Fixed race conditions due to writing to primitive id vertex buffers from a parallel task:
* Added GPrimitiveIdVertexBufferPool which provideds temporary volatile VBs and discards all allocated ones after finishing rendering the frame.
* MeshPassSetup task writes primitive id data to a temporary (MemStack) buffer, then before dispatching draw task, a special RHI command is used to copy its contents into VB.
* Also in general replaced all temporary FGlobalDynamicMeshVertexPool with a VB allocated GPrimitiveIdVertexBufferPool for efficiency.

diff --git a/Engine/Source/Runtime/Renderer/Private/MeshDrawCommands.cpp b/Engine/Source/Runtime/Renderer/Private/MeshDrawCommands.cpp
--- a/Engine/Source/Runtime/Renderer/Private/MeshDrawCommands.cpp
+++ b/Engine/Source/Runtime/Renderer/Private/MeshDrawCommands.cpp
@@ -9,12 +9,67 @@ MeshDrawCommandSetup.cpp: Mesh draw command setup.
 #include "ScenePrivate.h"
 #include "TranslucentRendering.h"
 
+FPrimitiveIdVertexBufferPool GPrimitiveIdVertexBufferPool;
+
 static TAutoConsoleVariable<int32> CVarMeshDrawCommandsParallelPassSetup(
 	TEXT("r.MeshDrawCommands.ParallelPassSetup"),
 	1,
 	TEXT("Whether to setup mesh draw command pass in parallel."),
 	ECVF_RenderThreadSafe);
 
+FPrimitiveIdVertexBufferPool::FPrimitiveIdVertexBufferPool()
+	: DiscardId(0)
+{
+}
+
+FPrimitiveIdVertexBufferPool::~FPrimitiveIdVertexBufferPool()
+{
+	Entries.Empty();
+}
+
+FVertexBufferRHIParamRef FPrimitiveIdVertexBufferPool::Allocate(int32 BufferSize)
+{
+	// First look for an allocated and unused one.
+	for (int32 Index = 0; Index < Entries.Num(); ++Index)
+	{
+		if (Entries[Index].LastDiscardId != DiscardId
+			&& Entries[Index].BufferSize >= BufferSize)
+		{
+			Entries[Index].LastDiscardId = DiscardId;
+			return Entries[Index].BufferRHI;
+		}
+	}
+
+
+	// Allocate new one.
+	FPrimitiveIdVertexBufferPoolEntry NewEntry;
+	NewEntry.LastDiscardId = DiscardId;
+	NewEntry.BufferSize = Align(BufferSize, 1024);
+	FRHIResourceCreateInfo CreateInfo;
+	NewEntry.BufferRHI = RHICreateVertexBuffer(NewEntry.BufferSize, BUF_Volatile, CreateInfo);
+	Entries.Add(NewEntry);
+
+	return NewEntry.BufferRHI;
+}
+
+void FPrimitiveIdVertexBufferPool::DiscardAll()
+{
+	++DiscardId;
+
+	// Remove old unused pool entries.
+	for (int32 Index = 0; Index < Entries.Num();)
+	{
+		if (DiscardId - Entries[Index].LastDiscardId > 1000u)
+		{
+			Entries.RemoveAtSwap(Index);
+		}
+		else
+		{
+			++Index;
+		}
+	}
+}
+
 struct FCompareFMeshDrawCommands
 {
 	FORCEINLINE bool operator() (FVisibleMeshDrawCommand A, FVisibleMeshDrawCommand B) const
@@ -232,7 +287,8 @@ void UpdateMobileBasePassMeshSortKeys(
 void BuildMeshDrawCommandPrimitiveIdBuffer(
 	FMeshCommandOneFrameArray& VisibleMeshDrawCommands,
 	FDynamicMeshDrawCommandStorage& MeshDrawCommandStorage,
-	FGlobalDynamicVertexBuffer::FAllocation& PrimitiveIdBuffer,
+	void* RESTRICT PrimitiveIdData,
+	int32 PrimitiveIdDataSize,
 	FMeshCommandOneFrameArray& TempVisibleMeshDrawCommands,
 	int32& MaxInstances,
 	int32& VisibleMeshDrawCommandsNum,
@@ -241,13 +297,14 @@ void BuildMeshDrawCommandPrimitiveIdBuffer(
 	)
 {
 	QUICK_SCOPE_CYCLE_COUNTER(STAT_BuildMeshDrawCommandPrimitiveIdBuffer);
-	check(VisibleMeshDrawCommands.Num() <= TempVisibleMeshDrawCommands.Max() && TempVisibleMeshDrawCommands.Num() == 0 && PrimitiveIdBuffer.IsValid());
+	check(VisibleMeshDrawCommands.Num() <= TempVisibleMeshDrawCommands.Max() && TempVisibleMeshDrawCommands.Num() == 0 && PrimitiveIdData);
 
 	const FVisibleMeshDrawCommand* RESTRICT PassVisibleMeshDrawCommands = VisibleMeshDrawCommands.GetData();
 	const int32 NumDrawCommands = VisibleMeshDrawCommands.Num();
 
 	uint32 PrimitiveIdIndex = 0;
-	int32* RESTRICT PrimitiveIds = (int32*)PrimitiveIdBuffer.Buffer;
+	int32* RESTRICT PrimitiveIds = (int32*)PrimitiveIdData;
+	const uint32 MaxPrimitiveId = PrimitiveIdDataSize / sizeof(int32);
 
 	if (IsDynamicInstancingEnabled())
 	{
@@ -320,6 +377,7 @@ void BuildMeshDrawCommandPrimitiveIdBuffer(
 			for (uint32 InstanceFactorIndex = 0; InstanceFactorIndex < InstanceFactor; InstanceFactorIndex++, PrimitiveIdIndex++)
 			{
 				//@todo - refactor into memcpy
+				checkSlow(PrimitiveIdIndex < MaxPrimitiveId);
 				PrimitiveIds[PrimitiveIdIndex] = VisibleMeshDrawCommand.DrawPrimitiveId;
 			}
 		}
@@ -343,6 +401,7 @@ void BuildMeshDrawCommandPrimitiveIdBuffer(
 			InstanceFactor = ensure(InstanceFactor <= MaxInstanceFactor) ? InstanceFactor : MaxInstanceFactor;
 			for (uint32 InstanceFactorIndex = 0; InstanceFactorIndex < InstanceFactor; InstanceFactorIndex++, PrimitiveIdIndex++)
 			{
+				checkSlow(PrimitiveIdIndex < MaxPrimitiveId);
 				PrimitiveIds[PrimitiveIdIndex] = VisibleMeshDrawCommand.DrawPrimitiveId;
 			}
 		}
@@ -685,7 +744,8 @@ public:
 				BuildMeshDrawCommandPrimitiveIdBuffer(
 					Context.MeshDrawCommands,
 					Context.MeshDrawCommandStorage,
-					Context.PrimitiveIdBuffer,
+					Context.PrimitiveIdBufferData,
+					Context.PrimitiveIdBufferDataSize,
 					Context.TempVisibleMeshDrawCommands,
 					Context.MaxInstances,
 					Context.VisibleMeshDrawCommandsNum,
@@ -709,8 +769,7 @@ void SortPassMeshDrawCommands(
 	ERHIFeatureLevel::Type FeatureLevel,
 	FMeshCommandOneFrameArray& VisibleMeshDrawCommands,
 	FDynamicMeshDrawCommandStorage& MeshDrawCommandStorage,
-	FGlobalDynamicVertexBuffer& DynamicVertexBuffer,
-	FGlobalDynamicVertexBuffer::FAllocation& PrimitiveIdBuffer)
+	FVertexBufferRHIParamRef& OutPrimitiveIdVertexBuffer)
 {
 	const bool bUseGPUScene = UseGPUScene(GMaxRHIShaderPlatform, FeatureLevel);
 
@@ -723,33 +782,34 @@ void SortPassMeshDrawCommands(
 		int32 NewPassVisibleMeshDrawCommandsNum = 0;
 		uint32 MaxInstanceFactor = 1;
 
-		// Preallocate context memory and resources.
+		VisibleMeshDrawCommands.Sort(FCompareFMeshDrawCommands());
+
 		if (bUseGPUScene)
 		{
-			//@todo - make the instance factor a pass parameter instead of a draw command parameter
-			MaxInstanceFactor = VisibleMeshDrawCommands[0].MeshDrawCommand->InstanceFactor;
-			PrimitiveIdBuffer = DynamicVertexBuffer.Allocate(NumDrawCommands * MaxInstanceFactor * sizeof(int32));
-
 			if (IsDynamicInstancingEnabled())
 			{
 				NewPassVisibleMeshDrawCommands.Empty(NumDrawCommands);
 			}
-		}
 
-		VisibleMeshDrawCommands.Sort(FCompareFMeshDrawCommands());
+			//@todo - make the instance factor a pass parameter instead of a draw command parameter
+			MaxInstanceFactor = VisibleMeshDrawCommands[0].MeshDrawCommand->InstanceFactor;
+			const int32 PrimitiveIdBufferDataSize = MaxInstanceFactor * NumDrawCommands * sizeof(int32);
+			OutPrimitiveIdVertexBuffer = GPrimitiveIdVertexBufferPool.Allocate(PrimitiveIdBufferDataSize);
+			void* PrimitiveIdBufferData = GDynamicRHI->RHILockVertexBuffer(OutPrimitiveIdVertexBuffer, 0, PrimitiveIdBufferDataSize, RLM_WriteOnly);
 
-		if (bUseGPUScene)
-		{
 			BuildMeshDrawCommandPrimitiveIdBuffer(
 				VisibleMeshDrawCommands,
 				MeshDrawCommandStorage,
-				PrimitiveIdBuffer,
+				PrimitiveIdBufferData,
+				PrimitiveIdBufferDataSize,
 				NewPassVisibleMeshDrawCommands,
 				MaxInstances,
 				VisibleMeshDrawCommandsNum,
 				NewPassVisibleMeshDrawCommandsNum,
 				MaxInstanceFactor
 			);
+
+			GDynamicRHI->RHIUnlockVertexBuffer(OutPrimitiveIdVertexBuffer);
 		}
 	}
 }
@@ -827,7 +887,9 @@ void FParallelMeshDrawCommandPass::DispatchPassSetup(
 	// Preallocate resources on rendering thread based on MaxNumDraws.
 	if (TaskContext.bDynamicInstancing)
 	{
-		TaskContext.PrimitiveIdBuffer = FGlobalDynamicVertexBuffer::Get().Allocate(MaxNumDraws * TaskContext.MaxInstanceFactor * sizeof(int32));
+		TaskContext.PrimitiveIdBufferDataSize = TaskContext.MaxInstanceFactor * MaxNumDraws * sizeof(int32);
+		TaskContext.PrimitiveIdBufferData = FMemStack::Get().Alloc(TaskContext.PrimitiveIdBufferDataSize, 8);
+		PrimitiveIdVertexBufferRHI = GPrimitiveIdVertexBufferPool.Allocate(TaskContext.PrimitiveIdBufferDataSize);
 	}
 	TaskContext.MeshDrawCommands.Reserve(MaxNumDraws);
 	TaskContext.TempVisibleMeshDrawCommands.Reserve(MaxNumDraws);
@@ -872,7 +934,8 @@ void FParallelMeshDrawCommandPass::Empty()
 	TaskContext.MobileBasePassCSMMeshDrawCommands.Empty();
 	TaskContext.DynamicMeshCommandBuildRequests.Empty();
 	TaskContext.TempVisibleMeshDrawCommands.Empty();
-	TaskContext.PrimitiveIdBuffer = FGlobalDynamicVertexBuffer::FAllocation();
+	TaskContext.PrimitiveIdBufferData = nullptr;
+	TaskContext.PrimitiveIdBufferDataSize = 0;
 }
 
 FParallelMeshDrawCommandPass::~FParallelMeshDrawCommandPass()
@@ -934,6 +997,42 @@ public:
 	}
 };
 
+/*
+ * Copies provided vertex data (assumed to be on MemStack) to a vertex buffer.
+ */
+struct FRHICommandUpdatePrimitiveIdBuffer : public FRHICommand<FRHICommandUpdatePrimitiveIdBuffer>
+{
+	FGraphEventArray Prereqs;
+
+	FVertexBufferRHIParamRef VertexBuffer;
+	const void* VertexBufferData;
+	int32 VertexBufferDataSize;
+
+	virtual ~FRHICommandUpdatePrimitiveIdBuffer() {}
+
+	FORCEINLINE_DEBUGGABLE FRHICommandUpdatePrimitiveIdBuffer(
+		const FGraphEventArray& InPrereqs,
+		FVertexBufferRHIParamRef InVertexBuffer,
+		const void* InVertexBufferData,
+		int32 InVertexBufferDataSize)
+		: Prereqs(InPrereqs)
+		, VertexBuffer(InVertexBuffer)
+		, VertexBufferData(InVertexBufferData)
+		, VertexBufferDataSize(InVertexBufferDataSize)
+	{
+	}
+
+	void Execute(FRHICommandListBase& CmdList)
+	{
+		FTaskGraphInterface::Get().WaitUntilTasksComplete(Prereqs, IsRunningRHIInSeparateThread() ? ENamedThreads::RHIThread : ENamedThreads::ActualRenderingThread);
+
+		// Upload vertex buffer data.
+		void* RESTRICT Data = (void* RESTRICT)GDynamicRHI->RHILockVertexBuffer(VertexBuffer, 0, VertexBufferDataSize, RLM_WriteOnly);
+		FMemory::BigBlockMemcpy(Data, VertexBufferData, VertexBufferDataSize);
+		GDynamicRHI->RHIUnlockVertexBuffer(VertexBuffer);
+	}
+};
+
 void FParallelMeshDrawCommandPass::DispatchDraw(FParallelCommandListSet* ParallelCommandListSet, FRHICommandList& RHICmdList) const
 {
 	if (MaxNumDraws <= 0)
@@ -942,11 +1041,24 @@ void FParallelMeshDrawCommandPass::DispatchDraw(FParallelCommandListSet* Paralle
 	}
 
 
-	FVertexBufferRHIParamRef PrimitiveIdsBuffer = TaskContext.bUseGPUScene ? TaskContext.PrimitiveIdBuffer.VertexBuffer->VertexBufferRHI : nullptr;
-	const int32 BasePrimitiveIdsOffset = TaskContext.PrimitiveIdBuffer.VertexOffset;
+	FVertexBufferRHIParamRef PrimitiveIdsBuffer = PrimitiveIdVertexBufferRHI;
+	const int32 BasePrimitiveIdsOffset = 0;
+
 
 	if (ParallelCommandListSet)
 	{
+		if (TaskContext.bUseGPUScene)
+		{
+			// Queue a command on the RHI thread which will upload PrimitiveIdVertexBuffer after finishing FMeshDrawCommandPassSetupTaskContext.
+			FRHICommandListImmediate &RHICommandList = GetImmediateCommandList_ForRenderCommand();
+			new (RHICommandList.AllocCommand<FRHICommandUpdatePrimitiveIdBuffer>())FRHICommandUpdatePrimitiveIdBuffer(
+				TaskEventRefs,
+				PrimitiveIdVertexBufferRHI,
+				TaskContext.PrimitiveIdBufferData,
+				TaskContext.PrimitiveIdBufferDataSize
+			);
+		}
+
 		const ENamedThreads::Type RenderThread = ENamedThreads::GetRenderThread();
 
 		FGraphEventArray Prereqs = TaskEventRefs;
@@ -977,6 +1089,15 @@ void FParallelMeshDrawCommandPass::DispatchDraw(FParallelCommandListSet* Paralle
 	else
 	{
 		WaitForMeshPassSetupTask();
+
+		if (TaskContext.bUseGPUScene)
+		{
+			// Can immediately upload vertex buffer data, as there is no parallel draw task.
+			void* RESTRICT Data = (void* RESTRICT)GDynamicRHI->RHILockVertexBuffer(PrimitiveIdVertexBufferRHI, 0, TaskContext.PrimitiveIdBufferDataSize, RLM_WriteOnly);
+			FMemory::BigBlockMemcpy(Data, TaskContext.PrimitiveIdBufferData, TaskContext.PrimitiveIdBufferDataSize);
+			GDynamicRHI->RHIUnlockVertexBuffer(PrimitiveIdVertexBufferRHI);
+		}
+
 		SubmitMeshDrawCommandsRange(TaskContext.MeshDrawCommands, PrimitiveIdsBuffer, BasePrimitiveIdsOffset, TaskContext.bDynamicInstancing, 0, TaskContext.MeshDrawCommands.Num(), RHICmdList);
 	}
 }
diff --git a/Engine/Source/Runtime/Renderer/Private/MeshDrawCommands.h b/Engine/Source/Runtime/Renderer/Private/MeshDrawCommands.h
--- a/Engine/Source/Runtime/Renderer/Private/MeshDrawCommands.h
+++ b/Engine/Source/Runtime/Renderer/Private/MeshDrawCommands.h
@@ -13,6 +13,32 @@ struct FMeshBatchAndRelevance;
 class FStaticMeshBatch;
 class FParallelCommandListSet;
 
+/**
+ * Global vertex buffer pool used for GPUScene primitive id arrays.
+ */
+class FPrimitiveIdVertexBufferPool
+{
+public:
+	FPrimitiveIdVertexBufferPool();
+	~FPrimitiveIdVertexBufferPool();
+
+	FVertexBufferRHIParamRef Allocate(int32 BufferSize);
+	void DiscardAll();
+
+private:
+	struct FPrimitiveIdVertexBufferPoolEntry
+	{
+		int32 BufferSize;
+		uint32 LastDiscardId;
+		FVertexBufferRHIRef BufferRHI;
+	};
+
+	uint32 DiscardId;
+	TArray<FPrimitiveIdVertexBufferPoolEntry> Entries;
+};
+
+extern FPrimitiveIdVertexBufferPool GPrimitiveIdVertexBufferPool;
+
 /**	
  * Parallel mesh draw command pass setup task context.
  */
@@ -64,8 +90,9 @@ public:
 	TArray<const FStaticMeshBatch*, SceneRenderingAllocator> MobileBasePassCSMDynamicMeshCommandBuildRequests;
 	FDynamicMeshDrawCommandStorage MeshDrawCommandStorage;
 
-	// Preallocated resources.
-	FGlobalDynamicVertexBuffer::FAllocation PrimitiveIdBuffer;
+	// Resources preallocated on rendering thread with SceneRenderingAllocator.
+	void* PrimitiveIdBufferData;
+	int32 PrimitiveIdBufferDataSize;
 	FMeshCommandOneFrameArray TempVisibleMeshDrawCommands;
 
 	// For UpdateTranslucentMeshSortKeys.
@@ -129,6 +156,7 @@ public:
 
 
 private:
+	FVertexBufferRHIParamRef PrimitiveIdVertexBufferRHI;
 	FMeshDrawCommandPassSetupTaskContext TaskContext;
 	FGraphEventArray TaskEventRefs;
 	FString PassNameForStats;
diff --git a/Engine/Source/Runtime/Renderer/Private/MeshPassProcessor.cpp b/Engine/Source/Runtime/Renderer/Private/MeshPassProcessor.cpp
--- a/Engine/Source/Runtime/Renderer/Private/MeshPassProcessor.cpp
+++ b/Engine/Source/Runtime/Renderer/Private/MeshPassProcessor.cpp
@@ -859,18 +859,13 @@ void DrawDynamicMeshPassPrivate(
 {
 	if (VisibleMeshDrawCommands.Num() > 0)
 	{
-		FGlobalDynamicVertexBuffer DynamicVertexBuffer;
-		FGlobalDynamicVertexBuffer::FAllocation PrimitiveIdBufferAllocation;
+		FVertexBufferRHIParamRef PrimitiveIdVertexBuffer = nullptr;
 
-		SortPassMeshDrawCommands(View.GetFeatureLevel(), VisibleMeshDrawCommands, DynamicMeshDrawCommandStorage, DynamicVertexBuffer, PrimitiveIdBufferAllocation);
+		SortPassMeshDrawCommands(View.GetFeatureLevel(), VisibleMeshDrawCommands, DynamicMeshDrawCommandStorage, PrimitiveIdVertexBuffer);
 
-		DynamicVertexBuffer.Commit();
-
-		FVertexBufferRHIParamRef PrimitiveIdsBuffer = UseGPUScene(GMaxRHIShaderPlatform, View.GetFeatureLevel()) ? PrimitiveIdBufferAllocation.VertexBuffer->VertexBufferRHI : nullptr;
-		const int32 BasePrimitiveIdsOffset = PrimitiveIdBufferAllocation.VertexOffset;
 		const bool bDynamicInstancing = IsDynamicInstancingEnabled() && UseGPUScene(GMaxRHIShaderPlatform, View.GetFeatureLevel());
 
-		SubmitMeshDrawCommandsRange(VisibleMeshDrawCommands, PrimitiveIdsBuffer, BasePrimitiveIdsOffset, bDynamicInstancing, 0, VisibleMeshDrawCommands.Num(), RHICmdList);
+		SubmitMeshDrawCommandsRange(VisibleMeshDrawCommands, PrimitiveIdVertexBuffer, 0, bDynamicInstancing, 0, VisibleMeshDrawCommands.Num(), RHICmdList);
 	}
 }
 
diff --git a/Engine/Source/Runtime/Renderer/Private/SceneRendering.cpp b/Engine/Source/Runtime/Renderer/Private/SceneRendering.cpp
--- a/Engine/Source/Runtime/Renderer/Private/SceneRendering.cpp
+++ b/Engine/Source/Runtime/Renderer/Private/SceneRendering.cpp
@@ -2881,6 +2881,9 @@ void FSceneRenderer::WaitForTasksClearSnapshotsAndDeleteSceneRenderer(FRHIComman
 		QUICK_SCOPE_CYCLE_COUNTER(STAT_DeleteSceneRenderer_WaitForTasks);
 		RHICmdList.ImmediateFlush(EImmediateFlushType::WaitForOutstandingTasksOnly);
 	}
+
+	GPrimitiveIdVertexBufferPool.DiscardAll();
+
 	FViewInfo::DestroyAllSnapshots(); // this destroys viewinfo snapshots
 	FSceneRenderTargets::GetGlobalUnsafe().DestroyAllSnapshots(); // this will destroy the render target snapshots
 	static const IConsoleVariable* AsyncDispatch	= IConsoleManager::Get().FindConsoleVariable(TEXT("r.RHICmdAsyncRHIThreadDispatch"));
diff --git a/Engine/Source/Runtime/Renderer/Private/SceneRendering.h b/Engine/Source/Runtime/Renderer/Private/SceneRendering.h
--- a/Engine/Source/Runtime/Renderer/Private/SceneRendering.h
+++ b/Engine/Source/Runtime/Renderer/Private/SceneRendering.h
@@ -1893,5 +1893,4 @@ extern void SortPassMeshDrawCommands(
 	ERHIFeatureLevel::Type FeatureLevel,
 	FMeshCommandOneFrameArray& VisibleMeshDrawCommands,
 	FDynamicMeshDrawCommandStorage& MeshDrawCommandStorage,
-	FGlobalDynamicVertexBuffer& DynamicVertexBuffer,
-	FGlobalDynamicVertexBuffer::FAllocation& PrimitiveIdBuffer);
\ No newline at end of file
+	FVertexBufferRHIParamRef& OutPrimitiveIdVertexBuffer);
\ No newline at end of file
diff --git a/Engine/Source/Runtime/Renderer/Private/ShadowRendering.cpp b/Engine/Source/Runtime/Renderer/Private/ShadowRendering.cpp
--- a/Engine/Source/Runtime/Renderer/Private/ShadowRendering.cpp
+++ b/Engine/Source/Runtime/Renderer/Private/ShadowRendering.cpp
@@ -612,11 +612,9 @@ void FProjectedShadowInfo::SetupProjectionStencilMask(
 		const FShadowMeshDrawCommandPass& ProjectionStencilingPass = ProjectionStencilingPasses[ViewIndex];
 		if (ProjectionStencilingPass.VisibleMeshDrawCommands.Num() > 0)
 		{
-			FVertexBufferRHIParamRef PrimitiveIdsBuffer = UseGPUScene(GMaxRHIShaderPlatform, View->GetFeatureLevel()) ? ProjectionStencilingPass.PrimitiveIdBuffer.VertexBuffer->VertexBufferRHI : nullptr;
-			const int32 BasePrimitiveIdsOffset = ProjectionStencilingPass.PrimitiveIdBuffer.VertexOffset;
 			const bool bDynamicInstancing = IsDynamicInstancingEnabled() && UseGPUScene(GMaxRHIShaderPlatform, View->GetFeatureLevel());
 
-			SubmitMeshDrawCommands(ProjectionStencilingPass.VisibleMeshDrawCommands, PrimitiveIdsBuffer, BasePrimitiveIdsOffset, bDynamicInstancing, RHICmdList);
+			SubmitMeshDrawCommands(ProjectionStencilingPass.VisibleMeshDrawCommands, ProjectionStencilingPass.PrimitiveIdVertexBuffer, 0, bDynamicInstancing, RHICmdList);
 		}
 
 		// Restore viewport
@@ -756,11 +754,9 @@ void FProjectedShadowInfo::SetupProjectionStencilMask(
 			const FShadowMeshDrawCommandPass& ProjectionStencilingPass = ProjectionStencilingPasses[ViewIndex];
 			if (ProjectionStencilingPass.VisibleMeshDrawCommands.Num() > 0)
 			{
-				FVertexBufferRHIParamRef PrimitiveIdsBuffer = UseGPUScene(GMaxRHIShaderPlatform, View->GetFeatureLevel()) ? ProjectionStencilingPass.PrimitiveIdBuffer.VertexBuffer->VertexBufferRHI : nullptr;
-				const int32 BasePrimitiveIdsOffset = ProjectionStencilingPass.PrimitiveIdBuffer.VertexOffset;
 				const bool bDynamicInstancing = IsDynamicInstancingEnabled() && UseGPUScene(GMaxRHIShaderPlatform, View->GetFeatureLevel());
 
-				SubmitMeshDrawCommands(ProjectionStencilingPass.VisibleMeshDrawCommands, PrimitiveIdsBuffer, BasePrimitiveIdsOffset, bDynamicInstancing, RHICmdList);
+				SubmitMeshDrawCommands(ProjectionStencilingPass.VisibleMeshDrawCommands, ProjectionStencilingPass.PrimitiveIdVertexBuffer, 0, bDynamicInstancing, RHICmdList);
 			}
 		}
 	}
diff --git a/Engine/Source/Runtime/Renderer/Private/ShadowRendering.h b/Engine/Source/Runtime/Renderer/Private/ShadowRendering.h
--- a/Engine/Source/Runtime/Renderer/Private/ShadowRendering.h
+++ b/Engine/Source/Runtime/Renderer/Private/ShadowRendering.h
@@ -189,7 +189,7 @@ class FShadowMeshDrawCommandPass
 {
 public:
 	FMeshCommandOneFrameArray VisibleMeshDrawCommands;
-	FGlobalDynamicVertexBuffer::FAllocation PrimitiveIdBuffer;
+	FVertexBufferRHIParamRef PrimitiveIdVertexBuffer;
 };
 
 /**
diff --git a/Engine/Source/Runtime/Renderer/Private/ShadowSetup.cpp b/Engine/Source/Runtime/Renderer/Private/ShadowSetup.cpp
--- a/Engine/Source/Runtime/Renderer/Private/ShadowSetup.cpp
+++ b/Engine/Source/Runtime/Renderer/Private/ShadowSetup.cpp
@@ -1290,7 +1290,7 @@ void FProjectedShadowInfo::SetupMeshDrawCommandsForProjectionStenciling(FSceneRe
 
 			ApplyViewOverridesToMeshDrawCommands(View, ProjectionStencilingPass.VisibleMeshDrawCommands);
 
-			SortPassMeshDrawCommands(Renderer.FeatureLevel, ProjectionStencilingPass.VisibleMeshDrawCommands, DynamicMeshDrawCommandStorage, FGlobalDynamicVertexBuffer::Get(), ProjectionStencilingPass.PrimitiveIdBuffer);
+			SortPassMeshDrawCommands(Renderer.FeatureLevel, ProjectionStencilingPass.VisibleMeshDrawCommands, DynamicMeshDrawCommandStorage, ProjectionStencilingPass.PrimitiveIdVertexBuffer);
 		}
 	}
 }
