Fixed compile errors due to missing PipelineStateCache.h include

diff --git a/Engine/Source/Runtime/Renderer/Private/PathTracing.cpp b/Engine/Source/Runtime/Renderer/Private/PathTracing.cpp
--- a/Engine/Source/Runtime/Renderer/Private/PathTracing.cpp
+++ b/Engine/Source/Runtime/Renderer/Private/PathTracing.cpp
@@ -10,6 +10,7 @@
 #include "PostProcess/PostProcessing.h"
 #include "PostProcess/SceneFilterRendering.h"
 #include "PathTracingUniformBuffers.h"
+#include "RHI/Public/PipelineStateCache.h"
 
 ////static TAutoConsoleVariable<int32> CVarRayTracingPrimaryDebugMaxBounces(
 ////	TEXT("r.RayTracing.PrimaryDebug.MaxBounces"),
diff --git a/Engine/Source/Runtime/Renderer/Private/PathTracing/SkyLightCdfBuild.cpp b/Engine/Source/Runtime/Renderer/Private/PathTracing/SkyLightCdfBuild.cpp
--- a/Engine/Source/Runtime/Renderer/Private/PathTracing/SkyLightCdfBuild.cpp
+++ b/Engine/Source/Runtime/Renderer/Private/PathTracing/SkyLightCdfBuild.cpp
@@ -12,6 +12,7 @@
 #include "DeferredShadingRenderer.h"
 #include "PostProcess/PostProcessing.h"
 #include "PostProcess/SceneFilterRendering.h"
+#include "RHI/Public/PipelineStateCache.h"
 
 class FBuildSkyLightRowCdfCS : public FGlobalShader
 {
diff --git a/Engine/Source/Runtime/Renderer/Private/PathTracing/VarianceMipTreeBuild.cpp b/Engine/Source/Runtime/Renderer/Private/PathTracing/VarianceMipTreeBuild.cpp
--- a/Engine/Source/Runtime/Renderer/Private/PathTracing/VarianceMipTreeBuild.cpp
+++ b/Engine/Source/Runtime/Renderer/Private/PathTracing/VarianceMipTreeBuild.cpp
@@ -12,6 +12,7 @@
 #include "DeferredShadingRenderer.h"
 #include "PostProcess/PostProcessing.h"
 #include "PostProcess/SceneFilterRendering.h"
+#include "RHI/Public/PipelineStateCache.h"
 
 class FBuildVarianceMipTreeCS : public FGlobalShader
 {
diff --git a/Engine/Source/Runtime/Renderer/Private/PostProcess/PostProcessMaterial.cpp b/Engine/Source/Runtime/Renderer/Private/PostProcess/PostProcessMaterial.cpp
--- a/Engine/Source/Runtime/Renderer/Private/PostProcess/PostProcessMaterial.cpp
+++ b/Engine/Source/Runtime/Renderer/Private/PostProcess/PostProcessMaterial.cpp
@@ -13,6 +13,7 @@
 #include "PostProcess/SceneFilterRendering.h"
 #include "SceneRendering.h"
 #include "ClearQuad.h"
+#include "RHI/Public/PipelineStateCache.h"
 
 enum class EPostProcessMaterialTarget
 {
diff --git a/Engine/Source/Runtime/Renderer/Private/RayTracing/RayTracingAmbientOcclusion.cpp b/Engine/Source/Runtime/Renderer/Private/RayTracing/RayTracingAmbientOcclusion.cpp
--- a/Engine/Source/Runtime/Renderer/Private/RayTracing/RayTracingAmbientOcclusion.cpp
+++ b/Engine/Source/Runtime/Renderer/Private/RayTracing/RayTracingAmbientOcclusion.cpp
@@ -11,6 +11,7 @@
 #include "RenderTargetPool.h"
 #include "RHIResources.h"
 #include "UniformBuffer.h"
+#include "RHI/Public/PipelineStateCache.h"
 
 #include "PostProcess/PostProcessing.h"
 #include "PostProcess/SceneFilterRendering.h"
diff --git a/Engine/Source/Runtime/Renderer/Private/RayTracing/RayTracingBarycentrics.cpp b/Engine/Source/Runtime/Renderer/Private/RayTracing/RayTracingBarycentrics.cpp
--- a/Engine/Source/Runtime/Renderer/Private/RayTracing/RayTracingBarycentrics.cpp
+++ b/Engine/Source/Runtime/Renderer/Private/RayTracing/RayTracingBarycentrics.cpp
@@ -8,6 +8,7 @@
 #include "GlobalShader.h"
 #include "SceneRenderTargets.h"
 #include "RenderGraphBuilder.h"
+#include "RHI/Public/PipelineStateCache.h"
 
 class FRayTracingBarycentricsRGS : public FGlobalShader
 {
diff --git a/Engine/Source/Runtime/Renderer/Private/RayTracing/RayTracingShadowsForRectLight.cpp b/Engine/Source/Runtime/Renderer/Private/RayTracing/RayTracingShadowsForRectLight.cpp
--- a/Engine/Source/Runtime/Renderer/Private/RayTracing/RayTracingShadowsForRectLight.cpp
+++ b/Engine/Source/Runtime/Renderer/Private/RayTracing/RayTracingShadowsForRectLight.cpp
@@ -16,6 +16,8 @@
 #include "PostProcess/SceneFilterRendering.h"
 #include "RectLightSceneProxy.h"
 
+#include "RHI/Public/PipelineStateCache.h"
+
 static int32 GRayTracingRectLight = 1;
 static FAutoConsoleVariableRef CVarRayTracingRectLight(
 	TEXT("r.RayTracing.RectLight"),
diff --git a/Engine/Source/Runtime/Renderer/Private/VolumetricFogLightFunction.cpp b/Engine/Source/Runtime/Renderer/Private/VolumetricFogLightFunction.cpp
--- a/Engine/Source/Runtime/Renderer/Private/VolumetricFogLightFunction.cpp
+++ b/Engine/Source/Runtime/Renderer/Private/VolumetricFogLightFunction.cpp
@@ -11,6 +11,7 @@
 #include "LightRendering.h"
 #include "SceneFilterRendering.h"
 #include "PostProcessing.h"
+#include "RHI/Public/PipelineStateCache.h"
 
 float GVolumetricFogLightFunctionSupersampleScale = 2.0f;
 FAutoConsoleVariableRef CVarVolumetricFogLightFunctionSupersampleScale(
