Fixed r.DoInitViewsLightingAfterPrepass=1 rendering path.
* Removed experimental r.FXSystemPreRenderAfterPrepass, as now there is no distinction between static and dynamic draw lists, so we can't first dispatch static one, then update FX, and after that dispatch the dynamic draw list.
* Moved AfterTasksAreStarted outside of parallel cmd list, as it calls a resource flush, which is not allowed when there is a pending parallel cmd list.
* Moved AfterTasksAreStarted to be done after the prepass dispatch, as this is what r.DoInitViewsLightingAfterPrepass is intended to do.

diff --git a/Engine/Source/Runtime/Renderer/Private/DeferredShadingRenderer.cpp b/Engine/Source/Runtime/Renderer/Private/DeferredShadingRenderer.cpp
--- a/Engine/Source/Runtime/Renderer/Private/DeferredShadingRenderer.cpp
+++ b/Engine/Source/Runtime/Renderer/Private/DeferredShadingRenderer.cpp
@@ -105,13 +105,6 @@ static TAutoConsoleVariable<int32> CVarIrisNormal(
 	TEXT(" 1: on"),
 	ECVF_ReadOnly);
 
-static TAutoConsoleVariable<int32> CVarFXSystemPreRenderAfterPrepass(
-	TEXT("r.FXSystemPreRenderAfterPrepass"),
-	0,
-	TEXT("If > 0, then do the FX prerender after the prepass. This improves pipelining for greater performance. Experimental option."),
-	ECVF_RenderThreadSafe
-	);
-
 int32 GbEnableAsyncComputeTranslucencyLightingVolumeClear = 1;
 static FAutoConsoleVariableRef CVarEnableAsyncComputeTranslucencyLightingVolumeClear(
 	TEXT("r.EnableAsyncComputeTranslucencyLightingVolumeClear"),
@@ -1167,10 +1160,9 @@ void FDeferredShadingSceneRenderer::Render(FRHICommandListImmediate& RHICmdList)
 	const bool bAllowGPUParticleSceneUpdate = !Views[0].bIsPlanarReflection && !Views[0].bIsSceneCapture && !Views[0].bIsReflectionCapture;
 
 	// Notify the FX system that the scene is about to be rendered.
-	bool bLateFXPrerender = CVarFXSystemPreRenderAfterPrepass.GetValueOnRenderThread() > 0;
 	bool bDoFXPrerender = Scene->FXSystem && Views.IsValidIndex(0) && bAllowGPUParticleSceneUpdate;
 
-	if (!bLateFXPrerender && bDoFXPrerender)
+	if (bDoFXPrerender)
 	{
 		SCOPE_CYCLE_COUNTER(STAT_FDeferredShadingSceneRenderer_FXSystem_PreRender);
 		RHICmdList.SetCurrentStat(GET_STATID(STAT_CLM_FXPreRender));
@@ -1178,7 +1170,7 @@ void FDeferredShadingSceneRenderer::Render(FRHICommandListImmediate& RHICmdList)
 	}
 
 	bool bDidAfterTaskWork = false;
-	auto AfterTasksAreStarted = [&bDidAfterTaskWork, bDoInitViewAftersPrepass, this, &RHICmdList, &ILCTaskData, &UpdateViewCustomDataEvents, bLateFXPrerender, bDoFXPrerender]()
+	auto AfterTasksAreStarted = [&bDidAfterTaskWork, bDoInitViewAftersPrepass, this, &RHICmdList, &ILCTaskData, &UpdateViewCustomDataEvents, bDoFXPrerender]()
 	{
 		if (!bDidAfterTaskWork)
 		{
@@ -1198,13 +1190,6 @@ void FDeferredShadingSceneRenderer::Render(FRHICommandListImmediate& RHICmdList)
 				}
 				ServiceLocalQueue();
 			}
-			if (bLateFXPrerender && bDoFXPrerender)
-			{
-				SCOPE_CYCLE_COUNTER(STAT_FDeferredShadingSceneRenderer_FXSystem_PreRender);
-				RHICmdList.SetCurrentStat(GET_STATID(STAT_CLM_FXPreRender));
-				Scene->FXSystem->PreRender(RHICmdList, &Views[0].GlobalDistanceFieldInfo.ParameterData);
-				ServiceLocalQueue();
-			}
 		}
 	};
 
diff --git a/Engine/Source/Runtime/Renderer/Private/DepthRendering.cpp b/Engine/Source/Runtime/Renderer/Private/DepthRendering.cpp
--- a/Engine/Source/Runtime/Renderer/Private/DepthRendering.cpp
+++ b/Engine/Source/Runtime/Renderer/Private/DepthRendering.cpp
@@ -299,11 +299,15 @@ bool FDeferredShadingSceneRenderer::RenderPrePassViewParallel(const FViewInfo& V
 
 		if (bDoPrePre)
 		{
-			AfterTasksAreStarted();
 			bDepthWasCleared = PreRenderPrePass(ParentCmdList);
 		}
 	}
 
+	if (bDoPrePre)
+	{
+		AfterTasksAreStarted();
+	}
+
 	return bDepthWasCleared;
 }
 
