Renamed SortPassMeshDrawCommands to SortAndMergeDynamicPassMeshDrawCommands to better describe what it actually does.

diff --git a/Engine/Source/Runtime/Renderer/Private/MeshDrawCommands.cpp b/Engine/Source/Runtime/Renderer/Private/MeshDrawCommands.cpp
--- a/Engine/Source/Runtime/Renderer/Private/MeshDrawCommands.cpp
+++ b/Engine/Source/Runtime/Renderer/Private/MeshDrawCommands.cpp
@@ -765,7 +765,10 @@ private:
 	FMeshDrawCommandPassSetupTaskContext& Context;
 };
 
-void SortPassMeshDrawCommands(
+/*
+ * Used by various dynamic passes to sort/merge mesh draw commands immediately on a rendering thread.
+ */
+void SortAndMergeDynamicPassMeshDrawCommands(
 	ERHIFeatureLevel::Type FeatureLevel,
 	FMeshCommandOneFrameArray& VisibleMeshDrawCommands,
 	FDynamicMeshDrawCommandStorage& MeshDrawCommandStorage,
diff --git a/Engine/Source/Runtime/Renderer/Private/MeshDrawCommands.h b/Engine/Source/Runtime/Renderer/Private/MeshDrawCommands.h
--- a/Engine/Source/Runtime/Renderer/Private/MeshDrawCommands.h
+++ b/Engine/Source/Runtime/Renderer/Private/MeshDrawCommands.h
@@ -168,3 +168,9 @@ private:
 	void DumpInstancingStats() const;
 	void WaitForMeshPassSetupTask() const;
 };
+
+extern void SortAndMergeDynamicPassMeshDrawCommands(
+	ERHIFeatureLevel::Type FeatureLevel,
+	FMeshCommandOneFrameArray& VisibleMeshDrawCommands,
+	FDynamicMeshDrawCommandStorage& MeshDrawCommandStorage,
+	FVertexBufferRHIParamRef& OutPrimitiveIdVertexBuffer);
\ No newline at end of file
diff --git a/Engine/Source/Runtime/Renderer/Private/MeshPassProcessor.cpp b/Engine/Source/Runtime/Renderer/Private/MeshPassProcessor.cpp
--- a/Engine/Source/Runtime/Renderer/Private/MeshPassProcessor.cpp
+++ b/Engine/Source/Runtime/Renderer/Private/MeshPassProcessor.cpp
@@ -861,7 +861,7 @@ void DrawDynamicMeshPassPrivate(
 	{
 		FVertexBufferRHIParamRef PrimitiveIdVertexBuffer = nullptr;
 
-		SortPassMeshDrawCommands(View.GetFeatureLevel(), VisibleMeshDrawCommands, DynamicMeshDrawCommandStorage, PrimitiveIdVertexBuffer);
+		SortAndMergeDynamicPassMeshDrawCommands(View.GetFeatureLevel(), VisibleMeshDrawCommands, DynamicMeshDrawCommandStorage, PrimitiveIdVertexBuffer);
 
 		const bool bDynamicInstancing = IsDynamicInstancingEnabled() && UseGPUScene(GMaxRHIShaderPlatform, View.GetFeatureLevel());
 
diff --git a/Engine/Source/Runtime/Renderer/Private/SceneRendering.h b/Engine/Source/Runtime/Renderer/Private/SceneRendering.h
--- a/Engine/Source/Runtime/Renderer/Private/SceneRendering.h
+++ b/Engine/Source/Runtime/Renderer/Private/SceneRendering.h
@@ -1887,10 +1887,4 @@ private:
 extern FFastVramConfig GFastVRamConfig;
 
 extern bool UseCachedMeshDrawCommands();
-extern bool IsDynamicInstancingEnabled();
-
-extern void SortPassMeshDrawCommands(
-	ERHIFeatureLevel::Type FeatureLevel,
-	FMeshCommandOneFrameArray& VisibleMeshDrawCommands,
-	FDynamicMeshDrawCommandStorage& MeshDrawCommandStorage,
-	FVertexBufferRHIParamRef& OutPrimitiveIdVertexBuffer);
\ No newline at end of file
+extern bool IsDynamicInstancingEnabled();
\ No newline at end of file
diff --git a/Engine/Source/Runtime/Renderer/Private/ShadowSetup.cpp b/Engine/Source/Runtime/Renderer/Private/ShadowSetup.cpp
--- a/Engine/Source/Runtime/Renderer/Private/ShadowSetup.cpp
+++ b/Engine/Source/Runtime/Renderer/Private/ShadowSetup.cpp
@@ -1290,7 +1290,7 @@ void FProjectedShadowInfo::SetupMeshDrawCommandsForProjectionStenciling(FSceneRe
 
 			ApplyViewOverridesToMeshDrawCommands(View, ProjectionStencilingPass.VisibleMeshDrawCommands);
 
-			SortPassMeshDrawCommands(Renderer.FeatureLevel, ProjectionStencilingPass.VisibleMeshDrawCommands, DynamicMeshDrawCommandStorage, ProjectionStencilingPass.PrimitiveIdVertexBuffer);
+			SortAndMergeDynamicPassMeshDrawCommands(Renderer.FeatureLevel, ProjectionStencilingPass.VisibleMeshDrawCommands, DynamicMeshDrawCommandStorage, ProjectionStencilingPass.PrimitiveIdVertexBuffer);
 		}
 	}
 }
