Empty parallel mesh draw command passes when destroying all snapshots.

diff --git a/Engine/Source/Runtime/Renderer/Private/SceneRendering.cpp b/Engine/Source/Runtime/Renderer/Private/SceneRendering.cpp
--- a/Engine/Source/Runtime/Renderer/Private/SceneRendering.cpp
+++ b/Engine/Source/Runtime/Renderer/Private/SceneRendering.cpp
@@ -1569,6 +1569,12 @@ void FViewInfo::DestroyAllSnapshots()
 		Snapshot->CachedViewUniformShaderParameters.Reset();
 		Snapshot->DynamicPrimitiveShaderData.Empty();
 		Snapshot->OneFramePrimitiveShaderDataBuffer.Release();
+
+		for (int32 Index = 0; Index < Snapshot->ParallelMeshDrawCommandPasses.Num(); ++Index)
+		{
+			Snapshot->ParallelMeshDrawCommandPasses[Index].Empty();
+		}
+
 		FreeViewInfoSnapshots.Add(Snapshot);
 	}
 	ViewInfoSnapshots.Reset();
