DR - Remove Editor Primitives (non Game Primitives) from the Raytracing Scene. Also added Showflags for static and Skelmeshes.

diff --git a/Engine/Source/Runtime/Engine/Public/PrimitiveSceneProxy.h b/Engine/Source/Runtime/Engine/Public/PrimitiveSceneProxy.h
--- a/Engine/Source/Runtime/Engine/Public/PrimitiveSceneProxy.h
+++ b/Engine/Source/Runtime/Engine/Public/PrimitiveSceneProxy.h
@@ -695,6 +695,9 @@ public:
 	 */	
 	ENGINE_API virtual void PostInitViewCustomData(const FSceneView& InView, void* InViewCustomData) const { }
 
+	/** Tell us if this proxy is drawn in game.*/
+	ENGINE_API virtual bool IsDrawnInGame() const { return DrawInGame; }
+
 	/** Tell us if we should rely on the default LOD computing rules or not.*/
 	ENGINE_API virtual bool IsUsingCustomLODRules() const { return false; }
 	
diff --git a/Engine/Source/Runtime/Renderer/Private/DeferredShadingRenderer.cpp b/Engine/Source/Runtime/Renderer/Private/DeferredShadingRenderer.cpp
--- a/Engine/Source/Runtime/Renderer/Private/DeferredShadingRenderer.cpp
+++ b/Engine/Source/Runtime/Renderer/Private/DeferredShadingRenderer.cpp
@@ -638,9 +638,9 @@ bool FDeferredShadingSceneRenderer::GatherRayTracingWorldInstances(FRHICommandLi
 
 				//#dxr_todo The Raytracing codepath does not support Showflags since data moved to the SceneInfo. 
 				//Touching the SceneProxy to determine this would simply cost too much
-				if (SceneInfo->bShouldRenderInMainPass)
+				if (SceneInfo->bShouldRenderInMainPass && SceneInfo->bDrawInGame)
 				{
-					if (SceneInfo->bIsRayTracingStaticRelevant)
+					if (SceneInfo->bIsRayTracingStaticRelevant && View.Family->EngineShowFlags.StaticMeshes)
 					{
 						static const auto ICVarStaticMeshLODDistanceScale = IConsoleManager::Get().FindConsoleVariable(TEXT("r.StaticMeshLODDistanceScale"));
 						float LODScale = ICVarStaticMeshLODDistanceScale->GetFloat() * View.LODDistanceFactor;
@@ -721,7 +721,7 @@ bool FDeferredShadingSceneRenderer::GatherRayTracingWorldInstances(FRHICommandLi
 						ensure(DrawCmdIndexCopy != DrawCmdIndex);
 						View.RayTracingGeometryInstances.Add(FRayTracingGeometryInstance{ RayTracingGeometryInstance, Scene->PrimitiveTransforms[PrimitiveIndex], (uint32)PrimitiveIndex });
 					}
-					else
+					else if (View.Family->EngineShowFlags.SkeletalMeshes)
 					{
 						RayTracedMeshElementsMask |= 1 << ViewIndex;
 					}
diff --git a/Engine/Source/Runtime/Renderer/Private/PrimitiveSceneInfo.cpp b/Engine/Source/Runtime/Renderer/Private/PrimitiveSceneInfo.cpp
--- a/Engine/Source/Runtime/Renderer/Private/PrimitiveSceneInfo.cpp
+++ b/Engine/Source/Runtime/Renderer/Private/PrimitiveSceneInfo.cpp
@@ -114,6 +114,7 @@ FPrimitiveSceneInfo::FPrimitiveSceneInfo(UPrimitiveComponent* InComponent,FScene
 	bIsUsingCustomLODRules(Proxy->IsUsingCustomLODRules()),
 	bIsUsingCustomWholeSceneShadowLODRules(Proxy->IsUsingCustomWholeSceneShadowLODRules()),
 #if RHI_RAYTRACING
+	bDrawInGame(Proxy->IsDrawnInGame()),
 	bShouldRenderInMainPass(InComponent->SceneProxy->ShouldRenderInMainPass()),
 	bIsVisibleInReflectionCaptures(InComponent->SceneProxy->IsVisibleInReflectionCaptures()),
 	bIsRayTracingRelevant(InComponent->SceneProxy->IsRayTracingRelevant()),
diff --git a/Engine/Source/Runtime/Renderer/Public/PrimitiveSceneInfo.h b/Engine/Source/Runtime/Renderer/Public/PrimitiveSceneInfo.h
--- a/Engine/Source/Runtime/Renderer/Public/PrimitiveSceneInfo.h
+++ b/Engine/Source/Runtime/Renderer/Public/PrimitiveSceneInfo.h
@@ -283,6 +283,7 @@ public:
 	bool bIsUsingCustomWholeSceneShadowLODRules : 1;
 
 #if RHI_RAYTRACING
+	bool bDrawInGame : 1;
 	bool bShouldRenderInMainPass : 1;
 	bool bIsVisibleInReflectionCaptures : 1;
 	bool bIsRayTracingRelevant : 1;
