Fixed deadlock inside primitive VB update -  need to wait on GetRenderThread_Local, as mesh pass setup task can wait on rendering thread inside InitResourceFromPossiblyParallelRendering().

diff --git a/Engine/Source/Runtime/Renderer/Private/MeshDrawCommands.cpp b/Engine/Source/Runtime/Renderer/Private/MeshDrawCommands.cpp
--- a/Engine/Source/Runtime/Renderer/Private/MeshDrawCommands.cpp
+++ b/Engine/Source/Runtime/Renderer/Private/MeshDrawCommands.cpp
@@ -1027,7 +1027,8 @@ struct FRHICommandUpdatePrimitiveIdBuffer : public FRHICommand<FRHICommandUpdate
 
 	void Execute(FRHICommandListBase& CmdList)
 	{
-		FTaskGraphInterface::Get().WaitUntilTasksComplete(Prereqs, IsRunningRHIInSeparateThread() ? ENamedThreads::RHIThread : ENamedThreads::ActualRenderingThread);
+		// Need to wait on GetRenderThread_Local, as mesh pass setup task can wait on rendering thread inside InitResourceFromPossiblyParallelRendering().
+		FTaskGraphInterface::Get().WaitUntilTasksComplete(Prereqs, IsRunningRHIInSeparateThread() ? ENamedThreads::RHIThread : ENamedThreads::GetRenderThread_Local());
 
 		// Upload vertex buffer data.
 		void* RESTRICT Data = (void* RESTRICT)GDynamicRHI->RHILockVertexBuffer(VertexBuffer, 0, VertexBufferDataSize, RLM_WriteOnly);
