Use FinalPostProcessSettings.ScreenSpaceReflectionMaxRoughness to drive ray traced reflection roughness threshold instead of a custom CVar. Added roughness-based fade for ray traced reflections (same as for SSR).

diff --git a/Engine/Shaders/Private/RayTracing/RayTracingReflections.usf b/Engine/Shaders/Private/RayTracing/RayTracingReflections.usf
--- a/Engine/Shaders/Private/RayTracing/RayTracingReflections.usf
+++ b/Engine/Shaders/Private/RayTracing/RayTracingReflections.usf
@@ -41,6 +41,12 @@ RaytracingAccelerationStructure TLAS;
 RWTexture2D<float4> ColorOutput;
 RWTexture2D<float> RayHitDistanceOutput;
 
+float GetRoughnessFade(float Roughness, float MaxRoughness)
+{
+	float RoughnessMaskScale = -2.0 / MaxRoughness;
+	return saturate(Roughness * RoughnessMaskScale + 2.0);
+}
+
 uint CalcLinearIndex(uint2 PixelCoord)
 {
 	return PixelCoord.y * View.BufferSizeAndInvSize.x + PixelCoord.x;
@@ -333,7 +339,9 @@ void RayTracingReflectionsRGS()
 	float ReflectedWeight = 0.0;
 	float HitDistance = -2.0;
 
-	if (Roughness < ReflectionMaxRoughness)
+	const float RoughnessFade = GetRoughnessFade(Roughness, ReflectionMaxRoughness);
+
+	if (RoughnessFade > 0)
 	{
 		for (uint SampleIndex = 0; SampleIndex < LocalSamplesPerPixel; ++SampleIndex)
 		{
@@ -410,6 +418,8 @@ void RayTracingReflectionsRGS()
 		ReflectedColor *= rcp(ReflectedWeight);
 	}
 
+	ReflectedColor *= RoughnessFade;
+
 	ColorOutput[PixelCoord] = float4(ReflectedColor);
 	RayHitDistanceOutput[PixelCoord] = HitDistance;
 }
diff --git a/Engine/Source/Runtime/Renderer/Private/RayTracing/RayTracingReflections.cpp b/Engine/Source/Runtime/Renderer/Private/RayTracing/RayTracingReflections.cpp
--- a/Engine/Source/Runtime/Renderer/Private/RayTracing/RayTracingReflections.cpp
+++ b/Engine/Source/Runtime/Renderer/Private/RayTracing/RayTracingReflections.cpp
@@ -44,13 +44,6 @@ static FAutoConsoleVariableRef CVarRayTracingReflectionsMaxRayDistance(
 	TEXT("Sets the maximum ray distance for ray traced reflection rays (default = 1.0e27)")
 );
 
-static float GRayTracingReflectionsMaxRoughness = 1.0;
-static FAutoConsoleVariableRef CVarRayTracingReflectionsMaxRoughness(
-	TEXT("r.RayTracing.Reflections.MaxRoughness"),
-	GRayTracingReflectionsMaxRoughness,
-	TEXT("Surfaces with roughness larger than this value will not receive ray traced reflections (default = 1)")
-);
-
 static const int32 GReflectionLightCountMaximum = 64;
 
 BEGIN_GLOBAL_SHADER_PARAMETER_STRUCT(FReflectionsLightData, )
@@ -218,7 +211,7 @@ void FDeferredShadingSceneRenderer::RayTraceReflections(
 	PassParameters->ShouldDoReflectedShadows = GRayTracingReflectionsShadows;
 	PassParameters->ShouldDoEmissiveAndIndirectLighting = GRayTracingReflectionsEmissiveAndIndirectLighting;
 	PassParameters->ReflectionMaxRayDistance = GRayTracingReflectionsMaxRayDistance;
-	PassParameters->ReflectionMaxRoughness = GRayTracingReflectionsMaxRoughness;
+	PassParameters->ReflectionMaxRoughness = FMath::Clamp(View.FinalPostProcessSettings.ScreenSpaceReflectionMaxRoughness, 0.01f, 1.0f);
 	PassParameters->LTCMatTexture = GSystemTextures.LTCMat->GetRenderTargetItem().ShaderResourceTexture;
 	PassParameters->LTCMatSampler = TStaticSamplerState<SF_Bilinear, AM_Clamp, AM_Clamp, AM_Clamp>::GetRHI();
 	PassParameters->LTCAmpTexture = GSystemTextures.LTCAmp->GetRenderTargetItem().ShaderResourceTexture;
