Make sure to account for Metal multiplying one-pass-point-light shadow draw calls so we have one for each face in the absence of geometry shaders.

diff --git a/Engine/Source/Runtime/Renderer/Private/ShadowSetup.cpp b/Engine/Source/Runtime/Renderer/Private/ShadowSetup.cpp
--- a/Engine/Source/Runtime/Renderer/Private/ShadowSetup.cpp
+++ b/Engine/Source/Runtime/Renderer/Private/ShadowSetup.cpp
@@ -1165,6 +1165,8 @@ void FProjectedShadowInfo::SetupMeshDrawCommandsForShadowDepth(FSceneRenderer& R
 		GetShadowTypeNameForDrawEvent(PassNameForStats);
 		ShadowDepthPass.SetDumpInstancingStats(TEXT("ShadowDepth ") + PassNameForStats);
 	}
+	
+	const uint32 InstanceFactor = !GetShadowDepthType().bOnePassPointLightShadow || RHISupportsGeometryShaders(Renderer.Scene->GetShaderPlatform()) ? 1 : 6;
 
 	ShadowDepthPass.DispatchPassSetup(
 		Renderer.Scene,
@@ -1173,9 +1175,9 @@ void FProjectedShadowInfo::SetupMeshDrawCommandsForShadowDepth(FSceneRenderer& R
 		FExclusiveDepthStencil::DepthNop_StencilNop,
 		MeshPassProcessor,
 		DynamicSubjectMeshElements,
-		NumDynamicSubjectMeshElements,
+		NumDynamicSubjectMeshElements * InstanceFactor,
 		SubjectMeshCommandBuildRequests,
-		NumSubjectMeshCommandBuildRequestElements,
+		NumSubjectMeshCommandBuildRequestElements * InstanceFactor,
 		ShadowDepthPassVisibleCommands);
 }
 
