Keeps ray tracing cvar r.RayTracing.*

diff --git a/Engine/Source/Runtime/Renderer/Private/ReflectionEnvironment.cpp b/Engine/Source/Runtime/Renderer/Private/ReflectionEnvironment.cpp
--- a/Engine/Source/Runtime/Renderer/Private/ReflectionEnvironment.cpp
+++ b/Engine/Source/Runtime/Renderer/Private/ReflectionEnvironment.cpp
@@ -103,24 +103,23 @@ static TAutoConsoleVariable<float> CVarSkySpecularOcclusionStrength(
 
 static int32 GRayTracingReflections = 1;
 static FAutoConsoleVariableRef CVarReflectionsMethod(
-	TEXT("r.Reflections.Method"),
+	TEXT("r.RayTracing.Reflections"),
 	GRayTracingReflections,
 	TEXT("0: use traditional rasterized SSR\n")
 	TEXT("1: use ray traced reflections (default when ray tracing is enabled)\n")
 );
 
 static TAutoConsoleVariable<float> CVarReflectionScreenPercentage(
-	TEXT("r.Reflections.RayTracing.ScreenPercentage"),
+	TEXT("r.RayTracing.Reflections.ScreenPercentage"),
 	100.0f,
-	TEXT("Try to guess"),
+	TEXT("Screen percentage the reflections should be ray traced at (default = 100)."),
 	ECVF_RenderThreadSafe);
 
 static int32 GRayTracingReflectionsSamplesPerPixel = 1;
 static FAutoConsoleVariableRef CVarRayTracingReflectionsSamplesPerPixel(
-	TEXT("r.Reflections.RayTracing.SamplesPerPixel"),
+	TEXT("r.RayTracing.Reflections.SamplesPerPixel"),
 	GRayTracingReflectionsSamplesPerPixel,
-	TEXT("Sets the samples-per-pixel for reflections (default = 1)")
-);
+	TEXT("Sets the samples-per-pixel for reflections (default = 1)"));
 
 static TAutoConsoleVariable<int32> CVarUseReflectionDenoiser(
 	TEXT("r.Reflections.Denoiser"),
