Addresses in code a warning for the reflection's hit distance when denoiser is disabled.

First case where it is not a temporal history case witch may be a sign this warning might be one to be disabled by default. Will see how it scale or not as we convert more pass to render graph.

diff --git a/Engine/Source/Runtime/Renderer/Private/ReflectionEnvironment.cpp b/Engine/Source/Runtime/Renderer/Private/ReflectionEnvironment.cpp
--- a/Engine/Source/Runtime/Renderer/Private/ReflectionEnvironment.cpp
+++ b/Engine/Source/Runtime/Renderer/Private/ReflectionEnvironment.cpp
@@ -792,6 +792,12 @@ void FDeferredShadingSceneRenderer::RenderDeferredReflectionsAndSkyLighting(FRHI
 			}
 			else
 			{
+				// The performance of ray tracing does not allow to run without a denoiser in real time.
+				// Multiple rays per pixel is unsupported by the denoiser that will most likely more bound by to
+				// many rays than exporting the hit distance buffer. Therefore no permutation of the ray generation
+				// shader has been judged required to be supported.
+				GraphBuilder.RemoveUnusedTextureWarning(DenoiserInputs.RayHitDistance);
+
 				GraphBuilder.QueueTextureExtraction(DenoiserInputs.Color, &ReflectionsColor);
 			}
 
