Fixed mesh draw command threading crashes on engine start. GRenderingThread is required to safely use rendering resources in parallel.

diff --git a/Engine/Source/Runtime/Renderer/Private/MeshDrawCommands.cpp b/Engine/Source/Runtime/Renderer/Private/MeshDrawCommands.cpp
--- a/Engine/Source/Runtime/Renderer/Private/MeshDrawCommands.cpp
+++ b/Engine/Source/Runtime/Renderer/Private/MeshDrawCommands.cpp
@@ -897,7 +897,9 @@ void FParallelMeshDrawCommandPass::DispatchPassSetup(
 	TaskContext.MeshDrawCommands.Reserve(MaxNumDraws);
 	TaskContext.TempVisibleMeshDrawCommands.Reserve(MaxNumDraws);
 
-	const bool bExecuteInParallel = FApp::ShouldUseThreadingForPerformance() && CVarMeshDrawCommandsParallelPassSetup.GetValueOnRenderThread() > 0;
+	const bool bExecuteInParallel = FApp::ShouldUseThreadingForPerformance()
+		&& CVarMeshDrawCommandsParallelPassSetup.GetValueOnRenderThread() > 0
+		&& GRenderingThread; // Rendering thread is required to safely use rendering resources in parallel.
 
 	if (bExecuteInParallel)
 	{
