Tightening IsInParallelRenderingThread condition - make sure that we are on a game thread if there is no rendering thread. This catches issues when e.g. initializing rendering resources from tasks while there is no rendering thread. In that case IsInRenderingThread returns true and all code silently runs as if it would single threaded.

diff --git a/Engine/Source/Runtime/Core/Private/HAL/ThreadingBase.cpp b/Engine/Source/Runtime/Core/Private/HAL/ThreadingBase.cpp
--- a/Engine/Source/Runtime/Core/Private/HAL/ThreadingBase.cpp
+++ b/Engine/Source/Runtime/Core/Private/HAL/ThreadingBase.cpp
@@ -55,7 +55,14 @@ CORE_API bool IsInRenderingThread()
 
 CORE_API bool IsInParallelRenderingThread()
 {
-	return !GRenderingThread || GIsRenderingThreadSuspended.Load(EMemoryOrder::Relaxed) || (FPlatformTLS::GetCurrentThreadId() != GGameThreadId);
+	if (!GRenderingThread || GIsRenderingThreadSuspended.Load(EMemoryOrder::Relaxed))
+	{
+		return FPlatformTLS::GetCurrentThreadId() == GGameThreadId;
+	}
+	else
+	{
+		return FPlatformTLS::GetCurrentThreadId() != GGameThreadId;
+	}
 }
 
 CORE_API uint32 GRHIThreadId = 0;
diff --git a/Engine/Source/Runtime/Renderer/Private/SceneVisibility.cpp b/Engine/Source/Runtime/Renderer/Private/SceneVisibility.cpp
--- a/Engine/Source/Runtime/Renderer/Private/SceneVisibility.cpp
+++ b/Engine/Source/Runtime/Renderer/Private/SceneVisibility.cpp
@@ -478,7 +478,7 @@ static int32 FrustumCull(const FScene* Scene, FViewInfo& View)
 				}
 			}
 		},
-		!FApp::ShouldUseThreadingForPerformance() || (UseCustomCulling && !View.CustomVisibilityQuery->IsThreadsafe()) || CVarParallelInitViews.GetValueOnRenderThread() == 0
+		!FApp::ShouldUseThreadingForPerformance() || (UseCustomCulling && !View.CustomVisibilityQuery->IsThreadsafe()) || CVarParallelInitViews.GetValueOnRenderThread() == 0 || !IsInActualRenderingThread()
 	);
 
 	return NumCulledPrimitives.GetValue();
@@ -2396,7 +2396,7 @@ static void ComputeAndMarkRelevanceForViewParallel(
 	TArray<FRelevancePacket*,SceneRenderingAllocator> Packets;
 	Packets.Reserve(EstimateOfNumPackets);
 
-	bool WillExecuteInParallel = FApp::ShouldUseThreadingForPerformance() && CVarParallelInitViews.GetValueOnRenderThread() > 0;
+	bool WillExecuteInParallel = FApp::ShouldUseThreadingForPerformance() && CVarParallelInitViews.GetValueOnRenderThread() > 0 && IsInActualRenderingThread();
 
 	{
 		FSceneSetBitIterator BitIt(View.PrimitiveVisibilityMap);
