DR - Fix crash exiting QAGame due to RHI objects not getting released
#rnx

diff --git a/Engine/Source/Runtime/Renderer/Private/MeshDrawCommands.cpp b/Engine/Source/Runtime/Renderer/Private/MeshDrawCommands.cpp
--- a/Engine/Source/Runtime/Renderer/Private/MeshDrawCommands.cpp
+++ b/Engine/Source/Runtime/Renderer/Private/MeshDrawCommands.cpp
@@ -9,7 +9,7 @@ MeshDrawCommandSetup.cpp: Mesh draw command setup.
 #include "ScenePrivate.h"
 #include "TranslucentRendering.h"
 
-FPrimitiveIdVertexBufferPool GPrimitiveIdVertexBufferPool;
+TGlobalResource<FPrimitiveIdVertexBufferPool> GPrimitiveIdVertexBufferPool;
 
 static TAutoConsoleVariable<int32> CVarMeshDrawCommandsParallelPassSetup(
 	TEXT("r.MeshDrawCommands.ParallelPassSetup"),
@@ -24,7 +24,7 @@ FPrimitiveIdVertexBufferPool::FPrimitiveIdVertexBufferPool()
 
 FPrimitiveIdVertexBufferPool::~FPrimitiveIdVertexBufferPool()
 {
-	Entries.Empty();
+	check(!Entries.Num());
 }
 
 FVertexBufferRHIParamRef FPrimitiveIdVertexBufferPool::Allocate(int32 BufferSize)
@@ -70,6 +70,12 @@ void FPrimitiveIdVertexBufferPool::DiscardAll()
 	}
 }
 
+void FPrimitiveIdVertexBufferPool::ReleaseDynamicRHI()
+{
+	Entries.Empty();
+}
+
+
 struct FCompareFMeshDrawCommands
 {
 	FORCEINLINE bool operator() (FVisibleMeshDrawCommand A, FVisibleMeshDrawCommand B) const
diff --git a/Engine/Source/Runtime/Renderer/Private/MeshDrawCommands.h b/Engine/Source/Runtime/Renderer/Private/MeshDrawCommands.h
--- a/Engine/Source/Runtime/Renderer/Private/MeshDrawCommands.h
+++ b/Engine/Source/Runtime/Renderer/Private/MeshDrawCommands.h
@@ -16,7 +16,7 @@ class FParallelCommandListSet;
 /**
  * Global vertex buffer pool used for GPUScene primitive id arrays.
  */
-class FPrimitiveIdVertexBufferPool
+class FPrimitiveIdVertexBufferPool : public FRenderResource
 {
 public:
 	FPrimitiveIdVertexBufferPool();
@@ -25,6 +25,8 @@ public:
 	FVertexBufferRHIParamRef Allocate(int32 BufferSize);
 	void DiscardAll();
 
+	virtual void ReleaseDynamicRHI() override;
+
 private:
 	struct FPrimitiveIdVertexBufferPoolEntry
 	{
@@ -37,7 +39,7 @@ private:
 	TArray<FPrimitiveIdVertexBufferPoolEntry> Entries;
 };
 
-extern FPrimitiveIdVertexBufferPool GPrimitiveIdVertexBufferPool;
+extern TGlobalResource<FPrimitiveIdVertexBufferPool> GPrimitiveIdVertexBufferPool;
 
 /**	
  * Parallel mesh draw command pass setup task context.
