Fix wireframe rendering issue when parallel rendering is enabled.

The underlying issue was that the wrong buffers (GBuffer and single samplle depth, instead of MSAA-ed color & depth) were bound when wireframe rendering was enabled in parallel rendering mode.

diff --git a/Engine/Source/Runtime/Renderer/Private/BasePassRendering.cpp b/Engine/Source/Runtime/Renderer/Private/BasePassRendering.cpp
--- a/Engine/Source/Runtime/Renderer/Private/BasePassRendering.cpp
+++ b/Engine/Source/Runtime/Renderer/Private/BasePassRendering.cpp
@@ -748,7 +748,7 @@ public:
 	virtual void SetStateOnCommandList(FRHICommandList& CmdList) override
 	{
 		FParallelCommandListSet::SetStateOnCommandList(CmdList);
-		FSceneRenderTargets::Get(CmdList).BeginRenderingGBuffer(CmdList, ERenderTargetLoadAction::ELoad, ERenderTargetLoadAction::ELoad, BasePassDepthStencilAccess, SceneRenderer->ViewFamily.EngineShowFlags.ShaderComplexity);
+		FSceneRenderTargets::Get(CmdList).BeginRenderingGBuffer(CmdList, ERenderTargetLoadAction::ELoad, ERenderTargetLoadAction::ELoad, BasePassDepthStencilAccess, SceneRenderer->ViewFamily.EngineShowFlags.ShaderComplexity, false, FLinearColor(0, 0, 0, 1), SceneRenderer->ViewFamily.EngineShowFlags.Wireframe);
 		SetupBasePassView(CmdList, View, SceneRenderer);
 	}
 };
diff --git a/Engine/Source/Runtime/Renderer/Private/DeferredShadingRenderer.cpp b/Engine/Source/Runtime/Renderer/Private/DeferredShadingRenderer.cpp
--- a/Engine/Source/Runtime/Renderer/Private/DeferredShadingRenderer.cpp
+++ b/Engine/Source/Runtime/Renderer/Private/DeferredShadingRenderer.cpp
@@ -1474,18 +1474,7 @@ void FDeferredShadingSceneRenderer::Render(FRHICommandListImmediate& RHICmdList)
 	// If we ran parallel in the basepass there will be no renderpass at this point.
 	if (bDoParallelBasePass && !bRenderLightmapDensity)
 	{
-		if (bIsWireframeRenderpass)
-		{
-			FRHIRenderPassInfo RPInfo(SceneContext.GetEditorPrimitivesColor(RHICmdList), ERenderTargetActions::Clear_Store);
-			RPInfo.DepthStencilRenderTarget.Action = EDepthStencilTargetActions::ClearDepthStencil_StoreDepthStencil;
-			RPInfo.DepthStencilRenderTarget.DepthStencilTarget = SceneContext.GetEditorPrimitivesDepth(RHICmdList);
-			RPInfo.DepthStencilRenderTarget.ExclusiveDepthStencil = FExclusiveDepthStencil::DepthWrite_StencilWrite;
-			RHICmdList.BeginRenderPass(RPInfo, TEXT("Wireframe"));
-		}
-		else
-		{
-			SceneContext.BeginRenderingGBuffer(RHICmdList, ERenderTargetLoadAction::ELoad, ERenderTargetLoadAction::ELoad, BasePassDepthStencilAccess, ViewFamily.EngineShowFlags.ShaderComplexity);
-		}
+		SceneContext.BeginRenderingGBuffer(RHICmdList, ERenderTargetLoadAction::ELoad, ERenderTargetLoadAction::ELoad, BasePassDepthStencilAccess, ViewFamily.EngineShowFlags.ShaderComplexity);
 	}
 
 	{
diff --git a/Engine/Source/Runtime/Renderer/Private/PostProcess/SceneRenderTargets.cpp b/Engine/Source/Runtime/Renderer/Private/PostProcess/SceneRenderTargets.cpp
--- a/Engine/Source/Runtime/Renderer/Private/PostProcess/SceneRenderTargets.cpp
+++ b/Engine/Source/Runtime/Renderer/Private/PostProcess/SceneRenderTargets.cpp
@@ -779,7 +779,7 @@ void FSceneRenderTargets::SetQuadOverdrawUAV(FRHICommandList& RHICmdList, bool b
 	}
 }
 
-void FSceneRenderTargets::BeginRenderingGBuffer(FRHICommandList& RHICmdList, ERenderTargetLoadAction ColorLoadAction, ERenderTargetLoadAction DepthLoadAction, FExclusiveDepthStencil::Type DepthStencilAccess, bool bBindQuadOverdrawBuffers, bool bClearQuadOverdrawBuffers, const FLinearColor& ClearColor/*=(0,0,0,1)*/)
+void FSceneRenderTargets::BeginRenderingGBuffer(FRHICommandList& RHICmdList, ERenderTargetLoadAction ColorLoadAction, ERenderTargetLoadAction DepthLoadAction, FExclusiveDepthStencil::Type DepthStencilAccess, bool bBindQuadOverdrawBuffers, bool bClearQuadOverdrawBuffers, const FLinearColor& ClearColor/*=(0,0,0,1)*/, bool bIsWireframe)
 {
 	check(RHICmdList.IsOutsideRenderPass());
 
@@ -872,6 +872,26 @@ void FSceneRenderTargets::BeginRenderingGBuffer(FRHICommandList& RHICmdList, ERe
 	bool bBindClearColor = !bClearColor && bGBuffersFastCleared;
 	bool bBindClearDepth = !bClearDepth && bSceneDepthCleared;
 	RHICmdList.BindClearMRTValues(bBindClearColor, bBindClearDepth, bBindClearDepth);
+
+	if (bIsWireframe)
+	{
+		RHICmdList.EndRenderPass();
+
+		FSceneRenderTargets& SceneContext = FSceneRenderTargets::Get(RHICmdList);
+
+		RPInfo = FRHIRenderPassInfo();
+
+		RPInfo.ColorRenderTargets[0].Action = MakeRenderTargetActions(ColorLoadAction, ERenderTargetStoreAction::EStore);
+		RPInfo.ColorRenderTargets[0].RenderTarget = SceneContext.GetEditorPrimitivesColor(RHICmdList);
+		RPInfo.ColorRenderTargets[0].ArraySlice = -1;
+		RPInfo.ColorRenderTargets[0].MipIndex = 0;
+
+		//FRHIRenderPassInfo RPInfo(SceneContext.GetEditorPrimitivesColor(RHICmdList), ColorLoadAction);
+		RPInfo.DepthStencilRenderTarget.Action = MakeDepthStencilTargetActions(MakeRenderTargetActions(DepthLoadAction, ERenderTargetStoreAction::EStore), MakeRenderTargetActions(DepthLoadAction, ERenderTargetStoreAction::EStore));
+		RPInfo.DepthStencilRenderTarget.DepthStencilTarget = SceneContext.GetEditorPrimitivesDepth(RHICmdList);
+		RPInfo.DepthStencilRenderTarget.ExclusiveDepthStencil = FExclusiveDepthStencil::DepthWrite_StencilWrite;
+		RHICmdList.BeginRenderPass(RPInfo, TEXT("Wireframe"));	
+	}
 }
 
 void FSceneRenderTargets::FinishGBufferPassAndResolve(FRHICommandListImmediate& RHICmdList)
diff --git a/Engine/Source/Runtime/Renderer/Private/PostProcess/SceneRenderTargets.h b/Engine/Source/Runtime/Renderer/Private/PostProcess/SceneRenderTargets.h
--- a/Engine/Source/Runtime/Renderer/Private/PostProcess/SceneRenderTargets.h
+++ b/Engine/Source/Runtime/Renderer/Private/PostProcess/SceneRenderTargets.h
@@ -212,7 +212,7 @@ public:
 
 	void SetQuadOverdrawUAV(FRHICommandList& RHICmdList, bool bBindQuadOverdrawBuffers, FRHISetRenderTargetsInfo& Info);
 	void SetQuadOverdrawUAV(FRHICommandList& RHICmdList, bool bBindQuadOverdrawBuffers, bool bClearQuadOverdrawBuffers, FRHIRenderPassInfo& Info);
-	void BeginRenderingGBuffer(FRHICommandList& RHICmdList, ERenderTargetLoadAction ColorLoadAction, ERenderTargetLoadAction DepthLoadAction, FExclusiveDepthStencil::Type DepthStencilAccess, bool bBindQuadOverdrawBuffers, bool bClearQuadOverdrawBuffers = false, const FLinearColor& ClearColor = FLinearColor(0, 0, 0, 1));
+	void BeginRenderingGBuffer(FRHICommandList& RHICmdList, ERenderTargetLoadAction ColorLoadAction, ERenderTargetLoadAction DepthLoadAction, FExclusiveDepthStencil::Type DepthStencilAccess, bool bBindQuadOverdrawBuffers, bool bClearQuadOverdrawBuffers = false, const FLinearColor& ClearColor = FLinearColor(0, 0, 0, 1), bool bIsWireframe=false);
 	void FinishGBufferPassAndResolve(FRHICommandListImmediate& RHICmdList);
 
 	/**
