Fixed mesh draw command crash when dynamic instancing is disabled.

diff --git a/Engine/Source/Runtime/Renderer/Private/MeshDrawCommands.cpp b/Engine/Source/Runtime/Renderer/Private/MeshDrawCommands.cpp
--- a/Engine/Source/Runtime/Renderer/Private/MeshDrawCommands.cpp
+++ b/Engine/Source/Runtime/Renderer/Private/MeshDrawCommands.cpp
@@ -303,7 +303,7 @@ void BuildMeshDrawCommandPrimitiveIdBuffer(
 	)
 {
 	QUICK_SCOPE_CYCLE_COUNTER(STAT_BuildMeshDrawCommandPrimitiveIdBuffer);
-	check(VisibleMeshDrawCommands.Num() <= TempVisibleMeshDrawCommands.Max() && TempVisibleMeshDrawCommands.Num() == 0 && PrimitiveIdData);
+	check(PrimitiveIdData && PrimitiveIdDataSize > 0);
 
 	const FVisibleMeshDrawCommand* RESTRICT PassVisibleMeshDrawCommands = VisibleMeshDrawCommands.GetData();
 	const int32 NumDrawCommands = VisibleMeshDrawCommands.Num();
@@ -315,6 +315,7 @@ void BuildMeshDrawCommandPrimitiveIdBuffer(
 	if (IsDynamicInstancingEnabled())
 	{
 		QUICK_SCOPE_CYCLE_COUNTER(STAT_DynamicInstancingOfVisibleMeshDrawCommands);
+		check(VisibleMeshDrawCommands.Num() <= TempVisibleMeshDrawCommands.Max() && TempVisibleMeshDrawCommands.Num() == 0);
 
 		int32 CurrentStateBucketId = -1;
 		uint32* RESTRICT CurrentDynamicallyInstancedMeshCommandNumInstances = nullptr;
@@ -894,12 +895,9 @@ void FParallelMeshDrawCommandPass::DispatchPassSetup(
 	}
 
 	// Preallocate resources on rendering thread based on MaxNumDraws.
-	if (TaskContext.bDynamicInstancing)
-	{
-		TaskContext.PrimitiveIdBufferDataSize = TaskContext.MaxInstanceFactor * MaxNumDraws * sizeof(int32);
-		TaskContext.PrimitiveIdBufferData = FMemStack::Get().Alloc(TaskContext.PrimitiveIdBufferDataSize, 8);
-		PrimitiveIdVertexBufferRHI = GPrimitiveIdVertexBufferPool.Allocate(TaskContext.PrimitiveIdBufferDataSize);
-	}
+	TaskContext.PrimitiveIdBufferDataSize = TaskContext.MaxInstanceFactor * MaxNumDraws * sizeof(int32);
+	TaskContext.PrimitiveIdBufferData = FMemStack::Get().Alloc(TaskContext.PrimitiveIdBufferDataSize, 8);
+	PrimitiveIdVertexBufferRHI = GPrimitiveIdVertexBufferPool.Allocate(TaskContext.PrimitiveIdBufferDataSize);
 	TaskContext.MeshDrawCommands.Reserve(MaxNumDraws);
 	TaskContext.TempVisibleMeshDrawCommands.Reserve(MaxNumDraws);
 
