Added missing #pragma once

diff --git a/Engine/Source/Runtime/Renderer/Private/EditorPrimitivesRendering.h b/Engine/Source/Runtime/Renderer/Private/EditorPrimitivesRendering.h
--- a/Engine/Source/Runtime/Renderer/Private/EditorPrimitivesRendering.h
+++ b/Engine/Source/Runtime/Renderer/Private/EditorPrimitivesRendering.h
@@ -3,6 +3,9 @@
 /*=============================================================================
 	EditorPrimitivesRendering
 =============================================================================*/
+
+#pragma once
+
 #include "MeshPassProcessor.h"
 
 class FEditorPrimitivesBasePassMeshProcessor : public FMeshPassProcessor
