Removed VisibleDynamicPrimitives, as this array was only used for counting visible primitives.

diff --git a/Engine/Source/Runtime/Renderer/Private/LightRendering.cpp b/Engine/Source/Runtime/Renderer/Private/LightRendering.cpp
--- a/Engine/Source/Runtime/Renderer/Private/LightRendering.cpp
+++ b/Engine/Source/Runtime/Renderer/Private/LightRendering.cpp
@@ -417,17 +417,12 @@ void FSceneRenderer::GatherSimpleLights(const FSceneViewFamily& ViewFamily, cons
 	for (int32 ViewIndex = 0; ViewIndex < Views.Num(); ViewIndex++)
 	{
 		const FViewInfo& View = Views[ViewIndex];
-		for (int32 PrimitiveIndex = 0; PrimitiveIndex < View.VisibleDynamicPrimitives.Num(); PrimitiveIndex++)
+		for (int32 PrimitiveIndex = 0; PrimitiveIndex < View.VisibleDynamicPrimitivesWithSimpleLights.Num(); PrimitiveIndex++)
 		{
-			const FPrimitiveSceneInfo* PrimitiveSceneInfo = View.VisibleDynamicPrimitives[PrimitiveIndex];
-			const int32 PrimitiveId = PrimitiveSceneInfo->GetIndex();
-			const FPrimitiveViewRelevance& PrimitiveViewRelevance = View.PrimitiveViewRelevanceMap[PrimitiveId];
+			const FPrimitiveSceneInfo* PrimitiveSceneInfo = View.VisibleDynamicPrimitivesWithSimpleLights[PrimitiveIndex];
 
-			if (PrimitiveViewRelevance.bHasSimpleLights)
-			{
-				// TArray::AddUnique is slow, but not expecting many entries in PrimitivesWithSimpleLights
-				PrimitivesWithSimpleLights.AddUnique(PrimitiveSceneInfo);
-			}
+			// TArray::AddUnique is slow, but not expecting many entries in PrimitivesWithSimpleLights
+			PrimitivesWithSimpleLights.AddUnique(PrimitiveSceneInfo);
 		}
 	}
 
diff --git a/Engine/Source/Runtime/Renderer/Private/SceneRendering.cpp b/Engine/Source/Runtime/Renderer/Private/SceneRendering.cpp
--- a/Engine/Source/Runtime/Renderer/Private/SceneRendering.cpp
+++ b/Engine/Source/Runtime/Renderer/Private/SceneRendering.cpp
@@ -896,6 +896,8 @@ void FViewInfo::Init()
 	{
 		NumVisibleDynamicMeshElements[PassIndex] = 0;
 	}
+
+	NumVisibleDynamicPrimitives = 0;
 }
 
 FViewInfo::~FViewInfo()
@@ -2587,7 +2589,7 @@ void FSceneRenderer::RenderFinish(FRHICommandListImmediate& RHICmdList)
 	{
 		const FViewInfo& View = Views[ViewIndex];
 		INC_DWORD_STAT_BY(STAT_VisibleStaticMeshElements,View.NumVisibleStaticMeshElements);
-		INC_DWORD_STAT_BY(STAT_VisibleDynamicPrimitives,View.VisibleDynamicPrimitives.Num());
+		INC_DWORD_STAT_BY(STAT_VisibleDynamicPrimitives,View.NumVisibleDynamicPrimitives);
 
 #if !(UE_BUILD_SHIPPING || UE_BUILD_TEST)
 		// update freezing info
diff --git a/Engine/Source/Runtime/Renderer/Private/SceneRendering.h b/Engine/Source/Runtime/Renderer/Private/SceneRendering.h
--- a/Engine/Source/Runtime/Renderer/Private/SceneRendering.h
+++ b/Engine/Source/Runtime/Renderer/Private/SceneRendering.h
@@ -896,8 +896,11 @@ public:
 	 set visible in StaticMeshVisibilityMap. */
 	TArray<uint64,SceneRenderingAllocator> StaticMeshBatchVisibility;
 
-	/** The dynamic primitives visible in this view. */
-	TArray<FPrimitiveSceneInfo*,SceneRenderingAllocator> VisibleDynamicPrimitives;
+	/** The dynamic primitives with simple lights visible in this view. */
+	TArray<FPrimitiveSceneInfo*, SceneRenderingAllocator> VisibleDynamicPrimitivesWithSimpleLights;
+
+	/** Number of dynamic primitives visible in this view. */
+	int32 NumVisibleDynamicPrimitives;
 
 	/** Marks used mesh passes by visible dynamic meshes. */
 	FMeshPassMask VisibleDynamicMeshesPassMask;
diff --git a/Engine/Source/Runtime/Renderer/Private/SceneVisibility.cpp b/Engine/Source/Runtime/Renderer/Private/SceneVisibility.cpp
--- a/Engine/Source/Runtime/Renderer/Private/SceneVisibility.cpp
+++ b/Engine/Source/Runtime/Renderer/Private/SceneVisibility.cpp
@@ -1747,7 +1747,8 @@ struct FRelevancePacket
 	FRelevancePrimSet<int32> RelevantStaticPrimitives;
 	FRelevancePrimSet<int32> NotDrawRelevant;
 	FRelevancePrimSet<int32> TranslucentSelfShadowPrimitives;
-	FRelevancePrimSet<FPrimitiveSceneInfo*> VisibleDynamicPrimitives;
+	FRelevancePrimSet<FPrimitiveSceneInfo*> VisibleDynamicPrimitivesWithSimpleLights;
+	int32 NumVisibleDynamicPrimitives;
 	FMeshPassMask VisibleDynamicMeshesPassMask;
 	FRelevancePrimSet<FTranslucentPrimSet::FTranslucentPrim, ETranslucencyPass::TPT_MAX> TranslucencyPrims;
 	// belongs to TranslucencyPrims
@@ -1829,6 +1830,7 @@ struct FRelevancePacket
 		, OutHasDynamicMeshElementsMasks(InOutHasDynamicMeshElementsMasks)
 		, OutHasDynamicEditorMeshElementsMasks(InOutHasDynamicEditorMeshElementsMasks)
 		, MarkMasks(InMarkMasks)
+		, NumVisibleDynamicPrimitives(0)
 		, bHasDistortionPrimitives(false)
 		, bHasCustomDepthPrimitives(false)
 		, PrimitiveCustomDataMemStack(InPrimitiveCustomDataMemStack)
@@ -1993,9 +1995,14 @@ struct FRelevancePacket
 			else if(bDynamicRelevance)
 			{
 				// Keep track of visible dynamic primitives.
-				VisibleDynamicPrimitives.AddPrim(PrimitiveSceneInfo);
+				++NumVisibleDynamicPrimitives;
 				OutHasDynamicMeshElementsMasks[BitIndex] |= ViewBit;
 				UpdateDynamicPrimMask(ShadingPath, bAddLightmapDensityCommands, ViewRelevance);
+
+				if (ViewRelevance.bHasSimpleLights)
+				{
+					VisibleDynamicPrimitivesWithSimpleLights.AddPrim(PrimitiveSceneInfo);
+				}
 			}
 
 			if (ViewRelevance.bUseCustomViewData)
@@ -2323,7 +2330,8 @@ struct FRelevancePacket
 		WriteView.bTranslucentSurfaceLighting |= bTranslucentSurfaceLighting;
 		WriteView.bUsesSceneDepth |= bUsesSceneDepth;
 		VisibleEditorPrimitives.AppendTo(WriteView.VisibleEditorPrimitives);
-		VisibleDynamicPrimitives.AppendTo(WriteView.VisibleDynamicPrimitives);
+		VisibleDynamicPrimitivesWithSimpleLights.AppendTo(WriteView.VisibleDynamicPrimitivesWithSimpleLights);
+		WriteView.NumVisibleDynamicPrimitives += NumVisibleDynamicPrimitives;
 		VisibleDynamicMeshesPassMask.AppendTo(WriteView.VisibleDynamicMeshesPassMask);
 		WriteView.TranslucentPrimSet.AppendScenePrimitives(TranslucencyPrims.Prims, TranslucencyPrims.NumPrims, TranslucencyPrimCount);
 		MeshDecalPrimSet.AppendTo(WriteView.MeshDecalPrimSet.Prims);
