Removed VisibleEditorPrimitives, as this array was only used for counting visible primitives.

diff --git a/Engine/Source/Runtime/Renderer/Private/SceneRendering.cpp b/Engine/Source/Runtime/Renderer/Private/SceneRendering.cpp
--- a/Engine/Source/Runtime/Renderer/Private/SceneRendering.cpp
+++ b/Engine/Source/Runtime/Renderer/Private/SceneRendering.cpp
@@ -898,6 +898,7 @@ void FViewInfo::Init()
 	}
 
 	NumVisibleDynamicPrimitives = 0;
+	NumVisibleDynamicEditorPrimitives = 0;
 }
 
 FViewInfo::~FViewInfo()
@@ -2873,7 +2874,11 @@ bool FSceneRenderer::ShouldCompositeEditorPrimitives(const FViewInfo& View)
 	}
 
 	// Any elements that needed compositing were drawn then compositing should be done
-	if (View.ViewMeshElements.Num() || View.TopViewMeshElements.Num() || View.BatchedViewElements.HasPrimsToDraw() || View.TopBatchedViewElements.HasPrimsToDraw() || View.VisibleEditorPrimitives.Num())
+	if (View.ViewMeshElements.Num() 
+		|| View.TopViewMeshElements.Num() 
+		|| View.BatchedViewElements.HasPrimsToDraw() 
+		|| View.TopBatchedViewElements.HasPrimsToDraw() 
+		|| View.NumVisibleDynamicEditorPrimitives > 0)
 	{
 		return true;
 	}
diff --git a/Engine/Source/Runtime/Renderer/Private/SceneRendering.h b/Engine/Source/Runtime/Renderer/Private/SceneRendering.h
--- a/Engine/Source/Runtime/Renderer/Private/SceneRendering.h
+++ b/Engine/Source/Runtime/Renderer/Private/SceneRendering.h
@@ -902,15 +902,15 @@ public:
 	/** Number of dynamic primitives visible in this view. */
 	int32 NumVisibleDynamicPrimitives;
 
+	/** Number of dynamic editor primitives visible in this view. */
+	int32 NumVisibleDynamicEditorPrimitives;
+
 	/** Marks used mesh passes by visible dynamic meshes. */
 	FMeshPassMask VisibleDynamicMeshesPassMask;
 
 	/** Number of dynamic mesh elements per mesh pass (inside FViewInfo::DynamicMeshElements). */
 	int32 NumVisibleDynamicMeshElements[EMeshPass::Num];
 
-	/** The dynamic editor primitives visible in this view. */
-	TArray<const FPrimitiveSceneInfo*,SceneRenderingAllocator> VisibleEditorPrimitives;
-
 	/** List of visible primitives with dirty indirect lighting cache buffers */
 	TArray<FPrimitiveSceneInfo*,SceneRenderingAllocator> DirtyIndirectLightingCacheBufferPrimitives;
 
diff --git a/Engine/Source/Runtime/Renderer/Private/SceneVisibility.cpp b/Engine/Source/Runtime/Renderer/Private/SceneVisibility.cpp
--- a/Engine/Source/Runtime/Renderer/Private/SceneVisibility.cpp
+++ b/Engine/Source/Runtime/Renderer/Private/SceneVisibility.cpp
@@ -1749,6 +1749,7 @@ struct FRelevancePacket
 	FRelevancePrimSet<int32> TranslucentSelfShadowPrimitives;
 	FRelevancePrimSet<FPrimitiveSceneInfo*> VisibleDynamicPrimitivesWithSimpleLights;
 	int32 NumVisibleDynamicPrimitives;
+	int32 NumVisibleDynamicEditorPrimitives;
 	FMeshPassMask VisibleDynamicMeshesPassMask;
 	FRelevancePrimSet<FTranslucentPrimSet::FTranslucentPrim, ETranslucencyPass::TPT_MAX> TranslucencyPrims;
 	// belongs to TranslucencyPrims
@@ -1758,7 +1759,6 @@ struct FRelevancePacket
 	bool bHasCustomDepthPrimitives;
 	FRelevancePrimSet<FPrimitiveSceneInfo*> LazyUpdatePrimitives;
 	FRelevancePrimSet<FPrimitiveSceneInfo*> DirtyIndirectLightingCacheBufferPrimitives;
-	FRelevancePrimSet<FPrimitiveSceneInfo*> VisibleEditorPrimitives;
 	FRelevancePrimSet<FPrimitiveSceneProxy*> VolumetricPrimSet;
 	FDrawCommandRelevancePacket DrawCommandPacket;
 
@@ -1831,6 +1831,7 @@ struct FRelevancePacket
 		, OutHasDynamicEditorMeshElementsMasks(InOutHasDynamicEditorMeshElementsMasks)
 		, MarkMasks(InMarkMasks)
 		, NumVisibleDynamicPrimitives(0)
+		, NumVisibleDynamicEditorPrimitives(0)
 		, bHasDistortionPrimitives(false)
 		, bHasCustomDepthPrimitives(false)
 		, PrimitiveCustomDataMemStack(InPrimitiveCustomDataMemStack)
@@ -1984,8 +1985,7 @@ struct FRelevancePacket
 
 			if (bEditorRelevance)
 			{
-				// Editor primitives are rendered after post processing and composited onto the scene
-				VisibleEditorPrimitives.AddPrim(PrimitiveSceneInfo);
+				++NumVisibleDynamicEditorPrimitives;
 
 				if (GIsEditor)
 				{
@@ -2329,9 +2329,9 @@ struct FRelevancePacket
 		WriteView.bUsesLightingChannels |= bUsesLightingChannels;
 		WriteView.bTranslucentSurfaceLighting |= bTranslucentSurfaceLighting;
 		WriteView.bUsesSceneDepth |= bUsesSceneDepth;
-		VisibleEditorPrimitives.AppendTo(WriteView.VisibleEditorPrimitives);
 		VisibleDynamicPrimitivesWithSimpleLights.AppendTo(WriteView.VisibleDynamicPrimitivesWithSimpleLights);
 		WriteView.NumVisibleDynamicPrimitives += NumVisibleDynamicPrimitives;
+		WriteView.NumVisibleDynamicEditorPrimitives += NumVisibleDynamicEditorPrimitives;
 		VisibleDynamicMeshesPassMask.AppendTo(WriteView.VisibleDynamicMeshesPassMask);
 		WriteView.TranslucentPrimSet.AppendScenePrimitives(TranslucencyPrims.Prims, TranslucencyPrims.NumPrims, TranslucencyPrimCount);
 		MeshDecalPrimSet.AppendTo(WriteView.MeshDecalPrimSet.Prims);
