Removed TranslucentPrimSet, as this array was only used for counting visible primitives.

diff --git a/Engine/Source/Runtime/Renderer/Private/SceneRendering.cpp b/Engine/Source/Runtime/Renderer/Private/SceneRendering.cpp
--- a/Engine/Source/Runtime/Renderer/Private/SceneRendering.cpp
+++ b/Engine/Source/Runtime/Renderer/Private/SceneRendering.cpp
@@ -3582,7 +3582,7 @@ static void DisplayInternals(FRHICommandListImmediate& RHICmdList, FViewInfo& In
 		CANVAS_LINE(false, TEXT("  ViewRotation: %s"), *InView.ViewRotation.ToString())
 		CANVAS_LINE(false, TEXT("  ViewRect: %s"), *ViewInfo.ViewRect.ToString())
 
-		CANVAS_LINE(false, TEXT("  DynMeshElements/TranslPrim: %d/%d"), ViewInfo.DynamicMeshElements.Num(), ViewInfo.TranslucentPrimSet.NumPrims())
+		CANVAS_LINE(false, TEXT("  DynMeshElements/TranslPrim: %d/%d"), ViewInfo.DynamicMeshElements.Num(), ViewInfo.TranslucentPrimCount.NumPrims())
 
 #undef CANVAS_LINE
 #undef CANVAS_HEADER
diff --git a/Engine/Source/Runtime/Renderer/Private/SceneRendering.h b/Engine/Source/Runtime/Renderer/Private/SceneRendering.h
--- a/Engine/Source/Runtime/Renderer/Private/SceneRendering.h
+++ b/Engine/Source/Runtime/Renderer/Private/SceneRendering.h
@@ -190,32 +190,19 @@ public:
 		DisableOffscreenRenderingPerPass[InPass] |= bDisableOffscreenRendering;
 	}
 
-	// @return range in SortedPrims[] after sorting
-	FInt32Range GetPassRange(ETranslucencyPass::Type InPass) const
+	int32 Num(ETranslucencyPass::Type InPass) const
 	{
-		checkSlow(InPass < ETranslucencyPass::TPT_MAX);
-
-		// can be optimized (if needed)
-
-		// inclusive
-		int32 Start = 0;
-
-		uint32 i = 0;
-
-		for(; i < (uint32)InPass; ++i)
-		{
-			Start += Count[i];
-		}
-
-		// exclusive
-		int32 End = Start + Count[i];
-		
-		return FInt32Range(Start, End);
+		return Count[InPass];
 	}
 
-	int32 Num(ETranslucencyPass::Type InPass) const
+	int32 NumPrims() const
 	{
-		return Count[InPass];
+		int32 NumTotal = 0;
+		for (uint32 PassIndex = 0; PassIndex < ETranslucencyPass::TPT_MAX; ++PassIndex)
+		{
+			NumTotal += Count[PassIndex];
+		}
+		return NumTotal;
 	}
 
 	bool UseSceneColorCopy(ETranslucencyPass::Type InPass) const
@@ -293,57 +280,6 @@ public:
 	}
 };
 
-/** 
-* Set of sorted translucent scene prims  
-*/
-class FTranslucentPrimSet
-{
-public:
-	/** contains a scene prim */
-	struct FTranslucentPrim
-	{
-		/** Default constructor. */
-		FTranslucentPrim() {}
-
-		// @param InPass
-		FTranslucentPrim(FPrimitiveSceneInfo* InPrimitiveSceneInfo, ETranslucencyPass::Type InPass)
-			:	PrimitiveSceneInfo(InPrimitiveSceneInfo)
-			,	Pass(InPass)
-		{}
-
-		FPrimitiveSceneInfo* PrimitiveSceneInfo;
-		ETranslucencyPass::Type Pass;
-	};
-
-	/**
-	* Insert a primitive to the translucency rendering list[s]
-	*/
-	static void PlaceScenePrimitive(FPrimitiveSceneInfo* PrimitiveSceneInfo, const FViewInfo& ViewInfo, const FPrimitiveViewRelevance& ViewRelevance,
-		FTranslucentPrim* InArrayStart, int32& InOutArrayNum, FTranslucenyPrimCount& OutCount);
-
-	/** 
-	* @return number of prims to render
-	*/
-	int32 NumPrims() const
-	{
-		return Prims.Num();
-	}
-
-	/**
-	* Adds primitives originally created with PlaceScenePrimitive
-	*/
-	void AppendScenePrimitives(FTranslucentPrim* Elements, int32 Num, const FTranslucenyPrimCount& TranslucentPrimitiveCountPerPass);
-
-	FTranslucenyPrimCount PrimsNum;
-
-private:
-
-	/** list of translucent primitives, sorted after calling Sort() */
-	TArray<FTranslucentPrim,SceneRenderingAllocator> Prims;
-};
-
-template <> struct TIsPODType<FTranslucentPrimSet::FTranslucentPrim> { enum { Value = true }; };
-
 /** A batched occlusion primitive. */
 struct FOcclusionPrimitive
 {
@@ -920,8 +856,8 @@ public:
 	/** View dependent global distance field clipmap info. */
 	FGlobalDistanceFieldInfo GlobalDistanceFieldInfo;
 
-	/** Set of translucent prims for this view */
-	FTranslucentPrimSet TranslucentPrimSet;
+	/** Count of translucent prims for this view. */
+	FTranslucenyPrimCount TranslucentPrimCount;
 
 	/** Set of mesh decal prims for this view */
 	FMeshDecalPrimSet MeshDecalPrimSet;
diff --git a/Engine/Source/Runtime/Renderer/Private/SceneVisibility.cpp b/Engine/Source/Runtime/Renderer/Private/SceneVisibility.cpp
--- a/Engine/Source/Runtime/Renderer/Private/SceneVisibility.cpp
+++ b/Engine/Source/Runtime/Renderer/Private/SceneVisibility.cpp
@@ -1751,9 +1751,7 @@ struct FRelevancePacket
 	int32 NumVisibleDynamicPrimitives;
 	int32 NumVisibleDynamicEditorPrimitives;
 	FMeshPassMask VisibleDynamicMeshesPassMask;
-	FRelevancePrimSet<FTranslucentPrimSet::FTranslucentPrim, ETranslucencyPass::TPT_MAX> TranslucencyPrims;
-	// belongs to TranslucencyPrims
-	FTranslucenyPrimCount TranslucencyPrimCount;
+	FTranslucenyPrimCount TranslucentPrimCount;
 	FRelevancePrimSet<FMeshDecalPrimSet::KeyType> MeshDecalPrimSet;
 	bool bHasDistortionPrimitives;
 	bool bHasCustomDepthPrimitives;
@@ -2012,8 +2010,23 @@ struct FRelevancePacket
 
 			if (bTranslucentRelevance && !bEditorRelevance && ViewRelevance.bRenderInMainPass)
 			{
-				// Add to set of dynamic translucent primitives
-				FTranslucentPrimSet::PlaceScenePrimitive(PrimitiveSceneInfo, View, ViewRelevance, &TranslucencyPrims.Prims[0], TranslucencyPrims.NumPrims, TranslucencyPrimCount);
+				if (View.Family->AllowTranslucencyAfterDOF())
+				{
+					if (ViewRelevance.bNormalTranslucencyRelevance)
+					{
+						TranslucentPrimCount.Add(ETranslucencyPass::TPT_StandardTranslucency, ViewRelevance.bUsesSceneColorCopy, ViewRelevance.bDisableOffscreenRendering);
+					}
+
+					if (ViewRelevance.bSeparateTranslucencyRelevance)
+					{
+						TranslucentPrimCount.Add(ETranslucencyPass::TPT_TranslucencyAfterDOF, ViewRelevance.bUsesSceneColorCopy, ViewRelevance.bDisableOffscreenRendering);
+					}
+				}
+				else // Otherwise, everything is rendered in a single bucket. This is not related to whether DOF is currently enabled or not.
+				{
+					// When using all translucency, Standard and AfterDOF are sorted together instead of being rendered like 2 buckets.
+					TranslucentPrimCount.Add(ETranslucencyPass::TPT_AllTranslucency, ViewRelevance.bUsesSceneColorCopy, ViewRelevance.bDisableOffscreenRendering);
+				}
 
 				if (ViewRelevance.bDistortionRelevance)
 				{
@@ -2333,7 +2346,7 @@ struct FRelevancePacket
 		WriteView.NumVisibleDynamicPrimitives += NumVisibleDynamicPrimitives;
 		WriteView.NumVisibleDynamicEditorPrimitives += NumVisibleDynamicEditorPrimitives;
 		VisibleDynamicMeshesPassMask.AppendTo(WriteView.VisibleDynamicMeshesPassMask);
-		WriteView.TranslucentPrimSet.AppendScenePrimitives(TranslucencyPrims.Prims, TranslucencyPrims.NumPrims, TranslucencyPrimCount);
+		WriteView.TranslucentPrimCount.Append(TranslucentPrimCount);
 		MeshDecalPrimSet.AppendTo(WriteView.MeshDecalPrimSet.Prims);
 		WriteView.bHasDistortionPrimitives |= bHasDistortionPrimitives;
 		WriteView.bHasCustomDepthPrimitives |= bHasCustomDepthPrimitives;
diff --git a/Engine/Source/Runtime/Renderer/Private/TranslucentRendering.cpp b/Engine/Source/Runtime/Renderer/Private/TranslucentRendering.cpp
--- a/Engine/Source/Runtime/Renderer/Private/TranslucentRendering.cpp
+++ b/Engine/Source/Runtime/Renderer/Private/TranslucentRendering.cpp
@@ -284,39 +284,6 @@ FMeshDrawCommandSortKey CalculateStaticTranslucentMeshSortKey(const FPrimitiveSc
 	return SortKey;
 }
 
-void FTranslucentPrimSet::AppendScenePrimitives(FTranslucentPrim* Elements, int32 Num, const FTranslucenyPrimCount& TranslucentPrimitiveCountPerPass)
-{
-	Prims.Append(Elements, Num);
-	PrimsNum.Append(TranslucentPrimitiveCountPerPass);
-}
-
-void FTranslucentPrimSet::PlaceScenePrimitive(FPrimitiveSceneInfo* PrimitiveSceneInfo, const FViewInfo& ViewInfo, const FPrimitiveViewRelevance& ViewRelevance,
-	FTranslucentPrimSet::FTranslucentPrim *InArrayStart, int32& InOutArrayNum, FTranslucenyPrimCount& OutCount)
-{
-	const auto FeatureLevel = ViewInfo.GetFeatureLevel();
-
-	if (ViewInfo.Family->AllowTranslucencyAfterDOF())
-	{
-		if (ViewRelevance.bNormalTranslucencyRelevance)
-		{
-			new(&InArrayStart[InOutArrayNum++]) FTranslucentPrim(PrimitiveSceneInfo, ETranslucencyPass::TPT_StandardTranslucency);
-			OutCount.Add(ETranslucencyPass::TPT_StandardTranslucency, ViewRelevance.bUsesSceneColorCopy, ViewRelevance.bDisableOffscreenRendering);
-		}
-
-		if (ViewRelevance.bSeparateTranslucencyRelevance)
-		{
-			new(&InArrayStart[InOutArrayNum++]) FTranslucentPrim(PrimitiveSceneInfo, ETranslucencyPass::TPT_TranslucencyAfterDOF);
-			OutCount.Add(ETranslucencyPass::TPT_TranslucencyAfterDOF, ViewRelevance.bUsesSceneColorCopy, ViewRelevance.bDisableOffscreenRendering);
-		}
-	}
-	else // Otherwise, everything is rendered in a single bucket. This is not related to whether DOF is currently enabled or not.
-	{
-		// When using all translucency, Standard and AfterDOF are sorted together instead of being rendered like 2 buckets.
-		new(&InArrayStart[InOutArrayNum++]) FTranslucentPrim(PrimitiveSceneInfo, ETranslucencyPass::TPT_AllTranslucency);
-		OutCount.Add(ETranslucencyPass::TPT_AllTranslucency, ViewRelevance.bUsesSceneColorCopy, ViewRelevance.bDisableOffscreenRendering);
-	}
-}
-
 extern int32 GLightShaftRenderAfterDOF;
 
 bool FSceneRenderer::ShouldRenderTranslucency(ETranslucencyPass::Type TranslucencyPass) const
@@ -346,7 +313,7 @@ bool FSceneRenderer::ShouldRenderTranslucency(ETranslucencyPass::Type Translucen
 
 	for (const FViewInfo& View : Views)
 	{
-		if (View.TranslucentPrimSet.PrimsNum.Num(TranslucencyPass) > 0)
+		if (View.TranslucentPrimCount.Num(TranslucencyPass) > 0)
 		{
 			return true;
 		}
@@ -555,7 +522,7 @@ void FDeferredShadingSceneRenderer::ConditionalResolveSceneColorForTranslucentMa
 		bool bNeedsResolve = false;
 		for (int32 TranslucencyPass = 0; TranslucencyPass < ETranslucencyPass::TPT_MAX && !bNeedsResolve; ++TranslucencyPass)
 		{
-			bNeedsResolve |= View.TranslucentPrimSet.PrimsNum.UseSceneColorCopy((ETranslucencyPass::Type)TranslucencyPass);
+			bNeedsResolve |= View.TranslucentPrimCount.UseSceneColorCopy((ETranslucencyPass::Type)TranslucencyPass);
 		}
 
 		if (bNeedsResolve)
@@ -948,7 +915,7 @@ void FDeferredShadingSceneRenderer::RenderTranslucency(FRHICommandListImmediate&
 		FMeshPassProcessorRenderState DrawRenderState(View, BasePassUniformBuffer);
 
 		// If downsampling we need to render in the separate buffer. Otherwise we also need to render offscreen to apply TPT_TranslucencyAfterDOF
-		if (RenderInSeparateTranslucency(SceneContext, TranslucencyPass, View.TranslucentPrimSet.PrimsNum.DisableOffscreenRendering(TranslucencyPass)))
+		if (RenderInSeparateTranslucency(SceneContext, TranslucencyPass, View.TranslucentPrimCount.DisableOffscreenRendering(TranslucencyPass)))
 		{
 			checkSlow(RHICmdList.IsOutsideRenderPass());
 
