Improved mesh draw command sorting. Most mesh draw commands are now first sorted by pixel shader's hash and then vertex shader's hash, and only finally by the instancing bucket.This decreases number of context and state switches between instancing buckets. Also removed sorting by CachedPipelineId.Id(), as it scatters similar state instead of grouping it together.

diff --git a/Engine/Source/Runtime/Landscape/Private/LandscapeGrass.cpp b/Engine/Source/Runtime/Landscape/Private/LandscapeGrass.cpp
--- a/Engine/Source/Runtime/Landscape/Private/LandscapeGrass.cpp
+++ b/Engine/Source/Runtime/Landscape/Private/LandscapeGrass.cpp
@@ -397,6 +397,8 @@ void FLandscapeGrassWeightMeshProcessor::Process(
 	FLandscapeGrassWeightShaderElementData ShaderElementData;
 	ShaderElementData.InitializeMeshMaterialData(ViewIfDynamicMeshCommand, PrimitiveSceneProxy, MeshBatch, -1, true);
 
+	const FMeshDrawCommandSortKey SortKey = CalculateMeshStaticSortKey(PassShaders.VertexShader, PassShaders.PixelShader);
+
 	for (int32 PassIndex = 0; PassIndex < NumPasses; ++PassIndex)
 	{
 		ShaderElementData.OutputPass = PassIndex;
@@ -413,7 +415,7 @@ void FLandscapeGrassWeightMeshProcessor::Process(
 			MeshFillMode,
 			MeshCullMode,
 			1,
-			FMeshDrawCommandSortKey::Default,
+			SortKey,
 			EMeshPassFeatures::Default,
 			ShaderElementData);
 	}
diff --git a/Engine/Source/Runtime/Renderer/Private/BasePassRendering.cpp b/Engine/Source/Runtime/Renderer/Private/BasePassRendering.cpp
--- a/Engine/Source/Runtime/Renderer/Private/BasePassRendering.cpp
+++ b/Engine/Source/Runtime/Renderer/Private/BasePassRendering.cpp
@@ -227,19 +227,54 @@ void SetTranslucentRenderState(FMeshPassProcessorRenderState& DrawRenderState, c
 	}
 }
 
-FMeshDrawCommandSortKey CalculateStaticBaseMeshSortKey(EDepthDrawingMode EarlyZPassMode, EBlendMode BlendMode)
+FMeshDrawCommandSortKey CalculateTranslucentMeshStaticSortKey(const FPrimitiveSceneProxy* RESTRICT PrimitiveSceneProxy, uint16 MeshIdInPrimitive)
 {
+	uint16 SortKeyPriority = 0;
+
+	if (PrimitiveSceneProxy)
+	{
+		const FPrimitiveSceneInfo* PrimitiveSceneInfo = PrimitiveSceneProxy->GetPrimitiveSceneInfo();
+		SortKeyPriority = (uint16)((int32)PrimitiveSceneInfo->Proxy->GetTranslucencySortPriority() - (int32)SHRT_MIN);
+	}
+
+	FTranslucentMeshSortKey TranslucentMeshSortKey;
+	TranslucentMeshSortKey.Fields.MeshIdInPrimitive = MeshIdInPrimitive;
+	TranslucentMeshSortKey.Fields.Priority = SortKeyPriority;
+	TranslucentMeshSortKey.Fields.Distance = 0; // View specific, so will be filled later inside VisibleMeshCommands.
+
 	FMeshDrawCommandSortKey SortKey;
+	SortKey.PackedData = TranslucentMeshSortKey.PackedData;
+	return SortKey;
+}
+
+FMeshDrawCommandSortKey CalculateBasePassMeshStaticSortKey(EDepthDrawingMode EarlyZPassMode, EBlendMode BlendMode, const FMeshMaterialShader* VertexShader, const FMeshMaterialShader* PixelShader)
+{
+	union FBasePassMeshSortKey
+	{
+		uint64 PackedData;
 
+		struct
+		{
+			uint64 VertexShaderHash	: 16;	// Order by vertex shader's hash.
+			uint64 PixelShaderHash	: 32;	// Order by pixel shader's hash.
+			uint64 Masked			: 16;	// First order by masked.
+		} Fields;
+	};
+
+	FBasePassMeshSortKey BasePassMeshSortKey;
+	BasePassMeshSortKey.Fields.VertexShaderHash = PointerHash(VertexShader) & 0xFFFF;
+	BasePassMeshSortKey.Fields.PixelShaderHash = PointerHash(PixelShader);
 	if (EarlyZPassMode != DDM_None)
 	{
-		SortKey.PackedData = BlendMode == EBlendMode::BLEND_Masked ? 0 : 1;
+		BasePassMeshSortKey.Fields.Masked = BlendMode == EBlendMode::BLEND_Masked ? 0 : 1;
 	}
 	else
 	{
-		SortKey.PackedData = BlendMode == EBlendMode::BLEND_Masked ? 1 : 0;
+		BasePassMeshSortKey.Fields.Masked = BlendMode == EBlendMode::BLEND_Masked ? 1 : 0;
 	}
 
+	FMeshDrawCommandSortKey SortKey;
+	SortKey.PackedData = BasePassMeshSortKey.PackedData;
 	return SortKey;
 }
 
@@ -947,20 +982,20 @@ void FBasePassMeshProcessor::Process(
 
 	SetBasePassDitheredLODTransitionState(ViewIfDynamicMeshCommand, MeshBatch, StaticMeshId, DrawRenderState);
 
+	TBasePassShaderElementData<LightMapPolicyType> ShaderElementData(LightMapElementData);
+	ShaderElementData.InitializeMeshMaterialData(ViewIfDynamicMeshCommand, PrimitiveSceneProxy, MeshBatch, StaticMeshId, true);
+
 	FMeshDrawCommandSortKey SortKey = FMeshDrawCommandSortKey::Default;
 
 	if (bTranslucentBasePass)
 	{
-		SortKey = CalculateStaticTranslucentMeshSortKey(PrimitiveSceneProxy, MeshBatch.MeshIdInPrimitive);
+		SortKey = CalculateTranslucentMeshStaticSortKey(PrimitiveSceneProxy, MeshBatch.MeshIdInPrimitive);
 	}
 	else
 	{
-		SortKey = CalculateStaticBaseMeshSortKey(EarlyZPassMode, BlendMode);
+		SortKey = CalculateBasePassMeshStaticSortKey(EarlyZPassMode, BlendMode, BasePassShaders.VertexShader, BasePassShaders.PixelShader);
 	}
 
-	TBasePassShaderElementData<LightMapPolicyType> ShaderElementData(LightMapElementData);
-	ShaderElementData.InitializeMeshMaterialData(ViewIfDynamicMeshCommand, PrimitiveSceneProxy, MeshBatch, StaticMeshId, true);
-
 	const bool bIsInstancedStereo = ViewIfDynamicMeshCommand ? ViewIfDynamicMeshCommand->IsInstancedStereoPass() : Scene->bStaticDrawInstancedStereo;
 	const int32 InstanceFactor = bIsInstancedStereo ? 2 : 1;
 
diff --git a/Engine/Source/Runtime/Renderer/Private/BasePassRendering.h b/Engine/Source/Runtime/Renderer/Private/BasePassRendering.h
--- a/Engine/Source/Runtime/Renderer/Private/BasePassRendering.h
+++ b/Engine/Source/Runtime/Renderer/Private/BasePassRendering.h
@@ -642,4 +642,5 @@ private:
 
 extern void SetDepthStencilStateForBasePass(FMeshPassProcessorRenderState& DrawRenderState, ERHIFeatureLevel::Type FeatureLevel, const FMeshBatch& Mesh, const FPrimitiveSceneProxy* PrimitiveSceneProxy, bool bEnableReceiveDecalOutput, bool bUseDebugViewPS, FDepthStencilStateRHIParamRef LodFadeOverrideDepthStencilState);
 extern void SetupBasePassState(FExclusiveDepthStencil::Type BasePassDepthStencilAccess, const bool bShaderComplexity, FMeshPassProcessorRenderState& DrawRenderState);
-extern FExclusiveDepthStencil::Type GetDefaultBasePassDepthStencilAccess(EShadingPath ShadingPath);
\ No newline at end of file
+extern FExclusiveDepthStencil::Type GetDefaultBasePassDepthStencilAccess(EShadingPath ShadingPath);
+extern FMeshDrawCommandSortKey CalculateTranslucentMeshStaticSortKey(const FPrimitiveSceneProxy* RESTRICT PrimitiveSceneProxy, uint16 MeshIdInPrimitive);
\ No newline at end of file
diff --git a/Engine/Source/Runtime/Renderer/Private/CompositionLighting/PostProcessMeshDecals.cpp b/Engine/Source/Runtime/Renderer/Private/CompositionLighting/PostProcessMeshDecals.cpp
--- a/Engine/Source/Runtime/Renderer/Private/CompositionLighting/PostProcessMeshDecals.cpp
+++ b/Engine/Source/Runtime/Renderer/Private/CompositionLighting/PostProcessMeshDecals.cpp
@@ -329,6 +329,8 @@ void FMeshDecalMeshProcessor::Process(
 	FMeshMaterialShaderElementData ShaderElementData;
 	ShaderElementData.InitializeMeshMaterialData(ViewIfDynamicMeshCommand, PrimitiveSceneProxy, MeshBatch, StaticMeshId, true);
 
+	const FMeshDrawCommandSortKey SortKey = CalculateMeshStaticSortKey(MeshDecalPassShaders.VertexShader, MeshDecalPassShaders.PixelShader);
+
 	BuildMeshDrawCommands(
 		MeshBatch,
 		BatchElementMask,
@@ -340,7 +342,7 @@ void FMeshDecalMeshProcessor::Process(
 		MeshFillMode,
 		MeshCullMode,
 		1,
-		FMeshDrawCommandSortKey::Default,
+		SortKey,
 		EMeshPassFeatures::Default,
 		ShaderElementData);
 }
diff --git a/Engine/Source/Runtime/Renderer/Private/ConvertToUniformMesh.cpp b/Engine/Source/Runtime/Renderer/Private/ConvertToUniformMesh.cpp
--- a/Engine/Source/Runtime/Renderer/Private/ConvertToUniformMesh.cpp
+++ b/Engine/Source/Runtime/Renderer/Private/ConvertToUniformMesh.cpp
@@ -190,6 +190,8 @@ void FConvertToUniformMeshProcessor::Process(
 	FMeshMaterialShaderElementData ShaderElementData;
 	ShaderElementData.InitializeMeshMaterialData(ViewIfDynamicMeshCommand, PrimitiveSceneProxy, MeshBatch, -1, true);
 
+	const FMeshDrawCommandSortKey SortKey = CalculateMeshStaticSortKey(PassShaders.VertexShader, PassShaders.PixelShader);
+
 	BuildMeshDrawCommands(
 		MeshBatch,
 		BatchElementMask,
@@ -201,7 +203,7 @@ void FConvertToUniformMeshProcessor::Process(
 		MeshFillMode,
 		MeshCullMode,
 		1,
-		FMeshDrawCommandSortKey::Default,
+		SortKey,
 		EMeshPassFeatures::Default,
 		ShaderElementData);
 }
diff --git a/Engine/Source/Runtime/Renderer/Private/CustomDepthRendering.cpp b/Engine/Source/Runtime/Renderer/Private/CustomDepthRendering.cpp
--- a/Engine/Source/Runtime/Renderer/Private/CustomDepthRendering.cpp
+++ b/Engine/Source/Runtime/Renderer/Private/CustomDepthRendering.cpp
@@ -173,6 +173,8 @@ void FCustomDepthPassMeshProcessor::Process(
 	FDepthOnlyShaderElementData ShaderElementData(MobileColorValue);
 	ShaderElementData.InitializeMeshMaterialData(ViewIfDynamicMeshCommand, PrimitiveSceneProxy, MeshBatch, StaticMeshId, false);
 
+	const FMeshDrawCommandSortKey SortKey = CalculateMeshStaticSortKey(DepthPassShaders.VertexShader, DepthPassShaders.PixelShader);
+
 	BuildMeshDrawCommands(
 		MeshBatch,
 		BatchElementMask,
@@ -184,7 +186,7 @@ void FCustomDepthPassMeshProcessor::Process(
 		MeshFillMode,
 		MeshCullMode,
 		1,
-		FMeshDrawCommandSortKey::Default,
+		SortKey,
 		bPositionOnly ? EMeshPassFeatures::PositionOnly : EMeshPassFeatures::Default,
 		ShaderElementData);
 }
diff --git a/Engine/Source/Runtime/Renderer/Private/DebugViewModeRendering.cpp b/Engine/Source/Runtime/Renderer/Private/DebugViewModeRendering.cpp
--- a/Engine/Source/Runtime/Renderer/Private/DebugViewModeRendering.cpp
+++ b/Engine/Source/Runtime/Renderer/Private/DebugViewModeRendering.cpp
@@ -328,6 +328,8 @@ void FDebugViewModeMeshProcessor::AddMeshBatch(const FMeshBatch& RESTRICT MeshBa
 
 	ShaderElementData.InitializeMeshMaterialData(ViewIfDynamicMeshCommand, PrimitiveSceneProxy, MeshBatch, StaticMeshId, true);
 
+	const FMeshDrawCommandSortKey SortKey = CalculateMeshStaticSortKey(DebugViewModePassShaders.VertexShader, DebugViewModePassShaders.PixelShader);
+
 	BuildMeshDrawCommands(
 		MeshBatch,
 		BatchElementMask,
@@ -339,7 +341,7 @@ void FDebugViewModeMeshProcessor::AddMeshBatch(const FMeshBatch& RESTRICT MeshBa
 		MeshFillMode,
 		MeshCullMode,
 		1,
-		FMeshDrawCommandSortKey::Default,
+		SortKey,
 		EMeshPassFeatures::Default,
 		ShaderElementData);
 }
diff --git a/Engine/Source/Runtime/Renderer/Private/DepthRendering.cpp b/Engine/Source/Runtime/Renderer/Private/DepthRendering.cpp
--- a/Engine/Source/Runtime/Renderer/Private/DepthRendering.cpp
+++ b/Engine/Source/Runtime/Renderer/Private/DepthRendering.cpp
@@ -722,6 +722,8 @@ void FDepthPassMeshProcessor::Process(
 	FDepthOnlyShaderElementData ShaderElementData(0.0f);
 	ShaderElementData.InitializeMeshMaterialData(ViewIfDynamicMeshCommand, PrimitiveSceneProxy, MeshBatch, StaticMeshId, true);
 
+	const FMeshDrawCommandSortKey SortKey = CalculateMeshStaticSortKey(DepthPassShaders.VertexShader, DepthPassShaders.PixelShader);
+
 	const bool bIsInstancedStereo = ViewIfDynamicMeshCommand ? ViewIfDynamicMeshCommand->IsInstancedStereoPass() : Scene->bStaticDrawInstancedStereo;
 	const int32 InstanceFactor = bIsInstancedStereo ? 2 : 1;
 
@@ -736,7 +738,7 @@ void FDepthPassMeshProcessor::Process(
 		MeshFillMode,
 		MeshCullMode,
 		InstanceFactor,
-		FMeshDrawCommandSortKey::Default,
+		SortKey,
 		bPositionOnly ? EMeshPassFeatures::PositionOnly : EMeshPassFeatures::Default,
 		ShaderElementData);
 }
diff --git a/Engine/Source/Runtime/Renderer/Private/DistortionRendering.cpp b/Engine/Source/Runtime/Renderer/Private/DistortionRendering.cpp
--- a/Engine/Source/Runtime/Renderer/Private/DistortionRendering.cpp
+++ b/Engine/Source/Runtime/Renderer/Private/DistortionRendering.cpp
@@ -756,6 +756,8 @@ void FDistortionMeshProcessor::Process(
 	FMeshMaterialShaderElementData ShaderElementData;
 	ShaderElementData.InitializeMeshMaterialData(ViewIfDynamicMeshCommand, PrimitiveSceneProxy, MeshBatch, StaticMeshId, false);
 
+	const FMeshDrawCommandSortKey SortKey = CalculateMeshStaticSortKey(DistortionPassShaders.VertexShader, DistortionPassShaders.PixelShader);
+
 	BuildMeshDrawCommands(
 		MeshBatch,
 		BatchElementMask,
@@ -767,7 +769,7 @@ void FDistortionMeshProcessor::Process(
 		MeshFillMode,
 		MeshCullMode,
 		1,
-		FMeshDrawCommandSortKey::Default,
+		SortKey,
 		EMeshPassFeatures::Default,
 		ShaderElementData);
 }
diff --git a/Engine/Source/Runtime/Renderer/Private/EditorPrimitivesRendering.cpp b/Engine/Source/Runtime/Renderer/Private/EditorPrimitivesRendering.cpp
--- a/Engine/Source/Runtime/Renderer/Private/EditorPrimitivesRendering.cpp
+++ b/Engine/Source/Runtime/Renderer/Private/EditorPrimitivesRendering.cpp
@@ -87,6 +87,8 @@ void FEditorPrimitivesBasePassMeshProcessor::ProcessDeferredShadingPath(const FM
 	TBasePassShaderElementData<LightMapPolicyType> ShaderElementData(nullptr);
 	ShaderElementData.InitializeMeshMaterialData(ViewIfDynamicMeshCommand, PrimitiveSceneProxy, MeshBatch, StaticMeshId, false);
 
+	const FMeshDrawCommandSortKey SortKey = CalculateMeshStaticSortKey(BasePassShaders.VertexShader, BasePassShaders.PixelShader);
+
 	BuildMeshDrawCommands(
 		MeshBatch,
 		BatchElementMask,
@@ -98,7 +100,7 @@ void FEditorPrimitivesBasePassMeshProcessor::ProcessDeferredShadingPath(const FM
 		MeshFillMode,
 		MeshCullMode,
 		1,
-		FMeshDrawCommandSortKey::Default,
+		SortKey,
 		EMeshPassFeatures::Default,
 		ShaderElementData);
 }
@@ -140,6 +142,8 @@ void FEditorPrimitivesBasePassMeshProcessor::ProcessMobileShadingPath(const FMes
 	TMobileBasePassShaderElementData<LightMapPolicyType> ShaderElementData(nullptr);
 	ShaderElementData.InitializeMeshMaterialData(ViewIfDynamicMeshCommand, PrimitiveSceneProxy, MeshBatch, StaticMeshId, false);
 
+	const FMeshDrawCommandSortKey SortKey = CalculateMeshStaticSortKey(BasePassShaders.VertexShader, BasePassShaders.PixelShader);
+
 	BuildMeshDrawCommands(
 		MeshBatch,
 		BatchElementMask,
@@ -151,7 +155,7 @@ void FEditorPrimitivesBasePassMeshProcessor::ProcessMobileShadingPath(const FMes
 		MeshFillMode,
 		MeshCullMode,
 		1,
-		FMeshDrawCommandSortKey::Default,
+		SortKey,
 		EMeshPassFeatures::Default,
 		ShaderElementData);
 }
\ No newline at end of file
diff --git a/Engine/Source/Runtime/Renderer/Private/LightMapDensityRendering.cpp b/Engine/Source/Runtime/Renderer/Private/LightMapDensityRendering.cpp
--- a/Engine/Source/Runtime/Renderer/Private/LightMapDensityRendering.cpp
+++ b/Engine/Source/Runtime/Renderer/Private/LightMapDensityRendering.cpp
@@ -208,6 +208,8 @@ void FLightmapDensityMeshProcessor::Process(
 
 	ShaderElementData.InitializeMeshMaterialData(ViewIfDynamicMeshCommand, PrimitiveSceneProxy, MeshBatch, StaticMeshId, false);
 
+	const FMeshDrawCommandSortKey SortKey = CalculateMeshStaticSortKey(LightmapDensityPassShaders.VertexShader, LightmapDensityPassShaders.PixelShader);
+
 	BuildMeshDrawCommands(
 		MeshBatch,
 		BatchElementMask,
@@ -219,7 +221,7 @@ void FLightmapDensityMeshProcessor::Process(
 		MeshFillMode,
 		MeshCullMode,
 		1,
-		FMeshDrawCommandSortKey::Default,
+		SortKey,
 		EMeshPassFeatures::Default,
 		ShaderElementData);
 }
diff --git a/Engine/Source/Runtime/Renderer/Private/MeshDrawCommands.cpp b/Engine/Source/Runtime/Renderer/Private/MeshDrawCommands.cpp
--- a/Engine/Source/Runtime/Renderer/Private/MeshDrawCommands.cpp
+++ b/Engine/Source/Runtime/Renderer/Private/MeshDrawCommands.cpp
@@ -92,15 +92,6 @@ struct FCompareFMeshDrawCommands
 			return A.StateBucketId < B.StateBucketId;
 		}
 
-		// No valid StateBucketId - we cannot dynamically instance but we can still order by PSO to maximize state filtering
-		if (A.StateBucketId == -1 && B.StateBucketId == -1)
-		{
-			if (A.MeshDrawCommand->CachedPipelineId.GetId() != B.MeshDrawCommand->CachedPipelineId.GetId())
-			{
-				return A.MeshDrawCommand->CachedPipelineId.GetId() < B.MeshDrawCommand->CachedPipelineId.GetId();
-			}
-		}
-
 		return false;
 	}
 };
@@ -1114,7 +1105,7 @@ void FParallelMeshDrawCommandPass::DispatchDraw(FParallelCommandListSet* Paralle
 
 void FParallelMeshDrawCommandPass::DumpInstancingStats() const
 {
-	if (!PassNameForStats.IsEmpty())
+	if (!PassNameForStats.IsEmpty() && TaskContext.VisibleMeshDrawCommandsNum > 0)
 	{
 		UE_LOG(LogRenderer, Log, TEXT("Instancing stats for %s"), *PassNameForStats);
 		UE_LOG(LogRenderer, Log, TEXT("   %i Mesh Draw Commands in %i instancing state buckets"), TaskContext.VisibleMeshDrawCommandsNum, TaskContext.NewPassVisibleMeshDrawCommandsNum);
diff --git a/Engine/Source/Runtime/Renderer/Private/MeshPassProcessor.cpp b/Engine/Source/Runtime/Renderer/Private/MeshPassProcessor.cpp
--- a/Engine/Source/Runtime/Renderer/Private/MeshPassProcessor.cpp
+++ b/Engine/Source/Runtime/Renderer/Private/MeshPassProcessor.cpp
@@ -869,6 +869,28 @@ void DrawDynamicMeshPassPrivate(
 	}
 }
 
+FMeshDrawCommandSortKey CalculateMeshStaticSortKey(const FMeshMaterialShader* VertexShader, const FMeshMaterialShader* PixelShader)
+{
+	union FMeshSortKey
+	{
+		uint64 PackedData;
+
+		struct
+		{
+			uint64 VertexShaderHash : 32;	// Order by vertex shader's hash.
+			uint64 PixelShaderHash : 32;	// First order by pixel shader's hash.
+		} Fields;
+	};
+
+	FMeshSortKey MeshSortKey;
+	MeshSortKey.Fields.VertexShaderHash = PointerHash(VertexShader);
+	MeshSortKey.Fields.PixelShaderHash = PointerHash(PixelShader);
+
+	FMeshDrawCommandSortKey SortKey;
+	SortKey.PackedData = MeshSortKey.PackedData;
+	return SortKey;
+}
+
 FMeshPassProcessor::FMeshPassProcessor(const FScene* InScene, ERHIFeatureLevel::Type InFeatureLevel, const FSceneView* InViewIfDynamicMeshCommand, FMeshPassDrawListContext* InDrawListContext) 
 	: Scene(InScene)
 	, FeatureLevel(InFeatureLevel)
diff --git a/Engine/Source/Runtime/Renderer/Private/MobileBasePass.cpp b/Engine/Source/Runtime/Renderer/Private/MobileBasePass.cpp
--- a/Engine/Source/Runtime/Renderer/Private/MobileBasePass.cpp
+++ b/Engine/Source/Runtime/Renderer/Private/MobileBasePass.cpp
@@ -645,7 +645,7 @@ void FMobileBasePassMeshProcessor::Process(
 	FMeshDrawCommandSortKey SortKey; 
 	if (bTranslucentBasePass)
 	{
-		SortKey = CalculateStaticTranslucentMeshSortKey(PrimitiveSceneProxy, MeshBatch.MeshIdInPrimitive);
+		SortKey = CalculateTranslucentMeshStaticSortKey(PrimitiveSceneProxy, MeshBatch.MeshIdInPrimitive);
 	}
 	else
 	{
diff --git a/Engine/Source/Runtime/Renderer/Private/MobileTranslucentRendering.cpp b/Engine/Source/Runtime/Renderer/Private/MobileTranslucentRendering.cpp
--- a/Engine/Source/Runtime/Renderer/Private/MobileTranslucentRendering.cpp
+++ b/Engine/Source/Runtime/Renderer/Private/MobileTranslucentRendering.cpp
@@ -369,7 +369,7 @@ private:
 		FMeshMaterialShaderElementData ShaderElementData;
 		ShaderElementData.InitializeMeshMaterialData(ViewIfDynamicMeshCommand, PrimitiveSceneProxy, MeshBatch, StaticMeshId, false);
 
-		FMeshDrawCommandSortKey SortKey = CalculateStaticTranslucentMeshSortKey(PrimitiveSceneProxy, MeshBatch.MeshIdInPrimitive);
+		FMeshDrawCommandSortKey SortKey = CalculateTranslucentMeshStaticSortKey(PrimitiveSceneProxy, MeshBatch.MeshIdInPrimitive);
 		
 		BuildMeshDrawCommands(
 			MeshBatch,
diff --git a/Engine/Source/Runtime/Renderer/Private/SceneHitProxyRendering.cpp b/Engine/Source/Runtime/Renderer/Private/SceneHitProxyRendering.cpp
--- a/Engine/Source/Runtime/Renderer/Private/SceneHitProxyRendering.cpp
+++ b/Engine/Source/Runtime/Renderer/Private/SceneHitProxyRendering.cpp
@@ -654,6 +654,8 @@ void FHitProxyMeshProcessor::Process(
 	FHitProxyShaderElementData ShaderElementData(MeshBatch.BatchHitProxyId);
 	ShaderElementData.InitializeMeshMaterialData(ViewIfDynamicMeshCommand, PrimitiveSceneProxy, MeshBatch, StaticMeshId, false);
 
+	const FMeshDrawCommandSortKey SortKey = CalculateMeshStaticSortKey(HitProxyPassShaders.VertexShader, HitProxyPassShaders.PixelShader);
+
 	BuildMeshDrawCommands(
 		MeshBatch,
 		BatchElementMask,
@@ -665,7 +667,7 @@ void FHitProxyMeshProcessor::Process(
 		MeshFillMode,
 		MeshCullMode,
 		1,
-		FMeshDrawCommandSortKey::Default,
+		SortKey,
 		EMeshPassFeatures::Default,
 		ShaderElementData);
 }
@@ -768,6 +770,8 @@ void FEditorSelectionMeshProcessor::Process(
 	FHitProxyShaderElementData ShaderElementData(DummyId);
 	ShaderElementData.InitializeMeshMaterialData(ViewIfDynamicMeshCommand, PrimitiveSceneProxy, MeshBatch, StaticMeshId, false);
 
+	const FMeshDrawCommandSortKey SortKey = CalculateMeshStaticSortKey(HitProxyPassShaders.VertexShader, HitProxyPassShaders.PixelShader);
+
 	BuildMeshDrawCommands(
 		MeshBatch,
 		BatchElementMask,
@@ -779,7 +783,7 @@ void FEditorSelectionMeshProcessor::Process(
 		MeshFillMode,
 		MeshCullMode,
 		1,
-		FMeshDrawCommandSortKey::Default,
+		SortKey,
 		EMeshPassFeatures::Default,
 		ShaderElementData);
 }
diff --git a/Engine/Source/Runtime/Renderer/Private/ShadowDepthRendering.cpp b/Engine/Source/Runtime/Renderer/Private/ShadowDepthRendering.cpp
--- a/Engine/Source/Runtime/Renderer/Private/ShadowDepthRendering.cpp
+++ b/Engine/Source/Runtime/Renderer/Private/ShadowDepthRendering.cpp
@@ -1760,6 +1760,8 @@ void FShadowDepthPassMeshProcessor::Process(
 	FShadowDepthShaderElementData ShaderElementData;
 	ShaderElementData.InitializeMeshMaterialData(ViewIfDynamicMeshCommand, PrimitiveSceneProxy, MeshBatch, StaticMeshId, false);
 
+	const FMeshDrawCommandSortKey SortKey = CalculateMeshStaticSortKey(ShadowDepthPassShaders.VertexShader, ShadowDepthPassShaders.PixelShader);
+
 	const uint32 InstanceFactor = !ShadowDepthType.bOnePassPointLightShadow || RHISupportsGeometryShaders(GShaderPlatformForFeatureLevel[FeatureLevel]) ? 1 : 6;
 	for (uint32 i = 0; i < InstanceFactor; i++)
 	{
@@ -1776,7 +1778,7 @@ void FShadowDepthPassMeshProcessor::Process(
 			MeshFillMode,
 			MeshCullMode,
 			1,
-			FMeshDrawCommandSortKey::Default,
+			SortKey,
 			bUsePositionOnlyVS ? EMeshPassFeatures::PositionOnly : EMeshPassFeatures::Default,
 			ShaderElementData);
 	}
diff --git a/Engine/Source/Runtime/Renderer/Private/ShadowSetup.cpp b/Engine/Source/Runtime/Renderer/Private/ShadowSetup.cpp
--- a/Engine/Source/Runtime/Renderer/Private/ShadowSetup.cpp
+++ b/Engine/Source/Runtime/Renderer/Private/ShadowSetup.cpp
@@ -831,7 +831,7 @@ void FProjectedShadowInfo::AddCachedMeshDrawCommandsForPass(
 				CachedMeshDrawCommand.StateBucketId,
 				CachedMeshDrawCommand.MeshFillMode,
 				CachedMeshDrawCommand.MeshCullMode,
-				FMeshDrawCommandSortKey::Default);
+				CachedMeshDrawCommand.SortKey);
 
 			VisibleMeshCommands.Add(NewVisibleMeshDrawCommand);
 		}
diff --git a/Engine/Source/Runtime/Renderer/Private/TranslucentRendering.cpp b/Engine/Source/Runtime/Renderer/Private/TranslucentRendering.cpp
--- a/Engine/Source/Runtime/Renderer/Private/TranslucentRendering.cpp
+++ b/Engine/Source/Runtime/Renderer/Private/TranslucentRendering.cpp
@@ -263,27 +263,6 @@ private:
 
 IMPLEMENT_SHADER_TYPE(, FCopySceneColorPS, TEXT("/Engine/Private/TranslucentLightingShaders.usf"), TEXT("CopySceneColorMain"), SF_Pixel);
 
-FMeshDrawCommandSortKey CalculateStaticTranslucentMeshSortKey(const FPrimitiveSceneProxy* RESTRICT PrimitiveSceneProxy, uint16 MeshIdInPrimitive)
-{
-	uint16 SortKeyPriority = 0;
-
-	if (PrimitiveSceneProxy)
-	{
-		const FPrimitiveSceneInfo* PrimitiveSceneInfo = PrimitiveSceneProxy->GetPrimitiveSceneInfo();
-		SortKeyPriority = (uint16)((int32)PrimitiveSceneInfo->Proxy->GetTranslucencySortPriority() - (int32)SHRT_MIN);
-	}
-
-	FTranslucentMeshSortKey TranslucentMeshSortKey;
-	TranslucentMeshSortKey.Fields.MeshIdInPrimitive = MeshIdInPrimitive;
-	TranslucentMeshSortKey.Fields.Priority = SortKeyPriority;
-	TranslucentMeshSortKey.Fields.Distance = 0; // View specific, so will be filled later inside VisibleMeshCommands.
-
-	FMeshDrawCommandSortKey SortKey;
-	SortKey.PackedData = TranslucentMeshSortKey.PackedData;
-
-	return SortKey;
-}
-
 extern int32 GLightShaftRenderAfterDOF;
 
 bool FSceneRenderer::ShouldRenderTranslucency(ETranslucencyPass::Type TranslucencyPass) const
diff --git a/Engine/Source/Runtime/Renderer/Private/TranslucentRendering.h b/Engine/Source/Runtime/Renderer/Private/TranslucentRendering.h
--- a/Engine/Source/Runtime/Renderer/Private/TranslucentRendering.h
+++ b/Engine/Source/Runtime/Renderer/Private/TranslucentRendering.h
@@ -16,7 +16,6 @@
 
 bool UseNearestDepthNeighborUpsampleForSeparateTranslucency(const FSceneRenderTargets& SceneContext);
 
-extern FMeshDrawCommandSortKey CalculateStaticTranslucentMeshSortKey(const FPrimitiveSceneProxy* RESTRICT PrimitiveSceneProxy, uint16 MeshIdInPrimitive);
 EMeshPass::Type TranslucencyPassToMeshPass(ETranslucencyPass::Type TranslucencyPass);
 
 /**
diff --git a/Engine/Source/Runtime/Renderer/Private/VelocityRendering.cpp b/Engine/Source/Runtime/Renderer/Private/VelocityRendering.cpp
--- a/Engine/Source/Runtime/Renderer/Private/VelocityRendering.cpp
+++ b/Engine/Source/Runtime/Renderer/Private/VelocityRendering.cpp
@@ -691,6 +691,8 @@ void FVelocityMeshProcessor::Process(
 	const bool bIsInstancedStereo = ViewIfDynamicMeshCommand ? ViewIfDynamicMeshCommand->IsInstancedStereoPass() : Scene->bStaticDrawInstancedStereo;
 	const int32 InstanceFactor = bIsInstancedStereo ? 2 : 1;
 
+	const FMeshDrawCommandSortKey SortKey = CalculateMeshStaticSortKey(VelocityPassShaders.VertexShader, VelocityPassShaders.PixelShader);
+
 	BuildMeshDrawCommands(
 		MeshBatch,
 		BatchElementMask,
@@ -702,7 +704,7 @@ void FVelocityMeshProcessor::Process(
 		MeshFillMode,
 		MeshCullMode,
 		InstanceFactor,
-		FMeshDrawCommandSortKey::Default,
+		SortKey,
 		EMeshPassFeatures::Default,
 		ShaderElementData);
 }
diff --git a/Engine/Source/Runtime/Renderer/Private/VolumetricFogVoxelization.cpp b/Engine/Source/Runtime/Renderer/Private/VolumetricFogVoxelization.cpp
--- a/Engine/Source/Runtime/Renderer/Private/VolumetricFogVoxelization.cpp
+++ b/Engine/Source/Runtime/Renderer/Private/VolumetricFogVoxelization.cpp
@@ -492,6 +492,8 @@ void FVoxelizeVolumeMeshProcessor::Process(
 		PassShaders.PixelShader = MaterialResource.GetShader<TVoxelizeVolumePS<VMode_Object_Box>>(VertexFactory->GetType());
 	};
 
+	const FMeshDrawCommandSortKey SortKey = CalculateMeshStaticSortKey(PassShaders.VertexShader, PassShaders.PixelShader);
+
 	for (int32 VoxelizationPassIndex = 0; VoxelizationPassIndex < NumVoxelizationPasses; VoxelizationPassIndex++)
 	{
 		if (GVolumetricFogVoxelizationShowOnlyPassIndex < 0 || GVolumetricFogVoxelizationShowOnlyPassIndex == VoxelizationPassIndex)
@@ -510,7 +512,7 @@ void FVoxelizeVolumeMeshProcessor::Process(
 				MeshFillMode,
 				MeshCullMode,
 				1,
-				FMeshDrawCommandSortKey::Default,
+				SortKey,
 				EMeshPassFeatures::Default,
 				ShaderElementData);
 		}
diff --git a/Engine/Source/Runtime/Renderer/Public/MeshPassProcessor.h b/Engine/Source/Runtime/Renderer/Public/MeshPassProcessor.h
--- a/Engine/Source/Runtime/Renderer/Public/MeshPassProcessor.h
+++ b/Engine/Source/Runtime/Renderer/Public/MeshPassProcessor.h
@@ -986,4 +986,6 @@ RENDERER_API extern void DrawDynamicMeshPassPrivate(
 	const FSceneView& View,
 	FRHICommandList& RHICmdList,
 	FMeshCommandOneFrameArray& VisibleMeshDrawCommands,
-	FDynamicMeshDrawCommandStorage& DynamicMeshDrawCommandStorage);
\ No newline at end of file
+	FDynamicMeshDrawCommandStorage& DynamicMeshDrawCommandStorage);
+
+RENDERER_API extern FMeshDrawCommandSortKey CalculateMeshStaticSortKey(const FMeshMaterialShader* VertexShader, const FMeshMaterialShader* PixelShader);
\ No newline at end of file
