Moved all sort keys into FMeshDrawCommandSortKey union for easier debugging and cleaner code.

diff --git a/Engine/Source/Runtime/Renderer/Private/BasePassRendering.cpp b/Engine/Source/Runtime/Renderer/Private/BasePassRendering.cpp
--- a/Engine/Source/Runtime/Renderer/Private/BasePassRendering.cpp
+++ b/Engine/Source/Runtime/Renderer/Private/BasePassRendering.cpp
@@ -237,44 +237,28 @@ FMeshDrawCommandSortKey CalculateTranslucentMeshStaticSortKey(const FPrimitiveSc
 		SortKeyPriority = (uint16)((int32)PrimitiveSceneInfo->Proxy->GetTranslucencySortPriority() - (int32)SHRT_MIN);
 	}
 
-	FTranslucentMeshSortKey TranslucentMeshSortKey;
-	TranslucentMeshSortKey.Fields.MeshIdInPrimitive = MeshIdInPrimitive;
-	TranslucentMeshSortKey.Fields.Priority = SortKeyPriority;
-	TranslucentMeshSortKey.Fields.Distance = 0; // View specific, so will be filled later inside VisibleMeshCommands.
-
 	FMeshDrawCommandSortKey SortKey;
-	SortKey.PackedData = TranslucentMeshSortKey.PackedData;
+	SortKey.Translucent.MeshIdInPrimitive = MeshIdInPrimitive;
+	SortKey.Translucent.Priority = SortKeyPriority;
+	SortKey.Translucent.Distance = 0; // View specific, so will be filled later inside VisibleMeshCommands.
+
 	return SortKey;
 }
 
 FMeshDrawCommandSortKey CalculateBasePassMeshStaticSortKey(EDepthDrawingMode EarlyZPassMode, EBlendMode BlendMode, const FMeshMaterialShader* VertexShader, const FMeshMaterialShader* PixelShader)
 {
-	union FBasePassMeshSortKey
-	{
-		uint64 PackedData;
-
-		struct
-		{
-			uint64 VertexShaderHash	: 16;	// Order by vertex shader's hash.
-			uint64 PixelShaderHash	: 32;	// Order by pixel shader's hash.
-			uint64 Masked			: 16;	// First order by masked.
-		} Fields;
-	};
-
-	FBasePassMeshSortKey BasePassMeshSortKey;
-	BasePassMeshSortKey.Fields.VertexShaderHash = PointerHash(VertexShader) & 0xFFFF;
-	BasePassMeshSortKey.Fields.PixelShaderHash = PointerHash(PixelShader);
+	FMeshDrawCommandSortKey SortKey;
+	SortKey.BasePass.VertexShaderHash = PointerHash(VertexShader) & 0xFFFF;
+	SortKey.BasePass.PixelShaderHash = PointerHash(PixelShader);
 	if (EarlyZPassMode != DDM_None)
 	{
-		BasePassMeshSortKey.Fields.Masked = BlendMode == EBlendMode::BLEND_Masked ? 0 : 1;
+		SortKey.BasePass.Masked = BlendMode == EBlendMode::BLEND_Masked ? 0 : 1;
 	}
 	else
 	{
-		BasePassMeshSortKey.Fields.Masked = BlendMode == EBlendMode::BLEND_Masked ? 1 : 0;
+		SortKey.BasePass.Masked = BlendMode == EBlendMode::BLEND_Masked ? 1 : 0;
 	}
 
-	FMeshDrawCommandSortKey SortKey;
-	SortKey.PackedData = BasePassMeshSortKey.PackedData;
 	return SortKey;
 }
 
diff --git a/Engine/Source/Runtime/Renderer/Private/MeshDrawCommands.cpp b/Engine/Source/Runtime/Renderer/Private/MeshDrawCommands.cpp
--- a/Engine/Source/Runtime/Renderer/Private/MeshDrawCommands.cpp
+++ b/Engine/Source/Runtime/Renderer/Private/MeshDrawCommands.cpp
@@ -147,10 +147,10 @@ void UpdateTranslucentMeshSortKeys(
 		}
 
 		// Patch distance inside translucent mesh sort key.
-		FTranslucentMeshSortKey TranslucentMeshSortKey;
-		TranslucentMeshSortKey.PackedData = VisibleCommand.SortKey.PackedData;
-		TranslucentMeshSortKey.Fields.Distance = (uint32)~BitInvertIfNegativeFloat(*((uint32*)&Distance));
-		VisibleCommand.SortKey.PackedData = TranslucentMeshSortKey.PackedData;		
+		FMeshDrawCommandSortKey SortKey;
+		SortKey.PackedData = VisibleCommand.SortKey.PackedData;
+		SortKey.Translucent.Distance = (uint32)~BitInvertIfNegativeFloat(*((uint32*)&Distance));
+		VisibleCommand.SortKey.PackedData = SortKey.PackedData;
 	}
 }
 
diff --git a/Engine/Source/Runtime/Renderer/Private/MeshPassProcessor.cpp b/Engine/Source/Runtime/Renderer/Private/MeshPassProcessor.cpp
--- a/Engine/Source/Runtime/Renderer/Private/MeshPassProcessor.cpp
+++ b/Engine/Source/Runtime/Renderer/Private/MeshPassProcessor.cpp
@@ -871,23 +871,10 @@ void DrawDynamicMeshPassPrivate(
 
 FMeshDrawCommandSortKey CalculateMeshStaticSortKey(const FMeshMaterialShader* VertexShader, const FMeshMaterialShader* PixelShader)
 {
-	union FMeshSortKey
-	{
-		uint64 PackedData;
-
-		struct
-		{
-			uint64 VertexShaderHash : 32;	// Order by vertex shader's hash.
-			uint64 PixelShaderHash : 32;	// First order by pixel shader's hash.
-		} Fields;
-	};
-
-	FMeshSortKey MeshSortKey;
-	MeshSortKey.Fields.VertexShaderHash = PointerHash(VertexShader);
-	MeshSortKey.Fields.PixelShaderHash = PointerHash(PixelShader);
-
 	FMeshDrawCommandSortKey SortKey;
-	SortKey.PackedData = MeshSortKey.PackedData;
+	SortKey.Generic.VertexShaderHash = PointerHash(VertexShader);
+	SortKey.Generic.PixelShaderHash = PointerHash(PixelShader);
+
 	return SortKey;
 }
 
diff --git a/Engine/Source/Runtime/Renderer/Private/TranslucentRendering.h b/Engine/Source/Runtime/Renderer/Private/TranslucentRendering.h
--- a/Engine/Source/Runtime/Renderer/Private/TranslucentRendering.h
+++ b/Engine/Source/Runtime/Renderer/Private/TranslucentRendering.h
@@ -16,19 +16,4 @@
 
 bool UseNearestDepthNeighborUpsampleForSeparateTranslucency(const FSceneRenderTargets& SceneContext);
 
-EMeshPass::Type TranslucencyPassToMeshPass(ETranslucencyPass::Type TranslucencyPass);
-
-/**
-* Translucent mesh sort key format.
-*/
-union FTranslucentMeshSortKey
-{
-	uint64 PackedData;
-
-	struct
-	{
-		uint64 MeshIdInPrimitive	: 16; // Order meshes belonging to the same primitive by a stable id.
-		uint64 Distance				: 32; // Order by distance.
-		uint64 Priority				: 16; // First order by priority.
-	} Fields;
-};
+EMeshPass::Type TranslucencyPassToMeshPass(ETranslucencyPass::Type TranslucencyPass);
\ No newline at end of file
diff --git a/Engine/Source/Runtime/Renderer/Public/MeshPassProcessor.h b/Engine/Source/Runtime/Renderer/Public/MeshPassProcessor.h
--- a/Engine/Source/Runtime/Renderer/Public/MeshPassProcessor.h
+++ b/Engine/Source/Runtime/Renderer/Public/MeshPassProcessor.h
@@ -453,7 +453,30 @@ private:
 class RENDERER_API FMeshDrawCommandSortKey
 {
 public:
-	uint64 PackedData;
+	union 
+	{
+		uint64 PackedData;
+
+		struct
+		{
+			uint64 VertexShaderHash		: 16; // Order by vertex shader's hash.
+			uint64 PixelShaderHash		: 32; // Order by pixel shader's hash.
+			uint64 Masked				: 16; // First order by masked.
+		} BasePass;
+
+		struct
+		{
+			uint64 MeshIdInPrimitive	: 16; // Order meshes belonging to the same primitive by a stable id.
+			uint64 Distance				: 32; // Order by distance.
+			uint64 Priority				: 16; // First order by priority.
+		} Translucent;
+
+		struct 
+		{
+			uint64 VertexShaderHash : 32;	// Order by vertex shader's hash.
+			uint64 PixelShaderHash : 32;	// First order by pixel shader's hash.
+		} Generic;
+	};
 
 	FORCEINLINE bool operator!=(FMeshDrawCommandSortKey B) const
 	{
