DR - Fix compile for Lumin
#rnx

diff --git a/Engine/Source/Runtime/Renderer/Private/MeshPassProcessor.cpp b/Engine/Source/Runtime/Renderer/Private/MeshPassProcessor.cpp
--- a/Engine/Source/Runtime/Renderer/Private/MeshPassProcessor.cpp
+++ b/Engine/Source/Runtime/Renderer/Private/MeshPassProcessor.cpp
@@ -17,7 +17,7 @@
 TSet<FGraphicsMinimalPipelineStateInitializer> FGraphicsMinimalPipelineStateId::GlobalTable;
 FCriticalSection GlobalTableCriticalSection;
 
-const FMeshDrawCommandSortKey FMeshDrawCommandSortKey::Default = { 0 };
+const FMeshDrawCommandSortKey FMeshDrawCommandSortKey::Default = { {0} };
 
 int32 GEmitMeshDrawEvent = 0;
 static FAutoConsoleVariableRef CVarEmitMeshDrawEvent(
