Enables shadow denoiser by default.

diff --git a/Engine/Source/Runtime/Renderer/Private/RayTracing/RayTracingShadows.cpp b/Engine/Source/Runtime/Renderer/Private/RayTracing/RayTracingShadows.cpp
--- a/Engine/Source/Runtime/Renderer/Private/RayTracing/RayTracingShadows.cpp
+++ b/Engine/Source/Runtime/Renderer/Private/RayTracing/RayTracingShadows.cpp
@@ -29,7 +29,7 @@ static FAutoConsoleVariableRef CVarRayTracingOcclusionSamplesPerPixel(
 
 static TAutoConsoleVariable<int32> CVarShadowUseDenoiser(
 	TEXT("r.Shadow.Denoiser"),
-	0, // TODO: change default to 2.
+	2,
 	TEXT("Choose the denoising algorithm.\n")
 	TEXT(" 0: Disabled (default);\n")
 	TEXT(" 1: Forces the default denoiser of the renderer;\n")
diff --git a/Engine/Source/Runtime/Renderer/Private/ScreenSpaceDenoise.cpp b/Engine/Source/Runtime/Renderer/Private/ScreenSpaceDenoise.cpp
--- a/Engine/Source/Runtime/Renderer/Private/ScreenSpaceDenoise.cpp
+++ b/Engine/Source/Runtime/Renderer/Private/ScreenSpaceDenoise.cpp
@@ -18,18 +18,18 @@
 // ---------------------------------------------------- Cvars
 
 
-static TAutoConsoleVariable<int32> CVarShadowHaveTAA(
-	TEXT("r.Shadow.Denoise.TAA"), 1,
+static TAutoConsoleVariable<int32> CVarShadowTemporalAccumulation(
+	TEXT("r.Shadow.Denoiser.TemporalAccumulation"), 1,
 	TEXT(""),
 	ECVF_RenderThreadSafe);
 
 static TAutoConsoleVariable<int32> CVarHistoryRejection(
-	TEXT("r.Shadow.Denoise.HistoryRejection"), 1,
+	TEXT("r.Shadow.Denoiser.HistoryRejection"), 1,
 	TEXT(""),
 	ECVF_RenderThreadSafe);
 
 static TAutoConsoleVariable<int32> CVarDoDebugPass(
-	TEXT("r.Shadow.Denoise.DebugPass"), 0,
+	TEXT("r.Shadow.Denoiser.DebugPass"), 0,
 	TEXT("Adds denoiser's the debug pass."),
 	ECVF_RenderThreadSafe);
 
@@ -479,7 +479,7 @@ static void DenoiseShadowPenumbra(
 	};
 
 	// Doing mip level.
-	if (View.ViewState && CVarShadowHaveTAA.GetValueOnRenderThread())
+	if (View.ViewState && CVarShadowTemporalAccumulation.GetValueOnRenderThread())
 	{
 		// Generate rejection signal history.
 		FRDGTextureRef HistoryRejectionSignal0 = nullptr;
