SubmitMeshDrawCommands - inlining bDynamicInstancing argument.

diff --git a/Engine/Source/Runtime/Renderer/Private/MeshPassProcessor.cpp b/Engine/Source/Runtime/Renderer/Private/MeshPassProcessor.cpp
--- a/Engine/Source/Runtime/Renderer/Private/MeshPassProcessor.cpp
+++ b/Engine/Source/Runtime/Renderer/Private/MeshPassProcessor.cpp
@@ -822,9 +822,9 @@ void SubmitMeshDrawCommands(
 	const FMeshCommandOneFrameArray& VisibleMeshDrawCommands,
 	FVertexBufferRHIParamRef PrimitiveIdsBuffer,
 	int32 BasePrimitiveIdsOffset,
-	bool bDynamicInstancing,
 	FRHICommandList& RHICmdList)
 {
+	const bool bDynamicInstancing = IsDynamicInstancingEnabled();
 	SubmitMeshDrawCommandsRange(VisibleMeshDrawCommands, PrimitiveIdsBuffer, BasePrimitiveIdsOffset, bDynamicInstancing, 0, VisibleMeshDrawCommands.Num(), RHICmdList);
 }
 
@@ -863,9 +863,7 @@ void DrawDynamicMeshPassPrivate(
 
 		SortAndMergeDynamicPassMeshDrawCommands(View.GetFeatureLevel(), VisibleMeshDrawCommands, DynamicMeshDrawCommandStorage, PrimitiveIdVertexBuffer);
 
-		const bool bDynamicInstancing = IsDynamicInstancingEnabled() && UseGPUScene(GMaxRHIShaderPlatform, View.GetFeatureLevel());
-
-		SubmitMeshDrawCommandsRange(VisibleMeshDrawCommands, PrimitiveIdVertexBuffer, 0, bDynamicInstancing, 0, VisibleMeshDrawCommands.Num(), RHICmdList);
+		SubmitMeshDrawCommandsRange(VisibleMeshDrawCommands, PrimitiveIdVertexBuffer, 0, IsDynamicInstancingEnabled(), 0, VisibleMeshDrawCommands.Num(), RHICmdList);
 	}
 }
 
diff --git a/Engine/Source/Runtime/Renderer/Private/ShadowRendering.cpp b/Engine/Source/Runtime/Renderer/Private/ShadowRendering.cpp
--- a/Engine/Source/Runtime/Renderer/Private/ShadowRendering.cpp
+++ b/Engine/Source/Runtime/Renderer/Private/ShadowRendering.cpp
@@ -612,9 +612,7 @@ void FProjectedShadowInfo::SetupProjectionStencilMask(
 		const FShadowMeshDrawCommandPass& ProjectionStencilingPass = ProjectionStencilingPasses[ViewIndex];
 		if (ProjectionStencilingPass.VisibleMeshDrawCommands.Num() > 0)
 		{
-			const bool bDynamicInstancing = IsDynamicInstancingEnabled() && UseGPUScene(GMaxRHIShaderPlatform, View->GetFeatureLevel());
-
-			SubmitMeshDrawCommands(ProjectionStencilingPass.VisibleMeshDrawCommands, ProjectionStencilingPass.PrimitiveIdVertexBuffer, 0, bDynamicInstancing, RHICmdList);
+			SubmitMeshDrawCommands(ProjectionStencilingPass.VisibleMeshDrawCommands, ProjectionStencilingPass.PrimitiveIdVertexBuffer, 0, RHICmdList);
 		}
 
 		// Restore viewport
@@ -754,9 +752,7 @@ void FProjectedShadowInfo::SetupProjectionStencilMask(
 			const FShadowMeshDrawCommandPass& ProjectionStencilingPass = ProjectionStencilingPasses[ViewIndex];
 			if (ProjectionStencilingPass.VisibleMeshDrawCommands.Num() > 0)
 			{
-				const bool bDynamicInstancing = IsDynamicInstancingEnabled() && UseGPUScene(GMaxRHIShaderPlatform, View->GetFeatureLevel());
-
-				SubmitMeshDrawCommands(ProjectionStencilingPass.VisibleMeshDrawCommands, ProjectionStencilingPass.PrimitiveIdVertexBuffer, 0, bDynamicInstancing, RHICmdList);
+				SubmitMeshDrawCommands(ProjectionStencilingPass.VisibleMeshDrawCommands, ProjectionStencilingPass.PrimitiveIdVertexBuffer, 0, RHICmdList);
 			}
 		}
 	}
diff --git a/Engine/Source/Runtime/Renderer/Public/MeshPassProcessor.h b/Engine/Source/Runtime/Renderer/Public/MeshPassProcessor.h
--- a/Engine/Source/Runtime/Renderer/Public/MeshPassProcessor.h
+++ b/Engine/Source/Runtime/Renderer/Public/MeshPassProcessor.h
@@ -993,7 +993,6 @@ extern void SubmitMeshDrawCommands(
 	const FMeshCommandOneFrameArray& VisibleMeshDrawCommands,
 	FVertexBufferRHIParamRef PrimitiveIdsBuffer,
 	int32 BasePrimitiveIdsOffset,
-	bool bDynamicInstancing,
 	FRHICommandList& RHICmdList);
 
 extern void SubmitMeshDrawCommandsRange(
