FTranslucencyDepthPassMeshProcessor for translucency depth rendering.

diff --git a/Engine/Shaders/Private/TranslucentShadowDepthShaders.usf b/Engine/Shaders/Private/TranslucentShadowDepthShaders.usf
--- a/Engine/Shaders/Private/TranslucentShadowDepthShaders.usf
+++ b/Engine/Shaders/Private/TranslucentShadowDepthShaders.usf
@@ -13,24 +13,16 @@ struct FTranslucencyShadowDepthVSToPS
 {
 	FVertexFactoryInterpolantsVSToPS FactoryInterpolants;
 	float ShadowDepth : TEXCOORD6;
-	float4 PixelPosition	: TEXCOORD7;
+	float4 PixelPosition : TEXCOORD7;
 };
 
-float4x4 ProjectionMatrix;
-
-//@todo - make this a bool, bool shader parameters currently do not work in vertex shaders on Xbox 360 (TTP 125134)
-float bClampToNearPlane;
-
-// x: Tweakable depth bias, y:InvMaxSubjectDepthUsed to normalize the outputted depth
-float2 ShadowParams;
-
 void SetShadowDepthOutputs(float4 WorldPosition, out float4 OutPosition, out float ShadowDepth)
 {
-	OutPosition = mul(WorldPosition,ProjectionMatrix);
+	OutPosition = mul(WorldPosition, TranslucentDepthPass.ProjectionMatrix);
 
 	// Clamp the vertex to the near plane if it is in front of the near plane
 	// This has problems if some vertices of a triangle get clamped and others do not, also causes artifacts with non-ortho projections
-	if (bClampToNearPlane && OutPosition.z < 0)
+	if (TranslucentDepthPass.bClampToNearPlane && OutPosition.z < 0)
 	{
 		OutPosition.z = 0.000001f;
 		OutPosition.w = 1.0f;
@@ -39,7 +31,7 @@ void SetShadowDepthOutputs(float4 WorldPosition, out float4 OutPosition, out flo
 	#if PERSPECTIVE_CORRECT_DEPTH
 		ShadowDepth = OutPosition.z;
 	#else
-		float InvMaxSubjectDepth = ShadowParams.y;
+		float InvMaxSubjectDepth = TranslucentDepthPass.InvMaxSubjectDepth;
 
 		// Output linear, normalized depth
 		ShadowDepth = OutPosition.z * InvMaxSubjectDepth;
@@ -73,8 +65,6 @@ void MainVS(
 }
 
 float TranslucentShadowStartOffset;
-/** Used to normalize the outputted depth */
-float TranslInvMaxSubjectDepth;
 
 /** Pixel shader used to accumulate layer opacities in different channels based on the first translucent layer's depth. */
 void MainOpacityPS( 
@@ -96,7 +86,8 @@ void MainOpacityPS(
 	float Density = GetMaterialOpacity(PixelMaterialInputs) * GetMaterialTranslucentShadowDensityScale();
 
 #if PERSPECTIVE_CORRECT_DEPTH
-	Inputs.ShadowDepth *= TranslInvMaxSubjectDepth;
+	/** Used to normalize the outputted depth */
+	Inputs.ShadowDepth *= TranslucentDepthPass.InvMaxSubjectDepth;
 #endif
 
 	Inputs.ShadowDepth += TranslucentShadowStartOffset;
diff --git a/Engine/Source/Runtime/Renderer/Private/TranslucentLighting.cpp b/Engine/Source/Runtime/Renderer/Private/TranslucentLighting.cpp
--- a/Engine/Source/Runtime/Renderer/Private/TranslucentLighting.cpp
+++ b/Engine/Source/Runtime/Renderer/Private/TranslucentLighting.cpp
@@ -44,6 +44,7 @@
 #include "VolumeLighting.h"
 #include "PipelineStateCache.h"
 #include "VisualizeTexture.h"
+#include "MeshPassProcessor.inl"
 
 class FMaterial;
 
@@ -238,15 +239,55 @@ private:
 	FShaderParameter ClampToNearPlane;
 };
 
+class FTranslucencyDepthShaderElementData : public FMeshMaterialShaderElementData
+{
+public:
+
+	float TranslucentShadowStartOffset;
+};
+
+BEGIN_GLOBAL_SHADER_PARAMETER_STRUCT(FTranslucencyDepthPassUniformParameters,)
+	SHADER_PARAMETER_STRUCT(FSceneTexturesUniformParameters, SceneTextures)
+	SHADER_PARAMETER(FMatrix, ProjectionMatrix)
+	SHADER_PARAMETER(float, bClampToNearPlane)
+	SHADER_PARAMETER(float, InvMaxSubjectDepth)
+	SHADER_PARAMETER_STRUCT(FTranslucentSelfShadowUniformParameters, TranslucentSelfShadow)
+END_GLOBAL_SHADER_PARAMETER_STRUCT()
+
+IMPLEMENT_GLOBAL_SHADER_PARAMETER_STRUCT(FTranslucencyDepthPassUniformParameters, "TranslucentDepthPass");
+
+void SetupTranslucencyDepthPassUniformBuffer(
+	const FProjectedShadowInfo* ShadowInfo,
+	FRHICommandList& RHICmdList,
+	const FViewInfo& View,
+	FTranslucencyDepthPassUniformParameters& TranslucencyDepthPassParameters)
+{
+	// Note - scene depth can be bound by the material for use in depth fades
+	// This is incorrect when rendering a shadowmap as it's not from the camera's POV
+	// Set the scene depth texture to something safe when rendering shadow depths
+	FSceneRenderTargets& SceneRenderTargets = FSceneRenderTargets::Get(RHICmdList);
+	SetupSceneTextureUniformParameters(SceneRenderTargets, View.FeatureLevel, ESceneTextureSetupMode::None, TranslucencyDepthPassParameters.SceneTextures);
+
+	TranslucencyDepthPassParameters.ProjectionMatrix = FTranslationMatrix(ShadowInfo->PreShadowTranslation - View.ViewMatrices.GetPreViewTranslation()) * ShadowInfo->SubjectAndReceiverMatrix;
+
+	// Only clamp vertices to the near plane when rendering whole scene directional light shadow depths or preshadows from directional lights
+	const bool bClampToNearPlaneValue = ShadowInfo->IsWholeSceneDirectionalShadow() || (ShadowInfo->bPreShadow && ShadowInfo->bDirectionalLight);
+	TranslucencyDepthPassParameters.bClampToNearPlane = bClampToNearPlaneValue ? 1.0f : 0.0f;
+
+	TranslucencyDepthPassParameters.InvMaxSubjectDepth = ShadowInfo->InvMaxSubjectDepth;
+
+	SetupTranslucentSelfShadowUniformParameters(ShadowInfo, TranslucencyDepthPassParameters.TranslucentSelfShadow);
+}
+
 /**
- * Vertex shader used to render shadow maps for translucency.
- */
+* Vertex shader used to render shadow maps for translucency.
+*/
 class FTranslucencyShadowDepthVS : public FMeshMaterialShader
 {
-	DECLARE_SHADER_TYPE(FTranslucencyShadowDepthVS,MeshMaterial);
+	DECLARE_SHADER_TYPE(FTranslucencyShadowDepthVS, MeshMaterial);
 public:
 
-	static bool ShouldCompilePermutation(EShaderPlatform Platform,const FMaterial* Material,const FVertexFactoryType* VertexFactoryType)
+	static bool ShouldCompilePermutation(EShaderPlatform Platform, const FMaterial* Material, const FVertexFactoryType* VertexFactoryType)
 	{
 		return IsTranslucentBlendMode(Material->GetBlendMode()) && IsFeatureLevelSupported(Platform, ERHIFeatureLevel::SM4);
 	}
@@ -255,36 +296,8 @@ public:
 	FTranslucencyShadowDepthVS(const FMeshMaterialShaderType::CompiledShaderInitializerType& Initializer) :
 		FMeshMaterialShader(Initializer)
 	{
-		ShadowParameters.Bind(Initializer.ParameterMap);
-		PassUniformBuffer.Bind(Initializer.ParameterMap, FSceneTexturesUniformParameters::StaticStructMetadata.GetShaderVariableName());
-	}
-
-	virtual bool Serialize(FArchive& Ar) override
-	{
-		bool bShaderHasOutdatedParameters = FMeshMaterialShader::Serialize(Ar);
-		Ar << ShadowParameters;
-		return bShaderHasOutdatedParameters;
-	}
-	/*
-	void SetParameters(
-		FRHICommandList& RHICmdList, 
-		const FMaterialRenderProxy* MaterialRenderProxy,
-		const FSceneView& View,
-		const FProjectedShadowInfo* ShadowInfo, 
-		const FMeshPassProcessorRenderState& DrawRenderState
-		)
-	{
-		FMeshMaterialShader::SetParameters(RHICmdList, GetVertexShader(), MaterialRenderProxy, *MaterialRenderProxy->GetMaterial(View.GetFeatureLevel()), View, DrawRenderState.GetViewUniformBuffer(), DrawRenderState.GetPassUniformBuffer());
-		ShadowParameters.SetVertexShader(RHICmdList, this, View, ShadowInfo, MaterialRenderProxy);
-	}
-
-	void SetMesh(FRHICommandList& RHICmdList, const FVertexFactory* VertexFactory,const FSceneView& View,const FPrimitiveSceneProxy* Proxy,const FMeshBatchElement& BatchElement,const FMeshPassProcessorRenderState& DrawRenderState)
-	{
-		FMeshMaterialShader::SetMesh(RHICmdList, GetVertexShader(),VertexFactory,View,Proxy,BatchElement,DrawRenderState);
+		PassUniformBuffer.Bind(Initializer.ParameterMap, FTranslucencyDepthPassUniformParameters::StaticStructMetadata.GetShaderVariableName());
 	}
-	*/
-private:
-	FShadowDepthShaderParameters ShadowParameters;
 };
 
 enum ETranslucencyShadowDepthShaderMode
@@ -296,17 +309,17 @@ enum ETranslucencyShadowDepthShaderMode
 template <ETranslucencyShadowDepthShaderMode ShaderMode>
 class TTranslucencyShadowDepthVS : public FTranslucencyShadowDepthVS
 {
-	DECLARE_SHADER_TYPE(TTranslucencyShadowDepthVS,MeshMaterial);
+	DECLARE_SHADER_TYPE(TTranslucencyShadowDepthVS, MeshMaterial);
 public:
 
-	TTranslucencyShadowDepthVS(const ShaderMetaType::CompiledShaderInitializerType& Initializer):
+	TTranslucencyShadowDepthVS(const ShaderMetaType::CompiledShaderInitializerType& Initializer) :
 		FTranslucencyShadowDepthVS(Initializer)
 	{
 	}
 
 	TTranslucencyShadowDepthVS() {}
 
-	static void ModifyCompilationEnvironment( EShaderPlatform Platform, const FMaterial* Material, FShaderCompilerEnvironment& OutEnvironment )
+	static void ModifyCompilationEnvironment(EShaderPlatform Platform, const FMaterial* Material, FShaderCompilerEnvironment& OutEnvironment)
 	{
 		FTranslucencyShadowDepthVS::ModifyCompilationEnvironment(Platform, Material, OutEnvironment);
 		OutEnvironment.SetDefine(TEXT("PERSPECTIVE_CORRECT_DEPTH"), (uint32)(ShaderMode == TranslucencyShadowDepth_PerspectiveCorrect ? 1 : 0));
@@ -332,52 +345,36 @@ public:
 	FTranslucencyShadowDepthPS(const ShaderMetaType::CompiledShaderInitializerType& Initializer):
 		FMeshMaterialShader(Initializer)
 	{
-		TranslInvMaxSubjectDepth.Bind(Initializer.ParameterMap,TEXT("TranslInvMaxSubjectDepth"));
-		TranslucentShadowStartOffset.Bind(Initializer.ParameterMap,TEXT("TranslucentShadowStartOffset"));
-		PassUniformBuffer.Bind(Initializer.ParameterMap, FSceneTexturesUniformParameters::StaticStructMetadata.GetShaderVariableName());
+		PassUniformBuffer.Bind(Initializer.ParameterMap, FTranslucencyDepthPassUniformParameters::StaticStructMetadata.GetShaderVariableName());
+
+		TranslucentShadowStartOffset.Bind(Initializer.ParameterMap, TEXT("TranslucentShadowStartOffset"));
 	}
 
 	FTranslucencyShadowDepthPS() {}
-	/*
-	void SetParameters(
-		FRHICommandList& RHICmdList, 
-		const FMaterialRenderProxy* MaterialRenderProxy,
-		const FSceneView& View,
-		const FProjectedShadowInfo* ShadowInfo, 
-		const FMeshPassProcessorRenderState& DrawRenderState
-		)
-	{
-		const FPixelShaderRHIParamRef ShaderRHI = GetPixelShader();
-		const auto FeatureLevel = View.GetFeatureLevel();
-
-		FMeshMaterialShader::SetParameters(RHICmdList, ShaderRHI, MaterialRenderProxy, *MaterialRenderProxy->GetMaterial(FeatureLevel), View, DrawRenderState.GetViewUniformBuffer(), DrawRenderState.GetPassUniformBuffer());
 
-		SetShaderValue(RHICmdList, ShaderRHI, TranslInvMaxSubjectDepth, ShadowInfo->InvMaxSubjectDepth);
-
-		const float LocalToWorldScale = ShadowInfo->GetParentSceneInfo()->Proxy->GetLocalToWorld().GetScaleVector().GetMax();
-		const float TranslucentShadowStartOffsetValue = MaterialRenderProxy->GetMaterial(FeatureLevel)->GetTranslucentShadowStartOffset() * LocalToWorldScale;
-		SetShaderValue(RHICmdList, ShaderRHI,TranslucentShadowStartOffset, TranslucentShadowStartOffsetValue / (ShadowInfo->MaxSubjectZ - ShadowInfo->MinSubjectZ));
+	void GetShaderBindings(
+		const FScene* Scene,
+		ERHIFeatureLevel::Type FeatureLevel,
+		const FPrimitiveSceneProxy* PrimitiveSceneProxy,
+		const FMaterialRenderProxy& MaterialRenderProxy,
+		const FMaterial& Material,
+		const FMeshPassProcessorRenderState& DrawRenderState,
+		const FTranslucencyDepthShaderElementData& ShaderElementData,
+		FMeshDrawSingleShaderBindings& ShaderBindings) const
+	{
+		FMeshMaterialShader::GetShaderBindings(Scene, FeatureLevel, PrimitiveSceneProxy, MaterialRenderProxy, Material, DrawRenderState, ShaderElementData, ShaderBindings);
 
-		FTranslucentSelfShadowUniformParameters TranslucentSelfShadowUniformParameters;
-		SetupTranslucentSelfShadowUniformParameters(ShadowInfo, TranslucentSelfShadowUniformParameters);
-		SetUniformBufferParameterImmediate(RHICmdList, ShaderRHI, GetUniformBufferParameter<FTranslucentSelfShadowUniformParameters>(), TranslucentSelfShadowUniformParameters);
+		ShaderBindings.Add(TranslucentShadowStartOffset, ShaderElementData.TranslucentShadowStartOffset);
 	}
 
-	void SetMesh(FRHICommandList& RHICmdList, const FVertexFactory* VertexFactory,const FSceneView& View,const FPrimitiveSceneProxy* Proxy,const FMeshBatchElement& BatchElement,const FMeshPassProcessorRenderState& DrawRenderState)
-	{
-		FMeshMaterialShader::SetMesh(RHICmdList, GetPixelShader(),VertexFactory,View,Proxy,BatchElement,DrawRenderState);
-	}
-	*/
 	virtual bool Serialize(FArchive& Ar) override
 	{
 		bool bShaderHasOutdatedParameters = FMeshMaterialShader::Serialize(Ar);
-		Ar << TranslInvMaxSubjectDepth;
 		Ar << TranslucentShadowStartOffset;
 		return bShaderHasOutdatedParameters;
 	}
 
 private:
-	FShaderParameter TranslInvMaxSubjectDepth;
 	FShaderParameter TranslucentShadowStartOffset;
 };
 
@@ -404,190 +401,131 @@ public:
 IMPLEMENT_MATERIAL_SHADER_TYPE(template<>,TTranslucencyShadowDepthPS<TranslucencyShadowDepth_PerspectiveCorrect>,TEXT("/Engine/Private/TranslucentShadowDepthShaders.usf"),TEXT("MainOpacityPS"),SF_Pixel);
 IMPLEMENT_MATERIAL_SHADER_TYPE(template<>,TTranslucencyShadowDepthPS<TranslucencyShadowDepth_Standard>,TEXT("/Engine/Private/TranslucentShadowDepthShaders.usf"),TEXT("MainOpacityPS"),SF_Pixel);
 
-//@todo MeshCommandPipeline
-#if 0
-
-/**
- * Drawing policy used to create Fourier opacity maps
- */
-class FTranslucencyShadowDepthDrawingPolicy : public FMeshDrawingPolicy
+class FTranslucencyDepthPassMeshProcessor : public FMeshPassProcessor
 {
 public:
+	FTranslucencyDepthPassMeshProcessor(const FScene* Scene,
+		const FSceneView* InViewIfDynamicMeshCommand,
+		const FMeshPassProcessorRenderState& InPassDrawRenderState,
+		const FProjectedShadowInfo* InShadowInfo,
+		FMeshPassDrawListContext* InDrawListContext);
 
-	struct ContextDataType : public FMeshDrawingPolicy::ContextDataType
-	{
-		const FProjectedShadowInfo* ShadowInfo;
-
-		explicit ContextDataType(const FProjectedShadowInfo* InShadowInfo)
-			: ShadowInfo(InShadowInfo)
-		{}
-	};
-
-	FTranslucencyShadowDepthDrawingPolicy(
-		const FVertexFactory* InVertexFactory,
-		const FMaterialRenderProxy* InMaterialRenderProxy,
-		const FMaterial& InMaterialResource,
-		const FMeshDrawingPolicyOverrideSettings& InOverrideSettings,
-		bool bInDirectionalLight
-		):
-		FMeshDrawingPolicy(InVertexFactory,InMaterialRenderProxy,InMaterialResource,InOverrideSettings)
-	{
-		const bool bUsePerspectiveCorrectShadowDepths = !bInDirectionalLight;
-
-		if (bUsePerspectiveCorrectShadowDepths)
-		{
-			VertexShader = InMaterialResource.GetShader<TTranslucencyShadowDepthVS<TranslucencyShadowDepth_PerspectiveCorrect> >(InVertexFactory->GetType());
-			PixelShader = InMaterialResource.GetShader<TTranslucencyShadowDepthPS<TranslucencyShadowDepth_PerspectiveCorrect> >(InVertexFactory->GetType());
-		}
-		else
-		{
-			VertexShader = InMaterialResource.GetShader<TTranslucencyShadowDepthVS<TranslucencyShadowDepth_Standard> >(InVertexFactory->GetType());
-			PixelShader = InMaterialResource.GetShader<TTranslucencyShadowDepthPS<TranslucencyShadowDepth_Standard> >(InVertexFactory->GetType());
-		}
-		BaseVertexShader = VertexShader;
-	}
-
-	void SetSharedState(FRHICommandList& RHICmdList, const FMeshPassProcessorRenderState& DrawRenderState, const FSceneView* View, const ContextDataType PolicyContext) const
-	{
-		// Set the shared mesh resources.
-		FMeshDrawingPolicy::SetSharedState(RHICmdList, DrawRenderState, View, PolicyContext);
-
-		VertexShader->SetParameters(RHICmdList, MaterialRenderProxy, *View, PolicyContext.ShadowInfo, DrawRenderState);
-		PixelShader->SetParameters(RHICmdList, MaterialRenderProxy, *View, PolicyContext.ShadowInfo, DrawRenderState);
-	}
-
-	/** 
-	* Create bound shader state using the vertex decl from the mesh draw policy
-	* as well as the shaders needed to draw the mesh
-	* @return new bound shader state object
-	*/
-	FBoundShaderStateInput GetBoundShaderStateInput(ERHIFeatureLevel::Type InFeatureLevel) const
-	{
-		return FBoundShaderStateInput(
-			FMeshDrawingPolicy::GetVertexDeclaration(), 
-			VertexShader->GetVertexShader(),
-			NULL, 
-			NULL,
-			PixelShader->GetPixelShader(),
-			NULL);
-	}
-
-	void SetMeshRenderState(
-		FRHICommandList& RHICmdList, 
-		const FSceneView& View,
-		const FPrimitiveSceneProxy* PrimitiveSceneProxy,
-		const FMeshBatch& Mesh,
-		int32 BatchElementIndex,
-		const FMeshPassProcessorRenderState& DrawRenderState,
-		const ElementDataType& ElementData,
-		const ContextDataType PolicyContext
-		) const
-	{
-		const FMeshBatchElement& BatchElement = Mesh.Elements[BatchElementIndex];
-		VertexShader->SetMesh(RHICmdList, VertexFactory,View,PrimitiveSceneProxy,BatchElement,DrawRenderState);
-		PixelShader->SetMesh(RHICmdList, VertexFactory,View,PrimitiveSceneProxy,BatchElement,DrawRenderState);
-		SetPrimitiveIdStream(RHICmdList, View, PrimitiveSceneProxy, BatchElement.PrimitiveIdMode, BatchElement.DynamicPrimitiveShaderDataIndex);
-	}
+	virtual void AddMeshBatch(const FMeshBatch& RESTRICT MeshBatch, uint64 BatchElementMask, const FPrimitiveSceneProxy* RESTRICT PrimitiveSceneProxy, int32 StaticMeshId = -1) override final;
 
 private:
-	FTranslucencyShadowDepthVS* VertexShader;
-	FTranslucencyShadowDepthPS* PixelShader;
+
+	template<ETranslucencyShadowDepthShaderMode ShaderMode>
+	void Process(
+		const FMeshBatch& MeshBatch,
+		uint64 BatchElementMask,
+		int32 StaticMeshId,
+		const FPrimitiveSceneProxy* RESTRICT PrimitiveSceneProxy,
+		const FMaterialRenderProxy& RESTRICT MaterialRenderProxy,
+		const FMaterial& RESTRICT MaterialResource,
+		float MaterialTranslucentShadowStartOffset,
+		ERasterizerFillMode MeshFillMode,
+		ERasterizerCullMode MeshCullMode);
+
+	FMeshPassProcessorRenderState PassDrawRenderState;
+	const FProjectedShadowInfo* ShadowInfo;
+	FShadowDepthType ShadowDepthType;
+	const bool bDirectionalLight;
 };
 
-class FTranslucencyShadowDepthDrawingPolicyFactory
+FTranslucencyDepthPassMeshProcessor::FTranslucencyDepthPassMeshProcessor(const FScene* Scene,
+	const FSceneView* InViewIfDynamicMeshCommand,
+	const FMeshPassProcessorRenderState& InPassDrawRenderState,
+	const FProjectedShadowInfo* InShadowInfo,
+	FMeshPassDrawListContext* InDrawListContext)
+	: FMeshPassProcessor(Scene, Scene->GetFeatureLevel(), InViewIfDynamicMeshCommand, InDrawListContext)
+	, PassDrawRenderState(InPassDrawRenderState)
+	, ShadowInfo(InShadowInfo)
+	, ShadowDepthType(InShadowInfo->GetShadowDepthType())
+	, bDirectionalLight(InShadowInfo->bDirectionalLight)
 {
-public:
-
-	enum { bAllowSimpleElements = false };
-	struct ContextType 
-	{
-		ContextType(const FProjectedShadowInfo* InShadowInfo, bool bInDirectionalLight)
-			: ShadowInfo(InShadowInfo)
-			, bDirectionalLight(bInDirectionalLight)
-		{}
+}
 
-		const FProjectedShadowInfo* ShadowInfo;
-		bool bDirectionalLight;
-	};
+template<ETranslucencyShadowDepthShaderMode ShaderMode>
+void FTranslucencyDepthPassMeshProcessor::Process(
+	const FMeshBatch& RESTRICT MeshBatch,
+	uint64 BatchElementMask,
+	int32 StaticMeshId,
+	const FPrimitiveSceneProxy* RESTRICT PrimitiveSceneProxy,
+	const FMaterialRenderProxy& RESTRICT MaterialRenderProxy,
+	const FMaterial& RESTRICT MaterialResource,
+	float MaterialTranslucentShadowStartOffset,
+	ERasterizerFillMode MeshFillMode,
+	ERasterizerCullMode MeshCullMode)
+{
+	const FVertexFactory* VertexFactory = MeshBatch.VertexFactory;
+
+	TMeshProcessorShaders<
+		TTranslucencyShadowDepthVS<ShaderMode>,
+		FMeshMaterialShader,
+		FMeshMaterialShader,
+		TTranslucencyShadowDepthPS<ShaderMode>> PassShaders;
+
+	PassShaders.VertexShader = MaterialResource.GetShader<TTranslucencyShadowDepthVS<ShaderMode> >(VertexFactory->GetType());
+	PassShaders.PixelShader = MaterialResource.GetShader<TTranslucencyShadowDepthPS<ShaderMode> >(VertexFactory->GetType());
+
+	FMeshPassProcessorRenderState DrawRenderState(PassDrawRenderState);
+
+	FTranslucencyDepthShaderElementData ShaderElementData;
+	ShaderElementData.InitializeMeshMaterialData(ViewIfDynamicMeshCommand, PrimitiveSceneProxy, MeshBatch, StaticMeshId, false);
+
+	const float LocalToWorldScale = ShadowInfo->GetParentSceneInfo()->Proxy->GetLocalToWorld().GetScaleVector().GetMax();
+	const float TranslucentShadowStartOffsetValue = MaterialTranslucentShadowStartOffset * LocalToWorldScale;
+	ShaderElementData.TranslucentShadowStartOffset = TranslucentShadowStartOffsetValue / (ShadowInfo->MaxSubjectZ - ShadowInfo->MinSubjectZ);
+
+	const FMeshDrawCommandSortKey SortKey = CalculateMeshStaticSortKey(PassShaders.VertexShader, PassShaders.PixelShader);
+
+	const uint32 InstanceFactor = 1;
+
+	BuildMeshDrawCommands(
+		MeshBatch,
+		BatchElementMask,
+		PrimitiveSceneProxy,
+		MaterialRenderProxy,
+		MaterialResource,
+		DrawRenderState,
+		PassShaders,
+		MeshFillMode,
+		MeshCullMode,
+		InstanceFactor,
+		SortKey,
+		EMeshPassFeatures::Default,
+		ShaderElementData);
+}
 
-	static bool DrawDynamicMesh(
-		FRHICommandList& RHICmdList, 
-		const FSceneView& View,
-		ContextType DrawingContext,
-		const FMeshBatch& Mesh,
-		bool bPreFog,
-		const FMeshPassProcessorRenderState& DrawRenderState,
-		const FPrimitiveSceneProxy* PrimitiveSceneProxy,
-		FHitProxyId HitProxyId
-		)
+void FTranslucencyDepthPassMeshProcessor::AddMeshBatch(const FMeshBatch& RESTRICT MeshBatch, uint64 BatchElementMask, const FPrimitiveSceneProxy* RESTRICT PrimitiveSceneProxy, int32 StaticMeshId)
+{
+	if (MeshBatch.CastShadow)
 	{
-		bool bDirty = false;
-		const auto FeatureLevel = View.GetFeatureLevel();
+		// Determine the mesh's material and blend mode.
+		const FMaterialRenderProxy* FallbackMaterialRenderProxyPtr = nullptr;
+		const FMaterial& Material = MeshBatch.MaterialRenderProxy->GetMaterialWithFallback(FeatureLevel, FallbackMaterialRenderProxyPtr);
+		const FMaterialRenderProxy& MaterialRenderProxy = FallbackMaterialRenderProxyPtr ? *FallbackMaterialRenderProxyPtr : *MeshBatch.MaterialRenderProxy;
 
-		if (Mesh.CastShadow)
-		{
-			const FMaterialRenderProxy* MaterialRenderProxy = Mesh.MaterialRenderProxy;
-			const FMaterial* Material = MaterialRenderProxy->GetMaterial(FeatureLevel);
-			const EBlendMode BlendMode = Material->GetBlendMode();
+		const EBlendMode BlendMode = Material.GetBlendMode();
+		const float MaterialTranslucentShadowStartOffset = Material.GetTranslucentShadowStartOffset();
+		const ERasterizerFillMode MeshFillMode = ComputeMeshFillMode(MeshBatch, Material);
+		const ERasterizerCullMode MeshCullMode = ComputeMeshCullMode(MeshBatch, Material);
+		const bool bIsTranslucent = IsTranslucentBlendMode(BlendMode);
 
-			// Only render translucent meshes into the Fourier opacity maps
-			if (IsTranslucentBlendMode(BlendMode) && ShouldIncludeDomainInMeshPass(Material->GetMaterialDomain()))
+		// Only render translucent meshes into the Fourier opacity maps
+		if (bIsTranslucent && ShouldIncludeDomainInMeshPass(Material.GetMaterialDomain()))
+		{
+			if (bDirectionalLight)
 			{
-				FTranslucencyShadowDepthDrawingPolicy DrawingPolicy(Mesh.VertexFactory, MaterialRenderProxy, *MaterialRenderProxy->GetMaterial(FeatureLevel), ComputeMeshOverrideSettings(Mesh), DrawingContext.bDirectionalLight);
-
-				FMeshPassProcessorRenderState DrawRenderStateLocal(DrawRenderState);
-				DrawRenderStateLocal.SetDitheredLODTransitionAlpha(Mesh.DitheredLODTransitionAlpha);
-				DrawingPolicy.SetupPipelineState(DrawRenderStateLocal, View);
-				CommitGraphicsPipelineState(RHICmdList, DrawingPolicy, DrawRenderStateLocal, DrawingPolicy.GetBoundShaderStateInput(View.GetFeatureLevel()), DrawingPolicy.GetMaterialRenderProxy());
-				DrawingPolicy.SetSharedState(RHICmdList, DrawRenderStateLocal, &View, FTranslucencyShadowDepthDrawingPolicy::ContextDataType(DrawingContext.ShadowInfo));
-
-				for (int32 BatchElementIndex = 0; BatchElementIndex < Mesh.Elements.Num(); BatchElementIndex++)
-				{
-					TDrawEvent<FRHICommandList> MeshEvent;
-					BeginMeshDrawEvent(RHICmdList, PrimitiveSceneProxy, Mesh, MeshEvent, EnumHasAnyFlags(EShowMaterialDrawEventTypes(GShowMaterialDrawEventTypes), EShowMaterialDrawEventTypes::TranslucentLighting));
-
-					DrawingPolicy.SetMeshRenderState(RHICmdList, View,PrimitiveSceneProxy,Mesh,BatchElementIndex,DrawRenderStateLocal,
-						FTranslucencyShadowDepthDrawingPolicy::ElementDataType(),
-						FTranslucencyShadowDepthDrawingPolicy::ContextDataType(DrawingContext.ShadowInfo)
-						);
-					DrawingPolicy.DrawMesh(RHICmdList,View,Mesh,BatchElementIndex);
-				}
-				bDirty = true;
+				Process<TranslucencyShadowDepth_Standard>(MeshBatch, BatchElementMask, StaticMeshId, PrimitiveSceneProxy, MaterialRenderProxy, Material, MaterialTranslucentShadowStartOffset, MeshFillMode, MeshCullMode);
+			}
+			else
+			{
+				Process<TranslucencyShadowDepth_PerspectiveCorrect>(MeshBatch, BatchElementMask, StaticMeshId, PrimitiveSceneProxy, MaterialRenderProxy, Material, MaterialTranslucentShadowStartOffset, MeshFillMode, MeshCullMode);
 			}
 		}
-	
-		return bDirty;
 	}
-
-	static bool DrawStaticMesh(
-		FRHICommandList& RHICmdList, 
-		const FViewInfo& View,
-		ContextType DrawingContext,
-		const FStaticMeshBatch& StaticMesh,
-		bool bPreFog,
-		const FMeshPassProcessorRenderState& DrawRenderState,
-		const FPrimitiveSceneProxy* PrimitiveSceneProxy,
-		FHitProxyId HitProxyId
-		)
-	{
-		bool bDirty = false;
-
-		bDirty |= DrawDynamicMesh(
-			RHICmdList, 
-			View,
-			DrawingContext,
-			StaticMesh,
-			bPreFog,
-			DrawRenderState,
-			PrimitiveSceneProxy,
-			HitProxyId
-			);
-
-		return bDirty;
-	}
-};
-
-#endif
+}
 
 void FProjectedShadowInfo::RenderTranslucencyDepths(FRHICommandList& RHICmdList, FSceneRenderer* SceneRenderer)
 {
@@ -596,13 +534,9 @@ void FProjectedShadowInfo::RenderTranslucencyDepths(FRHICommandList& RHICmdList,
 	checkSlow(!bWholeSceneShadow);
 	SCOPE_CYCLE_COUNTER(STAT_RenderPerObjectShadowDepthsTime);
 
-	// Note - scene depth can be bound by the material for use in depth fades
-	// This is incorrect when rendering a shadowmap as it's not from the camera's POV
-	// Set the scene depth texture to something safe when rendering shadow depths
-	FSceneTexturesUniformParameters SceneTextureParameters;
-	FSceneRenderTargets& SceneContext = FSceneRenderTargets::Get(RHICmdList);
-	SetupSceneTextureUniformParameters(SceneContext, ShadowDepthView->FeatureLevel, ESceneTextureSetupMode::None, SceneTextureParameters);
-	TUniformBufferRef<FSceneTexturesUniformParameters> PassUniformBuffer = TUniformBufferRef<FSceneTexturesUniformParameters>::CreateUniformBufferImmediate(SceneTextureParameters, UniformBuffer_SingleFrame);	
+	FTranslucencyDepthPassUniformParameters TranslucencyDepthPassParameters;
+	SetupTranslucencyDepthPassUniformBuffer(this, RHICmdList, *ShadowDepthView, TranslucencyDepthPassParameters);
+	TUniformBufferRef<FTranslucencyDepthPassUniformParameters> PassUniformBuffer = TUniformBufferRef<FTranslucencyDepthPassUniformParameters>::CreateUniformBufferImmediate(TranslucencyDepthPassParameters, UniformBuffer_SingleFrame, EUniformBufferValidation::None);
 
 	FMeshPassProcessorRenderState DrawRenderState(*ShadowDepthView, PassUniformBuffer);
 	{
@@ -642,15 +576,23 @@ void FProjectedShadowInfo::RenderTranslucencyDepths(FRHICommandList& RHICmdList,
 			CW_RGBA, BO_Add, BF_One, BF_One, BO_Add, BF_One, BF_One,
 			CW_RGBA, BO_Add, BF_One, BF_One, BO_Add, BF_One, BF_One>::GetRHI());
 
-//@todo MeshCommandPipeline
-#if 0
-		FTranslucencyShadowDepthDrawingPolicyFactory::ContextType DrawingContext(this,bDirectionalLight);
+
+		FMeshCommandOneFrameArray VisibleMeshDrawCommands;
+		FDynamicPassMeshDrawListContext TranslucencyDepthContext(DynamicMeshDrawCommandStorage, VisibleMeshDrawCommands);
+
+		FTranslucencyDepthPassMeshProcessor TranslucencyDepthPassMeshProcessor(
+			SceneRenderer->Scene,
+			ShadowDepthView,
+			DrawRenderState,
+			this,
+			&TranslucencyDepthContext);
 
 		for (int32 MeshBatchIndex = 0; MeshBatchIndex < DynamicSubjectTranslucentMeshElements.Num(); MeshBatchIndex++)
 		{
-			const FMeshBatchAndRelevance& MeshBatchAndRelevance = DynamicSubjectTranslucentMeshElements[MeshBatchIndex];
-			const FMeshBatch& MeshBatch = *MeshBatchAndRelevance.Mesh;
-			FTranslucencyShadowDepthDrawingPolicyFactory::DrawDynamicMesh(RHICmdList, *ShadowDepthView, DrawingContext, MeshBatch, true, DrawRenderState, MeshBatchAndRelevance.PrimitiveSceneProxy, MeshBatch.BatchHitProxyId);
+			const FMeshBatchAndRelevance& MeshAndRelevance = DynamicSubjectTranslucentMeshElements[MeshBatchIndex];
+			check(!MeshAndRelevance.Mesh->bRequiresPerElementVisibility);
+			const uint64 BatchElementMask = ~0ull;
+			TranslucencyDepthPassMeshProcessor.AddMeshBatch(*MeshAndRelevance.Mesh, BatchElementMask, MeshAndRelevance.PrimitiveSceneProxy);
 		}
 
 		for (int32 PrimitiveIndex = 0; PrimitiveIndex < SubjectTranslucentPrimitives.Num(); PrimitiveIndex++)
@@ -669,19 +611,20 @@ void FProjectedShadowInfo::RenderTranslucencyDepths(FRHICommandList& RHICmdList,
 			{
 				for (int32 MeshIndex = 0; MeshIndex < PrimitiveSceneInfo->StaticMeshes.Num(); MeshIndex++)
 				{
-					FTranslucencyShadowDepthDrawingPolicyFactory::DrawStaticMesh(
-						RHICmdList, 
-						*ShadowDepthView,
-						DrawingContext,
-						PrimitiveSceneInfo->StaticMeshes[MeshIndex],
-						true,
-						DrawRenderState,
-						PrimitiveSceneInfo->Proxy,
-						FHitProxyId());
+					const FStaticMeshBatch& StaticMeshBatch = PrimitiveSceneInfo->StaticMeshes[MeshIndex];
+					const uint64 BatchElementMask = StaticMeshBatch.bRequiresPerElementVisibility ? ShadowDepthView->StaticMeshBatchVisibility[StaticMeshBatch.BatchVisibilityId] : ~0ull;
+					TranslucencyDepthPassMeshProcessor.AddMeshBatch(StaticMeshBatch, BatchElementMask, StaticMeshBatch.PrimitiveSceneInfo->Proxy, StaticMeshBatch.Id);
 				}
 			}
 		}
-#endif
+
+		if (VisibleMeshDrawCommands.Num() > 0)
+		{
+			FVertexBufferRHIParamRef PrimitiveIdVertexBuffer = nullptr;
+			ApplyViewOverridesToMeshDrawCommands(*ShadowDepthView, VisibleMeshDrawCommands);
+			SortAndMergeDynamicPassMeshDrawCommands(SceneRenderer->FeatureLevel, VisibleMeshDrawCommands, DynamicMeshDrawCommandStorage, PrimitiveIdVertexBuffer);
+			SubmitMeshDrawCommands(VisibleMeshDrawCommands, PrimitiveIdVertexBuffer, 0, RHICmdList);
+		}
 	}
 }
 
