Fixed misc compilation / static analysis warnings.

diff --git a/Engine/Source/Runtime/Renderer/Private/BasePassRendering.cpp b/Engine/Source/Runtime/Renderer/Private/BasePassRendering.cpp
--- a/Engine/Source/Runtime/Renderer/Private/BasePassRendering.cpp
+++ b/Engine/Source/Runtime/Renderer/Private/BasePassRendering.cpp
@@ -980,7 +980,7 @@ void FBasePassMeshProcessor::Process(
 		SortKey = CalculateBasePassMeshStaticSortKey(EarlyZPassMode, BlendMode, BasePassShaders.VertexShader, BasePassShaders.PixelShader);
 	}
 
-	const bool bIsInstancedStereo = ViewIfDynamicMeshCommand ? ViewIfDynamicMeshCommand->IsInstancedStereoPass() : Scene->bStaticDrawInstancedStereo;
+	const bool bIsInstancedStereo = ViewIfDynamicMeshCommand ? ViewIfDynamicMeshCommand->IsInstancedStereoPass() : (Scene && Scene->bStaticDrawInstancedStereo);
 	const int32 InstanceFactor = bIsInstancedStereo ? 2 : 1;
 
 	BuildMeshDrawCommands(
diff --git a/Engine/Source/Runtime/Renderer/Private/DepthRendering.cpp b/Engine/Source/Runtime/Renderer/Private/DepthRendering.cpp
--- a/Engine/Source/Runtime/Renderer/Private/DepthRendering.cpp
+++ b/Engine/Source/Runtime/Renderer/Private/DepthRendering.cpp
@@ -724,7 +724,7 @@ void FDepthPassMeshProcessor::Process(
 
 	const FMeshDrawCommandSortKey SortKey = CalculateMeshStaticSortKey(DepthPassShaders.VertexShader, DepthPassShaders.PixelShader);
 
-	const bool bIsInstancedStereo = ViewIfDynamicMeshCommand ? ViewIfDynamicMeshCommand->IsInstancedStereoPass() : Scene->bStaticDrawInstancedStereo;
+	const bool bIsInstancedStereo = ViewIfDynamicMeshCommand ? ViewIfDynamicMeshCommand->IsInstancedStereoPass() : (Scene && Scene->bStaticDrawInstancedStereo);
 	const int32 InstanceFactor = bIsInstancedStereo ? 2 : 1;
 
 	BuildMeshDrawCommands(
diff --git a/Engine/Source/Runtime/Renderer/Private/MeshPassProcessor.cpp b/Engine/Source/Runtime/Renderer/Private/MeshPassProcessor.cpp
--- a/Engine/Source/Runtime/Renderer/Private/MeshPassProcessor.cpp
+++ b/Engine/Source/Runtime/Renderer/Private/MeshPassProcessor.cpp
@@ -94,6 +94,11 @@ void FMeshDrawShaderBindings::SetShaderBindings(
 	const FReadOnlyMeshDrawSingleShaderBindings& RESTRICT SingleShaderBindings,
 	FShaderBindingState& RESTRICT ShaderBindingState)
 {
+#if USING_CODE_ANALYSIS
+	MSVC_PRAGMA(warning(push))
+	MSVC_PRAGMA(warning(disable : 6386)) // Parameter.BaseIndex may be out of bounds.
+#endif
+
 	const FUniformBufferRHIParamRef* RESTRICT UniformBufferBindings = SingleShaderBindings.GetUniformBufferStart();
 	const FShaderParameterInfo* RESTRICT UniformBufferParameters = SingleShaderBindings.ParameterMapInfo.UniformBuffers.GetData();
 	const int32 NumUniformBuffers = SingleShaderBindings.ParameterMapInfo.UniformBuffers.Num();
@@ -181,6 +186,10 @@ void FMeshDrawShaderBindings::SetShaderBindings(
 			LooseDataStart += Parameter.Size;
 		}
 	}
+
+#if USING_CODE_ANALYSIS
+	MSVC_PRAGMA(warning(pop))
+#endif
 }
 
 template<class RHIShaderType>
