FPrimitiveIdVertexBufferPool - try to reuse best fitting unused buffer.

diff --git a/Engine/Source/Runtime/Renderer/Private/MeshDrawCommands.cpp b/Engine/Source/Runtime/Renderer/Private/MeshDrawCommands.cpp
--- a/Engine/Source/Runtime/Renderer/Private/MeshDrawCommands.cpp
+++ b/Engine/Source/Runtime/Renderer/Private/MeshDrawCommands.cpp
@@ -29,27 +29,46 @@ FPrimitiveIdVertexBufferPool::~FPrimitiveIdVertexBufferPool()
 
 FVertexBufferRHIParamRef FPrimitiveIdVertexBufferPool::Allocate(int32 BufferSize)
 {
-	// First look for an allocated and unused one.
+	BufferSize = Align(BufferSize, 1024);
+
+	// First look for a smallest unused one.
+	int32 BestFitBufferIndex = -1;
 	for (int32 Index = 0; Index < Entries.Num(); ++Index)
 	{
-		if (Entries[Index].LastDiscardId != DiscardId
-			&& Entries[Index].BufferSize >= BufferSize)
+		// Unused and fits?
+		if (Entries[Index].LastDiscardId != DiscardId && Entries[Index].BufferSize >= BufferSize)
 		{
-			Entries[Index].LastDiscardId = DiscardId;
-			return Entries[Index].BufferRHI;
+			// Is it a bet fit than current BestFitBufferIndex?
+			if (BestFitBufferIndex == -1 || Entries[Index].BufferSize < Entries[BestFitBufferIndex].BufferSize)
+			{
+				BestFitBufferIndex = Index;
+				
+				if (Entries[BestFitBufferIndex].BufferSize == BufferSize)
+				{
+					break;
+				}
+			}
 		}
 	}
-
-
-	// Allocate new one.
-	FPrimitiveIdVertexBufferPoolEntry NewEntry;
-	NewEntry.LastDiscardId = DiscardId;
-	NewEntry.BufferSize = Align(BufferSize, 1024);
-	FRHIResourceCreateInfo CreateInfo;
-	NewEntry.BufferRHI = RHICreateVertexBuffer(NewEntry.BufferSize, BUF_Volatile, CreateInfo);
-	Entries.Add(NewEntry);
-
-	return NewEntry.BufferRHI;
+	
+	if (BestFitBufferIndex >=0 )
+	{
+		// Reuse existing buffer.
+		Entries[BestFitBufferIndex].LastDiscardId = DiscardId;
+		return Entries[BestFitBufferIndex].BufferRHI;
+	}
+	else
+	{
+		// Allocate new one.
+		FPrimitiveIdVertexBufferPoolEntry NewEntry;
+		NewEntry.LastDiscardId = DiscardId;
+		NewEntry.BufferSize = BufferSize;
+		FRHIResourceCreateInfo CreateInfo;
+		NewEntry.BufferRHI = RHICreateVertexBuffer(NewEntry.BufferSize, BUF_Volatile, CreateInfo);
+		Entries.Add(NewEntry);
+
+		return NewEntry.BufferRHI;
+	}
 }
 
 void FPrimitiveIdVertexBufferPool::DiscardAll()
