Fixed material draw events for cached mesh draw commands. Event names are now reconstructed based on command's debug data instead of storing them using a separate string.

diff --git a/Engine/Source/Runtime/Renderer/Private/MeshPassProcessor.cpp b/Engine/Source/Runtime/Renderer/Private/MeshPassProcessor.cpp
--- a/Engine/Source/Runtime/Renderer/Private/MeshPassProcessor.cpp
+++ b/Engine/Source/Runtime/Renderer/Private/MeshPassProcessor.cpp
@@ -743,6 +743,24 @@ void FMeshDrawCommand::SubmitDraw(
 
 	if (GShowMaterialDrawEvents)
 	{
+		const FPrimitiveSceneProxy* PrimitiveSceneProxy = MeshDrawCommand.DebugData.PrimitiveSceneProxy;
+		const FMaterial* Material = MeshDrawCommand.DebugData.Material;
+
+		FString DrawEventName;
+
+		if (PrimitiveSceneProxy)
+		{
+			DrawEventName = FString::Printf(
+				TEXT("%s %s"),
+				// Note: this is the parent's material name, not the material instance
+				*Material->GetFriendlyName(),
+				PrimitiveSceneProxy->GetResourceName().IsValid() ? *PrimitiveSceneProxy->GetResourceName().ToString() : TEXT(""));
+		}
+		else
+		{
+			DrawEventName = Material->GetFriendlyName();
+		}
+
 		const uint32 Instances = MeshDrawCommand.NumInstances * MeshDrawCommand.InstanceFactor;
 		if (Instances > 1)
 		{
@@ -751,12 +769,12 @@ void FMeshDrawCommand::SubmitDraw(
 				MaterialEvent,
 				MeshEvent,
 				TEXT("%s %u instances"),
-				*MeshDrawCommand.DrawEventName,
+				*DrawEventName,
 				Instances);
 		}
 		else
 		{
-			BEGIN_DRAW_EVENTF(RHICmdList, MaterialEvent, MeshEvent, *MeshDrawCommand.DrawEventName);
+			BEGIN_DRAW_EVENTF(RHICmdList, MaterialEvent, MeshEvent, *DrawEventName);
 		}
 	}
 #endif
@@ -941,27 +959,6 @@ ERasterizerCullMode FMeshPassProcessor::ComputeMeshCullMode(const FMeshBatch& Me
 	return bMeshRenderTwoSided ? CM_None : (bInReverseCullModeOverride ? CM_CCW : CM_CW);
 }
 
-void FMeshPassProcessor::SetDrawCommandEvent(const FPrimitiveSceneProxy* PrimitiveSceneProxy, const FMaterial& RESTRICT MaterialResource, FMeshDrawCommand& MeshDrawCommand) const
-{
-#if WANTS_DRAW_MESH_EVENTS
-	if (GShowMaterialDrawEvents)
-	{
-		if (PrimitiveSceneProxy)
-		{
-			MeshDrawCommand.DrawEventName = FString::Printf(
-				TEXT("%s %s"),
-				// Note: this is the parent's material name, not the material instance
-				*MaterialResource.GetFriendlyName(),
-				PrimitiveSceneProxy->GetResourceName().IsValid() ? *PrimitiveSceneProxy->GetResourceName().ToString() : TEXT(""));
-		}
-		else
-		{
-			MeshDrawCommand.DrawEventName = MaterialResource.GetFriendlyName();
-		}
-	}
-#endif
-}
-
 int32 FMeshPassProcessor::GetDrawCommandPrimitiveId(const FPrimitiveSceneInfo* RESTRICT PrimitiveSceneInfo, const FMeshBatchElement& BatchElement) const
 {
 	int32 DrawPrimitiveId;
diff --git a/Engine/Source/Runtime/Renderer/Public/MeshPassProcessor.h b/Engine/Source/Runtime/Renderer/Public/MeshPassProcessor.h
--- a/Engine/Source/Runtime/Renderer/Public/MeshPassProcessor.h
+++ b/Engine/Source/Runtime/Renderer/Public/MeshPassProcessor.h
@@ -10,7 +10,7 @@
 #include "SceneUtils.h"
 #include "MeshBatch.h"
 
-#define MESH_DRAW_COMMAND_DEBUG_DATA ((!UE_BUILD_SHIPPING && !UE_BUILD_TEST) || VALIDATE_MESH_COMMAND_BINDINGS)
+#define MESH_DRAW_COMMAND_DEBUG_DATA ((!UE_BUILD_SHIPPING && !UE_BUILD_TEST) || VALIDATE_MESH_COMMAND_BINDINGS || WANTS_DRAW_MESH_EVENTS)
 
 /** Mesh pass types supported. */
 namespace EMeshPass
@@ -361,10 +361,6 @@ public:
 	/** Non-pipeline state */
 	uint8 StencilRef;
 
-#if WANTS_DRAW_MESH_EVENTS
-	FString DrawEventName;
-#endif
-
 	FMeshDrawCommand()
 	{}
 
@@ -929,7 +925,6 @@ public:
 		EMeshPassFeatures MeshPassFeatures,
 		const ShaderElementDataType& ShaderElementData);
 private:
-	RENDERER_API void SetDrawCommandEvent(const FPrimitiveSceneProxy* PrimitiveSceneProxy, const FMaterial& RESTRICT MaterialResource, FMeshDrawCommand& MeshDrawCommand) const;
 	RENDERER_API int32 GetDrawCommandPrimitiveId(const FPrimitiveSceneInfo* RESTRICT PrimitiveSceneInfo, const FMeshBatchElement& BatchElement) const;
 };
 
diff --git a/Engine/Source/Runtime/Renderer/Public/MeshPassProcessor.inl b/Engine/Source/Runtime/Renderer/Public/MeshPassProcessor.inl
--- a/Engine/Source/Runtime/Renderer/Public/MeshPassProcessor.inl
+++ b/Engine/Source/Runtime/Renderer/Public/MeshPassProcessor.inl
@@ -129,8 +129,6 @@ void FMeshPassProcessor::BuildMeshDrawCommands(
 				PassShaders.GeometryShader->GetElementShaderBindings(Scene, ViewIfDynamicMeshCommand, VertexFactory, false, FeatureLevel, PrimitiveSceneProxy, MeshBatch, BatchElement, ShaderElementData, GeometryShaderBindings, MeshDrawCommand.VertexStreams);
 			}
 
-			SetDrawCommandEvent(PrimitiveSceneProxy, MaterialResource, MeshDrawCommand);
-
 			const int32 DrawPrimitiveId = GetDrawCommandPrimitiveId(PrimitiveSceneInfo, BatchElement);
 
 			DrawListContext->FinalizeCommand(MeshBatch, BatchElementIndex, DrawPrimitiveId, MeshFillMode, MeshCullMode, InstanceFactor, SortKey, MeshDrawCommand, true);
@@ -211,8 +209,6 @@ void FMeshPassProcessor::BuildRayTracingDrawCommands(
 			}
 #endif
 
-			SetDrawCommandEvent(PrimitiveSceneProxy, MaterialResource, MeshDrawCommand);
-
 			const int32 DrawPrimitiveId = 0;
 			const int32 ScenePrimitiveId = 0;
 			DrawListContext->FinalizeCommand(MeshBatch, BatchElementIndex, DrawPrimitiveId, MeshFillMode, MeshCullMode, InstanceFactor, SortKey, MeshDrawCommand, false);
