DR - Fix static analysis warning
#rnx

diff --git a/Engine/Source/Runtime/Renderer/Public/MeshDrawShaderBindings.h b/Engine/Source/Runtime/Renderer/Public/MeshDrawShaderBindings.h
--- a/Engine/Source/Runtime/Renderer/Public/MeshDrawShaderBindings.h
+++ b/Engine/Source/Runtime/Renderer/Public/MeshDrawShaderBindings.h
@@ -20,7 +20,7 @@ public:
 	FMeshDrawShaderBindingsLayout(const FShader* Shader) :
 		ParameterMapInfo(Shader->GetParameterMapInfo())
 	{
-		check(Shader)
+		check(Shader);
 		Frequency = (EShaderFrequency)Shader->GetTarget().Frequency;
 		checkSlow((EShaderFrequency)Frequency == (EShaderFrequency)Shader->GetTarget().Frequency);
 	}
