DR - Deprecate ENQUEUE_UNIQUE_RENDER_COMMAND and ENQUEUE_UNIQUE_RENDER_COMMAND_FOUR
#rnx

diff --git a/Engine/Plugins/Experimental/StereoPanorama/Source/StereoPanorama/Private/SceneCapturer.cpp b/Engine/Plugins/Experimental/StereoPanorama/Source/StereoPanorama/Private/SceneCapturer.cpp
--- a/Engine/Plugins/Experimental/StereoPanorama/Source/StereoPanorama/Private/SceneCapturer.cpp
+++ b/Engine/Plugins/Experimental/StereoPanorama/Source/StereoPanorama/Private/SceneCapturer.cpp
@@ -810,11 +810,11 @@ void USceneCapturer::Tick( float DeltaTime )
         else if (CaptureStep == ECaptureStep::SetStartPosition)
         {
             //SetStartPosition();
-            ENQUEUE_UNIQUE_RENDER_COMMAND(
-                SceneCapturer_HeartbeatTickTickables,
-            {
-                TickRenderingTickables();
-            });
+            ENQUEUE_RENDER_COMMAND(SceneCapturer_HeartbeatTickTickables)(
+				[](FRHICommandList& RHICmdList)
+				{
+					TickRenderingTickables();
+				});
 
             FlushRenderingCommands();
             
diff --git a/Engine/Source/Editor/UnrealEd/Private/Lightmass/LightmassRender.cpp b/Engine/Source/Editor/UnrealEd/Private/Lightmass/LightmassRender.cpp
--- a/Engine/Source/Editor/UnrealEd/Private/Lightmass/LightmassRender.cpp
+++ b/Engine/Source/Editor/UnrealEd/Private/Lightmass/LightmassRender.cpp
@@ -1101,8 +1101,8 @@ bool FLightmassMaterialRenderer::GenerateMaterialPropertyData(
 			}
 			else
 			{
-				ENQUEUE_UNIQUE_RENDER_COMMAND(
-					InitializeSystemTextures,
+				ENQUEUE_RENDER_COMMAND(InitializeSystemTextures)(
+					[](FRHICommandListImmediate& RHICmdList)
 					{
 						GetRendererModule().InitializeSystemTextures(RHICmdList);
 					});
diff --git a/Engine/Source/Runtime/Engine/Private/GameViewportClient.cpp b/Engine/Source/Runtime/Engine/Private/GameViewportClient.cpp
--- a/Engine/Source/Runtime/Engine/Private/GameViewportClient.cpp
+++ b/Engine/Source/Runtime/Engine/Private/GameViewportClient.cpp
@@ -1487,10 +1487,11 @@ void UGameViewportClient::Draw(FViewport* InViewport, FCanvas* SceneCanvas)
 	else
 	{
 		// Make sure RHI resources get flushed if we're not using a renderer
-		ENQUEUE_UNIQUE_RENDER_COMMAND( UGameViewportClient_FlushRHIResources,
-		{ 
-			FRHICommandListExecutor::GetImmediateCommandList().ImmediateFlush(EImmediateFlushType::FlushRHIThreadFlushResources);
-		});
+		ENQUEUE_RENDER_COMMAND(UGameViewportClient_FlushRHIResources)(
+			[](FRHICommandListImmediate& RHICmdList)
+			{ 
+				RHICmdList.ImmediateFlush(EImmediateFlushType::FlushRHIThreadFlushResources);
+			});
 	}
 
 	// Beyond this point, only UI rendering independent from dynamc resolution.
diff --git a/Engine/Source/Runtime/Engine/Private/Particles/ParticleGpuSimulation.cpp b/Engine/Source/Runtime/Engine/Private/Particles/ParticleGpuSimulation.cpp
--- a/Engine/Source/Runtime/Engine/Private/Particles/ParticleGpuSimulation.cpp
+++ b/Engine/Source/Runtime/Engine/Private/Particles/ParticleGpuSimulation.cpp
@@ -5129,16 +5129,16 @@ void FFXSystem::UpdateMultiGPUResources(FRHICommandListImmediate& RHICmdList)
 
 void FFXSystem::VisualizeGPUParticles(FCanvas* Canvas)
 {
-	ENQUEUE_UNIQUE_RENDER_COMMAND_FOURPARAMETER(
-		FVisualizeGPUParticlesCommand,
-		FFXSystem*, FXSystem, this,
-		int32, VisualizationMode, FXConsoleVariables::VisualizeGPUSimulation,
-		FRenderTarget*, RenderTarget, Canvas->GetRenderTarget(),
-		ERHIFeatureLevel::Type, FeatureLevel, GetFeatureLevel(),
+	FFXSystem* FXSystem = this;
+	int32 VisualizationMode = FXConsoleVariables::VisualizeGPUSimulation;
+	FRenderTarget* RenderTarget = Canvas->GetRenderTarget();
+	ERHIFeatureLevel::Type InFeatureLevel = FeatureLevel;
+	ENQUEUE_RENDER_COMMAND(FVisualizeGPUParticlesCommand)(
+		[FXSystem, VisualizationMode, RenderTarget, InFeatureLevel](FRHICommandList& RHICmdList)
 	{
 		FParticleSimulationResources* Resources = FXSystem->GetParticleSimulationResources();
 		FParticleStateTextures& CurrentStateTextures = Resources->GetVisualizeStateTextures();
-		VisualizeGPUSimulation(RHICmdList, FeatureLevel, VisualizationMode, RenderTarget, CurrentStateTextures, GParticleCurveTexture.GetCurveTexture());
+		VisualizeGPUSimulation(RHICmdList, InFeatureLevel, VisualizationMode, RenderTarget, CurrentStateTextures, GParticleCurveTexture.GetCurveTexture());
 	});
 }
 
diff --git a/Engine/Source/Runtime/Engine/Private/Particles/ParticleSystemRender.cpp b/Engine/Source/Runtime/Engine/Private/Particles/ParticleSystemRender.cpp
--- a/Engine/Source/Runtime/Engine/Private/Particles/ParticleSystemRender.cpp
+++ b/Engine/Source/Runtime/Engine/Private/Particles/ParticleSystemRender.cpp
@@ -260,12 +260,11 @@ void ParticleVertexFactoryPool_FreePool_RenderingThread()
 
 void ParticleVertexFactoryPool_FreePool()
 {
-	ENQUEUE_UNIQUE_RENDER_COMMAND(
-		ParticleVertexFactoryFreePool,
-	{
-		ParticleVertexFactoryPool_FreePool_RenderingThread();
-	}
-	);		
+	ENQUEUE_RENDER_COMMAND(ParticleVertexFactoryFreePool)(
+		[](FRHICommandList& RHICmdList)
+		{
+			ParticleVertexFactoryPool_FreePool_RenderingThread();
+		});		
 }
 
 void ParticleVertexFactoryPool_ClearPool_RenderingThread()
@@ -276,12 +275,11 @@ void ParticleVertexFactoryPool_ClearPool_RenderingThread()
 /** Globally accessible function for clearing the pool */
 void ParticleVertexFactoryPool_ClearPool()
 {
-	ENQUEUE_UNIQUE_RENDER_COMMAND(
-		ParticleVertexFactoryFreePool,
-	{
+	ENQUEUE_RENDER_COMMAND(ParticleVertexFactoryFreePool)(
+		[](FRHICommandList& RHICmdList)
+			{
 		ParticleVertexFactoryPool_ClearPool_RenderingThread();
-	}
-	);		
+		});		
 }
 
 ///////////////////////////////////////////////////////////////////////////////
diff --git a/Engine/Source/Runtime/Engine/Private/SkeletalMesh.cpp b/Engine/Source/Runtime/Engine/Private/SkeletalMesh.cpp
--- a/Engine/Source/Runtime/Engine/Private/SkeletalMesh.cpp
+++ b/Engine/Source/Runtime/Engine/Private/SkeletalMesh.cpp
@@ -4228,28 +4228,28 @@ void FSkeletalMeshSceneProxy::UpdateMorphMaterialUsage_GameThread(TArray<UMateri
 		// update the new LODSections on the render thread proxy
 		if (MaterialsToSwap.Num())
 		{
-			ENQUEUE_UNIQUE_RENDER_COMMAND_FOURPARAMETER(
-			UpdateSkelProxyLODSectionElementsCmd,
-				TSet<UMaterialInterface*>,MaterialsToSwap,MaterialsToSwap,
-				UMaterialInterface*,DefaultMaterial,UMaterial::GetDefaultMaterial(MD_Surface),
-				ERHIFeatureLevel::Type, FeatureLevel, GetScene().GetFeatureLevel(),
-			FSkeletalMeshSceneProxy*,SkelMeshSceneProxy,this,
-			{
+			TSet<UMaterialInterface*> InMaterialsToSwap = MaterialsToSwap;
+			UMaterialInterface* DefaultMaterial = UMaterial::GetDefaultMaterial(MD_Surface);
+			ERHIFeatureLevel::Type InFeatureLevel = GetScene().GetFeatureLevel();
+			FSkeletalMeshSceneProxy* SkelMeshSceneProxy = this;
+			ENQUEUE_RENDER_COMMAND(UpdateSkelProxyLODSectionElementsCmd)(
+				[InMaterialsToSwap, DefaultMaterial, InFeatureLevel, SkelMeshSceneProxy](FRHICommandList& RHICmdList)
+				{
 					for( int32 LodIdx=0; LodIdx < SkelMeshSceneProxy->LODSections.Num(); LodIdx++ )
 					{
 						FLODSectionElements& LODSection = SkelMeshSceneProxy->LODSections[LodIdx];
 						for( int32 SectIdx=0; SectIdx < LODSection.SectionElements.Num(); SectIdx++ )
 						{
 							FSectionElementInfo& SectionElement = LODSection.SectionElements[SectIdx];
-							if( MaterialsToSwap.Contains(SectionElement.Material) )
+							if( InMaterialsToSwap.Contains(SectionElement.Material) )
 							{
 								// fallback to default material if needed
 								SectionElement.Material = DefaultMaterial;
 							}
 						}
 					}
-					SkelMeshSceneProxy->MaterialRelevance |= DefaultMaterial->GetRelevance(FeatureLevel);
-			});
+					SkelMeshSceneProxy->MaterialRelevance |= DefaultMaterial->GetRelevance(InFeatureLevel);
+				});
 		}
 	}
 }
diff --git a/Engine/Source/Runtime/Engine/Private/SkeletalRenderCPUSkin.cpp b/Engine/Source/Runtime/Engine/Private/SkeletalRenderCPUSkin.cpp
--- a/Engine/Source/Runtime/Engine/Private/SkeletalRenderCPUSkin.cpp
+++ b/Engine/Source/Runtime/Engine/Private/SkeletalRenderCPUSkin.cpp
@@ -142,17 +142,16 @@ void FSkeletalMeshObjectCPUSkin::Update(int32 LODIndex,USkinnedMeshComponent* In
 		}
 
 		// queue a call to update this data
-		ENQUEUE_UNIQUE_RENDER_COMMAND_FOURPARAMETER(
-			SkelMeshObjectUpdateDataCommand,
-			FSkeletalMeshObjectCPUSkin*, MeshObject, this,
-			uint32, FrameNumberToPrepare, FrameNumberToPrepare,
-			uint32, RevisionNumber, RevisionNumber, 
-			FDynamicSkelMeshObjectDataCPUSkin*, NewDynamicData, NewDynamicData,
 		{
-			FScopeCycleCounter Context(MeshObject->GetStatId());
-			MeshObject->UpdateDynamicData_RenderThread(RHICmdList, NewDynamicData, FrameNumberToPrepare, RevisionNumber);
+			FSkeletalMeshObjectCPUSkin* MeshObject = this;
+			ENQUEUE_RENDER_COMMAND(SkelMeshObjectUpdateDataCommand)(
+				[MeshObject, FrameNumberToPrepare, RevisionNumber, NewDynamicData](FRHICommandListImmediate& RHICmdList)
+				{
+					FScopeCycleCounter Context(MeshObject->GetStatId());
+					MeshObject->UpdateDynamicData_RenderThread(RHICmdList, NewDynamicData, FrameNumberToPrepare, RevisionNumber);
+				}
+			);
 		}
-		);
 	}
 }
 
diff --git a/Engine/Source/Runtime/Engine/Private/Streaming/StreamingManagerTexture.cpp b/Engine/Source/Runtime/Engine/Private/Streaming/StreamingManagerTexture.cpp
--- a/Engine/Source/Runtime/Engine/Private/Streaming/StreamingManagerTexture.cpp
+++ b/Engine/Source/Runtime/Engine/Private/Streaming/StreamingManagerTexture.cpp
@@ -268,11 +268,12 @@ bool FStreamingManagerTexture::StreamOutTextureData( int64 RequiredMemorySize )
 			if (TempMemoryUsed >= MaxTempMemoryAllowed)
 			{
 				// Queue up the process on the render thread and wait for everything to complete.
-				ENQUEUE_UNIQUE_RENDER_COMMAND(FlushResourceCommand,
-				{				
-					FRHICommandListExecutor::GetImmediateCommandList().ImmediateFlush(EImmediateFlushType::FlushRHIThreadFlushResources);
-					RHIFlushResources();
-				});
+				ENQUEUE_RENDER_COMMAND(FlushResourceCommand)(
+					[](FRHICommandList& RHICmdList)
+					{				
+						FRHICommandListExecutor::GetImmediateCommandList().ImmediateFlush(EImmediateFlushType::FlushRHIThreadFlushResources);
+						RHIFlushResources();
+					});
 				FlushRenderingCommands();
 				TempMemoryUsed = 0;
 			}
diff --git a/Engine/Source/Runtime/Engine/Private/UnrealEngine.cpp b/Engine/Source/Runtime/Engine/Private/UnrealEngine.cpp
--- a/Engine/Source/Runtime/Engine/Private/UnrealEngine.cpp
+++ b/Engine/Source/Runtime/Engine/Private/UnrealEngine.cpp
@@ -7834,7 +7834,13 @@ bool UEngine::PerformError(const TCHAR* Cmd, FOutputDevice& Ar)
 #if !UE_BUILD_SHIPPING
 	if (FParse::Command(&Cmd, TEXT("RENDERCRASH")))
 	{
-		ENQUEUE_UNIQUE_RENDER_COMMAND(CauseRenderThreadCrash, { UE_LOG(LogEngine, Warning, TEXT("Printed warning to log.")); SetCrashType(ECrashType::Debug); UE_LOG(LogEngine, Fatal, TEXT("Crashing the renderthread at your request")); });
+		ENQUEUE_RENDER_COMMAND(CauseRenderThreadCrash)(
+			[](FRHICommandList& RHICmdList)
+			{
+				UE_LOG(LogEngine, Warning, TEXT("Printed warning to log."));
+				SetCrashType(ECrashType::Debug);
+				UE_LOG(LogEngine, Fatal, TEXT("Crashing the renderthread at your request"));
+			});
 		return true;
 	}
 	if (FParse::Command(&Cmd, TEXT("RENDERCHECK")))
@@ -7848,34 +7854,46 @@ bool UEngine::PerformError(const TCHAR* Cmd, FOutputDevice& Ar)
 				check(!"Crashing the renderthread via check(0) at your request");
 			}
 		};
-		ENQUEUE_UNIQUE_RENDER_COMMAND(CauseRenderThreadCrash, { FRender::Check(); });
+		ENQUEUE_RENDER_COMMAND(CauseRenderThreadCrash)(
+			[](FRHICommandList& RHICmdList)
+			{
+				FRender::Check();
+			});
 		return true;
 	}
 	if (FParse::Command(&Cmd, TEXT("RENDERGPF")))
 	{
-		ENQUEUE_UNIQUE_RENDER_COMMAND(CauseRenderThreadCrash, { UE_LOG(LogEngine, Warning, TEXT("Printed warning to log.")); SetCrashType(ECrashType::Debug); *(int32 *)3 = 123; });
+		ENQUEUE_RENDER_COMMAND(CauseRenderThreadCrash)(
+			[](FRHICommandList& RHICmdList)
+			{
+				UE_LOG(LogEngine, Warning, TEXT("Printed warning to log."));
+				SetCrashType(ECrashType::Debug);
+				*(int32 *)3 = 123;
+			});
 		return true;
 	}
 	if (FParse::Command(&Cmd, TEXT("RENDERFATAL")))
 	{
-		ENQUEUE_UNIQUE_RENDER_COMMAND(CauseRenderThreadCrash,
-		{
-			UE_LOG(LogEngine, Warning, TEXT("Printed warning to log."));
-		SetCrashType(ECrashType::Debug);
-		LowLevelFatalError(TEXT("FError::LowLevelFatal test"));
-		});
+		ENQUEUE_RENDER_COMMAND(CauseRenderThreadCrash)(
+			[](FRHICommandList& RHICmdList)
+			{
+				UE_LOG(LogEngine, Warning, TEXT("Printed warning to log."));
+			SetCrashType(ECrashType::Debug);
+			LowLevelFatalError(TEXT("FError::LowLevelFatal test"));
+			});
 		return true;
 	}
 	if (FParse::Command(&Cmd, TEXT("RENDERENSURE")))
 	{
-		ENQUEUE_UNIQUE_RENDER_COMMAND(CauseRenderThreadEnsure,
-		{
-			UE_LOG(LogEngine, Warning, TEXT("Printed warning to log."));
-		if (!ensure(0))
-		{
-			UE_LOG(LogEngine, Warning, TEXT("Ensure condition failed (this is the expected behavior)."));
-		}
-		});
+		ENQUEUE_RENDER_COMMAND(CauseRenderThreadEnsure)(
+			[](FRHICommandList& RHICmdList)
+			{
+				UE_LOG(LogEngine, Warning, TEXT("Printed warning to log."));
+			if (!ensure(0))
+			{
+				UE_LOG(LogEngine, Warning, TEXT("Ensure condition failed (this is the expected behavior)."));
+			}
+			});
 		return true;
 	}
 	if (FParse::Command(&Cmd, TEXT("THREADCRASH")))
@@ -12421,13 +12439,13 @@ void UEngine::TrimMemory()
 	// For platforms which manage GPU memory directly we must Enqueue a flush, and wait for it to be processed
 	// so that any pending frees that depend on the GPU will be processed.  Otherwise a whole map's worth of GPU memory
 	// may be unavailable to load the next one.
-	ENQUEUE_UNIQUE_RENDER_COMMAND(FlushCommand,
-	{
-		GRHICommandList.GetImmediateCommandList().ImmediateFlush(EImmediateFlushType::FlushRHIThreadFlushResources);
-	RHIFlushResources();
-	GRHICommandList.GetImmediateCommandList().ImmediateFlush(EImmediateFlushType::FlushRHIThreadFlushResources);
-	}
-	);
+	ENQUEUE_RENDER_COMMAND(FlushCommand)(
+		[](FRHICommandList& RHICmdList)
+		{
+			GRHICommandList.GetImmediateCommandList().ImmediateFlush(EImmediateFlushType::FlushRHIThreadFlushResources);
+			RHIFlushResources();
+			GRHICommandList.GetImmediateCommandList().ImmediateFlush(EImmediateFlushType::FlushRHIThreadFlushResources);
+		});
 	FlushRenderingCommands();
 
 	// Ask systems to trim memory where possible
diff --git a/Engine/Source/Runtime/Landscape/Private/LandscapeEditProcedural.cpp b/Engine/Source/Runtime/Landscape/Private/LandscapeEditProcedural.cpp
--- a/Engine/Source/Runtime/Landscape/Private/LandscapeEditProcedural.cpp
+++ b/Engine/Source/Runtime/Landscape/Private/LandscapeEditProcedural.cpp
@@ -1781,15 +1781,15 @@ void ALandscape::ResolveProceduralHeightmapTexture(bool InUpdateDDC)
 				Flags.SetMip(MipIndex);
 				FIntRect Rect(0, 0, MipSizeU, MipSizeV);
 
-				ENQUEUE_UNIQUE_RENDER_COMMAND_FOURPARAMETER(
-					ReadSurfaceCommand,
-					FTextureRHIRef, SourceTextureRHI, HeightmapRenderData.HeightmapsCPUReadBack->TextureRHI,
-					FIntRect, Rect, Rect,
-					TArray<FColor>*, OutData, &MipData[MipIndex],
-					FReadSurfaceDataFlags, ReadFlags, Flags,
-					{
-						RHICmdList.ReadSurfaceData(SourceTextureRHI, Rect, *OutData, ReadFlags);
-					});
+				{
+					TArray<FColor>* OutData = &MipData[MipIndex];
+					FTextureRHIRef SourceTextureRHI = HeightmapRenderData.HeightmapsCPUReadBack->TextureRHI;
+					ENQUEUE_RENDER_COMMAND(ReadSurfaceCommand)(
+						[SourceTextureRHI, Rect, OutData, Flags](FRHICommandListImmediate& RHICmdList)
+						{
+							RHICmdList.ReadSurfaceData(SourceTextureRHI, Rect, *OutData, Flags);
+						});
+				}
 
 				MipSizeU >>= 1;
 				MipSizeV >>= 1;
diff --git a/Engine/Source/Runtime/Launch/Private/LaunchEngineLoop.cpp b/Engine/Source/Runtime/Launch/Private/LaunchEngineLoop.cpp
--- a/Engine/Source/Runtime/Launch/Private/LaunchEngineLoop.cpp
+++ b/Engine/Source/Runtime/Launch/Private/LaunchEngineLoop.cpp
@@ -3764,8 +3764,8 @@ void FEngineLoop::Tick()
 			{
 				FlushRenderingCommands();
 
-				ENQUEUE_UNIQUE_RENDER_COMMAND(
-					MeasureLongGPUTaskExecutionTimeCmd,
+				ENQUEUE_RENDER_COMMAND(MeasureLongGPUTaskExecutionTimeCmd)(
+					[](FRHICommandListImmediate& RHICmdList)
 					{
 						MeasureLongGPUTaskExecutionTime(RHICmdList);
 					});
@@ -3846,8 +3846,8 @@ void FEngineLoop::Tick()
 		MALLOC_PROFILER(GMalloc->Exec(nullptr, *FString::Printf(TEXT("SNAPSHOTMEMORYFRAME")),*GLog));
 
 		// handle some per-frame tasks on the rendering thread
-		ENQUEUE_UNIQUE_RENDER_COMMAND(
-			ResetDeferredUpdates,
+		ENQUEUE_RENDER_COMMAND(ResetDeferredUpdates)(
+			[](FRHICommandList& RHICmdList)
 			{
 				FDeferredUpdateResource::ResetNeedsUpdate();
 				FlushPendingDeleteRHIResources_RenderThread();
@@ -4062,11 +4062,12 @@ void FEngineLoop::Tick()
 			ConcurrentTask = nullptr;
 		}
 		{
-			ENQUEUE_UNIQUE_RENDER_COMMAND(WaitForOutstandingTasksOnly_for_DelaySceneRenderCompletion,
-			{
-				QUICK_SCOPE_CYCLE_COUNTER(STAT_DelaySceneRenderCompletion_TaskWait);
-				FRHICommandListExecutor::GetImmediateCommandList().ImmediateFlush(EImmediateFlushType::WaitForOutstandingTasksOnly);
-			});
+			ENQUEUE_RENDER_COMMAND(WaitForOutstandingTasksOnly_for_DelaySceneRenderCompletion)(
+				[](FRHICommandList& RHICmdList)
+				{
+					QUICK_SCOPE_CYCLE_COUNTER(STAT_DelaySceneRenderCompletion_TaskWait);
+					FRHICommandListExecutor::GetImmediateCommandList().ImmediateFlush(EImmediateFlushType::WaitForOutstandingTasksOnly);
+				});
 		}
 #endif
 
@@ -4162,10 +4163,11 @@ void FEngineLoop::Tick()
 		#endif
 
 		// end of RHI frame
-		ENQUEUE_UNIQUE_RENDER_COMMAND(EndFrame,
-		{
-			EndFrameRenderThread(RHICmdList);
-		});
+		ENQUEUE_RENDER_COMMAND(EndFrame)(
+			[](FRHICommandListImmediate& RHICmdList)
+			{
+				EndFrameRenderThread(RHICmdList);
+			});
 
 		// Set CPU utilization stats.
 		const FCPUTime CPUTime = FPlatformTime::GetCPUTime();
diff --git a/Engine/Source/Runtime/RenderCore/Private/RenderingThread.cpp b/Engine/Source/Runtime/RenderCore/Private/RenderingThread.cpp
--- a/Engine/Source/Runtime/RenderCore/Private/RenderingThread.cpp
+++ b/Engine/Source/Runtime/RenderCore/Private/RenderingThread.cpp
@@ -522,16 +522,16 @@ public:
 			if (!GIsRenderingThreadSuspended.Load(EMemoryOrder::Relaxed) && OutstandingHeartbeats.GetValue() < 4)
 			{
 				OutstandingHeartbeats.Increment();
-				ENQUEUE_UNIQUE_RENDER_COMMAND(
-					HeartbeatTickTickables,
-				{
-					OutstandingHeartbeats.Decrement();
-					// make sure that rendering thread tickables get a chance to tick, even if the render thread is starving
-					if (!GIsRenderingThreadSuspended.Load(EMemoryOrder::Relaxed))
+				ENQUEUE_RENDER_COMMAND(HeartbeatTickTickables)(
+					[](FRHICommandList& RHICmdList)
 					{
-						TickRenderingTickables();
-					}
-				});
+						OutstandingHeartbeats.Decrement();
+						// make sure that rendering thread tickables get a chance to tick, even if the render thread is starving
+						if (!GIsRenderingThreadSuspended.Load(EMemoryOrder::Relaxed))
+						{
+							TickRenderingTickables();
+						}
+					});
 			}
 		}
 		return 0;
@@ -1240,12 +1240,11 @@ void FlushPendingDeleteRHIResources_GameThread()
 {
 	if (!IsRunningRHIInSeparateThread())
 	{
-		ENQUEUE_UNIQUE_RENDER_COMMAND(
-			FlushPendingDeleteRHIResources,
-		{
-			FlushPendingDeleteRHIResources_RenderThread();
-		}
-		);
+		ENQUEUE_RENDER_COMMAND(FlushPendingDeleteRHIResources)(
+			[](FRHICommandList& RHICmdList)
+			{
+				FlushPendingDeleteRHIResources_RenderThread();
+			});
 	}
 }
 void FlushPendingDeleteRHIResources_RenderThread()
diff --git a/Engine/Source/Runtime/RenderCore/Public/RenderingThread.h b/Engine/Source/Runtime/RenderCore/Public/RenderingThread.h
--- a/Engine/Source/Runtime/RenderCore/Public/RenderingThread.h
+++ b/Engine/Source/Runtime/RenderCore/Public/RenderingThread.h
@@ -303,7 +303,8 @@ FORCEINLINE_DEBUGGABLE void EnqueueUniqueRenderCommand(LAMBDA& Lambda)
 			FScopeCycleCounter EURCMacro_Scope(TempCommand.GetStatId()); \
 			TempCommand.DoTask(ENamedThreads::GameThread, FGraphEventRef() ); \
 		} \
-	}
+	} \
+	DEPRECATED_MACRO(4.22, "Please use ENQUEUE_RENDER_COMMAND instead")
 
 /**
  * Declares a rendering command type with 1 parameters.
@@ -489,7 +490,8 @@ FORCEINLINE_DEBUGGABLE void EnqueueUniqueRenderCommand(LAMBDA& Lambda)
 
 #define ENQUEUE_UNIQUE_RENDER_COMMAND_FOURPARAMETER(TypeName,ParamType1,ParamName1,ParamValue1,ParamType2,ParamName2,ParamValue2,ParamType3,ParamName3,ParamValue3,ParamType4,ParamName4,ParamValue4,Code) \
 	ENQUEUE_UNIQUE_RENDER_COMMAND_FOURPARAMETER_DECLARE(TypeName,ParamType1,ParamName1,ParamValue1,ParamType2,ParamName2,ParamValue2,ParamType3,ParamName3,ParamValue3,ParamType4,ParamName4,ParamValue4,Code) \
-	ENQUEUE_UNIQUE_RENDER_COMMAND_FOURPARAMETER_CREATE(TypeName,ParamType1,ParamValue1,ParamType2,ParamValue2,ParamType3,ParamValue3,ParamType4,ParamValue4)
+	ENQUEUE_UNIQUE_RENDER_COMMAND_FOURPARAMETER_CREATE(TypeName,ParamType1,ParamValue1,ParamType2,ParamValue2,ParamType3,ParamValue3,ParamType4,ParamValue4) \
+	DEPRECATED_MACRO(4.22, "Please use ENQUEUE_RENDER_COMMAND instead")
 
 #define ENQUEUE_UNIQUE_RENDER_COMMAND_FOURPARAMETER_DECLARE_TEMPLATE(TypeName,TemplateParamName,ParamType1,ParamName1,ParamValue1,ParamType2,ParamName2,ParamValue2,ParamType3,ParamName3,ParamValue3,ParamType4,ParamName4,ParamValue4,Code) \
 	template <typename TemplateParamName> \
diff --git a/Engine/Source/Runtime/Renderer/Private/PostProcess/RenderingCompositionGraph.cpp b/Engine/Source/Runtime/Renderer/Private/PostProcess/RenderingCompositionGraph.cpp
--- a/Engine/Source/Runtime/Renderer/Private/PostProcess/RenderingCompositionGraph.cpp
+++ b/Engine/Source/Runtime/Renderer/Private/PostProcess/RenderingCompositionGraph.cpp
@@ -138,28 +138,26 @@ void Test()
 
 void ExecuteCompositionGraphDebug()
 {
-	ENQUEUE_UNIQUE_RENDER_COMMAND(
-		StartDebugCompositionGraph,
-	{
-		GDebugCompositionGraphFrames = 1;
-		Test();
-	}
-	);
+	ENQUEUE_RENDER_COMMAND(StartDebugCompositionGraph)(
+		[](FRHICommandList& RHICmdList)
+		{
+			GDebugCompositionGraphFrames = 1;
+			Test();
+		});
 }
 
 // main thread
 void CompositionGraph_OnStartFrame()
 {
 #if !UE_BUILD_SHIPPING
-	ENQUEUE_UNIQUE_RENDER_COMMAND(
-		DebugCompositionGraphDec,
-	{
-		if(GDebugCompositionGraphFrames)
+	ENQUEUE_RENDER_COMMAND(DebugCompositionGraphDec)(
+		[](FRHICommandList& RHICmdList)
 		{
-			--GDebugCompositionGraphFrames;
-		}
-	}
-	);		
+			if(GDebugCompositionGraphFrames)
+			{
+				--GDebugCompositionGraphFrames;
+			}
+		});		
 #endif
 }
 
diff --git a/Engine/Source/Runtime/Renderer/Private/ReflectionEnvironmentCapture.cpp b/Engine/Source/Runtime/Renderer/Private/ReflectionEnvironmentCapture.cpp
--- a/Engine/Source/Runtime/Renderer/Private/ReflectionEnvironmentCapture.cpp
+++ b/Engine/Source/Runtime/Renderer/Private/ReflectionEnvironmentCapture.cpp
@@ -1266,12 +1266,12 @@ void CaptureSceneIntoScratchCubemap(
 		{
 			// Alert the RHI that we're rendering a new frame
 			// Not really a new frame, but it will allow pooling mechanisms to update, like the uniform buffer pool
-			ENQUEUE_UNIQUE_RENDER_COMMAND(
-				BeginFrame,
-			{
-				GFrameNumberRenderThread++;
-				RHICmdList.BeginFrame();
-			})
+			ENQUEUE_RENDER_COMMAND(BeginFrame)(
+				[](FRHICommandList& RHICmdList)
+				{
+					GFrameNumberRenderThread++;
+					RHICmdList.BeginFrame();
+				});
 		}
 
 		GReflectionCaptureRenderTarget.SetSize(CubemapSize);
@@ -1489,20 +1489,21 @@ void FScene::CaptureOrUploadReflectionCapture(UReflectionCaptureComponent* Captu
 				check(!TEXT("Unknown reflection source type"));
 			}
 
-			ENQUEUE_UNIQUE_RENDER_COMMAND_FOURPARAMETER( 
-				FilterCommand,
-				ERHIFeatureLevel::Type, FeatureLevel, GetFeatureLevel(),
-				int32, ReflectionCaptureSize, ReflectionCaptureSize,
-				FScene*, Scene, this,
-				const UReflectionCaptureComponent*, CaptureComponent, CaptureComponent,
-				{
-					FindOrAllocateCubemapIndex(Scene, CaptureComponent);
-					FCaptureComponentSceneState& FoundState = Scene->ReflectionSceneData.AllocatedReflectionCaptureState.FindChecked(CaptureComponent);
+			{
+				ERHIFeatureLevel::Type InFeatureLevel = GetFeatureLevel();
+				int32 InReflectionCaptureSize = ReflectionCaptureSize;
+				FScene* Scene = this;
+				const UReflectionCaptureComponent* InCaptureComponent = CaptureComponent;
+				ENQUEUE_RENDER_COMMAND(FilterCommand)(
+					[InFeatureLevel, InReflectionCaptureSize, Scene, InCaptureComponent](FRHICommandListImmediate& RHICmdList)
+					{
+						FindOrAllocateCubemapIndex(Scene, InCaptureComponent);
+						FCaptureComponentSceneState& FoundState = Scene->ReflectionSceneData.AllocatedReflectionCaptureState.FindChecked(InCaptureComponent);
 
-					ComputeAverageBrightness(RHICmdList, FeatureLevel, ReflectionCaptureSize, FoundState.AverageBrightness);
-					FilterReflectionEnvironment(RHICmdList, FeatureLevel, ReflectionCaptureSize, NULL);
-				}
-			);
+						ComputeAverageBrightness(RHICmdList, InFeatureLevel, InReflectionCaptureSize, FoundState.AverageBrightness);
+						FilterReflectionEnvironment(RHICmdList, InFeatureLevel, InReflectionCaptureSize, nullptr);
+					});
+			}
 
 			// Create a proxy to represent the reflection capture to the rendering thread
 			// The rendering thread will be responsible for deleting this when done with the filtering operation
@@ -1638,24 +1639,26 @@ void FScene::UpdateSkyCaptureContents(
 				});
 		}
 
-		ENQUEUE_UNIQUE_RENDER_COMMAND_FOURPARAMETER( 
-			FilterCommand,
-			int32, CubemapSize, CaptureComponent->CubemapResolution,
-			float&, AverageBrightness, OutAverageBrightness,
-			FSHVectorRGB3*, IrradianceEnvironmentMap, &OutIrradianceEnvironmentMap,
-			ERHIFeatureLevel::Type, FeatureLevel, GetFeatureLevel(),
 		{
-			if (FeatureLevel <= ERHIFeatureLevel::ES3_1)
-			{
-				MobileReflectionEnvironmentCapture::ComputeAverageBrightness(RHICmdList, FeatureLevel, CubemapSize, AverageBrightness);
-				MobileReflectionEnvironmentCapture::FilterReflectionEnvironment(RHICmdList, FeatureLevel, CubemapSize, IrradianceEnvironmentMap);
-			}
-			else
-			{
-				ComputeAverageBrightness(RHICmdList, FeatureLevel, CubemapSize, AverageBrightness);
-				FilterReflectionEnvironment(RHICmdList, FeatureLevel, CubemapSize, IrradianceEnvironmentMap);
-			}
-		});
+			int32 CubemapSize = CaptureComponent->CubemapResolution;
+			float* AverageBrightness = &OutAverageBrightness;
+			FSHVectorRGB3* IrradianceEnvironmentMap = &OutIrradianceEnvironmentMap;
+			ERHIFeatureLevel::Type InFeatureLevel = GetFeatureLevel();
+			ENQUEUE_RENDER_COMMAND(FilterCommand)(
+				[CubemapSize, AverageBrightness, IrradianceEnvironmentMap, InFeatureLevel](FRHICommandListImmediate& RHICmdList)
+				{
+					if (InFeatureLevel <= ERHIFeatureLevel::ES3_1)
+					{
+						MobileReflectionEnvironmentCapture::ComputeAverageBrightness(RHICmdList, InFeatureLevel, CubemapSize, *AverageBrightness);
+						MobileReflectionEnvironmentCapture::FilterReflectionEnvironment(RHICmdList, InFeatureLevel, CubemapSize, IrradianceEnvironmentMap);
+					}
+					else
+					{
+						ComputeAverageBrightness(RHICmdList, InFeatureLevel, CubemapSize, *AverageBrightness);
+						FilterReflectionEnvironment(RHICmdList, InFeatureLevel, CubemapSize, IrradianceEnvironmentMap);
+					}
+				});
+		}
 
 		// Optionally copy the filtered mip chain to the output texture
 		if (OutProcessedTexture)
diff --git a/Engine/Source/Runtime/Renderer/Private/RendererScene.cpp b/Engine/Source/Runtime/Renderer/Private/RendererScene.cpp
--- a/Engine/Source/Runtime/Renderer/Private/RendererScene.cpp
+++ b/Engine/Source/Runtime/Renderer/Private/RendererScene.cpp
@@ -442,16 +442,16 @@ void FDistanceFieldSceneData::VerifyIntegrity()
 
 void FScene::UpdateSceneSettings(AWorldSettings* WorldSettings)
 {
-	ENQUEUE_UNIQUE_RENDER_COMMAND_FOURPARAMETER(
-		UpdateSceneSettings,
-		FScene*, Scene, this,
-		float, DefaultMaxDistanceFieldOcclusionDistance, WorldSettings->DefaultMaxDistanceFieldOcclusionDistance,
-		float, GlobalDistanceFieldViewDistance, WorldSettings->GlobalDistanceFieldViewDistance,
-		float, DynamicIndirectShadowsSelfShadowingIntensity, FMath::Clamp(WorldSettings->DynamicIndirectShadowsSelfShadowingIntensity, 0.0f, 1.0f),
-	{
-		Scene->DefaultMaxDistanceFieldOcclusionDistance = DefaultMaxDistanceFieldOcclusionDistance;
-		Scene->GlobalDistanceFieldViewDistance = GlobalDistanceFieldViewDistance;
-		Scene->DynamicIndirectShadowsSelfShadowingIntensity = DynamicIndirectShadowsSelfShadowingIntensity;
+	FScene* Scene = this;
+	float InDefaultMaxDistanceFieldOcclusionDistance = WorldSettings->DefaultMaxDistanceFieldOcclusionDistance;
+	float InGlobalDistanceFieldViewDistance = WorldSettings->GlobalDistanceFieldViewDistance;
+	float InDynamicIndirectShadowsSelfShadowingIntensity = FMath::Clamp(WorldSettings->DynamicIndirectShadowsSelfShadowingIntensity, 0.0f, 1.0f);
+	ENQUEUE_RENDER_COMMAND(UpdateSceneSettings)(
+		[Scene, InDefaultMaxDistanceFieldOcclusionDistance, InGlobalDistanceFieldViewDistance, InDynamicIndirectShadowsSelfShadowingIntensity](FRHICommandList& RHICmdList)
+	{
+		Scene->DefaultMaxDistanceFieldOcclusionDistance = InDefaultMaxDistanceFieldOcclusionDistance;
+		Scene->GlobalDistanceFieldViewDistance = InGlobalDistanceFieldViewDistance;
+		Scene->DynamicIndirectShadowsSelfShadowingIntensity = InDynamicIndirectShadowsSelfShadowingIntensity;
 	});
 }
 
diff --git a/Engine/Source/Runtime/Renderer/Private/ShadowSetup.cpp b/Engine/Source/Runtime/Renderer/Private/ShadowSetup.cpp
--- a/Engine/Source/Runtime/Renderer/Private/ShadowSetup.cpp
+++ b/Engine/Source/Runtime/Renderer/Private/ShadowSetup.cpp
@@ -251,11 +251,11 @@ FAutoConsoleVariableRef CVarUseOctreeForShadowCulling(
 bool GDumpShadowSetup = false;
 void DumpShadowDumpSetup()
 {
-	ENQUEUE_UNIQUE_RENDER_COMMAND(
-		DumpShadowDumpSetup,
-	{
-		GDumpShadowSetup = true;
-	});
+	ENQUEUE_RENDER_COMMAND(DumpShadowDumpSetup)(
+		[](FRHICommandList& RHICmdList)
+		{
+			GDumpShadowSetup = true;
+		});
 }
 
 FAutoConsoleCommand CmdDumpShadowDumpSetup(
