Fixed mesh draw command shader binding's assert conditions.

diff --git a/Engine/Source/Runtime/Renderer/Private/MeshPassProcessor.cpp b/Engine/Source/Runtime/Renderer/Private/MeshPassProcessor.cpp
--- a/Engine/Source/Runtime/Renderer/Private/MeshPassProcessor.cpp
+++ b/Engine/Source/Runtime/Renderer/Private/MeshPassProcessor.cpp
@@ -94,11 +94,6 @@ void FMeshDrawShaderBindings::SetShaderBindings(
 	const FReadOnlyMeshDrawSingleShaderBindings& RESTRICT SingleShaderBindings,
 	FShaderBindingState& RESTRICT ShaderBindingState)
 {
-#if USING_CODE_ANALYSIS
-	MSVC_PRAGMA(warning(push))
-	MSVC_PRAGMA(warning(disable : 6386)) // Parameter.BaseIndex may be out of bounds.
-#endif
-
 	const FUniformBufferRHIParamRef* RESTRICT UniformBufferBindings = SingleShaderBindings.GetUniformBufferStart();
 	const FShaderParameterInfo* RESTRICT UniformBufferParameters = SingleShaderBindings.ParameterMapInfo.UniformBuffers.GetData();
 	const int32 NumUniformBuffers = SingleShaderBindings.ParameterMapInfo.UniformBuffers.Num();
@@ -106,7 +101,7 @@ void FMeshDrawShaderBindings::SetShaderBindings(
 	for (int32 UniformBufferIndex = 0; UniformBufferIndex < NumUniformBuffers; UniformBufferIndex++)
 	{
 		FShaderParameterInfo Parameter = UniformBufferParameters[UniformBufferIndex];
-		checkSlow(Parameter.BaseIndex <= ARRAY_COUNT(ShaderBindingState.UniformBuffers));
+		checkSlow(Parameter.BaseIndex < ARRAY_COUNT(ShaderBindingState.UniformBuffers));
 		FUniformBufferRHIParamRef UniformBuffer = UniformBufferBindings[UniformBufferIndex];
 
 		if (UniformBuffer != ShaderBindingState.UniformBuffers[Parameter.BaseIndex])
@@ -124,7 +119,7 @@ void FMeshDrawShaderBindings::SetShaderBindings(
 	for (int32 SamplerIndex = 0; SamplerIndex < NumTextureSamplers; SamplerIndex++)
 	{
 		FShaderParameterInfo Parameter = TextureSamplerParameters[SamplerIndex];
-		checkSlow(Parameter.BaseIndex <= ARRAY_COUNT(ShaderBindingState.Samplers));
+		checkSlow(Parameter.BaseIndex < ARRAY_COUNT(ShaderBindingState.Samplers));
 		FSamplerStateRHIParamRef Sampler = SamplerBindings[SamplerIndex];
 
 		if (Sampler != ShaderBindingState.Samplers[Parameter.BaseIndex])
@@ -142,7 +137,7 @@ void FMeshDrawShaderBindings::SetShaderBindings(
 	for (int32 SRVIndex = 0; SRVIndex < NumSRVs; SRVIndex++)
 	{
 		FShaderParameterInfo Parameter = SRVParameters[SRVIndex];
-		checkSlow(Parameter.BaseIndex <= ARRAY_COUNT(ShaderBindingState.SRVs));
+		checkSlow(Parameter.BaseIndex < ARRAY_COUNT(ShaderBindingState.SRVs));
 		FShaderResourceViewRHIParamRef SRV = SRVBindings[SRVIndex];
 
 		if (SRV != ShaderBindingState.SRVs[Parameter.BaseIndex])
@@ -158,7 +153,7 @@ void FMeshDrawShaderBindings::SetShaderBindings(
 	for (int32 TextureIndex = 0; TextureIndex < NumSRVs; TextureIndex++)
 	{
 		FShaderParameterInfo Parameter = SRVParameters[TextureIndex];
-		checkSlow(Parameter.BaseIndex <= ARRAY_COUNT(ShaderBindingState.Textures));
+		checkSlow(Parameter.BaseIndex < ARRAY_COUNT(ShaderBindingState.Textures));
 		FTextureRHIParamRef Texture = TextureBindings[TextureIndex];
 
 		if (Texture != ShaderBindingState.Textures[Parameter.BaseIndex])
@@ -186,10 +181,6 @@ void FMeshDrawShaderBindings::SetShaderBindings(
 			LooseDataStart += Parameter.Size;
 		}
 	}
-
-#if USING_CODE_ANALYSIS
-	MSVC_PRAGMA(warning(pop))
-#endif
 }
 
 template<class RHIShaderType>
@@ -289,9 +280,9 @@ void FMeshDrawShaderBindings::SetOnRayTracingStructure(FRHICommandList& RHICmdLi
 	for (int32 UniformBufferIndex = 0; UniformBufferIndex < NumUniformBuffers; UniformBufferIndex++)
 	{
 		FShaderParameterInfo Parameter = UniformBufferParameters[UniformBufferIndex];
-		checkSlow(Parameter.BaseIndex <= ARRAY_COUNT(LocalUniformBuffers));
+		checkSlow(Parameter.BaseIndex < ARRAY_COUNT(LocalUniformBuffers));
 		FUniformBufferRHIParamRef UniformBuffer = UniformBufferBindings[UniformBufferIndex];
-		if (Parameter.BaseIndex <= ARRAY_COUNT(LocalUniformBuffers))
+		if (Parameter.BaseIndex < ARRAY_COUNT(LocalUniformBuffers))
 		{
 			LocalUniformBuffers[Parameter.BaseIndex] = UniformBuffer;
 			MaxUniformBufferUsed = FMath::Max((int32)Parameter.BaseIndex, MaxUniformBufferUsed);
