Implements AO denoiser using same pass layout as reflections.

diff --git a/Engine/Shaders/Private/RayTracing/CompositeAmbientOcclusionPS.usf b/Engine/Shaders/Private/RayTracing/CompositeAmbientOcclusionPS.usf
--- a/Engine/Shaders/Private/RayTracing/CompositeAmbientOcclusionPS.usf
+++ b/Engine/Shaders/Private/RayTracing/CompositeAmbientOcclusionPS.usf
@@ -11,6 +11,9 @@ void CompositeAmbientOcclusionPS(
 	out float4 OutColor : SV_Target0
 )
 {
-	float4 AmbientOcclusion = AmbientOcclusionTexture.Sample(AmbientOcclusionTextureSampler, UV);
-	OutColor = float4(AmbientOcclusion.rgb, 1.0);
+	float AmbientOcclusionMask = AmbientOcclusionTexture.Sample(AmbientOcclusionTextureSampler, UV).r;
+
+	AmbientOcclusionMask = saturate(AmbientOcclusionMask);
+
+	OutColor = float4(AmbientOcclusionMask, AmbientOcclusionMask, AmbientOcclusionMask, 1.0);
 }
diff --git a/Engine/Shaders/Private/RayTracing/RayTracingAmbientOcclusionRGS.usf b/Engine/Shaders/Private/RayTracing/RayTracingAmbientOcclusionRGS.usf
--- a/Engine/Shaders/Private/RayTracing/RayTracingAmbientOcclusionRGS.usf
+++ b/Engine/Shaders/Private/RayTracing/RayTracingAmbientOcclusionRGS.usf
@@ -7,8 +7,7 @@
 #include "RayTracingCommon.ush"
 
 RaytracingAccelerationStructure TLAS;
-RWTexture2D<float4> RWOcclusionMaskUAV;
-RWTexture2D<float> RWRayDistanceUAV;
+RWTexture2D<float2> RWOcclusionMaskUAV;
 
 uint CalcLinearIndex(uint2 PixelCoord)
 {
@@ -67,67 +66,71 @@ void AmbientOcclusionRGS()
 	float3 WorldPosition = ReconstructWorldPositionFromDepth(UV, Depth);
 	float3 WorldNormal = ScreenSpaceData.GBuffer.WorldNormal;
 
+	float RayCount = 0.0;
 	float Visibility = 0.0;
-	float RayDistance = 0.0;
-	float HitCount = 0.0;
+	float ClosestRayHitDistance = 1.#INF;
 	uint SamplesPerPixel = AmbientOcclusion.SamplesPerPixel;
 	
 	// Mask out depth values beyond far plane
-	if (Depth >= SCENE_TEXTURES_DISABLED_SCENE_DEPTH_VALUE)
+	if (Depth >= SCENE_TEXTURES_DISABLED_SCENE_DEPTH_VALUE || ScreenSpaceData.GBuffer.ShadingModelID == SHADINGMODELID_UNLIT)
 	{
-		SamplesPerPixel = 0.0;
 		Visibility = 1.0;
+		RayCount = SamplesPerPixel;
+		SamplesPerPixel = 0.0;
 	}
 
 	for (uint SampleIndex = 0; SampleIndex < SamplesPerPixel; ++SampleIndex)
 	{
 		RayDesc Ray;
 		GenerateCosineNormalRay(WorldPosition, WorldNormal, RandSequence, Ray.Origin, Ray.Direction, Ray.TMin, Ray.TMax);
-		if (dot(WorldNormal, Ray.Direction) > 0.0)
-		{
-			uint RayFlags = 0
-				| RAY_FLAG_CULL_BACK_FACING_TRIANGLES;
-
-			FDefaultPayload Payload = (FDefaultPayload)0;
-
-			TraceRay(
-				TLAS,
-				RayFlags,
-				0xFF, // InstanceInclusionMask
-				0,    // RayContributionToHitGroupIndex
-				1,    // MultiplierForGeometryContributionToShaderIndex
-				0,    // MissShaderIndex
-				Ray,
-				Payload
-			);
-
-			Visibility += Payload.HitT < 0 ? 1.0 : 0.0;
-			if (Payload.HitT >= 0.0)
-			{
-				RayDistance += Payload.HitT;
-				HitCount += 1.0;
-			}
-		}
-		else
+		
+		if (dot(WorldNormal, Ray.Direction) <= 0.0) // TODO: does it needs to be handled by the denoiser?
+			continue;
+		
+		uint RayFlags = 0
+			| RAY_FLAG_CULL_BACK_FACING_TRIANGLES;
+
+		FDefaultPayload Payload = (FDefaultPayload)0;
+
+		TraceRay(
+			TLAS,
+			RayFlags,
+			0xFF, // InstanceInclusionMask
+			0,    // RayContributionToHitGroupIndex
+			1,    // MultiplierForGeometryContributionToShaderIndex
+			0,    // MissShaderIndex
+			Ray,
+			Payload
+		);
+		
+		RayCount += 1;
+		Visibility += Payload.HitT < 0 ? 1.0 : 0.0;
+		if (Payload.HitT >= 0.0)
 		{
-			Visibility += 0.0;
-			RayDistance += 0.0;
-			HitCount += 1.0;
+			ClosestRayHitDistance = min(ClosestRayHitDistance, Payload.HitT);
 		}
 	}
 
-	RWOcclusionMaskUAV[PixelCoord] = Visibility;
-	if (SamplesPerPixel > 0)
+	// Output.
 	{
-		RWOcclusionMaskUAV[PixelCoord] /= SamplesPerPixel;
-	}
+		float2 RawOutput = 1;
 
-	if (HitCount > 0.0)
-	{
-		RWRayDistanceUAV[PixelCoord] = RayDistance / HitCount;
-	}
-	else
-	{
-		RWRayDistanceUAV[PixelCoord] = 1.0e27;
+		if (SamplesPerPixel == 0)
+		{
+			RawOutput.x = 1;
+			RawOutput.y = -2;
+		}
+		else if (RayCount == 0)
+		{
+			RawOutput.x = 1;
+			RawOutput.y = -1;
+		}
+		else
+		{
+			RawOutput.x = Visibility / RayCount;
+			RawOutput.y = ClosestRayHitDistance;
+		}
+
+		RWOcclusionMaskUAV[PixelCoord] = RawOutput;
 	}
 }
\ No newline at end of file
diff --git a/Engine/Shaders/Private/ScreenSpaceDenoise/SSDDefinitions.ush b/Engine/Shaders/Private/ScreenSpaceDenoise/SSDDefinitions.ush
--- a/Engine/Shaders/Private/ScreenSpaceDenoise/SSDDefinitions.ush
+++ b/Engine/Shaders/Private/ScreenSpaceDenoise/SSDDefinitions.ush
@@ -75,23 +75,24 @@
 	/** Buffer layout for the reflection output. */
 	#define SIGNAL_BUFFER_LAYOUT_REFLECTIONS_HISTORY 4
 
+	/** Buffer layout given by the ray generation shader for ambient occlusion. */
+	#define SIGNAL_BUFFER_LAYOUT_AO_INPUT 5
 
-/** Defines how the signal is being processed. Matches C++'s ESignalProcessing. */
-	/** Scene color anti aliasing processing. */
-	//#define SIGNAL_PROCESSING_COLORAA 0
+	/** Internal buffer layout for ambient occlusion. */
+	#define SIGNAL_BUFFER_LAYOUT_AO_HISTORY 6
 
-	/** Depth of field stabilization processing. */
-	//#define SIGNAL_PROCESSING_DOF 1
 
+/** Defines how the signal is being processed. Matches C++'s ESignalProcessing. */
 	/** Shadow penumbra denoising. */
 	#define SIGNAL_PROCESSING_PENUMBRA 0
 
-	/** Shadow penumbra denoising. */
-	//#define SIGNAL_PROCESSING_DIFUSE_GI 3
-
 	/** Reflection denoising. */
 	#define SIGNAL_PROCESSING_REFLECTIONS 1
 
+	/** Ambient occlusion denoising. */
+	#define SIGNAL_PROCESSING_AO 2
+
+
 
 /** Defines the color space bitfield. */
 	#define COLOR_SPACE_RGB 0x0
diff --git a/Engine/Shaders/Private/ScreenSpaceDenoise/SSDSignalCore.ush b/Engine/Shaders/Private/ScreenSpaceDenoise/SSDSignalCore.ush
--- a/Engine/Shaders/Private/ScreenSpaceDenoise/SSDSignalCore.ush
+++ b/Engine/Shaders/Private/ScreenSpaceDenoise/SSDSignalCore.ush
@@ -46,71 +46,82 @@ FSSFSignalSample CreateSignalSample()
 //------------------------------------------------------- INPUT/OUTPUT
 
 /** Decode input signal sample. */
-FSSFSignalSample DecodeInputSignal(const uint SignalBufferLayout, float4 Samples[2])
+FSSFSignalSample DecodeInputSignal(const uint SignalBufferLayout, float4 RawSample[2])
 {
 	FSSFSignalSample OutSample = CreateSignalSample();
 
 	if (SignalBufferLayout == SIGNAL_BUFFER_LAYOUT_PENUMBRA_INPUT)
 	{
 		OutSample.AccumulatedSampleCount = 1;
-		OutSample.LeavingRayCount = OutSample.AccumulatedSampleCount * Samples[0].r;
-		OutSample.PenumbraClosestHit = Samples[1].r;
+		OutSample.LeavingRayCount = OutSample.AccumulatedSampleCount * RawSample[0].r;
+		OutSample.PenumbraClosestHit = RawSample[1].r;
 	}
 	else if (SignalBufferLayout == SIGNAL_BUFFER_LAYOUT_PENUMBRA_INTERNAL)
 	{
 		if (DEBUG_STORE_NORMALIZED_COMPONENTS)
 		{
-			OutSample.LeavingRayCount = Samples[0].r * Samples[0].b;
-			OutSample.UncertainRayCount = Samples[0].g * Samples[0].b;
-			OutSample.AccumulatedSampleCount = Samples[0].b - OutSample.UncertainRayCount;
+			OutSample.LeavingRayCount = RawSample[0].r * RawSample[0].b;
+			OutSample.UncertainRayCount = RawSample[0].g * RawSample[0].b;
+			OutSample.AccumulatedSampleCount = RawSample[0].b - OutSample.UncertainRayCount;
 		}
 		else
 		{
-			OutSample.LeavingRayCount = Samples[0].r;
-			OutSample.AccumulatedSampleCount = Samples[0].g;
-			OutSample.UncertainRayCount = Samples[0].b;
+			OutSample.LeavingRayCount = RawSample[0].r;
+			OutSample.AccumulatedSampleCount = RawSample[0].g;
+			OutSample.UncertainRayCount = RawSample[0].b;
 		}
 	}
 	else if (SignalBufferLayout == SIGNAL_BUFFER_LAYOUT_PENUMBRA_HISTORY_REJECTION)
 	{
 		//if (Samples[0].r < 1)
-		if (Samples[0].r >= 0)
+		if (RawSample[0].r >= 0)
 		{
 			OutSample.AccumulatedSampleCount = 1;
-			OutSample.LeavingRayCount = OutSample.AccumulatedSampleCount * Samples[0].r; // *(255.0 / 254.0);
+			OutSample.LeavingRayCount = OutSample.AccumulatedSampleCount * RawSample[0].r; // *(255.0 / 254.0);
 		}
 	}
 	else if (SignalBufferLayout == SIGNAL_BUFFER_LAYOUT_REFLECTIONS_INPUT)
 	{
-		OutSample.AccumulatedSampleCount = Samples[1].r == -2 ? 0 : 1;
-		OutSample.LeavingRayCount = OutSample.AccumulatedSampleCount * (1 - Samples[0].a);
-		OutSample.SceneColor = Samples[0];
+		OutSample.AccumulatedSampleCount = RawSample[1].r == -2 ? 0 : 1;
+		OutSample.LeavingRayCount = OutSample.AccumulatedSampleCount * (1 - RawSample[0].a);
+		OutSample.SceneColor = RawSample[0];
 
 		// max() because ray hit distance is < 0 when miss / not ray has been shot.
-		OutSample.PenumbraClosestHit = Samples[1].r * OutSample.AccumulatedSampleCount;
-		OutSample.PenumbraClosestHit = ((Samples[1].r == -1.0) ? 1.0e20 : Samples[1].r) * OutSample.AccumulatedSampleCount;
+		OutSample.PenumbraClosestHit = RawSample[1].r * OutSample.AccumulatedSampleCount;
+		OutSample.PenumbraClosestHit = ((RawSample[1].r == -1.0) ? 1.0e20 : RawSample[1].r) * OutSample.AccumulatedSampleCount;
 	}
 	else if (SignalBufferLayout == SIGNAL_BUFFER_LAYOUT_REFLECTIONS_HISTORY)
 	{
-		OutSample.AccumulatedSampleCount = Samples[1].r;
-		OutSample.LeavingRayCount = OutSample.AccumulatedSampleCount * (1 - Samples[0].a);
-		OutSample.SceneColor = OutSample.AccumulatedSampleCount * Samples[0];
+		OutSample.AccumulatedSampleCount = RawSample[1].r;
+		OutSample.LeavingRayCount = OutSample.AccumulatedSampleCount * (1 - RawSample[0].a);
+		OutSample.SceneColor = OutSample.AccumulatedSampleCount * RawSample[0];
+	}
+	else if (SignalBufferLayout == SIGNAL_BUFFER_LAYOUT_AO_INPUT)
+	{
+		OutSample.AccumulatedSampleCount = RawSample[0].g < 0 ? 0 : 1;
+		OutSample.LeavingRayCount = OutSample.AccumulatedSampleCount * RawSample[0].r;
+		OutSample.PenumbraClosestHit = OutSample.AccumulatedSampleCount * RawSample[0].g;
+	}
+	else if (SignalBufferLayout == SIGNAL_BUFFER_LAYOUT_AO_HISTORY)
+	{
+		OutSample.AccumulatedSampleCount = RawSample[0].g;
+		OutSample.LeavingRayCount = OutSample.AccumulatedSampleCount * RawSample[0].r;
 	}
 
 	return OutSample;
 }
 
 /** Encode output signal sample. */
-void EncodeOutputSignal(const uint SignalBufferLayout, FSSFSignalSample Sample, out float4 OutSamples[2], out uint OutBufferCount)
+void EncodeOutputSignal(const uint SignalBufferLayout, FSSFSignalSample Sample, out float4 OutRawSample[2], out uint OutBufferCount)
 {
-	OutSamples[0] = OutSamples[1] = 0;
+	OutRawSample[0] = OutRawSample[1] = 0;
 	OutBufferCount = 1;
 
 	if (SignalBufferLayout == SIGNAL_BUFFER_LAYOUT_PENUMBRA_INTERNAL)
 	{
 		if (DEBUG_STORE_NORMALIZED_COMPONENTS)
 		{
-			OutSamples[0] = float4(
+			OutRawSample[0] = float4(
 				Sample.UncertainRayCount + Sample.AccumulatedSampleCount > 0
 					? Sample.LeavingRayCount / (Sample.UncertainRayCount + Sample.AccumulatedSampleCount)
 					: -1.0,
@@ -122,7 +133,7 @@ void EncodeOutputSignal(const uint SignalBufferLayout, FSSFSignalSample Sample,
 		}
 		else
 		{
-			OutSamples[0] = float4(
+			OutRawSample[0] = float4(
 				Sample.LeavingRayCount,
 				Sample.AccumulatedSampleCount,
 				Sample.UncertainRayCount,
@@ -133,43 +144,50 @@ void EncodeOutputSignal(const uint SignalBufferLayout, FSSFSignalSample Sample,
 	{
 		//if (Sample.AccumulatedSampleCount == 0)
 		//{
-		//	OutSamples[0].r = 1.0;
+		//	OutRawSample[0].r = 1.0;
 		//}
 		//else
 		//{
-		//	OutSamples[0].r = (254.0 / 255.0) * (Sample.LeavingRayCount / Sample.AccumulatedSampleCount);
+		//	OutRawSample[0].r = (254.0 / 255.0) * (Sample.LeavingRayCount / Sample.AccumulatedSampleCount);
 		//}
 
-		OutSamples[0].r = Sample.AccumulatedSampleCount == 0 ? -1 : Sample.LeavingRayCount / Sample.AccumulatedSampleCount;
+		OutRawSample[0].r = Sample.AccumulatedSampleCount == 0 ? -1 : Sample.LeavingRayCount / Sample.AccumulatedSampleCount;
 	}
 	else if (SignalBufferLayout == SIGNAL_BUFFER_LAYOUT_REFLECTIONS_HISTORY)
 	{
-		OutSamples[0] = Sample.SceneColor * (Sample.AccumulatedSampleCount > 0 ? rcp(Sample.AccumulatedSampleCount) : 0);
-		OutSamples[1].r = Sample.AccumulatedSampleCount;
+		// Outputs scene color to be compatible with a SSR output.
+		OutRawSample[0] = Sample.SceneColor * (Sample.AccumulatedSampleCount > 0 ? rcp(Sample.AccumulatedSampleCount) : 0);
+		OutRawSample[1].r = Sample.AccumulatedSampleCount;
 		OutBufferCount = 2;
 	}
+	else if (SignalBufferLayout == SIGNAL_BUFFER_LAYOUT_AO_HISTORY)
+	{
+		// Outputs number of ray miss in red to be compatible as a SSAO output.
+		OutRawSample[0].r = Sample.AccumulatedSampleCount > 0 ? Sample.LeavingRayCount / Sample.AccumulatedSampleCount : 1.0;
+		OutRawSample[0].g = Sample.AccumulatedSampleCount;
+	}
 }
 
 /** Sample input signal. */
 FSSFSignalSample SampleInputSignal(Texture2D SignalBuffer0, Texture2D SignalBuffer1, SamplerState Sampler, const uint SignalBufferLayout, float2 UV, float Level = 0)
 {
-	float4 Samples[2];
-	Samples[0] = SignalBuffer0.SampleLevel(Sampler, UV, Level);
-	Samples[1] = SignalBuffer1.SampleLevel(Sampler, UV, Level);
-	return DecodeInputSignal(SignalBufferLayout, Samples);
+	float4 RawSample[2];
+	RawSample[0] = SignalBuffer0.SampleLevel(Sampler, UV, Level);
+	RawSample[1] = SignalBuffer1.SampleLevel(Sampler, UV, Level);
+	return DecodeInputSignal(SignalBufferLayout, RawSample);
 }
 
 /** Sample output signal. */
 void OutputSignal(RWTexture2D<float4> OutputSignalBuffer0, RWTexture2D<float4> OutputSignalBuffer1, const uint SignalBufferLayout, const uint2 PixelPosition, FSSFSignalSample Sample)
 {
-	float4 Samples[2];
+	float4 RawSample[2];
 	uint BufferCount;
-	EncodeOutputSignal(SignalBufferLayout, Sample, Samples, BufferCount);
+	EncodeOutputSignal(SignalBufferLayout, Sample, RawSample, BufferCount);
 
 	if (BufferCount >= 1)
-		OutputSignalBuffer0[PixelPosition] = Samples[0];
+		OutputSignalBuffer0[PixelPosition] = RawSample[0];
 	if (BufferCount >= 2)
-		OutputSignalBuffer1[PixelPosition] = Samples[1];
+		OutputSignalBuffer1[PixelPosition] = RawSample[1];
 }
 
 
diff --git a/Engine/Shaders/Private/ScreenSpaceDenoise/SSDSignalFramework.ush b/Engine/Shaders/Private/ScreenSpaceDenoise/SSDSignalFramework.ush
--- a/Engine/Shaders/Private/ScreenSpaceDenoise/SSDSignalFramework.ush
+++ b/Engine/Shaders/Private/ScreenSpaceDenoise/SSDSignalFramework.ush
@@ -40,6 +40,8 @@
 	#include "../DeferredShadingCommon.ush"
 	#include "../CommonViewUniformBuffer.ush"
 
+#elif CONFIG_SIGNAL_PROCESSING == SIGNAL_PROCESSING_AO
+
 #else
 	#error Unknown CONFIG_SIGNAL_PROCESSING
 
@@ -383,7 +385,7 @@ float ComputeBilateralWeight(
 	float MaxWorldBluringDistance,
 	FSSFSampleSceneInfos RefSceneMetadata,
 	FSSFSampleSceneInfos NeighborSceneMetadata)
-#if CONFIG_SIGNAL_PROCESSING == SIGNAL_PROCESSING_PENUMBRA
+#if CONFIG_SIGNAL_PROCESSING == SIGNAL_PROCESSING_PENUMBRA || CONFIG_SIGNAL_PROCESSING == SIGNAL_PROCESSING_AO
 {
 	if (DEBUG_DISABLE_BILATERAL)
 	{
@@ -410,16 +412,27 @@ float ComputeBilateralWeight(
 
 		WorldBilateralWeight = saturate(1 - DistSquare * Multiplier);
 	}
-	
-	float LobeBilateralWeight = 1;
-	#if CONFIG_TOKOYASHI_BILATERAL
-	//if (View.GeneralPurposeTweak > 0.5)
-	{
-		LobeBilateralWeight = TokoyashiLobeBilateralWeight(RefSceneMetadata, NeighborSceneMetadata);
-	}
-	#endif
 
-	return WorldBilateralWeight;
+	float NormalBilateralWeight = 1;
+	if (CONFIG_SIGNAL_PROCESSING == SIGNAL_PROCESSING_AO)
+    {
+		float NoN = max(dot(RefSceneMetadata.WorldNormal, NeighborSceneMetadata.WorldNormal), 0);
+
+		#if 1 // TODO
+			NormalBilateralWeight = pow(NoN, 4);
+		#else
+			// [ Kontkanen et al 2004, "Irradiance Filtering for Monte Carlo Ray Tracing" ]
+			// TODO: looks broken
+			float m = MaxWorldBluringDistance;
+			float a = 0.3;
+			float d = sqrt(1 - NoN); // + length(RefSceneMetadata.TranslatedWorldPosition - NeighborSceneMetadata.TranslatedWorldPosition) * rcp(m);
+			float x = d * rcp(a);
+
+			NormalBilateralWeight = exp(-x * x);
+		#endif
+    }
+
+	return WorldBilateralWeight * NormalBilateralWeight;
 }
 #elif CONFIG_SIGNAL_PROCESSING == SIGNAL_PROCESSING_REFLECTIONS
 {
@@ -432,6 +445,8 @@ float ComputeBilateralWeight(
 
 	return LobeBilateralWeight;
 }
+#else
+	#error Unimplemented ComputeBilateralWeight()
 #endif
 
 float GetSignalWorldBluringRadius(FSSFSignalSample Sample, FSSFSampleSceneInfos SceneMetadata)
@@ -488,6 +503,12 @@ float GetSignalWorldBluringRadius(FSSFSignalSample Sample, FSSFSampleSceneInfos
 {
 	return 1;
 }
+#elif CONFIG_SIGNAL_PROCESSING == SIGNAL_PROCESSING_AO
+{
+	return 1000; //Sample.PenumbraClosestHit;
+}
+#else
+	#error Unimplemented GetSignalWorldBluringRadius()
 #endif
 
 FSSFSignalSample GetConstantSignalSample(uint TileMode)
diff --git a/Engine/Shaders/Private/ScreenSpaceDenoise/SSDSpatialAccumulation.usf b/Engine/Shaders/Private/ScreenSpaceDenoise/SSDSpatialAccumulation.usf
--- a/Engine/Shaders/Private/ScreenSpaceDenoise/SSDSpatialAccumulation.usf
+++ b/Engine/Shaders/Private/ScreenSpaceDenoise/SSDSpatialAccumulation.usf
@@ -27,9 +27,11 @@
 #define TILE_PIXEL_SIZE 8
 
 #define CONFIG_SIGNAL_PROCESSING DIM_SIGNAL_PROCESSING
+#define CONFIG_UPSCALE           DIM_UPSCALE
 
 
 #if CONFIG_SIGNAL_PROCESSING == SIGNAL_PROCESSING_PENUMBRA
+	// Input and output layout.
 	#define CONFIG_SIGNAL_INPUT_LAYOUT  SIGNAL_BUFFER_LAYOUT_PENUMBRA_INTERNAL
 	#define CONFIG_SIGNAL_OUTPUT_LAYOUT SIGNAL_BUFFER_LAYOUT_PENUMBRA_HISTORY_REJECTION
 
@@ -42,6 +44,7 @@
 
 #elif CONFIG_SIGNAL_PROCESSING == SIGNAL_PROCESSING_REFLECTIONS
 	#if DIM_STAGE == STAGE_RECONSTRUCTION
+		// Input and output layout.
 		#define CONFIG_SIGNAL_INPUT_LAYOUT  SIGNAL_BUFFER_LAYOUT_REFLECTIONS_INPUT
 		#define CONFIG_SIGNAL_OUTPUT_LAYOUT SIGNAL_BUFFER_LAYOUT_REFLECTIONS_HISTORY
 
@@ -53,7 +56,6 @@
 		// Reconstruction has 1 ray per input pixel, therefore can use MIS.
 		#define CONFIG_SPECULAR_RATIO_ESTIMATOR 1
 
-		#define CONFIG_UPSCALE              DIM_UPSCALE
 		//#define CONFIG_SAMPLE_COUNT_POLICY  SAMPLE_COUNT_POLICY_DISABLED
 
 		// Do color accumulation with karis weighting to avoid flickering specular highlight to show up the kernel pattern.
@@ -61,6 +63,7 @@
 
 
 	#elif DIM_STAGE == STAGE_POST_FILTERING
+		// Input and output layout.
 		#define CONFIG_SIGNAL_INPUT_LAYOUT  SIGNAL_BUFFER_LAYOUT_REFLECTIONS_HISTORY
 		#define CONFIG_SIGNAL_OUTPUT_LAYOUT SIGNAL_BUFFER_LAYOUT_REFLECTIONS_HISTORY
 
@@ -79,6 +82,25 @@
 		#error Unexpected stage.
 	#endif
 
+#elif CONFIG_SIGNAL_PROCESSING == SIGNAL_PROCESSING_AO
+	#if DIM_STAGE == STAGE_RECONSTRUCTION
+		// Input and output layout.
+		#define CONFIG_SIGNAL_INPUT_LAYOUT  SIGNAL_BUFFER_LAYOUT_AO_INPUT
+		#define CONFIG_SIGNAL_OUTPUT_LAYOUT SIGNAL_BUFFER_LAYOUT_AO_HISTORY
+
+		#define CONFIG_SAMPLE_SET           SAMPLE_SET_STACKOWIAK_4_SETS
+
+	#elif DIM_STAGE == STAGE_POST_FILTERING
+		// Input and output layout.
+		#define CONFIG_SIGNAL_INPUT_LAYOUT  SIGNAL_BUFFER_LAYOUT_AO_INPUT
+		#define CONFIG_SIGNAL_OUTPUT_LAYOUT SIGNAL_BUFFER_LAYOUT_AO_HISTORY
+
+		#define CONFIG_SAMPLE_SET           SAMPLE_SET_3X3
+
+	#else
+		#error Unexpected stage.
+	#endif
+
 #else
 	#error Unknown signal processing.
 #endif
diff --git a/Engine/Shaders/Private/ScreenSpaceDenoise/SSDTemporalAccumulation.usf b/Engine/Shaders/Private/ScreenSpaceDenoise/SSDTemporalAccumulation.usf
--- a/Engine/Shaders/Private/ScreenSpaceDenoise/SSDTemporalAccumulation.usf
+++ b/Engine/Shaders/Private/ScreenSpaceDenoise/SSDTemporalAccumulation.usf
@@ -46,6 +46,13 @@
 
 	//#define CONFIG_HISTORY_BLENDING_COLOR_SPACE (COLOR_SPACE_RGB | COLOR_SPACE_UNPREMULTIPLY | COLOR_SPACE_KARIS_WEIGHTING)
 
+#elif CONFIG_SIGNAL_PROCESSING == SIGNAL_PROCESSING_AO
+	// Input and output layout.
+	#define CONFIG_SIGNAL_INPUT_LAYOUT   SIGNAL_BUFFER_LAYOUT_AO_HISTORY
+	#define CONFIG_SIGNAL_HISTORY_LAYOUT SIGNAL_BUFFER_LAYOUT_AO_HISTORY
+
+	#define CONFIG_HISTORY_REJECTION 1
+
 #else
 	#error Unimplemented signal processing.
 #endif
@@ -260,6 +267,12 @@ void TemporallyAccumulate(
 		//	max(CurrentFrame.AccumulatedSampleCount - TargetedSampleCount, 0),
 		//	0);
 	}
+	#elif CONFIG_SIGNAL_PROCESSING == SIGNAL_PROCESSING_AO
+	{
+		TargetedSampleCount = TARGETED_SAMPLE_COUNT;
+	}
+	#else
+		#error
 	#endif
 
 	// Sample the history buffer.
@@ -558,7 +571,7 @@ void TemporallyAccumulate(
 	
 	FSSFSignalSample HistorySample = PreviousFrame;
 
-	const bool bPostRejectionBlending = !bBruteForce || CONFIG_SIGNAL_PROCESSING == SIGNAL_PROCESSING_REFLECTIONS;
+	const bool bPostRejectionBlending = !bBruteForce || CONFIG_SIGNAL_PROCESSING == SIGNAL_PROCESSING_REFLECTIONS || CONFIG_SIGNAL_PROCESSING == SIGNAL_PROCESSING_AO;
 
 	if (!bPostRejectionBlending)
 	{
@@ -849,6 +862,10 @@ void TemporallyAccumulate(
 					RejectionFactor = max(RejectionFactor, 1 * max(RejectedDiff.SceneColor.y, RejectedDiff.SceneColor.z));
 					RejectionFactor = max(RejectionFactor, 1 * RejectedDiff.SceneColor.a);
 				}
+				#elif CONFIG_SIGNAL_PROCESSING == SIGNAL_PROCESSING_AO
+				{
+					RejectionFactor = abs(ClampedTransformedHistorySample.LeavingRayCount - TransformedHistorySample.LeavingRayCount);
+				}
 				#else
 					#error Unsupported signal rejection.
 				#endif
diff --git a/Engine/Source/Runtime/Renderer/Private/DeferredShadingRenderer.cpp b/Engine/Source/Runtime/Renderer/Private/DeferredShadingRenderer.cpp
--- a/Engine/Source/Runtime/Renderer/Private/DeferredShadingRenderer.cpp
+++ b/Engine/Source/Runtime/Renderer/Private/DeferredShadingRenderer.cpp
@@ -31,6 +31,8 @@
 #include "GPUScene.h"
 #include "RayTracing/RayTracingMaterialHitShaders.h"
 #include "RayTracing/RayTracingDynamicGeometryCollection.h"
+#include "SceneViewFamilyBlackboard.h"
+#include "ScreenSpaceDenoise.h"
 
 TAutoConsoleVariable<int32> CVarEarlyZPass(
 	TEXT("r.EarlyZPass"),
@@ -157,6 +159,15 @@ static TAutoConsoleVariable<int32> CVarRayTracing(
 	TEXT(" 1: on"),
 	ECVF_RenderThreadSafe | ECVF_ReadOnly);
 
+static TAutoConsoleVariable<int32> CVarUseAODenoiser(
+	TEXT("r.AmbientOcclusion.Denoiser"),
+	2,
+	TEXT("Choose the denoising algorithm.\n")
+	TEXT(" 0: Disabled;\n")
+	TEXT(" 1: Forces the default denoiser of the renderer;\n")
+	TEXT(" 2: GScreenSpaceDenoiser witch may be overriden by a third party plugin (default)."),
+	ECVF_RenderThreadSafe);
+
 
 DECLARE_CYCLE_STAT(TEXT("PostInitViews FlushDel"), STAT_PostInitViews_FlushDel, STATGROUP_InitViews);
 DECLARE_CYCLE_STAT(TEXT("InitViews Intentional Stall"), STAT_InitViews_Intentional_Stall, STATGROUP_InitViews);
@@ -1667,6 +1678,8 @@ void FDeferredShadingSceneRenderer::Render(FRHICommandListImmediate& RHICmdList)
 
 		if (bCanOverlayRayTracingOutput)
 		{
+			// TODO: convert the entire AO and skylight to rendergraph.
+
 			// SkyLight takes priority over ambient occlusion
 			if (ShouldRenderRayTracingDynamicSkyLight(Scene, ViewFamily))
 			{
@@ -1682,6 +1695,43 @@ void FDeferredShadingSceneRenderer::Render(FRHICommandListImmediate& RHICmdList)
 				TRefCountPtr<IPooledRenderTarget> AmbientOcclusionRT;
 				RenderRayTracingAmbientOcclusion(RHICmdList, nullptr, AmbientOcclusionRT);
 				checkSlow(RHICmdList.IsOutsideRenderPass());
+				
+				int32 DenoiserMode = CVarUseAODenoiser.GetValueOnRenderThread();
+				if (DenoiserMode != 0)
+				{	
+					FRDGBuilder GraphBuilder(RHICmdList);
+
+					FSceneViewFamilyBlackboard SceneBlackboard;
+					SetupSceneViewFamilyBlackboard(GraphBuilder, &SceneBlackboard);
+
+					const IScreenSpaceDenoiser* DefaultDenoiser = IScreenSpaceDenoiser::GetDefaultDenoiser();
+					const IScreenSpaceDenoiser* DenoiserToUse = DenoiserMode == 1 ? DefaultDenoiser : GScreenSpaceDenoiser;
+					
+					IScreenSpaceDenoiser::FAmbientOcclusionRayTracingConfig RayTracingConfig;
+
+					IScreenSpaceDenoiser::FAmbientOcclusionInputs DenoiserInputs;
+					DenoiserInputs.MaskAndRayHitDistance = GraphBuilder.RegisterExternalTexture(AmbientOcclusionRT, TEXT("AOMaskAndHitDistance"));
+
+					const FViewInfo& View = Views[0];
+
+					{
+						RDG_EVENT_SCOPE(GraphBuilder, "%s%s(AmbientOcclusion) %dx%d",
+							DenoiserToUse != DefaultDenoiser ? TEXT("ThirdParty ") : TEXT(""),
+							DenoiserToUse->GetDebugName(),
+							View.ViewRect.Width(), View.ViewRect.Height());
+
+						IScreenSpaceDenoiser::FAmbientOcclusionOutputs DenoiserOutputs = DenoiserToUse->DenoiseAmbientOcclusion(
+							GraphBuilder,
+							View,
+							SceneBlackboard,
+							DenoiserInputs,
+							RayTracingConfig);
+						
+						GraphBuilder.QueueTextureExtraction(DenoiserOutputs.AmbientOcclusionMask, &AmbientOcclusionRT);
+					}
+					GraphBuilder.Execute();
+				}
+
 				// #dxr_todo: Denoise AO
 				CompositeRayTracingAmbientOcclusion(RHICmdList, AmbientOcclusionRT);
 				checkSlow(RHICmdList.IsOutsideRenderPass());
diff --git a/Engine/Source/Runtime/Renderer/Private/RayTracing/RayTracingAmbientOcclusion.cpp b/Engine/Source/Runtime/Renderer/Private/RayTracing/RayTracingAmbientOcclusion.cpp
--- a/Engine/Source/Runtime/Renderer/Private/RayTracing/RayTracingAmbientOcclusion.cpp
+++ b/Engine/Source/Runtime/Renderer/Private/RayTracing/RayTracingAmbientOcclusion.cpp
@@ -78,7 +78,6 @@ public:
 		AmbientOcclusionParameter.Bind(Initializer.ParameterMap, TEXT("AmbientOcclusion"));
 
 		OcclusionMaskUAVParameter.Bind(Initializer.ParameterMap, TEXT("RWOcclusionMaskUAV"));
-		RayDistanceUAVParameter.Bind(Initializer.ParameterMap, TEXT("RWRayDistanceUAV"));
 	}
 
 	bool Serialize(FArchive& Ar)
@@ -89,7 +88,6 @@ public:
 		Ar << SceneTexturesParameter;
 		Ar << AmbientOcclusionParameter;
 		Ar << OcclusionMaskUAVParameter;
-		Ar << RayDistanceUAVParameter;
 		return bShaderHasOutdatedParameters;
 	}
 
@@ -100,7 +98,6 @@ public:
 		FUniformBufferRHIParamRef SceneTexturesUniformBuffer,
 		FUniformBufferRHIParamRef AmbientOcclusionUniformBuffer,
 		FUnorderedAccessViewRHIParamRef OcclusionMaskUAV,
-		FUnorderedAccessViewRHIParamRef RayDistanceUAV,
 		uint32 Width, uint32 Height
 	)
 	{
@@ -115,7 +112,6 @@ public:
 		GlobalResources.Set(SceneTexturesParameter, SceneTexturesUniformBuffer);
 		GlobalResources.Set(AmbientOcclusionParameter, AmbientOcclusionUniformBuffer);
 		GlobalResources.Set(OcclusionMaskUAVParameter, OcclusionMaskUAV);
-		GlobalResources.Set(RayDistanceUAVParameter, RayDistanceUAV);
 
 		RHICmdList.RayTraceDispatch(Pipeline, GlobalResources, Width, Height);
 	}
@@ -145,18 +141,11 @@ void FDeferredShadingSceneRenderer::RenderRayTracingAmbientOcclusion(
 
 	FSceneRenderTargets& SceneContext = FSceneRenderTargets::Get(RHICmdList);
 	FPooledRenderTargetDesc Desc = SceneContext.GetSceneColor()->GetDesc();
-	Desc.Format = PF_FloatRGBA;
+	Desc.Format = PF_G16R16F;
 	Desc.Flags &= ~(TexCreate_FastVRAM | TexCreate_Transient);
 	GRenderTargetPool.FindFreeElement(RHICmdList, Desc, AmbientOcclusionRT, TEXT("RayTracingAmbientOcclusion"));
 	ClearUAV(RHICmdList, AmbientOcclusionRT->GetRenderTargetItem(), FLinearColor::Black);
 
-	TRefCountPtr<IPooledRenderTarget> RayDistanceRT;
-	Desc = SceneContext.GetSceneColor()->GetDesc();
-	Desc.Format = PF_R16F;
-	Desc.Flags &= ~(TexCreate_FastVRAM | TexCreate_Transient);
-	GRenderTargetPool.FindFreeElement(RHICmdList, Desc, RayDistanceRT, TEXT("RayTracingAmbientOcclusionDistance"));
-	ClearUAV(RHICmdList, RayDistanceRT->GetRenderTargetItem(), FLinearColor::Black);
-
 	// Add ambient occlusion parameters to uniform buffer
 	FAmbientOcclusionData AmbientOcclusionData;
 	AmbientOcclusionData.SamplesPerPixel = GRayTracingAmbientOcclusionSamplesPerPixel;
@@ -180,7 +169,6 @@ void FDeferredShadingSceneRenderer::RenderRayTracingAmbientOcclusion(
 			SceneTexturesUniformBuffer,
 			AmbientOcclusionUniformBuffer,
 			AmbientOcclusionRT->GetRenderTargetItem().UAV,
-			RayDistanceRT->GetRenderTargetItem().UAV,
 			ViewSize.X, ViewSize.Y
 		);
 	}
@@ -189,9 +177,6 @@ void FDeferredShadingSceneRenderer::RenderRayTracingAmbientOcclusion(
 	FComputeFenceRHIRef Fence = RHICmdList.CreateComputeFence(TEXT("RayTracingAmbientOcclusion"));
 	RHICmdList.TransitionResource(EResourceTransitionAccess::ERWBarrier, EResourceTransitionPipeline::EComputeToGfx, AmbientOcclusionRT->GetRenderTargetItem().UAV, Fence);
 	GVisualizeTexture.SetCheckPoint(RHICmdList, AmbientOcclusionRT);
-
-	RHICmdList.TransitionResource(EResourceTransitionAccess::ERWBarrier, EResourceTransitionPipeline::EComputeToGfx, RayDistanceRT->GetRenderTargetItem().UAV);
-	GVisualizeTexture.SetCheckPoint(RHICmdList, RayDistanceRT);
 }
 
 class FCompositeAmbientOcclusionPS : public FGlobalShader
diff --git a/Engine/Source/Runtime/Renderer/Private/SceneRendering.h b/Engine/Source/Runtime/Renderer/Private/SceneRendering.h
--- a/Engine/Source/Runtime/Renderer/Private/SceneRendering.h
+++ b/Engine/Source/Runtime/Renderer/Private/SceneRendering.h
@@ -737,6 +737,9 @@ struct FPreviousViewInfo
 
 	// History for the reflections
 	FScreenSpaceFilteringHistory ReflectionsHistory;
+	
+	// History for the ambient occlusion
+	FScreenSpaceFilteringHistory AmbientOcclusionHistory;
 
 	// History for shadow denoising.
 	TMap<const FLightSceneInfo*, FScreenSpaceFilteringHistory> ShadowHistories;
@@ -753,6 +756,7 @@ struct FPreviousViewInfo
 		DOFPostGatherBackgroundHistory.SafeRelease();
 		CustomSSRInput.SafeRelease();
 		ReflectionsHistory.SafeRelease();
+		AmbientOcclusionHistory.SafeRelease();
 		ShadowHistories.Reset();
 	}
 };
diff --git a/Engine/Source/Runtime/Renderer/Private/ScreenSpaceDenoise.cpp b/Engine/Source/Runtime/Renderer/Private/ScreenSpaceDenoise.cpp
--- a/Engine/Source/Runtime/Renderer/Private/ScreenSpaceDenoise.cpp
+++ b/Engine/Source/Runtime/Renderer/Private/ScreenSpaceDenoise.cpp
@@ -40,7 +40,7 @@ static TAutoConsoleVariable<int32> CVarReflectionReconstructionSampleCount(
 
 static TAutoConsoleVariable<int32> CVarReflectionTemporalAccumulation(
 	TEXT("r.Reflections.Denoiser.TemporalAccumulation"), 1,
-	TEXT(""),
+	TEXT("Accumulates the samples over multiple frames."),
 	ECVF_RenderThreadSafe);
 
 static TAutoConsoleVariable<int32> CVarReflectionHistoryConvolutionSampleCount(
@@ -48,6 +48,21 @@ static TAutoConsoleVariable<int32> CVarReflectionHistoryConvolutionSampleCount(
 	TEXT("Number of samples to use for history post filter (default = 1)."),
 	ECVF_RenderThreadSafe);
 
+static TAutoConsoleVariable<int32> CVarAOReconstructionSampleCount(
+	TEXT("r.AmbientOcclusion.Denoiser.ReconstructionSamples"), 16,
+	TEXT("Maximum number of samples for the reconstruction pass (default = 16)."),
+	ECVF_RenderThreadSafe);
+
+static TAutoConsoleVariable<int32> CVarAOTemporalAccumulation(
+	TEXT("r.AmbientOcclusion.Denoiser.TemporalAccumulation"), 1,
+	TEXT("Accumulates the samples over multiple frames."),
+	ECVF_RenderThreadSafe);
+
+static TAutoConsoleVariable<int32> CVarAOHistoryConvolutionSampleCount(
+	TEXT("r.AmbientOcclusion.Denoiser.HistoryConvolutionSampleCount"), 1,
+	TEXT("Number of samples to use for history post filter (default = 1)."),
+	ECVF_RenderThreadSafe);
+
 
 /** The maximum number of mip level supported in the denoiser. */
 static const int32 kMaxMipLevel = 4;
@@ -80,6 +95,7 @@ enum class ESignalProcessing
 {
 	Penumbra,
 	Reflections,
+	AmbientOcclusion,
 
 	MAX,
 };
@@ -112,6 +128,12 @@ BEGIN_SHADER_PARAMETER_STRUCT(FSSDCommonParameters, )
 	SHADER_PARAMETER_STRUCT_REF(FViewUniformShaderParameters, ViewUniformBuffer)
 END_SHADER_PARAMETER_STRUCT()
 
+/** Shader parameter structure use to bind all signal generically. */
+BEGIN_SHADER_PARAMETER_STRUCT(FSSDSignalTextures, )
+	SHADER_PARAMETER_RDG_TEXTURE(Texture2D, Texture0)
+	SHADER_PARAMETER_RDG_TEXTURE(Texture2D, Texture1)
+END_SHADER_PARAMETER_STRUCT()
+
 /** Shader parameter structure to have all information to spatial filtering. */
 BEGIN_SHADER_PARAMETER_STRUCT(FSSDConvolutionMetaData, )
 	SHADER_PARAMETER_STRUCT(FLightShaderParameters, Light)
@@ -212,7 +234,7 @@ class FSSDSpatialAccumulationCS : public FScreenSpaceDenoisingShader
 		FPermutationDomain PermutationVector(Parameters.PermutationId);
 
 		// Only reflections have  reconstruction spatial accumulation.
-		if (PermutationVector.Get<FSignalProcessingDim>() != ESignalProcessing::Reflections && 
+		if (PermutationVector.Get<FSignalProcessingDim>() == ESignalProcessing::Penumbra && 
 			PermutationVector.Get<FStageDim>() == EStage::ReConstruction)
 		{
 			return false;
@@ -259,8 +281,8 @@ class FSSDTemporalAccumulationCS : public FScreenSpaceDenoisingShader
 	{
 		FPermutationDomain PermutationVector(Parameters.PermutationId);
 
-		// Reflections are only processed at mip 0.
-		if (PermutationVector.Get<FSignalProcessingDim>() == ESignalProcessing::Reflections && 
+		// Shadows are only implemented at multiple MIP.
+		if (PermutationVector.Get<FSignalProcessingDim>() != ESignalProcessing::Penumbra && 
 			!PermutationVector.Get<FIsMip0Dim>())
 		{
 			return false;
@@ -640,24 +662,60 @@ static void DenoiseShadowPenumbra(
 } // DenoiseShadowPenumbra()
 
 
-IScreenSpaceDenoiser::FReflectionsOutputs DenoiseReflections(
+/** Generic settings to denoise signal at constant pixel density across the viewport. */
+struct FSSDConstantPixelDensitySettings
+{
+	ESignalProcessing SignalProcessing;
+	float InputResolutionFraction = 1.0f;
+	int32 ReconstructionSamples = 1;
+	bool bUseTemporalAccumulation = false;
+	int32 HistoryConvolutionSampleCount = 1;
+};
+
+
+/** Denoises a signal at constant pixel density across the viewport. */
+static void DenoiseSignalAtConstantPixelDensity(
 	FRDGBuilder& GraphBuilder,
 	const FViewInfo& View,
 	const FSceneViewFamilyBlackboard& SceneBlackboard,
-	const IScreenSpaceDenoiser::FReflectionsInputs& ReflectionInputs,
-	const IScreenSpaceDenoiser::FReflectionsRayTracingConfig RayTracingConfig)
+	const FSSDSignalTextures& InputSignal,
+	FSSDConstantPixelDensitySettings Settings,
+	FSSDSignalTextures* OutputSignal)
 {
-	ensure(RayTracingConfig.ResolutionFraction == 1.0f || RayTracingConfig.ResolutionFraction == 0.5f);
+	ensure(Settings.InputResolutionFraction == 1.0f || Settings.InputResolutionFraction == 0.5f);
 
 	const FIntPoint DenoiseResolution = View.ViewRect.Size();
 	
 	FSceneRenderTargets& SceneContext = FSceneRenderTargets::Get(GraphBuilder.RHICmdList);
 
 	// Descriptor to allocate internal denoising buffer.
-	FRDGTextureDesc SignalProcessingDesc = SceneContext.GetSceneColor()->GetDesc();
-	SignalProcessingDesc.Format = PF_FloatRGBA;
-	SignalProcessingDesc.Flags &= ~(TexCreate_FastVRAM | TexCreate_Transient);
-	SignalProcessingDesc.NumMips = 1;
+	int32 SignalTextureCount = 1;
+	FRDGTextureDesc SignalProcessingDesc[2];
+	FRDGTextureDesc DebugDesc;
+	{
+		SignalProcessingDesc[0] = SceneContext.GetSceneColor()->GetDesc();
+		SignalProcessingDesc[0].Flags &= ~(TexCreate_FastVRAM | TexCreate_Transient);
+		SignalProcessingDesc[0].NumMips = 1;
+		SignalProcessingDesc[1] = SignalProcessingDesc[0];
+
+		DebugDesc = SignalProcessingDesc[0];
+		DebugDesc.Format = PF_FloatRGBA;
+
+		if (Settings.SignalProcessing == ESignalProcessing::Reflections)
+		{
+			SignalProcessingDesc[0].Format = PF_FloatRGBA;
+			SignalProcessingDesc[1].Format = PF_R16F;
+			SignalTextureCount = 2;
+		}
+		else if (Settings.SignalProcessing == ESignalProcessing::AmbientOcclusion)
+		{
+			SignalProcessingDesc[0].Format = PF_G16R16F;
+		}
+		else
+		{
+			check(0);
+		}
+	}
 
 	// Setup common shader parameters.
 	FSSDCommonParameters CommonParameters;
@@ -668,32 +726,43 @@ IScreenSpaceDenoiser::FReflectionsOutputs DenoiseReflections(
 	}
 
 	FScreenSpaceFilteringHistory PrevFrameHistory;
-	if (!View.bCameraCut)
+	if (View.bCameraCut)
+	{
+	}
+	else if (Settings.SignalProcessing == ESignalProcessing::Reflections)
 	{
 		PrevFrameHistory = View.PrevViewInfo.ReflectionsHistory;
 	}
+	else if (Settings.SignalProcessing == ESignalProcessing::AmbientOcclusion)
+	{
+		PrevFrameHistory = View.PrevViewInfo.AmbientOcclusionHistory;
+	}
+	else
+	{
+		check(0);
+	}
 
 	FRDGTextureRef SignalHistory0;
 	FRDGTextureRef SignalHistory1;
 	
 	// Spatial reconstruction with multiple important sampling to be more precise in the history rejection.
 	{
-		FRDGTextureRef SignalOutput0 = GraphBuilder.CreateTexture(SignalProcessingDesc, TEXT("SSDReflectionsReconstruction0"));
-		FRDGTextureRef SignalOutput1 = GraphBuilder.CreateTexture(SignalProcessingDesc, TEXT("SSDReflectionsReconstruction1"));
+		FRDGTextureRef SignalOutput0 = GraphBuilder.CreateTexture(SignalProcessingDesc[0], TEXT("SSDReflectionsReconstruction0"));
+		FRDGTextureRef SignalOutput1 = GraphBuilder.CreateTexture(SignalProcessingDesc[1], TEXT("SSDReflectionsReconstruction1"));
 
 		FSSDSpatialAccumulationCS::FParameters* PassParameters = GraphBuilder.AllocParameters<FSSDSpatialAccumulationCS::FParameters>();
-		PassParameters->MaxSampleCount = FMath::Clamp(CVarReflectionReconstructionSampleCount.GetValueOnRenderThread(), 1, kStackowiakMaxSampleCountPerSet);
-		PassParameters->UpscaleFactor = int32(1.0f / RayTracingConfig.ResolutionFraction);
+		PassParameters->MaxSampleCount = FMath::Clamp(Settings.ReconstructionSamples, 1, kStackowiakMaxSampleCountPerSet);
+		PassParameters->UpscaleFactor = int32(1.0f / Settings.InputResolutionFraction);
 		PassParameters->CommonParameters = CommonParameters;
-		PassParameters->SignalInput0 = ReflectionInputs.Color;
-		PassParameters->SignalInput1 = ReflectionInputs.RayHitDistance;
+		PassParameters->SignalInput0 = InputSignal.Texture0;
+		PassParameters->SignalInput1 = InputSignal.Texture1;
 		PassParameters->SignalOutput0 = GraphBuilder.CreateUAV(SignalOutput0);
 		PassParameters->SignalOutput1 = GraphBuilder.CreateUAV(SignalOutput1);
 		
-		PassParameters->DebugOutput = GraphBuilder.CreateUAV(GraphBuilder.CreateTexture(SignalProcessingDesc, TEXT("SSDDebugReflectionReconstruction")));
+		PassParameters->DebugOutput = GraphBuilder.CreateUAV(GraphBuilder.CreateTexture(DebugDesc, TEXT("SSDDebugReflectionReconstruction")));
 
 		FSSDSpatialAccumulationCS::FPermutationDomain PermutationVector;
-		PermutationVector.Set<FSignalProcessingDim>(ESignalProcessing::Reflections);
+		PermutationVector.Set<FSignalProcessingDim>(Settings.SignalProcessing);
 		PermutationVector.Set<FSSDSpatialAccumulationCS::FStageDim>(FSSDSpatialAccumulationCS::EStage::ReConstruction);
 		PermutationVector.Set<FSSDSpatialAccumulationCS::FUpscaleDim>(PassParameters->UpscaleFactor != 1);
 
@@ -711,16 +780,16 @@ IScreenSpaceDenoiser::FReflectionsOutputs DenoiseReflections(
 	}
 
 	// Temporal pass.
-	if (View.ViewState && CVarReflectionTemporalAccumulation.GetValueOnRenderThread())
+	if (View.ViewState && Settings.bUseTemporalAccumulation)
 	{
-		FRDGTextureRef SignalOutput0 = GraphBuilder.CreateTexture(SignalProcessingDesc, TEXT("SSDReflectionsHistory0"));
-		FRDGTextureRef SignalOutput1 = GraphBuilder.CreateTexture(SignalProcessingDesc, TEXT("SSDReflectionsHistory1"));
+		FRDGTextureRef SignalOutput0 = GraphBuilder.CreateTexture(SignalProcessingDesc[0], TEXT("SSDReflectionsHistory0"));
+		FRDGTextureRef SignalOutput1 = GraphBuilder.CreateTexture(SignalProcessingDesc[1], TEXT("SSDReflectionsHistory1"));
 
 		const int32 MipLevel = 0;
 
 		FSSDTemporalAccumulationCS::FPermutationDomain PermutationVector;
 		PermutationVector.Set<FSSDTemporalAccumulationCS::FIsMip0Dim>(MipLevel == 0);
-		PermutationVector.Set<FSignalProcessingDim>(ESignalProcessing::Reflections);
+		PermutationVector.Set<FSignalProcessingDim>(Settings.SignalProcessing);
 
 		TShaderMapRef<FSSDTemporalAccumulationCS> ComputeShader(View.ShaderMap, PermutationVector);
 
@@ -742,7 +811,7 @@ IScreenSpaceDenoiser::FReflectionsOutputs DenoiseReflections(
 		PassParameters->SignalHistoryOutput0 = GraphBuilder.CreateUAV(FRDGTextureUAVDesc(SignalOutput0, MipLevel));
 		PassParameters->SignalHistoryOutput1 = GraphBuilder.CreateUAV(FRDGTextureUAVDesc(SignalOutput1, MipLevel));
 		
-		PassParameters->DebugOutput = GraphBuilder.CreateUAV(GraphBuilder.CreateTexture(SignalProcessingDesc, TEXT("SSDDebugReflectionTemporalAccumulation")));
+		PassParameters->DebugOutput = GraphBuilder.CreateUAV(GraphBuilder.CreateTexture(DebugDesc, TEXT("SSDDebugReflectionTemporalAccumulation")));
 
 		FComputeShaderUtils::AddPass(
 			GraphBuilder,
@@ -756,11 +825,11 @@ IScreenSpaceDenoiser::FReflectionsOutputs DenoiseReflections(
 	}
 	
 	// Spatial filter, to converge history faster.
-	int32 MaxPostFilterSampleCount = CVarReflectionHistoryConvolutionSampleCount.GetValueOnRenderThread();
+	int32 MaxPostFilterSampleCount = FMath::Clamp(Settings.HistoryConvolutionSampleCount, 1, 16);
 	if (MaxPostFilterSampleCount > 1)
 	{
-		FRDGTextureRef SignalOutput0 = GraphBuilder.CreateTexture(SignalProcessingDesc, TEXT("SSDReflectionsHistory0"));
-		FRDGTextureRef SignalOutput1 = GraphBuilder.CreateTexture(SignalProcessingDesc, TEXT("SSDReflectionsHistory1"));
+		FRDGTextureRef SignalOutput0 = GraphBuilder.CreateTexture(SignalProcessingDesc[0], TEXT("SSDReflectionsHistory0"));
+		FRDGTextureRef SignalOutput1 = GraphBuilder.CreateTexture(SignalProcessingDesc[1], TEXT("SSDReflectionsHistory1"));
 
 		FSSDSpatialAccumulationCS::FParameters* PassParameters = GraphBuilder.AllocParameters<FSSDSpatialAccumulationCS::FParameters>();
 		PassParameters->MaxSampleCount = FMath::Clamp(MaxPostFilterSampleCount, 1, kStackowiakMaxSampleCountPerSet);
@@ -771,10 +840,10 @@ IScreenSpaceDenoiser::FReflectionsOutputs DenoiseReflections(
 		PassParameters->SignalOutput1 = GraphBuilder.CreateUAV(SignalOutput1);
 
 		FSSDSpatialAccumulationCS::FPermutationDomain PermutationVector;
-		PermutationVector.Set<FSignalProcessingDim>(ESignalProcessing::Reflections);
+		PermutationVector.Set<FSignalProcessingDim>(Settings.SignalProcessing);
 		PermutationVector.Set<FSSDSpatialAccumulationCS::FStageDim>(FSSDSpatialAccumulationCS::EStage::PostFiltering);
 		
-		PassParameters->DebugOutput = GraphBuilder.CreateUAV(GraphBuilder.CreateTexture(SignalProcessingDesc, TEXT("SSDDebugReflectionPostfilter")));
+		PassParameters->DebugOutput = GraphBuilder.CreateUAV(GraphBuilder.CreateTexture(DebugDesc, TEXT("SSDDebugReflectionPostfilter")));
 
 		TShaderMapRef<FSSDSpatialAccumulationCS> ComputeShader(View.ShaderMap, PermutationVector);
 		FComputeShaderUtils::AddPass(
@@ -793,18 +862,32 @@ IScreenSpaceDenoiser::FReflectionsOutputs DenoiseReflections(
 		// Keep depth buffer and GBuffer around for next frame.
 		{
 			GraphBuilder.QueueTextureExtraction(SceneBlackboard.SceneDepthBuffer, &View.ViewState->PendingPrevFrameViewInfo.DepthBuffer);
+
+			// Requires the normal that are in GBuffer A.
 			GraphBuilder.QueueTextureExtraction(SceneBlackboard.SceneGBufferA, &View.ViewState->PendingPrevFrameViewInfo.GBufferA);
-			GraphBuilder.QueueTextureExtraction(SceneBlackboard.SceneGBufferB, &View.ViewState->PendingPrevFrameViewInfo.GBufferB);
+
+			// Reflections requires the roughness that is in GBuffer B.
+			if (Settings.SignalProcessing == ESignalProcessing::Reflections)
+				GraphBuilder.QueueTextureExtraction(SceneBlackboard.SceneGBufferB, &View.ViewState->PendingPrevFrameViewInfo.GBufferB);
 		}
 
 		// Saves history.
 		{
-			FScreenSpaceFilteringHistory& NewHistory = View.ViewState->PendingPrevFrameViewInfo.ReflectionsHistory;
-			GraphBuilder.QueueTextureExtraction(SignalHistory0, &NewHistory.RT[0]);
-			GraphBuilder.QueueTextureExtraction(SignalHistory1, &NewHistory.RT[1]);
+			FScreenSpaceFilteringHistory* NewHistory = nullptr;
+			if (Settings.SignalProcessing == ESignalProcessing::Reflections)
+				NewHistory = &View.ViewState->PendingPrevFrameViewInfo.ReflectionsHistory;
+			else if (Settings.SignalProcessing == ESignalProcessing::AmbientOcclusion)
+				NewHistory = &View.ViewState->PendingPrevFrameViewInfo.AmbientOcclusionHistory;
+			else
+				check(0);
+
+			GraphBuilder.QueueTextureExtraction(SignalHistory0, &NewHistory->RT[0]);
+
+			if (SignalTextureCount >= 2)
+				GraphBuilder.QueueTextureExtraction(SignalHistory1, &NewHistory->RT[1]);
 		}
 	}
-	else
+	else if (SignalTextureCount >= 2)
 	{
 		// The SignalHistory1 is always generated for temporal history, but will endup useless if there is no view state,
 		// in witch case we do not extract any textures. Don't support a shader permutation that does not produce it, because
@@ -812,10 +895,9 @@ IScreenSpaceDenoiser::FReflectionsOutputs DenoiseReflections(
 		GraphBuilder.RemoveUnusedTextureWarning(SignalHistory1);
 	}
 
-	IScreenSpaceDenoiser::FReflectionsOutputs Outputs;
-	Outputs.Color = SignalHistory0; // TODO.
-	return Outputs;
-} // DenoiseReflections()
+	OutputSignal->Texture0 = SignalHistory0;
+	OutputSignal->Texture1 = SignalHistory1;
+} // DenoiseSignalAtConstantPixelDensity()
 
 
 /** The implementation of the default denoiser of the renderer. */
@@ -824,7 +906,7 @@ class FDefaultScreenSpaceDenoiser : public IScreenSpaceDenoiser
 public:
 	const TCHAR* GetDebugName() const override
 	{
-		return TEXT("SSD");
+		return TEXT("ScreenSpaceDenoiser");
 	}
 
 	FShadowRayTracingConfig GetShadowRayTracingConfig(
@@ -886,7 +968,54 @@ public:
 		const FReflectionsInputs& ReflectionInputs,
 		const FReflectionsRayTracingConfig RayTracingConfig) const override
 	{
-		return ::DenoiseReflections(GraphBuilder, View, SceneBlackboard, ReflectionInputs, RayTracingConfig);
+		FSSDSignalTextures InputSignal;
+		InputSignal.Texture0 = ReflectionInputs.Color;
+		InputSignal.Texture1 = ReflectionInputs.RayHitDistance;
+
+		FSSDConstantPixelDensitySettings Settings;
+		Settings.SignalProcessing = ESignalProcessing::Reflections;
+		Settings.InputResolutionFraction = RayTracingConfig.ResolutionFraction;
+		Settings.ReconstructionSamples = CVarReflectionReconstructionSampleCount.GetValueOnRenderThread();
+		Settings.bUseTemporalAccumulation = CVarReflectionTemporalAccumulation.GetValueOnRenderThread() != 0;
+		Settings.HistoryConvolutionSampleCount = CVarReflectionHistoryConvolutionSampleCount.GetValueOnRenderThread();
+
+		FSSDSignalTextures SignalOutput;
+		DenoiseSignalAtConstantPixelDensity(
+			GraphBuilder, View, SceneBlackboard,
+			InputSignal, Settings,
+			&SignalOutput);
+
+		FReflectionsOutputs ReflectionsOutput;
+		ReflectionsOutput.Color = SignalOutput.Texture0;
+		return ReflectionsOutput;
+	}
+	
+	FAmbientOcclusionOutputs DenoiseAmbientOcclusion(
+		FRDGBuilder& GraphBuilder,
+		const FViewInfo& View,
+		const FSceneViewFamilyBlackboard& SceneBlackboard,
+		const FAmbientOcclusionInputs& ReflectionInputs,
+		const FAmbientOcclusionRayTracingConfig RayTracingConfig) const override
+	{
+		FSSDSignalTextures InputSignal;
+		InputSignal.Texture0 = ReflectionInputs.MaskAndRayHitDistance;
+		
+		FSSDConstantPixelDensitySettings Settings;
+		Settings.SignalProcessing = ESignalProcessing::AmbientOcclusion;
+		Settings.InputResolutionFraction = RayTracingConfig.ResolutionFraction;
+		Settings.ReconstructionSamples = CVarAOReconstructionSampleCount.GetValueOnRenderThread();
+		Settings.bUseTemporalAccumulation = CVarAOTemporalAccumulation.GetValueOnRenderThread() != 0;
+		Settings.HistoryConvolutionSampleCount = CVarAOHistoryConvolutionSampleCount.GetValueOnRenderThread();
+
+		FSSDSignalTextures SignalOutput;
+		DenoiseSignalAtConstantPixelDensity(
+			GraphBuilder, View, SceneBlackboard,
+			InputSignal, Settings,
+			&SignalOutput);
+
+		FAmbientOcclusionOutputs AmbientOcclusionOutput;
+		AmbientOcclusionOutput.AmbientOcclusionMask = SignalOutput.Texture0;
+		return AmbientOcclusionOutput;
 	}
 
 }; // class FDefaultScreenSpaceDenoiser
diff --git a/Engine/Source/Runtime/Renderer/Private/ScreenSpaceDenoise.h b/Engine/Source/Runtime/Renderer/Private/ScreenSpaceDenoise.h
--- a/Engine/Source/Runtime/Renderer/Private/ScreenSpaceDenoise.h
+++ b/Engine/Source/Runtime/Renderer/Private/ScreenSpaceDenoise.h
@@ -46,6 +46,16 @@ public:
 		// Resolution fraction the ray tracing is being traced at.
 		float ResolutionFraction;
 	};
+	
+	/** The configuration of the reflection ray tracing. */
+	struct FAmbientOcclusionRayTracingConfig
+	{
+		// Resolution fraction the ray tracing is being traced at.
+		float ResolutionFraction = 1.0f;
+
+		// Number of rays per pixels.
+		float RayCountPerPixel = 1.0f;
+	};
 
 
 	virtual ~IScreenSpaceDenoiser() {};
@@ -66,13 +76,13 @@ public:
 		const FSceneViewFamilyBlackboard& SceneBlackboard,
 		const FShadowPenumbraInputs& ShadowInputs) const = 0;
 	
-	/** All the inputs of the shadow denoiser. */
+	/** All the inputs of the reflection denoiser. */
 	BEGIN_SHADER_PARAMETER_STRUCT(FReflectionsInputs, )
 		SHADER_PARAMETER_RDG_TEXTURE(Texture2D, Color)
 		SHADER_PARAMETER_RDG_TEXTURE(Texture2D, RayHitDistance)
 	END_SHADER_PARAMETER_STRUCT()
 
-	/** All the outputs of the shadow denoiser may generate. */
+	/** All the outputs of the reflection denoiser may generate. */
 	BEGIN_SHADER_PARAMETER_STRUCT(FReflectionsOutputs, )
 		SHADER_PARAMETER_RDG_TEXTURE(Texture2D, Color)
 	END_SHADER_PARAMETER_STRUCT()
@@ -84,6 +94,28 @@ public:
 		const FSceneViewFamilyBlackboard& SceneBlackboard,
 		const FReflectionsInputs& ReflectionInputs,
 		const FReflectionsRayTracingConfig RayTracingConfig) const = 0;
+	
+	/** All the inputs of the AO denoisers. */
+	BEGIN_SHADER_PARAMETER_STRUCT(FAmbientOcclusionInputs, )
+		// G16R16F texture:
+		//  R: mask as [0; 1].
+		//  B: closest hit distance of the ray.
+		SHADER_PARAMETER_RDG_TEXTURE(Texture2D, MaskAndRayHitDistance)
+	END_SHADER_PARAMETER_STRUCT()
+
+	/** All the outputs of the AO denoiser may generate. */
+	BEGIN_SHADER_PARAMETER_STRUCT(FAmbientOcclusionOutputs, )
+		// Ambient occlusion mask stored in the red channel as [0; 1].
+		SHADER_PARAMETER_RDG_TEXTURE(Texture2D, AmbientOcclusionMask)
+	END_SHADER_PARAMETER_STRUCT()
+		
+	/** Entry point to denoise reflections. */
+	virtual FAmbientOcclusionOutputs DenoiseAmbientOcclusion(
+		FRDGBuilder& GraphBuilder,
+		const FViewInfo& View,
+		const FSceneViewFamilyBlackboard& SceneBlackboard,
+		const FAmbientOcclusionInputs& ReflectionInputs,
+		const FAmbientOcclusionRayTracingConfig RayTracingConfig) const = 0;
 
 	/** Returns the interface of the default denoiser of the renderer. */
 	static const IScreenSpaceDenoiser* GetDefaultDenoiser();
