Fixes a bug where shadow denoising histories were lost when changing the light components.

diff --git a/Engine/Source/Runtime/Renderer/Private/SceneRendering.h b/Engine/Source/Runtime/Renderer/Private/SceneRendering.h
--- a/Engine/Source/Runtime/Renderer/Private/SceneRendering.h
+++ b/Engine/Source/Runtime/Renderer/Private/SceneRendering.h
@@ -742,7 +742,7 @@ struct FPreviousViewInfo
 	FScreenSpaceFilteringHistory AmbientOcclusionHistory;
 
 	// History for shadow denoising.
-	TMap<const FLightSceneInfo*, FScreenSpaceFilteringHistory> ShadowHistories;
+	TMap<const ULightComponent*, FScreenSpaceFilteringHistory> ShadowHistories;
 
 
 	void SafeRelease()
diff --git a/Engine/Source/Runtime/Renderer/Private/ScreenSpaceDenoise.cpp b/Engine/Source/Runtime/Renderer/Private/ScreenSpaceDenoise.cpp
--- a/Engine/Source/Runtime/Renderer/Private/ScreenSpaceDenoise.cpp
+++ b/Engine/Source/Runtime/Renderer/Private/ScreenSpaceDenoise.cpp
@@ -555,7 +555,7 @@ static void DenoiseShadowPenumbra(
 		FScreenSpaceFilteringHistory PrevFrameHistory;
 		if (!View.bCameraCut)
 		{
-			PrevFrameHistory = View.PrevViewInfo.ShadowHistories.FindRef(&LightSceneInfo);
+			PrevFrameHistory = View.PrevViewInfo.ShadowHistories.FindRef(LightSceneInfo.Proxy->GetLightComponent());
 		}
 
 		// Process each mip level independently.
@@ -610,7 +610,7 @@ static void DenoiseShadowPenumbra(
 
 		// Saves shadow history.
 		{
-			FScreenSpaceFilteringHistory& NewShadowHistory = View.ViewState->PendingPrevFrameViewInfo.ShadowHistories.FindOrAdd(&LightSceneInfo);
+			FScreenSpaceFilteringHistory& NewShadowHistory = View.ViewState->PendingPrevFrameViewInfo.ShadowHistories.FindOrAdd(LightSceneInfo.Proxy->GetLightComponent());
 			GraphBuilder.QueueTextureExtraction(SignalHistory0, &NewShadowHistory.RT[0]);
 
 			GraphBuilder.QueueTextureExtraction(CommonParameters.TileClassificationTexture, &NewShadowHistory.TileClassification);
