Fix static analysis warning. Add SkyLight existence check before building CDF data.

diff --git a/Engine/Source/Runtime/Renderer/Private/PathTracing.cpp b/Engine/Source/Runtime/Renderer/Private/PathTracing.cpp
--- a/Engine/Source/Runtime/Renderer/Private/PathTracing.cpp
+++ b/Engine/Source/Runtime/Renderer/Private/PathTracing.cpp
@@ -431,18 +431,21 @@ void FDeferredShadingSceneRenderer::RenderPathTracing(FRHICommandListImmediate&
 	static int32 SPPCount = 0;
 
 	// Conditionally rebuild sky light CDFs
-	FRWBuffer& SkyLightRowCdf = Scene->SkyLight->RowCdf;
-	FRWBuffer& SkyLightColumnCdf = Scene->SkyLight->ColumnCdf;
-	FRWBuffer& SkyLightCubeFaceCdf = Scene->SkyLight->CubeFaceCdf;
+	FRWBuffer SkyLightRowCdf;
+	FRWBuffer SkyLightColumnCdf;
+	FRWBuffer SkyLightCubeFaceCdf;
 	if (Scene->SkyLight)
 	{
 		if (Scene->SkyLight->ShouldRebuildCdf())
 		{
 			SCOPED_GPU_STAT(RHICmdList, Stat_GPU_PathTracingBuildSkyLightCDF);
 
-			BuildSkyLightCdf(RHICmdList, View, *Scene->SkyLight->ProcessedTexture, SkyLightRowCdf, SkyLightColumnCdf, SkyLightCubeFaceCdf);
+			BuildSkyLightCdf(RHICmdList, View, *Scene->SkyLight->ProcessedTexture, Scene->SkyLight->RowCdf, Scene->SkyLight->ColumnCdf, Scene->SkyLight->CubeFaceCdf);
 			Scene->SkyLight->IsDirtyImportanceSamplingData = false;
 		}
+		SkyLightRowCdf = Scene->SkyLight->RowCdf;
+		SkyLightColumnCdf = Scene->SkyLight->ColumnCdf;
+		SkyLightCubeFaceCdf = Scene->SkyLight->CubeFaceCdf;
 
 		FIntVector Dimensions = FIntVector(Scene->SkyLight->ProcessedTexture->GetSizeX(), Scene->SkyLight->ProcessedTexture->GetSizeY(), 6);
 		VisualizeSkyLightCdf(RHICmdList, View, Dimensions, SkyLightRowCdf, SkyLightColumnCdf, SkyLightCubeFaceCdf);
