Replaced unsafe primitive id buffer update from memstack with a separate allocated memory.

diff --git a/Engine/Source/Runtime/Renderer/Private/MeshDrawCommands.cpp b/Engine/Source/Runtime/Renderer/Private/MeshDrawCommands.cpp
--- a/Engine/Source/Runtime/Renderer/Private/MeshDrawCommands.cpp
+++ b/Engine/Source/Runtime/Renderer/Private/MeshDrawCommands.cpp
@@ -850,7 +850,7 @@ void FParallelMeshDrawCommandPass::DispatchPassSetup(
 	FMeshCommandOneFrameArray* InOutMobileBasePassCSMMeshDrawCommands
 )
 {
-	check(TaskEventRefs.Num() == 0 && MeshPassProcessor != nullptr);
+	check(TaskEventRefs.Num() == 0 && MeshPassProcessor != nullptr && TaskContext.PrimitiveIdBufferData == nullptr);
 
 	MaxNumDraws = InOutMeshDrawCommands.Num() + NumDynamicMeshElements + NumDynamicMeshCommandBuildRequestElements;
 	if (MaxNumDraws <= 0)
@@ -905,8 +905,9 @@ void FParallelMeshDrawCommandPass::DispatchPassSetup(
 	}
 
 	// Preallocate resources on rendering thread based on MaxNumDraws.
+	bPrimitiveIdBufferDataOwnedByRHIThread = false;
 	TaskContext.PrimitiveIdBufferDataSize = TaskContext.MaxInstanceFactor * MaxNumDraws * sizeof(int32);
-	TaskContext.PrimitiveIdBufferData = FMemStack::Get().Alloc(TaskContext.PrimitiveIdBufferDataSize, 8);
+	TaskContext.PrimitiveIdBufferData = FMemory::Malloc(TaskContext.PrimitiveIdBufferDataSize);
 	PrimitiveIdVertexBufferRHI = GPrimitiveIdVertexBufferPool.Allocate(TaskContext.PrimitiveIdBufferDataSize);
 	TaskContext.MeshDrawCommands.Reserve(MaxNumDraws);
 	TaskContext.TempVisibleMeshDrawCommands.Reserve(MaxNumDraws);
@@ -944,11 +945,18 @@ void FParallelMeshDrawCommandPass::Empty()
 
 	DumpInstancingStats();
 
+	if (!bPrimitiveIdBufferDataOwnedByRHIThread)
+	{
+		FMemory::Free(TaskContext.PrimitiveIdBufferData);
+	}
+
+	bPrimitiveIdBufferDataOwnedByRHIThread = false;
 	MaxNumDraws = 0;
+	PassNameForStats.Empty();
+
 	TaskContext.MeshPassProcessor = nullptr;
 	TaskContext.MobileBasePassCSMMeshPassProcessor = nullptr;
 	TaskContext.DynamicMeshElements = nullptr;
-	PassNameForStats.Empty();
 	TaskContext.MeshDrawCommands.Empty();
 	TaskContext.MobileBasePassCSMMeshDrawCommands.Empty();
 	TaskContext.DynamicMeshCommandBuildRequests.Empty();
@@ -1060,11 +1068,9 @@ void FParallelMeshDrawCommandPass::DispatchDraw(FParallelCommandListSet* Paralle
 		return;
 	}
 
-
 	FVertexBufferRHIParamRef PrimitiveIdsBuffer = PrimitiveIdVertexBufferRHI;
 	const int32 BasePrimitiveIdsOffset = 0;
 
-
 	if (ParallelCommandListSet)
 	{
 		if (TaskContext.bUseGPUScene)
@@ -1077,6 +1083,7 @@ void FParallelMeshDrawCommandPass::DispatchDraw(FParallelCommandListSet* Paralle
 				TaskContext.PrimitiveIdBufferData,
 				TaskContext.PrimitiveIdBufferDataSize
 			);
+			bPrimitiveIdBufferDataOwnedByRHIThread = true;
 		}
 
 		const ENamedThreads::Type RenderThread = ENamedThreads::GetRenderThread();
diff --git a/Engine/Source/Runtime/Renderer/Private/MeshDrawCommands.h b/Engine/Source/Runtime/Renderer/Private/MeshDrawCommands.h
--- a/Engine/Source/Runtime/Renderer/Private/MeshDrawCommands.h
+++ b/Engine/Source/Runtime/Renderer/Private/MeshDrawCommands.h
@@ -61,6 +61,8 @@ public:
 		, DynamicMeshElements(nullptr)
 		, NumDynamicMeshElements(0)
 		, NumDynamicMeshCommandBuildRequestElements(0)
+		, PrimitiveIdBufferData(nullptr)
+		, PrimitiveIdBufferDataSize(0)
 		, PrimitiveBounds(nullptr)
 		, VisibleMeshDrawCommandsNum(0)
 		, NewPassVisibleMeshDrawCommandsNum(0)
@@ -92,7 +94,7 @@ public:
 	TArray<const FStaticMeshBatch*, SceneRenderingAllocator> MobileBasePassCSMDynamicMeshCommandBuildRequests;
 	FDynamicMeshDrawCommandStorage MeshDrawCommandStorage;
 
-	// Resources preallocated on rendering thread with SceneRenderingAllocator.
+	// Resources preallocated on rendering thread.
 	void* PrimitiveIdBufferData;
 	int32 PrimitiveIdBufferDataSize;
 	FMeshCommandOneFrameArray TempVisibleMeshDrawCommands;
@@ -122,7 +124,8 @@ class FParallelMeshDrawCommandPass
 {
 public:
 	FParallelMeshDrawCommandPass()
-		: MaxNumDraws(0)
+		: bPrimitiveIdBufferDataOwnedByRHIThread(false)
+		, MaxNumDraws(0)
 	{
 	}
 
@@ -163,6 +166,9 @@ private:
 	FGraphEventArray TaskEventRefs;
 	FString PassNameForStats;
 
+	// If TaskContext::PrimitiveIdBufferData will be released by RHI Thread.
+	mutable bool bPrimitiveIdBufferDataOwnedByRHIThread;
+
 	// Maximum number of draws for this pass. Used to prealocate resources on rendering thread. 
 	// Has a guarantee that if there won't be any draws, then MaxNumDraws = 0;
 	int32 MaxNumDraws;
