Added missing Free to FRHICommandUpdatePrimitiveIdBuffer.

diff --git a/Engine/Source/Runtime/Renderer/Private/MeshDrawCommands.cpp b/Engine/Source/Runtime/Renderer/Private/MeshDrawCommands.cpp
--- a/Engine/Source/Runtime/Renderer/Private/MeshDrawCommands.cpp
+++ b/Engine/Source/Runtime/Renderer/Private/MeshDrawCommands.cpp
@@ -1032,7 +1032,7 @@ struct FRHICommandUpdatePrimitiveIdBuffer : public FRHICommand<FRHICommandUpdate
 	FGraphEventArray Prereqs;
 
 	FVertexBufferRHIParamRef VertexBuffer;
-	const void* VertexBufferData;
+	void* VertexBufferData;
 	int32 VertexBufferDataSize;
 
 	virtual ~FRHICommandUpdatePrimitiveIdBuffer() {}
@@ -1040,7 +1040,7 @@ struct FRHICommandUpdatePrimitiveIdBuffer : public FRHICommand<FRHICommandUpdate
 	FORCEINLINE_DEBUGGABLE FRHICommandUpdatePrimitiveIdBuffer(
 		const FGraphEventArray& InPrereqs,
 		FVertexBufferRHIParamRef InVertexBuffer,
-		const void* InVertexBufferData,
+		void* InVertexBufferData,
 		int32 InVertexBufferDataSize)
 		: Prereqs(InPrereqs)
 		, VertexBuffer(InVertexBuffer)
@@ -1058,6 +1058,8 @@ struct FRHICommandUpdatePrimitiveIdBuffer : public FRHICommand<FRHICommandUpdate
 		void* RESTRICT Data = (void* RESTRICT)GDynamicRHI->RHILockVertexBuffer(VertexBuffer, 0, VertexBufferDataSize, RLM_WriteOnly);
 		FMemory::BigBlockMemcpy(Data, VertexBufferData, VertexBufferDataSize);
 		GDynamicRHI->RHIUnlockVertexBuffer(VertexBuffer);
+
+		FMemory::Free(VertexBufferData);
 	}
 };
 
