Fixed deadlock with mesh draw commands and RHIThread:
* FMeshDrawCommandPassSetupTask calls InitResourceFromPossiblyParallelRendering and waits on the Rendering Thread.
* RenderThread tries to process parallel creation tasks and waits on StallRHIThread.
* RHIThread waits inside FRHICommandUpdatePrimitiveIdBuffer for FMeshDrawCommandPassSetupTask.
Fixed by removing wait from FMeshDrawCommandPassSetupTask and instead using AddDispatchPrerequisite on the cmd list to setup proper chain dependency.

diff --git a/Engine/Source/Runtime/RHI/Private/PipelineStateCache.cpp b/Engine/Source/Runtime/RHI/Private/PipelineStateCache.cpp
--- a/Engine/Source/Runtime/RHI/Private/PipelineStateCache.cpp
+++ b/Engine/Source/Runtime/RHI/Private/PipelineStateCache.cpp
@@ -802,7 +802,7 @@ FComputePipelineState* PipelineStateCache::GetAndOrCreateComputePipelineState(FR
 		if (DoAsyncCompile)
 		{
 			OutCachedState->CompletionEvent = TGraphTask<FCompilePipelineStateTask>::CreateTask().ConstructAndDispatchWhenReady(OutCachedState, FGraphicsPipelineStateInitializer());
-			RHICmdList.QueueAsyncPipelineStateCompile(OutCachedState->CompletionEvent);
+			RHICmdList.AddDispatchPrerequisite(OutCachedState->CompletionEvent);
 		}
 		else
 		{
@@ -818,7 +818,7 @@ FComputePipelineState* PipelineStateCache::GetAndOrCreateComputePipelineState(FR
 			FGraphEventRef& CompletionEvent = OutCachedState->CompletionEvent;
 			if ( CompletionEvent.IsValid() && !CompletionEvent->IsComplete() )
 			{
-				RHICmdList.QueueAsyncPipelineStateCompile(CompletionEvent);
+				RHICmdList.AddDispatchPrerequisite(CompletionEvent);
 			}
 		}
 
@@ -1008,7 +1008,7 @@ FGraphicsPipelineState* PipelineStateCache::GetAndOrCreateGraphicsPipelineState(
 		if (DoAsyncCompile)
 		{
 			OutCachedState->CompletionEvent = TGraphTask<FCompilePipelineStateTask>::CreateTask().ConstructAndDispatchWhenReady(OutCachedState, *Initializer);
-			RHICmdList.QueueAsyncPipelineStateCompile(OutCachedState->CompletionEvent);
+			RHICmdList.AddDispatchPrerequisite(OutCachedState->CompletionEvent);
 		}
 		else
 		{
@@ -1029,7 +1029,7 @@ FGraphicsPipelineState* PipelineStateCache::GetAndOrCreateGraphicsPipelineState(
 			FGraphEventRef& CompletionEvent = OutCachedState->CompletionEvent;
 			if ( CompletionEvent.IsValid() && !CompletionEvent->IsComplete() )
 			{
-				RHICmdList.QueueAsyncPipelineStateCompile(CompletionEvent);
+				RHICmdList.AddDispatchPrerequisite(CompletionEvent);
 			}
 		}
 
diff --git a/Engine/Source/Runtime/RHI/Private/RHICommandList.cpp b/Engine/Source/Runtime/RHI/Private/RHICommandList.cpp
--- a/Engine/Source/Runtime/RHI/Private/RHICommandList.cpp
+++ b/Engine/Source/Runtime/RHI/Private/RHICommandList.cpp
@@ -573,6 +573,20 @@ void FRHICommandListExecutor::ExecuteInner(FRHICommandListBase& CmdList)
 			}
 		}
 	}
+	else
+	{
+		if (bIsInRenderingThread && CmdList.RTTasks.Num())
+		{
+			ENamedThreads::Type RenderThread_Local = ENamedThreads::GetRenderThread_Local();
+			if (FTaskGraphInterface::Get().IsThreadProcessingTasks(RenderThread_Local))
+			{
+				// this is a deadlock. RT tasks must be done by now or they won't be done. We could add a third queue...
+				UE_LOG(LogRHI, Fatal, TEXT("Deadlock in FRHICommandListExecutor::ExecuteInner (RTTasks)."));
+			}
+			FTaskGraphInterface::Get().WaitUntilTasksComplete(CmdList.RTTasks, RenderThread_Local);
+			CmdList.RTTasks.Reset();
+		}
+	}
 
 	ExecuteInner_DoExecute(CmdList);
 }
@@ -1404,11 +1418,11 @@ void FRHICommandListBase::QueueRenderThreadCommandListSubmit(FGraphEventRef& Ren
 	ALLOC_COMMAND(FRHICommandWaitForAndSubmitRTSubList)(RenderThreadCompletionEvent, CmdList);
 }
 
-void FRHICommandListBase::QueueAsyncPipelineStateCompile(FGraphEventRef& AsyncCompileCompletionEvent)
+void FRHICommandListBase::AddDispatchPrerequisite(const FGraphEventRef& Prereq)
 {
-	if (AsyncCompileCompletionEvent.GetReference())
+	if (Prereq.GetReference())
 	{
-		RTTasks.AddUnique(AsyncCompileCompletionEvent);
+		RTTasks.AddUnique(Prereq);
 	}
 }
 
diff --git a/Engine/Source/Runtime/RHI/Public/RHICommandList.h b/Engine/Source/Runtime/RHI/Public/RHICommandList.h
--- a/Engine/Source/Runtime/RHI/Public/RHICommandList.h
+++ b/Engine/Source/Runtime/RHI/Public/RHICommandList.h
@@ -401,8 +401,8 @@ public:
 	void QueueAsyncCommandListSubmit(FGraphEventRef& AnyThreadCompletionEvent, class FRHICommandList* CmdList);
 	void QueueParallelAsyncCommandListSubmit(FGraphEventRef* AnyThreadCompletionEvents, bool bIsPrepass, class FRHICommandList** CmdLists, int32* NumDrawsIfKnown, int32 Num, int32 MinDrawsPerTranslate, bool bSpewMerge);
 	void QueueRenderThreadCommandListSubmit(FGraphEventRef& RenderThreadCompletionEvent, class FRHICommandList* CmdList);
-	void QueueAsyncPipelineStateCompile(FGraphEventRef& AsyncCompileCompletionEvent);
 	void QueueCommandListSubmit(class FRHICommandList* CmdList);
+	void AddDispatchPrerequisite(const FGraphEventRef& Prereq);
 	void WaitForTasks(bool bKnownToBeComplete = false);
 	void WaitForDispatch();
 	void WaitForRHIThreadTasks();
diff --git a/Engine/Source/Runtime/Renderer/Private/MeshDrawCommands.cpp b/Engine/Source/Runtime/Renderer/Private/MeshDrawCommands.cpp
--- a/Engine/Source/Runtime/Renderer/Private/MeshDrawCommands.cpp
+++ b/Engine/Source/Runtime/Renderer/Private/MeshDrawCommands.cpp
@@ -850,7 +850,7 @@ void FParallelMeshDrawCommandPass::DispatchPassSetup(
 	FMeshCommandOneFrameArray* InOutMobileBasePassCSMMeshDrawCommands
 )
 {
-	check(TaskEventRefs.Num() == 0 && MeshPassProcessor != nullptr && TaskContext.PrimitiveIdBufferData == nullptr);
+	check(!TaskEventRef.IsValid() && MeshPassProcessor != nullptr && TaskContext.PrimitiveIdBufferData == nullptr);
 
 	MaxNumDraws = InOutMeshDrawCommands.Num() + NumDynamicMeshElements + NumDynamicMeshCommandBuildRequestElements;
 	if (MaxNumDraws <= 0)
@@ -918,8 +918,8 @@ void FParallelMeshDrawCommandPass::DispatchPassSetup(
 
 	if (bExecuteInParallel)
 	{
-		TaskEventRefs.Add(TGraphTask<FMeshDrawCommandPassSetupTask>::CreateTask(
-			nullptr, ENamedThreads::GetRenderThread()).ConstructAndDispatchWhenReady(TaskContext));
+		TaskEventRef = TGraphTask<FMeshDrawCommandPassSetupTask>::CreateTask(
+			nullptr, ENamedThreads::GetRenderThread()).ConstructAndDispatchWhenReady(TaskContext);
 	}
 	else
 	{
@@ -930,10 +930,10 @@ void FParallelMeshDrawCommandPass::DispatchPassSetup(
 
 void FParallelMeshDrawCommandPass::WaitForMeshPassSetupTask() const
 {
-	if (TaskEventRefs.Num() > 0)
+	if (TaskEventRef.IsValid())
 	{
 		// Need to wait on GetRenderThread_Local, as mesh pass setup task can wait on rendering thread inside InitResourceFromPossiblyParallelRendering().
-		FTaskGraphInterface::Get().WaitUntilTasksComplete(TaskEventRefs, ENamedThreads::GetRenderThread_Local());
+		FTaskGraphInterface::Get().WaitUntilTaskCompletes(TaskEventRef, ENamedThreads::GetRenderThread_Local());
 	}
 }
 
@@ -941,7 +941,7 @@ void FParallelMeshDrawCommandPass::Empty()
 {
 	// Need to wait in case if someone dispatched sort and draw merge task, but didn't draw it.
 	WaitForMeshPassSetupTask();
-	TaskEventRefs.Empty();
+	TaskEventRef = nullptr;
 
 	DumpInstancingStats();
 
@@ -1029,8 +1029,6 @@ public:
  */
 struct FRHICommandUpdatePrimitiveIdBuffer : public FRHICommand<FRHICommandUpdatePrimitiveIdBuffer>
 {
-	FGraphEventArray Prereqs;
-
 	FVertexBufferRHIParamRef VertexBuffer;
 	void* VertexBufferData;
 	int32 VertexBufferDataSize;
@@ -1038,12 +1036,10 @@ struct FRHICommandUpdatePrimitiveIdBuffer : public FRHICommand<FRHICommandUpdate
 	virtual ~FRHICommandUpdatePrimitiveIdBuffer() {}
 
 	FORCEINLINE_DEBUGGABLE FRHICommandUpdatePrimitiveIdBuffer(
-		const FGraphEventArray& InPrereqs,
 		FVertexBufferRHIParamRef InVertexBuffer,
 		void* InVertexBufferData,
 		int32 InVertexBufferDataSize)
-		: Prereqs(InPrereqs)
-		, VertexBuffer(InVertexBuffer)
+		: VertexBuffer(InVertexBuffer)
 		, VertexBufferData(InVertexBufferData)
 		, VertexBufferDataSize(InVertexBufferDataSize)
 	{
@@ -1051,9 +1047,6 @@ struct FRHICommandUpdatePrimitiveIdBuffer : public FRHICommand<FRHICommandUpdate
 
 	void Execute(FRHICommandListBase& CmdList)
 	{
-		// Need to wait on GetRenderThread_Local, as mesh pass setup task can wait on rendering thread inside InitResourceFromPossiblyParallelRendering().
-		FTaskGraphInterface::Get().WaitUntilTasksComplete(Prereqs, IsRunningRHIInSeparateThread() ? ENamedThreads::RHIThread : ENamedThreads::GetRenderThread_Local());
-
 		// Upload vertex buffer data.
 		void* RESTRICT Data = (void* RESTRICT)GDynamicRHI->RHILockVertexBuffer(VertexBuffer, 0, VertexBufferDataSize, RLM_WriteOnly);
 		FMemory::BigBlockMemcpy(Data, VertexBufferData, VertexBufferDataSize);
@@ -1077,10 +1070,15 @@ void FParallelMeshDrawCommandPass::DispatchDraw(FParallelCommandListSet* Paralle
 	{
 		if (TaskContext.bUseGPUScene)
 		{
-			// Queue a command on the RHI thread which will upload PrimitiveIdVertexBuffer after finishing FMeshDrawCommandPassSetupTaskContext.
+			// Queue a command on the RHI thread which will upload PrimitiveIdVertexBuffer after finishing FMeshDrawCommandPassSetupTask.
 			FRHICommandListImmediate &RHICommandList = GetImmediateCommandList_ForRenderCommand();
+
+			if (TaskEventRef.IsValid())
+			{
+				RHICommandList.AddDispatchPrerequisite(TaskEventRef);
+			}
+
 			new (RHICommandList.AllocCommand<FRHICommandUpdatePrimitiveIdBuffer>())FRHICommandUpdatePrimitiveIdBuffer(
-				TaskEventRefs,
 				PrimitiveIdVertexBufferRHI,
 				TaskContext.PrimitiveIdBufferData,
 				TaskContext.PrimitiveIdBufferDataSize
@@ -1090,11 +1088,15 @@ void FParallelMeshDrawCommandPass::DispatchDraw(FParallelCommandListSet* Paralle
 
 		const ENamedThreads::Type RenderThread = ENamedThreads::GetRenderThread();
 
-		FGraphEventArray Prereqs = TaskEventRefs;
+		FGraphEventArray Prereqs;
 		if (ParallelCommandListSet->GetPrereqs())
 		{
 			Prereqs.Append(*ParallelCommandListSet->GetPrereqs());
 		}
+		if (TaskEventRef.IsValid())
+		{
+			Prereqs.Add(TaskEventRef);
+		}
 
 		// Distribute work evenly to the available task graph workers based on NumEstimatedDraws.
 		// Every task will then adjust it's working range based on FVisibleMeshDrawCommandProcessTask results.
diff --git a/Engine/Source/Runtime/Renderer/Private/MeshDrawCommands.h b/Engine/Source/Runtime/Renderer/Private/MeshDrawCommands.h
--- a/Engine/Source/Runtime/Renderer/Private/MeshDrawCommands.h
+++ b/Engine/Source/Runtime/Renderer/Private/MeshDrawCommands.h
@@ -163,7 +163,7 @@ public:
 private:
 	FVertexBufferRHIParamRef PrimitiveIdVertexBufferRHI;
 	FMeshDrawCommandPassSetupTaskContext TaskContext;
-	FGraphEventArray TaskEventRefs;
+	FGraphEventRef TaskEventRef;
 	FString PassNameForStats;
 
 	// If TaskContext::PrimitiveIdBufferData will be released by RHI Thread.
