Fixing up GPU readback for the pathtracer.
Also cleaned up the dtor for d3d12 staging buffers.

diff --git a/Engine/Source/Runtime/D3D12RHI/Private/D3D12UAV.cpp b/Engine/Source/Runtime/D3D12RHI/Private/D3D12UAV.cpp
--- a/Engine/Source/Runtime/D3D12RHI/Private/D3D12UAV.cpp
+++ b/Engine/Source/Runtime/D3D12RHI/Private/D3D12UAV.cpp
@@ -262,6 +262,14 @@ FUnorderedAccessViewRHIRef FD3D12DynamicRHI::RHICreateUnorderedAccessView_Render
 	return RHICreateUnorderedAccessView(VertexBufferRHI, Format);
 }
 
+FD3D12StagingBuffer::~FD3D12StagingBuffer()
+{
+	if (StagedRead)
+	{
+		StagedRead.SafeRelease();
+	}
+}
+
 void* FD3D12StagingBuffer::Lock(uint32 Offset, uint32 NumBytes)
 {
 	check(!bIsLocked);
diff --git a/Engine/Source/Runtime/D3D12RHI/Public/D3D12Resources.h b/Engine/Source/Runtime/D3D12RHI/Public/D3D12Resources.h
--- a/Engine/Source/Runtime/D3D12RHI/Public/D3D12Resources.h
+++ b/Engine/Source/Runtime/D3D12RHI/Public/D3D12Resources.h
@@ -956,17 +956,7 @@ public:
 		: FRHIStagingBuffer()
 		, ShadowBufferSize(0)
 	{}
-	~FD3D12StagingBuffer()
-	{
-		if (StagedRead)
-		{
-			StagedRead->DeferDelete();
-		}
-	}
-	void SafeRelease()
-	{
-		StagedRead.SafeRelease();
-	}
+	virtual ~FD3D12StagingBuffer() final override;
 
 	virtual void* Lock(uint32 Offset, uint32 NumBytes) final override;
 	virtual void Unlock() final override;
diff --git a/Engine/Source/Runtime/Renderer/Private/PathTracing/RayCounterCompute.cpp b/Engine/Source/Runtime/Renderer/Private/PathTracing/RayCounterCompute.cpp
--- a/Engine/Source/Runtime/Renderer/Private/PathTracing/RayCounterCompute.cpp
+++ b/Engine/Source/Runtime/Renderer/Private/PathTracing/RayCounterCompute.cpp
@@ -108,23 +108,25 @@ void FDeferredShadingSceneRenderer::ComputeRayCount(FRHICommandListImmediate& RH
 	int32 NumGroups = FMath::DivideAndRoundUp<int32>(ViewSize.Y, FRayCounterCS::GetGroupSize());
 	DispatchComputeShader(RHICmdList, *RayCounterComputeShader, NumGroups, 1, 1);
 
+	FRHIGPUMemoryReadback* RayCountGPUReadback = ViewState->RayCountGPUReadback;
+
 	// Read read count data from the GPU using a stage buffer to avoid stalls
-	static FRHIGPUMemoryReadback* RayCountGPUReadback = nullptr;
-	if (!RayCountGPUReadback)
+	if (!ViewState->bReadbackInitialized)
 	{
-		RayCountGPUReadback = new FRHIGPUMemoryReadback(ViewState->TotalRayCountBuffer->Buffer, TEXT("Path Tracing Ray Count Readback"));
-		RayCountGPUReadback->Insert(RHICmdList);
+		RayCountGPUReadback->EnqueueCopy(RHICmdList, ViewState->TotalRayCountBuffer->Buffer);
+		ViewState->bReadbackInitialized = true;
 	}
 	else if (RayCountGPUReadback->IsReady())
 	{
-		uint32* RayCountResultBuffer = static_cast<uint32*>(RayCountGPUReadback->RetrieveData(1 * sizeof(uint32)));
+		uint32* RayCountResultBuffer = static_cast<uint32*>(RayCountGPUReadback->Lock(1 * sizeof(uint32)));
 		ViewState->TotalRayCount = RayCountResultBuffer[0];
 		extern ENGINE_API float GAveragePathTracedMRays;
 		GAveragePathTracedMRays = ViewState->TotalRayCount / 1000000.f;
 
-		RayCountGPUReadback->Finish();
-		delete RayCountGPUReadback;
-		RayCountGPUReadback = nullptr;
+		RayCountGPUReadback->Unlock();
+
+		// Enqueue another copy so we can get the data out.
+		RayCountGPUReadback->EnqueueCopy(RHICmdList, ViewState->TotalRayCountBuffer->Buffer);
 	}
 }
 
diff --git a/Engine/Source/Runtime/Renderer/Private/RendererScene.cpp b/Engine/Source/Runtime/Renderer/Private/RendererScene.cpp
--- a/Engine/Source/Runtime/Renderer/Private/RendererScene.cpp
+++ b/Engine/Source/Runtime/Renderer/Private/RendererScene.cpp
@@ -55,6 +55,7 @@
 #include "GPUScene.h"
 #if RHI_RAYTRACING
 #include "RayTracing/RayTracingDynamicGeometryCollection.h"
+#include "RHIGPUReadback.h"
 #endif
 
 // Enable this define to do slow checks for components being added to the wrong
@@ -184,6 +185,8 @@ FSceneViewState::FSceneViewState()
 				TotalRayCountBuffer->Initialize(sizeof(uint32), 1, PF_R32_UINT);
 			});
 	}
+	bReadbackInitialized = false;
+	RayCountGPUReadback = new FRHIGPUMemoryReadback(TEXT("Ray Count Readback"));
 #endif
 }
 
@@ -228,6 +231,8 @@ FSceneViewState::~FSceneViewState()
 #if RHI_RAYTRACING
 	DestroyRWBuffer(VarianceMipTree);
 	DestroyRWBuffer(TotalRayCountBuffer);
+
+	delete RayCountGPUReadback;
 #endif // RHI_RAYTRACING
 }
 
diff --git a/Engine/Source/Runtime/Renderer/Private/ScenePrivate.h b/Engine/Source/Runtime/Renderer/Private/ScenePrivate.h
--- a/Engine/Source/Runtime/Renderer/Private/ScenePrivate.h
+++ b/Engine/Source/Runtime/Renderer/Private/ScenePrivate.h
@@ -71,6 +71,7 @@ class UStaticMesh;
 class UStaticMeshComponent;
 class UTextureCube;
 class UWindDirectionalSourceComponent;
+class FRHIGPUMemoryReadback;
 
 /** Holds information about a single primitive's occlusion. */
 class FPrimitiveOcclusionHistory
@@ -911,6 +912,10 @@ public:
 	// Path tracer ray counter
 	uint32 TotalRayCount;
 	FRWBuffer* TotalRayCountBuffer;
+
+	// For Ray Count readback:
+	FRHIGPUMemoryReadback* RayCountGPUReadback;
+	bool bReadbackInitialized = false;
 #endif
 
 	// cache for stencil reads to a avoid reallocations of the SRV, Key is to detect if the object has changed
