Fixes some MSVC static analysis warnings.

diff --git a/Engine/Source/Runtime/RenderCore/Private/RenderGraphBuilder.cpp b/Engine/Source/Runtime/RenderCore/Private/RenderGraphBuilder.cpp
--- a/Engine/Source/Runtime/RenderCore/Private/RenderGraphBuilder.cpp
+++ b/Engine/Source/Runtime/RenderCore/Private/RenderGraphBuilder.cpp
@@ -153,13 +153,16 @@ void FRDGBuilder::Execute()
 	}
 
 	// Pops remaining scopes
-	if ((RENDER_GRAPH_DRAW_EVENTS != 0) && GetEmitDrawEvents())
+	if (RENDER_GRAPH_DRAW_EVENTS)
 	{
-		for (int32 i = 0; i < kMaxScopeCount; i++)
+		if (GetEmitDrawEvents())
 		{
-			if (!ScopesStack[i])
-				break;
-			RHICmdList.PopEvent();
+			for (int32 i = 0; i < kMaxScopeCount; i++)
+			{
+				if (!ScopesStack[i])
+					break;
+				RHICmdList.PopEvent();
+			}
 		}
 	}
 
diff --git a/Engine/Source/Runtime/Renderer/Private/ScreenSpaceDenoise.cpp b/Engine/Source/Runtime/Renderer/Private/ScreenSpaceDenoise.cpp
--- a/Engine/Source/Runtime/Renderer/Private/ScreenSpaceDenoise.cpp
+++ b/Engine/Source/Runtime/Renderer/Private/ScreenSpaceDenoise.cpp
@@ -785,10 +785,8 @@ static void DenoiseSignalAtConstantPixelDensity(
 		FRDGTextureRef SignalOutput0 = GraphBuilder.CreateTexture(SignalProcessingDesc[0], TEXT("SSDReflectionsHistory0"));
 		FRDGTextureRef SignalOutput1 = GraphBuilder.CreateTexture(SignalProcessingDesc[1], TEXT("SSDReflectionsHistory1"));
 
-		const int32 MipLevel = 0;
-
 		FSSDTemporalAccumulationCS::FPermutationDomain PermutationVector;
-		PermutationVector.Set<FSSDTemporalAccumulationCS::FIsMip0Dim>(MipLevel == 0);
+		PermutationVector.Set<FSSDTemporalAccumulationCS::FIsMip0Dim>(true);
 		PermutationVector.Set<FSignalProcessingDim>(Settings.SignalProcessing);
 
 		TShaderMapRef<FSSDTemporalAccumulationCS> ComputeShader(View.ShaderMap, PermutationVector);
@@ -808,8 +806,8 @@ static void DenoiseSignalAtConstantPixelDensity(
 		PassParameters->PrevGBufferA = RegisterExternalTextureWithFallback(GraphBuilder, View.PrevViewInfo.GBufferA, GSystemTextures.BlackDummy);
 		PassParameters->PrevGBufferB = RegisterExternalTextureWithFallback(GraphBuilder, View.PrevViewInfo.GBufferB, GSystemTextures.BlackDummy);
 
-		PassParameters->SignalHistoryOutput0 = GraphBuilder.CreateUAV(FRDGTextureUAVDesc(SignalOutput0, MipLevel));
-		PassParameters->SignalHistoryOutput1 = GraphBuilder.CreateUAV(FRDGTextureUAVDesc(SignalOutput1, MipLevel));
+		PassParameters->SignalHistoryOutput0 = GraphBuilder.CreateUAV(SignalOutput0);
+		PassParameters->SignalHistoryOutput1 = GraphBuilder.CreateUAV(SignalOutput1);
 		
 		PassParameters->DebugOutput = GraphBuilder.CreateUAV(GraphBuilder.CreateTexture(DebugDesc, TEXT("SSDDebugReflectionTemporalAccumulation")));
 
