PrimitiveIdBufferVB - use RHILockVertexBuffer instead of GDynamicRHI->RHILockVertexBuffer to lock/unlock VB on render thread.

diff --git a/Engine/Source/Runtime/Renderer/Private/MeshDrawCommands.cpp b/Engine/Source/Runtime/Renderer/Private/MeshDrawCommands.cpp
--- a/Engine/Source/Runtime/Renderer/Private/MeshDrawCommands.cpp
+++ b/Engine/Source/Runtime/Renderer/Private/MeshDrawCommands.cpp
@@ -815,7 +815,7 @@ void SortAndMergeDynamicPassMeshDrawCommands(
 			MaxInstanceFactor = VisibleMeshDrawCommands[0].MeshDrawCommand->InstanceFactor;
 			const int32 PrimitiveIdBufferDataSize = MaxInstanceFactor * NumDrawCommands * sizeof(int32);
 			OutPrimitiveIdVertexBuffer = GPrimitiveIdVertexBufferPool.Allocate(PrimitiveIdBufferDataSize);
-			void* PrimitiveIdBufferData = GDynamicRHI->RHILockVertexBuffer(OutPrimitiveIdVertexBuffer, 0, PrimitiveIdBufferDataSize, RLM_WriteOnly);
+			void* PrimitiveIdBufferData = RHILockVertexBuffer(OutPrimitiveIdVertexBuffer, 0, PrimitiveIdBufferDataSize, RLM_WriteOnly);
 
 			BuildMeshDrawCommandPrimitiveIdBuffer(
 				VisibleMeshDrawCommands,
@@ -829,7 +829,7 @@ void SortAndMergeDynamicPassMeshDrawCommands(
 				MaxInstanceFactor
 			);
 
-			GDynamicRHI->RHIUnlockVertexBuffer(OutPrimitiveIdVertexBuffer);
+			RHIUnlockVertexBuffer(OutPrimitiveIdVertexBuffer);
 		}
 	}
 }
@@ -1124,9 +1124,9 @@ void FParallelMeshDrawCommandPass::DispatchDraw(FParallelCommandListSet* Paralle
 		if (TaskContext.bUseGPUScene)
 		{
 			// Can immediately upload vertex buffer data, as there is no parallel draw task.
-			void* RESTRICT Data = (void* RESTRICT)GDynamicRHI->RHILockVertexBuffer(PrimitiveIdVertexBufferRHI, 0, TaskContext.PrimitiveIdBufferDataSize, RLM_WriteOnly);
+			void* RESTRICT Data = RHILockVertexBuffer(PrimitiveIdVertexBufferRHI, 0, TaskContext.PrimitiveIdBufferDataSize, RLM_WriteOnly);
 			FMemory::BigBlockMemcpy(Data, TaskContext.PrimitiveIdBufferData, TaskContext.PrimitiveIdBufferDataSize);
-			GDynamicRHI->RHIUnlockVertexBuffer(PrimitiveIdVertexBufferRHI);
+			RHIUnlockVertexBuffer(PrimitiveIdVertexBufferRHI);
 		}
 
 		SubmitMeshDrawCommandsRange(TaskContext.MeshDrawCommands, PrimitiveIdsBuffer, BasePrimitiveIdsOffset, TaskContext.bDynamicInstancing, 0, TaskContext.MeshDrawCommands.Num(), RHICmdList);
