Fixed race condition between parallel VB update inside RHI draw command, and parallel command lists.

diff --git a/Engine/Source/Runtime/Renderer/Private/MeshDrawCommands.cpp b/Engine/Source/Runtime/Renderer/Private/MeshDrawCommands.cpp
--- a/Engine/Source/Runtime/Renderer/Private/MeshDrawCommands.cpp
+++ b/Engine/Source/Runtime/Renderer/Private/MeshDrawCommands.cpp
@@ -1083,6 +1083,9 @@ void FParallelMeshDrawCommandPass::DispatchDraw(FParallelCommandListSet* Paralle
 				TaskContext.PrimitiveIdBufferData,
 				TaskContext.PrimitiveIdBufferDataSize
 			);
+
+			RHICommandList.RHIThreadFence(true);
+
 			bPrimitiveIdBufferDataOwnedByRHIThread = true;
 		}
 
