Fixed race condition when r.Cache.UpdatePrimsTaskEnabled=0 by moving IndirectLightingCache.UpdateCache so now updates can't be picked up in the same frame by the shadow setup.
* InitViews dispatched mesh draw command tasks operating on cached mesh draw commands.
* Scene->IndirectLightingCache.UpdateCache added static mesh updates for dispatched mesh draw commands.
* InitDynamicShadows picked up those updates and recached mesh draw commands, operating on the same memory as parallel tasks.

diff --git a/Engine/Source/Runtime/Renderer/Private/SceneVisibility.cpp b/Engine/Source/Runtime/Renderer/Private/SceneVisibility.cpp
--- a/Engine/Source/Runtime/Renderer/Private/SceneVisibility.cpp
+++ b/Engine/Source/Runtime/Renderer/Private/SceneVisibility.cpp
@@ -3612,17 +3612,10 @@ void FSceneRenderer::PostVisibilityFrameSetup(FILCUpdatePrimTaskData& OutILCTask
 		bCheckLightShafts = ViewFamily.EngineShowFlags.LightShafts && GLightShafts;
 	}
 
-	if (ViewFamily.EngineShowFlags.HitProxies == 0 && Scene->PrecomputedLightVolumes.Num() > 0)
+	if (ViewFamily.EngineShowFlags.HitProxies == 0 && Scene->PrecomputedLightVolumes.Num() > 0
+		&& GILCUpdatePrimTaskEnabled && FPlatformProcess::SupportsMultithreading())
 	{
-		QUICK_SCOPE_CYCLE_COUNTER(STAT_PostVisibilityFrameSetup_IndirectLightingCache_Update);
-		if (GILCUpdatePrimTaskEnabled && FPlatformProcess::SupportsMultithreading())
-		{
-			Scene->IndirectLightingCache.StartUpdateCachePrimitivesTask(Scene, *this, true, OutILCTaskData);
-		}
-		else
-		{
-			Scene->IndirectLightingCache.UpdateCache(Scene, *this, true);
-		}		
+		Scene->IndirectLightingCache.StartUpdateCachePrimitivesTask(Scene, *this, true, OutILCTaskData);
 	}
 
 	{
@@ -3966,7 +3959,17 @@ void FDeferredShadingSceneRenderer::InitViewsPossiblyAfterPrepass(FRHICommandLis
 		RHICmdList.ImmediateFlush(EImmediateFlushType::DispatchToRHIThread);
 	}
 
-	// if we kicked off ILC update via task, wait and finalize.
+	// If parallel ILC update is disabled, then process it in place.
+	if (ViewFamily.EngineShowFlags.HitProxies == 0
+		&& Scene->PrecomputedLightVolumes.Num() > 0
+		&& !(GILCUpdatePrimTaskEnabled && FPlatformProcess::SupportsMultithreading()))
+	{
+		QUICK_SCOPE_CYCLE_COUNTER(STAT_PostVisibilityFrameSetup_IndirectLightingCache_Update);
+		check(!ILCTaskData.TaskRef.IsValid());
+		Scene->IndirectLightingCache.UpdateCache(Scene, *this, true);
+	}
+
+	// If we kicked off ILC update via task, wait and finalize.
 	if (ILCTaskData.TaskRef.IsValid())
 	{
 		Scene->IndirectLightingCache.FinalizeCacheUpdates(Scene, *this, ILCTaskData);
diff --git a/Engine/Source/Runtime/Renderer/Private/ShadowSetup.cpp b/Engine/Source/Runtime/Renderer/Private/ShadowSetup.cpp
--- a/Engine/Source/Runtime/Renderer/Private/ShadowSetup.cpp
+++ b/Engine/Source/Runtime/Renderer/Private/ShadowSetup.cpp
@@ -1071,6 +1071,12 @@ void FProjectedShadowInfo::AddSubjectPrimitive(FPrimitiveSceneInfo* PrimitiveSce
 			*(PrimitiveSceneInfo->ComponentLastRenderTime) = CurrentWorldTime;
 			if (PrimitiveSceneInfo->NeedsUniformBufferUpdate() || PrimitiveSceneInfo->NeedsUpdateStaticMeshes())
 			{
+				for (int32 ViewIndex = 0; ViewIndex < Views.Num(); ++ViewIndex)
+				{
+					// Main view visible primitives are processed on parallel tasks, updating them here will cause a race condition.
+					check(!Views[ViewIndex]->PrimitiveVisibilityMap[PrimitiveSceneInfo->GetIndex()]);
+				}
+
 				PrimitiveSceneInfo->ConditionalUpdateStaticMeshes(FRHICommandListExecutor::GetImmediateCommandList());
 				PrimitiveSceneInfo->ConditionalUpdateUniformBuffer(FRHICommandListExecutor::GetImmediateCommandList());
 			}
