Replaced RemoveCachedMeshDrawCommands() with a check inside FPrimitiveSceneInfo destructor, as they should be removed when primitive is removed from the scene. Besides it's unsafe to operate on mesh draw commands here, as this may be called from a game thread.

diff --git a/Engine/Source/Runtime/Renderer/Private/PrimitiveSceneInfo.cpp b/Engine/Source/Runtime/Renderer/Private/PrimitiveSceneInfo.cpp
--- a/Engine/Source/Runtime/Renderer/Private/PrimitiveSceneInfo.cpp
+++ b/Engine/Source/Runtime/Renderer/Private/PrimitiveSceneInfo.cpp
@@ -167,8 +167,7 @@ FPrimitiveSceneInfo::FPrimitiveSceneInfo(UPrimitiveComponent* InComponent,FScene
 FPrimitiveSceneInfo::~FPrimitiveSceneInfo()
 {
 	check(!OctreeId.IsValidId());
-
-	RemoveCachedMeshDrawCommands();
+	check(StaticMeshCommandInfos.Num() == 0);
 }
 
 #if RHI_RAYTRACING
