Fixed crash inside DrawTileMesh when rendering translucent material. DrawTileMesh doesn't have a scene, so we need to create a new single frame global uniform buffer in that case.

diff --git a/Engine/Source/Runtime/Renderer/Private/TranslucentRendering.cpp b/Engine/Source/Runtime/Renderer/Private/TranslucentRendering.cpp
--- a/Engine/Source/Runtime/Renderer/Private/TranslucentRendering.cpp
+++ b/Engine/Source/Runtime/Renderer/Private/TranslucentRendering.cpp
@@ -674,12 +674,17 @@ void CreateTranslucentBasePassUniformBuffer(
 
 	BasePassParameters.SceneTextures.SceneColorCopyTexture = SceneColorCopy ? SceneColorCopy->GetRenderTargetItem().ShaderResourceTexture : GBlackTexture->TextureRHI;
 
+	FScene* Scene = View.Family->Scene ? View.Family->Scene->GetRenderScene() : nullptr;
 
-	check(View.Family->Scene);
-	FScene* Scene = View.Family->Scene->GetRenderScene();
-
-	Scene->UniformBuffers.TranslucentBasePassUniformBuffer.UpdateUniformBufferImmediate(BasePassParameters);
-	BasePassUniformBuffer = Scene->UniformBuffers.TranslucentBasePassUniformBuffer;
+	if (Scene)
+	{
+		Scene->UniformBuffers.TranslucentBasePassUniformBuffer.UpdateUniformBufferImmediate(BasePassParameters);
+		BasePassUniformBuffer = Scene->UniformBuffers.TranslucentBasePassUniformBuffer;
+	}
+	else
+	{
+		BasePassUniformBuffer = TUniformBufferRef<FTranslucentBasePassUniformParameters>::CreateUniformBufferImmediate(BasePassParameters, UniformBuffer_SingleFrame);
+	}
 }
 
 class FTranslucencyUpsamplingPS : public FGlobalShader
