Make FBasicRayData & FBasicRayIntersectionData match the definitions of their MPS equivalents as requested.

diff --git a/Engine/Shaders/Private/RayTracing/RayTracingCommon.ush b/Engine/Shaders/Private/RayTracing/RayTracingCommon.ush
--- a/Engine/Shaders/Private/RayTracing/RayTracingCommon.ush
+++ b/Engine/Shaders/Private/RayTracing/RayTracingCommon.ush
@@ -17,16 +17,17 @@
 struct FBasicRayData
 {
 	float3 Origin;
-	float TFar;
-	float3 Direction;
 	uint Mask;
+	float3 Direction;
+	float TFar;
 };
 
 struct FBasicRayIntersectionData
 {
-	float2 Barycentrics;
-	uint   InstanceIndex;  // 0xFFFFFFFF if no intersection if found
+	float  Distance;
 	uint   PrimitiveIndex; // 0xFFFFFFFF if no intersection if found
+	uint   InstanceIndex;  // 0xFFFFFFFF if no intersection if found
+	float2 Barycentrics;
 };
 
 struct FOcclusionPayload
diff --git a/Engine/Source/Runtime/RHI/Public/RHICommandList.h b/Engine/Source/Runtime/RHI/Public/RHICommandList.h
--- a/Engine/Source/Runtime/RHI/Public/RHICommandList.h
+++ b/Engine/Source/Runtime/RHI/Public/RHICommandList.h
@@ -122,18 +122,19 @@ struct FRayTracingShaderBindings
 struct FBasicRayData
 {
 	float Origin[3];
-	float TFar;
-	float Direction[3];
 	uint32 Mask;
+	float Direction[3];
+	float TFar;
 };
 
 // C++ counter-part of FBasicRayIntersectionData declared in RayTracingCommon.ush
 // #dxr_todo: ideally this should be shared with shaders by including a common header
 struct FBasicRayIntersectionData
 {
-	float  Barycentrics[2];
-	uint32 InstanceIndex;  // 0xFFFFFFFF if no intersection if found
+	float  Distance;
 	uint32 PrimitiveIndex; // 0xFFFFFFFF if no intersection if found
+	uint32 InstanceIndex;  // 0xFFFFFFFF if no intersection if found
+	float  Barycentrics[2];
 };
 #endif // RHI_RAYTRACING
 
diff --git a/Engine/Source/Runtime/Renderer/Private/RayTracing/RayTracingTestbed.cpp b/Engine/Source/Runtime/Renderer/Private/RayTracing/RayTracingTestbed.cpp
--- a/Engine/Source/Runtime/Renderer/Private/RayTracing/RayTracingTestbed.cpp
+++ b/Engine/Source/Runtime/Renderer/Private/RayTracing/RayTracingTestbed.cpp
@@ -53,10 +53,10 @@ void TestBasicRayTracing(bool bValidateResults)
 	{
 		TResourceArray<FBasicRayData> RayData;
 		RayData.SetNumUninitialized(NumRays);
-		RayData[0] = FBasicRayData{ { 0.75f, 0.0f, -1.0f}, 100000.0f, {0.0f, 0.0f,  1.0f}, 0xFFFFFFFF }; // expected to hit
-		RayData[1] = FBasicRayData{ { 0.75f, 0.0f, -1.0f},      0.5f, {0.0f, 0.0f,  1.0f}, 0xFFFFFFFF }; // expected to miss (short ray)
-		RayData[2] = FBasicRayData{ { 0.75f, 0.0f,  1.0f}, 100000.0f, {0.0f, 0.0f, -1.0f}, 0xFFFFFFFF }; // expected to miss (back face culled)
-		RayData[3] = FBasicRayData{ {-0.75f, 0.0f, -1.0f}, 100000.0f, {0.0f, 0.0f,  1.0f}, 0xFFFFFFFF }; // expected to miss (doesn't intersect)
+		RayData[0] = FBasicRayData{ { 0.75f, 0.0f, -1.0f}, 0xFFFFFFFF, {0.0f, 0.0f,  1.0f}, 100000.0f }; // expected to hit
+		RayData[1] = FBasicRayData{ { 0.75f, 0.0f, -1.0f}, 0xFFFFFFFF, {0.0f, 0.0f,  1.0f},      0.5f }; // expected to miss (short ray)
+		RayData[2] = FBasicRayData{ { 0.75f, 0.0f,  1.0f}, 0xFFFFFFFF, {0.0f, 0.0f, -1.0f}, 100000.0f }; // expected to miss (back face culled)
+		RayData[3] = FBasicRayData{ {-0.75f, 0.0f, -1.0f}, 0xFFFFFFFF, {0.0f, 0.0f,  1.0f}, 100000.0f }; // expected to miss (doesn't intersect)
 
 		FRHIResourceCreateInfo CreateInfo;
 		CreateInfo.ResourceArray = &RayData;
