Relevance for dynamic mesh elements. Before generating mesh draw commands dynamic mesh element list is filtered by mesh pass relevance. This both prevents crashes when proxy relevance mismatches actual mesh contents and speeds up traversal through dynamic mesh element lists for sparse passes like velocity.

diff --git a/Engine/Source/Runtime/Renderer/Private/MeshDrawCommands.cpp b/Engine/Source/Runtime/Renderer/Private/MeshDrawCommands.cpp
--- a/Engine/Source/Runtime/Renderer/Private/MeshDrawCommands.cpp
+++ b/Engine/Source/Runtime/Renderer/Private/MeshDrawCommands.cpp
@@ -434,6 +434,7 @@ void GenerateDynamicMeshDrawCommands(
 	EMeshPass::Type PassType,
 	FMeshPassProcessor* PassMeshProcessor,
 	const TArray<FMeshBatchAndRelevance, SceneRenderingAllocator>& DynamicMeshElements,
+	const TArray<FMeshPassMask, SceneRenderingAllocator>* DynamicMeshElementsPassRelevance,
 	int32 MaxNumDynamicMeshElements,
 	const TArray<const FStaticMeshBatch*, SceneRenderingAllocator>& DynamicMeshCommandBuildRequests,
 	int32 MaxNumBuildRequestElements,
@@ -443,6 +444,7 @@ void GenerateDynamicMeshDrawCommands(
 {
 	QUICK_SCOPE_CYCLE_COUNTER(STAT_GenerateDynamicMeshDrawCommands);
 	check(PassMeshProcessor);
+	check((PassType == EMeshPass::Num) == (DynamicMeshElementsPassRelevance == nullptr));
 
 	FDynamicPassMeshDrawListContext DynamicPassMeshDrawListContext(
 		MeshDrawCommandStorage,
@@ -456,11 +458,14 @@ void GenerateDynamicMeshDrawCommands(
 
 		for (int32 MeshIndex = 0; MeshIndex < NumDynamicMeshBatches; MeshIndex++)
 		{
-			const FMeshBatchAndRelevance& MeshAndRelevance = DynamicMeshElements[MeshIndex];
-			check(!MeshAndRelevance.Mesh->bRequiresPerElementVisibility);
-			const uint64 BatchElementMask = ~0ull;
+			if (!DynamicMeshElementsPassRelevance || (*DynamicMeshElementsPassRelevance)[MeshIndex].Get(PassType))
+			{
+				const FMeshBatchAndRelevance& MeshAndRelevance = DynamicMeshElements[MeshIndex];
+				check(!MeshAndRelevance.Mesh->bRequiresPerElementVisibility);
+				const uint64 BatchElementMask = ~0ull;
 
-			PassMeshProcessor->AddMeshBatch(*MeshAndRelevance.Mesh, BatchElementMask, MeshAndRelevance.PrimitiveSceneProxy);
+				PassMeshProcessor->AddMeshBatch(*MeshAndRelevance.Mesh, BatchElementMask, MeshAndRelevance.PrimitiveSceneProxy);
+			}
 		}
 
 		const int32 NumCommandsGenerated = VisibleCommands.Num() - NumCommandsBefore;
@@ -497,6 +502,7 @@ void GenerateMobileBasePassDynamicMeshDrawCommands(
 	FMeshPassProcessor* PassMeshProcessor,
 	FMeshPassProcessor* MobilePassCSMPassMeshProcessor,
 	const TArray<FMeshBatchAndRelevance, SceneRenderingAllocator>& DynamicMeshElements,
+	const TArray<FMeshPassMask, SceneRenderingAllocator>* DynamicMeshElementsPassRelevance,
 	int32 MaxNumDynamicMeshElements,
 	const TArray<const FStaticMeshBatch*, SceneRenderingAllocator>& DynamicMeshCommandBuildRequests,
 	int32 MaxNumBuildRequestElements,
@@ -506,6 +512,7 @@ void GenerateMobileBasePassDynamicMeshDrawCommands(
 {
 	QUICK_SCOPE_CYCLE_COUNTER(STAT_GenerateMobileBasePassDynamicMeshDrawCommands);
 	check(PassMeshProcessor && MobilePassCSMPassMeshProcessor);
+	check((PassType == EMeshPass::Num) == (DynamicMeshElementsPassRelevance == nullptr));
 
 	FDynamicPassMeshDrawListContext DynamicPassMeshDrawListContext(
 		MeshDrawCommandStorage,
@@ -522,19 +529,22 @@ void GenerateMobileBasePassDynamicMeshDrawCommands(
 
 		for (int32 MeshIndex = 0; MeshIndex < NumDynamicMeshBatches; MeshIndex++)
 		{
-			const FMeshBatchAndRelevance& MeshAndRelevance = DynamicMeshElements[MeshIndex];
-			check(!MeshAndRelevance.Mesh->bRequiresPerElementVisibility);
-			const uint64 BatchElementMask = ~0ull;
-
-			const int32 PrimitiveIndex = MeshAndRelevance.PrimitiveSceneProxy->GetPrimitiveSceneInfo()->GetIndex();
-			if (MobileCSMVisibilityInfo.bMobileDynamicCSMInUse
-				&& (MobileCSMVisibilityInfo.bAlwaysUseCSM || MobileCSMVisibilityInfo.MobilePrimitiveCSMReceiverVisibilityMap[PrimitiveIndex]))
-			{
-				MobilePassCSMPassMeshProcessor->AddMeshBatch(*MeshAndRelevance.Mesh, BatchElementMask, MeshAndRelevance.PrimitiveSceneProxy);
-			}
-			else
+			if (!DynamicMeshElementsPassRelevance || (*DynamicMeshElementsPassRelevance)[MeshIndex].Get(PassType))
 			{
-				PassMeshProcessor->AddMeshBatch(*MeshAndRelevance.Mesh, BatchElementMask, MeshAndRelevance.PrimitiveSceneProxy);
+				const FMeshBatchAndRelevance& MeshAndRelevance = DynamicMeshElements[MeshIndex];
+				check(!MeshAndRelevance.Mesh->bRequiresPerElementVisibility);
+				const uint64 BatchElementMask = ~0ull;
+
+				const int32 PrimitiveIndex = MeshAndRelevance.PrimitiveSceneProxy->GetPrimitiveSceneInfo()->GetIndex();
+				if (MobileCSMVisibilityInfo.bMobileDynamicCSMInUse
+					&& (MobileCSMVisibilityInfo.bAlwaysUseCSM || MobileCSMVisibilityInfo.MobilePrimitiveCSMReceiverVisibilityMap[PrimitiveIndex]))
+				{
+					MobilePassCSMPassMeshProcessor->AddMeshBatch(*MeshAndRelevance.Mesh, BatchElementMask, MeshAndRelevance.PrimitiveSceneProxy);
+				}
+				else
+				{
+					PassMeshProcessor->AddMeshBatch(*MeshAndRelevance.Mesh, BatchElementMask, MeshAndRelevance.PrimitiveSceneProxy);
+				}
 			}
 		}
 
@@ -690,6 +700,7 @@ public:
 				Context.MeshPassProcessor,
 				Context.MobileBasePassCSMMeshPassProcessor,
 				*Context.DynamicMeshElements,
+				Context.DynamicMeshElementsPassRelevance,
 				Context.NumDynamicMeshElements,
 				Context.DynamicMeshCommandBuildRequests,
 				Context.NumDynamicMeshCommandBuildRequestElements,
@@ -705,6 +716,7 @@ public:
 				Context.PassType,
 				Context.MeshPassProcessor,
 				*Context.DynamicMeshElements,
+				Context.DynamicMeshElementsPassRelevance,
 				Context.NumDynamicMeshElements,
 				Context.DynamicMeshCommandBuildRequests,
 				Context.NumDynamicMeshCommandBuildRequestElements,
@@ -842,6 +854,7 @@ void FParallelMeshDrawCommandPass::DispatchPassSetup(
 	FExclusiveDepthStencil::Type BasePassDepthStencilAccess,
 	FMeshPassProcessor* MeshPassProcessor,
 	const TArray<FMeshBatchAndRelevance, SceneRenderingAllocator>& DynamicMeshElements,
+	const TArray<FMeshPassMask, SceneRenderingAllocator>* DynamicMeshElementsPassRelevance,
 	int32 NumDynamicMeshElements,
 	TArray<const FStaticMeshBatch*, SceneRenderingAllocator>& InOutDynamicMeshCommandBuildRequests,
 	int32 NumDynamicMeshCommandBuildRequestElements,
@@ -851,6 +864,7 @@ void FParallelMeshDrawCommandPass::DispatchPassSetup(
 )
 {
 	check(!TaskEventRef.IsValid() && MeshPassProcessor != nullptr && TaskContext.PrimitiveIdBufferData == nullptr);
+	check((PassType == EMeshPass::Num) == (DynamicMeshElementsPassRelevance == nullptr));
 
 	MaxNumDraws = InOutMeshDrawCommands.Num() + NumDynamicMeshElements + NumDynamicMeshCommandBuildRequestElements;
 	if (MaxNumDraws <= 0)
@@ -861,6 +875,7 @@ void FParallelMeshDrawCommandPass::DispatchPassSetup(
 	TaskContext.MeshPassProcessor = MeshPassProcessor;
 	TaskContext.MobileBasePassCSMMeshPassProcessor = MobileBasePassCSMMeshPassProcessor;
 	TaskContext.DynamicMeshElements = &DynamicMeshElements;
+	TaskContext.DynamicMeshElementsPassRelevance = DynamicMeshElementsPassRelevance;
 
 	TaskContext.View = &View;
 	TaskContext.ShadingPath = Scene->GetShadingPath();
@@ -957,6 +972,7 @@ void FParallelMeshDrawCommandPass::Empty()
 	TaskContext.MeshPassProcessor = nullptr;
 	TaskContext.MobileBasePassCSMMeshPassProcessor = nullptr;
 	TaskContext.DynamicMeshElements = nullptr;
+	TaskContext.DynamicMeshElementsPassRelevance = nullptr;
 	TaskContext.MeshDrawCommands.Empty();
 	TaskContext.MobileBasePassCSMMeshDrawCommands.Empty();
 	TaskContext.DynamicMeshCommandBuildRequests.Empty();
diff --git a/Engine/Source/Runtime/Renderer/Private/MeshDrawCommands.h b/Engine/Source/Runtime/Renderer/Private/MeshDrawCommands.h
--- a/Engine/Source/Runtime/Renderer/Private/MeshDrawCommands.h
+++ b/Engine/Source/Runtime/Renderer/Private/MeshDrawCommands.h
@@ -84,6 +84,7 @@ public:
 	FMeshPassProcessor* MeshPassProcessor;
 	FMeshPassProcessor* MobileBasePassCSMMeshPassProcessor;
 	const TArray<FMeshBatchAndRelevance, SceneRenderingAllocator>* DynamicMeshElements;
+	const TArray<FMeshPassMask, SceneRenderingAllocator>* DynamicMeshElementsPassRelevance;
 
 	// Commands.
 	int32 NumDynamicMeshElements;
@@ -142,6 +143,7 @@ public:
 		FExclusiveDepthStencil::Type BasePassDepthStencilAccess,
 		FMeshPassProcessor* MeshPassProcessor,
 		const TArray<FMeshBatchAndRelevance, SceneRenderingAllocator>& DynamicMeshElements,
+		const TArray<FMeshPassMask, SceneRenderingAllocator>* DynamicMeshElementsPassRelevance,
 		int32 NumDynamicMeshElements,
 		TArray<const FStaticMeshBatch*, SceneRenderingAllocator>& InOutDynamicMeshCommandBuildRequests,
 		int32 NumDynamicMeshCommandBuildRequestElements,
diff --git a/Engine/Source/Runtime/Renderer/Private/MobileShadingRenderer.cpp b/Engine/Source/Runtime/Renderer/Private/MobileShadingRenderer.cpp
--- a/Engine/Source/Runtime/Renderer/Private/MobileShadingRenderer.cpp
+++ b/Engine/Source/Runtime/Renderer/Private/MobileShadingRenderer.cpp
@@ -149,6 +149,7 @@ void FMobileSceneRenderer::SetupMobileBasePassAfterShadowInit(FExclusiveDepthSte
 			BasePassDepthStencilAccess,
 			MeshPassProcessor,
 			View.DynamicMeshElements,
+			&View.DynamicMeshElementsPassRelevance,
 			View.NumVisibleDynamicMeshElements[EMeshPass::BasePass],
 			ViewCommands.DynamicMeshCommandBuildRequests[EMeshPass::BasePass],
 			ViewCommands.NumDynamicMeshCommandBuildRequestElements[EMeshPass::BasePass],
diff --git a/Engine/Source/Runtime/Renderer/Private/SceneRendering.cpp b/Engine/Source/Runtime/Renderer/Private/SceneRendering.cpp
--- a/Engine/Source/Runtime/Renderer/Private/SceneRendering.cpp
+++ b/Engine/Source/Runtime/Renderer/Private/SceneRendering.cpp
@@ -891,8 +891,6 @@ void FViewInfo::Init()
 	DitherFadeInUniformBuffer = nullptr;
 	DitherFadeOutUniformBuffer = nullptr;
 
-	VisibleDynamicMeshesPassMask.Reset();
-
 	for (int32 PassIndex = 0; PassIndex < EMeshPass::Num; ++PassIndex)
 	{
 		NumVisibleDynamicMeshElements[PassIndex] = 0;
@@ -2695,6 +2693,7 @@ void FSceneRenderer::SetupMeshPass(FViewInfo& View, FExclusiveDepthStencil::Type
 				BasePassDepthStencilAccess,
 				MeshPassProcessor,
 				View.DynamicMeshElements,
+				&View.DynamicMeshElementsPassRelevance,
 				View.NumVisibleDynamicMeshElements[PassType],
 				ViewCommands.DynamicMeshCommandBuildRequests[PassType],
 				ViewCommands.NumDynamicMeshCommandBuildRequestElements[PassType],
diff --git a/Engine/Source/Runtime/Renderer/Private/SceneRendering.h b/Engine/Source/Runtime/Renderer/Private/SceneRendering.h
--- a/Engine/Source/Runtime/Renderer/Private/SceneRendering.h
+++ b/Engine/Source/Runtime/Renderer/Private/SceneRendering.h
@@ -845,9 +845,6 @@ public:
 	/** Number of dynamic editor primitives visible in this view. */
 	int32 NumVisibleDynamicEditorPrimitives;
 
-	/** Marks used mesh passes by visible dynamic meshes. */
-	FMeshPassMask VisibleDynamicMeshesPassMask;
-
 	/** Number of dynamic mesh elements per mesh pass (inside FViewInfo::DynamicMeshElements). */
 	int32 NumVisibleDynamicMeshElements[EMeshPass::Num];
 
@@ -893,6 +890,9 @@ public:
 	/** Gathered in initviews from all the primitives with dynamic view relevance, used in each mesh pass. */
 	TArray<FMeshBatchAndRelevance,SceneRenderingAllocator> DynamicMeshElements;
 
+	/* Mesh pass relevance for gathered dynamic mesh elements. */
+	TArray<FMeshPassMask, SceneRenderingAllocator> DynamicMeshElementsPassRelevance;
+
 	/** Gathered in UpdateRayTracingWorld from all the primitives with dynamic view relevance, used in each mesh pass. */
 	TArray<FMeshBatchAndRelevance, SceneRenderingAllocator> RayTracedDynamicMeshElements;
 
diff --git a/Engine/Source/Runtime/Renderer/Private/SceneVisibility.cpp b/Engine/Source/Runtime/Renderer/Private/SceneVisibility.cpp
--- a/Engine/Source/Runtime/Renderer/Private/SceneVisibility.cpp
+++ b/Engine/Source/Runtime/Renderer/Private/SceneVisibility.cpp
@@ -1848,91 +1848,6 @@ struct FRelevancePacket
 		MarkRelevant();
 	}
 
-	void UpdateDynamicPrimMask(EShadingPath ShadingPath, bool bAddLightmapDensityCommands, const FPrimitiveViewRelevance& ViewRelevance)
-	{
-		// Compute number mesh pass mask for this dynamic primitive.
-		if (ViewRelevance.bDrawRelevance && (ViewRelevance.bRenderInMainPass || ViewRelevance.bRenderCustomDepth))
-		{
-			VisibleDynamicMeshesPassMask.Set(EMeshPass::DepthPass);
-			VisibleDynamicMeshesPassMask.Set(EMeshPass::BasePass);
-
-			if (ShadingPath == EShadingPath::Mobile)
-			{
-				VisibleDynamicMeshesPassMask.Set(EMeshPass::MobileBasePassCSM);
-			}
-
-			if (ViewRelevance.bRenderCustomDepth)
-			{
-				VisibleDynamicMeshesPassMask.Set(EMeshPass::CustomDepth);
-			}
-
-			if (bAddLightmapDensityCommands)
-			{
-				VisibleDynamicMeshesPassMask.Set(EMeshPass::LightmapDensity);
-			}
-#if !(UE_BUILD_SHIPPING || UE_BUILD_TEST)
-			else if (View.Family->UseDebugViewPS())
-			{
-				VisibleDynamicMeshesPassMask.Set(EMeshPass::DebugViewMode);
-			}
-#endif
-
-#if WITH_EDITOR
-			if (View.bAllowTranslucentPrimitivesInHitProxy)
-			{
-				VisibleDynamicMeshesPassMask.Set(EMeshPass::HitProxy);
-			}
-			else
-			{
-				VisibleDynamicMeshesPassMask.Set(EMeshPass::HitProxyOpaqueOnly);
-			}
-#endif
-
-			if (ViewRelevance.bVelocityRelevance)
-			{
-				VisibleDynamicMeshesPassMask.Set(EMeshPass::Velocity);
-			}
-		}
-
-		if (ViewRelevance.HasTranslucency()
-			&& !ViewRelevance.bEditorPrimitiveRelevance
-			&& ViewRelevance.bRenderInMainPass)
-		{
-			if (View.Family->AllowTranslucencyAfterDOF())
-			{
-				if (ViewRelevance.bNormalTranslucencyRelevance)
-				{
-					VisibleDynamicMeshesPassMask.Set(EMeshPass::TranslucencyStandard);
-				}
-
-				if (ViewRelevance.bSeparateTranslucencyRelevance)
-				{
-					VisibleDynamicMeshesPassMask.Set(EMeshPass::TranslucencyAfterDOF);
-				}
-			}
-			else
-			{
-				VisibleDynamicMeshesPassMask.Set(EMeshPass::TranslucencyAll);
-			}
-
-			if (ViewRelevance.bDistortionRelevance)
-			{
-				VisibleDynamicMeshesPassMask.Set(EMeshPass::Distortion);
-			}
-
-			if (ShadingPath == EShadingPath::Mobile && View.bIsSceneCapture)
-			{
-				VisibleDynamicMeshesPassMask.Set(EMeshPass::MobileInverseOpacity);
-			}
-		}
-#if WITH_EDITOR
-		if (ViewRelevance.bDrawRelevance)
-		{
-			VisibleDynamicMeshesPassMask.Set(EMeshPass::EditorSelection);
-		}
-#endif
-	}
-
 	void ComputeRelevance()
 	{
 		CombinedShadingModelMask = 0;
@@ -1995,7 +1910,6 @@ struct FRelevancePacket
 				// Keep track of visible dynamic primitives.
 				++NumVisibleDynamicPrimitives;
 				OutHasDynamicMeshElementsMasks[BitIndex] |= ViewBit;
-				UpdateDynamicPrimMask(ShadingPath, bAddLightmapDensityCommands, ViewRelevance);
 
 				if (ViewRelevance.bHasSimpleLights)
 				{
@@ -2345,7 +2259,6 @@ struct FRelevancePacket
 		VisibleDynamicPrimitivesWithSimpleLights.AppendTo(WriteView.VisibleDynamicPrimitivesWithSimpleLights);
 		WriteView.NumVisibleDynamicPrimitives += NumVisibleDynamicPrimitives;
 		WriteView.NumVisibleDynamicEditorPrimitives += NumVisibleDynamicEditorPrimitives;
-		VisibleDynamicMeshesPassMask.AppendTo(WriteView.VisibleDynamicMeshesPassMask);
 		WriteView.TranslucentPrimCount.Append(TranslucentPrimCount);
 		MeshDecalPrimSet.AppendTo(WriteView.MeshDecalPrimSet.Prims);
 		WriteView.bHasDistortionPrimitives |= bHasDistortionPrimitives;
@@ -2572,6 +2485,109 @@ static void SetDynamicMeshElementViewCustomData(TArray<FViewInfo>& InViews, cons
 	}
 }
 
+void ComputeDynamicMeshRelevance(EShadingPath ShadingPath, bool bAddLightmapDensityCommands, const FPrimitiveViewRelevance& ViewRelevance, const FMeshBatchAndRelevance& MeshBatch, FViewInfo& View, FMeshPassMask& PassMask)
+{
+	const int32 NumElements = MeshBatch.Mesh->Elements.Num();
+
+	if (ViewRelevance.bDrawRelevance && (ViewRelevance.bRenderInMainPass || ViewRelevance.bRenderCustomDepth))
+	{
+		PassMask.Set(EMeshPass::DepthPass);
+		View.NumVisibleDynamicMeshElements[EMeshPass::DepthPass] += NumElements;
+
+		PassMask.Set(EMeshPass::BasePass);
+		View.NumVisibleDynamicMeshElements[EMeshPass::BasePass] += NumElements;
+
+		if (ShadingPath == EShadingPath::Mobile)
+		{
+			PassMask.Set(EMeshPass::MobileBasePassCSM);
+			View.NumVisibleDynamicMeshElements[EMeshPass::MobileBasePassCSM] += NumElements;
+		}
+
+		if (ViewRelevance.bRenderCustomDepth)
+		{
+			PassMask.Set(EMeshPass::CustomDepth);
+			View.NumVisibleDynamicMeshElements[EMeshPass::CustomDepth] += NumElements;
+		}
+
+		if (bAddLightmapDensityCommands)
+		{
+			PassMask.Set(EMeshPass::LightmapDensity);
+			View.NumVisibleDynamicMeshElements[EMeshPass::LightmapDensity] += NumElements;
+		}
+#if !(UE_BUILD_SHIPPING || UE_BUILD_TEST)
+		else if (View.Family->UseDebugViewPS())
+		{
+			PassMask.Set(EMeshPass::DebugViewMode);
+			View.NumVisibleDynamicMeshElements[EMeshPass::DebugViewMode] += NumElements;
+		}
+#endif
+
+#if WITH_EDITOR
+		if (View.bAllowTranslucentPrimitivesInHitProxy)
+		{
+			PassMask.Set(EMeshPass::HitProxy);
+			View.NumVisibleDynamicMeshElements[EMeshPass::HitProxy] += NumElements;
+		}
+		else
+		{
+			PassMask.Set(EMeshPass::HitProxyOpaqueOnly);
+			View.NumVisibleDynamicMeshElements[EMeshPass::HitProxyOpaqueOnly] += NumElements;
+		}
+#endif
+
+		if (ViewRelevance.bVelocityRelevance)
+		{
+			PassMask.Set(EMeshPass::Velocity);
+			View.NumVisibleDynamicMeshElements[EMeshPass::Velocity] += NumElements;
+		}
+	}
+
+	if (ViewRelevance.HasTranslucency()
+		&& !ViewRelevance.bEditorPrimitiveRelevance
+		&& ViewRelevance.bRenderInMainPass)
+	{
+		if (View.Family->AllowTranslucencyAfterDOF())
+		{
+			if (ViewRelevance.bNormalTranslucencyRelevance)
+			{
+				PassMask.Set(EMeshPass::TranslucencyStandard);
+				View.NumVisibleDynamicMeshElements[EMeshPass::TranslucencyStandard] += NumElements;
+			}
+
+			if (ViewRelevance.bSeparateTranslucencyRelevance)
+			{
+				PassMask.Set(EMeshPass::TranslucencyAfterDOF);
+				View.NumVisibleDynamicMeshElements[EMeshPass::TranslucencyAfterDOF] += NumElements;
+			}
+		}
+		else
+		{
+			PassMask.Set(EMeshPass::TranslucencyAll);
+			View.NumVisibleDynamicMeshElements[EMeshPass::TranslucencyAll] += NumElements;
+		}
+
+		if (ViewRelevance.bDistortionRelevance)
+		{
+			PassMask.Set(EMeshPass::Distortion);
+			++View.NumVisibleDynamicMeshElements[EMeshPass::Distortion] += NumElements;
+		}
+
+		if (ShadingPath == EShadingPath::Mobile && View.bIsSceneCapture)
+		{
+			PassMask.Set(EMeshPass::MobileInverseOpacity);
+			++View.NumVisibleDynamicMeshElements[EMeshPass::MobileInverseOpacity] += NumElements;
+		}
+	}
+
+#if WITH_EDITOR
+	if (ViewRelevance.bDrawRelevance)
+	{
+		PassMask.Set(EMeshPass::EditorSelection);
+		View.NumVisibleDynamicMeshElements[EMeshPass::EditorSelection] += NumElements;
+	}
+#endif
+}
+
 void FSceneRenderer::GatherDynamicMeshElements(
 	TArray<FViewInfo>& InViews, 
 	const FScene* InScene, 
@@ -2601,6 +2617,7 @@ void FSceneRenderer::GatherDynamicMeshElements(
 		}
 
 		const bool bIsInstancedStereo = (ViewCount > 0) ? (InViews[0].IsInstancedStereoPass() || InViews[0].bIsMobileMultiViewEnabled) : false;
+		const EShadingPath ShadingPath = Scene->GetShadingPath();
 
 		for (int32 PrimitiveIndex = 0; PrimitiveIndex < NumPrimitives; ++PrimitiveIndex)
 		{
@@ -2624,18 +2641,25 @@ void FSceneRenderer::GatherDynamicMeshElements(
 			{
 				InViews[ViewIndex].DynamicMeshEndIndices[PrimitiveIndex] = Collector.GetMeshBatchCount(ViewIndex);
 			}
-		}
-
-		// Compute MaxNumVisibleDynamicMeshes.
-		for (int32 ViewIndex = 0; ViewIndex < Views.Num(); ++ViewIndex)
-		{
-			FViewInfo& View = Views[ViewIndex];
 
-			for (int32 PassIndex = 0; PassIndex < EMeshPass::Num; ++PassIndex)
+			// Compute DynamicMeshElementsMeshPassRelevance for this primitive.
+			for (int32 ViewIndex = 0; ViewIndex < ViewCount; ViewIndex++)
 			{
-				if (View.VisibleDynamicMeshesPassMask.Get((EMeshPass::Type) PassIndex))
+				FViewInfo& View = InViews[ViewIndex];
+				const bool bAddLightmapDensityCommands = View.Family->EngineShowFlags.LightMapDensity && AllowDebugViewmodes();
+				const FPrimitiveViewRelevance& ViewRelevance = View.PrimitiveViewRelevanceMap[PrimitiveIndex];
+
+				View.DynamicMeshElementsPassRelevance.SetNum(View.DynamicMeshElements.Num());
+
+				const int32 NumPrimitiveElements = Collector.GetMeshBatchCount(ViewIndex);
+				const int32 PrimitiveFirstElement = View.DynamicMeshElements.Num() - NumPrimitiveElements;
+
+				for (int32 ElementIndex = PrimitiveFirstElement; ElementIndex < NumPrimitiveElements; ++ElementIndex)
 				{
-					View.NumVisibleDynamicMeshElements[PassIndex] = Collector.GetMeshElementCount(ViewIndex);
+					const FMeshBatchAndRelevance& MeshBatch = View.DynamicMeshElements[ElementIndex];
+					FMeshPassMask& PassRelevance = View.DynamicMeshElementsPassRelevance[ElementIndex];
+
+					ComputeDynamicMeshRelevance(ShadingPath, bAddLightmapDensityCommands, ViewRelevance, MeshBatch, View, PassRelevance);
 				}
 			}
 		}
diff --git a/Engine/Source/Runtime/Renderer/Private/ShadowSetup.cpp b/Engine/Source/Runtime/Renderer/Private/ShadowSetup.cpp
--- a/Engine/Source/Runtime/Renderer/Private/ShadowSetup.cpp
+++ b/Engine/Source/Runtime/Renderer/Private/ShadowSetup.cpp
@@ -1181,6 +1181,7 @@ void FProjectedShadowInfo::SetupMeshDrawCommandsForShadowDepth(FSceneRenderer& R
 		FExclusiveDepthStencil::DepthNop_StencilNop,
 		MeshPassProcessor,
 		DynamicSubjectMeshElements,
+		nullptr,
 		NumDynamicSubjectMeshElements * InstanceFactor,
 		SubjectMeshCommandBuildRequests,
 		NumSubjectMeshCommandBuildRequestElements * InstanceFactor,
