Compilation hacks around mobile rendering

diff --git a/Engine/Source/Runtime/Renderer/Private/MobileBasePass.cpp b/Engine/Source/Runtime/Renderer/Private/MobileBasePass.cpp
--- a/Engine/Source/Runtime/Renderer/Private/MobileBasePass.cpp
+++ b/Engine/Source/Runtime/Renderer/Private/MobileBasePass.cpp
@@ -524,11 +524,14 @@ void TMobileBasePassPSPolicyParamType<FUniformLightMapPolicy>::GetShaderBindings
 
 		if (LightPositionAndInvRadiusParameter.IsBound())
 		{
+			//temp compile fixes.
+#if 0
 			// Set dynamic point lights
 			FMobileBasePassMovablePointLightInfo LightInfo(PrimitiveSceneProxy);
 			ShaderBindings.Add(NumDynamicPointLightsParameter, LightInfo.NumMovablePointLights);
 			ShaderBindings.Add(LightPositionAndInvRadiusParameter, LightInfo.LightPositionAndInvRadius);
 			ShaderBindings.Add(LightColorAndFalloffExponentParameter, LightInfo.LightColorAndFalloffExponent);
+#endif
 		}
 	}
 	else
diff --git a/Engine/Source/Runtime/Renderer/Private/MobileBasePassRendering.h b/Engine/Source/Runtime/Renderer/Private/MobileBasePassRendering.h
--- a/Engine/Source/Runtime/Renderer/Private/MobileBasePassRendering.h
+++ b/Engine/Source/Runtime/Renderer/Private/MobileBasePassRendering.h
@@ -246,8 +246,11 @@ public:
 				
 		CSMDebugHintParams.Bind(Initializer.ParameterMap, TEXT("CSMDebugHint"));
 
+		//temp compile hacks
+#if 0
 		SpotLightAnglesParameter.Bind(Initializer.ParameterMap, TEXT("SpotLightAngles"));
 		SpotLightDirectionParameter.Bind(Initializer.ParameterMap, TEXT("SpotLightDirection"));
+#endif
 	}
 	TMobileBasePassPSPolicyParamType() {}
 
@@ -271,9 +274,12 @@ public:
 		Ar << NumDynamicPointLightsParameter;
 	
 		Ar << CSMDebugHintParams;
-		
+	
+		//temp compile hacks
+#if 0
 		Ar << SpotLightAnglesParameter;
 		Ar << SpotLightDirectionParameter;
+#endif
 
 		return bShaderHasOutdatedParameters;
 	}
