resolve merge conflicts for mobile spot lights

diff --git a/Engine/Source/Runtime/Renderer/Private/MobileBasePass.cpp b/Engine/Source/Runtime/Renderer/Private/MobileBasePass.cpp
--- a/Engine/Source/Runtime/Renderer/Private/MobileBasePass.cpp
+++ b/Engine/Source/Runtime/Renderer/Private/MobileBasePass.cpp
@@ -383,8 +383,6 @@ void MobileBasePass::SetOpaqueRenderState(FMeshPassProcessorRenderState& DrawRen
 void MobileBasePass::SetTranslucentRenderState(FMeshPassProcessorRenderState& DrawRenderState, const FMaterial& Material)
 {
 	bool bEncodedHDR = GetMobileHDRMode() == EMobileHDRMode::EnabledRGBE && Material.GetMaterialDomain() != MD_UI;
-	static const auto CVarMonoscopicFarField = IConsoleManager::Get().FindTConsoleVariableDataInt(TEXT("vr.MonoscopicFarField"));
-	const bool bIsMobileMonoscopic = CVarMonoscopicFarField && (CVarMonoscopicFarField->GetValueOnRenderThread() != 0);
 
 	if (bEncodedHDR == false)
 	{
@@ -395,10 +393,6 @@ void MobileBasePass::SetTranslucentRenderState(FMeshPassProcessorRenderState& Dr
 			{
 				DrawRenderState.SetBlendState(TStaticBlendState<CW_ALPHA, BO_Add, BF_Zero, BF_Zero, BO_Add, BF_One, BF_Zero>::GetRHI());
 			} 
-			else if (bIsMobileMonoscopic)
-			{
-				DrawRenderState.SetBlendState(TStaticBlendState<CW_RGBA, BO_Add, BF_SourceAlpha, BF_InverseSourceAlpha, BO_Add, BF_One, BF_One>::GetRHI());
-			}
 			else
 			{
 				DrawRenderState.SetBlendState(TStaticBlendState<CW_RGB, BO_Add, BF_SourceAlpha, BF_InverseSourceAlpha, BO_Add, BF_Zero, BF_InverseSourceAlpha>::GetRHI());
@@ -522,16 +516,15 @@ void TMobileBasePassPSPolicyParamType<FUniformLightMapPolicy>::GetShaderBindings
 		ShaderBindings.Add(ReflectionPositionsAndRadii, CapturePositions);
 		ShaderBindings.Add(MobileReflectionParams, ReflectionParams);
 
-		if (LightPositionAndInvRadiusParameter.IsBound())
+		if (LightPositionAndInvRadiusParameter.IsBound() || SpotLightDirectionParameter.IsBound())
 		{
-			//temp compile fixes.
-#if 0
 			// Set dynamic point lights
-			FMobileBasePassMovablePointLightInfo LightInfo(PrimitiveSceneProxy);
+			FMobileBasePassMovableLightInfo LightInfo(PrimitiveSceneProxy);
 			ShaderBindings.Add(NumDynamicPointLightsParameter, LightInfo.NumMovablePointLights);
 			ShaderBindings.Add(LightPositionAndInvRadiusParameter, LightInfo.LightPositionAndInvRadius);
 			ShaderBindings.Add(LightColorAndFalloffExponentParameter, LightInfo.LightColorAndFalloffExponent);
-#endif
+			ShaderBindings.Add(SpotLightDirectionParameter, LightInfo.SpotLightDirection);
+			ShaderBindings.Add(SpotLightAnglesParameter, LightInfo.SpotLightAngles);
 		}
 	}
 	else
diff --git a/Engine/Source/Runtime/Renderer/Private/MobileBasePassRendering.cpp b/Engine/Source/Runtime/Renderer/Private/MobileBasePassRendering.cpp
--- a/Engine/Source/Runtime/Renderer/Private/MobileBasePassRendering.cpp
+++ b/Engine/Source/Runtime/Renderer/Private/MobileBasePassRendering.cpp
@@ -129,6 +129,14 @@ FMobileBasePassMovableLightInfo::FMobileBasePassMovableLightInfo(const FPrimitiv
 
 				LightPositionAndInvRadius[NumMovablePointLights] = FVector4(LightParameters.Position, LightParameters.InvRadius);
 				LightColorAndFalloffExponent[NumMovablePointLights] = FVector4(LightParameters.Color, LightParameters.FalloffExponent);
+				SpotLightDirection[NumMovablePointLights] = LightParameters.Direction;
+				SpotLightAngles[NumMovablePointLights].Set(LightParameters.SpotAngles.X, LightParameters.SpotAngles.Y, 0.f, LightType == LightType_Spot ? 1.0f : 0.0f);
+
+				if (LightType == LightType_Rect)
+				{
+					// Treat rect lights as point lights.
+					LightColorAndFalloffExponent[NumMovablePointLights] = FVector4(LightParameters.Color, LightParameters.FalloffExponent);
+				}
 
 				if (LightProxy->IsInverseSquared())
 				{
diff --git a/Engine/Source/Runtime/Renderer/Private/MobileBasePassRendering.h b/Engine/Source/Runtime/Renderer/Private/MobileBasePassRendering.h
--- a/Engine/Source/Runtime/Renderer/Private/MobileBasePassRendering.h
+++ b/Engine/Source/Runtime/Renderer/Private/MobileBasePassRendering.h
@@ -243,14 +243,10 @@ public:
 		LightPositionAndInvRadiusParameter.Bind(Initializer.ParameterMap, TEXT("LightPositionAndInvRadius"));
 		LightColorAndFalloffExponentParameter.Bind(Initializer.ParameterMap, TEXT("LightColorAndFalloffExponent"));
 		NumDynamicPointLightsParameter.Bind(Initializer.ParameterMap, TEXT("NumDynamicPointLights"));
-				
-		CSMDebugHintParams.Bind(Initializer.ParameterMap, TEXT("CSMDebugHint"));
-
-		//temp compile hacks
-#if 0
 		SpotLightAnglesParameter.Bind(Initializer.ParameterMap, TEXT("SpotLightAngles"));
 		SpotLightDirectionParameter.Bind(Initializer.ParameterMap, TEXT("SpotLightDirection"));
-#endif
+						
+		CSMDebugHintParams.Bind(Initializer.ParameterMap, TEXT("CSMDebugHint"));
 	}
 	TMobileBasePassPSPolicyParamType() {}
 
@@ -272,14 +268,10 @@ public:
 		Ar << LightPositionAndInvRadiusParameter;
 		Ar << LightColorAndFalloffExponentParameter;
 		Ar << NumDynamicPointLightsParameter;
-	
-		Ar << CSMDebugHintParams;
-	
-		//temp compile hacks
-#if 0
 		Ar << SpotLightAnglesParameter;
 		Ar << SpotLightDirectionParameter;
-#endif
+			
+		Ar << CSMDebugHintParams;
 
 		return bShaderHasOutdatedParameters;
 	}
@@ -299,6 +291,8 @@ private:
 	FShaderParameter LightPositionAndInvRadiusParameter;
 	FShaderParameter LightColorAndFalloffExponentParameter;
 	FShaderParameter NumDynamicPointLightsParameter;
+	FShaderParameter SpotLightDirectionParameter;
+	FShaderParameter SpotLightAnglesParameter;
 
 	FShaderParameter CSMDebugHintParams;
 
