Restore 'background" option for mobile base pass sorting

diff --git a/Engine/Source/Runtime/Renderer/Private/MeshDrawCommands.cpp b/Engine/Source/Runtime/Renderer/Private/MeshDrawCommands.cpp
--- a/Engine/Source/Runtime/Renderer/Private/MeshDrawCommands.cpp
+++ b/Engine/Source/Runtime/Renderer/Private/MeshDrawCommands.cpp
@@ -283,7 +283,8 @@ void UpdateMobileBasePassMeshSortKeys(
 		{
 			const FPrimitiveBounds& PrimitiveBounds = ScenePrimitiveBounds[Cmd.DrawPrimitiveId];
 			PrimitiveDistance = (PrimitiveBounds.BoxSphereBounds.Origin - ViewOrigin).Size();
-			bBackground = (PrimitiveBounds.BoxSphereBounds.SphereRadius > HALF_WORLD_MAX / 4.0f); // TODO: per-primitive bBackgound flag
+			bBackground = Cmd.SortKey.PackedData & 0x2 ? true : false; // Set in MobileBasePass.cpp - GetBasePassStaticSortKey
+			bBackground|= (PrimitiveBounds.BoxSphereBounds.SphereRadius > HALF_WORLD_MAX / 4.0f);
 			bMasked = Cmd.SortKey.PackedData & 0x1 ? true : false; // Set in MobileBasePass.cpp - GetBasePassStaticSortKey
 		}
 
diff --git a/Engine/Source/Runtime/Renderer/Private/MobileBasePass.cpp b/Engine/Source/Runtime/Renderer/Private/MobileBasePass.cpp
--- a/Engine/Source/Runtime/Renderer/Private/MobileBasePass.cpp
+++ b/Engine/Source/Runtime/Renderer/Private/MobileBasePass.cpp
@@ -425,10 +425,11 @@ void MobileBasePass::SetTranslucentRenderState(FMeshPassProcessorRenderState& Dr
 	}
 }
 
-static FMeshDrawCommandSortKey GetBasePassStaticSortKey(EBlendMode BlendMode)
+static FMeshDrawCommandSortKey GetBasePassStaticSortKey(EBlendMode BlendMode, bool bBackground)
 {
 	FMeshDrawCommandSortKey SortKey;
 	SortKey.PackedData = (BlendMode == EBlendMode::BLEND_Masked ? 1 : 0);
+	SortKey.PackedData|= (bBackground ? 2 : 0); // background flag in second bit
 	return SortKey;
 }
 
@@ -645,9 +646,11 @@ void FMobileBasePassMeshProcessor::Process(
 	}
 	else
 	{
+		// Background primitives will be rendered last in masked/non-masked buckets
+		bool bBackground = PrimitiveSceneProxy->TreatAsBackgroundForOcclusion();
 		// Default static sort key separates masked and non-masked geometry, generic mesh sorting will also sort by PSO
 		// if platform wants front to back sorting, this key will be recomputed in InitViews
-		SortKey = GetBasePassStaticSortKey(BlendMode);
+		SortKey = GetBasePassStaticSortKey(BlendMode, bBackground);
 	}
 	
 	ERasterizerFillMode MeshFillMode = ComputeMeshFillMode(MeshBatch, MaterialResource);
