Disabled dynamic instancing for editor preview levels which don't support it.

diff --git a/Engine/Source/Runtime/Renderer/Private/MeshDrawCommands.cpp b/Engine/Source/Runtime/Renderer/Private/MeshDrawCommands.cpp
--- a/Engine/Source/Runtime/Renderer/Private/MeshDrawCommands.cpp
+++ b/Engine/Source/Runtime/Renderer/Private/MeshDrawCommands.cpp
@@ -302,6 +302,7 @@ void UpdateMobileBasePassMeshSortKeys(
  * TempVisibleMeshDrawCommands must be presized for NewPassVisibleMeshDrawCommands.
  */
 void BuildMeshDrawCommandPrimitiveIdBuffer(
+	bool bDynamicInstancing,
 	FMeshCommandOneFrameArray& VisibleMeshDrawCommands,
 	FDynamicMeshDrawCommandStorage& MeshDrawCommandStorage,
 	void* RESTRICT PrimitiveIdData,
@@ -323,7 +324,7 @@ void BuildMeshDrawCommandPrimitiveIdBuffer(
 	int32* RESTRICT PrimitiveIds = (int32*)PrimitiveIdData;
 	const uint32 MaxPrimitiveId = PrimitiveIdDataSize / sizeof(int32);
 
-	if (IsDynamicInstancingEnabled())
+	if (bDynamicInstancing)
 	{
 		QUICK_SCOPE_CYCLE_COUNTER(STAT_DynamicInstancingOfVisibleMeshDrawCommands);
 		check(VisibleMeshDrawCommands.Num() <= TempVisibleMeshDrawCommands.Max() && TempVisibleMeshDrawCommands.Num() == 0);
@@ -772,6 +773,7 @@ public:
 			if (Context.bUseGPUScene)
 			{
 				BuildMeshDrawCommandPrimitiveIdBuffer(
+					Context.bDynamicInstancing,
 					Context.MeshDrawCommands,
 					Context.MeshDrawCommandStorage,
 					Context.PrimitiveIdBufferData,
@@ -819,7 +821,8 @@ void SortAndMergeDynamicPassMeshDrawCommands(
 
 		if (bUseGPUScene)
 		{
-			if (IsDynamicInstancingEnabled())
+			const bool bDynamicInstancing = IsDynamicInstancingEnabled(FeatureLevel);
+			if (bDynamicInstancing)
 			{
 				NewPassVisibleMeshDrawCommands.Empty(NumDrawCommands);
 			}
@@ -831,6 +834,7 @@ void SortAndMergeDynamicPassMeshDrawCommands(
 			void* PrimitiveIdBufferData = RHILockVertexBuffer(OutPrimitiveIdVertexBuffer, 0, PrimitiveIdBufferDataSize, RLM_WriteOnly);
 
 			BuildMeshDrawCommandPrimitiveIdBuffer(
+				bDynamicInstancing,
 				VisibleMeshDrawCommands,
 				MeshDrawCommandStorage,
 				PrimitiveIdBufferData,
@@ -882,7 +886,7 @@ void FParallelMeshDrawCommandPass::DispatchPassSetup(
 	TaskContext.ShadingPath = Scene->GetShadingPath();
 	TaskContext.PassType = PassType;
 	TaskContext.bUseGPUScene = UseGPUScene(GMaxRHIShaderPlatform, View.GetFeatureLevel());
-	TaskContext.bDynamicInstancing = TaskContext.bUseGPUScene && IsDynamicInstancingEnabled();
+	TaskContext.bDynamicInstancing = IsDynamicInstancingEnabled(View.GetFeatureLevel());
 	TaskContext.bReverseCulling = View.bReverseCulling;
 	TaskContext.bRenderSceneTwoSided = View.bRenderSceneTwoSided;
 	TaskContext.BasePassDepthStencilAccess = BasePassDepthStencilAccess;
diff --git a/Engine/Source/Runtime/Renderer/Private/MeshPassProcessor.cpp b/Engine/Source/Runtime/Renderer/Private/MeshPassProcessor.cpp
--- a/Engine/Source/Runtime/Renderer/Private/MeshPassProcessor.cpp
+++ b/Engine/Source/Runtime/Renderer/Private/MeshPassProcessor.cpp
@@ -840,9 +840,9 @@ void SubmitMeshDrawCommands(
 	const FMeshCommandOneFrameArray& VisibleMeshDrawCommands,
 	FVertexBufferRHIParamRef PrimitiveIdsBuffer,
 	int32 BasePrimitiveIdsOffset,
+	bool bDynamicInstancing,
 	FRHICommandList& RHICmdList)
 {
-	const bool bDynamicInstancing = IsDynamicInstancingEnabled();
 	SubmitMeshDrawCommandsRange(VisibleMeshDrawCommands, PrimitiveIdsBuffer, BasePrimitiveIdsOffset, bDynamicInstancing, 0, VisibleMeshDrawCommands.Num(), RHICmdList);
 }
 
@@ -877,11 +877,13 @@ void DrawDynamicMeshPassPrivate(
 {
 	if (VisibleMeshDrawCommands.Num() > 0)
 	{
+		const bool bDynamicInstancing = IsDynamicInstancingEnabled(View.GetFeatureLevel());
+
 		FVertexBufferRHIParamRef PrimitiveIdVertexBuffer = nullptr;
 
 		SortAndMergeDynamicPassMeshDrawCommands(View.GetFeatureLevel(), VisibleMeshDrawCommands, DynamicMeshDrawCommandStorage, PrimitiveIdVertexBuffer);
 
-		SubmitMeshDrawCommandsRange(VisibleMeshDrawCommands, PrimitiveIdVertexBuffer, 0, IsDynamicInstancingEnabled(), 0, VisibleMeshDrawCommands.Num(), RHICmdList);
+		SubmitMeshDrawCommandsRange(VisibleMeshDrawCommands, PrimitiveIdVertexBuffer, 0, bDynamicInstancing, 0, VisibleMeshDrawCommands.Num(), RHICmdList);
 	}
 }
 
diff --git a/Engine/Source/Runtime/Renderer/Private/SceneRendering.cpp b/Engine/Source/Runtime/Renderer/Private/SceneRendering.cpp
--- a/Engine/Source/Runtime/Renderer/Private/SceneRendering.cpp
+++ b/Engine/Source/Runtime/Renderer/Private/SceneRendering.cpp
@@ -71,10 +71,10 @@ static TAutoConsoleVariable<int32> CVarMeshDrawCommandsDynamicInstancing(
 	TEXT("Whether to dynamically combine multiple compatible visible Mesh Draw Commands into one instanced draw on vertex factories that support it."),
 	ECVF_RenderThreadSafe);
 
-bool IsDynamicInstancingEnabled()
+bool IsDynamicInstancingEnabled(ERHIFeatureLevel::Type FeatureLevel)
 {
 	return CVarMeshDrawCommandsDynamicInstancing.GetValueOnRenderThread() > 0
-		&& UseGPUScene(GMaxRHIShaderPlatform, GMaxRHIFeatureLevel); // always disabled on platforms that do not support auto-instacing
+		&& UseGPUScene(GMaxRHIShaderPlatform, FeatureLevel);
 }
 
 int32 GDumpInstancingStats = 0;
diff --git a/Engine/Source/Runtime/Renderer/Private/SceneRendering.h b/Engine/Source/Runtime/Renderer/Private/SceneRendering.h
--- a/Engine/Source/Runtime/Renderer/Private/SceneRendering.h
+++ b/Engine/Source/Runtime/Renderer/Private/SceneRendering.h
@@ -1828,4 +1828,4 @@ private:
 extern FFastVramConfig GFastVRamConfig;
 
 extern bool UseCachedMeshDrawCommands();
-extern bool IsDynamicInstancingEnabled();
\ No newline at end of file
+extern bool IsDynamicInstancingEnabled(ERHIFeatureLevel::Type FeatureLevel);
\ No newline at end of file
diff --git a/Engine/Source/Runtime/Renderer/Private/ShadowRendering.cpp b/Engine/Source/Runtime/Renderer/Private/ShadowRendering.cpp
--- a/Engine/Source/Runtime/Renderer/Private/ShadowRendering.cpp
+++ b/Engine/Source/Runtime/Renderer/Private/ShadowRendering.cpp
@@ -597,6 +597,8 @@ void FProjectedShadowInfo::SetupProjectionStencilMask(
 	DrawRenderState.SetDepthStencilState(TStaticDepthStencilState<false, CF_DepthNearOrEqual>::GetRHI());
 	DrawRenderState.SetBlendState(TStaticBlendState<CW_NONE>::GetRHI());
 
+	const bool bDynamicInstancing = IsDynamicInstancingEnabled(View->FeatureLevel);
+
 	// If this is a preshadow, mask the projection by the receiver primitives.
 	if (bPreShadow || bSelfShadowOnly)
 	{
@@ -612,7 +614,7 @@ void FProjectedShadowInfo::SetupProjectionStencilMask(
 		const FShadowMeshDrawCommandPass& ProjectionStencilingPass = ProjectionStencilingPasses[ViewIndex];
 		if (ProjectionStencilingPass.VisibleMeshDrawCommands.Num() > 0)
 		{
-			SubmitMeshDrawCommands(ProjectionStencilingPass.VisibleMeshDrawCommands, ProjectionStencilingPass.PrimitiveIdVertexBuffer, 0, RHICmdList);
+			SubmitMeshDrawCommands(ProjectionStencilingPass.VisibleMeshDrawCommands, ProjectionStencilingPass.PrimitiveIdVertexBuffer, 0, bDynamicInstancing, RHICmdList);
 		}
 
 		// Restore viewport
@@ -752,7 +754,7 @@ void FProjectedShadowInfo::SetupProjectionStencilMask(
 			const FShadowMeshDrawCommandPass& ProjectionStencilingPass = ProjectionStencilingPasses[ViewIndex];
 			if (ProjectionStencilingPass.VisibleMeshDrawCommands.Num() > 0)
 			{
-				SubmitMeshDrawCommands(ProjectionStencilingPass.VisibleMeshDrawCommands, ProjectionStencilingPass.PrimitiveIdVertexBuffer, 0, RHICmdList);
+				SubmitMeshDrawCommands(ProjectionStencilingPass.VisibleMeshDrawCommands, ProjectionStencilingPass.PrimitiveIdVertexBuffer, 0, bDynamicInstancing, RHICmdList);
 			}
 		}
 	}
diff --git a/Engine/Source/Runtime/Renderer/Private/TranslucentLighting.cpp b/Engine/Source/Runtime/Renderer/Private/TranslucentLighting.cpp
--- a/Engine/Source/Runtime/Renderer/Private/TranslucentLighting.cpp
+++ b/Engine/Source/Runtime/Renderer/Private/TranslucentLighting.cpp
@@ -620,10 +620,12 @@ void FProjectedShadowInfo::RenderTranslucencyDepths(FRHICommandList& RHICmdList,
 
 		if (VisibleMeshDrawCommands.Num() > 0)
 		{
+			const bool bDynamicInstancing = IsDynamicInstancingEnabled(ShadowDepthView->FeatureLevel);
+
 			FVertexBufferRHIParamRef PrimitiveIdVertexBuffer = nullptr;
 			ApplyViewOverridesToMeshDrawCommands(*ShadowDepthView, VisibleMeshDrawCommands);
 			SortAndMergeDynamicPassMeshDrawCommands(SceneRenderer->FeatureLevel, VisibleMeshDrawCommands, DynamicMeshDrawCommandStorage, PrimitiveIdVertexBuffer);
-			SubmitMeshDrawCommands(VisibleMeshDrawCommands, PrimitiveIdVertexBuffer, 0, RHICmdList);
+			SubmitMeshDrawCommands(VisibleMeshDrawCommands, PrimitiveIdVertexBuffer, 0, bDynamicInstancing, RHICmdList);
 		}
 	}
 }
diff --git a/Engine/Source/Runtime/Renderer/Public/MeshPassProcessor.h b/Engine/Source/Runtime/Renderer/Public/MeshPassProcessor.h
--- a/Engine/Source/Runtime/Renderer/Public/MeshPassProcessor.h
+++ b/Engine/Source/Runtime/Renderer/Public/MeshPassProcessor.h
@@ -988,6 +988,7 @@ extern void SubmitMeshDrawCommands(
 	const FMeshCommandOneFrameArray& VisibleMeshDrawCommands,
 	FVertexBufferRHIParamRef PrimitiveIdsBuffer,
 	int32 BasePrimitiveIdsOffset,
+	bool bDynamicInstancing,
 	FRHICommandList& RHICmdList);
 
 extern void SubmitMeshDrawCommandsRange(
