Need to use the mobile mesh processor in the mobile context.

diff --git a/Engine/Source/Runtime/Renderer/Private/Renderer.cpp b/Engine/Source/Runtime/Renderer/Private/Renderer.cpp
--- a/Engine/Source/Runtime/Renderer/Private/Renderer.cpp
+++ b/Engine/Source/Runtime/Renderer/Private/Renderer.cpp
@@ -199,30 +199,50 @@ void FRendererModule::DrawTileMesh(FRHICommandListImmediate& RHICmdList, FMeshPa
 				TUniformBufferRef<FTranslucentBasePassUniformParameters> TranslucentBasePassUniformBuffer;
 				CreateTranslucentBasePassUniformBuffer(RHICmdList, View, nullptr, ESceneTextureSetupMode::None, TranslucentBasePassUniformBuffer, 0);
 				BasePassUniformBuffer = TranslucentBasePassUniformBuffer;
+				
+				DrawRenderState.SetPassUniformBuffer(BasePassUniformBuffer);
+				
+				DrawDynamicMeshPass(View, RHICmdList,
+					[&View, &DrawRenderState, &Mesh](FDynamicPassMeshDrawListContext* DynamicMeshPassContext)
+				{
+					FBasePassMeshProcessor PassMeshProcessor(
+															 nullptr,
+															 View.GetFeatureLevel(),
+															 &View,
+															 DrawRenderState,
+															 DynamicMeshPassContext,
+															 ETranslucencyPass::TPT_AllTranslucency);
+					
+					const uint64 DefaultBatchElementMask = ~0ull;
+					PassMeshProcessor.AddMeshBatch(Mesh, DefaultBatchElementMask, nullptr);
+				});
 			}
 			else // Mobile
 			{
 				TUniformBufferRef<FMobileBasePassUniformParameters> MobileBasePassUniformBuffer;
 				CreateMobileBasePassUniformBuffer(RHICmdList, View, true, MobileBasePassUniformBuffer);
 				BasePassUniformBuffer = MobileBasePassUniformBuffer;
-			}
-
-			DrawRenderState.SetPassUniformBuffer(BasePassUniformBuffer);
-			
-			DrawDynamicMeshPass(View, RHICmdList,
-				[&View, &DrawRenderState, &Mesh](FDynamicPassMeshDrawListContext* DynamicMeshPassContext)
+				
+				DrawRenderState.SetPassUniformBuffer(BasePassUniformBuffer);
+				
+				DrawDynamicMeshPass(View, RHICmdList,
+					[&View, &DrawRenderState, &Mesh](FDynamicPassMeshDrawListContext* DynamicMeshPassContext)
 				{
-					FBasePassMeshProcessor PassMeshProcessor(
-						nullptr,
-						View.GetFeatureLevel(),
-						&View,
-						DrawRenderState,
-						DynamicMeshPassContext,
-						ETranslucencyPass::TPT_AllTranslucency);
-
+					FMobileBasePassMeshProcessor PassMeshProcessor(
+															 nullptr,
+															 View.GetFeatureLevel(),
+															 &View,
+															 DrawRenderState,
+															 DynamicMeshPassContext,
+															 false,
+															 ETranslucencyPass::TPT_AllTranslucency); // No CSM?
+					
 					const uint64 DefaultBatchElementMask = ~0ull;
 					PassMeshProcessor.AddMeshBatch(Mesh, DefaultBatchElementMask, nullptr);
 				});
+			}
+
+			
 		}
 		// handle opaque materials
 		else
@@ -260,29 +280,47 @@ void FRendererModule::DrawTileMesh(FRHICommandListImmediate& RHICmdList, FMeshPa
 					TUniformBufferRef<FOpaqueBasePassUniformParameters> OpaqueBasePassUniformBuffer;
 					CreateOpaqueBasePassUniformBuffer(RHICmdList, View, nullptr, OpaqueBasePassUniformBuffer);
 					BasePassUniformBuffer = OpaqueBasePassUniformBuffer;
+					
+					DrawRenderState.SetPassUniformBuffer(BasePassUniformBuffer);
+					
+					DrawDynamicMeshPass(View, RHICmdList,
+						[&View, &DrawRenderState, &Mesh](FDynamicPassMeshDrawListContext* DynamicMeshPassContext)
+					{
+						FBasePassMeshProcessor PassMeshProcessor(
+																 nullptr,
+																 View.GetFeatureLevel(),
+																 &View,
+																 DrawRenderState,
+																 DynamicMeshPassContext);
+						
+						const uint64 DefaultBatchElementMask = ~0ull;
+						PassMeshProcessor.AddMeshBatch(Mesh, DefaultBatchElementMask, nullptr);
+					});
 				}
 				else // Mobile
 				{
 					TUniformBufferRef<FMobileBasePassUniformParameters> MobileBasePassUniformBuffer;
 					CreateMobileBasePassUniformBuffer(RHICmdList, View, false, MobileBasePassUniformBuffer);
 					BasePassUniformBuffer = MobileBasePassUniformBuffer;
-				}
-
-				DrawRenderState.SetPassUniformBuffer(BasePassUniformBuffer);
-
-				DrawDynamicMeshPass(View, RHICmdList,
-					[&View, &DrawRenderState, &Mesh](FDynamicPassMeshDrawListContext* DynamicMeshPassContext)
+					
+					DrawRenderState.SetPassUniformBuffer(BasePassUniformBuffer);
+					
+					DrawDynamicMeshPass(View, RHICmdList,
+						[&View, &DrawRenderState, &Mesh](FDynamicPassMeshDrawListContext* DynamicMeshPassContext)
 					{
-						FBasePassMeshProcessor PassMeshProcessor(
-							nullptr,
-							View.GetFeatureLevel(),
-							&View,
-							DrawRenderState,
-							DynamicMeshPassContext);
-
+						FMobileBasePassMeshProcessor PassMeshProcessor(
+																 nullptr,
+																 View.GetFeatureLevel(),
+																 &View,
+																 DrawRenderState,
+																 DynamicMeshPassContext, true);
+						
 						const uint64 DefaultBatchElementMask = ~0ull;
 						PassMeshProcessor.AddMeshBatch(Mesh, DefaultBatchElementMask, nullptr);
 					});
+				}
+
+				
 			}
 		}
 
