Fixed static analysis warning.

diff --git a/Engine/Source/Runtime/Renderer/Public/MeshDrawShaderBindings.h b/Engine/Source/Runtime/Renderer/Public/MeshDrawShaderBindings.h
--- a/Engine/Source/Runtime/Renderer/Public/MeshDrawShaderBindings.h
+++ b/Engine/Source/Runtime/Renderer/Public/MeshDrawShaderBindings.h
@@ -20,7 +20,6 @@ public:
 	FMeshDrawShaderBindingsLayout(const FShader* Shader) :
 		ParameterMapInfo(Shader->GetParameterMapInfo())
 	{
-		check(Shader);
 		Frequency = (EShaderFrequency)Shader->GetTarget().Frequency;
 		checkSlow((EShaderFrequency)Frequency == (EShaderFrequency)Shader->GetTarget().Frequency);
 	}
