Fixed reflection captures with forward shading. Primitive reflection capture cache is updated only for the visible ones, so it wasn't enough just to rebuild entire GPUScene and all primitive uniform buffers when registered reflection captures were modifed. Now every primitive's reflection capture cache update also triggers it's uniform buffer/entry in GPUScene update.

diff --git a/Engine/Source/Runtime/Renderer/Private/SceneVisibility.cpp b/Engine/Source/Runtime/Renderer/Private/SceneVisibility.cpp
--- a/Engine/Source/Runtime/Renderer/Private/SceneVisibility.cpp
+++ b/Engine/Source/Runtime/Renderer/Private/SceneVisibility.cpp
@@ -1754,6 +1754,7 @@ struct FRelevancePacket
 	bool bHasCustomDepthPrimitives;
 	FRelevancePrimSet<FPrimitiveSceneInfo*> LazyUpdatePrimitives;
 	FRelevancePrimSet<FPrimitiveSceneInfo*> DirtyIndirectLightingCacheBufferPrimitives;
+	FRelevancePrimSet<FPrimitiveSceneInfo*> RecachedReflectionCapturePrimitives;
 	FRelevancePrimSet<FPrimitiveSceneProxy*> VolumetricPrimSet;
 	FDrawCommandRelevancePacket DrawCommandPacket;
 
@@ -1990,6 +1991,13 @@ struct FRelevancePacket
 				ensure(Scene->GetShadingPath() != EShadingPath::Mobile);
 		
 				PrimitiveSceneInfo->CacheReflectionCaptures();
+
+				// With forward shading we need to track reflection capture cache updates
+				// in order to update primitive's uniform buffer's closest reflection capture id.
+				if (IsForwardShadingEnabled(Scene->GetShaderPlatform()))
+				{
+					RecachedReflectionCapturePrimitives.AddPrim(PrimitiveSceneInfo);
+				}
 			}
 
 			if (PrimitiveSceneInfo->NeedsUniformBufferUpdate())
@@ -2263,6 +2271,17 @@ struct FRelevancePacket
 		DirtyIndirectLightingCacheBufferPrimitives.AppendTo(WriteView.DirtyIndirectLightingCacheBufferPrimitives);
 		VolumetricPrimSet.AppendTo(WriteView.VolumetricPrimSet);
 
+		for (int32 Index = 0; Index < RecachedReflectionCapturePrimitives.NumPrims; ++Index)
+		{
+			FPrimitiveSceneInfo* PrimitiveSceneInfo = RecachedReflectionCapturePrimitives.Prims[Index];
+
+			PrimitiveSceneInfo->SetNeedsUniformBufferUpdate(true);
+			PrimitiveSceneInfo->ConditionalUpdateUniformBuffer(RHICmdList);
+
+			FScene& WriteScene = *const_cast<FScene*>(Scene);
+			AddPrimitiveToUpdateGPU(WriteScene, PrimitiveSceneInfo->GetIndex());
+		}
+
 		for (int32 Index = 0; Index < LazyUpdatePrimitives.NumPrims; Index++)
 		{
 			LazyUpdatePrimitives.Prims[Index]->ConditionalUpdateUniformBuffer(RHICmdList);
@@ -3197,21 +3216,17 @@ void UpdateReflectionSceneData(FScene* Scene)
 
 	// If SortedCaptures change, then in case of forward renderer all scene primitives need to be updated, as they 
 	// store index into sorted reflection capture uniform buffer for the forward renderer.
-	if (IsForwardShadingEnabled(Scene->GetShaderPlatform()))
+	if (IsForwardShadingEnabled(Scene->GetShaderPlatform()) && ReflectionSceneData.AllocatedReflectionCaptureStateHasChanged)
 	{
-		if (ReflectionSceneData.bRegisteredReflectionCapturesHasChanged || ReflectionSceneData.AllocatedReflectionCaptureStateHasChanged)
+		const int32 NumPrimitives = Scene->Primitives.Num();
+		for (int32 PrimitiveIndex = 0; PrimitiveIndex < NumPrimitives; ++PrimitiveIndex)
 		{
-			const int32 NumPrimitives = Scene->Primitives.Num();
-
-			for (int32 PrimitiveIndex = 0; PrimitiveIndex < NumPrimitives; ++PrimitiveIndex)
-			{
-				Scene->Primitives[PrimitiveIndex]->SetNeedsUniformBufferUpdate(true);
-			}
+			Scene->Primitives[PrimitiveIndex]->SetNeedsUniformBufferUpdate(true);
+		}
 
-			Scene->GPUScene.bUpdateAllPrimitives = true;
+		Scene->GPUScene.bUpdateAllPrimitives = true;
 
-			ReflectionSceneData.AllocatedReflectionCaptureStateHasChanged = false;
-		}
+		ReflectionSceneData.AllocatedReflectionCaptureStateHasChanged = false;
 	}
 
 
