Fixed ComputeDynamicMeshRelevance conditions.

diff --git a/Engine/Source/Runtime/Renderer/Private/SceneVisibility.cpp b/Engine/Source/Runtime/Renderer/Private/SceneVisibility.cpp
--- a/Engine/Source/Runtime/Renderer/Private/SceneVisibility.cpp
+++ b/Engine/Source/Runtime/Renderer/Private/SceneVisibility.cpp
@@ -2585,13 +2585,13 @@ void ComputeDynamicMeshRelevance(EShadingPath ShadingPath, bool bAddLightmapDens
 		if (ViewRelevance.bDistortionRelevance)
 		{
 			PassMask.Set(EMeshPass::Distortion);
-			++View.NumVisibleDynamicMeshElements[EMeshPass::Distortion] += NumElements;
+			View.NumVisibleDynamicMeshElements[EMeshPass::Distortion] += NumElements;
 		}
 
 		if (ShadingPath == EShadingPath::Mobile && View.bIsSceneCapture)
 		{
 			PassMask.Set(EMeshPass::MobileInverseOpacity);
-			++View.NumVisibleDynamicMeshElements[EMeshPass::MobileInverseOpacity] += NumElements;
+			View.NumVisibleDynamicMeshElements[EMeshPass::MobileInverseOpacity] += NumElements;
 		}
 	}
 
