Fixed FMeshDrawSingleShaderBindings static analysis warning.

diff --git a/Engine/Source/Runtime/Renderer/Public/MeshDrawShaderBindings.h b/Engine/Source/Runtime/Renderer/Public/MeshDrawShaderBindings.h
--- a/Engine/Source/Runtime/Renderer/Public/MeshDrawShaderBindings.h
+++ b/Engine/Source/Runtime/Renderer/Public/MeshDrawShaderBindings.h
@@ -20,6 +20,7 @@ public:
 	FMeshDrawShaderBindingsLayout(const FShader* Shader) :
 		ParameterMapInfo(Shader->GetParameterMapInfo())
 	{
+		check(Shader);
 		Frequency = (EShaderFrequency)Shader->GetTarget().Frequency;
 		checkSlow((EShaderFrequency)Frequency == (EShaderFrequency)Shader->GetTarget().Frequency);
 	}
diff --git a/Engine/Source/Runtime/Renderer/Public/MeshPassProcessor.h b/Engine/Source/Runtime/Renderer/Public/MeshPassProcessor.h
--- a/Engine/Source/Runtime/Renderer/Public/MeshPassProcessor.h
+++ b/Engine/Source/Runtime/Renderer/Public/MeshPassProcessor.h
@@ -228,7 +228,8 @@ public:
 		}
 
 		checkf(0, TEXT("Invalid shader binding frequency requested"));
-		return FMeshDrawSingleShaderBindings(FMeshDrawShaderBindingsLayout(nullptr), nullptr);
+		FShader Shader;
+		return FMeshDrawSingleShaderBindings(FMeshDrawShaderBindingsLayout(&Shader), nullptr);
 	}
 
 	/** Set shader bindings on the commandlist, filtered by state cache. */
