Removed unnecessary flush inside cascade shadow map rendering.

diff --git a/Engine/Source/Runtime/Renderer/Private/ShadowDepthRendering.cpp b/Engine/Source/Runtime/Renderer/Private/ShadowDepthRendering.cpp
--- a/Engine/Source/Runtime/Renderer/Private/ShadowDepthRendering.cpp
+++ b/Engine/Source/Runtime/Renderer/Private/ShadowDepthRendering.cpp
@@ -1127,9 +1127,7 @@ void FProjectedShadowInfo::RenderDepthInner(FRHICommandListImmediate& RHICmdList
 
 		// parallel version
 		bool bFlush = CVarRHICmdFlushRenderThreadTasksShadowPass.GetValueOnRenderThread() > 0
-			|| CVarRHICmdFlushRenderThreadTasks.GetValueOnRenderThread() > 0
-			// Scene uniform buffers are reused between cascades
-			|| GetShadowDepthType() == CSMShadowDepthType;
+			|| CVarRHICmdFlushRenderThreadTasks.GetValueOnRenderThread() > 0;
 		FScopedCommandListWaitForTasks Flusher(bFlush);
 
 		// Dispatch commands
