Extended FFXSystemInterface framework to allow custom FX implementations.
NiagaraEmitterInstanceBatcher now derives from FFXSystemInterface and is hooked to the new implementation.
Removed the FComputeDispatcher framework (which was only used for Niagara hooks)

diff --git a/Engine/Plugins/FX/Niagara/Source/Niagara/Classes/NiagaraEmitterInstance.h b/Engine/Plugins/FX/Niagara/Source/Niagara/Classes/NiagaraEmitterInstance.h
--- a/Engine/Plugins/FX/Niagara/Source/Niagara/Classes/NiagaraEmitterInstance.h
+++ b/Engine/Plugins/FX/Niagara/Source/Niagara/Classes/NiagaraEmitterInstance.h
@@ -20,7 +20,7 @@ class NiagaraRenderer;
 struct FNiagaraEmitterHandle;
 class UNiagaraParameterCollection;
 class UNiagaraParameterCollectionInstance;
-
+class NiagaraEmitterInstanceBatcher;
 /**
 * A Niagara particle simulation.
 */
@@ -212,4 +212,6 @@ private:
 
 	/** A parameter store which contains the data interfaces parameters which were defined by the scripts. */
 	FNiagaraParameterStore ScriptDefinedDataInterfaceParameters;
+
+	NiagaraEmitterInstanceBatcher* Batcher = nullptr;
 };
diff --git a/Engine/Plugins/FX/Niagara/Source/Niagara/Classes/NiagaraEmitterInstanceBatcher.h b/Engine/Plugins/FX/Niagara/Source/Niagara/Classes/NiagaraEmitterInstanceBatcher.h
--- a/Engine/Plugins/FX/Niagara/Source/Niagara/Classes/NiagaraEmitterInstanceBatcher.h
+++ b/Engine/Plugins/FX/Niagara/Source/Niagara/Classes/NiagaraEmitterInstanceBatcher.h
@@ -15,57 +15,44 @@ the same VectorVM byte code / compute shader code
 #include "Tickable.h"
 #include "Modules/ModuleManager.h"
 #include "RHIResources.h"
+#include "FXSystem.h"
 
 struct FNiagaraScriptExecutionContext;
 struct FNiagaraComputeExecutionContext;
 
 #define SIMULATION_QUEUE_COUNT 2
 
-class NiagaraEmitterInstanceBatcher : public FTickableGameObject, public FComputeDispatcher
+class NiagaraEmitterInstanceBatcher : public FFXSystemInterface
 {
 public:
 	NiagaraEmitterInstanceBatcher()
 		: CurQueueIndex(0)
 	{
-		IRendererModule *RendererModule = FModuleManager::GetModulePtr<IRendererModule>("Renderer");
-		if (RendererModule)
-		{
-			RendererModule->RegisterPostOpaqueComputeDispatcher(this);
-		}
 	}
 
 	~NiagaraEmitterInstanceBatcher()
 	{
-		IRendererModule *RendererModule = FModuleManager::GetModulePtr<IRendererModule>("Renderer");
-		if (RendererModule)
-		{
-			RendererModule->UnRegisterPostOpaqueComputeDispatcher(this);
-		}
 	}
 
-	static NiagaraEmitterInstanceBatcher *Get()
-	{
-		if (BatcherSingleton == nullptr)
-		{
-			BatcherSingleton = new NiagaraEmitterInstanceBatcher();
-		}
-		return BatcherSingleton;
-	}
+	static const FName Name;
+	virtual FFXSystemInterface* GetInterface(const FName& InName) override;
 
 	void Queue(FNiagaraComputeExecutionContext *InContext);
 
 	void Remove(FNiagaraComputeExecutionContext* InContext);
 
-	virtual ETickableTickType GetTickableTickType() const override
-	{
-		return ETickableTickType::Always;
-	}
-
-	virtual TStatId GetStatId() const override
-	{
-		RETURN_QUICK_DECLARE_CYCLE_STAT(NiagaraEmitterInstanceBatcher, STATGROUP_Tickables);
-	}
+#if WITH_EDITOR
+	virtual void Suspend() override {}
+	virtual void Resume() override {}
+#endif // #if WITH_EDITOR
 
+	virtual void DrawDebug(FCanvas* Canvas) override {}
+	virtual void AddVectorField(UVectorFieldComponent* VectorFieldComponent) override {}
+	virtual void RemoveVectorField(UVectorFieldComponent* VectorFieldComponent) override {}
+	virtual void UpdateVectorField(UVectorFieldComponent* VectorFieldComponent) override {}
+	virtual void PreInitViews() override {}
+	virtual bool UsesGlobalDistanceField() const override { return false; }
+	virtual void PreRender(FRHICommandListImmediate& RHICmdList, const class FGlobalDistanceFieldParameterData* GlobalDistanceFieldParameterData) override {}
 
 	virtual void Tick(float DeltaTime) override
 	{
@@ -82,9 +69,11 @@ public:
 
 	uint32 GetEventSpawnTotal(const FNiagaraComputeExecutionContext *InContext) const;
 
-	// from FComputeDispatcher; called once per frame by the render thread, swaps buffers and works down 
-	// the queue submitted by the game thread; means we're one frame behind with this; we need a mechanism to determine execution order here.
-	virtual void Execute(FRHICommandList &RHICmdList, FUniformBufferRHIParamRef ViewUniformBuffer)
+	virtual void PostRenderOpaque(
+		FRHICommandListImmediate& RHICmdList,
+		const FUniformBufferRHIParamRef ViewUniformBuffer,
+		const class FShaderParametersMetadata* SceneTexturesUniformBufferStruct,
+		FUniformBufferRHIParamRef SceneTexturesUniformBuffer) override
 	{
 		CurQueueIndex ^= 0x1;
 
@@ -118,7 +107,6 @@ public:
 	void ResolveDatasetWrites(FRHICommandList &RHICmdList, FNiagaraComputeExecutionContext *Context) const;
 	void ResizeCurrentBuffer(FRHICommandList &RHICmdList, FNiagaraComputeExecutionContext *Context, uint32 NewNumInstances, uint32 PrevNumInstances) const;
 private:
-	static NiagaraEmitterInstanceBatcher* BatcherSingleton;
 
 	uint32 CurQueueIndex;
 	TArray<FNiagaraComputeExecutionContext*> SimulationQueue[SIMULATION_QUEUE_COUNT];
diff --git a/Engine/Plugins/FX/Niagara/Source/Niagara/Private/NiagaraEmitterInstance.cpp b/Engine/Plugins/FX/Niagara/Source/Niagara/Private/NiagaraEmitterInstance.cpp
--- a/Engine/Plugins/FX/Niagara/Source/Niagara/Private/NiagaraEmitterInstance.cpp
+++ b/Engine/Plugins/FX/Niagara/Source/Niagara/Private/NiagaraEmitterInstance.cpp
@@ -61,6 +61,8 @@ FNiagaraEmitterInstance::FNiagaraEmitterInstance(FNiagaraSystemInstance* InParen
 {
 	bDumpAfterEvent = false;
 	ParticleDataSet = new FNiagaraDataSet();
+
+	Batcher = ParentSystemInstance ? ParentSystemInstance->GetBatcher() : nullptr;
 }
 
 FNiagaraEmitterInstance::~FNiagaraEmitterInstance()
@@ -75,7 +77,10 @@ FNiagaraEmitterInstance::~FNiagaraEmitterInstance()
 		/** We defer the deletion of the particle dataset and the compute context to the RT to be sure all in-flight RT commands have finished using it.*/
 		if (GPUExecContext != nullptr)
 		{
-			NiagaraEmitterInstanceBatcher::Get()->Remove(GPUExecContext);
+			if (Batcher)
+			{
+				Batcher->Remove(GPUExecContext);
+			}
 
 			ENQUEUE_UNIQUE_RENDER_COMMAND_ONEPARAMETER(FDeleteComputeContextCommand,
 				FNiagaraComputeExecutionContext*, ExecContext, GPUExecContext,
@@ -331,8 +336,6 @@ void FNiagaraEmitterInstance::Init(int32 InEmitterIdx, FName InSystemInstanceNam
 
 	if (CachedEmitter->SimTarget == ENiagaraSimTarget::GPUComputeSim && GPUExecContext != nullptr)
 	{
-		//Just ensure we've generated the singleton here on the GT as it throws a wobbler if we do this later in parallel.
-		NiagaraEmitterInstanceBatcher::Get();
 	}
 	else
 	{
@@ -890,14 +893,12 @@ bool FNiagaraEmitterInstance::WaitForDebugInfo()
 {
 	if (CachedEmitter->SimTarget == ENiagaraSimTarget::GPUComputeSim)
 	{
-		
 		FNiagaraComputeExecutionContext* DebugContext = GPUExecContext;
 
-		ENQUEUE_RENDER_COMMAND(CaptureCommand)([DebugContext](FRHICommandListImmediate& RHICmdList)
-			{
-				NiagaraEmitterInstanceBatcher::Get()->ProcessDebugInfo(RHICmdList, DebugContext);
-			}
-		);
+		ENQUEUE_RENDER_COMMAND(CaptureCommand)([=](FRHICommandListImmediate& RHICmdList)
+		{
+			Batcher->ProcessDebugInfo(RHICmdList, GPUExecContext);
+		});
 		return true;
 	}
 	return false;
@@ -1109,7 +1110,10 @@ void FNiagaraEmitterInstance::Tick(float DeltaSeconds)
 		GPUExecContext->EventHandlerScriptProps = CachedEmitter->GetEventHandlers();
 		GPUExecContext->EventSets = EventSet;
 		GPUExecContext->EventSpawnCounts = EventHandlerSpawnCounts;
-		NiagaraEmitterInstanceBatcher::Get()->Queue(GPUExecContext);
+		if (Batcher)
+		{
+			Batcher->Queue(GPUExecContext);
+		}
 
 		// Need to call post-tick, which calls the copy to previous for interpolated spawning
 		SpawnExecContext.PostTick();
diff --git a/Engine/Plugins/FX/Niagara/Source/Niagara/Private/NiagaraEmitterInstanceBatcher.cpp b/Engine/Plugins/FX/Niagara/Source/Niagara/Private/NiagaraEmitterInstanceBatcher.cpp
--- a/Engine/Plugins/FX/Niagara/Source/Niagara/Private/NiagaraEmitterInstanceBatcher.cpp
+++ b/Engine/Plugins/FX/Niagara/Source/Niagara/Private/NiagaraEmitterInstanceBatcher.cpp
@@ -21,9 +21,15 @@ DECLARE_DWORD_COUNTER_STAT(TEXT("Readback latency (frames)"), STAT_NiagaraReadba
 DECLARE_GPU_STAT_NAMED(NiagaraGPUSimulation, TEXT("Niagara GPU Simulation"));
 DECLARE_GPU_STAT_NAMED(NiagaraIndexBufferClear, TEXT("Niagara index buffer clear"));
 
-NiagaraEmitterInstanceBatcher* NiagaraEmitterInstanceBatcher::BatcherSingleton = nullptr;
 uint32 FNiagaraComputeExecutionContext::TickCounter = 0;
 
+const FName NiagaraEmitterInstanceBatcher::Name(TEXT("NiagaraEmitterInstanceBatcher"));
+
+FFXSystemInterface* NiagaraEmitterInstanceBatcher::GetInterface(const FName& InName)
+{
+	return InName == Name ? this : nullptr;
+}
+
 void NiagaraEmitterInstanceBatcher::Queue(FNiagaraComputeExecutionContext *InContext)
 {
 	//UE_LOG(LogNiagara, Warning, TEXT("Submitted!"));
diff --git a/Engine/Plugins/FX/Niagara/Source/Niagara/Private/NiagaraModule.cpp b/Engine/Plugins/FX/Niagara/Source/Niagara/Private/NiagaraModule.cpp
--- a/Engine/Plugins/FX/Niagara/Source/Niagara/Private/NiagaraModule.cpp
+++ b/Engine/Plugins/FX/Niagara/Source/Niagara/Private/NiagaraModule.cpp
@@ -19,6 +19,7 @@
 #include "Misc/CoreDelegates.h"
 #include "NiagaraShaderModule.h"
 #include "UObject/CoreRedirects.h"
+#include "NiagaraEmitterInstanceBatcher.h"
 
 IMPLEMENT_MODULE(INiagaraModule, Niagara);
 
@@ -259,10 +260,16 @@ void INiagaraModule::StartupModule()
 		return FNiagaraTypeRegistry::GetDefaultDataInterfaceByName(DIClassName);
 	}));
 
+	FFXSystemInterface::RegisterCustomFXSystem(NiagaraEmitterInstanceBatcher::Name, FCreateCustomFXSystemDelegate::CreateLambda([](ERHIFeatureLevel::Type InFeatureLevel, EShaderPlatform InShaderPlatform) -> FFXSystemInterface*
+	{
+		return new NiagaraEmitterInstanceBatcher;
+	}));
 }
 
 void INiagaraModule::ShutdownRenderingResources()
 {
+	FFXSystemInterface::UnregisterCustomFXSystem(NiagaraEmitterInstanceBatcher::Name);
+
 	FNiagaraViewDataMgr::Shutdown();
 }
 
diff --git a/Engine/Plugins/FX/Niagara/Source/Niagara/Private/NiagaraSystemInstance.cpp b/Engine/Plugins/FX/Niagara/Source/Niagara/Private/NiagaraSystemInstance.cpp
--- a/Engine/Plugins/FX/Niagara/Source/Niagara/Private/NiagaraSystemInstance.cpp
+++ b/Engine/Plugins/FX/Niagara/Source/Niagara/Private/NiagaraSystemInstance.cpp
@@ -12,6 +12,7 @@
 #include "NiagaraRenderer.h"
 #include "GameFramework/PlayerController.h"
 #include "Templates/AlignmentTemplates.h"
+#include "NiagaraEmitterInstanceBatcher.h"
 
 
 DECLARE_CYCLE_STAT(TEXT("System Activate [GT]"), STAT_NiagaraSystemActivate, STATGROUP_Niagara);
@@ -54,6 +55,19 @@ FNiagaraSystemInstance::FNiagaraSystemInstance(UNiagaraComponent* InComponent)
 	, bDataInterfacesInitialized(false)
 {
 	SystemBounds.Init();
+
+	if (Component)
+	{
+		UWorld* World = Component->GetWorld();
+		if (World && World->Scene)
+		{
+			FFXSystemInterface*  FXSystemInterface = World->Scene->GetFXSystem();
+			if (FXSystemInterface)
+			{
+				Batcher = static_cast<NiagaraEmitterInstanceBatcher*>(FXSystemInterface->GetInterface(NiagaraEmitterInstanceBatcher::Name));
+			}
+		}
+	}
 }
 
 void FNiagaraSystemInstance::Init(bool bInForceSolo)
diff --git a/Engine/Plugins/FX/Niagara/Source/Niagara/Public/NiagaraSystemInstance.h b/Engine/Plugins/FX/Niagara/Source/Niagara/Public/NiagaraSystemInstance.h
--- a/Engine/Plugins/FX/Niagara/Source/Niagara/Public/NiagaraSystemInstance.h
+++ b/Engine/Plugins/FX/Niagara/Source/Niagara/Public/NiagaraSystemInstance.h
@@ -16,6 +16,7 @@ class FNiagaraWorldManager;
 class UNiagaraComponent;
 class FNiagaraSystemInstance;
 class FNiagaraSystemSimulation;
+class NiagaraEmitterInstanceBatcher;
 
 class NIAGARA_API FNiagaraSystemInstance 
 {
@@ -201,6 +202,8 @@ public:
 
 	bool GetPerInstanceDataAndOffsets(void*& OutData, uint32& OutDataSize, TMap<TWeakObjectPtr<UNiagaraDataInterface>, int32>*& OutOffsets);
 
+	NiagaraEmitterInstanceBatcher* GetBatcher() const { return Batcher; }
+
 private:
 
 	/** Builds the emitter simulations. */
@@ -319,4 +322,6 @@ private:
 	ENiagaraExecutionState ActualExecutionState;
 
 	bool bDataInterfacesInitialized;
+
+	NiagaraEmitterInstanceBatcher* Batcher = nullptr;
 };
diff --git a/Engine/Source/Runtime/Engine/Classes/Particles/ParticleSystemComponent.h b/Engine/Source/Runtime/Engine/Classes/Particles/ParticleSystemComponent.h
--- a/Engine/Source/Runtime/Engine/Classes/Particles/ParticleSystemComponent.h
+++ b/Engine/Source/Runtime/Engine/Classes/Particles/ParticleSystemComponent.h
@@ -696,7 +696,7 @@ private:
 	FVector SavedAutoAttachRelativeScale3D;
 
 public:
-	class FFXSystemInterface* FXSystem;
+	class FFXSystem* FXSystem;
 
 	// Called when the particle system is done
 	UPROPERTY(BlueprintAssignable)
diff --git a/Engine/Source/Runtime/Engine/Private/Particles/FXSystem.cpp b/Engine/Source/Runtime/Engine/Private/Particles/FXSystem.cpp
--- a/Engine/Source/Runtime/Engine/Private/Particles/FXSystem.cpp
+++ b/Engine/Source/Runtime/Engine/Private/Particles/FXSystem.cpp
@@ -8,6 +8,7 @@
 #include "RenderingThread.h"
 #include "VectorField.h"
 #include "Particles/FXSystemPrivate.h"
+#include "Particles/FXSystemSet.h"
 #include "GPUSort.h"
 #include "Particles/ParticleCurveTexture.h"
 #include "VectorField/VectorField.h"
@@ -15,13 +16,34 @@
 #include "SceneUtils.h"
 #include "Renderer/Private/SceneRendering.h" // needed for STATGROUP_CommandListMarkers
 
+
+TMap<FName, FCreateCustomFXSystemDelegate> FFXSystemInterface::CreateCustomFXDelegates;
+
 /*-----------------------------------------------------------------------------
 	External FX system interface.
 -----------------------------------------------------------------------------*/
 
 FFXSystemInterface* FFXSystemInterface::Create(ERHIFeatureLevel::Type InFeatureLevel, EShaderPlatform InShaderPlatform)
 {
-	return new FFXSystem(InFeatureLevel, InShaderPlatform);
+	if (CreateCustomFXDelegates.Num())
+	{
+		FFXSystemSet* Set = new FFXSystemSet;
+		Set->FXSystems.Add(new FFXSystem(InFeatureLevel, InShaderPlatform));
+
+		for (TMap<FName, FCreateCustomFXSystemDelegate>::TConstIterator Ite(CreateCustomFXDelegates); Ite; ++Ite)
+		{
+			FFXSystemInterface* CustomFX = Ite.Value().Execute(InFeatureLevel, InShaderPlatform);
+			if (CustomFX)
+			{
+				Set->FXSystems.Add(CustomFX);
+			}
+		}
+		return Set;
+	}
+	else
+	{
+		return new FFXSystem(InFeatureLevel, InShaderPlatform);
+	}
 }
 
 void FFXSystemInterface::Destroy( FFXSystemInterface* FXSystem )
@@ -33,6 +55,16 @@ void FFXSystemInterface::Destroy( FFXSystemInterface* FXSystem )
 		});
 }
 
+void FFXSystemInterface::RegisterCustomFXSystem(const FName& InterfaceName, const FCreateCustomFXSystemDelegate& InCreateDelegate)
+{
+	CreateCustomFXDelegates.Add(InterfaceName, InCreateDelegate);
+}
+
+void FFXSystemInterface::UnregisterCustomFXSystem(const FName& InterfaceName)
+{
+	CreateCustomFXDelegates.Remove(InterfaceName);
+}
+
 FFXSystemInterface::~FFXSystemInterface()
 {
 }
@@ -159,6 +191,13 @@ FFXSystem::~FFXSystem()
 	DestroyGPUSimulation();
 }
 
+const FName FFXSystem::Name(TEXT("FFXSystem"));
+
+FFXSystemInterface* FFXSystem::GetInterface(const FName& InName)
+{
+	return InName == Name ? this : nullptr;
+}
+
 void FFXSystem::Tick(float DeltaSeconds)
 {
 	if (RHISupportsGPUParticles())
@@ -209,7 +248,7 @@ void FFXSystem::AddVectorField( UVectorFieldComponent* VectorFieldComponent )
 	if (RHISupportsGPUParticles())
 	{
 		check( VectorFieldComponent->VectorFieldInstance == NULL );
-		check( VectorFieldComponent->FXSystem == this );
+		checkSlow( VectorFieldComponent->FXSystem && VectorFieldComponent->FXSystem->GetInterface(Name) == this );
 
 		if ( VectorFieldComponent->VectorField )
 		{
@@ -238,7 +277,7 @@ void FFXSystem::RemoveVectorField( UVectorFieldComponent* VectorFieldComponent )
 {
 	if (RHISupportsGPUParticles())
 	{
-		check( VectorFieldComponent->FXSystem == this );
+		checkSlow(VectorFieldComponent->FXSystem && VectorFieldComponent->FXSystem->GetInterface(Name) == this);
 
 		FVectorFieldInstance* Instance = VectorFieldComponent->VectorFieldInstance;
 		VectorFieldComponent->VectorFieldInstance = NULL;
@@ -264,7 +303,7 @@ void FFXSystem::UpdateVectorField( UVectorFieldComponent* VectorFieldComponent )
 {
 	if (RHISupportsGPUParticles())
 	{
-		check( VectorFieldComponent->FXSystem == this );
+		checkSlow(VectorFieldComponent->FXSystem && VectorFieldComponent->FXSystem->GetInterface(Name) == this);
 
 		FVectorFieldInstance* Instance = VectorFieldComponent->VectorFieldInstance;
 
diff --git a/Engine/Source/Runtime/Engine/Private/Particles/FXSystemPrivate.h b/Engine/Source/Runtime/Engine/Private/Particles/FXSystemPrivate.h
--- a/Engine/Source/Runtime/Engine/Private/Particles/FXSystemPrivate.h
+++ b/Engine/Source/Runtime/Engine/Private/Particles/FXSystemPrivate.h
@@ -135,6 +135,9 @@ public:
 	/** Destructor. */
 	virtual ~FFXSystem();
 
+	static const FName Name;
+	virtual FFXSystemInterface* GetInterface(const FName& InName) override;
+
 	// Begin FFXSystemInterface.
 	virtual void Tick(float DeltaSeconds) override;
 #if WITH_EDITOR
@@ -145,7 +148,7 @@ public:
 	virtual void AddVectorField(UVectorFieldComponent* VectorFieldComponent) override;
 	virtual void RemoveVectorField(UVectorFieldComponent* VectorFieldComponent) override;
 	virtual void UpdateVectorField(UVectorFieldComponent* VectorFieldComponent) override;
-	virtual FParticleEmitterInstance* CreateGPUSpriteEmitterInstance(FGPUSpriteEmitterInfo& EmitterInfo) override;
+	FParticleEmitterInstance* CreateGPUSpriteEmitterInstance(FGPUSpriteEmitterInfo& EmitterInfo);
 	virtual void PreInitViews() override;
 	virtual bool UsesGlobalDistanceField() const override;
 	virtual void PreRender(FRHICommandListImmediate& RHICmdList, const FGlobalDistanceFieldParameterData* GlobalDistanceFieldParameterData) override;
diff --git a/Engine/Source/Runtime/Engine/Private/Particles/FXSystemSet.cpp b/Engine/Source/Runtime/Engine/Private/Particles/FXSystemSet.cpp
new file mode 100644
--- /dev/null
+++ b/Engine/Source/Runtime/Engine/Private/Particles/FXSystemSet.cpp
@@ -0,0 +1,147 @@
+// Copyright 1998-2019 Epic Games, Inc. All Rights Reserved.
+
+/*=============================================================================
+	FXSystemSet.cpp: Internal redirector to several fx systems.
+=============================================================================*/
+
+#include "FXSystemSet.h"
+
+FFXSystemInterface* FFXSystemSet::GetInterface(const FName& InName)
+{
+	for (FFXSystemInterface* FXSystem : FXSystems)
+	{
+		check(FXSystem);
+		FFXSystemInterface* Res = FXSystem->GetInterface(InName);
+		if (Res)
+		{
+			return Res;
+		}
+	}
+	return nullptr;
+}
+
+void FFXSystemSet::Tick(float DeltaSeconds)
+{
+	for (FFXSystemInterface* FXSystem : FXSystems)
+	{
+		check(FXSystem);
+		FXSystem->Tick(DeltaSeconds);
+	}
+}
+
+#if WITH_EDITOR
+
+void FFXSystemSet::Suspend()
+{
+	for (FFXSystemInterface* FXSystem : FXSystems)
+	{
+		check(FXSystem);
+		FXSystem->Suspend();
+	}
+}
+
+void FFXSystemSet::Resume()
+{
+	for (FFXSystemInterface* FXSystem : FXSystems)
+	{
+		check(FXSystem);
+		FXSystem->Resume();
+	}
+}
+
+#endif // #if WITH_EDITOR
+
+void FFXSystemSet::DrawDebug(FCanvas* Canvas)
+{
+	for (FFXSystemInterface* FXSystem : FXSystems)
+	{
+		check(FXSystem);
+		FXSystem->DrawDebug(Canvas);
+	}
+}
+
+void FFXSystemSet::AddVectorField(UVectorFieldComponent* VectorFieldComponent)
+{
+	for (FFXSystemInterface* FXSystem : FXSystems)
+	{
+		check(FXSystem);
+		FXSystem->AddVectorField(VectorFieldComponent);
+	}
+}
+
+void FFXSystemSet::RemoveVectorField(UVectorFieldComponent* VectorFieldComponent)
+{
+	for (FFXSystemInterface* FXSystem : FXSystems)
+	{
+		check(FXSystem);
+		FXSystem->RemoveVectorField(VectorFieldComponent);
+	}
+}
+
+void FFXSystemSet::UpdateVectorField(UVectorFieldComponent* VectorFieldComponent)
+{
+	for (FFXSystemInterface* FXSystem : FXSystems)
+	{
+		check(FXSystem);
+		FXSystem->UpdateVectorField(VectorFieldComponent);
+	}
+}
+
+void FFXSystemSet::PreInitViews()
+{
+	for (FFXSystemInterface* FXSystem : FXSystems)
+	{
+		check(FXSystem);
+		FXSystem->PreInitViews();
+	}
+}
+
+bool FFXSystemSet::UsesGlobalDistanceField() const
+{
+	for (FFXSystemInterface* FXSystem : FXSystems)
+	{
+		check(FXSystem);
+		if (FXSystem->UsesGlobalDistanceField())
+		{
+			return true;
+		}
+	}
+	return false;
+}
+
+void FFXSystemSet::PreRender(FRHICommandListImmediate& RHICmdList, const class FGlobalDistanceFieldParameterData* GlobalDistanceFieldParameterData)
+{
+	for (FFXSystemInterface* FXSystem : FXSystems)
+	{
+		check(FXSystem);
+		FXSystem->PreRender(RHICmdList, GlobalDistanceFieldParameterData);
+	}
+}
+
+void FFXSystemSet::PostRenderOpaque(
+	FRHICommandListImmediate& RHICmdList,
+	const FUniformBufferRHIParamRef ViewUniformBuffer,
+	const class FShaderParametersMetadata* SceneTexturesUniformBufferStruct,
+	FUniformBufferRHIParamRef SceneTexturesUniformBuffer)
+{
+	for (FFXSystemInterface* FXSystem : FXSystems)
+	{
+		check(FXSystem);
+		FXSystem->PostRenderOpaque(
+			RHICmdList,
+			ViewUniformBuffer,
+			SceneTexturesUniformBufferStruct,
+			SceneTexturesUniformBuffer
+		);
+	}
+}
+
+FFXSystemSet::~FFXSystemSet()
+{
+	for (FFXSystemInterface* FXSystem : FXSystems)
+	{
+		check(FXSystem);
+		FFXSystemInterface::Destroy(FXSystem);
+	}
+	FXSystems.Empty();
+}
\ No newline at end of file
diff --git a/Engine/Source/Runtime/Engine/Private/Particles/FXSystemSet.h b/Engine/Source/Runtime/Engine/Private/Particles/FXSystemSet.h
new file mode 100644
--- /dev/null
+++ b/Engine/Source/Runtime/Engine/Private/Particles/FXSystemSet.h
@@ -0,0 +1,46 @@
+// Copyright 1998-2019 Epic Games, Inc. All Rights Reserved.
+
+/*=============================================================================
+	FXSystemSet.h: Internal redirector to several fx systems.
+=============================================================================*/
+
+#pragma once
+
+#include "CoreMinimal.h"
+#include "FXSystem.h"
+
+/**
+ * FX system.
+ */
+class FFXSystemSet : public FFXSystemInterface
+{
+public:
+
+	TArray<FFXSystemInterface*> FXSystems;
+
+	virtual FFXSystemInterface* GetInterface(const FName& InName) override;
+	virtual void Tick(float DeltaSeconds) override;
+
+#if WITH_EDITOR
+	virtual void Suspend() override;
+	virtual void Resume() override;
+#endif // #if WITH_EDITOR
+
+	virtual void DrawDebug(FCanvas* Canvas) override;
+	virtual void AddVectorField(UVectorFieldComponent* VectorFieldComponent) override;
+	virtual void RemoveVectorField(UVectorFieldComponent* VectorFieldComponent) override;
+	virtual void UpdateVectorField(UVectorFieldComponent* VectorFieldComponent) override;
+	virtual void PreInitViews() override;
+	virtual bool UsesGlobalDistanceField() const override;
+	virtual void PreRender(FRHICommandListImmediate& RHICmdList, const class FGlobalDistanceFieldParameterData* GlobalDistanceFieldParameterData) override;
+	virtual void PostRenderOpaque(
+		FRHICommandListImmediate& RHICmdList, 
+		const FUniformBufferRHIParamRef ViewUniformBuffer, 
+		const class FShaderParametersMetadata* SceneTexturesUniformBufferStruct,
+		FUniformBufferRHIParamRef SceneTexturesUniformBuffer) override;
+
+protected:
+
+	/** By making the destructor protected, an instance must be destroyed via FFXSystemInterface::Destroy. */
+	virtual ~FFXSystemSet();
+};
diff --git a/Engine/Source/Runtime/Engine/Private/Particles/ParticleComponents.cpp b/Engine/Source/Runtime/Engine/Private/Particles/ParticleComponents.cpp
--- a/Engine/Source/Runtime/Engine/Private/Particles/ParticleComponents.cpp
+++ b/Engine/Source/Runtime/Engine/Private/Particles/ParticleComponents.cpp
@@ -54,6 +54,7 @@
 #include "ObjectEditorUtils.h"
 #endif
 
+#include "Particles/FXSystemPrivate.h"
 #include "Particles/Camera/ParticleModuleCameraOffset.h"
 #include "Particles/Collision/ParticleModuleCollision.h"
 #include "Particles/Color/ParticleModuleColorOverLife.h"
@@ -3757,7 +3758,11 @@ void UParticleSystemComponent::OnRegister()
 
 	if (World->Scene)
 	{
-		FXSystem = World->Scene->GetFXSystem();
+		FFXSystemInterface*  FXSystemInterface = World->Scene->GetFXSystem();
+		if (FXSystemInterface)
+		{
+			FXSystem = static_cast<FFXSystem*>(FXSystemInterface->GetInterface(FFXSystem::Name));
+		}
 	}
 
 	if (bAutoManageAttachment && !IsActive())
diff --git a/Engine/Source/Runtime/Engine/Public/FXSystem.h b/Engine/Source/Runtime/Engine/Public/FXSystem.h
--- a/Engine/Source/Runtime/Engine/Public/FXSystem.h
+++ b/Engine/Source/Runtime/Engine/Public/FXSystem.h
@@ -110,6 +110,9 @@ inline bool RHISupportsGPUParticles()
 		&& GRHISupportsInstancing;
 }
 
+class FFXSystemInterface;
+DECLARE_DELEGATE_RetVal_TwoParams(FFXSystemInterface*, FCreateCustomFXSystemDelegate, ERHIFeatureLevel::Type, EShaderPlatform);
+
 /*-----------------------------------------------------------------------------
 	The interface to the FX system runtime.
 -----------------------------------------------------------------------------*/
@@ -131,6 +134,21 @@ public:
 	 */
 	ENGINE_API static void Destroy(FFXSystemInterface* FXSystem);
 
+	/**
+	 * Register a custom FX system implementation.
+	 */
+	ENGINE_API static void RegisterCustomFXSystem(const FName& InterfaceName, const FCreateCustomFXSystemDelegate& InCreateDelegate);
+
+	/**
+	 * Unregister a custom FX system implementation.
+	 */
+	ENGINE_API static void UnregisterCustomFXSystem(const FName& InterfaceName);
+
+	/**
+	 * Return the interface bound to the given name.
+	 */
+	virtual FFXSystemInterface* GetInterface(const FName& InName) { return nullptr; };
+
 	/**
 	 * Tick the effects system.
 	 * @param DeltaSeconds The number of seconds by which to step simulations forward.
@@ -175,13 +193,6 @@ public:
 	 */
 	virtual void UpdateVectorField(UVectorFieldComponent* VectorFieldComponent) = 0;
 
-	/**
-	 * Creates an emitter instance for simulating sprite particles on the GPU.
-	 * @param EmitterInfo - Information required for runtime simulation.
-	 * @returns a particle emitter instance for simulating and drawing GPU sprites.
-	 */
-	virtual FParticleEmitterInstance* CreateGPUSpriteEmitterInstance(FGPUSpriteEmitterInfo& EmitterInfo) = 0;
-
 	/**
 	 * Notification from the renderer that it is about to perform visibility
 	 * checks on FX belonging to this system.
@@ -213,7 +224,11 @@ public:
 protected:
 
 	/** By making the destructor protected, an instance must be destroyed via FFXSystemInterface::Destroy. */
-	virtual ~FFXSystemInterface();
+	ENGINE_API virtual ~FFXSystemInterface();
+
+private:
+
+	static TMap<FName, FCreateCustomFXSystemDelegate> CreateCustomFXDelegates;
 };
 
 /*-----------------------------------------------------------------------------
diff --git a/Engine/Source/Runtime/RenderCore/Public/RendererInterface.h b/Engine/Source/Runtime/RenderCore/Public/RendererInterface.h
--- a/Engine/Source/Runtime/RenderCore/Public/RendererInterface.h
+++ b/Engine/Source/Runtime/RenderCore/Public/RendererInterface.h
@@ -503,14 +503,6 @@ public:
 };
 DECLARE_DELEGATE_OneParam(FPostOpaqueRenderDelegate, class FPostOpaqueRenderParameters&);
 
-
-class FComputeDispatcher
-{
-public:
-	virtual void Execute(FRHICommandList &RHICmdList, FUniformBufferRHIParamRef ViewUniformBuffer) = 0;
-};
-
-
 class ICustomVisibilityQuery: public IRefCountedObject
 {
 public:
@@ -762,10 +754,6 @@ public:
 	virtual FPreSceneRenderValues PreSceneRenderExtension() = 0;
 	virtual bool HasPostOpaqueExtentions() const = 0;
 
-	virtual void RegisterPostOpaqueComputeDispatcher(FComputeDispatcher *Dispatcher) = 0;
-	virtual void UnRegisterPostOpaqueComputeDispatcher(FComputeDispatcher *Dispatcher) = 0;
-	virtual void DispatchPostOpaqueCompute(FRHICommandList &CmdList, FUniformBufferRHIParamRef ViewUniformBuffer) = 0;
-
 	/** Delegate that is called upon resolving scene color. */
 	DECLARE_MULTICAST_DELEGATE_TwoParams(FOnResolvedSceneColor, FRHICommandListImmediate& /*RHICmdList*/, class FSceneRenderTargets& /*SceneContext*/);
 
diff --git a/Engine/Source/Runtime/Renderer/Private/DeferredShadingRenderer.cpp b/Engine/Source/Runtime/Renderer/Private/DeferredShadingRenderer.cpp
--- a/Engine/Source/Runtime/Renderer/Private/DeferredShadingRenderer.cpp
+++ b/Engine/Source/Runtime/Renderer/Private/DeferredShadingRenderer.cpp
@@ -509,7 +509,7 @@ bool FDeferredShadingSceneRenderer::RenderHzb(FRHICommandListImmediate& RHICmdLi
 }
 
 void FDeferredShadingSceneRenderer::RenderOcclusion(FRHICommandListImmediate& RHICmdList)
-{
+{		
 	check(RHICmdList.IsOutsideRenderPass());
 
 	SCOPED_GPU_STAT(RHICmdList, HZB);
@@ -962,7 +962,7 @@ void FDeferredShadingSceneRenderer::Render(FRHICommandListImmediate& RHICmdList)
 #endif // WITH_MGPU
 
 	FSceneRenderTargets& SceneContext = FSceneRenderTargets::Get(RHICmdList);
-
+	
 	//make sure all the targets we're going to use will be safely writable.
 	GRenderTargetPool.TransitionTargetsWritable(RHICmdList);
 
@@ -982,7 +982,7 @@ void FDeferredShadingSceneRenderer::Render(FRHICommandListImmediate& RHICmdList)
 	// Anything rendered inside Render() which isn't accounted for will fall into this stat
 	// This works because child stat events do not contribute to their parents' times (see GPU_STATS_CHILD_TIMES_INCLUDED)
 	SCOPED_GPU_STAT(RHICmdList, Unaccounted);
-
+	
 	{
 		SCOPE_CYCLE_COUNTER(STAT_FDeferredShadingSceneRenderer_Render_Init);
 
@@ -1019,7 +1019,7 @@ void FDeferredShadingSceneRenderer::Render(FRHICommandListImmediate& RHICmdList)
 	if (CVarVirtualTextureLightmaps && CVarVirtualTextureLightmaps->GetValueOnRenderThread())
 	{
 		// TODO should probably be in InitViews
-		GetVirtualTextureSystem()->Update(RHICmdList, FeatureLevel);
+		GetVirtualTextureSystem()->Update( RHICmdList, FeatureLevel );
 	}
 
 #if !UE_BUILD_SHIPPING
@@ -1036,7 +1036,7 @@ void FDeferredShadingSceneRenderer::Render(FRHICommandListImmediate& RHICmdList)
 		for (int32 ViewIndex = 0; ViewIndex < Views.Num(); ViewIndex++)
 		{
 			Views[ViewIndex].GetEyeAdaptation(RHICmdList);
-		}
+		}	
 	}
 
 	if (GDoPrepareDistanceFieldSceneAfterRHIFlush && (GRHINeedsExtraDeletionLatency || !GRHICommandList.Bypass()))
@@ -1116,25 +1116,25 @@ void FDeferredShadingSceneRenderer::Render(FRHICommandListImmediate& RHICmdList)
 		switch (ClearMethod)
 		{
 		case 0: // No clear
-		{
-			bRequiresRHIClear = false;
-			bRequiresFarZQuadClear = false;
-			break;
-		}
-
+			{
+				bRequiresRHIClear = false;
+				bRequiresFarZQuadClear = false;
+				break;
+			}
+		
 		case 1: // RHICmdList.Clear
-		{
-			bRequiresRHIClear = true;
-			bRequiresFarZQuadClear = false;
-			break;
-		}
+			{
+				bRequiresRHIClear = true;
+				bRequiresFarZQuadClear = false;
+				break;
+			}
 
 		case 2: // Clear using far-z quad
-		{
-			bRequiresFarZQuadClear = true;
-			bRequiresRHIClear = false;
-			break;
-		}
+			{
+				bRequiresFarZQuadClear = true;
+				bRequiresRHIClear = false;
+				break;
+			}
 		}
 	}
 
@@ -1146,15 +1146,15 @@ void FDeferredShadingSceneRenderer::Render(FRHICommandListImmediate& RHICmdList)
 
 	// force using occ queries for wireframe if rendering is parented or frozen in the first view
 	check(Views.Num());
-#if (UE_BUILD_SHIPPING || UE_BUILD_TEST)
-	const bool bIsViewFrozen = false;
-	const bool bHasViewParent = false;
-#else
-	const bool bIsViewFrozen = Views[0].State && ((FSceneViewState*)Views[0].State)->bIsFrozen;
-	const bool bHasViewParent = Views[0].State && ((FSceneViewState*)Views[0].State)->HasViewParent();
-#endif
-
+	#if (UE_BUILD_SHIPPING || UE_BUILD_TEST)
+		const bool bIsViewFrozen = false;
+		const bool bHasViewParent = false;
+	#else
+		const bool bIsViewFrozen = Views[0].State && ((FSceneViewState*)Views[0].State)->bIsFrozen;
+		const bool bHasViewParent = Views[0].State && ((FSceneViewState*)Views[0].State)->HasViewParent();
+	#endif
 
+	
 	const bool bIsOcclusionTesting = DoOcclusionQueries(FeatureLevel) && (!bIsWireframe || bIsViewFrozen || bHasViewParent);
 
 	// Dynamic vertex and index buffers need to be committed before rendering.
@@ -1218,13 +1218,13 @@ void FDeferredShadingSceneRenderer::Render(FRHICommandListImmediate& RHICmdList)
 		QUICK_SCOPE_CYCLE_COUNTER(STAT_FDeferredShadingSceneRenderer_AsyncUpdateViewCustomData_Wait);
 		FTaskGraphInterface::Get().WaitUntilTasksComplete(UpdateViewCustomDataEvents, ENamedThreads::GetRenderThread());
 	}
-
+	
 	if (CVarVirtualTextureLightmaps && CVarVirtualTextureLightmaps->GetValueOnRenderThread())
 	{
 		// Create VT feedback buffer
 		FIntPoint Size = SceneContext.GetBufferSizeXY();
-		Size = FIntPoint::DivideAndRoundUp(Size, 16);
-		GVirtualTextureFeedback.CreateResourceGPU(RHICmdList, Size.X, Size.Y);
+		Size = FIntPoint::DivideAndRoundUp( Size, 16 );
+		GVirtualTextureFeedback.CreateResourceGPU( RHICmdList, Size.X, Size.Y );
 	}
 
 	checkSlow(RHICmdList.IsOutsideRenderPass());
@@ -1273,7 +1273,7 @@ void FDeferredShadingSceneRenderer::Render(FRHICommandListImmediate& RHICmdList)
 		SCOPE_CYCLE_COUNTER(STAT_FDeferredShadingSceneRenderer_AllocGBufferTargets);
 		SceneContext.PreallocGBufferTargets(); // Even if !bShouldRenderVelocities, the velocity buffer must be bound because it's a compile time option for the shader.
 		SceneContext.AllocGBufferTargets(RHICmdList);
-	}
+	}	
 
 	checkSlow(RHICmdList.IsOutsideRenderPass());
 
@@ -1287,7 +1287,7 @@ void FDeferredShadingSceneRenderer::Render(FRHICommandListImmediate& RHICmdList)
 			RenderOcclusion(RHICmdList);
 		}
 		bool bUseHzbOcclusion = RenderHzb(RHICmdList);
-
+		
 		SCOPED_GPU_STAT(RHICmdList, HZB);
 		if (bUseHzbOcclusion || bIsOcclusionTesting)
 		{
@@ -1329,7 +1329,7 @@ void FDeferredShadingSceneRenderer::Render(FRHICommandListImmediate& RHICmdList)
 		ServiceLocalQueue();
 	}
 
-	if (GetCustomDepthPassLocation() == 0)
+	if(GetCustomDepthPassLocation() == 0)
 	{
 		QUICK_SCOPE_CYCLE_COUNTER(STAT_FDeferredShadingSceneRenderer_CustomDepthPass0);
 		RenderCustomDepthPassAtLocation(RHICmdList, 0);
@@ -1347,8 +1347,8 @@ void FDeferredShadingSceneRenderer::Render(FRHICommandListImmediate& RHICmdList)
 		RenderForwardShadingShadowProjections(RHICmdList, ForwardScreenSpaceShadowMask);
 
 		RenderIndirectCapsuleShadows(
-			RHICmdList,
-			NULL,
+			RHICmdList, 
+			NULL, 
 			NULL);
 	}
 
@@ -1362,9 +1362,9 @@ void FDeferredShadingSceneRenderer::Render(FRHICommandListImmediate& RHICmdList)
 		SCOPE_CYCLE_COUNTER(STAT_FDeferredShadingSceneRenderer_DBuffer);
 
 		// e.g. DBuffer deferred decals
-		for (int32 ViewIndex = 0; ViewIndex < Views.Num(); ViewIndex++)
-		{
-			SCOPED_CONDITIONAL_DRAW_EVENTF(RHICmdList, EventView, Views.Num() > 1, TEXT("View%d"), ViewIndex);
+		for(int32 ViewIndex = 0;ViewIndex < Views.Num();ViewIndex++)
+		{	
+			SCOPED_CONDITIONAL_DRAW_EVENTF(RHICmdList, EventView,Views.Num() > 1, TEXT("View%d"), ViewIndex);
 			FViewInfo& View = Views[ViewIndex];
 
 			Scene->UniformBuffers.UpdateViewUniformBuffer(View);
@@ -1382,18 +1382,18 @@ void FDeferredShadingSceneRenderer::Render(FRHICommandListImmediate& RHICmdList)
 
 		ServiceLocalQueue();
 	}
-
+	
 	checkSlow(RHICmdList.IsOutsideRenderPass());
 
 	if (bRenderDeferredLighting)
 	{
 		bool bShouldAllocateDeferredShadingPathRenderTargets = false;
 		const char* str = SceneContext.ScreenSpaceAO ? "Allocated" : "Unallocated"; //ScreenSpaceAO is determining factor of detecting render target allocation
-		for (int Index = 0; Index < (NumTranslucentVolumeRenderTargetSets * Views.Num()); ++Index)
+		for(int Index = 0; Index < (NumTranslucentVolumeRenderTargetSets * Views.Num()); ++Index)
 		{
-			if (!SceneContext.TranslucencyLightingVolumeAmbient[Index] || !SceneContext.TranslucencyLightingVolumeDirectional[Index])
+			if(!SceneContext.TranslucencyLightingVolumeAmbient[Index] || !SceneContext.TranslucencyLightingVolumeDirectional[Index])
 			{
-
+				
 				ensureMsgf(SceneContext.TranslucencyLightingVolumeAmbient[Index], TEXT("%s%d is unallocated, Deferred Render Targets would be detected as: %s"), "TranslucencyLightingVolumeAmbient", Index, str);
 				ensureMsgf(SceneContext.TranslucencyLightingVolumeDirectional[Index], TEXT("%s%d is unallocated, Deferred Render Targets would be detected as: %s"), "TranslucencyLightingVolumeDirectional", Index, str);
 				bShouldAllocateDeferredShadingPathRenderTargets = true;
@@ -1401,11 +1401,11 @@ void FDeferredShadingSceneRenderer::Render(FRHICommandListImmediate& RHICmdList)
 			}
 		}
 
-		if (bShouldAllocateDeferredShadingPathRenderTargets)
+		if(bShouldAllocateDeferredShadingPathRenderTargets)
 		{
 			SceneContext.AllocateDeferredShadingPathRenderTargets(RHICmdList);
 		}
-
+		
 		if (GbEnableAsyncComputeTranslucencyLightingVolumeClear && GSupportsEfficientAsyncCompute)
 		{
 			ClearTranslucentVolumeLightingAsyncCompute(RHICmdList);
@@ -1438,7 +1438,7 @@ void FDeferredShadingSceneRenderer::Render(FRHICommandListImmediate& RHICmdList)
 
 		// The first time through we'll clear the Overdraw UAVs.
 		SceneContext.BeginRenderingGBuffer(RHICmdList, ColorLoadAction, DepthLoadAction, BasePassDepthStencilAccess, ViewFamily.EngineShowFlags.ShaderComplexity, true, ClearColor);
-
+		
 		// If we are in wireframe mode or will go wide later this pass is just the clear.
 		if (bIsWireframeRenderpass || bDoParallelBasePass)
 		{
@@ -1446,11 +1446,11 @@ void FDeferredShadingSceneRenderer::Render(FRHICommandListImmediate& RHICmdList)
 		}
 		else
 		{
-			bIsGBufferCurrent = true;
+		bIsGBufferCurrent = true;
 		}
 		ServiceLocalQueue();
 	}
-
+	
 	// Wireframe mode requires bRequiresRHIClear to be true. 
 	// Rendering will be very funny without it and the call to BeginRenderingGBuffer will call AllocSceneColor which is needed for the EditorPrimitives resolve.
 	if (bIsWireframeRenderpass)
@@ -1553,7 +1553,7 @@ void FDeferredShadingSceneRenderer::Render(FRHICommandListImmediate& RHICmdList)
 		DrawClearQuad(RHICmdList, FLinearColor(0, 0, 0, 0));
 		SceneContext.FinishRenderingSceneColor(RHICmdList);
 	}
-	
+
 	checkSlow(RHICmdList.IsOutsideRenderPass());
 
 	SceneContext.DBufferA.SafeRelease();
@@ -1980,7 +1980,6 @@ void FDeferredShadingSceneRenderer::Render(FRHICommandListImmediate& RHICmdList)
 	}
 	checkSlow(RHICmdList.IsOutsideRenderPass());
 	UnbindRenderTargets(RHICmdList);
-	RendererModule.DispatchPostOpaqueCompute(RHICmdList, Views[0].ViewUniformBuffer);
 
 	// No longer needed, release
 	LightShaftOutput.LightShaftOcclusion = NULL;
@@ -2096,8 +2095,8 @@ void FDeferredShadingSceneRenderer::Render(FRHICommandListImmediate& RHICmdList)
 			}
 			else
 			{
-				GPostProcessing.Process(RHICmdList, Views[ViewIndex], VelocityRT);
-			}
+			GPostProcessing.Process(RHICmdList, Views[ ViewIndex ], VelocityRT);
+		}
 		}
 
 		// End of frame, we don't need it anymore
@@ -2228,46 +2227,46 @@ void FDeferredShadingSceneRenderer::DownsampleDepthSurface(FRHICommandList& RHIC
 	RPInfo.DepthStencilRenderTarget.ExclusiveDepthStencil = FExclusiveDepthStencil::DepthWrite_StencilWrite;
 	RHICmdList.BeginRenderPass(RPInfo, TEXT("DownsampleDepth"));
 	{
-		SCOPED_DRAW_EVENT(RHICmdList, DownsampleDepth);
+	SCOPED_DRAW_EVENT(RHICmdList, DownsampleDepth);
 
-		// Set shaders and texture
-		TShaderMapRef<FScreenVS> ScreenVertexShader(View.ShaderMap);
-		TShaderMapRef<FDownsampleSceneDepthPS> PixelShader(View.ShaderMap);
-
-		extern TGlobalResource<FFilterVertexDeclaration> GFilterVertexDeclaration;
-
-		FGraphicsPipelineStateInitializer GraphicsPSOInit;
-		RHICmdList.ApplyCachedRenderTargets(GraphicsPSOInit);
+	// Set shaders and texture
+	TShaderMapRef<FScreenVS> ScreenVertexShader(View.ShaderMap);
+	TShaderMapRef<FDownsampleSceneDepthPS> PixelShader(View.ShaderMap);
 
-		GraphicsPSOInit.BlendState = TStaticBlendState<CW_NONE>::GetRHI();
-		GraphicsPSOInit.RasterizerState = TStaticRasterizerState<FM_Solid, CM_None>::GetRHI();
-		GraphicsPSOInit.DepthStencilState = TStaticDepthStencilState<true, CF_Always>::GetRHI();
+	extern TGlobalResource<FFilterVertexDeclaration> GFilterVertexDeclaration;
 
-		GraphicsPSOInit.BoundShaderState.VertexDeclarationRHI = GFilterVertexDeclaration.VertexDeclarationRHI;
-		GraphicsPSOInit.BoundShaderState.VertexShaderRHI = GETSAFERHISHADER_VERTEX(*ScreenVertexShader);
-		GraphicsPSOInit.BoundShaderState.PixelShaderRHI = GETSAFERHISHADER_PIXEL(*PixelShader);
-		GraphicsPSOInit.PrimitiveType = PT_TriangleList;
+	FGraphicsPipelineStateInitializer GraphicsPSOInit;
+	RHICmdList.ApplyCachedRenderTargets(GraphicsPSOInit);
 
-		SetGraphicsPipelineState(RHICmdList, GraphicsPSOInit);
+	GraphicsPSOInit.BlendState = TStaticBlendState<CW_NONE>::GetRHI();
+	GraphicsPSOInit.RasterizerState = TStaticRasterizerState<FM_Solid, CM_None>::GetRHI();
+	GraphicsPSOInit.DepthStencilState = TStaticDepthStencilState<true, CF_Always>::GetRHI();
 
-		PixelShader->SetParameters(RHICmdList, View, bUseMaxDepth, View.ViewRect.Size());
-		const uint32 DownsampledX = FMath::TruncToInt(View.ViewRect.Min.X * ScaleFactor);
-		const uint32 DownsampledY = FMath::TruncToInt(View.ViewRect.Min.Y * ScaleFactor);
-		const uint32 DownsampledSizeX = FMath::TruncToInt(View.ViewRect.Width() * ScaleFactor);
-		const uint32 DownsampledSizeY = FMath::TruncToInt(View.ViewRect.Height() * ScaleFactor);
+	GraphicsPSOInit.BoundShaderState.VertexDeclarationRHI = GFilterVertexDeclaration.VertexDeclarationRHI;
+	GraphicsPSOInit.BoundShaderState.VertexShaderRHI = GETSAFERHISHADER_VERTEX(*ScreenVertexShader);
+	GraphicsPSOInit.BoundShaderState.PixelShaderRHI = GETSAFERHISHADER_PIXEL(*PixelShader);
+	GraphicsPSOInit.PrimitiveType = PT_TriangleList;
 
-		RHICmdList.SetViewport(DownsampledX, DownsampledY, 0.0f, DownsampledX + DownsampledSizeX, DownsampledY + DownsampledSizeY, 1.0f);
+	SetGraphicsPipelineState(RHICmdList, GraphicsPSOInit);
 
-		DrawRectangle(
-			RHICmdList,
-			0, 0,
-			DownsampledSizeX, DownsampledSizeY,
-			View.ViewRect.Min.X, View.ViewRect.Min.Y,
-			View.ViewRect.Width(), View.ViewRect.Height(),
-			FIntPoint(DownsampledSizeX, DownsampledSizeY),
-			SceneContext.GetBufferSizeXY(),
-			*ScreenVertexShader,
-			EDRF_UseTriangleOptimization);
+	PixelShader->SetParameters(RHICmdList, View, bUseMaxDepth, View.ViewRect.Size());
+	const uint32 DownsampledX = FMath::TruncToInt(View.ViewRect.Min.X * ScaleFactor);
+	const uint32 DownsampledY = FMath::TruncToInt(View.ViewRect.Min.Y * ScaleFactor);
+	const uint32 DownsampledSizeX = FMath::TruncToInt(View.ViewRect.Width() * ScaleFactor);
+	const uint32 DownsampledSizeY = FMath::TruncToInt(View.ViewRect.Height() * ScaleFactor);
+
+	RHICmdList.SetViewport(DownsampledX, DownsampledY, 0.0f, DownsampledX + DownsampledSizeX, DownsampledY + DownsampledSizeY, 1.0f);
+
+	DrawRectangle(
+		RHICmdList,
+		0, 0,
+		DownsampledSizeX, DownsampledSizeY,
+		View.ViewRect.Min.X, View.ViewRect.Min.Y,
+		View.ViewRect.Width(), View.ViewRect.Height(),
+		FIntPoint(DownsampledSizeX, DownsampledSizeY),
+		SceneContext.GetBufferSizeXY(),
+		*ScreenVertexShader,
+		EDRF_UseTriangleOptimization);
 	}
 	RHICmdList.EndRenderPass();
 }
@@ -2340,42 +2339,42 @@ void FDeferredShadingSceneRenderer::CopyStencilToLightingChannelTexture(FRHIComm
 		TransitionRenderPassTargets(RHICmdList, RPInfo);
 		RHICmdList.BeginRenderPass(RPInfo, TEXT("CopyStencilToLightingChannel"));
 		{
-			FGraphicsPipelineStateInitializer GraphicsPSOInit;
-			RHICmdList.ApplyCachedRenderTargets(GraphicsPSOInit);
-			GraphicsPSOInit.BlendState = TStaticBlendState<CW_RGBA>::GetRHI();
-			GraphicsPSOInit.RasterizerState = TStaticRasterizerState<FM_Solid, CM_None>::GetRHI();
-			GraphicsPSOInit.DepthStencilState = TStaticDepthStencilState<false, CF_Always>::GetRHI();
+		FGraphicsPipelineStateInitializer GraphicsPSOInit;
+		RHICmdList.ApplyCachedRenderTargets(GraphicsPSOInit);
+		GraphicsPSOInit.BlendState = TStaticBlendState<CW_RGBA>::GetRHI();
+		GraphicsPSOInit.RasterizerState = TStaticRasterizerState<FM_Solid, CM_None>::GetRHI();
+		GraphicsPSOInit.DepthStencilState = TStaticDepthStencilState<false, CF_Always>::GetRHI();
 
-			extern TGlobalResource<FFilterVertexDeclaration> GFilterVertexDeclaration;
-			GraphicsPSOInit.BoundShaderState.VertexDeclarationRHI = GFilterVertexDeclaration.VertexDeclarationRHI;
+		extern TGlobalResource<FFilterVertexDeclaration> GFilterVertexDeclaration;
+		GraphicsPSOInit.BoundShaderState.VertexDeclarationRHI = GFilterVertexDeclaration.VertexDeclarationRHI;
 
-			for (int32 ViewIndex = 0; ViewIndex < Views.Num(); ViewIndex++)
-			{
-				const FViewInfo& View = Views[ViewIndex];
-				// Set shaders and texture
-				TShaderMapRef<FScreenVS> ScreenVertexShader(View.ShaderMap);
-				TShaderMapRef<FCopyStencilToLightingChannelsPS> PixelShader(View.ShaderMap);
-
-				GraphicsPSOInit.BoundShaderState.VertexShaderRHI = GETSAFERHISHADER_VERTEX(*ScreenVertexShader);
-				GraphicsPSOInit.BoundShaderState.PixelShaderRHI = GETSAFERHISHADER_PIXEL(*PixelShader);
-				GraphicsPSOInit.PrimitiveType = PT_TriangleList;
-
-				SetGraphicsPipelineState(RHICmdList, GraphicsPSOInit);
-				PixelShader->SetParameters(RHICmdList, View);
-
-				RHICmdList.SetViewport(View.ViewRect.Min.X, View.ViewRect.Min.Y, 0.0f, View.ViewRect.Min.X + View.ViewRect.Width(), View.ViewRect.Min.Y + View.ViewRect.Height(), 1.0f);
-
-				DrawRectangle(
-					RHICmdList,
-					0, 0,
-					View.ViewRect.Width(), View.ViewRect.Height(),
-					View.ViewRect.Min.X, View.ViewRect.Min.Y,
-					View.ViewRect.Width(), View.ViewRect.Height(),
-					FIntPoint(View.ViewRect.Width(), View.ViewRect.Height()),
-					SceneContext.GetBufferSizeXY(),
-					*ScreenVertexShader,
-					EDRF_UseTriangleOptimization);
-			}
+		for (int32 ViewIndex = 0; ViewIndex < Views.Num(); ViewIndex++)
+		{
+			const FViewInfo& View = Views[ViewIndex];
+			// Set shaders and texture
+			TShaderMapRef<FScreenVS> ScreenVertexShader(View.ShaderMap);
+			TShaderMapRef<FCopyStencilToLightingChannelsPS> PixelShader(View.ShaderMap);
+
+			GraphicsPSOInit.BoundShaderState.VertexShaderRHI = GETSAFERHISHADER_VERTEX(*ScreenVertexShader);
+			GraphicsPSOInit.BoundShaderState.PixelShaderRHI = GETSAFERHISHADER_PIXEL(*PixelShader);
+			GraphicsPSOInit.PrimitiveType = PT_TriangleList;
+
+			SetGraphicsPipelineState(RHICmdList, GraphicsPSOInit);
+			PixelShader->SetParameters(RHICmdList, View);
+
+			RHICmdList.SetViewport(View.ViewRect.Min.X, View.ViewRect.Min.Y, 0.0f, View.ViewRect.Min.X + View.ViewRect.Width(), View.ViewRect.Min.Y + View.ViewRect.Height(), 1.0f);
+
+			DrawRectangle(
+				RHICmdList,
+				0, 0,
+				View.ViewRect.Width(), View.ViewRect.Height(),
+				View.ViewRect.Min.X, View.ViewRect.Min.Y,
+				View.ViewRect.Width(), View.ViewRect.Height(),
+				FIntPoint(View.ViewRect.Width(), View.ViewRect.Height()),
+				SceneContext.GetBufferSizeXY(),
+				*ScreenVertexShader,
+				EDRF_UseTriangleOptimization);
+		}
 		}
 		RHICmdList.EndRenderPass();
 		RHICmdList.CopyToResolveTarget(SceneContext.LightingChannels->GetRenderTargetItem().TargetableTexture, SceneContext.LightingChannels->GetRenderTargetItem().TargetableTexture, FResolveParams());
diff --git a/Engine/Source/Runtime/Renderer/Private/MobileShadingRenderer.cpp b/Engine/Source/Runtime/Renderer/Private/MobileShadingRenderer.cpp
--- a/Engine/Source/Runtime/Renderer/Private/MobileShadingRenderer.cpp
+++ b/Engine/Source/Runtime/Renderer/Private/MobileShadingRenderer.cpp
@@ -463,7 +463,6 @@ void FMobileSceneRenderer::Render(FRHICommandListImmediate& RHICmdList)
 	}
 
 	IRendererModule& RendererModule = GetRendererModule();
-	RendererModule.DispatchPostOpaqueCompute(RHICmdList, View.ViewUniformBuffer);
 
 	if (!View.bIsPlanarReflection)
 	{
diff --git a/Engine/Source/Runtime/Renderer/Private/RendererModule.h b/Engine/Source/Runtime/Renderer/Private/RendererModule.h
--- a/Engine/Source/Runtime/Renderer/Private/RendererModule.h
+++ b/Engine/Source/Runtime/Renderer/Private/RendererModule.h
@@ -78,10 +78,6 @@ public:
 	virtual void RegisterCustomCullingImpl(ICustomCulling* impl) override;
 	virtual void UnregisterCustomCullingImpl(ICustomCulling* impl) override;
 
-	virtual void RegisterPostOpaqueComputeDispatcher(FComputeDispatcher *Dispatcher) override;
-	virtual void UnRegisterPostOpaqueComputeDispatcher(FComputeDispatcher *Dispatcher) override;
-	virtual void DispatchPostOpaqueCompute(FRHICommandList &CmdList, FUniformBufferRHIParamRef ViewUniformBuffer) override;
-
 	virtual FPreSceneRenderDelegate& OnPreSceneRender() override { return PreSceneRenderDelegate; }
 	virtual void RegisterPostOpaqueRenderDelegate(const FPostOpaqueRenderDelegate& PostOpaqueRenderDelegate) override;
 	virtual void RegisterOverlayRenderDelegate(const FPostOpaqueRenderDelegate& OverlayRenderDelegate) override;
@@ -112,7 +108,6 @@ private:
 	FPostOpaqueRenderDelegate PostOpaqueRenderDelegate;
 	FPostOpaqueRenderDelegate OverlayRenderDelegate;
 	FOnResolvedSceneColor PostResolvedSceneColorCallbacks;
-	TArray<FComputeDispatcher*>PostOpaqueDispatchers;
 };
 
 extern ICustomCulling* GCustomCullingImpl;
diff --git a/Engine/Source/Runtime/Renderer/Private/SceneRendering.cpp b/Engine/Source/Runtime/Renderer/Private/SceneRendering.cpp
--- a/Engine/Source/Runtime/Renderer/Private/SceneRendering.cpp
+++ b/Engine/Source/Runtime/Renderer/Private/SceneRendering.cpp
@@ -3358,24 +3358,6 @@ void FRendererModule::DrawRectangle(
 	::DrawRectangle( RHICmdList, X, Y, SizeX, SizeY, U, V, SizeU, SizeV, TargetSize, TextureSize, VertexShader, Flags );
 }
 
-void FRendererModule::RegisterPostOpaqueComputeDispatcher(FComputeDispatcher *Dispatcher)
-{
-	PostOpaqueDispatchers.AddUnique(Dispatcher);
-}
-
-void FRendererModule::UnRegisterPostOpaqueComputeDispatcher(FComputeDispatcher *Dispatcher)
-{
-	PostOpaqueDispatchers.Remove(Dispatcher);
-}
-
-void FRendererModule::DispatchPostOpaqueCompute(FRHICommandList &RHICmdList, FUniformBufferRHIParamRef ViewUniformBuffer)
-{
-	for (FComputeDispatcher *Dispatcher : PostOpaqueDispatchers)
-	{
-		Dispatcher->Execute(RHICmdList, ViewUniformBuffer);
-	}
-}
-
 void FRendererModule::RegisterPostOpaqueRenderDelegate(const FPostOpaqueRenderDelegate& InPostOpaqueRenderDelegate)
 {
 	this->PostOpaqueRenderDelegate = InPostOpaqueRenderDelegate;
