Downsampled translucent rendering now also sets the instanced view uniform buffer

diff --git a/Engine/Source/Runtime/Renderer/Private/DeferredShadingRenderer.h b/Engine/Source/Runtime/Renderer/Private/DeferredShadingRenderer.h
--- a/Engine/Source/Runtime/Renderer/Private/DeferredShadingRenderer.h
+++ b/Engine/Source/Runtime/Renderer/Private/DeferredShadingRenderer.h
@@ -233,11 +233,11 @@ private:
 	/** Issues a timestamp query for the end of the separate translucency pass. */
 	void EndTimingSeparateTranslucencyPass(FRHICommandListImmediate& RHICmdList, const FViewInfo& View);
 
-	/** Setup the downsampled view uniform buffer if it was not already built */
-	void SetupDownsampledTranslucencyViewUniformBuffer(
+	/** Setup the downsampled view uniform parameters if it was not already built */
+	void SetupDownsampledTranslucencyViewParameters(
 		FRHICommandListImmediate& RHICmdList, 
 		const FViewInfo& View,
-		TUniformBufferRef<FViewUniformShaderParameters>& DownsampledTranslucencyViewUniformBuffer);
+		FViewUniformShaderParameters& DownsampledTranslucencyViewParameters);
 
 	/** Resolve the scene color if any translucent material needs it. */
 	void ConditionalResolveSceneColorForTranslucentMaterials(FRHICommandListImmediate& RHICmdList, TRefCountPtr<IPooledRenderTarget>& SceneColorCopy);
diff --git a/Engine/Source/Runtime/Renderer/Private/TranslucentRendering.cpp b/Engine/Source/Runtime/Renderer/Private/TranslucentRendering.cpp
--- a/Engine/Source/Runtime/Renderer/Private/TranslucentRendering.cpp
+++ b/Engine/Source/Runtime/Renderer/Private/TranslucentRendering.cpp
@@ -463,10 +463,10 @@ void FDeferredShadingSceneRenderer::RenderViewTranslucencyParallel(FRHICommandLi
 	RenderViewTranslucencyInner(RHICmdList, View, DrawRenderState, TranslucencyPass, &ParallelCommandListSet);
 }
 
-void FDeferredShadingSceneRenderer::SetupDownsampledTranslucencyViewUniformBuffer(
+void FDeferredShadingSceneRenderer::SetupDownsampledTranslucencyViewParameters(
 	FRHICommandListImmediate& RHICmdList,
 	const FViewInfo& View,
-	TUniformBufferRef<FViewUniformShaderParameters>& DownsampledTranslucencyViewUniformBuffer)
+	FViewUniformShaderParameters& DownsampledTranslucencyViewParameters)
 {
 	FSceneRenderTargets& SceneContext = FSceneRenderTargets::Get(RHICmdList);
 	FIntPoint ScaledSize;
@@ -477,19 +477,16 @@ void FDeferredShadingSceneRenderer::SetupDownsampledTranslucencyViewUniformBuffe
 	SceneContext.GetDownsampledTranslucencyDepth(RHICmdList, ScaledSize);
 	DownsampleDepthSurface(RHICmdList, SceneContext.GetDownsampledTranslucencyDepthSurface(), View, DownsamplingScale, false);
 
-	FViewUniformShaderParameters DownsampledTranslucencyParameters = *View.CachedViewUniformShaderParameters;
+	DownsampledTranslucencyViewParameters  = *View.CachedViewUniformShaderParameters;
 
 	// Update the parts of DownsampledTranslucencyParameters which are dependent on the buffer size and view rect
 	View.SetupViewRectUniformBufferParameters(
-		DownsampledTranslucencyParameters,
+		DownsampledTranslucencyViewParameters,
 		ScaledSize,
 		FIntRect(View.ViewRect.Min.X * DownsamplingScale, View.ViewRect.Min.Y * DownsamplingScale, View.ViewRect.Max.X * DownsamplingScale, View.ViewRect.Max.Y * DownsamplingScale),
 		View.ViewMatrices,
 		View.PrevViewInfo.ViewMatrices
 	);
-
-	Scene->UniformBuffers.ViewUniformBuffer.UpdateUniformBufferImmediate(DownsampledTranslucencyParameters);
-	DownsampledTranslucencyViewUniformBuffer = Scene->UniformBuffers.ViewUniformBuffer;
 }
 
 void FDeferredShadingSceneRenderer::ConditionalResolveSceneColorForTranslucentMaterials(FRHICommandListImmediate& RHICmdList, TRefCountPtr<IPooledRenderTarget>& SceneColorCopy)
@@ -909,9 +906,21 @@ void FDeferredShadingSceneRenderer::RenderTranslucency(FRHICommandListImmediate&
 
 			if (DownsamplingScale < 1.f)
 			{
-				TUniformBufferRef<FViewUniformShaderParameters> DownsampledTranslucencyViewUniformBuffer;
-				SetupDownsampledTranslucencyViewUniformBuffer(RHICmdList, View, DownsampledTranslucencyViewUniformBuffer);
-				DrawRenderState.SetViewUniformBuffer(DownsampledTranslucencyViewUniformBuffer);
+				FViewUniformShaderParameters DownsampledTranslucencyViewParameters;
+				SetupDownsampledTranslucencyViewParameters(RHICmdList, View, DownsampledTranslucencyViewParameters);
+				Scene->UniformBuffers.ViewUniformBuffer.UpdateUniformBufferImmediate(DownsampledTranslucencyViewParameters);
+				DrawRenderState.SetViewUniformBuffer(Scene->UniformBuffers.ViewUniformBuffer);
+
+				if (View.IsInstancedStereoPass() && View.Family->Views.Num() > 0)
+				{
+					// When drawing the left eye in a stereo scene, copy the right eye view values into the instanced view uniform buffer.
+					const EStereoscopicPass StereoPassIndex = (View.StereoPass != eSSP_FULL) ? eSSP_RIGHT_EYE : eSSP_FULL;
+
+					const FViewInfo& InstancedView = static_cast<const FViewInfo&>(View.Family->GetStereoEyeView(StereoPassIndex));
+					SetupDownsampledTranslucencyViewParameters(RHICmdList, InstancedView, DownsampledTranslucencyViewParameters);
+					Scene->UniformBuffers.InstancedViewUniformBuffer.UpdateUniformBufferImmediate(reinterpret_cast<FInstancedViewUniformShaderParameters&>(DownsampledTranslucencyViewParameters));
+					DrawRenderState.SetInstancedViewUniformBuffer(Scene->UniformBuffers.InstancedViewUniformBuffer);
+				}
 			}
 			if (TranslucencyPass == ETranslucencyPass::TPT_TranslucencyAfterDOF)
 			{
