Move the InstanceFactor to draw command submission rather than building

diff --git a/Engine/Source/Runtime/Landscape/Private/LandscapeGrass.cpp b/Engine/Source/Runtime/Landscape/Private/LandscapeGrass.cpp
--- a/Engine/Source/Runtime/Landscape/Private/LandscapeGrass.cpp
+++ b/Engine/Source/Runtime/Landscape/Private/LandscapeGrass.cpp
@@ -414,7 +414,6 @@ void FLandscapeGrassWeightMeshProcessor::Process(
 			PassShaders,
 			MeshFillMode,
 			MeshCullMode,
-			1,
 			SortKey,
 			EMeshPassFeatures::Default,
 			ShaderElementData);
diff --git a/Engine/Source/Runtime/Renderer/Private/BasePassRendering.cpp b/Engine/Source/Runtime/Renderer/Private/BasePassRendering.cpp
--- a/Engine/Source/Runtime/Renderer/Private/BasePassRendering.cpp
+++ b/Engine/Source/Runtime/Renderer/Private/BasePassRendering.cpp
@@ -980,9 +980,6 @@ void FBasePassMeshProcessor::Process(
 		SortKey = CalculateBasePassMeshStaticSortKey(EarlyZPassMode, BlendMode, BasePassShaders.VertexShader, BasePassShaders.PixelShader);
 	}
 
-	const bool bIsInstancedStereo = ViewIfDynamicMeshCommand ? ViewIfDynamicMeshCommand->IsInstancedStereoPass() : (Scene && Scene->bStaticDrawInstancedStereo);
-	const int32 InstanceFactor = bIsInstancedStereo ? 2 : 1;
-
 	BuildMeshDrawCommands(
 		MeshBatch,
 		BatchElementMask,
@@ -993,7 +990,6 @@ void FBasePassMeshProcessor::Process(
 		BasePassShaders,
 		MeshFillMode,
 		MeshCullMode,
-		InstanceFactor,
 		SortKey,
 		EMeshPassFeatures::Default,
 		ShaderElementData);
diff --git a/Engine/Source/Runtime/Renderer/Private/CompositionLighting/PostProcessMeshDecals.cpp b/Engine/Source/Runtime/Renderer/Private/CompositionLighting/PostProcessMeshDecals.cpp
--- a/Engine/Source/Runtime/Renderer/Private/CompositionLighting/PostProcessMeshDecals.cpp
+++ b/Engine/Source/Runtime/Renderer/Private/CompositionLighting/PostProcessMeshDecals.cpp
@@ -341,7 +341,6 @@ void FMeshDecalMeshProcessor::Process(
 		MeshDecalPassShaders,
 		MeshFillMode,
 		MeshCullMode,
-		1,
 		SortKey,
 		EMeshPassFeatures::Default,
 		ShaderElementData);
diff --git a/Engine/Source/Runtime/Renderer/Private/ConvertToUniformMesh.cpp b/Engine/Source/Runtime/Renderer/Private/ConvertToUniformMesh.cpp
--- a/Engine/Source/Runtime/Renderer/Private/ConvertToUniformMesh.cpp
+++ b/Engine/Source/Runtime/Renderer/Private/ConvertToUniformMesh.cpp
@@ -202,7 +202,6 @@ void FConvertToUniformMeshProcessor::Process(
 		PassShaders,
 		MeshFillMode,
 		MeshCullMode,
-		1,
 		SortKey,
 		EMeshPassFeatures::Default,
 		ShaderElementData);
diff --git a/Engine/Source/Runtime/Renderer/Private/CustomDepthRendering.cpp b/Engine/Source/Runtime/Renderer/Private/CustomDepthRendering.cpp
--- a/Engine/Source/Runtime/Renderer/Private/CustomDepthRendering.cpp
+++ b/Engine/Source/Runtime/Renderer/Private/CustomDepthRendering.cpp
@@ -185,7 +185,6 @@ void FCustomDepthPassMeshProcessor::Process(
 		DepthPassShaders,
 		MeshFillMode,
 		MeshCullMode,
-		1,
 		SortKey,
 		bPositionOnly ? EMeshPassFeatures::PositionOnly : EMeshPassFeatures::Default,
 		ShaderElementData);
diff --git a/Engine/Source/Runtime/Renderer/Private/DebugViewModeRendering.cpp b/Engine/Source/Runtime/Renderer/Private/DebugViewModeRendering.cpp
--- a/Engine/Source/Runtime/Renderer/Private/DebugViewModeRendering.cpp
+++ b/Engine/Source/Runtime/Renderer/Private/DebugViewModeRendering.cpp
@@ -340,7 +340,6 @@ void FDebugViewModeMeshProcessor::AddMeshBatch(const FMeshBatch& RESTRICT MeshBa
 		DebugViewModePassShaders,
 		MeshFillMode,
 		MeshCullMode,
-		1,
 		SortKey,
 		EMeshPassFeatures::Default,
 		ShaderElementData);
diff --git a/Engine/Source/Runtime/Renderer/Private/DepthRendering.cpp b/Engine/Source/Runtime/Renderer/Private/DepthRendering.cpp
--- a/Engine/Source/Runtime/Renderer/Private/DepthRendering.cpp
+++ b/Engine/Source/Runtime/Renderer/Private/DepthRendering.cpp
@@ -566,12 +566,12 @@ bool FDeferredShadingSceneRenderer::RenderPrePass(FRHICommandListImmediate& RHIC
 
 			FMeshPassProcessorRenderState DrawRenderState(View, PassUniformBuffer);
 
+			SetupDepthPassState(DrawRenderState);
+
 			if (View.ShouldRenderView())
 			{
 				Scene->UniformBuffers.UpdateViewUniformBuffer(View);
 
-				SetupDepthPassState(DrawRenderState);
-
 				if (bParallel)
 				{
 					check(RHICmdList.IsOutsideRenderPass());
@@ -724,9 +724,6 @@ void FDepthPassMeshProcessor::Process(
 
 	const FMeshDrawCommandSortKey SortKey = CalculateMeshStaticSortKey(DepthPassShaders.VertexShader, DepthPassShaders.PixelShader);
 
-	const bool bIsInstancedStereo = ViewIfDynamicMeshCommand ? ViewIfDynamicMeshCommand->IsInstancedStereoPass() : (Scene && Scene->bStaticDrawInstancedStereo);
-	const int32 InstanceFactor = bIsInstancedStereo ? 2 : 1;
-
 	BuildMeshDrawCommands(
 		MeshBatch,
 		BatchElementMask,
@@ -737,7 +734,6 @@ void FDepthPassMeshProcessor::Process(
 		DepthPassShaders,
 		MeshFillMode,
 		MeshCullMode,
-		InstanceFactor,
 		SortKey,
 		bPositionOnly ? EMeshPassFeatures::PositionOnly : EMeshPassFeatures::Default,
 		ShaderElementData);
diff --git a/Engine/Source/Runtime/Renderer/Private/DistortionRendering.cpp b/Engine/Source/Runtime/Renderer/Private/DistortionRendering.cpp
--- a/Engine/Source/Runtime/Renderer/Private/DistortionRendering.cpp
+++ b/Engine/Source/Runtime/Renderer/Private/DistortionRendering.cpp
@@ -768,7 +768,6 @@ void FDistortionMeshProcessor::Process(
 		DistortionPassShaders,
 		MeshFillMode,
 		MeshCullMode,
-		1,
 		SortKey,
 		EMeshPassFeatures::Default,
 		ShaderElementData);
diff --git a/Engine/Source/Runtime/Renderer/Private/EditorPrimitivesRendering.cpp b/Engine/Source/Runtime/Renderer/Private/EditorPrimitivesRendering.cpp
--- a/Engine/Source/Runtime/Renderer/Private/EditorPrimitivesRendering.cpp
+++ b/Engine/Source/Runtime/Renderer/Private/EditorPrimitivesRendering.cpp
@@ -99,7 +99,6 @@ void FEditorPrimitivesBasePassMeshProcessor::ProcessDeferredShadingPath(const FM
 		BasePassShaders,
 		MeshFillMode,
 		MeshCullMode,
-		1,
 		SortKey,
 		EMeshPassFeatures::Default,
 		ShaderElementData);
@@ -154,7 +153,6 @@ void FEditorPrimitivesBasePassMeshProcessor::ProcessMobileShadingPath(const FMes
 		BasePassShaders,
 		MeshFillMode,
 		MeshCullMode,
-		1,
 		SortKey,
 		EMeshPassFeatures::Default,
 		ShaderElementData);
diff --git a/Engine/Source/Runtime/Renderer/Private/LightMapDensityRendering.cpp b/Engine/Source/Runtime/Renderer/Private/LightMapDensityRendering.cpp
--- a/Engine/Source/Runtime/Renderer/Private/LightMapDensityRendering.cpp
+++ b/Engine/Source/Runtime/Renderer/Private/LightMapDensityRendering.cpp
@@ -220,7 +220,6 @@ void FLightmapDensityMeshProcessor::Process(
 		LightmapDensityPassShaders,
 		MeshFillMode,
 		MeshCullMode,
-		1,
 		SortKey,
 		EMeshPassFeatures::Default,
 		ShaderElementData);
diff --git a/Engine/Source/Runtime/Renderer/Private/MeshDrawCommands.cpp b/Engine/Source/Runtime/Renderer/Private/MeshDrawCommands.cpp
--- a/Engine/Source/Runtime/Renderer/Private/MeshDrawCommands.cpp
+++ b/Engine/Source/Runtime/Renderer/Private/MeshDrawCommands.cpp
@@ -311,7 +311,7 @@ void BuildMeshDrawCommandPrimitiveIdBuffer(
 	int32& MaxInstances,
 	int32& VisibleMeshDrawCommandsNum,
 	int32& NewPassVisibleMeshDrawCommandsNum,
-	uint32 MaxInstanceFactor = 1
+	uint32 InstanceFactor = 1
 	)
 {
 	QUICK_SCOPE_CYCLE_COUNTER(STAT_BuildMeshDrawCommandPrimitiveIdBuffer);
@@ -391,8 +391,6 @@ void BuildMeshDrawCommandPrimitiveIdBuffer(
 			}
 
 			//@todo - refactor into instance step rate in the RHI
-			uint32 InstanceFactor = VisibleMeshDrawCommand.MeshDrawCommand->InstanceFactor;
-			InstanceFactor = ensure(InstanceFactor <= MaxInstanceFactor) ? InstanceFactor : MaxInstanceFactor;
 			for (uint32 InstanceFactorIndex = 0; InstanceFactorIndex < InstanceFactor; InstanceFactorIndex++, PrimitiveIdIndex++)
 			{
 				//@todo - refactor into memcpy
@@ -416,8 +414,6 @@ void BuildMeshDrawCommandPrimitiveIdBuffer(
 		for (int32 DrawCommandIndex = 0; DrawCommandIndex < NumDrawCommands; DrawCommandIndex++)
 		{
 			const FVisibleMeshDrawCommand& VisibleMeshDrawCommand = VisibleMeshDrawCommands[DrawCommandIndex];
-			uint32 InstanceFactor = VisibleMeshDrawCommand.MeshDrawCommand->InstanceFactor;
-			InstanceFactor = ensure(InstanceFactor <= MaxInstanceFactor) ? InstanceFactor : MaxInstanceFactor;
 			for (uint32 InstanceFactorIndex = 0; InstanceFactorIndex < InstanceFactor; InstanceFactorIndex++, PrimitiveIdIndex++)
 			{
 				checkSlow(PrimitiveIdIndex < MaxPrimitiveId);
@@ -781,7 +777,8 @@ public:
 					Context.TempVisibleMeshDrawCommands,
 					Context.MaxInstances,
 					Context.VisibleMeshDrawCommandsNum,
-					Context.NewPassVisibleMeshDrawCommandsNum
+					Context.NewPassVisibleMeshDrawCommandsNum,
+					Context.InstanceFactor
 				);
 			}
 		}
@@ -804,7 +801,8 @@ void SortAndMergeDynamicPassMeshDrawCommands(
 	ERHIFeatureLevel::Type FeatureLevel,
 	FMeshCommandOneFrameArray& VisibleMeshDrawCommands,
 	FDynamicMeshDrawCommandStorage& MeshDrawCommandStorage,
-	FVertexBufferRHIParamRef& OutPrimitiveIdVertexBuffer)
+	FVertexBufferRHIParamRef& OutPrimitiveIdVertexBuffer,
+	uint32 InstanceFactor)
 {
 	const bool bUseGPUScene = UseGPUScene(GMaxRHIShaderPlatform, FeatureLevel);
 
@@ -815,7 +813,6 @@ void SortAndMergeDynamicPassMeshDrawCommands(
 		int32 MaxInstances = 1;
 		int32 VisibleMeshDrawCommandsNum = 0;
 		int32 NewPassVisibleMeshDrawCommandsNum = 0;
-		uint32 MaxInstanceFactor = 1;
 
 		VisibleMeshDrawCommands.Sort(FCompareFMeshDrawCommands());
 
@@ -827,9 +824,7 @@ void SortAndMergeDynamicPassMeshDrawCommands(
 				NewPassVisibleMeshDrawCommands.Empty(NumDrawCommands);
 			}
 
-			//@todo - make the instance factor a pass parameter instead of a draw command parameter
-			MaxInstanceFactor = VisibleMeshDrawCommands[0].MeshDrawCommand->InstanceFactor;
-			const int32 PrimitiveIdBufferDataSize = MaxInstanceFactor * NumDrawCommands * sizeof(int32);
+			const int32 PrimitiveIdBufferDataSize = InstanceFactor * NumDrawCommands * sizeof(int32);
 			OutPrimitiveIdVertexBuffer = GPrimitiveIdVertexBufferPool.Allocate(PrimitiveIdBufferDataSize);
 			void* PrimitiveIdBufferData = RHILockVertexBuffer(OutPrimitiveIdVertexBuffer, 0, PrimitiveIdBufferDataSize, RLM_WriteOnly);
 
@@ -843,7 +838,7 @@ void SortAndMergeDynamicPassMeshDrawCommands(
 				MaxInstances,
 				VisibleMeshDrawCommandsNum,
 				NewPassVisibleMeshDrawCommandsNum,
-				MaxInstanceFactor
+				InstanceFactor
 			);
 
 			RHIUnlockVertexBuffer(OutPrimitiveIdVertexBuffer);
@@ -893,6 +888,10 @@ void FParallelMeshDrawCommandPass::DispatchPassSetup(
 	TaskContext.NumDynamicMeshElements = NumDynamicMeshElements;
 	TaskContext.NumDynamicMeshCommandBuildRequestElements = NumDynamicMeshCommandBuildRequestElements;
 
+	// Only apply instancing for ISR to main view passes
+	const bool bIsMainViewPass = (FPassProcessorManager::GetPassFlags(TaskContext.ShadingPath, TaskContext.PassType) & EMeshPassFlags::MainView) != EMeshPassFlags::None;
+	TaskContext.InstanceFactor = (bIsMainViewPass && View.IsInstancedStereoPass()) ? 2 : 1;
+
 	// Setup translucency sort key update pass based on view.
 	TaskContext.TranslucencyPass = ETranslucencyPass::TPT_MAX;
 	TaskContext.TranslucentSortPolicy = View.TranslucentSortPolicy;
@@ -901,9 +900,6 @@ void FParallelMeshDrawCommandPass::DispatchPassSetup(
 	TaskContext.ViewMatrix = View.ViewMatrices.GetViewMatrix();
 	TaskContext.PrimitiveBounds = &Scene->PrimitiveBounds;
 
-	const bool bIsInstancedStereo = View.IsInstancedStereoPass() && Scene->bStaticDrawInstancedStereo;
-	TaskContext.MaxInstanceFactor = bIsInstancedStereo ? 2 : 1;
-
 	switch (PassType)
 	{
 		case EMeshPass::TranslucencyStandard: TaskContext.TranslucencyPass = ETranslucencyPass::TPT_StandardTranslucency; break;
@@ -926,7 +922,7 @@ void FParallelMeshDrawCommandPass::DispatchPassSetup(
 
 	// Preallocate resources on rendering thread based on MaxNumDraws.
 	bPrimitiveIdBufferDataOwnedByRHIThread = false;
-	TaskContext.PrimitiveIdBufferDataSize = TaskContext.MaxInstanceFactor * MaxNumDraws * sizeof(int32);
+	TaskContext.PrimitiveIdBufferDataSize = TaskContext.InstanceFactor * MaxNumDraws * sizeof(int32);
 	TaskContext.PrimitiveIdBufferData = FMemory::Malloc(TaskContext.PrimitiveIdBufferDataSize);
 	PrimitiveIdVertexBufferRHI = GPrimitiveIdVertexBufferPool.Allocate(TaskContext.PrimitiveIdBufferDataSize);
 	TaskContext.MeshDrawCommands.Reserve(MaxNumDraws);
@@ -998,6 +994,7 @@ class FDrawVisibleMeshCommandsAnyThreadTask : public FRenderTask
 	FVertexBufferRHIParamRef PrimitiveIdsBuffer;
 	int32 BasePrimitiveIdsOffset;
 	bool bDynamicInstancing;
+	uint32 InstanceFactor;
 	int32 TaskIndex;
 	int32 TaskNum;
 
@@ -1009,6 +1006,7 @@ public:
 		FVertexBufferRHIParamRef InPrimitiveIdsBuffer,
 		int32 InBasePrimitiveIdsOffset,
 		bool bInDynamicInstancing,
+		uint32 InInstanceFactor,
 		int32 InTaskIndex,
 		int32 InTaskNum
 	)
@@ -1017,6 +1015,7 @@ public:
 		, PrimitiveIdsBuffer(InPrimitiveIdsBuffer)
 		, BasePrimitiveIdsOffset(InBasePrimitiveIdsOffset)
 		, bDynamicInstancing(bInDynamicInstancing)
+		, InstanceFactor(InInstanceFactor)
 		, TaskIndex(InTaskIndex)
 		, TaskNum(InTaskNum)
 	{}
@@ -1038,7 +1037,7 @@ public:
 		const int32 StartIndex = TaskIndex * NumDrawsPerTask;
 		const int32 NumDraws = FMath::Min(NumDrawsPerTask, DrawNum - StartIndex);
 
-		SubmitMeshDrawCommandsRange(VisibleMeshDrawCommands, PrimitiveIdsBuffer, BasePrimitiveIdsOffset, bDynamicInstancing, StartIndex, NumDraws, RHICmdList);
+		SubmitMeshDrawCommandsRange(VisibleMeshDrawCommands, PrimitiveIdsBuffer, BasePrimitiveIdsOffset, bDynamicInstancing, StartIndex, NumDraws, InstanceFactor, RHICmdList);
 
 		RHICmdList.EndRenderPass();
 		RHICmdList.HandleRTThreadTaskCompletion(MyCompletionGraphEvent);
@@ -1137,7 +1136,7 @@ void FParallelMeshDrawCommandPass::DispatchDraw(FParallelCommandListSet* Paralle
 			FRHICommandList* CmdList = ParallelCommandListSet->NewParallelCommandList();
 
 			FGraphEventRef AnyThreadCompletionEvent = TGraphTask<FDrawVisibleMeshCommandsAnyThreadTask>::CreateTask(&Prereqs, RenderThread)
-				.ConstructAndDispatchWhenReady(*CmdList, TaskContext.MeshDrawCommands, PrimitiveIdsBuffer, BasePrimitiveIdsOffset, TaskContext.bDynamicInstancing, TaskIndex, NumTasks);
+				.ConstructAndDispatchWhenReady(*CmdList, TaskContext.MeshDrawCommands, PrimitiveIdsBuffer, BasePrimitiveIdsOffset, TaskContext.bDynamicInstancing, TaskContext.InstanceFactor, TaskIndex, NumTasks);
 			ParallelCommandListSet->AddParallelCommandList(CmdList, AnyThreadCompletionEvent, NumDraws);
 		}
 	}
@@ -1153,7 +1152,7 @@ void FParallelMeshDrawCommandPass::DispatchDraw(FParallelCommandListSet* Paralle
 			RHIUnlockVertexBuffer(PrimitiveIdVertexBufferRHI);
 		}
 
-		SubmitMeshDrawCommandsRange(TaskContext.MeshDrawCommands, PrimitiveIdsBuffer, BasePrimitiveIdsOffset, TaskContext.bDynamicInstancing, 0, TaskContext.MeshDrawCommands.Num(), RHICmdList);
+		SubmitMeshDrawCommandsRange(TaskContext.MeshDrawCommands, PrimitiveIdsBuffer, BasePrimitiveIdsOffset, TaskContext.bDynamicInstancing, 0, TaskContext.MeshDrawCommands.Num(), TaskContext.InstanceFactor, RHICmdList);
 	}
 }
 
diff --git a/Engine/Source/Runtime/Renderer/Private/MeshDrawCommands.h b/Engine/Source/Runtime/Renderer/Private/MeshDrawCommands.h
--- a/Engine/Source/Runtime/Renderer/Private/MeshDrawCommands.h
+++ b/Engine/Source/Runtime/Renderer/Private/MeshDrawCommands.h
@@ -59,6 +59,7 @@ public:
 		, MeshPassProcessor(nullptr)
 		, MobileBasePassCSMMeshPassProcessor(nullptr)
 		, DynamicMeshElements(nullptr)
+		, InstanceFactor(1)
 		, NumDynamicMeshElements(0)
 		, NumDynamicMeshCommandBuildRequestElements(0)
 		, PrimitiveIdBufferData(nullptr)
@@ -67,7 +68,6 @@ public:
 		, VisibleMeshDrawCommandsNum(0)
 		, NewPassVisibleMeshDrawCommandsNum(0)
 		, MaxInstances(1)
-		, MaxInstanceFactor(1)
 	{
 	}
 
@@ -87,6 +87,7 @@ public:
 	const TArray<FMeshPassMask, SceneRenderingAllocator>* DynamicMeshElementsPassRelevance;
 
 	// Commands.
+	int32 InstanceFactor;
 	int32 NumDynamicMeshElements;
 	int32 NumDynamicMeshCommandBuildRequestElements;
 	FMeshCommandOneFrameArray MeshDrawCommands;
@@ -112,9 +113,6 @@ public:
 	int32 VisibleMeshDrawCommandsNum;
 	int32 NewPassVisibleMeshDrawCommandsNum;
 	int32 MaxInstances;
-
-	// Assertion check
-	uint32 MaxInstanceFactor;
 };
 
 /**
@@ -183,4 +181,5 @@ extern void SortAndMergeDynamicPassMeshDrawCommands(
 	ERHIFeatureLevel::Type FeatureLevel,
 	FMeshCommandOneFrameArray& VisibleMeshDrawCommands,
 	FDynamicMeshDrawCommandStorage& MeshDrawCommandStorage,
-	FVertexBufferRHIParamRef& OutPrimitiveIdVertexBuffer);
\ No newline at end of file
+	FVertexBufferRHIParamRef& OutPrimitiveIdVertexBuffer,
+	uint32 InstanceFactor);
\ No newline at end of file
diff --git a/Engine/Source/Runtime/Renderer/Private/MeshPassProcessor.cpp b/Engine/Source/Runtime/Renderer/Private/MeshPassProcessor.cpp
--- a/Engine/Source/Runtime/Renderer/Private/MeshPassProcessor.cpp
+++ b/Engine/Source/Runtime/Renderer/Private/MeshPassProcessor.cpp
@@ -608,7 +608,7 @@ void FMeshDrawCommand::SetRayTracingShaders(const FMeshProcessorShaders& Shaders
 }
 #endif // RHI_RAYTRACING
 
-void FMeshDrawCommand::SetDrawParametersAndFinalize(const FMeshBatch& MeshBatch, int32 BatchElementIndex, int32 BatchInstanceFactor, bool bDoSetupPsoStateForRasterization)
+void FMeshDrawCommand::SetDrawParametersAndFinalize(const FMeshBatch& MeshBatch, int32 BatchElementIndex, bool bDoSetupPsoStateForRasterization)
 {
 	const FMeshBatchElement& BatchElement = MeshBatch.Elements[BatchElementIndex];
 
@@ -618,7 +618,6 @@ void FMeshDrawCommand::SetDrawParametersAndFinalize(const FMeshBatch& MeshBatch,
 	FirstIndex = BatchElement.FirstIndex;
 	NumPrimitives = BatchElement.NumPrimitives;
 	NumInstances = BatchElement.NumInstances;
-	InstanceFactor = BatchInstanceFactor;
 	BaseVertexIndex = BatchElement.BaseVertexIndex;
 	NumVertices = BatchElement.MaxVertexIndex - BatchElement.MinVertexIndex + 1;
 	IndirectArgsBuffer = BatchElement.IndirectArgsBuffer;
@@ -724,7 +723,8 @@ void FMeshDrawCommand::SubmitDraw(
 	const FMeshDrawCommand& RESTRICT MeshDrawCommand, 
 	FVertexBufferRHIParamRef ScenePrimitiveIdsBuffer,
 	int32 PrimitiveIdOffset,
-	FRHICommandList& RHICmdList, 
+	uint32 InstanceFactor,
+	FRHICommandList& RHICmdList,
 	FMeshDrawCommandStateCache& RESTRICT StateCache)
 {
 	checkSlow(MeshDrawCommand.CachedPipelineId.IsValid());
@@ -752,7 +752,7 @@ void FMeshDrawCommand::SubmitDraw(
 			DrawEventName = Material->GetFriendlyName();
 		}
 
-		const uint32 Instances = MeshDrawCommand.NumInstances * MeshDrawCommand.InstanceFactor;
+		const uint32 Instances = MeshDrawCommand.NumInstances * InstanceFactor;
 		if (Instances > 1)
 		{
 			BEGIN_DRAW_EVENTF(
@@ -821,7 +821,7 @@ void FMeshDrawCommand::SubmitDraw(
 				MeshDrawCommand.NumVertices,
 				MeshDrawCommand.FirstIndex,
 				MeshDrawCommand.NumPrimitives,
-				MeshDrawCommand.NumInstances * MeshDrawCommand.InstanceFactor
+				MeshDrawCommand.NumInstances * InstanceFactor
 			);
 		}
 	}
@@ -830,7 +830,7 @@ void FMeshDrawCommand::SubmitDraw(
 		RHICmdList.DrawPrimitive(
 			MeshDrawCommand.BaseVertexIndex + MeshDrawCommand.FirstIndex,
 			MeshDrawCommand.NumPrimitives,
-			MeshDrawCommand.NumInstances * MeshDrawCommand.InstanceFactor
+			MeshDrawCommand.NumInstances * InstanceFactor
 		);
 	}
 
@@ -841,9 +841,10 @@ void SubmitMeshDrawCommands(
 	FVertexBufferRHIParamRef PrimitiveIdsBuffer,
 	int32 BasePrimitiveIdsOffset,
 	bool bDynamicInstancing,
+	uint32 InstanceFactor,
 	FRHICommandList& RHICmdList)
 {
-	SubmitMeshDrawCommandsRange(VisibleMeshDrawCommands, PrimitiveIdsBuffer, BasePrimitiveIdsOffset, bDynamicInstancing, 0, VisibleMeshDrawCommands.Num(), RHICmdList);
+	SubmitMeshDrawCommandsRange(VisibleMeshDrawCommands, PrimitiveIdsBuffer, BasePrimitiveIdsOffset, bDynamicInstancing, 0, VisibleMeshDrawCommands.Num(), InstanceFactor, RHICmdList);
 }
 
 void SubmitMeshDrawCommandsRange(
@@ -853,6 +854,7 @@ void SubmitMeshDrawCommandsRange(
 	bool bDynamicInstancing,
 	int32 StartIndex,
 	int32 NumMeshDrawCommands,
+	uint32 InstanceFactor,
 	FRHICommandList& RHICmdList)
 {
 	FMeshDrawCommandStateCache StateCache;
@@ -865,7 +867,7 @@ void SubmitMeshDrawCommandsRange(
 		const FVisibleMeshDrawCommand& VisibleMeshDrawCommand = VisibleMeshDrawCommands[DrawCommandIndex];
 		const int32 PrimitiveIdBufferOffset = BasePrimitiveIdsOffset + (bDynamicInstancing ? VisibleMeshDrawCommand.PrimitiveIdBufferOffset : DrawCommandIndex) * sizeof(int32);
 		checkSlow(!bDynamicInstancing || VisibleMeshDrawCommand.PrimitiveIdBufferOffset >= 0);
-		FMeshDrawCommand::SubmitDraw(*VisibleMeshDrawCommand.MeshDrawCommand, PrimitiveIdsBuffer, PrimitiveIdBufferOffset, RHICmdList, StateCache);
+		FMeshDrawCommand::SubmitDraw(*VisibleMeshDrawCommand.MeshDrawCommand, PrimitiveIdsBuffer, PrimitiveIdBufferOffset, InstanceFactor, RHICmdList, StateCache);
 	}
 }
 
@@ -873,7 +875,8 @@ void DrawDynamicMeshPassPrivate(
 	const FSceneView& View,
 	FRHICommandList& RHICmdList,
 	FMeshCommandOneFrameArray& VisibleMeshDrawCommands,
-	FDynamicMeshDrawCommandStorage& DynamicMeshDrawCommandStorage)
+	FDynamicMeshDrawCommandStorage& DynamicMeshDrawCommandStorage,
+	uint32 InstanceFactor)
 {
 	if (VisibleMeshDrawCommands.Num() > 0)
 	{
@@ -881,9 +884,9 @@ void DrawDynamicMeshPassPrivate(
 
 		FVertexBufferRHIParamRef PrimitiveIdVertexBuffer = nullptr;
 
-		SortAndMergeDynamicPassMeshDrawCommands(View.GetFeatureLevel(), VisibleMeshDrawCommands, DynamicMeshDrawCommandStorage, PrimitiveIdVertexBuffer);
+		SortAndMergeDynamicPassMeshDrawCommands(View.GetFeatureLevel(), VisibleMeshDrawCommands, DynamicMeshDrawCommandStorage, PrimitiveIdVertexBuffer, InstanceFactor);
 
-		SubmitMeshDrawCommandsRange(VisibleMeshDrawCommands, PrimitiveIdVertexBuffer, 0, bDynamicInstancing, 0, VisibleMeshDrawCommands.Num(), RHICmdList);
+		SubmitMeshDrawCommandsRange(VisibleMeshDrawCommands, PrimitiveIdVertexBuffer, 0, bDynamicInstancing, 0, VisibleMeshDrawCommands.Num(), InstanceFactor, RHICmdList);
 	}
 }
 
@@ -1001,14 +1004,13 @@ void FCachedPassMeshDrawListContext::FinalizeCommand(
 	int32 DrawPrimitiveId,
 	ERasterizerFillMode MeshFillMode,
 	ERasterizerCullMode MeshCullMode,
-	int32 InstanceFactor,
 	FMeshDrawCommandSortKey SortKey,
 	FMeshDrawCommand& MeshDrawCommand,
 	bool bDoSetupPsoStateForRasterization)
 {
 	QUICK_SCOPE_CYCLE_COUNTER(STAT_FinalizeCachedMeshDrawCommand);
 
-	MeshDrawCommand.SetDrawParametersAndFinalize(MeshBatch, BatchElementIndex, InstanceFactor, bDoSetupPsoStateForRasterization);
+	MeshDrawCommand.SetDrawParametersAndFinalize(MeshBatch, BatchElementIndex, bDoSetupPsoStateForRasterization);
 
 	check(CommandInfo.CommandIndex != -1);
 	const bool bCanUseGPUScene = UseGPUScene(GMaxRHIShaderPlatform, GMaxRHIFeatureLevel);
diff --git a/Engine/Source/Runtime/Renderer/Private/MobileBasePass.cpp b/Engine/Source/Runtime/Renderer/Private/MobileBasePass.cpp
--- a/Engine/Source/Runtime/Renderer/Private/MobileBasePass.cpp
+++ b/Engine/Source/Runtime/Renderer/Private/MobileBasePass.cpp
@@ -669,7 +669,6 @@ void FMobileBasePassMeshProcessor::Process(
 		BasePassShaders,
 		MeshFillMode,
 		MeshCullMode,
-		1,
 		SortKey,
 		EMeshPassFeatures::Default,
 		ShaderElementData);
diff --git a/Engine/Source/Runtime/Renderer/Private/MobileTranslucentRendering.cpp b/Engine/Source/Runtime/Renderer/Private/MobileTranslucentRendering.cpp
--- a/Engine/Source/Runtime/Renderer/Private/MobileTranslucentRendering.cpp
+++ b/Engine/Source/Runtime/Renderer/Private/MobileTranslucentRendering.cpp
@@ -385,7 +385,6 @@ private:
 			InverseOpacityShaders,
 			MeshFillMode,
 			MeshCullMode,
-			1,
 			SortKey,
 			EMeshPassFeatures::Default,
 			ShaderElementData);
diff --git a/Engine/Source/Runtime/Renderer/Private/RayTracing/RayTracingMaterialHitShaders.cpp b/Engine/Source/Runtime/Renderer/Private/RayTracing/RayTracingMaterialHitShaders.cpp
--- a/Engine/Source/Runtime/Renderer/Private/RayTracing/RayTracingMaterialHitShaders.cpp
+++ b/Engine/Source/Runtime/Renderer/Private/RayTracing/RayTracingMaterialHitShaders.cpp
@@ -235,7 +235,6 @@ void FRayTracingMeshProcessor::Process(
 		RayTracingShaders,
 		MeshFillMode,
 		MeshCullMode,
-		1,
 		FMeshDrawCommandSortKey::Default,
 		EMeshPassFeatures::Default,
 		ShaderElementData);
diff --git a/Engine/Source/Runtime/Renderer/Private/RendererScene.cpp b/Engine/Source/Runtime/Renderer/Private/RendererScene.cpp
--- a/Engine/Source/Runtime/Renderer/Private/RendererScene.cpp
+++ b/Engine/Source/Runtime/Renderer/Private/RendererScene.cpp
@@ -884,10 +884,6 @@ FScene::FScene(UWorld* InWorld, bool bInRequiresHitProxies, bool bInIsEditorScen
 	static auto* ShaderPipelinesCvar = IConsoleManager::Get().FindTConsoleVariableDataInt(TEXT("r.ShaderPipelines"));
 	StaticDrawShaderPipelines = ShaderPipelinesCvar->GetValueOnAnyThread();
 
-	// Query instanced stereo and multi-view state
-	static const auto InstancedStereoCvar = IConsoleManager::Get().FindTConsoleVariableDataInt(TEXT("vr.InstancedStereo"));
-	bStaticDrawInstancedStereo = RHISupportsInstancedStereo(GShaderPlatformForFeatureLevel[InFeatureLevel]) && GEngine->IsStereoscopic3D() && InstancedStereoCvar->GetValueOnAnyThread();
-
 	if (World->FXSystem)
 	{
 		FFXSystemInterface::Destroy(World->FXSystem);
@@ -2862,21 +2858,17 @@ void FScene::ConditionalMarkStaticMeshElementsForUpdate()
 {
 	static auto* EarlyZPassCvar = IConsoleManager::Get().FindTConsoleVariableDataInt(TEXT("r.EarlyZPass"));
 	static auto* ShaderPipelinesCvar = IConsoleManager::Get().FindTConsoleVariableDataInt(TEXT("r.ShaderPipelines"));
-	static auto* InstancedStereoCvar = IConsoleManager::Get().FindTConsoleVariableDataInt(TEXT("vr.InstancedStereo"));
 
 	bool bMobileHDR = IsMobileHDR();
 	bool bMobileHDR32bpp = IsMobileHDR32bpp();
 	int32 DesiredStaticDrawListsEarlyZPassMode = EarlyZPassCvar->GetValueOnRenderThread();
 	int32 DesiredStaticDrawShaderPipelines = ShaderPipelinesCvar->GetValueOnRenderThread();
-	bool bDesiredStaticDrawInstancedStereo = RHISupportsInstancedStereo(GShaderPlatformForFeatureLevel[FeatureLevel]) &&
-		GEngine->IsStereoscopic3D() && InstancedStereoCvar->GetValueOnAnyThread();
 
 	if (bScenesPrimitivesNeedStaticMeshElementUpdate
 		|| bStaticDrawListsMobileHDR != bMobileHDR
 		|| bStaticDrawListsMobileHDR32bpp != bMobileHDR32bpp
 		|| StaticDrawShaderPipelines != DesiredStaticDrawShaderPipelines
-		|| StaticDrawListsEarlyZPassMode != DesiredStaticDrawListsEarlyZPassMode
-		|| bStaticDrawInstancedStereo != bDesiredStaticDrawInstancedStereo)
+		|| StaticDrawListsEarlyZPassMode != DesiredStaticDrawListsEarlyZPassMode)
 	{
 		// Mark all primitives as needing an update
 		// Note: Only visible primitives will actually update their static mesh elements
@@ -2890,7 +2882,6 @@ void FScene::ConditionalMarkStaticMeshElementsForUpdate()
 		bStaticDrawListsMobileHDR32bpp = bMobileHDR32bpp;
 		StaticDrawListsEarlyZPassMode = DesiredStaticDrawListsEarlyZPassMode;
 		StaticDrawShaderPipelines = DesiredStaticDrawShaderPipelines;
-		bStaticDrawInstancedStereo = bDesiredStaticDrawInstancedStereo;
 	}
 }
 
diff --git a/Engine/Source/Runtime/Renderer/Private/SceneHitProxyRendering.cpp b/Engine/Source/Runtime/Renderer/Private/SceneHitProxyRendering.cpp
--- a/Engine/Source/Runtime/Renderer/Private/SceneHitProxyRendering.cpp
+++ b/Engine/Source/Runtime/Renderer/Private/SceneHitProxyRendering.cpp
@@ -666,7 +666,6 @@ void FHitProxyMeshProcessor::Process(
 		HitProxyPassShaders,
 		MeshFillMode,
 		MeshCullMode,
-		1,
 		SortKey,
 		EMeshPassFeatures::Default,
 		ShaderElementData);
@@ -782,7 +781,6 @@ void FEditorSelectionMeshProcessor::Process(
 		HitProxyPassShaders,
 		MeshFillMode,
 		MeshCullMode,
-		1,
 		SortKey,
 		EMeshPassFeatures::Default,
 		ShaderElementData);
diff --git a/Engine/Source/Runtime/Renderer/Private/ScenePrivate.h b/Engine/Source/Runtime/Renderer/Private/ScenePrivate.h
--- a/Engine/Source/Runtime/Renderer/Private/ScenePrivate.h
+++ b/Engine/Source/Runtime/Renderer/Private/ScenePrivate.h
@@ -2500,9 +2500,6 @@ public:
 	/** Whether the ShaderPipelines were enabled when the static draw lists were built. */
 	int32 StaticDrawShaderPipelines;
 
-	/** Whether instanced stereoscopy was enabled when the static draw lists were built. */
-	bool bStaticDrawInstancedStereo;
-
 	/** True if a change to SkyLight / Lighting has occurred that requires static draw lists to be updated. */
 	bool bScenesPrimitivesNeedStaticMeshElementUpdate;
 
diff --git a/Engine/Source/Runtime/Renderer/Private/ShadowDepthRendering.cpp b/Engine/Source/Runtime/Renderer/Private/ShadowDepthRendering.cpp
--- a/Engine/Source/Runtime/Renderer/Private/ShadowDepthRendering.cpp
+++ b/Engine/Source/Runtime/Renderer/Private/ShadowDepthRendering.cpp
@@ -1775,7 +1775,6 @@ void FShadowDepthPassMeshProcessor::Process(
 			ShadowDepthPassShaders,
 			MeshFillMode,
 			MeshCullMode,
-			1,
 			SortKey,
 			bUsePositionOnlyVS ? EMeshPassFeatures::PositionOnly : EMeshPassFeatures::Default,
 			ShaderElementData);
diff --git a/Engine/Source/Runtime/Renderer/Private/ShadowRendering.cpp b/Engine/Source/Runtime/Renderer/Private/ShadowRendering.cpp
--- a/Engine/Source/Runtime/Renderer/Private/ShadowRendering.cpp
+++ b/Engine/Source/Runtime/Renderer/Private/ShadowRendering.cpp
@@ -614,7 +614,7 @@ void FProjectedShadowInfo::SetupProjectionStencilMask(
 		const FShadowMeshDrawCommandPass& ProjectionStencilingPass = ProjectionStencilingPasses[ViewIndex];
 		if (ProjectionStencilingPass.VisibleMeshDrawCommands.Num() > 0)
 		{
-			SubmitMeshDrawCommands(ProjectionStencilingPass.VisibleMeshDrawCommands, ProjectionStencilingPass.PrimitiveIdVertexBuffer, 0, bDynamicInstancing, RHICmdList);
+			SubmitMeshDrawCommands(ProjectionStencilingPass.VisibleMeshDrawCommands, ProjectionStencilingPass.PrimitiveIdVertexBuffer, 0, bDynamicInstancing, 1, RHICmdList);
 		}
 
 		// Restore viewport
@@ -754,7 +754,7 @@ void FProjectedShadowInfo::SetupProjectionStencilMask(
 			const FShadowMeshDrawCommandPass& ProjectionStencilingPass = ProjectionStencilingPasses[ViewIndex];
 			if (ProjectionStencilingPass.VisibleMeshDrawCommands.Num() > 0)
 			{
-				SubmitMeshDrawCommands(ProjectionStencilingPass.VisibleMeshDrawCommands, ProjectionStencilingPass.PrimitiveIdVertexBuffer, 0, bDynamicInstancing, RHICmdList);
+				SubmitMeshDrawCommands(ProjectionStencilingPass.VisibleMeshDrawCommands, ProjectionStencilingPass.PrimitiveIdVertexBuffer, 0, bDynamicInstancing, 1, RHICmdList);
 			}
 		}
 	}
diff --git a/Engine/Source/Runtime/Renderer/Private/ShadowSetup.cpp b/Engine/Source/Runtime/Renderer/Private/ShadowSetup.cpp
--- a/Engine/Source/Runtime/Renderer/Private/ShadowSetup.cpp
+++ b/Engine/Source/Runtime/Renderer/Private/ShadowSetup.cpp
@@ -1299,7 +1299,7 @@ void FProjectedShadowInfo::SetupMeshDrawCommandsForProjectionStenciling(FSceneRe
 
 			ApplyViewOverridesToMeshDrawCommands(View, ProjectionStencilingPass.VisibleMeshDrawCommands);
 
-			SortAndMergeDynamicPassMeshDrawCommands(Renderer.FeatureLevel, ProjectionStencilingPass.VisibleMeshDrawCommands, DynamicMeshDrawCommandStorage, ProjectionStencilingPass.PrimitiveIdVertexBuffer);
+			SortAndMergeDynamicPassMeshDrawCommands(Renderer.FeatureLevel, ProjectionStencilingPass.VisibleMeshDrawCommands, DynamicMeshDrawCommandStorage, ProjectionStencilingPass.PrimitiveIdVertexBuffer, 1);
 		}
 	}
 }
diff --git a/Engine/Source/Runtime/Renderer/Private/TranslucentLighting.cpp b/Engine/Source/Runtime/Renderer/Private/TranslucentLighting.cpp
--- a/Engine/Source/Runtime/Renderer/Private/TranslucentLighting.cpp
+++ b/Engine/Source/Runtime/Renderer/Private/TranslucentLighting.cpp
@@ -479,8 +479,6 @@ void FTranslucencyDepthPassMeshProcessor::Process(
 
 	const FMeshDrawCommandSortKey SortKey = CalculateMeshStaticSortKey(PassShaders.VertexShader, PassShaders.PixelShader);
 
-	const uint32 InstanceFactor = 1;
-
 	BuildMeshDrawCommands(
 		MeshBatch,
 		BatchElementMask,
@@ -491,7 +489,6 @@ void FTranslucencyDepthPassMeshProcessor::Process(
 		PassShaders,
 		MeshFillMode,
 		MeshCullMode,
-		InstanceFactor,
 		SortKey,
 		EMeshPassFeatures::Default,
 		ShaderElementData);
@@ -624,8 +621,8 @@ void FProjectedShadowInfo::RenderTranslucencyDepths(FRHICommandList& RHICmdList,
 
 			FVertexBufferRHIParamRef PrimitiveIdVertexBuffer = nullptr;
 			ApplyViewOverridesToMeshDrawCommands(*ShadowDepthView, VisibleMeshDrawCommands);
-			SortAndMergeDynamicPassMeshDrawCommands(SceneRenderer->FeatureLevel, VisibleMeshDrawCommands, DynamicMeshDrawCommandStorage, PrimitiveIdVertexBuffer);
-			SubmitMeshDrawCommands(VisibleMeshDrawCommands, PrimitiveIdVertexBuffer, 0, bDynamicInstancing, RHICmdList);
+			SortAndMergeDynamicPassMeshDrawCommands(SceneRenderer->FeatureLevel, VisibleMeshDrawCommands, DynamicMeshDrawCommandStorage, PrimitiveIdVertexBuffer, 1);
+			SubmitMeshDrawCommands(VisibleMeshDrawCommands, PrimitiveIdVertexBuffer, 0, bDynamicInstancing, 1, RHICmdList);
 		}
 	}
 }
diff --git a/Engine/Source/Runtime/Renderer/Private/VelocityRendering.cpp b/Engine/Source/Runtime/Renderer/Private/VelocityRendering.cpp
--- a/Engine/Source/Runtime/Renderer/Private/VelocityRendering.cpp
+++ b/Engine/Source/Runtime/Renderer/Private/VelocityRendering.cpp
@@ -688,9 +688,6 @@ void FVelocityMeshProcessor::Process(
 	FMeshMaterialShaderElementData ShaderElementData;
 	ShaderElementData.InitializeMeshMaterialData(ViewIfDynamicMeshCommand, PrimitiveSceneProxy, MeshBatch, StaticMeshId, false);
 
-	const bool bIsInstancedStereo = ViewIfDynamicMeshCommand ? ViewIfDynamicMeshCommand->IsInstancedStereoPass() : Scene->bStaticDrawInstancedStereo;
-	const int32 InstanceFactor = bIsInstancedStereo ? 2 : 1;
-
 	const FMeshDrawCommandSortKey SortKey = CalculateMeshStaticSortKey(VelocityPassShaders.VertexShader, VelocityPassShaders.PixelShader);
 
 	BuildMeshDrawCommands(
@@ -703,7 +700,6 @@ void FVelocityMeshProcessor::Process(
 		VelocityPassShaders,
 		MeshFillMode,
 		MeshCullMode,
-		InstanceFactor,
 		SortKey,
 		EMeshPassFeatures::Default,
 		ShaderElementData);
diff --git a/Engine/Source/Runtime/Renderer/Private/VolumetricFogVoxelization.cpp b/Engine/Source/Runtime/Renderer/Private/VolumetricFogVoxelization.cpp
--- a/Engine/Source/Runtime/Renderer/Private/VolumetricFogVoxelization.cpp
+++ b/Engine/Source/Runtime/Renderer/Private/VolumetricFogVoxelization.cpp
@@ -511,7 +511,6 @@ void FVoxelizeVolumeMeshProcessor::Process(
 				PassShaders,
 				MeshFillMode,
 				MeshCullMode,
-				1,
 				SortKey,
 				EMeshPassFeatures::Default,
 				ShaderElementData);
diff --git a/Engine/Source/Runtime/Renderer/Public/MeshPassProcessor.h b/Engine/Source/Runtime/Renderer/Public/MeshPassProcessor.h
--- a/Engine/Source/Runtime/Renderer/Public/MeshPassProcessor.h
+++ b/Engine/Source/Runtime/Renderer/Public/MeshPassProcessor.h
@@ -351,7 +351,6 @@ public:
 	uint32 FirstIndex;
 	uint32 NumPrimitives;
 	uint32 NumInstances;
-	uint32 InstanceFactor;
 	uint32 BaseVertexIndex;
 	uint32 NumVertices;
 	FVertexBufferRHIParamRef IndirectArgsBuffer;
@@ -376,7 +375,6 @@ public:
 			&& FirstIndex == Rhs.FirstIndex
 			&& NumPrimitives == Rhs.NumPrimitives
 			&& NumInstances == Rhs.NumInstances
-			&& InstanceFactor == Rhs.InstanceFactor
 			&& BaseVertexIndex == Rhs.BaseVertexIndex
 			&& NumVertices == Rhs.NumVertices
 			&& IndirectArgsBuffer == Rhs.IndirectArgsBuffer;
@@ -398,7 +396,7 @@ public:
 	}
 
 	/** Called when the mesh draw command is complete. */
-	RENDERER_API void SetDrawParametersAndFinalize(const FMeshBatch& MeshBatch, int32 BatchElementIndex, int32 BatchInstanceFactor, bool bDoSetupPsoStateForRasterization);
+	RENDERER_API void SetDrawParametersAndFinalize(const FMeshBatch& MeshBatch, int32 BatchElementIndex, bool bDoSetupPsoStateForRasterization);
 
 	void Finalize(bool bDoSetupPsoStateForRasterization)
 	{
@@ -415,7 +413,8 @@ public:
 	static void SubmitDraw(
 		const FMeshDrawCommand& RESTRICT MeshDrawCommand, 
 		FVertexBufferRHIParamRef ScenePrimitiveIdsBuffer, 
-		int32 PrimitiveIdOffset, 
+		int32 PrimitiveIdOffset,
+		uint32 InstanceFactor,
 		FRHICommandList& CommandList, 
 		class FMeshDrawCommandStateCache& RESTRICT StateCache);
 
@@ -503,7 +502,6 @@ public:
 		int32 DrawPrimitiveId,
 		ERasterizerFillMode MeshFillMode,
 		ERasterizerCullMode MeshCullMode,
-		int32 InstanceFactor,
 		FMeshDrawCommandSortKey SortKey,
 		FMeshDrawCommand& MeshDrawCommand,
 		bool bDoSetupPsoStateForRasterization) = 0;
@@ -601,12 +599,11 @@ public:
 		int32 DrawPrimitiveId,
 		ERasterizerFillMode MeshFillMode,
 		ERasterizerCullMode MeshCullMode,
-		int32 InstanceFactor,
 		FMeshDrawCommandSortKey SortKey,
 		FMeshDrawCommand& MeshDrawCommand,
 		bool bDoSetupPsoStateForRasterization) override final
 	{
-		MeshDrawCommand.SetDrawParametersAndFinalize(MeshBatch, BatchElementIndex, InstanceFactor, bDoSetupPsoStateForRasterization);
+		MeshDrawCommand.SetDrawParametersAndFinalize(MeshBatch, BatchElementIndex, bDoSetupPsoStateForRasterization);
 
 		FVisibleMeshDrawCommand NewVisibleMeshDrawCommand;
 		//@todo MeshCommandPipeline - assign usable state ID for dynamic path draws
@@ -671,7 +668,6 @@ public:
 		int32 DrawPrimitiveId,
 		ERasterizerFillMode MeshFillMode,
 		ERasterizerCullMode MeshCullMode,
-		int32 InstanceFactor,
 		FMeshDrawCommandSortKey SortKey,
 		FMeshDrawCommand& MeshDrawCommand,
 		bool bDoSetupPsoStateForRasterization) override final;
@@ -905,7 +901,6 @@ public:
 		PassShadersType PassShaders,
 		ERasterizerFillMode MeshFillMode,
 		ERasterizerCullMode MeshCullMode,
-		int32 InstanceFactor,
 		FMeshDrawCommandSortKey SortKey,
 		EMeshPassFeatures MeshPassFeatures,
 		const ShaderElementDataType& ShaderElementData);
@@ -921,7 +916,6 @@ public:
 		PassShadersType PassShaders,
 		ERasterizerFillMode MeshFillMode,
 		ERasterizerCullMode MeshCullMode,
-		int32 InstanceFactor,
 		FMeshDrawCommandSortKey SortKey,
 		EMeshPassFeatures MeshPassFeatures,
 		const ShaderElementDataType& ShaderElementData);
@@ -990,6 +984,7 @@ extern void SubmitMeshDrawCommands(
 	FVertexBufferRHIParamRef PrimitiveIdsBuffer,
 	int32 BasePrimitiveIdsOffset,
 	bool bDynamicInstancing,
+	uint32 InstanceFactor,
 	FRHICommandList& RHICmdList);
 
 extern void SubmitMeshDrawCommandsRange(
@@ -999,12 +994,14 @@ extern void SubmitMeshDrawCommandsRange(
 	bool bDynamicInstancing,
 	int32 StartIndex,
 	int32 NumMeshDrawCommands,
+	uint32 InstanceFactor,
 	FRHICommandList& RHICmdList);
 
 RENDERER_API extern void DrawDynamicMeshPassPrivate(
 	const FSceneView& View,
 	FRHICommandList& RHICmdList,
 	FMeshCommandOneFrameArray& VisibleMeshDrawCommands,
-	FDynamicMeshDrawCommandStorage& DynamicMeshDrawCommandStorage);
+	FDynamicMeshDrawCommandStorage& DynamicMeshDrawCommandStorage,
+	uint32 InstanceFactor);
 
 RENDERER_API extern FMeshDrawCommandSortKey CalculateMeshStaticSortKey(const FMeshMaterialShader* VertexShader, const FMeshMaterialShader* PixelShader);
\ No newline at end of file
diff --git a/Engine/Source/Runtime/Renderer/Public/MeshPassProcessor.inl b/Engine/Source/Runtime/Renderer/Public/MeshPassProcessor.inl
--- a/Engine/Source/Runtime/Renderer/Public/MeshPassProcessor.inl
+++ b/Engine/Source/Runtime/Renderer/Public/MeshPassProcessor.inl
@@ -17,7 +17,6 @@ void FMeshPassProcessor::BuildMeshDrawCommands(
 	PassShadersType PassShaders,
 	ERasterizerFillMode MeshFillMode,
 	ERasterizerCullMode MeshCullMode,
-	int32 InstanceFactor,
 	FMeshDrawCommandSortKey SortKey,
 	EMeshPassFeatures MeshPassFeatures,
 	const ShaderElementDataType& ShaderElementData)
@@ -131,7 +130,7 @@ void FMeshPassProcessor::BuildMeshDrawCommands(
 
 			const int32 DrawPrimitiveId = GetDrawCommandPrimitiveId(PrimitiveSceneInfo, BatchElement);
 
-			DrawListContext->FinalizeCommand(MeshBatch, BatchElementIndex, DrawPrimitiveId, MeshFillMode, MeshCullMode, InstanceFactor, SortKey, MeshDrawCommand, true);
+			DrawListContext->FinalizeCommand(MeshBatch, BatchElementIndex, DrawPrimitiveId, MeshFillMode, MeshCullMode, SortKey, MeshDrawCommand, true);
 		}
 	}
 }
@@ -150,7 +149,9 @@ void DrawDynamicMeshPass(const FSceneView& View, FRHICommandList& RHICmdList, co
 
 	BuildPassProcessorLambda(&DynamicMeshPassContext);
 
-	DrawDynamicMeshPassPrivate(View, RHICmdList, VisibleMeshDrawCommands, DynamicMeshDrawCommandStorage);
+	// We assume all dynamic passes are in stereo if it is enabled in the view, so we apply ISR to them
+	const uint32 InstanceFactor = View.IsInstancedStereoPass() ? 2 : 1;
+	DrawDynamicMeshPassPrivate(View, RHICmdList, VisibleMeshDrawCommands, DynamicMeshDrawCommandStorage, InstanceFactor);
 }
 
 template<typename PassShadersType, typename ShaderElementDataType>
@@ -164,7 +165,6 @@ void FMeshPassProcessor::BuildRayTracingDrawCommands(
 	PassShadersType PassShaders,
 	ERasterizerFillMode MeshFillMode,
 	ERasterizerCullMode MeshCullMode,
-	int32 InstanceFactor,
 	FMeshDrawCommandSortKey SortKey,
 	EMeshPassFeatures MeshPassFeatures,
 	const ShaderElementDataType& ShaderElementData)
@@ -211,7 +211,7 @@ void FMeshPassProcessor::BuildRayTracingDrawCommands(
 
 			const int32 DrawPrimitiveId = 0;
 			const int32 ScenePrimitiveId = 0;
-			DrawListContext->FinalizeCommand(MeshBatch, BatchElementIndex, DrawPrimitiveId, MeshFillMode, MeshCullMode, InstanceFactor, SortKey, MeshDrawCommand, false);
+			DrawListContext->FinalizeCommand(MeshBatch, BatchElementIndex, DrawPrimitiveId, MeshFillMode, MeshCullMode, SortKey, MeshDrawCommand, false);
 		}
 	}
 }
\ No newline at end of file
