Fixed rendering corruption and crashes when r.DoInitViewsLightingAfterPrepass 1. This was caused by using FGlobalDynamic(Vertex/Index/Read)Buffer before it was committed (also commit was overlapping the parallel draw cmd list):
* IntiViews
* Dispatch prepass parallel cmd list (flickering, as global buffers weren't committed).
* InitShadows.
* Commit global buffers (race condition and potential crash).

This was fixed by:
* Renamed FNiagaraGlobalReadBuffer to FGlobalDynamicReadBuffer and moved it out from Niagara.
* Removed all global buffer singletons and placed them inside appropriate renderers as static members.
* They are also separated into two - one for InitViews and one for InitShadows, so we can dispatch prepass while writing to the shadow global buffers.
* All global buffers are now passed as a parameter and/or through mesh collector.

#rnx

diff --git a/Engine/Plugins/FX/Niagara/Source/Niagara/Private/NiagaraRenderer.cpp b/Engine/Plugins/FX/Niagara/Source/Niagara/Private/NiagaraRenderer.cpp
--- a/Engine/Plugins/FX/Niagara/Source/Niagara/Private/NiagaraRenderer.cpp
+++ b/Engine/Plugins/FX/Niagara/Source/Niagara/Private/NiagaraRenderer.cpp
@@ -7,6 +7,7 @@
 #include "NiagaraStats.h"
 #include "NiagaraVertexFactory.h"
 #include "Engine/Engine.h"
+#include "DynamicBufferAllocator.h"
 
 DECLARE_CYCLE_STAT(TEXT("Generate Particle Lights"), STAT_NiagaraGenLights, STATGROUP_Niagara);
 DECLARE_CYCLE_STAT(TEXT("Sort Particles"), STAT_NiagaraSortParticles, STATGROUP_Niagara);
@@ -88,7 +89,7 @@ void NiagaraRenderer::Release()
 	);
 }
 
-void NiagaraRenderer::SortIndices(ENiagaraSortMode SortMode, int32 SortAttributeOffset, const FNiagaraDataBuffer& Buffer, const FMatrix& LocalToWorld, const FSceneView* View, FNiagaraGlobalReadBuffer::FAllocation& OutIndices)const
+void NiagaraRenderer::SortIndices(ENiagaraSortMode SortMode, int32 SortAttributeOffset, const FNiagaraDataBuffer& Buffer, const FMatrix& LocalToWorld, const FSceneView* View, FGlobalDynamicReadBuffer::FAllocation& OutIndices)const
 {
 	SCOPE_CYCLE_COUNTER(STAT_NiagaraSortParticles);
 
@@ -323,243 +324,4 @@ const TArray<FNiagaraVariable>& NiagaraRendererLights::GetOptionalAttributes()
 	return Properties->GetOptionalAttributes();
 }
 
-#endif
-
-
-
-
-//////////////////////////////////////////////////////////////////////////
-// FNiagaraGlobalReadBuffer
-//////////////////////////////////////////////////////////////////////////
-
-int32 GMaxNiagaraRenderingBytesAllocatedPerFrame = 32 * 1024 * 1024;
-
-FAutoConsoleVariableRef CVarMaxNiagaraRenderingBytesAllocatedPerFrame(
-	TEXT("fx.MaxRenderingBytesAllocatedPerFrame"),
-	GMaxNiagaraRenderingBytesAllocatedPerFrame,
-	TEXT("The maximum number of transient rendering buffer bytes to allocate before we start panic logging who is doing the allocations"));
-
-
-int32 GMinNiagaraRenderingBufferSize = 8 * 1024;
-FAutoConsoleVariableRef CVarMinNiagaraRenderingBufferSize(
-	TEXT("fx.MinNiagaraRenderingBufferSize"),
-	GMinNiagaraRenderingBufferSize,
-	TEXT("The minimum size (in instances) to allocate in blocks for niagara rendering buffers."));
-
-
-struct FDynamicReadBufferPool
-{
-	/** List of vertex buffers. */
-	TIndirectArray<FDynamicAllocReadBuffer> Buffers;
-	/** The current buffer from which allocations are being made. */
-	FDynamicAllocReadBuffer* CurrentBuffer;
-
-	/** Default constructor. */
-	FDynamicReadBufferPool()
-		: CurrentBuffer(NULL)
-	{
-	}
-
-	/** Destructor. */
-	~FDynamicReadBufferPool()
-	{
-		int32 NumVertexBuffers = Buffers.Num();
-		for (int32 BufferIndex = 0; BufferIndex < NumVertexBuffers; ++BufferIndex)
-		{
-			Buffers[BufferIndex].Release();
-		}
-	}
-
-	FCriticalSection CriticalSection;
-};
-
-
-
-FNiagaraGlobalReadBuffer::FNiagaraGlobalReadBuffer()
-	: TotalAllocatedSinceLastCommit(0)
-{
-	FloatBufferPool = new FDynamicReadBufferPool();
-	Int32BufferPool = new FDynamicReadBufferPool();
-
-	CommitCallbackHandle = GEngine->GetPreRenderDelegate().AddRaw(this, &FNiagaraGlobalReadBuffer::Commit);
-}
-
-FNiagaraGlobalReadBuffer::~FNiagaraGlobalReadBuffer()
-{
-	Cleanup();
-
-	//GEngine->GetPostInitViewsDelegate().Remove(CommitCallbackHandle);
-}
-
-void FNiagaraGlobalReadBuffer::Cleanup()
-{
-	if (FloatBufferPool)
-	{
-		UE_LOG(LogNiagara, Log, TEXT("FNiagaraGlobalReadBuffer::Cleanup()"));
-		delete FloatBufferPool;
-		FloatBufferPool = nullptr;
-	}
-
-	if (Int32BufferPool)
-	{
-		delete Int32BufferPool;
-		Int32BufferPool = nullptr;
-	}
-}
-void FNiagaraGlobalReadBuffer::InitRHI()
-{
-	UE_LOG(LogNiagara, Log, TEXT("FNiagaraGlobalReadBuffer::InitRHI"));
-}
-
-void FNiagaraGlobalReadBuffer::ReleaseRHI() 
-{
-	UE_LOG(LogNiagara, Log, TEXT("FNiagaraGlobalReadBuffer::ReleaseRHI"));
-	Cleanup();
-}
-
-FNiagaraGlobalReadBuffer::FAllocation FNiagaraGlobalReadBuffer::AllocateFloat(uint32 Num)
-{
-	FScopeLock ScopeLock(&FloatBufferPool->CriticalSection);
-
-	FAllocation Allocation;
-
-	TotalAllocatedSinceLastCommit += Num;
-	if (IsRenderAlarmLoggingEnabled())
-	{
-		UE_LOG(LogNiagara, Warning, TEXT("FNiagaraGlobalReadBuffer::AllocateFloat(%u), will have allocated %u total this frame"), Num, TotalAllocatedSinceLastCommit);
-	}
-	uint32 SizeInBytes = sizeof(float) * Num;
-	FDynamicAllocReadBuffer* Buffer = FloatBufferPool->CurrentBuffer;
-	if (Buffer == NULL || Buffer->AllocatedByteCount + SizeInBytes > Buffer->NumBytes)
-	{
-		// Find a buffer in the pool big enough to service the request.
-		Buffer = NULL;
-		for (int32 BufferIndex = 0, NumBuffers = FloatBufferPool->Buffers.Num(); BufferIndex < NumBuffers; ++BufferIndex)
-		{
-			FDynamicAllocReadBuffer& BufferToCheck = FloatBufferPool->Buffers[BufferIndex];
-			if (BufferToCheck.AllocatedByteCount + SizeInBytes <= BufferToCheck.NumBytes)
-			{
-				Buffer = &BufferToCheck;
-				break;
-			}
-		}
-
-		// Create a new vertex buffer if needed.
-		if (Buffer == NULL)
-		{
-			uint32 NewBufferSize = FMath::Max(Num, (uint32)GMinNiagaraRenderingBufferSize);
-			Buffer = new FDynamicAllocReadBuffer();
-			FloatBufferPool->Buffers.Add(Buffer);
-			Buffer->Initialize(sizeof(float), NewBufferSize, PF_R32_FLOAT, BUF_Dynamic);
-		}
-
-		// Lock the buffer if needed.
-		if (Buffer->MappedBuffer == NULL)
-		{
-			Buffer->Lock();
-		}
-
-		// Remember this buffer, we'll try to allocate out of it in the future.
-		FloatBufferPool->CurrentBuffer = Buffer;
-	}
-
-	check(Buffer != NULL);
-	checkf(Buffer->AllocatedByteCount + SizeInBytes <= Buffer->NumBytes, TEXT("Niagara global float buffer allocation failed: BufferSize=%d AllocatedByteCount=%d SizeInBytes=%d"), Buffer->NumBytes, Buffer->AllocatedByteCount, SizeInBytes);
-	Allocation.Buffer = Buffer->MappedBuffer + Buffer->AllocatedByteCount;
-	Allocation.ReadBuffer = Buffer;
-	Allocation.FirstIndex = Buffer->AllocatedByteCount;
-	Buffer->AllocatedByteCount += SizeInBytes;
-
-	return Allocation;
-}
-
-FNiagaraGlobalReadBuffer::FAllocation FNiagaraGlobalReadBuffer::AllocateInt32(uint32 Num)
-{
-	FScopeLock ScopeLock(&Int32BufferPool->CriticalSection);
-	FAllocation Allocation;
-
-	TotalAllocatedSinceLastCommit += Num;
-	if (IsRenderAlarmLoggingEnabled())
-	{
-		UE_LOG(LogNiagara, Warning, TEXT("FNiagaraGlobalReadBuffer::AllocateInt32(%u), will have allocated %u total this frame"), Num, TotalAllocatedSinceLastCommit);
-	}
-	uint32 SizeInBytes = sizeof(int32) * Num;
-	FDynamicAllocReadBuffer* Buffer = Int32BufferPool->CurrentBuffer;
-	if (Buffer == NULL || Buffer->AllocatedByteCount + SizeInBytes > Buffer->NumBytes)
-	{
-		// Find a buffer in the pool big enough to service the request.
-		Buffer = NULL;
-		for (int32 BufferIndex = 0, NumBuffers = Int32BufferPool->Buffers.Num(); BufferIndex < NumBuffers; ++BufferIndex)
-		{
-			FDynamicAllocReadBuffer& BufferToCheck = Int32BufferPool->Buffers[BufferIndex];
-			if (BufferToCheck.AllocatedByteCount + SizeInBytes <= BufferToCheck.NumBytes)
-			{
-				Buffer = &BufferToCheck;
-				break;
-			}
-		}
-
-		// Create a new vertex buffer if needed.
-		if (Buffer == NULL)
-		{
-			uint32 NewBufferSize = FMath::Max(Num, (uint32)GMinNiagaraRenderingBufferSize);
-			Buffer = new FDynamicAllocReadBuffer();
-			Int32BufferPool->Buffers.Add(Buffer);
-			Buffer->Initialize(sizeof(int32), NewBufferSize, PF_R32_SINT, BUF_Dynamic);
-		}
-
-		// Lock the buffer if needed.
-		if (Buffer->MappedBuffer == NULL)
-		{
-			Buffer->Lock();
-		}
-
-		// Remember this buffer, we'll try to allocate out of it in the future.
-		Int32BufferPool->CurrentBuffer = Buffer;
-	}
-
-	check(Buffer != NULL);
-	checkf(Buffer->AllocatedByteCount + SizeInBytes <= Buffer->NumBytes, TEXT("Niagara global int32 buffer allocation failed: BufferSize=%d AllocatedByteCount=%d SizeInBytes=%d"), Buffer->NumBytes, Buffer->AllocatedByteCount, SizeInBytes);
-	Allocation.Buffer = Buffer->MappedBuffer + Buffer->AllocatedByteCount;
-	Allocation.ReadBuffer = Buffer;
-	Allocation.FirstIndex = Buffer->AllocatedByteCount;
-	Buffer->AllocatedByteCount += SizeInBytes;
-
-	return Allocation;
-}
-
-bool FNiagaraGlobalReadBuffer::IsRenderAlarmLoggingEnabled() const
-{
-	return GMaxNiagaraRenderingBytesAllocatedPerFrame > 0 && TotalAllocatedSinceLastCommit >= (size_t)GMaxNiagaraRenderingBytesAllocatedPerFrame;
-}
-
-void FNiagaraGlobalReadBuffer::Commit()
-{
-	for (int32 BufferIndex = 0, NumBuffers = FloatBufferPool->Buffers.Num(); BufferIndex < NumBuffers; ++BufferIndex)
-	{
-		FDynamicAllocReadBuffer& Buffer = FloatBufferPool->Buffers[BufferIndex];
-		if (Buffer.MappedBuffer != NULL)
-		{
-			Buffer.Unlock();
-		}
-	}
-	FloatBufferPool->CurrentBuffer = NULL;
-	for (int32 BufferIndex = 0, NumBuffers = Int32BufferPool->Buffers.Num(); BufferIndex < NumBuffers; ++BufferIndex)
-	{
-		FDynamicAllocReadBuffer& Buffer = Int32BufferPool->Buffers[BufferIndex];
-		if (Buffer.MappedBuffer != NULL)
-		{
-			Buffer.Unlock();
-		}
-	}
-	Int32BufferPool->CurrentBuffer = NULL;
-	TotalAllocatedSinceLastCommit = 0;
-}
-
-
-FNiagaraGlobalReadBuffer& FNiagaraGlobalReadBuffer::Get()
-{
-	check(IsInRenderingThread());
-	static TGlobalResource<FNiagaraGlobalReadBuffer> GlobalDynamicReadBuffer;
-	return GlobalDynamicReadBuffer;
-}
+#endif
\ No newline at end of file
diff --git a/Engine/Plugins/FX/Niagara/Source/Niagara/Private/NiagaraRendererMeshes.cpp b/Engine/Plugins/FX/Niagara/Source/Niagara/Private/NiagaraRendererMeshes.cpp
--- a/Engine/Plugins/FX/Niagara/Source/Niagara/Private/NiagaraRendererMeshes.cpp
+++ b/Engine/Plugins/FX/Niagara/Source/Niagara/Private/NiagaraRendererMeshes.cpp
@@ -133,12 +133,14 @@ void NiagaraRendererMeshes::GetDynamicMeshElements(const TArray<const FSceneView
 	int32 NumInstances = DynamicDataMesh->RTParticleData.GetNumInstances();
 
 	int32 TotalFloatSize = DynamicDataMesh->RTParticleData.GetFloatBuffer().Num() / sizeof(float);
-	FNiagaraGlobalReadBuffer::FAllocation ParticleData;
+
+	FGlobalDynamicReadBuffer& DynamicReadBuffer = Collector.GetDynamicReadBuffer();
+	FGlobalDynamicReadBuffer::FAllocation ParticleData;
 	
 	//For cpu sims we allocate render buffers from the global pool. GPU sims own their own.
 	if (DynamicDataMesh->DataSet->GetSimTarget() == ENiagaraSimTarget::CPUSim)
 	{
-		ParticleData = FNiagaraGlobalReadBuffer::Get().AllocateFloat(TotalFloatSize);
+		ParticleData = DynamicReadBuffer.AllocateFloat(TotalFloatSize);
 		FMemory::Memcpy(ParticleData.Buffer, DynamicDataMesh->RTParticleData.GetFloatBuffer().GetData(), DynamicDataMesh->RTParticleData.GetFloatBuffer().Num());
 	}
 
@@ -240,7 +242,7 @@ void NiagaraRendererMeshes::GetDynamicMeshElements(const TArray<const FSceneView
 				}
 
 				//Sort particles if needed.
-				FNiagaraGlobalReadBuffer::FAllocation SortedIndices;
+				FGlobalDynamicReadBuffer::FAllocation SortedIndices;
 				CollectorResources.VertexFactory.SetSortedIndices(nullptr, 0xFFFFFFFF);
 				if (DynamicDataMesh->DataSet->GetSimTarget() == ENiagaraSimTarget::CPUSim)//TODO: Compute shader for sorting gpu sims and larger cpu sims.
 				{
@@ -252,7 +254,7 @@ void NiagaraRendererMeshes::GetDynamicMeshElements(const TArray<const FSceneView
 						int32 SortAttributeOffest = bCustomSortMode ? CustomSortingOffset : PositionOffset;
 						if (SortMode != ENiagaraSortMode::None && SortAttributeOffest != INDEX_NONE)
 						{
-							SortedIndices = FNiagaraGlobalReadBuffer::Get().AllocateInt32(NumInstances);
+							SortedIndices = DynamicReadBuffer.AllocateInt32(NumInstances);
 							SortIndices(SortMode, SortAttributeOffest, DynamicDataMesh->RTParticleData, SceneProxy->GetLocalToWorld(), View, SortedIndices);
 							CollectorResources.VertexFactory.SetSortedIndices(SortedIndices.ReadBuffer->SRV, SortedIndices.FirstIndex / sizeof(float));
 						}
diff --git a/Engine/Plugins/FX/Niagara/Source/Niagara/Private/NiagaraRendererRibbons.cpp b/Engine/Plugins/FX/Niagara/Source/Niagara/Private/NiagaraRendererRibbons.cpp
--- a/Engine/Plugins/FX/Niagara/Source/Niagara/Private/NiagaraRendererRibbons.cpp
+++ b/Engine/Plugins/FX/Niagara/Source/Niagara/Private/NiagaraRendererRibbons.cpp
@@ -103,14 +103,16 @@ void NiagaraRendererRibbons::GetDynamicMeshElements(const TArray<const FSceneVie
 	const bool bIsWireframe = ViewFamily.EngineShowFlags.Wireframe;
 	FMaterialRenderProxy* MaterialRenderProxy = Material->GetRenderProxy();
 
-	FGlobalDynamicIndexBuffer::FAllocation DynamicIndexAllocation = FGlobalDynamicIndexBuffer::Get().Allocate(DynamicDataRibbon->IndexData.Num(), sizeof(int16));
+	FGlobalDynamicIndexBuffer& DynamicIndexBuffer = Collector.GetDynamicIndexBuffer();
+	FGlobalDynamicIndexBuffer::FAllocation DynamicIndexAllocation = DynamicIndexBuffer.Allocate(DynamicDataRibbon->IndexData.Num(), sizeof(int16));
 
 	int32 TotalFloatSize = DynamicDataRibbon->RTParticleData.GetFloatBuffer().Num() / sizeof(float);
-	FNiagaraGlobalReadBuffer::FAllocation ParticleData;
+	FGlobalDynamicReadBuffer::FAllocation ParticleData;
 
 	if (DynamicDataRibbon->DataSet->GetSimTarget() == ENiagaraSimTarget::CPUSim)
 	{
-		ParticleData = FNiagaraGlobalReadBuffer::Get().AllocateFloat(TotalFloatSize);
+		FGlobalDynamicReadBuffer& DynamicReadBuffer = Collector.GetDynamicReadBuffer();
+		ParticleData = DynamicReadBuffer.AllocateFloat(TotalFloatSize);
 		FMemory::Memcpy(ParticleData.Buffer, DynamicDataRibbon->RTParticleData.GetFloatBuffer().GetData(), DynamicDataRibbon->RTParticleData.GetFloatBuffer().Num());
 	}
 
diff --git a/Engine/Plugins/FX/Niagara/Source/Niagara/Private/NiagaraRendererSprites.cpp b/Engine/Plugins/FX/Niagara/Source/Niagara/Private/NiagaraRendererSprites.cpp
--- a/Engine/Plugins/FX/Niagara/Source/Niagara/Private/NiagaraRendererSprites.cpp
+++ b/Engine/Plugins/FX/Niagara/Source/Niagara/Private/NiagaraRendererSprites.cpp
@@ -128,11 +128,13 @@ void NiagaraRendererSprites::GetDynamicMeshElements(const TArray<const FSceneVie
 	int32 NumInstances = DynamicDataSprites->RTParticleData.GetNumInstances();
 
 	int32 TotalFloatSize = DynamicDataSprites->RTParticleData.GetFloatBuffer().Num() / sizeof(float);
-	FNiagaraGlobalReadBuffer::FAllocation ParticleData;
+
+	FGlobalDynamicReadBuffer& DynamicReadBuffer = Collector.GetDynamicReadBuffer();
+	FGlobalDynamicReadBuffer::FAllocation ParticleData;
 	
 	if (DynamicDataSprites->DataSet->GetSimTarget() == ENiagaraSimTarget::CPUSim)
 	{
-		ParticleData = FNiagaraGlobalReadBuffer::Get().AllocateFloat(TotalFloatSize);
+		ParticleData = DynamicReadBuffer.AllocateFloat(TotalFloatSize);
 		FMemory::Memcpy(ParticleData.Buffer, DynamicDataSprites->RTParticleData.GetFloatBuffer().GetData(), DynamicDataSprites->RTParticleData.GetFloatBuffer().Num());
 	}
 
@@ -250,7 +252,7 @@ void NiagaraRendererSprites::GetDynamicMeshElements(const TArray<const FSceneVie
 
 				//Sort particles if needed.
 				EBlendMode BlendMode = MaterialRenderProxy->GetMaterial(VertexFactory->GetFeatureLevel())->GetBlendMode();
-				FNiagaraGlobalReadBuffer::FAllocation SortedIndices;
+				FGlobalDynamicReadBuffer::FAllocation SortedIndices;
 				CollectorResources.VertexFactory.SetSortedIndices(nullptr, 0xFFFFFFFF);
 				if (DynamicDataSprites->DataSet->GetSimTarget() == ENiagaraSimTarget::CPUSim)//TODO: Compute shader for sorting gpu sims and larger cpu sims.
 				{
@@ -262,7 +264,7 @@ void NiagaraRendererSprites::GetDynamicMeshElements(const TArray<const FSceneVie
 						int32 SortAttributeOffest = bCustomSortMode ? CustomSortingOffset : PositionOffset;
 						if (SortMode != ENiagaraSortMode::None && SortAttributeOffest != INDEX_NONE)
 						{
-							SortedIndices = FNiagaraGlobalReadBuffer::Get().AllocateInt32(NumInstances);
+							SortedIndices = DynamicReadBuffer.AllocateInt32(NumInstances);
 							SortIndices(SortMode, SortAttributeOffest, DynamicDataSprites->RTParticleData, SceneProxy->GetLocalToWorld(), View, SortedIndices);
 							CollectorResources.VertexFactory.SetSortedIndices(SortedIndices.ReadBuffer->SRV, SortedIndices.FirstIndex / sizeof(float));
 						}
diff --git a/Engine/Plugins/FX/Niagara/Source/Niagara/Public/NiagaraGlobalReadBuffer.h b/Engine/Plugins/FX/Niagara/Source/Niagara/Public/NiagaraGlobalReadBuffer.h
deleted file mode 100644
diff --git a/Engine/Plugins/FX/Niagara/Source/Niagara/Public/NiagaraRenderer.h b/Engine/Plugins/FX/Niagara/Source/Niagara/Public/NiagaraRenderer.h
--- a/Engine/Plugins/FX/Niagara/Source/Niagara/Public/NiagaraRenderer.h
+++ b/Engine/Plugins/FX/Niagara/Source/Niagara/Public/NiagaraRenderer.h
@@ -19,7 +19,6 @@ NiagaraRenderer.h: Base class for Niagara render modules
 #include "RenderingThread.h"
 #include "SceneView.h"
 #include "NiagaraComponent.h"
-#include "NiagaraGlobalReadBuffer.h"
 
 class FNiagaraDataSet;
 
@@ -157,7 +156,7 @@ public:
 	
 	const FVector& GetBaseExtents() const {	return BaseExtents; }
 
-	void SortIndices(ENiagaraSortMode SortMode, int32 SortAttributeOffset, const FNiagaraDataBuffer& Buffer, const FMatrix& LocalToWorld, const FSceneView* View, FNiagaraGlobalReadBuffer::FAllocation& OutIndices)const;
+	void SortIndices(ENiagaraSortMode SortMode, int32 SortAttributeOffset, const FNiagaraDataBuffer& Buffer, const FMatrix& LocalToWorld, const FSceneView* View, FGlobalDynamicReadBuffer::FAllocation& OutIndices)const;
 
 	static FRWBuffer& GetDummyFloatBuffer(); 
 	static FRWBuffer& GetDummyIntBuffer();
diff --git a/Engine/Plugins/FX/Niagara/Source/NiagaraVertexFactories/Public/NiagaraMeshVertexFactory.h b/Engine/Plugins/FX/Niagara/Source/NiagaraVertexFactories/Public/NiagaraMeshVertexFactory.h
--- a/Engine/Plugins/FX/Niagara/Source/NiagaraVertexFactories/Public/NiagaraMeshVertexFactory.h
+++ b/Engine/Plugins/FX/Niagara/Source/NiagaraVertexFactories/Public/NiagaraMeshVertexFactory.h
@@ -15,7 +15,6 @@ ParticleVertexFactory.h: Particle vertex factory definitions.
 #include "Components.h"
 #include "SceneManagement.h"
 #include "VertexFactory.h"
-#include "NiagaraGlobalReadBuffer.h"
 
 
 
diff --git a/Engine/Plugins/FX/Niagara/Source/NiagaraVertexFactories/Public/NiagaraSpriteVertexFactory.h b/Engine/Plugins/FX/Niagara/Source/NiagaraVertexFactories/Public/NiagaraSpriteVertexFactory.h
--- a/Engine/Plugins/FX/Niagara/Source/NiagaraVertexFactories/Public/NiagaraSpriteVertexFactory.h
+++ b/Engine/Plugins/FX/Niagara/Source/NiagaraVertexFactories/Public/NiagaraSpriteVertexFactory.h
@@ -12,7 +12,6 @@
 #include "NiagaraVertexFactory.h"
 #include "NiagaraDataSet.h"
 #include "SceneView.h"
-#include "NiagaraGlobalReadBuffer.h"
 
 class FMaterial;
 
diff --git a/Engine/Source/Runtime/Engine/Private/ModelRender.cpp b/Engine/Source/Runtime/Engine/Private/ModelRender.cpp
--- a/Engine/Source/Runtime/Engine/Private/ModelRender.cpp
+++ b/Engine/Source/Runtime/Engine/Private/ModelRender.cpp
@@ -350,7 +350,8 @@ public:
 
 						if (TotalIndices > 0)
 						{
-							FGlobalDynamicIndexBuffer::FAllocation IndexAllocation = FGlobalDynamicIndexBuffer::Get().Allocate<uint32>(TotalIndices);
+							FGlobalDynamicIndexBuffer& DynamicIndexBuffer = Collector.GetDynamicIndexBuffer();
+							FGlobalDynamicIndexBuffer::FAllocation IndexAllocation = DynamicIndexBuffer.Allocate<uint32>(TotalIndices);
 
 							if (IndexAllocation.IsValid())
 							{
diff --git a/Engine/Source/Runtime/Engine/Private/Particles/ParticleSystemRender.cpp b/Engine/Source/Runtime/Engine/Private/Particles/ParticleSystemRender.cpp
--- a/Engine/Source/Runtime/Engine/Private/Particles/ParticleSystemRender.cpp
+++ b/Engine/Source/Runtime/Engine/Private/Particles/ParticleSystemRender.cpp
@@ -499,6 +499,8 @@ void FDynamicSpriteEmitterDataBase::BuildViewFillData(
 	int32 InVertexCount, 
 	int32 InVertexSize, 
 	int32 InDynamicParameterVertexStride, 
+	FGlobalDynamicIndexBuffer& DynamicIndexBuffer,
+	FGlobalDynamicVertexBuffer& DynamicVertexBuffer,
 	FGlobalDynamicVertexBuffer::FAllocation& DynamicVertexAllocation,
 	FGlobalDynamicIndexBuffer::FAllocation& DynamicIndexAllocation,
 	FGlobalDynamicVertexBuffer::FAllocation* DynamicParameterAllocation,
@@ -509,9 +511,9 @@ void FDynamicSpriteEmitterDataBase::BuildViewFillData(
 	Data.View = InView;
 	check(Data.VertexSize == 0 || Data.VertexSize == InVertexSize);
 
-	DynamicVertexAllocation = FGlobalDynamicVertexBuffer::Get().Allocate( InVertexCount * InVertexSize );
+	DynamicVertexAllocation = DynamicVertexBuffer.Allocate( InVertexCount * InVertexSize );
 
-	if (FGlobalDynamicVertexBuffer::Get().IsRenderAlarmLoggingEnabled())
+	if (DynamicVertexBuffer.IsRenderAlarmLoggingEnabled())
 	{
 		UE_LOG(LogParticles, Warning, TEXT("Panic logging.  Allocated %u bytes for Resource: %s, Owner: %s"), InVertexCount * InVertexSize, *Proxy->GetResourceName().ToString(), *Proxy->GetOwnerName().ToString())
 	}
@@ -524,7 +526,7 @@ void FDynamicSpriteEmitterDataBase::BuildViewFillData(
 	GetIndexAllocInfo(NumIndices, IndexStride);
 	check(IndexStride > 0);
 
-	DynamicIndexAllocation = FGlobalDynamicIndexBuffer::Get().Allocate( NumIndices, IndexStride );
+	DynamicIndexAllocation = DynamicIndexBuffer.Allocate( NumIndices, IndexStride );
 	Data.IndexData = DynamicIndexAllocation.Buffer;
 	Data.IndexCount = NumIndices;
 
@@ -533,9 +535,9 @@ void FDynamicSpriteEmitterDataBase::BuildViewFillData(
 	if( bUsesDynamicParameter )
 	{
 		check( InDynamicParameterVertexStride > 0 );
-		*DynamicParameterAllocation = FGlobalDynamicVertexBuffer::Get().Allocate( InVertexCount * InDynamicParameterVertexStride );
+		*DynamicParameterAllocation = DynamicVertexBuffer.Allocate( InVertexCount * InDynamicParameterVertexStride );
 
-		if (FGlobalDynamicVertexBuffer::Get().IsRenderAlarmLoggingEnabled())
+		if (DynamicVertexBuffer.IsRenderAlarmLoggingEnabled())
 		{
 			UE_LOG(LogParticles, Warning, TEXT("Panic logging.  Allocated %u bytes for Resource: %s, Owner: %s"), InVertexCount * InDynamicParameterVertexStride, *Proxy->GetResourceName().ToString(), *Proxy->GetOwnerName().ToString())
 		}
@@ -1077,21 +1079,22 @@ void FDynamicSpriteEmitterData::GetDynamicMeshElementsEmitter(const FParticleSys
 					Proxy->UpdateWorldSpacePrimitiveUniformBuffer();
 				}
 
+				FGlobalDynamicVertexBuffer& DynamicVertexBuffer = Collector.GetDynamicVertexBuffer();
 				FGlobalDynamicVertexBuffer::FAllocation Allocation;
 				FGlobalDynamicVertexBuffer::FAllocation DynamicParameterAllocation;
 
 				// Allocate memory for render data.
-				Allocation = FGlobalDynamicVertexBuffer::Get().Allocate( InstanceFactor * ParticleCount * VertexSize * NumVerticesPerParticleInBuffer );
-				if (FGlobalDynamicVertexBuffer::Get().IsRenderAlarmLoggingEnabled())
+				Allocation = DynamicVertexBuffer.Allocate( InstanceFactor * ParticleCount * VertexSize * NumVerticesPerParticleInBuffer );
+				if (DynamicVertexBuffer.IsRenderAlarmLoggingEnabled())
 				{
 					UE_LOG(LogParticles, Warning, TEXT("Panic logging.  Allocated %u bytes for Resource: %s, Owner: %s"), InstanceFactor * ParticleCount * VertexSize * NumVerticesPerParticleInBuffer, *Proxy->GetResourceName().ToString(), *Proxy->GetOwnerName().ToString())
 				}
 
 				if (bUsesDynamicParameter)
 				{
-					DynamicParameterAllocation = FGlobalDynamicVertexBuffer::Get().Allocate( InstanceFactor * ParticleCount * DynamicParameterVertexSize * NumVerticesPerParticleInBuffer );
+					DynamicParameterAllocation = DynamicVertexBuffer.Allocate( InstanceFactor * ParticleCount * DynamicParameterVertexSize * NumVerticesPerParticleInBuffer );
 
-					if (FGlobalDynamicVertexBuffer::Get().IsRenderAlarmLoggingEnabled())
+					if (DynamicVertexBuffer.IsRenderAlarmLoggingEnabled())
 					{
 						UE_LOG(LogParticles, Warning, TEXT("Panic logging.  Allocated %u bytes for Resource: %s, Owner: %s"), InstanceFactor * ParticleCount * DynamicParameterVertexSize * NumVerticesPerParticleInBuffer, *Proxy->GetResourceName().ToString(), *Proxy->GetOwnerName().ToString())
 					}
@@ -1669,15 +1672,16 @@ void FDynamicMeshEmitterData::GetDynamicMeshElementsEmitter(const FParticleSyste
 				if (bInstanced)
 				{
 					const uint32 InstanceFactor = View->IsInstancedStereoPass() ? 2 : 1;
-					FGlobalDynamicVertexBuffer::FAllocation Allocation = FGlobalDynamicVertexBuffer::Get().Allocate(InstanceFactor * ParticleCount * InstanceVertexStride);
+					FGlobalDynamicVertexBuffer& DynamicVertexBuffer = Collector.GetDynamicVertexBuffer();
+					FGlobalDynamicVertexBuffer::FAllocation Allocation = DynamicVertexBuffer.Allocate(InstanceFactor * ParticleCount * InstanceVertexStride);
 					FGlobalDynamicVertexBuffer::FAllocation DynamicParameterAllocation;
 					uint8* PrevTransformBuffer = nullptr;
 
 					if (bUsesDynamicParameter)
 					{
-						DynamicParameterAllocation = FGlobalDynamicVertexBuffer::Get().Allocate(InstanceFactor * ParticleCount * DynamicParameterVertexStride);
+						DynamicParameterAllocation = DynamicVertexBuffer.Allocate(InstanceFactor * ParticleCount * DynamicParameterVertexStride);
 
-						if (FGlobalDynamicVertexBuffer::Get().IsRenderAlarmLoggingEnabled())
+						if (DynamicVertexBuffer.IsRenderAlarmLoggingEnabled())
 						{
 							UE_LOG(LogParticles, Warning, TEXT("Panic logging.  Allocated %u bytes for Resource: %s, Owner: %s"), InstanceFactor * ParticleCount * DynamicParameterVertexStride, *Proxy->GetResourceName().ToString(), *Proxy->GetOwnerName().ToString());
 						}
@@ -2767,7 +2771,9 @@ void FDynamicBeam2EmitterData::GetDynamicMeshElementsEmitter(const FParticleSyst
 	FGlobalDynamicVertexBuffer::FAllocation DynamicParameterAllocation;
 	FAsyncBufferFillData Data;
 
-	BuildViewFillData(Proxy, View, Source.VertexCount, sizeof(FParticleBeamTrailVertex), 0, DynamicVertexAllocation, DynamicIndexAllocation, &DynamicParameterAllocation, Data);
+	BuildViewFillData(Proxy, View, Source.VertexCount, sizeof(FParticleBeamTrailVertex), 0, 
+		Collector.GetDynamicIndexBuffer(), Collector.GetDynamicVertexBuffer(), 
+		DynamicVertexAllocation, DynamicIndexAllocation, &DynamicParameterAllocation, Data);
 	DoBufferFill(Data);
 		OutTriangleCount = Data.OutTriangleCount;
 
@@ -5543,7 +5549,10 @@ void FDynamicTrailsEmitterData::GetDynamicMeshElementsEmitter(const FParticleSys
 	const int32 VertexStride = GetDynamicVertexStride(ViewFamily.GetFeatureLevel());
 	const int32 DynamicParameterVertexStride = bUsesDynamicParameter ? GetDynamicParameterVertexStride() : 0;
 
-	BuildViewFillData(Proxy, View, SourcePointer->VertexCount, VertexStride, DynamicParameterVertexStride, DynamicVertexAllocation, DynamicIndexAllocation, &DynamicParameterAllocation, Data);
+	DoBufferFill(Data);
+	BuildViewFillData(Proxy, View, SourcePointer->VertexCount, VertexStride, DynamicParameterVertexStride, 
+		Collector.GetDynamicIndexBuffer(), Collector.GetDynamicVertexBuffer(), 
+		DynamicVertexAllocation, DynamicIndexAllocation, &DynamicParameterAllocation, Data);
 	DoBufferFill(Data);
 		OutTriangleCount = Data.OutTriangleCount;
 
diff --git a/Engine/Source/Runtime/Engine/Private/SceneManagement.cpp b/Engine/Source/Runtime/Engine/Private/SceneManagement.cpp
--- a/Engine/Source/Runtime/Engine/Private/SceneManagement.cpp
+++ b/Engine/Source/Runtime/Engine/Private/SceneManagement.cpp
@@ -207,6 +207,9 @@ static TAutoConsoleVariable<int32> CVarUseParallelGetDynamicMeshElementsTasks(
 
 FMeshElementCollector::FMeshElementCollector(ERHIFeatureLevel::Type InFeatureLevel) :
 	PrimitiveSceneProxy(NULL),
+	DynamicIndexBuffer(nullptr),
+	DynamicVertexBuffer(nullptr),
+	DynamicReadBuffer(nullptr),
 	FeatureLevel(InFeatureLevel),
 	bUseAsyncTasks(FApp::ShouldUseThreadingForPerformance() && CVarUseParallelGetDynamicMeshElementsTasks.GetValueOnAnyThread() > 0)
 {	
diff --git a/Engine/Source/Runtime/Engine/Public/ParticleHelper.h b/Engine/Source/Runtime/Engine/Public/ParticleHelper.h
--- a/Engine/Source/Runtime/Engine/Public/ParticleHelper.h
+++ b/Engine/Source/Runtime/Engine/Public/ParticleHelper.h
@@ -1628,6 +1628,8 @@ struct FDynamicSpriteEmitterDataBase : public FDynamicEmitterDataBase
 		int32 InVertexCount, 
 		int32 InVertexSize, 
 		int32 InDynamicParameterVertexSize, 
+		FGlobalDynamicIndexBuffer& DynamicIndexBuffer,
+		FGlobalDynamicVertexBuffer& DynamicVertexBuffer,
 		FGlobalDynamicVertexBuffer::FAllocation& DynamicVertexAllocation,
 		FGlobalDynamicIndexBuffer::FAllocation& DynamicIndexAllocation,
 		FGlobalDynamicVertexBuffer::FAllocation* DynamicParameterAllocation,
diff --git a/Engine/Source/Runtime/Engine/Public/SceneManagement.h b/Engine/Source/Runtime/Engine/Public/SceneManagement.h
--- a/Engine/Source/Runtime/Engine/Public/SceneManagement.h
+++ b/Engine/Source/Runtime/Engine/Public/SceneManagement.h
@@ -26,6 +26,7 @@
 #include "MeshBatch.h"
 #include "SceneUtils.h"
 #include "LightmapUniformShaderParameters.h"
+#include "DynamicBufferAllocator.h"
 
 #ifndef ENVIRONMENT_TEXTURE_ARRAY_WORKAROUND // RHI_RAYTRACING
 #define ENVIRONMENT_TEXTURE_ARRAY_WORKAROUND	1
@@ -1911,6 +1912,27 @@ public:
 		return MeshBatchStorage[Index];
 	}
 
+	/** Return dynamic index buffer for this collector. */
+	FGlobalDynamicIndexBuffer& GetDynamicIndexBuffer()
+	{
+		check(DynamicIndexBuffer);
+		return *DynamicIndexBuffer;
+	}
+
+	/** Return dynamic vertex buffer for this collector. */
+	FGlobalDynamicVertexBuffer& GetDynamicVertexBuffer()
+	{
+		check(DynamicVertexBuffer);
+		return *DynamicVertexBuffer;
+	}
+
+	/** Return dynamic read buffer for this collector. */
+	FGlobalDynamicReadBuffer& GetDynamicReadBuffer()
+	{
+		check(DynamicReadBuffer);
+		return *DynamicReadBuffer;
+	}
+
 	// @return number of MeshBatches collected (so far) for a given view
 	uint32 GetMeshBatchCount(uint32 ViewIndex) const
 	{
@@ -2009,6 +2031,9 @@ private:
 		MeshIdInPrimitivePerView.Empty();
 		DynamicPrimitiveShaderDataPerView.Empty();
 		NumMeshBatchElementsPerView.Empty();
+		DynamicIndexBuffer = nullptr;
+		DynamicVertexBuffer = nullptr;
+		DynamicReadBuffer = nullptr;
 	}
 
 	void AddViewMeshArrays(
@@ -2016,7 +2041,10 @@ private:
 		TArray<FMeshBatchAndRelevance,SceneRenderingAllocator>* ViewMeshes,
 		FSimpleElementCollector* ViewSimpleElementCollector, 
 		TArray<FPrimitiveUniformShaderParameters>* InDynamicPrimitiveShaderData,
-		ERHIFeatureLevel::Type InFeatureLevel)
+		ERHIFeatureLevel::Type InFeatureLevel,
+		FGlobalDynamicIndexBuffer* InDynamicIndexBuffer,
+		FGlobalDynamicVertexBuffer* InDynamicVertexBuffer,
+		FGlobalDynamicReadBuffer* InDynamicReadBuffer)
 	{
 		Views.Add(InView);
 		MeshIdInPrimitivePerView.Add(0);
@@ -2024,6 +2052,11 @@ private:
 		NumMeshBatchElementsPerView.Add(0);
 		SimpleElementCollectors.Add(ViewSimpleElementCollector);
 		DynamicPrimitiveShaderDataPerView.Add(InDynamicPrimitiveShaderData);
+
+		check(InDynamicIndexBuffer && InDynamicVertexBuffer && InDynamicReadBuffer);
+		DynamicIndexBuffer = InDynamicIndexBuffer;
+		DynamicVertexBuffer = InDynamicVertexBuffer;
+		DynamicReadBuffer = InDynamicReadBuffer;
 	}
 
 	/** 
@@ -2056,6 +2089,11 @@ private:
 	/** Current primitive being gathered. */
 	const FPrimitiveSceneProxy* PrimitiveSceneProxy;
 
+	/** Dynamic buffer pools. */
+	FGlobalDynamicIndexBuffer* DynamicIndexBuffer;
+	FGlobalDynamicVertexBuffer* DynamicVertexBuffer;
+	FGlobalDynamicReadBuffer* DynamicReadBuffer;
+
 	ERHIFeatureLevel::Type FeatureLevel;
 
 	/** This is related to some cvars and FApp stuff and if true means calling code should use async tasks. */
diff --git a/Engine/Source/Runtime/RenderCore/Private/DynamicBufferAllocator.cpp b/Engine/Source/Runtime/RenderCore/Private/DynamicBufferAllocator.cpp
new file mode 100644
--- /dev/null
+++ b/Engine/Source/Runtime/RenderCore/Private/DynamicBufferAllocator.cpp
@@ -0,0 +1,226 @@
+// Copyright 1998-2019 Epic Games, Inc. All Rights Reserved.
+
+/*==============================================================================
+DynamicBufferAllocator.cpp: Classes for allocating transient rendering data.
+==============================================================================*/
+
+#include "DynamicBufferAllocator.h"
+#include "RenderResource.h"
+
+int32 GMaxReadBufferRenderingBytesAllocatedPerFrame = 32 * 1024 * 1024;
+
+FAutoConsoleVariableRef CVarMaxReadBufferRenderingBytesAllocatedPerFrame(
+	TEXT("r.ReadBuffer.MaxRenderingBytesAllocatedPerFrame"),
+	GMaxReadBufferRenderingBytesAllocatedPerFrame,
+	TEXT("The maximum number of transient rendering read buffer bytes to allocate before we start panic logging who is doing the allocations"));
+
+int32 GMinReadBufferRenderingBufferSize = 8 * 1024;
+FAutoConsoleVariableRef CVarMinReadBufferSize(
+	TEXT("r.ReadBuffer.MinSize"),
+	GMinReadBufferRenderingBufferSize,
+	TEXT("The minimum size (in instances) to allocate in blocks for rendering read buffers."));
+
+struct FDynamicReadBufferPool
+{
+	/** List of vertex buffers. */
+	TIndirectArray<FDynamicAllocReadBuffer> Buffers;
+	/** The current buffer from which allocations are being made. */
+	FDynamicAllocReadBuffer* CurrentBuffer;
+
+	/** Default constructor. */
+	FDynamicReadBufferPool()
+		: CurrentBuffer(NULL)
+	{
+	}
+
+	/** Destructor. */
+	~FDynamicReadBufferPool()
+	{
+		int32 NumVertexBuffers = Buffers.Num();
+		for (int32 BufferIndex = 0; BufferIndex < NumVertexBuffers; ++BufferIndex)
+		{
+			Buffers[BufferIndex].Release();
+		}
+	}
+
+	FCriticalSection CriticalSection;
+};
+
+
+
+FGlobalDynamicReadBuffer::FGlobalDynamicReadBuffer()
+	: TotalAllocatedSinceLastCommit(0)
+{
+	FloatBufferPool = new FDynamicReadBufferPool();
+	Int32BufferPool = new FDynamicReadBufferPool();
+}
+
+FGlobalDynamicReadBuffer::~FGlobalDynamicReadBuffer()
+{
+	Cleanup();
+}
+
+void FGlobalDynamicReadBuffer::Cleanup()
+{
+	if (FloatBufferPool)
+	{
+		UE_LOG(LogRendererCore, Log, TEXT("FGlobalDynamicReadBuffer::Cleanup()"));
+		delete FloatBufferPool;
+		FloatBufferPool = nullptr;
+	}
+
+	if (Int32BufferPool)
+	{
+		delete Int32BufferPool;
+		Int32BufferPool = nullptr;
+	}
+}
+void FGlobalDynamicReadBuffer::InitRHI()
+{
+	UE_LOG(LogRendererCore, Log, TEXT("FGlobalReadBuffer::InitRHI"));
+}
+
+void FGlobalDynamicReadBuffer::ReleaseRHI()
+{
+	UE_LOG(LogRendererCore, Log, TEXT("FGlobalReadBuffer::ReleaseRHI"));
+	Cleanup();
+}
+
+FGlobalDynamicReadBuffer::FAllocation FGlobalDynamicReadBuffer::AllocateFloat(uint32 Num)
+{
+	FScopeLock ScopeLock(&FloatBufferPool->CriticalSection);
+
+	FAllocation Allocation;
+
+	TotalAllocatedSinceLastCommit += Num;
+	if (IsRenderAlarmLoggingEnabled())
+	{
+		UE_LOG(LogRendererCore, Warning, TEXT("FGlobalReadBuffer::AllocateFloat(%u), will have allocated %u total this frame"), Num, TotalAllocatedSinceLastCommit);
+	}
+	uint32 SizeInBytes = sizeof(float) * Num;
+	FDynamicAllocReadBuffer* Buffer = FloatBufferPool->CurrentBuffer;
+	if (Buffer == NULL || Buffer->AllocatedByteCount + SizeInBytes > Buffer->NumBytes)
+	{
+		// Find a buffer in the pool big enough to service the request.
+		Buffer = NULL;
+		for (int32 BufferIndex = 0, NumBuffers = FloatBufferPool->Buffers.Num(); BufferIndex < NumBuffers; ++BufferIndex)
+		{
+			FDynamicAllocReadBuffer& BufferToCheck = FloatBufferPool->Buffers[BufferIndex];
+			if (BufferToCheck.AllocatedByteCount + SizeInBytes <= BufferToCheck.NumBytes)
+			{
+				Buffer = &BufferToCheck;
+				break;
+			}
+		}
+
+		// Create a new vertex buffer if needed.
+		if (Buffer == NULL)
+		{
+			uint32 NewBufferSize = FMath::Max(Num, (uint32)GMinReadBufferRenderingBufferSize);
+			Buffer = new FDynamicAllocReadBuffer();
+			FloatBufferPool->Buffers.Add(Buffer);
+			Buffer->Initialize(sizeof(float), NewBufferSize, PF_R32_FLOAT, BUF_Dynamic);
+		}
+
+		// Lock the buffer if needed.
+		if (Buffer->MappedBuffer == NULL)
+		{
+			Buffer->Lock();
+		}
+
+		// Remember this buffer, we'll try to allocate out of it in the future.
+		FloatBufferPool->CurrentBuffer = Buffer;
+	}
+
+	check(Buffer != NULL);
+	checkf(Buffer->AllocatedByteCount + SizeInBytes <= Buffer->NumBytes, TEXT("Global dynamic read buffer float buffer allocation failed: BufferSize=%d AllocatedByteCount=%d SizeInBytes=%d"), Buffer->NumBytes, Buffer->AllocatedByteCount, SizeInBytes);
+	Allocation.Buffer = Buffer->MappedBuffer + Buffer->AllocatedByteCount;
+	Allocation.ReadBuffer = Buffer;
+	Allocation.FirstIndex = Buffer->AllocatedByteCount;
+	Buffer->AllocatedByteCount += SizeInBytes;
+
+	return Allocation;
+}
+
+FGlobalDynamicReadBuffer::FAllocation FGlobalDynamicReadBuffer::AllocateInt32(uint32 Num)
+{
+	FScopeLock ScopeLock(&Int32BufferPool->CriticalSection);
+	FAllocation Allocation;
+
+	TotalAllocatedSinceLastCommit += Num;
+	if (IsRenderAlarmLoggingEnabled())
+	{
+		UE_LOG(LogRendererCore, Warning, TEXT("FGlobalReadBuffer::AllocateInt32(%u), will have allocated %u total this frame"), Num, TotalAllocatedSinceLastCommit);
+	}
+	uint32 SizeInBytes = sizeof(int32) * Num;
+	FDynamicAllocReadBuffer* Buffer = Int32BufferPool->CurrentBuffer;
+	if (Buffer == NULL || Buffer->AllocatedByteCount + SizeInBytes > Buffer->NumBytes)
+	{
+		// Find a buffer in the pool big enough to service the request.
+		Buffer = NULL;
+		for (int32 BufferIndex = 0, NumBuffers = Int32BufferPool->Buffers.Num(); BufferIndex < NumBuffers; ++BufferIndex)
+		{
+			FDynamicAllocReadBuffer& BufferToCheck = Int32BufferPool->Buffers[BufferIndex];
+			if (BufferToCheck.AllocatedByteCount + SizeInBytes <= BufferToCheck.NumBytes)
+			{
+				Buffer = &BufferToCheck;
+				break;
+			}
+		}
+
+		// Create a new vertex buffer if needed.
+		if (Buffer == NULL)
+		{
+			uint32 NewBufferSize = FMath::Max(Num, (uint32)GMinReadBufferRenderingBufferSize);
+			Buffer = new FDynamicAllocReadBuffer();
+			Int32BufferPool->Buffers.Add(Buffer);
+			Buffer->Initialize(sizeof(int32), NewBufferSize, PF_R32_SINT, BUF_Dynamic);
+		}
+
+		// Lock the buffer if needed.
+		if (Buffer->MappedBuffer == NULL)
+		{
+			Buffer->Lock();
+		}
+
+		// Remember this buffer, we'll try to allocate out of it in the future.
+		Int32BufferPool->CurrentBuffer = Buffer;
+	}
+
+	check(Buffer != NULL);
+	checkf(Buffer->AllocatedByteCount + SizeInBytes <= Buffer->NumBytes, TEXT("Global dynamic read buffer int32 buffer allocation failed: BufferSize=%d AllocatedByteCount=%d SizeInBytes=%d"), Buffer->NumBytes, Buffer->AllocatedByteCount, SizeInBytes);
+	Allocation.Buffer = Buffer->MappedBuffer + Buffer->AllocatedByteCount;
+	Allocation.ReadBuffer = Buffer;
+	Allocation.FirstIndex = Buffer->AllocatedByteCount;
+	Buffer->AllocatedByteCount += SizeInBytes;
+
+	return Allocation;
+}
+
+bool FGlobalDynamicReadBuffer::IsRenderAlarmLoggingEnabled() const
+{
+	return GMaxReadBufferRenderingBytesAllocatedPerFrame > 0 && TotalAllocatedSinceLastCommit >= (size_t)GMaxReadBufferRenderingBytesAllocatedPerFrame;
+}
+
+void FGlobalDynamicReadBuffer::Commit()
+{
+	for (int32 BufferIndex = 0, NumBuffers = FloatBufferPool->Buffers.Num(); BufferIndex < NumBuffers; ++BufferIndex)
+	{
+		FDynamicAllocReadBuffer& Buffer = FloatBufferPool->Buffers[BufferIndex];
+		if (Buffer.MappedBuffer != NULL)
+		{
+			Buffer.Unlock();
+		}
+	}
+	FloatBufferPool->CurrentBuffer = NULL;
+	for (int32 BufferIndex = 0, NumBuffers = Int32BufferPool->Buffers.Num(); BufferIndex < NumBuffers; ++BufferIndex)
+	{
+		FDynamicAllocReadBuffer& Buffer = Int32BufferPool->Buffers[BufferIndex];
+		if (Buffer.MappedBuffer != NULL)
+		{
+			Buffer.Unlock();
+		}
+	}
+	Int32BufferPool->CurrentBuffer = NULL;
+	TotalAllocatedSinceLastCommit = 0;
+}
\ No newline at end of file
diff --git a/Engine/Source/Runtime/RenderCore/Private/RenderResource.cpp b/Engine/Source/Runtime/RenderCore/Private/RenderResource.cpp
--- a/Engine/Source/Runtime/RenderCore/Private/RenderResource.cpp
+++ b/Engine/Source/Runtime/RenderCore/Private/RenderResource.cpp
@@ -492,13 +492,8 @@ void FGlobalDynamicVertexBuffer::Commit()
 	TotalAllocatedSinceLastCommit = 0;
 }
 
-FGlobalDynamicVertexBuffer& FGlobalDynamicVertexBuffer::Get()
-{
-	check(IsInRenderingThread());
-
-	static FGlobalDynamicVertexBuffer GlobalDynamicVertexBuffer;
-	return GlobalDynamicVertexBuffer;
-}
+FGlobalDynamicVertexBuffer InitViewDynamicVertexBuffer;
+FGlobalDynamicVertexBuffer InitShadowViewDynamicVertexBuffer;
 
 /*------------------------------------------------------------------------------
 	FGlobalDynamicIndexBuffer implementation.
@@ -693,14 +688,6 @@ void FGlobalDynamicIndexBuffer::Commit()
 	}
 }
 
-FGlobalDynamicIndexBuffer& FGlobalDynamicIndexBuffer::Get()
-{
-	check(IsInRenderingThread());
-
-	static FGlobalDynamicIndexBuffer GlobalDynamicIndexBuffer;
-	return GlobalDynamicIndexBuffer;
-}
-
 /*=============================================================================
 	FMipBiasFade class
 =============================================================================*/
diff --git a/Engine/Source/Runtime/RenderCore/Public/DynamicBufferAllocator.h b/Engine/Source/Runtime/RenderCore/Public/DynamicBufferAllocator.h
new file mode 100644
--- /dev/null
+++ b/Engine/Source/Runtime/RenderCore/Public/DynamicBufferAllocator.h
@@ -0,0 +1,97 @@
+// Copyright 1998-2019 Epic Games, Inc. All Rights Reserved.
+
+/*==============================================================================
+DynamicBufferAllocator.h: Classes for allocating transient rendering data.
+==============================================================================*/
+
+#pragma once
+
+#include "RenderResource.h"
+
+struct FDynamicReadBufferPool;
+
+struct FDynamicAllocReadBuffer : public FDynamicReadBuffer
+{
+	FDynamicAllocReadBuffer()
+		: FDynamicReadBuffer()
+		, AllocatedByteCount(0)
+	{}
+
+	int32 AllocatedByteCount;
+
+	/**
+	* Unocks the buffer so the GPU may read from it.
+	*/
+	void Unlock()
+	{
+		FDynamicReadBuffer::Unlock();
+		AllocatedByteCount = 0;
+	}
+};
+
+/**
+* A system for dynamically allocating GPU memory for rendering. Note that this must derive from FRenderResource 
+  so that we can safely free the shader resource views for OpenGL and other platforms. If we wait until the module is shutdown,
+  the renderer RHI will have already been destroyed and we can execute code on invalid data. By making ourself a render resource, we
+  clean up immediately before the renderer dies.
+*/
+class RENDERCORE_API FGlobalDynamicReadBuffer : public FRenderResource
+{
+public:
+	/**
+	* Information regarding an allocation from this buffer.
+	*/
+	struct FAllocation
+	{
+		/** The location of the buffer in main memory. */
+		uint8* Buffer;
+		/** The read buffer to bind for draw calls. */
+		FDynamicAllocReadBuffer* ReadBuffer;
+		/** The offset in to the read buffer. */
+		uint32 FirstIndex;
+
+		/** Default constructor. */
+		FAllocation()
+			: Buffer(NULL)
+			, ReadBuffer(NULL)
+			, FirstIndex(0)
+		{
+		}
+
+		/** Returns true if the allocation is valid. */
+		FORCEINLINE bool IsValid() const
+		{
+			return Buffer != NULL;
+		}
+	};
+
+	FGlobalDynamicReadBuffer();
+	~FGlobalDynamicReadBuffer();
+	
+	FAllocation AllocateFloat(uint32 Num);
+	FAllocation AllocateInt32(uint32 Num);
+
+	/**
+	* Commits allocated memory to the GPU.
+	*		WARNING: Once this buffer has been committed to the GPU, allocations
+	*		remain valid only until the next call to Allocate!
+	*/
+	void Commit();
+
+
+	/** Returns true if log statements should be made because we exceeded GMaxVertexBytesAllocatedPerFrame */
+	bool IsRenderAlarmLoggingEnabled() const;
+
+protected:
+	virtual void InitRHI() override;
+	virtual void ReleaseRHI() override;
+	void Cleanup();
+
+	/** The pools of read buffers from which allocations are made. */
+	FDynamicReadBufferPool* FloatBufferPool;
+	FDynamicReadBufferPool* Int32BufferPool;
+
+	/** A total of all allocations made since the last commit. Used to alert about spikes in memory usage. */
+	size_t TotalAllocatedSinceLastCommit;
+};
+
diff --git a/Engine/Source/Runtime/RenderCore/Public/RenderResource.h b/Engine/Source/Runtime/RenderCore/Public/RenderResource.h
--- a/Engine/Source/Runtime/RenderCore/Public/RenderResource.h
+++ b/Engine/Source/Runtime/RenderCore/Public/RenderResource.h
@@ -667,11 +667,6 @@ public:
 	 */
 	void Commit();
 
-	/**
-	 * Obtain a reference to the global dynamic vertex buffer instance.
-	 */
-	static FGlobalDynamicVertexBuffer& Get();
-
 	/** Returns true if log statements should be made because we exceeded GMaxVertexBytesAllocatedPerFrame */
 	bool IsRenderAlarmLoggingEnabled() const;
 
@@ -748,11 +743,6 @@ public:
 	 */
 	void Commit();
 
-	/**
-	 * Obtain a reference to the global dynamic index buffer instance.
-	 */
-	static FGlobalDynamicIndexBuffer& Get();
-
 private:
 	/** The pool of vertex buffers from which allocations are made. */
 	struct FDynamicIndexBufferPool* Pools[2];
diff --git a/Engine/Source/Runtime/Renderer/Private/DeferredShadingRenderer.cpp b/Engine/Source/Runtime/Renderer/Private/DeferredShadingRenderer.cpp
--- a/Engine/Source/Runtime/Renderer/Private/DeferredShadingRenderer.cpp
+++ b/Engine/Source/Runtime/Renderer/Private/DeferredShadingRenderer.cpp
@@ -452,6 +452,12 @@ DECLARE_CYCLE_STAT(TEXT("RenderFinish"), STAT_CLM_RenderFinish, STATGROUP_Comman
 DECLARE_CYCLE_STAT(TEXT("AfterFrame"), STAT_CLM_AfterFrame, STATGROUP_CommandListMarkers);
 
 FGraphEventRef FDeferredShadingSceneRenderer::TranslucencyTimestampQuerySubmittedFence[FOcclusionQueryHelpers::MaxBufferedOcclusionFrames + 1];
+FGlobalDynamicIndexBuffer FDeferredShadingSceneRenderer::DynamicIndexBufferForInitViews;
+FGlobalDynamicIndexBuffer FDeferredShadingSceneRenderer::DynamicIndexBufferForInitShadows;
+FGlobalDynamicVertexBuffer FDeferredShadingSceneRenderer::DynamicVertexBufferForInitViews;
+FGlobalDynamicVertexBuffer FDeferredShadingSceneRenderer::DynamicVertexBufferForInitShadows;
+TGlobalResource<FGlobalDynamicReadBuffer> FDeferredShadingSceneRenderer::DynamicReadBufferForInitShadows;
+TGlobalResource<FGlobalDynamicReadBuffer> FDeferredShadingSceneRenderer::DynamicReadBufferForInitViews;
 
 /**
  * Returns true if the depth Prepass needs to run
@@ -609,7 +615,11 @@ bool FDeferredShadingSceneRenderer::GatherRayTracingWorldInstances(FRHICommandLi
 				&View.RayTracedDynamicMeshElements,
 				&View.SimpleElementCollector,
 				&DummyDynamicPrimitiveShaderData,
-				ViewFamily.GetFeatureLevel());
+				ViewFamily.GetFeatureLevel(),
+				&DynamicIndexBufferForInitViews,
+				&DynamicVertexBufferForInitViews,
+				&DynamicReadBufferForInitViews
+				);
 
 			View.RaytraycingDynamicMeshDrawCommandStorage.MeshDrawCommands.Reserve(Scene->Primitives.Num());
 			View.RaytraycingVisibleMeshDrawCommands.Reserve(Scene->Primitives.Num());
@@ -1158,14 +1168,12 @@ void FDeferredShadingSceneRenderer::Render(FRHICommandListImmediate& RHICmdList)
 	const bool bIsOcclusionTesting = DoOcclusionQueries(FeatureLevel) && (!bIsWireframe || bIsViewFrozen || bHasViewParent);
 
 	// Dynamic vertex and index buffers need to be committed before rendering.
-	if (!bDoInitViewAftersPrepass)
+	GEngine->GetPreRenderDelegate().Broadcast();
 	{
-		GEngine->GetPreRenderDelegate().Broadcast();
-		{
-			SCOPE_CYCLE_COUNTER(STAT_FDeferredShadingSceneRenderer_FGlobalDynamicVertexBuffer_Commit);
-			FGlobalDynamicVertexBuffer::Get().Commit();
-			FGlobalDynamicIndexBuffer::Get().Commit();
-		}
+		SCOPE_CYCLE_COUNTER(STAT_FDeferredShadingSceneRenderer_FGlobalDynamicVertexBuffer_Commit);
+		DynamicIndexBufferForInitViews.Commit();
+		DynamicVertexBufferForInitViews.Commit();
+		DynamicReadBufferForInitViews.Commit();
 	}
 
 	// Only update the GPU particle simulation for the main view
@@ -1195,13 +1203,14 @@ void FDeferredShadingSceneRenderer::Render(FRHICommandListImmediate& RHICmdList)
 				InitViewsPossiblyAfterPrepass(RHICmdList, ILCTaskData, UpdateViewCustomDataEvents);
 				PrepareDistanceFieldScene(RHICmdList, false);
 
-				GEngine->GetPreRenderDelegate().Broadcast();
-
 				{
 					SCOPE_CYCLE_COUNTER(STAT_FDeferredShadingSceneRenderer_FGlobalDynamicVertexBuffer_Commit);
-					FGlobalDynamicVertexBuffer::Get().Commit();
-					FGlobalDynamicIndexBuffer::Get().Commit();
+
+					DynamicVertexBufferForInitShadows.Commit();
+					DynamicIndexBufferForInitShadows.Commit();
+					DynamicReadBufferForInitShadows.Commit();
 				}
+
 				ServiceLocalQueue();
 			}
 		}
diff --git a/Engine/Source/Runtime/Renderer/Private/DeferredShadingRenderer.h b/Engine/Source/Runtime/Renderer/Private/DeferredShadingRenderer.h
--- a/Engine/Source/Runtime/Renderer/Private/DeferredShadingRenderer.h
+++ b/Engine/Source/Runtime/Renderer/Private/DeferredShadingRenderer.h
@@ -112,6 +112,12 @@ public:
 private:
 
 	static FGraphEventRef TranslucencyTimestampQuerySubmittedFence[FOcclusionQueryHelpers::MaxBufferedOcclusionFrames + 1];
+	static FGlobalDynamicIndexBuffer DynamicIndexBufferForInitViews;
+	static FGlobalDynamicIndexBuffer DynamicIndexBufferForInitShadows;
+	static FGlobalDynamicVertexBuffer DynamicVertexBufferForInitViews;
+	static FGlobalDynamicVertexBuffer DynamicVertexBufferForInitShadows;
+	static TGlobalResource<FGlobalDynamicReadBuffer> DynamicReadBufferForInitViews;
+	static TGlobalResource<FGlobalDynamicReadBuffer> DynamicReadBufferForInitShadows;
 
 	/** Creates a per object projected shadow for the given interaction. */
 	void CreatePerObjectProjectedShadow(
diff --git a/Engine/Source/Runtime/Renderer/Private/MobileShadingRenderer.cpp b/Engine/Source/Runtime/Renderer/Private/MobileShadingRenderer.cpp
--- a/Engine/Source/Runtime/Renderer/Private/MobileShadingRenderer.cpp
+++ b/Engine/Source/Runtime/Renderer/Private/MobileShadingRenderer.cpp
@@ -85,6 +85,10 @@ DECLARE_CYCLE_STAT(TEXT("Post"), STAT_CLMM_Post, STATGROUP_CommandListMarkers);
 DECLARE_CYCLE_STAT(TEXT("Translucency"), STAT_CLMM_Translucency, STATGROUP_CommandListMarkers);
 DECLARE_CYCLE_STAT(TEXT("Shadows"), STAT_CLMM_Shadows, STATGROUP_CommandListMarkers);
 
+FGlobalDynamicIndexBuffer FMobileSceneRenderer::DynamicIndexBuffer;
+FGlobalDynamicVertexBuffer FMobileSceneRenderer::DynamicVertexBuffer;
+TGlobalResource<FGlobalDynamicReadBuffer> FMobileSceneRenderer::DynamicReadBuffer;
+
 FMobileSceneRenderer::FMobileSceneRenderer(const FSceneViewFamily* InViewFamily,FHitProxyConsumer* HitProxyConsumer)
 	:	FSceneRenderer(InViewFamily, HitProxyConsumer)
 {
@@ -198,7 +202,7 @@ void FMobileSceneRenderer::InitViews(FRHICommandListImmediate& RHICmdList)
 	const FExclusiveDepthStencil::Type BasePassDepthStencilAccess = FExclusiveDepthStencil::DepthWrite_StencilWrite;
 
 	PreVisibilityFrameSetup(RHICmdList);
-	ComputeViewVisibility(RHICmdList, BasePassDepthStencilAccess, ViewCommandsPerView);
+	ComputeViewVisibility(RHICmdList, BasePassDepthStencilAccess, ViewCommandsPerView, DynamicIndexBuffer, DynamicVertexBuffer, DynamicReadBuffer);
 	PostVisibilityFrameSetup(ILCTaskData);
 
 	const bool bDynamicShadows = ViewFamily.EngineShowFlags.DynamicShadows;
@@ -302,9 +306,10 @@ void FMobileSceneRenderer::Render(FRHICommandListImmediate& RHICmdList)
 
 	GEngine->GetPreRenderDelegate().Broadcast();
 
-	// Dynamic vertex and index buffers need to be committed before rendering.
-	FGlobalDynamicVertexBuffer::Get().Commit();
-	FGlobalDynamicIndexBuffer::Get().Commit();
+	// Global dynamic buffers need to be committed before rendering.
+	DynamicIndexBuffer.Commit();
+	DynamicVertexBuffer.Commit();
+	DynamicReadBuffer.Commit();
 	RHICmdList.ImmediateFlush(EImmediateFlushType::DispatchToRHIThread);
 
 	// Notify the FX system that the scene is about to be rendered.
diff --git a/Engine/Source/Runtime/Renderer/Private/SceneHitProxyRendering.cpp b/Engine/Source/Runtime/Renderer/Private/SceneHitProxyRendering.cpp
--- a/Engine/Source/Runtime/Renderer/Private/SceneHitProxyRendering.cpp
+++ b/Engine/Source/Runtime/Renderer/Private/SceneHitProxyRendering.cpp
@@ -249,12 +249,6 @@ static void DoRenderHitProxies(FRHICommandListImmediate& RHICmdList, const FScen
 
 	const auto FeatureLevel = SceneRenderer->FeatureLevel;
 
-	GEngine->GetPreRenderDelegate().Broadcast();
-
-	// Dynamic vertex and index buffers need to be committed before rendering.
-	FGlobalDynamicVertexBuffer::Get().Commit();
-	FGlobalDynamicIndexBuffer::Get().Commit();
-
 	const bool bNeedToSwitchVerticalAxis = RHINeedsToSwitchVerticalAxis(GShaderPlatformForFeatureLevel[SceneRenderer->FeatureLevel]);
 	FSceneRenderTargets& SceneContext = FSceneRenderTargets::Get(RHICmdList);
 
@@ -511,6 +505,14 @@ void FMobileSceneRenderer::RenderHitProxies(FRHICommandListImmediate& RHICmdList
 	{
 		// Find the visible primitives.
 		InitViews(RHICmdList);
+
+		GEngine->GetPreRenderDelegate().Broadcast();
+
+		// Global dynamic buffers need to be committed before rendering.
+		DynamicIndexBuffer.Commit();
+		DynamicVertexBuffer.Commit();
+		DynamicReadBuffer.Commit();
+
 		::DoRenderHitProxies(RHICmdList, this, HitProxyRT, HitProxyDepthRT);
 	}
 
@@ -552,6 +554,13 @@ void FDeferredShadingSceneRenderer::RenderHitProxies(FRHICommandListImmediate& R
 			FTaskGraphInterface::Get().WaitUntilTasksComplete(UpdateViewCustomDataEvents, ENamedThreads::GetRenderThread());
 		}
 
+		GEngine->GetPreRenderDelegate().Broadcast();
+
+		// Global dynamic buffers need to be committed before rendering.
+		DynamicIndexBufferForInitViews.Commit();
+		DynamicVertexBufferForInitViews.Commit();
+		DynamicReadBufferForInitViews.Commit();
+
 		::DoRenderHitProxies(RHICmdList, this, HitProxyRT, HitProxyDepthRT);
 		ClearPrimitiveSingleFrameIndirectLightingCacheBuffers();
 	}
diff --git a/Engine/Source/Runtime/Renderer/Private/SceneOcclusion.cpp b/Engine/Source/Runtime/Renderer/Private/SceneOcclusion.cpp
--- a/Engine/Source/Runtime/Renderer/Private/SceneOcclusion.cpp
+++ b/Engine/Source/Runtime/Renderer/Private/SceneOcclusion.cpp
@@ -388,7 +388,7 @@ void FOcclusionQueryBatcher::Flush(FRHICommandList& RHICmdList)
 	}
 }
 
-FRenderQueryRHIParamRef FOcclusionQueryBatcher::BatchPrimitive(const FVector& BoundsOrigin,const FVector& BoundsBoxExtent)
+FRenderQueryRHIParamRef FOcclusionQueryBatcher::BatchPrimitive(const FVector& BoundsOrigin,const FVector& BoundsBoxExtent, FGlobalDynamicVertexBuffer& DynamicVertexBuffer)
 {
 	// Check if the current batch is full.
 	if(CurrentBatchOcclusionQuery == NULL || NumBatchedPrimitives >= MaxBatchedPrimitives)
@@ -396,7 +396,7 @@ FRenderQueryRHIParamRef FOcclusionQueryBatcher::BatchPrimitive(const FVector& Bo
 		check(OcclusionQueryPool);
 		CurrentBatchOcclusionQuery = new(BatchOcclusionQueries) FOcclusionBatch;
 		CurrentBatchOcclusionQuery->Query = OcclusionQueryPool->AllocateQuery();
-		CurrentBatchOcclusionQuery->VertexAllocation = FGlobalDynamicVertexBuffer::Get().Allocate(MaxBatchedPrimitives * 8 * sizeof(FVector));
+		CurrentBatchOcclusionQuery->VertexAllocation = DynamicVertexBuffer.Allocate(MaxBatchedPrimitives * 8 * sizeof(FVector));
 		check(CurrentBatchOcclusionQuery->VertexAllocation.IsValid());
 		NumBatchedPrimitives = 0;
 	}
diff --git a/Engine/Source/Runtime/Renderer/Private/SceneRendering.h b/Engine/Source/Runtime/Renderer/Private/SceneRendering.h
--- a/Engine/Source/Runtime/Renderer/Private/SceneRendering.h
+++ b/Engine/Source/Runtime/Renderer/Private/SceneRendering.h
@@ -313,7 +313,7 @@ public:
 	 * Batches a primitive's occlusion query for rendering.
 	 * @param Bounds - The primitive's bounds.
 	 */
-	FRenderQueryRHIParamRef BatchPrimitive(const FVector& BoundsOrigin, const FVector& BoundsBoxExtent);
+	FRenderQueryRHIParamRef BatchPrimitive(const FVector& BoundsOrigin, const FVector& BoundsBoxExtent, FGlobalDynamicVertexBuffer& DynamicVertexBuffer);
 	inline int32 GetNumBatchOcclusionQueries() const
 	{
 		return BatchOcclusionQueries.Num();
@@ -1484,7 +1484,7 @@ protected:
 
 	// Shared functionality between all scene renderers
 
-	void InitDynamicShadows(FRHICommandListImmediate& RHICmdList);
+	void InitDynamicShadows(FRHICommandListImmediate& RHICmdList, FGlobalDynamicIndexBuffer& DynamicIndexBuffer, FGlobalDynamicVertexBuffer& DynamicVertexBuffer, FGlobalDynamicReadBuffer& DynamicReadBuffer);
 
 	void SetupMeshPass(FViewInfo& View, FExclusiveDepthStencil::Type BasePassDepthStencilAccess, FViewCommands& ViewCommands);
 
@@ -1576,13 +1576,14 @@ protected:
 	void InitProjectedShadowVisibility(FRHICommandListImmediate& RHICmdList);	
 
 	/** Gathers dynamic mesh elements for all shadows. */
-	void GatherShadowDynamicMeshElements();
+	void GatherShadowDynamicMeshElements(FGlobalDynamicIndexBuffer& DynamicIndexBuffer, FGlobalDynamicVertexBuffer& DynamicVertexBuffer, FGlobalDynamicReadBuffer& DynamicReadBuffer);
 
 	/** Performs once per frame setup prior to visibility determination. */
 	void PreVisibilityFrameSetup(FRHICommandListImmediate& RHICmdList);
 
 	/** Computes which primitives are visible and relevant for each view. */
-	void ComputeViewVisibility(FRHICommandListImmediate& RHICmdList, FExclusiveDepthStencil::Type BasePassDepthStencilAccess, FViewVisibleCommandsPerView& ViewCommandsPerView);
+	void ComputeViewVisibility(FRHICommandListImmediate& RHICmdList, FExclusiveDepthStencil::Type BasePassDepthStencilAccess, FViewVisibleCommandsPerView& ViewCommandsPerView, 
+		FGlobalDynamicIndexBuffer& DynamicIndexBuffer, FGlobalDynamicVertexBuffer& DynamicVertexBuffer, FGlobalDynamicReadBuffer& DynamicReadBuffer);
 
 	/** Performs once per frame setup after to visibility determination. */
 	void PostVisibilityFrameSetup(FILCUpdatePrimTaskData& OutILCTaskData);
@@ -1591,6 +1592,9 @@ protected:
 		TArray<FViewInfo>& InViews, 
 		const FScene* InScene, 
 		const FSceneViewFamily& InViewFamily, 
+		FGlobalDynamicIndexBuffer& DynamicIndexBuffer,
+		FGlobalDynamicVertexBuffer& DynamicVertexBuffer,
+		FGlobalDynamicReadBuffer& DynamicReadBuffer,
 		const FPrimitiveViewMasks& HasDynamicMeshElementsMasks, 
 		const FPrimitiveViewMasks& HasDynamicEditorMeshElementsMasks, 
 		const FPrimitiveViewMasks& HasViewCustomDataMasks,
@@ -1704,6 +1708,9 @@ protected:
 	
 private:
 	bool bModulatedShadowsInUse;
+	static FGlobalDynamicIndexBuffer DynamicIndexBuffer;
+	static FGlobalDynamicVertexBuffer DynamicVertexBuffer;
+	static TGlobalResource<FGlobalDynamicReadBuffer> DynamicReadBuffer;
 };
 
 // The noise textures need to be set in Slate too.
diff --git a/Engine/Source/Runtime/Renderer/Private/SceneVisibility.cpp b/Engine/Source/Runtime/Renderer/Private/SceneVisibility.cpp
--- a/Engine/Source/Runtime/Renderer/Private/SceneVisibility.cpp
+++ b/Engine/Source/Runtime/Renderer/Private/SceneVisibility.cpp
@@ -663,7 +663,7 @@ struct FVisForPrimParams
 //since only one of the task threads will ever reference it.  However, any modifications to shared state like the ViewState must be buffered
 //to be recombined later.
 template<bool bSingleThreaded>
-static void FetchVisibilityForPrimitives_Range(FVisForPrimParams& Params)
+static void FetchVisibilityForPrimitives_Range(FVisForPrimParams& Params, FGlobalDynamicVertexBuffer* DynamicVertexBufferIfSingleThreaded)
 {	
 	int32 NumOccludedPrimitives = 0;
 	
@@ -1043,6 +1043,8 @@ static void FetchVisibilityForPrimitives_Range(FVisForPrimParams& Params)
 
 								if (bSingleThreaded)
 								{
+									checkSlow(DynamicVertexBufferIfSingleThreaded);
+
 									if (GRHIMaximumReccommendedOustandingOcclusionQueries < MAX_int32 && !bGroupedQuery)
 									{
 										QueriesToAdd->Emplace(FPrimitiveOcclusionHistoryKey(PrimitiveId, SubQuery), BoundOrigin, BoundExtent, PrimitiveOcclusionHistory->LastQuerySubmitFrame());
@@ -1051,8 +1053,8 @@ static void FetchVisibilityForPrimitives_Range(FVisForPrimParams& Params)
 									{
 										PrimitiveOcclusionHistory->SetCurrentQuery(OcclusionFrameCounter,
 											bGroupedQuery ?
-											View.GroupedOcclusionQueries.BatchPrimitive(BoundOrigin, BoundExtent) :
-											View.IndividualOcclusionQueries.BatchPrimitive(BoundOrigin, BoundExtent),
+											View.GroupedOcclusionQueries.BatchPrimitive(BoundOrigin, BoundExtent, *DynamicVertexBufferIfSingleThreaded) :
+											View.IndividualOcclusionQueries.BatchPrimitive(BoundOrigin, BoundExtent, *DynamicVertexBufferIfSingleThreaded),
 											NumBufferedFrames,
 											bGroupedQuery,
 											Params.bNeedsScanOnRead
@@ -1171,11 +1173,11 @@ public:
 
 	void DoTask(ENamedThreads::Type CurrentThread, const FGraphEventRef& MyCompletionGraphEvent)
 	{
-		FetchVisibilityForPrimitives_Range<false>(Params);
+		FetchVisibilityForPrimitives_Range<false>(Params, nullptr);
 	}
 };
 
-static int32 FetchVisibilityForPrimitives(const FScene* Scene, FViewInfo& View, const bool bSubmitQueries, const bool bHZBOcclusion)
+static int32 FetchVisibilityForPrimitives(const FScene* Scene, FViewInfo& View, const bool bSubmitQueries, const bool bHZBOcclusion, FGlobalDynamicVertexBuffer& DynamicVertexBuffer)
 {
 	QUICK_SCOPE_CYCLE_COUNTER(STAT_FetchVisibilityForPrimitives);
 	FSceneViewState* ViewState = (FSceneViewState*)View.State;
@@ -1341,8 +1343,8 @@ static int32 FetchVisibilityForPrimitives(const FScene* Scene, FViewInfo& View,
 				{
 					RunQueriesIter->PrimitiveOcclusionHistory->SetCurrentQuery(OcclusionFrameCounter,
 						RunQueriesIter->bGroupedQuery ?
-						View.GroupedOcclusionQueries.BatchPrimitive(RunQueriesIter->BoundsOrigin, RunQueriesIter->BoundsExtent) :
-						View.IndividualOcclusionQueries.BatchPrimitive(RunQueriesIter->BoundsOrigin, RunQueriesIter->BoundsExtent),
+						View.GroupedOcclusionQueries.BatchPrimitive(RunQueriesIter->BoundsOrigin, RunQueriesIter->BoundsExtent, DynamicVertexBuffer) :
+						View.IndividualOcclusionQueries.BatchPrimitive(RunQueriesIter->BoundsOrigin, RunQueriesIter->BoundsExtent, DynamicVertexBuffer),
 						NumBufferedFrames,
 						RunQueriesIter->bGroupedQuery,
 						Params[i].bNeedsScanOnRead
@@ -1397,7 +1399,7 @@ static int32 FetchVisibilityForPrimitives(const FScene* Scene, FViewInfo& View,
 			&SubIsOccluded
 			);
 
-		FetchVisibilityForPrimitives_Range<true>(Params);
+		FetchVisibilityForPrimitives_Range<true>(Params, &DynamicVertexBuffer);
 
 		int32 IndQueries = PendingIndividualQueriesWhenOptimizing.Num();
 		if (IndQueries)
@@ -1427,7 +1429,7 @@ static int32 FetchVisibilityForPrimitives(const FScene* Scene, FViewInfo& View,
 					FPrimitiveOcclusionHistory* PrimitiveOcclusionHistory = ViewPrimitiveOcclusionHistory.Find(RunQueriesIter->PrimitiveOcclusionHistoryKey);
 
 					PrimitiveOcclusionHistory->SetCurrentQuery(OcclusionFrameCounter,
-						View.IndividualOcclusionQueries.BatchPrimitive(RunQueriesIter->BoundsOrigin, RunQueriesIter->BoundsExtent),
+						View.IndividualOcclusionQueries.BatchPrimitive(RunQueriesIter->BoundsOrigin, RunQueriesIter->BoundsExtent, DynamicVertexBuffer),
 						NumBufferedFrames,
 						false,
 						Params.bNeedsScanOnRead
@@ -1455,7 +1457,7 @@ static int32 FetchVisibilityForPrimitives(const FScene* Scene, FViewInfo& View,
 					FOcclusionBounds* RunQueriesIter = PendingIndividualQueriesWhenOptimizingSorter[Index];
 					FPrimitiveOcclusionHistory* PrimitiveOcclusionHistory = ViewPrimitiveOcclusionHistory.Find(RunQueriesIter->PrimitiveOcclusionHistoryKey);
 					PrimitiveOcclusionHistory->SetCurrentQuery(OcclusionFrameCounter,
-						View.IndividualOcclusionQueries.BatchPrimitive(RunQueriesIter->BoundsOrigin, RunQueriesIter->BoundsExtent),
+						View.IndividualOcclusionQueries.BatchPrimitive(RunQueriesIter->BoundsOrigin, RunQueriesIter->BoundsExtent, DynamicVertexBuffer),
 						NumBufferedFrames,
 						false,
 						Params.bNeedsScanOnRead
@@ -1483,7 +1485,7 @@ static int32 FetchVisibilityForPrimitives(const FScene* Scene, FViewInfo& View,
 /**
  * Cull occluded primitives in the view.
  */
-static int32 OcclusionCull(FRHICommandListImmediate& RHICmdList, const FScene* Scene, FViewInfo& View)
+static int32 OcclusionCull(FRHICommandListImmediate& RHICmdList, const FScene* Scene, FViewInfo& View, FGlobalDynamicVertexBuffer& DynamicVertexBuffer)
 {
 	SCOPE_CYCLE_COUNTER(STAT_OcclusionCull);	
 	RHICmdList.SetCurrentStat(GET_STATID(STAT_CLMM_OcclusionReadback));
@@ -1563,7 +1565,7 @@ static int32 OcclusionCull(FRHICommandListImmediate& RHICmdList, const FScene* S
 								  (!FrameParity && View.StereoPass == eSSP_RIGHT_EYE);
 			}
 
-			NumOccludedPrimitives += FetchVisibilityForPrimitives(Scene, View, bSubmitQueries, bHZBOcclusion);
+			NumOccludedPrimitives += FetchVisibilityForPrimitives(Scene, View, bSubmitQueries, bHZBOcclusion, DynamicVertexBuffer);
 
 			if( bHZBOcclusion )
 			{
@@ -2608,6 +2610,9 @@ void FSceneRenderer::GatherDynamicMeshElements(
 	TArray<FViewInfo>& InViews, 
 	const FScene* InScene, 
 	const FSceneViewFamily& InViewFamily, 
+	FGlobalDynamicIndexBuffer& DynamicIndexBuffer,
+	FGlobalDynamicVertexBuffer& DynamicVertexBuffer,
+	FGlobalDynamicReadBuffer& DynamicReadBuffer,
 	const FPrimitiveViewMasks& HasDynamicMeshElementsMasks, 
 	const FPrimitiveViewMasks& HasDynamicEditorMeshElementsMasks, 
 	const FPrimitiveViewMasks& HasViewCustomDataMasks,
@@ -2629,7 +2634,10 @@ void FSceneRenderer::GatherDynamicMeshElements(
 				&InViews[ViewIndex].DynamicMeshElements,
 				&InViews[ViewIndex].SimpleElementCollector,
 				&InViews[ViewIndex].DynamicPrimitiveShaderData, 
-				InViewFamily.GetFeatureLevel());
+				InViewFamily.GetFeatureLevel(),
+				&DynamicIndexBuffer,
+				&DynamicVertexBuffer,
+				&DynamicReadBuffer);
 		}
 
 		const bool bIsInstancedStereo = (ViewCount > 0) ? (InViews[0].IsInstancedStereoPass() || InViews[0].bIsMobileMultiViewEnabled) : false;
@@ -2692,7 +2700,10 @@ void FSceneRenderer::GatherDynamicMeshElements(
 				&InViews[ViewIndex].DynamicEditorMeshElements, 
 				&InViews[ViewIndex].EditorSimpleElementCollector, 
 				&InViews[ViewIndex].DynamicPrimitiveShaderData, 
-				InViewFamily.GetFeatureLevel());
+				InViewFamily.GetFeatureLevel(),
+				&DynamicIndexBuffer,
+				&DynamicVertexBuffer,
+				&DynamicReadBuffer);
 		}
 
 		for (int32 PrimitiveIndex = 0; PrimitiveIndex < NumPrimitives; ++PrimitiveIndex)
@@ -3261,7 +3272,8 @@ void UpdateReflectionSceneData(FScene* Scene)
 	}
 }
 
-void FSceneRenderer::ComputeViewVisibility(FRHICommandListImmediate& RHICmdList, FExclusiveDepthStencil::Type BasePassDepthStencilAccess, FViewVisibleCommandsPerView& ViewCommandsPerView)
+void FSceneRenderer::ComputeViewVisibility(FRHICommandListImmediate& RHICmdList, FExclusiveDepthStencil::Type BasePassDepthStencilAccess, FViewVisibleCommandsPerView& ViewCommandsPerView, 
+	FGlobalDynamicIndexBuffer& DynamicIndexBuffer, FGlobalDynamicVertexBuffer& DynamicVertexBuffer, FGlobalDynamicReadBuffer& DynamicReadBuffer)
 {
 	SCOPE_CYCLE_COUNTER(STAT_ViewVisibilityTime);
 	SCOPED_NAMED_EVENT(FSceneRenderer_ComputeViewVisibility, FColor::Magenta);
@@ -3507,7 +3519,7 @@ void FSceneRenderer::ComputeViewVisibility(FRHICommandListImmediate& RHICmdList,
 		// Occlusion cull for all primitives in the view frustum, but not in wireframe.
 		if (!View.Family->EngineShowFlags.Wireframe)
 		{
-			int32 NumOccludedPrimitivesInView = OcclusionCull(RHICmdList, Scene, View);
+			int32 NumOccludedPrimitivesInView = OcclusionCull(RHICmdList, Scene, View, DynamicVertexBuffer);
 			STAT(NumOccludedPrimitives += NumOccludedPrimitivesInView);
 		}
 
@@ -3593,7 +3605,8 @@ void FSceneRenderer::ComputeViewVisibility(FRHICommandListImmediate& RHICmdList,
 	}
 
 	// Gather FMeshBatches from scene proxies
-	GatherDynamicMeshElements(Views, Scene, ViewFamily, HasDynamicMeshElementsMasks, HasDynamicEditorMeshElementsMasks, HasViewCustomDataMasks, MeshCollector);
+	GatherDynamicMeshElements(Views, Scene, ViewFamily, DynamicIndexBuffer, DynamicVertexBuffer, DynamicReadBuffer,
+		HasDynamicMeshElementsMasks, HasDynamicEditorMeshElementsMasks, HasViewCustomDataMasks, MeshCollector);
 
 	for (int32 ViewIndex = 0; ViewIndex < Views.Num(); ViewIndex++)
 	{
@@ -3868,7 +3881,7 @@ bool FDeferredShadingSceneRenderer::InitViews(FRHICommandListImmediate& RHICmdLi
 	FViewVisibleCommandsPerView ViewCommandsPerView;
 	ViewCommandsPerView.SetNum(Views.Num());
 
-	ComputeViewVisibility(RHICmdList, BasePassDepthStencilAccess, ViewCommandsPerView);
+	ComputeViewVisibility(RHICmdList, BasePassDepthStencilAccess, ViewCommandsPerView, DynamicIndexBufferForInitViews, DynamicVertexBufferForInitViews, DynamicReadBufferForInitViews);
 
 	RHICmdList.ImmediateFlush(EImmediateFlushType::DispatchToRHIThread);
 
@@ -3991,7 +4004,7 @@ void FDeferredShadingSceneRenderer::InitViewsPossiblyAfterPrepass(FRHICommandLis
 		&& !ViewFamily.EngineShowFlags.HitProxies)
 	{
 		// Setup dynamic shadows.
-		InitDynamicShadows(RHICmdList);
+		InitDynamicShadows(RHICmdList, DynamicIndexBufferForInitShadows, DynamicVertexBufferForInitShadows, DynamicReadBufferForInitShadows);
 
 		RHICmdList.ImmediateFlush(EImmediateFlushType::DispatchToRHIThread);
 	}
diff --git a/Engine/Source/Runtime/Renderer/Private/ShadowRendering.h b/Engine/Source/Runtime/Renderer/Private/ShadowRendering.h
--- a/Engine/Source/Runtime/Renderer/Private/ShadowRendering.h
+++ b/Engine/Source/Runtime/Renderer/Private/ShadowRendering.h
@@ -423,7 +423,8 @@ public:
 	void AddReceiverPrimitive(FPrimitiveSceneInfo* PrimitiveSceneInfo);
 
 	/** Gathers dynamic mesh elements for all the shadow's primitives arrays. */
-	void GatherDynamicMeshElements(FSceneRenderer& Renderer, class FVisibleLightInfo& VisibleLightInfo, TArray<const FSceneView*>& ReusedViewsArray);
+	void GatherDynamicMeshElements(FSceneRenderer& Renderer, class FVisibleLightInfo& VisibleLightInfo, TArray<const FSceneView*>& ReusedViewsArray, 
+		FGlobalDynamicIndexBuffer& DynamicIndexBuffer, FGlobalDynamicVertexBuffer& DynamicVertexBuffer, FGlobalDynamicReadBuffer& DynamicReadBuffer);
 
 	void SetupMeshDrawCommandsForShadowDepth(FSceneRenderer& Renderer, FUniformBufferRHIParamRef PassUniformBuffer);
 
@@ -585,6 +586,9 @@ private:
 	void GatherDynamicMeshElementsArray(
 		FViewInfo* FoundView,
 		FSceneRenderer& Renderer, 
+		FGlobalDynamicIndexBuffer& DynamicIndexBuffer,
+		FGlobalDynamicVertexBuffer& DynamicVertexBuffer,
+		FGlobalDynamicReadBuffer& DynamicReadBuffer,
 		const PrimitiveArrayType& PrimitiveArray, 
 		const TArray<const FSceneView*>& ReusedViewsArray,
 		TArray<FMeshBatchAndRelevance,SceneRenderingAllocator>& OutDynamicMeshElements,
diff --git a/Engine/Source/Runtime/Renderer/Private/ShadowSetup.cpp b/Engine/Source/Runtime/Renderer/Private/ShadowSetup.cpp
--- a/Engine/Source/Runtime/Renderer/Private/ShadowSetup.cpp
+++ b/Engine/Source/Runtime/Renderer/Private/ShadowSetup.cpp
@@ -1345,7 +1345,8 @@ void FProjectedShadowInfo::ApplyViewOverridesToMeshDrawCommands(const FViewInfo&
 	}
 }
 
-void FProjectedShadowInfo::GatherDynamicMeshElements(FSceneRenderer& Renderer, FVisibleLightInfo& VisibleLightInfo, TArray<const FSceneView*>& ReusedViewsArray)
+void FProjectedShadowInfo::GatherDynamicMeshElements(FSceneRenderer& Renderer, FVisibleLightInfo& VisibleLightInfo, TArray<const FSceneView*>& ReusedViewsArray,
+	FGlobalDynamicIndexBuffer& DynamicIndexBuffer, FGlobalDynamicVertexBuffer& DynamicVertexBuffer, FGlobalDynamicReadBuffer& DynamicReadBuffer)
 {
 	QUICK_SCOPE_CYCLE_COUNTER(STAT_Shadow_GatherDynamicMeshElements);
 
@@ -1370,25 +1371,29 @@ void FProjectedShadowInfo::GatherDynamicMeshElements(FSceneRenderer& Renderer, F
 		{
 			ShadowDepthView->SetPreShadowTranslation(FVector(0, 0, 0));
 			ShadowDepthView->SetDynamicMeshElementsShadowCullFrustum(&CascadeSettings.ShadowBoundsAccurate);
-			GatherDynamicMeshElementsArray(ShadowDepthView, Renderer, DynamicSubjectPrimitives, ReusedViewsArray, DynamicSubjectMeshElements, NumDynamicSubjectMeshElements);
+			GatherDynamicMeshElementsArray(ShadowDepthView, Renderer, DynamicIndexBuffer, DynamicVertexBuffer, DynamicReadBuffer, 
+				DynamicSubjectPrimitives, ReusedViewsArray, DynamicSubjectMeshElements, NumDynamicSubjectMeshElements);
 			ShadowDepthView->SetPreShadowTranslation(PreShadowTranslation);
 		}
 		else
 		{
 			ShadowDepthView->SetPreShadowTranslation(PreShadowTranslation);
 			ShadowDepthView->SetDynamicMeshElementsShadowCullFrustum(&CasterFrustum);
-			GatherDynamicMeshElementsArray(ShadowDepthView, Renderer, DynamicSubjectPrimitives, ReusedViewsArray, DynamicSubjectMeshElements, NumDynamicSubjectMeshElements);
+			GatherDynamicMeshElementsArray(ShadowDepthView, Renderer, DynamicIndexBuffer, DynamicVertexBuffer, DynamicReadBuffer, 
+				DynamicSubjectPrimitives, ReusedViewsArray, DynamicSubjectMeshElements, NumDynamicSubjectMeshElements);
 		}
 
 		ShadowDepthView->DrawDynamicFlags = EDrawDynamicFlags::None;
 
 		int32 NumDynamicReceiverMeshElements = 0;
 		ShadowDepthView->SetDynamicMeshElementsShadowCullFrustum(&ReceiverFrustum);
-		GatherDynamicMeshElementsArray(ShadowDepthView, Renderer, ReceiverPrimitives, ReusedViewsArray, DynamicReceiverMeshElements, NumDynamicReceiverMeshElements);
+		GatherDynamicMeshElementsArray(ShadowDepthView, Renderer, DynamicIndexBuffer, DynamicVertexBuffer, DynamicReadBuffer, 
+			ReceiverPrimitives, ReusedViewsArray, DynamicReceiverMeshElements, NumDynamicReceiverMeshElements);
 
 		int32 NumDynamicSubjectTranslucentMeshElements = 0;
 		ShadowDepthView->SetDynamicMeshElementsShadowCullFrustum(&CasterFrustum);
-		GatherDynamicMeshElementsArray(ShadowDepthView, Renderer, SubjectTranslucentPrimitives, ReusedViewsArray, DynamicSubjectTranslucentMeshElements, NumDynamicSubjectTranslucentMeshElements);
+		GatherDynamicMeshElementsArray(ShadowDepthView, Renderer, DynamicIndexBuffer, DynamicVertexBuffer, DynamicReadBuffer, 
+			SubjectTranslucentPrimitives, ReusedViewsArray, DynamicSubjectTranslucentMeshElements, NumDynamicSubjectTranslucentMeshElements);
 
 		Renderer.MeshCollector.ProcessTasks();
     }
@@ -1417,6 +1422,9 @@ void FProjectedShadowInfo::GatherDynamicMeshElements(FSceneRenderer& Renderer, F
 void FProjectedShadowInfo::GatherDynamicMeshElementsArray(
 	FViewInfo* FoundView,
 	FSceneRenderer& Renderer, 
+	FGlobalDynamicIndexBuffer& DynamicIndexBuffer,
+	FGlobalDynamicVertexBuffer& DynamicVertexBuffer,
+	FGlobalDynamicReadBuffer& DynamicReadBuffer,
 	const PrimitiveArrayType& PrimitiveArray, 
 	const TArray<const FSceneView*>& ReusedViewsArray,
 	TArray<FMeshBatchAndRelevance,SceneRenderingAllocator>& OutDynamicMeshElements,
@@ -1426,7 +1434,16 @@ void FProjectedShadowInfo::GatherDynamicMeshElementsArray(
 	FSimpleElementCollector DynamicSubjectSimpleElements;
 
 	Renderer.MeshCollector.ClearViewMeshArrays();
-	Renderer.MeshCollector.AddViewMeshArrays(FoundView, &OutDynamicMeshElements, &DynamicSubjectSimpleElements, &FoundView->DynamicPrimitiveShaderData, Renderer.ViewFamily.GetFeatureLevel());
+	Renderer.MeshCollector.AddViewMeshArrays(
+		FoundView, 
+		&OutDynamicMeshElements, 
+		&DynamicSubjectSimpleElements, 
+		&FoundView->DynamicPrimitiveShaderData, 
+		Renderer.ViewFamily.GetFeatureLevel(),
+		&DynamicIndexBuffer,
+		&DynamicVertexBuffer,
+		&DynamicReadBuffer
+		);
 
 	const uint32 PrimitiveCount = PrimitiveArray.Num();
 
@@ -2871,7 +2888,7 @@ void FSceneRenderer::InitProjectedShadowVisibility(FRHICommandListImmediate& RHI
 #endif // !UE_BUILD_SHIPPING
 }
 
-void FSceneRenderer::GatherShadowDynamicMeshElements()
+void FSceneRenderer::GatherShadowDynamicMeshElements(FGlobalDynamicIndexBuffer& DynamicIndexBuffer, FGlobalDynamicVertexBuffer& DynamicVertexBuffer, FGlobalDynamicReadBuffer& DynamicReadBuffer)
 {
 	TArray<const FSceneView*> ReusedViewsArray;
 	ReusedViewsArray.AddZeroed(1);
@@ -2884,7 +2901,7 @@ void FSceneRenderer::GatherShadowDynamicMeshElements()
 		{
 			FProjectedShadowInfo* ProjectedShadowInfo = Atlas.Shadows[ShadowIndex];
 			FVisibleLightInfo& VisibleLightInfo = VisibleLightInfos[ProjectedShadowInfo->GetLightSceneInfo().Id];
-			ProjectedShadowInfo->GatherDynamicMeshElements(*this, VisibleLightInfo, ReusedViewsArray);
+			ProjectedShadowInfo->GatherDynamicMeshElements(*this, VisibleLightInfo, ReusedViewsArray, DynamicIndexBuffer, DynamicVertexBuffer, DynamicReadBuffer);
 		}
 	}
 
@@ -2896,7 +2913,7 @@ void FSceneRenderer::GatherShadowDynamicMeshElements()
 		{
 			FProjectedShadowInfo* ProjectedShadowInfo = Atlas.Shadows[ShadowIndex];
 			FVisibleLightInfo& VisibleLightInfo = VisibleLightInfos[ProjectedShadowInfo->GetLightSceneInfo().Id];
-			ProjectedShadowInfo->GatherDynamicMeshElements(*this, VisibleLightInfo, ReusedViewsArray);
+			ProjectedShadowInfo->GatherDynamicMeshElements(*this, VisibleLightInfo, ReusedViewsArray, DynamicIndexBuffer, DynamicVertexBuffer, DynamicReadBuffer);
 		}
 	}
 
@@ -2908,7 +2925,7 @@ void FSceneRenderer::GatherShadowDynamicMeshElements()
 		{
 			FProjectedShadowInfo* ProjectedShadowInfo = Atlas.Shadows[ShadowIndex];
 			FVisibleLightInfo& VisibleLightInfo = VisibleLightInfos[ProjectedShadowInfo->GetLightSceneInfo().Id];
-			ProjectedShadowInfo->GatherDynamicMeshElements(*this, VisibleLightInfo, ReusedViewsArray);
+			ProjectedShadowInfo->GatherDynamicMeshElements(*this, VisibleLightInfo, ReusedViewsArray, DynamicIndexBuffer, DynamicVertexBuffer, DynamicReadBuffer);
 		}
 	}
 
@@ -2916,7 +2933,7 @@ void FSceneRenderer::GatherShadowDynamicMeshElements()
 	{
 		FProjectedShadowInfo* ProjectedShadowInfo = SortedShadowsForShadowDepthPass.PreshadowCache.Shadows[ShadowIndex];
 		FVisibleLightInfo& VisibleLightInfo = VisibleLightInfos[ProjectedShadowInfo->GetLightSceneInfo().Id];
-		ProjectedShadowInfo->GatherDynamicMeshElements(*this, VisibleLightInfo, ReusedViewsArray);
+		ProjectedShadowInfo->GatherDynamicMeshElements(*this, VisibleLightInfo, ReusedViewsArray, DynamicIndexBuffer, DynamicVertexBuffer, DynamicReadBuffer);
 	}
 
 	for (int32 AtlasIndex = 0; AtlasIndex < SortedShadowsForShadowDepthPass.TranslucencyShadowMapAtlases.Num(); AtlasIndex++)
@@ -2927,7 +2944,7 @@ void FSceneRenderer::GatherShadowDynamicMeshElements()
 		{
 			FProjectedShadowInfo* ProjectedShadowInfo = Atlas.Shadows[ShadowIndex];
 			FVisibleLightInfo& VisibleLightInfo = VisibleLightInfos[ProjectedShadowInfo->GetLightSceneInfo().Id];
-			ProjectedShadowInfo->GatherDynamicMeshElements(*this, VisibleLightInfo, ReusedViewsArray);
+			ProjectedShadowInfo->GatherDynamicMeshElements(*this, VisibleLightInfo, ReusedViewsArray, DynamicIndexBuffer, DynamicVertexBuffer, DynamicReadBuffer);
 		}
 	}
 }
@@ -4042,7 +4059,7 @@ void FSceneRenderer::AllocateTranslucentShadowDepthTargets(FRHICommandListImmedi
 	}
 }
 
-void FSceneRenderer::InitDynamicShadows(FRHICommandListImmediate& RHICmdList)
+void FSceneRenderer::InitDynamicShadows(FRHICommandListImmediate& RHICmdList, FGlobalDynamicIndexBuffer& DynamicIndexBuffer, FGlobalDynamicVertexBuffer& DynamicVertexBuffer, FGlobalDynamicReadBuffer& DynamicReadBuffer)
 {
 	SCOPE_CYCLE_COUNTER(STAT_DynamicShadowSetupTime);
 	SCOPED_NAMED_EVENT(FSceneRenderer_InitDynamicShadows, FColor::Magenta);
@@ -4192,6 +4209,6 @@ void FSceneRenderer::InitDynamicShadows(FRHICommandListImmediate& RHICmdList)
 	AllocateShadowDepthTargets(RHICmdList);
 
 	// Generate mesh element arrays from shadow primitive arrays
-	GatherShadowDynamicMeshElements();
+	GatherShadowDynamicMeshElements(DynamicIndexBuffer, DynamicVertexBuffer, DynamicReadBuffer);
 
 }
diff --git a/Engine/Source/Runtime/Renderer/Private/ShadowSetupMobile.cpp b/Engine/Source/Runtime/Renderer/Private/ShadowSetupMobile.cpp
--- a/Engine/Source/Runtime/Renderer/Private/ShadowSetupMobile.cpp
+++ b/Engine/Source/Runtime/Renderer/Private/ShadowSetupMobile.cpp
@@ -303,7 +303,7 @@ void FMobileSceneRenderer::InitDynamicShadows(FRHICommandListImmediate& RHICmdLi
 		}
 	}
 
-	FSceneRenderer::InitDynamicShadows(RHICmdList);
+	FSceneRenderer::InitDynamicShadows(RHICmdList, DynamicIndexBuffer, DynamicVertexBuffer, DynamicReadBuffer);
 
 	PrepareViewVisibilityLists();
 
