Optimized dynamic mesh relevance computation.

diff --git a/Engine/Source/Runtime/Renderer/Private/SceneVisibility.cpp b/Engine/Source/Runtime/Renderer/Private/SceneVisibility.cpp
--- a/Engine/Source/Runtime/Renderer/Private/SceneVisibility.cpp
+++ b/Engine/Source/Runtime/Renderer/Private/SceneVisibility.cpp
@@ -2658,34 +2658,34 @@ void FSceneRenderer::GatherDynamicMeshElements(
 				SetDynamicMeshElementViewCustomData(InViews, HasViewCustomDataMasks, PrimitiveSceneInfo);
 
 				PrimitiveSceneInfo->Proxy->GetDynamicMeshElements(InViewFamily.Views, InViewFamily, ViewMaskFinal, Collector);
-			}
-
-			// to support GetDynamicMeshElementRange()
-			for (int32 ViewIndex = 0; ViewIndex < ViewCount; ViewIndex++)
-			{
-				InViews[ViewIndex].DynamicMeshEndIndices[PrimitiveIndex] = Collector.GetMeshBatchCount(ViewIndex);
-			}
 
-			// Compute DynamicMeshElementsMeshPassRelevance for this primitive.
-			for (int32 ViewIndex = 0; ViewIndex < ViewCount; ViewIndex++)
-			{
-				FViewInfo& View = InViews[ViewIndex];
-				const bool bAddLightmapDensityCommands = View.Family->EngineShowFlags.LightMapDensity && AllowDebugViewmodes();
-				const FPrimitiveViewRelevance& ViewRelevance = View.PrimitiveViewRelevanceMap[PrimitiveIndex];
+				// Compute DynamicMeshElementsMeshPassRelevance for this primitive.
+				for (int32 ViewIndex = 0; ViewIndex < ViewCount; ViewIndex++)
+				{
+					FViewInfo& View = InViews[ViewIndex];
+					const bool bAddLightmapDensityCommands = View.Family->EngineShowFlags.LightMapDensity && AllowDebugViewmodes();
+					const FPrimitiveViewRelevance& ViewRelevance = View.PrimitiveViewRelevanceMap[PrimitiveIndex];
 
-				View.DynamicMeshElementsPassRelevance.SetNum(View.DynamicMeshElements.Num());
+					View.DynamicMeshElementsPassRelevance.SetNum(View.DynamicMeshElements.Num());
 
-				const int32 NumPrimitiveElements = Collector.GetMeshBatchCount(ViewIndex);
-				const int32 PrimitiveFirstElement = View.DynamicMeshElements.Num() - NumPrimitiveElements;
+					const int32 NumPrimitiveElements = Collector.GetMeshBatchCount(ViewIndex);
+					const int32 PrimitiveFirstElement = View.DynamicMeshElements.Num() - NumPrimitiveElements;
 
-				for (int32 ElementIndex = PrimitiveFirstElement; ElementIndex < NumPrimitiveElements; ++ElementIndex)
-				{
-					const FMeshBatchAndRelevance& MeshBatch = View.DynamicMeshElements[ElementIndex];
-					FMeshPassMask& PassRelevance = View.DynamicMeshElementsPassRelevance[ElementIndex];
+					for (int32 ElementIndex = PrimitiveFirstElement; ElementIndex < NumPrimitiveElements; ++ElementIndex)
+					{
+						const FMeshBatchAndRelevance& MeshBatch = View.DynamicMeshElements[ElementIndex];
+						FMeshPassMask& PassRelevance = View.DynamicMeshElementsPassRelevance[ElementIndex];
 
-					ComputeDynamicMeshRelevance(ShadingPath, bAddLightmapDensityCommands, ViewRelevance, MeshBatch, View, PassRelevance);
+						ComputeDynamicMeshRelevance(ShadingPath, bAddLightmapDensityCommands, ViewRelevance, MeshBatch, View, PassRelevance);
+					}
 				}
 			}
+
+			// to support GetDynamicMeshElementRange()
+			for (int32 ViewIndex = 0; ViewIndex < ViewCount; ViewIndex++)
+			{
+				InViews[ViewIndex].DynamicMeshEndIndices[PrimitiveIndex] = Collector.GetMeshBatchCount(ViewIndex);
+			}
 		}
 	}
 
