InitView optimization. Removed DynamicMeshEndIndices array. This array was walked for every primitive in the scene, even through it was required only for a small subset of them. Namely mesh decals and volumetric density injection primitives. Replaced it with building directly two separate lists - one for mesh decals and one for volumetric primitives.

diff --git a/Engine/Source/Runtime/Renderer/Private/CompositionLighting/PostProcessDeferredDecals.cpp b/Engine/Source/Runtime/Renderer/Private/CompositionLighting/PostProcessDeferredDecals.cpp
--- a/Engine/Source/Runtime/Renderer/Private/CompositionLighting/PostProcessDeferredDecals.cpp
+++ b/Engine/Source/Runtime/Renderer/Private/CompositionLighting/PostProcessDeferredDecals.cpp
@@ -546,7 +546,7 @@ void FRCPassPostProcessDeferredDecals::Process(FRenderingCompositePassContext& C
 		if (ViewFamily.EngineShowFlags.Decals)
 		{
 			FScene& Scene = *(FScene*)ViewFamily.Scene;
-			if (Scene.Decals.Num() > 0 || Context.View.MeshDecalPrimSet.NumPrims() > 0)
+			if (Scene.Decals.Num() > 0 || Context.View.MeshDecalBatches.Num() > 0)
 			{
 				bNeedsDBufferTargets = true;
 			}
@@ -642,7 +642,7 @@ void FRCPassPostProcessDeferredDecals::Process(FRenderingCompositePassContext& C
 
 		if (CurrentStage == DRS_BeforeBasePass || CurrentStage == DRS_BeforeLighting || CurrentStage == DRS_Emissive)
 		{
-			if (Context.View.MeshDecalPrimSet.NumPrims() > 0)
+			if (Context.View.MeshDecalBatches.Num() > 0)
 			{
 				check(bNeedsDBufferTargets || CurrentStage != DRS_BeforeBasePass);
 				RenderMeshDecals(Context, CurrentStage);
@@ -659,7 +659,7 @@ void FRCPassPostProcessDeferredDecals::Process(FRenderingCompositePassContext& C
 		FDecalRenderTargetManager RenderTargetManager(RHICmdList, Context.GetShaderPlatform(), CurrentStage);
 
 		//don't early return. Resolves must be run for fast clears to work.
-		if (Scene.Decals.Num() || Context.View.MeshDecalPrimSet.NumPrims() > 0)
+		if (Scene.Decals.Num() || Context.View.MeshDecalBatches.Num() > 0)
 		{
 			check(bNeedsDBufferTargets || CurrentStage != DRS_BeforeBasePass);
 
diff --git a/Engine/Source/Runtime/Renderer/Private/CompositionLighting/PostProcessMeshDecals.cpp b/Engine/Source/Runtime/Renderer/Private/CompositionLighting/PostProcessMeshDecals.cpp
--- a/Engine/Source/Runtime/Renderer/Private/CompositionLighting/PostProcessMeshDecals.cpp
+++ b/Engine/Source/Runtime/Renderer/Private/CompositionLighting/PostProcessMeshDecals.cpp
@@ -222,7 +222,7 @@ FMeshDecalMeshProcessor::FMeshDecalMeshProcessor(const FScene* Scene,
 
 void FMeshDecalMeshProcessor::AddMeshBatch(const FMeshBatch& RESTRICT MeshBatch, uint64 BatchElementMask, const FPrimitiveSceneProxy* RESTRICT PrimitiveSceneProxy, int32 StaticMeshId)
 {
-	if (MeshBatch.bUseForMaterial)
+	if (MeshBatch.bUseForMaterial && MeshBatch.IsDecal(FeatureLevel))
 	{
 		const FMaterialRenderProxy* MaterialRenderProxy = MeshBatch.MaterialRenderProxy;
 		const FMaterial* Material = MaterialRenderProxy->GetMaterialNoFallback(FeatureLevel);
@@ -371,50 +371,13 @@ void DrawDecalMeshCommands(FRenderingCompositePassContext& Context, EDecalRender
 			RenderTargetMode,
 			DynamicMeshPassContext);
 
-		const int32 PrimitiveNum = View.MeshDecalPrimSet.NumPrims();
-		for (int32 PrimitiveIndex = 0; PrimitiveIndex < PrimitiveNum; ++PrimitiveIndex)
+		for (int32 MeshBatchIndex = 0; MeshBatchIndex < View.MeshDecalBatches.Num(); ++MeshBatchIndex)
 		{
-			FPrimitiveSceneInfo* PrimitiveSceneInfo = View.MeshDecalPrimSet.Prims[PrimitiveIndex].PrimitiveSceneInfo;
+			const FMeshBatch* Mesh = View.MeshDecalBatches[MeshBatchIndex].Mesh;
+			const FPrimitiveSceneProxy* PrimitiveSceneProxy = View.MeshDecalBatches[MeshBatchIndex].Proxy;
+			const uint64 DefaultBatchElementMask = ~0ull;
 
-			const uint32 PrimitiveId = PrimitiveSceneInfo->GetIndex();
-			const FPrimitiveViewRelevance& ViewRelevance = View.PrimitiveViewRelevanceMap[PrimitiveId];
-			if (ViewRelevance.bDrawRelevance)
-			{
-				const uint64 DefaultBatchElementMask = ~0ull;
-
-				// Render dynamic scene prim
-				{
-					// range in View.DynamicMeshElements[]
-					FInt32Range range = View.GetDynamicMeshElementRange(PrimitiveSceneInfo->GetIndex());
-
-					for (int32 MeshBatchIndex = range.GetLowerBoundValue(); MeshBatchIndex < range.GetUpperBoundValue(); MeshBatchIndex++)
-					{
-						const FMeshBatchAndRelevance& MeshBatchAndRelevance = View.DynamicMeshElements[MeshBatchIndex];
-
-						checkSlow(MeshBatchAndRelevance.PrimitiveSceneProxy == PrimitiveSceneInfo->Proxy);
-
-						const FMeshBatch& MeshBatch = *MeshBatchAndRelevance.Mesh;
-
-						PassMeshProcessor.AddMeshBatch(MeshBatch, DefaultBatchElementMask, PrimitiveSceneInfo->Proxy);
-					}
-				}
-
-				// Render static scene prim
-				if (ViewRelevance.bStaticRelevance)
-				{
-					// Render static meshes from static scene prim
-					for (int32 StaticMeshIdx = 0, Count = PrimitiveSceneInfo->StaticMeshes.Num(); StaticMeshIdx < Count; StaticMeshIdx++)
-					{
-						const FStaticMeshBatch& StaticMesh = PrimitiveSceneInfo->StaticMeshes[StaticMeshIdx];
-
-						// Only render static mesh elements using decal materials
-						if (View.StaticMeshVisibilityMap[StaticMesh.Id] && StaticMesh.IsDecal(View.FeatureLevel))
-						{
-							PassMeshProcessor.AddMeshBatch(StaticMesh, DefaultBatchElementMask, PrimitiveSceneInfo->Proxy);
-						}
-					}
-				}
-			}
+			PassMeshProcessor.AddMeshBatch(*Mesh, DefaultBatchElementMask, PrimitiveSceneProxy);
 		}
 	});
 }
@@ -433,7 +396,7 @@ void RenderMeshDecals(FRenderingCompositePassContext& Context, EDecalRenderStage
 	SetupSceneTextureUniformParameters(SceneContext, View.FeatureLevel, ESceneTextureSetupMode::All, SceneTextureParameters);
 	Scene->UniformBuffers.MeshDecalPassUniformBuffer.UpdateUniformBufferImmediate(SceneTextureParameters);
 
-	if (View.MeshDecalPrimSet.NumPrims() > 0)
+	if (View.MeshDecalBatches.Num() > 0)
 	{
 		switch (CurrentDecalStage)
 		{
diff --git a/Engine/Source/Runtime/Renderer/Private/SceneRendering.h b/Engine/Source/Runtime/Renderer/Private/SceneRendering.h
--- a/Engine/Source/Runtime/Renderer/Private/SceneRendering.h
+++ b/Engine/Source/Runtime/Renderer/Private/SceneRendering.h
@@ -216,70 +216,6 @@ public:
 	}
 };
 
-
-/** 
-* Set of sorted scene prims  
-*/
-template <class TKey>
-class FSortedPrimSet
-{
-public:
-	// contains a scene prim and its sort key
-	struct FSortedPrim
-	{
-		// Default constructor
-		FSortedPrim() {}
-
-		FSortedPrim(FPrimitiveSceneInfo* InPrimitiveSceneInfo, const TKey InSortKey)
-			:	PrimitiveSceneInfo(InPrimitiveSceneInfo)
-			,	SortKey(InSortKey)
-		{
-		}
-
-		FORCEINLINE bool operator<( const FSortedPrim& rhs ) const
-		{
-			return SortKey < rhs.SortKey;
-		}
-
-		//
-		FPrimitiveSceneInfo* PrimitiveSceneInfo;
-		//
-		TKey SortKey;
-	};
-
-	/**
-	* Sort any primitives that were added to the set back-to-front
-	*/
-	void SortPrimitives()
-	{
-		Prims.Sort();
-	}
-
-	/** 
-	* @return number of prims to render
-	*/
-	int32 NumPrims() const
-	{
-		return Prims.Num();
-	}
-
-	/** list of primitives, sorted after calling Sort() */
-	TArray<FSortedPrim, SceneRenderingAllocator> Prims;
-};
-
-template <> struct TIsPODType<FSortedPrimSet<uint32>::FSortedPrim> { enum { Value = true }; };
-
-class FMeshDecalPrimSet : public FSortedPrimSet<uint32>
-{
-public:
-	typedef FSortedPrimSet<uint32>::FSortedPrim KeyType;
-
-	static KeyType GenerateKey(FPrimitiveSceneInfo* PrimitiveSceneInfo, int16 InSortPriority)
-	{
-		return KeyType(PrimitiveSceneInfo, (uint32)(InSortPriority - SHRT_MIN));
-	}
-};
-
 /** A batched occlusion primitive. */
 struct FOcclusionPrimitive
 {
@@ -604,39 +540,22 @@ public:
 	}
 };
 
-class FVolumetricPrimSet
+struct FVolumetricMeshBatch
 {
-public:
-
-	/**
-	* Adds a new primitives to the list of distortion prims
-	* @param PrimitiveSceneProxies - primitive info to add.
-	*/
-	void Append(FPrimitiveSceneProxy** PrimitiveSceneProxies, int32 NumProxies)
-	{
-		Prims.Append(PrimitiveSceneProxies, NumProxies);
-	}
+	const FMeshBatch* Mesh;
+	const FPrimitiveSceneProxy* Proxy;
+};
 
-	/** 
-	* @return number of prims to render
-	*/
-	int32 NumPrims() const
-	{
-		return Prims.Num();
-	}
+struct FMeshDecalBatch
+{
+	const FMeshBatch* Mesh;
+	const FPrimitiveSceneProxy* Proxy;
+	int16 SortKey;
 
-	/** 
-	* @return a prim currently set to render
-	*/
-	const FPrimitiveSceneProxy* GetPrim(int32 i)const
+	FORCEINLINE bool operator<(const FMeshDecalBatch& rhs) const
 	{
-		check(i>=0 && i<NumPrims());
-		return Prims[i];
+		return SortKey < rhs.SortKey;
 	}
-
-private:
-	/** list of distortion prims added from the scene */
-	TArray<FPrimitiveSceneProxy*, SceneRenderingAllocator> Prims;
 };
 
 static const int32 GMaxNumReflectionCaptures = 341;
@@ -859,15 +778,15 @@ public:
 
 	/** Count of translucent prims for this view. */
 	FTranslucenyPrimCount TranslucentPrimCount;
-
-	/** Set of mesh decal prims for this view */
-	FMeshDecalPrimSet MeshDecalPrimSet;
 	
 	bool bHasDistortionPrimitives;
 	bool bHasCustomDepthPrimitives;
 
-	/** Primitives with a volumetric material. */
-	FVolumetricPrimSet VolumetricPrimSet;
+	/** Mesh batches with for mesh decal rendering. */
+	TArray<FMeshDecalBatch, SceneRenderingAllocator> MeshDecalBatches;
+
+	/** Mesh batches with a volumetric material. */
+	TArray<FVolumetricMeshBatch, SceneRenderingAllocator> VolumetricMeshBatches;
 
 	/** A map from light ID to a boolean visibility value. */
 	TArray<FVisibleLightViewInfo,SceneRenderingAllocator> VisibleLightInfos;
@@ -896,9 +815,6 @@ public:
 	/** Gathered in UpdateRayTracingWorld from all the primitives with dynamic view relevance, used in each mesh pass. */
 	TArray<FMeshBatchAndRelevance, SceneRenderingAllocator> RayTracedDynamicMeshElements;
 
-	// [PrimitiveIndex] = end index index in DynamicMeshElements[], to support GetDynamicMeshElementRange()
-	TArray<uint32,SceneRenderingAllocator> DynamicMeshEndIndices;
-
 	TArray<FMeshBatchAndRelevance,SceneRenderingAllocator> DynamicEditorMeshElements;
 
 	FSimpleElementCollector SimpleElementCollector;
@@ -1187,18 +1103,6 @@ public:
 
 	/** Destroy all snapshots before we wipe the scene allocator. */
 	static void DestroyAllSnapshots();
-	
-	// Get the range in DynamicMeshElements[] for a given PrimitiveIndex
-	// @return range (start is inclusive, end is exclusive)
-	FInt32Range GetDynamicMeshElementRange(uint32 PrimitiveIndex) const
-	{
-		// inclusive
-		int32 Start = (PrimitiveIndex == 0) ? 0 : DynamicMeshEndIndices[PrimitiveIndex - 1];
-		// exclusive
-		int32 AfterEnd = DynamicMeshEndIndices[PrimitiveIndex];
-		
-		return FInt32Range(Start, AfterEnd);
-	}
 
 	/** Set the custom data associated with a primitive scene info.	*/
 	void SetCustomData(const FPrimitiveSceneInfo* InPrimitiveSceneInfo, void* InCustomData);
diff --git a/Engine/Source/Runtime/Renderer/Private/SceneVisibility.cpp b/Engine/Source/Runtime/Renderer/Private/SceneVisibility.cpp
--- a/Engine/Source/Runtime/Renderer/Private/SceneVisibility.cpp
+++ b/Engine/Source/Runtime/Renderer/Private/SceneVisibility.cpp
@@ -1751,13 +1751,14 @@ struct FRelevancePacket
 	int32 NumVisibleDynamicEditorPrimitives;
 	FMeshPassMask VisibleDynamicMeshesPassMask;
 	FTranslucenyPrimCount TranslucentPrimCount;
-	FRelevancePrimSet<FMeshDecalPrimSet::KeyType> MeshDecalPrimSet;
 	bool bHasDistortionPrimitives;
 	bool bHasCustomDepthPrimitives;
 	FRelevancePrimSet<FPrimitiveSceneInfo*> LazyUpdatePrimitives;
 	FRelevancePrimSet<FPrimitiveSceneInfo*> DirtyIndirectLightingCacheBufferPrimitives;
 	FRelevancePrimSet<FPrimitiveSceneInfo*> RecachedReflectionCapturePrimitives;
-	FRelevancePrimSet<FPrimitiveSceneProxy*> VolumetricPrimSet;
+
+	TArray<FMeshDecalBatch> MeshDecalBatches;
+	TArray<FVolumetricMeshBatch> VolumetricMeshBatches;
 	FDrawCommandRelevancePacket DrawCommandPacket;
 
 	struct FPrimitiveLODMask
@@ -1891,11 +1892,6 @@ struct FRelevancePacket
 				continue;
 			}
 
-			if (ViewRelevance.bDecal && ViewRelevance.bRenderInMainPass)
-			{
-				MeshDecalPrimSet.AddPrim(FMeshDecalPrimSet::GenerateKey(PrimitiveSceneInfo, PrimitiveSceneInfo->Proxy->GetTranslucencySortPriority()));
-			}
-
 			if (bEditorRelevance)
 			{
 				++NumVisibleDynamicEditorPrimitives;
@@ -1948,11 +1944,6 @@ struct FRelevancePacket
 				}
 			}
 
-			if (ViewRelevance.bHasVolumeMaterialDomain)
-			{
-				VolumetricPrimSet.AddPrim(PrimitiveSceneInfo->Proxy);
-			}
-
 			CombinedShadingModelMask |= ViewRelevance.ShadingModelMaskRelevance;
 			bUsesGlobalDistanceField |= ViewRelevance.bUsesGlobalDistanceField;
 			bUsesLightingChannels |= ViewRelevance.bUsesLightingChannels;
@@ -2231,6 +2222,24 @@ struct FRelevancePacket
 						DrawCommandPacket.AddCommandsForMesh(PrimitiveIndex, PrimitiveSceneInfo, StaticMeshRelevance, StaticMesh, Scene, bCanCache, EMeshPass::EditorSelection);
 					}
 #endif
+
+					if (ViewRelevance.bHasVolumeMaterialDomain)
+					{
+						VolumetricMeshBatches.AddUninitialized(1);
+						FVolumetricMeshBatch& BatchAndProxy = VolumetricMeshBatches.Last();
+						BatchAndProxy.Mesh = &StaticMesh;
+						BatchAndProxy.Proxy = PrimitiveSceneInfo->Proxy;
+					}
+
+					if (ViewRelevance.bRenderInMainPass && ViewRelevance.bDecal)
+					{
+						MeshDecalBatches.AddUninitialized(1);
+						FMeshDecalBatch& BatchAndProxy = MeshDecalBatches.Last();
+						BatchAndProxy.Mesh = &StaticMesh;
+						BatchAndProxy.Proxy = PrimitiveSceneInfo->Proxy;
+						BatchAndProxy.SortKey = PrimitiveSceneInfo->Proxy->GetTranslucencySortPriority();
+					}
+
 					if (MarkMask)
 					{
 						MarkMasks[StaticMeshRelevance.Id] = MarkMask;
@@ -2267,11 +2276,12 @@ struct FRelevancePacket
 		WriteView.NumVisibleDynamicPrimitives += NumVisibleDynamicPrimitives;
 		WriteView.NumVisibleDynamicEditorPrimitives += NumVisibleDynamicEditorPrimitives;
 		WriteView.TranslucentPrimCount.Append(TranslucentPrimCount);
-		MeshDecalPrimSet.AppendTo(WriteView.MeshDecalPrimSet.Prims);
 		WriteView.bHasDistortionPrimitives |= bHasDistortionPrimitives;
 		WriteView.bHasCustomDepthPrimitives |= bHasCustomDepthPrimitives;
 		DirtyIndirectLightingCacheBufferPrimitives.AppendTo(WriteView.DirtyIndirectLightingCacheBufferPrimitives);
-		VolumetricPrimSet.AppendTo(WriteView.VolumetricPrimSet);
+
+		WriteView.MeshDecalBatches.Append(MeshDecalBatches);
+		WriteView.VolumetricMeshBatches.Append(VolumetricMeshBatches);
 
 		for (int32 Index = 0; Index < RecachedReflectionCapturePrimitives.NumPrims; ++Index)
 		{
@@ -2604,6 +2614,23 @@ void ComputeDynamicMeshRelevance(EShadingPath ShadingPath, bool bAddLightmapDens
 		View.NumVisibleDynamicMeshElements[EMeshPass::EditorSelection] += NumElements;
 	}
 #endif
+
+	if (ViewRelevance.bHasVolumeMaterialDomain)
+	{
+		View.VolumetricMeshBatches.AddUninitialized(1);
+		FVolumetricMeshBatch& BatchAndProxy = View.VolumetricMeshBatches.Last();
+		BatchAndProxy.Mesh = MeshBatch.Mesh;
+		BatchAndProxy.Proxy = MeshBatch.PrimitiveSceneProxy;
+	}
+
+	if (ViewRelevance.bRenderInMainPass && ViewRelevance.bDecal)
+	{
+		View.MeshDecalBatches.AddUninitialized(1);
+		FMeshDecalBatch& BatchAndProxy = View.MeshDecalBatches.Last();
+		BatchAndProxy.Mesh = MeshBatch.Mesh;
+		BatchAndProxy.Proxy = MeshBatch.PrimitiveSceneProxy;
+		BatchAndProxy.SortKey = MeshBatch.PrimitiveSceneProxy->GetTranslucencySortPriority();
+	}
 }
 
 void FSceneRenderer::GatherDynamicMeshElements(
@@ -2680,12 +2707,6 @@ void FSceneRenderer::GatherDynamicMeshElements(
 					}
 				}
 			}
-
-			// to support GetDynamicMeshElementRange()
-			for (int32 ViewIndex = 0; ViewIndex < ViewCount; ViewIndex++)
-			{
-				InViews[ViewIndex].DynamicMeshEndIndices[PrimitiveIndex] = Collector.GetMeshBatchCount(ViewIndex);
-			}
 		}
 	}
 
@@ -3319,7 +3340,6 @@ void FSceneRenderer::ComputeViewVisibility(FRHICommandListImmediate& RHICmdList,
 		// Allocate the view's visibility maps.
 		View.PrimitiveVisibilityMap.Init(false,Scene->Primitives.Num());
 		// we don't initialized as we overwrite the whole array (in GatherDynamicMeshElements)
-		View.DynamicMeshEndIndices.SetNumUninitialized(Scene->Primitives.Num());
 		View.PrimitiveDefinitelyUnoccludedMap.Init(false,Scene->Primitives.Num());
 		View.PotentiallyFadingPrimitiveMap.Init(false,Scene->Primitives.Num());
 		View.PrimitiveFadeUniformBuffers.AddZeroed(Scene->Primitives.Num());
@@ -3635,7 +3655,7 @@ void FSceneRenderer::PostVisibilityFrameSetup(FILCUpdatePrimTaskData& OutILCTask
 		{		
 			FViewInfo& View = Views[ViewIndex];
 
-			View.MeshDecalPrimSet.SortPrimitives();
+			View.MeshDecalBatches.Sort();
 
 			if (View.State)
 			{
diff --git a/Engine/Source/Runtime/Renderer/Private/VolumetricFogVoxelization.cpp b/Engine/Source/Runtime/Renderer/Private/VolumetricFogVoxelization.cpp
--- a/Engine/Source/Runtime/Renderer/Private/VolumetricFogVoxelization.cpp
+++ b/Engine/Source/Runtime/Renderer/Private/VolumetricFogVoxelization.cpp
@@ -586,8 +586,7 @@ void FDeferredShadingSceneRenderer::VoxelizeFogVolumePrimitives(
 	FVector GridZParams,
 	float VolumetricFogDistance)
 {
-	if (View.VolumetricPrimSet.NumPrims() > 0
-		&& DoesPlatformSupportVolumetricFogVoxelization(View.GetShaderPlatform()))
+	if (View.VolumetricMeshBatches.Num() > 0 && DoesPlatformSupportVolumetricFogVoxelization(View.GetShaderPlatform()))
 	{
 		FRenderTargetParameters* PassParameters = GraphBuilder.AllocParameters<FRenderTargetParameters>();
 		PassParameters->RenderTargets[0] = FRenderTargetBinding(IntegrationData.VBufferA, ERenderTargetLoadAction::ENoAction, ERenderTargetStoreAction::EStore);
@@ -628,44 +627,16 @@ void FDeferredShadingSceneRenderer::VoxelizeFogVolumePrimitives(
 					&View,
 					DynamicMeshPassContext);
 
-				for (int32 PrimIdx = 0; PrimIdx < View.VolumetricPrimSet.NumPrims(); PrimIdx++)
+				for (int32 MeshBatchIndex = 0; MeshBatchIndex < View.VolumetricMeshBatches.Num(); ++MeshBatchIndex)
 				{
-					const FPrimitiveSceneProxy* PrimitiveSceneProxy = View.VolumetricPrimSet.GetPrim(PrimIdx);
+					const FMeshBatch* Mesh = View.VolumetricMeshBatches[MeshBatchIndex].Mesh;
+					const FPrimitiveSceneProxy* PrimitiveSceneProxy = View.VolumetricMeshBatches[MeshBatchIndex].Proxy;
 					const FPrimitiveSceneInfo* PrimitiveSceneInfo = PrimitiveSceneProxy->GetPrimitiveSceneInfo();
+					const FBoxSphereBounds Bounds = PrimitiveSceneProxy->GetBounds();
 
-					if (View.PrimitiveVisibilityMap[PrimitiveSceneInfo->GetIndex()])
+					if ((View.ViewMatrices.GetViewOrigin() - Bounds.Origin).SizeSquared() < (VolumetricFogDistance + Bounds.SphereRadius) * (VolumetricFogDistance + Bounds.SphereRadius))
 					{
-						const FPrimitiveViewRelevance& ViewRelevance = View.PrimitiveViewRelevanceMap[PrimitiveSceneInfo->GetIndex()];
-						FBoxSphereBounds Bounds = PrimitiveSceneProxy->GetBounds();
-
-						if ((View.ViewMatrices.GetViewOrigin() - Bounds.Origin).SizeSquared() < (VolumetricFogDistance + Bounds.SphereRadius) * (VolumetricFogDistance + Bounds.SphereRadius))
-						{
-							// Range in View.DynamicMeshElements[] corresponding to this PrimitiveSceneInfo
-							FInt32Range Range = View.GetDynamicMeshElementRange(PrimitiveSceneInfo->GetIndex());
-
-							for (int32 MeshBatchIndex = Range.GetLowerBoundValue(); MeshBatchIndex < Range.GetUpperBoundValue(); MeshBatchIndex++)
-							{
-								const FMeshBatchAndRelevance& MeshBatchAndRelevance = View.DynamicMeshElements[MeshBatchIndex];
-
-								checkSlow(MeshBatchAndRelevance.PrimitiveSceneProxy == PrimitiveSceneInfo->Proxy);
-
-								const FMeshBatch& MeshBatch = *MeshBatchAndRelevance.Mesh;
-								VoxelizeVolumePrimitive(PassMeshProcessor, RHICmdList, View, VolumetricFogGridSize, GridZParams, PrimitiveSceneProxy, MeshBatch);
-							}
-						}
-
-						if (ViewRelevance.bStaticRelevance)
-						{
-							for (int32 StaticMeshIdx = 0; StaticMeshIdx < PrimitiveSceneInfo->StaticMeshes.Num(); StaticMeshIdx++)
-							{
-								const FStaticMeshBatch& StaticMesh = PrimitiveSceneInfo->StaticMeshes[StaticMeshIdx];
-
-								if (View.StaticMeshVisibilityMap[StaticMesh.Id])
-								{
-									VoxelizeVolumePrimitive(PassMeshProcessor, RHICmdList, View, VolumetricFogGridSize, GridZParams, PrimitiveSceneProxy, StaticMesh);
-								}
-							}
-						}
+						VoxelizeVolumePrimitive(PassMeshProcessor, RHICmdList, View, VolumetricFogGridSize, GridZParams, PrimitiveSceneProxy, *Mesh);
 					}
 				}
 			});
