Post Process Histogram Optimizations provided by [REDACTED_PERSON].

diff --git a/Engine/Shaders/Private/PostProcessHistogram.usf b/Engine/Shaders/Private/PostProcessHistogram.usf
--- a/Engine/Shaders/Private/PostProcessHistogram.usf
+++ b/Engine/Shaders/Private/PostProcessHistogram.usf
@@ -12,45 +12,55 @@
 float4 HistogramParameters;
 
 // Output histogram texture (UAV)
-RWTexture2D<float4> HistogramRWTexture;
-
+RWTexture2D<uint> RWHistogramTexture;
 
 //
 uint2 ThreadGroupCount;
 //
 uint2 LeftTopOffset;
 
-// THREADGROUP_SIZEX*THREADGROUP_SIZEY histograms of the size HISTOGRAM_SIZE
-groupshared float SharedHistogram[HISTOGRAM_SIZE][THREADGROUP_SIZEX][THREADGROUP_SIZEY];
+// One shared histogram for all threads in group
+static const uint HistogramPadding = 4;	// AMD/Nvidia = 32 banks, Intel = 16 banks.
+static const uint HistogramOffset = HISTOGRAM_SIZE + HistogramPadding;
+groupshared uint SharedHistograms[HISTOGRAM_SIZE * 4 + HistogramPadding * 3];
+
+
+
+[numthreads(64, 1, 1)]	
+void ClearHistogramCS(uint GroupIndex: SV_GroupIndex)
+{
+	if(GroupIndex < HISTOGRAM_SIZE)
+	{
+		RWHistogramTexture[uint2(GroupIndex, 0)] = 0;
+	}
+}
 
 [numthreads(THREADGROUP_SIZEX, THREADGROUP_SIZEY, 1)]	// dimx,dimy,dimz
 void MainCS(
 	uint3 GroupId : SV_GroupID,
 	uint3 DispatchThreadId : SV_DispatchThreadID,	// DispatchThreadId = GroupId * int2(dimx,dimy) + GroupThreadId
-    uint3 GroupThreadId : SV_GroupThreadID,	// 0..THREADGROUP_SIZEX 0..THREADGROUP_SIZEY
-	uint GroupIndex: SV_GroupIndex) // SV_GroupIndex = SV_GroupThreadID.z*dimx*dimy + SV_GroupThreadID.y*dimx + SV_GroupThreadID.x
+    	uint3 GroupThreadId : SV_GroupThreadID,			// 0..THREADGROUP_SIZEX 0..THREADGROUP_SIZEY
+	uint GroupIndex: SV_GroupIndex)					// SV_GroupIndex = SV_GroupThreadID.z*dimx*dimy + SV_GroupThreadID.y*dimx + SV_GroupThreadID.x
 {
-	// todo: can be cleared more efficiently
-	// clear all THREADGROUP_SIZEX*THREADGROUP_SIZEY histograms
-	UNROLL for(uint i = 0; i < HISTOGRAM_SIZE; ++i)
+	// Clear the shared histogram in one go. Use all 256 threads to clear 4x 64 bin histograms.
+	// IMPORTANT: Group size must be 16x16 = 256
 	{
-		SharedHistogram[i][GroupThreadId.x][GroupThreadId.y] = 0.0f;
+		uint Index = (GroupIndex % 64) + HistogramOffset * (GroupIndex / 64);
+		SharedHistograms[Index] = 0;
 	}
 	 
-	GroupMemoryBarrierWithGroupSync();
+    GroupMemoryBarrierWithGroupSync();
 
 	uint2 TileSize = uint2(LOOP_SIZEX, LOOP_SIZEY);
-	uint2 LeftTop = DispatchThreadId.xy * TileSize + LeftTopOffset;
-
+	uint2 LeftTop = GroupId.xy * uint2(THREADGROUP_SIZEX * LOOP_SIZEX, THREADGROUP_SIZEY * LOOP_SIZEY) + GroupThreadId.xy + LeftTopOffset;
 	uint2 GatherExtentInt = (uint2)HistogramParameters.xy + (uint2)LeftTopOffset.xy;
 
-	// accumulate all pixels into THREADGROUP_SIZEX*THREADGROUP_SIZEY histograms
 	{
-		LOOP for(uint y = 0; y < LOOP_SIZEY; ++y)
+		UNROLL for(uint y = 0; y < LOOP_SIZEY; ++y)
 		{
-			LOOP for(uint x = 0; x < LOOP_SIZEX; ++x)
+			UNROLL for(uint x = 0; x < LOOP_SIZEX; ++x)
 			{
-				uint2 Tile = uint2(x, y);
+				uint2 Tile = uint2(x, y) * uint2(THREADGROUP_SIZEX, THREADGROUP_SIZEY);
 				uint2 TexelPos = LeftTop + Tile;
 
 				if(TexelPos.x < GatherExtentInt.x && TexelPos.y < GatherExtentInt.y)
@@ -71,42 +81,33 @@ void MainCS(
 					float Weight0 = 1.0f - Weight1;
 			
 					// accumulate the weight to the nearby history buckets
-					if (Bucket0 != 0)
-					{
-						SharedHistogram[Bucket0][GroupThreadId.x][GroupThreadId.y] += Weight0;
-					}
-					SharedHistogram[Bucket1][GroupThreadId.x][GroupThreadId.y] += Weight1;
+					// NOTE: Pack "overflow" to next bucket to same counter high bits. This halves the number of atomics. 
+					uint PackedHistogram = (uint(Weight1 * 64.0) << 16) | uint(Weight0 * 64.0);
+                   			uint Offset = HistogramOffset * (GroupIndex % 4);
+					InterlockedAdd(SharedHistograms[Bucket0 + Offset], PackedHistogram);
 				}
 			}
 		}
 	}
 
-	GroupMemoryBarrierWithGroupSync();
+    GroupMemoryBarrierWithGroupSync();
 
-	// accumulate all histograms into one
-	if(GroupIndex < HISTOGRAM_SIZE / 4) 
+	[branch]
+	if(GroupIndex < HISTOGRAM_SIZE)
 	{
-		float4 Sum = 0;
+		uint Sum = 0;
 
-		LOOP for(uint y = 0; y < THREADGROUP_SIZEY; ++y)
+		[flatten]
+		if (GroupIndex > 0)	// UE comment: good hack to prevent adapting to extreme darks
 		{
-			LOOP for(uint x = 0; x < THREADGROUP_SIZEX; ++x)
-			{
-				Sum += float4(
-					SharedHistogram[GroupIndex * 4 + 0][x][y],
-					SharedHistogram[GroupIndex * 4 + 1][x][y],
-					SharedHistogram[GroupIndex * 4 + 2][x][y],
-					SharedHistogram[GroupIndex * 4 + 3][x][y]);
-			}
-		}
+			// Sum low+high bit histograms for each 4 histograms.
+			Sum = SharedHistograms[GroupIndex     + HistogramOffset * 0] + SharedHistograms[GroupIndex     + HistogramOffset * 1] + SharedHistograms[GroupIndex     + HistogramOffset * 2] + SharedHistograms[GroupIndex     + HistogramOffset * 3] + 
+				((SharedHistograms[GroupIndex - 1 + HistogramOffset * 0] + SharedHistograms[GroupIndex - 1 + HistogramOffset * 1] + SharedHistograms[GroupIndex - 1 + HistogramOffset * 2] + SharedHistograms[GroupIndex - 1 + HistogramOffset * 3]) >> 16);
 
-		float2 MaxExtent = float2(THREADGROUP_SIZEX * LOOP_SIZEX, THREADGROUP_SIZEY * LOOP_SIZEY);
-		float Area = MaxExtent.x * MaxExtent.y;
-	
-		// doesn't take borders into account but the error should be minor
-		float NormalizeFactor = 1.0f / Area;
+			Sum = Sum & 0x0000ffff;
+		}
 
-		// output texture with one histogram per line, x and y unwrapped into all the lines
-		HistogramRWTexture[uint2(GroupIndex, GroupId.x + GroupId.y * ThreadGroupCount.x)] = Sum * NormalizeFactor;
+		// Accumulate groups to shared histogram
+		InterlockedAdd(RWHistogramTexture[uint2(GroupIndex, 0)], Sum);
 	}
 }
diff --git a/Engine/Shaders/Private/PostProcessHistogramReduce.usf b/Engine/Shaders/Private/PostProcessHistogramReduce.usf
--- a/Engine/Shaders/Private/PostProcessHistogramReduce.usf
+++ b/Engine/Shaders/Private/PostProcessHistogramReduce.usf
@@ -14,19 +14,26 @@ float4 EyeAdapationTemporalParams;
 // the last frames eye adaptation settings
 Texture2D EyeAdaptationTexture;
 
+Texture2D<int> HistogramTexture;
+
 void MainPS(noperspective float4 UVAndScreenPos : TEXCOORD0, float4 SvPosition: SV_POSITION, out float4 OutColor : SV_Target0)
 {
 	float2 UV = UVAndScreenPos.xy;
 
 	float4 SceneColor = 0;
 
-	// accumulate all histograms into a single one
-	// could be optimized using bilinear filtering (half sample count)
-	for(uint y = 0; y < LoopSize; ++y)
-	{
-		SceneColor += Texture2DSample(PostprocessInput0, PostprocessInput0Sampler, UV + float2(0, y * PostprocessInput0Size.w));
-	}
-
+	// NOTE: Could do single Buffer.Load4, but UE doesn't support buffers as post process outputs
+	int PosX = int(SvPosition.x);
+	uint4 HistogramUint = uint4(
+		HistogramTexture.Load(int3(PosX * 4 + 0, 0, 0)).x,
+		HistogramTexture.Load(int3(PosX * 4 + 1, 0, 0)).x,
+		HistogramTexture.Load(int3(PosX * 4 + 2, 0, 0)).x,
+		HistogramTexture.Load(int3(PosX * 4 + 3, 0, 0)).x);
+
+	//SceneColor = float4(HistogramUint) * ((1.0 / 64.0) * (1.0 / (8.0 * 8.0)));				// Loopsize = 1x1 (Tile = 16x16)
+	SceneColor = float4(HistogramUint) * ((1.0 / 64.0) * (1.0 / (8.0 * 8.0 * 2.0 * 2.0)));		// Loopsize = 2x2 (Tile = 16x16)
+	
+	[flatten]
 	if(SvPosition.y < 1.0f)
 	{
 		// line 0: histogram
diff --git a/Engine/Source/Runtime/Renderer/Private/PostProcess/PostProcessHistogram.cpp b/Engine/Source/Runtime/Renderer/Private/PostProcess/PostProcessHistogram.cpp
--- a/Engine/Source/Runtime/Renderer/Private/PostProcess/PostProcessHistogram.cpp
+++ b/Engine/Source/Runtime/Renderer/Private/PostProcess/PostProcessHistogram.cpp
@@ -8,6 +8,53 @@
 #include "SceneUtils.h"
 #include "PostProcess/PostProcessEyeAdaptation.h"
 
+
+class FPostProcessClearHistogramCS : public FGlobalShader
+{
+	DECLARE_SHADER_TYPE(FPostProcessClearHistogramCS, Global);
+
+	static bool ShouldCompilePermutation(const FGlobalShaderPermutationParameters& Parameters)
+	{
+		return IsFeatureLevelSupported(Parameters.Platform, ERHIFeatureLevel::SM5);
+	}
+
+	static void ModifyCompilationEnvironment(const FGlobalShaderPermutationParameters& Parameters, FShaderCompilerEnvironment& OutEnvironment)
+	{
+		FGlobalShader::ModifyCompilationEnvironment(Parameters, OutEnvironment);
+		OutEnvironment.SetDefine(TEXT("THREADGROUP_SIZEX"), FRCPassPostProcessHistogram::ThreadGroupSizeX);
+		OutEnvironment.SetDefine(TEXT("THREADGROUP_SIZEY"), FRCPassPostProcessHistogram::ThreadGroupSizeY);
+		OutEnvironment.SetDefine(TEXT("LOOP_SIZEX"), FRCPassPostProcessHistogram::LoopCountX);
+		OutEnvironment.SetDefine(TEXT("LOOP_SIZEY"), FRCPassPostProcessHistogram::LoopCountY);
+		OutEnvironment.SetDefine(TEXT("HISTOGRAM_SIZE"), FRCPassPostProcessHistogram::HistogramSize);
+		OutEnvironment.SetDefine(TEXT("EYE_ADAPTATION_PARAMS_SIZE"), (uint32)EYE_ADAPTATION_PARAMS_SIZE);
+		OutEnvironment.CompilerFlags.Add(CFLAG_StandardOptimization);
+	}
+
+	/** Default constructor. */
+	FPostProcessClearHistogramCS() {}
+
+public:
+	FRWShaderParameter HistogramRWTexture;
+
+
+	/** Initialization constructor. */
+	FPostProcessClearHistogramCS(const ShaderMetaType::CompiledShaderInitializerType& Initializer)
+		: FGlobalShader(Initializer)
+	{
+		HistogramRWTexture.Bind(Initializer.ParameterMap, TEXT("HistogramTexture"));
+	}
+
+	// FShader interface.
+	virtual bool Serialize(FArchive& Ar) override
+	{
+		bool bShaderHasOutdatedParameters = FGlobalShader::Serialize(Ar);
+		Ar << HistogramRWTexture;
+		return bShaderHasOutdatedParameters;
+	}
+};
+IMPLEMENT_SHADER_TYPE(, FPostProcessClearHistogramCS, TEXT("/Engine/Private/PostProcessHistogram.usf"), TEXT("ClearHistogramCS"), SF_Compute);
+
+
 /** Encapsulates the post processing histogram compute shader. */
 class FPostProcessHistogramCS : public FGlobalShader
 {
@@ -35,7 +82,7 @@ class FPostProcessHistogramCS : public FGlobalShader
 
 public:
 	FPostProcessPassParameters PostprocessParameter;
-	FShaderResourceParameter HistogramRWTexture;
+	FRWShaderParameter HistogramRWTexture;
 	FShaderParameter HistogramParameters;
 	FShaderParameter ThreadGroupCount;
 	FShaderParameter LeftTopOffset;
@@ -46,7 +93,7 @@ public:
 		: FGlobalShader(Initializer)
 	{
 		PostprocessParameter.Bind(Initializer.ParameterMap);
-		HistogramRWTexture.Bind(Initializer.ParameterMap, TEXT("HistogramRWTexture"));
+		HistogramRWTexture.Bind(Initializer.ParameterMap, TEXT("HistogramTexture"));
 		HistogramParameters.Bind(Initializer.ParameterMap, TEXT("HistogramParameters"));
 		ThreadGroupCount.Bind(Initializer.ParameterMap, TEXT("ThreadGroupCount"));
 		LeftTopOffset.Bind(Initializer.ParameterMap, TEXT("LeftTopOffset"));
@@ -104,28 +151,33 @@ void FRCPassPostProcessHistogram::Process(FRenderingCompositePassContext& Contex
 	FIntRect DestRect = Context.SceneColorViewRect;
 
 	const FSceneRenderTargetItem& DestRenderTarget = PassOutputs[0].RequestSurface(Context);
-
-	TShaderMapRef<FPostProcessHistogramCS> ComputeShader(Context.GetShaderMap());
-
-	// #todo-renderpasses remove once everything is renderpasses
+	
 	UnbindRenderTargets(Context.RHICmdList);
-    Context.RHICmdList.SetComputeShader(ComputeShader->GetComputeShader());
-    
-
-	// set destination
+	// Clear histogram counters
 	check(DestRenderTarget.UAV);
 	Context.RHICmdList.TransitionResource(EResourceTransitionAccess::ERWBarrier, EResourceTransitionPipeline::EGfxToCompute, DestRenderTarget.UAV);
-	Context.RHICmdList.SetUAVParameter(ComputeShader->GetComputeShader(), ComputeShader->HistogramRWTexture.GetBaseIndex(), DestRenderTarget.UAV);
+	SetRenderTarget(Context.RHICmdList, FTextureRHIRef(), FTextureRHIRef());
+	{
+		TShaderMapRef<FPostProcessClearHistogramCS> ComputeShader(Context.GetShaderMap());
+		const FComputeShaderRHIParamRef ShaderRHI = ComputeShader->GetComputeShader();
+		Context.RHICmdList.SetComputeShader(ShaderRHI);
+		ComputeShader->HistogramRWTexture.SetTexture(Context.RHICmdList, ShaderRHI, nullptr, DestRenderTarget.UAV);
+		DispatchComputeShader(Context.RHICmdList, *ComputeShader, 1, 1, 1);
+	}
+
+	TShaderMapRef<FPostProcessHistogramCS> ComputeShader(Context.GetShaderMap());
+	const FComputeShaderRHIParamRef ShaderRHI = ComputeShader->GetComputeShader();
+	Context.RHICmdList.SetComputeShader(ShaderRHI);
+ 	ComputeShader->HistogramRWTexture.SetTexture(Context.RHICmdList, ShaderRHI, nullptr, DestRenderTarget.UAV);
 
 	FIntPoint GatherExtent = ComputeGatherExtent(Context);
 	FIntPoint ThreadGroupCountValue = ComputeThreadGroupCount(GatherExtent);
 
 	ComputeShader->SetCS(Context.RHICmdList, Context, ThreadGroupCountValue, (DestRect.Min + FIntPoint(1, 1)) / 2, GatherExtent);
-	
 	DispatchComputeShader(Context.RHICmdList, *ComputeShader, ThreadGroupCountValue.X, ThreadGroupCountValue.Y, 1);
 
-	// un-set destination
-	Context.RHICmdList.SetUAVParameter(ComputeShader->GetComputeShader(), ComputeShader->HistogramRWTexture.GetBaseIndex(), NULL);
+	// unset destination
+	ComputeShader->HistogramRWTexture.UnsetUAV(Context.RHICmdList, ComputeShader->GetComputeShader());
 
 	Context.RHICmdList.TransitionResource(EResourceTransitionAccess::EReadable, EResourceTransitionPipeline::EComputeToGfx, DestRenderTarget.UAV);
 	ensureMsgf(DestRenderTarget.TargetableTexture == DestRenderTarget.ShaderResourceTexture, TEXT("%s should be resolved to a separate SRV"), *DestRenderTarget.TargetableTexture->GetName().ToString());	
@@ -158,11 +210,10 @@ FPooledRenderTargetDesc FRCPassPostProcessHistogram::ComputeOutputDesc(EPassOutp
 
 	FIntPoint ThreadGroupCount = ComputeThreadGroupCount(PixelExtent);
 
-	// each ThreadGroup outputs one histogram
-	FIntPoint NewSize = FIntPoint(HistogramTexelCount, ThreadGroupCount.X * ThreadGroupCount.Y);
+	// One shared histogram (groups counts accumulated using integer atomics)
+	FIntPoint NewSize = FIntPoint(HistogramSize, 1);
+	FPooledRenderTargetDesc Ret(FPooledRenderTargetDesc::Create2DDesc(NewSize, PF_R32_UINT, FClearValueBinding::None, TexCreate_None, TexCreate_UAV, false));
 
-	// format can be optimized later
-	FPooledRenderTargetDesc Ret(FPooledRenderTargetDesc::Create2DDesc(NewSize, PF_FloatRGBA, FClearValueBinding::None, TexCreate_None, TexCreate_RenderTargetable | TexCreate_UAV, false));
 	Ret.Flags |= GFastVRamConfig.Histogram;
 	Ret.DebugName = TEXT("Histogram");
 
diff --git a/Engine/Source/Runtime/Renderer/Private/PostProcess/PostProcessHistogram.h b/Engine/Source/Runtime/Renderer/Private/PostProcess/PostProcessHistogram.h
--- a/Engine/Source/Runtime/Renderer/Private/PostProcess/PostProcessHistogram.h
+++ b/Engine/Source/Runtime/Renderer/Private/PostProcess/PostProcessHistogram.h
@@ -21,10 +21,10 @@ public:
 
 	// -------------------------------------------
 
+	// IMPORTANT: Optimized algorithm needs groups thread count == 256. Change shader if you want to change group size.
 	// changing this number require Histogram.usf to be recompiled
-	static const uint32 ThreadGroupSizeX = 8;
-	// changing this number require Histogram.usf to be recompiled
-	static const uint32 ThreadGroupSizeY = 4;
+	static const uint32 ThreadGroupSizeX = 16;
+	static const uint32 ThreadGroupSizeY = 16;
 
 	static const uint32 HistogramSize = 64;
 
@@ -33,10 +33,10 @@ public:
 
 	// one ThreadGroup processes LoopCountX*LoopCountY blocks of size ThreadGroupSizeX*ThreadGroupSizeY
 
-	// changing this number require Histogram.usf to be recompiled
-	static const uint32 LoopCountX = 8;
-	// changing this number require Histogram.usf to be recompiled
-	static const uint32 LoopCountY = 8;
+	// Using smaller loop count (2,2) to increase wave count. Otherwise will only use a few CUs (GPU practically idles).
+	// Higher amount of groups is not a problem since new algorithm doesn't need a separate reduction step (all groups accumulate to same histogram with atomic adds).
+	static const uint32 LoopCountX = 2;
+	static const uint32 LoopCountY = 2;
 
 	// -------------------------------------------
 
diff --git a/Engine/Source/Runtime/Renderer/Private/PostProcess/PostProcessHistogramReduce.cpp b/Engine/Source/Runtime/Renderer/Private/PostProcess/PostProcessHistogramReduce.cpp
--- a/Engine/Source/Runtime/Renderer/Private/PostProcess/PostProcessHistogramReduce.cpp
+++ b/Engine/Source/Runtime/Renderer/Private/PostProcess/PostProcessHistogramReduce.cpp
@@ -34,6 +34,7 @@ public:
 	FPostProcessPassParameters PostprocessParameter;
 	FShaderParameter LoopSize;
 	FShaderResourceParameter EyeAdaptationTexture;
+	FShaderResourceParameter HistogramTexture;
 	FShaderParameter EyeAdapationTemporalParams;
 
 	/** Initialization constructor. */
@@ -44,16 +45,18 @@ public:
 		LoopSize.Bind(Initializer.ParameterMap, TEXT("LoopSize"));
 		EyeAdaptationTexture.Bind(Initializer.ParameterMap, TEXT("EyeAdaptationTexture"));
 		EyeAdapationTemporalParams.Bind(Initializer.ParameterMap, TEXT("EyeAdapationTemporalParams"));
+		HistogramTexture.Bind(Initializer.ParameterMap, TEXT("HistogramTexture"));
 	}
 
 	template <typename TRHICmdList>
-	void SetPS(TRHICmdList& RHICmdList, const FRenderingCompositePassContext& Context, uint32 LoopSizeValue)
+	void SetPS(TRHICmdList& RHICmdList, const FRenderingCompositePassContext& Context, uint32 LoopSizeValue, IPooledRenderTarget* InHistogramTexture)
 	{
 		const FPixelShaderRHIParamRef ShaderRHI = GetPixelShader();
 
 		FGlobalShader::SetParameters<FViewUniformShaderParameters>(RHICmdList, ShaderRHI, Context.View.ViewUniformBuffer);
 
 		PostprocessParameter.SetPS(RHICmdList, ShaderRHI, Context, TStaticSamplerState<SF_Bilinear, AM_Clamp, AM_Clamp, AM_Clamp>::GetRHI());
+		SetTextureParameter(Context.RHICmdList, ShaderRHI, HistogramTexture, InHistogramTexture->GetRenderTargetItem().TargetableTexture);
 
 		SetShaderValue(RHICmdList, ShaderRHI, LoopSize, LoopSizeValue);
 
@@ -80,7 +83,7 @@ public:
 	virtual bool Serialize(FArchive& Ar) override
 	{
 		bool bShaderHasOutdatedParameters = FGlobalShader::Serialize(Ar);
-		Ar << PostprocessParameter << LoopSize << EyeAdaptationTexture << EyeAdapationTemporalParams;
+		Ar << PostprocessParameter << LoopSize << EyeAdaptationTexture << EyeAdapationTemporalParams << HistogramTexture;
 		return bShaderHasOutdatedParameters;
 	}
 };
@@ -131,7 +134,7 @@ void FRCPassPostProcessHistogramReduce::Process(FRenderingCompositePassContext&
 
 		uint32 LoopSizeValue = ComputeLoopSize(GatherExtent);
 
-		PixelShader->SetPS(Context.RHICmdList, Context, LoopSizeValue);
+		PixelShader->SetPS(Context.RHICmdList, Context, LoopSizeValue, GetInput(ePId_Input0)->GetOutput()->PooledRenderTarget);
 
 		DrawRectangle(
 			Context.RHICmdList,
