Fixed crash bug with translucent materials using the SceneDepth node.

diff --git a/Engine/Source/Runtime/Renderer/Private/PostProcess/PostProcessSelectionOutline.cpp b/Engine/Source/Runtime/Renderer/Private/PostProcess/PostProcessSelectionOutline.cpp
--- a/Engine/Source/Runtime/Renderer/Private/PostProcess/PostProcessSelectionOutline.cpp
+++ b/Engine/Source/Runtime/Renderer/Private/PostProcess/PostProcessSelectionOutline.cpp
@@ -97,14 +97,15 @@ void FRCPassPostProcessSelectionOutlineColor::Process(FRenderingCompositePassCon
 				EditorView.ViewUniformBuffer = TUniformBufferRef<FViewUniformShaderParameters>::CreateUniformBufferImmediate(*EditorView.CachedViewUniformShaderParameters, UniformBuffer_SingleFrame);
 			}
 
+			FScene* Scene = View.Family->Scene->GetRenderScene();
+
 			FSceneTexturesUniformParameters SceneTextureParameters;
 			SetupSceneTextureUniformParameters(SceneContext, EditorView.FeatureLevel, ESceneTextureSetupMode::None, SceneTextureParameters);
-			TUniformBufferRef<FSceneTexturesUniformParameters> PassUniformBuffer = TUniformBufferRef<FSceneTexturesUniformParameters>::CreateUniformBufferImmediate(SceneTextureParameters, UniformBuffer_SingleFrame);
+			Scene->UniformBuffers.EditorSelectionPassUniformBuffer.UpdateUniformBufferImmediate(SceneTextureParameters);
 
 			FBox VolumeBounds[TVC_MAX];
 			EditorView.SetupUniformBufferParameters(SceneContext, VolumeBounds, TVC_MAX, *EditorView.CachedViewUniformShaderParameters);
 
-			FScene* Scene = View.Family->Scene->GetRenderScene();
 			Scene->UniformBuffers.ViewUniformBuffer.UpdateUniformBufferImmediate(*EditorView.CachedViewUniformShaderParameters);
 			EditorView.ViewUniformBuffer = Scene->UniformBuffers.ViewUniformBuffer;
 
