Added missing dynamic buffer commit when r.DoInitViewsLightingAfterPrepass 0.

diff --git a/Engine/Source/Runtime/Renderer/Private/DeferredShadingRenderer.cpp b/Engine/Source/Runtime/Renderer/Private/DeferredShadingRenderer.cpp
--- a/Engine/Source/Runtime/Renderer/Private/DeferredShadingRenderer.cpp
+++ b/Engine/Source/Runtime/Renderer/Private/DeferredShadingRenderer.cpp
@@ -1174,6 +1174,13 @@ void FDeferredShadingSceneRenderer::Render(FRHICommandListImmediate& RHICmdList)
 		DynamicIndexBufferForInitViews.Commit();
 		DynamicVertexBufferForInitViews.Commit();
 		DynamicReadBufferForInitViews.Commit();
+
+		if (!bDoInitViewAftersPrepass)
+		{
+			DynamicVertexBufferForInitShadows.Commit();
+			DynamicIndexBufferForInitShadows.Commit();
+			DynamicReadBufferForInitShadows.Commit();
+		}
 	}
 
 	// Only update the GPU particle simulation for the main view
@@ -1205,7 +1212,6 @@ void FDeferredShadingSceneRenderer::Render(FRHICommandListImmediate& RHICmdList)
 
 				{
 					SCOPE_CYCLE_COUNTER(STAT_FDeferredShadingSceneRenderer_FGlobalDynamicVertexBuffer_Commit);
-
 					DynamicVertexBufferForInitShadows.Commit();
 					DynamicIndexBufferForInitShadows.Commit();
 					DynamicReadBufferForInitShadows.Commit();
