Enable cube scene captures for ES3.1 feature level

diff --git a/Engine/Source/Runtime/Renderer/Private/SceneCaptureRendering.cpp b/Engine/Source/Runtime/Renderer/Private/SceneCaptureRendering.cpp
--- a/Engine/Source/Runtime/Renderer/Private/SceneCaptureRendering.cpp
+++ b/Engine/Source/Runtime/Renderer/Private/SceneCaptureRendering.cpp
@@ -847,7 +847,7 @@ void FScene::UpdateSceneCaptureContents(USceneCaptureComponentCube* CaptureCompo
 	{
 		UTextureRenderTargetCube* const TextureTarget = TextureTargets[CaptureIter];
 
-		if (GetFeatureLevel() >= ERHIFeatureLevel::SM4 && TextureTarget)
+		if (GetFeatureLevel() >= ERHIFeatureLevel::ES3_1 && TextureTarget)
 		{
 			const float FOV = 90 * (float)PI / 360.0f;
 			for (int32 faceidx = 0; faceidx < (int32)ECubeFace::CubeFace_MAX; faceidx++)
