Fixed out of bounds GPUScene GPU buffer write. PrimitivesToUpdate may contain a stale out of bounds index, as we don't remove update request on primitive removal from scene.

diff --git a/Engine/Source/Runtime/Renderer/Private/GPUScene.cpp b/Engine/Source/Runtime/Renderer/Private/GPUScene.cpp
--- a/Engine/Source/Runtime/Renderer/Private/GPUScene.cpp
+++ b/Engine/Source/Runtime/Renderer/Private/GPUScene.cpp
@@ -196,6 +196,7 @@ void UpdateGPUScene(FRHICommandList& RHICmdList, FScene& Scene)
 
 			for (int32 Index : Scene.GPUScene.PrimitivesToUpdate)
 			{
+				// PrimitivesToUpdate may contain a stale out of bounds index, as we don't remove update request on primitive removal from scene.
 				if (Index < Scene.PrimitiveSceneProxies.Num())
 				{
 					FPrimitiveSceneProxy* PrimitiveSceneProxy = Scene.PrimitiveSceneProxies[Index];
@@ -204,13 +205,6 @@ void UpdateGPUScene(FRHICommandList& RHICmdList, FScene& Scene)
 					FPrimitiveSceneShaderData PrimitiveSceneData(PrimitiveSceneProxy);
 					PrimitivesUploadBuilder.Add(Index, &PrimitiveSceneData.Data[0]);
 				}
-				else
-				{
-					// Dirty index belongs to a primitive that is now out of bounds
-					// Upload identity data to the dirty index to expose any incorrect shader indexing, even though no primitive should be referencing it
-					FPrimitiveSceneShaderData PrimitiveSceneData;
-					PrimitivesUploadBuilder.Add(Index, &PrimitiveSceneData.Data[0]);
-				}
 			}
 
 			if (bResizedPrimitiveData)
@@ -247,6 +241,7 @@ void UpdateGPUScene(FRHICommandList& RHICmdList, FScene& Scene)
 
 				for (int32 Index : Scene.GPUScene.PrimitivesToUpdate)
 				{
+					// PrimitivesToUpdate may contain a stale out of bounds index, as we don't remove update request on primitive removal from scene.
 					if (Index < Scene.PrimitiveSceneProxies.Num())
 					{
 						FPrimitiveSceneProxy* PrimitiveSceneProxy = Scene.PrimitiveSceneProxies[Index];
