GPUScene update optimization - prevent from update same primitive multiple times.

diff --git a/Engine/Source/Runtime/Renderer/Private/GPUScene.cpp b/Engine/Source/Runtime/Renderer/Private/GPUScene.cpp
--- a/Engine/Source/Runtime/Renderer/Private/GPUScene.cpp
+++ b/Engine/Source/Runtime/Renderer/Private/GPUScene.cpp
@@ -159,6 +159,10 @@ void UpdateGPUScene(FRHICommandList& RHICmdList, FScene& Scene)
 
 		if (GGPUSceneUploadEveryFrame || Scene.GPUScene.bUpdateAllPrimitives)
 		{
+			for (int32 Index : Scene.GPUScene.PrimitivesToUpdate)
+			{
+				Scene.GPUScene.PrimitivesMarkedToUpdate[Index] = false;
+			}
 			Scene.GPUScene.PrimitivesToUpdate.Reset();
 
 			for (int32 i = 0; i < Scene.Primitives.Num(); i++)
@@ -205,6 +209,8 @@ void UpdateGPUScene(FRHICommandList& RHICmdList, FScene& Scene)
 					FPrimitiveSceneShaderData PrimitiveSceneData(PrimitiveSceneProxy);
 					PrimitivesUploadBuilder.Add(Index, &PrimitiveSceneData.Data[0]);
 				}
+
+				Scene.GPUScene.PrimitivesMarkedToUpdate[Index] = false;
 			}
 
 			if (bResizedPrimitiveData)
@@ -348,7 +354,18 @@ void UploadDynamicPrimitiveShaderDataForView(FRHICommandList& RHICmdList, FScene
 void AddPrimitiveToUpdateGPU(FScene& Scene, int32 PrimitiveId)
 {
 	if (UseGPUScene(GMaxRHIShaderPlatform, Scene.GetFeatureLevel()))
-	{
-		Scene.GPUScene.PrimitivesToUpdate.Add(PrimitiveId);
+	{ 
+		if (PrimitiveId + 1 > Scene.GPUScene.PrimitivesMarkedToUpdate.Num())
+		{
+			const int32 NewSize = Align(PrimitiveId + 1, 64);
+			Scene.GPUScene.PrimitivesMarkedToUpdate.Add(0, NewSize - Scene.GPUScene.PrimitivesMarkedToUpdate.Num());
+		}
+
+		// Make sure we are no updating same primitive multiple times.
+		if (!Scene.GPUScene.PrimitivesMarkedToUpdate[PrimitiveId])
+		{
+			Scene.GPUScene.PrimitivesToUpdate.Add(PrimitiveId);
+			Scene.GPUScene.PrimitivesMarkedToUpdate[PrimitiveId] = true;
+		}
 	}
 }
\ No newline at end of file
diff --git a/Engine/Source/Runtime/Renderer/Private/ScenePrivate.h b/Engine/Source/Runtime/Renderer/Private/ScenePrivate.h
--- a/Engine/Source/Runtime/Renderer/Private/ScenePrivate.h
+++ b/Engine/Source/Runtime/Renderer/Private/ScenePrivate.h
@@ -1711,6 +1711,9 @@ public:
 	/** Indices of primitives that need to be updated in GPU Scene */
 	TArray<int32> PrimitivesToUpdate;
 
+	/** Bit array of all scene primitives. Set bit means that current primitive is in PrimitivesToUpdate array. */
+	TBitArray<> PrimitivesMarkedToUpdate;
+
 	/** GPU mirror of Primitives */
 	FRWBufferStructured PrimitiveBuffer;
 
