Fixed dynamic mesh elements revelance computation.

diff --git a/Engine/Source/Runtime/Renderer/Private/SceneVisibility.cpp b/Engine/Source/Runtime/Renderer/Private/SceneVisibility.cpp
--- a/Engine/Source/Runtime/Renderer/Private/SceneVisibility.cpp
+++ b/Engine/Source/Runtime/Renderer/Private/SceneVisibility.cpp
@@ -2692,21 +2692,22 @@ void FSceneRenderer::GatherDynamicMeshElements(
 				// Compute DynamicMeshElementsMeshPassRelevance for this primitive.
 				for (int32 ViewIndex = 0; ViewIndex < ViewCount; ViewIndex++)
 				{
-					FViewInfo& View = InViews[ViewIndex];
-					const bool bAddLightmapDensityCommands = View.Family->EngineShowFlags.LightMapDensity && AllowDebugViewmodes();
-					const FPrimitiveViewRelevance& ViewRelevance = View.PrimitiveViewRelevanceMap[PrimitiveIndex];
-
-					View.DynamicMeshElementsPassRelevance.SetNum(View.DynamicMeshElements.Num());
+					if (ViewMaskFinal & (1 << ViewIndex))
+					{
+						FViewInfo& View = InViews[ViewIndex];
+						const bool bAddLightmapDensityCommands = View.Family->EngineShowFlags.LightMapDensity && AllowDebugViewmodes();
+						const FPrimitiveViewRelevance& ViewRelevance = View.PrimitiveViewRelevanceMap[PrimitiveIndex];
 
-					const int32 NumPrimitiveElements = Collector.GetMeshBatchCount(ViewIndex);
-					const int32 PrimitiveFirstElement = View.DynamicMeshElements.Num() - NumPrimitiveElements;
+						const int32 LastNumDynamicMeshElements = View.DynamicMeshElementsPassRelevance.Num();
+						View.DynamicMeshElementsPassRelevance.SetNum(View.DynamicMeshElements.Num());
 
-					for (int32 ElementIndex = PrimitiveFirstElement; ElementIndex < NumPrimitiveElements; ++ElementIndex)
-					{
-						const FMeshBatchAndRelevance& MeshBatch = View.DynamicMeshElements[ElementIndex];
-						FMeshPassMask& PassRelevance = View.DynamicMeshElementsPassRelevance[ElementIndex];
+						for (int32 ElementIndex = LastNumDynamicMeshElements; ElementIndex < View.DynamicMeshElements.Num(); ++ElementIndex)
+						{
+							const FMeshBatchAndRelevance& MeshBatch = View.DynamicMeshElements[ElementIndex];
+							FMeshPassMask& PassRelevance = View.DynamicMeshElementsPassRelevance[ElementIndex];
 
-						ComputeDynamicMeshRelevance(ShadingPath, bAddLightmapDensityCommands, ViewRelevance, MeshBatch, View, PassRelevance);
+							ComputeDynamicMeshRelevance(ShadingPath, bAddLightmapDensityCommands, ViewRelevance, MeshBatch, View, PassRelevance);
+						}
 					}
 				}
 			}
