Unshelved from pending changelist '4932645':

Merge conflict.
rolando.caloca, please merge this change by hand.
More info at https://robomerge.epicgames.net#RELEASE

/src/ROBOMERGE_RELEASE_Main/Engine/Plugins/FX/Niagara/Source/Niagara/Private/NiagaraEmitterInstanceBatcher.cpp - merging //UE4/Release-4.22/Engine/Plugins/FX/Niagara/Source/Niagara/Private/NiagaraEmitterInstanceBatcher.cpp#2

--------------------------------------
UE4.22 - Deprecate ENQUEUE_UNIQUE_RENDER_COMMAND_THREEPARAMETER

diff --git a/Engine/Plugins/2D/Paper2D/Source/Paper2D/Private/PaperFlipbookComponent.cpp b/Engine/Plugins/2D/Paper2D/Source/Paper2D/Private/PaperFlipbookComponent.cpp
--- a/Engine/Plugins/2D/Paper2D/Source/Paper2D/Private/PaperFlipbookComponent.cpp
+++ b/Engine/Plugins/2D/Paper2D/Source/Paper2D/Private/PaperFlipbookComponent.cpp
@@ -357,13 +357,12 @@ void UPaperFlipbookComponent::SendRenderDynamicData_Concurrent()
 		DrawCall.BuildFromSprite(SpriteToSend);
 		DrawCall.Color = SpriteColor.ToFColor(/*bSRGB=*/ false);
 
-		ENQUEUE_UNIQUE_RENDER_COMMAND_THREEPARAMETER(
-				FSendPaperRenderComponentDynamicData,
-				FPaperRenderSceneProxy*,InSceneProxy,(FPaperRenderSceneProxy*)SceneProxy,
-				FSpriteDrawCallRecord,InSpriteToSend,DrawCall,
-				UBodySetup*,InBodySetup,CachedBodySetup,
+		FPaperRenderSceneProxy* InSceneProxy = (FPaperRenderSceneProxy*)SceneProxy;
+		UBodySetup* InBodySetup = CachedBodySetup;
+		ENQUEUE_RENDER_COMMAND(FSendPaperRenderComponentDynamicData)(
+			[InSceneProxy, DrawCall, InBodySetup](FRHICommandListImmediate& RHICmdList)
 			{
-				InSceneProxy->SetDrawCall_RenderThread(InSpriteToSend);
+				InSceneProxy->SetDrawCall_RenderThread(DrawCall);
 				InSceneProxy->SetBodySetup_RenderThread(InBodySetup);
 			});
 	}
diff --git a/Engine/Plugins/2D/Paper2D/Source/Paper2D/Private/PaperSpriteComponent.cpp b/Engine/Plugins/2D/Paper2D/Source/Paper2D/Private/PaperSpriteComponent.cpp
--- a/Engine/Plugins/2D/Paper2D/Source/Paper2D/Private/PaperSpriteComponent.cpp
+++ b/Engine/Plugins/2D/Paper2D/Source/Paper2D/Private/PaperSpriteComponent.cpp
@@ -119,13 +119,11 @@ void UPaperSpriteComponent::SendRenderDynamicData_Concurrent()
 		DrawCall.Color = SpriteColor.ToFColor(/*bSRGB=*/ false);
 		const int32 SplitIndex = (SourceSprite != nullptr) ? SourceSprite->AlternateMaterialSplitIndex : INDEX_NONE;
 
-		ENQUEUE_UNIQUE_RENDER_COMMAND_THREEPARAMETER(
-				FSendPaperSpriteComponentDynamicData,
-				FPaperSpriteSceneProxy*,InSceneProxy,(FPaperSpriteSceneProxy*)SceneProxy,
-				FSpriteDrawCallRecord,InSpriteToSend,DrawCall,
-				int32,InSplitIndex,SplitIndex,
+		FPaperSpriteSceneProxy* InSceneProxy = (FPaperSpriteSceneProxy*)SceneProxy;
+		ENQUEUE_RENDER_COMMAND(FSendPaperSpriteComponentDynamicData)(
+			[InSceneProxy, DrawCall, SplitIndex](FRHICommandListImmediate& RHICmdList)
 			{
-				InSceneProxy->SetSprite_RenderThread(InSpriteToSend, InSplitIndex);
+				InSceneProxy->SetSprite_RenderThread(DrawCall, SplitIndex);
 			});
 	}
 }
@@ -378,13 +376,11 @@ void UPaperSpriteComponent::CheckForErrors()
 #if WITH_EDITOR
 void UPaperSpriteComponent::SetTransientTextureOverride(const UTexture* TextureToModifyOverrideFor, UTexture* OverrideTexture)
 {
-	ENQUEUE_UNIQUE_RENDER_COMMAND_THREEPARAMETER(
-		FSendPaperSpriteComponentDynamicData,
-		FPaperSpriteSceneProxy*, InSceneProxy, (FPaperSpriteSceneProxy*)SceneProxy,
-		const UTexture*, InTextureToModifyOverrideFor, TextureToModifyOverrideFor,
-		UTexture*, InOverrideTexture, OverrideTexture,
+	FPaperSpriteSceneProxy* InSceneProxy = (FPaperSpriteSceneProxy*)SceneProxy;
+	ENQUEUE_RENDER_COMMAND(FSendPaperSpriteComponentDynamicData)(
+		[InSceneProxy, TextureToModifyOverrideFor, OverrideTexture](FRHICommandListImmediate& RHICmdList)
 		{
-			InSceneProxy->SetTransientTextureOverride_RenderThread(InTextureToModifyOverrideFor, InOverrideTexture);
+			InSceneProxy->SetTransientTextureOverride_RenderThread(TextureToModifyOverrideFor, OverrideTexture);
 		});
 }
 #endif
diff --git a/Engine/Plugins/Developer/RenderDocPlugin/Source/RenderDocPlugin/Private/RenderDocPluginModule.cpp b/Engine/Plugins/Developer/RenderDocPlugin/Source/RenderDocPlugin/Private/RenderDocPluginModule.cpp
--- a/Engine/Plugins/Developer/RenderDocPlugin/Source/RenderDocPlugin/Private/RenderDocPluginModule.cpp
+++ b/Engine/Plugins/Developer/RenderDocPlugin/Source/RenderDocPlugin/Private/RenderDocPluginModule.cpp
@@ -216,11 +216,9 @@ void FRenderDocPluginModule::BeginCapture()
 	HWND WindowHandle = GetActiveWindow();
 
 	typedef FRenderDocPluginLoader::RENDERDOC_API_CONTEXT RENDERDOC_API_CONTEXT;
-	ENQUEUE_UNIQUE_RENDER_COMMAND_THREEPARAMETER(
-		StartRenderDocCapture,
-		HWND, WindowHandle, WindowHandle,
-		RENDERDOC_API_CONTEXT*, RenderDocAPI, RenderDocAPI,
-		FRenderDocPluginModule*, Plugin, this,
+	FRenderDocPluginModule* Plugin = this;
+	ENQUEUE_RENDER_COMMAND(StartRenderDocCapture)(
+		[Plugin, WindowHandle, RenderDocAPI](FRHICommandListImmediate& RHICmdList)
 		{
 			FRenderDocFrameCapturer::BeginCapture(WindowHandle, RenderDocAPI, Plugin);
 		});
@@ -252,11 +250,9 @@ void FRenderDocPluginModule::EndCapture()
 	HWND WindowHandle = GetActiveWindow();
 
 	typedef FRenderDocPluginLoader::RENDERDOC_API_CONTEXT RENDERDOC_API_CONTEXT;
-	ENQUEUE_UNIQUE_RENDER_COMMAND_THREEPARAMETER(
-		EndRenderDocCapture,
-		HWND, WindowHandle, WindowHandle,
-		RENDERDOC_API_CONTEXT*, RenderDocAPI, RenderDocAPI,
-		FRenderDocPluginModule*, Plugin, this,
+	FRenderDocPluginModule* Plugin = this;
+	ENQUEUE_RENDER_COMMAND(EndRenderDocCapture)(
+		[WindowHandle, RenderDocAPI, Plugin](FRHICommandListImmediate& RHICmdList)
 		{
 			FRenderDocFrameCapturer::EndCapture(WindowHandle, RenderDocAPI, Plugin);
 		});
diff --git a/Engine/Plugins/Experimental/PixelStreaming/Source/PixelStreaming/Private/NvVideoEncoder.cpp b/Engine/Plugins/Experimental/PixelStreaming/Source/PixelStreaming/Private/NvVideoEncoder.cpp
--- a/Engine/Plugins/Experimental/PixelStreaming/Source/PixelStreaming/Private/NvVideoEncoder.cpp
+++ b/Engine/Plugins/Experimental/PixelStreaming/Source/PixelStreaming/Private/NvVideoEncoder.cpp
@@ -495,18 +495,17 @@ void FNvVideoEncoder::FNvVideoEncoderImpl::EncoderCheckLoop()
 		// When resolution changes, render thread is stopped and later restarted from game thread.
 		// We can't enqueue render commands when render thread is stopped, so pause until render thread is restarted.
 		while (bWaitForRenderThreadToResume) {}
-		ENQUEUE_UNIQUE_RENDER_COMMAND_THREEPARAMETER(
-			NvEncProcessFrame,
-			FNvVideoEncoderImpl*, this_, this,
-			FFrame*, Frame, &Frame,
-			int32, CurrImplCounter, CurrImplCounter,
+		FNvVideoEncoderImpl* This = this;
+		FFrame* InFrame = &Frame;
+		ENQUEUE_RENDER_COMMAND(NvEncProcessFrame)(
+			[This, InFrame, CurrImplCounter](FRHICommandListImmediate& RHICmdList)
 			{
 				if (CurrImplCounter != ImplCounter.GetValue()) // Check if the "this" we captured is still valid
 				{
 					return;
 				}
 	
-				this_->ProcessFrame(*Frame);
+				This->ProcessFrame(*InFrame);
 			}
 		);
 
diff --git a/Engine/Plugins/Runtime/AR/Google/GoogleARCore/Source/GoogleARCoreBase/Private/GoogleARCoreCameraImageBlitter.cpp b/Engine/Plugins/Runtime/AR/Google/GoogleARCore/Source/GoogleARCoreBase/Private/GoogleARCoreCameraImageBlitter.cpp
--- a/Engine/Plugins/Runtime/AR/Google/GoogleARCore/Source/GoogleARCoreBase/Private/GoogleARCoreCameraImageBlitter.cpp
+++ b/Engine/Plugins/Runtime/AR/Google/GoogleARCore/Source/GoogleARCoreBase/Private/GoogleARCoreCameraImageBlitter.cpp
@@ -76,14 +76,12 @@ void FGoogleARCoreDeviceCameraBlitter::LateInit(FIntPoint ImageSize)
 			// We have to get the OpenGL texture ID on the render
 			// thread, after the native resource has been
 			// initialized.
-			ENQUEUE_UNIQUE_RENDER_COMMAND_THREEPARAMETER(
-				InitCameraBlitter,
-				FTexture2DResource *, resource, ((FTexture2DResource *)resource),
-				uint32 *, TextureId, TextureId,
-				FIntPoint, CameraSize, CameraSize,
+			FTexture2DResource* Resource = ((FTexture2DResource *)resource);
+			ENQUEUE_RENDER_COMMAND(InitCameraBlitter)(
+				[Resource, TextureId, CameraSize](FRHICommandListImmediate& RHICmdList)
 				{
-					resource->InitRHI();
-					*TextureId = *reinterpret_cast<uint32*>(resource->GetTexture2DRHI()->GetNativeResource());
+					Resource->InitRHI();
+					*TextureId = *reinterpret_cast<uint32*>(Resource->GetTexture2DRHI()->GetNativeResource());
 				});
 		}
 	}
diff --git a/Engine/Plugins/Runtime/ProceduralMeshComponent/Source/ProceduralMeshComponent/Private/ProceduralMeshComponent.cpp b/Engine/Plugins/Runtime/ProceduralMeshComponent/Source/ProceduralMeshComponent/Private/ProceduralMeshComponent.cpp
--- a/Engine/Plugins/Runtime/ProceduralMeshComponent/Source/ProceduralMeshComponent/Private/ProceduralMeshComponent.cpp
+++ b/Engine/Plugins/Runtime/ProceduralMeshComponent/Source/ProceduralMeshComponent/Private/ProceduralMeshComponent.cpp
@@ -640,15 +640,12 @@ void UProceduralMeshComponent::SetMeshSectionVisible(int32 SectionIndex, bool bN
 		if (SceneProxy)
 		{
 			// Enqueue command to modify render thread info
-			ENQUEUE_UNIQUE_RENDER_COMMAND_THREEPARAMETER(
-				FProcMeshSectionVisibilityUpdate,
-				FProceduralMeshSceneProxy*, ProcMeshSceneProxy, (FProceduralMeshSceneProxy*)SceneProxy,
-				int32, SectionIndex, SectionIndex,
-				bool, bNewVisibility, bNewVisibility,
+			FProceduralMeshSceneProxy* ProcMeshSceneProxy = (FProceduralMeshSceneProxy*)SceneProxy;
+			ENQUEUE_RENDER_COMMAND(FProcMeshSectionVisibilityUpdate)(
+				[ProcMeshSceneProxy, SectionIndex, bNewVisibility](FRHICommandListImmediate& RHICmdList)
 				{
 					ProcMeshSceneProxy->SetSectionVisibility_RenderThread(SectionIndex, bNewVisibility);
-				}
-			);
+				});
 		}
 	}
 }
diff --git a/Engine/Source/Editor/UnrealEd/Private/AssetThumbnail.cpp b/Engine/Source/Editor/UnrealEd/Private/AssetThumbnail.cpp
--- a/Engine/Source/Editor/UnrealEd/Private/AssetThumbnail.cpp
+++ b/Engine/Source/Editor/UnrealEd/Private/AssetThumbnail.cpp
@@ -1169,13 +1169,12 @@ FSlateTexture2DRHIRef* FAssetThumbnailPool::AccessTexture( const FAssetData& Ass
 			{
 				ThumbnailInfo = FreeThumbnails.Pop();
 
-				ENQUEUE_UNIQUE_RENDER_COMMAND_THREEPARAMETER( SlateUpdateThumbSizeCommand,
-					FSlateTextureRenderTarget2DResource*, ThumbnailRenderTarget, ThumbnailInfo->ThumbnailRenderTarget,
-					uint32, Width, Width,
-					uint32, Height, Height,
-				{
-					ThumbnailRenderTarget->SetSize(Width, Height);
-				});
+				FSlateTextureRenderTarget2DResource* ThumbnailRenderTarget = ThumbnailInfo->ThumbnailRenderTarget;
+				ENQUEUE_RENDER_COMMAND(SlateUpdateThumbSizeCommand)(
+					[ThumbnailRenderTarget, Width, Height](FRHICommandListImmediate& RHICmdList)
+					{
+						ThumbnailRenderTarget->SetSize(Width, Height);
+					});
 			}
 			else
 			{
diff --git a/Engine/Source/Runtime/Engine/Private/Components/SkinnedMeshComponent.cpp b/Engine/Source/Runtime/Engine/Private/Components/SkinnedMeshComponent.cpp
--- a/Engine/Source/Runtime/Engine/Private/Components/SkinnedMeshComponent.cpp
+++ b/Engine/Source/Runtime/Engine/Private/Components/SkinnedMeshComponent.cpp
@@ -2166,13 +2166,11 @@ void USkinnedMeshComponent::ShowMaterialSection(int32 MaterialID, bool bShow, in
 		if ( MeshObject )
 		{
 			// need to send render thread for updated hidden section
-			ENQUEUE_UNIQUE_RENDER_COMMAND_THREEPARAMETER(
-				FUpdateHiddenSectionCommand, 
-				FSkeletalMeshObject*, MeshObject, MeshObject, 
-				TArray<bool>, HiddenMaterials, HiddenMaterials, 
-				int32, LODIndex, LODIndex,
+			FSkeletalMeshObject* InMeshObject = MeshObject;
+			ENQUEUE_RENDER_COMMAND(FUpdateHiddenSectionCommand)(
+				[InMeshObject, HiddenMaterials, LODIndex](FRHICommandListImmediate& RHICmdList)
 			{
-				MeshObject->SetHiddenMaterials(LODIndex,HiddenMaterials);
+				InMeshObject->SetHiddenMaterials(LODIndex, HiddenMaterials);
 			});
 		}
 	}
@@ -2197,13 +2195,11 @@ void USkinnedMeshComponent::ShowAllMaterialSections(int32 LODIndex)
 			if (MeshObject)
 			{
 				// need to send render thread for updated hidden section
-				ENQUEUE_UNIQUE_RENDER_COMMAND_THREEPARAMETER(
-					FUpdateHiddenSectionCommand,
-					FSkeletalMeshObject*, MeshObject, MeshObject,
-					TArray<bool>, HiddenMaterials, HiddenMaterials,
-					int32, LODIndex, LODIndex,
+				FSkeletalMeshObject* InMeshObject = MeshObject;
+				ENQUEUE_RENDER_COMMAND(FUpdateHiddenSectionCommand)(
+					[InMeshObject, HiddenMaterials, LODIndex](FRHICommandListImmediate& RHICmdList)
 					{
-						MeshObject->SetHiddenMaterials(LODIndex,HiddenMaterials);
+						InMeshObject->SetHiddenMaterials(LODIndex, HiddenMaterials);
 					});
 			}
 		}
diff --git a/Engine/Source/Runtime/Engine/Private/Materials/MaterialInstance.cpp b/Engine/Source/Runtime/Engine/Private/Materials/MaterialInstance.cpp
--- a/Engine/Source/Runtime/Engine/Private/Materials/MaterialInstance.cpp
+++ b/Engine/Source/Runtime/Engine/Private/Materials/MaterialInstance.cpp
@@ -323,27 +323,20 @@ void FMaterialInstanceResource::GameThread_SetParent(UMaterialInterface* ParentM
 	}
 }
 
-ENQUEUE_UNIQUE_RENDER_COMMAND_THREEPARAMETER_DECLARE_TEMPLATE(
-	SetMIParameterValue, ParameterType,
-	FMaterialInstanceResource*, Resource, Resource,
-	FMaterialParameterInfo, ParameterInfo, Parameter.ParameterInfo,
-	typename ParameterType::ValueType, Value, ParameterType::GetValue(Parameter),
-	{
-		Resource->RenderThread_UpdateParameter(ParameterInfo, Value);
-	});
-
 /**
 * Updates a parameter on the material instance from the game thread.
 */
 template <typename ParameterType>
 void GameThread_UpdateMIParameter(const UMaterialInstance* Instance, const ParameterType& Parameter)
 {
-	ENQUEUE_UNIQUE_RENDER_COMMAND_THREEPARAMETER_CREATE_TEMPLATE(
-		SetMIParameterValue, ParameterType,
-		FMaterialInstanceResource*, Instance->Resource,
-		FMaterialParameterInfo, Parameter.ParameterInfo,
-		typename ParameterType::ValueType, ParameterType::GetValue(Parameter)
-		);
+	FMaterialInstanceResource* Resource = Instance->Resource;
+	const FMaterialParameterInfo& ParameterInfo = Parameter.ParameterInfo;
+	typename ParameterType::ValueType Value = ParameterType::GetValue(Parameter);
+	ENQUEUE_RENDER_COMMAND(SetMIParameterValue)(
+		[Resource, ParameterInfo, Value](FRHICommandListImmediate& RHICmdList)
+		{
+			Resource->RenderThread_UpdateParameter(ParameterInfo, Value);
+		});
 }
 
 bool UMaterialInstance::UpdateParameters()
diff --git a/Engine/Source/Runtime/Engine/Private/Materials/MaterialInterface.cpp b/Engine/Source/Runtime/Engine/Private/Materials/MaterialInterface.cpp
--- a/Engine/Source/Runtime/Engine/Private/Materials/MaterialInterface.cpp
+++ b/Engine/Source/Runtime/Engine/Private/Materials/MaterialInterface.cpp
@@ -487,22 +487,20 @@ void UMaterialInterface::UpdateMaterialRenderProxy(FMaterialRenderProxy& Proxy)
 			Settings = LocalSubsurfaceProfile->Settings;
 		}
 
-		ENQUEUE_UNIQUE_RENDER_COMMAND_THREEPARAMETER(
-			UpdateMaterialRenderProxySubsurface,
-			const FSubsurfaceProfileStruct, Settings, Settings,
-			USubsurfaceProfile*, LocalSubsurfaceProfile, LocalSubsurfaceProfile,
-			FMaterialRenderProxy&, Proxy, Proxy,
-		{
-			uint32 AllocationId = 0;
-
-			if (LocalSubsurfaceProfile)
+		FMaterialRenderProxy* InProxy = &Proxy;
+		ENQUEUE_RENDER_COMMAND(UpdateMaterialRenderProxySubsurface)(
+			[Settings, LocalSubsurfaceProfile, InProxy](FRHICommandListImmediate& RHICmdList)
 			{
-				AllocationId = GSubsurfaceProfileTextureObject.AddOrUpdateProfile(Settings, LocalSubsurfaceProfile);
+				uint32 AllocationId = 0;
 
-				check(AllocationId >= 0 && AllocationId <= 255);
-			}
-			Proxy.SetSubsurfaceProfileRT(LocalSubsurfaceProfile);
-		});
+				if (LocalSubsurfaceProfile)
+				{
+					AllocationId = GSubsurfaceProfileTextureObject.AddOrUpdateProfile(Settings, LocalSubsurfaceProfile);
+
+					check(AllocationId >= 0 && AllocationId <= 255);
+				}
+				InProxy->SetSubsurfaceProfileRT(LocalSubsurfaceProfile);
+			});
 	}
 }
 
diff --git a/Engine/Source/Runtime/Engine/Private/Particles/FXSystem.cpp b/Engine/Source/Runtime/Engine/Private/Particles/FXSystem.cpp
--- a/Engine/Source/Runtime/Engine/Private/Particles/FXSystem.cpp
+++ b/Engine/Source/Runtime/Engine/Private/Particles/FXSystem.cpp
@@ -259,16 +259,15 @@ void FFXSystem::AddVectorField( UVectorFieldComponent* VectorFieldComponent )
 			Instance->Intensity = VectorFieldComponent->Intensity;
 			Instance->Tightness = VectorFieldComponent->Tightness;
 
-			ENQUEUE_UNIQUE_RENDER_COMMAND_THREEPARAMETER(
-				FAddVectorFieldCommand,
-				FFXSystem*, FXSystem, this,
-				FVectorFieldInstance*, Instance, Instance,
-				FMatrix, ComponentToWorld, VectorFieldComponent->GetComponentTransform().ToMatrixWithScale(),
-			{
-				Instance->UpdateTransforms( ComponentToWorld );
-				Instance->Index = FXSystem->VectorFields.AddUninitialized().Index;
-				FXSystem->VectorFields[ Instance->Index ] = Instance;
-			});
+			FFXSystem* FXSystem = this;
+			FMatrix ComponentToWorld = VectorFieldComponent->GetComponentTransform().ToMatrixWithScale();
+			ENQUEUE_RENDER_COMMAND(FAddVectorFieldCommand)(
+				[FXSystem, Instance, ComponentToWorld](FRHICommandListImmediate& RHICmdList)
+				{
+					Instance->UpdateTransforms( ComponentToWorld );
+					Instance->Index = FXSystem->VectorFields.AddUninitialized().Index;
+					FXSystem->VectorFields[ Instance->Index ] = Instance;
+				});
 		}
 	}
 }
@@ -323,17 +322,15 @@ void FFXSystem::UpdateVectorField( UVectorFieldComponent* VectorFieldComponent )
 			UpdateParams.Intensity = VectorFieldComponent->Intensity;
 			UpdateParams.Tightness = VectorFieldComponent->Tightness;
 
-			ENQUEUE_UNIQUE_RENDER_COMMAND_THREEPARAMETER(
-				FUpdateVectorFieldCommand,
-				FFXSystem*, FXSystem, this,
-				FVectorFieldInstance*, Instance, Instance,
-				FUpdateVectorFieldParams, UpdateParams, UpdateParams,
-			{
-				Instance->WorldBounds = UpdateParams.Bounds;
-				Instance->Intensity = UpdateParams.Intensity;
-				Instance->Tightness = UpdateParams.Tightness;
-				Instance->UpdateTransforms( UpdateParams.ComponentToWorld );
-			});
+			FFXSystem* FXSystem = this;
+			ENQUEUE_RENDER_COMMAND(FUpdateVectorFieldCommand)(
+				[FXSystem, Instance, UpdateParams](FRHICommandListImmediate& RHICmdList)
+				{
+					Instance->WorldBounds = UpdateParams.Bounds;
+					Instance->Intensity = UpdateParams.Intensity;
+					Instance->Tightness = UpdateParams.Tightness;
+					Instance->UpdateTransforms( UpdateParams.ComponentToWorld );
+				});
 		}
 	}
 }
diff --git a/Engine/Source/Runtime/Engine/Private/Particles/ParticleGpuSimulation.cpp b/Engine/Source/Runtime/Engine/Private/Particles/ParticleGpuSimulation.cpp
--- a/Engine/Source/Runtime/Engine/Private/Particles/ParticleGpuSimulation.cpp
+++ b/Engine/Source/Runtime/Engine/Private/Particles/ParticleGpuSimulation.cpp
@@ -2832,11 +2832,10 @@ void FParticleSimulationGPU::InitResources(const TArray<uint32>& Tiles, FGPUSpri
 
 	if (InGPUSpriteResources)
 	{
-		ENQUEUE_UNIQUE_RENDER_COMMAND_THREEPARAMETER(
-			FInitParticleSimulationGPUCommand,
-			FParticleSimulationGPU*, Simulation, this,
-			TArray<uint32>, Tiles, Tiles,
-			TRefCountPtr<FGPUSpriteResources>, InGPUSpriteResources, InGPUSpriteResources, // TRefCountPtr to take reference for lifetime of this render command
+		TRefCountPtr<FGPUSpriteResources> InGPUSpriteResourcesRef = InGPUSpriteResources;
+		FParticleSimulationGPU* Simulation = this;
+		ENQUEUE_RENDER_COMMAND(FInitParticleSimulationGPUCommand)(
+			[Simulation, Tiles, InGPUSpriteResourcesRef](FRHICommandListImmediate& RHICmdList)
 			{
 				// Release vertex buffers.
 				Simulation->VertexBuffer.ReleaseResource();
@@ -2847,8 +2846,8 @@ void FParticleSimulationGPU::InitResources(const TArray<uint32>& Tiles, FGPUSpri
 				Simulation->TileVertexBuffer.Init(Tiles);
 
 				// Store simulation resources for this emitter.
-				Simulation->GPUSpriteResources = InGPUSpriteResources;
-				Simulation->EmitterSimulationResources = &InGPUSpriteResources->EmitterSimulationResources;
+				Simulation->GPUSpriteResources = InGPUSpriteResourcesRef;
+				Simulation->EmitterSimulationResources = &InGPUSpriteResourcesRef->EmitterSimulationResources;
 
 				// If a visualization vertex factory has been created, initialize it.
 				if (Simulation->VectorFieldVisualizationVertexFactory)
diff --git a/Engine/Source/Runtime/Engine/Private/PrimitiveSceneProxy.cpp b/Engine/Source/Runtime/Engine/Private/PrimitiveSceneProxy.cpp
--- a/Engine/Source/Runtime/Engine/Private/PrimitiveSceneProxy.cpp
+++ b/Engine/Source/Runtime/Engine/Private/PrimitiveSceneProxy.cpp
@@ -406,14 +406,12 @@ void FPrimitiveSceneProxy::SetSelection_GameThread(const bool bInParentSelected,
 	check(IsInGameThread());
 
 	// Enqueue a message to the rendering thread containing the interaction to add.
-	ENQUEUE_UNIQUE_RENDER_COMMAND_THREEPARAMETER(
-		SetNewSelection,
-		FPrimitiveSceneProxy*,PrimitiveSceneProxy,this,
-		const bool,bNewParentSelection,bInParentSelected,
-		const bool,bNewIndividuallySelected,bInIndividuallySelected,
-	{
-		PrimitiveSceneProxy->SetSelection_RenderThread(bNewParentSelection,bNewIndividuallySelected);
-	});
+	FPrimitiveSceneProxy* PrimitiveSceneProxy = this;
+	ENQUEUE_RENDER_COMMAND(SetNewSelection)(
+		[PrimitiveSceneProxy, bInParentSelected, bInIndividuallySelected](FRHICommandListImmediate& RHICmdList)
+		{
+			PrimitiveSceneProxy->SetSelection_RenderThread(bInParentSelected, bInIndividuallySelected);
+		});
 }
 
 /**
diff --git a/Engine/Source/Runtime/Engine/Private/Slate/DebugCanvas.cpp b/Engine/Source/Runtime/Engine/Private/Slate/DebugCanvas.cpp
--- a/Engine/Source/Runtime/Engine/Private/Slate/DebugCanvas.cpp
+++ b/Engine/Source/Runtime/Engine/Private/Slate/DebugCanvas.cpp
@@ -123,25 +123,25 @@ void FDebugCanvasDrawer::BeginRenderingCanvas( const FIntRect& CanvasRect )
 	if( CanvasRect.Size().X > 0 && CanvasRect.Size().Y > 0 )
 	{
 		bCanvasRenderedLastFrame = true;
-		ENQUEUE_UNIQUE_RENDER_COMMAND_THREEPARAMETER
-		(
-			BeginRenderingDebugCanvas,
-			FDebugCanvasDrawer*, CanvasDrawer, this, 
-			FCanvasPtr, CanvasToRender, GameThreadCanvas,
-			FIntRect, CanvasRect, CanvasRect,
+		FDebugCanvasDrawer* CanvasDrawer = this;
+		FCanvasPtr CanvasToRender = GameThreadCanvas;
+		ENQUEUE_RENDER_COMMAND(BeginRenderingDebugCanvas)(
+			[CanvasDrawer, CanvasToRender, CanvasRect = CanvasRect](FRHICommandListImmediate& RHICmdList)
 			{
+				FCanvasPtr LocalCanvasToRender = CanvasToRender;
+			
 				// Delete the old rendering thread canvas
-				if( CanvasDrawer->GetRenderThreadCanvas().IsValid() && CanvasToRender.IsValid() )
+				if( CanvasDrawer->GetRenderThreadCanvas().IsValid() && LocalCanvasToRender.IsValid() )
 				{
 					CanvasDrawer->DeleteRenderThreadCanvas();
 				}
 
-				if(!CanvasToRender.IsValid())
+				if (!LocalCanvasToRender.IsValid())
 				{
-					CanvasToRender = CanvasDrawer->GetRenderThreadCanvas();
+					LocalCanvasToRender = CanvasDrawer->GetRenderThreadCanvas();
 				}
 
-				CanvasDrawer->SetRenderThreadCanvas( CanvasRect, CanvasToRender );
+				CanvasDrawer->SetRenderThreadCanvas( CanvasRect, LocalCanvasToRender);
 			}
 		);
 		
diff --git a/Engine/Source/Runtime/Engine/Private/TextureRenderTarget2D.cpp b/Engine/Source/Runtime/Engine/Private/TextureRenderTarget2D.cpp
--- a/Engine/Source/Runtime/Engine/Private/TextureRenderTarget2D.cpp
+++ b/Engine/Source/Runtime/Engine/Private/TextureRenderTarget2D.cpp
@@ -126,14 +126,14 @@ void UTextureRenderTarget2D::ResizeTarget(uint32 InSizeX, uint32 InSizeY)
 
 		if (Resource)
 		{
-			ENQUEUE_UNIQUE_RENDER_COMMAND_THREEPARAMETER(
-				ResizeRenderTarget,
-				FTextureRenderTarget2DResource*, Resource, static_cast<FTextureRenderTarget2DResource*>(Resource),
-				int32, NewSizeX, SizeX,
-				int32, NewSizeY, SizeY,
+			FTextureRenderTarget2DResource* InResource = static_cast<FTextureRenderTarget2DResource*>(Resource);
+			int32 NewSizeX = SizeX;
+			int32 NewSizeY = SizeY;
+			ENQUEUE_RENDER_COMMAND(ResizeRenderTarget)(
+				[InResource, NewSizeX, NewSizeY](FRHICommandListImmediate& RHICmdList)
 				{
-					Resource->Resize(NewSizeX, NewSizeY);
-					Resource->UpdateDeferredResource(RHICmdList, true);
+					InResource->Resize(NewSizeX, NewSizeY);
+					InResource->UpdateDeferredResource(RHICmdList, true);
 				}
 			);
 
@@ -150,12 +150,11 @@ void UTextureRenderTarget2D::UpdateResourceImmediate(bool bClearRenderTarget/*=t
 {
 	if (Resource)
 	{
-		ENQUEUE_UNIQUE_RENDER_COMMAND_TWOPARAMETER(
-			UpdateResourceImmediate,
-			FTextureRenderTarget2DResource*, Resource, static_cast<FTextureRenderTarget2DResource*>(Resource),
-			bool, bClearRenderTarget, bClearRenderTarget,
+		FTextureRenderTarget2DResource* InResource = static_cast<FTextureRenderTarget2DResource*>(Resource);
+		ENQUEUE_RENDER_COMMAND(UpdateResourceImmediate)(
+			[InResource, bClearRenderTarget](FRHICommandListImmediate& RHICmdList)
 			{
-				Resource->UpdateDeferredResource(RHICmdList, bClearRenderTarget);
+				InResource->UpdateDeferredResource(RHICmdList, bClearRenderTarget);
 			}
 		);
 	}
diff --git a/Engine/Source/Runtime/Experimental/GeometryCollectionEngine/Private/GeometryCollection/GeometryCollectionComponent.cpp b/Engine/Source/Runtime/Experimental/GeometryCollectionEngine/Private/GeometryCollection/GeometryCollectionComponent.cpp
--- a/Engine/Source/Runtime/Experimental/GeometryCollectionEngine/Private/GeometryCollection/GeometryCollectionComponent.cpp
+++ b/Engine/Source/Runtime/Experimental/GeometryCollectionEngine/Private/GeometryCollection/GeometryCollectionComponent.cpp
@@ -289,10 +289,9 @@ void UGeometryCollectionComponent::CreateRenderState_Concurrent()
 		InitDynamicData(DynamicData);
 
 		// Enqueue command to send to render thread
-		ENQUEUE_UNIQUE_RENDER_COMMAND_THREEPARAMETER(FSendGeometryCollectionData,
-			FGeometryCollectionSceneProxy*, GeometryCollectionSceneProxy, (FGeometryCollectionSceneProxy*)SceneProxy,
-			FGeometryCollectionConstantData*, ConstantData, ConstantData,
-			FGeometryCollectionDynamicData*, DynamicData, DynamicData,
+		FGeometryCollectionSceneProxy* GeometryCollectionSceneProxy = (FGeometryCollectionSceneProxy*)SceneProxy;
+		ENQUEUE_RENDER_COMMAND(FSendGeometryCollectionData)(
+			[GeometryCollectionSceneProxy, ConstantData, DynamicData](FRHICommandListImmediate& RHICmdList)
 			{
 				GeometryCollectionSceneProxy->SetConstantData_RenderThread(ConstantData);
 				GeometryCollectionSceneProxy->SetDynamicData_RenderThread(DynamicData);
@@ -438,10 +437,9 @@ void UGeometryCollectionComponent::ForceInitRenderData()
 	InitDynamicData(DynamicData);
 
 	// Enqueue command to send to render thread
-	ENQUEUE_UNIQUE_RENDER_COMMAND_THREEPARAMETER(FSendGeometryCollectionData,
-		FGeometryCollectionSceneProxy*, GeometryCollectionSceneProxy, (FGeometryCollectionSceneProxy*)SceneProxy,
-		FGeometryCollectionConstantData*, ConstantData, ConstantData,
-		FGeometryCollectionDynamicData*, DynamicData, DynamicData,
+	FGeometryCollectionSceneProxy* GeometryCollectionSceneProxy = (FGeometryCollectionSceneProxy*)SceneProxy;
+	ENQUEUE_RENDER_COMMAND(FSendGeometryCollectionData)(
+		[GeometryCollectionSceneProxy, ConstantData, DynamicData](FRHICommandListImmediate& RHICmdList)
 		{
 			GeometryCollectionSceneProxy->SetConstantData_RenderThread(ConstantData, true);
 			GeometryCollectionSceneProxy->SetDynamicData_RenderThread(DynamicData);
diff --git a/Engine/Source/Runtime/Foliage/Private/FoliageComponent.cpp b/Engine/Source/Runtime/Foliage/Private/FoliageComponent.cpp
--- a/Engine/Source/Runtime/Foliage/Private/FoliageComponent.cpp
+++ b/Engine/Source/Runtime/Foliage/Private/FoliageComponent.cpp
@@ -41,11 +41,9 @@ public:
 	void UpdateParameters_GameThread(const FVector& NewFoliageImpluseDirection, const FVector4& NewFoliageNormalizedRotationAxisAndAngle)
 	{
 		checkSlow(IsInGameThread());
-		ENQUEUE_UNIQUE_RENDER_COMMAND_THREEPARAMETER(
-			UpdateFoliageParameters,
-			FInteractiveFoliageSceneProxy*,FoliageProxy,this,
-			FVector,NewFoliageImpluseDirection,NewFoliageImpluseDirection,
-			FVector4,NewFoliageNormalizedRotationAxisAndAngle,NewFoliageNormalizedRotationAxisAndAngle,
+		FInteractiveFoliageSceneProxy* FoliageProxy = this;
+		ENQUEUE_RENDER_COMMAND(UpdateFoliageParameters)(
+			[FoliageProxy, NewFoliageImpluseDirection, NewFoliageNormalizedRotationAxisAndAngle](FRHICommandListImmediate& RHICmdList)
 			{
 				FoliageProxy->FoliageImpluseDirection = NewFoliageImpluseDirection;
 				FoliageProxy->FoliageNormalizedRotationAxisAndAngle = NewFoliageNormalizedRotationAxisAndAngle;
diff --git a/Engine/Source/Runtime/Landscape/Private/LandscapeComponent.cpp b/Engine/Source/Runtime/Landscape/Private/LandscapeComponent.cpp
--- a/Engine/Source/Runtime/Landscape/Private/LandscapeComponent.cpp
+++ b/Engine/Source/Runtime/Landscape/Private/LandscapeComponent.cpp
@@ -48,11 +48,9 @@ void ULandscapeComponent::UpdateEditToolRenderData()
 		const bool bGetDebugMaterials = true;
 		GetUsedMaterials(UsedMaterialsForVerification, bGetDebugMaterials);
 
-		ENQUEUE_UNIQUE_RENDER_COMMAND_THREEPARAMETER(
-			UpdateEditToolRenderData,
-			FLandscapeEditToolRenderData, LandscapeEditToolRenderData, EditToolRenderData,
-			FLandscapeComponentSceneProxy*, LandscapeSceneProxy, LandscapeSceneProxy,
-			TArray<UMaterialInterface*>, UsedMaterialsForVerification, UsedMaterialsForVerification,
+		FLandscapeEditToolRenderData LandscapeEditToolRenderData = EditToolRenderData;
+		ENQUEUE_RENDER_COMMAND(UpdateEditToolRenderData)(
+			[LandscapeEditToolRenderData, LandscapeSceneProxy, UsedMaterialsForVerification](FRHICommandListImmediate& RHICmdList)
 			{
 				LandscapeSceneProxy->EditToolRenderData = LandscapeEditToolRenderData;				
 				LandscapeSceneProxy->SetUsedMaterialForVerification(UsedMaterialsForVerification);
@@ -60,4 +58,4 @@ void ULandscapeComponent::UpdateEditToolRenderData()
 	}
 }
 
-#endif
\ No newline at end of file
+#endif
diff --git a/Engine/Source/Runtime/Landscape/Private/LandscapeRender.cpp b/Engine/Source/Runtime/Landscape/Private/LandscapeRender.cpp
--- a/Engine/Source/Runtime/Landscape/Private/LandscapeRender.cpp
+++ b/Engine/Source/Runtime/Landscape/Private/LandscapeRender.cpp
@@ -4267,21 +4267,19 @@ void ALandscapeProxy::ChangeTessellationComponentScreenSize(float InTessellation
 			RenderProxies[Idx] = (FLandscapeComponentSceneProxy*)(LandscapeComponents[Idx]->SceneProxy);
 		}
 
-		ENQUEUE_UNIQUE_RENDER_COMMAND_THREEPARAMETER(
-			LandscapeChangeTessellationComponentScreenSizeCommand,
-			FLandscapeComponentSceneProxy**, RenderProxies, RenderProxies,
-			int32, ComponentCount, ComponentCount,
-			float, InTessellationComponentScreenSize, TessellationComponentScreenSize,
+		float TessellationComponentScreenSizeLocal = TessellationComponentScreenSize;
+		ENQUEUE_RENDER_COMMAND(LandscapeChangeTessellationComponentScreenSizeCommand)(
+			[RenderProxies, ComponentCount, TessellationComponentScreenSizeLocal](FRHICommandListImmediate& RHICmdList)
 			{
 				for (int32 Idx = 0; Idx < ComponentCount; ++Idx)
 				{
 					if (RenderProxies[Idx] != nullptr)
 					{
-						RenderProxies[Idx]->ChangeTessellationComponentScreenSize_RenderThread(InTessellationComponentScreenSize);
+						RenderProxies[Idx]->ChangeTessellationComponentScreenSize_RenderThread(TessellationComponentScreenSizeLocal);
 					}
 				}
 
-		delete[] RenderProxies;
+				delete[] RenderProxies;
 			}
 		);
 	}
@@ -4300,21 +4298,19 @@ void ALandscapeProxy::ChangeComponentScreenSizeToUseSubSections(float InComponen
 			RenderProxies[Idx] = (FLandscapeComponentSceneProxy*)(LandscapeComponents[Idx]->SceneProxy);
 		}
 
-		ENQUEUE_UNIQUE_RENDER_COMMAND_THREEPARAMETER(
-			LandscapeChangeComponentScreenSizeToUseSubSectionsCommand,
-			FLandscapeComponentSceneProxy**, RenderProxies, RenderProxies,
-			int32, ComponentCount, ComponentCount,
-			float, InComponentScreenSizeToUseSubSections, ComponentScreenSizeToUseSubSections,
+		float ComponentScreenSizeToUseSubSectionsLocal = ComponentScreenSizeToUseSubSections;
+		ENQUEUE_RENDER_COMMAND(LandscapeChangeComponentScreenSizeToUseSubSectionsCommand)(
+			[RenderProxies, ComponentCount, ComponentScreenSizeToUseSubSectionsLocal](FRHICommandListImmediate& RHICmdList)
 			{
 				for (int32 Idx = 0; Idx < ComponentCount; ++Idx)
 				{
 					if (RenderProxies[Idx] != nullptr)
 					{
-						RenderProxies[Idx]->ChangeComponentScreenSizeToUseSubSections_RenderThread(InComponentScreenSizeToUseSubSections);
+						RenderProxies[Idx]->ChangeComponentScreenSizeToUseSubSections_RenderThread(ComponentScreenSizeToUseSubSectionsLocal);
 					}
 				}
 
-		delete[] RenderProxies;
+				delete[] RenderProxies;
 			}
 		);
 	}
@@ -4333,11 +4329,8 @@ void ALandscapeProxy::ChangeUseTessellationComponentScreenSizeFalloff(bool InUse
 			RenderProxies[Idx] = (FLandscapeComponentSceneProxy*)(LandscapeComponents[Idx]->SceneProxy);
 		}
 
-		ENQUEUE_UNIQUE_RENDER_COMMAND_THREEPARAMETER(
-			LandscapeChangeUseTessellationComponentScreenSizeFalloffCommand,
-			FLandscapeComponentSceneProxy**, RenderProxies, RenderProxies,
-			int32, ComponentCount, ComponentCount,
-			bool, InUseTessellationComponentScreenSizeFalloff, UseTessellationComponentScreenSizeFalloff,
+		ENQUEUE_RENDER_COMMAND(LandscapeChangeUseTessellationComponentScreenSizeFalloffCommand)(
+			[RenderProxies, ComponentCount, InUseTessellationComponentScreenSizeFalloff](FRHICommandListImmediate& RHICmdList)
 			{
 				for (int32 Idx = 0; Idx < ComponentCount; ++Idx)
 				{
@@ -4347,7 +4340,7 @@ void ALandscapeProxy::ChangeUseTessellationComponentScreenSizeFalloff(bool InUse
 					}
 				}
 
-		delete[] RenderProxies;
+				delete[] RenderProxies;
 			}
 		);
 	}
@@ -4366,21 +4359,19 @@ void ALandscapeProxy::ChangeTessellationComponentScreenSizeFalloff(float InTesse
 			RenderProxies[Idx] = (FLandscapeComponentSceneProxy*)(LandscapeComponents[Idx]->SceneProxy);
 		}
 
-		ENQUEUE_UNIQUE_RENDER_COMMAND_THREEPARAMETER(
-			LandscapeChangeTessellationComponentScreenSizeFalloffCommand,
-			FLandscapeComponentSceneProxy**, RenderProxies, RenderProxies,
-			int32, ComponentCount, ComponentCount,
-			float, InTessellationComponentScreenSizeFalloff, TessellationComponentScreenSizeFalloff,
+		float TessellationComponentScreenSizeFalloffLocal = TessellationComponentScreenSizeFalloff;
+		ENQUEUE_RENDER_COMMAND(LandscapeChangeTessellationComponentScreenSizeFalloffCommand)(
+			[RenderProxies, ComponentCount, TessellationComponentScreenSizeFalloffLocal](FRHICommandListImmediate& RHICmdList)
 			{
 				for (int32 Idx = 0; Idx < ComponentCount; ++Idx)
 				{
 					if (RenderProxies[Idx] != nullptr)
 					{
-						RenderProxies[Idx]->ChangeTessellationComponentScreenSizeFalloff_RenderThread(InTessellationComponentScreenSizeFalloff);
+						RenderProxies[Idx]->ChangeTessellationComponentScreenSizeFalloff_RenderThread(TessellationComponentScreenSizeFalloffLocal);
 					}
 				}
 
-		delete[] RenderProxies;
+				delete[] RenderProxies;
 			}
 		);
 	}
diff --git a/Engine/Source/Runtime/MovieSceneCapture/Private/FrameGrabber.cpp b/Engine/Source/Runtime/MovieSceneCapture/Private/FrameGrabber.cpp
--- a/Engine/Source/Runtime/MovieSceneCapture/Private/FrameGrabber.cpp
+++ b/Engine/Source/Runtime/MovieSceneCapture/Private/FrameGrabber.cpp
@@ -53,24 +53,22 @@ void FViewportSurfaceReader::Resize(uint32 Width, uint32 Height)
 {
 	ReadbackTexture.SafeRelease();
 
-	ENQUEUE_UNIQUE_RENDER_COMMAND_THREEPARAMETER(
-		CreateCaptureFrameTexture,
-		int32, InWidth, Width,
-		int32, InHeight, Height,
-		FViewportSurfaceReader*, This, this,
-	{
-		FRHIResourceCreateInfo CreateInfo;
-
-		This->ReadbackTexture = RHICreateTexture2D(
-			InWidth,
-			InHeight,
-			This->PixelFormat,
-			1,
-			1,
-			TexCreate_CPUReadback,
-			CreateInfo
-			);
-	});
+	FViewportSurfaceReader* This = this;
+	ENQUEUE_RENDER_COMMAND(CreateCaptureFrameTexture)(
+		[Width, Height, This](FRHICommandListImmediate& RHICmdList)
+		{
+			FRHIResourceCreateInfo CreateInfo;
+
+			This->ReadbackTexture = RHICreateTexture2D(
+				Width,
+				Height,
+				This->PixelFormat,
+				1,
+				1,
+				TexCreate_CPUReadback,
+				CreateInfo
+				);
+		});
 }
 
 void FViewportSurfaceReader::BlockUntilAvailable()
diff --git a/Engine/Source/Runtime/RenderCore/Public/RenderingThread.h b/Engine/Source/Runtime/RenderCore/Public/RenderingThread.h
--- a/Engine/Source/Runtime/RenderCore/Public/RenderingThread.h
+++ b/Engine/Source/Runtime/RenderCore/Public/RenderingThread.h
@@ -438,7 +438,8 @@ FORCEINLINE_DEBUGGABLE void EnqueueUniqueRenderCommand(LAMBDA& Lambda)
 
 #define ENQUEUE_UNIQUE_RENDER_COMMAND_THREEPARAMETER(TypeName,ParamType1,ParamName1,ParamValue1,ParamType2,ParamName2,ParamValue2,ParamType3,ParamName3,ParamValue3,Code) \
 	ENQUEUE_UNIQUE_RENDER_COMMAND_THREEPARAMETER_DECLARE(TypeName,ParamType1,ParamName1,ParamValue1,ParamType2,ParamName2,ParamValue2,ParamType3,ParamName3,ParamValue3,Code) \
-	ENQUEUE_UNIQUE_RENDER_COMMAND_THREEPARAMETER_CREATE(TypeName,ParamType1,ParamValue1,ParamType2,ParamValue2,ParamType3,ParamValue3)
+	ENQUEUE_UNIQUE_RENDER_COMMAND_THREEPARAMETER_CREATE(TypeName,ParamType1,ParamValue1,ParamType2,ParamValue2,ParamType3,ParamValue3) \
+	DEPRECATED_MACRO(4.22, "Please use ENQUEUE_RENDER_COMMAND instead")
 
 #define ENQUEUE_UNIQUE_RENDER_COMMAND_THREEPARAMETER_DECLARE_TEMPLATE(TypeName,TemplateParamName,ParamType1,ParamName1,ParamValue1,ParamType2,ParamName2,ParamValue2,ParamType3,ParamName3,ParamValue3,Code) \
 	template <typename TemplateParamName> \
diff --git a/Engine/Source/Runtime/Renderer/Private/PlanarReflectionRendering.cpp b/Engine/Source/Runtime/Renderer/Private/PlanarReflectionRendering.cpp
--- a/Engine/Source/Runtime/Renderer/Private/PlanarReflectionRendering.cpp
+++ b/Engine/Source/Runtime/Renderer/Private/PlanarReflectionRendering.cpp
@@ -600,14 +600,14 @@ void FScene::AddPlanarReflection(UPlanarReflectionComponent* Component)
 	check(Component->SceneProxy);
 	PlanarReflections_GameThread.Add(Component);
 
-	ENQUEUE_UNIQUE_RENDER_COMMAND_TWOPARAMETER(
-		FAddPlanarReflectionCommand,
-		FPlanarReflectionSceneProxy*,SceneProxy,Component->SceneProxy,
-		FScene*,Scene,this,
-	{
-		Scene->ReflectionSceneData.bRegisteredReflectionCapturesHasChanged = true;
-		Scene->PlanarReflections.Add(SceneProxy);
-	});
+	FPlanarReflectionSceneProxy* SceneProxy = Component->SceneProxy;
+	FScene* Scene = this;
+	ENQUEUE_RENDER_COMMAND(FAddPlanarReflectionCommand)(
+		[SceneProxy, Scene](FRHICommandListImmediate& RHICmdList)
+		{
+			Scene->ReflectionSceneData.bRegisteredReflectionCapturesHasChanged = true;
+			Scene->PlanarReflections.Add(SceneProxy);
+		});
 }
 
 void FScene::RemovePlanarReflection(UPlanarReflectionComponent* Component) 
@@ -615,29 +615,29 @@ void FScene::RemovePlanarReflection(UPlanarReflectionComponent* Component)
 	check(Component->SceneProxy);
 	PlanarReflections_GameThread.Remove(Component);
 
-	ENQUEUE_UNIQUE_RENDER_COMMAND_TWOPARAMETER(
-		FRemovePlanarReflectionCommand,
-		FPlanarReflectionSceneProxy*,SceneProxy,Component->SceneProxy,
-		FScene*,Scene,this,
-	{
-		Scene->ReflectionSceneData.bRegisteredReflectionCapturesHasChanged = true;
-		Scene->PlanarReflections.Remove(SceneProxy);
-	});
+	FPlanarReflectionSceneProxy* SceneProxy = Component->SceneProxy;
+	FScene* Scene = this;
+	ENQUEUE_RENDER_COMMAND(FRemovePlanarReflectionCommand)(
+		[SceneProxy, Scene](FRHICommandListImmediate& RHICmdList)
+		{
+			Scene->ReflectionSceneData.bRegisteredReflectionCapturesHasChanged = true;
+			Scene->PlanarReflections.Remove(SceneProxy);
+		});
 }
 
 void FScene::UpdatePlanarReflectionTransform(UPlanarReflectionComponent* Component)
 {	
 	check(Component->SceneProxy);
 
-	ENQUEUE_UNIQUE_RENDER_COMMAND_THREEPARAMETER(
-		FUpdatePlanarReflectionCommand,
-		FPlanarReflectionSceneProxy*,SceneProxy,Component->SceneProxy,
-		FMatrix,Transform,Component->GetComponentTransform().ToMatrixWithScale(),
-		FScene*,Scene,this,
-	{
-		Scene->ReflectionSceneData.bRegisteredReflectionCapturesHasChanged = true;
-		SceneProxy->UpdateTransform(Transform);
-	});
+	FPlanarReflectionSceneProxy* SceneProxy = Component->SceneProxy;
+	FMatrix Transform = Component->GetComponentTransform().ToMatrixWithScale();
+	FScene* Scene = this;
+	ENQUEUE_RENDER_COMMAND(FUpdatePlanarReflectionCommand)(
+		[SceneProxy, Transform, Scene](FRHICommandListImmediate& RHICmdList)
+		{
+			Scene->ReflectionSceneData.bRegisteredReflectionCapturesHasChanged = true;
+			SceneProxy->UpdateTransform(Transform);
+		});
 }
 
 class FPlanarReflectionPS : public FGlobalShader
diff --git a/Engine/Source/Runtime/Renderer/Private/ReflectionEnvironmentCapture.cpp b/Engine/Source/Runtime/Renderer/Private/ReflectionEnvironmentCapture.cpp
--- a/Engine/Source/Runtime/Renderer/Private/ReflectionEnvironmentCapture.cpp
+++ b/Engine/Source/Runtime/Renderer/Private/ReflectionEnvironmentCapture.cpp
@@ -945,11 +945,10 @@ void FScene::AllocateReflectionCaptures(const TArray<UReflectionCaptureComponent
 				{
 					// We can do a fast GPU copy to realloc the array, so we don't need to update all captures
 					ReflectionSceneData.MaxAllocatedReflectionCubemapsGameThread = DesiredMaxCubemaps;
-					ENQUEUE_UNIQUE_RENDER_COMMAND_THREEPARAMETER(
-						GPUResizeArrayCommand,
-						FScene*, Scene, this,
-						uint32, MaxSize, ReflectionSceneData.MaxAllocatedReflectionCubemapsGameThread,
-						int32, ReflectionCaptureSize, ReflectionCaptureSize,
+					FScene* Scene = this;
+					uint32 MaxSize = ReflectionSceneData.MaxAllocatedReflectionCubemapsGameThread;
+					ENQUEUE_RENDER_COMMAND(GPUResizeArrayCommand)(
+						[Scene, MaxSize, ReflectionCaptureSize](FRHICommandListImmediate& RHICmdList)
 						{
 							// Update the scene's cubemap array, preserving the original contents with a GPU-GPU copy
 							Scene->ReflectionSceneData.ResizeCubemapArrayGPU(MaxSize, ReflectionCaptureSize);
@@ -963,15 +962,14 @@ void FScene::AllocateReflectionCaptures(const TArray<UReflectionCaptureComponent
 			{
 				ReflectionSceneData.MaxAllocatedReflectionCubemapsGameThread = DesiredMaxCubemaps;
 
-				ENQUEUE_UNIQUE_RENDER_COMMAND_THREEPARAMETER( 
-					ResizeArrayCommand,
-					FScene*, Scene, this,
-					uint32, MaxSize, ReflectionSceneData.MaxAllocatedReflectionCubemapsGameThread,
-					int32, ReflectionCaptureSize, ReflectionCaptureSize,
-				{
-					// Update the scene's cubemap array, which will reallocate it, so we no longer have the contents of existing entries
-					Scene->ReflectionSceneData.CubemapArray.UpdateMaxCubemaps(MaxSize, ReflectionCaptureSize);
-				});
+				FScene* Scene = this;
+				uint32 MaxSize = ReflectionSceneData.MaxAllocatedReflectionCubemapsGameThread;
+				ENQUEUE_RENDER_COMMAND(ResizeArrayCommand)(
+					[Scene, MaxSize, ReflectionCaptureSize](FRHICommandListImmediate& RHICmdList)
+					{
+						// Update the scene's cubemap array, which will reallocate it, so we no longer have the contents of existing entries
+						Scene->ReflectionSceneData.CubemapArray.UpdateMaxCubemaps(MaxSize, ReflectionCaptureSize);
+					});
 
 				// Recapture all reflection captures now that we have reallocated the cubemap array
 				UpdateAllReflectionCaptures(CaptureReason, bVerifyOnlyCapturing);
@@ -1116,14 +1114,13 @@ void FScene::GetReflectionCaptureData(UReflectionCaptureComponent* Component, FR
 {
 	check(GetFeatureLevel() >= ERHIFeatureLevel::SM5);
 
-	ENQUEUE_UNIQUE_RENDER_COMMAND_THREEPARAMETER(
-		GetReflectionDataCommand,
-		FScene*,Scene,this,
-		const UReflectionCaptureComponent*,Component,Component,
-		FReflectionCaptureData*,OutCaptureData,&OutCaptureData,
-	{
-		GetReflectionCaptureData_RenderingThread(RHICmdList, Scene, Component, OutCaptureData);
-	});
+	FScene* Scene = this;
+	FReflectionCaptureData* OutCaptureDataPtr = &OutCaptureData;
+	ENQUEUE_RENDER_COMMAND(GetReflectionDataCommand)(
+		[Scene, Component, OutCaptureDataPtr](FRHICommandListImmediate& RHICmdList)
+		{
+			GetReflectionCaptureData_RenderingThread(RHICmdList, Scene, Component, OutCaptureDataPtr);
+		});
 
 	// Necessary since the RT is writing to OutDerivedData directly
 	FlushRenderingCommands();
@@ -1510,13 +1507,12 @@ void FScene::CaptureOrUploadReflectionCapture(UReflectionCaptureComponent* Captu
 			// We can't use the component's SceneProxy here because the component may not be registered with the scene
 			FReflectionCaptureProxy* ReflectionProxy = new FReflectionCaptureProxy(CaptureComponent);
 
-			ENQUEUE_UNIQUE_RENDER_COMMAND_THREEPARAMETER( 
-				CopyCommand,
-				FScene*, Scene, this,
-				FReflectionCaptureProxy*, ReflectionProxy, ReflectionProxy,
-				ERHIFeatureLevel::Type, FeatureLevel, GetFeatureLevel(),
+			FScene* Scene = this;
+			ERHIFeatureLevel::Type InFeatureLevel = GetFeatureLevel();
+			ENQUEUE_RENDER_COMMAND(CopyCommand)(
+				[Scene, ReflectionProxy, InFeatureLevel](FRHICommandListImmediate& RHICmdList)
 			{
-				if (FeatureLevel == ERHIFeatureLevel::SM5)
+				if (InFeatureLevel == ERHIFeatureLevel::SM5)
 				{
 					CopyToSceneArray(RHICmdList, Scene, ReflectionProxy);
 				}
@@ -1662,21 +1658,20 @@ void FScene::UpdateSkyCaptureContents(
 		// Optionally copy the filtered mip chain to the output texture
 		if (OutProcessedTexture)
 		{
-			ENQUEUE_UNIQUE_RENDER_COMMAND_THREEPARAMETER( 
-				CopyCommand,
-				FScene*, Scene, this,
-				FTexture*, ProcessedTexture, OutProcessedTexture,
-				ERHIFeatureLevel::Type, FeatureLevel, GetFeatureLevel(),
-			{
-				if (FeatureLevel <= ERHIFeatureLevel::ES3_1)
-				{
-					MobileReflectionEnvironmentCapture::CopyToSkyTexture(RHICmdList, Scene, ProcessedTexture);
-				}
-				else
+			FScene* Scene = this;
+			ERHIFeatureLevel::Type InFeatureLevel = GetFeatureLevel();
+			ENQUEUE_RENDER_COMMAND(CopyCommand)(
+				[Scene, OutProcessedTexture, InFeatureLevel](FRHICommandListImmediate& RHICmdList)
 				{
-					CopyToSkyTexture(RHICmdList, Scene, ProcessedTexture);
-				}
-			});
+					if (InFeatureLevel <= ERHIFeatureLevel::ES3_1)
+					{
+						MobileReflectionEnvironmentCapture::CopyToSkyTexture(RHICmdList, Scene, OutProcessedTexture);
+					}
+					else
+					{
+						CopyToSkyTexture(RHICmdList, Scene, OutProcessedTexture);
+					}
+				});
 		}
 	}
 }
diff --git a/Engine/Source/Runtime/Renderer/Private/Renderer.cpp b/Engine/Source/Runtime/Renderer/Private/Renderer.cpp
--- a/Engine/Source/Runtime/Renderer/Private/Renderer.cpp
+++ b/Engine/Source/Runtime/Renderer/Private/Renderer.cpp
@@ -397,14 +397,12 @@ void FRendererModule::GPUBenchmark(FSynthBenchmarkResults& InOut, float WorkScal
 
 	FSceneView DummyView(ViewInitOptions);
 	FlushRenderingCommands();
-	ENQUEUE_UNIQUE_RENDER_COMMAND_THREEPARAMETER(
-	  RendererGPUBenchmarkCommand,
-	  FSceneView, DummyView, DummyView,
-	  float, WorkScale, WorkScale,
-	  FSynthBenchmarkResults&, InOut, InOut,
-	{
-		RendererGPUBenchmark(RHICmdList, InOut, DummyView, WorkScale);
-	});
+	FSynthBenchmarkResults* InOutPtr = &InOut;
+	ENQUEUE_RENDER_COMMAND(RendererGPUBenchmarkCommand)(
+		[DummyView, WorkScale, InOutPtr](FRHICommandListImmediate& RHICmdList)
+		{
+			RendererGPUBenchmark(RHICmdList, *InOutPtr, DummyView, WorkScale);
+		});
 	FlushRenderingCommands();
 }
 
diff --git a/Engine/Source/Runtime/Renderer/Private/RendererScene.cpp b/Engine/Source/Runtime/Renderer/Private/RendererScene.cpp
--- a/Engine/Source/Runtime/Renderer/Private/RendererScene.cpp
+++ b/Engine/Source/Runtime/Renderer/Private/RendererScene.cpp
@@ -2206,11 +2206,10 @@ void FScene::UpdateLightTransform(ULightComponent* Light)
 		FUpdateLightTransformParameters Parameters;
 		Parameters.LightToWorld = Light->GetComponentTransform().ToMatrixNoScale();
 		Parameters.Position = Light->GetLightPosition();
-		ENQUEUE_UNIQUE_RENDER_COMMAND_THREEPARAMETER(
-			UpdateLightTransform,
-			FScene*,Scene,this,
-			FLightSceneInfo*,LightSceneInfo,Light->SceneProxy->GetLightSceneInfo(),
-			FUpdateLightTransformParameters,Parameters,Parameters,
+		FScene* Scene = this;
+		FLightSceneInfo* LightSceneInfo = Light->SceneProxy->GetLightSceneInfo();
+		ENQUEUE_RENDER_COMMAND(UpdateLightTransform)(
+			[Scene, LightSceneInfo, Parameters](FRHICommandListImmediate& RHICmdList)
 			{
 				FScopeCycleCounter Context(LightSceneInfo->Proxy->GetStatId());
 				Scene->UpdateLightTransform_RenderThread(LightSceneInfo, Parameters);
@@ -2244,11 +2243,10 @@ void FScene::UpdateLightColorAndBrightness(ULightComponent* Light)
 			 NewParameters.NewColor *= FLinearColor::MakeFromColorTemperature(Light->Temperature);
 		}
 	
-		ENQUEUE_UNIQUE_RENDER_COMMAND_THREEPARAMETER(
-			UpdateLightColorAndBrightness,
-			FLightSceneInfo*,LightSceneInfo,Light->SceneProxy->GetLightSceneInfo(),
-			FScene*,Scene,this,
-			FUpdateLightColorParameters,Parameters,NewParameters,
+		FScene* Scene = this;
+		FLightSceneInfo* LightSceneInfo = Light->SceneProxy->GetLightSceneInfo();
+		ENQUEUE_RENDER_COMMAND(UpdateLightColorAndBrightness)(
+			[LightSceneInfo, Scene, NewParameters](FRHICommandListImmediate& RHICmdList)
 			{
 				if( LightSceneInfo && LightSceneInfo->bVisible )
 				{
@@ -2258,16 +2256,16 @@ void FScene::UpdateLightColorAndBrightness(ULightComponent* Light)
 					Scene->bScenesPrimitivesNeedStaticMeshElementUpdate =
 						Scene->bScenesPrimitivesNeedStaticMeshElementUpdate ||
 						( Scene->GetShadingPath() == EShadingPath::Mobile 
-						&& Parameters.NewColor.IsAlmostBlack() != LightSceneInfo->Proxy->GetColor().IsAlmostBlack() );
+						&& NewParameters.NewColor.IsAlmostBlack() != LightSceneInfo->Proxy->GetColor().IsAlmostBlack() );
 
-					LightSceneInfo->Proxy->SetColor(Parameters.NewColor);
-					LightSceneInfo->Proxy->IndirectLightingScale = Parameters.NewIndirectLightingScale;
-					LightSceneInfo->Proxy->VolumetricScatteringIntensity = Parameters.NewVolumetricScatteringIntensity;
+					LightSceneInfo->Proxy->SetColor(NewParameters.NewColor);
+					LightSceneInfo->Proxy->IndirectLightingScale = NewParameters.NewIndirectLightingScale;
+					LightSceneInfo->Proxy->VolumetricScatteringIntensity = NewParameters.NewVolumetricScatteringIntensity;
 
 					// Also update the LightSceneInfoCompact
 					if( LightSceneInfo->Id != INDEX_NONE )
 					{
-						Scene->Lights[ LightSceneInfo->Id ].Color = Parameters.NewColor;
+						Scene->Lights[ LightSceneInfo->Id ].Color = NewParameters.NewColor;
 					}
 				}
 			});
@@ -2589,11 +2587,9 @@ void FScene::AddSpeedTreeWind(FVertexFactory* VertexFactory, const UStaticMesh*
 {
 	if (StaticMesh != NULL && StaticMesh->SpeedTreeWind.IsValid() && StaticMesh->RenderData.IsValid())
 	{
-		ENQUEUE_UNIQUE_RENDER_COMMAND_THREEPARAMETER(
-			FAddSpeedTreeWindCommand,
-			FScene*,Scene,this,
-			const UStaticMesh*,StaticMesh,StaticMesh,
-			FVertexFactory*,VertexFactory,VertexFactory,
+		FScene* Scene = this;
+		ENQUEUE_RENDER_COMMAND(FAddSpeedTreeWindCommand)(
+			[Scene, StaticMesh, VertexFactory](FRHICommandListImmediate& RHICmdList)
 			{
 				Scene->SpeedTreeVertexFactoryMap.Add(VertexFactory, StaticMesh);
 
@@ -2761,11 +2757,9 @@ void FScene::GetRelevantLights( UPrimitiveComponent* Primitive, TArray<const ULi
 	if( Primitive && RelevantLights )
 	{
 		// Add interacting lights to the array.
-		ENQUEUE_UNIQUE_RENDER_COMMAND_THREEPARAMETER(
-			FGetRelevantLightsCommand,
-			const FScene*,Scene,this,
-			UPrimitiveComponent*,Primitive,Primitive,
-			TArray<const ULightComponent*>*,RelevantLights,RelevantLights,
+		const FScene* Scene = this;
+		ENQUEUE_RENDER_COMMAND(FGetRelevantLightsCommand)(
+			[Scene, Primitive, RelevantLights](FRHICommandListImmediate& RHICmdList)
 			{
 				Scene->GetRelevantLights_RenderThread( Primitive, RelevantLights );
 			});
