Decrease priority of MeshDrawCommandPassSetupTask, so it doesn't block high priority game thread tasks.

diff --git a/Engine/Source/Runtime/Renderer/Private/MeshDrawCommands.cpp b/Engine/Source/Runtime/Renderer/Private/MeshDrawCommands.cpp
--- a/Engine/Source/Runtime/Renderer/Private/MeshDrawCommands.cpp
+++ b/Engine/Source/Runtime/Renderer/Private/MeshDrawCommands.cpp
@@ -646,9 +646,8 @@ void ApplyViewOverridesToMeshDrawCommands(
 FAutoConsoleTaskPriority CPrio_FMeshDrawCommandPassSetupTask(
 	TEXT("TaskGraph.TaskPriorities.FMeshDrawCommandPassSetupTask"),
 	TEXT("Task and thread priority for FMeshDrawCommandPassSetupTask."),
-	ENamedThreads::HighThreadPriority, // if we have high priority task threads, then use them...
-	ENamedThreads::NormalTaskPriority, // .. at normal task priority
-	ENamedThreads::HighTaskPriority // if we don't have hi pri threads, then use normal priority threads at high task priority instead
+	ENamedThreads::NormalThreadPriority,
+	ENamedThreads::HighTaskPriority
 );
 
 /**
