Added STAT_GatherShadowPrimitivesPacket into parallel for.
#rnx

diff --git a/Engine/Source/Runtime/Renderer/Private/ShadowSetup.cpp b/Engine/Source/Runtime/Renderer/Private/ShadowSetup.cpp
--- a/Engine/Source/Runtime/Renderer/Private/ShadowSetup.cpp
+++ b/Engine/Source/Runtime/Renderer/Private/ShadowSetup.cpp
@@ -2999,6 +2999,8 @@ struct FGatherShadowPrimitivesPacket
 
 	void AnyThreadTask()
 	{
+		QUICK_SCOPE_CYCLE_COUNTER(STAT_GatherShadowPrimitivesPacket);
+
 		if (Node)
 		{
 			// Check all the primitives in this octree node.
