Removes a check() that was missing in 4721425

Repro steps -> disable DOF, and allow alpha channel through post processing in project settings.

diff --git a/Engine/Source/Runtime/Renderer/Private/PostProcess/PostProcessing.cpp b/Engine/Source/Runtime/Renderer/Private/PostProcess/PostProcessing.cpp
--- a/Engine/Source/Runtime/Renderer/Private/PostProcess/PostProcessing.cpp
+++ b/Engine/Source/Runtime/Renderer/Private/PostProcess/PostProcessing.cpp
@@ -1585,7 +1585,6 @@ void FPostProcessing::Process(FRHICommandListImmediate& RHICmdList, const FViewI
 
 			if(SeparateTranslucency.IsValid() && !bSepTransWasApplied)
 			{
-				checkf(!FPostProcessing::HasAlphaChannelSupport(), TEXT("Separate translucency was supposed to be disabled automatically."));
 				const bool bIsComputePass = ShouldDoComputePostProcessing(Context.View);
 				// separate translucency is done here or in AddPostProcessDepthOfFieldBokeh()
 				FRenderingCompositePass* NodeRecombined = Context.Graph.RegisterPass(new(FMemStack::Get()) FRCPassPostProcessBokehDOFRecombine(bIsComputePass));
