Decrease amount of allocations inside mark relevance by prealocating a bit of memory for mesh draw commands.

diff --git a/Engine/Source/Runtime/Renderer/Private/MeshDrawCommands.cpp b/Engine/Source/Runtime/Renderer/Private/MeshDrawCommands.cpp
--- a/Engine/Source/Runtime/Renderer/Private/MeshDrawCommands.cpp
+++ b/Engine/Source/Runtime/Renderer/Private/MeshDrawCommands.cpp
@@ -1068,7 +1068,7 @@ struct FRHICommandUpdatePrimitiveIdBuffer : public FRHICommand<FRHICommandUpdate
 	{
 		// Upload vertex buffer data.
 		void* RESTRICT Data = (void* RESTRICT)GDynamicRHI->RHILockVertexBuffer(VertexBuffer, 0, VertexBufferDataSize, RLM_WriteOnly);
-		FMemory::BigBlockMemcpy(Data, VertexBufferData, VertexBufferDataSize);
+		FMemory::Memcpy(Data, VertexBufferData, VertexBufferDataSize);
 		GDynamicRHI->RHIUnlockVertexBuffer(VertexBuffer);
 
 		FMemory::Free(VertexBufferData);
@@ -1147,7 +1147,7 @@ void FParallelMeshDrawCommandPass::DispatchDraw(FParallelCommandListSet* Paralle
 		{
 			// Can immediately upload vertex buffer data, as there is no parallel draw task.
 			void* RESTRICT Data = RHILockVertexBuffer(PrimitiveIdVertexBufferRHI, 0, TaskContext.PrimitiveIdBufferDataSize, RLM_WriteOnly);
-			FMemory::BigBlockMemcpy(Data, TaskContext.PrimitiveIdBufferData, TaskContext.PrimitiveIdBufferDataSize);
+			FMemory::Memcpy(Data, TaskContext.PrimitiveIdBufferData, TaskContext.PrimitiveIdBufferDataSize);
 			RHIUnlockVertexBuffer(PrimitiveIdVertexBufferRHI);
 		}
 
diff --git a/Engine/Source/Runtime/Renderer/Private/SceneVisibility.cpp b/Engine/Source/Runtime/Renderer/Private/SceneVisibility.cpp
--- a/Engine/Source/Runtime/Renderer/Private/SceneVisibility.cpp
+++ b/Engine/Source/Runtime/Renderer/Private/SceneVisibility.cpp
@@ -1663,9 +1663,6 @@ namespace EMarkMaskBits
 	};
 }
 
-typedef TChunkedArray<FVisibleMeshDrawCommand> FPassDrawCommandArray;
-typedef TChunkedArray<const FStaticMeshBatch*> FPassDrawCommandBuildRequestArray;
-
 struct FDrawCommandRelevancePacket
 {
 	FDrawCommandRelevancePacket()
@@ -1678,8 +1675,8 @@ struct FDrawCommandRelevancePacket
 		}
 	}
 
-	FPassDrawCommandArray VisibleCachedDrawCommands[EMeshPass::Num];
-	FPassDrawCommandBuildRequestArray DynamicBuildRequests[EMeshPass::Num];
+	TArray<FVisibleMeshDrawCommand, TInlineAllocator<256>> VisibleCachedDrawCommands[EMeshPass::Num];
+	TArray<const FStaticMeshBatch*, TInlineAllocator<256>> DynamicBuildRequests[EMeshPass::Num];
 	int32 NumDynamicBuildRequestElements[EMeshPass::Num];
 	bool bUseCachedMeshDrawCommands;
 
@@ -1716,13 +1713,13 @@ struct FDrawCommandRelevancePacket
 					CachedMeshDrawCommand.MeshCullMode,
 					CachedMeshDrawCommand.SortKey);
 
-				VisibleCachedDrawCommands[(uint32)PassType].AddElement(NewVisibleMeshDrawCommand);
+				VisibleCachedDrawCommands[(uint32)PassType].Add(NewVisibleMeshDrawCommand);
 			}
 		}
 		else
 		{
 			NumDynamicBuildRequestElements[PassType] += StaticMeshRelevance.NumElements;
-			DynamicBuildRequests[PassType].AddElement(&StaticMesh);
+			DynamicBuildRequests[PassType].Add(&StaticMesh);
 		}
 	}
 };
@@ -2314,8 +2311,24 @@ struct FRelevancePacket
 
 		for (int32 PassIndex = 0; PassIndex < EMeshPass::Num; PassIndex++)
 		{
-			DrawCommandPacket.VisibleCachedDrawCommands[PassIndex].CopyToLinearArray(WriteViewCommands.MeshCommands[PassIndex]);
-			DrawCommandPacket.DynamicBuildRequests[PassIndex].CopyToLinearArray(WriteViewCommands.DynamicMeshCommandBuildRequests[PassIndex]);
+			TArray<FVisibleMeshDrawCommand, TInlineAllocator<256>>& SrcCommands = DrawCommandPacket.VisibleCachedDrawCommands[PassIndex];
+			FMeshCommandOneFrameArray& DstCommands = WriteViewCommands.MeshCommands[PassIndex];
+			if (SrcCommands.Num() > 0)
+			{
+				static_assert(sizeof(SrcCommands[0]) == sizeof(DstCommands[0]), "Memcpy sizes must match.");
+				const int32 PrevNum = DstCommands.AddUninitialized(SrcCommands.Num());
+				FMemory::Memcpy(&DstCommands[PrevNum], &SrcCommands[0], SrcCommands.Num() * sizeof(SrcCommands[0]));
+			}
+
+			TArray<const FStaticMeshBatch*, TInlineAllocator<256>>& SrcRequests = DrawCommandPacket.DynamicBuildRequests[PassIndex];
+			TArray<const FStaticMeshBatch*, SceneRenderingAllocator>& DstRequests = WriteViewCommands.DynamicMeshCommandBuildRequests[PassIndex];
+			if (SrcRequests.Num() > 0)
+			{
+				static_assert(sizeof(SrcRequests[0]) == sizeof(DstRequests[0]), "Memcpy sizes must match.");
+				const int32 PrevNum = DstRequests.AddUninitialized(SrcRequests.Num());
+				FMemory::Memcpy(&DstRequests[PrevNum], &SrcRequests[0], SrcRequests.Num() * sizeof(SrcRequests[0]));
+			}
+
 			WriteViewCommands.NumDynamicMeshCommandBuildRequestElements[PassIndex] += DrawCommandPacket.NumDynamicBuildRequestElements[PassIndex];
 		}
 
