Fixed crash when multiple views try to run editor selection mesh processor in parallel.

diff --git a/Engine/Source/Runtime/Renderer/Private/SceneHitProxyRendering.cpp b/Engine/Source/Runtime/Renderer/Private/SceneHitProxyRendering.cpp
--- a/Engine/Source/Runtime/Renderer/Private/SceneHitProxyRendering.cpp
+++ b/Engine/Source/Runtime/Renderer/Private/SceneHitProxyRendering.cpp
@@ -827,17 +827,11 @@ int32 FEditorSelectionMeshProcessor::GetStencilValue(const FSceneView* View, con
 	return StencilValue;
 }
 
-TMap<const FPrimitiveSceneProxy*, int32> FEditorSelectionMeshProcessor::ProxyToStencilIndex;
-TMap<FName, int32> FEditorSelectionMeshProcessor::ActorNameToStencilIndex;
-
 FEditorSelectionMeshProcessor::FEditorSelectionMeshProcessor(const FScene* Scene, const FSceneView* InViewIfDynamicMeshCommand, FMeshPassDrawListContext* InDrawListContext)
 	: FMeshPassProcessor(Scene, Scene->GetFeatureLevel(), InViewIfDynamicMeshCommand, InDrawListContext)
 {
 	checkf(InViewIfDynamicMeshCommand, TEXT("Editor selection mesh process required dynamic mesh command mode."));
 
-	// Reset common outline maps per every view render.
-	ProxyToStencilIndex.Reset();
-	ActorNameToStencilIndex.Reset();
 	ActorNameToStencilIndex.Add(NAME_BSP, 1);
 
 	PassDrawRenderState.SetDepthStencilState(TStaticDepthStencilState<true, CF_DepthNearOrEqual, true, CF_Always, SO_Keep, SO_Keep, SO_Replace>::GetRHI());
diff --git a/Engine/Source/Runtime/Renderer/Private/SceneHitProxyRendering.h b/Engine/Source/Runtime/Renderer/Private/SceneHitProxyRendering.h
--- a/Engine/Source/Runtime/Renderer/Private/SceneHitProxyRendering.h
+++ b/Engine/Source/Runtime/Renderer/Private/SceneHitProxyRendering.h
@@ -67,9 +67,9 @@ private:
 	int32 GetStencilValue(const FSceneView* View, const FPrimitiveSceneProxy* PrimitiveSceneProxy);
 
 	/** This map is needed to ensure that individually selected proxies rendered more than once a frame (if they have multiple sections) share a common outline */
-	static TMap<const FPrimitiveSceneProxy*, int32> ProxyToStencilIndex;
+	TMap<const FPrimitiveSceneProxy*, int32> ProxyToStencilIndex;
 	/** This map is needed to ensure that proxies rendered more than once a frame (if they have multiple sections) share a common outline */
-	static TMap<FName, int32> ActorNameToStencilIndex;
+	TMap<FName, int32> ActorNameToStencilIndex;
 };
 
 #endif
\ No newline at end of file
