Merged to relevance packet constants into one, named one.

diff --git a/Engine/Source/Runtime/Renderer/Private/SceneVisibility.cpp b/Engine/Source/Runtime/Renderer/Private/SceneVisibility.cpp
--- a/Engine/Source/Runtime/Renderer/Private/SceneVisibility.cpp
+++ b/Engine/Source/Runtime/Renderer/Private/SceneVisibility.cpp
@@ -225,7 +225,6 @@ static FAutoConsoleVariableRef CVarLightMaxDrawDistanceScale(
 	ECVF_Scalability | ECVF_RenderThreadSafe
 );
 
-
 DECLARE_CYCLE_STAT(TEXT("Occlusion Readback"), STAT_CLMM_OcclusionReadback, STATGROUP_CommandListMarkers);
 DECLARE_CYCLE_STAT(TEXT("After Occlusion Readback"), STAT_CLMM_AfterOcclusionReadback, STATGROUP_CommandListMarkers);
 
@@ -1592,12 +1591,15 @@ static int32 OcclusionCull(FRHICommandListImmediate& RHICmdList, const FScene* S
 	return NumOccludedPrimitives;
 }
 
+const int32 InputsPrimNumPerRelevancePacket = 128;
+const int32 AverageMeshBatchNumPerRelevancePacket = InputsPrimNumPerRelevancePacket * 2;
+
 template<class T, int TAmplifyFactor = 1>
 struct FRelevancePrimSet
 {
 	enum
 	{
-		MaxInputPrims = 127, //like 128, but we leave space for NumPrims
+		MaxInputPrims = InputsPrimNumPerRelevancePacket - 1, // leave space for NumPrims.
 		MaxOutputPrims = MaxInputPrims * TAmplifyFactor
 	};
 	int32 NumPrims;
@@ -1663,6 +1665,9 @@ namespace EMarkMaskBits
 	};
 }
 
+typedef TArray<FVisibleMeshDrawCommand, TInlineAllocator<AverageMeshBatchNumPerRelevancePacket>> FPassDrawCommandArray;
+typedef TArray<const FStaticMeshBatch*, TInlineAllocator<AverageMeshBatchNumPerRelevancePacket>> FPassDrawCommandBuildRequestArray;
+
 struct FDrawCommandRelevancePacket
 {
 	FDrawCommandRelevancePacket()
@@ -1675,8 +1680,8 @@ struct FDrawCommandRelevancePacket
 		}
 	}
 
-	TArray<FVisibleMeshDrawCommand, TInlineAllocator<256>> VisibleCachedDrawCommands[EMeshPass::Num];
-	TArray<const FStaticMeshBatch*, TInlineAllocator<256>> DynamicBuildRequests[EMeshPass::Num];
+	FPassDrawCommandArray VisibleCachedDrawCommands[EMeshPass::Num];
+	FPassDrawCommandBuildRequestArray DynamicBuildRequests[EMeshPass::Num];
 	int32 NumDynamicBuildRequestElements[EMeshPass::Num];
 	bool bUseCachedMeshDrawCommands;
 
@@ -2311,7 +2316,7 @@ struct FRelevancePacket
 
 		for (int32 PassIndex = 0; PassIndex < EMeshPass::Num; PassIndex++)
 		{
-			TArray<FVisibleMeshDrawCommand, TInlineAllocator<256>>& SrcCommands = DrawCommandPacket.VisibleCachedDrawCommands[PassIndex];
+			FPassDrawCommandArray& SrcCommands = DrawCommandPacket.VisibleCachedDrawCommands[PassIndex];
 			FMeshCommandOneFrameArray& DstCommands = WriteViewCommands.MeshCommands[PassIndex];
 			if (SrcCommands.Num() > 0)
 			{
@@ -2320,7 +2325,7 @@ struct FRelevancePacket
 				FMemory::Memcpy(&DstCommands[PrevNum], &SrcCommands[0], SrcCommands.Num() * sizeof(SrcCommands[0]));
 			}
 
-			TArray<const FStaticMeshBatch*, TInlineAllocator<256>>& SrcRequests = DrawCommandPacket.DynamicBuildRequests[PassIndex];
+			FPassDrawCommandBuildRequestArray& SrcRequests = DrawCommandPacket.DynamicBuildRequests[PassIndex];
 			TArray<const FStaticMeshBatch*, SceneRenderingAllocator>& DstRequests = WriteViewCommands.DynamicMeshCommandBuildRequests[PassIndex];
 			if (SrcRequests.Num() > 0)
 			{
