Merging 4975335 to unblock RM

diff --git a/Engine/Source/Runtime/Renderer/Private/DeferredShadingRenderer.cpp b/Engine/Source/Runtime/Renderer/Private/DeferredShadingRenderer.cpp
--- a/Engine/Source/Runtime/Renderer/Private/DeferredShadingRenderer.cpp
+++ b/Engine/Source/Runtime/Renderer/Private/DeferredShadingRenderer.cpp
@@ -1651,30 +1651,6 @@ void FDeferredShadingSceneRenderer::Render(FRHICommandListImmediate& RHICmdList)
 
 	ServiceLocalQueue();
 
-	// Make sure the rendertargets the particle system might need aren't currently bound on RHIs
-	// that don't really have internal renderpasses.
-	UnbindRenderTargets(RHICmdList);
-
-	// Notify the FX system that opaque primitives have been rendered and we now have a valid depth buffer.
-	if (Scene->FXSystem && Views.IsValidIndex(0) && bAllowGPUParticleSceneUpdate)
-	{
-		SCOPE_CYCLE_COUNTER(STAT_FDeferredShadingSceneRenderer_FXSystem_PostRenderOpaque);
-
-		FSceneTexturesUniformParameters SceneTextureParameters;
-		SetupSceneTextureUniformParameters(SceneContext, FeatureLevel, ESceneTextureSetupMode::SceneDepth | ESceneTextureSetupMode::GBuffers, SceneTextureParameters);
-		TUniformBufferRef<FSceneTexturesUniformParameters> SceneTextureUniformBuffer = TUniformBufferRef<FSceneTexturesUniformParameters>::CreateUniformBufferImmediate(SceneTextureParameters, UniformBuffer_SingleFrame);
-
-		Scene->FXSystem->PostRenderOpaque(
-			RHICmdList,
-			Views[0].ViewUniformBuffer,
-			&FSceneTexturesUniformParameters::StaticStructMetadata,
-			SceneTextureUniformBuffer.GetReference()
-			);
-		ServiceLocalQueue();
-	}
-
-	checkSlow(RHICmdList.IsOutsideRenderPass());
-
 	TRefCountPtr<IPooledRenderTarget> VelocityRT;
 
 	if (bBasePassCanOutputVelocity)
@@ -2000,9 +1976,29 @@ void FDeferredShadingSceneRenderer::Render(FRHICommandListImmediate& RHICmdList)
 	checkSlow(RHICmdList.IsOutsideRenderPass());
 	UnbindRenderTargets(RHICmdList);
 
+	// Notify the FX system that opaque primitives have been rendered and we now have a valid depth buffer.
+	if (Scene->FXSystem && Views.IsValidIndex(0) && bAllowGPUParticleSceneUpdate)
+	{
+		SCOPE_CYCLE_COUNTER(STAT_FDeferredShadingSceneRenderer_FXSystem_PostRenderOpaque);
+
+		FSceneTexturesUniformParameters SceneTextureParameters;
+		SetupSceneTextureUniformParameters(SceneContext, FeatureLevel, ESceneTextureSetupMode::SceneDepth | ESceneTextureSetupMode::GBuffers, SceneTextureParameters);
+		TUniformBufferRef<FSceneTexturesUniformParameters> SceneTextureUniformBuffer = TUniformBufferRef<FSceneTexturesUniformParameters>::CreateUniformBufferImmediate(SceneTextureParameters, UniformBuffer_SingleFrame);
+
+		Scene->FXSystem->PostRenderOpaque(
+			RHICmdList,
+			Views[0].ViewUniformBuffer,
+			&FSceneTexturesUniformParameters::StaticStructMetadata,
+			SceneTextureUniformBuffer.GetReference()
+		);
+		ServiceLocalQueue();
+	}
+
 	// No longer needed, release
 	LightShaftOutput.LightShaftOcclusion = NULL;
 
+	checkSlow(RHICmdList.IsOutsideRenderPass());
+
 	GRenderTargetPool.AddPhaseEvent(TEXT("Translucency"));
 
 	// Draw translucency.
