Fixed: Noticeable frame time increases on low end Adreno devices in SunTemple
Added cvar (r.Mobile.MeshSortingMethod) to optionally enable strict front to back mesh sorting, similar to sorting in 4.19 (enabled for SunTemple).
Default is sort by state.
On GPU bound devices when scene has < 300 drawcalls it's faster to use strict front to back mesh sorting.

diff --git a/Engine/Source/Runtime/Renderer/Private/MeshDrawCommands.cpp b/Engine/Source/Runtime/Renderer/Private/MeshDrawCommands.cpp
--- a/Engine/Source/Runtime/Renderer/Private/MeshDrawCommands.cpp
+++ b/Engine/Source/Runtime/Renderer/Private/MeshDrawCommands.cpp
@@ -17,6 +17,14 @@ static TAutoConsoleVariable<int32> CVarMeshDrawCommandsParallelPassSetup(
 	TEXT("Whether to setup mesh draw command pass in parallel."),
 	ECVF_RenderThreadSafe);
 
+static TAutoConsoleVariable<int32> CVarMobileMeshSortingMethod(
+	TEXT("r.Mobile.MeshSortingMethod"),
+	0,
+	TEXT("How to sort mesh commands on mobile:\n")
+	TEXT("\t0: Sort by state, roughly front to back (Default).\n")
+	TEXT("\t1: Strict front to back sorting.\n"),
+	ECVF_RenderThreadSafe);
+
 FPrimitiveIdVertexBufferPool::FPrimitiveIdVertexBufferPool()
 	: DiscardId(0)
 {
@@ -173,7 +181,35 @@ void UpdateTranslucentMeshSortKeys(
 	}
 }
 
-static uint64 GetMobileBasePassSortKey(bool bMasked, bool bBackground, int32 PipelineId, int32 StateBucketId, float PipelineDistance, float PrimitiveDistance)
+static uint64 GetMobileBasePassSortKey_FrontToBack(bool bMasked, bool bBackground, int32 PipelineId, int32 StateBucketId, float PrimitiveDistance)
+{
+	union
+	{
+		uint64 PackedInt;
+		struct
+		{
+			uint64 StateBucketId : 27; 		// Order by state bucket
+			uint64 PipelineId : 20;			// Order by PSO
+			uint64 DepthBits : 15;			// Order by primitive depth
+			uint64 Background : 1;			// Non-background meshes first 
+			uint64 Masked : 1;				// Non-masked first
+		} Fields;
+	} Key;
+
+	union FFloatToInt { float F; uint32 I; };
+	FFloatToInt F2I;
+
+	Key.Fields.Masked = bMasked;
+	Key.Fields.Background = bBackground;
+	F2I.F = PrimitiveDistance;
+	Key.Fields.DepthBits = ((-int32(F2I.I >> 31) | 0x80000000) ^ F2I.I) >> 17;
+	Key.Fields.PipelineId = PipelineId;
+	Key.Fields.StateBucketId = StateBucketId;
+	
+	return Key.PackedInt;
+}
+
+static uint64 GetMobileBasePassSortKey_ByState(bool bMasked, bool bBackground, int32 PipelineId, int32 StateBucketId, float PipelineDistance, float PrimitiveDistance)
 {
 	const float PrimitiveDepthQuantization = ((1 << 14) - 1);
 
@@ -247,53 +283,74 @@ void UpdateMobileBasePassMeshSortKeys(
 )
 {
 	QUICK_SCOPE_CYCLE_COUNTER(STAT_UpdateMobileBasePassMeshSortKeys);
-
+	
 	int32 NumCmds = VisibleMeshCommands.Num();
-
-	// pre-compute distance to a group of meshes that share same PSO
-	TMap<int32, float> PipelineDistances;
-	PipelineDistances.Reserve(256);
-
-	for (int32 CmdIdx = 0; CmdIdx < NumCmds; ++CmdIdx)
+	int32 MeshSortingMethod = CVarMobileMeshSortingMethod.GetValueOnAnyThread();
+	
+	if (MeshSortingMethod == 1) //strict front to back sorting
 	{
-		FVisibleMeshDrawCommand& Cmd = VisibleMeshCommands[CmdIdx];
-
-		float PrimitiveDistance = 0;
-
-		if (Cmd.DrawPrimitiveId < ScenePrimitiveBounds.Num())
+		// compute sort key for each mesh command
+		for (int32 CmdIdx = 0; CmdIdx < NumCmds; ++CmdIdx)
 		{
-			const FPrimitiveBounds& PrimitiveBounds = ScenePrimitiveBounds[Cmd.DrawPrimitiveId];
-			PrimitiveDistance = (PrimitiveBounds.BoxSphereBounds.Origin - ViewOrigin).Size();
-		}
+			FVisibleMeshDrawCommand& Cmd = VisibleMeshCommands[CmdIdx];
+			// Set in MobileBasePass.cpp - GetBasePassStaticSortKey;
+			bool bMasked = Cmd.SortKey.PackedData & 0x1 ? true : false; 
+			bool bBackground = Cmd.SortKey.PackedData & 0x2 ? true : false;
+			float PrimitiveDistance = 0;
+			if (Cmd.DrawPrimitiveId < ScenePrimitiveBounds.Num())
+			{
+				const FPrimitiveBounds& PrimitiveBounds = ScenePrimitiveBounds[Cmd.DrawPrimitiveId];
+				PrimitiveDistance = (PrimitiveBounds.BoxSphereBounds.Origin - ViewOrigin).Size();
+				bBackground|= (PrimitiveBounds.BoxSphereBounds.SphereRadius > HALF_WORLD_MAX / 4.0f);
+			}
 
-		float& PipelineDistance = PipelineDistances.FindOrAdd(Cmd.MeshDrawCommand->CachedPipelineId.GetId());
-		// not sure what could be better: average distance, max or min
-		PipelineDistance = FMath::Max(PipelineDistance, PrimitiveDistance);
+			int32 PipelineId = Cmd.MeshDrawCommand->CachedPipelineId.GetId();
+			int32 StateBucketId = PointerHash(Cmd.MeshDrawCommand->IndexBuffer);
+			Cmd.SortKey.PackedData = GetMobileBasePassSortKey_FrontToBack(bMasked, bBackground, PipelineId, StateBucketId, PrimitiveDistance);
+		}
 	}
-
-	// compute sort key for each mesh command
-	for (int32 CmdIdx = 0; CmdIdx < NumCmds; ++CmdIdx)
+	else // prefer state then distance
 	{
-		FVisibleMeshDrawCommand& Cmd = VisibleMeshCommands[CmdIdx];
-		float PrimitiveDistance = 0;
-		bool bMasked = false;
-		bool bBackground = false;
-
-		if (Cmd.DrawPrimitiveId < ScenePrimitiveBounds.Num())
+		TMap<int32, float> PipelineDistances;
+		PipelineDistances.Reserve(256);
+				
+		// pre-compute distance to a group of meshes that share same PSO
+		for (int32 CmdIdx = 0; CmdIdx < NumCmds; ++CmdIdx)
 		{
-			const FPrimitiveBounds& PrimitiveBounds = ScenePrimitiveBounds[Cmd.DrawPrimitiveId];
-			PrimitiveDistance = (PrimitiveBounds.BoxSphereBounds.Origin - ViewOrigin).Size();
-			bBackground = Cmd.SortKey.PackedData & 0x2 ? true : false; // Set in MobileBasePass.cpp - GetBasePassStaticSortKey
-			bBackground|= (PrimitiveBounds.BoxSphereBounds.SphereRadius > HALF_WORLD_MAX / 4.0f);
-			bMasked = Cmd.SortKey.PackedData & 0x1 ? true : false; // Set in MobileBasePass.cpp - GetBasePassStaticSortKey
-		}
+			FVisibleMeshDrawCommand& Cmd = VisibleMeshCommands[CmdIdx];
+			float PrimitiveDistance = 0;
+			if (Cmd.DrawPrimitiveId < ScenePrimitiveBounds.Num())
+			{
+				const FPrimitiveBounds& PrimitiveBounds = ScenePrimitiveBounds[Cmd.DrawPrimitiveId];
+				PrimitiveDistance = (PrimitiveBounds.BoxSphereBounds.Origin - ViewOrigin).Size();
+			}
 
-		int32 PipelineId = Cmd.MeshDrawCommand->CachedPipelineId.GetId();
-		float PipelineDistance = PipelineDistances.FindRef(PipelineId);
-		// poor man StateID, can't use Cmd.StateBucketId as it is unique for each primitive if platform does not support auto-instancing
-		int32 StateBucketId = PointerHash(Cmd.MeshDrawCommand->IndexBuffer);
+			float& PipelineDistance = PipelineDistances.FindOrAdd(Cmd.MeshDrawCommand->CachedPipelineId.GetId());
+			// not sure what could be better: average distance, max or min
+			PipelineDistance = FMath::Max(PipelineDistance, PrimitiveDistance);
+		}
 
-		Cmd.SortKey.PackedData = GetMobileBasePassSortKey(bMasked, bBackground, PipelineId, StateBucketId, PipelineDistance, PrimitiveDistance);
+		// compute sort key for each mesh command
+		for (int32 CmdIdx = 0; CmdIdx < NumCmds; ++CmdIdx)
+		{
+			FVisibleMeshDrawCommand& Cmd = VisibleMeshCommands[CmdIdx];
+			// Set in MobileBasePass.cpp - GetBasePassStaticSortKey;
+			bool bMasked = Cmd.SortKey.PackedData & 0x1 ? true : false; 
+			bool bBackground = Cmd.SortKey.PackedData & 0x2 ? true : false;
+			float PrimitiveDistance = 0;
+			if (Cmd.DrawPrimitiveId < ScenePrimitiveBounds.Num())
+			{
+				const FPrimitiveBounds& PrimitiveBounds = ScenePrimitiveBounds[Cmd.DrawPrimitiveId];
+				PrimitiveDistance = (PrimitiveBounds.BoxSphereBounds.Origin - ViewOrigin).Size();
+				bBackground|= (PrimitiveBounds.BoxSphereBounds.SphereRadius > HALF_WORLD_MAX / 4.0f);
+			}
+			
+			int32 PipelineId = Cmd.MeshDrawCommand->CachedPipelineId.GetId();
+			float PipelineDistance = PipelineDistances.FindRef(PipelineId);
+			// poor man StateID, can't use Cmd.StateBucketId as it is unique for each primitive if platform does not support auto-instancing
+			int32 StateBucketId = PointerHash(Cmd.MeshDrawCommand->IndexBuffer);
+			Cmd.SortKey.PackedData = GetMobileBasePassSortKey_ByState(bMasked, bBackground, PipelineId, StateBucketId, PipelineDistance, PrimitiveDistance);
+		}
 	}
 }
 
