Duplicate 4978967 & 4980500: Fixes some causes of #jira UE-64700 & #jira UE-65563

Support separate MSAA and resolve targets on Metal MRT for iOS devices - this means the Forward renderer can support MSAA the same way as on macOS. It isn't optimal for iOS GPUs but it at least gets it to run without validation errors - though it doesn't yet render correctly. However now that Metal is using the RenderPassInfo descriptor we can start making changes to fix and improve it.

On iOS Metal is a bit more specific about what LoadAction::NoAction really means - it will reallocate the backing store for the target to uninitialised memory, abandoning whatever was already there. This fixes corruption in the Forward renderer when MSAA is enabled.

diff --git a/Engine/Source/Runtime/Apple/MetalRHI/Private/MetalCommandQueue.h b/Engine/Source/Runtime/Apple/MetalRHI/Private/MetalCommandQueue.h
--- a/Engine/Source/Runtime/Apple/MetalRHI/Private/MetalCommandQueue.h
+++ b/Engine/Source/Runtime/Apple/MetalRHI/Private/MetalCommandQueue.h
@@ -171,7 +171,7 @@ public:
 	* @returns True if RHISupportsSeparateMSAAAndResolveTextures will be true.  
 	* Currently Mac only.
 	*/
-	static inline bool SupportsSeparateMSAAAndResolveTarget() { return PLATFORM_MAC != 0; }
+	static inline bool SupportsSeparateMSAAAndResolveTarget() { return (PLATFORM_MAC != 0 || GMaxRHIFeatureLevel >= ERHIFeatureLevel::SM5); }
 
 #pragma mark - Public Debug Support -
 
diff --git a/Engine/Source/Runtime/Apple/MetalRHI/Private/MetalTexture.cpp b/Engine/Source/Runtime/Apple/MetalRHI/Private/MetalTexture.cpp
--- a/Engine/Source/Runtime/Apple/MetalRHI/Private/MetalTexture.cpp
+++ b/Engine/Source/Runtime/Apple/MetalRHI/Private/MetalTexture.cpp
@@ -877,7 +877,7 @@ FMetalSurface::FMetalSurface(ERHIResourceType ResourceType, EPixelFormat Format,
 		{
 			check(!(Flags & TexCreate_CPUReadback));
 #if PLATFORM_IOS
-			if (FMetalCommandQueue::SupportsFeature(EMetalFeaturesMemoryLessResources) && !(Flags & (TexCreate_ShaderResource|TexCreate_UAV)))
+			if (FMetalCommandQueue::SupportsFeature(EMetalFeaturesMemoryLessResources) && !(Flags & (TexCreate_ShaderResource|TexCreate_UAV)) && (GMaxRHIFeatureLevel < ERHIFeatureLevel::SM5))
 			{
 				Desc.SetStorageMode(mtlpp::StorageMode::Memoryless);
 				Desc.SetResourceOptions(mtlpp::ResourceOptions::StorageModeMemoryless);
@@ -1014,7 +1014,7 @@ FMetalSurface::FMetalSurface(ERHIResourceType ResourceType, EPixelFormat Format,
 			
 			bool bMemoryless = false;
 #if PLATFORM_IOS
-			if (FMetalCommandQueue::SupportsFeature(EMetalFeaturesMemoryLessResources) && GMaxRHIShaderPlatform == SP_METAL)
+			if (FMetalCommandQueue::SupportsFeature(EMetalFeaturesMemoryLessResources) && (GMaxRHIFeatureLevel < ERHIFeatureLevel::SM5))
 			{
 				bMemoryless = true;
 				Desc.SetStorageMode(mtlpp::StorageMode::Memoryless);
diff --git a/Engine/Source/Runtime/RHI/Public/RHI.h b/Engine/Source/Runtime/RHI/Public/RHI.h
--- a/Engine/Source/Runtime/RHI/Public/RHI.h
+++ b/Engine/Source/Runtime/RHI/Public/RHI.h
@@ -141,12 +141,6 @@ inline bool RHISupportsMSAA(EShaderPlatform Platform)
 	return 
 		//@todo-rco: Fix when iOS OpenGL supports MSAA
 		Platform != SP_OPENGL_ES2_IOS
-		// @todo marksatt Metal on macOS 10.12 and earlier (or Intel on any macOS < 10.13.2) don't reliably support our MSAA usage & custom resolve.
-#if PLATFORM_MAC
-		&& IsMetalPlatform(Platform) && (FPlatformMisc::MacOSXVersionCompare(10, 13, 0) >= 0) && (!IsRHIDeviceIntel() || FPlatformMisc::MacOSXVersionCompare(10, 13, 2) >= 0)
-#endif
-		// @todo marksatt iOS Desktop Forward needs more work internally
-		&& Platform != SP_METAL_MRT
 		// @todo optimise MSAA for XboxOne, currently uses significant eRAM.
 		&& Platform != SP_XBOXONE_D3D12;
 }
diff --git a/Engine/Source/Runtime/RHI/Public/RHIDefinitions.h b/Engine/Source/Runtime/RHI/Public/RHIDefinitions.h
--- a/Engine/Source/Runtime/RHI/Public/RHIDefinitions.h
+++ b/Engine/Source/Runtime/RHI/Public/RHIDefinitions.h
@@ -1075,7 +1075,7 @@ inline bool RHINeedsToSwitchVerticalAxis(EShaderPlatform Platform)
 inline bool RHISupportsSeparateMSAAAndResolveTextures(const EShaderPlatform Platform)
 {
 	// Metal mobile devices, Vulkan and Android ES2/3.1 need to handle MSAA and resolve textures internally (unless RHICreateTexture2D was changed to take an optional resolve target)
-	const bool bMobileMetalDevice = (Platform == SP_METAL || Platform == SP_METAL_TVOS || Platform == SP_METAL_MRT || Platform == SP_METAL_MRT_TVOS);
+	const bool bMobileMetalDevice = (Platform == SP_METAL || Platform == SP_METAL_TVOS);
 	return !bMobileMetalDevice && !IsVulkanPlatform(Platform) && !IsAndroidOpenGLESPlatform(Platform);
 }
 
diff --git a/Engine/Source/Runtime/Renderer/Private/DeferredShadingRenderer.cpp b/Engine/Source/Runtime/Renderer/Private/DeferredShadingRenderer.cpp
--- a/Engine/Source/Runtime/Renderer/Private/DeferredShadingRenderer.cpp
+++ b/Engine/Source/Runtime/Renderer/Private/DeferredShadingRenderer.cpp
@@ -1491,7 +1491,7 @@ void FDeferredShadingSceneRenderer::Render(FRHICommandListImmediate& RHICmdList)
 	{
 		// Make sure we have began the renderpass
 		ERenderTargetLoadAction DepthLoadAction = bDepthWasCleared ? ERenderTargetLoadAction::ELoad : ERenderTargetLoadAction::EClear;
-		SceneContext.BeginRenderingGBuffer(RHICmdList, ERenderTargetLoadAction::ENoAction, DepthLoadAction, BasePassDepthStencilAccess, ViewFamily.EngineShowFlags.ShaderComplexity);
+		SceneContext.BeginRenderingGBuffer(RHICmdList, (!IsMetalPlatform(ShaderPlatform) ? ERenderTargetLoadAction::ENoAction : ERenderTargetLoadAction::ELoad), DepthLoadAction, BasePassDepthStencilAccess, ViewFamily.EngineShowFlags.ShaderComplexity);
 	}
 	// Wait for Async SSAO before rendering base pass with forward rendering
 	if (IsForwardShadingEnabled(ShaderPlatform))
@@ -1964,7 +1964,7 @@ void FDeferredShadingSceneRenderer::Render(FRHICommandListImmediate& RHICmdList)
 		SetupSceneTextureUniformParameters(SceneContext, FeatureLevel, ESceneTextureSetupMode::SceneDepth | ESceneTextureSetupMode::GBuffers, SceneTextureParameters);
 		TUniformBufferRef<FSceneTexturesUniformParameters> SceneTextureUniformBuffer = TUniformBufferRef<FSceneTexturesUniformParameters>::CreateUniformBufferImmediate(SceneTextureParameters, UniformBuffer_SingleFrame);
 
-		SceneContext.BeginRenderingSceneColor(RHICmdList);
+        SceneContext.BeginRenderingSceneColor(RHICmdList, (!IsMetalPlatform(ShaderPlatform) ? ESimpleRenderTargetMode::EUninitializedColorExistingDepth : ESimpleRenderTargetMode::EExistingColorAndDepth));
 		for (int32 ViewIndex = 0; ViewIndex < Views.Num(); ++ViewIndex)
 		{
 			const FViewInfo& View = Views[ViewIndex];
