Moving Daniel Wright's HLR memory optimizations from Dev-Rendering.

CL 4871955
Fixed TInlineAllocator::GetAllocatedSize reporting SecondaryData AllocatedSize even when there wasn't any SecondaryData

CL 4871994
MeshDrawCommand memory optimizations
* When GPUScene is enabled, cached FMeshDrawCommands are only stored in FScene::CachedMeshDrawCommandStateBuckets, so we only have one command per state bucket.
* Removed FGraphicsMinimalPipelineStateInitializer, now relying on index into the global table
* Union IndirectArgsBuffer with draw params
* Trimmed debug data
* Removed duplicated space for SRV values, now using a bitmask to indicate SRV or TextureRHI
* Deferred renderer GPUPerfTest cached MeshDrawCommands 28Mb -> 2.1Mb
* Mobile renderer GPUPerfTest cached MeshDrawCommands 30.7Mb -> 20.5Mb

CL 4892243
Resupported r.DoLazyStaticMeshUpdate with mesh draw commands.  Now enabled when in game mode, no longer requires cooked build.
 * GPUPerfTest Total cached MeshDrawCommands 2.324Mb -> 1.691Mb and Primitive StaticMeshCommandInfos 3529.2Kb -> 563.7Kb

CL 4925071
Allocate MeshDrawCommands at lowest free index instead of last freed index
Shrink MeshDrawCommands periodically when r.DoLazyStaticMeshUpdate is enabled
Measured in replay ES31: Cached MeshDrawCommands 5.236Mb -> 2.216Mb

CL 4937280
Removed r.DeferUniformBufferUpdatesUntilVisible since it is no longer supported with MeshDrawCommands, which cache the uniform buffer pointer at AddToScene time.  Removed dead code referencing it.

CL 4937591
Fixed FMeshDrawCommandStateCache always resetting resources at index 0.  Found by Rolando.

CL 4951781
Attempt to silence static analysis warning

#rnx

diff --git a/Engine/Plugins/FX/Niagara/Source/Niagara/Private/NiagaraRendererMeshes.cpp b/Engine/Plugins/FX/Niagara/Source/Niagara/Private/NiagaraRendererMeshes.cpp
--- a/Engine/Plugins/FX/Niagara/Source/Niagara/Private/NiagaraRendererMeshes.cpp
+++ b/Engine/Plugins/FX/Niagara/Source/Niagara/Private/NiagaraRendererMeshes.cpp
@@ -299,10 +299,6 @@ void NiagaraRendererMeshes::GetDynamicMeshElements(const TArray<const FSceneView
 					BatchElement.MinVertexIndex = 0;
 					BatchElement.MaxVertexIndex = 0;
 					BatchElement.NumInstances = NumInstances;
-					if (DynamicDataMesh->DataSet->GetSimTarget() == ENiagaraSimTarget::GPUComputeSim)
-					{
-						BatchElement.IndirectArgsBuffer = DynamicDataMesh->DataSet->GetCurDataSetIndices().Buffer;
-					}
 
 					if (bIsWireframe)
 					{
@@ -334,10 +330,19 @@ void NiagaraRendererMeshes::GetDynamicMeshElements(const TArray<const FSceneView
 						BatchElement.NumPrimitives = Section.NumTriangles;
 					}
 
+					if (DynamicDataMesh->DataSet->GetSimTarget() == ENiagaraSimTarget::GPUComputeSim)
+					{
+						BatchElement.NumPrimitives = 0;
+						BatchElement.IndirectArgsBuffer = DynamicDataMesh->DataSet->GetCurDataSetIndices().Buffer;
+					}
+					else
+					{
+						check(BatchElement.NumPrimitives > 0);
+					}
+
 					Mesh.bCanApplyViewModeOverrides = true;
 					Mesh.bUseWireframeSelectionColoring = SceneProxy->IsSelected();
 
-					check(BatchElement.NumPrimitives > 0);
 					Collector.AddMesh(ViewIndex, Mesh);
 				}
 			}
diff --git a/Engine/Plugins/FX/Niagara/Source/Niagara/Private/NiagaraRendererSprites.cpp b/Engine/Plugins/FX/Niagara/Source/Niagara/Private/NiagaraRendererSprites.cpp
--- a/Engine/Plugins/FX/Niagara/Source/Niagara/Private/NiagaraRendererSprites.cpp
+++ b/Engine/Plugins/FX/Niagara/Source/Niagara/Private/NiagaraRendererSprites.cpp
@@ -332,6 +332,7 @@ void NiagaraRendererSprites::GetDynamicMeshElements(const TArray<const FSceneVie
 				if (DynamicDataSprites->DataSet->GetSimTarget() == ENiagaraSimTarget::GPUComputeSim)
 				{
 					MeshElement.IndirectArgsBuffer = DynamicDataSprites->DataSet->GetCurDataSetIndices().Buffer;
+					MeshElement.NumPrimitives = 0;
 				}
 				
 				Collector.AddMesh(ViewIndex, MeshBatch);				
diff --git a/Engine/Source/Runtime/Core/Public/Containers/BitArray.h b/Engine/Source/Runtime/Core/Public/Containers/BitArray.h
--- a/Engine/Source/Runtime/Core/Public/Containers/BitArray.h
+++ b/Engine/Source/Runtime/Core/Public/Containers/BitArray.h
@@ -642,13 +642,13 @@ public:
 	 * Finds the first zero bit in the array, sets it to true, and returns the bit index.
 	 * If there is none, INDEX_NONE is returned.
 	 */
-	int32 FindAndSetFirstZeroBit()
+	int32 FindAndSetFirstZeroBit(int32 ConservativeStartIndex = 0)
 	{
 		// Iterate over the array until we see a word with a zero bit.
 		uint32* RESTRICT DwordArray = GetData();
 		const int32 LocalNumBits = NumBits;
 		const int32 DwordCount = FMath::DivideAndRoundUp(LocalNumBits, NumBitsPerDWORD);
-		int32 DwordIndex = 0;
+		int32 DwordIndex = FMath::DivideAndRoundDown(ConservativeStartIndex, NumBitsPerDWORD);
 		while (DwordIndex < DwordCount && DwordArray[DwordIndex] == (uint32)-1)
 		{
 			++DwordIndex;
diff --git a/Engine/Source/Runtime/Core/Public/Containers/ContainerAllocationPolicies.h b/Engine/Source/Runtime/Core/Public/Containers/ContainerAllocationPolicies.h
--- a/Engine/Source/Runtime/Core/Public/Containers/ContainerAllocationPolicies.h
+++ b/Engine/Source/Runtime/Core/Public/Containers/ContainerAllocationPolicies.h
@@ -545,7 +545,11 @@ public:
 
 		SIZE_T GetAllocatedSize(int32 NumAllocatedElements, SIZE_T NumBytesPerElement) const
 		{
-			return SecondaryData.GetAllocatedSize(NumAllocatedElements, NumBytesPerElement);
+			if (NumAllocatedElements > NumInlineElements)
+			{
+				return SecondaryData.GetAllocatedSize(NumAllocatedElements, NumBytesPerElement);
+			}
+			return 0;
 		}
 
 		bool HasAllocation()
diff --git a/Engine/Source/Runtime/Core/Public/Containers/Set.h b/Engine/Source/Runtime/Core/Public/Containers/Set.h
--- a/Engine/Source/Runtime/Core/Public/Containers/Set.h
+++ b/Engine/Source/Runtime/Core/Public/Containers/Set.h
@@ -454,6 +454,11 @@ public:
 		return Elements.Num();
 	}
 
+	FORCEINLINE int32 GetMaxIndex() const
+	{
+		return Elements.GetMaxIndex();
+	}
+
 	/**
 	 * Checks whether an element id is valid.
 	 * @param Id - The element id to check.
diff --git a/Engine/Source/Runtime/Core/Public/Containers/SparseArray.h b/Engine/Source/Runtime/Core/Public/Containers/SparseArray.h
--- a/Engine/Source/Runtime/Core/Public/Containers/SparseArray.h
+++ b/Engine/Source/Runtime/Core/Public/Containers/SparseArray.h
@@ -141,6 +141,59 @@ public:
 		return Allocation.Index;
 	}
 
+	FSparseArrayAllocationInfo AddUninitializedAtLowestFreeIndex(int32& LowestFreeIndexSearchStart)
+	{
+		int32 Index;
+		if(NumFreeIndices)
+		{
+			Index = AllocationFlags.FindAndSetFirstZeroBit(LowestFreeIndexSearchStart);
+			LowestFreeIndexSearchStart = Index + 1;
+
+			auto& IndexData = GetData(Index);
+
+			// Update FirstFreeIndex
+			if (FirstFreeIndex == Index)
+			{
+				FirstFreeIndex = IndexData.NextFreeIndex;
+			}
+
+			// Link our next and prev free nodes together
+			if (IndexData.NextFreeIndex >= 0)
+			{
+				GetData(IndexData.NextFreeIndex).PrevFreeIndex = IndexData.PrevFreeIndex;
+			}
+
+			if (IndexData.PrevFreeIndex >= 0)
+			{
+				GetData(IndexData.PrevFreeIndex).NextFreeIndex = IndexData.NextFreeIndex;
+			}
+
+			--NumFreeIndices;
+		}
+		else
+		{
+			// Add a new element.
+			Index = Data.AddUninitialized(1);
+			AllocationFlags.Add(true);
+		}
+
+		FSparseArrayAllocationInfo Result;
+		Result.Index = Index;
+		Result.Pointer = &GetData(Result.Index).ElementData;
+		return Result;
+	}
+
+	/** 
+	 * Add an element at the lowest free index, instead of the last freed index. 
+	 * This requires a search which can be accelerated with LowestFreeIndexSearchStart.
+	 */
+	int32 AddAtLowestFreeIndex(const ElementType& Element, int32& LowestFreeIndexSearchStart)
+	{
+		FSparseArrayAllocationInfo Allocation = AddUninitializedAtLowestFreeIndex(LowestFreeIndexSearchStart);
+		new(Allocation) ElementType(Element);
+		return Allocation.Index;
+	}
+
 	/**
 	 * Allocates space for an element in the array at a given index.  The element is not initialized, and you must use the corresponding placement new operator
 	 * to construct the element in the allocated memory.
diff --git a/Engine/Source/Runtime/Engine/Private/PrimitiveSceneProxy.cpp b/Engine/Source/Runtime/Engine/Private/PrimitiveSceneProxy.cpp
--- a/Engine/Source/Runtime/Engine/Private/PrimitiveSceneProxy.cpp
+++ b/Engine/Source/Runtime/Engine/Private/PrimitiveSceneProxy.cpp
@@ -232,41 +232,6 @@ FPrimitiveViewRelevance FPrimitiveSceneProxy::GetViewRelevance(const FSceneView*
 	return FPrimitiveViewRelevance();
 }
 
-static TAutoConsoleVariable<int32> CVarDeferUniformBufferUpdatesUntilVisible(
-	TEXT("r.DeferUniformBufferUpdatesUntilVisible"),
-	0,
-	TEXT("If > 0, then don't update the primitive uniform buffer until it is visible. Incompatible with the Mesh Draw Command pipeline."));
-
-void FPrimitiveSceneProxy::UpdateUniformBufferMaybeLazy()
-{
-	//@todo MeshCommandPipeline r.DeferUniformBufferUpdatesUntilVisible isn't currently supported.
-	/*
-	if (PrimitiveSceneInfo && CVarDeferUniformBufferUpdatesUntilVisible.GetValueOnAnyThread() > 0)
-	{
-		PrimitiveSceneInfo->SetNeedsUniformBufferUpdate(true);
-	}
-	else
-	{
-		UpdateUniformBuffer();
-	}
-	*/
-
-	UpdateUniformBuffer();
-}
-
-bool FPrimitiveSceneProxy::NeedsUniformBufferUpdate() const
-{
-	//@todo MeshCommandPipeline r.DeferUniformBufferUpdatesUntilVisible isn't currently supported.
-	/*
-	if (PrimitiveSceneInfo && CVarDeferUniformBufferUpdatesUntilVisible.GetValueOnAnyThread() > 0)
-	{
-		return PrimitiveSceneInfo->NeedsUniformBufferUpdate();
-	}
-	*/
-
-	return false;
-}
-
 void FPrimitiveSceneProxy::UpdateUniformBuffer()
 {
 	QUICK_SCOPE_CYCLE_COUNTER(STAT_FPrimitiveSceneProxy_UpdateUniformBuffer);
@@ -328,7 +293,7 @@ void FPrimitiveSceneProxy::SetTransform(const FMatrix& InLocalToWorld, const FBo
 	LocalBounds = InLocalBounds;
 	ActorPosition = InActorPosition;
 	
-	UpdateUniformBufferMaybeLazy();
+	UpdateUniformBuffer();
 	
 	// Notify the proxy's implementation of the change.
 	OnTransformChanged();
diff --git a/Engine/Source/Runtime/Engine/Public/MeshBatch.h b/Engine/Source/Runtime/Engine/Public/MeshBatch.h
--- a/Engine/Source/Runtime/Engine/Public/MeshBatch.h
+++ b/Engine/Source/Runtime/Engine/Public/MeshBatch.h
@@ -70,6 +70,7 @@ struct FMeshBatchElement
 	const void* UserData;
 
 	uint32 FirstIndex;
+	/** When 0, IndirectArgsBuffer will be used. */
 	uint32 NumPrimitives;
 
 	/** Number of instances to draw.  If InstanceRuns is valid, this is actually the number of runs in InstanceRuns. */
diff --git a/Engine/Source/Runtime/Engine/Public/PrimitiveSceneProxy.h b/Engine/Source/Runtime/Engine/Public/PrimitiveSceneProxy.h
--- a/Engine/Source/Runtime/Engine/Public/PrimitiveSceneProxy.h
+++ b/Engine/Source/Runtime/Engine/Public/PrimitiveSceneProxy.h
@@ -606,10 +606,6 @@ public:
 	 * Updates the primitive proxy's uniform buffer.
 	 */
 	ENGINE_API void UpdateUniformBuffer();
-	/**
-	 * Updates the primitive proxy's uniform buffer.
-	 */
-	ENGINE_API bool NeedsUniformBufferUpdate() const;
 
 #if !UE_BUILD_SHIPPING
 
@@ -1045,11 +1041,6 @@ private:
 
 	ENGINE_API bool WouldSetTransformBeRedundant(const FMatrix& InLocalToWorld, const FBoxSphereBounds& InBounds, const FBoxSphereBounds& InLocalBounds, FVector InActorPosition);
 
-	/**
-	 * Either updates the uniform buffer or defers it until it becomes visible depending on a cvar
-	 */
-	ENGINE_API void UpdateUniformBufferMaybeLazy();
-
 	/** Updates the hidden editor view visibility map on the render thread */
 	void SetHiddenEdViews_RenderThread( uint64 InHiddenEditorViews );
 
diff --git a/Engine/Source/Runtime/Landscape/Private/LandscapeRender.cpp b/Engine/Source/Runtime/Landscape/Private/LandscapeRender.cpp
--- a/Engine/Source/Runtime/Landscape/Private/LandscapeRender.cpp
+++ b/Engine/Source/Runtime/Landscape/Private/LandscapeRender.cpp
@@ -950,10 +950,6 @@ void FLandscapeComponentSceneProxy::CreateRenderThreadResources()
 		BatchElementParams->SubY = -1;
 		BatchElementParams->CurrentLOD = 0;
 		GrassBatchElement->UserData = BatchElementParams;
-		if (NeedsUniformBufferUpdate())
-		{
-			UpdateUniformBuffer();
-		}
 		GrassBatchElement->PrimitiveUniformBuffer = GetUniformBuffer();
 		GrassBatchElement->IndexBuffer = SharedBuffers->GrassIndexBuffer;
 		GrassBatchElement->NumPrimitives = FMath::Square(NumSubsections) * FMath::Square(SubsectionSizeVerts);
diff --git a/Engine/Source/Runtime/RenderCore/Public/VertexFactory.h b/Engine/Source/Runtime/RenderCore/Public/VertexFactory.h
--- a/Engine/Source/Runtime/RenderCore/Public/VertexFactory.h
+++ b/Engine/Source/Runtime/RenderCore/Public/VertexFactory.h
@@ -67,7 +67,7 @@ struct FVertexInputStream
  * This is tweaked so that the bindings for FLocalVertexFactory fit into the inline storage.
  * Overflow of the inline storage will cause a heap allocation per draw (and corresponding cache miss on traversal)
  */
-typedef TArray<FVertexInputStream, TInlineAllocator<2>> FVertexInputStreamArray;
+typedef TArray<FVertexInputStream, TInlineAllocator<4>> FVertexInputStreamArray;
 
 enum class EVertexStreamUsage : uint8
 {
diff --git a/Engine/Source/Runtime/Renderer/Private/MeshDrawCommands.cpp b/Engine/Source/Runtime/Renderer/Private/MeshDrawCommands.cpp
--- a/Engine/Source/Runtime/Renderer/Private/MeshDrawCommands.cpp
+++ b/Engine/Source/Runtime/Renderer/Private/MeshDrawCommands.cpp
@@ -669,16 +669,17 @@ void ApplyViewOverridesToMeshDrawCommands(
 				NewMeshCommand = MeshCommand;
 
 				const ERasterizerCullMode LocalCullMode = bRenderSceneTwoSided ? CM_None : bReverseCulling ? FMeshPassProcessor::InverseCullMode(VisibleMeshDrawCommand.MeshCullMode) : VisibleMeshDrawCommand.MeshCullMode;
-				NewMeshCommand.PipelineState.RasterizerState = GetStaticRasterizerState<true>(VisibleMeshDrawCommand.MeshFillMode, LocalCullMode);
+				FGraphicsMinimalPipelineStateInitializer PipelineState = MeshCommand.CachedPipelineId.GetPipelineState();
+				PipelineState.RasterizerState = GetStaticRasterizerState<true>(VisibleMeshDrawCommand.MeshFillMode, LocalCullMode);
 
 				if (BasePassDepthStencilAccess != DefaultBasePassDepthStencilAccess && PassType == EMeshPass::BasePass)
 				{
 					FMeshPassProcessorRenderState PassDrawRenderState;
 					SetupBasePassState(BasePassDepthStencilAccess, false, PassDrawRenderState);
-					NewMeshCommand.PipelineState.DepthStencilState = PassDrawRenderState.GetDepthStencilState();
+					PipelineState.DepthStencilState = PassDrawRenderState.GetDepthStencilState();
 				}
 
-				NewMeshCommand.Finalize(true);
+				NewMeshCommand.Finalize(PipelineState, nullptr, true);
 
 				FVisibleMeshDrawCommand NewVisibleMeshDrawCommand;
 
diff --git a/Engine/Source/Runtime/Renderer/Private/MeshPassProcessor.cpp b/Engine/Source/Runtime/Renderer/Private/MeshPassProcessor.cpp
--- a/Engine/Source/Runtime/Renderer/Private/MeshPassProcessor.cpp
+++ b/Engine/Source/Runtime/Renderer/Private/MeshPassProcessor.cpp
@@ -35,13 +35,13 @@ enum { MAX_SAMPLERS_PER_SHADER_STAGE = 32 };
 class FShaderBindingState
 {
 public:
-	int32 MaxSRVUsed = 0;
+	int32 MaxSRVUsed = -1;
 	FShaderResourceViewRHIParamRef SRVs[MAX_SRVs_PER_SHADER_STAGE] = {};
-	int32 MaxUniformBufferUsed = 0;
+	int32 MaxUniformBufferUsed = -1;
 	FUniformBufferRHIParamRef UniformBuffers[MAX_UNIFORM_BUFFERS_PER_SHADER_STAGE] = {};
-	int32 MaxTextureUsed = 0;
+	int32 MaxTextureUsed = -1;
 	FTextureRHIParamRef Textures[MAX_SRVs_PER_SHADER_STAGE] = {};
-	int32 MaxSamplerUsed = 0;
+	int32 MaxSamplerUsed = -1;
 	FSamplerStateRHIParamRef Samplers[MAX_SAMPLERS_PER_SHADER_STAGE] = {};
 };
 
@@ -65,16 +65,16 @@ public:
 		return (const FSamplerStateRHIParamRef*)SamplerDataStart;
 	}
 
-	inline const FShaderResourceViewRHIParamRef* GetSRVStart() const
+	inline const FRHIResource** GetSRVStart() const
 	{
 		const uint8* SRVDataStart = Data + GetSRVOffset();
-		return (const FShaderResourceViewRHIParamRef*)SRVDataStart;
+		return (const FRHIResource**)SRVDataStart;
 	}
 
-	inline const FTextureRHIParamRef* GetTextureStart() const
+	inline const uint8* GetSRVTypeStart() const
 	{
-		const uint8* TextureDataStart = Data + GetTextureOffset();
-		return (const FTextureRHIParamRef*)TextureDataStart;
+		const uint8* SRVTypeDataStart = Data + GetSRVTypeOffset();
+		return SRVTypeDataStart;
 	}
 
 	inline const uint8* GetLooseDataStart() const
@@ -130,37 +130,40 @@ void FMeshDrawShaderBindings::SetShaderBindings(
 		}
 	}
 
-	const FShaderResourceViewRHIParamRef* RESTRICT SRVBindings = SingleShaderBindings.GetSRVStart();
+	const uint8* RESTRICT SRVType = SingleShaderBindings.GetSRVTypeStart();
+	const FRHIResource** RESTRICT SRVBindings = SingleShaderBindings.GetSRVStart();
 	const FShaderParameterInfo* RESTRICT SRVParameters = SingleShaderBindings.ParameterMapInfo.SRVs.GetData();
-	const int32 NumSRVs = SingleShaderBindings.ParameterMapInfo.SRVs.Num();
+	const uint32 NumSRVs = SingleShaderBindings.ParameterMapInfo.SRVs.Num();
 
-	for (int32 SRVIndex = 0; SRVIndex < NumSRVs; SRVIndex++)
+	for (uint32 SRVIndex = 0; SRVIndex < NumSRVs; SRVIndex++)
 	{
 		FShaderParameterInfo Parameter = SRVParameters[SRVIndex];
 		checkSlow(Parameter.BaseIndex < ARRAY_COUNT(ShaderBindingState.SRVs));
-		FShaderResourceViewRHIParamRef SRV = SRVBindings[SRVIndex];
-
-		if (SRV != ShaderBindingState.SRVs[Parameter.BaseIndex])
-		{
-			RHICmdList.SetShaderResourceViewParameter(Shader, Parameter.BaseIndex, SRV);
-			ShaderBindingState.SRVs[Parameter.BaseIndex] = SRV;
-			ShaderBindingState.MaxSRVUsed = FMath::Max((int32)Parameter.BaseIndex, ShaderBindingState.MaxSRVUsed);
-		}
-	}
 
-	const FTextureRHIParamRef* RESTRICT TextureBindings = SingleShaderBindings.GetTextureStart();
+		uint32 TypeByteIndex = SRVIndex / 8;
+		uint32 TypeBitIndex = SRVIndex - TypeByteIndex;
 
-	for (int32 TextureIndex = 0; TextureIndex < NumSRVs; TextureIndex++)
-	{
-		FShaderParameterInfo Parameter = SRVParameters[TextureIndex];
-		checkSlow(Parameter.BaseIndex < ARRAY_COUNT(ShaderBindingState.Textures));
-		FTextureRHIParamRef Texture = TextureBindings[TextureIndex];
+		if (SRVType[TypeByteIndex] & (1 << TypeBitIndex))
+		{
+			FShaderResourceViewRHIParamRef SRV = (FShaderResourceViewRHIParamRef)SRVBindings[SRVIndex];
 
-		if (Texture != ShaderBindingState.Textures[Parameter.BaseIndex])
+			if (SRV != ShaderBindingState.SRVs[Parameter.BaseIndex])
+			{
+				RHICmdList.SetShaderResourceViewParameter(Shader, Parameter.BaseIndex, SRV);
+				ShaderBindingState.SRVs[Parameter.BaseIndex] = SRV;
+				ShaderBindingState.MaxSRVUsed = FMath::Max((int32)Parameter.BaseIndex, ShaderBindingState.MaxSRVUsed);
+			}
+		}
+		else
 		{
-			RHICmdList.SetShaderTexture(Shader, Parameter.BaseIndex, Texture);
-			ShaderBindingState.Textures[Parameter.BaseIndex] = Texture;
-			ShaderBindingState.MaxTextureUsed = FMath::Max((int32)Parameter.BaseIndex, ShaderBindingState.MaxTextureUsed);
+			FTextureRHIParamRef Texture = (FTextureRHIParamRef)SRVBindings[SRVIndex];
+
+			if (Texture != ShaderBindingState.Textures[Parameter.BaseIndex])
+			{
+				RHICmdList.SetShaderTexture(Shader, Parameter.BaseIndex, Texture);
+				ShaderBindingState.Textures[Parameter.BaseIndex] = Texture;
+				ShaderBindingState.MaxTextureUsed = FMath::Max((int32)Parameter.BaseIndex, ShaderBindingState.MaxTextureUsed);
+			}
 		}
 	}
 
@@ -213,34 +216,27 @@ void FMeshDrawShaderBindings::SetShaderBindings(
 		RHICmdList.SetShaderSampler(Shader, Parameter.BaseIndex, Sampler);
 	}
 
+	const uint8* RESTRICT SRVType = SingleShaderBindings.GetSRVTypeStart();
+	const FRHIResource** RESTRICT SRVBindings = SingleShaderBindings.GetSRVStart();
 	const FShaderParameterInfo* RESTRICT SRVParameters = SingleShaderBindings.ParameterMapInfo.SRVs.GetData();
-	const int32 NumSRVs = SingleShaderBindings.ParameterMapInfo.SRVs.Num();
-
-	const FShaderResourceViewRHIParamRef* RESTRICT SRVBindings = SingleShaderBindings.GetSRVStart();
-	const FTextureRHIParamRef* RESTRICT TextureBindings = SingleShaderBindings.GetTextureStart();
+	const uint32 NumSRVs = SingleShaderBindings.ParameterMapInfo.SRVs.Num();
 
-	for (int32 SRVIndex = 0; SRVIndex < NumSRVs; SRVIndex++)
+	for (uint32 SRVIndex = 0; SRVIndex < NumSRVs; SRVIndex++)
 	{
 		FShaderParameterInfo Parameter = SRVParameters[SRVIndex];
 
-		FShaderResourceViewRHIParamRef SRV = SRVBindings[SRVIndex];
-		FTextureRHIParamRef Texture = TextureBindings[SRVIndex];
+		uint32 TypeByteIndex = SRVIndex / 8;
+		uint32 TypeBitIndex = SRVIndex - TypeByteIndex;
 
-		// Check for both being non-null
-		checkf(!(SRV != nullptr && Texture != nullptr), TEXT("SRV and Texture cannot be set to the same slot"));
-
-		if (SRV != nullptr)
+		if (SRVType[TypeByteIndex] & (1 << TypeBitIndex))
 		{
+			FShaderResourceViewRHIParamRef SRV = (FShaderResourceViewRHIParamRef)SRVBindings[SRVIndex];
 			RHICmdList.SetShaderResourceViewParameter(Shader, Parameter.BaseIndex, SRV);
 		}
-		else if (Texture != nullptr)
-		{
-			RHICmdList.SetShaderTexture(Shader, Parameter.BaseIndex, Texture);
-		}
 		else
 		{
-			// Both are null, empty slot.
-			RHICmdList.SetShaderResourceViewParameter(Shader, Parameter.BaseIndex, nullptr);
+			FTextureRHIParamRef Texture = (FTextureRHIParamRef)SRVBindings[SRVIndex];
+			RHICmdList.SetShaderTexture(Shader, Parameter.BaseIndex, Texture);
 		}
 	}
 	
@@ -322,7 +318,6 @@ class FMeshDrawCommandStateCache
 public:
 
 	int32 PipelineId;
-	FGraphicsMinimalPipelineStateInitializer PipelineState;
 	uint32 StencilRef;
 	FShaderBindingState ShaderBindings[SF_NumStandardFrequencies];
 	FVertexInputStream VertexStreams[MaxVertexElementCount];
@@ -334,9 +329,8 @@ public:
 		StencilRef = -1;
 	}
 
-	inline void SetPipelineState(const FGraphicsMinimalPipelineStateInitializer& NewPipelineState, int32 NewPipelineId)
+	inline void SetPipelineState(int32 NewPipelineId)
 	{
-		PipelineState = NewPipelineState;
 		PipelineId = NewPipelineId;
 		StencilRef = -1;
 
@@ -356,28 +350,28 @@ public:
 				ShaderBinding.SRVs[SlotIndex] = nullptr;
 			}
 
-			ShaderBinding.MaxSRVUsed = 0;
+			ShaderBinding.MaxSRVUsed = -1;
 
 			for (int32 SlotIndex = 0; SlotIndex <= ShaderBinding.MaxUniformBufferUsed; SlotIndex++)
 			{
 				ShaderBinding.UniformBuffers[SlotIndex] = nullptr;
 			}
 
-			ShaderBinding.MaxUniformBufferUsed = 0;
+			ShaderBinding.MaxUniformBufferUsed = -1;
 			
 			for (int32 SlotIndex = 0; SlotIndex <= ShaderBinding.MaxTextureUsed; SlotIndex++)
 			{
 				ShaderBinding.Textures[SlotIndex] = nullptr;
 			}
 
-			ShaderBinding.MaxTextureUsed = 0;
+			ShaderBinding.MaxTextureUsed = -1;
 
 			for (int32 SlotIndex = 0; SlotIndex <= ShaderBinding.MaxSamplerUsed; SlotIndex++)
 			{
 				ShaderBinding.Samplers[SlotIndex] = nullptr;
 			}
 
-			ShaderBinding.MaxSamplerUsed = 0;
+			ShaderBinding.MaxSamplerUsed = -1;
 		}
 	}
 };
@@ -478,16 +472,21 @@ void FMeshDrawShaderBindings::Initialize(FMeshProcessorShaders Shaders)
 	}
 }
 
-void FMeshDrawShaderBindings::Finalize(const FMeshDrawCommandDebugData& DebugData)
+void FMeshDrawShaderBindings::Finalize(const FMeshProcessorShaders* ShadersForDebugging)
 {
 #if VALIDATE_MESH_COMMAND_BINDINGS
+	if (!ShadersForDebugging)
+	{
+		return;
+	}
+
 	const uint8* ShaderBindingDataPtr = GetData();
 
 	for (int32 ShaderBindingsIndex = 0; ShaderBindingsIndex < ShaderLayouts.Num(); ShaderBindingsIndex++)
 	{
 		const FMeshDrawShaderBindingsLayout& ShaderLayout = ShaderLayouts[ShaderBindingsIndex];
 
-		FMeshMaterialShader* Shader = DebugData.Shaders.GetShader(ShaderLayout.Frequency);
+		FMeshMaterialShader* Shader = ShadersForDebugging->GetShader(ShaderLayout.Frequency);
 		check(Shader);
 
 		FReadOnlyMeshDrawSingleShaderBindings SingleShaderBindings(ShaderLayout, ShaderBindingDataPtr);
@@ -535,18 +534,36 @@ void FMeshDrawShaderBindings::Finalize(const FMeshDrawCommandDebugData& DebugDat
 				ParameterInfo.BaseIndex);
 		}
 
-		const FShaderResourceViewRHIParamRef* SRVBindings = SingleShaderBindings.GetSRVStart();
-		const FTextureRHIParamRef* TextureBindings = SingleShaderBindings.GetTextureStart();
+		const uint8* RESTRICT SRVType = SingleShaderBindings.GetSRVTypeStart();
+		const FRHIResource** RESTRICT SRVBindings = SingleShaderBindings.GetSRVStart();
+		const FShaderParameterInfo* RESTRICT SRVParameters = SingleShaderBindings.ParameterMapInfo.SRVs.GetData();
+		const uint32 NumSRVs = SingleShaderBindings.ParameterMapInfo.SRVs.Num();
 
-		for (int32 BindingIndex = 0; BindingIndex < ShaderLayout.ParameterMapInfo.SRVs.Num(); BindingIndex++)
+		for (uint32 SRVIndex = 0; SRVIndex < NumSRVs; SRVIndex++)
 		{
-			FShaderParameterInfo ParameterInfo = ShaderLayout.ParameterMapInfo.SRVs[BindingIndex];
-			FShaderResourceViewRHIParamRef SRVValue = SRVBindings[BindingIndex];
-			FTextureRHIParamRef TextureValue = TextureBindings[BindingIndex];
-			ensureMsgf(SRVValue || TextureValue, TEXT("Shader %s with vertex factory %s never set texture or SRV at BaseIndex %u.  This can cause GPU hangs, depending on how the shader uses it."), 
-				Shader->GetType()->GetName(), 
-				Shader->GetVertexFactoryType()->GetName(),
-				ParameterInfo.BaseIndex);
+			FShaderParameterInfo Parameter = SRVParameters[SRVIndex];
+
+			uint32 TypeByteIndex = SRVIndex / 8;
+			uint32 TypeBitIndex = SRVIndex - TypeByteIndex;
+
+			if (SRVType[TypeByteIndex] & (1 << TypeBitIndex))
+			{
+				FShaderResourceViewRHIParamRef SRV = (FShaderResourceViewRHIParamRef)SRVBindings[SRVIndex];
+
+				ensureMsgf(SRV, TEXT("Shader %s with vertex factory %s never set SRV at BaseIndex %u.  This can cause GPU hangs, depending on how the shader uses it."), 
+					Shader->GetType()->GetName(), 
+					Shader->GetVertexFactoryType()->GetName(),
+					Parameter.BaseIndex);
+			}
+			else
+			{
+				FTextureRHIParamRef Texture = (FTextureRHIParamRef)SRVBindings[SRVIndex];
+
+				ensureMsgf(Texture, TEXT("Shader %s with vertex factory %s never set texture at BaseIndex %u.  This can cause GPU hangs, depending on how the shader uses it."), 
+					Shader->GetType()->GetName(), 
+					Shader->GetVertexFactoryType()->GetName(),
+					Parameter.BaseIndex);
+			}
 		}
 
 		ShaderBindingDataPtr += ShaderLayout.GetDataSizeBytes();
@@ -585,7 +602,7 @@ void FMeshDrawShaderBindings::CopyFrom(const FMeshDrawShaderBindings& Other)
 #endif
 }
 
-void FMeshDrawCommand::SetShaders(FVertexDeclarationRHIParamRef VertexDeclaration, const FMeshProcessorShaders& Shaders)
+void FMeshDrawCommand::SetShaders(FVertexDeclarationRHIParamRef VertexDeclaration, const FMeshProcessorShaders& Shaders, FGraphicsMinimalPipelineStateInitializer& PipelineState)
 {
 	PipelineState.BoundShaderState = FBoundShaderStateInput(
 		VertexDeclaration,
@@ -608,7 +625,12 @@ void FMeshDrawCommand::SetRayTracingShaders(const FMeshProcessorShaders& Shaders
 }
 #endif // RHI_RAYTRACING
 
-void FMeshDrawCommand::SetDrawParametersAndFinalize(const FMeshBatch& MeshBatch, int32 BatchElementIndex, bool bDoSetupPsoStateForRasterization)
+void FMeshDrawCommand::SetDrawParametersAndFinalize(
+	const FMeshBatch& MeshBatch, 
+	int32 BatchElementIndex,
+	const FGraphicsMinimalPipelineStateInitializer& PipelineState, 
+	const FMeshProcessorShaders* ShadersForDebugging,
+	bool bDoSetupPsoStateForRasterization)
 {
 	const FMeshBatchElement& BatchElement = MeshBatch.Elements[BatchElementIndex];
 
@@ -618,14 +640,23 @@ void FMeshDrawCommand::SetDrawParametersAndFinalize(const FMeshBatch& MeshBatch,
 	FirstIndex = BatchElement.FirstIndex;
 	NumPrimitives = BatchElement.NumPrimitives;
 	NumInstances = BatchElement.NumInstances;
-	BaseVertexIndex = BatchElement.BaseVertexIndex;
-	NumVertices = BatchElement.MaxVertexIndex - BatchElement.MinVertexIndex + 1;
-	IndirectArgsBuffer = BatchElement.IndirectArgsBuffer;
+
+	if (NumPrimitives > 0)
+	{
+		VertexParams.BaseVertexIndex = BatchElement.BaseVertexIndex;
+		VertexParams.NumVertices = BatchElement.MaxVertexIndex - BatchElement.MinVertexIndex + 1;
+		checkf(!BatchElement.IndirectArgsBuffer, TEXT("FMeshBatchElement::NumPrimitives must be set to 0 when a IndirectArgsBuffer is used"));
+	}
+	else
+	{
+		checkf(BatchElement.IndirectArgsBuffer, TEXT("It is only valid to set BatchElement.NumPrimitives == 0 when a IndirectArgsBuffer is used"));
+		IndirectArgsBuffer = BatchElement.IndirectArgsBuffer;
+	}
 
 	int32 SegmentIndex = MeshBatch.SegmentIndex + BatchElementIndex;
 	RayTracedSegmentIndex = (SegmentIndex < UINT8_MAX) ? uint8(MeshBatch.SegmentIndex + BatchElementIndex) : UINT8_MAX;
 
-	Finalize(bDoSetupPsoStateForRasterization);
+	Finalize(PipelineState, ShadersForDebugging, bDoSetupPsoStateForRasterization);
 }
 
 void FMeshDrawShaderBindings::SetOnCommandList(FRHICommandList& RHICmdList, FBoundShaderStateInput Shaders, FShaderBindingState* StateCacheShaderBindings) const
@@ -770,12 +801,14 @@ void FMeshDrawCommand::SubmitDraw(
 	}
 #endif
 
+	const FGraphicsMinimalPipelineStateInitializer& MeshPipelineState = MeshDrawCommand.CachedPipelineId.GetPipelineState();
+
 	if (MeshDrawCommand.CachedPipelineId.GetId() != StateCache.PipelineId)
 	{
-		FGraphicsPipelineStateInitializer GraphicsPSOInit = MeshDrawCommand.PipelineState;
+		FGraphicsPipelineStateInitializer GraphicsPSOInit = MeshPipelineState;
 		RHICmdList.ApplyCachedRenderTargets(GraphicsPSOInit);
 		SetGraphicsPipelineState(RHICmdList, GraphicsPSOInit);
-		StateCache.SetPipelineState(MeshDrawCommand.PipelineState, MeshDrawCommand.CachedPipelineId.GetId());
+		StateCache.SetPipelineState(MeshDrawCommand.CachedPipelineId.GetId());
 	}
 
 	if (MeshDrawCommand.StencilRef != StateCache.StencilRef)
@@ -800,35 +833,35 @@ void FMeshDrawCommand::SubmitDraw(
 		}
 	}
 
-	MeshDrawCommand.ShaderBindings.SetOnCommandList(RHICmdList, MeshDrawCommand.PipelineState.BoundShaderState, StateCache.ShaderBindings);
+	MeshDrawCommand.ShaderBindings.SetOnCommandList(RHICmdList, MeshPipelineState.BoundShaderState, StateCache.ShaderBindings);
 
 	if (MeshDrawCommand.IndexBuffer)
 	{
-		if (MeshDrawCommand.IndirectArgsBuffer)
-		{
-			RHICmdList.DrawIndexedPrimitiveIndirect(
-				MeshDrawCommand.IndexBuffer, 
-				MeshDrawCommand.IndirectArgsBuffer, 
-				0
-				);
-		}
-		else
+		if (MeshDrawCommand.NumPrimitives > 0)
 		{
 			RHICmdList.DrawIndexedPrimitive(
 				MeshDrawCommand.IndexBuffer,
-				MeshDrawCommand.BaseVertexIndex,
+				MeshDrawCommand.VertexParams.BaseVertexIndex,
 				0,
-				MeshDrawCommand.NumVertices,
+				MeshDrawCommand.VertexParams.NumVertices,
 				MeshDrawCommand.FirstIndex,
 				MeshDrawCommand.NumPrimitives,
 				MeshDrawCommand.NumInstances * InstanceFactor
 			);
 		}
+		else
+		{
+			RHICmdList.DrawIndexedPrimitiveIndirect(
+				MeshDrawCommand.IndexBuffer, 
+				MeshDrawCommand.IndirectArgsBuffer, 
+				0
+				);
+		}
 	}
 	else
 	{
 		RHICmdList.DrawPrimitive(
-			MeshDrawCommand.BaseVertexIndex + MeshDrawCommand.FirstIndex,
+			MeshDrawCommand.VertexParams.BaseVertexIndex + MeshDrawCommand.FirstIndex,
 			MeshDrawCommand.NumPrimitives,
 			MeshDrawCommand.NumInstances * InstanceFactor
 		);
@@ -988,14 +1021,25 @@ FCachedPassMeshDrawListContext::FCachedPassMeshDrawListContext(FCachedMeshDrawCo
 	CommandInfo(InCommandInfo),
 	DrawList(InDrawList),
 	Scene(InScene)
-{}
+{
+	bUseStateBuckets = UseGPUScene(GMaxRHIShaderPlatform, GMaxRHIFeatureLevel) && InCommandInfo.MeshPass != EMeshPass::RayTracing;
+}
 
 FMeshDrawCommand& FCachedPassMeshDrawListContext::AddCommand(const FMeshDrawCommand& Initializer)
 {
-	// Only one FMeshDrawCommand supported per FStaticMesh in a pass
-	check(CommandInfo.CommandIndex == -1);
-	CommandInfo.CommandIndex = DrawList.MeshDrawCommands.Add(Initializer);
-	return DrawList.MeshDrawCommands[CommandInfo.CommandIndex];
+	if (bUseStateBuckets)
+	{
+		MeshDrawCommandForStateBucketing = Initializer;
+		return MeshDrawCommandForStateBucketing;
+	}
+	else
+	{
+		// Only one FMeshDrawCommand supported per FStaticMesh in a pass
+		check(CommandInfo.CommandIndex == -1);
+		// Allocate at lowest free index so that 'r.DoLazyStaticMeshUpdate' can shrink the TSparseArray more effectively
+		CommandInfo.CommandIndex = DrawList.MeshDrawCommands.AddAtLowestFreeIndex(Initializer, DrawList.LowestFreeIndexSearchStart);
+		return DrawList.MeshDrawCommands[CommandInfo.CommandIndex];
+	}
 }
 
 void FCachedPassMeshDrawListContext::FinalizeCommand(
@@ -1005,17 +1049,16 @@ void FCachedPassMeshDrawListContext::FinalizeCommand(
 	ERasterizerFillMode MeshFillMode,
 	ERasterizerCullMode MeshCullMode,
 	FMeshDrawCommandSortKey SortKey,
+	const FGraphicsMinimalPipelineStateInitializer& PipelineState,
+	const FMeshProcessorShaders* ShadersForDebugging,
 	FMeshDrawCommand& MeshDrawCommand,
 	bool bDoSetupPsoStateForRasterization)
 {
 	QUICK_SCOPE_CYCLE_COUNTER(STAT_FinalizeCachedMeshDrawCommand);
 
-	MeshDrawCommand.SetDrawParametersAndFinalize(MeshBatch, BatchElementIndex, bDoSetupPsoStateForRasterization);
-
-	check(CommandInfo.CommandIndex != -1);
-	const bool bCanUseGPUScene = UseGPUScene(GMaxRHIShaderPlatform, GMaxRHIFeatureLevel);
+	MeshDrawCommand.SetDrawParametersAndFinalize(MeshBatch, BatchElementIndex, PipelineState, ShadersForDebugging, bDoSetupPsoStateForRasterization);
 
-	if (bCanUseGPUScene /* && bDoSetupPsoStateForRasterization*/)
+	if (bUseStateBuckets)
 	{
 		FSetElementId SetId = Scene.CachedMeshDrawCommandStateBuckets.FindId(MeshDrawCommand);
 
@@ -1028,8 +1071,15 @@ void FCachedPassMeshDrawListContext::FinalizeCommand(
 			SetId = Scene.CachedMeshDrawCommandStateBuckets.Add(FMeshDrawCommandStateBucket(1, MeshDrawCommand));
 		}
 
+		check(CommandInfo.StateBucketId == -1);
 		CommandInfo.StateBucketId = SetId.AsInteger();
+		check(CommandInfo.CommandIndex == -1);
 	}
+	else
+	{
+		check(CommandInfo.CommandIndex != -1);
+	}
+
 	CommandInfo.SortKey = SortKey;
 	CommandInfo.MeshFillMode = MeshFillMode;
 	CommandInfo.MeshCullMode = MeshCullMode;
diff --git a/Engine/Source/Runtime/Renderer/Private/PrimitiveSceneInfo.cpp b/Engine/Source/Runtime/Renderer/Private/PrimitiveSceneInfo.cpp
--- a/Engine/Source/Runtime/Renderer/Private/PrimitiveSceneInfo.cpp
+++ b/Engine/Source/Runtime/Renderer/Private/PrimitiveSceneInfo.cpp
@@ -229,7 +229,6 @@ void FPrimitiveSceneInfo::CacheMeshDrawCommands(FRHICommandListImmediate& RHICmd
 					if ((FPassProcessorManager::GetPassFlags(ShadingPath, PassType) & EMeshPassFlags::CachedMeshCommands) != EMeshPassFlags::None)
 					{
 						FCachedMeshDrawCommandInfo CommandInfo;
-						CommandInfo.CommandIndex = -1;
 						CommandInfo.MeshPass = PassType;
 
 						FCachedPassMeshDrawList& SceneDrawList = Scene->CachedDrawLists[PassType];
@@ -247,13 +246,38 @@ void FPrimitiveSceneInfo::CacheMeshDrawCommands(FRHICommandListImmediate& RHICmd
 							PassMeshProcessor->~FMeshPassProcessor();
 						}
 
-						if (CommandInfo.CommandIndex != -1)
+						if (CommandInfo.CommandIndex != -1 || CommandInfo.StateBucketId != -1)
 						{
 							static_assert(sizeof(MeshRelevance.CommandInfosMask) * 8 >= EMeshPass::Num, "CommandInfosMask is too small to contain all mesh passes.");
 
 							MeshRelevance.CommandInfosMask.Set(PassType);
 
 							StaticMeshCommandInfos.Add(CommandInfo);
+
+#if DO_GUARD_SLOW
+							if (ShadingPath == EShadingPath::Deferred)
+							{
+								FMeshDrawCommand* MeshDrawCommand = CommandInfo.StateBucketId >= 0
+									? &Scene->CachedMeshDrawCommandStateBuckets[FSetElementId::FromInteger(CommandInfo.StateBucketId)].MeshDrawCommand
+									: &SceneDrawList.MeshDrawCommands[CommandInfo.CommandIndex];
+
+								ensureMsgf(MeshDrawCommand->VertexStreams.GetAllocatedSize() == 0, TEXT("Cached Mesh Draw command overflows VertexStreams.  VertexStream inline size should be tweaked."));
+								
+								if (PassType == EMeshPass::BasePass || PassType == EMeshPass::DepthPass || PassType == EMeshPass::CSMShadowDepth)
+								{
+									TArray<EShaderFrequency, TInlineAllocator<SF_NumFrequencies>> ShaderFrequencies;
+									MeshDrawCommand->ShaderBindings.GetShaderFrequencies(ShaderFrequencies);
+
+									for (int32 i = 0; i < ShaderFrequencies.Num(); i++)
+									{
+										FMeshDrawSingleShaderBindings SingleShaderBindings = MeshDrawCommand->ShaderBindings.GetSingleShaderBindings(ShaderFrequencies[i]);
+										ensureMsgf(SingleShaderBindings.ParameterMapInfo.LooseParameterBuffers.Num() == 0, TEXT("Cached Mesh Draw command uses loose parameters.  This will break dynamic instancing in performance critical pass.  Use Uniform Buffers instead."));
+										ensureMsgf(SingleShaderBindings.ParameterMapInfo.SRVs.Num() == 0, TEXT("Cached Mesh Draw command uses individual SRVs.  This will break dynamic instancing in performance critical pass.  Use Uniform Buffers instead."));
+										ensureMsgf(SingleShaderBindings.ParameterMapInfo.TextureSamplers.Num() == 0, TEXT("Cached Mesh Draw command uses individual Texture Samplers.  This will break dynamic instancing in performance critical pass.  Use Uniform Buffers instead."));
+									}
+								}
+							}
+#endif
 						}
 					}
 				}
@@ -264,34 +288,33 @@ void FPrimitiveSceneInfo::CacheMeshDrawCommands(FRHICommandListImmediate& RHICmd
 
 void FPrimitiveSceneInfo::RemoveCachedMeshDrawCommands()
 {
-	const bool bCanUseGPUScene = UseGPUScene(GMaxRHIShaderPlatform, GMaxRHIFeatureLevel);
-	
+	checkSlow(IsInRenderingThread());
+
 	for (int32 CommandIndex = 0; CommandIndex < StaticMeshCommandInfos.Num(); ++CommandIndex)
 	{
 		const FCachedMeshDrawCommandInfo& CachedCommand = StaticMeshCommandInfos[CommandIndex];
-		if (CachedCommand.CommandIndex != -1)
-		{
-			FCachedPassMeshDrawList& PassDrawList = Scene->CachedDrawLists[CachedCommand.MeshPass];
+		const FSetElementId StateBucketId = FSetElementId::FromInteger(CachedCommand.StateBucketId);
 			
-			if (bCanUseGPUScene)
+		if (StateBucketId.IsValidId())
+		{
+			FMeshDrawCommandStateBucket& StateBucket = Scene->CachedMeshDrawCommandStateBuckets[StateBucketId];
+
+			if (StateBucket.Num == 1)
 			{
-				const FSetElementId StateBucketId = FSetElementId::FromInteger(CachedCommand.StateBucketId);
-				checkSlow(StateBucketId.IsValidId());
-				FMeshDrawCommandStateBucket& StateBucket = Scene->CachedMeshDrawCommandStateBuckets[StateBucketId];
-				if (CachedCommand.StateBucketId != -1)
-				{
-					if (StateBucket.Num == 1)
-					{
-						Scene->CachedMeshDrawCommandStateBuckets.Remove(StateBucketId);
-					}
-					else
-					{
-						StateBucket.Num--;
-					}
-				}
+				Scene->CachedMeshDrawCommandStateBuckets.Remove(StateBucketId);
 			}
-
+			else
+			{
+				StateBucket.Num--;
+			}
+		}
+		else if (CachedCommand.CommandIndex >= 0)
+		{
+			FCachedPassMeshDrawList& PassDrawList = Scene->CachedDrawLists[CachedCommand.MeshPass];
 			PassDrawList.MeshDrawCommands.RemoveAt(CachedCommand.CommandIndex);
+
+			// Track the lowest index that might be free for faster AddAtLowestFreeIndex
+			PassDrawList.LowestFreeIndexSearchStart = FMath::Min(PassDrawList.LowestFreeIndexSearchStart, CachedCommand.CommandIndex);
 		}
 	}
 
diff --git a/Engine/Source/Runtime/Renderer/Private/RayTracing/RayTracingDynamicGeometry.cpp b/Engine/Source/Runtime/Renderer/Private/RayTracing/RayTracingDynamicGeometry.cpp
--- a/Engine/Source/Runtime/Renderer/Private/RayTracing/RayTracingDynamicGeometry.cpp
+++ b/Engine/Source/Runtime/Renderer/Private/RayTracing/RayTracingDynamicGeometry.cpp
@@ -131,10 +131,8 @@ void FRayTracingDynamicGeometryCollection::AddDynamicMeshBatchForGeometryUpdate(
 	check(DispatchCmd.TargetBuffer->NumBytes >= DispatchCmd.NumVertices * sizeof(FVector));
 
 #if MESH_DRAW_COMMAND_DEBUG_DATA
-	FMeshDrawCommandDebugData DebugData;
-	DebugData.Shaders = Shaders.GetUntypedShaders();
-
-	ShaderBindings.Finalize(DebugData);
+	FMeshProcessorShaders ShadersForDebug = Shaders.GetUntypedShaders();
+	ShaderBindings.Finalize(&ShadersForDebug);
 #endif
 
 	DispatchCommands.Add(DispatchCmd);
diff --git a/Engine/Source/Runtime/Renderer/Private/RendererScene.cpp b/Engine/Source/Runtime/Renderer/Private/RendererScene.cpp
--- a/Engine/Source/Runtime/Renderer/Private/RendererScene.cpp
+++ b/Engine/Source/Runtime/Renderer/Private/RendererScene.cpp
@@ -523,14 +523,12 @@ void FScene::CheckPrimitiveArrays()
 static TAutoConsoleVariable<int32> CVarDoLazyStaticMeshUpdate(
 	TEXT("r.DoLazyStaticMeshUpdate"),
 	0,
-	TEXT("If true, then do not add meshes to the static mesh draw lists until they are visible. Experiemental option. Incompatible with the Mesh Draw Command pipeline."));
+	TEXT("If true, then do not add meshes to the static mesh draw lists until they are visible. Experiemental option."));
 
 static void DoLazyStaticMeshUpdateCVarSinkFunction()
 {
-	//@todo MeshCommandPipeline r.DoLazyStaticMeshUpdate isn't currently supported.
-	/*
-	static bool CachedDoLazyStaticMeshUpdate = CVarDoLazyStaticMeshUpdate.GetValueOnGameThread() && !WITH_EDITOR;
-	bool DoLazyStaticMeshUpdate = CVarDoLazyStaticMeshUpdate.GetValueOnGameThread() && !WITH_EDITOR;
+	static bool CachedDoLazyStaticMeshUpdate = CVarDoLazyStaticMeshUpdate.GetValueOnGameThread() && !GIsEditor;
+	bool DoLazyStaticMeshUpdate = CVarDoLazyStaticMeshUpdate.GetValueOnGameThread() && !GIsEditor;
 
 	if (DoLazyStaticMeshUpdate != CachedDoLazyStaticMeshUpdate)
 	{
@@ -549,22 +547,129 @@ static void DoLazyStaticMeshUpdateCVarSinkFunction()
 			}
 		}
 	}
-	*/
 }
 
 static FAutoConsoleVariableSink CVarDoLazyStaticMeshUpdateSink(FConsoleCommandDelegate::CreateStatic(&DoLazyStaticMeshUpdateCVarSinkFunction));
 
 void FScene::UpdateDoLazyStaticMeshUpdate(FRHICommandListImmediate& CmdList)
 {
-	//@todo MeshCommandPipeline r.DoLazyStaticMeshUpdate isn't currently supported.
-	/*
-	bool DoLazyStaticMeshUpdate = CVarDoLazyStaticMeshUpdate.GetValueOnRenderThread() && !WITH_EDITOR;
+	bool DoLazyStaticMeshUpdate = CVarDoLazyStaticMeshUpdate.GetValueOnRenderThread() && !GIsEditor;
 
 	for (int32 PrimitiveIndex = 0; PrimitiveIndex < Primitives.Num(); PrimitiveIndex++)
 	{
 		Primitives[PrimitiveIndex]->UpdateStaticMeshes(CmdList, !DoLazyStaticMeshUpdate);
 	}
-	*/
+}
+
+void FScene::DumpMeshDrawCommandMemoryStats()
+{
+	SIZE_T TotalCachedMeshDrawCommands = 0;
+	SIZE_T TotalStaticMeshCommandInfos = 0;
+
+	struct FPassStats
+	{
+		SIZE_T CachedMeshDrawCommandBytes = 0;
+		SIZE_T PSOBytes = 0;
+		SIZE_T ShaderBindingInlineBytes = 0;
+		SIZE_T ShaderBindingHeapBytes = 0;
+		SIZE_T VertexStreamsInlineBytes = 0;
+		SIZE_T DebugDataBytes = 0;
+		SIZE_T DrawCommandParameterBytes = 0;
+		uint32 NumCommands = 0;
+	};
+
+	FPassStats AllPassStats[EMeshPass::Num];
+	TArray<bool> StateBucketAccounted;
+	StateBucketAccounted.Empty(CachedMeshDrawCommandStateBuckets.GetMaxIndex());
+	StateBucketAccounted.AddZeroed(CachedMeshDrawCommandStateBuckets.GetMaxIndex());
+
+	for (int32 i = 0; i < Primitives.Num(); i++)
+	{
+		FPrimitiveSceneInfo* PrimitiveSceneInfo = Primitives[i];
+
+		TotalStaticMeshCommandInfos += PrimitiveSceneInfo->StaticMeshCommandInfos.GetAllocatedSize();
+
+		for (int32 CommandIndex = 0; CommandIndex < PrimitiveSceneInfo->StaticMeshCommandInfos.Num(); ++CommandIndex)
+		{
+			const FCachedMeshDrawCommandInfo& CachedCommand = PrimitiveSceneInfo->StaticMeshCommandInfos[CommandIndex];
+			const FSetElementId StateBucketId = FSetElementId::FromInteger(CachedCommand.StateBucketId);
+			const FMeshDrawCommand* MeshDrawCommandPtr = nullptr;
+
+			if (StateBucketId.IsValidId())
+			{
+				if (!StateBucketAccounted[CachedCommand.StateBucketId])
+				{
+					StateBucketAccounted[CachedCommand.StateBucketId] = true;
+					FMeshDrawCommandStateBucket& StateBucket = CachedMeshDrawCommandStateBuckets[StateBucketId];
+					MeshDrawCommandPtr = &StateBucket.MeshDrawCommand;
+				}
+			}
+			else if (CachedCommand.CommandIndex >= 0)
+			{
+				FCachedPassMeshDrawList& PassDrawList = CachedDrawLists[CachedCommand.MeshPass];
+				MeshDrawCommandPtr = &PassDrawList.MeshDrawCommands[CachedCommand.CommandIndex];
+			}
+
+			if (MeshDrawCommandPtr)
+			{
+				const FMeshDrawCommand& MeshDrawCommand = *MeshDrawCommandPtr;
+				FPassStats& PassStats = AllPassStats[CachedCommand.MeshPass];
+				SIZE_T CommandBytes = sizeof(MeshDrawCommand) + MeshDrawCommand.GetAllocatedSize();
+				PassStats.CachedMeshDrawCommandBytes += CommandBytes;
+				TotalCachedMeshDrawCommands += MeshDrawCommand.GetAllocatedSize();
+				PassStats.PSOBytes += sizeof(MeshDrawCommand.CachedPipelineId);
+				PassStats.ShaderBindingInlineBytes += sizeof(MeshDrawCommand.ShaderBindings);
+				PassStats.ShaderBindingHeapBytes += MeshDrawCommand.ShaderBindings.GetAllocatedSize();
+				PassStats.VertexStreamsInlineBytes += sizeof(MeshDrawCommand.VertexStreams);
+				PassStats.DebugDataBytes += MeshDrawCommand.GetDebugDataSize();
+				PassStats.DrawCommandParameterBytes += sizeof(MeshDrawCommand.IndexBuffer) + sizeof(MeshDrawCommand.FirstIndex) + sizeof(MeshDrawCommand.NumPrimitives) + sizeof(MeshDrawCommand.NumInstances) + sizeof(MeshDrawCommand.VertexParams); //-V568
+				PassStats.NumCommands++;
+			}
+		}
+	}
+
+	TotalCachedMeshDrawCommands += CachedMeshDrawCommandStateBuckets.GetAllocatedSize();
+
+	for (int32 i = 0; i < EMeshPass::Num; i++)
+	{
+		TotalCachedMeshDrawCommands += CachedDrawLists[i].MeshDrawCommands.GetAllocatedSize();
+	}
+
+	for (int32 i = 0; i < EMeshPass::Num; i++)
+	{
+		const FPassStats& PassStats = AllPassStats[i];
+
+		if (PassStats.NumCommands > 0)
+		{
+			UE_LOG(LogRenderer, Log, TEXT("%s: %.1fKb for %u CachedMeshDrawCommands"), GetMeshPassName((EMeshPass::Type)i), PassStats.CachedMeshDrawCommandBytes / 1024.0f, PassStats.NumCommands);
+
+			if (PassStats.CachedMeshDrawCommandBytes > 1024 && i <= EMeshPass::BasePass)
+			{
+				UE_LOG(LogRenderer, Log, TEXT("     avg %.1f bytes PSO"), PassStats.PSOBytes / (float)PassStats.NumCommands);
+				UE_LOG(LogRenderer, Log, TEXT("     avg %.1f bytes ShaderBindingInline"), PassStats.ShaderBindingInlineBytes / (float)PassStats.NumCommands);
+				UE_LOG(LogRenderer, Log, TEXT("     avg %.1f bytes ShaderBindingHeap"), PassStats.ShaderBindingHeapBytes / (float)PassStats.NumCommands);
+				UE_LOG(LogRenderer, Log, TEXT("     avg %.1f bytes VertexStreamsInline"), PassStats.VertexStreamsInlineBytes / (float)PassStats.NumCommands);
+				UE_LOG(LogRenderer, Log, TEXT("     avg %.1f bytes DebugData"), PassStats.DebugDataBytes / (float)PassStats.NumCommands);
+				UE_LOG(LogRenderer, Log, TEXT("     avg %.1f bytes DrawCommandParameters"), PassStats.DrawCommandParameterBytes / (float)PassStats.NumCommands);
+
+				const SIZE_T Other = PassStats.CachedMeshDrawCommandBytes -
+					(PassStats.PSOBytes +
+					PassStats.ShaderBindingInlineBytes +
+					PassStats.ShaderBindingHeapBytes +
+					PassStats.VertexStreamsInlineBytes +
+					PassStats.DebugDataBytes +
+					PassStats.DrawCommandParameterBytes);
+
+				UE_LOG(LogRenderer, Log, TEXT("     avg %.1f bytes Other"), Other / (float)PassStats.NumCommands);
+			}
+		}
+	}
+
+	UE_LOG(LogRenderer, Log, TEXT("sizeof(FMeshDrawCommand) %u"), sizeof(FMeshDrawCommand));
+	UE_LOG(LogRenderer, Log, TEXT("Total cached MeshDrawCommands %.3fMb"), TotalCachedMeshDrawCommands / 1024.0f / 1024.0f);
+	UE_LOG(LogRenderer, Log, TEXT("Primitive StaticMeshCommandInfos %.1fKb"), TotalStaticMeshCommandInfos / 1024.0f);
+	UE_LOG(LogRenderer, Log, TEXT("GPUScene CPU structures %.1fKb"), GPUScene.PrimitivesToUpdate.GetAllocatedSize() / 1024.0f);
+	UE_LOG(LogRenderer, Log, TEXT("PSO Id Grow-Only Table %.1fKb"), FGraphicsMinimalPipelineStateId::GetGlobalTableSize() / 1024.0f);
 }
 
 template<typename T>
@@ -672,26 +777,23 @@ void FScene::AddPrimitiveSceneInfo_RenderThread(FRHICommandListImmediate& RHICmd
 	// Set lod Parent information if valid
 	PrimitiveSceneInfo->LinkLODParentComponent();
 
-	// Add the primitive to the scene.
-	PrimitiveSceneInfo->AddToScene(RHICmdList, true);
-
-	//@todo MeshCommandPipeline r.DoLazyStaticMeshUpdate isn't currently supported.
-	/*
-#if WITH_EDITOR
-	PrimitiveSceneInfo->AddToScene(RHICmdList, true);
-#else
-	const bool bAddToDrawLists = !(CVarDoLazyStaticMeshUpdate.GetValueOnRenderThread());
-	if (bAddToDrawLists)
+	if (GIsEditor)
 	{
 		PrimitiveSceneInfo->AddToScene(RHICmdList, true);
 	}
 	else
 	{
-		PrimitiveSceneInfo->AddToScene(RHICmdList, true, false);
-		PrimitiveSceneInfo->BeginDeferredUpdateStaticMeshes();
+		const bool bAddToDrawLists = !(CVarDoLazyStaticMeshUpdate.GetValueOnRenderThread());
+		if (bAddToDrawLists)
+		{
+			PrimitiveSceneInfo->AddToScene(RHICmdList, true);
+		}
+		else
+		{
+			PrimitiveSceneInfo->AddToScene(RHICmdList, true, false);
+			PrimitiveSceneInfo->BeginDeferredUpdateStaticMeshes();
+		}
 	}
-#endif
-	*/
 
 	AddPrimitiveToUpdateGPU(*this, SourceIndex);
 
diff --git a/Engine/Source/Runtime/Renderer/Private/ScenePrivate.h b/Engine/Source/Runtime/Renderer/Private/ScenePrivate.h
--- a/Engine/Source/Runtime/Renderer/Private/ScenePrivate.h
+++ b/Engine/Source/Runtime/Renderer/Private/ScenePrivate.h
@@ -2419,7 +2419,7 @@ struct MeshDrawCommandKeyFuncs : DefaultKeyFuncs<FMeshDrawCommandStateBucket,fal
 	/** Calculates a hash index for a key. */
 	static FORCEINLINE uint32 GetKeyHash(KeyInitType Key)
 	{
-		return PointerHash(Key.IndexBuffer, GetTypeHash(Key.PipelineState));
+		return PointerHash(Key.IndexBuffer, GetTypeHash(Key.CachedPipelineId.GetId()));
 	}
 };
 
@@ -2863,6 +2863,8 @@ public:
 	/** Debug function to abtest lazy static mesh drawlists. */
 	void UpdateDoLazyStaticMeshUpdate(FRHICommandListImmediate& CmdList);
 
+	void DumpMeshDrawCommandMemoryStats();
+
 private:
 
 	/**
diff --git a/Engine/Source/Runtime/Renderer/Private/SceneRendering.cpp b/Engine/Source/Runtime/Renderer/Private/SceneRendering.cpp
--- a/Engine/Source/Runtime/Renderer/Private/SceneRendering.cpp
+++ b/Engine/Source/Runtime/Renderer/Private/SceneRendering.cpp
@@ -85,6 +85,14 @@ FAutoConsoleVariableRef CVarDumpInstancingStats(
 	ECVF_Scalability | ECVF_RenderThreadSafe
 	);
 
+int32 GDumpMeshDrawCommandMemoryStats = 0;
+FAutoConsoleVariableRef CVarDumpMeshDrawCommandMemoryStats(
+	TEXT("r.MeshDrawCommands.LogMeshDrawCommandMemoryStats"),
+	GDumpMeshDrawCommandMemoryStats,
+	TEXT("Whether to log mesh draw command memory stats on the next frame"),
+	ECVF_Scalability | ECVF_RenderThreadSafe
+	);
+
 extern ENGINE_API FLightMap2D* GDebugSelectedLightmap;
 extern ENGINE_API UPrimitiveComponent* GDebugSelectedComponent;
 
@@ -2639,6 +2647,12 @@ void FSceneRenderer::RenderFinish(FRHICommandListImmediate& RHICmdList)
 
 	// Notify the RHI we are done rendering a scene.
 	RHICmdList.EndScene();
+
+	if (GDumpMeshDrawCommandMemoryStats)
+	{
+		GDumpMeshDrawCommandMemoryStats = 0;
+		Scene->DumpMeshDrawCommandMemoryStats();
+	}
 }
 
 void FSceneRenderer::SetupMeshPass(FViewInfo& View, FExclusiveDepthStencil::Type BasePassDepthStencilAccess, FViewCommands& ViewCommands)
diff --git a/Engine/Source/Runtime/Renderer/Private/SceneVisibility.cpp b/Engine/Source/Runtime/Renderer/Private/SceneVisibility.cpp
--- a/Engine/Source/Runtime/Renderer/Private/SceneVisibility.cpp
+++ b/Engine/Source/Runtime/Renderer/Private/SceneVisibility.cpp
@@ -1710,8 +1710,12 @@ struct FDrawCommandRelevancePacket
 
 				FVisibleMeshDrawCommand NewVisibleMeshDrawCommand;
 
+				const FMeshDrawCommand* MeshDrawCommand = CachedMeshDrawCommand.StateBucketId >= 0
+					? &Scene->CachedMeshDrawCommandStateBuckets[FSetElementId::FromInteger(CachedMeshDrawCommand.StateBucketId)].MeshDrawCommand
+					: &SceneDrawList.MeshDrawCommands[CachedMeshDrawCommand.CommandIndex];
+
 				NewVisibleMeshDrawCommand.Setup(
-					&SceneDrawList.MeshDrawCommands[CachedMeshDrawCommand.CommandIndex],
+					MeshDrawCommand,
 					PrimitiveIndex,
 					CachedMeshDrawCommand.StateBucketId,
 					CachedMeshDrawCommand.MeshFillMode,
@@ -2804,19 +2808,25 @@ void FSceneRenderer::PreVisibilityFrameSetup(FRHICommandListImmediate& RHICmdLis
 	// Notify the RHI we are beginning to render a scene.
 	RHICmdList.BeginScene();
 
-	//@todo MeshCommandPipeline r.DoLazyStaticMeshUpdate isn't currently supported.
-	/*
 	{
 		static auto CVar = IConsoleManager::Get().FindConsoleVariable(TEXT("r.DoLazyStaticMeshUpdate"));
-		const bool DoLazyStaticMeshUpdate = (CVar->GetInt() && !WITH_EDITOR);
+		const bool DoLazyStaticMeshUpdate = (CVar->GetInt() && !GIsEditor);
 
 		if (DoLazyStaticMeshUpdate)
 		{
 			QUICK_SCOPE_CYCLE_COUNTER(STAT_PreVisibilityFrameSetup_EvictionForLazyStaticMeshUpdate);
 			static int32 RollingRemoveIndex = 0;
+			static int32 RollingPassShrinkIndex = 0;
 			if (RollingRemoveIndex >= Scene->Primitives.Num())
 			{
 				RollingRemoveIndex = 0;
+				RollingPassShrinkIndex++;
+				if (RollingPassShrinkIndex >= ARRAY_COUNT(Scene->CachedDrawLists))
+				{
+					RollingPassShrinkIndex = 0;
+				}
+				// Periodically shrink the SparseArray containing cached mesh draw commands which we are causing to be regenerated with UpdateStaticMeshes
+				Scene->CachedDrawLists[RollingPassShrinkIndex].MeshDrawCommands.Shrink();
 			}
 			const int32 NumRemovedPerFrame = 10;
 			for (int32 NumRemoved = 0; NumRemoved < NumRemovedPerFrame && RollingRemoveIndex < Scene->Primitives.Num(); NumRemoved++, RollingRemoveIndex++)
@@ -2824,7 +2834,7 @@ void FSceneRenderer::PreVisibilityFrameSetup(FRHICommandListImmediate& RHICmdLis
 				Scene->Primitives[RollingRemoveIndex]->UpdateStaticMeshes(RHICmdList, false);
 			}
 		}
-	}*/
+	}
 
 	// Notify the FX system that the scene is about to perform visibility checks.
 	if (Scene->FXSystem && !Views[0].bIsPlanarReflection)
diff --git a/Engine/Source/Runtime/Renderer/Private/ShadowSetup.cpp b/Engine/Source/Runtime/Renderer/Private/ShadowSetup.cpp
--- a/Engine/Source/Runtime/Renderer/Private/ShadowSetup.cpp
+++ b/Engine/Source/Runtime/Renderer/Private/ShadowSetup.cpp
@@ -822,11 +822,14 @@ void FProjectedShadowInfo::AddCachedMeshDrawCommandsForPass(
 		{
 			const FCachedMeshDrawCommandInfo& CachedMeshDrawCommand = InPrimitiveSceneInfo->StaticMeshCommandInfos[StaticMeshCommandInfoIndex];
 			const FCachedPassMeshDrawList& SceneDrawList = Scene->CachedDrawLists[PassType];
+			const FMeshDrawCommand* MeshDrawCommand = CachedMeshDrawCommand.StateBucketId >= 0
+					? &Scene->CachedMeshDrawCommandStateBuckets[FSetElementId::FromInteger(CachedMeshDrawCommand.StateBucketId)].MeshDrawCommand
+					: &SceneDrawList.MeshDrawCommands[CachedMeshDrawCommand.CommandIndex];
 
 			FVisibleMeshDrawCommand NewVisibleMeshDrawCommand;
 
 			NewVisibleMeshDrawCommand.Setup(
-				&SceneDrawList.MeshDrawCommands[CachedMeshDrawCommand.CommandIndex],
+				MeshDrawCommand,
 				PrimitiveIndex,
 				CachedMeshDrawCommand.StateBucketId,
 				CachedMeshDrawCommand.MeshFillMode,
@@ -1325,9 +1328,10 @@ void FProjectedShadowInfo::ApplyViewOverridesToMeshDrawCommands(const FViewInfo&
 			NewMeshCommand = MeshCommand;
 
 			const ERasterizerCullMode LocalCullMode = View.bRenderSceneTwoSided ? CM_None : View.bReverseCulling ? FMeshPassProcessor::InverseCullMode(VisibleMeshDrawCommand.MeshCullMode) : VisibleMeshDrawCommand.MeshCullMode;
-			NewMeshCommand.PipelineState.RasterizerState = GetStaticRasterizerState<true>(VisibleMeshDrawCommand.MeshFillMode, LocalCullMode);
+			FGraphicsMinimalPipelineStateInitializer PipelineState = MeshCommand.CachedPipelineId.GetPipelineState();
+			PipelineState.RasterizerState = GetStaticRasterizerState<true>(VisibleMeshDrawCommand.MeshFillMode, LocalCullMode);
 
-			NewMeshCommand.Finalize(true);
+			NewMeshCommand.Finalize(PipelineState, nullptr, true);
 
 			FVisibleMeshDrawCommand NewVisibleMeshDrawCommand;
 
diff --git a/Engine/Source/Runtime/Renderer/Public/MeshDrawShaderBindings.h b/Engine/Source/Runtime/Renderer/Public/MeshDrawShaderBindings.h
--- a/Engine/Source/Runtime/Renderer/Public/MeshDrawShaderBindings.h
+++ b/Engine/Source/Runtime/Renderer/Public/MeshDrawShaderBindings.h
@@ -33,21 +33,26 @@ public:
 
 	inline uint32 GetDataSizeBytes() const
 	{
-		uint32 DataSize = (ParameterMapInfo.UniformBuffers.Num() 
+		uint32 DataSize = sizeof(void*) * 
+			(ParameterMapInfo.UniformBuffers.Num() 
 			+ ParameterMapInfo.TextureSamplers.Num() 
-			// At the RHI level we don't know whether SRV's will come from FTextureRHIParamRef or FShaderResourceViewRHIParamRef so we have to store space for both
-			+ 2 * ParameterMapInfo.SRVs.Num()) * sizeof(void*);
+			+ ParameterMapInfo.SRVs.Num());
+
+		// Allocate a bit for each SRV tracking whether it is a FTextureRHIParamRef or FShaderResourceViewRHIParamRef
+		DataSize += FMath::DivideAndRoundUp(ParameterMapInfo.SRVs.Num(), 8);
 
 		for (int32 LooseBufferIndex = 0; LooseBufferIndex < ParameterMapInfo.LooseParameterBuffers.Num(); LooseBufferIndex++)
 		{
 			DataSize += ParameterMapInfo.LooseParameterBuffers[LooseBufferIndex].BufferSize;
 		}
 
-		return DataSize;
+		// Align to pointer size so subsequent packed shader bindings will have their pointers aligned
+		return Align(DataSize, sizeof(void*));
 	}
 
 protected:
 
+	// Note: pointers first in layout, so they stay aligned
 	inline uint32 GetUniformBufferOffset() const
 	{
 		return 0;
@@ -60,23 +65,20 @@ protected:
 
 	inline uint32 GetSRVOffset() const
 	{
-		return ParameterMapInfo.UniformBuffers.Num() * sizeof(FUniformBufferRHIParamRef) 
+		return GetSamplerOffset() 
 			+ ParameterMapInfo.TextureSamplers.Num() * sizeof(FSamplerStateRHIParamRef);
 	}
 
-	inline uint32 GetTextureOffset() const
+	inline uint32 GetSRVTypeOffset() const
 	{
-		return ParameterMapInfo.UniformBuffers.Num() * sizeof(FUniformBufferRHIParamRef) 
-			+ ParameterMapInfo.TextureSamplers.Num() * sizeof(FSamplerStateRHIParamRef) 
+		return GetSRVOffset() 
 			+ ParameterMapInfo.SRVs.Num() * sizeof(FShaderResourceViewRHIParamRef);
 	}
 
 	inline uint32 GetLooseDataOffset() const
 	{
-		return ParameterMapInfo.UniformBuffers.Num() * sizeof(FUniformBufferRHIParamRef) 
-			+ ParameterMapInfo.TextureSamplers.Num() * sizeof(FSamplerStateRHIParamRef) 
-			+ ParameterMapInfo.SRVs.Num() * sizeof(FShaderResourceViewRHIParamRef)
-			+ ParameterMapInfo.SRVs.Num() * sizeof(FTextureRHIParamRef);
+		return GetSRVTypeOffset()
+			+ FMath::DivideAndRoundUp(ParameterMapInfo.SRVs.Num(), 8);
 	}
 
 	friend class FMeshDrawShaderBindings;
@@ -220,16 +222,17 @@ private:
 		return (FSamplerStateRHIParamRef*)SamplerDataStart;
 	}
 
-	inline FShaderResourceViewRHIParamRef* GetSRVStart() const
+	inline FRHIResource** GetSRVStart() const
 	{
 		uint8* SRVDataStart = Data + GetSRVOffset();
-		return (FShaderResourceViewRHIParamRef*)SRVDataStart;
+		checkfSlow(Align(*SRVDataStart, sizeof(void*)) == *SRVDataStart, TEXT("FMeshDrawSingleShaderBindings should have been laid out so that stored pointers are aligned"));
+		return (FRHIResource**)SRVDataStart;
 	}
 
-	inline FTextureRHIParamRef* GetTextureStart() const
+	inline uint8* GetSRVTypeStart() const
 	{
-		uint8* TextureDataStart = Data + GetTextureOffset();
-		return (FTextureRHIParamRef*)TextureDataStart;
+		uint8* SRVTypeDataStart = Data + GetSRVTypeOffset();
+		return SRVTypeDataStart;
 	}
 
 	inline uint8* GetLooseDataStart() const
@@ -310,6 +313,9 @@ private:
 
 		if (FoundIndex >= 0)
 		{
+			uint32 TypeByteIndex = FoundIndex / 8;
+			uint32 TypeBitIndex = FoundIndex - TypeByteIndex;
+			GetSRVTypeStart()[TypeByteIndex] |= 1 << TypeBitIndex;
 			GetSRVStart()[FoundIndex] = Value;
 		}
 
@@ -333,7 +339,7 @@ private:
 
 		if (FoundIndex >= 0)
 		{
-			GetTextureStart()[FoundIndex] = Value;
+			GetSRVStart()[FoundIndex] = Value;
 		}
 
 		checkfSlow(FoundIndex >= 0, TEXT("Attempted to set Texture at BaseIndex %u which was never in the shader's parameter map."), BaseIndex);
diff --git a/Engine/Source/Runtime/Renderer/Public/MeshPassProcessor.h b/Engine/Source/Runtime/Renderer/Public/MeshPassProcessor.h
--- a/Engine/Source/Runtime/Renderer/Public/MeshPassProcessor.h
+++ b/Engine/Source/Runtime/Renderer/Public/MeshPassProcessor.h
@@ -111,6 +111,23 @@ public:
 
 	void Setup(const FGraphicsMinimalPipelineStateInitializer& InPipelineState);
 
+	inline bool operator==(const FGraphicsMinimalPipelineStateId& rhs) const
+	{
+		return Id == rhs.Id;
+	}
+
+	inline bool operator!=(const FGraphicsMinimalPipelineStateId& rhs) const
+	{
+		return !(*this == rhs);
+	}
+	
+	inline const FGraphicsMinimalPipelineStateInitializer& GetPipelineState() const
+	{
+		return GlobalTable[Id];
+	}
+
+	static SIZE_T GetGlobalTableSize() { return GlobalTable.GetAllocatedSize(); }
+
 private:
 	FSetElementId Id;
 
@@ -183,7 +200,8 @@ struct FMeshDrawCommandDebugData
 	const FPrimitiveSceneProxy* PrimitiveSceneProxy;
 	const FMaterial* Material;
 	const FMaterialRenderProxy* MaterialRenderProxy;
-	FMeshProcessorShaders Shaders;
+	FMeshMaterialShader* VertexShader;
+	FMeshMaterialShader* PixelShader;
 #endif
 };
 
@@ -211,7 +229,7 @@ public:
 	void Initialize(FMeshProcessorShaders Shaders);
 
 	/** Called once binding setup is complete. */
-	void Finalize(const FMeshDrawCommandDebugData& DebugData);
+	void Finalize(const FMeshProcessorShaders* ShadersForDebugging);
 
 	inline FMeshDrawSingleShaderBindings GetSingleShaderBindings(EShaderFrequency Frequency)
 	{
@@ -247,7 +265,7 @@ public:
 	SIZE_T GetAllocatedSize() const
 	{
 		SIZE_T Bytes = ShaderLayouts.GetAllocatedSize();
-		if (Size > ARRAY_COUNT(InlineStorage))
+		if (Size > sizeof(InlineStorage))
 		{
 			Bytes += Size;
 		}
@@ -255,6 +273,18 @@ public:
 		return Bytes;
 	}
 
+	void GetShaderFrequencies(TArray<EShaderFrequency, TInlineAllocator<SF_NumFrequencies>>& OutShaderFrequencies) const
+	{
+		OutShaderFrequencies.Empty(ShaderLayouts.Num());
+
+		for (int32 BindingIndex = 0; BindingIndex < ShaderLayouts.Num(); BindingIndex++)
+		{
+			OutShaderFrequencies.Add(ShaderLayouts[BindingIndex].Frequency);
+		}
+	}
+
+	inline int32 GetDataSize() const { return Size; }
+
 private:
 
 	TArray<FMeshDrawShaderBindingsLayout, TInlineAllocator<2>> ShaderLayouts;
@@ -327,12 +357,7 @@ private:
 class FMeshDrawCommand
 {
 public:
-	/**
-	 * PSO
-	 */
-	FGraphicsMinimalPipelineStateInitializer PipelineState;
-	FGraphicsMinimalPipelineStateId CachedPipelineId;
-
+	
 	/**
 	 * Resource bindings
 	 */
@@ -341,32 +366,45 @@ public:
 	FIndexBufferRHIParamRef IndexBuffer;
 
 	/**
-	* Ray tracing specific
-	*/
-	uint32 RayTracingMaterialLibraryIndex = UINT_MAX;
+	 * PSO
+	 */
+	FGraphicsMinimalPipelineStateId CachedPipelineId;
 
 	/**
-	* Draw command parameters
-	*/
+	 * Draw command parameters
+	 */
 	uint32 FirstIndex;
 	uint32 NumPrimitives;
 	uint32 NumInstances;
-	uint32 BaseVertexIndex;
-	uint32 NumVertices;
-	FVertexBufferRHIParamRef IndirectArgsBuffer;
-	uint8 RayTracedSegmentIndex;
+
+	union
+	{
+		struct 
+		{
+			uint32 BaseVertexIndex;
+			uint32 NumVertices;
+		} VertexParams;
+		
+		FVertexBufferRHIParamRef IndirectArgsBuffer;
+	};
 
 	int8 PrimitiveIdStreamIndex;
 
 	/** Non-pipeline state */
 	uint8 StencilRef;
 
+    /**
+	 * Ray tracing specific
+	 */
+	uint8 RayTracedSegmentIndex;
+	uint32 RayTracingMaterialLibraryIndex = UINT_MAX;
+
 	FMeshDrawCommand()
 	{}
 
 	bool MatchesForDynamicInstancing(const FMeshDrawCommand& Rhs) const
 	{
-		return PipelineState == Rhs.PipelineState
+		return CachedPipelineId == Rhs.CachedPipelineId
 			&& StencilRef == Rhs.StencilRef
 			&& ShaderBindings.MatchesForDynamicInstancing(Rhs.ShaderBindings)
 			&& VertexStreams == Rhs.VertexStreams
@@ -375,13 +413,12 @@ public:
 			&& FirstIndex == Rhs.FirstIndex
 			&& NumPrimitives == Rhs.NumPrimitives
 			&& NumInstances == Rhs.NumInstances
-			&& BaseVertexIndex == Rhs.BaseVertexIndex
-			&& NumVertices == Rhs.NumVertices
-			&& IndirectArgsBuffer == Rhs.IndirectArgsBuffer;
+			&& ((NumPrimitives > 0 && VertexParams.BaseVertexIndex == Rhs.VertexParams.BaseVertexIndex && VertexParams.NumVertices == Rhs.VertexParams.NumVertices)
+				|| (NumPrimitives == 0 && IndirectArgsBuffer == Rhs.IndirectArgsBuffer));
 	}
 
 	/** Sets shaders on the mesh draw command and allocates room for the shader bindings. */
-	RENDERER_API void SetShaders(FVertexDeclarationRHIParamRef VertexDeclaration, const FMeshProcessorShaders& Shaders);
+	RENDERER_API void SetShaders(FVertexDeclarationRHIParamRef VertexDeclaration, const FMeshProcessorShaders& Shaders, FGraphicsMinimalPipelineStateInitializer& PipelineState);
 
 #if RHI_RAYTRACING
 	/** Sets ray hit group shaders on the mesh draw command and allocates room for the shader bindings. */
@@ -396,17 +433,21 @@ public:
 	}
 
 	/** Called when the mesh draw command is complete. */
-	RENDERER_API void SetDrawParametersAndFinalize(const FMeshBatch& MeshBatch, int32 BatchElementIndex, bool bDoSetupPsoStateForRasterization);
+	RENDERER_API void SetDrawParametersAndFinalize(
+		const FMeshBatch& MeshBatch, 
+		int32 BatchElementIndex,
+		const FGraphicsMinimalPipelineStateInitializer& PipelineState, 
+		const FMeshProcessorShaders* ShadersForDebugging,
+		bool bDoSetupPsoStateForRasterization);
 
-	void Finalize(bool bDoSetupPsoStateForRasterization)
+	void Finalize(const FGraphicsMinimalPipelineStateInitializer& PipelineState, const FMeshProcessorShaders* ShadersForDebugging, bool bDoSetupPsoStateForRasterization)
 	{
 		if (bDoSetupPsoStateForRasterization)
 		{
 			CachedPipelineId.Setup(PipelineState);
 		}
-#if MESH_DRAW_COMMAND_DEBUG_DATA
-		ShaderBindings.Finalize(DebugData);	
-#endif
+
+		ShaderBindings.Finalize(ShadersForDebugging);	
 	}
 
 	/** Submits commands to the RHI Commandlist to draw the MeshDrawCommand. */
@@ -418,7 +459,7 @@ public:
 		FRHICommandList& CommandList, 
 		class FMeshDrawCommandStateCache& RESTRICT StateCache);
 
-	FORCEINLINE friend uint32 GetTypeHash( const FMeshDrawCommand& Other )
+	FORCENOINLINE friend uint32 GetTypeHash( const FMeshDrawCommand& Other )
 	{
 		return Other.CachedPipelineId.GetId();
 	}
@@ -429,7 +470,8 @@ public:
 		DebugData.PrimitiveSceneProxy = PrimitiveSceneProxy;
 		DebugData.Material = Material;
 		DebugData.MaterialRenderProxy = MaterialRenderProxy;
-		DebugData.Shaders = UntypedShaders;
+		DebugData.VertexShader = UntypedShaders.VertexShader;
+		DebugData.PixelShader = UntypedShaders.PixelShader;
 #endif
 	}
 
@@ -438,6 +480,14 @@ public:
 		return ShaderBindings.GetAllocatedSize() + VertexStreams.GetAllocatedSize();
 	}
 
+	SIZE_T GetDebugDataSize() const
+	{
+#if MESH_DRAW_COMMAND_DEBUG_DATA
+		return sizeof(DebugData);
+#endif
+		return 0;
+	}
+
 #if MESH_DRAW_COMMAND_DEBUG_DATA
 private:
 	FMeshDrawCommandDebugData DebugData;
@@ -503,6 +553,8 @@ public:
 		ERasterizerFillMode MeshFillMode,
 		ERasterizerCullMode MeshCullMode,
 		FMeshDrawCommandSortKey SortKey,
+		const FGraphicsMinimalPipelineStateInitializer& PipelineState,
+		const FMeshProcessorShaders* ShadersForDebugging,
 		FMeshDrawCommand& MeshDrawCommand,
 		bool bDoSetupPsoStateForRasterization) = 0;
 };
@@ -600,10 +652,12 @@ public:
 		ERasterizerFillMode MeshFillMode,
 		ERasterizerCullMode MeshCullMode,
 		FMeshDrawCommandSortKey SortKey,
+		const FGraphicsMinimalPipelineStateInitializer& PipelineState,
+		const FMeshProcessorShaders* ShadersForDebugging,
 		FMeshDrawCommand& MeshDrawCommand,
 		bool bDoSetupPsoStateForRasterization) override final
 	{
-		MeshDrawCommand.SetDrawParametersAndFinalize(MeshBatch, BatchElementIndex, bDoSetupPsoStateForRasterization);
+		MeshDrawCommand.SetDrawParametersAndFinalize(MeshBatch, BatchElementIndex, PipelineState, ShadersForDebugging, bDoSetupPsoStateForRasterization);
 
 		FVisibleMeshDrawCommand NewVisibleMeshDrawCommand;
 		//@todo MeshCommandPipeline - assign usable state ID for dynamic path draws
@@ -617,6 +671,10 @@ private:
 	FMeshCommandOneFrameArray& DrawList;
 };
 
+#if PLATFORM_SUPPORTS_PRAGMA_PACK
+	#pragma pack (push,4)
+#endif
+
 /** 
  * Stores information about a mesh draw command which is cached in the scene. 
  * This is stored separately from the cached FMeshDrawCommand so that InitViews does not have to load the FMeshDrawCommand into cache.
@@ -635,7 +693,10 @@ public:
 
 	FMeshDrawCommandSortKey SortKey;
 
+	// Stores the index into FScene::CachedDrawLists of the corresponding FMeshDrawCommand, or -1 if not stored there
 	int32 CommandIndex;
+
+	// Stores the index into FScene::CachedMeshDrawCommandStateBuckets of the corresponding FMeshDrawCommand, or -1 if not stored there
 	int32 StateBucketId;
 
 	// Needed for easier debugging and faster removal of cached mesh draw commands.
@@ -646,11 +707,21 @@ public:
 	ERasterizerCullMode MeshCullMode : ERasterizerCullMode_NumBits + 1;
 };
 
+#if PLATFORM_SUPPORTS_PRAGMA_PACK
+	#pragma pack (pop)
+#endif
+
 class FCachedPassMeshDrawList
 {
 public:
+
+	FCachedPassMeshDrawList() :
+		LowestFreeIndexSearchStart(0)
+	{}
+
 	/** Indices held by FStaticMeshBatch::CachedMeshDrawCommands must be stable */
 	TSparseArray<FMeshDrawCommand> MeshDrawCommands;
+	int32 LowestFreeIndexSearchStart;
 };
 
 typedef TArray<int32, TInlineAllocator<5>> FDrawCommandIndices;
@@ -669,13 +740,17 @@ public:
 		ERasterizerFillMode MeshFillMode,
 		ERasterizerCullMode MeshCullMode,
 		FMeshDrawCommandSortKey SortKey,
+		const FGraphicsMinimalPipelineStateInitializer& PipelineState,
+		const FMeshProcessorShaders* ShadersForDebugging,
 		FMeshDrawCommand& MeshDrawCommand,
 		bool bDoSetupPsoStateForRasterization) override final;
 
 private:
+	FMeshDrawCommand MeshDrawCommandForStateBucketing;
 	FCachedMeshDrawCommandInfo& CommandInfo;
 	FCachedPassMeshDrawList& DrawList;
 	FScene& Scene;
+	bool bUseStateBuckets;
 };
 
 template<typename VertexType, typename HullType, typename DomainType, typename PixelType, typename GeometryType = FMeshMaterialShader, typename RayHitGroupType = FMeshMaterialShader, typename ComputeType = FMeshMaterialShader>
diff --git a/Engine/Source/Runtime/Renderer/Public/MeshPassProcessor.inl b/Engine/Source/Runtime/Renderer/Public/MeshPassProcessor.inl
--- a/Engine/Source/Runtime/Renderer/Public/MeshPassProcessor.inl
+++ b/Engine/Source/Runtime/Renderer/Public/MeshPassProcessor.inl
@@ -29,7 +29,9 @@ void FMeshPassProcessor::BuildMeshDrawCommands(
 	FMeshDrawCommand SharedMeshDrawCommand;
 
 	SharedMeshDrawCommand.SetStencilRef(DrawRenderState.GetStencilRef());
-	SharedMeshDrawCommand.PipelineState.PrimitiveType = (EPrimitiveType)MeshBatch.Type;
+
+	FGraphicsMinimalPipelineStateInitializer PipelineState;
+	PipelineState.PrimitiveType = (EPrimitiveType)MeshBatch.Type;
 
 	const bool bPositionOnly = (MeshPassFeatures & EMeshPassFeatures::PositionOnly) != EMeshPassFeatures::Default;
 
@@ -37,15 +39,15 @@ void FMeshPassProcessor::BuildMeshDrawCommands(
 	FVertexDeclarationRHIParamRef VertexDeclaration = bPositionOnly ? VertexFactory->GetPositionDeclaration() : VertexFactory->GetDeclaration();
 	check(!VertexFactory->NeedsDeclaration() || VertexDeclaration);
 
-	SharedMeshDrawCommand.SetShaders(VertexDeclaration, PassShaders.GetUntypedShaders());
+	SharedMeshDrawCommand.SetShaders(VertexDeclaration, PassShaders.GetUntypedShaders(), PipelineState);
 
-	SharedMeshDrawCommand.PipelineState.RasterizerState = GetStaticRasterizerState<true>(MeshFillMode, MeshCullMode);
+	PipelineState.RasterizerState = GetStaticRasterizerState<true>(MeshFillMode, MeshCullMode);
 
 	check(DrawRenderState.GetDepthStencilState());
 	check(DrawRenderState.GetBlendState());
 
-	SharedMeshDrawCommand.PipelineState.BlendState = DrawRenderState.GetBlendState();
-	SharedMeshDrawCommand.PipelineState.DepthStencilState = DrawRenderState.GetDepthStencilState();
+	PipelineState.BlendState = DrawRenderState.GetBlendState();
+	PipelineState.DepthStencilState = DrawRenderState.GetDepthStencilState();
 
 	check(VertexFactory && VertexFactory->IsInitialized());
 
@@ -129,8 +131,8 @@ void FMeshPassProcessor::BuildMeshDrawCommands(
 			}
 
 			const int32 DrawPrimitiveId = GetDrawCommandPrimitiveId(PrimitiveSceneInfo, BatchElement);
-
-			DrawListContext->FinalizeCommand(MeshBatch, BatchElementIndex, DrawPrimitiveId, MeshFillMode, MeshCullMode, SortKey, MeshDrawCommand, true);
+			FMeshProcessorShaders ShadersForDebugging = PassShaders.GetUntypedShaders();
+			DrawListContext->FinalizeCommand(MeshBatch, BatchElementIndex, DrawPrimitiveId, MeshFillMode, MeshCullMode, SortKey, PipelineState, &ShadersForDebugging, MeshDrawCommand, true);
 		}
 	}
 }
@@ -173,6 +175,7 @@ void FMeshPassProcessor::BuildRayTracingDrawCommands(
 
 	checkf(MaterialRenderProxy.ImmutableSamplerState.ImmutableSamplers[0] == nullptr, TEXT("Immutable samplers not yet supported in Mesh Draw Command pipeline"));
 
+	FGraphicsMinimalPipelineStateInitializer DummyPipelineState;
 	FMeshDrawCommand SharedMeshDrawCommand;
 
 	check(VertexFactory && VertexFactory->IsInitialized());
@@ -211,7 +214,8 @@ void FMeshPassProcessor::BuildRayTracingDrawCommands(
 
 			const int32 DrawPrimitiveId = 0;
 			const int32 ScenePrimitiveId = 0;
-			DrawListContext->FinalizeCommand(MeshBatch, BatchElementIndex, DrawPrimitiveId, MeshFillMode, MeshCullMode, SortKey, MeshDrawCommand, false);
+			FMeshProcessorShaders ShadersForDebugging = PassShaders.GetUntypedShaders();
+			DrawListContext->FinalizeCommand(MeshBatch, BatchElementIndex, DrawPrimitiveId, MeshFillMode, MeshCullMode, SortKey, DummyPipelineState, &ShadersForDebugging, MeshDrawCommand, false);
 		}
 	}
 }
\ No newline at end of file
