UE4 - Android merge/compile fix
#rnx

diff --git a/Engine/Source/Runtime/Renderer/Private/MeshPassProcessor.cpp b/Engine/Source/Runtime/Renderer/Private/MeshPassProcessor.cpp
--- a/Engine/Source/Runtime/Renderer/Private/MeshPassProcessor.cpp
+++ b/Engine/Source/Runtime/Renderer/Private/MeshPassProcessor.cpp
@@ -605,12 +605,16 @@ void FMeshDrawShaderBindings::CopyFrom(const FMeshDrawShaderBindings& Other)
 void FMeshDrawCommand::SetShaders(FVertexDeclarationRHIParamRef VertexDeclaration, const FMeshProcessorShaders& Shaders, FGraphicsMinimalPipelineStateInitializer& PipelineState)
 {
 	PipelineState.BoundShaderState = FBoundShaderStateInput(
-		VertexDeclaration,
-		GETSAFERHISHADER_VERTEX(Shaders.VertexShader),
-		GETSAFERHISHADER_HULL(Shaders.HullShader),
-		GETSAFERHISHADER_DOMAIN(Shaders.DomainShader),
-		GETSAFERHISHADER_PIXEL(Shaders.PixelShader),
-		GETSAFERHISHADER_GEOMETRY(Shaders.GeometryShader)
+		VertexDeclaration
+		, GETSAFERHISHADER_VERTEX(Shaders.VertexShader)
+#if PLATFORM_SUPPORTS_TESSELLATION_SHADERS
+		, GETSAFERHISHADER_HULL(Shaders.HullShader)
+		, GETSAFERHISHADER_DOMAIN(Shaders.DomainShader)
+#endif
+		, GETSAFERHISHADER_PIXEL(Shaders.PixelShader)
+#if PLATFORM_SUPPORTS_GEOMETRY_SHADERS
+		, GETSAFERHISHADER_GEOMETRY(Shaders.GeometryShader)
+#endif
 	);
 
 	ShaderBindings.Initialize(Shaders);
