Refactors material sorting pass to shader permutation API to work arround clamg compilation failure

diff --git a/Engine/Shaders/Private/RayTracing/MaterialSort.usf b/Engine/Shaders/Private/RayTracing/MaterialSort.usf
--- a/Engine/Shaders/Private/RayTracing/MaterialSort.usf
+++ b/Engine/Shaders/Private/RayTracing/MaterialSort.usf
@@ -1,8 +1,16 @@
 #include "RayTracingCommon.ush"
 
-#ifndef NUM_ELEMENTS
-#define NUM_THREADS  512
-#define NUM_ELEMENTS  NUM_THREADS * 2
+#if DIM_SORT_SIZE == 0
+	#define NUM_ELEMENTS 256
+
+#elif DIM_SORT_SIZE == 1
+	#define NUM_ELEMENTS 512
+
+#elif DIM_SORT_SIZE == 2
+	#define NUM_ELEMENTS 1024
+
+#else
+	#error Unknown sorting size
 #endif
 
 groupshared FDeferredMaterialPayload Entries[NUM_ELEMENTS];
diff --git a/Engine/Source/Runtime/Renderer/Private/RayTracing/RayTracingDeferredMaterials.cpp b/Engine/Source/Runtime/Renderer/Private/RayTracing/RayTracingDeferredMaterials.cpp
--- a/Engine/Source/Runtime/Renderer/Private/RayTracing/RayTracingDeferredMaterials.cpp
+++ b/Engine/Source/Runtime/Renderer/Private/RayTracing/RayTracingDeferredMaterials.cpp
@@ -82,85 +82,59 @@ FRHIRayTracingPipelineState* FDeferredShadingSceneRenderer::BindRayTracingPipeli
 	return PipelineState;
 }
 
-template<uint32 SortSize>
-class TMaterialSortCS : public FGlobalShader
+class FMaterialSortCS : public FGlobalShader
 {
-	DECLARE_GLOBAL_SHADER(TMaterialSortCS);
-	SHADER_USE_PARAMETER_STRUCT(TMaterialSortCS, FGlobalShader);
+	DECLARE_GLOBAL_SHADER(FMaterialSortCS);
+	SHADER_USE_PARAMETER_STRUCT(FMaterialSortCS, FGlobalShader);
+
+	class FSortSize : SHADER_PERMUTATION_INT("DIM_SORT_SIZE", 3);
+
+	using FPermutationDomain = TShaderPermutationDomain<FSortSize>;
 
 	BEGIN_SHADER_PARAMETER_STRUCT(FParameters, )
 		SHADER_PARAMETER(int, NumTotalEntries)
 		SHADER_PARAMETER_RDG_BUFFER_UAV(StructuredBuffer<FDeferredMaterialPayload>, MaterialBuffer)
-		END_SHADER_PARAMETER_STRUCT()
+	END_SHADER_PARAMETER_STRUCT()
 
-public:
 	static bool ShouldCompilePermutation(const FGlobalShaderPermutationParameters& Parameters)
 	{
 		return Parameters.Platform == SP_PCD3D_SM5;
 	}
-
-	static void ModifyCompilationEnvironment(const FGlobalShaderPermutationParameters& Parameters, FShaderCompilerEnvironment& OutEnvironment)
-	{
-		// Compile time template-based conditional
-		OutEnvironment.SetDefine(TEXT("NUM_ELEMENTS"), SortSize);
-	}
-
 };
 
-IMPLEMENT_SHADER_TYPE(template<>, TMaterialSortCS<256>, TEXT("/Engine/Private/RayTracing/MaterialSort.usf"), TEXT("MaterialSortLocal"), SF_Compute);
-IMPLEMENT_SHADER_TYPE(template<>, TMaterialSortCS<512>, TEXT("/Engine/Private/RayTracing/MaterialSort.usf"), TEXT("MaterialSortLocal"), SF_Compute);
-IMPLEMENT_SHADER_TYPE(template<>, TMaterialSortCS<1024>, TEXT("/Engine/Private/RayTracing/MaterialSort.usf"), TEXT("MaterialSortLocal"), SF_Compute);
+IMPLEMENT_GLOBAL_SHADER(FMaterialSortCS, "/Engine/Private/RayTracing/MaterialSort.usf", "MaterialSortLocal", SF_Compute);
 
-template<uint32 SortSize>
-static void TemplateSortMaterials(FRDGBuilder& GraphBuilder,
+void SortDeferredMaterials(
+	FRDGBuilder& GraphBuilder,
 	const FViewInfo& View,
+	uint32 SortSize,
 	uint32 NumElements,
 	FRDGBufferRef MaterialBuffer)
 {
+	if (SortSize == 0)
+	{
+		return;
+	}
+	SortSize = FMath::Min(SortSize, 3u);
+
 	// Setup shader and parameters
-	TMaterialSortCS<SortSize>::FParameters* PassParameters = GraphBuilder.AllocParameters<TMaterialSortCS<SortSize>::FParameters>();
+	FMaterialSortCS::FParameters* PassParameters = GraphBuilder.AllocParameters<FMaterialSortCS::FParameters>();
 	PassParameters->NumTotalEntries = NumElements;
 	PassParameters->MaterialBuffer = GraphBuilder.CreateUAV(MaterialBuffer);
 
-	// Get the CS
-	auto SortShader = View.ShaderMap->GetShader<TMaterialSortCS<SortSize> >();
-	ClearUnusedGraphResources(SortShader, PassParameters);
-
-	// Add the pass to the graph
-	GraphBuilder.AddPass(
-		RDG_EVENT_NAME("Material Sort<%d>", SortSize),
-		PassParameters,
-		ERenderGraphPassFlags::Compute,
-		[PassParameters, &View, SortShader, NumElements](FRHICommandList& RHICmdList)
-
-	{
-		// Note that we are presently rounding down and leaving the last N elements unsorted
-		const uint32 DispatchWidth = (NumElements + SortSize - 1) / SortSize;
-
-		FRHIComputeShader* ShaderRHI = SortShader->GetComputeShader();
+	FMaterialSortCS::FPermutationDomain PermutationVector;
+	PermutationVector.Set<FMaterialSortCS::FSortSize>(SortSize - 1);
 
-		RHICmdList.SetComputeShader(ShaderRHI);
-		SetShaderParameters(RHICmdList, SortShader, ShaderRHI, *PassParameters);
-		RHICmdList.DispatchComputeShader(DispatchWidth, 1, 1);
-		UnsetShaderUAVs(RHICmdList, SortShader, ShaderRHI);
-	});
-}
+	// Note that we are presently rounding down and leaving the last N elements unsorted
+	const uint32 DispatchWidth = FMath::DivideAndRoundUp(NumElements, SortSize);
 
-void SortDeferredMaterials(
-	FRDGBuilder& GraphBuilder,
-	const FViewInfo& View,
-	uint32 SortSize,
-	uint32 NumElements,
-	FRDGBufferRef MaterialBuffer)
-{
-	SortSize = FMath::Min(SortSize, 3u);
-
-	switch (SortSize)
-	{
-	case 1: TemplateSortMaterials<256>(GraphBuilder, View, NumElements, MaterialBuffer); break;
-	case 2: TemplateSortMaterials<512>(GraphBuilder, View, NumElements, MaterialBuffer); break;
-	case 3: TemplateSortMaterials<1024>(GraphBuilder, View, NumElements, MaterialBuffer); break;
-	};
+	TShaderMapRef<FMaterialSortCS> SortShader(View.ShaderMap, PermutationVector);
+	FComputeShaderUtils::AddPass(
+		GraphBuilder,
+		RDG_EVENT_NAME("MaterialSort SortSize=%d NumElements=%d", SortSize, NumElements),
+		*SortShader,
+		PassParameters,
+		FIntVector(DispatchWidth, 1, 1));
 }
 
 #else // RHI_RAYTRACING
