Temporary fix for shadow rendering on instanced meshes.

diff --git a/Engine/Source/Runtime/Renderer/Private/ShadowRendering.cpp b/Engine/Source/Runtime/Renderer/Private/ShadowRendering.cpp
--- a/Engine/Source/Runtime/Renderer/Private/ShadowRendering.cpp
+++ b/Engine/Source/Runtime/Renderer/Private/ShadowRendering.cpp
@@ -605,21 +605,24 @@ void FProjectedShadowInfo::SetupProjectionStencilMask(
 		SCOPED_DRAW_EVENTF(RHICmdList, EventMaskSubjects, TEXT("Stencil Mask Subjects"));
 
 		// If instanced stereo is enabled, we need to render each view of the stereo pair using the instanced stereo transform to avoid bias issues.
+		// TODO: Support instanced stereo properly in the projection stenciling pass.
 		const bool bIsInstancedStereoEmulated = View->bIsInstancedStereoEnabled && !View->bIsMultiViewEnabled && View->StereoPass != eSSP_FULL;
 		if (bIsInstancedStereoEmulated)
 		{
 			RHICmdList.SetViewport(0, 0, 0, SceneRender->InstancedStereoWidth, View->ViewRect.Max.Y, 1);
+			RHICmdList.SetScissorRect(true, View->ViewRect.Min.X, View->ViewRect.Min.Y, View->ViewRect.Max.X, View->ViewRect.Max.Y);
 		}
 
 		const FShadowMeshDrawCommandPass& ProjectionStencilingPass = ProjectionStencilingPasses[ViewIndex];
 		if (ProjectionStencilingPass.VisibleMeshDrawCommands.Num() > 0)
 		{
-			SubmitMeshDrawCommands(ProjectionStencilingPass.VisibleMeshDrawCommands, ProjectionStencilingPass.PrimitiveIdVertexBuffer, 0, bDynamicInstancing, 1, RHICmdList);
+			SubmitMeshDrawCommands(ProjectionStencilingPass.VisibleMeshDrawCommands, ProjectionStencilingPass.PrimitiveIdVertexBuffer, 0, bDynamicInstancing, bIsInstancedStereoEmulated ? 2 : 1, RHICmdList);
 		}
 
 		// Restore viewport
 		if (bIsInstancedStereoEmulated)
 		{
+			RHICmdList.SetScissorRect(false, 0, 0, 0, 0);
 			RHICmdList.SetViewport(View->ViewRect.Min.X, View->ViewRect.Min.Y, 0.0f, View->ViewRect.Max.X, View->ViewRect.Max.Y, 1.0f);
 		}
 		
diff --git a/Engine/Source/Runtime/Renderer/Private/ShadowSetup.cpp b/Engine/Source/Runtime/Renderer/Private/ShadowSetup.cpp
--- a/Engine/Source/Runtime/Renderer/Private/ShadowSetup.cpp
+++ b/Engine/Source/Runtime/Renderer/Private/ShadowSetup.cpp
@@ -1304,7 +1304,10 @@ void FProjectedShadowInfo::SetupMeshDrawCommandsForProjectionStenciling(FSceneRe
 
 			ApplyViewOverridesToMeshDrawCommands(View, ProjectionStencilingPass.VisibleMeshDrawCommands);
 
-			SortAndMergeDynamicPassMeshDrawCommands(Renderer.FeatureLevel, ProjectionStencilingPass.VisibleMeshDrawCommands, DynamicMeshDrawCommandStorage, ProjectionStencilingPass.PrimitiveIdVertexBuffer, 1);
+			// If instanced stereo is enabled, we need to render each view of the stereo pair using the instanced stereo transform to avoid bias issues.
+			// TODO: Support instanced stereo properly in the projection stenciling pass.
+			const uint32 InstanceFactor = View.bIsInstancedStereoEnabled && !View.bIsMultiViewEnabled && View.StereoPass != eSSP_FULL ? 2 : 1;
+			SortAndMergeDynamicPassMeshDrawCommands(Renderer.FeatureLevel, ProjectionStencilingPass.VisibleMeshDrawCommands, DynamicMeshDrawCommandStorage, ProjectionStencilingPass.PrimitiveIdVertexBuffer, InstanceFactor);
 		}
 	}
 }
