Set the instanced uniform buffer for mobile multi-view

diff --git a/Engine/Source/Runtime/Engine/Private/HierarchicalInstancedStaticMesh.cpp b/Engine/Source/Runtime/Engine/Private/HierarchicalInstancedStaticMesh.cpp
--- a/Engine/Source/Runtime/Engine/Private/HierarchicalInstancedStaticMesh.cpp
+++ b/Engine/Source/Runtime/Engine/Private/HierarchicalInstancedStaticMesh.cpp
@@ -1522,7 +1522,7 @@ void FHierarchicalStaticMeshSceneProxy::GetDynamicMeshElements(const TArray<cons
 				{
 					// Instanced stereo needs to use the right plane from the right eye when constructing the frustum bounds to cull against.
 					// Otherwise we'll cull objects visible in the right eye, but not the left.
-					if (Views[0]->IsInstancedStereoPass() && ViewIndex == 0)
+					if ((Views[0]->IsInstancedStereoPass() || Views[0]->bIsMobileMultiViewEnabled) && ViewIndex == 0)
 					{
 						check(Views.Num() == 2);
 						
diff --git a/Engine/Source/Runtime/Renderer/Private/MobileBasePass.cpp b/Engine/Source/Runtime/Renderer/Private/MobileBasePass.cpp
--- a/Engine/Source/Runtime/Renderer/Private/MobileBasePass.cpp
+++ b/Engine/Source/Runtime/Renderer/Private/MobileBasePass.cpp
@@ -679,6 +679,7 @@ FMeshPassProcessor* CreateMobileBasePassProcessor(const FScene* Scene, const FSc
 	FExclusiveDepthStencil::Type DefaultBasePassDepthStencilAccess = GetDefaultBasePassDepthStencilAccess(EShadingPath::Mobile);
 	
 	FMeshPassProcessorRenderState PassDrawRenderState(Scene->UniformBuffers.ViewUniformBuffer, Scene->UniformBuffers.MobileOpaqueBasePassUniformBuffer);
+	PassDrawRenderState.SetInstancedViewUniformBuffer(Scene->UniformBuffers.InstancedViewUniformBuffer);
 	PassDrawRenderState.SetBlendState(TStaticBlendStateWriteMask<CW_RGBA>::GetRHI());
 	PassDrawRenderState.SetDepthStencilAccess(DefaultBasePassDepthStencilAccess);
 	PassDrawRenderState.SetDepthStencilState(TStaticDepthStencilState<true, CF_DepthNearOrEqual>::GetRHI());
@@ -691,6 +692,7 @@ FMeshPassProcessor* CreateMobileBasePassCSMProcessor(const FScene* Scene, const
 	FExclusiveDepthStencil::Type DefaultBasePassDepthStencilAccess = GetDefaultBasePassDepthStencilAccess(EShadingPath::Mobile);
 
 	FMeshPassProcessorRenderState PassDrawRenderState(Scene->UniformBuffers.ViewUniformBuffer, Scene->UniformBuffers.MobileOpaqueBasePassUniformBuffer);
+	PassDrawRenderState.SetInstancedViewUniformBuffer(Scene->UniformBuffers.InstancedViewUniformBuffer);
 	PassDrawRenderState.SetBlendState(TStaticBlendStateWriteMask<CW_RGBA>::GetRHI());
 	PassDrawRenderState.SetDepthStencilAccess(DefaultBasePassDepthStencilAccess);
 	PassDrawRenderState.SetDepthStencilState(TStaticDepthStencilState<true, CF_DepthNearOrEqual>::GetRHI());
@@ -701,6 +703,7 @@ FMeshPassProcessor* CreateMobileBasePassCSMProcessor(const FScene* Scene, const
 FMeshPassProcessor* CreateMobileTranslucencyStandardPassProcessor(const FScene* Scene, const FSceneView* InViewIfDynamicMeshCommand, FMeshPassDrawListContext* InDrawListContext)
 {
 	FMeshPassProcessorRenderState PassDrawRenderState(Scene->UniformBuffers.ViewUniformBuffer, Scene->UniformBuffers.MobileTranslucentBasePassUniformBuffer);
+	PassDrawRenderState.SetInstancedViewUniformBuffer(Scene->UniformBuffers.InstancedViewUniformBuffer);
 	PassDrawRenderState.SetDepthStencilState(TStaticDepthStencilState<false, CF_DepthNearOrEqual>::GetRHI());
 	PassDrawRenderState.SetDepthStencilAccess(FExclusiveDepthStencil::DepthRead_StencilRead);
 
@@ -710,6 +713,7 @@ FMeshPassProcessor* CreateMobileTranslucencyStandardPassProcessor(const FScene*
 FMeshPassProcessor* CreateMobileTranslucencyAfterDOFProcessor(const FScene* Scene, const FSceneView* InViewIfDynamicMeshCommand, FMeshPassDrawListContext* InDrawListContext)
 {
 	FMeshPassProcessorRenderState PassDrawRenderState(Scene->UniformBuffers.ViewUniformBuffer, Scene->UniformBuffers.MobileTranslucentBasePassUniformBuffer);
+	PassDrawRenderState.SetInstancedViewUniformBuffer(Scene->UniformBuffers.InstancedViewUniformBuffer);
 	PassDrawRenderState.SetDepthStencilState(TStaticDepthStencilState<false, CF_DepthNearOrEqual>::GetRHI());
 	PassDrawRenderState.SetDepthStencilAccess(FExclusiveDepthStencil::DepthRead_StencilRead);
 
@@ -719,6 +723,7 @@ FMeshPassProcessor* CreateMobileTranslucencyAfterDOFProcessor(const FScene* Scen
 FMeshPassProcessor* CreateMobileTranslucencyAllPassProcessor(const FScene* Scene, const FSceneView* InViewIfDynamicMeshCommand, FMeshPassDrawListContext* InDrawListContext)
 {
 	FMeshPassProcessorRenderState PassDrawRenderState(Scene->UniformBuffers.ViewUniformBuffer, Scene->UniformBuffers.MobileTranslucentBasePassUniformBuffer);
+	PassDrawRenderState.SetInstancedViewUniformBuffer(Scene->UniformBuffers.InstancedViewUniformBuffer);
 	PassDrawRenderState.SetDepthStencilState(TStaticDepthStencilState<false, CF_DepthNearOrEqual>::GetRHI());
 	PassDrawRenderState.SetDepthStencilAccess(FExclusiveDepthStencil::DepthRead_StencilRead);
 
diff --git a/Engine/Source/Runtime/Renderer/Private/RendererScene.cpp b/Engine/Source/Runtime/Renderer/Private/RendererScene.cpp
--- a/Engine/Source/Runtime/Renderer/Private/RendererScene.cpp
+++ b/Engine/Source/Runtime/Renderer/Private/RendererScene.cpp
@@ -914,7 +914,7 @@ bool FPersistentUniformBuffers::UpdateViewUniformBuffer(const FViewInfo& View)
 	{
 		ViewUniformBuffer.UpdateUniformBufferImmediate(*View.CachedViewUniformShaderParameters);
 
-		if (View.IsInstancedStereoPass() && View.Family->Views.Num() > 0)
+		if ((View.IsInstancedStereoPass() || View.bIsMobileMultiViewEnabled) && View.Family->Views.Num() > 0)
 		{
 			// When drawing the left eye in a stereo scene, copy the right eye view values into the instanced view uniform buffer.
 			const EStereoscopicPass StereoPassIndex = (View.StereoPass != eSSP_FULL) ? eSSP_RIGHT_EYE : eSSP_FULL;
diff --git a/Engine/Source/Runtime/Renderer/Private/TranslucentRendering.cpp b/Engine/Source/Runtime/Renderer/Private/TranslucentRendering.cpp
--- a/Engine/Source/Runtime/Renderer/Private/TranslucentRendering.cpp
+++ b/Engine/Source/Runtime/Renderer/Private/TranslucentRendering.cpp
@@ -911,7 +911,7 @@ void FDeferredShadingSceneRenderer::RenderTranslucency(FRHICommandListImmediate&
 				Scene->UniformBuffers.ViewUniformBuffer.UpdateUniformBufferImmediate(DownsampledTranslucencyViewParameters);
 				DrawRenderState.SetViewUniformBuffer(Scene->UniformBuffers.ViewUniformBuffer);
 
-				if (View.IsInstancedStereoPass() && View.Family->Views.Num() > 0)
+				if ((View.IsInstancedStereoPass() || View.bIsMobileMultiViewEnabled) && View.Family->Views.Num() > 0)
 				{
 					// When drawing the left eye in a stereo scene, copy the right eye view values into the instanced view uniform buffer.
 					const EStereoscopicPass StereoPassIndex = (View.StereoPass != eSSP_FULL) ? eSSP_RIGHT_EYE : eSSP_FULL;
