Split grow only FGraphicsMinimalPipelineStateId cache into two caches: one for persistent Ids, and one for one frame Ids. One frame Id cache is cleared every frame. This reduces total FGraphicsMinimalPipelineStateId cache memory usage and fixes memory leaking for dynamic draws. Leaking here was caused by caching pointers to RHI resources, which may be later recreated and same PSO wil be re-added to the cache without removing previous one.

diff --git a/Engine/Source/Runtime/Renderer/Private/MeshDrawCommands.cpp b/Engine/Source/Runtime/Renderer/Private/MeshDrawCommands.cpp
--- a/Engine/Source/Runtime/Renderer/Private/MeshDrawCommands.cpp
+++ b/Engine/Source/Runtime/Renderer/Private/MeshDrawCommands.cpp
@@ -181,7 +181,7 @@ void UpdateTranslucentMeshSortKeys(
 	}
 }
 
-static uint64 GetMobileBasePassSortKey_FrontToBack(bool bMasked, bool bBackground, int32 PipelineId, int32 StateBucketId, float PrimitiveDistance)
+static uint64 GetMobileBasePassSortKey_FrontToBack(bool bMasked, bool bBackground, uint32 PipelineId, int32 StateBucketId, float PrimitiveDistance)
 {
 	union
 	{
@@ -304,14 +304,14 @@ void UpdateMobileBasePassMeshSortKeys(
 				bBackground|= (PrimitiveBounds.BoxSphereBounds.SphereRadius > HALF_WORLD_MAX / 4.0f);
 			}
 
-			int32 PipelineId = Cmd.MeshDrawCommand->CachedPipelineId.GetId();
+			uint32 PipelineId = Cmd.MeshDrawCommand->CachedPipelineId.GetId();
 			int32 StateBucketId = PointerHash(Cmd.MeshDrawCommand->IndexBuffer);
 			Cmd.SortKey.PackedData = GetMobileBasePassSortKey_FrontToBack(bMasked, bBackground, PipelineId, StateBucketId, PrimitiveDistance);
 		}
 	}
 	else // prefer state then distance
 	{
-		TMap<int32, float> PipelineDistances;
+		TMap<uint32, float> PipelineDistances;
 		PipelineDistances.Reserve(256);
 				
 		// pre-compute distance to a group of meshes that share same PSO
@@ -679,7 +679,8 @@ void ApplyViewOverridesToMeshDrawCommands(
 					PipelineState.DepthStencilState = PassDrawRenderState.GetDepthStencilState();
 				}
 
-				NewMeshCommand.Finalize(PipelineState, nullptr, true);
+				const FGraphicsMinimalPipelineStateId PipelineId = FGraphicsMinimalPipelineStateId::GetOneFrameId(PipelineState);
+				NewMeshCommand.Finalize(PipelineId, nullptr);
 
 				FVisibleMeshDrawCommand NewVisibleMeshDrawCommand;
 
diff --git a/Engine/Source/Runtime/Renderer/Private/MeshPassProcessor.cpp b/Engine/Source/Runtime/Renderer/Private/MeshPassProcessor.cpp
--- a/Engine/Source/Runtime/Renderer/Private/MeshPassProcessor.cpp
+++ b/Engine/Source/Runtime/Renderer/Private/MeshPassProcessor.cpp
@@ -14,8 +14,9 @@
 #include "MeshPassProcessor.inl"
 #include "PipelineStateCache.h"
 
-TSet<FGraphicsMinimalPipelineStateInitializer> FGraphicsMinimalPipelineStateId::GlobalTable;
-FCriticalSection GlobalTableCriticalSection;
+TSet<FGraphicsMinimalPipelineStateInitializer> FGraphicsMinimalPipelineStateId::PersistentIdTable;
+TSet<FGraphicsMinimalPipelineStateInitializer> FGraphicsMinimalPipelineStateId::OneFrameIdTable;
+FCriticalSection FGraphicsMinimalPipelineStateId::OneFrameIdTableCriticalSection;
 
 const FMeshDrawCommandSortKey FMeshDrawCommandSortKey::Default = { {0} };
 
@@ -297,29 +298,62 @@ void FMeshDrawShaderBindings::SetOnRayTracingStructure(FRHICommandList& RHICmdLi
 }
 #endif // RHI_RAYTRACING
 
-void FGraphicsMinimalPipelineStateId::Setup(const FGraphicsMinimalPipelineStateInitializer& InPipelineState)
+FGraphicsMinimalPipelineStateId FGraphicsMinimalPipelineStateId::GetPersistentId(const FGraphicsMinimalPipelineStateInitializer& InPipelineState)
+{
+	FSetElementId TableId = PersistentIdTable.FindId(InPipelineState);
+
+	if (!TableId.IsValidId())
+	{
+		// Note: grow-only table. Assuming finite and small enough set of FGraphicsMinimalPipelineStateInitializer permutations.
+		TableId = PersistentIdTable.Add(InPipelineState);
+	}
+
+	checkf(TableId.AsInteger() < (MAX_uint32 >> 2), TEXT("Persistent FGraphicsMinimalPipelineStateId table overflow!"));
+
+	FGraphicsMinimalPipelineStateId Ret;
+	Ret.bValid = 1;
+	Ret.bOneFrameId = 0;
+	Ret.SetElementIndex = TableId.AsInteger();
+	return Ret;
+}
+
+FGraphicsMinimalPipelineStateId FGraphicsMinimalPipelineStateId::GetOneFrameId(const FGraphicsMinimalPipelineStateInitializer& InPipelineState)
 {
 	// Need to lock as this is called from multiple parallel tasks during mesh draw command generation or patching.
-	FScopeLock Lock(&GlobalTableCriticalSection);
+	FScopeLock Lock(&OneFrameIdTableCriticalSection);
 
-	FSetElementId TableId = GlobalTable.FindId(InPipelineState);
+	FGraphicsMinimalPipelineStateId Ret;
+	Ret.bValid = 1;
+	Ret.bOneFrameId = 0;
 
+	FSetElementId TableId = PersistentIdTable.FindId(InPipelineState);
 	if (!TableId.IsValidId())
 	{
-		// Note: grow-only table.  Assuming finite and small enough set of FGraphicsMinimalPipelineStateInitializer permutations.
-		TableId = GlobalTable.Add(InPipelineState);
+		Ret.bOneFrameId = 1;
+
+		TableId = OneFrameIdTable.FindId(InPipelineState);
+		if (!TableId.IsValidId())
+		{
+			TableId = OneFrameIdTable.Add(InPipelineState);
+		}
 	}
 
-	checkf(TableId.AsInteger() < MAX_int32, TEXT("FGraphicsMinimalPipelineStateId overflow!"));
+	checkf(TableId.AsInteger() < (MAX_uint32 >> 2), TEXT("One frame FGraphicsMinimalPipelineStateId table overflow!"));
 
-	Id = TableId;
+	Ret.SetElementIndex = TableId.AsInteger();
+	return Ret;
+}
+
+void FGraphicsMinimalPipelineStateId::ResetOneFrameIdTable()
+{
+	OneFrameIdTable.Reset();
 }
 
 class FMeshDrawCommandStateCache
 {
 public:
 
-	int32 PipelineId;
+	uint32 PipelineId;
 	uint32 StencilRef;
 	FShaderBindingState ShaderBindings[SF_NumStandardFrequencies];
 	FVertexInputStream VertexStreams[MaxVertexElementCount];
@@ -634,9 +668,8 @@ void FMeshDrawCommand::SetRayTracingShaders(const FMeshProcessorShaders& Shaders
 void FMeshDrawCommand::SetDrawParametersAndFinalize(
 	const FMeshBatch& MeshBatch, 
 	int32 BatchElementIndex,
-	const FGraphicsMinimalPipelineStateInitializer& PipelineState, 
-	const FMeshProcessorShaders* ShadersForDebugging,
-	bool bDoSetupPsoStateForRasterization)
+	FGraphicsMinimalPipelineStateId PipelineId,
+	const FMeshProcessorShaders* ShadersForDebugging)
 {
 	const FMeshBatchElement& BatchElement = MeshBatch.Elements[BatchElementIndex];
 
@@ -662,7 +695,7 @@ void FMeshDrawCommand::SetDrawParametersAndFinalize(
 	int32 SegmentIndex = MeshBatch.SegmentIndex + BatchElementIndex;
 	RayTracedSegmentIndex = (SegmentIndex < UINT8_MAX) ? uint8(MeshBatch.SegmentIndex + BatchElementIndex) : UINT8_MAX;
 
-	Finalize(PipelineState, ShadersForDebugging, bDoSetupPsoStateForRasterization);
+	Finalize(PipelineId, ShadersForDebugging);
 }
 
 void FMeshDrawShaderBindings::SetOnCommandList(FRHICommandList& RHICmdList, FBoundShaderStateInput Shaders, FShaderBindingState* StateCacheShaderBindings) const
@@ -1062,7 +1095,13 @@ void FCachedPassMeshDrawListContext::FinalizeCommand(
 {
 	QUICK_SCOPE_CYCLE_COUNTER(STAT_FinalizeCachedMeshDrawCommand);
 
-	MeshDrawCommand.SetDrawParametersAndFinalize(MeshBatch, BatchElementIndex, PipelineState, ShadersForDebugging, bDoSetupPsoStateForRasterization);
+	FGraphicsMinimalPipelineStateId PipelineId;
+	if (bDoSetupPsoStateForRasterization)
+	{
+		PipelineId = FGraphicsMinimalPipelineStateId::GetPersistentId(PipelineState);
+	}
+
+	MeshDrawCommand.SetDrawParametersAndFinalize(MeshBatch, BatchElementIndex, PipelineId, ShadersForDebugging);
 
 	if (bUseStateBuckets)
 	{
diff --git a/Engine/Source/Runtime/Renderer/Private/RendererScene.cpp b/Engine/Source/Runtime/Renderer/Private/RendererScene.cpp
--- a/Engine/Source/Runtime/Renderer/Private/RendererScene.cpp
+++ b/Engine/Source/Runtime/Renderer/Private/RendererScene.cpp
@@ -669,7 +669,8 @@ void FScene::DumpMeshDrawCommandMemoryStats()
 	UE_LOG(LogRenderer, Log, TEXT("Total cached MeshDrawCommands %.3fMb"), TotalCachedMeshDrawCommands / 1024.0f / 1024.0f);
 	UE_LOG(LogRenderer, Log, TEXT("Primitive StaticMeshCommandInfos %.1fKb"), TotalStaticMeshCommandInfos / 1024.0f);
 	UE_LOG(LogRenderer, Log, TEXT("GPUScene CPU structures %.1fKb"), GPUScene.PrimitivesToUpdate.GetAllocatedSize() / 1024.0f);
-	UE_LOG(LogRenderer, Log, TEXT("PSO Id Grow-Only Table %.1fKb"), FGraphicsMinimalPipelineStateId::GetGlobalTableSize() / 1024.0f);
+	UE_LOG(LogRenderer, Log, TEXT("PSO persistent Id table %.1fKb %d elements"), FGraphicsMinimalPipelineStateId::GetPersistentIdTableSize() / 1024.0f, FGraphicsMinimalPipelineStateId::GetPersistentIdNum());
+	UE_LOG(LogRenderer, Log, TEXT("PSO one frame Id %.1fKb"), FGraphicsMinimalPipelineStateId::GetOneFrameIdTableSize() / 1024.0f);
 }
 
 template<typename T>
diff --git a/Engine/Source/Runtime/Renderer/Private/SceneRendering.cpp b/Engine/Source/Runtime/Renderer/Private/SceneRendering.cpp
--- a/Engine/Source/Runtime/Renderer/Private/SceneRendering.cpp
+++ b/Engine/Source/Runtime/Renderer/Private/SceneRendering.cpp
@@ -2897,6 +2897,7 @@ void FSceneRenderer::WaitForTasksClearSnapshotsAndDeleteSceneRenderer(FRHIComman
 	}
 
 	GPrimitiveIdVertexBufferPool.DiscardAll();
+	FGraphicsMinimalPipelineStateId::ResetOneFrameIdTable();
 
 	// Destroy cached preshadow transient arrays (allocated with SceneRenderingAllocator).
 	TArray<TRefCountPtr<FProjectedShadowInfo>>& CachedPreshadows = SceneRenderer->Scene->CachedPreshadows;
diff --git a/Engine/Source/Runtime/Renderer/Private/ShadowSetup.cpp b/Engine/Source/Runtime/Renderer/Private/ShadowSetup.cpp
--- a/Engine/Source/Runtime/Renderer/Private/ShadowSetup.cpp
+++ b/Engine/Source/Runtime/Renderer/Private/ShadowSetup.cpp
@@ -1334,7 +1334,8 @@ void FProjectedShadowInfo::ApplyViewOverridesToMeshDrawCommands(const FViewInfo&
 			FGraphicsMinimalPipelineStateInitializer PipelineState = MeshCommand.CachedPipelineId.GetPipelineState();
 			PipelineState.RasterizerState = GetStaticRasterizerState<true>(VisibleMeshDrawCommand.MeshFillMode, LocalCullMode);
 
-			NewMeshCommand.Finalize(PipelineState, nullptr, true);
+			const FGraphicsMinimalPipelineStateId PipelineId = FGraphicsMinimalPipelineStateId::GetOneFrameId(PipelineState);
+			NewMeshCommand.Finalize(PipelineId, nullptr);
 
 			FVisibleMeshDrawCommand NewVisibleMeshDrawCommand;
 
diff --git a/Engine/Source/Runtime/Renderer/Public/MeshPassProcessor.h b/Engine/Source/Runtime/Renderer/Public/MeshPassProcessor.h
--- a/Engine/Source/Runtime/Renderer/Public/MeshPassProcessor.h
+++ b/Engine/Source/Runtime/Renderer/Public/MeshPassProcessor.h
@@ -98,22 +98,20 @@ static_assert(sizeof(FMeshPassMask::Data) * 8 >= EMeshPass::Num, "FMeshPassMask:
 class FGraphicsMinimalPipelineStateId
 {
 public:
-	FORCEINLINE_DEBUGGABLE int32 GetId() const
+	FORCEINLINE_DEBUGGABLE uint32 GetId() const
 	{
 		checkSlow(IsValid());
-		return Id.AsInteger();
+		return PackedId;
 	}
 
 	inline bool IsValid() const 
 	{
-		return Id.IsValidId();
+		return bValid != 0;
 	}
 
-	void Setup(const FGraphicsMinimalPipelineStateInitializer& InPipelineState);
-
 	inline bool operator==(const FGraphicsMinimalPipelineStateId& rhs) const
 	{
-		return Id == rhs.Id;
+		return PackedId == rhs.PackedId;
 	}
 
 	inline bool operator!=(const FGraphicsMinimalPipelineStateId& rhs) const
@@ -123,15 +121,47 @@ public:
 	
 	inline const FGraphicsMinimalPipelineStateInitializer& GetPipelineState() const
 	{
-		return GlobalTable[Id];
+		const FSetElementId SetElementId = FSetElementId::FromInteger(SetElementIndex);
+
+		if (bOneFrameId)
+		{
+			return OneFrameIdTable[SetElementId];
+		}
+
+		return PersistentIdTable[SetElementId];
 	}
 
-	static SIZE_T GetGlobalTableSize() { return GlobalTable.GetAllocatedSize(); }
+	/**
+	 * Get a pipeline id, which won't be ever released.
+	 */
+	static FGraphicsMinimalPipelineStateId GetPersistentId(const FGraphicsMinimalPipelineStateInitializer& InPipelineState);
+
+	/**
+	 * Get a pipeline id, which is valid only for a single frame.
+	 */
+	RENDERER_API static FGraphicsMinimalPipelineStateId GetOneFrameId(const FGraphicsMinimalPipelineStateInitializer& InPipelineState);
+
+	static void ResetOneFrameIdTable();
+	static SIZE_T GetPersistentIdTableSize() { return PersistentIdTable.GetAllocatedSize(); }
+	static int32 GetPersistentIdNum() { return PersistentIdTable.Num(); }
+	static SIZE_T GetOneFrameIdTableSize() { return OneFrameIdTable.GetAllocatedSize(); }
 
 private:
-	FSetElementId Id;
+	union
+	{
+		uint32 PackedId = 0;
 
-	static TSet<FGraphicsMinimalPipelineStateInitializer> GlobalTable;
+		struct
+		{
+			uint32 SetElementIndex	: 30;
+			uint32 bOneFrameId		: 1;
+			uint32 bValid			: 1;
+		};
+	};
+
+	static TSet<FGraphicsMinimalPipelineStateInitializer> PersistentIdTable;
+	static TSet<FGraphicsMinimalPipelineStateInitializer> OneFrameIdTable;
+	static FCriticalSection OneFrameIdTableCriticalSection;
 };
 
 struct FMeshProcessorShaders
@@ -436,17 +466,12 @@ public:
 	RENDERER_API void SetDrawParametersAndFinalize(
 		const FMeshBatch& MeshBatch, 
 		int32 BatchElementIndex,
-		const FGraphicsMinimalPipelineStateInitializer& PipelineState, 
-		const FMeshProcessorShaders* ShadersForDebugging,
-		bool bDoSetupPsoStateForRasterization);
+		FGraphicsMinimalPipelineStateId PipelineId,
+		const FMeshProcessorShaders* ShadersForDebugging);
 
-	void Finalize(const FGraphicsMinimalPipelineStateInitializer& PipelineState, const FMeshProcessorShaders* ShadersForDebugging, bool bDoSetupPsoStateForRasterization)
+	void Finalize(FGraphicsMinimalPipelineStateId PipelineId, const FMeshProcessorShaders* ShadersForDebugging)
 	{
-		if (bDoSetupPsoStateForRasterization)
-		{
-			CachedPipelineId.Setup(PipelineState);
-		}
-
+		CachedPipelineId = PipelineId;
 		ShaderBindings.Finalize(ShadersForDebugging);	
 	}
 
@@ -657,7 +682,13 @@ public:
 		FMeshDrawCommand& MeshDrawCommand,
 		bool bDoSetupPsoStateForRasterization) override final
 	{
-		MeshDrawCommand.SetDrawParametersAndFinalize(MeshBatch, BatchElementIndex, PipelineState, ShadersForDebugging, bDoSetupPsoStateForRasterization);
+		FGraphicsMinimalPipelineStateId PipelineId;
+		if (bDoSetupPsoStateForRasterization)
+		{
+			PipelineId = FGraphicsMinimalPipelineStateId::GetOneFrameId(PipelineState);
+		}
+
+		MeshDrawCommand.SetDrawParametersAndFinalize(MeshBatch, BatchElementIndex, PipelineId, ShadersForDebugging);
 
 		FVisibleMeshDrawCommand NewVisibleMeshDrawCommand;
 		//@todo MeshCommandPipeline - assign usable state ID for dynamic path draws
