UE4.22 - Remove almost all of ENQUEUE_UNIQUE_RENDER_COMMAND_TWOPARAMETER
#rnx

diff --git a/Engine/Plugins/Compositing/OpenColorIO/Source/OpenColorIO/Private/OpenColorIOShaderCompilationManager.cpp b/Engine/Plugins/Compositing/OpenColorIO/Source/OpenColorIO/Private/OpenColorIOShaderCompilationManager.cpp
--- a/Engine/Plugins/Compositing/OpenColorIO/Source/OpenColorIO/Private/OpenColorIOShaderCompilationManager.cpp
+++ b/Engine/Plugins/Compositing/OpenColorIO/Source/OpenColorIO/Private/OpenColorIOShaderCompilationManager.cpp
@@ -396,12 +396,10 @@ void FOpenColorIOShaderCompilationManager::ProcessCompiledOpenColorIOShaderMaps(
 
 			ColorTransform->SetGameThreadShaderMap(It.Value());
 
-			ENQUEUE_UNIQUE_RENDER_COMMAND_TWOPARAMETER(
-				FSetShaderMapOnColorTransformResources,
-				FOpenColorIOTransformResource*, ColorTransform, ColorTransform,
-				FOpenColorIOShaderMap*, CompiledShaderMap, ShaderMap,
+			ENQUEUE_RENDER_COMMAND(FSetShaderMapOnColorTransformResources)(
+				[ColorTransform, ShaderMap](FRHICommandListImmediate& RHICmdList)
 				{
-					ColorTransform->SetRenderingThreadShaderMap(CompiledShaderMap);
+					ColorTransform->SetRenderingThreadShaderMap(ShaderMap);
 				});
 
 
diff --git a/Engine/Plugins/Compositing/OpenColorIO/Source/OpenColorIO/Private/OpenColorIOShared.cpp b/Engine/Plugins/Compositing/OpenColorIO/Source/OpenColorIO/Private/OpenColorIOShared.cpp
--- a/Engine/Plugins/Compositing/OpenColorIO/Source/OpenColorIO/Private/OpenColorIOShared.cpp
+++ b/Engine/Plugins/Compositing/OpenColorIO/Source/OpenColorIO/Private/OpenColorIOShared.cpp
@@ -316,10 +316,10 @@ bool FOpenColorIOTransformResource::CacheShaders(const FOpenColorIOShaderMapId&
 		bSucceeded = true;
 	}
 
-	ENQUEUE_UNIQUE_RENDER_COMMAND_TWOPARAMETER(
-		FSetShaderMapOnColorSpaceTransformResources,
-		FOpenColorIOTransformResource*, ColorSpaceTransform, this,
-		FOpenColorIOShaderMap*, LoadedShaderMap, GameThreadShaderMap,
+	FOpenColorIOTransformResource* ColorSpaceTransform = this;
+	FOpenColorIOShaderMap* LoadedShaderMap = GameThreadShaderMap;
+	ENQUEUE_RENDER_COMMAND(FSetShaderMapOnColorSpaceTransformResources)(
+		[ColorSpaceTransform, LoadedShaderMap](FRHICommandListImmediate& RHICmdList)
 		{
 			ColorSpaceTransform->SetRenderingThreadShaderMap(LoadedShaderMap);
 		});
diff --git a/Engine/Plugins/FX/Niagara/Source/Niagara/Private/NiagaraComponent.cpp b/Engine/Plugins/FX/Niagara/Source/Niagara/Private/NiagaraComponent.cpp
--- a/Engine/Plugins/FX/Niagara/Source/Niagara/Private/NiagaraComponent.cpp
+++ b/Engine/Plugins/FX/Niagara/Source/Niagara/Private/NiagaraComponent.cpp
@@ -123,7 +123,8 @@ SIZE_T FNiagaraSceneProxy::GetTypeHash() const
 	static size_t UniquePointer;
 	return reinterpret_cast<size_t>(&UniquePointer);
 }
-void FNiagaraSceneProxy::UpdateEmitterRenderers(TArray<NiagaraRenderer*>& InRenderers)
+
+void FNiagaraSceneProxy::UpdateEmitterRenderers(const TArray<NiagaraRenderer*>& InRenderers)
 {
 	EmitterRenderers.Empty();
 	for (NiagaraRenderer* EmitterRenderer : InRenderers)
@@ -821,20 +822,18 @@ void UNiagaraComponent::SendRenderDynamicData_Concurrent()
 					if (bRendererEditorEnabled && !Emitter->IsComplete() && !SystemInstance->IsComplete())
 					{
 						FNiagaraDynamicDataBase* DynamicData = Renderer->GenerateVertexData(NiagaraProxy, Emitter->GetData(), Emitter->GetEmitterHandle().GetInstance()->SimTarget);
-
-						ENQUEUE_UNIQUE_RENDER_COMMAND_TWOPARAMETER(
-							FSendNiagaraDynamicData,
-							NiagaraRenderer*, EmitterRenderer, Emitter->GetEmitterRenderer(EmitterIdx),
-							FNiagaraDynamicDataBase*, DynamicData, DynamicData,
+						NiagaraRenderer* EmitterRenderer = Emitter->GetEmitterRenderer(EmitterIdx);
+						ENQUEUE_RENDER_COMMAND(FSendNiagaraDynamicData)(
+							[EmitterRenderer, DynamicData](FRHICommandListImmediate& RHICmdList)
 							{
 								EmitterRenderer->SetDynamicData_RenderThread(DynamicData);
 							});
 					}
 					else
 					{
-						ENQUEUE_UNIQUE_RENDER_COMMAND_ONEPARAMETER(
-							FSendNiagaraDynamicData,
-							NiagaraRenderer*, EmitterRenderer, Emitter->GetEmitterRenderer(EmitterIdx),
+						NiagaraRenderer* EmitterRenderer = Emitter->GetEmitterRenderer(EmitterIdx);
+						ENQUEUE_RENDER_COMMAND(FSendNiagaraDynamicData)(
+							[EmitterRenderer](FRHICommandListImmediate& RHICmdList)
 							{
 								EmitterRenderer->SetDynamicData_RenderThread(nullptr);
 							});
diff --git a/Engine/Plugins/FX/Niagara/Source/Niagara/Private/NiagaraSystemInstance.cpp b/Engine/Plugins/FX/Niagara/Source/Niagara/Private/NiagaraSystemInstance.cpp
--- a/Engine/Plugins/FX/Niagara/Source/Niagara/Private/NiagaraSystemInstance.cpp
+++ b/Engine/Plugins/FX/Niagara/Source/Niagara/Private/NiagaraSystemInstance.cpp
@@ -1280,12 +1280,11 @@ void FNiagaraSystemInstance::UpdateProxy(TArray<NiagaraRenderer*>& InRenderers)
 		if (Component->GetWorld() != nullptr)
 		{
 			// Tell the scene proxy on the render thread to update its System renderers.
-			ENQUEUE_UNIQUE_RENDER_COMMAND_TWOPARAMETER(
-				FChangeNiagaraRenderModule,
-				FNiagaraSceneProxy*, InProxy, NiagaraProxy,
-				TArray<NiagaraRenderer*>, InRendererArray, InRenderers,
+			TArray<NiagaraRenderer*> InRendererArray = InRenderers;
+			ENQUEUE_RENDER_COMMAND(FChangeNiagaraRenderModule)(
+				[NiagaraProxy, InRendererArray](FRHICommandListImmediate& RHICmdList)
 				{
-					InProxy->UpdateEmitterRenderers(InRendererArray);
+					NiagaraProxy->UpdateEmitterRenderers(InRendererArray);
 				}
 			);
 		}
diff --git a/Engine/Plugins/FX/Niagara/Source/Niagara/Public/NiagaraComponent.h b/Engine/Plugins/FX/Niagara/Source/Niagara/Public/NiagaraComponent.h
--- a/Engine/Plugins/FX/Niagara/Source/Niagara/Public/NiagaraComponent.h
+++ b/Engine/Plugins/FX/Niagara/Source/Niagara/Public/NiagaraComponent.h
@@ -425,7 +425,7 @@ public:
 	/** Called on render thread to assign new dynamic data */
 	void SetDynamicData_RenderThread(struct FNiagaraDynamicDataBase* NewDynamicData);
 	TArray<class NiagaraRenderer*>& GetEmitterRenderers() { return EmitterRenderers; }
-	void UpdateEmitterRenderers(TArray<NiagaraRenderer*>& InRenderers);
+	void UpdateEmitterRenderers(const TArray<NiagaraRenderer*>& InRenderers);
 
 	/** Gets whether or not this scene proxy should be rendered. */
 	bool GetRenderingEnabled() const;
diff --git a/Engine/Plugins/FX/Niagara/Source/NiagaraShader/Private/NiagaraShaderCompilationManager.cpp b/Engine/Plugins/FX/Niagara/Source/NiagaraShader/Private/NiagaraShaderCompilationManager.cpp
--- a/Engine/Plugins/FX/Niagara/Source/NiagaraShader/Private/NiagaraShaderCompilationManager.cpp
+++ b/Engine/Plugins/FX/Niagara/Source/NiagaraShader/Private/NiagaraShaderCompilationManager.cpp
@@ -398,12 +398,10 @@ void FNiagaraShaderCompilationManager::ProcessCompiledNiagaraShaderMaps(
 
 			Script->SetGameThreadShaderMap(It.Value());
 
-			ENQUEUE_UNIQUE_RENDER_COMMAND_TWOPARAMETER(
-				FSetShaderMapOnScriptResources,
-				FNiagaraShaderScript*, Script, Script,
-				FNiagaraShaderMap*, CompiledShaderMap, ShaderMap,
+			ENQUEUE_RENDER_COMMAND(FSetShaderMapOnScriptResources)(
+				[Script, ShaderMap](FRHICommandListImmediate& RHICmdList)
 				{
-					Script->SetRenderingThreadShaderMap(CompiledShaderMap);
+					Script->SetRenderingThreadShaderMap(ShaderMap);
 				});
 
 			// Can't call NotifyCompilationFinished() when post-loading. 
diff --git a/Engine/Plugins/FX/Niagara/Source/NiagaraShader/Private/NiagaraShared.cpp b/Engine/Plugins/FX/Niagara/Source/NiagaraShader/Private/NiagaraShared.cpp
--- a/Engine/Plugins/FX/Niagara/Source/NiagaraShader/Private/NiagaraShared.cpp
+++ b/Engine/Plugins/FX/Niagara/Source/NiagaraShader/Private/NiagaraShared.cpp
@@ -328,11 +328,10 @@ bool FNiagaraShaderScript::CacheShaders(const FNiagaraShaderMapId& ShaderMapId,
 		bSucceeded = true;
 	}
 
-
-	ENQUEUE_UNIQUE_RENDER_COMMAND_TWOPARAMETER(
-		FSetShaderMapOnScriptResources,
-		FNiagaraShaderScript*, Script, this,
-		FNiagaraShaderMap*, LoadedShaderMap, GameThreadShaderMap,
+	FNiagaraShaderScript* Script = this;
+	FNiagaraShaderMap* LoadedShaderMap = GameThreadShaderMap;
+	ENQUEUE_RENDER_COMMAND(FSetShaderMapOnScriptResources)(
+		[Script, LoadedShaderMap](FRHICommandListImmediate& RHICmdList)
 		{
 			Script->SetRenderingThreadShaderMap(LoadedShaderMap);
 		});
diff --git a/Engine/Plugins/Media/WebMMedia/Source/WebMMedia/Private/Player/WebMVideoDecoder.cpp b/Engine/Plugins/Media/WebMMedia/Source/WebMMedia/Private/Player/WebMVideoDecoder.cpp
--- a/Engine/Plugins/Media/WebMMedia/Source/WebMMedia/Private/Player/WebMVideoDecoder.cpp
+++ b/Engine/Plugins/Media/WebMMedia/Source/WebMMedia/Private/Player/WebMVideoDecoder.cpp
@@ -132,14 +132,14 @@ void FWebMVideoDecoder::DoDecodeVideoFrames(const TArray<TSharedPtr<FWebMFrame>>
 		const void* ImageIter = nullptr;
 		while (const vpx_image_t* Image = vpx_codec_get_frame(&Context, &ImageIter))
 		{
+			FWebMVideoDecoder* Self = this;
 			if (!bTexturesCreated)
 			{
 				// First creation of conversion textures
 
 				bTexturesCreated = true;
-
-				ENQUEUE_UNIQUE_RENDER_COMMAND_TWOPARAMETER(WebMMediaPlayerCreateTextures,
-					FWebMVideoDecoder*, Self, this, const vpx_image_t*, Image, Image,
+				ENQUEUE_RENDER_COMMAND(WebMMediaPlayerCreateTextures)(
+					[Self, Image](FRHICommandListImmediate& RHICmdList)
 					{
 						Self->CreateTextures(Image);
 					});
@@ -152,9 +152,8 @@ void FWebMVideoDecoder::DoDecodeVideoFrames(const TArray<TSharedPtr<FWebMFrame>>
 			FConvertParams Params;
 			Params.VideoSample = VideoSample;
 			Params.Image = Image;
-
-			ENQUEUE_UNIQUE_RENDER_COMMAND_TWOPARAMETER(WebMMediaPlayerConvertYUVToRGB,
-				FWebMVideoDecoder*, Self, this, FConvertParams, Params, Params,
+			ENQUEUE_RENDER_COMMAND(WebMMediaPlayerConvertYUVToRGB)(
+				[Self, Params](FRHICommandListImmediate& RHICmdList)
 				{
 					Self->ConvertYUVToRGBAndSubmit(Params);
 				});
diff --git a/Engine/Plugins/Runtime/AR/Google/GoogleARCore/Source/GoogleARCoreBase/Private/GoogleARCoreXRCamera.cpp b/Engine/Plugins/Runtime/AR/Google/GoogleARCore/Source/GoogleARCoreBase/Private/GoogleARCoreXRCamera.cpp
--- a/Engine/Plugins/Runtime/AR/Google/GoogleARCore/Source/GoogleARCoreBase/Private/GoogleARCoreXRCamera.cpp
+++ b/Engine/Plugins/Runtime/AR/Google/GoogleARCore/Source/GoogleARCoreBase/Private/GoogleARCoreXRCamera.cpp
@@ -67,10 +67,9 @@ bool FGoogleARCoreXRCamera::IsActiveThisFrame(class FViewport* InViewport) const
 void FGoogleARCoreXRCamera::ConfigXRCamera(bool bInMatchDeviceCameraFOV, bool bInEnablePassthroughCameraRendering)
 {
 	bMatchDeviceCameraFOV = bInMatchDeviceCameraFOV;
-	ENQUEUE_UNIQUE_RENDER_COMMAND_TWOPARAMETER(
-		ConfigXRCamera,
-		FGoogleARCoreXRCamera*, ARCoreXRCamera, this,
-		bool, bInEnablePassthroughCameraRendering, bInEnablePassthroughCameraRendering,
+	FGoogleARCoreXRCamera* ARCoreXRCamera = this;
+	ENQUEUE_RENDER_COMMAND(ConfigXRCamera)(
+		[ARCoreXRCamera, bInEnablePassthroughCameraRendering](FRHICommandListImmediate& RHICmdList)
 		{
 			ARCoreXRCamera->bEnablePassthroughCameraRendering_RT = bInEnablePassthroughCameraRendering;
 		}
diff --git a/Engine/Plugins/Runtime/AndroidMoviePlayer/Source/AndroidMoviePlayer/Private/AndroidMovieStreamer.cpp b/Engine/Plugins/Runtime/AndroidMoviePlayer/Source/AndroidMoviePlayer/Private/AndroidMovieStreamer.cpp
--- a/Engine/Plugins/Runtime/AndroidMoviePlayer/Source/AndroidMoviePlayer/Private/AndroidMovieStreamer.cpp
+++ b/Engine/Plugins/Runtime/AndroidMoviePlayer/Source/AndroidMoviePlayer/Private/AndroidMovieStreamer.cpp
@@ -281,16 +281,16 @@ bool FAndroidMediaPlayerStreamer::StartNextMovie()
 
 		uint32 FrameBytes = VideoDimensions.X * VideoDimensions.Y * GPixelFormats[PF_B8G8R8A8].BlockBytes;
 
-		ENQUEUE_UNIQUE_RENDER_COMMAND_TWOPARAMETER(InitMovieTexture,
-			FSlateTexture2DRHIRef*, TextureRHIRef, Texture.Get(),
-			uint32, Bytes, FrameBytes,
+		FSlateTexture2DRHIRef* TextureRHIRef = Texture.Get();
+		ENQUEUE_RENDER_COMMAND(InitMovieTexture)(
+			[TextureRHIRef, FrameBytes](FRHICommandListImmediate& RHICmdList)
 			{
 				TextureRHIRef->InitResource();
 	
 				// clear texture to black
 				uint32 Stride = 0;
 				void* TextureBuffer = RHILockTexture2D(TextureRHIRef->GetTypedResource(), 0, RLM_WriteOnly, Stride, false);
-				FMemory::Memset(TextureBuffer, 0, Bytes);
+				FMemory::Memset(TextureBuffer, 0, FrameBytes);
 				RHIUnlockTexture2D(TextureRHIRef->GetTypedResource(), 0, false);
 			});
 
diff --git a/Engine/Plugins/Runtime/CableComponent/Source/CableComponent/Private/CableComponent.cpp b/Engine/Plugins/Runtime/CableComponent/Source/CableComponent/Private/CableComponent.cpp
--- a/Engine/Plugins/Runtime/CableComponent/Source/CableComponent/Private/CableComponent.cpp
+++ b/Engine/Plugins/Runtime/CableComponent/Source/CableComponent/Private/CableComponent.cpp
@@ -708,10 +708,9 @@ void UCableComponent::SendRenderDynamicData_Concurrent()
 		}
 
 		// Enqueue command to send to render thread
-		ENQUEUE_UNIQUE_RENDER_COMMAND_TWOPARAMETER(
-			FSendCableDynamicData,
-			FCableSceneProxy*,CableSceneProxy,(FCableSceneProxy*)SceneProxy,
-			FCableDynamicData*,DynamicData,DynamicData,
+		FCableSceneProxy* CableSceneProxy = (FCableSceneProxy*)SceneProxy;
+		ENQUEUE_RENDER_COMMAND(FSendCableDynamicData)(
+			[CableSceneProxy, DynamicData](FRHICommandListImmediate& RHICmdList)
 		{
 			CableSceneProxy->SetDynamicData_RenderThread(DynamicData);
 		});
diff --git a/Engine/Plugins/Runtime/GoogleVR/GoogleVRHMD/Source/GoogleVRHMD/Private/GoogleVRSplash.cpp b/Engine/Plugins/Runtime/GoogleVR/GoogleVRHMD/Source/GoogleVRHMD/Private/GoogleVRSplash.cpp
--- a/Engine/Plugins/Runtime/GoogleVR/GoogleVRHMD/Source/GoogleVRHMD/Private/GoogleVRSplash.cpp
+++ b/Engine/Plugins/Runtime/GoogleVR/GoogleVRHMD/Source/GoogleVRHMD/Private/GoogleVRSplash.cpp
@@ -110,13 +110,14 @@ void FGoogleVRSplash::Show()
 	//SplashScreenRenderingOrientation = FRotator(GVRHMD->CachedFinalHeadRotation);
 
 	RenderThreadTicker = MakeShareable(new FGoogleVRSplashTicker(this));
-	ENQUEUE_UNIQUE_RENDER_COMMAND_TWOPARAMETER(RegisterAsyncTick,
-	FTickableObjectRenderThread*, RenderThreadTicker, RenderThreadTicker.Get(),
-	FGoogleVRSplash*, pGVRSplash, this,
-	{
-		pGVRSplash->AllocateSplashScreenRenderTarget();
-		RenderThreadTicker->Register();
-	});
+	FTickableObjectRenderThread* RenderThreadTickerLocal = RenderThreadTicker.Get();
+	FGoogleVRSplash* pGVRSplash = this;
+	ENQUEUE_RENDER_COMMAND(RegisterAsyncTick)(
+		[RenderThreadTickerLocal, pGVRSplash](FRHICommandListImmediate& RHICmdList)
+		{
+			pGVRSplash->AllocateSplashScreenRenderTarget();
+			RenderThreadTickerLocal->Register();
+		});
 
 	bIsShown = true;
 }
@@ -130,7 +131,9 @@ void FGoogleVRSplash::Hide()
 		return;
 	}
 
-	ENQUEUE_UNIQUE_RENDER_COMMAND_TWOPARAMETER(UnregisterAsyncTick,
+	FTickableObjectRenderThread* RenderThreadTickerLocal = RenderThreadTicker.Get();
+	FGoogleVRSplash* pGVRSplash = this;
+	ENQUEUE_RENDER_COMMAND(UnregisterAsyncTick)(
 	TSharedPtr<FGoogleVRSplashTicker>&, RenderThreadTicker, RenderThreadTicker,
 	FGoogleVRSplash*, pGVRSplash, this,
 	{
diff --git a/Engine/Plugins/Runtime/ProceduralMeshComponent/Source/ProceduralMeshComponent/Private/ProceduralMeshComponent.cpp b/Engine/Plugins/Runtime/ProceduralMeshComponent/Source/ProceduralMeshComponent/Private/ProceduralMeshComponent.cpp
--- a/Engine/Plugins/Runtime/ProceduralMeshComponent/Source/ProceduralMeshComponent/Private/ProceduralMeshComponent.cpp
+++ b/Engine/Plugins/Runtime/ProceduralMeshComponent/Source/ProceduralMeshComponent/Private/ProceduralMeshComponent.cpp
@@ -571,10 +571,9 @@ void UProceduralMeshComponent::UpdateMeshSection(int32 SectionIndex, const TArra
 			SectionData->NewVertexBuffer = Section.ProcVertexBuffer;
 
 			// Enqueue command to send to render thread
-			ENQUEUE_UNIQUE_RENDER_COMMAND_TWOPARAMETER(
-				FProcMeshSectionUpdate,
-				FProceduralMeshSceneProxy*, ProcMeshSceneProxy, (FProceduralMeshSceneProxy*)SceneProxy,
-				FProcMeshSectionUpdateData*, SectionData, SectionData,
+			FProceduralMeshSceneProxy* ProcMeshSceneProxy = (FProceduralMeshSceneProxy*)SceneProxy;
+			ENQUEUE_RENDER_COMMAND(FProcMeshSectionUpdate)(
+				[ProcMeshSceneProxy, SectionData](FRHICommandListImmediate& RHICmdList)
 				{
 					ProcMeshSceneProxy->UpdateSection_RenderThread(SectionData);
 				}
diff --git a/Engine/Plugins/Runtime/WindowsMoviePlayer/Source/WindowsMoviePlayer/Private/WindowsMovieStreamer.cpp b/Engine/Plugins/Runtime/WindowsMoviePlayer/Source/WindowsMoviePlayer/Private/WindowsMovieStreamer.cpp
--- a/Engine/Plugins/Runtime/WindowsMoviePlayer/Source/WindowsMoviePlayer/Private/WindowsMovieStreamer.cpp
+++ b/Engine/Plugins/Runtime/WindowsMoviePlayer/Source/WindowsMoviePlayer/Private/WindowsMovieStreamer.cpp
@@ -278,12 +278,12 @@ bool FMediaFoundationMovieStreamer::OpenNextMovie()
 
 			if( Texture->GetWidth() != VideoDimensions.X || Texture->GetHeight() != VideoDimensions.Y )
 			{
-				ENQUEUE_UNIQUE_RENDER_COMMAND_TWOPARAMETER(UpdateMovieTexture,
-					FSlateTexture2DRHIRef*, TextureRHIRef, Texture.Get(),
-					FIntPoint, InVideoDimensions, VideoDimensions,
-				{
-					TextureRHIRef->Resize( InVideoDimensions.X, InVideoDimensions.Y );
-				});
+				FSlateTexture2DRHIRef* TextureRHIRef = Texture.Get();
+				ENQUEUE_RENDER_COMMAND(UpdateMovieTexture)(
+					[TextureRHIRef, VideoDimensions](FRHICommandListImmediate& RHICmdList)
+					{
+						TextureRHIRef->Resize(VideoDimensions.X, VideoDimensions.Y);
+					});
 			}
 		}
 		else
@@ -291,11 +291,12 @@ bool FMediaFoundationMovieStreamer::OpenNextMovie()
 			const bool bCreateEmptyTexture = true;
 			Texture = MakeShareable(new FSlateTexture2DRHIRef(VideoDimensions.X, VideoDimensions.Y, PF_B8G8R8A8, NULL, TexCreate_RenderTargetable, bCreateEmptyTexture));
 
-			ENQUEUE_UNIQUE_RENDER_COMMAND_ONEPARAMETER(InitMovieTexture,
-				FSlateTexture2DRHIRef*, TextureRHIRef, Texture.Get(),
-			{
-				TextureRHIRef->InitResource();
-			});
+			FSlateTexture2DRHIRef* TextureRHIRef = Texture.Get();
+			ENQUEUE_RENDER_COMMAND(InitMovieTexture)(
+				[TextureRHIRef](FRHICommandListImmediate& RHICmdList)
+				{
+					TextureRHIRef->InitResource();
+				});
 		}
 
 		VideoPlayer->StartPlayback();
diff --git a/Engine/Source/Editor/GraphEditor/Private/MaterialNodes/SGraphNodeMaterialBase.cpp b/Engine/Source/Editor/GraphEditor/Private/MaterialNodes/SGraphNodeMaterialBase.cpp
--- a/Engine/Source/Editor/GraphEditor/Private/MaterialNodes/SGraphNodeMaterialBase.cpp
+++ b/Engine/Source/Editor/GraphEditor/Private/MaterialNodes/SGraphNodeMaterialBase.cpp
@@ -187,18 +187,16 @@ bool FPreviewElement::BeginRenderingCanvas( const FIntRect& InCanvasRect, const
 		RenderInfo.RenderProxy = InGraphNode->GetExpressionPreview();
 		RenderInfo.bIsRealtime = bInIsRealtime;
 
-		ENQUEUE_UNIQUE_RENDER_COMMAND_TWOPARAMETER
-		(
-			BeginRenderingPreviewCanvas,
-			FPreviewElement*, PreviewElement, this, 
-			FPreviewRenderInfo, InRenderInfo, RenderInfo,
-		{
-			PreviewElement->RenderTarget->SetViewRect(InRenderInfo.CanvasRect);
-			PreviewElement->RenderTarget->SetClippingRect(InRenderInfo.ClippingRect);
-			PreviewElement->ExpressionPreview = InRenderInfo.RenderProxy;
-			PreviewElement->bIsRealtime = InRenderInfo.bIsRealtime;
-		}
-		);
+		FPreviewElement* PreviewElement = this;
+		ENQUEUE_RENDER_COMMAND(BeginRenderingPreviewCanvas)(
+			[PreviewElement, RenderInfo](FRHICommandListImmediate& RHICmdList)
+			{
+				PreviewElement->RenderTarget->SetViewRect(RenderInfo.CanvasRect);
+				PreviewElement->RenderTarget->SetClippingRect(RenderInfo.ClippingRect);
+				PreviewElement->ExpressionPreview = RenderInfo.RenderProxy;
+				PreviewElement->bIsRealtime = RenderInfo.bIsRealtime;
+			}
+			);
 		return true;
 	}
 
@@ -207,11 +205,10 @@ bool FPreviewElement::BeginRenderingCanvas( const FIntRect& InCanvasRect, const
 
 void FPreviewElement::UpdateExpressionPreview(UMaterialGraphNode* MaterialNode)
 {
-	ENQUEUE_UNIQUE_RENDER_COMMAND_TWOPARAMETER
-	(
-		UpdatePreviewNodeRenderProxy,
-		FPreviewElement*, PreviewElement, this,
-		FMaterialRenderProxy*, InRenderProxy, MaterialNode ? MaterialNode->GetExpressionPreview() : nullptr,
+	FPreviewElement* PreviewElement = this;
+	FMaterialRenderProxy* InRenderProxy = MaterialNode ? MaterialNode->GetExpressionPreview() : nullptr;
+	ENQUEUE_RENDER_COMMAND(UpdatePreviewNodeRenderProxy)(
+		[PreviewElement, InRenderProxy](FRHICommandListImmediate& RHICmdList)
 		{
 			PreviewElement->ExpressionPreview = InRenderProxy;
 		}
diff --git a/Engine/Source/Editor/NewLevelDialog/Private/NewLevelDialogModule.cpp b/Engine/Source/Editor/NewLevelDialog/Private/NewLevelDialogModule.cpp
--- a/Engine/Source/Editor/NewLevelDialog/Private/NewLevelDialogModule.cpp
+++ b/Engine/Source/Editor/NewLevelDialog/Private/NewLevelDialogModule.cpp
@@ -43,12 +43,11 @@ public:
 	{
 		Size = FIntPoint(InTexture2D->GetSizeX(), InTexture2D->GetSizeY());
 
-		ENQUEUE_UNIQUE_RENDER_COMMAND_TWOPARAMETER(
-			UpdateSTexture2DView,
-			STexture2DView*,TextureView,this,
-			UTexture2D*,Texture,InTexture2D,
+		STexture2DView* TextureView = this;
+		ENQUEUE_RENDER_COMMAND(UpdateSTexture2DView)(
+			[TextureView, InTexture2D](FRHICommandListImmediate& RHICmdList)
 			{
-				TextureView->ShaderResource = ((FTexture2DResource*)(Texture->Resource))->GetTexture2DRHI();
+				TextureView->ShaderResource = ((FTexture2DResource*)(InTexture2D->Resource))->GetTexture2DRHI();
 			});
 	}
 
diff --git a/Engine/Source/Editor/UnrealEd/Private/AssetThumbnail.cpp b/Engine/Source/Editor/UnrealEd/Private/AssetThumbnail.cpp
--- a/Engine/Source/Editor/UnrealEd/Private/AssetThumbnail.cpp
+++ b/Engine/Source/Editor/UnrealEd/Private/AssetThumbnail.cpp
@@ -1086,17 +1086,16 @@ void FAssetThumbnailPool::Tick( float DeltaTime )
 							FSlateTextureData* BulkData = new FSlateTextureData(ObjectThumbnail->GetImageWidth(),ObjectThumbnail->GetImageHeight(),GPixelFormats[PF_B8G8R8A8].BlockBytes, ObjectThumbnail->AccessImageData() );
 
 							// Update the texture RHI
-							ENQUEUE_UNIQUE_RENDER_COMMAND_TWOPARAMETER(
-								ClearSlateTextureCommand,
-								FThumbnailInfo_RenderThread, ThumbInfo, InfoRef.Get(),
-								FSlateTextureData*, BulkData, BulkData,
+							FThumbnailInfo_RenderThread ThumbInfo = InfoRef.Get();
+							ENQUEUE_RENDER_COMMAND(ClearSlateTextureCommand)(
+								[ThumbInfo, BulkData](FRHICommandListImmediate& RHICmdList)
 							{
-								if ( ThumbInfo.ThumbnailTexture->GetTypedResource() == ThumbInfo.ThumbnailRenderTarget->GetTextureRHI() )
+								if (ThumbInfo.ThumbnailTexture->GetTypedResource() == ThumbInfo.ThumbnailRenderTarget->GetTextureRHI())
 								{
 									ThumbInfo.ThumbnailTexture->SetRHIRef(NULL, ThumbInfo.Width, ThumbInfo.Height);
 								}
 
-								ThumbInfo.ThumbnailTexture->SetTextureData( MakeShareable(BulkData) );
+								ThumbInfo.ThumbnailTexture->SetTextureData(MakeShareable(BulkData));
 								ThumbInfo.ThumbnailTexture->UpdateRHI();
 							});
 
diff --git a/Engine/Source/Editor/UnrealEd/Private/EditorEngine.cpp b/Engine/Source/Editor/UnrealEd/Private/EditorEngine.cpp
--- a/Engine/Source/Editor/UnrealEd/Private/EditorEngine.cpp
+++ b/Engine/Source/Editor/UnrealEd/Private/EditorEngine.cpp
@@ -1865,18 +1865,18 @@ void UEditorEngine::Tick( float DeltaSeconds, bool bIdleMode )
 	{
 		// rendering thread commands
 
-		ENQUEUE_UNIQUE_RENDER_COMMAND_TWOPARAMETER(
-			TickRenderingTimer,
-			bool, bPauseRenderingRealtimeClock, GPauseRenderingRealtimeClock,
-			float, DeltaTime, DeltaSeconds,
-		{
-			if(!bPauseRenderingRealtimeClock)
+		bool bPauseRenderingRealtimeClock = GPauseRenderingRealtimeClock;
+		float DeltaTime = DeltaSeconds;
+		ENQUEUE_RENDER_COMMAND(TickRenderingTimer)(
+			[bPauseRenderingRealtimeClock, DeltaTime](FRHICommandListImmediate& RHICmdList)
 			{
-				// Tick the GRenderingRealtimeClock, unless it's paused
-				GRenderingRealtimeClock.Tick(DeltaTime);
-			}
-			GetRendererModule().TickRenderTargetPool();
-		});
+				if(!bPauseRenderingRealtimeClock)
+				{
+					// Tick the GRenderingRealtimeClock, unless it's paused
+					GRenderingRealtimeClock.Tick(DeltaTime);
+				}
+				GetRendererModule().TickRenderTargetPool();
+			});
 	}
 
 	// After the play world has ticked, see if a request was made to end pie
diff --git a/Engine/Source/Runtime/Engine/Private/Atmosphere/Atmosphere.cpp b/Engine/Source/Runtime/Engine/Private/Atmosphere/Atmosphere.cpp
--- a/Engine/Source/Runtime/Engine/Private/Atmosphere/Atmosphere.cpp
+++ b/Engine/Source/Runtime/Engine/Private/Atmosphere/Atmosphere.cpp
@@ -277,58 +277,55 @@ void UAtmosphericFogComponent::BeginDestroy()
 void UAtmosphericFogComponent::ReleaseResource()
 {
 	FSceneInterface* Scene = GetScene();
-	if ( TransmittanceResource != NULL )
+	if ( TransmittanceResource != nullptr )
 	{
-		ENQUEUE_UNIQUE_RENDER_COMMAND_TWOPARAMETER(
-			ReleaseAtmosphericFogTransmittanceTextureCommand,
-			FRenderResource*,TransmittanceResource,TransmittanceResource,
-			FSceneInterface*,Scene,Scene,
-		{
-			if (Scene)
+		FRenderResource* InTransmittanceResource = TransmittanceResource;
+		ENQUEUE_RENDER_COMMAND(ReleaseAtmosphericFogTransmittanceTextureCommand)(
+			[InTransmittanceResource, Scene](FRHICommandListImmediate& RHICmdList)
 			{
-				Scene->RemoveAtmosphericFogResource_RenderThread(TransmittanceResource);
-			}
-			TransmittanceResource->ReleaseResource();
-			delete TransmittanceResource;
-		});
+				if (Scene)
+				{
+					Scene->RemoveAtmosphericFogResource_RenderThread(InTransmittanceResource);
+				}
+				InTransmittanceResource->ReleaseResource();
+				delete InTransmittanceResource;
+			});
 
-		TransmittanceResource = NULL;
+		TransmittanceResource = nullptr;
 	}
 
-	if ( IrradianceResource != NULL )
+	if ( IrradianceResource != nullptr)
 	{
-		ENQUEUE_UNIQUE_RENDER_COMMAND_TWOPARAMETER(
-			ReleaseAtmosphericFogIrradianceTextureCommand,
-			FRenderResource*,IrradianceResource,IrradianceResource,
-			FSceneInterface*,Scene,Scene,
-		{
-			if (Scene)
+		FRenderResource* InIrradianceResource = IrradianceResource;
+		ENQUEUE_RENDER_COMMAND(ReleaseAtmosphericFogIrradianceTextureCommand)(
+			[InIrradianceResource, Scene](FRHICommandListImmediate& RHICmdList)
 			{
-				Scene->RemoveAtmosphericFogResource_RenderThread(IrradianceResource);
-			}
-			IrradianceResource->ReleaseResource();
-			delete IrradianceResource;
-		});
+				if (Scene)
+				{
+					Scene->RemoveAtmosphericFogResource_RenderThread(InIrradianceResource);
+				}
+				InIrradianceResource->ReleaseResource();
+				delete InIrradianceResource;
+			});
 
-		IrradianceResource = NULL;
+		IrradianceResource = nullptr;
 	}
 
-	if ( InscatterResource != NULL )
+	if ( InscatterResource != nullptr )
 	{
-		ENQUEUE_UNIQUE_RENDER_COMMAND_TWOPARAMETER(
-			ReleaseAtmosphericFogInscatterTextureCommand,
-			FRenderResource*,InscatterResource,InscatterResource,
-			FSceneInterface*,Scene,Scene,
-		{
-			if (Scene)
+		FRenderResource* InInscatterResource = InscatterResource;
+		ENQUEUE_RENDER_COMMAND(ReleaseAtmosphericFogInscatterTextureCommand)(
+			[InInscatterResource, Scene](FRHICommandListImmediate& RHICmdList)
 			{
-				Scene->RemoveAtmosphericFogResource_RenderThread(InscatterResource);
-			}
-			InscatterResource->ReleaseResource();
-			delete InscatterResource;
-		});
+				if (Scene)
+				{
+					Scene->RemoveAtmosphericFogResource_RenderThread(InInscatterResource);
+				}
+				InInscatterResource->ReleaseResource();
+				delete InInscatterResource;
+			});
 
-		InscatterResource = NULL;
+		InscatterResource = nullptr;
 	}
 }
 
@@ -578,18 +575,17 @@ void UAtmosphericFogComponent::StartPrecompute()
 			}
 
 			// This is largely redundant, the component will be reattached anyway, thus it will be recomputed.
-			ENQUEUE_UNIQUE_RENDER_COMMAND_TWOPARAMETER(
-				FStartPrecompute,
-				FSceneInterface*, AtmosphericFogScene, AtmosphericFogScene, 
-				UAtmosphericFogComponent*, Component, this,
-			{
-				FAtmosphericFogSceneInfo* AtmosphericFogSceneInfo = AtmosphericFogScene->GetAtmosphericFogSceneInfo();
-
-				if (AtmosphericFogSceneInfo && AtmosphericFogSceneInfo->Component == Component)
+			UAtmosphericFogComponent* Component = this;
+			ENQUEUE_RENDER_COMMAND(FStartPrecompute)(
+				[AtmosphericFogScene, Component](FRHICommandListImmediate& RHICmdList)
 				{
-					AtmosphericFogSceneInfo->bNeedRecompute = true;
-				}
-			});
+					FAtmosphericFogSceneInfo* AtmosphericFogSceneInfo = AtmosphericFogScene->GetAtmosphericFogSceneInfo();
+
+					if (AtmosphericFogSceneInfo && AtmosphericFogSceneInfo->Component == Component)
+					{
+						AtmosphericFogSceneInfo->bNeedRecompute = true;
+					}
+				});
 		}
 	}
 #endif
diff --git a/Engine/Source/Runtime/Engine/Private/GameEngine.cpp b/Engine/Source/Runtime/Engine/Private/GameEngine.cpp
--- a/Engine/Source/Runtime/Engine/Private/GameEngine.cpp
+++ b/Engine/Source/Runtime/Engine/Private/GameEngine.cpp
@@ -1501,15 +1501,14 @@ void UGameEngine::Tick( float DeltaSeconds, bool bIdleMode )
 
 	// rendering thread commands
 	{
-		ENQUEUE_UNIQUE_RENDER_COMMAND_TWOPARAMETER(
-			TickRenderingTimer,
-			bool, bPauseRenderingRealtimeClock, GPauseRenderingRealtimeClock,
-			float, DeltaTime, DeltaSeconds,
+		bool bPauseRenderingRealtimeClock = GPauseRenderingRealtimeClock;
+		ENQUEUE_RENDER_COMMAND(TickRenderingTimer)(
+			[bPauseRenderingRealtimeClock, DeltaSeconds](FRHICommandListImmediate& RHICmdList)
 		{
 			if(!bPauseRenderingRealtimeClock)
 			{
 				// Tick the GRenderingRealtimeClock, unless it's paused
-				GRenderingRealtimeClock.Tick(DeltaTime);
+				GRenderingRealtimeClock.Tick(DeltaSeconds);
 			}
 
 			GetRendererModule().TickRenderTargetPool();
diff --git a/Engine/Source/Runtime/Engine/Private/InstancedStaticMesh.cpp b/Engine/Source/Runtime/Engine/Private/InstancedStaticMesh.cpp
--- a/Engine/Source/Runtime/Engine/Private/InstancedStaticMesh.cpp
+++ b/Engine/Source/Runtime/Engine/Private/InstancedStaticMesh.cpp
@@ -672,10 +672,10 @@ void FPerInstanceRenderData::UpdateFromPreallocatedData(FStaticMeshInstanceData&
 
 	typedef TSharedPtr<FStaticMeshInstanceData, ESPMode::ThreadSafe> FStaticMeshInstanceDataPtr;
 
-	ENQUEUE_UNIQUE_RENDER_COMMAND_TWOPARAMETER(
-		FInstanceBuffer_UpdateFromPreallocatedData,
-		FStaticMeshInstanceDataPtr, InInstanceBufferDataPtr, InstanceBuffer_GameThread,
-		FStaticMeshInstanceBuffer*, InInstanceBuffer, &InstanceBuffer,
+	FStaticMeshInstanceDataPtr InInstanceBufferDataPtr = InstanceBuffer_GameThread;
+	FStaticMeshInstanceBuffer* InInstanceBuffer = &InstanceBuffer;
+	ENQUEUE_RENDER_COMMAND(FInstanceBuffer_UpdateFromPreallocatedData)(
+		[InInstanceBufferDataPtr, InInstanceBuffer](FRHICommandListImmediate& RHICmdList)
 		{
 			InInstanceBuffer->InstanceData = InInstanceBufferDataPtr;
 			InInstanceBuffer->UpdateRHI();
diff --git a/Engine/Source/Runtime/Engine/Private/Level.cpp b/Engine/Source/Runtime/Engine/Private/Level.cpp
--- a/Engine/Source/Runtime/Engine/Private/Level.cpp
+++ b/Engine/Source/Runtime/Engine/Private/Level.cpp
@@ -2122,13 +2122,12 @@ void ULevel::ApplyWorldOffset(const FVector& InWorldOffset, bool bWorldShift)
 		// Otherwise we need to send a command to move just this volume
 		else if (!bWorldShift) 
 		{
-			ENQUEUE_UNIQUE_RENDER_COMMAND_TWOPARAMETER(
- 				ApplyWorldOffset_PLV,
- 				FPrecomputedLightVolume*, InPrecomputedLightVolume, PrecomputedLightVolume,
- 				FVector, InWorldOffset, InWorldOffset,
- 			{
-				InPrecomputedLightVolume->ApplyWorldOffset(InWorldOffset);
- 			});
+			FPrecomputedLightVolume* InPrecomputedLightVolume = PrecomputedLightVolume;
+			ENQUEUE_RENDER_COMMAND(ApplyWorldOffset_PLV)(
+				[InPrecomputedLightVolume, InWorldOffset](FRHICommandListImmediate& RHICmdList)
+	 			{
+					InPrecomputedLightVolume->ApplyWorldOffset(InWorldOffset);
+ 				});
 		}
 	}
 
@@ -2149,13 +2148,12 @@ void ULevel::ApplyWorldOffset(const FVector& InWorldOffset, bool bWorldShift)
 		// Otherwise we need to send a command to move just this volume
 		else if (!bWorldShift) 
 		{
-			ENQUEUE_UNIQUE_RENDER_COMMAND_TWOPARAMETER(
- 				ApplyWorldOffset_PLV,
- 				FPrecomputedVolumetricLightmap*, InPrecomputedVolumetricLightmap, PrecomputedVolumetricLightmap,
- 				FVector, InWorldOffset, InWorldOffset,
- 			{
-				InPrecomputedVolumetricLightmap->ApplyWorldOffset(InWorldOffset);
- 			});
+			FPrecomputedVolumetricLightmap* InPrecomputedVolumetricLightmap = PrecomputedVolumetricLightmap;
+			ENQUEUE_RENDER_COMMAND(ApplyWorldOffset_PLV)(
+				[InPrecomputedVolumetricLightmap, InWorldOffset](FRHICommandListImmediate& RHICmdList)
+	 			{
+					InPrecomputedVolumetricLightmap->ApplyWorldOffset(InWorldOffset);
+ 				});
 		}
 	}
 
diff --git a/Engine/Source/Runtime/Engine/Private/Materials/MaterialInstance.cpp b/Engine/Source/Runtime/Engine/Private/Materials/MaterialInstance.cpp
--- a/Engine/Source/Runtime/Engine/Private/Materials/MaterialInstance.cpp
+++ b/Engine/Source/Runtime/Engine/Private/Materials/MaterialInstance.cpp
@@ -306,14 +306,13 @@ void FMaterialInstanceResource::GameThread_SetParent(UMaterialInterface* ParentM
 
 		// Set the rendering thread's parent and instance pointers.
 		check(ParentMaterialInterface != NULL);
-		ENQUEUE_UNIQUE_RENDER_COMMAND_TWOPARAMETER(
-			InitMaterialInstanceResource,
-			FMaterialInstanceResource*, Resource, this,
-			UMaterialInterface*, Parent, ParentMaterialInterface,
+		FMaterialInstanceResource* Resource = this;
+		ENQUEUE_RENDER_COMMAND(InitMaterialInstanceResource)(
+			[Resource, ParentMaterialInterface](FRHICommandListImmediate& RHICmdList)
 			{
-			Resource->Parent = Parent;
-			Resource->InvalidateUniformExpressionCache(false);
-		});
+				Resource->Parent = ParentMaterialInterface;
+				Resource->InvalidateUniformExpressionCache(false);
+			});
 
 		if (OldParent)
 		{
diff --git a/Engine/Source/Runtime/Engine/Private/Materials/MaterialShader.cpp b/Engine/Source/Runtime/Engine/Private/Materials/MaterialShader.cpp
--- a/Engine/Source/Runtime/Engine/Private/Materials/MaterialShader.cpp
+++ b/Engine/Source/Runtime/Engine/Private/Materials/MaterialShader.cpp
@@ -1325,13 +1325,11 @@ void FMaterialShaderMap::LoadForRemoteRecompile(FArchive& Ar, EShaderPlatform Sh
 							MaterialResource->SetGameThreadShaderMap(LoadedShaderMap);
 							MaterialResource->RegisterInlineShaderMap(false);
 
-							ENQUEUE_UNIQUE_RENDER_COMMAND_TWOPARAMETER(
-								FSetShaderMapOnMaterialResources,
-								FMaterial*,MaterialResource,MaterialResource,
-								FMaterialShaderMap*,LoadedShaderMap,LoadedShaderMap,
-							{
-								MaterialResource->SetRenderingThreadShaderMap(LoadedShaderMap);
-							});
+							ENQUEUE_RENDER_COMMAND(FSetShaderMapOnMaterialResources)(
+								[MaterialResource, LoadedShaderMap](FRHICommandListImmediate& RHICmdList)
+								{
+									MaterialResource->SetRenderingThreadShaderMap(LoadedShaderMap);
+								});
 						}
 					}
 				}
diff --git a/Engine/Source/Runtime/Engine/Private/Materials/MaterialUniformExpressions.cpp b/Engine/Source/Runtime/Engine/Private/Materials/MaterialUniformExpressions.cpp
--- a/Engine/Source/Runtime/Engine/Private/Materials/MaterialUniformExpressions.cpp
+++ b/Engine/Source/Runtime/Engine/Private/Materials/MaterialUniformExpressions.cpp
@@ -605,12 +605,12 @@ void FMaterialUniformExpressionTexture::Serialize(FArchive& Ar)
 void FMaterialUniformExpressionTexture::SetTransientOverrideTextureValue( UTexture* InOverrideTexture )
 {
 	TransientOverrideValue_GameThread = InOverrideTexture;
-	ENQUEUE_UNIQUE_RENDER_COMMAND_TWOPARAMETER(SetTransientOverrideTextureValueCommand,
-											   FMaterialUniformExpressionTexture*, ExpressionTexture, this,
-											   UTexture*, InOverrideTexture, InOverrideTexture,
-											   {
-												   ExpressionTexture->TransientOverrideValue_RenderThread = InOverrideTexture;
-											   });
+	FMaterialUniformExpressionTexture* ExpressionTexture = this;
+	ENQUEUE_RENDER_COMMAND(SetTransientOverrideTextureValueCommand)(
+		[ExpressionTexture, InOverrideTexture](FRHICommandListImmediate& RHICmdList)
+		{
+			ExpressionTexture->TransientOverrideValue_RenderThread = InOverrideTexture;
+		});
 }
 
 void FMaterialUniformExpressionTexture::GetTextureValue(const FMaterialRenderContext& Context,const FMaterial& Material,const UTexture*& OutValue,ESamplerSourceMode& OutSamplerSource) const
diff --git a/Engine/Source/Runtime/Engine/Private/Particles/FXSystem.cpp b/Engine/Source/Runtime/Engine/Private/Particles/FXSystem.cpp
--- a/Engine/Source/Runtime/Engine/Private/Particles/FXSystem.cpp
+++ b/Engine/Source/Runtime/Engine/Private/Particles/FXSystem.cpp
@@ -283,17 +283,16 @@ void FFXSystem::RemoveVectorField( UVectorFieldComponent* VectorFieldComponent )
 
 		if ( Instance )
 		{
-			ENQUEUE_UNIQUE_RENDER_COMMAND_TWOPARAMETER(
-				FRemoveVectorFieldCommand,
-				FFXSystem*, FXSystem, this,
-				FVectorFieldInstance*, Instance, Instance,
-			{
-				if ( Instance->Index != INDEX_NONE )
+			FFXSystem* FXSystem = this;
+			ENQUEUE_RENDER_COMMAND(FRemoveVectorFieldCommand)(
+				[FXSystem, Instance](FRHICommandListImmediate& RHICmdList)
 				{
-					FXSystem->VectorFields.RemoveAt( Instance->Index );
-					delete Instance;
-				}
-			});
+					if ( Instance->Index != INDEX_NONE )
+					{
+						FXSystem->VectorFields.RemoveAt( Instance->Index );
+						delete Instance;
+					}
+				});
 		}
 	}
 }
diff --git a/Engine/Source/Runtime/Engine/Private/Particles/ParticleGpuSimulation.cpp b/Engine/Source/Runtime/Engine/Private/Particles/ParticleGpuSimulation.cpp
--- a/Engine/Source/Runtime/Engine/Private/Particles/ParticleGpuSimulation.cpp
+++ b/Engine/Source/Runtime/Engine/Private/Particles/ParticleGpuSimulation.cpp
@@ -4535,35 +4535,33 @@ void FFXSystem::ReleaseGPUResources()
 
 void FFXSystem::AddGPUSimulation(FParticleSimulationGPU* Simulation)
 {
-	ENQUEUE_UNIQUE_RENDER_COMMAND_TWOPARAMETER(
-		FAddGPUSimulationCommand,
-		FFXSystem*, FXSystem, this,
-		FParticleSimulationGPU*, Simulation, Simulation,
-	{
-		if (Simulation->SimulationIndex == INDEX_NONE)
+	FFXSystem* FXSystem = this;
+	ENQUEUE_RENDER_COMMAND(FAddGPUSimulationCommand)(
+		[FXSystem, Simulation](FRHICommandListImmediate& RHICmdList)
 		{
-			FSparseArrayAllocationInfo Allocation = FXSystem->GPUSimulations.AddUninitialized();
-			Simulation->SimulationIndex = Allocation.Index;
-			FXSystem->GPUSimulations[Allocation.Index] = Simulation;
-		}
-		check(FXSystem->GPUSimulations[Simulation->SimulationIndex] == Simulation);
-	});
+			if (Simulation->SimulationIndex == INDEX_NONE)
+			{
+				FSparseArrayAllocationInfo Allocation = FXSystem->GPUSimulations.AddUninitialized();
+				Simulation->SimulationIndex = Allocation.Index;
+				FXSystem->GPUSimulations[Allocation.Index] = Simulation;
+			}
+			check(FXSystem->GPUSimulations[Simulation->SimulationIndex] == Simulation);
+		});
 }
 
 void FFXSystem::RemoveGPUSimulation(FParticleSimulationGPU* Simulation)
 {
-	ENQUEUE_UNIQUE_RENDER_COMMAND_TWOPARAMETER(
-		FRemoveGPUSimulationCommand,
-		FFXSystem*, FXSystem, this,
-		FParticleSimulationGPU*, Simulation, Simulation,
-	{
-		if (Simulation->SimulationIndex != INDEX_NONE)
+	FFXSystem* FXSystem = this;
+	ENQUEUE_RENDER_COMMAND(FRemoveGPUSimulationCommand)(
+		[FXSystem, Simulation](FRHICommandListImmediate& RHICmdList)
 		{
-			check(FXSystem->GPUSimulations[Simulation->SimulationIndex] == Simulation);
-			FXSystem->GPUSimulations.RemoveAt(Simulation->SimulationIndex);
-		}
-		Simulation->SimulationIndex = INDEX_NONE;
-	});
+			if (Simulation->SimulationIndex != INDEX_NONE)
+			{
+				check(FXSystem->GPUSimulations[Simulation->SimulationIndex] == Simulation);
+				FXSystem->GPUSimulations.RemoveAt(Simulation->SimulationIndex);
+			}
+			Simulation->SimulationIndex = INDEX_NONE;
+		});
 }
 
 int32 FFXSystem::AddSortedGPUSimulation(FParticleSimulationGPU* Simulation, const FVector& ViewOrigin)
diff --git a/Engine/Source/Runtime/Engine/Private/Particles/ParticleSystemRender.cpp b/Engine/Source/Runtime/Engine/Private/Particles/ParticleSystemRender.cpp
--- a/Engine/Source/Runtime/Engine/Private/Particles/ParticleSystemRender.cpp
+++ b/Engine/Source/Runtime/Engine/Private/Particles/ParticleSystemRender.cpp
@@ -7070,10 +7070,9 @@ void FParticleSystemSceneProxy::ReleaseRenderThreadResourcesForEmitterData()
 
 void FParticleSystemSceneProxy::UpdateData(FParticleDynamicData* NewDynamicData)
 {
-	ENQUEUE_UNIQUE_RENDER_COMMAND_TWOPARAMETER(
-		ParticleUpdateDataCommand,
-		FParticleSystemSceneProxy*, Proxy, this,
-		FParticleDynamicData*, NewDynamicData, NewDynamicData,
+	FParticleSystemSceneProxy* Proxy = this;
+	ENQUEUE_RENDER_COMMAND(ParticleUpdateDataCommand)(
+		[Proxy, NewDynamicData](FRHICommandListImmediate& RHICmdList)
 		{
 			SCOPE_CYCLE_COUNTER(STAT_ParticleUpdateRTTime);
 			STAT(FScopeCycleCounter Context(Proxy->GetStatId());)
diff --git a/Engine/Source/Runtime/Engine/Private/PrimitiveSceneProxy.cpp b/Engine/Source/Runtime/Engine/Private/PrimitiveSceneProxy.cpp
--- a/Engine/Source/Runtime/Engine/Private/PrimitiveSceneProxy.cpp
+++ b/Engine/Source/Runtime/Engine/Private/PrimitiveSceneProxy.cpp
@@ -458,13 +458,12 @@ void FPrimitiveSceneProxy::SetHovered_GameThread(const bool bInHovered)
 	check(IsInGameThread());
 
 	// Enqueue a message to the rendering thread containing the interaction to add.
-	ENQUEUE_UNIQUE_RENDER_COMMAND_TWOPARAMETER(
-		SetNewHovered,
-		FPrimitiveSceneProxy*,PrimitiveSceneProxy,this,
-		const bool,bNewHovered,bInHovered,
-	{
-		PrimitiveSceneProxy->SetHovered_RenderThread(bNewHovered);
-	});
+	FPrimitiveSceneProxy* PrimitiveSceneProxy = this;
+	ENQUEUE_RENDER_COMMAND(SetNewHovered)(
+		[PrimitiveSceneProxy, bInHovered](FRHICommandListImmediate& RHICmdList)
+		{
+			PrimitiveSceneProxy->SetHovered_RenderThread(bInHovered);
+		});
 }
 
 #if !UE_BUILD_SHIPPING
@@ -504,13 +503,12 @@ void FPrimitiveSceneProxy::SetHiddenEdViews_GameThread( uint64 InHiddenEditorVie
 {
 	check(IsInGameThread());
 
-	ENQUEUE_UNIQUE_RENDER_COMMAND_TWOPARAMETER(
-		SetEditorVisibility,
-		FPrimitiveSceneProxy*,PrimitiveSceneProxy,this,
-		const uint64,NewHiddenEditorViews,InHiddenEditorViews,
-	{
-		PrimitiveSceneProxy->SetHiddenEdViews_RenderThread(NewHiddenEditorViews);
-	});
+	FPrimitiveSceneProxy* PrimitiveSceneProxy = this;
+	ENQUEUE_RENDER_COMMAND(SetEditorVisibility)(
+		[PrimitiveSceneProxy, InHiddenEditorViews](FRHICommandListImmediate& RHICmdList)
+		{
+			PrimitiveSceneProxy->SetHiddenEdViews_RenderThread(InHiddenEditorViews);
+		});
 }
 
 /**
@@ -529,13 +527,12 @@ void FPrimitiveSceneProxy::SetCollisionEnabled_GameThread(const bool bNewEnabled
 	check(IsInGameThread());
 
 	// Enqueue a message to the rendering thread to change draw state
-	ENQUEUE_UNIQUE_RENDER_COMMAND_TWOPARAMETER(
-		SetCollisionEnabled,
-		FPrimitiveSceneProxy*,PrimSceneProxy,this,
-		const bool,bEnabled,bNewEnabled,
-	{
-		PrimSceneProxy->SetCollisionEnabled_RenderThread(bEnabled);
-	});
+	FPrimitiveSceneProxy* PrimSceneProxy = this;
+	ENQUEUE_RENDER_COMMAND(SetCollisionEnabled)(
+		[PrimSceneProxy, bNewEnabled](FRHICommandListImmediate& RHICmdList)
+		{
+			PrimSceneProxy->SetCollisionEnabled_RenderThread(bNewEnabled);
+		});
 }
 
 void FPrimitiveSceneProxy::SetCollisionEnabled_RenderThread(const bool bNewEnabled)
diff --git a/Engine/Source/Runtime/Engine/Private/Rendering/SubsurfaceProfile.cpp b/Engine/Source/Runtime/Engine/Private/Rendering/SubsurfaceProfile.cpp
--- a/Engine/Source/Runtime/Engine/Private/Rendering/SubsurfaceProfile.cpp
+++ b/Engine/Source/Runtime/Engine/Private/Rendering/SubsurfaceProfile.cpp
@@ -383,12 +383,12 @@ void USubsurfaceProfile::BeginDestroy()
 
 void USubsurfaceProfile::PostEditChangeProperty(struct FPropertyChangedEvent& PropertyChangedEvent)
 {
-	ENQUEUE_UNIQUE_RENDER_COMMAND_TWOPARAMETER(
-		UpdateSubsurfaceProfile,
-		const FSubsurfaceProfileStruct, Settings, this->Settings,
-		USubsurfaceProfile*, Profile, this,
-	{
-		// any changes to the setting require an update of the texture
-		GSubsurfaceProfileTextureObject.UpdateProfile(Settings, Profile);
-	});
+	const FSubsurfaceProfileStruct SettingsLocal = this->Settings;
+	USubsurfaceProfile* Profile = this;
+	ENQUEUE_RENDER_COMMAND(UpdateSubsurfaceProfile)(
+		[SettingsLocal, Profile](FRHICommandListImmediate& RHICmdList)
+		{
+			// any changes to the setting require an update of the texture
+			GSubsurfaceProfileTextureObject.UpdateProfile(SettingsLocal, Profile);
+		});
 }
diff --git a/Engine/Source/Runtime/Engine/Private/SkeletalMeshComponentPhysics.cpp b/Engine/Source/Runtime/Engine/Private/SkeletalMeshComponentPhysics.cpp
--- a/Engine/Source/Runtime/Engine/Private/SkeletalMeshComponentPhysics.cpp
+++ b/Engine/Source/Runtime/Engine/Private/SkeletalMeshComponentPhysics.cpp
@@ -1401,10 +1401,10 @@ void USkeletalMeshComponent::SendRenderDebugPhysics(FPrimitiveSceneProxy* Overri
 			}
 		}
 
-		ENQUEUE_UNIQUE_RENDER_COMMAND_TWOPARAMETER(
-			SkeletalMesh_SendRenderDebugPhysics, FPrimitiveSceneProxy*, PassedSceneProxy, UseSceneProxy, TArray<FPrimitiveSceneProxy::FDebugMassData>, UseDebugMassData, DebugMassData,
+		ENQUEUE_RENDER_COMMAND(SkeletalMesh_SendRenderDebugPhysics)(
+			[UseSceneProxy, DebugMassData](FRHICommandListImmediate& RHICmdList)
 			{
-				PassedSceneProxy->SetDebugMassData(UseDebugMassData);
+				UseSceneProxy->SetDebugMassData(DebugMassData);
 			}
 		);
 		
diff --git a/Engine/Source/Runtime/Engine/Private/Slate/SceneViewport.cpp b/Engine/Source/Runtime/Engine/Private/Slate/SceneViewport.cpp
--- a/Engine/Source/Runtime/Engine/Private/Slate/SceneViewport.cpp
+++ b/Engine/Source/Runtime/Engine/Private/Slate/SceneViewport.cpp
@@ -1602,18 +1602,20 @@ void FSceneViewport::UpdateViewportRHI(bool bDestroyed, uint32 NewSizeX, uint32
 		else
 		{
 			// Enqueue a render command to delete the handle.  It must be deleted on the render thread after the resource is released
-			ENQUEUE_UNIQUE_RENDER_COMMAND_TWOPARAMETER(DeleteSlateRenderTarget, TArray<FSlateRenderTargetRHI*>&, BufferedSlateHandles, BufferedSlateHandles,
-																				FSlateRenderTargetRHI*&, RenderThreadSlateTexture, RenderThreadSlateTexture,
-			{
-				for (int32 i = 0; i < BufferedSlateHandles.Num(); ++i)
+			FSlateRenderTargetRHI** RenderThreadSlateTexturePtr = &RenderThreadSlateTexture;
+			TArray<FSlateRenderTargetRHI*>* BufferedSlateHandlesPtr = &BufferedSlateHandles;
+			ENQUEUE_RENDER_COMMAND(DeleteSlateRenderTarget)(
+				[BufferedSlateHandlesPtr, RenderThreadSlateTexturePtr](FRHICommandListImmediate& RHICmdList)
 				{
-					delete BufferedSlateHandles[i];
-					BufferedSlateHandles[i] = nullptr;
+					for (int32 i = 0; i < BufferedSlateHandlesPtr->Num(); ++i)
+					{
+						delete (*BufferedSlateHandlesPtr)[i];
+						(*BufferedSlateHandlesPtr)[i] = nullptr;
+					}
 
-					delete RenderThreadSlateTexture;
-					RenderThreadSlateTexture = nullptr;
-				}						
-			});
+					delete *RenderThreadSlateTexturePtr;
+					*RenderThreadSlateTexturePtr = nullptr;
+				});
 
 		}
 	}
diff --git a/Engine/Source/Runtime/Engine/Private/Slate/SlateTextures.cpp b/Engine/Source/Runtime/Engine/Private/Slate/SlateTextures.cpp
--- a/Engine/Source/Runtime/Engine/Private/Slate/SlateTextures.cpp
+++ b/Engine/Source/Runtime/Engine/Private/Slate/SlateTextures.cpp
@@ -150,12 +150,12 @@ void FSlateTexture2DRHIRef::ResizeTexture(uint32 InWidth, uint32 InHeight)
 		else
 		{
 			FIntPoint Dimensions(InWidth, InHeight);
-			ENQUEUE_UNIQUE_RENDER_COMMAND_TWOPARAMETER(ResizeSlateTexture,
-			FSlateTexture2DRHIRef*, TextureRHIRef, this,
-			FIntPoint, InDimensions, Dimensions,
-			{
-				TextureRHIRef->Resize(InDimensions.X, InDimensions.Y);
-			});
+			FSlateTexture2DRHIRef* TextureRHIRef = this;
+			ENQUEUE_RENDER_COMMAND(ResizeSlateTexture)(
+				[TextureRHIRef, Dimensions](FRHICommandListImmediate& RHICmdList)
+				{
+					TextureRHIRef->Resize(Dimensions.X, Dimensions.Y);
+				});
 		}
 	}
 }
@@ -196,12 +196,12 @@ void FSlateTexture2DRHIRef::UpdateTexture(const TArray<uint8>& Bytes)
 	}
 	else
 	{
-		ENQUEUE_UNIQUE_RENDER_COMMAND_TWOPARAMETER(UpdateSlateTexture,
-		FSlateTexture2DRHIRef*, TextureRHIRef, this,
-		const TArray<uint8>&, TextureData, Bytes,
-		{
-			TextureRHIRef->SetTextureData(TextureData);
-		});
+		FSlateTexture2DRHIRef* TextureRHIRef = this;
+		ENQUEUE_RENDER_COMMAND(UpdateSlateTexture)(
+			[TextureRHIRef, Bytes](FRHICommandListImmediate& RHICmdList)
+			{
+				TextureRHIRef->SetTextureData(Bytes);
+			});
 	}
 }
 
@@ -213,10 +213,9 @@ void FSlateTexture2DRHIRef::UpdateTextureThreadSafe(const TArray<uint8>& Bytes)
 		FSlateTextureData* BulkData = new FSlateTextureData( Bytes.Num(), 0, 1, Bytes );
 
 		// Update the texture RHI
-		ENQUEUE_UNIQUE_RENDER_COMMAND_TWOPARAMETER(
-			FSlateTexture2DRHIRef_UpdateTextureThreadSafe,
-			FSlateTexture2DRHIRef*, ThisTexture, this,
-			FSlateTextureData*, BulkData, BulkData,
+		FSlateTexture2DRHIRef* ThisTexture = this;
+		ENQUEUE_RENDER_COMMAND(FSlateTexture2DRHIRef_UpdateTextureThreadSafe)(
+			[ThisTexture, BulkData](FRHICommandListImmediate& RHICmdList)
 			{
 				ThisTexture->UpdateTexture( BulkData->GetRawBytes() );
 				delete BulkData;
@@ -238,10 +237,9 @@ void FSlateTexture2DRHIRef::UpdateTextureThreadSafeWithTextureData(FSlateTexture
 {
 	check(IsInGameThread());
 	// Update the texture RHI
-	ENQUEUE_UNIQUE_RENDER_COMMAND_TWOPARAMETER(
-		FSlateTexture2DRHIRef_UpdateTextureThreadSafeWithTextureData,
-		FSlateTexture2DRHIRef*, ThisTexture, this,
-		FSlateTextureData*, BulkData, BulkData,
+	FSlateTexture2DRHIRef* ThisTexture = this;
+	ENQUEUE_RENDER_COMMAND(FSlateTexture2DRHIRef_UpdateTextureThreadSafeWithTextureData)(
+		[ThisTexture, BulkData](FRHICommandListImmediate& RHICmdList)
 		{
 			if (ThisTexture->GetWidth() != BulkData->GetWidth() || ThisTexture->GetHeight() != BulkData->GetHeight())
 			{
diff --git a/Engine/Source/Runtime/Engine/Private/SplineMeshSceneProxy.h b/Engine/Source/Runtime/Engine/Private/SplineMeshSceneProxy.h
--- a/Engine/Source/Runtime/Engine/Private/SplineMeshSceneProxy.h
+++ b/Engine/Source/Runtime/Engine/Private/SplineMeshSceneProxy.h
@@ -50,13 +50,13 @@ public:
 	/** Copy the data from another vertex factory */
 	void Copy(const FSplineMeshVertexFactory& Other)
 	{
-		ENQUEUE_UNIQUE_RENDER_COMMAND_TWOPARAMETER(
-			FSplineMeshVertexFactoryCopyData,
-			FSplineMeshVertexFactory*, VertexFactory, this,
-			const FDataType*, DataCopy, &Other.Data,
+		FSplineMeshVertexFactory* VertexFactory = this;
+		const FDataType* DataCopy = &Other.Data;
+		ENQUEUE_RENDER_COMMAND(FSplineMeshVertexFactoryCopyData)(
+			[VertexFactory, DataCopy](FRHICommandListImmediate& RHICmdList)
 			{
-			VertexFactory->Data = *DataCopy;
-		});
+				VertexFactory->Data = *DataCopy;
+			});
 		BeginUpdateResourceRHI(this);
 	}
 
diff --git a/Engine/Source/Runtime/Engine/Private/Texture2D.cpp b/Engine/Source/Runtime/Engine/Private/Texture2D.cpp
--- a/Engine/Source/Runtime/Engine/Private/Texture2D.cpp
+++ b/Engine/Source/Runtime/Engine/Private/Texture2D.cpp
@@ -1097,30 +1097,28 @@ void UTexture2D::UpdateTextureRegions(int32 MipIndex, uint32 NumRegions, const F
 		RegionData->SrcBpp = SrcBpp;
 		RegionData->SrcData = SrcData;
 
-		ENQUEUE_UNIQUE_RENDER_COMMAND_TWOPARAMETER(
-			UpdateTextureRegionsData,
-			FUpdateTextureRegionsData*, RegionData, RegionData,			
-			TFunction<void(uint8* SrcData, const FUpdateTextureRegion2D* Regions)>, DataCleanupFunc, DataCleanupFunc,
+		ENQUEUE_RENDER_COMMAND(UpdateTextureRegionsData)(
+			[RegionData, DataCleanupFunc](FRHICommandListImmediate& RHICmdList)
 			{
-			for (uint32 RegionIndex = 0; RegionIndex < RegionData->NumRegions; ++RegionIndex)
-			{
-				int32 CurrentFirstMip = RegionData->Texture2DResource->GetCurrentFirstMip();
-				if (RegionData->MipIndex >= CurrentFirstMip)
+				for (uint32 RegionIndex = 0; RegionIndex < RegionData->NumRegions; ++RegionIndex)
 				{
-					RHIUpdateTexture2D(
-						RegionData->Texture2DResource->GetTexture2DRHI(),
-						RegionData->MipIndex - CurrentFirstMip,
-						RegionData->Regions[RegionIndex],
-						RegionData->SrcPitch,
-						RegionData->SrcData
-						+ RegionData->Regions[RegionIndex].SrcY * RegionData->SrcPitch
-						+ RegionData->Regions[RegionIndex].SrcX * RegionData->SrcBpp
-						);
+					int32 CurrentFirstMip = RegionData->Texture2DResource->GetCurrentFirstMip();
+					if (RegionData->MipIndex >= CurrentFirstMip)
+					{
+						RHIUpdateTexture2D(
+							RegionData->Texture2DResource->GetTexture2DRHI(),
+							RegionData->MipIndex - CurrentFirstMip,
+							RegionData->Regions[RegionIndex],
+							RegionData->SrcPitch,
+							RegionData->SrcData
+							+ RegionData->Regions[RegionIndex].SrcY * RegionData->SrcPitch
+							+ RegionData->Regions[RegionIndex].SrcX * RegionData->SrcBpp
+							);
+					}
 				}
-			}
-			DataCleanupFunc(RegionData->SrcData, RegionData->Regions);
-			delete RegionData;
-		});
+				DataCleanupFunc(RegionData->SrcData, RegionData->Regions);
+				delete RegionData;
+			});
 	}
 }
 
diff --git a/Engine/Source/Runtime/Engine/Private/UnrealClient.cpp b/Engine/Source/Runtime/Engine/Private/UnrealClient.cpp
--- a/Engine/Source/Runtime/Engine/Private/UnrealClient.cpp
+++ b/Engine/Source/Runtime/Engine/Private/UnrealClient.cpp
@@ -1247,14 +1247,12 @@ void FViewport::HighResScreenshot()
 		ForceLODVar->Set(OldForceLOD, ECVF_SetByCode);
 	}
 
-	ENQUEUE_UNIQUE_RENDER_COMMAND_TWOPARAMETER(
-		EndDrawingCommand,
-		FViewport*,Viewport,DummyViewport,
-		FIntPoint,InRestoreSize,RestoreSize,
-	{
-		Viewport->EndRenderFrame(RHICmdList, false, false);
-		GetRendererModule().SceneRenderTargetsSetBufferSize(InRestoreSize.X, InRestoreSize.Y);
-	});
+	ENQUEUE_RENDER_COMMAND(EndDrawingCommand)(
+		[DummyViewport, RestoreSize](FRHICommandListImmediate& RHICmdList)
+		{
+			DummyViewport->EndRenderFrame(RHICmdList, false, false);
+			GetRendererModule().SceneRenderTargetsSetBufferSize(RestoreSize.X, RestoreSize.Y);
+		});
 
 	BeginReleaseResource(DummyViewport);
 	FlushRenderingCommands();
diff --git a/Engine/Source/Runtime/Engine/Private/VectorField.cpp b/Engine/Source/Runtime/Engine/Private/VectorField.cpp
--- a/Engine/Source/Runtime/Engine/Private/VectorField.cpp
+++ b/Engine/Source/Runtime/Engine/Private/VectorField.cpp
@@ -239,25 +239,24 @@ public:
 		UpdateParams.VolumeData = NULL;
 		InVectorField->SourceData.GetCopy(&UpdateParams.VolumeData, /*bDiscardInternalCopy=*/ true);
 
-		ENQUEUE_UNIQUE_RENDER_COMMAND_TWOPARAMETER(
-			FUpdateStaticVectorFieldCommand,
-			FVectorFieldStaticResource*, Resource, this,
-			FUpdateParams, UpdateParams, UpdateParams,
-		{
-				// Free any existing volume data on the resource.
-				FMemory::Free(Resource->VolumeData);
-
-				// Update settings on this resource.
-				Resource->SizeX = UpdateParams.SizeX;
-				Resource->SizeY = UpdateParams.SizeY;
-				Resource->SizeZ = UpdateParams.SizeZ;
-				Resource->Intensity = UpdateParams.Intensity;
-				Resource->LocalBounds = UpdateParams.Bounds;
-				Resource->VolumeData = UpdateParams.VolumeData;
-
-				// Update RHI resources.
-				Resource->UpdateRHI();
-		});
+		FVectorFieldStaticResource* Resource = this;
+		ENQUEUE_RENDER_COMMAND(FUpdateStaticVectorFieldCommand)(
+			[Resource, UpdateParams](FRHICommandListImmediate& RHICmdList)
+			{
+					// Free any existing volume data on the resource.
+					FMemory::Free(Resource->VolumeData);
+
+					// Update settings on this resource.
+					Resource->SizeX = UpdateParams.SizeX;
+					Resource->SizeY = UpdateParams.SizeY;
+					Resource->SizeZ = UpdateParams.SizeZ;
+					Resource->Intensity = UpdateParams.Intensity;
+					Resource->LocalBounds = UpdateParams.Bounds;
+					Resource->VolumeData = UpdateParams.VolumeData;
+
+					// Update RHI resources.
+					Resource->UpdateRHI();
+			});
 	}
 
 private:
diff --git a/Engine/Source/Runtime/Experimental/FieldSystem/Source/FieldSystemEngine/Private/Field/FieldSystemComponent.cpp b/Engine/Source/Runtime/Experimental/FieldSystem/Source/FieldSystemEngine/Private/Field/FieldSystemComponent.cpp
--- a/Engine/Source/Runtime/Experimental/FieldSystem/Source/FieldSystemEngine/Private/Field/FieldSystemComponent.cpp
+++ b/Engine/Source/Runtime/Experimental/FieldSystem/Source/FieldSystemEngine/Private/Field/FieldSystemComponent.cpp
@@ -44,9 +44,9 @@ void UFieldSystemComponent::CreateRenderState_Concurrent()
 		InitSampleData(SampleData);
 
 		// Enqueue command to send to render thread
-		ENQUEUE_UNIQUE_RENDER_COMMAND_TWOPARAMETER(FSendFieldSystemData,
-			FFieldSystemSceneProxy*, FieldSystemSceneProxy, (FFieldSystemSceneProxy*)SceneProxy,
-			FFieldSystemSampleData*, SampleData, SampleData,
+		FFieldSystemSceneProxy* FieldSystemSceneProxy = (FFieldSystemSceneProxy*)SceneProxy;
+		ENQUEUE_RENDER_COMMAND(FSendFieldSystemData)(
+			[FieldSystemSceneProxy, SampleData](FRHICommandListImmediate& RHICmdList)
 			{
 				//FieldSystemSceneProxy->SetSampleData_RenderThread(SampleData);
 			});
@@ -103,9 +103,9 @@ void UFieldSystemComponent::SendRenderDynamicData_Concurrent()
 			InitSampleData(SampleData);
 
 			// Enqueue command to send to render thread
-			ENQUEUE_UNIQUE_RENDER_COMMAND_TWOPARAMETER(FSendFieldSystemData,
-				FFieldSystemSceneProxy*, FieldSystemSceneProxy, (FFieldSystemSceneProxy*)SceneProxy,
-				FFieldSystemSampleData*, SampleData, SampleData,
+			FFieldSystemSceneProxy* FieldSystemSceneProxy = (FFieldSystemSceneProxy*)SceneProxy;
+			ENQUEUE_RENDER_COMMAND(FSendFieldSystemData)(
+				[FieldSystemSceneProxy, SampleData](FRHICommandListImmediate& RHICmdList)
 				{
 					//FieldSystemSceneProxy->SetSampleData_RenderThread(SampleData);
 				});
diff --git a/Engine/Source/Runtime/Experimental/GeometryCollectionEngine/Private/GeometryCollection/GeometryCollectionComponent.cpp b/Engine/Source/Runtime/Experimental/GeometryCollectionEngine/Private/GeometryCollection/GeometryCollectionComponent.cpp
--- a/Engine/Source/Runtime/Experimental/GeometryCollectionEngine/Private/GeometryCollection/GeometryCollectionComponent.cpp
+++ b/Engine/Source/Runtime/Experimental/GeometryCollectionEngine/Private/GeometryCollection/GeometryCollectionComponent.cpp
@@ -693,9 +693,9 @@ void UGeometryCollectionComponent::SendRenderDynamicData_Concurrent()
 			InitDynamicData(DynamicData);
 
 			// Enqueue command to send to render thread
-			ENQUEUE_UNIQUE_RENDER_COMMAND_TWOPARAMETER(FSendGeometryCollectionData,
-				FGeometryCollectionSceneProxy*, GeometryCollectionSceneProxy, (FGeometryCollectionSceneProxy*)SceneProxy,
-				FGeometryCollectionDynamicData*, DynamicData, DynamicData,
+			FGeometryCollectionSceneProxy* GeometryCollectionSceneProxy = (FGeometryCollectionSceneProxy*)SceneProxy;
+			ENQUEUE_RENDER_COMMAND(FSendGeometryCollectionData)(
+				[GeometryCollectionSceneProxy, DynamicData](FRHICommandListImmediate& RHICmdList)
 				{
 					GeometryCollectionSceneProxy->SetDynamicData_RenderThread(DynamicData);
 				});
diff --git a/Engine/Source/Runtime/HeadMountedDisplay/Private/DefaultSpectatorScreenController.cpp b/Engine/Source/Runtime/HeadMountedDisplay/Private/DefaultSpectatorScreenController.cpp
--- a/Engine/Source/Runtime/HeadMountedDisplay/Private/DefaultSpectatorScreenController.cpp
+++ b/Engine/Source/Runtime/HeadMountedDisplay/Private/DefaultSpectatorScreenController.cpp
@@ -110,18 +110,17 @@ void FDefaultSpectatorScreenController::SetSpectatorScreenTextureRenderCommand(U
 	}
 
 	// setting the texture must be done on the thread that's executing RHI commandlists.
-	ENQUEUE_UNIQUE_RENDER_COMMAND_TWOPARAMETER(
-		SetSpectatorScreenTexture,
-		FDefaultSpectatorScreenController*, SpectatorScreenController, this,
-		FTexture2DRHIRef, Texture, Texture2DRHIRef,
+	FDefaultSpectatorScreenController* SpectatorScreenController = this;
+	ENQUEUE_RENDER_COMMAND(SetSpectatorScreenTexture)(
+		[SpectatorScreenController, Texture2DRHIRef](FRHICommandListImmediate& RHICmdList)
 		{
 			if (RHICmdList.Bypass())
 			{
-				FRHISetSpectatorScreenTexture Command(SpectatorScreenController, Texture);
+				FRHISetSpectatorScreenTexture Command(SpectatorScreenController, Texture2DRHIRef);
 				Command.Execute(RHICmdList);
 				return;
 			}
-			ALLOC_COMMAND_CL(RHICmdList, FRHISetSpectatorScreenTexture)(SpectatorScreenController, Texture);
+			ALLOC_COMMAND_CL(RHICmdList, FRHISetSpectatorScreenTexture)(SpectatorScreenController, Texture2DRHIRef);
 		}
 	);
 }
@@ -154,18 +153,17 @@ void FDefaultSpectatorScreenController::SetSpectatorScreenModeTexturePlusEyeLayo
 	check(IsInGameThread());
 
 	// setting the layout must be done on the thread that's executing RHI commandlists.
-	ENQUEUE_UNIQUE_RENDER_COMMAND_TWOPARAMETER(
-		SetSpectatorScreenTexture,
-		FDefaultSpectatorScreenController*, SpectatorScreenController, this,
-		FSpectatorScreenModeTexturePlusEyeLayout, Layout, NewLayout,
+	FDefaultSpectatorScreenController* SpectatorScreenController = this;
+	ENQUEUE_RENDER_COMMAND(SetSpectatorScreenTexture)(
+		[SpectatorScreenController, NewLayout](FRHICommandListImmediate& RHICmdList)
 		{
 			if (RHICmdList.Bypass())
 			{
-				FRHISetSpectatorScreenModeTexturePlusEyeLayout Command(SpectatorScreenController, Layout);
+				FRHISetSpectatorScreenModeTexturePlusEyeLayout Command(SpectatorScreenController, NewLayout);
 				Command.Execute(RHICmdList);
 				return;
 			}
-			ALLOC_COMMAND_CL(RHICmdList, FRHISetSpectatorScreenModeTexturePlusEyeLayout)(SpectatorScreenController, Layout);
+			ALLOC_COMMAND_CL(RHICmdList, FRHISetSpectatorScreenModeTexturePlusEyeLayout)(SpectatorScreenController, NewLayout);
 		}
 	);
 }
diff --git a/Engine/Source/Runtime/ImageWriteQueue/Private/ImageWriteBlueprintLibrary.cpp b/Engine/Source/Runtime/ImageWriteQueue/Private/ImageWriteBlueprintLibrary.cpp
--- a/Engine/Source/Runtime/ImageWriteQueue/Private/ImageWriteBlueprintLibrary.cpp
+++ b/Engine/Source/Runtime/ImageWriteQueue/Private/ImageWriteBlueprintLibrary.cpp
@@ -54,10 +54,9 @@ bool UImageWriteBlueprintLibrary::ResolvePixelData(UTexture* InTexture, const FO
 		break;
 	}
 
-	ENQUEUE_UNIQUE_RENDER_COMMAND_TWOPARAMETER(
-		ResolvePixelData,
-		FTextureResource*, TextureResource, InTexture->Resource,
-		FOnPixelsReady, OnPixelsReady, OnPixelsReady,
+	FTextureResource* TextureResource = InTexture->Resource;
+	ENQUEUE_RENDER_COMMAND(ResolvePixelData)(
+		[TextureResource, OnPixelsReady](FRHICommandListImmediate& RHICmdList)
 		{
 			FTexture2DRHIRef Texture2D = TextureResource->TextureRHI ? TextureResource->TextureRHI->GetTexture2D() : nullptr;
 			if (!Texture2D)
diff --git a/Engine/Source/Runtime/Landscape/Private/LandscapeRender.cpp b/Engine/Source/Runtime/Landscape/Private/LandscapeRender.cpp
--- a/Engine/Source/Runtime/Landscape/Private/LandscapeRender.cpp
+++ b/Engine/Source/Runtime/Landscape/Private/LandscapeRender.cpp
@@ -3674,10 +3674,10 @@ IMPLEMENT_VERTEX_FACTORY_TYPE(FLandscapeVertexFactory, "/Engine/Private/Landscap
 void FLandscapeVertexFactory::Copy(const FLandscapeVertexFactory& Other)
 {
 	//SetSceneProxy(Other.Proxy());
-	ENQUEUE_UNIQUE_RENDER_COMMAND_TWOPARAMETER(
-		FLandscapeVertexFactoryCopyData,
-		FLandscapeVertexFactory*, VertexFactory, this,
-		const FDataType*, DataCopy, &Other.Data,
+	FLandscapeVertexFactory* VertexFactory = this;
+	const FDataType* DataCopy = &Other.Data;
+	ENQUEUE_RENDER_COMMAND(FLandscapeVertexFactoryCopyData)(
+		[VertexFactory, DataCopy](FRHICommandListImmediate& RHICmdList)
 		{
 			VertexFactory->Data = *DataCopy;
 		});
diff --git a/Engine/Source/Runtime/MRMesh/Private/MRMeshComponent.cpp b/Engine/Source/Runtime/MRMesh/Private/MRMeshComponent.cpp
--- a/Engine/Source/Runtime/MRMesh/Private/MRMeshComponent.cpp
+++ b/Engine/Source/Runtime/MRMesh/Private/MRMeshComponent.cpp
@@ -684,10 +684,9 @@ void UMRMeshComponent::SendBrickData_Internal(IMRMesh::FSendBrickDataArgs Args)
 			check(SceneProxy != nullptr);
 
 			// Graphics update
-			ENQUEUE_UNIQUE_RENDER_COMMAND_TWOPARAMETER(
-				FSendBrickDataLambda,
-				UMRMeshComponent*, This, this,
-				IMRMesh::FSendBrickDataArgs, Args, Args,
+			UMRMeshComponent* This = this;
+			ENQUEUE_RENDER_COMMAND(FSendBrickDataLambda)(
+				[This, Args](FRHICommandListImmediate& RHICmdList)
 				{
 					FMRMeshProxy* MRMeshProxy = static_cast<FMRMeshProxy*>(This->SceneProxy);
 					if (MRMeshProxy)
@@ -735,9 +734,9 @@ void UMRMeshComponent::ClearAllBrickData_Internal()
 	}
 
 	// Graphics update
-	ENQUEUE_UNIQUE_RENDER_COMMAND_ONEPARAMETER(
-		FClearAllBricksLambda,
-		UMRMeshComponent*, This, this,
+	UMRMeshComponent* This = this;
+	ENQUEUE_RENDER_COMMAND(FClearAllBricksLambda)(
+		[This](FRHICommandListImmediate& RHICmdList)
 		{
 			FMRMeshProxy* MRMeshProxy = static_cast<FMRMeshProxy*>(This->SceneProxy);
 			if (MRMeshProxy)
@@ -764,15 +763,15 @@ void UMRMeshComponent::SendRenderDynamicData_Concurrent()
 	if (SceneProxy)
 	{
 		// Enqueue command to send to render thread
-		ENQUEUE_UNIQUE_RENDER_COMMAND_TWOPARAMETER(
-			FSetMaterialLambda,
-			UMRMeshComponent*, This, this,
-			UMaterialInterface*, Material, Material,
+		UMRMeshComponent* This = this;
+		UMaterialInterface* InMaterial = Material;
+		ENQUEUE_RENDER_COMMAND(FSetMaterialLambda)(
+			[This, InMaterial](FRHICommandListImmediate& RHICmdList)
 			{
 				FMRMeshProxy* MRMeshProxy = static_cast<FMRMeshProxy*>(This->SceneProxy);
 				if (MRMeshProxy)
 				{
-					MRMeshProxy->RenderThread_SetMaterial(Material);
+					MRMeshProxy->RenderThread_SetMaterial(InMaterial);
 				}
 			});
 	}
diff --git a/Engine/Source/Runtime/MediaAssets/Private/Assets/MediaTexture.cpp b/Engine/Source/Runtime/MediaAssets/Private/Assets/MediaTexture.cpp
--- a/Engine/Source/Runtime/MediaAssets/Private/Assets/MediaTexture.cpp
+++ b/Engine/Source/Runtime/MediaAssets/Private/Assets/MediaTexture.cpp
@@ -329,11 +329,11 @@ void UMediaTexture::TickResource(FTimespan Timecode)
 	RenderParams.SrgbOutput = SRGB;
 
 	// redraw texture resource on render thread
-	ENQUEUE_UNIQUE_RENDER_COMMAND_TWOPARAMETER(MediaTextureResourceRender,
-		FMediaTextureResource*, ResourceParam, (FMediaTextureResource*)Resource,
-		FMediaTextureResource::FRenderParams, RenderParamsParam, RenderParams,
+	FMediaTextureResource* ResourceParam = (FMediaTextureResource*)Resource;
+	ENQUEUE_RENDER_COMMAND(MediaTextureResourceRender)(
+		[ResourceParam, RenderParams](FRHICommandListImmediate& RHICmdList)
 		{
-			ResourceParam->Render(RenderParamsParam);
+			ResourceParam->Render(RenderParams);
 		});
 }
 
diff --git a/Engine/Source/Runtime/MoviePlayer/Private/DefaultGameMoviePlayer.cpp b/Engine/Source/Runtime/MoviePlayer/Private/DefaultGameMoviePlayer.cpp
--- a/Engine/Source/Runtime/MoviePlayer/Private/DefaultGameMoviePlayer.cpp
+++ b/Engine/Source/Runtime/MoviePlayer/Private/DefaultGameMoviePlayer.cpp
@@ -488,15 +488,14 @@ void FDefaultGameMoviePlayer::WaitForMovieToFinish(bool bAllowEngineTick)
 					GEngine->Tick(DeltaTime, false);
 				}
 
-				ENQUEUE_UNIQUE_RENDER_COMMAND_TWOPARAMETER(
-					BeginLoadingMovieFrameAndTickMovieStreamer,
-					FDefaultGameMoviePlayer*, MoviePlayer, this,
-					float, DeltaTime, DeltaTime,
+				FDefaultGameMoviePlayer* InMoviePlayer = this;
+				ENQUEUE_RENDER_COMMAND(BeginLoadingMovieFrameAndTickMovieStreamer)(
+					[InMoviePlayer, DeltaTime](FRHICommandListImmediate& RHICmdList)
 					{
 						GFrameNumberRenderThread++;
 						GRHICommandList.GetImmediateCommandList().BeginFrame();
 				
-						MoviePlayer->TickStreamer(DeltaTime);
+						InMoviePlayer->TickStreamer(DeltaTime);
 					}
 				);
 				
diff --git a/Engine/Source/Runtime/MovieSceneCapture/Private/FrameGrabber.cpp b/Engine/Source/Runtime/MovieSceneCapture/Private/FrameGrabber.cpp
--- a/Engine/Source/Runtime/MovieSceneCapture/Private/FrameGrabber.cpp
+++ b/Engine/Source/Runtime/MovieSceneCapture/Private/FrameGrabber.cpp
@@ -220,12 +220,10 @@ void FViewportSurfaceReader::ResolveRenderTarget(FViewportSurfaceReader* RenderT
 		}
 	};
 
-	ENQUEUE_UNIQUE_RENDER_COMMAND_TWOPARAMETER(
-		ResolveCaptureFrameTexture,
-		TFunction<void(FRHICommandListImmediate&, FViewportSurfaceReader* InRenderToReadback)>, InRenderCommand, RenderCommand,
-		FViewportSurfaceReader*, InRenderToReadback, RenderToReadback,
+	ENQUEUE_RENDER_COMMAND(ResolveCaptureFrameTexture)(
+		[RenderCommand, RenderToReadback](FRHICommandListImmediate& RHICmdList)
 		{
-			InRenderCommand(RHICmdList, InRenderToReadback);
+			RenderCommand(RHICmdList, RenderToReadback);
 		}
 	);
 }
diff --git a/Engine/Source/Runtime/PreLoadScreen/Private/PreLoadScreenManager.cpp b/Engine/Source/Runtime/PreLoadScreen/Private/PreLoadScreenManager.cpp
--- a/Engine/Source/Runtime/PreLoadScreen/Private/PreLoadScreenManager.cpp
+++ b/Engine/Source/Runtime/PreLoadScreen/Private/PreLoadScreenManager.cpp
@@ -328,10 +328,8 @@ void FPreLoadScreenManager::EarlyPlayRenderFrameTick()
         float SlateDeltaTime = SlateApp.GetDeltaTime();
 
         //Setup Slate Render Command
-        ENQUEUE_UNIQUE_RENDER_COMMAND_TWOPARAMETER(
-            BeginPreLoadScreenFrame,
-            IPreLoadScreen*, ActivePreLoadScreen, ActivePreLoadScreen,
-            float, SlateDeltaTime, SlateDeltaTime,
+        ENQUEUE_RENDER_COMMAND(BeginPreLoadScreenFrame)(
+			[ActivePreLoadScreen, SlateDeltaTime](FRHICommandListImmediate& RHICmdList)
             {
                 if (FPreLoadScreenManager::ShouldRender())
                 {
diff --git a/Engine/Source/Runtime/RenderCore/Private/RenderingThread.cpp b/Engine/Source/Runtime/RenderCore/Private/RenderingThread.cpp
--- a/Engine/Source/Runtime/RenderCore/Private/RenderingThread.cpp
+++ b/Engine/Source/Runtime/RenderCore/Private/RenderingThread.cpp
@@ -543,46 +543,42 @@ struct FConsoleRenderThreadPropagation : public IConsoleThreadPropagation
 {
 	virtual void OnCVarChange(int32& Dest, int32 NewValue)
 	{
-		ENQUEUE_UNIQUE_RENDER_COMMAND_TWOPARAMETER(
-			OnCVarChange1,
-			int32&, Dest, Dest,
-			int32, NewValue, NewValue,
-		{
-			Dest = NewValue;
-		});
+		int32* DestPtr = &Dest;
+		ENQUEUE_RENDER_COMMAND(OnCVarChange1)(
+			[DestPtr, NewValue](FRHICommandListImmediate& RHICmdList)
+			{
+				*DestPtr = NewValue;
+			});
 	}
 	
 	virtual void OnCVarChange(float& Dest, float NewValue)
 	{
-		ENQUEUE_UNIQUE_RENDER_COMMAND_TWOPARAMETER(
-			OnCVarChange2,
-			float&, Dest, Dest,
-			float, NewValue, NewValue,
-		{
-			Dest = NewValue;
-		});
+		float* DestPtr = &Dest;
+		ENQUEUE_RENDER_COMMAND(OnCVarChange2)(
+			[DestPtr, NewValue](FRHICommandListImmediate& RHICmdList)
+			{
+				*DestPtr = NewValue;
+			});
 	}
 
 	virtual void OnCVarChange(bool& Dest, bool NewValue)
 	{
-		ENQUEUE_UNIQUE_RENDER_COMMAND_TWOPARAMETER(
-			OnCVarChange2,
-			bool&, Dest, Dest,
-			bool, NewValue, NewValue,
-		{
-			Dest = NewValue;
-		});
+		bool* DestPtr = &Dest;
+		ENQUEUE_RENDER_COMMAND(OnCVarChange2)(
+			[DestPtr, NewValue](FRHICommandListImmediate& RHICmdList)
+			{
+				*DestPtr = NewValue;
+			});
 	}
 	
 	virtual void OnCVarChange(FString& Dest, const FString& NewValue)
 	{
-		ENQUEUE_UNIQUE_RENDER_COMMAND_TWOPARAMETER(
-			OnCVarChange3,
-			FString&, Dest, Dest,
-			const FString&, NewValue, NewValue,
-		{
-			Dest = NewValue;
-		});
+		FString* DestPtr = &Dest;
+		ENQUEUE_RENDER_COMMAND(OnCVarChange3)(
+			[DestPtr, NewValue](FRHICommandListImmediate& RHICmdList)
+			{
+				*DestPtr = NewValue;
+			});
 	}
 
 	static FConsoleRenderThreadPropagation& GetSingleton()
diff --git a/Engine/Source/Runtime/Renderer/Private/AtmosphereRendering.cpp b/Engine/Source/Runtime/Renderer/Private/AtmosphereRendering.cpp
--- a/Engine/Source/Runtime/Renderer/Private/AtmosphereRendering.cpp
+++ b/Engine/Source/Runtime/Renderer/Private/AtmosphereRendering.cpp
@@ -1807,31 +1807,28 @@ void FScene::AddAtmosphericFog(UAtmosphericFogComponent* FogComponent)
 	check(FogComponent);
 
 	FAtmosphericFogSceneInfo* FogSceneInfo = new FAtmosphericFogSceneInfo(FogComponent, this);
-
-	ENQUEUE_UNIQUE_RENDER_COMMAND_TWOPARAMETER(
-		FAddAtmosphericFogCommand,
-		FScene*,Scene,this,
-		FAtmosphericFogSceneInfo*,FogSceneInfo,FogSceneInfo,
-	{
-		delete Scene->AtmosphericFog;
-		Scene->AtmosphericFog = FogSceneInfo;
-	});
+	FScene* Scene = this;
+	ENQUEUE_RENDER_COMMAND(FAddAtmosphericFogCommand)(
+		[Scene, FogSceneInfo](FRHICommandListImmediate& RHICmdList)
+		{
+			delete Scene->AtmosphericFog;
+			Scene->AtmosphericFog = FogSceneInfo;
+		});
 }
 
 void FScene::RemoveAtmosphericFog(UAtmosphericFogComponent* FogComponent)
 {
-	ENQUEUE_UNIQUE_RENDER_COMMAND_TWOPARAMETER(
-		FRemoveAtmosphericFogCommand,
-		FScene*,Scene,this,
-		UAtmosphericFogComponent*,FogComponent,FogComponent,
-	{
-		// Remove the given component's FExponentialHeightFogSceneInfo from the scene's fog array.
-		if (Scene->AtmosphericFog && Scene->AtmosphericFog->Component == FogComponent)
+	FScene* Scene = this;
+	ENQUEUE_RENDER_COMMAND(FRemoveAtmosphericFogCommand)(
+		[Scene, FogComponent](FRHICommandListImmediate& RHICmdList)
 		{
-			delete Scene->AtmosphericFog;
-			Scene->AtmosphericFog = NULL;
-		}
-	});
+			// Remove the given component's FExponentialHeightFogSceneInfo from the scene's fog array.
+			if (Scene->AtmosphericFog && Scene->AtmosphericFog->Component == FogComponent)
+			{
+				delete Scene->AtmosphericFog;
+				Scene->AtmosphericFog = NULL;
+			}
+		});
 }
 
 
diff --git a/Engine/Source/Runtime/Renderer/Private/DynamicPrimitiveDrawing.inl b/Engine/Source/Runtime/Renderer/Private/DynamicPrimitiveDrawing.inl
--- a/Engine/Source/Runtime/Renderer/Private/DynamicPrimitiveDrawing.inl
+++ b/Engine/Source/Runtime/Renderer/Private/DynamicPrimitiveDrawing.inl
@@ -55,12 +55,12 @@ inline void FViewElementPDI::RegisterDynamicResource(FDynamicPrimitiveResource*
 	if (IsInGameThread())
 	{
 		// Render thread might be reading the array while we are adding in the game thread
-		ENQUEUE_UNIQUE_RENDER_COMMAND_TWOPARAMETER(AddViewInfoDynamicResource,
-			FViewInfo*, InViewInfo, ViewInfo,
-			FDynamicPrimitiveResource*, InDynamicResource, DynamicResource,
+		FViewInfo* InViewInfo = ViewInfo;
+		ENQUEUE_RENDER_COMMAND(AddViewInfoDynamicResource)(
+			[InViewInfo, DynamicResource](FRHICommandListImmediate& RHICmdList)
 			{
-				InViewInfo->DynamicResources.Add(InDynamicResource);
-				InDynamicResource->InitPrimitiveResource();
+				InViewInfo->DynamicResources.Add(DynamicResource);
+				DynamicResource->InitPrimitiveResource();
 			});
 	}
 	else
diff --git a/Engine/Source/Runtime/Renderer/Private/PlanarReflectionRendering.cpp b/Engine/Source/Runtime/Renderer/Private/PlanarReflectionRendering.cpp
--- a/Engine/Source/Runtime/Renderer/Private/PlanarReflectionRendering.cpp
+++ b/Engine/Source/Runtime/Renderer/Private/PlanarReflectionRendering.cpp
@@ -443,13 +443,13 @@ void FScene::UpdatePlanarReflectionContents(UPlanarReflectionComponent* CaptureC
 
 		if (CaptureComponent->RenderTarget != NULL && CaptureComponent->RenderTarget->GetSizeXY() != DesiredPlanarReflectionTextureSize)
 		{
-			ENQUEUE_UNIQUE_RENDER_COMMAND_ONEPARAMETER( 
-				ReleaseRenderTargetCommand,
-				FPlanarReflectionRenderTarget*, RenderTarget, CaptureComponent->RenderTarget,
-			{
-				RenderTarget->ReleaseResource();
-				delete RenderTarget;
-			});
+			FPlanarReflectionRenderTarget* RenderTarget = CaptureComponent->RenderTarget;
+			ENQUEUE_RENDER_COMMAND(ReleaseRenderTargetCommand)(
+				[RenderTarget](FRHICommandListImmediate& RHICmdList)
+				{
+					RenderTarget->ReleaseResource();
+					delete RenderTarget;
+				});
 
 			CaptureComponent->RenderTarget = NULL;
 		}
@@ -458,24 +458,24 @@ void FScene::UpdatePlanarReflectionContents(UPlanarReflectionComponent* CaptureC
 		{
 			CaptureComponent->RenderTarget = new FPlanarReflectionRenderTarget(DesiredPlanarReflectionTextureSize);
 
-			ENQUEUE_UNIQUE_RENDER_COMMAND_TWOPARAMETER( 
-				InitRenderTargetCommand,
-				FPlanarReflectionRenderTarget*, RenderTarget, CaptureComponent->RenderTarget,
-				FPlanarReflectionSceneProxy*, SceneProxy, CaptureComponent->SceneProxy,
-			{
-				RenderTarget->InitResource();
-				SceneProxy->RenderTarget = nullptr;
-			});
+			FPlanarReflectionRenderTarget* RenderTarget = CaptureComponent->RenderTarget;
+			FPlanarReflectionSceneProxy* SceneProxy = CaptureComponent->SceneProxy;
+			ENQUEUE_RENDER_COMMAND(InitRenderTargetCommand)(
+				[RenderTarget, SceneProxy](FRHICommandListImmediate& RHICmdList)
+				{
+					RenderTarget->InitResource();
+					SceneProxy->RenderTarget = nullptr;
+				});
 		}
 		else
 		{
 			// Remove the render target on the planar reflection proxy so that this planar reflection is not getting drawn in its own FSceneRenderer.
-			ENQUEUE_UNIQUE_RENDER_COMMAND_ONEPARAMETER(
-				InitRenderTargetCommand,
-				FPlanarReflectionSceneProxy*, SceneProxy, CaptureComponent->SceneProxy,
-			{
-				SceneProxy->RenderTarget = nullptr;
-			});
+			FPlanarReflectionSceneProxy* SceneProxy = CaptureComponent->SceneProxy;
+			ENQUEUE_RENDER_COMMAND(InitRenderTargetCommand)(
+				[SceneProxy](FRHICommandListImmediate& RHICmdList)
+				{
+					SceneProxy->RenderTarget = nullptr;
+				});
 		}
 
 		const FMatrix ComponentTransform = CaptureComponent->GetComponentTransform().ToMatrixWithScale();
diff --git a/Engine/Source/Runtime/Renderer/Private/RendererScene.cpp b/Engine/Source/Runtime/Renderer/Private/RendererScene.cpp
--- a/Engine/Source/Runtime/Renderer/Private/RendererScene.cpp
+++ b/Engine/Source/Runtime/Renderer/Private/RendererScene.cpp
@@ -1659,10 +1659,10 @@ void FScene::AddLight(ULightComponent* Light)
 		++NumVisibleLights_GameThread;
 
 		// Send a command to the rendering thread to add the light to the scene.
-		ENQUEUE_UNIQUE_RENDER_COMMAND_TWOPARAMETER(
-			FAddLightCommand,
-			FScene*,Scene,this,
-			FLightSceneInfo*,LightSceneInfo,Proxy->LightSceneInfo,
+		FScene* Scene = this;
+		FLightSceneInfo* LightSceneInfo = Proxy->LightSceneInfo;
+		ENQUEUE_RENDER_COMMAND(FAddLightCommand)(
+			[Scene, LightSceneInfo](FRHICommandListImmediate& RHICmdList)
 			{
 				FScopeCycleCounter Context(LightSceneInfo->Proxy->GetStatId());
 				Scene->AddLightSceneInfo_RenderThread(LightSceneInfo);
@@ -1689,14 +1689,14 @@ void FScene::AddInvisibleLight(ULightComponent* Light)
 		INC_DWORD_STAT(STAT_SceneLights);
 
 		// Send a command to the rendering thread to add the light to the scene.
-		ENQUEUE_UNIQUE_RENDER_COMMAND_TWOPARAMETER(
-			FAddLightCommand,
-			FScene*,Scene,this,
-			FLightSceneInfo*,LightSceneInfo,Proxy->LightSceneInfo,
-		{
-			FScopeCycleCounter Context(LightSceneInfo->Proxy->GetStatId());
-			LightSceneInfo->Id = Scene->InvisibleLights.Add(FLightSceneInfoCompact(LightSceneInfo));
-		});
+		FScene* Scene = this;
+		FLightSceneInfo* LightSceneInfo = Proxy->LightSceneInfo;
+		ENQUEUE_RENDER_COMMAND(FAddLightCommand)(
+			[Scene, LightSceneInfo](FRHICommandListImmediate& RHICmdList)
+			{
+				FScopeCycleCounter Context(LightSceneInfo->Proxy->GetStatId());
+				LightSceneInfo->Id = Scene->InvisibleLights.Add(FLightSceneInfoCompact(LightSceneInfo));
+			});
 	}
 }
 
@@ -1709,23 +1709,23 @@ void FScene::SetSkyLight(FSkyLightSceneProxy* LightProxy)
 
 	ENQUEUE_RENDER_COMMAND(FSetSkyLightCommand)
 		([Scene, LightProxy](FRHICommandListImmediate& RHICmdList)
-	{
-		check(!Scene->SkyLightStack.Contains(LightProxy));
-		Scene->SkyLightStack.Push(LightProxy);
-		const bool bOriginalHadSkylight = Scene->ShouldRenderSkylightInBasePass(BLEND_Opaque);
+		{
+			check(!Scene->SkyLightStack.Contains(LightProxy));
+			Scene->SkyLightStack.Push(LightProxy);
+			const bool bOriginalHadSkylight = Scene->ShouldRenderSkylightInBasePass(BLEND_Opaque);
 
-		// Use the most recently enabled skylight
-		Scene->SkyLight = LightProxy;
+			// Use the most recently enabled skylight
+			Scene->SkyLight = LightProxy;
 
-		const bool bNewHasSkylight = Scene->ShouldRenderSkylightInBasePass(BLEND_Opaque);
+			const bool bNewHasSkylight = Scene->ShouldRenderSkylightInBasePass(BLEND_Opaque);
 
-		if (bOriginalHadSkylight != bNewHasSkylight)
-		{
-			// Mark the scene as needing static draw lists to be recreated if needed
-			// The base pass chooses shaders based on whether there's a skylight in the scene, and that is cached in static draw lists
-			Scene->bScenesPrimitivesNeedStaticMeshElementUpdate = true;
-		}
-	});
+			if (bOriginalHadSkylight != bNewHasSkylight)
+			{
+				// Mark the scene as needing static draw lists to be recreated if needed
+				// The base pass chooses shaders based on whether there's a skylight in the scene, and that is cached in static draw lists
+				Scene->bScenesPrimitivesNeedStaticMeshElementUpdate = true;
+			}
+		});
 }
 
 void FScene::DisableSkyLight(FSkyLightSceneProxy* LightProxy)
@@ -1795,13 +1795,13 @@ void FScene::AddDecal(UDecalComponent* Component)
 		INC_DWORD_STAT(STAT_SceneDecals);
 
 		// Send a command to the rendering thread to add the light to the scene.
-		ENQUEUE_UNIQUE_RENDER_COMMAND_TWOPARAMETER(
-			FAddDecalCommand,
-			FScene*,Scene,this,
-			FDeferredDecalProxy*,Proxy,Component->SceneProxy,
-		{
-			Scene->AddOrRemoveDecal_RenderThread(Proxy, true);
-		});
+		FScene* Scene = this;
+		FDeferredDecalProxy* Proxy = Component->SceneProxy;
+		ENQUEUE_RENDER_COMMAND(FAddDecalCommand)(
+			[Scene, Proxy](FRHICommandListImmediate& RHICmdList)
+			{
+				Scene->AddOrRemoveDecal_RenderThread(Proxy, true);
+			});
 	}
 }
 
@@ -1812,13 +1812,13 @@ void FScene::RemoveDecal(UDecalComponent* Component)
 		DEC_DWORD_STAT(STAT_SceneDecals);
 
 		// Send a command to the rendering thread to remove the light from the scene.
-		ENQUEUE_UNIQUE_RENDER_COMMAND_TWOPARAMETER(
-			FRemoveDecalCommand,
-			FScene*,Scene,this,
-			FDeferredDecalProxy*,Proxy,Component->SceneProxy,
-		{
-			Scene->AddOrRemoveDecal_RenderThread(Proxy, false);
-		});
+		FScene* Scene = this;
+		FDeferredDecalProxy* Proxy = Component->SceneProxy;
+		ENQUEUE_RENDER_COMMAND(FRemoveDecalCommand)(
+			[Scene, Proxy](FRHICommandListImmediate& RHICmdList)
+			{
+				Scene->AddOrRemoveDecal_RenderThread(Proxy, false);
+			});
 
 		// Disassociate the primitive's scene proxy.
 		Component->SceneProxy = NULL;
@@ -1830,14 +1830,14 @@ void FScene::UpdateDecalTransform(UDecalComponent* Decal)
 	if(Decal->SceneProxy)
 	{
 		//Send command to the rendering thread to update the decal's transform.
-		ENQUEUE_UNIQUE_RENDER_COMMAND_TWOPARAMETER(
-			UpdateTransformCommand,
-			FDeferredDecalProxy*,DecalSceneProxy,Decal->SceneProxy,
-			FTransform,ComponentToWorldIncludingDecalSize,Decal->GetTransformIncludingDecalSize(),
-		{
-			// Update the primitive's transform.
-			DecalSceneProxy->SetTransformIncludingDecalSize(ComponentToWorldIncludingDecalSize);
-		});
+		FDeferredDecalProxy* DecalSceneProxy = Decal->SceneProxy;
+		FTransform ComponentToWorldIncludingDecalSize = Decal->GetTransformIncludingDecalSize();
+		ENQUEUE_RENDER_COMMAND(UpdateTransformCommand)(
+			[DecalSceneProxy, ComponentToWorldIncludingDecalSize](FRHICommandListImmediate& RHICmdList)
+			{
+				// Update the primitive's transform.
+				DecalSceneProxy->SetTransformIncludingDecalSize(ComponentToWorldIncludingDecalSize);
+			});
 	}
 }
 
@@ -2008,20 +2008,19 @@ void FScene::ReleaseReflectionCubemap(UReflectionCaptureComponent* CaptureCompon
 
 	if (bRemoved)
 	{
-	ENQUEUE_UNIQUE_RENDER_COMMAND_TWOPARAMETER(
-		RemoveCaptureCommand,
-			UReflectionCaptureComponent*, Component, CaptureComponent,
-			FScene*, Scene, this,
-		{
-			const FCaptureComponentSceneState* ComponentStatePtr = Scene->ReflectionSceneData.AllocatedReflectionCaptureState.Find(Component);
-			if (ComponentStatePtr)
+		FScene* Scene = this;
+		ENQUEUE_RENDER_COMMAND(RemoveCaptureCommand)(
+			[CaptureComponent, Scene](FRHICommandListImmediate& RHICmdList)
 			{
-				// We track removed captures so we can remap them when reallocating the cubemap array
-				check(ComponentStatePtr->CubemapIndex != -1);
-				Scene->ReflectionSceneData.CubemapArraySlotsUsed[ComponentStatePtr->CubemapIndex] = false;
-		}
-		Scene->ReflectionSceneData.AllocatedReflectionCaptureState.Remove(Component);
-	});
+				const FCaptureComponentSceneState* ComponentStatePtr = Scene->ReflectionSceneData.AllocatedReflectionCaptureState.Find(CaptureComponent);
+				if (ComponentStatePtr)
+				{
+					// We track removed captures so we can remap them when reallocating the cubemap array
+					check(ComponentStatePtr->CubemapIndex != -1);
+					Scene->ReflectionSceneData.CubemapArraySlotsUsed[ComponentStatePtr->CubemapIndex] = false;
+				}
+				Scene->ReflectionSceneData.AllocatedReflectionCaptureState.Remove(CaptureComponent);
+			});
 	}
 }
 
@@ -2491,10 +2490,9 @@ void FScene::RemoveLight(ULightComponent* Light)
 		Light->SceneProxy = NULL;
 
 		// Send a command to the rendering thread to remove the light from the scene.
-		ENQUEUE_UNIQUE_RENDER_COMMAND_TWOPARAMETER(
-			FRemoveLightCommand,
-			FScene*,Scene,this,
-			FLightSceneInfo*,LightSceneInfo,LightSceneInfo,
+		FScene* Scene = this;
+		ENQUEUE_RENDER_COMMAND(FRemoveLightCommand)(
+			[Scene, LightSceneInfo](FRHICommandListImmediate& RHICmdList)
 			{
 				FScopeCycleCounter Context(LightSceneInfo->Proxy->GetStatId());
 				Scene->RemoveLightSceneInfo_RenderThread(LightSceneInfo);
@@ -2504,10 +2502,10 @@ void FScene::RemoveLight(ULightComponent* Light)
 
 void FScene::AddExponentialHeightFog(UExponentialHeightFogComponent* FogComponent)
 {
-	ENQUEUE_UNIQUE_RENDER_COMMAND_TWOPARAMETER(
-		FAddFogCommand,
-		FScene*,Scene,this,
-		FExponentialHeightFogSceneInfo,HeightFogSceneInfo,FExponentialHeightFogSceneInfo(FogComponent),
+	FScene* Scene = this;
+	FExponentialHeightFogSceneInfo HeightFogSceneInfo = FExponentialHeightFogSceneInfo(FogComponent);
+	ENQUEUE_RENDER_COMMAND(FAddFogCommand)(
+		[Scene, HeightFogSceneInfo](FRHICommandListImmediate& RHICmdList)
 		{
 			// Create a FExponentialHeightFogSceneInfo for the component in the scene's fog array.
 			new(Scene->ExponentialFogs) FExponentialHeightFogSceneInfo(HeightFogSceneInfo);
@@ -2516,10 +2514,9 @@ void FScene::AddExponentialHeightFog(UExponentialHeightFogComponent* FogComponen
 
 void FScene::RemoveExponentialHeightFog(UExponentialHeightFogComponent* FogComponent)
 {
-	ENQUEUE_UNIQUE_RENDER_COMMAND_TWOPARAMETER(
-		FRemoveFogCommand,
-		FScene*,Scene,this,
-		UExponentialHeightFogComponent*,FogComponent,FogComponent,
+	FScene* Scene = this;
+	ENQUEUE_RENDER_COMMAND(FRemoveFogCommand)(
+		[Scene, FogComponent](FRHICommandListImmediate& RHICmdList)
 		{
 			// Remove the given component's FExponentialHeightFogSceneInfo from the scene's fog array.
 			for(int32 FogIndex = 0;FogIndex < Scene->ExponentialFogs.Num();FogIndex++)
@@ -2546,10 +2543,9 @@ void FScene::AddWindSource(UWindDirectionalSourceComponent* WindComponent)
 	FWindSourceSceneProxy* SceneProxy = WindComponent->CreateSceneProxy();
 	WindComponent->SceneProxy = SceneProxy;
 
-	ENQUEUE_UNIQUE_RENDER_COMMAND_TWOPARAMETER(
-		FAddWindSourceCommand,
-		FScene*,Scene,this,
-		FWindSourceSceneProxy*,SceneProxy,SceneProxy,
+	FScene* Scene = this;
+	ENQUEUE_RENDER_COMMAND(FAddWindSourceCommand)(
+		[Scene, SceneProxy](FRHICommandListImmediate& RHICmdList)
 		{
 			Scene->WindSources.Add(SceneProxy);
 		});
@@ -2564,10 +2560,9 @@ void FScene::RemoveWindSource(UWindDirectionalSourceComponent* WindComponent)
 
 	if(SceneProxy)
 	{
-		ENQUEUE_UNIQUE_RENDER_COMMAND_TWOPARAMETER(
-			FRemoveWindSourceCommand,
-			FScene*,Scene,this,
-			FWindSourceSceneProxy*,SceneProxy,SceneProxy,
+		FScene* Scene = this;
+		ENQUEUE_RENDER_COMMAND(FRemoveWindSourceCommand)(
+			[Scene, SceneProxy](FRHICommandListImmediate& RHICmdList)
 			{
 				Scene->WindSources.Remove(SceneProxy);
 
@@ -2875,13 +2870,12 @@ void FScene::GetRelevantLights( UPrimitiveComponent* Primitive, TArray<const ULi
 /** Sets the precomputed visibility handler for the scene, or NULL to clear the current one. */
 void FScene::SetPrecomputedVisibility(const FPrecomputedVisibilityHandler* NewPrecomputedVisibilityHandler)
 {
-	ENQUEUE_UNIQUE_RENDER_COMMAND_TWOPARAMETER(
-		UpdatePrecomputedVisibility,
-		FScene*,Scene,this,
-		const FPrecomputedVisibilityHandler*,PrecomputedVisibilityHandler,NewPrecomputedVisibilityHandler,
-	{
-		Scene->PrecomputedVisibilityHandler = PrecomputedVisibilityHandler;
-	});
+	FScene* Scene = this;
+	ENQUEUE_RENDER_COMMAND(UpdatePrecomputedVisibility)(
+		[Scene, NewPrecomputedVisibilityHandler](FRHICommandListImmediate& RHICmdList)
+		{
+			Scene->PrecomputedVisibilityHandler = NewPrecomputedVisibilityHandler;
+		});
 }
 
 void FScene::UpdateStaticDrawLists_RenderThread(FRHICommandListImmediate& RHICmdList)
@@ -3053,13 +3047,12 @@ void FScene::Export( FArchive& Ar ) const
 void FScene::ApplyWorldOffset(FVector InOffset)
 {
 	// Send a command to the rendering thread to shift scene data
-	ENQUEUE_UNIQUE_RENDER_COMMAND_TWOPARAMETER(
-		FApplyWorldOffset,
-		FScene*,Scene,this,
-		FVector,InOffset,InOffset,
-	{
-		Scene->ApplyWorldOffset_RenderThread(InOffset);
-	});
+	FScene* Scene = this;
+	ENQUEUE_RENDER_COMMAND(FApplyWorldOffset)(
+		[Scene, InOffset](FRHICommandListImmediate& RHICmdList)
+		{
+			Scene->ApplyWorldOffset_RenderThread(InOffset);
+		});
 }
 
 void FScene::ApplyWorldOffset_RenderThread(FVector InOffset)
@@ -3164,13 +3157,12 @@ void FScene::OnLevelAddedToWorld(FName LevelAddedName, UWorld* InWorld, bool bIs
 		InWorld->PropagateLightingScenarioChange();
 	}
 
-	ENQUEUE_UNIQUE_RENDER_COMMAND_TWOPARAMETER(
-		FLevelAddedToWorld,
-		class FScene*, Scene, this,
-		FName, LevelName, LevelAddedName,
-	{
-		Scene->OnLevelAddedToWorld_RenderThread(LevelName);
-	});
+	FScene* Scene = this;
+	ENQUEUE_RENDER_COMMAND(FLevelAddedToWorld)(
+		[Scene, LevelAddedName](FRHICommandListImmediate& RHICmdList)
+		{
+			Scene->OnLevelAddedToWorld_RenderThread(LevelAddedName);
+		});
 }
 
 void FScene::OnLevelAddedToWorld_RenderThread(FName InLevelName)
diff --git a/Engine/Source/Runtime/SlateRHIRenderer/Private/SlateRHITextureAtlas.cpp b/Engine/Source/Runtime/SlateRHIRenderer/Private/SlateRHITextureAtlas.cpp
--- a/Engine/Source/Runtime/SlateRHIRenderer/Private/SlateRHITextureAtlas.cpp
+++ b/Engine/Source/Runtime/SlateRHIRenderer/Private/SlateRHITextureAtlas.cpp
@@ -71,12 +71,12 @@ void FSlateTextureAtlasRHI::ConditionalUpdateTexture( )
 	{
 		// Copy the game thread data. This is deleted on the render thread
 		FSlateTextureData* RenderThreadData = new FSlateTextureData( AtlasWidth, AtlasHeight, BytesPerPixel, AtlasData ); 
-		ENQUEUE_UNIQUE_RENDER_COMMAND_TWOPARAMETER( SlateUpdateAtlasTextureCommand,
-			FSlateTextureAtlasRHI&, Atlas, *this,
-			FSlateTextureData*, InRenderThreadData, RenderThreadData,
-		{
-			Atlas.UpdateTexture_RenderThread( InRenderThreadData );
-		});
+		FSlateTextureAtlasRHI* Atlas = this;
+		ENQUEUE_RENDER_COMMAND(SlateUpdateAtlasTextureCommand)(
+			[Atlas, RenderThreadData](FRHICommandListImmediate& RHICmdList)
+			{
+				Atlas->UpdateTexture_RenderThread( RenderThreadData );
+			});
 
 		bNeedsUpdate = false;
 
diff --git a/Engine/Source/Runtime/SlateRHIRenderer/Private/SlateUpdatableBuffer.cpp b/Engine/Source/Runtime/SlateRHIRenderer/Private/SlateUpdatableBuffer.cpp
--- a/Engine/Source/Runtime/SlateRHIRenderer/Private/SlateUpdatableBuffer.cpp
+++ b/Engine/Source/Runtime/SlateRHIRenderer/Private/SlateUpdatableBuffer.cpp
@@ -65,14 +65,14 @@ void FSlateUpdatableInstanceBuffer::UpdateRenderingData(int32 NumInstancesToUse)
 	NumInstances = NumInstancesToUse;
 	if (NumInstances > 0)
 	{
-		
 		// Enqueue a command to unlock the draw buffer after all windows have been drawn
-		ENQUEUE_UNIQUE_RENDER_COMMAND_TWOPARAMETER( SlateBeginDrawingWindowsCommand, 
-			FSlateUpdatableInstanceBuffer&, Self, *this,
-			int32, BufferIndex, FreeBufferIndex,
-		{
-			Self.UpdateRenderingData_RenderThread(RHICmdList, BufferIndex);
-		});
+		FSlateUpdatableInstanceBuffer* Self = this;
+		int32 BufferIndex = FreeBufferIndex;
+		ENQUEUE_RENDER_COMMAND(SlateBeginDrawingWindowsCommand)(
+			[Self, BufferIndex](FRHICommandListImmediate& RHICmdList)
+			{
+				Self->UpdateRenderingData_RenderThread(RHICmdList, BufferIndex);
+			});
 	
 		FreeBufferIndex = (FreeBufferIndex + 1) % SlateRHIConstants::NumBuffers;
 	}
diff --git a/Engine/Source/Runtime/UMG/Private/Blueprint/AsyncTaskDownloadImage.cpp b/Engine/Source/Runtime/UMG/Private/Blueprint/AsyncTaskDownloadImage.cpp
--- a/Engine/Source/Runtime/UMG/Private/Blueprint/AsyncTaskDownloadImage.cpp
+++ b/Engine/Source/Runtime/UMG/Private/Blueprint/AsyncTaskDownloadImage.cpp
@@ -110,13 +110,13 @@ void UAsyncTaskDownloadImage::HandleImageRequest(FHttpRequestPtr HttpRequest, FH
 						Texture->SRGB = true;
 						Texture->UpdateResource();
 
-						ENQUEUE_UNIQUE_RENDER_COMMAND_TWOPARAMETER(
-							FWriteRawDataToTexture,
-							FTexture2DDynamicResource*, TextureResource, static_cast<FTexture2DDynamicResource*>(Texture->Resource),
-							TArray<uint8>, RawData, *RawData,
-						{
-							WriteRawToTexture_RenderThread(TextureResource, RawData);
-						});
+						FTexture2DDynamicResource* TextureResource = static_cast<FTexture2DDynamicResource*>(Texture->Resource);
+						TArray<uint8> RawDataCopy = *RawData;
+						ENQUEUE_RENDER_COMMAND(FWriteRawDataToTexture)(
+							[TextureResource, RawDataCopy](FRHICommandListImmediate& RHICmdList)
+							{
+								WriteRawToTexture_RenderThread(TextureResource, RawDataCopy);
+							});
 						
 						OnSuccess.Broadcast(Texture);
 						return;
diff --git a/Engine/Source/Runtime/WebBrowserTexture/Private/WebBrowserTexture.cpp b/Engine/Source/Runtime/WebBrowserTexture/Private/WebBrowserTexture.cpp
--- a/Engine/Source/Runtime/WebBrowserTexture/Private/WebBrowserTexture.cpp
+++ b/Engine/Source/Runtime/WebBrowserTexture/Private/WebBrowserTexture.cpp
@@ -129,11 +129,11 @@ void UWebBrowserTexture::TickResource(TSharedPtr<FWebBrowserTextureSample, ESPMo
 		RenderParams.SampleSource = SampleQueue;
 	}
 
-	ENQUEUE_UNIQUE_RENDER_COMMAND_TWOPARAMETER(UWebBrowserTextureResourceRender,
-		FWebBrowserTextureResource*, ResourceParam, (FWebBrowserTextureResource*)Resource,
-		FWebBrowserTextureResource::FRenderParams, RenderParamsParam, RenderParams,
+	FWebBrowserTextureResource* ResourceParam = (FWebBrowserTextureResource*)Resource;
+	ENQUEUE_RENDER_COMMAND(UWebBrowserTextureResourceRender)(
+		[ResourceParam, RenderParams](FRHICommandListImmediate& RHICmdList)
 		{
-			ResourceParam->Render(RenderParamsParam);
+			ResourceParam->Render(RenderParams);
 		});
 }
 
@@ -144,8 +144,9 @@ void UWebBrowserTexture::UnregisterPlayerGuid()
 		return;
 	}
 
-	ENQUEUE_UNIQUE_RENDER_COMMAND_ONEPARAMETER(UWebBrowserTextureUnregisterPlayerGuid,
-		FGuid, PlayerGuid, WebPlayerGuid,
+	FGuid PlayerGuid = WebPlayerGuid;
+	ENQUEUE_RENDER_COMMAND(UWebBrowserTextureUnregisterPlayerGuid)(
+		[PlayerGuid](FRHICommandListImmediate& RHICmdList)
 		{
 			FExternalTextureRegistry::Get().UnregisterExternalTexture(PlayerGuid);
 		});
