Reference counting persistent FGraphicsMinimalPipelineStateId cache. This reduces total FGraphicsMinimalPipelineStateId cache memory usage and fixes memory leaking for cached draws. Leaking here was caused by caching pointers to RHI resources, which may be later recreated and same PSO wil be re-added to the cache without removing previous one.
#rnx

diff --git a/Engine/Source/Runtime/Renderer/Private/MeshPassProcessor.cpp b/Engine/Source/Runtime/Renderer/Private/MeshPassProcessor.cpp
--- a/Engine/Source/Runtime/Renderer/Private/MeshPassProcessor.cpp
+++ b/Engine/Source/Runtime/Renderer/Private/MeshPassProcessor.cpp
@@ -14,7 +14,7 @@
 #include "MeshPassProcessor.inl"
 #include "PipelineStateCache.h"
 
-TSet<FGraphicsMinimalPipelineStateInitializer> FGraphicsMinimalPipelineStateId::PersistentIdTable;
+TSet<FRefCountedGraphicsMinimalPipelineStateInitializer, RefCountedGraphicsMinimalPipelineStateInitializerKeyFuncs> FGraphicsMinimalPipelineStateId::PersistentIdTable;
 TSet<FGraphicsMinimalPipelineStateInitializer> FGraphicsMinimalPipelineStateId::OneFrameIdTable;
 FCriticalSection FGraphicsMinimalPipelineStateId::OneFrameIdTableCriticalSection;
 
@@ -302,10 +302,13 @@ FGraphicsMinimalPipelineStateId FGraphicsMinimalPipelineStateId::GetPersistentId
 {
 	FSetElementId TableId = PersistentIdTable.FindId(InPipelineState);
 
-	if (!TableId.IsValidId())
+	if (TableId.IsValidId())
+	{
+		++PersistentIdTable[TableId].RefNum;
+	}
+	else
 	{
-		// Note: grow-only table. Assuming finite and small enough set of FGraphicsMinimalPipelineStateInitializer permutations.
-		TableId = PersistentIdTable.Add(InPipelineState);
+		TableId = PersistentIdTable.Add(FRefCountedGraphicsMinimalPipelineStateInitializer(InPipelineState, 1));
 	}
 
 	checkf(TableId.AsInteger() < (MAX_uint32 >> 2), TEXT("Persistent FGraphicsMinimalPipelineStateId table overflow!"));
@@ -317,6 +320,21 @@ FGraphicsMinimalPipelineStateId FGraphicsMinimalPipelineStateId::GetPersistentId
 	return Ret;
 }
 
+void FGraphicsMinimalPipelineStateId::RemovePersistentId(FGraphicsMinimalPipelineStateId Id)
+{
+	check(Id.bValid && !Id.bOneFrameId);
+
+	const FSetElementId SetElementId = FSetElementId::FromInteger(Id.SetElementIndex);
+	FRefCountedGraphicsMinimalPipelineStateInitializer& RefCountedStateInitializer = PersistentIdTable[SetElementId];
+
+	check(RefCountedStateInitializer.RefNum > 0);
+	--RefCountedStateInitializer.RefNum;
+	if (RefCountedStateInitializer.RefNum <= 0)
+	{
+		PersistentIdTable.Remove(SetElementId);
+	}
+}
+
 FGraphicsMinimalPipelineStateId FGraphicsMinimalPipelineStateId::GetOneFrameId(const FGraphicsMinimalPipelineStateInitializer& InPipelineState)
 {
 	// Need to lock as this is called from multiple parallel tasks during mesh draw command generation or patching.
diff --git a/Engine/Source/Runtime/Renderer/Private/PrimitiveSceneInfo.cpp b/Engine/Source/Runtime/Renderer/Private/PrimitiveSceneInfo.cpp
--- a/Engine/Source/Runtime/Renderer/Private/PrimitiveSceneInfo.cpp
+++ b/Engine/Source/Runtime/Renderer/Private/PrimitiveSceneInfo.cpp
@@ -294,11 +294,13 @@ void FPrimitiveSceneInfo::RemoveCachedMeshDrawCommands()
 	{
 		const FCachedMeshDrawCommandInfo& CachedCommand = StaticMeshCommandInfos[CommandIndex];
 		const FSetElementId StateBucketId = FSetElementId::FromInteger(CachedCommand.StateBucketId);
-			
+
 		if (StateBucketId.IsValidId())
 		{
 			FMeshDrawCommandStateBucket& StateBucket = Scene->CachedMeshDrawCommandStateBuckets[StateBucketId];
 
+			FGraphicsMinimalPipelineStateId::RemovePersistentId(StateBucket.MeshDrawCommand.CachedPipelineId);
+
 			if (StateBucket.Num == 1)
 			{
 				Scene->CachedMeshDrawCommandStateBuckets.Remove(StateBucketId);
@@ -311,6 +313,8 @@ void FPrimitiveSceneInfo::RemoveCachedMeshDrawCommands()
 		else if (CachedCommand.CommandIndex >= 0)
 		{
 			FCachedPassMeshDrawList& PassDrawList = Scene->CachedDrawLists[CachedCommand.MeshPass];
+
+			FGraphicsMinimalPipelineStateId::RemovePersistentId(PassDrawList.MeshDrawCommands[CachedCommand.CommandIndex].CachedPipelineId);
 			PassDrawList.MeshDrawCommands.RemoveAt(CachedCommand.CommandIndex);
 
 			// Track the lowest index that might be free for faster AddAtLowestFreeIndex
diff --git a/Engine/Source/Runtime/Renderer/Public/MeshPassProcessor.h b/Engine/Source/Runtime/Renderer/Public/MeshPassProcessor.h
--- a/Engine/Source/Runtime/Renderer/Public/MeshPassProcessor.h
+++ b/Engine/Source/Runtime/Renderer/Public/MeshPassProcessor.h
@@ -94,6 +94,45 @@ public:
 
 static_assert(sizeof(FMeshPassMask::Data) * 8 >= EMeshPass::Num, "FMeshPassMask::Data is too small to fit all mesh passes.");
 
+struct FRefCountedGraphicsMinimalPipelineStateInitializer
+{
+	FRefCountedGraphicsMinimalPipelineStateInitializer(const FGraphicsMinimalPipelineStateInitializer& InStateInitializer, int32 InRefNum = 0)
+		: StateInitializer(InStateInitializer)
+		, RefNum(InRefNum)
+	{
+	}
+
+	FGraphicsMinimalPipelineStateInitializer StateInitializer;
+	int32 RefNum = 0;
+};
+
+struct RefCountedGraphicsMinimalPipelineStateInitializerKeyFuncs : DefaultKeyFuncs<FRefCountedGraphicsMinimalPipelineStateInitializer, false>
+{
+	typedef typename TCallTraits<FGraphicsMinimalPipelineStateInitializer>::ConstReference KeyInitType;
+
+	/**
+	 * @return True if the keys match.
+	 */
+	static FORCEINLINE bool Matches(KeyInitType A, KeyInitType B)
+	{
+		return A == B;
+	}
+
+	/**
+	 * @return The key used to index the given element.
+	 */
+	static FORCEINLINE KeyInitType GetSetKey(ElementInitType Element)
+	{
+		return Element.StateInitializer;
+	}
+
+	/** Calculates a hash index for a key. */
+	static FORCEINLINE uint32 GetKeyHash(KeyInitType Key)
+	{
+		return GetTypeHash(Key);
+	}
+};
+
 /** Uniquely represents a FGraphicsMinimalPipelineStateInitializer for fast compares. */
 class FGraphicsMinimalPipelineStateId
 {
@@ -128,16 +167,21 @@ public:
 			return OneFrameIdTable[SetElementId];
 		}
 
-		return PersistentIdTable[SetElementId];
+		return PersistentIdTable[SetElementId].StateInitializer;
 	}
 
 	/**
-	 * Get a pipeline id, which won't be ever released.
+	 * Get a ref counted persistent pipeline id, which needs to manually released.
 	 */
 	static FGraphicsMinimalPipelineStateId GetPersistentId(const FGraphicsMinimalPipelineStateInitializer& InPipelineState);
 
 	/**
-	 * Get a pipeline id, which is valid only for a single frame.
+	 * Removes a persistent pipeline Id from the global persistent Id table.
+	 */
+	static void RemovePersistentId(FGraphicsMinimalPipelineStateId Id);
+
+	/**
+	 * Get a pipeline id, which is valid only for a single frame and doesn't need to be released manually.
 	 */
 	RENDERER_API static FGraphicsMinimalPipelineStateId GetOneFrameId(const FGraphicsMinimalPipelineStateInitializer& InPipelineState);
 
@@ -159,7 +203,7 @@ private:
 		};
 	};
 
-	static TSet<FGraphicsMinimalPipelineStateInitializer> PersistentIdTable;
+	static TSet<FRefCountedGraphicsMinimalPipelineStateInitializer, RefCountedGraphicsMinimalPipelineStateInitializerKeyFuncs> PersistentIdTable;
 	static TSet<FGraphicsMinimalPipelineStateInitializer> OneFrameIdTable;
 	static FCriticalSection OneFrameIdTableCriticalSection;
 };
