Fixed parallel mesh pass race condition, where task resources could be freed while task is still in flight.
#rnx

diff --git a/Engine/Source/Runtime/Renderer/Private/SceneRendering.cpp b/Engine/Source/Runtime/Renderer/Private/SceneRendering.cpp
--- a/Engine/Source/Runtime/Renderer/Private/SceneRendering.cpp
+++ b/Engine/Source/Runtime/Renderer/Private/SceneRendering.cpp
@@ -2893,9 +2893,6 @@ void FSceneRenderer::WaitForTasksClearSnapshotsAndDeleteSceneRenderer(FRHIComman
 		RHICmdList.ImmediateFlush(EImmediateFlushType::WaitForOutstandingTasksOnly);
 	}
 
-	GPrimitiveIdVertexBufferPool.DiscardAll();
-	FGraphicsMinimalPipelineStateId::ResetOneFrameIdTable();
-
 	// Destroy cached preshadow transient arrays (allocated with SceneRenderingAllocator).
 	TArray<TRefCountPtr<FProjectedShadowInfo>>& CachedPreshadows = SceneRenderer->Scene->CachedPreshadows;
 	for (int32 CachedShadowIndex = 0; CachedShadowIndex < CachedPreshadows.Num(); ++CachedShadowIndex)
@@ -2922,6 +2919,10 @@ void FSceneRenderer::WaitForTasksClearSnapshotsAndDeleteSceneRenderer(FRHIComman
 		delete SceneRenderer;
 	}
 
+	// Can relase only after all mesh pass tasks are finished.
+	GPrimitiveIdVertexBufferPool.DiscardAll();
+	FGraphicsMinimalPipelineStateId::ResetOneFrameIdTable();
+
 	delete LocalRootMark;
 }
 
