FParallelMeshDrawCommandPass - surrounded slower parts with some quick scope counters.
#rnx

diff --git a/Engine/Source/Runtime/Renderer/Private/MeshDrawCommands.cpp b/Engine/Source/Runtime/Renderer/Private/MeshDrawCommands.cpp
--- a/Engine/Source/Runtime/Renderer/Private/MeshDrawCommands.cpp
+++ b/Engine/Source/Runtime/Renderer/Private/MeshDrawCommands.cpp
@@ -995,6 +995,7 @@ void FParallelMeshDrawCommandPass::DispatchPassSetup(
 		}
 		else
 		{
+			QUICK_SCOPE_CYCLE_COUNTER(STAT_MeshPassSetupImmediate);
 			FMeshDrawCommandPassSetupTask Task(TaskContext);
 			Task.AnyThreadTask();
 		}
@@ -1006,6 +1007,7 @@ void FParallelMeshDrawCommandPass::WaitForMeshPassSetupTask() const
 	if (TaskEventRef.IsValid())
 	{
 		// Need to wait on GetRenderThread_Local, as mesh pass setup task can wait on rendering thread inside InitResourceFromPossiblyParallelRendering().
+		QUICK_SCOPE_CYCLE_COUNTER(STAT_WaitForMeshPassSetupTask);
 		FTaskGraphInterface::Get().WaitUntilTaskCompletes(TaskEventRef, ENamedThreads::GetRenderThread_Local());
 	}
 }
@@ -1208,6 +1210,8 @@ void FParallelMeshDrawCommandPass::DispatchDraw(FParallelCommandListSet* Paralle
 	}
 	else
 	{
+		QUICK_SCOPE_CYCLE_COUNTER(STAT_MeshPassDrawImmediate);
+
 		WaitForMeshPassSetupTask();
 
 		if (TaskContext.bUseGPUScene)
