Subsurface scattering fix mostly for Vulkan

diff --git a/Engine/Source/Runtime/Renderer/Private/PostProcess/PostProcessSubsurface.cpp b/Engine/Source/Runtime/Renderer/Private/PostProcess/PostProcessSubsurface.cpp
--- a/Engine/Source/Runtime/Renderer/Private/PostProcess/PostProcessSubsurface.cpp
+++ b/Engine/Source/Runtime/Renderer/Private/PostProcess/PostProcessSubsurface.cpp
@@ -554,8 +554,13 @@ void FRCPassPostProcessSubsurfaceSetup::Process(FRenderingCompositePassContext&
 			View.StereoPass,
 			Context.HasHmdMesh(),
 			EDRF_UseTriangleOptimization);
+
 	}
 	Context.RHICmdList.EndRenderPass();
+
+	FTextureRHIParamRef TexturesToReadable[1];
+	TexturesToReadable[0] = DestRenderTarget.ShaderResourceTexture;
+	Context.RHICmdList.TransitionResources(EResourceTransitionAccess::EReadable, TexturesToReadable, ARRAY_COUNT(TexturesToReadable));
 }
 
 FPooledRenderTargetDesc FRCPassPostProcessSubsurfaceSetup::ComputeOutputDesc(EPassOutputId InPassOutputId) const
@@ -799,6 +804,10 @@ void FRCPassPostProcessSubsurface::Process(FRenderingCompositePassContext& Conte
 			EDRF_UseTriangleOptimization);
 	}
 	Context.RHICmdList.EndRenderPass();
+
+	FTextureRHIParamRef TexturesToReadable[1];
+	TexturesToReadable[0] = DestRenderTarget.ShaderResourceTexture;
+	Context.RHICmdList.TransitionResources(EResourceTransitionAccess::EReadable, TexturesToReadable, ARRAY_COUNT(TexturesToReadable));
 }
 
 
