Force primitiive uniform buffer update on the frame after a primitive is moved, since the uniform buffer contains PreviousLocalToWorld

diff --git a/Engine/Source/Runtime/Renderer/Private/RendererScene.cpp b/Engine/Source/Runtime/Renderer/Private/RendererScene.cpp
--- a/Engine/Source/Runtime/Renderer/Private/RendererScene.cpp
+++ b/Engine/Source/Runtime/Renderer/Private/RendererScene.cpp
@@ -1180,7 +1180,7 @@ void FScene::UpdatePrimitiveTransform_RenderThread(FRHICommandListImmediate& RHI
 	// (note that the octree update relies on the bounds not being modified yet).
 	PrimitiveSceneInfo->RemoveFromScene(bUpdateStaticDrawLists);
 
-	VelocityData.UpdateTransform(PrimitiveSceneInfo->PrimitiveComponentId, LocalToWorld, PrimitiveSceneProxy->GetLocalToWorld());
+	VelocityData.UpdateTransform(PrimitiveSceneInfo, LocalToWorld, PrimitiveSceneProxy->GetLocalToWorld());
 
 	if (GWarningOnRedundantTransformUpdate && PrimitiveSceneProxy->WouldSetTransformBeRedundant(LocalToWorld, WorldBounds, LocalBounds, AttachmentRootPosition))
 	{
@@ -2251,6 +2251,32 @@ bool FScene::GetPreviousLocalToWorld(const FPrimitiveSceneInfo* PrimitiveSceneIn
 	return VelocityData.GetComponentPreviousLocalToWorld(PrimitiveSceneInfo->PrimitiveComponentId, OutPreviousLocalToWorld);
 }
 
+void FSceneVelocityData::StartFrame(FScene* Scene)
+{
+	InternalFrameIndex++;
+
+	const bool bTrimOld = InternalFrameIndex % 100 == 0;
+
+	for (TMap<FPrimitiveComponentId, FComponentVelocityData>::TIterator It(ComponentData); It; ++It)
+	{
+		FComponentVelocityData& VelocityData = It.Value();
+		VelocityData.PreviousLocalToWorld = VelocityData.LocalToWorld;
+		VelocityData.bPreviousLocalToWorldValid = true;
+
+		if (InternalFrameIndex - VelocityData.LastFrameUpdated == 1)
+		{
+			// Recreate PrimitiveUniformBuffer on the frame after the primitive moved, since it contains PreviousLocalToWorld
+			VelocityData.PrimitiveSceneInfo->SetNeedsUniformBufferUpdate(true);
+			AddPrimitiveToUpdateGPU(*Scene, VelocityData.PrimitiveSceneInfo->GetIndex());
+		}
+
+		if (bTrimOld && (InternalFrameIndex - VelocityData.LastFrameUsed) > 10)
+		{
+			It.RemoveCurrent();
+		}
+	}
+}
+
 void FScene::GetPrimitiveUniformShaderParameters_RenderThread(const FPrimitiveSceneInfo* PrimitiveSceneInfo, bool& bHasPrecomputedVolumetricLightmap, FMatrix& PreviousLocalToWorld, int32& SingleCaptureIndex) const 
 {
 	bHasPrecomputedVolumetricLightmap = VolumetricLightmapSceneData.HasData();
diff --git a/Engine/Source/Runtime/Renderer/Private/ScenePrivate.h b/Engine/Source/Runtime/Renderer/Private/ScenePrivate.h
--- a/Engine/Source/Runtime/Renderer/Private/ScenePrivate.h
+++ b/Engine/Source/Runtime/Renderer/Private/ScenePrivate.h
@@ -2140,9 +2140,11 @@ class FComponentVelocityData
 {
 public:
 
+	FPrimitiveSceneInfo* PrimitiveSceneInfo;
 	FMatrix LocalToWorld;
 	FMatrix PreviousLocalToWorld;
 	mutable uint64 LastFrameUsed;
+	uint64 LastFrameUpdated;
 	bool bPreviousLocalToWorldValid = false;
 };
 
@@ -2156,24 +2158,7 @@ public:
 	/**
 	 * Must be called once per frame, even when there are multiple BeginDrawingViewports.
 	 */
-	void StartFrame()
-	{
-		InternalFrameIndex++;
-
-		const bool bTrimOld = InternalFrameIndex % 100 == 0;
-
-		for (TMap<FPrimitiveComponentId, FComponentVelocityData>::TIterator It(ComponentData); It; ++It)
-		{
-			FComponentVelocityData& VelocityData = It.Value();
-			VelocityData.PreviousLocalToWorld = VelocityData.LocalToWorld;
-			VelocityData.bPreviousLocalToWorldValid = true;
-
-			if (bTrimOld && (InternalFrameIndex - VelocityData.LastFrameUsed) > 10)
-			{
-				It.RemoveCurrent();
-			}
-		}
-	}
+	void StartFrame(FScene* Scene);
 
 	/** 
 	 * Looks up the PreviousLocalToWorld state for the given component.  Returns false if none is found (the primitive has never been moved). 
@@ -2196,11 +2181,13 @@ public:
 	/** 
 	 * Updates a primitives current LocalToWorld state.
 	 */
-	void UpdateTransform(FPrimitiveComponentId PrimitiveComponentId, FMatrix LocalToWorld, FMatrix PreviousLocalToWorld)
+	void UpdateTransform(FPrimitiveSceneInfo* PrimitiveSceneInfo, FMatrix LocalToWorld, FMatrix PreviousLocalToWorld)
 	{
-		FComponentVelocityData& VelocityData = ComponentData.FindOrAdd(PrimitiveComponentId);
+		FComponentVelocityData& VelocityData = ComponentData.FindOrAdd(PrimitiveSceneInfo->PrimitiveComponentId);
 		VelocityData.LocalToWorld = LocalToWorld;
 		VelocityData.LastFrameUsed = InternalFrameIndex;
+		VelocityData.LastFrameUpdated = InternalFrameIndex;
+		VelocityData.PrimitiveSceneInfo = PrimitiveSceneInfo;
 
 		// If this transform state is newly added, use the passed in PreviousLocalToWorld for this frame
 		if (!VelocityData.bPreviousLocalToWorldValid)
@@ -2847,7 +2834,7 @@ public:
 
 	virtual void StartFrame() override
 	{
-		VelocityData.StartFrame();
+		VelocityData.StartFrame(this);
 	}
 
 	virtual uint32 GetFrameNumber() const override
