Fixed incorrect velocity vectors on moving objects which are having MarkRenderStateDirty() called on them every frame from BP / gameplay changes (eg UPrimitiveComponent::SetCastShadow).

diff --git a/Engine/Source/Runtime/Renderer/Private/RendererScene.cpp b/Engine/Source/Runtime/Renderer/Private/RendererScene.cpp
--- a/Engine/Source/Runtime/Renderer/Private/RendererScene.cpp
+++ b/Engine/Source/Runtime/Renderer/Private/RendererScene.cpp
@@ -688,7 +688,8 @@ void FScene::AddPrimitiveSceneInfo_RenderThread(FRHICommandListImmediate& RHICmd
 	CheckPrimitiveArrays();
 
 	Primitives.Add(PrimitiveSceneInfo);
-	PrimitiveTransforms.Add(PrimitiveSceneInfo->Proxy->GetLocalToWorld());
+	const FMatrix LocalToWorld = PrimitiveSceneInfo->Proxy->GetLocalToWorld();
+	PrimitiveTransforms.Add(LocalToWorld);
 	PrimitiveSceneProxies.Add(PrimitiveSceneInfo->Proxy);
 	PrimitiveBounds.AddUninitialized();
 	PrimitiveFlagsCompact.AddUninitialized();
@@ -796,6 +797,13 @@ void FScene::AddPrimitiveSceneInfo_RenderThread(FRHICommandListImmediate& RHICmd
 		}
 	}
 
+	if (PrimitiveSceneInfo->Proxy->IsMovable())
+	{
+		// We must register the initial LocalToWorld with the velocity state. 
+		// In the case of a moving component with MarkRenderStateDirty() called every frame, UpdateTransform will never happen.
+		VelocityData.UpdateTransform(PrimitiveSceneInfo, LocalToWorld, LocalToWorld);
+	}
+
 	AddPrimitiveToUpdateGPU(*this, SourceIndex);
 
 	DistanceFieldSceneData.AddPrimitive(PrimitiveSceneInfo);
@@ -1457,8 +1465,11 @@ void FScene::RemovePrimitiveSceneInfo_RenderThread(FPrimitiveSceneInfo* Primitiv
 	
 	CheckPrimitiveArrays();
 
-	// Remove primitive's motion blur information.
-	VelocityData.RemoveTransform(PrimitiveSceneInfo->PrimitiveComponentId);
+	if (PrimitiveSceneInfo->Proxy->IsMovable())
+	{
+		// Remove primitive's motion blur information.
+		VelocityData.RemoveFromScene(PrimitiveSceneInfo->PrimitiveComponentId);
+	}
 
 	// Unlink the primitive from its shadow parent.
 	PrimitiveSceneInfo->UnlinkAttachmentGroup();
@@ -2263,7 +2274,7 @@ void FSceneVelocityData::StartFrame(FScene* Scene)
 		VelocityData.PreviousLocalToWorld = VelocityData.LocalToWorld;
 		VelocityData.bPreviousLocalToWorldValid = true;
 
-		if (InternalFrameIndex - VelocityData.LastFrameUpdated == 1)
+		if ((InternalFrameIndex - VelocityData.LastFrameUpdated == 1) && VelocityData.PrimitiveSceneInfo)
 		{
 			// Recreate PrimitiveUniformBuffer on the frame after the primitive moved, since it contains PreviousLocalToWorld
 			VelocityData.PrimitiveSceneInfo->SetNeedsUniformBufferUpdate(true);
diff --git a/Engine/Source/Runtime/Renderer/Private/ScenePrivate.h b/Engine/Source/Runtime/Renderer/Private/ScenePrivate.h
--- a/Engine/Source/Runtime/Renderer/Private/ScenePrivate.h
+++ b/Engine/Source/Runtime/Renderer/Private/ScenePrivate.h
@@ -2197,12 +2197,14 @@ public:
 		}
 	}
 
-	/**
-	* Removes primitives velocity data.
-	*/
-	void RemoveTransform(FPrimitiveComponentId PrimitiveComponentId)
+	void RemoveFromScene(FPrimitiveComponentId PrimitiveComponentId)
 	{
-		ComponentData.Remove(PrimitiveComponentId);
+		FComponentVelocityData* VelocityData = ComponentData.Find(PrimitiveComponentId);
+
+		if (VelocityData)
+		{
+			VelocityData->PrimitiveSceneInfo = nullptr;
+		}
 	}
 
 	void ApplyOffset(FVector Offset)
