Fixed invalid primitive index used from cl 5074346

diff --git a/Engine/Source/Runtime/Renderer/Private/RendererScene.cpp b/Engine/Source/Runtime/Renderer/Private/RendererScene.cpp
--- a/Engine/Source/Runtime/Renderer/Private/RendererScene.cpp
+++ b/Engine/Source/Runtime/Renderer/Private/RendererScene.cpp
@@ -2278,7 +2278,11 @@ void FSceneVelocityData::StartFrame(FScene* Scene)
 		{
 			// Recreate PrimitiveUniformBuffer on the frame after the primitive moved, since it contains PreviousLocalToWorld
 			VelocityData.PrimitiveSceneInfo->SetNeedsUniformBufferUpdate(true);
-			AddPrimitiveToUpdateGPU(*Scene, VelocityData.PrimitiveSceneInfo->GetIndex());
+
+			if (VelocityData.PrimitiveSceneInfo->IsIndexValid())
+			{
+				AddPrimitiveToUpdateGPU(*Scene, VelocityData.PrimitiveSceneInfo->GetIndex());
+			}
 		}
 
 		if (bTrimOld && (InternalFrameIndex - VelocityData.LastFrameUsed) > 10)
