Copying //UE4/Private-Woodchuck-Staging[REDACTED_HANDLE] to Release-4.22 (//UE4/Release-4.22)

diff --git a/Engine/Shaders/Private/DataStructures/SceneViewFamilyBlackboard.ush b/Engine/Shaders/Private/DataStructures/SceneViewFamilyBlackboard.ush
--- a/Engine/Shaders/Private/DataStructures/SceneViewFamilyBlackboard.ush
+++ b/Engine/Shaders/Private/DataStructures/SceneViewFamilyBlackboard.ush
@@ -6,8 +6,10 @@
 
 #pragma once
 
-#include "../SceneTexturesCommon.ush"
-#include "../DeferredShadingCommon.ush"
+#include "SceneTexturesCommon.ush"
+#include "DeferredShadingCommon.ush"
+
+uint bIsSceneLightingChannelsValid;
 
 // Matches FSceneViewFamilyBlackboard
 Texture2D SceneDepthBuffer;
@@ -17,6 +19,7 @@ Texture2D SceneGBufferB;
 Texture2D SceneGBufferC;
 Texture2D SceneGBufferD;
 Texture2D SceneGBufferE;
+Texture2D<uint> SceneLightingChannels;
 
 
 // @param UV - UV space in the GBuffer textures (BufferSize resolution)
@@ -37,3 +40,14 @@ FGBufferData GetGBufferDataFromSceneBlackboard(float2 UV, bool bGetNormalizedNor
 
 	return DecodeGBufferData(GBufferA, GBufferB, GBufferC, GBufferD, GBufferE, GBufferVelocity, CustomNativeDepth, CustomStencil, SceneDepth, bGetNormalizedNormal, CheckerFromSceneColorUV(UV));
 }
+
+/** Returns the light channel mask that should be executed for this pixel. */
+uint GetSceneLightingChannel(uint2 PixelCoord)
+{
+	BRANCH
+	if (bIsSceneLightingChannelsValid)
+	{
+		return SceneLightingChannels.Load(uint3(PixelCoord, 0)).x;
+	}
+	return ~0;
+}
diff --git a/Engine/Shaders/Private/RayTracing/CommonRayTracingBuiltInResources.ush b/Engine/Shaders/Private/RayTracing/CommonRayTracingBuiltInResources.ush
deleted file mode 100644
diff --git a/Engine/Shaders/Private/RayTracing/RayTracingDebugDefinitions.ush b/Engine/Shaders/Private/RayTracing/RayTracingDebugDefinitions.ush
deleted file mode 100644
diff --git a/Engine/Shaders/Private/RayTracing/RayTracingDefinitions.ush b/Engine/Shaders/Private/RayTracing/RayTracingDefinitions.ush
deleted file mode 100644
diff --git a/Engine/Shaders/Private/RayTracing/RayTracingRectLightOcclusionRGS.usf b/Engine/Shaders/Private/RayTracing/RayTracingRectLightOcclusionRGS.usf
--- a/Engine/Shaders/Private/RayTracing/RayTracingRectLightOcclusionRGS.usf
+++ b/Engine/Shaders/Private/RayTracing/RayTracingRectLightOcclusionRGS.usf
@@ -10,7 +10,7 @@
 #include "MipTreeCommon.ush"
 
 RaytracingAccelerationStructure TLAS;
-RWTexture2D<float4> RWOcclusionMaskUAV;
+RWTexture2D<float4> RWLuminanceUAV;
 RWTexture2D<float> RWRayDistanceUAV;
 
 uint CalcLinearIndex(uint2 PixelCoord)
@@ -45,7 +45,7 @@ bool GenerateRectLightImportanceSampledRay(
 	out float RayPdf
 )
 {
-	uint MipCount = Log2(RectLight.MipTreeDimensions.x);
+	uint MipCount = log2(RectLight.MipTreeDimensions.x);
 	uint2 MipPixel = 0;
 	float MipPdf = 1.0;
 
@@ -168,7 +168,7 @@ bool GenerateRectLightImportanceSampledRay(
 }
 
 [shader("raygeneration")]
-void RectLightOcclusionRGS()
+void RectLightRGS()
 {
 	uint2 PixelCoord = DispatchRaysIndex().xy;
 
@@ -195,41 +195,60 @@ void RectLightOcclusionRGS()
 	float HitCount = 0.0;
 	uint SamplesPerPixel = RectLight.SamplesPerPixel;
 
-	// Mask out depth values at or beyond far plane
-	if (Depth >= SCENE_TEXTURES_DISABLED_SCENE_DEPTH_VALUE)
+	// Mask out depth values that are infinitely far away
+	float DeviceZ = SceneTexturesStruct.SceneDepthTexture.Load(int3(PixelCoord, 0)).r;
+	bool IsFiniteDepth = DeviceZ > 0.0;
+	bool bTraceRay = (
+		IsFiniteDepth &&
+		GBuffer.ShadingModelID != SHADINGMODELID_UNLIT);
+	if (!bTraceRay)
 	{
 		SamplesPerPixel = 0.0;
 	}
 
+	float3 LightPosition = RectLight.Position;
+	float3 LightNormal = RectLight.Normal;
+	float3 LightdPdu = RectLight.dPdu;
+	float3 LightdPdv = RectLight.dPdv;
+	float LightWidth = RectLight.Width;
+	float LightHeight = RectLight.Height;
+	
+	// Encode the rect light  parameter into the light parameters struct
+	// todo: send direction the light shader parameter data instead of the custon RectLight data.
+	FLightShaderParameters LightParameters;
+	LightParameters.Position = RectLight.Position;
+	LightParameters.Direction = -RectLight.Normal;
+	LightParameters.Tangent = RectLight.dPdu;
+	LightParameters.SourceRadius = 0.5f * RectLight.Width;
+	LightParameters.SourceLength = 0.5f * RectLight.Height;
+	
 	float3 ExitantRadiance = 0.0;
 	float3 DiffuseExitantRadiance = 0.0;
 	for (uint SampleIndex = 0; SampleIndex < SamplesPerPixel; ++SampleIndex)
 	{
 		RayDesc Ray;
-		float3 LightPosition = RectLight.Position;
-		float3 LightNormal = RectLight.Normal;
-		float3 LightdPdu = RectLight.dPdu;
-		float3 LightdPdv = RectLight.dPdv;
-		float LightWidth = RectLight.Width;
-		float LightHeight = RectLight.Height;
 
 		uint DummyVariable;
 		float2 RandSample = RandomSequence_GenerateSample2D(RandSequence, DummyVariable);
 		float RayPdf = 0.0;
 
+		bool bIsValid = false;
 #if TEXTURE_IMPORTANCE_SAMPLING
-		bool bIsValid = RectLight.bIsTextureImportanceSampling ?
-			GenerateRectLightImportanceSampledRay(WorldPosition, WorldNormal, LightPosition, LightNormal, LightdPdu, LightdPdv, LightWidth, LightHeight, RandSequence, Ray.Origin, Ray.Direction, Ray.TMin, Ray.TMax, RayPdf) :
-			GenerateRectLightOcclusionRay(WorldPosition, WorldNormal, LightPosition, LightNormal, LightdPdu, LightdPdv, LightWidth, LightHeight, RandSample, Ray.Origin, Ray.Direction, Ray.TMin, Ray.TMax, RayPdf);
-#else
-		bool bIsValid = GenerateRectLightOcclusionRay(WorldPosition, WorldNormal, LightPosition, LightNormal, LightdPdu, LightdPdv, LightWidth, LightHeight, RandSample, Ray.Origin, Ray.Direction, Ray.TMin, Ray.TMax, RayPdf);
+		if (RectLight.bIsTextureImportanceSampling)
+		{
+			bIsValid = GenerateRectLightImportanceSampledRay(WorldPosition, WorldNormal, LightPosition, LightNormal, LightdPdu, LightdPdv, LightWidth, LightHeight, RandSequence, Ray.Origin, Ray.Direction, Ray.TMin, Ray.TMax, RayPdf);
+		}
+		else
 #endif
-		ApplyPositionBias(Ray, WorldNormal, RectLight.MaxNormalBias);
+		{
+			bIsValid = GenerateRectLightOcclusionRay(LightParameters, WorldPosition, WorldNormal, RandSample, Ray.Origin, Ray.Direction, Ray.TMin, Ray.TMax, RayPdf);
+		}
 
 		if (!bIsValid)
 		{
 			continue;
 		}
+		ApplyPositionBias(Ray, WorldNormal, RectLight.MaxNormalBias);
 
 		float NoL = dot(WorldNormal, Ray.Direction);
 		if (NoL <= 0.0)
@@ -246,9 +265,9 @@ void RectLightOcclusionRGS()
 		TraceRay(
 			TLAS,      // AccelerationStructure
 			RayFlags,
-			0xFF,      // InstanceInclusionMask
-			0,         // RayContributionToHitGroupIndex
-			1,         // MultiplierForGeometryContributionToShaderIndex
+			RAY_TRACING_MASK_OPAQUE,          // InstanceInclusionMask
+			RAY_TRACING_SHADER_SLOT_MATERIAL, // RayContributionToHitGroupIndex
+			RAY_TRACING_NUM_SHADER_SLOTS,     // MultiplierForGeometryContributionToShaderIndex
 			0,         // MissShaderIndex
 			Ray,       // RayDesc
 			Payload    // Payload
@@ -261,14 +280,13 @@ void RectLightOcclusionRGS()
 		}
 		else
 		{
-#if CALC_DIRECT_LIGHTING
 			// Sample light texture
 			float3 LightPosition = Ray.Origin + Ray.Direction * Ray.TMax;
 			float3 RectLightCorner = RectLight.Position - 0.5 * (RectLight.dPdu * RectLight.Width + RectLight.dPdv * RectLight.Height);
 			float3 DeltaPosition = LightPosition - RectLightCorner;
-			float2 UV = float2(dot(DeltaPosition, RectLight.dPdu) / RectLight.Width, dot(DeltaPosition, RectLight.dPdv) / RectLight.Height);
-			UV = 1.0 - UV;
-			float3 TextureColor = Texture2DSample(RectLight.Texture, RectLight.TextureSampler, UV).xyz;
+			float2 SampleUV = float2(dot(DeltaPosition, RectLight.dPdu) / RectLight.Width, dot(DeltaPosition, RectLight.dPdv) / RectLight.Height);
+			SampleUV = 1.0 - SampleUV;
+			float3 TextureColor = Texture2DSample(RectLight.Texture, RectLight.TextureSampler, SampleUV).xyz;
 			float3 IncomingRadiance = RectLight.Color * TextureColor;
 
 			// Evaluate material
@@ -286,19 +304,11 @@ void RectLightOcclusionRGS()
 			float3 Irradiance = IncomingRadiance * NoL / RayPdf;
 			ExitantRadiance += Irradiance * Brdf;
 			DiffuseExitantRadiance += Irradiance * (LightingSample.Diffuse + LightingSample.Transmission);
-#else
-			// Record exitant radiance as visibility score
-			ExitantRadiance += 1.0;
-#endif
 		}
 	}
 
 	float3 Result = (SamplesPerPixel > 0) ? ExitantRadiance / SamplesPerPixel : ExitantRadiance;
-#if	ENCODE_VISIBILITY
-	// Encode visibility term to match shadow map pipeline
-	Result = sqrt(Result);
-#endif
-	RWOcclusionMaskUAV[PixelCoord].rgb = Result;
-	RWOcclusionMaskUAV[PixelCoord].a = 1.0;
+	RWLuminanceUAV[PixelCoord].rgb = Result;
+	RWLuminanceUAV[PixelCoord].a = 1.0;
 	RWRayDistanceUAV[PixelCoord] = (HitCount > 0.0) ? RayDistance / HitCount : 1.0e27;
 }
diff --git a/Engine/Shaders/Shared/RayTracingBuiltInResources.h b/Engine/Shaders/Shared/RayTracingBuiltInResources.h
new file mode 100644
--- /dev/null
+++ b/Engine/Shaders/Shared/RayTracingBuiltInResources.h
@@ -0,0 +1,94 @@
+// Copyright 1998-2019 Epic Games, Inc. All Rights Reserved.
+
+/*================================================================================================
+	RayTracingBuiltInResources.h: used in ray tracing shaders and C++ code to define resources 
+	available in all hit groups, such as root nostants, index and vertex buffers.
+	!!! Changing this file requires recompilation of the engine !!!
+=================================================================================================*/
+
+#pragma once
+
+#ifndef RAYTRACINGBUILTINRESOURCES_USH_INCLUDED
+#define RAYTRACINGBUILTINRESOURCES_USH_INCLUDED // Workarround for UE-66460
+
+#include "RayTracingDefinitions.h"
+
+#if defined(__cplusplus)
+	#define INCLUDED_FROM_CPP_CODE  1
+	#define INCLUDED_FROM_HLSL_CODE 0
+#elif defined(SM5_PROFILE)
+	// #dxr_todo: we should use a built-in macro to detect if this shader is compiled using DXC (depends on [REDACTED_PROFILE_URL]/DirectXShaderCompiler/issues/1686)
+	#define INCLUDED_FROM_CPP_CODE  0
+	#define INCLUDED_FROM_HLSL_CODE 1
+#else
+	#error Unknown Compiler
+#endif
+
+#if INCLUDED_FROM_HLSL_CODE
+	#define UINT_TYPE uint
+#elif INCLUDED_FROM_CPP_CODE
+	#define UINT_TYPE unsigned int
+#endif
+
+struct FHitGroupSystemRootConstants
+{
+	// Config is a bitfield:
+	// uint IndexStride  : 8; // Can be just 1 bit to indicate 16 or 32 bit indices
+	// uint VertexStride : 8; // Can be just 2 bits to indicate float3, float2 or half2 format
+	// uint Unused       : 16;
+	UINT_TYPE Config;
+
+	// Offset into HitGroupSystemIndexBuffer
+	UINT_TYPE IndexBufferOffsetInBytes;
+
+	// User-provided constant assigned to the hit group
+	UINT_TYPE UserData;
+
+	// Padding to ensure that root parameters are properly aligned to 8-byte boundary
+	UINT_TYPE Unused;
+
+	// Helper functions
+
+	UINT_TYPE GetIndexStride()
+	{
+		return Config & 0xFF;
+	}
+
+	UINT_TYPE GetVertexStride()
+	{
+		return (Config >> 8) & 0xFF;
+	}
+
+	#if INCLUDED_FROM_CPP_CODE
+		void SetVertexAndIndexStride(UINT_TYPE Vertex, UINT_TYPE Index)
+		{
+			Config = (Index & 0xFF) | ((Vertex & 0xFF) << 8);
+		}
+	#endif
+};
+
+#define RAY_TRACING_SYSTEM_INDEXBUFFER_REGISTER  0
+#define RAY_TRACING_SYSTEM_VERTEXBUFFER_REGISTER 1
+#define RAY_TRACING_SYSTEM_ROOTCONSTANT_REGISTER 0
+
+#if INCLUDED_FROM_HLSL_CODE
+	#define RT_CONCATENATE2(a, b) a##b
+	#define RT_CONCATENATE(a, b) RT_CONCATENATE2(a, b)
+	#define RT_REGISTER(InType, InIndex, InSpace) register(RT_CONCATENATE(InType, InIndex), RT_CONCATENATE(space, InSpace))
+
+	// Built-in local root parameters that are always bound to all hit shaders
+	ByteAddressBuffer									HitGroupSystemIndexBuffer   : RT_REGISTER(t, RAY_TRACING_SYSTEM_INDEXBUFFER_REGISTER,  RAY_TRACING_REGISTER_SPACE_SYSTEM);
+	ByteAddressBuffer									HitGroupSystemVertexBuffer  : RT_REGISTER(t, RAY_TRACING_SYSTEM_VERTEXBUFFER_REGISTER, RAY_TRACING_REGISTER_SPACE_SYSTEM);
+	ConstantBuffer<FHitGroupSystemRootConstants>		HitGroupSystemRootConstants : RT_REGISTER(b, RAY_TRACING_SYSTEM_ROOTCONSTANT_REGISTER, RAY_TRACING_REGISTER_SPACE_SYSTEM);
+
+	#undef RT_REGISTER
+	#undef RT_CONCATENATE
+	#undef RT_CONCATENATE2
+#endif // INCLUDED_FROM_HLSL_CODE
+
+
+#undef INCLUDED_FROM_CPP_CODE
+#undef INCLUDED_FROM_HLSL_CODE
+#undef UINT_TYPE
+
+#endif // RAYTRACINGBUILTINRESOURCES_USH_INCLUDED // Workarround for UE-66460
diff --git a/Engine/Shaders/Shared/RayTracingDefinitions.h b/Engine/Shaders/Shared/RayTracingDefinitions.h
new file mode 100644
--- /dev/null
+++ b/Engine/Shaders/Shared/RayTracingDefinitions.h
@@ -0,0 +1,26 @@
+// Copyright 1998-2019 Epic Games, Inc. All Rights Reserved.
+
+/*================================================================================================
+	RayTracingDefinitions.ush: used in ray tracing shaders and C++ code to define common constants
+	!!! Changing this file requires recompilation of the engine !!!
+=================================================================================================*/
+
+#pragma once
+
+// Change this to force recompilation of all ray tracing shaders (use https://www.random.org/cgi-bin/randbyte?nbytes=4&format=h)
+// This avoids changing the global ShaderVersion.ush and forcing recompilation of all shaders in the engine (only RT shaders will be affected)
+#define RAY_TRACING_SHADER_VERSION 0x510cacb6
+
+#define RAY_TRACING_REGISTER_SPACE_LOCAL  0 // default register space for hit group (closest hit, any hit, intersection) shader resources
+#define RAY_TRACING_REGISTER_SPACE_GLOBAL 1 // register space for ray generation and miss shaders (#dxr_todo: make global resources also available to hit shaders)
+#define RAY_TRACING_REGISTER_SPACE_SYSTEM 2 // register space for "system" parameters (index buffer, vertex buffer, fetch parameters)
+
+#define RAY_TRACING_MASK_OPAQUE						0x01    // Opaque and alpha tested meshes and particles (e.g. used by reflection, shadow, AO and GI tracing passes)
+#define RAY_TRACING_MASK_TRANSLUCENT				0x02    // Opaque and alpha tested meshes and particles (e.g. used by translucency tracing pass)
+#define RAY_TRACING_MASK_SHADOW						0x08    // Whether the geometry is visible for shadow rays
+#define RAY_TRACING_MASK_ALL						0xFF
+
+#define RAY_TRACING_SHADER_SLOT_MATERIAL	0
+#define RAY_TRACING_SHADER_SLOT_SHADOW		1
+#define RAY_TRACING_NUM_SHADER_SLOTS		2
+
diff --git a/Engine/Shaders/Shared/RaytracingDebugDefinitions.h b/Engine/Shaders/Shared/RaytracingDebugDefinitions.h
new file mode 100644
--- /dev/null
+++ b/Engine/Shaders/Shared/RaytracingDebugDefinitions.h
@@ -0,0 +1,29 @@
+// Copyright 1998-2019 Epic Games, Inc. All Rights Reserved.
+
+/*================================================================================================
+	RayTracingDebugDefinitions.ush: used in ray tracing debug shaders and C++ code to define common constants
+	!!! Changing this file requires recompilation of the engine !!!
+=================================================================================================*/
+
+#pragma once
+
+#define RAY_TRACING_DEBUG_VIZ_RADIANCE					0 
+#define RAY_TRACING_DEBUG_VIZ_WORLD_NORMAL				1 
+#define RAY_TRACING_DEBUG_VIZ_BASE_COLOR				2 
+#define RAY_TRACING_DEBUG_VIZ_DIFFUSE_COLOR				3 
+#define RAY_TRACING_DEBUG_VIZ_SPECULAR_COLOR			4 
+#define RAY_TRACING_DEBUG_VIZ_OPACITY					5 
+#define RAY_TRACING_DEBUG_VIZ_METALLIC					6 
+#define RAY_TRACING_DEBUG_VIZ_SPECULAR					7 
+#define RAY_TRACING_DEBUG_VIZ_ROUGHNESS					8 
+#define RAY_TRACING_DEBUG_VIZ_IOR						9 
+#define RAY_TRACING_DEBUG_VIZ_SHADING_MODEL				10 
+#define RAY_TRACING_DEBUG_VIZ_BLENDING_MODE				11  
+#define RAY_TRACING_DEBUG_VIZ_LIGHTING_CHANNEL_MASK		12  
+#define RAY_TRACING_DEBUG_VIZ_CUSTOM_DATA				13  
+#define RAY_TRACING_DEBUG_VIZ_GBUFFER_AO				14  
+#define RAY_TRACING_DEBUG_VIZ_INDIRECT_IRRADIANCE		15  
+#define RAY_TRACING_DEBUG_VIZ_WORLD_POSITION			16  
+#define RAY_TRACING_DEBUG_VIZ_HITKIND					17  
+#define RAY_TRACING_DEBUG_VIZ_BARYCENTRICS				18  
+
diff --git a/Engine/Source/Runtime/Renderer/Private/RayTracing/RayTracingDynamicGeometryCollection.h b/Engine/Source/Runtime/Renderer/Private/RayTracing/RayTracingDynamicGeometryCollection.h
--- a/Engine/Source/Runtime/Renderer/Private/RayTracing/RayTracingDynamicGeometryCollection.h
+++ b/Engine/Source/Runtime/Renderer/Private/RayTracing/RayTracingDynamicGeometryCollection.h
@@ -7,23 +7,14 @@
 class RENDERER_API FRayTracingDynamicGeometryCollection
 {
 public:
+	FRayTracingDynamicGeometryCollection();
+
 	void AddDynamicMeshBatchForGeometryUpdate(const FScene* Scene, const FSceneView* View, const FPrimitiveSceneProxy* PrimitiveSceneProxy, const FMeshBatch& MeshBatch, FRayTracingGeometry& Geometry, uint32 NumMaxVertices, FRWBuffer& Buffer);
 	void DispatchUpdates(FRHICommandListImmediate& RHICmdList);
 	void Clear();
 
 private:
-	struct FMeshComputeDispatchCommand
-	{
-		FMeshDrawShaderBindings ShaderBindings;
-		class FRayTracingDynamicGeometryConverterCS* MaterialShader;
-
-		uint32 NumMaxVertices;
-		uint32 NumCPUVertices;
-		FRWBuffer* TargetBuffer;
-		FRayTracingGeometry* TargetGeometry;
-	};
-
-	TArray<FMeshComputeDispatchCommand> DispatchCommands;
+	TUniquePtr<TArray<struct FMeshComputeDispatchCommand>> DispatchCommands;
 };
 
 #endif // RHI_RAYTRACING
diff --git a/Engine/Source/Runtime/Renderer/Private/RayTracing/RayTracingShadowsForRectLight.cpp b/Engine/Source/Runtime/Renderer/Private/RayTracing/RayTracingShadowsForRectLight.cpp
--- a/Engine/Source/Runtime/Renderer/Private/RayTracing/RayTracingShadowsForRectLight.cpp
+++ b/Engine/Source/Runtime/Renderer/Private/RayTracing/RayTracingShadowsForRectLight.cpp
@@ -19,47 +19,33 @@
 
 #include "RHI/Public/PipelineStateCache.h"
 
-static int32 GRayTracingRectLight = 0;
-static FAutoConsoleVariableRef CVarRayTracingRectLight(
-	TEXT("r.RayTracing.RectLight"),
-	GRayTracingRectLight,
-	TEXT("0: use traditional rasterized shadow map\n")
-	TEXT("1: use ray gen shader (default)\n")
+static int32 GRayTracingStochasticRectLight = 0;
+static FAutoConsoleVariableRef CVarRayTracingStochasticRectLight(
+	TEXT("r.RayTracing.StochasticRectLight"),
+	GRayTracingStochasticRectLight,
+	TEXT("0: use analytical evaluation (default)\n")
+	TEXT("1: use stochastic evaluation\n")
 );
 
-static int32 GRayTracingRectLightSamplesPerPixel = 1;
-static FAutoConsoleVariableRef CVarRayTracingRecLightSamplesPerPixel(
-	TEXT("r.RayTracing.RectLight.SamplesPerPixel"),
-	GRayTracingRectLightSamplesPerPixel,
-	TEXT("Sets the samples-per-pixel for directional light occlusion (default = 1)")
+static int32 GRayTracingStochasticRectLightSamplesPerPixel = 1;
+static FAutoConsoleVariableRef CVarRayTracingRecLightStochasticSamplesPerPixel(
+	TEXT("r.RayTracing.StochasticRectLight.SamplesPerPixel"),
+	GRayTracingStochasticRectLightSamplesPerPixel,
+	TEXT("Sets the samples-per-pixel for rect light evaluation (default = 1)")
 );
 
-static int32 GRayTracingRectLightIsTextureImportanceSampling = 1;
-static FAutoConsoleVariableRef CVarRayTracingRecLightIsTextureImportanceSampling(
-	TEXT("r.RayTracing.RectLight.IsTextureImportanceSampling"),
-	GRayTracingRectLightIsTextureImportanceSampling,
-	TEXT("Sets the samples-per-pixel for directional light occlusion (default = 1)")
+static int32 GRayTracingStochasticRectLightIsTextureImportanceSampling = 1;
+static FAutoConsoleVariableRef CVarRayTracingStochasticRecLightIsTextureImportanceSampling(
+	TEXT("r.RayTracing.StochasticRectLight.IsTextureImportanceSampling"),
+	GRayTracingStochasticRectLightIsTextureImportanceSampling,
+	TEXT("Enable importance sampling for rect light evaluation (default = 1)")
 );
 
-bool IsRayTracingRectLightSelected()
+bool ShouldRenderRayTracingStochasticRectLight(const FLightSceneInfo& LightSceneInfo)
 {
-	return IsRayTracingEnabled() && GRayTracingRectLight == 1;
-}
-
-bool ShouldRenderRayTracingStaticOrStationaryRectLight(const FLightSceneInfo& LightSceneInfo)
-{
-	return IsRayTracingRectLightSelected()
+	return IsRayTracingEnabled() && GRayTracingStochasticRectLight == 1
 		&& LightSceneInfo.Proxy->CastsRaytracedShadow()
-		&& LightSceneInfo.Proxy->GetLightType() == LightType_Rect
-		&& !LightSceneInfo.Proxy->IsMovable();
-}
-
-bool ShouldRenderRayTracingDynamicRectLight(const FLightSceneInfo& LightSceneInfo)
-{
-	return IsRayTracingRectLightSelected()
-		&& LightSceneInfo.Proxy->CastsRaytracedShadow()
-		&& LightSceneInfo.Proxy->GetLightType() == LightType_Rect
-		&& LightSceneInfo.Proxy->IsMovable();
+		&& LightSceneInfo.Proxy->GetLightType() == LightType_Rect;
 }
 
 BEGIN_GLOBAL_SHADER_PARAMETER_STRUCT(FRectLightData, )
@@ -76,6 +62,8 @@ SHADER_PARAMETER(float, Width)
 SHADER_PARAMETER(float, Height)
 SHADER_PARAMETER(FIntVector, MipTreeDimensions)
 SHADER_PARAMETER(float, MaxNormalBias)
+SHADER_PARAMETER(float, BarnCosAngle)
+SHADER_PARAMETER(float, BarnLength)
 SHADER_PARAMETER_TEXTURE(Texture2D, Texture)
 SHADER_PARAMETER_SAMPLER(SamplerState, TextureSampler)
 // Sampling data
@@ -83,19 +71,16 @@ SHADER_PARAMETER_SRV(Buffer<float>, MipTree)
 END_GLOBAL_SHADER_PARAMETER_STRUCT()
 
 DECLARE_GPU_STAT_NAMED(RayTracingRectLight, TEXT("Ray Tracing RectLight"));
-DECLARE_GPU_STAT_NAMED(RayTracingRectLightOcclusion, TEXT("Ray Tracing RectLight Occlusion"));
 
 IMPLEMENT_GLOBAL_SHADER_PARAMETER_STRUCT(FRectLightData, "RectLight");
 
-template <int CalcDirectLighting, int EncodeVisibility, int TextureImportanceSampling>
-class FRectLightOcclusionRGS : public FGlobalShader
+template <int TextureImportanceSampling>
+class FRectLightRGS : public FGlobalShader
 {
-	DECLARE_SHADER_TYPE(FRectLightOcclusionRGS, Global)
+	DECLARE_SHADER_TYPE(FRectLightRGS, Global)
 
 	static void ModifyCompilationEnvironment(const FGlobalShaderPermutationParameters& Parameters, FShaderCompilerEnvironment& OutEnvironment)
 	{
-		OutEnvironment.SetDefine(TEXT("CALC_DIRECT_LIGHTING"), CalcDirectLighting);
-		OutEnvironment.SetDefine(TEXT("ENCODE_VISIBILITY"), EncodeVisibility);
 		OutEnvironment.SetDefine(TEXT("TEXTURE_IMPORTANCE_SAMPLING"), TextureImportanceSampling);
 	}
 
@@ -105,10 +90,10 @@ public:
 		return ShouldCompileRayTracingShadersForProject(Parameters.Platform);
 	}
 
-	FRectLightOcclusionRGS() {}
-	virtual ~FRectLightOcclusionRGS() {}
+	FRectLightRGS() {}
+	virtual ~FRectLightRGS() {}
 
-	FRectLightOcclusionRGS(const ShaderMetaType::CompiledShaderInitializerType& Initializer)
+	FRectLightRGS(const ShaderMetaType::CompiledShaderInitializerType& Initializer)
 		: FGlobalShader(Initializer)
 	{
 		ViewParameter.Bind(Initializer.ParameterMap, TEXT("View"));
@@ -118,7 +103,7 @@ public:
 		TransmissionProfilesTextureParameter.Bind(Initializer.ParameterMap, TEXT("SSProfilesTexture"));
 		TransmissionProfilesLinearSamplerParameter.Bind(Initializer.ParameterMap, TEXT("TransmissionProfilesLinearSampler"));
 
-		OcclusionMaskUAVParameter.Bind(Initializer.ParameterMap, TEXT("RWOcclusionMaskUAV"));
+		LuminanceUAVParameter.Bind(Initializer.ParameterMap, TEXT("RWLuminanceUAV"));
 		RayDistanceUAVParameter.Bind(Initializer.ParameterMap, TEXT("RWRayDistanceUAV"));
 	}
 
@@ -131,7 +116,7 @@ public:
 		Ar << TLASParameter;
 		Ar << TransmissionProfilesTextureParameter;
 		Ar << TransmissionProfilesLinearSamplerParameter;
-		Ar << OcclusionMaskUAVParameter;
+		Ar << LuminanceUAVParameter;
 		Ar << RayDistanceUAVParameter;
 		return bShaderHasOutdatedParameters;
 	}
@@ -142,7 +127,7 @@ public:
 		FUniformBufferRHIParamRef ViewUniformBuffer,
 		FUniformBufferRHIParamRef SceneTexturesUniformBuffer,
 		FUniformBufferRHIParamRef RectLightUniformBuffer,
-		FUnorderedAccessViewRHIParamRef OcclusionMaskUAV,
+		FUnorderedAccessViewRHIParamRef LuminanceUAV,
 		FUnorderedAccessViewRHIParamRef RayDistanceUAV,
 		uint32 Width, uint32 Height
 	)
@@ -159,7 +144,7 @@ public:
 		GlobalResources.Set(ViewParameter, ViewUniformBuffer);
 		GlobalResources.Set(SceneTexturesParameter, SceneTexturesUniformBuffer);
 		GlobalResources.Set(RectLightParameter, RectLightUniformBuffer);
-		GlobalResources.Set(OcclusionMaskUAVParameter, OcclusionMaskUAV);
+		GlobalResources.Set(LuminanceUAVParameter, LuminanceUAV);
 		GlobalResources.Set(RayDistanceUAVParameter, RayDistanceUAV);
 
 		if (TransmissionProfilesTextureParameter.IsBound())
@@ -178,8 +163,7 @@ public:
 			GlobalResources.SetSampler(TransmissionProfilesLinearSamplerParameter.GetBaseIndex(), TStaticSamplerState<SF_Bilinear, AM_Clamp, AM_Clamp, AM_Clamp>::GetRHI());
 		}
 
-		const uint32 RayGenShaderIndex = 0;
-		RHICmdList.RayTraceDispatch(Pipeline, RayGenShaderIndex, GlobalResources, Width, Height);
+		RHICmdList.RayTraceDispatch(Pipeline, GetRayTracingShader(), RayTracingScene.RayTracingSceneRHI, GlobalResources, Width, Height);
 	}
 
 private:
@@ -194,20 +178,16 @@ private:
 	FShaderResourceParameter TransmissionProfilesLinearSamplerParameter;
 
 	// Output
-	FShaderResourceParameter OcclusionMaskUAVParameter;
+	FShaderResourceParameter LuminanceUAVParameter;
 	FShaderResourceParameter RayDistanceUAVParameter;
 };
 
-#define IMPLEMENT_RECT_LIGHT_OCCLUSION_TYPE(CalcDirectLightingType, EncodeVisbilityType, TextureImportanceSampling) \
-	typedef FRectLightOcclusionRGS<CalcDirectLightingType, EncodeVisbilityType, TextureImportanceSampling> FRectLightOcclusionRGS##CalcDirectLightingType##EncodeVisbilityType##TextureImportanceSampling; \
-	IMPLEMENT_SHADER_TYPE(template<>, FRectLightOcclusionRGS##CalcDirectLightingType##EncodeVisbilityType##TextureImportanceSampling, TEXT("/Engine/Private/RayTracing/RayTracingRectLightOcclusionRGS.usf"), TEXT("RectLightOcclusionRGS"), SF_RayGen);
+#define IMPLEMENT_RECT_LIGHT_TYPE(TextureImportanceSampling) \
+	typedef FRectLightRGS<TextureImportanceSampling> FRectLightRGS##TextureImportanceSampling; \
+	IMPLEMENT_SHADER_TYPE(template<>, FRectLightRGS##TextureImportanceSampling, TEXT("/Engine/Private/RayTracing/RayTracingRectLightRGS.usf"), TEXT("RectLightRGS"), SF_RayGen);
 
-IMPLEMENT_RECT_LIGHT_OCCLUSION_TYPE(0, 0, 0);
-IMPLEMENT_RECT_LIGHT_OCCLUSION_TYPE(0, 0, 1);
-IMPLEMENT_RECT_LIGHT_OCCLUSION_TYPE(0, 1, 0);
-IMPLEMENT_RECT_LIGHT_OCCLUSION_TYPE(0, 1, 1);
-IMPLEMENT_RECT_LIGHT_OCCLUSION_TYPE(1, 0, 0);
-IMPLEMENT_RECT_LIGHT_OCCLUSION_TYPE(1, 0, 1);
+IMPLEMENT_RECT_LIGHT_TYPE(0);
+IMPLEMENT_RECT_LIGHT_TYPE(1);
 
 class FVisualizeRectLightMipTreePS : public FGlobalShader
 {
@@ -282,7 +262,7 @@ void FDeferredShadingSceneRenderer::VisualizeRectLightMipTree(
 	GRenderTargetPool.FindFreeElement(RHICmdList, Desc, RectLightMipTreeRT, TEXT("RectLightMipTreeRT"));
 
 	// Define shaders
-	const auto ShaderMap = GetGlobalShaderMap(FeatureLevel);
+	const auto ShaderMap = GetGlobalShaderMap(View.FeatureLevel);
 	TShaderMapRef<FPostProcessVS> VertexShader(ShaderMap);
 	TShaderMapRef<FVisualizeRectLightMipTreePS> PixelShader(ShaderMap);
 	FTextureRHIParamRef RenderTargets[2] =
@@ -330,71 +310,21 @@ void FDeferredShadingSceneRenderer::VisualizeRectLightMipTree(
 	SceneContext.FinishRenderingSceneColor(RHICmdList);
 }
 
-void FDeferredShadingSceneRenderer::RenderRayTracingRectLight(
-	FRHICommandListImmediate& RHICmdList,
-	const FLightSceneInfo& RectLightSceneInfo,
-	TRefCountPtr<IPooledRenderTarget>& RectLightRT,
-	TRefCountPtr<IPooledRenderTarget>& HitDistanceRT
-)
+void FDeferredShadingSceneRenderer::PrepareRayTracingRectLight(const FViewInfo& View, TArray<FRayTracingShaderRHIParamRef>& OutRayGenShaders)
 {
-	SCOPED_DRAW_EVENT(RHICmdList, RayTracingRectLight);
-	SCOPED_GPU_STAT(RHICmdList, RayTracingRectLight);
+	// Declare all RayGen shaders that require material closest hit shaders to be bound
 
-	FSceneRenderTargets& SceneContext = FSceneRenderTargets::Get(RHICmdList);
-	FPooledRenderTargetDesc Desc = SceneContext.GetSceneColor()->GetDesc();
-	Desc.Format = PF_FloatRGBA;
-	Desc.Flags &= ~(TexCreate_FastVRAM | TexCreate_Transient);
-	GRenderTargetPool.FindFreeElement(RHICmdList, Desc, RectLightRT, TEXT("RayTracingRectLight"));
-	//ClearUAV(RHICmdList, RectLightRT->GetRenderTargetItem(), FLinearColor::Black);
+	TShaderMapRef<FRectLightRGS<0>> Shader0(GetGlobalShaderMap(View.FeatureLevel));
+	TShaderMapRef<FRectLightRGS<1>> Shader1(GetGlobalShaderMap(View.FeatureLevel));
 
-	Desc.Format = PF_R16F;
-	GRenderTargetPool.FindFreeElement(RHICmdList, Desc, HitDistanceRT, TEXT("RayTracingRectLightDistance"));
-	//ClearUAV(RHICmdList, HitDistanceRT->GetRenderTargetItem(), FLinearColor::Black);
-
-	if (RectLightSceneInfo.Proxy->HasSourceTexture())
-	{
-		RenderRayTracingRectLightInternal<1, 0, 1>(RHICmdList, RectLightSceneInfo, RectLightRT, HitDistanceRT);
-	}
-	else
-	{
-		RenderRayTracingRectLightInternal<1, 0, 0>(RHICmdList, RectLightSceneInfo, RectLightRT, HitDistanceRT);
-	}
+	OutRayGenShaders.Add(Shader0->GetRayTracingShader());
+	OutRayGenShaders.Add(Shader1->GetRayTracingShader());
 }
 
-void FDeferredShadingSceneRenderer::RenderRayTracingOcclusionForRectLight(
-	FRHICommandListImmediate& RHICmdList,
-	const FLightSceneInfo& RectLightSceneInfo,
-	TRefCountPtr<IPooledRenderTarget>& ScreenShadowMaskTexture
-)
-{
-	SCOPED_DRAW_EVENT(RHICmdList, RayTracingRectLightOcclusion);
-	SCOPED_GPU_STAT(RHICmdList, RayTracingRectLightOcclusion);
-
-	FSceneRenderTargets& SceneContext = FSceneRenderTargets::Get(RHICmdList);
-	FPooledRenderTargetDesc Desc = SceneContext.GetSceneColor()->GetDesc();
-	Desc.Format = PF_FloatRGBA;
-	Desc.Flags &= ~(TexCreate_FastVRAM | TexCreate_Transient);
-	GRenderTargetPool.FindFreeElement(RHICmdList, Desc, ScreenShadowMaskTexture, TEXT("RayTracingRectLightOcclusion"));
-	//ClearUAV(RHICmdList, ScreenShadowMaskTexture->GetRenderTargetItem(), FLinearColor::Black);
-
-	TRefCountPtr<IPooledRenderTarget> HitDistanceTexture;
-	Desc.Format = PF_R16F;
-	GRenderTargetPool.FindFreeElement(RHICmdList, Desc, HitDistanceTexture, TEXT("RayTracingRectLightDistance"));
-	//ClearUAV(RHICmdList, HitDistanceTexture->GetRenderTargetItem(), FLinearColor::Black);
-
-	if (RectLightSceneInfo.Proxy->HasSourceTexture())
-	{
-		RenderRayTracingRectLightInternal<0, 1, 1>(RHICmdList, RectLightSceneInfo, ScreenShadowMaskTexture, HitDistanceTexture);
-	}
-	else
-	{
-		RenderRayTracingRectLightInternal<0, 1, 0>(RHICmdList, RectLightSceneInfo, ScreenShadowMaskTexture, HitDistanceTexture);
-	}
-}
-
-template <int CalcDirectLighting, int EncodeVisibility, int TextureImportanceSampling>
-void FDeferredShadingSceneRenderer::RenderRayTracingRectLightInternal(
+template <int TextureImportanceSampling>
+void RenderRayTracingRectLightInternal(
 	FRHICommandListImmediate& RHICmdList,
+	const TArray<FViewInfo>& Views,
 	const FLightSceneInfo& RectLightSceneInfo,
 	TRefCountPtr<IPooledRenderTarget>& ScreenShadowMaskTexture,
 	TRefCountPtr<IPooledRenderTarget>& RayDistanceTexture
@@ -402,14 +332,14 @@ void FDeferredShadingSceneRenderer::RenderRayTracingRectLightInternal(
 {
 	check(RectLightSceneInfo.Proxy);
 	check(RectLightSceneInfo.Proxy->IsRectLight());
-	FRectLightSceneProxy* RectLightSceneProxy = (FRectLightSceneProxy*) RectLightSceneInfo.Proxy;
+	FRectLightSceneProxy* RectLightSceneProxy = (FRectLightSceneProxy*)RectLightSceneInfo.Proxy;
 
 	FLightShaderParameters LightShaderParameters;
 	RectLightSceneProxy->GetLightShaderParameters(LightShaderParameters);
 
 	FRectLightData RectLightData;
-	RectLightData.SamplesPerPixel = GRayTracingRectLightSamplesPerPixel;
-	RectLightData.bIsTextureImportanceSampling = GRayTracingRectLightIsTextureImportanceSampling;
+	RectLightData.SamplesPerPixel = GRayTracingStochasticRectLightSamplesPerPixel;
+	RectLightData.bIsTextureImportanceSampling = GRayTracingStochasticRectLightIsTextureImportanceSampling;
 	RectLightData.Position = RectLightSceneInfo.Proxy->GetOrigin();
 	RectLightData.Normal = RectLightSceneInfo.Proxy->GetDirection();
 	const FMatrix& WorldToLight = RectLightSceneInfo.Proxy->GetWorldToLight();
@@ -430,24 +360,26 @@ void FDeferredShadingSceneRenderer::RenderRayTracingRectLightInternal(
 	RectLightData.MipTree = RectLightSceneProxy->RectLightMipTree.SRV;
 	RectLightData.MipTreeDimensions = RectLightSceneProxy->RectLightMipTreeDimensions;
 	RectLightData.MaxNormalBias = GetRaytracingOcclusionMaxNormalBias();
+	RectLightData.BarnCosAngle = FMath::Cos(FMath::DegreesToRadians(RectLightSceneProxy->BarnDoorAngle));
+	RectLightData.BarnLength = RectLightSceneProxy->BarnDoorLength;
 	FUniformBufferRHIRef RectLightUniformBuffer = RHICreateUniformBuffer(&RectLightData, FRectLightData::StaticStructMetadata.GetLayout(), EUniformBufferUsage::UniformBuffer_SingleDraw);
 
 	for (int32 ViewIndex = 0; ViewIndex < Views.Num(); ++ViewIndex)
 	{
-		FViewInfo& View = Views[ViewIndex];
-		FIntPoint ViewSize = View.ViewRect.Size();
+		const FViewInfo& View = Views[ViewIndex];
+		const FIntPoint ViewSize = View.ViewRect.Size();
 
-		TShaderMapRef<FRectLightOcclusionRGS<CalcDirectLighting, EncodeVisibility, TextureImportanceSampling>> RectLightOcclusionRayGenerationShader(GetGlobalShaderMap(FeatureLevel));
+		TShaderMapRef<FRectLightRGS<TextureImportanceSampling>> RectLightRayGenerationShader(GetGlobalShaderMap(View.FeatureLevel));
 
 		FSceneRenderTargets& SceneContext = FSceneRenderTargets::Get(RHICmdList);
 		FSceneTexturesUniformParameters SceneTextures;
-		SetupSceneTextureUniformParameters(SceneContext, FeatureLevel, ESceneTextureSetupMode::All, SceneTextures);
+		SetupSceneTextureUniformParameters(SceneContext, View.FeatureLevel, ESceneTextureSetupMode::All, SceneTextures);
 		FUniformBufferRHIRef SceneTexturesUniformBuffer = RHICreateUniformBuffer(&SceneTextures, FSceneTexturesUniformParameters::StaticStructMetadata.GetLayout(), EUniformBufferUsage::UniformBuffer_SingleDraw);
 
 		// Dispatch
-		RectLightOcclusionRayGenerationShader->Dispatch(
+		RectLightRayGenerationShader->Dispatch(
 			RHICmdList,
-			View.PerViewRayTracingScene,
+			View.RayTracingScene,
 			View.ViewUniformBuffer,
 			SceneTexturesUniformBuffer,
 			RectLightUniformBuffer,
@@ -466,4 +398,41 @@ void FDeferredShadingSceneRenderer::RenderRayTracingRectLightInternal(
 	GVisualizeTexture.SetCheckPoint(RHICmdList, RayDistanceTexture);
 }
 
-#endif
\ No newline at end of file
+#endif // RHI_RAYTRACING
+
+void FDeferredShadingSceneRenderer::RenderRayTracingStochasticRectLight(
+	FRHICommandListImmediate& RHICmdList,
+	const FLightSceneInfo& RectLightSceneInfo,
+	TRefCountPtr<IPooledRenderTarget>& RectLightRT,
+	TRefCountPtr<IPooledRenderTarget>& HitDistanceRT
+)
+#if RHI_RAYTRACING
+{
+	SCOPED_DRAW_EVENT(RHICmdList, RayTracingRectLight);
+	SCOPED_GPU_STAT(RHICmdList, RayTracingRectLight);
+
+	FSceneRenderTargets& SceneContext = FSceneRenderTargets::Get(RHICmdList);
+	FPooledRenderTargetDesc Desc = SceneContext.GetSceneColor()->GetDesc();
+	Desc.Format = PF_FloatRGBA;
+	Desc.Flags &= ~(TexCreate_FastVRAM | TexCreate_Transient);
+	GRenderTargetPool.FindFreeElement(RHICmdList, Desc, RectLightRT, TEXT("RayTracingRectLight"));
+	//ClearUAV(RHICmdList, RectLightRT->GetRenderTargetItem(), FLinearColor::Black);
+
+	Desc.Format = PF_R16F;
+	GRenderTargetPool.FindFreeElement(RHICmdList, Desc, HitDistanceRT, TEXT("RayTracingRectLightDistance"));
+	//ClearUAV(RHICmdList, HitDistanceRT->GetRenderTargetItem(), FLinearColor::Black);
+
+	if (RectLightSceneInfo.Proxy->HasSourceTexture())
+	{
+		RenderRayTracingRectLightInternal<1>(RHICmdList, Views, RectLightSceneInfo, RectLightRT, HitDistanceRT);
+	}
+	else
+	{
+		RenderRayTracingRectLightInternal<0>(RHICmdList, Views, RectLightSceneInfo, RectLightRT, HitDistanceRT);
+	}
+}
+#else
+{
+	unimplemented();
+}
+#endif
