Fix ISR being applied in the shader to non-instanced views.

diff --git a/Engine/Source/Runtime/Renderer/Private/RendererScene.cpp b/Engine/Source/Runtime/Renderer/Private/RendererScene.cpp
--- a/Engine/Source/Runtime/Renderer/Private/RendererScene.cpp
+++ b/Engine/Source/Runtime/Renderer/Private/RendererScene.cpp
@@ -933,6 +933,13 @@ bool FPersistentUniformBuffers::UpdateViewUniformBuffer(const FViewInfo& View)
 			const FViewInfo& InstancedView = static_cast<const FViewInfo&>(View.Family->GetStereoEyeView(StereoPassIndex));
 			InstancedViewUniformBuffer.UpdateUniformBufferImmediate(reinterpret_cast<FInstancedViewUniformShaderParameters&>(*InstancedView.CachedViewUniformShaderParameters));
 		}
+		else
+		{
+			// If we don't render this pass in stereo we simply update the buffer with the same view uniform parameters.
+			// The shader will detect this and it will not attempt to apply ISR while this view is being rendered.
+			// TODO: It's more efficient to change the shader binding to point to ViewUniformBuffer instead of updating InstancedViewUniformBuffer.
+			InstancedViewUniformBuffer.UpdateUniformBufferImmediate(reinterpret_cast<FInstancedViewUniformShaderParameters&>(*View.CachedViewUniformShaderParameters));
+		}
 
 		CachedView = &View;
 		return true;
