Fixed wireframe in editor build by enabling it to go through the MSAA path.

diff --git a/Engine/Source/Runtime/Engine/Private/ShowFlags.cpp b/Engine/Source/Runtime/Engine/Private/ShowFlags.cpp
--- a/Engine/Source/Runtime/Engine/Private/ShowFlags.cpp
+++ b/Engine/Source/Runtime/Engine/Private/ShowFlags.cpp
@@ -424,6 +424,7 @@ void EngineShowFlagOverride(EShowFlagInitMode ShowFlagInitMode, EViewModeIndex V
 		{
 			EngineShowFlags.SetLighting(false);
 			EngineShowFlags.AtmosphericFog = 0;
+			EngineShowFlags.Fog = (ViewModeIndex == VMI_Unlit); // Only unlit uses it 
 		}
 
 		if( ViewModeIndex == VMI_Lit ||
diff --git a/Engine/Source/Runtime/Renderer/Private/PostProcess/PostProcessCompositeEditorPrimitives.cpp b/Engine/Source/Runtime/Renderer/Private/PostProcess/PostProcessCompositeEditorPrimitives.cpp
--- a/Engine/Source/Runtime/Renderer/Private/PostProcess/PostProcessCompositeEditorPrimitives.cpp
+++ b/Engine/Source/Runtime/Renderer/Private/PostProcess/PostProcessCompositeEditorPrimitives.cpp
@@ -499,6 +499,9 @@ void FRCPassPostProcessCompositeEditorPrimitives::Process(FRenderingCompositePas
 	const FPooledRenderTargetDesc* InputDesc = GetInputDesc(ePId_Input0);
 	FIntPoint SrcSize = InputDesc->Extent;
 
+	// Editor primitive is used when rendering VMI_WIREFRAME in order to use MSAA.
+	// This mean we might not actually want to render composite editor primitives here.
+	if (Context.View.Family->EngineShowFlags.CompositeEditorPrimitives)
 	{
 		FRHIRenderPassInfo RPInfo;
 		RPInfo.ColorRenderTargets[0].RenderTarget = EditorColorTarget;
diff --git a/Engine/Source/Runtime/Renderer/Private/SceneRendering.cpp b/Engine/Source/Runtime/Renderer/Private/SceneRendering.cpp
--- a/Engine/Source/Runtime/Renderer/Private/SceneRendering.cpp
+++ b/Engine/Source/Runtime/Renderer/Private/SceneRendering.cpp
@@ -2869,32 +2869,28 @@ void FSceneRenderer::OnStartRender(FRHICommandListImmediate& RHICmdList)
 
 bool FSceneRenderer::ShouldCompositeEditorPrimitives(const FViewInfo& View)
 {
-	// If the show flag is enabled
-	if (!View.Family->EngineShowFlags.CompositeEditorPrimitives)
-	{
-		return false;
-	}
-
 	if (View.Family->EngineShowFlags.VisualizeHDR || View.Family->UseDebugViewPS())
 	{
 		// certain visualize modes get obstructed too much
 		return false;
 	}
 
-	if (GIsEditor && View.Family->EngineShowFlags.Wireframe)
+	if (View.Family->EngineShowFlags.Wireframe)
 	{
-		// In Editor we want wire frame view modes to be in MSAA
+		// We want wireframe view use MSAA if possible.
 		return true;
 	}
-
-	// Any elements that needed compositing were drawn then compositing should be done
-	if (View.ViewMeshElements.Num() 
-		|| View.TopViewMeshElements.Num() 
-		|| View.BatchedViewElements.HasPrimsToDraw() 
-		|| View.TopBatchedViewElements.HasPrimsToDraw() 
-		|| View.NumVisibleDynamicEditorPrimitives > 0)
-	{
-		return true;
+	else if (View.Family->EngineShowFlags.CompositeEditorPrimitives)
+	{
+	    // Any elements that needed compositing were drawn then compositing should be done
+	    if (View.ViewMeshElements.Num() 
+		    || View.TopViewMeshElements.Num() 
+		    || View.BatchedViewElements.HasPrimsToDraw() 
+		    || View.TopBatchedViewElements.HasPrimsToDraw() 
+		    || View.NumVisibleDynamicEditorPrimitives > 0)
+	    {
+		    return true;
+	    }
 	}
 
 	return false;
