UE4 - Fix warning
#rnx

diff --git a/Engine/Source/Runtime/Renderer/Private/RayTracing/RayTracingShadows.cpp b/Engine/Source/Runtime/Renderer/Private/RayTracing/RayTracingShadows.cpp
--- a/Engine/Source/Runtime/Renderer/Private/RayTracing/RayTracingShadows.cpp
+++ b/Engine/Source/Runtime/Renderer/Private/RayTracing/RayTracingShadows.cpp
@@ -210,8 +210,6 @@ void FDeferredShadingSceneRenderer::RenderRayTracingShadows(
 			FRayTracingShaderBindingsWriter GlobalResources;
 			SetShaderParameters(GlobalResources, *RayGenerationShader, *PassParameters);
 
-			FRHIRayTracingPipelineState* Pipeline = nullptr;
-
 			FRayTracingSceneRHIParamRef RayTracingSceneRHI = View.RayTracingScene.RayTracingSceneRHI;
 
 			if (GRayTracingShadowsEnableMaterials)
