UE4 - Fix warning
#rnx

diff --git a/Engine/Source/Runtime/Renderer/Private/CompositionLighting/PostProcessAmbientOcclusion.cpp b/Engine/Source/Runtime/Renderer/Private/CompositionLighting/PostProcessAmbientOcclusion.cpp
--- a/Engine/Source/Runtime/Renderer/Private/CompositionLighting/PostProcessAmbientOcclusion.cpp
+++ b/Engine/Source/Runtime/Renderer/Private/CompositionLighting/PostProcessAmbientOcclusion.cpp
@@ -589,7 +589,7 @@ void FRCPassPostProcessAmbientOcclusionSmooth::Process(FRenderingCompositePassCo
 {
 	SCOPED_GPU_STAT(Context.RHICmdList, SSAOSmooth);
 
-	SetRenderTarget(Context.RHICmdList, nullptr, nullptr);
+	UnbindRenderTargets(Context.RHICmdList);
 	Context.SetViewportAndCallRHI(Context.View.ViewRect);
 
 	FSceneRenderTargets& SceneContext = FSceneRenderTargets::Get(Context.RHICmdList);
