Fixed ray traced reflection material sorting shader dispatch size

Original fix by [REDACTED_PERSON]

diff --git a/Engine/Source/Runtime/Renderer/Private/RayTracing/RayTracingDeferredMaterials.cpp b/Engine/Source/Runtime/Renderer/Private/RayTracing/RayTracingDeferredMaterials.cpp
--- a/Engine/Source/Runtime/Renderer/Private/RayTracing/RayTracingDeferredMaterials.cpp
+++ b/Engine/Source/Runtime/Renderer/Private/RayTracing/RayTracingDeferredMaterials.cpp
@@ -131,13 +131,14 @@ void SortDeferredMaterials(
 	FMaterialSortCS::FPermutationDomain PermutationVector;
 	PermutationVector.Set<FMaterialSortCS::FSortSize>(SortSize - 1);
 
-	// Note that we are presently rounding down and leaving the last N elements unsorted
-	const uint32 DispatchWidth = FMath::DivideAndRoundUp(NumElements, SortSize);
+	// Sort size represents an index into pow2 sizes, not an actual size, so convert to the actual number of elements being sorted
+	const uint32 ElementBlockSize = 256 * (1 << (SortSize - 1));
+	const uint32 DispatchWidth = FMath::DivideAndRoundUp(NumElements, ElementBlockSize);
 
 	TShaderMapRef<FMaterialSortCS> SortShader(View.ShaderMap, PermutationVector);
 	FComputeShaderUtils::AddPass(
 		GraphBuilder,
-		RDG_EVENT_NAME("MaterialSort SortSize=%d NumElements=%d", SortSize, NumElements),
+		RDG_EVENT_NAME("MaterialSort SortSize=%d NumElements=%d", ElementBlockSize, NumElements),
 		*SortShader,
 		PassParameters,
 		FIntVector(DispatchWidth, 1, 1));		
