Updated the builtin instruction count for shader complexity viewmode.

diff --git a/Engine/Source/Runtime/Renderer/Private/DebugViewModeRendering.cpp b/Engine/Source/Runtime/Renderer/Private/DebugViewModeRendering.cpp
--- a/Engine/Source/Runtime/Renderer/Private/DebugViewModeRendering.cpp
+++ b/Engine/Source/Runtime/Renderer/Private/DebugViewModeRendering.cpp
@@ -350,8 +350,7 @@ void FDebugViewModeMeshProcessor::AddMeshBatch(const FMeshBatch& RESTRICT MeshBa
 	// Shadermap can be null while shaders are compiling.
 	if (DebugViewModeInterface->bNeedsInstructionCount && BatchMaterial->GetRenderingThreadShaderMap())
 	{
-		ShaderElementData.NumVSInstructions = BatchMaterial->GetShader<TBasePassVS<TUniformLightMapPolicy<LMP_NO_LIGHTMAP>, false>>(VertexFactoryType)->GetNumInstructions();
-		ShaderElementData.NumPSInstructions = BatchMaterial->GetShader<TBasePassPS<TUniformLightMapPolicy<LMP_NO_LIGHTMAP>, false>>(VertexFactoryType)->GetNumInstructions();
+		UpdateInstructionCount(ShaderElementData, BatchMaterial, VertexFactoryType);
 	}
 
 	ShaderElementData.InitializeMeshMaterialData(ViewIfDynamicMeshCommand, PrimitiveSceneProxy, MeshBatch, StaticMeshId, true);
@@ -373,6 +372,35 @@ void FDebugViewModeMeshProcessor::AddMeshBatch(const FMeshBatch& RESTRICT MeshBa
 		ShaderElementData);
 }
 
+void FDebugViewModeMeshProcessor::UpdateInstructionCount(FDebugViewModeShaderElementData& OutShaderElementData, const FMaterial* InBatchMaterial, FVertexFactoryType* InVertexFactoryType)
+{
+	check(InBatchMaterial && InVertexFactoryType);
+
+	const bool bDeferred = !IsAnyForwardShadingEnabled(GetFeatureLevelShaderPlatform(InBatchMaterial->GetFeatureLevel()));
+	const bool bLit = InBatchMaterial->GetShadingModel() != MSM_Unlit;
+
+	OutShaderElementData.NumVSInstructions = InBatchMaterial->GetShader<TBasePassVS<TUniformLightMapPolicy<LMP_NO_LIGHTMAP>, false>>(InVertexFactoryType)->GetNumInstructions();
+	OutShaderElementData.NumPSInstructions = InBatchMaterial->GetShader<TBasePassPS<TUniformLightMapPolicy<LMP_NO_LIGHTMAP>, false>>(InVertexFactoryType)->GetNumInstructions();
+
+	// Those numbers are taken from a simple material where common inputs are bound to vector parameters (to prevent constant optimizations).
+	if (bDeferred)
+	{
+		OutShaderElementData.NumVSInstructions -= 41;
+		OutShaderElementData.NumPSInstructions -= bLit ? 111 : 33;
+	}
+	else
+	{
+		OutShaderElementData.NumVSInstructions -= 134;
+		OutShaderElementData.NumPSInstructions -= bLit ? 635 : 47;
+	}
+
+	// Ensure a minimum instruction count so that overdraw accumulates.
+	const int32 MinInstructionCount = IsTranslucentBlendMode(InBatchMaterial->GetBlendMode()) ? 20 : 0;
+	
+	OutShaderElementData.NumVSInstructions = FMath::Max<int32>(MinInstructionCount, OutShaderElementData.NumVSInstructions);
+	OutShaderElementData.NumPSInstructions = FMath::Max<int32>(MinInstructionCount, OutShaderElementData.NumPSInstructions);
+}
+
 FMeshPassProcessor* CreateDebugViewModePassProcessor(const FScene* Scene, const FSceneView* InViewIfDynamicMeshCommand, FMeshPassDrawListContext* InDrawListContext)
 {
 	const ERHIFeatureLevel::Type FeatureLevel = Scene ? Scene->GetFeatureLevel() : (InViewIfDynamicMeshCommand ? InViewIfDynamicMeshCommand->GetFeatureLevel() : GMaxRHIFeatureLevel);
diff --git a/Engine/Source/Runtime/Renderer/Private/DebugViewModeRendering.h b/Engine/Source/Runtime/Renderer/Private/DebugViewModeRendering.h
--- a/Engine/Source/Runtime/Renderer/Private/DebugViewModeRendering.h
+++ b/Engine/Source/Runtime/Renderer/Private/DebugViewModeRendering.h
@@ -241,6 +241,8 @@ public:
 
 private:
 
+	void UpdateInstructionCount(FDebugViewModeShaderElementData& OutShaderElementData, const FMaterial* InBatchMaterial, FVertexFactoryType* InVertexFactoryType);
+
 	TUniformBufferRef<FViewUniformShaderParameters> ViewUniformBuffer;
 	FUniformBufferRHIRef PassUniformBuffer;
 	EDebugViewShaderMode DebugViewMode;
diff --git a/Engine/Source/Runtime/Renderer/Private/ShaderComplexityRendering.cpp b/Engine/Source/Runtime/Renderer/Private/ShaderComplexityRendering.cpp
--- a/Engine/Source/Runtime/Renderer/Private/ShaderComplexityRendering.cpp
+++ b/Engine/Source/Runtime/Renderer/Private/ShaderComplexityRendering.cpp
@@ -29,14 +29,6 @@ void TComplexityAccumulatePS<bQuadComplexity>::GetDebugViewModeShaderBindings(
 	FMeshDrawSingleShaderBindings& ShaderBindings
 ) const
 {
-	const float NormalizeMul = 1.0f / GetMaxShaderComplexityCount(Material.GetFeatureLevel());
-	const int32 DeferredBasePassBuiltinInstructions = 83;
-	const int32 ForwardBasePassBuiltinInstructions = 476;
-	// Attempt to remove instructions from code features only present in the forward renderer, so we are showing users their graph cost
-	const int32 LitBaseline = IsAnyForwardShadingEnabled(GetFeatureLevelShaderPlatform(Material.GetFeatureLevel())) ? (ForwardBasePassBuiltinInstructions - DeferredBasePassBuiltinInstructions) : 0;
-	const int32 Baseline = Material.GetShadingModel() == MSM_Unlit ? 0 : LitBaseline;
-	const int32 AdjustedInstructionCount = FMath::Max<int32>(NumPSInstructions - Baseline, 0);
-
 	// normalize the complexity so we can fit it in a low precision scene color which is necessary on some platforms
 	// late value is for overdraw which can be problematic with a low precision float format, at some point the precision isn't there any more and it doesn't accumulate
 	if (DebugViewMode == DVSM_QuadComplexity)
@@ -45,7 +37,8 @@ void TComplexityAccumulatePS<bQuadComplexity>::GetDebugViewModeShaderBindings(
 	}
 	else
 	{
-		ShaderBindings.Add(NormalizedComplexity, FVector4(AdjustedInstructionCount * NormalizeMul, NumVSInstructions * NormalizeMul, 1 / 32.0f));
+		const float NormalizeMul = 1.0f / GetMaxShaderComplexityCount(Material.GetFeatureLevel());
+		ShaderBindings.Add(NormalizedComplexity, FVector4(NumPSInstructions * NormalizeMul, NumVSInstructions * NormalizeMul, 1 / 32.0f));
 	}
 	ShaderBindings.Add(ShowQuadOverdraw, DebugViewMode != DVSM_ShaderComplexity ? 1 : 0);
 }
