Ensure that static mesh visibility mask is set for the shadow LODs
Tessellated landscape can use a different mesh for the main and shadow pass so we were seeing a bug where the shadow visibility mask was unset and so the shadow mesh was not rendered
#rnx

diff --git a/Engine/Source/Runtime/Renderer/Private/ShadowSetup.cpp b/Engine/Source/Runtime/Renderer/Private/ShadowSetup.cpp
--- a/Engine/Source/Runtime/Renderer/Private/ShadowSetup.cpp
+++ b/Engine/Source/Runtime/Renderer/Private/ShadowSetup.cpp
@@ -945,6 +945,12 @@ bool FProjectedShadowInfo::ShouldDrawStaticMeshes(FViewInfo& InCurrentView, bool
 						SubjectMeshCommandBuildRequests.Add(&StaticMesh);
 					}
 
+					if (StaticMeshRelevance.bRequiresPerElementVisibility && !VisibilePrimitiveLODMask->ContainsLOD(StaticMeshRelevance.LODIndex))
+					{
+						// StaticMeshBatchVisibility wasn't updated in main pass so update it now
+						InCurrentView.StaticMeshBatchVisibility[StaticMesh.BatchVisibilityId] = StaticMesh.VertexFactory->GetStaticBatchElementVisibility(InCurrentView, &StaticMesh);
+					}
+
 					bDrawingStaticMeshes = true;
 				}
 			}
