Check whether stereo is enabled before attempting to copy multi-view layers.

diff --git a/Engine/Source/Runtime/Renderer/Private/MobileShadingRenderer.cpp b/Engine/Source/Runtime/Renderer/Private/MobileShadingRenderer.cpp
--- a/Engine/Source/Runtime/Renderer/Private/MobileShadingRenderer.cpp
+++ b/Engine/Source/Runtime/Renderer/Private/MobileShadingRenderer.cpp
@@ -876,7 +876,7 @@ IMPLEMENT_SHADER_TYPE(, FCopyMobileMultiViewSceneColorPS, TEXT("/Engine/Private/
 
 void FMobileSceneRenderer::CopyMobileMultiViewSceneColor(FRHICommandListImmediate& RHICmdList)
 {
-	if (!Views[0].bIsMobileMultiViewEnabled)
+	if (Views.Num() <= 1 || !Views[0].bIsMobileMultiViewEnabled)
 	{
 		return;
 	}
