Fixes copyright noticed on mitchell-netravali downsample.

diff --git a/Engine/Shaders/Private/PostProcessMitchellNetravali.usf b/Engine/Shaders/Private/PostProcessMitchellNetravali.usf
--- a/Engine/Shaders/Private/PostProcessMitchellNetravali.usf
+++ b/Engine/Shaders/Private/PostProcessMitchellNetravali.usf
@@ -1,4 +1,4 @@
-// Copyright 1998-2018 Epic Games, Inc. All Rights Reserved.
+// Copyright 1998-2019 Epic Games, Inc. All Rights Reserved.
 
 /*=============================================================================
 	PostProcessDownsample.usf: PostProcessing down sample.
diff --git a/Engine/Source/Runtime/Renderer/Private/PostProcess/PostProcessMitchellNetravali.cpp b/Engine/Source/Runtime/Renderer/Private/PostProcess/PostProcessMitchellNetravali.cpp
--- a/Engine/Source/Runtime/Renderer/Private/PostProcess/PostProcessMitchellNetravali.cpp
+++ b/Engine/Source/Runtime/Renderer/Private/PostProcess/PostProcessMitchellNetravali.cpp
@@ -1,4 +1,4 @@
-// Copyright 1998-2018 Epic Games, Inc. All Rights Reserved.
+// Copyright 1998-2019 Epic Games, Inc. All Rights Reserved.
 
 /*=============================================================================
 	PostProcessTemporalAA.cpp: Post process MotionBlur implementation.
diff --git a/Engine/Source/Runtime/Renderer/Private/PostProcess/PostProcessMitchellNetravali.h b/Engine/Source/Runtime/Renderer/Private/PostProcess/PostProcessMitchellNetravali.h
--- a/Engine/Source/Runtime/Renderer/Private/PostProcess/PostProcessMitchellNetravali.h
+++ b/Engine/Source/Runtime/Renderer/Private/PostProcess/PostProcessMitchellNetravali.h
@@ -1,4 +1,4 @@
-// Copyright 1998-2018 Epic Games, Inc. All Rights Reserved.
+// Copyright 1998-2019 Epic Games, Inc. All Rights Reserved.
 
 /*=============================================================================
 	PostProcessTemporalAA.h: Post process MotionBlur implementation.
