Fixed: Perf regression in SunTemple on some Mali GPUs
Do not mask stencil writes in mobile base pass, as there atm only single user of stencil - decals

diff --git a/Engine/Source/Runtime/Renderer/Private/MobileBasePass.cpp b/Engine/Source/Runtime/Renderer/Private/MobileBasePass.cpp
--- a/Engine/Source/Runtime/Renderer/Private/MobileBasePass.cpp
+++ b/Engine/Source/Runtime/Renderer/Private/MobileBasePass.cpp
@@ -370,8 +370,10 @@ void MobileBasePass::SetOpaqueRenderState(FMeshPassProcessorRenderState& DrawRen
 				true, CF_DepthNearOrEqual,
 				true, CF_Always, SO_Keep, SO_Keep, SO_Replace,
 				false, CF_Always, SO_Keep, SO_Keep, SO_Keep,
-				0x00, GET_STENCIL_BIT_MASK(RECEIVE_DECAL, 1)>::GetRHI());
-		
+				// decals atm are singe user of stencil in mobile base pass
+				// don't use masking as it has significant performance hit on Mali GPUs (T860MP2)
+				0x00, 0xff /*GET_STENCIL_BIT_MASK(RECEIVE_DECAL, 1)*/ >::GetRHI());
+				
 		DrawRenderState.SetStencilRef(GET_STENCIL_BIT_MASK(RECEIVE_DECAL, StencilValue)); // we hash the stencil group because we only have 6 bits.
 	}
 	else
