Fix SceneTextures not being set in CompositeGlobalIllumination.

diff --git a/Engine/Source/Runtime/Renderer/Private/RayTracing/RayTracingGlobalIllumination.cpp b/Engine/Source/Runtime/Renderer/Private/RayTracing/RayTracingGlobalIllumination.cpp
--- a/Engine/Source/Runtime/Renderer/Private/RayTracing/RayTracingGlobalIllumination.cpp
+++ b/Engine/Source/Runtime/Renderer/Private/RayTracing/RayTracingGlobalIllumination.cpp
@@ -512,12 +512,19 @@ void FDeferredShadingSceneRenderer::CompositeGlobalIllumination(
 	TRefCountPtr<IPooledRenderTarget>& GlobalIlluminationRT
 )
 {
+
 	FSceneRenderTargets& SceneContext = FSceneRenderTargets::Get(RHICmdList);
+
+	FSceneTexturesUniformParameters SceneTextures;
+	SetupSceneTextureUniformParameters(SceneContext, FeatureLevel, ESceneTextureSetupMode::All, SceneTextures);
+
 	FRDGBuilder GraphBuilder(RHICmdList);
+
 	FRayTracingGlobalIlluminationSceneColorCompositePS::FParameters *PassParameters = GraphBuilder.AllocParameters<FRayTracingGlobalIlluminationSceneColorCompositePS::FParameters>();
 	PassParameters->GlobalIlluminationTexture = GraphBuilder.RegisterExternalTexture(GlobalIlluminationRT);
 	PassParameters->GlobalIlluminationSampler = TStaticSamplerState<SF_Point, AM_Clamp, AM_Clamp, AM_Clamp>::GetRHI();
 	PassParameters->RenderTargets[0] = FRenderTargetBinding(GraphBuilder.RegisterExternalTexture(SceneContext.GetSceneColor()), ERenderTargetLoadAction::ENoAction, ERenderTargetStoreAction::ENoAction);
+	PassParameters->SceneTexturesStruct = CreateUniformBufferImmediate(SceneTextures, EUniformBufferUsage::UniformBuffer_SingleDraw);
 
 	GraphBuilder.AddPass(
 		RDG_EVENT_NAME("GlobalIlluminationComposite"),
