UE-71339: Add shader permutation to support optional back-face culling in ray tracing SkyLight pass

diff --git a/Engine/Shaders/Private/RayTracing/RayTracingSkyLightRGS.usf b/Engine/Shaders/Private/RayTracing/RayTracingSkyLightRGS.usf
--- a/Engine/Shaders/Private/RayTracing/RayTracingSkyLightRGS.usf
+++ b/Engine/Shaders/Private/RayTracing/RayTracingSkyLightRGS.usf
@@ -310,6 +310,9 @@ void SkyLightRGS()
 
 		uint RayFlags = 0
 			| RAY_FLAG_FORCE_OPAQUE; // #dxr_todo: generate a special shader permutation for occlusion ray tracing that includes any-hit shaders for masked geometry
+#if !ENABLE_TWO_SIDED_GEOMETRY
+		RayFlags |= RAY_FLAG_CULL_BACK_FACING_TRIANGLES;
+#endif
 		FDefaultPayload Payload = (FDefaultPayload)0;
 
 		TraceRay(
diff --git a/Engine/Source/Runtime/Renderer/Private/RayTracing/RaytracingSkylight.cpp b/Engine/Source/Runtime/Renderer/Private/RayTracing/RaytracingSkylight.cpp
--- a/Engine/Source/Runtime/Renderer/Private/RayTracing/RaytracingSkylight.cpp
+++ b/Engine/Source/Runtime/Renderer/Private/RayTracing/RaytracingSkylight.cpp
@@ -69,6 +69,13 @@ static FAutoConsoleVariableRef CVarRayTracingSkyLightDenoiser(
 	TEXT("Denoising options (default = 1)")
 );
 
+static TAutoConsoleVariable<int32> CVarRayTracingSkyLightEnableTwoSidedGeometry(
+	TEXT("r.RayTracing.SkyLight.EnableTwoSidedGeometry"),
+	0,
+	TEXT("Enables two-sided geometry when tracing shadow rays (default = 0)"),
+	ECVF_RenderThreadSafe
+);
+
 IMPLEMENT_GLOBAL_SHADER_PARAMETER_STRUCT(FSkyLightData, "SkyLight");
 
 void SetupSkyLightParameters(
@@ -132,13 +139,16 @@ void SetupSkyLightParameters(
 DECLARE_GPU_STAT_NAMED(RayTracingSkyLight, TEXT("Ray Tracing SkyLight"));
 DECLARE_GPU_STAT_NAMED(BuildSkyLightMipTree, TEXT("Build SkyLight Mip Tree"));
 
-class FSkylightRGS : public FGlobalShader
+template<uint32 EnableTwoSidedGeometry>
+class TSkyLightRGS : public FGlobalShader
 {
-	DECLARE_SHADER_TYPE(FSkylightRGS, Global)
+	DECLARE_SHADER_TYPE(TSkyLightRGS, Global)
 
 	static void ModifyCompilationEnvironment(const FGlobalShaderPermutationParameters& Parameters, FShaderCompilerEnvironment& OutEnvironment)
 	{
+		FGlobalShader::ModifyCompilationEnvironment(Parameters, OutEnvironment);
 		OutEnvironment.SetDefine(TEXT("USE_TRANSMISSION"), 1);
+		OutEnvironment.SetDefine(TEXT("ENABLE_TWO_SIDED_GEOMETRY"), EnableTwoSidedGeometry);
 	}
 
 public:
@@ -147,10 +157,10 @@ public:
 		return ShouldCompileRayTracingShadersForProject(Parameters.Platform);
 	}
 
-	FSkylightRGS() {}
-	virtual ~FSkylightRGS() {}
+	TSkyLightRGS() {}
+	virtual ~TSkyLightRGS() {}
 
-	FSkylightRGS(const ShaderMetaType::CompiledShaderInitializerType& Initializer)
+	TSkyLightRGS(const ShaderMetaType::CompiledShaderInitializerType& Initializer)
 		: FGlobalShader(Initializer)
 	{
 		ViewParameter.Bind(Initializer.ParameterMap, TEXT("View"));
@@ -239,7 +249,8 @@ private:
 	FShaderResourceParameter RayDistanceUAVParameter;
 };
 
-IMPLEMENT_SHADER_TYPE(, FSkylightRGS, TEXT("/Engine/Private/Raytracing/RaytracingSkylightRGS.usf"), TEXT("SkyLightRGS"), SF_RayGen);
+IMPLEMENT_SHADER_TYPE(template<>, TSkyLightRGS<0>, TEXT("/Engine/Private/Raytracing/RaytracingSkylightRGS.usf"), TEXT("SkyLightRGS"), SF_RayGen);
+IMPLEMENT_SHADER_TYPE(template<>, TSkyLightRGS<1>, TEXT("/Engine/Private/Raytracing/RaytracingSkylightRGS.usf"), TEXT("SkyLightRGS"), SF_RayGen);
 
 void FDeferredShadingSceneRenderer::BuildSkyLightCdfs(FRHICommandListImmediate& RHICmdList, FSkyLightSceneProxy* SkyLight)
 {
@@ -849,21 +860,39 @@ void FDeferredShadingSceneRenderer::RenderRayTracingSkyLight(
 		FViewInfo& View = Views[ViewIndex];
 		FIntPoint ViewSize = View.ViewRect.Size();
 
-		TShaderMapRef<FSkylightRGS> SkyLightRayGenerationShader(GetGlobalShaderMap(FeatureLevel));
 		FSceneTexturesUniformParameters SceneTextures;
 		SetupSceneTextureUniformParameters(SceneContext, FeatureLevel, ESceneTextureSetupMode::All, SceneTextures);
 		FUniformBufferRHIRef SceneTexturesUniformBuffer = RHICreateUniformBuffer(&SceneTextures, FSceneTexturesUniformParameters::StaticStructMetadata.GetLayout(), EUniformBufferUsage::UniformBuffer_SingleDraw);
 
-		SkyLightRayGenerationShader->Dispatch(
-			RHICmdList,
-			View.RayTracingScene,
-			View.ViewUniformBuffer,
-			SceneTexturesUniformBuffer,
-			SkyLightUniformBuffer,
-			SkyLightRT->GetRenderTargetItem().UAV,
-			HitDistanceRT->GetRenderTargetItem().UAV,
-			ViewSize.X, ViewSize.Y
-		);
+		int32 EnableTwoSidedGeometry = CVarRayTracingSkyLightEnableTwoSidedGeometry.GetValueOnRenderThread();
+		if (EnableTwoSidedGeometry)
+		{
+			TShaderMapRef<TSkyLightRGS<1>> SkyLightRayGenerationShader(GetGlobalShaderMap(FeatureLevel));
+			SkyLightRayGenerationShader->Dispatch(
+				RHICmdList,
+				View.RayTracingScene,
+				View.ViewUniformBuffer,
+				SceneTexturesUniformBuffer,
+				SkyLightUniformBuffer,
+				SkyLightRT->GetRenderTargetItem().UAV,
+				HitDistanceRT->GetRenderTargetItem().UAV,
+				ViewSize.X, ViewSize.Y
+			);
+		}
+		else
+		{
+			TShaderMapRef<TSkyLightRGS<0>> SkyLightRayGenerationShader(GetGlobalShaderMap(FeatureLevel));
+			SkyLightRayGenerationShader->Dispatch(
+				RHICmdList,
+				View.RayTracingScene,
+				View.ViewUniformBuffer,
+				SceneTexturesUniformBuffer,
+				SkyLightUniformBuffer,
+				SkyLightRT->GetRenderTargetItem().UAV,
+				HitDistanceRT->GetRenderTargetItem().UAV,
+				ViewSize.X, ViewSize.Y
+			);
+		}
 	}
 
 	// Transition to graphics pipeline
