UE-71372: Dynamic SkyLight rebuild is slow

Reduce compute fences during SkyLight CDF build

diff --git a/Engine/Source/Runtime/Renderer/Private/RayTracing/RaytracingSkylight.cpp b/Engine/Source/Runtime/Renderer/Private/RayTracing/RaytracingSkylight.cpp
--- a/Engine/Source/Runtime/Renderer/Private/RayTracing/RaytracingSkylight.cpp
+++ b/Engine/Source/Runtime/Renderer/Private/RayTracing/RaytracingSkylight.cpp
@@ -254,6 +254,9 @@ IMPLEMENT_SHADER_TYPE(template<>, TSkyLightRGS<1>, TEXT("/Engine/Private/Raytrac
 
 void FDeferredShadingSceneRenderer::BuildSkyLightCdfs(FRHICommandListImmediate& RHICmdList, FSkyLightSceneProxy* SkyLight)
 {
+	SCOPED_DRAW_EVENT(RHICmdList, BuildSkyLightMipTree);
+	SCOPED_GPU_STAT(RHICmdList, BuildSkyLightMipTree);
+
 	BuildSkyLightMipTree(RHICmdList, SkyLight->ProcessedTexture->TextureRHI, SkyLight->SkyLightMipTreePosX, SkyLight->SkyLightMipTreeNegX, SkyLight->SkyLightMipTreePosY, SkyLight->SkyLightMipTreeNegY, SkyLight->SkyLightMipTreePosZ, SkyLight->SkyLightMipTreeNegZ, SkyLight->SkyLightMipDimensions);
 	BuildSkyLightMipTreePdf(RHICmdList, SkyLight->SkyLightMipTreePosX, SkyLight->SkyLightMipTreeNegX, SkyLight->SkyLightMipTreePosY, SkyLight->SkyLightMipTreeNegY, SkyLight->SkyLightMipTreePosZ, SkyLight->SkyLightMipTreeNegZ, SkyLight->SkyLightMipDimensions,
 		SkyLight->SkyLightMipTreePdfPosX, SkyLight->SkyLightMipTreePdfNegX, SkyLight->SkyLightMipTreePdfPosY, SkyLight->SkyLightMipTreePdfNegY, SkyLight->SkyLightMipTreePdfPosZ, SkyLight->SkyLightMipTreePdfNegZ);
@@ -318,13 +321,12 @@ public:
 		FRHICommandList& RHICmdList,
 		EResourceTransitionAccess TransitionAccess,
 		EResourceTransitionPipeline TransitionPipeline,
-		FRWBuffer& MipTree,
-		FComputeFenceRHIParamRef Fence)
+		FRWBuffer& MipTree)
 	{
 		FComputeShaderRHIParamRef ShaderRHI = GetComputeShader();
 
 		MipTreeParameter.UnsetUAV(RHICmdList, ShaderRHI);
-		RHICmdList.TransitionResource(TransitionAccess, TransitionPipeline, MipTree.UAV, Fence);
+		RHICmdList.TransitionResource(TransitionAccess, TransitionPipeline, MipTree.UAV);
 	}
 
 	virtual bool Serialize(FArchive& Ar)
@@ -363,9 +365,6 @@ void FDeferredShadingSceneRenderer::BuildSkyLightMipTree(
 	FIntVector& SkyLightMipTreeDimensions
 )
 {
-	SCOPED_DRAW_EVENT(RHICmdList, BuildSkyLightMipTree);
-	SCOPED_GPU_STAT(RHICmdList, BuildSkyLightMipTree);
-
 	const auto ShaderMap = GetGlobalShaderMap(FeatureLevel);
 	TShaderMapRef<FBuildMipTreeCS> BuildSkyLightMipTreeComputeShader(ShaderMap);
 	RHICmdList.SetComputeShader(BuildSkyLightMipTreeComputeShader->GetComputeShader());
@@ -393,19 +392,27 @@ void FDeferredShadingSceneRenderer::BuildSkyLightMipTree(
 	for (uint32 FaceIndex = 0; FaceIndex < 6; ++FaceIndex)
 	{
 		MipTrees[FaceIndex]->Initialize(sizeof(float), NumElements, PF_R32_FLOAT, BUF_UnorderedAccess | BUF_ShaderResource);
+	}
 
-		// Execute hierarchical build
-		for (uint32 MipLevel = 0; MipLevel <= MipLevelCount; ++MipLevel)
+	// Execute hierarchical build
+	for (uint32 MipLevel = 0; MipLevel <= MipLevelCount; ++MipLevel)
+	{
+		for (uint32 FaceIndex = 0; FaceIndex < 6; ++FaceIndex)
 		{
-			FComputeFenceRHIRef MipLevelFence = RHICmdList.CreateComputeFence(TEXT("SkyLightMipTree Build"));
 			BuildSkyLightMipTreeComputeShader->SetParameters(RHICmdList, SkyLightTexture, SkyLightMipTreeDimensions, FaceIndex, MipLevel, *MipTrees[FaceIndex]);
 			FIntVector MipLevelDimensions = FIntVector(SkyLightMipTreeDimensions.X >> MipLevel, SkyLightMipTreeDimensions.Y >> MipLevel, 1);
 			FIntVector NumGroups = FIntVector::DivideAndRoundUp(MipLevelDimensions, FBuildMipTreeCS::GetGroupSize());
 			DispatchComputeShader(RHICmdList, *BuildSkyLightMipTreeComputeShader, NumGroups.X, NumGroups.Y, 1);
-			BuildSkyLightMipTreeComputeShader->UnsetParameters(RHICmdList, EResourceTransitionAccess::ERWBarrier, EResourceTransitionPipeline::EComputeToCompute, *MipTrees[FaceIndex], MipLevelFence);
+			BuildSkyLightMipTreeComputeShader->UnsetParameters(RHICmdList, EResourceTransitionAccess::ERWBarrier, EResourceTransitionPipeline::EComputeToCompute, *MipTrees[FaceIndex]);
 		}
-		FComputeFenceRHIRef TransitionFence = RHICmdList.CreateComputeFence(TEXT("SkyLightMipTree Transition"));
-		BuildSkyLightMipTreeComputeShader->UnsetParameters(RHICmdList, EResourceTransitionAccess::ERWBarrier, EResourceTransitionPipeline::EComputeToCompute, *MipTrees[FaceIndex], TransitionFence);
+
+		FComputeFenceRHIRef Fence = RHICmdList.CreateComputeFence(TEXT("SkyLightMipTree"));
+		RHICmdList.TransitionResource(EResourceTransitionAccess::ERWBarrier, EResourceTransitionPipeline::EComputeToCompute, MipTrees[0]->UAV);
+		RHICmdList.TransitionResource(EResourceTransitionAccess::ERWBarrier, EResourceTransitionPipeline::EComputeToCompute, MipTrees[1]->UAV);
+		RHICmdList.TransitionResource(EResourceTransitionAccess::ERWBarrier, EResourceTransitionPipeline::EComputeToCompute, MipTrees[2]->UAV);
+		RHICmdList.TransitionResource(EResourceTransitionAccess::ERWBarrier, EResourceTransitionPipeline::EComputeToCompute, MipTrees[3]->UAV);
+		RHICmdList.TransitionResource(EResourceTransitionAccess::ERWBarrier, EResourceTransitionPipeline::EComputeToCompute, MipTrees[4]->UAV);
+		RHICmdList.TransitionResource(EResourceTransitionAccess::ERWBarrier, EResourceTransitionPipeline::EComputeToCompute, MipTrees[5]->UAV, Fence);
 	}
 }
 
@@ -567,13 +574,12 @@ public:
 		FRHICommandList& RHICmdList,
 		EResourceTransitionAccess TransitionAccess,
 		EResourceTransitionPipeline TransitionPipeline,
-		FRWBuffer& MipTreePdf,
-		FComputeFenceRHIParamRef Fence)
+		FRWBuffer& MipTreePdf)
 	{
 		FComputeShaderRHIParamRef ShaderRHI = GetComputeShader();
 
 		MipTreePdfParameter.UnsetUAV(RHICmdList, ShaderRHI);
-		RHICmdList.TransitionResource(TransitionAccess, TransitionPipeline, MipTreePdf.UAV, Fence);
+		RHICmdList.TransitionResource(TransitionAccess, TransitionPipeline, MipTreePdf.UAV);
 	}
 
 	virtual bool Serialize(FArchive& Ar)
@@ -642,18 +648,23 @@ void FDeferredShadingSceneRenderer::BuildSkyLightMipTreePdf(
 		MipTreePdfs[FaceIndex]->Initialize(sizeof(float), NumElements, PF_R32_FLOAT, BUF_UnorderedAccess | BUF_ShaderResource);
 
 		// Execute hierarchical build
-		for (uint32 MipLevel = 0; MipLevel <= MipLevelCount; ++MipLevel)
+		uint32 MipLevel = 0;
 		{
-			FComputeFenceRHIRef MipLevelFence = RHICmdList.CreateComputeFence(TEXT("SkyLightMipTree Build"));
 			BuildSkyLightMipTreePdfComputeShader->SetParameters(RHICmdList, *MipTrees[FaceIndex], SkyLightMipTreeDimensions, MipLevel, *MipTreePdfs[FaceIndex]);
 			FIntVector MipLevelDimensions = FIntVector(SkyLightMipTreeDimensions.X >> MipLevel, SkyLightMipTreeDimensions.Y >> MipLevel, 1);
 			FIntVector NumGroups = FIntVector::DivideAndRoundUp(MipLevelDimensions, FBuildMipTreeCS::GetGroupSize());
 			DispatchComputeShader(RHICmdList, *BuildSkyLightMipTreePdfComputeShader, NumGroups.X, NumGroups.Y, 1);
-			//BuildSkyLightMipTreePdfComputeShader->UnsetParameters(RHICmdList, EResourceTransitionAccess::ERWBarrier, EResourceTransitionPipeline::EComputeToCompute, *MipTreePdfs[FaceIndex], MipLevelFence);
 		}
-		FComputeFenceRHIRef TransitionFence = RHICmdList.CreateComputeFence(TEXT("SkyLightMipTree Transition"));
-		BuildSkyLightMipTreePdfComputeShader->UnsetParameters(RHICmdList, EResourceTransitionAccess::ERWBarrier, EResourceTransitionPipeline::EComputeToCompute, *MipTreePdfs[FaceIndex], TransitionFence);
+		BuildSkyLightMipTreePdfComputeShader->UnsetParameters(RHICmdList, EResourceTransitionAccess::ERWBarrier, EResourceTransitionPipeline::EComputeToCompute, *MipTreePdfs[FaceIndex]);
 	}
+
+	FComputeFenceRHIRef Fence = RHICmdList.CreateComputeFence(TEXT("SkyLightMipTreePdf"));
+	RHICmdList.TransitionResource(EResourceTransitionAccess::ERWBarrier, EResourceTransitionPipeline::EComputeToCompute, MipTreePdfs[0]->UAV);
+	RHICmdList.TransitionResource(EResourceTransitionAccess::ERWBarrier, EResourceTransitionPipeline::EComputeToCompute, MipTreePdfs[1]->UAV);
+	RHICmdList.TransitionResource(EResourceTransitionAccess::ERWBarrier, EResourceTransitionPipeline::EComputeToCompute, MipTreePdfs[2]->UAV);
+	RHICmdList.TransitionResource(EResourceTransitionAccess::ERWBarrier, EResourceTransitionPipeline::EComputeToCompute, MipTreePdfs[3]->UAV);
+	RHICmdList.TransitionResource(EResourceTransitionAccess::ERWBarrier, EResourceTransitionPipeline::EComputeToCompute, MipTreePdfs[4]->UAV);
+	RHICmdList.TransitionResource(EResourceTransitionAccess::ERWBarrier, EResourceTransitionPipeline::EComputeToCompute, MipTreePdfs[5]->UAV, Fence);
 }
 
 class FVisualizeSkyLightMipTreePS : public FGlobalShader
