Cherry-pick 5357715: Switches to multi spp shadow denoiser code path, because found an energy conservation mistake in 1spp code path.

diff --git a/Engine/Source/Runtime/Renderer/Private/LightRendering.cpp b/Engine/Source/Runtime/Renderer/Private/LightRendering.cpp
--- a/Engine/Source/Runtime/Renderer/Private/LightRendering.cpp
+++ b/Engine/Source/Runtime/Renderer/Private/LightRendering.cpp
@@ -903,7 +903,7 @@ void FDeferredShadingSceneRenderer::RenderLights(FRHICommandListImmediate& RHICm
 
 					// Not worth batching and increase memory pressure if the denoiser do not support this ray tracing config.
 					// TODO: add suport for batch with multiple SPP.
-					if (DenoiserRequirements != IScreenSpaceDenoiser::EShadowRequirements::ClosestOccluder)
+					if (DenoiserRequirements != IScreenSpaceDenoiser::EShadowRequirements::PenumbraAndClosestOccluder)
 					{
 						continue;
 					}
diff --git a/Engine/Source/Runtime/Renderer/Private/ScreenSpaceDenoise.cpp b/Engine/Source/Runtime/Renderer/Private/ScreenSpaceDenoise.cpp
--- a/Engine/Source/Runtime/Renderer/Private/ScreenSpaceDenoise.cpp
+++ b/Engine/Source/Runtime/Renderer/Private/ScreenSpaceDenoise.cpp
@@ -18,6 +18,11 @@
 // ---------------------------------------------------- Cvars
 
 
+static TAutoConsoleVariable<int32> CVarShadowUse1SPPCodePath(
+	TEXT("r.Shadow.Denoiser.Use1SPPCodePath"), 0,
+	TEXT("Whether to use the 1spp code path."),
+	ECVF_RenderThreadSafe);
+
 static TAutoConsoleVariable<int32> CVarShadowReconstructionSampleCount(
 	TEXT("r.Shadow.Denoiser.ReconstructionSamples"), 16,
 	TEXT("Maximum number of samples for the reconstruction pass (default = 16)."),
@@ -659,7 +664,9 @@ static void DenoiseSignalAtConstantPixelDensity(
 		return ResourceNames + (int32(Settings.SignalProcessing) * kMaxBufferProcessingCount);
 	};
 
-	const bool bUseMultiInputSPPShaderPath = Settings.MaxInputSPP > 1;
+	const bool bUseMultiInputSPPShaderPath = (
+		Settings.MaxInputSPP > 1 || 
+		(CVarShadowUse1SPPCodePath.GetValueOnRenderThread() == 0 && Settings.SignalProcessing == ESignalProcessing::MonochromaticPenumbra));
 
 	const FIntPoint DenoiseResolution = View.ViewRect.Size();
 	
@@ -1147,7 +1154,7 @@ public:
 		const FLightSceneInfo& LightSceneInfo,
 		const FShadowRayTracingConfig& RayTracingConfig) const override
 	{
-		if (RayTracingConfig.RayCountPerPixel != 1)
+		if (RayTracingConfig.RayCountPerPixel != 1 || CVarShadowUse1SPPCodePath.GetValueOnRenderThread() == 0)
 		{
 			check(SignalSupportMultiSPP(ESignalProcessing::MonochromaticPenumbra));
 			return IScreenSpaceDenoiser::EShadowRequirements::PenumbraAndClosestOccluder;
@@ -1188,7 +1195,7 @@ public:
 			ensure(IsSupportedLightType(ELightComponentType(Parameters.LightSceneInfo->Proxy->GetLightType())));
 
 			Settings.LightSceneInfo[BatchedSignalId] = Parameters.LightSceneInfo;
-			if (Settings.MaxInputSPP == 1)
+			if (Settings.MaxInputSPP == 1 && CVarShadowUse1SPPCodePath.GetValueOnRenderThread() != 0)
 			{
 				// Only have it distance in ClosestOccluder.
 				InputSignal.Textures[BatchedSignalId] = Parameters.InputTextures.ClosestOccluder;
