Force correct fast rendering path in case if full depth prepasss is disabled:
* Moved depth prepass settings into FScene, as those settings are used by cached mesh draw commands (MDCs) and MDCs are cached per FScene
* FScene updates depth prepass settings every frame and recaches MDCs if default base pass depth stencil acccess changes

Before the change cached MDCs always were using DepthRead_StencilWrite, even if full depth prepass was disabled and DepthWrite_StencilWrite was required. This forced ApplyViewOverridesToMeshDrawCommands rendering path, which is a slow and can't correctly set depth stencil state due to lack of mesh data at that point.

diff --git a/Engine/Source/Runtime/Engine/Classes/Engine/RendererSettings.h b/Engine/Source/Runtime/Engine/Classes/Engine/RendererSettings.h
--- a/Engine/Source/Runtime/Engine/Classes/Engine/RendererSettings.h
+++ b/Engine/Source/Runtime/Engine/Classes/Engine/RendererSettings.h
@@ -447,11 +447,6 @@ class ENGINE_API URendererSettings : public UDeveloperSettings
 		ToolTip="Whether to use a depth only pass to initialize Z culling for the base pass."))
 	TEnumAsByte<EEarlyZPass::Type> EarlyZPass;
 
-	UPROPERTY(config, EditAnywhere, Category=Optimizations, meta=(
-		ConsoleVariable="r.EarlyZPassMovable",DisplayName="Movables in early Z-pass",
-		ToolTip="Whether to render movable objects in the early Z pass. Need to reload the level!"))
-	uint32 bEarlyZPassMovable:1;
-
 	UPROPERTY(config, EditAnywhere, Category = Optimizations, meta = (
 		EditCondition = "EarlyZPass == OpaqueAndMasked && bEarlyZPassMovable",
 		ConsoleVariable = "r.EarlyZPassOnlyMaterialMasking", DisplayName = "Mask material only in early Z-pass",
diff --git a/Engine/Source/Runtime/Renderer/Private/BasePassRendering.cpp b/Engine/Source/Runtime/Renderer/Private/BasePassRendering.cpp
--- a/Engine/Source/Runtime/Renderer/Private/BasePassRendering.cpp
+++ b/Engine/Source/Runtime/Renderer/Private/BasePassRendering.cpp
@@ -352,18 +352,6 @@ void SetupBasePassState(FExclusiveDepthStencil::Type BasePassDepthStencilAccess,
 	}
 }
 
-FExclusiveDepthStencil::Type GetDefaultBasePassDepthStencilAccess(EShadingPath ShadingPath)
-{
-	if (ShadingPath == EShadingPath::Mobile)
-	{
-		return FExclusiveDepthStencil::DepthWrite_StencilWrite;
-	}
-	else // deferred
-	{
-		return FExclusiveDepthStencil::DepthRead_StencilWrite;
-	}
-}
-
 /**
  * Get shader templates allowing to redirect between compatible shaders.
  */
@@ -1521,20 +1509,15 @@ FBasePassMeshProcessor::FBasePassMeshProcessor(const FScene* Scene, ERHIFeatureL
 	, PassDrawRenderState(InDrawRenderState)
 	, TranslucencyPassType(InTranslucencyPassType)
 	, bTranslucentBasePass(InTranslucencyPassType != ETranslucencyPass::TPT_MAX)
+	, EarlyZPassMode(Scene ? Scene->EarlyZPassMode : DDM_None)
 {
-	extern void GetEarlyZPassMode(EShaderPlatform ShaderPlatform, EDepthDrawingMode& EarlyZPassMode, bool& bEarlyZPassMovable);
-
-	bool bEarlyZPassMovable;
-	GetEarlyZPassMode(GShaderPlatformForFeatureLevel[InFeatureLevel], EarlyZPassMode, bEarlyZPassMovable);
 }
 
 FMeshPassProcessor* CreateBasePassProcessor(const FScene* Scene, const FSceneView* InViewIfDynamicMeshCommand, FMeshPassDrawListContext* InDrawListContext)
 {
-	FExclusiveDepthStencil::Type DefaultBasePassDepthStencilAccess = GetDefaultBasePassDepthStencilAccess(EShadingPath::Deferred);
-
 	FMeshPassProcessorRenderState PassDrawRenderState(Scene->UniformBuffers.ViewUniformBuffer, Scene->UniformBuffers.OpaqueBasePassUniformBuffer);
 	PassDrawRenderState.SetInstancedViewUniformBuffer(Scene->UniformBuffers.InstancedViewUniformBuffer);
-	SetupBasePassState(DefaultBasePassDepthStencilAccess, false, PassDrawRenderState);
+	SetupBasePassState(Scene->DefaultBasePassDepthStencilAccess, false, PassDrawRenderState);
 
 	return new(FMemStack::Get()) FBasePassMeshProcessor(Scene, Scene->GetFeatureLevel(), InViewIfDynamicMeshCommand, PassDrawRenderState, InDrawListContext);
 }
diff --git a/Engine/Source/Runtime/Renderer/Private/BasePassRendering.h b/Engine/Source/Runtime/Renderer/Private/BasePassRendering.h
--- a/Engine/Source/Runtime/Renderer/Private/BasePassRendering.h
+++ b/Engine/Source/Runtime/Renderer/Private/BasePassRendering.h
@@ -642,5 +642,4 @@ private:
 
 extern void SetDepthStencilStateForBasePass(FMeshPassProcessorRenderState& DrawRenderState, ERHIFeatureLevel::Type FeatureLevel, const FMeshBatch& Mesh, const FPrimitiveSceneProxy* PrimitiveSceneProxy, bool bEnableReceiveDecalOutput, bool bUseDebugViewPS, FDepthStencilStateRHIParamRef LodFadeOverrideDepthStencilState);
 extern void SetupBasePassState(FExclusiveDepthStencil::Type BasePassDepthStencilAccess, const bool bShaderComplexity, FMeshPassProcessorRenderState& DrawRenderState);
-extern FExclusiveDepthStencil::Type GetDefaultBasePassDepthStencilAccess(EShadingPath ShadingPath);
 extern FMeshDrawCommandSortKey CalculateTranslucentMeshStaticSortKey(const FPrimitiveSceneProxy* RESTRICT PrimitiveSceneProxy, uint16 MeshIdInPrimitive);
\ No newline at end of file
diff --git a/Engine/Source/Runtime/Renderer/Private/DeferredShadingRenderer.cpp b/Engine/Source/Runtime/Renderer/Private/DeferredShadingRenderer.cpp
--- a/Engine/Source/Runtime/Renderer/Private/DeferredShadingRenderer.cpp
+++ b/Engine/Source/Runtime/Renderer/Private/DeferredShadingRenderer.cpp
@@ -37,37 +37,6 @@
 #include "RayTracingDefinitions.h"
 #include "RayTracingInstance.h"
 
-TAutoConsoleVariable<int32> CVarEarlyZPass(
-	TEXT("r.EarlyZPass"),
-	3,	
-	TEXT("Whether to use a depth only pass to initialize Z culling for the base pass. Cannot be changed at runtime.\n")
-	TEXT("Note: also look at r.EarlyZPassMovable\n")
-	TEXT("  0: off\n")
-	TEXT("  1: good occluders only: not masked, and large on screen\n")
-	TEXT("  2: all opaque (including masked)\n")
-	TEXT("  x: use built in heuristic (default is 3)"),
-	ECVF_Scalability);
-
-int32 GEarlyZPassMovable = 1;
-
-/** Affects static draw lists so must reload level to propagate. */
-static FAutoConsoleVariableRef CVarEarlyZPassMovable(
-	TEXT("r.EarlyZPassMovable"),
-	GEarlyZPassMovable,
-	TEXT("Whether to render movable objects into the depth only pass. Defaults to on.\n")
-	TEXT("Note: also look at r.EarlyZPass"),
-	ECVF_RenderThreadSafe | ECVF_Scalability
-	);
-
-/** Affects BasePassPixelShader.usf so must relaunch editor to recompile shaders. */
-static TAutoConsoleVariable<int32> CVarEarlyZPassOnlyMaterialMasking(
-	TEXT("r.EarlyZPassOnlyMaterialMasking"),
-	0,
-	TEXT("Whether to compute materials' mask opacity only in early Z pass. Changing this setting requires restarting the editor.\n")
-	TEXT("Note: Needs r.EarlyZPass == 2 && r.EarlyZPassMovable == 1"),
-	ECVF_RenderThreadSafe | ECVF_ReadOnly
-	);
-
 static TAutoConsoleVariable<int32> CVarStencilForLODDither(
 	TEXT("r.StencilForLODDither"),
 	0,
@@ -126,12 +95,6 @@ static FAutoConsoleVariableRef CVarDoPrepareDistanceFieldSceneAfterRHIFlush(
 	ECVF_RenderThreadSafe
 );
 
-static TAutoConsoleVariable<int32> CVarBasePassWriteDepthEvenWithFullPrepass(
-	TEXT("r.BasePassWriteDepthEvenWithFullPrepass"),
-	0,
-	TEXT("0 to allow a readonly base pass, which skips an MSAA depth resolve, and allows masked materials to get EarlyZ (writing to depth while doing clip() disables EarlyZ) (default)\n")
-	TEXT("1 to force depth writes in the base pass.  Useful for debugging when the prepass and base pass don't match what they render."));
-
 static TAutoConsoleVariable<int32> CVarParallelBasePass(
 	TEXT("r.ParallelBasePass"),
 	1,
@@ -206,45 +169,6 @@ DECLARE_GPU_STAT(HZB);
 DECLARE_GPU_STAT_NAMED(AmbientOcclusionDenoiser, TEXT("Ambient Occlusion Denoiser"));
 DECLARE_GPU_STAT_NAMED(Unaccounted, TEXT("[unaccounted]"));
 
-bool ShouldForceFullDepthPass(EShaderPlatform ShaderPlatform)
-{
-	const bool bDBufferAllowed = IsUsingDBuffers(ShaderPlatform);
-
-	static const auto StencilLODDitherCVar = IConsoleManager::Get().FindTConsoleVariableDataInt(TEXT("r.StencilForLODDither"));
-	const bool bStencilLODDither = StencilLODDitherCVar->GetValueOnAnyThread() != 0;
-
-	const bool bEarlyZMaterialMasking = CVarEarlyZPassOnlyMaterialMasking.GetValueOnAnyThread() != 0;
-
-	// Note: ShouldForceFullDepthPass affects which static draw lists meshes go into, so nothing it depends on can change at runtime, unless you do a FGlobalComponentRecreateRenderStateContext to propagate the cvar change
-	return bDBufferAllowed || bStencilLODDither || bEarlyZMaterialMasking || IsForwardShadingEnabled(ShaderPlatform) || UseSelectiveBasePassOutputs();
-}
-
-void GetEarlyZPassMode(EShaderPlatform ShaderPlatform, EDepthDrawingMode& EarlyZPassMode, bool& bEarlyZPassMovable)
-{
-	EarlyZPassMode = DDM_NonMaskedOnly;
-	bEarlyZPassMovable = false;
-
-	// developer override, good for profiling, can be useful as project setting
-	{
-		const int32 CVarValue = CVarEarlyZPass.GetValueOnAnyThread();
-
-		switch(CVarValue)
-		{
-			case 0: EarlyZPassMode = DDM_None; break;
-			case 1: EarlyZPassMode = DDM_NonMaskedOnly; break;
-			case 2: EarlyZPassMode = DDM_AllOccluders; break;
-			case 3: break;	// Note: 3 indicates "default behavior" and does not specify an override
-		}
-	}
-
-	if (ShouldForceFullDepthPass(ShaderPlatform))
-	{
-		// DBuffer decals and stencil LOD dithering force a full prepass
-		EarlyZPassMode = DDM_AllOpaque;
-		bEarlyZPassMovable = true;
-	}
-}
-
 const TCHAR* GetDepthPassReason(bool bDitheredLODTransitionsUseStencil, EShaderPlatform ShaderPlatform)
 {
 	if (IsForwardShadingEnabled(ShaderPlatform))
@@ -273,12 +197,12 @@ const TCHAR* GetDepthPassReason(bool bDitheredLODTransitionsUseStencil, EShaderP
 
 FDeferredShadingSceneRenderer::FDeferredShadingSceneRenderer(const FSceneViewFamily* InViewFamily,FHitProxyConsumer* HitProxyConsumer)
 	: FSceneRenderer(InViewFamily, HitProxyConsumer)
+	, EarlyZPassMode(Scene ? Scene->EarlyZPassMode : DDM_None)
+	, bEarlyZPassMovable(Scene ? Scene->bEarlyZPassMovable : false)
 {
 	static const auto StencilLODDitherCVar = IConsoleManager::Get().FindTConsoleVariableDataInt(TEXT("r.StencilForLODDither"));
 	bDitheredLODTransitionsUseStencil = StencilLODDitherCVar->GetValueOnAnyThread() != 0;
 
-	GetEarlyZPassMode(ShaderPlatform, EarlyZPassMode, bEarlyZPassMovable);
-
 	// Shader complexity requires depth only pass to display masked material cost correctly
 	if (ViewFamily.UseDebugViewPS() && ViewFamily.GetDebugViewShaderMode() != DVSM_OutputMaterialTextureScales)
 	{
@@ -925,7 +849,6 @@ void FDeferredShadingSceneRenderer::Render(FRHICommandListImmediate& RHICmdList)
 
 	// Use readonly depth in the base pass if we have a full depth prepass
 	const bool bAllowReadonlyDepthBasePass = EarlyZPassMode == DDM_AllOpaque
-		&& CVarBasePassWriteDepthEvenWithFullPrepass.GetValueOnRenderThread() == 0
 		&& !ViewFamily.EngineShowFlags.ShaderComplexity
 		&& !ViewFamily.UseDebugViewPS()
 		&& !bIsWireframe
diff --git a/Engine/Source/Runtime/Renderer/Private/DepthRendering.cpp b/Engine/Source/Runtime/Renderer/Private/DepthRendering.cpp
--- a/Engine/Source/Runtime/Renderer/Private/DepthRendering.cpp
+++ b/Engine/Source/Runtime/Renderer/Private/DepthRendering.cpp
@@ -836,15 +836,9 @@ FDepthPassMeshProcessor::FDepthPassMeshProcessor(const FScene* Scene,
 
 FMeshPassProcessor* CreateDepthPassProcessor(const FScene* Scene, const FSceneView* InViewIfDynamicMeshCommand, FMeshPassDrawListContext* InDrawListContext)
 {
-	extern void GetEarlyZPassMode(EShaderPlatform ShaderPlatform, EDepthDrawingMode& EarlyZPassMode, bool& bEarlyZPassMovable);
-
-	EDepthDrawingMode EarlyZPassMode;
-	bool bEarlyZPassMovable;
-	GetEarlyZPassMode(Scene->GetShaderPlatform(), EarlyZPassMode, bEarlyZPassMovable);
-
 	FMeshPassProcessorRenderState DepthPassState;
 	SetupDepthPassState(DepthPassState);
-	return new(FMemStack::Get()) FDepthPassMeshProcessor(Scene, InViewIfDynamicMeshCommand, DepthPassState, true, EarlyZPassMode, bEarlyZPassMovable, InDrawListContext);
+	return new(FMemStack::Get()) FDepthPassMeshProcessor(Scene, InViewIfDynamicMeshCommand, DepthPassState, true, Scene->EarlyZPassMode, Scene->bEarlyZPassMovable, InDrawListContext);
 }
 
 FRegisterPassProcessorCreateFunction RegisterDepthPass(&CreateDepthPassProcessor, EShadingPath::Deferred, EMeshPass::DepthPass, EMeshPassFlags::CachedMeshCommands | EMeshPassFlags::MainView);
\ No newline at end of file
diff --git a/Engine/Source/Runtime/Renderer/Private/MeshDrawCommands.cpp b/Engine/Source/Runtime/Renderer/Private/MeshDrawCommands.cpp
--- a/Engine/Source/Runtime/Renderer/Private/MeshDrawCommands.cpp
+++ b/Engine/Source/Runtime/Renderer/Private/MeshDrawCommands.cpp
@@ -645,6 +645,7 @@ void ApplyViewOverridesToMeshDrawCommands(
 	bool bReverseCulling,
 	bool bRenderSceneTwoSided,
 	FExclusiveDepthStencil::Type BasePassDepthStencilAccess,
+	FExclusiveDepthStencil::Type DefaultBasePassDepthStencilAccess,
 	FMeshCommandOneFrameArray& VisibleMeshDrawCommands,
 	FDynamicMeshDrawCommandStorage& MeshDrawCommandStorage,
 	FMeshCommandOneFrameArray& TempVisibleMeshDrawCommands
@@ -655,8 +656,6 @@ void ApplyViewOverridesToMeshDrawCommands(
 
 	if ((FPassProcessorManager::GetPassFlags(ShadingPath, PassType) & EMeshPassFlags::MainView) != EMeshPassFlags::None)
 	{
-		const FExclusiveDepthStencil::Type DefaultBasePassDepthStencilAccess = GetDefaultBasePassDepthStencilAccess(ShadingPath);
-
 		if (bReverseCulling || bRenderSceneTwoSided || (BasePassDepthStencilAccess != DefaultBasePassDepthStencilAccess && PassType == EMeshPass::BasePass))
 		{
 			for (int32 MeshCommandIndex = 0; MeshCommandIndex < VisibleMeshDrawCommands.Num(); MeshCommandIndex++)
@@ -791,6 +790,7 @@ public:
 					Context.bReverseCulling,
 					Context.bRenderSceneTwoSided,
 					Context.BasePassDepthStencilAccess,
+					Context.DefaultBasePassDepthStencilAccess,
 					Context.MeshDrawCommands,
 					Context.MeshDrawCommandStorage,
 					Context.TempVisibleMeshDrawCommands
@@ -939,6 +939,7 @@ void FParallelMeshDrawCommandPass::DispatchPassSetup(
 	TaskContext.bReverseCulling = View.bReverseCulling;
 	TaskContext.bRenderSceneTwoSided = View.bRenderSceneTwoSided;
 	TaskContext.BasePassDepthStencilAccess = BasePassDepthStencilAccess;
+	TaskContext.DefaultBasePassDepthStencilAccess = Scene->DefaultBasePassDepthStencilAccess;
 	TaskContext.NumDynamicMeshElements = NumDynamicMeshElements;
 	TaskContext.NumDynamicMeshCommandBuildRequestElements = NumDynamicMeshCommandBuildRequestElements;
 
diff --git a/Engine/Source/Runtime/Renderer/Private/MeshDrawCommands.h b/Engine/Source/Runtime/Renderer/Private/MeshDrawCommands.h
--- a/Engine/Source/Runtime/Renderer/Private/MeshDrawCommands.h
+++ b/Engine/Source/Runtime/Renderer/Private/MeshDrawCommands.h
@@ -79,6 +79,7 @@ public:
 	bool bReverseCulling;
 	bool bRenderSceneTwoSided;
 	FExclusiveDepthStencil::Type BasePassDepthStencilAccess;
+	FExclusiveDepthStencil::Type DefaultBasePassDepthStencilAccess;
 
 	// Mesh pass processor.
 	FMeshPassProcessor* MeshPassProcessor;
diff --git a/Engine/Source/Runtime/Renderer/Private/MobileBasePass.cpp b/Engine/Source/Runtime/Renderer/Private/MobileBasePass.cpp
--- a/Engine/Source/Runtime/Renderer/Private/MobileBasePass.cpp
+++ b/Engine/Source/Runtime/Renderer/Private/MobileBasePass.cpp
@@ -678,12 +678,10 @@ void FMobileBasePassMeshProcessor::Process(
 
 FMeshPassProcessor* CreateMobileBasePassProcessor(const FScene* Scene, const FSceneView* InViewIfDynamicMeshCommand, FMeshPassDrawListContext* InDrawListContext)
 {
-	FExclusiveDepthStencil::Type DefaultBasePassDepthStencilAccess = GetDefaultBasePassDepthStencilAccess(EShadingPath::Mobile);
-	
 	FMeshPassProcessorRenderState PassDrawRenderState(Scene->UniformBuffers.ViewUniformBuffer, Scene->UniformBuffers.MobileOpaqueBasePassUniformBuffer);
 	PassDrawRenderState.SetInstancedViewUniformBuffer(Scene->UniformBuffers.InstancedViewUniformBuffer);
 	PassDrawRenderState.SetBlendState(TStaticBlendStateWriteMask<CW_RGBA>::GetRHI());
-	PassDrawRenderState.SetDepthStencilAccess(DefaultBasePassDepthStencilAccess);
+	PassDrawRenderState.SetDepthStencilAccess(Scene->DefaultBasePassDepthStencilAccess);
 	PassDrawRenderState.SetDepthStencilState(TStaticDepthStencilState<true, CF_DepthNearOrEqual>::GetRHI());
 
 	return new(FMemStack::Get()) FMobileBasePassMeshProcessor(Scene, Scene->GetFeatureLevel(), InViewIfDynamicMeshCommand, PassDrawRenderState, InDrawListContext, false);
@@ -691,12 +689,10 @@ FMeshPassProcessor* CreateMobileBasePassProcessor(const FScene* Scene, const FSc
 
 FMeshPassProcessor* CreateMobileBasePassCSMProcessor(const FScene* Scene, const FSceneView* InViewIfDynamicMeshCommand, FMeshPassDrawListContext* InDrawListContext)
 {
-	FExclusiveDepthStencil::Type DefaultBasePassDepthStencilAccess = GetDefaultBasePassDepthStencilAccess(EShadingPath::Mobile);
-
 	FMeshPassProcessorRenderState PassDrawRenderState(Scene->UniformBuffers.ViewUniformBuffer, Scene->UniformBuffers.MobileOpaqueBasePassUniformBuffer);
 	PassDrawRenderState.SetInstancedViewUniformBuffer(Scene->UniformBuffers.InstancedViewUniformBuffer);
 	PassDrawRenderState.SetBlendState(TStaticBlendStateWriteMask<CW_RGBA>::GetRHI());
-	PassDrawRenderState.SetDepthStencilAccess(DefaultBasePassDepthStencilAccess);
+	PassDrawRenderState.SetDepthStencilAccess(Scene->DefaultBasePassDepthStencilAccess);
 	PassDrawRenderState.SetDepthStencilState(TStaticDepthStencilState<true, CF_DepthNearOrEqual>::GetRHI());
 
 	return new(FMemStack::Get()) FMobileBasePassMeshProcessor(Scene, Scene->GetFeatureLevel(), InViewIfDynamicMeshCommand, PassDrawRenderState, InDrawListContext, true);
diff --git a/Engine/Source/Runtime/Renderer/Private/MobileBasePassRendering.cpp b/Engine/Source/Runtime/Renderer/Private/MobileBasePassRendering.cpp
--- a/Engine/Source/Runtime/Renderer/Private/MobileBasePassRendering.cpp
+++ b/Engine/Source/Runtime/Renderer/Private/MobileBasePassRendering.cpp
@@ -248,7 +248,7 @@ void FMobileSceneRenderer::RenderMobileBasePass(FRHICommandListImmediate& RHICmd
 		{
 			FMeshPassProcessorRenderState DrawRenderState(View, Scene->UniformBuffers.MobileOpaqueBasePassUniformBuffer);
 			DrawRenderState.SetBlendState(TStaticBlendStateWriteMask<CW_RGBA>::GetRHI());
-			DrawRenderState.SetDepthStencilAccess(GetDefaultBasePassDepthStencilAccess(EShadingPath::Mobile));
+			DrawRenderState.SetDepthStencilAccess(Scene->DefaultBasePassDepthStencilAccess);
 			DrawRenderState.SetDepthStencilState(TStaticDepthStencilState<true, CF_DepthNearOrEqual>::GetRHI());
 			RenderMobileEditorPrimitives(RHICmdList, View, DrawRenderState);
 		}
diff --git a/Engine/Source/Runtime/Renderer/Private/RendererScene.cpp b/Engine/Source/Runtime/Renderer/Private/RendererScene.cpp
--- a/Engine/Source/Runtime/Renderer/Private/RendererScene.cpp
+++ b/Engine/Source/Runtime/Renderer/Private/RendererScene.cpp
@@ -63,6 +63,31 @@
 // while GWorld is the editor world, for example.
 #define CHECK_FOR_PIE_PRIMITIVE_ATTACH_SCENE_MISMATCH	0
 
+/** Affects BasePassPixelShader.usf so must relaunch editor to recompile shaders. */
+static TAutoConsoleVariable<int32> CVarEarlyZPassOnlyMaterialMasking(
+	TEXT("r.EarlyZPassOnlyMaterialMasking"),
+	0,
+	TEXT("Whether to compute materials' mask opacity only in early Z pass. Changing this setting requires restarting the editor.\n")
+	TEXT("Note: Needs r.EarlyZPass == 2 && r.EarlyZPassMovable == 1"),
+	ECVF_RenderThreadSafe | ECVF_ReadOnly
+);
+
+TAutoConsoleVariable<int32> CVarEarlyZPass(
+	TEXT("r.EarlyZPass"),
+	3,
+	TEXT("Whether to use a depth only pass to initialize Z culling for the base pass. Cannot be changed at runtime.\n")
+	TEXT("Note: also look at r.EarlyZPassMovable\n")
+	TEXT("  0: off\n")
+	TEXT("  1: good occluders only: not masked, and large on screen\n")
+	TEXT("  2: all opaque (including masked)\n")
+	TEXT("  x: use built in heuristic (default is 3)"),
+	ECVF_Scalability);
+
+static TAutoConsoleVariable<int32> CVarBasePassWriteDepthEvenWithFullPrepass(
+	TEXT("r.BasePassWriteDepthEvenWithFullPrepass"),
+	0,
+	TEXT("0 to allow a readonly base pass, which skips an MSAA depth resolve, and allows masked materials to get EarlyZ (writing to depth while doing clip() disables EarlyZ) (default)\n")
+	TEXT("1 to force depth writes in the base pass.  Useful for debugging when the prepass and base pass don't match what they render."));
 
 DECLARE_CYCLE_STAT(TEXT("DeferredShadingSceneRenderer MotionBlurStartFrame"), STAT_FDeferredShadingSceneRenderer_MotionBlurStartFrame, STATGROUP_SceneRendering);
 
@@ -951,10 +976,6 @@ FScene::FScene(UWorld* InWorld, bool bInRequiresHitProxies, bool bInIsEditorScen
 :	FSceneInterface(InFeatureLevel)
 ,	World(InWorld)
 ,	FXSystem(NULL)
-,	bStaticDrawListsMobileHDR(false)
-,	bStaticDrawListsMobileHDR32bpp(false)
-,	StaticDrawListsEarlyZPassMode(0)
-,	StaticDrawShaderPipelines(0)
 ,	bScenesPrimitivesNeedStaticMeshElementUpdate(false)
 ,	bPathTracingNeedsInvalidation(true)
 ,	SkyLight(NULL)
@@ -994,17 +1015,6 @@ FScene::FScene(UWorld* InWorld, bool bInRequiresHitProxies, bool bInIsEditorScen
 
 	FeatureLevel = World->FeatureLevel;
 
-	static auto* MobileHDRCvar = IConsoleManager::Get().FindTConsoleVariableDataInt(TEXT("r.MobileHDR"));
-	static auto* MobileHDR32bppModeCvar = IConsoleManager::Get().FindTConsoleVariableDataInt(TEXT("r.MobileHDR32bppMode"));
-	bStaticDrawListsMobileHDR = MobileHDRCvar->GetValueOnAnyThread() == 1;
-	bStaticDrawListsMobileHDR32bpp = bStaticDrawListsMobileHDR && (GSupportsRenderTargetFormat_PF_FloatRGBA == false || MobileHDR32bppModeCvar->GetValueOnAnyThread() != 0);
-
-	static auto* EarlyZPassCvar = IConsoleManager::Get().FindTConsoleVariableDataInt(TEXT("r.EarlyZPass"));
-	StaticDrawListsEarlyZPassMode = EarlyZPassCvar->GetValueOnAnyThread();
-
-	static auto* ShaderPipelinesCvar = IConsoleManager::Get().FindTConsoleVariableDataInt(TEXT("r.ShaderPipelines"));
-	StaticDrawShaderPipelines = ShaderPipelinesCvar->GetValueOnAnyThread();
-
 	if (World->FXSystem)
 	{
 		FFXSystemInterface::Destroy(World->FXSystem);
@@ -1041,6 +1051,8 @@ FScene::FScene(UWorld* InWorld, bool bInRequiresHitProxies, bool bInIsEditorScen
 	{
 		PersistentUniformBuffers->Initialize();
 	});
+
+	UpdateEarlyZPassMode();
 }
 
 FScene::~FScene()
@@ -3000,21 +3012,62 @@ void FScene::Release()
 		});
 }
 
-void FScene::ConditionalMarkStaticMeshElementsForUpdate()
+bool ShouldForceFullDepthPass(EShaderPlatform ShaderPlatform)
+{
+	const bool bDBufferAllowed = IsUsingDBuffers(ShaderPlatform);
+
+	static const auto StencilLODDitherCVar = IConsoleManager::Get().FindTConsoleVariableDataInt(TEXT("r.StencilForLODDither"));
+	const bool bStencilLODDither = StencilLODDitherCVar->GetValueOnAnyThread() != 0;
+
+	const bool bEarlyZMaterialMasking = CVarEarlyZPassOnlyMaterialMasking.GetValueOnAnyThread() != 0;
+
+	// Note: ShouldForceFullDepthPass affects which static draw lists meshes go into, so nothing it depends on can change at runtime, unless you do a FGlobalComponentRecreateRenderStateContext to propagate the cvar change
+	return bDBufferAllowed || bStencilLODDither || bEarlyZMaterialMasking || IsForwardShadingEnabled(ShaderPlatform) || UseSelectiveBasePassOutputs();
+}
+
+void FScene::UpdateEarlyZPassMode()
 {
-	static auto* EarlyZPassCvar = IConsoleManager::Get().FindTConsoleVariableDataInt(TEXT("r.EarlyZPass"));
-	static auto* ShaderPipelinesCvar = IConsoleManager::Get().FindTConsoleVariableDataInt(TEXT("r.ShaderPipelines"));
+	DefaultBasePassDepthStencilAccess = FExclusiveDepthStencil::DepthWrite_StencilWrite;
+	EarlyZPassMode = DDM_NonMaskedOnly;
+	bEarlyZPassMovable = false;
+
+	if (GetShadingPath(GetFeatureLevel()) == EShadingPath::Deferred)
+	{
+		// developer override, good for profiling, can be useful as project setting
+		{
+			const int32 CVarValue = CVarEarlyZPass.GetValueOnAnyThread();
+
+				switch (CVarValue)
+				{
+				case 0: EarlyZPassMode = DDM_None; break;
+				case 1: EarlyZPassMode = DDM_NonMaskedOnly; break;
+				case 2: EarlyZPassMode = DDM_AllOccluders; break;
+				case 3: break;	// Note: 3 indicates "default behavior" and does not specify an override
+				}
+		}
+
+		const EShaderPlatform ShaderPlatform = GetFeatureLevelShaderPlatform(FeatureLevel);
+		if (ShouldForceFullDepthPass(ShaderPlatform))
+		{
+			// DBuffer decals and stencil LOD dithering force a full prepass
+			EarlyZPassMode = DDM_AllOpaque;
+			bEarlyZPassMovable = true;
+		}
+
+		if (EarlyZPassMode == DDM_AllOpaque
+			&& CVarBasePassWriteDepthEvenWithFullPrepass.GetValueOnAnyThread() == 0)
+		{
+			DefaultBasePassDepthStencilAccess = FExclusiveDepthStencil::DepthRead_StencilWrite;
+		}
+	}
+}
 
-	bool bMobileHDR = IsMobileHDR();
-	bool bMobileHDR32bpp = IsMobileHDR32bpp();
-	int32 DesiredStaticDrawListsEarlyZPassMode = EarlyZPassCvar->GetValueOnRenderThread();
-	int32 DesiredStaticDrawShaderPipelines = ShaderPipelinesCvar->GetValueOnRenderThread();
+void FScene::ConditionalMarkStaticMeshElementsForUpdate()
+{
+	UpdateEarlyZPassMode();
 
 	if (bScenesPrimitivesNeedStaticMeshElementUpdate
-		|| bStaticDrawListsMobileHDR != bMobileHDR
-		|| bStaticDrawListsMobileHDR32bpp != bMobileHDR32bpp
-		|| StaticDrawShaderPipelines != DesiredStaticDrawShaderPipelines
-		|| StaticDrawListsEarlyZPassMode != DesiredStaticDrawListsEarlyZPassMode)
+		|| CachedDefaultBasePassDepthStencilAccess != DefaultBasePassDepthStencilAccess)
 	{
 		// Mark all primitives as needing an update
 		// Note: Only visible primitives will actually update their static mesh elements
@@ -3024,10 +3077,7 @@ void FScene::ConditionalMarkStaticMeshElementsForUpdate()
 		}
 
 		bScenesPrimitivesNeedStaticMeshElementUpdate = false;
-		bStaticDrawListsMobileHDR = bMobileHDR;
-		bStaticDrawListsMobileHDR32bpp = bMobileHDR32bpp;
-		StaticDrawListsEarlyZPassMode = DesiredStaticDrawListsEarlyZPassMode;
-		StaticDrawShaderPipelines = DesiredStaticDrawShaderPipelines;
+		CachedDefaultBasePassDepthStencilAccess = DefaultBasePassDepthStencilAccess;
 	}
 }
 
diff --git a/Engine/Source/Runtime/Renderer/Private/ScenePrivate.h b/Engine/Source/Runtime/Renderer/Private/ScenePrivate.h
--- a/Engine/Source/Runtime/Renderer/Private/ScenePrivate.h
+++ b/Engine/Source/Runtime/Renderer/Private/ScenePrivate.h
@@ -2518,15 +2518,17 @@ public:
 	/** Shadow casting lights that couldn't get a shadowmap channel assigned and therefore won't have valid dynamic shadows, forward renderer only. */
 	TArray<FName> OverflowingDynamicShadowedLights;
 
-	/** The mobile quality level for which static draw lists have been built. */
-	bool bStaticDrawListsMobileHDR;
-	bool bStaticDrawListsMobileHDR32bpp;
+	/** Early Z pass mode. */
+	EDepthDrawingMode EarlyZPassMode;
 
-	/** Whether the early Z pass was force enabled when static draw lists were built. */
-	int32 StaticDrawListsEarlyZPassMode;
+	/** Early Z pass movable. */
+	bool bEarlyZPassMovable;
 
-	/** Whether the ShaderPipelines were enabled when the static draw lists were built. */
-	int32 StaticDrawShaderPipelines;
+	/** Default base pass depth stencil access. */
+	FExclusiveDepthStencil::Type DefaultBasePassDepthStencilAccess;
+
+	/** Default base pass depth stencil access used to cache mesh draw commands. */
+	FExclusiveDepthStencil::Type CachedDefaultBasePassDepthStencilAccess;
 
 	/** True if a change to SkyLight / Lighting has occurred that requires static draw lists to be updated. */
 	bool bScenesPrimitivesNeedStaticMeshElementUpdate;
@@ -2763,6 +2765,8 @@ public:
 
 	int64 GetCachedWholeSceneShadowMapsSize() const;
 
+	void UpdateEarlyZPassMode();
+
 	/**
 	 * Marks static mesh elements as needing an update if necessary.
 	 */
diff --git a/Engine/Source/Runtime/Renderer/Private/SceneRendering.cpp b/Engine/Source/Runtime/Renderer/Private/SceneRendering.cpp
--- a/Engine/Source/Runtime/Renderer/Private/SceneRendering.cpp
+++ b/Engine/Source/Runtime/Renderer/Private/SceneRendering.cpp
@@ -3233,9 +3233,6 @@ FRendererModule::FRendererModule()
 	static auto EarlyZPassVar = IConsoleManager::Get().FindConsoleVariable(TEXT("r.EarlyZPass"));
 	EarlyZPassVar->SetOnChangedCallback(FConsoleVariableDelegate::CreateStatic(&OnChangeCVarRequiringRecreateRenderState));
 
-	static auto EarlyZPassMovableVar = IConsoleManager::Get().FindConsoleVariable(TEXT("r.EarlyZPassMovable"));
-	EarlyZPassMovableVar->SetOnChangedCallback(FConsoleVariableDelegate::CreateStatic(&OnChangeCVarRequiringRecreateRenderState));
-
 #if !(UE_BUILD_SHIPPING || UE_BUILD_TEST)
 	void InitDebugViewModeInterfaces();
 	InitDebugViewModeInterfaces();
