Changes the default of r.Shadow.Denoiser.MaxBatchSize to 4.

diff --git a/Engine/Source/Runtime/Renderer/Private/LightRendering.cpp b/Engine/Source/Runtime/Renderer/Private/LightRendering.cpp
--- a/Engine/Source/Runtime/Renderer/Private/LightRendering.cpp
+++ b/Engine/Source/Runtime/Renderer/Private/LightRendering.cpp
@@ -63,7 +63,7 @@ static TAutoConsoleVariable<int32> CVarShadowUseDenoiser(
 	ECVF_RenderThreadSafe);
 
 static TAutoConsoleVariable<int32> CVarMaxShadowDenoisingBatchSize(
-	TEXT("r.Shadow.Denoiser.MaxBatchSize"), 1,
+	TEXT("r.Shadow.Denoiser.MaxBatchSize"), 4,
 	TEXT("Maximum number of shadow to denoise at the same time."),
 	ECVF_RenderThreadSafe);
 
