Fix upscale issues on GLES devices, output was not flipping viewrect on y axis when rendering to backbuffer.

diff --git a/Engine/Source/Runtime/Renderer/Private/PostProcess/PostProcessUpscale.cpp b/Engine/Source/Runtime/Renderer/Private/PostProcess/PostProcessUpscale.cpp
--- a/Engine/Source/Runtime/Renderer/Private/PostProcess/PostProcessUpscale.cpp
+++ b/Engine/Source/Runtime/Renderer/Private/PostProcess/PostProcessUpscale.cpp
@@ -415,6 +415,15 @@ void FRCPassPostProcessUpscale::Process(FRenderingCompositePassContext& Context)
 		// with distortion (bTessellatedQuad) we need to clear the background
 		FIntRect ExcludeRect = bTessellatedQuad ? FIntRect() : DestRect;
 
+		if (bIsMobileRenderer && RHINeedsToSwitchVerticalAxis(GShaderPlatformForFeatureLevel[Context.GetFeatureLevel()]))
+		{
+			// Move dest rect on y axis as we're rendering directly onto the ogl backbuffer.
+			const int32 OutputTargetHeight = PassOutputs[0].RenderTargetDesc.Extent.Y;
+			int32 NewMaxY = OutputTargetHeight - DestRect.Min.Y;
+			DestRect.Min.Y = OutputTargetHeight - DestRect.Max.Y;
+			DestRect.Max.Y = NewMaxY;
+		}
+
 		Context.SetViewportAndCallRHI(DestRect);
 		if (Context.GetLoadActionForRenderTarget(DestRenderTarget) == ERenderTargetLoadAction::EClear && !bMobilePlatform)
 		{
