Fixed landscape shadows from landscape chunks which aren't visible in the main camera. Extra VisibilePrimitiveLODMask condition was preventing StaticMeshBatchVisibility from being initialized inside whole scene directional shadow pass.

diff --git a/Engine/Source/Runtime/Renderer/Private/ShadowSetup.cpp b/Engine/Source/Runtime/Renderer/Private/ShadowSetup.cpp
--- a/Engine/Source/Runtime/Renderer/Private/ShadowSetup.cpp
+++ b/Engine/Source/Runtime/Renderer/Private/ShadowSetup.cpp
@@ -945,9 +945,8 @@ bool FProjectedShadowInfo::ShouldDrawStaticMeshes(FViewInfo& InCurrentView, bool
 						SubjectMeshCommandBuildRequests.Add(&StaticMesh);
 					}
 
-					if (StaticMeshRelevance.bRequiresPerElementVisibility && !VisibilePrimitiveLODMask->ContainsLOD(StaticMeshRelevance.LODIndex))
+					if (StaticMeshRelevance.bRequiresPerElementVisibility)
 					{
-						// StaticMeshBatchVisibility wasn't updated in main pass so update it now
 						InCurrentView.StaticMeshBatchVisibility[StaticMesh.BatchVisibilityId] = StaticMesh.VertexFactory->GetStaticBatchElementVisibility(InCurrentView, &StaticMesh);
 					}
 
