#UE4 Fix for crash when using nullrhi with r.DoLazyStaticMeshUpdate

diff --git a/Engine/Source/Runtime/Renderer/Private/RendererScene.cpp b/Engine/Source/Runtime/Renderer/Private/RendererScene.cpp
--- a/Engine/Source/Runtime/Renderer/Private/RendererScene.cpp
+++ b/Engine/Source/Runtime/Renderer/Private/RendererScene.cpp
@@ -553,8 +553,8 @@ static TAutoConsoleVariable<int32> CVarDoLazyStaticMeshUpdate(
 
 static void DoLazyStaticMeshUpdateCVarSinkFunction()
 {
-	static bool CachedDoLazyStaticMeshUpdate = CVarDoLazyStaticMeshUpdate.GetValueOnGameThread() && !GIsEditor;
-	bool DoLazyStaticMeshUpdate = CVarDoLazyStaticMeshUpdate.GetValueOnGameThread() && !GIsEditor;
+	static bool CachedDoLazyStaticMeshUpdate = CVarDoLazyStaticMeshUpdate.GetValueOnGameThread() && !GIsEditor && FApp::CanEverRender();
+	bool DoLazyStaticMeshUpdate = CVarDoLazyStaticMeshUpdate.GetValueOnGameThread() && !GIsEditor && FApp::CanEverRender();
 
 	if (DoLazyStaticMeshUpdate != CachedDoLazyStaticMeshUpdate)
 	{
@@ -579,7 +579,7 @@ static FAutoConsoleVariableSink CVarDoLazyStaticMeshUpdateSink(FConsoleCommandDe
 
 void FScene::UpdateDoLazyStaticMeshUpdate(FRHICommandListImmediate& CmdList)
 {
-	bool DoLazyStaticMeshUpdate = CVarDoLazyStaticMeshUpdate.GetValueOnRenderThread() && !GIsEditor;
+	bool DoLazyStaticMeshUpdate = CVarDoLazyStaticMeshUpdate.GetValueOnRenderThread() && !GIsEditor && FApp::CanEverRender();
 
 	for (int32 PrimitiveIndex = 0; PrimitiveIndex < Primitives.Num(); PrimitiveIndex++)
 	{
