Removed dead code

#rnx

diff --git a/Engine/Source/Runtime/Renderer/Private/DistortionRendering.cpp b/Engine/Source/Runtime/Renderer/Private/DistortionRendering.cpp
--- a/Engine/Source/Runtime/Renderer/Private/DistortionRendering.cpp
+++ b/Engine/Source/Runtime/Renderer/Private/DistortionRendering.cpp
@@ -738,7 +738,6 @@ void FDistortionMeshProcessor::AddMeshBatch(const FMeshBatch& RESTRICT MeshBatch
 		const FMaterialRenderProxy& MaterialRenderProxy = FallbackMaterialRenderProxyPtr ? *FallbackMaterialRenderProxyPtr : *MeshBatch.MaterialRenderProxy;
 
 		const EBlendMode BlendMode = Material.GetBlendMode();
-		const EMaterialShadingModel ShadingModel = Material.GetShadingModel();
 		const ERasterizerFillMode MeshFillMode = ComputeMeshFillMode(MeshBatch, Material);
 		const ERasterizerCullMode MeshCullMode = ComputeMeshCullMode(MeshBatch, Material);
 		const bool bIsTranslucent = IsTranslucentBlendMode(BlendMode);
