Fix high resolution screenshot.

The CL invalidates the PT history when the final resolution no longer matches the history resolution.

diff --git a/Engine/Source/Runtime/Renderer/Private/RendererScene.cpp b/Engine/Source/Runtime/Renderer/Private/RendererScene.cpp
--- a/Engine/Source/Runtime/Renderer/Private/RendererScene.cpp
+++ b/Engine/Source/Runtime/Renderer/Private/RendererScene.cpp
@@ -200,7 +200,7 @@ FSceneViewState::FSceneViewState()
 #if RHI_RAYTRACING
 	VarianceMipTreeDimensions = FIntVector(0);
 	VarianceMipTree = new FRWBuffer;
-
+	PathTracingRect = FIntRect(0, 0, 0, 0);
 	TotalRayCount = 0;
 	TotalRayCountBuffer = new FRWBuffer;
 	if (GetFeatureLevel() >= ERHIFeatureLevel::SM5)
diff --git a/Engine/Source/Runtime/Renderer/Private/ScenePrivate.h b/Engine/Source/Runtime/Renderer/Private/ScenePrivate.h
--- a/Engine/Source/Runtime/Renderer/Private/ScenePrivate.h
+++ b/Engine/Source/Runtime/Renderer/Private/ScenePrivate.h
@@ -909,6 +909,7 @@ public:
 	// Reference path tracing cached results
 	TRefCountPtr<IPooledRenderTarget> PathTracingIrradianceRT;
 	TRefCountPtr<IPooledRenderTarget> PathTracingSampleCountRT;
+	FIntRect PathTracingRect;
 	FRWBuffer* VarianceMipTree;
 	FIntVector VarianceMipTreeDimensions;
 
@@ -1312,6 +1313,7 @@ public:
 		PathTracingSampleCountRT.SafeRelease();
 		VarianceMipTreeDimensions = FIntVector(0);
 		TotalRayCount = 0;
+		PathTracingRect = FIntRect(0, 0, 0, 0);
 		IESLightProfileResources.Release();
 #endif 
 	}
diff --git a/Engine/Source/Runtime/Renderer/Private/SceneVisibility.cpp b/Engine/Source/Runtime/Renderer/Private/SceneVisibility.cpp
--- a/Engine/Source/Runtime/Renderer/Private/SceneVisibility.cpp
+++ b/Engine/Source/Runtime/Renderer/Private/SceneVisibility.cpp
@@ -3083,24 +3083,35 @@ void FSceneRenderer::PreVisibilityFrameSetup(FRHICommandListImmediate& RHICmdLis
 			// determine if we are initializing or we should reset the persistent state
 			const float DeltaTime = View.Family->CurrentRealTime - ViewState->LastRenderTime;
 			const bool bFirstFrameOrTimeWasReset = DeltaTime < -0.0001f || ViewState->LastRenderTime < 0.0001f;
-
-			bool bIsLargeCameraMovement = IsLargeCameraMovement(
+			const bool bIsLargeCameraMovement = IsLargeCameraMovement(
 				View,
 				ViewState->PrevFrameViewInfo.ViewMatrices.GetViewMatrix(),
 				ViewState->PrevFrameViewInfo.ViewMatrices.GetViewOrigin(),
 				45.0f, 10000.0f);
-
-			bool bResetCamera = (
-				bFirstFrameOrTimeWasReset ||
-				View.bCameraCut ||
-				IsLargeCameraMovement(View,
-					ViewState->PrevFrameViewInfo.ViewMatrices.GetViewMatrix(),
-					ViewState->PrevFrameViewInfo.ViewMatrices.GetViewOrigin(),
-					45.0f, 10000.0f) ||
+			const bool bResetCamera = (bFirstFrameOrTimeWasReset || View.bCameraCut || bIsLargeCameraMovement);
+			
 #if RHI_RAYTRACING
-				(View.RayTracingRenderMode == ERayTracingRenderMode::PathTracing && (Scene->bPathTracingNeedsInvalidation || IsLargeCameraMovement(View, ViewState->PrevFrameViewInfo.ViewMatrices.GetViewMatrix(), ViewState->PrevFrameViewInfo.ViewMatrices.GetViewOrigin(), 0.1f, 0.1f))) ||
+			const bool bInvalidatePathTracer = View.RayTracingRenderMode == ERayTracingRenderMode::PathTracing &&
+			(
+				bResetCamera ||
+				Scene->bPathTracingNeedsInvalidation ||
+				View.ViewRect != ViewState->PathTracingRect ||
+				IsLargeCameraMovement(
+					View, ViewState->PrevFrameViewInfo.ViewMatrices.GetViewMatrix(), 
+					ViewState->PrevFrameViewInfo.ViewMatrices.GetViewOrigin(), 
+					0.1f, 0.1f)
+			);
+
+			if (bInvalidatePathTracer)
+			{
+				ViewState->PathTracingIrradianceRT.SafeRelease();
+				ViewState->PathTracingSampleCountRT.SafeRelease();
+				ViewState->VarianceMipTreeDimensions = FIntVector(0);
+				ViewState->PathTracingRect = View.ViewRect;
+				ViewState->TotalRayCount = 0;
+				Scene->bPathTracingNeedsInvalidation = false;
+			}
 #endif // RHI_RAYTRACING
-				0);
 
 			if (bResetCamera)
 			{
@@ -3112,14 +3123,6 @@ void FSceneRenderer::PreVisibilityFrameSetup(FRHICommandListImmediate& RHICmdLis
 				//     the uber-postprocessing effect as the last effect in the chain.
 
 				View.bPrevTransformsReset = true;
-
-#if RHI_RAYTRACING
-				ViewState->PathTracingIrradianceRT.SafeRelease();
-				ViewState->PathTracingSampleCountRT.SafeRelease();
-				ViewState->VarianceMipTreeDimensions = FIntVector(0);
-				ViewState->TotalRayCount = 0;
-				Scene->bPathTracingNeedsInvalidation = false;
-#endif // RHI_RAYTRACING
 			}
 			else
 			{
