* Rename variable/function RaytracingOcclusionMaxNormalBias > RayTracingMaxNormalBias.
* Rename command r.Shadow.RayTracing.NormalBias > r.RayTracing.NormalBias

diff --git a/Engine/Source/Runtime/Renderer/Private/PathTracing.cpp b/Engine/Source/Runtime/Renderer/Private/PathTracing.cpp
--- a/Engine/Source/Runtime/Renderer/Private/PathTracing.cpp
+++ b/Engine/Source/Runtime/Renderer/Private/PathTracing.cpp
@@ -239,7 +239,7 @@ public:
 		// Adaptive sampling
 		{
 			FPathTracingAdaptiveSamplingData AdaptiveSamplingData;
-			AdaptiveSamplingData.MaxNormalBias = GetRaytracingOcclusionMaxNormalBias();
+			AdaptiveSamplingData.MaxNormalBias = GetRaytracingMaxNormalBias();
 			AdaptiveSamplingData.UseAdaptiveSampling = CVarPathTracingAdaptiveSampling.GetValueOnRenderThread();
 			AdaptiveSamplingData.RandomSequence = CVarPathTracingRandomSequence.GetValueOnRenderThread();
 			if (VarianceMipTree.NumBytes > 0)
diff --git a/Engine/Source/Runtime/Renderer/Private/RayTracing/RayTracingAmbientOcclusion.cpp b/Engine/Source/Runtime/Renderer/Private/RayTracing/RayTracingAmbientOcclusion.cpp
--- a/Engine/Source/Runtime/Renderer/Private/RayTracing/RayTracingAmbientOcclusion.cpp
+++ b/Engine/Source/Runtime/Renderer/Private/RayTracing/RayTracingAmbientOcclusion.cpp
@@ -174,7 +174,7 @@ void FDeferredShadingSceneRenderer::RenderRayTracingAmbientOcclusion(
 		AmbientOcclusionData.SamplesPerPixel = GRayTracingAmbientOcclusionSamplesPerPixel >= 0 ? GRayTracingAmbientOcclusionSamplesPerPixel : View.FinalPostProcessSettings.RayTracingAOSamplesPerPixel;
 		AmbientOcclusionData.MaxRayDistance = View.FinalPostProcessSettings.AmbientOcclusionRadius;
 		AmbientOcclusionData.Intensity = View.FinalPostProcessSettings.AmbientOcclusionIntensity;
-		AmbientOcclusionData.MaxNormalBias = GetRaytracingOcclusionMaxNormalBias();
+		AmbientOcclusionData.MaxNormalBias = GetRaytracingMaxNormalBias();
 		FUniformBufferRHIRef AmbientOcclusionUniformBuffer = RHICreateUniformBuffer(&AmbientOcclusionData, FAmbientOcclusionData::StaticStructMetadata.GetLayout(), EUniformBufferUsage::UniformBuffer_SingleDraw);
 
 		int32 EnableTwoSidedGeometry = CVarRayTracingAmbientOcclusionEnableTwoSidedGeometry.GetValueOnRenderThread();
diff --git a/Engine/Source/Runtime/Renderer/Private/RayTracing/RayTracingGlobalIllumination.cpp b/Engine/Source/Runtime/Renderer/Private/RayTracing/RayTracingGlobalIllumination.cpp
--- a/Engine/Source/Runtime/Renderer/Private/RayTracing/RayTracingGlobalIllumination.cpp
+++ b/Engine/Source/Runtime/Renderer/Private/RayTracing/RayTracingGlobalIllumination.cpp
@@ -390,7 +390,7 @@ void FDeferredShadingSceneRenderer::RenderRayTracingGlobalIllumination(
 		FGlobalIlluminationRGS::FParameters* PassParameters = GraphBuilder.AllocParameters<FGlobalIlluminationRGS::FParameters>();
 		PassParameters->SamplesPerPixel = RayTracingGISamplesPerPixel;
 		PassParameters->MaxBounces = GRayTracingGlobalIlluminationMaxBounces > -1? GRayTracingGlobalIlluminationMaxBounces : View.FinalPostProcessSettings.RayTracingGIMaxBounces;
-		PassParameters->MaxNormalBias = GetRaytracingOcclusionMaxNormalBias();
+		PassParameters->MaxNormalBias = GetRaytracingMaxNormalBias();
 		float MaxRayDistanceForGI = GRayTracingGlobalIlluminationMaxRayDistance;
 		if (MaxRayDistanceForGI == -1.0)
 		{
diff --git a/Engine/Source/Runtime/Renderer/Private/RayTracing/RayTracingRectLight.cpp b/Engine/Source/Runtime/Renderer/Private/RayTracing/RayTracingRectLight.cpp
--- a/Engine/Source/Runtime/Renderer/Private/RayTracing/RayTracingRectLight.cpp
+++ b/Engine/Source/Runtime/Renderer/Private/RayTracing/RayTracingRectLight.cpp
@@ -361,7 +361,7 @@ void RenderRayTracingRectLightInternal(
 	RectLightData.TextureSampler = RHICreateSamplerState(FSamplerStateInitializerRHI(SF_Bilinear, AM_Border, AM_Border, AM_Border));
 	RectLightData.MipTree = RectLightSceneProxy->RectLightMipTree.SRV;
 	RectLightData.MipTreeDimensions = RectLightSceneProxy->RectLightMipTreeDimensions;
-	RectLightData.MaxNormalBias = GetRaytracingOcclusionMaxNormalBias();
+	RectLightData.MaxNormalBias = GetRaytracingMaxNormalBias();
 	RectLightData.BarnCosAngle = FMath::Cos(FMath::DegreesToRadians(RectLightSceneProxy->BarnDoorAngle));
 	RectLightData.BarnLength = RectLightSceneProxy->BarnDoorLength;
 	FUniformBufferRHIRef RectLightUniformBuffer = RHICreateUniformBuffer(&RectLightData, FRectLightData::StaticStructMetadata.GetLayout(), EUniformBufferUsage::UniformBuffer_SingleDraw);
diff --git a/Engine/Source/Runtime/Renderer/Private/RayTracing/RayTracingReflections.cpp b/Engine/Source/Runtime/Renderer/Private/RayTracing/RayTracingReflections.cpp
--- a/Engine/Source/Runtime/Renderer/Private/RayTracing/RayTracingReflections.cpp
+++ b/Engine/Source/Runtime/Renderer/Private/RayTracing/RayTracingReflections.cpp
@@ -267,7 +267,7 @@ void FDeferredShadingSceneRenderer::RenderRayTracingReflections(
 	CommonParameters.ReflectionMinRayDistance = FMath::Min(GRayTracingReflectionsMinRayDistance, GRayTracingReflectionsMaxRayDistance);
 	CommonParameters.ReflectionMaxRayDistance = GRayTracingReflectionsMaxRayDistance;
 	CommonParameters.ReflectionMaxRoughness = FMath::Clamp(GRayTracingReflectionsMaxRoughness >= 0 ? GRayTracingReflectionsMaxRoughness : View.FinalPostProcessSettings.RayTracingReflectionsMaxRoughness, 0.01f, 1.0f);
-	CommonParameters.ReflectionMaxNormalBias = GetRaytracingOcclusionMaxNormalBias();
+	CommonParameters.ReflectionMaxNormalBias = GetRaytracingMaxNormalBias();
 	CommonParameters.RayTracingResolution = RayTracingResolution;
 	CommonParameters.TileAlignedResolution = TileAlignedResolution;
 
diff --git a/Engine/Source/Runtime/Renderer/Private/RayTracing/RayTracingShadows.cpp b/Engine/Source/Runtime/Renderer/Private/RayTracing/RayTracingShadows.cpp
--- a/Engine/Source/Runtime/Renderer/Private/RayTracing/RayTracingShadows.cpp
+++ b/Engine/Source/Runtime/Renderer/Private/RayTracing/RayTracingShadows.cpp
@@ -20,11 +20,11 @@
 #include "Containers/DynamicRHIResourceArray.h"
 #include "SceneViewFamilyBlackboard.h"
 
-static float GRayTracingOcclusionNormalBias = 0.1f;
-static FAutoConsoleVariableRef CVarRayTracingOcclusionNormalBias(
-	TEXT("r.Shadow.RayTracing.NormalBias"),
-	GRayTracingOcclusionNormalBias,
-	TEXT("Sets the max. normal bias used for offseting the ray start position along the normal for light occlusion (default = 0.1)")
+static float GRayTracingMaxNormalBias = 0.1f;
+static FAutoConsoleVariableRef CVarRayTracingNormalBias(
+	TEXT("r.RayTracing.NormalBias"),
+	GRayTracingMaxNormalBias,
+	TEXT("Sets the max. normal bias used for offseting the ray start position along the normal (default = 0.1, i.e., 1mm)")
 );
 
 static int32 GRayTracingShadowsEnableMaterials = 1;
@@ -75,9 +75,9 @@ class FOcclusionRGS : public FGlobalShader
 
 IMPLEMENT_GLOBAL_SHADER(FOcclusionRGS, "/Engine/Private/RayTracing/RayTracingOcclusionRGS.usf", "OcclusionRGS", SF_RayGen);
 
-float GetRaytracingOcclusionMaxNormalBias()
+float GetRaytracingMaxNormalBias()
 {
-	return FMath::Max(0.01f, GRayTracingOcclusionNormalBias);
+	return FMath::Max(0.01f, GRayTracingMaxNormalBias);
 }
 
 void FDeferredShadingSceneRenderer::PrepareRayTracingShadows(const FViewInfo& View, TArray<FRayTracingShaderRHIParamRef>& OutRayGenShaders)
@@ -166,7 +166,7 @@ void FDeferredShadingSceneRenderer::RenderRayTracingShadows(
 		PassParameters->RWOcclusionMaskUAV = GraphBuilder.CreateUAV(FRDGTextureUAVDesc(ScreenShadowMaskTexture));
 		PassParameters->RWRayDistanceUAV = GraphBuilder.CreateUAV(FRDGTextureUAVDesc(RayDistanceTexture));
 		PassParameters->SamplesPerPixel = RayTracingConfig.RayCountPerPixel;
-		PassParameters->NormalBias = GetRaytracingOcclusionMaxNormalBias();
+		PassParameters->NormalBias = GetRaytracingMaxNormalBias();
 		PassParameters->LightingChannelMask = LightSceneProxy->GetLightingChannelMask();
 		LightSceneProxy->GetLightShaderParameters(PassParameters->Light);
 		PassParameters->TLAS = View.RayTracingScene.RayTracingSceneRHI->GetShaderResourceView();
diff --git a/Engine/Source/Runtime/Renderer/Private/RayTracing/RayTracingTranslucency.cpp b/Engine/Source/Runtime/Renderer/Private/RayTracing/RayTracingTranslucency.cpp
--- a/Engine/Source/Runtime/Renderer/Private/RayTracing/RayTracingTranslucency.cpp
+++ b/Engine/Source/Runtime/Renderer/Private/RayTracing/RayTracingTranslucency.cpp
@@ -365,7 +365,7 @@ void FDeferredShadingSceneRenderer::RenderRayTracingTranslucencyView(
 	PassParameters->TranslucencyMaxRayDistance = GRayTracingTranslucencyMaxRayDistance;
 	PassParameters->TranslucencyMaxRoughness = FMath::Clamp(GRayTracingTranslucencyMaxRoughness >= 0 ? GRayTracingTranslucencyMaxRoughness : View.FinalPostProcessSettings.RayTracingTranslucencyMaxRoughness, 0.01f, 1.0f);
 	PassParameters->TranslucencyRefraction = GRayTracingTranslucencyRefraction >= 0 ? GRayTracingTranslucencyRefraction : View.FinalPostProcessSettings.RayTracingTranslucencyRefraction;
-	PassParameters->MaxNormalBias = GetRaytracingOcclusionMaxNormalBias();
+	PassParameters->MaxNormalBias = GetRaytracingMaxNormalBias();
 
 	PassParameters->TLAS = View.RayTracingScene.RayTracingSceneRHI->GetShaderResourceView();
 	PassParameters->ViewUniformBuffer = View.ViewUniformBuffer;
diff --git a/Engine/Source/Runtime/Renderer/Private/RayTracing/RaytracingOptions.h b/Engine/Source/Runtime/Renderer/Private/RayTracing/RaytracingOptions.h
--- a/Engine/Source/Runtime/Renderer/Private/RayTracing/RaytracingOptions.h
+++ b/Engine/Source/Runtime/Renderer/Private/RayTracing/RaytracingOptions.h
@@ -16,7 +16,7 @@ extern bool ShouldRenderRayTracingAmbientOcclusion();
 extern bool ShouldRenderRayTracingGlobalIllumination(const TArray<FViewInfo>& Views);
 extern bool ShouldRenderRayTracingStochasticRectLight(const FLightSceneInfo& LightSceneInfo);
 
-extern float GetRaytracingOcclusionMaxNormalBias();
+extern float GetRaytracingMaxNormalBias();
 
 #else
 
diff --git a/Engine/Source/Runtime/Renderer/Private/RayTracing/RaytracingSkylight.cpp b/Engine/Source/Runtime/Renderer/Private/RayTracing/RaytracingSkylight.cpp
--- a/Engine/Source/Runtime/Renderer/Private/RayTracing/RaytracingSkylight.cpp
+++ b/Engine/Source/Runtime/Renderer/Private/RayTracing/RaytracingSkylight.cpp
@@ -87,7 +87,7 @@ void SetupSkyLightParameters(
 	SkyLightData->SamplesPerPixel = -1;
 	SkyLightData->SamplingStopLevel = 0;
 	SkyLightData->MaxRayDistance = 1.0e27;
-	SkyLightData->MaxNormalBias = GetRaytracingOcclusionMaxNormalBias();
+	SkyLightData->MaxNormalBias = GetRaytracingMaxNormalBias();
 
 	if (Scene.SkyLight && Scene.SkyLight->ProcessedTexture)
 	{
