Change rect light texture slot allocation for raytraced reflection/translucency to increase the number possible textures.

The new allocation deduplicates similar textures.

diff --git a/Engine/Source/Runtime/Renderer/Private/RayTracing/RayTracingLighting.cpp b/Engine/Source/Runtime/Renderer/Private/RayTracing/RayTracingLighting.cpp
--- a/Engine/Source/Runtime/Renderer/Private/RayTracing/RayTracingLighting.cpp
+++ b/Engine/Source/Runtime/Renderer/Private/RayTracing/RayTracingLighting.cpp
@@ -14,6 +14,7 @@ void SetupRaytracingLightData(
 	FRaytracingLightData* LightData)
 {
 	TMap<UTextureLightProfile*, int32> IESLightProfilesMap;
+	TMap<FTextureRHIParamRef, uint32> RectTextureMap;
 
 	LightData->Count = 0;
 	LightData->LTCMatTexture = GSystemTextures.LTCMat->GetRenderTargetItem().ShaderResourceTexture;
@@ -30,8 +31,8 @@ void SetupRaytracingLightData(
 	LightData->RectLightTexture5 = DymmyWhiteTexture;
 	LightData->RectLightTexture6 = DymmyWhiteTexture;
 	LightData->RectLightTexture7 = DymmyWhiteTexture;
-	const uint32 RectLightTextureSlotCount = 8;
-	uint32 CurrentRectLightIndex = 0;
+	static constexpr uint32 MaxRectLightTextureSlos = 8;
+	static constexpr uint32 InvalidTextureIndex = 99; // #dxr_todo: share this definition with ray tracing shaders
 
 	for (auto Light : Lights)
 	{
@@ -78,15 +79,34 @@ void SetupRaytracingLightData(
 		LightData->SourceLength[LightData->Count] = LightParameters.SourceLength;
 		LightData->SoftSourceRadius[LightData->Count] = LightParameters.SoftSourceRadius;
 		LightData->LightProfileIndex[LightData->Count] = IESLightProfileIndex;
-		LightData->RectLightTextureIndex[LightData->Count] = 99;
+		LightData->RectLightTextureIndex[LightData->Count] = InvalidTextureIndex;
 		LightData->RectLightBarnCosAngle[LightData->Count] = LightParameters.RectLightBarnCosAngle;
 		LightData->RectLightBarnLength[LightData->Count] = LightParameters.RectLightBarnLength;
 
-		const bool bAllocateRectTextureSlot = Light.LightType == ELightComponentType::LightType_Rect && LightParameters.SourceTexture && CurrentRectLightIndex < RectLightTextureSlotCount;
-		if (bAllocateRectTextureSlot)
+		const bool bRequireTexture = Light.LightType == ELightComponentType::LightType_Rect && LightParameters.SourceTexture;
+
+		uint32 RectLightTextureIndex = InvalidTextureIndex;
+		if (bRequireTexture)
+		{
+			const uint32* IndexFound = RectTextureMap.Find(LightParameters.SourceTexture);
+			if (!IndexFound)
+			{
+				if (RectTextureMap.Num() < MaxRectLightTextureSlos)
+				{
+					RectLightTextureIndex = RectTextureMap.Num();
+					RectTextureMap.Add(LightParameters.SourceTexture, RectLightTextureIndex);
+				}
+			}
+			else
+			{
+				RectLightTextureIndex = *IndexFound;
+			}
+		}
+
+		if (RectLightTextureIndex != InvalidTextureIndex)
 		{
-			LightData->RectLightTextureIndex[LightData->Count] = CurrentRectLightIndex;
-			switch (CurrentRectLightIndex)
+			LightData->RectLightTextureIndex[LightData->Count] = RectLightTextureIndex;
+			switch (RectLightTextureIndex)
 			{
 			case 0: LightData->RectLightTexture0 = LightParameters.SourceTexture; break;
 			case 1: LightData->RectLightTexture1 = LightParameters.SourceTexture; break;
@@ -97,8 +117,6 @@ void SetupRaytracingLightData(
 			case 6: LightData->RectLightTexture6 = LightParameters.SourceTexture; break;
 			case 7: LightData->RectLightTexture7 = LightParameters.SourceTexture; break;
 			}
-			++CurrentRectLightIndex;
-			
 		}
 
 		const FVector2D FadeParams = Light.LightSceneInfo->Proxy->GetDirectionalLightDistanceFadeParameters(View.GetFeatureLevel(), Light.LightSceneInfo->IsPrecomputedLightingValid(), View.MaxShadowCascades);
