Player outlines fix (iOS, Android)

diff --git a/Engine/Shaders/Private/PostProcessMaterialShaders.usf b/Engine/Shaders/Private/PostProcessMaterialShaders.usf
--- a/Engine/Shaders/Private/PostProcessMaterialShaders.usf
+++ b/Engine/Shaders/Private/PostProcessMaterialShaders.usf
@@ -257,8 +257,14 @@ void MainPS_ES2(
 
 	// Grab emissive colour as output
 	half3 EmissiveColor = GetMaterialEmissive(PixelMaterialInputs);
-	half4 FullSceneColor = half4(EmissiveColor, Parameters.BackupSceneColorAlpha);
-
+	
+#if MATERIAL_OUTPUT_OPACITY_AS_ALPHA
+	half Alpha = GetMaterialOpacity(PixelMaterialInputs);
+#else
+	half Alpha = Parameters.BackupSceneColorAlpha;
+#endif
+	half4 FullSceneColor = half4(EmissiveColor, Alpha);
+	   
 #if POST_PROCESS_MATERIAL_BEFORE_TONEMAP
 #if OUTPUT_GAMMA_SPACE
 	FullSceneColor.rgb = sqrt(FullSceneColor.rgb);
diff --git a/Engine/Source/Runtime/Renderer/Private/PostProcess/PostProcessMaterial.cpp b/Engine/Source/Runtime/Renderer/Private/PostProcess/PostProcessMaterial.cpp
--- a/Engine/Source/Runtime/Renderer/Private/PostProcess/PostProcessMaterial.cpp
+++ b/Engine/Source/Runtime/Renderer/Private/PostProcess/PostProcessMaterial.cpp
@@ -268,12 +268,12 @@ void FRCPassPostProcessMaterial::Process(FRenderingCompositePassContext& Context
 			bool AllowStencilTestWithCopy = AllowStencilTest == 2;
 
 			// do the stencil test if 
-			// >= SM5 
+			// not SM4 
 			//   OR 
 			// reads DS but allowed make DS copy 
 			//   OR 
 			// DS not read at all.
-			bDoStencilTest = (FeatureLevel >= ERHIFeatureLevel::SM5) || 
+			bDoStencilTest = (FeatureLevel != ERHIFeatureLevel::SM4) || 
 				((bReadsCustomDepthStencil == AllowStencilTestWithCopy) || AllowStencilTestWithCopy);
 		}
 		else
@@ -292,6 +292,7 @@ void FRCPassPostProcessMaterial::Process(FRenderingCompositePassContext& Context
 
 	// The PP target - either from the render target pool or the ePId_Input0
 	const FSceneRenderTargetItem* DestRenderTarget = nullptr;
+	ERenderTargetLoadAction DestRenderTargetLoadAction = ERenderTargetLoadAction::Num;
 	const FSceneRenderTargetItem* CustomDepthStencilTarget = nullptr;
 
 	FDepthStencilStateRHIParamRef DepthStencilState;
@@ -302,7 +303,7 @@ void FRCPassPostProcessMaterial::Process(FRenderingCompositePassContext& Context
 		CustomDepthStencilTarget = &SceneContext.CustomDepth->GetRenderTargetItem();
 
 		// SM4 HW lacks support for texture bound as DepthRead_StencilRead and SRV simultaneously thus make a copy of DS buffer
-		if (FeatureLevel < ERHIFeatureLevel::SM5 && bReadsCustomDepthStencil)
+		if (FeatureLevel == ERHIFeatureLevel::SM4 && bReadsCustomDepthStencil)
 		{
 			// Dest param of CopyResource() call can only be an SRV (No render target flags) on DX10.0 (SM4)
 			FPooledRenderTargetDesc DSCopyDesc = SceneContext.CustomDepth->GetDesc();
@@ -332,10 +333,12 @@ void FRCPassPostProcessMaterial::Process(FRenderingCompositePassContext& Context
 		{
 			PassOutputs[0].PooledRenderTarget = GetInput(ePId_Input0)->GetOutput()->RequestInput();
 			DestRenderTarget = &PassOutputs[0].RequestSurface(Context);
+			DestRenderTargetLoadAction = ERenderTargetLoadAction::ELoad;
 		}
 		else
 		{
 			DestRenderTarget = &PassOutputs[0].RequestSurface(Context);
+			DestRenderTargetLoadAction = ERenderTargetLoadAction::ELoad;
 
 			Context.RHICmdList.CopyTexture(
 				GetInput(ePId_Input0)->GetOutput()->RequestSurface(Context).ShaderResourceTexture,
@@ -383,6 +386,11 @@ void FRCPassPostProcessMaterial::Process(FRenderingCompositePassContext& Context
 		BlendState = BlendStates[Material->GetBlendMode()];
 	}
 
+	if (DestRenderTargetLoadAction == ERenderTargetLoadAction::Num)
+	{
+		DestRenderTargetLoadAction = Context.GetLoadActionForRenderTarget(*DestRenderTarget);
+	}
+
 	FIntRect SrcRect = Context.SceneColorViewRect;
 	FIntRect DestRect = Context.GetSceneColorDestRect(*DestRenderTarget);
 	checkf(DestRect.Size() == SrcRect.Size(), TEXT("Post process material should not be used as upscaling pass."));
@@ -392,7 +400,7 @@ void FRCPassPostProcessMaterial::Process(FRenderingCompositePassContext& Context
 	{
 		RPInfo = FRHIRenderPassInfo(
 			DestRenderTarget->TargetableTexture,
-			MakeRenderTargetActions(Context.GetLoadActionForRenderTarget(*DestRenderTarget), ERenderTargetStoreAction::EStore),
+			MakeRenderTargetActions(DestRenderTargetLoadAction, ERenderTargetStoreAction::EStore),
 			CustomDepthStencilTarget->TargetableTexture,
 			MakeDepthStencilTargetActions(
 				MakeRenderTargetActions(ERenderTargetLoadAction::ENoAction, ERenderTargetStoreAction::ENoAction),
