Fixed shadows when depth bias is 0:
* For depth bias 0 FProjectedShadowInfo::ComputeTransitionSize() returns 0
* 1/ComputeTransitionSize() gets set as a shader parameter, which results in a INF which breaks shadows
Make sure that shadow soft transition size is greater than zero so 1/TransitionSize shader parameter won't be INF.

diff --git a/Engine/Source/Runtime/Renderer/Private/ShadowRendering.cpp b/Engine/Source/Runtime/Renderer/Private/ShadowRendering.cpp
--- a/Engine/Source/Runtime/Renderer/Private/ShadowRendering.cpp
+++ b/Engine/Source/Runtime/Renderer/Private/ShadowRendering.cpp
@@ -1188,7 +1188,7 @@ void FProjectedShadowInfo::UpdateShaderDepthBias()
 
 float FProjectedShadowInfo::ComputeTransitionSize() const
 {
-	float TransitionSize = 1;
+	float TransitionSize = 1.0f;
 
 	if (IsWholeScenePointLightShadow())
 	{
@@ -1215,7 +1215,7 @@ float FProjectedShadowInfo::ComputeTransitionSize() const
 	else if (bPreShadow)
 	{
 		// Preshadows don't have self shadowing, so make sure the shadow starts as close to the caster as possible
-		TransitionSize = 0.00001f;
+		TransitionSize = 0.0f;
 	}
 	else
 	{
@@ -1225,7 +1225,9 @@ float FProjectedShadowInfo::ComputeTransitionSize() const
 		TransitionSize *= 2.0f * LightSceneInfo->Proxy->GetUserShadowBias();
 	}
 
-	return TransitionSize;
+	// Make sure that shadow soft transition size is greater than zero so 1/TransitionSize shader parameter won't be INF.
+	const float MinTransitionSize = 0.00001f;
+	return FMath::Max(TransitionSize, MinTransitionSize);
 }
 
 /*-----------------------------------------------------------------------------
