Correct blending behavior for PPM. Bug was if stencil test was disabled then we did not prime the target output, causing blend modes to not work. Fix allows independent selection of stencil test and blend mode, then correctly primes output if necessary.

diff --git a/Engine/Source/Runtime/Renderer/Private/PostProcess/PostProcessMaterial.cpp b/Engine/Source/Runtime/Renderer/Private/PostProcess/PostProcessMaterial.cpp
--- a/Engine/Source/Runtime/Renderer/Private/PostProcess/PostProcessMaterial.cpp
+++ b/Engine/Source/Runtime/Renderer/Private/PostProcess/PostProcessMaterial.cpp
@@ -197,6 +197,13 @@ FRCPassPostProcessMaterial::FRCPassPostProcessMaterial(UMaterialInterface* InMat
 	{
 		MaterialInterface = UMaterial::GetDefaultMaterial(MD_PostProcess);
 	}
+
+	if (Material->IsStencilTestEnabled() || Material->GetBlendableOutputAlpha())
+	{
+		// Only allowed to have blend/stencil test if output format is compatible with ePId_Input0. 
+		// PF_Unknown implies output format is that of EPId_Input0
+		ensure(OutputFormat == PF_Unknown);
+	}
 }
 		
 /** The filter vertex declaration resource type. */
@@ -290,9 +297,6 @@ void FRCPassPostProcessMaterial::Process(FRenderingCompositePassContext& Context
 	// Copy of custom depth/stencil buffer if HW does not support simultaneously a texture bound as DepthRead_StencilRead and SRV
 	TRefCountPtr<IPooledRenderTarget> CustomDepthStencilCopy;
 
-	// The PP target - either from the render target pool or the ePId_Input0
-	const FSceneRenderTargetItem* DestRenderTarget = nullptr;
-	ERenderTargetLoadAction DestRenderTargetLoadAction = ERenderTargetLoadAction::Num;
 	const FSceneRenderTargetItem* CustomDepthStencilTarget = nullptr;
 
 	FDepthStencilStateRHIParamRef DepthStencilState;
@@ -324,29 +328,6 @@ void FRCPassPostProcessMaterial::Process(FRenderingCompositePassContext& Context
 			Swap(SceneContext.CustomDepth, CustomDepthStencilCopy);
 		}
 
-		// essentially only care about format and dimensions
-		FPooledRenderTargetDesc TargetDesc = PassOutputs[0].RenderTargetDesc;
-		TargetDesc.AutoWritable = InputDesc->AutoWritable;
-
-		// write directly into PPI0 if material does not read from it, otherwise make a copy of PPI0.
-		if (InputDesc->Compare(TargetDesc, false) && !MaterialShaderMap->UsesSceneTexture(PPI_PostProcessInput0))
-		{
-			PassOutputs[0].PooledRenderTarget = GetInput(ePId_Input0)->GetOutput()->RequestInput();
-			DestRenderTarget = &PassOutputs[0].RequestSurface(Context);
-			DestRenderTargetLoadAction = ERenderTargetLoadAction::ELoad;
-		}
-		else
-		{
-			DestRenderTarget = &PassOutputs[0].RequestSurface(Context);
-			DestRenderTargetLoadAction = ERenderTargetLoadAction::ELoad;
-
-			Context.RHICmdList.CopyTexture(
-				GetInput(ePId_Input0)->GetOutput()->RequestSurface(Context).ShaderResourceTexture,
-				DestRenderTarget->TargetableTexture,
-				FRHICopyTextureInfo()
-				);
-		}
-
 		static const FDepthStencilStateRHIParamRef StencilStates[] =
 		{
 			TStaticDepthStencilState<false, CF_Always, true, CF_Less>::GetRHI(),
@@ -365,12 +346,12 @@ void FRCPassPostProcessMaterial::Process(FRenderingCompositePassContext& Context
 	}
 	else
 	{
-		DestRenderTarget = &PassOutputs[0].RequestSurface(Context);
 		DepthStencilState = TStaticDepthStencilState<false, CF_Always>::GetRHI();
 	}
 
 	FBlendStateRHIParamRef BlendState = TStaticBlendState<>::GetRHI();
-	if (Material->GetBlendableOutputAlpha() && CVarPostProcessAllowBlendModes.GetValueOnRenderThread() != 0)
+	bool bDoOutputBlend = Material->GetBlendableOutputAlpha() && CVarPostProcessAllowBlendModes.GetValueOnRenderThread() != 0;
+	if (bDoOutputBlend)
 	{
 		static const FBlendStateRHIParamRef BlendStates[] =
 		{
@@ -386,8 +367,33 @@ void FRCPassPostProcessMaterial::Process(FRenderingCompositePassContext& Context
 		BlendState = BlendStates[Material->GetBlendMode()];
 	}
 
-	if (DestRenderTargetLoadAction == ERenderTargetLoadAction::Num)
+	// The PP target - either from the render target pool or the ePId_Input0
+	const FSceneRenderTargetItem* DestRenderTarget = nullptr;
+	ERenderTargetLoadAction DestRenderTargetLoadAction = ERenderTargetLoadAction::Num;
+		
+	if (bDoStencilTest || bDoOutputBlend)
 	{
+		if (!MaterialShaderMap->UsesSceneTexture(PPI_PostProcessInput0))
+		{
+			PassOutputs[0].PooledRenderTarget = GetInput(ePId_Input0)->GetOutput()->RequestInput();
+			DestRenderTarget = &PassOutputs[0].RequestSurface(Context);
+		}
+		else
+		{
+			DestRenderTarget = &PassOutputs[0].RequestSurface(Context);
+
+			Context.RHICmdList.CopyTexture(
+				GetInput(ePId_Input0)->GetOutput()->RequestSurface(Context).ShaderResourceTexture,
+				DestRenderTarget->TargetableTexture,
+				FRHICopyTextureInfo()
+				);
+		}
+
+		DestRenderTargetLoadAction = ERenderTargetLoadAction::ELoad;
+	}
+	else
+	{
+		DestRenderTarget = &PassOutputs[0].RequestSurface(Context);
 		DestRenderTargetLoadAction = Context.GetLoadActionForRenderTarget(*DestRenderTarget);
 	}
 
