Fixed: World flickers with World Origin Rebasing enabled

diff --git a/Engine/Source/Runtime/Renderer/Private/RendererScene.cpp b/Engine/Source/Runtime/Renderer/Private/RendererScene.cpp
--- a/Engine/Source/Runtime/Renderer/Private/RendererScene.cpp
+++ b/Engine/Source/Runtime/Renderer/Private/RendererScene.cpp
@@ -3160,16 +3160,36 @@ void FScene::ApplyWorldOffset(FVector InOffset)
 		});
 }
 
-void FScene::ApplyWorldOffset_RenderThread(FVector InOffset)
+void FScene::ApplyWorldOffset_RenderThread(const FVector& InOffset)
 {
 	QUICK_SCOPE_CYCLE_COUNTER(STAT_SceneApplyWorldOffset);
 	
+	GPUScene.bUpdateAllPrimitives = true;
+
 	// Primitives
-	for (auto It = Primitives.CreateIterator(); It; ++It)
+	for (int32 Idx = 0; Idx < Primitives.Num(); ++Idx)
 	{
-		(*It)->ApplyWorldOffset(InOffset);
+		Primitives[Idx]->ApplyWorldOffset(InOffset);
+	}
+	
+	// Primitive transforms
+	for (int32 Idx = 0; Idx < PrimitiveTransforms.Num(); ++Idx)
+	{
+		PrimitiveTransforms[Idx].SetOrigin(PrimitiveTransforms[Idx].GetOrigin() + InOffset);
 	}
 
+	// Primitive bounds
+	for (int32 Idx = 0; Idx < PrimitiveBounds.Num(); ++Idx)
+	{
+		PrimitiveBounds[Idx].BoxSphereBounds.Origin+= InOffset;
+	}
+
+	// Primitive occlusion bounds
+	for (int32 Idx = 0; Idx < PrimitiveOcclusionBounds.Num(); ++Idx)
+	{
+		PrimitiveOcclusionBounds[Idx].Origin+= InOffset;
+	}
+	
 	// Precomputed light volumes
 	for (const FPrecomputedLightVolume* It : PrecomputedLightVolumes)
 	{
@@ -3188,18 +3208,6 @@ void FScene::ApplyWorldOffset_RenderThread(FVector InOffset)
 	// Primitives octree
 	PrimitiveOctree.ApplyOffset(InOffset, /*bGlobalOctee*/ true);
 
-	// Primitive bounds
-	for (auto It = PrimitiveBounds.CreateIterator(); It; ++It)
-	{
-		(*It).BoxSphereBounds.Origin+= InOffset;
-	}
-
-	// Primitive occlusion bounds
-	for (auto It = PrimitiveOcclusionBounds.CreateIterator(); It; ++It)
-	{
-		(*It).Origin+= InOffset;
-	}
-
 	// Lights
 	VectorRegister OffsetReg = VectorLoadFloat3_W0(&InOffset);
 	for (auto It = Lights.CreateIterator(); It; ++It)
diff --git a/Engine/Source/Runtime/Renderer/Private/ScenePrivate.h b/Engine/Source/Runtime/Renderer/Private/ScenePrivate.h
--- a/Engine/Source/Runtime/Renderer/Private/ScenePrivate.h
+++ b/Engine/Source/Runtime/Renderer/Private/ScenePrivate.h
@@ -2973,7 +2973,7 @@ private:
 	 * 
 	 * @param	InOffset	Delta to shift scene by
 	 */
-	void ApplyWorldOffset_RenderThread(FVector InOffset);
+	void ApplyWorldOffset_RenderThread(const FVector& InOffset);
 
 	/**
 	 * Notification from game thread that level was added to a world
