Guard against nullptr deref based on StaticAnalysis.

diff --git a/Engine/Source/Runtime/Renderer/Private/PostProcess/PostProcessMaterial.cpp b/Engine/Source/Runtime/Renderer/Private/PostProcess/PostProcessMaterial.cpp
--- a/Engine/Source/Runtime/Renderer/Private/PostProcess/PostProcessMaterial.cpp
+++ b/Engine/Source/Runtime/Renderer/Private/PostProcess/PostProcessMaterial.cpp
@@ -198,7 +198,7 @@ FRCPassPostProcessMaterial::FRCPassPostProcessMaterial(UMaterialInterface* InMat
 		MaterialInterface = UMaterial::GetDefaultMaterial(MD_PostProcess);
 	}
 
-	if (Material->IsStencilTestEnabled() || Material->GetBlendableOutputAlpha())
+	if (Material && (Material->IsStencilTestEnabled() || Material->GetBlendableOutputAlpha()))
 	{
 		// Only allowed to have blend/stencil test if output format is compatible with ePId_Input0. 
 		// PF_Unknown implies output format is that of EPId_Input0
