Fixed crash in mobile renderer when trying to draw opaque mesh without a proxy (DrawTileMesh rendering path).

diff --git a/Engine/Source/Runtime/Renderer/Private/MobileBasePass.cpp b/Engine/Source/Runtime/Renderer/Private/MobileBasePass.cpp
--- a/Engine/Source/Runtime/Renderer/Private/MobileBasePass.cpp
+++ b/Engine/Source/Runtime/Renderer/Private/MobileBasePass.cpp
@@ -649,7 +649,7 @@ void FMobileBasePassMeshProcessor::Process(
 	else
 	{
 		// Background primitives will be rendered last in masked/non-masked buckets
-		bool bBackground = PrimitiveSceneProxy->TreatAsBackgroundForOcclusion();
+		bool bBackground = PrimitiveSceneProxy ? PrimitiveSceneProxy->TreatAsBackgroundForOcclusion() : false;
 		// Default static sort key separates masked and non-masked geometry, generic mesh sorting will also sort by PSO
 		// if platform wants front to back sorting, this key will be recomputed in InitViews
 		SortKey = GetBasePassStaticSortKey(BlendMode, bBackground);
