UE-67120 Editor crash on startup after enabling Morpheus plugin...
UE-67297 ShaderCompiler crash when opening the editor with Morpheus plugin and project setting enabled
-Moved this shader file out of platform shader directory into the  plugin shaders (it isn't actually for use on that platform, its used in the editor).  Sadly we still need to use it from engine code if the shader file exists because a postprocesspass cannot be defined in a plugin and I couldn't find a way to pull the shader from the plugin and use it in the pass so I had to hardcode the shader path mapping for this plugin in launchengineloop.  A rendering refactor is coming along that will make postprocess passes more like other passes, so it makes sense to defer refactoring this until after that.

diff --git a/Engine/Shaders/Private/PostProcessHMDMorpheus.usf b/Engine/Shaders/Private/PostProcessHMDMorpheus.usf
deleted file mode 100644
diff --git a/Engine/Source/Programs/UnrealBuildTool/Platform/Windows/UEBuildWindows.cs b/Engine/Source/Programs/UnrealBuildTool/Platform/Windows/UEBuildWindows.cs
--- a/Engine/Source/Programs/UnrealBuildTool/Platform/Windows/UEBuildWindows.cs
+++ b/Engine/Source/Programs/UnrealBuildTool/Platform/Windows/UEBuildWindows.cs
@@ -1641,7 +1641,7 @@ namespace UnrealBuildTool
 			CompileEnvironment.Definitions.Add(String.Format("WINVER=0x{0:X4}", Target.WindowsPlatform.TargetWindowsVersion));
 			CompileEnvironment.Definitions.Add("PLATFORM_WINDOWS=1");
 
-			FileReference MorpheusShaderPath = FileReference.Combine(UnrealBuildTool.EngineDirectory, "Shaders", "Private", "Platform", "PS4", "PostProcessHMDMorpheus.usf");
+			FileReference MorpheusShaderPath = FileReference.Combine(UnrealBuildTool.EngineDirectory, "Plugins", "Runtime", "PS4", "Morpheus", "Shaders", "Private", "PostProcessHMDMorpheus.usf");
 			if (FileReference.Exists(MorpheusShaderPath))
 			{
 				CompileEnvironment.Definitions.Add("HAS_MORPHEUS=1");
diff --git a/Engine/Source/Runtime/Launch/Private/LaunchEngineLoop.cpp b/Engine/Source/Runtime/Launch/Private/LaunchEngineLoop.cpp
--- a/Engine/Source/Runtime/Launch/Private/LaunchEngineLoop.cpp
+++ b/Engine/Source/Runtime/Launch/Private/LaunchEngineLoop.cpp
@@ -1288,6 +1288,13 @@ int32 FEngineLoop::PreInit(const TCHAR* CmdLine)
 	// Add the default engine shader dir
 	AddShaderSourceDirectoryMapping(TEXT("/Engine"), FGenericPlatformProcess::ShaderDir());
 
+// This shader file is in the plugin, but it is compiled even if the plugin is disabled so we can't have the plugin add this mapping.
+// This is a temporary workaround.
+#if defined(MORPHEUS_ENGINE_DISTORTION) && MORPHEUS_ENGINE_DISTORTION 
+	FString PluginShaderDir = FPaths::Combine(IPluginManager::Get().FindPlugin(TEXT("Morpheus"))->GetBaseDir(), TEXT("Shaders"));
+	AddShaderSourceDirectoryMapping(TEXT("/Plugin/Morpheus"), PluginShaderDir);
+#endif
+
 	TArray<FString> Tokens;
 	TArray<FString> Switches;
 	UCommandlet::ParseCommandLine(CommandLineCopy, Tokens, Switches);
diff --git a/Engine/Source/Runtime/Renderer/Private/PostProcess/PostProcessMorpheus.cpp b/Engine/Source/Runtime/Renderer/Private/PostProcess/PostProcessMorpheus.cpp
--- a/Engine/Source/Runtime/Renderer/Private/PostProcess/PostProcessMorpheus.cpp
+++ b/Engine/Source/Runtime/Renderer/Private/PostProcess/PostProcessMorpheus.cpp
@@ -189,8 +189,9 @@ public:
 	}
 };
 
-IMPLEMENT_SHADER_TYPE(, FPostProcessMorpheusVS, TEXT("/Engine/Private/PostProcessHMDMorpheus.usf"), TEXT("MainVS"), SF_Vertex);
-IMPLEMENT_SHADER_TYPE(, FPostProcessMorpheusPS, TEXT("/Engine/Private/PostProcessHMDMorpheus.usf"), TEXT("MainPS"), SF_Pixel);
+// This .usf file is in the morpheus plugin, MORPHEUS_ENGINE_DISTORTION, under which all this code is ifdef-ed can only be true when the plugin exists. 
+IMPLEMENT_SHADER_TYPE(, FPostProcessMorpheusVS, TEXT("/Plugin/Morpheus/Private/PostProcessHMDMorpheus.usf"), TEXT("MainVS"), SF_Vertex);
+IMPLEMENT_SHADER_TYPE(, FPostProcessMorpheusPS, TEXT("/Plugin/Morpheus/Private/PostProcessHMDMorpheus.usf"), TEXT("MainPS"), SF_Pixel);
 
 void FRCPassPostProcessMorpheus::Process(FRenderingCompositePassContext& Context)
 {
