Back out changelist 5417701
-This fix breaks the ddc build, which is more important.

diff --git a/Engine/Shaders/Private/PostProcessHMDMorpheus.usf b/Engine/Shaders/Private/PostProcessHMDMorpheus.usf
new file mode 100644
--- /dev/null
+++ b/Engine/Shaders/Private/PostProcessHMDMorpheus.usf
@@ -0,0 +1,11 @@
+// Copyright 1998-2019 Epic Games, Inc. All Rights Reserved.
+
+// PS4 directory may not exists. As a result a shader compiler may not be able to
+// find file for include scanning, but the include scanner will silently ignore
+// PS4 directory by design.
+//
+// TODO(VR team): move shader file and C++ into a plugin.
+
+#include "Common.ush"
+#include "PostProcessCommon.ush"
+#include "Platform/PS4/PostProcessHMDMorpheus.usf"
diff --git a/Engine/Source/Programs/UnrealBuildTool/Platform/Windows/UEBuildWindows.cs b/Engine/Source/Programs/UnrealBuildTool/Platform/Windows/UEBuildWindows.cs
--- a/Engine/Source/Programs/UnrealBuildTool/Platform/Windows/UEBuildWindows.cs
+++ b/Engine/Source/Programs/UnrealBuildTool/Platform/Windows/UEBuildWindows.cs
@@ -1641,7 +1641,7 @@ namespace UnrealBuildTool
 			CompileEnvironment.Definitions.Add(String.Format("WINVER=0x{0:X4}", Target.WindowsPlatform.TargetWindowsVersion));
 			CompileEnvironment.Definitions.Add("PLATFORM_WINDOWS=1");
 
-			FileReference MorpheusShaderPath = FileReference.Combine(UnrealBuildTool.EngineDirectory, "Plugins", "Runtime", "PS4", "Morpheus", "Shaders", "Private", "PostProcessHMDMorpheus.usf");
+			FileReference MorpheusShaderPath = FileReference.Combine(UnrealBuildTool.EngineDirectory, "Shaders", "Private", "Platform", "PS4", "PostProcessHMDMorpheus.usf");
 			if (FileReference.Exists(MorpheusShaderPath))
 			{
 				CompileEnvironment.Definitions.Add("HAS_MORPHEUS=1");
diff --git a/Engine/Source/Runtime/Launch/Private/LaunchEngineLoop.cpp b/Engine/Source/Runtime/Launch/Private/LaunchEngineLoop.cpp
--- a/Engine/Source/Runtime/Launch/Private/LaunchEngineLoop.cpp
+++ b/Engine/Source/Runtime/Launch/Private/LaunchEngineLoop.cpp
@@ -1288,19 +1288,6 @@ int32 FEngineLoop::PreInit(const TCHAR* CmdLine)
 	// Add the default engine shader dir
 	AddShaderSourceDirectoryMapping(TEXT("/Engine"), FGenericPlatformProcess::ShaderDir());
 
-// This shader file is in the plugin, but it is compiled even if the plugin is disabled so we can't have the plugin add this mapping.
-// This is a temporary workaround.
-#if defined(MORPHEUS_ENGINE_DISTORTION) && MORPHEUS_ENGINE_DISTORTION 
-	{
-		TSharedPtr<IPlugin> Plugin = IPluginManager::Get().FindPlugin(TEXT("Morpheus"));
-		if (Plugin.IsValid())
-		{
-			FString PluginShaderDir = FPaths::Combine(Plugin->GetBaseDir(), TEXT("Shaders"));
-			AddShaderSourceDirectoryMapping(TEXT("/Plugin/Morpheus"), PluginShaderDir);
-		}
-	}
-#endif
-
 	TArray<FString> Tokens;
 	TArray<FString> Switches;
 	UCommandlet::ParseCommandLine(CommandLineCopy, Tokens, Switches);
diff --git a/Engine/Source/Runtime/Renderer/Private/PostProcess/PostProcessMorpheus.cpp b/Engine/Source/Runtime/Renderer/Private/PostProcess/PostProcessMorpheus.cpp
--- a/Engine/Source/Runtime/Renderer/Private/PostProcess/PostProcessMorpheus.cpp
+++ b/Engine/Source/Runtime/Renderer/Private/PostProcess/PostProcessMorpheus.cpp
@@ -189,9 +189,8 @@ public:
 	}
 };
 
-// This .usf file is in the morpheus plugin, MORPHEUS_ENGINE_DISTORTION, under which all this code is ifdef-ed can only be true when the plugin exists. 
-IMPLEMENT_SHADER_TYPE(, FPostProcessMorpheusVS, TEXT("/Plugin/Morpheus/Private/PostProcessHMDMorpheus.usf"), TEXT("MainVS"), SF_Vertex);
-IMPLEMENT_SHADER_TYPE(, FPostProcessMorpheusPS, TEXT("/Plugin/Morpheus/Private/PostProcessHMDMorpheus.usf"), TEXT("MainPS"), SF_Pixel);
+IMPLEMENT_SHADER_TYPE(, FPostProcessMorpheusVS, TEXT("/Engine/Private/PostProcessHMDMorpheus.usf"), TEXT("MainVS"), SF_Vertex);
+IMPLEMENT_SHADER_TYPE(, FPostProcessMorpheusPS, TEXT("/Engine/Private/PostProcessHMDMorpheus.usf"), TEXT("MainPS"), SF_Pixel);
 
 void FRCPassPostProcessMorpheus::Process(FRenderingCompositePassContext& Context)
 {
