Ray traced reflection can now correctly fallback to environment lighting contribution with sky visible from primary reflection ray.
With or without Reflection.MaxDistance at -1.

diff --git a/Engine/Source/Runtime/Renderer/Private/DistanceFieldAmbientOcclusion.cpp b/Engine/Source/Runtime/Renderer/Private/DistanceFieldAmbientOcclusion.cpp
--- a/Engine/Source/Runtime/Renderer/Private/DistanceFieldAmbientOcclusion.cpp
+++ b/Engine/Source/Runtime/Renderer/Private/DistanceFieldAmbientOcclusion.cpp
@@ -722,7 +722,7 @@ bool ShouldRenderDeferredDynamicSkyLight(const FScene* Scene, const FSceneViewFa
 {
 	return Scene->SkyLight
 		&& Scene->SkyLight->ProcessedTexture
-		&& !ShouldRenderRayTracingSkyLight(Scene->SkyLight)
+		&& !ShouldRenderRayTracingSkyLight(Scene->SkyLight) // Disable diffuse sky contribution if evaluated by RT Sky.
 		&& !Scene->SkyLight->bWantsStaticShadowing
 		&& !Scene->SkyLight->bHasStaticLighting
 		&& ViewFamily.EngineShowFlags.SkyLighting
diff --git a/Engine/Source/Runtime/Renderer/Private/ReflectionEnvironment.cpp b/Engine/Source/Runtime/Renderer/Private/ReflectionEnvironment.cpp
--- a/Engine/Source/Runtime/Renderer/Private/ReflectionEnvironment.cpp
+++ b/Engine/Source/Runtime/Renderer/Private/ReflectionEnvironment.cpp
@@ -726,10 +726,10 @@ void FDeferredShadingSceneRenderer::RenderDeferredReflectionsAndSkyLighting(FRHI
 
 	const bool bRayTracedReflections = IsRayTracingEnabled() && (GRayTracingReflections < 0 ? bAnyViewWithRaytracingReflections : GRayTracingReflections);
 
+	// The specular sky light contribution is also needed by RT Reflections as a fallback.
 	const bool bSkyLight = Scene->SkyLight
 		&& Scene->SkyLight->ProcessedTexture
-		&& !Scene->SkyLight->bHasStaticLighting
-		&& !ShouldRenderRayTracingSkyLight(Scene->SkyLight);
+		&& !Scene->SkyLight->bHasStaticLighting;
 
 	bool bDynamicSkyLight = ShouldRenderDeferredDynamicSkyLight(Scene, ViewFamily);
 	bool bApplySkyShadowing = false;
