Adreno occlusion mode, reduces GPU time on lower end Adreno devices

diff --git a/Engine/Source/Runtime/OpenGLDrv/Private/OpenGLCommands.cpp b/Engine/Source/Runtime/OpenGLDrv/Private/OpenGLCommands.cpp
--- a/Engine/Source/Runtime/OpenGLDrv/Private/OpenGLCommands.cpp
+++ b/Engine/Source/Runtime/OpenGLDrv/Private/OpenGLCommands.cpp
@@ -3424,8 +3424,11 @@ void FOpenGLDynamicRHI::RHIBlockUntilGPUIdle()
 
 void FOpenGLDynamicRHI::RHISubmitCommandsAndFlushGPU()
 {
-	FOpenGL::Flush();
-	RHIPollOcclusionQueries();
+	RunOnGLRenderContextThread([&]()
+	{
+		FOpenGL::Flush();
+		RHIPollOcclusionQueries();
+	});
 }
 
 /**
diff --git a/Engine/Source/Runtime/Renderer/Private/MobileShadingRenderer.cpp b/Engine/Source/Runtime/Renderer/Private/MobileShadingRenderer.cpp
--- a/Engine/Source/Runtime/Renderer/Private/MobileShadingRenderer.cpp
+++ b/Engine/Source/Runtime/Renderer/Private/MobileShadingRenderer.cpp
@@ -70,6 +70,13 @@ static TAutoConsoleVariable<int32> CVarMobileMoveSubmissionHintAfterTranslucency
 	TEXT("1: Submission hint occurs after translucency. (Default)"),
 	ECVF_Scalability | ECVF_RenderThreadSafe);
 
+static TAutoConsoleVariable<int32> CVarMobileAdrenoOcclusionMode(
+	TEXT("r.Mobile.AdrenoOcclusionMode"),
+	0,
+	TEXT("0: Render occlusion queries after the base pass (default).\n")
+	TEXT("1: Render occlusion queries after translucency and a flush, which can help Adreno devices in GL mode."),
+	ECVF_RenderThreadSafe);
+
 DECLARE_CYCLE_STAT(TEXT("SceneStart"), STAT_CLMM_SceneStart, STATGROUP_CommandListMarkers);
 DECLARE_CYCLE_STAT(TEXT("SceneEnd"), STAT_CLMM_SceneEnd, STATGROUP_CommandListMarkers);
 DECLARE_CYCLE_STAT(TEXT("InitViews"), STAT_CLMM_InitViews, STATGROUP_CommandListMarkers);
@@ -459,10 +466,14 @@ void FMobileSceneRenderer::Render(FRHICommandListImmediate& RHICmdList)
 	RenderMobileBasePass(RHICmdList, ViewList);
 	RHICmdList.ImmediateFlush(EImmediateFlushType::DispatchToRHIThread);
 
-	// Issue occlusion queries
-	RHICmdList.SetCurrentStat(GET_STATID(STAT_CLMM_Occlusion));
-	RenderOcclusion(RHICmdList);
-	RHICmdList.ImmediateFlush(EImmediateFlushType::DispatchToRHIThread);
+	const bool bAdrenoOcclusionMode = CVarMobileAdrenoOcclusionMode.GetValueOnRenderThread() != 0;
+	if (!bAdrenoOcclusionMode)
+	{
+	    // Issue occlusion queries
+	    RHICmdList.SetCurrentStat(GET_STATID(STAT_CLMM_Occlusion));
+	    RenderOcclusion(RHICmdList);
+	    RHICmdList.ImmediateFlush(EImmediateFlushType::DispatchToRHIThread);
+	}
 
 	{
 		CSV_SCOPED_TIMING_STAT_EXCLUSIVE(ViewExtensionPostRenderBasePass);
@@ -548,6 +559,16 @@ void FMobileSceneRenderer::Render(FRHICommandListImmediate& RHICmdList)
 		RHICmdList.ImmediateFlush(EImmediateFlushType::DispatchToRHIThread);
 	}
 
+	if (bAdrenoOcclusionMode)
+	{
+	    RHICmdList.SetCurrentStat(GET_STATID(STAT_CLMM_Occlusion));
+		// flush
+		RHICmdList.SubmitCommandsAndFlushGPU();
+		// Issue occlusion queries
+	    RenderOcclusion(RHICmdList);
+	    RHICmdList.ImmediateFlush(EImmediateFlushType::DispatchToRHIThread);
+	}
+
 	// Pre-tonemap before MSAA resolve (iOS only)
 	if (!bGammaSpace)
 	{
