Fix to shader complexity on forward shading when using the velocity buffer

diff --git a/Engine/Shaders/Private/QuadOverdraw.ush b/Engine/Shaders/Private/QuadOverdraw.ush
--- a/Engine/Shaders/Private/QuadOverdraw.ush
+++ b/Engine/Shaders/Private/QuadOverdraw.ush
@@ -47,12 +47,6 @@ float ComplexityToFloat(uint Complexity)
 
 #if OUTPUT_QUAD_OVERDRAW
 
-#if FORWARD_SHADING
-#define QUAD_BUFFER_REGISTER u1
-#else
-#define QUAD_BUFFER_REGISTER u7
-#endif
-
 // The temporary buffer used to synchronize and exchange data between quad sub-pixels.
 // Left half hold QuadDescriptor, right half hold QuadComplexity
 // Both are halfres here.
diff --git a/Engine/Source/Runtime/Renderer/Private/PostProcess/SceneRenderTargets.cpp b/Engine/Source/Runtime/Renderer/Private/PostProcess/SceneRenderTargets.cpp
--- a/Engine/Source/Runtime/Renderer/Private/PostProcess/SceneRenderTargets.cpp
+++ b/Engine/Source/Runtime/Renderer/Private/PostProcess/SceneRenderTargets.cpp
@@ -706,7 +706,7 @@ int32 FSceneRenderTargets::FillGBufferRenderPassInfo(ERenderTargetLoadAction Col
 
 	// The velocity buffer needs to be bound before other optionnal rendertargets (when UseSelectiveBasePassOutputs() is true).
 	// Otherwise there is an issue on some AMD hardware where the target does not get updated. Seems to be related to the velocity buffer format as it works fine with other targets.
-	if (bAllocateVelocityGBuffer)
+	if (bAllocateVelocityGBuffer && !IsSimpleForwardShadingEnabled(ShaderPlatform))
 	{
 		OutVelocityRTIndex = MRTCount;
 		check(OutVelocityRTIndex == 4 || (!bUseGBuffer && OutVelocityRTIndex == 1)); // As defined in BasePassPixelShader.usf
@@ -787,13 +787,33 @@ int32 FSceneRenderTargets::GetGBufferRenderTargets(ERenderTargetLoadAction Color
 	return MRTCount;
 }
 
+#if !(UE_BUILD_SHIPPING || UE_BUILD_TEST)
+int32 FSceneRenderTargets::GetQuadOverdrawUAVIndex(EShaderPlatform Platform, ERHIFeatureLevel::Type FeatureLevel)
+{
+	if (IsSimpleForwardShadingEnabled(Platform))
+	{
+		return 1;
+	}
+	else if (IsForwardShadingEnabled(Platform))
+	{
+		return FVelocityRendering::BasePassCanOutputVelocity(FeatureLevel) ? 2 : 1;
+	}
+	else // GBuffer
+	{
+		return FVelocityRendering::BasePassCanOutputVelocity(FeatureLevel) ? 7 : 6;
+	}
+}
+#endif // !(UE_BUILD_SHIPPING || UE_BUILD_TEST)
+
 void FSceneRenderTargets::SetQuadOverdrawUAV(FRHICommandList& RHICmdList, bool bBindQuadOverdrawBuffers, bool bClearQuadOverdrawBuffers, FRHIRenderPassInfo& OutInfo)
 {
-	if (bBindQuadOverdrawBuffers && AllowDebugViewShaderMode(DVSM_QuadComplexity, GetFeatureLevelShaderPlatform(CurrentFeatureLevel), CurrentFeatureLevel))
+#if !(UE_BUILD_SHIPPING || UE_BUILD_TEST)
+	const EShaderPlatform ShaderPlatform = GetFeatureLevelShaderPlatform(CurrentFeatureLevel);
+	if (bBindQuadOverdrawBuffers && AllowDebugViewShaderMode(DVSM_QuadComplexity, ShaderPlatform, CurrentFeatureLevel))
 	{
 		if (QuadOverdrawBuffer.IsValid() && QuadOverdrawBuffer->GetRenderTargetItem().UAV.IsValid())
 		{
-			QuadOverdrawIndex = IsAnyForwardShadingEnabled(GetFeatureLevelShaderPlatform(CurrentFeatureLevel)) ? 1 : 7; // As defined in QuadOverdraw.usf
+			QuadOverdrawIndex = GetQuadOverdrawUAVIndex(ShaderPlatform, CurrentFeatureLevel);
 
 			// Increase the rendertarget count in order to control the bound slot of the UAV.
 			check(OutInfo.GetNumColorRenderTargets() <= QuadOverdrawIndex);
@@ -809,6 +829,7 @@ void FSceneRenderTargets::SetQuadOverdrawUAV(FRHICommandList& RHICmdList, bool b
 			}
 		}
 	}
+#endif // !(UE_BUILD_SHIPPING || UE_BUILD_TEST)
 }
 
 void FSceneRenderTargets::BeginRenderingGBuffer(FRHICommandList& RHICmdList, ERenderTargetLoadAction ColorLoadAction, ERenderTargetLoadAction DepthLoadAction, FExclusiveDepthStencil::Type DepthStencilAccess, bool bBindQuadOverdrawBuffers, bool bClearQuadOverdrawBuffers, const FLinearColor& ClearColor/*=(0,0,0,1)*/, bool bIsWireframe)
diff --git a/Engine/Source/Runtime/Renderer/Private/PostProcess/SceneRenderTargets.h b/Engine/Source/Runtime/Renderer/Private/PostProcess/SceneRenderTargets.h
--- a/Engine/Source/Runtime/Renderer/Private/PostProcess/SceneRenderTargets.h
+++ b/Engine/Source/Runtime/Renderer/Private/PostProcess/SceneRenderTargets.h
@@ -210,6 +210,11 @@ public:
 
 	void SetSeparateTranslucencyBufferSize(bool bAnyViewWantsDownsampledSeparateTranslucency);
 
+#if !(UE_BUILD_SHIPPING || UE_BUILD_TEST)
+	/** Returns the RT index where the QuadOverdrawUAV will be bound. */
+	static int32 GetQuadOverdrawUAVIndex(EShaderPlatform Platform, ERHIFeatureLevel::Type FeatureLevel);
+#endif
+
 	void SetQuadOverdrawUAV(FRHICommandList& RHICmdList, bool bBindQuadOverdrawBuffers, FRHISetRenderTargetsInfo& Info);
 	void SetQuadOverdrawUAV(FRHICommandList& RHICmdList, bool bBindQuadOverdrawBuffers, bool bClearQuadOverdrawBuffers, FRHIRenderPassInfo& Info);
 	void BeginRenderingGBuffer(FRHICommandList& RHICmdList, ERenderTargetLoadAction ColorLoadAction, ERenderTargetLoadAction DepthLoadAction, FExclusiveDepthStencil::Type DepthStencilAccess, bool bBindQuadOverdrawBuffers, bool bClearQuadOverdrawBuffers = false, const FLinearColor& ClearColor = FLinearColor(0, 0, 0, 1), bool bIsWireframe=false);
diff --git a/Engine/Source/Runtime/Renderer/Private/ShaderComplexityRendering.h b/Engine/Source/Runtime/Renderer/Private/ShaderComplexityRendering.h
--- a/Engine/Source/Runtime/Renderer/Private/ShaderComplexityRendering.h
+++ b/Engine/Source/Runtime/Renderer/Private/ShaderComplexityRendering.h
@@ -12,6 +12,7 @@ ShaderComplexityRendering.h: Declarations used for the shader complexity viewmod
 #include "GlobalShader.h"
 #include "DebugViewModeRendering.h"
 #include "DebugViewModeInterface.h"
+#include "PostProcess/SceneRenderTargets.h"
 
 class FPrimitiveSceneProxy;
 struct FMeshBatchElement;
@@ -60,6 +61,9 @@ public:
 	static void ModifyCompilationEnvironment(EShaderPlatform Platform, const FMaterial* Material, FShaderCompilerEnvironment& OutEnvironment)
 	{
 		OutEnvironment.SetDefine(TEXT("OUTPUT_QUAD_OVERDRAW"), AllowDebugViewShaderMode(DVSM_QuadComplexity, Platform, GetMaxSupportedFeatureLevel(Platform)));
+		TCHAR BufferRegister[] = { 'u', '0', 0 };
+		BufferRegister[1] += FSceneRenderTargets::GetQuadOverdrawUAVIndex(Platform, Material->GetFeatureLevel());
+		OutEnvironment.SetDefine(TEXT("QUAD_BUFFER_REGISTER"), BufferRegister);
 	}
 
 	virtual void GetDebugViewModeShaderBindings(
