Fix issue where no alpha was written by tonemapper even if the post process alpha passthrough setting was enabled

#rnx

diff --git a/Engine/Source/Runtime/Renderer/Private/PostProcess/PostProcessTonemap.cpp b/Engine/Source/Runtime/Renderer/Private/PostProcess/PostProcessTonemap.cpp
--- a/Engine/Source/Runtime/Renderer/Private/PostProcess/PostProcessTonemap.cpp
+++ b/Engine/Source/Runtime/Renderer/Private/PostProcess/PostProcessTonemap.cpp
@@ -11,6 +11,7 @@
 #include "PostProcess/SceneFilterRendering.h"
 #include "PostProcess/PostProcessCombineLUTs.h"
 #include "PostProcess/PostProcessMobile.h"
+#include "PostProcess/PostProcessing.h"
 #include "ClearQuad.h"
 #include "PipelineStateCache.h"
 
@@ -1085,9 +1086,9 @@ void FRCPassPostProcessTonemap::Process(FRenderingCompositePassContext& Context)
 			{
 				FGraphicsPipelineStateInitializer GraphicsPSOInit;
 				Context.RHICmdList.ApplyCachedRenderTargets(GraphicsPSOInit);
-				if (Context.View.AntiAliasingMethod == AAM_FXAA)
+				if (Context.View.AntiAliasingMethod == AAM_FXAA || FPostProcessing::HasAlphaChannelSupport())
 				{
-					// Alpha is used to store luminance for post process AA
+					// Alpha is used to store luminance for post process AA or as passthrough alpha
 					GraphicsPSOInit.BlendState = TStaticBlendState<>::GetRHI();
 				}
 				else
