Copying over fix from Dev-Rendering:

When applying PPM and we must prime the RT, check if CopyResource() is unavailable due to format conversion. If so use copy pixel shader.

diff --git a/Engine/Shaders/Private/ScreenPass.usf b/Engine/Shaders/Private/ScreenPass.usf
new file mode 100644
--- /dev/null
+++ b/Engine/Shaders/Private/ScreenPass.usf
@@ -0,0 +1,14 @@
+// Copyright 1998-2019 Epic Games, Inc. All Rights Reserved.
+
+#include "Common.ush"
+
+Texture2D SrcTexture;
+SamplerState SrcSampler;
+
+float4 CopyRectPS(
+	noperspective float4 UVAndScreenPos : TEXCOORD0
+    ) : SV_Target0
+{
+    float2 UV = UVAndScreenPos.xy;
+    return Texture2DSample(SrcTexture, SrcSampler, UV);
+}
diff --git a/Engine/Source/Runtime/RenderCore/Private/CommonRenderResources.cpp b/Engine/Source/Runtime/RenderCore/Private/CommonRenderResources.cpp
--- a/Engine/Source/Runtime/RenderCore/Private/CommonRenderResources.cpp
+++ b/Engine/Source/Runtime/RenderCore/Private/CommonRenderResources.cpp
@@ -15,6 +15,7 @@ TGlobalResource<FScreenRectangleVertexBuffer> GScreenRectangleVertexBuffer;
 TGlobalResource<FScreenRectangleIndexBuffer> GScreenRectangleIndexBuffer;
 
 IMPLEMENT_GLOBAL_SHADER(FVisualizeTextureVS, "/Engine/Private/Tools/FullscreenVertexShader.usf", "MainVS", SF_Vertex);
+IMPLEMENT_GLOBAL_SHADER(FCopyRectPS, "/Engine/Private/ScreenPass.usf", "CopyRectPS", SF_Pixel);
 
 
 void FScreenRectangleVertexBuffer::InitRHI()
diff --git a/Engine/Source/Runtime/RenderCore/Public/CommonRenderResources.h b/Engine/Source/Runtime/RenderCore/Public/CommonRenderResources.h
--- a/Engine/Source/Runtime/RenderCore/Public/CommonRenderResources.h
+++ b/Engine/Source/Runtime/RenderCore/Public/CommonRenderResources.h
@@ -104,3 +104,20 @@ class RENDERCORE_API FVisualizeTextureVS : public FGlobalShader
 	BEGIN_SHADER_PARAMETER_STRUCT(FParameters, )
 	END_SHADER_PARAMETER_STRUCT()
 };
+
+/** Pixel shader to copy pixels from src to dst performing a format change that works on all platforms. */
+class RENDERCORE_API FCopyRectPS : public FGlobalShader
+{
+	DECLARE_GLOBAL_SHADER(FCopyRectPS);
+	SHADER_USE_PARAMETER_STRUCT(FCopyRectPS, FGlobalShader);
+
+	static bool ShouldCompilePermutation(const FGlobalShaderPermutationParameters& Parameters) 
+	{
+		return true;
+	}
+
+	BEGIN_SHADER_PARAMETER_STRUCT(FParameters, )
+		SHADER_PARAMETER_TEXTURE(Texture2D, SrcTexture)
+		SHADER_PARAMETER_SAMPLER(SamplerState, SrcSampler)
+	END_SHADER_PARAMETER_STRUCT()
+};
diff --git a/Engine/Source/Runtime/Renderer/Private/PostProcess/PostProcessMaterial.cpp b/Engine/Source/Runtime/Renderer/Private/PostProcess/PostProcessMaterial.cpp
--- a/Engine/Source/Runtime/Renderer/Private/PostProcess/PostProcessMaterial.cpp
+++ b/Engine/Source/Runtime/Renderer/Private/PostProcess/PostProcessMaterial.cpp
@@ -16,6 +16,7 @@
 #include "ClearQuad.h"
 #include "Materials/MaterialExpressionSceneTexture.h"
 #include "RHI/Public/PipelineStateCache.h"
+#include "PostProcessing.h"
 #include "PostProcessMobile.h"
 
 static TAutoConsoleVariable<int32> CVarPostProcessAllowStencilTest(
@@ -393,12 +394,77 @@ void FRCPassPostProcessMaterial::Process(FRenderingCompositePassContext& Context
 		{
 			DestRenderTarget = &PassOutputs[0].RequestSurface(Context);
 
+			FTextureRHIParamRef DstTexture = DestRenderTarget->TargetableTexture;
+			FTextureRHIParamRef SrcTexture = GetInput(ePId_Input0)->GetOutput()->RequestSurface(Context).ShaderResourceTexture;
+
+			// CopyResource() can only be called when format and size match. Otherwise must use shader to do stretch & format conversion.
+			if (DstTexture->GetFormat() == SrcTexture->GetFormat() && DstTexture->GetSizeXYZ() == SrcTexture->GetSizeXYZ())
+			{
 			Context.RHICmdList.CopyTexture(
-				GetInput(ePId_Input0)->GetOutput()->RequestSurface(Context).ShaderResourceTexture,
-				DestRenderTarget->TargetableTexture,
+					SrcTexture,
+					DstTexture,
 				FRHICopyTextureInfo()
 				);
 		}
+			else
+			{
+				TShaderMap<FGlobalShaderType>* GlobalShaderMap = GetGlobalShaderMap(Context.View.FeatureLevel);
+
+				FIntRect DstRect = Context.GetSceneColorDestRect(*DestRenderTarget);
+				FIntRect SrcRect = Context.SceneColorViewRect;
+
+				FRHIRenderPassInfo RPInfo = FRHIRenderPassInfo(
+					DstTexture,
+					MakeRenderTargetActions(Context.GetLoadActionForRenderTarget(*DestRenderTarget), ERenderTargetStoreAction::EStore)
+					);
+
+				Context.RHICmdList.BeginRenderPass(RPInfo, TEXT("Copy Rect"));
+				{
+					Context.SetViewportAndCallRHI(DstRect);
+
+					FGraphicsPipelineStateInitializer GraphicsPSOInit;
+					Context.RHICmdList.ApplyCachedRenderTargets(GraphicsPSOInit);
+
+					GraphicsPSOInit.BlendState = TStaticBlendState<>::GetRHI();
+					GraphicsPSOInit.RasterizerState = TStaticRasterizerState<>::GetRHI();
+					GraphicsPSOInit.DepthStencilState = TStaticDepthStencilState<false, CF_Always>::GetRHI();
+
+					auto VS = GlobalShaderMap->GetShader<FPostProcessVS>();
+					auto PS = GlobalShaderMap->GetShader<FCopyRectPS>();
+
+					GraphicsPSOInit.BoundShaderState.VertexDeclarationRHI = GFilterVertexDeclaration.VertexDeclarationRHI;
+					GraphicsPSOInit.BoundShaderState.VertexShaderRHI = GETSAFERHISHADER_VERTEX(VS);
+					GraphicsPSOInit.BoundShaderState.PixelShaderRHI = GETSAFERHISHADER_PIXEL(PS);
+					GraphicsPSOInit.PrimitiveType = PT_TriangleList;
+
+					SetGraphicsPipelineState(Context.RHICmdList, GraphicsPSOInit);
+
+					FCopyRectPS::FParameters PixelParams;
+					PixelParams.SrcTexture = SrcTexture;
+					PixelParams.SrcSampler = TStaticSamplerState<>::GetRHI();
+
+					SetShaderParameters(Context.RHICmdList, PS, PS->GetPixelShader(), PixelParams);
+
+					DrawRectangle(
+						Context.RHICmdList,
+						DstRect.Min.X,
+						DstRect.Min.Y,
+						DstRect.Width(),
+						DstRect.Height(),
+						SrcRect.Min.X,
+						SrcRect.Min.Y,
+						SrcRect.Width(),
+						SrcRect.Height(),
+						DstTexture->GetTexture2D()->GetSizeXY(),
+						SrcTexture->GetTexture2D()->GetSizeXY(),
+						VS,
+						EDRF_UseTriangleOptimization
+						);
+				}
+				Context.RHICmdList.EndRenderPass();
+				Context.RHICmdList.CopyToResolveTarget(DestRenderTarget->TargetableTexture, DestRenderTarget->ShaderResourceTexture, FResolveParams());
+			}
+		}
 
 		DestRenderTargetLoadAction = ERenderTargetLoadAction::ELoad;
 	}
