Fix inverted rendering issues on GLES devices with 'after tonemapper' post process materials when viewrect != scenecolor target size.

diff --git a/Engine/Shaders/Private/MaterialTemplate.ush b/Engine/Shaders/Private/MaterialTemplate.ush
--- a/Engine/Shaders/Private/MaterialTemplate.ush
+++ b/Engine/Shaders/Private/MaterialTemplate.ush
@@ -1394,11 +1394,6 @@ MaterialFloat4 MobileSceneTextureLookup(inout FMaterialPixelParameters Parameter
 	else if (SceneTextureId == 14)
 	{
 #if POST_PROCESS_MATERIAL	
-		// SceneColor was transformed to OpenGL coordinates during tonemapping
-		#if !POST_PROCESS_MATERIAL_BEFORE_TONEMAP && (COMPILER_GLSL_ES2 || COMPILER_GLSL_ES3_1)
-		UV.y = 1.0 - UV.y; 
-		#endif
-
 		MaterialFloat4 Input0 = Texture2DSample(PostprocessInput0, PostprocessInput0Sampler, UV);
 		#if POST_PROCESS_MATERIAL_BEFORE_TONEMAP
 		Input0 = Decode32BPPHDR(Input0);
diff --git a/Engine/Shaders/Private/PostProcessMaterialShaders.usf b/Engine/Shaders/Private/PostProcessMaterialShaders.usf
--- a/Engine/Shaders/Private/PostProcessMaterialShaders.usf
+++ b/Engine/Shaders/Private/PostProcessMaterialShaders.usf
@@ -194,15 +194,11 @@ void MainVS_ES2(
 	out float2 OutViewportUV : TEXCOORD8
 	)
 {
-	// flip UV.y if we rendering to backbuffer in OpenGL space
-#if !POST_PROCESS_MATERIAL_BEFORE_TONEMAP && (COMPILER_GLSL_ES2 || COMPILER_GLSL_ES3_1)
-	InTexCoord.y = 1.0 - InTexCoord.y;
-#endif
-	
 	Output = (FPostProcessMaterialVSToPS)0;
 	float2 OutUV;
 	DrawRectangle(InPosition, InTexCoord, Output.Position, OutUV);
 
+	OutUV = BufferUVToViewportUV(OutUV);
 	OutViewportUV = OutUV;
 	
 #if NUM_TEX_COORD_INTERPOLATORS
@@ -246,14 +242,10 @@ void MainPS_ES2(
 	Parameters.VertexColor = 1;
 	
 	float4 SvPosition = Input.Position;
-	float4 SvPositionMP = SvPosition;
-	#if !POST_PROCESS_MATERIAL_BEFORE_TONEMAP && (COMPILER_GLSL_ES2 || COMPILER_GLSL_ES3_1)
-		SvPositionMP.y = View.ViewSizeAndInvSize.y - SvPositionMP.y; 
-	#endif
-	SvPositionMP.z = LookupDeviceZ(InViewportUV);
-	SvPositionMP.z = max(SvPositionMP.z, 1e-18);
+	SvPosition.z = LookupDeviceZ(InViewportUV);
+	SvPosition.z = max(SvPosition.z, 1e-18);
 	// fill out other related material parameters
-	CalcMaterialParametersPost(Parameters, PixelMaterialInputs, SvPositionMP, true);
+	CalcMaterialParametersPost(Parameters, PixelMaterialInputs, SvPosition, true);
 
 	// Grab emissive colour as output
 	half3 EmissiveColor = GetMaterialEmissive(PixelMaterialInputs);
diff --git a/Engine/Source/Runtime/Renderer/Private/PostProcess/PostProcessMaterial.cpp b/Engine/Source/Runtime/Renderer/Private/PostProcess/PostProcessMaterial.cpp
--- a/Engine/Source/Runtime/Renderer/Private/PostProcess/PostProcessMaterial.cpp
+++ b/Engine/Source/Runtime/Renderer/Private/PostProcess/PostProcessMaterial.cpp
@@ -387,6 +387,10 @@ void FRCPassPostProcessMaterial::Process(FRenderingCompositePassContext& Context
 	FIntRect DestRect = Context.GetSceneColorDestRect(*DestRenderTarget);
 	checkf(DestRect.Size() == SrcRect.Size(), TEXT("Post process material should not be used as upscaling pass."));
 	
+	FIntPoint DestPos(0, 0);
+	FIntPoint DestSize = DestRect.Size();
+	EDrawRectangleFlags DrawRectangleFlags = EDRF_UseTriangleOptimization;
+
 	FRHIRenderPassInfo RPInfo;
 	if (CustomDepthStencilTarget)
 	{
@@ -430,6 +434,20 @@ void FRCPassPostProcessMaterial::Process(FRenderingCompositePassContext& Context
 		const bool bViewSizeMatchesBufferSize = (View.ViewRect == Context.SceneColorViewRect && View.ViewRect.Size() == SrcSize && View.ViewRect.Min == FIntPoint::ZeroValue);
 		if (FeatureLevel <= ERHIFeatureLevel::ES3_1)
 		{
+			if (RHINeedsToSwitchVerticalAxis(GShaderPlatformForFeatureLevel[Context.GetFeatureLevel()]))
+			{
+				if (EBlendableLocation(Material->GetBlendableLocation()) == EBlendableLocation::BL_AfterTonemapping)
+				{
+					// flip dest rect y axis.
+					GraphicsPSOInit.RasterizerState = TStaticRasterizerState<FM_Solid, CM_None>::GetRHI();
+					int32 NewDestMaxY = DestSize.Y - DestRect.Min.Y;
+					DestPos.Y = DestRect.Max.Y;
+					DestSize.Y = -DestRect.Max.Y;
+					// triangle optimization currently doesn't work when flipped.
+					DrawRectangleFlags = EDRF_Default;
+				}
+			}
+
 			// use mobile's post process material.
 			if (bViewSizeMatchesBufferSize)
 			{
@@ -449,7 +467,7 @@ void FRCPassPostProcessMaterial::Process(FRenderingCompositePassContext& Context
 			GraphicsPSOInit.BoundShaderState.VertexDeclarationRHI = GPostProcessMaterialVertexDeclaration.VertexDeclarationRHI;
 			GraphicsPSOInit.BoundShaderState.VertexShaderRHI = GETSAFERHISHADER_VERTEX(VertexShader_HighEnd);
 			GraphicsPSOInit.BoundShaderState.PixelShaderRHI = GETSAFERHISHADER_PIXEL(PixelShader_HighEnd);
-
+			
 			SetGraphicsPipelineState(Context.RHICmdList, GraphicsPSOInit);
 			Context.RHICmdList.SetStencilRef(StencilRefValue);
 
@@ -477,8 +495,8 @@ void FRCPassPostProcessMaterial::Process(FRenderingCompositePassContext& Context
 
 		DrawPostProcessPass(
 			Context.RHICmdList,
-			0, 0,
-			DestRect.Width(), DestRect.Height(),
+			DestPos.X, DestPos.Y,
+			DestSize.X, DestSize.Y,
 			SrcRect.Min.X, SrcRect.Min.Y,
 			SrcRect.Width(), SrcRect.Height(),
 			DestRect.Size(),
@@ -486,7 +504,7 @@ void FRCPassPostProcessMaterial::Process(FRenderingCompositePassContext& Context
 			VertexShader,
 			View.StereoPass,
 			Context.HasHmdMesh(),
-			EDRF_UseTriangleOptimization);
+			DrawRectangleFlags);
 	}
 	Context.RHICmdList.EndRenderPass();
 	Context.RHICmdList.CopyToResolveTarget(DestRenderTarget->TargetableTexture, DestRenderTarget->ShaderResourceTexture, FResolveParams());
