Fix DoLazyStaticMeshUpdateCVarSinkFunction crash when GIsEditor value changes during engine startup while the actual r.DoLazyStaticMeshUpdate variable stays disabled

diff --git a/Engine/Source/Runtime/Renderer/Private/RendererScene.cpp b/Engine/Source/Runtime/Renderer/Private/RendererScene.cpp
--- a/Engine/Source/Runtime/Renderer/Private/RendererScene.cpp
+++ b/Engine/Source/Runtime/Renderer/Private/RendererScene.cpp
@@ -528,8 +528,13 @@ static TAutoConsoleVariable<int32> CVarDoLazyStaticMeshUpdate(
 
 static void DoLazyStaticMeshUpdateCVarSinkFunction()
 {
-	static bool CachedDoLazyStaticMeshUpdate = CVarDoLazyStaticMeshUpdate.GetValueOnGameThread() && !GIsEditor;
-	bool DoLazyStaticMeshUpdate = CVarDoLazyStaticMeshUpdate.GetValueOnGameThread() && !GIsEditor;
+	if (!GIsRunning || GIsEditor)
+	{
+		return;
+	}
+
+	static bool CachedDoLazyStaticMeshUpdate = CVarDoLazyStaticMeshUpdate.GetValueOnGameThread();
+	bool DoLazyStaticMeshUpdate = CVarDoLazyStaticMeshUpdate.GetValueOnGameThread();
 
 	if (DoLazyStaticMeshUpdate != CachedDoLazyStaticMeshUpdate)
 	{
