Back out revision 6 from //UE4/Dev-Core/Engine/Source/Runtime/Renderer/Private/RayTracing/RayTracingDynamicGeometryCollection.h

(fixing cis)

diff --git a/Engine/Source/Runtime/Renderer/Private/RayTracing/RayTracingDynamicGeometryCollection.h b/Engine/Source/Runtime/Renderer/Private/RayTracing/RayTracingDynamicGeometryCollection.h
new file mode 100644
--- /dev/null
+++ b/Engine/Source/Runtime/Renderer/Private/RayTracing/RayTracingDynamicGeometryCollection.h
@@ -0,0 +1,20 @@
+// Copyright 1998-2019 Epic Games, Inc. All Rights Reserved.
+
+#pragma once
+
+#if RHI_RAYTRACING
+
+class RENDERER_API FRayTracingDynamicGeometryCollection
+{
+public:
+	FRayTracingDynamicGeometryCollection();
+
+	void AddDynamicMeshBatchForGeometryUpdate(const FScene* Scene, const FSceneView* View, const FPrimitiveSceneProxy* PrimitiveSceneProxy, const FMeshBatch& MeshBatch, FRayTracingGeometry& Geometry, uint32 NumMaxVertices, FRWBuffer& Buffer);
+	void DispatchUpdates(FRHICommandListImmediate& RHICmdList);
+	void Clear();
+
+private:
+	TUniquePtr<TArray<struct FMeshComputeDispatchCommand>> DispatchCommands;
+};
+
+#endif // RHI_RAYTRACING
