#UE4 Fix for crash when using nullrhi with r.DoLazyStaticMeshUpdate

diff --git a/Engine/Source/Runtime/Renderer/Private/RendererScene.cpp b/Engine/Source/Runtime/Renderer/Private/RendererScene.cpp
--- a/Engine/Source/Runtime/Renderer/Private/RendererScene.cpp
+++ b/Engine/Source/Runtime/Renderer/Private/RendererScene.cpp
@@ -553,7 +553,7 @@ static TAutoConsoleVariable<int32> CVarDoLazyStaticMeshUpdate(
 
 static void DoLazyStaticMeshUpdateCVarSinkFunction()
 {
-	if (!GIsRunning || GIsEditor)
+	if (!GIsRunning || GIsEditor || !FApp::CanEverRender())
 	{
 		return;
 	}
@@ -584,7 +584,7 @@ static FAutoConsoleVariableSink CVarDoLazyStaticMeshUpdateSink(FConsoleCommandDe
 
 void FScene::UpdateDoLazyStaticMeshUpdate(FRHICommandListImmediate& CmdList)
 {
-	bool DoLazyStaticMeshUpdate = CVarDoLazyStaticMeshUpdate.GetValueOnRenderThread() && !GIsEditor;
+	bool DoLazyStaticMeshUpdate = CVarDoLazyStaticMeshUpdate.GetValueOnRenderThread() && !GIsEditor && FApp::CanEverRender();
 
 	for (int32 PrimitiveIndex = 0; PrimitiveIndex < Primitives.Num(); PrimitiveIndex++)
 	{
