RL - Move RectLight MipTree generation from the Proxy to a livetime of the Component and its latest proxy. This avoids the rebuild when other light properties change and can consitently save up to 1.5ms of GPU time.

diff --git a/Engine/Source/Runtime/Engine/Classes/Components/RectLightComponent.h b/Engine/Source/Runtime/Engine/Classes/Components/RectLightComponent.h
--- a/Engine/Source/Runtime/Engine/Classes/Components/RectLightComponent.h
+++ b/Engine/Source/Runtime/Engine/Classes/Components/RectLightComponent.h
@@ -48,9 +48,12 @@ class ENGINE_API URectLightComponent : public ULocalLightComponent
 	float BarnDoorLength;
 
 	/** Texture mapped to the light source rectangle */
-	UPROPERTY(EditAnywhere, BlueprintReadWrite, Category=Light)
+	UPROPERTY(EditAnywhere, BlueprintReadOnly, Category = Light)
 	class UTexture* SourceTexture;
 
+	UFUNCTION(BlueprintCallable, Category = "Rendering|Lighting")
+	void SetSourceTexture(UTexture* bNewValue);
+
 	UFUNCTION(BlueprintCallable, Category="Rendering|Lighting")
 	void SetSourceWidth(float bNewValue);
 
@@ -75,9 +78,15 @@ public:
 	virtual float GetUniformPenumbraSize() const override;
 	virtual FLightSceneProxy* CreateSceneProxy() const override;
 
+	virtual void BeginDestroy() override;
 	//~ Begin UObject Interface
 #if WITH_EDITOR
 	virtual void PostEditChangeProperty(FPropertyChangedEvent& PropertyChangedEvent) override;
 #endif // WITH_EDITOR
 	//~ End UObject Interface
+
+private:
+	void UpdateRayTracingData();
+	friend class FRectLightSceneProxy;
+	struct FRectLightRayTracingData* RayTracingData;
 };
\ No newline at end of file
diff --git a/Engine/Source/Runtime/Engine/Private/Components/RectLightComponent.cpp b/Engine/Source/Runtime/Engine/Private/Components/RectLightComponent.cpp
--- a/Engine/Source/Runtime/Engine/Private/Components/RectLightComponent.cpp
+++ b/Engine/Source/Runtime/Engine/Private/Components/RectLightComponent.cpp
@@ -23,8 +23,29 @@ float GetRectLightBarnDoorMaxAngle()
 	return 88.f;
 }
 
+#if RHI_RAYTRACING
+FRectLightRayTracingData BuildRectLightMipTree(FRHICommandListImmediate& RHICmdList, UTexture* SourceTexture);
+#endif
+
+void URectLightComponent::UpdateRayTracingData()
+{
+#if RHI_RAYTRACING
+	if (IsRayTracingEnabled())
+	{
+		UTexture* NewSourceTexture = SourceTexture;
+		FRectLightRayTracingData* UpdatedRayTracingData = RayTracingData;
+		ENQUEUE_RENDER_COMMAND(BuildRectLightMipTree)(
+			[NewSourceTexture, UpdatedRayTracingData](FRHICommandListImmediate& RHICmdList) mutable
+		{
+			*UpdatedRayTracingData = BuildRectLightMipTree(RHICmdList, NewSourceTexture);
+		});
+	}
+#endif
+}
+
 URectLightComponent::URectLightComponent(const FObjectInitializer& ObjectInitializer)
 	: Super(ObjectInitializer)
+	, RayTracingData(new FRectLightRayTracingData())
 {
 #if WITH_EDITORONLY_DATA
 	if (!IsRunningCommandlet())
@@ -44,6 +65,8 @@ URectLightComponent::URectLightComponent(const FObjectInitializer& ObjectInitial
 	SourceTexture = nullptr;
 	BarnDoorAngle = GetRectLightBarnDoorMaxAngle();
 	BarnDoorLength = 20.0f;
+
+	UpdateRayTracingData();
 }
 
 FLightSceneProxy* URectLightComponent::CreateSceneProxy() const
@@ -51,6 +74,18 @@ FLightSceneProxy* URectLightComponent::CreateSceneProxy() const
 	return new FRectLightSceneProxy(this);
 }
 
+void URectLightComponent::SetSourceTexture(UTexture* NewValue)
+{
+	if (AreDynamicDataChangesAllowed()
+		&& SourceTexture != NewValue)
+	{
+		SourceTexture = NewValue;
+
+		UpdateRayTracingData();
+		MarkRenderStateDirty();
+	}
+}
+
 void URectLightComponent::SetSourceWidth(float NewValue)
 {
 	if (AreDynamicDataChangesAllowed()
@@ -153,6 +188,18 @@ float URectLightComponent::GetUniformPenumbraSize() const
 	}
 }
 
+void URectLightComponent::BeginDestroy()
+{
+	FRectLightRayTracingData* DeletedRenderData = RayTracingData;
+	RayTracingData = nullptr;
+	ENQUEUE_RENDER_COMMAND(DeleteBuildRectLightMipTree)(
+		[DeletedRenderData](FRHICommandListImmediate& RHICmdList)
+	{
+		delete DeletedRenderData;
+	});
+	Super::BeginDestroy();
+}
+
 #if WITH_EDITOR
 /**
  * Called after property has changed via e.g. property window or set command.
@@ -174,18 +221,9 @@ FRectLightSceneProxy::FRectLightSceneProxy(const URectLightComponent* Component)
 	, SourceHeight(Component->SourceHeight)
 	, BarnDoorAngle(FMath::Clamp(Component->BarnDoorAngle, 0.f, GetRectLightBarnDoorMaxAngle()))
 	, BarnDoorLength(FMath::Max(0.1f, Component->BarnDoorLength))
+	, RayTracingData(Component->RayTracingData)
 	, SourceTexture(Component->SourceTexture)
 {
-#if RHI_RAYTRACING
-	if (IsRayTracingEnabled())
-	{
-		ENQUEUE_RENDER_COMMAND(BuildRectLightMipTree)(
-			[this](FRHICommandListImmediate& RHICmdList)
-		{
-			BuildRectLightMipTree(RHICmdList);
-		});
-	}
-#endif
 }
 
 FRectLightSceneProxy::~FRectLightSceneProxy() {}
@@ -337,41 +375,46 @@ private:
 
 IMPLEMENT_SHADER_TYPE(, FBuildRectLightMipTreeCS, TEXT("/Engine/Private/Raytracing/BuildMipTreeCS.usf"), TEXT("BuildRectLightMipTreeCS"), SF_Compute)
 
-void FRectLightSceneProxy::BuildRectLightMipTree(FRHICommandListImmediate& RHICmdList)
+DECLARE_GPU_STAT_NAMED(BuildRectLightMipTreeStat, TEXT("build RectLight MipTree"));
+
+FRectLightRayTracingData BuildRectLightMipTree(FRHICommandListImmediate& RHICmdList, UTexture* SourceTexture)
 {
+	SCOPED_GPU_STAT(RHICmdList, BuildRectLightMipTreeStat);
+
 	check(IsInRenderingThread());
+	FRectLightRayTracingData Data;
+	FTextureRHIRef RhiTexture = SourceTexture ? SourceTexture->Resource->TextureRHI : GWhiteTexture->TextureRHI;
 
 	const auto ShaderMap = GetGlobalShaderMap(ERHIFeatureLevel::SM5);
 	TShaderMapRef<FBuildRectLightMipTreeCS> BuildRectLightMipTreeComputeShader(ShaderMap);
 	RHICmdList.SetComputeShader(BuildRectLightMipTreeComputeShader->GetComputeShader());
 
 	// Allocate MIP tree
-	FLightShaderParameters LightParameters;
-	GetLightShaderParameters(LightParameters);
-	FIntVector TextureSize = LightParameters.SourceTexture->GetSizeXYZ();
+	FIntVector TextureSize = RhiTexture->GetSizeXYZ();
 	uint32 MipLevelCount = FMath::Min(FMath::CeilLogTwo(TextureSize.X), FMath::CeilLogTwo(TextureSize.Y));
-	RectLightMipTreeDimensions = FIntVector(1 << MipLevelCount, 1 << MipLevelCount, 1);
-	uint32 NumElements = RectLightMipTreeDimensions.X * RectLightMipTreeDimensions.Y;
+	Data.RectLightMipTreeDimensions = FIntVector(1 << MipLevelCount, 1 << MipLevelCount, 1);
+	uint32 NumElements = Data.RectLightMipTreeDimensions.X * Data.RectLightMipTreeDimensions.Y;
 	for (uint32 MipLevel = 1; MipLevel <= MipLevelCount; ++MipLevel)
 	{
-		uint32 NumElementsInLevel = (RectLightMipTreeDimensions.X >> MipLevel) * (RectLightMipTreeDimensions.Y >> MipLevel);
+		uint32 NumElementsInLevel = (Data.RectLightMipTreeDimensions.X >> MipLevel) * (Data.RectLightMipTreeDimensions.Y >> MipLevel);
 		NumElements += NumElementsInLevel;
 	}
 
-	RectLightMipTree.Initialize(sizeof(float), NumElements, PF_R32_FLOAT, BUF_UnorderedAccess | BUF_ShaderResource);
+	Data.RectLightMipTree.Initialize(sizeof(float), NumElements, PF_R32_FLOAT, BUF_UnorderedAccess | BUF_ShaderResource);
 
 	// Execute hierarchical build
 	for (uint32 MipLevel = 0; MipLevel <= MipLevelCount; ++MipLevel)
 	{
 		FComputeFenceRHIRef MipLevelFence = RHICmdList.CreateComputeFence(TEXT("RectLightMipTree Build"));
-		BuildRectLightMipTreeComputeShader->SetParameters(RHICmdList, LightParameters.SourceTexture, RectLightMipTreeDimensions, MipLevel, RectLightMipTree);
-		FIntVector MipLevelDimensions = FIntVector(RectLightMipTreeDimensions.X >> MipLevel, RectLightMipTreeDimensions.Y >> MipLevel, 1);
+		BuildRectLightMipTreeComputeShader->SetParameters(RHICmdList, RhiTexture, Data.RectLightMipTreeDimensions, MipLevel, Data.RectLightMipTree);
+		FIntVector MipLevelDimensions = FIntVector(Data.RectLightMipTreeDimensions.X >> MipLevel, Data.RectLightMipTreeDimensions.Y >> MipLevel, 1);
 		FIntVector NumGroups = FIntVector::DivideAndRoundUp(MipLevelDimensions, FBuildRectLightMipTreeCS::GetGroupSize());
 		DispatchComputeShader(RHICmdList, *BuildRectLightMipTreeComputeShader, NumGroups.X, NumGroups.Y, 1);
-		BuildRectLightMipTreeComputeShader->UnsetParameters(RHICmdList, EResourceTransitionAccess::ERWBarrier, EResourceTransitionPipeline::EComputeToCompute, RectLightMipTree, MipLevelFence);
+		BuildRectLightMipTreeComputeShader->UnsetParameters(RHICmdList, EResourceTransitionAccess::ERWBarrier, EResourceTransitionPipeline::EComputeToCompute, Data.RectLightMipTree, MipLevelFence);
 	}
 	FComputeFenceRHIRef TransitionFence = RHICmdList.CreateComputeFence(TEXT("RectLightMipTree Transition"));
-	BuildRectLightMipTreeComputeShader->UnsetParameters(RHICmdList, EResourceTransitionAccess::ERWBarrier, EResourceTransitionPipeline::EComputeToCompute, RectLightMipTree, TransitionFence);
-}
+	BuildRectLightMipTreeComputeShader->UnsetParameters(RHICmdList, EResourceTransitionAccess::ERWBarrier, EResourceTransitionPipeline::EComputeToCompute, Data.RectLightMipTree, TransitionFence);
 
+	return Data;
+}
 #endif // RHI_RAYTRACING
diff --git a/Engine/Source/Runtime/Engine/Public/RectLightSceneProxy.h b/Engine/Source/Runtime/Engine/Public/RectLightSceneProxy.h
--- a/Engine/Source/Runtime/Engine/Public/RectLightSceneProxy.h
+++ b/Engine/Source/Runtime/Engine/Public/RectLightSceneProxy.h
@@ -11,6 +11,14 @@
 #include "LocalLightSceneProxy.h"
 #include "SceneManagement.h"
 
+struct FRectLightRayTracingData
+{
+#if RHI_RAYTRACING
+	FRWBuffer   RectLightMipTree;
+	FIntVector  RectLightMipTreeDimensions;
+#endif
+};
+
 class FRectLightSceneProxy : public FLocalLightSceneProxy
 {
 public:
@@ -18,11 +26,8 @@ public:
 	float		SourceHeight;
 	float		BarnDoorAngle;
 	float		BarnDoorLength;
+	FRectLightRayTracingData* RayTracingData;
 	UTexture*	SourceTexture;
-#if RHI_RAYTRACING
-	FRWBuffer   RectLightMipTree;
-	FIntVector  RectLightMipTreeDimensions;
-#endif
 
 	FRectLightSceneProxy(const URectLightComponent* Component);
 	virtual ~FRectLightSceneProxy();
@@ -39,8 +44,4 @@ public:
 	* @return True if the whole-scene projected shadow should be used.
 	*/
 	virtual bool GetWholeSceneProjectedShadowInitializer(const FSceneViewFamily& ViewFamily, TArray<FWholeSceneProjectedShadowInitializer, TInlineAllocator<6> >& OutInitializers) const;
-
-#if RHI_RAYTRACING
-	void BuildRectLightMipTree(FRHICommandListImmediate& RHICmdList);
-#endif
 };
\ No newline at end of file
diff --git a/Engine/Source/Runtime/Renderer/Private/RayTracing/RayTracingShadowsForRectLight.cpp b/Engine/Source/Runtime/Renderer/Private/RayTracing/RayTracingShadowsForRectLight.cpp
new file mode 100644
--- /dev/null
+++ b/Engine/Source/Runtime/Renderer/Private/RayTracing/RayTracingShadowsForRectLight.cpp
@@ -0,0 +1,440 @@
+// Copyright 1998-2019 Epic Games, Inc. All Rights Reserved.
+
+#include "DeferredShadingRenderer.h"
+
+#if RHI_RAYTRACING
+
+#include "ClearQuad.h"
+#include "SceneRendering.h"
+#include "SceneRenderTargets.h"
+#include "RenderTargetPool.h"
+#include "RHIResources.h"
+#include "UniformBuffer.h"
+#include "VisualizeTexture.h"
+
+#include "PostProcess/PostProcessing.h"
+#include "PostProcess/SceneFilterRendering.h"
+#include "RectLightSceneProxy.h"
+#include "RaytracingOptions.h"
+
+#include "RHI/Public/PipelineStateCache.h"
+
+static int32 GRayTracingStochasticRectLight = 0;
+static FAutoConsoleVariableRef CVarRayTracingStochasticRectLight(
+	TEXT("r.RayTracing.StochasticRectLight"),
+	GRayTracingStochasticRectLight,
+	TEXT("0: use analytical evaluation (default)\n")
+	TEXT("1: use stochastic evaluation\n")
+);
+
+static int32 GRayTracingStochasticRectLightSamplesPerPixel = 1;
+static FAutoConsoleVariableRef CVarRayTracingRecLightStochasticSamplesPerPixel(
+	TEXT("r.RayTracing.StochasticRectLight.SamplesPerPixel"),
+	GRayTracingStochasticRectLightSamplesPerPixel,
+	TEXT("Sets the samples-per-pixel for rect light evaluation (default = 1)")
+);
+
+static int32 GRayTracingStochasticRectLightIsTextureImportanceSampling = 1;
+static FAutoConsoleVariableRef CVarRayTracingStochasticRecLightIsTextureImportanceSampling(
+	TEXT("r.RayTracing.StochasticRectLight.IsTextureImportanceSampling"),
+	GRayTracingStochasticRectLightIsTextureImportanceSampling,
+	TEXT("Enable importance sampling for rect light evaluation (default = 1)")
+);
+
+bool ShouldRenderRayTracingStochasticRectLight(const FLightSceneInfo& LightSceneInfo)
+{
+	return IsRayTracingEnabled() && GRayTracingStochasticRectLight == 1
+		&& LightSceneInfo.Proxy->CastsRaytracedShadow()
+		&& LightSceneInfo.Proxy->GetLightType() == LightType_Rect;
+}
+
+BEGIN_GLOBAL_SHADER_PARAMETER_STRUCT(FRectLightData, )
+// Pass settings
+SHADER_PARAMETER(int, SamplesPerPixel)
+SHADER_PARAMETER(int, bIsTextureImportanceSampling)
+// Light data
+SHADER_PARAMETER(FVector, Position)
+SHADER_PARAMETER(FVector, Normal)
+SHADER_PARAMETER(FVector, dPdu)
+SHADER_PARAMETER(FVector, dPdv)
+SHADER_PARAMETER(FVector, Color)
+SHADER_PARAMETER(float, Width)
+SHADER_PARAMETER(float, Height)
+SHADER_PARAMETER(FIntVector, MipTreeDimensions)
+SHADER_PARAMETER(float, MaxNormalBias)
+SHADER_PARAMETER(float, BarnCosAngle)
+SHADER_PARAMETER(float, BarnLength)
+SHADER_PARAMETER_TEXTURE(Texture2D, Texture)
+SHADER_PARAMETER_SAMPLER(SamplerState, TextureSampler)
+// Sampling data
+SHADER_PARAMETER_SRV(Buffer<float>, MipTree)
+END_GLOBAL_SHADER_PARAMETER_STRUCT()
+
+DECLARE_GPU_STAT_NAMED(RayTracingRectLight, TEXT("Ray Tracing RectLight"));
+
+IMPLEMENT_GLOBAL_SHADER_PARAMETER_STRUCT(FRectLightData, "RectLight");
+
+template <int TextureImportanceSampling>
+class FRectLightRGS : public FGlobalShader
+{
+	DECLARE_SHADER_TYPE(FRectLightRGS, Global)
+
+	static void ModifyCompilationEnvironment(const FGlobalShaderPermutationParameters& Parameters, FShaderCompilerEnvironment& OutEnvironment)
+	{
+		OutEnvironment.SetDefine(TEXT("TEXTURE_IMPORTANCE_SAMPLING"), TextureImportanceSampling);
+	}
+
+public:
+	static bool ShouldCompilePermutation(const FGlobalShaderPermutationParameters& Parameters)
+	{
+		return ShouldCompileRayTracingShadersForProject(Parameters.Platform);
+	}
+
+	FRectLightRGS() {}
+	virtual ~FRectLightRGS() {}
+
+	FRectLightRGS(const ShaderMetaType::CompiledShaderInitializerType& Initializer)
+		: FGlobalShader(Initializer)
+	{
+		ViewParameter.Bind(Initializer.ParameterMap, TEXT("View"));
+		SceneTexturesParameter.Bind(Initializer.ParameterMap, TEXT("SceneTexturesStruct"));
+		RectLightParameter.Bind(Initializer.ParameterMap, TEXT("RectLight"));
+		TLASParameter.Bind(Initializer.ParameterMap, TEXT("TLAS"));
+		TransmissionProfilesTextureParameter.Bind(Initializer.ParameterMap, TEXT("SSProfilesTexture"));
+		TransmissionProfilesLinearSamplerParameter.Bind(Initializer.ParameterMap, TEXT("TransmissionProfilesLinearSampler"));
+
+		LuminanceUAVParameter.Bind(Initializer.ParameterMap, TEXT("RWLuminanceUAV"));
+		RayDistanceUAVParameter.Bind(Initializer.ParameterMap, TEXT("RWRayDistanceUAV"));
+	}
+
+	bool Serialize(FArchive& Ar)
+	{
+		bool bShaderHasOutdatedParameters = FGlobalShader::Serialize(Ar);
+		Ar << ViewParameter;
+		Ar << SceneTexturesParameter;
+		Ar << RectLightParameter;
+		Ar << TLASParameter;
+		Ar << TransmissionProfilesTextureParameter;
+		Ar << TransmissionProfilesLinearSamplerParameter;
+		Ar << LuminanceUAVParameter;
+		Ar << RayDistanceUAVParameter;
+		return bShaderHasOutdatedParameters;
+	}
+
+	void Dispatch(
+		FRHICommandListImmediate& RHICmdList,
+		const FRayTracingScene& RayTracingScene,
+		FUniformBufferRHIParamRef ViewUniformBuffer,
+		FUniformBufferRHIParamRef SceneTexturesUniformBuffer,
+		FUniformBufferRHIParamRef RectLightUniformBuffer,
+		FUnorderedAccessViewRHIParamRef LuminanceUAV,
+		FUnorderedAccessViewRHIParamRef RayDistanceUAV,
+		uint32 Width, uint32 Height
+	)
+	{
+		FRayTracingPipelineStateInitializer Initializer;
+
+		FRayTracingShaderRHIParamRef RayGenShaderTable[] = { GetRayTracingShader() };
+		Initializer.SetRayGenShaderTable(RayGenShaderTable);
+
+		FRHIRayTracingPipelineState* Pipeline = PipelineStateCache::GetAndOrCreateRayTracingPipelineState(Initializer); // #dxr_todo: this should be done once at load-time and cached
+
+		FRayTracingShaderBindingsWriter GlobalResources;
+		GlobalResources.Set(TLASParameter, RayTracingScene.RayTracingSceneRHI->GetShaderResourceView());
+		GlobalResources.Set(ViewParameter, ViewUniformBuffer);
+		GlobalResources.Set(SceneTexturesParameter, SceneTexturesUniformBuffer);
+		GlobalResources.Set(RectLightParameter, RectLightUniformBuffer);
+		GlobalResources.Set(LuminanceUAVParameter, LuminanceUAV);
+		GlobalResources.Set(RayDistanceUAVParameter, RayDistanceUAV);
+
+		if (TransmissionProfilesTextureParameter.IsBound())
+		{
+			FSceneRenderTargets& SceneContext = FSceneRenderTargets::Get(RHICmdList);
+			const IPooledRenderTarget* PooledRT = GetSubsufaceProfileTexture_RT((FRHICommandListImmediate&)RHICmdList);
+
+			if (!PooledRT)
+			{
+				// no subsurface profile was used yet
+				PooledRT = GSystemTextures.BlackDummy;
+			}
+			const FSceneRenderTargetItem& Item = PooledRT->GetRenderTargetItem();
+
+			GlobalResources.SetTexture(TransmissionProfilesTextureParameter.GetBaseIndex(), Item.ShaderResourceTexture);
+			GlobalResources.SetSampler(TransmissionProfilesLinearSamplerParameter.GetBaseIndex(), TStaticSamplerState<SF_Bilinear, AM_Clamp, AM_Clamp, AM_Clamp>::GetRHI());
+		}
+
+		RHICmdList.RayTraceDispatch(Pipeline, GetRayTracingShader(), RayTracingScene.RayTracingSceneRHI, GlobalResources, Width, Height);
+	}
+
+private:
+	// Input
+	FShaderResourceParameter TLASParameter;
+	FShaderUniformBufferParameter ViewParameter;
+	FShaderUniformBufferParameter SceneTexturesParameter;
+	FShaderUniformBufferParameter RectLightParameter;
+
+	// SSS Profile
+	FShaderResourceParameter TransmissionProfilesTextureParameter;
+	FShaderResourceParameter TransmissionProfilesLinearSamplerParameter;
+
+	// Output
+	FShaderResourceParameter LuminanceUAVParameter;
+	FShaderResourceParameter RayDistanceUAVParameter;
+};
+
+#define IMPLEMENT_RECT_LIGHT_TYPE(TextureImportanceSampling) \
+	typedef FRectLightRGS<TextureImportanceSampling> FRectLightRGS##TextureImportanceSampling; \
+	IMPLEMENT_SHADER_TYPE(template<>, FRectLightRGS##TextureImportanceSampling, TEXT("/Engine/Private/RayTracing/RayTracingRectLightRGS.usf"), TEXT("RectLightRGS"), SF_RayGen);
+
+IMPLEMENT_RECT_LIGHT_TYPE(0);
+IMPLEMENT_RECT_LIGHT_TYPE(1);
+
+class FVisualizeRectLightMipTreePS : public FGlobalShader
+{
+	DECLARE_SHADER_TYPE(FVisualizeRectLightMipTreePS, Global);
+
+public:
+	static bool ShouldCache(EShaderPlatform Platform)
+	{
+		return ShouldCompileRayTracingShadersForProject(Platform);
+	}
+
+	static bool ShouldCompilePermutation(const FGlobalShaderPermutationParameters& Parameters)
+	{
+		return ShouldCompileRayTracingShadersForProject(Parameters.Platform);
+	}
+
+	static void ModifyCompilationEnvironment(const FGlobalShaderPermutationParameters& Parameters, FShaderCompilerEnvironment& OutEnvironment)
+	{
+		FGlobalShader::ModifyCompilationEnvironment(Parameters, OutEnvironment);
+	}
+
+	FVisualizeRectLightMipTreePS(const ShaderMetaType::CompiledShaderInitializerType& Initializer)
+		: FGlobalShader(Initializer)
+	{
+		DimensionsParameter.Bind(Initializer.ParameterMap, TEXT("Dimensions"));
+		MipTreeParameter.Bind(Initializer.ParameterMap, TEXT("MipTree"));
+	}
+
+	FVisualizeRectLightMipTreePS() {}
+
+	template<typename TRHICommandList>
+	void SetParameters(
+		TRHICommandList& RHICmdList,
+		const FViewInfo& View,
+		const FRWBuffer& MipTree,
+		const FIntVector Dimensions)
+	{
+		const FPixelShaderRHIParamRef ShaderRHI = GetPixelShader();
+		FGlobalShader::SetParameters<FViewUniformShaderParameters>(RHICmdList, ShaderRHI, View.ViewUniformBuffer);
+
+		SetShaderValue(RHICmdList, ShaderRHI, DimensionsParameter, Dimensions);
+		SetSRVParameter(RHICmdList, ShaderRHI, MipTreeParameter, MipTree.SRV);
+	}
+
+	virtual bool Serialize(FArchive& Ar) override
+	{
+		bool bShaderHasOutdatedParameters = FGlobalShader::Serialize(Ar);
+		Ar << DimensionsParameter;
+		Ar << MipTreeParameter;
+		return bShaderHasOutdatedParameters;
+	}
+
+private:
+	FShaderParameter DimensionsParameter;
+	FShaderResourceParameter MipTreeParameter;
+};
+
+IMPLEMENT_SHADER_TYPE(, FVisualizeRectLightMipTreePS, TEXT("/Engine/Private/PathTracing/VisualizeMipTreePixelShader.usf"), TEXT("VisualizeMipTreePS"), SF_Pixel)
+
+void FDeferredShadingSceneRenderer::VisualizeRectLightMipTree(
+	FRHICommandListImmediate& RHICmdList,
+	const FViewInfo& View,
+	const FRWBuffer& RectLightMipTree,
+	const FIntVector& RectLightMipTreeDimensions
+)
+{
+	// Allocate render target
+	FSceneRenderTargets& SceneContext = FSceneRenderTargets::Get(RHICmdList);
+	FPooledRenderTargetDesc Desc = SceneContext.GetSceneColor()->GetDesc();
+	Desc.Flags &= ~(TexCreate_FastVRAM | TexCreate_Transient);
+	TRefCountPtr<IPooledRenderTarget> RectLightMipTreeRT;
+	GRenderTargetPool.FindFreeElement(RHICmdList, Desc, RectLightMipTreeRT, TEXT("RectLightMipTreeRT"));
+
+	// Define shaders
+	const auto ShaderMap = GetGlobalShaderMap(View.FeatureLevel);
+	TShaderMapRef<FPostProcessVS> VertexShader(ShaderMap);
+	TShaderMapRef<FVisualizeRectLightMipTreePS> PixelShader(ShaderMap);
+	FTextureRHIParamRef RenderTargets[2] =
+	{
+		SceneContext.GetSceneColor()->GetRenderTargetItem().TargetableTexture,
+		RectLightMipTreeRT->GetRenderTargetItem().TargetableTexture
+	};
+	PRAGMA_DISABLE_DEPRECATION_WARNINGS
+	SetRenderTargets(RHICmdList, 2, RenderTargets, SceneContext.GetSceneDepthSurface(), ESimpleRenderTargetMode::EExistingColorAndDepth, FExclusiveDepthStencil::DepthRead_StencilNop);
+	PRAGMA_ENABLE_DEPRECATION_WARNINGS
+
+	// PSO definition
+	FGraphicsPipelineStateInitializer GraphicsPSOInit;
+	SceneContext.BeginRenderingSceneColor(RHICmdList, ESimpleRenderTargetMode::EExistingColorAndDepth, FExclusiveDepthStencil::DepthRead_StencilWrite, true);
+	RHICmdList.ApplyCachedRenderTargets(GraphicsPSOInit);
+	GraphicsPSOInit.BlendState = TStaticBlendState<CW_RGB, BO_Add, BF_One, BF_One>::GetRHI();
+	GraphicsPSOInit.RasterizerState = TStaticRasterizerState<FM_Solid, CM_None>::GetRHI();
+	GraphicsPSOInit.DepthStencilState = TStaticDepthStencilState<false, CF_Always>::GetRHI();
+	GraphicsPSOInit.BoundShaderState.VertexDeclarationRHI = GFilterVertexDeclaration.VertexDeclarationRHI;
+	GraphicsPSOInit.BoundShaderState.VertexShaderRHI = GETSAFERHISHADER_VERTEX(*VertexShader);
+	GraphicsPSOInit.BoundShaderState.PixelShaderRHI = GETSAFERHISHADER_PIXEL(*PixelShader);
+	GraphicsPSOInit.PrimitiveType = PT_TriangleList;
+	SetGraphicsPipelineState(RHICmdList, GraphicsPSOInit);
+
+	// Transition to graphics
+	RHICmdList.TransitionResource(EResourceTransitionAccess::EReadable, EResourceTransitionPipeline::EComputeToGfx, RectLightMipTree.UAV);
+
+	// Draw
+	RHICmdList.SetViewport(View.ViewRect.Min.X, View.ViewRect.Min.Y, 0.0f, View.ViewRect.Max.X, View.ViewRect.Max.Y, 1.0f);
+	PixelShader->SetParameters(RHICmdList, View, RectLightMipTree, RectLightMipTreeDimensions);
+	DrawRectangle(
+		RHICmdList,
+		0, 0,
+		View.ViewRect.Width(), View.ViewRect.Height(),
+		View.ViewRect.Min.X, View.ViewRect.Min.Y,
+		View.ViewRect.Width(), View.ViewRect.Height(),
+		FIntPoint(View.ViewRect.Width(), View.ViewRect.Height()),
+		SceneContext.GetBufferSizeXY(),
+		*VertexShader);
+	ResolveSceneColor(RHICmdList);
+	GVisualizeTexture.SetCheckPoint(RHICmdList, RectLightMipTreeRT);
+
+	// Transition to compute
+	RHICmdList.TransitionResource(EResourceTransitionAccess::ERWBarrier, EResourceTransitionPipeline::EGfxToCompute, RectLightMipTree.UAV);
+
+	ResolveSceneColor(RHICmdList);
+	SceneContext.FinishRenderingSceneColor(RHICmdList);
+}
+
+void FDeferredShadingSceneRenderer::PrepareRayTracingRectLight(const FViewInfo& View, TArray<FRayTracingShaderRHIParamRef>& OutRayGenShaders)
+{
+	// Declare all RayGen shaders that require material closest hit shaders to be bound
+
+	TShaderMapRef<FRectLightRGS<0>> Shader0(GetGlobalShaderMap(View.FeatureLevel));
+	TShaderMapRef<FRectLightRGS<1>> Shader1(GetGlobalShaderMap(View.FeatureLevel));
+
+	OutRayGenShaders.Add(Shader0->GetRayTracingShader());
+	OutRayGenShaders.Add(Shader1->GetRayTracingShader());
+}
+
+template <int TextureImportanceSampling>
+void RenderRayTracingRectLightInternal(
+	FRHICommandListImmediate& RHICmdList,
+	const TArray<FViewInfo>& Views,
+	const FLightSceneInfo& RectLightSceneInfo,
+	TRefCountPtr<IPooledRenderTarget>& ScreenShadowMaskTexture,
+	TRefCountPtr<IPooledRenderTarget>& RayDistanceTexture
+)
+{
+	check(RectLightSceneInfo.Proxy);
+	check(RectLightSceneInfo.Proxy->IsRectLight());
+	FRectLightSceneProxy* RectLightSceneProxy = (FRectLightSceneProxy*)RectLightSceneInfo.Proxy;
+
+	FLightShaderParameters LightShaderParameters;
+	RectLightSceneProxy->GetLightShaderParameters(LightShaderParameters);
+
+	FRectLightData RectLightData;
+	RectLightData.SamplesPerPixel = GRayTracingStochasticRectLightSamplesPerPixel;
+	RectLightData.bIsTextureImportanceSampling = GRayTracingStochasticRectLightIsTextureImportanceSampling;
+	RectLightData.Position = RectLightSceneInfo.Proxy->GetOrigin();
+	RectLightData.Normal = RectLightSceneInfo.Proxy->GetDirection();
+	const FMatrix& WorldToLight = RectLightSceneInfo.Proxy->GetWorldToLight();
+	RectLightData.dPdu = FVector(WorldToLight.M[0][1], WorldToLight.M[1][1], WorldToLight.M[2][1]);
+	RectLightData.dPdv = FVector(WorldToLight.M[0][2], WorldToLight.M[1][2], WorldToLight.M[2][2]);
+	RectLightData.Color = LightShaderParameters.Color / 2.0;
+
+	// #dxr_todo: Ray traced textured area lights are 1.5X brighter than those in lit mode.
+	if (RectLightSceneProxy->HasSourceTexture())
+	{
+		RectLightData.Color *= 2.0 / 3.0;
+	}
+
+	RectLightData.Width = 2.0f * LightShaderParameters.SourceRadius;
+	RectLightData.Height = 2.0f * LightShaderParameters.SourceLength;
+	RectLightData.Texture = LightShaderParameters.SourceTexture;
+	RectLightData.TextureSampler = RHICreateSamplerState(FSamplerStateInitializerRHI(SF_Bilinear, AM_Border, AM_Border, AM_Border));
+	RectLightData.MipTree = RectLightSceneProxy->RayTracingData->RectLightMipTree.SRV;
+	RectLightData.MipTreeDimensions = RectLightSceneProxy->RayTracingData->RectLightMipTreeDimensions;
+	RectLightData.MaxNormalBias = GetRaytracingMaxNormalBias();
+	RectLightData.BarnCosAngle = FMath::Cos(FMath::DegreesToRadians(RectLightSceneProxy->BarnDoorAngle));
+	RectLightData.BarnLength = RectLightSceneProxy->BarnDoorLength;
+	FUniformBufferRHIRef RectLightUniformBuffer = RHICreateUniformBuffer(&RectLightData, FRectLightData::StaticStructMetadata.GetLayout(), EUniformBufferUsage::UniformBuffer_SingleDraw);
+
+	for (int32 ViewIndex = 0; ViewIndex < Views.Num(); ++ViewIndex)
+	{
+		const FViewInfo& View = Views[ViewIndex];
+		const FIntPoint ViewSize = View.ViewRect.Size();
+
+		TShaderMapRef<FRectLightRGS<TextureImportanceSampling>> RectLightRayGenerationShader(GetGlobalShaderMap(View.FeatureLevel));
+
+		FSceneRenderTargets& SceneContext = FSceneRenderTargets::Get(RHICmdList);
+		FSceneTexturesUniformParameters SceneTextures;
+		SetupSceneTextureUniformParameters(SceneContext, View.FeatureLevel, ESceneTextureSetupMode::All, SceneTextures);
+		FUniformBufferRHIRef SceneTexturesUniformBuffer = RHICreateUniformBuffer(&SceneTextures, FSceneTexturesUniformParameters::StaticStructMetadata.GetLayout(), EUniformBufferUsage::UniformBuffer_SingleDraw);
+
+		// Dispatch
+		RectLightRayGenerationShader->Dispatch(
+			RHICmdList,
+			View.RayTracingScene,
+			View.ViewUniformBuffer,
+			SceneTexturesUniformBuffer,
+			RectLightUniformBuffer,
+			ScreenShadowMaskTexture->GetRenderTargetItem().UAV,
+			RayDistanceTexture->GetRenderTargetItem().UAV,
+			ViewSize.X, ViewSize.Y
+		);
+	}
+
+	// Transition out to graphics pipeline
+	FComputeFenceRHIRef Fence = RHICmdList.CreateComputeFence(TEXT("RayTracingRectLight"));
+	RHICmdList.TransitionResource(EResourceTransitionAccess::ERWBarrier, EResourceTransitionPipeline::EComputeToGfx, ScreenShadowMaskTexture->GetRenderTargetItem().UAV, Fence);
+	GVisualizeTexture.SetCheckPoint(RHICmdList, ScreenShadowMaskTexture);
+
+	RHICmdList.TransitionResource(EResourceTransitionAccess::ERWBarrier, EResourceTransitionPipeline::EComputeToGfx, RayDistanceTexture->GetRenderTargetItem().UAV);
+	GVisualizeTexture.SetCheckPoint(RHICmdList, RayDistanceTexture);
+}
+
+#endif // RHI_RAYTRACING
+
+void FDeferredShadingSceneRenderer::RenderRayTracingStochasticRectLight(
+	FRHICommandListImmediate& RHICmdList,
+	const FLightSceneInfo& RectLightSceneInfo,
+	TRefCountPtr<IPooledRenderTarget>& RectLightRT,
+	TRefCountPtr<IPooledRenderTarget>& HitDistanceRT
+)
+#if RHI_RAYTRACING
+{
+	SCOPED_DRAW_EVENT(RHICmdList, RayTracingRectLight);
+	SCOPED_GPU_STAT(RHICmdList, RayTracingRectLight);
+
+	FSceneRenderTargets& SceneContext = FSceneRenderTargets::Get(RHICmdList);
+	FPooledRenderTargetDesc Desc = SceneContext.GetSceneColor()->GetDesc();
+	Desc.Format = PF_FloatRGBA;
+	Desc.Flags &= ~(TexCreate_FastVRAM | TexCreate_Transient);
+	GRenderTargetPool.FindFreeElement(RHICmdList, Desc, RectLightRT, TEXT("RayTracingRectLight"));
+	//ClearUAV(RHICmdList, RectLightRT->GetRenderTargetItem(), FLinearColor::Black);
+
+	Desc.Format = PF_R16F;
+	GRenderTargetPool.FindFreeElement(RHICmdList, Desc, HitDistanceRT, TEXT("RayTracingRectLightDistance"));
+	//ClearUAV(RHICmdList, HitDistanceRT->GetRenderTargetItem(), FLinearColor::Black);
+
+	if (RectLightSceneInfo.Proxy->HasSourceTexture())
+	{
+		RenderRayTracingRectLightInternal<1>(RHICmdList, Views, RectLightSceneInfo, RectLightRT, HitDistanceRT);
+	}
+	else
+	{
+		RenderRayTracingRectLightInternal<0>(RHICmdList, Views, RectLightSceneInfo, RectLightRT, HitDistanceRT);
+	}
+}
+#else
+{
+	unimplemented();
+}
+#endif
