p4 copy //ue4/main/Engine/Source/Runtime/Renderer/Private/RayTracing/RayTracingRectLight.cpp#4

This file was ignored when 5337956 was robomerged in 5401207.

diff --git a/Engine/Source/Runtime/Renderer/Private/RayTracing/RayTracingRectLight.cpp b/Engine/Source/Runtime/Renderer/Private/RayTracing/RayTracingRectLight.cpp
--- a/Engine/Source/Runtime/Renderer/Private/RayTracing/RayTracingRectLight.cpp
+++ b/Engine/Source/Runtime/Renderer/Private/RayTracing/RayTracingRectLight.cpp
@@ -270,7 +270,9 @@ void FDeferredShadingSceneRenderer::VisualizeRectLightMipTree(
 		SceneContext.GetSceneColor()->GetRenderTargetItem().TargetableTexture,
 		RectLightMipTreeRT->GetRenderTargetItem().TargetableTexture
 	};
+	PRAGMA_DISABLE_DEPRECATION_WARNINGS
 	SetRenderTargets(RHICmdList, 2, RenderTargets, SceneContext.GetSceneDepthSurface(), ESimpleRenderTargetMode::EExistingColorAndDepth, FExclusiveDepthStencil::DepthRead_StencilNop);
+	PRAGMA_ENABLE_DEPRECATION_WARNINGS
 
 	// PSO definition
 	FGraphicsPipelineStateInitializer GraphicsPSOInit;
