Fixed conflicting integration with the MesPassProcessor.
Added the new IndirectArgs.Offset to the hash in GetDynamicInstancingHash().

diff --git a/Engine/Source/Runtime/Renderer/Public/MeshPassProcessor.h b/Engine/Source/Runtime/Renderer/Public/MeshPassProcessor.h
--- a/Engine/Source/Runtime/Renderer/Public/MeshPassProcessor.h
+++ b/Engine/Source/Runtime/Renderer/Public/MeshPassProcessor.h
@@ -525,7 +525,8 @@ public:
 		}
 		else
 		{
-			Hash = PointerHash(IndirectArgsBuffer, Hash);
+			Hash = PointerHash(IndirectArgs.Buffer, Hash);
+			Hash = FCrc::TypeCrc32(IndirectArgs.Offset, Hash);
 		}		
 
 		return Hash;
