Preserve content written by MaterialSetupCS when voxelizing primitives in Volumetric Fog.

diff --git a/Engine/Source/Runtime/Renderer/Private/VolumetricFogVoxelization.cpp b/Engine/Source/Runtime/Renderer/Private/VolumetricFogVoxelization.cpp
--- a/Engine/Source/Runtime/Renderer/Private/VolumetricFogVoxelization.cpp
+++ b/Engine/Source/Runtime/Renderer/Private/VolumetricFogVoxelization.cpp
@@ -589,8 +589,8 @@ void FDeferredShadingSceneRenderer::VoxelizeFogVolumePrimitives(
 	if (View.VolumetricMeshBatches.Num() > 0 && DoesPlatformSupportVolumetricFogVoxelization(View.GetShaderPlatform()))
 	{
 		FRenderTargetParameters* PassParameters = GraphBuilder.AllocParameters<FRenderTargetParameters>();
-		PassParameters->RenderTargets[0] = FRenderTargetBinding(IntegrationData.VBufferA, ERenderTargetLoadAction::ENoAction, ERenderTargetStoreAction::EStore);
-		PassParameters->RenderTargets[1] = FRenderTargetBinding(IntegrationData.VBufferB, ERenderTargetLoadAction::ENoAction, ERenderTargetStoreAction::EStore);
+		PassParameters->RenderTargets[0] = FRenderTargetBinding(IntegrationData.VBufferA, ERenderTargetLoadAction::ELoad, ERenderTargetStoreAction::EStore);
+		PassParameters->RenderTargets[1] = FRenderTargetBinding(IntegrationData.VBufferB, ERenderTargetLoadAction::ELoad, ERenderTargetStoreAction::EStore);
 
 		GraphBuilder.AddPass(
 			RDG_EVENT_NAME("VoxelizeVolumePrimitives"),
