Fix assert failing due to unreleased UniformBuffers

diff --git a/Engine/Source/Runtime/Renderer/Private/MeshDrawCommands.cpp b/Engine/Source/Runtime/Renderer/Private/MeshDrawCommands.cpp
--- a/Engine/Source/Runtime/Renderer/Private/MeshDrawCommands.cpp
+++ b/Engine/Source/Runtime/Renderer/Private/MeshDrawCommands.cpp
@@ -871,14 +871,16 @@ void FParallelMeshDrawCommandPass::DispatchPassSetup(
 	check(!TaskEventRef.IsValid() && MeshPassProcessor != nullptr && TaskContext.PrimitiveIdBufferData == nullptr);
 	check((PassType == EMeshPass::Num) == (DynamicMeshElementsPassRelevance == nullptr));
 
+	// Assign these before the early-out because we still need to destroy them in Empty()
+	TaskContext.MeshPassProcessor = MeshPassProcessor;
+	TaskContext.MobileBasePassCSMMeshPassProcessor = MobileBasePassCSMMeshPassProcessor;
+
 	MaxNumDraws = InOutMeshDrawCommands.Num() + NumDynamicMeshElements + NumDynamicMeshCommandBuildRequestElements;
 	if (MaxNumDraws <= 0)
 	{
 		return;
 	}
 
-	TaskContext.MeshPassProcessor = MeshPassProcessor;
-	TaskContext.MobileBasePassCSMMeshPassProcessor = MobileBasePassCSMMeshPassProcessor;
 	TaskContext.DynamicMeshElements = &DynamicMeshElements;
 	TaskContext.DynamicMeshElementsPassRelevance = DynamicMeshElementsPassRelevance;
 
@@ -965,6 +967,15 @@ void FParallelMeshDrawCommandPass::Empty()
 
 	DumpInstancingStats();
 
+	if (TaskContext.MeshPassProcessor)
+	{
+		TaskContext.MeshPassProcessor->~FMeshPassProcessor();
+	}
+	if (TaskContext.MobileBasePassCSMMeshPassProcessor)
+	{
+		TaskContext.MobileBasePassCSMMeshPassProcessor->~FMeshPassProcessor();
+	}
+
 	if (!bPrimitiveIdBufferDataOwnedByRHIThread)
 	{
 		FMemory::Free(TaskContext.PrimitiveIdBufferData);
