Separates hit distance from ao mask regressing AO denoiser's reconstruction pass from 0.37ms to 0.43ms.

diff --git a/Engine/Shaders/Private/RayTracing/RayTracingAmbientOcclusionRGS.usf b/Engine/Shaders/Private/RayTracing/RayTracingAmbientOcclusionRGS.usf
--- a/Engine/Shaders/Private/RayTracing/RayTracingAmbientOcclusionRGS.usf
+++ b/Engine/Shaders/Private/RayTracing/RayTracingAmbientOcclusionRGS.usf
@@ -7,7 +7,8 @@
 #include "RayTracingCommon.ush"
 
 RaytracingAccelerationStructure TLAS;
-RWTexture2D<float2> RWOcclusionMaskUAV;
+RWTexture2D<float> RWOcclusionMaskUAV;
+RWTexture2D<float> RWHitDistanceUAV;
 
 uint CalcLinearIndex(uint2 PixelCoord)
 {
@@ -131,6 +132,7 @@ void AmbientOcclusionRGS()
 			RawOutput.y = ClosestRayHitDistance;
 		}
 
-		RWOcclusionMaskUAV[PixelCoord] = RawOutput;
+		RWOcclusionMaskUAV[PixelCoord] = RawOutput.x;
+		RWHitDistanceUAV[PixelCoord] = RawOutput.y;
 	}
 }
\ No newline at end of file
diff --git a/Engine/Shaders/Private/ScreenSpaceDenoise/SSDSignalCore.ush b/Engine/Shaders/Private/ScreenSpaceDenoise/SSDSignalCore.ush
--- a/Engine/Shaders/Private/ScreenSpaceDenoise/SSDSignalCore.ush
+++ b/Engine/Shaders/Private/ScreenSpaceDenoise/SSDSignalCore.ush
@@ -98,9 +98,9 @@ FSSFSignalSample DecodeInputSignal(const uint SignalBufferLayout, float4 RawSamp
 	}
 	else if (SignalBufferLayout == SIGNAL_BUFFER_LAYOUT_AO_INPUT)
 	{
-		OutSample.AccumulatedSampleCount = RawSample[0].g < 0 ? 0 : 1;
+		OutSample.AccumulatedSampleCount = RawSample[1].r < 0 ? 0 : 1;
 		OutSample.LeavingRayCount = OutSample.AccumulatedSampleCount * RawSample[0].r;
-		OutSample.PenumbraClosestHit = OutSample.AccumulatedSampleCount * RawSample[0].g;
+		OutSample.PenumbraClosestHit = OutSample.AccumulatedSampleCount * RawSample[1].r;
 	}
 	else if (SignalBufferLayout == SIGNAL_BUFFER_LAYOUT_AO_HISTORY)
 	{
diff --git a/Engine/Source/Runtime/Renderer/Private/DeferredShadingRenderer.cpp b/Engine/Source/Runtime/Renderer/Private/DeferredShadingRenderer.cpp
--- a/Engine/Source/Runtime/Renderer/Private/DeferredShadingRenderer.cpp
+++ b/Engine/Source/Runtime/Renderer/Private/DeferredShadingRenderer.cpp
@@ -1731,7 +1731,8 @@ void FDeferredShadingSceneRenderer::Render(FRHICommandListImmediate& RHICmdList)
 			{
 				checkSlow(RHICmdList.IsOutsideRenderPass());
 				TRefCountPtr<IPooledRenderTarget> AmbientOcclusionRT;
-				RenderRayTracingAmbientOcclusion(RHICmdList, nullptr, AmbientOcclusionRT);
+				TRefCountPtr<IPooledRenderTarget> AmbientOcclusionHitDistanceRT;
+				RenderRayTracingAmbientOcclusion(RHICmdList, nullptr, AmbientOcclusionRT, AmbientOcclusionHitDistanceRT);
 				checkSlow(RHICmdList.IsOutsideRenderPass());
 				
 				int32 DenoiserMode = CVarUseAODenoiser.GetValueOnRenderThread();
@@ -1748,7 +1749,8 @@ void FDeferredShadingSceneRenderer::Render(FRHICommandListImmediate& RHICmdList)
 					IScreenSpaceDenoiser::FAmbientOcclusionRayTracingConfig RayTracingConfig;
 
 					IScreenSpaceDenoiser::FAmbientOcclusionInputs DenoiserInputs;
-					DenoiserInputs.MaskAndRayHitDistance = GraphBuilder.RegisterExternalTexture(AmbientOcclusionRT, TEXT("AOMaskAndHitDistance"));
+					DenoiserInputs.Mask = GraphBuilder.RegisterExternalTexture(AmbientOcclusionRT, TEXT("AOMask"));
+					DenoiserInputs.RayHitDistance = GraphBuilder.RegisterExternalTexture(AmbientOcclusionHitDistanceRT, TEXT("AOHitDistance"));
 
 					const FViewInfo& View = Views[0];
 
diff --git a/Engine/Source/Runtime/Renderer/Private/DeferredShadingRenderer.h b/Engine/Source/Runtime/Renderer/Private/DeferredShadingRenderer.h
--- a/Engine/Source/Runtime/Renderer/Private/DeferredShadingRenderer.h
+++ b/Engine/Source/Runtime/Renderer/Private/DeferredShadingRenderer.h
@@ -449,7 +449,7 @@ private:
 	void RenderRayTracingOcclusionForRectLight(FRHICommandListImmediate& RHICmdList, const FLightSceneInfo& RectLightSceneInfo, TRefCountPtr<IPooledRenderTarget>& ScreenShadowMaskTexture);
 	template <int CalcDirectLighting, int EncodeVisibility, int TextureImportanceSampling> void RenderRayTracingRectLightInternal(FRHICommandListImmediate& RHICmdList, const FLightSceneInfo& RectLightSceneInfo, TRefCountPtr<IPooledRenderTarget>& ScreenShadowMaskTexture, TRefCountPtr<IPooledRenderTarget>& RayDistanceTexture);
 
-	void RenderRayTracingAmbientOcclusion(FRHICommandListImmediate& RHICmdList, const FLightSceneInfo* LightSceneInfo, TRefCountPtr<IPooledRenderTarget>& AmbientOcclusionRT);
+	void RenderRayTracingAmbientOcclusion(FRHICommandListImmediate& RHICmdList, const FLightSceneInfo* LightSceneInfo, TRefCountPtr<IPooledRenderTarget>& AmbientOcclusionRT, TRefCountPtr<IPooledRenderTarget>& AmbientOcclusionHitDistanceRT);
 	void CompositeRayTracingAmbientOcclusion(FRHICommandListImmediate& RHICmdList, TRefCountPtr<IPooledRenderTarget>& AmbientOcclusionRT);
 
 	void BuildSkyLightMipTree(FRHICommandListImmediate& RHICmdList, FTextureRHIRef SkyLightTexture, FRWBuffer& SkyLightMipTreePosX, FRWBuffer& SkyLightMipTreePosY, FRWBuffer& SkyLightMipTreePosZ, FRWBuffer& SkyLightMipTreeNegX, FRWBuffer& SkyLightMipTreeNegY, FRWBuffer& SkyLightMipTreeNegZ, FIntVector& SkyLightMipTreeDimensions);
diff --git a/Engine/Source/Runtime/Renderer/Private/RayTracing/RayTracingAmbientOcclusion.cpp b/Engine/Source/Runtime/Renderer/Private/RayTracing/RayTracingAmbientOcclusion.cpp
--- a/Engine/Source/Runtime/Renderer/Private/RayTracing/RayTracingAmbientOcclusion.cpp
+++ b/Engine/Source/Runtime/Renderer/Private/RayTracing/RayTracingAmbientOcclusion.cpp
@@ -78,6 +78,7 @@ public:
 		AmbientOcclusionParameter.Bind(Initializer.ParameterMap, TEXT("AmbientOcclusion"));
 
 		OcclusionMaskUAVParameter.Bind(Initializer.ParameterMap, TEXT("RWOcclusionMaskUAV"));
+		RayDistanceUAVParameter.Bind(Initializer.ParameterMap, TEXT("RWHitDistanceUAV"));
 	}
 
 	bool Serialize(FArchive& Ar)
@@ -88,6 +89,7 @@ public:
 		Ar << SceneTexturesParameter;
 		Ar << AmbientOcclusionParameter;
 		Ar << OcclusionMaskUAVParameter;
+		Ar << RayDistanceUAVParameter;
 		return bShaderHasOutdatedParameters;
 	}
 
@@ -98,6 +100,7 @@ public:
 		FUniformBufferRHIParamRef SceneTexturesUniformBuffer,
 		FUniformBufferRHIParamRef AmbientOcclusionUniformBuffer,
 		FUnorderedAccessViewRHIParamRef OcclusionMaskUAV,
+		FUnorderedAccessViewRHIParamRef HitDistanceUAV,
 		uint32 Width, uint32 Height
 	)
 	{
@@ -114,6 +117,7 @@ public:
 		GlobalResources.Set(SceneTexturesParameter, SceneTexturesUniformBuffer);
 		GlobalResources.Set(AmbientOcclusionParameter, AmbientOcclusionUniformBuffer);
 		GlobalResources.Set(OcclusionMaskUAVParameter, OcclusionMaskUAV);
+		GlobalResources.Set(RayDistanceUAVParameter, HitDistanceUAV);
 
 		const uint32 RayGenShaderIndex = 0;
 		RHICmdList.RayTraceDispatch(Pipeline, RayGenShaderIndex, GlobalResources, Width, Height);
@@ -136,7 +140,8 @@ IMPLEMENT_SHADER_TYPE(, FAmbientOcclusionRGS, TEXT("/Engine/Private/RayTracing/R
 void FDeferredShadingSceneRenderer::RenderRayTracingAmbientOcclusion(
 	FRHICommandListImmediate& RHICmdList,
 	const FLightSceneInfo* SkyLightSceneInfo,
-	TRefCountPtr<IPooledRenderTarget>& AmbientOcclusionRT
+	TRefCountPtr<IPooledRenderTarget>& AmbientOcclusionMask,
+	TRefCountPtr<IPooledRenderTarget>& HitDistance
 )
 {
 	SCOPED_DRAW_EVENT(RHICmdList, RayTracingAmbientOcclusion);
@@ -144,10 +149,14 @@ void FDeferredShadingSceneRenderer::RenderRayTracingAmbientOcclusion(
 
 	FSceneRenderTargets& SceneContext = FSceneRenderTargets::Get(RHICmdList);
 	FPooledRenderTargetDesc Desc = SceneContext.GetSceneColor()->GetDesc();
-	Desc.Format = PF_G16R16F;
+	Desc.Format = PF_R16F;
 	Desc.Flags &= ~(TexCreate_FastVRAM | TexCreate_Transient);
-	GRenderTargetPool.FindFreeElement(RHICmdList, Desc, AmbientOcclusionRT, TEXT("RayTracingAmbientOcclusion"));
-	ClearUAV(RHICmdList, AmbientOcclusionRT->GetRenderTargetItem(), FLinearColor::Black);
+	GRenderTargetPool.FindFreeElement(RHICmdList, Desc, AmbientOcclusionMask, TEXT("RayTracingAmbientOcclusion"));
+	GRenderTargetPool.FindFreeElement(RHICmdList, Desc, HitDistance, TEXT("RayTracingAmbientOcclusionHitDistance"));
+
+	// TODO: ClearUAV is useless.
+	ClearUAV(RHICmdList, AmbientOcclusionMask->GetRenderTargetItem(), FLinearColor::Black);
+	ClearUAV(RHICmdList, HitDistance->GetRenderTargetItem(), FLinearColor::Black);
 
 	// Add ambient occlusion parameters to uniform buffer
 	FAmbientOcclusionData AmbientOcclusionData;
@@ -171,15 +180,21 @@ void FDeferredShadingSceneRenderer::RenderRayTracingAmbientOcclusion(
 			View.ViewUniformBuffer,
 			SceneTexturesUniformBuffer,
 			AmbientOcclusionUniformBuffer,
-			AmbientOcclusionRT->GetRenderTargetItem().UAV,
+			AmbientOcclusionMask->GetRenderTargetItem().UAV,
+			HitDistance->GetRenderTargetItem().UAV,
 			ViewSize.X, ViewSize.Y
 		);
 	}
 
 	// Transition to graphics pipeline
 	FComputeFenceRHIRef Fence = RHICmdList.CreateComputeFence(TEXT("RayTracingAmbientOcclusion"));
-	RHICmdList.TransitionResource(EResourceTransitionAccess::ERWBarrier, EResourceTransitionPipeline::EComputeToGfx, AmbientOcclusionRT->GetRenderTargetItem().UAV, Fence);
-	GVisualizeTexture.SetCheckPoint(RHICmdList, AmbientOcclusionRT);
+	RHICmdList.TransitionResource(EResourceTransitionAccess::ERWBarrier, EResourceTransitionPipeline::EComputeToGfx, AmbientOcclusionMask->GetRenderTargetItem().UAV, Fence);
+
+	FComputeFenceRHIRef Fence2 = RHICmdList.CreateComputeFence(TEXT("RayTracingAmbientOcclusion"));
+	RHICmdList.TransitionResource(EResourceTransitionAccess::ERWBarrier, EResourceTransitionPipeline::EComputeToGfx, HitDistance->GetRenderTargetItem().UAV, Fence2);
+	
+	GVisualizeTexture.SetCheckPoint(RHICmdList, AmbientOcclusionMask);
+	GVisualizeTexture.SetCheckPoint(RHICmdList, HitDistance);
 }
 
 class FCompositeAmbientOcclusionPS : public FGlobalShader
diff --git a/Engine/Source/Runtime/Renderer/Private/ScreenSpaceDenoise.cpp b/Engine/Source/Runtime/Renderer/Private/ScreenSpaceDenoise.cpp
--- a/Engine/Source/Runtime/Renderer/Private/ScreenSpaceDenoise.cpp
+++ b/Engine/Source/Runtime/Renderer/Private/ScreenSpaceDenoise.cpp
@@ -996,7 +996,8 @@ public:
 		const FAmbientOcclusionRayTracingConfig RayTracingConfig) const override
 	{
 		FSSDSignalTextures InputSignal;
-		InputSignal.Texture0 = ReflectionInputs.MaskAndRayHitDistance;
+		InputSignal.Texture0 = ReflectionInputs.Mask;
+		InputSignal.Texture1 = ReflectionInputs.RayHitDistance;
 		
 		FSSDConstantPixelDensitySettings Settings;
 		Settings.SignalProcessing = ESignalProcessing::AmbientOcclusion;
diff --git a/Engine/Source/Runtime/Renderer/Private/ScreenSpaceDenoise.h b/Engine/Source/Runtime/Renderer/Private/ScreenSpaceDenoise.h
--- a/Engine/Source/Runtime/Renderer/Private/ScreenSpaceDenoise.h
+++ b/Engine/Source/Runtime/Renderer/Private/ScreenSpaceDenoise.h
@@ -97,10 +97,8 @@ public:
 	
 	/** All the inputs of the AO denoisers. */
 	BEGIN_SHADER_PARAMETER_STRUCT(FAmbientOcclusionInputs, )
-		// G16R16F texture:
-		//  R: mask as [0; 1].
-		//  B: closest hit distance of the ray.
-		SHADER_PARAMETER_RDG_TEXTURE(Texture2D, MaskAndRayHitDistance)
+		SHADER_PARAMETER_RDG_TEXTURE(Texture2D, Mask)
+		SHADER_PARAMETER_RDG_TEXTURE(Texture2D, RayHitDistance)
 	END_SHADER_PARAMETER_STRUCT()
 
 	/** All the outputs of the AO denoiser may generate. */
