Fixed crash when destoying mesh pass processor embeded inside cached preshadows - this could happen on next frame, when SceneRenderingAllocator memory is invalid and result in a crash.
#rnx

diff --git a/Engine/Source/Runtime/Renderer/Private/MeshDrawCommands.cpp b/Engine/Source/Runtime/Renderer/Private/MeshDrawCommands.cpp
--- a/Engine/Source/Runtime/Renderer/Private/MeshDrawCommands.cpp
+++ b/Engine/Source/Runtime/Renderer/Private/MeshDrawCommands.cpp
@@ -871,16 +871,10 @@ void FParallelMeshDrawCommandPass::DispatchPassSetup(
 	check(!TaskEventRef.IsValid() && MeshPassProcessor != nullptr && TaskContext.PrimitiveIdBufferData == nullptr);
 	check((PassType == EMeshPass::Num) == (DynamicMeshElementsPassRelevance == nullptr));
 
-	// Assign these before the early-out because we still need to destroy them in Empty()
-	TaskContext.MeshPassProcessor = MeshPassProcessor;
-	TaskContext.MobileBasePassCSMMeshPassProcessor = MobileBasePassCSMMeshPassProcessor;
-
 	MaxNumDraws = InOutMeshDrawCommands.Num() + NumDynamicMeshElements + NumDynamicMeshCommandBuildRequestElements;
-	if (MaxNumDraws <= 0)
-	{
-		return;
-	}
 
+	TaskContext.MeshPassProcessor = MeshPassProcessor;
+	TaskContext.MobileBasePassCSMMeshPassProcessor = MobileBasePassCSMMeshPassProcessor;
 	TaskContext.DynamicMeshElements = &DynamicMeshElements;
 	TaskContext.DynamicMeshElementsPassRelevance = DynamicMeshElementsPassRelevance;
 
@@ -926,27 +920,30 @@ void FParallelMeshDrawCommandPass::DispatchPassSetup(
 		check(MobileBasePassCSMMeshPassProcessor == nullptr && InOutMobileBasePassCSMMeshDrawCommands == nullptr);
 	}
 
-	// Preallocate resources on rendering thread based on MaxNumDraws.
-	bPrimitiveIdBufferDataOwnedByRHIThread = false;
-	TaskContext.PrimitiveIdBufferDataSize = TaskContext.MaxInstanceFactor * MaxNumDraws * sizeof(int32);
-	TaskContext.PrimitiveIdBufferData = FMemory::Malloc(TaskContext.PrimitiveIdBufferDataSize);
-	PrimitiveIdVertexBufferRHI = GPrimitiveIdVertexBufferPool.Allocate(TaskContext.PrimitiveIdBufferDataSize);
-	TaskContext.MeshDrawCommands.Reserve(MaxNumDraws);
-	TaskContext.TempVisibleMeshDrawCommands.Reserve(MaxNumDraws);
-
-	const bool bExecuteInParallel = FApp::ShouldUseThreadingForPerformance()
-		&& CVarMeshDrawCommandsParallelPassSetup.GetValueOnRenderThread() > 0
-		&& GRenderingThread; // Rendering thread is required to safely use rendering resources in parallel.
-
-	if (bExecuteInParallel)
-	{
-		TaskEventRef = TGraphTask<FMeshDrawCommandPassSetupTask>::CreateTask(
-			nullptr, ENamedThreads::GetRenderThread()).ConstructAndDispatchWhenReady(TaskContext);
-	}
-	else
+	if (MaxNumDraws > 0)
 	{
-		FMeshDrawCommandPassSetupTask Task(TaskContext);
-		Task.AnyThreadTask();
+		// Preallocate resources on rendering thread based on MaxNumDraws.
+		bPrimitiveIdBufferDataOwnedByRHIThread = false;
+		TaskContext.PrimitiveIdBufferDataSize = TaskContext.MaxInstanceFactor * MaxNumDraws * sizeof(int32);
+		TaskContext.PrimitiveIdBufferData = FMemory::Malloc(TaskContext.PrimitiveIdBufferDataSize);
+		PrimitiveIdVertexBufferRHI = GPrimitiveIdVertexBufferPool.Allocate(TaskContext.PrimitiveIdBufferDataSize);
+		TaskContext.MeshDrawCommands.Reserve(MaxNumDraws);
+		TaskContext.TempVisibleMeshDrawCommands.Reserve(MaxNumDraws);
+
+		const bool bExecuteInParallel = FApp::ShouldUseThreadingForPerformance()
+			&& CVarMeshDrawCommandsParallelPassSetup.GetValueOnRenderThread() > 0
+			&& GRenderingThread; // Rendering thread is required to safely use rendering resources in parallel.
+
+		if (bExecuteInParallel)
+		{
+			TaskEventRef = TGraphTask<FMeshDrawCommandPassSetupTask>::CreateTask(
+				nullptr, ENamedThreads::GetRenderThread()).ConstructAndDispatchWhenReady(TaskContext);
+		}
+		else
+		{
+			FMeshDrawCommandPassSetupTask Task(TaskContext);
+			Task.AnyThreadTask();
+		}
 	}
 }
 
@@ -970,10 +967,12 @@ void FParallelMeshDrawCommandPass::Empty()
 	if (TaskContext.MeshPassProcessor)
 	{
 		TaskContext.MeshPassProcessor->~FMeshPassProcessor();
+		TaskContext.MeshPassProcessor = nullptr;
 	}
 	if (TaskContext.MobileBasePassCSMMeshPassProcessor)
 	{
 		TaskContext.MobileBasePassCSMMeshPassProcessor->~FMeshPassProcessor();
+		TaskContext.MobileBasePassCSMMeshPassProcessor = nullptr;
 	}
 
 	if (!bPrimitiveIdBufferDataOwnedByRHIThread)
@@ -985,8 +984,6 @@ void FParallelMeshDrawCommandPass::Empty()
 	MaxNumDraws = 0;
 	PassNameForStats.Empty();
 
-	TaskContext.MeshPassProcessor = nullptr;
-	TaskContext.MobileBasePassCSMMeshPassProcessor = nullptr;
 	TaskContext.DynamicMeshElements = nullptr;
 	TaskContext.DynamicMeshElementsPassRelevance = nullptr;
 	TaskContext.MeshDrawCommands.Empty();
diff --git a/Engine/Source/Runtime/Renderer/Private/SceneRendering.cpp b/Engine/Source/Runtime/Renderer/Private/SceneRendering.cpp
--- a/Engine/Source/Runtime/Renderer/Private/SceneRendering.cpp
+++ b/Engine/Source/Runtime/Renderer/Private/SceneRendering.cpp
@@ -2884,6 +2884,13 @@ void FSceneRenderer::WaitForTasksClearSnapshotsAndDeleteSceneRenderer(FRHIComman
 
 	GPrimitiveIdVertexBufferPool.DiscardAll();
 
+	// Destroy cached preshadow transient arrays (allocated with SceneRenderingAllocator).
+	TArray<TRefCountPtr<FProjectedShadowInfo>>& CachedPreshadows = SceneRenderer->Scene->CachedPreshadows;
+	for (int32 CachedShadowIndex = 0; CachedShadowIndex < CachedPreshadows.Num(); ++CachedShadowIndex)
+	{
+		CachedPreshadows[CachedShadowIndex]->ClearTransientArrays();
+	}
+
 	FViewInfo::DestroyAllSnapshots(); // this destroys viewinfo snapshots
 	FSceneRenderTargets::GetGlobalUnsafe().DestroyAllSnapshots(); // this will destroy the render target snapshots
 	static const IConsoleVariable* AsyncDispatch	= IConsoleManager::Get().FindConsoleVariable(TEXT("r.RHICmdAsyncRHIThreadDispatch"));
