Rectifying RectLight energy in path tracing to match Lit mode integrators.

diff --git a/Engine/Source/Runtime/Renderer/Private/PathTracing.cpp b/Engine/Source/Runtime/Renderer/Private/PathTracing.cpp
--- a/Engine/Source/Runtime/Renderer/Private/PathTracing.cpp
+++ b/Engine/Source/Runtime/Renderer/Private/PathTracing.cpp
@@ -173,7 +173,7 @@ public:
 						LightData.Normal[LightData.Count] = -LightParameters.Direction;
 						LightData.dPdu[LightData.Count] = FVector::CrossProduct(LightParameters.Tangent, LightParameters.Direction);
 						LightData.dPdv[LightData.Count] = LightParameters.Tangent;
-						LightData.Color[LightData.Count] = LightParameters.Color;
+						LightData.Color[LightData.Count] = LightParameters.Color / 4.0;
 						LightData.Dimensions[LightData.Count] = FVector(2.0f * LightParameters.SourceRadius, 2.0f * LightParameters.SourceLength, 0.0f);
 						break;
 					}
