Unifying RTAO pass with SSAO parameters.

diff --git a/Engine/Shaders/Private/RayTracing/RayTracingAmbientOcclusionRGS.usf b/Engine/Shaders/Private/RayTracing/RayTracingAmbientOcclusionRGS.usf
--- a/Engine/Shaders/Private/RayTracing/RayTracingAmbientOcclusionRGS.usf
+++ b/Engine/Shaders/Private/RayTracing/RayTracingAmbientOcclusionRGS.usf
@@ -69,6 +69,7 @@ void AmbientOcclusionRGS()
 
 	float RayCount = 0.0;
 	float Visibility = 0.0;
+	float Intensity = AmbientOcclusion.Intensity;
 	float ClosestRayHitDistance = 1.#INF;
 	uint SamplesPerPixel = AmbientOcclusion.SamplesPerPixel;
 	
@@ -105,7 +106,7 @@ void AmbientOcclusionRGS()
 		);
 		
 		RayCount += 1;
-		Visibility += Payload.HitT < 0 ? 1.0 : 0.0;
+		Visibility += Payload.HitT < 0 ? 1.0 : 1.0 - Intensity;
 		if (Payload.HitT >= 0.0)
 		{
 			ClosestRayHitDistance = min(ClosestRayHitDistance, Payload.HitT);
diff --git a/Engine/Source/Runtime/Renderer/Private/CompositionLighting/CompositionLighting.cpp b/Engine/Source/Runtime/Renderer/Private/CompositionLighting/CompositionLighting.cpp
--- a/Engine/Source/Runtime/Renderer/Private/CompositionLighting/CompositionLighting.cpp
+++ b/Engine/Source/Runtime/Renderer/Private/CompositionLighting/CompositionLighting.cpp
@@ -17,6 +17,7 @@
 #include "LightPropagationVolumeSettings.h"
 #include "DecalRenderingShared.h"
 #include "VisualizeTexture.h"
+#include "RayTracing/RayTracingOptions.h"
 
 /** The global center for all deferred lighting activities. */
 FCompositionLighting GCompositionLighting;
@@ -122,11 +123,10 @@ bool ShouldRenderScreenSpaceAmbientOcclusion(const FViewInfo& View)
 			&& (FSSAOHelper::IsBasePassAmbientOcclusionRequired(View) || IsAmbientCubemapPassRequired(View) || IsReflectionEnvironmentActive(View) || IsSkylightActive(View) || View.Family->EngineShowFlags.VisualizeBuffer)
 			&& !IsSimpleForwardShadingEnabled(View.GetShaderPlatform());
 	}
-
-	// #dxr_todo: hard-coded FindConsoleVariable is a bad hack. The current variable is hidden behind DeferredShadingRenderer
-	static IConsoleVariable* UseRayTracingAmbientOcclusionCVar = IConsoleManager::Get().FindConsoleVariable(TEXT("r.RayTracing.AmbientOcclusion"));
-	bool bUseRayTracingAmbientOcclusion = IsRayTracingEnabled() && UseRayTracingAmbientOcclusionCVar && UseRayTracingAmbientOcclusionCVar->GetInt() > 0;
-	return bEnabled && !bUseRayTracingAmbientOcclusion;
+#if RHI_RAYTRACING
+	bEnabled &= !ShouldRenderRayTracingAmbientOcclusion();
+#endif
+	return bEnabled;
 }
 
 static void AddPostProcessingAmbientCubemap(FPostprocessContext& Context, FRenderingCompositeOutputRef AmbientOcclusion)
@@ -329,7 +329,12 @@ void FCompositionLighting::ProcessAfterBasePass(FRHICommandListImmediate& RHICmd
 		if (!IsForwardShadingEnabled(View.GetShaderPlatform()))
 		{
 			FRenderingCompositeOutputRef AmbientOcclusion;
-
+#if RHI_RAYTRACING
+			if (ShouldRenderRayTracingAmbientOcclusion() && SceneContext.bScreenSpaceAOIsValid)
+			{
+				AmbientOcclusion = Context.Graph.RegisterPass(new(FMemStack::Get()) FRCPassPostProcessInput(SceneContext.ScreenSpaceAO));
+			}
+#endif
 			uint32 SSAOLevels = FSSAOHelper::ComputeAmbientOcclusionPassCount(Context.View);
 			if (SSAOLevels)
 			{
@@ -352,13 +357,14 @@ void FCompositionLighting::ProcessAfterBasePass(FRHICommandListImmediate& RHICmd
 						TEXT("Ambient occlusion decals are not supported with Async compute SSAO."));
 				}
 
-				if (FSSAOHelper::IsBasePassAmbientOcclusionRequired(Context.View))
-				{
-					FRenderingCompositePass* Pass = Context.Graph.RegisterPass(new(FMemStack::Get()) FRCPassPostProcessBasePassAO());
-					Pass->AddDependency(Context.FinalOutput);
+			}
 
-					Context.FinalOutput = FRenderingCompositeOutputRef(Pass);
-				}
+			if (SceneContext.bScreenSpaceAOIsValid && FSSAOHelper::IsBasePassAmbientOcclusionRequired(Context.View))
+			{
+				FRenderingCompositePass* Pass = Context.Graph.RegisterPass(new(FMemStack::Get()) FRCPassPostProcessBasePassAO());
+				Pass->AddDependency(Context.FinalOutput);
+
+				Context.FinalOutput = FRenderingCompositeOutputRef(Pass);
 			}
 
 			if (IsAmbientCubemapPassRequired(Context.View))
diff --git a/Engine/Source/Runtime/Renderer/Private/DeferredShadingRenderer.cpp b/Engine/Source/Runtime/Renderer/Private/DeferredShadingRenderer.cpp
--- a/Engine/Source/Runtime/Renderer/Private/DeferredShadingRenderer.cpp
+++ b/Engine/Source/Runtime/Renderer/Private/DeferredShadingRenderer.cpp
@@ -33,6 +33,7 @@
 #include "RayTracing/RayTracingDynamicGeometryCollection.h"
 #include "SceneViewFamilyBlackboard.h"
 #include "ScreenSpaceDenoise.h"
+#include "RayTracing/RayTracingOptions.h"
 
 TAutoConsoleVariable<int32> CVarEarlyZPass(
 	TEXT("r.EarlyZPass"),
@@ -197,6 +198,7 @@ DECLARE_CYCLE_STAT(TEXT("DeferredShadingSceneRenderer ViewExtensionPostRenderBas
 
 DECLARE_GPU_STAT(Postprocessing);
 DECLARE_GPU_STAT(HZB);
+DECLARE_GPU_STAT_NAMED(AmbientOcclusionDenoiser, TEXT("Ambient Occlusion Denoiser"));
 DECLARE_GPU_STAT_NAMED(Unaccounted, TEXT("[unaccounted]"));
 
 bool ShouldForceFullDepthPass(EShaderPlatform ShaderPlatform)
@@ -1730,14 +1732,14 @@ void FDeferredShadingSceneRenderer::Render(FRHICommandListImmediate& RHICmdList)
 			else if (ShouldRenderRayTracingAmbientOcclusion())
 			{
 				checkSlow(RHICmdList.IsOutsideRenderPass());
-				TRefCountPtr<IPooledRenderTarget> AmbientOcclusionRT;
-				TRefCountPtr<IPooledRenderTarget> AmbientOcclusionHitDistanceRT;
-				RenderRayTracingAmbientOcclusion(RHICmdList, nullptr, AmbientOcclusionRT, AmbientOcclusionHitDistanceRT);
 				checkSlow(RHICmdList.IsOutsideRenderPass());
+				TRefCountPtr<IPooledRenderTarget> AmbientOcclusionHitDistanceRT;
+				RenderRayTracingAmbientOcclusion(RHICmdList, nullptr, SceneContext.ScreenSpaceAO, AmbientOcclusionHitDistanceRT);
 				
 				int32 DenoiserMode = CVarUseAODenoiser.GetValueOnRenderThread();
 				if (DenoiserMode != 0)
 				{	
+					SCOPED_GPU_STAT(RHICmdList, AmbientOcclusionDenoiser);
 					FRDGBuilder GraphBuilder(RHICmdList);
 
 					FSceneViewFamilyBlackboard SceneBlackboard;
@@ -1749,7 +1751,7 @@ void FDeferredShadingSceneRenderer::Render(FRHICommandListImmediate& RHICmdList)
 					IScreenSpaceDenoiser::FAmbientOcclusionRayTracingConfig RayTracingConfig;
 
 					IScreenSpaceDenoiser::FAmbientOcclusionInputs DenoiserInputs;
-					DenoiserInputs.Mask = GraphBuilder.RegisterExternalTexture(AmbientOcclusionRT, TEXT("AOMask"));
+					DenoiserInputs.Mask = GraphBuilder.RegisterExternalTexture(SceneContext.ScreenSpaceAO, TEXT("AOMask"));
 					DenoiserInputs.RayHitDistance = GraphBuilder.RegisterExternalTexture(AmbientOcclusionHitDistanceRT, TEXT("AOHitDistance"));
 
 					const FViewInfo& View = Views[0];
@@ -1767,14 +1769,12 @@ void FDeferredShadingSceneRenderer::Render(FRHICommandListImmediate& RHICmdList)
 							DenoiserInputs,
 							RayTracingConfig);
 						
-						GraphBuilder.QueueTextureExtraction(DenoiserOutputs.AmbientOcclusionMask, &AmbientOcclusionRT);
+						GraphBuilder.QueueTextureExtraction(DenoiserOutputs.AmbientOcclusionMask, &SceneContext.ScreenSpaceAO);
 					}
 					GraphBuilder.Execute();
 				}
 
 				// #dxr_todo: Denoise AO
-				CompositeRayTracingAmbientOcclusion(RHICmdList, AmbientOcclusionRT);
-				checkSlow(RHICmdList.IsOutsideRenderPass());
 			}
 		}
 	}
diff --git a/Engine/Source/Runtime/Renderer/Private/DeferredShadingRenderer.h b/Engine/Source/Runtime/Renderer/Private/DeferredShadingRenderer.h
--- a/Engine/Source/Runtime/Renderer/Private/DeferredShadingRenderer.h
+++ b/Engine/Source/Runtime/Renderer/Private/DeferredShadingRenderer.h
@@ -440,8 +440,6 @@ private:
 		float ResolutionFraction);
 
 #if RHI_RAYTRACING
-	bool ShouldRenderRayTracingAmbientOcclusion() const;
-
 	void RenderRayTracingOcclusion(FRHICommandListImmediate& RHICmdList, const FLightSceneInfo* LightSceneInfo, TRefCountPtr<IPooledRenderTarget>& ScreenShadowMaskTexture);
 
 	void VisualizeRectLightMipTree(FRHICommandListImmediate& RHICmdList, const FViewInfo& View, const FRWBuffer& RectLightMipTree, const FIntVector& RectLightMipTreeDimensions);
diff --git a/Engine/Source/Runtime/Renderer/Private/RayTracing/RayTracingAmbientOcclusion.cpp b/Engine/Source/Runtime/Renderer/Private/RayTracing/RayTracingAmbientOcclusion.cpp
--- a/Engine/Source/Runtime/Renderer/Private/RayTracing/RayTracingAmbientOcclusion.cpp
+++ b/Engine/Source/Runtime/Renderer/Private/RayTracing/RayTracingAmbientOcclusion.cpp
@@ -23,7 +23,7 @@ static FAutoConsoleVariableRef CVarRayTracingAmbientOcclusion(
 	TEXT("Enables ray tracing ambient occlusion (default = 1)")
 );
 
-bool FDeferredShadingSceneRenderer::ShouldRenderRayTracingAmbientOcclusion() const
+bool ShouldRenderRayTracingAmbientOcclusion()
 {
 	return GRayTracingAmbientOcclusion != 0;
 }
@@ -45,11 +45,12 @@ static FAutoConsoleVariableRef CVarRayTracingAmbientOcclusionMaxRayDistance(
 BEGIN_GLOBAL_SHADER_PARAMETER_STRUCT(FAmbientOcclusionData, )
 SHADER_PARAMETER(int, SamplesPerPixel)
 SHADER_PARAMETER(float, MaxRayDistance)
+SHADER_PARAMETER(float, Intensity)
 END_GLOBAL_SHADER_PARAMETER_STRUCT()
 
 IMPLEMENT_GLOBAL_SHADER_PARAMETER_STRUCT(FAmbientOcclusionData, "AmbientOcclusion");
 
-DECLARE_GPU_STAT_NAMED(RayTracingAmbientOcclusion, TEXT("RTAO"));
+DECLARE_GPU_STAT_NAMED(RayTracingAmbientOcclusion, TEXT("Ray Tracing Ambient Occlusion"));
 DECLARE_GPU_STAT_NAMED(CompositeRayTracingAmbientOcclusion, TEXT("Denoising: RTAO"));
 
 class FAmbientOcclusionRGS : public FGlobalShader
@@ -161,8 +162,6 @@ void FDeferredShadingSceneRenderer::RenderRayTracingAmbientOcclusion(
 	// Add ambient occlusion parameters to uniform buffer
 	FAmbientOcclusionData AmbientOcclusionData;
 	AmbientOcclusionData.SamplesPerPixel = GRayTracingAmbientOcclusionSamplesPerPixel;
-	AmbientOcclusionData.MaxRayDistance = GRayTracingAmbientOcclusionMaxRayDistance;
-	FUniformBufferRHIRef AmbientOcclusionUniformBuffer = RHICreateUniformBuffer(&AmbientOcclusionData, FAmbientOcclusionData::StaticStructMetadata.GetLayout(), EUniformBufferUsage::UniformBuffer_SingleDraw);
 
 	for (int32 ViewIndex = 0; ViewIndex < Views.Num(); ++ViewIndex)
 	{
@@ -174,6 +173,10 @@ void FDeferredShadingSceneRenderer::RenderRayTracingAmbientOcclusion(
 		SetupSceneTextureUniformParameters(SceneContext, FeatureLevel, ESceneTextureSetupMode::All, SceneTextures);
 		FUniformBufferRHIRef SceneTexturesUniformBuffer = RHICreateUniformBuffer(&SceneTextures, FSceneTexturesUniformParameters::StaticStructMetadata.GetLayout(), EUniformBufferUsage::UniformBuffer_SingleDraw);
 
+		AmbientOcclusionData.MaxRayDistance = View.FinalPostProcessSettings.AmbientOcclusionRadius;
+		AmbientOcclusionData.Intensity = View.FinalPostProcessSettings.AmbientOcclusionIntensity;
+		FUniformBufferRHIRef AmbientOcclusionUniformBuffer = RHICreateUniformBuffer(&AmbientOcclusionData, FAmbientOcclusionData::StaticStructMetadata.GetLayout(), EUniformBufferUsage::UniformBuffer_SingleDraw);
+
 		AmbientOcclusionRayGenerationShader->Dispatch(
 			RHICmdList,
 			View.PerViewRayTracingScene,
@@ -195,6 +198,7 @@ void FDeferredShadingSceneRenderer::RenderRayTracingAmbientOcclusion(
 	
 	GVisualizeTexture.SetCheckPoint(RHICmdList, AmbientOcclusionMask);
 	GVisualizeTexture.SetCheckPoint(RHICmdList, HitDistance);
+	SceneContext.bScreenSpaceAOIsValid = true;
 }
 
 class FCompositeAmbientOcclusionPS : public FGlobalShader
diff --git a/Engine/Source/Runtime/Renderer/Private/RayTracing/RaytracingOptions.h b/Engine/Source/Runtime/Renderer/Private/RayTracing/RaytracingOptions.h
--- a/Engine/Source/Runtime/Renderer/Private/RayTracing/RaytracingOptions.h
+++ b/Engine/Source/Runtime/Renderer/Private/RayTracing/RaytracingOptions.h
@@ -10,6 +10,7 @@ extern bool IsRayTracingSkyLightSelected();
 
 extern bool IsRayTracingRectLightSelected();
 
+extern bool ShouldRenderRayTracingAmbientOcclusion();
 extern bool ShouldRenderRayTracingStaticOrStationaryRectLight(const FLightSceneInfo& LightSceneInfo);
 extern bool ShouldRenderRayTracingDynamicRectLight(const FLightSceneInfo& LightSceneInfo);
 
