Adding denoising support for RTGI pass

diff --git a/Engine/Shaders/Private/ScreenSpaceDenoise/SSDDefinitions.ush b/Engine/Shaders/Private/ScreenSpaceDenoise/SSDDefinitions.ush
--- a/Engine/Shaders/Private/ScreenSpaceDenoise/SSDDefinitions.ush
+++ b/Engine/Shaders/Private/ScreenSpaceDenoise/SSDDefinitions.ush
@@ -81,6 +81,11 @@
 	/** Internal buffer layout for ambient occlusion. */
 	#define SIGNAL_BUFFER_LAYOUT_AO_HISTORY 6
 
+	/** Buffer layout given by the ray generation shader for global illumination. */
+	#define SIGNAL_BUFFER_LAYOUT_GI_AND_AO_INPUT 7
+
+	/** Internal buffer layout for global illumination. */
+	#define SIGNAL_BUFFER_LAYOUT_GI_AND_AO_HISTORY 8
 
 /** Defines how the signal is being processed. Matches C++'s ESignalProcessing. */
 	/** Shadow penumbra denoising. */
@@ -92,6 +97,9 @@
 	/** Ambient occlusion denoising. */
 	#define SIGNAL_PROCESSING_AO 2
 
+	/** Global illumination denoising. */
+	#define SIGNAL_PROCESSING_GI_AND_AO 3
+
 
 
 /** Defines the color space bitfield. */
diff --git a/Engine/Shaders/Private/ScreenSpaceDenoise/SSDSignalCore.ush b/Engine/Shaders/Private/ScreenSpaceDenoise/SSDSignalCore.ush
--- a/Engine/Shaders/Private/ScreenSpaceDenoise/SSDSignalCore.ush
+++ b/Engine/Shaders/Private/ScreenSpaceDenoise/SSDSignalCore.ush
@@ -107,6 +107,22 @@ FSSFSignalSample DecodeInputSignal(const uint SignalBufferLayout, float4 RawSamp
 		OutSample.AccumulatedSampleCount = RawSample[0].g;
 		OutSample.LeavingRayCount = OutSample.AccumulatedSampleCount * RawSample[0].r;
 	}
+	else if (SignalBufferLayout == SIGNAL_BUFFER_LAYOUT_GI_AND_AO_INPUT)
+	{
+		OutSample.AccumulatedSampleCount = RawSample[1].g;
+		OutSample.LeavingRayCount = OutSample.AccumulatedSampleCount * (1 - RawSample[0].a);
+		OutSample.SceneColor = OutSample.AccumulatedSampleCount * RawSample[0];
+
+		// max() because ray hit distance is < 0 when miss / not ray has been shot.
+		OutSample.PenumbraClosestHit = RawSample[1].r * OutSample.AccumulatedSampleCount;
+		OutSample.PenumbraClosestHit = ((RawSample[1].r == -1.0) ? 1.0e20 : RawSample[1].r) * OutSample.AccumulatedSampleCount;
+	}
+	else if (SignalBufferLayout == SIGNAL_BUFFER_LAYOUT_GI_AND_AO_HISTORY)
+	{
+		OutSample.AccumulatedSampleCount = RawSample[1].r;
+		OutSample.LeavingRayCount = OutSample.AccumulatedSampleCount * (1 - RawSample[0].a);
+		OutSample.SceneColor = OutSample.AccumulatedSampleCount * RawSample[0];
+	}
 
 	return OutSample;
 }
@@ -166,6 +182,13 @@ void EncodeOutputSignal(const uint SignalBufferLayout, FSSFSignalSample Sample,
 		OutRawSample[0].r = Sample.AccumulatedSampleCount > 0 ? Sample.LeavingRayCount / Sample.AccumulatedSampleCount : 1.0;
 		OutRawSample[0].g = Sample.AccumulatedSampleCount;
 	}
+	else if (SignalBufferLayout == SIGNAL_BUFFER_LAYOUT_GI_AND_AO_HISTORY)
+	{
+		// Outputs number of ray miss in red to be compatible as a SSAO output.
+		OutRawSample[0] = Sample.SceneColor * (Sample.AccumulatedSampleCount > 0 ? rcp(Sample.AccumulatedSampleCount) : 0);
+		OutRawSample[1].r = Sample.AccumulatedSampleCount;
+		OutBufferCount = 2;
+	}
 }
 
 /** Sample input signal. */
diff --git a/Engine/Shaders/Private/ScreenSpaceDenoise/SSDSignalFramework.ush b/Engine/Shaders/Private/ScreenSpaceDenoise/SSDSignalFramework.ush
--- a/Engine/Shaders/Private/ScreenSpaceDenoise/SSDSignalFramework.ush
+++ b/Engine/Shaders/Private/ScreenSpaceDenoise/SSDSignalFramework.ush
@@ -56,7 +56,7 @@
 	#include "../DeferredShadingCommon.ush"
 	#include "../CommonViewUniformBuffer.ush"
 
-#elif CONFIG_SIGNAL_PROCESSING == SIGNAL_PROCESSING_AO
+#elif CONFIG_SIGNAL_PROCESSING == SIGNAL_PROCESSING_AO || CONFIG_SIGNAL_PROCESSING == SIGNAL_PROCESSING_GI_AND_AO
 
 #else
 	#error Unknown CONFIG_SIGNAL_PROCESSING
@@ -524,6 +524,10 @@ float GetSignalWorldBluringRadius(FSSFSignalSample Sample, FSSFSampleSceneInfos
 {
 	return 1000; //Sample.PenumbraClosestHit;
 }
+#elif CONFIG_SIGNAL_PROCESSING == SIGNAL_PROCESSING_GI_AND_AO
+{
+	Sample.PenumbraClosestHit;
+}
 #else
 	#error Unimplemented GetSignalWorldBluringRadius()
 #endif
diff --git a/Engine/Shaders/Private/ScreenSpaceDenoise/SSDSpatialAccumulation.usf b/Engine/Shaders/Private/ScreenSpaceDenoise/SSDSpatialAccumulation.usf
--- a/Engine/Shaders/Private/ScreenSpaceDenoise/SSDSpatialAccumulation.usf
+++ b/Engine/Shaders/Private/ScreenSpaceDenoise/SSDSpatialAccumulation.usf
@@ -109,6 +109,25 @@
 		#error Unexpected stage.
 	#endif
 
+#elif CONFIG_SIGNAL_PROCESSING == SIGNAL_PROCESSING_GI_AND_AO
+	#if DIM_STAGE == STAGE_RECONSTRUCTION
+		// Input and output layout.
+		#define CONFIG_SIGNAL_INPUT_LAYOUT  SIGNAL_BUFFER_LAYOUT_GI_AND_AO_INPUT
+		#define CONFIG_SIGNAL_OUTPUT_LAYOUT SIGNAL_BUFFER_LAYOUT_GI_AND_AO_HISTORY
+
+	#define CONFIG_SAMPLE_SET           SAMPLE_SET_STACKOWIAK_4_SETS
+
+	#elif DIM_STAGE == STAGE_POST_FILTERING
+		// Input and output layout.
+		#define CONFIG_SIGNAL_INPUT_LAYOUT  SIGNAL_BUFFER_LAYOUT_GI_AND_AO_INPUT
+		#define CONFIG_SIGNAL_OUTPUT_LAYOUT SIGNAL_BUFFER_LAYOUT_GI_AND_AO_HISTORY
+
+		#define CONFIG_SAMPLE_SET           SAMPLE_SET_3X3
+
+	#else
+		#error Unexpected stage.
+	#endif
+
 #else
 	#error Unknown signal processing.
 #endif
diff --git a/Engine/Shaders/Private/ScreenSpaceDenoise/SSDTemporalAccumulation.usf b/Engine/Shaders/Private/ScreenSpaceDenoise/SSDTemporalAccumulation.usf
--- a/Engine/Shaders/Private/ScreenSpaceDenoise/SSDTemporalAccumulation.usf
+++ b/Engine/Shaders/Private/ScreenSpaceDenoise/SSDTemporalAccumulation.usf
@@ -63,6 +63,16 @@
 	#define CONFIG_VIEW_POSITION_BASED_BILATERAL 1
 	#define CONFIG_NORMAL_BASED_BILATERAL 1
 
+#elif CONFIG_SIGNAL_PROCESSING == SIGNAL_PROCESSING_GI_AND_AO
+	// Input and output layout.
+	#define CONFIG_SIGNAL_INPUT_LAYOUT   SIGNAL_BUFFER_LAYOUT_GI_AND_AO_HISTORY
+	#define CONFIG_SIGNAL_HISTORY_LAYOUT SIGNAL_BUFFER_LAYOUT_GI_AND_AO_HISTORY
+
+	#define CONFIG_HISTORY_REJECTION 1
+
+	#define CONFIG_VIEW_POSITION_BASED_BILATERAL 1
+	#define CONFIG_NORMAL_BASED_BILATERAL 1
+
 #else
 	#error Unimplemented signal processing.
 #endif
@@ -277,7 +287,7 @@ void TemporallyAccumulate(
 		//	max(CurrentFrame.AccumulatedSampleCount - TargetedSampleCount, 0),
 		//	0);
 	}
-	#elif CONFIG_SIGNAL_PROCESSING == SIGNAL_PROCESSING_AO
+	#elif CONFIG_SIGNAL_PROCESSING == SIGNAL_PROCESSING_AO || CONFIG_SIGNAL_PROCESSING == SIGNAL_PROCESSING_GI_AND_AO
 	{
 		TargetedSampleCount = TARGETED_SAMPLE_COUNT;
 	}
@@ -581,7 +591,7 @@ void TemporallyAccumulate(
 	
 	FSSFSignalSample HistorySample = PreviousFrame;
 
-	const bool bPostRejectionBlending = !bBruteForce || CONFIG_SIGNAL_PROCESSING == SIGNAL_PROCESSING_REFLECTIONS || CONFIG_SIGNAL_PROCESSING == SIGNAL_PROCESSING_AO;
+	const bool bPostRejectionBlending = !bBruteForce || CONFIG_SIGNAL_PROCESSING == SIGNAL_PROCESSING_REFLECTIONS || CONFIG_SIGNAL_PROCESSING == SIGNAL_PROCESSING_AO || CONFIG_SIGNAL_PROCESSING == SIGNAL_PROCESSING_GI_AND_AO;
 
 	if (!bPostRejectionBlending)
 	{
@@ -880,7 +890,7 @@ void TemporallyAccumulate(
 					RejectionFactor = max(RejectionFactor, 1 * max(RejectedDiff.SceneColor.y, RejectedDiff.SceneColor.z));
 					RejectionFactor = max(RejectionFactor, 1 * RejectedDiff.SceneColor.a);
 				}
-				#elif CONFIG_SIGNAL_PROCESSING == SIGNAL_PROCESSING_AO
+				#elif CONFIG_SIGNAL_PROCESSING == SIGNAL_PROCESSING_AO || CONFIG_SIGNAL_PROCESSING == SIGNAL_PROCESSING_GI_AND_AO
 				{
 					RejectionFactor = abs(ClampedTransformedHistorySample.LeavingRayCount - TransformedHistorySample.LeavingRayCount);
 				}
diff --git a/Engine/Source/Runtime/Renderer/Private/SceneRendering.h b/Engine/Source/Runtime/Renderer/Private/SceneRendering.h
--- a/Engine/Source/Runtime/Renderer/Private/SceneRendering.h
+++ b/Engine/Source/Runtime/Renderer/Private/SceneRendering.h
@@ -741,6 +741,9 @@ struct FPreviousViewInfo
 	// History for the ambient occlusion
 	FScreenSpaceFilteringHistory AmbientOcclusionHistory;
 
+	// History for global illumination
+	FScreenSpaceFilteringHistory GlobalIlluminationHistory;
+
 	// History for shadow denoising.
 	TMap<const ULightComponent*, FScreenSpaceFilteringHistory> ShadowHistories;
 
diff --git a/Engine/Source/Runtime/Renderer/Private/ScreenSpaceDenoise.cpp b/Engine/Source/Runtime/Renderer/Private/ScreenSpaceDenoise.cpp
--- a/Engine/Source/Runtime/Renderer/Private/ScreenSpaceDenoise.cpp
+++ b/Engine/Source/Runtime/Renderer/Private/ScreenSpaceDenoise.cpp
@@ -63,6 +63,21 @@ static TAutoConsoleVariable<int32> CVarAOHistoryConvolutionSampleCount(
 	TEXT("Number of samples to use for history post filter (default = 1)."),
 	ECVF_RenderThreadSafe);
 
+static TAutoConsoleVariable<int32> CVarGIReconstructionSampleCount(
+	TEXT("r.GlobalIllumination.Denoiser.ReconstructionSamples"), 16,
+	TEXT("Maximum number of samples for the reconstruction pass (default = 16)."),
+	ECVF_RenderThreadSafe);
+
+static TAutoConsoleVariable<int32> CVarGITemporalAccumulation(
+	TEXT("r.GlobalIllumination.Denoiser.TemporalAccumulation"), 1,
+	TEXT("Accumulates the samples over multiple frames."),
+	ECVF_RenderThreadSafe);
+
+static TAutoConsoleVariable<int32> CVarGIHistoryConvolutionSampleCount(
+	TEXT("r.GlobalIllumination.Denoiser.HistoryConvolutionSampleCount"), 1,
+	TEXT("Number of samples to use for history post filter (default = 1)."),
+	ECVF_RenderThreadSafe);
+
 
 /** The maximum number of mip level supported in the denoiser. */
 static const int32 kMaxMipLevel = 4;
@@ -96,6 +111,7 @@ enum class ESignalProcessing
 	Penumbra,
 	Reflections,
 	AmbientOcclusion,
+	GlobalIllumination,
 
 	MAX,
 };
@@ -711,6 +727,12 @@ static void DenoiseSignalAtConstantPixelDensity(
 		{
 			SignalProcessingDesc[0].Format = PF_G16R16F;
 		}
+		else if (Settings.SignalProcessing == ESignalProcessing::GlobalIllumination)
+		{
+			SignalProcessingDesc[0].Format = PF_FloatRGBA;
+			SignalProcessingDesc[1].Format = PF_G16R16F;
+			SignalTextureCount = 2;
+		}
 		else
 		{
 			check(0);
@@ -737,6 +759,10 @@ static void DenoiseSignalAtConstantPixelDensity(
 	{
 		PrevFrameHistory = View.PrevViewInfo.AmbientOcclusionHistory;
 	}
+	else if (Settings.SignalProcessing == ESignalProcessing::GlobalIllumination)
+	{
+		PrevFrameHistory = View.PrevViewInfo.GlobalIlluminationHistory;
+	}
 	else
 	{
 		check(0);
@@ -876,6 +902,8 @@ static void DenoiseSignalAtConstantPixelDensity(
 				NewHistory = &View.ViewState->PendingPrevFrameViewInfo.ReflectionsHistory;
 			else if (Settings.SignalProcessing == ESignalProcessing::AmbientOcclusion)
 				NewHistory = &View.ViewState->PendingPrevFrameViewInfo.AmbientOcclusionHistory;
+			else if (Settings.SignalProcessing == ESignalProcessing::GlobalIllumination)
+				NewHistory = &View.ViewState->PendingPrevFrameViewInfo.GlobalIlluminationHistory;
 			else
 				check(0);
 
@@ -1017,6 +1045,35 @@ public:
 		return AmbientOcclusionOutput;
 	}
 
+	FGlobalIlluminationOutputs DenoiseGlobalIllumination(
+		FRDGBuilder& GraphBuilder,
+		const FViewInfo& View,
+		const FSceneViewFamilyBlackboard& SceneBlackboard,
+		const FGlobalIlluminationInputs& Inputs,
+		const FAmbientOcclusionRayTracingConfig Config) const override
+	{
+		FSSDSignalTextures InputSignal;
+		InputSignal.Texture0 = Inputs.Color;
+		InputSignal.Texture1 = Inputs.RayHitDistance;
+
+		FSSDConstantPixelDensitySettings Settings;
+		Settings.SignalProcessing = ESignalProcessing::GlobalIllumination;
+		Settings.InputResolutionFraction = Config.ResolutionFraction;
+		Settings.ReconstructionSamples = CVarGIReconstructionSampleCount.GetValueOnRenderThread();
+		Settings.bUseTemporalAccumulation = CVarGITemporalAccumulation.GetValueOnRenderThread() != 0;
+		Settings.HistoryConvolutionSampleCount = CVarGIHistoryConvolutionSampleCount.GetValueOnRenderThread();
+
+		FSSDSignalTextures SignalOutput;
+		DenoiseSignalAtConstantPixelDensity(
+			GraphBuilder, View, SceneBlackboard,
+			InputSignal, Settings,
+			&SignalOutput);
+
+		FGlobalIlluminationOutputs GlobalIlluminationOutputs;
+		GlobalIlluminationOutputs.Color = SignalOutput.Texture0;
+		return GlobalIlluminationOutputs;
+	}
+
 }; // class FDefaultScreenSpaceDenoiser
 
 
diff --git a/Engine/Source/Runtime/Renderer/Private/ScreenSpaceDenoise.h b/Engine/Source/Runtime/Renderer/Private/ScreenSpaceDenoise.h
--- a/Engine/Source/Runtime/Renderer/Private/ScreenSpaceDenoise.h
+++ b/Engine/Source/Runtime/Renderer/Private/ScreenSpaceDenoise.h
@@ -115,6 +115,25 @@ public:
 		const FAmbientOcclusionInputs& ReflectionInputs,
 		const FAmbientOcclusionRayTracingConfig RayTracingConfig) const = 0;
 
+	/** All the inputs of the GI denoisers. */
+	BEGIN_SHADER_PARAMETER_STRUCT(FGlobalIlluminationInputs, )
+		SHADER_PARAMETER_RDG_TEXTURE(Texture2D, Color)
+		SHADER_PARAMETER_RDG_TEXTURE(Texture2D, RayHitDistance)
+	END_SHADER_PARAMETER_STRUCT()
+
+	/** All the outputs of the GI denoiser may generate. */
+	BEGIN_SHADER_PARAMETER_STRUCT(FGlobalIlluminationOutputs, )
+		SHADER_PARAMETER_RDG_TEXTURE(Texture2D, Color)
+	END_SHADER_PARAMETER_STRUCT()
+
+	/** Entry point to denoise reflections. */
+	virtual FGlobalIlluminationOutputs DenoiseGlobalIllumination(
+		FRDGBuilder& GraphBuilder,
+		const FViewInfo& View,
+		const FSceneViewFamilyBlackboard& SceneBlackboard,
+		const FGlobalIlluminationInputs& Inputs,
+		const FAmbientOcclusionRayTracingConfig Config) const = 0;
+
 	/** Returns the interface of the default denoiser of the renderer. */
 	static const IScreenSpaceDenoiser* GetDefaultDenoiser();
 }; // class FScreenSpaceDenoisingInterface
