Fix shadow artifact with raytraced directional occlusion.

Add a normal bias based on the (shading) normal orientation to compensate the mismatch betweeen shading surface & geometric surface.

Expose a cvar for settings the desire bias. By default set to 3mm as it seems to be an arbitrary good value (tested on several assets).

diff --git a/Engine/Shaders/Private/RayTracing/RayTracingDirectionalLight.ush b/Engine/Shaders/Private/RayTracing/RayTracingDirectionalLight.ush
--- a/Engine/Shaders/Private/RayTracing/RayTracingDirectionalLight.ush
+++ b/Engine/Shaders/Private/RayTracing/RayTracingDirectionalLight.ush
@@ -42,6 +42,7 @@ float2 ToConcentricMap(float2 RectangularCoords)
 void GenerateDirectionalLightOcclusionRay(
 	float3 WorldPosition,
 	float3 WorldNormal,
+	float MaxNormalBias,
 	float3 LightDirection,
 	float LightRadius,
 	float2 RandSample,
@@ -69,9 +70,20 @@ void GenerateDirectionalLightOcclusionRay(
 	float3 dPdv = cross(dPdu, N);
 	LightDirection += dPdu * DiskUV.x + dPdv * DiskUV.y;
 
-	// Apply normal perturbation when defining ray
-	RayOrigin = WorldPosition + WorldNormal * 0.01;
-	RayDirection = normalize(LightDirection);
+	// Apply normal perturbation when defining ray to:
+	// * avoid self intersection with current underlying triangle
+	// * hide mismatch between shading surface & geometric surface
+	//
+	// While using shading normal is not correct (we should use the 
+	// geometry normal, but it is not available atm/too costly to compute), 
+	// it is good enough for a cheap solution
+    const float3 NormalizedLightDirection = normalize(LightDirection);
+    const float MinBias = 0.01f;
+    const float MaxBias = max(MinBias, MaxNormalBias);
+    const float NormalBias = lerp(MaxBias, MinBias, saturate(dot(WorldNormal, NormalizedLightDirection)));
+
+    RayOrigin = WorldPosition + WorldNormal * NormalBias;
+    RayDirection = NormalizedLightDirection;
 	RayTMin = 0.0;
 	RayTMax = 1.0e27;
 }
diff --git a/Engine/Shaders/Private/RayTracing/RayTracingOcclusionRGS.usf b/Engine/Shaders/Private/RayTracing/RayTracingOcclusionRGS.usf
--- a/Engine/Shaders/Private/RayTracing/RayTracingOcclusionRGS.usf
+++ b/Engine/Shaders/Private/RayTracing/RayTracingOcclusionRGS.usf
@@ -20,6 +20,7 @@ RWTexture2D<float4> RWOcclusionMaskUAV;
 RWTexture2D<float> RWRayDistanceUAV;
 
 uint SamplesPerPixel;
+float NormalBias;
 
 float2 SobolIndexToUniformUnitSquare(uint2 SobolRand)
 {
@@ -111,7 +112,7 @@ bool GenerateOcclusionRay(
 		float LightRadius = LightParameters.SourceRadius;
 
 		GenerateDirectionalLightOcclusionRay(
-			WorldPosition, WorldNormal,
+			WorldPosition, WorldNormal, NormalBias,
 			LightDirection, LightRadius,
 			RandSample,
 			/* out */ RayOrigin,
diff --git a/Engine/Source/Runtime/Renderer/Private/RayTracing/RayTracingShadows.cpp b/Engine/Source/Runtime/Renderer/Private/RayTracing/RayTracingShadows.cpp
--- a/Engine/Source/Runtime/Renderer/Private/RayTracing/RayTracingShadows.cpp
+++ b/Engine/Source/Runtime/Renderer/Private/RayTracing/RayTracingShadows.cpp
@@ -27,6 +27,12 @@ static FAutoConsoleVariableRef CVarRayTracingOcclusionSamplesPerPixel(
 	TEXT("Sets the samples-per-pixel for directional light occlusion (default = 1)")
 );
 
+static float GRayTracingOcclusionNormalBias = 0.3f;
+static FAutoConsoleVariableRef CVarRayTracingOcclusionNormalBias(
+	TEXT("r.Shadow.RayTracing.NormalBias"),
+	GRayTracingOcclusionNormalBias,
+	TEXT("Sets the max. normal bias used for offseting the ray start position along the normal for directional light occlusion (default = 0.3)")
+);
 static TAutoConsoleVariable<int32> CVarShadowUseDenoiser(
 	TEXT("r.Shadow.Denoiser"),
 	2,
@@ -52,7 +58,7 @@ class FOcclusionRGS : public FGlobalShader
 
 	BEGIN_SHADER_PARAMETER_STRUCT(FParameters, )
 		SHADER_PARAMETER(uint32, SamplesPerPixel)
-
+		SHADER_PARAMETER(float, NormalBias)
 		SHADER_PARAMETER_STRUCT(FLightShaderParameters, Light)
 
 		SHADER_PARAMETER_SRV(RaytracingAccelerationStructure, TLAS)
@@ -108,6 +114,7 @@ void FDeferredShadingSceneRenderer::RenderRayTracingOcclusion(
 		PassParameters->RWOcclusionMaskUAV = GraphBuilder.CreateUAV(FRDGTextureUAVDesc(ScreenShadowMaskTexture));
 		PassParameters->RWRayDistanceUAV = GraphBuilder.CreateUAV(FRDGTextureUAVDesc(RayDistanceTexture));
 		PassParameters->SamplesPerPixel = GRayTracingOcclusionSamplesPerPixel;
+		PassParameters->NormalBias = GRayTracingOcclusionNormalBias;
 		LightSceneProxy->GetLightShaderParameters(PassParameters->Light);
 		PassParameters->TLAS = View.PerViewRayTracingScene.RayTracingSceneRHI->GetShaderResourceView();
 		PassParameters->ViewUniformBuffer = View.ViewUniformBuffer;
