Fixed linux build due to incorrect header name case

diff --git a/Engine/Source/Runtime/Renderer/Private/CompositionLighting/CompositionLighting.cpp b/Engine/Source/Runtime/Renderer/Private/CompositionLighting/CompositionLighting.cpp
--- a/Engine/Source/Runtime/Renderer/Private/CompositionLighting/CompositionLighting.cpp
+++ b/Engine/Source/Runtime/Renderer/Private/CompositionLighting/CompositionLighting.cpp
@@ -17,7 +17,7 @@
 #include "LightPropagationVolumeSettings.h"
 #include "DecalRenderingShared.h"
 #include "VisualizeTexture.h"
-#include "RayTracing/RayTracingOptions.h"
+#include "RayTracing/RaytracingOptions.h"
 
 /** The global center for all deferred lighting activities. */
 FCompositionLighting GCompositionLighting;
diff --git a/Engine/Source/Runtime/Renderer/Private/DeferredShadingRenderer.cpp b/Engine/Source/Runtime/Renderer/Private/DeferredShadingRenderer.cpp
--- a/Engine/Source/Runtime/Renderer/Private/DeferredShadingRenderer.cpp
+++ b/Engine/Source/Runtime/Renderer/Private/DeferredShadingRenderer.cpp
@@ -33,7 +33,7 @@
 #include "RayTracing/RayTracingDynamicGeometryCollection.h"
 #include "SceneViewFamilyBlackboard.h"
 #include "ScreenSpaceDenoise.h"
-#include "RayTracing/RayTracingOptions.h"
+#include "RayTracing/RaytracingOptions.h"
 
 TAutoConsoleVariable<int32> CVarEarlyZPass(
 	TEXT("r.EarlyZPass"),
