Dead code removal

diff --git a/Engine/Source/Runtime/Renderer/Private/RayTracing/RayTracingAmbientOcclusion.cpp b/Engine/Source/Runtime/Renderer/Private/RayTracing/RayTracingAmbientOcclusion.cpp
--- a/Engine/Source/Runtime/Renderer/Private/RayTracing/RayTracingAmbientOcclusion.cpp
+++ b/Engine/Source/Runtime/Renderer/Private/RayTracing/RayTracingAmbientOcclusion.cpp
@@ -35,13 +35,6 @@ static FAutoConsoleVariableRef CVarRayTracingAmbientOcclusionSamplesPerPixel(
 	TEXT("Sets the samples-per-pixel for ambient occlusion (default = 1)")
 );
 
-static float GRayTracingAmbientOcclusionMaxRayDistance = 1.0e27;
-static FAutoConsoleVariableRef CVarRayTracingAmbientOcclusionMaxRayDistance(
-	TEXT("r.RayTracing.AmbientOcclusion.MaxRayDistance"),
-	GRayTracingAmbientOcclusionMaxRayDistance,
-	TEXT("Sets the maximum ray distance for ambient occlusion rays (default = 1.0e27)")
-);
-
 BEGIN_GLOBAL_SHADER_PARAMETER_STRUCT(FAmbientOcclusionData, )
 SHADER_PARAMETER(int, SamplesPerPixel)
 SHADER_PARAMETER(float, MaxRayDistance)
@@ -51,7 +44,6 @@ END_GLOBAL_SHADER_PARAMETER_STRUCT()
 IMPLEMENT_GLOBAL_SHADER_PARAMETER_STRUCT(FAmbientOcclusionData, "AmbientOcclusion");
 
 DECLARE_GPU_STAT_NAMED(RayTracingAmbientOcclusion, TEXT("Ray Tracing Ambient Occlusion"));
-DECLARE_GPU_STAT_NAMED(CompositeRayTracingAmbientOcclusion, TEXT("Denoising: RTAO"));
 
 class FAmbientOcclusionRGS : public FGlobalShader
 {
@@ -155,10 +147,6 @@ void FDeferredShadingSceneRenderer::RenderRayTracingAmbientOcclusion(
 	GRenderTargetPool.FindFreeElement(RHICmdList, Desc, AmbientOcclusionMask, TEXT("RayTracingAmbientOcclusion"));
 	GRenderTargetPool.FindFreeElement(RHICmdList, Desc, HitDistance, TEXT("RayTracingAmbientOcclusionHitDistance"));
 
-	// TODO: ClearUAV is useless.
-	ClearUAV(RHICmdList, AmbientOcclusionMask->GetRenderTargetItem(), FLinearColor::Black);
-	ClearUAV(RHICmdList, HitDistance->GetRenderTargetItem(), FLinearColor::Black);
-
 	// Add ambient occlusion parameters to uniform buffer
 	FAmbientOcclusionData AmbientOcclusionData;
 	AmbientOcclusionData.SamplesPerPixel = GRayTracingAmbientOcclusionSamplesPerPixel;
@@ -189,122 +177,17 @@ void FDeferredShadingSceneRenderer::RenderRayTracingAmbientOcclusion(
 		);
 	}
 
-	// Transition to graphics pipeline
+	FUnorderedAccessViewRHIParamRef UAVs[]
+	{
+		AmbientOcclusionMask->GetRenderTargetItem().UAV,
+		HitDistance->GetRenderTargetItem().UAV
+	};
 	FComputeFenceRHIRef Fence = RHICmdList.CreateComputeFence(TEXT("RayTracingAmbientOcclusion"));
-	RHICmdList.TransitionResource(EResourceTransitionAccess::ERWBarrier, EResourceTransitionPipeline::EComputeToGfx, AmbientOcclusionMask->GetRenderTargetItem().UAV, Fence);
+	RHICmdList.TransitionResources(EResourceTransitionAccess::ERWBarrier, EResourceTransitionPipeline::EComputeToGfx, UAVs, ARRAY_COUNT(UAVs), Fence);
 
-	FComputeFenceRHIRef Fence2 = RHICmdList.CreateComputeFence(TEXT("RayTracingAmbientOcclusion"));
-	RHICmdList.TransitionResource(EResourceTransitionAccess::ERWBarrier, EResourceTransitionPipeline::EComputeToGfx, HitDistance->GetRenderTargetItem().UAV, Fence2);
-	
 	GVisualizeTexture.SetCheckPoint(RHICmdList, AmbientOcclusionMask);
 	GVisualizeTexture.SetCheckPoint(RHICmdList, HitDistance);
 	SceneContext.bScreenSpaceAOIsValid = true;
 }
 
-class FCompositeAmbientOcclusionPS : public FGlobalShader
-{
-	DECLARE_SHADER_TYPE(FCompositeAmbientOcclusionPS, Global)
-
-	static void ModifyCompilationEnvironment(const FGlobalShaderPermutationParameters& Parameters, FShaderCompilerEnvironment& OutEnvironment)
-	{
-	}
-
-public:
-	static bool ShouldCompilePermutation(const FGlobalShaderPermutationParameters& Parameters)
-	{
-		return ShouldCompileRayTracingShadersForProject(Parameters.Platform);
-	}
-
-	FCompositeAmbientOcclusionPS() {}
-	virtual ~FCompositeAmbientOcclusionPS() {}
-
-	FCompositeAmbientOcclusionPS(const ShaderMetaType::CompiledShaderInitializerType& Initializer)
-		: FGlobalShader(Initializer)
-	{
-		SceneTextureParameters.Bind(Initializer);
-		AmbientOcclusionTextureParameter.Bind(Initializer.ParameterMap, TEXT("AmbientOcclusionTexture"));
-		AmbientOcclusionTextureSamplerParameter.Bind(Initializer.ParameterMap, TEXT("AmbientOcclusionTextureSampler"));
-	}
-
-	bool Serialize(FArchive& Ar)
-	{
-		bool bShaderHasOutdatedParameters = FGlobalShader::Serialize(Ar);
-		Ar << SceneTextureParameters;
-		Ar << AmbientOcclusionTextureParameter;
-		Ar << AmbientOcclusionTextureSamplerParameter;
-		return bShaderHasOutdatedParameters;
-	}
-
-	template<typename TRHICommandList>
-	void SetParameters(
-		TRHICommandList& RHICmdList,
-		const FViewInfo& View,
-		FTextureRHIParamRef AmbientOcclusionTexture
-	)
-	{
-		const FPixelShaderRHIParamRef ShaderRHI = GetPixelShader();
-		FGlobalShader::SetParameters<FViewUniformShaderParameters>(RHICmdList, ShaderRHI, View.ViewUniformBuffer);
-		SceneTextureParameters.Set(RHICmdList, ShaderRHI, View.FeatureLevel, ESceneTextureSetupMode::All);
-
-		SetTextureParameter(RHICmdList, ShaderRHI, AmbientOcclusionTextureParameter, AmbientOcclusionTextureSamplerParameter, TStaticSamplerState<SF_Bilinear>::GetRHI(), AmbientOcclusionTexture);
-	}
-
-private:
-	FSceneTextureShaderParameters SceneTextureParameters;
-
-	FShaderResourceParameter AmbientOcclusionTextureParameter;
-	FShaderResourceParameter AmbientOcclusionTextureSamplerParameter;
-};
-
-IMPLEMENT_SHADER_TYPE(, FCompositeAmbientOcclusionPS, TEXT("/Engine/Private/RayTracing/CompositeAmbientOcclusionPS.usf"), TEXT("CompositeAmbientOcclusionPS"), SF_Pixel)
-
-void FDeferredShadingSceneRenderer::CompositeRayTracingAmbientOcclusion(
-	FRHICommandListImmediate& RHICmdList,
-	TRefCountPtr<IPooledRenderTarget>& AmbientOcclusionRT
-)
-{
-	SCOPED_DRAW_EVENT(RHICmdList, CompositeRayTracingAmbientOcclusion);
-	SCOPED_GPU_STAT(RHICmdList, CompositeRayTracingAmbientOcclusion);
-
-	auto ShaderMap = GetGlobalShaderMap(FeatureLevel);
-	TShaderMapRef<FPostProcessVS> VertexShader(ShaderMap);
-	TShaderMapRef<FCompositeAmbientOcclusionPS> PixelShader(ShaderMap);
-	FGraphicsPipelineStateInitializer GraphicsPSOInit;
-
-	FSceneRenderTargets& SceneContext = FSceneRenderTargets::Get(RHICmdList);
-	SceneContext.BeginRenderingSceneColor(RHICmdList, ESimpleRenderTargetMode::EExistingColorAndDepth, FExclusiveDepthStencil::DepthRead_StencilWrite, true);
-	RHICmdList.ApplyCachedRenderTargets(GraphicsPSOInit);
-
-	// Multiply scene color by ambient term
-	GraphicsPSOInit.BlendState = TStaticBlendState<CW_RGBA, BO_Add, BF_Zero, BF_SourceColor, BO_Add, BF_One, BF_One>::GetRHI();
-	GraphicsPSOInit.RasterizerState = TStaticRasterizerState<FM_Solid, CM_None>::GetRHI();
-	GraphicsPSOInit.DepthStencilState = TStaticDepthStencilState<false, CF_Always>::GetRHI();
-	GraphicsPSOInit.BoundShaderState.VertexDeclarationRHI = GFilterVertexDeclaration.VertexDeclarationRHI;
-	GraphicsPSOInit.BoundShaderState.VertexShaderRHI = GETSAFERHISHADER_VERTEX(*VertexShader);
-	GraphicsPSOInit.BoundShaderState.PixelShaderRHI = GETSAFERHISHADER_PIXEL(*PixelShader);
-	GraphicsPSOInit.PrimitiveType = PT_TriangleList;
-	SetGraphicsPipelineState(RHICmdList, GraphicsPSOInit);
-
-	for (int32 ViewIndex = 0; ViewIndex < Views.Num(); ViewIndex++)
-	{
-		const FViewInfo& View = Views[ViewIndex];
-		RHICmdList.SetViewport(View.ViewRect.Min.X, View.ViewRect.Min.Y, 0.0f, View.ViewRect.Max.X, View.ViewRect.Max.Y, 1.0f);
-		FTextureRHIParamRef AmbientOcclusionTexture = AmbientOcclusionRT->GetRenderTargetItem().ShaderResourceTexture;
-		PixelShader->SetParameters(RHICmdList, View, AmbientOcclusionTexture);
-		DrawRectangle(
-			RHICmdList,
-			0, 0,
-			View.ViewRect.Width(), View.ViewRect.Height(),
-			View.ViewRect.Min.X, View.ViewRect.Min.Y,
-			View.ViewRect.Width(), View.ViewRect.Height(),
-			FIntPoint(View.ViewRect.Width(), View.ViewRect.Height()),
-			SceneContext.GetBufferSizeXY(),
-			*VertexShader
-		);
-	}
-
-	ResolveSceneColor(RHICmdList);
-	SceneContext.FinishRenderingSceneColor(RHICmdList);
-}
-
 #endif // RHI_RAYTRACING
