Turns on history convolution by default in the denoiser for AO and GI.

diff --git a/Engine/Shaders/Private/ScreenSpaceDenoise/SSDSignalFramework.ush b/Engine/Shaders/Private/ScreenSpaceDenoise/SSDSignalFramework.ush
--- a/Engine/Shaders/Private/ScreenSpaceDenoise/SSDSignalFramework.ush
+++ b/Engine/Shaders/Private/ScreenSpaceDenoise/SSDSignalFramework.ush
@@ -402,8 +402,8 @@ float ComputeBilateralWeight(
 		return 1;
 	}
 	
-	float ViewPositionBilateralWeight = 1;
-	#if CONFIG_VIEW_POSITION_BASED_BILATERAL
+	float PositionBilateralWeight = 1;
+	#if CONFIG_VIEW_POSITION_BASED_BILATERAL == 1
 	{
 		const float WorldRadius = MaxWorldBluringDistance;
 
@@ -420,7 +420,13 @@ float ComputeBilateralWeight(
 
 		float Multiplier = rcp(WorldRadius * WorldRadius);
 
-		ViewPositionBilateralWeight = saturate(1 - DistSquare * Multiplier);
+		PositionBilateralWeight = saturate(1 - DistSquare * Multiplier);
+	}
+	#elif CONFIG_VIEW_POSITION_BASED_BILATERAL == 2
+	{
+		float DistFromRefPlane = dot(RefSceneMetadata.WorldNormal, RefSceneMetadata.TranslatedWorldPosition - NeighborSceneMetadata.TranslatedWorldPosition);
+
+		PositionBilateralWeight = saturate(1 - DistFromRefPlane * rcp(MaxWorldBluringDistance));
 	}
 	#endif
 
@@ -464,7 +470,7 @@ float ComputeBilateralWeight(
 	}
 	#endif
 
-	return ViewPositionBilateralWeight * NormalBilateralWeight * LobeSimilarity * AxesSimilarity;
+	return PositionBilateralWeight * NormalBilateralWeight * LobeSimilarity * AxesSimilarity;
 }
 
 float GetSignalWorldBluringRadius(FSSFSignalSample Sample, FSSFSampleSceneInfos SceneMetadata)
diff --git a/Engine/Shaders/Private/ScreenSpaceDenoise/SSDSpatialAccumulation.usf b/Engine/Shaders/Private/ScreenSpaceDenoise/SSDSpatialAccumulation.usf
--- a/Engine/Shaders/Private/ScreenSpaceDenoise/SSDSpatialAccumulation.usf
+++ b/Engine/Shaders/Private/ScreenSpaceDenoise/SSDSpatialAccumulation.usf
@@ -20,6 +20,7 @@
 	#define SAMPLE_COUNT_POLICY_METADATA_BASED 1
 	#define SAMPLE_COUNT_POLICY_HIT_DISTANCE_BASED 2
 	#define SAMPLE_COUNT_POLICY_ACCUMULATION_BASED 3
+	#define SAMPLE_COUNT_POLICY_SAMPLE_ACCUMULATION_BASED 4
 
 
 //------------------------------------------------------- CONFIGS
@@ -74,9 +75,9 @@
 		#define CONFIG_SIGNAL_INPUT_LAYOUT  SIGNAL_BUFFER_LAYOUT_REFLECTIONS_HISTORY
 		#define CONFIG_SIGNAL_OUTPUT_LAYOUT SIGNAL_BUFFER_LAYOUT_REFLECTIONS_HISTORY
 
-		#define CONFIG_SAMPLE_SET           SAMPLE_SET_3X3
-		//#define CONFIG_SAMPLE_SET           SAMPLE_SET_STACKOWIAK_4_SETS
-		//#define CONFIG_SAMPLE_SUBSET        1
+		//#define CONFIG_SAMPLE_SET           SAMPLE_SET_3X3
+		#define CONFIG_SAMPLE_SET           SAMPLE_SET_STACKOWIAK_4_SETS
+		#define CONFIG_SAMPLE_SUBSET        1
 		#define CONFIG_SAMPLE_COUNT_POLICY  SAMPLE_COUNT_POLICY_ACCUMULATION_BASED
 
 		// Can only be done using tokoyashi because have more than one sample at a time.
@@ -88,7 +89,7 @@
 	#endif
 
 #elif CONFIG_SIGNAL_PROCESSING == SIGNAL_PROCESSING_AO
-	#define CONFIG_VIEW_POSITION_BASED_BILATERAL 1
+	#define CONFIG_VIEW_POSITION_BASED_BILATERAL 2
 	#define CONFIG_NORMAL_BASED_BILATERAL 1
 
 	#if DIM_STAGE == STAGE_RECONSTRUCTION
@@ -100,17 +101,20 @@
 
 	#elif DIM_STAGE == STAGE_POST_FILTERING
 		// Input and output layout.
-		#define CONFIG_SIGNAL_INPUT_LAYOUT  SIGNAL_BUFFER_LAYOUT_AO_INPUT
+		#define CONFIG_SIGNAL_INPUT_LAYOUT  SIGNAL_BUFFER_LAYOUT_AO_HISTORY
 		#define CONFIG_SIGNAL_OUTPUT_LAYOUT SIGNAL_BUFFER_LAYOUT_AO_HISTORY
 
-		#define CONFIG_SAMPLE_SET           SAMPLE_SET_3X3
+		#define CONFIG_SAMPLE_SET           SAMPLE_SET_STACKOWIAK_4_SETS
+		#define CONFIG_SAMPLE_SUBSET        1
+
+		#define CONFIG_SAMPLE_COUNT_POLICY  SAMPLE_COUNT_POLICY_SAMPLE_ACCUMULATION_BASED
 
 	#else
 		#error Unexpected stage.
 	#endif
 
 #elif CONFIG_SIGNAL_PROCESSING == SIGNAL_PROCESSING_GI_AND_AO
-	#define CONFIG_VIEW_POSITION_BASED_BILATERAL 1
+	#define CONFIG_VIEW_POSITION_BASED_BILATERAL 2
 	#define CONFIG_NORMAL_BASED_BILATERAL 1
 
 	#if DIM_STAGE == STAGE_RECONSTRUCTION
@@ -122,10 +126,13 @@
 
 	#elif DIM_STAGE == STAGE_POST_FILTERING
 		// Input and output layout.
-		#define CONFIG_SIGNAL_INPUT_LAYOUT  SIGNAL_BUFFER_LAYOUT_GI_AND_AO_INPUT
+		#define CONFIG_SIGNAL_INPUT_LAYOUT  SIGNAL_BUFFER_LAYOUT_GI_AND_AO_HISTORY
 		#define CONFIG_SIGNAL_OUTPUT_LAYOUT SIGNAL_BUFFER_LAYOUT_GI_AND_AO_HISTORY
 
-		#define CONFIG_SAMPLE_SET           SAMPLE_SET_3X3
+		#define CONFIG_SAMPLE_SET           SAMPLE_SET_STACKOWIAK_4_SETS
+		#define CONFIG_SAMPLE_SUBSET        1
+
+		#define CONFIG_SAMPLE_COUNT_POLICY  SAMPLE_COUNT_POLICY_SAMPLE_ACCUMULATION_BASED
 
 	#else
 		#error Unexpected stage.
@@ -314,6 +321,10 @@ static const uint kStackowiakSampleCountPerSet = 56;
 uint MaxSampleCount;
 uint UpscaleFactor;
 
+#if !CONFIG_UPSCALE && CONFIG_CUSTOM_SPREAD_FACTOR
+	float KernelSpreadFactor;
+#endif
+
 Texture2D SignalInput0;
 Texture2D SignalInput1;
 
@@ -410,9 +421,10 @@ void MainCS(
 	#endif
 	
 	/** factor by witch should be spread out. */
-	float KernelSpreadFactor = 1;
 	#if CONFIG_UPSCALE
-		KernelSpreadFactor = UpscaleFactor;
+		float KernelSpreadFactor = UpscaleFactor;
+	#elif !CONFIG_CUSTOM_SPREAD_FACTOR
+		const float KernelSpreadFactor = 1;
 	#endif
 
 	/** Find out the number of samples that should be done. */
@@ -443,16 +455,19 @@ void MainCS(
 		IdealSampleCount = min(IdealSampleCount, MaxSampleCount);
 
 		RequestedSampleCount = IdealSampleCount / clamp(RefSample.AccumulatedSampleCount, 1, MaxSampleCount);
-		
-		DebugOutput[DispatchThreadId] = float4(IdealSampleCount, RequestedSampleCount, 0, 0);
+	}
+	#elif CONFIG_SAMPLE_COUNT_POLICY == SAMPLE_COUNT_POLICY_SAMPLE_ACCUMULATION_BASED
+	{
+		const float IdealSampleCount = TARGETED_SAMPLE_COUNT;
+
+		RequestedSampleCount = clamp(IdealSampleCount / RefSample.AccumulatedSampleCount, 1, MaxSampleCount);
 	}
 	#else
 		#error Unknown policy to control the number of samples.
 	#endif
 
 
-	/** Iterate over the samples of the kernel.
-	 */
+	/** Iterate over the samples of the kernel. */
 	float2 KernelCenterBufferUV = SceneBufferUV;
 
 	#if CONFIG_SAMPLE_SET == SAMPLE_SET_NONE
diff --git a/Engine/Source/Runtime/Renderer/Private/ScreenSpaceDenoise.cpp b/Engine/Source/Runtime/Renderer/Private/ScreenSpaceDenoise.cpp
--- a/Engine/Source/Runtime/Renderer/Private/ScreenSpaceDenoise.cpp
+++ b/Engine/Source/Runtime/Renderer/Private/ScreenSpaceDenoise.cpp
@@ -44,7 +44,7 @@ static TAutoConsoleVariable<int32> CVarReflectionTemporalAccumulation(
 	ECVF_RenderThreadSafe);
 
 static TAutoConsoleVariable<int32> CVarReflectionHistoryConvolutionSampleCount(
-	TEXT("r.Reflections.Denoiser.HistoryConvolutionSampleCount"), 1,
+	TEXT("r.Reflections.Denoiser.HistoryConvolution.SampleCount"), 1,
 	TEXT("Number of samples to use for history post filter (default = 1)."),
 	ECVF_RenderThreadSafe);
 
@@ -59,8 +59,13 @@ static TAutoConsoleVariable<int32> CVarAOTemporalAccumulation(
 	ECVF_RenderThreadSafe);
 
 static TAutoConsoleVariable<int32> CVarAOHistoryConvolutionSampleCount(
-	TEXT("r.AmbientOcclusion.Denoiser.HistoryConvolutionSampleCount"), 1,
-	TEXT("Number of samples to use for history post filter (default = 1)."),
+	TEXT("r.AmbientOcclusion.Denoiser.HistoryConvolution.SampleCount"), 16,
+	TEXT("Number of samples to use for history post filter (default = 16)."),
+	ECVF_RenderThreadSafe);
+
+static TAutoConsoleVariable<float> CVarAOHistoryConvolutionKernelSpreadFactor(
+	TEXT("r.AmbientOcclusion.Denoiser.HistoryConvolution.KernelSpreadFactor"), 3,
+	TEXT("Multiplication factor applied on the kernel sample offset (default=3)."),
 	ECVF_RenderThreadSafe);
 
 static TAutoConsoleVariable<int32> CVarGIReconstructionSampleCount(
@@ -74,10 +79,15 @@ static TAutoConsoleVariable<int32> CVarGITemporalAccumulation(
 	ECVF_RenderThreadSafe);
 
 static TAutoConsoleVariable<int32> CVarGIHistoryConvolutionSampleCount(
-	TEXT("r.GlobalIllumination.Denoiser.HistoryConvolutionSampleCount"), 1,
+	TEXT("r.GlobalIllumination.Denoiser.HistoryConvolution.SampleCount"), 16,
 	TEXT("Number of samples to use for history post filter (default = 1)."),
 	ECVF_RenderThreadSafe);
 
+static TAutoConsoleVariable<float> CVarGIHistoryConvolutionKernelSpreadFactor(
+	TEXT("r.GlobalIllumination.Denoiser.HistoryConvolution.KernelSpreadFactor"), 3,
+	TEXT("Multiplication factor applied on the kernel sample offset (default=3)."),
+	ECVF_RenderThreadSafe);
+
 
 /** The maximum number of mip level supported in the denoiser. */
 static const int32 kMaxMipLevel = 4;
@@ -269,6 +279,7 @@ class FSSDSpatialAccumulationCS : public FScreenSpaceDenoisingShader
 	BEGIN_SHADER_PARAMETER_STRUCT(FParameters, )
 		SHADER_PARAMETER(uint32, MaxSampleCount)
 		SHADER_PARAMETER(int32, UpscaleFactor)
+		SHADER_PARAMETER(float, KernelSpreadFactor)
 
 		SHADER_PARAMETER_STRUCT_INCLUDE(FSSDCommonParameters, CommonParameters)
 
@@ -686,6 +697,7 @@ struct FSSDConstantPixelDensitySettings
 	int32 ReconstructionSamples = 1;
 	bool bUseTemporalAccumulation = false;
 	int32 HistoryConvolutionSampleCount = 1;
+	float HistoryConvolutionKernelSpreadFactor = 1.0f;
 };
 
 
@@ -849,7 +861,7 @@ static void DenoiseSignalAtConstantPixelDensity(
 	}
 	
 	// Spatial filter, to converge history faster.
-	int32 MaxPostFilterSampleCount = FMath::Clamp(Settings.HistoryConvolutionSampleCount, 1, 16);
+	int32 MaxPostFilterSampleCount = FMath::Clamp(Settings.HistoryConvolutionSampleCount, 1, kStackowiakMaxSampleCountPerSet);
 	if (MaxPostFilterSampleCount > 1)
 	{
 		FRDGTextureRef SignalOutput0 = GraphBuilder.CreateTexture(SignalProcessingDesc[0], TEXT("SSDReflectionsHistory0"));
@@ -857,6 +869,7 @@ static void DenoiseSignalAtConstantPixelDensity(
 
 		FSSDSpatialAccumulationCS::FParameters* PassParameters = GraphBuilder.AllocParameters<FSSDSpatialAccumulationCS::FParameters>();
 		PassParameters->MaxSampleCount = FMath::Clamp(MaxPostFilterSampleCount, 1, kStackowiakMaxSampleCountPerSet);
+		PassParameters->KernelSpreadFactor = Settings.HistoryConvolutionKernelSpreadFactor;
 		PassParameters->CommonParameters = CommonParameters;
 		PassParameters->SignalInput0 = SignalHistory0;
 		PassParameters->SignalInput1 = SignalHistory1;
@@ -1033,6 +1046,7 @@ public:
 		Settings.ReconstructionSamples = CVarAOReconstructionSampleCount.GetValueOnRenderThread();
 		Settings.bUseTemporalAccumulation = CVarAOTemporalAccumulation.GetValueOnRenderThread() != 0;
 		Settings.HistoryConvolutionSampleCount = CVarAOHistoryConvolutionSampleCount.GetValueOnRenderThread();
+		Settings.HistoryConvolutionKernelSpreadFactor = CVarAOHistoryConvolutionKernelSpreadFactor.GetValueOnRenderThread();
 
 		FSSDSignalTextures SignalOutput;
 		DenoiseSignalAtConstantPixelDensity(
@@ -1062,6 +1076,7 @@ public:
 		Settings.ReconstructionSamples = CVarGIReconstructionSampleCount.GetValueOnRenderThread();
 		Settings.bUseTemporalAccumulation = CVarGITemporalAccumulation.GetValueOnRenderThread() != 0;
 		Settings.HistoryConvolutionSampleCount = CVarGIHistoryConvolutionSampleCount.GetValueOnRenderThread();
+		Settings.HistoryConvolutionKernelSpreadFactor = CVarGIHistoryConvolutionKernelSpreadFactor.GetValueOnRenderThread();
 
 		FSSDSignalTextures SignalOutput;
 		DenoiseSignalAtConstantPixelDensity(
