DR - Merge fix
#rnx

diff --git a/Engine/Source/Runtime/Renderer/Private/RendererScene.cpp b/Engine/Source/Runtime/Renderer/Private/RendererScene.cpp
--- a/Engine/Source/Runtime/Renderer/Private/RendererScene.cpp
+++ b/Engine/Source/Runtime/Renderer/Private/RendererScene.cpp
@@ -628,7 +628,7 @@ void FScene::DumpMeshDrawCommandMemoryStats()
 				PassStats.ShaderBindingHeapBytes += MeshDrawCommand.ShaderBindings.GetAllocatedSize();
 				PassStats.VertexStreamsInlineBytes += sizeof(MeshDrawCommand.VertexStreams);
 				PassStats.DebugDataBytes += MeshDrawCommand.GetDebugDataSize();
-				PassStats.DrawCommandParameterBytes += sizeof(MeshDrawCommand.IndexBuffer) + sizeof(MeshDrawCommand.FirstIndex) + sizeof(MeshDrawCommand.NumPrimitives) + sizeof(MeshDrawCommand.NumInstances) + sizeof(MeshDrawCommand.InstanceFactor) + sizeof(MeshDrawCommand.VertexParams);
+				PassStats.DrawCommandParameterBytes += sizeof(MeshDrawCommand.IndexBuffer) + sizeof(MeshDrawCommand.FirstIndex) + sizeof(MeshDrawCommand.NumPrimitives) + sizeof(MeshDrawCommand.NumInstances) + sizeof(MeshDrawCommand.VertexParams);
 				PassStats.NumCommands++;
 			}
 		}
