Removes an unecessary channel in AO+GI denoiser's history buffer.

diff --git a/Engine/Source/Runtime/Renderer/Private/ScreenSpaceDenoise.cpp b/Engine/Source/Runtime/Renderer/Private/ScreenSpaceDenoise.cpp
--- a/Engine/Source/Runtime/Renderer/Private/ScreenSpaceDenoise.cpp
+++ b/Engine/Source/Runtime/Renderer/Private/ScreenSpaceDenoise.cpp
@@ -742,7 +742,7 @@ static void DenoiseSignalAtConstantPixelDensity(
 		else if (Settings.SignalProcessing == ESignalProcessing::GlobalIllumination)
 		{
 			SignalProcessingDesc[0].Format = PF_FloatRGBA;
-			SignalProcessingDesc[1].Format = PF_G16R16F;
+			SignalProcessingDesc[1].Format = PF_R16F;
 			SignalTextureCount = 2;
 		}
 		else
