Properly names denoiser's intermediary textures for vis tool.

diff --git a/Engine/Source/Runtime/Renderer/Private/ScreenSpaceDenoise.cpp b/Engine/Source/Runtime/Renderer/Private/ScreenSpaceDenoise.cpp
--- a/Engine/Source/Runtime/Renderer/Private/ScreenSpaceDenoise.cpp
+++ b/Engine/Source/Runtime/Renderer/Private/ScreenSpaceDenoise.cpp
@@ -130,6 +130,65 @@ enum class ESignalProcessing
 class FSignalProcessingDim : SHADER_PERMUTATION_ENUM_CLASS("DIM_SIGNAL_PROCESSING", ESignalProcessing);
 
 
+const TCHAR* const kReconstructionResourceNames[] = {
+	// Penumbra
+	nullptr,
+	nullptr,
+
+	// Reflections
+	TEXT("ReflectionsReconstruction0"),
+	TEXT("ReflectionsReconstruction1"),
+
+	// AmbientOcclusion
+	TEXT("AOReconstruction0"),
+	TEXT("AOReconstruction1"),
+
+	// GlobalIllumination
+	TEXT("GIReconstruction0"),
+	TEXT("GIReconstruction1"),
+};
+
+const TCHAR* const kTemporalAccumulationResourceNames[] = {
+	// Penumbra
+	nullptr,
+	nullptr,
+
+	// Reflections
+	TEXT("ReflectionsTemporalAccumulation0"),
+	TEXT("ReflectionsTemporalAccumulation1"),
+
+	// AmbientOcclusion
+	TEXT("AOTemporalAccumulation0"),
+	TEXT("AOTemporalAccumulation1"),
+
+	// GlobalIllumination
+	TEXT("GITemporalAccumulation0"),
+	TEXT("GITemporalAccumulation1"),
+};
+
+const TCHAR* const kHistoryConvolutionResourceNames[] = {
+	// Penumbra
+	nullptr,
+	nullptr,
+
+	// Reflections
+	TEXT("ReflectionsHistoryConvolution0"),
+	TEXT("ReflectionsHistoryConvolution1"),
+
+	// AmbientOcclusion
+	TEXT("AOHistoryConvolution0"),
+	TEXT("AOHistoryConvolution1"),
+
+	// GlobalIllumination
+	TEXT("GIHistoryConvolution0"),
+	TEXT("GIHistoryConvolution1"),
+};
+
+static_assert(ARRAY_COUNT(kReconstructionResourceNames) == int32(ESignalProcessing::MAX) * 2, "You forgot me!");
+static_assert(ARRAY_COUNT(kTemporalAccumulationResourceNames) == int32(ESignalProcessing::MAX) * 2, "You forgot me!");
+static_assert(ARRAY_COUNT(kHistoryConvolutionResourceNames) == int32(ESignalProcessing::MAX) * 2, "You forgot me!");
+
+
 /** Base class for a screen space denoising shader. */
 class FScreenSpaceDenoisingShader : public FGlobalShader
 {
@@ -711,6 +770,12 @@ static void DenoiseSignalAtConstantPixelDensity(
 	FSSDSignalTextures* OutputSignal)
 {
 	ensure(Settings.InputResolutionFraction == 1.0f || Settings.InputResolutionFraction == 0.5f);
+	
+	auto GetResourceName = [&](const TCHAR* const ResourceNames[], int32 SubResourceId)
+	{
+		check(SubResourceId >= 0 && SubResourceId < 2);
+		return ResourceNames[int32(Settings.SignalProcessing) * 2 + SubResourceId];
+	};
 
 	const FIntPoint DenoiseResolution = View.ViewRect.Size();
 	
@@ -785,8 +850,8 @@ static void DenoiseSignalAtConstantPixelDensity(
 	
 	// Spatial reconstruction with multiple important sampling to be more precise in the history rejection.
 	{
-		FRDGTextureRef SignalOutput0 = GraphBuilder.CreateTexture(SignalProcessingDesc[0], TEXT("SSDReflectionsReconstruction0"));
-		FRDGTextureRef SignalOutput1 = GraphBuilder.CreateTexture(SignalProcessingDesc[1], TEXT("SSDReflectionsReconstruction1"));
+		FRDGTextureRef SignalOutput0 = GraphBuilder.CreateTexture(SignalProcessingDesc[0], GetResourceName(kReconstructionResourceNames, 0));
+		FRDGTextureRef SignalOutput1 = GraphBuilder.CreateTexture(SignalProcessingDesc[1], GetResourceName(kReconstructionResourceNames, 1));
 
 		FSSDSpatialAccumulationCS::FParameters* PassParameters = GraphBuilder.AllocParameters<FSSDSpatialAccumulationCS::FParameters>();
 		PassParameters->MaxSampleCount = FMath::Clamp(Settings.ReconstructionSamples, 1, kStackowiakMaxSampleCountPerSet);
@@ -820,8 +885,8 @@ static void DenoiseSignalAtConstantPixelDensity(
 	// Temporal pass.
 	if (View.ViewState && Settings.bUseTemporalAccumulation)
 	{
-		FRDGTextureRef SignalOutput0 = GraphBuilder.CreateTexture(SignalProcessingDesc[0], TEXT("SSDReflectionsHistory0"));
-		FRDGTextureRef SignalOutput1 = GraphBuilder.CreateTexture(SignalProcessingDesc[1], TEXT("SSDReflectionsHistory1"));
+		FRDGTextureRef SignalOutput0 = GraphBuilder.CreateTexture(SignalProcessingDesc[0], GetResourceName(kTemporalAccumulationResourceNames, 0));
+		FRDGTextureRef SignalOutput1 = GraphBuilder.CreateTexture(SignalProcessingDesc[1], GetResourceName(kTemporalAccumulationResourceNames, 1));
 
 		FSSDTemporalAccumulationCS::FPermutationDomain PermutationVector;
 		PermutationVector.Set<FSSDTemporalAccumulationCS::FIsMip0Dim>(true);
@@ -864,8 +929,8 @@ static void DenoiseSignalAtConstantPixelDensity(
 	int32 MaxPostFilterSampleCount = FMath::Clamp(Settings.HistoryConvolutionSampleCount, 1, kStackowiakMaxSampleCountPerSet);
 	if (MaxPostFilterSampleCount > 1)
 	{
-		FRDGTextureRef SignalOutput0 = GraphBuilder.CreateTexture(SignalProcessingDesc[0], TEXT("SSDReflectionsHistory0"));
-		FRDGTextureRef SignalOutput1 = GraphBuilder.CreateTexture(SignalProcessingDesc[1], TEXT("SSDReflectionsHistory1"));
+		FRDGTextureRef SignalOutput0 = GraphBuilder.CreateTexture(SignalProcessingDesc[0], GetResourceName(kHistoryConvolutionResourceNames, 0));
+		FRDGTextureRef SignalOutput1 = GraphBuilder.CreateTexture(SignalProcessingDesc[1], GetResourceName(kHistoryConvolutionResourceNames, 1));
 
 		FSSDSpatialAccumulationCS::FParameters* PassParameters = GraphBuilder.AllocParameters<FSSDSpatialAccumulationCS::FParameters>();
 		PassParameters->MaxSampleCount = FMath::Clamp(MaxPostFilterSampleCount, 1, kStackowiakMaxSampleCountPerSet);
