Implemented infrastructure for ray tracing translucency pass.
TODO: Change architecture depending on denoiser needs. Probably unify compositing with non ray traced pass. RT translucency disabled by default until then.

diff --git a/Engine/Config/ConsoleVariables.ini b/Engine/Config/ConsoleVariables.ini
--- a/Engine/Config/ConsoleVariables.ini
+++ b/Engine/Config/ConsoleVariables.ini
@@ -33,7 +33,7 @@
 
 
 ; Uncomment to get detailed logs on shader compiles and the opportunity to retry on errors
-;r.ShaderDevelopmentMode=1
+r.ShaderDevelopmentMode=1
 ; Uncomment to dump shaders in the Saved folder
 ; Warning: leaving this on for a while will fill your hard drive with many small files and folders
 ;r.DumpShaderDebugInfo=1
diff --git a/Engine/Shaders/Private/RayTracing/CompositeTranslucencyPS.usf b/Engine/Shaders/Private/RayTracing/CompositeTranslucencyPS.usf
new file mode 100644
--- /dev/null
+++ b/Engine/Shaders/Private/RayTracing/CompositeTranslucencyPS.usf
@@ -0,0 +1,17 @@
+// Copyright 1998-2019 Epic Games, Inc. All Rights Reserved.
+
+#include "../Common.ush"
+#include "../DeferredShadingCommon.ush"
+
+Texture2D<float4> TranslucencyTexture;
+
+SamplerState TranslucencyTextureSampler;
+
+void CompositeTranslucencyPS(
+	in noperspective float2 UV : TEXCOORD0,
+	out float4 OutColor : SV_Target0
+)
+{
+	float4 Translucency = TranslucencyTexture.Sample(TranslucencyTextureSampler, UV);
+	OutColor = Translucency;
+}
diff --git a/Engine/Shaders/Private/RayTracing/RayTracingLightingCommon.ush b/Engine/Shaders/Private/RayTracing/RayTracingLightingCommon.ush
new file mode 100644
--- /dev/null
+++ b/Engine/Shaders/Private/RayTracing/RayTracingLightingCommon.ush
@@ -0,0 +1,229 @@
+// Copyright 1998-2019 Epic Games, Inc. All Rights Reserved.
+
+// Light types: should match SceneTypes.h until there is a common header
+#define LIGHT_TYPE_DIRECTIONAL		0 
+#define LIGHT_TYPE_POINT			1 
+#define LIGHT_TYPE_SPOT				2 
+#define LIGHT_TYPE_RECT				3 
+#define LIGHT_TYPE_MAX				4 
+
+
+float GetRoughnessFade(float Roughness, float MaxRoughness)
+{
+	float RoughnessMaskScale = -2.0 / MaxRoughness;
+	return saturate(Roughness * RoughnessMaskScale + 2.0);
+}
+
+float3 GenerateReflectedRayDirection(
+	float3 IncidentDirection,
+	float3 WorldNormal,
+	float Roughness,
+	float2 RandSample
+)
+{
+	float3 RayDirection;
+	if (Roughness < 0.001) //ReflectionSmoothClamp)
+	{
+		RayDirection = reflect(IncidentDirection, WorldNormal);
+	}
+	else
+	{
+		float3 N = WorldNormal;
+		float3 V = -IncidentDirection;
+		float2 E = RandSample;
+
+		float3x3 TangentBasis = GetTangentBasis(N);
+		float3 TangentV = mul(TangentBasis, V);
+
+		float NoV = saturate(dot(V, WorldNormal));
+
+		float4 Sample = ImportanceSampleVisibleGGX(E, Pow4(Roughness), TangentV);
+
+		float3 H = mul(Sample.xyz, TangentBasis);
+		float3 L = 2 * dot(V, H) * H - V;
+
+		RayDirection = L;
+	}
+
+	return RayDirection;
+}
+
+float TraceShadowRay(
+	float3 Origin,
+	float3 Direction,
+	float TMax,
+	in FRayCone ShadowRayCone)
+{
+	RayDesc Ray;
+	Ray.Origin = Origin;
+	Ray.Direction = Direction;
+	Ray.TMin = 1e-4f;
+	Ray.TMax = TMax;
+
+	uint RayFlags = RAY_FLAG_CULL_BACK_FACING_TRIANGLES | RAY_FLAG_ACCEPT_FIRST_HIT_AND_END_SEARCH | RAY_FLAG_SKIP_CLOSEST_HIT_SHADER;
+
+	FMaterialClosestHitPayload Payload = TraceRayInternal(
+		TLAS,      // AccelerationStructure
+		RayFlags,
+		0xFF,      // InstanceInclusionMask
+		0,         // RayContributionToHitGroupIndex
+		1,         // MultiplierForGeometryContributionToShaderIndex
+		0,         // MissShaderIndex
+		Ray,       // RayDesc
+		ShadowRayCone
+	);
+
+	return Payload.HitT;
+}
+
+float3 ComputeDirectLighting(RayDesc Ray, in FRayCone RayCone, FMaterialClosestHitPayload Payload, bool bReflectedShadows)
+{
+	float3 DirectLighting = float3(0.0, 0.0, 0.0);
+
+	float3 WorldPosition = Payload.WorldPos;
+	float3 CameraVector = -Ray.Direction;
+	float AmbientOcclusion = 1.0;
+
+	// #dxr_todo: review when interpolants fill the Payload with all the data we need
+	FGBufferData GBufferData = (FGBufferData)0;
+	GBufferData.Depth = 1.f; // Do not use depth
+	GBufferData.WorldNormal = Payload.WorldNormal;
+	GBufferData.BaseColor = Payload.BaseColor;
+	GBufferData.Metallic = Payload.Metallic;
+	GBufferData.Specular = Payload.Specular;
+	GBufferData.Roughness = Payload.Roughness;
+	GBufferData.ShadingModelID = Payload.ShadingModelID;
+	GBufferData.CustomData = Payload.CustomData;
+	GBufferData.GBufferAO = Payload.GBufferAO;
+	//#dxr_todo: change float3 to float in payload to reduce space?;
+	GBufferData.IndirectIrradiance = (Payload.IndirectIrradiance.x + Payload.IndirectIrradiance.y + Payload.IndirectIrradiance.z) / 3.f; 
+	GBufferData.ShadingModelID = 1;
+
+	//#dxr_todo: calculate on the fly or store in Payload?
+#if 1
+	GBufferData.SpecularColor = lerp(0.08 * GBufferData.Specular.xxx, GBufferData.BaseColor, GBufferData.Metallic);
+	GBufferData.DiffuseColor = GBufferData.BaseColor - GBufferData.BaseColor * GBufferData.Metallic;
+#else	
+	GBufferData.SpecularColor = Payload.SpecularColor;
+	GBufferData.DiffuseColor = Payload.DiffuseColor;
+#endif
+
+	//#dxr_todo: think on the more efficient way of iterate through lights (split in light types, etc)
+	for (uint DirIndex = 0; DirIndex < ReflectionLightsData.Count; DirIndex++)
+	{
+		FDeferredLightData LightData = (FDeferredLightData)0;
+
+		uint LightType = ReflectionLightsData.Type[DirIndex];
+		LightData.Position = ReflectionLightsData.LightPosition[DirIndex];
+		LightData.InvRadius = ReflectionLightsData.LightInvRadius[DirIndex];
+		LightData.Color = ReflectionLightsData.LightColor[DirIndex];
+		LightData.FalloffExponent = ReflectionLightsData.LightFalloffExponent[DirIndex];
+		LightData.Direction = ReflectionLightsData.Direction[DirIndex];
+		LightData.Tangent = ReflectionLightsData.Tangent[DirIndex];
+		LightData.SpotAngles = ReflectionLightsData.SpotAngles[DirIndex];
+		LightData.SourceRadius = ReflectionLightsData.SourceRadius[DirIndex];
+		LightData.SourceLength = ReflectionLightsData.SourceLength[DirIndex];
+		LightData.SoftSourceRadius = ReflectionLightsData.SoftSourceRadius[DirIndex];
+		LightData.SpecularScale = ReflectionLightsData.SpecularScale[DirIndex];
+		LightData.ContactShadowLength = 0.0;
+		LightData.DistanceFadeMAD = ReflectionLightsData.DistanceFadeMAD[DirIndex];
+		LightData.ShadowMapChannelMask = float4(0, 0, 0, 0); // #dxr_todo;
+		LightData.ShadowedBits = 0; // Not lit dynamic shadows 
+		LightData.ContactShadowLengthInWS = false;
+
+		LightData.bRadialLight = (LightType != LIGHT_TYPE_DIRECTIONAL);
+		LightData.bSpotLight = (LightType == LIGHT_TYPE_SPOT);
+		LightData.bRectLight = (LightType == LIGHT_TYPE_RECT);
+
+		if (LightType == LIGHT_TYPE_DIRECTIONAL)
+		{
+			LightData.bInverseSquared = false;
+		}
+		else
+		{
+			LightData.bInverseSquared = LightData.FalloffExponent == 0;
+		}
+
+		// #dxr_todo: get a valid texture once they are properly bound
+		Texture2D LightSourceTexture = ReflectionLightsData.DummyRectLightTexture;
+
+		float4 LightAttenuation = 1.0f;
+		float LightProfileMultiplier = 1.0; //#dxr_todo: implement support for light profiles
+
+		float3 ShadowRayDirection;
+		float3 ToLight = LightData.Position - WorldPosition;
+		float LightMask = 1.0;
+
+		if (LightType == LIGHT_TYPE_DIRECTIONAL)
+		{
+			ShadowRayDirection = LightData.Direction;
+		}
+		else
+		{
+			LightMask = GetLocalLightAttenuation(WorldPosition, LightData, ToLight, ShadowRayDirection);
+		}
+
+		// Skip the light sample that's below the horizon or does not contribute anything due to attenuation.
+		if (dot(Payload.WorldNormal, ToLight) <= 0 || LightMask <= 0.0)
+		{
+			continue;
+		}
+
+		//#dxr_todo: check perf improvements forcing TMin == TMax when bReflectedShadows = false, to avoid tracing rays inside if statements
+		if (bReflectedShadows)
+		{
+			FRayCone ShadowRayCone = RayCone;
+
+			//#dxr_todo: add support for sampling area lights (and perhaps add CVar to enable stocastic shadows?)
+			float DistanceToLight = length(ToLight);
+			float HitT = TraceShadowRay(WorldPosition + Payload.WorldNormal*0.001, ShadowRayDirection, DistanceToLight, ShadowRayCone);
+			AmbientOcclusion = HitT < 0.0f;
+		}
+
+		float3 LightContribution = GetDynamicLighting(WorldPosition, -CameraVector, GBufferData, AmbientOcclusion, GBufferData.ShadingModelID, LightData, LightAttenuation, 0.5, uint2(0, 0), LightSourceTexture).xyz;
+		DirectLighting += LightContribution * LightProfileMultiplier;
+	}
+
+	return DirectLighting;
+}
+
+FMaterialClosestHitPayload TraceRayAndAccumulateResults(
+	in RayDesc Ray,
+	inout FRayCone RayCone,
+	inout float3 Radiance)
+{
+	// Call TraceRay iteratively until alpha is 0 and accumulate the result
+	//#dxr-todo: iterate until opacity == 1 if the material is not opaque
+	//#dxr-todo: add optimizations (TMax depending on roughness, etc)
+
+	uint RayFlags = RAY_FLAG_CULL_BACK_FACING_TRIANGLES;
+
+	FMaterialClosestHitPayload Payload = TraceRayInternal(
+		TLAS,      // AccelerationStructure
+		RayFlags,
+		0xFF,      // InstanceInclusionMask
+		0,         // RayContributionToHitGroupIndex
+		1,         // MultiplierForGeometryContributionToShaderIndex
+		0,         // MissShaderIndex
+		Ray,       // RayDesc
+		RayCone
+	);
+
+	if (ShouldDoDirectLighting)
+	{
+		float3 DirectLighting = ComputeDirectLighting(Ray, RayCone, Payload, ShouldDoReflectedShadows > 0);
+
+		// Transform NaNs to black, transform negative colors to black.
+		DirectLighting = -min(-DirectLighting, 0.0);
+		Radiance += DirectLighting;
+	}
+
+	if (ShouldDoEmissiveAndIndirectLighting)
+	{
+		// Emissive & indirect contribution
+		Radiance += Payload.Radiance;
+		Radiance += Payload.IndirectIrradiance;
+	}
+
+	return Payload;
+}
diff --git a/Engine/Shaders/Private/RayTracing/RayTracingReflections.usf b/Engine/Shaders/Private/RayTracing/RayTracingReflections.usf
--- a/Engine/Shaders/Private/RayTracing/RayTracingReflections.usf
+++ b/Engine/Shaders/Private/RayTracing/RayTracingReflections.usf
@@ -14,13 +14,6 @@ float2 SobolIndexToUniformUnitSquare(uint2 SobolRand)
 #define SUPPORT_CONTACT_SHADOWS		0
 #define USE_SOURCE_TEXTURE			0 //Used in RectLight.ush  //#dxr_todo: enable textures
 
-// Light types: should match SceneTypes.h until there is a common header
-#define LIGHT_TYPE_DIRECTIONAL		0 
-#define LIGHT_TYPE_POINT			1 
-#define LIGHT_TYPE_SPOT				2 
-#define LIGHT_TYPE_RECT				3 
-#define LIGHT_TYPE_MAX				4 
-
 #include "../DeferredShadingCommon.ush"
 #include "../DeferredLightingCommon.ush"
 #include "../ReflectionEnvironmentShared.ush"
@@ -46,11 +39,7 @@ RaytracingAccelerationStructure TLAS;
 RWTexture2D<float4> ColorOutput;
 RWTexture2D<float> RayHitDistanceOutput;
 
-float GetRoughnessFade(float Roughness, float MaxRoughness)
-{
-	float RoughnessMaskScale = -2.0 / MaxRoughness;
-	return saturate(Roughness * RoughnessMaskScale + 2.0);
-}
+#include "RayTracingLightingCommon.ush"
 
 // #dxr-todo: move to RayTracingCommon.ush
 float3 ReconstructWorldPositionFromDepth(float2 UV, float Depth)
@@ -62,263 +51,6 @@ float3 ReconstructWorldPositionFromDepth(float2 UV, float Depth)
 	return WorldPosition;
 }
 
-float3 GenerateReflectedRayDirection(
-	float3 IncidentDirection,
-	float3 WorldNormal,
-	float Roughness,
-	float2 RandSample
-)
-{
-	float3 RayDirection;
-	if (Roughness < 0.001 ) //ReflectionSmoothClamp)
-	{
-		RayDirection = reflect(IncidentDirection, WorldNormal);
-	}
-	else
-	{
-		float3 N = WorldNormal;
-		float3 V = -IncidentDirection;
-		float2 E = RandSample;
-
-		float3x3 TangentBasis = GetTangentBasis(N);
-		float3 TangentV = mul(TangentBasis, V);
-
-		float NoV = saturate(dot(V, WorldNormal));
-
-		float4 Sample = ImportanceSampleVisibleGGX(E, Pow4(Roughness), TangentV);
-
-		float3 H = mul(Sample.xyz, TangentBasis);
-		float3 L = 2 * dot( V, H ) * H - V;
-
-		RayDirection = L;
-	}
-
-	return RayDirection;
-}
-
-#if 0
-float3 GatherRadiance(float CompositeAlpha, float3 WorldPosition, float3 RayDirection, float Roughness, float IndirectIrradiance)
-{
-	float IndirectSpecularOcclusion = 1.0f;
-	float3 ExtraIndirectSpecular = 0;
-	return CompositeReflectionCapturesAndSkylight(CompositeAlpha, WorldPosition, RayDirection, Roughness, IndirectIrradiance, IndirectSpecularOcclusion, ExtraIndirectSpecular, 0, 0);
-}
-
-float4 BlendWithEnvironment(inout float3 BRDFThroughput, const in FMaterialClosestHitPayload CurrentBounceHitInfo, float3 WorldPosition, float4 ReflectionColor, float3 RayDirection, float3 Normal)
-{
-#if USE_PREEXPOSURE
-	const float PreExposure = View.PreExposure;
-#else
-	const float PreExposure = 1.f;
-#endif
-
-	float3 ProbeReflectDirection = reflect(RayDirection, Normal);
-	float  NoV = saturate(dot(-RayDirection, Normal));
-	float IndirectIrradiance = CurrentBounceHitInfo.IndirectIrradiance;
-
-	// Point lobe in off-specular peak direction
-	ProbeReflectDirection = GetOffSpecularPeakReflectionDir(Normal, ProbeReflectDirection, CurrentBounceHitInfo.Roughness);
-
-	float RoughnessFade = saturate(CurrentBounceHitInfo.Roughness * RTRRoughnessFadeFactor + 2);
-	ReflectionColor.a = RoughnessFade; 	
-	  
-	float4 FinalReflectionColor = float4(ReflectionColor.rgb, 1 - ReflectionColor.a); 
-	float AO = CurrentBounceHitInfo.GBufferAO;
-	float RoughnessSq = CurrentBounceHitInfo.Roughness * CurrentBounceHitInfo.Roughness;
-	float SpecularOcclusion = GetSpecularOcclusion(NoV, RoughnessSq, AO);
-	FinalReflectionColor.a *= SpecularOcclusion;
-
-	FinalReflectionColor.rgb += PreExposure * GatherRadiance(FinalReflectionColor.a, WorldPosition, ProbeReflectDirection, CurrentBounceHitInfo.Roughness, IndirectIrradiance);
-
-	float3 BRDF = EnvBRDF(CurrentBounceHitInfo.SpecularColor, CurrentBounceHitInfo.Roughness, NoV);
-	FinalReflectionColor.rgb *= BRDF;
-
-	FinalReflectionColor.rgb *= BRDFThroughput;
-	BRDFThroughput *= BRDF;
-
-	return FinalReflectionColor;
-}
-#endif
-
-float TraceShadowRay(
-	float3 Origin,
-	float3 Direction,
-	float TMax,
-	in FRayCone ShadowRayCone)
-{
-	RayDesc Ray;
-	Ray.Origin = Origin;
-	Ray.Direction = Direction;
-	Ray.TMin = 1e-4f;
-	Ray.TMax = TMax;
-
-	uint RayFlags = RAY_FLAG_CULL_BACK_FACING_TRIANGLES | RAY_FLAG_ACCEPT_FIRST_HIT_AND_END_SEARCH | RAY_FLAG_SKIP_CLOSEST_HIT_SHADER;
-
-	FMaterialClosestHitPayload Payload = TraceRayInternal(
-		TLAS,      // AccelerationStructure
-		RayFlags,
-		0xFF,      // InstanceInclusionMask
-		0,         // RayContributionToHitGroupIndex
-		1,         // MultiplierForGeometryContributionToShaderIndex
-		0,         // MissShaderIndex
-		Ray,       // RayDesc
-		ShadowRayCone
-	);
-
-	return Payload.HitT;
-}
-
-float3 ComputeDirectLighting(RayDesc Ray, in FRayCone RayCone, FMaterialClosestHitPayload Payload, bool bReflectedShadows)
-{
-	float3 DirectLighting = float3(0.0, 0.0, 0.0);
-
-	float3 WorldPosition = Payload.WorldPos;
-	float3 CameraVector = -Ray.Direction;
-	float AmbientOcclusion = 1.0;
-
-	// #dxr_todo: review when interpolants fill the Payload with all the data we need
-	FGBufferData GBufferData = (FGBufferData)0;
-	GBufferData.Depth = 1.f; // Do not use depth
-	GBufferData.WorldNormal = Payload.WorldNormal;
-	GBufferData.BaseColor = Payload.BaseColor;
-	GBufferData.Metallic = Payload.Metallic;
-	GBufferData.Specular = Payload.Specular;
-	GBufferData.Roughness = Payload.Roughness;
-	GBufferData.ShadingModelID = Payload.ShadingModelID;
-	GBufferData.CustomData = Payload.CustomData;
-	GBufferData.GBufferAO = Payload.GBufferAO;
-	GBufferData.IndirectIrradiance = Payload.IndirectIrradiance;
-	GBufferData.ShadingModelID = 1; 
-	   
-	//#dxr_todo: calculate on the fly or store in Payload?
-#if 1
-	GBufferData.SpecularColor = lerp(0.08 * GBufferData.Specular.xxx, GBufferData.BaseColor, GBufferData.Metallic);
-	GBufferData.DiffuseColor = GBufferData.BaseColor - GBufferData.BaseColor * GBufferData.Metallic;
-#else	
-	GBufferData.SpecularColor = Payload.SpecularColor;	
-	GBufferData.DiffuseColor = Payload.DiffuseColor;	
-#endif
-
-	//#dxr_todo: think on the more efficient way of iterate through lights (split in light types, etc)
-	for (uint DirIndex = 0; DirIndex < ReflectionLightsData.Count; DirIndex++)
-	{
-		FDeferredLightData LightData = (FDeferredLightData)0;
-
-		uint LightType = ReflectionLightsData.Type[DirIndex];
-		LightData.Position = ReflectionLightsData.LightPosition[DirIndex];
-		LightData.InvRadius	= ReflectionLightsData.LightInvRadius[DirIndex];
-		LightData.Color = ReflectionLightsData.LightColor[DirIndex];
-		LightData.FalloffExponent = ReflectionLightsData.LightFalloffExponent[DirIndex];
-		LightData.Direction = ReflectionLightsData.Direction[DirIndex];
-		LightData.Tangent = ReflectionLightsData.Tangent[DirIndex];
-		LightData.SpotAngles = ReflectionLightsData.SpotAngles[DirIndex];
-		LightData.SourceRadius = ReflectionLightsData.SourceRadius[DirIndex];
-		LightData.SourceLength = ReflectionLightsData.SourceLength[DirIndex];
-		LightData.SoftSourceRadius = ReflectionLightsData.SoftSourceRadius[DirIndex];
-		LightData.SpecularScale = ReflectionLightsData.SpecularScale[DirIndex];
-		LightData.ContactShadowLength = 0.0;
-		LightData.DistanceFadeMAD = ReflectionLightsData.DistanceFadeMAD[DirIndex];
-		LightData.ShadowMapChannelMask = float4(0, 0, 0, 0); // #dxr_todo;
-		LightData.ShadowedBits = 0; // Not lit dynamic shadows 
-		LightData.ContactShadowLengthInWS = false;
-
-		LightData.bRadialLight = (LightType != LIGHT_TYPE_DIRECTIONAL);
-		LightData.bSpotLight = (LightType == LIGHT_TYPE_SPOT);
-		LightData.bRectLight = (LightType == LIGHT_TYPE_RECT);
-
-		if (LightType == LIGHT_TYPE_DIRECTIONAL)
-		{
-			LightData.bInverseSquared = false;
-		}
-		else
-		{
-			LightData.bInverseSquared = LightData.FalloffExponent == 0;
-		}
-
-		// #dxr_todo: get a valid texture once they are properly bound
-		Texture2D LightSourceTexture = ReflectionLightsData.DummyRectLightTexture;
-
-		float4 LightAttenuation = 1.0f;
-		float LightProfileMultiplier = 1.0; //#dxr_todo: implement support for light profiles
-
-		float3 ShadowRayDirection;
-		float3 ToLight = LightData.Position - WorldPosition;
-		float LightMask = 1.0;
-		
-		if (LightType == LIGHT_TYPE_DIRECTIONAL)
-		{
-			ShadowRayDirection = LightData.Direction;
-		}
-		else
-		{
-			LightMask = GetLocalLightAttenuation(WorldPosition, LightData, ToLight, ShadowRayDirection);
-		}
-
-		// Skip the light sample that's below the horizon or does not contribute anything due to attenuation.
-		if (dot(Payload.WorldNormal, ToLight) <= 0 || LightMask <= 0.0)
-		{
-			continue;
-		}
-
-		//#dxr_todo: check perf improvements forcing TMin == TMax when bReflectedShadows = false, to avoid tracing rays inside if statements
-		if (bReflectedShadows)
-		{
-			FRayCone ShadowRayCone = RayCone;
-
-			//#dxr_todo: add support for sampling area lights (and perhaps add CVar to enable stocastic shadows?)
-			float DistanceToLight = length(ToLight);
-			float HitT = TraceShadowRay(WorldPosition + Payload.WorldNormal*0.001, ShadowRayDirection, DistanceToLight, ShadowRayCone);
-			AmbientOcclusion = HitT < 0.0f;
-		}
-
-		float3 LightContribution = GetDynamicLighting(WorldPosition, -CameraVector, GBufferData, AmbientOcclusion, GBufferData.ShadingModelID, LightData, LightAttenuation, 0.5, uint2(0, 0), LightSourceTexture).xyz;
-		DirectLighting += LightContribution * LightProfileMultiplier;
-	}
-
-	return DirectLighting;
-}
-
-FMaterialClosestHitPayload TraceRayAndAccumulateResults(
-	in RayDesc Ray,
-	inout FRayCone RayCone,
-	inout float3 Radiance)
-{
-	// Call TraceRay iteratively until alpha is 0 and accumulate the result
-	//#dxr-todo: iterate until opacity == 1 if the material is not opaque
-	//#dxr-todo: add optimizations (TMax depending on roughness, etc)
-
-	uint RayFlags = RAY_FLAG_CULL_BACK_FACING_TRIANGLES;
-
-	FMaterialClosestHitPayload Payload = TraceRayInternal(
-		TLAS,      // AccelerationStructure
-		RayFlags,
-		0xFF,      // InstanceInclusionMask
-		0,         // RayContributionToHitGroupIndex
-		1,         // MultiplierForGeometryContributionToShaderIndex
-		0,         // MissShaderIndex
-		Ray,       // RayDesc
-		RayCone
-	);
-
-	if (ShouldDoDirectLighting)
-	{
-		float3 DirectLighting = ComputeDirectLighting(Ray, RayCone, Payload, ShouldDoReflectedShadows > 0);
-
-		// Transform NaNs to black, transform negative colors to black.
-		DirectLighting = -min(-DirectLighting, 0.0);
-		Radiance += DirectLighting;
-	}
-    
-	if (ShouldDoEmissiveAndIndirectLighting)
-	{
-	    // Emissive & indirect contribution
-	    Radiance += Payload.Radiance;
-		Radiance += Payload.IndirectIrradiance;
-	}
-
-	return Payload;
-}
-
 uint2 GetPixelCoord(uint2 DispatchThreadId)
 {
 	uint UpscaleFactorPow2 = UpscaleFactor * UpscaleFactor;
diff --git a/Engine/Shaders/Private/RayTracing/RayTracingTranslucency.usf b/Engine/Shaders/Private/RayTracing/RayTracingTranslucency.usf
new file mode 100644
--- /dev/null
+++ b/Engine/Shaders/Private/RayTracing/RayTracingTranslucency.usf
@@ -0,0 +1,130 @@
+// Copyright 1998-2019 Epic Games, Inc. All Rights Reserved.
+
+#ifndef DIM_DENOISER_OUTPUT
+	#define DIM_DENOISER_OUTPUT 0
+#endif
+
+#include "../Common.ush"
+
+float2 SobolIndexToUniformUnitSquare(uint2 SobolRand)
+{
+	return float2(SobolRand) * rcp(65536.0) + rcp(65536.0 * 2.0);
+}
+
+#define SUPPORT_CONTACT_SHADOWS		0
+#define USE_SOURCE_TEXTURE			0 //Used in RectLight.ush  //#dxr_todo: enable textures
+
+#include "../DeferredShadingCommon.ush"
+#include "../DeferredLightingCommon.ush"
+#include "../ReflectionEnvironmentShared.ush"
+#include "../Montecarlo.ush"
+#include "../PathTracing/Utilities/PathTracingRandomSequence.ush" // #dxr-todo: move to a common folder
+#include "../HeightFogCommon.ush"
+#include "RayTracingCommon.ush"
+#include "RayTracingHitGroupCommon.ush"
+
+uint SamplesPerPixel;
+uint HeightFog;
+uint ShouldDoDirectLighting;
+uint ShouldDoReflectedShadows;
+uint ShouldDoEmissiveAndIndirectLighting;
+uint UpscaleFactor;
+
+float ReflectionMinRayDistance;
+float ReflectionMaxRayDistance;
+float ReflectionMaxRoughness;
+
+RaytracingAccelerationStructure TLAS;
+
+RWTexture2D<float4> ColorOutput;
+RWTexture2D<float> RayHitDistanceOutput;
+
+#include "RayTracingLightingCommon.ush"
+
+// #dxr-todo: move to RayTracingCommon.ush
+float3 ReconstructWorldPositionFromDepth(float2 UV, float Depth)
+{
+	float2 ScreenPosition = (UV - View.ScreenPositionScaleBias.wz) / View.ScreenPositionScaleBias.xy;
+	float4 HomogeneousWorldPosition = mul(float4(ScreenPosition * Depth, Depth, 1), View.ScreenToWorld);
+	float3 WorldPosition = HomogeneousWorldPosition.xyz / HomogeneousWorldPosition.w;
+
+	return WorldPosition;
+}
+
+uint2 GetPixelCoord(uint2 DispatchThreadId)
+{
+	uint UpscaleFactorPow2 = UpscaleFactor * UpscaleFactor;
+
+	// TODO: find a way to not interfer with TAA's jittering.
+	uint SubPixelId = View.StateFrameIndex & (UpscaleFactorPow2 - 1);
+
+	return DispatchThreadId * UpscaleFactor + uint2(SubPixelId & (UpscaleFactor - 1), SubPixelId / UpscaleFactor);
+}
+
+
+[shader("raygeneration")]
+void RayTracingTranslucencyRGS()
+{
+	uint2 DispatchThreadId = DispatchRaysIndex().xy;
+
+	uint2 PixelCoord = GetPixelCoord(DispatchThreadId);
+	
+	uint LinearIndex = PixelCoord.y * View.BufferSizeAndInvSize.x + PixelCoord.x; // TODO: PixelCoord or DispatchThreadId
+	
+	RandomSequence RandSequence;
+	RandomSequence_Initialize(RandSequence, LinearIndex, View.FrameNumber);
+
+	float2 InvBufferSize = View.BufferSizeAndInvSize.zw;
+	float2 UV = (float2(PixelCoord) + 0.5) * InvBufferSize;
+	FScreenSpaceData ScreenSpaceData = GetScreenSpaceData(UV);
+	float Depth = ScreenSpaceData.GBuffer.Depth;
+	float3 CameraOrigin = ReconstructWorldPositionFromDepth(UV, 0.0f);
+	float3 WorldPosition = ReconstructWorldPositionFromDepth(UV, Depth);
+	float3 CameraDirection = normalize(WorldPosition - CameraOrigin);
+	//#dxr - todo: use also the shading normal (interpolate, or apply a more sophisticated approach)
+	float3 WorldNormal = ScreenSpaceData.GBuffer.WorldNormal;
+	float Roughness = ScreenSpaceData.GBuffer.Roughness;
+	
+	float3 BRDFThroughput = float3(1.0f, 1.0f, 1.0f);
+
+	uint MaxBounces = 1; //#dxr-todo: only works with 1 bounce until hit shaders are supported
+	uint LocalSamplesPerPixel = SamplesPerPixel;
+
+	// Mask out depth values at or beyond far plane
+	if (Depth >= SCENE_TEXTURES_DISABLED_SCENE_DEPTH_VALUE)
+	{
+		LocalSamplesPerPixel = 0;
+	}
+
+	float4 TranslucencyColor = 0.0;
+	float HitDistance = 0.0f;
+
+	ColorOutput[DispatchThreadId] = float4(TranslucencyColor);
+	RayHitDistanceOutput[DispatchThreadId] = HitDistance;
+}
+
+
+[shader("miss")]
+void RayTracingTranslucencyMainMS(inout FPackedMaterialClosestHitPayload PackedPayload)
+{
+	PackedPayload.Payload.HitT = -1;
+}
+
+
+[shader("closesthit")]
+void RayTracingTranslucencyMainCHS(inout FPackedMaterialClosestHitPayload PackedPayload, in FDefaultAttributes Attributes)
+{
+	FMaterialClosestHitPayload Payload = (FMaterialClosestHitPayload)0;
+
+	Payload.HitT = RayTCurrent();
+	Payload.BaseColor = float3(0.2, 0.3, 0.4);
+
+	FTriangleBaseAttributes Triangle = LoadTriangleBaseAttributes(PrimitiveIndex());
+	float3 Edge0 = Triangle.LocalPositions[2] - Triangle.LocalPositions[0];
+	float3 Edge1 = Triangle.LocalPositions[1] - Triangle.LocalPositions[0];
+	float3x3 WorldToLocal = (float3x3)WorldToObject();
+	float3x3 LocalToWorldNormal = transpose(WorldToLocal);
+	Payload.WorldNormal = normalize(mul(LocalToWorldNormal, cross(Edge0, Edge1)));
+
+	PackedPayload = PackRayTracingPayload(Payload, PackedPayload.RayCone);
+}
diff --git a/Engine/Source/Runtime/Renderer/Private/DeferredShadingRenderer.cpp b/Engine/Source/Runtime/Renderer/Private/DeferredShadingRenderer.cpp
--- a/Engine/Source/Runtime/Renderer/Private/DeferredShadingRenderer.cpp
+++ b/Engine/Source/Runtime/Renderer/Private/DeferredShadingRenderer.cpp
@@ -169,6 +169,15 @@ static TAutoConsoleVariable<int32> CVarUseAODenoiser(
 	TEXT(" 2: GScreenSpaceDenoiser witch may be overriden by a third party plugin (default)."),
 	ECVF_RenderThreadSafe);
 
+static TAutoConsoleVariable<int32> CVarRayTracingTranslucency(
+	TEXT("r.RayTracing.Translucency"),
+	0,
+	TEXT("0 to disable ray tracing translucency.\n")
+	TEXT(" 0: off\n")
+	TEXT(" 1: on"),
+	ECVF_RenderThreadSafe);
+
+
 
 DECLARE_CYCLE_STAT(TEXT("PostInitViews FlushDel"), STAT_PostInitViews_FlushDel, STATGROUP_InitViews);
 DECLARE_CYCLE_STAT(TEXT("InitViews Intentional Stall"), STAT_InitViews_Intentional_Stall, STATGROUP_InitViews);
@@ -2029,36 +2038,47 @@ void FDeferredShadingSceneRenderer::Render(FRHICommandListImmediate& RHICmdList)
 
 		RHICmdList.SetCurrentStat(GET_STATID(STAT_CLM_Translucency));
 
-		// For now there is only one resolve for all translucency passes. This can be changed by enabling the resolve in RenderTranslucency()
-		TRefCountPtr<IPooledRenderTarget> SceneColorCopy;
-		ConditionalResolveSceneColorForTranslucentMaterials(RHICmdList, SceneColorCopy);
-
-		if (ViewFamily.AllowTranslucencyAfterDOF())
+#if RHI_RAYTRACING
+		const bool bRaytracedTranslucency = CVarRayTracingTranslucency.GetValueOnRenderThread() != 0;
+		if (bRaytracedTranslucency)
 		{
-			RenderTranslucency(RHICmdList, ETranslucencyPass::TPT_StandardTranslucency, SceneColorCopy);
-			// Translucency after DOF is rendered now, but stored in the separate translucency RT for later use.
-			RenderTranslucency(RHICmdList, ETranslucencyPass::TPT_TranslucencyAfterDOF, SceneColorCopy);
+			ResolveSceneColor(RHICmdList);
+			RayTraceTranslucency(RHICmdList);
 		}
-		else // Otherwise render translucent primitives in a single bucket.
+		else
+#endif
 		{
-			RenderTranslucency(RHICmdList, ETranslucencyPass::TPT_AllTranslucency, SceneColorCopy);
-		}
-		ServiceLocalQueue();
-
-		static const auto DisableDistortionCVar = IConsoleManager::Get().FindTConsoleVariableDataInt(TEXT("r.DisableDistortion"));
-		const bool bAllowDistortion = DisableDistortionCVar->GetValueOnAnyThread() != 1;
+			// For now there is only one resolve for all translucency passes. This can be changed by enabling the resolve in RenderTranslucency()
+			TRefCountPtr<IPooledRenderTarget> SceneColorCopy;
+			ConditionalResolveSceneColorForTranslucentMaterials(RHICmdList, SceneColorCopy);
 
-		if (GetRefractionQuality(ViewFamily) > 0 && bAllowDistortion)
-		{
-			// To apply refraction effect by distorting the scene color.
-			// After non separate translucency as that is considered at scene depth anyway
-			// It allows skybox translucency (set to non separate translucency) to be refracted.
-			RHICmdList.SetCurrentStat(GET_STATID(STAT_CLM_RenderDistortion));
-			RenderDistortion(RHICmdList);
+			if (ViewFamily.AllowTranslucencyAfterDOF())
+			{
+				RenderTranslucency(RHICmdList, ETranslucencyPass::TPT_StandardTranslucency, SceneColorCopy);
+				// Translucency after DOF is rendered now, but stored in the separate translucency RT for later use.
+				RenderTranslucency(RHICmdList, ETranslucencyPass::TPT_TranslucencyAfterDOF, SceneColorCopy);
+			}
+			else // Otherwise render translucent primitives in a single bucket.
+			{
+				RenderTranslucency(RHICmdList, ETranslucencyPass::TPT_AllTranslucency, SceneColorCopy);
+			}
 			ServiceLocalQueue();
-		}
 
-		RHICmdList.SetCurrentStat(GET_STATID(STAT_CLM_AfterTranslucency));
+			static const auto DisableDistortionCVar = IConsoleManager::Get().FindTConsoleVariableDataInt(TEXT("r.DisableDistortion"));
+			const bool bAllowDistortion = DisableDistortionCVar->GetValueOnAnyThread() != 1;
+
+			if (GetRefractionQuality(ViewFamily) > 0 && bAllowDistortion)
+			{
+				// To apply refraction effect by distorting the scene color.
+				// After non separate translucency as that is considered at scene depth anyway
+				// It allows skybox translucency (set to non separate translucency) to be refracted.
+				RHICmdList.SetCurrentStat(GET_STATID(STAT_CLM_RenderDistortion));
+				RenderDistortion(RHICmdList);
+				ServiceLocalQueue();
+			}
+
+			RHICmdList.SetCurrentStat(GET_STATID(STAT_CLM_AfterTranslucency));
+		}
 	}
 
 	checkSlow(RHICmdList.IsOutsideRenderPass());
diff --git a/Engine/Source/Runtime/Renderer/Private/DeferredShadingRenderer.h b/Engine/Source/Runtime/Renderer/Private/DeferredShadingRenderer.h
--- a/Engine/Source/Runtime/Renderer/Private/DeferredShadingRenderer.h
+++ b/Engine/Source/Runtime/Renderer/Private/DeferredShadingRenderer.h
@@ -480,6 +480,16 @@ private:
 	void RenderRayTracingSkyLight(FRHICommandListImmediate& RHICmdList, TRefCountPtr<IPooledRenderTarget>& SkyLightRT, TRefCountPtr<IPooledRenderTarget>& HitDistanceRT);
 	void CompositeRayTracingSkyLight(FRHICommandListImmediate& RHICmdList, TRefCountPtr<IPooledRenderTarget>& SkyLightRT, TRefCountPtr<IPooledRenderTarget>& HitDistanceRT);
 
+	void RayTraceTranslucency(FRHICommandListImmediate& RHICmdList);
+	void RayTraceTranslucencyView(
+		FRDGBuilder& GraphBuilder,
+		const FViewInfo& View,
+		FRDGTextureRef* OutColorTexture,
+		FRDGTextureRef* OutRayHitDistanceTexture,
+		int32 SamplePerPixel,
+		int32 HeightFog,
+		float ResolutionFraction);
+
 	/** Path tracing functions. */
 	void RenderPathTracing(FRHICommandListImmediate& RHICmdList, const FViewInfo& View);
 
diff --git a/Engine/Source/Runtime/Renderer/Private/RayTracing/RayTracingTranslucency.cpp b/Engine/Source/Runtime/Renderer/Private/RayTracing/RayTracingTranslucency.cpp
new file mode 100644
--- /dev/null
+++ b/Engine/Source/Runtime/Renderer/Private/RayTracing/RayTracingTranslucency.cpp
@@ -0,0 +1,457 @@
+// Copyright 1998-2019 Epic Games, Inc. All Rights Reserved.
+
+#include "RendererPrivate.h"
+#include "GlobalShader.h"
+#include "DeferredShadingRenderer.h"
+
+#if RHI_RAYTRACING
+
+#include "ClearQuad.h"
+#include "SceneRendering.h"
+#include "SceneRenderTargets.h"
+#include "RenderTargetPool.h"
+#include "RHIResources.h"
+#include "UniformBuffer.h"
+#include "VisualizeTexture.h"
+#include "LightRendering.h"
+#include "SystemTextures.h"
+#include "SceneViewFamilyBlackboard.h"
+#include "ScreenSpaceDenoise.h"
+#include "PostProcess/PostProcessing.h"
+#include "PostProcess/SceneFilterRendering.h"
+#include "PipelineStateCache.h"
+
+static int32 GRayTracingTranslucencyEmissiveAndIndirectLighting = 1;
+static FAutoConsoleVariableRef CVarRayTracingTranslucencyEmissiveAndIndirectLighting(
+	TEXT("r.RayTracing.Translucency.EmissiveAndIndirectLighting"),
+	GRayTracingTranslucencyEmissiveAndIndirectLighting,
+	TEXT("Enables ray tracing translucency emissive and indirect lighting (default = 1)")
+);
+
+static int32 GRayTracingTranslucencyDirectLighting = 1;
+static FAutoConsoleVariableRef CVarRayTracingTranslucencyDirectLighting(
+	TEXT("r.RayTracing.Translucency.DirectLighting"),
+	GRayTracingTranslucencyDirectLighting,
+	TEXT("Enables ray tracing translucency direct lighting (default = 1)")
+);
+
+static int32 GRayTracingTranslucencyShadows = 1;
+static FAutoConsoleVariableRef CVarRayTracingTranslucencyShadows(
+	TEXT("r.RayTracing.Translucency.Shadows"),
+	GRayTracingTranslucencyShadows,
+	TEXT("Enables shadows in ray tracing Translucency (default = 1)")
+);
+
+static float GRayTracingTranslucencyMinRayDistance = -1;
+static FAutoConsoleVariableRef CVarRayTracingTranslucencyMinRayDistance(
+	TEXT("r.RayTracing.Translucency.MinRayDistance"),
+	GRayTracingTranslucencyMinRayDistance,
+	TEXT("Sets the minimum ray distance for ray traced translucency rays. Actual translucency ray length is computed as Lerp(MaxRayDistance, MinRayDistance, Roughness), i.e. translucency rays become shorter when traced from rougher surfaces. (default = -1 (infinite rays))")
+);
+
+static float GRayTracingTranslucencyMaxRayDistance = -1;
+static FAutoConsoleVariableRef CVarRayTracingTranslucencyMaxRayDistance(
+	TEXT("r.RayTracing.Translucency.MaxRayDistance"),
+	GRayTracingTranslucencyMaxRayDistance,
+	TEXT("Sets the maximum ray distance for ray traced translucency rays. When ray shortening is used, skybox will not be sampled in RT translucency pass and will be composited later, together with local reflection captures. Negative values turn off this optimization. (default = -1 (infinite rays))")
+);
+
+static int32 GRayTracingTranslucencySamplesPerPixel = 1;
+static FAutoConsoleVariableRef CVarRayTracingTranslucencySamplesPerPixel(
+	TEXT("r.RayTracing.Translucency.SamplesPerPixel"),
+	GRayTracingTranslucencySamplesPerPixel,
+	TEXT("Sets the samples-per-pixel for Translucency (default = 1)"));
+
+static int32 GRayTracingTranslucencyHeightFog = 1;
+static FAutoConsoleVariableRef CVarRayTracingTranslucencyHeightFog(
+	TEXT("r.RayTracing.Translucency.HeightFog"),
+	GRayTracingTranslucencyHeightFog,
+	TEXT("Enables height fog in ray traced Translucency (default = 1)"));
+
+DECLARE_GPU_STAT_NAMED(RayTracingTranslucency, TEXT("Ray Tracing Translucency"));
+
+
+static const int32 GTranslucencyLightCountMaximum = 64;
+
+BEGIN_GLOBAL_SHADER_PARAMETER_STRUCT(FTranslucencyLightData, )
+SHADER_PARAMETER(uint32, Count)
+SHADER_PARAMETER_ARRAY(uint32, Type, [GTranslucencyLightCountMaximum])
+SHADER_PARAMETER_ARRAY(FVector, LightPosition, [GTranslucencyLightCountMaximum])
+SHADER_PARAMETER_ARRAY(float, LightInvRadius, [GTranslucencyLightCountMaximum])
+SHADER_PARAMETER_ARRAY(FVector, LightColor, [GTranslucencyLightCountMaximum])
+SHADER_PARAMETER_ARRAY(float, LightFalloffExponent, [GTranslucencyLightCountMaximum])
+SHADER_PARAMETER_ARRAY(FVector, Direction, [GTranslucencyLightCountMaximum])
+SHADER_PARAMETER_ARRAY(FVector, Tangent, [GTranslucencyLightCountMaximum])
+SHADER_PARAMETER_ARRAY(FVector2D, SpotAngles, [GTranslucencyLightCountMaximum])
+SHADER_PARAMETER_ARRAY(float, SpecularScale, [GTranslucencyLightCountMaximum])
+SHADER_PARAMETER_ARRAY(float, SourceRadius, [GTranslucencyLightCountMaximum])
+SHADER_PARAMETER_ARRAY(float, SourceLength, [GTranslucencyLightCountMaximum])
+SHADER_PARAMETER_ARRAY(float, SoftSourceRadius, [GTranslucencyLightCountMaximum])
+SHADER_PARAMETER_ARRAY(FVector2D, DistanceFadeMAD, [GTranslucencyLightCountMaximum])
+SHADER_PARAMETER_TEXTURE(Texture2D, DummyRectLightTexture) //#dxr_todo: replace with an array of textures when there is support for SHADER_PARAMETER_TEXTURE_ARRAY 
+END_GLOBAL_SHADER_PARAMETER_STRUCT()
+
+
+IMPLEMENT_GLOBAL_SHADER_PARAMETER_STRUCT(FTranslucencyLightData, "TranslucencyLightsData");
+
+
+void SetupTranslucencyLightData(
+	const TSparseArray<FLightSceneInfoCompact>& Lights,
+	const FViewInfo& View,
+	FTranslucencyLightData* LightData)
+{
+	LightData->Count = 0;
+
+	for (auto Light : Lights)
+	{
+		if (Light.LightSceneInfo->Proxy->HasStaticLighting() && Light.LightSceneInfo->IsPrecomputedLightingValid()) continue;
+
+		FLightShaderParameters LightParameters;
+		Light.LightSceneInfo->Proxy->GetLightShaderParameters(LightParameters);
+
+		if (Light.LightSceneInfo->Proxy->IsInverseSquared())
+		{
+			LightParameters.FalloffExponent = 0;
+		}
+
+		LightData->Type[LightData->Count] = Light.LightType;
+		LightData->LightPosition[LightData->Count] = LightParameters.Position;
+		LightData->LightInvRadius[LightData->Count] = LightParameters.InvRadius;
+		LightData->LightColor[LightData->Count] = LightParameters.Color;
+		LightData->LightFalloffExponent[LightData->Count] = LightParameters.FalloffExponent;
+		LightData->Direction[LightData->Count] = LightParameters.Direction;
+		LightData->Tangent[LightData->Count] = LightParameters.Tangent;
+		LightData->SpotAngles[LightData->Count] = LightParameters.SpotAngles;
+		LightData->SpecularScale[LightData->Count] = LightParameters.SpecularScale;
+		LightData->SourceRadius[LightData->Count] = LightParameters.SourceRadius;
+		LightData->SourceLength[LightData->Count] = LightParameters.SourceLength;
+		LightData->SoftSourceRadius[LightData->Count] = LightParameters.SoftSourceRadius;
+
+		const FVector2D FadeParams = Light.LightSceneInfo->Proxy->GetDirectionalLightDistanceFadeParameters(View.GetFeatureLevel(), Light.LightSceneInfo->IsPrecomputedLightingValid(), View.MaxShadowCascades);
+		LightData->DistanceFadeMAD[LightData->Count] = FVector2D(FadeParams.Y, -FadeParams.X * FadeParams.Y);
+
+		LightData->Count++;
+
+		if (LightData->Count >= GTranslucencyLightCountMaximum) break;
+	}
+
+	LightData->DummyRectLightTexture = GWhiteTexture->TextureRHI; //#dxr_todo: replace with valid textures per rect light
+}
+
+
+class FRayTracingTranslucencyRG : public FGlobalShader
+{
+	DECLARE_GLOBAL_SHADER(FRayTracingTranslucencyRG)
+	SHADER_USE_ROOT_PARAMETER_STRUCT(FRayTracingTranslucencyRG, FGlobalShader)
+
+		class FDenoiserOutput : SHADER_PERMUTATION_BOOL("DIM_DENOISER_OUTPUT");
+	using FPermutationDomain = TShaderPermutationDomain<FDenoiserOutput>;
+
+	BEGIN_SHADER_PARAMETER_STRUCT(FParameters, )
+		SHADER_PARAMETER(int32, SamplesPerPixel)
+		SHADER_PARAMETER(int32, HeightFog)
+		SHADER_PARAMETER(int32, ShouldDoDirectLighting)
+		SHADER_PARAMETER(int32, ShouldDoReflectedShadows)
+		SHADER_PARAMETER(int32, ShouldDoEmissiveAndIndirectLighting)
+		SHADER_PARAMETER(int32, UpscaleFactor)
+		SHADER_PARAMETER(float, TranslucencyMinRayDistance)
+		SHADER_PARAMETER(float, TranslucencyMaxRayDistance)
+		SHADER_PARAMETER(float, TranslucencyMaxRoughness)
+
+		SHADER_PARAMETER_SRV(RaytracingAccelerationStructure, TLAS)
+
+		SHADER_PARAMETER_TEXTURE(Texture2D, LTCMatTexture)
+		SHADER_PARAMETER_SAMPLER(SamplerState, LTCMatSampler)
+		SHADER_PARAMETER_TEXTURE(Texture2D, LTCAmpTexture)
+		SHADER_PARAMETER_SAMPLER(SamplerState, LTCAmpSampler)
+
+		SHADER_PARAMETER_STRUCT_REF(FViewUniformShaderParameters, ViewUniformBuffer)
+		SHADER_PARAMETER_STRUCT_REF(FSceneTexturesUniformParameters, SceneTexturesStruct)
+		SHADER_PARAMETER_STRUCT_REF(FTranslucencyLightData, LightData)
+		SHADER_PARAMETER_STRUCT_REF(FReflectionUniformParameters, ReflectionStruct)
+		SHADER_PARAMETER_STRUCT_REF(FFogUniformParameters, FogUniformParameters)
+
+		SHADER_PARAMETER_RDG_TEXTURE_UAV(RWTexture2D<float4>, ColorOutput)
+		SHADER_PARAMETER_RDG_TEXTURE_UAV(RWTexture2D<float>, RayHitDistanceOutput)
+		END_SHADER_PARAMETER_STRUCT()
+
+		static bool ShouldCompilePermutation(const FGlobalShaderPermutationParameters& Parameters)
+	{
+		return ShouldCompileRayTracingShadersForProject(Parameters.Platform);
+	}
+};
+
+class FRayTracingTranslucencyCHS : public FGlobalShader
+{
+	DECLARE_GLOBAL_SHADER(FRayTracingTranslucencyCHS);
+	SHADER_USE_ROOT_PARAMETER_STRUCT(FRayTracingTranslucencyCHS, FGlobalShader)
+
+		static bool ShouldCompilePermutation(const FGlobalShaderPermutationParameters& Parameters)
+	{
+		return ShouldCompileRayTracingShadersForProject(Parameters.Platform);
+	}
+
+	using FParameters = FEmptyShaderParameters;
+};
+
+class FRayTracingTranslucencyMS : public FGlobalShader
+{
+	DECLARE_GLOBAL_SHADER(FRayTracingTranslucencyMS);
+	SHADER_USE_ROOT_PARAMETER_STRUCT(FRayTracingTranslucencyMS, FGlobalShader)
+
+		static bool ShouldCompilePermutation(const FGlobalShaderPermutationParameters& Parameters)
+	{
+		return ShouldCompileRayTracingShadersForProject(Parameters.Platform);
+	}
+
+	using FParameters = FEmptyShaderParameters;
+};
+
+IMPLEMENT_GLOBAL_SHADER(FRayTracingTranslucencyRG, "/Engine/Private/RayTracing/RayTracingTranslucency.usf", "RayTracingTranslucencyRGS", SF_RayGen);
+IMPLEMENT_GLOBAL_SHADER(FRayTracingTranslucencyCHS, "/Engine/Private/RayTracing/RayTracingTranslucency.usf", "RayTracingTranslucencyMainCHS", SF_RayHitGroup);
+IMPLEMENT_GLOBAL_SHADER(FRayTracingTranslucencyMS, "/Engine/Private/RayTracing/RayTracingTranslucency.usf", "RayTracingTranslucencyMainMS", SF_RayMiss);
+
+
+//#dxr-todo: should we unify it with the composition happening in the non raytraced translucency pass? In that case it should use FCopySceneColorPS
+// Probably, but the architecture depends on the denoiser -> discuss
+
+class FCompositeTranslucencyPS : public FGlobalShader
+{
+	DECLARE_SHADER_TYPE(FCompositeTranslucencyPS, Global);
+
+public:
+	static bool ShouldCache(EShaderPlatform Platform)
+	{
+		return ShouldCompileRayTracingShadersForProject(Platform);
+	}
+
+	static bool ShouldCompilePermutation(const FGlobalShaderPermutationParameters& Parameters)
+	{
+		return ShouldCompileRayTracingShadersForProject(Parameters.Platform);
+	}
+
+	static void ModifyCompilationEnvironment(const FGlobalShaderPermutationParameters& Parameters, FShaderCompilerEnvironment& OutEnvironment)
+	{
+		FGlobalShader::ModifyCompilationEnvironment(Parameters, OutEnvironment);
+	}
+
+	FCompositeTranslucencyPS() {}
+	virtual ~FCompositeTranslucencyPS() {}
+
+	FCompositeTranslucencyPS(const ShaderMetaType::CompiledShaderInitializerType& Initializer)
+		: FGlobalShader(Initializer)
+	{
+		SceneTextureParameters.Bind(Initializer);
+		TranslucencyTextureParameter.Bind(Initializer.ParameterMap, TEXT("TranslucencyTexture"));
+		TranslucencyTextureSamplerParameter.Bind(Initializer.ParameterMap, TEXT("TranslucencyTextureSampler"));
+	}
+
+	template<typename TRHICommandList>
+	void SetParameters(
+		TRHICommandList& RHICmdList,
+		const FViewInfo& View,
+		FTextureRHIParamRef TranslucencyTexture,
+		FTextureRHIParamRef HitDistanceTexture)
+	{
+		const FPixelShaderRHIParamRef ShaderRHI = GetPixelShader();
+		FGlobalShader::SetParameters<FViewUniformShaderParameters>(RHICmdList, ShaderRHI, View.ViewUniformBuffer);
+		SceneTextureParameters.Set(RHICmdList, ShaderRHI, View.FeatureLevel, ESceneTextureSetupMode::All);
+
+		SetTextureParameter(RHICmdList, ShaderRHI, TranslucencyTextureParameter, TranslucencyTextureSamplerParameter, TStaticSamplerState<SF_Bilinear>::GetRHI(), TranslucencyTexture);
+		// #dxr_todo: Use hit-distance texture for denoising
+	}
+
+	virtual bool Serialize(FArchive& Ar) override
+	{
+		bool bShaderHasOutdatedParameters = FGlobalShader::Serialize(Ar);
+		Ar << TranslucencyTextureParameter;
+		Ar << TranslucencyTextureSamplerParameter;
+		return bShaderHasOutdatedParameters;
+	}
+
+private:
+	FSceneTextureShaderParameters SceneTextureParameters;
+	FShaderResourceParameter TranslucencyTextureParameter;
+	FShaderResourceParameter TranslucencyTextureSamplerParameter;
+};
+
+IMPLEMENT_SHADER_TYPE(, FCompositeTranslucencyPS, TEXT("/Engine/Private/RayTracing/CompositeTranslucencyPS.usf"), TEXT("CompositeTranslucencyPS"), SF_Pixel)
+
+
+void FDeferredShadingSceneRenderer::RayTraceTranslucency(FRHICommandListImmediate& RHICmdList)
+{
+	for (int32 ViewIndex = 0, Num = Views.Num(); ViewIndex < Num; ViewIndex++)
+	{
+		FViewInfo& View = Views[ViewIndex];
+
+		SCOPED_DRAW_EVENT(RHICmdList, RayTracingTranslucency);
+		SCOPED_GPU_STAT(RHICmdList, RayTracingTranslucency);
+
+		FRDGBuilder GraphBuilder(RHICmdList); //#dxr_todo: convert the entire translucency pass to render graph.
+		FSceneRenderTargets& SceneContext = FSceneRenderTargets::Get(GraphBuilder.RHICmdList);
+
+		FSceneViewFamilyBlackboard SceneBlackboard;
+		SetupSceneViewFamilyBlackboard(GraphBuilder, &SceneBlackboard);
+
+		//#dxr_todo: do not use reflections denoiser structs but separated ones
+		IScreenSpaceDenoiser::FReflectionsInputs DenoiserInputs;
+		float ResolutionFraction = 1.0f;
+
+		RayTraceTranslucencyView(
+			GraphBuilder,
+			View, &DenoiserInputs.Color, &DenoiserInputs.RayHitDistance,
+			GRayTracingTranslucencySamplesPerPixel, GRayTracingTranslucencyHeightFog, ResolutionFraction);
+
+		//#dxr_todo: denoise: replace DenoiserInputs with DenoiserOutputs in the following lines!
+		TRefCountPtr<IPooledRenderTarget> TranslucencyColor = GSystemTextures.BlackDummy;
+		TRefCountPtr<IPooledRenderTarget> TranslucencyHitDistanceColor = GSystemTextures.BlackDummy;
+
+		GraphBuilder.QueueTextureExtraction(DenoiserInputs.Color, &TranslucencyColor);
+		GraphBuilder.QueueTextureExtraction(DenoiserInputs.RayHitDistance, &TranslucencyHitDistanceColor);
+
+		GraphBuilder.Execute();
+
+		// Compositing result with the scene color
+		//#dxr-todo: should we unify it with the composition happening in the non raytraced translucency pass? In that case it should use FCopySceneColorPS
+		// Probably, but the architecture depends on the denoiser -> discuss
+		{
+			const auto ShaderMap = GetGlobalShaderMap(FeatureLevel);
+			TShaderMapRef<FPostProcessVS> VertexShader(ShaderMap);
+			TShaderMapRef<FCompositeTranslucencyPS> PixelShader(ShaderMap);
+
+			FGraphicsPipelineStateInitializer GraphicsPSOInit;
+			SceneContext.BeginRenderingSceneColor(RHICmdList, ESimpleRenderTargetMode::EExistingColorAndDepth, FExclusiveDepthStencil::DepthNop_StencilNop, true);
+			RHICmdList.ApplyCachedRenderTargets(GraphicsPSOInit);
+
+			//#dxr_todo review render mode
+			GraphicsPSOInit.BlendState = TStaticBlendState<CW_RGBA, BO_Add, BF_One, BF_Zero>::GetRHI();
+			GraphicsPSOInit.RasterizerState = TStaticRasterizerState<FM_Solid, CM_None>::GetRHI();
+			GraphicsPSOInit.DepthStencilState = TStaticDepthStencilState<false, CF_Always>::GetRHI();
+			GraphicsPSOInit.BoundShaderState.VertexDeclarationRHI = GFilterVertexDeclaration.VertexDeclarationRHI;
+			GraphicsPSOInit.BoundShaderState.VertexShaderRHI = GETSAFERHISHADER_VERTEX(*VertexShader);
+			GraphicsPSOInit.BoundShaderState.PixelShaderRHI = GETSAFERHISHADER_PIXEL(*PixelShader);
+			GraphicsPSOInit.PrimitiveType = PT_TriangleList;
+			SetGraphicsPipelineState(RHICmdList, GraphicsPSOInit);
+
+			RHICmdList.SetViewport(View.ViewRect.Min.X, View.ViewRect.Min.Y, 0.0f, View.ViewRect.Max.X, View.ViewRect.Max.Y, 1.0f);
+			PixelShader->SetParameters(RHICmdList, View, TranslucencyColor->GetRenderTargetItem().ShaderResourceTexture, TranslucencyHitDistanceColor->GetRenderTargetItem().ShaderResourceTexture);
+
+			DrawRectangle(
+				RHICmdList,
+				0, 0,
+				View.ViewRect.Width(), View.ViewRect.Height(),
+				View.ViewRect.Min.X, View.ViewRect.Min.Y,
+				View.ViewRect.Width(), View.ViewRect.Height(),
+				FIntPoint(View.ViewRect.Width(), View.ViewRect.Height()),
+				SceneContext.GetBufferSizeXY(),
+				*VertexShader);
+		}
+
+		ResolveSceneColor(RHICmdList);
+	}
+}
+
+void FDeferredShadingSceneRenderer::RayTraceTranslucencyView(
+	FRDGBuilder& GraphBuilder,
+	const FViewInfo& View,
+	FRDGTextureRef* OutColorTexture,
+	FRDGTextureRef* OutRayHitDistanceTexture,
+	int32 SamplePerPixel,
+	int32 HeightFog,
+	float ResolutionFraction)
+{
+	if (!View.bHasTranslucentViewMeshElements)
+	{
+		return;
+	}
+
+	FSceneRenderTargets& SceneContext = FSceneRenderTargets::Get(GraphBuilder.RHICmdList);
+
+	int32 UpscaleFactor = int32(1.0f / ResolutionFraction);
+	ensure(ResolutionFraction == 1.0 / UpscaleFactor);
+	ensureMsgf(FComputeShaderUtils::kGolden2DGroupSize % UpscaleFactor == 0, TEXT("Translucency ray tracing will have uv misalignement."));
+	FIntPoint RayTracingResolution = FIntPoint::DivideAndRoundUp(View.ViewRect.Size(), UpscaleFactor);
+
+	{
+		FPooledRenderTargetDesc Desc = SceneContext.GetSceneColor()->GetDesc();
+		Desc.Format = PF_FloatRGBA;
+		Desc.Flags &= ~(TexCreate_FastVRAM | TexCreate_Transient);
+		Desc.Extent /= UpscaleFactor;
+
+		*OutColorTexture = GraphBuilder.CreateTexture(Desc, TEXT("RayTracingTranslucency"));
+
+		Desc.Format = PF_R16F;
+		*OutRayHitDistanceTexture = GraphBuilder.CreateTexture(Desc, TEXT("RayTracingTranslucencyHitDistance"));
+	}
+
+	FRayTracingTranslucencyRG::FParameters* PassParameters = GraphBuilder.AllocParameters<FRayTracingTranslucencyRG::FParameters>();
+
+	PassParameters->SamplesPerPixel = SamplePerPixel;
+	PassParameters->HeightFog = HeightFog;
+	PassParameters->ShouldDoDirectLighting = GRayTracingTranslucencyDirectLighting;
+	PassParameters->ShouldDoReflectedShadows = GRayTracingTranslucencyShadows;
+	PassParameters->ShouldDoEmissiveAndIndirectLighting = GRayTracingTranslucencyEmissiveAndIndirectLighting;
+	PassParameters->UpscaleFactor = UpscaleFactor;
+	PassParameters->TranslucencyMinRayDistance = FMath::Min(GRayTracingTranslucencyMinRayDistance, GRayTracingTranslucencyMaxRayDistance);
+	PassParameters->TranslucencyMaxRayDistance = GRayTracingTranslucencyMaxRayDistance;
+	//#dxr-todo: do we want to use SSR parameter here?
+	PassParameters->TranslucencyMaxRoughness = FMath::Clamp(View.FinalPostProcessSettings.ScreenSpaceReflectionMaxRoughness, 0.01f, 1.0f); 
+	PassParameters->LTCMatTexture = GSystemTextures.LTCMat->GetRenderTargetItem().ShaderResourceTexture;
+	PassParameters->LTCMatSampler = TStaticSamplerState<SF_Bilinear, AM_Clamp, AM_Clamp, AM_Clamp>::GetRHI();
+	PassParameters->LTCAmpTexture = GSystemTextures.LTCAmp->GetRenderTargetItem().ShaderResourceTexture;
+	PassParameters->LTCAmpSampler = TStaticSamplerState<SF_Bilinear, AM_Clamp, AM_Clamp, AM_Clamp>::GetRHI();
+
+	PassParameters->TLAS = View.PerViewRayTracingScene.RayTracingSceneRHI->GetShaderResourceView();
+	PassParameters->ViewUniformBuffer = View.ViewUniformBuffer;
+	{
+		FTranslucencyLightData LightData;
+		SetupTranslucencyLightData(Scene->Lights, View, &LightData);
+		PassParameters->LightData = CreateUniformBufferImmediate(LightData, EUniformBufferUsage::UniformBuffer_SingleDraw);
+	}
+	{ // TODO: use FSceneViewFamilyBlackboard.
+		FSceneTexturesUniformParameters SceneTextures;
+		SetupSceneTextureUniformParameters(SceneContext, FeatureLevel, ESceneTextureSetupMode::All, SceneTextures);
+		PassParameters->SceneTexturesStruct = CreateUniformBufferImmediate(SceneTextures, EUniformBufferUsage::UniformBuffer_SingleDraw);
+	}
+	{
+		FReflectionUniformParameters ReflectionStruct;
+		SetupReflectionUniformParameters(View, ReflectionStruct);
+		PassParameters->ReflectionStruct = CreateUniformBufferImmediate(ReflectionStruct, EUniformBufferUsage::UniformBuffer_SingleDraw);
+	}
+	{
+		FFogUniformParameters FogStruct;
+		SetupFogUniformParameters(View, FogStruct);
+		PassParameters->FogUniformParameters = CreateUniformBufferImmediate(FogStruct, EUniformBufferUsage::UniformBuffer_SingleDraw);
+	}
+	PassParameters->ColorOutput = GraphBuilder.CreateUAV(*OutColorTexture);
+	PassParameters->RayHitDistanceOutput = GraphBuilder.CreateUAV(*OutRayHitDistanceTexture);
+
+	auto RayGenShader = View.ShaderMap->GetShader<FRayTracingTranslucencyRG>();
+	ClearUnusedGraphResources(RayGenShader, PassParameters);
+
+	GraphBuilder.AddPass(
+		RDG_EVENT_NAME("TranslucencyRayTracing %dx%d", RayTracingResolution.X, RayTracingResolution.X),
+		PassParameters,
+		ERenderGraphPassFlags::Compute,
+		[PassParameters, this, &View, RayGenShader, RayTracingResolution](FRHICommandList& RHICmdList)
+	{
+		auto ClosestHitShader = View.ShaderMap->GetShader<FRayTracingTranslucencyCHS>();
+		auto MissShader = View.ShaderMap->GetShader<FRayTracingTranslucencyMS>();
+
+		FRHIRayTracingPipelineState* Pipeline = BindRayTracingPipeline(
+			RHICmdList, View,
+			RayGenShader->GetRayTracingShader(),
+			MissShader->GetRayTracingShader(),
+			ClosestHitShader->GetRayTracingShader()); // #dxr_todo: this should be done once at load-time and cached
+
+		FRayTracingShaderBindingsWriter GlobalResources;
+		SetShaderParameters(GlobalResources, RayGenShader, *PassParameters);
+
+		FRayTracingSceneRHIParamRef RayTracingSceneRHI = View.PerViewRayTracingScene.RayTracingSceneRHI;
+		const uint32 RayGenShaderIndex = 0;
+		RHICmdList.RayTraceDispatch(Pipeline, RayGenShaderIndex, RayTracingSceneRHI, GlobalResources, RayTracingResolution.X, RayTracingResolution.Y);
+	});
+}
+
+
+#endif // RHI_RAYTRACING
