Fixed render graph event name for ray tracing reflections pass

diff --git a/Engine/Source/Runtime/Renderer/Private/RayTracing/RayTracingReflections.cpp b/Engine/Source/Runtime/Renderer/Private/RayTracing/RayTracingReflections.cpp
--- a/Engine/Source/Runtime/Renderer/Private/RayTracing/RayTracingReflections.cpp
+++ b/Engine/Source/Runtime/Renderer/Private/RayTracing/RayTracingReflections.cpp
@@ -268,7 +268,7 @@ void FDeferredShadingSceneRenderer::RayTraceReflections(
 	ClearUnusedGraphResources(RayGenShader, PassParameters);
 
 	GraphBuilder.AddPass(
-		RDG_EVENT_NAME("ReflectionRayTracing %dx%d", RayTracingResolution.X, RayTracingResolution.X),
+		RDG_EVENT_NAME("ReflectionRayTracing %dx%d", RayTracingResolution.X, RayTracingResolution.Y),
 		PassParameters,
 		ERenderGraphPassFlags::Compute,
 		[PassParameters, this, &View, RayGenShader, RayTracingResolution](FRHICommandList& RHICmdList)
