Move FSkyLightData to central header.
Removed dead code.

diff --git a/Engine/Source/Runtime/Renderer/Private/RayTracing/RayTracingSkyLight.h b/Engine/Source/Runtime/Renderer/Private/RayTracing/RayTracingSkyLight.h
new file mode 100644
--- /dev/null
+++ b/Engine/Source/Runtime/Renderer/Private/RayTracing/RayTracingSkyLight.h
@@ -0,0 +1,31 @@
+// Copyright 1998-2019 Epic Games, Inc. All Rights Reserved.
+
+/*=============================================================================
+RaytracingOptions.h declares ray tracing options for use in rendering
+=============================================================================*/
+
+#pragma once
+
+#include "UniformBuffer.h"
+
+BEGIN_GLOBAL_SHADER_PARAMETER_STRUCT(FSkyLightData, )
+SHADER_PARAMETER(int, SamplesPerPixel)
+SHADER_PARAMETER(int, SamplingStopLevel)
+SHADER_PARAMETER(FVector, Color)
+SHADER_PARAMETER(FIntVector, MipDimensions)
+SHADER_PARAMETER_TEXTURE(TextureCube, Texture)
+SHADER_PARAMETER_SAMPLER(SamplerState, TextureSampler)
+SHADER_PARAMETER_SRV(Buffer<float>, MipTreePosX)
+SHADER_PARAMETER_SRV(Buffer<float>, MipTreeNegX)
+SHADER_PARAMETER_SRV(Buffer<float>, MipTreePosY)
+SHADER_PARAMETER_SRV(Buffer<float>, MipTreeNegY)
+SHADER_PARAMETER_SRV(Buffer<float>, MipTreePosZ)
+SHADER_PARAMETER_SRV(Buffer<float>, MipTreeNegZ)
+SHADER_PARAMETER_SRV(Buffer<float>, MipTreePdfPosX)
+SHADER_PARAMETER_SRV(Buffer<float>, MipTreePdfNegX)
+SHADER_PARAMETER_SRV(Buffer<float>, MipTreePdfPosY)
+SHADER_PARAMETER_SRV(Buffer<float>, MipTreePdfNegY)
+SHADER_PARAMETER_SRV(Buffer<float>, MipTreePdfPosZ)
+SHADER_PARAMETER_SRV(Buffer<float>, MipTreePdfNegZ)
+SHADER_PARAMETER_SRV(Buffer<float>, SolidAnglePdf)
+END_GLOBAL_SHADER_PARAMETER_STRUCT()
diff --git a/Engine/Source/Runtime/Renderer/Private/RayTracing/RaytracingSkylight.cpp b/Engine/Source/Runtime/Renderer/Private/RayTracing/RaytracingSkylight.cpp
--- a/Engine/Source/Runtime/Renderer/Private/RayTracing/RaytracingSkylight.cpp
+++ b/Engine/Source/Runtime/Renderer/Private/RayTracing/RaytracingSkylight.cpp
@@ -18,6 +18,7 @@ bool IsRayTracingSkyLightSelected()
 
 #if RHI_RAYTRACING
 
+#include "RayTracingSkyLight.h"
 #include "ClearQuad.h"
 #include "DistanceFieldAmbientOcclusion.h"
 #include "SceneRendering.h"
@@ -52,117 +53,6 @@ static FAutoConsoleVariableRef CVarRayTracingSkyLightSamplingStopLevel(
 	TEXT("Sets the stop level for MIP-sampling (default = 0)")
 );
 
-
-#if 0
-class FSkylightRG : public FGlobalShader
-{
-	DECLARE_GLOBAL_SHADER(FSkylightRG)
-	SHADER_USE_ROOT_PARAMETER_STRUCT(FSkylightRG, FGlobalShader)
-
-	static bool ShouldCompilePermutation(const FGlobalShaderPermutationParameters& Parameters)
-	{
-		return ShouldCompileRayTracingShadersForProject(Parameters.Platform);
-	}
-
-	BEGIN_SHADER_PARAMETER_STRUCT(FParameters, )
-		SHADER_PARAMETER(uint32, SamplesPerPixel)
-		SHADER_PARAMETER(FVector, SkyLightColor)
-
-		//SHADER_PARAMETER_TEXTURE(TextureCube, SkyLightTexture)
-		//SHADER_PARAMETER_SAMPLER(SamplerState, SkyLightTextureSampler)
-
-		// Ray tracing parameters
-		SHADER_PARAMETER_SRV(RaytracingAccelerationStructure, TLAS)
-		SHADER_PARAMETER_RDG_TEXTURE_UAV(RWTexture2D<float4>, RWOcclusionMaskUAV)
-		SHADER_PARAMETER_RDG_TEXTURE_UAV(RWTexture2D<float>, RWRayDistanceUAV)
-		SHADER_PARAMETER_STRUCT_REF(FViewUniformShaderParameters, ViewUniformBuffer)
-		SHADER_PARAMETER_STRUCT_REF(FSceneTexturesUniformParameters, SceneTexturesStruct)
-	END_SHADER_PARAMETER_STRUCT()
-};
-
-IMPLEMENT_GLOBAL_SHADER(FSkylightRG, "/Engine/Private/Raytracing/RaytracingSkylightRGS.usf", "SkyLightRGS", SF_RayGen);
-
-void FDeferredShadingSceneRenderer::RenderRayTracingSkyLight(
-	FRHICommandListImmediate& RHICmdList,
-	TRefCountPtr<IPooledRenderTarget>& OutSkyLightTexture
-)
-{
-	//check(LightSceneInfo);
-	check(Views.Num() == 1);
-	const FViewInfo& View = Views[0];
-
-	FSceneRenderTargets& SceneContext = FSceneRenderTargets::Get(RHICmdList);
-	FRDGBuilder GraphBuilder(RHICmdList);
-
-	// Render targets
-	FRDGTextureRef ScreenShadowMaskTexture;
-	{
-		FRDGTextureDesc Desc = SceneContext.GetSceneColor()->GetDesc();
-		Desc.Format = PF_FloatRGBA;
-		Desc.Flags &= ~(TexCreate_FastVRAM | TexCreate_Transient);
-		ScreenShadowMaskTexture = GraphBuilder.CreateTexture(Desc, TEXT("RayTracingSkylight"));
-	}
-
-	FRDGTextureRef RayDistanceTexture;
-	{
-		FRDGTextureDesc Desc = SceneContext.GetSceneColor()->GetDesc();
-		Desc.Format = PF_R16F;
-		Desc.Flags &= ~(TexCreate_FastVRAM | TexCreate_Transient);
-		RayDistanceTexture = GraphBuilder.CreateTexture(Desc, TEXT("RayTracingOcclusionDistance"));
-	}
-
-	// TODO(RDG): use FSceneViewFamilyBlackboard.
-	FSceneTexturesUniformParameters SceneTextures;
-	SetupSceneTextureUniformParameters(SceneContext, FeatureLevel, ESceneTextureSetupMode::All, SceneTextures);
-
-	FSkylightRG::FParameters* PassParameters = GraphBuilder.AllocParameters<FSkylightRG::FParameters>();
-	PassParameters->RWOcclusionMaskUAV = GraphBuilder.CreateUAV(FRDGTextureUAVDesc(ScreenShadowMaskTexture));
-	PassParameters->RWRayDistanceUAV = GraphBuilder.CreateUAV(FRDGTextureUAVDesc(RayDistanceTexture));
-	PassParameters->SamplesPerPixel = GRayTracingSkyLightSamplesPerPixel;
-	PassParameters->SamplingStopLevel = GRayTracingSkyLightSamplingStopLevel;
-	PassParameters->SkyLightColor = Scene->SkyLight->LightColor
-	//PassParameters->SkyLightTexture = Scene->SkyLight->ProcessedTexture->TextureRHI;
-	//PassParameters->SkyLightTextureSampler = Scene->SkyLight->ProcessedTexture->SamplerStateRHI;
-	PassParameters->TLAS = View.PerViewRayTracingScene.RayTracingSceneRHI->GetShaderResourceView();
-	PassParameters->ViewUniformBuffer = View.ViewUniformBuffer;
-	PassParameters->SceneTexturesStruct = CreateUniformBufferImmediate(SceneTextures, EUniformBufferUsage::UniformBuffer_SingleDraw);
-
-	TShaderMapRef<FSkylightRG> RayGenerationShader(GetGlobalShaderMap(FeatureLevel));
-
-	FRayGenShaderUtils::AddRayTraceDispatchPass(
-		GraphBuilder,
-		RDG_EVENT_NAME("RayTracedSkyLight %dx%d", View.ViewRect.Width(), View.ViewRect.Height()),
-		*RayGenerationShader,
-		PassParameters,
-		View.ViewRect.Size());
-
-	GraphBuilder.QueueTextureExtraction(ScreenShadowMaskTexture, &OutSkyLightTexture);
-	GraphBuilder.Execute();
-}
-#else
-
-BEGIN_GLOBAL_SHADER_PARAMETER_STRUCT(FSkyLightData, )
-SHADER_PARAMETER(int, SamplesPerPixel)
-SHADER_PARAMETER(int, SamplingStopLevel)
-SHADER_PARAMETER(FVector, Color)
-SHADER_PARAMETER(FIntVector, MipDimensions)
-SHADER_PARAMETER_TEXTURE(TextureCube, Texture)
-SHADER_PARAMETER_SAMPLER(SamplerState, TextureSampler)
-SHADER_PARAMETER_SRV(Buffer<float>, MipTreePosX)
-SHADER_PARAMETER_SRV(Buffer<float>, MipTreeNegX)
-SHADER_PARAMETER_SRV(Buffer<float>, MipTreePosY)
-SHADER_PARAMETER_SRV(Buffer<float>, MipTreeNegY)
-SHADER_PARAMETER_SRV(Buffer<float>, MipTreePosZ)
-SHADER_PARAMETER_SRV(Buffer<float>, MipTreeNegZ)
-SHADER_PARAMETER_SRV(Buffer<float>, MipTreePdfPosX)
-SHADER_PARAMETER_SRV(Buffer<float>, MipTreePdfNegX)
-SHADER_PARAMETER_SRV(Buffer<float>, MipTreePdfPosY)
-SHADER_PARAMETER_SRV(Buffer<float>, MipTreePdfNegY)
-SHADER_PARAMETER_SRV(Buffer<float>, MipTreePdfPosZ)
-SHADER_PARAMETER_SRV(Buffer<float>, MipTreePdfNegZ)
-SHADER_PARAMETER_SRV(Buffer<float>, SolidAnglePdf)
-END_GLOBAL_SHADER_PARAMETER_STRUCT()
-
 IMPLEMENT_GLOBAL_SHADER_PARAMETER_STRUCT(FSkyLightData, "SkyLight");
 
 DECLARE_GPU_STAT_NAMED(RayTracingSkyLight, TEXT("Ray Tracing SkyLight"));
@@ -1030,7 +920,4 @@ void FDeferredShadingSceneRenderer::CompositeRayTracingSkyLight(
 	SceneContext.FinishRenderingSceneColor(RHICmdList);
 }
 
-
-#endif
-
 #endif // RHI_RAYTRACING
