Resupported r.DoLazyStaticMeshUpdate with mesh draw commands.  Now enabled when in game mode, no longer requires cooked build.
 * GPUPerfTest Total cached MeshDrawCommands 2.324Mb -> 1.691Mb and Primitive StaticMeshCommandInfos 3529.2Kb -> 563.7Kb

diff --git a/Engine/Source/Runtime/Renderer/Private/RendererScene.cpp b/Engine/Source/Runtime/Renderer/Private/RendererScene.cpp
--- a/Engine/Source/Runtime/Renderer/Private/RendererScene.cpp
+++ b/Engine/Source/Runtime/Renderer/Private/RendererScene.cpp
@@ -523,14 +523,12 @@ void FScene::CheckPrimitiveArrays()
 static TAutoConsoleVariable<int32> CVarDoLazyStaticMeshUpdate(
 	TEXT("r.DoLazyStaticMeshUpdate"),
 	0,
-	TEXT("If true, then do not add meshes to the static mesh draw lists until they are visible. Experiemental option. Incompatible with the Mesh Draw Command pipeline."));
+	TEXT("If true, then do not add meshes to the static mesh draw lists until they are visible. Experiemental option."));
 
 static void DoLazyStaticMeshUpdateCVarSinkFunction()
 {
-	//@todo MeshCommandPipeline r.DoLazyStaticMeshUpdate isn't currently supported.
-	/*
-	static bool CachedDoLazyStaticMeshUpdate = CVarDoLazyStaticMeshUpdate.GetValueOnGameThread() && !WITH_EDITOR;
-	bool DoLazyStaticMeshUpdate = CVarDoLazyStaticMeshUpdate.GetValueOnGameThread() && !WITH_EDITOR;
+	static bool CachedDoLazyStaticMeshUpdate = CVarDoLazyStaticMeshUpdate.GetValueOnGameThread() && !GIsEditor;
+	bool DoLazyStaticMeshUpdate = CVarDoLazyStaticMeshUpdate.GetValueOnGameThread() && !GIsEditor;
 
 	if (DoLazyStaticMeshUpdate != CachedDoLazyStaticMeshUpdate)
 	{
@@ -549,22 +547,18 @@ static void DoLazyStaticMeshUpdateCVarSinkFunction()
 			}
 		}
 	}
-	*/
 }
 
 static FAutoConsoleVariableSink CVarDoLazyStaticMeshUpdateSink(FConsoleCommandDelegate::CreateStatic(&DoLazyStaticMeshUpdateCVarSinkFunction));
 
 void FScene::UpdateDoLazyStaticMeshUpdate(FRHICommandListImmediate& CmdList)
 {
-	//@todo MeshCommandPipeline r.DoLazyStaticMeshUpdate isn't currently supported.
-	/*
-	bool DoLazyStaticMeshUpdate = CVarDoLazyStaticMeshUpdate.GetValueOnRenderThread() && !WITH_EDITOR;
+	bool DoLazyStaticMeshUpdate = CVarDoLazyStaticMeshUpdate.GetValueOnRenderThread() && !GIsEditor;
 
 	for (int32 PrimitiveIndex = 0; PrimitiveIndex < Primitives.Num(); PrimitiveIndex++)
 	{
 		Primitives[PrimitiveIndex]->UpdateStaticMeshes(CmdList, !DoLazyStaticMeshUpdate);
 	}
-*/
 }
 
 void FScene::DumpMeshDrawCommandMemoryStats()
@@ -783,26 +777,23 @@ void FScene::AddPrimitiveSceneInfo_RenderThread(FRHICommandListImmediate& RHICmd
 	// Set lod Parent information if valid
 	PrimitiveSceneInfo->LinkLODParentComponent();
 
-	// Add the primitive to the scene.
-	PrimitiveSceneInfo->AddToScene(RHICmdList, true);
-
-	//@todo MeshCommandPipeline r.DoLazyStaticMeshUpdate isn't currently supported.
-	/*
-#if WITH_EDITOR
-	PrimitiveSceneInfo->AddToScene(RHICmdList, true);
-#else
-	const bool bAddToDrawLists = !(CVarDoLazyStaticMeshUpdate.GetValueOnRenderThread());
-	if (bAddToDrawLists)
+	if (GIsEditor)
 	{
 		PrimitiveSceneInfo->AddToScene(RHICmdList, true);
 	}
 	else
 	{
-		PrimitiveSceneInfo->AddToScene(RHICmdList, true, false);
-		PrimitiveSceneInfo->BeginDeferredUpdateStaticMeshes();
+		const bool bAddToDrawLists = !(CVarDoLazyStaticMeshUpdate.GetValueOnRenderThread());
+		if (bAddToDrawLists)
+		{
+			PrimitiveSceneInfo->AddToScene(RHICmdList, true);
+		}
+		else
+		{
+			PrimitiveSceneInfo->AddToScene(RHICmdList, true, false);
+			PrimitiveSceneInfo->BeginDeferredUpdateStaticMeshes();
+		}
 	}
-#endif
-	*/
 
 	AddPrimitiveToUpdateGPU(*this, SourceIndex);
 
diff --git a/Engine/Source/Runtime/Renderer/Private/SceneVisibility.cpp b/Engine/Source/Runtime/Renderer/Private/SceneVisibility.cpp
--- a/Engine/Source/Runtime/Renderer/Private/SceneVisibility.cpp
+++ b/Engine/Source/Runtime/Renderer/Private/SceneVisibility.cpp
@@ -2765,11 +2765,9 @@ void FSceneRenderer::PreVisibilityFrameSetup(FRHICommandListImmediate& RHICmdLis
 	// Notify the RHI we are beginning to render a scene.
 	RHICmdList.BeginScene();
 
-	//@todo MeshCommandPipeline r.DoLazyStaticMeshUpdate isn't currently supported.
-	/*
 	{
 		static auto CVar = IConsoleManager::Get().FindConsoleVariable(TEXT("r.DoLazyStaticMeshUpdate"));
-		const bool DoLazyStaticMeshUpdate = (CVar->GetInt() && !WITH_EDITOR);
+		const bool DoLazyStaticMeshUpdate = (CVar->GetInt() && !GIsEditor);
 
 		if (DoLazyStaticMeshUpdate)
 		{
@@ -2785,7 +2783,7 @@ void FSceneRenderer::PreVisibilityFrameSetup(FRHICommandListImmediate& RHICmdLis
 				Scene->Primitives[RollingRemoveIndex]->UpdateStaticMeshes(RHICmdList, false);
 			}
 		}
-	}*/
+	}
 
 	// Notify the FX system that the scene is about to perform visibility checks.
 	if (Scene->FXSystem && !Views[0].bIsPlanarReflection)
