Ray tracing global illumination. Currently disabled by default to evaluate performance and quality in WoodChuck.

r.RayTracing.GlobalIllumination (enables RTGI, default = 0)
r.RayTracing.GlobalIllumination.Denoiser (enables denoising, default = 1)
r.RayTracing.GlobalIllumination.SamplesPerPixel (determines samples-per-pixel, default = 1)
r.RayTracing.GlobalIllumination.MaxBounces (determines maximum number of diffuse bounces, default = 1)
r.RayTracing.GlobalIllumination.EvalSkyLight (determines whether or not SkyLight can be hit with indirect diffuse, default = 0)
r.RayTracing.GlobalIllumination.MaxRayDistance (maximum ray distance for diffuse and visibility rays, default = 1.0e27)

diff --git a/Engine/Shaders/Private/RayTracing/RayTracingGlobalIlluminationCompositePS.usf b/Engine/Shaders/Private/RayTracing/RayTracingGlobalIlluminationCompositePS.usf
new file mode 100644
--- /dev/null
+++ b/Engine/Shaders/Private/RayTracing/RayTracingGlobalIlluminationCompositePS.usf
@@ -0,0 +1,21 @@
+// Copyright 1998-2019 Epic Games, Inc. All Rights Reserved.
+
+#include "../Common.ush"
+#include "../DeferredShadingCommon.ush"
+
+Texture2D<float4> GlobalIlluminationTexture;
+SamplerState GlobalIlluminationSampler;
+
+void GlobalIlluminationCompositePS(
+	in noperspective float2 UV : TEXCOORD0,
+	out float4 OutColor : SV_Target0,
+	out float OutAmbientOcclusion : SV_Target1
+)
+{
+	float4 GlobalIllumination = GlobalIlluminationTexture.Sample(GlobalIlluminationSampler, UV);
+	// Scene color is ultimately multiplied by the occlusion term, so we counteract the weighting term here
+	// dxr_todo: this is currently a hack, due to catastrophic cancelation when AO is close to 0.
+	OutAmbientOcclusion = max(1.0 - GlobalIllumination.a, 1.0 / 256.0);
+	GlobalIllumination.rgb *= rcp(OutAmbientOcclusion);
+	OutColor = float4(GlobalIllumination.rgb, 0.0);
+}
diff --git a/Engine/Shaders/Private/RayTracing/RayTracingGlobalIlluminationRGS.usf b/Engine/Shaders/Private/RayTracing/RayTracingGlobalIlluminationRGS.usf
new file mode 100644
--- /dev/null
+++ b/Engine/Shaders/Private/RayTracing/RayTracingGlobalIlluminationRGS.usf
@@ -0,0 +1,301 @@
+// Copyright 1998-2019 Epic Games, Inc. All Rights Reserved.
+
+#include "../Common.ush"
+#include "../RectLight.ush"
+//#include "../MonteCarlo.ush"
+#include "../DeferredShadingCommon.ush"
+#include "RayTracingCommon.ush"
+#include "RayTracingHitGroupCommon.ush"
+
+#include "../PathTracing/Utilities/PathTracingRandomSequence.ush" 
+#include "../PathTracing/Light/PathTracingLightSampling.ush"
+#include "../PathTracing/Material/PathTracingMaterialSampling.ush"
+
+RaytracingAccelerationStructure TLAS;
+RWTexture2D<float4> RWGlobalIlluminationUAV;
+RWTexture2D<float2> RWRayDistanceUAV;
+
+uint SamplesPerPixel;
+uint MaxBounces;
+float MaxRayDistance;
+bool EvalSkyLight;
+
+uint CalcLinearIndex(uint2 PixelCoord)
+{
+	return PixelCoord.y * View.BufferSizeAndInvSize.x + PixelCoord.x;
+}
+
+void GenerateCosineNormalRay(
+	float3 WorldPosition,
+	float3 WorldNormal,
+	inout RandomSequence RandSequence,
+	out float3 RayOrigin,
+	out float3 RayDirection,
+	out float RayTMin,
+	out float RayTMax,
+	out float RayPdf
+)
+{
+	// Draw random variable
+	float2 BufferSize = View.BufferSizeAndInvSize.xy;
+	uint DummyVariable;
+	float2 RandSample = RandomSequence_GenerateSample2D(RandSequence, DummyVariable);
+
+	// Perform cosine-hemispherical sampling and convert to world-space
+	float4 Direction_Tangent = CosineSampleHemisphere(RandSample);
+	float3 Direction_World = TangentToWorld(Direction_Tangent.xyz, WorldNormal);
+
+	RayOrigin = WorldPosition + WorldNormal * 0.01;
+	RayDirection = Direction_World;
+	RayTMin = 0.0;
+	RayTMax = MaxRayDistance;
+	RayPdf = Direction_Tangent.w;
+}
+
+float GetHitT(FMaterialClosestHitPayload HitInfo)
+{
+	return HitInfo.HitT;
+}
+
+bool IsHit(RayDesc Ray, FMaterialClosestHitPayload HitInfo)
+{
+	return HitInfo.HitT >= 0.0;
+}
+
+[shader("raygeneration")]
+void GlobalIlluminationRGS()
+{
+	uint2 PixelCoord = DispatchRaysIndex().xy;
+	RandomSequence RandSequence;
+	uint LinearIndex = CalcLinearIndex(PixelCoord);
+	RandomSequence_Initialize(RandSequence, LinearIndex, View.FrameNumber);
+
+	bool IsUnidirectionalEnabled = false;
+	bool IsNextEventEstimationEnabled = true;
+
+	// Get G-Buffer surface data
+	float2 InvBufferSize = View.BufferSizeAndInvSize.zw;
+	float2 UV = (float2(PixelCoord) + 0.5) * InvBufferSize;
+	FScreenSpaceData ScreenSpaceData = GetScreenSpaceData(UV);
+	float Depth = ScreenSpaceData.GBuffer.Depth;
+	float3 WorldPosition = ReconstructWorldPositionFromDepth(UV, Depth);
+	float3 WorldNormal = ScreenSpaceData.GBuffer.WorldNormal;
+
+	// #dxr_todo: Determine importance of computing GI based on BRDF and geometric term
+	float3 DiffuseColor = ScreenSpaceData.GBuffer.DiffuseColor;
+	float3 BaseColor = ScreenSpaceData.GBuffer.BaseColor;
+
+	float3 Irradiance = 0;
+	float HitDistance = 0.0;
+	float HitCount = 0.0;
+	for (uint SampleIndex = 0; SampleIndex < SamplesPerPixel; ++SampleIndex)
+	{
+#if 0
+		const uint TemporalBits = 7;
+		uint FrameIndex = View.FrameNumber % 1024;
+		RandomSequence_Initialize(RandSequence, LinearIndex, FrameIndex * SamplesPerPixel + SampleIndex);
+		RandSequence.Type = 2;
+#endif
+
+		// Initialize ray
+		RayDesc Ray;
+		float RayPdf = 1.0;
+#if 1
+		GenerateCosineNormalRay(WorldPosition, WorldNormal, RandSequence, Ray.Origin, Ray.Direction, Ray.TMin, Ray.TMax, RayPdf);
+		float3 RayThroughput = DiffuseColor * RayPdf;
+#else
+		uint DummyVariable;
+		float2 RandSample = RandomSequence_GenerateSample2D(RandSequence, DummyVariable);
+		float2 ViewportUV = (PixelCoord.xy + RandSample.xy) * View.BufferSizeAndInvSize.zw;
+		Ray.Origin = ReconstructWorldPositionFromDepth(ViewportUV, 0.0f);
+		Ray.Direction = normalize(ReconstructWorldPositionFromDepth(ViewportUV, 1.f) - Ray.Origin);
+		Ray.TMin = 0.0;
+		Ray.TMax = 1.0e12;
+		float3 RayThroughput = 1.0;
+#endif
+		Ray.TMax = MaxRayDistance;
+		
+		float MaterialPdf = 0.0;
+		uint Bounce = 0;
+		while (Bounce < MaxBounces)
+		{
+			// Cast ray
+			uint RayFlags = 0;
+			FRayCone RayCone = (FRayCone)0;
+			FMaterialClosestHitPayload Payload = TraceRayInternal(
+				TLAS,   // AccelerationStructure
+				RayFlags,
+				0xFF,   // InstanceInclusionMask
+				0,      // RayContributionToHitGroupIndex
+				1,      // MultiplierForGeometryContributionToShaderIndex #dxr_todo: we currently only support shader binding per mesh instance
+				0,      // MissShaderIndex
+				Ray,    // RayDesc
+				RayCone
+			);
+
+			// Environment hit
+			if (!IsHit(Ray, Payload))
+			{
+				// Optional multi-bounce SkyLight contribution
+				if (EvalSkyLight && Bounce > 0)
+				{
+					uint SkyLightId = 0;
+					float3 EnvironmentRadiance = 0.0;
+					SkyLight_EvalLight(SkyLightId, Ray.Direction, Ray, EnvironmentRadiance);
+					Irradiance += EnvironmentRadiance * RayThroughput / RayPdf;
+				}
+				break;
+			}
+			// #dxr_todo: Allow for material emission?
+
+			if (Bounce == 0)
+			{
+				HitDistance += Payload.HitT;
+				HitCount += 1.0;
+			}
+
+			// Update intersection
+			Ray.Origin += Ray.Direction * Payload.HitT;
+
+			// Perform next-event estimation
+			if (IsNextEventEstimationEnabled)
+			{
+				// #dxr_todo: Many-light selection
+				int LightId;
+				float3 LightUV;
+				float NeePdf = 0.0;
+				uint DummyVariable;
+				float4 RandSample4 = RandomSequence_GenerateSample4D(RandSequence, DummyVariable);
+				SampleLight(Ray, Payload, RandSample4, LightId, LightUV, NeePdf);
+
+				if (NeePdf > 0.0)
+				{
+					RayDesc LightRay;
+					GenerateLightRay(Ray, LightId, LightUV, LightRay);
+					LightRay.Origin += Payload.WorldNormal * 0.01;
+					// #dxr_todo: Guard against tracing a back-facing ray
+
+					// Trace visibility ray
+					uint RayFlags = RAY_FLAG_ACCEPT_FIRST_HIT_AND_END_SEARCH | RAY_FLAG_SKIP_CLOSEST_HIT_SHADER;
+					FRayCone LightRayCone = (FRayCone)0;
+					FMaterialClosestHitPayload NeePayload = TraceRayInternal(
+						TLAS,   // AccelerationStructure
+						RayFlags,
+						0xFF,   // InstanceInclusionMask
+						0,      // RayContributionToHitGroupIndex
+						1,      // MultiplierForGeometryContributionToShaderIndex #dxr_todo: we currently only support shader binding per mesh instance
+						0,      // MissShaderIndex
+						LightRay,    // RayDesc
+						LightRayCone
+					);
+
+					// No hit indicates successful next-event connection
+					if (!IsHit(LightRay, NeePayload))
+					{
+						// Evaluate radiance
+						float3 Radiance;
+						EvalLight(LightId, LightUV, LightRay, Radiance);
+
+						// Evaluate material
+						float3 MaterialThroughput;
+						float MaterialEvalPdf = 0.0;
+#if 0
+						EvalMaterial(Ray.Direction, LightRay.Direction, Payload, MaterialThroughput, MaterialEvalPdf);
+#else
+						MaterialEvalPdf = saturate(dot(Payload.WorldNormal, LightRay.Direction)) / PI;
+						MaterialThroughput = Payload.DiffuseColor * MaterialEvalPdf;
+#endif
+						// Apply material Pdf for correct MIS weight
+						float MisWeight = 1.0;
+#if 0
+						if (IsUnidirectionalEnabled && IsPhysicalLight(LightId))
+						{
+							MisWeight = NeePdf / (NeePdf + MaterialEvalPdf);
+						}
+#endif
+						// Record the contribution
+						float3 ExitantRadianceSample = Radiance * MaterialThroughput * RayThroughput * MisWeight / (NeePdf * RayPdf);
+						Irradiance += isfinite(ExitantRadianceSample) ? ExitantRadianceSample : 0.0;
+					}
+				}
+			}
+
+			// Perform material sampling
+			// dxr_todo: only worth doing when Bounce + 1 < BouncesMax
+			{
+				float3 Direction;
+				float3 Throughput = 1.0;
+#if 0
+				uint DummyVariable;
+				float4 RandSample = RandomSequence_GenerateSample4D(RandSequence, DummyVariable);
+				SampleMaterial(Ray.Direction, Payload, RandSample, Direction, Throughput, MaterialPdf);
+#else
+				float3 RayOrigin = Ray.Origin;
+				GenerateCosineNormalRay(RayOrigin, Payload.WorldNormal, RandSequence, Ray.Origin, Direction, Ray.TMin, Ray.TMax, MaterialPdf);
+				Throughput = Payload.DiffuseColor * MaterialPdf;
+#endif
+				// #dxr_todo: Degenerate guard?
+				if (MaterialPdf <= 0.0)
+				{
+					break;
+				}
+
+				// Update ray
+				Ray.Direction = Direction;
+				RayThroughput *= Throughput;
+				RayPdf *= MaterialPdf;
+
+				// #dxr_todo: Russian roulette?
+
+				// #dxr_todo: Firefly rejection?
+			}
+
+			Bounce++;
+		}
+	}
+
+	float AmbientOcclusion = HitCount;
+	if (SamplesPerPixel > 0)
+	{
+		Irradiance /= SamplesPerPixel;
+		AmbientOcclusion /= SamplesPerPixel;
+	}
+
+	if (HitCount > 0.0)
+	{
+		HitDistance /= HitCount;
+	}
+	else
+	{
+		HitDistance = -1.0;
+	}
+
+	RWGlobalIlluminationUAV[PixelCoord] = float4(Irradiance, AmbientOcclusion);
+	RWRayDistanceUAV[PixelCoord] = float2(HitDistance, SamplesPerPixel);
+	//RWRayDistanceUAV[PixelCoord] = float2(HitDistance, HitCount);
+	// For AO denoiser..
+	//RWRayDistanceUAV[PixelCoord] = float2(Luminance(Irradiance), HitDistance);
+}
+
+[shader("miss")]
+void RayTracingGlobalIlluminationMS(inout FPackedMaterialClosestHitPayload PackedPayload)
+{
+	PackedPayload.Payload.HitT = -1;
+}
+
+
+[shader("closesthit")]
+void RayTracingGlobalIlluminationCHS(inout FPackedMaterialClosestHitPayload PackedPayload, in FDefaultAttributes Attributes)
+{
+	FMaterialClosestHitPayload Payload = (FMaterialClosestHitPayload)0;
+	Payload.HitT = RayTCurrent();
+
+	FTriangleBaseAttributes Triangle = LoadTriangleBaseAttributes(PrimitiveIndex());
+	float3 Edge0 = Triangle.LocalPositions[2] - Triangle.LocalPositions[0];
+	float3 Edge1 = Triangle.LocalPositions[1] - Triangle.LocalPositions[0];
+	float3x3 WorldToLocal = (float3x3)WorldToObject();
+	float3x3 LocalToWorldNormal = transpose(WorldToLocal);
+	Payload.WorldNormal = normalize(mul(LocalToWorldNormal, cross(Edge0, Edge1)));
+
+	PackedPayload = PackRayTracingPayload(Payload, PackedPayload.RayCone);
+}
diff --git a/Engine/Source/Runtime/Renderer/Private/DeferredShadingRenderer.cpp b/Engine/Source/Runtime/Renderer/Private/DeferredShadingRenderer.cpp
--- a/Engine/Source/Runtime/Renderer/Private/DeferredShadingRenderer.cpp
+++ b/Engine/Source/Runtime/Renderer/Private/DeferredShadingRenderer.cpp
@@ -1750,6 +1750,13 @@ void FDeferredShadingSceneRenderer::Render(FRHICommandListImmediate& RHICmdList)
 				// #dxr_todo: Denoise SkyLight
 				CompositeRayTracingSkyLight(RHICmdList, SkyLightRT, HitDistanceRT);
 			}
+			if (ShouldRenderRayTracingGlobalIllumination())
+			{
+				for (int32 ViewIndex = 0; ViewIndex < Views.Num(); ViewIndex++)
+				{
+					RenderRayTracingGlobalIllumination(RHICmdList, Views[ViewIndex], SceneContext.ScreenSpaceAO);
+				}
+			}
 			else if (ShouldRenderRayTracingAmbientOcclusion())
 			{
 				checkSlow(RHICmdList.IsOutsideRenderPass());
@@ -1794,8 +1801,6 @@ void FDeferredShadingSceneRenderer::Render(FRHICommandListImmediate& RHICmdList)
 					}
 					GraphBuilder.Execute();
 				}
-
-				// #dxr_todo: Denoise AO
 			}
 		}
 	}
diff --git a/Engine/Source/Runtime/Renderer/Private/DeferredShadingRenderer.h b/Engine/Source/Runtime/Renderer/Private/DeferredShadingRenderer.h
--- a/Engine/Source/Runtime/Renderer/Private/DeferredShadingRenderer.h
+++ b/Engine/Source/Runtime/Renderer/Private/DeferredShadingRenderer.h
@@ -456,6 +456,7 @@ private:
 
 	void RenderRayTracingAmbientOcclusion(FRHICommandListImmediate& RHICmdList, const FLightSceneInfo* LightSceneInfo, TRefCountPtr<IPooledRenderTarget>& AmbientOcclusionRT, TRefCountPtr<IPooledRenderTarget>& AmbientOcclusionHitDistanceRT);
 	void CompositeRayTracingAmbientOcclusion(FRHICommandListImmediate& RHICmdList, TRefCountPtr<IPooledRenderTarget>& AmbientOcclusionRT);
+	void RenderRayTracingGlobalIllumination(FRHICommandListImmediate& RHICmdList, const FViewInfo& View, TRefCountPtr<IPooledRenderTarget>& AmbientOcclusionRT);
 
 	void BuildSkyLightMipTree(FRHICommandListImmediate& RHICmdList, FTextureRHIRef SkyLightTexture, FRWBuffer& SkyLightMipTreePosX, FRWBuffer& SkyLightMipTreePosY, FRWBuffer& SkyLightMipTreePosZ, FRWBuffer& SkyLightMipTreeNegX, FRWBuffer& SkyLightMipTreeNegY, FRWBuffer& SkyLightMipTreeNegZ, FIntVector& SkyLightMipTreeDimensions);
 	void BuildSkyLightMipTreePdf(
diff --git a/Engine/Source/Runtime/Renderer/Private/RayTracing/RayTracingGlobalIllumination.cpp b/Engine/Source/Runtime/Renderer/Private/RayTracing/RayTracingGlobalIllumination.cpp
new file mode 100644
--- /dev/null
+++ b/Engine/Source/Runtime/Renderer/Private/RayTracing/RayTracingGlobalIllumination.cpp
@@ -0,0 +1,425 @@
+// Copyright 1998-2019 Epic Games, Inc. All Rights Reserved.
+
+#include "DeferredShadingRenderer.h"
+
+#if RHI_RAYTRACING
+
+#include "RayTracingSkyLight.h"
+#include "ScenePrivate.h"
+#include "SceneRenderTargets.h"
+#include "RenderGraphBuilder.h"
+#include "RenderTargetPool.h"
+#include "RHIResources.h"
+#include "UniformBuffer.h"
+#include "RayGenShaderUtils.h"
+#include "PathTracingUniformBuffers.h"
+#include "SceneViewFamilyBlackboard.h"
+#include "ScreenSpaceDenoise.h"
+#include "ClearQuad.h"
+#include "PostProcess/PostProcessing.h"
+#include "PostProcess/SceneFilterRendering.h"
+
+static int32 GRayTracingGlobalIllumination = 0;
+static FAutoConsoleVariableRef CVarRayTracingGlobalIllumination(
+	TEXT("r.RayTracing.GlobalIllumination"),
+	GRayTracingGlobalIllumination,
+	TEXT("Enabled ray tracing global illumination (default = 0)")
+);
+
+static int32 GRayTracingGlobalIlluminationSamplesPerPixel = 1;
+static FAutoConsoleVariableRef CVarRayTracingGlobalIlluminationSamplesPerPixel(
+	TEXT("r.RayTracing.GlobalIllumination.SamplesPerPixel"),
+	GRayTracingGlobalIlluminationSamplesPerPixel,
+	TEXT("Samples per pixel (default = 1)")
+);
+
+static float GRayTracingGlobalIlluminationMaxRayDistance = 1.0e27;
+static FAutoConsoleVariableRef CVarRayTracingGlobalIlluminationMaxRayDistance(
+	TEXT("r.RayTracing.GlobalIllumination.MaxRayDistance"),
+	GRayTracingGlobalIlluminationMaxRayDistance,
+	TEXT("Max ray distance (default = 1.0e27)")
+);
+
+static int32 GRayTracingGlobalIlluminationMaxBounces = 1;
+static FAutoConsoleVariableRef CVarRayTracingGlobalIlluminationMaxBounces(
+	TEXT("r.RayTracing.GlobalIllumination.MaxBounces"),
+	GRayTracingGlobalIlluminationMaxBounces,
+	TEXT("Max bounces (default = 1)")
+);
+
+static int32 GRayTracingGlobalIlluminationDenoiser = 1;
+static FAutoConsoleVariableRef CVarRayTracingGlobalIlluminationDenoiser(
+	TEXT("r.RayTracing.GlobalIllumination.Denoiser"),
+	GRayTracingGlobalIlluminationDenoiser,
+	TEXT("Denoising options (default = 1)")
+);
+
+static int32 GRayTracingGlobalIlluminationEvalSkyLight = 0;
+static FAutoConsoleVariableRef CVarRayTracingGlobalIlluminationEvalSkyLight(
+	TEXT("r.RayTracing.GlobalIllumination.EvalSkyLight"),
+	GRayTracingGlobalIlluminationEvalSkyLight,
+	TEXT("Evaluate SkyLight multi-bounce contribution")
+);
+
+static const int32 GLightCountMax = 64;
+
+DECLARE_GPU_STAT_NAMED(RayTracingGlobalIllumination, TEXT("Ray Tracing Global Illumination"));
+
+void SetupLightParameters(
+	const TSparseArray<FLightSceneInfoCompact>& Lights,
+	const FViewInfo& View,
+	FPathTracingLightData* LightParameters)
+{
+	LightParameters->Count = 0;
+	for (auto Light : Lights)
+	{
+		if (LightParameters->Count >= GLightCountMax) break;
+
+		if (Light.LightSceneInfo->Proxy->HasStaticLighting() && Light.LightSceneInfo->IsPrecomputedLightingValid()) continue;
+
+		FLightShaderParameters LightShaderParameters;
+		Light.LightSceneInfo->Proxy->GetLightShaderParameters(LightShaderParameters);
+
+		ELightComponentType LightComponentType = (ELightComponentType)Light.LightSceneInfo->Proxy->GetLightType();
+		switch (LightComponentType)
+		{
+			// TODO: LightType_Spot
+		case LightType_Directional:
+		{
+			LightParameters->Type[LightParameters->Count] = 2;
+			LightParameters->Normal[LightParameters->Count] = LightShaderParameters.Direction;
+			LightParameters->Color[LightParameters->Count] = LightShaderParameters.Color;
+			break;
+		}
+		case LightType_Rect:
+		{
+			LightParameters->Type[LightParameters->Count] = 3;
+			LightParameters->Position[LightParameters->Count] = LightShaderParameters.Position;
+			LightParameters->Normal[LightParameters->Count] = -LightShaderParameters.Direction;
+			LightParameters->dPdu[LightParameters->Count] = FVector::CrossProduct(LightShaderParameters.Direction, LightShaderParameters.Tangent);
+			LightParameters->dPdv[LightParameters->Count] = LightShaderParameters.Tangent;
+			LightParameters->Color[LightParameters->Count] = LightShaderParameters.Color / 4.0;
+			LightParameters->Dimensions[LightParameters->Count] = FVector(2.0f * LightShaderParameters.SourceRadius, 2.0f * LightShaderParameters.SourceLength, 0.0f);
+			break;
+		}
+		case LightType_Point:
+		default:
+		{
+			LightParameters->Type[LightParameters->Count] = 1;
+			LightParameters->Position[LightParameters->Count] = LightShaderParameters.Position;
+			LightParameters->Color[LightParameters->Count] = LightShaderParameters.Color;
+			break;
+		}
+		};
+
+		LightParameters->Count++;
+	}
+}
+
+void SetupSkyLightParameters(
+	const FScene& Scene,
+	FSkyLightData* SkyLight
+)
+{
+	// dxr_todo: factor out these pass constants
+	SkyLight->SamplesPerPixel = 0;
+	SkyLight->SamplingStopLevel = 0;
+
+	if (Scene.SkyLight)
+	{
+		SkyLight->Color = FVector(Scene.SkyLight->LightColor);
+		SkyLight->Texture = Scene.SkyLight->ProcessedTexture->TextureRHI;
+		SkyLight->TextureSampler = Scene.SkyLight->ProcessedTexture->SamplerStateRHI;
+		SkyLight->MipDimensions = Scene.SkyLight->SkyLightMipDimensions;
+	}
+	else
+	{
+		SkyLight->Color = FVector(0.0);
+		SkyLight->Texture = GBlackTextureCube->TextureRHI;
+		SkyLight->TextureSampler = TStaticSamplerState<SF_Bilinear, AM_Clamp, AM_Clamp, AM_Clamp>::GetRHI();
+		SkyLight->MipDimensions = FIntVector(0);
+	}
+
+	// dxr_todo: Sky light importance sampling is currently disabled
+	auto BlackTextureBuffer = RHICreateShaderResourceView(GBlackTexture->TextureRHI->GetTexture2D(), 0);
+	SkyLight->MipTreePosX = BlackTextureBuffer;
+	SkyLight->MipTreeNegX = BlackTextureBuffer;
+	SkyLight->MipTreePosY = BlackTextureBuffer;
+	SkyLight->MipTreeNegY = BlackTextureBuffer;
+	SkyLight->MipTreePosZ = BlackTextureBuffer;
+	SkyLight->MipTreeNegZ = BlackTextureBuffer;
+
+	SkyLight->MipTreePdfPosX = BlackTextureBuffer;
+	SkyLight->MipTreePdfNegX = BlackTextureBuffer;
+	SkyLight->MipTreePdfPosY = BlackTextureBuffer;
+	SkyLight->MipTreePdfNegY = BlackTextureBuffer;
+	SkyLight->MipTreePdfPosZ = BlackTextureBuffer;
+	SkyLight->MipTreePdfNegZ = BlackTextureBuffer;
+	SkyLight->SolidAnglePdf = BlackTextureBuffer;
+}
+
+bool ShouldRenderRayTracingGlobalIllumination()
+{
+	return GRayTracingGlobalIllumination == 1;
+}
+
+class FGlobalIlluminationRGS : public FGlobalShader
+{
+	DECLARE_GLOBAL_SHADER(FGlobalIlluminationRGS)
+	SHADER_USE_ROOT_PARAMETER_STRUCT(FGlobalIlluminationRGS, FGlobalShader)
+
+	using FPermutationDomain = TShaderPermutationDomain<>;
+	static bool ShouldCompilePermutation(const FGlobalShaderPermutationParameters& Parameters)
+	{
+		return ShouldCompileRayTracingShadersForProject(Parameters.Platform);
+	}
+
+	BEGIN_SHADER_PARAMETER_STRUCT(FParameters, )
+		SHADER_PARAMETER(uint32, SamplesPerPixel)
+		SHADER_PARAMETER(uint32, MaxBounces)
+		SHADER_PARAMETER(float, MaxRayDistance)
+		SHADER_PARAMETER(bool, EvalSkyLight)
+
+		SHADER_PARAMETER_SRV(RaytracingAccelerationStructure, TLAS)
+		SHADER_PARAMETER_RDG_TEXTURE_UAV(RWTexture2D<float4>, RWGlobalIlluminationUAV)
+		SHADER_PARAMETER_RDG_TEXTURE_UAV(RWTexture2D<float>, RWRayDistanceUAV)
+		SHADER_PARAMETER_STRUCT_REF(FViewUniformShaderParameters, ViewUniformBuffer)
+		SHADER_PARAMETER_STRUCT_REF(FSceneTexturesUniformParameters, SceneTexturesStruct)
+		SHADER_PARAMETER_STRUCT_REF(FPathTracingLightData, LightParameters)
+		SHADER_PARAMETER_STRUCT_REF(FSkyLightData, SkyLight)
+	END_SHADER_PARAMETER_STRUCT()
+};
+
+class FRayTracingGlobalIlluminationCompositePS : public FGlobalShader
+{
+	DECLARE_GLOBAL_SHADER(FRayTracingGlobalIlluminationCompositePS)
+	SHADER_USE_ROOT_PARAMETER_STRUCT(FRayTracingGlobalIlluminationCompositePS, FGlobalShader)
+
+	static bool ShouldCompilePermutation(const FGlobalShaderPermutationParameters& Parameters)
+	{
+		return ShouldCompileRayTracingShadersForProject(Parameters.Platform);
+	}
+
+	BEGIN_SHADER_PARAMETER_STRUCT(FParameters, )
+		RENDER_TARGET_BINDING_SLOTS()
+		SHADER_PARAMETER_RDG_TEXTURE(Texture2D, GlobalIlluminationTexture)
+		SHADER_PARAMETER_SAMPLER(SamplerState, GlobalIlluminationSampler)
+		//SHADER_PARAMETER_STRUCT_REF(FViewUniformShaderParameters, ViewUniformBuffer)
+		//SHADER_PARAMETER_STRUCT_REF(FSceneTexturesUniformParameters, SceneTexturesStruct)
+	END_SHADER_PARAMETER_STRUCT()
+};
+
+class FRayTracingGlobalIlluminationCHS : public FGlobalShader
+{
+	DECLARE_GLOBAL_SHADER(FRayTracingGlobalIlluminationCHS);
+	SHADER_USE_ROOT_PARAMETER_STRUCT(FRayTracingGlobalIlluminationCHS, FGlobalShader)
+
+		static bool ShouldCompilePermutation(const FGlobalShaderPermutationParameters& Parameters)
+	{
+		return ShouldCompileRayTracingShadersForProject(Parameters.Platform);
+	}
+
+	using FParameters = FEmptyShaderParameters;
+};
+
+class FRayTracingGlobalIlluminationMS : public FGlobalShader
+{
+	DECLARE_GLOBAL_SHADER(FRayTracingGlobalIlluminationMS);
+	SHADER_USE_ROOT_PARAMETER_STRUCT(FRayTracingGlobalIlluminationMS, FGlobalShader)
+
+		static bool ShouldCompilePermutation(const FGlobalShaderPermutationParameters& Parameters)
+	{
+		return ShouldCompileRayTracingShadersForProject(Parameters.Platform);
+	}
+
+	using FParameters = FEmptyShaderParameters;
+};
+
+IMPLEMENT_GLOBAL_SHADER(FGlobalIlluminationRGS, "/Engine/Private/RayTracing/RayTracingGlobalIlluminationRGS.usf", "GlobalIlluminationRGS", SF_RayGen);
+IMPLEMENT_GLOBAL_SHADER(FRayTracingGlobalIlluminationCHS, "/Engine/Private/RayTracing/RayTracingGlobalIlluminationRGS.usf", "RayTracingGlobalIlluminationCHS", SF_RayHitGroup);
+IMPLEMENT_GLOBAL_SHADER(FRayTracingGlobalIlluminationMS, "/Engine/Private/RayTracing/RayTracingGlobalIlluminationRGS.usf", "RayTracingGlobalIlluminationMS", SF_RayMiss);
+IMPLEMENT_GLOBAL_SHADER(FRayTracingGlobalIlluminationCompositePS, "/Engine/Private/RayTracing/RayTracingGlobalIlluminationCompositePS.usf", "GlobalIlluminationCompositePS", SF_Pixel);
+
+void FDeferredShadingSceneRenderer::RenderRayTracingGlobalIllumination(
+	FRHICommandListImmediate& RHICmdList,
+	const FViewInfo& View,
+	TRefCountPtr<IPooledRenderTarget>& AmbientOcclusionRT
+)
+{
+	if (!GRayTracingGlobalIllumination) return;
+	SCOPED_GPU_STAT(RHICmdList, RayTracingGlobalIllumination);
+
+	FSceneRenderTargets& SceneContext = FSceneRenderTargets::Get(RHICmdList);
+	FRDGBuilder GraphBuilder(RHICmdList);
+
+	// Render targets
+	FRDGTextureRef GlobalIlluminationTexture;
+	FRDGTextureRef ResultTexture;
+	{
+		FRDGTextureDesc Desc = SceneContext.GetSceneColor()->GetDesc();
+		Desc.Format = PF_FloatRGBA;
+		Desc.Flags &= ~(TexCreate_FastVRAM | TexCreate_Transient);
+		GlobalIlluminationTexture = GraphBuilder.CreateTexture(Desc, TEXT("RayTracingGlobalIllumination"));
+		ResultTexture = GraphBuilder.CreateTexture(Desc, TEXT("RTGI"));
+	}
+
+	FRDGTextureRef RayDistanceTexture;
+	{
+		FRDGTextureDesc Desc = SceneContext.GetSceneColor()->GetDesc();
+		Desc.Format = PF_G16R16;
+		Desc.Flags &= ~(TexCreate_FastVRAM | TexCreate_Transient);
+		RayDistanceTexture = GraphBuilder.CreateTexture(Desc, TEXT("RayTracingGlobalIlluminationRayDistance"));
+	}
+
+	FRDGTextureRef AmbientOcclusionTexture;
+	{
+		FRDGTextureDesc Desc = SceneContext.GetSceneColor()->GetDesc();
+		Desc.Format = PF_R16F;
+		Desc.Flags &= ~(TexCreate_FastVRAM | TexCreate_Transient);
+		AmbientOcclusionTexture = GraphBuilder.CreateTexture(Desc, TEXT("RayTracingAmbientOcclusion"));
+	}
+	
+	FSceneTexturesUniformParameters SceneTextures;
+	SetupSceneTextureUniformParameters(SceneContext, FeatureLevel, ESceneTextureSetupMode::All, SceneTextures);
+	// Ray generation
+	{
+		FPathTracingLightData LightParameters;
+		SetupLightParameters(Scene->Lights, View, &LightParameters);
+
+		FSkyLightData SkyLightParameters;
+		SetupSkyLightParameters(*Scene, &SkyLightParameters);
+
+		FGlobalIlluminationRGS::FParameters* PassParameters = GraphBuilder.AllocParameters<FGlobalIlluminationRGS::FParameters>();
+		PassParameters->SamplesPerPixel = GRayTracingGlobalIlluminationSamplesPerPixel;
+		PassParameters->MaxBounces = GRayTracingGlobalIlluminationMaxBounces;
+		PassParameters->MaxRayDistance = GRayTracingGlobalIlluminationMaxRayDistance;
+		PassParameters->EvalSkyLight = GRayTracingGlobalIlluminationEvalSkyLight;
+		PassParameters->TLAS = View.PerViewRayTracingScene.RayTracingSceneRHI->GetShaderResourceView();
+		PassParameters->ViewUniformBuffer = View.ViewUniformBuffer;
+		PassParameters->SceneTexturesStruct = CreateUniformBufferImmediate(SceneTextures, EUniformBufferUsage::UniformBuffer_SingleDraw);
+		PassParameters->LightParameters = CreateUniformBufferImmediate(LightParameters, EUniformBufferUsage::UniformBuffer_SingleDraw);
+		PassParameters->SkyLight = CreateUniformBufferImmediate(SkyLightParameters, EUniformBufferUsage::UniformBuffer_SingleDraw);
+		PassParameters->RWGlobalIlluminationUAV = GraphBuilder.CreateUAV(GlobalIlluminationTexture);
+		PassParameters->RWRayDistanceUAV = GraphBuilder.CreateUAV(RayDistanceTexture);
+
+		auto RayGenerationShader = View.ShaderMap->GetShader<FGlobalIlluminationRGS>();
+		ClearUnusedGraphResources(RayGenerationShader, PassParameters);
+
+		float ResolutionFraction = 1.0;
+		int32 UpscaleFactor = int32(1.0f / ResolutionFraction);
+		FIntPoint RayTracingResolution = FIntPoint::DivideAndRoundUp(View.ViewRect.Size(), UpscaleFactor);
+		GraphBuilder.AddPass(
+			RDG_EVENT_NAME("GlobalIlluminationRayTracing %dx%d", RayTracingResolution.X, RayTracingResolution.Y),
+			PassParameters,
+			ERenderGraphPassFlags::Compute,
+			[PassParameters, this, &View, RayGenerationShader, RayTracingResolution](FRHICommandList& RHICmdList)
+		{
+			auto ClosestHitShader = View.ShaderMap->GetShader<FRayTracingGlobalIlluminationCHS>();
+			auto MissShader = View.ShaderMap->GetShader<FRayTracingGlobalIlluminationMS>();
+
+			FRHIRayTracingPipelineState* Pipeline = BindRayTracingPipeline(
+				RHICmdList, View,
+				RayGenerationShader->GetRayTracingShader(),
+				MissShader->GetRayTracingShader(),
+				ClosestHitShader->GetRayTracingShader()); // #dxr_todo: this should be done once at load-time and cached
+
+			FRayTracingShaderBindingsWriter GlobalResources;
+			SetShaderParameters(GlobalResources, RayGenerationShader, *PassParameters);
+
+			FRayTracingSceneRHIParamRef RayTracingSceneRHI = View.PerViewRayTracingScene.RayTracingSceneRHI;
+			uint32 RayGenShaderIndex = 0;
+			RHICmdList.RayTraceDispatch(Pipeline, RayGenShaderIndex, RayTracingSceneRHI, GlobalResources, RayTracingResolution.X, RayTracingResolution.Y);
+		});
+	}
+
+	// Denoising
+	if (GRayTracingGlobalIlluminationDenoiser != 0)
+	{
+		FSceneViewFamilyBlackboard SceneBlackboard;
+		SetupSceneViewFamilyBlackboard(GraphBuilder, &SceneBlackboard);
+
+		const IScreenSpaceDenoiser* DefaultDenoiser = IScreenSpaceDenoiser::GetDefaultDenoiser();
+		const IScreenSpaceDenoiser* DenoiserToUse = GRayTracingGlobalIlluminationDenoiser == 1 ? DefaultDenoiser : GScreenSpaceDenoiser;
+
+		IScreenSpaceDenoiser::FAmbientOcclusionRayTracingConfig RayTracingConfig;
+		RayTracingConfig.RayCountPerPixel = GRayTracingGlobalIlluminationSamplesPerPixel;
+
+		IScreenSpaceDenoiser::FGlobalIlluminationInputs DenoiserInputs;
+		DenoiserInputs.Color = GlobalIlluminationTexture;
+		DenoiserInputs.RayHitDistance = RayDistanceTexture;
+
+		{
+			RDG_EVENT_SCOPE(GraphBuilder, "%s%s(AmbientOcclusion) %dx%d",
+				DenoiserToUse != DefaultDenoiser ? TEXT("ThirdParty ") : TEXT(""),
+				DenoiserToUse->GetDebugName(),
+				View.ViewRect.Width(), View.ViewRect.Height());
+
+			IScreenSpaceDenoiser::FGlobalIlluminationOutputs DenoiserOutputs = DenoiserToUse->DenoiseGlobalIllumination(
+				GraphBuilder,
+				View,
+				SceneBlackboard,
+				DenoiserInputs,
+				RayTracingConfig);
+
+			ResultTexture = DenoiserOutputs.Color;
+		}
+	}
+	else
+	{
+		ResultTexture = GlobalIlluminationTexture;
+	}
+
+	// Compositing
+	{
+		FRayTracingGlobalIlluminationCompositePS::FParameters *PassParameters = GraphBuilder.AllocParameters<FRayTracingGlobalIlluminationCompositePS::FParameters>();
+		PassParameters->GlobalIlluminationTexture = ResultTexture;
+		PassParameters->GlobalIlluminationSampler = TStaticSamplerState<SF_Point, AM_Clamp, AM_Clamp, AM_Clamp>::GetRHI();
+		PassParameters->RenderTargets[0] = FRenderTargetBinding(GraphBuilder.RegisterExternalTexture(SceneContext.GetSceneColor()), ERenderTargetLoadAction::ENoAction, ERenderTargetStoreAction::ENoAction);
+		PassParameters->RenderTargets[1] = FRenderTargetBinding(GraphBuilder.RegisterExternalTexture(AmbientOcclusionRT), ERenderTargetLoadAction::ENoAction, ERenderTargetStoreAction::ENoAction);
+
+		GraphBuilder.AddPass(
+			RDG_EVENT_NAME("GlobalIlluminationComposite"),
+			PassParameters,
+			ERenderGraphPassFlags::None,
+			[this, &SceneContext, &View, &SceneTextures, PassParameters](FRHICommandListImmediate& RHICmdList)
+			{
+				TShaderMapRef<FPostProcessVS> VertexShader(View.ShaderMap);
+				TShaderMapRef<FRayTracingGlobalIlluminationCompositePS> PixelShader(View.ShaderMap);
+				FGraphicsPipelineStateInitializer GraphicsPSOInit;
+				RHICmdList.ApplyCachedRenderTargets(GraphicsPSOInit);
+
+				// Additive blending
+				GraphicsPSOInit.BlendState = TStaticBlendState<CW_RGBA, BO_Add, BF_One, BF_One, BO_Add, BF_One, BF_Zero>::GetRHI();
+				//GraphicsPSOInit.BlendState = TStaticBlendState<CW_RGBA, BO_Add, BF_One, BF_One, BO_Add, BF_One, BF_One>::GetRHI();
+				GraphicsPSOInit.RasterizerState = TStaticRasterizerState<FM_Solid, CM_None>::GetRHI();
+				GraphicsPSOInit.DepthStencilState = TStaticDepthStencilState<false, CF_Always>::GetRHI();
+
+				GraphicsPSOInit.BoundShaderState.VertexDeclarationRHI = GFilterVertexDeclaration.VertexDeclarationRHI;
+				GraphicsPSOInit.BoundShaderState.VertexShaderRHI = GETSAFERHISHADER_VERTEX(*VertexShader);
+				GraphicsPSOInit.BoundShaderState.PixelShaderRHI = GETSAFERHISHADER_PIXEL(*PixelShader);
+				GraphicsPSOInit.PrimitiveType = PT_TriangleList;
+				SetGraphicsPipelineState(RHICmdList, GraphicsPSOInit);
+
+				SetShaderParameters(RHICmdList, *PixelShader, PixelShader->GetPixelShader(), *PassParameters);
+
+				RHICmdList.SetViewport(View.ViewRect.Min.X, View.ViewRect.Min.Y, 0.0f, View.ViewRect.Max.X, View.ViewRect.Max.Y, 1.0f);
+
+				DrawRectangle(
+					RHICmdList,
+					0, 0,
+					View.ViewRect.Width(), View.ViewRect.Height(),
+					View.ViewRect.Min.X, View.ViewRect.Min.Y,
+					View.ViewRect.Width(), View.ViewRect.Height(),
+					FIntPoint(View.ViewRect.Width(), View.ViewRect.Height()),
+					SceneContext.GetBufferSizeXY(),
+					*VertexShader
+				);
+			}
+		);
+	}
+
+	GraphBuilder.Execute();
+	SceneContext.bScreenSpaceAOIsValid = true;
+}
+
+#endif // RHI_RAYTRACING
diff --git a/Engine/Source/Runtime/Renderer/Private/RayTracing/RaytracingOptions.h b/Engine/Source/Runtime/Renderer/Private/RayTracing/RaytracingOptions.h
--- a/Engine/Source/Runtime/Renderer/Private/RayTracing/RaytracingOptions.h
+++ b/Engine/Source/Runtime/Renderer/Private/RayTracing/RaytracingOptions.h
@@ -11,6 +11,7 @@ extern bool IsRayTracingSkyLightSelected();
 extern bool IsRayTracingRectLightSelected();
 
 extern bool ShouldRenderRayTracingAmbientOcclusion();
+extern bool ShouldRenderRayTracingGlobalIllumination();
 extern bool ShouldRenderRayTracingStaticOrStationaryRectLight(const FLightSceneInfo& LightSceneInfo);
 extern bool ShouldRenderRayTracingDynamicRectLight(const FLightSceneInfo& LightSceneInfo);
 
