Preferring RayTracing.Occlusion pass to RayTracing.RectLight occlusion pass.

TODO: unify occlusion methods such that r.RayTracing.Occlusion can incorporate textured area light sampling when importance sampling visibility.

diff --git a/Engine/Source/Runtime/Renderer/Private/RayTracing/RayTracingShadowsForRectLight.cpp b/Engine/Source/Runtime/Renderer/Private/RayTracing/RayTracingShadowsForRectLight.cpp
--- a/Engine/Source/Runtime/Renderer/Private/RayTracing/RayTracingShadowsForRectLight.cpp
+++ b/Engine/Source/Runtime/Renderer/Private/RayTracing/RayTracingShadowsForRectLight.cpp
@@ -18,7 +18,7 @@
 
 #include "RHI/Public/PipelineStateCache.h"
 
-static int32 GRayTracingRectLight = 1;
+static int32 GRayTracingRectLight = 0;
 static FAutoConsoleVariableRef CVarRayTracingRectLight(
 	TEXT("r.RayTracing.RectLight"),
 	GRayTracingRectLight,
