Added CVar to adjust reflections Max Bounces

diff --git a/Engine/Shaders/Private/RayTracing/RayTracingReflections.usf b/Engine/Shaders/Private/RayTracing/RayTracingReflections.usf
--- a/Engine/Shaders/Private/RayTracing/RayTracingReflections.usf
+++ b/Engine/Shaders/Private/RayTracing/RayTracingReflections.usf
@@ -24,6 +24,7 @@ float2 SobolIndexToUniformUnitSquare(uint2 SobolRand)
 #include "RayTracingHitGroupCommon.ush"
 
 uint SamplesPerPixel;
+uint MaxBounces;
 uint HeightFog;
 uint ShouldDoDirectLighting;
 uint ShouldDoReflectedShadows;
@@ -150,7 +151,6 @@ void RayTracingReflectionsRGS()
 
 			uint DummyVariable;
 				
-			const uint MaxBounces = 1;
 			// Add one more 'fake bounce' and we will evaluate cubemap there
 			for (uint BounceIndex = 0; BounceIndex < MaxBounces + 1; ++BounceIndex)
 			{
diff --git a/Engine/Source/Runtime/Renderer/Private/RayTracing/RayTracingReflections.cpp b/Engine/Source/Runtime/Renderer/Private/RayTracing/RayTracingReflections.cpp
--- a/Engine/Source/Runtime/Renderer/Private/RayTracing/RayTracingReflections.cpp
+++ b/Engine/Source/Runtime/Renderer/Private/RayTracing/RayTracingReflections.cpp
@@ -16,6 +16,13 @@
 #include "LightRendering.h"
 #include "SystemTextures.h"
 
+static int32 GRayTracingReflectionsMaxBounces = 1;
+static FAutoConsoleVariableRef CVarRayTracingReflectionsMaxBounces(
+	TEXT("r.RayTracing.Reflections.MaxBounces"),
+	GRayTracingReflectionsMaxBounces,
+	TEXT("Sets the maximum number of ray tracing reflection bounces (default = 1)")
+);
+
 static int32 GRayTracingReflectionsEmissiveAndIndirectLighting = 1;
 static FAutoConsoleVariableRef CVarRayTracingReflectionsEmissiveAndIndirectLighting(
 	TEXT("r.RayTracing.Reflections.EmissiveAndIndirectLighting"),
@@ -129,6 +136,7 @@ class FRayTracingReflectionsRG : public FGlobalShader
 
 	BEGIN_SHADER_PARAMETER_STRUCT(FParameters, )
 		SHADER_PARAMETER(int32, SamplesPerPixel)
+		SHADER_PARAMETER(int32, MaxBounces)
 		SHADER_PARAMETER(int32, HeightFog)
 		SHADER_PARAMETER(int32, ShouldDoDirectLighting)
 		SHADER_PARAMETER(int32, ShouldDoReflectedShadows)
@@ -229,6 +237,7 @@ void FDeferredShadingSceneRenderer::RayTraceReflections(
 	FRayTracingReflectionsRG::FParameters* PassParameters = GraphBuilder.AllocParameters<FRayTracingReflectionsRG::FParameters>();
 
 	PassParameters->SamplesPerPixel = SamplePerPixel;
+	PassParameters->MaxBounces = GRayTracingReflectionsMaxBounces;
 	PassParameters->HeightFog = HeightFog;
 	PassParameters->ShouldDoDirectLighting = GRayTracingReflectionsDirectLighting;
 	PassParameters->ShouldDoReflectedShadows = GRayTracingReflectionsShadows;
