Fixed FMeshDrawCommandStateCache always resetting resources at index 0.  Found by Rolando.

diff --git a/Engine/Source/Runtime/Renderer/Private/MeshPassProcessor.cpp b/Engine/Source/Runtime/Renderer/Private/MeshPassProcessor.cpp
--- a/Engine/Source/Runtime/Renderer/Private/MeshPassProcessor.cpp
+++ b/Engine/Source/Runtime/Renderer/Private/MeshPassProcessor.cpp
@@ -35,13 +35,13 @@ enum { MAX_SAMPLERS_PER_SHADER_STAGE = 32 };
 class FShaderBindingState
 {
 public:
-	int32 MaxSRVUsed = 0;
+	int32 MaxSRVUsed = -1;
 	FShaderResourceViewRHIParamRef SRVs[MAX_SRVs_PER_SHADER_STAGE] = {};
-	int32 MaxUniformBufferUsed = 0;
+	int32 MaxUniformBufferUsed = -1;
 	FUniformBufferRHIParamRef UniformBuffers[MAX_UNIFORM_BUFFERS_PER_SHADER_STAGE] = {};
-	int32 MaxTextureUsed = 0;
+	int32 MaxTextureUsed = -1;
 	FTextureRHIParamRef Textures[MAX_SRVs_PER_SHADER_STAGE] = {};
-	int32 MaxSamplerUsed = 0;
+	int32 MaxSamplerUsed = -1;
 	FSamplerStateRHIParamRef Samplers[MAX_SAMPLERS_PER_SHADER_STAGE] = {};
 };
 
@@ -351,28 +351,28 @@ public:
 				ShaderBinding.SRVs[SlotIndex] = nullptr;
 			}
 
-			ShaderBinding.MaxSRVUsed = 0;
+			ShaderBinding.MaxSRVUsed = -1;
 
 			for (int32 SlotIndex = 0; SlotIndex <= ShaderBinding.MaxUniformBufferUsed; SlotIndex++)
 			{
 				ShaderBinding.UniformBuffers[SlotIndex] = nullptr;
 			}
 
-			ShaderBinding.MaxUniformBufferUsed = 0;
+			ShaderBinding.MaxUniformBufferUsed = -1;
 			
 			for (int32 SlotIndex = 0; SlotIndex <= ShaderBinding.MaxTextureUsed; SlotIndex++)
 			{
 				ShaderBinding.Textures[SlotIndex] = nullptr;
 			}
 
-			ShaderBinding.MaxTextureUsed = 0;
+			ShaderBinding.MaxTextureUsed = -1;
 
 			for (int32 SlotIndex = 0; SlotIndex <= ShaderBinding.MaxSamplerUsed; SlotIndex++)
 			{
 				ShaderBinding.Samplers[SlotIndex] = nullptr;
 			}
 
-			ShaderBinding.MaxSamplerUsed = 0;
+			ShaderBinding.MaxSamplerUsed = -1;
 		}
 	}
 };
