On iOS Metal is a bit more specific about what LoadAction::NoAction really means - it will reallocate the backing store for the target to uninitialised memory, abandoning whatever was already there. This fixes corruption in the Forward renderer when MSAA is enabled.

diff --git a/Engine/Source/Runtime/Renderer/Private/DeferredShadingRenderer.cpp b/Engine/Source/Runtime/Renderer/Private/DeferredShadingRenderer.cpp
--- a/Engine/Source/Runtime/Renderer/Private/DeferredShadingRenderer.cpp
+++ b/Engine/Source/Runtime/Renderer/Private/DeferredShadingRenderer.cpp
@@ -1521,7 +1521,7 @@ void FDeferredShadingSceneRenderer::Render(FRHICommandListImmediate& RHICmdList)
 	{
 		// Make sure we have began the renderpass
 		ERenderTargetLoadAction DepthLoadAction = bDepthWasCleared ? ERenderTargetLoadAction::ELoad : ERenderTargetLoadAction::EClear;
-		SceneContext.BeginRenderingGBuffer(RHICmdList, ERenderTargetLoadAction::ENoAction, DepthLoadAction, BasePassDepthStencilAccess, ViewFamily.EngineShowFlags.ShaderComplexity);
+		SceneContext.BeginRenderingGBuffer(RHICmdList, (!IsMetalPlatform(ShaderPlatform) ? ERenderTargetLoadAction::ENoAction : ERenderTargetLoadAction::ELoad), DepthLoadAction, BasePassDepthStencilAccess, ViewFamily.EngineShowFlags.ShaderComplexity);
 	}
 	// Wait for Async SSAO before rendering base pass with forward rendering
 	if (IsForwardShadingEnabled(ShaderPlatform))
@@ -2022,7 +2022,7 @@ void FDeferredShadingSceneRenderer::Render(FRHICommandListImmediate& RHICmdList)
 		SetupSceneTextureUniformParameters(SceneContext, FeatureLevel, ESceneTextureSetupMode::SceneDepth | ESceneTextureSetupMode::GBuffers, SceneTextureParameters);
 		TUniformBufferRef<FSceneTexturesUniformParameters> SceneTextureUniformBuffer = TUniformBufferRef<FSceneTexturesUniformParameters>::CreateUniformBufferImmediate(SceneTextureParameters, UniformBuffer_SingleFrame);
 
-		SceneContext.BeginRenderingSceneColor(RHICmdList);
+        SceneContext.BeginRenderingSceneColor(RHICmdList, (!IsMetalPlatform(ShaderPlatform) ? ESimpleRenderTargetMode::EUninitializedColorExistingDepth : ESimpleRenderTargetMode::EExistingColorAndDepth));
 		for (int32 ViewIndex = 0; ViewIndex < Views.Num(); ++ViewIndex)
 		{
 			const FViewInfo& View = Views[ViewIndex];
