Adds a warning message about the experimental nature of SSGI.

diff --git a/Engine/Source/Runtime/Renderer/Private/IndirectLightRendering.cpp b/Engine/Source/Runtime/Renderer/Private/IndirectLightRendering.cpp
--- a/Engine/Source/Runtime/Renderer/Private/IndirectLightRendering.cpp
+++ b/Engine/Source/Runtime/Renderer/Private/IndirectLightRendering.cpp
@@ -9,6 +9,7 @@
 #include "ScreenSpaceRayTracing.h"
 #include "DeferredShadingRenderer.h"
 #include "PostProcessing.h" // for FPostProcessVS
+#include "RendererModule.h"
 
 
 static TAutoConsoleVariable<int32> CVarDiffuseIndirectDenoiser(
@@ -133,6 +134,13 @@ void FDeferredShadingSceneRenderer::RenderDiffuseIndirectAndAmbientOcclusion(FRH
 		IScreenSpaceDenoiser::FDiffuseIndirectInputs DenoiserInputs;
 		if (bApplySSGI)
 		{
+			static bool bWarnExperimental = false;
+			if (!bWarnExperimental)
+			{
+				UE_LOG(LogRenderer, Warning, TEXT("SSGI is experimental."));
+				bWarnExperimental = true;
+			}
+
 			RenderScreenSpaceDiffuseIndirect(GraphBuilder, SceneTextures, SceneColor, View, /* out */ &DenoiserInputs);
 			
 			// TODO: Denoise.
