SSAO can now be disabled when in renderer is in forward and ssao is not executed in async compute (special case that was not taken into account)

diff --git a/Engine/Source/Runtime/Renderer/Private/CompositionLighting/PostProcessAmbientOcclusion.cpp b/Engine/Source/Runtime/Renderer/Private/CompositionLighting/PostProcessAmbientOcclusion.cpp
--- a/Engine/Source/Runtime/Renderer/Private/CompositionLighting/PostProcessAmbientOcclusion.cpp
+++ b/Engine/Source/Runtime/Renderer/Private/CompositionLighting/PostProcessAmbientOcclusion.cpp
@@ -178,7 +178,7 @@ uint32 FSSAOHelper::ComputeAmbientOcclusionPassCount(const FViewInfo& View)
 		if (IsAmbientOcclusionCompute(View) || IsForwardShadingEnabled(View.GetShaderPlatform()))
 		{	
 			// Compute and forward only support one pass currently.
-			return 1;
+			return FMath::Min<int32>(GetNumAmbientOcclusionLevels(), 1);
 		}
 
 		// usually in the range 0..100
