Fixed crash when renderer is in forward or ssao in compute and r.ssaolevels is lower than zero

diff --git a/Engine/Source/Runtime/Renderer/Private/CompositionLighting/PostProcessAmbientOcclusion.cpp b/Engine/Source/Runtime/Renderer/Private/CompositionLighting/PostProcessAmbientOcclusion.cpp
--- a/Engine/Source/Runtime/Renderer/Private/CompositionLighting/PostProcessAmbientOcclusion.cpp
+++ b/Engine/Source/Runtime/Renderer/Private/CompositionLighting/PostProcessAmbientOcclusion.cpp
@@ -175,13 +175,19 @@ uint32 FSSAOHelper::ComputeAmbientOcclusionPassCount(const FViewInfo& View)
 
 	if (bEnabled)
 	{
+		int32 CVarLevel = GetNumAmbientOcclusionLevels();
+
 		if (IsAmbientOcclusionCompute(View) || IsForwardShadingEnabled(View.GetShaderPlatform()))
 		{	
+			if (CVarLevel<0)
+			{
+				CVarLevel = 1;
+			}
 			// Compute and forward only support one pass currently.
-			return FMath::Min<int32>(GetNumAmbientOcclusionLevels(), 1);
+			return FMath::Min<int32>(CVarLevel, 1);
 		}
 
-		// usually in the range 0..100
+		// usually in the range 0..100 
 		float QualityPercent = GetAmbientOcclusionQualityRT(View);
 
 		// don't expose 0 as the lowest quality should still render
@@ -189,8 +195,6 @@ uint32 FSSAOHelper::ComputeAmbientOcclusionPassCount(const FViewInfo& View)
 			(QualityPercent > 70.0f) +
 			(QualityPercent > 35.0f);
 
-		int32 CVarLevel = GetNumAmbientOcclusionLevels();
-
 		if (CVarLevel >= 0)
 		{
 			// cvar can override (for scalability or to profile/test)
