Sobol random texture fix

diff --git a/Engine/Source/Runtime/Renderer/Private/SystemTextures.cpp b/Engine/Source/Runtime/Renderer/Private/SystemTextures.cpp
--- a/Engine/Source/Runtime/Renderer/Private/SystemTextures.cpp
+++ b/Engine/Source/Runtime/Renderer/Private/SystemTextures.cpp
@@ -196,14 +196,14 @@ void FSystemTextures::InitializeFeatureLevelDependentTextures(FRHICommandListImm
 				// 16x16 block starting at 16,0 = Sobol X,Y from 2nd 4 bits of cell X,Y
 				for (int x = 0; x < 16; ++x, ++Dest)
 				{
-					Result  = (x & 0x010) ? 0x4414 : 0;
-					Result ^= (x & 0x020) ? 0x8828 : 0;
-					Result ^= (x & 0x040) ? 0xe69e : 0;
-					Result ^= (x & 0x080) ? 0xae76 : 0;
-					Result ^= (y & 0x010) ? 0xa28a : 0;
-					Result ^= (y & 0x020) ? 0x265e : 0;
-					Result ^= (y & 0x040) ? 0xe69e : 0;
-					Result ^= (y & 0x080) ? 0xae76 : 0;
+					Result  = (x & 0x001) ? 0x4414 : 0;
+					Result ^= (x & 0x002) ? 0x8828 : 0;
+					Result ^= (x & 0x004) ? 0xe69e : 0;
+					Result ^= (x & 0x008) ? 0xae76 : 0;
+					Result ^= (y & 0x001) ? 0xa28a : 0;
+					Result ^= (y & 0x002) ? 0x265e : 0;
+					Result ^= (y & 0x004) ? 0xe69e : 0;
+					Result ^= (y & 0x008) ? 0xae76 : 0;
 					*Dest = Result;
 				}
 			}
