Hand-integration of ERenderGraphPassFlags enum rename.

diff --git a/Engine/Source/Runtime/RenderCore/Private/GenerateMips.cpp b/Engine/Source/Runtime/RenderCore/Private/GenerateMips.cpp
--- a/Engine/Source/Runtime/RenderCore/Private/GenerateMips.cpp
+++ b/Engine/Source/Runtime/RenderCore/Private/GenerateMips.cpp
@@ -131,7 +131,7 @@ void FGenerateMips::Compute(FRHICommandListImmediate& RHIImmCmdList, FRHITexture
 		GraphBuilder.AddPass(
 			RDG_EVENT_NAME("Generate2DTextureMips DestMipLevel=%d", MipLevel),
 			PassParameters,
-			ERenderGraphPassFlags::Compute | ERDGPassFlags::GenerateMips,
+			ERDGPassFlags::Compute | ERDGPassFlags::GenerateMips,
 			[PassParameters, ComputeShader, GenMipsGroupCount](FRHICommandList& RHICmdList)
 		{
 			FComputeShaderUtils::Dispatch(RHICmdList, *ComputeShader, *PassParameters, GenMipsGroupCount);
diff --git a/Engine/Source/Runtime/RenderCore/Private/VisualizeTexture.cpp b/Engine/Source/Runtime/RenderCore/Private/VisualizeTexture.cpp
--- a/Engine/Source/Runtime/RenderCore/Private/VisualizeTexture.cpp
+++ b/Engine/Source/Runtime/RenderCore/Private/VisualizeTexture.cpp
@@ -260,7 +260,7 @@ void FVisualizeTexture::CreateContentCapturePass(FRDGBuilder& GraphBuilder, cons
 	GraphBuilder.AddPass(
 		RDG_EVENT_NAME("VisualizeTextureCapture(%s)", SrcTexture->Name),
 		PassParameters,
-		ERenderGraphPassFlags::None,
+		ERDGPassFlags::Raster,
 		[this, PassParameters, ShaderMap, PixelShader, RTExtent](FRHICommandList& RHICmdList)
 	{
 		FVisualizeTexturePS::FParameters ShaderParameter = *PassParameters;
diff --git a/Engine/Source/Runtime/RenderCore/Public/PixelShaderUtils.h b/Engine/Source/Runtime/RenderCore/Public/PixelShaderUtils.h
--- a/Engine/Source/Runtime/RenderCore/Public/PixelShaderUtils.h
+++ b/Engine/Source/Runtime/RenderCore/Public/PixelShaderUtils.h
@@ -67,7 +67,7 @@ struct RENDERCORE_API FPixelShaderUtils
 		GraphBuilder.AddPass(
 			Forward<FRDGEventName>(PassName),
 			Parameters,
-			ERenderGraphPassFlags::None,
+			ERDGPassFlags::Raster,
 			[Parameters, GlobalShaderMap, PixelShader, Viewport](FRHICommandList& RHICmdList)
 		{
 			FPixelShaderUtils::DrawFullscreenPixelShader(RHICmdList, GlobalShaderMap, PixelShader, *Parameters, Viewport);
diff --git a/Engine/Source/Runtime/RenderCore/Public/RayGenShaderUtils.h b/Engine/Source/Runtime/RenderCore/Public/RayGenShaderUtils.h
--- a/Engine/Source/Runtime/RenderCore/Public/RayGenShaderUtils.h
+++ b/Engine/Source/Runtime/RenderCore/Public/RayGenShaderUtils.h
@@ -28,7 +28,7 @@ struct RENDERCORE_API FRayGenShaderUtils
 		GraphBuilder.AddPass(
 			Forward<FRDGEventName>(PassName),
 			Parameters,
-			ERenderGraphPassFlags::Compute,
+			ERDGPassFlags::Compute,
 			[RayGenerationShader, Parameters, Resolution](FRHICommandList& RHICmdList)
 		{
 			FRayTracingShaderBindingsWriter GlobalResources;
diff --git a/Engine/Source/Runtime/RenderCore/Public/RenderGraphPass.h b/Engine/Source/Runtime/RenderCore/Public/RenderGraphPass.h
--- a/Engine/Source/Runtime/RenderCore/Public/RenderGraphPass.h
+++ b/Engine/Source/Runtime/RenderCore/Public/RenderGraphPass.h
@@ -118,11 +118,6 @@ public:
 		: FRDGPassParameterStruct(Parameters, &FParameterStruct::FTypeInfo::GetStructMetadata()->GetLayout())
 	{}
 
-	bool IsCompute() const
-	{
-		return Layout->NumRenderTargets() == 0;
-	}
-
 	uint32 GetParameterCount() const
 	{
 		return Layout->Resources.Num();
@@ -160,21 +155,18 @@ private:
 /** Flags to annotate passes. */
 enum class ERDGPassFlags
 {
-	None = 0,
+	/** Pass uses raster pipeline. */
+	Raster = 0x1,
 
 	/** Pass uses compute only */
-	Compute = 1 << 0,
+	Compute = 0x2,
 
 	//#todo-rco: Remove this when we can do split/per mip layout transitions.
 	/** Hint to some RHIs this pass will be generating mips to optimize transitions. */
-	GenerateMips = 1 << 1,
+	GenerateMips = 0x8,
 };
-
 ENUM_CLASS_FLAGS(ERDGPassFlags)
 
-// TODO(RDG): Bulk rename across codebase.
-using ERenderGraphPassFlags = ERDGPassFlags;
-
 /** Base class of a render graph pass. */
 class RENDERCORE_API FRDGPass
 {
@@ -198,9 +190,9 @@ public:
 		return PassFlags;
 	}
 
-	bool IsGenerateMips() const
+	bool IsRaster() const
 	{
-		return (PassFlags & ERDGPassFlags::GenerateMips) == ERDGPassFlags::GenerateMips;
+		return (PassFlags & ERDGPassFlags::Raster) == ERDGPassFlags::Raster;
 	}
 
 	bool IsCompute() const
@@ -208,6 +200,11 @@ public:
 		return (PassFlags & ERDGPassFlags::Compute) == ERDGPassFlags::Compute;
 	}
 
+	bool IsGenerateMips() const
+	{
+		return (PassFlags & ERDGPassFlags::GenerateMips) == ERDGPassFlags::GenerateMips;
+	}
+
 	FRDGPassParameterStruct GetParameters() const
 	{
 		return ParameterStruct;
diff --git a/Engine/Source/Runtime/RenderCore/Public/RenderGraphUtils.h b/Engine/Source/Runtime/RenderCore/Public/RenderGraphUtils.h
--- a/Engine/Source/Runtime/RenderCore/Public/RenderGraphUtils.h
+++ b/Engine/Source/Runtime/RenderCore/Public/RenderGraphUtils.h
@@ -147,7 +147,7 @@ struct RENDERCORE_API FComputeShaderUtils
 		GraphBuilder.AddPass(
 			Forward<FRDGEventName>(PassName),
 			Parameters,
-			ERenderGraphPassFlags::Compute,
+			ERDGPassFlags::Compute,
 			[Parameters, ComputeShader, GroupCount](FRHICommandList& RHICmdList)
 		{
 			FComputeShaderUtils::Dispatch(RHICmdList, ComputeShader, *Parameters, GroupCount);
@@ -171,7 +171,7 @@ struct RENDERCORE_API FComputeShaderUtils
 		GraphBuilder.AddPass(
 			Forward<FRDGEventName>(PassName),
 			Parameters,
-			ERenderGraphPassFlags::Compute,
+			ERDGPassFlags::Compute,
 			[Parameters, ComputeShader, IndirectArgsBuffer, IndirectArgOffset](FRHICommandList& RHICmdList)
 		{			
 			// Marks the indirect draw parameter as used by the pass manually, given it can't be bound directly by any of the shader,
diff --git a/Engine/Source/Runtime/Renderer/Private/ClusteredDeferredShadingPass.cpp b/Engine/Source/Runtime/Renderer/Private/ClusteredDeferredShadingPass.cpp
--- a/Engine/Source/Runtime/Renderer/Private/ClusteredDeferredShadingPass.cpp
+++ b/Engine/Source/Runtime/Renderer/Private/ClusteredDeferredShadingPass.cpp
@@ -139,7 +139,7 @@ void FDeferredShadingSceneRenderer::AddClusteredDeferredShadingPass(FRHICommandL
 			GraphBuilder.AddPass(
 				RDG_EVENT_NAME("ClusteredDeferredShading, #Lights: %d", NumLightsToRender),
 				PassParameters,
-				ERenderGraphPassFlags::None,
+				ERDGPassFlags::Raster,
 				[PassParameters, &View, &SceneContext](FRHICommandListImmediate& InRHICmdList)
 			{
 				TShaderMapRef<FPostProcessVS> VertexShader(View.ShaderMap);
diff --git a/Engine/Source/Runtime/Renderer/Private/IndirectLightRendering.cpp b/Engine/Source/Runtime/Renderer/Private/IndirectLightRendering.cpp
--- a/Engine/Source/Runtime/Renderer/Private/IndirectLightRendering.cpp
+++ b/Engine/Source/Runtime/Renderer/Private/IndirectLightRendering.cpp
@@ -226,7 +226,7 @@ void FDeferredShadingSceneRenderer::RenderDiffuseIndirectAndAmbientOcclusion(FRH
 					PermutationVector.Get<FDiffuseIndirectCompositePS::FApplyDiffuseIndirectDim>() ? TEXT("Yes") : TEXT("No"),
 					View.ViewRect.Width(), View.ViewRect.Height()),
 				PassParameters,
-				ERenderGraphPassFlags::None,
+				ERDGPassFlags::Raster,
 				[PassParameters, &View, PixelShader, PermutationVector](FRHICommandList& RHICmdList)
 			{
 				RHICmdList.SetViewport(View.ViewRect.Min.X, View.ViewRect.Min.Y, 0.0f, View.ViewRect.Max.X, View.ViewRect.Max.Y, 0.0);
@@ -277,7 +277,7 @@ void FDeferredShadingSceneRenderer::RenderDiffuseIndirectAndAmbientOcclusion(FRH
 			GraphBuilder.AddPass(
 				RDG_EVENT_NAME("AmbientCubemapComposite %dx%d", View.ViewRect.Width(), View.ViewRect.Height()),
 				PassParameters,
-				ERDGPassFlags::None,
+				ERDGPassFlags::Raster,
 				[PassParameters, &View, PixelShader](FRHICommandList& RHICmdList)
 			{
 				TShaderMapRef<FPostProcessVS> VertexShader(View.ShaderMap);
diff --git a/Engine/Source/Runtime/Renderer/Private/LightGridInjection.cpp b/Engine/Source/Runtime/Renderer/Private/LightGridInjection.cpp
--- a/Engine/Source/Runtime/Renderer/Private/LightGridInjection.cpp
+++ b/Engine/Source/Runtime/Renderer/Private/LightGridInjection.cpp
@@ -92,7 +92,7 @@ void AddPass_ClearUAV(
 	GraphBuilder.AddPass(
 		Forward<FRDGEventName>(PassName),
 		Parameters,
-		ERenderGraphPassFlags::Compute,
+		ERDGPassFlags::Compute,
 		[&Parameters, BufferUAV, Value](FRHICommandList& RHICmdList)
 		{
 			BufferUAV->MarkResourceAsUsed();
diff --git a/Engine/Source/Runtime/Renderer/Private/PlanarReflectionRendering.cpp b/Engine/Source/Runtime/Renderer/Private/PlanarReflectionRendering.cpp
--- a/Engine/Source/Runtime/Renderer/Private/PlanarReflectionRendering.cpp
+++ b/Engine/Source/Runtime/Renderer/Private/PlanarReflectionRendering.cpp
@@ -734,7 +734,7 @@ void FDeferredShadingSceneRenderer::RenderDeferredPlanarReflections(FRDGBuilder&
 	GraphBuilder.AddPass(
 		RDG_EVENT_NAME("CompositePlanarReflections"),
 		PassParameters,
-		ERDGPassFlags::None,
+		ERDGPassFlags::Raster,
 		[PassParameters, &View, this](FRHICommandList& RHICmdList)
 	{
 		RHICmdList.SetViewport(View.ViewRect.Min.X, View.ViewRect.Min.Y, 0.0f, View.ViewRect.Max.X, View.ViewRect.Max.Y, 1.0f);
diff --git a/Engine/Source/Runtime/Renderer/Private/PostProcess/DiaphragmDOF.cpp b/Engine/Source/Runtime/Renderer/Private/PostProcess/DiaphragmDOF.cpp
--- a/Engine/Source/Runtime/Renderer/Private/PostProcess/DiaphragmDOF.cpp
+++ b/Engine/Source/Runtime/Renderer/Private/PostProcess/DiaphragmDOF.cpp
@@ -2273,7 +2273,7 @@ FRDGTextureRef DiaphragmDOF::AddPasses(
 					PermutationVector.Get<FDDOFScatterOcclusionDim>() ? TEXT("Yes") : TEXT("No"),
 					GatheringViewSize.X, GatheringViewSize.Y),
 				PassParameters,
-				ERenderGraphPassFlags::None,
+				ERDGPassFlags::Raster,
 				[PassParameters, VertexShader, PixelShader, GatheringViewSize, DrawIndirectParametersOffset](FRHICommandList& RHICmdList)
 			{
 				RHICmdList.SetViewport(0, 0, 0.0f, GatheringViewSize.X, GatheringViewSize.Y, 1.0f);
diff --git a/Engine/Source/Runtime/Renderer/Private/PostProcess/PostProcessMotionBlur.cpp b/Engine/Source/Runtime/Renderer/Private/PostProcess/PostProcessMotionBlur.cpp
--- a/Engine/Source/Runtime/Renderer/Private/PostProcess/PostProcessMotionBlur.cpp
+++ b/Engine/Source/Runtime/Renderer/Private/PostProcess/PostProcessMotionBlur.cpp
@@ -551,7 +551,7 @@ void ComputeMotionBlurVelocity(
 		GraphBuilder.AddPass(
 			RDG_EVENT_NAME("VelocityTileScatter %dx%d", VelocityTileCount.X, VelocityTileCount.Y),
 			PassParameters,
-			ERenderGraphPassFlags::None,
+			ERDGPassFlags::Raster,
 			[VertexShader, PixelShader, VelocityTileCount, PassParameters](FRHICommandListImmediate& RHICmdList)
 		{
 			FRHIVertexShader* RHIVertexShader = GETSAFERHISHADER_VERTEX(*VertexShader);
@@ -840,7 +840,7 @@ FRDGTextureRef VisualizeMotionBlur(
 	GraphBuilder.AddPass(
 		RDG_EVENT_NAME("MotionBlurVisualize"),
 		PassParameters,
-		ERenderGraphPassFlags::None,
+		ERDGPassFlags::Raster,
 		[ScreenPassView, ColorTextureOutput, ColorViewport, PixelShader, PassParameters](FRHICommandListImmediate& RHICmdList)
 	{
 		DrawScreenPass(RHICmdList, ScreenPassView, ColorViewport, ColorViewport, *PixelShader, *PassParameters);
diff --git a/Engine/Source/Runtime/Renderer/Private/PostProcess/PostProcessSubsurface.cpp b/Engine/Source/Runtime/Renderer/Private/PostProcess/PostProcessSubsurface.cpp
--- a/Engine/Source/Runtime/Renderer/Private/PostProcess/PostProcessSubsurface.cpp
+++ b/Engine/Source/Runtime/Renderer/Private/PostProcess/PostProcessSubsurface.cpp
@@ -633,7 +633,7 @@ FRDGTextureRef ComputeSubsurface(
 		GraphBuilder.AddPass(
 			RDG_EVENT_NAME("SubsurfaceViewportCopy"),
 			PassParameters,
-			ERenderGraphPassFlags::None,
+			ERDGPassFlags::Raster,
 			[&Views, ViewMask, ViewCount, PixelShader, InputTextureSize, PassParameters](FRHICommandListImmediate& RHICmdList)
 		{
 			for (uint32 ViewIndex = 0; ViewIndex < ViewCount; ++ViewIndex)
@@ -706,7 +706,7 @@ void VisualizeSubsurface(
 	GraphBuilder.AddPass(
 		RDG_EVENT_NAME("SubsurfaceVisualize"),
 		PassParameters,
-		ERenderGraphPassFlags::None,
+		ERDGPassFlags::Raster,
 		[ScreenPassView, SceneViewport, SceneTextureOutput, PixelShader, PassParameters](FRHICommandListImmediate& RHICmdList)
 	{
 		DrawScreenPass(RHICmdList, ScreenPassView, SceneViewport, SceneViewport, *PixelShader, *PassParameters);
diff --git a/Engine/Source/Runtime/Renderer/Private/PostProcess/PostProcessTemporalAA.cpp b/Engine/Source/Runtime/Renderer/Private/PostProcess/PostProcessTemporalAA.cpp
--- a/Engine/Source/Runtime/Renderer/Private/PostProcess/PostProcessTemporalAA.cpp
+++ b/Engine/Source/Runtime/Renderer/Private/PostProcess/PostProcessTemporalAA.cpp
@@ -731,7 +731,7 @@ FTAAOutputs FTAAPassParameters::AddTemporalAAPass(
 				PassName, bUseFast ? TEXT(" Fast") : TEXT(""),
 				PracticableDestRect.Width(), PracticableDestRect.Height()),
 			PassParameters,
-			ERenderGraphPassFlags::None,
+			ERDGPassFlags::Raster,
 			[PassParameters, &View, PracticableDestRect, bCameraCut, BasePermutationVector, bUseResponsiveStencilTest](FRHICommandList& RHICmdList)
 		{
 			RHICmdList.SetViewport(PracticableDestRect.Min.X, PracticableDestRect.Min.Y, 0.0f, PracticableDestRect.Max.X, PracticableDestRect.Max.Y, 1.0f);
diff --git a/Engine/Source/Runtime/Renderer/Private/RayTracing/RayTracingAmbientOcclusion.cpp b/Engine/Source/Runtime/Renderer/Private/RayTracing/RayTracingAmbientOcclusion.cpp
--- a/Engine/Source/Runtime/Renderer/Private/RayTracing/RayTracingAmbientOcclusion.cpp
+++ b/Engine/Source/Runtime/Renderer/Private/RayTracing/RayTracingAmbientOcclusion.cpp
@@ -202,7 +202,7 @@ void FDeferredShadingSceneRenderer::RenderRayTracingAmbientOcclusion(
 	GraphBuilder.AddPass(
 		RDG_EVENT_NAME("AmbientOcclusionRayTracing %dx%d", RayTracingResolution.X, RayTracingResolution.Y),
 		PassParameters,
-		ERenderGraphPassFlags::Compute,
+		ERDGPassFlags::Compute,
 		[PassParameters, this, &View, RayGenerationShader, RayTracingResolution](FRHICommandList& RHICmdList)
 	{
 		FRayTracingShaderBindingsWriter GlobalResources;
diff --git a/Engine/Source/Runtime/Renderer/Private/RayTracing/RayTracingBarycentrics.cpp b/Engine/Source/Runtime/Renderer/Private/RayTracing/RayTracingBarycentrics.cpp
--- a/Engine/Source/Runtime/Renderer/Private/RayTracing/RayTracingBarycentrics.cpp
+++ b/Engine/Source/Runtime/Renderer/Private/RayTracing/RayTracingBarycentrics.cpp
@@ -87,7 +87,7 @@ void FDeferredShadingSceneRenderer::RenderRayTracingBarycentrics(FRHICommandList
 	GraphBuilder.AddPass(
 		RDG_EVENT_NAME("Barycentrics"),
 		RayGenParameters,
-		ERenderGraphPassFlags::Compute,
+		ERDGPassFlags::Compute,
 		[this, RayGenParameters, RayGenShader, &SceneContext, RayTracingSceneRHI, Pipeline, ViewRect](FRHICommandList& RHICmdList)
 	{
 		FRayTracingShaderBindingsWriter GlobalResources;
diff --git a/Engine/Source/Runtime/Renderer/Private/RayTracing/RayTracingDebug.cpp b/Engine/Source/Runtime/Renderer/Private/RayTracing/RayTracingDebug.cpp
--- a/Engine/Source/Runtime/Renderer/Private/RayTracing/RayTracingDebug.cpp
+++ b/Engine/Source/Runtime/Renderer/Private/RayTracing/RayTracingDebug.cpp
@@ -122,7 +122,7 @@ void FDeferredShadingSceneRenderer::RenderRayTracingDebug(FRHICommandListImmedia
 	GraphBuilder.AddPass(
 		RDG_EVENT_NAME("RayTracingDebug"),
 		RayGenParameters,
-		ERenderGraphPassFlags::Compute,
+		ERDGPassFlags::Compute,
 		[this, RayGenParameters, RayGenShader, &SceneContext, RayTracingSceneRHI, Pipeline, ViewRect](FRHICommandList& RHICmdList)
 	{
 		SCOPED_GPU_STAT(RHICmdList, RayTracingDebug);
diff --git a/Engine/Source/Runtime/Renderer/Private/RayTracing/RayTracingGlobalIllumination.cpp b/Engine/Source/Runtime/Renderer/Private/RayTracing/RayTracingGlobalIllumination.cpp
--- a/Engine/Source/Runtime/Renderer/Private/RayTracing/RayTracingGlobalIllumination.cpp
+++ b/Engine/Source/Runtime/Renderer/Private/RayTracing/RayTracingGlobalIllumination.cpp
@@ -497,7 +497,7 @@ void FDeferredShadingSceneRenderer::RenderRayTracingGlobalIllumination(
 		GraphBuilder.AddPass(
 			RDG_EVENT_NAME("GlobalIlluminationRayTracing %dx%d", RayTracingResolution.X, RayTracingResolution.Y),
 			PassParameters,
-			ERenderGraphPassFlags::Compute,
+			ERDGPassFlags::Compute,
 			[PassParameters, this, &View, RayGenerationShader, RayTracingResolution](FRHICommandList& RHICmdList)
 		{
 			FRayTracingShaderBindingsWriter GlobalResources;
@@ -556,7 +556,7 @@ void FDeferredShadingSceneRenderer::RenderRayTracingGlobalIllumination(
 		GraphBuilder.AddPass(
 			RDG_EVENT_NAME("GlobalIlluminationComposite"),
 			PassParameters,
-			ERenderGraphPassFlags::None,
+			ERDGPassFlags::Raster,
 			[this, &SceneContext, &View, &SceneTextures, PassParameters](FRHICommandListImmediate& RHICmdList)
 			{
 				TShaderMapRef<FPostProcessVS> VertexShader(View.ShaderMap);
@@ -616,7 +616,7 @@ void FDeferredShadingSceneRenderer::CompositeGlobalIllumination(
 	GraphBuilder.AddPass(
 		RDG_EVENT_NAME("GlobalIlluminationComposite"),
 		PassParameters,
-		ERenderGraphPassFlags::None,
+		ERDGPassFlags::Raster,
 		[this, &SceneContext, &View, PassParameters](FRHICommandListImmediate& RHICmdList)
 		{
 			TShaderMapRef<FPostProcessVS> VertexShader(View.ShaderMap);
diff --git a/Engine/Source/Runtime/Renderer/Private/RayTracing/RayTracingReflections.cpp b/Engine/Source/Runtime/Renderer/Private/RayTracing/RayTracingReflections.cpp
--- a/Engine/Source/Runtime/Renderer/Private/RayTracing/RayTracingReflections.cpp
+++ b/Engine/Source/Runtime/Renderer/Private/RayTracing/RayTracingReflections.cpp
@@ -333,7 +333,7 @@ void FDeferredShadingSceneRenderer::RenderRayTracingReflections(
 			GraphBuilder.AddPass(
 				RDG_EVENT_NAME("ReflectionRayTracingGatherMaterials %dx%d", TileAlignedResolution.X, TileAlignedResolution.Y),
 				PassParameters,
-				ERenderGraphPassFlags::Compute,
+				ERDGPassFlags::Compute,
 				[PassParameters, this, &View, RayGenShader, TileAlignedResolution](FRHICommandList& RHICmdList)
 			{
 				FRayTracingPipelineState* Pipeline = BindRayTracingDeferredMaterialGatherPipeline(RHICmdList, View, RayGenShader->GetRayTracingShader());
@@ -357,7 +357,7 @@ void FDeferredShadingSceneRenderer::RenderRayTracingReflections(
 			GraphBuilder.AddPass(
 				RDG_EVENT_NAME("ReflectionRayTracing %dx%d", RayTracingResolution.X, RayTracingResolution.Y),
 				PassParameters,
-				ERenderGraphPassFlags::Compute,
+				ERDGPassFlags::Compute,
 				[PassParameters, this, &View, RayGenShader, RayTracingResolution, DeferredMaterialBufferNumElements, DeferredMaterialMode](FRHICommandList& RHICmdList)
 			{
 				FRayTracingShaderBindingsWriter GlobalResources;
diff --git a/Engine/Source/Runtime/Renderer/Private/RayTracing/RayTracingShadows.cpp b/Engine/Source/Runtime/Renderer/Private/RayTracing/RayTracingShadows.cpp
--- a/Engine/Source/Runtime/Renderer/Private/RayTracing/RayTracingShadows.cpp
+++ b/Engine/Source/Runtime/Renderer/Private/RayTracing/RayTracingShadows.cpp
@@ -221,7 +221,7 @@ void FDeferredShadingSceneRenderer::RenderRayTracingShadows(
 		GraphBuilder.AddPass(
 			RDG_EVENT_NAME("RayTracedShadow (spp=%d) %dx%d", RayTracingConfig.RayCountPerPixel, Resolution.X, Resolution.Y),
 			PassParameters,
-			ERenderGraphPassFlags::Compute,
+			ERDGPassFlags::Compute,
 			[this, &View, RayGenerationShader, PassParameters, Resolution](FRHICommandList& RHICmdList)
 		{
 			FRayTracingShaderBindingsWriter GlobalResources;
diff --git a/Engine/Source/Runtime/Renderer/Private/RayTracing/RayTracingTranslucency.cpp b/Engine/Source/Runtime/Renderer/Private/RayTracing/RayTracingTranslucency.cpp
--- a/Engine/Source/Runtime/Renderer/Private/RayTracing/RayTracingTranslucency.cpp
+++ b/Engine/Source/Runtime/Renderer/Private/RayTracing/RayTracingTranslucency.cpp
@@ -409,7 +409,7 @@ void FDeferredShadingSceneRenderer::RenderRayTracingTranslucencyView(
 	GraphBuilder.AddPass(
 		RDG_EVENT_NAME("TranslucencyRayTracing %dx%d", RayTracingResolution.X, RayTracingResolution.Y),
 		PassParameters,
-		ERenderGraphPassFlags::Compute,
+		ERDGPassFlags::Compute,
 		[PassParameters, this, &View, RayGenShader, RayTracingResolution](FRHICommandList& RHICmdList)
 	{
 		FRayTracingPipelineState* Pipeline = View.RayTracingMaterialPipeline;
diff --git a/Engine/Source/Runtime/Renderer/Private/RayTracing/RaytracingSkylight.cpp b/Engine/Source/Runtime/Renderer/Private/RayTracing/RaytracingSkylight.cpp
--- a/Engine/Source/Runtime/Renderer/Private/RayTracing/RaytracingSkylight.cpp
+++ b/Engine/Source/Runtime/Renderer/Private/RayTracing/RaytracingSkylight.cpp
@@ -1035,7 +1035,7 @@ void FDeferredShadingSceneRenderer::RenderRayTracingSkyLight(
 		GraphBuilder.AddPass(
 			RDG_EVENT_NAME("SkyLightRayTracing %dx%d", RayTracingResolution.X, RayTracingResolution.Y),
 			PassParameters,
-			ERenderGraphPassFlags::Compute,
+			ERDGPassFlags::Compute,
 			[PassParameters, this, &View, RayGenerationShader, RayTracingResolution](FRHICommandList& RHICmdList)
 		{
 			FRayTracingShaderBindingsWriter GlobalResources;
@@ -1153,7 +1153,7 @@ void FDeferredShadingSceneRenderer::CompositeRayTracingSkyLight(
 		GraphBuilder.AddPass(
 			RDG_EVENT_NAME("GlobalIlluminationComposite"),
 			PassParameters,
-			ERenderGraphPassFlags::None,
+			ERDGPassFlags::Raster,
 			[this, &SceneContext, &View, PassParameters](FRHICommandListImmediate& RHICmdList)
 		{
 			TShaderMapRef<FPostProcessVS> VertexShader(View.ShaderMap);
diff --git a/Engine/Source/Runtime/Renderer/Private/ReflectionEnvironment.cpp b/Engine/Source/Runtime/Renderer/Private/ReflectionEnvironment.cpp
--- a/Engine/Source/Runtime/Renderer/Private/ReflectionEnvironment.cpp
+++ b/Engine/Source/Runtime/Renderer/Private/ReflectionEnvironment.cpp
@@ -854,7 +854,7 @@ void FDeferredShadingSceneRenderer::RenderDeferredReflectionsAndSkyLighting(FRHI
 			GraphBuilder.AddPass(
 				RDG_EVENT_NAME("ReflectionEnvironmentAndSky %dx%d", View.ViewRect.Width(), View.ViewRect.Height()),
 				PassParameters,
-				ERDGPassFlags::None,
+				ERDGPassFlags::Raster,
 				[PassParameters, &View, PixelShader](FRHICommandList& InRHICmdList)
 			{
 				InRHICmdList.SetViewport(View.ViewRect.Min.X, View.ViewRect.Min.Y, 0.0f, View.ViewRect.Max.X, View.ViewRect.Max.Y, 1.0f);
diff --git a/Engine/Source/Runtime/Renderer/Private/SceneOcclusion.cpp b/Engine/Source/Runtime/Renderer/Private/SceneOcclusion.cpp
--- a/Engine/Source/Runtime/Renderer/Private/SceneOcclusion.cpp
+++ b/Engine/Source/Runtime/Renderer/Private/SceneOcclusion.cpp
@@ -1120,7 +1120,7 @@ void BuildHZB(FRDGBuilder& GraphBuilder, FViewInfo& View)
 		GraphBuilder.AddPass(
 			RDG_EVENT_NAME("HZB(mip=0) %dx%d", HZBSize.X, HZBSize.Y),
 			PassParameters,
-			ERenderGraphPassFlags::None,
+			ERDGPassFlags::Raster,
 			[PassParameters, &View, HZBSize](FRHICommandListImmediate& RHICmdList)
 			{
 				FGraphicsPipelineStateInitializer GraphicsPSOInit;
@@ -1203,7 +1203,7 @@ void BuildHZB(FRDGBuilder& GraphBuilder, FViewInfo& View)
 		GraphBuilder.AddPass(
 			RDG_EVENT_NAME("HZB(mip=%d) %dx%d", MipIndex, DstSize.X, DstSize.Y),
 			PassParameters,
-			ERenderGraphPassFlags::GenerateMips,
+			ERDGPassFlags::GenerateMips,
 			[PassParameters, SrcSize, DstSize, &View](FRHICommandListImmediate& RHICmdList)
 			{
 				FHZBBuildPS::FPermutationDomain PermutationVector;
diff --git a/Engine/Source/Runtime/Renderer/Private/ScreenPass.h b/Engine/Source/Runtime/Renderer/Private/ScreenPass.h
--- a/Engine/Source/Runtime/Renderer/Private/ScreenPass.h
+++ b/Engine/Source/Runtime/Renderer/Private/ScreenPass.h
@@ -250,7 +250,7 @@ inline void AddDrawScreenPass(
 	GraphBuilder.AddPass(
 		Forward<FRDGEventName>(PassName),
 		PixelShaderParameters,
-		ERenderGraphPassFlags::None,
+		ERDGPassFlags::Raster,
 		[ScreenPassView, OutputViewportRect, InputViewportRect, InputExtent, PixelShader, PixelShaderParameters]
 		(FRHICommandListImmediate& RHICmdList)
 	{
diff --git a/Engine/Source/Runtime/Renderer/Private/ScreenSpaceRayTracing.cpp b/Engine/Source/Runtime/Renderer/Private/ScreenSpaceRayTracing.cpp
--- a/Engine/Source/Runtime/Renderer/Private/ScreenSpaceRayTracing.cpp
+++ b/Engine/Source/Runtime/Renderer/Private/ScreenSpaceRayTracing.cpp
@@ -406,7 +406,7 @@ void RenderScreenSpaceReflections(
 		GraphBuilder.AddPass(
 			RDG_EVENT_NAME("SSR StencilSetup %dx%d", View.ViewRect.Width(), View.ViewRect.Height()),
 			PassParameters,
-			ERenderGraphPassFlags::None,
+			ERDGPassFlags::Raster,
 			[PassParameters, &View, PixelShader](FRHICommandList& RHICmdList)
 		{
 			SCOPED_GPU_STAT(RHICmdList, ScreenSpaceReflections);
@@ -488,7 +488,7 @@ void RenderScreenSpaceReflections(
 				SSRQuality, RayTracingConfigs.RayCountPerPixel, bDenoiser ? TEXT(" DenoiserOutput") : TEXT(""),
 				View.ViewRect.Width(), View.ViewRect.Height()),
 			PassParameters,
-			ERenderGraphPassFlags::None,
+			ERDGPassFlags::Raster,
 			[PassParameters, &View, PixelShader, SSRStencilPrePass](FRHICommandList& RHICmdList)
 		{
 			SCOPED_GPU_STAT(RHICmdList, ScreenSpaceReflections);
diff --git a/Engine/Source/Runtime/Renderer/Private/ShaderPrint.cpp b/Engine/Source/Runtime/Renderer/Private/ShaderPrint.cpp
--- a/Engine/Source/Runtime/Renderer/Private/ShaderPrint.cpp
+++ b/Engine/Source/Runtime/Renderer/Private/ShaderPrint.cpp
@@ -424,7 +424,7 @@ namespace ShaderPrint
 			GraphBuilder.AddPass(
 				RDG_EVENT_NAME("DrawSymbols"),
 				PassParameters,
-				ERenderGraphPassFlags::None,
+				ERDGPassFlags::Raster,
 				[VertexShader, PixelShader, PassParameters](FRHICommandListImmediate& RHICmdListImmediate)
 			{
 				FGraphicsPipelineStateInitializer GraphicsPSOInit;
diff --git a/Engine/Source/Runtime/Renderer/Private/VT/RuntimeVirtualTextureRender.cpp b/Engine/Source/Runtime/Renderer/Private/VT/RuntimeVirtualTextureRender.cpp
--- a/Engine/Source/Runtime/Renderer/Private/VT/RuntimeVirtualTextureRender.cpp
+++ b/Engine/Source/Runtime/Renderer/Private/VT/RuntimeVirtualTextureRender.cpp
@@ -658,7 +658,7 @@ namespace RuntimeVirtualTexture
 			GraphBuilder.AddPass(
 				RDG_EVENT_NAME("VirtualTextureDraw"),
 				PassParameters,
-				ERenderGraphPassFlags::None,
+				ERDGPassFlags::Raster,
 				[Scene, View, MaterialType, RuntimeVirtualTextureMask, vLevel, MaxLevel](FRHICommandListImmediate& RHICmdListImmediate)
 			{
 				DrawMeshes(RHICmdListImmediate, Scene, View, MaterialType, RuntimeVirtualTextureMask, vLevel, MaxLevel);
diff --git a/Engine/Source/Runtime/Renderer/Private/VolumetricFog.cpp b/Engine/Source/Runtime/Renderer/Private/VolumetricFog.cpp
--- a/Engine/Source/Runtime/Renderer/Private/VolumetricFog.cpp
+++ b/Engine/Source/Runtime/Renderer/Private/VolumetricFog.cpp
@@ -506,7 +506,7 @@ void FDeferredShadingSceneRenderer::RenderLocalLightsForVolumetricFog(
 		GraphBuilder.AddPass(
 			RDG_EVENT_NAME("ShadowedLights"),
 			PassParameters,
-			ERenderGraphPassFlags::None,
+			ERDGPassFlags::Raster,
 			[PassParameters, &View, this, LightsToInject, VolumetricFogGridSize, GridZParams, bUseTemporalReprojection, IntegrationData, FogInfo](FRHICommandListImmediate& RHICmdList)
 		{
 			for (int32 LightIndex = 0; LightIndex < LightsToInject.Num(); LightIndex++)
@@ -1062,7 +1062,7 @@ void FDeferredShadingSceneRenderer::ComputeVolumetricFog(FRHICommandListImmediat
 				GraphBuilder.AddPass(
 					RDG_EVENT_NAME("InitializeVolumeAttributes"),
 					PassParameters,
-					ERenderGraphPassFlags::Compute,
+					ERDGPassFlags::Compute,
 					[PassParameters, &View, VolumetricFogGridSize, IntegrationData, ComputeShader](FRHICommandListImmediate& RHICmdList)
 				{
 					const FIntVector NumGroups = FIntVector::DivideAndRoundUp(VolumetricFogGridSize, VolumetricFogGridInjectionGroupSize);
@@ -1128,7 +1128,7 @@ void FDeferredShadingSceneRenderer::ComputeVolumetricFog(FRHICommandListImmediat
 						bUseDistanceFieldSkyOcclusion ? TEXT("DFAO") : TEXT(""),
 						PassParameters->LightFunctionTexture ? TEXT("LF") : TEXT("")),
 					PassParameters,
-					ERenderGraphPassFlags::Compute,
+					ERDGPassFlags::Compute,
 					[PassParameters, ComputeShader, &View, this, FogInfo, bUseTemporalReprojection, VolumetricFogGridSize, IntegrationData, bUseDirectionalLightShadowing, bUseDistanceFieldSkyOcclusion, LightFunctionWorldToShadow](FRHICommandListImmediate& RHICmdList)
 				{
 					UnbindRenderTargets(RHICmdList);
@@ -1159,7 +1159,7 @@ void FDeferredShadingSceneRenderer::ComputeVolumetricFog(FRHICommandListImmediat
 				GraphBuilder.AddPass(
 					RDG_EVENT_NAME("FinalIntegration"),
 					PassParameters,
-					ERenderGraphPassFlags::Compute,
+					ERDGPassFlags::Compute,
 					[PassParameters, &View, VolumetricFogGridSize, IntegrationData, this](FRHICommandListImmediate& RHICmdList)
 				{
 					const FIntVector NumGroups = FIntVector::DivideAndRoundUp(VolumetricFogGridSize, VolumetricFogIntegrationGroupSize);
diff --git a/Engine/Source/Runtime/Renderer/Private/VolumetricFogLightFunction.cpp b/Engine/Source/Runtime/Renderer/Private/VolumetricFogLightFunction.cpp
--- a/Engine/Source/Runtime/Renderer/Private/VolumetricFogLightFunction.cpp
+++ b/Engine/Source/Runtime/Renderer/Private/VolumetricFogLightFunction.cpp
@@ -207,7 +207,7 @@ void FDeferredShadingSceneRenderer::RenderLightFunctionForVolumetricFog(
 			GraphBuilder.AddPass(
 				RDG_EVENT_NAME("LightFunction"),
 				PassParameters,
-				ERenderGraphPassFlags::None,
+				ERDGPassFlags::Raster,
 				[PassParameters, &View, MaterialProxy, LightFunctionResolution, DirectionalLightSceneInfo, WorldToShadowValue, this](FRHICommandListImmediate& RHICmdList)
 			{
 				const FMaterial* Material = MaterialProxy->GetMaterial(Scene->GetFeatureLevel());
diff --git a/Engine/Source/Runtime/Renderer/Private/VolumetricFogVoxelization.cpp b/Engine/Source/Runtime/Renderer/Private/VolumetricFogVoxelization.cpp
--- a/Engine/Source/Runtime/Renderer/Private/VolumetricFogVoxelization.cpp
+++ b/Engine/Source/Runtime/Renderer/Private/VolumetricFogVoxelization.cpp
@@ -595,7 +595,7 @@ void FDeferredShadingSceneRenderer::VoxelizeFogVolumePrimitives(
 		GraphBuilder.AddPass(
 			RDG_EVENT_NAME("VoxelizeVolumePrimitives"),
 			PassParameters,
-			ERenderGraphPassFlags::None,
+			ERDGPassFlags::Raster,
 			[PassParameters, Scene = Scene, &View, VolumetricFogGridSize, IntegrationData, VolumetricFogDistance, GridZParams](FRHICommandListImmediate& RHICmdList)
 		{
 			FViewUniformShaderParameters ViewVoxelizeParameters = *View.CachedViewUniformShaderParameters;
