lot of unaccounted time in gpu profiler

diff --git a/Engine/Source/Runtime/Renderer/Private/DeferredShadingRenderer.cpp b/Engine/Source/Runtime/Renderer/Private/DeferredShadingRenderer.cpp
--- a/Engine/Source/Runtime/Renderer/Private/DeferredShadingRenderer.cpp
+++ b/Engine/Source/Runtime/Renderer/Private/DeferredShadingRenderer.cpp
@@ -159,6 +159,12 @@ DECLARE_CYCLE_STAT(TEXT("DeferredShadingSceneRenderer ViewExtensionPostRenderBas
 
 DECLARE_GPU_STAT_NAMED(RayTracingTLAS, TEXT("Ray Tracing Top Level Acceleration Structure"));
 DECLARE_GPU_STAT(Postprocessing);
+DECLARE_GPU_STAT(VisibilityCommands);
+DECLARE_GPU_STAT(RenderDeferredLighting);
+DECLARE_GPU_STAT(AllocateRendertargets);
+DECLARE_GPU_STAT(FrameRenderFinish);
+DECLARE_GPU_STAT(SortLights);
+DECLARE_GPU_STAT(PostRenderOpsFX);
 DECLARE_GPU_STAT(HZB);
 DECLARE_GPU_STAT_NAMED(Unaccounted, TEXT("[unaccounted]"));
 
@@ -871,6 +877,7 @@ void FDeferredShadingSceneRenderer::Render(FRHICommandListImmediate& RHICmdList)
 	
 	{
 		SCOPE_CYCLE_COUNTER(STAT_FDeferredShadingSceneRenderer_Render_Init);
+		SCOPED_GPU_STAT(RHICmdList, AllocateRendertargets);
 
 		// Initialize global system textures (pass-through if already initialized).
 		GSystemTextures.InitializeTextures(RHICmdList, FeatureLevel);
@@ -905,7 +912,12 @@ void FDeferredShadingSceneRenderer::Render(FRHICommandListImmediate& RHICmdList)
 
 	// Find the visible primitives.
 	RHICmdList.ImmediateFlush(EImmediateFlushType::DispatchToRHIThread);
-	bool bDoInitViewAftersPrepass = InitViews(RHICmdList, BasePassDepthStencilAccess, ILCTaskData, UpdateViewCustomDataEvents);
+	
+	bool bDoInitViewAftersPrepass = false;
+	{
+		SCOPED_GPU_STAT(RHICmdList, VisibilityCommands);
+		bDoInitViewAftersPrepass = InitViews(RHICmdList, BasePassDepthStencilAccess, ILCTaskData, UpdateViewCustomDataEvents);
+	}
 
 #if !UE_BUILD_SHIPPING
 	if (CVarStallInitViews.GetValueOnRenderThread() > 0.0f)
@@ -1157,8 +1169,12 @@ void FDeferredShadingSceneRenderer::Render(FRHICommandListImmediate& RHICmdList)
 
 	// NOTE: The ordering of the lights is used to select sub-sets for different purposes, e.g., those that support clustered deferred.
 	FSortedLightSetSceneInfo SortedLightSet;
-	GatherAndSortLights(SortedLightSet);
-	ComputeLightGrid(RHICmdList, bComputeLightGrid, SortedLightSet);
+	
+	{
+		SCOPED_GPU_STAT(RHICmdList, SortLights);
+		GatherAndSortLights(SortedLightSet);
+		ComputeLightGrid(RHICmdList, bComputeLightGrid, SortedLightSet);
+	}
 
 	{
 		SCOPE_CYCLE_COUNTER(STAT_FDeferredShadingSceneRenderer_AllocGBufferTargets);
@@ -1173,13 +1189,15 @@ void FDeferredShadingSceneRenderer::Render(FRHICommandListImmediate& RHICmdList)
 
 	if (bOcclusionBeforeBasePass)
 	{
+		SCOPED_GPU_STAT(RHICmdList, HZB);
+
 		if (bIsOcclusionTesting)
 		{
 			RenderOcclusion(RHICmdList);
 		}
+
 		bool bUseHzbOcclusion = RenderHzb(RHICmdList);
 		
-		SCOPED_GPU_STAT(RHICmdList, HZB);
 		if (bUseHzbOcclusion || bIsOcclusionTesting)
 		{
 			FinishOcclusion(RHICmdList);
@@ -1460,13 +1478,15 @@ void FDeferredShadingSceneRenderer::Render(FRHICommandListImmediate& RHICmdList)
 	// Occlusion after base pass
 	if (!bOcclusionBeforeBasePass)
 	{
+		SCOPED_GPU_STAT(RHICmdList, HZB);
 		// #todo-renderpasses Needs its own renderpass. Does this need more than the depth?
 		if (bIsOcclusionTesting)
 		{
 			RenderOcclusion(RHICmdList);
 		}
+
 		bool bUseHzbOcclusion = RenderHzb(RHICmdList);
-		SCOPED_GPU_STAT(RHICmdList, HZB);
+
 		if (bUseHzbOcclusion || bIsOcclusionTesting)
 		{
 			FinishOcclusion(RHICmdList);
@@ -1619,6 +1639,7 @@ void FDeferredShadingSceneRenderer::Render(FRHICommandListImmediate& RHICmdList)
 	// Render lighting.
 	if (bRenderDeferredLighting)
 	{
+		SCOPED_GPU_STAT(RHICmdList, RenderDeferredLighting);
 		CSV_SCOPED_TIMING_STAT_EXCLUSIVE(RenderLighting);
 		SCOPE_CYCLE_COUNTER(STAT_FDeferredShadingSceneRenderer_Lighting);
 
@@ -1797,6 +1818,7 @@ void FDeferredShadingSceneRenderer::Render(FRHICommandListImmediate& RHICmdList)
 	{
 		CSV_SCOPED_TIMING_STAT_EXCLUSIVE(RenderOpaqueFX);
 		SCOPE_CYCLE_COUNTER(STAT_FDeferredShadingSceneRenderer_FXSystem_PostRenderOpaque);
+		SCOPED_GPU_STAT(RHICmdList, PostRenderOpsFX);
 
 		FSceneTexturesUniformParameters SceneTextureParameters;
 		SetupSceneTextureUniformParameters(SceneContext, FeatureLevel, ESceneTextureSetupMode::SceneDepth | ESceneTextureSetupMode::GBuffers, SceneTextureParameters);
@@ -2033,6 +2055,7 @@ void FDeferredShadingSceneRenderer::Render(FRHICommandListImmediate& RHICmdList)
 
 	{
 		SCOPE_CYCLE_COUNTER(STAT_FDeferredShadingSceneRenderer_RenderFinish);
+		SCOPED_GPU_STAT(RHICmdList, FrameRenderFinish);
 		RHICmdList.SetCurrentStat(GET_STATID(STAT_CLM_RenderFinish));
 		RenderFinish(RHICmdList);
 		RHICmdList.SetCurrentStat(GET_STATID(STAT_CLM_AfterFrame));
