remove unneeded UAV flag

diff --git a/Engine/Source/Runtime/Renderer/Private/ScreenSpaceRayTracing.cpp b/Engine/Source/Runtime/Renderer/Private/ScreenSpaceRayTracing.cpp
--- a/Engine/Source/Runtime/Renderer/Private/ScreenSpaceRayTracing.cpp
+++ b/Engine/Source/Runtime/Renderer/Private/ScreenSpaceRayTracing.cpp
@@ -351,7 +351,7 @@ void RenderScreenSpaceReflections(
 		FRDGTextureDesc Desc = FPooledRenderTargetDesc::Create2DDesc(
 			FSceneRenderTargets::Get_FrameConstantsOnly().GetBufferSizeXY(),
 			PF_FloatRGBA, FClearValueBinding(FLinearColor(0, 0, 0, 0)),
-			TexCreate_None, TexCreate_RenderTargetable | TexCreate_ShaderResource | TexCreate_UAV,
+			TexCreate_None, TexCreate_RenderTargetable | TexCreate_ShaderResource,
 			false);
 
 		Desc.AutoWritable = false;
