Ray Tracing Translucency Fix

(ushell-p4-cherrypick of 7476085 by [REDACTED_PERSON])

diff --git a/Engine/Shaders/Private/RayTracing/RayTracingTranslucency.usf b/Engine/Shaders/Private/RayTracing/RayTracingTranslucency.usf
--- a/Engine/Shaders/Private/RayTracing/RayTracingTranslucency.usf
+++ b/Engine/Shaders/Private/RayTracing/RayTracingTranslucency.usf
@@ -46,6 +46,7 @@ float TranslucencyMaxRayDistance;
 float TranslucencyMaxRoughness;
 uint  TranslucencyRefraction;
 float MaxNormalBias;
+float PrimaryRayBias;
 
 Texture2D SceneColorTexture;
 
@@ -103,7 +104,7 @@ void RayTracingTranslucencyRGS()
 	uint2 DispatchThreadId = DispatchRaysIndex().xy + View.ViewRectMin;
 	uint2 PixelCoord = GetPixelCoord(DispatchThreadId, UpscaleFactor);
 	uint LinearIndex = PixelCoord.y * View.BufferSizeAndInvSize.x + PixelCoord.x; // TODO: PixelCoord or DispatchThreadId
-	
+
 	RandomSequence RandSequence;
 	RandomSequence_Initialize(RandSequence, LinearIndex, View.StateFrameIndex);
 
@@ -111,10 +112,12 @@ void RayTracingTranslucencyRGS()
 	float2 UV = (float2(PixelCoord) + 0.5) * InvBufferSize;
 	FGBufferData GBufferData = GetGBufferDataFromSceneTextures(UV);
 	float Depth = GBufferData.Depth;
+	float3 WorldPosition = ReconstructWorldPositionFromDepth(UV, Depth);
+	float WorldSpaceDistance = length(WorldPosition - View.WorldCameraOrigin);
 
 	// Trace rays from camera origin to (Gbuffer - epsilon) to only intersect translucent objects
 	RayDesc Ray = CreatePrimaryRay(UV);
-	Ray.TMax = Depth - 0.1;
+	Ray.TMax = WorldSpaceDistance - PrimaryRayBias;
 	FRayCone RayCone = (FRayCone)0;
 	RayCone.SpreadAngle = View.EyeToPixelSpreadAngle;
 
@@ -142,7 +145,7 @@ void RayTracingTranslucencyRGS()
 		//
 		// Handle no hit condition
 		//
-		if (Payload.IsMiss())
+		if (Payload.IsMiss() || RefractionRayIndex == 0 && Payload.BlendingMode == RAY_TRACING_BLEND_MODE_OPAQUE)
 		{
 			if (bHasScattered && bAllowSkySampling)
 			{
@@ -254,7 +257,7 @@ void RayTracingTranslucencyRGS()
 	if (!bHasScattered)
 	{
 		// Use the scene radiance for ray that has not been scattered/refracted (no surface or IORin=IORout). Still apply the throughtput in case we have traversed surfaces with opacity>0.
-		PathRadiance += PathThroughput * SceneColorTexture.SampleLevel(GlobalPointClampedSampler, UV, 0).xyz;
+		PathRadiance += PathThroughput * SceneColorTexture.SampleLevel(GlobalPointClampedSampler, UV, 0).xyz / View.PreExposure;
 	}
 
 #if USE_PREEXPOSURE
diff --git a/Engine/Source/Runtime/Renderer/Private/RayTracing/RayTracingTranslucency.cpp b/Engine/Source/Runtime/Renderer/Private/RayTracing/RayTracingTranslucency.cpp
--- a/Engine/Source/Runtime/Renderer/Private/RayTracing/RayTracingTranslucency.cpp
+++ b/Engine/Source/Runtime/Renderer/Private/RayTracing/RayTracingTranslucency.cpp
@@ -103,6 +103,12 @@ static FAutoConsoleVariableRef CVarRayTracingTranslucencyRefraction(
 	GRayTracingTranslucencyRefraction,
 	TEXT("Enables refraction in ray traced Translucency (default = 1)"));
 
+static float GRayTracingTranslucencyPrimaryRayBias = 1e-5;
+static FAutoConsoleVariableRef CVarRayTracingTranslucencyPrimaryRayBias(
+	TEXT("r.RayTracing.Translucency.PrimaryRayBias"),
+	GRayTracingTranslucencyPrimaryRayBias,
+	TEXT("Sets the bias to be subtracted from the primary ray TMax in ray traced Translucency. Larger bias reduces the chance of opaque objects being intersected in ray traversal, saving performance, but at the risk of skipping some thin translucent objects in proximity of opaque objects. (recommended range: 0.00001 - 0.1) (default = 0.00001)"));
+
 DECLARE_GPU_STAT_NAMED(RayTracingTranslucency, TEXT("Ray Tracing Translucency"));
 
 #if RHI_RAYTRACING
@@ -140,6 +146,7 @@ class FRayTracingTranslucencyRGS : public FGlobalShader
 		SHADER_PARAMETER(float, TranslucencyMaxRoughness)
 		SHADER_PARAMETER(int32, TranslucencyRefraction)
 		SHADER_PARAMETER(float, MaxNormalBias)
+		SHADER_PARAMETER(float, PrimaryRayBias)
 
 		SHADER_PARAMETER_SRV(RaytracingAccelerationStructure, TLAS)
 		SHADER_PARAMETER_SRV(StructuredBuffer<FRTLightingData>, LightDataBuffer)
@@ -383,6 +390,7 @@ void FDeferredShadingSceneRenderer::RenderRayTracingTranslucencyView(
 	PassParameters->TranslucencyMaxRoughness = FMath::Clamp(GRayTracingTranslucencyMaxRoughness >= 0 ? GRayTracingTranslucencyMaxRoughness : View.FinalPostProcessSettings.RayTracingTranslucencyMaxRoughness, 0.01f, 1.0f);
 	PassParameters->TranslucencyRefraction = GRayTracingTranslucencyRefraction >= 0 ? GRayTracingTranslucencyRefraction : View.FinalPostProcessSettings.RayTracingTranslucencyRefraction;
 	PassParameters->MaxNormalBias = GetRaytracingMaxNormalBias();
+	PassParameters->PrimaryRayBias = GRayTracingTranslucencyPrimaryRayBias;
 
 	PassParameters->TLAS = View.RayTracingScene.RayTracingSceneRHI->GetShaderResourceView();
 	PassParameters->ViewUniformBuffer = View.ViewUniformBuffer;
