When rendering in extreme resolutions, clamp volumetric fog texture to hardware limits.

diff --git a/Engine/Source/Runtime/RHI/Private/RHI.cpp b/Engine/Source/Runtime/RHI/Private/RHI.cpp
--- a/Engine/Source/Runtime/RHI/Private/RHI.cpp
+++ b/Engine/Source/Runtime/RHI/Private/RHI.cpp
@@ -495,6 +495,7 @@ bool GRHISupportsLazyShaderCodeLoading = false;
 TRHIGlobal<int32> GMaxShadowDepthBufferSizeX(2048);
 TRHIGlobal<int32> GMaxShadowDepthBufferSizeY(2048);
 TRHIGlobal<int32> GMaxTextureDimensions(2048);
+TRHIGlobal<int32> GMaxVolumeTextureDimensions(2048);
 TRHIGlobal<int32> GMaxCubeTextureDimensions(2048);
 int32 GMaxTextureArrayLayers = 256;
 int32 GMaxTextureSamplers = 16;
diff --git a/Engine/Source/Runtime/RHI/Public/RHI.h b/Engine/Source/Runtime/RHI/Public/RHI.h
--- a/Engine/Source/Runtime/RHI/Public/RHI.h
+++ b/Engine/Source/Runtime/RHI/Public/RHI.h
@@ -381,6 +381,10 @@ extern RHI_API TRHIGlobal<int32> GMaxShadowDepthBufferSizeY;
 /** The maximum size allowed for 2D textures in both dimensions. */
 extern RHI_API TRHIGlobal<int32> GMaxTextureDimensions;
 
+/** The maximum size allowed for 3D textures in all three dimensions. */
+extern RHI_API TRHIGlobal<int32> GMaxVolumeTextureDimensions;
+
+
 FORCEINLINE uint32 GetMax2DTextureDimension()
 {
 	return GMaxTextureDimensions;
diff --git a/Engine/Source/Runtime/Renderer/Private/VolumetricFog.cpp b/Engine/Source/Runtime/Renderer/Private/VolumetricFog.cpp
--- a/Engine/Source/Runtime/Renderer/Private/VolumetricFog.cpp
+++ b/Engine/Source/Runtime/Renderer/Private/VolumetricFog.cpp
@@ -870,10 +870,20 @@ FVector GetVolumetricFogGridZParams(float NearPlane, float FarPlane, int32 GridS
 	return FVector(B, O, S);
 }
 
-FIntVector GetVolumetricFogGridSize(FIntPoint ViewRectSize)
+FIntVector GetVolumetricFogGridSize(FIntPoint ViewRectSize, int32& OutVolumetricFogGridPixelSize)
 {
 	extern int32 GLightGridSizeZ;
-	const FIntPoint VolumetricFogGridSizeXY = FIntPoint::DivideAndRoundUp(ViewRectSize, GVolumetricFogGridPixelSize);
+	FIntPoint VolumetricFogGridSizeXY;
+	int32 VolumetricFogGridPixelSize = GVolumetricFogGridPixelSize;
+	VolumetricFogGridSizeXY = FIntPoint::DivideAndRoundUp(ViewRectSize, VolumetricFogGridPixelSize);
+	if(VolumetricFogGridSizeXY.X > GMaxVolumeTextureDimensions || VolumetricFogGridSizeXY.Y > GMaxVolumeTextureDimensions) //clamp to max volume texture dimensions. only happens for extreme resolutions (~8x2k)
+	{
+		float PixelSizeX = (float)ViewRectSize.X / GMaxVolumeTextureDimensions;
+		float PixelSizeY = (float)ViewRectSize.Y / GMaxVolumeTextureDimensions;
+		VolumetricFogGridPixelSize = FMath::Max(FMath::CeilToInt(PixelSizeX), FMath::CeilToInt(PixelSizeY));
+		VolumetricFogGridSizeXY = FIntPoint::DivideAndRoundUp(ViewRectSize, VolumetricFogGridPixelSize);
+	}
+	OutVolumetricFogGridPixelSize = VolumetricFogGridPixelSize;
 	return FIntVector(VolumetricFogGridSizeXY.X, VolumetricFogGridSizeXY.Y, GVolumetricFogGridSizeZ);
 }
 
@@ -882,7 +892,8 @@ void SetupVolumetricFogGlobalData(const FViewInfo& View, FVolumetricFogGlobalDat
 	const FScene* Scene = (FScene*)View.Family->Scene;
 	const FExponentialHeightFogSceneInfo& FogInfo = Scene->ExponentialFogs[0];
 
-	const FIntVector VolumetricFogGridSize = GetVolumetricFogGridSize(View.ViewRect.Size());
+	int32 VolumetricFogGridPixelSize;
+	const FIntVector VolumetricFogGridSize = GetVolumetricFogGridSize(View.ViewRect.Size(), VolumetricFogGridPixelSize);
 
 	Parameters.GridSizeInt = VolumetricFogGridSize;
 	Parameters.GridSize = FVector(VolumetricFogGridSize);
@@ -890,8 +901,8 @@ void SetupVolumetricFogGlobalData(const FViewInfo& View, FVolumetricFogGlobalDat
 	FVector ZParams = GetVolumetricFogGridZParams(View.NearClippingDistance, FogInfo.VolumetricFogDistance, VolumetricFogGridSize.Z);
 	Parameters.GridZParams = ZParams;
 
-	Parameters.SVPosToVolumeUV = FVector2D(1.0f, 1.0f) / (FVector2D(Parameters.GridSize) * GVolumetricFogGridPixelSize);
-	Parameters.FogGridToPixelXY = FIntPoint(GVolumetricFogGridPixelSize, GVolumetricFogGridPixelSize);
+	Parameters.SVPosToVolumeUV = FVector2D(1.0f, 1.0f) / (FVector2D(Parameters.GridSize) * VolumetricFogGridPixelSize);
+	Parameters.FogGridToPixelXY = FIntPoint(VolumetricFogGridPixelSize, VolumetricFogGridPixelSize);
 	Parameters.MaxDistance = FogInfo.VolumetricFogDistance;
 
 	Parameters.HeightFogInscatteringColor = View.ExponentialFogColor;
@@ -912,14 +923,15 @@ void FViewInfo::SetupVolumetricFogUniformBufferParameters(FViewUniformShaderPara
 	{
 		const FExponentialHeightFogSceneInfo& FogInfo = Scene->ExponentialFogs[0];
 
-		const FIntVector VolumetricFogGridSize = GetVolumetricFogGridSize(ViewRect.Size());
+		int32 VolumetricFogGridPixelSize;
+		const FIntVector VolumetricFogGridSize = GetVolumetricFogGridSize(ViewRect.Size(), VolumetricFogGridPixelSize);
 
 		ViewUniformShaderParameters.VolumetricFogInvGridSize = FVector(1.0f / VolumetricFogGridSize.X, 1.0f / VolumetricFogGridSize.Y, 1.0f / VolumetricFogGridSize.Z);
 
 		const FVector ZParams = GetVolumetricFogGridZParams(NearClippingDistance, FogInfo.VolumetricFogDistance, VolumetricFogGridSize.Z);
 		ViewUniformShaderParameters.VolumetricFogGridZParams = ZParams;
 
-		ViewUniformShaderParameters.VolumetricFogSVPosToVolumeUV = FVector2D(1.0f, 1.0f) / (FVector2D(VolumetricFogGridSize.X, VolumetricFogGridSize.Y) * GVolumetricFogGridPixelSize);
+		ViewUniformShaderParameters.VolumetricFogSVPosToVolumeUV = FVector2D(1.0f, 1.0f) / (FVector2D(VolumetricFogGridSize.X, VolumetricFogGridSize.Y) * VolumetricFogGridPixelSize);
 		ViewUniformShaderParameters.VolumetricFogMaxDistance = FogInfo.VolumetricFogDistance;
 	}
 	else
@@ -945,8 +957,9 @@ void FDeferredShadingSceneRenderer::SetupVolumetricFog()
 		for (int32 ViewIndex = 0; ViewIndex < Views.Num(); ViewIndex++)
 		{
 			FViewInfo& View = Views[ViewIndex];
-
-			const FIntVector VolumetricFogGridSize = GetVolumetricFogGridSize(View.ViewRect.Size());
+			
+			int32 VolumetricFogGridPixelSize;
+			const FIntVector VolumetricFogGridSize = GetVolumetricFogGridSize(View.ViewRect.Size(), VolumetricFogGridPixelSize);
 
 			FVolumetricFogGlobalData GlobalData;
 			SetupVolumetricFogGlobalData(View, GlobalData);
@@ -982,8 +995,9 @@ void FDeferredShadingSceneRenderer::ComputeVolumetricFog(FRHICommandListImmediat
 		for (int32 ViewIndex = 0; ViewIndex < Views.Num(); ViewIndex++)
 		{
 			FViewInfo& View = Views[ViewIndex];
-
-			const FIntVector VolumetricFogGridSize = GetVolumetricFogGridSize(View.ViewRect.Size());
+			
+			int32 VolumetricFogGridPixelSize;
+			const FIntVector VolumetricFogGridSize = GetVolumetricFogGridSize(View.ViewRect.Size(), VolumetricFogGridPixelSize);
 			const FVector GridZParams = GetVolumetricFogGridZParams(View.NearClippingDistance, FogInfo.VolumetricFogDistance, VolumetricFogGridSize.Z);
 
 			//@DW - graph todo
