Integrated from //UE4/Dev-Rendering
#rnx

diff --git a/Engine/Source/Editor/UnrealEd/Private/Lightmass/LightmassRender.cpp b/Engine/Source/Editor/UnrealEd/Private/Lightmass/LightmassRender.cpp
--- a/Engine/Source/Editor/UnrealEd/Private/Lightmass/LightmassRender.cpp
+++ b/Engine/Source/Editor/UnrealEd/Private/Lightmass/LightmassRender.cpp
@@ -1125,30 +1125,19 @@ bool FLightmassMaterialRenderer::GenerateMaterialPropertyData(
 						GetRendererModule().InitializeSystemTextures(RHICmdList);
 					});
 
-				// prefetch VT textures so we have sensible content available
+				// Prefetch all virtual textures so that we have content available
+				//todo[vt]: Move this to calling function to avoid multiple prefetches
 				if (UseVirtualTexturing(GMaxRHIFeatureLevel))
 				{					
-					const FVector2D InScreenSpaceSize(InOutSizeX, InOutSizeY);
+					const ERHIFeatureLevel::Type FeatureLevel = GMaxRHIFeatureLevel;
+					const FVector2D ScreenSpaceSize(InOutSizeX, InOutSizeY);
 
-					const FMaterialRenderProxy* InMaterialProxy = InMaterial.GetRenderProxy();
-					if (InMaterialProxy)
+					ENQUEUE_RENDER_COMMAND(LoadTiles)(
+						[FeatureLevel, ScreenSpaceSize](FRHICommandListImmediate& RHICmdList)
 					{
-						const ERHIFeatureLevel::Type InFeatureLevel = GMaxRHIFeatureLevel;
-						const FUniformExpressionCache& UniformExpressionCache = InMaterialProxy->UniformExpressionCache[InFeatureLevel];
-						for (IAllocatedVirtualTexture* AllocatedVT : UniformExpressionCache.AllocatedVTs)
-						{
-							if (AllocatedVT)
-							{
-								GetRendererModule().RequestVirtualTextureTilesForRegion(AllocatedVT, InScreenSpaceSize, FIntRect(), -1);
-							}
-						}
-
-						ENQUEUE_RENDER_COMMAND(LoadTiles)(
-							[InFeatureLevel](FRHICommandListImmediate& RHICmdList)
-						{
-							GetRendererModule().LoadPendingVirtualTextureTiles(RHICmdList, InFeatureLevel);
-						});
-					}
+						GetRendererModule().RequestVirtualTextureTiles(ScreenSpaceSize, -1);
+						GetRendererModule().LoadPendingVirtualTextureTiles(RHICmdList, FeatureLevel);
+					});
 
 					FlushRenderingCommands();
 				}
diff --git a/Engine/Source/Runtime/RenderCore/Public/RendererInterface.h b/Engine/Source/Runtime/RenderCore/Public/RendererInterface.h
--- a/Engine/Source/Runtime/RenderCore/Public/RendererInterface.h
+++ b/Engine/Source/Runtime/RenderCore/Public/RendererInterface.h
@@ -695,6 +695,7 @@ public:
 	virtual uint32 RemoveAllVirtualTextureProducerDestroyedCallbacks(const void* Baton) = 0;
 
 	/**	Provided a list of packed virtual texture tile ids, let the VT system request them. Note this should be called as long as the tiles are needed.*/
+	virtual void RequestVirtualTextureTiles(const FVector2D& InScreenSpaceSize, int32 InMipLevel) = 0;
 	virtual void RequestVirtualTextureTilesForRegion(IAllocatedVirtualTexture* AllocatedVT, const FVector2D& InScreenSpaceSize, const FIntRect& InTextureRegion, int32 InMipLevel) = 0;
 
 	/** Ensure that any tiles requested by 'RequestVirtualTextureTilesForRegion' are loaded, must be called from render thread */
diff --git a/Engine/Source/Runtime/Renderer/Private/ReflectionEnvironmentCapture.cpp b/Engine/Source/Runtime/Renderer/Private/ReflectionEnvironmentCapture.cpp
--- a/Engine/Source/Runtime/Renderer/Private/ReflectionEnvironmentCapture.cpp
+++ b/Engine/Source/Runtime/Renderer/Private/ReflectionEnvironmentCapture.cpp
@@ -1494,37 +1494,16 @@ void FScene::CaptureOrUploadReflectionCapture(UReflectionCaptureComponent* Captu
 
 			const int32 ReflectionCaptureSize = UReflectionCaptureComponent::GetReflectionCaptureSize();
 
-			// prefetch VT textures so we have sensible content available
+			// Prefetch all virtual textures so that we have content available
 			if (UseVirtualTexturing(GetFeatureLevel()))
 			{
 				const ERHIFeatureLevel::Type InFeatureLevel = FeatureLevel;
-				const FVector2D InScreenSpaceSize(ReflectionCaptureSize, ReflectionCaptureSize);
-
-				for (TObjectIterator<UMaterialInterface> MaterialIt; MaterialIt; ++MaterialIt)
-				{
-					const FMaterialRenderProxy* MaterialProxy = MaterialIt->GetRenderProxy();
-					if (MaterialProxy)
-					{
-						const FUniformExpressionCache& UniformExpressionCache = MaterialProxy->UniformExpressionCache[InFeatureLevel];
-						for (IAllocatedVirtualTexture* AllocatedVT : UniformExpressionCache.AllocatedVTs)
-						{
-							//todo[vt]: 
-							// It's possible to get here before a RuntimeVT material has been initialized by its virtual texture actor.
-							// In that case AllocatedVT will be nullptr. We can skip the update but it means the VT won't be filled before the capture.
-							if (AllocatedVT)
-							{
-								GetRendererModule().RequestVirtualTextureTilesForRegion(AllocatedVT, InScreenSpaceSize, FIntRect(), -1);
-							}
-						}
-					}
-				}
-
-				//todo[vt]: 
-				// Need to request lightmap VTs as well
+				const FVector2D ScreenSpaceSize(ReflectionCaptureSize, ReflectionCaptureSize);
 
 				ENQUEUE_RENDER_COMMAND(LoadTiles)(
-					[InFeatureLevel](FRHICommandListImmediate& RHICmdList)
+					[InFeatureLevel, ScreenSpaceSize](FRHICommandListImmediate& RHICmdList)
 				{
+					GetRendererModule().RequestVirtualTextureTiles(ScreenSpaceSize, -1);
 					GetRendererModule().LoadPendingVirtualTextureTiles(RHICmdList, InFeatureLevel);
 				});
 
diff --git a/Engine/Source/Runtime/Renderer/Private/RendererModule.h b/Engine/Source/Runtime/Renderer/Private/RendererModule.h
--- a/Engine/Source/Runtime/Renderer/Private/RendererModule.h
+++ b/Engine/Source/Runtime/Renderer/Private/RendererModule.h
@@ -101,6 +101,7 @@ public:
 	virtual void ReleaseVirtualTextureProducer(const FVirtualTextureProducerHandle& Handle) override;
 	virtual void AddVirtualTextureProducerDestroyedCallback(const FVirtualTextureProducerHandle& Handle, FVTProducerDestroyedFunction* Function, void* Baton) override;
 	virtual uint32 RemoveAllVirtualTextureProducerDestroyedCallbacks(const void* Baton) override;
+	virtual void RequestVirtualTextureTiles(const FVector2D& InScreenSpaceSize, int32 InMipLevel) override;
 	virtual void RequestVirtualTextureTilesForRegion(IAllocatedVirtualTexture* AllocatedVT, const FVector2D& InScreenSpaceSize, const FIntRect& InTextureRegion, int32 InMipLevel) override;
 	virtual void LoadPendingVirtualTextureTiles(FRHICommandListImmediate& RHICmdList, ERHIFeatureLevel::Type FeatureLevel) override;
 	virtual void FlushVirtualTextureCache() override;
diff --git a/Engine/Source/Runtime/Renderer/Private/SceneRendering.cpp b/Engine/Source/Runtime/Renderer/Private/SceneRendering.cpp
--- a/Engine/Source/Runtime/Renderer/Private/SceneRendering.cpp
+++ b/Engine/Source/Runtime/Renderer/Private/SceneRendering.cpp
@@ -3579,6 +3579,11 @@ uint32 FRendererModule::RemoveAllVirtualTextureProducerDestroyedCallbacks(const
 	return FVirtualTextureSystem::Get().RemoveAllProducerDestroyedCallbacks(Baton);
 }
 
+void FRendererModule::RequestVirtualTextureTiles(const FVector2D& InScreenSpaceSize, int32 InMipLevel)
+{
+	FVirtualTextureSystem::Get().RequestTiles(InScreenSpaceSize, InMipLevel);
+}
+
 void FRendererModule::RequestVirtualTextureTilesForRegion(IAllocatedVirtualTexture* AllocatedVT, const FVector2D& InScreenSpaceSize, const FIntRect& InTextureRegion, int32 InMipLevel)
 {
 	FVirtualTextureSystem::Get().RequestTilesForRegion(AllocatedVT, InScreenSpaceSize, InTextureRegion, InMipLevel);
diff --git a/Engine/Source/Runtime/Renderer/Private/VT/VirtualTextureSystem.cpp b/Engine/Source/Runtime/Renderer/Private/VT/VirtualTextureSystem.cpp
--- a/Engine/Source/Runtime/Renderer/Private/VT/VirtualTextureSystem.cpp
+++ b/Engine/Source/Runtime/Renderer/Private/VT/VirtualTextureSystem.cpp
@@ -658,6 +658,16 @@ static float ComputeMipLevel(const IAllocatedVirtualTexture* AllocatedVT, const
 	return 0.5f * FMath::Log2(FMath::Max(ppx, ppy));
 }
 
+void FVirtualTextureSystem::RequestTiles(const FVector2D& InScreenSpaceSize, int32 InMipLevel)
+{
+	check(IsInRenderingThread());
+
+	for (auto const Pair : AllocatedVTs)
+	{
+		RequestTilesForRegion(Pair.Value, InScreenSpaceSize, FIntRect(), InMipLevel);
+	}
+}
+
 void FVirtualTextureSystem::RequestTilesForRegion(const IAllocatedVirtualTexture* AllocatedVT, const FVector2D& InScreenSpaceSize, const FIntRect& InTextureRegion, int32 InMipLevel)
 {
 	FIntRect TextureRegion(InTextureRegion);
diff --git a/Engine/Source/Runtime/Renderer/Private/VT/VirtualTextureSystem.h b/Engine/Source/Runtime/Renderer/Private/VT/VirtualTextureSystem.h
--- a/Engine/Source/Runtime/Renderer/Private/VT/VirtualTextureSystem.h
+++ b/Engine/Source/Runtime/Renderer/Private/VT/VirtualTextureSystem.h
@@ -59,6 +59,7 @@ public:
 	void LockTile(const FVirtualTextureLocalTile& Tile);
 	void UnlockTile(const FVirtualTextureLocalTile& Tile, const FVirtualTextureProducer* Producer);
 	void ForceUnlockAllTiles(const FVirtualTextureProducerHandle& ProducerHandle, const FVirtualTextureProducer* Producer);
+	void RequestTiles(const FVector2D& InScreenSpaceSize, int32 InMipLevel = -1);
 	void RequestTilesForRegion(const IAllocatedVirtualTexture* AllocatedVT, const FVector2D& InScreenSpaceSize, const FIntRect& InTextureRegion, int32 InMipLevel = -1);
 	void LoadPendingTiles(FRHICommandListImmediate& RHICmdList, ERHIFeatureLevel::Type FeatureLevel);
 	void FlushCache();
