Integration: Fixed always allowed translucent hit proxies to be selectable even when translucency selection is off.

diff --git a/Engine/Source/Runtime/Renderer/Private/SceneHitProxyRendering.cpp b/Engine/Source/Runtime/Renderer/Private/SceneHitProxyRendering.cpp
--- a/Engine/Source/Runtime/Renderer/Private/SceneHitProxyRendering.cpp
+++ b/Engine/Source/Runtime/Renderer/Private/SceneHitProxyRendering.cpp
@@ -338,11 +338,7 @@ static void DoRenderHitProxies(FRHICommandListImmediate& RHICmdList, const FScen
 			for (int32 MeshIndex = 0; MeshIndex < View.DynamicEditorMeshElements.Num(); MeshIndex++)
 			{
 				const FMeshBatchAndRelevance& MeshBatchAndRelevance = View.DynamicEditorMeshElements[MeshIndex];
-
-				if (MeshBatchAndRelevance.Mesh->bSelectable)
-				{
-					PassMeshProcessor.AddMeshBatch(*MeshBatchAndRelevance.Mesh, DefaultBatchElementMask, MeshBatchAndRelevance.PrimitiveSceneProxy);
-				}
+				PassMeshProcessor.AddMeshBatch(*MeshBatchAndRelevance.Mesh, DefaultBatchElementMask, MeshBatchAndRelevance.PrimitiveSceneProxy);
 			}
 		});
 
@@ -608,11 +604,6 @@ void FDeferredShadingSceneRenderer::RenderHitProxies(FRHICommandListImmediate& R
 
 void FHitProxyMeshProcessor::AddMeshBatch(const FMeshBatch& RESTRICT MeshBatch, uint64 BatchElementMask, const FPrimitiveSceneProxy* RESTRICT PrimitiveSceneProxy, int32 StaticMeshId)
 {
-	if (ViewIfDynamicMeshCommand && ViewIfDynamicMeshCommand->bAllowTranslucentPrimitivesInHitProxy != bAllowTranslucentPrimitivesInHitProxy)
-	{
-		return;
-	}
-
 	if (MeshBatch.bUseForMaterial && MeshBatch.bSelectable && Scene->RequiresHitProxies() && (!PrimitiveSceneProxy || PrimitiveSceneProxy->IsSelectable()))
 	{
 		// Determine the mesh's material and blend mode.
@@ -636,7 +627,29 @@ void FHitProxyMeshProcessor::AddMeshBatch(const FMeshBatch& RESTRICT MeshBatch,
 
 		check(Material && MaterialRenderProxy);
 
-		if (bAllowTranslucentPrimitivesInHitProxy || !IsTranslucentBlendMode(BlendMode))
+		bool bAddTranslucentPrimitive = bAllowTranslucentPrimitivesInHitProxy;
+
+		// Check whether the primitive overrides the pass to force translucent hit proxies.
+		if (!bAddTranslucentPrimitive)
+		{
+			FHitProxyId HitProxyId = MeshBatch.BatchHitProxyId;
+
+			// Fallback to the primitive default hit proxy id if the mesh batch doesn't have one.
+			if (MeshBatch.BatchHitProxyId == FHitProxyId() && PrimitiveSceneProxy)
+			{
+				if (const FPrimitiveSceneInfo* PrimitiveSceneInfo = PrimitiveSceneProxy->GetPrimitiveSceneInfo())
+				{
+					HitProxyId = PrimitiveSceneInfo->DefaultDynamicHitProxyId;
+				}
+			}
+
+			if (const HHitProxy* HitProxy = GetHitProxyById(HitProxyId))
+			{
+				bAddTranslucentPrimitive = HitProxy->AlwaysAllowsTranslucentPrimitives();
+			}
+		}
+
+		if (bAddTranslucentPrimitive || !IsTranslucentBlendMode(BlendMode))
 		{
 			Process(MeshBatch, BatchElementMask, StaticMeshId, PrimitiveSceneProxy, *MaterialRenderProxy, *Material, MeshFillMode, MeshCullMode);
 		}
