Mesh commands generated via DrawDynamicMeshPass now apply view overrides and take view's cull mode overrides in to account.

diff --git a/Engine/Source/Runtime/Renderer/Private/MeshPassProcessor.cpp b/Engine/Source/Runtime/Renderer/Private/MeshPassProcessor.cpp
--- a/Engine/Source/Runtime/Renderer/Private/MeshPassProcessor.cpp
+++ b/Engine/Source/Runtime/Renderer/Private/MeshPassProcessor.cpp
@@ -1058,6 +1058,51 @@ void SubmitMeshDrawCommandsRange(
 	}
 }
 
+void ApplyViewOverridesToMeshDrawCommands(const FSceneView& View, FMeshCommandOneFrameArray& VisibleMeshDrawCommands, FDynamicMeshDrawCommandStorage& DynamicMeshDrawCommandStorage, FGraphicsMinimalPipelineStateSet& GraphicsMinimalPipelineStateSet)
+{
+	if (View.bReverseCulling || View.bRenderSceneTwoSided)
+	{
+		const FMeshCommandOneFrameArray& PassVisibleMeshDrawCommands = VisibleMeshDrawCommands;
+
+		FMeshCommandOneFrameArray ViewOverriddenMeshCommands;
+		ViewOverriddenMeshCommands.Empty(PassVisibleMeshDrawCommands.Num());
+
+		for (int32 MeshCommandIndex = 0; MeshCommandIndex < PassVisibleMeshDrawCommands.Num(); MeshCommandIndex++)
+		{
+			DynamicMeshDrawCommandStorage.MeshDrawCommands.Add(1);
+			FMeshDrawCommand& NewMeshCommand = DynamicMeshDrawCommandStorage.MeshDrawCommands[DynamicMeshDrawCommandStorage.MeshDrawCommands.Num() - 1];
+
+			const FVisibleMeshDrawCommand& VisibleMeshDrawCommand = PassVisibleMeshDrawCommands[MeshCommandIndex];
+			const FMeshDrawCommand& MeshCommand = *VisibleMeshDrawCommand.MeshDrawCommand;
+			NewMeshCommand = MeshCommand;
+
+			const ERasterizerCullMode LocalCullMode = View.bRenderSceneTwoSided ? CM_None : View.bReverseCulling ? FMeshPassProcessor::InverseCullMode(VisibleMeshDrawCommand.MeshCullMode) : VisibleMeshDrawCommand.MeshCullMode;
+
+			FGraphicsMinimalPipelineStateInitializer PipelineState = MeshCommand.CachedPipelineId.GetPipelineState(GraphicsMinimalPipelineStateSet);
+			PipelineState.RasterizerState = GetStaticRasterizerState<true>(VisibleMeshDrawCommand.MeshFillMode, LocalCullMode);
+
+			const FGraphicsMinimalPipelineStateId PipelineId = FGraphicsMinimalPipelineStateId::GetPipelineStateId(PipelineState, GraphicsMinimalPipelineStateSet);
+			NewMeshCommand.Finalize(PipelineId, nullptr);
+
+			FVisibleMeshDrawCommand NewVisibleMeshDrawCommand;
+
+			NewVisibleMeshDrawCommand.Setup(
+				&NewMeshCommand,
+				VisibleMeshDrawCommand.DrawPrimitiveId,
+				VisibleMeshDrawCommand.ScenePrimitiveId,
+				VisibleMeshDrawCommand.StateBucketId,
+				VisibleMeshDrawCommand.MeshFillMode,
+				VisibleMeshDrawCommand.MeshCullMode,
+				VisibleMeshDrawCommand.SortKey);
+
+			ViewOverriddenMeshCommands.Add(NewVisibleMeshDrawCommand);
+		}
+
+		// Replace VisibleMeshDrawCommands
+		FMemory::Memswap(&VisibleMeshDrawCommands, &ViewOverriddenMeshCommands, sizeof(ViewOverriddenMeshCommands));
+	}
+}
+
 void DrawDynamicMeshPassPrivate(
 	const FSceneView& View,
 	FRHICommandList& RHICmdList,
@@ -1072,6 +1117,7 @@ void DrawDynamicMeshPassPrivate(
 
 		FRHIVertexBuffer* PrimitiveIdVertexBuffer = nullptr;
 
+		ApplyViewOverridesToMeshDrawCommands(View, VisibleMeshDrawCommands, DynamicMeshDrawCommandStorage, GraphicsMinimalPipelineStateSet);
 		SortAndMergeDynamicPassMeshDrawCommands(View.GetFeatureLevel(), VisibleMeshDrawCommands, DynamicMeshDrawCommandStorage, PrimitiveIdVertexBuffer, InstanceFactor);
 
 		SubmitMeshDrawCommandsRange(VisibleMeshDrawCommands, GraphicsMinimalPipelineStateSet, PrimitiveIdVertexBuffer, 0, bDynamicInstancing, 0, VisibleMeshDrawCommands.Num(), InstanceFactor, RHICmdList);
@@ -1137,14 +1183,10 @@ ERasterizerCullMode FMeshPassProcessor::ComputeMeshCullMode(const FMeshBatch& Me
 	const FMeshDrawingPolicyOverrideSettings InOverrideSettings = ComputeMeshOverrideSettings(Mesh);
 	const bool bMaterialResourceIsTwoSided = InMaterialResource.IsTwoSided();
 	const bool bInTwoSidedOverride = !!(InOverrideSettings.MeshOverrideFlags & EDrawingPolicyOverrideFlags::TwoSided);
-	const bool bViewTwoSidedOverride = ViewIfDynamicMeshCommand ? ViewIfDynamicMeshCommand->bRenderSceneTwoSided : false;
-	const bool bViewReverseCullingOverride = ViewIfDynamicMeshCommand ? ViewIfDynamicMeshCommand->bReverseCulling : false;
 	const bool bInReverseCullModeOverride = !!(InOverrideSettings.MeshOverrideFlags & EDrawingPolicyOverrideFlags::ReverseCullMode);
-	const bool bReverseCullMode = bViewReverseCullingOverride ? !bInReverseCullModeOverride : bInReverseCullModeOverride;
 	const bool bIsTwoSided = (bMaterialResourceIsTwoSided || bInTwoSidedOverride);
-	const bool bMeshRenderTwoSided = bIsTwoSided || bInTwoSidedOverride || bViewTwoSidedOverride;
-
-	return bMeshRenderTwoSided ? CM_None : (bReverseCullMode ? CM_CCW : CM_CW);
+	const bool bMeshRenderTwoSided = bIsTwoSided || bInTwoSidedOverride;
+	return bMeshRenderTwoSided ? CM_None : (bInReverseCullModeOverride ? CM_CCW : CM_CW);
 }
 
 void FMeshPassProcessor::GetDrawCommandPrimitiveId(
diff --git a/Engine/Source/Runtime/Renderer/Private/ShadowRendering.h b/Engine/Source/Runtime/Renderer/Private/ShadowRendering.h
--- a/Engine/Source/Runtime/Renderer/Private/ShadowRendering.h
+++ b/Engine/Source/Runtime/Renderer/Private/ShadowRendering.h
@@ -439,8 +439,6 @@ public:
 
 	void SetupMeshDrawCommandsForProjectionStenciling(FSceneRenderer& Renderer);
 
-	void ApplyViewOverridesToMeshDrawCommands(const FViewInfo& View, FMeshCommandOneFrameArray& VisibleMeshDrawCommands);
-
 	/** 
 	 * @param View view to check visibility in
 	 * @return true if this shadow info has any subject prims visible in the view
diff --git a/Engine/Source/Runtime/Renderer/Private/ShadowSetup.cpp b/Engine/Source/Runtime/Renderer/Private/ShadowSetup.cpp
--- a/Engine/Source/Runtime/Renderer/Private/ShadowSetup.cpp
+++ b/Engine/Source/Runtime/Renderer/Private/ShadowSetup.cpp
@@ -1319,7 +1319,7 @@ void FProjectedShadowInfo::SetupMeshDrawCommandsForProjectionStenciling(FSceneRe
 				}
 			}
 
-			ApplyViewOverridesToMeshDrawCommands(View, ProjectionStencilingPass.VisibleMeshDrawCommands);
+			ApplyViewOverridesToMeshDrawCommands(View, ProjectionStencilingPass.VisibleMeshDrawCommands, DynamicMeshDrawCommandStorage, GraphicsMinimalPipelineStateSet);
 
 			// If instanced stereo is enabled, we need to render each view of the stereo pair using the instanced stereo transform to avoid bias issues.
 			// TODO: Support instanced stereo properly in the projection stenciling pass.
@@ -1329,51 +1329,6 @@ void FProjectedShadowInfo::SetupMeshDrawCommandsForProjectionStenciling(FSceneRe
 	}
 }
 
-void FProjectedShadowInfo::ApplyViewOverridesToMeshDrawCommands(const FViewInfo& View, FMeshCommandOneFrameArray& VisibleMeshDrawCommands)
-{
-	if (View.bReverseCulling || View.bRenderSceneTwoSided)
-	{
-		const FMeshCommandOneFrameArray& PassVisibleMeshDrawCommands = VisibleMeshDrawCommands;
-
-		FMeshCommandOneFrameArray ViewOverriddenMeshCommands;
-		ViewOverriddenMeshCommands.Empty(PassVisibleMeshDrawCommands.Num());
-
-		for (int32 MeshCommandIndex = 0; MeshCommandIndex < PassVisibleMeshDrawCommands.Num(); MeshCommandIndex++)
-		{
-			DynamicMeshDrawCommandStorage.MeshDrawCommands.Add(1);
-			FMeshDrawCommand& NewMeshCommand = DynamicMeshDrawCommandStorage.MeshDrawCommands[DynamicMeshDrawCommandStorage.MeshDrawCommands.Num() - 1];
-
-			const FVisibleMeshDrawCommand& VisibleMeshDrawCommand = PassVisibleMeshDrawCommands[MeshCommandIndex];
-			const FMeshDrawCommand& MeshCommand = *VisibleMeshDrawCommand.MeshDrawCommand;
-			NewMeshCommand = MeshCommand;
-
-			const ERasterizerCullMode LocalCullMode = View.bRenderSceneTwoSided ? CM_None : View.bReverseCulling ? FMeshPassProcessor::InverseCullMode(VisibleMeshDrawCommand.MeshCullMode) : VisibleMeshDrawCommand.MeshCullMode;
-
-			FGraphicsMinimalPipelineStateInitializer PipelineState = MeshCommand.CachedPipelineId.GetPipelineState(GraphicsMinimalPipelineStateSet);
-			PipelineState.RasterizerState = GetStaticRasterizerState<true>(VisibleMeshDrawCommand.MeshFillMode, LocalCullMode);
-
-			const FGraphicsMinimalPipelineStateId PipelineId = FGraphicsMinimalPipelineStateId::GetPipelineStateId(PipelineState, GraphicsMinimalPipelineStateSet);
-			NewMeshCommand.Finalize(PipelineId, nullptr);
-
-			FVisibleMeshDrawCommand NewVisibleMeshDrawCommand;
-
-			NewVisibleMeshDrawCommand.Setup(
-				&NewMeshCommand,
-				VisibleMeshDrawCommand.DrawPrimitiveId,
-				VisibleMeshDrawCommand.ScenePrimitiveId,
-				VisibleMeshDrawCommand.StateBucketId,
-				VisibleMeshDrawCommand.MeshFillMode,
-				VisibleMeshDrawCommand.MeshCullMode,
-				VisibleMeshDrawCommand.SortKey);
-
-			ViewOverriddenMeshCommands.Add(NewVisibleMeshDrawCommand);
-		}
-
-		// Replace VisibleMeshDrawCommands
-		FMemory::Memswap(&VisibleMeshDrawCommands, &ViewOverriddenMeshCommands, sizeof(ViewOverriddenMeshCommands));
-	}
-}
-
 void FProjectedShadowInfo::GatherDynamicMeshElements(FSceneRenderer& Renderer, FVisibleLightInfo& VisibleLightInfo, TArray<const FSceneView*>& ReusedViewsArray,
 	FGlobalDynamicIndexBuffer& DynamicIndexBuffer, FGlobalDynamicVertexBuffer& DynamicVertexBuffer, FGlobalDynamicReadBuffer& DynamicReadBuffer)
 {
diff --git a/Engine/Source/Runtime/Renderer/Private/TranslucentLighting.cpp b/Engine/Source/Runtime/Renderer/Private/TranslucentLighting.cpp
--- a/Engine/Source/Runtime/Renderer/Private/TranslucentLighting.cpp
+++ b/Engine/Source/Runtime/Renderer/Private/TranslucentLighting.cpp
@@ -620,7 +620,7 @@ void FProjectedShadowInfo::RenderTranslucencyDepths(FRHICommandList& RHICmdList,
 			const bool bDynamicInstancing = IsDynamicInstancingEnabled(ShadowDepthView->FeatureLevel);
 
 			FRHIVertexBuffer* PrimitiveIdVertexBuffer = nullptr;
-			ApplyViewOverridesToMeshDrawCommands(*ShadowDepthView, VisibleMeshDrawCommands);
+			ApplyViewOverridesToMeshDrawCommands(*ShadowDepthView, VisibleMeshDrawCommands, DynamicMeshDrawCommandStorage, GraphicsMinimalPipelineStateSet);
 			SortAndMergeDynamicPassMeshDrawCommands(SceneRenderer->FeatureLevel, VisibleMeshDrawCommands, DynamicMeshDrawCommandStorage, PrimitiveIdVertexBuffer, 1);
 			SubmitMeshDrawCommands(VisibleMeshDrawCommands, GraphicsMinimalPipelineStateSet, PrimitiveIdVertexBuffer, 0, bDynamicInstancing, 1, RHICmdList);
 		}
diff --git a/Engine/Source/Runtime/Renderer/Public/MeshPassProcessor.h b/Engine/Source/Runtime/Renderer/Public/MeshPassProcessor.h
--- a/Engine/Source/Runtime/Renderer/Public/MeshPassProcessor.h
+++ b/Engine/Source/Runtime/Renderer/Public/MeshPassProcessor.h
@@ -1204,6 +1204,12 @@ extern void SubmitMeshDrawCommandsRange(
 	uint32 InstanceFactor,
 	FRHICommandList& RHICmdList);
 
+extern void ApplyViewOverridesToMeshDrawCommands(
+	const FSceneView& View,
+	FMeshCommandOneFrameArray& VisibleMeshDrawCommands,
+	FDynamicMeshDrawCommandStorage& DynamicMeshDrawCommandStorage,
+	FGraphicsMinimalPipelineStateSet& GraphicsMinimalPipelineStateSet);
+
 RENDERER_API extern void DrawDynamicMeshPassPrivate(
 	const FSceneView& View,
 	FRHICommandList& RHICmdList,
