Fix for Corrupted motion blur

diff --git a/Engine/Source/Runtime/Renderer/Private/PostProcess/PostProcessMotionBlur.cpp b/Engine/Source/Runtime/Renderer/Private/PostProcess/PostProcessMotionBlur.cpp
--- a/Engine/Source/Runtime/Renderer/Private/PostProcess/PostProcessMotionBlur.cpp
+++ b/Engine/Source/Runtime/Renderer/Private/PostProcess/PostProcessMotionBlur.cpp
@@ -530,7 +530,7 @@ void ComputeMotionBlurVelocity(
 		PassParameters->RenderTargets.DepthStencil =
 			FDepthStencilBinding(
 				VelocityTileDepthTexture,
-				ERenderTargetLoadAction::ENoAction,
+				ERenderTargetLoadAction::EClear,
 				ERenderTargetStoreAction::EStore,
 				ERenderTargetLoadAction::ENoAction,
 				ERenderTargetStoreAction::ENoAction,
