Removing redundant mesh draw command invalidation. This is no longer required as we now do not recreate ILC uniform buffer at runtime - instead we always update the existing one. Redundant invalidation was causing performance regression on mobile when GPUScene is disabled:
https://udn.unrealengine.com/questions/517462/bupdatestaticdrawlists-performance-problem.html

diff --git a/Engine/Source/Runtime/Renderer/Private/RendererScene.cpp b/Engine/Source/Runtime/Renderer/Private/RendererScene.cpp
--- a/Engine/Source/Runtime/Renderer/Private/RendererScene.cpp
+++ b/Engine/Source/Runtime/Renderer/Private/RendererScene.cpp
@@ -1290,8 +1290,7 @@ void FScene::UpdatePrimitiveTransform_RenderThread(FRHICommandListImmediate& RHI
 
 	FPrimitiveSceneInfo* PrimitiveSceneInfo = PrimitiveSceneProxy->GetPrimitiveSceneInfo();
 
-	const bool bUpdateStaticDrawLists = !PrimitiveSceneProxy->StaticElementsAlwaysUseProxyPrimitiveUniformBuffer() 
-		|| !UseGPUScene(GMaxRHIShaderPlatform, GetFeatureLevel());
+	const bool bUpdateStaticDrawLists = !PrimitiveSceneProxy->StaticElementsAlwaysUseProxyPrimitiveUniformBuffer();
 
 	// Remove the primitive from the scene at its old location
 	// (note that the octree update relies on the bounds not being modified yet).
