Runtime virtual texture flush per VT and region instead of flushing everything

diff --git a/Engine/Source/Editor/LandscapeEditor/Private/LandscapeEdModePaintTools.cpp b/Engine/Source/Editor/LandscapeEditor/Private/LandscapeEdModePaintTools.cpp
--- a/Engine/Source/Editor/LandscapeEditor/Private/LandscapeEdModePaintTools.cpp
+++ b/Engine/Source/Editor/LandscapeEditor/Private/LandscapeEdModePaintTools.cpp
@@ -259,14 +259,17 @@ public:
 		this->Cache.SetCachedData(X1, Y1, X2, Y2, Data, UISettings->PaintingRestriction);
 		this->Cache.Flush();
 		
+		// If we render to a runtime virtual texture then we mark touched components as dirty to trigger updates
+		//todo[vt]: Would be more efficient to update VT in a single flush instead of one flush per component. Also dirtying all render state is a bit heavyweight.
 		ALandscape* Landscape = LandscapeInfo->LandscapeActor.Get();
-
-		// If we render to a runtime virtual texture then we need to flush here
 		if (Landscape != nullptr && Landscape->RuntimeVirtualTextures.Num() > 0)
 		{
-			//todo[vt]: Only flush Bounds 
-			//todo[vt]: Only flush specific virtual textures
-			GetRendererModule().FlushVirtualTextureCache();
+			TSet<ULandscapeComponent*> Components;
+			LandscapeInfo->GetComponentsInRegion(X1 + 1, Y1 + 1, X2 - 1, Y2 - 1, Components);
+			for (ULandscapeComponent* Component : Components)
+			{
+				Component->MarkRenderStateDirty();
+			}
 		}
 	}
 };
diff --git a/Engine/Source/Runtime/Engine/Private/PrimitiveSceneProxy.cpp b/Engine/Source/Runtime/Engine/Private/PrimitiveSceneProxy.cpp
--- a/Engine/Source/Runtime/Engine/Private/PrimitiveSceneProxy.cpp
+++ b/Engine/Source/Runtime/Engine/Private/PrimitiveSceneProxy.cpp
@@ -219,7 +219,7 @@ FPrimitiveSceneProxy::FPrimitiveSceneProxy(const UPrimitiveComponent* InComponen
 	bRequiresVisibleLevelToRender = (ComponentLevel && ComponentLevel->bRequireFullVisibilityToRender);
 	bIsComponentLevelVisible = (!ComponentLevel || ComponentLevel->bIsVisible);
 
-	// Setup the runtime virtual texture information and flush the virtual texture if necessary
+	// Setup the runtime virtual texture information
 	if (UseVirtualTexturing(GetScene().GetFeatureLevel()))
 	{
 		for (URuntimeVirtualTexture* VirtualTexture : InComponent->GetRuntimeVirtualTextures())
@@ -228,10 +228,6 @@ FPrimitiveSceneProxy::FPrimitiveSceneProxy(const UPrimitiveComponent* InComponen
 			{
 				RuntimeVirtualTextures.Add(VirtualTexture);
 				RuntimeVirtualTextureMaterialTypes.AddUnique(VirtualTexture->GetMaterialType());
-			
-				//todo[vt]: Only flush this specific virtual texture
-				//todo[vt]: Only flush primitive bounds 
-				GetRendererModule().FlushVirtualTextureCache();
 			}
 		}
 	}
@@ -264,13 +260,6 @@ void FPrimitiveSceneProxy::SetUsedMaterialForVerification(const TArray<UMaterial
 FPrimitiveSceneProxy::~FPrimitiveSceneProxy()
 {
 	check(IsInRenderingThread());
-
-	for (URuntimeVirtualTexture* VirtualTexture : RuntimeVirtualTextures)
-	{
-		//todo[vt]: Only flush Bounds 
-		//todo[vt]: Only flush specific virtual textures
-		GetRendererModule().FlushVirtualTextureCache();
-	}
 }
 
 HHitProxy* FPrimitiveSceneProxy::CreateHitProxies(UPrimitiveComponent* Component,TArray<TRefCountPtr<HHitProxy> >& OutHitProxies)
diff --git a/Engine/Source/Runtime/Engine/Public/PrimitiveSceneProxy.h b/Engine/Source/Runtime/Engine/Public/PrimitiveSceneProxy.h
--- a/Engine/Source/Runtime/Engine/Public/PrimitiveSceneProxy.h
+++ b/Engine/Source/Runtime/Engine/Public/PrimitiveSceneProxy.h
@@ -501,6 +501,7 @@ public:
 	inline bool CastsStaticShadow() const { return bCastStaticShadow; }
 	inline bool CastsDynamicShadow() const { return bCastDynamicShadow; }
 	inline bool WritesVirtualTexture() const{ return RuntimeVirtualTextures.Num() > 0; }
+	inline bool WritesVirtualTexture(URuntimeVirtualTexture* VirtualTexture) const { return RuntimeVirtualTextures.Find(VirtualTexture) != INDEX_NONE; }
 	inline bool AffectsDynamicIndirectLighting() const { return bAffectDynamicIndirectLighting; }
 	inline bool AffectsDistanceFieldLighting() const { return bAffectDistanceFieldLighting; }
 	inline float GetLpvBiasMultiplier() const { return LpvBiasMultiplier; }
diff --git a/Engine/Source/Runtime/Landscape/Private/LandscapeEdit.cpp b/Engine/Source/Runtime/Landscape/Private/LandscapeEdit.cpp
--- a/Engine/Source/Runtime/Landscape/Private/LandscapeEdit.cpp
+++ b/Engine/Source/Runtime/Landscape/Private/LandscapeEdit.cpp
@@ -606,13 +606,6 @@ void ULandscapeComponent::PostEditUndo()
 		Components.Add(this);
 		GetLandscapeProxy()->FlushGrassComponents(&Components);
 	}
-
-	if (GetLandscapeProxy()->RuntimeVirtualTextures.Num() > 0)
-	{
-		//todo[vt]: Only flush this specific virtual textures
-		//todo[vt]: Only flush Bounds 
-		GetRendererModule().FlushVirtualTextureCache();
-	}
 }
 
 void ALandscapeProxy::FixupWeightmaps()
diff --git a/Engine/Source/Runtime/Renderer/Private/DeferredShadingRenderer.cpp b/Engine/Source/Runtime/Renderer/Private/DeferredShadingRenderer.cpp
--- a/Engine/Source/Runtime/Renderer/Private/DeferredShadingRenderer.cpp
+++ b/Engine/Source/Runtime/Renderer/Private/DeferredShadingRenderer.cpp
@@ -948,6 +948,7 @@ void FDeferredShadingSceneRenderer::Render(FRHICommandListImmediate& RHICmdList)
 
 	if (bUseVirtualTexturing)
 	{
+		Scene->FlushDirtyRuntimeVirtualTextures();
 		FVirtualTextureSystem::Get().Update(RHICmdList, FeatureLevel);
 	}
 
diff --git a/Engine/Source/Runtime/Renderer/Private/PrimitiveSceneInfo.cpp b/Engine/Source/Runtime/Renderer/Private/PrimitiveSceneInfo.cpp
--- a/Engine/Source/Runtime/Renderer/Private/PrimitiveSceneInfo.cpp
+++ b/Engine/Source/Runtime/Renderer/Private/PrimitiveSceneInfo.cpp
@@ -13,6 +13,7 @@
 #include "ScenePrivate.h"
 #include "HAL/LowLevelMemTracker.h"
 #include "RayTracing/RayTracingMaterialHitShaders.h"
+#include "VT/RuntimeVirtualTextureSceneProxy.h"
 #include "VT/VirtualTextureSystem.h"
 #include "GPUScene.h"
 
@@ -636,13 +637,13 @@ void FPrimitiveSceneInfo::AddToScene(FRHICommandListImmediate& RHICmdList, bool
 	// Store the component.
 	Scene->PrimitiveComponentIds[PackedIndex] = PrimitiveComponentId;
 
-	// Store the runtime virtual texture flags and lod info.
-	const bool bRenderToVirtualTexture = Proxy->WritesVirtualTexture();
-	Scene->PrimitiveVirtualTextureFlags[PackedIndex].bRenderToVirtualTexture = bRenderToVirtualTexture;
-	if (bRenderToVirtualTexture)
-	{
-		Scene->PrimitiveVirtualTextureFlags[PackedIndex].RuntimeVirtualTextureMask = Scene->GetRuntimeVirtualTextureMask(Proxy);
+	// Store the runtime virtual texture flags.
+	UpdateRuntimeVirtualTextureFlags();
+	Scene->PrimitiveVirtualTextureFlags[PackedIndex] = RuntimeVirtualTextureFlags;
 
+	// Store the runtime virtual texture Lod info.
+	if (RuntimeVirtualTextureFlags.bRenderToVirtualTexture)
+	{
 		int8 MinLod, MaxLod;
 		GetRuntimeVirtualTextureLODRange(StaticMeshRelevances, MinLod, MaxLod);
 
@@ -740,6 +741,31 @@ void FPrimitiveSceneInfo::RemoveFromScene(bool bUpdateStaticDrawLists)
 	}
 }
 
+void FPrimitiveSceneInfo::UpdateRuntimeVirtualTextureFlags()
+{
+	const bool bRenderToVirtualTexture = Proxy->WritesVirtualTexture();
+
+	RuntimeVirtualTextureFlags.bRenderToVirtualTexture = bRenderToVirtualTexture;
+	RuntimeVirtualTextureFlags.RuntimeVirtualTextureMask = 0;
+
+	if (bRenderToVirtualTexture)
+	{
+		// Performance assumption: The arrays of runtime virtual textures are small (less that 5?) so that O(n^2) scan isn't expensive
+		for (TSparseArray<FRuntimeVirtualTextureSceneProxy*>::TConstIterator It(Scene->RuntimeVirtualTextures); It; ++It)
+		{
+			int32 SceneIndex = It.GetIndex();
+			if (SceneIndex < FPrimitiveVirtualTextureFlags::RuntimeVirtualTexture_BitCount)
+			{
+				URuntimeVirtualTexture* SceneVirtualTexture = (*It)->VirtualTexture;
+				if (Proxy->WritesVirtualTexture(SceneVirtualTexture))
+				{
+					RuntimeVirtualTextureFlags.RuntimeVirtualTextureMask |= 1 << SceneIndex;
+				}
+			}
+		}
+	}
+}
+
 bool FPrimitiveSceneInfo::NeedsUpdateStaticMeshes()
 {
 	return Scene->PrimitivesNeedingStaticMeshUpdate[PackedIndex];
@@ -791,6 +817,24 @@ void FPrimitiveSceneInfo::BeginDeferredUpdateStaticMeshesWithoutVisibilityCheck(
 	}
 }
 
+void FPrimitiveSceneInfo::FlushRuntimeVirtualTexture()
+{
+	if (RuntimeVirtualTextureFlags.bRenderToVirtualTexture)
+	{
+		uint32 RuntimeVirtualTextureIndex = 0;
+		uint32 Mask = RuntimeVirtualTextureFlags.RuntimeVirtualTextureMask;
+		while (Mask != 0)
+		{
+			if (Mask & 1)
+			{
+				Scene->RuntimeVirtualTextures[RuntimeVirtualTextureIndex]->Dirty(Proxy->GetBounds());
+			}
+			Mask >>= 1;
+			RuntimeVirtualTextureIndex++;
+		}
+	}
+}
+
 void FPrimitiveSceneInfo::LinkLODParentComponent()
 {
 	if (LODParentComponentId.IsValid())
diff --git a/Engine/Source/Runtime/Renderer/Private/RendererScene.cpp b/Engine/Source/Runtime/Renderer/Private/RendererScene.cpp
--- a/Engine/Source/Runtime/Renderer/Private/RendererScene.cpp
+++ b/Engine/Source/Runtime/Renderer/Private/RendererScene.cpp
@@ -907,6 +907,9 @@ void FScene::AddPrimitiveSceneInfo_RenderThread(FRHICommandListImmediate& RHICmd
 
 	DistanceFieldSceneData.AddPrimitive(PrimitiveSceneInfo);
 
+	// Flush virtual textures touched by primitive
+	PrimitiveSceneInfo->FlushRuntimeVirtualTexture();
+
 	// LOD Parent, if this is LOD parent, we should update Proxy Scene Info
 	// LOD parent gets removed WHEN no children is accessing
 	// LOD parent can be recreated as scene updates
@@ -1292,6 +1295,8 @@ void FScene::UpdatePrimitiveTransform_RenderThread(FRHICommandListImmediate& RHI
 
 	const bool bUpdateStaticDrawLists = !PrimitiveSceneProxy->StaticElementsAlwaysUseProxyPrimitiveUniformBuffer();
 
+	PrimitiveSceneInfo->FlushRuntimeVirtualTexture();
+
 	// Remove the primitive from the scene at its old location
 	// (note that the octree update relies on the bounds not being modified yet).
 	PrimitiveSceneInfo->RemoveFromScene(bUpdateStaticDrawLists);
@@ -1324,6 +1329,8 @@ void FScene::UpdatePrimitiveTransform_RenderThread(FRHICommandListImmediate& RHI
 
 	// Re-add the primitive to the scene with the new transform.
 	PrimitiveSceneInfo->AddToScene(RHICmdList, bUpdateStaticDrawLists);
+
+	PrimitiveSceneInfo->FlushRuntimeVirtualTexture();
 }
 
 void FScene::UpdatePrimitiveTransform(UPrimitiveComponent* Primitive)
@@ -1640,6 +1647,9 @@ void FScene::RemovePrimitiveSceneInfo_RenderThread(FPrimitiveSceneInfo* Primitiv
 	// Unlink the LOD parent info if valid
 	PrimitiveSceneInfo->UnlinkLODParentComponent();
 
+	// Flush virtual textures touched by primitive
+	PrimitiveSceneInfo->FlushRuntimeVirtualTexture();
+
 	// Remove the primitive from the scene.
 	PrimitiveSceneInfo->RemoveFromScene(true);
 
@@ -2533,7 +2543,8 @@ void FScene::UpdateRuntimeVirtualTextureForAllPrimitives_RenderThread()
 	{
 		if (PrimitiveVirtualTextureFlags[Index].bRenderToVirtualTexture)
 		{
-			PrimitiveVirtualTextureFlags[Index].RuntimeVirtualTextureMask = GetRuntimeVirtualTextureMask(PrimitiveSceneProxies[Index]);
+			Primitives[Index]->UpdateRuntimeVirtualTextureFlags();
+			PrimitiveVirtualTextureFlags[Index] = Primitives[Index]->GetRuntimeVirtualTextureFlags();
 		}
 	}
 }
@@ -2553,25 +2564,13 @@ uint32 FScene::GetRuntimeVirtualTextureSceneIndex(uint32 ProducerId)
 	return 0;
 }
 
-uint32 FScene::GetRuntimeVirtualTextureMask(FPrimitiveSceneProxy const* Proxy)
+void FScene::FlushDirtyRuntimeVirtualTextures()
 {
-	uint32 Mask = 0;
-	for (TSparseArray<FRuntimeVirtualTextureSceneProxy*>::TConstIterator It(RuntimeVirtualTextures); It; ++It)
+	checkSlow(IsInRenderingThread());
+	for (TSparseArray<FRuntimeVirtualTextureSceneProxy*>::TIterator It(RuntimeVirtualTextures); It; ++It)
 	{
-		int32 SceneIndex = It.GetIndex();
-		if (SceneIndex < FPrimitiveVirtualTextureFlags::RuntimeVirtualTexture_BitCount)
-		{
-			URuntimeVirtualTexture* SceneVirtualTexture = (*It)->VirtualTexture;
-			for (URuntimeVirtualTexture* PrimitiveVirtualTexture : Proxy->RuntimeVirtualTextures)
-			{
-				if (SceneVirtualTexture == PrimitiveVirtualTexture)
-				{
-					Mask |= 1 << SceneIndex;
-				}
-			}
-		}
+		(*It)->FlushDirtyPages();
 	}
-	return Mask;
 }
 
 bool FScene::GetPreviousLocalToWorld(const FPrimitiveSceneInfo* PrimitiveSceneInfo, FMatrix& OutPreviousLocalToWorld) const
diff --git a/Engine/Source/Runtime/Renderer/Private/ScenePrivate.h b/Engine/Source/Runtime/Renderer/Private/ScenePrivate.h
--- a/Engine/Source/Runtime/Renderer/Private/ScenePrivate.h
+++ b/Engine/Source/Runtime/Renderer/Private/ScenePrivate.h
@@ -2835,8 +2835,9 @@ public:
 
 	/** Get the scene index of the FRuntimeVirtualTextureSceneProxy associated with the producer. */
 	uint32 GetRuntimeVirtualTextureSceneIndex(uint32 ProducerId);
-	/** Build the mask used for filtering this Proxy against all of the runtime virtual textures in the scene. */
-	uint32 GetRuntimeVirtualTextureMask(FPrimitiveSceneProxy const* Proxy);
+
+	/** Flush any dirty runtime virtual texture pages */
+	void FlushDirtyRuntimeVirtualTextures();
 
 #if WITH_EDITOR
 	virtual bool InitializePixelInspector(FRenderTarget* BufferFinalColor, FRenderTarget* BufferSceneColor, FRenderTarget* BufferDepth, FRenderTarget* BufferHDR, FRenderTarget* BufferA, FRenderTarget* BufferBCDE, int32 BufferIndex) override;
diff --git a/Engine/Source/Runtime/Renderer/Private/VT/RuntimeVirtualTextureProducer.cpp b/Engine/Source/Runtime/Renderer/Private/VT/RuntimeVirtualTextureProducer.cpp
--- a/Engine/Source/Runtime/Renderer/Private/VT/RuntimeVirtualTextureProducer.cpp
+++ b/Engine/Source/Runtime/Renderer/Private/VT/RuntimeVirtualTextureProducer.cpp
@@ -5,6 +5,8 @@
 #include "RendererInterface.h"
 #include "RuntimeVirtualTextureRender.h"
 #include "ScenePrivate.h"
+#include "VT/RuntimeVirtualTextureSceneProxy.h"
+
 
 FRuntimeVirtualTextureFinalizer::FRuntimeVirtualTextureFinalizer(FVTProducerDescription const& InDesc, uint32 InProducerId, ERuntimeVirtualTextureMaterialType InMaterialType, FSceneInterface* InScene, FTransform const& InUVToWorld)
 	: Desc(InDesc)
@@ -25,27 +27,37 @@ bool FRuntimeVirtualTextureFinalizer::IsReady()
 	return Scene != nullptr && Scene->GetRenderScene() != nullptr && Scene->GetRenderScene()->GetFrameNumber() > 1;
 }
 
-void FRuntimeVirtualTextureFinalizer::AddTile(FTileEntry& Tile)
-{
-	Tiles.Add(Tile);
-}
-
-void FRuntimeVirtualTextureFinalizer::Finalize(FRHICommandListImmediate& RHICmdList)
+void FRuntimeVirtualTextureFinalizer::InitProducer(const FVirtualTextureProducerHandle& ProducerHandle)
 {
 	if (RuntimeVirtualTextureMask == 0)
 	{
+		FScene* RenderScene = Scene->GetRenderScene();
+
 		// Initialize the RuntimeVirtualTextureMask by matching this producer with those registered in the scene's runtime virtual textures.
 		// We only need to do this once. If the associated scene proxy is removed this finalizer will also be destroyed.
-		const uint32 VirtualTextureSceneIndex = Scene->GetRenderScene()->GetRuntimeVirtualTextureSceneIndex(ProducerId);
+		const uint32 VirtualTextureSceneIndex = RenderScene->GetRuntimeVirtualTextureSceneIndex(ProducerId);
 		RuntimeVirtualTextureMask = 1 << VirtualTextureSceneIndex;
-	
+
+		// Store the ProducerHandle in the FRuntimeVirtualTextureSceneProxy object.
+		// This is a bit of a hack: the proxy needs to know the producer handle but can't know it on proxy creation because the producer registration is deferred to the render thread.
+		check(ProducerHandle.PackedValue != 0);
+		RenderScene->RuntimeVirtualTextures[VirtualTextureSceneIndex]->ProducerHandle = ProducerHandle;
+
 		//todo[vt]: 
 		// Add a slow render path inside RenderPage() when this check fails. 
 		// It will need to iterate the virtual textures on each primitive instead of using the RuntimeVirtualTextureMask.
 		// Currently nothing will render for this finalizer when the check fails.
 		checkSlow(VirtualTextureSceneIndex < FPrimitiveVirtualTextureFlags::RuntimeVirtualTexture_BitCount);
 	}
+}
 
+void FRuntimeVirtualTextureFinalizer::AddTile(FTileEntry& Tile)
+{
+	Tiles.Add(Tile);
+}
+
+void FRuntimeVirtualTextureFinalizer::Finalize(FRHICommandListImmediate& RHICmdList)
+{
 	for (auto Entry : Tiles)
 	{
 		const int32 TileSize = Desc.TileSize + 2 * Desc.TileBorderSize;
@@ -142,6 +154,7 @@ IVirtualTextureFinalizer* FRuntimeVirtualTextureProducer::ProducePageData(
 		Tile.DestY1 = TargetLayers[1].pPageLocation.Y;
 	}
 
+	Finalizer.InitProducer(ProducerHandle);
 	Finalizer.AddTile(Tile);
 
 	return &Finalizer;
diff --git a/Engine/Source/Runtime/Renderer/Private/VT/RuntimeVirtualTextureProducer.h b/Engine/Source/Runtime/Renderer/Private/VT/RuntimeVirtualTextureProducer.h
--- a/Engine/Source/Runtime/Renderer/Private/VT/RuntimeVirtualTextureProducer.h
+++ b/Engine/Source/Runtime/Renderer/Private/VT/RuntimeVirtualTextureProducer.h
@@ -33,6 +33,9 @@ public:
 	/** Returns false if we don't yet have everything we need to render a VT page. */
 	bool IsReady();
 
+	/** Does some one time work at the first call to set up the Producer */
+	void InitProducer(const FVirtualTextureProducerHandle& ProducerHandle);
+
 	/** Add a tile to the finalize queue. */
 	void AddTile(FTileEntry& Tile);
 
diff --git a/Engine/Source/Runtime/Renderer/Private/VT/RuntimeVirtualTextureSceneProxy.cpp b/Engine/Source/Runtime/Renderer/Private/VT/RuntimeVirtualTextureSceneProxy.cpp
--- a/Engine/Source/Runtime/Renderer/Private/VT/RuntimeVirtualTextureSceneProxy.cpp
+++ b/Engine/Source/Runtime/Renderer/Private/VT/RuntimeVirtualTextureSceneProxy.cpp
@@ -3,6 +3,7 @@
 #include "VT/RuntimeVirtualTextureSceneProxy.h"
 
 #include "Components/RuntimeVirtualTextureComponent.h"
+#include "VirtualTextureSystem.h"
 #include "VT/RuntimeVirtualTexture.h"
 #include "VT/RuntimeVirtualTextureProducer.h"
 
@@ -11,6 +12,7 @@ int32 FRuntimeVirtualTextureSceneProxy::ProducerIdGenerator = 1;
 FRuntimeVirtualTextureSceneProxy::FRuntimeVirtualTextureSceneProxy(URuntimeVirtualTextureComponent* InComponent)
 	: ProducerId(0)
 	, VirtualTexture(InComponent->GetVirtualTexture())
+	, CombinedDirtyRect(0, 0, 0, 0)
 {
 	if (VirtualTexture != nullptr)
 	{
@@ -21,11 +23,12 @@ FRuntimeVirtualTextureSceneProxy::FRuntimeVirtualTextureSceneProxy(URuntimeVirtu
 		const ERuntimeVirtualTextureMaterialType MaterialType = VirtualTexture->GetMaterialType();
 
 		// Transform is based on bottom left of the URuntimeVirtualTextureComponent unit box (which is centered on the origin)
-		FTransform Transform = FTransform(FVector(-0.5f, -0.5f, 0.f)) * InComponent->GetComponentTransform();
+		Transform = FTransform(FVector(-0.5f, -0.5f, 0.f)) * InComponent->GetComponentTransform();
 
 		// The producer description is calculated using the transform to determine the aspect ratio
 		FVTProducerDescription Desc;
 		VirtualTexture->GetProducerDescription(Desc, Transform);
+		VirtualTextureSize = FIntPoint(Desc.WidthInBlocks * Desc.BlockWidthInTiles * Desc.TileSize, Desc.HeightInBlocks * Desc.BlockHeightInTiles * Desc.TileSize);
 
 		// The Producer object created here will be passed into the virtual texture system which will take ownership.
 		// The Initialize() call will allocate the virtual texture by spawning work on the render thread.
@@ -47,3 +50,67 @@ void FRuntimeVirtualTextureSceneProxy::Release()
 		VirtualTexture = nullptr;
 	}
 }
+
+void FRuntimeVirtualTextureSceneProxy::Dirty(FBoxSphereBounds const& Bounds)
+{
+	// Transform world bounds into Virtual Texture UV space
+	const FVector O = Transform.GetTranslation();
+	const FVector U = Transform.GetUnitAxis(EAxis::X) * 1.f / Transform.GetScale3D().X;
+	const FVector V = Transform.GetUnitAxis(EAxis::Y) * 1.f / Transform.GetScale3D().Y;
+	const FVector P = Bounds.GetSphere().Center - O;
+	const FVector2D UVCenter = FVector2D(FVector::DotProduct(P, U), FVector::DotProduct(P, V));
+	const float Scale = FMath::Max(1.f / Transform.GetScale3D().X, 1.f / Transform.GetScale3D().Y);
+	const float UVRadius = Bounds.GetSphere().W * Scale;
+	const FVector2D UVExtent(UVRadius, UVRadius);
+	const FBox2D UVRect = FBox2D(UVCenter - UVExtent, UVCenter + UVExtent);
+
+	// Convert to Texel coordinate space
+	const FIntRect TextureRect(0, 0, VirtualTextureSize.X, VirtualTextureSize.Y);
+	FIntRect TexelRect(
+		FMath::FloorToInt(UVRect.Min.X * VirtualTextureSize.X),
+		FMath::FloorToInt(UVRect.Min.Y * VirtualTextureSize.Y),
+		FMath::CeilToInt(UVRect.Max.X * VirtualTextureSize.X),
+		FMath::CeilToInt(UVRect.Max.Y * VirtualTextureSize.Y));
+	TexelRect.Clip(TextureRect);
+
+	// Only add rect if it has some area
+	if (TexelRect.Min != TexelRect.Max)
+	{
+		const bool bFirst = DirtyRects.Add(TexelRect) == 0;
+		if (bFirst)
+		{
+			CombinedDirtyRect = TexelRect;
+		}
+		else
+		{
+			CombinedDirtyRect.Union(TexelRect);
+		}
+	}
+}
+
+void FRuntimeVirtualTextureSceneProxy::FlushDirtyPages()
+{
+	// If Producer handle is not initialized yet it's safe to do nothing because we won't have rendered anything to the VT that needs flushing.
+	if (ProducerHandle.PackedValue != 0)
+	{
+		//todo[vt]: 
+		// Profile to work out best heuristic for when we should use the CombinedDirtyRect
+		// Also consider using some other structure to represent dirty area such as a course 2D bitfield
+		const bool bCombinedFlush = (DirtyRects.Num() > 2 || CombinedDirtyRect == FIntRect(0, 0, VirtualTextureSize.X, VirtualTextureSize.Y));
+
+		if (bCombinedFlush)
+		{
+			FVirtualTextureSystem::Get().FlushCache(ProducerHandle, CombinedDirtyRect);
+		}
+		else
+		{
+			for (FIntRect Rect : DirtyRects)
+			{
+				FVirtualTextureSystem::Get().FlushCache(ProducerHandle, Rect);
+			}
+		}
+	}
+
+	DirtyRects.Reset();
+	CombinedDirtyRect = FIntRect(0, 0, 0, 0);
+}
diff --git a/Engine/Source/Runtime/Renderer/Private/VT/RuntimeVirtualTextureSceneProxy.h b/Engine/Source/Runtime/Renderer/Private/VT/RuntimeVirtualTextureSceneProxy.h
--- a/Engine/Source/Runtime/Renderer/Private/VT/RuntimeVirtualTextureSceneProxy.h
+++ b/Engine/Source/Runtime/Renderer/Private/VT/RuntimeVirtualTextureSceneProxy.h
@@ -21,14 +21,34 @@ public:
 	 */
 	void Release();
 
+	/** Mark an area of the associated runtime virtual texture as dirty. */
+	void Dirty(FBoxSphereBounds const& Bounds);
+
+	/** Flush the cached physical pages of the virtual texture at the given world space bounds. */
+	void FlushDirtyPages();
+
 	/** Index in FScene::RuntimeVirtualTextures. */
 	int32 SceneIndex;
 	/** Unique ID for producer that this Proxy created. Used for finding this object and it's SceneIndex from the producer.  */
 	int32 ProducerId;
 
+	/** Handle for the producer that this Proxy initialized. This is only filled in by the render thread sometime after construction! */
+	union FVirtualTextureProducerHandle ProducerHandle;
+
 	/** Pointer to linked URuntimeVirtualTexture. Not for dereferencing, just for pointer comparison. */
 	URuntimeVirtualTexture* VirtualTexture;
 
 private:
+	/** UVToWorld transform for the URuntimeVirtualTexture object. */
+	FTransform Transform;
+	/** Virtual texture size of the URuntimeVirtualTexture object. */
+	FIntPoint VirtualTextureSize;
+
+	/** Array of dirty rectangles to process at the next flush. */
+	TArray<FIntRect> DirtyRects;
+	/** Combined dirty rectangle to process at the next flush. */
+	FIntRect CombinedDirtyRect;
+
+	/** Static ProducerId counter. */
 	static int32 ProducerIdGenerator;
 };
diff --git a/Engine/Source/Runtime/Renderer/Private/VT/TexturePagePool.cpp b/Engine/Source/Runtime/Renderer/Private/VT/TexturePagePool.cpp
--- a/Engine/Source/Runtime/Renderer/Private/VT/TexturePagePool.cpp
+++ b/Engine/Source/Runtime/Renderer/Private/VT/TexturePagePool.cpp
@@ -88,6 +88,40 @@ void FTexturePagePool::EvictPages(FVirtualTextureSystem* System, const FVirtualT
 	}
 }
 
+void FTexturePagePool::EvictPages(FVirtualTextureSystem* System, FVirtualTextureProducerHandle const& ProducerHandle, FVTProducerDescription const& Desc, FIntRect const& TextureRegion, TArray<union FVirtualTextureLocalTile>& OutLocked)
+{
+	//todo[vt]: 
+	// Simple linear iteration of all physical pages here. Can we do better?
+	// We should test if it's faster to store a physical page list sorted by Morton code and find upper and lower bounds in that.
+	for (uint32 i = 0; i < NumPages; ++i)
+	{
+		if (Pages[i].PackedProducerHandle == ProducerHandle.PackedValue)
+		{
+			const uint32 vAddress = Pages[i].Local_vAddress;
+			const uint32 vLevel = Pages[i].Local_vLevel;
+
+			const int32 TileSize = Desc.TileSize << vLevel;
+			const int32 X = FMath::ReverseMortonCode2(vAddress) * TileSize;
+			const int32 Y = FMath::ReverseMortonCode2(vAddress >> 1) * TileSize;
+			const int32 TileBorderSize = Desc.TileBorderSize << vLevel;
+			const FIntRect PageRect(X - TileBorderSize, Y - TileBorderSize, X + TileSize + TileBorderSize, Y + TileSize + TileBorderSize);
+
+			if (!(PageRect.Min.X > TextureRegion.Max.X) && !(TextureRegion.Min.X > PageRect.Max.X) && !(PageRect.Min.Y > TextureRegion.Max.Y) && !(TextureRegion.Min.Y > PageRect.Max.Y))
+			{
+				if (FreeHeap.IsPresent(i))
+				{
+					OutLocked.Add(FVirtualTextureLocalTile(ProducerHandle, vAddress, vLevel));
+				}
+				else
+				{
+					UnmapAllPages(System, i);
+					FreeHeap.Update(0, i);
+				}
+			}
+		}
+	}
+}
+
 void FTexturePagePool::GetAllLockedPages(FVirtualTextureSystem* System, TSet<FVirtualTextureLocalTile>& OutPages)
 {
 
diff --git a/Engine/Source/Runtime/Renderer/Private/VT/TexturePagePool.h b/Engine/Source/Runtime/Renderer/Private/VT/TexturePagePool.h
--- a/Engine/Source/Runtime/Renderer/Private/VT/TexturePagePool.h
+++ b/Engine/Source/Runtime/Renderer/Private/VT/TexturePagePool.h
@@ -42,6 +42,12 @@ public:
 	 */
 	void EvictPages(FVirtualTextureSystem* System, const FVirtualTextureProducerHandle& ProducerHandle);
 
+	/**
+	 * Unmap/remove any pages that were allocated by the given producer and are inside the TextureRegion.
+	 * Outputs the locked pages that can't be unmapped to the OutLocked array.
+	 */
+	void EvictPages(FVirtualTextureSystem* System, FVirtualTextureProducerHandle const& ProducerHandle, FVTProducerDescription const& Desc, FIntRect const& TextureRegion, TArray<union FVirtualTextureLocalTile>& OutLocked);
+
 	/**
 	* Unmap all pages from the given space...pages will remain resident in the pool, but no longer by mapped to any page table
 	*/
diff --git a/Engine/Source/Runtime/Renderer/Private/VT/VirtualTextureSystem.cpp b/Engine/Source/Runtime/Renderer/Private/VT/VirtualTextureSystem.cpp
--- a/Engine/Source/Runtime/Renderer/Private/VT/VirtualTextureSystem.cpp
+++ b/Engine/Source/Runtime/Renderer/Private/VT/VirtualTextureSystem.cpp
@@ -30,6 +30,8 @@ DECLARE_CYCLE_STAT(TEXT("Merge Unique Pages"), STAT_ProcessRequests_MergePages,
 DECLARE_CYCLE_STAT(TEXT("Merge Requests"), STAT_ProcessRequests_MergeRequests, STATGROUP_VirtualTexturing);
 DECLARE_CYCLE_STAT(TEXT("Submit Tasks"), STAT_ProcessRequests_SubmitTasks, STATGROUP_VirtualTexturing);
 
+DECLARE_CYCLE_STAT(TEXT("Flush Cache"), STAT_FlushCache, STATGROUP_VirtualTexturing);
+
 DECLARE_DWORD_COUNTER_STAT(TEXT("Num page visible"), STAT_NumPageVisible, STATGROUP_VirtualTexturing);
 DECLARE_DWORD_COUNTER_STAT(TEXT("Num page visible resident"), STAT_NumPageVisibleResident, STATGROUP_VirtualTexturing);
 DECLARE_DWORD_COUNTER_STAT(TEXT("Num page visible not resident"), STAT_NumPageVisibleNotResident, STATGROUP_VirtualTexturing);
@@ -40,6 +42,8 @@ DECLARE_DWORD_COUNTER_STAT(TEXT("Num continuous page update"), STAT_NumContinuou
 DECLARE_DWORD_COUNTER_STAT(TEXT("Num stacks requested"), STAT_NumStacksRequested, STATGROUP_VirtualTexturing);
 DECLARE_DWORD_COUNTER_STAT(TEXT("Num stacks produced"), STAT_NumStacksProduced, STATGROUP_VirtualTexturing);
 
+DECLARE_DWORD_COUNTER_STAT(TEXT("Num flush caches"), STAT_NumFlushCache, STATGROUP_VirtualTexturing);
+
 DECLARE_MEMORY_STAT_POOL(TEXT("Total Physical Memory"), STAT_TotalPhysicalMemory, STATGROUP_VirtualTextureMemory, FPlatformMemory::MCR_GPU);
 DECLARE_MEMORY_STAT_POOL(TEXT("Total Pagetable Memory"), STAT_TotalPagetableMemory, STATGROUP_VirtualTextureMemory, FPlatformMemory::MCR_GPU);
 
@@ -240,6 +244,38 @@ void FVirtualTextureSystem::FlushCache()
 	bFlushCaches = true;
 }
 
+void FVirtualTextureSystem::FlushCache(FVirtualTextureProducerHandle const& ProducerHandle, FIntRect const& TextureRegion)
+{
+	checkSlow(IsInRenderingThread());
+
+	SCOPE_CYCLE_COUNTER(STAT_FlushCache);
+	INC_DWORD_STAT_BY(STAT_NumFlushCache, 1);
+
+	FVirtualTextureProducer const* Producer = Producers.FindProducer(ProducerHandle);
+	check(Producer != nullptr);
+	FVTProducerDescription const& ProducerDescription = Producer->GetDescription();
+
+	TArray<FVirtualTexturePhysicalSpace*> PhysicalSpacesForProducer;
+	for (uint32 i = 0; i < Producer->GetNumLayers(); ++i)
+	{
+		PhysicalSpacesForProducer.AddUnique(Producer->GetPhysicalSpace(i));
+	}
+
+	// Don't resize to allow this container to grow as needed (avoid allocations when collecting)
+	TransientCollectedPages.Reset();
+
+	for (int32 i = 0; i < PhysicalSpacesForProducer.Num(); ++i)
+	{
+		FTexturePagePool& Pool = PhysicalSpacesForProducer[i]->GetPagePool();
+		Pool.EvictPages(this, ProducerHandle, ProducerDescription, TextureRegion, TransientCollectedPages);
+	}
+
+	for (auto& Page : TransientCollectedPages)
+	{
+		MappedTilesToProduce.Add(Page);
+	}
+}
+
 void FVirtualTextureSystem::DumpFromConsole()
 {
 	bool verbose = false;
@@ -805,6 +841,9 @@ void FVirtualTextureSystem::Update(FRHICommandListImmediate& RHICmdList, ERHIFea
 
 	if (bFlushCaches)
 	{
+		SCOPE_CYCLE_COUNTER(STAT_FlushCache);
+		INC_DWORD_STAT_BY(STAT_NumFlushCache, 1);
+
 		for (int i = 0; i < PhysicalSpaces.Num(); ++i)
 		{
 			FVirtualTexturePhysicalSpace* PhysicalSpace = PhysicalSpaces[i].Get();
@@ -995,13 +1034,16 @@ void FVirtualTextureSystem::Update(FRHICommandListImmediate& RHICmdList, ERHIFea
 	{
 		GatherRequests(MergedRequestList, MergedUniquePageList, Frame - PendingFrameDelay, MemStack);
 	}
+
 	{
 		SCOPE_CYCLE_COUNTER(STAT_ProcessRequests_Sort);
-		// Limit the number of uploads
+
+		// Limit the number of uploads (account for MappedTilesToProduce this frame)
 		// Are all pages equal? Should there be different limits on different types of pages?
-		// if not async, 'infinite' uploads
-		const uint32 MaxNumUploads = CVarVTMaxUploadsPerFrame.GetValueOnRenderThread();
-		MergedRequestList->SortRequests(Producers, MemStack, MaxNumUploads);
+		const int32 MaxNumUploads = CVarVTMaxUploadsPerFrame.GetValueOnRenderThread();
+		const int32 MaxRequestUploads = FMath::Max(MaxNumUploads - MappedTilesToProduce.Num(), 1);
+
+		MergedRequestList->SortRequests(Producers, MemStack, MaxRequestUploads);
 	}
 
 	// Submit the requests to produce pages that are already mapped
diff --git a/Engine/Source/Runtime/Renderer/Private/VT/VirtualTextureSystem.h b/Engine/Source/Runtime/Renderer/Private/VT/VirtualTextureSystem.h
--- a/Engine/Source/Runtime/Renderer/Private/VT/VirtualTextureSystem.h
+++ b/Engine/Source/Runtime/Renderer/Private/VT/VirtualTextureSystem.h
@@ -62,7 +62,9 @@ public:
 	void RequestTiles(const FVector2D& InScreenSpaceSize, int32 InMipLevel = -1);
 	void RequestTilesForRegion(const IAllocatedVirtualTexture* AllocatedVT, const FVector2D& InScreenSpaceSize, const FIntRect& InTextureRegion, int32 InMipLevel = -1);
 	void LoadPendingTiles(FRHICommandListImmediate& RHICmdList, ERHIFeatureLevel::Type FeatureLevel);
+	
 	void FlushCache();
+	void FlushCache(FVirtualTextureProducerHandle const& ProducerHandle, FIntRect const& TextureRegion);
 
 private:
 	friend class FFeedbackAnalysisTask;
@@ -119,6 +121,7 @@ private:
 	FCriticalSection ContinuousUpdateTilesToProduceCS;
 	TSet<FVirtualTextureLocalTile> ContinuousUpdateTilesToProduce;
 	TSet<FVirtualTextureLocalTile> MappedTilesToProduce;
+	TArray<FVirtualTextureLocalTile> TransientCollectedPages;
 	TArray<FAllocatedVirtualTexture*> AllocatedVTsToMap;
 	TArray<IVirtualTextureFinalizer*> Finalizers;
 };
diff --git a/Engine/Source/Runtime/Renderer/Public/PrimitiveSceneInfo.h b/Engine/Source/Runtime/Renderer/Public/PrimitiveSceneInfo.h
--- a/Engine/Source/Runtime/Renderer/Public/PrimitiveSceneInfo.h
+++ b/Engine/Source/Runtime/Renderer/Public/PrimitiveSceneInfo.h
@@ -478,6 +478,15 @@ public:
 	/** Updates static lighting uniform buffer, returns the number of entries needed for GPUScene */
 	int32 UpdateStaticLightingBuffer();
 
+	/** Update the cached runtime virtual texture flags for this primitive. Do this when runtime virtual textures are created or destroyed. */
+	void UpdateRuntimeVirtualTextureFlags();
+
+	/** Get the cached runtime virtual texture flags for this primitive. */
+	FPrimitiveVirtualTextureFlags GetRuntimeVirtualTextureFlags() const { return RuntimeVirtualTextureFlags; }
+
+	/** Mark the runtime virtual textures covered by this primitive as dirty. */
+	void FlushRuntimeVirtualTexture();
+
 #if RHI_RAYTRACING
 	RENDERER_API FRayTracingGeometryRHIRef GetStaticRayTracingGeometryInstance(int LodLevel);
 #endif
@@ -509,6 +518,7 @@ private:
 	/** If this is TRUE, this primitive's indirect lighting cache buffer needs to be updated before it can be rendered. */
 	bool bIndirectLightingCacheBufferDirty : 1;
 
+	/** If this is TRUE, this primitive has registerd with the virtual texture system for a callback on virtual texture changes. */
 	bool bRegisteredVirtualTextureProducerCallback : 1;
 
 	/** Offset into the scene's lightmap data buffer, when GPUScene is enabled. */
@@ -529,6 +539,9 @@ private:
 	/** Removes cached mesh draw commands for all meshes. */
 	void RemoveCachedMeshDrawCommands();
 
+	/** These flags carry information about which runtime virtual textures are bound to this primitive. */
+	FPrimitiveVirtualTextureFlags RuntimeVirtualTextureFlags;
+
 #if RHI_RAYTRACING
 	TArray<FRayTracingGeometry*> RayTracingGeometries;
 #endif
