NonUnity build fix

diff --git a/Engine/Source/Runtime/Renderer/Private/VT/TexturePagePool.h b/Engine/Source/Runtime/Renderer/Private/VT/TexturePagePool.h
--- a/Engine/Source/Runtime/Renderer/Private/VT/TexturePagePool.h
+++ b/Engine/Source/Runtime/Renderer/Private/VT/TexturePagePool.h
@@ -7,10 +7,12 @@
 #include "Containers/BinaryHeap.h"
 #include "Containers/HashTable.h"
 
-class FVirtualTextureSystem;
-class FVirtualTextureSpace;
+union FVirtualTextureLocalTile;
 class FVirtualTexturePhysicalSpace;
 union FVirtualTextureProducerHandle;
+class FVirtualTextureSpace;
+class FVirtualTextureSystem;
+struct FVTProducerDescription;
 
 /**
  * Manages a pool of texture pages, backed by a large GPU texture atlas.
