Fix race condition in VT material updates that triggered check in FMaterialShader::GetShaderBindings()
The VT update can run concurrently with worker thread mesh command generation
VT system callbacks which can touch materials need to be complete before this phase
Moved VT callbacks to before InitViews to fix this

diff --git a/Engine/Source/Runtime/Renderer/Private/DeferredShadingRenderer.cpp b/Engine/Source/Runtime/Renderer/Private/DeferredShadingRenderer.cpp
--- a/Engine/Source/Runtime/Renderer/Private/DeferredShadingRenderer.cpp
+++ b/Engine/Source/Runtime/Renderer/Private/DeferredShadingRenderer.cpp
@@ -891,6 +891,7 @@ void FDeferredShadingSceneRenderer::Render(FRHICommandListImmediate& RHICmdList)
 	{
 		// AllocateResources needs to be called before RHIBeginScene
 		FVirtualTextureSystem::Get().AllocateResources(RHICmdList, FeatureLevel);
+		FVirtualTextureSystem::Get().CallPendingCallbacks();
 	}
 
 	const bool bIsWireframe = ViewFamily.EngineShowFlags.Wireframe;
diff --git a/Engine/Source/Runtime/Renderer/Private/VT/VirtualTextureSystem.cpp b/Engine/Source/Runtime/Renderer/Private/VT/VirtualTextureSystem.cpp
--- a/Engine/Source/Runtime/Renderer/Private/VT/VirtualTextureSystem.cpp
+++ b/Engine/Source/Runtime/Renderer/Private/VT/VirtualTextureSystem.cpp
@@ -859,8 +859,6 @@ void FVirtualTextureSystem::Update(FRHICommandListImmediate& RHICmdList, ERHIFea
 		bFlushCaches = false;
 	}
 
-	Producers.CallPendingCallbacks();
-
 	DestroyPendingVirtualTextures();
 
 	FMemStack& MemStack = FMemStack::Get();
@@ -1945,3 +1943,8 @@ void FVirtualTextureSystem::AllocateResources(FRHICommandListImmediate& RHICmdLi
 		}
 	}
 }
+
+void FVirtualTextureSystem::CallPendingCallbacks()
+{
+	Producers.CallPendingCallbacks();
+}
diff --git a/Engine/Source/Runtime/Renderer/Private/VT/VirtualTextureSystem.h b/Engine/Source/Runtime/Renderer/Private/VT/VirtualTextureSystem.h
--- a/Engine/Source/Runtime/Renderer/Private/VT/VirtualTextureSystem.h
+++ b/Engine/Source/Runtime/Renderer/Private/VT/VirtualTextureSystem.h
@@ -34,6 +34,7 @@ public:
 	static FVirtualTextureSystem& Get();
 
 	void AllocateResources(FRHICommandListImmediate& RHICmdList, ERHIFeatureLevel::Type FeatureLevel);
+	void CallPendingCallbacks();
 	void Update( FRHICommandListImmediate& RHICmdList, ERHIFeatureLevel::Type FeatureLevel );
 
 	IAllocatedVirtualTexture* AllocateVirtualTexture(const FAllocatedVTDescription& Desc);
