NonUnity build fix

diff --git a/Engine/Source/Runtime/Renderer/Private/VT/RuntimeVirtualTextureSceneProxy.h b/Engine/Source/Runtime/Renderer/Private/VT/RuntimeVirtualTextureSceneProxy.h
--- a/Engine/Source/Runtime/Renderer/Private/VT/RuntimeVirtualTextureSceneProxy.h
+++ b/Engine/Source/Runtime/Renderer/Private/VT/RuntimeVirtualTextureSceneProxy.h
@@ -3,6 +3,7 @@
 #pragma once
 
 #include "CoreMinimal.h"
+#include "VirtualTexturing.h"
 
 class URuntimeVirtualTexture;
 class URuntimeVirtualTextureComponent;
@@ -33,7 +34,7 @@ public:
 	int32 ProducerId;
 
 	/** Handle for the producer that this Proxy initialized. This is only filled in by the render thread sometime after construction! */
-	union FVirtualTextureProducerHandle ProducerHandle;
+	FVirtualTextureProducerHandle ProducerHandle;
 
 	/** Pointer to linked URuntimeVirtualTexture. Not for dereferencing, just for pointer comparison. */
 	URuntimeVirtualTexture* VirtualTexture;
