Cherry-pick 7936902: Added option to exclude decals in ray tracing (r.RayTracing.ExcludeDecals)

diff --git a/Engine/Source/Runtime/Renderer/Private/DeferredShadingRenderer.cpp b/Engine/Source/Runtime/Renderer/Private/DeferredShadingRenderer.cpp
--- a/Engine/Source/Runtime/Renderer/Private/DeferredShadingRenderer.cpp
+++ b/Engine/Source/Runtime/Renderer/Private/DeferredShadingRenderer.cpp
@@ -122,6 +122,14 @@ static TAutoConsoleVariable<int32> CVarForceAllRayTracingEffects(
 	TEXT(" 1: All ray tracing effects enabled"),
 	ECVF_RenderThreadSafe);
 
+static int32 GRayTracingExcludeDecals = 0;
+static FAutoConsoleVariableRef CRayTracingExcludeDecals(
+	TEXT("r.RayTracing.ExcludeDecals"),
+	GRayTracingExcludeDecals,
+	TEXT("A toggle that modifies the inclusion of decals in the ray tracing BVH.\n")
+	TEXT(" 0: Decals included in the ray tracing BVH (default)\n")
+	TEXT(" 1: Decals excluded from the ray tracing BVH"),
+	ECVF_RenderThreadSafe);
 
 #if !UE_BUILD_SHIPPING
 static TAutoConsoleVariable<int32> CVarForceBlackVelocityBuffer(
@@ -657,6 +665,7 @@ bool FDeferredShadingSceneRenderer::GatherRayTracingWorldInstances(FRHICommandLi
 						uint8 NewInstanceMask = 0;
 						bool bAllSegmentsOpaque = true;
 						bool bAnySegmentsCastShadow = false;
+						bool bAnySegmentsDecal = false;
 
 						uint32 LODIndex = LODToRender.GetRayTracedLOD();
 						// Sometimes LODIndex is out of range because it is clamped by ClampToFirstLOD, like the requested LOD is being streamed in and hasn't been available
@@ -679,6 +688,7 @@ bool FDeferredShadingSceneRenderer::GatherRayTracingWorldInstances(FRHICommandLi
 									NewInstanceMask |= NewVisibleMeshCommand.RayTracingMeshCommand->InstanceMask;
 									bAllSegmentsOpaque &= NewVisibleMeshCommand.RayTracingMeshCommand->bOpaque;
 									bAnySegmentsCastShadow |= NewVisibleMeshCommand.RayTracingMeshCommand->bCastRayTracedShadows;
+									bAnySegmentsDecal |= NewVisibleMeshCommand.RayTracingMeshCommand->bDecal;
 								}
 								else
 								{
@@ -687,6 +697,11 @@ bool FDeferredShadingSceneRenderer::GatherRayTracingWorldInstances(FRHICommandLi
 								}
 							}
 
+							if (GRayTracingExcludeDecals && bAnySegmentsDecal)
+							{
+								continue;
+							}
+
 							NewInstanceMask |= bAnySegmentsCastShadow ? RAY_TRACING_MASK_SHADOW : 0;
 
 							// When no cached command is found, NewInstanceMask == 0 and the instance is effectively filtered out
diff --git a/Engine/Source/Runtime/Renderer/Private/RayTracing/RayTracingMaterialHitShaders.cpp b/Engine/Source/Runtime/Renderer/Private/RayTracing/RayTracingMaterialHitShaders.cpp
--- a/Engine/Source/Runtime/Renderer/Private/RayTracing/RayTracingMaterialHitShaders.cpp
+++ b/Engine/Source/Runtime/Renderer/Private/RayTracing/RayTracingMaterialHitShaders.cpp
@@ -206,6 +206,7 @@ void FRayTracingMeshProcessor::BuildRayTracingMeshCommands(
 	SharedCommand.InstanceMask = ComputeBlendModeMask(MaterialResource.GetBlendMode());
 	SharedCommand.bCastRayTracedShadows = MeshBatch.CastRayTracedShadow && MaterialResource.CastsRayTracedShadows();
 	SharedCommand.bOpaque = MaterialResource.GetBlendMode() == EBlendMode::BLEND_Opaque;
+	SharedCommand.bDecal = MaterialResource.GetMaterialDomain() == EMaterialDomain::MD_DeferredDecal;
 
 	FVertexInputStreamArray VertexStreams;
 	VertexFactory->GetStreams(ERHIFeatureLevel::SM5, EVertexInputStreamType::Default, VertexStreams);
diff --git a/Engine/Source/Runtime/Renderer/Public/MeshPassProcessor.h b/Engine/Source/Runtime/Renderer/Public/MeshPassProcessor.h
--- a/Engine/Source/Runtime/Renderer/Public/MeshPassProcessor.h
+++ b/Engine/Source/Runtime/Renderer/Public/MeshPassProcessor.h
@@ -1233,6 +1233,7 @@ public:
 
 	bool bCastRayTracedShadows = true;
 	bool bOpaque = true;
+	bool bDecal = false;
 
 	/** Sets ray hit group shaders on the mesh command and allocates room for the shader bindings. */
 	RENDERER_API void SetShaders(const FMeshProcessorShaders& Shaders);
