Add WorldHeight layer type to Runtime Virtual Texture

diff --git a/Engine/Shaders/Private/VirtualTextureCommon.ush b/Engine/Shaders/Private/VirtualTextureCommon.ush
--- a/Engine/Shaders/Private/VirtualTextureCommon.ush
+++ b/Engine/Shaders/Private/VirtualTextureCommon.ush
@@ -544,18 +544,12 @@ float3 VirtualTextureUnpackNormalBC5(in float4 PackedValue)
 	return VirtualTextureUnpackNormal(PackedValue.rb);
 }
 
-/** Generic height unpack funtion. */
-float VirtualTextureUnpackHeight(in float2 PackedValue)
+/** Unpack 16 bit height (with some hardcoded range scale/bias). */
+float VirtualTextureUnpackHeightR16(in float4 PackedValue)
 {
-	float ScaledHeight = PackedValue.x * 255.f * 256.f + PackedValue.y * 255.f;
-	float Height = (ScaledHeight - 32768.f) * 2.f;
+	float Height = PackedValue.r * 65535.f - 32768.f;
+	Height *= 2.f;
 	return Height;
 }
 
-/** Unpack height from R8G8. */
-float VirtualTextureUnpackHeightR8G8(in float4 PackedValue)
-{
-	return VirtualTextureUnpackHeight(PackedValue.rg);
-}
-
 #endif //VirtualTextureCommon_ush__
diff --git a/Engine/Shaders/Private/VirtualTextureMaterial.usf b/Engine/Shaders/Private/VirtualTextureMaterial.usf
--- a/Engine/Shaders/Private/VirtualTextureMaterial.usf
+++ b/Engine/Shaders/Private/VirtualTextureMaterial.usf
@@ -44,13 +44,11 @@ float2 PackNormal( float3 N )
 	return normalize(N).xy * 0.5 + 0.5;
 }
 
-// Prepare for VirtualTextureUnpackHeight() in VirtualTextureCommon.ush
-float2 PackHeight( float Height )
+// Prepare for VirtualTextureUnpackHeightR16() in VirtualTextureCommon.ush
+float PackWorldHeight( float Height )
 {
-	uint H = (uint)clamp(Height * 0.5f + 32768.f, 0.f, 65535.f);
-	uint LSH = H & 0xff;
-	uint MSH = (H - LSH) >> 8;
-	return float2((float)MSH, (float)LSH) * (1.f / 255.f);
+	Height *= 0.5f;
+	return saturate((Height + 32768.f) / 65535.f);
 }
 
 void FPixelShaderInOut_MainPS(
@@ -75,13 +73,15 @@ void FPixelShaderInOut_MainPS(
 	half Specular = GetVirtualTextureOutput1(MaterialParameters);
 	half Roughness = GetVirtualTextureOutput2(MaterialParameters);
 	half3 Normal = GetVirtualTextureOutput3(MaterialParameters);
-	float Opacity = GetVirtualTextureOutput4(MaterialParameters);
+	half WorldHeight = GetVirtualTextureOutput4(MaterialParameters);
+	float Opacity = GetVirtualTextureOutput5(MaterialParameters);
 #else
 	// Output is from standard material output attribute node
 	half3 BaseColor = GetMaterialBaseColor(PixelMaterialInputs);
 	half Specular = GetMaterialSpecular(PixelMaterialInputs);
 	half Roughness = GetMaterialRoughness(PixelMaterialInputs);
 	half3 Normal = mul(MaterialParameters.WorldNormal, (float3x3)ResolvedView.WorldToVirtualTexture);
+	float WorldHeight = MaterialParameters.AbsoluteWorldPosition.z;
 	float Opacity = GetMaterialOpacity(PixelMaterialInputs);
 #endif
 
@@ -94,6 +94,8 @@ void FPixelShaderInOut_MainPS(
 	Out.MRT[0] = float4(BaseColor, 1.f) * Opacity;
 	Out.MRT[1] = float4(PackNormal(Normal), 0.f, 1.f) * Opacity;
 	Out.MRT[2] = float4(Roughness, Specular, 0.f, 1.f) * Opacity;
+#elif defined(OUT_WORLDHEIGHT)
+	Out.MRT[0] = float4(PackWorldHeight(WorldHeight), 0, 0, 1);
 #endif
 }
 
@@ -108,6 +110,8 @@ void FPixelShaderInOut_MainPS(
 #define PIXELSHADEROUTPUT_MRT0 1
 #define PIXELSHADEROUTPUT_MRT1 1
 #define PIXELSHADEROUTPUT_MRT2 1
+#elif defined(OUT_WORLDHEIGHT)
+#define PIXELSHADEROUTPUT_MRT0 1
 #endif
 
 // all PIXELSHADEROUTPUT_ and "void FPixelShaderInOut_MainPS()" need to be setup before this include
diff --git a/Engine/Source/Runtime/Engine/Classes/Materials/MaterialExpressionRuntimeVirtualTextureOutput.h b/Engine/Source/Runtime/Engine/Classes/Materials/MaterialExpressionRuntimeVirtualTextureOutput.h
--- a/Engine/Source/Runtime/Engine/Classes/Materials/MaterialExpressionRuntimeVirtualTextureOutput.h
+++ b/Engine/Source/Runtime/Engine/Classes/Materials/MaterialExpressionRuntimeVirtualTextureOutput.h
@@ -29,6 +29,10 @@ class UMaterialExpressionRuntimeVirtualTextureOutput : public UMaterialExpressio
 	UPROPERTY()
 	FExpressionInput Normal;
 
+	/** Input for World Height to output to virtual texture. */
+	UPROPERTY()
+	FExpressionInput WorldHeight;
+
 	/** Input for Opacity value used for blending to virtual texture. */
 	UPROPERTY()
 	FExpressionInput Opacity;
diff --git a/Engine/Source/Runtime/Engine/Private/Materials/HLSLMaterialTranslator.h b/Engine/Source/Runtime/Engine/Private/Materials/HLSLMaterialTranslator.h
--- a/Engine/Source/Runtime/Engine/Private/Materials/HLSLMaterialTranslator.h
+++ b/Engine/Source/Runtime/Engine/Private/Materials/HLSLMaterialTranslator.h
@@ -5067,9 +5067,9 @@ protected:
 				FString	SampleCode(TEXT("VirtualTextureUnpackNormalBC5(%s)"));
 				return AddCodeChunk(MCT_Float3, *SampleCode, *GetParameterCode(CodeIndex));
 			}
-			else if (UnpackType == EVirtualTextureUnpackType::HeightR8G8)
+			else if (UnpackType == EVirtualTextureUnpackType::HeightR16)
 			{
-				FString	SampleCode(TEXT("VirtualTextureUnpackHeightR8G8(%s)"));
+				FString	SampleCode(TEXT("VirtualTextureUnpackHeightR16(%s)"));
 				return AddCodeChunk(MCT_Float, *SampleCode, *GetParameterCode(CodeIndex));
 			}
 		}
diff --git a/Engine/Source/Runtime/Engine/Private/Materials/MaterialExpressions.cpp b/Engine/Source/Runtime/Engine/Private/Materials/MaterialExpressions.cpp
--- a/Engine/Source/Runtime/Engine/Private/Materials/MaterialExpressions.cpp
+++ b/Engine/Source/Runtime/Engine/Private/Materials/MaterialExpressions.cpp
@@ -1947,6 +1947,10 @@ int32 UMaterialExpressionRuntimeVirtualTextureOutput::Compile(class FMaterialCom
 		CodeInput = Normal.IsConnected() ? Normal.Compile(Compiler) : Compiler->Constant3(0.f, 0.f, 1.f);
 	}
 	else if (OutputIndex == 4)
+	{
+		CodeInput = WorldHeight.IsConnected() ? WorldHeight.Compile(Compiler) : Compiler->Constant(0.f);
+	}
+	else if (OutputIndex == 5)
 	{
 		CodeInput = Opacity.IsConnected() ? Opacity.Compile(Compiler) : Compiler->Constant(1.f);
 	}
@@ -1964,7 +1968,7 @@ void UMaterialExpressionRuntimeVirtualTextureOutput::GetCaption(TArray<FString>&
 
 int32 UMaterialExpressionRuntimeVirtualTextureOutput::GetNumOutputs() const
 {
-	return 5; 
+	return 6; 
 }
 
 FString UMaterialExpressionRuntimeVirtualTextureOutput::GetFunctionName() const
@@ -2019,6 +2023,7 @@ void UMaterialExpressionRuntimeVirtualTextureSample::InitOutputs()
 	Outputs.Add(FExpressionOutput(TEXT("Specular"), 1, 0, 0, 1, 0));
 	Outputs.Add(FExpressionOutput(TEXT("Roughness"), 1, 1, 0, 0, 0));
 	Outputs.Add(FExpressionOutput(TEXT("Normal")));
+	Outputs.Add(FExpressionOutput(TEXT("WorldHeight")));
 #endif // WITH_EDITORONLY_DATA
 }
 
@@ -2078,6 +2083,7 @@ int32 UMaterialExpressionRuntimeVirtualTextureSample::Compile(class FMaterialCom
 	const bool bIsSpecularValid = MaterialType == ERuntimeVirtualTextureMaterialType::BaseColor_Normal_Specular;
 	const bool bIsNormalValid = MaterialType == ERuntimeVirtualTextureMaterialType::BaseColor_Normal || MaterialType == ERuntimeVirtualTextureMaterialType::BaseColor_Normal_Specular;
 	const bool bIsNormalBC3 = MaterialType == ERuntimeVirtualTextureMaterialType::BaseColor_Normal_Specular;
+	const bool bIsWorldHeightValid = MaterialType == ERuntimeVirtualTextureMaterialType::WorldHeight;
 
 	int32 LayerIndex = 0;
 	EMaterialSamplerType SamplerType = SAMPLERTYPE_VirtualMasks;
@@ -2117,6 +2123,17 @@ int32 UMaterialExpressionRuntimeVirtualTextureSample::Compile(class FMaterialCom
 			return Compiler->Constant3(0.f, 0.f, 1.f);
 		}
 		break;
+	case 4:
+		if (bIsVirtualTextureValid && bIsWorldHeightValid)
+		{
+			LayerIndex = 0;
+			UnpackType = EVirtualTextureUnpackType::HeightR16;
+		}
+		else
+		{
+			return Compiler->Constant(0.f);
+		}
+		break;
 	default:
 		return INDEX_NONE;
 	}
diff --git a/Engine/Source/Runtime/Engine/Private/VT/RuntimeVirtualTexture.cpp b/Engine/Source/Runtime/Engine/Private/VT/RuntimeVirtualTexture.cpp
--- a/Engine/Source/Runtime/Engine/Private/VT/RuntimeVirtualTexture.cpp
+++ b/Engine/Source/Runtime/Engine/Private/VT/RuntimeVirtualTexture.cpp
@@ -221,6 +221,10 @@ void URuntimeVirtualTexture::GetProducerDescription(FVTProducerDescription& OutD
 		OutDesc.LayerFormat[0] = bCompressTextures ? PF_DXT1 : PF_B8G8R8A8;
 		OutDesc.LayerFormat[1] = bCompressTextures ? PF_DXT5 : PF_B8G8R8A8;
 		break;
+	case ERuntimeVirtualTextureMaterialType::WorldHeight:
+		OutDesc.NumLayers = 1;
+		OutDesc.LayerFormat[0] = PF_G16;
+		break;
 	default:
 		checkf(0, TEXT("Invalid Runtime Virtual Texture setup: %s, %d"), *GetName(), MaterialType);
 		OutDesc.NumLayers = 1;
@@ -238,6 +242,8 @@ bool URuntimeVirtualTexture::IsLayerSRGB(int32 LayerIndex) const
 	case ERuntimeVirtualTextureMaterialType::BaseColor_Normal_Specular:
 		// Only BaseColor layer is sRGB
 		return LayerIndex == 0;
+	case ERuntimeVirtualTextureMaterialType::WorldHeight:
+		return false;
 	default:
 		// Implement logic for any missing material types
 		check(false);
diff --git a/Engine/Source/Runtime/Engine/Public/MaterialCompiler.h b/Engine/Source/Runtime/Engine/Public/MaterialCompiler.h
--- a/Engine/Source/Runtime/Engine/Public/MaterialCompiler.h
+++ b/Engine/Source/Runtime/Engine/Public/MaterialCompiler.h
@@ -33,7 +33,7 @@ enum class EVirtualTextureUnpackType
 	None,
 	NormalBC3,
 	NormalBC5,
-	HeightR8G8,
+	HeightR16,
 };
 
 /** 
diff --git a/Engine/Source/Runtime/Engine/Public/VT/RuntimeVirtualTextureEnum.h b/Engine/Source/Runtime/Engine/Public/VT/RuntimeVirtualTextureEnum.h
--- a/Engine/Source/Runtime/Engine/Public/VT/RuntimeVirtualTextureEnum.h
+++ b/Engine/Source/Runtime/Engine/Public/VT/RuntimeVirtualTextureEnum.h
@@ -16,6 +16,7 @@ enum class ERuntimeVirtualTextureMaterialType
 	BaseColor UMETA(DisplayName = "Base Color"),
 	BaseColor_Normal UMETA(DisplayName = "Base Color, Normal"),
 	BaseColor_Normal_Specular UMETA(DisplayName = "Base Color, Normal, Roughness, Specular"),
+	WorldHeight UMETA(DisplayName = "World Height"),
 	
 	Count UMETA(Hidden),
 };
diff --git a/Engine/Source/Runtime/Landscape/Private/LandscapeRender.cpp b/Engine/Source/Runtime/Landscape/Private/LandscapeRender.cpp
--- a/Engine/Source/Runtime/Landscape/Private/LandscapeRender.cpp
+++ b/Engine/Source/Runtime/Landscape/Private/LandscapeRender.cpp
@@ -4410,6 +4410,16 @@ public:
 			// Forward shading required
 			FName(TEXT("TBasePassPSFCachedPointIndirectLightingPolicySkylight")),
 			FName(TEXT("TBasePassPSFNoLightMapPolicySkylight")),
+
+			// Runtime virtual texture
+			FName(TEXT("TVirtualTextureVSBaseColor")),
+			FName(TEXT("TVirtualTextureVSBaseColorNormal")),
+			FName(TEXT("TVirtualTextureVSBaseColorNormalSpecular")),
+			FName(TEXT("TVirtualTextureVSWorldHeight")),
+			FName(TEXT("TVirtualTexturePSBaseColor")),
+			FName(TEXT("TVirtualTexturePSBaseColorNormal")),
+			FName(TEXT("TVirtualTexturePSBaseColorNormalSpecular")),
+			FName(TEXT("TVirtualTexturePSWorldHeight")),
 		};
 		return AllowedShaderTypes;
 	}
@@ -4580,9 +4590,11 @@ public:
 			FName(TEXT("TVirtualTextureVSBaseColor")),
 			FName(TEXT("TVirtualTextureVSBaseColorNormal")),
 			FName(TEXT("TVirtualTextureVSBaseColorNormalSpecular")),
+			FName(TEXT("TVirtualTextureVSWorldHeight")),
 			FName(TEXT("TVirtualTexturePSBaseColor")),
 			FName(TEXT("TVirtualTexturePSBaseColorNormal")),
 			FName(TEXT("TVirtualTexturePSBaseColorNormalSpecular")),
+			FName(TEXT("TVirtualTexturePSWorldHeight")),
 		};
 		return ShaderTypes;
 	}
diff --git a/Engine/Source/Runtime/Renderer/Private/VT/RuntimeVirtualTextureRender.cpp b/Engine/Source/Runtime/Renderer/Private/VT/RuntimeVirtualTextureRender.cpp
--- a/Engine/Source/Runtime/Renderer/Private/VT/RuntimeVirtualTextureRender.cpp
+++ b/Engine/Source/Runtime/Renderer/Private/VT/RuntimeVirtualTextureRender.cpp
@@ -111,6 +111,21 @@ namespace RuntimeVirtualTexture
 		}
 	};
 
+	/** Specialization for ERuntimeVirtualTextureMaterialType::WorldHeight */
+	class FMaterialPolicy_WorldHeight
+	{
+	public:
+		static void ModifyCompilationEnvironment(FShaderCompilerEnvironment& OutEnvironment)
+		{
+			OutEnvironment.SetDefine(TEXT("OUT_WORLDHEIGHT"), 1);
+		}
+
+		static FRHIBlendState* GetBlendState()
+		{
+			return TStaticBlendState< CW_RED, BO_Max, BF_One, BF_One, BO_Add, BF_One, BF_One >::GetRHI();
+		}
+	};
+
 
 	/** Vertex shader derivation of material shader. Templated on policy for virtual texture layout. */
 	template< class MaterialPolicy >
@@ -165,7 +180,7 @@ namespace RuntimeVirtualTexture
 	IMPLEMENT_VIRTUALTEXTURE_SHADER_TYPE(FMaterialPolicy_BaseColor, BaseColor);
 	IMPLEMENT_VIRTUALTEXTURE_SHADER_TYPE(FMaterialPolicy_BaseColorNormal, BaseColorNormal);
 	IMPLEMENT_VIRTUALTEXTURE_SHADER_TYPE(FMaterialPolicy_BaseColorNormalSpecular, BaseColorNormalSpecular);
-
+	IMPLEMENT_VIRTUALTEXTURE_SHADER_TYPE(FMaterialPolicy_WorldHeight, WorldHeight);
 
 	/** Mesh processor for rendering static meshes to the virtual texture */
 	class FRuntimeVirtualTextureMeshProcessor : public FMeshPassProcessor
@@ -248,6 +263,9 @@ namespace RuntimeVirtualTexture
 				case ERuntimeVirtualTextureMaterialType::BaseColor_Normal_Specular:
 					Process<FMaterialPolicy_BaseColorNormalSpecular>(MeshBatch, BatchElementMask, StaticMeshId, PrimitiveSceneProxy, MaterialRenderProxy, Material);
 					break;
+				case ERuntimeVirtualTextureMaterialType::WorldHeight:
+					Process<FMaterialPolicy_WorldHeight>(MeshBatch, BatchElementMask, StaticMeshId, PrimitiveSceneProxy, MaterialRenderProxy, Material);
+					break;
 				default:
 					break;
 				}
@@ -553,6 +571,12 @@ namespace RuntimeVirtualTexture
 					OutputAlias1 = CopyTexture0 = GraphBuilder.CreateTexture(FPooledRenderTargetDesc::Create2DDesc(TextureSize, PF_B8G8R8A8, FClearValueBinding::None, TexCreate_None, TexCreate_UAV, false), TEXT("CopyTexture0"));
 				}
 				break;
+			case ERuntimeVirtualTextureMaterialType::WorldHeight:
+				if (bRenderPass)
+				{
+					OutputAlias0 = RenderTexture0 = GraphBuilder.CreateTexture(FPooledRenderTargetDesc::Create2DDesc(TextureSize, PF_G16, FClearValueBinding::Black, TexCreate_SRGB, TexCreate_RenderTargetable, false), TEXT("RenderTexture0"));
+				}
+				break;
 			}
 		}
 
