Fix dynamic shadows stereo rendering logic bug.
An update to use correct APIs to determine stereo-pass introduced a logic error in the place where we determine FadeAlpha (and duplicate for the secondary pass eye if in stereo). The original logic checked for if the current view is primary stereo, then if there's a second view and it's secondary, set the second fade alpha to 1 also. The update checked the first view to be primary AND seconday (which can never happen). This has been fixed, with the new (correct) APIs.

diff --git a/Engine/Source/Runtime/Renderer/Private/ShadowSetup.cpp b/Engine/Source/Runtime/Renderer/Private/ShadowSetup.cpp
--- a/Engine/Source/Runtime/Renderer/Private/ShadowSetup.cpp
+++ b/Engine/Source/Runtime/Renderer/Private/ShadowSetup.cpp
@@ -3296,7 +3296,7 @@ void FSceneRenderer::AddViewDependentWholeSceneShadowsForView(
 
 		if (IStereoRendering::IsAPrimaryView(View.StereoPass, GEngine->StereoRenderingDevice)
 			&& Views.IsValidIndex(ViewIndex + 1)
-			&& IStereoRendering::IsASecondaryView(View.StereoPass, GEngine->StereoRenderingDevice))
+			&& IStereoRendering::IsASecondaryView(Views[ViewIndex + 1].StereoPass, GEngine->StereoRenderingDevice))
 		{
 			FadeAlphas[ViewIndex + 1] = 1.0f;
 		}		
