Disabled auto-exposure when using the player collision view.

diff --git a/Engine/Source/Runtime/Renderer/Private/PostProcess/PostProcessEyeAdaptation.cpp b/Engine/Source/Runtime/Renderer/Private/PostProcess/PostProcessEyeAdaptation.cpp
--- a/Engine/Source/Runtime/Renderer/Private/PostProcess/PostProcessEyeAdaptation.cpp
+++ b/Engine/Source/Runtime/Renderer/Private/PostProcess/PostProcessEyeAdaptation.cpp
@@ -102,7 +102,9 @@ inline static void ComputeEyeAdaptationValues(const ERHIFeatureLevel::Type MinFe
 		|| View.Family->UseDebugViewPS() 
 		|| EngineShowFlags.RayTracingDebug 
 		|| EngineShowFlags.VisualizeDistanceFieldAO
-		|| EngineShowFlags.VisualizeGlobalDistanceField)
+		|| EngineShowFlags.VisualizeGlobalDistanceField
+		|| EngineShowFlags.CollisionVisibility
+		|| EngineShowFlags.CollisionPawn)
 	{
 		LocalExposureMultipler = 1.f;
 	}
