Disabled occlusion queries in collision viewmode to prevent flickering of collision geometry.
This is because collision primitives can be authored to be bigger than the rendering primitive bounds.

diff --git a/Engine/Source/Runtime/Renderer/Private/SceneVisibility.cpp b/Engine/Source/Runtime/Renderer/Private/SceneVisibility.cpp
--- a/Engine/Source/Runtime/Renderer/Private/SceneVisibility.cpp
+++ b/Engine/Source/Runtime/Renderer/Private/SceneVisibility.cpp
@@ -3001,7 +3001,9 @@ void FSceneRenderer::PreVisibilityFrameSetup(FRHICommandListImmediate& RHICmdLis
 		// HighResScreenshot should get best results so we don't do the occlusion optimization based on the former frame
 		extern bool GIsHighResScreenshot;
 		const bool bIsHitTesting = ViewFamily.EngineShowFlags.HitProxies;
-		if (GIsHighResScreenshot || !DoOcclusionQueries(FeatureLevel) || bIsHitTesting)
+		// Don't test occlusion queries in collision viewmode as they can be bigger then the rendering bounds.
+		const bool bCollisionView = ViewFamily.EngineShowFlags.CollisionVisibility || ViewFamily.EngineShowFlags.CollisionPawn;
+		if (GIsHighResScreenshot || !DoOcclusionQueries(FeatureLevel) || bIsHitTesting || bCollisionView)
 		{
 			View.bDisableQuerySubmissions = true;
 			View.bIgnoreExistingQueries = true;
