Fixed: scene capture on mobile with MSAA enabled

diff --git a/Engine/Shaders/Private/SceneTexturesCommon.ush b/Engine/Shaders/Private/SceneTexturesCommon.ush
--- a/Engine/Shaders/Private/SceneTexturesCommon.ush
+++ b/Engine/Shaders/Private/SceneTexturesCommon.ush
@@ -139,7 +139,9 @@ float CalcSceneDepth(float2 ScreenUV)
 /** Returns DeviceZ which is the z value stored in the depth buffer. */
 float LookupDeviceZ( float2 ScreenUV )
 {
-#if	MOBILE_FORCE_DEPTH_TEXTURE_READS
+#if	SCENE_TEXTURES_DISABLED
+	return SCENE_TEXTURES_DISABLED_SCENE_DEPTH_VALUE;
+#elif MOBILE_FORCE_DEPTH_TEXTURE_READS
 	// native Depth buffer lookup
 	return Texture2DSampleLevel(MobileSceneTextures.SceneDepthTexture, MobileSceneTextures.SceneDepthTextureSampler, ScreenUV, 0).r;
 #elif POST_PROCESS_MATERIAL
diff --git a/Engine/Source/Runtime/Renderer/Private/MobileTranslucentRendering.cpp b/Engine/Source/Runtime/Renderer/Private/MobileTranslucentRendering.cpp
--- a/Engine/Source/Runtime/Renderer/Private/MobileTranslucentRendering.cpp
+++ b/Engine/Source/Runtime/Renderer/Private/MobileTranslucentRendering.cpp
@@ -210,7 +210,7 @@ public:
 	static void ModifyCompilationEnvironment(const FMaterialShaderPermutationParameters& Parameters, FShaderCompilerEnvironment& OutEnvironment)
 	{
 		FMeshMaterialShader::ModifyCompilationEnvironment(Parameters, OutEnvironment);
-		OutEnvironment.SetDefine(TEXT("MOBILE_FORCE_DEPTH_TEXTURE_READS"), 1u);
+		OutEnvironment.SetDefine(TEXT("SCENE_TEXTURES_DISABLED"), 1u);
 	}
 
 	FOpacityOnlyPS(const ShaderMetaType::CompiledShaderInitializerType& Initializer) :
@@ -238,18 +238,29 @@ bool FMobileSceneRenderer::RenderInverseOpacity(FRHICommandListImmediate& RHICmd
 	bool bDirty = false;
 	FSceneRenderTargets& SceneContext = FSceneRenderTargets::Get(RHICmdList);
 
+	SceneContext.AllocSceneColor(RHICmdList);
+
+	const bool bMobileMSAA = SceneContext.GetSceneColorSurface()->GetNumSamples() > 1;
+	
+	FRHITexture* SceneColorResolve = bMobileMSAA ? SceneContext.GetSceneColorTexture() : nullptr;
+	ERenderTargetActions ColorTargetAction = bMobileMSAA ? ERenderTargetActions::Clear_Resolve : ERenderTargetActions::Clear_Store;
+	FRHIRenderPassInfo RPInfo(
+		SceneContext.GetSceneColorSurface(), 
+		ColorTargetAction,
+		SceneColorResolve,
+		SceneContext.GetSceneDepthSurface(),
+		EDepthStencilTargetActions::ClearDepthStencil_DontStoreDepthStencil,
+		nullptr,
+		FExclusiveDepthStencil::DepthRead_StencilRead
+	);
+		
+	RHICmdList.BeginRenderPass(RPInfo, TEXT("RenderInverseOpacity"));
+
 	if (ShouldRenderTranslucency(ETranslucencyPass::TPT_AllTranslucency))
-	{
+	{		
 		const bool bGammaSpace = !IsMobileHDR();
-		const bool bLinearHDR64 = !bGammaSpace && !IsMobileHDR32bpp();
-		
-		if (!bGammaSpace)
-		{
-			SceneContext.BeginRenderingTranslucency(RHICmdList, View, *this);
-		}
-		else
+		if (bGammaSpace)
 		{
-			SceneContext.BeginRenderingSceneColor(RHICmdList, ESimpleRenderTargetMode::EClearColorExistingDepth);
 			// Mobile multi-view is not side by side stereo
 			const FViewInfo& TranslucentViewport = (View.bIsMobileMultiViewEnabled) ? Views[0] : View;
 			RHICmdList.SetViewport(TranslucentViewport.ViewRect.Min.X, TranslucentViewport.ViewRect.Min.Y, 0.0f, TranslucentViewport.ViewRect.Max.X, TranslucentViewport.ViewRect.Max.Y, 1.0f);
@@ -264,23 +275,10 @@ bool FMobileSceneRenderer::RenderInverseOpacity(FRHICommandListImmediate& RHICmd
 		View.ParallelMeshDrawCommandPasses[EMeshPass::MobileInverseOpacity].DispatchDraw(nullptr, RHICmdList);
 				
 		bDirty |= View.ParallelMeshDrawCommandPasses[EMeshPass::MobileInverseOpacity].HasAnyDraw();
-
-		if (!bGammaSpace)
-		{
-			RHICmdList.EndRenderPass();
-			SceneContext.FinishRenderingTranslucency(RHICmdList);
-		}
-		else
-		{
-			SceneContext.FinishRenderingSceneColor(RHICmdList);
-		}
-	}
-	else
-	{
-		// This is to preserve the previous behavior.
-		SceneContext.BeginRenderingSceneColor(RHICmdList, ESimpleRenderTargetMode::EClearColorExistingDepth);
-		SceneContext.FinishRenderingSceneColor(RHICmdList);
 	}
+	
+	RHICmdList.EndRenderPass();
+
 	return bDirty;
 }
 
