Remove magic scalar attached to ray tracing RectLight

diff --git a/Engine/Source/Runtime/Renderer/Private/PathTracing.cpp b/Engine/Source/Runtime/Renderer/Private/PathTracing.cpp
--- a/Engine/Source/Runtime/Renderer/Private/PathTracing.cpp
+++ b/Engine/Source/Runtime/Renderer/Private/PathTracing.cpp
@@ -197,8 +197,7 @@ public:
 						LightData.Normal[LightData.Count] = -LightParameters.Direction;
 						LightData.dPdu[LightData.Count] = FVector::CrossProduct(LightParameters.Tangent, LightParameters.Direction);
 						LightData.dPdv[LightData.Count] = LightParameters.Tangent;
-						// #dxr_todo: UE-72556  define these differences from Lit..
-						LightData.Color[LightData.Count] = LightParameters.Color / 4.0;
+						LightData.Color[LightData.Count] = LightParameters.Color;
 						LightData.Dimensions[LightData.Count] = FVector(2.0f * LightParameters.SourceRadius, 2.0f * LightParameters.SourceLength, 0.0f);
 						LightData.Attenuation[LightData.Count] = 1.0 / LightParameters.InvRadius;
 						LightData.RectLightBarnCosAngle[LightData.Count] = LightParameters.RectLightBarnCosAngle;
diff --git a/Engine/Source/Runtime/Renderer/Private/RayTracing/RayTracingGlobalIllumination.cpp b/Engine/Source/Runtime/Renderer/Private/RayTracing/RayTracingGlobalIllumination.cpp
--- a/Engine/Source/Runtime/Renderer/Private/RayTracing/RayTracingGlobalIllumination.cpp
+++ b/Engine/Source/Runtime/Renderer/Private/RayTracing/RayTracingGlobalIllumination.cpp
@@ -150,8 +150,7 @@ void SetupLightParameters(
 			LightParameters->Normal[LightParameters->Count] = -LightShaderParameters.Direction;
 			LightParameters->dPdu[LightParameters->Count] = FVector::CrossProduct(LightShaderParameters.Direction, LightShaderParameters.Tangent);
 			LightParameters->dPdv[LightParameters->Count] = LightShaderParameters.Tangent;
-			// #dxr_todo: UE-72556 define these differences from Lit..
-			LightParameters->Color[LightParameters->Count] = LightShaderParameters.Color / 4.0;
+			LightParameters->Color[LightParameters->Count] = LightShaderParameters.Color;
 			LightParameters->Dimensions[LightParameters->Count] = FVector(2.0f * LightShaderParameters.SourceRadius, 2.0f * LightShaderParameters.SourceLength, 0.0f);
 			LightParameters->Attenuation[LightParameters->Count] = 1.0 / LightShaderParameters.InvRadius;
 			break;
