- FVirtualTexturePhysicalSpace* now stored using TRefCountPtr<>, since code was changed to destroy previous unused physical spaces
- Otherwise we get into bad state since AllocatedVT could reference already deleted physical space, and cause system to get out of sync and/or crash
- Also defer deletion of unreferenced physical spaces until end of each render thread update, not strictly needed, but keeps large textures from getting destroyed/recreated in some cases

diff --git a/Engine/Source/Runtime/Renderer/Private/VT/AllocatedVirtualTexture.cpp b/Engine/Source/Runtime/Renderer/Private/VT/AllocatedVirtualTexture.cpp
--- a/Engine/Source/Runtime/Renderer/Private/VT/AllocatedVirtualTexture.cpp
+++ b/Engine/Source/Runtime/Renderer/Private/VT/AllocatedVirtualTexture.cpp
@@ -22,7 +22,6 @@ FAllocatedVirtualTexture::FAllocatedVirtualTexture(FVirtualTextureSystem* InSyst
 	, NumUniqueProducers(0u)
 {
 	check(IsInRenderingThread());
-	FMemory::Memzero(PhysicalSpace);
 	FMemory::Memzero(UniqueProducerHandles);
 	FMemory::Memzero(UniqueProducerMipBias);
 
@@ -130,6 +129,7 @@ void FAllocatedVirtualTexture::Release(FVirtualTextureSystem* System)
 			FTexturePageMap& PageMap = Space->GetPageMap(LayerIndex);
 			PhysicalSpace[LayerIndex]->GetPagePool().UnmapAllPagesForSpace(System, Space->GetID());
 			PageMap.VerifyPhysicalSpaceUnmapped(PhysicalSpace[LayerIndex]->GetID());
+			PhysicalSpace[LayerIndex].SafeRelease();
 		}
 	}
 
diff --git a/Engine/Source/Runtime/Renderer/Private/VT/AllocatedVirtualTexture.h b/Engine/Source/Runtime/Renderer/Private/VT/AllocatedVirtualTexture.h
--- a/Engine/Source/Runtime/Renderer/Private/VT/AllocatedVirtualTexture.h
+++ b/Engine/Source/Runtime/Renderer/Private/VT/AllocatedVirtualTexture.h
@@ -56,7 +56,7 @@ private:
 	uint32 AddUniqueProducer(const FVirtualTextureProducerHandle& InHandle, FVirtualTextureProducer* InProducer);
 
 	FVirtualTextureSpace* Space;
-	FVirtualTexturePhysicalSpace* PhysicalSpace[VIRTUALTEXTURE_SPACE_MAXLAYERS];
+	TRefCountPtr<FVirtualTexturePhysicalSpace> PhysicalSpace[VIRTUALTEXTURE_SPACE_MAXLAYERS];
 	FVirtualTextureProducerHandle UniqueProducerHandles[VIRTUALTEXTURE_SPACE_MAXLAYERS];
 	uint8 UniqueProducerIndexForLayer[VIRTUALTEXTURE_SPACE_MAXLAYERS];
 	uint8 UniqueProducerMipBias[VIRTUALTEXTURE_SPACE_MAXLAYERS];
diff --git a/Engine/Source/Runtime/Renderer/Private/VT/TexturePagePool.cpp b/Engine/Source/Runtime/Renderer/Private/VT/TexturePagePool.cpp
--- a/Engine/Source/Runtime/Renderer/Private/VT/TexturePagePool.cpp
+++ b/Engine/Source/Runtime/Renderer/Private/VT/TexturePagePool.cpp
@@ -54,14 +54,13 @@ void FTexturePagePool::EvictAllPages(FVirtualTextureSystem* System)
 
 	for (int32 i = 0; i < PagesToEvict.Num(); i++)
 	{
-		UnmapAllPages(System, PagesToEvict[i]);
+		UnmapAllPages(System, PagesToEvict[i], false);
 		FreeHeap.Add(0, PagesToEvict[i]);
 	}
 }
 
 void FTexturePagePool::UnmapAllPagesForSpace(FVirtualTextureSystem* System, uint8 SpaceID)
 {
-
 	// walk through all of our current mapping entries, and unmap any that belong to the current space
 	for (int32 MappingIndex = NumPages + 1; MappingIndex < PageMapping.Num(); ++MappingIndex)
 	{
@@ -82,7 +81,7 @@ void FTexturePagePool::EvictPages(FVirtualTextureSystem* System, const FVirtualT
 		const FPageEntry& PageEntry = Pages[pAddress];
 		if (PageEntry.PackedProducerHandle == ProducerHandle.PackedValue)
 		{
-			UnmapAllPages(System, pAddress);
+			UnmapAllPages(System, pAddress, false);
 			FreeHeap.Update(0, pAddress);
 		}
 	}
@@ -223,7 +222,7 @@ uint32 FTexturePagePool::Alloc(FVirtualTextureSystem* System, uint32 Frame, cons
 	const uint16 pAddress = FreeHeap.Top();
 
 	// Unmap any previous usage
-	UnmapAllPages(System, pAddress);
+	UnmapAllPages(System, pAddress, true);
 
 	// Mark the page as used for the given producer
 	FPageEntry& PageEntry = Pages[pAddress];
@@ -247,7 +246,7 @@ uint32 FTexturePagePool::Alloc(FVirtualTextureSystem* System, uint32 Frame, cons
 
 void FTexturePagePool::Free(FVirtualTextureSystem* System, uint16 pAddress)
 {
-	UnmapAllPages(System, pAddress);
+	UnmapAllPages(System, pAddress, true);
 
 	if (FreeHeap.IsPresent(pAddress))
 	{
@@ -313,7 +312,7 @@ void FTexturePagePool::UnmapPageMapping(FVirtualTextureSystem* System, uint32 Ma
 	ReleaseMapping(MappingIndex);
 }
 
-void FTexturePagePool::UnmapAllPages(FVirtualTextureSystem* System, uint16 pAddress)
+void FTexturePagePool::UnmapAllPages(FVirtualTextureSystem* System, uint16 pAddress, bool bMapAncestorPages)
 {
 	FPageEntry& PageEntry = Pages[pAddress];
 	if (PageEntry.PackedProducerHandle != 0u)
@@ -328,7 +327,7 @@ void FTexturePagePool::UnmapAllPages(FVirtualTextureSystem* System, uint16 pAddr
 	{
 		FPageMapping& Mapping = PageMapping[MappingIndex];
 		const uint32 NextIndex = Mapping.NextIndex;
-		UnmapPageMapping(System, MappingIndex, true);
+		UnmapPageMapping(System, MappingIndex, bMapAncestorPages);
 
 		MappingIndex = NextIndex;
 	}
diff --git a/Engine/Source/Runtime/Renderer/Private/VT/TexturePagePool.h b/Engine/Source/Runtime/Renderer/Private/VT/TexturePagePool.h
--- a/Engine/Source/Runtime/Renderer/Private/VT/TexturePagePool.h
+++ b/Engine/Source/Runtime/Renderer/Private/VT/TexturePagePool.h
@@ -161,7 +161,7 @@ private:
 	static uint16 GetPageHash(const FPageEntry& Entry);
 
 	void UnmapPageMapping(FVirtualTextureSystem* System, uint32 MappingIndex, bool bMapAncestorPage);
-	void UnmapAllPages(FVirtualTextureSystem* System, uint16 pAddress);
+	void UnmapAllPages(FVirtualTextureSystem* System, uint16 pAddress, bool bMapAncestorPages);
 
 	void RemoveMappingFromList(uint32 Index)
 	{
diff --git a/Engine/Source/Runtime/Renderer/Private/VT/VirtualTextureProducer.cpp b/Engine/Source/Runtime/Renderer/Private/VT/VirtualTextureProducer.cpp
--- a/Engine/Source/Runtime/Renderer/Private/VT/VirtualTextureProducer.cpp
+++ b/Engine/Source/Runtime/Renderer/Private/VT/VirtualTextureProducer.cpp
@@ -14,9 +14,8 @@ void FVirtualTextureProducer::Release(FVirtualTextureSystem* System, const FVirt
 	for (uint32 LayerIndex = 0u; LayerIndex < Description.NumLayers; ++LayerIndex)
 	{
 		FVirtualTexturePhysicalSpace* Space = PhysicalSpace[LayerIndex];
-		PhysicalSpace[LayerIndex] = nullptr;
 		Space->GetPagePool().EvictPages(System, HandleToSelf);
-		System->ReleasePhysicalSpace(Space);
+		PhysicalSpace[LayerIndex].SafeRelease();
 	}
 
 	delete VirtualTexture;
diff --git a/Engine/Source/Runtime/Renderer/Private/VT/VirtualTextureProducer.h b/Engine/Source/Runtime/Renderer/Private/VT/VirtualTextureProducer.h
--- a/Engine/Source/Runtime/Renderer/Private/VT/VirtualTextureProducer.h
+++ b/Engine/Source/Runtime/Renderer/Private/VT/VirtualTextureProducer.h
@@ -13,7 +13,7 @@ class FVirtualTexturePhysicalSpace;
 class FVirtualTextureProducer
 {
 public:
-	FVirtualTextureProducer() : VirtualTexture(nullptr) { FMemory::Memzero(PhysicalSpace); }
+	FVirtualTextureProducer() : VirtualTexture(nullptr) {}
 
 	void Release(FVirtualTextureSystem* System, const FVirtualTextureProducerHandle& HandleToSelf);
 
@@ -33,7 +33,7 @@ private:
 	~FVirtualTextureProducer() {}
 
 	IVirtualTexture* VirtualTexture;
-	FVirtualTexturePhysicalSpace* PhysicalSpace[VIRTUALTEXTURE_SPACE_MAXLAYERS];
+	TRefCountPtr<FVirtualTexturePhysicalSpace> PhysicalSpace[VIRTUALTEXTURE_SPACE_MAXLAYERS];
 	FVTProducerDescription Description;
 };
 
diff --git a/Engine/Source/Runtime/Renderer/Private/VT/VirtualTextureSystem.cpp b/Engine/Source/Runtime/Renderer/Private/VT/VirtualTextureSystem.cpp
--- a/Engine/Source/Runtime/Renderer/Private/VT/VirtualTextureSystem.cpp
+++ b/Engine/Source/Runtime/Renderer/Private/VT/VirtualTextureSystem.cpp
@@ -223,7 +223,7 @@ FVirtualTextureSystem::~FVirtualTextureSystem()
 	}
 	for(int i = 0; i < PhysicalSpaces.Num(); ++i)
 	{
-		FVirtualTexturePhysicalSpace* PhysicalSpace = PhysicalSpaces[i].Get();
+		FVirtualTexturePhysicalSpace* PhysicalSpace = PhysicalSpaces[i];
 		if (PhysicalSpace)
 		{
 			check(PhysicalSpace->GetRefCount() == 0u);
@@ -293,7 +293,7 @@ void FVirtualTextureSystem::ListPhysicalPoolsFromConsole()
 {
 	for(int i = 0; i < PhysicalSpaces.Num(); ++i)
 	{
-		if (PhysicalSpaces[i].IsValid())
+		if (PhysicalSpaces[i])
 		{
 			const FVirtualTexturePhysicalSpace& PhysicalSpace = *PhysicalSpaces[i];
 			const FVTPhysicalSpaceDescription& Desc = PhysicalSpace.GetDescription();
@@ -586,10 +586,9 @@ FVirtualTexturePhysicalSpace* FVirtualTextureSystem::AcquirePhysicalSpace(const
 {
 	for (int i = 0; i < PhysicalSpaces.Num(); ++i)
 	{
-		FVirtualTexturePhysicalSpace* PhysicalSpace = PhysicalSpaces[i].Get();
+		FVirtualTexturePhysicalSpace* PhysicalSpace = PhysicalSpaces[i];
 		if (PhysicalSpace && PhysicalSpace->GetDescription() == InDesc)
 		{
-			PhysicalSpace->AddRef();
 			return PhysicalSpace;
 		}
 	}
@@ -599,7 +598,7 @@ FVirtualTexturePhysicalSpace* FVirtualTextureSystem::AcquirePhysicalSpace(const
 
 	for (int i = 0; i < PhysicalSpaces.Num(); ++i)
 	{
-		if (!PhysicalSpaces[i].IsValid())
+		if (!PhysicalSpaces[i])
 		{
 			ID = i;
 			break;
@@ -608,27 +607,37 @@ FVirtualTexturePhysicalSpace* FVirtualTextureSystem::AcquirePhysicalSpace(const
 
 	if (ID == PhysicalSpaces.Num())
 	{
-		PhysicalSpaces.AddDefaulted_GetRef();
+		PhysicalSpaces.AddZeroed();
 	}
 
-	TUniquePtr<FVirtualTexturePhysicalSpace>& PhysicalSpace = PhysicalSpaces[ID];
-	PhysicalSpace.Reset(new FVirtualTexturePhysicalSpace(InDesc, ID));
+	FVirtualTexturePhysicalSpace* PhysicalSpace = new FVirtualTexturePhysicalSpace(InDesc, ID);
+	PhysicalSpaces[ID] = PhysicalSpace;
+
 	INC_MEMORY_STAT_BY(STAT_TotalPhysicalMemory, PhysicalSpace->GetSizeInBytes());
-	BeginInitResource(PhysicalSpace.Get());
-	PhysicalSpace->AddRef();
-	return PhysicalSpace.Get();
+	BeginInitResource(PhysicalSpace);
+	return PhysicalSpace;
 }
 
-void FVirtualTextureSystem::ReleasePhysicalSpace(FVirtualTexturePhysicalSpace* Space)
+void FVirtualTextureSystem::ReleasePendingSpaces()
 {
-	const uint32 NumRefs = Space->Release();
-	// Physical space is released when ref count hits 0
-	// Might need to have some mechanism to hold an extra reference if we know we will be recycling very soon (such when doing level reload)
-	if (NumRefs == 0)
+	check(IsInRenderingThread());
+	for (int32 Id = 0; Id < PhysicalSpaces.Num(); ++Id)
 	{
-		DEC_MEMORY_STAT_BY(STAT_TotalPhysicalMemory, Space->GetSizeInBytes());
-		Space->ReleaseResource();
-		PhysicalSpaces[Space->GetID()] = nullptr;
+		// Physical space is released when ref count hits 0
+		// Might need to have some mechanism to hold an extra reference if we know we will be recycling very soon (such when doing level reload)
+		FVirtualTexturePhysicalSpace* PhysicalSpace = PhysicalSpaces[Id];
+		if ((bool)PhysicalSpace && PhysicalSpace->GetRefCount() == 0u)
+		{
+			DEC_MEMORY_STAT_BY(STAT_TotalPhysicalMemory, PhysicalSpace->GetSizeInBytes());
+
+			const FTexturePagePool& PagePool = PhysicalSpace->GetPagePool();
+			check(PagePool.GetNumMappedPages() == 0u);
+			check(PagePool.GetNumLockedPages() == 0u);
+
+			PhysicalSpace->ReleaseResource();
+			delete PhysicalSpace;
+			PhysicalSpaces[Id] = nullptr;
+		}
 	}
 }
 
@@ -846,7 +855,7 @@ void FVirtualTextureSystem::Update(FRHICommandListImmediate& RHICmdList, ERHIFea
 
 		for (int i = 0; i < PhysicalSpaces.Num(); ++i)
 		{
-			FVirtualTexturePhysicalSpace* PhysicalSpace = PhysicalSpaces[i].Get();
+			FVirtualTexturePhysicalSpace* PhysicalSpace = PhysicalSpaces[i];
 			if (PhysicalSpace)
 			{
 				// Collect locked pages to be produced again
@@ -1048,6 +1057,8 @@ void FVirtualTextureSystem::Update(FRHICommandListImmediate& RHICmdList, ERHIFea
 	SubmitPreMappedRequests(RHICmdList, FeatureLevel);
 	// Submit the merged requests
 	SubmitRequests(RHICmdList, FeatureLevel, MemStack, MergedRequestList, true);
+
+	ReleasePendingSpaces();
 }
 
 void FVirtualTextureSystem::GatherRequests(FUniqueRequestList* MergedRequestList, const FUniquePageList* UniquePageList, uint32 FrameRequested, FMemStack& MemStack)
@@ -1535,7 +1546,7 @@ void FVirtualTextureSystem::GatherRequestsTask(const FGatherRequestsParameters&
 
 	for (uint32 PhysicalSpaceID = 0u; PhysicalSpaceID < (uint32)PhysicalSpaces.Num(); ++PhysicalSpaceID)
 	{
-		if (PhysicalSpaces[PhysicalSpaceID].Get() == nullptr)
+		if (PhysicalSpaces[PhysicalSpaceID] == nullptr)
 		{
 			continue;
 		}
diff --git a/Engine/Source/Runtime/Renderer/Private/VT/VirtualTextureSystem.h b/Engine/Source/Runtime/Renderer/Private/VT/VirtualTextureSystem.h
--- a/Engine/Source/Runtime/Renderer/Private/VT/VirtualTextureSystem.h
+++ b/Engine/Source/Runtime/Renderer/Private/VT/VirtualTextureSystem.h
@@ -52,10 +52,9 @@ public:
 	void ReleaseSpace(FVirtualTextureSpace* Space);
 
 	FVirtualTexturePhysicalSpace* AcquirePhysicalSpace(const FVTPhysicalSpaceDescription& InDesc);
-	void ReleasePhysicalSpace(FVirtualTexturePhysicalSpace* Space);
 
 	FVirtualTextureSpace* GetSpace(uint8 ID) const { check(ID < MaxSpaces); return Spaces[ID].Get(); }
-	FVirtualTexturePhysicalSpace* GetPhysicalSpace(uint16 ID) const { check(PhysicalSpaces[ID].IsValid());  return PhysicalSpaces[ID].Get(); }
+	FVirtualTexturePhysicalSpace* GetPhysicalSpace(uint16 ID) const { check(PhysicalSpaces[ID]);  return PhysicalSpaces[ID]; }
 
 	void LockTile(const FVirtualTextureLocalTile& Tile);
 	void UnlockTile(const FVirtualTextureLocalTile& Tile, const FVirtualTextureProducer* Producer);
@@ -75,6 +74,7 @@ private:
 	~FVirtualTextureSystem();
 
 	void DestroyPendingVirtualTextures();
+	void ReleasePendingSpaces();
 
 	void RequestTilesForRegionInternal(const IAllocatedVirtualTexture* AllocatedVT, const FIntRect& InTextureRegion, uint32 vLevel);
 	
@@ -95,7 +95,7 @@ private:
 
 	static const uint32 MaxSpaces = 16;
 	TUniquePtr<FVirtualTextureSpace> Spaces[MaxSpaces];
-	TArray< TUniquePtr<FVirtualTexturePhysicalSpace> > PhysicalSpaces;
+	TArray<FVirtualTexturePhysicalSpace*> PhysicalSpaces;
 	FVirtualTextureProducerCollection Producers;
 
 	FCriticalSection PendingDeleteLock;
