Fix non-unity build

diff --git a/Engine/Source/Runtime/Renderer/Private/VT/VirtualTextureProducer.cpp b/Engine/Source/Runtime/Renderer/Private/VT/VirtualTextureProducer.cpp
--- a/Engine/Source/Runtime/Renderer/Private/VT/VirtualTextureProducer.cpp
+++ b/Engine/Source/Runtime/Renderer/Private/VT/VirtualTextureProducer.cpp
@@ -4,6 +4,10 @@
 #include "VirtualTextureSystem.h"
 #include "VirtualTexturePhysicalSpace.h"
 
+FVirtualTextureProducer::~FVirtualTextureProducer()
+{
+}
+
 void FVirtualTextureProducer::Release(FVirtualTextureSystem* System, const FVirtualTextureProducerHandle& HandleToSelf)
 {
 	if (Description.bPersistentHighestMip)
diff --git a/Engine/Source/Runtime/Renderer/Private/VT/VirtualTextureProducer.h b/Engine/Source/Runtime/Renderer/Private/VT/VirtualTextureProducer.h
--- a/Engine/Source/Runtime/Renderer/Private/VT/VirtualTextureProducer.h
+++ b/Engine/Source/Runtime/Renderer/Private/VT/VirtualTextureProducer.h
@@ -30,7 +30,7 @@ public:
 
 private:
 	friend class FVirtualTextureProducerCollection;
-	~FVirtualTextureProducer() {}
+	~FVirtualTextureProducer();
 
 	IVirtualTexture* VirtualTexture;
 	TRefCountPtr<FVirtualTexturePhysicalSpace> PhysicalSpace[VIRTUALTEXTURE_SPACE_MAXLAYERS];
