Fix build error introduced by merge of 7882089 and 8012534

diff --git a/Engine/Source/Runtime/Renderer/Private/VT/TexturePagePool.cpp b/Engine/Source/Runtime/Renderer/Private/VT/TexturePagePool.cpp
--- a/Engine/Source/Runtime/Renderer/Private/VT/TexturePagePool.cpp
+++ b/Engine/Source/Runtime/Renderer/Private/VT/TexturePagePool.cpp
@@ -113,7 +113,7 @@ void FTexturePagePool::EvictPages(FVirtualTextureSystem* System, FVirtualTexture
 				}
 				else
 				{
-					UnmapAllPages(System, i);
+					UnmapAllPages(System, i, false);
 					FreeHeap.Update(0, i);
 				}
 			}
