Cherry-pick CL 8041599

Decouple ray tracing ambient occlusion and global illumination passes.

The ray tracing global illumination-computed ambient occlusion render target is undefined where G-buffer Albedo is < DiffuseThreshold.

diff --git a/Engine/Source/Runtime/Renderer/Private/RayTracing/RayTracingAmbientOcclusion.cpp b/Engine/Source/Runtime/Renderer/Private/RayTracing/RayTracingAmbientOcclusion.cpp
--- a/Engine/Source/Runtime/Renderer/Private/RayTracing/RayTracingAmbientOcclusion.cpp
+++ b/Engine/Source/Runtime/Renderer/Private/RayTracing/RayTracingAmbientOcclusion.cpp
@@ -62,7 +62,7 @@ bool ShouldRenderRayTracingAmbientOcclusion(const FViewInfo& View)
 	const bool bRTAOEnabled = (ForceAllRayTracingEffects > 0 || (GRayTracingAmbientOcclusion > 0 && ForceAllRayTracingEffects < 0));
 
 	//#dxr_todo: add option to enable RTAO in View.FinalPostProcessSettings
-	return IsRayTracingEnabled() && !ShouldRenderRayTracingGlobalIllumination(View) && bRTAOEnabled;
+	return IsRayTracingEnabled() && bRTAOEnabled;
 }
 
 DECLARE_GPU_STAT_NAMED(RayTracingAmbientOcclusion, TEXT("Ray Tracing Ambient Occlusion"));
