Try again to fix non-unity compile

diff --git a/Engine/Source/Runtime/Renderer/Private/VT/AllocatedVirtualTexture.h b/Engine/Source/Runtime/Renderer/Private/VT/AllocatedVirtualTexture.h
--- a/Engine/Source/Runtime/Renderer/Private/VT/AllocatedVirtualTexture.h
+++ b/Engine/Source/Runtime/Renderer/Private/VT/AllocatedVirtualTexture.h
@@ -6,6 +6,7 @@
 #include "VirtualTextureShared.h"
 #include "RendererInterface.h"
 #include "VirtualTexturing.h"
+#include "VirtualTexturePhysicalSpace.h"
 #include "HAL/ThreadSafeCounter.h"
 
 class FVirtualTextureSpace;
diff --git a/Engine/Source/Runtime/Renderer/Private/VT/VirtualTextureProducer.h b/Engine/Source/Runtime/Renderer/Private/VT/VirtualTextureProducer.h
--- a/Engine/Source/Runtime/Renderer/Private/VT/VirtualTextureProducer.h
+++ b/Engine/Source/Runtime/Renderer/Private/VT/VirtualTextureProducer.h
@@ -6,6 +6,7 @@
 #include "VirtualTextureShared.h"
 #include "RendererInterface.h"
 #include "VirtualTexturing.h"
+#include "VirtualTexturePhysicalSpace.h"
 
 class FVirtualTextureSystem;
 class FVirtualTexturePhysicalSpace;
