Cherry-pick 8047649: Added option to enable RTAO in post process volume settings

diff --git a/Engine/Source/Runtime/Engine/Classes/Engine/Scene.h b/Engine/Source/Runtime/Engine/Classes/Engine/Scene.h
--- a/Engine/Source/Runtime/Engine/Classes/Engine/Scene.h
+++ b/Engine/Source/Runtime/Engine/Classes/Engine/Scene.h
@@ -902,6 +902,9 @@ struct FPostProcessSettings
 	UPROPERTY(EditAnywhere, BlueprintReadWrite, Category=Overrides, meta=(PinHiddenByDefault, InlineEditConditionToggle))
 	uint8 bOverride_AmbientOcclusionMipThreshold:1;
 
+	UPROPERTY(EditAnywhere, BlueprintReadWrite, Category = Overrides, meta = (PinHiddenByDefault, InlineEditConditionToggle))
+	uint32 bOverride_RayTracingAO : 1;
+
 	UPROPERTY(EditAnywhere, BlueprintReadWrite, Category = Overrides, meta = (PinHiddenByDefault, InlineEditConditionToggle))
 	uint32 bOverride_RayTracingAOSamplesPerPixel : 1;
 
@@ -1544,7 +1547,11 @@ struct FPostProcessSettings
 	UPROPERTY(interp, BlueprintReadWrite, Category="Rendering Features|Ambient Occlusion", AdvancedDisplay, meta=(ClampMin = "0.0", UIMax = "0.1", editcondition = "bOverride_AmbientOcclusionMipThreshold", DisplayName = "Mip Threshold"))
 	float AmbientOcclusionMipThreshold;
 
-	/** Sets the samples per pixel for ray tracing global illumination. */
+	/** Enables ray tracing ambient occlusion. */
+	UPROPERTY(EditAnywhere, BlueprintReadWrite, Category = "Rendering Features|Ray Tracing Ambient Occlusion", meta = (editcondition = "bOverride_RayTracingAO", DisplayName = "Enabled"))
+	uint32 RayTracingAO : 1;
+
+	/** Sets the samples per pixel for ray tracing ambient occlusion. */
 	UPROPERTY(EditAnywhere, BlueprintReadWrite, Category = "Rendering Features|Ray Tracing Ambient Occlusion", meta = (ClampMin = "1", ClampMax = "64", editcondition = "bOverride_RayTracingAOSamplesPerPixel", DisplayName = "Samples Per Pixel"))
 	int32 RayTracingAOSamplesPerPixel;
 
diff --git a/Engine/Source/Runtime/Engine/Private/Scene.cpp b/Engine/Source/Runtime/Engine/Private/Scene.cpp
--- a/Engine/Source/Runtime/Engine/Private/Scene.cpp
+++ b/Engine/Source/Runtime/Engine/Private/Scene.cpp
@@ -508,6 +508,7 @@ FPostProcessSettings::FPostProcessSettings()
 	AmbientOcclusionMipScale = 1.7f;
 	AmbientOcclusionMipThreshold = 0.01f;
 	AmbientOcclusionRadiusInWS = false;
+	RayTracingAO = 1;
 	RayTracingAOSamplesPerPixel = 1;
 	IndirectLightingColor = FLinearColor(1.0f, 1.0f, 1.0f);
 	IndirectLightingIntensity = 1.0f;
@@ -690,6 +691,7 @@ FPostProcessSettings::FPostProcessSettings(const FPostProcessSettings& Settings)
 	, bOverride_AmbientOcclusionMipBlend(Settings.bOverride_AmbientOcclusionMipBlend)
 	, bOverride_AmbientOcclusionMipScale(Settings.bOverride_AmbientOcclusionMipScale)
 	, bOverride_AmbientOcclusionMipThreshold(Settings.bOverride_AmbientOcclusionMipThreshold)
+	, bOverride_RayTracingAO(Settings.bOverride_RayTracingAO)
 	, bOverride_RayTracingAOSamplesPerPixel(Settings.bOverride_RayTracingAOSamplesPerPixel)
 	, bOverride_LPVIntensity(Settings.bOverride_LPVIntensity)
 	, bOverride_LPVDirectionalOcclusionIntensity(Settings.bOverride_LPVDirectionalOcclusionIntensity)
@@ -868,6 +870,7 @@ FPostProcessSettings::FPostProcessSettings(const FPostProcessSettings& Settings)
 	, AmbientOcclusionMipBlend(Settings.AmbientOcclusionMipBlend)
 	, AmbientOcclusionMipScale(Settings.AmbientOcclusionMipScale)
 	, AmbientOcclusionMipThreshold(Settings.AmbientOcclusionMipThreshold)
+	, RayTracingAO(Settings.RayTracingAO)
 	, RayTracingAOSamplesPerPixel(Settings.RayTracingAOSamplesPerPixel)
 	, IndirectLightingColor(Settings.IndirectLightingColor)
 	, IndirectLightingIntensity(Settings.IndirectLightingIntensity)
diff --git a/Engine/Source/Runtime/Engine/Private/SceneView.cpp b/Engine/Source/Runtime/Engine/Private/SceneView.cpp
--- a/Engine/Source/Runtime/Engine/Private/SceneView.cpp
+++ b/Engine/Source/Runtime/Engine/Private/SceneView.cpp
@@ -1438,10 +1438,16 @@ void FSceneView::OverridePostProcessSettings(const FPostProcessSettings& Src, fl
 			Dest.RayTracingGISamplesPerPixel = Src.RayTracingGISamplesPerPixel;
 		}
 
+		if (Src.bOverride_RayTracingAO)
+		{
+			Dest.RayTracingAO = Src.RayTracingAO;
+		}
+
 		if (Src.bOverride_RayTracingAOSamplesPerPixel)
 		{
 			Dest.RayTracingAOSamplesPerPixel = Src.RayTracingAOSamplesPerPixel;
 		}
+
 		if (Src.bOverride_PathTracingMaxBounces)
 		{
 			Dest.PathTracingMaxBounces = Src.PathTracingMaxBounces;
diff --git a/Engine/Source/Runtime/Renderer/Private/RayTracing/RayTracingAmbientOcclusion.cpp b/Engine/Source/Runtime/Renderer/Private/RayTracing/RayTracingAmbientOcclusion.cpp
--- a/Engine/Source/Runtime/Renderer/Private/RayTracing/RayTracingAmbientOcclusion.cpp
+++ b/Engine/Source/Runtime/Renderer/Private/RayTracing/RayTracingAmbientOcclusion.cpp
@@ -19,11 +19,13 @@
 #include "PostProcess/PostProcessing.h"
 #include "PostProcess/SceneFilterRendering.h"
 
-static int32 GRayTracingAmbientOcclusion = 1;
+static int32 GRayTracingAmbientOcclusion = -1;
 static FAutoConsoleVariableRef CVarRayTracingAmbientOcclusion(
 	TEXT("r.RayTracing.AmbientOcclusion"),
 	GRayTracingAmbientOcclusion,
-	TEXT("Enables ray tracing ambient occlusion (default = 1)")
+	TEXT("-1: Value driven by postprocess volume (default) \n")
+	TEXT(" 0: ray tracing ambient occlusion off \n")
+	TEXT(" 1: ray tracing ambient occlusion enabled")
 );
 
 static TAutoConsoleVariable<int32> CVarUseAODenoiser(
@@ -58,11 +60,22 @@ static TAutoConsoleVariable<int32> CVarRayTracingAmbientOcclusionEnableMaterials
 
 bool ShouldRenderRayTracingAmbientOcclusion(const FViewInfo& View)
 {
-	const int32 ForceAllRayTracingEffects = GetForceRayTracingEffectsCVarValue();
-	const bool bRTAOEnabled = (ForceAllRayTracingEffects > 0 || (GRayTracingAmbientOcclusion > 0 && ForceAllRayTracingEffects < 0));
-
-	//#dxr_todo: add option to enable RTAO in View.FinalPostProcessSettings
-	return IsRayTracingEnabled() && bRTAOEnabled;
+	if (!IsRayTracingEnabled())
+	{
+		return (false);
+	}
+	else if (GetForceRayTracingEffectsCVarValue() >= 0)
+	{
+		return GetForceRayTracingEffectsCVarValue() > 0;
+	}
+	else if (GRayTracingAmbientOcclusion >= 0)
+	{
+		return (GRayTracingAmbientOcclusion > 0);
+	}
+	else
+	{
+		return View.FinalPostProcessSettings.RayTracingAO > 0;
+	}
 }
 
 DECLARE_GPU_STAT_NAMED(RayTracingAmbientOcclusion, TEXT("Ray Tracing Ambient Occlusion"));
diff --git a/Engine/Source/Runtime/Renderer/Private/RayTracing/RayTracingGlobalIllumination.cpp b/Engine/Source/Runtime/Renderer/Private/RayTracing/RayTracingGlobalIllumination.cpp
--- a/Engine/Source/Runtime/Renderer/Private/RayTracing/RayTracingGlobalIllumination.cpp
+++ b/Engine/Source/Runtime/Renderer/Private/RayTracing/RayTracingGlobalIllumination.cpp
@@ -27,7 +27,7 @@ static FAutoConsoleVariableRef CVarRayTracingGlobalIllumination(
 	TEXT("r.RayTracing.GlobalIllumination"),
 	GRayTracingGlobalIllumination,
 	TEXT("-1: Value driven by postprocess volume (default) \n")
-	TEXT(" 0: ray tracing ray tracing global illumination off \n")
+	TEXT(" 0: ray tracing global illumination off \n")
 	TEXT(" 1: ray tracing global illumination enabled")
 );
 
@@ -187,7 +187,11 @@ void SetupLightParameters(
 
 bool ShouldRenderRayTracingGlobalIllumination(const FViewInfo& View)
 {
-	if (GetForceRayTracingEffectsCVarValue() >= 0)
+	if (!IsRayTracingEnabled())
+	{
+		return (false);
+	}
+	else if (GetForceRayTracingEffectsCVarValue() >= 0)
 	{
 		return GetForceRayTracingEffectsCVarValue() > 0;
 	}
