Don't pay cost of VT feedback analysis if no VTs are allocated

diff --git a/Engine/Source/Runtime/Renderer/Private/VT/VirtualTextureSystem.cpp b/Engine/Source/Runtime/Renderer/Private/VT/VirtualTextureSystem.cpp
--- a/Engine/Source/Runtime/Renderer/Private/VT/VirtualTextureSystem.cpp
+++ b/Engine/Source/Runtime/Renderer/Private/VT/VirtualTextureSystem.cpp
@@ -870,6 +870,13 @@ void FVirtualTextureSystem::Update(FRHICommandListImmediate& RHICmdList, ERHIFea
 
 	DestroyPendingVirtualTextures();
 
+	// Early out when no allocated VTs
+	if (AllocatedVTs.Num() == 0)
+	{
+		MappedTilesToProduce.Reset();
+		return;
+	}
+
 	FMemStack& MemStack = FMemStack::Get();
 	FMemMark Mark(MemStack);
 	FUniquePageList* MergedUniquePageList = new(MemStack) FUniquePageList;
