Don't allocate the runtime virtual texture in the proxy if the asset is disabled

diff --git a/Engine/Source/Runtime/Renderer/Private/VT/RuntimeVirtualTextureSceneProxy.cpp b/Engine/Source/Runtime/Renderer/Private/VT/RuntimeVirtualTextureSceneProxy.cpp
--- a/Engine/Source/Runtime/Renderer/Private/VT/RuntimeVirtualTextureSceneProxy.cpp
+++ b/Engine/Source/Runtime/Renderer/Private/VT/RuntimeVirtualTextureSceneProxy.cpp
@@ -10,12 +10,15 @@
 int32 FRuntimeVirtualTextureSceneProxy::ProducerIdGenerator = 1;
 
 FRuntimeVirtualTextureSceneProxy::FRuntimeVirtualTextureSceneProxy(URuntimeVirtualTextureComponent* InComponent)
-	: ProducerId(0)
-	, VirtualTexture(InComponent->GetVirtualTexture())
+	: SceneIndex(0)
+	, ProducerId(0)
+	, VirtualTexture(nullptr)
 	, CombinedDirtyRect(0, 0, 0, 0)
 {
-	if (VirtualTexture != nullptr)
+	if (InComponent->GetVirtualTexture() != nullptr && InComponent->GetVirtualTexture()->GetEnabled())
 	{
+		VirtualTexture = InComponent->GetVirtualTexture();
+
 		// We store a ProducerId here so that we will be able to find our SceneIndex from the Producer during rendering.
 		// We will need the SceneIndex to determine which primitives should render to this Producer.
 		ProducerId = ProducerIdGenerator++;
