Added csv stats to track RT time easier

#rnx

diff --git a/Engine/Source/Runtime/Engine/Private/TextureRenderTarget.cpp b/Engine/Source/Runtime/Engine/Private/TextureRenderTarget.cpp
--- a/Engine/Source/Runtime/Engine/Private/TextureRenderTarget.cpp
+++ b/Engine/Source/Runtime/Engine/Private/TextureRenderTarget.cpp
@@ -128,6 +128,7 @@ void FDeferredUpdateResource::UpdateResources(FRHICommandListImmediate& RHICmdLi
 
 			if( RTResource )
 			{
+				CSV_SCOPED_TIMING_STAT_EXCLUSIVE(FlushDeferredResourceUpdate);
 				RTResource->FlushDeferredResourceUpdate(RHICmdList);
 			}
 		}
diff --git a/Engine/Source/Runtime/RenderCore/Private/RenderResource.cpp b/Engine/Source/Runtime/RenderCore/Private/RenderResource.cpp
--- a/Engine/Source/Runtime/RenderCore/Private/RenderResource.cpp
+++ b/Engine/Source/Runtime/RenderCore/Private/RenderResource.cpp
@@ -83,6 +83,7 @@ void FRenderResource::InitResource()
 #endif
 		if(GIsRHIInitialized)
 		{
+			CSV_SCOPED_TIMING_STAT_EXCLUSIVE(InitRenderResource);
 			InitDynamicRHI();
 			InitRHI();
 		}
diff --git a/Engine/Source/Runtime/Renderer/Private/CapsuleShadowRendering.cpp b/Engine/Source/Runtime/Renderer/Private/CapsuleShadowRendering.cpp
--- a/Engine/Source/Runtime/Renderer/Private/CapsuleShadowRendering.cpp
+++ b/Engine/Source/Runtime/Renderer/Private/CapsuleShadowRendering.cpp
@@ -1303,6 +1303,7 @@ void FDeferredShadingSceneRenderer::RenderIndirectCapsuleShadows(
 		&& ViewFamily.EngineShowFlags.CapsuleShadows
 		&& FSceneRenderTargets::Get(RHICmdList).IsStaticLightingAllowed())
 	{
+		CSV_SCOPED_TIMING_STAT_EXCLUSIVE(RenderIndirectCapsuleShadows);
 		QUICK_SCOPE_CYCLE_COUNTER(STAT_RenderIndirectCapsuleShadows);
 
 		bool bAnyViewsUseCapsuleShadows = false;
diff --git a/Engine/Source/Runtime/Renderer/Private/DeferredShadingRenderer.cpp b/Engine/Source/Runtime/Renderer/Private/DeferredShadingRenderer.cpp
--- a/Engine/Source/Runtime/Renderer/Private/DeferredShadingRenderer.cpp
+++ b/Engine/Source/Runtime/Renderer/Private/DeferredShadingRenderer.cpp
@@ -1147,7 +1147,8 @@ void FDeferredShadingSceneRenderer::Render(FRHICommandListImmediate& RHICmdList)
 	// Before starting the render, all async task for the Custom data must be completed
 	if (UpdateViewCustomDataEvents.Num() > 0)
 	{
-		QUICK_SCOPE_CYCLE_COUNTER(STAT_FDeferredShadingSceneRenderer_AsyncUpdateViewCustomData_Wait);
+		QUICK_SCOPE_CYCLE_COUNTER(STAT_FDeferndershaddShadingSceneRenderer_AsyncUpdateViewCustomData_Wait);
+		CSV_SCOPED_TIMING_STAT_EXCLUSIVE(AsyncUpdateViewCustomData_Wait);
 		FTaskGraphInterface::Get().WaitUntilTasksComplete(UpdateViewCustomDataEvents, ENamedThreads::GetRenderThread());
 	}
 
@@ -1298,6 +1299,7 @@ void FDeferredShadingSceneRenderer::Render(FRHICommandListImmediate& RHICmdList)
 
 	if (bDBuffer || IsForwardShadingEnabled(ShaderPlatform))
 	{
+		CSV_SCOPED_TIMING_STAT_EXCLUSIVE(DeferredShadingSceneRenderer_DBuffer);
 		SCOPE_CYCLE_COUNTER(STAT_FDeferredShadingSceneRenderer_DBuffer);
 
 		// e.g. DBuffer deferred decals
diff --git a/Engine/Source/Runtime/Renderer/Private/GPUScene.cpp b/Engine/Source/Runtime/Renderer/Private/GPUScene.cpp
--- a/Engine/Source/Runtime/Renderer/Private/GPUScene.cpp
+++ b/Engine/Source/Runtime/Renderer/Private/GPUScene.cpp
@@ -155,6 +155,7 @@ void UpdateGPUScene(FRHICommandList& RHICmdList, FScene& Scene)
 {
 	if (UseGPUScene(GMaxRHIShaderPlatform, Scene.GetFeatureLevel()))
 	{
+		CSV_SCOPED_TIMING_STAT_EXCLUSIVE(UpdateGPUScene);
 		QUICK_SCOPE_CYCLE_COUNTER(STAT_UpdateGPUScene);
 
 		if (GGPUSceneUploadEveryFrame || Scene.GPUScene.bUpdateAllPrimitives)
@@ -303,6 +304,7 @@ void UploadDynamicPrimitiveShaderDataForView(FRHICommandList& RHICmdList, FScene
 {
 	if (UseGPUScene(GMaxRHIShaderPlatform, Scene.GetFeatureLevel()))
 	{
+		CSV_SCOPED_TIMING_STAT_EXCLUSIVE(UploadDynamicPrimitiveShaderDataForView);
 		QUICK_SCOPE_CYCLE_COUNTER(STAT_UploadDynamicPrimitiveShaderDataForView);
 
 		const int32 NumPrimitiveDataUploads = View.DynamicPrimitiveShaderData.Num();
diff --git a/Engine/Source/Runtime/Renderer/Private/LightGridInjection.cpp b/Engine/Source/Runtime/Renderer/Private/LightGridInjection.cpp
--- a/Engine/Source/Runtime/Renderer/Private/LightGridInjection.cpp
+++ b/Engine/Source/Runtime/Renderer/Private/LightGridInjection.cpp
@@ -805,6 +805,7 @@ void FDeferredShadingSceneRenderer::RenderForwardShadingShadowProjections(FRHICo
 
 	if (bScreenShadowMaskNeeded)
 	{
+		CSV_SCOPED_TIMING_STAT_EXCLUSIVE(RenderForwardShadingShadowProjections);
 		FSceneRenderTargets& SceneRenderTargets = FSceneRenderTargets::Get(RHICmdList);
 		SceneRenderTargets.AllocateScreenShadowMask(RHICmdList, ForwardScreenSpaceShadowMask);
 
diff --git a/Engine/Source/Runtime/Renderer/Private/RendererScene.cpp b/Engine/Source/Runtime/Renderer/Private/RendererScene.cpp
--- a/Engine/Source/Runtime/Renderer/Private/RendererScene.cpp
+++ b/Engine/Source/Runtime/Renderer/Private/RendererScene.cpp
@@ -1291,6 +1291,7 @@ static FAutoConsoleVariableRef CVarWarningOnRedundantTransformUpdate(
 
 void FScene::UpdatePrimitiveTransform_RenderThread(FRHICommandListImmediate& RHICmdList, FPrimitiveSceneProxy* PrimitiveSceneProxy, const FBoxSphereBounds& WorldBounds, const FBoxSphereBounds& LocalBounds, const FMatrix& LocalToWorld, const FVector& AttachmentRootPosition)
 {
+	CSV_SCOPED_TIMING_STAT_EXCLUSIVE(UpdatePrimitiveTransform);
 	SCOPE_CYCLE_COUNTER(STAT_UpdatePrimitiveTransformRenderThreadTime);
 
 	FPrimitiveSceneInfo* PrimitiveSceneInfo = PrimitiveSceneProxy->GetPrimitiveSceneInfo();
diff --git a/Engine/Source/Runtime/Renderer/Private/SceneRendering.cpp b/Engine/Source/Runtime/Renderer/Private/SceneRendering.cpp
--- a/Engine/Source/Runtime/Renderer/Private/SceneRendering.cpp
+++ b/Engine/Source/Runtime/Renderer/Private/SceneRendering.cpp
@@ -3244,7 +3244,9 @@ static void ViewExtensionPreRender_RenderThread(FRHICommandListImmediate& RHICmd
 	FMemMark MemStackMark(FMemStack::Get());
 
 	{
+		CSV_SCOPED_TIMING_STAT_EXCLUSIVE(PreRender);
 		SCOPE_CYCLE_COUNTER(STAT_FDeferredShadingSceneRenderer_ViewExtensionPreRenderView);
+
 		for (int ViewExt = 0; ViewExt < SceneRenderer->ViewFamily.ViewExtensions.Num(); ViewExt++)
 		{
 			SceneRenderer->ViewFamily.ViewExtensions[ViewExt]->PreRenderViewFamily_RenderThread(RHICmdList, SceneRenderer->ViewFamily);
