Condition for detecting locked pages in FTexturePagePool::EvictPages was inverted :(
Fixes crash in FTexturePagePool::EvictPages()
Fixes check triggering in FVirtualTextureSystem::SubmitRequestsFromLocalTileList()

diff --git a/Engine/Source/Runtime/Renderer/Private/VT/TexturePagePool.cpp b/Engine/Source/Runtime/Renderer/Private/VT/TexturePagePool.cpp
--- a/Engine/Source/Runtime/Renderer/Private/VT/TexturePagePool.cpp
+++ b/Engine/Source/Runtime/Renderer/Private/VT/TexturePagePool.cpp
@@ -107,7 +107,7 @@ void FTexturePagePool::EvictPages(FVirtualTextureSystem* System, FVirtualTexture
 
 			if (!(PageRect.Min.X > TextureRegion.Max.X) && !(TextureRegion.Min.X > PageRect.Max.X) && !(PageRect.Min.Y > TextureRegion.Max.Y) && !(TextureRegion.Min.Y > PageRect.Max.Y))
 			{
-				if (FreeHeap.IsPresent(i))
+				if (!FreeHeap.IsPresent(i))
 				{
 					OutLocked.Add(FVirtualTextureLocalTile(ProducerHandle, vAddress, vLevel));
 				}
@@ -123,7 +123,6 @@ void FTexturePagePool::EvictPages(FVirtualTextureSystem* System, FVirtualTexture
 
 void FTexturePagePool::GetAllLockedPages(FVirtualTextureSystem* System, TSet<FVirtualTextureLocalTile>& OutPages)
 {
-
 	OutPages.Reserve(OutPages.Num() + GetNumLockedPages());
 
 	for (uint32 i = 0; i < NumPages; ++i)
