Default settings for Ray Tracing Global Illumination changed to 4 SPP and 50% screen percentage

diff --git a/Engine/Source/Runtime/Engine/Private/Scene.cpp b/Engine/Source/Runtime/Engine/Private/Scene.cpp
--- a/Engine/Source/Runtime/Engine/Private/Scene.cpp
+++ b/Engine/Source/Runtime/Engine/Private/Scene.cpp
@@ -515,7 +515,7 @@ FPostProcessSettings::FPostProcessSettings()
 	ColorGradingIntensity = 1.0f;
 	RayTracingGI = 0;
 	RayTracingGIMaxBounces = 1;
-	RayTracingGISamplesPerPixel = 1;
+	RayTracingGISamplesPerPixel = 4;
 
 	DepthOfFieldFocalDistance = 0; // Intentionally invalid to disable DOF by default.
 	DepthOfFieldFstop = 4.0f; 
diff --git a/Engine/Source/Runtime/Renderer/Private/RayTracing/RayTracingGlobalIllumination.cpp b/Engine/Source/Runtime/Renderer/Private/RayTracing/RayTracingGlobalIllumination.cpp
--- a/Engine/Source/Runtime/Renderer/Private/RayTracing/RayTracingGlobalIllumination.cpp
+++ b/Engine/Source/Runtime/Renderer/Private/RayTracing/RayTracingGlobalIllumination.cpp
@@ -91,7 +91,7 @@ static FAutoConsoleVariableRef CVarRayTracingGlobalIlluminationUseRussianRoulett
 	TEXT("NOTE: This parameter is experimental")
 );
 
-static float GRayTracingGlobalIlluminationScreenPercentage = 100.0;
+static float GRayTracingGlobalIlluminationScreenPercentage = 50.0;
 static FAutoConsoleVariableRef CVarRayTracingGlobalIlluminationScreenPercentage(
 	TEXT("r.RayTracing.GlobalIllumination.ScreenPercentage"),
 	GRayTracingGlobalIlluminationScreenPercentage,
