Fixed Compile UE4Editor Win64 (VS2019) - SkyAtmosphereRendering: Implicit conversion from 'T' to bool

diff --git a/Engine/Source/Runtime/Renderer/Private/SkyAtmosphereRendering.cpp b/Engine/Source/Runtime/Renderer/Private/SkyAtmosphereRendering.cpp
--- a/Engine/Source/Runtime/Renderer/Private/SkyAtmosphereRendering.cpp
+++ b/Engine/Source/Runtime/Renderer/Private/SkyAtmosphereRendering.cpp
@@ -1383,7 +1383,7 @@ void FDeferredShadingSceneRenderer::RenderDebugSkyAtmosphere(FRHICommandListImme
 	SCOPED_DRAW_EVENT(RHICmdList, SkyAtmosphereDebugViz);
 	SCOPED_GPU_STAT(RHICmdList, SkyAtmosphereDebugViz);
 
-	const bool SkyAtmosphereVisualize = CVarSkyAtmosphereVisualize.GetValueOnRenderThread();
+	const int32 SkyAtmosphereVisualize = CVarSkyAtmosphereVisualize.GetValueOnRenderThread();
 	FSkyAtmosphereRenderSceneInfo& SkyInfo = *Scene->GetSkyAtmosphereSceneInfo();
 	const FAtmosphereSetup& Atmosphere = SkyInfo.GetAtmosphereSetup();
 	FSceneRenderTargets& SceneContext = FSceneRenderTargets::Get(RHICmdList);
