Guard against invalid light function render target resolutions

diff --git a/Engine/Source/Runtime/Renderer/Private/VolumetricFogLightFunction.cpp b/Engine/Source/Runtime/Renderer/Private/VolumetricFogLightFunction.cpp
--- a/Engine/Source/Runtime/Renderer/Private/VolumetricFogLightFunction.cpp
+++ b/Engine/Source/Runtime/Renderer/Private/VolumetricFogLightFunction.cpp
@@ -157,6 +157,25 @@ void FDeferredShadingSceneRenderer::RenderLightFunctionForVolumetricFog(
 			LightFunctionResolution.X = FMath::DivideAndRoundUp(LightFunctionResolution.X, ResolutionSnapFactor) * ResolutionSnapFactor;
 			LightFunctionResolution.Y = FMath::DivideAndRoundUp(LightFunctionResolution.Y, ResolutionSnapFactor) * ResolutionSnapFactor;
 
+			// Guard against invalid resolutions
+			const uint32 LiFuncResXU32 = (uint32)LightFunctionResolution.X;
+			const uint32 LiFuncResYU32 = (uint32)LightFunctionResolution.Y;
+			const uint32 MaxTextureResU32 = GMaxTextureDimensions;
+			if (LiFuncResXU32 > MaxTextureResU32 || LiFuncResYU32 > MaxTextureResU32)
+			{
+#if !(UE_BUILD_SHIPPING)
+				UE_LOG(LogRenderer, Error,
+					TEXT("Invalid LightFunctionResolution %dx%d, View={ %s}, LightDirection={ %s }"),
+					LightFunctionResolution.X,
+					LightFunctionResolution.Y,
+					*View.ViewMatrices.GetOverriddenTranslatedViewMatrix().ToString(),
+					*LightDirection.ToString());
+#endif
+				const uint32 ClampedRes = FMath::Max(FMath::Min3(LiFuncResXU32, LiFuncResYU32, MaxTextureResU32), (uint32)ResolutionSnapFactor);
+				LightFunctionResolution.X = ClampedRes;
+				LightFunctionResolution.Y = ClampedRes;
+			}
+
 			FWholeSceneProjectedShadowInitializer ShadowInitializer;
 
 			check(VolumetricFogMaxDistance > 0);
