Cherry-pick 7609529: Adding tiling support to ray tracing global illumination pass.
r.RayTracing.GlobalIllumination.TileSize (default = 0, disabled)

Submitted on behalf of [REDACTED_HANDLE]

diff --git a/Engine/Shaders/Private/RayTracing/RayTracingGlobalIlluminationRGS.usf b/Engine/Shaders/Private/RayTracing/RayTracingGlobalIlluminationRGS.usf
--- a/Engine/Shaders/Private/RayTracing/RayTracingGlobalIlluminationRGS.usf
+++ b/Engine/Shaders/Private/RayTracing/RayTracingGlobalIlluminationRGS.usf
@@ -30,6 +30,8 @@ float DiffuseThreshold;
 bool EvalSkyLight;
 bool UseRussianRoulette;
 float MaxNormalBias;
+uint TileOffsetX;
+uint TileOffsetY;
 
 void GenerateCosineNormalRay(
 	float3 WorldPosition,
@@ -56,7 +58,7 @@ void GenerateCosineNormalRay(
 [shader("raygeneration")]
 void GlobalIlluminationRGS()
 {
-	uint2 DispatchThreadId = DispatchRaysIndex().xy + View.ViewRectMin;
+	uint2 DispatchThreadId = DispatchRaysIndex().xy + View.ViewRectMin + uint2(TileOffsetX, TileOffsetY);
 	RWGlobalIlluminationUAV[DispatchThreadId] = 0.0;
 	RWRayDistanceUAV[DispatchThreadId] = float2(-1.0, 0.0);
 
diff --git a/Engine/Source/Runtime/Renderer/Private/RayTracing/RayTracingGlobalIllumination.cpp b/Engine/Source/Runtime/Renderer/Private/RayTracing/RayTracingGlobalIllumination.cpp
--- a/Engine/Source/Runtime/Renderer/Private/RayTracing/RayTracingGlobalIllumination.cpp
+++ b/Engine/Source/Runtime/Renderer/Private/RayTracing/RayTracingGlobalIllumination.cpp
@@ -105,6 +105,13 @@ static TAutoConsoleVariable<int32> CVarRayTracingGlobalIlluminationEnableTwoSide
 	ECVF_RenderThreadSafe
 );
 
+static int32 GRayTracingGlobalIlluminationTileSize = 0;
+static FAutoConsoleVariableRef CVarRayTracingGlobalIlluminationTileSize(
+	TEXT("r.RayTracing.GlobalIllumination.TileSize"),
+	GRayTracingGlobalIlluminationTileSize,
+	TEXT("Render ray traced global illumination in NxN piel tiles, where each tile is submitted as separate GPU command buffer, allowing high quality rendering without triggering timeout detection. (default = 0, tiling disabled)")
+);
+
 static const int32 GLightCountMax = 64;
 
 DECLARE_GPU_STAT_NAMED(RayTracingGlobalIllumination, TEXT("Ray Tracing Global Illumination"));
@@ -231,6 +238,8 @@ class FGlobalIlluminationRGS : public FGlobalShader
 		SHADER_PARAMETER(bool, EvalSkyLight)
 		SHADER_PARAMETER(bool, UseRussianRoulette)
 		SHADER_PARAMETER(float, MaxNormalBias)
+		SHADER_PARAMETER(uint32, TileOffsetX)
+		SHADER_PARAMETER(uint32, TileOffsetY)
 
 		SHADER_PARAMETER_SRV(RaytracingAccelerationStructure, TLAS)
 		SHADER_PARAMETER_RDG_TEXTURE_UAV(RWTexture2D<float4>, RWGlobalIlluminationUAV)
@@ -488,7 +497,8 @@ void FDeferredShadingSceneRenderer::RenderRayTracingGlobalIllumination(
 		PassParameters->TransmissionProfilesLinearSampler = TStaticSamplerState<SF_Bilinear, AM_Clamp, AM_Clamp, AM_Clamp>::GetRHI();
 		PassParameters->RWGlobalIlluminationUAV = GlobalIlluminationUAV;
 		PassParameters->RWRayDistanceUAV = RayDistanceUAV;
-
+		PassParameters->TileOffsetX = 0;
+		PassParameters->TileOffsetY = 0;
 
 		FGlobalIlluminationRGS::FPermutationDomain PermutationVector;
 		PermutationVector.Set<FGlobalIlluminationRGS::FUseAttenuationTermDim>(true);
@@ -497,18 +507,64 @@ void FDeferredShadingSceneRenderer::RenderRayTracingGlobalIllumination(
 		ClearUnusedGraphResources(*RayGenerationShader, PassParameters);
 
 		FIntPoint RayTracingResolution = FIntPoint::DivideAndRoundUp(View.ViewRect.Size(), UpscaleFactor);
-		GraphBuilder.AddPass(
-			RDG_EVENT_NAME("GlobalIlluminationRayTracing %dx%d", RayTracingResolution.X, RayTracingResolution.Y),
-			PassParameters,
-			ERDGPassFlags::Compute,
-			[PassParameters, this, &View, RayGenerationShader, RayTracingResolution](FRHICommandList& RHICmdList)
+
+		if (GRayTracingGlobalIlluminationTileSize <= 0)
 		{
-			FRayTracingShaderBindingsWriter GlobalResources;
-			SetShaderParameters(GlobalResources, *RayGenerationShader, *PassParameters);
+			FIntPoint RayTracingResolution = FIntPoint::DivideAndRoundUp(View.ViewRect.Size(), UpscaleFactor);
+			GraphBuilder.AddPass(
+				RDG_EVENT_NAME("GlobalIlluminationRayTracing %dx%d", RayTracingResolution.X, RayTracingResolution.Y),
+				PassParameters,
+				ERDGPassFlags::Compute,
+				[PassParameters, this, &View, RayGenerationShader, RayTracingResolution](FRHICommandList& RHICmdList)
+			{
+				FRayTracingShaderBindingsWriter GlobalResources;
+				SetShaderParameters(GlobalResources, *RayGenerationShader, *PassParameters);
+
+				FRHIRayTracingScene* RayTracingSceneRHI = View.RayTracingScene.RayTracingSceneRHI;
+				RHICmdList.RayTraceDispatch(View.RayTracingMaterialPipeline, RayGenerationShader->GetRayTracingShader(), RayTracingSceneRHI, GlobalResources, RayTracingResolution.X, RayTracingResolution.Y);
+			});
+		}
+		else
+		{
+			int32 TileSize = FMath::Max(32, GRayTracingGlobalIlluminationTileSize);
+			int32 NumTilesX = FMath::DivideAndRoundUp(RayTracingResolution.X, TileSize);
+			int32 NumTilesY = FMath::DivideAndRoundUp(RayTracingResolution.Y, TileSize);
+			for (int32 Y = 0; Y < NumTilesY; ++Y)
+			{
+				for (int32 X = 0; X < NumTilesX; ++X)
+				{
+					FGlobalIlluminationRGS::FParameters* TilePassParameters = PassParameters;
+
+					if (X > 0 || Y > 0)
+					{
+						TilePassParameters = GraphBuilder.AllocParameters<FGlobalIlluminationRGS::FParameters>();
+						*TilePassParameters = *PassParameters;
+
+						TilePassParameters->TileOffsetX = X * TileSize;
+						TilePassParameters->TileOffsetY = Y * TileSize;
+					}
+
+					int32 DispatchSizeX = FMath::Min<int32>(TileSize, RayTracingResolution.X - TilePassParameters->TileOffsetX);
+					int32 DispatchSizeY = FMath::Min<int32>(TileSize, RayTracingResolution.Y - TilePassParameters->TileOffsetY);
+
+					GraphBuilder.AddPass(
+						RDG_EVENT_NAME("GlobalIlluminationRayTracing %dx%d (tile %dx%d)", DispatchSizeX, DispatchSizeY, X, Y),
+						TilePassParameters,
+						ERDGPassFlags::Compute,
+						[TilePassParameters, this, &View, RayGenerationShader, DispatchSizeX, DispatchSizeY](FRHICommandList& RHICmdList)
+					{
+						FRHIRayTracingScene* RayTracingSceneRHI = View.RayTracingScene.RayTracingSceneRHI;
+
+						FRayTracingShaderBindingsWriter GlobalResources;
+						SetShaderParameters(GlobalResources, *RayGenerationShader, *TilePassParameters);
+						RHICmdList.RayTraceDispatch(View.RayTracingMaterialPipeline, RayGenerationShader->GetRayTracingShader(), RayTracingSceneRHI, 
+							GlobalResources, DispatchSizeX, DispatchSizeY);
+						RHICmdList.SubmitCommandsHint();
+					});
+				}
+			}
+		}
 
-			FRHIRayTracingScene* RayTracingSceneRHI = View.RayTracingScene.RayTracingSceneRHI;
-			RHICmdList.RayTraceDispatch(View.RayTracingMaterialPipeline, RayGenerationShader->GetRayTracingShader(), RayTracingSceneRHI, GlobalResources, RayTracingResolution.X, RayTracingResolution.Y);
-		});
 	}
 
 	// Denoising
