Fixed Compile UE4Editor Win64 warning

diff --git a/Engine/Source/Runtime/Renderer/Private/RayTracing/RayTracingGlobalIllumination.cpp b/Engine/Source/Runtime/Renderer/Private/RayTracing/RayTracingGlobalIllumination.cpp
--- a/Engine/Source/Runtime/Renderer/Private/RayTracing/RayTracingGlobalIllumination.cpp
+++ b/Engine/Source/Runtime/Renderer/Private/RayTracing/RayTracingGlobalIllumination.cpp
@@ -510,7 +510,6 @@ void FDeferredShadingSceneRenderer::RenderRayTracingGlobalIllumination(
 
 		if (GRayTracingGlobalIlluminationTileSize <= 0)
 		{
-			FIntPoint RayTracingResolution = FIntPoint::DivideAndRoundUp(View.ViewRect.Size(), UpscaleFactor);
 			GraphBuilder.AddPass(
 				RDG_EVENT_NAME("GlobalIlluminationRayTracing %dx%d", RayTracingResolution.X, RayTracingResolution.Y),
 				PassParameters,
