Editgrate of 8096916 MobileShadingRenderer to fix static analysis errors.

diff --git a/Engine/Source/Runtime/Renderer/Private/MobileShadingRenderer.cpp b/Engine/Source/Runtime/Renderer/Private/MobileShadingRenderer.cpp
--- a/Engine/Source/Runtime/Renderer/Private/MobileShadingRenderer.cpp
+++ b/Engine/Source/Runtime/Renderer/Private/MobileShadingRenderer.cpp
@@ -228,6 +228,8 @@ void FMobileSceneRenderer::InitViews(FRHICommandListImmediate& RHICmdList)
 	SCOPE_CYCLE_COUNTER(STAT_InitViewsTime);
 	CSV_SCOPED_TIMING_STAT_EXCLUSIVE(InitViews_Scene);
 
+	check(Scene);
+
 	FILCUpdatePrimTaskData ILCTaskData;
 	FViewVisibleCommandsPerView ViewCommandsPerView;
 	ViewCommandsPerView.SetNum(Views.Num());
@@ -238,7 +240,7 @@ void FMobileSceneRenderer::InitViews(FRHICommandListImmediate& RHICmdList)
 	ComputeViewVisibility(RHICmdList, BasePassDepthStencilAccess, ViewCommandsPerView, DynamicIndexBuffer, DynamicVertexBuffer, DynamicReadBuffer);
 
 	// Initialise Sky/View resources before the view global uniform buffer is built.
-	if (Scene && ShouldRenderSkyAtmosphere(Scene->GetSkyAtmosphereSceneInfo(), Scene->GetShaderPlatform()))
+	if (ShouldRenderSkyAtmosphere(Scene->GetSkyAtmosphereSceneInfo(), Scene->GetShaderPlatform()))
 	{
 		InitSkyAtmosphereForViews(RHICmdList);
 	}
