Fixed display issue in RTGI default settings

diff --git a/Engine/Source/Runtime/Renderer/Private/RayTracing/RayTracingGlobalIllumination.cpp b/Engine/Source/Runtime/Renderer/Private/RayTracing/RayTracingGlobalIllumination.cpp
--- a/Engine/Source/Runtime/Renderer/Private/RayTracing/RayTracingGlobalIllumination.cpp
+++ b/Engine/Source/Runtime/Renderer/Private/RayTracing/RayTracingGlobalIllumination.cpp
@@ -95,7 +95,7 @@ static float GRayTracingGlobalIlluminationScreenPercentage = 50.0;
 static FAutoConsoleVariableRef CVarRayTracingGlobalIlluminationScreenPercentage(
 	TEXT("r.RayTracing.GlobalIllumination.ScreenPercentage"),
 	GRayTracingGlobalIlluminationScreenPercentage,
-	TEXT("Screen percentage for ray tracing global illumination (default = 100)")
+	TEXT("Screen percentage for ray tracing global illumination (default = 50)")
 );
 
 static TAutoConsoleVariable<int32> CVarRayTracingGlobalIlluminationEnableTwoSidedGeometry(
