Fixed MaxDrawDistance not working for HLODs

diff --git a/Engine/Source/Runtime/Renderer/Private/SceneVisibility.cpp b/Engine/Source/Runtime/Renderer/Private/SceneVisibility.cpp
--- a/Engine/Source/Runtime/Renderer/Private/SceneVisibility.cpp
+++ b/Engine/Source/Runtime/Renderer/Private/SceneVisibility.cpp
@@ -4489,7 +4489,9 @@ void FLODSceneTree::UpdateVisibilityStates(FViewInfo& View)
 
 			// Update visibility states of this node and owned children
 			const float DistanceSquared = Bounds.BoxSphereBounds.ComputeSquaredDistanceFromBoxToPoint(View.ViewMatrices.GetViewOrigin());
-			const bool bIsInDrawRange = DistanceSquared >= Bounds.MinDrawDistanceSq * HLODState.FOVDistanceScaleSq;
+			const bool bNearCulled = DistanceSquared < Bounds.MinDrawDistanceSq * HLODState.FOVDistanceScaleSq;
+			const bool bFarCulled = DistanceSquared > Bounds.MaxDrawDistance * Bounds.MaxDrawDistance * HLODState.FOVDistanceScaleSq;
+			const bool bIsInDrawRange = !bNearCulled && !bFarCulled;
 
 			const bool bWasFadingPreUpdate = !!NodeVisibility.bIsFading;
 			const bool bIsDitheredTransition = NodeMeshRelevances[0].bDitheredLODTransition;
@@ -4569,6 +4571,12 @@ void FLODSceneTree::UpdateVisibilityStates(FViewInfo& View)
 			{
 				// Not visible and waiting for a transition to fade, keep HLOD hidden
 				HLODState.ForcedHiddenPrimitiveMap[NodeIndex] = true;
+
+				// Also hide children when performing far culling
+				if (bFarCulled)
+				{
+					HideNodeChildren(ViewState, Node);
+				}
 			}
 		}
 	}	
