Fixed default value for ray tracing translucency was incorrect. That was making the translucency flag in the postprocessing being ignored.

diff --git a/Engine/Source/Runtime/Renderer/Private/RayTracing/RayTracingTranslucency.cpp b/Engine/Source/Runtime/Renderer/Private/RayTracing/RayTracingTranslucency.cpp
--- a/Engine/Source/Runtime/Renderer/Private/RayTracing/RayTracingTranslucency.cpp
+++ b/Engine/Source/Runtime/Renderer/Private/RayTracing/RayTracingTranslucency.cpp
@@ -27,7 +27,7 @@
 
 static TAutoConsoleVariable<int32> CVarRayTracingTranslucency(
 	TEXT("r.RayTracing.Translucency"),
-	0,
+	-1,
 	TEXT("-1: Value driven by postprocess volume (default) \n")
 	TEXT(" 0: ray tracing translucency off (use raster) \n")
 	TEXT(" 1: ray tracing translucency enabled"),
