Pushing 0-spp guard to top of function for ray tracing SkyLight pass.

diff --git a/Engine/Source/Runtime/Renderer/Private/RayTracing/RaytracingSkylight.cpp b/Engine/Source/Runtime/Renderer/Private/RayTracing/RaytracingSkylight.cpp
--- a/Engine/Source/Runtime/Renderer/Private/RayTracing/RaytracingSkylight.cpp
+++ b/Engine/Source/Runtime/Renderer/Private/RayTracing/RaytracingSkylight.cpp
@@ -864,7 +864,8 @@ void FDeferredShadingSceneRenderer::GenerateSkyLightVisibilityRays(
 	// SkyLight data setup
 	FSkyLightData SkyLightData;
 	SetupSkyLightParameters(*Scene, &SkyLightData);
-	uint32 SamplesPerPixel = GRayTracingSkyLightSamplesPerPixel >= 0 ? GRayTracingSkyLightSamplesPerPixel : Scene->SkyLight->SamplesPerPixel;
+	int32 SamplesPerPixel = GRayTracingSkyLightSamplesPerPixel >= 0 ? GRayTracingSkyLightSamplesPerPixel : Scene->SkyLight->SamplesPerPixel;
+	ensure(SamplesPerPixel > 0);
 	SkyLightData.SamplesPerPixel = SamplesPerPixel;
 	SkyLightData.MaxRayDistance = GRayTracingSkyLightMaxRayDistance;
 	SkyLightData.MaxShadowThickness = GRayTracingSkyLightMaxShadowThickness;
@@ -879,7 +880,7 @@ void FDeferredShadingSceneRenderer::GenerateSkyLightVisibilityRays(
 	Dimensions = FIntVector(BlueNoise.Dimensions.X, BlueNoise.Dimensions.Y, 0);
 
 	// Halton iteration
-	uint32 IterationCount = FMath::Max(SamplesPerPixel, 1u);
+	uint32 IterationCount = FMath::Max(SamplesPerPixel, 1);
 	uint32 SequenceCount = 1;
 	uint32 DimensionCount = 3;
 	FHaltonSequenceIteration HaltonSequenceIteration(Scene->HaltonSequence, IterationCount, SequenceCount, DimensionCount, Views[0].ViewState ? (Views[0].ViewState->FrameIndex % 1024) : 0);
@@ -925,6 +926,12 @@ void FDeferredShadingSceneRenderer::RenderRayTracingSkyLight(
 		return;
 	}
 
+	int32 SamplesPerPixel = GRayTracingSkyLightSamplesPerPixel >= 0 ? GRayTracingSkyLightSamplesPerPixel : Scene->SkyLight->SamplesPerPixel;
+	if (SamplesPerPixel <= 0)
+	{
+		return;
+	}
+
 	check(Scene->SkyLight->ProcessedTexture);
 
 	if (Scene->SkyLight->ShouldRebuildCdf())
@@ -985,7 +992,6 @@ void FDeferredShadingSceneRenderer::RenderRayTracingSkyLight(
 	// Fill SkyLight parameters
 	FSkyLightData SkyLightData;
 	SetupSkyLightParameters(*Scene, &SkyLightData);
-	uint32 SamplesPerPixel = GRayTracingSkyLightSamplesPerPixel >= 0 ? GRayTracingSkyLightSamplesPerPixel : Scene->SkyLight->SamplesPerPixel;
 	SkyLightData.SamplesPerPixel = SamplesPerPixel;
 	SkyLightData.MaxRayDistance = GRayTracingSkyLightMaxRayDistance;
 	SkyLightData.SamplingStopLevel = GRayTracingSkyLightSamplingStopLevel;
@@ -1020,7 +1026,7 @@ void FDeferredShadingSceneRenderer::RenderRayTracingSkyLight(
 		FViewInfo& View = Views[ViewIndex];
 
 		// Halton iteration
-		uint32 IterationCount = FMath::Max(SamplesPerPixel, 1u);
+		uint32 IterationCount = SamplesPerPixel;
 		uint32 SequenceCount = 1;
 		uint32 DimensionCount = 3;
 		FHaltonSequenceIteration HaltonSequenceIteration(Scene->HaltonSequence, IterationCount, SequenceCount, DimensionCount, View.ViewState ? (View.ViewState->FrameIndex % 1024) : 0);
@@ -1071,7 +1077,7 @@ void FDeferredShadingSceneRenderer::RenderRayTracingSkyLight(
 		});
 
 		// Denoising
-		if (GRayTracingSkyLightDenoiser != 0 && SamplesPerPixel > 0)
+		if (GRayTracingSkyLightDenoiser != 0)
 		{
 			const IScreenSpaceDenoiser* DefaultDenoiser = IScreenSpaceDenoiser::GetDefaultDenoiser();
 			const IScreenSpaceDenoiser* DenoiserToUse = DefaultDenoiser;// GRayTracingGlobalIlluminationDenoiser == 1 ? DefaultDenoiser : GScreenSpaceDenoiser;
