Pushing 0-spp guard to top of function of ray tracing global illumination pass.

diff --git a/Engine/Source/Runtime/Renderer/Private/RayTracing/RayTracingGlobalIllumination.cpp b/Engine/Source/Runtime/Renderer/Private/RayTracing/RayTracingGlobalIllumination.cpp
--- a/Engine/Source/Runtime/Renderer/Private/RayTracing/RayTracingGlobalIllumination.cpp
+++ b/Engine/Source/Runtime/Renderer/Private/RayTracing/RayTracingGlobalIllumination.cpp
@@ -433,6 +433,7 @@ void FDeferredShadingSceneRenderer::RenderRayTracingGlobalIllumination(
 	RDG_EVENT_SCOPE(GraphBuilder, "RTGI");
 
 	int32 RayTracingGISamplesPerPixel = GRayTracingGlobalIlluminationSamplesPerPixel > -1 ? GRayTracingGlobalIlluminationSamplesPerPixel : View.FinalPostProcessSettings.RayTracingGISamplesPerPixel;
+	RayTracingGISamplesPerPixel = FMath::Max(RayTracingGISamplesPerPixel, 0);
 	RayTracingConfig.RayCountPerPixel = RayTracingGISamplesPerPixel;
 
 	FSceneTextureParameters SceneTextures;
@@ -440,7 +441,7 @@ void FDeferredShadingSceneRenderer::RenderRayTracingGlobalIllumination(
 
 	// Ray generation
 	{
-		uint32 IterationCount = RayTracingGISamplesPerPixel;
+		uint32 IterationCount = FMath::Max(RayTracingGISamplesPerPixel, 1);
 		uint32 SequenceCount = 1;
 		uint32 DimensionCount = 24;
 		FHaltonSequenceIteration HaltonSequenceIteration(Scene->HaltonSequence, IterationCount, SequenceCount, DimensionCount, View.ViewState ? (View.ViewState->FrameIndex % 1024) : 0);
@@ -570,7 +571,7 @@ void FDeferredShadingSceneRenderer::RenderRayTracingGlobalIllumination(
 
 	FRDGTextureRef ResultTexture; //#dxr_todo review
 
-	if (GRayTracingGlobalIlluminationDenoiser != 0)
+	if (GRayTracingGlobalIlluminationDenoiser != 0 && RayTracingGISamplesPerPixel > 0)
 	{
 		const IScreenSpaceDenoiser* DefaultDenoiser = IScreenSpaceDenoiser::GetDefaultDenoiser();
 		const IScreenSpaceDenoiser* DenoiserToUse = GRayTracingGlobalIlluminationDenoiser == 1 ? DefaultDenoiser : GScreenSpaceDenoiser;
