Improves performance by no longer clearing the render target during mip generation as the render target should be fully written by the downsampling shader.

diff --git a/Engine/Source/Runtime/Renderer/Private/PostProcess/PostProcessDownsample.cpp b/Engine/Source/Runtime/Renderer/Private/PostProcess/PostProcessDownsample.cpp
--- a/Engine/Source/Runtime/Renderer/Private/PostProcess/PostProcessDownsample.cpp
+++ b/Engine/Source/Runtime/Renderer/Private/PostProcess/PostProcessDownsample.cpp
@@ -367,7 +367,7 @@ void FRCPassPostProcessDownsample::Process(FRenderingCompositePassContext& Conte
 	else
 	{
 		// #todo-renderpasses only clear dest rectangle if it's been computed
-		FRHIRenderPassInfo RPInfo(DestRenderTarget.TargetableTexture, ERenderTargetActions::Clear_Store);
+		FRHIRenderPassInfo RPInfo(DestRenderTarget.TargetableTexture, ERenderTargetActions::DontLoad_Store);
 		Context.RHICmdList.BeginRenderPass(RPInfo, TEXT("PostProcessDownsample"));
 		{
 			Context.SetViewportAndCallRHI(0, 0, 0.0f, DestSize.X, DestSize.Y, 1.0f);
