Fixed assert when using r.SkyAtmosphere.Visualize.

diff --git a/Engine/Source/Runtime/Renderer/Private/SkyAtmosphereRendering.cpp b/Engine/Source/Runtime/Renderer/Private/SkyAtmosphereRendering.cpp
--- a/Engine/Source/Runtime/Renderer/Private/SkyAtmosphereRendering.cpp
+++ b/Engine/Source/Runtime/Renderer/Private/SkyAtmosphereRendering.cpp
@@ -1301,7 +1301,7 @@ void FSceneRenderer::RenderSkyAtmosphere(FRHICommandListImmediate& RHICmdList)
 
 			FLinearColor TextColor(1.0f, 0.5f, 0.0f);
 			FString Text = TEXT("You are using a FastAerialPespective without FastSky, visuals might look wrong.");
-			Canvas.DrawShadowedString(ViewPortWidth*0.5 - Text.Len()*7, ViewPortHeight*0.5, *Text, GetStatsFont(), TextColor);
+			Canvas.DrawShadowedString(ViewPortWidth*0.5f - Text.Len()*7.0f, ViewPortHeight*0.4f, *Text, GetStatsFont(), TextColor);
 
 			Canvas.Flush_RenderThread(RHICmdList);
 		}
@@ -1422,8 +1422,8 @@ void FDeferredShadingSceneRenderer::RenderDebugSkyAtmosphere(FRHICommandListImme
 				PassParameters->SkyAtmosphere = InternalCommonParametersRef;
 				PassParameters->ViewUniformBuffer = View.ViewUniformBuffer;
 				PassParameters->RenderTargets[0] = FRenderTargetBinding(SceneColor, ERenderTargetLoadAction::ELoad, ERenderTargetStoreAction::EStore);
-				PassParameters->RenderTargets.DepthStencil = FDepthStencilBinding(SceneDepth, ERenderTargetLoadAction::ELoad, ERenderTargetStoreAction::ENoAction,
-					ERenderTargetLoadAction::ENoAction, ERenderTargetStoreAction::ENoAction, FExclusiveDepthStencil::DepthRead_StencilNop);
+				PassParameters->RenderTargets.DepthStencil = FDepthStencilBinding(SceneDepth, ERenderTargetLoadAction::ELoad, ERenderTargetStoreAction::EStore,
+					ERenderTargetLoadAction::ENoAction, ERenderTargetStoreAction::ENoAction, FExclusiveDepthStencil::DepthWrite_StencilNop);
 				PassParameters->TransmittanceLutTextureSampler = SamplerLinearClamp;
 				PassParameters->MultiScatteredLuminanceLutTextureSampler = SamplerLinearClamp;
 				PassParameters->TransmittanceLutTexture = SkyInfo.GetTransmittanceLutTexture()->GetRenderTargetItem().ShaderResourceTexture;
