Fixed wrong start depth in some cases (units and difference between shaders)

diff --git a/Engine/Shaders/Private/SkyAtmosphere.usf b/Engine/Shaders/Private/SkyAtmosphere.usf
--- a/Engine/Shaders/Private/SkyAtmosphere.usf
+++ b/Engine/Shaders/Private/SkyAtmosphere.usf
@@ -61,6 +61,8 @@
 #define MULTISCATTERING_STORE_SCALE (4.0f * PI * 5.0f)
 #define MULTISCATTERING_STORE_SCALE_INV 1.0f/MULTISCATTERING_STORE_SCALE
 
+// Only available for shaders ran per view, so this constant should not be accessed in common code such as IntegrateSingleScatteredLuminance for instance.
+float AerialPerspectiveStartDepth;
 
 
 // - RayOrigin: ray origin
@@ -783,7 +785,7 @@ void RenderSkyAtmosphereRayMarchingPS(
 		CameraAerialPerspectiveVolumeTexture, CameraAerialPerspectiveVolumeTextureSampler,
 		SkyAtmosphere.CameraAerialPerspectiveVolumeDepthResolutionInv,
 		SkyAtmosphere.CameraAerialPerspectiveVolumeDepthResolution,
-		SkyAtmosphere.AerialPerspectiveStartDepth,
+		AerialPerspectiveStartDepth,
 		SkyAtmosphere.CameraAerialPerspectiveVolumeDepthSliceLength,
 		SkyAtmosphere.CameraAerialPerspectiveVolumeDepthSliceLengthInv,
 		OneOverPreExposure);
@@ -797,7 +799,7 @@ void RenderSkyAtmosphereRayMarchingPS(
 
 #else // FASTAERIALPERSPECTIVE_ENABLED
 
-	WorldPos += WorldDir * SkyAtmosphere.AerialPerspectiveStartDepth; // Offset according to start depth
+	WorldPos += WorldDir * AerialPerspectiveStartDepth; // Offset according to start depth
 
 	// Move to top atmosphere as the starting point for ray marching.
 	// This is critical to be after the above to not disrupt above atmosphere tests and voxel selection.
@@ -1166,7 +1168,7 @@ void RenderCameraAerialPerspectiveVolumeCS(uint3 ThreadId : SV_DispatchThreadID)
 	Slice *= Slice;	// squared distribution
 	Slice *= SkyAtmosphere.CameraAerialPerspectiveVolumeDepthResolution;
 
-	float3 RayStartWorldPos = CamPos + SkyAtmosphere.AerialPerspectiveStartDepth * WorldDir; // Offset according to start depth
+	float3 RayStartWorldPos = CamPos + AerialPerspectiveStartDepth * WorldDir; // Offset according to start depth
 	float ViewHeight;
 
 
@@ -1219,7 +1221,7 @@ void RenderCameraAerialPerspectiveVolumeCS(uint3 ThreadId : SV_DispatchThreadID)
 		}
 		 
 		WorldDir = normalize(VoxelWorldPos - CamPos);
-		RayStartWorldPos = CamPos + SkyAtmosphere.AerialPerspectiveStartDepth * WorldDir; // Offset according to start depth
+		RayStartWorldPos = CamPos + AerialPerspectiveStartDepth * WorldDir; // Offset according to start depth
 		tMax = length(VoxelWorldPos - RayStartWorldPos);
 	}
 	float tMaxMax = tMax;
diff --git a/Engine/Source/Runtime/Renderer/Private/SkyAtmosphereRendering.cpp b/Engine/Source/Runtime/Renderer/Private/SkyAtmosphereRendering.cpp
--- a/Engine/Source/Runtime/Renderer/Private/SkyAtmosphereRendering.cpp
+++ b/Engine/Source/Runtime/Renderer/Private/SkyAtmosphereRendering.cpp
@@ -210,7 +210,6 @@ BEGIN_GLOBAL_SHADER_PARAMETER_STRUCT(FSkyAtmosphereInternalCommonParameters, )
 	SHADER_PARAMETER(float, CameraAerialPerspectiveVolumeDepthSliceLength)		// Also on View UB
 	SHADER_PARAMETER(float, CameraAerialPerspectiveVolumeDepthSliceLengthInv)	// Also on View UB
 	SHADER_PARAMETER(float, CameraAerialPerspectiveSampleCountPerSlice)
-	SHADER_PARAMETER(float, AerialPerspectiveStartDepth)
 
 	SHADER_PARAMETER(FVector4, TransmittanceLutSizeAndInvSize)
 	SHADER_PARAMETER(FVector4, MultiScatteredLuminanceLutSizeAndInvSize)
@@ -241,12 +240,21 @@ IMPLEMENT_GLOBAL_SHADER_PARAMETER_STRUCT(FSkyAtmosphereInternalCommonParameters,
 	int32 SkyViewLutHeight = ValidateLUTResolution(CVarSkyAtmosphereFastSkyLUTHeight.GetValueOnRenderThread()); \
 	int32 CameraAerialPerspectiveVolumeScreenResolution = ValidateLUTResolution(CVarSkyAtmosphereAerialPerspectiveLUTWidth.GetValueOnRenderThread()); \
 	int32 CameraAerialPerspectiveVolumeDepthResolution = ValidateLUTResolution(CVarSkyAtmosphereAerialPerspectiveLUTDepthResolution.GetValueOnRenderThread()); \
-	float AerialPerspectiveStartDepth = CVarSkyAtmosphereAerialPerspectiveStartDepth.GetValueOnRenderThread(); \
-	AerialPerspectiveStartDepth = AerialPerspectiveStartDepth < 0.0f ? 0.0f : AerialPerspectiveStartDepth; \
 	float CameraAerialPerspectiveVolumeDepth = CVarSkyAtmosphereAerialPerspectiveLUTDepth.GetValueOnRenderThread(); \
 	CameraAerialPerspectiveVolumeDepth = CameraAerialPerspectiveVolumeDepth < 1.0f ? 1.0f : CameraAerialPerspectiveVolumeDepth;	/* 1 kilometer minimum */ \
 	float CameraAerialPerspectiveVolumeDepthSliceLength = CameraAerialPerspectiveVolumeDepth / CameraAerialPerspectiveVolumeDepthResolution;
 
+#define KM_TO_CM  100000.0f
+#define CM_TO_KM  (1.0f / KM_TO_CM)
+
+static float GetValidAerialPerspectiveStartDepthInCm(const FViewInfo& View)
+{
+	float AerialPerspectiveStartDepth = CVarSkyAtmosphereAerialPerspectiveStartDepth.GetValueOnRenderThread();
+	AerialPerspectiveStartDepth = AerialPerspectiveStartDepth < 0.0f ? 0.0f : AerialPerspectiveStartDepth;
+	// For sky reflection capture, the start depth can be super large. So we max it to make sure the triangle is never in front the NearClippingDistance.
+	const float StartDepthInCm = FMath::Max(AerialPerspectiveStartDepth * KM_TO_CM, View.NearClippingDistance);
+	return StartDepthInCm;
+}
 
 static bool ShouldPipelineCompileSkyAtmosphereShader(EShaderPlatform ShaderPlatform)
 {
@@ -261,7 +269,7 @@ bool ShouldRenderSkyAtmosphere(const FSkyAtmosphereRenderSceneInfo* SkyAtmospher
 	// TODO: Add new or reuse EngineShowFlags.AtmosphericFog? ALso take into account EngineShowFlags.Fog as previously?
 }
 
-void SetupSkyAtmosphereViewSharedUniformShaderParameters(const class FViewInfo& View, FSkyAtmosphereViewSharedUniformShaderParameters& OutParameters)
+void SetupSkyAtmosphereViewSharedUniformShaderParameters(const FViewInfo& View, FSkyAtmosphereViewSharedUniformShaderParameters& OutParameters)
 {
 	GET_VALID_DATA_FROM_CVAR;
 
@@ -277,7 +285,7 @@ void SetupSkyAtmosphereViewSharedUniformShaderParameters(const class FViewInfo&
 	OutParameters.CameraAerialPerspectiveVolumeDepthSliceLength = CameraAerialPerspectiveVolumeDepthSliceLength;
 	OutParameters.CameraAerialPerspectiveVolumeDepthSliceLengthInv = 1.0f / OutParameters.CameraAerialPerspectiveVolumeDepthSliceLength;
 
-	OutParameters.AerialPerspectiveStartDepth = AerialPerspectiveStartDepth;
+	OutParameters.AerialPerspectiveStartDepth = GetValidAerialPerspectiveStartDepthInCm(View);
 
 	SetBlackAlpha13DIfNull(SkyAtmosphereCameraAerialPerspectiveVolume); // Needs to be after we set ApplyCameraAerialPerspectiveVolume
 }
@@ -500,6 +508,7 @@ class FRenderSkyAtmospherePS : public FGlobalShader
 		SHADER_PARAMETER_SAMPLER(SamplerState, MultiScatteredLuminanceLutTextureSampler)
 		SHADER_PARAMETER_SAMPLER(SamplerState, SkyViewLutTextureSampler)
 		SHADER_PARAMETER_SAMPLER(SamplerState, CameraAerialPerspectiveVolumeTextureSampler)
+		SHADER_PARAMETER(float, AerialPerspectiveStartDepth)
 	END_SHADER_PARAMETER_STRUCT()
 
 	static FPermutationDomain RemapPermutation(FPermutationDomain PermutationVector)
@@ -718,6 +727,7 @@ public:
 		SHADER_PARAMETER_TEXTURE(Texture2D<float3>, MultiScatteredLuminanceLutTexture)
 		SHADER_PARAMETER_SAMPLER(SamplerState, TransmittanceLutTextureSampler)
 		SHADER_PARAMETER_SAMPLER(SamplerState, MultiScatteredLuminanceLutTextureSampler)
+		SHADER_PARAMETER(float, AerialPerspectiveStartDepth)
 	END_SHADER_PARAMETER_STRUCT()
 
 	static bool ShouldCompilePermutation(const FGlobalShaderPermutationParameters& Parameters)
@@ -924,7 +934,6 @@ static void SetupSkyAtmosphereInternalCommonParameters(
 	InternalCommonParameters.FastSkySampleCountMax = CVarSkyAtmosphereFastSkyLUTSampleCountMax.GetValueOnRenderThread();
 	float FastSkyDistanceToSampleCountMaxInv = CVarSkyAtmosphereFastSkyLUTDistanceToSampleCountMax.GetValueOnRenderThread();
 
-	InternalCommonParameters.AerialPerspectiveStartDepth = AerialPerspectiveStartDepth;
 	InternalCommonParameters.CameraAerialPerspectiveVolumeDepthResolution = float(CameraAerialPerspectiveVolumeDepthResolution);
 	InternalCommonParameters.CameraAerialPerspectiveVolumeDepthResolutionInv = 1.0f / InternalCommonParameters.CameraAerialPerspectiveVolumeDepthResolution;
 	InternalCommonParameters.CameraAerialPerspectiveVolumeDepthSliceLength = CameraAerialPerspectiveVolumeDepthSliceLength;
@@ -1071,6 +1080,7 @@ void FSceneRenderer::RenderSkyAtmosphereLookUpTables(FRHICommandListImmediate& R
 	for (int32 ViewIndex = 0; ViewIndex < Views.Num(); ViewIndex++)
 	{
 		FViewInfo& View = Views[ViewIndex];
+		const float AerialPerspectiveStartDepthInCm = GetValidAerialPerspectiveStartDepthInCm(View);
 
 		const bool bLightDiskEnabled = !View.bIsReflectionCapture;
 
@@ -1116,6 +1126,7 @@ void FSceneRenderer::RenderSkyAtmosphereLookUpTables(FRHICommandListImmediate& R
 			PassParameters->TransmittanceLutTexture = SkyInfo.GetTransmittanceLutTexture()->GetRenderTargetItem().ShaderResourceTexture;
 			PassParameters->MultiScatteredLuminanceLutTexture = SkyInfo.GetMultiScatteredLuminanceLutTexture()->GetRenderTargetItem().ShaderResourceTexture;
 			PassParameters->CameraAerialPerspectiveVolumeUAV = View.SkyAtmosphereCameraAerialPerspectiveVolume->GetRenderTargetItem().UAV;
+			PassParameters->AerialPerspectiveStartDepth = AerialPerspectiveStartDepthInCm * CM_TO_KM;
 
 			FIntVector TextureSize = View.SkyAtmosphereCameraAerialPerspectiveVolume->GetDesc().GetSize();
 			const FIntVector NumGroups = FIntVector::DivideAndRoundUp(TextureSize, FRenderCameraAerialPerspectiveVolumeCS::GroupSize);
@@ -1162,6 +1173,7 @@ void FSceneRenderer::RenderSkyAtmosphere(FRHICommandListImmediate& RHICmdList)
 	for (int32 ViewIndex = 0; ViewIndex < Views.Num(); ViewIndex++)
 	{
 		FViewInfo& View = Views[ViewIndex];
+		const float AerialPerspectiveStartDepthInCm = GetValidAerialPerspectiveStartDepthInCm(View);
 
 		const bool bLightDiskEnabled = !View.bIsReflectionCapture;
 
@@ -1171,9 +1183,8 @@ void FSceneRenderer::RenderSkyAtmosphere(FRHICommandListImmediate& RHICmdList)
 		const bool bRenderSkyPixel = !View.bSceneHasSkyMaterial;
 
 		const FVector ViewOrigin = View.ViewMatrices.GetViewOrigin();
-		const float KmToCm = 100000.0f;
-		const FVector PlanetOrigin = FVector(0.0f, 0.0f, -Atmosphere.BottomRadius * KmToCm);
-		const float TopOfAtmosphere = Atmosphere.TopRadius * KmToCm;
+		const FVector PlanetOrigin = FVector(0.0f, 0.0f, -Atmosphere.BottomRadius * KM_TO_CM);
+		const float TopOfAtmosphere = Atmosphere.TopRadius * KM_TO_CM;
 		const float SafeEdge = 1000.0f;	// 10 meters
 		const bool ForceRayMarching = (FVector::Distance(ViewOrigin, PlanetOrigin) - TopOfAtmosphere - SafeEdge) > 0.0f;
 
@@ -1211,6 +1222,7 @@ void FSceneRenderer::RenderSkyAtmosphere(FRHICommandListImmediate& RHICmdList)
 			PsPassParameters->MultiScatteredLuminanceLutTexture = SkyInfo.GetMultiScatteredLuminanceLutTexture()->GetRenderTargetItem().ShaderResourceTexture;
 			PsPassParameters->SkyViewLutTexture = View.SkyAtmosphereViewLutTexture->GetRenderTargetItem().ShaderResourceTexture;
 			PsPassParameters->CameraAerialPerspectiveVolumeTexture = View.SkyAtmosphereCameraAerialPerspectiveVolume->GetRenderTargetItem().ShaderResourceTexture;
+			PsPassParameters->AerialPerspectiveStartDepth = AerialPerspectiveStartDepthInCm * CM_TO_KM;
 			ClearUnusedGraphResources(*PixelShader, PsPassParameters);
 
 			float StartDepthZ = 0.1f;
@@ -1219,10 +1231,8 @@ void FSceneRenderer::RenderSkyAtmosphere(FRHICommandListImmediate& RHICmdList)
 				const FMatrix ViewProjectionMatrix = View.ViewMatrices.GetProjectionMatrix();
 				float HalfHorizontalFOV = FMath::Atan(1.0f / ViewProjectionMatrix.M[0][0]);
 				float HalfVerticalFOV = FMath::Atan(1.0f / ViewProjectionMatrix.M[1][1]);
-				// For sky reflection capture, the start depth can be super large. Se we max it to make sure the triangle is never in front the NearClippingDistance.
-				float StartDepthInKm = FMath::Max(InternalCommonParameters.AerialPerspectiveStartDepth, View.NearClippingDistance) * KmToCm;
-				const float StartDepthViewKm = FMath::Cos(FMath::Max(HalfHorizontalFOV, HalfVerticalFOV)) * StartDepthInKm;
-				const FVector4 Projected = ViewProjectionMatrix.TransformFVector4(FVector4(0.0f, 0.0f, StartDepthViewKm, 1.0f));
+				const float StartDepthViewCm = FMath::Cos(FMath::Max(HalfHorizontalFOV, HalfVerticalFOV)) * AerialPerspectiveStartDepthInCm;
+				const FVector4 Projected = ViewProjectionMatrix.TransformFVector4(FVector4(0.0f, 0.0f, StartDepthViewCm, 1.0f));
 				StartDepthZ = Projected.Z / Projected.W;
 			}
 
