Add stat for Insights profiling

diff --git a/Engine/Source/Runtime/Renderer/Private/VT/VirtualTextureSystem.cpp b/Engine/Source/Runtime/Renderer/Private/VT/VirtualTextureSystem.cpp
--- a/Engine/Source/Runtime/Renderer/Private/VT/VirtualTextureSystem.cpp
+++ b/Engine/Source/Runtime/Renderer/Private/VT/VirtualTextureSystem.cpp
@@ -845,6 +845,7 @@ void FVirtualTextureSystem::Update(FRHICommandListImmediate& RHICmdList, ERHIFea
 {
 	check(IsInRenderingThread());
 
+	CSV_SCOPED_TIMING_STAT_EXCLUSIVE(VirtualTextureSystem_Update);
 	SCOPE_CYCLE_COUNTER(STAT_VirtualTextureSystem_Update);
 	SCOPED_GPU_STAT(RHICmdList, VirtualTexture);
 
