Compile denoiser's shaders only when r.RayTracing=1

diff --git a/Engine/Source/Runtime/Renderer/Private/ScreenSpaceDenoise.cpp b/Engine/Source/Runtime/Renderer/Private/ScreenSpaceDenoise.cpp
--- a/Engine/Source/Runtime/Renderer/Private/ScreenSpaceDenoise.cpp
+++ b/Engine/Source/Runtime/Renderer/Private/ScreenSpaceDenoise.cpp
@@ -439,10 +439,16 @@ static_assert(ARRAY_COUNT(kDenoiserOutputResourceNames) == int32(ESignalProcessi
 /** Returns whether should compole pipeline for a given shader platform.*/
 bool ShouldCompileSignalPipeline(ESignalProcessing SignalProcessing, EShaderPlatform Platform)
 {
+	// Only denoise for project that have ray tracing.
+	if (!ShouldCompileRayTracingShadersForProject(Platform))
+	{
+		return false;
+	}
+
 	if (SignalProcessing == ESignalProcessing::Reflections)
 	{
 		// Ray traced reflection and SSR.
-		return Platform == SP_PCD3D_SM5 || Platform == SP_XBOXONE_D3D12;
+		return Platform == SP_PCD3D_SM5;
 	}
 	else if (
 		SignalProcessing == ESignalProcessing::MonochromaticPenumbra ||
