Making RuntimeVirtualTextureRender.h public so that it can be called by incoming thumbnail and baking systems
Add very simple debug shading mode to runtime virtual texture
Will extend to other modes later

diff --git a/Engine/Shaders/Private/VirtualTextureMaterial.usf b/Engine/Shaders/Private/VirtualTextureMaterial.usf
--- a/Engine/Shaders/Private/VirtualTextureMaterial.usf
+++ b/Engine/Shaders/Private/VirtualTextureMaterial.usf
@@ -85,6 +85,9 @@ void FPixelShaderInOut_MainPS(
 	float Opacity = GetMaterialOpacity(PixelMaterialInputs);
 #endif
 
+	// Apply debug shading
+	BaseColor = lerp(BaseColor, half3(In.SvPosition.xy * ResolvedView.ViewSizeAndInvSize.zw, 0), ResolvedView.VirtualTextureParams.y * 0.25);
+
 #if defined(OUT_BASECOLOR)
 	Out.MRT[0] = float4(BaseColor, 1.f) * Opacity;
 #elif defined(OUT_BASECOLOR_NORMAL)
diff --git a/Engine/Source/Runtime/Engine/Public/VT/RuntimeVirtualTextureEnum.h b/Engine/Source/Runtime/Engine/Public/VT/RuntimeVirtualTextureEnum.h
--- a/Engine/Source/Runtime/Engine/Public/VT/RuntimeVirtualTextureEnum.h
+++ b/Engine/Source/Runtime/Engine/Public/VT/RuntimeVirtualTextureEnum.h
@@ -35,3 +35,11 @@ enum class ERuntimeVirtualTextureMainPassType : uint8
 	/** We will render to any valid virtual texture target AND the main pass. Use this for items that need to both read and write the virtual texture. For example, some landscape setups need this. */
 	Always UMETA(DisplayName = "Virtual Texture AND Main Pass"),
 };
+
+/** Enumeration of runtime virtual texture debug modes. */
+UENUM()
+enum class ERuntimeVirtualTextureDebugType
+{
+	None,
+	Debug
+};
diff --git a/Engine/Source/Runtime/Renderer/Private/VT/RuntimeVirtualTextureProducer.cpp b/Engine/Source/Runtime/Renderer/Private/VT/RuntimeVirtualTextureProducer.cpp
--- a/Engine/Source/Runtime/Renderer/Private/VT/RuntimeVirtualTextureProducer.cpp
+++ b/Engine/Source/Runtime/Renderer/Private/VT/RuntimeVirtualTextureProducer.cpp
@@ -3,8 +3,8 @@
 #include "RuntimeVirtualTextureProducer.h"
 
 #include "RendererInterface.h"
-#include "RuntimeVirtualTextureRender.h"
 #include "ScenePrivate.h"
+#include "VT/RuntimeVirtualTextureRender.h"
 #include "VT/RuntimeVirtualTextureSceneProxy.h"
 
 
@@ -90,7 +90,8 @@ void FRuntimeVirtualTextureFinalizer::Finalize(FRHICommandListImmediate& RHICmdL
 			UVToWorld,
 			UVRange, 
 			Entry.vLevel,
-			Desc.MaxLevel);
+			Desc.MaxLevel,
+			ERuntimeVirtualTextureDebugType::None);
 	}
 
 	Tiles.SetNumUnsafeInternal(0);
diff --git a/Engine/Source/Runtime/Renderer/Private/VT/RuntimeVirtualTextureRender.cpp b/Engine/Source/Runtime/Renderer/Private/VT/RuntimeVirtualTextureRender.cpp
--- a/Engine/Source/Runtime/Renderer/Private/VT/RuntimeVirtualTextureRender.cpp
+++ b/Engine/Source/Runtime/Renderer/Private/VT/RuntimeVirtualTextureRender.cpp
@@ -1,6 +1,6 @@
 // Copyright 1998-2019 Epic Games, Inc. All Rights Reserved.
 
-#include "RuntimeVirtualTextureRender.h"
+#include "VT/RuntimeVirtualTextureRender.h"
 
 #include "GlobalShader.h"
 #include "GPUScene.h"
@@ -612,7 +612,8 @@ namespace RuntimeVirtualTexture
 		FTransform const& UVToWorld,
 		FBox2D const& UVRange,
 		uint8 vLevel,
-		uint8 MaxLevel)
+		uint8 MaxLevel,
+		ERuntimeVirtualTextureDebugType DebugType)
 	{
 		SCOPED_DRAW_EVENT(RHICmdList, VirtualTextureDynamicCache);
 
@@ -662,7 +663,7 @@ namespace RuntimeVirtualTexture
 		View->CachedViewUniformShaderParameters = MakeUnique<FViewUniformShaderParameters>();
 		View->SetupUniformBufferParameters(SceneContext, nullptr, 0, *View->CachedViewUniformShaderParameters);
 		View->CachedViewUniformShaderParameters->WorldToVirtualTexture = WorldToUVRotate.ToMatrixNoScale();
-		View->CachedViewUniformShaderParameters->VirtualTextureParams = FVector4((float)vLevel, (float)MaxLevel, OrthoWidth/(float)TextureSize.X, OrthoHeight / (float)TextureSize.Y);
+		View->CachedViewUniformShaderParameters->VirtualTextureParams = FVector4((float)vLevel, DebugType == ERuntimeVirtualTextureDebugType::Debug ? 1.f : 0.f, OrthoWidth / (float)TextureSize.X, OrthoHeight / (float)TextureSize.Y);
 		View->ViewUniformBuffer = TUniformBufferRef<FViewUniformShaderParameters>::CreateUniformBufferImmediate(*View->CachedViewUniformShaderParameters, UniformBuffer_SingleFrame);
 		UploadDynamicPrimitiveShaderDataForView(RHICmdList, *(const_cast<FScene*>(Scene)), *View);
 		Scene->UniformBuffers.VirtualTextureViewUniformBuffer.UpdateUniformBufferImmediate(*View->CachedViewUniformShaderParameters);
diff --git a/Engine/Source/Runtime/Renderer/Private/VT/RuntimeVirtualTextureRender.h b/Engine/Source/Runtime/Renderer/Private/VT/RuntimeVirtualTextureRender.h
deleted file mode 100644
diff --git a/Engine/Source/Runtime/Renderer/Public/VT/RuntimeVirtualTextureRender.h b/Engine/Source/Runtime/Renderer/Public/VT/RuntimeVirtualTextureRender.h
new file mode 100644
--- /dev/null
+++ b/Engine/Source/Runtime/Renderer/Public/VT/RuntimeVirtualTextureRender.h
@@ -0,0 +1,31 @@
+// Copyright 1998-2019 Epic Games, Inc. All Rights Reserved.
+
+#pragma once
+
+#include "CoreMinimal.h"
+
+enum class ERuntimeVirtualTextureDebugType;
+enum class ERuntimeVirtualTextureMaterialType : uint8;
+class FRHICommandListImmediate;
+class FRHITexture2D;
+class FScene;
+
+namespace RuntimeVirtualTexture
+{
+	/**
+	 * Render a single page of a virtual texture with a given material.
+	 * todo[vt]: Likely to be more optimal to batch several pages at a time and share setup/visibility/render targets.
+	 */
+	RENDERER_API void RenderPage(
+		FRHICommandListImmediate& RHICmdList,
+		FScene* Scene,
+		uint32 RuntimeVirtualTextureMask,
+		ERuntimeVirtualTextureMaterialType MaterialType,
+		FRHITexture2D* Texture0, FBox2D const& DestBox0,
+		FRHITexture2D* Texture1, FBox2D const& DestBox1,
+		FTransform const& UVToWorld,
+		FBox2D const& UVRange,
+		uint8 vLevel,
+		uint8 MaxLevel,
+		ERuntimeVirtualTextureDebugType DebugType);
+}
