Add support for max level to flush runtime virtual texture mips

diff --git a/Engine/Source/Runtime/Renderer/Private/VT/RuntimeVirtualTextureSceneProxy.cpp b/Engine/Source/Runtime/Renderer/Private/VT/RuntimeVirtualTextureSceneProxy.cpp
--- a/Engine/Source/Runtime/Renderer/Private/VT/RuntimeVirtualTextureSceneProxy.cpp
+++ b/Engine/Source/Runtime/Renderer/Private/VT/RuntimeVirtualTextureSceneProxy.cpp
@@ -32,6 +32,7 @@ FRuntimeVirtualTextureSceneProxy::FRuntimeVirtualTextureSceneProxy(URuntimeVirtu
 		FVTProducerDescription Desc;
 		VirtualTexture->GetProducerDescription(Desc, Transform);
 		VirtualTextureSize = FIntPoint(Desc.WidthInBlocks * Desc.BlockWidthInTiles * Desc.TileSize, Desc.HeightInBlocks * Desc.BlockHeightInTiles * Desc.TileSize);
+		MaxDirtyLevel = Desc.MaxLevel;
 
 		// The Producer object created here will be passed into the virtual texture system which will take ownership.
 		// The Initialize() call will allocate the virtual texture by spawning work on the render thread.
@@ -103,13 +104,13 @@ void FRuntimeVirtualTextureSceneProxy::FlushDirtyPages()
 
 		if (bCombinedFlush)
 		{
-			FVirtualTextureSystem::Get().FlushCache(ProducerHandle, CombinedDirtyRect);
+			FVirtualTextureSystem::Get().FlushCache(ProducerHandle, CombinedDirtyRect, MaxDirtyLevel);
 		}
 		else
 		{
 			for (FIntRect Rect : DirtyRects)
 			{
-				FVirtualTextureSystem::Get().FlushCache(ProducerHandle, Rect);
+				FVirtualTextureSystem::Get().FlushCache(ProducerHandle, Rect, MaxDirtyLevel);
 			}
 		}
 	}
diff --git a/Engine/Source/Runtime/Renderer/Private/VT/RuntimeVirtualTextureSceneProxy.h b/Engine/Source/Runtime/Renderer/Private/VT/RuntimeVirtualTextureSceneProxy.h
--- a/Engine/Source/Runtime/Renderer/Private/VT/RuntimeVirtualTextureSceneProxy.h
+++ b/Engine/Source/Runtime/Renderer/Private/VT/RuntimeVirtualTextureSceneProxy.h
@@ -45,6 +45,8 @@ private:
 	/** Virtual texture size of the URuntimeVirtualTexture object. */
 	FIntPoint VirtualTextureSize;
 
+	/** Maxmimum mip level to mark dirty. Can be less than the virtual texture's MaxLevel if we have streaming mips. */
+	uint32 MaxDirtyLevel;
 	/** Array of dirty rectangles to process at the next flush. */
 	TArray<FIntRect> DirtyRects;
 	/** Combined dirty rectangle to process at the next flush. */
diff --git a/Engine/Source/Runtime/Renderer/Private/VT/TexturePagePool.cpp b/Engine/Source/Runtime/Renderer/Private/VT/TexturePagePool.cpp
--- a/Engine/Source/Runtime/Renderer/Private/VT/TexturePagePool.cpp
+++ b/Engine/Source/Runtime/Renderer/Private/VT/TexturePagePool.cpp
@@ -87,7 +87,7 @@ void FTexturePagePool::EvictPages(FVirtualTextureSystem* System, const FVirtualT
 	}
 }
 
-void FTexturePagePool::EvictPages(FVirtualTextureSystem* System, FVirtualTextureProducerHandle const& ProducerHandle, FVTProducerDescription const& Desc, FIntRect const& TextureRegion, TArray<union FVirtualTextureLocalTile>& OutLocked)
+void FTexturePagePool::EvictPages(FVirtualTextureSystem* System, FVirtualTextureProducerHandle const& ProducerHandle, FVTProducerDescription const& Desc, FIntRect const& TextureRegion, uint32 MaxLevel, TArray<union FVirtualTextureLocalTile>& OutLocked)
 {
 	//todo[vt]: 
 	// Simple linear iteration of all physical pages here. Can we do better?
@@ -99,22 +99,25 @@ void FTexturePagePool::EvictPages(FVirtualTextureSystem* System, FVirtualTexture
 			const uint32 vAddress = Pages[i].Local_vAddress;
 			const uint32 vLevel = Pages[i].Local_vLevel;
 
-			const int32 TileSize = Desc.TileSize << vLevel;
-			const int32 X = FMath::ReverseMortonCode2(vAddress) * TileSize;
-			const int32 Y = FMath::ReverseMortonCode2(vAddress >> 1) * TileSize;
-			const int32 TileBorderSize = Desc.TileBorderSize << vLevel;
-			const FIntRect PageRect(X - TileBorderSize, Y - TileBorderSize, X + TileSize + TileBorderSize, Y + TileSize + TileBorderSize);
-
-			if (!(PageRect.Min.X > TextureRegion.Max.X) && !(TextureRegion.Min.X > PageRect.Max.X) && !(PageRect.Min.Y > TextureRegion.Max.Y) && !(TextureRegion.Min.Y > PageRect.Max.Y))
+			if (vLevel <= MaxLevel)
 			{
-				if (!FreeHeap.IsPresent(i))
-				{
-					OutLocked.Add(FVirtualTextureLocalTile(ProducerHandle, vAddress, vLevel));
-				}
-				else
+				const int32 TileSize = Desc.TileSize << vLevel;
+				const int32 X = FMath::ReverseMortonCode2(vAddress) * TileSize;
+				const int32 Y = FMath::ReverseMortonCode2(vAddress >> 1) * TileSize;
+				const int32 TileBorderSize = Desc.TileBorderSize << vLevel;
+				const FIntRect PageRect(X - TileBorderSize, Y - TileBorderSize, X + TileSize + TileBorderSize, Y + TileSize + TileBorderSize);
+
+				if (!(PageRect.Min.X > TextureRegion.Max.X) && !(TextureRegion.Min.X > PageRect.Max.X) && !(PageRect.Min.Y > TextureRegion.Max.Y) && !(TextureRegion.Min.Y > PageRect.Max.Y))
 				{
-					UnmapAllPages(System, i, false);
-					FreeHeap.Update(0, i);
+					if (!FreeHeap.IsPresent(i))
+					{
+						OutLocked.Add(FVirtualTextureLocalTile(ProducerHandle, vAddress, vLevel));
+					}
+					else
+					{
+						UnmapAllPages(System, i, false);
+						FreeHeap.Update(0, i);
+					}
 				}
 			}
 		}
diff --git a/Engine/Source/Runtime/Renderer/Private/VT/TexturePagePool.h b/Engine/Source/Runtime/Renderer/Private/VT/TexturePagePool.h
--- a/Engine/Source/Runtime/Renderer/Private/VT/TexturePagePool.h
+++ b/Engine/Source/Runtime/Renderer/Private/VT/TexturePagePool.h
@@ -48,7 +48,7 @@ public:
 	 * Unmap/remove any pages that were allocated by the given producer and are inside the TextureRegion.
 	 * Outputs the locked pages that can't be unmapped to the OutLocked array.
 	 */
-	void EvictPages(FVirtualTextureSystem* System, FVirtualTextureProducerHandle const& ProducerHandle, FVTProducerDescription const& Desc, FIntRect const& TextureRegion, TArray<union FVirtualTextureLocalTile>& OutLocked);
+	void EvictPages(FVirtualTextureSystem* System, FVirtualTextureProducerHandle const& ProducerHandle, FVTProducerDescription const& Desc, FIntRect const& TextureRegion, uint32 MaxLevel, TArray<union FVirtualTextureLocalTile>& OutLocked);
 
 	/**
 	* Unmap all pages from the given space...pages will remain resident in the pool, but no longer by mapped to any page table
diff --git a/Engine/Source/Runtime/Renderer/Private/VT/VirtualTextureSystem.cpp b/Engine/Source/Runtime/Renderer/Private/VT/VirtualTextureSystem.cpp
--- a/Engine/Source/Runtime/Renderer/Private/VT/VirtualTextureSystem.cpp
+++ b/Engine/Source/Runtime/Renderer/Private/VT/VirtualTextureSystem.cpp
@@ -244,7 +244,7 @@ void FVirtualTextureSystem::FlushCache()
 	bFlushCaches = true;
 }
 
-void FVirtualTextureSystem::FlushCache(FVirtualTextureProducerHandle const& ProducerHandle, FIntRect const& TextureRegion)
+void FVirtualTextureSystem::FlushCache(FVirtualTextureProducerHandle const& ProducerHandle, FIntRect const& TextureRegion, uint32 MaxLevel)
 {
 	checkSlow(IsInRenderingThread());
 
@@ -267,7 +267,7 @@ void FVirtualTextureSystem::FlushCache(FVirtualTextureProducerHandle const& Prod
 	for (int32 i = 0; i < PhysicalSpacesForProducer.Num(); ++i)
 	{
 		FTexturePagePool& Pool = PhysicalSpacesForProducer[i]->GetPagePool();
-		Pool.EvictPages(this, ProducerHandle, ProducerDescription, TextureRegion, TransientCollectedPages);
+		Pool.EvictPages(this, ProducerHandle, ProducerDescription, TextureRegion, MaxLevel, TransientCollectedPages);
 	}
 
 	for (auto& Page : TransientCollectedPages)
diff --git a/Engine/Source/Runtime/Renderer/Private/VT/VirtualTextureSystem.h b/Engine/Source/Runtime/Renderer/Private/VT/VirtualTextureSystem.h
--- a/Engine/Source/Runtime/Renderer/Private/VT/VirtualTextureSystem.h
+++ b/Engine/Source/Runtime/Renderer/Private/VT/VirtualTextureSystem.h
@@ -64,7 +64,7 @@ public:
 	void LoadPendingTiles(FRHICommandListImmediate& RHICmdList, ERHIFeatureLevel::Type FeatureLevel);
 	
 	void FlushCache();
-	void FlushCache(FVirtualTextureProducerHandle const& ProducerHandle, FIntRect const& TextureRegion);
+	void FlushCache(FVirtualTextureProducerHandle const& ProducerHandle, FIntRect const& TextureRegion, uint32 MaxLevel);
 
 private:
 	friend class FFeedbackAnalysisTask;
