Fixed double release on AtmosphereUniformBuffer

diff --git a/Engine/Source/Runtime/Renderer/Private/SkyAtmosphereRendering.cpp b/Engine/Source/Runtime/Renderer/Private/SkyAtmosphereRendering.cpp
--- a/Engine/Source/Runtime/Renderer/Private/SkyAtmosphereRendering.cpp
+++ b/Engine/Source/Runtime/Renderer/Private/SkyAtmosphereRendering.cpp
@@ -347,9 +347,6 @@ FSkyAtmosphereRenderSceneInfo::FSkyAtmosphereRenderSceneInfo(const USkyAtmospher
 
 FSkyAtmosphereRenderSceneInfo::~FSkyAtmosphereRenderSceneInfo()
 {
-	// Destructor is called on the render thread so safe to release here
-	AtmosphereUniformBuffer->Release();
-	//TransmittanceLutTexture->Release(); This is automatically tracked by FRenderTargetPool and RefCount
 }
 
 void FSkyAtmosphereRenderSceneInfo::PrepareSunLightProxy(FLightSceneInfo& SunLight) const
