Use a dummy with the same color as the clear color of separate translucency render target

diff --git a/Engine/Source/Runtime/Renderer/Private/PostProcess/PostProcessing.cpp b/Engine/Source/Runtime/Renderer/Private/PostProcess/PostProcessing.cpp
--- a/Engine/Source/Runtime/Renderer/Private/PostProcess/PostProcessing.cpp
+++ b/Engine/Source/Runtime/Renderer/Private/PostProcess/PostProcessing.cpp
@@ -1362,6 +1362,11 @@ void FPostProcessing::Process(FRHICommandListImmediate& RHICmdList, const FViewI
 					check(!FSceneRenderTargets::Get(RHICmdList).SeparateTranslucencyRT);
 				}
 			}
+			else
+			{
+				FRenderingCompositePass* NodeSeparateTranslucency = Context.Graph.RegisterPass(new (FMemStack::Get()) FRCPassPostProcessInput(FSceneRenderTargets::Get(RHICmdList).GetSeparateTranslucencyDummy()));
+				SeparateTranslucency = FRenderingCompositeOutputRef(NodeSeparateTranslucency);
+			}
 		}
 
 		bool bVisualizeHDR = View.Family->EngineShowFlags.VisualizeHDR && FeatureLevel >= ERHIFeatureLevel::SM5;
diff --git a/Engine/Source/Runtime/Renderer/Private/PostProcess/SceneRenderTargets.cpp b/Engine/Source/Runtime/Renderer/Private/PostProcess/SceneRenderTargets.cpp
--- a/Engine/Source/Runtime/Renderer/Private/PostProcess/SceneRenderTargets.cpp
+++ b/Engine/Source/Runtime/Renderer/Private/PostProcess/SceneRenderTargets.cpp
@@ -1550,6 +1550,12 @@ TRefCountPtr<IPooledRenderTarget>& FSceneRenderTargets::GetSeparateTranslucency(
 	return SeparateTranslucencyRT;
 }
 
+TRefCountPtr<IPooledRenderTarget>& FSceneRenderTargets::GetSeparateTranslucencyDummy()
+{
+	// Make sure the dummy is consistent with the clear color of separate translucnecy render target
+	return GSystemTextures.BlackAlphaOneDummy;
+}
+
 TRefCountPtr<IPooledRenderTarget>& FSceneRenderTargets::GetDownsampledTranslucencyDepth(FRHICommandList& RHICmdList, FIntPoint Size)
 {
 	if (!DownsampledTranslucencyDepthRT || DownsampledTranslucencyDepthRT->GetDesc().Extent != Size)
diff --git a/Engine/Source/Runtime/Renderer/Private/PostProcess/SceneRenderTargets.h b/Engine/Source/Runtime/Renderer/Private/PostProcess/SceneRenderTargets.h
--- a/Engine/Source/Runtime/Renderer/Private/PostProcess/SceneRenderTargets.h
+++ b/Engine/Source/Runtime/Renderer/Private/PostProcess/SceneRenderTargets.h
@@ -295,6 +295,9 @@ public:
 	/** Separate translucency buffer can be downsampled or not (as it is used to store the AfterDOF translucency) */
 	TRefCountPtr<IPooledRenderTarget>& GetSeparateTranslucency(FRHICommandList& RHICmdList, FIntPoint Size);
 
+	/** Returns a dummy texture that can be used for shader resource binding when separate translucency render target is not available */
+	TRefCountPtr<IPooledRenderTarget>& GetSeparateTranslucencyDummy();
+
 	bool IsDownsampledTranslucencyDepthValid()
 	{
 		return DownsampledTranslucencyDepthRT != nullptr;
