Fix RuntimeVirtualTexture RemoveLowMips logic
Previously RemoveLowMips didn't work correctly with StreamLowMips

diff --git a/Engine/Source/Editor/VirtualTexturingEditor/Private/RuntimeVirtualTextureBuild.cpp b/Engine/Source/Editor/VirtualTexturingEditor/Private/RuntimeVirtualTextureBuild.cpp
--- a/Engine/Source/Editor/VirtualTexturingEditor/Private/RuntimeVirtualTextureBuild.cpp
+++ b/Engine/Source/Editor/VirtualTexturingEditor/Private/RuntimeVirtualTextureBuild.cpp
@@ -52,7 +52,7 @@ namespace RuntimeVirtualTexture
 		const int32 TextureSizeX = VTDesc.WidthInBlocks * VTDesc.BlockWidthInTiles * TileSize;
 		const int32 TextureSizeY = VTDesc.HeightInBlocks * VTDesc.BlockHeightInTiles * TileSize;
 		const int32 MaxLevel = (int32)FMath::CeilLogTwo(FMath::Max(VTDesc.BlockWidthInTiles, VTDesc.BlockHeightInTiles));
-		const int32 RenderLevel = FMath::Max(MaxLevel - RuntimeVirtualTexture->GetStreamLowMips() + 1, 1);
+		const int32 RenderLevel = FMath::Max(MaxLevel - RuntimeVirtualTexture->GetStreamLowMips() + 1, 0);
 		const int32 StreamingTextureSizeX = FMath::Max(TileSize, TextureSizeX >> RenderLevel);
 		const int32 StreamingTextureSizeY = FMath::Max(TileSize, TextureSizeY >> RenderLevel);
 		const int32 NumTilesX = StreamingTextureSizeX / TileSize;
diff --git a/Engine/Source/Runtime/Engine/Classes/VT/RuntimeVirtualTexture.h b/Engine/Source/Runtime/Engine/Classes/VT/RuntimeVirtualTexture.h
--- a/Engine/Source/Runtime/Engine/Classes/VT/RuntimeVirtualTexture.h
+++ b/Engine/Source/Runtime/Engine/Classes/VT/RuntimeVirtualTexture.h
@@ -37,11 +37,11 @@ protected:
 	int32 TileBorderSize = 2; // 4
 
 	/** Number of low mips to cut from the virtual texture. This can reduce peak virtual texture update cost but will also increase the probability of mip shimmering. */
-	UPROPERTY(EditAnywhere, BlueprintReadOnly, Category = LowMips, meta = (UIMin = "0", UIMax = "5", DisplayName = "Number of low mips to remove from the virtual texture"))
+	UPROPERTY(EditAnywhere, BlueprintReadOnly, Category = LowMips, meta = (UIMin = "0", UIMax = "6", DisplayName = "Number of low mips to remove from the virtual texture"))
 	int32 RemoveLowMips = 0;
 
 	/** Number of low mips to serialize and stream for the virtual texture. This can reduce rendering update cost. */
-	UPROPERTY(EditAnywhere, BlueprintReadOnly, Category = LowMips, meta = (UIMin = "0", UIMax = "7", DisplayName = "Number of low mips to stream to the virtual texture"))
+	UPROPERTY(EditAnywhere, BlueprintReadOnly, Category = LowMips, meta = (UIMin = "0", UIMax = "6", DisplayName = "Number of low mips to stream to the virtual texture"))
 	int32 StreamLowMips = 0;
 
 	/** Texture object containing streamed low mips. */
@@ -66,9 +66,9 @@ public:
 	/** Public getter for virtual texture tile border size */
 	int32 GetTileBorderSize() const { return 2 * FMath::Clamp(TileBorderSize, 0, 4); }
 	/** Public getter for virtual texture removed low mips */
-	int32 GetRemoveLowMips() const { return RemoveLowMips; }
+	int32 GetRemoveLowMips() const { return FMath::Clamp(RemoveLowMips, 0, 5); }
 	/** Public getter for virtual texture streaming low mips */
-	int32 GetStreamLowMips() const { return StreamLowMips; }
+	int32 GetStreamLowMips() const { return FMath::Clamp(StreamLowMips, 0, 6); }
 
 	/** Returns an approximate estimated value for the memory used by the page table texture. */
 	int32 GetEstimatedPageTableTextureMemoryKb() const;
diff --git a/Engine/Source/Runtime/Engine/Private/VT/RuntimeVirtualTexture.cpp b/Engine/Source/Runtime/Engine/Private/VT/RuntimeVirtualTexture.cpp
--- a/Engine/Source/Runtime/Engine/Private/VT/RuntimeVirtualTexture.cpp
+++ b/Engine/Source/Runtime/Engine/Private/VT/RuntimeVirtualTexture.cpp
@@ -217,7 +217,7 @@ void URuntimeVirtualTexture::GetProducerDescription(FVTProducerDescription& OutD
 
 	OutDesc.BlockWidthInTiles = Width / GetTileSize();
 	OutDesc.BlockHeightInTiles = Height / GetTileSize();
-	OutDesc.MaxLevel = FMath::Max((int32)FMath::CeilLogTwo(FMath::Max(OutDesc.BlockWidthInTiles, OutDesc.BlockHeightInTiles)) - RemoveLowMips, 1);
+	OutDesc.MaxLevel = FMath::Max((int32)FMath::CeilLogTwo(FMath::Max(OutDesc.BlockWidthInTiles, OutDesc.BlockHeightInTiles)) - GetRemoveLowMips(), 0);
 
 	// Set layer description based on material type
 	switch (MaterialType)
@@ -399,7 +399,6 @@ uint32 URuntimeVirtualTexture::GetStreamingTextureBuildHash() const
 			uint32 Size : 4;
 			uint32 TileSize : 4;
 			uint32 TileBorderSize : 4;
-			uint32 RemoveLowMips : 4;
 			uint32 StreamLowMips : 4;
 		};
 	};
@@ -410,8 +409,7 @@ uint32 URuntimeVirtualTexture::GetStreamingTextureBuildHash() const
 	Settings.CompressTextures = (uint32)bCompressTextures;
 	Settings.Size = (uint32)Size;
 	Settings.TileSize = (uint32)TileSize;
-	Settings.RemoveLowMips = (uint32)RemoveLowMips;
-	Settings.StreamLowMips = (uint32)StreamLowMips;
+	Settings.StreamLowMips = (uint32)GetStreamLowMips();
 
 	return Settings.PackedValue;
 }
@@ -457,7 +455,7 @@ IVirtualTexture* URuntimeVirtualTexture::CreateStreamingTextureProducer(IVirtual
 		FTexturePlatformData** StreamingTextureData = StreamingTexture->GetRunningPlatformData();
 		if (StreamingTextureData != nullptr && *StreamingTextureData != nullptr)
 		{
-			OutTransitionLevel = InMaxLevel - StreamLowMips;
+			OutTransitionLevel = FMath::Max(InMaxLevel - GetStreamLowMips() + 1, 0);
 			IVirtualTexture* StreamingProducer = new FUploadingVirtualTexture((*StreamingTextureData)->VTData, 0);
 			return new FVirtualTextureLevelRedirector(InProducer, StreamingProducer, OutTransitionLevel);
 		}
diff --git a/Engine/Source/Runtime/Engine/Private/VT/VirtualTextureLevelRedirector.cpp b/Engine/Source/Runtime/Engine/Private/VT/VirtualTextureLevelRedirector.cpp
--- a/Engine/Source/Runtime/Engine/Private/VT/VirtualTextureLevelRedirector.cpp
+++ b/Engine/Source/Runtime/Engine/Private/VT/VirtualTextureLevelRedirector.cpp
@@ -21,8 +21,8 @@ FVTRequestPageResult FVirtualTextureLevelRedirector::RequestPageData(
 	uint32 vAddress,
 	EVTRequestPagePriority Priority)
 {
-	int32 VirtualTextureIndex = vLevel <= TransitionLevel ? 0 : 1;
-	int32 vLevelOffset = vLevel <= TransitionLevel ? 0 : TransitionLevel + 1;
+	int32 VirtualTextureIndex = vLevel < TransitionLevel ? 0 : 1;
+	int32 vLevelOffset = vLevel < TransitionLevel ? 0 : TransitionLevel;
 	return VirtualTextures[VirtualTextureIndex]->RequestPageData(ProducerHandle, LayerMask, vLevel - vLevelOffset, vAddress, Priority);
 }
 
@@ -37,7 +37,7 @@ IVirtualTextureFinalizer* FVirtualTextureLevelRedirector::ProducePageData(
 	uint64 RequestHandle,
 	const FVTProduceTargetLayer* TargetLayers)
 {
-	int32 VirtualTextureIndex = vLevel <= TransitionLevel ? 0 : 1;
-	int32 vLevelOffset = vLevel <= TransitionLevel ? 0 : TransitionLevel + 1;
+	int32 VirtualTextureIndex = vLevel < TransitionLevel ? 0 : 1;
+	int32 vLevelOffset = vLevel < TransitionLevel ? 0 : TransitionLevel;
 	return VirtualTextures[VirtualTextureIndex]->ProducePageData(RHICmdList, FeatureLevel, Flags, ProducerHandle, LayerMask, vLevel - vLevelOffset, vAddress, RequestHandle, TargetLayers);
 }
diff --git a/Engine/Source/Runtime/Renderer/Private/VT/RuntimeVirtualTextureSceneProxy.cpp b/Engine/Source/Runtime/Renderer/Private/VT/RuntimeVirtualTextureSceneProxy.cpp
--- a/Engine/Source/Runtime/Renderer/Private/VT/RuntimeVirtualTextureSceneProxy.cpp
+++ b/Engine/Source/Runtime/Renderer/Private/VT/RuntimeVirtualTextureSceneProxy.cpp
@@ -27,7 +27,8 @@ FRuntimeVirtualTextureSceneProxy::FRuntimeVirtualTextureSceneProxy(URuntimeVirtu
 		// The producer description is calculated using the transform to determine the aspect ratio
 		FVTProducerDescription Desc;
 		VirtualTexture->GetProducerDescription(Desc, Transform);
-		VirtualTextureSize = FIntPoint(Desc.WidthInBlocks * Desc.BlockWidthInTiles * Desc.TileSize, Desc.HeightInBlocks * Desc.BlockHeightInTiles * Desc.TileSize);
+		VirtualTextureSize = FIntPoint(Desc.BlockWidthInTiles * Desc.TileSize, Desc.BlockHeightInTiles * Desc.TileSize);
+		// We only need to dirty flush up to the producer description MaxLevel which accounts for the RemoveLowMips
 		MaxDirtyLevel = Desc.MaxLevel;
 
 		const ERuntimeVirtualTextureMaterialType MaterialType = VirtualTexture->GetMaterialType();
@@ -37,11 +38,13 @@ FRuntimeVirtualTextureSceneProxy::FRuntimeVirtualTextureSceneProxy(URuntimeVirtu
 
 		if (InComponent->IsStreamingLowMips() && VirtualTexture->GetStreamLowMips() > 0)
 		{
+			// Streaming mips start from the MaxLevel before taking into account the RemoveLowMips
+			const int32 MaxLevel = FMath::CeilLogTwo(FMath::Max(Desc.BlockWidthInTiles, Desc.BlockHeightInTiles));
 			// Wrap our producer to use a streaming producer for low mips
 			int32 StreamingTransitionLevel;
-			Producer = VirtualTexture->CreateStreamingTextureProducer(Producer, Desc.MaxLevel, StreamingTransitionLevel);
+			Producer = VirtualTexture->CreateStreamingTextureProducer(Producer, MaxLevel, StreamingTransitionLevel);
 			// Any dirty flushes don't need to flush the streaming mips (they only change with a build step).
-			MaxDirtyLevel = StreamingTransitionLevel;
+			MaxDirtyLevel = FMath::Min(MaxDirtyLevel, StreamingTransitionLevel);
 		}
 
 		// The Initialize() call will allocate the virtual texture by spawning work on the render thread.
diff --git a/Engine/Source/Runtime/Renderer/Private/VT/RuntimeVirtualTextureSceneProxy.h b/Engine/Source/Runtime/Renderer/Private/VT/RuntimeVirtualTextureSceneProxy.h
--- a/Engine/Source/Runtime/Renderer/Private/VT/RuntimeVirtualTextureSceneProxy.h
+++ b/Engine/Source/Runtime/Renderer/Private/VT/RuntimeVirtualTextureSceneProxy.h
@@ -45,8 +45,8 @@ private:
 	/** Virtual texture size of the URuntimeVirtualTexture object. */
 	FIntPoint VirtualTextureSize;
 
-	/** Maxmimum mip level to mark dirty. Can be less than the virtual texture's MaxLevel if we have streaming mips. */
-	uint32 MaxDirtyLevel;
+	/** Maximum mip level to mark dirty. Can be less than the virtual texture's MaxLevel if we have streaming mips. */
+	int32 MaxDirtyLevel;
 	/** Array of dirty rectangles to process at the next flush. */
 	TArray<FIntRect> DirtyRects;
 	/** Combined dirty rectangle to process at the next flush. */
