Disabled pre-exposure outside the editor through the new project config "Enable pre-exposure only in the editor".
This is to prevent the buffer readback to stall the renderthread by waiting on the RHI threads on cooked builds while providing
support for physically based lighting in the editor.

diff --git a/Engine/Source/Runtime/Engine/Classes/Engine/RendererSettings.h b/Engine/Source/Runtime/Engine/Classes/Engine/RendererSettings.h
--- a/Engine/Source/Runtime/Engine/Classes/Engine/RendererSettings.h
+++ b/Engine/Source/Runtime/Engine/Classes/Engine/RendererSettings.h
@@ -445,6 +445,12 @@ class ENGINE_API URendererSettings : public UDeveloperSettings
 		ConfigRestartRequired=true))
 	uint32 bUsePreExposure : 1;
 
+	UPROPERTY(config, EditAnywhere, Category = DefaultSettings, meta = (
+		EditCondition = "bUsePreExposure", ConsoleVariable = "r.EyeAdaptation.EditorOnly", DisplayName = "Enable pre-exposure only in the editor",
+		ToolTip = "Whether pre-exposure should be an editor only feature. This is to because it currently has an impact on the renderthread performance.",
+		ConfigRestartRequired=false))
+	uint32 bEnablePreExposureOnlyInTheEditor : 1;
+
 	UPROPERTY(config, EditAnywhere, Category = DefaultSettings, meta = (
 		ConsoleVariable = "r.DefaultFeature.MotionBlur", DisplayName = "Motion Blur",
 		ToolTip = "Whether the default for MotionBlur is enabled or not (postprocess volume/camera/game setting can still override and enable or disable it independently)"))
diff --git a/Engine/Source/Runtime/Renderer/Private/PostProcess/PostProcessEyeAdaptation.cpp b/Engine/Source/Runtime/Renderer/Private/PostProcess/PostProcessEyeAdaptation.cpp
--- a/Engine/Source/Runtime/Renderer/Private/PostProcess/PostProcessEyeAdaptation.cpp
+++ b/Engine/Source/Runtime/Renderer/Private/PostProcess/PostProcessEyeAdaptation.cpp
@@ -54,6 +54,14 @@ TAutoConsoleVariable<int32> CVarEyeAdaptationBasicCompute(
 	TEXT("> 0 : Compute Shader (default) \n"),
 	ECVF_Scalability | ECVF_RenderThreadSafe);
 
+static TAutoConsoleVariable<int32> CVarEnablePreExposureOnlyInTheEditor(
+	TEXT("r.EyeAdaptation.EditorOnly"),
+	1,
+	TEXT("When pre-exposure is enabled, 0 to enable it everywhere, 1 to enable it only in the editor (default).\n")
+	TEXT("This is to because it currently has an impact on the renderthread performance\n"),
+	ECVF_ReadOnly);
+
+
 /**
  *   Shared functionality used in computing the eye-adaptation parameters
  *   Compute the parameters used for eye-adaptation.  These will default to values
@@ -984,7 +992,8 @@ void FSceneViewState::FEyeAdaptationRTManager::SwapRTs(bool bInUpdateLastExposur
 	QUICK_SCOPE_CYCLE_COUNTER(STAT_FEyeAdaptationRTManager_SwapRTs);
 
 	FRHICommandListImmediate& RHICmdList = FRHICommandListExecutor::GetImmediateCommandList();
-	if (bInUpdateLastExposure && PooledRenderTarget[CurrentBuffer].IsValid())
+
+	if (bInUpdateLastExposure && PooledRenderTarget[CurrentBuffer].IsValid() && (GIsEditor || CVarEnablePreExposureOnlyInTheEditor.GetValueOnRenderThread() == 0))
 	{
 		if (!ExposureTextureReadback)
 		{
