Clear the entire render target to improve performance (fast clear).

diff --git a/Engine/Source/Runtime/Renderer/Private/LightRendering.cpp b/Engine/Source/Runtime/Renderer/Private/LightRendering.cpp
--- a/Engine/Source/Runtime/Renderer/Private/LightRendering.cpp
+++ b/Engine/Source/Runtime/Renderer/Private/LightRendering.cpp
@@ -76,10 +76,11 @@ static TAutoConsoleVariable<int32> CVarMaxShadowDenoisingBatchSize(
 	TEXT("Maximum number of shadow to denoise at the same time."),
 	ECVF_RenderThreadSafe);
 
-	static TAutoConsoleVariable<int32> CVarMaxShadowRayTracingBatchSize(
-		TEXT("r.RayTracing.Shadow.MaxBatchSize"), 8,
-		TEXT("Maximum number of shadows to trace at the same time."),
-		ECVF_RenderThreadSafe);
+static TAutoConsoleVariable<int32> CVarMaxShadowRayTracingBatchSize(
+	TEXT("r.RayTracing.Shadow.MaxBatchSize"), 8,
+	TEXT("Maximum number of shadows to trace at the same time."),
+	ECVF_RenderThreadSafe);
+
 
 #if ENABLE_DEBUG_DISCARD_PROP
 static float GDebugLightDiscardProp = 0.0f;
@@ -1246,49 +1247,14 @@ void FDeferredShadingSceneRenderer::RenderLights(FRHICommandListImmediate& RHICm
 							const FViewInfo& View = Views[ViewIndex];
 							View.HeightfieldLightingViewInfo.ClearShadowing(View, RHICmdList, LightSceneInfo);
 						}
-
-						// Clear light attenuation for local lights with a quad covering their extents
-						const bool bClearLightScreenExtentsOnly = SortedLightInfo.SortKey.Fields.LightType != LightType_Directional;
-						// All shadows render with min blending
-						bool bClearToWhite = !bClearLightScreenExtentsOnly;
-
+					
 						FRHIRenderPassInfo RPInfo(ScreenShadowMaskTexture->GetRenderTargetItem().TargetableTexture, ERenderTargetActions::Load_Store);
 						RPInfo.DepthStencilRenderTarget.Action = MakeDepthStencilTargetActions(ERenderTargetActions::Load_DontStore, ERenderTargetActions::Load_Store);
 						RPInfo.DepthStencilRenderTarget.DepthStencilTarget = SceneContext.GetSceneDepthSurface();
 						RPInfo.DepthStencilRenderTarget.ExclusiveDepthStencil = FExclusiveDepthStencil::DepthRead_StencilWrite;
-						if (bClearToWhite)
-						{
-							RPInfo.ColorRenderTargets[0].Action = ERenderTargetActions::Clear_Store;
-						}
-
+						RPInfo.ColorRenderTargets[0].Action = ERenderTargetActions::Clear_Store;
 						TransitionRenderPassTargets(RHICmdList, RPInfo);
 						RHICmdList.BeginRenderPass(RPInfo, TEXT("ClearScreenShadowMask"));
-						if (bClearLightScreenExtentsOnly)
-						{
-							SCOPED_DRAW_EVENT(RHICmdList, ClearQuad);
-
-							for (int32 ViewIndex = 0; ViewIndex < Views.Num(); ViewIndex++)
-							{
-								const FViewInfo& View = Views[ViewIndex];
-								FIntRect ScissorRect;
-
-								if (!LightSceneInfo.Proxy->GetScissorRect(ScissorRect, View, View.ViewRect))
-								{
-									ScissorRect = View.ViewRect;
-								}
-
-								if (ScissorRect.Min.X < ScissorRect.Max.X && ScissorRect.Min.Y < ScissorRect.Max.Y)
-								{
-									RHICmdList.SetViewport(ScissorRect.Min.X, ScissorRect.Min.Y, 0.0f, ScissorRect.Max.X, ScissorRect.Max.Y, 1.0f);
-									DrawClearQuad(RHICmdList, true, FLinearColor(1, 1, 1, 1), false, 0, false, 0);
-								}
-								else
-								{
-									LightSceneInfo.Proxy->GetScissorRect(ScissorRect, View, View.ViewRect);
-								}
-							}
-						}
-
 						RHICmdList.EndRenderPass();
 
 						RenderShadowProjections(RHICmdList, &LightSceneInfo, ScreenShadowMaskTexture, bInjectedTranslucentVolume);
