Added shadow amount control to fade out directional light shadow.

diff --git a/Engine/Source/Runtime/Engine/Classes/Components/DirectionalLightComponent.h b/Engine/Source/Runtime/Engine/Classes/Components/DirectionalLightComponent.h
--- a/Engine/Source/Runtime/Engine/Classes/Components/DirectionalLightComponent.h
+++ b/Engine/Source/Runtime/Engine/Classes/Components/DirectionalLightComponent.h
@@ -155,6 +155,12 @@ class ENGINE_API UDirectionalLightComponent : public ULightComponent
 	**/
 	UPROPERTY(BlueprintReadOnly, interp, Category = Light, meta = (HideAlphaChannel), AdvancedDisplay)
 	FColor ModulatedShadowColor;
+	
+	/**
+	 * Control the amount of shadow occlusion. A value of 0 means no occlusion, thus no shadow.
+	 */
+	UPROPERTY(BlueprintReadOnly, interp, Category = Light, meta = (HideAlphaChannel, UIMin = "0", UIMax = "1", ClampMin = "0", ClampMax = "1"), AdvancedDisplay)
+	float ShadowAmount;
 
 	UPROPERTY(EditAnywhere, AdvancedDisplay, BlueprintReadOnly, Category=Light, meta=(DisplayName = "Atmosphere / Fog Sun Light"))
 	uint32 bUsedAsAtmosphereSunLight : 1;
@@ -189,6 +195,9 @@ class ENGINE_API UDirectionalLightComponent : public ULightComponent
 	UFUNCTION(BlueprintCallable, Category="Rendering|Lighting")
 	void SetLightShaftOverrideDirection(FVector NewValue);
 
+	UFUNCTION(BlueprintCallable, Category="Rendering|Lighting")
+	void SetShadowAmount(float NewValue);
+
 	//~ Begin ULightComponent Interface
 	virtual FVector4 GetLightPosition() const override;
 	virtual ELightComponentType GetLightType() const override;
diff --git a/Engine/Source/Runtime/Engine/Private/Components/DirectionalLightComponent.cpp b/Engine/Source/Runtime/Engine/Private/Components/DirectionalLightComponent.cpp
--- a/Engine/Source/Runtime/Engine/Private/Components/DirectionalLightComponent.cpp
+++ b/Engine/Source/Runtime/Engine/Private/Components/DirectionalLightComponent.cpp
@@ -175,6 +175,7 @@ public:
 		}
 		bCastModulatedShadows = Component->bCastModulatedShadows;
 		ModulatedShadowColor = FLinearColor(Component->ModulatedShadowColor);
+		ShadowAmount = Component->ShadowAmount;
 	}
 
 	void UpdateLightShaftOverrideDirection_GameThread(const UDirectionalLightComponent* Component)
@@ -819,6 +820,7 @@ UDirectionalLightComponent::UDirectionalLightComponent(const FObjectInitializer&
 	bCastVolumetricShadow = true;
 
 	ModulatedShadowColor = FColor(128, 128, 128);
+	ShadowAmount = 1.0f;
 }
 
 #if WITH_EDITOR
@@ -845,6 +847,7 @@ void UDirectionalLightComponent::PostEditChangeProperty(FPropertyChangedEvent& P
 	// max range is larger than UI
 	ShadowBias = FMath::Clamp(ShadowBias, 0.0f, 10.0f);
 	ShadowSlopeBias = FMath::Clamp(ShadowSlopeBias, 0.0f, 10.0f);
+	ShadowAmount = FMath::Clamp(ShadowAmount, 0.0f, 1.0f);
 
 	Super::PostEditChangeProperty(PropertyChangedEvent);
 }
@@ -1036,6 +1039,16 @@ void UDirectionalLightComponent::SetLightShaftOverrideDirection(FVector NewValue
 	}
 }
 
+void UDirectionalLightComponent::SetShadowAmount(float NewValue)
+{
+	if (AreDynamicDataChangesAllowed()
+		&& ShadowAmount != NewValue)
+	{
+		ShadowAmount = NewValue;
+		MarkRenderStateDirty();
+	}
+}
+
 void UDirectionalLightComponent::Serialize(FArchive& Ar)
 {
 	Super::Serialize(Ar);
diff --git a/Engine/Source/Runtime/Engine/Private/Components/LightComponent.cpp b/Engine/Source/Runtime/Engine/Private/Components/LightComponent.cpp
--- a/Engine/Source/Runtime/Engine/Private/Components/LightComponent.cpp
+++ b/Engine/Source/Runtime/Engine/Private/Components/LightComponent.cpp
@@ -285,6 +285,7 @@ FLightSceneProxy::FLightSceneProxy(const ULightComponent* InLightComponent)
 	, LevelName(InLightComponent->GetOutermost()->GetFName())
 	, FarShadowDistance(0)
 	, FarShadowCascadeCount(0)
+	, ShadowAmount(1.0f)
 	, SamplesPerPixel(1)
 {
 	check(SceneInterface);
@@ -684,6 +685,7 @@ void ULightComponent::PostEditChangeProperty(FPropertyChangedEvent& PropertyChan
 		PropertyName != GET_MEMBER_NAME_STRING_CHECKED(UDirectionalLightComponent, LightShaftOverrideDirection) &&
 		PropertyName != GET_MEMBER_NAME_STRING_CHECKED(UDirectionalLightComponent, bCastModulatedShadows) &&
 		PropertyName != GET_MEMBER_NAME_STRING_CHECKED(UDirectionalLightComponent, ModulatedShadowColor) &&
+		PropertyName != GET_MEMBER_NAME_STRING_CHECKED(UDirectionalLightComponent, ShadowAmount) &&
 		// Properties that should only unbuild lighting for a Static light (can be changed dynamically on a Stationary light)
 		(PropertyName != GET_MEMBER_NAME_STRING_CHECKED(ULightComponent, Intensity) || Mobility == EComponentMobility::Static) &&
 		(PropertyName != GET_MEMBER_NAME_STRING_CHECKED(ULightComponent, LightColor) || Mobility == EComponentMobility::Static) &&
diff --git a/Engine/Source/Runtime/Engine/Public/SceneManagement.h b/Engine/Source/Runtime/Engine/Public/SceneManagement.h
--- a/Engine/Source/Runtime/Engine/Public/SceneManagement.h
+++ b/Engine/Source/Runtime/Engine/Public/SceneManagement.h
@@ -1325,6 +1325,7 @@ public:
 	inline bool CastsShadowsFromCinematicObjectsOnly() const { return bCastShadowsFromCinematicObjectsOnly; }
 	inline bool CastsModulatedShadows() const { return bCastModulatedShadows; }
 	inline const FLinearColor& GetModulatedShadowColor() const { return ModulatedShadowColor; }
+	inline const float GetShadowAmount() const { return ShadowAmount; }
 	inline bool AffectsTranslucentLighting() const { return bAffectTranslucentLighting; }
 	inline bool Transmission() const { return bTransmission; }
 	inline bool UseRayTracedDistanceFieldShadows() const { return bUseRayTracedDistanceFieldShadows; }
@@ -1543,6 +1544,9 @@ protected:
 	/** Modulated shadow color. */
 	FLinearColor ModulatedShadowColor;
 
+	/** Control the amount of shadow occlusion. */
+	float ShadowAmount;
+
 	/** Samples per pixel for ray tracing */
 	uint32 SamplesPerPixel;
 
diff --git a/Engine/Source/Runtime/Renderer/Private/ShadowSetup.cpp b/Engine/Source/Runtime/Renderer/Private/ShadowSetup.cpp
--- a/Engine/Source/Runtime/Renderer/Private/ShadowSetup.cpp
+++ b/Engine/Source/Runtime/Renderer/Private/ShadowSetup.cpp
@@ -1855,11 +1855,11 @@ void FSceneRenderer::CreatePerObjectProjectedShadow(
 
 		// Calculate fading based on resolution
 		// Compute FadeAlpha before ShadowResolutionScale contribution (artists want to modify the softness of the shadow, not change the fade ranges)
-		const float ViewSpecificAlpha = CalculateShadowFadeAlpha( UnclampedResolution, ShadowFadeResolution, MinShadowResolution );
+		const float ViewSpecificAlpha = CalculateShadowFadeAlpha(UnclampedResolution, ShadowFadeResolution, MinShadowResolution) * LightSceneInfo->Proxy->GetShadowAmount();
 		MaxResolutionFadeAlpha = FMath::Max(MaxResolutionFadeAlpha, ViewSpecificAlpha);
 		ResolutionFadeAlphas.Add(ViewSpecificAlpha);
 
-		const float ViewSpecificPreShadowAlpha = CalculateShadowFadeAlpha(UnclampedResolution * CVarPreShadowResolutionFactor.GetValueOnRenderThread(), PreShadowFadeResolution, MinPreShadowResolution);
+		const float ViewSpecificPreShadowAlpha = CalculateShadowFadeAlpha(UnclampedResolution * CVarPreShadowResolutionFactor.GetValueOnRenderThread(), PreShadowFadeResolution, MinPreShadowResolution) * LightSceneInfo->Proxy->GetShadowAmount();
 		MaxResolutionPreShadowFadeAlpha = FMath::Max(MaxResolutionPreShadowFadeAlpha, ViewSpecificPreShadowAlpha);
 		ResolutionPreShadowFadeAlphas.Add(ViewSpecificPreShadowAlpha);
 
@@ -2451,7 +2451,7 @@ void FSceneRenderer::CreateWholeSceneProjectedShadow(
 			}
 
 			// Compute FadeAlpha before ShadowResolutionScale contribution (artists want to modify the softness of the shadow, not change the fade ranges)
-			const float FadeAlpha = CalculateShadowFadeAlpha( UnclampedResolution, ShadowFadeResolution, MinShadowResolution );
+			const float FadeAlpha = CalculateShadowFadeAlpha( UnclampedResolution, ShadowFadeResolution, MinShadowResolution ) * LightSceneInfo->Proxy->GetShadowAmount();
 			MaxFadeAlpha = FMath::Max(MaxFadeAlpha, FadeAlpha);
 			FadeAlphas.Add(FadeAlpha);
 
@@ -3290,15 +3290,16 @@ void FSceneRenderer::AddViewDependentWholeSceneShadowsForView(
 	{
 		FViewInfo& View = Views[ViewIndex];
 
+		const float LightShadowAmount = LightSceneInfo.Proxy->GetShadowAmount();
 		TArray<float, TInlineAllocator<2> > FadeAlphas;
 		FadeAlphas.Init(0.0f, Views.Num());
-		FadeAlphas[ViewIndex] = 1.0f;
+		FadeAlphas[ViewIndex] = LightShadowAmount;
 
 		if (IStereoRendering::IsAPrimaryView(View.StereoPass, GEngine->StereoRenderingDevice)
 			&& Views.IsValidIndex(ViewIndex + 1)
 			&& IStereoRendering::IsASecondaryView(Views[ViewIndex + 1].StereoPass, GEngine->StereoRenderingDevice))
 		{
-			FadeAlphas[ViewIndex + 1] = 1.0f;
+			FadeAlphas[ViewIndex + 1] = LightShadowAmount;
 		}		
 		
 		// If rendering in stereo mode we render shadow depths only for the left eye, but project for both eyes!
