Added atmospehre light direction override on SkyAtmosphere component.

diff --git a/Engine/Source/Runtime/Engine/Classes/Components/SkyAtmosphereComponent.h b/Engine/Source/Runtime/Engine/Classes/Components/SkyAtmosphereComponent.h
--- a/Engine/Source/Runtime/Engine/Classes/Components/SkyAtmosphereComponent.h
+++ b/Engine/Source/Runtime/Engine/Classes/Components/SkyAtmosphereComponent.h
@@ -123,6 +123,11 @@ class USkyAtmosphereComponent : public USceneComponent
 
 
 
+	UFUNCTION(BlueprintCallable, Category = "Rendering")
+	void OverrideAtmosphereLightDirection(int32 AtmosphereLightIndex, const FVector& LightDirection);
+
+
+
 protected:
 	//~ Begin UActorComponent Interface.
 	virtual void CreateRenderState_Concurrent() override;
diff --git a/Engine/Source/Runtime/Engine/Private/Components/SkyAtmosphereComponent.cpp b/Engine/Source/Runtime/Engine/Private/Components/SkyAtmosphereComponent.cpp
--- a/Engine/Source/Runtime/Engine/Private/Components/SkyAtmosphereComponent.cpp
+++ b/Engine/Source/Runtime/Engine/Private/Components/SkyAtmosphereComponent.cpp
@@ -235,6 +235,15 @@ void USkyAtmosphereComponent::Serialize(FArchive& Ar)
 	Ar << bStaticLightingBuiltGUID;
 }
 
+void USkyAtmosphereComponent::OverrideAtmosphereLightDirection(int32 AtmosphereLightIndex, const FVector& LightDirection)
+{
+	if (AreDynamicDataChangesAllowed())
+	{
+		FSceneInterface* Scene = GetWorld()->Scene;
+		Scene->OverrideSkyAtmosphereLightDirection(this, AtmosphereLightIndex, LightDirection);
+	}
+}
+
 
 /*=============================================================================
 	ASkyAtmosphere implementation.
diff --git a/Engine/Source/Runtime/Engine/Public/SceneInterface.h b/Engine/Source/Runtime/Engine/Public/SceneInterface.h
--- a/Engine/Source/Runtime/Engine/Public/SceneInterface.h
+++ b/Engine/Source/Runtime/Engine/Public/SceneInterface.h
@@ -261,6 +261,13 @@ public:
 	 * Returns the scene's unique FSkyAtmosphereRenderSceneInfo if it exists
 	 */
 	virtual FSkyAtmosphereRenderSceneInfo* GetSkyAtmosphereSceneInfo() = 0;
+	/**
+	 * Override a sky atmosphere light direction
+	 * @param SkyAtmosphereComponent - component to verify it is the actual unique SkyAtmosphere
+	 * @param AtmosphereLightIndex - the atmosphere light index to consider
+	 * @param LightDirection - the new light direction to override the atmosphere light with
+	 */
+	virtual void OverrideSkyAtmosphereLightDirection(const class USkyAtmosphereComponent* SkyAtmosphereComponent, int32 AtmosphereLightIndex, const FVector& LightDirection) = 0;
 
 	/**
 	 * Adds a wind source component to the scene.
diff --git a/Engine/Source/Runtime/Renderer/Private/DeferredShadingRenderer.cpp b/Engine/Source/Runtime/Renderer/Private/DeferredShadingRenderer.cpp
--- a/Engine/Source/Runtime/Renderer/Private/DeferredShadingRenderer.cpp
+++ b/Engine/Source/Runtime/Renderer/Private/DeferredShadingRenderer.cpp
@@ -851,7 +851,7 @@ void FDeferredShadingSceneRenderer::Render(FRHICommandListImmediate& RHICmdList)
 		{
 			if (Scene->AtmosphereLights[LightIndex])
 			{
-				Scene->GetSkyAtmosphereSceneInfo()->PrepareSunLightProxy(*Scene->AtmosphereLights[LightIndex]);
+				Scene->GetSkyAtmosphereSceneInfo()->PrepareSunLightProxy(LightIndex, *Scene->AtmosphereLights[LightIndex]);
 			}
 		}
 	}
diff --git a/Engine/Source/Runtime/Renderer/Private/MobileShadingRenderer.cpp b/Engine/Source/Runtime/Renderer/Private/MobileShadingRenderer.cpp
--- a/Engine/Source/Runtime/Renderer/Private/MobileShadingRenderer.cpp
+++ b/Engine/Source/Runtime/Renderer/Private/MobileShadingRenderer.cpp
@@ -323,7 +323,7 @@ void FMobileSceneRenderer::Render(FRHICommandListImmediate& RHICmdList)
 		{
 			if (Scene->AtmosphereLights[LightIndex])
 			{
-				Scene->GetSkyAtmosphereSceneInfo()->PrepareSunLightProxy(*Scene->AtmosphereLights[LightIndex]);
+				Scene->GetSkyAtmosphereSceneInfo()->PrepareSunLightProxy(LightIndex, *Scene->AtmosphereLights[LightIndex]);
 			}
 		}
 	}
diff --git a/Engine/Source/Runtime/Renderer/Private/RendererScene.cpp b/Engine/Source/Runtime/Renderer/Private/RendererScene.cpp
--- a/Engine/Source/Runtime/Renderer/Private/RendererScene.cpp
+++ b/Engine/Source/Runtime/Renderer/Private/RendererScene.cpp
@@ -3707,6 +3707,7 @@ public:
 	virtual void AddSkyAtmosphere(const class USkyAtmosphereComponent* SkyAtmosphereComponent, bool bStaticLightingBuilt) override {}
 	virtual void RemoveSkyAtmosphere(const class USkyAtmosphereComponent* SkyAtmosphereComponent) override {}
 	virtual FSkyAtmosphereRenderSceneInfo* GetSkyAtmosphereSceneInfo() override { return NULL; }
+	virtual void OverrideSkyAtmosphereLightDirection(const class USkyAtmosphereComponent* SkyAtmosphereComponent, int32 AtmosphereLightIndex, const FVector& LightDirection) override {}
 
 	virtual void AddWindSource(class UWindDirectionalSourceComponent* WindComponent) override {}
 	virtual void RemoveWindSource(class UWindDirectionalSourceComponent* WindComponent) override {}
diff --git a/Engine/Source/Runtime/Renderer/Private/ScenePrivate.h b/Engine/Source/Runtime/Renderer/Private/ScenePrivate.h
--- a/Engine/Source/Runtime/Renderer/Private/ScenePrivate.h
+++ b/Engine/Source/Runtime/Renderer/Private/ScenePrivate.h
@@ -2689,6 +2689,7 @@ public:
 	virtual void AddSkyAtmosphere(const class USkyAtmosphereComponent* SkyAtmosphereComponent, bool bStaticLightingBuilt) override;
 	virtual void RemoveSkyAtmosphere(const class USkyAtmosphereComponent* SkyAtmosphereComponent) override;
 	virtual FSkyAtmosphereRenderSceneInfo* GetSkyAtmosphereSceneInfo() override { return SkyAtmosphere; }
+	virtual void OverrideSkyAtmosphereLightDirection(const class USkyAtmosphereComponent* SkyAtmosphereComponent, int32 AtmosphereLightIndex, const FVector& LightDirection) override;
 
 	virtual void AddWindSource(UWindDirectionalSourceComponent* WindComponent) override;
 	virtual void RemoveWindSource(UWindDirectionalSourceComponent* WindComponent) override;
diff --git a/Engine/Source/Runtime/Renderer/Private/SceneRendering.cpp b/Engine/Source/Runtime/Renderer/Private/SceneRendering.cpp
--- a/Engine/Source/Runtime/Renderer/Private/SceneRendering.cpp
+++ b/Engine/Source/Runtime/Renderer/Private/SceneRendering.cpp
@@ -1246,7 +1246,7 @@ void FViewInfo::SetupUniformBufferParameters(
 				ViewUniformShaderParameters.AtmosphereLightDiscCosHalfApexAngle[Index] = FVector4(FMath::Cos(Light->Proxy->GetSunLightHalfApexAngleRadian()));
 				ViewUniformShaderParameters.AtmosphereLightDiscLuminance[Index] = Light->Proxy->GetOuterSpaceLuminance();
 				ViewUniformShaderParameters.AtmosphereLightColor[Index] = Light->Proxy->GetColor();
-				ViewUniformShaderParameters.AtmosphereLightDirection[Index] = -Light->Proxy->GetDirection();
+				ViewUniformShaderParameters.AtmosphereLightDirection[Index] = SkyAtmosphere->GetAtmosphereLightDirection(Index, -Light->Proxy->GetDirection());
 			}
 			else
 			{
diff --git a/Engine/Source/Runtime/Renderer/Private/SkyAtmosphereRendering.cpp b/Engine/Source/Runtime/Renderer/Private/SkyAtmosphereRendering.cpp
--- a/Engine/Source/Runtime/Renderer/Private/SkyAtmosphereRendering.cpp
+++ b/Engine/Source/Runtime/Renderer/Private/SkyAtmosphereRendering.cpp
@@ -328,11 +328,13 @@ static FLinearColor GetLightDiskLuminance(FLightSceneInfo& Light, FLinearColor L
 
 FSkyAtmosphereRenderSceneInfo::FSkyAtmosphereRenderSceneInfo(const USkyAtmosphereComponent* InComponent)
 	: bStaticLightingBuilt(false)
+	, Component_DoNotDereference(InComponent)
 	, AtmosphereSetup(*InComponent)
 	, TransmittanceAtZenith(AtmosphereSetup.GetTransmittanceAtGroundLevel(FVector(0.0f, 0.0f, 1.0f)))
 {
 	SkyLuminanceFactor = InComponent->SkyLuminanceFactor;
 	AerialPespectiveViewDistanceScale = InComponent->AerialPespectiveViewDistanceScale;
+	memset(OverrideAtmosphericLight, 0, sizeof(OverrideAtmosphericLight));
 
 	// Create a multiframe uniform buffer. A render command is used because FSkyAtmosphereRenderSceneInfo ctor is called on the Game thread.
 	TUniformBufferRef<FAtmosphereUniformShaderParameters>* AtmosphereUniformBufferPtr = &AtmosphereUniformBuffer;
@@ -349,18 +351,19 @@ FSkyAtmosphereRenderSceneInfo::~FSkyAtmosphereRenderSceneInfo()
 {
 }
 
-void FSkyAtmosphereRenderSceneInfo::PrepareSunLightProxy(FLightSceneInfo& SunLight) const
+void FSkyAtmosphereRenderSceneInfo::PrepareSunLightProxy(uint32 AtmosphereLightIndex, FLightSceneInfo& AtmosphereLight) const
 {
 	// See explanation in https://media.contentapi.ea.com/content/dam/eacom/frostbite/files/s2016-pbs-frostbite-sky-clouds-new.pdf page 26
 	const bool bAtmosphereAffectsSunIlluminance = true;
-	FLinearColor TransmittanceTowardSun = bAtmosphereAffectsSunIlluminance ? AtmosphereSetup.GetTransmittanceAtGroundLevel(-SunLight.Proxy->GetDirection()) : FLinearColor(FLinearColor::White);
+	const FVector AtmosphereLightDirection = GetAtmosphereLightDirection(AtmosphereLightIndex, -AtmosphereLight.Proxy->GetDirection());
+	FLinearColor TransmittanceTowardSun = bAtmosphereAffectsSunIlluminance ? AtmosphereSetup.GetTransmittanceAtGroundLevel(AtmosphereLightDirection) : FLinearColor(FLinearColor::White);
 	FLinearColor TransmittanceAtZenithFinal = bAtmosphereAffectsSunIlluminance ? TransmittanceAtZenith : FLinearColor(FLinearColor::White);
 
-	FLinearColor SunZenithIlluminance = SunLight.Proxy->GetColor();
+	FLinearColor SunZenithIlluminance = AtmosphereLight.Proxy->GetColor();
 	FLinearColor SunOuterSpaceIlluminance = SunZenithIlluminance / TransmittanceAtZenithFinal;
-	FLinearColor SunDiskOuterSpaceLuminance = GetLightDiskLuminance(SunLight, SunOuterSpaceIlluminance);
+	FLinearColor SunDiskOuterSpaceLuminance = GetLightDiskLuminance(AtmosphereLight, SunOuterSpaceIlluminance);
 
-	SunLight.Proxy->SetAtmosphereRelatedProperties(TransmittanceTowardSun / TransmittanceAtZenithFinal, SunDiskOuterSpaceLuminance);
+	AtmosphereLight.Proxy->SetAtmosphereRelatedProperties(TransmittanceTowardSun / TransmittanceAtZenithFinal, SunDiskOuterSpaceLuminance);
 }
 
 FTextureRHIRef FSkyAtmosphereRenderSceneInfo::GetDistantSkyLightLutTextureRHI()
@@ -372,7 +375,25 @@ FTextureRHIRef FSkyAtmosphereRenderSceneInfo::GetDistantSkyLightLutTextureRHI()
 	return GBlackTexture->TextureRHI;
 }
 
+void FSkyAtmosphereRenderSceneInfo::OverrideAtmosphereLightDirection(const class USkyAtmosphereComponent* SkyAtmosphereComponent, int32 AtmosphereLightIndex, const FVector& LightDirection)
+{
+	if (SkyAtmosphereComponent == Component_DoNotDereference)
+	{
+		check(AtmosphereLightIndex >= 0 && AtmosphereLightIndex < NUM_ATMOSPHERE_LIGHTS);
+		AtmosphereLightIndex = FMath::Clamp(AtmosphereLightIndex, 0, NUM_ATMOSPHERE_LIGHTS - 1);	// To make sure we do not crash, blueprint function cannot enforce input ranges
+		OverrideAtmosphericLight[AtmosphereLightIndex] = true;
+		OverrideAtmosphericLightDirection[AtmosphereLightIndex] = LightDirection;
+	}
+}
 
+FVector FSkyAtmosphereRenderSceneInfo::GetAtmosphereLightDirection(int32 AtmosphereLightIndex, const FVector& DefaultDirection) const
+{
+	if (OverrideAtmosphericLight[AtmosphereLightIndex])
+	{
+		return OverrideAtmosphericLightDirection[AtmosphereLightIndex];
+	}
+	return DefaultDirection;
+}
 
 /*=============================================================================
 	FScene functions
@@ -420,6 +441,22 @@ void FScene::RemoveSkyAtmosphere(const USkyAtmosphereComponent* SkyAtmosphereCom
 		} );
 }
 
+void FScene::OverrideSkyAtmosphereLightDirection(const class USkyAtmosphereComponent* SkyAtmosphereComponent, int32 AtmosphereLightIndex, const FVector& InLightDirection)
+{
+	FScene* Scene = this;
+	FVector LightDirection = InLightDirection;
+	LightDirection.Normalize();
+	ENQUEUE_RENDER_COMMAND(FOverrideSkyAtmosphereLightDirection)(
+		[Scene, SkyAtmosphereComponent, AtmosphereLightIndex, LightDirection](FRHICommandListImmediate& RHICmdList)
+	{
+		FSkyAtmosphereRenderSceneInfo* SkyInfo = Scene->GetSkyAtmosphereSceneInfo();
+		if (SkyInfo)
+		{
+			SkyInfo->OverrideAtmosphereLightDirection(SkyAtmosphereComponent, AtmosphereLightIndex, LightDirection);
+		}
+	});
+}
+
 void FScene::ResetAtmosphereLightsProperties()
 {
 	// Also rest the current atmospheric light to default atmosphere
diff --git a/Engine/Source/Runtime/Renderer/Private/SkyAtmosphereRendering.h b/Engine/Source/Runtime/Renderer/Private/SkyAtmosphereRendering.h
--- a/Engine/Source/Runtime/Renderer/Private/SkyAtmosphereRendering.h
+++ b/Engine/Source/Runtime/Renderer/Private/SkyAtmosphereRendering.h
@@ -7,6 +7,7 @@
 #pragma once
 
 #include "CoreMinimal.h"
+#include "EngineDefines.h"
 #include "RendererInterface.h"
 #include "RenderResource.h"
 #include "Rendering/SkyAtmosphereCommonData.h"
@@ -61,7 +62,7 @@ public:
 	~FSkyAtmosphereRenderSceneInfo();
 
 	/** Prepare the sun light data as a function of current atmosphere state. */
-	void PrepareSunLightProxy(FLightSceneInfo& SunLight) const;
+	void PrepareSunLightProxy(uint32 AtmosphereLightIndex, FLightSceneInfo& AtmosphereLight) const;
 
 	bool IsMultiScatteringEnabled() const { return AtmosphereSetup.MultiScatteringFactor > 0.0f; }
 	FLinearColor GetSkyLuminanceFactor() const { return SkyLuminanceFactor; }
@@ -76,10 +77,15 @@ public:
 	const FAtmosphereSetup& GetAtmosphereSetup() const { return AtmosphereSetup; }
 	const FAtmosphereUniformShaderParameters* GetAtmosphereShaderParameters() const { return &AtmosphereUniformShaderParameters; }
 
+	void OverrideAtmosphereLightDirection(const class USkyAtmosphereComponent* SkyAtmosphereComponent, int32 AtmosphereLightIndex, const FVector& LightDirection);
+	FVector GetAtmosphereLightDirection(int32 AtmosphereLightIndex, const FVector& DefaultDirection) const;
+
 	bool bStaticLightingBuilt;
 
 private:
 
+	const USkyAtmosphereComponent* Component_DoNotDereference;
+
 	FAtmosphereSetup AtmosphereSetup;
 	FAtmosphereUniformShaderParameters AtmosphereUniformShaderParameters;
 
@@ -91,6 +97,9 @@ private:
 	FLinearColor TransmittanceAtZenith;
 	FLinearColor SkyLuminanceFactor;
 	float AerialPespectiveViewDistanceScale;
+
+	bool OverrideAtmosphericLight[NUM_ATMOSPHERE_LIGHTS];
+	FVector OverrideAtmosphericLightDirection[NUM_ATMOSPHERE_LIGHTS];
 };
 
 bool ShouldRenderSkyAtmosphere(const FSkyAtmosphereRenderSceneInfo* SkyAtmosphere, EShaderPlatform ShaderPlatform);
