Use uint8 for RuntimeVirtualTextureMaterialType to reduce mem cost

diff --git a/Engine/Source/Runtime/Renderer/Private/SceneCore.h b/Engine/Source/Runtime/Renderer/Private/SceneCore.h
--- a/Engine/Source/Runtime/Renderer/Private/SceneCore.h
+++ b/Engine/Source/Runtime/Renderer/Private/SceneCore.h
@@ -224,7 +224,7 @@ public:
 	/** Whether the mesh batch can be used for rendering to a virtual texture. */
 	uint8 bRenderToVirtualTexture : 1;
 	/** What virtual texture material type this mesh batch should be rendered with. */
-	uint32 RuntimeVirtualTextureMaterialType : ERuntimeVirtualTextureMaterialType_NumBits;
+	uint8 RuntimeVirtualTextureMaterialType : ERuntimeVirtualTextureMaterialType_NumBits;
 
 	/** Cached from vertex factory to avoid dereferencing VF in InitViews. */
 	uint8 bSupportsCachingMeshDrawCommands : 1;
