Use 2xBC3 textures for BaseColor+Normal+Specular runtime virtual texture
Better normal quality then BC1+BC3
Old layout kept for now in a deprecated enum. Will probably remove after any comparisons are carried out.

diff --git a/Engine/Shaders/Private/VirtualTextureCommon.ush b/Engine/Shaders/Private/VirtualTextureCommon.ush
--- a/Engine/Shaders/Private/VirtualTextureCommon.ush
+++ b/Engine/Shaders/Private/VirtualTextureCommon.ush
@@ -554,6 +554,12 @@ float3 VirtualTextureUnpackNormalBC3(in float4 PackedValue)
 	return VirtualTextureUnpackNormal(PackedValue.ag);
 }
 
+/** Unpack normal from two BC3 textures. */
+float3 VirtualTextureUnpackNormalBC3BC3(in float4 PackedValue0, in float4 PackedValue1)
+{
+	return VirtualTextureUnpackNormal(float2(PackedValue0.a, PackedValue1.a));
+}
+
 /** Unpack normal from BC5. */
 float3 VirtualTextureUnpackNormalBC5(in float4 PackedValue)
 {
diff --git a/Engine/Shaders/Private/VirtualTextureCompress.usf b/Engine/Shaders/Private/VirtualTextureCompress.usf
--- a/Engine/Shaders/Private/VirtualTextureCompress.usf
+++ b/Engine/Shaders/Private/VirtualTextureCompress.usf
@@ -16,9 +16,11 @@ Texture2D<float4> RenderTexture1;
 SamplerState TextureSampler1;
 Texture2D<float4> RenderTexture2;
 SamplerState TextureSampler2;
-RWTexture2D<uint2> OutCompressTexture0;
+RWTexture2D<uint2> OutCompressTexture0_u2;
+RWTexture2D<uint4> OutCompressTexture0_u4;
 RWTexture2D<uint4> OutCompressTexture1;
 RWTexture2D<float4> OutCopyTexture0;
+RWTexture2D<float4> OutCopyTexture1;
 
 // Read a 4x4 color block ready for BC1 compression
 void ReadBlockRGB( Texture2D<float4> SourceTexture, SamplerState TextureSampler, float2 UV, float2 TexelUVSize, out float3 Block[16] )
@@ -104,15 +106,6 @@ void ReadBlockRG( Texture2D<float4> SourceTexture, SamplerState TextureSampler,
 	}
 }
 
-// Combine NormalY, Roughness and Smoothness into an RGB block ready for BC3 compression
-void CombineNormalSpecularBlocksForBC3( in float BlockNormalY[16], inout float3 BlockSpecular[16] )
-{
-	for (int i=0; i<16; i++)
-	{
-		BlockSpecular[i] = float3(BlockSpecular[i].x, BlockNormalY[i], BlockSpecular[i].y);
-	}
-}
-
 /** Compress base color to a single BC1 target. */
 [numthreads(8,8,1)]
 void CompressBaseColorCS(
@@ -129,7 +122,7 @@ void CompressBaseColorCS(
 	float3 BlockBaseColor[16];
 	ReadBlockRGB(RenderTexture0, TextureSampler0, SampleUV, TexelUVSize, BlockBaseColor);
 
-	OutCompressTexture0[ThreadId.xy + DestRect.xy] = CompressBC1Block_SRGB(BlockBaseColor);
+	OutCompressTexture0_u2[ThreadId.xy + DestRect.xy] = CompressBC1Block_SRGB(BlockBaseColor);
 }
 
 /** Compress base color and normal to BC1/BC5 target pair. */
@@ -152,11 +145,11 @@ void CompressBaseColorNormalCS(
 	float BlockNormalY[16];
 	ReadBlockRG(RenderTexture1, TextureSampler1, SampleUV, TexelUVSize, BlockNormalX, BlockNormalY);
 
-	OutCompressTexture0[ThreadId.xy + DestRect.xy] = CompressBC1Block_SRGB(BlockBaseColor);
+	OutCompressTexture0_u2[ThreadId.xy + DestRect.xy] = CompressBC1Block_SRGB(BlockBaseColor);
 	OutCompressTexture1[ThreadId.xy + DestRect.xy] = CompressBC5Block(BlockNormalX, BlockNormalY);
 }
 
-/** Compress base color, normal, roughness and specular to BC1/BC3 target pair. */
+/** Compress base color, normal, roughness and specular to BC3 target pair. */
 [numthreads(8,8,1)]
 void CompressBaseColorNormalSpecularCS(
 	uint3 ThreadId : SV_DispatchThreadID 
@@ -179,16 +172,65 @@ void CompressBaseColorNormalSpecularCS(
 	float3 BlockSpecular[16];
 	ReadBlockRGB(RenderTexture2, TextureSampler2, SampleUV, TexelUVSize, BlockSpecular);
 
-	CombineNormalSpecularBlocksForBC3(BlockNormalY, BlockSpecular);
+	OutCompressTexture0_u4[ThreadId.xy + DestRect.xy] = CompressBC3Block_SRGB(BlockBaseColor, BlockNormalX);
+	OutCompressTexture1[ThreadId.xy + DestRect.xy] = CompressBC3Block(BlockSpecular, BlockNormalY);
+}
+
+/** Copy normal X and Y to the alpha channels of the final target textures (without compression)
+  * This path is used when we disable texture compression, because we need to keep the same final channel layout. */
+[numthreads(8,8,1)]
+void CopyNormalSpecularCS(
+	uint3 ThreadId : SV_DispatchThreadID 
+	)
+{
+	uint2 SamplePos = ThreadId.xy;
+	if (any(ThreadId.xy >= DestRect.zw))
+		return;
+
+	float3 BaseColor = LinearToSrgb(RenderTexture0[SamplePos].xyz);
+	float2 NormalXY = RenderTexture1[SamplePos].xy;
+	float2 RoughnessSpecular = RenderTexture2[SamplePos].xy;
+	
+	OutCopyTexture0[ThreadId.xy + DestRect.xy] = float4(BaseColor, NormalXY.x);
+	OutCopyTexture1[ThreadId.xy + DestRect.xy] = float4(RoughnessSpecular.x, RoughnessSpecular.y, 0.0f, NormalXY.y);
+}
+
+/** Compress base color, normal, roughness and specular to BC1/BC3 target pair. */
+[numthreads(8,8,1)]
+void CompressBaseColorNormalSpecularDeprecatedCS(
+	uint3 ThreadId : SV_DispatchThreadID 
+	)
+{
+	uint2 SamplePos = ThreadId.xy * 4;
+	if (any(SamplePos >= DestRect.zw))
+		return;
+
+	float2 TexelUVSize = 1.f / float2(DestRect.zw);
+	float2 SampleUV = (float2(SamplePos) + 0.5f) * TexelUVSize;
+
+	float3 BlockBaseColor[16];
+	ReadBlockRGB(RenderTexture0, TextureSampler0, SampleUV, TexelUVSize, BlockBaseColor);
+
+	float BlockNormalX[16];
+	float BlockNormalY[16];
+	ReadBlockRG(RenderTexture1, TextureSampler1, SampleUV, TexelUVSize, BlockNormalX, BlockNormalY);
+
+	float3 BlockSpecular[16];
+	ReadBlockRGB(RenderTexture2, TextureSampler2, SampleUV, TexelUVSize, BlockSpecular);
+
+	for (int i=0; i<16; i++)
+	{
+		BlockSpecular[i] = float3(BlockSpecular[i].x, BlockNormalY[i], BlockSpecular[i].y);
+	}
 
-	OutCompressTexture0[ThreadId.xy + DestRect.xy] = CompressBC1Block_SRGB(BlockBaseColor);
+	OutCompressTexture0_u2[ThreadId.xy + DestRect.xy] = CompressBC1Block_SRGB(BlockBaseColor);
 	OutCompressTexture1[ThreadId.xy + DestRect.xy] = CompressBC3Block(BlockSpecular, BlockNormalX);
 }
 
 /** Combine and copy normal, roughness and specular to a single packed target (without compression).
   * This path is used when we disable texture compression, because we need to keep the same final channel layout. */
 [numthreads(8,8,1)]
-void CopyNormalSpecularCS(
+void CopyNormalSpecularDeprecatedCS(
 	uint3 ThreadId : SV_DispatchThreadID 
 	)
 {
diff --git a/Engine/Shaders/Private/VirtualTextureMaterial.usf b/Engine/Shaders/Private/VirtualTextureMaterial.usf
--- a/Engine/Shaders/Private/VirtualTextureMaterial.usf
+++ b/Engine/Shaders/Private/VirtualTextureMaterial.usf
@@ -96,7 +96,7 @@ void FPixelShaderInOut_MainPS(
 #elif defined(OUT_BASECOLOR_NORMAL_SPECULAR)
 	Out.MRT[0] = float4(BaseColor, 1.f) * Opacity;
 	Out.MRT[1] = float4(PackNormal(Normal), 0.f, 1.f) * Opacity;
-	Out.MRT[2] = float4(Roughness, Specular, 0.f, 1.f) * Opacity;
+	Out.MRT[2] = float4(Specular, Roughness, 0.f, 1.f) * Opacity;
 #elif defined(OUT_WORLDHEIGHT)
 	Out.MRT[0] = float4(PackWorldHeight(WorldHeight), 0, 0, 1);
 #endif
diff --git a/Engine/Source/Runtime/Engine/Private/Materials/HLSLMaterialTranslator.h b/Engine/Source/Runtime/Engine/Private/Materials/HLSLMaterialTranslator.h
--- a/Engine/Source/Runtime/Engine/Private/Materials/HLSLMaterialTranslator.h
+++ b/Engine/Source/Runtime/Engine/Private/Materials/HLSLMaterialTranslator.h
@@ -2060,6 +2060,15 @@ protected:
 		// Can only create temporaries for certain types
 		else if ((Type & (MCT_Float | MCT_VTPageTableResult)) || Type == MCT_ShadingModel)
 		{
+			// Check for existing
+			for (int32 i = 0; i < CurrentScopeChunks->Num(); ++i)
+			{
+				if ((*CurrentScopeChunks)[i].Hash == Hash)
+				{
+					return i;
+				}
+			}
+
 			const int32 CodeIndex = CurrentScopeChunks->Num();
 			// Allocate a local variable name
 			const FString SymbolName = CreateSymbolName(TEXT("Local"));
@@ -5070,28 +5079,46 @@ protected:
 		return AddInlinedCodeChunk(MCT_Float2, *SampleCode, *GetParameterCode(WorldPositionIndex), *GetParameterCode(P0), *GetParameterCode(P1), *GetParameterCode(P2));
 	}
 
-	virtual int32 VirtualTextureUnpack(int32 CodeIndex, EVirtualTextureUnpackType UnpackType) override
+	virtual int32 VirtualTextureUnpack(int32 CodeIndex0, int32 CodeIndex1, EVirtualTextureUnpackType UnpackType) override
 	{
-		if (CodeIndex != INDEX_NONE)
+		if (UnpackType == EVirtualTextureUnpackType::BaseColor)
 		{
-			if (UnpackType == EVirtualTextureUnpackType::NormalBC3)
-			{
-				FString	SampleCode(TEXT("VirtualTextureUnpackNormalBC3(%s)"));
-				return AddCodeChunk(MCT_Float3, *SampleCode, *GetParameterCode(CodeIndex));
-			}
-			if (UnpackType == EVirtualTextureUnpackType::NormalBC5)
-			{
-				FString	SampleCode(TEXT("VirtualTextureUnpackNormalBC5(%s)"));
-				return AddCodeChunk(MCT_Float3, *SampleCode, *GetParameterCode(CodeIndex));
-			}
-			else if (UnpackType == EVirtualTextureUnpackType::HeightR16)
-			{
-				FString	SampleCode(TEXT("VirtualTextureUnpackHeightR16(%s)"));
-				return AddCodeChunk(MCT_Float, *SampleCode, *GetParameterCode(CodeIndex));
-			}
+			return CodeIndex0 == INDEX_NONE ? INDEX_NONE : ComponentMask(CodeIndex0, 1, 1, 1, 0);
+		}
+		else if (UnpackType == EVirtualTextureUnpackType::NormalBC3)
+		{
+			FString	SampleCode(TEXT("VirtualTextureUnpackNormalBC3(%s)"));
+			return CodeIndex1 == INDEX_NONE ? INDEX_NONE : AddCodeChunk(MCT_Float3, *SampleCode, *GetParameterCode(CodeIndex1));
+		}
+		else if (UnpackType == EVirtualTextureUnpackType::NormalBC5)
+		{
+			FString	SampleCode(TEXT("VirtualTextureUnpackNormalBC5(%s)"));
+			return CodeIndex1 == INDEX_NONE ? INDEX_NONE : AddCodeChunk(MCT_Float3, *SampleCode, *GetParameterCode(CodeIndex1));
+		}
+		else if (UnpackType == EVirtualTextureUnpackType::NormalBC3BC3)
+		{
+			FString	SampleCode(TEXT("VirtualTextureUnpackNormalBC3BC3(%s, %s)"));
+			return CodeIndex0 == INDEX_NONE || CodeIndex1 == INDEX_NONE ? INDEX_NONE : AddCodeChunk(MCT_Float3, *SampleCode, *GetParameterCode(CodeIndex0), *GetParameterCode(CodeIndex1));
+		}
+		else if (UnpackType == EVirtualTextureUnpackType::SpecularR8)
+		{
+			return CodeIndex1 == INDEX_NONE ? INDEX_NONE : ComponentMask(CodeIndex1, 1, 0, 0, 0);
+		}
+		else if (UnpackType == EVirtualTextureUnpackType::RoughnessG8)
+		{
+			return CodeIndex1 == INDEX_NONE ? INDEX_NONE : ComponentMask(CodeIndex1, 0, 1, 0, 0);
+		}
+		else if (UnpackType == EVirtualTextureUnpackType::RoughnessB8)
+		{
+			return CodeIndex1 == INDEX_NONE ? INDEX_NONE : ComponentMask(CodeIndex1, 0, 0, 1, 0);
+		}
+		else if (UnpackType == EVirtualTextureUnpackType::HeightR16)
+		{
+			FString	SampleCode(TEXT("VirtualTextureUnpackHeightR16(%s)"));
+			return CodeIndex0 == INDEX_NONE ? INDEX_NONE : AddCodeChunk(MCT_Float, *SampleCode, *GetParameterCode(CodeIndex0));
 		}
 
-		return CodeIndex;
+		return CodeIndex0;
 	}
 
 	virtual int32 ExternalTexture(const FGuid& ExternalTextureGuid) override
diff --git a/Engine/Source/Runtime/Engine/Private/Materials/MaterialExpressions.cpp b/Engine/Source/Runtime/Engine/Private/Materials/MaterialExpressions.cpp
--- a/Engine/Source/Runtime/Engine/Private/Materials/MaterialExpressions.cpp
+++ b/Engine/Source/Runtime/Engine/Private/Materials/MaterialExpressions.cpp
@@ -2023,8 +2023,8 @@ void UMaterialExpressionRuntimeVirtualTextureSample::InitOutputs()
 	Outputs.Reset();
 	
 	Outputs.Add(FExpressionOutput(TEXT("BaseColor"), 1, 1, 1, 1, 0));
-	Outputs.Add(FExpressionOutput(TEXT("Specular"), 1, 0, 0, 1, 0));
-	Outputs.Add(FExpressionOutput(TEXT("Roughness"), 1, 1, 0, 0, 0));
+	Outputs.Add(FExpressionOutput(TEXT("Specular")));
+	Outputs.Add(FExpressionOutput(TEXT("Roughness")));
 	Outputs.Add(FExpressionOutput(TEXT("Normal")));
 	Outputs.Add(FExpressionOutput(TEXT("WorldHeight")));
 #endif // WITH_EDITORONLY_DATA
@@ -2082,22 +2082,28 @@ int32 UMaterialExpressionRuntimeVirtualTextureSample::Compile(class FMaterialCom
 
 	// Calculate the virtual texture layer and sampling/unpacking functions for this output
 	// Fallback to a sensible default value if the output isn't valid for the bound virtual texture
-	const bool bIsBaseColorValid = MaterialType == ERuntimeVirtualTextureMaterialType::BaseColor || MaterialType == ERuntimeVirtualTextureMaterialType::BaseColor_Normal || MaterialType == ERuntimeVirtualTextureMaterialType::BaseColor_Normal_Specular;
-	const bool bIsSpecularValid = MaterialType == ERuntimeVirtualTextureMaterialType::BaseColor_Normal_Specular;
-	const bool bIsNormalValid = MaterialType == ERuntimeVirtualTextureMaterialType::BaseColor_Normal || MaterialType == ERuntimeVirtualTextureMaterialType::BaseColor_Normal_Specular;
-	const bool bIsNormalBC3 = MaterialType == ERuntimeVirtualTextureMaterialType::BaseColor_Normal_Specular;
-	const bool bIsWorldHeightValid = MaterialType == ERuntimeVirtualTextureMaterialType::WorldHeight;
-
-	int32 LayerIndex = 0;
-	EMaterialSamplerType SamplerType = SAMPLERTYPE_VirtualMasks;
 	EVirtualTextureUnpackType UnpackType = EVirtualTextureUnpackType::None;
 
+	bool bIsBaseColorValid = false;
+	bool bIsSpecularValid = false;
+	bool bIsNormalValid = false;
+	bool bIsWorldHeightValid = false;
+
+	switch (MaterialType)
+	{
+	case ERuntimeVirtualTextureMaterialType::BaseColor: bIsBaseColorValid = true; break;
+	case ERuntimeVirtualTextureMaterialType::BaseColor_Normal: bIsBaseColorValid = bIsNormalValid = true; break;
+	case ERuntimeVirtualTextureMaterialType::BaseColor_Normal_Specular: bIsBaseColorValid = bIsNormalValid = bIsSpecularValid = true; break;
+	case ERuntimeVirtualTextureMaterialType::BaseColor_Normal_Specular_Deprecated: bIsBaseColorValid = bIsNormalValid = bIsSpecularValid = true; break;
+	case ERuntimeVirtualTextureMaterialType::WorldHeight: bIsWorldHeightValid = true; break;
+	}
+
 	switch (OutputIndex)
 	{
 	case 0: 
 		if (bIsVirtualTextureValid && bIsBaseColorValid)
 		{
-			SamplerType = SAMPLERTYPE_VirtualColor;
+			UnpackType = EVirtualTextureUnpackType::BaseColor;
 		}
 		else
 		{
@@ -2105,10 +2111,23 @@ int32 UMaterialExpressionRuntimeVirtualTextureSample::Compile(class FMaterialCom
 		}
 		break;
 	case 1:
+		if (bIsVirtualTextureValid && bIsSpecularValid)
+		{
+			UnpackType = EVirtualTextureUnpackType::SpecularR8;
+		}
+		else
+		{
+			return Compiler->Constant(0.5f);
+		}
+		break;
 	case 2:
 		if (bIsVirtualTextureValid && bIsSpecularValid)
 		{
-			LayerIndex = 1;
+			switch (MaterialType)
+			{
+			case ERuntimeVirtualTextureMaterialType::BaseColor_Normal_Specular: UnpackType = EVirtualTextureUnpackType::RoughnessG8; break;
+			case ERuntimeVirtualTextureMaterialType::BaseColor_Normal_Specular_Deprecated: UnpackType = EVirtualTextureUnpackType::RoughnessB8; break;
+			}
 		}
 		else
 		{
@@ -2118,8 +2137,12 @@ int32 UMaterialExpressionRuntimeVirtualTextureSample::Compile(class FMaterialCom
 	case 3:
 		if (bIsVirtualTextureValid && bIsNormalValid)
 		{
-			LayerIndex = 1;
-			UnpackType = bIsNormalBC3 ? EVirtualTextureUnpackType::NormalBC3 : EVirtualTextureUnpackType::NormalBC5;
+			switch (MaterialType)
+			{
+			case ERuntimeVirtualTextureMaterialType::BaseColor_Normal: UnpackType = EVirtualTextureUnpackType::NormalBC5; break;
+			case ERuntimeVirtualTextureMaterialType::BaseColor_Normal_Specular: UnpackType = EVirtualTextureUnpackType::NormalBC3BC3; break;
+			case ERuntimeVirtualTextureMaterialType::BaseColor_Normal_Specular_Deprecated: UnpackType = EVirtualTextureUnpackType::NormalBC3; break;
+			}
 		}
 		else
 		{
@@ -2129,7 +2152,6 @@ int32 UMaterialExpressionRuntimeVirtualTextureSample::Compile(class FMaterialCom
 	case 4:
 		if (bIsVirtualTextureValid && bIsWorldHeightValid)
 		{
-			LayerIndex = 0;
 			UnpackType = EVirtualTextureUnpackType::HeightR16;
 		}
 		else
@@ -2141,9 +2163,18 @@ int32 UMaterialExpressionRuntimeVirtualTextureSample::Compile(class FMaterialCom
 		return INDEX_NONE;
 	}
 	
-	// Compile the texture object reference
-	int32 TextureReferenceIndex = INDEX_NONE;
-	const int32 TextureCodeIndex = Compiler->VirtualTexture(VirtualTexture, LayerIndex, TextureReferenceIndex, SamplerType);
+	// Compile the texture object references
+	enum { MAX_RVT_LAYERS = 2 };
+	const int32 LayerCount = VirtualTexture->GetLayerCount();
+	check(LayerCount <= MAX_RVT_LAYERS);
+
+	int32 TextureCodeIndex[MAX_RVT_LAYERS] = { INDEX_NONE };
+	int32 TextureReferenceIndex[MAX_RVT_LAYERS] = { INDEX_NONE };
+	for (int32 Layer = 0; Layer < LayerCount; Layer++)
+	{
+		TextureCodeIndex[Layer] = Compiler->VirtualTexture(VirtualTexture, Layer, TextureReferenceIndex[Layer], SAMPLERTYPE_VirtualMasks);
+		check(TextureReferenceIndex[Layer] == TextureReferenceIndex[0]);
+	}
 
 	// Compile the coordinates
 	// We use the virtual texture world space transform by default
@@ -2151,9 +2182,9 @@ int32 UMaterialExpressionRuntimeVirtualTextureSample::Compile(class FMaterialCom
 	if (Coordinates.GetTracedInput().Expression == nullptr)
 	{
 		int32 WorldPositionIndex = Compiler->WorldPosition(WPT_Default);
-		int32 P0 = Compiler->VirtualTextureParam(TextureReferenceIndex, 0);
-		int32 P1 = Compiler->VirtualTextureParam(TextureReferenceIndex, 1);
-		int32 P2 = Compiler->VirtualTextureParam(TextureReferenceIndex, 2);
+		int32 P0 = Compiler->VirtualTextureParam(TextureReferenceIndex[0], 0);
+		int32 P1 = Compiler->VirtualTextureParam(TextureReferenceIndex[0], 1);
+		int32 P2 = Compiler->VirtualTextureParam(TextureReferenceIndex[0], 2);
 		CoordinateIndex = Compiler->VirtualTextureWorldToUV(WorldPositionIndex, P0, P1, P2);
 	}
 	else
@@ -2162,12 +2193,21 @@ int32 UMaterialExpressionRuntimeVirtualTextureSample::Compile(class FMaterialCom
 	}
 	
 	// Compile the texture sample code
-	//todo[vt]: Expose support for mip sample settings through the URuntimeVirtualTexture object
-	const int32 SampleCodeIndex = Compiler->TextureSample(TextureCodeIndex, CoordinateIndex, SamplerType, INDEX_NONE, INDEX_NONE, TMVM_None, SSM_Wrap_WorldGroupSettings, TextureReferenceIndex, false);
+	//todo[vt]: Expose support for mip settings in this expression
+	int32 SampleCodeIndex[MAX_RVT_LAYERS] = { INDEX_NONE };
+	for (int32 Layer = 0; Layer < LayerCount; Layer++)
+	{
+		SampleCodeIndex[Layer] = Compiler->TextureSample(
+			TextureCodeIndex[Layer],
+			CoordinateIndex, 
+			SAMPLERTYPE_VirtualMasks,
+			INDEX_NONE, INDEX_NONE, TMVM_None, SSM_Wrap_WorldGroupSettings,
+			TextureReferenceIndex[Layer], 
+			false);
+	}
 
 	// Compile any unpacking code
-	const int32 UnpackCodeIndex = Compiler->VirtualTextureUnpack(SampleCodeIndex, UnpackType);
-	
+	const int32 UnpackCodeIndex = Compiler->VirtualTextureUnpack(SampleCodeIndex[0], SampleCodeIndex[1], UnpackType);
 	return UnpackCodeIndex;
 }
 
diff --git a/Engine/Source/Runtime/Engine/Private/VT/RuntimeVirtualTexture.cpp b/Engine/Source/Runtime/Engine/Private/VT/RuntimeVirtualTexture.cpp
--- a/Engine/Source/Runtime/Engine/Private/VT/RuntimeVirtualTexture.cpp
+++ b/Engine/Source/Runtime/Engine/Private/VT/RuntimeVirtualTexture.cpp
@@ -235,6 +235,7 @@ int32 URuntimeVirtualTexture::GetLayerCount() const
 		return 1;
 	case ERuntimeVirtualTextureMaterialType::BaseColor_Normal:
 	case ERuntimeVirtualTextureMaterialType::BaseColor_Normal_Specular:
+	case ERuntimeVirtualTextureMaterialType::BaseColor_Normal_Specular_Deprecated:
 		return 2;
 	default:
 		break;
@@ -253,8 +254,10 @@ EPixelFormat URuntimeVirtualTexture::GetLayerFormat(int32 LayerIndex) const
 		{
 		case ERuntimeVirtualTextureMaterialType::BaseColor:
 		case ERuntimeVirtualTextureMaterialType::BaseColor_Normal:
-		case ERuntimeVirtualTextureMaterialType::BaseColor_Normal_Specular:
+		case ERuntimeVirtualTextureMaterialType::BaseColor_Normal_Specular_Deprecated:
 			return bCompressTextures ? PF_DXT1 : PF_B8G8R8A8;
+		case ERuntimeVirtualTextureMaterialType::BaseColor_Normal_Specular:
+			return bCompressTextures ? PF_DXT5 : PF_B8G8R8A8;
 		case ERuntimeVirtualTextureMaterialType::WorldHeight:
 			return PF_G16;
 		default:
@@ -268,6 +271,7 @@ EPixelFormat URuntimeVirtualTexture::GetLayerFormat(int32 LayerIndex) const
 		case ERuntimeVirtualTextureMaterialType::BaseColor_Normal:
 			return bCompressTextures ? PF_BC5 : PF_B8G8R8A8;
 		case ERuntimeVirtualTextureMaterialType::BaseColor_Normal_Specular:
+		case ERuntimeVirtualTextureMaterialType::BaseColor_Normal_Specular_Deprecated:
 			return bCompressTextures ? PF_DXT5 : PF_B8G8R8A8;
 		case ERuntimeVirtualTextureMaterialType::BaseColor:
 		case ERuntimeVirtualTextureMaterialType::WorldHeight:
@@ -288,6 +292,7 @@ bool URuntimeVirtualTexture::IsLayerSRGB(int32 LayerIndex) const
 	case ERuntimeVirtualTextureMaterialType::BaseColor:
 	case ERuntimeVirtualTextureMaterialType::BaseColor_Normal:
 	case ERuntimeVirtualTextureMaterialType::BaseColor_Normal_Specular:
+	case ERuntimeVirtualTextureMaterialType::BaseColor_Normal_Specular_Deprecated:
 		// Only BaseColor layer is sRGB
 		return LayerIndex == 0;
 	case ERuntimeVirtualTextureMaterialType::WorldHeight:
diff --git a/Engine/Source/Runtime/Engine/Public/MaterialCompiler.h b/Engine/Source/Runtime/Engine/Public/MaterialCompiler.h
--- a/Engine/Source/Runtime/Engine/Public/MaterialCompiler.h
+++ b/Engine/Source/Runtime/Engine/Public/MaterialCompiler.h
@@ -31,8 +31,13 @@ enum EMaterialForceCastFlags
 enum class EVirtualTextureUnpackType
 {
 	None,
+	BaseColor,
 	NormalBC3,
 	NormalBC5,
+	NormalBC3BC3,
+	SpecularR8,
+	RoughnessG8,
+	RoughnessB8,
 	HeightR16,
 };
 
@@ -185,7 +190,7 @@ public:
 	virtual int32 VirtualTexture(URuntimeVirtualTexture* InTexture, int32 LayerIndex, int32& TextureReferenceIndex, EMaterialSamplerType SamplerType) = 0;
 	virtual int32 VirtualTextureParam(int32 TextureIndex, int32 ParamIndex) = 0;
 	virtual int32 VirtualTextureWorldToUV(int32 WorldPositionIndex, int32 P0, int32 P1, int32 P2) = 0;
-	virtual int32 VirtualTextureUnpack(int32 CodeIndex, EVirtualTextureUnpackType UnpackType) = 0;
+	virtual int32 VirtualTextureUnpack(int32 CodeIndex0, int32 CodeIndex1, EVirtualTextureUnpackType UnpackType) = 0;
 
 	virtual int32 ExternalTexture(const FGuid& ExternalTextureGuid) = 0;
 	virtual int32 ExternalTexture(UTexture* InTexture, int32& TextureReferenceIndex) = 0;
@@ -459,7 +464,7 @@ public:
 	virtual int32 VirtualTexture(URuntimeVirtualTexture* InTexture, int32 LayerIndex, int32& TextureReferenceIndex, EMaterialSamplerType SamplerType) override { return Compiler->VirtualTexture(InTexture, LayerIndex, TextureReferenceIndex, SamplerType); }
 	virtual int32 VirtualTextureParam(int32 TextureIndex, int32 ParamIndex) override { return Compiler->VirtualTextureParam(TextureIndex, ParamIndex); }
 	virtual int32 VirtualTextureWorldToUV(int32 WorldPositionIndex, int32 P0, int32 P1, int32 P2) override { return Compiler->VirtualTextureWorldToUV(WorldPositionIndex, P0, P1, P2); }
-	virtual int32 VirtualTextureUnpack(int32 CodeIndex, EVirtualTextureUnpackType UnpackType) override { return Compiler->VirtualTextureUnpack(CodeIndex, UnpackType); }
+	virtual int32 VirtualTextureUnpack(int32 CodeIndex0, int32 CodeIndex1, EVirtualTextureUnpackType UnpackType) override { return Compiler->VirtualTextureUnpack(CodeIndex0, CodeIndex1, UnpackType); }
 
 	virtual int32 ExternalTexture(const FGuid& ExternalTextureGuid) override { return Compiler->ExternalTexture(ExternalTextureGuid); }
 	virtual int32 ExternalTexture(UTexture* InTexture, int32& TextureReferenceIndex) override { return Compiler->ExternalTexture(InTexture, TextureReferenceIndex); }
diff --git a/Engine/Source/Runtime/Engine/Public/VT/RuntimeVirtualTextureEnum.h b/Engine/Source/Runtime/Engine/Public/VT/RuntimeVirtualTextureEnum.h
--- a/Engine/Source/Runtime/Engine/Public/VT/RuntimeVirtualTextureEnum.h
+++ b/Engine/Source/Runtime/Engine/Public/VT/RuntimeVirtualTextureEnum.h
@@ -17,11 +17,12 @@ enum class ERuntimeVirtualTextureMaterialType : uint8
 	BaseColor_Normal UMETA(DisplayName = "Base Color, Normal"),
 	BaseColor_Normal_Specular UMETA(DisplayName = "Base Color, Normal, Roughness, Specular"),
 	WorldHeight UMETA(DisplayName = "World Height"),
-	
+	BaseColor_Normal_Specular_Deprecated UMETA(DisplayName = "[OLD] Base Color, Normal, Roughness, Specular"),
+
 	Count UMETA(Hidden),
 };
 
-enum { ERuntimeVirtualTextureMaterialType_NumBits = 2 };
+enum { ERuntimeVirtualTextureMaterialType_NumBits = 3 };
 static_assert((uint32)ERuntimeVirtualTextureMaterialType::Count <= (1 << (uint32)ERuntimeVirtualTextureMaterialType_NumBits), "NumBits is too small");
 
 /** Enumeration of main pass behaviors when rendering to a runtime virtual texture. */
diff --git a/Engine/Source/Runtime/Renderer/Private/VT/RuntimeVirtualTextureRender.cpp b/Engine/Source/Runtime/Renderer/Private/VT/RuntimeVirtualTextureRender.cpp
--- a/Engine/Source/Runtime/Renderer/Private/VT/RuntimeVirtualTextureRender.cpp
+++ b/Engine/Source/Runtime/Renderer/Private/VT/RuntimeVirtualTextureRender.cpp
@@ -263,6 +263,7 @@ namespace RuntimeVirtualTexture
 					Process<FMaterialPolicy_BaseColorNormal>(MeshBatch, BatchElementMask, StaticMeshId, PrimitiveSceneProxy, MaterialRenderProxy, Material);
 					break;
 				case ERuntimeVirtualTextureMaterialType::BaseColor_Normal_Specular:
+				case ERuntimeVirtualTextureMaterialType::BaseColor_Normal_Specular_Deprecated:
 					Process<FMaterialPolicy_BaseColorNormalSpecular>(MeshBatch, BatchElementMask, StaticMeshId, PrimitiveSceneProxy, MaterialRenderProxy, Material);
 					break;
 				case ERuntimeVirtualTextureMaterialType::WorldHeight:
@@ -429,9 +430,11 @@ namespace RuntimeVirtualTexture
 			SHADER_PARAMETER_SAMPLER(SamplerState, TextureSampler1)
 			SHADER_PARAMETER_RDG_TEXTURE(Texture2D<float4>, RenderTexture2)
 			SHADER_PARAMETER_SAMPLER(SamplerState, TextureSampler2)
-			SHADER_PARAMETER_RDG_TEXTURE_UAV(RWTexture<uint2>, OutCompressTexture0)
+			SHADER_PARAMETER_RDG_TEXTURE_UAV(RWTexture<uint2>, OutCompressTexture0_u2)
+			SHADER_PARAMETER_RDG_TEXTURE_UAV(RWTexture<uint4>, OutCompressTexture0_u4)
 			SHADER_PARAMETER_RDG_TEXTURE_UAV(RWTexture<uint4>, OutCompressTexture1)
 			SHADER_PARAMETER_RDG_TEXTURE_UAV(RWTexture<float4>, OutCopyTexture0)
+			SHADER_PARAMETER_RDG_TEXTURE_UAV(RWTexture<float4>, OutCopyTexture1)
 		END_SHADER_PARAMETER_STRUCT()
 
 		static bool ShouldCompilePermutation(FGlobalShaderPermutationParameters const& Parameters)
@@ -472,10 +475,16 @@ namespace RuntimeVirtualTexture
 	
 	typedef FShader_VirtualTextureCompress_CS< ERuntimeVirtualTextureMaterialType::BaseColor_Normal_Specular, false > FShader_VirtualTextureCompress_BaseColorNormalSpecular_CS;
 	IMPLEMENT_SHADER_TYPE(template<>, FShader_VirtualTextureCompress_BaseColorNormalSpecular_CS, TEXT("/Engine/Private/VirtualTextureCompress.usf"), TEXT("CompressBaseColorNormalSpecularCS"), SF_Compute);
-	
+
+	typedef FShader_VirtualTextureCompress_CS< ERuntimeVirtualTextureMaterialType::BaseColor_Normal_Specular_Deprecated, false > FShader_VirtualTextureCompress_BaseColorNormalSpecularDeprecated_CS;
+	IMPLEMENT_SHADER_TYPE(template<>, FShader_VirtualTextureCompress_BaseColorNormalSpecularDeprecated_CS, TEXT("/Engine/Private/VirtualTextureCompress.usf"), TEXT("CompressBaseColorNormalSpecularDeprecatedCS"), SF_Compute);
+
 	typedef FShader_VirtualTextureCompress_CS< ERuntimeVirtualTextureMaterialType::BaseColor_Normal_Specular, true > FShader_VirtualTextureCompress_BaseColorNormalSpecular_CopyOnly_CS;
 	IMPLEMENT_SHADER_TYPE(template<>, FShader_VirtualTextureCompress_BaseColorNormalSpecular_CopyOnly_CS, TEXT("/Engine/Private/VirtualTextureCompress.usf"), TEXT("CopyNormalSpecularCS"), SF_Compute);
 
+	typedef FShader_VirtualTextureCompress_CS< ERuntimeVirtualTextureMaterialType::BaseColor_Normal_Specular_Deprecated, true > FShader_VirtualTextureCompress_BaseColorNormalSpecularDeprecated_CopyOnly_CS;
+	IMPLEMENT_SHADER_TYPE(template<>, FShader_VirtualTextureCompress_BaseColorNormalSpecularDeprecated_CopyOnly_CS, TEXT("/Engine/Private/VirtualTextureCompress.usf"), TEXT("CopyNormalSpecularDeprecatedCS"), SF_Compute);
+
 
 	/** Set up the BC compression pass for the specific MaterialType. */
 	template< ERuntimeVirtualTextureMaterialType MaterialType, bool CopyOnly = false >
@@ -507,17 +516,26 @@ namespace RuntimeVirtualTexture
 		case ERuntimeVirtualTextureMaterialType::BaseColor_Normal_Specular:
 			AddCompressOrCopyPass<ERuntimeVirtualTextureMaterialType::BaseColor_Normal_Specular>(GraphBuilder, FeatureLevel, Parameters, GroupCount);
 			break;
+		case ERuntimeVirtualTextureMaterialType::BaseColor_Normal_Specular_Deprecated:
+			AddCompressOrCopyPass<ERuntimeVirtualTextureMaterialType::BaseColor_Normal_Specular_Deprecated>(GraphBuilder, FeatureLevel, Parameters, GroupCount);
+			break;
 		}
 	}
 
 	/** Set up the copy pass used when BC compression is disabled */
 	void AddCopyPass(FRDGBuilder& GraphBuilder, ERHIFeatureLevel::Type FeatureLevel, FShader_VirtualTextureCompress::FParameters* Parameters, FIntPoint TextureSize, ERuntimeVirtualTextureMaterialType MaterialType)
 	{
+		const FIntVector GroupCount((TextureSize.X + 7) / 8, (TextureSize.Y + 7) / 8, 1);
+
 		// Only needed for BaseColor_Normal_Specular where we need to pack normal and specular into one VT layer
-		if (MaterialType == ERuntimeVirtualTextureMaterialType::BaseColor_Normal_Specular)
+		switch (MaterialType)
 		{
-			const FIntVector GroupCount((TextureSize.X + 7) / 8, (TextureSize.Y + 7) / 8, 1);
+		case ERuntimeVirtualTextureMaterialType::BaseColor_Normal_Specular:
 			AddCompressOrCopyPass<ERuntimeVirtualTextureMaterialType::BaseColor_Normal_Specular, true>(GraphBuilder, FeatureLevel, Parameters, GroupCount);
+			break;
+		case ERuntimeVirtualTextureMaterialType::BaseColor_Normal_Specular_Deprecated:
+			AddCompressOrCopyPass<ERuntimeVirtualTextureMaterialType::BaseColor_Normal_Specular_Deprecated, true>(GraphBuilder, FeatureLevel, Parameters, GroupCount);
+			break;
 		}
 	}
 
@@ -529,8 +547,8 @@ namespace RuntimeVirtualTexture
 		FRenderGraphSetup(FRDGBuilder& GraphBuilder, ERuntimeVirtualTextureMaterialType MaterialType, FRHITexture2D* OutputTexture0, FRHITexture2D* OutputTexture1, FIntPoint TextureSize)
 		{
 			bRenderPass = OutputTexture0 != nullptr;
-			bCompressPass = bRenderPass && (OutputTexture0->GetFormat() == PF_DXT1 || OutputTexture0->GetFormat() == PF_DXT3 || OutputTexture0->GetFormat() == PF_BC5);
-			bCopyPass = bRenderPass && !bCompressPass && MaterialType == ERuntimeVirtualTextureMaterialType::BaseColor_Normal_Specular;
+			bCompressPass = bRenderPass && (OutputTexture0->GetFormat() == PF_DXT1 || OutputTexture0->GetFormat() == PF_DXT5 || OutputTexture0->GetFormat() == PF_BC5);
+			bCopyPass = bRenderPass && !bCompressPass && (MaterialType == ERuntimeVirtualTextureMaterialType::BaseColor_Normal_Specular || MaterialType == ERuntimeVirtualTextureMaterialType::BaseColor_Normal_Specular_Deprecated);
 		
 			switch (MaterialType)
 			{
@@ -541,7 +559,7 @@ namespace RuntimeVirtualTexture
 				}
 				if (bCompressPass)
 				{
-					OutputAlias0 = CompressTexture0 = GraphBuilder.CreateTexture(FPooledRenderTargetDesc::Create2DDesc(TextureSize / 4, PF_R32G32_UINT, FClearValueBinding::None, TexCreate_None, TexCreate_UAV, false), TEXT("CompressTexture0"));
+					OutputAlias0 = CompressTexture0_u2 = GraphBuilder.CreateTexture(FPooledRenderTargetDesc::Create2DDesc(TextureSize / 4, PF_R32G32_UINT, FClearValueBinding::None, TexCreate_None, TexCreate_UAV, false), TEXT("CompressTexture0"));
 				}
 				break;
 			case ERuntimeVirtualTextureMaterialType::BaseColor_Normal:
@@ -552,7 +570,7 @@ namespace RuntimeVirtualTexture
 				}
 				if (bCompressPass)
 				{
-					OutputAlias0 = CompressTexture0 = GraphBuilder.CreateTexture(FPooledRenderTargetDesc::Create2DDesc(TextureSize / 4, PF_R32G32_UINT, FClearValueBinding::None, TexCreate_None, TexCreate_UAV, false), TEXT("CompressTexture0"));
+					OutputAlias0 = CompressTexture0_u2 = GraphBuilder.CreateTexture(FPooledRenderTargetDesc::Create2DDesc(TextureSize / 4, PF_R32G32_UINT, FClearValueBinding::None, TexCreate_None, TexCreate_UAV, false), TEXT("CompressTexture0"));
 					OutputAlias1 = CompressTexture1 = GraphBuilder.CreateTexture(FPooledRenderTargetDesc::Create2DDesc(TextureSize / 4, PF_R32G32B32A32_UINT, FClearValueBinding::None, TexCreate_None, TexCreate_UAV, false), TEXT("CompressTexture1"));
 				}
 				break;
@@ -565,7 +583,25 @@ namespace RuntimeVirtualTexture
 				}
 				if (bCompressPass)
 				{
-					OutputAlias0 = CompressTexture0 = GraphBuilder.CreateTexture(FPooledRenderTargetDesc::Create2DDesc(TextureSize / 4, PF_R32G32_UINT, FClearValueBinding::None, TexCreate_None, TexCreate_UAV, false), TEXT("CompressTexture0"));
+					OutputAlias0 = CompressTexture0_u4 = GraphBuilder.CreateTexture(FPooledRenderTargetDesc::Create2DDesc(TextureSize / 4, PF_R32G32B32A32_UINT, FClearValueBinding::None, TexCreate_None, TexCreate_UAV, false), TEXT("CompressTexture0"));
+					OutputAlias1 = CompressTexture1 = GraphBuilder.CreateTexture(FPooledRenderTargetDesc::Create2DDesc(TextureSize / 4, PF_R32G32B32A32_UINT, FClearValueBinding::None, TexCreate_None, TexCreate_UAV, false), TEXT("CompressTexture1"));
+				}
+				if (bCopyPass)
+				{
+					OutputAlias0 = CopyTexture0 = GraphBuilder.CreateTexture(FPooledRenderTargetDesc::Create2DDesc(TextureSize, PF_B8G8R8A8, FClearValueBinding::None, TexCreate_None, TexCreate_UAV, false), TEXT("CopyTexture0"));
+					OutputAlias1 = CopyTexture1 = GraphBuilder.CreateTexture(FPooledRenderTargetDesc::Create2DDesc(TextureSize, PF_B8G8R8A8, FClearValueBinding::None, TexCreate_None, TexCreate_UAV, false), TEXT("CopyTexture1"));
+				}
+				break;
+			case ERuntimeVirtualTextureMaterialType::BaseColor_Normal_Specular_Deprecated:
+				if (bRenderPass)
+				{
+					OutputAlias0 = RenderTexture0 = GraphBuilder.CreateTexture(FPooledRenderTargetDesc::Create2DDesc(TextureSize, PF_B8G8R8A8, FClearValueBinding::Black, TexCreate_SRGB, TexCreate_RenderTargetable, false), TEXT("RenderTexture0"));
+					RenderTexture1 = GraphBuilder.CreateTexture(FPooledRenderTargetDesc::Create2DDesc(TextureSize, PF_B8G8R8A8, FClearValueBinding::Black, TexCreate_None, TexCreate_RenderTargetable, false), TEXT("RenderTexture1"));
+					RenderTexture2 = GraphBuilder.CreateTexture(FPooledRenderTargetDesc::Create2DDesc(TextureSize, PF_B8G8R8A8, FClearValueBinding::Black, TexCreate_None, TexCreate_RenderTargetable, false), TEXT("RenderTexture2"));
+				}
+				if (bCompressPass)
+				{
+					OutputAlias0 = CompressTexture0_u2 = GraphBuilder.CreateTexture(FPooledRenderTargetDesc::Create2DDesc(TextureSize / 4, PF_R32G32_UINT, FClearValueBinding::None, TexCreate_None, TexCreate_UAV, false), TEXT("CompressTexture0"));
 					OutputAlias1 = CompressTexture1 = GraphBuilder.CreateTexture(FPooledRenderTargetDesc::Create2DDesc(TextureSize / 4, PF_R32G32B32A32_UINT, FClearValueBinding::None, TexCreate_None, TexCreate_UAV, false), TEXT("CompressTexture1"));
 				}
 				if (bCopyPass)
@@ -576,7 +612,7 @@ namespace RuntimeVirtualTexture
 			case ERuntimeVirtualTextureMaterialType::WorldHeight:
 				if (bRenderPass)
 				{
-					OutputAlias0 = RenderTexture0 = GraphBuilder.CreateTexture(FPooledRenderTargetDesc::Create2DDesc(TextureSize, PF_G16, FClearValueBinding::Black, TexCreate_SRGB, TexCreate_RenderTargetable, false), TEXT("RenderTexture0"));
+					OutputAlias0 = RenderTexture0 = GraphBuilder.CreateTexture(FPooledRenderTargetDesc::Create2DDesc(TextureSize, PF_G16, FClearValueBinding::Black, TexCreate_None, TexCreate_RenderTargetable, false), TEXT("RenderTexture0"));
 				}
 				break;
 			}
@@ -591,9 +627,11 @@ namespace RuntimeVirtualTexture
 		FRDGTextureRef RenderTexture0 = nullptr;
 		FRDGTextureRef RenderTexture1 = nullptr;
 		FRDGTextureRef RenderTexture2 = nullptr;
-		FRDGTextureRef CompressTexture0 = nullptr;
+		FRDGTextureRef CompressTexture0_u2 = nullptr;
+		FRDGTextureRef CompressTexture0_u4 = nullptr;
 		FRDGTextureRef CompressTexture1 = nullptr;
 		FRDGTextureRef CopyTexture0 = nullptr;
+		FRDGTextureRef CopyTexture1 = nullptr;
 
 		/** Aliases to one of the render/compress/copy textures. This is what we will Copy into the final physical texture. */
 		//todo[vt]: On platforms that support direct aliasing we can not set these and compress direct to the final destination
@@ -703,7 +741,8 @@ namespace RuntimeVirtualTexture
 			PassParameters->TextureSampler1 = TStaticSamplerState<SF_Point, AM_Clamp, AM_Clamp, AM_Clamp>::GetRHI();
 			PassParameters->RenderTexture2 = GraphSetup.RenderTexture2;
 			PassParameters->TextureSampler2 = TStaticSamplerState<SF_Point, AM_Clamp, AM_Clamp, AM_Clamp>::GetRHI();
-			PassParameters->OutCompressTexture0 = GraphSetup.CompressTexture0 ? GraphBuilder.CreateUAV(FRDGTextureUAVDesc(GraphSetup.CompressTexture0)) : nullptr;
+			PassParameters->OutCompressTexture0_u2 = GraphSetup.CompressTexture0_u2 ? GraphBuilder.CreateUAV(FRDGTextureUAVDesc(GraphSetup.CompressTexture0_u2)) : nullptr;
+			PassParameters->OutCompressTexture0_u4 = GraphSetup.CompressTexture0_u4 ? GraphBuilder.CreateUAV(FRDGTextureUAVDesc(GraphSetup.CompressTexture0_u4)) : nullptr;
 			PassParameters->OutCompressTexture1 = GraphSetup.CompressTexture1 ? GraphBuilder.CreateUAV(FRDGTextureUAVDesc(GraphSetup.CompressTexture1)) : nullptr;
 
 			AddCompressPass(GraphBuilder, View->GetFeatureLevel(), PassParameters, TextureSize, MaterialType);
@@ -721,6 +760,7 @@ namespace RuntimeVirtualTexture
 			PassParameters->RenderTexture2 = GraphSetup.RenderTexture2;
 			PassParameters->TextureSampler2 = TStaticSamplerState<SF_Point, AM_Clamp, AM_Clamp, AM_Clamp>::GetRHI();
 			PassParameters->OutCopyTexture0 = GraphSetup.CopyTexture0 ? GraphBuilder.CreateUAV(FRDGTextureUAVDesc(GraphSetup.CopyTexture0)) : nullptr;
+			PassParameters->OutCopyTexture1 = GraphSetup.CopyTexture1 ? GraphBuilder.CreateUAV(FRDGTextureUAVDesc(GraphSetup.CopyTexture1)) : nullptr;
 
 			AddCopyPass(GraphBuilder, View->GetFeatureLevel(), PassParameters, TextureSize, MaterialType);
 		}
