Changed ERuntimeVirtualTextureMaterialType enum so that we don't break backwards compatibility with the new RVT format

diff --git a/Engine/Shaders/Private/VirtualTextureCompress.usf b/Engine/Shaders/Private/VirtualTextureCompress.usf
--- a/Engine/Shaders/Private/VirtualTextureCompress.usf
+++ b/Engine/Shaders/Private/VirtualTextureCompress.usf
@@ -149,7 +149,7 @@ void CompressBaseColorNormalCS(
 	OutCompressTexture1[ThreadId.xy + DestRect.xy] = CompressBC5Block(BlockNormalX, BlockNormalY);
 }
 
-/** Compress base color, normal, roughness and specular to BC3 target pair. */
+/** Compress base color, normal, roughness and specular to BC1/BC3 target pair. */
 [numthreads(8,8,1)]
 void CompressBaseColorNormalSpecularCS(
 	uint3 ThreadId : SV_DispatchThreadID 
@@ -172,11 +172,16 @@ void CompressBaseColorNormalSpecularCS(
 	float3 BlockSpecular[16];
 	ReadBlockRGB(RenderTexture2, TextureSampler2, SampleUV, TexelUVSize, BlockSpecular);
 
-	OutCompressTexture0_u4[ThreadId.xy + DestRect.xy] = CompressBC3Block_SRGB(BlockBaseColor, BlockNormalX);
-	OutCompressTexture1[ThreadId.xy + DestRect.xy] = CompressBC3Block(BlockSpecular, BlockNormalY);
+	for (int i=0; i<16; i++)
+	{
+		BlockSpecular[i] = float3(BlockSpecular[i].x, BlockNormalY[i], BlockSpecular[i].y);
+	}
+
+	OutCompressTexture0_u2[ThreadId.xy + DestRect.xy] = CompressBC1Block_SRGB(BlockBaseColor);
+	OutCompressTexture1[ThreadId.xy + DestRect.xy] = CompressBC3Block(BlockSpecular, BlockNormalX);
 }
 
-/** Copy normal X and Y to the alpha channels of the final target textures (without compression)
+/** Combine and copy normal, roughness and specular to a single packed target (without compression).
   * This path is used when we disable texture compression, because we need to keep the same final channel layout. */
 [numthreads(8,8,1)]
 void CopyNormalSpecularCS(
@@ -187,17 +192,15 @@ void CopyNormalSpecularCS(
 	if (any(ThreadId.xy >= DestRect.zw))
 		return;
 
-	float3 BaseColor = LinearToSrgb(RenderTexture0[SamplePos].xyz);
 	float2 NormalXY = RenderTexture1[SamplePos].xy;
 	float2 RoughnessSpecular = RenderTexture2[SamplePos].xy;
 	
-	OutCopyTexture0[ThreadId.xy + DestRect.xy] = float4(BaseColor, NormalXY.x);
-	OutCopyTexture1[ThreadId.xy + DestRect.xy] = float4(RoughnessSpecular.x, RoughnessSpecular.y, 0.0f, NormalXY.y);
+	OutCopyTexture0[ThreadId.xy + DestRect.xy] = float4(RoughnessSpecular.x, NormalXY.y, RoughnessSpecular.y, NormalXY.x);
 }
 
-/** Compress base color, normal, roughness and specular to BC1/BC3 target pair. */
+/** Compress base color, normal, roughness and specular to BC3 target pair. */
 [numthreads(8,8,1)]
-void CompressBaseColorNormalSpecularDeprecatedCS(
+void CompressBaseColorNormalSpecularExCS(
 	uint3 ThreadId : SV_DispatchThreadID 
 	)
 {
@@ -218,19 +221,14 @@ void CompressBaseColorNormalSpecularDeprecatedCS(
 	float3 BlockSpecular[16];
 	ReadBlockRGB(RenderTexture2, TextureSampler2, SampleUV, TexelUVSize, BlockSpecular);
 
-	for (int i=0; i<16; i++)
-	{
-		BlockSpecular[i] = float3(BlockSpecular[i].x, BlockNormalY[i], BlockSpecular[i].y);
-	}
-
-	OutCompressTexture0_u2[ThreadId.xy + DestRect.xy] = CompressBC1Block_SRGB(BlockBaseColor);
-	OutCompressTexture1[ThreadId.xy + DestRect.xy] = CompressBC3Block(BlockSpecular, BlockNormalX);
+	OutCompressTexture0_u4[ThreadId.xy + DestRect.xy] = CompressBC3Block_SRGB(BlockBaseColor, BlockNormalX);
+	OutCompressTexture1[ThreadId.xy + DestRect.xy] = CompressBC3Block(BlockSpecular, BlockNormalY);
 }
 
-/** Combine and copy normal, roughness and specular to a single packed target (without compression).
+/** Copy normal X and Y to the alpha channels of the final target textures (without compression)
   * This path is used when we disable texture compression, because we need to keep the same final channel layout. */
 [numthreads(8,8,1)]
-void CopyNormalSpecularDeprecatedCS(
+void CopyNormalSpecularExCS(
 	uint3 ThreadId : SV_DispatchThreadID 
 	)
 {
@@ -238,8 +236,10 @@ void CopyNormalSpecularDeprecatedCS(
 	if (any(ThreadId.xy >= DestRect.zw))
 		return;
 
+	float3 BaseColor = LinearToSrgb(RenderTexture0[SamplePos].xyz);
 	float2 NormalXY = RenderTexture1[SamplePos].xy;
 	float2 RoughnessSpecular = RenderTexture2[SamplePos].xy;
 	
-	OutCopyTexture0[ThreadId.xy + DestRect.xy] = float4(RoughnessSpecular.x, NormalXY.y, RoughnessSpecular.y, NormalXY.x);
+	OutCopyTexture0[ThreadId.xy + DestRect.xy] = float4(BaseColor, NormalXY.x);
+	OutCopyTexture1[ThreadId.xy + DestRect.xy] = float4(RoughnessSpecular.x, RoughnessSpecular.y, 0.0f, NormalXY.y);
 }
diff --git a/Engine/Source/Runtime/Engine/Private/Materials/MaterialExpressions.cpp b/Engine/Source/Runtime/Engine/Private/Materials/MaterialExpressions.cpp
--- a/Engine/Source/Runtime/Engine/Private/Materials/MaterialExpressions.cpp
+++ b/Engine/Source/Runtime/Engine/Private/Materials/MaterialExpressions.cpp
@@ -2093,8 +2093,8 @@ int32 UMaterialExpressionRuntimeVirtualTextureSample::Compile(class FMaterialCom
 	{
 	case ERuntimeVirtualTextureMaterialType::BaseColor: bIsBaseColorValid = true; break;
 	case ERuntimeVirtualTextureMaterialType::BaseColor_Normal: bIsBaseColorValid = bIsNormalValid = true; break;
-	case ERuntimeVirtualTextureMaterialType::BaseColor_Normal_Specular: bIsBaseColorValid = bIsNormalValid = bIsSpecularValid = true; break;
-	case ERuntimeVirtualTextureMaterialType::BaseColor_Normal_Specular_Deprecated: bIsBaseColorValid = bIsNormalValid = bIsSpecularValid = true; break;
+	case ERuntimeVirtualTextureMaterialType::BaseColor_Normal_Specular:
+	case ERuntimeVirtualTextureMaterialType::BaseColor_Normal_Specular_Ex: bIsBaseColorValid = bIsNormalValid = bIsSpecularValid = true; break;
 	case ERuntimeVirtualTextureMaterialType::WorldHeight: bIsWorldHeightValid = true; break;
 	}
 
@@ -2125,8 +2125,8 @@ int32 UMaterialExpressionRuntimeVirtualTextureSample::Compile(class FMaterialCom
 		{
 			switch (MaterialType)
 			{
-			case ERuntimeVirtualTextureMaterialType::BaseColor_Normal_Specular: UnpackType = EVirtualTextureUnpackType::RoughnessG8; break;
-			case ERuntimeVirtualTextureMaterialType::BaseColor_Normal_Specular_Deprecated: UnpackType = EVirtualTextureUnpackType::RoughnessB8; break;
+			case ERuntimeVirtualTextureMaterialType::BaseColor_Normal_Specular: UnpackType = EVirtualTextureUnpackType::RoughnessB8; break;
+			case ERuntimeVirtualTextureMaterialType::BaseColor_Normal_Specular_Ex: UnpackType = EVirtualTextureUnpackType::RoughnessG8; break;
 			}
 		}
 		else
@@ -2140,8 +2140,8 @@ int32 UMaterialExpressionRuntimeVirtualTextureSample::Compile(class FMaterialCom
 			switch (MaterialType)
 			{
 			case ERuntimeVirtualTextureMaterialType::BaseColor_Normal: UnpackType = EVirtualTextureUnpackType::NormalBC5; break;
-			case ERuntimeVirtualTextureMaterialType::BaseColor_Normal_Specular: UnpackType = EVirtualTextureUnpackType::NormalBC3BC3; break;
-			case ERuntimeVirtualTextureMaterialType::BaseColor_Normal_Specular_Deprecated: UnpackType = EVirtualTextureUnpackType::NormalBC3; break;
+			case ERuntimeVirtualTextureMaterialType::BaseColor_Normal_Specular: UnpackType = EVirtualTextureUnpackType::NormalBC3; break;
+			case ERuntimeVirtualTextureMaterialType::BaseColor_Normal_Specular_Ex: UnpackType = EVirtualTextureUnpackType::NormalBC3BC3; break;
 			}
 		}
 		else
diff --git a/Engine/Source/Runtime/Engine/Private/VT/RuntimeVirtualTexture.cpp b/Engine/Source/Runtime/Engine/Private/VT/RuntimeVirtualTexture.cpp
--- a/Engine/Source/Runtime/Engine/Private/VT/RuntimeVirtualTexture.cpp
+++ b/Engine/Source/Runtime/Engine/Private/VT/RuntimeVirtualTexture.cpp
@@ -235,7 +235,7 @@ int32 URuntimeVirtualTexture::GetLayerCount() const
 		return 1;
 	case ERuntimeVirtualTextureMaterialType::BaseColor_Normal:
 	case ERuntimeVirtualTextureMaterialType::BaseColor_Normal_Specular:
-	case ERuntimeVirtualTextureMaterialType::BaseColor_Normal_Specular_Deprecated:
+	case ERuntimeVirtualTextureMaterialType::BaseColor_Normal_Specular_Ex:
 		return 2;
 	default:
 		break;
@@ -254,9 +254,9 @@ EPixelFormat URuntimeVirtualTexture::GetLayerFormat(int32 LayerIndex) const
 		{
 		case ERuntimeVirtualTextureMaterialType::BaseColor:
 		case ERuntimeVirtualTextureMaterialType::BaseColor_Normal:
-		case ERuntimeVirtualTextureMaterialType::BaseColor_Normal_Specular_Deprecated:
-			return bCompressTextures ? PF_DXT1 : PF_B8G8R8A8;
 		case ERuntimeVirtualTextureMaterialType::BaseColor_Normal_Specular:
+			return bCompressTextures ? PF_DXT1 : PF_B8G8R8A8;
+		case ERuntimeVirtualTextureMaterialType::BaseColor_Normal_Specular_Ex:
 			return bCompressTextures ? PF_DXT5 : PF_B8G8R8A8;
 		case ERuntimeVirtualTextureMaterialType::WorldHeight:
 			return PF_G16;
@@ -271,7 +271,7 @@ EPixelFormat URuntimeVirtualTexture::GetLayerFormat(int32 LayerIndex) const
 		case ERuntimeVirtualTextureMaterialType::BaseColor_Normal:
 			return bCompressTextures ? PF_BC5 : PF_B8G8R8A8;
 		case ERuntimeVirtualTextureMaterialType::BaseColor_Normal_Specular:
-		case ERuntimeVirtualTextureMaterialType::BaseColor_Normal_Specular_Deprecated:
+		case ERuntimeVirtualTextureMaterialType::BaseColor_Normal_Specular_Ex:
 			return bCompressTextures ? PF_DXT5 : PF_B8G8R8A8;
 		case ERuntimeVirtualTextureMaterialType::BaseColor:
 		case ERuntimeVirtualTextureMaterialType::WorldHeight:
@@ -292,7 +292,7 @@ bool URuntimeVirtualTexture::IsLayerSRGB(int32 LayerIndex) const
 	case ERuntimeVirtualTextureMaterialType::BaseColor:
 	case ERuntimeVirtualTextureMaterialType::BaseColor_Normal:
 	case ERuntimeVirtualTextureMaterialType::BaseColor_Normal_Specular:
-	case ERuntimeVirtualTextureMaterialType::BaseColor_Normal_Specular_Deprecated:
+	case ERuntimeVirtualTextureMaterialType::BaseColor_Normal_Specular_Ex:
 		// Only BaseColor layer is sRGB
 		return LayerIndex == 0;
 	case ERuntimeVirtualTextureMaterialType::WorldHeight:
diff --git a/Engine/Source/Runtime/Engine/Public/VT/RuntimeVirtualTextureEnum.h b/Engine/Source/Runtime/Engine/Public/VT/RuntimeVirtualTextureEnum.h
--- a/Engine/Source/Runtime/Engine/Public/VT/RuntimeVirtualTextureEnum.h
+++ b/Engine/Source/Runtime/Engine/Public/VT/RuntimeVirtualTextureEnum.h
@@ -15,9 +15,9 @@ enum class ERuntimeVirtualTextureMaterialType : uint8
 {
 	BaseColor UMETA(DisplayName = "Base Color"),
 	BaseColor_Normal UMETA(DisplayName = "Base Color, Normal"),
-	BaseColor_Normal_Specular UMETA(DisplayName = "Base Color, Normal, Roughness, Specular"),
+	BaseColor_Normal_Specular UMETA(DisplayName = "[BC1|BC3] Base Color, Normal, Roughness, Specular"),
+	BaseColor_Normal_Specular_Ex UMETA(DisplayName = "[BC3|BC3] Base Color, Normal, Roughness, Specular"),
 	WorldHeight UMETA(DisplayName = "World Height"),
-	BaseColor_Normal_Specular_Deprecated UMETA(DisplayName = "[OLD] Base Color, Normal, Roughness, Specular"),
 
 	Count UMETA(Hidden),
 };
diff --git a/Engine/Source/Runtime/Renderer/Private/VT/RuntimeVirtualTextureRender.cpp b/Engine/Source/Runtime/Renderer/Private/VT/RuntimeVirtualTextureRender.cpp
--- a/Engine/Source/Runtime/Renderer/Private/VT/RuntimeVirtualTextureRender.cpp
+++ b/Engine/Source/Runtime/Renderer/Private/VT/RuntimeVirtualTextureRender.cpp
@@ -263,7 +263,7 @@ namespace RuntimeVirtualTexture
 					Process<FMaterialPolicy_BaseColorNormal>(MeshBatch, BatchElementMask, StaticMeshId, PrimitiveSceneProxy, MaterialRenderProxy, Material);
 					break;
 				case ERuntimeVirtualTextureMaterialType::BaseColor_Normal_Specular:
-				case ERuntimeVirtualTextureMaterialType::BaseColor_Normal_Specular_Deprecated:
+				case ERuntimeVirtualTextureMaterialType::BaseColor_Normal_Specular_Ex:
 					Process<FMaterialPolicy_BaseColorNormalSpecular>(MeshBatch, BatchElementMask, StaticMeshId, PrimitiveSceneProxy, MaterialRenderProxy, Material);
 					break;
 				case ERuntimeVirtualTextureMaterialType::WorldHeight:
@@ -472,18 +472,18 @@ namespace RuntimeVirtualTexture
 	
 	typedef FShader_VirtualTextureCompress_CS< ERuntimeVirtualTextureMaterialType::BaseColor_Normal, false > FShader_VirtualTextureCompress_BaseColorNormal_CS;
 	IMPLEMENT_SHADER_TYPE(template<>, FShader_VirtualTextureCompress_BaseColorNormal_CS, TEXT("/Engine/Private/VirtualTextureCompress.usf"), TEXT("CompressBaseColorNormalCS"), SF_Compute);
-	
+
 	typedef FShader_VirtualTextureCompress_CS< ERuntimeVirtualTextureMaterialType::BaseColor_Normal_Specular, false > FShader_VirtualTextureCompress_BaseColorNormalSpecular_CS;
 	IMPLEMENT_SHADER_TYPE(template<>, FShader_VirtualTextureCompress_BaseColorNormalSpecular_CS, TEXT("/Engine/Private/VirtualTextureCompress.usf"), TEXT("CompressBaseColorNormalSpecularCS"), SF_Compute);
 
-	typedef FShader_VirtualTextureCompress_CS< ERuntimeVirtualTextureMaterialType::BaseColor_Normal_Specular_Deprecated, false > FShader_VirtualTextureCompress_BaseColorNormalSpecularDeprecated_CS;
-	IMPLEMENT_SHADER_TYPE(template<>, FShader_VirtualTextureCompress_BaseColorNormalSpecularDeprecated_CS, TEXT("/Engine/Private/VirtualTextureCompress.usf"), TEXT("CompressBaseColorNormalSpecularDeprecatedCS"), SF_Compute);
+	typedef FShader_VirtualTextureCompress_CS< ERuntimeVirtualTextureMaterialType::BaseColor_Normal_Specular_Ex, false > FShader_VirtualTextureCompress_BaseColorNormalSpecularEx_CS;
+	IMPLEMENT_SHADER_TYPE(template<>, FShader_VirtualTextureCompress_BaseColorNormalSpecularEx_CS, TEXT("/Engine/Private/VirtualTextureCompress.usf"), TEXT("CompressBaseColorNormalSpecularExCS"), SF_Compute);
 
 	typedef FShader_VirtualTextureCompress_CS< ERuntimeVirtualTextureMaterialType::BaseColor_Normal_Specular, true > FShader_VirtualTextureCompress_BaseColorNormalSpecular_CopyOnly_CS;
 	IMPLEMENT_SHADER_TYPE(template<>, FShader_VirtualTextureCompress_BaseColorNormalSpecular_CopyOnly_CS, TEXT("/Engine/Private/VirtualTextureCompress.usf"), TEXT("CopyNormalSpecularCS"), SF_Compute);
 
-	typedef FShader_VirtualTextureCompress_CS< ERuntimeVirtualTextureMaterialType::BaseColor_Normal_Specular_Deprecated, true > FShader_VirtualTextureCompress_BaseColorNormalSpecularDeprecated_CopyOnly_CS;
-	IMPLEMENT_SHADER_TYPE(template<>, FShader_VirtualTextureCompress_BaseColorNormalSpecularDeprecated_CopyOnly_CS, TEXT("/Engine/Private/VirtualTextureCompress.usf"), TEXT("CopyNormalSpecularDeprecatedCS"), SF_Compute);
+	typedef FShader_VirtualTextureCompress_CS< ERuntimeVirtualTextureMaterialType::BaseColor_Normal_Specular_Ex, true > FShader_VirtualTextureCompress_BaseColorNormalSpecularEx_CopyOnly_CS;
+	IMPLEMENT_SHADER_TYPE(template<>, FShader_VirtualTextureCompress_BaseColorNormalSpecularEx_CopyOnly_CS, TEXT("/Engine/Private/VirtualTextureCompress.usf"), TEXT("CopyNormalSpecularExCS"), SF_Compute);
 
 
 	/** Set up the BC compression pass for the specific MaterialType. */
@@ -516,8 +516,8 @@ namespace RuntimeVirtualTexture
 		case ERuntimeVirtualTextureMaterialType::BaseColor_Normal_Specular:
 			AddCompressOrCopyPass<ERuntimeVirtualTextureMaterialType::BaseColor_Normal_Specular>(GraphBuilder, FeatureLevel, Parameters, GroupCount);
 			break;
-		case ERuntimeVirtualTextureMaterialType::BaseColor_Normal_Specular_Deprecated:
-			AddCompressOrCopyPass<ERuntimeVirtualTextureMaterialType::BaseColor_Normal_Specular_Deprecated>(GraphBuilder, FeatureLevel, Parameters, GroupCount);
+		case ERuntimeVirtualTextureMaterialType::BaseColor_Normal_Specular_Ex:
+			AddCompressOrCopyPass<ERuntimeVirtualTextureMaterialType::BaseColor_Normal_Specular_Ex>(GraphBuilder, FeatureLevel, Parameters, GroupCount);
 			break;
 		}
 	}
@@ -533,8 +533,8 @@ namespace RuntimeVirtualTexture
 		case ERuntimeVirtualTextureMaterialType::BaseColor_Normal_Specular:
 			AddCompressOrCopyPass<ERuntimeVirtualTextureMaterialType::BaseColor_Normal_Specular, true>(GraphBuilder, FeatureLevel, Parameters, GroupCount);
 			break;
-		case ERuntimeVirtualTextureMaterialType::BaseColor_Normal_Specular_Deprecated:
-			AddCompressOrCopyPass<ERuntimeVirtualTextureMaterialType::BaseColor_Normal_Specular_Deprecated, true>(GraphBuilder, FeatureLevel, Parameters, GroupCount);
+		case ERuntimeVirtualTextureMaterialType::BaseColor_Normal_Specular_Ex:
+			AddCompressOrCopyPass<ERuntimeVirtualTextureMaterialType::BaseColor_Normal_Specular_Ex, true>(GraphBuilder, FeatureLevel, Parameters, GroupCount);
 			break;
 		}
 	}
@@ -548,7 +548,7 @@ namespace RuntimeVirtualTexture
 		{
 			bRenderPass = OutputTexture0 != nullptr;
 			bCompressPass = bRenderPass && (OutputTexture0->GetFormat() == PF_DXT1 || OutputTexture0->GetFormat() == PF_DXT5 || OutputTexture0->GetFormat() == PF_BC5);
-			bCopyPass = bRenderPass && !bCompressPass && (MaterialType == ERuntimeVirtualTextureMaterialType::BaseColor_Normal_Specular || MaterialType == ERuntimeVirtualTextureMaterialType::BaseColor_Normal_Specular_Deprecated);
+			bCopyPass = bRenderPass && !bCompressPass && (MaterialType == ERuntimeVirtualTextureMaterialType::BaseColor_Normal_Specular || MaterialType == ERuntimeVirtualTextureMaterialType::BaseColor_Normal_Specular_Ex);
 		
 			switch (MaterialType)
 			{
@@ -583,16 +583,15 @@ namespace RuntimeVirtualTexture
 				}
 				if (bCompressPass)
 				{
-					OutputAlias0 = CompressTexture0_u4 = GraphBuilder.CreateTexture(FPooledRenderTargetDesc::Create2DDesc(TextureSize / 4, PF_R32G32B32A32_UINT, FClearValueBinding::None, TexCreate_None, TexCreate_UAV, false), TEXT("CompressTexture0"));
+					OutputAlias0 = CompressTexture0_u2 = GraphBuilder.CreateTexture(FPooledRenderTargetDesc::Create2DDesc(TextureSize / 4, PF_R32G32_UINT, FClearValueBinding::None, TexCreate_None, TexCreate_UAV, false), TEXT("CompressTexture0"));
 					OutputAlias1 = CompressTexture1 = GraphBuilder.CreateTexture(FPooledRenderTargetDesc::Create2DDesc(TextureSize / 4, PF_R32G32B32A32_UINT, FClearValueBinding::None, TexCreate_None, TexCreate_UAV, false), TEXT("CompressTexture1"));
 				}
 				if (bCopyPass)
 				{
-					OutputAlias0 = CopyTexture0 = GraphBuilder.CreateTexture(FPooledRenderTargetDesc::Create2DDesc(TextureSize, PF_B8G8R8A8, FClearValueBinding::None, TexCreate_None, TexCreate_UAV, false), TEXT("CopyTexture0"));
-					OutputAlias1 = CopyTexture1 = GraphBuilder.CreateTexture(FPooledRenderTargetDesc::Create2DDesc(TextureSize, PF_B8G8R8A8, FClearValueBinding::None, TexCreate_None, TexCreate_UAV, false), TEXT("CopyTexture1"));
+					OutputAlias1 = CopyTexture0 = GraphBuilder.CreateTexture(FPooledRenderTargetDesc::Create2DDesc(TextureSize, PF_B8G8R8A8, FClearValueBinding::None, TexCreate_None, TexCreate_UAV, false), TEXT("CopyTexture0"));
 				}
 				break;
-			case ERuntimeVirtualTextureMaterialType::BaseColor_Normal_Specular_Deprecated:
+			case ERuntimeVirtualTextureMaterialType::BaseColor_Normal_Specular_Ex:
 				if (bRenderPass)
 				{
 					OutputAlias0 = RenderTexture0 = GraphBuilder.CreateTexture(FPooledRenderTargetDesc::Create2DDesc(TextureSize, PF_B8G8R8A8, FClearValueBinding::Black, TexCreate_SRGB, TexCreate_RenderTargetable, false), TEXT("RenderTexture0"));
@@ -601,12 +600,13 @@ namespace RuntimeVirtualTexture
 				}
 				if (bCompressPass)
 				{
-					OutputAlias0 = CompressTexture0_u2 = GraphBuilder.CreateTexture(FPooledRenderTargetDesc::Create2DDesc(TextureSize / 4, PF_R32G32_UINT, FClearValueBinding::None, TexCreate_None, TexCreate_UAV, false), TEXT("CompressTexture0"));
+					OutputAlias0 = CompressTexture0_u4 = GraphBuilder.CreateTexture(FPooledRenderTargetDesc::Create2DDesc(TextureSize / 4, PF_R32G32B32A32_UINT, FClearValueBinding::None, TexCreate_None, TexCreate_UAV, false), TEXT("CompressTexture0"));
 					OutputAlias1 = CompressTexture1 = GraphBuilder.CreateTexture(FPooledRenderTargetDesc::Create2DDesc(TextureSize / 4, PF_R32G32B32A32_UINT, FClearValueBinding::None, TexCreate_None, TexCreate_UAV, false), TEXT("CompressTexture1"));
 				}
 				if (bCopyPass)
 				{
-					OutputAlias1 = CopyTexture0 = GraphBuilder.CreateTexture(FPooledRenderTargetDesc::Create2DDesc(TextureSize, PF_B8G8R8A8, FClearValueBinding::None, TexCreate_None, TexCreate_UAV, false), TEXT("CopyTexture0"));
+					OutputAlias0 = CopyTexture0 = GraphBuilder.CreateTexture(FPooledRenderTargetDesc::Create2DDesc(TextureSize, PF_B8G8R8A8, FClearValueBinding::None, TexCreate_None, TexCreate_UAV, false), TEXT("CopyTexture0"));
+					OutputAlias1 = CopyTexture1 = GraphBuilder.CreateTexture(FPooledRenderTargetDesc::Create2DDesc(TextureSize, PF_B8G8R8A8, FClearValueBinding::None, TexCreate_None, TexCreate_UAV, false), TEXT("CopyTexture1"));
 				}
 				break;
 			case ERuntimeVirtualTextureMaterialType::WorldHeight:
