Fixed crash when rendering a scene with from DrawTileMesh.

diff --git a/Engine/Source/Runtime/Engine/Public/SceneInterface.h b/Engine/Source/Runtime/Engine/Public/SceneInterface.h
--- a/Engine/Source/Runtime/Engine/Public/SceneInterface.h
+++ b/Engine/Source/Runtime/Engine/Public/SceneInterface.h
@@ -261,6 +261,7 @@ public:
 	 * Returns the scene's unique FSkyAtmosphereRenderSceneInfo if it exists
 	 */
 	virtual FSkyAtmosphereRenderSceneInfo* GetSkyAtmosphereSceneInfo() = 0;
+	virtual const FSkyAtmosphereRenderSceneInfo* GetSkyAtmosphereSceneInfo() const = 0;
 	/**
 	 * Override a sky atmosphere light direction
 	 * @param SkyAtmosphereComponent - component to verify it is the actual unique SkyAtmosphere
diff --git a/Engine/Source/Runtime/Renderer/Private/DeferredShadingRenderer.cpp b/Engine/Source/Runtime/Renderer/Private/DeferredShadingRenderer.cpp
--- a/Engine/Source/Runtime/Renderer/Private/DeferredShadingRenderer.cpp
+++ b/Engine/Source/Runtime/Renderer/Private/DeferredShadingRenderer.cpp
@@ -845,7 +845,7 @@ void FDeferredShadingSceneRenderer::Render(FRHICommandListImmediate& RHICmdList)
 
 	PrepareViewRectsForRendering();
 
-	if (Scene->HasSkyAtmosphere() && ShouldRenderSkyAtmosphere(Scene->SkyAtmosphere, Scene->GetShaderPlatform()))
+	if (ShouldRenderSkyAtmosphere(Scene))
 	{
 		for (int32 LightIndex = 0; LightIndex < NUM_ATMOSPHERE_LIGHTS; ++LightIndex)
 		{
@@ -1153,7 +1153,7 @@ void FDeferredShadingSceneRenderer::Render(FRHICommandListImmediate& RHICmdList)
 	}
 
 	// Generate the Sky/Atmosphere look up tables
-	const bool bShouldRenderSkyAtmosphere = ShouldRenderSkyAtmosphere(Scene->GetSkyAtmosphereSceneInfo(), Scene->GetShaderPlatform());
+	const bool bShouldRenderSkyAtmosphere = ShouldRenderSkyAtmosphere(Scene);
 	if (bShouldRenderSkyAtmosphere)
 	{
 		RenderSkyAtmosphereLookUpTables(RHICmdList);
diff --git a/Engine/Source/Runtime/Renderer/Private/MobileShadingRenderer.cpp b/Engine/Source/Runtime/Renderer/Private/MobileShadingRenderer.cpp
--- a/Engine/Source/Runtime/Renderer/Private/MobileShadingRenderer.cpp
+++ b/Engine/Source/Runtime/Renderer/Private/MobileShadingRenderer.cpp
@@ -240,7 +240,7 @@ void FMobileSceneRenderer::InitViews(FRHICommandListImmediate& RHICmdList)
 	ComputeViewVisibility(RHICmdList, BasePassDepthStencilAccess, ViewCommandsPerView, DynamicIndexBuffer, DynamicVertexBuffer, DynamicReadBuffer);
 
 	// Initialise Sky/View resources before the view global uniform buffer is built.
-	if (ShouldRenderSkyAtmosphere(Scene->GetSkyAtmosphereSceneInfo(), Scene->GetShaderPlatform()))
+	if (ShouldRenderSkyAtmosphere(Scene))
 	{
 		InitSkyAtmosphereForViews(RHICmdList);
 	}
@@ -317,7 +317,7 @@ void FMobileSceneRenderer::Render(FRHICommandListImmediate& RHICmdList)
 
 	PrepareViewRectsForRendering();
 
-	if (Scene->HasSkyAtmosphere() && ShouldRenderSkyAtmosphere(Scene->SkyAtmosphere, Scene->GetShaderPlatform()))
+	if (ShouldRenderSkyAtmosphere(Scene))
 	{
 		for (int32 LightIndex = 0; LightIndex < NUM_ATMOSPHERE_LIGHTS; ++LightIndex)
 		{
@@ -406,7 +406,7 @@ void FMobileSceneRenderer::Render(FRHICommandListImmediate& RHICmdList)
 	const bool bGammaSpace = !IsMobileHDR();
 
 	// Generate the Sky/Atmosphere look up tables
-	const bool bShouldRenderSkyAtmosphere = ShouldRenderSkyAtmosphere(Scene->GetSkyAtmosphereSceneInfo(), Scene->GetShaderPlatform());
+	const bool bShouldRenderSkyAtmosphere = ShouldRenderSkyAtmosphere(Scene);
 	if (bShouldRenderSkyAtmosphere)
 	{
 		RenderSkyAtmosphereLookUpTables(RHICmdList);
diff --git a/Engine/Source/Runtime/Renderer/Private/Renderer.cpp b/Engine/Source/Runtime/Renderer/Private/Renderer.cpp
--- a/Engine/Source/Runtime/Renderer/Private/Renderer.cpp
+++ b/Engine/Source/Runtime/Renderer/Private/Renderer.cpp
@@ -89,7 +89,7 @@ void FRendererModule::DrawTileMesh(FRHICommandListImmediate& RHICmdList, FMeshPa
 		const auto FeatureLevel = View.GetFeatureLevel();
 		const EShadingPath ShadingPath = FSceneInterface::GetShadingPath(FeatureLevel);
 		const FSceneViewFamily* ViewFamily = View.Family;
-		const FScene* Scene = nullptr;
+		FScene* Scene = nullptr;
 
 		if (ViewFamily->Scene)
 		{
@@ -135,7 +135,14 @@ void FRendererModule::DrawTileMesh(FRHICommandListImmediate& RHICmdList, FMeshPa
 				View.LightmapSceneDataOverrideSRV = SinglePrimitiveStructuredBuffer.LightmapSceneDataBufferSRV;
 			}
 		}
-		
+
+		// Initialise Sky/View resources before the view global uniform buffer is built.
+		if (ShouldRenderSkyAtmosphere(Scene))
+		{
+			InitSkyAtmosphereForScene(RHICmdList, Scene);
+			InitSkyAtmosphereForView(RHICmdList, Scene, View);
+		}
+
 		View.InitRHIResources();
 		DrawRenderState.SetViewUniformBuffer(View.ViewUniformBuffer);
 
diff --git a/Engine/Source/Runtime/Renderer/Private/RendererScene.cpp b/Engine/Source/Runtime/Renderer/Private/RendererScene.cpp
--- a/Engine/Source/Runtime/Renderer/Private/RendererScene.cpp
+++ b/Engine/Source/Runtime/Renderer/Private/RendererScene.cpp
@@ -3707,6 +3707,7 @@ public:
 	virtual void AddSkyAtmosphere(const class USkyAtmosphereComponent* SkyAtmosphereComponent, bool bStaticLightingBuilt) override {}
 	virtual void RemoveSkyAtmosphere(const class USkyAtmosphereComponent* SkyAtmosphereComponent) override {}
 	virtual FSkyAtmosphereRenderSceneInfo* GetSkyAtmosphereSceneInfo() override { return NULL; }
+	virtual const FSkyAtmosphereRenderSceneInfo* GetSkyAtmosphereSceneInfo() const override { return NULL; }
 	virtual void OverrideSkyAtmosphereLightDirection(const class USkyAtmosphereComponent* SkyAtmosphereComponent, int32 AtmosphereLightIndex, const FVector& LightDirection) override {}
 
 	virtual void AddWindSource(class UWindDirectionalSourceComponent* WindComponent) override {}
diff --git a/Engine/Source/Runtime/Renderer/Private/ScenePrivate.h b/Engine/Source/Runtime/Renderer/Private/ScenePrivate.h
--- a/Engine/Source/Runtime/Renderer/Private/ScenePrivate.h
+++ b/Engine/Source/Runtime/Renderer/Private/ScenePrivate.h
@@ -2689,6 +2689,7 @@ public:
 	virtual void AddSkyAtmosphere(const class USkyAtmosphereComponent* SkyAtmosphereComponent, bool bStaticLightingBuilt) override;
 	virtual void RemoveSkyAtmosphere(const class USkyAtmosphereComponent* SkyAtmosphereComponent) override;
 	virtual FSkyAtmosphereRenderSceneInfo* GetSkyAtmosphereSceneInfo() override { return SkyAtmosphere; }
+	virtual const FSkyAtmosphereRenderSceneInfo* GetSkyAtmosphereSceneInfo() const override { return SkyAtmosphere; }
 	virtual void OverrideSkyAtmosphereLightDirection(const class USkyAtmosphereComponent* SkyAtmosphereComponent, int32 AtmosphereLightIndex, const FVector& LightDirection) override;
 
 	virtual void AddWindSource(UWindDirectionalSourceComponent* WindComponent) override;
diff --git a/Engine/Source/Runtime/Renderer/Private/SceneRendering.cpp b/Engine/Source/Runtime/Renderer/Private/SceneRendering.cpp
--- a/Engine/Source/Runtime/Renderer/Private/SceneRendering.cpp
+++ b/Engine/Source/Runtime/Renderer/Private/SceneRendering.cpp
@@ -1220,7 +1220,7 @@ void FViewInfo::SetupUniformBufferParameters(
 	FRHITexture* SkyViewLutTextureFound = nullptr;
 	FRHITexture* CameraAerialPerspectiveVolumeFound = nullptr;
 	FRHITexture* DistantSkyLightLutTextureFound = nullptr;
-	if (Scene && ShouldRenderSkyAtmosphere(Scene->SkyAtmosphere, Scene->GetShaderPlatform()))
+	if (ShouldRenderSkyAtmosphere(Scene))
 	{
 		check(Scene->SkyAtmosphere->GetTransmittanceLutTexture().IsValid());
 		FSkyAtmosphereRenderSceneInfo* SkyAtmosphere = Scene->SkyAtmosphere;
diff --git a/Engine/Source/Runtime/Renderer/Private/SceneVisibility.cpp b/Engine/Source/Runtime/Renderer/Private/SceneVisibility.cpp
--- a/Engine/Source/Runtime/Renderer/Private/SceneVisibility.cpp
+++ b/Engine/Source/Runtime/Renderer/Private/SceneVisibility.cpp
@@ -4169,7 +4169,7 @@ bool FDeferredShadingSceneRenderer::InitViews(FRHICommandListImmediate& RHICmdLi
 	RHICmdList.ImmediateFlush(EImmediateFlushType::DispatchToRHIThread);
 
 	// Initialise Sky/View resources before the view global uniform buffer is built.
-	if (Scene && ShouldRenderSkyAtmosphere(Scene->GetSkyAtmosphereSceneInfo(), Scene->GetShaderPlatform()))
+	if (ShouldRenderSkyAtmosphere(Scene))
 	{
 		InitSkyAtmosphereForViews(RHICmdList);
 	}
diff --git a/Engine/Source/Runtime/Renderer/Private/SkyAtmosphereRendering.cpp b/Engine/Source/Runtime/Renderer/Private/SkyAtmosphereRendering.cpp
--- a/Engine/Source/Runtime/Renderer/Private/SkyAtmosphereRendering.cpp
+++ b/Engine/Source/Runtime/Renderer/Private/SkyAtmosphereRendering.cpp
@@ -269,11 +269,19 @@ static bool ShouldPipelineCompileSkyAtmosphereShader(EShaderPlatform ShaderPlatf
 	return RHISupportsComputeShaders(ShaderPlatform);
 }
 
-bool ShouldRenderSkyAtmosphere(const FSkyAtmosphereRenderSceneInfo* SkyAtmosphere, EShaderPlatform ShaderPlatform)
+bool ShouldRenderSkyAtmosphere(const FScene* Scene)
 {
-	const bool ShadersCompiled = ShouldPipelineCompileSkyAtmosphereShader(ShaderPlatform);
-	return FReadOnlyCVARCache::Get().bSupportSkyAtmosphere && SkyAtmosphere && ShadersCompiled && CVarSkyAtmosphere.GetValueOnRenderThread() > 0;
-	// TODO: Add new or reuse EngineShowFlags.AtmosphericFog? ALso take into account EngineShowFlags.Fog as previously?
+	if (Scene && Scene->HasSkyAtmosphere())
+	{
+		EShaderPlatform ShaderPlatform = Scene->GetShaderPlatform();
+		const FSkyAtmosphereRenderSceneInfo* SkyAtmosphere = Scene->GetSkyAtmosphereSceneInfo();
+		check(SkyAtmosphere);
+
+		const bool ShadersCompiled = ShouldPipelineCompileSkyAtmosphereShader(ShaderPlatform);
+		return FReadOnlyCVARCache::Get().bSupportSkyAtmosphere && ShadersCompiled && CVarSkyAtmosphere.GetValueOnRenderThread() > 0;
+		// TODO: Add new or reuse EngineShowFlags.AtmosphericFog? ALso take into account EngineShowFlags.Fog as previously?
+	}
+	return false;
 }
 
 void SetupSkyAtmosphereViewSharedUniformShaderParameters(const FViewInfo& View, FSkyAtmosphereViewSharedUniformShaderParameters& OutParameters)
@@ -880,24 +888,45 @@ TGlobalResource<FUniformSphereSamplesBuffer> GUniformSphereSamplesBuffer;
 =============================================================================*/
 
 void FSceneRenderer::InitSkyAtmosphereForViews(FRHICommandListImmediate& RHICmdList)
+{
+	InitSkyAtmosphereForScene(RHICmdList, Scene);
+
+	for (int32 ViewIndex = 0; ViewIndex < Views.Num(); ViewIndex++)
+	{
+		FViewInfo& View = Views[ViewIndex];
+		InitSkyAtmosphereForView(RHICmdList, Scene, View);
+	}
+}
+
+static EPixelFormat GetSkyLutTextureFormat(ERHIFeatureLevel::Type FeatureLevel)
+{
+	EPixelFormat TextureLUTFormat = PF_FloatRGB;
+	if (FeatureLevel <= ERHIFeatureLevel::ES3_1)
+	{
+		// OpenGL ES3.1 does not support storing into 3-component images
+		// TODO: check if need this for Metal, Vulkan
+		TextureLUTFormat = PF_FloatRGBA;
+	}
+	return TextureLUTFormat;
+}
+static EPixelFormat GetSkyLutSmallTextureFormat()
+{
+	return PF_R8G8B8A8;
+}
+
+void InitSkyAtmosphereForScene(FRHICommandListImmediate& RHICmdList, FScene* Scene)
 {
 	if (Scene)
 	{
 		GET_VALID_DATA_FROM_CVAR;
 
-		check(ShouldRenderSkyAtmosphere(Scene->GetSkyAtmosphereSceneInfo(), Scene->GetShaderPlatform())); // This should not be called if we should not render SkyAtmosphere
+		check(ShouldRenderSkyAtmosphere(Scene)); // This should not be called if we should not render SkyAtmosphere
 		FPooledRenderTargetDesc Desc;
 		check(Scene->GetSkyAtmosphereSceneInfo());
 		FSkyAtmosphereRenderSceneInfo& SkyInfo = *Scene->GetSkyAtmosphereSceneInfo();
 
-		EPixelFormat TextureLUTFormat = PF_FloatRGB;
-		EPixelFormat TextureLUTSmallFormat = PF_R8G8B8A8;
-		if (Scene->GetFeatureLevel() <= ERHIFeatureLevel::ES3_1)
-		{
-			// OpenGL ES3.1 does not support storing into 3-component images
-			// TODO: check if need this for Metal, Vulkan
-			TextureLUTFormat = PF_FloatRGBA;
-		}  
+		EPixelFormat TextureLUTFormat = GetSkyLutTextureFormat(Scene->GetFeatureLevel());
+		EPixelFormat TextureLUTSmallFormat = GetSkyLutSmallTextureFormat();
 
 		//
 		// Initialise per scene/atmosphere resources
@@ -907,13 +936,13 @@ void FSceneRenderer::InitSkyAtmosphereForViews(FRHICommandListImmediate& RHICmdL
 
 		TRefCountPtr<IPooledRenderTarget>& TransmittanceLutTexture = SkyInfo.GetTransmittanceLutTexture();
 		Desc = FPooledRenderTargetDesc::Create2DDesc(
-			FIntPoint(TransmittanceLutWidth, TransmittanceLutHeight), 
+			FIntPoint(TransmittanceLutWidth, TransmittanceLutHeight),
 			TranstmittanceLUTUseSmallFormat ? TextureLUTSmallFormat : TextureLUTFormat, FClearValueBinding::None, TexCreate_HideInVisualizeTexture, TexCreate_ShaderResource | TexCreate_UAV, false);
 		GRenderTargetPool.FindFreeElement(RHICmdList, Desc, TransmittanceLutTexture, TEXT("TransmittanceLutTexture"), true, ERenderTargetTransience::NonTransient);
 
 		TRefCountPtr<IPooledRenderTarget>& MultiScatteredLuminanceLutTexture = SkyInfo.GetMultiScatteredLuminanceLutTexture();
 		Desc = FPooledRenderTargetDesc::Create2DDesc(
-			FIntPoint(MultiScatteredLuminanceLutWidth, MultiScatteredLuminanceLutHeight), 
+			FIntPoint(MultiScatteredLuminanceLutWidth, MultiScatteredLuminanceLutHeight),
 			MultiScatteringLUTUseSmallFormat ? TextureLUTSmallFormat : TextureLUTFormat, FClearValueBinding::None, TexCreate_HideInVisualizeTexture, TexCreate_ShaderResource | TexCreate_UAV, false);
 		GRenderTargetPool.FindFreeElement(RHICmdList, Desc, MultiScatteredLuminanceLutTexture, TEXT("MultiScatteredLuminanceLutTexture"), true, ERenderTargetTransience::NonTransient);
 
@@ -925,27 +954,39 @@ void FSceneRenderer::InitSkyAtmosphereForViews(FRHICommandListImmediate& RHICmdL
 				TextureLUTFormat, FClearValueBinding::None, TexCreate_HideInVisualizeTexture, TexCreate_ShaderResource | TexCreate_UAV, false);
 			GRenderTargetPool.FindFreeElement(RHICmdList, Desc, DistantSkyLightLutTexture, TEXT("DistantSkyLightLutTexture"), true, ERenderTargetTransience::NonTransient);
 		}
+	}
+}
+
+void InitSkyAtmosphereForView(FRHICommandListImmediate& RHICmdList, const FScene* Scene, FViewInfo& View)
+{
+	if (Scene)
+	{
+		GET_VALID_DATA_FROM_CVAR;
+
+		check(ShouldRenderSkyAtmosphere(Scene)); // This should not be called if we should not render SkyAtmosphere
+		FPooledRenderTargetDesc Desc;
+		check(Scene->GetSkyAtmosphereSceneInfo());
+		const FSkyAtmosphereRenderSceneInfo& SkyInfo = *Scene->GetSkyAtmosphereSceneInfo();
+
+		EPixelFormat TextureLUTFormat = GetSkyLutTextureFormat(Scene->GetFeatureLevel());
+		EPixelFormat TextureLUTSmallFormat = GetSkyLutSmallTextureFormat();
 
 		//
 		// Initialise transient per view resources.
 		//
 
 		FPooledRenderTargetDesc SkyAtmosphereViewLutTextureDesc = FPooledRenderTargetDesc::Create2DDesc(
-			FIntPoint(SkyViewLutWidth, SkyViewLutHeight), 
+			FIntPoint(SkyViewLutWidth, SkyViewLutHeight),
 			TextureLUTFormat, FClearValueBinding::None, TexCreate_HideInVisualizeTexture, TexCreate_ShaderResource | TexCreate_UAV, false);
 
 		FPooledRenderTargetDesc SkyAtmosphereCameraAerialPerspectiveVolumeDesc = FPooledRenderTargetDesc::CreateVolumeDesc(
-			CameraAerialPerspectiveVolumeScreenResolution, CameraAerialPerspectiveVolumeScreenResolution, CameraAerialPerspectiveVolumeDepthResolution, 
+			CameraAerialPerspectiveVolumeScreenResolution, CameraAerialPerspectiveVolumeScreenResolution, CameraAerialPerspectiveVolumeDepthResolution,
 			PF_FloatRGBA, FClearValueBinding::None, TexCreate_HideInVisualizeTexture, TexCreate_ShaderResource | TexCreate_UAV, false);
 
-		for (int32 ViewIndex = 0; ViewIndex < Views.Num(); ViewIndex++)
-		{
-			FViewInfo& View = Views[ViewIndex];
-			// Set textures and data that will be needed later on the view.
-			View.SkyAtmosphereUniformShaderParameters = SkyInfo.GetAtmosphereShaderParameters();
-			GRenderTargetPool.FindFreeElement(RHICmdList, SkyAtmosphereViewLutTextureDesc, View.SkyAtmosphereViewLutTexture, TEXT("View.SkyAtmosphereViewLutTexture"), true, ERenderTargetTransience::Transient);
-			GRenderTargetPool.FindFreeElement(RHICmdList, SkyAtmosphereCameraAerialPerspectiveVolumeDesc, View.SkyAtmosphereCameraAerialPerspectiveVolume, TEXT("View.SkyAtmosphereCameraAerialPerspectiveVolume"), true, ERenderTargetTransience::Transient);
-		}
+		// Set textures and data that will be needed later on the view.
+		View.SkyAtmosphereUniformShaderParameters = SkyInfo.GetAtmosphereShaderParameters();
+		GRenderTargetPool.FindFreeElement(RHICmdList, SkyAtmosphereViewLutTextureDesc, View.SkyAtmosphereViewLutTexture, TEXT("View.SkyAtmosphereViewLutTexture"), true, ERenderTargetTransience::Transient);
+		GRenderTargetPool.FindFreeElement(RHICmdList, SkyAtmosphereCameraAerialPerspectiveVolumeDesc, View.SkyAtmosphereCameraAerialPerspectiveVolume, TEXT("View.SkyAtmosphereCameraAerialPerspectiveVolume"), true, ERenderTargetTransience::Transient);
 	}
 }
 
@@ -1024,7 +1065,7 @@ static bool IsSecondAtmosphereLightEnabled(FScene* Scene)
 
 void FSceneRenderer::RenderSkyAtmosphereLookUpTables(FRHICommandListImmediate& RHICmdList)
 {
-	check(ShouldRenderSkyAtmosphere(Scene->GetSkyAtmosphereSceneInfo(), Scene->GetShaderPlatform())); // This should not be called if we should not render SkyAtmosphere
+	check(ShouldRenderSkyAtmosphere(Scene)); // This should not be called if we should not render SkyAtmosphere
 
 	DECLARE_GPU_STAT(SkyAtmosphereLUTs);
 	SCOPED_DRAW_EVENT(RHICmdList, SkyAtmosphereLUTs);
@@ -1185,7 +1226,7 @@ void FSceneRenderer::RenderSkyAtmosphere(FRHICommandListImmediate& RHICmdList)
 	// In this case, the sky must be rendered as an opaque mesh using the shader graph as a composition tool.
 	check(!IsMobilePlatform(Scene->GetShaderPlatform()));
 
-	check(ShouldRenderSkyAtmosphere(Scene->GetSkyAtmosphereSceneInfo(), Scene->GetShaderPlatform())); // This should not be called if we should not render SkyAtmosphere
+	check(ShouldRenderSkyAtmosphere(Scene)); // This should not be called if we should not render SkyAtmosphere
 
 	DECLARE_GPU_STAT(SkyAtmosphere);
 	SCOPED_DRAW_EVENT(RHICmdList, SkyAtmosphere);
@@ -1426,7 +1467,7 @@ void FSceneRenderer::RenderSkyAtmosphereEditorNotifications(FRHICommandListImmed
 void FDeferredShadingSceneRenderer::RenderDebugSkyAtmosphere(FRHICommandListImmediate& RHICmdList)
 {
 #if WITH_EDITOR
-	check(ShouldRenderSkyAtmosphere(Scene->GetSkyAtmosphereSceneInfo(), Scene->GetShaderPlatform())); // This should not be called if we should not render SkyAtmosphere
+	check(ShouldRenderSkyAtmosphere(Scene)); // This should not be called if we should not render SkyAtmosphere
 
 	//if (!RHISupportsComputeShaders()) return;	// TODO cannot render, add a ShouldRender function. Also should PipelineShouldCook ?
 
diff --git a/Engine/Source/Runtime/Renderer/Private/SkyAtmosphereRendering.h b/Engine/Source/Runtime/Renderer/Private/SkyAtmosphereRendering.h
--- a/Engine/Source/Runtime/Renderer/Private/SkyAtmosphereRendering.h
+++ b/Engine/Source/Runtime/Renderer/Private/SkyAtmosphereRendering.h
@@ -102,7 +102,10 @@ private:
 	FVector OverrideAtmosphericLightDirection[NUM_ATMOSPHERE_LIGHTS];
 };
 
-bool ShouldRenderSkyAtmosphere(const FSkyAtmosphereRenderSceneInfo* SkyAtmosphere, EShaderPlatform ShaderPlatform);
+bool ShouldRenderSkyAtmosphere(const FScene* Scene);
+
+void InitSkyAtmosphereForScene(FRHICommandListImmediate& RHICmdList, FScene* Scene);
+void InitSkyAtmosphereForView(FRHICommandListImmediate& RHICmdList, const FScene* Scene, FViewInfo& View);
 
 extern void SetupSkyAtmosphereViewSharedUniformShaderParameters(const class FViewInfo& View, FSkyAtmosphereViewSharedUniformShaderParameters& OutParameters);
 
