Move sky LUT textures to be transient for now since they are evaluated every frame.
Use RDG shader resource declaration.
Remove unbinRenderTarget as RDG should handle it.

diff --git a/Engine/Source/Runtime/Renderer/Private/SkyAtmosphereRendering.cpp b/Engine/Source/Runtime/Renderer/Private/SkyAtmosphereRendering.cpp
--- a/Engine/Source/Runtime/Renderer/Private/SkyAtmosphereRendering.cpp
+++ b/Engine/Source/Runtime/Renderer/Private/SkyAtmosphereRendering.cpp
@@ -549,10 +549,10 @@ class FRenderSkyAtmospherePS : public FGlobalShader
 		SHADER_PARAMETER_STRUCT_REF(FViewUniformShaderParameters, ViewUniformBuffer)
 		RENDER_TARGET_BINDING_SLOTS()
 		SHADER_PARAMETER_STRUCT_REF(FSceneTexturesUniformParameters, SceneTextures)
-		SHADER_PARAMETER_TEXTURE(Texture2D<float3>, TransmittanceLutTexture)
-		SHADER_PARAMETER_TEXTURE(Texture2D<float3>, MultiScatteredLuminanceLutTexture)
-		SHADER_PARAMETER_TEXTURE(Texture2D<float3>, SkyViewLutTexture)
-		SHADER_PARAMETER_TEXTURE(Texture3D<float4>, CameraAerialPerspectiveVolumeTexture)
+		SHADER_PARAMETER_RDG_TEXTURE(Texture2D<float3>, TransmittanceLutTexture)
+		SHADER_PARAMETER_RDG_TEXTURE(Texture2D<float3>, MultiScatteredLuminanceLutTexture)
+		SHADER_PARAMETER_RDG_TEXTURE(Texture2D<float3>, SkyViewLutTexture)
+		SHADER_PARAMETER_RDG_TEXTURE(Texture3D<float4>, CameraAerialPerspectiveVolumeTexture)
 		SHADER_PARAMETER_SAMPLER(SamplerState, TransmittanceLutTextureSampler)
 		SHADER_PARAMETER_SAMPLER(SamplerState, MultiScatteredLuminanceLutTextureSampler)
 		SHADER_PARAMETER_SAMPLER(SamplerState, SkyViewLutTextureSampler)
@@ -625,7 +625,7 @@ public:
 	BEGIN_SHADER_PARAMETER_STRUCT(FParameters, )
 		SHADER_PARAMETER_STRUCT_REF(FAtmosphereUniformShaderParameters, Atmosphere)
 		SHADER_PARAMETER_STRUCT_REF(FSkyAtmosphereInternalCommonParameters, SkyAtmosphere)
-		SHADER_PARAMETER_UAV(RWTexture2D<float3>, TransmittanceLutUAV)
+		SHADER_PARAMETER_RDG_TEXTURE_UAV(RWTexture2D<float3>, TransmittanceLutUAV)
 	END_SHADER_PARAMETER_STRUCT()
 
 	static bool ShouldCompilePermutation(const FGlobalShaderPermutationParameters& Parameters)
@@ -658,8 +658,8 @@ public:
 	BEGIN_SHADER_PARAMETER_STRUCT(FParameters, )
 		SHADER_PARAMETER_STRUCT_REF(FAtmosphereUniformShaderParameters, Atmosphere)
 		SHADER_PARAMETER_STRUCT_REF(FSkyAtmosphereInternalCommonParameters, SkyAtmosphere)
-		SHADER_PARAMETER_UAV(RWTexture2D<float3>, MultiScatteredLuminanceLutUAV)
-		SHADER_PARAMETER_TEXTURE(Texture2D<float3>, TransmittanceLutTexture)
+		SHADER_PARAMETER_RDG_TEXTURE_UAV(RWTexture2D<float3>, MultiScatteredLuminanceLutUAV)
+		SHADER_PARAMETER_RDG_TEXTURE(Texture2D<float3>, TransmittanceLutTexture)
 		SHADER_PARAMETER_SAMPLER(SamplerState, TransmittanceLutTextureSampler)
 	END_SHADER_PARAMETER_STRUCT()
 
@@ -692,9 +692,9 @@ public:
 	BEGIN_SHADER_PARAMETER_STRUCT(FParameters, )
 		SHADER_PARAMETER_STRUCT_REF(FAtmosphereUniformShaderParameters, Atmosphere)
 		SHADER_PARAMETER_STRUCT_REF(FSkyAtmosphereInternalCommonParameters, SkyAtmosphere)
-		SHADER_PARAMETER_UAV(RWTexture2D<float3>, DistantSkyLightLutUAV)
-		SHADER_PARAMETER_TEXTURE(Texture2D<float3>, TransmittanceLutTexture)
-		SHADER_PARAMETER_TEXTURE(Texture2D<float3>, MultiScatteredLuminanceLutTexture)
+		SHADER_PARAMETER_RDG_TEXTURE_UAV(RWTexture2D<float3>, DistantSkyLightLutUAV)
+		SHADER_PARAMETER_RDG_TEXTURE(Texture2D<float3>, TransmittanceLutTexture)
+		SHADER_PARAMETER_RDG_TEXTURE(Texture2D<float3>, MultiScatteredLuminanceLutTexture)
 		SHADER_PARAMETER_SAMPLER(SamplerState, TransmittanceLutTextureSampler)
 		SHADER_PARAMETER_SAMPLER(SamplerState, MultiScatteredLuminanceLutTextureSampler)
 		SHADER_PARAMETER_SRV(Buffer<float4>, UniformSphereSamplesBuffer)
@@ -735,9 +735,9 @@ public:
 		SHADER_PARAMETER_STRUCT_REF(FAtmosphereUniformShaderParameters, Atmosphere)
 		SHADER_PARAMETER_STRUCT_REF(FSkyAtmosphereInternalCommonParameters, SkyAtmosphere)
 		SHADER_PARAMETER_STRUCT_REF(FViewUniformShaderParameters, ViewUniformBuffer)
-		SHADER_PARAMETER_UAV(RWTexture2D<float3>, SkyViewLutUAV)
-		SHADER_PARAMETER_TEXTURE(Texture2D<float3>, TransmittanceLutTexture)
-		SHADER_PARAMETER_TEXTURE(Texture2D<float3>, MultiScatteredLuminanceLutTexture)
+		SHADER_PARAMETER_RDG_TEXTURE_UAV(RWTexture2D<float3>, SkyViewLutUAV)
+		SHADER_PARAMETER_RDG_TEXTURE(Texture2D<float3>, TransmittanceLutTexture)
+		SHADER_PARAMETER_RDG_TEXTURE(Texture2D<float3>, MultiScatteredLuminanceLutTexture)
 		SHADER_PARAMETER_SAMPLER(SamplerState, TransmittanceLutTextureSampler)
 		SHADER_PARAMETER_SAMPLER(SamplerState, MultiScatteredLuminanceLutTextureSampler)
 	END_SHADER_PARAMETER_STRUCT()
@@ -771,9 +771,9 @@ public:
 		SHADER_PARAMETER_STRUCT_REF(FAtmosphereUniformShaderParameters, Atmosphere)
 		SHADER_PARAMETER_STRUCT_REF(FSkyAtmosphereInternalCommonParameters, SkyAtmosphere)
 		SHADER_PARAMETER_STRUCT_REF(FViewUniformShaderParameters, ViewUniformBuffer)
-		SHADER_PARAMETER_UAV(RWTexture3D<float4>, CameraAerialPerspectiveVolumeUAV)
-		SHADER_PARAMETER_TEXTURE(Texture2D<float3>, TransmittanceLutTexture)
-		SHADER_PARAMETER_TEXTURE(Texture2D<float3>, MultiScatteredLuminanceLutTexture)
+		SHADER_PARAMETER_RDG_TEXTURE_UAV(RWTexture3D<float4>, CameraAerialPerspectiveVolumeUAV)
+		SHADER_PARAMETER_RDG_TEXTURE(Texture2D<float3>, TransmittanceLutTexture)
+		SHADER_PARAMETER_RDG_TEXTURE(Texture2D<float3>, MultiScatteredLuminanceLutTexture)
 		SHADER_PARAMETER_SAMPLER(SamplerState, TransmittanceLutTextureSampler)
 		SHADER_PARAMETER_SAMPLER(SamplerState, MultiScatteredLuminanceLutTextureSampler)
 		SHADER_PARAMETER(float, AerialPerspectiveStartDepth)
@@ -806,8 +806,8 @@ class FRenderDebugSkyAtmospherePS : public FGlobalShader
 		SHADER_PARAMETER_STRUCT_REF(FSkyAtmosphereInternalCommonParameters, SkyAtmosphere)
 		SHADER_PARAMETER_STRUCT_REF(FViewUniformShaderParameters, ViewUniformBuffer)
 		RENDER_TARGET_BINDING_SLOTS()
-		SHADER_PARAMETER_TEXTURE(Texture2D<float3>, TransmittanceLutTexture)
-		SHADER_PARAMETER_TEXTURE(Texture2D<float3>, MultiScatteredLuminanceLutTexture)
+		SHADER_PARAMETER_RDG_TEXTURE(Texture2D<float3>, TransmittanceLutTexture)
+		SHADER_PARAMETER_RDG_TEXTURE(Texture2D<float3>, MultiScatteredLuminanceLutTexture)
 		SHADER_PARAMETER_SAMPLER(SamplerState, TransmittanceLutTextureSampler)
 		SHADER_PARAMETER_SAMPLER(SamplerState, MultiScatteredLuminanceLutTextureSampler)
 		SHADER_PARAMETER(float, ViewPortWidth)
@@ -938,13 +938,13 @@ void InitSkyAtmosphereForScene(FRHICommandListImmediate& RHICmdList, FScene* Sce
 		Desc = FPooledRenderTargetDesc::Create2DDesc(
 			FIntPoint(TransmittanceLutWidth, TransmittanceLutHeight),
 			TranstmittanceLUTUseSmallFormat ? TextureLUTSmallFormat : TextureLUTFormat, FClearValueBinding::None, TexCreate_HideInVisualizeTexture, TexCreate_ShaderResource | TexCreate_UAV, false);
-		GRenderTargetPool.FindFreeElement(RHICmdList, Desc, TransmittanceLutTexture, TEXT("TransmittanceLutTexture"), true, ERenderTargetTransience::NonTransient);
+		GRenderTargetPool.FindFreeElement(RHICmdList, Desc, TransmittanceLutTexture, TEXT("TransmittanceLutTexture"), true, ERenderTargetTransience::Transient);
 
 		TRefCountPtr<IPooledRenderTarget>& MultiScatteredLuminanceLutTexture = SkyInfo.GetMultiScatteredLuminanceLutTexture();
 		Desc = FPooledRenderTargetDesc::Create2DDesc(
 			FIntPoint(MultiScatteredLuminanceLutWidth, MultiScatteredLuminanceLutHeight),
 			MultiScatteringLUTUseSmallFormat ? TextureLUTSmallFormat : TextureLUTFormat, FClearValueBinding::None, TexCreate_HideInVisualizeTexture, TexCreate_ShaderResource | TexCreate_UAV, false);
-		GRenderTargetPool.FindFreeElement(RHICmdList, Desc, MultiScatteredLuminanceLutTexture, TEXT("MultiScatteredLuminanceLutTexture"), true, ERenderTargetTransience::NonTransient);
+		GRenderTargetPool.FindFreeElement(RHICmdList, Desc, MultiScatteredLuminanceLutTexture, TEXT("MultiScatteredLuminanceLutTexture"), true, ERenderTargetTransience::Transient);
 
 		if (CVarSkyAtmosphereDistantSkyLightLUT.GetValueOnRenderThread() > 0)
 		{
@@ -952,7 +952,7 @@ void InitSkyAtmosphereForScene(FRHICommandListImmediate& RHICmdList, FScene* Sce
 			Desc = FPooledRenderTargetDesc::Create2DDesc(
 				FIntPoint(1, 1),
 				TextureLUTFormat, FClearValueBinding::None, TexCreate_HideInVisualizeTexture, TexCreate_ShaderResource | TexCreate_UAV, false);
-			GRenderTargetPool.FindFreeElement(RHICmdList, Desc, DistantSkyLightLutTexture, TEXT("DistantSkyLightLutTexture"), true, ERenderTargetTransience::NonTransient);
+			GRenderTargetPool.FindFreeElement(RHICmdList, Desc, DistantSkyLightLutTexture, TEXT("DistantSkyLightLutTexture"), true, ERenderTargetTransience::Transient);
 		}
 	}
 }
@@ -1084,9 +1084,15 @@ void FSceneRenderer::RenderSkyAtmosphereLookUpTables(FRHICommandListImmediate& R
 	SetupSkyAtmosphereInternalCommonParameters(InternalCommonParameters, *Scene, SkyInfo);
 	TUniformBufferRef<FSkyAtmosphereInternalCommonParameters> InternalCommonParametersRef = TUniformBufferRef<FSkyAtmosphereInternalCommonParameters>::CreateUniformBufferImmediate(InternalCommonParameters, UniformBuffer_SingleFrame);
 
-	UnbindRenderTargets(RHICmdList);
 	FRDGBuilder GraphBuilder(RHICmdList);
 
+	FRDGTextureRef TransmittanceLut = GraphBuilder.RegisterExternalTexture(SkyInfo.GetTransmittanceLutTexture());
+	FRDGTextureUAVRef TransmittanceLutUAV = GraphBuilder.CreateUAV(FRDGTextureUAVDesc(TransmittanceLut, 0));
+	FRDGTextureRef MultiScatteredLuminanceLut = GraphBuilder.RegisterExternalTexture(SkyInfo.GetMultiScatteredLuminanceLutTexture());
+	FRDGTextureUAVRef MultiScatteredLuminanceLutUAV = GraphBuilder.CreateUAV(FRDGTextureUAVDesc(MultiScatteredLuminanceLut, 0));
+	FRDGTextureRef DistantSkyLightLut = GraphBuilder.RegisterExternalTexture(SkyInfo.GetDistantSkyLightLutTexture());
+	FRDGTextureUAVRef DistantSkyLightLutUAV = GraphBuilder.CreateUAV(FRDGTextureUAVDesc(DistantSkyLightLut, 0));
+
 	// Transmittance LUT
 	TShaderMap<FGlobalShaderType>* GlobalShaderMap = GetGlobalShaderMap(FeatureLevel);
 	{
@@ -1095,14 +1101,14 @@ void FSceneRenderer::RenderSkyAtmosphereLookUpTables(FRHICommandListImmediate& R
 		FRenderTransmittanceLutCS::FParameters * PassParameters = GraphBuilder.AllocParameters<FRenderTransmittanceLutCS::FParameters>();
 		PassParameters->Atmosphere = Scene->GetSkyAtmosphereSceneInfo()->GetAtmosphereUniformBuffer();
 		PassParameters->SkyAtmosphere = InternalCommonParametersRef;
-		PassParameters->TransmittanceLutUAV = SkyInfo.GetTransmittanceLutTexture()->GetRenderTargetItem().UAV;
+		PassParameters->TransmittanceLutUAV = TransmittanceLutUAV;
 
-		FIntVector TextureSize = SkyInfo.GetTransmittanceLutTexture()->GetDesc().GetSize();
+		FIntVector TextureSize = TransmittanceLut->Desc.GetSize();
 		TextureSize.Z = 1;
 		const FIntVector NumGroups = FIntVector::DivideAndRoundUp(TextureSize, FRenderTransmittanceLutCS::GroupSize);
 		FComputeShaderUtils::AddPass(GraphBuilder, RDG_EVENT_NAME("TransmittanceLut"), *ComputeShader, PassParameters, NumGroups);
 	}
-
+	
 	// Mean Illuminance LUT
 	{
 		TShaderMapRef<FRenderMultiScatteredLuminanceLutCS> ComputeShader(GlobalShaderMap);
@@ -1111,10 +1117,10 @@ void FSceneRenderer::RenderSkyAtmosphereLookUpTables(FRHICommandListImmediate& R
 		PassParameters->Atmosphere = Scene->GetSkyAtmosphereSceneInfo()->GetAtmosphereUniformBuffer();
 		PassParameters->SkyAtmosphere = InternalCommonParametersRef;
 		PassParameters->TransmittanceLutTextureSampler = SamplerLinearClamp;
-		PassParameters->TransmittanceLutTexture = SkyInfo.GetTransmittanceLutTexture()->GetRenderTargetItem().ShaderResourceTexture;
-		PassParameters->MultiScatteredLuminanceLutUAV = SkyInfo.GetMultiScatteredLuminanceLutTexture()->GetRenderTargetItem().UAV;
+		PassParameters->TransmittanceLutTexture = TransmittanceLut;
+		PassParameters->MultiScatteredLuminanceLutUAV = MultiScatteredLuminanceLutUAV;
 
-		FIntVector TextureSize = SkyInfo.GetMultiScatteredLuminanceLutTexture()->GetDesc().GetSize();
+		FIntVector TextureSize = MultiScatteredLuminanceLut->Desc.GetSize();
 		TextureSize.Z = 1;
 		const FIntVector NumGroups = FIntVector::DivideAndRoundUp(TextureSize, FRenderMultiScatteredLuminanceLutCS::GroupSize);
 		FComputeShaderUtils::AddPass(GraphBuilder, RDG_EVENT_NAME("MeanIllumLut"), *ComputeShader, PassParameters, NumGroups);
@@ -1135,10 +1141,10 @@ void FSceneRenderer::RenderSkyAtmosphereLookUpTables(FRHICommandListImmediate& R
 		PassParameters->SkyAtmosphere = InternalCommonParametersRef;
 		PassParameters->TransmittanceLutTextureSampler = SamplerLinearClamp;
 		PassParameters->MultiScatteredLuminanceLutTextureSampler = SamplerLinearClamp;
-		PassParameters->TransmittanceLutTexture = SkyInfo.GetTransmittanceLutTexture()->GetRenderTargetItem().ShaderResourceTexture;
-		PassParameters->MultiScatteredLuminanceLutTexture = SkyInfo.GetMultiScatteredLuminanceLutTexture()->GetRenderTargetItem().ShaderResourceTexture;
+		PassParameters->TransmittanceLutTexture = TransmittanceLut;
+		PassParameters->MultiScatteredLuminanceLutTexture = MultiScatteredLuminanceLut;
 		PassParameters->UniformSphereSamplesBuffer = GUniformSphereSamplesBuffer.UniformSphereSamplesBuffer.SRV;
-		PassParameters->DistantSkyLightLutUAV = SkyInfo.GetDistantSkyLightLutTexture()->GetRenderTargetItem().UAV;
+		PassParameters->DistantSkyLightLutUAV = DistantSkyLightLutUAV;
 
 		FLightSceneInfo* Light0 = Scene->AtmosphereLights[0];
 		FLightSceneInfo* Light1 = Scene->AtmosphereLights[1];
@@ -1163,10 +1169,12 @@ void FSceneRenderer::RenderSkyAtmosphereLookUpTables(FRHICommandListImmediate& R
 	{
 		FViewInfo& View = Views[ViewIndex];
 		const float AerialPerspectiveStartDepthInCm = GetValidAerialPerspectiveStartDepthInCm(View);
-
 		const bool bLightDiskEnabled = !View.bIsReflectionCapture;
 
-		UnbindRenderTargets(RHICmdList);
+		FRDGTextureRef SkyAtmosphereViewLutTexture = GraphBuilder.RegisterExternalTexture(View.SkyAtmosphereViewLutTexture);
+		FRDGTextureUAVRef SkyAtmosphereViewLutTextureUAV = GraphBuilder.CreateUAV(FRDGTextureUAVDesc(SkyAtmosphereViewLutTexture, 0));
+		FRDGTextureRef SkyAtmosphereCameraAerialPerspectiveVolume = GraphBuilder.RegisterExternalTexture(View.SkyAtmosphereCameraAerialPerspectiveVolume);
+		FRDGTextureUAVRef SkyAtmosphereCameraAerialPerspectiveVolumeUAV = GraphBuilder.CreateUAV(FRDGTextureUAVDesc(SkyAtmosphereCameraAerialPerspectiveVolume, 0));
 
 		// Sky View LUT
 		{
@@ -1182,11 +1190,11 @@ void FSceneRenderer::RenderSkyAtmosphereLookUpTables(FRHICommandListImmediate& R
 			PassParameters->ViewUniformBuffer = View.ViewUniformBuffer;
 			PassParameters->TransmittanceLutTextureSampler = SamplerLinearClamp;
 			PassParameters->MultiScatteredLuminanceLutTextureSampler = SamplerLinearClamp;
-			PassParameters->TransmittanceLutTexture = SkyInfo.GetTransmittanceLutTexture()->GetRenderTargetItem().ShaderResourceTexture;
-			PassParameters->MultiScatteredLuminanceLutTexture = SkyInfo.GetMultiScatteredLuminanceLutTexture()->GetRenderTargetItem().ShaderResourceTexture;
-			PassParameters->SkyViewLutUAV = View.SkyAtmosphereViewLutTexture->GetRenderTargetItem().UAV;
+			PassParameters->TransmittanceLutTexture = TransmittanceLut;
+			PassParameters->MultiScatteredLuminanceLutTexture = MultiScatteredLuminanceLut;
+			PassParameters->SkyViewLutUAV = SkyAtmosphereViewLutTextureUAV;
 
-			FIntVector TextureSize = View.SkyAtmosphereViewLutTexture->GetDesc().GetSize();
+			FIntVector TextureSize = SkyAtmosphereViewLutTexture->Desc.GetSize();
 			TextureSize.Z = 1;
 			const FIntVector NumGroups = FIntVector::DivideAndRoundUp(TextureSize, FRenderSkyViewLutCS::GroupSize);
 			FComputeShaderUtils::AddPass(GraphBuilder, RDG_EVENT_NAME("SkyViewLut"), *ComputeShader, PassParameters, NumGroups);
@@ -1205,17 +1213,17 @@ void FSceneRenderer::RenderSkyAtmosphereLookUpTables(FRHICommandListImmediate& R
 			PassParameters->ViewUniformBuffer = View.ViewUniformBuffer;
 			PassParameters->TransmittanceLutTextureSampler = SamplerLinearClamp;
 			PassParameters->MultiScatteredLuminanceLutTextureSampler = SamplerLinearClamp;
-			PassParameters->TransmittanceLutTexture = SkyInfo.GetTransmittanceLutTexture()->GetRenderTargetItem().ShaderResourceTexture;
-			PassParameters->MultiScatteredLuminanceLutTexture = SkyInfo.GetMultiScatteredLuminanceLutTexture()->GetRenderTargetItem().ShaderResourceTexture;
-			PassParameters->CameraAerialPerspectiveVolumeUAV = View.SkyAtmosphereCameraAerialPerspectiveVolume->GetRenderTargetItem().UAV;
+			PassParameters->TransmittanceLutTexture = TransmittanceLut;
+			PassParameters->MultiScatteredLuminanceLutTexture = MultiScatteredLuminanceLut;
+			PassParameters->CameraAerialPerspectiveVolumeUAV = SkyAtmosphereCameraAerialPerspectiveVolumeUAV;
 			PassParameters->AerialPerspectiveStartDepth = AerialPerspectiveStartDepthInCm * CM_TO_KM;
 
-			FIntVector TextureSize = View.SkyAtmosphereCameraAerialPerspectiveVolume->GetDesc().GetSize();
+			FIntVector TextureSize = SkyAtmosphereCameraAerialPerspectiveVolume->Desc.GetSize();
 			const FIntVector NumGroups = FIntVector::DivideAndRoundUp(TextureSize, FRenderCameraAerialPerspectiveVolumeCS::GroupSize);
 			FComputeShaderUtils::AddPass(GraphBuilder, RDG_EVENT_NAME("CameraVolumeLut"), *ComputeShader, PassParameters, NumGroups);
 		}
 	}
-
+	
 	GraphBuilder.Execute();
 	// TODO have RDG execute those above passes with compute overlap similarly to using AutomaticCacheFlushAfterComputeShader(true);
 }
@@ -1252,6 +1260,9 @@ void FSceneRenderer::RenderSkyAtmosphere(FRHICommandListImmediate& RHICmdList)
 	FRDGTextureRef SceneColor = GraphBuilder.RegisterExternalTexture(SceneContext.GetSceneColor(), TEXT("SceneColor"));
 	FRDGTextureRef SceneDepth = GraphBuilder.RegisterExternalTexture(SceneContext.SceneDepthZ, TEXT("SceneDepth"));
 
+	FRDGTextureRef TransmittanceLut = GraphBuilder.RegisterExternalTexture(SkyInfo.GetTransmittanceLutTexture());
+	FRDGTextureRef MultiScatteredLuminanceLut = GraphBuilder.RegisterExternalTexture(SkyInfo.GetMultiScatteredLuminanceLutTexture());
+
 	for (int32 ViewIndex = 0; ViewIndex < Views.Num(); ViewIndex++)
 	{
 		FViewInfo& View = Views[ViewIndex];
@@ -1270,6 +1281,9 @@ void FSceneRenderer::RenderSkyAtmosphere(FRHICommandListImmediate& RHICmdList)
 		const float SafeEdge = 1000.0f;	// 10 meters
 		const bool ForceRayMarching = (FVector::Distance(ViewOrigin, PlanetOrigin) - TopOfAtmosphere - SafeEdge) > 0.0f;
 
+		FRDGTextureRef SkyAtmosphereViewLutTexture = GraphBuilder.RegisterExternalTexture(View.SkyAtmosphereViewLutTexture);
+		FRDGTextureRef SkyAtmosphereCameraAerialPerspectiveVolume = GraphBuilder.RegisterExternalTexture(View.SkyAtmosphereCameraAerialPerspectiveVolume);
+
 		// Render the sky, and optionally the atmosphere aerial perspective, on the scene luminance buffer
 		{
 			FRenderSkyAtmospherePS::FPermutationDomain PsPermutationVector;
@@ -1300,10 +1314,10 @@ void FSceneRenderer::RenderSkyAtmosphere(FRHICommandListImmediate& RHICmdList)
 			PsPassParameters->MultiScatteredLuminanceLutTextureSampler = SamplerLinearClamp;
 			PsPassParameters->SkyViewLutTextureSampler = SamplerLinearClamp;
 			PsPassParameters->CameraAerialPerspectiveVolumeTextureSampler = SamplerLinearClamp;
-			PsPassParameters->TransmittanceLutTexture = SkyInfo.GetTransmittanceLutTexture()->GetRenderTargetItem().ShaderResourceTexture;
-			PsPassParameters->MultiScatteredLuminanceLutTexture = SkyInfo.GetMultiScatteredLuminanceLutTexture()->GetRenderTargetItem().ShaderResourceTexture;
-			PsPassParameters->SkyViewLutTexture = View.SkyAtmosphereViewLutTexture->GetRenderTargetItem().ShaderResourceTexture;
-			PsPassParameters->CameraAerialPerspectiveVolumeTexture = View.SkyAtmosphereCameraAerialPerspectiveVolume->GetRenderTargetItem().ShaderResourceTexture;
+			PsPassParameters->TransmittanceLutTexture = TransmittanceLut;
+			PsPassParameters->MultiScatteredLuminanceLutTexture = MultiScatteredLuminanceLut;
+			PsPassParameters->SkyViewLutTexture = SkyAtmosphereViewLutTexture;
+			PsPassParameters->CameraAerialPerspectiveVolumeTexture = SkyAtmosphereCameraAerialPerspectiveVolume;
 			PsPassParameters->AerialPerspectiveStartDepth = AerialPerspectiveStartDepthInCm * CM_TO_KM;
 			ClearUnusedGraphResources(*PixelShader, PsPassParameters);
 
@@ -1490,6 +1504,9 @@ void FDeferredShadingSceneRenderer::RenderDebugSkyAtmosphere(FRHICommandListImme
 		FRDGTextureRef SceneColor = GraphBuilder.RegisterExternalTexture(SceneContext.GetSceneColor(), TEXT("SceneColor"));
 		FRDGTextureRef SceneDepth = GraphBuilder.RegisterExternalTexture(SceneContext.SceneDepthZ, TEXT("SceneDepth"));
 
+		FRDGTextureRef TransmittanceLut = GraphBuilder.RegisterExternalTexture(SkyInfo.GetTransmittanceLutTexture());
+		FRDGTextureRef MultiScatteredLuminanceLut = GraphBuilder.RegisterExternalTexture(SkyInfo.GetMultiScatteredLuminanceLutTexture());
+
 		FRHIBlendState* PreMultipliedColorTransmittanceBlend = TStaticBlendState<CW_RGB, BO_Add, BF_One, BF_SourceAlpha, BO_Add, BF_Zero, BF_One>::GetRHI();
 		FRHIDepthStencilState* DepthStencilStateWrite = TStaticDepthStencilState<true, CF_Always>::GetRHI();
 		FRHISamplerState* SamplerLinearClamp = TStaticSamplerState<SF_Trilinear>::GetRHI();
@@ -1518,8 +1535,8 @@ void FDeferredShadingSceneRenderer::RenderDebugSkyAtmosphere(FRHICommandListImme
 					ERenderTargetLoadAction::ENoAction, ERenderTargetStoreAction::ENoAction, FExclusiveDepthStencil::DepthWrite_StencilNop);
 				PassParameters->TransmittanceLutTextureSampler = SamplerLinearClamp;
 				PassParameters->MultiScatteredLuminanceLutTextureSampler = SamplerLinearClamp;
-				PassParameters->TransmittanceLutTexture = SkyInfo.GetTransmittanceLutTexture()->GetRenderTargetItem().ShaderResourceTexture;
-				PassParameters->MultiScatteredLuminanceLutTexture = SkyInfo.GetMultiScatteredLuminanceLutTexture()->GetRenderTargetItem().ShaderResourceTexture;
+				PassParameters->TransmittanceLutTexture = TransmittanceLut;
+				PassParameters->MultiScatteredLuminanceLutTexture = MultiScatteredLuminanceLut;
 				PassParameters->ViewPortWidth = float(View.ViewRect.Width());
 				PassParameters->ViewPortHeight = float(View.ViewRect.Height());
 
