Fix for non unity build.

diff --git a/Engine/Source/Runtime/Renderer/Private/SkyAtmosphereRendering.h b/Engine/Source/Runtime/Renderer/Private/SkyAtmosphereRendering.h
--- a/Engine/Source/Runtime/Renderer/Private/SkyAtmosphereRendering.h
+++ b/Engine/Source/Runtime/Renderer/Private/SkyAtmosphereRendering.h
@@ -13,6 +13,7 @@
 #include "Rendering/SkyAtmosphereCommonData.h"
 
 
+class FScene;
 class FLightSceneInfo;
 class USkyAtmosphereComponent;
 
