Submit cmd-buffer with offscreen rendering before doing post-processing on Mobile

diff --git a/Engine/Source/Runtime/Renderer/Private/MobileShadingRenderer.cpp b/Engine/Source/Runtime/Renderer/Private/MobileShadingRenderer.cpp
--- a/Engine/Source/Runtime/Renderer/Private/MobileShadingRenderer.cpp
+++ b/Engine/Source/Runtime/Renderer/Private/MobileShadingRenderer.cpp
@@ -65,13 +65,6 @@ static TAutoConsoleVariable<int32> CVarMobileForceDepthResolve(
 	TEXT("1: Depth buffer is resolved by switching out render targets and drawing with the depth texture.\n"),
 	ECVF_Scalability | ECVF_RenderThreadSafe);
 
-static TAutoConsoleVariable<int32> CVarMobileMoveSubmissionHintAfterTranslucency(
-	TEXT("r.Mobile.MoveSubmissionHintAfterTranslucency"),
-	1,
-	TEXT("0: Submission hint occurs after occlusion query.\n")
-	TEXT("1: Submission hint occurs after translucency. (Default)"),
-	ECVF_Scalability | ECVF_RenderThreadSafe);
-
 static TAutoConsoleVariable<int32> CVarMobileAdrenoOcclusionMode(
 	TEXT("r.Mobile.AdrenoOcclusionMode"),
 	0,
@@ -447,6 +440,9 @@ void FMobileSceneRenderer::Render(FRHICommandListImmediate& RHICmdList)
 	bool bKeepDepthContent = bRenderToSceneColor && 
 		(bForceDepthResolve || bSeparateTranslucencyActive || (View.bIsSceneCapture && (ViewFamily.SceneCaptureSource == ESceneCaptureSource::SCS_SceneColorHDR || ViewFamily.SceneCaptureSource == ESceneCaptureSource::SCS_SceneColorSceneDepth)));
 
+	// Whether to submit cmdbuffer with offscreen rendering before doing post-processing
+	bool bSubmitOffscreenRendering = !bGammaSpace || bRenderToSceneColor;
+
 	//
 	FRHITexture* SceneColor = nullptr;
 	FRHITexture* SceneColorResolve = nullptr;
@@ -628,6 +624,7 @@ void FMobileSceneRenderer::Render(FRHICommandListImmediate& RHICmdList)
 	    RHICmdList.SetCurrentStat(GET_STATID(STAT_CLMM_Occlusion));
 		// flush
 		RHICmdList.SubmitCommandsHint();
+		bSubmitOffscreenRendering = false; // submit once
 		// Issue occlusion queries
 	    RenderOcclusion(RHICmdList);
 	    RHICmdList.ImmediateFlush(EImmediateFlushType::DispatchToRHIThread);
@@ -642,6 +639,12 @@ void FMobileSceneRenderer::Render(FRHICommandListImmediate& RHICmdList)
 	// End of scene color rendering
 	RHICmdList.EndRenderPass();
 
+	// Flush / submit cmdbuffer
+	if (bSubmitOffscreenRendering)
+	{
+		RHICmdList.SubmitCommandsHint();
+	}
+	
 	if (!bGammaSpace || bRenderToSceneColor)
 	{
 		// transition scene color to Readable for post-processing
@@ -788,13 +791,6 @@ void FMobileSceneRenderer::RenderOcclusion(FRHICommandListImmediate& RHICmdList)
 
 	BeginOcclusionTests(RHICmdList, true);
 	FenceOcclusionTests(RHICmdList);
-
-	// Optionally hint submission later to avoid render pass churn but delay query results
-	const bool bSubmissionAfterTranslucency = (CVarMobileMoveSubmissionHintAfterTranslucency.GetValueOnRenderThread() == 1);
-	if (!bSubmissionAfterTranslucency)
-	{	
-		RHICmdList.SubmitCommandsHint();
-	}
 }
 
 int32 FMobileSceneRenderer::ComputeNumOcclusionQueriesToBatch() const
