Flush RHIT when submiting cmdbuffer

diff --git a/Engine/Source/Runtime/Renderer/Private/MobileShadingRenderer.cpp b/Engine/Source/Runtime/Renderer/Private/MobileShadingRenderer.cpp
--- a/Engine/Source/Runtime/Renderer/Private/MobileShadingRenderer.cpp
+++ b/Engine/Source/Runtime/Renderer/Private/MobileShadingRenderer.cpp
@@ -643,6 +643,7 @@ void FMobileSceneRenderer::Render(FRHICommandListImmediate& RHICmdList)
 	if (bSubmitOffscreenRendering)
 	{
 		RHICmdList.SubmitCommandsHint();
+		RHICmdList.ImmediateFlush(EImmediateFlushType::DispatchToRHIThread);
 	}
 	
 	if (!bGammaSpace || bRenderToSceneColor)
