Fixed fog on transparent broken by 8239614.

diff --git a/Engine/Source/Runtime/Renderer/Private/SkyAtmosphereRendering.cpp b/Engine/Source/Runtime/Renderer/Private/SkyAtmosphereRendering.cpp
--- a/Engine/Source/Runtime/Renderer/Private/SkyAtmosphereRendering.cpp
+++ b/Engine/Source/Runtime/Renderer/Private/SkyAtmosphereRendering.cpp
@@ -300,7 +300,7 @@ void SetupSkyAtmosphereViewSharedUniformShaderParameters(const FViewInfo& View,
 	OutParameters.CameraAerialPerspectiveVolumeDepthSliceLength = CameraAerialPerspectiveVolumeDepthSliceLength;
 	OutParameters.CameraAerialPerspectiveVolumeDepthSliceLengthInv = 1.0f / OutParameters.CameraAerialPerspectiveVolumeDepthSliceLength;
 
-	OutParameters.AerialPerspectiveStartDepth = GetValidAerialPerspectiveStartDepthInCm(View);
+	OutParameters.AerialPerspectiveStartDepth = GetValidAerialPerspectiveStartDepthInCm(View) * CM_TO_KM;
 
 	SetBlackAlpha13DIfNull(SkyAtmosphereCameraAerialPerspectiveVolume); // Needs to be after we set ApplyCameraAerialPerspectiveVolume
 }
