Cherry-pick 8349661: Texture samples applied to the base color input of clear coat materials in ray traced reflections - Packing proper data into custom data field fixed the issue.

diff --git a/Engine/Shaders/Private/RayTracing/RayTracingCommon.ush b/Engine/Shaders/Private/RayTracing/RayTracingCommon.ush
--- a/Engine/Shaders/Private/RayTracing/RayTracingCommon.ush
+++ b/Engine/Shaders/Private/RayTracing/RayTracingCommon.ush
@@ -228,7 +228,7 @@ FPackedMaterialClosestHitPayload PackRayTracingPayload(FMaterialClosestHitPayloa
 	Output.IndirectIrradiance[0] = f32tof16(Input.IndirectIrradiance.x);
 	Output.IndirectIrradiance[0] |= f32tof16(Input.IndirectIrradiance.y) << 16;
 	Output.IndirectIrradiance[1] = f32tof16(Input.IndirectIrradiance.z);
-	Output.CustomData = f32tof16(Input.CustomData.x) | (f32tof16(Input.CustomData.y) << 16);
+	Output.CustomData = f32tof16(Input.CustomData.z) | (f32tof16(Input.CustomData.w) << 16);
 
 	return Output;
 }
@@ -267,9 +267,9 @@ FMaterialClosestHitPayload UnpackRayTracingPayload(FPackedMaterialClosestHitPayl
 	Output.IndirectIrradiance.x	= f16tof32(Input.IndirectIrradiance[0]);
 	Output.IndirectIrradiance.y	= f16tof32(Input.IndirectIrradiance[0] >> 16);
 	Output.IndirectIrradiance.z	= f16tof32(Input.IndirectIrradiance[1]);
-	Output.CustomData.x = f16tof32(Input.CustomData.x & 0xFFFF);
-	Output.CustomData.y = f16tof32(Input.CustomData.x >> 16);
-	Output.CustomData.zw = 0;
+	Output.CustomData.z = f16tof32(Input.CustomData.x & 0xFFFF);
+	Output.CustomData.w = f16tof32(Input.CustomData.x >> 16);
+	Output.CustomData.xy = 0;
 
 	Output.DiffuseColor  = Output.BaseColor - Output.BaseColor * Output.Metallic;
 	Output.SpecularColor = ComputeF0(Output.Specular, Output.BaseColor, Output.Metallic);
diff --git a/Engine/Shaders/Private/RayTracing/RayTracingLightingCommon.ush b/Engine/Shaders/Private/RayTracing/RayTracingLightingCommon.ush
--- a/Engine/Shaders/Private/RayTracing/RayTracingLightingCommon.ush
+++ b/Engine/Shaders/Private/RayTracing/RayTracingLightingCommon.ush
@@ -323,7 +323,7 @@ float3 ComputeDirectLightingCulled(
 	GBufferData.CustomData = Payload.CustomData;
 	GBufferData.GBufferAO = Payload.GBufferAO;
 	GBufferData.IndirectIrradiance = (Payload.IndirectIrradiance.x + Payload.IndirectIrradiance.y + Payload.IndirectIrradiance.z) / 3.f; 
-	GBufferData.ShadingModelID = SHADINGMODELID_DEFAULT_LIT;
+	GBufferData.ShadingModelID = Payload.ShadingModelID;
 	GBufferData.SpecularColor = Payload.SpecularColor;
 	GBufferData.DiffuseColor = Payload.DiffuseColor;
 
diff --git a/Engine/Source/Runtime/Renderer/Private/RayTracing/RayTracingReflections.cpp b/Engine/Source/Runtime/Renderer/Private/RayTracing/RayTracingReflections.cpp
--- a/Engine/Source/Runtime/Renderer/Private/RayTracing/RayTracingReflections.cpp
+++ b/Engine/Source/Runtime/Renderer/Private/RayTracing/RayTracingReflections.cpp
@@ -164,6 +164,7 @@ class FRayTracingReflectionsRGS : public FGlobalShader
 		SHADER_PARAMETER(int32, MaxUnderCoatBounces)
 
 		SHADER_PARAMETER_SRV(RaytracingAccelerationStructure, TLAS)
+		SHADER_PARAMETER_RDG_TEXTURE(Texture2D, SSProfilesTexture)
 		SHADER_PARAMETER_SRV(StructuredBuffer<FRTLightingData>, LightDataBuffer)
 
 		SHADER_PARAMETER_STRUCT_INCLUDE(FSceneTextureParameters, SceneTextures)
@@ -295,6 +296,14 @@ void FDeferredShadingSceneRenderer::RenderRayTracingReflections(
 	CommonParameters.SceneTextures = SceneTextures;
 	SetupSceneTextureSamplers(&CommonParameters.SceneTextureSamplers);
 
+	// TODO: should be converted to RDG
+	TRefCountPtr<IPooledRenderTarget> SubsurfaceProfileRT((IPooledRenderTarget*)GetSubsufaceProfileTexture_RT(GraphBuilder.RHICmdList));
+	if (!SubsurfaceProfileRT)
+	{
+		SubsurfaceProfileRT = GSystemTextures.BlackDummy;
+	}
+	CommonParameters.SSProfilesTexture = GraphBuilder.RegisterExternalTexture(SubsurfaceProfileRT);
+
 	CommonParameters.ReflectionStruct = CreateReflectionUniformBuffer(View, EUniformBufferUsage::UniformBuffer_SingleFrame);
 	CommonParameters.FogUniformParameters = CreateFogUniformBuffer(View, EUniformBufferUsage::UniformBuffer_SingleFrame);
 	CommonParameters.ColorOutput = GraphBuilder.CreateUAV(OutDenoiserInputs->Color);
diff --git a/Engine/Source/Runtime/Renderer/Private/RayTracing/RayTracingTranslucency.cpp b/Engine/Source/Runtime/Renderer/Private/RayTracing/RayTracingTranslucency.cpp
--- a/Engine/Source/Runtime/Renderer/Private/RayTracing/RayTracingTranslucency.cpp
+++ b/Engine/Source/Runtime/Renderer/Private/RayTracing/RayTracingTranslucency.cpp
@@ -150,6 +150,7 @@ class FRayTracingTranslucencyRGS : public FGlobalShader
 
 		SHADER_PARAMETER_SRV(RaytracingAccelerationStructure, TLAS)
 		SHADER_PARAMETER_SRV(StructuredBuffer<FRTLightingData>, LightDataBuffer)
+		SHADER_PARAMETER_RDG_TEXTURE(Texture2D, SSProfilesTexture)
 
 		SHADER_PARAMETER_STRUCT_REF(FViewUniformShaderParameters, ViewUniformBuffer)
 		SHADER_PARAMETER_STRUCT_REF(FRaytracingLightDataPacked, LightDataPacked)
@@ -411,6 +412,14 @@ void FDeferredShadingSceneRenderer::RenderRayTracingTranslucencyView(
 	PassParameters->ColorOutput = GraphBuilder.CreateUAV(*OutColorTexture);
 	PassParameters->RayHitDistanceOutput = GraphBuilder.CreateUAV(*OutRayHitDistanceTexture);
 
+	// TODO: should be converted to RDG
+	TRefCountPtr<IPooledRenderTarget> SubsurfaceProfileRT((IPooledRenderTarget*)GetSubsufaceProfileTexture_RT(GraphBuilder.RHICmdList));
+	if (!SubsurfaceProfileRT)
+	{
+		SubsurfaceProfileRT = GSystemTextures.BlackDummy;
+	}
+	PassParameters->SSProfilesTexture = GraphBuilder.RegisterExternalTexture(SubsurfaceProfileRT);
+
 	auto RayGenShader = View.ShaderMap->GetShader<FRayTracingTranslucencyRGS>();
 	ClearUnusedGraphResources(RayGenShader, PassParameters);
 
