Only show the "under the ground" message when r.SkyAtmosphere.Visualize is 1.

diff --git a/Engine/Source/Runtime/Renderer/Private/SkyAtmosphereRendering.cpp b/Engine/Source/Runtime/Renderer/Private/SkyAtmosphereRendering.cpp
--- a/Engine/Source/Runtime/Renderer/Private/SkyAtmosphereRendering.cpp
+++ b/Engine/Source/Runtime/Renderer/Private/SkyAtmosphereRendering.cpp
@@ -1562,16 +1562,16 @@ void FDeferredShadingSceneRenderer::RenderDebugSkyAtmosphere(FRHICommandListImme
 		FLinearColor WarningColor(1.0f, 0.5f, 0.0f);
 		FString Text;
 
-		const float ViewPlanetAltitude = (View.ViewLocation*0.00001f - FVector(0.0f, 0.0f, -Atmosphere.BottomRadius)).Size() - Atmosphere.BottomRadius;
-		const bool ViewUnderGroundLevel = ViewPlanetAltitude < 0.0f;
-		if (ViewUnderGroundLevel)
-		{
-			Text = FString::Printf(TEXT("SkyAtmosphere: View is %.3f km under the planet ground level!"), -ViewPlanetAltitude);
-			Canvas.DrawShadowedString(ViewPortWidth*0.5 - 250.0f, ViewPortHeight*0.5f, *Text, GetStatsFont(), WarningColor);
-		}
-
 		if (SkyAtmosphereVisualize)
 		{
+			const float ViewPlanetAltitude = (View.ViewLocation*0.00001f - FVector(0.0f, 0.0f, -Atmosphere.BottomRadius)).Size() - Atmosphere.BottomRadius;
+			const bool bViewUnderGroundLevel = ViewPlanetAltitude < 0.0f;
+			if (bViewUnderGroundLevel)
+			{
+				Text = FString::Printf(TEXT("SkyAtmosphere: View is %.3f km under the planet ground level!"), -ViewPlanetAltitude);
+				Canvas.DrawShadowedString(ViewPortWidth*0.5 - 250.0f, ViewPortHeight*0.5f, *Text, GetStatsFont(), WarningColor);
+			}
+
 			// This needs to stay in sync with RenderSkyAtmosphereDebugPS.
 			const float DensityViewTop = ViewPortHeight * 0.1f;
 			const float DensityViewBottom = ViewPortHeight * 0.8f;
@@ -1602,10 +1602,7 @@ void FDeferredShadingSceneRenderer::RenderDebugSkyAtmosphere(FRHICommandListImme
 
 			Text = FString::Printf(TEXT("Hemisphere view"));
 			Canvas.DrawShadowedString(Margin, HemiViewTop, *Text, GetStatsFont(), TextColor);
-		}
 
-		if (ViewUnderGroundLevel || SkyAtmosphereVisualize)
-		{
 			Canvas.Flush_RenderThread(RHICmdList);
 		}
 	}
