Take into account SkyLuminanceFactor in the material node graph.

diff --git a/Engine/Shaders/Private/MaterialTemplate.ush b/Engine/Shaders/Private/MaterialTemplate.ush
--- a/Engine/Shaders/Private/MaterialTemplate.ush
+++ b/Engine/Shaders/Private/MaterialTemplate.ush
@@ -1032,6 +1032,7 @@ float3 MaterialExpressionSkyAtmosphereViewLuminance(FMaterialPixelParameters Par
 	float2 SkyViewLutUv;
 	SkyViewLutParamsToUv(IntersectGround, ViewZenithCosAngle, WorldDirLocal, ViewHeight, ResolvedView.SkyAtmosphereBottomRadius, ResolvedView.SkyViewLutSizeAndInvSize, SkyViewLutUv);
 	float3 SkyAtmosphereViewLuminance = Texture2DSampleLevel(View.SkyViewLutTexture, View.SkyViewLutTextureSampler, SkyViewLutUv, 0.0f).rgb;
+	SkyAtmosphereViewLuminance *= ResolvedView.SkyAtmosphereSkyLuminanceFactor;
 #if USE_PREEXPOSURE
 	SkyAtmosphereViewLuminance *= ResolvedView.OneOverPreExposure;
 #endif
diff --git a/Engine/Shaders/Private/SkyAtmosphere.usf b/Engine/Shaders/Private/SkyAtmosphere.usf
--- a/Engine/Shaders/Private/SkyAtmosphere.usf
+++ b/Engine/Shaders/Private/SkyAtmosphere.usf
@@ -1099,7 +1099,7 @@ void RenderDistantSkyLightLutCS(uint3 ThreadId : SV_DispatchThreadID)
 		AtmosphereLightDirection0.xyz, AtmosphereLightDirection1.xyz, AtmosphereLightColor0.rgb, AtmosphereLightColor1.rgb,
 		AerialPespectiveViewDistanceScale);
 
-	GroupSkyLuminanceSamples[LinearIndex] = ss.L;
+	GroupSkyLuminanceSamples[LinearIndex] = ss.L * SkyAtmosphere.SkyLuminanceFactor;
 
 	// Wait for all group threads to be done
 	GroupMemoryBarrierWithGroupSync();
diff --git a/Engine/Source/Runtime/Engine/Public/SceneView.h b/Engine/Source/Runtime/Engine/Public/SceneView.h
--- a/Engine/Source/Runtime/Engine/Public/SceneView.h
+++ b/Engine/Source/Runtime/Engine/Public/SceneView.h
@@ -672,6 +672,7 @@ enum ETranslucencyVolumeCascade
 	VIEW_UNIFORM_BUFFER_MEMBER_ARRAY(FVector4, AtmosphereLightDiscCosHalfApexAngle, [NUM_ATMOSPHERE_LIGHTS]) \
 	VIEW_UNIFORM_BUFFER_MEMBER(FVector4, SkyViewLutSizeAndInvSize) \
 	VIEW_UNIFORM_BUFFER_MEMBER(FVector, SkyWorldCameraOrigin) \
+	VIEW_UNIFORM_BUFFER_MEMBER(FLinearColor, SkyAtmosphereSkyLuminanceFactor) \
 	VIEW_UNIFORM_BUFFER_MEMBER(float, SkyAtmosphereBottomRadius) \
 	VIEW_UNIFORM_BUFFER_MEMBER(float, SkyAtmosphereTopRadius) \
 	VIEW_UNIFORM_BUFFER_MEMBER(float, SkyAtmosphereAerialPerspectiveStartDepth) \
diff --git a/Engine/Source/Runtime/Renderer/Private/SceneRendering.cpp b/Engine/Source/Runtime/Renderer/Private/SceneRendering.cpp
--- a/Engine/Source/Runtime/Renderer/Private/SceneRendering.cpp
+++ b/Engine/Source/Runtime/Renderer/Private/SceneRendering.cpp
@@ -1233,8 +1233,8 @@ void FViewInfo::SetupUniformBufferParameters(
 		ViewUniformShaderParameters.SkyViewLutSizeAndInvSize = FVector4(SkyViewLutWidth, SkyViewLutHeight, 1.0f / SkyViewLutWidth, 1.0f / SkyViewLutHeight);
 
 		CameraAerialPerspectiveVolumeFound = this->SkyAtmosphereCameraAerialPerspectiveVolume->GetRenderTargetItem().ShaderResourceTexture;
-
 		DistantSkyLightLutTextureFound = SkyAtmosphere->GetDistantSkyLightLutTextureRHI();
+		ViewUniformShaderParameters.SkyAtmosphereSkyLuminanceFactor = SkyAtmosphere->GetSkyLuminanceFactor();
 
 		const FAtmosphereSetup& AtmosphereSetup =  SkyAtmosphere->GetAtmosphereSetup();
 		ViewUniformShaderParameters.SkyAtmosphereBottomRadius = AtmosphereSetup.BottomRadius;
@@ -1288,6 +1288,7 @@ void FViewInfo::SetupUniformBufferParameters(
 		ViewUniformShaderParameters.SkyViewLutSizeAndInvSize = FVector4(1.0f, 1.0f, 1.0f, 1.0f);
 		ViewUniformShaderParameters.SkyAtmosphereBottomRadius = 1.0f;
 		ViewUniformShaderParameters.SkyAtmosphereTopRadius = 1.0f;
+		ViewUniformShaderParameters.SkyAtmosphereSkyLuminanceFactor = FLinearColor::White;
 		ViewUniformShaderParameters.SkyAtmosphereAerialPerspectiveStartDepth = 1.0f;
 		ViewUniformShaderParameters.SkyAtmosphereCameraAerialPerspectiveVolumeDepthResolution = 1.0f;
 		ViewUniformShaderParameters.SkyAtmosphereCameraAerialPerspectiveVolumeDepthResolutionInv = 1.0f;
