Do sky compute work before any graphics on mobile

diff --git a/Engine/Source/Runtime/Renderer/Private/MobileShadingRenderer.cpp b/Engine/Source/Runtime/Renderer/Private/MobileShadingRenderer.cpp
--- a/Engine/Source/Runtime/Renderer/Private/MobileShadingRenderer.cpp
+++ b/Engine/Source/Runtime/Renderer/Private/MobileShadingRenderer.cpp
@@ -371,6 +371,14 @@ void FMobileSceneRenderer::Render(FRHICommandListImmediate& RHICmdList)
 	DynamicReadBuffer.Commit();
 	RHICmdList.ImmediateFlush(EImmediateFlushType::DispatchToRHIThread);
 
+	// Generate the Sky/Atmosphere look up tables
+	// Do compute work first, before any graphics work
+	const bool bShouldRenderSkyAtmosphere = ShouldRenderSkyAtmosphere(Scene);
+	if (bShouldRenderSkyAtmosphere)
+	{
+		RenderSkyAtmosphereLookUpTables(RHICmdList);
+	}
+
 	// Notify the FX system that the scene is about to be rendered.
 	if (Scene->FXSystem && ViewFamily.EngineShowFlags.Particles)
 	{
@@ -397,13 +405,6 @@ void FMobileSceneRenderer::Render(FRHICommandListImmediate& RHICmdList)
 	}
 
 	const bool bGammaSpace = !IsMobileHDR();
-
-	// Generate the Sky/Atmosphere look up tables
-	const bool bShouldRenderSkyAtmosphere = ShouldRenderSkyAtmosphere(Scene);
-	if (bShouldRenderSkyAtmosphere)
-	{
-		RenderSkyAtmosphereLookUpTables(RHICmdList);
-	}
 	
 	// Custom depth
 	if (!bGammaSpace)
