Check that we have a valid GPU feedback buffer ready to map before flushing runtime virtual texture caches.
This prevents reported bug when flushing at low frame rates (which can happen when sculpting many layer landscape).
If we flush without a feedback buffer, then we will have to wait for at least one frame before refilling the cache and this can cause the landscape to look constantly low res.

diff --git a/Engine/Source/Runtime/Renderer/Private/DeferredShadingRenderer.cpp b/Engine/Source/Runtime/Renderer/Private/DeferredShadingRenderer.cpp
--- a/Engine/Source/Runtime/Renderer/Private/DeferredShadingRenderer.cpp
+++ b/Engine/Source/Runtime/Renderer/Private/DeferredShadingRenderer.cpp
@@ -979,8 +979,7 @@ void FDeferredShadingSceneRenderer::Render(FRHICommandListImmediate& RHICmdList)
 
 	if (bUseVirtualTexturing)
 	{
-		Scene->FlushDirtyRuntimeVirtualTextures();
-		FVirtualTextureSystem::Get().Update(RHICmdList, FeatureLevel);
+		FVirtualTextureSystem::Get().Update(RHICmdList, FeatureLevel, Scene);
 	}
 
 	for (int32 ViewIndex = 0; ViewIndex < Views.Num(); ViewIndex++)
diff --git a/Engine/Source/Runtime/Renderer/Private/VT/VirtualTextureFeedback.cpp b/Engine/Source/Runtime/Renderer/Private/VT/VirtualTextureFeedback.cpp
--- a/Engine/Source/Runtime/Renderer/Private/VT/VirtualTextureFeedback.cpp
+++ b/Engine/Source/Runtime/Renderer/Private/VT/VirtualTextureFeedback.cpp
@@ -99,7 +99,13 @@ void FVirtualTextureFeedback::TransferGPUToCPU( FRHICommandListImmediate& RHICmd
 	++PendingTargetCount;
 }
 
-bool FVirtualTextureFeedback::Map( FRHICommandListImmediate& RHICmdList, MapResult& OutResult )
+bool FVirtualTextureFeedback::CanMap()
+{
+	const FeedBackItem& FeedbackEntryCPU = FeedbackTextureCPU[CPUReadIndex];
+	return (PendingTargetCount > 0u && FeedbackEntryCPU.TextureCPU.IsValid() && FeedbackEntryCPU.GPUFenceRHI->Poll());
+}
+
+bool FVirtualTextureFeedback::Map(FRHICommandListImmediate& RHICmdList, MapResult& OutResult)
 {
 	const FeedBackItem& FeedbackEntryCPU = FeedbackTextureCPU[CPUReadIndex];
 	if (PendingTargetCount > 0u &&
diff --git a/Engine/Source/Runtime/Renderer/Private/VT/VirtualTextureFeedback.h b/Engine/Source/Runtime/Renderer/Private/VT/VirtualTextureFeedback.h
--- a/Engine/Source/Runtime/Renderer/Private/VT/VirtualTextureFeedback.h
+++ b/Engine/Source/Runtime/Renderer/Private/VT/VirtualTextureFeedback.h
@@ -38,8 +38,9 @@ public:
 
 	void			TransferGPUToCPU( FRHICommandListImmediate& RHICmdList, FIntRect const& Rect);
 
-	bool			Map( FRHICommandListImmediate& RHICmdList, MapResult& OutResult );
-	void			Unmap( FRHICommandListImmediate& RHICmdList, int32 MapHandle );
+	bool			CanMap();
+	bool			Map(FRHICommandListImmediate& RHICmdList, MapResult& OutResult);
+	void			Unmap(FRHICommandListImmediate& RHICmdList, int32 MapHandle);
 
 private:
 	struct FeedBackItem
diff --git a/Engine/Source/Runtime/Renderer/Private/VT/VirtualTextureSystem.cpp b/Engine/Source/Runtime/Renderer/Private/VT/VirtualTextureSystem.cpp
--- a/Engine/Source/Runtime/Renderer/Private/VT/VirtualTextureSystem.cpp
+++ b/Engine/Source/Runtime/Renderer/Private/VT/VirtualTextureSystem.cpp
@@ -11,6 +11,7 @@
 #include "UniquePageList.h"
 #include "UniqueRequestList.h"
 #include "Stats/Stats.h"
+#include "ScenePrivate.h"
 #include "SceneUtils.h"
 #include "HAL/IConsoleManager.h"
 #include "PostProcess/SceneRenderTargets.h"
@@ -841,7 +842,7 @@ void FVirtualTextureSystem::FeedbackAnalysisTask(const FFeedbackAnalysisParamete
 	}
 }
 
-void FVirtualTextureSystem::Update(FRHICommandListImmediate& RHICmdList, ERHIFeatureLevel::Type FeatureLevel)
+void FVirtualTextureSystem::Update(FRHICommandListImmediate& RHICmdList, ERHIFeatureLevel::Type FeatureLevel, FScene* Scene)
 {
 	check(IsInRenderingThread());
 
@@ -878,6 +879,19 @@ void FVirtualTextureSystem::Update(FRHICommandListImmediate& RHICmdList, ERHIFea
 		return;
 	}
 
+	// Flush any dirty runtime virtual textures for the current scene
+	if (Scene != nullptr)
+	{
+		// Only flush if we know that there is GPU feedback available to refill the visible data this frame
+		// This prevents bugs when low frame rate causes feedback buffer to stall so that the physical cache isn't filled immediately which causes visible glitching
+		FSceneRenderTargets& SceneContext = FSceneRenderTargets::Get(RHICmdList);
+		if (SceneContext.VirtualTextureFeedback.CanMap())
+		{
+			// Each RVT will call FVirtualTextureSystem::FlushCache()
+			Scene->FlushDirtyRuntimeVirtualTextures();
+		}
+	}
+
 	FMemStack& MemStack = FMemStack::Get();
 	FMemMark Mark(MemStack);
 	FUniquePageList* MergedUniquePageList = new(MemStack) FUniquePageList;
diff --git a/Engine/Source/Runtime/Renderer/Private/VT/VirtualTextureSystem.h b/Engine/Source/Runtime/Renderer/Private/VT/VirtualTextureSystem.h
--- a/Engine/Source/Runtime/Renderer/Private/VT/VirtualTextureSystem.h
+++ b/Engine/Source/Runtime/Renderer/Private/VT/VirtualTextureSystem.h
@@ -35,7 +35,7 @@ public:
 
 	void AllocateResources(FRHICommandListImmediate& RHICmdList, ERHIFeatureLevel::Type FeatureLevel);
 	void CallPendingCallbacks();
-	void Update( FRHICommandListImmediate& RHICmdList, ERHIFeatureLevel::Type FeatureLevel );
+	void Update( FRHICommandListImmediate& RHICmdList, ERHIFeatureLevel::Type FeatureLevel, FScene* Scene);
 
 	IAllocatedVirtualTexture* AllocateVirtualTexture(const FAllocatedVTDescription& Desc);
 	void DestroyVirtualTexture(IAllocatedVirtualTexture* AllocatedVT);
