Fixed mistake where AtmosphericFog was linked to Data for the SkyAtmospohere.

diff --git a/Engine/Source/Runtime/Renderer/Private/SceneRendering.cpp b/Engine/Source/Runtime/Renderer/Private/SceneRendering.cpp
--- a/Engine/Source/Runtime/Renderer/Private/SceneRendering.cpp
+++ b/Engine/Source/Runtime/Renderer/Private/SceneRendering.cpp
@@ -1124,6 +1124,19 @@ void FViewInfo::SetupUniformBufferParameters(
 	}
 
 	const FVector DefaultSunDirection(0.0f, 0.0f, 1.0f); // Up vector so that the AtmosphericLightVector node always output a valid direction.
+	auto ClearAtmosphereLightData = [&](uint32 Index)
+	{
+		check(Index < NUM_ATMOSPHERE_LIGHTS);
+		ViewUniformShaderParameters.AtmosphereLightDiscCosHalfApexAngle[Index] = FVector4(1.0f);
+		ViewUniformShaderParameters.AtmosphereLightDiscLuminance[Index] = FLinearColor::Black;
+		ViewUniformShaderParameters.AtmosphereLightColor[Index] = FLinearColor::Black;
+		ViewUniformShaderParameters.AtmosphereLightColor[Index].A = 0.0f;
+		ViewUniformShaderParameters.AtmosphereLightColorGlobalPostTransmittance[Index] = FLinearColor::Black;
+		ViewUniformShaderParameters.AtmosphereLightColorGlobalPostTransmittance[Index].A = 0.0f;
+		ViewUniformShaderParameters.AtmosphereLightDirection[Index] = DefaultSunDirection;
+	};
+
+	uint32 AtmosphereLightDataClearStartIndex = 0;
 	if (Scene)
 	{
 		if (Scene->SimpleDirectionalLight)
@@ -1158,6 +1171,9 @@ void FViewInfo::SetupUniformBufferParameters(
 			ViewUniformShaderParameters.AtmosphereLightDiscCosHalfApexAngle[0] = FVector4(FMath::Cos(Scene->AtmosphericFog->SunDiscScale * SunLightDiskHalfApexAngleRadian));
 			ViewUniformShaderParameters.AtmosphereLightDiscLuminance[0] = SunLight ? SunLight->Proxy->GetOuterSpaceLuminance() : FLinearColor::White;
 			ViewUniformShaderParameters.AtmosphereLightColor[0] = SunLight ? SunLight->Proxy->GetColor() : Scene->AtmosphericFog->DefaultSunColor; // Sun light color unaffected by atmosphere transmittance
+			ViewUniformShaderParameters.AtmosphereLightColor[0].A = 1.0f;
+			ViewUniformShaderParameters.AtmosphereLightColorGlobalPostTransmittance[0] = FLinearColor::Black;
+			ViewUniformShaderParameters.AtmosphereLightColorGlobalPostTransmittance[0].A = 0.0f;
 			ViewUniformShaderParameters.AtmosphereLightDirection[0] = SunLight ? -SunLight->Proxy->GetDirection() : -Scene->AtmosphericFog->DefaultSunDirection;
 		}
 		else
@@ -1179,8 +1195,12 @@ void FViewInfo::SetupUniformBufferParameters(
 			//Added check so atmospheric light color and vector can use a directional light without needing an atmospheric fog actor in the scene
 			ViewUniformShaderParameters.AtmosphereLightDiscLuminance[0] = SunLight ? SunLight->Proxy->GetOuterSpaceLuminance() : FLinearColor::Black;
 			ViewUniformShaderParameters.AtmosphereLightColor[0] = SunLight ? SunLight->Proxy->GetColor() : FLinearColor::Black;
+			ViewUniformShaderParameters.AtmosphereLightColor[0].A = 1.0f;
+			ViewUniformShaderParameters.AtmosphereLightColorGlobalPostTransmittance[0] = FLinearColor::Black;
+			ViewUniformShaderParameters.AtmosphereLightColorGlobalPostTransmittance[0].A = 0.0f;
 			ViewUniformShaderParameters.AtmosphereLightDirection[0] = SunLight ? -SunLight->Proxy->GetDirection() : DefaultSunDirection;
 		}
+		AtmosphereLightDataClearStartIndex = 1; // Do not clear the first atmosphere light data
 	}
 	else
 	{
@@ -1198,24 +1218,9 @@ void FViewInfo::SetupUniformBufferParameters(
 		ViewUniformShaderParameters.AtmosphericFogSunDiscScale = 1.f;
 		ViewUniformShaderParameters.AtmosphericFogRenderMask = EAtmosphereRenderFlag::E_EnableAll;
 		ViewUniformShaderParameters.AtmosphericFogInscatterAltitudeSampleNum = 0;
-		ViewUniformShaderParameters.AtmosphereLightDiscCosHalfApexAngle[0] = FVector4(1.0f);
-		ViewUniformShaderParameters.AtmosphereLightDiscLuminance[0] = FLinearColor::Black;
-		ViewUniformShaderParameters.AtmosphereLightColor[0] = FLinearColor::Black;
-		ViewUniformShaderParameters.AtmosphereLightDirection[0] = DefaultSunDirection;
+		AtmosphereLightDataClearStartIndex = 0; // Clear every atmosphere light data
 	}
 
-	auto ClearAtmosphereLightData = [&](uint32 Index)
-	{
-		check(Index < NUM_ATMOSPHERE_LIGHTS);
-		ViewUniformShaderParameters.AtmosphereLightDiscCosHalfApexAngle[Index] = FVector4(1.0f);
-		ViewUniformShaderParameters.AtmosphereLightDiscLuminance[Index] = FLinearColor::Black;
-		ViewUniformShaderParameters.AtmosphereLightColor[Index] = FLinearColor::Black;
-		ViewUniformShaderParameters.AtmosphereLightColor[Index].A = 0.0f;
-		ViewUniformShaderParameters.AtmosphereLightColorGlobalPostTransmittance[Index] = FLinearColor::Black;
-		ViewUniformShaderParameters.AtmosphereLightColorGlobalPostTransmittance[Index].A = 0.0f;
-		ViewUniformShaderParameters.AtmosphereLightDirection[Index] = DefaultSunDirection;
-	};
-
 	FRHITexture* TransmittanceLutTextureFound = nullptr;
 	FRHITexture* SkyViewLutTextureFound = nullptr;
 	FRHITexture* CameraAerialPerspectiveVolumeFound = nullptr;
@@ -1268,6 +1273,7 @@ void FViewInfo::SetupUniformBufferParameters(
 				ClearAtmosphereLightData(Index);
 			}
 		}
+		AtmosphereLightDataClearStartIndex = NUM_ATMOSPHERE_LIGHTS;	// Do not clear any atmosphere light data, this component sets everything it needs
 
 		// The constants below should match the one in SkyAtmosphereCommon.ush
 		const float SkyUnitToCm = 1.0f / 0.00001f;
@@ -1296,11 +1302,11 @@ void FViewInfo::SetupUniformBufferParameters(
 		ViewUniformShaderParameters.SkyAtmosphereCameraAerialPerspectiveVolumeDepthSliceLengthInv = 1.0f;
 		ViewUniformShaderParameters.SkyAtmosphereApplyCameraAerialPerspectiveVolume = 1.0f;
 		ViewUniformShaderParameters.SkyWorldCameraOrigin = ViewUniformShaderParameters.WorldCameraOrigin;
+	}
 
-		for (uint8 Index = 0; Index < NUM_ATMOSPHERE_LIGHTS; ++Index)
-		{
-			ClearAtmosphereLightData(Index);
-		}
+	for (uint8 Index = AtmosphereLightDataClearStartIndex; Index < NUM_ATMOSPHERE_LIGHTS; ++Index)
+	{
+		ClearAtmosphereLightData(Index);
 	}
 
 	ViewUniformShaderParameters.TransmittanceLutTexture = OrWhite2DIfNull(TransmittanceLutTextureFound);
