editgrate
Fixed non skydome sky to be capture in reflections.
Depth tested plane needs to be at least == nearclip plane distance.

diff --git a/Engine/Source/Runtime/Renderer/Private/SkyAtmosphereRendering.cpp b/Engine/Source/Runtime/Renderer/Private/SkyAtmosphereRendering.cpp
--- a/Engine/Source/Runtime/Renderer/Private/SkyAtmosphereRendering.cpp
+++ b/Engine/Source/Runtime/Renderer/Private/SkyAtmosphereRendering.cpp
@@ -1324,11 +1324,12 @@ void FSceneRenderer::RenderSkyAtmosphere(FRHICommandListImmediate& RHICmdList)
 			float StartDepthZ = 0.1f;
 			if (bFastAerialPerspectiveDepthTest)
 			{
-				const FMatrix ViewProjectionMatrix = View.ViewMatrices.GetProjectionMatrix();
-				float HalfHorizontalFOV = FMath::Atan(1.0f / ViewProjectionMatrix.M[0][0]);
-				float HalfVerticalFOV = FMath::Atan(1.0f / ViewProjectionMatrix.M[1][1]);
-				const float StartDepthViewCm = FMath::Cos(FMath::Max(HalfHorizontalFOV, HalfVerticalFOV)) * AerialPerspectiveStartDepthInCm;
-				const FVector4 Projected = ViewProjectionMatrix.TransformFVector4(FVector4(0.0f, 0.0f, StartDepthViewCm, 1.0f));
+				const FMatrix ProjectionMatrix = View.ViewMatrices.GetProjectionMatrix();
+				float HalfHorizontalFOV = FMath::Atan(1.0f / ProjectionMatrix.M[0][0]);
+				float HalfVerticalFOV = FMath::Atan(1.0f / ProjectionMatrix.M[1][1]);
+				float StartDepthViewCm = FMath::Cos(FMath::Max(HalfHorizontalFOV, HalfVerticalFOV)) * AerialPerspectiveStartDepthInCm;
+				StartDepthViewCm = FMath::Max(StartDepthViewCm, View.NearClippingDistance); // In any case, we need to limit the distance to frustum near plane to not be clipped away.
+				const FVector4 Projected = ProjectionMatrix.TransformFVector4(FVector4(0.0f, 0.0f, StartDepthViewCm, 1.0f));
 				StartDepthZ = Projected.Z / Projected.W;
 			}
 
