Fixed DFAO artifacts in the distance. DFAO distance fade out was done before the history blending and final upscaling. This caused bent normal artifacts in the distance where no DFAO computation takes place. Fixed it by replacing existing fade out with a fade out to world space normal after final upscaling.

diff --git a/Engine/Shaders/Private/DistanceFieldAOShared.ush b/Engine/Shaders/Private/DistanceFieldAOShared.ush
--- a/Engine/Shaders/Private/DistanceFieldAOShared.ush
+++ b/Engine/Shaders/Private/DistanceFieldAOShared.ush
@@ -29,6 +29,7 @@ float AOObjectMaxDistance;
 float AOStepScale;
 float AOStepExponentScale;
 float AOMaxViewDistance;
+float DistanceFadeScale; // 1 / ((1 - FadeDistanceFraction) * AOMaxViewDistance)
 
 float AOGlobalMaxOcclusionDistance;
 
@@ -168,7 +169,7 @@ float3 ComputeBentNormal(float3 RecordWorldNormal, uint RelativeRecordIndex)
 
 float2 AOBufferBilinearUVMax;
 
-float3 UpsampleDFAO(float2 BufferUV, float SceneDepth)
+float3 UpsampleDFAO(float2 BufferUV, float SceneDepth, float3 WorldNormal)
 {
 	// Distance field AO was computed at 0,0 regardless of viewrect min
 	float2 DistanceFieldUVs = BufferUV - View.ViewRectMin.xy * View.BufferSizeAndInvSize.zw;
@@ -214,6 +215,10 @@ float3 UpsampleDFAO(float2 BufferUV, float SceneDepth)
 
 #endif
 
+	// Fade to unoccluded in the distance
+	float FadeAlpha = saturate((AOMaxViewDistance - SceneDepth) * DistanceFadeScale);
+	BentNormal = lerp(WorldNormal, BentNormal, FadeAlpha);
+
 	return BentNormal;
 }
 
diff --git a/Engine/Shaders/Private/DistanceFieldLightingPost.usf b/Engine/Shaders/Private/DistanceFieldLightingPost.usf
--- a/Engine/Shaders/Private/DistanceFieldLightingPost.usf
+++ b/Engine/Shaders/Private/DistanceFieldLightingPost.usf
@@ -76,9 +76,6 @@ float4 GetNormalWeights(float2 Corner00UV, float2 LowResTexelSize, float3 WorldN
 	return max(NormalWeights, .0001f);
 }
 
-// 1 / ((1 - FadeDistanceFraction) * AOMaxViewDistance)
-float DistanceFadeScale;
-
 float ComputeSampleWeightBasedOnPosition(float4 ReferencePlane, float2 SampleScreenUV, float SampleDepth)
 {
 	float2 SampleScreenPosition = (SampleScreenUV.xy - View.ScreenPositionScaleBias.wz) / View.ScreenPositionScaleBias.xy;
@@ -146,12 +143,9 @@ void GeometryAwareUpsample(float4 UVAndScreenPos, out float4 OutBentNormal)
 
 	OutBentNormal = float4(AverageBentNormal, SceneDepth);
 
-	// Fade to unoccluded in the distance
-	//@todo - box distance from largest cascade
 	float BentNormalLength = length(OutBentNormal.rgb);
-	float FadeAlpha = saturate((AOMaxViewDistance - SceneDepth) * DistanceFadeScale);
 	float3 NormalizedBentNormal = OutBentNormal.rgb / max(BentNormalLength, .0001f);
-	OutBentNormal.rgb = NormalizedBentNormal * lerp(1, BentNormalLength, FadeAlpha);
+	OutBentNormal.rgb = NormalizedBentNormal * BentNormalLength;
 	//OutBentNormal = float4(WorldNormal, SceneDepth);
 }
 
@@ -372,7 +366,9 @@ void AOUpsamplePS(
 	,out float4 OutSceneColor : SV_Target0
 	)
 {
-	float3 BentNormal = UpsampleDFAO(UVAndScreenPos.xy, CalcSceneDepth(UVAndScreenPos.xy));
+	FGBufferData GBuffer = GetGBufferData(UVAndScreenPos.xy);
+
+	float3 BentNormal = UpsampleDFAO(UVAndScreenPos.xy, GBuffer.Depth, GBuffer.WorldNormal);
 
 #if MODULATE_SCENE_COLOR
 	float Visibility = lerp(length(BentNormal), 1.0f, MinIndirectDiffuseOcclusion);
diff --git a/Engine/Shaders/Private/ReflectionEnvironmentPixelShader.usf b/Engine/Shaders/Private/ReflectionEnvironmentPixelShader.usf
--- a/Engine/Shaders/Private/ReflectionEnvironmentPixelShader.usf
+++ b/Engine/Shaders/Private/ReflectionEnvironmentPixelShader.usf
@@ -314,7 +314,7 @@ void ReflectionEnvironmentSkyLighting(
 	float3 BentNormal = GBuffer.WorldNormal;
 	#if APPLY_SKY_SHADOWING
 	{
-		BentNormal = UpsampleDFAO(BufferUV, GBuffer.Depth);
+		BentNormal = UpsampleDFAO(BufferUV, GBuffer.Depth, GBuffer.WorldNormal);
 	}
 	#endif
 	
diff --git a/Engine/Source/Runtime/Renderer/Private/DistanceFieldAmbientOcclusion.h b/Engine/Source/Runtime/Renderer/Private/DistanceFieldAmbientOcclusion.h
--- a/Engine/Source/Runtime/Renderer/Private/DistanceFieldAmbientOcclusion.h
+++ b/Engine/Source/Runtime/Renderer/Private/DistanceFieldAmbientOcclusion.h
@@ -286,6 +286,14 @@ BEGIN_GLOBAL_SHADER_PARAMETER_STRUCT(FAOSampleData2,)
 	SHADER_PARAMETER_ARRAY(FVector4,SampleDirections,[NumConeSampleDirections])
 END_GLOBAL_SHADER_PARAMETER_STRUCT()
 
+inline float GetMaxAOViewDistance()
+{
+	extern float GAOMaxViewDistance;
+	// Scene depth stored in fp16 alpha, must fade out before it runs out of range
+	// The fade extends past GAOMaxViewDistance a bit
+	return FMath::Min(GAOMaxViewDistance, 65000.0f);
+}
+
 class FAOParameters
 {
 public:
@@ -323,7 +331,6 @@ public:
 
 		SetShaderValue(RHICmdList, ShaderRHI, AOStepExponentScale, GAOStepExponentScale);
 
-		extern float GetMaxAOViewDistance();
 		SetShaderValue(RHICmdList, ShaderRHI, AOMaxViewDistance, GetMaxAOViewDistance());
 
 		const float GlobalMaxOcclusionDistance = Parameters.GlobalMaxOcclusionDistance;
@@ -346,6 +353,8 @@ public:
 		BentNormalAOTexture.Bind(ParameterMap, TEXT("BentNormalAOTexture"));
 		BentNormalAOSampler.Bind(ParameterMap, TEXT("BentNormalAOSampler"));
 		AOBufferBilinearUVMax.Bind(ParameterMap, TEXT("AOBufferBilinearUVMax"));
+		DistanceFadeScale.Bind(ParameterMap, TEXT("DistanceFadeScale"));
+		AOMaxViewDistance.Bind(ParameterMap, TEXT("AOMaxViewDistance"));
 	}
 
 	void Set(FRHICommandList& RHICmdList, FRHIPixelShader* ShaderRHI, const FViewInfo& View, const TRefCountPtr<IPooledRenderTarget>& DistanceFieldAOBentNormal)
@@ -358,6 +367,12 @@ public:
 			(View.ViewRect.Width() / GAODownsampleFactor - 0.51f) / AOBufferSize.X, // 0.51 - so bilateral gather4 won't sample invalid texels
 			(View.ViewRect.Height() / GAODownsampleFactor - 0.51f) / AOBufferSize.Y);
 		SetShaderValue(RHICmdList, ShaderRHI, AOBufferBilinearUVMax, UVMax);
+
+		SetShaderValue(RHICmdList, ShaderRHI, AOMaxViewDistance, GetMaxAOViewDistance());
+
+		extern float GAOViewFadeDistanceScale;
+		const float DistanceFadeScaleValue = 1.0f / ((1.0f - GAOViewFadeDistanceScale) * GetMaxAOViewDistance());
+		SetShaderValue(RHICmdList, ShaderRHI, DistanceFadeScale, DistanceFadeScaleValue);
 	}
 
 	/** Serializer. */
@@ -366,6 +381,8 @@ public:
 		Ar << P.BentNormalAOTexture;
 		Ar << P.BentNormalAOSampler;
 		Ar << P.AOBufferBilinearUVMax;
+		Ar << P.DistanceFadeScale;
+		Ar << P.AOMaxViewDistance;
 
 		return Ar;
 	}
@@ -374,16 +391,10 @@ private:
 	FShaderResourceParameter BentNormalAOTexture;
 	FShaderResourceParameter BentNormalAOSampler;
 	FShaderParameter AOBufferBilinearUVMax;
+	FShaderParameter DistanceFadeScale;
+	FShaderParameter AOMaxViewDistance;
 };
 
-inline float GetMaxAOViewDistance()
-{
-	extern float GAOMaxViewDistance;
-	// Scene depth stored in fp16 alpha, must fade out before it runs out of range
-	// The fade extends past GAOMaxViewDistance a bit
-	return FMath::Min(GAOMaxViewDistance, 65000.0f);
-}
-
 class FMaxSizedRWBuffers : public FRenderResource
 {
 public:
diff --git a/Engine/Source/Runtime/Renderer/Private/ReflectionEnvironment.cpp b/Engine/Source/Runtime/Renderer/Private/ReflectionEnvironment.cpp
--- a/Engine/Source/Runtime/Renderer/Private/ReflectionEnvironment.cpp
+++ b/Engine/Source/Runtime/Renderer/Private/ReflectionEnvironment.cpp
@@ -550,6 +550,8 @@ class FReflectionEnvironmentSkyLightingPS : public FGlobalShader
 		// Distance field AO parameters.
 		// TODO. FDFAOUpsampleParameters
 		SHADER_PARAMETER(FVector2D, AOBufferBilinearUVMax)
+		SHADER_PARAMETER(float, DistanceFadeScale)
+		SHADER_PARAMETER(float, AOMaxViewDistance)
 		SHADER_PARAMETER_RDG_TEXTURE(Texture2D, BentNormalAOTexture)
 		SHADER_PARAMETER_SAMPLER(SamplerState,  BentNormalAOSampler)
 		
@@ -817,6 +819,10 @@ void FDeferredShadingSceneRenderer::RenderDeferredReflectionsAndSkyLighting(FRHI
 						(View.ViewRect.Width() / GAODownsampleFactor - 0.51f) / AOBufferSize.X, // 0.51 - so bilateral gather4 won't sample invalid texels
 						(View.ViewRect.Height() / GAODownsampleFactor - 0.51f) / AOBufferSize.Y);
 
+					extern float GAOViewFadeDistanceScale;
+					PassParameters->AOMaxViewDistance = GetMaxAOViewDistance();
+					PassParameters->DistanceFadeScale = 1.0f / ((1.0f - GAOViewFadeDistanceScale) * GetMaxAOViewDistance());
+
 					PassParameters->BentNormalAOTexture = DynamicBentNormalAOTexture;
 					PassParameters->BentNormalAOSampler = TStaticSamplerState<SF_Bilinear>::GetRHI();
 				}
