On world cleanup, free unused render targets.

#rnx

diff --git a/Engine/Source/Runtime/Engine/Private/World.cpp b/Engine/Source/Runtime/Engine/Private/World.cpp
--- a/Engine/Source/Runtime/Engine/Private/World.cpp
+++ b/Engine/Source/Runtime/Engine/Private/World.cpp
@@ -4123,6 +4123,8 @@ void UWorld::CleanupWorld(bool bSessionEnded, bool bCleanupResources, UWorld* Ne
 
 	FWorldDelegates::OnWorldCleanup.Broadcast(this, bSessionEnded, bCleanupResources);
 
+	GetRendererModule().OnWorldCleanup(this, bSessionEnded, bCleanupResources);
+
 	if (AISystem != nullptr)
 	{
 		AISystem->CleanupWorld(bSessionEnded, bCleanupResources, NewWorld);
diff --git a/Engine/Source/Runtime/RenderCore/Public/RendererInterface.h b/Engine/Source/Runtime/RenderCore/Public/RendererInterface.h
--- a/Engine/Source/Runtime/RenderCore/Public/RendererInterface.h
+++ b/Engine/Source/Runtime/RenderCore/Public/RendererInterface.h
@@ -614,6 +614,8 @@ public:
 	/** Forces reallocation of scene render targets. */
 	virtual void ReallocateSceneRenderTargets() = 0;
 
+	virtual void OnWorldCleanup(UWorld* World, bool bSessionEnded, bool bCleanupResources) = 0;
+
 	/** Sets the buffer size of the render targets. */
 	virtual void SceneRenderTargetsSetBufferSize(uint32 SizeX, uint32 SizeY) = 0;
 
diff --git a/Engine/Source/Runtime/Renderer/Private/Renderer.cpp b/Engine/Source/Runtime/Renderer/Private/Renderer.cpp
--- a/Engine/Source/Runtime/Renderer/Private/Renderer.cpp
+++ b/Engine/Source/Runtime/Renderer/Private/Renderer.cpp
@@ -45,6 +45,10 @@ IMPLEMENT_MODULE(FRendererModule, Renderer);
 	FSystemSettings* GSystemSettingsForVisualizers = &GSystemSettings;
 #endif
 
+static int32 bFlushRenderTargetsOnWorldCleanup = 1;
+FAutoConsoleVariableRef CVarFlushRenderTargetsOnWorldCleanup(TEXT("r.bFlushRenderTargetsOnWorldCleanup"), bFlushRenderTargetsOnWorldCleanup, TEXT(""));
+
+
 void FRendererModule::StartupModule()
 {
 	GScreenSpaceDenoiser = IScreenSpaceDenoiser::GetDefaultDenoiser();
@@ -66,6 +70,18 @@ void FRendererModule::ReallocateSceneRenderTargets()
 	FSceneRenderTargets::GetGlobalUnsafe().UpdateRHI();
 }
 
+void FRendererModule::OnWorldCleanup(UWorld* World, bool bSessionEnded, bool bCleanupResources)
+{
+	if (bFlushRenderTargetsOnWorldCleanup > 0)
+	{
+		ENQUEUE_RENDER_COMMAND(OnWorldCleanup)(
+			[](FRHICommandListImmediate& RHICmdList)
+		{
+			GRenderTargetPool.FreeUnusedResources();
+		});
+	}
+}
+
 void FRendererModule::SceneRenderTargetsSetBufferSize(uint32 SizeX, uint32 SizeY)
 {
 	FSceneRenderTargets::GetGlobalUnsafe().SetBufferSize(SizeX, SizeY);
diff --git a/Engine/Source/Runtime/Renderer/Private/RendererModule.h b/Engine/Source/Runtime/Renderer/Private/RendererModule.h
--- a/Engine/Source/Runtime/Renderer/Private/RendererModule.h
+++ b/Engine/Source/Runtime/Renderer/Private/RendererModule.h
@@ -45,6 +45,7 @@ public:
 	virtual FSceneViewStateInterface* AllocateViewState() override;
 	virtual uint32 GetNumDynamicLightsAffectingPrimitive(const FPrimitiveSceneInfo* PrimitiveSceneInfo,const FLightCacheInterface* LCI) override;
 	virtual void ReallocateSceneRenderTargets() override;
+	virtual void OnWorldCleanup(UWorld* World, bool bSessionEnded, bool bCleanupResources) override;
 	virtual void SceneRenderTargetsSetBufferSize(uint32 SizeX, uint32 SizeY) override;
 	virtual void InitializeSystemTextures(FRHICommandListImmediate& RHICmdList);
 	virtual void DrawTileMesh(FRHICommandListImmediate& RHICmdList, FMeshPassProcessorRenderState& DrawRenderState, const FSceneView& View, FMeshBatch& Mesh, bool bIsHitTesting, const FHitProxyId& HitProxyId) override;
