Using a better proxy behavior for the sky atmosphere component.
Tested on apollo map and others

diff --git a/Engine/Source/Runtime/Engine/Classes/Components/SkyAtmosphereComponent.h b/Engine/Source/Runtime/Engine/Classes/Components/SkyAtmosphereComponent.h
--- a/Engine/Source/Runtime/Engine/Classes/Components/SkyAtmosphereComponent.h
+++ b/Engine/Source/Runtime/Engine/Classes/Components/SkyAtmosphereComponent.h
@@ -11,6 +11,9 @@
 #include "SkyAtmosphereComponent.generated.h"
 
 
+class FSkyAtmosphereSceneProxy;
+
+
 USTRUCT(BlueprintType)
 struct FTentDistribution
 {
@@ -188,9 +191,7 @@ public:
 
 private:
 
-	/** Add this component to the render scene */
-	void AddToRenderScene() const;
-
+	FSkyAtmosphereSceneProxy* SkyAtmosphereSceneProxy;
 
 	/**
 	 * GUID used to associate a atmospheric component with precomputed lighting/shadowing information across levels.
diff --git a/Engine/Source/Runtime/Engine/Private/Components/SkyAtmosphereComponent.cpp b/Engine/Source/Runtime/Engine/Private/Components/SkyAtmosphereComponent.cpp
--- a/Engine/Source/Runtime/Engine/Private/Components/SkyAtmosphereComponent.cpp
+++ b/Engine/Source/Runtime/Engine/Private/Components/SkyAtmosphereComponent.cpp
@@ -11,6 +11,7 @@
 #include "Logging/TokenizedMessage.h"
 #include "Misc/MapErrors.h"
 #include "Misc/UObjectToken.h"
+#include "Rendering/SkyAtmosphereCommonData.h"
 #include "UObject/UObjectIterator.h"
 #include "UObject/ConstructorHelpers.h"
 
@@ -30,6 +31,7 @@
 
 USkyAtmosphereComponent::USkyAtmosphereComponent(const FObjectInitializer& ObjectInitializer)
 	: Super(ObjectInitializer)
+	, SkyAtmosphereSceneProxy(nullptr)
 {
 	// All distance here are in kilometer and scattering/absorptions coefficient in 1/kilometers.
 	const float EarthBottomRadius = 6360.0f;
@@ -91,36 +93,46 @@ static bool SkyAtmosphereComponentStaticLightingBuilt(const USkyAtmosphereCompon
 	return true;	// The component has not been spawned in any world yet so let's mark it as built for now.
 }
 
-void USkyAtmosphereComponent::AddToRenderScene() const
-{
-	if (ShouldComponentAddToScene() && ShouldRender() && IsRegistered() && (GetOuter() == NULL || !GetOuter()->HasAnyFlags(RF_ClassDefaultObject)))
-	{
-		GetWorld()->Scene->AddSkyAtmosphere(this, SkyAtmosphereComponentStaticLightingBuilt(this));
-	}
-}
-
 void USkyAtmosphereComponent::CreateRenderState_Concurrent()
 {
 	Super::CreateRenderState_Concurrent();
 	// If one day we need to look up lightmass built data, lookup it up here using the guid from the correct MapBuildData.
-	AddToRenderScene();
+
+	bool bHidden = false;
+#if WITH_EDITORONLY_DATA
+	bHidden = GetOwner() ? GetOwner()->bHiddenEdLevel : false;
+#endif // WITH_EDITORONLY_DATA
+	if (!ShouldComponentAddToScene())
+	{
+		bHidden = true;
+	}
+
+	if (bVisible && !bHidden &&
+		ShouldComponentAddToScene() && ShouldRender() && IsRegistered() && (GetOuter() == NULL || !GetOuter()->HasAnyFlags(RF_ClassDefaultObject)))
+	{
+		// Create the scene proxy.
+		SkyAtmosphereSceneProxy = new FSkyAtmosphereSceneProxy(this);
+		GetWorld()->Scene->AddSkyAtmosphere(SkyAtmosphereSceneProxy, SkyAtmosphereComponentStaticLightingBuilt(this));
+	}
+
 }
 
 void USkyAtmosphereComponent::DestroyRenderState_Concurrent()
 {
 	Super::DestroyRenderState_Concurrent();
-	GetWorld()->Scene->RemoveSkyAtmosphere(this);
 
-	// Search for new fog component
-	for (TObjectIterator<USkyAtmosphereComponent> It; It; ++It)
+	if (SkyAtmosphereSceneProxy)
 	{
-		USkyAtmosphereComponent* Component = *It;
-		checkSlow(Component);
-		if (Component != this && Component->IsRegistered())
+		GetWorld()->Scene->RemoveSkyAtmosphere(SkyAtmosphereSceneProxy);
+
+		FSkyAtmosphereSceneProxy* SceneProxy = SkyAtmosphereSceneProxy;
+		ENQUEUE_RENDER_COMMAND(FDestroySkyAtmosphereSceneProxyCommand)(
+			[SceneProxy](FRHICommandList& RHICmdList)
 		{
-			Component->AddToRenderScene();
-			break;
-		}
+			delete SceneProxy;
+		});
+
+		SkyAtmosphereSceneProxy = nullptr;
 	}
 }
 
@@ -237,10 +249,10 @@ void USkyAtmosphereComponent::Serialize(FArchive& Ar)
 
 void USkyAtmosphereComponent::OverrideAtmosphereLightDirection(int32 AtmosphereLightIndex, const FVector& LightDirection)
 {
-	if (AreDynamicDataChangesAllowed())
+	if (AreDynamicDataChangesAllowed() && SkyAtmosphereSceneProxy)
 	{
 		FSceneInterface* Scene = GetWorld()->Scene;
-		Scene->OverrideSkyAtmosphereLightDirection(this, AtmosphereLightIndex, LightDirection);
+		Scene->OverrideSkyAtmosphereLightDirection(SkyAtmosphereSceneProxy, AtmosphereLightIndex, LightDirection);
 	}
 }
 
@@ -331,6 +343,47 @@ ASkyAtmosphere::ASkyAtmosphere(const FObjectInitializer& ObjectInitializer)
 	bHidden = false;
 }
 
+
+
+/*=============================================================================
+	FSkyAtmosphereSceneProxy implementation.
+=============================================================================*/
+
+
+
+FSkyAtmosphereSceneProxy::FSkyAtmosphereSceneProxy(const USkyAtmosphereComponent* InComponent)
+	: bStaticLightingBuilt(false)
+	, AtmosphereSetup(*InComponent)
+{
+	SkyLuminanceFactor = InComponent->SkyLuminanceFactor;
+	AerialPespectiveViewDistanceScale = InComponent->AerialPespectiveViewDistanceScale;
+	memset(OverrideAtmosphericLight, 0, sizeof(OverrideAtmosphericLight));
+
+	TransmittanceAtZenith = AtmosphereSetup.GetTransmittanceAtGroundLevel(FVector(0.0f, 0.0f, 1.0f));
+}
+
+FSkyAtmosphereSceneProxy::~FSkyAtmosphereSceneProxy()
+{
+}
+
+void FSkyAtmosphereSceneProxy::OverrideAtmosphereLightDirection(int32 AtmosphereLightIndex, const FVector& LightDirection)
+{
+	check(AtmosphereLightIndex >= 0 && AtmosphereLightIndex < NUM_ATMOSPHERE_LIGHTS);
+	AtmosphereLightIndex = FMath::Clamp(AtmosphereLightIndex, 0, NUM_ATMOSPHERE_LIGHTS - 1);	// To make sure we do not crash, blueprint function cannot enforce input ranges
+	OverrideAtmosphericLight[AtmosphereLightIndex] = true;
+	OverrideAtmosphericLightDirection[AtmosphereLightIndex] = LightDirection;
+}
+
+FVector FSkyAtmosphereSceneProxy::GetAtmosphereLightDirection(int32 AtmosphereLightIndex, const FVector& DefaultDirection) const
+{
+	if (OverrideAtmosphericLight[AtmosphereLightIndex])
+	{
+		return OverrideAtmosphericLightDirection[AtmosphereLightIndex];
+	}
+	return DefaultDirection;
+}
+
+
 #undef LOCTEXT_NAMESPACE
 
 
diff --git a/Engine/Source/Runtime/Engine/Public/SceneInterface.h b/Engine/Source/Runtime/Engine/Public/SceneInterface.h
--- a/Engine/Source/Runtime/Engine/Public/SceneInterface.h
+++ b/Engine/Source/Runtime/Engine/Public/SceneInterface.h
@@ -11,6 +11,7 @@
 class AWorldSettings;
 class FAtmosphericFogSceneInfo;
 class FSkyAtmosphereRenderSceneInfo;
+class FSkyAtmosphereSceneProxy;
 class FMaterial;
 class FMaterialShaderMap;
 class FPrimitiveSceneInfo;
@@ -248,27 +249,27 @@ public:
 	/**
 	 * Adds the unique sky atmosphere component to the scene
 	 *
-	 * @param SkyAtmosphereComponent - component to add
+	 * @param SkyAtmosphereSceneProxy - the sky atmosphere proxy
 	 */
-	virtual void AddSkyAtmosphere(const class USkyAtmosphereComponent* SkyAtmosphereComponent, bool bStaticLightingBuilt) = 0;
+	virtual void AddSkyAtmosphere(FSkyAtmosphereSceneProxy* SkyAtmosphereSceneProxy, bool bStaticLightingBuilt) = 0;
 	/**
 	 * Removes the unique sky atmosphere component to the scene
 	 *
-	 * @param SkyAtmosphereComponent - component to remove
+	 * @param SkyAtmosphereSceneProxy - the sky atmosphere proxy
 	 */
-	virtual void RemoveSkyAtmosphere(const class USkyAtmosphereComponent* SkyAtmosphereComponent) = 0;
+	virtual void RemoveSkyAtmosphere(FSkyAtmosphereSceneProxy* SkyAtmosphereSceneProxy) = 0;
 	/**
-	 * Returns the scene's unique FSkyAtmosphereRenderSceneInfo if it exists
+	 * Returns the scene's unique info if it exists
 	 */
 	virtual FSkyAtmosphereRenderSceneInfo* GetSkyAtmosphereSceneInfo() = 0;
 	virtual const FSkyAtmosphereRenderSceneInfo* GetSkyAtmosphereSceneInfo() const = 0;
 	/**
 	 * Override a sky atmosphere light direction
-	 * @param SkyAtmosphereComponent - component to verify it is the actual unique SkyAtmosphere
+	 * @param SkyAtmosphereSceneProxy - the sky atmosphere proxy
 	 * @param AtmosphereLightIndex - the atmosphere light index to consider
 	 * @param LightDirection - the new light direction to override the atmosphere light with
 	 */
-	virtual void OverrideSkyAtmosphereLightDirection(const class USkyAtmosphereComponent* SkyAtmosphereComponent, int32 AtmosphereLightIndex, const FVector& LightDirection) = 0;
+	virtual void OverrideSkyAtmosphereLightDirection(FSkyAtmosphereSceneProxy* SkyAtmosphereSceneProxy, int32 AtmosphereLightIndex, const FVector& LightDirection) = 0;
 
 	/**
 	 * Adds a wind source component to the scene.
diff --git a/Engine/Source/Runtime/Engine/Public/SceneManagement.h b/Engine/Source/Runtime/Engine/Public/SceneManagement.h
--- a/Engine/Source/Runtime/Engine/Public/SceneManagement.h
+++ b/Engine/Source/Runtime/Engine/Public/SceneManagement.h
@@ -27,6 +27,7 @@
 #include "SceneUtils.h"
 #include "LightmapUniformShaderParameters.h"
 #include "DynamicBufferAllocator.h"
+#include "Rendering/SkyAtmosphereCommonData.h"
 
 class FCanvas;
 class FLightMap;
@@ -44,6 +45,8 @@ class ULightMapTexture2D;
 class UMaterialInstanceDynamic;
 class UMaterialInterface;
 class UShadowMapTexture2D;
+class USkyAtmosphereComponent;
+class FSkyAtmosphereRenderSceneInfo;
 class USkyLightComponent;
 struct FDynamicMeshVertex;
 class ULightMapVirtualTexture2D;
@@ -1090,6 +1093,38 @@ private:
 	const uint8 bMovable : 1;
 };
 
+/** Represents a USkyAtmosphereComponent to the rendering thread. */
+class ENGINE_API FSkyAtmosphereSceneProxy
+{
+public:
+
+	// Initialization constructor.
+	FSkyAtmosphereSceneProxy(const USkyAtmosphereComponent* InComponent);
+	~FSkyAtmosphereSceneProxy();
+
+	bool IsMultiScatteringEnabled() const { return AtmosphereSetup.MultiScatteringFactor > 0.0f; }
+	FLinearColor GetSkyLuminanceFactor() const { return SkyLuminanceFactor; }
+	FLinearColor GetTransmittanceAtZenith() const { return TransmittanceAtZenith; };
+	float GetAerialPespectiveViewDistanceScale() const { return AerialPespectiveViewDistanceScale; }
+
+	const FAtmosphereSetup& GetAtmosphereSetup() const { return AtmosphereSetup; }
+
+	void OverrideAtmosphereLightDirection(int32 AtmosphereLightIndex, const FVector& LightDirection);
+	FVector GetAtmosphereLightDirection(int32 AtmosphereLightIndex, const FVector& DefaultDirection) const;
+
+	bool bStaticLightingBuilt;
+	FSkyAtmosphereRenderSceneInfo* RenderSceneInfo;
+private:
+
+	FAtmosphereSetup AtmosphereSetup;
+
+	FLinearColor TransmittanceAtZenith;
+	FLinearColor SkyLuminanceFactor;
+	float AerialPespectiveViewDistanceScale;
+
+	bool OverrideAtmosphericLight[NUM_ATMOSPHERE_LIGHTS];
+	FVector OverrideAtmosphericLightDirection[NUM_ATMOSPHERE_LIGHTS];
+};
 
 /** Shader paraneter structure for rendering lights. */
 BEGIN_SHADER_PARAMETER_STRUCT(FLightShaderParameters, ENGINE_API)
diff --git a/Engine/Source/Runtime/Renderer/Private/DeferredShadingRenderer.cpp b/Engine/Source/Runtime/Renderer/Private/DeferredShadingRenderer.cpp
--- a/Engine/Source/Runtime/Renderer/Private/DeferredShadingRenderer.cpp
+++ b/Engine/Source/Runtime/Renderer/Private/DeferredShadingRenderer.cpp
@@ -851,7 +851,7 @@ void FDeferredShadingSceneRenderer::Render(FRHICommandListImmediate& RHICmdList)
 		{
 			if (Scene->AtmosphereLights[LightIndex])
 			{
-				Scene->GetSkyAtmosphereSceneInfo()->PrepareSunLightProxy(LightIndex, *Scene->AtmosphereLights[LightIndex]);
+				PrepareSunLightProxy(*Scene->GetSkyAtmosphereSceneInfo(),LightIndex, *Scene->AtmosphereLights[LightIndex]);
 			}
 		}
 	}
diff --git a/Engine/Source/Runtime/Renderer/Private/MobileShadingRenderer.cpp b/Engine/Source/Runtime/Renderer/Private/MobileShadingRenderer.cpp
--- a/Engine/Source/Runtime/Renderer/Private/MobileShadingRenderer.cpp
+++ b/Engine/Source/Runtime/Renderer/Private/MobileShadingRenderer.cpp
@@ -316,7 +316,7 @@ void FMobileSceneRenderer::Render(FRHICommandListImmediate& RHICmdList)
 		{
 			if (Scene->AtmosphereLights[LightIndex])
 			{
-				Scene->GetSkyAtmosphereSceneInfo()->PrepareSunLightProxy(LightIndex, *Scene->AtmosphereLights[LightIndex]);
+				PrepareSunLightProxy(*Scene->GetSkyAtmosphereSceneInfo(), LightIndex, *Scene->AtmosphereLights[LightIndex]);
 			}
 		}
 	}
diff --git a/Engine/Source/Runtime/Renderer/Private/RendererScene.cpp b/Engine/Source/Runtime/Renderer/Private/RendererScene.cpp
--- a/Engine/Source/Runtime/Renderer/Private/RendererScene.cpp
+++ b/Engine/Source/Runtime/Renderer/Private/RendererScene.cpp
@@ -3704,11 +3704,11 @@ public:
 	virtual void RemoveAtmosphericFogResource_RenderThread(FRenderResource* FogResource) override {}
 	virtual FAtmosphericFogSceneInfo* GetAtmosphericFogSceneInfo() override { return NULL; }
 
-	virtual void AddSkyAtmosphere(const class USkyAtmosphereComponent* SkyAtmosphereComponent, bool bStaticLightingBuilt) override {}
-	virtual void RemoveSkyAtmosphere(const class USkyAtmosphereComponent* SkyAtmosphereComponent) override {}
+	virtual void AddSkyAtmosphere(FSkyAtmosphereSceneProxy* SkyAtmosphereSceneProxy, bool bStaticLightingBuilt) override {}
+	virtual void RemoveSkyAtmosphere(FSkyAtmosphereSceneProxy* SkyAtmosphereSceneProxy) override {}
 	virtual FSkyAtmosphereRenderSceneInfo* GetSkyAtmosphereSceneInfo() override { return NULL; }
 	virtual const FSkyAtmosphereRenderSceneInfo* GetSkyAtmosphereSceneInfo() const override { return NULL; }
-	virtual void OverrideSkyAtmosphereLightDirection(const class USkyAtmosphereComponent* SkyAtmosphereComponent, int32 AtmosphereLightIndex, const FVector& LightDirection) override {}
+	virtual void OverrideSkyAtmosphereLightDirection(FSkyAtmosphereSceneProxy* SkyAtmosphereSceneProxy, int32 AtmosphereLightIndex, const FVector& LightDirection) override {}
 
 	virtual void AddWindSource(class UWindDirectionalSourceComponent* WindComponent) override {}
 	virtual void RemoveWindSource(class UWindDirectionalSourceComponent* WindComponent) override {}
diff --git a/Engine/Source/Runtime/Renderer/Private/ScenePrivate.h b/Engine/Source/Runtime/Renderer/Private/ScenePrivate.h
--- a/Engine/Source/Runtime/Renderer/Private/ScenePrivate.h
+++ b/Engine/Source/Runtime/Renderer/Private/ScenePrivate.h
@@ -2559,6 +2559,9 @@ public:
 	/** The sky/atmosphere components of the scene. */
 	FSkyAtmosphereRenderSceneInfo* SkyAtmosphere;
 
+	/** Used to track the order that skylights were enabled in. */
+	TArray<FSkyAtmosphereSceneProxy*> SkyAtmosphereStack;
+
 	/** The wind sources in the scene. */
 	TArray<class FWindSourceSceneProxy*> WindSources;
 
@@ -2686,11 +2689,11 @@ public:
 	virtual void RemoveAtmosphericFogResource_RenderThread(FRenderResource* FogResource) override;
 	virtual FAtmosphericFogSceneInfo* GetAtmosphericFogSceneInfo() override { return AtmosphericFog; }
 
-	virtual void AddSkyAtmosphere(const class USkyAtmosphereComponent* SkyAtmosphereComponent, bool bStaticLightingBuilt) override;
-	virtual void RemoveSkyAtmosphere(const class USkyAtmosphereComponent* SkyAtmosphereComponent) override;
+	virtual void AddSkyAtmosphere(FSkyAtmosphereSceneProxy* SkyAtmosphereSceneProxy, bool bStaticLightingBuilt) override;
+	virtual void RemoveSkyAtmosphere(FSkyAtmosphereSceneProxy* SkyAtmosphereSceneProxy) override;
 	virtual FSkyAtmosphereRenderSceneInfo* GetSkyAtmosphereSceneInfo() override { return SkyAtmosphere; }
 	virtual const FSkyAtmosphereRenderSceneInfo* GetSkyAtmosphereSceneInfo() const override { return SkyAtmosphere; }
-	virtual void OverrideSkyAtmosphereLightDirection(const class USkyAtmosphereComponent* SkyAtmosphereComponent, int32 AtmosphereLightIndex, const FVector& LightDirection) override;
+	virtual void OverrideSkyAtmosphereLightDirection(FSkyAtmosphereSceneProxy* SkyAtmosphereSceneProxy, int32 AtmosphereLightIndex, const FVector& LightDirection) override;
 
 	virtual void AddWindSource(UWindDirectionalSourceComponent* WindComponent) override;
 	virtual void RemoveWindSource(UWindDirectionalSourceComponent* WindComponent) override;
diff --git a/Engine/Source/Runtime/Renderer/Private/SceneRendering.cpp b/Engine/Source/Runtime/Renderer/Private/SceneRendering.cpp
--- a/Engine/Source/Runtime/Renderer/Private/SceneRendering.cpp
+++ b/Engine/Source/Runtime/Renderer/Private/SceneRendering.cpp
@@ -1229,6 +1229,7 @@ void FViewInfo::SetupUniformBufferParameters(
 	{
 		check(Scene->SkyAtmosphere->GetTransmittanceLutTexture().IsValid());
 		FSkyAtmosphereRenderSceneInfo* SkyAtmosphere = Scene->SkyAtmosphere;
+		const FSkyAtmosphereSceneProxy& SkyAtmosphereSceneProxy = SkyAtmosphere->GetSkyAtmosphereSceneProxy();
 		const TRefCountPtr<IPooledRenderTarget>& PooledTransmittanceLutTexture = SkyAtmosphere->GetTransmittanceLutTexture();
 		TransmittanceLutTextureFound = PooledTransmittanceLutTexture->GetRenderTargetItem().ShaderResourceTexture;
 
@@ -1239,9 +1240,9 @@ void FViewInfo::SetupUniformBufferParameters(
 
 		CameraAerialPerspectiveVolumeFound = this->SkyAtmosphereCameraAerialPerspectiveVolume->GetRenderTargetItem().ShaderResourceTexture;
 		DistantSkyLightLutTextureFound = SkyAtmosphere->GetDistantSkyLightLutTextureRHI();
-		ViewUniformShaderParameters.SkyAtmosphereSkyLuminanceFactor = SkyAtmosphere->GetSkyLuminanceFactor();
+		ViewUniformShaderParameters.SkyAtmosphereSkyLuminanceFactor = SkyAtmosphereSceneProxy.GetSkyLuminanceFactor();
 
-		const FAtmosphereSetup& AtmosphereSetup =  SkyAtmosphere->GetAtmosphereSetup();
+		const FAtmosphereSetup& AtmosphereSetup = SkyAtmosphereSceneProxy.GetAtmosphereSetup();
 		ViewUniformShaderParameters.SkyAtmosphereBottomRadius = AtmosphereSetup.BottomRadius;
 		ViewUniformShaderParameters.SkyAtmosphereTopRadius = AtmosphereSetup.TopRadius;
 
@@ -1266,7 +1267,7 @@ void FViewInfo::SetupUniformBufferParameters(
 				ViewUniformShaderParameters.AtmosphereLightColor[Index].A = 1.0f;
 				ViewUniformShaderParameters.AtmosphereLightColorGlobalPostTransmittance[Index] = Light->Proxy->GetColor() * Light->Proxy->GetTransmittanceFactor();
 				ViewUniformShaderParameters.AtmosphereLightColorGlobalPostTransmittance[Index].A = 1.0f;
-				ViewUniformShaderParameters.AtmosphereLightDirection[Index] = SkyAtmosphere->GetAtmosphereLightDirection(Index, -Light->Proxy->GetDirection());
+				ViewUniformShaderParameters.AtmosphereLightDirection[Index] = SkyAtmosphereSceneProxy.GetAtmosphereLightDirection(Index, -Light->Proxy->GetDirection());
 			}
 			else
 			{
diff --git a/Engine/Source/Runtime/Renderer/Private/SkyAtmosphereRendering.cpp b/Engine/Source/Runtime/Renderer/Private/SkyAtmosphereRendering.cpp
--- a/Engine/Source/Runtime/Renderer/Private/SkyAtmosphereRendering.cpp
+++ b/Engine/Source/Runtime/Renderer/Private/SkyAtmosphereRendering.cpp
@@ -12,6 +12,7 @@
 #include "LightSceneInfo.h"
 #include "PixelShaderUtils.h"
 #include "RenderTargetTemp.h"
+#include "Rendering/SkyAtmosphereCommonData.h"
 #include "ScenePrivate.h"
 #include "SceneRenderTargetParameters.h"
 
@@ -333,27 +334,36 @@ static FLinearColor GetLightDiskLuminance(FLightSceneInfo& Light, FLinearColor L
 	return LightIlluminance / SunSolidAngle; // approximation
 }
 
+void PrepareSunLightProxy(const FSkyAtmosphereRenderSceneInfo& SkyAtmosphere, uint32 AtmosphereLightIndex, FLightSceneInfo& AtmosphereLight)
+{
+	// See explanation in https://media.contentapi.ea.com/content/dam/eacom/frostbite/files/s2016-pbs-frostbite-sky-clouds-new.pdf page 26
+	const bool bAtmosphereAffectsSunIlluminance = true;
+	const FSkyAtmosphereSceneProxy& SkyAtmosphereProxy = SkyAtmosphere.GetSkyAtmosphereSceneProxy();
+	const FVector AtmosphereLightDirection = SkyAtmosphereProxy.GetAtmosphereLightDirection(AtmosphereLightIndex, -AtmosphereLight.Proxy->GetDirection());
+	FLinearColor TransmittanceTowardSun = bAtmosphereAffectsSunIlluminance ? SkyAtmosphereProxy.GetAtmosphereSetup().GetTransmittanceAtGroundLevel(AtmosphereLightDirection) : FLinearColor(FLinearColor::White);
+	FLinearColor TransmittanceAtZenithFinal = bAtmosphereAffectsSunIlluminance ? SkyAtmosphereProxy.GetTransmittanceAtZenith() : FLinearColor(FLinearColor::White);
+
+	FLinearColor SunZenithIlluminance = AtmosphereLight.Proxy->GetColor();
+	FLinearColor SunOuterSpaceIlluminance = SunZenithIlluminance / TransmittanceAtZenithFinal;
+	FLinearColor SunDiskOuterSpaceLuminance = GetLightDiskLuminance(AtmosphereLight, SunOuterSpaceIlluminance);
+
+	AtmosphereLight.Proxy->SetAtmosphereRelatedProperties(TransmittanceTowardSun / TransmittanceAtZenithFinal, SunDiskOuterSpaceLuminance);
+}
+
+
 
 /*=============================================================================
 	FSkyAtmosphereRenderSceneInfo implementation.
 =============================================================================*/
 
 
-
-FSkyAtmosphereRenderSceneInfo::FSkyAtmosphereRenderSceneInfo(const USkyAtmosphereComponent* InComponent)
-	: bStaticLightingBuilt(false)
-	, Component_DoNotDereference(InComponent)
-	, AtmosphereSetup(*InComponent)
-	, TransmittanceAtZenith(AtmosphereSetup.GetTransmittanceAtGroundLevel(FVector(0.0f, 0.0f, 1.0f)))
+FSkyAtmosphereRenderSceneInfo::FSkyAtmosphereRenderSceneInfo(FSkyAtmosphereSceneProxy& SkyAtmosphereSceneProxyIn)
+	:SkyAtmosphereSceneProxy(SkyAtmosphereSceneProxyIn)
 {
-	SkyLuminanceFactor = InComponent->SkyLuminanceFactor;
-	AerialPespectiveViewDistanceScale = InComponent->AerialPespectiveViewDistanceScale;
-	memset(OverrideAtmosphericLight, 0, sizeof(OverrideAtmosphericLight));
-
 	// Create a multiframe uniform buffer. A render command is used because FSkyAtmosphereRenderSceneInfo ctor is called on the Game thread.
 	TUniformBufferRef<FAtmosphereUniformShaderParameters>* AtmosphereUniformBufferPtr = &AtmosphereUniformBuffer;
 	FAtmosphereUniformShaderParameters* AtmosphereUniformShaderParametersPtr = &AtmosphereUniformShaderParameters;
-	CopyAtmosphereSetupToUniformShaderParameters(AtmosphereUniformShaderParameters, AtmosphereSetup);
+	CopyAtmosphereSetupToUniformShaderParameters(AtmosphereUniformShaderParameters, SkyAtmosphereSceneProxy.GetAtmosphereSetup());
 	ENQUEUE_RENDER_COMMAND(FCreateUniformBuffer)(
 		[AtmosphereUniformBufferPtr, AtmosphereUniformShaderParametersPtr](FRHICommandListImmediate& RHICmdList)
 	{
@@ -365,21 +375,6 @@ FSkyAtmosphereRenderSceneInfo::~FSkyAtmosphereRenderSceneInfo()
 {
 }
 
-void FSkyAtmosphereRenderSceneInfo::PrepareSunLightProxy(uint32 AtmosphereLightIndex, FLightSceneInfo& AtmosphereLight) const
-{
-	// See explanation in https://media.contentapi.ea.com/content/dam/eacom/frostbite/files/s2016-pbs-frostbite-sky-clouds-new.pdf page 26
-	const bool bAtmosphereAffectsSunIlluminance = true;
-	const FVector AtmosphereLightDirection = GetAtmosphereLightDirection(AtmosphereLightIndex, -AtmosphereLight.Proxy->GetDirection());
-	FLinearColor TransmittanceTowardSun = bAtmosphereAffectsSunIlluminance ? AtmosphereSetup.GetTransmittanceAtGroundLevel(AtmosphereLightDirection) : FLinearColor(FLinearColor::White);
-	FLinearColor TransmittanceAtZenithFinal = bAtmosphereAffectsSunIlluminance ? TransmittanceAtZenith : FLinearColor(FLinearColor::White);
-
-	FLinearColor SunZenithIlluminance = AtmosphereLight.Proxy->GetColor();
-	FLinearColor SunOuterSpaceIlluminance = SunZenithIlluminance / TransmittanceAtZenithFinal;
-	FLinearColor SunDiskOuterSpaceLuminance = GetLightDiskLuminance(AtmosphereLight, SunOuterSpaceIlluminance);
-
-	AtmosphereLight.Proxy->SetAtmosphereRelatedProperties(TransmittanceTowardSun / TransmittanceAtZenithFinal, SunDiskOuterSpaceLuminance);
-}
-
 FTextureRHIRef FSkyAtmosphereRenderSceneInfo::GetDistantSkyLightLutTextureRHI()
 {
 	if (CVarSkyAtmosphereDistantSkyLightLUT.GetValueOnRenderThread() > 0)
@@ -389,85 +384,72 @@ FTextureRHIRef FSkyAtmosphereRenderSceneInfo::GetDistantSkyLightLutTextureRHI()
 	return GBlackTexture->TextureRHI;
 }
 
-void FSkyAtmosphereRenderSceneInfo::OverrideAtmosphereLightDirection(const class USkyAtmosphereComponent* SkyAtmosphereComponent, int32 AtmosphereLightIndex, const FVector& LightDirection)
-{
-	if (SkyAtmosphereComponent == Component_DoNotDereference)
-	{
-		check(AtmosphereLightIndex >= 0 && AtmosphereLightIndex < NUM_ATMOSPHERE_LIGHTS);
-		AtmosphereLightIndex = FMath::Clamp(AtmosphereLightIndex, 0, NUM_ATMOSPHERE_LIGHTS - 1);	// To make sure we do not crash, blueprint function cannot enforce input ranges
-		OverrideAtmosphericLight[AtmosphereLightIndex] = true;
-		OverrideAtmosphericLightDirection[AtmosphereLightIndex] = LightDirection;
-	}
-}
 
-FVector FSkyAtmosphereRenderSceneInfo::GetAtmosphereLightDirection(int32 AtmosphereLightIndex, const FVector& DefaultDirection) const
-{
-	if (OverrideAtmosphericLight[AtmosphereLightIndex])
-	{
-		return OverrideAtmosphericLightDirection[AtmosphereLightIndex];
-	}
-	return DefaultDirection;
-}
 
 /*=============================================================================
 	FScene functions
 =============================================================================*/
 
-void FScene::AddSkyAtmosphere(const USkyAtmosphereComponent* SkyAtmosphereComponent, bool bStaticLightingBuilt)
-{
-	check(SkyAtmosphereComponent);
 
-	FSkyAtmosphereRenderSceneInfo* SceneInfo = new FSkyAtmosphereRenderSceneInfo(SkyAtmosphereComponent);
+
+void FScene::AddSkyAtmosphere(FSkyAtmosphereSceneProxy* SkyAtmosphereSceneProxy, bool bStaticLightingBuilt)
+{
+	check(SkyAtmosphereSceneProxy);
 	FScene* Scene = this;
+
 	ENQUEUE_RENDER_COMMAND(FAddSkyAtmosphereCommand)(
-		[Scene, SceneInfo, bStaticLightingBuilt](FRHICommandListImmediate& RHICmdList)
+		[Scene, SkyAtmosphereSceneProxy, bStaticLightingBuilt](FRHICommandListImmediate& RHICmdList)
 		{
-			if (Scene->SkyAtmosphere)
-			{
-				delete Scene->SkyAtmosphere;
-			}
-			Scene->SkyAtmosphere = SceneInfo;
-			Scene->SkyAtmosphere->bStaticLightingBuilt = bStaticLightingBuilt;
-			if (!Scene->SkyAtmosphere->bStaticLightingBuilt)
+			check(!Scene->SkyAtmosphereStack.Contains(SkyAtmosphereSceneProxy));
+			Scene->SkyAtmosphereStack.Push(SkyAtmosphereSceneProxy);
+
+			SkyAtmosphereSceneProxy->RenderSceneInfo = new FSkyAtmosphereRenderSceneInfo(*SkyAtmosphereSceneProxy);
+
+			// Use the most recently enabled SkyAtmosphere
+			Scene->SkyAtmosphere = SkyAtmosphereSceneProxy->RenderSceneInfo;
+			SkyAtmosphereSceneProxy->bStaticLightingBuilt = bStaticLightingBuilt;
+			if (!SkyAtmosphereSceneProxy->bStaticLightingBuilt)
 			{
 				FPlatformAtomics::InterlockedIncrement(&Scene->NumUncachedStaticLightingInteractions);
 			}
 		} );
 }
 
-void FScene::RemoveSkyAtmosphere(const USkyAtmosphereComponent* SkyAtmosphereComponent)
+void FScene::RemoveSkyAtmosphere(FSkyAtmosphereSceneProxy* SkyAtmosphereSceneProxy)
 {
+	check(SkyAtmosphereSceneProxy);
 	FScene* Scene = this;
-	check(SkyAtmosphereComponent);
 
 	ENQUEUE_RENDER_COMMAND(FRemoveSkyAtmosphereCommand)(
-		[Scene](FRHICommandListImmediate& RHICmdList)
+		[Scene, SkyAtmosphereSceneProxy](FRHICommandListImmediate& RHICmdList)
 		{
-			if (Scene->SkyAtmosphere)
+			if (!SkyAtmosphereSceneProxy->bStaticLightingBuilt)
 			{
-				if (Scene->SkyAtmosphere && !Scene->SkyAtmosphere->bStaticLightingBuilt)
-				{
-					FPlatformAtomics::InterlockedDecrement(&Scene->NumUncachedStaticLightingInteractions);
-				}
-				delete Scene->SkyAtmosphere;
-				Scene->SkyAtmosphere = NULL;
+				FPlatformAtomics::InterlockedDecrement(&Scene->NumUncachedStaticLightingInteractions);
+			}
+			delete SkyAtmosphereSceneProxy->RenderSceneInfo;
+			Scene->SkyAtmosphereStack.RemoveSingle(SkyAtmosphereSceneProxy);
+
+			if (Scene->SkyAtmosphereStack.Num() > 0)
+			{
+				// Use the most recently enabled SkyAtmosphere
+				Scene->SkyAtmosphere = Scene->SkyAtmosphereStack.Last()->RenderSceneInfo;
+			}
+			else
+			{
+				Scene->SkyAtmosphere = nullptr;
 			}
 		} );
 }
 
-void FScene::OverrideSkyAtmosphereLightDirection(const class USkyAtmosphereComponent* SkyAtmosphereComponent, int32 AtmosphereLightIndex, const FVector& InLightDirection)
+void FScene::OverrideSkyAtmosphereLightDirection(FSkyAtmosphereSceneProxy* SkyAtmosphereSceneProxy, int32 AtmosphereLightIndex, const FVector& InLightDirection)
 {
-	FScene* Scene = this;
 	FVector LightDirection = InLightDirection;
 	LightDirection.Normalize();
 	ENQUEUE_RENDER_COMMAND(FOverrideSkyAtmosphereLightDirection)(
-		[Scene, SkyAtmosphereComponent, AtmosphereLightIndex, LightDirection](FRHICommandListImmediate& RHICmdList)
+		[SkyAtmosphereSceneProxy, AtmosphereLightIndex, LightDirection](FRHICommandListImmediate& RHICmdList)
 	{
-		FSkyAtmosphereRenderSceneInfo* SkyInfo = Scene->GetSkyAtmosphereSceneInfo();
-		if (SkyInfo)
-		{
-			SkyInfo->OverrideAtmosphereLightDirection(SkyAtmosphereComponent, AtmosphereLightIndex, LightDirection);
-		}
+		SkyAtmosphereSceneProxy->OverrideAtmosphereLightDirection(AtmosphereLightIndex, LightDirection);
 	});
 }
 
@@ -1023,8 +1005,9 @@ static void SetupSkyAtmosphereInternalCommonParameters(
 	InternalCommonParameters.TransmittanceSampleCount = CVarSkyAtmosphereTranstmittanceLUTSampleCount.GetValueOnRenderThread();
 	InternalCommonParameters.MultiScatteringSampleCount = CVarSkyAtmosphereMultiScatteringLUTSampleCount.GetValueOnRenderThread();
 
-	InternalCommonParameters.SkyLuminanceFactor = FVector(SkyInfo.GetSkyLuminanceFactor());
-	InternalCommonParameters.AerialPespectiveViewDistanceScale = SkyInfo.GetAerialPespectiveViewDistanceScale();
+	const FSkyAtmosphereSceneProxy& SkyAtmosphereSceneProxy = SkyInfo.GetSkyAtmosphereSceneProxy();
+	InternalCommonParameters.SkyLuminanceFactor = FVector(SkyAtmosphereSceneProxy.GetSkyLuminanceFactor());
+	InternalCommonParameters.AerialPespectiveViewDistanceScale = SkyAtmosphereSceneProxy.GetAerialPespectiveViewDistanceScale();
 
 	auto ValidateDistanceValue = [](float& Value)
 	{
@@ -1070,7 +1053,9 @@ void FSceneRenderer::RenderSkyAtmosphereLookUpTables(FRHICommandListImmediate& R
 	SCOPED_GPU_STAT(RHICmdList, SkyAtmosphereLUTs);
 
 	FSkyAtmosphereRenderSceneInfo& SkyInfo = *Scene->GetSkyAtmosphereSceneInfo();
-	const bool bMultiScattering = SkyInfo.IsMultiScatteringEnabled();
+	const FSkyAtmosphereSceneProxy& SkyAtmosphereSceneProxy = SkyInfo.GetSkyAtmosphereSceneProxy();
+
+	const bool bMultiScattering = SkyAtmosphereSceneProxy.IsMultiScatteringEnabled();
 	const bool bFastSky = CVarSkyAtmosphereFastSkyLUT.GetValueOnRenderThread() > 0;
 	const bool bFastAerialPespective = CVarSkyAtmosphereAerialPerspectiveApplyOnOpaque.GetValueOnRenderThread() > 0;
 	const bool bSecondAtmosphereLightEnabled = IsSecondAtmosphereLightEnabled(Scene);
@@ -1239,8 +1224,10 @@ void FSceneRenderer::RenderSkyAtmosphere(FRHICommandListImmediate& RHICmdList)
 	SCOPED_GPU_STAT(RHICmdList, SkyAtmosphere);
 
 	FSkyAtmosphereRenderSceneInfo& SkyInfo = *Scene->GetSkyAtmosphereSceneInfo();
-	const FAtmosphereSetup& Atmosphere = SkyInfo.GetAtmosphereSetup();
-	const bool bMultiScattering = SkyInfo.IsMultiScatteringEnabled();
+	const FSkyAtmosphereSceneProxy& SkyAtmosphereSceneProxy = SkyInfo.GetSkyAtmosphereSceneProxy();
+
+	const FAtmosphereSetup& Atmosphere = SkyAtmosphereSceneProxy.GetAtmosphereSetup();
+	const bool bMultiScattering = SkyAtmosphereSceneProxy.IsMultiScatteringEnabled();
 	const bool bFastSky = CVarSkyAtmosphereFastSkyLUT.GetValueOnRenderThread() > 0;
 	const bool bFastAerialPerspective = CVarSkyAtmosphereAerialPerspectiveApplyOnOpaque.GetValueOnRenderThread() > 0;
 	const bool bFastAerialPerspectiveDepthTest = CVarSkyAtmosphereAerialPerspectiveDepthTest.GetValueOnRenderThread() > 0;
@@ -1490,12 +1477,14 @@ void FDeferredShadingSceneRenderer::RenderDebugSkyAtmosphere(FRHICommandListImme
 
 	const int32 SkyAtmosphereVisualize = CVarSkyAtmosphereVisualize.GetValueOnRenderThread();
 	FSkyAtmosphereRenderSceneInfo& SkyInfo = *Scene->GetSkyAtmosphereSceneInfo();
-	const FAtmosphereSetup& Atmosphere = SkyInfo.GetAtmosphereSetup();
+	const FSkyAtmosphereSceneProxy& SkyAtmosphereSceneProxy = SkyInfo.GetSkyAtmosphereSceneProxy();
+
+	const FAtmosphereSetup& Atmosphere = SkyAtmosphereSceneProxy.GetAtmosphereSetup();
 	FSceneRenderTargets& SceneContext = FSceneRenderTargets::Get(RHICmdList);
 
 	if (SkyAtmosphereVisualize)
 	{
-		const bool bMultiScattering = SkyInfo.IsMultiScatteringEnabled();
+		const bool bMultiScattering = SkyAtmosphereSceneProxy.IsMultiScatteringEnabled();
 		const bool bFastSky = CVarSkyAtmosphereFastSkyLUT.GetValueOnRenderThread() > 0;
 		const bool bFastAerialPespective = CVarSkyAtmosphereAerialPerspectiveApplyOnOpaque.GetValueOnRenderThread() > 0;
 
diff --git a/Engine/Source/Runtime/Renderer/Private/SkyAtmosphereRendering.h b/Engine/Source/Runtime/Renderer/Private/SkyAtmosphereRendering.h
--- a/Engine/Source/Runtime/Renderer/Private/SkyAtmosphereRendering.h
+++ b/Engine/Source/Runtime/Renderer/Private/SkyAtmosphereRendering.h
@@ -17,6 +17,7 @@ class FScene;
 class FViewInfo;
 class FLightSceneInfo;
 class USkyAtmosphereComponent;
+class FSkyAtmosphereSceneProxy;
 
 struct FEngineShowFlags;
 
@@ -61,49 +62,32 @@ class FSkyAtmosphereRenderSceneInfo
 public:
 
 	/** Initialization constructor. */
-	explicit FSkyAtmosphereRenderSceneInfo(const USkyAtmosphereComponent* InComponent);
+	explicit FSkyAtmosphereRenderSceneInfo(FSkyAtmosphereSceneProxy& SkyAtmosphereSceneProxy);
 	~FSkyAtmosphereRenderSceneInfo();
 
-	/** Prepare the sun light data as a function of current atmosphere state. */
-	void PrepareSunLightProxy(uint32 AtmosphereLightIndex, FLightSceneInfo& AtmosphereLight) const;
-
-	bool IsMultiScatteringEnabled() const { return AtmosphereSetup.MultiScatteringFactor > 0.0f; }
-	FLinearColor GetSkyLuminanceFactor() const { return SkyLuminanceFactor; }
-	float GetAerialPespectiveViewDistanceScale() const { return AerialPespectiveViewDistanceScale; }
-
 	const TUniformBufferRef<FAtmosphereUniformShaderParameters>& GetAtmosphereUniformBuffer() { return AtmosphereUniformBuffer; }
 	TRefCountPtr<IPooledRenderTarget>& GetTransmittanceLutTexture() { return TransmittanceLutTexture; }
 	TRefCountPtr<IPooledRenderTarget>& GetMultiScatteredLuminanceLutTexture() { return MultiScatteredLuminanceLutTexture; }
 	TRefCountPtr<IPooledRenderTarget>& GetDistantSkyLightLutTexture() { return DistantSkyLightLutTexture; }
 	FTextureRHIRef GetDistantSkyLightLutTextureRHI();
 
-	const FAtmosphereSetup& GetAtmosphereSetup() const { return AtmosphereSetup; }
 	const FAtmosphereUniformShaderParameters* GetAtmosphereShaderParameters() const { return &AtmosphereUniformShaderParameters; }
-
-	void OverrideAtmosphereLightDirection(const class USkyAtmosphereComponent* SkyAtmosphereComponent, int32 AtmosphereLightIndex, const FVector& LightDirection);
-	FVector GetAtmosphereLightDirection(int32 AtmosphereLightIndex, const FVector& DefaultDirection) const;
-
-	bool bStaticLightingBuilt;
+	const FSkyAtmosphereSceneProxy& GetSkyAtmosphereSceneProxy() const { return SkyAtmosphereSceneProxy; }
 
 private:
 
-	const USkyAtmosphereComponent* Component_DoNotDereference;
+	FSkyAtmosphereSceneProxy& SkyAtmosphereSceneProxy;
 
-	FAtmosphereSetup AtmosphereSetup;
 	FAtmosphereUniformShaderParameters AtmosphereUniformShaderParameters;
 
 	TUniformBufferRef<FAtmosphereUniformShaderParameters> AtmosphereUniformBuffer;
+
 	TRefCountPtr<IPooledRenderTarget> TransmittanceLutTexture;
 	TRefCountPtr<IPooledRenderTarget> MultiScatteredLuminanceLutTexture;
 	TRefCountPtr<IPooledRenderTarget> DistantSkyLightLutTexture;
+};
 
-	FLinearColor TransmittanceAtZenith;
-	FLinearColor SkyLuminanceFactor;
-	float AerialPespectiveViewDistanceScale;
 
-	bool OverrideAtmosphericLight[NUM_ATMOSPHERE_LIGHTS];
-	FVector OverrideAtmosphericLightDirection[NUM_ATMOSPHERE_LIGHTS];
-};
 
 bool ShouldRenderSkyAtmosphere(const FScene* Scene, const FEngineShowFlags& EngineShowFlags);
 
@@ -112,4 +96,6 @@ void InitSkyAtmosphereForView(FRHICommandListImmediate& RHICmdList, const FScene
 
 extern void SetupSkyAtmosphereViewSharedUniformShaderParameters(const class FViewInfo& View, FSkyAtmosphereViewSharedUniformShaderParameters& OutParameters);
 
+// Prepare the sun light data as a function of the atmosphere state. 
+void PrepareSunLightProxy(const FSkyAtmosphereRenderSceneInfo& SkyAtmosphere, uint32 AtmosphereLightIndex, FLightSceneInfo& AtmosphereLight);
 
