Another way to fix sky+tilerenderer without resource allocation for the sky/scene.
Tested on apollo map and others

diff --git a/Engine/Source/Runtime/Renderer/Private/Renderer.cpp b/Engine/Source/Runtime/Renderer/Private/Renderer.cpp
--- a/Engine/Source/Runtime/Renderer/Private/Renderer.cpp
+++ b/Engine/Source/Runtime/Renderer/Private/Renderer.cpp
@@ -152,13 +152,6 @@ void FRendererModule::DrawTileMesh(FRHICommandListImmediate& RHICmdList, FMeshPa
 			}
 		}
 
-		// Initialise Sky/View resources before the view global uniform buffer is built.
-		if (ShouldRenderSkyAtmosphere(Scene, View.Family->EngineShowFlags))
-		{
-			InitSkyAtmosphereForScene(RHICmdList, Scene);
-			InitSkyAtmosphereForView(RHICmdList, Scene, View);
-		}
-
 		View.InitRHIResources();
 		DrawRenderState.SetViewUniformBuffer(View.ViewUniformBuffer);
 
diff --git a/Engine/Source/Runtime/Renderer/Private/SceneRendering.cpp b/Engine/Source/Runtime/Renderer/Private/SceneRendering.cpp
--- a/Engine/Source/Runtime/Renderer/Private/SceneRendering.cpp
+++ b/Engine/Source/Runtime/Renderer/Private/SceneRendering.cpp
@@ -1230,17 +1230,37 @@ void FViewInfo::SetupUniformBufferParameters(
 		check(Scene->SkyAtmosphere->GetTransmittanceLutTexture().IsValid());
 		FSkyAtmosphereRenderSceneInfo* SkyAtmosphere = Scene->SkyAtmosphere;
 		const FSkyAtmosphereSceneProxy& SkyAtmosphereSceneProxy = SkyAtmosphere->GetSkyAtmosphereSceneProxy();
+
+		// Get access to texture resource if we have valid pointer.
+		// (Valid pointer checks are needed because some resources might not have been initialized when coming from FCanvasTileRendererItem or FCanvasTriangleRendererItem)
+
 		const TRefCountPtr<IPooledRenderTarget>& PooledTransmittanceLutTexture = SkyAtmosphere->GetTransmittanceLutTexture();
-		TransmittanceLutTextureFound = PooledTransmittanceLutTexture->GetRenderTargetItem().ShaderResourceTexture;
+		if (PooledTransmittanceLutTexture.IsValid())
+		{
+			TransmittanceLutTextureFound = PooledTransmittanceLutTexture->GetRenderTargetItem().ShaderResourceTexture;
+		}
+		const TRefCountPtr<IPooledRenderTarget>& PooledDistantSkyLightLutTexture = SkyAtmosphere->GetDistantSkyLightLutTexture();
+		if (PooledDistantSkyLightLutTexture.IsValid())
+		{
+			DistantSkyLightLutTextureFound = PooledDistantSkyLightLutTexture->GetRenderTargetItem().ShaderResourceTexture;
+		}
 
-		SkyViewLutTextureFound = this->SkyAtmosphereViewLutTexture->GetRenderTargetItem().ShaderResourceTexture;
-		uint32 SkyViewLutWidth = float(this->SkyAtmosphereViewLutTexture->GetDesc().GetSize().X);
-		uint32 SkyViewLutHeight = float(this->SkyAtmosphereViewLutTexture->GetDesc().GetSize().Y);
+		if (this->SkyAtmosphereCameraAerialPerspectiveVolume.IsValid())
+		{
+			CameraAerialPerspectiveVolumeFound = this->SkyAtmosphereCameraAerialPerspectiveVolume->GetRenderTargetItem().ShaderResourceTexture;
+		}
+
+		float SkyViewLutWidth = 1.0f;
+		float SkyViewLutHeight = 1.0f;
+		if (this->SkyAtmosphereViewLutTexture.IsValid())
+		{
+			SkyViewLutTextureFound = this->SkyAtmosphereViewLutTexture->GetRenderTargetItem().ShaderResourceTexture;
+			SkyViewLutWidth = float(this->SkyAtmosphereViewLutTexture->GetDesc().GetSize().X);
+			SkyViewLutHeight = float(this->SkyAtmosphereViewLutTexture->GetDesc().GetSize().Y);
+		}
 		ViewUniformShaderParameters.SkyViewLutSizeAndInvSize = FVector4(SkyViewLutWidth, SkyViewLutHeight, 1.0f / SkyViewLutWidth, 1.0f / SkyViewLutHeight);
 
-		CameraAerialPerspectiveVolumeFound = this->SkyAtmosphereCameraAerialPerspectiveVolume->GetRenderTargetItem().ShaderResourceTexture;
-		DistantSkyLightLutTextureFound = SkyAtmosphere->GetDistantSkyLightLutTextureRHI();
-		ViewUniformShaderParameters.SkyAtmosphereSkyLuminanceFactor = SkyAtmosphereSceneProxy.GetSkyLuminanceFactor();
+		// Now initialize remaining view parameters.
 
 		const FAtmosphereSetup& AtmosphereSetup = SkyAtmosphereSceneProxy.GetAtmosphereSetup();
 		ViewUniformShaderParameters.SkyAtmosphereBottomRadius = AtmosphereSetup.BottomRadius;
@@ -1254,6 +1274,7 @@ void FViewInfo::SetupUniformBufferParameters(
 		ViewUniformShaderParameters.SkyAtmosphereCameraAerialPerspectiveVolumeDepthSliceLength = OutParameters.CameraAerialPerspectiveVolumeDepthSliceLength;
 		ViewUniformShaderParameters.SkyAtmosphereCameraAerialPerspectiveVolumeDepthSliceLengthInv = OutParameters.CameraAerialPerspectiveVolumeDepthSliceLengthInv;
 		ViewUniformShaderParameters.SkyAtmosphereApplyCameraAerialPerspectiveVolume = OutParameters.ApplyCameraAerialPerspectiveVolume;
+		ViewUniformShaderParameters.SkyAtmosphereSkyLuminanceFactor = SkyAtmosphereSceneProxy.GetSkyLuminanceFactor();
 
 		// Fill atmosphere lights shader parameters
 		for (uint8 Index = 0; Index < NUM_ATMOSPHERE_LIGHTS; ++Index)
diff --git a/Engine/Source/Runtime/Renderer/Private/SkyAtmosphereRendering.cpp b/Engine/Source/Runtime/Renderer/Private/SkyAtmosphereRendering.cpp
--- a/Engine/Source/Runtime/Renderer/Private/SkyAtmosphereRendering.cpp
+++ b/Engine/Source/Runtime/Renderer/Private/SkyAtmosphereRendering.cpp
@@ -375,13 +375,13 @@ FSkyAtmosphereRenderSceneInfo::~FSkyAtmosphereRenderSceneInfo()
 {
 }
 
-FTextureRHIRef FSkyAtmosphereRenderSceneInfo::GetDistantSkyLightLutTextureRHI()
+TRefCountPtr<IPooledRenderTarget>& FSkyAtmosphereRenderSceneInfo::GetDistantSkyLightLutTexture()
 {
 	if (CVarSkyAtmosphereDistantSkyLightLUT.GetValueOnRenderThread() > 0)
 	{
-		return DistantSkyLightLutTexture->GetRenderTargetItem().ShaderResourceTexture;
+		return DistantSkyLightLutTexture;
 	}
-	return GBlackTexture->TextureRHI;
+	return GSystemTextures.BlackDummy;
 }
 
 
@@ -1073,8 +1073,6 @@ void FSceneRenderer::RenderSkyAtmosphereLookUpTables(FRHICommandListImmediate& R
 	FRDGTextureUAVRef TransmittanceLutUAV = GraphBuilder.CreateUAV(FRDGTextureUAVDesc(TransmittanceLut, 0));
 	FRDGTextureRef MultiScatteredLuminanceLut = GraphBuilder.RegisterExternalTexture(SkyInfo.GetMultiScatteredLuminanceLutTexture());
 	FRDGTextureUAVRef MultiScatteredLuminanceLutUAV = GraphBuilder.CreateUAV(FRDGTextureUAVDesc(MultiScatteredLuminanceLut, 0));
-	FRDGTextureRef DistantSkyLightLut = GraphBuilder.RegisterExternalTexture(SkyInfo.GetDistantSkyLightLutTexture());
-	FRDGTextureUAVRef DistantSkyLightLutUAV = GraphBuilder.CreateUAV(FRDGTextureUAVDesc(DistantSkyLightLut, 0));
 
 	// Transmittance LUT
 	TShaderMap<FGlobalShaderType>* GlobalShaderMap = GetGlobalShaderMap(FeatureLevel);
@@ -1112,6 +1110,9 @@ void FSceneRenderer::RenderSkyAtmosphereLookUpTables(FRHICommandListImmediate& R
 	// Distant Sky Light LUT
 	if(CVarSkyAtmosphereDistantSkyLightLUT.GetValueOnRenderThread() > 0)
 	{
+		FRDGTextureRef DistantSkyLightLut = GraphBuilder.RegisterExternalTexture(SkyInfo.GetDistantSkyLightLutTexture());
+		FRDGTextureUAVRef DistantSkyLightLutUAV = GraphBuilder.CreateUAV(FRDGTextureUAVDesc(DistantSkyLightLut, 0));
+
 		FRenderDistantSkyLightLutCS::FPermutationDomain PermutationVector;
 		PermutationVector.Set<FSkyPermutationMultiScatteringApprox>(bMultiScattering);
 		PermutationVector.Set<FSecondAtmosphereLight>(bSecondAtmosphereLightEnabled);
diff --git a/Engine/Source/Runtime/Renderer/Private/SkyAtmosphereRendering.h b/Engine/Source/Runtime/Renderer/Private/SkyAtmosphereRendering.h
--- a/Engine/Source/Runtime/Renderer/Private/SkyAtmosphereRendering.h
+++ b/Engine/Source/Runtime/Renderer/Private/SkyAtmosphereRendering.h
@@ -68,8 +68,7 @@ public:
 	const TUniformBufferRef<FAtmosphereUniformShaderParameters>& GetAtmosphereUniformBuffer() { return AtmosphereUniformBuffer; }
 	TRefCountPtr<IPooledRenderTarget>& GetTransmittanceLutTexture() { return TransmittanceLutTexture; }
 	TRefCountPtr<IPooledRenderTarget>& GetMultiScatteredLuminanceLutTexture() { return MultiScatteredLuminanceLutTexture; }
-	TRefCountPtr<IPooledRenderTarget>& GetDistantSkyLightLutTexture() { return DistantSkyLightLutTexture; }
-	FTextureRHIRef GetDistantSkyLightLutTextureRHI();
+	TRefCountPtr<IPooledRenderTarget>& GetDistantSkyLightLutTexture();
 
 	const FAtmosphereUniformShaderParameters* GetAtmosphereShaderParameters() const { return &AtmosphereUniformShaderParameters; }
 	const FSkyAtmosphereSceneProxy& GetSkyAtmosphereSceneProxy() const { return SkyAtmosphereSceneProxy; }
