Remove check() firing in FVirtualTextureSystem::FlushCache and test virtual texture producer handle is still valid before flush.
There are cases where the handle will not be valid.
For example when we have two RVT volumes pointing at the same RVT (not a good thing to do, but it can happen).
In that case the RVT will have two producers created and one discarded.

diff --git a/Engine/Source/Runtime/Renderer/Private/VT/VirtualTextureSystem.cpp b/Engine/Source/Runtime/Renderer/Private/VT/VirtualTextureSystem.cpp
--- a/Engine/Source/Runtime/Renderer/Private/VT/VirtualTextureSystem.cpp
+++ b/Engine/Source/Runtime/Renderer/Private/VT/VirtualTextureSystem.cpp
@@ -253,27 +253,29 @@ void FVirtualTextureSystem::FlushCache(FVirtualTextureProducerHandle const& Prod
 	INC_DWORD_STAT_BY(STAT_NumFlushCache, 1);
 
 	FVirtualTextureProducer const* Producer = Producers.FindProducer(ProducerHandle);
-	check(Producer != nullptr);
-	FVTProducerDescription const& ProducerDescription = Producer->GetDescription();
-
-	TArray<FVirtualTexturePhysicalSpace*> PhysicalSpacesForProducer;
-	for (uint32 i = 0; i < Producer->GetNumLayers(); ++i)
+	if (Producer != nullptr)
 	{
-		PhysicalSpacesForProducer.AddUnique(Producer->GetPhysicalSpace(i));
-	}
+		FVTProducerDescription const& ProducerDescription = Producer->GetDescription();
 
-	// Don't resize to allow this container to grow as needed (avoid allocations when collecting)
-	TransientCollectedPages.Reset();
+		TArray<FVirtualTexturePhysicalSpace*> PhysicalSpacesForProducer;
+		for (uint32 i = 0; i < Producer->GetNumLayers(); ++i)
+		{
+			PhysicalSpacesForProducer.AddUnique(Producer->GetPhysicalSpace(i));
+		}
 
-	for (int32 i = 0; i < PhysicalSpacesForProducer.Num(); ++i)
-	{
-		FTexturePagePool& Pool = PhysicalSpacesForProducer[i]->GetPagePool();
-		Pool.EvictPages(this, ProducerHandle, ProducerDescription, TextureRegion, MaxLevel, TransientCollectedPages);
-	}
+		// Don't resize to allow this container to grow as needed (avoid allocations when collecting)
+		TransientCollectedPages.Reset();
 
-	for (auto& Page : TransientCollectedPages)
-	{
-		MappedTilesToProduce.Add(Page);
+		for (int32 i = 0; i < PhysicalSpacesForProducer.Num(); ++i)
+		{
+			FTexturePagePool& Pool = PhysicalSpacesForProducer[i]->GetPagePool();
+			Pool.EvictPages(this, ProducerHandle, ProducerDescription, TextureRegion, MaxLevel, TransientCollectedPages);
+		}
+
+		for (auto& Page : TransientCollectedPages)
+		{
+			MappedTilesToProduce.Add(Page);
+		}
 	}
 }
 
