Make runtime virtual texture page clear before render optional

diff --git a/Engine/Source/Editor/VirtualTexturingEditor/Private/RuntimeVirtualTextureBuild.cpp b/Engine/Source/Editor/VirtualTexturingEditor/Private/RuntimeVirtualTextureBuild.cpp
--- a/Engine/Source/Editor/VirtualTexturingEditor/Private/RuntimeVirtualTextureBuild.cpp
+++ b/Engine/Source/Editor/VirtualTexturingEditor/Private/RuntimeVirtualTextureBuild.cpp
@@ -110,6 +110,7 @@ namespace RuntimeVirtualTexture
 						Scene->GetRenderScene(),
 						1 << VirtualTextureSceneIndex,
 						MaterialType,
+						true,
 						RenderTargetResource[0] != nullptr ? RenderTargetResource[0]->GetRenderTargetTexture() : nullptr, TileBox,
 						RenderTargetResource[1] != nullptr ? RenderTargetResource[1]->GetRenderTargetTexture() : nullptr, TileBox,
 						Transform,
diff --git a/Engine/Source/Editor/VirtualTexturingEditor/Private/RuntimeVirtualTextureThumbnailRenderer.cpp b/Engine/Source/Editor/VirtualTexturingEditor/Private/RuntimeVirtualTextureThumbnailRenderer.cpp
--- a/Engine/Source/Editor/VirtualTexturingEditor/Private/RuntimeVirtualTextureThumbnailRenderer.cpp
+++ b/Engine/Source/Editor/VirtualTexturingEditor/Private/RuntimeVirtualTextureThumbnailRenderer.cpp
@@ -84,6 +84,7 @@ void URuntimeVirtualTextureThumbnailRenderer::Draw(UObject* Object, int32 X, int
 			Scene->GetRenderScene(),
 			1 << VirtualTextureSceneIndex,
 			MaterialType,
+			true,
 			RenderTarget->GetRenderTargetTexture(), DestRect,
 			nullptr, DestRect,
 			Transform,
diff --git a/Engine/Source/Runtime/Engine/Classes/VT/RuntimeVirtualTexture.h b/Engine/Source/Runtime/Engine/Classes/VT/RuntimeVirtualTexture.h
--- a/Engine/Source/Runtime/Engine/Classes/VT/RuntimeVirtualTexture.h
+++ b/Engine/Source/Runtime/Engine/Classes/VT/RuntimeVirtualTexture.h
@@ -24,6 +24,10 @@ protected:
 	UPROPERTY(EditAnywhere, BlueprintReadOnly, Category = Layout, meta = (DisplayName = "Enable BC texture compression"))
 	bool bCompressTextures = true;
 
+	/** Enable clear before rendering a page of the virtual texture. Disabling this can be an optimization if you know that the texture will always be fully covered by rendering.  */
+	UPROPERTY(EditAnywhere, BlueprintReadOnly, AdvancedDisplay, Category = Layout, meta = (DisplayName = "Enable clear before render"))
+	bool bClearTextures = true;
+
 	/** Size of virtual texture along the largest axis. (Actual values increase in powers of 2) */
 	UPROPERTY(EditAnywhere, BlueprintReadOnly, Category = Size, meta = (UIMin = "0", UIMax = "8", DisplayName = "Size of the virtual texture"))
 	int32 Size = 6; // 65536
@@ -37,7 +41,7 @@ protected:
 	int32 TileBorderSize = 2; // 4
 
 	/** Number of low mips to cut from the virtual texture. This can reduce peak virtual texture update cost but will also increase the probability of mip shimmering. */
-	UPROPERTY(EditAnywhere, BlueprintReadOnly, Category = LowMips, meta = (UIMin = "0", UIMax = "6", DisplayName = "Number of low mips to remove from the virtual texture"))
+	UPROPERTY(EditAnywhere, BlueprintReadOnly, AdvancedDisplay, Category = LowMips, meta = (UIMin = "0", UIMax = "6", DisplayName = "Number of low mips to remove from the virtual texture"))
 	int32 RemoveLowMips = 0;
 
 	/** Number of low mips to serialize and stream for the virtual texture. This can reduce rendering update cost. */
@@ -49,7 +53,7 @@ protected:
 	class URuntimeVirtualTextureStreamingProxy* StreamingTexture;
 
 	/** Enable usage of the virtual texture. When disabled there is no rendering into the virtual texture, and sampling will return zero values. */
-	UPROPERTY(EditAnywhere, BlueprintReadOnly, Category = Layout, meta = (DisplayName = "Enable virtual texture"))
+	UPROPERTY(EditAnywhere, BlueprintReadOnly, AdvancedDisplay, Category = Layout, meta = (DisplayName = "Enable virtual texture"))
 	bool bEnable = true;
 
 public:
@@ -84,6 +88,8 @@ public:
 	EPixelFormat GetLayerFormat(int32 LayerIndex) const;
 	/** Return true if the virtual texture layer should be sampled as sRGB */
 	bool IsLayerSRGB(int32 LayerIndex) const;
+	/** Returns true if texture pages should be cleared before render */
+	bool GetClearTextures() const { return bClearTextures; }
 
 	/** (Re)Initialize this object. Call this whenever we modify the producer or transform. */
 	void Initialize(IVirtualTexture* InProducer, FTransform const& VolumeToWorld);
diff --git a/Engine/Source/Runtime/Renderer/Private/VT/RuntimeVirtualTextureProducer.cpp b/Engine/Source/Runtime/Renderer/Private/VT/RuntimeVirtualTextureProducer.cpp
--- a/Engine/Source/Runtime/Renderer/Private/VT/RuntimeVirtualTextureProducer.cpp
+++ b/Engine/Source/Runtime/Renderer/Private/VT/RuntimeVirtualTextureProducer.cpp
@@ -8,11 +8,12 @@
 #include "VT/RuntimeVirtualTextureSceneProxy.h"
 
 
-FRuntimeVirtualTextureFinalizer::FRuntimeVirtualTextureFinalizer(FVTProducerDescription const& InDesc, uint32 InProducerId, ERuntimeVirtualTextureMaterialType InMaterialType, FSceneInterface* InScene, FTransform const& InUVToWorld)
+FRuntimeVirtualTextureFinalizer::FRuntimeVirtualTextureFinalizer(FVTProducerDescription const& InDesc, uint32 InProducerId, ERuntimeVirtualTextureMaterialType InMaterialType, bool InClearTextures, FSceneInterface* InScene, FTransform const& InUVToWorld)
 	: Desc(InDesc)
 	, ProducerId(InProducerId)
 	, RuntimeVirtualTextureMask(0)
 	, MaterialType(InMaterialType)
+	, bClearTextures(InClearTextures)
 	, Scene(InScene)
 	, UVToWorld(InUVToWorld)
 {
@@ -79,6 +80,7 @@ void FRuntimeVirtualTextureFinalizer::Finalize(FRHICommandListImmediate& RHICmdL
 			Scene->GetRenderScene(),
 			RuntimeVirtualTextureMask,
 			MaterialType,
+			bClearTextures,
 			Entry.Texture0, 
 			DestinationBox0, 
 			Entry.Texture1,
@@ -93,8 +95,8 @@ void FRuntimeVirtualTextureFinalizer::Finalize(FRHICommandListImmediate& RHICmdL
 	Tiles.SetNumUnsafeInternal(0);
 }
 
-FRuntimeVirtualTextureProducer::FRuntimeVirtualTextureProducer(FVTProducerDescription const& InDesc, uint32 InProducerId, ERuntimeVirtualTextureMaterialType InMaterialType, FSceneInterface* InScene, FTransform const& InUVToWorld)
-	: Finalizer(InDesc, InProducerId, InMaterialType, InScene, InUVToWorld)
+FRuntimeVirtualTextureProducer::FRuntimeVirtualTextureProducer(FVTProducerDescription const& InDesc, uint32 InProducerId, ERuntimeVirtualTextureMaterialType InMaterialType, bool InClearTextures, FSceneInterface* InScene, FTransform const& InUVToWorld)
+	: Finalizer(InDesc, InProducerId, InMaterialType, InClearTextures, InScene, InUVToWorld)
 {
 }
 
diff --git a/Engine/Source/Runtime/Renderer/Private/VT/RuntimeVirtualTextureProducer.h b/Engine/Source/Runtime/Renderer/Private/VT/RuntimeVirtualTextureProducer.h
--- a/Engine/Source/Runtime/Renderer/Private/VT/RuntimeVirtualTextureProducer.h
+++ b/Engine/Source/Runtime/Renderer/Private/VT/RuntimeVirtualTextureProducer.h
@@ -14,7 +14,7 @@ class FSceneInterface;
 class FRuntimeVirtualTextureFinalizer : public IVirtualTextureFinalizer
 {
 public:
-	FRuntimeVirtualTextureFinalizer(FVTProducerDescription const& InDesc, uint32 InProducerId, ERuntimeVirtualTextureMaterialType InMaterialType, FSceneInterface* InScene, FTransform const& InUVToWorld);
+	FRuntimeVirtualTextureFinalizer(FVTProducerDescription const& InDesc, uint32 InProducerId, ERuntimeVirtualTextureMaterialType InMaterialType, bool InClearTextures, FSceneInterface* InScene, FTransform const& InUVToWorld);
 	virtual ~FRuntimeVirtualTextureFinalizer() {}
 
 	/** A description for a single tile to render. */
@@ -52,6 +52,8 @@ private:
 	uint32 RuntimeVirtualTextureMask;
 	/** Contents of virtual texture layer stack. */
 	ERuntimeVirtualTextureMaterialType MaterialType;
+	/** Clear before render flag. */
+	bool bClearTextures;
 	/** Scene that the virtual texture is placed within. */
 	FSceneInterface* Scene;
 	/** Transform from UV space to world space. */
@@ -64,7 +66,7 @@ private:
 class FRuntimeVirtualTextureProducer : public IVirtualTexture
 {
 public:
-	RENDERER_API FRuntimeVirtualTextureProducer(FVTProducerDescription const& InDesc, uint32 InProducerId, ERuntimeVirtualTextureMaterialType InMaterialType, FSceneInterface* InScene, FTransform const& InUVToWorld);
+	RENDERER_API FRuntimeVirtualTextureProducer(FVTProducerDescription const& InDesc, uint32 InProducerId, ERuntimeVirtualTextureMaterialType InMaterialType, bool InClearTextures, FSceneInterface* InScene, FTransform const& InUVToWorld);
 	RENDERER_API virtual ~FRuntimeVirtualTextureProducer() {}
 
 	//~ Begin IVirtualTexture Interface.
diff --git a/Engine/Source/Runtime/Renderer/Private/VT/RuntimeVirtualTextureRender.cpp b/Engine/Source/Runtime/Renderer/Private/VT/RuntimeVirtualTextureRender.cpp
--- a/Engine/Source/Runtime/Renderer/Private/VT/RuntimeVirtualTextureRender.cpp
+++ b/Engine/Source/Runtime/Renderer/Private/VT/RuntimeVirtualTextureRender.cpp
@@ -566,7 +566,7 @@ namespace RuntimeVirtualTexture
 				if (bRenderPass)
 				{
 					OutputAlias0 = RenderTexture0 = GraphBuilder.CreateTexture(FPooledRenderTargetDesc::Create2DDesc(TextureSize, PF_B8G8R8A8, FClearValueBinding::Black, TexCreate_SRGB, TexCreate_RenderTargetable, false), TEXT("RenderTexture0"));
-					OutputAlias1 = RenderTexture1 = GraphBuilder.CreateTexture(FPooledRenderTargetDesc::Create2DDesc(TextureSize, PF_B8G8R8A8, FClearValueBinding::Black, TexCreate_None, TexCreate_RenderTargetable, false), TEXT("RenderTexture1"));
+					OutputAlias1 = RenderTexture1 = GraphBuilder.CreateTexture(FPooledRenderTargetDesc::Create2DDesc(TextureSize, PF_B8G8R8A8, FClearValueBinding::DefaultNormal8Bit, TexCreate_None, TexCreate_RenderTargetable, false), TEXT("RenderTexture1"));
 				}
 				if (bCompressPass)
 				{
@@ -578,8 +578,8 @@ namespace RuntimeVirtualTexture
 				if (bRenderPass)
 				{
 					OutputAlias0 = RenderTexture0 = GraphBuilder.CreateTexture(FPooledRenderTargetDesc::Create2DDesc(TextureSize, PF_B8G8R8A8, FClearValueBinding::Black, TexCreate_SRGB, TexCreate_RenderTargetable, false), TEXT("RenderTexture0"));
-					RenderTexture1 = GraphBuilder.CreateTexture(FPooledRenderTargetDesc::Create2DDesc(TextureSize, PF_B8G8R8A8, FClearValueBinding::Black, TexCreate_None, TexCreate_RenderTargetable, false), TEXT("RenderTexture1"));
-					RenderTexture2 = GraphBuilder.CreateTexture(FPooledRenderTargetDesc::Create2DDesc(TextureSize, PF_B8G8R8A8, FClearValueBinding::Black, TexCreate_None, TexCreate_RenderTargetable, false), TEXT("RenderTexture2"));
+					RenderTexture1 = GraphBuilder.CreateTexture(FPooledRenderTargetDesc::Create2DDesc(TextureSize, PF_B8G8R8A8, FClearValueBinding::DefaultNormal8Bit, TexCreate_None, TexCreate_RenderTargetable, false), TEXT("RenderTexture1"));
+					RenderTexture2 = GraphBuilder.CreateTexture(FPooledRenderTargetDesc::Create2DDesc(TextureSize, PF_B8G8R8A8, FClearValueBinding::DefaultNormal8Bit, TexCreate_None, TexCreate_RenderTargetable, false), TEXT("RenderTexture2"));
 				}
 				if (bCompressPass)
 				{
@@ -595,8 +595,8 @@ namespace RuntimeVirtualTexture
 				if (bRenderPass)
 				{
 					OutputAlias0 = RenderTexture0 = GraphBuilder.CreateTexture(FPooledRenderTargetDesc::Create2DDesc(TextureSize, PF_B8G8R8A8, FClearValueBinding::Black, TexCreate_SRGB, TexCreate_RenderTargetable, false), TEXT("RenderTexture0"));
-					RenderTexture1 = GraphBuilder.CreateTexture(FPooledRenderTargetDesc::Create2DDesc(TextureSize, PF_B8G8R8A8, FClearValueBinding::Black, TexCreate_None, TexCreate_RenderTargetable, false), TEXT("RenderTexture1"));
-					RenderTexture2 = GraphBuilder.CreateTexture(FPooledRenderTargetDesc::Create2DDesc(TextureSize, PF_B8G8R8A8, FClearValueBinding::Black, TexCreate_None, TexCreate_RenderTargetable, false), TEXT("RenderTexture2"));
+					RenderTexture1 = GraphBuilder.CreateTexture(FPooledRenderTargetDesc::Create2DDesc(TextureSize, PF_B8G8R8A8, FClearValueBinding::DefaultNormal8Bit, TexCreate_None, TexCreate_RenderTargetable, false), TEXT("RenderTexture1"));
+					RenderTexture2 = GraphBuilder.CreateTexture(FPooledRenderTargetDesc::Create2DDesc(TextureSize, PF_B8G8R8A8, FClearValueBinding::DefaultNormal8Bit, TexCreate_None, TexCreate_RenderTargetable, false), TEXT("RenderTexture2"));
 				}
 				if (bCompressPass)
 				{
@@ -653,6 +653,7 @@ namespace RuntimeVirtualTexture
 		FScene* Scene,
 		uint32 RuntimeVirtualTextureMask,
 		ERuntimeVirtualTextureMaterialType MaterialType,
+		bool bClearTextures,
 		FRHITexture2D* OutputTexture0,
 		FBox2D const& DestBox0,
 		FRHITexture2D* OutputTexture1,
@@ -723,10 +724,11 @@ namespace RuntimeVirtualTexture
 		// Draw Pass
 		if (GraphSetup.bRenderPass)
 		{
+			ERenderTargetLoadAction LoadAction = bClearTextures ? ERenderTargetLoadAction::EClear : ERenderTargetLoadAction::ENoAction;
 			FShader_VirtualTextureMaterialDraw::FParameters* PassParameters = GraphBuilder.AllocParameters<FShader_VirtualTextureMaterialDraw::FParameters>();
-			PassParameters->RenderTargets[0] = GraphSetup.RenderTexture0 ? FRenderTargetBinding(GraphSetup.RenderTexture0, ERenderTargetLoadAction::EClear, ERenderTargetStoreAction::EStore) : FRenderTargetBinding();
-			PassParameters->RenderTargets[1] = GraphSetup.RenderTexture1 ? FRenderTargetBinding(GraphSetup.RenderTexture1, ERenderTargetLoadAction::EClear, ERenderTargetStoreAction::EStore) : FRenderTargetBinding();
-			PassParameters->RenderTargets[2] = GraphSetup.RenderTexture2 ? FRenderTargetBinding(GraphSetup.RenderTexture2, ERenderTargetLoadAction::EClear, ERenderTargetStoreAction::EStore) : FRenderTargetBinding();
+			PassParameters->RenderTargets[0] = GraphSetup.RenderTexture0 ? FRenderTargetBinding(GraphSetup.RenderTexture0, LoadAction, ERenderTargetStoreAction::EStore) : FRenderTargetBinding();
+			PassParameters->RenderTargets[1] = GraphSetup.RenderTexture1 ? FRenderTargetBinding(GraphSetup.RenderTexture1, LoadAction, ERenderTargetStoreAction::EStore) : FRenderTargetBinding();
+			PassParameters->RenderTargets[2] = GraphSetup.RenderTexture2 ? FRenderTargetBinding(GraphSetup.RenderTexture2, LoadAction, ERenderTargetStoreAction::EStore) : FRenderTargetBinding();
 
 			GraphBuilder.AddPass(
 				RDG_EVENT_NAME("VirtualTextureDraw"),
diff --git a/Engine/Source/Runtime/Renderer/Private/VT/RuntimeVirtualTextureSceneProxy.cpp b/Engine/Source/Runtime/Renderer/Private/VT/RuntimeVirtualTextureSceneProxy.cpp
--- a/Engine/Source/Runtime/Renderer/Private/VT/RuntimeVirtualTextureSceneProxy.cpp
+++ b/Engine/Source/Runtime/Renderer/Private/VT/RuntimeVirtualTextureSceneProxy.cpp
@@ -32,9 +32,10 @@ FRuntimeVirtualTextureSceneProxy::FRuntimeVirtualTextureSceneProxy(URuntimeVirtu
 		MaxDirtyLevel = Desc.MaxLevel;
 
 		const ERuntimeVirtualTextureMaterialType MaterialType = VirtualTexture->GetMaterialType();
+		const bool bClearTextures = VirtualTexture->GetClearTextures();
 
 		// The Producer object created here will be passed into the virtual texture system which will take ownership.
-		IVirtualTexture* Producer = new FRuntimeVirtualTextureProducer(Desc, ProducerId, MaterialType, InComponent->GetScene(), Transform);
+		IVirtualTexture* Producer = new FRuntimeVirtualTextureProducer(Desc, ProducerId, MaterialType, bClearTextures, InComponent->GetScene(), Transform);
 
 		if (InComponent->IsStreamingLowMips() && VirtualTexture->GetStreamLowMips() > 0)
 		{
diff --git a/Engine/Source/Runtime/Renderer/Public/VT/RuntimeVirtualTextureRender.h b/Engine/Source/Runtime/Renderer/Public/VT/RuntimeVirtualTextureRender.h
--- a/Engine/Source/Runtime/Renderer/Public/VT/RuntimeVirtualTextureRender.h
+++ b/Engine/Source/Runtime/Renderer/Public/VT/RuntimeVirtualTextureRender.h
@@ -34,6 +34,7 @@ namespace RuntimeVirtualTexture
 		FScene* Scene,
 		uint32 RuntimeVirtualTextureMask,
 		ERuntimeVirtualTextureMaterialType MaterialType,
+		bool bClearTextures,
 		FRHITexture2D* Texture0, FBox2D const& DestBox0,
 		FRHITexture2D* Texture1, FBox2D const& DestBox1,
 		FTransform const& UVToWorld,
