When multiple HLODs are set, the transition does not occur correctly

Visibility was dependent on the order of items in the SceneNodes map.
If a child was processed first before it's parent, it was forced visible, even if the parent would later mark all it's children as hidden.

diff --git a/Engine/Source/Runtime/Renderer/Private/SceneVisibility.cpp b/Engine/Source/Runtime/Renderer/Private/SceneVisibility.cpp
--- a/Engine/Source/Runtime/Renderer/Private/SceneVisibility.cpp
+++ b/Engine/Source/Runtime/Renderer/Private/SceneVisibility.cpp
@@ -498,14 +498,17 @@ static int32 FrustumCull(const FScene* Scene, FViewInfo& View)
 					}
 
 					// Fading HLODs and their children must be visible, objects hidden by HLODs can be culled
-					if (HLODState && HLODState->IsNodeForcedVisible(Index))
+					if (HLODState)
 					{
-						MaxDrawDistance = FLT_MAX;
-						MinDrawDistanceSq = 0.f;
-					}
-					else if (HLODState && HLODState->IsNodeForcedHidden(Index))
-					{
-						MaxDrawDistance = 0.f;
+						if (HLODState->IsNodeForcedVisible(Index))
+						{
+							MaxDrawDistance = FLT_MAX;
+							MinDrawDistanceSq = 0.f;
+						}
+						else if (HLODState->IsNodeForcedHidden(Index))
+						{
+							MaxDrawDistance = 0.f;
+						}
 					}
 
 					if (DistanceSquared > FMath::Square(MaxDrawDistance + FadeRadius) ||
@@ -4653,6 +4656,9 @@ void FLODSceneTree::HideNodeChildren(FSceneViewState* ViewState, FLODSceneNode&
 				}
 
 				HLODState.ForcedHiddenPrimitiveMap[ChildIndex] = true;
+				
+				// Clear the force visible flag in case the child was processed before it's parent
+				HLODState.ForcedVisiblePrimitiveMap[ChildIndex] = false;
 
 				if (FLODSceneNode* ChildNode = SceneNodes.Find(Child->PrimitiveComponentId))
 				{
