Avoid copying FVisibleMeshDrawCommand when building visible mesh commands list, saves about 0.5ms on low end iPhone devices

diff --git a/Engine/Source/Runtime/Renderer/Private/SceneVisibility.cpp b/Engine/Source/Runtime/Renderer/Private/SceneVisibility.cpp
--- a/Engine/Source/Runtime/Renderer/Private/SceneVisibility.cpp
+++ b/Engine/Source/Runtime/Renderer/Private/SceneVisibility.cpp
@@ -1787,7 +1787,9 @@ struct FDrawCommandRelevancePacket
 				const FCachedMeshDrawCommandInfo& CachedMeshDrawCommand = InPrimitiveSceneInfo->StaticMeshCommandInfos[StaticMeshCommandInfoIndex];
 				const FCachedPassMeshDrawList& SceneDrawList = Scene->CachedDrawLists[PassType];
 
-				FVisibleMeshDrawCommand NewVisibleMeshDrawCommand;
+				// AddUninitialized_GetRef()
+				VisibleCachedDrawCommands[(uint32)PassType].AddUninitialized();
+				FVisibleMeshDrawCommand& NewVisibleMeshDrawCommand = VisibleCachedDrawCommands[(uint32)PassType].Last();
 
 				const FMeshDrawCommand* MeshDrawCommand = CachedMeshDrawCommand.StateBucketId >= 0
 					? &Scene->CachedMeshDrawCommandStateBuckets[FSetElementId::FromInteger(CachedMeshDrawCommand.StateBucketId)].MeshDrawCommand
@@ -1801,8 +1803,6 @@ struct FDrawCommandRelevancePacket
 					CachedMeshDrawCommand.MeshFillMode,
 					CachedMeshDrawCommand.MeshCullMode,
 					CachedMeshDrawCommand.SortKey);
-
-				VisibleCachedDrawCommands[(uint32)PassType].Add(NewVisibleMeshDrawCommand);
 			}
 		}
 		else
