[NO NON-PII COMMIT MESSAGE]

diff --git a/Engine/Platforms/XXX/Config/DataDrivenPlatformInfo.ini b/Engine/Platforms/XXX/Config/DataDrivenPlatformInfo.ini
--- a/Engine/Platforms/XXX/Config/DataDrivenPlatformInfo.ini
+++ b/Engine/Platforms/XXX/Config/DataDrivenPlatformInfo.ini
@@ -11,7 +11,6 @@ bSupportsVolumeTextureCompression=true
 bSupportsDistanceFields=true
 bSupportsDiaphragmDOF=true
 bSupportsRGBColorBuffer=true
-bSupportsByteBufferComputeShaders=true
 bSupportsCapsuleShadows=true
 bSupportsVolumetricFog=true
 bSupportsIndexBufferUAVs=true
diff --git a/Engine/Source/Runtime/RHI/Private/RHI.cpp b/Engine/Source/Runtime/RHI/Private/RHI.cpp
--- a/Engine/Source/Runtime/RHI/Private/RHI.cpp
+++ b/Engine/Source/Runtime/RHI/Private/RHI.cpp
@@ -1175,7 +1175,6 @@ inline void ParseDataDrivenShaderInfo(const FConfigSection& Section, FDataDriven
 	Info.bSupportsDistanceFields = GetSectionBool(Section, "bSupportsDistanceFields");
 	Info.bSupportsDiaphragmDOF = GetSectionBool(Section, "bSupportsDiaphragmDOF");
 	Info.bSupportsRGBColorBuffer = GetSectionBool(Section, "bSupportsRGBColorBuffer");
-	Info.bSupportsByteBufferComputeShaders = GetSectionBool(Section, "bSupportsByteBufferComputeShaders");
 	Info.bSupportsCapsuleShadows = GetSectionBool(Section, "bSupportsCapsuleShadows");
 	Info.bSupportsVolumetricFog = GetSectionBool(Section, "bSupportsVolumetricFog");
 	Info.bSupportsIndexBufferUAVs = GetSectionBool(Section, "bSupportsIndexBufferUAVs");
diff --git a/Engine/Source/Runtime/RHI/Public/RHIDefinitions.h b/Engine/Source/Runtime/RHI/Public/RHIDefinitions.h
--- a/Engine/Source/Runtime/RHI/Public/RHIDefinitions.h
+++ b/Engine/Source/Runtime/RHI/Public/RHIDefinitions.h
@@ -151,7 +151,6 @@ struct RHI_API FDataDrivenShaderPlatformInfo
 	bool bSupportsDistanceFields; // used for DFShadows and DFAO - since they had the same checks
 	bool bSupportsDiaphragmDOF;
 	bool bSupportsRGBColorBuffer;
-	bool bSupportsByteBufferComputeShaders;
 	bool bSupportsCapsuleShadows;
 	bool bSupportsVolumetricFog; // also used for FVVoxelization
 	bool bSupportsIndexBufferUAVs;
@@ -1160,7 +1159,7 @@ inline bool RHISupportsSeparateMSAAAndResolveTextures(const EShaderPlatform Plat
 inline bool RHISupportsComputeShaders(const EShaderPlatform Platform)
 {
 	return IsFeatureLevelSupported(Platform, ERHIFeatureLevel::SM5) 
-		|| GetMaxSupportedFeatureLevel(Platform) == ERHIFeatureLevel::ES3_1;
+		|| (GetMaxSupportedFeatureLevel(Platform) == ERHIFeatureLevel::ES3_1 && !IsSwitchPlatform(Platform));
 }
 
 inline bool RHISupportsGeometryShaders(const EShaderPlatform Platform)
diff --git a/Engine/Source/Runtime/Renderer/Private/ByteBuffer.cpp b/Engine/Source/Runtime/Renderer/Private/ByteBuffer.cpp
--- a/Engine/Source/Runtime/Renderer/Private/ByteBuffer.cpp
+++ b/Engine/Source/Runtime/Renderer/Private/ByteBuffer.cpp
@@ -16,9 +16,7 @@ class FMemsetBufferCS : public FGlobalShader
 	
 	static bool ShouldCompilePermutation( const FGlobalShaderPermutationParameters& Parameters )
 	{
-		return RHISupportsComputeShaders(Parameters.Platform)
-			// @todo platplug can we remove ths property now that we are just checking SM5 instead of platforms (and all uses below)?
-			|| FDataDrivenShaderPlatformInfo::GetInfo(Parameters.Platform).bSupportsByteBufferComputeShaders;
+		return RHISupportsComputeShaders(Parameters.Platform);
 	}
 
 	static void ModifyCompilationEnvironment( const FGlobalShaderPermutationParameters& Parameters, FShaderCompilerEnvironment& OutEnvironment )
@@ -87,8 +85,7 @@ class FMemcpyBufferCS : public FGlobalShader
 	
 	static bool ShouldCompilePermutation( const FGlobalShaderPermutationParameters& Parameters )
 	{
-		return RHISupportsComputeShaders(Parameters.Platform)
-			|| FDataDrivenShaderPlatformInfo::GetInfo(Parameters.Platform).bSupportsByteBufferComputeShaders;
+		return RHISupportsComputeShaders(Parameters.Platform);
 	}
 
 	static void ModifyCompilationEnvironment( const FGlobalShaderPermutationParameters& Parameters, FShaderCompilerEnvironment& OutEnvironment )
@@ -187,8 +184,7 @@ class FScatterCopyCS : public FGlobalShader
 	
 	static bool ShouldCompilePermutation( const FGlobalShaderPermutationParameters& Parameters )
 	{
-		return RHISupportsComputeShaders(Parameters.Platform)
-			|| FDataDrivenShaderPlatformInfo::GetInfo(Parameters.Platform).bSupportsByteBufferComputeShaders;
+		return RHISupportsComputeShaders(Parameters.Platform);
 	}
 
 	static void ModifyCompilationEnvironment( const FGlobalShaderPermutationParameters& Parameters, FShaderCompilerEnvironment& OutEnvironment )
