Avoid doing unnecesary work gathering ray tracing instances if all the ray tracing effects are disabled

diff --git a/Engine/Source/Runtime/Renderer/Private/DeferredShadingRenderer.cpp b/Engine/Source/Runtime/Renderer/Private/DeferredShadingRenderer.cpp
--- a/Engine/Source/Runtime/Renderer/Private/DeferredShadingRenderer.cpp
+++ b/Engine/Source/Runtime/Renderer/Private/DeferredShadingRenderer.cpp
@@ -544,6 +544,11 @@ bool FDeferredShadingSceneRenderer::GatherRayTracingWorldInstances(FRHICommandLi
 		return false;
 	}
 
+	if (GetForceRayTracingEffectsCVarValue() == 0 && Views.Num() > 0 && CanOverlayRayTracingOutput(Views[0])) // #dxr_todo: UE-72557 multi-view case
+	{
+		return false;
+	}
+
 	{
 		SCOPE_CYCLE_COUNTER(STAT_GenerateVisibleRayTracingMeshCommands);
 		RayTracingCollector.ClearViewMeshArrays();
@@ -788,6 +793,11 @@ bool FDeferredShadingSceneRenderer::DispatchRayTracingWorldUpdates(FRHICommandLi
 		return false;
 	}
 
+	if (GetForceRayTracingEffectsCVarValue() == 0 && Views.Num() > 0 && CanOverlayRayTracingOutput(Views[0])) // #dxr_todo: UE-72557 multi-view case
+	{
+		return false; 
+	}
+
 	SCOPED_GPU_STAT(RHICmdList, RayTracingTLAS);
 
 	Scene->GetRayTracingDynamicGeometryCollection()->DispatchUpdates(RHICmdList);
