Removed wrong assert (the case is carefully handled right below)

diff --git a/Engine/Source/Runtime/Renderer/Private/SceneRendering.cpp b/Engine/Source/Runtime/Renderer/Private/SceneRendering.cpp
--- a/Engine/Source/Runtime/Renderer/Private/SceneRendering.cpp
+++ b/Engine/Source/Runtime/Renderer/Private/SceneRendering.cpp
@@ -1227,7 +1227,6 @@ void FViewInfo::SetupUniformBufferParameters(
 	FRHITexture* DistantSkyLightLutTextureFound = nullptr;
 	if (ShouldRenderSkyAtmosphere(Scene, Family->EngineShowFlags))
 	{
-		check(Scene->SkyAtmosphere->GetTransmittanceLutTexture().IsValid());
 		FSkyAtmosphereRenderSceneInfo* SkyAtmosphere = Scene->SkyAtmosphere;
 		const FSkyAtmosphereSceneProxy& SkyAtmosphereSceneProxy = SkyAtmosphere->GetSkyAtmosphereSceneProxy();
 
