Encapsulating ray tracing global illumination ray generation dispatch in preparation for Final Gather version.

diff --git a/Engine/Source/Runtime/Renderer/Private/DeferredShadingRenderer.h b/Engine/Source/Runtime/Renderer/Private/DeferredShadingRenderer.h
--- a/Engine/Source/Runtime/Renderer/Private/DeferredShadingRenderer.h
+++ b/Engine/Source/Runtime/Renderer/Private/DeferredShadingRenderer.h
@@ -517,6 +517,16 @@ private:
 	);
 	void CompositeGlobalIllumination(FRHICommandListImmediate& RHICmdList, const FViewInfo& View, TRefCountPtr<IPooledRenderTarget>& GlobalIlluminationRT);
 
+	bool RenderRayTracingGlobalIlluminationBruteForce(
+		FRHICommandListImmediate& RHICmdList,
+		FRDGBuilder& GraphBuilder,
+		FSceneTextureParameters& SceneTextures,
+		FViewInfo& View,
+		const IScreenSpaceDenoiser::FAmbientOcclusionRayTracingConfig& RayTracingConfig,
+		int32 UpscaleFactor,
+		FRDGTextureUAV* GlobalIlluminationUAV,
+		FRDGTextureUAV* RayDistanceUAV);
+
 	void BuildSkyLightCdfs(FRHICommandListImmediate& RHICmdList, FSkyLightSceneProxy* SkyLight);
 	void BuildSkyLightMipTree(FRHICommandListImmediate& RHICmdList, FTextureRHIRef SkyLightTexture, FRWBuffer& SkyLightMipTreePosX, FRWBuffer& SkyLightMipTreePosY, FRWBuffer& SkyLightMipTreePosZ, FRWBuffer& SkyLightMipTreeNegX, FRWBuffer& SkyLightMipTreeNegY, FRWBuffer& SkyLightMipTreeNegZ, FIntVector& SkyLightMipTreeDimensions);
 	void BuildSkyLightMipTreePdf(
diff --git a/Engine/Source/Runtime/Renderer/Private/RayTracing/RayTracingGlobalIllumination.cpp b/Engine/Source/Runtime/Renderer/Private/RayTracing/RayTracingGlobalIllumination.cpp
--- a/Engine/Source/Runtime/Renderer/Private/RayTracing/RayTracingGlobalIllumination.cpp
+++ b/Engine/Source/Runtime/Renderer/Private/RayTracing/RayTracingGlobalIllumination.cpp
@@ -114,7 +114,7 @@ static FAutoConsoleVariableRef CVarRayTracingGlobalIlluminationTileSize(
 
 static const int32 GLightCountMax = 64;
 
-DECLARE_GPU_STAT_NAMED(RayTracingGlobalIllumination, TEXT("Ray Tracing Global Illumination"));
+DECLARE_GPU_STAT_NAMED(RayTracingGIBruteForce, TEXT("Ray Tracing GI: Brute Force"));
 
 void SetupLightParameters(
 	const TSparseArray<FLightSceneInfoCompact>& Lights,
@@ -193,13 +193,25 @@ void SetupLightParameters(
 	}
 }
 
+int32 GetRayTracingGlobalIlluminationSamplesPerPixel(const FViewInfo& View)
+{
+	int32 SamplesPerPixel = GRayTracingGlobalIlluminationSamplesPerPixel > -1 ? GRayTracingGlobalIlluminationSamplesPerPixel : View.FinalPostProcessSettings.RayTracingGISamplesPerPixel;
+	return SamplesPerPixel;
+}
+
 bool ShouldRenderRayTracingGlobalIllumination(const FViewInfo& View)
 {
 	if (!IsRayTracingEnabled())
 	{
 		return (false);
 	}
-	else if (GetForceRayTracingEffectsCVarValue() >= 0)
+
+	if (GetRayTracingGlobalIlluminationSamplesPerPixel(View) <= 0)
+	{
+		return false;
+	}
+
+	if (GetForceRayTracingEffectsCVarValue() >= 0)
 	{
 		return GetForceRayTracingEffectsCVarValue() > 0;
 	}
@@ -334,7 +346,7 @@ void FDeferredShadingSceneRenderer::RenderRayTracingGlobalIllumination(
 	TRefCountPtr<IPooledRenderTarget>& GlobalIlluminationRT
 )
 {
-	SCOPED_GPU_STAT(RHICmdList, RayTracingGlobalIllumination);
+	SCOPED_GPU_STAT(RHICmdList, RayTracingGIBruteForce);
 
 	bool bAnyViewWithRTGI = false;
 	for (int32 ViewIndex = 0; ViewIndex < Views.Num(); ++ViewIndex)
@@ -433,146 +445,24 @@ void FDeferredShadingSceneRenderer::RenderRayTracingGlobalIllumination(
 
 	RDG_EVENT_SCOPE(GraphBuilder, "RTGI");
 
-	int32 RayTracingGISamplesPerPixel = GRayTracingGlobalIlluminationSamplesPerPixel > -1 ? GRayTracingGlobalIlluminationSamplesPerPixel : View.FinalPostProcessSettings.RayTracingGISamplesPerPixel;
-	RayTracingGISamplesPerPixel = FMath::Max(RayTracingGISamplesPerPixel, 0);
-	RayTracingConfig.RayCountPerPixel = RayTracingGISamplesPerPixel;
+	RayTracingConfig.RayCountPerPixel = GetRayTracingGlobalIlluminationSamplesPerPixel(View);
 
 	FSceneTextureParameters SceneTextures;
 	SetupSceneTextureParameters(GraphBuilder, &SceneTextures);
 
-	// Ray generation
+	if (Scene->SkyLight && Scene->SkyLight->ShouldRebuildCdf())
 	{
-		uint32 IterationCount = FMath::Max(RayTracingGISamplesPerPixel, 1);
-		uint32 SequenceCount = 1;
-		uint32 DimensionCount = 24;
-		FHaltonSequenceIteration HaltonSequenceIteration(Scene->HaltonSequence, IterationCount, SequenceCount, DimensionCount, View.ViewState ? (View.ViewState->FrameIndex % 1024) : 0);
-
-		FHaltonIteration HaltonIteration;
-		InitializeHaltonSequenceIteration(HaltonSequenceIteration, HaltonIteration);
-
-		FHaltonPrimes HaltonPrimes;
-		InitializeHaltonPrimes(Scene->HaltonPrimesResource, HaltonPrimes);
-
-		FBlueNoise BlueNoise;
-		InitializeBlueNoise(BlueNoise);
-
-		FPathTracingLightData LightParameters;
-		SetupLightParameters(Scene->Lights, View, &LightParameters);
-
-		if (Scene->SkyLight && Scene->SkyLight->ShouldRebuildCdf())
-		{
-			BuildSkyLightCdfs(RHICmdList, Scene->SkyLight);
-		}
-		FSkyLightData SkyLightParameters;
-		SetupSkyLightParameters(*Scene, &SkyLightParameters);
-
-		FGlobalIlluminationRGS::FParameters* PassParameters = GraphBuilder.AllocParameters<FGlobalIlluminationRGS::FParameters>();
-		PassParameters->SamplesPerPixel = RayTracingGISamplesPerPixel;
-		PassParameters->MaxBounces = GRayTracingGlobalIlluminationMaxBounces > -1? GRayTracingGlobalIlluminationMaxBounces : View.FinalPostProcessSettings.RayTracingGIMaxBounces;
-		PassParameters->MaxNormalBias = GetRaytracingMaxNormalBias();
-		float MaxRayDistanceForGI = GRayTracingGlobalIlluminationMaxRayDistance;
-		if (MaxRayDistanceForGI == -1.0)
-		{
-			MaxRayDistanceForGI = View.FinalPostProcessSettings.AmbientOcclusionRadius;
-		}
-		PassParameters->MaxRayDistanceForGI = MaxRayDistanceForGI;
-		PassParameters->MaxRayDistanceForAO = View.FinalPostProcessSettings.AmbientOcclusionRadius;
-		PassParameters->UpscaleFactor = UpscaleFactor;
-		PassParameters->EvalSkyLight = GRayTracingGlobalIlluminationEvalSkyLight != 0;
-		PassParameters->UseRussianRoulette = GRayTracingGlobalIlluminationUseRussianRoulette != 0;
-		PassParameters->DiffuseThreshold = GRayTracingGlobalIlluminationDiffuseThreshold;
-		PassParameters->NextEventEstimationSamples = GRayTracingGlobalIlluminationNextEventEstimationSamples;
-		PassParameters->TLAS = View.RayTracingScene.RayTracingSceneRHI->GetShaderResourceView();
-		PassParameters->ViewUniformBuffer = View.ViewUniformBuffer;
-		PassParameters->HaltonIteration = CreateUniformBufferImmediate(HaltonIteration, EUniformBufferUsage::UniformBuffer_SingleDraw);
-		PassParameters->HaltonPrimes = CreateUniformBufferImmediate(HaltonPrimes, EUniformBufferUsage::UniformBuffer_SingleDraw);
-		PassParameters->BlueNoise = CreateUniformBufferImmediate(BlueNoise, EUniformBufferUsage::UniformBuffer_SingleDraw);
-		PassParameters->LightParameters = CreateUniformBufferImmediate(LightParameters, EUniformBufferUsage::UniformBuffer_SingleDraw);
-		PassParameters->SceneTextures = SceneTextures;
-		PassParameters->SkyLight = CreateUniformBufferImmediate(SkyLightParameters, EUniformBufferUsage::UniformBuffer_SingleDraw);
-		TRefCountPtr<IPooledRenderTarget> SubsurfaceProfileRT((IPooledRenderTarget*) GetSubsufaceProfileTexture_RT(RHICmdList));
-		if (!SubsurfaceProfileRT)
-		{
-			SubsurfaceProfileRT = GSystemTextures.BlackDummy;
-		}
-		PassParameters->SSProfilesTexture = GraphBuilder.RegisterExternalTexture(SubsurfaceProfileRT);
-		PassParameters->TransmissionProfilesLinearSampler = TStaticSamplerState<SF_Bilinear, AM_Clamp, AM_Clamp, AM_Clamp>::GetRHI();
-		PassParameters->RWGlobalIlluminationUAV = GlobalIlluminationUAV;
-		PassParameters->RWRayDistanceUAV = RayDistanceUAV;
-		PassParameters->TileOffsetX = 0;
-		PassParameters->TileOffsetY = 0;
-
-		FGlobalIlluminationRGS::FPermutationDomain PermutationVector;
-		PermutationVector.Set<FGlobalIlluminationRGS::FUseAttenuationTermDim>(true);
-		PermutationVector.Set<FGlobalIlluminationRGS::FEnableTwoSidedGeometryDim>(CVarRayTracingGlobalIlluminationEnableTwoSidedGeometry.GetValueOnRenderThread() != 0);
-		TShaderMapRef<FGlobalIlluminationRGS> RayGenerationShader(GetGlobalShaderMap(FeatureLevel), PermutationVector);
-		ClearUnusedGraphResources(*RayGenerationShader, PassParameters);
-
-		FIntPoint RayTracingResolution = FIntPoint::DivideAndRoundUp(View.ViewRect.Size(), UpscaleFactor);
-
-		if (GRayTracingGlobalIlluminationTileSize <= 0)
-		{
-			GraphBuilder.AddPass(
-				RDG_EVENT_NAME("GlobalIlluminationRayTracing %dx%d", RayTracingResolution.X, RayTracingResolution.Y),
-				PassParameters,
-				ERDGPassFlags::Compute,
-				[PassParameters, this, &View, RayGenerationShader, RayTracingResolution](FRHICommandList& RHICmdList)
-			{
-				FRayTracingShaderBindingsWriter GlobalResources;
-				SetShaderParameters(GlobalResources, *RayGenerationShader, *PassParameters);
-
-				FRHIRayTracingScene* RayTracingSceneRHI = View.RayTracingScene.RayTracingSceneRHI;
-				RHICmdList.RayTraceDispatch(View.RayTracingMaterialPipeline, RayGenerationShader->GetRayTracingShader(), RayTracingSceneRHI, GlobalResources, RayTracingResolution.X, RayTracingResolution.Y);
-			});
-		}
-		else
-		{
-			int32 TileSize = FMath::Max(32, GRayTracingGlobalIlluminationTileSize);
-			int32 NumTilesX = FMath::DivideAndRoundUp(RayTracingResolution.X, TileSize);
-			int32 NumTilesY = FMath::DivideAndRoundUp(RayTracingResolution.Y, TileSize);
-			for (int32 Y = 0; Y < NumTilesY; ++Y)
-			{
-				for (int32 X = 0; X < NumTilesX; ++X)
-				{
-					FGlobalIlluminationRGS::FParameters* TilePassParameters = PassParameters;
-
-					if (X > 0 || Y > 0)
-					{
-						TilePassParameters = GraphBuilder.AllocParameters<FGlobalIlluminationRGS::FParameters>();
-						*TilePassParameters = *PassParameters;
-
-						TilePassParameters->TileOffsetX = X * TileSize;
-						TilePassParameters->TileOffsetY = Y * TileSize;
-					}
-
-					int32 DispatchSizeX = FMath::Min<int32>(TileSize, RayTracingResolution.X - TilePassParameters->TileOffsetX);
-					int32 DispatchSizeY = FMath::Min<int32>(TileSize, RayTracingResolution.Y - TilePassParameters->TileOffsetY);
-
-					GraphBuilder.AddPass(
-						RDG_EVENT_NAME("GlobalIlluminationRayTracing %dx%d (tile %dx%d)", DispatchSizeX, DispatchSizeY, X, Y),
-						TilePassParameters,
-						ERDGPassFlags::Compute,
-						[TilePassParameters, this, &View, RayGenerationShader, DispatchSizeX, DispatchSizeY](FRHICommandList& RHICmdList)
-					{
-						FRHIRayTracingScene* RayTracingSceneRHI = View.RayTracingScene.RayTracingSceneRHI;
-
-						FRayTracingShaderBindingsWriter GlobalResources;
-						SetShaderParameters(GlobalResources, *RayGenerationShader, *TilePassParameters);
-						RHICmdList.RayTraceDispatch(View.RayTracingMaterialPipeline, RayGenerationShader->GetRayTracingShader(), RayTracingSceneRHI, 
-							GlobalResources, DispatchSizeX, DispatchSizeY);
-						RHICmdList.SubmitCommandsHint();
-					});
-				}
-			}
-		}
-
+		BuildSkyLightCdfs(RHICmdList, Scene->SkyLight);
 	}
 
+	// Ray generation
+	bool bIsValid = RenderRayTracingGlobalIlluminationBruteForce(RHICmdList, GraphBuilder, SceneTextures, View, RayTracingConfig, UpscaleFactor, GlobalIlluminationUAV, RayDistanceUAV);
+
 	// Denoising
 
 	FRDGTextureRef ResultTexture; //#dxr_todo review
 
-	if (GRayTracingGlobalIlluminationDenoiser != 0 && RayTracingGISamplesPerPixel > 0)
+	if (GRayTracingGlobalIlluminationDenoiser != 0 && bIsValid)
 	{
 		const IScreenSpaceDenoiser* DefaultDenoiser = IScreenSpaceDenoiser::GetDefaultDenoiser();
 		const IScreenSpaceDenoiser* DenoiserToUse = GRayTracingGlobalIlluminationDenoiser == 1 ? DefaultDenoiser : GScreenSpaceDenoiser;
@@ -604,6 +494,7 @@ void FDeferredShadingSceneRenderer::RenderRayTracingGlobalIllumination(
 	}
 
 	// Compositing
+	if (bIsValid)
 	{
 		FRayTracingGlobalIlluminationCompositePS::FParameters *PassParameters = GraphBuilder.AllocParameters<FRayTracingGlobalIlluminationCompositePS::FParameters>();
 		PassParameters->GlobalIlluminationTexture = ResultTexture;
@@ -655,6 +546,151 @@ void FDeferredShadingSceneRenderer::RenderRayTracingGlobalIllumination(
 	}
 }
 
+bool FDeferredShadingSceneRenderer::RenderRayTracingGlobalIlluminationBruteForce(
+	FRHICommandListImmediate& RHICmdList,
+	FRDGBuilder& GraphBuilder,
+	FSceneTextureParameters& SceneTextures,
+	FViewInfo& View,
+	const IScreenSpaceDenoiser::FAmbientOcclusionRayTracingConfig& RayTracingConfig,
+	int32 UpscaleFactor,
+	FRDGTextureUAV* GlobalIlluminationUAV,
+	FRDGTextureUAV* RayDistanceUAV)
+#if RHI_RAYTRACING
+{
+	if (!View.ViewState) return false;
+
+	RDG_GPU_STAT_SCOPE(GraphBuilder, RayTracingGIBruteForce);
+	RDG_EVENT_SCOPE(GraphBuilder, "Ray Tracing GI: Brute Force");
+
+	int32 RayTracingGISamplesPerPixel = GetRayTracingGlobalIlluminationSamplesPerPixel(View);
+	uint32 IterationCount = FMath::Max(RayTracingGISamplesPerPixel, 1);
+	uint32 SequenceCount = 1;
+	uint32 DimensionCount = 24;
+	FHaltonSequenceIteration HaltonSequenceIteration(Scene->HaltonSequence, IterationCount, SequenceCount, DimensionCount, View.ViewState ? (View.ViewState->FrameIndex % 1024) : 0);
+
+	FHaltonIteration HaltonIteration;
+	InitializeHaltonSequenceIteration(HaltonSequenceIteration, HaltonIteration);
+
+	FHaltonPrimes HaltonPrimes;
+	InitializeHaltonPrimes(Scene->HaltonPrimesResource, HaltonPrimes);
+
+	FBlueNoise BlueNoise;
+	InitializeBlueNoise(BlueNoise);
+
+	FPathTracingLightData LightParameters;
+	SetupLightParameters(Scene->Lights, View, &LightParameters);
+
+	FSkyLightData SkyLightParameters;
+	SetupSkyLightParameters(*Scene, &SkyLightParameters);
+
+	FGlobalIlluminationRGS::FParameters* PassParameters = GraphBuilder.AllocParameters<FGlobalIlluminationRGS::FParameters>();
+	PassParameters->SamplesPerPixel = RayTracingGISamplesPerPixel;
+	PassParameters->MaxBounces = GRayTracingGlobalIlluminationMaxBounces > -1 ? GRayTracingGlobalIlluminationMaxBounces : View.FinalPostProcessSettings.RayTracingGIMaxBounces;
+	PassParameters->MaxNormalBias = GetRaytracingMaxNormalBias();
+	float MaxRayDistanceForGI = GRayTracingGlobalIlluminationMaxRayDistance;
+	if (MaxRayDistanceForGI == -1.0)
+	{
+		MaxRayDistanceForGI = View.FinalPostProcessSettings.AmbientOcclusionRadius;
+	}
+	PassParameters->MaxRayDistanceForGI = MaxRayDistanceForGI;
+	PassParameters->MaxRayDistanceForAO = View.FinalPostProcessSettings.AmbientOcclusionRadius;
+	PassParameters->UpscaleFactor = UpscaleFactor;
+	PassParameters->EvalSkyLight = GRayTracingGlobalIlluminationEvalSkyLight != 0;
+	PassParameters->UseRussianRoulette = GRayTracingGlobalIlluminationUseRussianRoulette != 0;
+	PassParameters->DiffuseThreshold = GRayTracingGlobalIlluminationDiffuseThreshold;
+	PassParameters->NextEventEstimationSamples = GRayTracingGlobalIlluminationNextEventEstimationSamples;
+	PassParameters->TLAS = View.RayTracingScene.RayTracingSceneRHI->GetShaderResourceView();
+	PassParameters->ViewUniformBuffer = View.ViewUniformBuffer;
+	PassParameters->HaltonIteration = CreateUniformBufferImmediate(HaltonIteration, EUniformBufferUsage::UniformBuffer_SingleDraw);
+	PassParameters->HaltonPrimes = CreateUniformBufferImmediate(HaltonPrimes, EUniformBufferUsage::UniformBuffer_SingleDraw);
+	PassParameters->BlueNoise = CreateUniformBufferImmediate(BlueNoise, EUniformBufferUsage::UniformBuffer_SingleDraw);
+	PassParameters->LightParameters = CreateUniformBufferImmediate(LightParameters, EUniformBufferUsage::UniformBuffer_SingleDraw);
+	PassParameters->SceneTextures = SceneTextures;
+	PassParameters->SkyLight = CreateUniformBufferImmediate(SkyLightParameters, EUniformBufferUsage::UniformBuffer_SingleDraw);
+	TRefCountPtr<IPooledRenderTarget> SubsurfaceProfileRT((IPooledRenderTarget*)GetSubsufaceProfileTexture_RT(RHICmdList));
+	if (!SubsurfaceProfileRT)
+	{
+		SubsurfaceProfileRT = GSystemTextures.BlackDummy;
+	}
+	PassParameters->SSProfilesTexture = GraphBuilder.RegisterExternalTexture(SubsurfaceProfileRT);
+	PassParameters->TransmissionProfilesLinearSampler = TStaticSamplerState<SF_Bilinear, AM_Clamp, AM_Clamp, AM_Clamp>::GetRHI();
+	PassParameters->RWGlobalIlluminationUAV = GlobalIlluminationUAV;
+	PassParameters->RWRayDistanceUAV = RayDistanceUAV;
+	PassParameters->TileOffsetX = 0;
+	PassParameters->TileOffsetY = 0;
+
+	FGlobalIlluminationRGS::FPermutationDomain PermutationVector;
+	PermutationVector.Set<FGlobalIlluminationRGS::FUseAttenuationTermDim>(true);
+	PermutationVector.Set<FGlobalIlluminationRGS::FEnableTwoSidedGeometryDim>(CVarRayTracingGlobalIlluminationEnableTwoSidedGeometry.GetValueOnRenderThread() != 0);
+	TShaderMapRef<FGlobalIlluminationRGS> RayGenerationShader(GetGlobalShaderMap(FeatureLevel), PermutationVector);
+	ClearUnusedGraphResources(*RayGenerationShader, PassParameters);
+
+	FIntPoint RayTracingResolution = FIntPoint::DivideAndRoundUp(View.ViewRect.Size(), UpscaleFactor);
+
+	if (GRayTracingGlobalIlluminationTileSize <= 0)
+	{
+		GraphBuilder.AddPass(
+			RDG_EVENT_NAME("GlobalIlluminationRayTracing %dx%d", RayTracingResolution.X, RayTracingResolution.Y),
+			PassParameters,
+			ERDGPassFlags::Compute,
+			[PassParameters, this, &View, RayGenerationShader, RayTracingResolution](FRHICommandList& RHICmdList)
+		{
+			FRayTracingShaderBindingsWriter GlobalResources;
+			SetShaderParameters(GlobalResources, *RayGenerationShader, *PassParameters);
+
+			FRHIRayTracingScene* RayTracingSceneRHI = View.RayTracingScene.RayTracingSceneRHI;
+			RHICmdList.RayTraceDispatch(View.RayTracingMaterialPipeline, RayGenerationShader->GetRayTracingShader(), RayTracingSceneRHI, GlobalResources, RayTracingResolution.X, RayTracingResolution.Y);
+		});
+	}
+	else
+	{
+		int32 TileSize = FMath::Max(32, GRayTracingGlobalIlluminationTileSize);
+		int32 NumTilesX = FMath::DivideAndRoundUp(RayTracingResolution.X, TileSize);
+		int32 NumTilesY = FMath::DivideAndRoundUp(RayTracingResolution.Y, TileSize);
+		for (int32 Y = 0; Y < NumTilesY; ++Y)
+		{
+			for (int32 X = 0; X < NumTilesX; ++X)
+			{
+				FGlobalIlluminationRGS::FParameters* TilePassParameters = PassParameters;
+
+				if (X > 0 || Y > 0)
+				{
+					TilePassParameters = GraphBuilder.AllocParameters<FGlobalIlluminationRGS::FParameters>();
+					*TilePassParameters = *PassParameters;
+
+					TilePassParameters->TileOffsetX = X * TileSize;
+					TilePassParameters->TileOffsetY = Y * TileSize;
+				}
+
+				int32 DispatchSizeX = FMath::Min<int32>(TileSize, RayTracingResolution.X - TilePassParameters->TileOffsetX);
+				int32 DispatchSizeY = FMath::Min<int32>(TileSize, RayTracingResolution.Y - TilePassParameters->TileOffsetY);
+
+				GraphBuilder.AddPass(
+					RDG_EVENT_NAME("GlobalIlluminationRayTracing %dx%d (tile %dx%d)", DispatchSizeX, DispatchSizeY, X, Y),
+					TilePassParameters,
+					ERDGPassFlags::Compute,
+					[TilePassParameters, this, &View, RayGenerationShader, DispatchSizeX, DispatchSizeY](FRHICommandList& RHICmdList)
+				{
+					FRHIRayTracingScene* RayTracingSceneRHI = View.RayTracingScene.RayTracingSceneRHI;
+
+					FRayTracingShaderBindingsWriter GlobalResources;
+					SetShaderParameters(GlobalResources, *RayGenerationShader, *TilePassParameters);
+					RHICmdList.RayTraceDispatch(View.RayTracingMaterialPipeline, RayGenerationShader->GetRayTracingShader(), RayTracingSceneRHI,
+						GlobalResources, DispatchSizeX, DispatchSizeY);
+					RHICmdList.SubmitCommandsHint();
+				});
+			}
+		}
+	}
+	return true;
+}
+#else
+{
+	unimplemented();
+	return false;
+}
+#endif // RHI_RAYTRACING
+
 void FDeferredShadingSceneRenderer::CompositeGlobalIllumination(
 	FRHICommandListImmediate& RHICmdList,
 	const FViewInfo& View,
