Add mem mark around RVT rendering
Fixes mem allocator check() firing in debug builds

diff --git a/Engine/Source/Runtime/Renderer/Private/VT/RuntimeVirtualTextureRender.cpp b/Engine/Source/Runtime/Renderer/Private/VT/RuntimeVirtualTextureRender.cpp
--- a/Engine/Source/Runtime/Renderer/Private/VT/RuntimeVirtualTextureRender.cpp
+++ b/Engine/Source/Runtime/Renderer/Private/VT/RuntimeVirtualTextureRender.cpp
@@ -718,6 +718,7 @@ namespace RuntimeVirtualTexture
 		Scene->UniformBuffers.VirtualTextureViewUniformBuffer.UpdateUniformBufferImmediate(*View->CachedViewUniformShaderParameters);
 
 		// Build graph
+		FMemMark Mark(FMemStack::Get());
 		FRDGBuilder GraphBuilder(RHICmdList);
 		FRenderGraphSetup GraphSetup(GraphBuilder, MaterialType, OutputTexture0, OutputTexture1, TextureSize);
 
