Remove creation of unused cached draw commands from the virtual texture mesh pass

diff --git a/Engine/Source/Runtime/Renderer/Private/VT/RuntimeVirtualTextureRender.cpp b/Engine/Source/Runtime/Renderer/Private/VT/RuntimeVirtualTextureRender.cpp
--- a/Engine/Source/Runtime/Renderer/Private/VT/RuntimeVirtualTextureRender.cpp
+++ b/Engine/Source/Runtime/Renderer/Private/VT/RuntimeVirtualTextureRender.cpp
@@ -248,29 +248,32 @@ namespace RuntimeVirtualTexture
 	public:
 		virtual void AddMeshBatch(const FMeshBatch& RESTRICT MeshBatch, uint64 BatchElementMask, const FPrimitiveSceneProxy* RESTRICT PrimitiveSceneProxy, int32 StaticMeshId = -1) override final
 		{
-			const FMaterialRenderProxy* FallbackMaterialRenderProxyPtr = nullptr;
-			const FMaterial& Material = MeshBatch.MaterialRenderProxy->GetMaterialWithFallback(FeatureLevel, FallbackMaterialRenderProxyPtr);
-			const FMaterialRenderProxy& MaterialRenderProxy = FallbackMaterialRenderProxyPtr ? *FallbackMaterialRenderProxyPtr : *MeshBatch.MaterialRenderProxy;
-
-			if (Material.GetMaterialDomain() == MD_RuntimeVirtualTexture || Material.HasRuntimeVirtualTextureOutput_RenderThread())
+			if (MeshBatch.bRenderToVirtualTexture)
 			{
-				switch ((ERuntimeVirtualTextureMaterialType)MeshBatch.RuntimeVirtualTextureMaterialType)
+				const FMaterialRenderProxy* FallbackMaterialRenderProxyPtr = nullptr;
+				const FMaterial& Material = MeshBatch.MaterialRenderProxy->GetMaterialWithFallback(FeatureLevel, FallbackMaterialRenderProxyPtr);
+				const FMaterialRenderProxy& MaterialRenderProxy = FallbackMaterialRenderProxyPtr ? *FallbackMaterialRenderProxyPtr : *MeshBatch.MaterialRenderProxy;
+
+				if (Material.GetMaterialDomain() == MD_RuntimeVirtualTexture || Material.HasRuntimeVirtualTextureOutput_RenderThread())
 				{
-				case ERuntimeVirtualTextureMaterialType::BaseColor:
-					Process<FMaterialPolicy_BaseColor>(MeshBatch, BatchElementMask, StaticMeshId, PrimitiveSceneProxy, MaterialRenderProxy, Material);
-					break;
-				case ERuntimeVirtualTextureMaterialType::BaseColor_Normal:
-					Process<FMaterialPolicy_BaseColorNormal>(MeshBatch, BatchElementMask, StaticMeshId, PrimitiveSceneProxy, MaterialRenderProxy, Material);
-					break;
-				case ERuntimeVirtualTextureMaterialType::BaseColor_Normal_Specular:
-				case ERuntimeVirtualTextureMaterialType::BaseColor_Normal_Specular_Ex:
-					Process<FMaterialPolicy_BaseColorNormalSpecular>(MeshBatch, BatchElementMask, StaticMeshId, PrimitiveSceneProxy, MaterialRenderProxy, Material);
-					break;
-				case ERuntimeVirtualTextureMaterialType::WorldHeight:
-					Process<FMaterialPolicy_WorldHeight>(MeshBatch, BatchElementMask, StaticMeshId, PrimitiveSceneProxy, MaterialRenderProxy, Material);
-					break;
-				default:
-					break;
+					switch ((ERuntimeVirtualTextureMaterialType)MeshBatch.RuntimeVirtualTextureMaterialType)
+					{
+					case ERuntimeVirtualTextureMaterialType::BaseColor:
+						Process<FMaterialPolicy_BaseColor>(MeshBatch, BatchElementMask, StaticMeshId, PrimitiveSceneProxy, MaterialRenderProxy, Material);
+						break;
+					case ERuntimeVirtualTextureMaterialType::BaseColor_Normal:
+						Process<FMaterialPolicy_BaseColorNormal>(MeshBatch, BatchElementMask, StaticMeshId, PrimitiveSceneProxy, MaterialRenderProxy, Material);
+						break;
+					case ERuntimeVirtualTextureMaterialType::BaseColor_Normal_Specular:
+					case ERuntimeVirtualTextureMaterialType::BaseColor_Normal_Specular_Ex:
+						Process<FMaterialPolicy_BaseColorNormalSpecular>(MeshBatch, BatchElementMask, StaticMeshId, PrimitiveSceneProxy, MaterialRenderProxy, Material);
+						break;
+					case ERuntimeVirtualTextureMaterialType::WorldHeight:
+						Process<FMaterialPolicy_WorldHeight>(MeshBatch, BatchElementMask, StaticMeshId, PrimitiveSceneProxy, MaterialRenderProxy, Material);
+						break;
+					default:
+						break;
+					}
 				}
 			}
 		}
