Adds final gather algorithm for ray tracing global illumination pass.

diff --git a/Engine/Shaders/Private/RayTracing/RayTracingCreateGatherPointsRGS.usf b/Engine/Shaders/Private/RayTracing/RayTracingCreateGatherPointsRGS.usf
new file mode 100644
--- /dev/null
+++ b/Engine/Shaders/Private/RayTracing/RayTracingCreateGatherPointsRGS.usf
@@ -0,0 +1,253 @@
+// Copyright 1998-2019 Epic Games, Inc. All Rights Reserved.
+
+#include "../Common.ush"
+#include "../RectLight.ush"
+#include "../DeferredShadingCommon.ush"
+#include "../ShadingModels.ush"
+#include "../SceneTextureParameters.ush"
+
+#include "RayTracingCommon.ush"
+#include "RayTracingHitGroupCommon.ush"
+#include "RayTracingGatherPoints.ush"
+
+#include "../PathTracing/Light/PathTracingLightSampling.ush"
+#include "../PathTracing/Material/PathTracingMaterialSampling.ush"
+
+#include "../QuasiRandom.ush"
+
+RaytracingAccelerationStructure TLAS;
+
+uint SamplesPerPixel;
+uint SampleIndex;
+uint MaxBounces;
+uint UpscaleFactor;
+uint TileOffsetX;
+uint TileOffsetY;
+float MaxRayDistanceForGI;
+float NextEventEstimationSamples;
+float DiffuseThreshold;
+float MaxNormalBias;
+bool EvalSkyLight;
+bool UseRussianRoulette;
+
+int2 GatherPointsResolution;
+RWStructuredBuffer<FGatherPoints> RWGatherPointsBuffer;
+
+uint DispatchThreadIdToLinearIndex(uint2 DispatchThreadId)
+{
+	return DispatchThreadId.y * GatherPointsResolution.x + DispatchThreadId.x;
+}
+
+void GenerateCosineNormalRay(
+	float3 WorldPosition,
+	float3 WorldNormal,
+	float2 RandSample,
+	out float3 RayOrigin,
+	out float3 RayDirection,
+	out float RayTMin,
+	out float RayTMax,
+	out float RayPdf
+)
+{
+	// Perform cosine-hemispherical sampling and convert to world-space
+	float4 Direction_Tangent = CosineSampleHemisphere(RandSample);
+	float3 Direction_World = TangentToWorld(Direction_Tangent.xyz, WorldNormal);
+
+	RayOrigin = WorldPosition;
+	RayDirection = Direction_World;
+	RayTMin = 0.01;
+	RayTMax = MaxRayDistanceForGI;
+	RayPdf = Direction_Tangent.w;
+}
+
+[shader("raygeneration")]
+void RayTracingCreateGatherPointsRGS()
+{
+	uint2 DispatchThreadId = DispatchRaysIndex().xy + View.ViewRectMin + uint2(TileOffsetX, TileOffsetY);
+	uint GatherPointsIndex = DispatchThreadIdToLinearIndex(DispatchThreadId);
+#if 1
+	// Default initialization
+	for (uint Index = 0; Index < SamplesPerPixel; ++Index)
+	{
+		uint SampleIndexOffset = SampleIndex + Index;
+		RWGatherPointsBuffer[GatherPointsIndex].CreationPoint[SampleIndexOffset] = 0.0;
+		RWGatherPointsBuffer[GatherPointsIndex].Position[SampleIndexOffset] = 0.0;
+		RWGatherPointsBuffer[GatherPointsIndex].Irradiance[SampleIndexOffset] = 0.0;
+	}
+#endif
+	// Get G-Buffer surface data
+	uint2 PixelCoord = GetPixelCoord(DispatchThreadId, UpscaleFactor);
+	float2 InvBufferSize = View.BufferSizeAndInvSize.zw;
+	float2 UV = (float2(PixelCoord) + 0.5) * InvBufferSize;
+	FGBufferData GBufferData = GetGBufferDataFromSceneTextures(UV);
+	// Remap DiffuseColor when using SubsurfaceProfile (GBuffer decoding replaces with 100% albedo)
+	if (UseSubsurfaceProfile(GBufferData.ShadingModelID))
+	{
+		GBufferData.DiffuseColor = GBufferData.StoredBaseColor;
+	}
+
+	// Cull based on shading model
+	uint ShadingModelID = GBufferData.ShadingModelID;
+	if (ShadingModelID == SHADINGMODELID_UNLIT
+		|| ShadingModelID == SHADINGMODELID_TWOSIDED_FOLIAGE
+		)
+	{
+		return;
+	}
+
+	// Remap DiffuseColor when using SubsurfaceProfile (GBuffer decoding replaces with 100% albedo)
+	if (UseSubsurfaceProfile(GBufferData.ShadingModelID))
+	{
+		GBufferData.DiffuseColor = GBufferData.StoredBaseColor;
+	}
+	float3 DiffuseColor = GBufferData.DiffuseColor;
+	if (Luminance(DiffuseColor) < DiffuseThreshold)
+	{
+		return;
+	}
+
+	// Construct position, normal, and view direction
+	float DeviceZ = SceneDepthBuffer.Load(int3(PixelCoord, 0)).r;
+	float3 WorldPosition;
+	float3 CameraDirection;
+	ReconstructWorldPositionAndCameraDirectionFromDeviceZ(PixelCoord, DeviceZ, WorldPosition, CameraDirection);
+	float3 WorldNormal = GBufferData.WorldNormal;
+
+	FQuasiRandom QuasiRandom;
+	QuasiRandom_Initialize(QuasiRandom, PixelCoord, 0, SamplesPerPixel);
+
+	for (uint Index = 0; Index < SamplesPerPixel; ++Index)
+	{
+		uint SampleIndexOffset = (SampleIndex + Index) % 16;
+#if 0
+		// Adaptive rejection optimization
+		{
+			float3 GPPosition = RWGatherPointsBuffer[GatherPointsIndex].CreationPoint[SampleIndexOffset];
+			if (length(GPPosition - WorldPosition) < 10.0) continue;
+		}
+#endif
+		float3 Irradiance = 0;
+		float3 Position = 0;
+		float3 Normal = 0;
+
+		QuasiRandom_Update(QuasiRandom, Index, SamplesPerPixel);
+
+		float3 RayThroughput = 1.0;
+		float3 RandSample = QuasiRandom_Sample3(QuasiRandom);
+
+		// Initialize ray
+		RayDesc Ray;
+		float RayPdf = 1.0;
+		GenerateCosineNormalRay(WorldPosition, WorldNormal, RandSample.yz, Ray.Origin, Ray.Direction, Ray.TMin, Ray.TMax, RayPdf);
+		// Distribution PDF is not used for gather points, only gather queries.
+		RayPdf = 1.0;
+		ApplyCameraRelativeDepthBias(Ray, PixelCoord, DeviceZ, WorldNormal, MaxNormalBias);
+
+		{
+			// Cast ray
+			uint RayFlags = 0;
+#if !ENABLE_TWO_SIDED_GEOMETRY
+			RayFlags |= RAY_FLAG_CULL_BACK_FACING_TRIANGLES;
+#endif
+			FRayCone RayCone = (FRayCone)0;
+			FMaterialClosestHitPayload Payload = TraceRayInternal(
+				TLAS,   // AccelerationStructure
+				RayFlags,
+				RAY_TRACING_MASK_OPAQUE,
+				RAY_TRACING_SHADER_SLOT_MATERIAL, // RayContributionToHitGroupIndex
+				RAY_TRACING_NUM_SHADER_SLOTS,     // MultiplierForGeometryContributionToShaderIndex
+				0,      // MissShaderIndex
+				Ray,    // RayDesc
+				RayCone
+			);
+
+			if (Payload.HitT > MaxRayDistanceForGI) continue;
+
+			// Update intersection
+			Ray.Origin += Ray.Direction * Payload.HitT;
+
+			// Create faux GBuffer to use with EvaluateBxDF
+			FGBufferData GBufferData = (FGBufferData)0;
+			GBufferData.Depth = 1.f; // Do not use depth
+			GBufferData.WorldNormal = Payload.WorldNormal;
+			GBufferData.BaseColor = Payload.BaseColor;
+			GBufferData.CustomData = Payload.CustomData;
+			GBufferData.GBufferAO = Payload.GBufferAO;
+			GBufferData.IndirectIrradiance = (Payload.IndirectIrradiance.x + Payload.IndirectIrradiance.y + Payload.IndirectIrradiance.z) / 3.f;
+			GBufferData.SpecularColor = Payload.SpecularColor;
+			GBufferData.DiffuseColor = Payload.DiffuseColor;
+			GBufferData.Metallic = Payload.Metallic;
+			GBufferData.Specular = Payload.Specular;
+			GBufferData.Roughness = Payload.Roughness;
+			GBufferData.ShadingModelID = Payload.ShadingModelID;
+			GBufferData.CustomData = Payload.CustomData;
+
+			Position = Ray.Origin;
+			Normal = Payload.WorldNormal;
+
+			// Perform next-event estimation
+			float SplitFactor = 1.0 / NextEventEstimationSamples;
+			for (uint NeeTrial = 0; NeeTrial < NextEventEstimationSamples; ++NeeTrial)
+			{
+				// Light selection
+				int LightId;
+				float3 LightUV;
+				float NeePdf = 0.0;
+
+				float4 RandSample4;
+				RandSample4.xyz = QuasiRandom_Sample3(QuasiRandom);
+				RandSample4.w = QuasiRandom_Sample3(QuasiRandom).x;
+				SampleLight(Ray, Payload, RandSample4, LightId, LightUV, NeePdf);
+
+				if (NeePdf > 0.0)
+				{
+					RayDesc LightRay;
+					GenerateLightRay(Ray, LightId, LightUV, LightRay);
+					ApplyPositionBias(LightRay, Payload.WorldNormal, MaxNormalBias);
+
+					// Trace visibility ray
+					uint RayFlags = RAY_FLAG_ACCEPT_FIRST_HIT_AND_END_SEARCH | RAY_FLAG_SKIP_CLOSEST_HIT_SHADER;
+#if !ENABLE_TWO_SIDED_GEOMETRY
+					RayFlags |= RAY_FLAG_CULL_BACK_FACING_TRIANGLES;
+#endif
+					FRayCone LightRayCone = (FRayCone)0;
+					FMaterialClosestHitPayload NeePayload = TraceRayInternal(
+						TLAS,   // AccelerationStructure
+						RayFlags,
+						RAY_TRACING_MASK_OPAQUE,
+						RAY_TRACING_SHADER_SLOT_SHADOW,  // RayContributionToHitGroupIndex
+						RAY_TRACING_NUM_SHADER_SLOTS,     // MultiplierForGeometryContributionToShaderIndex
+						0,      // MissShaderIndex
+						LightRay,    // RayDesc
+						LightRayCone
+					);
+
+					// No hit indicates successful next-event connection
+					if (NeePayload.IsMiss())
+					{
+						// Evaluate radiance
+						float3 Radiance;
+						EvalLight(LightId, LightUV, LightRay, Radiance);
+
+						// Evaluate material
+						half3 N = Payload.WorldNormal;
+						half3 V = -Ray.Direction;
+						half3 L = LightRay.Direction;
+						float NoL = saturate(dot(N, L));
+						FShadowTerms ShadowTerms = { 0.0, 0.0, 0.0 };
+						FDirectLighting LightingSample = EvaluateBxDF(GBufferData, N, V, L, NoL, ShadowTerms);
+						float3 MaterialThroughput = LightingSample.Diffuse;
+
+						// Record the contribution
+						float3 ExitantRadianceSample = Radiance * MaterialThroughput * SplitFactor / NeePdf;
+						Irradiance += isfinite(ExitantRadianceSample) ? ExitantRadianceSample : 0.0;
+					}
+				}
+			}
+		}
+
+		RWGatherPointsBuffer[GatherPointsIndex].CreationPoint[SampleIndexOffset] = WorldPosition;
+		RWGatherPointsBuffer[GatherPointsIndex].Position[SampleIndexOffset] = Position;
+		RWGatherPointsBuffer[GatherPointsIndex].Irradiance[SampleIndexOffset] = Irradiance;
+	}
+}
diff --git a/Engine/Shaders/Private/RayTracing/RayTracingFinalGatherRGS.usf b/Engine/Shaders/Private/RayTracing/RayTracingFinalGatherRGS.usf
new file mode 100644
--- /dev/null
+++ b/Engine/Shaders/Private/RayTracing/RayTracingFinalGatherRGS.usf
@@ -0,0 +1,172 @@
+// Copyright 1998-2019 Epic Games, Inc. All Rights Reserved.
+
+#include "../Common.ush"
+#include "../RectLight.ush"
+#include "../DeferredShadingCommon.ush"
+#include "../ShadingModels.ush"
+#include "../SceneTextureParameters.ush"
+
+#include "RayTracingCommon.ush"
+#include "RayTracingHitGroupCommon.ush"
+#include "RayTracingGatherPoints.ush"
+
+#include "../PathTracing/Light/PathTracingLightSampling.ush"
+#include "../PathTracing/Material/PathTracingMaterialSampling.ush"
+
+#include "../QuasiRandom.ush"
+
+#define USE_PATHTRACING_MATERIALS 0
+
+uint SampleIndex;
+uint SamplesPerPixel;
+uint UpscaleFactor;
+uint TileOffsetX;
+uint TileOffsetY;
+float DiffuseThreshold;
+float MaxNormalBias;
+float FinalGatherDistance;
+
+RaytracingAccelerationStructure TLAS;
+
+int2 GatherPointsResolution;
+StructuredBuffer<FGatherPoints> GatherPointsBuffer;
+
+RWTexture2D<float4> RWGlobalIlluminationUAV;
+RWTexture2D<float> RWRayDistanceUAV;
+
+uint DispatchThreadIdToLinearIndex(uint2 DispatchThreadId)
+{
+	return DispatchThreadId.y * GatherPointsResolution.x + DispatchThreadId.x;
+}
+
+[shader("raygeneration")]
+void RayTracingFinalGatherRGS()
+{
+	uint2 DispatchThreadId = DispatchRaysIndex().xy + View.ViewRectMin + uint2(TileOffsetX, TileOffsetY);
+	RWGlobalIlluminationUAV[DispatchThreadId] = 0.0;
+	RWRayDistanceUAV[DispatchThreadId] = float2(-1.0, 0.0);
+
+	// Get G-Buffer surface data
+	uint2 PixelCoord = GetPixelCoord(DispatchThreadId, UpscaleFactor);
+	float2 InvBufferSize = View.BufferSizeAndInvSize.zw;
+	float2 UV = (float2(PixelCoord) + 0.5) * InvBufferSize;
+	FGBufferData GBufferData = GetGBufferDataFromSceneTextures(UV);
+	// Remap DiffuseColor when using SubsurfaceProfile (GBuffer decoding replaces with 100% albedo)
+	if (UseSubsurfaceProfile(GBufferData.ShadingModelID))
+	{
+		GBufferData.DiffuseColor = GBufferData.StoredBaseColor;
+	}
+
+	// Cull based on shading model
+	uint ShadingModelID = GBufferData.ShadingModelID;
+	if (ShadingModelID == SHADINGMODELID_UNLIT
+		|| ShadingModelID == SHADINGMODELID_TWOSIDED_FOLIAGE
+		)
+	{
+		return;
+	}
+
+	// Remap DiffuseColor when using SubsurfaceProfile (GBuffer decoding replaces with 100% albedo)
+	if (UseSubsurfaceProfile(GBufferData.ShadingModelID))
+	{
+		GBufferData.DiffuseColor = GBufferData.StoredBaseColor;
+	}
+	float3 DiffuseColor = GBufferData.DiffuseColor;
+	if (Luminance(DiffuseColor) < DiffuseThreshold)
+	{
+		return;
+	}
+
+	// Construct position, normal, and view direction
+	float DeviceZ = SceneDepthBuffer.Load(int3(PixelCoord, 0)).r;
+	float3 WorldPosition;
+	float3 CameraDirection;
+	ReconstructWorldPositionAndCameraDirectionFromDeviceZ(PixelCoord, DeviceZ, WorldPosition, CameraDirection);
+	float3 WorldNormal = GBufferData.WorldNormal;
+
+	uint GatherPointsIndex = DispatchThreadIdToLinearIndex(DispatchThreadId);
+
+	// Define rejection criteria and evaluate
+	bool bShouldReject[16];
+	FRejectionCriteria RejectionCriteria = CreateRejectionCriteria(WorldPosition, FinalGatherDistance);
+	CreateRejectionMask(RejectionCriteria, GatherPointsBuffer[GatherPointsIndex], SamplesPerPixel, bShouldReject);
+
+	float3 ExitantRadiance = 0.0;
+	float HitDistance = 1.0e7;
+	float HitCount = 0.0;
+	uint ValidSamplesPerPixel = 0;
+	for (uint LocalSampleIndex = 0; LocalSampleIndex < SamplesPerPixel; ++LocalSampleIndex)
+	{
+
+		if (bShouldReject[LocalSampleIndex])
+		{
+			continue;
+		}
+		ValidSamplesPerPixel++;
+
+		float3 GPPosition = GatherPointsBuffer[GatherPointsIndex].Position[LocalSampleIndex];
+		float3 GPIrradiance = GatherPointsBuffer[GatherPointsIndex].Irradiance[LocalSampleIndex];
+
+		// Initialize throughput based on BRDF
+		half3 N = WorldNormal;
+		half3 V = -CameraDirection;
+		half3 L = normalize(GPPosition - WorldPosition);
+		float NoL = saturate(dot(N, L));
+		FShadowTerms ShadowTerms = { 0.0, 0.0, 0.0 };
+		FDirectLighting LightingSample = EvaluateBxDF(GBufferData, N, V, L, NoL, ShadowTerms);
+		float3 RayThroughput = LightingSample.Diffuse / DiffuseColor;
+		float RayPdf = NoL / PI;
+
+		// Trace visibility ray
+		RayDesc Ray;
+		Ray.Origin = WorldPosition;
+		Ray.Direction = normalize(GPPosition - WorldPosition);
+		Ray.TMin = 0.01;
+		float RayLength = length(GPPosition - WorldPosition);
+		Ray.TMax = max(RayLength - 0.01, Ray.TMin);
+
+		// Cast ray
+		uint RayFlags = 0;
+#if !ENABLE_TWO_SIDED_GEOMETRY
+		RayFlags |= RAY_FLAG_CULL_BACK_FACING_TRIANGLES;
+#endif
+		FRayCone RayCone = (FRayCone)0;
+		FMaterialClosestHitPayload Payload = TraceRayInternal(
+			TLAS,   // AccelerationStructure
+			RayFlags,
+			RAY_TRACING_MASK_OPAQUE,
+			RAY_TRACING_SHADER_SLOT_SHADOW, // RayContributionToHitGroupIndex
+			RAY_TRACING_NUM_SHADER_SLOTS,     // MultiplierForGeometryContributionToShaderIndex
+			0,      // MissShaderIndex
+			Ray,    // RayDesc
+			RayCone
+		);
+
+		// No hit indicates successful next-event connection
+		if (Payload.IsMiss() && RayPdf > 0.0)
+		{
+			ExitantRadiance += GPIrradiance * RayThroughput / RayPdf;
+			HitDistance = min(HitDistance, Ray.TMax);
+			HitCount += 1.0;
+		}
+	}
+
+	float AmbientOcclusion = 0.0;
+	if (ValidSamplesPerPixel > 0)
+	{
+		ExitantRadiance /= ValidSamplesPerPixel;
+		AmbientOcclusion = HitCount / ValidSamplesPerPixel;
+	}
+#if USE_PREEXPOSURE
+	ExitantRadiance *= View.PreExposure;
+#endif
+	ExitantRadiance = ClampToHalfFloatRange(ExitantRadiance);
+
+	if (HitCount == 0.0)
+	{
+		HitDistance = -1.0;
+	}
+
+	RWGlobalIlluminationUAV[DispatchThreadId] = float4(ExitantRadiance, AmbientOcclusion);
+	RWRayDistanceUAV[DispatchThreadId] = float2(HitDistance, ValidSamplesPerPixel);
+}
diff --git a/Engine/Shaders/Private/RayTracing/RayTracingGatherPoints.ush b/Engine/Shaders/Private/RayTracing/RayTracingGatherPoints.ush
new file mode 100644
--- /dev/null
+++ b/Engine/Shaders/Private/RayTracing/RayTracingGatherPoints.ush
@@ -0,0 +1,31 @@
+// Copyright 1998-2019 Epic Games, Inc. All Rights Reserved.
+
+struct FGatherPoints
+{
+	float3 CreationPoint[16];
+	float3 Position[16];
+	float3 Irradiance[16];
+};
+
+struct FRejectionCriteria
+{
+	float3 CreationPoint;
+	float DistanceThreshold2;
+};
+
+FRejectionCriteria CreateRejectionCriteria(float3 CreationPoint, float DistanceThreshold)
+{
+	FRejectionCriteria RejectionCriteria;
+	RejectionCriteria.CreationPoint = CreationPoint;
+	RejectionCriteria.DistanceThreshold2 = DistanceThreshold * DistanceThreshold;
+	return RejectionCriteria;
+}
+
+void CreateRejectionMask(FRejectionCriteria RejectionCriteria, FGatherPoints GatherPoints, int SampleCount, out bool bShouldReject[16])
+{
+	for (int SampleIndex = 0; SampleIndex < SampleCount; ++SampleIndex)
+	{
+		float3 GPDistance2 = RejectionCriteria.CreationPoint - GatherPoints.CreationPoint[SampleIndex];
+		bShouldReject[SampleIndex] = dot(GPDistance2, GPDistance2) > RejectionCriteria.DistanceThreshold2;
+	}
+}
diff --git a/Engine/Source/Runtime/Renderer/Private/DeferredShadingRenderer.h b/Engine/Source/Runtime/Renderer/Private/DeferredShadingRenderer.h
--- a/Engine/Source/Runtime/Renderer/Private/DeferredShadingRenderer.h
+++ b/Engine/Source/Runtime/Renderer/Private/DeferredShadingRenderer.h
@@ -527,6 +527,23 @@ private:
 		FRDGTextureUAV* GlobalIlluminationUAV,
 		FRDGTextureUAV* RayDistanceUAV);
 
+	void RayTracingGlobalIlluminationCreateGatherPoints(
+		FRDGBuilder& GraphBuilder,
+		FSceneTextureParameters& SceneTextures,
+		FViewInfo& View,
+		int32 UpscaleFactor,
+		FRDGBufferRef& GatherPointsBuffer,
+		FIntPoint& GatherPointsResolution);
+
+	bool RenderRayTracingGlobalIlluminationFinalGather(
+		FRDGBuilder& GraphBuilder,
+		FSceneTextureParameters& SceneTextures,
+		FViewInfo& View,
+		const IScreenSpaceDenoiser::FAmbientOcclusionRayTracingConfig& RayTracingConfig,
+		int32 UpscaleFactor,
+		FRDGTextureUAV* GlobalIlluminationUAV,
+		FRDGTextureUAV* RayDistanceUAV);
+
 	void BuildSkyLightCdfs(FRHICommandListImmediate& RHICmdList, FSkyLightSceneProxy* SkyLight);
 	void BuildSkyLightMipTree(FRHICommandListImmediate& RHICmdList, FTextureRHIRef SkyLightTexture, FRWBuffer& SkyLightMipTreePosX, FRWBuffer& SkyLightMipTreePosY, FRWBuffer& SkyLightMipTreePosZ, FRWBuffer& SkyLightMipTreeNegX, FRWBuffer& SkyLightMipTreeNegY, FRWBuffer& SkyLightMipTreeNegZ, FIntVector& SkyLightMipTreeDimensions);
 	void BuildSkyLightMipTreePdf(
diff --git a/Engine/Source/Runtime/Renderer/Private/RayTracing/RayTracingGlobalIllumination.cpp b/Engine/Source/Runtime/Renderer/Private/RayTracing/RayTracingGlobalIllumination.cpp
--- a/Engine/Source/Runtime/Renderer/Private/RayTracing/RayTracingGlobalIllumination.cpp
+++ b/Engine/Source/Runtime/Renderer/Private/RayTracing/RayTracingGlobalIllumination.cpp
@@ -22,13 +22,13 @@
 #include "BlueNoise.h"
 #include "SceneTextureParameters.h"
 
-static int32 GRayTracingGlobalIllumination = -1;
-static FAutoConsoleVariableRef CVarRayTracingGlobalIllumination(
+static TAutoConsoleVariable<int32> CVarRayTracingGlobalIllumination(
 	TEXT("r.RayTracing.GlobalIllumination"),
-	GRayTracingGlobalIllumination,
+	-1,
 	TEXT("-1: Value driven by postprocess volume (default) \n")
 	TEXT(" 0: ray tracing global illumination off \n")
-	TEXT(" 1: ray tracing global illumination enabled")
+	TEXT(" 1: ray tracing global illumination enabled"),
+	ECVF_RenderThreadSafe
 );
 
 static int32 GRayTracingGlobalIlluminationSamplesPerPixel = -1;
@@ -112,9 +112,25 @@ static FAutoConsoleVariableRef CVarRayTracingGlobalIlluminationTileSize(
 	TEXT("Render ray traced global illumination in NxN piel tiles, where each tile is submitted as separate GPU command buffer, allowing high quality rendering without triggering timeout detection. (default = 0, tiling disabled)")
 );
 
+static TAutoConsoleVariable<int32> CVarRayTracingGlobalIlluminationEnableFinalGather(
+	TEXT("r.RayTracing.GlobalIllumination.EnableFinalGather"),
+	0,
+	TEXT("Enables final gather algorithm for 1-bounce global illumination (default = 0)"),
+	ECVF_RenderThreadSafe
+);
+
+static float GRayTracingGlobalIlluminationFinalGatherDistance = 10.0;
+static FAutoConsoleVariableRef CVarRayTracingGlobalIlluminationFinalGatherDistance(
+	TEXT("r.RayTracing.GlobalIllumination.FinalGatherDistance"),
+	GRayTracingGlobalIlluminationFinalGatherDistance,
+	TEXT("Maximum world-space distance for valid, reprojected final gather points (default = 10)")
+);
+
 static const int32 GLightCountMax = 64;
 
 DECLARE_GPU_STAT_NAMED(RayTracingGIBruteForce, TEXT("Ray Tracing GI: Brute Force"));
+DECLARE_GPU_STAT_NAMED(RayTracingGIFinalGather, TEXT("Ray Tracing GI: Final Gather"));
+DECLARE_GPU_STAT_NAMED(RayTracingGICreateGatherPoints, TEXT("Ray Tracing GI: Create Gather Points"));
 
 void SetupLightParameters(
 	const TSparseArray<FLightSceneInfoCompact>& Lights,
@@ -215,9 +231,11 @@ bool ShouldRenderRayTracingGlobalIllumination(const FViewInfo& View)
 	{
 		return GetForceRayTracingEffectsCVarValue() > 0;
 	}
-	else if (GRayTracingGlobalIllumination >= 0)
+
+	int32 CVarRayTracingGlobalIlluminationValue = CVarRayTracingGlobalIllumination.GetValueOnRenderThread();
+	if (CVarRayTracingGlobalIlluminationValue >= 0)
 	{
-		return (GRayTracingGlobalIllumination > 0);
+		return (CVarRayTracingGlobalIlluminationValue > 0);
 	}
 	else 
 	{
@@ -324,6 +342,114 @@ IMPLEMENT_GLOBAL_SHADER(FRayTracingGlobalIlluminationCHS, "/Engine/Private/RayTr
 IMPLEMENT_GLOBAL_SHADER(FRayTracingGlobalIlluminationCompositePS, "/Engine/Private/RayTracing/RayTracingGlobalIlluminationCompositePS.usf", "GlobalIlluminationCompositePS", SF_Pixel);
 IMPLEMENT_GLOBAL_SHADER(FRayTracingGlobalIlluminationSceneColorCompositePS, "/Engine/Private/RayTracing/RayTracingGlobalIlluminationCompositePS.usf", "GlobalIlluminationSceneColorCompositePS", SF_Pixel);
 
+struct GatherPoints
+{
+	FVector CreationPoint[16];
+	FVector Position[16];
+	FVector Irradiance[16];
+};
+
+class FRayTracingGlobalIlluminationCreateGatherPointsRGS : public FGlobalShader
+{
+	DECLARE_GLOBAL_SHADER(FRayTracingGlobalIlluminationCreateGatherPointsRGS)
+	SHADER_USE_ROOT_PARAMETER_STRUCT(FRayTracingGlobalIlluminationCreateGatherPointsRGS, FGlobalShader)
+
+	class FUseAttenuationTermDim : SHADER_PERMUTATION_BOOL("USE_ATTENUATION_TERM");
+	class FEnableTwoSidedGeometryDim : SHADER_PERMUTATION_BOOL("ENABLE_TWO_SIDED_GEOMETRY");
+
+	using FPermutationDomain = TShaderPermutationDomain<FUseAttenuationTermDim, FEnableTwoSidedGeometryDim>;
+
+	static bool ShouldCompilePermutation(const FGlobalShaderPermutationParameters& Parameters)
+	{
+		return ShouldCompileRayTracingShadersForProject(Parameters.Platform);
+	}
+
+	BEGIN_SHADER_PARAMETER_STRUCT(FParameters, )
+		SHADER_PARAMETER(uint32, SamplesPerPixel)
+		SHADER_PARAMETER(uint32, SampleIndex)
+		SHADER_PARAMETER(uint32, MaxBounces)
+		SHADER_PARAMETER(uint32, UpscaleFactor)
+		SHADER_PARAMETER(uint32, TileOffsetX)
+		SHADER_PARAMETER(uint32, TileOffsetY)
+		SHADER_PARAMETER(float, MaxRayDistanceForGI)
+		SHADER_PARAMETER(float, NextEventEstimationSamples)
+		SHADER_PARAMETER(float, DiffuseThreshold)
+		SHADER_PARAMETER(float, MaxNormalBias)
+		SHADER_PARAMETER(bool, EvalSkyLight)
+		SHADER_PARAMETER(bool, UseRussianRoulette)
+
+		// Scene data
+		SHADER_PARAMETER_SRV(RaytracingAccelerationStructure, TLAS)
+		SHADER_PARAMETER_STRUCT_REF(FViewUniformShaderParameters, ViewUniformBuffer)
+
+		// Sampling sequence
+		SHADER_PARAMETER_STRUCT_REF(FHaltonIteration, HaltonIteration)
+		SHADER_PARAMETER_STRUCT_REF(FHaltonPrimes, HaltonPrimes)
+		SHADER_PARAMETER_STRUCT_REF(FBlueNoise, BlueNoise)
+
+		// Light data
+		SHADER_PARAMETER_STRUCT_REF(FPathTracingLightData, LightParameters)
+		SHADER_PARAMETER_STRUCT_REF(FSkyLightData, SkyLight)
+
+		// Shading data
+		SHADER_PARAMETER_STRUCT_INCLUDE(FSceneTextureParameters, SceneTextures)
+		SHADER_PARAMETER_RDG_TEXTURE(Texture2D, SSProfilesTexture)
+		SHADER_PARAMETER_SAMPLER(SamplerState, TransmissionProfilesLinearSampler)
+
+		SHADER_PARAMETER(FIntPoint, GatherPointsResolution)
+		// Output
+		SHADER_PARAMETER_RDG_BUFFER_UAV(RWStructuredBuffer<GatherPoints>, RWGatherPointsBuffer)
+		END_SHADER_PARAMETER_STRUCT()
+};
+
+IMPLEMENT_GLOBAL_SHADER(FRayTracingGlobalIlluminationCreateGatherPointsRGS, "/Engine/Private/RayTracing/RayTracingCreateGatherPointsRGS.usf", "RayTracingCreateGatherPointsRGS", SF_RayGen);
+
+class FRayTracingGlobalIlluminationFinalGatherRGS : public FGlobalShader
+{
+	DECLARE_GLOBAL_SHADER(FRayTracingGlobalIlluminationFinalGatherRGS)
+	SHADER_USE_ROOT_PARAMETER_STRUCT(FRayTracingGlobalIlluminationFinalGatherRGS, FGlobalShader)
+
+	class FUseAttenuationTermDim : SHADER_PERMUTATION_BOOL("USE_ATTENUATION_TERM");
+	class FEnableTwoSidedGeometryDim : SHADER_PERMUTATION_BOOL("ENABLE_TWO_SIDED_GEOMETRY");
+
+	using FPermutationDomain = TShaderPermutationDomain<FUseAttenuationTermDim, FEnableTwoSidedGeometryDim>;
+
+	static bool ShouldCompilePermutation(const FGlobalShaderPermutationParameters& Parameters)
+	{
+		return ShouldCompileRayTracingShadersForProject(Parameters.Platform);
+	}
+
+	BEGIN_SHADER_PARAMETER_STRUCT(FParameters, )
+		SHADER_PARAMETER(uint32, SampleIndex)
+		SHADER_PARAMETER(uint32, SamplesPerPixel)
+		SHADER_PARAMETER(uint32, UpscaleFactor)
+		SHADER_PARAMETER(uint32, TileOffsetX)
+		SHADER_PARAMETER(uint32, TileOffsetY)
+		SHADER_PARAMETER(float, DiffuseThreshold)
+		SHADER_PARAMETER(float, MaxNormalBias)
+		SHADER_PARAMETER(float, FinalGatherDistance)
+
+		// Scene data
+		SHADER_PARAMETER_SRV(RaytracingAccelerationStructure, TLAS)
+		SHADER_PARAMETER_STRUCT_REF(FViewUniformShaderParameters, ViewUniformBuffer)
+
+		// Shading data
+		SHADER_PARAMETER_STRUCT_INCLUDE(FSceneTextureParameters, SceneTextures)
+		SHADER_PARAMETER_RDG_TEXTURE(Texture2D, SSProfilesTexture)
+		SHADER_PARAMETER_SAMPLER(SamplerState, TransmissionProfilesLinearSampler)
+
+		// Gather points
+		SHADER_PARAMETER_RDG_BUFFER_SRV(StructuredBuffer<GatherPoints>, GatherPointsBuffer)
+		SHADER_PARAMETER(FIntPoint, GatherPointsResolution)
+
+		// Output
+		SHADER_PARAMETER_RDG_TEXTURE_UAV(RWTexture2D<float4>, RWGlobalIlluminationUAV)
+		SHADER_PARAMETER_RDG_TEXTURE_UAV(RWTexture2D<float>, RWRayDistanceUAV)
+		END_SHADER_PARAMETER_STRUCT()
+};
+
+IMPLEMENT_GLOBAL_SHADER(FRayTracingGlobalIlluminationFinalGatherRGS, "/Engine/Private/RayTracing/RayTracingFinalGatherRGS.usf", "RayTracingFinalGatherRGS", SF_RayGen);
+
 void FDeferredShadingSceneRenderer::PrepareRayTracingGlobalIllumination(const FViewInfo& View, TArray<FRHIRayTracingShader*>& OutRayGenShaders)
 {
 	// Declare all RayGen shaders that require material closest hit shaders to be bound
@@ -336,6 +462,18 @@ void FDeferredShadingSceneRenderer::PrepareRayTracingGlobalIllumination(const FV
 			PermutationVector.Set<FGlobalIlluminationRGS::FEnableTwoSidedGeometryDim>(EnableTwoSidedGeometry == 1);
 			TShaderMapRef<FGlobalIlluminationRGS> RayGenerationShader(View.ShaderMap, PermutationVector);
 			OutRayGenShaders.Add(RayGenerationShader->GetRayTracingShader());
+
+			FRayTracingGlobalIlluminationCreateGatherPointsRGS::FPermutationDomain CreateGatherPointsPermutationVector;
+			CreateGatherPointsPermutationVector.Set<FRayTracingGlobalIlluminationCreateGatherPointsRGS::FUseAttenuationTermDim>(UseAttenuationTerm == 1);
+			CreateGatherPointsPermutationVector.Set<FRayTracingGlobalIlluminationCreateGatherPointsRGS::FEnableTwoSidedGeometryDim>(EnableTwoSidedGeometry == 1);
+			TShaderMapRef<FRayTracingGlobalIlluminationCreateGatherPointsRGS> CreateGatherPointsRayGenerationShader(View.ShaderMap, CreateGatherPointsPermutationVector);
+			OutRayGenShaders.Add(CreateGatherPointsRayGenerationShader->GetRayTracingShader());
+
+			FRayTracingGlobalIlluminationFinalGatherRGS::FPermutationDomain GatherPassPermutationVector;
+			GatherPassPermutationVector.Set<FRayTracingGlobalIlluminationFinalGatherRGS::FUseAttenuationTermDim>(UseAttenuationTerm == 1);
+			GatherPassPermutationVector.Set<FRayTracingGlobalIlluminationFinalGatherRGS::FEnableTwoSidedGeometryDim>(EnableTwoSidedGeometry == 1);
+			TShaderMapRef<FRayTracingGlobalIlluminationFinalGatherRGS> GatherPassRayGenerationShader(View.ShaderMap, GatherPassPermutationVector);
+			OutRayGenShaders.Add(GatherPassRayGenerationShader->GetRayTracingShader());
 		}
 	}
 
@@ -456,7 +594,15 @@ void FDeferredShadingSceneRenderer::RenderRayTracingGlobalIllumination(
 	}
 
 	// Ray generation
-	bool bIsValid = RenderRayTracingGlobalIlluminationBruteForce(RHICmdList, GraphBuilder, SceneTextures, View, RayTracingConfig, UpscaleFactor, GlobalIlluminationUAV, RayDistanceUAV);
+	bool bIsValid;
+	if (CVarRayTracingGlobalIlluminationEnableFinalGather.GetValueOnRenderThread() != 0)
+	{
+		bIsValid = RenderRayTracingGlobalIlluminationFinalGather(GraphBuilder, SceneTextures, View, RayTracingConfig, UpscaleFactor, GlobalIlluminationUAV, RayDistanceUAV);
+	}
+	else
+	{
+		bIsValid = RenderRayTracingGlobalIlluminationBruteForce(RHICmdList, GraphBuilder, SceneTextures, View, RayTracingConfig, UpscaleFactor, GlobalIlluminationUAV, RayDistanceUAV);
+	}
 
 	// Denoising
 
@@ -546,6 +692,208 @@ void FDeferredShadingSceneRenderer::RenderRayTracingGlobalIllumination(
 	}
 }
 
+void FDeferredShadingSceneRenderer::RayTracingGlobalIlluminationCreateGatherPoints(
+	FRDGBuilder& GraphBuilder,
+	FSceneTextureParameters& SceneTextures,
+	FViewInfo& View,
+	int32 UpscaleFactor,
+	FRDGBufferRef& GatherPointsBuffer,
+	FIntPoint& GatherPointsResolution
+)
+#if RHI_RAYTRACING
+{
+	RDG_GPU_STAT_SCOPE(GraphBuilder, RayTracingGICreateGatherPoints);
+	RDG_EVENT_SCOPE(GraphBuilder, "Ray Tracing GI: Create Gather Points");
+
+	int32 GatherSamples = GetRayTracingGlobalIlluminationSamplesPerPixel(View);
+	int32 SamplesPerPixel = 1;
+
+	uint32 IterationCount = SamplesPerPixel;
+	uint32 SequenceCount = 1;
+	uint32 DimensionCount = 24;
+	int32 FrameIndex = View.ViewState->FrameIndex % 1024;
+	FHaltonSequenceIteration HaltonSequenceIteration(Scene->HaltonSequence, IterationCount, SequenceCount, DimensionCount, FrameIndex);
+
+	FHaltonIteration HaltonIteration;
+	InitializeHaltonSequenceIteration(HaltonSequenceIteration, HaltonIteration);
+
+	FHaltonPrimes HaltonPrimes;
+	InitializeHaltonPrimes(Scene->HaltonPrimesResource, HaltonPrimes);
+
+	FBlueNoise BlueNoise;
+	InitializeBlueNoise(BlueNoise);
+
+	FPathTracingLightData LightParameters;
+	SetupLightParameters(Scene->Lights, View, &LightParameters);
+
+	FSkyLightData SkyLightParameters;
+	SetupSkyLightParameters(*Scene, &SkyLightParameters);
+
+	FRayTracingGlobalIlluminationCreateGatherPointsRGS::FParameters* PassParameters = GraphBuilder.AllocParameters<FRayTracingGlobalIlluminationCreateGatherPointsRGS::FParameters>();
+	PassParameters->SampleIndex = (FrameIndex * SamplesPerPixel) % GatherSamples;
+	PassParameters->SamplesPerPixel = SamplesPerPixel;
+	PassParameters->MaxBounces = 1;
+	PassParameters->MaxNormalBias = GetRaytracingMaxNormalBias();
+	PassParameters->MaxRayDistanceForGI = GRayTracingGlobalIlluminationMaxRayDistance;
+	PassParameters->EvalSkyLight = GRayTracingGlobalIlluminationEvalSkyLight != 0;
+	PassParameters->UseRussianRoulette = GRayTracingGlobalIlluminationUseRussianRoulette != 0;
+	PassParameters->DiffuseThreshold = GRayTracingGlobalIlluminationDiffuseThreshold;
+	PassParameters->NextEventEstimationSamples = GRayTracingGlobalIlluminationNextEventEstimationSamples;
+	PassParameters->UpscaleFactor = UpscaleFactor;
+	PassParameters->TileOffsetX = 0;
+	PassParameters->TileOffsetY = 0;
+
+	// Global
+	PassParameters->TLAS = View.RayTracingScene.RayTracingSceneRHI->GetShaderResourceView();
+	PassParameters->ViewUniformBuffer = View.ViewUniformBuffer;
+
+	// Sampling sequence
+	PassParameters->HaltonIteration = CreateUniformBufferImmediate(HaltonIteration, EUniformBufferUsage::UniformBuffer_SingleDraw);
+	PassParameters->HaltonPrimes = CreateUniformBufferImmediate(HaltonPrimes, EUniformBufferUsage::UniformBuffer_SingleDraw);
+	PassParameters->BlueNoise = CreateUniformBufferImmediate(BlueNoise, EUniformBufferUsage::UniformBuffer_SingleDraw);
+
+	// Light data
+	PassParameters->LightParameters = CreateUniformBufferImmediate(LightParameters, EUniformBufferUsage::UniformBuffer_SingleDraw);
+	PassParameters->SceneTextures = SceneTextures;
+	PassParameters->SkyLight = CreateUniformBufferImmediate(SkyLightParameters, EUniformBufferUsage::UniformBuffer_SingleDraw);
+
+	// Shading data
+	TRefCountPtr<IPooledRenderTarget> SubsurfaceProfileRT((IPooledRenderTarget*)GetSubsufaceProfileTexture_RT(GraphBuilder.RHICmdList));
+	if (!SubsurfaceProfileRT)
+	{
+		SubsurfaceProfileRT = GSystemTextures.BlackDummy;
+	}
+	PassParameters->SSProfilesTexture = GraphBuilder.RegisterExternalTexture(SubsurfaceProfileRT);
+	PassParameters->TransmissionProfilesLinearSampler = TStaticSamplerState<SF_Bilinear, AM_Clamp, AM_Clamp, AM_Clamp>::GetRHI();
+
+	// Output
+	FIntPoint LocalGatherPointsResolution = FIntPoint::DivideAndRoundUp(View.ViewRect.Size(), UpscaleFactor);
+	if (GatherPointsResolution != LocalGatherPointsResolution)
+	{
+		GatherPointsResolution = LocalGatherPointsResolution;
+		FRDGBufferDesc BufferDesc = FRDGBufferDesc::CreateStructuredDesc(sizeof(GatherPoints), GatherPointsResolution.X * GatherPointsResolution.Y);
+		GatherPointsBuffer = GraphBuilder.CreateBuffer(BufferDesc, TEXT("GatherPointsBuffer"), ERDGResourceFlags::MultiFrame);
+	}
+	else
+	{
+		GatherPointsBuffer = GraphBuilder.RegisterExternalBuffer(((FSceneViewState*)View.State)->GatherPointsBuffer, TEXT("GatherPointsBuffer"));
+	}
+	PassParameters->GatherPointsResolution = GatherPointsResolution;
+	PassParameters->RWGatherPointsBuffer = GraphBuilder.CreateUAV(GatherPointsBuffer, EPixelFormat::PF_R32_UINT);
+
+	FRayTracingGlobalIlluminationCreateGatherPointsRGS::FPermutationDomain PermutationVector;
+	PermutationVector.Set<FRayTracingGlobalIlluminationCreateGatherPointsRGS::FUseAttenuationTermDim>(true);
+	PermutationVector.Set<FRayTracingGlobalIlluminationCreateGatherPointsRGS::FEnableTwoSidedGeometryDim>(CVarRayTracingGlobalIlluminationEnableTwoSidedGeometry.GetValueOnRenderThread() != 0);
+	TShaderMapRef<FRayTracingGlobalIlluminationCreateGatherPointsRGS> RayGenerationShader(GetGlobalShaderMap(FeatureLevel), PermutationVector);
+	ClearUnusedGraphResources(*RayGenerationShader, PassParameters);
+
+	GraphBuilder.AddPass(
+		RDG_EVENT_NAME("GatherPoints %d%d", GatherPointsResolution.X, GatherPointsResolution.Y),
+		PassParameters,
+		ERDGPassFlags::Compute,
+		[PassParameters, this, &View, RayGenerationShader, GatherPointsResolution](FRHICommandList& RHICmdList)
+	{
+		FRHIRayTracingScene* RayTracingSceneRHI = View.RayTracingScene.RayTracingSceneRHI;
+		FRayTracingShaderBindingsWriter GlobalResources;
+		SetShaderParameters(GlobalResources, *RayGenerationShader, *PassParameters);
+		RHICmdList.RayTraceDispatch(View.RayTracingMaterialPipeline, RayGenerationShader->GetRayTracingShader(), RayTracingSceneRHI, GlobalResources, GatherPointsResolution.X, GatherPointsResolution.Y);
+	});
+}
+#else
+{
+	unimplemented();
+}
+#endif
+
+bool FDeferredShadingSceneRenderer::RenderRayTracingGlobalIlluminationFinalGather(
+	FRDGBuilder& GraphBuilder,
+	FSceneTextureParameters& SceneTextures,
+	FViewInfo& View,
+	const IScreenSpaceDenoiser::FAmbientOcclusionRayTracingConfig& RayTracingConfig,
+	int32 UpscaleFactor,
+	// Output
+	FRDGTextureUAV* GlobalIlluminationUAV,
+	FRDGTextureUAV* RayDistanceUAV)
+#if RHI_RAYTRACING
+{
+	FSceneViewState* SceneViewState = (FSceneViewState*)View.State;
+	if (!SceneViewState) return false;
+
+	int32 SamplesPerPixel = GetRayTracingGlobalIlluminationSamplesPerPixel(View);
+	if (SamplesPerPixel <= 0) return false;
+
+	// Generate gather points
+	FRDGBufferRef GatherPointsBuffer;
+	RayTracingGlobalIlluminationCreateGatherPoints(GraphBuilder, SceneTextures, View, UpscaleFactor, GatherPointsBuffer, SceneViewState->GatherPointsResolution);
+
+	// Perform gather
+	RDG_GPU_STAT_SCOPE(GraphBuilder, RayTracingGIFinalGather);
+	RDG_EVENT_SCOPE(GraphBuilder, "Ray Tracing GI: Final Gather");
+
+	FRayTracingGlobalIlluminationFinalGatherRGS::FParameters* PassParameters = GraphBuilder.AllocParameters<FRayTracingGlobalIlluminationFinalGatherRGS::FParameters>();
+	int32 SampleIndex = View.ViewState->FrameIndex % SamplesPerPixel;
+	PassParameters->SampleIndex = SampleIndex;
+	PassParameters->SamplesPerPixel = SamplesPerPixel;
+	PassParameters->DiffuseThreshold = GRayTracingGlobalIlluminationDiffuseThreshold;
+	PassParameters->MaxNormalBias = GetRaytracingMaxNormalBias();
+	PassParameters->FinalGatherDistance = GRayTracingGlobalIlluminationFinalGatherDistance;
+	PassParameters->UpscaleFactor = UpscaleFactor;
+	PassParameters->TileOffsetX = 0;
+	PassParameters->TileOffsetY = 0;
+
+	// Scene data
+	PassParameters->TLAS = View.RayTracingScene.RayTracingSceneRHI->GetShaderResourceView();
+	PassParameters->ViewUniformBuffer = View.ViewUniformBuffer;
+
+	// Shading data
+	PassParameters->SceneTextures = SceneTextures;
+	TRefCountPtr<IPooledRenderTarget> SubsurfaceProfileRT((IPooledRenderTarget*)GetSubsufaceProfileTexture_RT(GraphBuilder.RHICmdList));
+	if (!SubsurfaceProfileRT)
+	{
+		SubsurfaceProfileRT = GSystemTextures.BlackDummy;
+	}
+	PassParameters->SSProfilesTexture = GraphBuilder.RegisterExternalTexture(SubsurfaceProfileRT);
+	PassParameters->TransmissionProfilesLinearSampler = TStaticSamplerState<SF_Bilinear, AM_Clamp, AM_Clamp, AM_Clamp>::GetRHI();
+
+	// Gather points
+	PassParameters->GatherPointsResolution = SceneViewState->GatherPointsResolution;
+	PassParameters->GatherPointsBuffer = GraphBuilder.CreateSRV(GatherPointsBuffer);
+
+	// Output
+	PassParameters->RWGlobalIlluminationUAV = GlobalIlluminationUAV;
+	PassParameters->RWRayDistanceUAV = RayDistanceUAV;
+
+	FRayTracingGlobalIlluminationFinalGatherRGS::FPermutationDomain PermutationVector;
+	PermutationVector.Set<FRayTracingGlobalIlluminationFinalGatherRGS::FUseAttenuationTermDim>(true);
+	PermutationVector.Set<FRayTracingGlobalIlluminationFinalGatherRGS::FEnableTwoSidedGeometryDim>(CVarRayTracingGlobalIlluminationEnableTwoSidedGeometry.GetValueOnRenderThread() != 0);
+	TShaderMapRef<FRayTracingGlobalIlluminationFinalGatherRGS> RayGenerationShader(GetGlobalShaderMap(FeatureLevel), PermutationVector);
+	ClearUnusedGraphResources(*RayGenerationShader, PassParameters);
+
+	FIntPoint RayTracingResolution = FIntPoint::DivideAndRoundUp(View.ViewRect.Size(), UpscaleFactor);
+	GraphBuilder.AddPass(
+		RDG_EVENT_NAME("GlobalIlluminationRayTracing %dx%d", RayTracingResolution.X, RayTracingResolution.Y),
+		PassParameters,
+		ERDGPassFlags::Compute,
+		[PassParameters, this, &View, RayGenerationShader, RayTracingResolution](FRHICommandList& RHICmdList)
+	{
+		FRHIRayTracingScene* RayTracingSceneRHI = View.RayTracingScene.RayTracingSceneRHI;
+
+		FRayTracingShaderBindingsWriter GlobalResources;
+		SetShaderParameters(GlobalResources, *RayGenerationShader, *PassParameters);
+		RHICmdList.RayTraceDispatch(View.RayTracingMaterialPipeline, RayGenerationShader->GetRayTracingShader(), RayTracingSceneRHI, GlobalResources, RayTracingResolution.X, RayTracingResolution.Y);
+	});
+
+
+	GraphBuilder.QueueBufferExtraction(GatherPointsBuffer, &SceneViewState->GatherPointsBuffer);
+	return true;
+}
+#else
+{
+	unimplemented();
+	return true;
+}
+#endif
+
 bool FDeferredShadingSceneRenderer::RenderRayTracingGlobalIlluminationBruteForce(
 	FRHICommandListImmediate& RHICmdList,
 	FRDGBuilder& GraphBuilder,
diff --git a/Engine/Source/Runtime/Renderer/Private/RendererScene.cpp b/Engine/Source/Runtime/Renderer/Private/RendererScene.cpp
--- a/Engine/Source/Runtime/Renderer/Private/RendererScene.cpp
+++ b/Engine/Source/Runtime/Renderer/Private/RendererScene.cpp
@@ -220,6 +220,9 @@ FSceneViewState::FSceneViewState()
 	}
 	bReadbackInitialized = false;
 	RayCountGPUReadback = new FRHIGPUBufferReadback(TEXT("Ray Count Readback"));
+
+	GatherPointsBuffer = nullptr;
+	GatherPointsResolution = FIntPoint(0, 0);
 #endif
 }
 
diff --git a/Engine/Source/Runtime/Renderer/Private/ScenePrivate.h b/Engine/Source/Runtime/Renderer/Private/ScenePrivate.h
--- a/Engine/Source/Runtime/Renderer/Private/ScenePrivate.h
+++ b/Engine/Source/Runtime/Renderer/Private/ScenePrivate.h
@@ -914,6 +914,9 @@ public:
 
 	// IES light profiles
 	FIESLightProfileResource IESLightProfileResources;
+
+	TRefCountPtr<FPooledRDGBuffer> GatherPointsBuffer;
+	FIntPoint GatherPointsResolution;
 #endif
 
 	// cache for stencil reads to a avoid reallocations of the SRV, Key is to detect if the object has changed
