Cherry-pick: CL 8556592

Extend maximum light count for ray tracing global illumination pass.

Added CVar to control the maximum count of GI-casting lights:
r.RayTracing.GlobalIllumination.MaxLightCount (default = 256)

diff --git a/Engine/Shaders/Shared/RayTracingDefinitions.h b/Engine/Shaders/Shared/RayTracingDefinitions.h
--- a/Engine/Shaders/Shared/RayTracingDefinitions.h
+++ b/Engine/Shaders/Shared/RayTracingDefinitions.h
@@ -24,5 +24,4 @@
 #define RAY_TRACING_SHADER_SLOT_SHADOW		1
 #define RAY_TRACING_NUM_SHADER_SLOTS		2
 
-#define RAY_TRACING_LIGHT_COUNT_MAXIMUM		64
-
+#define RAY_TRACING_LIGHT_COUNT_MAXIMUM		256
diff --git a/Engine/Source/Runtime/Renderer/Private/RayTracing/RayTracingGlobalIllumination.cpp b/Engine/Source/Runtime/Renderer/Private/RayTracing/RayTracingGlobalIllumination.cpp
--- a/Engine/Source/Runtime/Renderer/Private/RayTracing/RayTracingGlobalIllumination.cpp
+++ b/Engine/Source/Runtime/Renderer/Private/RayTracing/RayTracingGlobalIllumination.cpp
@@ -21,6 +21,7 @@
 #include "Raytracing/RaytracingOptions.h"
 #include "BlueNoise.h"
 #include "SceneTextureParameters.h"
+#include "RayTracingDefinitions.h"
 
 static TAutoConsoleVariable<int32> CVarRayTracingGlobalIllumination(
 	TEXT("r.RayTracing.GlobalIllumination"),
@@ -112,6 +113,13 @@ static FAutoConsoleVariableRef CVarRayTracingGlobalIlluminationTileSize(
 	TEXT("Render ray traced global illumination in NxN piel tiles, where each tile is submitted as separate GPU command buffer, allowing high quality rendering without triggering timeout detection. (default = 0, tiling disabled)")
 );
 
+static TAutoConsoleVariable<int32> CVarRayTracingGlobalIlluminationMaxLightCount(
+	TEXT("r.RayTracing.GlobalIllumination.MaxLightCount"),
+	RAY_TRACING_LIGHT_COUNT_MAXIMUM,
+	TEXT("Enables two-sided geometry when tracing GI rays (default = 256)"),
+	ECVF_RenderThreadSafe
+);
+
 static TAutoConsoleVariable<int32> CVarRayTracingGlobalIlluminationEnableFinalGather(
 	TEXT("r.RayTracing.GlobalIllumination.EnableFinalGather"),
 	0,
@@ -126,7 +134,6 @@ static FAutoConsoleVariableRef CVarRayTracingGlobalIlluminationFinalGatherDistan
 	TEXT("Maximum world-space distance for valid, reprojected final gather points (default = 10)")
 );
 
-static const int32 GLightCountMax = 64;
 
 DECLARE_GPU_STAT_NAMED(RayTracingGIBruteForce, TEXT("Ray Tracing GI: Brute Force"));
 DECLARE_GPU_STAT_NAMED(RayTracingGIFinalGather, TEXT("Ray Tracing GI: Final Gather"));
@@ -146,9 +153,10 @@ void SetupLightParameters(
 	LightParameters->Color[SkyLightIndex] = FVector(1.0);
 	LightParameters->Count++;
 
+	uint32 MaxLightCount = FMath::Min(CVarRayTracingGlobalIlluminationMaxLightCount.GetValueOnRenderThread(), RAY_TRACING_LIGHT_COUNT_MAXIMUM);
 	for (auto Light : Lights)
 	{
-		if (LightParameters->Count >= GLightCountMax) break;
+		if (LightParameters->Count >= MaxLightCount) break;
 
 		if (Light.LightSceneInfo->Proxy->HasStaticLighting() && Light.LightSceneInfo->IsPrecomputedLightingValid()) continue;
 		if (!Light.LightSceneInfo->Proxy->AffectGlobalIllumination()) continue;
