Fix ICWYU error.

diff --git a/Engine/Source/Runtime/Renderer/Private/ConvertToUniformMesh.cpp b/Engine/Source/Runtime/Renderer/Private/ConvertToUniformMesh.cpp
--- a/Engine/Source/Runtime/Renderer/Private/ConvertToUniformMesh.cpp
+++ b/Engine/Source/Runtime/Renderer/Private/ConvertToUniformMesh.cpp
@@ -15,6 +15,7 @@
 #include "ScenePrivate.h"
 #include "DistanceFieldLightingShared.h"
 #include "MeshPassProcessor.inl"
+#include "MaterialShared.h"
 
 class FConvertToUniformMeshVS : public FMeshMaterialShader
 {
@@ -306,7 +307,7 @@ public:
 
 	static bool ShouldCompilePermutation(const FMaterialShaderPermutationParameters& Parameters)
 	{
-		if (Material->GetMaterialDomain() == MD_UI)
+		if (Material->IsUIMaterial())
 		{
 			return false;
 		}
