Fix Compile Error.

diff --git a/Engine/Source/Runtime/Renderer/Private/ConvertToUniformMesh.cpp b/Engine/Source/Runtime/Renderer/Private/ConvertToUniformMesh.cpp
--- a/Engine/Source/Runtime/Renderer/Private/ConvertToUniformMesh.cpp
+++ b/Engine/Source/Runtime/Renderer/Private/ConvertToUniformMesh.cpp
@@ -307,12 +307,12 @@ public:
 
 	static bool ShouldCompilePermutation(const FMaterialShaderPermutationParameters& Parameters)
 	{
-		if (Material->IsUIMaterial())
+		if (Parameters.Material->IsUIMaterial())
 		{
 			return false;
 		}
 
-		if (Material->GetShadingModels().IsUnlit())
+		if (Parameters.Material->GetShadingModels().IsUnlit())
 		{
 			return false;
 		}
