Defining uninitialized barn door variables for ray tracing global illumination pass.

diff --git a/Engine/Source/Runtime/Renderer/Private/RayTracing/RayTracingGlobalIllumination.cpp b/Engine/Source/Runtime/Renderer/Private/RayTracing/RayTracingGlobalIllumination.cpp
--- a/Engine/Source/Runtime/Renderer/Private/RayTracing/RayTracingGlobalIllumination.cpp
+++ b/Engine/Source/Runtime/Renderer/Private/RayTracing/RayTracingGlobalIllumination.cpp
@@ -185,6 +185,8 @@ void SetupLightParameters(
 			LightParameters->Color[LightParameters->Count] = LightShaderParameters.Color;
 			LightParameters->Dimensions[LightParameters->Count] = FVector(2.0f * LightShaderParameters.SourceRadius, 2.0f * LightShaderParameters.SourceLength, 0.0f);
 			LightParameters->Attenuation[LightParameters->Count] = 1.0 / LightShaderParameters.InvRadius;
+			LightParameters->RectLightBarnCosAngle[LightParameters->Count] = LightShaderParameters.RectLightBarnCosAngle;
+			LightParameters->RectLightBarnLength[LightParameters->Count] = LightShaderParameters.RectLightBarnLength;
 			break;
 		}
 		case LightType_Point:
