New Screen Space Ambient Occlusion (SSAO) type added.
This is GTAO based on the activision paper.
[REDACTED_PROFILE_URL]

diff --git a/Engine/Shaders/Private/Common.ush b/Engine/Shaders/Private/Common.ush
--- a/Engine/Shaders/Private/Common.ush
+++ b/Engine/Shaders/Private/Common.ush
@@ -1212,6 +1212,17 @@ float2 ScreenPosToViewportUV(float2 ScreenPos)
 	return float2(0.5 + 0.5 * ScreenPos.x, 0.5 - 0.5 * ScreenPos.y);
 }
 
+// Maps standard viewport UV to an unprojected viewpos.
+// Viewpos can then be achieved via out.xy / out.z
+float3 ScreenToViewPos(float2 ViewportUV, float SceneDepth)
+{
+	float2 ProjViewPos;
+
+	ProjViewPos.x =  ViewportUV.x * View.ScreenToViewSpace.x  + View.ScreenToViewSpace.z;
+	ProjViewPos.y =  ViewportUV.y * View.ScreenToViewSpace.y  + View.ScreenToViewSpace.w;
+	return float3(ProjViewPos * SceneDepth, SceneDepth);
+}
+
 // ----------------------------
 
 /** 
diff --git a/Engine/Shaders/Private/PostProcessAmbientOcclusion.usf b/Engine/Shaders/Private/PostProcessAmbientOcclusion.usf
--- a/Engine/Shaders/Private/PostProcessAmbientOcclusion.usf
+++ b/Engine/Shaders/Private/PostProcessAmbientOcclusion.usf
@@ -19,6 +19,9 @@
 // 0:no / 1:yes
 // USE_UPSAMPLE
 
+// Testing the GTAO Implementation between a heavy weight reference and an in game one
+#define GTAO_REF 0
+
 // 0: classic with weighted sample, 1: don't normalize and adjust the formula to be simpler and faster - can look better and is cheaper (Alchemy like?)
 #define OPTIMIZATION_O1 1
 
@@ -879,7 +882,7 @@ void MainSSAOSmoothCS(
 	uint2 GroupThreadId : SV_GroupThreadID)
 {
 	BRANCH
-	if (any(DispatchThreadId >= SSAOSmoothParams.zw))
+	if (any(DispatchThreadId >= (uint2)SSAOSmoothParams.zw))
 	{
 		return;
 	}
@@ -897,3 +900,450 @@ void MainSSAOSmoothCS(
 	SSAOSmoothResult[DestPixelPos] = Result * 0.25;
 }
 #endif
+
+
+float4 GTAOParams[1];
+
+#define FALLOFF_START		100.0
+#define FALLOFF_START_SQ	(FALLOFF_START*FALLOFF_START)
+#define FALLOFF_END			150.0
+#define FALLOFF_END_SQ		(FALLOFF_END*FALLOFF_END)
+
+#if GTAO_REF
+#define NumPhi  32
+#define NumTaps  64
+#else
+#define NumPhi  1
+#define NumTaps  7
+#endif
+
+#define PI_HALF (PI*0.5)
+
+#if COMPUTE_SHADER
+RWTexture2D<float2> HorizonOutTexture;
+#endif
+
+// The following is used to help the compiler generate more optimal code.
+struct GTAO_COMMON
+{
+	float HZB_BiasX;
+	float HZB_BiasY;
+	float Falloff_Scale;
+	float Falloff_Bias;
+};
+
+
+void Fill_GTAO_Common( inout GTAO_COMMON GtaoCommon)
+{
+	GtaoCommon.HZB_BiasX = HZBRemapping.z;
+	GtaoCommon.HZB_BiasY = HZBRemapping.w;
+
+	GtaoCommon.Falloff_Scale = 1.0f / (FALLOFF_END_SQ - FALLOFF_START_SQ);
+	GtaoCommon.Falloff_Bias = -FALLOFF_START_SQ * GtaoCommon.Falloff_Scale; 
+}
+
+
+float Falloff(float distSq, GTAO_COMMON gtao_common)
+{
+	return 2.0 * saturate(distSq* gtao_common.Falloff_Scale + gtao_common.Falloff_Bias   );
+
+}
+
+float Falloff(float distSq)
+{
+	float Falloff_Scale = 1.0f / (FALLOFF_END_SQ - FALLOFF_START_SQ);
+	float Falloff_Bias = -FALLOFF_START_SQ * Falloff_Scale; 
+
+	return 2.0 * saturate(distSq* Falloff_Scale + Falloff_Bias   );
+}
+
+
+float3 GetViewSpacePosFromHZB(float2 UV, float MipLevel)
+{
+	float SceneDepth = GetDepthFromAOInput(UV);
+
+	float2 HZBUV;
+	HZBUV.x = UV.x * HZBRemapping.x + HZBRemapping.z;
+	HZBUV.y = UV.y * HZBRemapping.y + HZBRemapping.w;
+	SceneDepth = Texture2DSampleLevel( PostprocessInput1, PostprocessInput1Sampler, HZBUV, MipLevel).r;
+	SceneDepth = ConvertFromDeviceZ(SceneDepth);
+
+	return ScreenToViewPos(UV, SceneDepth);
+}
+
+float3 GetViewSpacePosFromHZB(float2 UV, float MipLevel, GTAO_COMMON gtao_common)
+{
+	float SceneDepth = GetDepthFromAOInput(UV);
+
+	float2 HZBUV;
+	HZBUV.x = UV.x * HZBRemapping.x + gtao_common.HZB_BiasX;
+	HZBUV.y = UV.y * HZBRemapping.y + gtao_common.HZB_BiasY;
+	SceneDepth = Texture2DSampleLevel( PostprocessInput1, PostprocessInput1Sampler, HZBUV,MipLevel).r;
+	SceneDepth = ConvertFromDeviceZ(SceneDepth);
+
+	return ScreenToViewPos(UV, SceneDepth);
+}
+
+
+float SearchForLargestAngleRef(float2 inUV, float2 deltaUV, float3 ViewSpacePosition, float3 ViewDir)
+{
+	float BiggestAngle = -1.0f;
+
+	deltaUV.y	  *=-1;
+	deltaUV	      *= 1;
+
+	for(int i=1; i<NumTaps; i++)
+	{
+		float fi = (float) i;
+		float2 UV		= inUV +  fi*deltaUV;
+
+		float SceneDepth	= GetDepthFromAOInput(UV);
+		float3 SampVS		= ScreenToViewPos(UV, SceneDepth);
+		float3 DiffV		= (SampVS - ViewSpacePosition) ;
+		float distSQ		= dot(DiffV, DiffV);
+
+		// At this point we could early out if DistSQ is greater than falloff
+		BRANCH
+		if(distSQ < FALLOFF_END_SQ)
+		{
+			float ooDist    = 1.0f / sqrt(distSQ);
+			float cosh		= ooDist * dot(DiffV,ViewDir );
+
+			float falloff = Falloff(distSQ);
+			BiggestAngle = max(BiggestAngle, cosh - falloff);
+		}
+	}
+
+	return BiggestAngle;
+}
+
+
+float SearchForLargestAngle(float2 inUV, float2 deltaUV, float3 ViewSpacePosition, float3 ViewDir)
+{
+#if GTAO_REF
+	return SearchForLargestAngleRef(inUV,  deltaUV,  ViewSpacePosition,  ViewDir);
+#endif
+
+	float BiggestAngle = -1.0f;
+
+	deltaUV.y	  *=-1;
+	deltaUV		  *= 2;
+
+	float2 UV	= inUV;
+
+	GTAO_COMMON GtaoCommon;
+	Fill_GTAO_Common(GtaoCommon);
+
+
+	for(int i=1; i<=NumTaps; i++)
+	{
+		float MipLevel = 0;
+
+		if(i==3)
+		{
+			deltaUV += deltaUV;
+			MipLevel = 1;
+		}
+		else if(i==5)
+		{
+			deltaUV += deltaUV;
+			MipLevel = 2;
+		}
+
+		UV				+= deltaUV;
+		
+		float3 SampVS   = GetViewSpacePosFromHZB(UV,MipLevel, GtaoCommon);
+		float3 DiffV	= (SampVS - ViewSpacePosition) ;
+		float distSQ	= dot(DiffV, DiffV);
+
+		// At this point we could early out if DistSQ is greater than falloff
+		// Profiling on Xbox with the branch on and off. Without the ALU usage is better but the VGPR count is much higher.
+		BRANCH
+		if(distSQ < FALLOFF_END_SQ)
+		{
+			float ooDist    = rsqrt(distSQ);
+			float cosh		= ooDist * dot(DiffV,ViewDir );
+			float falloff = Falloff(distSQ, GtaoCommon);
+			BiggestAngle = max(BiggestAngle, cosh - falloff);
+		}
+
+	}
+
+	return BiggestAngle;
+}
+
+
+// Given a screen space Axis this will search for the min and max angles
+float2 SearchAxisForAngles(float2 UV, float3 Dir, float3 ViewDir, float3 ViewSpacePos, float Scale)
+{
+	// TODO - Fix aspect ratio
+	float RandScale = 1.0f;//GTAOParams[0].y;
+	float2 DeltaUV	= (RandScale * Scale) * float2(View.BufferSizeAndInvSize.z * Dir.x, View.BufferSizeAndInvSize.w * Dir.y);
+
+	float Ang1 = SearchForLargestAngle( UV,   DeltaUV, ViewSpacePos, ViewDir);
+	float Ang2 = SearchForLargestAngle( UV,  -DeltaUV, ViewSpacePos, ViewDir);
+
+	return float2(Ang1, Ang2);
+}
+
+float ComputeInnerIntegral(float2 Angles, float3 ScreenDir, float3 ViewDir, float3 ViewSpaceNormal, float SceneDepth)
+{
+	// Given the angles found in the search plane we need to project the View Space GBuffer Normal onto the plane defined by the search axis and the View Direction and perform the inner integrate
+	float3 PlaneNormal		= normalize(cross(ScreenDir,ViewDir));
+	float3 Perp				= cross(ViewDir, PlaneNormal);
+	float3 ProjNormal		= ViewSpaceNormal - PlaneNormal * dot(ViewSpaceNormal, PlaneNormal);
+
+	float ProjNormalMagSq	= dot(ProjNormal, ProjNormal);
+	float RecipMag			= rsqrt(ProjNormalMagSq);
+	float ProjNormalMag		= 1.0f / max(RecipMag, 1e-7);
+	float CosAng			= dot(ProjNormal, Perp) * RecipMag;	
+	float Gamma				= acosFast(CosAng) - PI_HALF;				
+	float CosGamma			= dot(ProjNormal, ViewDir) * RecipMag;
+	float SinGamma  		= CosAng * -2.0f;					
+
+	// clamp to normal hemisphere
+	Angles.x = Gamma + max(-Angles.x - Gamma, -PI_HALF);
+	Angles.y = Gamma + min( Angles.y - Gamma,  PI_HALF);
+
+	float AO = (ProjNormalMag * 0.25 * ( (Angles.x * SinGamma + CosGamma - cos((2.0 * Angles.x) - Gamma)) +
+				  			             (Angles.y * SinGamma + CosGamma - cos((2.0 * Angles.y) - Gamma)) ));
+
+	// Fade out based on user defined distance
+	float Mul = ScreenSpaceAOParams[4].x;
+	float Add = ScreenSpaceAOParams[4].y;
+	AO = lerp(AO, 1, saturate(SceneDepth * Mul + Add));
+
+	return AO;
+}
+
+
+
+/*
+*
+* HORIZON SEARCH AND INNER INTEGRATE
+*
+*/
+
+void HorizonSearchAndInnerIntegratePSandCS(in float2 UV, out float4 OutColor)
+{
+	OutColor = 0;
+
+	float SceneDepth	 = GetDepthFromAOInput(UV);		
+	if(SceneDepth > ScreenSpaceAOParams[4].w)	
+	{		
+		OutColor = 1;
+		return;
+	}
+	
+	float2 ViewportUVToRandomUV = ScreenSpaceAOParams[1].xy;
+	const float2 TemporalOffset = ScreenSpaceAOParams[3].xy;
+	float2 RandomVec = (Texture2DSample(RandomNormalTexture, RandomNormalTextureSampler, TemporalOffset + UV * ViewportUVToRandomUV).rg * 2 - 1) ;
+
+	float3 ViewSpacePos  = ScreenToViewPos(UV,SceneDepth);
+
+	// Get Normal and transform into viewspace
+	float3 WorldNormal		= GetGBufferData(UV, false).WorldNormal;								// TAP
+	float3 ViewSpaceNormal	= normalize(mul(WorldNormal, (float3x3)View.TranslatedWorldToView));
+	float3 ViewDir			= -normalize(ViewSpacePos.xyz);
+
+	float Sum=0.0;
+
+	// Get radius based on sampling 
+	// Get radius based on distance - TODO - Need to use FOV and aspect ratio
+	float Scale = 1000.0 / ViewSpacePos.z;
+
+	Scale = max(Scale, 2.0f);
+
+	// TODO - Opt - Determine how many taps we'll step before we hit the attenuation radius.
+
+#if GTAO_REF
+	for(int i=0; i<NumPhi; i++)
+	{
+		float phi			= ((float)i) *  (PI / NumPhi);
+		float3 ScreenDir	= float3(cos(phi), sin(phi), 0.0);
+#else
+		float3 ScreenDir	= float3(RandomVec.x, RandomVec.y, 0.0);
+#endif
+
+		float2 Angles = SearchAxisForAngles( UV,  ScreenDir,  ViewDir,  ViewSpacePos, Scale);
+
+		Angles.x = acosFast(Angles.x);
+		Angles.y = acosFast(Angles.y);
+
+		Sum += ComputeInnerIntegral(Angles, ScreenDir, ViewDir, ViewSpaceNormal, SceneDepth);
+#if GTAO_REF
+	}
+#endif
+
+	OutColor.r = Sum   / (float)NumPhi;
+
+	return;
+}
+
+
+void HorizonSearchAndInnerIntegratePS(in noperspective float4 UVAndScreenPos : TEXCOORD0, out float4 OutColor : SV_Target0)
+{
+	HorizonSearchAndInnerIntegratePSandCS(UVAndScreenPos.xy, OutColor);	
+}
+
+
+#if COMPUTE_SHADER
+[numthreads(THREADGROUP_SIZEX, THREADGROUP_SIZEY, 1)]
+void HorizonSearchAndInnerIntegrateCS(
+	uint2 GroupId : SV_GroupID,
+	uint2 DispatchThreadId : SV_DispatchThreadID,
+	uint2 GroupThreadId : SV_GroupThreadID) 
+{
+	float4 OutColor		= 0;
+	int2   PixelPos		= DispatchThreadId + ScreenSpaceAOParams[5].zw; 
+	float2 PixelCenter	= (float2)PixelPos + float2(0.5, 0.5);
+	float2 BufferUV		= SvPositionToBufferUV(float4(PixelCenter, 0, 0));
+
+	HorizonSearchAndInnerIntegratePSandCS(BufferUV , OutColor);
+
+	OutTexture[PixelPos] = OutColor.r;
+}
+#endif
+
+
+
+
+/*
+*
+*
+*	JUST INNER INTEGRATE
+*
+*
+*/
+
+Texture2D		HorizonsTexture;
+SamplerState	HorizonsTextureSampler;
+
+void GTAOInnerIntegratePSandCS(in float2 UV, out float4 OutColor)
+{
+	// Read the angles buffer
+	OutColor = 0;
+
+	float SceneDepth	 = GetDepthFromAOInput(UV);		
+	if(SceneDepth > ScreenSpaceAOParams[4].w)	
+	{		
+		OutColor = 1;
+		return;
+	}
+
+	float2 Angles			= Texture2DSample(HorizonsTexture, HorizonsTextureSampler, UV).xy;	// Angles computed from previous pass
+	Angles				    = Angles *2-1;
+
+	// Get Angle
+	float2 ViewportUVToRandomUV = ScreenSpaceAOParams[1].xy;
+	const float2 TemporalOffset = ScreenSpaceAOParams[3].xy;
+	float2 RandomVec			= (Texture2DSample(RandomNormalTexture, RandomNormalTextureSampler, TemporalOffset + UV * ViewportUVToRandomUV).rg * 2 - 1) ;
+	float3 ScreenDir			= float3(RandomVec.x, RandomVec.y, 0.0);
+
+	// ViewspacePos and Normal
+	float3 ViewSpacePos		= ScreenToViewPos(UV, SceneDepth);
+	float3 WorldNormal		= GetGBufferData(UV, false).WorldNormal;
+	float3 ViewSpaceNormal	= normalize(mul(WorldNormal, (float3x3)View.TranslatedWorldToView));
+	float3 ViewDir			= -normalize(ViewSpacePos.xyz);	// TODO - This is a function of UV only.
+
+	Angles.x = acosFast(Angles.x);
+	Angles.y = acosFast(Angles.y);
+
+	float AO = ComputeInnerIntegral(Angles, ScreenDir, ViewDir, ViewSpaceNormal, SceneDepth);
+	OutColor.rgb = AO ;
+
+	return;
+}
+
+
+
+void GTAOInnerIntegratePS(in noperspective float4 UVAndScreenPos : TEXCOORD0, out float4 OutColor : SV_Target0)
+{
+	GTAOInnerIntegratePSandCS(UVAndScreenPos.xy, OutColor);	
+}
+
+
+#if COMPUTE_SHADER
+[numthreads(THREADGROUP_SIZEX, THREADGROUP_SIZEY, 1)]
+void GTAOInnerIntegrateCS(
+	uint2 GroupId : SV_GroupID,
+	uint2 DispatchThreadId : SV_DispatchThreadID,
+	uint2 GroupThreadId : SV_GroupThreadID) 
+{
+	float4 OutColor		= 0;
+	int2   PixelPos		= DispatchThreadId + ScreenSpaceAOParams[5].zw; 
+	float2 PixelCenter	= (float2)PixelPos + float2(0.5, 0.5);
+	float2 BufferUV		= SvPositionToBufferUV(float4(PixelCenter, 0, 0));
+
+	GTAOInnerIntegratePSandCS(BufferUV , OutColor);
+	
+	OutTexture[PixelPos] = OutColor.r;
+}
+#endif
+
+
+
+
+
+
+/*
+*
+* HORIZON SEARCH ONLY
+*
+*/
+
+float2 HorizonSearchPSandCS(in float2 UV)
+{
+	float2 OutHorizons = 0;
+
+	float SceneDepth	 = GetDepthFromAOInput(UV);
+	if(SceneDepth > ScreenSpaceAOParams[4].w)	
+	{		
+		OutHorizons = 0;
+		return OutHorizons;
+	}
+
+	float2 ViewportUVToRandomUV = ScreenSpaceAOParams[1].xy;
+	const float2 TemporalOffset = ScreenSpaceAOParams[3].xy;
+	float2 RandomVec			= (Texture2DSample(RandomNormalTexture, RandomNormalTextureSampler, TemporalOffset + UV * ViewportUVToRandomUV).rg * 2 - 1) ;
+	float3 ScreenDir			= float3(RandomVec.x, RandomVec.y, 0.0);
+	
+	float3 ViewSpacePos			= ScreenToViewPos(UV, SceneDepth);
+	float3 ViewDir				= -normalize(ViewSpacePos.xyz);
+
+	// Get radius based on distance - TODO - Need to use FOV and aspect ratio
+	float Scale = 1000.0 / ViewSpacePos.z;
+	Scale = max(Scale, 1.0f);
+
+	float2 Angles = SearchAxisForAngles( UV,  ScreenDir,  ViewDir,  ViewSpacePos, Scale);
+	OutHorizons.rg = Angles*0.5+0.5;
+
+	return OutHorizons;
+}
+
+
+void HorizonSearchPS(in noperspective float4 UVAndScreenPos : TEXCOORD0, float4 SvPosition : SV_POSITION, out float4 OutColor : SV_Target0)
+{
+	float2 Horizons = HorizonSearchPSandCS(UVAndScreenPos.xy);	
+	OutColor = float4(Horizons.xy,0,0);;
+
+}
+
+
+#if COMPUTE_SHADER
+[numthreads(THREADGROUP_SIZEX, THREADGROUP_SIZEY, 1)]
+void HorizonSearchCS(
+	uint2 GroupId : SV_GroupID,
+	uint2 DispatchThreadId : SV_DispatchThreadID,
+	uint2 GroupThreadId : SV_GroupThreadID) 
+{
+	float2 OutColor		= 0;
+	int2   PixelPos		= DispatchThreadId + ScreenSpaceAOParams[5].zw; 
+	float2 PixelCenter	= (float2)PixelPos + float2(0.5, 0.5);
+	float2 BufferUV		= SvPositionToBufferUV(float4(PixelCenter, 0, 0));
+
+	float2 Horizons = HorizonSearchPSandCS(BufferUV);
+	
+	HorizonOutTexture[PixelPos] = Horizons;
+}
+#endif
diff --git a/Engine/Source/Runtime/Engine/Private/SceneView.cpp b/Engine/Source/Runtime/Engine/Private/SceneView.cpp
--- a/Engine/Source/Runtime/Engine/Private/SceneView.cpp
+++ b/Engine/Source/Runtime/Engine/Private/SceneView.cpp
@@ -2211,6 +2211,18 @@ void FSceneView::SetupViewRectUniformBufferParameters(FViewUniformShaderParamete
 		ViewUniformShaderParameters.AdaptiveTessellationFactor = 0.5f * InViewMatrices.GetProjectionMatrix().M[1][1] * float(EffectiveViewRect.Height()) / TessellationAdaptivePixelsPerEdge;
 	}
 
+	// Compute coefficients which takes a screen UV and converts to Viewspace.xy / ViewZ
+	float InvTanHalfFov = InViewMatrices.GetProjectionMatrix().M[0][0];
+	float Ratio = UnscaledViewRect.Width() / (float)UnscaledViewRect.Height();
+
+	float InvFovFixX = 1.0f / (InvTanHalfFov);
+	float InvFovFixY = 1.0f / (Ratio * InvTanHalfFov);
+
+	ViewUniformShaderParameters.ScreenToViewSpace.X = BufferSize.X * ViewUniformShaderParameters.ViewSizeAndInvSize.Z * 2 * InvFovFixX;
+	ViewUniformShaderParameters.ScreenToViewSpace.Y = BufferSize.Y * ViewUniformShaderParameters.ViewSizeAndInvSize.W  * -2 * InvFovFixY;
+
+	ViewUniformShaderParameters.ScreenToViewSpace.Z = -((ViewUniformShaderParameters.ViewRectMin.X * ViewUniformShaderParameters.ViewSizeAndInvSize.Z * 2 * InvFovFixX) + InvFovFixX);
+	ViewUniformShaderParameters.ScreenToViewSpace.W = (ViewUniformShaderParameters.ViewRectMin.Y * ViewUniformShaderParameters.ViewSizeAndInvSize.W * 2 * InvFovFixY) + InvFovFixY;
 }
 
 void FSceneView::SetupCommonViewUniformBufferParameters(
diff --git a/Engine/Source/Runtime/Engine/Public/SceneView.h b/Engine/Source/Runtime/Engine/Public/SceneView.h
--- a/Engine/Source/Runtime/Engine/Public/SceneView.h
+++ b/Engine/Source/Runtime/Engine/Public/SceneView.h
@@ -610,6 +610,7 @@ enum ETranslucencyVolumeCascade
 	VIEW_UNIFORM_BUFFER_MEMBER(FVector4, ViewSizeAndInvSize) \
 	VIEW_UNIFORM_BUFFER_MEMBER(FVector4, BufferSizeAndInvSize) \
 	VIEW_UNIFORM_BUFFER_MEMBER(FVector4, BufferBilinearUVMinMax) \
+	VIEW_UNIFORM_BUFFER_MEMBER(FVector4, ScreenToViewSpace) \
 	VIEW_UNIFORM_BUFFER_MEMBER(int32, NumSceneColorMSAASamples) \
 	VIEW_UNIFORM_BUFFER_MEMBER_EX(float, PreExposure, EShaderPrecisionModifier::Half) \
 	VIEW_UNIFORM_BUFFER_MEMBER_EX(float, OneOverPreExposure, EShaderPrecisionModifier::Half) \
diff --git a/Engine/Source/Runtime/Renderer/Private/CompositionLighting/CompositionLighting.cpp b/Engine/Source/Runtime/Renderer/Private/CompositionLighting/CompositionLighting.cpp
--- a/Engine/Source/Runtime/Renderer/Private/CompositionLighting/CompositionLighting.cpp
+++ b/Engine/Source/Runtime/Renderer/Private/CompositionLighting/CompositionLighting.cpp
@@ -106,6 +106,80 @@ bool ShouldRenderScreenSpaceAmbientOcclusion(const FViewInfo& View)
 	return bEnabled;
 }
 
+static FRenderingCompositeOutputRef AddPostProcessingGTAOAsyncHorizonSearch(FRHICommandListImmediate& RHICmdList, FPostprocessContext& Context)
+{
+	FRenderingCompositePass* FinalOutputPass;
+
+	FRenderingCompositePass* HZBInput = Context.Graph.RegisterPass(new FRCPassPostProcessInput(const_cast<FViewInfo&>(Context.View).HZB));
+		
+	FRenderingCompositePass* AmbientOcclusionHorizonSearch = Context.Graph.RegisterPass(new(FMemStack::Get()) FRCPassPostProcessAmbientOcclusion_HorizonSearch(Context.View, ESSAOType::EAsyncCS));
+
+	AmbientOcclusionHorizonSearch->SetInput(ePId_Input0, Context.SceneDepth);
+	AmbientOcclusionHorizonSearch->SetInput(ePId_Input1, HZBInput);
+
+	FinalOutputPass = AmbientOcclusionHorizonSearch;
+
+	Context.FinalOutput = FRenderingCompositeOutputRef(FinalOutputPass);
+	return FRenderingCompositeOutputRef(FinalOutputPass);
+}
+
+
+static FRenderingCompositeOutputRef AddPostProcessingGTAOCombined(FRHICommandListImmediate& RHICmdList, FPostprocessContext& Context)
+{
+	FRenderingCompositePass* FinalOutputPass;
+
+	const bool bNeedSmoothingPass = CVarSSAOSmoothPass.GetValueOnRenderThread();
+
+	FRenderingCompositePass* HZBInput = Context.Graph.RegisterPass(new FRCPassPostProcessInput(const_cast<FViewInfo&>(Context.View).HZB));
+
+	FRenderingCompositePass* AmbientOcclusionGTAO = Context.Graph.RegisterPass(new(FMemStack::Get()) FRCPassPostProcessAmbientOcclusion_HorizonSearchInnerIntegrate(Context.View, !bNeedSmoothingPass));
+	AmbientOcclusionGTAO->SetInput(ePId_Input0, Context.SceneDepth);
+	AmbientOcclusionGTAO->SetInput(ePId_Input1, HZBInput);
+	FinalOutputPass = AmbientOcclusionGTAO;
+
+	FSceneRenderTargets& SceneContext = FSceneRenderTargets::Get(RHICmdList);
+	SceneContext.bScreenSpaceAOIsValid = true;
+
+	if (bNeedSmoothingPass)
+	{
+		FRenderingCompositePass* SSAOSmoothPass;
+		SSAOSmoothPass = Context.Graph.RegisterPass(new (FMemStack::Get()) FRCPassPostProcessAmbientOcclusionSmooth(ESSAOType::ECS, true));
+		SSAOSmoothPass->SetInput(ePId_Input0, FinalOutputPass);
+		FinalOutputPass = SSAOSmoothPass;
+	}
+
+	Context.FinalOutput = FRenderingCompositeOutputRef(FinalOutputPass);
+	return FRenderingCompositeOutputRef(FinalOutputPass);
+}
+
+
+static FRenderingCompositeOutputRef AddPostProcessingGTAOIntegration(FRHICommandListImmediate& RHICmdList, FPostprocessContext& Context)
+{
+	FRenderingCompositePass* FinalOutputPass;
+
+	const bool bNeedSmoothingPass = CVarSSAOSmoothPass.GetValueOnRenderThread();
+
+	FRenderingCompositePass* AmbientOcclusionInnerIntegrate = Context.Graph.RegisterPass(new(FMemStack::Get()) FRCPassPostProcessAmbientOcclusion_InnerIntegrate(Context.View, !bNeedSmoothingPass));
+	AmbientOcclusionInnerIntegrate->SetInput(ePId_Input0, Context.SceneDepth);
+	FinalOutputPass = AmbientOcclusionInnerIntegrate;
+
+	FSceneRenderTargets& SceneContext = FSceneRenderTargets::Get(RHICmdList);
+	SceneContext.bScreenSpaceAOIsValid = true;
+
+	if (bNeedSmoothingPass)
+	{
+		FRenderingCompositePass* SSAOSmoothPass;
+		SSAOSmoothPass = Context.Graph.RegisterPass(new (FMemStack::Get()) FRCPassPostProcessAmbientOcclusionSmooth(ESSAOType::ECS, true));
+		SSAOSmoothPass->SetInput(ePId_Input0, FinalOutputPass);
+		FinalOutputPass = SSAOSmoothPass;
+	}
+
+	Context.FinalOutput = FRenderingCompositeOutputRef(FinalOutputPass);
+	return FRenderingCompositeOutputRef(FinalOutputPass);
+
+}
+
+
 // @param Levels 0..3, how many different resolution levels we want to render
 static FRenderingCompositeOutputRef AddPostProcessingAmbientOcclusion(FRHICommandListImmediate& RHICmdList, FPostprocessContext& Context, uint32 Levels)
 {
@@ -239,7 +313,11 @@ void FCompositionLighting::ProcessBeforeBasePass(FRHICommandListImmediate& RHICm
 
 		if (SSAOLevels)
 		{
-			AddPostProcessingAmbientOcclusion(RHICmdList, Context, SSAOLevels);
+
+			if (FSSAOHelper::GetGTAOPassType(View) != EGTAOType::ECombinedNonAsync)
+			{
+				AddPostProcessingAmbientOcclusion(RHICmdList, Context, SSAOLevels);
+			}
 		}
 
 		// The graph setup should be finished before this line ----------------------------------------
@@ -323,7 +401,14 @@ void FCompositionLighting::ProcessAfterBasePass(FRHICommandListImmediate& RHICmd
 			{
 				if(!FSSAOHelper::IsAmbientOcclusionAsyncCompute(Context.View, SSAOLevels))
 				{
-					AmbientOcclusion = AddPostProcessingAmbientOcclusion(RHICmdList, Context, SSAOLevels);
+					if (FSSAOHelper::GetGTAOPassType(View) == EGTAOType::ECombinedNonAsync)
+					{
+						AmbientOcclusion = AddPostProcessingGTAOCombined(RHICmdList, Context);
+					}
+					else
+					{
+						AmbientOcclusion = AddPostProcessingAmbientOcclusion(RHICmdList, Context, SSAOLevels);
+					}
 
 					if (bDoDecal)
 					{
@@ -335,6 +420,12 @@ void FCompositionLighting::ProcessAfterBasePass(FRHICommandListImmediate& RHICmd
 				}
 				else
 				{
+					// If doing the Split GTAO method then we need to do the second part here.
+					if (FSSAOHelper::GetGTAOPassType(View) == EGTAOType::ESplitAsync)
+					{
+						AmbientOcclusion = AddPostProcessingGTAOIntegration(RHICmdList, Context);
+					}
+
 					ensureMsgf(
 						FDecalRendering::BuildVisibleDecalList(*(FScene*)Context.View.Family->Scene, Context.View, DRS_AmbientOcclusion, nullptr) == false,
 						TEXT("Ambient occlusion decals are not supported with Async compute SSAO."));
@@ -436,9 +527,17 @@ void FCompositionLighting::ProcessAsyncSSAO(FRHICommandListImmediate& RHICmdList
 			{
 				FPostprocessContext Context(RHICmdList, CompositeContext.Graph, View);
 
-				FRenderingCompositeOutputRef AmbientOcclusion = AddPostProcessingAmbientOcclusion(RHICmdList, Context, Levels);
-				Context.FinalOutput = FRenderingCompositeOutputRef(AmbientOcclusion);			
-
+				if (FSSAOHelper::GetGTAOPassType(View) == EGTAOType::ESplitAsync)
+				{
+					FRenderingCompositeOutputRef AmbientOcclusion = AddPostProcessingGTAOAsyncHorizonSearch(RHICmdList, Context);
+					Context.FinalOutput = FRenderingCompositeOutputRef(AmbientOcclusion);
+				}
+				else
+				{
+					FRenderingCompositeOutputRef AmbientOcclusion = AddPostProcessingAmbientOcclusion(RHICmdList, Context, Levels);
+					Context.FinalOutput = FRenderingCompositeOutputRef(AmbientOcclusion);
+				}
+		
 				// The graph setup should be finished before this line ----------------------------------------
 				CompositeContext.Process(Context.FinalOutput.GetPass(), TEXT("Composition_ProcessAsyncSSAO"));
 			}		
diff --git a/Engine/Source/Runtime/Renderer/Private/CompositionLighting/PostProcessAmbientOcclusion.cpp b/Engine/Source/Runtime/Renderer/Private/CompositionLighting/PostProcessAmbientOcclusion.cpp
--- a/Engine/Source/Runtime/Renderer/Private/CompositionLighting/PostProcessAmbientOcclusion.cpp
+++ b/Engine/Source/Runtime/Renderer/Private/CompositionLighting/PostProcessAmbientOcclusion.cpp
@@ -20,6 +20,9 @@ DECLARE_GPU_STAT_NAMED(SSAOSetup, TEXT("ScreenSpace AO Setup") );
 DECLARE_GPU_STAT_NAMED(SSAO, TEXT("ScreenSpace AO") );
 DECLARE_GPU_STAT_NAMED(BasePassAO, TEXT("BasePass AO") );
 DECLARE_GPU_STAT_NAMED(SSAOSmooth, TEXT("SSAO smooth"));
+DECLARE_GPU_STAT_NAMED(GTAO_HorizonSearch, TEXT("GTAO HorizonSearch"));
+DECLARE_GPU_STAT_NAMED(GTAO_HorizonSearchInnerIntegrate, TEXT("GTAO HorizonSetupAndInnerIntegrate"));
+DECLARE_GPU_STAT_NAMED(GTAO_InnerIntegrate, TEXT("GTAO InnerIntegrate"));
 
 // Tile size for the AmbientOcclusion compute shader, tweaked for 680 GTX. */
 // see GCN Performance Tip 21 http://developer.amd.com/wordpress/media/2013/05/GCNPerformanceTweets.pdf 
@@ -87,6 +90,15 @@ static TAutoConsoleVariable<int32> CVarAmbientOcclusionDepthBoundsTest(
 	TEXT("Whether to use depth bounds test to cull distant pixels during AO pass. This option is only valid when pixel shader path is used (r.AmbientOcclusion.Compute=0), without upsampling."),
 	ECVF_RenderThreadSafe);
 
+static TAutoConsoleVariable<int32> CVarAmbientOcclusionMethod(
+	TEXT("r.AmbientOcclusion.Method"),
+	0,
+	TEXT("Select between SSAO methods \n ")
+	TEXT("0: SSAO (default)\n ")
+	TEXT("1: GTAO\n "),
+	ECVF_RenderThreadSafe | ECVF_Scalability);
+
+
 float FSSAOHelper::GetAmbientOcclusionQualityRT(const FSceneView& View)
 {
 	float CVarValue = CVarAmbientOcclusionMaxQuality.GetValueOnRenderThread();
@@ -208,11 +220,39 @@ uint32 FSSAOHelper::ComputeAmbientOcclusionPassCount(const FViewInfo& View)
 	return Ret;
 }
 
+EGTAOType FSSAOHelper::GetGTAOPassType(const FViewInfo& View)
+{
+	int32 Method = CVarAmbientOcclusionMethod.GetValueOnRenderThread();
+
+	if (Method == 1)
+	{
+		if (IsAmbientOcclusionAsyncCompute(View, 1))
+		{
+			return EGTAOType::ESplitAsync;
+		}
+		else
+		{
+			return EGTAOType::ECombinedNonAsync;
+		}
+	}
+	return EGTAOType::EOff;
+
+}
+
+
+
 /** Shader parameters needed for screen space AmbientOcclusion passes. */
 class FScreenSpaceAOParameters
 {
 public:
 
+	enum class ERandTexType
+	{
+		SSAO,
+		GTAO,
+	};
+
+
 	void Bind(const FShaderParameterMap& ParameterMap)
 	{
 		ScreenSpaceAOParams.Bind(ParameterMap, TEXT("ScreenSpaceAOParams"));
@@ -220,12 +260,20 @@ public:
 
 	//@param TRHICmdList could be async compute or compute dispatch, so template on commandlist type.
 	template<typename ShaderRHIParamRef, typename TRHICmdList>
-	void Set(TRHICmdList& RHICmdList, const FViewInfo& View, const ShaderRHIParamRef ShaderRHI, FIntPoint InputTextureSize) const
+	void Set(TRHICmdList& RHICmdList, const FViewInfo& View, const ShaderRHIParamRef ShaderRHI, FIntPoint InputTextureSize, ERandTexType RandTexType = ERandTexType::SSAO) const
 	{
 		const FFinalPostProcessSettings& Settings = View.FinalPostProcessSettings;
 
-		FIntPoint SSAORandomizationSize = GSystemTextures.SSAORandomization->GetDesc().Extent;
-		FVector2D ViewportUVToRandomUV(InputTextureSize.X / (float)SSAORandomizationSize.X, InputTextureSize.Y / (float)SSAORandomizationSize.Y);
+		FIntPoint RandomizationSize;
+		if (RandTexType == ERandTexType::GTAO)
+		{
+			RandomizationSize = GSystemTextures.GTAORandomization->GetDesc().Extent;
+		}
+		else
+		{
+			RandomizationSize = GSystemTextures.SSAORandomization->GetDesc().Extent;
+		}
+		FVector2D ViewportUVToRandomUV(InputTextureSize.X / (float)RandomizationSize.X, InputTextureSize.Y / (float)RandomizationSize.Y);
 
 		// e.g. 4 means the input texture is 4x smaller than the buffer size
 		uint32 ScaleToFullRes = FSceneRenderTargets::Get(RHICmdList).GetBufferSizeXY().X / InputTextureSize.X;
@@ -264,7 +312,7 @@ public:
 		
 		if(View.State)
 		{
-			TemporalOffset = (View.State->GetCurrentTemporalAASampleIndex() % 8) * FVector2D(2.48f, 7.52f) / 64.0f;
+			TemporalOffset = (View.State->GetCurrentTemporalAASampleIndex() % 8) * FVector2D(2.48f, 7.52f) / (float)RandomizationSize.X;
 		}
 		const float HzbStepMipLevelFactorValue = FMath::Clamp(FSSAOHelper::GetAmbientOcclusionStepMipLevelFactor(), 0.0f, 100.0f);
 		const float InvAmbientOcclusionDistance = 1.0f / FMath::Max(Settings.AmbientOcclusionDistance_DEPRECATED, KINDA_SMALL_NUMBER);
@@ -274,7 +322,7 @@ public:
 		Value[1] = FVector4(ViewportUVToRandomUV.X, ViewportUVToRandomUV.Y, AORadiusInShader, Ratio);
 		Value[2] = FVector4(ScaleToFullRes, Settings.AmbientOcclusionMipThreshold / ScaleToFullRes, ScaleRadiusInWorldSpace, Settings.AmbientOcclusionMipBlend);
 		Value[3] = FVector4(TemporalOffset.X, TemporalOffset.Y, StaticFraction, InvTanHalfFov);
-		Value[4] = FVector4(InvFadeRadius, -(Settings.AmbientOcclusionFadeDistance - FadeRadius) * InvFadeRadius, HzbStepMipLevelFactorValue, 0);
+		Value[4] = FVector4(InvFadeRadius, -(Settings.AmbientOcclusionFadeDistance - FadeRadius) * InvFadeRadius, HzbStepMipLevelFactorValue, Settings.AmbientOcclusionFadeDistance);
 		Value[5] = FVector4(View.ViewRect.Width(), View.ViewRect.Height(), ViewRect.Min.X, ViewRect.Min.Y);
 
 		SetShaderValueArray(RHICmdList, ShaderRHI, ScreenSpaceAOParams, Value, 6);
@@ -1153,3 +1201,780 @@ FPooledRenderTargetDesc FRCPassPostProcessAmbientOcclusion::ComputeOutputDesc(EP
 
 	return Ret;
 }
+
+
+
+
+/** Shader parameters needed for screen space AmbientOcclusion passes. */
+class FGTAOParameters
+{
+public:
+
+	void Bind(const FShaderParameterMap& ParameterMap)
+	{
+		GTAOParams.Bind(ParameterMap, TEXT("GTAOParams"));
+	}
+
+	//@param TRHICmdList could be async compute or compute dispatch, so template on commandlist type.
+	template<typename ShaderRHIParamRef, typename TRHICmdList>
+	void Set(TRHICmdList& RHICmdList, const FViewInfo& View, const ShaderRHIParamRef ShaderRHI) const
+	{
+		const FFinalPostProcessSettings& Settings = View.FinalPostProcessSettings;
+
+		FVector4 Value[1];
+		static float CurrRand = 0.0f;
+		float NewRand = FMath::FRand();
+		if (fabs(NewRand - CurrRand) < 0.2f)
+		{
+			NewRand = CurrRand < 0.5f ? NewRand + 0.2f : NewRand - 0.2f;
+		}
+		CurrRand = NewRand;
+		Value[0] = FVector4(NewRand, NewRand*0.5f + 1.0f, 0.0f, 0.0f);
+
+		SetShaderValueArray(RHICmdList, ShaderRHI, GTAOParams, Value, 1);
+	}
+
+	friend FArchive& operator<<(FArchive& Ar, FGTAOParameters& This);
+
+private:
+	FShaderParameter GTAOParams;
+};
+
+FArchive& operator<<(FArchive& Ar, FGTAOParameters& This)
+{
+	Ar << This.GTAOParams;
+	return Ar;
+}
+
+template<uint32 bComputeShader>
+class FPostProcessGTAOHorizonSearchPSandCS : public FGlobalShader
+{
+	DECLARE_SHADER_TYPE(FPostProcessGTAOHorizonSearchPSandCS, Global);
+
+	static bool ShouldCompilePermutation(const FGlobalShaderPermutationParameters& Parameters)
+	{
+		return IsFeatureLevelSupported(Parameters.Platform, ERHIFeatureLevel::SM5);
+	}
+
+	static void ModifyCompilationEnvironment(const FGlobalShaderPermutationParameters& Parameters, FShaderCompilerEnvironment& OutEnvironment)
+	{
+		FGlobalShader::ModifyCompilationEnvironment(Parameters, OutEnvironment);
+		OutEnvironment.SetDefine(TEXT("COMPUTE_SHADER"), bComputeShader);
+		if (bComputeShader)
+		{
+			OutEnvironment.SetDefine(TEXT("THREADGROUP_SIZEX"), 8);
+			OutEnvironment.SetDefine(TEXT("THREADGROUP_SIZEY"), 8);
+		}
+	}
+
+	/** Default constructor. */
+	FPostProcessGTAOHorizonSearchPSandCS() {}
+
+public:
+	FSceneTextureShaderParameters SceneTextureParameters;
+	FPostProcessPassParameters	  PostprocessParameter;
+	FScreenSpaceAOParameters	  ScreenSpaceAOParams;
+	FGTAOParameters				  GTAOParams;
+	FShaderResourceParameter 	  HorizonOutTexture;
+	FShaderResourceParameter      RandomNormalTexture;
+	FShaderResourceParameter      RandomNormalTextureSampler;
+	FShaderParameter			  HZBRemapping;
+
+	/** Initialization constructor. */
+	FPostProcessGTAOHorizonSearchPSandCS(const ShaderMetaType::CompiledShaderInitializerType& Initializer)
+		: FGlobalShader(Initializer)
+	{
+		PostprocessParameter.Bind(Initializer.ParameterMap);
+		SceneTextureParameters.Bind(Initializer);
+		ScreenSpaceAOParams.Bind(Initializer.ParameterMap);
+		GTAOParams.Bind(Initializer.ParameterMap);
+		HZBRemapping.Bind(Initializer.ParameterMap, TEXT("HZBRemapping"));
+		RandomNormalTexture.Bind(Initializer.ParameterMap, TEXT("RandomNormalTexture"));
+		RandomNormalTextureSampler.Bind(Initializer.ParameterMap, TEXT("RandomNormalTextureSampler"));
+
+		if (bComputeShader)
+		{
+			HorizonOutTexture.Bind(Initializer.ParameterMap, TEXT("HorizonOutTexture"));
+		}
+	}
+
+	FVector4 GetHZBRemapVal(const FRenderingCompositePassContext& Context, FIntPoint DestSize, FIntPoint InputTextureSize)
+	{
+		const FVector2D HZBScaleFactor(float(Context.View.ViewRect.Width()) / float(2 * Context.View.HZBMipmap0Size.X),
+			float(Context.View.ViewRect.Height()) / float(2 * Context.View.HZBMipmap0Size.Y));
+
+		const FVector2D UVScaleFactor(float(DestSize.X) / float(InputTextureSize.X),
+			float(DestSize.Y) / float(InputTextureSize.Y));
+
+		return FVector4(HZBScaleFactor.X / UVScaleFactor.X, HZBScaleFactor.Y / UVScaleFactor.Y, 0.0f, 0.0f);
+	}
+
+
+	template <typename TRHICmdList>
+	void SetParametersCS(TRHICmdList& RHICmdList, const FRenderingCompositePassContext& Context, FIntPoint DestSize, FIntPoint InputTextureSize, FRHIUnorderedAccessView *OutUAV)
+	{
+		const FFinalPostProcessSettings& Settings = Context.View.FinalPostProcessSettings;
+		FRHIComputeShader* ShaderRHI = GetComputeShader();
+
+		FGlobalShader::SetParameters<FViewUniformShaderParameters>(RHICmdList, ShaderRHI, Context.View.ViewUniformBuffer);
+		SceneTextureParameters.Set(RHICmdList, ShaderRHI, Context.View.FeatureLevel, ESceneTextureSetupMode::All);
+		ScreenSpaceAOParams.Set(RHICmdList, Context.View, ShaderRHI, InputTextureSize, FScreenSpaceAOParameters::ERandTexType::GTAO);
+
+		GTAOParams.Set(RHICmdList, Context.View, ShaderRHI);
+
+
+		PostprocessParameter.SetCS(ShaderRHI, Context, RHICmdList, TStaticSamplerState<SF_Point, AM_Clamp, AM_Clamp, AM_Clamp>::GetRHI());
+
+		RHICmdList.SetUAVParameter(ShaderRHI, HorizonOutTexture.GetBaseIndex(), OutUAV);
+
+		FVector4 HZBRemappingValue = GetHZBRemapVal(Context, DestSize, InputTextureSize);
+		SetShaderValue(RHICmdList, ShaderRHI, HZBRemapping, HZBRemappingValue);
+
+		const FSceneRenderTargetItem& GTAORandomization = GSystemTextures.GTAORandomization->GetRenderTargetItem();
+		SetTextureParameter(RHICmdList, ShaderRHI, RandomNormalTexture, RandomNormalTextureSampler, TStaticSamplerState<SF_Point, AM_Wrap, AM_Wrap, AM_Wrap>::GetRHI(), GTAORandomization.ShaderResourceTexture);
+	}
+
+	void SetParametersPS(const FRenderingCompositePassContext& Context, FIntPoint DestSize, FIntPoint InputTextureSize)
+	{
+		const FFinalPostProcessSettings& Settings = Context.View.FinalPostProcessSettings;
+		FRHIPixelShader* ShaderRHI = GetPixelShader();
+
+		PostprocessParameter.SetPS(Context.RHICmdList, ShaderRHI, Context, TStaticSamplerState<SF_Point, AM_Clamp, AM_Clamp, AM_Clamp>::GetRHI());
+
+		FGlobalShader::SetParameters<FViewUniformShaderParameters>(Context.RHICmdList, ShaderRHI, Context.View.ViewUniformBuffer);
+		SceneTextureParameters.Set(Context.RHICmdList, ShaderRHI, Context.View.FeatureLevel, ESceneTextureSetupMode::All);
+		ScreenSpaceAOParams.Set(Context.RHICmdList, Context.View, ShaderRHI, InputTextureSize, FScreenSpaceAOParameters::ERandTexType::GTAO);
+		GTAOParams.Set(Context.RHICmdList, Context.View, ShaderRHI);
+
+		FVector4 HZBRemappingValue = GetHZBRemapVal(Context, DestSize, InputTextureSize);
+		SetShaderValue(Context.RHICmdList, ShaderRHI, HZBRemapping, HZBRemappingValue);
+
+		const FSceneRenderTargetItem& GTAORandomization = GSystemTextures.GTAORandomization->GetRenderTargetItem();
+		SetTextureParameter(Context.RHICmdList, ShaderRHI, RandomNormalTexture, RandomNormalTextureSampler, TStaticSamplerState<SF_Point, AM_Wrap, AM_Wrap, AM_Wrap>::GetRHI(), GTAORandomization.ShaderResourceTexture);
+	}
+
+
+	template <typename TRHICmdList>
+	void UnsetParameters(TRHICmdList& RHICmdList)
+	{
+		FRHIComputeShader* ShaderRHI = GetComputeShader();
+		RHICmdList.SetUAVParameter(ShaderRHI, HorizonOutTexture.GetBaseIndex(), nullptr);
+	}
+
+
+	// FShader interface.
+	virtual bool Serialize(FArchive& Ar) override
+	{
+		bool bShaderHasOutdatedParameters = FGlobalShader::Serialize(Ar);
+		Ar << PostprocessParameter << SceneTextureParameters << ScreenSpaceAOParams << GTAOParams << HorizonOutTexture << HZBRemapping << RandomNormalTexture << RandomNormalTextureSampler;
+		return bShaderHasOutdatedParameters;
+	}
+
+	static const TCHAR* GetSourceFilename()
+	{
+		return TEXT("/Engine/Private/PostProcessAmbientOcclusion.usf");
+	}
+
+	static const TCHAR* GetFunctionName()
+	{
+		return bComputeShader ? TEXT("HorizonSearchCS") : TEXT("HorizonSearchPS");
+	}
+};
+IMPLEMENT_SHADER_TYPE2(FPostProcessGTAOHorizonSearchPSandCS<false>, SF_Pixel);
+IMPLEMENT_SHADER_TYPE2(FPostProcessGTAOHorizonSearchPSandCS<true>, SF_Compute);
+
+
+template<uint32 bComputeShader>
+class FPostProcessGTAOCommon
+{
+public:
+
+	static bool ShouldCompilePermutation(const FGlobalShaderPermutationParameters& Parameters)
+	{
+		return IsFeatureLevelSupported(Parameters.Platform, ERHIFeatureLevel::SM5);
+	}
+
+	static void ModifyCompilationEnvironment(const FGlobalShaderPermutationParameters& Parameters, FShaderCompilerEnvironment& OutEnvironment)
+	{
+		FGlobalShader::ModifyCompilationEnvironment(Parameters, OutEnvironment);
+		OutEnvironment.SetDefine(TEXT("COMPUTE_SHADER"), bComputeShader);
+		if (bComputeShader)
+		{
+			OutEnvironment.SetDefine(TEXT("THREADGROUP_SIZEX"), 8);
+			OutEnvironment.SetDefine(TEXT("THREADGROUP_SIZEY"), 8);
+		}
+	}
+};
+
+
+template<uint32 bComputeShader>
+class FPostProcessGTAOHorizonSearchAndInnerIntegratePSandCS : public FGlobalShader, FPostProcessGTAOCommon<bComputeShader>
+{
+	DECLARE_SHADER_TYPE(FPostProcessGTAOHorizonSearchAndInnerIntegratePSandCS, Global);
+
+	static bool ShouldCompilePermutation(const FGlobalShaderPermutationParameters& Parameters)
+	{
+		return FPostProcessGTAOCommon<bComputeShader>::ShouldCompilePermutation(Parameters);
+	}
+
+	static void ModifyCompilationEnvironment(const FGlobalShaderPermutationParameters& Parameters, FShaderCompilerEnvironment& OutEnvironment)
+	{
+		FPostProcessGTAOCommon<bComputeShader>::ModifyCompilationEnvironment(Parameters, OutEnvironment);
+		
+	}
+
+	/** Default constructor. */
+	FPostProcessGTAOHorizonSearchAndInnerIntegratePSandCS() {}
+
+public:
+	FSceneTextureShaderParameters SceneTextureParameters;
+	FPostProcessPassParameters	  PostprocessParameter;
+	FScreenSpaceAOParameters	  ScreenSpaceAOParams;
+	FShaderResourceParameter 	  OutTexture;
+	FShaderParameter			  HZBRemapping;
+	FShaderResourceParameter      RandomNormalTexture;
+	FShaderResourceParameter      RandomNormalTextureSampler;
+
+	/** Initialization constructor. */
+	FPostProcessGTAOHorizonSearchAndInnerIntegratePSandCS(const ShaderMetaType::CompiledShaderInitializerType& Initializer)
+		: FGlobalShader(Initializer)
+	{
+		PostprocessParameter.Bind(Initializer.ParameterMap);
+		SceneTextureParameters.Bind(Initializer);
+		ScreenSpaceAOParams.Bind(Initializer.ParameterMap);
+		HZBRemapping.Bind(Initializer.ParameterMap, TEXT("HZBRemapping"));
+		RandomNormalTexture.Bind(Initializer.ParameterMap, TEXT("RandomNormalTexture"));
+		RandomNormalTextureSampler.Bind(Initializer.ParameterMap, TEXT("RandomNormalTextureSampler"));
+
+		if (bComputeShader)
+		{
+			OutTexture.Bind(Initializer.ParameterMap, TEXT("OutTexture"));
+		}
+	}
+
+	FVector4 GetHZBRemapVal(const FRenderingCompositePassContext& Context, FIntPoint DestSize, FIntPoint InputTextureSize)
+	{
+		const FVector2D HZBScaleFactor(float(Context.View.ViewRect.Width()) / float(2 * Context.View.HZBMipmap0Size.X),
+			float(Context.View.ViewRect.Height()) / float(2 * Context.View.HZBMipmap0Size.Y));
+
+		const FVector2D UVScaleFactor(float(DestSize.X) / float(InputTextureSize.X),
+			float(DestSize.Y) / float(InputTextureSize.Y));
+
+		return FVector4(HZBScaleFactor.X / UVScaleFactor.X, HZBScaleFactor.Y / UVScaleFactor.Y, 0.0f, 0.0f);
+	}
+
+
+	template <typename TRHICmdList>
+	void SetParametersCS(TRHICmdList& RHICmdList, const FRenderingCompositePassContext& Context, FIntPoint DestSize, FIntPoint InputTextureSize, FRHIUnorderedAccessView *OutUAV)
+	{
+		const FFinalPostProcessSettings& Settings = Context.View.FinalPostProcessSettings;
+		FRHIComputeShader* ShaderRHI = GetComputeShader();
+
+		FGlobalShader::SetParameters<FViewUniformShaderParameters>(RHICmdList, ShaderRHI, Context.View.ViewUniformBuffer);
+		SceneTextureParameters.Set(RHICmdList, ShaderRHI, Context.View.FeatureLevel, ESceneTextureSetupMode::All);
+		ScreenSpaceAOParams.Set(RHICmdList, Context.View, ShaderRHI, InputTextureSize, FScreenSpaceAOParameters::ERandTexType::GTAO);
+
+		PostprocessParameter.SetCS(ShaderRHI, Context, RHICmdList, TStaticSamplerState<SF_Point, AM_Clamp, AM_Clamp, AM_Clamp>::GetRHI());
+
+		RHICmdList.SetUAVParameter(ShaderRHI, OutTexture.GetBaseIndex(), OutUAV);
+
+		FVector4 HZBRemappingValue = GetHZBRemapVal(Context, DestSize, InputTextureSize);
+		SetShaderValue(Context.RHICmdList, ShaderRHI, HZBRemapping, HZBRemappingValue);
+
+		const FSceneRenderTargetItem& GTAORandomization = GSystemTextures.GTAORandomization->GetRenderTargetItem();
+		SetTextureParameter(Context.RHICmdList, ShaderRHI, RandomNormalTexture, RandomNormalTextureSampler, TStaticSamplerState<SF_Point, AM_Wrap, AM_Wrap, AM_Wrap>::GetRHI(), GTAORandomization.ShaderResourceTexture);
+	}
+
+	void SetParametersPS(const FRenderingCompositePassContext& Context, FIntPoint DestSize, FIntPoint InputTextureSize)
+	{
+		const FFinalPostProcessSettings& Settings = Context.View.FinalPostProcessSettings;
+		FRHIPixelShader* ShaderRHI = GetPixelShader();
+
+		PostprocessParameter.SetPS(Context.RHICmdList, ShaderRHI, Context, TStaticSamplerState<SF_Point, AM_Clamp, AM_Clamp, AM_Clamp>::GetRHI());
+
+		FGlobalShader::SetParameters<FViewUniformShaderParameters>(Context.RHICmdList, ShaderRHI, Context.View.ViewUniformBuffer);
+		SceneTextureParameters.Set(Context.RHICmdList, ShaderRHI, Context.View.FeatureLevel, ESceneTextureSetupMode::All);
+		ScreenSpaceAOParams.Set(Context.RHICmdList, Context.View, ShaderRHI, InputTextureSize, FScreenSpaceAOParameters::ERandTexType::GTAO);
+
+		FVector4 HZBRemappingValue = GetHZBRemapVal(Context, DestSize, InputTextureSize);
+		SetShaderValue(Context.RHICmdList, ShaderRHI, HZBRemapping, HZBRemappingValue);
+
+		const FSceneRenderTargetItem& GTAORandomization = GSystemTextures.GTAORandomization->GetRenderTargetItem();
+		SetTextureParameter(Context.RHICmdList, ShaderRHI, RandomNormalTexture, RandomNormalTextureSampler, TStaticSamplerState<SF_Point, AM_Wrap, AM_Wrap, AM_Wrap>::GetRHI(), GTAORandomization.ShaderResourceTexture);
+	}
+
+
+	template <typename TRHICmdList>
+	void UnsetParameters(TRHICmdList& RHICmdList)
+	{
+		FRHIComputeShader* ShaderRHI = GetComputeShader();
+		RHICmdList.SetUAVParameter(ShaderRHI, OutTexture.GetBaseIndex(), nullptr);
+	}
+
+
+	// FShader interface.
+	virtual bool Serialize(FArchive& Ar) override
+	{
+		bool bShaderHasOutdatedParameters = FGlobalShader::Serialize(Ar);
+		Ar << PostprocessParameter << SceneTextureParameters << ScreenSpaceAOParams << OutTexture << HZBRemapping << RandomNormalTexture << RandomNormalTextureSampler;
+		return bShaderHasOutdatedParameters;
+	}
+
+	static const TCHAR* GetSourceFilename()
+	{
+		return TEXT("/Engine/Private/PostProcessAmbientOcclusion.usf");
+	}
+
+	static const TCHAR* GetFunctionName()
+	{
+		return bComputeShader ? TEXT("HorizonSearchAndInnerIntegrateCS") : TEXT("HorizonSearchAndInnerIntegratePS");
+	}
+};
+IMPLEMENT_SHADER_TYPE2(FPostProcessGTAOHorizonSearchAndInnerIntegratePSandCS<false>, SF_Pixel);
+IMPLEMENT_SHADER_TYPE2(FPostProcessGTAOHorizonSearchAndInnerIntegratePSandCS<true>, SF_Compute);
+
+
+template<uint32 bComputeShader>
+class FPostProcessGTAOInnerIntegratePSandCS : public FGlobalShader
+{
+	DECLARE_SHADER_TYPE(FPostProcessGTAOInnerIntegratePSandCS, Global);
+
+	static bool ShouldCompilePermutation(const FGlobalShaderPermutationParameters& Parameters)
+	{
+		return IsFeatureLevelSupported(Parameters.Platform, ERHIFeatureLevel::SM5);
+	}
+
+	static void ModifyCompilationEnvironment(const FGlobalShaderPermutationParameters& Parameters, FShaderCompilerEnvironment& OutEnvironment)
+	{
+		FGlobalShader::ModifyCompilationEnvironment(Parameters, OutEnvironment);
+		OutEnvironment.SetDefine(TEXT("COMPUTE_SHADER"), bComputeShader);
+		if (bComputeShader)
+		{
+			OutEnvironment.SetDefine(TEXT("THREADGROUP_SIZEX"), 8);
+			OutEnvironment.SetDefine(TEXT("THREADGROUP_SIZEY"), 8);
+		}
+	}
+
+	/** Default constructor. */
+	FPostProcessGTAOInnerIntegratePSandCS() {}
+
+public:
+	FSceneTextureShaderParameters SceneTextureParameters;
+	FPostProcessPassParameters	  PostprocessParameter;
+	FScreenSpaceAOParameters	  ScreenSpaceAOParams;
+	FShaderResourceParameter 	  OutTexture;
+	FShaderResourceParameter      RandomNormalTexture;
+	FShaderResourceParameter      RandomNormalTextureSampler;
+	FShaderResourceParameter      HorizonsTexture;
+	FShaderResourceParameter      HorizonsTextureSampler;
+
+
+	/** Initialization constructor. */
+	FPostProcessGTAOInnerIntegratePSandCS(const ShaderMetaType::CompiledShaderInitializerType& Initializer)
+		: FGlobalShader(Initializer)
+	{
+		PostprocessParameter.Bind(Initializer.ParameterMap);
+		SceneTextureParameters.Bind(Initializer);
+		ScreenSpaceAOParams.Bind(Initializer.ParameterMap);
+		RandomNormalTexture.Bind(Initializer.ParameterMap, TEXT("RandomNormalTexture"));
+		RandomNormalTextureSampler.Bind(Initializer.ParameterMap, TEXT("RandomNormalTextureSampler"));
+		HorizonsTexture.Bind(Initializer.ParameterMap, TEXT("HorizonsTexture"));
+		HorizonsTextureSampler.Bind(Initializer.ParameterMap, TEXT("HorizonsTextureSampler"));
+
+		if (bComputeShader)
+		{
+			OutTexture.Bind(Initializer.ParameterMap, TEXT("OutTexture"));
+		}
+	}
+	// FShader interface.
+	virtual bool Serialize(FArchive& Ar) override
+	{
+		bool bShaderHasOutdatedParameters = FGlobalShader::Serialize(Ar);
+		Ar << PostprocessParameter << SceneTextureParameters << ScreenSpaceAOParams << OutTexture <<  RandomNormalTexture << RandomNormalTextureSampler;
+		Ar << HorizonsTexture << HorizonsTextureSampler;
+
+		return bShaderHasOutdatedParameters;
+	}
+
+
+	template <typename TRHICmdList>
+	void SetParametersCS(TRHICmdList& RHICmdList, const FRenderingCompositePassContext& Context, FIntPoint DestSize, FIntPoint InputTextureSize, FRHIUnorderedAccessView *OutUAV)
+	{
+		const FFinalPostProcessSettings& Settings = Context.View.FinalPostProcessSettings;
+		FRHIComputeShader* ShaderRHI = GetComputeShader();
+
+		FGlobalShader::SetParameters<FViewUniformShaderParameters>(RHICmdList, ShaderRHI, Context.View.ViewUniformBuffer);
+		SceneTextureParameters.Set(RHICmdList, ShaderRHI, Context.View.FeatureLevel, ESceneTextureSetupMode::All);
+		ScreenSpaceAOParams.Set(RHICmdList, Context.View, ShaderRHI, InputTextureSize, FScreenSpaceAOParameters::ERandTexType::GTAO);
+
+		PostprocessParameter.SetCS(ShaderRHI, Context, RHICmdList, TStaticSamplerState<SF_Point, AM_Clamp, AM_Clamp, AM_Clamp>::GetRHI());
+
+		RHICmdList.SetUAVParameter(ShaderRHI, OutTexture.GetBaseIndex(), OutUAV);
+
+		const FSceneRenderTargetItem& GTAORandomization = GSystemTextures.GTAORandomization->GetRenderTargetItem();
+		SetTextureParameter(RHICmdList, ShaderRHI, RandomNormalTexture, RandomNormalTextureSampler, TStaticSamplerState<SF_Point, AM_Wrap, AM_Wrap, AM_Wrap>::GetRHI(), GTAORandomization.ShaderResourceTexture);
+
+		FSceneRenderTargets& SceneContext = FSceneRenderTargets::Get(Context.RHICmdList);
+
+		SetTextureParameter(RHICmdList, ShaderRHI, HorizonsTexture, HorizonsTextureSampler, TStaticSamplerState<SF_Point, AM_Wrap, AM_Wrap, AM_Wrap>::GetRHI(), 
+			SceneContext.ScreenSpaceGTAOHorizons->GetRenderTargetItem().ShaderResourceTexture);
+
+	}
+
+	void SetParametersPS(const FRenderingCompositePassContext& Context, FIntPoint DestSize, FIntPoint InputTextureSize)
+	{
+		const FFinalPostProcessSettings& Settings = Context.View.FinalPostProcessSettings;
+		FRHIPixelShader* ShaderRHI = GetPixelShader();
+
+		PostprocessParameter.SetPS(Context.RHICmdList, ShaderRHI, Context, TStaticSamplerState<SF_Point, AM_Clamp, AM_Clamp, AM_Clamp>::GetRHI());
+
+		FGlobalShader::SetParameters<FViewUniformShaderParameters>(Context.RHICmdList, ShaderRHI, Context.View.ViewUniformBuffer);
+		SceneTextureParameters.Set(Context.RHICmdList, ShaderRHI, Context.View.FeatureLevel, ESceneTextureSetupMode::All);
+		ScreenSpaceAOParams.Set(Context.RHICmdList, Context.View, ShaderRHI, InputTextureSize, FScreenSpaceAOParameters::ERandTexType::GTAO);
+
+		const FSceneRenderTargetItem& GTAORandomization = GSystemTextures.GTAORandomization->GetRenderTargetItem();
+		SetTextureParameter(Context.RHICmdList, ShaderRHI, RandomNormalTexture, RandomNormalTextureSampler, TStaticSamplerState<SF_Point, AM_Wrap, AM_Wrap, AM_Wrap>::GetRHI(), GTAORandomization.ShaderResourceTexture);
+		FSceneRenderTargets& SceneContext = FSceneRenderTargets::Get(Context.RHICmdList);
+
+		SetTextureParameter(Context.RHICmdList, ShaderRHI, HorizonsTexture, HorizonsTextureSampler, TStaticSamplerState<SF_Point, AM_Wrap, AM_Wrap, AM_Wrap>::GetRHI(),
+		SceneContext.ScreenSpaceGTAOHorizons->GetRenderTargetItem().ShaderResourceTexture);
+
+	}
+
+	template <typename TRHICmdList>
+	void UnsetParameters(TRHICmdList& RHICmdList)
+	{
+		FRHIComputeShader* ShaderRHI = GetComputeShader();
+		RHICmdList.SetUAVParameter(ShaderRHI, OutTexture.GetBaseIndex(), nullptr);
+	}
+
+	static const TCHAR* GetSourceFilename()
+	{
+		return TEXT("/Engine/Private/PostProcessAmbientOcclusion.usf");
+	}
+
+	static const TCHAR* GetFunctionName()
+	{
+		return bComputeShader ? TEXT("GTAOInnerIntegrateCS") : TEXT("GTAOInnerIntegratePS");
+	}
+};
+
+
+IMPLEMENT_SHADER_TYPE2(FPostProcessGTAOInnerIntegratePSandCS<false>, SF_Pixel);
+IMPLEMENT_SHADER_TYPE2(FPostProcessGTAOInnerIntegratePSandCS<true>, SF_Compute);
+
+
+
+FRCPassPostProcessAmbientOcclusion_HorizonSearchInnerIntegrate::FRCPassPostProcessAmbientOcclusion_HorizonSearchInnerIntegrate(const FSceneView& View, bool FinalOutput)
+	:bFinalOutput(FinalOutput)
+{
+
+}
+
+void FRCPassPostProcessAmbientOcclusion_HorizonSearchInnerIntegrate::Process(FRenderingCompositePassContext& Context)
+{
+	SCOPED_GPU_STAT(Context.RHICmdList, GTAO_HorizonSearchInnerIntegrate);
+	const FViewInfo& View = Context.View;
+	FIntRect ViewRect = View.ViewRect;
+	FIntRect SrcRect = Context.View.ViewRect;
+	FIntRect DestRect = SrcRect;
+	FIntPoint DestSize(DestRect.Width(), DestRect.Height());
+	const FPooledRenderTargetDesc* InputDesc0 = GetInputDesc(ePId_Input0);
+	FSceneRenderTargets& SceneContext = FSceneRenderTargets::Get(Context.RHICmdList);
+	FIntPoint TexSize = InputDesc0 ? InputDesc0->Extent : SceneContext.GetBufferSizeXY();
+
+	const FSceneRenderTargetItem& DestRenderTarget = bFinalOutput ? SceneContext.ScreenSpaceAO->GetRenderTargetItem() : PassOutputs[0].RequestSurface(Context);
+
+	static const auto CVar = IConsoleManager::Get().FindTConsoleVariableDataInt(TEXT("r.AmbientOcclusion.Compute"));
+
+	if (CVar->GetValueOnRenderThread() >= 1)
+	{
+		Context.RHICmdList.TransitionResource(EResourceTransitionAccess::ERWBarrier, EResourceTransitionPipeline::EGfxToCompute, DestRenderTarget.UAV, nullptr);
+
+		// Compute  Version
+		UnbindRenderTargets(Context.RHICmdList);
+
+		TShaderMapRef<FPostProcessGTAOHorizonSearchAndInnerIntegratePSandCS<true>> ComputeShader(Context.GetShaderMap());
+		Context.RHICmdList.SetComputeShader(ComputeShader->GetComputeShader());
+		ComputeShader->SetParametersCS(Context.RHICmdList, Context, DestSize, TexSize, DestRenderTarget.UAV);
+
+		uint32 GroupSizeX = FMath::DivideAndRoundUp(DestSize.X, 8);
+		uint32 GroupSizeY = FMath::DivideAndRoundUp(DestSize.Y, 8);
+		DispatchComputeShader(Context.RHICmdList, *ComputeShader, GroupSizeX, GroupSizeY, 1);
+
+		ComputeShader->UnsetParameters(Context.RHICmdList);
+	}
+	else
+	{
+		// Pixel  Version
+		FRHIRenderPassInfo RPInfo(DestRenderTarget.TargetableTexture, ERenderTargetActions::Load_Store);
+		Context.RHICmdList.BeginRenderPass(RPInfo, TEXT("AmbientOcclusionSetup"));
+		{
+			// PS Version
+			Context.SetViewportAndCallRHI(ViewRect);
+
+			TShaderMapRef<FPostProcessGTAOHorizonSearchAndInnerIntegratePSandCS<false> > PixelShader(Context.GetShaderMap());
+			TShaderMapRef<FPostProcessVS> VertexShader(Context.GetShaderMap());
+
+			FGraphicsPipelineStateInitializer GraphicsPSOInit;
+			Context.RHICmdList.ApplyCachedRenderTargets(GraphicsPSOInit);
+			GraphicsPSOInit.BlendState = TStaticBlendState<>::GetRHI();
+			GraphicsPSOInit.RasterizerState = TStaticRasterizerState<>::GetRHI();
+			GraphicsPSOInit.DepthStencilState = TStaticDepthStencilState<false, CF_Always>::GetRHI();
+			GraphicsPSOInit.PrimitiveType = PT_TriangleList;
+			GraphicsPSOInit.BoundShaderState.VertexDeclarationRHI = GFilterVertexDeclaration.VertexDeclarationRHI;
+			GraphicsPSOInit.BoundShaderState.VertexShaderRHI = GETSAFERHISHADER_VERTEX(*VertexShader);
+			GraphicsPSOInit.BoundShaderState.PixelShaderRHI = GETSAFERHISHADER_PIXEL(*PixelShader);
+			SetGraphicsPipelineState(Context.RHICmdList, GraphicsPSOInit);
+
+			PixelShader->SetParametersPS(Context, DestSize, TexSize);
+
+			DrawRectangle(
+				Context.RHICmdList,
+				0, 0,
+				ViewRect.Width(), ViewRect.Height(),
+				ViewRect.Min.X, ViewRect.Min.Y,
+				ViewRect.Width(), ViewRect.Height(),
+				ViewRect.Size(),
+				TexSize,
+				*VertexShader,
+				EDRF_UseTriangleOptimization);
+		}
+		Context.RHICmdList.EndRenderPass();
+	}
+
+	Context.RHICmdList.TransitionResource(EResourceTransitionAccess::EReadable, DestRenderTarget.TargetableTexture);
+}
+
+
+FPooledRenderTargetDesc FRCPassPostProcessAmbientOcclusion_HorizonSearchInnerIntegrate::ComputeOutputDesc(EPassOutputId InPassOutputId) const
+{
+	FPooledRenderTargetDesc Ret;
+
+	Ret = GetInput(ePId_Input0)->GetOutput()->RenderTargetDesc;
+	Ret.Reset();
+	Ret.Format = PF_B8G8R8A8;
+
+	Ret.ClearValue = FClearValueBinding::None;
+	Ret.TargetableFlags &= ~TexCreate_DepthStencilTargetable;
+	Ret.TargetableFlags |= TexCreate_RenderTargetable | TexCreate_ShaderResource;
+	Ret.DebugName = TEXT("GTAOInnerHorizonInnerIntegrate");
+
+	return Ret;
+}
+
+
+
+
+FRCPassPostProcessAmbientOcclusion_InnerIntegrate::FRCPassPostProcessAmbientOcclusion_InnerIntegrate(const FSceneView& View, bool FinalOutput)
+	:bFinalOutput(FinalOutput)
+{
+
+}
+
+
+
+
+void FRCPassPostProcessAmbientOcclusion_InnerIntegrate::Process(FRenderingCompositePassContext& Context)
+{
+	SCOPED_GPU_STAT(Context.RHICmdList, GTAO_InnerIntegrate);
+	const FViewInfo& View = Context.View;
+	FIntRect ViewRect = View.ViewRect;
+	FIntRect SrcRect = Context.View.ViewRect;
+	FIntRect DestRect = SrcRect;
+	FIntPoint DestSize(DestRect.Width(), DestRect.Height());
+	const FPooledRenderTargetDesc* InputDesc0 = GetInputDesc(ePId_Input0);
+
+	FSceneRenderTargets& SceneContext = FSceneRenderTargets::Get(Context.RHICmdList);
+	FIntPoint TexSize = InputDesc0 ? InputDesc0->Extent : SceneContext.GetBufferSizeXY();
+
+	const FSceneRenderTargetItem& DestRenderTarget = bFinalOutput ? SceneContext.ScreenSpaceAO->GetRenderTargetItem() : PassOutputs[0].RequestSurface(Context);
+	static const auto CVar = IConsoleManager::Get().FindTConsoleVariableDataInt(TEXT("r.AmbientOcclusion.Compute"));
+
+	if (CVar->GetValueOnRenderThread() >= 1)
+	{
+		// Compute  Version
+		Context.RHICmdList.TransitionResource(EResourceTransitionAccess::ERWBarrier, EResourceTransitionPipeline::EGfxToCompute, DestRenderTarget.UAV, nullptr);
+
+		UnbindRenderTargets(Context.RHICmdList);
+
+		TShaderMapRef<FPostProcessGTAOInnerIntegratePSandCS<true>> ComputeShader(Context.GetShaderMap());
+		Context.RHICmdList.SetComputeShader(ComputeShader->GetComputeShader());
+		ComputeShader->SetParametersCS(Context.RHICmdList, Context, DestSize, TexSize, DestRenderTarget.UAV);
+
+		uint32 GroupSizeX = FMath::DivideAndRoundUp(DestSize.X, 8);
+		uint32 GroupSizeY = FMath::DivideAndRoundUp(DestSize.Y, 8);
+		DispatchComputeShader(Context.RHICmdList, *ComputeShader, GroupSizeX, GroupSizeY, 1);
+
+		ComputeShader->UnsetParameters(Context.RHICmdList);
+	}
+	else
+	{
+		// Pixel  Version
+		FRHIRenderPassInfo RPInfo(DestRenderTarget.TargetableTexture, ERenderTargetActions::Load_Store);
+		Context.RHICmdList.BeginRenderPass(RPInfo, TEXT("AmbientOcclusionSetup"));
+		{
+			Context.SetViewportAndCallRHI(ViewRect);
+
+			TShaderMapRef<FPostProcessGTAOInnerIntegratePSandCS<false> > PixelShader(Context.GetShaderMap());
+			TShaderMapRef<FPostProcessVS> VertexShader(Context.GetShaderMap());
+
+			FGraphicsPipelineStateInitializer GraphicsPSOInit;
+			Context.RHICmdList.ApplyCachedRenderTargets(GraphicsPSOInit);
+			GraphicsPSOInit.BlendState = TStaticBlendState<>::GetRHI();
+			GraphicsPSOInit.RasterizerState = TStaticRasterizerState<>::GetRHI();
+			GraphicsPSOInit.DepthStencilState = TStaticDepthStencilState<false, CF_Always>::GetRHI();
+			GraphicsPSOInit.PrimitiveType = PT_TriangleList;
+			GraphicsPSOInit.BoundShaderState.VertexDeclarationRHI = GFilterVertexDeclaration.VertexDeclarationRHI;
+			GraphicsPSOInit.BoundShaderState.VertexShaderRHI = GETSAFERHISHADER_VERTEX(*VertexShader);
+			GraphicsPSOInit.BoundShaderState.PixelShaderRHI = GETSAFERHISHADER_PIXEL(*PixelShader);
+			SetGraphicsPipelineState(Context.RHICmdList, GraphicsPSOInit);
+
+			PixelShader->SetParametersPS(Context, DestSize, TexSize);
+
+			DrawRectangle(
+				Context.RHICmdList,
+				0, 0,
+				ViewRect.Width(), ViewRect.Height(),
+				ViewRect.Min.X, ViewRect.Min.Y,
+				ViewRect.Width(), ViewRect.Height(),
+				ViewRect.Size(),
+				TexSize,
+				*VertexShader,
+				EDRF_UseTriangleOptimization);
+		}
+		Context.RHICmdList.EndRenderPass();
+	}
+
+	Context.RHICmdList.TransitionResource(EResourceTransitionAccess::EReadable, DestRenderTarget.TargetableTexture);
+}
+
+
+FPooledRenderTargetDesc FRCPassPostProcessAmbientOcclusion_InnerIntegrate::ComputeOutputDesc(EPassOutputId InPassOutputId) const
+{
+	FPooledRenderTargetDesc Ret;
+
+	Ret = GetInput(ePId_Input0)->GetOutput()->RenderTargetDesc;
+	Ret.Reset();
+	Ret.Format = PF_B8G8R8A8;
+
+	Ret.ClearValue = FClearValueBinding::None;
+	Ret.TargetableFlags &= ~TexCreate_DepthStencilTargetable;
+	Ret.TargetableFlags |= TexCreate_UAV;
+	Ret.TargetableFlags |= TexCreate_RenderTargetable | TexCreate_ShaderResource;
+	Ret.DebugName = TEXT("GTAOInnerIntegrate");
+	return Ret;
+}
+
+
+
+FRCPassPostProcessAmbientOcclusion_HorizonSearch::FRCPassPostProcessAmbientOcclusion_HorizonSearch(const FSceneView& View, const ESSAOType InAOType)
+	: AOType(InAOType)
+{
+}
+
+
+
+
+void FRCPassPostProcessAmbientOcclusion_HorizonSearch::Process(FRenderingCompositePassContext& Context)
+{
+	SCOPED_GPU_STAT(Context.RHICmdList, GTAO_HorizonSearch);
+	const FViewInfo& View = Context.View;
+	FIntRect ViewRect = View.ViewRect;
+	FIntRect SrcRect = Context.View.ViewRect;
+	FIntRect DestRect = SrcRect;
+	FIntPoint DestSize(DestRect.Width(), DestRect.Height());
+	const FPooledRenderTargetDesc* InputDesc0 = GetInputDesc(ePId_Input0);
+	FSceneRenderTargets& SceneContext = FSceneRenderTargets::Get(Context.RHICmdList);
+	FIntPoint TexSize = InputDesc0 ? InputDesc0->Extent : SceneContext.GetBufferSizeXY();
+
+
+	const FSceneRenderTargetItem& DestRenderTarget = (AOType == ESSAOType::EAsyncCS) ? SceneContext.ScreenSpaceGTAOHorizons->GetRenderTargetItem()  : PassOutputs[0].RequestSurface(Context);
+
+	static const auto CVar = IConsoleManager::Get().FindTConsoleVariableDataInt(TEXT("r.AmbientOcclusion.Compute"));
+
+	if (AOType == ESSAOType::EAsyncCS)
+	{
+		static FName AsyncStartFenceName(TEXT("AsyncStartFence"));
+		FComputeFenceRHIRef AsyncStartFence = Context.RHICmdList.CreateComputeFence(AsyncStartFenceName);
+		FRHIAsyncComputeCommandListImmediate& RHICmdListComputeImmediate = FRHICommandListExecutor::GetImmediateAsyncComputeCommandList();
+
+		//Fence to let us know when the Gfx pipe is done with the RT we want to write to.
+		Context.RHICmdList.TransitionResource(EResourceTransitionAccess::ERWBarrier, EResourceTransitionPipeline::EGfxToCompute, DestRenderTarget.UAV, AsyncStartFence);
+
+		SCOPED_COMPUTE_EVENT(RHICmdListComputeImmediate, AsyncSSAO);
+		RHICmdListComputeImmediate.WaitComputeFence(AsyncStartFence);
+
+		// Compute  Version
+		UnbindRenderTargets(Context.RHICmdList);
+
+		TShaderMapRef<FPostProcessGTAOHorizonSearchPSandCS<true>> ComputeShader(Context.GetShaderMap());
+		RHICmdListComputeImmediate.SetComputeShader(ComputeShader->GetComputeShader());
+		ComputeShader->SetParametersCS(RHICmdListComputeImmediate, Context, DestSize, TexSize, DestRenderTarget.UAV);
+
+		uint32 GroupSizeX = FMath::DivideAndRoundUp(DestSize.X, 8);
+		uint32 GroupSizeY = FMath::DivideAndRoundUp(DestSize.Y, 8);
+		DispatchComputeShader(RHICmdListComputeImmediate, *ComputeShader, GroupSizeX, GroupSizeY, 1);
+
+		ComputeShader->UnsetParameters(RHICmdListComputeImmediate);
+	}
+	else
+	{
+		// Pixel  Version
+		FRHIRenderPassInfo RPInfo(DestRenderTarget.TargetableTexture, ERenderTargetActions::Load_Store);
+		Context.RHICmdList.BeginRenderPass(RPInfo, TEXT("AmbientOcclusionSetup"));
+		{
+			Context.SetViewportAndCallRHI(ViewRect);
+
+			TShaderMapRef<FPostProcessGTAOHorizonSearchPSandCS<false> > PixelShader(Context.GetShaderMap());
+			TShaderMapRef<FPostProcessVS> VertexShader(Context.GetShaderMap());
+
+			FGraphicsPipelineStateInitializer GraphicsPSOInit;
+			Context.RHICmdList.ApplyCachedRenderTargets(GraphicsPSOInit);
+			GraphicsPSOInit.BlendState = TStaticBlendState<>::GetRHI();
+			GraphicsPSOInit.RasterizerState = TStaticRasterizerState<>::GetRHI();
+			GraphicsPSOInit.DepthStencilState = TStaticDepthStencilState<false, CF_Always>::GetRHI();
+			GraphicsPSOInit.PrimitiveType = PT_TriangleList;
+			GraphicsPSOInit.BoundShaderState.VertexDeclarationRHI = GFilterVertexDeclaration.VertexDeclarationRHI;
+			GraphicsPSOInit.BoundShaderState.VertexShaderRHI = GETSAFERHISHADER_VERTEX(*VertexShader);
+			GraphicsPSOInit.BoundShaderState.PixelShaderRHI = GETSAFERHISHADER_PIXEL(*PixelShader);
+			SetGraphicsPipelineState(Context.RHICmdList, GraphicsPSOInit);
+
+			PixelShader->SetParametersPS(Context, DestSize, TexSize);
+
+			DrawRectangle(
+				Context.RHICmdList,
+				0, 0,
+				ViewRect.Width(), ViewRect.Height(),
+				ViewRect.Min.X, ViewRect.Min.Y,
+				ViewRect.Width(), ViewRect.Height(),
+				ViewRect.Size(),
+				TexSize,
+				*VertexShader,
+				EDRF_UseTriangleOptimization);
+		}
+		Context.RHICmdList.EndRenderPass();
+	}
+
+	Context.RHICmdList.TransitionResource(EResourceTransitionAccess::EReadable, DestRenderTarget.TargetableTexture);
+}
+
+
+FPooledRenderTargetDesc FRCPassPostProcessAmbientOcclusion_HorizonSearch::ComputeOutputDesc(EPassOutputId InPassOutputId) const
+{
+	FPooledRenderTargetDesc Ret;
+
+	Ret = GetInput(ePId_Input0)->GetOutput()->RenderTargetDesc;
+	Ret.Reset();
+	Ret.Format = PF_B8G8R8A8;
+
+	Ret.ClearValue = FClearValueBinding::None;
+	Ret.TargetableFlags &= ~TexCreate_DepthStencilTargetable;
+	Ret.TargetableFlags |= TexCreate_UAV;
+	Ret.TargetableFlags |= TexCreate_RenderTargetable | TexCreate_ShaderResource;
+
+
+	Ret.DebugName = TEXT("GTAOInnerHorizon");
+
+	return Ret;
+}
\ No newline at end of file
diff --git a/Engine/Source/Runtime/Renderer/Private/CompositionLighting/PostProcessAmbientOcclusion.h b/Engine/Source/Runtime/Renderer/Private/CompositionLighting/PostProcessAmbientOcclusion.h
--- a/Engine/Source/Runtime/Renderer/Private/CompositionLighting/PostProcessAmbientOcclusion.h
+++ b/Engine/Source/Runtime/Renderer/Private/CompositionLighting/PostProcessAmbientOcclusion.h
@@ -23,6 +23,23 @@ enum class ESSAOType
 	EAsyncCS,
 };
 
+
+
+enum class EGTAOType
+{
+	// Not on (use legacy if at all)
+	EOff,
+
+	// non async compute shader where the Horizon Search and Inner Integrate are in separate passes
+	ESplitNonAsync,
+
+	// Async compute shader where the Horizon Search and Inner Integrate are in separate passes
+	ESplitAsync,
+
+	// Combined version which has to be run after base pass
+	ECombinedNonAsync,
+};
+
 class FSSAOHelper
 {
 public:
@@ -47,6 +64,8 @@ public:
 
 	// @return 0:off, 0..3
 	static uint32 ComputeAmbientOcclusionPassCount(const FViewInfo& View);
+
+	static EGTAOType GetGTAOPassType(const FViewInfo& View);
 };
 
 // ePId_Input0: SceneDepth
@@ -123,3 +142,52 @@ private:
 	const bool bAOSetupAsInput;
 	const bool bForceIntermediateOutput;
 };
+
+
+class FRCPassPostProcessAmbientOcclusion_HorizonSearch : public TRenderingCompositePassBase<2, 1>
+{
+public:
+
+	FRCPassPostProcessAmbientOcclusion_HorizonSearch(const FSceneView& View, const ESSAOType AOType );
+
+	// interface FRenderingCompositePass ---------
+	virtual void Process(FRenderingCompositePassContext& Context) override;
+	virtual void Release() override { delete this; }
+	virtual FPooledRenderTargetDesc ComputeOutputDesc(EPassOutputId InPassOutputId) const override;
+
+private:
+	const ESSAOType AOType;
+};
+
+
+class FRCPassPostProcessAmbientOcclusion_InnerIntegrate : public TRenderingCompositePassBase<2, 1>
+{
+public:
+
+	FRCPassPostProcessAmbientOcclusion_InnerIntegrate(const FSceneView& View, bool FinalOutput);
+
+	// interface FRenderingCompositePass ---------
+	virtual void Process(FRenderingCompositePassContext& Context) override;
+	virtual void Release() override { delete this; }
+	virtual FPooledRenderTargetDesc ComputeOutputDesc(EPassOutputId InPassOutputId) const override;
+
+private:
+	const bool bFinalOutput;
+};
+
+
+class FRCPassPostProcessAmbientOcclusion_HorizonSearchInnerIntegrate : public TRenderingCompositePassBase<2, 1>
+{
+public:
+
+	FRCPassPostProcessAmbientOcclusion_HorizonSearchInnerIntegrate(const FSceneView& View, bool FinalOutput);
+
+	// interface FRenderingCompositePass ---------
+	virtual void Process(FRenderingCompositePassContext& Context) override;
+	virtual void Release() override { delete this; }
+	virtual FPooledRenderTargetDesc ComputeOutputDesc(EPassOutputId InPassOutputId) const override;
+
+private:
+	const bool bFinalOutput;
+
+};
diff --git a/Engine/Source/Runtime/Renderer/Private/PostProcess/SceneRenderTargets.cpp b/Engine/Source/Runtime/Renderer/Private/PostProcess/SceneRenderTargets.cpp
--- a/Engine/Source/Runtime/Renderer/Private/PostProcess/SceneRenderTargets.cpp
+++ b/Engine/Source/Runtime/Renderer/Private/PostProcess/SceneRenderTargets.cpp
@@ -241,6 +241,7 @@ FSceneRenderTargets::FSceneRenderTargets(const FViewInfo& View, const FSceneRend
 	, DBufferC(GRenderTargetPool.MakeSnapshot(SnapshotSource.DBufferC))
 	, DBufferMask(GRenderTargetPool.MakeSnapshot(SnapshotSource.DBufferMask))
 	, ScreenSpaceAO(GRenderTargetPool.MakeSnapshot(SnapshotSource.ScreenSpaceAO))
+	, ScreenSpaceGTAOHorizons(GRenderTargetPool.MakeSnapshot(SnapshotSource.ScreenSpaceGTAOHorizons))
 	, QuadOverdrawBuffer(GRenderTargetPool.MakeSnapshot(SnapshotSource.QuadOverdrawBuffer))
 	, CustomDepth(GRenderTargetPool.MakeSnapshot(SnapshotSource.CustomDepth))
 	, MobileCustomStencil(GRenderTargetPool.MakeSnapshot(SnapshotSource.MobileCustomStencil))
@@ -2343,6 +2344,19 @@ void FSceneRenderTargets::AllocateDeferredShadingPathRenderTargets(FRHICommandLi
 			}
 			GRenderTargetPool.FindFreeElement(RHICmdList, Desc, ScreenSpaceAO, TEXT("ScreenSpaceAO"), true, ERenderTargetTransience::NonTransient);
 		}
+
+		// Create the screen space Horizon texture for Async GTAO
+		{
+			FPooledRenderTargetDesc Desc(FPooledRenderTargetDesc::Create2DDesc(BufferSize, PF_R8G8, FClearValueBinding::White, TexCreate_None, TexCreate_RenderTargetable, false));
+			Desc.Flags |= GFastVRamConfig.ScreenSpaceAO;
+
+			if (CurrentFeatureLevel >= ERHIFeatureLevel::SM5)
+			{
+				Desc.TargetableFlags |= TexCreate_UAV;
+			}
+			GRenderTargetPool.FindFreeElement(RHICmdList, Desc, ScreenSpaceGTAOHorizons, TEXT("ScreenSpaceGTAOHorizons"), true, ERenderTargetTransience::NonTransient);
+		}
+		
 		
 		{
 			const int32 TranslucencyLightingVolumeDim = GetTranslucencyLightingVolumeDim();
diff --git a/Engine/Source/Runtime/Renderer/Private/PostProcess/SceneRenderTargets.h b/Engine/Source/Runtime/Renderer/Private/PostProcess/SceneRenderTargets.h
--- a/Engine/Source/Runtime/Renderer/Private/PostProcess/SceneRenderTargets.h
+++ b/Engine/Source/Runtime/Renderer/Private/PostProcess/SceneRenderTargets.h
@@ -539,6 +539,7 @@ public:
 
 	// for AmbientOcclusion, only valid for a short time during the frame to allow reuse
 	TRefCountPtr<IPooledRenderTarget> ScreenSpaceAO;
+	TRefCountPtr<IPooledRenderTarget> ScreenSpaceGTAOHorizons;
 	// for shader/quad complexity, the temporary quad descriptors and complexity.
 	TRefCountPtr<IPooledRenderTarget> QuadOverdrawBuffer;
 	// used by the CustomDepth material feature, is allocated on demand or if r.CustomDepth is 2
diff --git a/Engine/Source/Runtime/Renderer/Private/RendererScene.cpp b/Engine/Source/Runtime/Renderer/Private/RendererScene.cpp
--- a/Engine/Source/Runtime/Renderer/Private/RendererScene.cpp
+++ b/Engine/Source/Runtime/Renderer/Private/RendererScene.cpp
@@ -3311,10 +3311,13 @@ bool ShouldForceFullDepthPass(EShaderPlatform ShaderPlatform)
 	static const auto StencilLODDitherCVar = IConsoleManager::Get().FindTConsoleVariableDataInt(TEXT("r.StencilForLODDither"));
 	const bool bStencilLODDither = StencilLODDitherCVar->GetValueOnAnyThread() != 0;
 
+	static const auto AOComputeCVar = IConsoleManager::Get().FindTConsoleVariableDataInt(TEXT("r.AmbientOcclusion.Compute"));
+	const bool bAOCompute		= AOComputeCVar->GetValueOnAnyThread() > 0;
+
 	const bool bEarlyZMaterialMasking = CVarEarlyZPassOnlyMaterialMasking.GetValueOnAnyThread() != 0;
 
 	// Note: ShouldForceFullDepthPass affects which static draw lists meshes go into, so nothing it depends on can change at runtime, unless you do a FGlobalComponentRecreateRenderStateContext to propagate the cvar change
-	return bDBufferAllowed || bStencilLODDither || bEarlyZMaterialMasking || IsForwardShadingEnabled(ShaderPlatform) || IsUsingSelectiveBasePassOutputs(ShaderPlatform);
+	return bAOCompute || bDBufferAllowed || bStencilLODDither || bEarlyZMaterialMasking || IsForwardShadingEnabled(ShaderPlatform) || IsUsingSelectiveBasePassOutputs(ShaderPlatform);
 }
 
 void FScene::UpdateEarlyZPassMode()
diff --git a/Engine/Source/Runtime/Renderer/Private/SystemTextures.cpp b/Engine/Source/Runtime/Renderer/Private/SystemTextures.cpp
--- a/Engine/Source/Runtime/Renderer/Private/SystemTextures.cpp
+++ b/Engine/Source/Runtime/Renderer/Private/SystemTextures.cpp
@@ -392,7 +392,36 @@ void FSystemTextures::InitializeFeatureLevelDependentTextures(FRHICommandListImm
 		}
 		RHICmdList.UnlockTexture2D((FTexture2DRHIRef&)SSAORandomization->GetRenderTargetItem().ShaderResourceTexture, 0, false);
 			}
+	
+			{
+				FPooledRenderTargetDesc Desc(FPooledRenderTargetDesc::Create2DDesc(FIntPoint(4, 4), PF_R8G8, FClearValueBinding::None, TexCreate_HideInVisualizeTexture, TexCreate_None | TexCreate_NoFastClear, false));
+				Desc.AutoWritable = false;
+				GRenderTargetPool.FindFreeElement(RHICmdList, Desc,GTAORandomization, TEXT("GTAORandomization"), true, ERenderTargetTransience::NonTransient);
+				// Write the contents of the texture.
+				uint32 DestStride;
+				uint8* DestBuffer = (uint8*)RHICmdList.LockTexture2D((FTexture2DRHIRef&)GTAORandomization->GetRenderTargetItem().ShaderResourceTexture, 0, RLM_WriteOnly, DestStride, false);
+
+				for(int32 y = 0; y < Desc.Extent.Y; ++y)
+				{
+					for(int32 x = 0; x < Desc.Extent.X; ++x)
+					{
+						uint8* Dest = (uint8*)(DestBuffer + x * sizeof(uint16) + y * DestStride);
+			
+						float Angle  = (PI/16.0f)  * ((((x + y) & 0x3) << 2) + (x & 0x3));
+						float Step	 = 0.5f + ((1.0f/8.0f) * ((y - x) & 0x3));
+
+						float ScaleCos = FMath::Cos(Angle) * Step;
+						float ScaleSin = FMath::Sin(Angle) * Step;
+			
+						Dest[0] = (uint8_t)(ScaleCos*127.5f + 127.5f);
+						Dest[1] = (uint8_t)(ScaleSin*127.5f + 127.5f);
+					}
 
+				}
+			}
+			RHICmdList.UnlockTexture2D((FTexture2DRHIRef&)GTAORandomization->GetRenderTargetItem().ShaderResourceTexture, 0, false);
+			
+			
 		{
 			// for testing, with 128x128 R8G8 we are very close to the reference (if lower res is needed we might have to add an offset to counter the 0.5f texel shift)
 			const bool bReference = false;
@@ -581,6 +610,7 @@ void FSystemTextures::ReleaseDynamicRHI()
 	PerlinNoise3D.SafeRelease();
 	SobolSampling.SafeRelease();
 	SSAORandomization.SafeRelease();
+	GTAORandomization.SafeRelease();
 	PreintegratedGF.SafeRelease();
 	LTCMat.SafeRelease();
 	LTCAmp.SafeRelease();
diff --git a/Engine/Source/Runtime/Renderer/Private/SystemTextures.h b/Engine/Source/Runtime/Renderer/Private/SystemTextures.h
--- a/Engine/Source/Runtime/Renderer/Private/SystemTextures.h
+++ b/Engine/Source/Runtime/Renderer/Private/SystemTextures.h
@@ -67,6 +67,9 @@ public:
 	TRefCountPtr<IPooledRenderTarget> SobolSampling;
 	/** SSAO randomization */
 	TRefCountPtr<IPooledRenderTarget> SSAORandomization;
+	/** GTAO randomization */
+	TRefCountPtr<IPooledRenderTarget> GTAORandomization;
+
 	/** Preintegrated GF for single sample IBL */
 	TRefCountPtr<IPooledRenderTarget> PreintegratedGF;
 	/** Linearly Transformed Cosines LUTs */
