Fix build break due to implicit conversion to bool

diff --git a/Engine/Source/Runtime/Renderer/Private/CompositionLighting/CompositionLighting.cpp b/Engine/Source/Runtime/Renderer/Private/CompositionLighting/CompositionLighting.cpp
--- a/Engine/Source/Runtime/Renderer/Private/CompositionLighting/CompositionLighting.cpp
+++ b/Engine/Source/Runtime/Renderer/Private/CompositionLighting/CompositionLighting.cpp
@@ -128,7 +128,7 @@ static FRenderingCompositeOutputRef AddPostProcessingGTAOCombined(FRHICommandLis
 {
 	FRenderingCompositePass* FinalOutputPass;
 
-	const bool bNeedSmoothingPass = CVarSSAOSmoothPass.GetValueOnRenderThread();
+	const bool bNeedSmoothingPass = CVarSSAOSmoothPass.GetValueOnRenderThread()!=0;
 
 	FRenderingCompositePass* HZBInput = Context.Graph.RegisterPass(new FRCPassPostProcessInput(const_cast<FViewInfo&>(Context.View).HZB));
 
@@ -157,7 +157,7 @@ static FRenderingCompositeOutputRef AddPostProcessingGTAOIntegration(FRHICommand
 {
 	FRenderingCompositePass* FinalOutputPass;
 
-	const bool bNeedSmoothingPass = CVarSSAOSmoothPass.GetValueOnRenderThread();
+	const bool bNeedSmoothingPass = CVarSSAOSmoothPass.GetValueOnRenderThread()!=0;
 
 	FRenderingCompositePass* AmbientOcclusionInnerIntegrate = Context.Graph.RegisterPass(new(FMemStack::Get()) FRCPassPostProcessAmbientOcclusion_InnerIntegrate(Context.View, !bNeedSmoothingPass));
 	AmbientOcclusionInnerIntegrate->SetInput(ePId_Input0, Context.SceneDepth);
