Use memoryless depth buffer on mobile platforms which support it.

Replaced the unused TexCreate_BiasNormalMap texture creation flag with a new TexCreate_Memoryless flag which can be used to request render targets which only exist in tile memory on TBDR GPUs. The mobile renderer sets this flag when it knows that the contents of the depth buffer will not be read.

Implemented support for TexCreate_Memoryless on Metal (copied from CL 8129626) and Vulkan. On Vulkan, memoryless mode is requested by adding VK_IMAGE_USAGE_TRANSIENT_ATTACHMENT_BIT to the image usage flags, and allocating the target from a memory type which supports VK_MEMORY_PROPERTY_LAZILY_ALLOCATED_BIT. If no such memory type is exposed by the device, it means it doesn't support memoryless targets, so we fall back to regular allocations.

diff --git a/Engine/Source/Runtime/Apple/MetalRHI/Private/MetalTexture.cpp b/Engine/Source/Runtime/Apple/MetalRHI/Private/MetalTexture.cpp
--- a/Engine/Source/Runtime/Apple/MetalRHI/Private/MetalTexture.cpp
+++ b/Engine/Source/Runtime/Apple/MetalRHI/Private/MetalTexture.cpp
@@ -826,6 +826,17 @@ FMetalSurface::FMetalSurface(ERHIResourceType ResourceType, EPixelFormat Format,
 #endif
 		}
 		
+#if PLATFORM_IOS
+		if (Flags & TexCreate_Memoryless)
+		{
+			ensure(Flags & (TexCreate_RenderTargetable | TexCreate_DepthStencilTargetable));
+			ensure(!(Flags & (TexCreate_CPUReadback | TexCreate_CPUWritable)));
+			ensure(!(Flags & TexCreate_UAV));
+			Desc.SetStorageMode(mtlpp::StorageMode::Memoryless);
+			Desc.SetResourceOptions((mtlpp::ResourceOptions)(mtlpp::ResourceOptions::CpuCacheModeDefaultCache | mtlpp::ResourceOptions::StorageModeMemoryless));
+		}
+#endif
+
 		static mtlpp::ResourceOptions GeneralResourceOption = FMetalCommandQueue::GetCompatibleResourceOptions(mtlpp::ResourceOptions::HazardTrackingModeUntracked);
 		Desc.SetResourceOptions((mtlpp::ResourceOptions)(Desc.GetResourceOptions() | GeneralResourceOption));
 	}
diff --git a/Engine/Source/Runtime/RHI/Public/RHIDefinitions.h b/Engine/Source/Runtime/RHI/Public/RHIDefinitions.h
--- a/Engine/Source/Runtime/RHI/Public/RHIDefinitions.h
+++ b/Engine/Source/Runtime/RHI/Public/RHIDefinitions.h
@@ -785,8 +785,8 @@ enum ETextureCreateFlags
 	TexCreate_InputAttachmentRead	= 1<<9,
 	// Disable automatic defragmentation if the initial texture memory allocation fails.
 	TexCreate_DisableAutoDefrag		= 1<<10,
-	// Create the texture with automatic -1..1 biasing
-	TexCreate_BiasNormalMap			= 1<<11,
+	// This texture has no GPU or CPU backing. It only exists in tile memory on TBDR GPUs (i.e., mobile).
+	TexCreate_Memoryless			= 1<<11,
 	// Create the texture with the flag that allows mip generation later, only applicable to D3D11
 	TexCreate_GenerateMipCapable	= 1<<12,
 	// The texture can be partially allocated in fastvram
diff --git a/Engine/Source/Runtime/Renderer/Private/MobileShadingRenderer.cpp b/Engine/Source/Runtime/Renderer/Private/MobileShadingRenderer.cpp
--- a/Engine/Source/Runtime/Renderer/Private/MobileShadingRenderer.cpp
+++ b/Engine/Source/Runtime/Renderer/Private/MobileShadingRenderer.cpp
@@ -338,14 +338,42 @@
 	FRHICommandListExecutor::GetImmediateCommandList().PollOcclusionQueries();
 	RHICmdList.ImmediateFlush(EImmediateFlushType::DispatchToRHIThread);
 
-
 	const ERHIFeatureLevel::Type ViewFeatureLevel = ViewFamily.GetFeatureLevel();
 
+	// This might eventually be a problem with multiple views.
+	// Using only view 0 to check to do on-chip transform of alpha.
+	FViewInfo& View = Views[0];
+
+	const bool bGammaSpace = !IsMobileHDR();
+
+	const FIntPoint RenderTargetSize = (ViewFamily.RenderTarget->GetRenderTargetTexture().IsValid()) ? ViewFamily.RenderTarget->GetRenderTargetTexture()->GetSizeXY() : ViewFamily.RenderTarget->GetSizeXY();
+	const bool bRequiresUpscale = ((int32)RenderTargetSize.X > FamilySize.X || (int32)RenderTargetSize.Y > FamilySize.Y);
+
+	// ES2 requires that the back buffer and depth match dimensions.
+	// For the most part this is not the case when using scene captures. Thus scene captures always render to scene color target.
+	const bool bStereoRenderingAndHMD = View.Family->EngineShowFlags.StereoRendering && View.Family->EngineShowFlags.HMDDistortion;
+	const bool bRenderToSceneColor = bStereoRenderingAndHMD || bRequiresUpscale || FSceneRenderer::ShouldCompositeEditorPrimitives(View) || View.bIsSceneCapture || View.bIsReflectionCapture;
+
+	// Whether we need to render translucency in a separate render pass
+	// On mobile it's better to render as much as possible in a single pass
+	const bool bRequiresTranslucencyPass = RequiresTranslucencyPass(RHICmdList, View);
+	// Whether we need to store depth for post-processing
+	// On PowerVR we see flickering of shadows and depths not updating correctly if targets are discarded.
+	// See CVarMobileForceDepthResolve use in ConditionalResolveSceneDepth.
+	const bool bForceDepthResolve = CVarMobileForceDepthResolve.GetValueOnRenderThread() == 1;
+	const bool bSeparateTranslucencyActive = IsMobileSeparateTranslucencyActive(View);
+	bool bKeepDepthContent = bRenderToSceneColor &&
+		(bForceDepthResolve || bSeparateTranslucencyActive || (View.bIsSceneCapture && (ViewFamily.SceneCaptureSource == ESceneCaptureSource::SCS_SceneColorHDR || ViewFamily.SceneCaptureSource == ESceneCaptureSource::SCS_SceneColorSceneDepth)));
+
+	// Whether to submit cmdbuffer with offscreen rendering before doing post-processing
+	bool bSubmitOffscreenRendering = !bGammaSpace || bRenderToSceneColor;
+
 	// Initialize global system textures (pass-through if already initialized).
 	GSystemTextures.InitializeTextures(RHICmdList, ViewFeatureLevel);
 	FSceneRenderTargets& SceneContext = FSceneRenderTargets::Get(RHICmdList);
 
 	// Allocate the maximum scene render target space for the current view family.
+	SceneContext.SetKeepDepthContent(bKeepDepthContent);
 	SceneContext.Allocate(RHICmdList, this);
 
 	//make sure all the targets we're going to use will be safely writable.
@@ -393,10 +421,6 @@
 	FRHICommandListExecutor::GetImmediateCommandList().PollOcclusionQueries();
 	RHICmdList.ImmediateFlush(EImmediateFlushType::DispatchToRHIThread);
 
-	// This might eventually be a problem with multiple views.
-	// Using only view 0 to check to do on-chip transform of alpha.
-	FViewInfo& View = Views[0];
-
 	// Default view list
 	TArray<const FViewInfo*> ViewList;
 	for (int32 ViewIndex = 0; ViewIndex < Views.Num(); ViewIndex++) 
@@ -404,8 +428,6 @@
 		ViewList.Add(&Views[ViewIndex]);
 	}
 
-	const bool bGammaSpace = !IsMobileHDR();
-	
 	// Custom depth
 	if (!bGammaSpace)
 	{
@@ -422,28 +444,6 @@
 		}
 	}
 		
-	const FIntPoint RenderTargetSize = (ViewFamily.RenderTarget->GetRenderTargetTexture().IsValid()) ? ViewFamily.RenderTarget->GetRenderTargetTexture()->GetSizeXY() : ViewFamily.RenderTarget->GetSizeXY();
-	const bool bRequiresUpscale = ((int32)RenderTargetSize.X > FamilySize.X || (int32)RenderTargetSize.Y > FamilySize.Y);
-
-	// ES2 requires that the back buffer and depth match dimensions.
-	// For the most part this is not the case when using scene captures. Thus scene captures always render to scene color target.
-	const bool bStereoRenderingAndHMD = View.Family->EngineShowFlags.StereoRendering && View.Family->EngineShowFlags.HMDDistortion;
-	const bool bRenderToSceneColor = bStereoRenderingAndHMD || bRequiresUpscale || FSceneRenderer::ShouldCompositeEditorPrimitives(View) || View.bIsSceneCapture || View.bIsReflectionCapture;
-
-	// Whether we need to render translucency in a separate render pass
-	// On mobile it's better to render as much as possible in a single pass
-	const bool bRequiresTranslucencyPass = RequiresTranslucencyPass(RHICmdList, View);
-	// Whether we need to store depth for post-processing
-	// On PowerVR we see flickering of shadows and depths not updating correctly if targets are discarded.
-	// See CVarMobileForceDepthResolve use in ConditionalResolveSceneDepth.
-	const bool bForceDepthResolve = CVarMobileForceDepthResolve.GetValueOnRenderThread() == 1;
-	const bool bSeparateTranslucencyActive = IsMobileSeparateTranslucencyActive(View);
-	bool bKeepDepthContent = bRenderToSceneColor && 
-		(bForceDepthResolve || bSeparateTranslucencyActive || (View.bIsSceneCapture && (ViewFamily.SceneCaptureSource == ESceneCaptureSource::SCS_SceneColorHDR || ViewFamily.SceneCaptureSource == ESceneCaptureSource::SCS_SceneColorSceneDepth)));
-
-	// Whether to submit cmdbuffer with offscreen rendering before doing post-processing
-	bool bSubmitOffscreenRendering = !bGammaSpace || bRenderToSceneColor;
-
 	//
 	FRHITexture* SceneColor = nullptr;
 	FRHITexture* SceneColorResolve = nullptr;
diff --git a/Engine/Source/Runtime/Renderer/Private/PostProcess/SceneRenderTargets.cpp b/Engine/Source/Runtime/Renderer/Private/PostProcess/SceneRenderTargets.cpp
--- a/Engine/Source/Runtime/Renderer/Private/PostProcess/SceneRenderTargets.cpp
+++ b/Engine/Source/Runtime/Renderer/Private/PostProcess/SceneRenderTargets.cpp
@@ -287,6 +287,7 @@ FSceneRenderTargets::FSceneRenderTargets(const FViewInfo& View, const FSceneRend
 	, DefaultDepthClear(SnapshotSource.DefaultDepthClear)
 	, QuadOverdrawIndex(SnapshotSource.QuadOverdrawIndex)
 	, bHMDAllocatedDepthTarget(SnapshotSource.bHMDAllocatedDepthTarget)
+	, bKeepDepthContent(SnapshotSource.bKeepDepthContent)
 {
 	FMemory::Memcpy(LargestDesiredSizes, SnapshotSource.LargestDesiredSizes);
 #if PREVENT_RENDERTARGET_SIZE_THRASHING
@@ -2218,6 +2219,11 @@ void FSceneRenderTargets::AllocateCommonDepthTargets(FRHICommandList& RHICmdList
 		Desc.NumSamples = GetNumSceneColorMSAASamples(CurrentFeatureLevel);
 		Desc.Flags |= GFastVRamConfig.SceneDepth;
 
+		if (!bKeepDepthContent)
+		{
+			Desc.TargetableFlags |= TexCreate_Memoryless;
+		}
+
 		// Only defer texture allocation if we're an HMD-allocated target, and we're not MSAA.
 		const bool bDeferTextureAllocation = bHMDAllocatedDepthTarget && Desc.NumSamples == 1;
 		GRenderTargetPool.FindFreeElement(RHICmdList, Desc, SceneDepthZ, TEXT("SceneDepthZ"), true, ERenderTargetTransience::Transient, bDeferTextureAllocation);
@@ -2646,6 +2652,20 @@ void FSceneRenderTargets::AllocateRenderTargets(FRHICommandListImmediate& RHICmd
 			AllocateDeferredShadingPathRenderTargets(RHICmdList, NumViews);
 		}
 	}
+	else if ((EShadingPath)CurrentShadingPath == EShadingPath::Mobile && SceneDepthZ)
+	{
+		// If the render targets are already allocated, but the keep depth content flag has changed,
+		// we need to reallocate the depth buffer.
+		uint32 DepthBufferFlags = SceneDepthZ->GetRenderTargetItem().TargetableTexture->GetFlags();
+		bool bCurrentKeepDepthContent = (DepthBufferFlags & TexCreate_Memoryless) == 0;
+		if (bCurrentKeepDepthContent != bKeepDepthContent)
+		{
+			SceneDepthZ.SafeRelease();
+			// Make sure the old depth buffer is freed by flushing the target pool.
+			GRenderTargetPool.FreeUnusedResources();
+			AllocateCommonDepthTargets(RHICmdList);
+		}
+	}
 }
 
 void FSceneRenderTargets::ReleaseSceneColor()
diff --git a/Engine/Source/Runtime/Renderer/Private/PostProcess/SceneRenderTargets.h b/Engine/Source/Runtime/Renderer/Private/PostProcess/SceneRenderTargets.h
--- a/Engine/Source/Runtime/Renderer/Private/PostProcess/SceneRenderTargets.h
+++ b/Engine/Source/Runtime/Renderer/Private/PostProcess/SceneRenderTargets.h
@@ -185,7 +185,8 @@ protected:
 		DefaultColorClear(FClearValueBinding::Black),
 		DefaultDepthClear(FClearValueBinding::DepthFar),
 		QuadOverdrawIndex(INDEX_NONE),
-		bHMDAllocatedDepthTarget(false)
+		bHMDAllocatedDepthTarget(false),
+		bKeepDepthContent(true)
 		{
 			FMemory::Memset(LargestDesiredSizes, 0);
 #if PREVENT_RENDERTARGET_SIZE_THRASHING
@@ -211,6 +212,11 @@ public:
 
 	void SetSeparateTranslucencyBufferSize(bool bAnyViewWantsDownsampledSeparateTranslucency);
 
+	void SetKeepDepthContent(bool bKeep)
+	{
+		bKeepDepthContent = bKeep;
+	}
+
 #if !(UE_BUILD_SHIPPING || UE_BUILD_TEST)
 	/** Returns the RT index where the QuadOverdrawUAV will be bound. */
 	static int32 GetQuadOverdrawUAVIndex(EShaderPlatform Platform, ERHIFeatureLevel::Type FeatureLevel);
@@ -760,6 +766,9 @@ private:
 	/** True if the depth target is allocated by an HMD plugin. This is a temporary fix to deal with HMD depth target swap chains not tracking the stencil SRV. */
 	bool bHMDAllocatedDepthTarget;
 
+	/** True if the contents of the depth buffer must be kept for post-processing. When this is false, the depth buffer can be allocated as memoryless on mobile platforms which support it. */
+	bool bKeepDepthContent;
+
 	/** CAUTION: When adding new data, make sure you copy it in the snapshot constructor! **/
 
 };
diff --git a/Engine/Source/Runtime/VulkanRHI/Private/VulkanCommandWrappers.h b/Engine/Source/Runtime/VulkanRHI/Private/VulkanCommandWrappers.h
--- a/Engine/Source/Runtime/VulkanRHI/Private/VulkanCommandWrappers.h
+++ b/Engine/Source/Runtime/VulkanRHI/Private/VulkanCommandWrappers.h
@@ -45,7 +45,7 @@ struct FWrapLayer
 	static void UnmapMemory(VkResult Result, VkDevice Device, VkDeviceMemory Memory) VULKAN_LAYER_BODY
 	static void FlushMappedMemoryRanges(VkResult Result, VkDevice Device, uint32 MemoryRangeCount, const VkMappedMemoryRange* pMemoryRanges) VULKAN_LAYER_BODY
 	static void InvalidateMappedMemoryRanges(VkResult Result, VkDevice Device, uint32 MemoryRangeCount, const VkMappedMemoryRange* pMemoryRanges) VULKAN_LAYER_BODY
-	static void GetDeviceMemoryCommitment(VkResult Result, VkDevice Device, VkDeviceMemory Memory, VkDeviceSize* pCommittedMemoryInBytes);
+	static void GetDeviceMemoryCommitment(VkResult Result, VkDevice Device, VkDeviceMemory Memory, VkDeviceSize* pCommittedMemoryInBytes) VULKAN_LAYER_BODY
 	static void BindBufferMemory(VkResult Result, VkDevice Device, VkBuffer Buffer, VkDeviceMemory Memory, VkDeviceSize MemoryOffset) VULKAN_LAYER_BODY
 	VULKAN_EXTERN_EXPORT static void BindImageMemory(VkResult Result, VkDevice Device, VkImage Image, VkDeviceMemory Memory, VkDeviceSize MemoryOffset) VULKAN_LAYER_BODY
 	static void GetBufferMemoryRequirements(VkResult Result, VkDevice Device, VkBuffer Buffer, VkMemoryRequirements* MemoryRequirements) VULKAN_LAYER_BODY
diff --git a/Engine/Source/Runtime/VulkanRHI/Private/VulkanMemory.cpp b/Engine/Source/Runtime/VulkanRHI/Private/VulkanMemory.cpp
--- a/Engine/Source/Runtime/VulkanRHI/Private/VulkanMemory.cpp
+++ b/Engine/Source/Runtime/VulkanRHI/Private/VulkanMemory.cpp
@@ -1058,7 +1058,7 @@ namespace VulkanRHI
 	}
 #endif
 
-	FOldResourceAllocation* FOldResourceHeap::AllocateResource(EType Type, uint32 Size, uint32 Alignment, bool bMapAllocation, const char* File, uint32 Line)
+	FOldResourceAllocation* FOldResourceHeap::AllocateResource(EType Type, uint32 Size, uint32 Alignment, bool bMapAllocation, bool bForceSeparateAllocation, const char* File, uint32 Line)
 	{
 		FScopeLock ScopeLock(&GOldResourceLock);
 
@@ -1074,6 +1074,7 @@ namespace VulkanRHI
 #endif
 		uint32 AllocationSize;
 		
+		if(!bForceSeparateAllocation)
 		{
 			if (Size < TargetDefaultPageSize)
 			{
@@ -1145,7 +1146,12 @@ namespace VulkanRHI
 				AllocationSize = FMath::Max(Size, TargetDefaultPageSize);
 			}
 		}
-#endif
+		else
+		{
+			// We get here when bForceSeparateAllocation is true, which is used for lazy allocations, since pooling those doesn't make sense.
+			AllocationSize = Size;
+		}
+#endif // VULKAN_SINGLE_ALLOCATION_PER_RESOURCE
 		FDeviceMemoryAllocation* DeviceMemoryAllocation = Owner->GetParent()->GetMemoryManager().Alloc(true, AllocationSize, MemoryTypeIndex, nullptr, VULKAN_MEMORY_HIGHEST_PRIORITY, File, Line);
 		if (!DeviceMemoryAllocation && Size != AllocationSize)
 		{
@@ -1246,14 +1252,6 @@ namespace VulkanRHI
 				InTypeBits >>= 1;
 			}
 
-			for (int32 Index = OutTypeIndices.Num() - 1; Index >= 1; --Index)
-			{
-				if (MemoryProperties.memoryTypes[Index].propertyFlags != MemoryProperties.memoryTypes[0].propertyFlags)
-				{
-					OutTypeIndices.RemoveAtSwap(Index, 1, false);
-				}
-			}
-
 			// No memory types matched, return failure
 			return OutTypeIndices.Num() > 0;
 		};
diff --git a/Engine/Source/Runtime/VulkanRHI/Private/VulkanRHIPrivate.h b/Engine/Source/Runtime/VulkanRHI/Private/VulkanRHIPrivate.h
--- a/Engine/Source/Runtime/VulkanRHI/Private/VulkanRHIPrivate.h
+++ b/Engine/Source/Runtime/VulkanRHI/Private/VulkanRHIPrivate.h
@@ -629,18 +629,18 @@ static inline VkAttachmentLoadOp RenderTargetLoadActionToVulkan(ERenderTargetLoa
 	return OutLoadAction;
 }
 
-static inline VkAttachmentStoreOp RenderTargetStoreActionToVulkan(ERenderTargetStoreAction InStoreAction, bool bRealRenderPass = false)
+static inline VkAttachmentStoreOp RenderTargetStoreActionToVulkan(ERenderTargetStoreAction InStoreAction)
 {
 	VkAttachmentStoreOp OutStoreAction = VK_ATTACHMENT_STORE_OP_MAX_ENUM;
 
 	switch (InStoreAction)
 	{
-	case ERenderTargetStoreAction::EStore:		OutStoreAction = VK_ATTACHMENT_STORE_OP_STORE;
+	case ERenderTargetStoreAction::EStore:
+		OutStoreAction = VK_ATTACHMENT_STORE_OP_STORE;
 		break;
-	//#todo-rco: Temp until we have fully switched to RenderPass system
 	case ERenderTargetStoreAction::ENoAction:
 	case ERenderTargetStoreAction::EMultisampleResolve:
-		OutStoreAction = bRealRenderPass ? VK_ATTACHMENT_STORE_OP_DONT_CARE : VK_ATTACHMENT_STORE_OP_STORE;
+		OutStoreAction = VK_ATTACHMENT_STORE_OP_DONT_CARE;
 		break;
 	default:
 		break;
diff --git a/Engine/Source/Runtime/VulkanRHI/Private/VulkanRenderTarget.cpp b/Engine/Source/Runtime/VulkanRHI/Private/VulkanRenderTarget.cpp
--- a/Engine/Source/Runtime/VulkanRHI/Private/VulkanRenderTarget.cpp
+++ b/Engine/Source/Runtime/VulkanRHI/Private/VulkanRenderTarget.cpp
@@ -1858,6 +1858,11 @@ FVulkanRenderTargetLayout::FVulkanRenderTargetLayout(FVulkanDevice& InDevice, co
 			CurrDesc.stencilLoadOp = VK_ATTACHMENT_LOAD_OP_DONT_CARE;
 			CurrDesc.stencilStoreOp = VK_ATTACHMENT_STORE_OP_DONT_CARE;
 
+			if (Texture->Surface.UEFlags & TexCreate_Memoryless)
+			{
+				ensure(CurrDesc.storeOp == VK_ATTACHMENT_STORE_OP_DONT_CARE);
+			}
+
 			// If the initial != final we need to change the FullHashInfo and use FinalLayout
 			CurrDesc.initialLayout = VK_IMAGE_LAYOUT_COLOR_ATTACHMENT_OPTIMAL;
 			CurrDesc.finalLayout = VK_IMAGE_LAYOUT_COLOR_ATTACHMENT_OPTIMAL;
@@ -1912,6 +1917,12 @@ FVulkanRenderTargetLayout::FVulkanRenderTargetLayout(FVulkanDevice& InDevice, co
 		{
 			CurrDesc.storeOp = RenderTargetStoreActionToVulkan(RTInfo.DepthStencilRenderTarget.DepthStoreAction);
 			CurrDesc.stencilStoreOp = RenderTargetStoreActionToVulkan(RTInfo.DepthStencilRenderTarget.GetStencilStoreAction());
+
+			if (Texture->Surface.UEFlags & TexCreate_Memoryless)
+			{
+				ensure(CurrDesc.storeOp == VK_ATTACHMENT_STORE_OP_DONT_CARE);
+				ensure(CurrDesc.stencilStoreOp == VK_ATTACHMENT_STORE_OP_DONT_CARE);
+			}
 		}
 		else
 		{
@@ -2035,10 +2046,15 @@ FVulkanRenderTargetLayout::FVulkanRenderTargetLayout(FVulkanDevice& InDevice, co
 		CurrDesc.format = UEToVkTextureFormat(ColorEntry.RenderTarget->GetFormat(), (Texture->Surface.UEFlags & TexCreate_SRGB) == TexCreate_SRGB);
 		CurrDesc.loadOp = RenderTargetLoadActionToVulkan(GetLoadAction(ColorEntry.Action));
 		bFoundClearOp = bFoundClearOp || (CurrDesc.loadOp == VK_ATTACHMENT_LOAD_OP_CLEAR);
-		CurrDesc.storeOp = RenderTargetStoreActionToVulkan(GetStoreAction(ColorEntry.Action), true);
+		CurrDesc.storeOp = RenderTargetStoreActionToVulkan(GetStoreAction(ColorEntry.Action));
 		CurrDesc.stencilLoadOp = VK_ATTACHMENT_LOAD_OP_DONT_CARE;
 		CurrDesc.stencilStoreOp = VK_ATTACHMENT_STORE_OP_DONT_CARE;
 
+		if (Texture->Surface.UEFlags & TexCreate_Memoryless)
+		{
+			ensure(CurrDesc.storeOp == VK_ATTACHMENT_STORE_OP_DONT_CARE);
+		}
+
 		// If the initial != final we need to change the FullHashInfo and use FinalLayout
 		CurrDesc.initialLayout = VK_IMAGE_LAYOUT_COLOR_ATTACHMENT_OPTIMAL;
 		CurrDesc.finalLayout = VK_IMAGE_LAYOUT_COLOR_ATTACHMENT_OPTIMAL;
@@ -2092,8 +2108,14 @@ FVulkanRenderTargetLayout::FVulkanRenderTargetLayout(FVulkanDevice& InDevice, co
 			ensure(GetStoreAction(GetStencilActions(RPInfo.DepthStencilRenderTarget.Action)) != ERenderTargetStoreAction::EMultisampleResolve);
 		}
 
-		CurrDesc.storeOp = RenderTargetStoreActionToVulkan(GetStoreAction(GetDepthActions(RPInfo.DepthStencilRenderTarget.Action)), true);
-		CurrDesc.stencilStoreOp = RenderTargetStoreActionToVulkan(GetStoreAction(GetStencilActions(RPInfo.DepthStencilRenderTarget.Action)), true);
+		CurrDesc.storeOp = RenderTargetStoreActionToVulkan(GetStoreAction(GetDepthActions(RPInfo.DepthStencilRenderTarget.Action)));
+		CurrDesc.stencilStoreOp = RenderTargetStoreActionToVulkan(GetStoreAction(GetStencilActions(RPInfo.DepthStencilRenderTarget.Action)));
+
+		if (Texture->Surface.UEFlags & TexCreate_Memoryless)
+		{
+			ensure(CurrDesc.storeOp == VK_ATTACHMENT_STORE_OP_DONT_CARE);
+			ensure(CurrDesc.stencilStoreOp == VK_ATTACHMENT_STORE_OP_DONT_CARE);
+		}
 
 		DepthStencilLayout = VulkanRHI::GetDepthStencilLayout(RPInfo.DepthStencilRenderTarget.ExclusiveDepthStencil, InDevice);
 		// If the initial != final we need to change the FullHashInfo and use FinalLayout
diff --git a/Engine/Source/Runtime/VulkanRHI/Private/VulkanTexture.cpp b/Engine/Source/Runtime/VulkanRHI/Private/VulkanTexture.cpp
--- a/Engine/Source/Runtime/VulkanRHI/Private/VulkanTexture.cpp
+++ b/Engine/Source/Runtime/VulkanRHI/Private/VulkanTexture.cpp
@@ -290,6 +290,10 @@ VkImage FVulkanSurface::CreateImage(
 		}
 		ImageCreateInfo.usage |= ((UEFlags & TexCreate_RenderTargetable) ? VK_IMAGE_USAGE_COLOR_ATTACHMENT_BIT : VK_IMAGE_USAGE_DEPTH_STENCIL_ATTACHMENT_BIT);
 		ImageCreateInfo.tiling = VK_IMAGE_TILING_OPTIMAL;
+		if ((UEFlags & TexCreate_Memoryless) == TexCreate_Memoryless)
+		{
+			ImageCreateInfo.usage |= VK_IMAGE_USAGE_TRANSIENT_ATTACHMENT_BIT;
+		}
 	}
 	else if (UEFlags & (TexCreate_DepthStencilResolveTarget))
 	{
@@ -537,13 +541,25 @@ FVulkanSurface::FVulkanSurface(FVulkanDevice& InDevice, VkImageViewType Resource
 	const bool bCPUReadback = (UEFlags & TexCreate_CPUReadback) != 0;
 	const bool bDynamic = (UEFlags & TexCreate_Dynamic) != 0;
 
-
 	VkMemoryPropertyFlags MemoryFlags = bCPUReadback ? VK_MEMORY_PROPERTY_HOST_VISIBLE_BIT : VK_MEMORY_PROPERTY_DEVICE_LOCAL_BIT;
 
+	bool bMemoryless = (UEFlags & TexCreate_Memoryless) != 0;
+	if (bMemoryless)
+	{
+		if (ensureMsgf(bRenderTarget, TEXT("Memoryless surfaces can only be used for render targets")) && ensureMsgf(!bCPUReadback, TEXT("Memoryless surfaces cannot be read back on CPU")))
+		{
+			MemoryFlags |= VK_MEMORY_PROPERTY_LAZILY_ALLOCATED_BIT;
+		}
+		else
+		{
+			bMemoryless = false;
+		}
+	}
+
 #if VULKAN_SUPPORTS_DEDICATED_ALLOCATION
 	// Per https://developer.nvidia.com/what%E2%80%99s-your-vulkan-memory-type
 	VkDeviceSize SizeToBeConsideredForDedicated = 12 * 1024 * 1024;
-	if ((bRenderTarget || MemoryRequirements.size >= SizeToBeConsideredForDedicated) && InDevice.GetOptionalExtensions().HasKHRDedicatedAllocation)
+	if ((bRenderTarget || MemoryRequirements.size >= SizeToBeConsideredForDedicated) && !bMemoryless && InDevice.GetOptionalExtensions().HasKHRDedicatedAllocation)
 	{
 		ResourceAllocation = InDevice.GetResourceHeapManager().AllocateDedicatedImageMemory(Image, MemoryRequirements, MemoryFlags, __FILE__, __LINE__);
 	}
diff --git a/Engine/Source/Runtime/VulkanRHI/Public/VulkanMemory.h b/Engine/Source/Runtime/VulkanRHI/Public/VulkanMemory.h
--- a/Engine/Source/Runtime/VulkanRHI/Public/VulkanMemory.h
+++ b/Engine/Source/Runtime/VulkanRHI/Public/VulkanMemory.h
@@ -750,7 +750,7 @@ namespace VulkanRHI
 		TArray<FOldResourceHeapPage*> FreePages;
 #endif
 
-		FOldResourceAllocation* AllocateResource(EType Type, uint32 Size, uint32 Alignment, bool bMapAllocation, const char* File, uint32 Line);
+		FOldResourceAllocation* AllocateResource(EType Type, uint32 Size, uint32 Alignment, bool bMapAllocation, bool bForceSeparateAllocation, const char* File, uint32 Line);
 
 #if VULKAN_SUPPORTS_DEDICATED_ALLOCATION
 		TArray<FOldResourceHeapPage*> UsedDedicatedImagePages;
@@ -789,19 +789,32 @@ namespace VulkanRHI
 
 		inline FOldResourceAllocation* AllocateImageMemory(const VkMemoryRequirements& MemoryReqs, VkMemoryPropertyFlags MemoryPropertyFlags, const char* File, uint32 Line)
 		{
-			uint32 TypeIndex = 0;
-			VERIFYVULKANRESULT(DeviceMemoryManager->GetMemoryTypeFromProperties(MemoryReqs.memoryTypeBits, MemoryPropertyFlags, &TypeIndex));
 			bool bMapped = (MemoryPropertyFlags & VK_MEMORY_PROPERTY_HOST_VISIBLE_BIT) == VK_MEMORY_PROPERTY_HOST_VISIBLE_BIT;
+			uint32 TypeIndex = 0;
+			if (DeviceMemoryManager->GetMemoryTypeFromProperties(MemoryReqs.memoryTypeBits, MemoryPropertyFlags, &TypeIndex) != VK_SUCCESS)
+			{
+				if ((MemoryPropertyFlags & VK_MEMORY_PROPERTY_LAZILY_ALLOCATED_BIT) == VK_MEMORY_PROPERTY_LAZILY_ALLOCATED_BIT)
+				{
+					// If lazy allocations are not supported, we can fall back to real allocations.
+					MemoryPropertyFlags = MemoryPropertyFlags & ~VK_MEMORY_PROPERTY_LAZILY_ALLOCATED_BIT;
+					VERIFYVULKANRESULT(DeviceMemoryManager->GetMemoryTypeFromProperties(MemoryReqs.memoryTypeBits, MemoryPropertyFlags, &TypeIndex));
+				}
+				else
+				{
+					UE_LOG(LogVulkanRHI, Fatal, TEXT("Cannot find memory type for MemSize %d, MemPropTypeBits %u, MemPropertyFlags %u, %s(%d)"), (uint32)MemoryReqs.size, (uint32)MemoryReqs.memoryTypeBits, (uint32)MemoryPropertyFlags, ANSI_TO_TCHAR(File), Line);
+				}
+			}
 			if (!ResourceTypeHeaps[TypeIndex])
 			{
 				UE_LOG(LogVulkanRHI, Fatal, TEXT("Missing memory type index %d, MemSize %d, MemPropTypeBits %u, MemPropertyFlags %u, %s(%d)"), TypeIndex, (uint32)MemoryReqs.size, (uint32)MemoryReqs.memoryTypeBits, (uint32)MemoryPropertyFlags, ANSI_TO_TCHAR(File), Line);
 			}
-			FOldResourceAllocation* Allocation = ResourceTypeHeaps[TypeIndex]->AllocateResource(FOldResourceHeap::EType::Image, MemoryReqs.size, MemoryReqs.alignment, bMapped, File, Line);
+			const bool bForceSeparateAllocation = (MemoryPropertyFlags & VK_MEMORY_PROPERTY_LAZILY_ALLOCATED_BIT) == VK_MEMORY_PROPERTY_LAZILY_ALLOCATED_BIT;
+			FOldResourceAllocation* Allocation = ResourceTypeHeaps[TypeIndex]->AllocateResource(FOldResourceHeap::EType::Image, MemoryReqs.size, MemoryReqs.alignment, bMapped, bForceSeparateAllocation, File, Line);
 			if (!Allocation)
 			{
 				VERIFYVULKANRESULT(DeviceMemoryManager->GetMemoryTypeFromPropertiesExcluding(MemoryReqs.memoryTypeBits, MemoryPropertyFlags, TypeIndex, &TypeIndex));
 				bMapped = (MemoryPropertyFlags & VK_MEMORY_PROPERTY_HOST_VISIBLE_BIT) == VK_MEMORY_PROPERTY_HOST_VISIBLE_BIT;
-				Allocation = ResourceTypeHeaps[TypeIndex]->AllocateResource(FOldResourceHeap::EType::Image, MemoryReqs.size, MemoryReqs.alignment, bMapped, File, Line);
+				Allocation = ResourceTypeHeaps[TypeIndex]->AllocateResource(FOldResourceHeap::EType::Image, MemoryReqs.size, MemoryReqs.alignment, bMapped, bForceSeparateAllocation, File, Line);
 			}
 			return Allocation;
 		}
@@ -841,7 +854,7 @@ namespace VulkanRHI
 				}
 			}
 
-			FOldResourceAllocation* Allocation = ResourceTypeHeaps[TypeIndex]->AllocateResource(FOldResourceHeap::EType::Buffer, MemoryReqs.size, MemoryReqs.alignment, bMapped, File, Line); //-V595
+			FOldResourceAllocation* Allocation = ResourceTypeHeaps[TypeIndex]->AllocateResource(FOldResourceHeap::EType::Buffer, MemoryReqs.size, MemoryReqs.alignment, bMapped, false, File, Line); //-V595
 			if (!Allocation)
 			{
 				// Try another memory type if the allocation failed
@@ -851,7 +864,7 @@ namespace VulkanRHI
 				{
 					UE_LOG(LogVulkanRHI, Fatal, TEXT("Missing memory type index %d, MemSize %d, MemPropTypeBits %u, MemPropertyFlags %u, %s(%d)"), TypeIndex, (uint32)MemoryReqs.size, (uint32)MemoryReqs.memoryTypeBits, (uint32)MemoryPropertyFlags, ANSI_TO_TCHAR(File), Line);
 				}
-				Allocation = ResourceTypeHeaps[TypeIndex]->AllocateResource(FOldResourceHeap::EType::Buffer, MemoryReqs.size, MemoryReqs.alignment, bMapped, File, Line);
+				Allocation = ResourceTypeHeaps[TypeIndex]->AllocateResource(FOldResourceHeap::EType::Buffer, MemoryReqs.size, MemoryReqs.alignment, bMapped, false, File, Line);
 			}
 			return Allocation;
 		}
