- Added bExcludeFromLightAttachmentGroup flag to exclude a Primitive component from a light attachment group.
- Made bExcludeFromLightAttachmentGroup default to 'true' for particle system components.
- Fixed shadow caster filtering such that proxies that are in a light attachment group check the root CastsInsetShadows(), not just the current proxy.

diff --git a/Engine/Source/Runtime/Engine/Classes/Components/PrimitiveComponent.h b/Engine/Source/Runtime/Engine/Classes/Components/PrimitiveComponent.h
--- a/Engine/Source/Runtime/Engine/Classes/Components/PrimitiveComponent.h
+++ b/Engine/Source/Runtime/Engine/Classes/Components/PrimitiveComponent.h
@@ -447,6 +447,12 @@ public:
 	UPROPERTY(EditAnywhere, AdvancedDisplay, BlueprintReadOnly, Category=Lighting)
 	uint8 bLightAttachmentsAsGroup:1;
 
+	/** 
+	 * If set, then it overrides any bLightAttachmentsAsGroup set in a parent.
+	 */
+	UPROPERTY(EditAnywhere, AdvancedDisplay, BlueprintReadOnly, Category=Lighting)
+	uint8 bExcludeFromLightAttachmentGroup :1;
+
 	/**
 	* Mobile only:
 	* If disabled this component will not receive CSM shadows. (Components that do not receive CSM may have reduced shading cost)
@@ -1437,6 +1443,10 @@ public:
 	UFUNCTION(BlueprintCallable, Category="Rendering")
 	void SetLightAttachmentsAsGroup(UPARAM(DisplayName="LightAttachmentsAsGroup") bool bInLightAttachmentsAsGroup);
 
+	/** Changes the value of ExcludeFromLightAttachmentGroup. */
+	UFUNCTION(BlueprintCallable, Category="Rendering")
+	void SetExcludeFromLightAttachmentGroup(UPARAM(DisplayName = "ExcludeFromLightAttachmentGroup") bool bInExcludeFromLightAttachmentGroup);
+
 	/** Changes the value of bSingleSampleShadowFromStationaryLights. */
 	UFUNCTION(BlueprintCallable, Category="Rendering")
 	void SetSingleSampleShadowFromStationaryLights(bool bNewSingleSampleShadowFromStationaryLights);
diff --git a/Engine/Source/Runtime/Engine/Private/Components/PrimitiveComponent.cpp b/Engine/Source/Runtime/Engine/Private/Components/PrimitiveComponent.cpp
--- a/Engine/Source/Runtime/Engine/Private/Components/PrimitiveComponent.cpp
+++ b/Engine/Source/Runtime/Engine/Private/Components/PrimitiveComponent.cpp
@@ -1429,6 +1429,15 @@ void UPrimitiveComponent::SetLightAttachmentsAsGroup(bool bInLightAttachmentsAsG
 	}
 }
 
+void UPrimitiveComponent::SetExcludeFromLightAttachmentGroup(bool bInExcludeFromLightAttachmentGroup)
+{
+	if (bExcludeFromLightAttachmentGroup != bInExcludeFromLightAttachmentGroup)
+	{
+		bExcludeFromLightAttachmentGroup = bInExcludeFromLightAttachmentGroup;
+		MarkRenderStateDirty();
+	}
+}
+
 void UPrimitiveComponent::SetSingleSampleShadowFromStationaryLights(bool bNewSingleSampleShadowFromStationaryLights)
 {
 	if (bNewSingleSampleShadowFromStationaryLights != bSingleSampleShadowFromStationaryLights)
@@ -3255,6 +3264,12 @@ bool UPrimitiveComponent::ComponentOverlapMultiImpl(TArray<struct FOverlapResult
 
 const UPrimitiveComponent* UPrimitiveComponent::GetLightingAttachmentRoot() const
 {
+	// Exclude  from light attachment group whatever the parent says
+	if (bExcludeFromLightAttachmentGroup)
+	{
+		return nullptr;
+	}
+
 	const USceneComponent* CurrentHead = this;
 
 	while (CurrentHead)
diff --git a/Engine/Source/Runtime/Engine/Private/Particles/ParticleComponents.cpp b/Engine/Source/Runtime/Engine/Private/Particles/ParticleComponents.cpp
--- a/Engine/Source/Runtime/Engine/Private/Particles/ParticleComponents.cpp
+++ b/Engine/Source/Runtime/Engine/Private/Particles/ParticleComponents.cpp
@@ -3397,6 +3397,8 @@ UParticleSystemComponent::UParticleSystemComponent(const FObjectInitializer& Obj
 	ManagerHandle = INDEX_NONE;
 	bPendingManagerAdd = false;
 	bPendingManagerRemove = false;
+
+	bExcludeFromLightAttachmentGroup = true;
 }
 
 void UParticleSystemComponent::SetRequiredSignificance(EParticleSignificanceLevel NewRequiredSignificance)
diff --git a/Engine/Source/Runtime/Renderer/Private/ShadowSetup.cpp b/Engine/Source/Runtime/Renderer/Private/ShadowSetup.cpp
--- a/Engine/Source/Runtime/Renderer/Private/ShadowSetup.cpp
+++ b/Engine/Source/Runtime/Renderer/Private/ShadowSetup.cpp
@@ -3076,6 +3076,14 @@ struct FGatherShadowPrimitivesPacket
 					bScreenSpaceSizeCulled = FMath::Square(PrimitiveBounds.SphereRadius) < FMath::Square(MinScreenRadiusForShadowCaster) * DistanceSquared * ProjectedShadowInfo->DependentView->LODDistanceFactorSquared;
 				}
 
+				bool bCastsInsetShadows = PrimitiveProxy->CastsInsetShadow();
+				// If light attachment root is valid, we're in a group and need to get the flag from the root.
+				if (PrimitiveSceneInfo->LightingAttachmentRoot.IsValid())
+				{
+					FAttachmentGroupSceneInfo& AttachmentGroup = PrimitiveSceneInfo->Scene->AttachmentGroups.FindChecked(PrimitiveSceneInfo->LightingAttachmentRoot);
+					bCastsInsetShadows = AttachmentGroup.ParentSceneInfo && AttachmentGroup.ParentSceneInfo->Proxy->CastsInsetShadow();
+				}
+
 				if (!bScreenSpaceSizeCulled
 					&& ProjectedShadowInfo->GetLightSceneInfoCompact().AffectsPrimitive(PrimitiveBounds, PrimitiveProxy)
 					// Include all primitives for movable lights, but only statically shadowed primitives from a light with static shadowing,
@@ -3084,7 +3092,7 @@ struct FGatherShadowPrimitivesPacket
 					// Only render primitives into a reflective shadowmap that are supposed to affect indirect lighting
 					&& !(ProjectedShadowInfo->bReflectiveShadowmap && !PrimitiveProxy->AffectsDynamicIndirectLighting())
 					// Exclude primitives that will create their own per-object shadow, except when rendering RSMs
-					&& (!PrimitiveProxy->CastsInsetShadow() || ProjectedShadowInfo->bReflectiveShadowmap)
+					&& (!bCastsInsetShadows || ProjectedShadowInfo->bReflectiveShadowmap)
 					// Exclude primitives that will create a per-object shadow from a stationary light
 					&& !ShouldCreateObjectShadowForStationaryLight(&LightSceneInfo, PrimitiveProxy, true)
 					// Only render shadows from objects that use static lighting during a reflection capture, since the reflection capture doesn't update at runtime
