Cherry-picking CL 7445043

Fixes shadow denoiser batching.

diff --git a/Engine/Source/Runtime/Renderer/Private/LightRendering.cpp b/Engine/Source/Runtime/Renderer/Private/LightRendering.cpp
--- a/Engine/Source/Runtime/Renderer/Private/LightRendering.cpp
+++ b/Engine/Source/Runtime/Renderer/Private/LightRendering.cpp
@@ -960,6 +960,8 @@ void FDeferredShadingSceneRenderer::RenderLights(FRHICommandListImmediate& RHICm
 						View.ViewState->PrevFrameViewInfo.ShadowHistories.Empty();
 						View.ViewState->PrevFrameViewInfo.ShadowHistories.Reserve(SortedLights.Num());
 					}
+
+					PreprocessedShadowMaskTextures.SetNum(SortedLights.Num());
 				}
 			} // if (RHI_RAYTRACING)
 
@@ -1012,7 +1014,9 @@ void FDeferredShadingSceneRenderer::RenderLights(FRHICommandListImmediate& RHICm
 							SortedLightInfo.SortKey.Fields.bShadowed && !ShouldRenderRayTracingStochasticRectLight(LightSceneInfo);
 
 						// determine if this light doesn't yet have a precomuted shadow and execute a batch to amortize costs if one is needed
-						if (RHI_RAYTRACING && PreprocessedShadowMaskTextures.Num() > 0 && bWantsBatchedShadow && !PreprocessedShadowMaskTextures[LightIndex - AttenuationLightStart])
+						if (RHI_RAYTRACING &&
+							bWantsBatchedShadow &&
+							!PreprocessedShadowMaskTextures[LightIndex - AttenuationLightStart])
 						{
 							SCOPED_DRAW_EVENT(RHICmdList, ShadowBatch);
 							TStaticArray<IScreenSpaceDenoiser::FShadowParameters, IScreenSpaceDenoiser::kMaxBatchSize> DenoisingQueue;
