Setup runtime virtual world height render to correctly mark render target output type to PF_G16

diff --git a/Engine/Source/Runtime/Renderer/Private/VT/RuntimeVirtualTextureRender.cpp b/Engine/Source/Runtime/Renderer/Private/VT/RuntimeVirtualTextureRender.cpp
--- a/Engine/Source/Runtime/Renderer/Private/VT/RuntimeVirtualTextureRender.cpp
+++ b/Engine/Source/Runtime/Renderer/Private/VT/RuntimeVirtualTextureRender.cpp
@@ -120,6 +120,7 @@ namespace RuntimeVirtualTexture
 		static void ModifyCompilationEnvironment(FShaderCompilerEnvironment& OutEnvironment)
 		{
 			OutEnvironment.SetDefine(TEXT("OUT_WORLDHEIGHT"), 1);
+			OutEnvironment.SetRenderTargetOutputFormat(0, PF_G16);
 		}
 
 		static FRHIBlendState* GetBlendState()
