Fix GPU performance regression with shadowed lights introduced in CL 8188641
- Add r.AllowClearLightSceneExtentsOnly cvar (default to 1). If this is enabled, allow clearing lights using a quad instead of a full clear. This is much faster for smaller lights

diff --git a/Engine/Source/Runtime/Renderer/Private/LightRendering.cpp b/Engine/Source/Runtime/Renderer/Private/LightRendering.cpp
--- a/Engine/Source/Runtime/Renderer/Private/LightRendering.cpp
+++ b/Engine/Source/Runtime/Renderer/Private/LightRendering.cpp
@@ -81,6 +81,11 @@ static TAutoConsoleVariable<int32> CVarMaxShadowRayTracingBatchSize(
 	TEXT("Maximum number of shadows to trace at the same time."),
 	ECVF_RenderThreadSafe);
 
+static TAutoConsoleVariable<int32> CVarAllowClearLightSceneExtentsOnly(
+	TEXT("r.AllowClearLightSceneExtentsOnly"), 1,
+	TEXT(""),
+	ECVF_RenderThreadSafe);
+
 
 #if ENABLE_DEBUG_DISCARD_PROP
 static float GDebugLightDiscardProp = 0.0f;
@@ -1252,13 +1257,48 @@ void FDeferredShadingSceneRenderer::RenderLights(FRHICommandListImmediate& RHICm
 							View.HeightfieldLightingViewInfo.ClearShadowing(View, RHICmdList, LightSceneInfo);
 						}
 					
+						// Clear light attenuation for local lights with a quad covering their extents
+						const bool bClearLightScreenExtentsOnly = CVarAllowClearLightSceneExtentsOnly.GetValueOnRenderThread() && SortedLightInfo.SortKey.Fields.LightType != LightType_Directional;
+						// All shadows render with min blending
+						bool bClearToWhite = !bClearLightScreenExtentsOnly;
+
 						FRHIRenderPassInfo RPInfo(ScreenShadowMaskTexture->GetRenderTargetItem().TargetableTexture, ERenderTargetActions::Load_Store);
 						RPInfo.DepthStencilRenderTarget.Action = MakeDepthStencilTargetActions(ERenderTargetActions::Load_DontStore, ERenderTargetActions::Load_Store);
 						RPInfo.DepthStencilRenderTarget.DepthStencilTarget = SceneContext.GetSceneDepthSurface();
 						RPInfo.DepthStencilRenderTarget.ExclusiveDepthStencil = FExclusiveDepthStencil::DepthRead_StencilWrite;
-						RPInfo.ColorRenderTargets[0].Action = ERenderTargetActions::Clear_Store;
+						if (bClearToWhite)
+						{
+							RPInfo.ColorRenderTargets[0].Action = ERenderTargetActions::Clear_Store;
+						}
+
 						TransitionRenderPassTargets(RHICmdList, RPInfo);
 						RHICmdList.BeginRenderPass(RPInfo, TEXT("ClearScreenShadowMask"));
+						if (bClearLightScreenExtentsOnly)
+						{
+							SCOPED_DRAW_EVENT(RHICmdList, ClearQuad);
+
+							for (int32 ViewIndex = 0; ViewIndex < Views.Num(); ViewIndex++)
+							{
+								const FViewInfo& View = Views[ViewIndex];
+								FIntRect ScissorRect;
+
+								if (!LightSceneInfo.Proxy->GetScissorRect(ScissorRect, View, View.ViewRect))
+								{
+									ScissorRect = View.ViewRect;
+								}
+
+								if (ScissorRect.Min.X < ScissorRect.Max.X && ScissorRect.Min.Y < ScissorRect.Max.Y)
+								{
+									RHICmdList.SetViewport(ScissorRect.Min.X, ScissorRect.Min.Y, 0.0f, ScissorRect.Max.X, ScissorRect.Max.Y, 1.0f);
+									DrawClearQuad(RHICmdList, true, FLinearColor(1, 1, 1, 1), false, 0, false, 0);
+								}
+								else
+								{
+									LightSceneInfo.Proxy->GetScissorRect(ScissorRect, View, View.ViewRect);
+								}
+							}
+						}
+
 						RHICmdList.EndRenderPass();
 
 						RenderShadowProjections(RHICmdList, &LightSceneInfo, ScreenShadowMaskTexture, bInjectedTranslucentVolume);
