Remove check() to unblock production while I investigate root cause
Issue is that a RVT rendering proxy has no mesh batches created which render to RVT
Should be safe to ignore check() since no mesh batches mean no further rendering to the RVT downstream

diff --git a/Engine/Source/Runtime/Renderer/Private/PrimitiveSceneInfo.cpp b/Engine/Source/Runtime/Renderer/Private/PrimitiveSceneInfo.cpp
--- a/Engine/Source/Runtime/Renderer/Private/PrimitiveSceneInfo.cpp
+++ b/Engine/Source/Runtime/Renderer/Private/PrimitiveSceneInfo.cpp
@@ -463,7 +463,7 @@ static void GetRuntimeVirtualTextureLODRange(TArray<class FStaticMeshBatchReleva
 		}
 	}
 
-	check(OutMinLOD <= OutMaxLOD);
+	//check(OutMinLOD <= OutMaxLOD);
 }
 
 int32 FPrimitiveSceneInfo::UpdateStaticLightingBuffer()
