Fix some IncludeTool warnings.

#rnx

diff --git a/Engine/Plugins/Experimental/GeometryProcessing/Source/GeometryAlgorithms/Private/GteUtil.h b/Engine/Plugins/Experimental/GeometryProcessing/Source/GeometryAlgorithms/Private/GteUtil.h
--- a/Engine/Plugins/Experimental/GeometryProcessing/Source/GeometryAlgorithms/Private/GteUtil.h
+++ b/Engine/Plugins/Experimental/GeometryProcessing/Source/GeometryAlgorithms/Private/GteUtil.h
@@ -1,5 +1,7 @@
 // Copyright 1998-2019 Epic Games, Inc. All Rights Reserved.
 
+#pragma once
+
 #include "VectorTypes.h"
 #include "LineTypes.h"
 #include "CircleTypes.h"
diff --git a/Engine/Source/Runtime/Core/Private/FramePro/FramePro.cpp b/Engine/Source/Runtime/Core/Private/FramePro/FramePro.cpp
--- a/Engine/Source/Runtime/Core/Private/FramePro/FramePro.cpp
+++ b/Engine/Source/Runtime/Core/Private/FramePro/FramePro.cpp
@@ -8430,7 +8430,7 @@ namespace FramePro
 		#include <winsock2.h>
 		#include <ws2tcpip.h>
 		#if FRAMEPRO_PLATFORM_UE4
-			#include "Windows//HideWindowsPlatformTypes.h"
+			#include "Windows/HideWindowsPlatformTypes.h"
 		#endif
 	#endif
 
diff --git a/Engine/Source/Runtime/Engine/Private/PhysicsEngine/Experimental/ChaosDerivedDataUtil.h b/Engine/Source/Runtime/Engine/Private/PhysicsEngine/Experimental/ChaosDerivedDataUtil.h
--- a/Engine/Source/Runtime/Engine/Private/PhysicsEngine/Experimental/ChaosDerivedDataUtil.h
+++ b/Engine/Source/Runtime/Engine/Private/PhysicsEngine/Experimental/ChaosDerivedDataUtil.h
@@ -1,5 +1,7 @@
 // Copyright 1998-2019 Epic Games, Inc. All Rights Reserved.
 
+#pragma once
+
 #include "CoreMinimal.h"
 
 #if INCLUDE_CHAOS
diff --git a/Engine/Source/Runtime/Engine/Private/TimerManagerTests.cpp b/Engine/Source/Runtime/Engine/Private/TimerManagerTests.cpp
--- a/Engine/Source/Runtime/Engine/Private/TimerManagerTests.cpp
+++ b/Engine/Source/Runtime/Engine/Private/TimerManagerTests.cpp
@@ -26,9 +26,8 @@ void TimerTest_TickWorld(UWorld* World, float Time = 1.f)
 	}
 }
 
-class FDummy
+struct FDummy
 {
-public:
 	FDummy() { Count = 0; }
 
 	void Callback() { ++Count; }
diff --git a/Engine/Source/Runtime/Renderer/Private/GpuDebugRendering.cpp b/Engine/Source/Runtime/Renderer/Private/GpuDebugRendering.cpp
--- a/Engine/Source/Runtime/Renderer/Private/GpuDebugRendering.cpp
+++ b/Engine/Source/Runtime/Renderer/Private/GpuDebugRendering.cpp
@@ -41,7 +41,7 @@ FGpuDebugPrimitiveBuffers AllocateGpuDebugPrimitiveBuffers(FRHICommandListImmedi
 	return Output;
 }
 
-struct FLine { FVector Start; FVector End; FLinearColor Color; };
+class FLine { public: FVector Start; FVector End; FLinearColor Color; };
 typedef TArray<FLine> FLineArray;
 static FLineArray ReadGPUDebugPrimitives(FRHICommandListImmediate& RHICmdList, FGpuDebugPrimitiveBuffers& DebugPrimitiveBuffer)
 {
