Merge: Fixed motion blur length being tied to target extent rather than viewport size.

diff --git a/Engine/Source/Runtime/Renderer/Private/PostProcess/PostProcessMotionBlur.cpp b/Engine/Source/Runtime/Renderer/Private/PostProcess/PostProcessMotionBlur.cpp
--- a/Engine/Source/Runtime/Renderer/Private/PostProcess/PostProcessMotionBlur.cpp
+++ b/Engine/Source/Runtime/Renderer/Private/PostProcess/PostProcessMotionBlur.cpp
@@ -116,13 +116,13 @@ bool IsVisualizeMotionBlurEnabled(const FViewInfo& View)
 bool IsMotionBlurScatterRequired(const FViewInfo& View, const FScreenPassTextureViewport& SceneViewport)
 {
 	const FSceneViewState* ViewState = View.ViewState;
-	const float SceneExtentX = SceneViewport.Extent.X;
+	const float ViewportWidth = SceneViewport.Rect.Width();
 
 	// Normalize percentage value.
 	const float VelocityMax = View.FinalPostProcessSettings.MotionBlurMax / 100.0f;
 
 	// Scale by 0.5 due to blur samples going both ways and convert to tiles.
-	const float VelocityMaxInTiles = VelocityMax * SceneExtentX * (0.5f / 16.0f);
+	const float VelocityMaxInTiles = VelocityMax * ViewportWidth * (0.5f / 16.0f);
 
 	// Compute path only supports the immediate neighborhood of tiles.
 	const float TileDistanceMaxGathered = 3.0f;
@@ -139,10 +139,10 @@ bool IsMotionBlurScatterRequired(const FViewInfo& View, const FScreenPassTexture
 	return bIsScatterRequiredByUser || bIsScatterRequiredByVelocityLength;
 }
 
-FIntPoint GetMotionBlurTileCount(FIntPoint PixelExtent)
+FIntPoint GetMotionBlurTileCount(FIntPoint SizeInPixels)
 {
-	const uint32 TilesX = FMath::DivideAndRoundUp(PixelExtent.X, kMotionBlurTileSize);
-	const uint32 TilesY = FMath::DivideAndRoundUp(PixelExtent.Y, kMotionBlurTileSize);
+	const uint32 TilesX = FMath::DivideAndRoundUp(SizeInPixels.X, kMotionBlurTileSize);
+	const uint32 TilesY = FMath::DivideAndRoundUp(SizeInPixels.Y, kMotionBlurTileSize);
 	return FIntPoint(TilesX, TilesY);
 }
 
@@ -185,24 +185,24 @@ BEGIN_SHADER_PARAMETER_STRUCT(FMotionBlurParameters, )
 	SHADER_PARAMETER(float, VelocityMax)
 END_SHADER_PARAMETER_STRUCT()
 
-FMotionBlurParameters GetMotionBlurParameters(const FViewInfo& View, FIntPoint SceneViewportExtent, float BlurScale)
+FMotionBlurParameters GetMotionBlurParameters(const FViewInfo& View, FIntPoint SceneViewportSize, float BlurScale)
 {
 	const FSceneViewState* ViewState = View.ViewState;
 
 	const float TileSize = kMotionBlurTileSize;
-	const float SceneViewportExtentX = SceneViewportExtent.X;
-	const float SceneViewportExtentY = SceneViewportExtent.Y;
+	const float SceneViewportSizeX = SceneViewportSize.X;
+	const float SceneViewportSizeY = SceneViewportSize.Y;
 	const float MotionBlurTimeScale = ViewState ? ViewState->MotionBlurTimeScale : 1.0f;
 
 	// Scale by 0.5 due to blur samples going both ways.
 	const float VelocityScale = MotionBlurTimeScale * View.FinalPostProcessSettings.MotionBlurAmount * 0.5f;
-	const float VelocityUVToPixel = BlurScale * SceneViewportExtentX * 0.5f;
+	const float VelocityUVToPixel = BlurScale * SceneViewportSizeX * 0.5f;
 
 	// 0:no 1:full screen width, percent conversion
 	const float UVVelocityMax = View.FinalPostProcessSettings.MotionBlurMax / 100.0f;
 
 	FMotionBlurParameters MotionBlurParameters;
-	MotionBlurParameters.AspectRatio = SceneViewportExtentY / SceneViewportExtentX;
+	MotionBlurParameters.AspectRatio = SceneViewportSizeY / SceneViewportSizeX;
 	MotionBlurParameters.VelocityScale = VelocityUVToPixel * VelocityScale;
 	MotionBlurParameters.VelocityScaleForTiles = MotionBlurParameters.VelocityScale / TileSize;
 	MotionBlurParameters.VelocityMax = FMath::Abs(VelocityUVToPixel) * UVVelocityMax;
@@ -477,7 +477,7 @@ void ComputeMotionBlurVelocity(
 				false),
 			TEXT("VelocityTile"));
 
-	const FMotionBlurParameters MotionBlurParametersNoScale = GetMotionBlurParameters(View, Viewports.Color.Extent, 1.0f);
+	const FMotionBlurParameters MotionBlurParametersNoScale = GetMotionBlurParameters(View, Viewports.Color.Rect.Size(), 1.0f);
 
 	// Velocity flatten pass: combines depth / velocity into a single target for sampling efficiency.
 	{
@@ -496,7 +496,7 @@ void ComputeMotionBlurVelocity(
 			RDG_EVENT_NAME("Velocity Flatten"),
 			*ComputeShader,
 			PassParameters,
-			FComputeShaderUtils::GetGroupCount(Viewports.Velocity.Extent, FComputeShaderUtils::kGolden2DGroupSize));
+			FComputeShaderUtils::GetGroupCount(Viewports.Velocity.Rect.Size(), FComputeShaderUtils::kGolden2DGroupSize));
 	}
 
 	bool ScatterDilatation = IsMotionBlurScatterRequired(View, Viewports.Color);
@@ -671,7 +671,7 @@ FRDGTextureRef AddMotionBlurFilterPass(
 	const float BlurScale = BlurScaleLUT[static_cast<uint32>(MotionBlurFilterPass)][static_cast<uint32>(MotionBlurQuality)];
 
 	FMotionBlurFilterParameters MotionBlurFilterParameters;
-	MotionBlurFilterParameters.MotionBlur = GetMotionBlurParameters(View, Viewports.Color.Extent, BlurScale);
+	MotionBlurFilterParameters.MotionBlur = GetMotionBlurParameters(View, Viewports.Color.Rect.Size(), BlurScale);
 	MotionBlurFilterParameters.Color = Viewports.ColorParameters;
 	MotionBlurFilterParameters.Velocity = Viewports.VelocityParameters;
 	MotionBlurFilterParameters.VelocityTile = Viewports.VelocityTileParameters;
@@ -711,7 +711,7 @@ FRDGTextureRef AddMotionBlurFilterPass(
 			RDG_EVENT_NAME("Motion Blur %dx%d (CS)", Viewports.Color.Rect.Width(), Viewports.Color.Rect.Height()),
 			*ComputeShader,
 			PassParameters,
-			FComputeShaderUtils::GetGroupCount(Viewports.Color.Extent, FComputeShaderUtils::kGolden2DGroupSize));
+			FComputeShaderUtils::GetGroupCount(Viewports.Color.Rect.Size(), FComputeShaderUtils::kGolden2DGroupSize));
 	}
 	else
 	{
