Decouple the concepts of VT texture layer and VT page table layer.
This allows us to put texture layers into physical groups that will be accessed through the same page table layer (and so must share the same UV in different physical textures).
Physical spaces can now contain multiple texture layers (one for each member of a physical group).
Changed the logic for the physical space pool descriptions so that we can support multiple texture layers. And also made the description work for Min and Max tile size values so that we don't need to specify pool sizes multiple times for each possible tile size.
Runtime virtual texture can select if it wants to use a single physical group. By default this is on so that RVT uses a single page table channel.

diff --git a/Engine/Config/BaseEngine.ini b/Engine/Config/BaseEngine.ini
--- a/Engine/Config/BaseEngine.ini
+++ b/Engine/Config/BaseEngine.ini
@@ -2354,14 +2354,6 @@ DefaultChaosSolverActorClass=/Script/ChaosSolverEngine.ChaosSolverActor
 ; Configure VT physical memory pools. Usefull commands to set these
 ; "stat virtualtexturing" show dynamic use of the VT system including the cache loads
 ; "r.VT.ListPhysicalPool" list in-depth details on the allocated physical pools
-; Size in megabyte of any pools that are created and not explicitly defined below
+; The upper allocation limit for a pool is caluclated by matching a size and format from the list below or falling back to the default
 DefaultSizeInMegabyte=64
-; Explicit per pool configurations
-+Pools=(SizeInMegabyte=128, TileSize=136, Format=PF_DXT1)
-+Pools=(SizeInMegabyte=64, TileSize=136, Format=PF_DXT5)
-+Pools=(SizeInMegabyte=128, TileSize=136, Format=PF_BC5)
-+Pools=(SizeInMegabyte=32, TileSize=136, Format=PF_FloatRGBA)
-+Pools=(SizeInMegabyte=32, TileSize=136, Format=PF_B8G8R8A8)
-
-
-
++Pools=(SizeInMegabyte=64, MinTileSize=0, MaxTileSize=9999, Formats=(PF_DXT1, PF_DXT5))
diff --git a/Engine/Source/Editor/MaterialEditor/Private/MaterialInstanceEditor.cpp b/Engine/Source/Editor/MaterialEditor/Private/MaterialInstanceEditor.cpp
--- a/Engine/Source/Editor/MaterialEditor/Private/MaterialInstanceEditor.cpp
+++ b/Engine/Source/Editor/MaterialEditor/Private/MaterialInstanceEditor.cpp
@@ -1393,6 +1393,20 @@ void FMaterialInstanceEditor::DrawSamplerWarningStrings(FCanvas* Canvas, int32&
 
 								DrawPositionY += SpacingBetweenLines;
 							}
+							if (Expression->bSinglePhysicalSpace != RuntimeVirtualTexture->GetSinglePhysicalSpace())
+							{
+								Canvas->DrawShadowedString(
+									5, DrawPositionY,
+									*FString::Printf(TEXT("Warning: '%s' interprets the virtual texture page table packing as '%d' not '%d'"),
+										*RuntimeVirtualTextureParameterValue->ParameterInfo.Name.ToString(),
+										RuntimeVirtualTexture->GetSinglePhysicalSpace() ? 1 : 0,
+										Expression->bSinglePhysicalSpace ? 1 : 0,
+										*RuntimeVirtualTexture->GetPathName()),
+									FontToUse,
+									FLinearColor(1, 1, 0));
+
+								DrawPositionY += SpacingBetweenLines;
+							}
 						}
 					}
 				}
diff --git a/Engine/Source/Editor/TextureEditor/Private/Models/TextureEditorViewportClient.cpp b/Engine/Source/Editor/TextureEditor/Private/Models/TextureEditorViewportClient.cpp
--- a/Engine/Source/Editor/TextureEditor/Private/Models/TextureEditorViewportClient.cpp
+++ b/Engine/Source/Editor/TextureEditor/Private/Models/TextureEditorViewportClient.cpp
@@ -121,23 +121,24 @@ void FTextureEditorViewportClient::Draw(FViewport* Viewport, FCanvas* Canvas)
 		{
 			bool bIsNormalMap = Texture2D->IsNormalMap();
 			bool bIsSingleChannel = Texture2D->CompressionSettings == TC_Grayscale || Texture2D->CompressionSettings == TC_Alpha;
+			bool bSingleVTPhysicalSpace = Texture2D->IsVirtualTexturedWithSinglePhysicalSpace();
 			bIsVirtualTexture = Texture2D->IsCurrentlyVirtualTextured();
-			BatchedElementParameters = new FBatchedElementTexture2DPreviewParameters(MipLevel, LayerIndex, bIsNormalMap, bIsSingleChannel, bIsVirtualTexture);
+			BatchedElementParameters = new FBatchedElementTexture2DPreviewParameters(MipLevel, LayerIndex, bIsNormalMap, bIsSingleChannel, bSingleVTPhysicalSpace, bIsVirtualTexture, false);
 		}
 		else if (Texture2DArray) 
 		{
 			bool bIsNormalMap = Texture2DArray->IsNormalMap();
 			bool bIsSingleChannel = Texture2DArray->CompressionSettings == TC_Grayscale || Texture2DArray->CompressionSettings == TC_Alpha;
-			BatchedElementParameters = new FBatchedElementTexture2DPreviewParameters(MipLevel, LayerIndex, bIsNormalMap, bIsSingleChannel, false, true);
+			BatchedElementParameters = new FBatchedElementTexture2DPreviewParameters(MipLevel, LayerIndex, bIsNormalMap, bIsSingleChannel, false, false, true);
 		}
 		else if (TextureRT2D)
 		{
-			BatchedElementParameters = new FBatchedElementTexture2DPreviewParameters(MipLevel, LayerIndex, false, false, false);
+			BatchedElementParameters = new FBatchedElementTexture2DPreviewParameters(MipLevel, LayerIndex, false, false, false, false, false);
 		}
 		else
 		{
 			// Default to treating any UTexture derivative as a 2D texture resource
-			BatchedElementParameters = new FBatchedElementTexture2DPreviewParameters(MipLevel, LayerIndex, false, false, false);
+			BatchedElementParameters = new FBatchedElementTexture2DPreviewParameters(MipLevel, LayerIndex, false, false, false, false, false);
 		}
 	}
 
diff --git a/Engine/Source/Editor/UnrealEd/Private/Texture2DPreview.cpp b/Engine/Source/Editor/UnrealEd/Private/Texture2DPreview.cpp
--- a/Engine/Source/Editor/UnrealEd/Private/Texture2DPreview.cpp
+++ b/Engine/Source/Editor/UnrealEd/Private/Texture2DPreview.cpp
@@ -63,7 +63,7 @@ public:
 		return IsFeatureLevelSupported(Parameters.Platform, ERHIFeatureLevel::SM5) && !IsConsolePlatform(Parameters.Platform);
 	}
 
-	void SetParameters(FRHICommandList& RHICmdList, const FTexture* TextureValue, const FMatrix& ColorWeightsValue, float GammaValue, float MipLevel, float LayerIndex, bool bIsNormalMap, bool bIsVirtualTexture, bool bIsTextureArray)
+	void SetParameters(FRHICommandList& RHICmdList, const FTexture* TextureValue, const FMatrix& ColorWeightsValue, float GammaValue, float MipLevel, float LayerIndex, bool bIsNormalMap, bool bIsSingleVTPhysicalSpace, bool bIsVirtualTexture, bool bIsTextureArray)
 	{
 		if (bIsVirtualTexture)
 		{
@@ -101,7 +101,7 @@ public:
 		}
 		
 		SetShaderValue(RHICmdList, GetPixelShader(),ColorWeights,ColorWeightsValue);
-		FVector4 PackedParametersValue(GammaValue, MipLevel, bIsNormalMap ? 1.0 : -1.0f, LayerIndex);
+		FVector4 PackedParametersValue(GammaValue, MipLevel, bIsNormalMap ? 1.0 : -1.0f, bIsSingleVTPhysicalSpace ? 0 : LayerIndex);
 		SetShaderValue(RHICmdList, GetPixelShader(), PackedParameters, PackedParametersValue);
 
 		// Store slice count for texture array
@@ -179,5 +179,5 @@ void FBatchedElementTexture2DPreviewParameters::BindShaders(
 	SetGraphicsPipelineState(RHICmdList, GraphicsPSOInit, EApplyRendertargetOption::ForceApply);
 
 	VertexShader->SetParameters(RHICmdList, InTransform);
-	PixelShader->SetParameters(RHICmdList, Texture, ColorWeights, InGamma, MipLevel, LayerIndex, bIsNormalMap, bIsVirtualTexture, bIsTextureArray);
+	PixelShader->SetParameters(RHICmdList, Texture, ColorWeights, InGamma, MipLevel, LayerIndex, bIsNormalMap, bIsSingleVTPhysicalSpace, bIsVirtualTexture, bIsTextureArray);
 }
diff --git a/Engine/Source/Editor/UnrealEd/Public/Texture2DPreview.h b/Engine/Source/Editor/UnrealEd/Public/Texture2DPreview.h
--- a/Engine/Source/Editor/UnrealEd/Public/Texture2DPreview.h
+++ b/Engine/Source/Editor/UnrealEd/Public/Texture2DPreview.h
@@ -15,11 +15,12 @@
 class UNREALED_API FBatchedElementTexture2DPreviewParameters : public FBatchedElementParameters
 {
 public:
-	FBatchedElementTexture2DPreviewParameters(float InMipLevel, float InLayerIndex, bool bInIsNormalMap, bool bInIsSingleChannel, bool bInIsVirtualTexture, bool bInIsTextureArray = false)
+	FBatchedElementTexture2DPreviewParameters(float InMipLevel, float InLayerIndex, bool bInIsNormalMap, bool bInIsSingleChannel, bool bInIsSingleVTPhysicalSpace, bool bInIsVirtualTexture, bool bInIsTextureArray)
 		: MipLevel(InMipLevel)
 		, LayerIndex(InLayerIndex)
 		, bIsNormalMap( bInIsNormalMap )
 		, bIsSingleChannelFormat( bInIsSingleChannel )
+		, bIsSingleVTPhysicalSpace(bInIsSingleVTPhysicalSpace)
 		, bIsVirtualTexture(bInIsVirtualTexture)
 		, bIsTextureArray(bInIsTextureArray)
 	{
@@ -35,6 +36,7 @@ private:
 	float LayerIndex;
 	bool bIsNormalMap;
 	bool bIsSingleChannelFormat;
+	bool bIsSingleVTPhysicalSpace;
 
 	/** Parameters that are used to select a shader permutation */
 	bool bIsVirtualTexture;
diff --git a/Engine/Source/Runtime/Engine/Classes/Engine/Texture2D.h b/Engine/Source/Runtime/Engine/Classes/Engine/Texture2D.h
--- a/Engine/Source/Runtime/Engine/Classes/Engine/Texture2D.h
+++ b/Engine/Source/Runtime/Engine/Classes/Engine/Texture2D.h
@@ -409,4 +409,10 @@ public:
 		}
 		return false;
 	}
+
+	/**
+	 * Returns true if this virtual texture uses a single physical space all of its texture layers.
+	 * This can reduce page table overhead but potentially increase the number of physical pools allocated.
+	 */
+	virtual bool IsVirtualTexturedWithSinglePhysicalSpace() const { return false;  }
 };
diff --git a/Engine/Source/Runtime/Engine/Classes/Materials/MaterialExpressionRuntimeVirtualTextureSample.h b/Engine/Source/Runtime/Engine/Classes/Materials/MaterialExpressionRuntimeVirtualTextureSample.h
--- a/Engine/Source/Runtime/Engine/Classes/Materials/MaterialExpressionRuntimeVirtualTextureSample.h
+++ b/Engine/Source/Runtime/Engine/Classes/Materials/MaterialExpressionRuntimeVirtualTextureSample.h
@@ -50,13 +50,17 @@ class ENGINE_API UMaterialExpressionRuntimeVirtualTextureSample : public UMateri
 	UPROPERTY(EditAnywhere, Category = VirtualTexture, meta = (DisplayName = "Virtual texture content"))
 	ERuntimeVirtualTextureMaterialType MaterialType = ERuntimeVirtualTextureMaterialType::BaseColor;
 
+	/** Enable page table channel packing. Note that the bound Virtual Texture should have the same setting for sampling to work correctly. */
+	UPROPERTY(EditAnywhere, Category = VirtualTexture, meta = (DisplayName = "Enable packed page table"))
+	bool bSinglePhysicalSpace = true;
+
 	/** Defines how the MipValue property is applied to the virtual texture lookup. */
-	UPROPERTY(EditAnywhere, Category = VirtualTexture)
-	TEnumAsByte<enum ERuntimeVirtualTextureMipValueMode> MipValueMode;
+	UPROPERTY(EditAnywhere, Category = TextureSample)
+	TEnumAsByte<enum ERuntimeVirtualTextureMipValueMode> MipValueMode = RVTMVM_None;
 
 protected:
-	/** Init MaterialType based on the current VirtualTexture */
-	void InitMaterialType();
+	/** Init settings that affect shader compilation and need to match the current VirtualTexture */
+	void InitVirtualTextureDependentSettings();
 
 	/** Initialize the output pins. */
 	void InitOutputs();
diff --git a/Engine/Source/Runtime/Engine/Classes/VT/RuntimeVirtualTexture.h b/Engine/Source/Runtime/Engine/Classes/VT/RuntimeVirtualTexture.h
--- a/Engine/Source/Runtime/Engine/Classes/VT/RuntimeVirtualTexture.h
+++ b/Engine/Source/Runtime/Engine/Classes/VT/RuntimeVirtualTexture.h
@@ -24,10 +24,22 @@ protected:
 	UPROPERTY(EditAnywhere, BlueprintReadOnly, Category = Layout, meta = (DisplayName = "Enable BC texture compression"))
 	bool bCompressTextures = true;
 
+	/** Enable usage of the virtual texture. When disabled there is no rendering into the virtual texture, and sampling will return zero values. */
+	UPROPERTY(EditAnywhere, BlueprintReadOnly, AdvancedDisplay, Category = Layout, meta = (DisplayName = "Enable virtual texture"))
+	bool bEnable = true;
+
 	/** Enable clear before rendering a page of the virtual texture. Disabling this can be an optimization if you know that the texture will always be fully covered by rendering.  */
 	UPROPERTY(EditAnywhere, BlueprintReadOnly, AdvancedDisplay, Category = Layout, meta = (DisplayName = "Enable clear before render"))
 	bool bClearTextures = true;
 
+	/** Enable page table channel packing. This reduces page table memory and update cost but can reduce the ability to share physical memory with other virtual textures.  */
+	UPROPERTY(EditAnywhere, BlueprintReadOnly, AdvancedDisplay, Category = Layout, meta = (DisplayName = "Enable packed page table"))
+	bool bSinglePhysicalSpace = true;
+
+	/** Enable private page table allocation. This can reduce total page table memory allocation but can also reduce the total number of virtual textures supported. */
+	UPROPERTY(EditAnywhere, BlueprintReadOnly, AdvancedDisplay, Category = Layout, meta = (DisplayName = "Enable private page table"))
+	bool bPrivateSpace = true;
+
 	/** Size of virtual texture along the largest axis. (Actual values increase in powers of 2) */
 	UPROPERTY(EditAnywhere, BlueprintReadOnly, Category = Size, meta = (UIMin = "0", UIMax = "8", DisplayName = "Size of the virtual texture"))
 	int32 Size = 6; // 65536
@@ -52,10 +64,6 @@ protected:
 	UPROPERTY(VisibleAnywhere, Category = LowMips, meta = (DisplayName = "Streaming low mip texture"))
 	class URuntimeVirtualTextureStreamingProxy* StreamingTexture;
 
-	/** Enable usage of the virtual texture. When disabled there is no rendering into the virtual texture, and sampling will return zero values. */
-	UPROPERTY(EditAnywhere, BlueprintReadOnly, AdvancedDisplay, Category = Layout, meta = (DisplayName = "Enable virtual texture"))
-	bool bEnable = true;
-
 public:
 	/** Public getter for enabled status */
 	bool GetEnabled() { return bEnable; }
@@ -73,6 +81,8 @@ public:
 	int32 GetRemoveLowMips() const { return FMath::Clamp(RemoveLowMips, 0, 5); }
 	/** Public getter for virtual texture streaming low mips */
 	int32 GetStreamLowMips() const { return FMath::Clamp(StreamLowMips, 0, 6); }
+	/** Public getter for virtual texture using single physical space flag. */
+	bool GetSinglePhysicalSpace() const { return bSinglePhysicalSpace; }
 
 	/** Returns an approximate estimated value for the memory used by the page table texture. */
 	int32 GetEstimatedPageTableTextureMemoryKb() const;
diff --git a/Engine/Source/Runtime/Engine/Classes/VT/RuntimeVirtualTextureStreamingProxy.h b/Engine/Source/Runtime/Engine/Classes/VT/RuntimeVirtualTextureStreamingProxy.h
--- a/Engine/Source/Runtime/Engine/Classes/VT/RuntimeVirtualTextureStreamingProxy.h
+++ b/Engine/Source/Runtime/Engine/Classes/VT/RuntimeVirtualTextureStreamingProxy.h
@@ -17,12 +17,20 @@ class URuntimeVirtualTextureStreamingProxy : public UTexture2D
 	UPROPERTY()
 	FVirtualTextureBuildSettings Settings;
 
+	/** 
+	 * Enables combining layers into a single physical space. 
+	 * If this value doesn't match the owning URuntimeVirtualTexture then unwanted physical pools may be allocated.
+	 */
+	UPROPERTY()
+	bool bSinglePhysicalSpace;
+
 	/** Hash of settings used when building this texture. Used to invalidate when build settings have changed. */
 	UPROPERTY()
 	uint32 BuildHash;
 
 	//~ Begin UTexture Interface.
 	virtual void GetVirtualTextureBuildSettings(FVirtualTextureBuildSettings& OutSettings) const override;
+	virtual bool IsVirtualTexturedWithSinglePhysicalSpace() const override { return bSinglePhysicalSpace; }
 #if WITH_EDITOR
 	virtual void BeginCacheForCookedPlatformData(const ITargetPlatform* TargetPlatform) override;
 #endif
diff --git a/Engine/Source/Runtime/Engine/Classes/VT/VirtualTexturePoolConfig.h b/Engine/Source/Runtime/Engine/Classes/VT/VirtualTexturePoolConfig.h
--- a/Engine/Source/Runtime/Engine/Classes/VT/VirtualTexturePoolConfig.h
+++ b/Engine/Source/Runtime/Engine/Classes/VT/VirtualTexturePoolConfig.h
@@ -18,18 +18,26 @@ struct FVirtualTextureSpacePoolConfig
 {
 	GENERATED_USTRUCT_BODY()
 
-	FVirtualTextureSpacePoolConfig() : SizeInMegabyte(0), TileSize(0), Format(PF_Unknown) {}
+	FVirtualTextureSpacePoolConfig() : MinTileSize(0), MaxTileSize(0), SizeInMegabyte(0) {}
 
+	/** Minimum tile size to match (including tile border). */
 	UPROPERTY()
-	int32 SizeInMegabyte; // Size counted in tiles of the pool
+	int32 MinTileSize;
 
+	/** Maximum tile size to match (including tile border). */
 	UPROPERTY()
-	int32 TileSize; // This is *including* any filtering borders
+	int32 MaxTileSize;
 
+	/** Format set to match. One pool can contain multiple layers with synchronized page table mappings. */
 	UPROPERTY()
-	TEnumAsByte<EPixelFormat> Format;
+	TArray< TEnumAsByte<EPixelFormat> > Formats;
 
-	bool IsDefault() const { return TileSize == 0; }
+	/** Upper limit of size in megabytes to allocate for the pool. The allocator will allocate as close as possible to this limit. */
+	UPROPERTY()
+	int32 SizeInMegabyte;
+
+	/** Is this the default config? Use this setting when we can't find any other match. */
+	bool IsDefault() const { return Formats.Num() == 0; }
 };
 
 UCLASS(config = Engine, transient)
@@ -43,7 +51,7 @@ public:
 	UPROPERTY(config)
 	TArray<FVirtualTextureSpacePoolConfig> Pools; // All the VT pools specified in the config
 
-	const FVirtualTextureSpacePoolConfig *FindPoolConfig(int32 TileSize, EPixelFormat Format);
+	const FVirtualTextureSpacePoolConfig *FindPoolConfig(TEnumAsByte<EPixelFormat> const* Formats, int32 NumLayers, int32 TileSize);
 
 	virtual ~UVirtualTexturePoolConfig();
 
diff --git a/Engine/Source/Runtime/Engine/Private/LightMap.cpp b/Engine/Source/Runtime/Engine/Private/LightMap.cpp
--- a/Engine/Source/Runtime/Engine/Private/LightMap.cpp
+++ b/Engine/Source/Runtime/Engine/Private/LightMap.cpp
@@ -3333,16 +3333,16 @@ IAllocatedVirtualTexture* FLightmapResourceCluster::AcquireAllocatedVT() const
 		VTDesc.Dimensions = 2;
 		VTDesc.TileSize = Resource->GetTileSize();
 		VTDesc.TileBorderSize = Resource->GetBorderSize();
-		VTDesc.NumLayers = 0u;
+		VTDesc.NumTextureLayers = 0u;
 
 		for (uint32 TypeIndex = 0u; TypeIndex < (uint32)ELightMapVirtualTextureType::Count; ++TypeIndex)
 		{
 			const uint32 LayerIndex = VirtualTexture->GetLayerForType((ELightMapVirtualTextureType)TypeIndex);
 			if (LayerIndex != ~0u)
 			{
-				VTDesc.NumLayers = TypeIndex + 1u;
+				VTDesc.NumTextureLayers = TypeIndex + 1u;
 				VTDesc.ProducerHandle[TypeIndex] = ProducerHandle; // use the same producer for each layer
-				VTDesc.LocalLayerToProduce[TypeIndex] = LayerIndex;
+				VTDesc.ProducerLayerIndex[TypeIndex] = LayerIndex;
 			}
 			else
 			{
@@ -3350,8 +3350,8 @@ IAllocatedVirtualTexture* FLightmapResourceCluster::AcquireAllocatedVT() const
 			}
 		}
 
-		check(VTDesc.NumLayers > 0u);
-		for (uint32 LayerIndex = 0u; LayerIndex < VTDesc.NumLayers; ++LayerIndex)
+		check(VTDesc.NumTextureLayers > 0u);
+		for (uint32 LayerIndex = 0u; LayerIndex < VTDesc.NumTextureLayers; ++LayerIndex)
 		{
 			if (VTDesc.ProducerHandle[LayerIndex].PackedValue == 0u)
 			{
@@ -3360,7 +3360,7 @@ IAllocatedVirtualTexture* FLightmapResourceCluster::AcquireAllocatedVT() const
 				// and attempt to do some extra work before determining there's nothing else to do...this wastes CPU time
 				// By mapping to layer0, we ensure that every layer has a valid mapping, and the overhead of mapping the empty layer to layer0 is very small, since layer0 will already be resident
 				VTDesc.ProducerHandle[LayerIndex] = ProducerHandle;
-				VTDesc.LocalLayerToProduce[LayerIndex] = 0u;
+				VTDesc.ProducerLayerIndex[LayerIndex] = 0u;
 			}
 		}
 
diff --git a/Engine/Source/Runtime/Engine/Private/LightmapUniformShaderParameters.cpp b/Engine/Source/Runtime/Engine/Private/LightmapUniformShaderParameters.cpp
--- a/Engine/Source/Runtime/Engine/Private/LightmapUniformShaderParameters.cpp
+++ b/Engine/Source/Runtime/Engine/Private/LightmapUniformShaderParameters.cpp
@@ -104,7 +104,7 @@ void GetPrecomputedLightingParameters(
 			check(AllocatedVT);
 
 			AllocatedVT->GetPackedPageTableUniform(&Parameters.LightmapVTPackedPageTableUniform[0], true);
-			NumLightmapVTLayers = AllocatedVT->GetNumLayers();
+			NumLightmapVTLayers = AllocatedVT->GetNumTextureLayers();
 			for (uint32 LayerIndex = 0u; LayerIndex < NumLightmapVTLayers; ++LayerIndex)
 			{
 				AllocatedVT->GetPackedUniform(&Parameters.LightmapVTPackedUniform[LayerIndex], LayerIndex);
diff --git a/Engine/Source/Runtime/Engine/Private/Materials/HLSLMaterialTranslator.h b/Engine/Source/Runtime/Engine/Private/Materials/HLSLMaterialTranslator.h
--- a/Engine/Source/Runtime/Engine/Private/Materials/HLSLMaterialTranslator.h
+++ b/Engine/Source/Runtime/Engine/Private/Materials/HLSLMaterialTranslator.h
@@ -4582,6 +4582,7 @@ protected:
 
 		int32 VTStackIndex = INDEX_NONE;
 		int32 VTLayerIndex = INDEX_NONE;
+		int32 VTPageTableIndex = INDEX_NONE;
 		if (bVirtualTexture)
 		{
 			check(VirtualTextureIndex >= 0);
@@ -4628,12 +4629,13 @@ protected:
 			//todo[vt]: Support feedback from other shader types
 			const bool bGenerateFeedback = ShaderFrequency == SF_Pixel;
 
-			VTLayerIndex = MaterialCompilationOutput.UniformExpressionSet.UniformVirtualTextureExpressions[VirtualTextureIndex]->GetLayerIndex();
+			VTLayerIndex = MaterialCompilationOutput.UniformExpressionSet.UniformVirtualTextureExpressions[VirtualTextureIndex]->GetTextureLayerIndex();
 			if (VTLayerIndex != INDEX_NONE)
 			{
 				// The layer index in the virtual texture stack is already known
 				// Create a page table sample for each new combination of virtual texture and sample parameters
 				VTStackIndex = AcquireVTStackIndex(MipValueMode, AddressU, AddressV, 1.0f, CoordinateIndex, MipValue0Index, MipValue1Index, TextureReferenceIndex, bGenerateFeedback);
+				VTPageTableIndex = MaterialCompilationOutput.UniformExpressionSet.UniformVirtualTextureExpressions[VirtualTextureIndex]->GetPageTableLayerIndex();
 			}
 			else
 			{
@@ -4651,6 +4653,7 @@ protected:
 				VTStackIndex = AcquireVTStackIndex(MipValueMode, AddressU, AddressV, TextureAspectRatio, CoordinateIndex, MipValue0Index, MipValue1Index, INDEX_NONE, bGenerateFeedback);
 				// Allocate a layer in the virtual texture stack for this physical sample
 				VTLayerIndex = MaterialCompilationOutput.UniformExpressionSet.VTStacks[VTStackIndex].AddLayer();
+				VTPageTableIndex = VTLayerIndex;
 			}
 
 			MaterialCompilationOutput.UniformExpressionSet.VTStacks[VTStackIndex].SetLayer(VTLayerIndex, VirtualTextureIndex);
@@ -4667,7 +4670,7 @@ protected:
 				*SampleCode,
 				*TextureName,
 				*VTPageTableResult,
-				VTLayerIndex,
+				VTPageTableIndex,
 				VirtualTextureIndex);
 
 			// TODO
@@ -5086,9 +5089,9 @@ protected:
 		FMaterialParameterInfo ParameterInfo = GetParameterAssociationInfo();
 		ParameterInfo.Name = ParameterName;
 
-		const bool bVirtualTexturesEnabeled = UseVirtualTexturing(FeatureLevel, TargetPlatform);
+		const bool bVirtualTexturesEnabled = UseVirtualTexturing(FeatureLevel, TargetPlatform);
 		bool bVirtual = ShaderType == MCT_TextureVirtual;
-		if (bVirtualTexturesEnabeled == false && ShaderType == MCT_TextureVirtual)
+		if (bVirtualTexturesEnabled == false && ShaderType == MCT_TextureVirtual)
 		{
 			bVirtual = false;
 			ShaderType = MCT_Texture2D;
@@ -5096,7 +5099,7 @@ protected:
 		return AddUniformExpression(new FMaterialUniformExpressionTextureParameter(ParameterInfo, TextureReferenceIndex, SamplerType, SamplerSource, bVirtual),ShaderType,TEXT(""));
 	}
 
-	virtual int32 VirtualTexture(URuntimeVirtualTexture* InTexture, int32 LayerIndex, int32& TextureReferenceIndex, EMaterialSamplerType SamplerType) override
+	virtual int32 VirtualTexture(URuntimeVirtualTexture* InTexture, int32 TextureLayerIndex, int32 PageTableLayerIndex, int32& TextureReferenceIndex, EMaterialSamplerType SamplerType) override
 	{
 		if (!UseVirtualTexturing(FeatureLevel, TargetPlatform))
 		{
@@ -5106,10 +5109,10 @@ protected:
 		TextureReferenceIndex = Material->GetReferencedTextures().Find(InTexture);
 		checkf(TextureReferenceIndex != INDEX_NONE, TEXT("Material expression called Compiler->VirtualTexture() without implementing UMaterialExpression::GetReferencedTexture properly"));
 
-		return AddUniformExpression(new FMaterialUniformExpressionTexture(TextureReferenceIndex, LayerIndex, SamplerType), MCT_TextureVirtual, TEXT(""));
+		return AddUniformExpression(new FMaterialUniformExpressionTexture(TextureReferenceIndex, TextureLayerIndex, PageTableLayerIndex, SamplerType), MCT_TextureVirtual, TEXT(""));
 	}
 
-	virtual int32 VirtualTextureParameter(FName ParameterName, URuntimeVirtualTexture* DefaultValue, int32 LayerIndex, int32& TextureReferenceIndex, EMaterialSamplerType SamplerType) override
+	virtual int32 VirtualTextureParameter(FName ParameterName, URuntimeVirtualTexture* DefaultValue, int32 TextureLayerIndex, int32 PageTableLayerIndex, int32& TextureReferenceIndex, EMaterialSamplerType SamplerType) override
 	{
 		if (!UseVirtualTexturing(FeatureLevel, TargetPlatform))
 		{
@@ -5122,7 +5125,7 @@ protected:
 		FMaterialParameterInfo ParameterInfo = GetParameterAssociationInfo();
 		ParameterInfo.Name = ParameterName;
 
-		return AddUniformExpression(new FMaterialUniformExpressionTextureParameter(ParameterInfo, TextureReferenceIndex, LayerIndex, SamplerType), MCT_TextureVirtual, TEXT(""));
+		return AddUniformExpression(new FMaterialUniformExpressionTextureParameter(ParameterInfo, TextureReferenceIndex, TextureLayerIndex, PageTableLayerIndex, SamplerType), MCT_TextureVirtual, TEXT(""));
 	}
 
 	virtual int32 VirtualTextureUniform(int32 TextureIndex, int32 VectorIndex) override
diff --git a/Engine/Source/Runtime/Engine/Private/Materials/MaterialExpressions.cpp b/Engine/Source/Runtime/Engine/Private/Materials/MaterialExpressions.cpp
--- a/Engine/Source/Runtime/Engine/Private/Materials/MaterialExpressions.cpp
+++ b/Engine/Source/Runtime/Engine/Private/Materials/MaterialExpressions.cpp
@@ -2074,11 +2074,12 @@ UMaterialExpressionRuntimeVirtualTextureSample::UMaterialExpressionRuntimeVirtua
 #endif
 }
 
-void UMaterialExpressionRuntimeVirtualTextureSample::InitMaterialType()
+void UMaterialExpressionRuntimeVirtualTextureSample::InitVirtualTextureDependentSettings()
 {
 	if (VirtualTexture != nullptr)
 	{
 		MaterialType = VirtualTexture->GetMaterialType();
+		bSinglePhysicalSpace = VirtualTexture->GetSinglePhysicalSpace();
 	}
 }
 
@@ -2115,7 +2116,7 @@ void UMaterialExpressionRuntimeVirtualTextureSample::PostEditChangeProperty(FPro
 	{
 		if (VirtualTexture != nullptr)
 		{
-			InitMaterialType();
+			InitVirtualTextureDependentSettings();
 			FEditorSupportDelegates::ForcePropertyWindowRebuild.Broadcast(this);
 		}
 	}
@@ -2150,6 +2151,15 @@ int32 UMaterialExpressionRuntimeVirtualTextureSample::Compile(class FMaterialCom
 
 		bIsVirtualTextureValid = false;
 	}
+	else if (VirtualTexture->GetSinglePhysicalSpace() != bSinglePhysicalSpace)
+	{
+		Compiler->Errorf(TEXT("%Page table packing is '%d', should be '%d' to match %s"),
+			bSinglePhysicalSpace ? 1 : 0,
+			VirtualTexture->GetSinglePhysicalSpace() ? 1 : 0,
+			*VirtualTexture->GetName());
+
+		bIsVirtualTextureValid = false;
+	}
 
 	// Calculate the virtual texture layer and sampling/unpacking functions for this output
 	// Fallback to a sensible default value if the output isn't valid for the bound virtual texture
@@ -2236,20 +2246,22 @@ int32 UMaterialExpressionRuntimeVirtualTextureSample::Compile(class FMaterialCom
 	
 	// Compile the texture object references
 	enum { MAX_RVT_LAYERS = 2 };
-	const int32 LayerCount = URuntimeVirtualTexture::GetLayerCount(MaterialType);
-	check(LayerCount <= MAX_RVT_LAYERS);
+	const int32 TextureLayerCount = URuntimeVirtualTexture::GetLayerCount(MaterialType);
+	check(TextureLayerCount <= MAX_RVT_LAYERS);
 
 	int32 TextureCodeIndex[MAX_RVT_LAYERS] = { INDEX_NONE };
 	int32 TextureReferenceIndex[MAX_RVT_LAYERS] = { INDEX_NONE };
-	for (int32 Layer = 0; Layer < LayerCount; Layer++)
+	for (int32 TexureLayerIndex = 0; TexureLayerIndex < TextureLayerCount; TexureLayerIndex++)
 	{
+		const int32 PageTableLayerIndex = bSinglePhysicalSpace ? 0 : TexureLayerIndex;
+
 		if (bIsParameter)
 		{
-			TextureCodeIndex[Layer] = Compiler->VirtualTextureParameter(GetParameterName(), VirtualTexture, Layer, TextureReferenceIndex[Layer], SAMPLERTYPE_VirtualMasks);
+			TextureCodeIndex[TexureLayerIndex] = Compiler->VirtualTextureParameter(GetParameterName(), VirtualTexture, TexureLayerIndex, PageTableLayerIndex, TextureReferenceIndex[TexureLayerIndex], SAMPLERTYPE_VirtualMasks);
 		}
 		else
 		{
-			TextureCodeIndex[Layer] = Compiler->VirtualTexture(VirtualTexture, Layer, TextureReferenceIndex[Layer], SAMPLERTYPE_VirtualMasks);
+			TextureCodeIndex[TexureLayerIndex] = Compiler->VirtualTexture(VirtualTexture, TexureLayerIndex, PageTableLayerIndex, TextureReferenceIndex[TexureLayerIndex], SAMPLERTYPE_VirtualMasks);
 		}
 	}
 
@@ -2297,14 +2309,14 @@ int32 UMaterialExpressionRuntimeVirtualTextureSample::Compile(class FMaterialCom
 
 	// Compile the texture sample code
 	int32 SampleCodeIndex[MAX_RVT_LAYERS] = { INDEX_NONE };
-	for (int32 Layer = 0; Layer < LayerCount; Layer++)
+	for (int32 TexureLayerIndex = 0; TexureLayerIndex < TextureLayerCount; TexureLayerIndex++)
 	{
-		SampleCodeIndex[Layer] = Compiler->TextureSample(
-			TextureCodeIndex[Layer],
+		SampleCodeIndex[TexureLayerIndex] = Compiler->TextureSample(
+			TextureCodeIndex[TexureLayerIndex],
 			CoordinateIndex, 
 			SAMPLERTYPE_VirtualMasks,
 			MipValueIndex, INDEX_NONE, TextureMipLevelMode, SSM_Wrap_WorldGroupSettings,
-			TextureReferenceIndex[Layer], 
+			TextureReferenceIndex[TexureLayerIndex],
 			false);
 	}
 
diff --git a/Engine/Source/Runtime/Engine/Private/Materials/MaterialShared.cpp b/Engine/Source/Runtime/Engine/Private/Materials/MaterialShared.cpp
--- a/Engine/Source/Runtime/Engine/Private/Materials/MaterialShared.cpp
+++ b/Engine/Source/Runtime/Engine/Private/Materials/MaterialShared.cpp
@@ -2317,9 +2317,9 @@ IAllocatedVirtualTexture* FMaterialRenderProxy::AllocateVTStack(const FMaterialR
 
 	FAllocatedVTDescription VTDesc;
 	VTDesc.Dimensions = 2;
-	VTDesc.NumLayers = NumLayers;
+	VTDesc.NumTextureLayers = NumLayers;
 	bool bFoundValidLayer = false;
-	for (uint32 LayerIndex = 0u; LayerIndex < VTDesc.NumLayers; ++LayerIndex)
+	for (uint32 LayerIndex = 0u; LayerIndex < NumLayers; ++LayerIndex)
 	{
 		const UTexture2D* Texture = LayerTextures[LayerIndex];
 		const FVirtualTexture2DResource* VirtualTextureResourceForLayer = (Texture && Texture->IsCurrentlyVirtualTextured()) ? (FVirtualTexture2DResource*)Texture->Resource : nullptr;
@@ -2333,7 +2333,7 @@ IAllocatedVirtualTexture* FMaterialRenderProxy::AllocateVTStack(const FMaterialR
 			VTDesc.TileBorderSize = VirtualTextureResourceForLayer->GetBorderSize();
 			const FVirtualTextureProducerHandle& ProducerHandle = VirtualTextureResourceForLayer->GetProducerHandle();
 			VTDesc.ProducerHandle[LayerIndex] = ProducerHandle;
-			VTDesc.LocalLayerToProduce[LayerIndex] = 0u;
+			VTDesc.ProducerLayerIndex[LayerIndex] = 0u;
 			GetRendererModule().AddVirtualTextureProducerDestroyedCallback(ProducerHandle, &OnVirtualTextureDestroyedCB, const_cast<FMaterialRenderProxy*>(this));
 			bFoundValidLayer = true;
 		}
diff --git a/Engine/Source/Runtime/Engine/Private/Materials/MaterialUniformExpressions.cpp b/Engine/Source/Runtime/Engine/Private/Materials/MaterialUniformExpressions.cpp
--- a/Engine/Source/Runtime/Engine/Private/Materials/MaterialUniformExpressions.cpp
+++ b/Engine/Source/Runtime/Engine/Private/Materials/MaterialUniformExpressions.cpp
@@ -518,7 +518,7 @@ void FUniformExpressionSet::FillUniformBuffer(const FMaterialRenderContext& Mate
 				UniformVirtualTextureExpressions[ExpressionIndex]->GetTextureValue(MaterialRenderContext, MaterialRenderContext.Material, Texture);
 				if (Texture != nullptr)
 				{
-					int32 LayerIndex = UniformVirtualTextureExpressions[ExpressionIndex]->GetLayerIndex();
+					int32 LayerIndex = UniformVirtualTextureExpressions[ExpressionIndex]->GetTextureLayerIndex();
 					IAllocatedVirtualTexture const* AllocatedVT = Texture->GetAllocatedVirtualTexture();
 					if (AllocatedVT)
 					{
@@ -884,7 +884,7 @@ void FUniformExpressionSet::FillUniformBuffer(const FMaterialRenderContext& Mate
 					IAllocatedVirtualTexture const* AllocatedVT = Texture->GetAllocatedVirtualTexture();
 					if (AllocatedVT != nullptr)
 					{
-						const int32 LayerIndex = UniformVirtualTextureExpressions[ExpressionIndex]->GetLayerIndex();
+						const int32 LayerIndex = UniformVirtualTextureExpressions[ExpressionIndex]->GetTextureLayerIndex();
 						FRHIShaderResourceView* PhysicalViewRHI = AllocatedVT->GetPhysicalTextureView(LayerIndex, Texture->IsLayerSRGB(LayerIndex));
 						if (PhysicalViewRHI != nullptr)
 						{
@@ -942,7 +942,8 @@ uint32 FUniformExpressionSet::GetReferencedTexture2DRHIHash(const FMaterialRende
 
 FMaterialUniformExpressionTexture::FMaterialUniformExpressionTexture() :
 	TextureIndex(INDEX_NONE),
-	LayerIndex(INDEX_NONE),
+	TextureLayerIndex(INDEX_NONE),
+	PageTableLayerIndex(INDEX_NONE),
 #if WITH_EDITORONLY_DATA
 	SamplerType(SAMPLERTYPE_Color),
 #endif
@@ -954,7 +955,8 @@ FMaterialUniformExpressionTexture::FMaterialUniformExpressionTexture() :
 
 FMaterialUniformExpressionTexture::FMaterialUniformExpressionTexture(int32 InTextureIndex, EMaterialSamplerType InSamplerType, ESamplerSourceMode InSamplerSource, bool InVirtualTexture) :
 	TextureIndex(InTextureIndex),
-	LayerIndex(INDEX_NONE),
+	TextureLayerIndex(INDEX_NONE),
+	PageTableLayerIndex(INDEX_NONE),
 #if WITH_EDITORONLY_DATA
 	SamplerType(InSamplerType),
 #endif
@@ -965,9 +967,10 @@ FMaterialUniformExpressionTexture::FMaterialUniformExpressionTexture(int32 InTex
 {
 }
 
-FMaterialUniformExpressionTexture::FMaterialUniformExpressionTexture(int32 InTextureIndex, int32 InLayerIndex, EMaterialSamplerType InSamplerType)
+FMaterialUniformExpressionTexture::FMaterialUniformExpressionTexture(int32 InTextureIndex, int16 InTextureLayerIndex, int16 InPageTableLayerIndex, EMaterialSamplerType InSamplerType)
 	: TextureIndex(InTextureIndex)
-	, LayerIndex(InLayerIndex)
+	, TextureLayerIndex(InTextureLayerIndex)
+	, PageTableLayerIndex(InPageTableLayerIndex)
 #if WITH_EDITORONLY_DATA
 	, SamplerType(InSamplerType)
 #endif
@@ -981,7 +984,7 @@ FMaterialUniformExpressionTexture::FMaterialUniformExpressionTexture(int32 InTex
 void FMaterialUniformExpressionTexture::Serialize(FArchive& Ar)
 {
 	int32 SamplerSourceInt = (int32)SamplerSource;
-	Ar << TextureIndex << LayerIndex << SamplerSourceInt << bVirtualTexture;
+	Ar << TextureIndex << TextureLayerIndex << PageTableLayerIndex << SamplerSourceInt << bVirtualTexture;
 	SamplerSource = (ESamplerSourceMode)SamplerSourceInt;
 }
 
@@ -1036,7 +1039,10 @@ bool FMaterialUniformExpressionTexture::IsIdentical(const FMaterialUniformExpres
 	}
 	FMaterialUniformExpressionTexture* OtherTextureExpression = (FMaterialUniformExpressionTexture*)OtherExpression;
 
-	return TextureIndex == OtherTextureExpression->TextureIndex && LayerIndex == OtherTextureExpression->LayerIndex && bVirtualTexture == OtherTextureExpression->bVirtualTexture;
+	return TextureIndex == OtherTextureExpression->TextureIndex && 
+		TextureLayerIndex == OtherTextureExpression->TextureLayerIndex &&
+		PageTableLayerIndex == OtherTextureExpression->PageTableLayerIndex &&
+		bVirtualTexture == OtherTextureExpression->bVirtualTexture;
 }
 
 void FMaterialUniformExpressionTextureParameter::GetTextureValue(const FMaterialRenderContext& Context, const FMaterial& Material, const UTexture*& OutValue) const
diff --git a/Engine/Source/Runtime/Engine/Private/Materials/MaterialUniformExpressions.h b/Engine/Source/Runtime/Engine/Private/Materials/MaterialUniformExpressions.h
--- a/Engine/Source/Runtime/Engine/Private/Materials/MaterialUniformExpressions.h
+++ b/Engine/Source/Runtime/Engine/Private/Materials/MaterialUniformExpressions.h
@@ -302,8 +302,8 @@ public:
 		, ParameterInfo(InParameterInfo)
 	{}
 
-	FMaterialUniformExpressionTextureParameter(const FMaterialParameterInfo& InParameterInfo, int32 InTextureIndex, int32 InLayerIndex, EMaterialSamplerType InSamplerType)
-		: Super(InTextureIndex, InLayerIndex, InSamplerType)
+	FMaterialUniformExpressionTextureParameter(const FMaterialParameterInfo& InParameterInfo, int32 InTextureIndex, int32 InTextureLayerIndex, int32 InPageTableLayerIndex, EMaterialSamplerType InSamplerType)
+		: Super(InTextureIndex, InTextureLayerIndex, InPageTableLayerIndex, InSamplerType)
 		, ParameterInfo(InParameterInfo)
 	{}
 
diff --git a/Engine/Source/Runtime/Engine/Private/Texture2D.cpp b/Engine/Source/Runtime/Engine/Private/Texture2D.cpp
--- a/Engine/Source/Runtime/Engine/Private/Texture2D.cpp
+++ b/Engine/Source/Runtime/Engine/Private/Texture2D.cpp
@@ -1889,6 +1889,8 @@ void FVirtualTexture2DResource::InitRHI()
 	const int32 MaxLevel = VTData->GetNumMips() - FirstMipToUse - 1;
 	check(MaxLevel >= 0);
 
+	const bool bSinglePhysicalSpace = TextureOwner->IsVirtualTexturedWithSinglePhysicalSpace();
+
 	FVTProducerDescription ProducerDesc;
 	ProducerDesc.Name = TextureOwner->GetFName();
 	ProducerDesc.Dimensions = 2;
@@ -1900,10 +1902,12 @@ void FVirtualTexture2DResource::InitRHI()
 	ProducerDesc.HeightInBlocks = VTData->HeightInBlocks;
 	ProducerDesc.DepthInTiles = 1u;
 	ProducerDesc.MaxLevel = MaxLevel;
-	ProducerDesc.NumLayers = VTData->GetNumLayers();
-	for (uint32 LayerIndex = 0u; LayerIndex < ProducerDesc.NumLayers; ++LayerIndex)
+	ProducerDesc.NumTextureLayers = VTData->GetNumLayers();
+	ProducerDesc.NumPhysicalGroups = bSinglePhysicalSpace ? 1 : VTData->GetNumLayers();
+	for (uint32 LayerIndex = 0u; LayerIndex < VTData->GetNumLayers(); ++LayerIndex)
 	{
 		ProducerDesc.LayerFormat[LayerIndex] = VTData->LayerTypes[LayerIndex];
+		ProducerDesc.PhysicalGroupIndex[LayerIndex] = bSinglePhysicalSpace ? 0 : LayerIndex;
 	}
 
 	FUploadingVirtualTexture* VirtualTexture = new FUploadingVirtualTexture(VTData, FirstMipToUse);
@@ -2043,11 +2047,11 @@ class IAllocatedVirtualTexture* FVirtualTexture2DResource::AcquireAllocatedVT()
 		VTDesc.Dimensions = 2;
 		VTDesc.TileSize = VTData->TileSize;
 		VTDesc.TileBorderSize = VTData->TileBorderSize;
-		VTDesc.NumLayers = VTData->GetNumLayers();
-		for (uint32 LayerIndex = 0u; LayerIndex < VTDesc.NumLayers; ++LayerIndex)
+		VTDesc.NumTextureLayers = VTData->GetNumLayers();
+		for (uint32 LayerIndex = 0u; LayerIndex < VTDesc.NumTextureLayers; ++LayerIndex)
 		{
 			VTDesc.ProducerHandle[LayerIndex] = ProducerHandle; // use the same producer for each layer
-			VTDesc.LocalLayerToProduce[LayerIndex] = LayerIndex;
+			VTDesc.ProducerLayerIndex[LayerIndex] = LayerIndex;
 		}
 		AllocatedVT = GetRendererModule().AllocateVirtualTexture(VTDesc);
 	}
diff --git a/Engine/Source/Runtime/Engine/Private/VT/RuntimeVirtualTexture.cpp b/Engine/Source/Runtime/Engine/Private/VT/RuntimeVirtualTexture.cpp
--- a/Engine/Source/Runtime/Engine/Private/VT/RuntimeVirtualTexture.cpp
+++ b/Engine/Source/Runtime/Engine/Private/VT/RuntimeVirtualTexture.cpp
@@ -28,7 +28,8 @@ namespace
 			OutDesc.DepthInTiles = 1;
 			OutDesc.WidthInBlocks = 1;
 			OutDesc.HeightInBlocks = 1;
-			OutDesc.NumLayers = 0;
+			OutDesc.NumTextureLayers = 0;
+			OutDesc.NumPhysicalGroups = 0;
 		}
 
 		//~ Begin IVirtualTexture Interface.
@@ -89,17 +90,17 @@ public:
 	}
 
 	/** Queues up render thread work to create resources and also releases any old resources. */
-	void Init(FVTProducerDescription const& InDesc, IVirtualTexture* InVirtualTextureProducer)
+	void Init(FVTProducerDescription const& InDesc, IVirtualTexture* InVirtualTextureProducer, bool InPrivateSpace)
 	{
 		FRuntimeVirtualTextureRenderResource* Resource = this;
 		
 		ENQUEUE_RENDER_COMMAND(FRuntimeVirtualTextureRenderResource_Init)(
-			[Resource, InDesc, InVirtualTextureProducer](FRHICommandList& RHICmdList)
+			[Resource, InDesc, InVirtualTextureProducer, InPrivateSpace](FRHICommandList& RHICmdList)
 		{
 			FVirtualTextureProducerHandle OldProducerHandle = Resource->ProducerHandle;
 			ReleaseVirtualTexture(Resource->AllocatedVirtualTexture);
 			Resource->ProducerHandle = GetRendererModule().RegisterVirtualTextureProducer(InDesc, InVirtualTextureProducer);
-			Resource->AllocatedVirtualTexture = AllocateVirtualTexture(InDesc, Resource->ProducerHandle);
+			Resource->AllocatedVirtualTexture = AllocateVirtualTexture(InDesc, Resource->ProducerHandle, InPrivateSpace);
 			// Release old producer after new one is created so that any destroy callbacks can access the new producer
 			GetRendererModule().ReleaseVirtualTextureProducer(OldProducerHandle);
 		});
@@ -123,24 +124,24 @@ public:
 
 protected:
 	/** Allocate in the global virtual texture system. */
-	static IAllocatedVirtualTexture* AllocateVirtualTexture(FVTProducerDescription const& InDesc, FVirtualTextureProducerHandle const& InProducerHandle)
+	static IAllocatedVirtualTexture* AllocateVirtualTexture(FVTProducerDescription const& InDesc, FVirtualTextureProducerHandle const& InProducerHandle, bool InPrivateSpace)
 	{
 		IAllocatedVirtualTexture* OutAllocatedVirtualTexture = nullptr;
 
 		// Check for NumLayers avoids allocating for the null producer
-		if (InDesc.NumLayers > 0)
+		if (InDesc.NumTextureLayers > 0)
 		{
 			FAllocatedVTDescription VTDesc;
 			VTDesc.Dimensions = InDesc.Dimensions;
 			VTDesc.TileSize = InDesc.TileSize;
 			VTDesc.TileBorderSize = InDesc.TileBorderSize;
-			VTDesc.NumLayers = InDesc.NumLayers;
-			VTDesc.bPrivateSpace = true; // Dedicated page table allocation for runtime VTs
+			VTDesc.NumTextureLayers = InDesc.NumTextureLayers;
+			VTDesc.bPrivateSpace = InPrivateSpace;
 
-			for (uint32 LayerIndex = 0u; LayerIndex < VTDesc.NumLayers; ++LayerIndex)
+			for (uint32 LayerIndex = 0u; LayerIndex < VTDesc.NumTextureLayers; ++LayerIndex)
 			{
 				VTDesc.ProducerHandle[LayerIndex] = InProducerHandle;
-				VTDesc.LocalLayerToProduce[LayerIndex] = LayerIndex;
+				VTDesc.ProducerLayerIndex[LayerIndex] = LayerIndex;
 			}
 
 			OutAllocatedVirtualTexture = GetRendererModule().AllocateVirtualTexture(VTDesc);
@@ -236,10 +237,13 @@ void URuntimeVirtualTexture::GetProducerDescription(FVTProducerDescription& OutD
 	OutDesc.BlockHeightInTiles = Height / GetTileSize();
 	OutDesc.MaxLevel = FMath::Max((int32)FMath::CeilLogTwo(FMath::Max(OutDesc.BlockWidthInTiles, OutDesc.BlockHeightInTiles)) - GetRemoveLowMips(), 0);
 
-	OutDesc.NumLayers = GetLayerCount();
-	for (int32 Layer = 0; Layer < OutDesc.NumLayers; Layer++)
+	OutDesc.NumTextureLayers = GetLayerCount();
+	OutDesc.NumPhysicalGroups = bSinglePhysicalSpace ? 1 : GetLayerCount();
+
+	for (int32 Layer = 0; Layer < OutDesc.NumTextureLayers; Layer++)
 	{
 		OutDesc.LayerFormat[Layer] = GetLayerFormat(Layer);
+		OutDesc.PhysicalGroupIndex[Layer] = bSinglePhysicalSpace ? 0 : Layer;
 	}
 }
 
@@ -383,7 +387,7 @@ void URuntimeVirtualTexture::InitResource(IVirtualTexture* InProducer, FTransfor
 {
 	FVTProducerDescription Desc;
 	GetProducerDescription(Desc, VolumeToWorld);
-	Resource->Init(Desc, InProducer);
+	Resource->Init(Desc, InProducer, bPrivateSpace);
 }
 
 void URuntimeVirtualTexture::InitNullResource()
@@ -391,7 +395,7 @@ void URuntimeVirtualTexture::InitNullResource()
 	FVTProducerDescription Desc;
 	FNullVirtualTextureProducer::GetNullProducerDescription(Desc);
 	FNullVirtualTextureProducer* Producer = new FNullVirtualTextureProducer;
-	Resource->Init(Desc, Producer);
+	Resource->Init(Desc, Producer, false);
 }
 
 void URuntimeVirtualTexture::GetAssetRegistryTags(TArray<FAssetRegistryTag>& OutTags) const
@@ -444,7 +448,7 @@ uint32 URuntimeVirtualTexture::GetStreamingTextureBuildHash() const
 		{
 			uint32 MaterialType : 3;
 			uint32 CompressTextures : 1;
-			uint32 Size : 4;
+			uint32 SinglePhysicalSpace : 1;
 			uint32 TileSize : 4;
 			uint32 TileBorderSize : 4;
 			uint32 StreamLowMips : 4;
@@ -455,7 +459,7 @@ uint32 URuntimeVirtualTexture::GetStreamingTextureBuildHash() const
 	Settings.PackedValue = 0;
 	Settings.MaterialType = (uint32)MaterialType;
 	Settings.CompressTextures = (uint32)bCompressTextures;
-	Settings.Size = (uint32)Size;
+	Settings.SinglePhysicalSpace = (uint32)bSinglePhysicalSpace;
 	Settings.TileSize = (uint32)TileSize;
 	Settings.TileBorderSize = (uint32)TileBorderSize;
 	Settings.StreamLowMips = (uint32)GetStreamLowMips();
@@ -471,7 +475,7 @@ void URuntimeVirtualTexture::InitializeStreamingTexture(uint32 InSizeX, uint32 I
 
 	StreamingTexture = NewObject<URuntimeVirtualTextureStreamingProxy>(GetOutermost(), TEXT("StreamingTexture"));
 	StreamingTexture->VirtualTextureStreaming = true;
-	StreamingTexture->MipGenSettings = TMGS_SimpleAverage;
+	StreamingTexture->bSinglePhysicalSpace = bSinglePhysicalSpace;
 
 	StreamingTexture->Settings.Init();
 	StreamingTexture->Settings.TileSize = GetTileSize();
diff --git a/Engine/Source/Runtime/Engine/Private/VT/VirtualTexturePoolConfig.cpp b/Engine/Source/Runtime/Engine/Private/VT/VirtualTexturePoolConfig.cpp
--- a/Engine/Source/Runtime/Engine/Private/VT/VirtualTexturePoolConfig.cpp
+++ b/Engine/Source/Runtime/Engine/Private/VT/VirtualTexturePoolConfig.cpp
@@ -16,15 +16,27 @@ void UVirtualTexturePoolConfig::PostLoad()
 	DefaultConfig.SizeInMegabyte = DefaultSizeInMegabyte;
 }
 
-const FVirtualTextureSpacePoolConfig *UVirtualTexturePoolConfig::FindPoolConfig(int32 TileSize, EPixelFormat Format)
+const FVirtualTextureSpacePoolConfig *UVirtualTexturePoolConfig::FindPoolConfig(TEnumAsByte<EPixelFormat> const* Formats, int32 NumLayers, int32 TileSize)
 {
 	// Reverse iterate so that project config can override base config
 	for (int32 Id = Pools.Num() - 1; Id >= 0 ; Id--)
 	{
 		const FVirtualTextureSpacePoolConfig& Config = Pools[Id];
-		if (Config.TileSize == TileSize && Config.Format == Format)
+		if (Config.MinTileSize <= TileSize && Config.MaxTileSize >= TileSize && NumLayers == Config.Formats.Num())
 		{
-			return &Config;
+			bool bAllFormatsMatch = true;
+			for (int Layer = 0; Layer < NumLayers && bAllFormatsMatch; ++Layer)
+			{
+				if (Formats[Layer] != Config.Formats[Layer])
+				{
+					bAllFormatsMatch = false;
+				}
+			}
+
+			if (bAllFormatsMatch)
+			{
+				return &Config;
+			}
 		}
 	}
 
diff --git a/Engine/Source/Runtime/Engine/Public/MaterialCompiler.h b/Engine/Source/Runtime/Engine/Public/MaterialCompiler.h
--- a/Engine/Source/Runtime/Engine/Public/MaterialCompiler.h
+++ b/Engine/Source/Runtime/Engine/Public/MaterialCompiler.h
@@ -206,8 +206,8 @@ public:
 	virtual int32 Texture(UTexture* Texture,int32& TextureReferenceIndex, EMaterialSamplerType SamplerType, ESamplerSourceMode SamplerSource=SSM_FromTextureAsset,ETextureMipValueMode MipValueMode=TMVM_None) = 0;
 	virtual int32 TextureParameter(FName ParameterName,UTexture* DefaultTexture,int32& TextureReferenceIndex, EMaterialSamplerType SamplerType, ESamplerSourceMode SamplerSource=SSM_FromTextureAsset) = 0;
 
-	virtual int32 VirtualTexture(URuntimeVirtualTexture* InTexture, int32 LayerIndex, int32& TextureReferenceIndex, EMaterialSamplerType SamplerType) = 0;
-	virtual int32 VirtualTextureParameter(FName ParameterName, URuntimeVirtualTexture* DefaultValue, int32 LayerIndex, int32& TextureReferenceIndex, EMaterialSamplerType SamplerType) = 0;
+	virtual int32 VirtualTexture(URuntimeVirtualTexture* InTexture, int32 TextureLayerIndex, int32 PageTableLayerIndex, int32& TextureReferenceIndex, EMaterialSamplerType SamplerType) = 0;
+	virtual int32 VirtualTextureParameter(FName ParameterName, URuntimeVirtualTexture* DefaultValue, int32 TextureLayerIndex, int32 PageTableLayerIndex, int32& TextureReferenceIndex, EMaterialSamplerType SamplerType) = 0;
 	virtual int32 VirtualTextureUniform(int32 TextureIndex, int32 VectorIndex) = 0;
 	virtual int32 VirtualTextureUniform(FName ParameterName, int32 TextureIndex, int32 VectorIndex) = 0;
 	virtual int32 VirtualTextureWorldToUV(int32 WorldPositionIndex, int32 P0, int32 P1, int32 P2) = 0;
@@ -229,16 +229,16 @@ public:
 		return Texture(InTexture, TextureReferenceIndex, SamplerType, SamplerSource);
 	}
 
-	int32 VirtualTexture(URuntimeVirtualTexture* InTexture, int32 LayerIndex, EMaterialSamplerType SamplerType)
+	int32 VirtualTexture(URuntimeVirtualTexture* InTexture, int32 TextureLayerIndex, int32 PageTableLayerIndex, EMaterialSamplerType SamplerType)
 	{
 		int32 TextureReferenceIndex = INDEX_NONE;
-		return VirtualTexture(InTexture, LayerIndex, TextureReferenceIndex, SamplerType);
+		return VirtualTexture(InTexture, TextureLayerIndex, PageTableLayerIndex, TextureReferenceIndex, SamplerType);
 	}
 
-	int32 VirtualTextureParameter(FName ParameterName, URuntimeVirtualTexture* DefaultValue, int32 LayerIndex, EMaterialSamplerType SamplerType)
+	int32 VirtualTextureParameter(FName ParameterName, URuntimeVirtualTexture* DefaultValue, int32 TextureLayerIndex, int32 PageTableLayerIndex, EMaterialSamplerType SamplerType)
 	{
 		int32 TextureReferenceIndex = INDEX_NONE;
-		return VirtualTextureParameter(ParameterName, DefaultValue, LayerIndex, TextureReferenceIndex, SamplerType);
+		return VirtualTextureParameter(ParameterName, DefaultValue, TextureLayerIndex, PageTableLayerIndex, TextureReferenceIndex, SamplerType);
 	}
 
 	int32 ExternalTexture(UTexture* DefaultTexture)
@@ -487,8 +487,14 @@ public:
 	virtual int32 Texture(UTexture* InTexture,int32& TextureReferenceIndex,EMaterialSamplerType SamplerType,ESamplerSourceMode SamplerSource=SSM_FromTextureAsset,ETextureMipValueMode MipValueMode=TMVM_None) override { return Compiler->Texture(InTexture,TextureReferenceIndex, SamplerType, SamplerSource,MipValueMode); }
 	virtual int32 TextureParameter(FName ParameterName,UTexture* DefaultValue,int32& TextureReferenceIndex, EMaterialSamplerType SamplerType, ESamplerSourceMode SamplerSource=SSM_FromTextureAsset) override { return Compiler->TextureParameter(ParameterName,DefaultValue,TextureReferenceIndex, SamplerType, SamplerSource); }
 
-	virtual int32 VirtualTexture(URuntimeVirtualTexture* InTexture, int32 LayerIndex, int32& TextureReferenceIndex, EMaterialSamplerType SamplerType) override { return Compiler->VirtualTexture(InTexture, LayerIndex, TextureReferenceIndex, SamplerType); }
-	virtual int32 VirtualTextureParameter(FName ParameterName, URuntimeVirtualTexture* DefaultValue, int32 LayerIndex, int32& TextureReferenceIndex, EMaterialSamplerType SamplerType) override { return Compiler->VirtualTextureParameter(ParameterName, DefaultValue, LayerIndex, TextureReferenceIndex, SamplerType); }
+	virtual int32 VirtualTexture(URuntimeVirtualTexture* InTexture, int32 TextureLayerIndex, int32 PageTableLayerIndex, int32& TextureReferenceIndex, EMaterialSamplerType SamplerType) override
+	{
+		return Compiler->VirtualTexture(InTexture, TextureLayerIndex, PageTableLayerIndex, TextureReferenceIndex, SamplerType); 
+	}
+	virtual int32 VirtualTextureParameter(FName ParameterName, URuntimeVirtualTexture* DefaultValue, int32 TextureLayerIndex, int32 PageTableLayerIndex, int32& TextureReferenceIndex, EMaterialSamplerType SamplerType) override
+	{
+		return Compiler->VirtualTextureParameter(ParameterName, DefaultValue, TextureLayerIndex, PageTableLayerIndex, TextureReferenceIndex, SamplerType); 
+	}
 	virtual int32 VirtualTextureUniform(int32 TextureIndex, int32 VectorIndex) override { return Compiler->VirtualTextureUniform(TextureIndex, VectorIndex); }
 	virtual int32 VirtualTextureUniform(FName ParameterName, int32 TextureIndex, int32 VectorIndex) override { return Compiler->VirtualTextureUniform(ParameterName, TextureIndex, VectorIndex); }
 	virtual int32 VirtualTextureWorldToUV(int32 WorldPositionIndex, int32 P0, int32 P1, int32 P2) override { return Compiler->VirtualTextureWorldToUV(WorldPositionIndex, P0, P1, P2); }
diff --git a/Engine/Source/Runtime/Engine/Public/MaterialShared.h b/Engine/Source/Runtime/Engine/Public/MaterialShared.h
--- a/Engine/Source/Runtime/Engine/Public/MaterialShared.h
+++ b/Engine/Source/Runtime/Engine/Public/MaterialShared.h
@@ -275,14 +275,14 @@ public:
 /**
  * A texture expression.
  */
-class FMaterialUniformExpressionTexture: public FMaterialUniformExpression
+class FMaterialUniformExpressionTexture : public FMaterialUniformExpression
 {
 	DECLARE_MATERIALUNIFORMEXPRESSION_TYPE(FMaterialUniformExpressionTexture);
 
 public:
 	FMaterialUniformExpressionTexture();
 	FMaterialUniformExpressionTexture(int32 InTextureIndex, EMaterialSamplerType InSamplerType, ESamplerSourceMode InSamplerSource, bool InVirtualTexture);
-	FMaterialUniformExpressionTexture(int32 InTextureIndex, int32 InLayerIndex, EMaterialSamplerType InSamplerType);
+	FMaterialUniformExpressionTexture(int32 InTextureIndex, int16 InTextureLayerIndex, int16 InPageTableLayerIndex, EMaterialSamplerType InSamplerType);
 
 	/** Sets an override texture. */
 	void SetTransientOverrideTextureValue(UTexture* InOverrideTexture);
@@ -298,8 +298,10 @@ public:
 
 	/** Gets texture index which is the index in the full set of referenced textures for this material. */
 	int32 GetTextureIndex() const { return TextureIndex; }
-	/** Gets the layer index in the virtual texture stack if this is fixed. If we don't have a fixed layer then we will allocate during compilation (and not store here). */
-	int32 GetLayerIndex() const { return LayerIndex; }
+	/** Gets the texture layer index in the virtual texture stack if this is fixed. If we don't have a fixed layer then we will allocate during compilation (and not store here). */
+	int32 GetTextureLayerIndex() const { return TextureLayerIndex; }
+	/** Gets the page table channel index in the virtual texture stack if this is fixed. If we don't have a fixed layer then we will allocate during compilation (and not store here). */
+	int32 GetPageTableLayerIndex() const { return PageTableLayerIndex; }
 
 #if WITH_EDITORONLY_DATA
 	/** Get the sampling/decoding logic to compile in the shader for this texture. */
@@ -313,15 +315,17 @@ public:
 	virtual void GetTextureValue(const FMaterialRenderContext& Context, const FMaterial& Material, const UTexture*& OutValue) const;
 	/** Get the URuntimeVirtualTexture referenced by this expression (can be a nullptr). */
 	virtual void GetTextureValue(const FMaterialRenderContext& Context, const FMaterial& Material, const URuntimeVirtualTexture*& OutValue) const;
-	
+
 	/** Get the UTexture referenced by this expression including any transient override logic. */
 	virtual void GetGameThreadTextureValue(const UMaterialInterface* MaterialInterface, const FMaterial& Material, UTexture*& OutValue, bool bAllowOverride = true) const;
 
 protected:
 	/** Index into FMaterial::GetReferencedTextures */
 	int32 TextureIndex;
-	/** Fixed layer in virtual texture stack if preallocated */
-	int32 LayerIndex;
+	/** Texture layer index in virtual texture stack if preallocated */
+	int16 TextureLayerIndex;
+	/** Page table layer index in virtual texture stack if preallocated */
+	int16 PageTableLayerIndex;
 #if WITH_EDITORONLY_DATA
 	/** Sampler logic for this expression */
 	EMaterialSamplerType SamplerType;
diff --git a/Engine/Source/Runtime/RenderCore/Public/VirtualTexturing.h b/Engine/Source/Runtime/RenderCore/Public/VirtualTexturing.h
--- a/Engine/Source/Runtime/RenderCore/Public/VirtualTexturing.h
+++ b/Engine/Source/Runtime/RenderCore/Public/VirtualTexturing.h
@@ -42,27 +42,28 @@ inline bool operator!=(const FVirtualTextureProducerHandle& Lhs, const FVirtualT
  */
 struct FAllocatedVTDescription
 {
-	FVirtualTextureProducerHandle ProducerHandle[VIRTUALTEXTURE_SPACE_MAXLAYERS];
 	uint32 TileSize = 0u;
 	uint32 TileBorderSize = 0u;
 	uint8 Dimensions = 0u;
-	uint8 NumLayers = 0u;
+	uint8 NumTextureLayers = 0u;
+	
+	/** Producer for each texture layer. */
+	FVirtualTextureProducerHandle ProducerHandle[VIRTUALTEXTURE_SPACE_MAXLAYERS];
+	/** Local layer inside producer for each texture layer. */
+	uint8 ProducerLayerIndex[VIRTUALTEXTURE_SPACE_MAXLAYERS] = { 0u };
+
 	union
 	{
 		uint8 PackedFlags = 0u;
 		struct
 		{
 			/**
-			 * Should the AllocatedVT create its own dedicated page table allocation?  The system only supports a limited number of unique page tables, so this must be used carefully
+			 * Should the AllocatedVT create its own dedicated page table allocation? Can be useful to control total allocation.
+			 * The system only supports a limited number of unique page tables, so this must be used carefully
 			 */
 			uint8 bPrivateSpace : 1;
 		};
 	};
-	
-	/**
-	 * This allows arbitrary layers of each IVirtualTexture producer to map to each layer of the allocated VT
-	 */
-	uint8 LocalLayerToProduce[VIRTUALTEXTURE_SPACE_MAXLAYERS] = { 0u };
 };
 
 inline bool operator==(const FAllocatedVTDescription& Lhs, const FAllocatedVTDescription& Rhs)
@@ -70,15 +71,15 @@ inline bool operator==(const FAllocatedVTDescription& Lhs, const FAllocatedVTDes
 	if (Lhs.TileSize != Rhs.TileSize ||
 		Lhs.TileBorderSize != Rhs.TileBorderSize ||
 		Lhs.Dimensions != Rhs.Dimensions ||
-		Lhs.NumLayers != Rhs.NumLayers ||
+		Lhs.NumTextureLayers != Rhs.NumTextureLayers ||
 		Lhs.PackedFlags != Rhs.PackedFlags)
 	{
 		return false;
 	}
-	for (uint32 LayerIndex = 0u; LayerIndex < Lhs.NumLayers; ++LayerIndex)
+	for (uint32 LayerIndex = 0u; LayerIndex < Lhs.NumTextureLayers; ++LayerIndex)
 	{
 		if (Lhs.ProducerHandle[LayerIndex] != Rhs.ProducerHandle[LayerIndex] ||
-			Lhs.LocalLayerToProduce[LayerIndex] != Rhs.LocalLayerToProduce[LayerIndex])
+			Lhs.ProducerLayerIndex[LayerIndex] != Rhs.ProducerLayerIndex[LayerIndex])
 		{
 			return false;
 		}
@@ -93,9 +94,11 @@ inline bool operator!=(const FAllocatedVTDescription& Lhs, const FAllocatedVTDes
 struct FVTProducerDescription
 {
 	FName Name; /** Will be name of UTexture for streaming VTs, mostly here for debugging */
+	
 	bool bPersistentHighestMip = true;
 	bool bContinuousUpdate = false;
 	bool bCreateRenderTarget = false;
+	
 	uint32 TileSize = 0u;
 	uint32 TileBorderSize = 0u;
 
@@ -111,9 +114,19 @@ struct FVTProducerDescription
 	uint16 WidthInBlocks = 1u;
 	uint16 HeightInBlocks = 1u;
 	uint8 Dimensions = 0u;
-	uint8 NumLayers = 0u;
 	uint8 MaxLevel = 0u;
+
+	/**
+	 * Producers will fill a number of texture layers.
+	 * These texture layers can be distributed across one or more physical groups.
+	 * Each physical group can contain one or more of the texture layers. 
+	 * Within a physical group the texture layers share the same UV allocation/mapping and can be referenced by a single page table lookup.
+	 */
+	uint8 NumTextureLayers = 0u;
 	TEnumAsByte<EPixelFormat> LayerFormat[VIRTUALTEXTURE_SPACE_MAXLAYERS] = { PF_Unknown };
+	
+	uint8 NumPhysicalGroups = 0u;
+	uint8 PhysicalGroupIndex[VIRTUALTEXTURE_SPACE_MAXLAYERS] = { 0 };
 };
 
 typedef void (FVTProducerDestroyedFunction)(const FVirtualTextureProducerHandle& InHandle, void* Baton);
@@ -287,6 +300,7 @@ public:
 	virtual FRHITexture* GetPhysicalTexture(uint32 InLayerIndex) const = 0;
 	virtual FRHIShaderResourceView* GetPhysicalTextureView(uint32 InLayerIndex, bool bSRGB) const = 0;
 	virtual uint32 GetPhysicalTextureSize(uint32 InLayerIndex) const = 0;
+	virtual uint32 GetNumPageTableTextures() const = 0;
 
 	/** Writes 2x FUintVector4 */
 	virtual void GetPackedPageTableUniform(FUintVector4* OutUniform, bool bApplyBlockScale) const = 0;
@@ -297,13 +311,12 @@ public:
 	virtual void DumpToConsole(bool bVerbose) const {}
 
 	inline const FAllocatedVTDescription& GetDescription() const { return Description; }
-	inline const FVirtualTextureProducerHandle& GetProducerHandle(uint32 InLayerIndex) const { check(InLayerIndex < Description.NumLayers); return Description.ProducerHandle[InLayerIndex]; }
-	inline uint32 GetLocalLayerToProduce(uint32 InLayerIndex) const { check(InLayerIndex < Description.NumLayers); return Description.LocalLayerToProduce[InLayerIndex]; }
+	inline const FVirtualTextureProducerHandle& GetProducerHandle(uint32 InLayerIndex) const { check(InLayerIndex < Description.NumTextureLayers); return Description.ProducerHandle[InLayerIndex]; }
 	
 	inline uint32 GetVirtualTileSize() const { return Description.TileSize; }
 	inline uint32 GetTileBorderSize() const { return Description.TileBorderSize; }
 	inline uint32 GetPhysicalTileSize() const { return Description.TileSize + Description.TileBorderSize * 2u; }
-	inline uint32 GetNumLayers() const { return Description.NumLayers; }
+	inline uint32 GetNumTextureLayers() const { return Description.NumTextureLayers; }
 	inline uint8 GetDimensions() const { return Description.Dimensions; }
 	inline uint32 GetWidthInBlocks() const { return WidthInBlocks; }
 	inline uint32 GetHeightInBlocks() const { return HeightInBlocks; }
@@ -313,7 +326,6 @@ public:
 	inline uint32 GetWidthInPixels() const { return GetWidthInTiles() * Description.TileSize; }
 	inline uint32 GetHeightInPixels() const { return GetHeightInTiles() * Description.TileSize; }
 	inline uint32 GetDepthInPixels() const { return DepthInTiles * Description.TileSize; }
-	inline uint32 GetNumPageTableTextures() const { return (Description.NumLayers + LayersPerPageTableTexture - 1u) / LayersPerPageTableTexture; }
 	inline uint32 GetSpaceID() const { return SpaceID; }
 	inline uint32 GetVirtualAddress() const { return VirtualAddress; }
 	inline uint32 GetMaxLevel() const { return MaxLevel; }
diff --git a/Engine/Source/Runtime/Renderer/Private/VT/AllocatedVirtualTexture.cpp b/Engine/Source/Runtime/Renderer/Private/VT/AllocatedVirtualTexture.cpp
--- a/Engine/Source/Runtime/Renderer/Private/VT/AllocatedVirtualTexture.cpp
+++ b/Engine/Source/Runtime/Renderer/Private/VT/AllocatedVirtualTexture.cpp
@@ -15,32 +15,37 @@ FAllocatedVirtualTexture::FAllocatedVirtualTexture(FVirtualTextureSystem* InSyst
 	uint32 InWidthInBlocks,
 	uint32 InHeightInBlocks,
 	uint32 InDepthInTiles)
-	: IAllocatedVirtualTexture(InDesc, InSpace->GetID(), InSpace->GetDescription().Format, InBlockWidthInTiles, InBlockHeightInTiles, InWidthInBlocks, InHeightInBlocks, InDepthInTiles)
-	, Space(InSpace)
+	: IAllocatedVirtualTexture(InDesc, InSpace->GetID(), InSpace->GetPageTableFormat(), InBlockWidthInTiles, InBlockHeightInTiles, InWidthInBlocks, InHeightInBlocks, InDepthInTiles)
 	, RefCount(1)
 	, FrameAllocated(InFrame)
-	, NumUniqueProducers(0u)
+	, Space(InSpace)
 {
 	check(IsInRenderingThread());
-	FMemory::Memzero(UniqueProducerHandles);
-	FMemory::Memzero(UniqueProducerMipBias);
+	FMemory::Memzero(TextureLayers);
 
-	for (uint32 LayerIndex = 0u; LayerIndex < Description.NumLayers; ++LayerIndex)
+	for (uint32 LayerIndex = 0u; LayerIndex < Description.NumTextureLayers; ++LayerIndex)
 	{
-		FVirtualTextureProducer* Producer = InProducers[LayerIndex];
+		FVirtualTextureProducer const* Producer = InProducers[LayerIndex];
 		if (Producer)
 		{
-			PhysicalSpace[LayerIndex] = Producer->GetPhysicalSpace(InDesc.LocalLayerToProduce[LayerIndex]);
-			UniqueProducerIndexForLayer[LayerIndex] = AddUniqueProducer(InDesc.ProducerHandle[LayerIndex], Producer);
-		}
-		else
-		{
-			UniqueProducerIndexForLayer[LayerIndex] = 0xff;
+			const uint32 UniqueProducerIndex = AddUniqueProducer(InDesc.ProducerHandle[LayerIndex], Producer->GetDescription());
+
+			const int32 ProducerLayerIndex = InDesc.ProducerLayerIndex[LayerIndex];
+			const uint32 ProducerPhysicalGroupIndex = Producer->GetPhysicalGroupIndexForTextureLayer(ProducerLayerIndex);
+			FVirtualTexturePhysicalSpace* PhysicalSpace = Producer->GetPhysicalSpaceForPhysicalGroup(ProducerPhysicalGroupIndex);
+			const uint32 UniquePhysicalSpaceIndex = AddUniquePhysicalSpace(PhysicalSpace, UniqueProducerIndex, ProducerPhysicalGroupIndex);
+			UniquePageTableLayers[UniquePhysicalSpaceIndex].ProducerTextureLayerMask |= 1 << ProducerLayerIndex;
+			const uint8 PageTableLayerLocalIndex = UniquePageTableLayers[UniquePhysicalSpaceIndex].TextureLayerCount++;
+			
+			TextureLayers[LayerIndex].UniquePageTableLayerIndex = UniquePhysicalSpaceIndex;
+			TextureLayers[LayerIndex].PhysicalTextureIndex = PageTableLayerLocalIndex;
 		}
 	}
 
 	// Must have at least 1 valid layer/producer
-	check(NumUniqueProducers > 0u);
+	check(UniqueProducers.Num() > 0u);
+	// Layout should match the FVirtualTextureSpace
+	check(UniquePageTableLayers.Num() == Space->GetNumPageTableLayers());
 
 	// Max level of overall allocated VT is limited by size in tiles
 	// With multiple layers of different sizes, some layers may have mips smaller than a single tile
@@ -49,13 +54,13 @@ FAllocatedVirtualTexture::FAllocatedVirtualTexture(FVirtualTextureSystem* InSyst
 	// Lock lowest resolution mip from each producer
 	// Depending on the block dimensions of the producers that make up this allocated VT, different allocated VTs may need to lock different low resolution mips from the same producer
 	// In the common case where block dimensions match, same mip will be locked by all allocated VTs that make use of the same producer
-	for (uint32 LayerIndex = 0u; LayerIndex < Description.NumLayers; ++LayerIndex)
+	for (int32 ProducerIndex = 0u; ProducerIndex < UniqueProducers.Num(); ++ProducerIndex)
 	{
-		FVirtualTextureProducer* Producer = InProducers[LayerIndex];
+		FVirtualTextureProducerHandle ProducerHandle = UniqueProducers[ProducerIndex].Handle;
+		FVirtualTextureProducer* Producer = InSystem->FindProducer(ProducerHandle);
 		if (Producer && Producer->GetDescription().bPersistentHighestMip)
 		{
-			const uint32 ProducerIndex = UniqueProducerIndexForLayer[LayerIndex];
-			const uint32 MipBias = UniqueProducerMipBias[ProducerIndex];
+			const uint32 MipBias = UniqueProducers[ProducerIndex].MipBias;
 			check(MipBias <= MaxLevel);
 			const uint32 Local_vLevel = MaxLevel - MipBias;
 			check(Local_vLevel <= Producer->GetMaxLevel());
@@ -67,7 +72,7 @@ FAllocatedVirtualTexture::FAllocatedVirtualTexture(FVirtualTextureSystem* InSyst
 				for (uint32 TileX = 0u; TileX < RootWidthInTiles; ++TileX)
 				{
 					const uint32 Local_vAddress = FMath::MortonCode2(TileX) | (FMath::MortonCode2(TileY) << 1);
-					const FVirtualTextureLocalTile TileToUnlock(InDesc.ProducerHandle[LayerIndex], Local_vAddress, Local_vLevel);
+					const FVirtualTextureLocalTile TileToUnlock(ProducerHandle, Local_vAddress, Local_vLevel);
 					InSystem->LockTile(TileToUnlock);
 				}
 			}
@@ -97,14 +102,13 @@ void FAllocatedVirtualTexture::Release(FVirtualTextureSystem* System)
 	check(RefCount.GetValue() == 0);
 
 	// Unlock any locked tiles
-	for (uint32 LayerIndex = 0u; LayerIndex < Description.NumLayers; ++LayerIndex)
+	for (int32 ProducerIndex = 0u; ProducerIndex < UniqueProducers.Num(); ++ProducerIndex)
 	{
-		const FVirtualTextureProducerHandle ProducerHandle = GetProducerHandle(LayerIndex);
+		const FVirtualTextureProducerHandle ProducerHandle = UniqueProducers[ProducerIndex].Handle;
 		FVirtualTextureProducer* Producer = System->FindProducer(ProducerHandle);
 		if (Producer && Producer->GetDescription().bPersistentHighestMip)
 		{
-			const uint32 ProducerIndex = UniqueProducerIndexForLayer[LayerIndex];
-			const uint32 MipBias = UniqueProducerMipBias[ProducerIndex];
+			const uint32 MipBias = UniqueProducers[ProducerIndex].MipBias;
 			check(MipBias <= MaxLevel);
 			const uint32 Local_vLevel = MaxLevel - MipBias;
 			check(Local_vLevel <= Producer->GetMaxLevel());
@@ -121,15 +125,31 @@ void FAllocatedVirtualTexture::Release(FVirtualTextureSystem* System)
 				}
 			}
 		}
+	}
+
+	// Physical pool needs to evict all pages that belong to this VT's space
+	//todo[vt]: Could improve this to only evict pages belonging to this VT
+	{
+		TArray<FVirtualTexturePhysicalSpace*> UniquePhysicalSpaces;
+		for (int32 PageTableIndex = 0u; PageTableIndex < UniquePageTableLayers.Num(); ++PageTableIndex)
+		{
+			UniquePhysicalSpaces.Add(UniquePageTableLayers[PageTableIndex].PhysicalSpace);
+		}
+
+		for (FVirtualTexturePhysicalSpace* PhysicalSpace : UniquePhysicalSpaces)
+		{
+			PhysicalSpace->GetPagePool().UnmapAllPagesForSpace(System, Space->GetID());
+
+			for (int32 PageTableLayerIndex = 0u; PageTableLayerIndex < UniquePageTableLayers.Num(); ++PageTableLayerIndex)
+			{
+				FTexturePageMap& PageMap = Space->GetPageMapForPageTableLayer(PageTableLayerIndex);
+				PageMap.VerifyPhysicalSpaceUnmapped(PhysicalSpace->GetID());
+			}
+		}
 
-		// Physical pool needs to evict all pages that belong to this VT's space
-		// TODO - could improve this to only evict pages belonging to this VT
-		if (PhysicalSpace[LayerIndex])
+		for (int32 PageTableIndex = 0u; PageTableIndex < UniquePageTableLayers.Num(); ++PageTableIndex)
 		{
-			FTexturePageMap& PageMap = Space->GetPageMap(LayerIndex);
-			PhysicalSpace[LayerIndex]->GetPagePool().UnmapAllPagesForSpace(System, Space->GetID());
-			PageMap.VerifyPhysicalSpaceUnmapped(PhysicalSpace[LayerIndex]->GetID());
-			PhysicalSpace[LayerIndex].SafeRelease();
+			UniquePageTableLayers[PageTableIndex].PhysicalSpace.SafeRelease();
 		}
 	}
 
@@ -140,41 +160,62 @@ void FAllocatedVirtualTexture::Release(FVirtualTextureSystem* System)
 	delete this;
 }
 
-uint32 FAllocatedVirtualTexture::AddUniqueProducer(const FVirtualTextureProducerHandle& InHandle, FVirtualTextureProducer* InProducer)
+uint32 FAllocatedVirtualTexture::AddUniqueProducer(FVirtualTextureProducerHandle const& InHandle, FVTProducerDescription const& InProducerDesc)
 {
-	for (uint32 Index = 0u; Index < NumUniqueProducers; ++Index)
+	for (int32 Index = 0u; Index < UniqueProducers.Num(); ++Index)
 	{
-		if (UniqueProducerHandles[Index] == InHandle)
+		if (UniqueProducers[Index].Handle == InHandle)
 		{
 			return Index;
 		}
 	}
-	const uint32 Index = NumUniqueProducers++;
+	const uint32 Index = UniqueProducers.AddDefaulted();
 	check(Index < VIRTUALTEXTURE_SPACE_MAXLAYERS);
 	
-	const FVTProducerDescription& ProducerDesc = InProducer->GetDescription();
-
 	// maybe these values should just be set by producers, rather than also set on AllocatedVT desc
-	check(ProducerDesc.Dimensions == Description.Dimensions);
-	check(ProducerDesc.TileSize == Description.TileSize);
-	check(ProducerDesc.TileBorderSize == Description.TileBorderSize);
+	check(InProducerDesc.Dimensions == Description.Dimensions);
+	check(InProducerDesc.TileSize == Description.TileSize);
+	check(InProducerDesc.TileBorderSize == Description.TileBorderSize);
 
 	const uint32 BlockSizeInTiles = FMath::Max(BlockWidthInTiles, BlockHeightInTiles);
-	const uint32 ProducerBlockSizeInTiles = FMath::Max(ProducerDesc.BlockWidthInTiles, ProducerDesc.BlockHeightInTiles);
+	const uint32 ProducerBlockSizeInTiles = FMath::Max(InProducerDesc.BlockWidthInTiles, InProducerDesc.BlockHeightInTiles);
 	const uint32 MipBias = FMath::CeilLogTwo(BlockSizeInTiles / ProducerBlockSizeInTiles);
 
 	check((BlockSizeInTiles / ProducerBlockSizeInTiles) * ProducerBlockSizeInTiles == BlockSizeInTiles);
-	check(ProducerDesc.BlockWidthInTiles << MipBias == BlockWidthInTiles);
-	check(ProducerDesc.BlockHeightInTiles << MipBias == BlockHeightInTiles);
+	check(InProducerDesc.BlockWidthInTiles << MipBias == BlockWidthInTiles);
+	check(InProducerDesc.BlockHeightInTiles << MipBias == BlockHeightInTiles);
 
-	MaxLevel = FMath::Max<uint32>(MaxLevel, ProducerDesc.MaxLevel + MipBias);
+	MaxLevel = FMath::Max<uint32>(MaxLevel, InProducerDesc.MaxLevel + MipBias);
 
-	UniqueProducerHandles[Index] = InHandle;
-	UniqueProducerMipBias[Index] = MipBias;
+	UniqueProducers[Index].Handle = InHandle;
+	UniqueProducers[Index].MipBias = MipBias;
 	
 	return Index;
 }
 
+uint32 FAllocatedVirtualTexture::AddUniquePhysicalSpace(FVirtualTexturePhysicalSpace* InPhysicalSpace, uint32 InUniqueProducerIndex, uint32 InProducerPhysicalSpaceIndex)
+{
+	for (int32 Index = 0u; Index < UniquePageTableLayers.Num(); ++Index)
+	{
+		if (UniquePageTableLayers[Index].PhysicalSpace == InPhysicalSpace && 
+			UniquePageTableLayers[Index].UniqueProducerIndex == InUniqueProducerIndex &&
+			UniquePageTableLayers[Index].ProducerPhysicalGroupIndex == InProducerPhysicalSpaceIndex)
+		{
+			return Index;
+		}
+	}
+	const uint32 Index = UniquePageTableLayers.AddDefaulted();
+	check(Index < VIRTUALTEXTURE_SPACE_MAXLAYERS);
+
+	UniquePageTableLayers[Index].PhysicalSpace = InPhysicalSpace;
+	UniquePageTableLayers[Index].UniqueProducerIndex = InUniqueProducerIndex;
+	UniquePageTableLayers[Index].ProducerPhysicalGroupIndex = InProducerPhysicalSpaceIndex;
+	UniquePageTableLayers[Index].ProducerTextureLayerMask = 0;
+	UniquePageTableLayers[Index].TextureLayerCount = 0;
+
+	return Index;
+}
+
 FRHITexture* FAllocatedVirtualTexture::GetPageTableTexture(uint32 InPageTableIndex) const
 {
 	return Space->GetPageTableTexture(InPageTableIndex);
@@ -182,34 +223,39 @@ FRHITexture* FAllocatedVirtualTexture::GetPageTableTexture(uint32 InPageTableInd
 
 FRHITexture* FAllocatedVirtualTexture::GetPhysicalTexture(uint32 InLayerIndex) const
 {
-	if (InLayerIndex < Description.NumLayers)
+	if (InLayerIndex < Description.NumTextureLayers)
 	{
-		const FVirtualTexturePhysicalSpace* LayerSpace = PhysicalSpace[InLayerIndex];
-		return LayerSpace ? LayerSpace->GetPhysicalTexture() : nullptr;
+		const FVirtualTexturePhysicalSpace* PhysicalSpace = UniquePageTableLayers[TextureLayers[InLayerIndex].UniquePageTableLayerIndex].PhysicalSpace;
+		return PhysicalSpace ? PhysicalSpace->GetPhysicalTexture(TextureLayers[InLayerIndex].PhysicalTextureIndex) : nullptr;
 	}
 	return nullptr;
 }
 
 FRHIShaderResourceView* FAllocatedVirtualTexture::GetPhysicalTextureView(uint32 InLayerIndex, bool bSRGB) const
 {
-	if (InLayerIndex < Description.NumLayers)
+	if (InLayerIndex < Description.NumTextureLayers)
 	{
-		const FVirtualTexturePhysicalSpace* LayerSpace = PhysicalSpace[InLayerIndex];
-		return LayerSpace ? LayerSpace->GetPhysicalTextureView(bSRGB) : nullptr;
+		const FVirtualTexturePhysicalSpace* PhysicalSpace = UniquePageTableLayers[TextureLayers[InLayerIndex].UniquePageTableLayerIndex].PhysicalSpace;
+		return PhysicalSpace ? PhysicalSpace->GetPhysicalTextureView(TextureLayers[InLayerIndex].PhysicalTextureIndex, bSRGB) : nullptr;
 	}
 	return nullptr;
 }
 
 uint32 FAllocatedVirtualTexture::GetPhysicalTextureSize(uint32 InLayerIndex) const
 {
-	if (InLayerIndex < Description.NumLayers)
+	if (InLayerIndex < Description.NumTextureLayers)
 	{
-		const FVirtualTexturePhysicalSpace* LayerSpace = PhysicalSpace[InLayerIndex];
-		return LayerSpace ? LayerSpace->GetTextureSize() : 0u;
+		const FVirtualTexturePhysicalSpace* PhysicalSpace = UniquePageTableLayers[TextureLayers[InLayerIndex].UniquePageTableLayerIndex].PhysicalSpace;
+		return PhysicalSpace ? PhysicalSpace->GetTextureSize() : 0u;
 	}
 	return 0u;
 }
 
+uint32 FAllocatedVirtualTexture::GetNumPageTableTextures() const
+{
+	return Space->GetNumPageTableTextures();
+}
+
 static inline uint32 BitcastFloatToUInt32(float v)
 {
 	const union
diff --git a/Engine/Source/Runtime/Renderer/Private/VT/AllocatedVirtualTexture.h b/Engine/Source/Runtime/Renderer/Private/VT/AllocatedVirtualTexture.h
--- a/Engine/Source/Runtime/Renderer/Private/VT/AllocatedVirtualTexture.h
+++ b/Engine/Source/Runtime/Renderer/Private/VT/AllocatedVirtualTexture.h
@@ -13,7 +13,6 @@ class FVirtualTextureSpace;
 class FVirtualTexturePhysicalSpace;
 class FVirtualTextureProducer;
 class FVirtualTextureSystem;
-class FTexturePageLocks;
 
 class FAllocatedVirtualTexture final : public IAllocatedVirtualTexture
 {
@@ -33,35 +32,68 @@ public:
 
 	void Release(FVirtualTextureSystem* System);
 
+	inline void IncrementRefCount() { RefCount.Increment(); }
+	inline uint32 GetFrameAllocated() const { return FrameAllocated; }
+
 	// begin IAllocatedVirtualTexture
 	virtual FRHITexture* GetPageTableTexture(uint32 InPageTableIndex) const override;
 	virtual FRHITexture* GetPhysicalTexture(uint32 InLayerIndex) const override;
 	virtual FRHIShaderResourceView* GetPhysicalTextureView(uint32 InLayerIndex, bool bSRGB) const override;
 	virtual uint32 GetPhysicalTextureSize(uint32 InLayerIndex) const override;
+	virtual uint32 GetNumPageTableTextures() const override;
 	virtual void GetPackedPageTableUniform(FUintVector4* OutUniform, bool bApplyBlockScale) const override;
 	virtual void GetPackedUniform(FUintVector4* OutUniform, uint32 LayerIndex) const override;
 	virtual void Destroy(FVirtualTextureSystem* System) override;
 	// end IAllocatedVirtualTexture
 
-	inline void IncrementRefCount() { RefCount.Increment(); }
-
-	inline uint32 GetFrameAllocated() const { return FrameAllocated; }
 	inline FVirtualTextureSpace* GetSpace() const { return Space; }
-	inline FVirtualTexturePhysicalSpace* GetPhysicalSpace(uint32 InLayer) const { checkSlow(InLayer < Description.NumLayers); return PhysicalSpace[InLayer]; }
-	inline uint32 GetNumUniqueProducers() const { return NumUniqueProducers; }
-	inline uint32 GetUniqueProducerIndexForLayer(uint32 InLayer) const { return UniqueProducerIndexForLayer[InLayer]; }
-	inline const FVirtualTextureProducerHandle& GetUniqueProducerHandle(uint32 InProducerIndex) const { checkSlow(InProducerIndex < NumUniqueProducers); return UniqueProducerHandles[InProducerIndex]; }
-	inline uint8 GetUniqueProducerMipBias(uint32 InProducerIndex) const { checkSlow(InProducerIndex < NumUniqueProducers); return UniqueProducerMipBias[InProducerIndex]; }
+
+	inline uint32 GetNumUniqueProducers() const { return UniqueProducers.Num(); }
+	inline const FVirtualTextureProducerHandle& GetUniqueProducerHandle(uint32 InProducerIndex) const { return UniqueProducers[InProducerIndex].Handle; }
+	inline uint8 GetUniqueProducerMipBias(uint32 InProducerIndex) const { return UniqueProducers[InProducerIndex].MipBias; }
+
+	inline uint32 GetNumPageTableLayers() const { return UniquePageTableLayers.Num(); }
+	inline FVirtualTexturePhysicalSpace* GetPhysicalSpaceForPageTableLayer(uint32 InLayerIndex) const { return UniquePageTableLayers[InLayerIndex].PhysicalSpace; }
+	inline uint32 GetProducerIndexForPageTableLayer(uint32 InLayerIndex) const { return UniquePageTableLayers[InLayerIndex].UniqueProducerIndex; }
+	inline uint32 GetProducerTextureLayerMaskForPageTableLayer(uint32 InLayerIndex) const { return UniquePageTableLayers[InLayerIndex].ProducerTextureLayerMask; }
+	inline uint32 GetProducerPhysicalGroupIndexForPageTableLayer(uint32 InLayerIndex) const { return UniquePageTableLayers[InLayerIndex].ProducerPhysicalGroupIndex; }
 
 private:
-	uint32 AddUniqueProducer(const FVirtualTextureProducerHandle& InHandle, FVirtualTextureProducer* InProducer);
+	uint32 AddUniqueProducer(FVirtualTextureProducerHandle const& InHandle, FVTProducerDescription const& InProducerDesc);
+	uint32 AddUniquePhysicalSpace(FVirtualTexturePhysicalSpace* InPhysicalSpace, uint32 InUniqueProducerIndex, uint32 InProducerPhysicalSpaceIndex);
 
-	FVirtualTextureSpace* Space;
-	TRefCountPtr<FVirtualTexturePhysicalSpace> PhysicalSpace[VIRTUALTEXTURE_SPACE_MAXLAYERS];
-	FVirtualTextureProducerHandle UniqueProducerHandles[VIRTUALTEXTURE_SPACE_MAXLAYERS];
-	uint8 UniqueProducerIndexForLayer[VIRTUALTEXTURE_SPACE_MAXLAYERS];
-	uint8 UniqueProducerMipBias[VIRTUALTEXTURE_SPACE_MAXLAYERS];
 	FThreadSafeCounter RefCount;
 	uint32 FrameAllocated;
-	uint8 NumUniqueProducers;
+
+	FVirtualTextureSpace* Space;
+
+	// Unique producers in this AllocatedVT.
+	struct FProducerDesc
+	{
+		FVirtualTextureProducerHandle Handle;
+		uint8 MipBias;
+	};
+	TArray<FProducerDesc> UniqueProducers;
+
+	// Unique page table layers in this AllocatedVT.
+	// We create a page table layer for each physical space group in each producer.
+	struct FPageTableLayerDesc
+	{
+		FVirtualTexturePhysicalSpace* PhysicalSpace;
+		uint8 UniqueProducerIndex;
+		uint8 ProducerPhysicalGroupIndex;
+		uint16 ProducerTextureLayerMask;
+		uint8 TextureLayerCount;
+	};
+	TArray<FPageTableLayerDesc> UniquePageTableLayers;
+
+	// Texture layers in this AllocatedVT.
+	// Each texture layer maps onto one page table layer.
+	// Each texture layer needs a unique index within the page table layer to describe the unique physical texture it references.
+	struct FTextureLayerDesc
+	{
+		uint8 UniquePageTableLayerIndex;
+		uint8 PhysicalTextureIndex;
+	};
+	FTextureLayerDesc TextureLayers[VIRTUALTEXTURE_SPACE_MAXLAYERS];
 };
diff --git a/Engine/Source/Runtime/Renderer/Private/VT/TexturePagePool.cpp b/Engine/Source/Runtime/Renderer/Private/VT/TexturePagePool.cpp
--- a/Engine/Source/Runtime/Renderer/Private/VT/TexturePagePool.cpp
+++ b/Engine/Source/Runtime/Renderer/Private/VT/TexturePagePool.cpp
@@ -36,7 +36,7 @@ void FTexturePagePool::Initialize(uint32 InNumPages)
 	{
 		FPageMapping& Mapping = PageMapping[i];
 		Mapping.PackedValues = 0xffffffff;
-		Mapping.LayerIndex = 0xff;
+		Mapping.PageTableLayerIndex = 0xff;
 		Mapping.NextIndex = Mapping.PrevIndex = i;
 	}
 }
@@ -65,7 +65,7 @@ void FTexturePagePool::UnmapAllPagesForSpace(FVirtualTextureSystem* System, uint
 	for (int32 MappingIndex = NumPages + 1; MappingIndex < PageMapping.Num(); ++MappingIndex)
 	{
 		struct FPageMapping& Mapping = PageMapping[MappingIndex];
-		if (Mapping.LayerIndex != 0xff && Mapping.SpaceID == SpaceID)
+		if (Mapping.PageTableLayerIndex != 0xff && Mapping.SpaceID == SpaceID)
 		{
 			// we're unmapping all pages for space, so don't try to map any ancestor pages...they'll be unmapped as well
 			UnmapPageMapping(System, MappingIndex, false);
@@ -161,13 +161,13 @@ uint16 FTexturePagePool::GetPageHash(const FPageEntry& Entry)
 	return (uint16)MurmurFinalize64(Entry.PackedValue);
 }
 
-uint32 FTexturePagePool::FindPageAddress(const FVirtualTextureProducerHandle& ProducerHandle, uint8 LayerIndex, uint32 Local_vAddress, uint8 Local_vLevel) const
+uint32 FTexturePagePool::FindPageAddress(const FVirtualTextureProducerHandle& ProducerHandle, uint8 GroupIndex, uint32 Local_vAddress, uint8 Local_vLevel) const
 {
 	FPageEntry CheckPage;
 	CheckPage.PackedProducerHandle = ProducerHandle.PackedValue;
 	CheckPage.Local_vAddress = Local_vAddress;
 	CheckPage.Local_vLevel = Local_vLevel;
-	CheckPage.LayerIndex = LayerIndex;
+	CheckPage.GroupIndex = GroupIndex;
 
 	const uint16 Hash = GetPageHash(CheckPage);
 	for (uint32 PageIndex = PageHash.First(Hash); PageHash.IsValid(PageIndex); PageIndex = PageHash.Next(PageIndex))
@@ -182,11 +182,11 @@ uint32 FTexturePagePool::FindPageAddress(const FVirtualTextureProducerHandle& Pr
 	return ~0u;
 }
 
-uint32 FTexturePagePool::FindNearestPageAddress(const FVirtualTextureProducerHandle& ProducerHandle, uint8 LayerIndex, uint32 Local_vAddress, uint8 Local_vLevel, uint8 MaxLevel) const
+uint32 FTexturePagePool::FindNearestPageAddress(const FVirtualTextureProducerHandle& ProducerHandle, uint8 GroupIndex, uint32 Local_vAddress, uint8 Local_vLevel, uint8 MaxLevel) const
 {
 	while (Local_vLevel <= MaxLevel)
 	{
-		const uint32 pAddress = FindPageAddress(ProducerHandle, LayerIndex, Local_vAddress, Local_vLevel);
+		const uint32 pAddress = FindPageAddress(ProducerHandle, GroupIndex, Local_vAddress, Local_vLevel);
 		if (pAddress != ~0u)
 		{
 			return pAddress;
@@ -198,11 +198,11 @@ uint32 FTexturePagePool::FindNearestPageAddress(const FVirtualTextureProducerHan
 	return ~0u;
 }
 
-uint32 FTexturePagePool::FindNearestPageLevel(const FVirtualTextureProducerHandle& ProducerHandle, uint8 LayerIndex, uint32 Local_vAddress, uint8 Local_vLevel) const
+uint32 FTexturePagePool::FindNearestPageLevel(const FVirtualTextureProducerHandle& ProducerHandle, uint8 GroupIndex, uint32 Local_vAddress, uint8 Local_vLevel) const
 {
 	while (Local_vLevel < 16u)
 	{
-		const uint32 pAddress = FindPageAddress(ProducerHandle, LayerIndex, Local_vAddress, Local_vLevel);
+		const uint32 pAddress = FindPageAddress(ProducerHandle, GroupIndex, Local_vAddress, Local_vLevel);
 		if (pAddress != ~0u)
 		{
 			return Pages[pAddress].Local_vLevel;
@@ -214,11 +214,11 @@ uint32 FTexturePagePool::FindNearestPageLevel(const FVirtualTextureProducerHandl
 	return ~0u;
 }
 
-uint32 FTexturePagePool::Alloc(FVirtualTextureSystem* System, uint32 Frame, const FVirtualTextureProducerHandle& ProducerHandle, uint8 LayerIndex, uint32 Local_vAddress, uint8 Local_vLevel, bool bLock)
+uint32 FTexturePagePool::Alloc(FVirtualTextureSystem* System, uint32 Frame, const FVirtualTextureProducerHandle& ProducerHandle, uint8 GroupIndex, uint32 Local_vAddress, uint8 Local_vLevel, bool bLock)
 {
 	check(ProducerHandle.PackedValue != 0u);
 	check(AnyFreeAvailable(Frame));
-	checkSlow(FindPageAddress(ProducerHandle, LayerIndex, Local_vAddress, Local_vLevel) == ~0u);
+	checkSlow(FindPageAddress(ProducerHandle, GroupIndex, Local_vAddress, Local_vLevel) == ~0u);
 
 	// Grab the LRU free page
 	const uint16 pAddress = FreeHeap.Top();
@@ -231,7 +231,7 @@ uint32 FTexturePagePool::Alloc(FVirtualTextureSystem* System, uint32 Frame, cons
 	PageEntry.PackedProducerHandle = ProducerHandle.PackedValue;
 	PageEntry.Local_vAddress = Local_vAddress;
 	PageEntry.Local_vLevel = Local_vLevel;
-	PageEntry.LayerIndex = LayerIndex;
+	PageEntry.GroupIndex = GroupIndex;
 	PageHash.Add(GetPageHash(PageEntry), pAddress);
 
 	if (bLock)
@@ -281,10 +281,10 @@ void FTexturePagePool::UpdateUsage(uint32 Frame, uint16 pAddress)
 	}
 }
 
-void FTexturePagePool::MapPage(FVirtualTextureSpace* Space, FVirtualTexturePhysicalSpace* PhysicalSpace, uint8 Layer, uint8 vLogSize, uint32 vAddress, uint8 vLevel, uint16 pAddress)
+void FTexturePagePool::MapPage(FVirtualTextureSpace* Space, FVirtualTexturePhysicalSpace* PhysicalSpace, uint8 PageTableLayerIndex, uint8 vLogSize, uint32 vAddress, uint8 vLevel, uint16 pAddress)
 {
 	check(pAddress < NumPages);
-	FTexturePageMap& PageMap = Space->GetPageMap(Layer);
+	FTexturePageMap& PageMap = Space->GetPageMapForPageTableLayer(PageTableLayerIndex);
 	PageMap.MapPage(Space, PhysicalSpace, vLogSize, vAddress, vLevel, pAddress);
 
 	++NumPagesMapped;
@@ -295,14 +295,14 @@ void FTexturePagePool::MapPage(FVirtualTextureSpace* Space, FVirtualTexturePhysi
 	Mapping.SpaceID = Space->GetID();
 	Mapping.vAddress = vAddress;
 	Mapping.vLogSize = vLogSize;
-	Mapping.LayerIndex = Layer;
+	Mapping.PageTableLayerIndex = PageTableLayerIndex;
 }
 
 void FTexturePagePool::UnmapPageMapping(FVirtualTextureSystem* System, uint32 MappingIndex, bool bMapAncestorPage)
 {
 	FPageMapping& Mapping = PageMapping[MappingIndex];
 	FVirtualTextureSpace* Space = System->GetSpace(Mapping.SpaceID);
-	FTexturePageMap& PageMap = Space->GetPageMap(Mapping.LayerIndex);
+	FTexturePageMap& PageMap = Space->GetPageMapForPageTableLayer(Mapping.PageTableLayerIndex);
 
 	PageMap.UnmapPage(System, Space, Mapping.vLogSize, Mapping.vAddress, bMapAncestorPage);
 
@@ -310,7 +310,7 @@ void FTexturePagePool::UnmapPageMapping(FVirtualTextureSystem* System, uint32 Ma
 	--NumPagesMapped;
 
 	Mapping.PackedValues = 0xffffffff;
-	Mapping.LayerIndex = 0xff;
+	Mapping.PageTableLayerIndex = 0xff;
 	ReleaseMapping(MappingIndex);
 }
 
diff --git a/Engine/Source/Runtime/Renderer/Private/VT/TexturePagePool.h b/Engine/Source/Runtime/Renderer/Private/VT/TexturePagePool.h
--- a/Engine/Source/Runtime/Renderer/Private/VT/TexturePagePool.h
+++ b/Engine/Source/Runtime/Renderer/Private/VT/TexturePagePool.h
@@ -73,17 +73,17 @@ public:
 	/**
 	 * Find physical address of the page allocated for the given VT address, or ~0 if not allocated
 	 */
-	uint32		FindPageAddress(const FVirtualTextureProducerHandle& ProducerHandle, uint8 LayerIndex, uint32 Local_vAddress, uint8 Local_vLevel) const;
+	uint32		FindPageAddress(const FVirtualTextureProducerHandle& ProducerHandle, uint8 GroupIndex, uint32 Local_vAddress, uint8 Local_vLevel) const;
 
 	/**
 	 * Find the physical address of the allocated page that's closest to the given page, or ~0 if not found
 	 */
-	uint32		FindNearestPageAddress(const FVirtualTextureProducerHandle& ProducerHandle, uint8 LayerIndex, uint32 Local_vAddress, uint8 Local_vLevel, uint8 MaxLevel) const;
+	uint32		FindNearestPageAddress(const FVirtualTextureProducerHandle& ProducerHandle, uint8 GroupIndex, uint32 Local_vAddress, uint8 Local_vLevel, uint8 MaxLevel) const;
 
 	/**
 	 * Find the level of the allocated page that's closest to the given page, or ~0 if not found
 	 */
-	uint32		FindNearestPageLevel(const FVirtualTextureProducerHandle& ProducerHandle, uint8 LayerIndex, uint32 Local_vAddress, uint8 Local_vLevel) const;
+	uint32		FindNearestPageLevel(const FVirtualTextureProducerHandle& ProducerHandle, uint8 GroupIndex, uint32 Local_vAddress, uint8 Local_vLevel) const;
 	
 	/**
 	 * Allocate a physical address
@@ -91,12 +91,12 @@ public:
 	 * Assuming the pool is full, the returned physical address will first be unmapped from anything that was previously using it
 	 * @param Frame The current frame, used to manage LRU allocations
 	 * @param ProducerHandle Handle of the VT producer requesting this page
-	 * @param LayerIndex Layer in the producer
+	 * @param GroupIndex Physical group in the producer
 	 * @param Local_vAddress Virtual address relative to the given producer
 	 * @param Local_vLevel Mip level in the producer
 	 * @param bLock Should the allocation be locked; locked allocations will never be evicted
 	 */
-	uint32		Alloc(FVirtualTextureSystem* System, uint32 Frame, const FVirtualTextureProducerHandle& ProducerHandle, uint8 LayerIndex, uint32 Local_vAddress, uint8 Local_vLevel, bool bLock);
+	uint32		Alloc(FVirtualTextureSystem* System, uint32 Frame, const FVirtualTextureProducerHandle& ProducerHandle, uint8 GroupIndex, uint32 Local_vAddress, uint8 Local_vLevel, bool bLock);
 
 	/**
 	 * Marks the given physical address as free, will be unlocked if needed, moved to top of LRU list, no longer associated with any producer
@@ -122,7 +122,7 @@ public:
 	/**
 	* Map the physical address to a specific virtual address.
 	*/
-	void		MapPage(FVirtualTextureSpace* Space, FVirtualTexturePhysicalSpace* PhysicalSpace, uint8 Layer, uint8 vLogSize, uint32 vAddress, uint8 vLevel, uint16 pAddress);
+	void		MapPage(FVirtualTextureSpace* Space, FVirtualTexturePhysicalSpace* PhysicalSpace, uint8 PageTableLayerIndex, uint8 vLogSize, uint32 vAddress, uint8 vLevel, uint16 pAddress);
 
 private:
 	// Allocate 24 bits to store next/prev indices, pack layer index into 8 bits
@@ -143,7 +143,7 @@ private:
 	
 		uint32 NextIndex;
 		uint32 PrevIndex : 24;
-		uint32 LayerIndex : 8;
+		uint32 PageTableLayerIndex : 8;
 	};
 
 	union FPageEntry
@@ -154,7 +154,7 @@ private:
 			uint32 PackedProducerHandle;
 			uint32 Local_vAddress : 24;
 			uint32 Local_vLevel : 4;
-			uint32 LayerIndex : 4;
+			uint32 GroupIndex : 4;
 		};
 	};
 
diff --git a/Engine/Source/Runtime/Renderer/Private/VT/UniqueRequestList.h b/Engine/Source/Runtime/Renderer/Private/VT/UniqueRequestList.h
--- a/Engine/Source/Runtime/Renderer/Private/VT/UniqueRequestList.h
+++ b/Engine/Source/Runtime/Renderer/Private/VT/UniqueRequestList.h
@@ -12,8 +12,8 @@
 union FMappingRequest
 {
 	inline FMappingRequest() {}
-	inline FMappingRequest(uint16 InLoadIndex, uint8 InLocalLayerIndex, uint8 InSpaceID, uint8 InLayerIndex, uint32 InAddress, uint8 InLevel, uint8 InLocalLevel)
-		: vAddress(InAddress), vLevel(InLevel), SpaceID(InSpaceID), LoadRequestIndex(InLoadIndex), Local_vLevel(InLocalLevel), LocalLayerIndex(InLocalLayerIndex), LayerIndex(InLayerIndex), Pad(0)
+	inline FMappingRequest(uint16 InLoadIndex, uint8 InPhysicalGroupIndex, uint8 InSpaceID, uint8 InPageTableLayerIndex, uint32 InAddress, uint8 InLevel, uint8 InLocalLevel)
+		: vAddress(InAddress), vLevel(InLevel), SpaceID(InSpaceID), LoadRequestIndex(InLoadIndex), Local_vLevel(InLocalLevel), ProducerPhysicalGroupIndex(InPhysicalGroupIndex), PageTableLayerIndex(InPageTableLayerIndex), Pad(0)
 	{}
 
 	uint64 PackedValue;
@@ -24,8 +24,8 @@ union FMappingRequest
 		uint32 SpaceID : 4;
 		uint32 LoadRequestIndex : 16;
 		uint32 Local_vLevel : 4;
-		uint32 LocalLayerIndex : 4;
-		uint32 LayerIndex : 4;
+		uint32 ProducerPhysicalGroupIndex : 4;
+		uint32 PageTableLayerIndex : 4;
 		uint32 Pad : 4;
 	};
 };
@@ -36,8 +36,8 @@ inline bool operator!=(const FMappingRequest& Lhs, const FMappingRequest& Rhs) {
 union FDirectMappingRequest
 {
 	inline FDirectMappingRequest() {}
-	inline FDirectMappingRequest(uint8 InSpaceID, uint16 InPhysicalSpaceID, uint8 InLayer, uint8 InLogSize, uint32 InAddress, uint8 InLevel, uint16 InPhysicalAddress)
-		: vAddress(InAddress), vLevel(InLevel), SpaceID(InSpaceID), pAddress(InPhysicalAddress), PhysicalSpaceID(InPhysicalSpaceID), Local_vLevel(InLogSize), LayerIndex(InLayer)
+	inline FDirectMappingRequest(uint8 InSpaceID, uint16 InPhysicalSpaceID, uint8 InPageTableLayerIndex, uint8 InLogSize, uint32 InAddress, uint8 InLevel, uint16 InPhysicalAddress)
+		: vAddress(InAddress), vLevel(InLevel), SpaceID(InSpaceID), pAddress(InPhysicalAddress), PhysicalSpaceID(InPhysicalSpaceID), Local_vLevel(InLogSize), PageTableLayerIndex(InPageTableLayerIndex)
 	{}
 
 	uint64 PackedValue;
@@ -49,7 +49,7 @@ union FDirectMappingRequest
 		uint32 pAddress : 16;
 		uint32 PhysicalSpaceID : 8;
 		uint32 Local_vLevel : 4;
-		uint32 LayerIndex : 4;
+		uint32 PageTableLayerIndex : 4;
 	};
 };
 static_assert(sizeof(FDirectMappingRequest) == sizeof(uint64), "Bad packing");
@@ -69,7 +69,7 @@ public:
 		, MappingRequests(new(MemStack) FMappingRequest[MappingRequestCapacity])
 		, DirectMappingRequests(new(MemStack) FDirectMappingRequest[DirectMappingRequestCapacity])
 		, LoadRequestCount(new(MemStack) uint16[LoadRequestCapacity])
-		, LoadRequestLayerMask(new(MemStack) uint8[LoadRequestCapacity])
+		, LoadRequestGroupMask(new(MemStack) uint8[LoadRequestCapacity])
 		, NumLoadRequests(0u)
 		, NumLockRequests(0u)
 		, NumMappingRequests(0u)
@@ -91,15 +91,15 @@ public:
 	inline const FVirtualTextureLocalTile& GetLoadRequest(uint32 i) const { checkSlow(i < NumLoadRequests); return LoadRequests[i]; }
 	inline const FMappingRequest& GetMappingRequest(uint32 i) const { checkSlow(i < NumMappingRequests); return MappingRequests[i]; }
 	inline const FDirectMappingRequest& GetDirectMappingRequest(uint32 i) const { checkSlow(i < NumDirectMappingRequests); return DirectMappingRequests[i]; }
-	inline uint8 GetLocalLayerMask(uint32 i) const { checkSlow(i < NumLoadRequests); return LoadRequestLayerMask[i]; }
+	inline uint8 GetGroupMask(uint32 i) const { checkSlow(i < NumLoadRequests); return LoadRequestGroupMask[i]; }
 	inline bool IsLocked(uint32 i) const { checkSlow(i < NumLoadRequests); return i < NumLockRequests; }
 
-	uint16 AddLoadRequest(const FVirtualTextureLocalTile& Tile, uint8 LayerMask, uint16 Count);
-	uint16 LockLoadRequest(const FVirtualTextureLocalTile& Tile, uint8 LayerMask);
+	uint16 AddLoadRequest(const FVirtualTextureLocalTile& Tile, uint8 GroupMask, uint16 Count);
+	uint16 LockLoadRequest(const FVirtualTextureLocalTile& Tile, uint8 GroupMask);
 
-	void AddMappingRequest(uint16 LoadRequestIndex, uint8 LocalLayerIndex, uint8 SpaceID, uint8 LayerIndex, uint32 vAddress, uint8 vLevel, uint8 Local_vLevel);
+	void AddMappingRequest(uint16 LoadRequestIndex, uint8 ProducerPhysicalGroupIndex, uint8 SpaceID, uint8 PageTableLayerIndex, uint32 vAddress, uint8 vLevel, uint8 Local_vLevel);
 
-	void AddDirectMappingRequest(uint8 InSpaceID, uint16 InPhysicalSpaceID, uint8 InLayer, uint8 InLogSize, uint32 InAddress, uint8 InLevel, uint16 InPhysicalAddress);
+	void AddDirectMappingRequest(uint8 InSpaceID, uint16 InPhysicalSpaceID, uint8 InPageTableLayerIndex, uint8 InLogSize, uint32 InAddress, uint8 InLevel, uint16 InPhysicalAddress);
 	void AddDirectMappingRequest(const FDirectMappingRequest& Request);
 
 	void MergeRequests(const FUniqueRequestList* RESTRICT Other, FMemStack& MemStack);
@@ -119,7 +119,7 @@ private:
 	FMappingRequest* MappingRequests;
 	FDirectMappingRequest* DirectMappingRequests;
 	uint16* LoadRequestCount;
-	uint8* LoadRequestLayerMask;
+	uint8* LoadRequestGroupMask;
 
 	uint32 NumLoadRequests;
 	uint32 NumLockRequests;
@@ -128,10 +128,10 @@ private:
 };
 
 
-inline uint16 FUniqueRequestList::AddLoadRequest(const FVirtualTextureLocalTile& Tile, uint8 LayerMask, uint16 Count)
+inline uint16 FUniqueRequestList::AddLoadRequest(const FVirtualTextureLocalTile& Tile, uint8 GroupMask, uint16 Count)
 {
 	const uint16 Hash = MurmurFinalize64(Tile.PackedValue);
-	check(LayerMask != 0u);
+	check(GroupMask != 0u);
 
 	for (uint16 Index = LoadRequestHash.First(Hash); LoadRequestHash.IsValid(Index); Index = LoadRequestHash.Next(Index))
 	{
@@ -143,7 +143,7 @@ inline uint16 FUniqueRequestList::AddLoadRequest(const FVirtualTextureLocalTile&
 				// Don't adjust count if already locked, don't allow request to transition to lock
 				LoadRequestCount[Index] = FMath::Min<uint32>(PrevCount + Count, 0xfffe);
 			}
-			LoadRequestLayerMask[Index] |= LayerMask;
+			LoadRequestGroupMask[Index] |= GroupMask;
 			return Index;
 		}
 	}
@@ -154,16 +154,16 @@ inline uint16 FUniqueRequestList::AddLoadRequest(const FVirtualTextureLocalTile&
 		LoadRequestHash.Add(Hash, Index);
 		LoadRequests[Index] = Tile;
 		LoadRequestCount[Index] = FMath::Min<uint32>(Count, 0xfffe);
-		LoadRequestLayerMask[Index] = LayerMask;
+		LoadRequestGroupMask[Index] = GroupMask;
 		return Index;
 	}
 	return 0xffff;
 }
 
-inline uint16 FUniqueRequestList::LockLoadRequest(const FVirtualTextureLocalTile& Tile, uint8 LayerMask)
+inline uint16 FUniqueRequestList::LockLoadRequest(const FVirtualTextureLocalTile& Tile, uint8 GroupMask)
 {
 	const uint16 Hash = MurmurFinalize64(Tile.PackedValue);
-	check(LayerMask != 0u);
+	check(GroupMask != 0u);
 
 	for (uint16 Index = LoadRequestHash.First(Hash); LoadRequestHash.IsValid(Index); Index = LoadRequestHash.Next(Index))
 	{
@@ -174,7 +174,7 @@ inline uint16 FUniqueRequestList::LockLoadRequest(const FVirtualTextureLocalTile
 				LoadRequestCount[Index] = 0xffff;
 				++NumLockRequests;
 			}
-			LoadRequestLayerMask[Index] |= LayerMask;
+			LoadRequestGroupMask[Index] |= GroupMask;
 			return Index;
 		}
 	}
@@ -185,7 +185,7 @@ inline uint16 FUniqueRequestList::LockLoadRequest(const FVirtualTextureLocalTile
 		LoadRequestHash.Add(Hash, Index);
 		LoadRequests[Index] = Tile;
 		LoadRequestCount[Index] = 0xffff;
-		LoadRequestLayerMask[Index] = LayerMask;
+		LoadRequestGroupMask[Index] = GroupMask;
 		++NumLockRequests;
 		return Index;
 	}
@@ -193,10 +193,10 @@ inline uint16 FUniqueRequestList::LockLoadRequest(const FVirtualTextureLocalTile
 	return 0xffff;
 }
 
-inline void FUniqueRequestList::AddMappingRequest(uint16 LoadRequestIndex, uint8 LocalLayerIndex, uint8 SpaceID, uint8 LayerIndex, uint32 vAddress, uint8 vLevel, uint8 Local_vLevel)
+inline void FUniqueRequestList::AddMappingRequest(uint16 LoadRequestIndex, uint8 ProducerPhysicalGroupIndex, uint8 SpaceID, uint8 PageTableLayerIndex, uint32 vAddress, uint8 vLevel, uint8 Local_vLevel)
 {
 	check(LoadRequestIndex < NumLoadRequests);
-	const FMappingRequest Request(LoadRequestIndex, LocalLayerIndex, SpaceID, LayerIndex, vAddress, vLevel, Local_vLevel);
+	const FMappingRequest Request(LoadRequestIndex, ProducerPhysicalGroupIndex, SpaceID, PageTableLayerIndex, vAddress, vLevel, Local_vLevel);
 	const uint16 Hash = MurmurFinalize64(Request.PackedValue);
 
 	for (uint16 Index = MappingRequestHash.First(Hash); MappingRequestHash.IsValid(Index); Index = MappingRequestHash.Next(Index))
@@ -215,9 +215,9 @@ inline void FUniqueRequestList::AddMappingRequest(uint16 LoadRequestIndex, uint8
 	}
 }
 
-inline void FUniqueRequestList::AddDirectMappingRequest(uint8 InSpaceID, uint16 InPhysicalSpaceID, uint8 InLayer, uint8 InLogSize, uint32 InAddress, uint8 InLevel, uint16 InPhysicalAddress)
+inline void FUniqueRequestList::AddDirectMappingRequest(uint8 InSpaceID, uint16 InPhysicalSpaceID, uint8 InPageTableLayerIndex, uint8 InLogSize, uint32 InAddress, uint8 InLevel, uint16 InPhysicalAddress)
 {
-	const FDirectMappingRequest Request(InSpaceID, InPhysicalSpaceID, InLayer, InLogSize, InAddress, InLevel, InPhysicalAddress);
+	const FDirectMappingRequest Request(InSpaceID, InPhysicalSpaceID, InPageTableLayerIndex, InLogSize, InAddress, InLevel, InPhysicalAddress);
 	AddDirectMappingRequest(Request);
 }
 
@@ -249,11 +249,11 @@ inline void FUniqueRequestList::MergeRequests(const FUniqueRequestList* RESTRICT
 	{
 		if (Other->IsLocked(Index))
 		{
-			LoadRequestIndexRemap[Index] = LockLoadRequest(Other->GetLoadRequest(Index), Other->LoadRequestLayerMask[Index]);
+			LoadRequestIndexRemap[Index] = LockLoadRequest(Other->GetLoadRequest(Index), Other->LoadRequestGroupMask[Index]);
 		}
 		else
 		{
-			LoadRequestIndexRemap[Index] = AddLoadRequest(Other->GetLoadRequest(Index), Other->LoadRequestLayerMask[Index], Other->LoadRequestCount[Index]);
+			LoadRequestIndexRemap[Index] = AddLoadRequest(Other->GetLoadRequest(Index), Other->LoadRequestGroupMask[Index], Other->LoadRequestCount[Index]);
 		}
 	}
 
@@ -264,7 +264,7 @@ inline void FUniqueRequestList::MergeRequests(const FUniqueRequestList* RESTRICT
 		const uint16 LoadRequestIndex = LoadRequestIndexRemap[Request.LoadRequestIndex];
 		if (LoadRequestIndex != 0xffff)
 		{
-			AddMappingRequest(LoadRequestIndex, Request.LocalLayerIndex, Request.SpaceID, Request.LayerIndex, Request.vAddress, Request.vLevel, Request.Local_vLevel);
+			AddMappingRequest(LoadRequestIndex, Request.ProducerPhysicalGroupIndex, Request.SpaceID, Request.PageTableLayerIndex, Request.vAddress, Request.vLevel, Request.Local_vLevel);
 		}
 	}
 
@@ -318,19 +318,19 @@ inline void FUniqueRequestList::SortRequests(FVirtualTextureProducerCollection&
 
 	// Re-index load request list, using sorted indices
 	FVirtualTextureLocalTile* SortedLoadRequests = new(MemStack) FVirtualTextureLocalTile[NewNumLoadRequests];
-	uint8* SortedLayerMask = new(MemStack) uint8[NewNumLoadRequests];
+	uint8* SortedGroupMask = new(MemStack) uint8[NewNumLoadRequests];
 	uint16* LoadIndexToSortedLoadIndex = new(MemStack) uint16[NumLoadRequests];
 	FMemory::Memset(LoadIndexToSortedLoadIndex, 0xff, NumLoadRequests * sizeof(uint16));
 	for (uint32 i = 0u; i < NewNumLoadRequests; ++i)
 	{
 		const uint32 SortedIndex = SortedKeys[i].Index;
 		SortedLoadRequests[i] = LoadRequests[SortedIndex];
-		SortedLayerMask[i] = LoadRequestLayerMask[SortedIndex];
+		SortedGroupMask[i] = LoadRequestGroupMask[SortedIndex];
 		checkSlow(SortedIndex < NumLoadRequests);
 		LoadIndexToSortedLoadIndex[SortedIndex] = i;
 	}
 	FMemory::Memcpy(LoadRequests, SortedLoadRequests, sizeof(FVirtualTextureLocalTile) * NewNumLoadRequests);
-	FMemory::Memcpy(LoadRequestLayerMask, SortedLayerMask, sizeof(uint8) * NewNumLoadRequests);
+	FMemory::Memcpy(LoadRequestGroupMask, SortedGroupMask, sizeof(uint8) * NewNumLoadRequests);
 
 	// Remap LoadRequest indices for all the mapping requests
 	// Can discard any mapping request that refers to a LoadRequest that's no longer being performed this frame
diff --git a/Engine/Source/Runtime/Renderer/Private/VT/VirtualTexturePhysicalSpace.cpp b/Engine/Source/Runtime/Renderer/Private/VT/VirtualTexturePhysicalSpace.cpp
--- a/Engine/Source/Runtime/Renderer/Private/VT/VirtualTexturePhysicalSpace.cpp
+++ b/Engine/Source/Runtime/Renderer/Private/VT/VirtualTexturePhysicalSpace.cpp
@@ -14,15 +14,20 @@ FVirtualTexturePhysicalSpace::FVirtualTexturePhysicalSpace(const FVTPhysicalSpac
 	, bPageTableLimit(false)
 	, bGpuTextureLimit(false)
 {
+	// Find matching physical pool
 	UVirtualTexturePoolConfig* PoolConfig = GetMutableDefault<UVirtualTexturePoolConfig>();
-	const FVirtualTextureSpacePoolConfig* Config = PoolConfig->FindPoolConfig(InDesc.TileSize, InDesc.Format);
+	const FVirtualTextureSpacePoolConfig* Config = PoolConfig->FindPoolConfig(InDesc.Format, InDesc.NumLayers, InDesc.TileSize);
 	const uint32 PoolSizeInBytes = Config->SizeInMegabyte * 1024u * 1024u;
 	const bool bForce16BitPageTable = false;
 
-	const FPixelFormatInfo& FormatInfo = GPixelFormats[InDesc.Format];
+	const FPixelFormatInfo& FormatInfo = GPixelFormats[InDesc.Format[0]];
 	check(InDesc.TileSize % FormatInfo.BlockSizeX == 0);
 	check(InDesc.TileSize % FormatInfo.BlockSizeY == 0);
-	const SIZE_T TileSizeBytes = CalculateImageBytes(InDesc.TileSize, InDesc.TileSize, 0, InDesc.Format);
+	SIZE_T TileSizeBytes = 0;
+	for (int32 Layer = 0; Layer < InDesc.NumLayers; ++Layer)
+	{
+		TileSizeBytes += CalculateImageBytes(InDesc.TileSize, InDesc.TileSize, 0, InDesc.Format[Layer]);
+	}
 	const uint32 MaxTiles = (uint32)(PoolSizeInBytes / TileSizeBytes);
 	
 	TextureSizeInTiles = FMath::CeilToInt(FMath::Sqrt((float)MaxTiles));
@@ -56,36 +61,46 @@ void FVirtualTexturePhysicalSpace::InitRHI()
 {
 	FRHICommandListImmediate& RHICmdList = FRHICommandListExecutor::GetImmediateCommandList();
 
-	const uint32 TextureSize = GetTextureSize();
-	const FPooledRenderTargetDesc Desc = FPooledRenderTargetDesc::Create2DDesc(
-		FIntPoint(TextureSize, TextureSize),
-		Description.Format,
-		FClearValueBinding::None,
-		TexCreate_None,
-		TexCreate_ShaderResource | (Description.bCreateRenderTarget ? (TexCreate_RenderTargetable | TexCreate_UAV) : 0),
-		false);
+	for (int32 Layer = 0; Layer < Description.NumLayers; ++Layer)
+	{
+		const uint32 TextureSize = GetTextureSize();
+		const FPooledRenderTargetDesc Desc = FPooledRenderTargetDesc::Create2DDesc(
+			FIntPoint(TextureSize, TextureSize),
+			Description.Format[Layer],
+			FClearValueBinding::None,
+			TexCreate_None,
+			TexCreate_ShaderResource | (Description.bCreateRenderTarget ? (TexCreate_RenderTargetable | TexCreate_UAV) : 0),
+			false);
 
-	GRenderTargetPool.FindFreeElement(RHICmdList, Desc, PooledRenderTarget, TEXT("PhysicalTexture"));
-	FRHITexture* TextureRHI = PooledRenderTarget->GetRenderTargetItem().ShaderResourceTexture;
+		GRenderTargetPool.FindFreeElement(RHICmdList, Desc, PooledRenderTarget[Layer], TEXT("PhysicalTexture"));
+		FRHITexture* TextureRHI = PooledRenderTarget[Layer]->GetRenderTargetItem().ShaderResourceTexture;
 
-	// Create sRGB/non-sRGB views into the physical texture
-	FRHITextureSRVCreateInfo ViewInfo;
-	TextureView = RHICreateShaderResourceView(TextureRHI, ViewInfo);
+		// Create sRGB/non-sRGB views into the physical texture
+		FRHITextureSRVCreateInfo ViewInfo;
+		TextureView[Layer] = RHICreateShaderResourceView(TextureRHI, ViewInfo);
 
-	ViewInfo.SRGBOverride = SRGBO_ForceEnable;
-	TextureSRGBView = RHICreateShaderResourceView(TextureRHI, ViewInfo);
+		ViewInfo.SRGBOverride = SRGBO_ForceEnable;
+		TextureSRGBView[Layer] = RHICreateShaderResourceView(TextureRHI, ViewInfo);
+	}
 }
 
 void FVirtualTexturePhysicalSpace::ReleaseRHI()
 {
-	GRenderTargetPool.FreeUnusedResource(PooledRenderTarget);
-	TextureView.SafeRelease();
-	TextureSRGBView.SafeRelease();
+	for (int32 Layer = 0; Layer < Description.NumLayers; ++Layer)
+	{
+		GRenderTargetPool.FreeUnusedResource(PooledRenderTarget[Layer]);
+		TextureView[Layer].SafeRelease();
+		TextureSRGBView[Layer].SafeRelease();
+	}
 }
 
 uint32 FVirtualTexturePhysicalSpace::GetSizeInBytes() const
 {
-	const SIZE_T TileSizeBytes = CalculateImageBytes(Description.TileSize, Description.TileSize, 0, Description.Format);
+	SIZE_T TileSizeBytes = 0;
+	for (int32 Layer = 0; Layer < Description.NumLayers; ++Layer)
+	{
+		TileSizeBytes += CalculateImageBytes(Description.TileSize, Description.TileSize, 0, Description.Format[Layer]);
+	}
 	return GetNumTiles() * TileSizeBytes;
 }
 
diff --git a/Engine/Source/Runtime/Renderer/Private/VT/VirtualTexturePhysicalSpace.h b/Engine/Source/Runtime/Renderer/Private/VT/VirtualTexturePhysicalSpace.h
--- a/Engine/Source/Runtime/Renderer/Private/VT/VirtualTexturePhysicalSpace.h
+++ b/Engine/Source/Runtime/Renderer/Private/VT/VirtualTexturePhysicalSpace.h
@@ -11,15 +11,31 @@
 struct FVTPhysicalSpaceDescription
 {
 	uint32 TileSize;
-	TEnumAsByte<EPixelFormat> Format;
 	uint8 Dimensions;
+	uint8 NumLayers;
+	TEnumAsByte<EPixelFormat> Format[VIRTUALTEXTURE_SPACE_MAXLAYERS];
 	bool bContinuousUpdate;
 	bool bCreateRenderTarget;
 };
 
 inline bool operator==(const FVTPhysicalSpaceDescription& Lhs, const FVTPhysicalSpaceDescription& Rhs)
 {
-	return Lhs.TileSize == Rhs.TileSize && Lhs.Format == Rhs.Format && Lhs.Dimensions == Rhs.Dimensions && Lhs.bContinuousUpdate == Rhs.bContinuousUpdate && Lhs.bCreateRenderTarget == Rhs.bCreateRenderTarget;
+	if (Lhs.TileSize != Rhs.TileSize || 
+		Lhs.NumLayers != Rhs.NumLayers || 
+		Lhs.Dimensions != Rhs.Dimensions || 
+		Lhs.bContinuousUpdate != Rhs.bContinuousUpdate || 
+		Lhs.bCreateRenderTarget != Rhs.bCreateRenderTarget)
+	{
+		return false;
+	}
+	for (int32 Layer = 0; Layer < Lhs.NumLayers; ++Layer)
+	{
+		if (Lhs.Format[Layer] != Rhs.Format[Layer])
+		{
+			return false;
+		}
+	}
+	return true;
 }
 inline bool operator!=(const FVTPhysicalSpaceDescription& Lhs, const FVTPhysicalSpaceDescription& Rhs)
 {
@@ -33,7 +49,7 @@ public:
 	virtual ~FVirtualTexturePhysicalSpace();
 
 	inline const FVTPhysicalSpaceDescription& GetDescription() const { return Description; }
-	inline EPixelFormat GetFormat() const { return Description.Format; }
+	inline EPixelFormat GetFormat(int32 Layer) const { return Description.Format[Layer]; }
 	inline uint16 GetID() const { return ID; }
 	inline uint32 GetNumTiles() const { return TextureSizeInTiles * TextureSizeInTiles; }
 	inline uint32 GetSizeInTiles() const { return TextureSizeInTiles; }
@@ -56,22 +72,22 @@ public:
 	inline uint32 Release() { check(NumRefs > 0u); return --NumRefs; }
 	inline uint32 GetRefCount() const { return NumRefs; }
 
-	FRHITexture* GetPhysicalTexture() const
+	FRHITexture* GetPhysicalTexture(int32 Layer) const
 	{
-		check(PooledRenderTarget.IsValid());
-		return PooledRenderTarget->GetRenderTargetItem().ShaderResourceTexture;
+		check(PooledRenderTarget[Layer].IsValid());
+		return PooledRenderTarget[Layer]->GetRenderTargetItem().ShaderResourceTexture;
 	}
 
-	TRefCountPtr<IPooledRenderTarget> GetPhysicalTexturePooledRenderTarget() const
+	TRefCountPtr<IPooledRenderTarget> GetPhysicalTexturePooledRenderTarget(int32 Layer) const
 	{
-		check(PooledRenderTarget.IsValid());
+		check(PooledRenderTarget[Layer].IsValid());
 		check(Description.bCreateRenderTarget);
-		return PooledRenderTarget;
+		return PooledRenderTarget[Layer];
 	}
 
-	FRHIShaderResourceView* GetPhysicalTextureView(bool bSRGB) const
+	FRHIShaderResourceView* GetPhysicalTextureView(int32 Layer, bool bSRGB) const
 	{
-		return bSRGB ? TextureSRGBView : TextureView;
+		return bSRGB ? TextureSRGBView[Layer] : TextureView[Layer];
 	}
 
 #if STATS
@@ -87,9 +103,9 @@ public:
 private:
 	FVTPhysicalSpaceDescription Description;
 	FTexturePagePool Pool;
-	TRefCountPtr<IPooledRenderTarget> PooledRenderTarget;
-	FShaderResourceViewRHIRef TextureView;
-	FShaderResourceViewRHIRef TextureSRGBView;
+	TRefCountPtr<IPooledRenderTarget> PooledRenderTarget[VIRTUALTEXTURE_SPACE_MAXLAYERS];
+	FShaderResourceViewRHIRef TextureView[VIRTUALTEXTURE_SPACE_MAXLAYERS];
+	FShaderResourceViewRHIRef TextureSRGBView[VIRTUALTEXTURE_SPACE_MAXLAYERS];
 
 	uint32 TextureSizeInTiles;
 	uint32 NumRefs;
diff --git a/Engine/Source/Runtime/Renderer/Private/VT/VirtualTextureProducer.cpp b/Engine/Source/Runtime/Renderer/Private/VT/VirtualTextureProducer.cpp
--- a/Engine/Source/Runtime/Renderer/Private/VT/VirtualTextureProducer.cpp
+++ b/Engine/Source/Runtime/Renderer/Private/VT/VirtualTextureProducer.cpp
@@ -15,13 +15,14 @@ void FVirtualTextureProducer::Release(FVirtualTextureSystem* System, const FVirt
 		System->ForceUnlockAllTiles(HandleToSelf, this);
 	}
 
-	for (uint32 LayerIndex = 0u; LayerIndex < Description.NumLayers; ++LayerIndex)
+	for (uint32 LayerIndex = 0u; LayerIndex < Description.NumPhysicalGroups; ++LayerIndex)
 	{
-		FVirtualTexturePhysicalSpace* Space = PhysicalSpace[LayerIndex];
+		FVirtualTexturePhysicalSpace* Space = PhysicalGroups[LayerIndex].PhysicalSpace;
 		Space->GetPagePool().EvictPages(System, HandleToSelf);
-		PhysicalSpace[LayerIndex].SafeRelease();
+		PhysicalGroups[LayerIndex].PhysicalSpace.SafeRelease();
 	}
 
+	PhysicalGroups.Reset();
 	delete VirtualTexture;
 	VirtualTexture = nullptr;
 	Description = FVTProducerDescription();
@@ -48,6 +49,8 @@ FVirtualTextureProducerHandle FVirtualTextureProducerCollection::RegisterProduce
 	check(ProducerWidth > 0u);
 	check(ProducerHeight > 0u);
 	check(InDesc.MaxLevel <= FMath::CeilLogTwo(FMath::Max(ProducerWidth, ProducerHeight)));
+	check(InDesc.NumTextureLayers <= VIRTUALTEXTURE_SPACE_MAXLAYERS);
+	check(InDesc.NumPhysicalGroups <= InDesc.NumTextureLayers);
 	check(InProducer);
 
 	const uint32 Index = AcquireEntry();
@@ -56,16 +59,26 @@ FVirtualTextureProducerHandle FVirtualTextureProducerCollection::RegisterProduce
 	Entry.Producer.VirtualTexture = InProducer;
 	Entry.DestroyedCallbacksIndex = AcquireCallback();
 
-	check(InDesc.NumLayers <= VIRTUALTEXTURE_SPACE_MAXLAYERS);
-	for (uint32 LayerIndex = 0u; LayerIndex < InDesc.NumLayers; ++LayerIndex)
+	FVTPhysicalSpaceDescription PhysicalSpaceDesc;
+	PhysicalSpaceDesc.Dimensions = InDesc.Dimensions;
+	PhysicalSpaceDesc.TileSize = InDesc.TileSize + InDesc.TileBorderSize * 2u;
+	PhysicalSpaceDesc.bContinuousUpdate = InDesc.bContinuousUpdate;
+	PhysicalSpaceDesc.bCreateRenderTarget = InDesc.bCreateRenderTarget;
+
+	Entry.Producer.PhysicalGroups.SetNum(InDesc.NumPhysicalGroups);
+	for (uint32 PhysicalGroupIndex = 0u; PhysicalGroupIndex < InDesc.NumPhysicalGroups; ++PhysicalGroupIndex)
 	{
-		FVTPhysicalSpaceDescription PhysicalSpaceDesc;
-		PhysicalSpaceDesc.Dimensions = InDesc.Dimensions;
-		PhysicalSpaceDesc.TileSize = InDesc.TileSize + InDesc.TileBorderSize * 2u;
-		PhysicalSpaceDesc.Format = InDesc.LayerFormat[LayerIndex];
-		PhysicalSpaceDesc.bContinuousUpdate = InDesc.bContinuousUpdate;
-		PhysicalSpaceDesc.bCreateRenderTarget = InDesc.bCreateRenderTarget;
-		Entry.Producer.PhysicalSpace[LayerIndex] = System->AcquirePhysicalSpace(PhysicalSpaceDesc);
+		PhysicalSpaceDesc.NumLayers = 0;
+		for (uint32 LayerIndex = 0u; LayerIndex < InDesc.NumTextureLayers; ++LayerIndex)
+		{
+			if (InDesc.PhysicalGroupIndex[LayerIndex] == PhysicalGroupIndex)
+			{
+				PhysicalSpaceDesc.Format[PhysicalSpaceDesc.NumLayers] = InDesc.LayerFormat[LayerIndex];
+				PhysicalSpaceDesc.NumLayers++;
+			}
+		}
+		check(PhysicalSpaceDesc.NumLayers > 0);
+		Entry.Producer.PhysicalGroups[PhysicalGroupIndex].PhysicalSpace = System->AcquirePhysicalSpace(PhysicalSpaceDesc);
 	}
 
 	const FVirtualTextureProducerHandle Handle(Index, Entry.Magic);
diff --git a/Engine/Source/Runtime/Renderer/Private/VT/VirtualTextureProducer.h b/Engine/Source/Runtime/Renderer/Private/VT/VirtualTextureProducer.h
--- a/Engine/Source/Runtime/Renderer/Private/VT/VirtualTextureProducer.h
+++ b/Engine/Source/Runtime/Renderer/Private/VT/VirtualTextureProducer.h
@@ -14,28 +14,42 @@ class FVirtualTexturePhysicalSpace;
 class FVirtualTextureProducer
 {
 public:
-	FVirtualTextureProducer() : VirtualTexture(nullptr) {}
+	FVirtualTextureProducer() : VirtualTexture(nullptr) { }
 
 	void Release(FVirtualTextureSystem* System, const FVirtualTextureProducerHandle& HandleToSelf);
 
 	inline const FVTProducerDescription& GetDescription() const { return Description; }
 	inline IVirtualTexture* GetVirtualTexture() const { return VirtualTexture; }
 	inline const FName& GetName() const { return Description.Name; }
-	inline uint32 GetNumLayers() const { return Description.NumLayers; }
 	inline uint32 GetWidthInTiles() const { return Description.BlockWidthInTiles * Description.WidthInBlocks; }
 	inline uint32 GetHeightInTiles() const { return Description.BlockHeightInTiles * Description.HeightInBlocks; }
 	inline uint32 GetDepthInTiles() const { return Description.DepthInTiles; }
-	inline EPixelFormat GetLayerFormat(uint32 LayerIndex) const { check(LayerIndex < Description.NumLayers); return Description.LayerFormat[LayerIndex]; }
-	inline FVirtualTexturePhysicalSpace* GetPhysicalSpace(uint32 LayerIndex) const { check(LayerIndex < Description.NumLayers); return PhysicalSpace[LayerIndex]; }
 	inline uint32 GetMaxLevel() const { return Description.MaxLevel; }
 
+	inline uint32 GetNumTextureLayers() const { return Description.NumTextureLayers; }
+	inline EPixelFormat GetLayerFormat(uint32 LayerIndex) const { check(LayerIndex < Description.NumTextureLayers); return Description.LayerFormat[LayerIndex]; }
+	inline uint32 GetPhysicalGroupIndexForTextureLayer(uint32 LayerIndex) const { check(LayerIndex < Description.NumTextureLayers); return Description.PhysicalGroupIndex[LayerIndex]; }
+
+	inline uint32 GetNumPhysicalGroups() const { return Description.NumPhysicalGroups; }
+	inline FVirtualTexturePhysicalSpace* GetPhysicalSpaceForPhysicalGroup(uint32 GroupIndex) const { check(GroupIndex < Description.NumPhysicalGroups); return PhysicalGroups[GroupIndex].PhysicalSpace; }
+
 private:
 	friend class FVirtualTextureProducerCollection;
 	~FVirtualTextureProducer();
 
-	IVirtualTexture* VirtualTexture;
-	TRefCountPtr<FVirtualTexturePhysicalSpace> PhysicalSpace[VIRTUALTEXTURE_SPACE_MAXLAYERS];
 	FVTProducerDescription Description;
+	
+	// A physical group is a physical space which multiple texture layers map to
+	// Texture layers that share a physical group will always be sampled via the same page table channel (they will have the same physical UV coordinates)
+	// More than one physical group can share the same FVirtualTexturePhysicalSpace
+	struct FPhysicalGroupDesc
+	{
+		TRefCountPtr<FVirtualTexturePhysicalSpace> PhysicalSpace;
+	};
+	TArray<FPhysicalGroupDesc> PhysicalGroups;
+
+	// The virtual texture provider implementation
+	IVirtualTexture* VirtualTexture;
 };
 
 class FVirtualTextureProducerCollection
diff --git a/Engine/Source/Runtime/Renderer/Private/VT/VirtualTextureSpace.cpp b/Engine/Source/Runtime/Renderer/Private/VT/VirtualTextureSpace.cpp
--- a/Engine/Source/Runtime/Renderer/Private/VT/VirtualTextureSpace.cpp
+++ b/Engine/Source/Runtime/Renderer/Private/VT/VirtualTextureSpace.cpp
@@ -58,18 +58,18 @@ FVirtualTextureSpace::FVirtualTextureSpace(FVirtualTextureSystem* InSystem, uint
 	const uint32 PhysicalTileSize = InDesc.TileSize + InDesc.TileBorderSize * 2u;
 	const uint32 MaxSizeInTiles = GetMax2DTextureDimension() / PhysicalTileSize;
 	const uint32 MaxNumTiles = MaxSizeInTiles * MaxSizeInTiles;
-	for (uint32 LayerIndex = 0u; LayerIndex < InDesc.NumLayers; ++LayerIndex)
+	for (uint32 LayerIndex = 0u; LayerIndex < InDesc.NumPageTableLayers; ++LayerIndex)
 	{
 		PhysicalPageMap[LayerIndex].Initialize(MaxNumTiles, LayerIndex, InDesc.Dimensions);
 	}
 
-	uint32 NumLayersToAllocate = InDesc.NumLayers;
+	uint32 NumLayersToAllocate = InDesc.NumPageTableLayers;
 	uint32 PageTableIndex = 0u;
 	FMemory::Memzero(TexturePixelFormat);
 	while (NumLayersToAllocate > 0u)
 	{
 		const uint32 NumLayersForTexture = FMath::Min(NumLayersToAllocate, LayersPerPageTableTexture);
-		const EPixelFormat PixelFormat = GetFormatForNumLayers(NumLayersForTexture, InDesc.Format);
+		const EPixelFormat PixelFormat = GetFormatForNumLayers(NumLayersForTexture, InDesc.PageTableFormat);
 		TexturePixelFormat[PageTableIndex] = PixelFormat;
 		NumLayersToAllocate -= NumLayersForTexture;
 		++PageTableIndex;
@@ -317,7 +317,7 @@ void FVirtualTextureSpace::ApplyUpdates(FVirtualTextureSystem* System, FRHIComma
 {
 	static TArray<FPageTableUpdate> ExpandedUpdates[VIRTUALTEXTURE_SPACE_MAXLAYERS][16];
 
-	for (uint32 LayerIndex = 0u; LayerIndex < Description.NumLayers; ++LayerIndex)
+	for (uint32 LayerIndex = 0u; LayerIndex < Description.NumPageTableLayers; ++LayerIndex)
 	{
 		FTexturePageMap& PageMap = PhysicalPageMap[LayerIndex];
 		if (bForceEntireUpdate || CVarVTRefreshEntirePageTable.GetValueOnRenderThread())
@@ -345,7 +345,7 @@ void FVirtualTextureSpace::ApplyUpdates(FVirtualTextureSystem* System, FRHIComma
 	// TODO Expand 3D updates for slices of volume texture
 
 	uint32 TotalNumUpdates = 0;
-	for (uint32 LayerIndex = 0u; LayerIndex < Description.NumLayers; ++LayerIndex)
+	for (uint32 LayerIndex = 0u; LayerIndex < Description.NumPageTableLayers; ++LayerIndex)
 	{
 		for (uint32 Mip = 0; Mip < NumPageTableLevels; Mip++)
 		{
@@ -380,7 +380,7 @@ void FVirtualTextureSpace::ApplyUpdates(FVirtualTextureSystem* System, FRHIComma
 	// This flushes the RHI thread!
 	{
 		uint8* Buffer = (uint8*)RHILockVertexBuffer(UpdateBuffer, 0, TotalNumUpdates * sizeof(FPageTableUpdate), RLM_WriteOnly);
-		for (uint32 LayerIndex = 0u; LayerIndex < Description.NumLayers; ++LayerIndex)
+		for (uint32 LayerIndex = 0u; LayerIndex < Description.NumPageTableLayers; ++LayerIndex)
 		{
 			for (uint32 Mip = 0; Mip < NumPageTableLevels; Mip++)
 			{
@@ -401,7 +401,7 @@ void FVirtualTextureSpace::ApplyUpdates(FVirtualTextureSystem* System, FRHIComma
 
 	auto ShaderMap = GetGlobalShaderMap(GetFeatureLevel());
 	FPageTableUpdateVS* VertexShader = nullptr;
-	if (Description.Format == EVTPageTableFormat::UInt16)
+	if (Description.PageTableFormat == EVTPageTableFormat::UInt16)
 	{
 		VertexShader = ShaderMap->GetShader< TPageTableUpdateVS<true> >();
 	}
@@ -412,7 +412,7 @@ void FVirtualTextureSpace::ApplyUpdates(FVirtualTextureSystem* System, FRHIComma
 	check(VertexShader);
 
 	uint32 FirstUpdate = 0;
-	for (uint32 LayerIndex = 0u; LayerIndex < Description.NumLayers; ++LayerIndex)
+	for (uint32 LayerIndex = 0u; LayerIndex < Description.NumPageTableLayers; ++LayerIndex)
 	{
 		const uint32 TextureIndex = LayerIndex / LayersPerPageTableTexture;
 		const uint32 LayerInTexture = LayerIndex % LayersPerPageTableTexture;
diff --git a/Engine/Source/Runtime/Renderer/Private/VT/VirtualTextureSpace.h b/Engine/Source/Runtime/Renderer/Private/VT/VirtualTextureSpace.h
--- a/Engine/Source/Runtime/Renderer/Private/VT/VirtualTextureSpace.h
+++ b/Engine/Source/Runtime/Renderer/Private/VT/VirtualTextureSpace.h
@@ -17,9 +17,9 @@ struct FVTSpaceDescription
 {
 	uint32 TileSize = 0u;
 	uint32 TileBorderSize = 0u;
-	EVTPageTableFormat Format = EVTPageTableFormat::UInt16;
 	uint8 Dimensions = 0u;
-	uint8 NumLayers = 0u;
+	EVTPageTableFormat PageTableFormat = EVTPageTableFormat::UInt16;
+	uint8 NumPageTableLayers = 0u;
 	uint8 bPrivateSpace : 1;
 };
 
@@ -28,8 +28,8 @@ inline bool operator==(const FVTSpaceDescription& Lhs, const FVTSpaceDescription
 	return Lhs.Dimensions == Rhs.Dimensions &&
 		Lhs.TileSize == Rhs.TileSize &&
 		Lhs.TileBorderSize == Rhs.TileBorderSize &&
-		Lhs.NumLayers == Rhs.NumLayers &&
-		Lhs.Format == Rhs.Format &&
+		Lhs.NumPageTableLayers == Rhs.NumPageTableLayers &&
+		Lhs.PageTableFormat == Rhs.PageTableFormat &&
 		Lhs.bPrivateSpace == Rhs.bPrivateSpace;
 }
 inline bool operator!=(const FVTSpaceDescription& Lhs, const FVTSpaceDescription& Rhs)
@@ -48,16 +48,17 @@ public:
 
 	inline const FVTSpaceDescription& GetDescription() const { return Description; }
 	inline uint32 GetPageTableSize() const { return PageTableSize; }
-	inline uint32 GetNumLayers() const { return Description.NumLayers; }
 	inline uint8 GetDimensions() const { return Description.Dimensions; }
-	inline uint32 GetNumPageTableTextures() const { return (Description.NumLayers + LayersPerPageTableTexture - 1u) / LayersPerPageTableTexture; }
+	inline EVTPageTableFormat GetPageTableFormat() const { return Description.PageTableFormat; }
+	inline uint32 GetNumPageTableLayers() const { return Description.NumPageTableLayers; }
+	inline uint32 GetNumPageTableTextures() const { return (Description.NumPageTableLayers + LayersPerPageTableTexture - 1u) / LayersPerPageTableTexture; }
 	inline uint8 GetID() const { return ID; }
 
 	inline uint32 GetNumPageTableLevels() const { return NumPageTableLevels; }
 	inline FVirtualTextureAllocator& GetAllocator() { return Allocator; }
 	inline const FVirtualTextureAllocator& GetAllocator() const { return Allocator; }
-	inline FTexturePageMap& GetPageMap(uint32 LayerIndex) { check(LayerIndex < Description.NumLayers); return PhysicalPageMap[LayerIndex]; }
-	inline const FTexturePageMap& GetPageMap(uint32 LayerIndex) const { check(LayerIndex < Description.NumLayers); return PhysicalPageMap[LayerIndex]; }
+	inline FTexturePageMap& GetPageMapForPageTableLayer(uint32 PageTableLayerIndex) { check(PageTableLayerIndex < Description.NumPageTableLayers); return PhysicalPageMap[PageTableLayerIndex]; }
+	inline const FTexturePageMap& GetPageMapForPageTableLayer(uint32 PageTableLayerIndex) const { check(PageTableLayerIndex < Description.NumPageTableLayers); return PhysicalPageMap[PageTableLayerIndex]; }
 
 	// FRenderResource interface
 	virtual void		InitRHI() override;
diff --git a/Engine/Source/Runtime/Renderer/Private/VT/VirtualTextureSystem.cpp b/Engine/Source/Runtime/Renderer/Private/VT/VirtualTextureSystem.cpp
--- a/Engine/Source/Runtime/Renderer/Private/VT/VirtualTextureSystem.cpp
+++ b/Engine/Source/Runtime/Renderer/Private/VT/VirtualTextureSystem.cpp
@@ -258,9 +258,9 @@ void FVirtualTextureSystem::FlushCache(FVirtualTextureProducerHandle const& Prod
 		FVTProducerDescription const& ProducerDescription = Producer->GetDescription();
 
 		TArray<FVirtualTexturePhysicalSpace*> PhysicalSpacesForProducer;
-		for (uint32 i = 0; i < Producer->GetNumLayers(); ++i)
+		for (uint32 i = 0; i < Producer->GetNumPhysicalGroups(); ++i)
 		{
-			PhysicalSpacesForProducer.AddUnique(Producer->GetPhysicalSpace(i));
+			PhysicalSpacesForProducer.AddUnique(Producer->GetPhysicalSpaceForPhysicalGroup(i));
 		}
 
 		// Don't resize to allow this container to grow as needed (avoid allocations when collecting)
@@ -303,7 +303,13 @@ void FVirtualTextureSystem::ListPhysicalPoolsFromConsole()
 			const FTexturePagePool& PagePool = PhysicalSpace.GetPagePool();
 			const uint32 TotalSizeInBytes = PhysicalSpace.GetSizeInBytes();
 
-			UE_LOG(LogConsoleResponse, Display, TEXT("PhysicaPool: [%i] PF_%s %ix%i:"), i, GPixelFormats[Desc.Format].Name, Desc.TileSize, Desc.TileSize);
+			UE_LOG(LogConsoleResponse, Display, TEXT("PhysicaPool: [%i] %ix%i:"), i, Desc.TileSize, Desc.TileSize);
+			
+			for (int Layer = 0; Layer < Desc.NumLayers; ++Layer)
+			{
+				UE_LOG(LogConsoleResponse, Display, TEXT("  Layer %i=%s"), Layer, GPixelFormats[Desc.Format[Layer]].Name);
+			}
+
 			UE_LOG(LogConsoleResponse, Display, TEXT("  SizeInMegabyte= %f"), (float)TotalSizeInBytes / 1024.0f / 1024.0f);
 			UE_LOG(LogConsoleResponse, Display, TEXT("  Dimensions= %ix%i"), PhysicalSpace.GetTextureSize(), PhysicalSpace.GetTextureSize());
 			UE_LOG(LogConsoleResponse, Display, TEXT("  Tiles= %i"), PhysicalSpace.GetNumTiles());
@@ -334,14 +340,14 @@ void FVirtualTextureSystem::ListPhysicalPoolsFromConsole()
 
 		const uint32 PhysicalTileSize = Desc.TileSize + Desc.TileBorderSize * 2u;
 		const TCHAR* FormatName = nullptr;
-		switch (Desc.Format)
+		switch (Desc.PageTableFormat)
 		{
 		case EVTPageTableFormat::UInt16: FormatName = TEXT("UInt16"); break;
 		case EVTPageTableFormat::UInt32: FormatName = TEXT("UInt32"); break;
 		default: checkNoEntry(); break;
 		}
 
-		UE_LOG(LogConsoleResponse, Display, TEXT("Pool: [%i] %s (%ix%i) x %i:"), ID, FormatName, PhysicalTileSize, PhysicalTileSize, Desc.NumLayers);
+		UE_LOG(LogConsoleResponse, Display, TEXT("Pool: [%i] %s (%ix%i) x %i:"), ID, FormatName, PhysicalTileSize, PhysicalTileSize, Desc.NumPageTableLayers);
 		UE_LOG(LogConsoleResponse, Display, TEXT("  PageTableSize= %ix%i"), PageTableSize, PageTableSize);
 		UE_LOG(LogConsoleResponse, Display, TEXT("  Allocations= %i, %i%% (%fMB)"),
 			Allocator.GetNumAllocations(),
@@ -357,7 +363,7 @@ uint32 GetTypeHash(const FAllocatedVTDescription& Description)
 
 IAllocatedVirtualTexture* FVirtualTextureSystem::AllocateVirtualTexture(const FAllocatedVTDescription& Desc)
 {
-	check(Desc.NumLayers <= VIRTUALTEXTURE_SPACE_MAXLAYERS);
+	check(Desc.NumTextureLayers <= VIRTUALTEXTURE_SPACE_MAXLAYERS);
 
 	// Make sure any pending VTs are destroyed before attempting to allocate a new one
 	// Otherwise, we might find/return an existing IAllocatedVirtualTexture* that's pending deletion
@@ -380,7 +386,7 @@ IAllocatedVirtualTexture* FVirtualTextureSystem::AllocateVirtualTexture(const FA
 	bool bSupport16BitPageTable = true;
 	FVirtualTextureProducer* ProducerForLayer[VIRTUALTEXTURE_SPACE_MAXLAYERS] = { nullptr };
 	bool bAnyLayerProducerWantsPersistentHighestMip = false;
-	for (uint32 LayerIndex = 0u; LayerIndex < Desc.NumLayers; ++LayerIndex)
+	for (uint32 LayerIndex = 0u; LayerIndex < Desc.NumTextureLayers; ++LayerIndex)
 	{
 		FVirtualTextureProducer* Producer = Producers.FindProducer(Desc.ProducerHandle[LayerIndex]);
 		ProducerForLayer[LayerIndex] = Producer;
@@ -393,7 +399,9 @@ IAllocatedVirtualTexture* FVirtualTextureSystem::AllocateVirtualTexture(const FA
 			MinHeightInBlocks = FMath::Min<uint32>(MinHeightInBlocks, ProducerDesc.HeightInBlocks);
 			DepthInTiles = FMath::Max(DepthInTiles, ProducerDesc.DepthInTiles);
 
-			FVirtualTexturePhysicalSpace* PhysicalSpace = Producer->GetPhysicalSpace(Desc.LocalLayerToProduce[LayerIndex]);
+			uint32 ProducerLayerIndex = Desc.ProducerLayerIndex[LayerIndex];
+			uint32 ProducerPhysicalGroup = Producer->GetPhysicalGroupIndexForTextureLayer(ProducerLayerIndex);
+			FVirtualTexturePhysicalSpace* PhysicalSpace = Producer->GetPhysicalSpaceForPhysicalGroup(ProducerPhysicalGroup);
 			if (!PhysicalSpace->DoesSupport16BitPageTable())
 			{
 				bSupport16BitPageTable = false;
@@ -414,7 +422,7 @@ IAllocatedVirtualTexture* FVirtualTextureSystem::AllocateVirtualTexture(const FA
 		while (!bFoundValidWidthInBlocks)
 		{
 			bFoundValidWidthInBlocks = true;
-			for (uint32 LayerIndex = 0u; LayerIndex < Desc.NumLayers; ++LayerIndex)
+			for (uint32 LayerIndex = 0u; LayerIndex < Desc.NumTextureLayers; ++LayerIndex)
 			{
 				const FVirtualTextureProducer* Producer = ProducerForLayer[LayerIndex];
 				if (Producer)
@@ -438,7 +446,7 @@ IAllocatedVirtualTexture* FVirtualTextureSystem::AllocateVirtualTexture(const FA
 		while (!bFoundValidHeightInBlocks)
 		{
 			bFoundValidHeightInBlocks = true;
-			for (uint32 LayerIndex = 0u; LayerIndex < Desc.NumLayers; ++LayerIndex)
+			for (uint32 LayerIndex = 0u; LayerIndex < Desc.NumTextureLayers; ++LayerIndex)
 			{
 				const FVirtualTextureProducer* Producer = ProducerForLayer[LayerIndex];
 				if (Producer)
@@ -455,13 +463,28 @@ IAllocatedVirtualTexture* FVirtualTextureSystem::AllocateVirtualTexture(const FA
 		}
 	}
 
+	// Sum the total number of physical groups from all producers
+	uint32 NumPhysicalGroups = 0;
+	TArray<FVirtualTextureProducer*> UniqueProducers;
+	for (uint32 LayerIndex = 0u; LayerIndex < Desc.NumTextureLayers; ++LayerIndex)
+	{
+		if (ProducerForLayer[LayerIndex] != nullptr)
+		{
+			UniqueProducers.AddUnique(ProducerForLayer[LayerIndex]);
+		}
+	}
+	for (int32 ProducerIndex = 0u; ProducerIndex < UniqueProducers.Num(); ++ProducerIndex)
+	{
+		NumPhysicalGroups += UniqueProducers[ProducerIndex]->GetNumPhysicalGroups();
+	}
+
 	FVTSpaceDescription SpaceDesc;
 	SpaceDesc.Dimensions = Desc.Dimensions;
-	SpaceDesc.NumLayers = Desc.NumLayers;
+	SpaceDesc.NumPageTableLayers = (uint8)NumPhysicalGroups;
 	SpaceDesc.TileSize = Desc.TileSize;
 	SpaceDesc.TileBorderSize = Desc.TileBorderSize;
 	SpaceDesc.bPrivateSpace = Desc.bPrivateSpace;
-	SpaceDesc.Format = bSupport16BitPageTable ? EVTPageTableFormat::UInt16 : EVTPageTableFormat::UInt32;
+	SpaceDesc.PageTableFormat = bSupport16BitPageTable ? EVTPageTableFormat::UInt16 : EVTPageTableFormat::UInt32;
 	FVirtualTextureSpace* Space = AcquireSpace(SpaceDesc, FMath::Max(BlockWidthInTiles * WidthInBlocks, BlockHeightInTiles * HeightInBlocks));
 
 	AllocatedVT = new FAllocatedVirtualTexture(this, Frame, Desc, Space, ProducerForLayer, BlockWidthInTiles, BlockHeightInTiles, WidthInBlocks, HeightInBlocks, DepthInTiles);
@@ -657,11 +680,11 @@ void FVirtualTextureSystem::LockTile(const FVirtualTextureLocalTile& Tile)
 
 static void UnlockTileInternal(const FVirtualTextureProducerHandle& ProducerHandle, const FVirtualTextureProducer* Producer, const FVirtualTextureLocalTile& Tile, uint32 Frame)
 {
-	for (uint32 LocalLayerIndex = 0u; LocalLayerIndex < Producer->GetNumLayers(); ++LocalLayerIndex)
+	for (uint32 ProducerPhysicalGroupIndex = 0u; ProducerPhysicalGroupIndex < Producer->GetNumPhysicalGroups(); ++ProducerPhysicalGroupIndex)
 	{
-		FVirtualTexturePhysicalSpace* PhysicalSpace = Producer->GetPhysicalSpace(LocalLayerIndex);
+		FVirtualTexturePhysicalSpace* PhysicalSpace = Producer->GetPhysicalSpaceForPhysicalGroup(ProducerPhysicalGroupIndex);
 		FTexturePagePool& PagePool = PhysicalSpace->GetPagePool();
-		const uint32 pAddress = PagePool.FindPageAddress(ProducerHandle, LocalLayerIndex, Tile.Local_vAddress, Tile.Local_vLevel);
+		const uint32 pAddress = PagePool.FindPageAddress(ProducerHandle, ProducerPhysicalGroupIndex, Tile.Local_vAddress, Tile.Local_vLevel);
 		if (pAddress != ~0u)
 		{
 			PagePool.Unlock(Frame, pAddress);
@@ -1012,24 +1035,25 @@ void FVirtualTextureSystem::Update(FRHICommandListImmediate& RHICmdList, ERHIFea
 			checkSlow(TileLocks.IsLocked(Tile));
 			if (Producer)
 			{
-				uint8 LocalLayerMaskToLoad = 0u;
-				for (uint32 LocalLayerIndex = 0u; LocalLayerIndex < Producer->GetNumLayers(); ++LocalLayerIndex)
+				uint8 ProducerLayerMaskToLoad = 0u;
+				for (uint32 ProducerLayerIndex = 0u; ProducerLayerIndex < Producer->GetNumTextureLayers(); ++ProducerLayerIndex)
 				{
-					FVirtualTexturePhysicalSpace* PhysicalSpace = Producer->GetPhysicalSpace(LocalLayerIndex);
+					uint32 GroupIndex = Producer->GetPhysicalGroupIndexForTextureLayer(ProducerLayerIndex);
+					FVirtualTexturePhysicalSpace* PhysicalSpace = Producer->GetPhysicalSpaceForPhysicalGroup(GroupIndex);
 					FTexturePagePool& PagePool = PhysicalSpace->GetPagePool();
-					const uint32 pAddress = PagePool.FindPageAddress(ProducerHandle, LocalLayerIndex, Tile.Local_vAddress, Tile.Local_vLevel);
+					const uint32 pAddress = PagePool.FindPageAddress(ProducerHandle, GroupIndex, Tile.Local_vAddress, Tile.Local_vLevel);
 					if (pAddress == ~0u)
 					{
-						LocalLayerMaskToLoad |= (1u << LocalLayerIndex);
+						ProducerLayerMaskToLoad |= (1u << ProducerLayerIndex);
 					}
 					else
 					{
 						PagePool.Lock(pAddress);
 					}
 				}
-				if (LocalLayerMaskToLoad != 0u)
+				if (ProducerLayerMaskToLoad != 0u)
 				{
-					MergedRequestList->LockLoadRequest(FVirtualTextureLocalTile(Tile.GetProducerHandle(), Tile.Local_vAddress, Tile.Local_vLevel), LocalLayerMaskToLoad);
+					MergedRequestList->LockLoadRequest(FVirtualTextureLocalTile(Tile.GetProducerHandle(), Tile.Local_vAddress, Tile.Local_vLevel), ProducerLayerMaskToLoad);
 				}
 			}
 		}
@@ -1254,14 +1278,14 @@ void FVirtualTextureSystem::GatherRequestsTask(const FGatherRequestsParameters&
 		const uint32 vDimensions = Space->GetDimensions();
 		const uint32 vAddress = vPosition << (vLevel * vDimensions);
 
-		uint32 LayersToLoad[VIRTUALTEXTURE_SPACE_MAXLAYERS] = { 0 };
-		uint32 NumLayersToLoad = 0u;
+		uint32 PageTableLayersToLoad[VIRTUALTEXTURE_SPACE_MAXLAYERS] = { 0 };
+		uint32 NumPageTableLayersToLoad = 0u;
 		{
 			const FTexturePage VirtualPage(vLevel, vAddress);
 			const uint16 VirtualPageHash = MurmurFinalize32(VirtualPage.Packed);
-			for (uint32 LayerIndex = 0u; LayerIndex < Space->GetNumLayers(); ++LayerIndex)
+			for (uint32 PageTableLayerIndex = 0u; PageTableLayerIndex < Space->GetNumPageTableLayers(); ++PageTableLayerIndex)
 			{
-				const FTexturePageMap& RESTRICT PageMap = Space->GetPageMap(LayerIndex);
+				const FTexturePageMap& RESTRICT PageMap = Space->GetPageMapForPageTableLayer(PageTableLayerIndex);
 
 				++NumRequestsPages;
 				const FPhysicalSpaceIDAndAddress PhysicalSpaceIDAndAddress = PageMap.FindPagePhysicalSpaceIDAndAddress(VirtualPage, VirtualPageHash);
@@ -1292,12 +1316,12 @@ void FVirtualTextureSystem::GatherRequestsTask(const FGatherRequestsParameters&
 				else
 				{
 					// Page not resident, store for later processing
-					LayersToLoad[NumLayersToLoad++] = LayerIndex;
+					PageTableLayersToLoad[NumPageTableLayersToLoad++] = PageTableLayerIndex;
 				}
 			}
 		}
 
-		if (NumLayersToLoad == 0u)
+		if (NumPageTableLayersToLoad == 0u)
 		{
 			// All pages are resident and properly mapped, we're done
 			// This is the fast path, as most frames should generally have the majority of tiles already mapped
@@ -1322,36 +1346,43 @@ void FVirtualTextureSystem::GatherRequestsTask(const FGatherRequestsParameters&
 			continue;
 		}
 
-		checkSlow(AllocatedVT->GetNumLayers() == Space->GetNumLayers());
+		check(AllocatedVT->GetNumPageTableLayers() == Space->GetNumPageTableLayers());
 		if (vLevel > AllocatedVT->GetMaxLevel())
 		{
 			// Requested level is outside the given allocated VT
 			// This can happen for requests made by expanding mips, since we don't know the current allocated VT in that context
-			check(NumLayersToLoad == Space->GetNumLayers()); // no pages from this request should have been resident
-			check(NumRequestsPages >= Space->GetNumLayers()); // don't want to track these requests, since it turns out they're not valid
-			NumRequestsPages -= Space->GetNumLayers();
+			check(NumPageTableLayersToLoad == Space->GetNumPageTableLayers()); // no pages from this request should have been resident
+			check(NumRequestsPages >= Space->GetNumPageTableLayers()); // don't want to track these requests, since it turns out they're not valid
+			NumRequestsPages -= Space->GetNumPageTableLayers();
 			continue;
 		}
 
+		// Build producer local layer masks from physical layers that we need to load
+		uint8 ProducerGroupMaskToLoad[VIRTUALTEXTURE_SPACE_MAXLAYERS] = { 0u };
+		uint8 ProducerTextureLayerMaskToLoad[VIRTUALTEXTURE_SPACE_MAXLAYERS] = { 0u };
+
 		const uint32 NumUniqueProducers = AllocatedVT->GetNumUniqueProducers();
-		uint8 LocalLayerMaskToLoadForProducer[VIRTUALTEXTURE_SPACE_MAXLAYERS] = { 0u };
-		for (uint32 LoadLayerIndex = 0u; LoadLayerIndex < NumLayersToLoad; ++LoadLayerIndex)
+
+		for (uint32 LoadPageTableLayerIndex = 0u; LoadPageTableLayerIndex < NumPageTableLayersToLoad; ++LoadPageTableLayerIndex)
 		{
-			const uint32 LayerIndex = LayersToLoad[LoadLayerIndex];
-			const FVirtualTexturePhysicalSpace* RESTRICT PhysicalSpace = AllocatedVT->GetPhysicalSpace(LayerIndex);
-			const uint32 ProducerIndexForLayer = AllocatedVT->GetUniqueProducerIndexForLayer(LayerIndex);
-			if (ProducerIndexForLayer < NumUniqueProducers)
-			{
-				const uint32 LocalLayerToProduce = AllocatedVT->GetLocalLayerToProduce(LayerIndex);
-				LocalLayerMaskToLoadForProducer[ProducerIndexForLayer] |= (1u << LocalLayerToProduce);
-				++PageUpdateBuffers[PhysicalSpace->GetID()].WorkingSetSize;
-			}
+			const uint32 PageTableLayerIndex = PageTableLayersToLoad[LoadPageTableLayerIndex];
+			const FVirtualTexturePhysicalSpace* RESTRICT PhysicalSpace = AllocatedVT->GetPhysicalSpaceForPageTableLayer(PageTableLayerIndex);
+			const uint32 ProducerIndex = AllocatedVT->GetProducerIndexForPageTableLayer(PageTableLayerIndex);
+			check(ProducerIndex < NumUniqueProducers);
+			
+			const uint32 ProducerTextureLayerMask = AllocatedVT->GetProducerTextureLayerMaskForPageTableLayer(PageTableLayerIndex);
+			ProducerTextureLayerMaskToLoad[ProducerIndex] |= ProducerTextureLayerMask;
+			
+			const uint32 ProducerPhysicalGroupIndex = AllocatedVT->GetProducerPhysicalGroupIndexForPageTableLayer(PageTableLayerIndex);
+			ProducerGroupMaskToLoad[ProducerIndex] |= 1 << ProducerPhysicalGroupIndex;
+
+			++PageUpdateBuffers[PhysicalSpace->GetID()].WorkingSetSize;
 		}
 
 		for (uint32 ProducerIndex = 0u; ProducerIndex < NumUniqueProducers; ++ProducerIndex)
 		{
-			uint8 LocalLayerMaskToLoad = LocalLayerMaskToLoadForProducer[ProducerIndex];
-			if (LocalLayerMaskToLoad == 0u)
+			uint8 GroupMaskToLoad = ProducerGroupMaskToLoad[ProducerIndex];
+			if (GroupMaskToLoad == 0u)
 			{
 				continue;
 			}
@@ -1397,21 +1428,22 @@ void FVirtualTextureSystem::GatherRequestsTask(const FGatherRequestsParameters&
 				Mapping_vLevel = FMath::Max(vLevel, LocalMipBias + ProducerMipBias);
 			}
 
-			uint8 LocalLayerMaskToPrefetchForLevel[16] = { 0u };
+			uint8 ProducerPhysicalGroupMaskToPrefetchForLevel[16] = { 0u };
 			uint32 MaxPrefetchLocal_vLevel = Local_vLevel;
 
-			for (uint32 LocalLayerIndex = 0u; LocalLayerIndex < Producer->GetNumLayers(); ++LocalLayerIndex)
+			// Iterate local layers that we found unmapped
+			for (uint32 ProducerGroupIndex = 0u; ProducerGroupIndex < Producer->GetNumPhysicalGroups(); ++ProducerGroupIndex)
 			{
-				if ((LocalLayerMaskToLoad & (1u << LocalLayerIndex)) == 0u)
+				if ((GroupMaskToLoad & (1u << ProducerGroupIndex)) == 0u)
 				{
 					continue;
 				}
 
-				const FVirtualTexturePhysicalSpace* RESTRICT PhysicalSpace = Producer->GetPhysicalSpace(LocalLayerIndex);
+				const FVirtualTexturePhysicalSpace* RESTRICT PhysicalSpace = Producer->GetPhysicalSpaceForPhysicalGroup(ProducerGroupIndex);
 				const FTexturePagePool& RESTRICT PagePool = PhysicalSpace->GetPagePool();
 
 				// Find the highest resolution tile that's currently loaded
-				const uint32 pAddress = PagePool.FindNearestPageAddress(ProducerHandle, LocalLayerIndex, Local_vAddress, Local_vLevel, MaxLevel);
+				const uint32 pAddress = PagePool.FindNearestPageAddress(ProducerHandle, ProducerGroupIndex, Local_vAddress, Local_vLevel, MaxLevel);
 				uint32 AllocatedLocal_vLevel = MaxLevel + 1u;
 				if (pAddress != ~0u)
 				{
@@ -1427,18 +1459,18 @@ void FVirtualTextureSystem::GatherRequestsTask(const FGatherRequestsParameters&
 					AddPageUpdate(PageUpdateBuffers, PageUpdateFlushCount, PhysicalSpace->GetID(), pAddress);
 
 					uint32 NumMappedPages = 0u;
-					for (uint32 LoadLayerIndex = 0u; LoadLayerIndex < NumLayersToLoad; ++LoadLayerIndex)
+					for (uint32 LoadLayerIndex = 0u; LoadLayerIndex < NumPageTableLayersToLoad; ++LoadLayerIndex)
 					{
-						const uint32 LayerIndex = LayersToLoad[LoadLayerIndex];
-						if (AllocatedVT->GetLocalLayerToProduce(LayerIndex) == LocalLayerIndex &&
-							AllocatedVT->GetUniqueProducerIndexForLayer(LayerIndex) == ProducerIndex)
+						const uint32 PageTableLayerIndex = PageTableLayersToLoad[LoadLayerIndex];
+						if (AllocatedVT->GetProducerPhysicalGroupIndexForPageTableLayer(PageTableLayerIndex) == ProducerGroupIndex &&
+							AllocatedVT->GetProducerIndexForPageTableLayer(PageTableLayerIndex) == ProducerIndex)
 						{
 							bool bPageWasMapped = false;
 							if (Allocated_vLevel != vLevel)
 							{
 								// if we found a lower resolution tile than was requested, it may have already been mapped, check for that first
 								// don't need to check this if the allocated page is at the level that was requested...if that was already mapped we wouldn't have gotten this far
-								const FTexturePageMap& PageMap = Space->GetPageMap(LayerIndex);
+								const FTexturePageMap& PageMap = Space->GetPageMapForPageTableLayer(PageTableLayerIndex);
 								const FPhysicalSpaceIDAndAddress PrevPhysicalSpaceIDAndAddress = PageMap.FindPagePhysicalSpaceIDAndAddress(Allocated_vLevel, Allocated_vAddress);
 								if (PrevPhysicalSpaceIDAndAddress.Packed != ~0u)
 								{
@@ -1454,7 +1486,7 @@ void FVirtualTextureSystem::GatherRequestsTask(const FGatherRequestsParameters&
 							else
 							{
 								// verify our assumption that the page shouldn't be mapped yet
-								const FTexturePageMap& PageMap = Space->GetPageMap(LayerIndex);
+								const FTexturePageMap& PageMap = Space->GetPageMapForPageTableLayer(PageTableLayerIndex);
 								const FPhysicalSpaceIDAndAddress PrevPhysicalSpaceIDAndAddress = PageMap.FindPagePhysicalSpaceIDAndAddress(Allocated_vLevel, Allocated_vAddress);
 								checkSlow(PrevPhysicalSpaceIDAndAddress.Packed == ~0u);
 							}
@@ -1463,7 +1495,7 @@ void FVirtualTextureSystem::GatherRequestsTask(const FGatherRequestsParameters&
 							if (!bPageWasMapped)
 							{
 								// map the page now if it wasn't already mapped
-								RequestList->AddDirectMappingRequest(Space->GetID(), PhysicalSpace->GetID(), LayerIndex, Allocated_vLevel, Allocated_vAddress, AllocatedMapping_vLevel, pAddress);
+								RequestList->AddDirectMappingRequest(Space->GetID(), PhysicalSpace->GetID(), PageTableLayerIndex, Allocated_vLevel, Allocated_vAddress, AllocatedMapping_vLevel, pAddress);
 							}
 							++NumMappedPages;
 						}
@@ -1474,7 +1506,7 @@ void FVirtualTextureSystem::GatherRequestsTask(const FGatherRequestsParameters&
 				if (AllocatedLocal_vLevel == Local_vLevel)
 				{
 					// page at the requested level was already resident, no longer need to load
-					LocalLayerMaskToLoad &= ~(1u << LocalLayerIndex);
+					GroupMaskToLoad &= ~(1u << ProducerGroupIndex);
 					++NumResidentPages;
 				}
 				else
@@ -1485,7 +1517,7 @@ void FVirtualTextureSystem::GatherRequestsTask(const FGatherRequestsParameters&
 					const uint32 PrefetchLocal_vLevel = AllocatedLocal_vLevel - FMath::Min(2u, AllocatedLocal_vLevel);
 					if (PrefetchLocal_vLevel > Local_vLevel)
 					{
-						LocalLayerMaskToPrefetchForLevel[PrefetchLocal_vLevel] |= (1u << LocalLayerIndex);
+						ProducerPhysicalGroupMaskToPrefetchForLevel[PrefetchLocal_vLevel] |= (1u << ProducerGroupIndex);
 						MaxPrefetchLocal_vLevel = FMath::Max(MaxPrefetchLocal_vLevel, PrefetchLocal_vLevel);
 						++NumPrefetchPages;
 					}
@@ -1496,23 +1528,23 @@ void FVirtualTextureSystem::GatherRequestsTask(const FGatherRequestsParameters&
 			// Check to see if we have any levels to prefetch
 			for (uint32 PrefetchLocal_vLevel = Local_vLevel + 1u; PrefetchLocal_vLevel <= MaxPrefetchLocal_vLevel; ++PrefetchLocal_vLevel)
 			{
-				uint32 LocalLayerMaskToPrefetch = LocalLayerMaskToPrefetchForLevel[PrefetchLocal_vLevel];
-				if (LocalLayerMaskToPrefetch != 0u)
+				uint32 ProducerPhysicalGroupMaskToPrefetch = ProducerPhysicalGroupMaskToPrefetchForLevel[PrefetchLocal_vLevel];
+				if (ProducerPhysicalGroupMaskToPrefetch != 0u)
 				{
 					const uint32 PrefetchLocal_vAddress = Local_vAddress >> ((PrefetchLocal_vLevel - Local_vLevel) * vDimensions);
 
 					// If we want to prefetch any layers for a given level, need to ensure that we request all the layers that aren't currently loaded
 					// This is required since the VT producer interface needs to be able to write data for all layers if desired, so we need to make sure that all layers are allocated
-					for (uint32 LocalLayerIndex = 0u; LocalLayerIndex < Producer->GetNumLayers(); ++LocalLayerIndex)
+					for (uint32 ProducerPhysicalGroupIndex = 0u; ProducerPhysicalGroupIndex < Producer->GetNumPhysicalGroups(); ++ProducerPhysicalGroupIndex)
 					{
-						if ((LocalLayerMaskToPrefetch & (1u << LocalLayerIndex)) == 0u)
+						if ((ProducerPhysicalGroupMaskToPrefetch & (1u << ProducerPhysicalGroupIndex)) == 0u)
 						{
-							const FVirtualTexturePhysicalSpace* RESTRICT PhysicalSpace = Producer->GetPhysicalSpace(LocalLayerIndex);
+							const FVirtualTexturePhysicalSpace* RESTRICT PhysicalSpace = Producer->GetPhysicalSpaceForPhysicalGroup(ProducerPhysicalGroupIndex);
 							const FTexturePagePool& RESTRICT PagePool = PhysicalSpace->GetPagePool();
-							const uint32 pAddress = PagePool.FindPageAddress(ProducerHandle, LocalLayerIndex, PrefetchLocal_vAddress, PrefetchLocal_vLevel);
+							const uint32 pAddress = PagePool.FindPageAddress(ProducerHandle, ProducerPhysicalGroupIndex, PrefetchLocal_vAddress, PrefetchLocal_vLevel);
 							if (pAddress == ~0u)
 							{
-								LocalLayerMaskToPrefetch |= (1u << LocalLayerIndex);
+								ProducerPhysicalGroupMaskToPrefetch |= (1u << ProducerPhysicalGroupIndex);
 								++NumPrefetchPages;
 							}
 							else
@@ -1523,22 +1555,22 @@ void FVirtualTextureSystem::GatherRequestsTask(const FGatherRequestsParameters&
 						}
 					}
 
-					const uint16 LoadRequestIndex = RequestList->AddLoadRequest(FVirtualTextureLocalTile(ProducerHandle, PrefetchLocal_vAddress, PrefetchLocal_vLevel), LocalLayerMaskToPrefetch, PageCount);
+					const uint16 LoadRequestIndex = RequestList->AddLoadRequest(FVirtualTextureLocalTile(ProducerHandle, PrefetchLocal_vAddress, PrefetchLocal_vLevel), ProducerPhysicalGroupMaskToPrefetch, PageCount);
 					if (LoadRequestIndex != 0xffff)
 					{
 						const uint32 Prefetch_vLevel = PrefetchLocal_vLevel + ProducerMipBias;
 						ensure(Prefetch_vLevel <= AllocatedVT->GetMaxLevel());
 						const uint32 PrefetchMapping_vLevel = FMath::Max(Prefetch_vLevel, ProducerMipBias);
 						const uint32 Prefetch_vAddress = vAddress & (0xffffffff << (Prefetch_vLevel * vDimensions));
-						for (uint32 LoadLayerIndex = 0u; LoadLayerIndex < NumLayersToLoad; ++LoadLayerIndex)
+						for (uint32 LoadLayerIndex = 0u; LoadLayerIndex < NumPageTableLayersToLoad; ++LoadLayerIndex)
 						{
-							const uint32 LayerIndex = LayersToLoad[LoadLayerIndex];
-							if (AllocatedVT->GetUniqueProducerIndexForLayer(LayerIndex) == ProducerIndex)
+							const uint32 LayerIndex = PageTableLayersToLoad[LoadLayerIndex];
+							if (AllocatedVT->GetProducerIndexForPageTableLayer(LayerIndex) == ProducerIndex)
 							{
-								const uint32 LocalLayerToProduce = AllocatedVT->GetLocalLayerToProduce(LayerIndex);
-								if (LocalLayerMaskToPrefetch & (1u << LocalLayerToProduce))
+								const uint32 ProducerPhysicalGroupIndex = AllocatedVT->GetProducerPhysicalGroupIndexForPageTableLayer(LayerIndex);
+								if (ProducerPhysicalGroupMaskToPrefetch & (1u << ProducerPhysicalGroupIndex))
 								{
-									RequestList->AddMappingRequest(LoadRequestIndex, LocalLayerToProduce, ID, LayerIndex, Prefetch_vAddress, Prefetch_vLevel, PrefetchMapping_vLevel);
+									RequestList->AddMappingRequest(LoadRequestIndex, ProducerPhysicalGroupIndex, ID, LayerIndex, Prefetch_vAddress, Prefetch_vLevel, PrefetchMapping_vLevel);
 								}
 							}
 						}
@@ -1546,20 +1578,20 @@ void FVirtualTextureSystem::GatherRequestsTask(const FGatherRequestsParameters&
 				}
 			}
 
-			if (LocalLayerMaskToLoad != 0u)
+			if (GroupMaskToLoad != 0u)
 			{
-				const uint16 LoadRequestIndex = RequestList->AddLoadRequest(FVirtualTextureLocalTile(ProducerHandle, Local_vAddress, Local_vLevel), LocalLayerMaskToLoad, PageCount);
+				const uint16 LoadRequestIndex = RequestList->AddLoadRequest(FVirtualTextureLocalTile(ProducerHandle, Local_vAddress, Local_vLevel), GroupMaskToLoad, PageCount);
 				if (LoadRequestIndex != 0xffff)
 				{
-					for (uint32 LoadLayerIndex = 0u; LoadLayerIndex < NumLayersToLoad; ++LoadLayerIndex)
+					for (uint32 LoadLayerIndex = 0u; LoadLayerIndex < NumPageTableLayersToLoad; ++LoadLayerIndex)
 					{
-						const uint32 LayerIndex = LayersToLoad[LoadLayerIndex];
-						if (AllocatedVT->GetUniqueProducerIndexForLayer(LayerIndex) == ProducerIndex)
+						const uint32 LayerIndex = PageTableLayersToLoad[LoadLayerIndex];
+						if (AllocatedVT->GetProducerIndexForPageTableLayer(LayerIndex) == ProducerIndex)
 						{
-							const uint32 LocalLayerToProduce = AllocatedVT->GetLocalLayerToProduce(LayerIndex);
-							if (LocalLayerMaskToLoad & (1u << LocalLayerToProduce))
+							const uint32 ProducerPhysicalGroupIndex = AllocatedVT->GetProducerPhysicalGroupIndexForPageTableLayer(LayerIndex);
+							if (GroupMaskToLoad & (1u << ProducerPhysicalGroupIndex))
 							{
-								RequestList->AddMappingRequest(LoadRequestIndex, LocalLayerToProduce, ID, LayerIndex, vAddress, vLevel, Mapping_vLevel);
+								RequestList->AddMappingRequest(LoadRequestIndex, ProducerPhysicalGroupIndex, ID, LayerIndex, vAddress, vLevel, Mapping_vLevel);
 							}
 						}
 					}
@@ -1620,22 +1652,32 @@ void FVirtualTextureSystem::SubmitRequestsFromLocalTileList(const TSet<FVirtualT
 		const FVirtualTextureProducer& Producer = Producers.GetProducer(ProducerHandle);
 
 		// Fill targets for each layer
+		// Each producer can have multiple physical layers
+		// If the phys layer is mapped then we get the textures it owns and map them into the producer local slots and set the flags
 		uint32 LayerMask = 0;
 		FVTProduceTargetLayer ProduceTarget[VIRTUALTEXTURE_SPACE_MAXLAYERS];
-		for (uint32 LocalLayerIndex = 0u; LocalLayerIndex < Producer.GetNumLayers(); ++LocalLayerIndex)
+		for (uint32 ProducerPhysicalGroupIndex = 0u; ProducerPhysicalGroupIndex < Producer.GetNumPhysicalGroups(); ++ProducerPhysicalGroupIndex)
 		{
-			FVirtualTexturePhysicalSpace* RESTRICT PhysicalSpace = Producer.GetPhysicalSpace(LocalLayerIndex);
+			FVirtualTexturePhysicalSpace* RESTRICT PhysicalSpace = Producer.GetPhysicalSpaceForPhysicalGroup(ProducerPhysicalGroupIndex);
 			FTexturePagePool& RESTRICT PagePool = PhysicalSpace->GetPagePool();
-			const uint32 pAddress = PagePool.FindPageAddress(ProducerHandle, LocalLayerIndex, Tile.Local_vAddress, Tile.Local_vLevel);
+			const uint32 pAddress = PagePool.FindPageAddress(ProducerHandle, ProducerPhysicalGroupIndex, Tile.Local_vAddress, Tile.Local_vLevel);
 			if (pAddress != ~0u)
 			{
-				ProduceTarget[LocalLayerIndex].TextureRHI = PhysicalSpace->GetPhysicalTexture();
-				if (PhysicalSpace->GetDescription().bCreateRenderTarget)
+				int32 PhysicalLocalTextureIndex = 0;
+				for (uint32 ProducerLayerIndex = 0u; ProducerLayerIndex < Producer.GetNumTextureLayers(); ++ProducerLayerIndex)
 				{
-					ProduceTarget[LocalLayerIndex].PooledRenderTarget = PhysicalSpace->GetPhysicalTexturePooledRenderTarget();
+					if (Producer.GetPhysicalGroupIndexForTextureLayer(ProducerLayerIndex) == ProducerPhysicalGroupIndex)
+					{
+						ProduceTarget[ProducerLayerIndex].TextureRHI = PhysicalSpace->GetPhysicalTexture(PhysicalLocalTextureIndex);
+						if (PhysicalSpace->GetDescription().bCreateRenderTarget)
+						{
+							ProduceTarget[ProducerLayerIndex].PooledRenderTarget = PhysicalSpace->GetPhysicalTexturePooledRenderTarget(PhysicalLocalTextureIndex);
+						}
+						ProduceTarget[ProducerLayerIndex].pPageLocation = PhysicalSpace->GetPhysicalLocation(pAddress);
+						LayerMask |= 1 << ProducerLayerIndex;
+						PhysicalLocalTextureIndex++;
+					}
 				}
-				ProduceTarget[LocalLayerIndex].pPageLocation = PhysicalSpace->GetPhysicalLocation(pAddress);
-				LayerMask |= 1 << LocalLayerIndex;
 			}
 		}
 
@@ -1696,14 +1738,23 @@ void FVirtualTextureSystem::SubmitRequests(FRHICommandListImmediate& RHICmdList,
 		for (uint32 RequestIndex = 0u; RequestIndex < RequestList->GetNumLoadRequests(); ++RequestIndex)
 		{
 			const FVirtualTextureLocalTile TileToLoad = RequestList->GetLoadRequest(RequestIndex);
-			const uint32 LocalLayerMask = RequestList->GetLocalLayerMask(RequestIndex);
 			const bool bLockTile = RequestList->IsLocked(RequestIndex);
 
 			const FVirtualTextureProducerHandle ProducerHandle = TileToLoad.GetProducerHandle();
 			const FVirtualTextureProducer& Producer = Producers.GetProducer(ProducerHandle);
 
+			const uint32 ProducerPhysicalGroupMask = RequestList->GetGroupMask(RequestIndex);
+			uint32 ProducerTextureLayerMask = 0;
+			for (uint32 ProducerLayerIndex = 0; ProducerLayerIndex < Producer.GetNumTextureLayers(); ++ProducerLayerIndex)
+			{
+				if (ProducerPhysicalGroupMask & (1 << Producer.GetPhysicalGroupIndexForTextureLayer(ProducerTextureLayerMask)))
+				{
+					ProducerTextureLayerMask |= (1 << ProducerLayerIndex);
+				}
+			}
+
 			const EVTRequestPagePriority Priority = bLockTile ? EVTRequestPagePriority::High : EVTRequestPagePriority::Normal;
-			FVTRequestPageResult RequestPageResult = Producer.GetVirtualTexture()->RequestPageData(ProducerHandle, LocalLayerMask, TileToLoad.Local_vLevel, TileToLoad.Local_vAddress, Priority);
+			FVTRequestPageResult RequestPageResult = Producer.GetVirtualTexture()->RequestPageData(ProducerHandle, ProducerTextureLayerMask, TileToLoad.Local_vLevel, TileToLoad.Local_vAddress, Priority);
 			if (RequestPageResult.Status == EVTRequestPageStatus::Pending && (bLockTile || !bAsync))
 			{
 				// If we're trying to lock this tile, we're OK producing data now (and possibly waiting) as long as data is pending
@@ -1727,30 +1778,45 @@ void FVirtualTextureSystem::SubmitRequests(FRHICommandListImmediate& RHICmdList,
 
 				// try to allocate a page for each layer we need to load
 				bool bProduceTargetValid = true;
-				for (uint32 LocalLayerIndex = 0u; LocalLayerIndex < Producer.GetNumLayers(); ++LocalLayerIndex)
+				for (uint32 ProducerPhysicalGroupIndex = 0u; ProducerPhysicalGroupIndex < Producer.GetNumPhysicalGroups(); ++ProducerPhysicalGroupIndex)
 				{
 					// If mask isn't set, we must already have a physical tile allocated for this layer, don't need to allocate another one
-					if (LocalLayerMask & (1u << LocalLayerIndex))
+					if (ProducerPhysicalGroupMask & (1u << ProducerPhysicalGroupIndex))
 					{
-						FVirtualTexturePhysicalSpace* RESTRICT PhysicalSpace = Producer.GetPhysicalSpace(LocalLayerIndex);
+						FVirtualTexturePhysicalSpace* RESTRICT PhysicalSpace = Producer.GetPhysicalSpaceForPhysicalGroup(ProducerPhysicalGroupIndex);
 						FTexturePagePool& RESTRICT PagePool = PhysicalSpace->GetPagePool();
 						if (PagePool.AnyFreeAvailable(Frame))
 						{
-							const uint32 pAddress = PagePool.Alloc(this, Frame, ProducerHandle, LocalLayerIndex, TileToLoad.Local_vAddress, TileToLoad.Local_vLevel, bLockTile);
+							const uint32 pAddress = PagePool.Alloc(this, Frame, ProducerHandle, ProducerPhysicalGroupIndex, TileToLoad.Local_vAddress, TileToLoad.Local_vLevel, bLockTile);
 							check(pAddress != ~0u);
 
-							ProduceTarget[LocalLayerIndex].TextureRHI = PhysicalSpace->GetPhysicalTexture();
-							if (PhysicalSpace->GetDescription().bCreateRenderTarget)
+							//JM Loop for textures
+							int32 PhysicalLocalTextureIndex = 0;
+							for (uint32 ProducerLayerIndex = 0u; ProducerLayerIndex < Producer.GetNumTextureLayers(); ++ProducerLayerIndex)
 							{
-								ProduceTarget[LocalLayerIndex].PooledRenderTarget = PhysicalSpace->GetPhysicalTexturePooledRenderTarget();
+								if (Producer.GetPhysicalGroupIndexForTextureLayer(ProducerLayerIndex) == ProducerPhysicalGroupIndex)
+								{
+									ProduceTarget[ProducerLayerIndex].TextureRHI = PhysicalSpace->GetPhysicalTexture(PhysicalLocalTextureIndex);
+									if (PhysicalSpace->GetDescription().bCreateRenderTarget)
+									{
+										ProduceTarget[ProducerLayerIndex].PooledRenderTarget = PhysicalSpace->GetPhysicalTexturePooledRenderTarget(PhysicalLocalTextureIndex);
+									}
+									ProduceTarget[ProducerLayerIndex].pPageLocation = PhysicalSpace->GetPhysicalLocation(pAddress);
+									
+									PhysicalLocalTextureIndex++;
+
+									Allocate_pAddress[ProducerPhysicalGroupIndex] = pAddress;
+								}
 							}
-							ProduceTarget[LocalLayerIndex].pPageLocation = PhysicalSpace->GetPhysicalLocation(pAddress);
-							Allocate_pAddress[LocalLayerIndex] = pAddress;
 						}
 						else
 						{
-							const FPixelFormatInfo& PoolFormatInfo = GPixelFormats[PhysicalSpace->GetFormat()];
-							UE_LOG(LogConsoleResponse, Display, TEXT("Failed to allocate VT page from pool PF_%s"), PoolFormatInfo.Name);
+							UE_LOG(LogConsoleResponse, Display, TEXT("Failed to allocate VT page from pool %d"), PhysicalSpace->GetID());
+							for (int TextureIndex = 0; TextureIndex < PhysicalSpace->GetDescription().NumLayers; ++TextureIndex)
+							{
+								const FPixelFormatInfo& PoolFormatInfo = GPixelFormats[PhysicalSpace->GetFormat(TextureIndex)];
+								UE_LOG(LogConsoleResponse, Display, TEXT("  PF_%s"), PoolFormatInfo.Name);
+							}
 							bProduceTargetValid = false;
 							break;
 						}
@@ -1760,32 +1826,42 @@ void FVirtualTextureSystem::SubmitRequests(FRHICommandListImmediate& RHICmdList,
 				if (bProduceTargetValid)
 				{
 					// Successfully allocated required pages, now we can make the request
-					for (uint32 LocalLayerIndex = 0u; LocalLayerIndex < Producer.GetNumLayers(); ++LocalLayerIndex)
+					for (uint32 ProducerPhysicalGroupIndex = 0u; ProducerPhysicalGroupIndex < Producer.GetNumPhysicalGroups(); ++ProducerPhysicalGroupIndex)
 					{
-						if (LocalLayerMask & (1u << LocalLayerIndex))
+						if (ProducerPhysicalGroupMask & (1u << ProducerPhysicalGroupIndex))
 						{
 							// Associate the addresses we allocated with this request, so they can be mapped if required
-							const uint32 pAddress = Allocate_pAddress[LocalLayerIndex];
+							const uint32 pAddress = Allocate_pAddress[ProducerPhysicalGroupIndex];
 							check(pAddress != ~0u);
-							RequestPhysicalAddress[RequestIndex * VIRTUALTEXTURE_SPACE_MAXLAYERS + LocalLayerIndex] = pAddress;
+							RequestPhysicalAddress[RequestIndex * VIRTUALTEXTURE_SPACE_MAXLAYERS + ProducerPhysicalGroupIndex] = pAddress;
 						}
 						else
 						{
 							// Fill in pAddress for layers that are already resident
-							const FVirtualTexturePhysicalSpace* RESTRICT PhysicalSpace = Producer.GetPhysicalSpace(LocalLayerIndex);
+							const FVirtualTexturePhysicalSpace* RESTRICT PhysicalSpace = Producer.GetPhysicalSpaceForPhysicalGroup(ProducerPhysicalGroupIndex);
 							const FTexturePagePool& RESTRICT PagePool = PhysicalSpace->GetPagePool();
-							const uint32 pAddress = PagePool.FindPageAddress(ProducerHandle, LocalLayerIndex, TileToLoad.Local_vAddress, TileToLoad.Local_vLevel);
+							const uint32 pAddress = PagePool.FindPageAddress(ProducerHandle, ProducerPhysicalGroupIndex, TileToLoad.Local_vAddress, TileToLoad.Local_vLevel);
 							checkf(pAddress != ~0u,
 								TEXT("%s missing tile: LayerMask: %X, Layer %d, vAddress %06X, vLevel %d"),
-								*Producer.GetName().ToString(), LocalLayerMask, LocalLayerIndex, TileToLoad.Local_vAddress, TileToLoad.Local_vLevel);
-							ProduceTarget[LocalLayerIndex].TextureRHI = PhysicalSpace->GetPhysicalTexture();
-							ProduceTarget[LocalLayerIndex].pPageLocation = PhysicalSpace->GetPhysicalLocation(pAddress);
+								*Producer.GetName().ToString(), ProducerPhysicalGroupMask, ProducerPhysicalGroupIndex, TileToLoad.Local_vAddress, TileToLoad.Local_vLevel);
+							
+							// JM for all layers
+							int32 PhysicalLocalTextureIndex = 0;
+							for (uint32 ProducerLayerIndex = 0u; ProducerLayerIndex < Producer.GetNumTextureLayers(); ++ProducerLayerIndex)
+							{
+								if (Producer.GetPhysicalGroupIndexForTextureLayer(ProducerLayerIndex) == ProducerPhysicalGroupIndex)
+								{
+									ProduceTarget[ProducerLayerIndex].TextureRHI = PhysicalSpace->GetPhysicalTexture(PhysicalLocalTextureIndex);
+									ProduceTarget[ProducerLayerIndex].pPageLocation = PhysicalSpace->GetPhysicalLocation(pAddress);
+									PhysicalLocalTextureIndex++;
+								}
+							}
 						}
 					}
 
 					IVirtualTextureFinalizer* VTFinalizer = Producer.GetVirtualTexture()->ProducePageData(RHICmdList, FeatureLevel,
 						EVTProducePageFlags::None,
-						ProducerHandle, LocalLayerMask, TileToLoad.Local_vLevel, TileToLoad.Local_vAddress,
+						ProducerHandle, ProducerTextureLayerMask, TileToLoad.Local_vLevel, TileToLoad.Local_vAddress,
 						RequestPageResult.Handle,
 						ProduceTarget);
 					if (VTFinalizer)
@@ -1799,12 +1875,12 @@ void FVirtualTextureSystem::SubmitRequests(FRHICommandListImmediate& RHICmdList,
 				else
 				{
 					// Failed to allocate required physical pages for the tile, free any pages we did manage to allocate
-					for (uint32 LocalLayerIndex = 0u; LocalLayerIndex < Producer.GetNumLayers(); ++LocalLayerIndex)
+					for (uint32 ProducerPhysicalGroupIndex = 0u; ProducerPhysicalGroupIndex < Producer.GetNumPhysicalGroups(); ++ProducerPhysicalGroupIndex)
 					{
-						const uint32 pAddress = Allocate_pAddress[LocalLayerIndex];
+						const uint32 pAddress = Allocate_pAddress[ProducerPhysicalGroupIndex];
 						if (pAddress != ~0u)
 						{
-							FVirtualTexturePhysicalSpace* RESTRICT PhysicalSpace = Producer.GetPhysicalSpace(LocalLayerIndex);
+							FVirtualTexturePhysicalSpace* RESTRICT PhysicalSpace = Producer.GetPhysicalSpaceForPhysicalGroup(ProducerPhysicalGroupIndex);
 							FTexturePagePool& RESTRICT PagePool = PhysicalSpace->GetPagePool();
 							PagePool.Free(this, pAddress);
 						}
@@ -1833,24 +1909,24 @@ void FVirtualTextureSystem::SubmitRequests(FRHICommandListImmediate& RHICmdList,
 			FVirtualTextureSpace* Space = GetSpace(MappingRequest.SpaceID);
 			FVirtualTexturePhysicalSpace* PhysicalSpace = GetPhysicalSpace(MappingRequest.PhysicalSpaceID);
 
-			PhysicalSpace->GetPagePool().MapPage(Space, PhysicalSpace, MappingRequest.LayerIndex, MappingRequest.vLevel, MappingRequest.vAddress, MappingRequest.Local_vLevel, MappingRequest.pAddress);
+			PhysicalSpace->GetPagePool().MapPage(Space, PhysicalSpace, MappingRequest.PageTableLayerIndex, MappingRequest.vLevel, MappingRequest.vAddress, MappingRequest.Local_vLevel, MappingRequest.pAddress);
 		}
 
 		// Update page mappings for any requested page that completed allocation this frame
 		for (uint32 RequestIndex = 0u; RequestIndex < RequestList->GetNumMappingRequests(); ++RequestIndex)
 		{
 			const FMappingRequest MappingRequest = RequestList->GetMappingRequest(RequestIndex);
-			const uint32 pAddress = RequestPhysicalAddress[MappingRequest.LoadRequestIndex * VIRTUALTEXTURE_SPACE_MAXLAYERS + MappingRequest.LocalLayerIndex];
+			const uint32 pAddress = RequestPhysicalAddress[MappingRequest.LoadRequestIndex * VIRTUALTEXTURE_SPACE_MAXLAYERS + MappingRequest.ProducerPhysicalGroupIndex];
 			if (pAddress != ~0u)
 			{
 				const FVirtualTextureLocalTile& TileToLoad = RequestList->GetLoadRequest(MappingRequest.LoadRequestIndex);
 				const FVirtualTextureProducerHandle ProducerHandle = TileToLoad.GetProducerHandle();
 				FVirtualTextureProducer& Producer = Producers.GetProducer(ProducerHandle);
-				FVirtualTexturePhysicalSpace* PhysicalSpace = Producer.GetPhysicalSpace(MappingRequest.LocalLayerIndex);
+				FVirtualTexturePhysicalSpace* PhysicalSpace = Producer.GetPhysicalSpaceForPhysicalGroup(MappingRequest.ProducerPhysicalGroupIndex);
 				FVirtualTextureSpace* Space = GetSpace(MappingRequest.SpaceID);
-				check(RequestList->GetLocalLayerMask(MappingRequest.LoadRequestIndex) & (1u << MappingRequest.LocalLayerIndex));
+				check(RequestList->GetGroupMask(MappingRequest.LoadRequestIndex) & (1u << MappingRequest.ProducerPhysicalGroupIndex));
 
-				PhysicalSpace->GetPagePool().MapPage(Space, PhysicalSpace, MappingRequest.LayerIndex, MappingRequest.vLevel, MappingRequest.vAddress, MappingRequest.Local_vLevel, pAddress);
+				PhysicalSpace->GetPagePool().MapPage(Space, PhysicalSpace, MappingRequest.PageTableLayerIndex, MappingRequest.vLevel, MappingRequest.vAddress, MappingRequest.Local_vLevel, pAddress);
 			}
 		}
 	}
@@ -1867,9 +1943,10 @@ void FVirtualTextureSystem::SubmitRequests(FRHICommandListImmediate& RHICmdList,
 			FVirtualTextureSpace* Space = AllocatedVT->GetSpace();
 
 			uint32 NumFullyMappedLayers = 0u;
-			for (uint32 LayerIndex = 0u; LayerIndex < AllocatedVT->GetNumLayers(); ++LayerIndex)
+			for (uint32 PageTableLayerIndex = 0u; PageTableLayerIndex < AllocatedVT->GetNumPageTableLayers(); ++PageTableLayerIndex)
 			{
-				const FVirtualTextureProducerHandle ProducerHandle = AllocatedVT->GetProducerHandle(LayerIndex);
+				uint32 ProducerIndex = AllocatedVT->GetProducerIndexForPageTableLayer(PageTableLayerIndex);
+				const FVirtualTextureProducerHandle ProducerHandle = AllocatedVT->GetUniqueProducerHandle(ProducerIndex);
 				const FVirtualTextureProducer* Producer = Producers.FindProducer(ProducerHandle);
 				if (!Producer)
 				{
@@ -1877,16 +1954,16 @@ void FVirtualTextureSystem::SubmitRequests(FRHICommandListImmediate& RHICmdList,
 					continue;
 				}
 
-				const uint32 ProducerIndex = AllocatedVT->GetUniqueProducerIndexForLayer(LayerIndex);
+				uint32 ProducerPhysicalGroupIndex = AllocatedVT->GetProducerPhysicalGroupIndexForPageTableLayer(PageTableLayerIndex);
+
 				const uint32 ProducerMipBias = AllocatedVT->GetUniqueProducerMipBias(ProducerIndex);
-				const uint32 LocalLayerIndex = AllocatedVT->GetLocalLayerToProduce(LayerIndex);
 				const uint32 WidthInTiles = Producer->GetWidthInTiles();
 				const uint32 HeightInTiles = Producer->GetHeightInTiles();
 				const uint32 MaxLevel = FMath::Min(Producer->GetMaxLevel(), AllocatedVT->GetMaxLevel());
 
-				FVirtualTexturePhysicalSpace* PhysicalSpace = AllocatedVT->GetPhysicalSpace(LayerIndex);
+				FVirtualTexturePhysicalSpace* PhysicalSpace = AllocatedVT->GetPhysicalSpaceForPageTableLayer(PageTableLayerIndex);
 				FTexturePagePool& PagePool = PhysicalSpace->GetPagePool();
-				FTexturePageMap& PageMap = Space->GetPageMap(LayerIndex);
+				FTexturePageMap& PageMap = Space->GetPageMapForPageTableLayer(PageTableLayerIndex);
 				
 				bool bIsLayerFullyMapped = false;
 				for (uint32 Local_vLevel = 0u; Local_vLevel <= MaxLevel; ++Local_vLevel)
@@ -1906,10 +1983,10 @@ void FVirtualTextureSystem::SubmitRequests(FRHICommandListImmediate& RHICmdList,
 							{
 								const uint32 Local_vAddress = FMath::MortonCode2(TileX) | (FMath::MortonCode2(TileY) << 1);
 
-								pAddress = PagePool.FindPageAddress(ProducerHandle, LocalLayerIndex, Local_vAddress, Local_vLevel);
+								pAddress = PagePool.FindPageAddress(ProducerHandle, ProducerPhysicalGroupIndex, Local_vAddress, Local_vLevel);
 								if (pAddress != ~0u)
 								{
-									PagePool.MapPage(Space, PhysicalSpace, LayerIndex, vLevel, vAddress, vLevel, pAddress);
+									PagePool.MapPage(Space, PhysicalSpace, PageTableLayerIndex, vLevel, vAddress, vLevel, pAddress);
 								}
 								else
 								{
@@ -1927,7 +2004,7 @@ void FVirtualTextureSystem::SubmitRequests(FRHICommandListImmediate& RHICmdList,
 				}
 			}
 
-			if (NumFullyMappedLayers < AllocatedVT->GetNumLayers())
+			if (NumFullyMappedLayers < AllocatedVT->GetNumPageTableLayers())
 			{
 				++Index;
 			}
