Preserve stencil ref when setting base pass dither stencil state to read-only Equal.

diff --git a/Engine/Source/Runtime/Renderer/Private/BasePassRendering.cpp b/Engine/Source/Runtime/Renderer/Private/BasePassRendering.cpp
--- a/Engine/Source/Runtime/Renderer/Private/BasePassRendering.cpp
+++ b/Engine/Source/Runtime/Renderer/Private/BasePassRendering.cpp
@@ -135,18 +135,9 @@ void SetBasePassDitheredLODTransitionState(const FSceneView* SceneView, const FM
 
 		if (ViewInfo->bAllowStencilDither)
 		{
-			if (ViewInfo->StaticMeshFadeOutDitheredLODMap[StaticMeshId])
-			{
-				DrawRenderState.SetDepthStencilState(
-					TStaticDepthStencilState<
-					false, CF_Equal,
-					true, CF_Always, SO_Keep, SO_Keep, SO_Replace,
-					false, CF_Always, SO_Keep, SO_Keep, SO_Keep,
-					0xFF, GET_STENCIL_BIT_MASK(RECEIVE_DECAL, 1) | STENCIL_LIGHTING_CHANNELS_MASK(0x7)
-					>::GetRHI());
-			}
-			else if (ViewInfo->StaticMeshFadeInDitheredLODMap[StaticMeshId])
+			if (ViewInfo->StaticMeshFadeOutDitheredLODMap[StaticMeshId] || ViewInfo->StaticMeshFadeInDitheredLODMap[StaticMeshId])
 			{
+				const uint32 RestoreStencilRef = DrawRenderState.GetStencilRef();
 				DrawRenderState.SetDepthStencilState(
 					TStaticDepthStencilState<
 					false, CF_Equal,
@@ -154,6 +145,7 @@ void SetBasePassDitheredLODTransitionState(const FSceneView* SceneView, const FM
 					false, CF_Always, SO_Keep, SO_Keep, SO_Keep,
 					0xFF, GET_STENCIL_BIT_MASK(RECEIVE_DECAL, 1) | STENCIL_LIGHTING_CHANNELS_MASK(0x7)
 					>::GetRHI());
+				DrawRenderState.SetStencilRef(RestoreStencilRef);
 			}
 		}
 	}
