GTAO Updates.
Added tuning controls for intensity and power.
Added Quality options set via r.AmbientOcclusionMaxQuality
Added optimized version of the main horizon searhc inner loop that pretends the depth buffer is othographic height map rather than perspective.
Use modified random normal texture

diff --git a/Engine/Shaders/Private/PostProcessAmbientOcclusion.usf b/Engine/Shaders/Private/PostProcessAmbientOcclusion.usf
--- a/Engine/Shaders/Private/PostProcessAmbientOcclusion.usf
+++ b/Engine/Shaders/Private/PostProcessAmbientOcclusion.usf
@@ -22,10 +22,15 @@
 // Testing the GTAO Implementation between a heavy weight reference and an in game one
 #define GTAO_REF 0
 
+// Simplified optimized version of GTAO.
+// The biggest difference here is the depth map is treated as an othographic projected height map rather than a perspective one. 
+// This leads to differences in angles and fall off distances but generally the differences are minor. For extreme FOV cases this might be a problem.
+#define GTAO_OPT 0
+
 // 0: classic with weighted sample, 1: don't normalize and adjust the formula to be simpler and faster - can look better and is cheaper (Alchemy like?)
 #define OPTIMIZATION_O1 1
 
-// 1:lowest quality, 2:medium , 3:high, more doesn't give too much (maybe HZB mip computations should be adjusted)
+// 1:lowest quality, 2:medium , 3:high, more doesn't give too much (maybe HZB mip computations should `be adjusted)
 //#define SAMPLE_STEPS 3
 
 // 0:off / 1:show samples on the right side of the screen
@@ -901,60 +906,56 @@ void MainSSAOSmoothCS(
 }
 #endif
 
+#if GTAO_REF
 
-float4 GTAOParams[1];
-
-#define FALLOFF_START		100.0
-#define FALLOFF_START_SQ	(FALLOFF_START*FALLOFF_START)
-#define FALLOFF_END			150.0
-#define FALLOFF_END_SQ		(FALLOFF_END*FALLOFF_END)
+#define GTAO_NUMROTS  32
+#define GTAO_NUMTAPS  64
 
-#if GTAO_REF
-#define NumPhi  32
-#define NumTaps  64
 #else
-#define NumPhi  1
-#define NumTaps  7
-#endif
 
-#define PI_HALF (PI*0.5)
+#define GTAO_NUMROTS  1
 
-#if COMPUTE_SHADER
-RWTexture2D<float2> HorizonOutTexture;
+#if SHADER_QUALITY == 0
+	// very low
+	#define GTAO_NUMTAPS 5
+#elif SHADER_QUALITY == 1
+	// low
+	#define GTAO_NUMTAPS 9
+#elif SHADER_QUALITY == 2
+	// medium
+	#define GTAO_NUMTAPS 15
+#elif SHADER_QUALITY == 3
+	// high
+	#define GTAO_NUMTAPS 25
+#else // SHADER_QUALITY == 4
+	// very high
+	#define GTAO_NUMTAPS 35
+#endif
 #endif
 
-// The following is used to help the compiler generate more optimal code.
-struct GTAO_COMMON
-{
-	float HZB_BiasX;
-	float HZB_BiasY;
-	float Falloff_Scale;
-	float Falloff_Bias;
-};
 
 
-void Fill_GTAO_Common( inout GTAO_COMMON GtaoCommon)
-{
-	GtaoCommon.HZB_BiasX = HZBRemapping.z;
-	GtaoCommon.HZB_BiasY = HZBRemapping.w;
+float4 GTAOParams[2];
 
-	GtaoCommon.Falloff_Scale = 1.0f / (FALLOFF_END_SQ - FALLOFF_START_SQ);
-	GtaoCommon.Falloff_Bias = -FALLOFF_START_SQ * GtaoCommon.Falloff_Scale; 
-}
+#define FALLOFF_START		50.0
+#define FALLOFF_END			250.0
+#define FALLOFF_START_SQ	(FALLOFF_START*FALLOFF_START)
+#define FALLOFF_END_SQ		(FALLOFF_END*FALLOFF_END)
 
 
-float Falloff(float distSq, GTAO_COMMON gtao_common)
-{
-	return 2.0 * saturate(distSq* gtao_common.Falloff_Scale + gtao_common.Falloff_Bias   );
 
-}
+#define PI_HALF (PI*0.5)
+
+#if COMPUTE_SHADER
+RWTexture2D<float2> HorizonOutTexture;
+#endif
+
 
 float Falloff(float distSq)
 {
 	float Falloff_Scale = 1.0f / (FALLOFF_END_SQ - FALLOFF_START_SQ);
 	float Falloff_Bias = -FALLOFF_START_SQ * Falloff_Scale; 
-
-	return 2.0 * saturate(distSq* Falloff_Scale + Falloff_Bias   );
+	return  2.0 * saturate(distSq* Falloff_Scale + Falloff_Bias   );
 }
 
 
@@ -971,17 +972,50 @@ float3 GetViewSpacePosFromHZB(float2 UV, float MipLevel)
 	return ScreenToViewPos(UV, SceneDepth);
 }
 
-float3 GetViewSpacePosFromHZB(float2 UV, float MipLevel, GTAO_COMMON gtao_common)
+
+float3 GetNormal(float2 UV, float3 ViewSpacePosMid, bool bConstructFromDepth)
 {
-	float SceneDepth = GetDepthFromAOInput(UV);
+	float3 ViewSpaceNormal;
 
-	float2 HZBUV;
-	HZBUV.x = UV.x * HZBRemapping.x + gtao_common.HZB_BiasX;
-	HZBUV.y = UV.y * HZBRemapping.y + gtao_common.HZB_BiasY;
-	SceneDepth = Texture2DSampleLevel( PostprocessInput1, PostprocessInput1Sampler, HZBUV,MipLevel).r;
-	SceneDepth = ConvertFromDeviceZ(SceneDepth);
+	if(bConstructFromDepth)
+	{
+		// Get the normal derived from the depth buffer
+		float2 DeltaUV = View.BufferSizeAndInvSize.zw;
+	
+		float DeviceZ		= Texture2DSampleLevel(SceneTexturesStruct.SceneDepthTexture, SceneTexturesStruct.SceneDepthTextureSampler, UV,0).r;									
+		float DeviceZLeft	= Texture2DSampleLevel(SceneTexturesStruct.SceneDepthTexture, SceneTexturesStruct.SceneDepthTextureSampler, UV + float2(-DeltaUV.x,  0.0f),0).r;		
+		float DeviceZTop	= Texture2DSampleLevel(SceneTexturesStruct.SceneDepthTexture, SceneTexturesStruct.SceneDepthTextureSampler, UV + float2( 0.0f     , -DeltaUV.y),0).r;	
+		float DeviceZRight	= Texture2DSampleLevel(SceneTexturesStruct.SceneDepthTexture, SceneTexturesStruct.SceneDepthTextureSampler, UV + float2( DeltaUV.x,  0.0f),0).r;		
+		float DeviceZBottom = Texture2DSampleLevel(SceneTexturesStruct.SceneDepthTexture, SceneTexturesStruct.SceneDepthTextureSampler, UV + float2( 0.0f     ,  DeltaUV.y),0).r;	
 
-	return ScreenToViewPos(UV, SceneDepth);
+		float DeviceZDdx	= TakeSmallerAbsDelta(DeviceZLeft, DeviceZ, DeviceZRight);
+		float DeviceZDdy	= TakeSmallerAbsDelta(DeviceZTop, DeviceZ, DeviceZBottom);
+
+		float ZRight		= ConvertFromDeviceZ(DeviceZ + DeviceZDdx);
+		float ZDown			= ConvertFromDeviceZ(DeviceZ + DeviceZDdy);
+
+		float3 Right		= ScreenToViewPos(UV+ float2( DeltaUV.x,      0.0f) , ZRight)-ViewSpacePosMid;
+		float3 Down			= ScreenToViewPos(UV+ float2(      0.0f, DeltaUV.y) , ZDown) -ViewSpacePosMid;
+		
+		ViewSpaceNormal = normalize(cross(Right, Down));
+
+	}
+	else
+	{
+		// Get the normal from the normal buffer
+		float3 WorldNormal		= GetGBufferData(UV, false).WorldNormal;
+		ViewSpaceNormal	= normalize(mul(WorldNormal, (float3x3)View.TranslatedWorldToView));
+
+	}
+
+	return ViewSpaceNormal;
+}
+
+
+float GetLinearDepthProj(float2 ScreenUV)
+{
+	float DeviceZ = Texture2DSampleLevel(SceneTexturesStruct.SceneDepthTexture, SceneTexturesStruct.SceneDepthTextureSampler, ScreenUV, 0).r;
+	return  1.0f / (DeviceZ * View.InvDeviceZToWorldZTransform[2] - View.InvDeviceZToWorldZTransform[3]);
 }
 
 
@@ -990,34 +1024,67 @@ float SearchForLargestAngleRef(float2 inUV, float2 deltaUV, float3 ViewSpacePosi
 	float BiggestAngle = -1.0f;
 
 	deltaUV.y	  *=-1;
-	deltaUV	      *= 1;
+	deltaUV.y	  *=.3;
+	float2 UV	   = inUV ;
+
+#if GTAO_OPT
+
+	// Precompute as much stuff upfront
+	float CurrDepth			= GetLinearDepthProj(UV);
+	UV						+= deltaUV;
+	float2 DeltaVS_XY	    = ScreenToViewPos(UV, CurrDepth).xy - ViewSpacePosition.xy;
+	float  DeltaVS_XYLen    = sqrt(dot( DeltaVS_XY.xy, DeltaVS_XY.xy));
+	float  CurrVS_XYLen		= DeltaVS_XYLen;
 
-	for(int i=1; i<NumTaps; i++)
+	for(int i=1; i<=GTAO_NUMTAPS; i++)
 	{
-		float fi = (float) i;
-		float2 UV		= inUV +  fi*deltaUV;
+		float SceneDepth	= GetLinearDepthProj(UV);
+		float DiffDepth		= (CurrDepth - SceneDepth);
+
+		float distSQ		= (CurrVS_XYLen*CurrVS_XYLen) + (DiffDepth*DiffDepth);
+
+		UV					+= deltaUV;
+		CurrVS_XYLen		+= DeltaVS_XYLen;
+
+		// At this point we could early out if DistSQ is greater than falloff
+		BRANCH
+		if(distSQ < FALLOFF_END_SQ)
+		{
+			float ooDist    = rsqrt(distSQ);
+			float cosh		= ooDist * DiffDepth;
 
+			float falloff	= Falloff(distSQ);
+			BiggestAngle	= max(BiggestAngle, cosh - falloff);
+		}
+
+	}	
+#else
+	for(int i=1; i<=GTAO_NUMTAPS; i++)
+	{
+		UV					+= deltaUV;
+	
 		float SceneDepth	= GetDepthFromAOInput(UV);
 		float3 SampVS		= ScreenToViewPos(UV, SceneDepth);
+	
 		float3 DiffV		= (SampVS - ViewSpacePosition) ;
 		float distSQ		= dot(DiffV, DiffV);
-
+	
 		// At this point we could early out if DistSQ is greater than falloff
 		BRANCH
 		if(distSQ < FALLOFF_END_SQ)
 		{
-			float ooDist    = 1.0f / sqrt(distSQ);
-			float cosh		= ooDist * dot(DiffV,ViewDir );
-
-			float falloff = Falloff(distSQ);
-			BiggestAngle = max(BiggestAngle, cosh - falloff);
+			float ooDist    = rsqrt(distSQ);
+			float cosh		= ooDist * dot(DiffV, ViewDir );
+	
+			float falloff	= Falloff(distSQ);
+			BiggestAngle	= max(BiggestAngle, cosh - falloff);
 		}
-	}
+	}	
+#endif
 
 	return BiggestAngle;
 }
 
-
 float SearchForLargestAngle(float2 inUV, float2 deltaUV, float3 ViewSpacePosition, float3 ViewDir)
 {
 #if GTAO_REF
@@ -1027,32 +1094,90 @@ float SearchForLargestAngle(float2 inUV, float2 deltaUV, float3 ViewSpacePositio
 	float BiggestAngle = -1.0f;
 
 	deltaUV.y	  *=-1;
-	deltaUV		  *= 2;
-
 	float2 UV	= inUV;
 
-	GTAO_COMMON GtaoCommon;
-	Fill_GTAO_Common(GtaoCommon);
+	float MipLevel = 0;
+	
 
+#if GTAO_OPT
+	// Precompute as much stuff upfront
+	float CurrDepth			= GetLinearDepthProj(UV);
+	UV						+= deltaUV;
+	float2 DeltaVS_XY	    = ScreenToViewPos(UV, CurrDepth).xy - ViewSpacePosition.xy;
+	float  DeltaVS_XYLen    = sqrt(dot( DeltaVS_XY.xy, DeltaVS_XY.xy));
+	float  CurrVS_XYLen		= DeltaVS_XYLen;
 
-	for(int i=1; i<=NumTaps; i++)
+	float2 HZB_UV  = float2( UV.x * HZBRemapping.x + HZBRemapping.z, UV.y * HZBRemapping.y + HZBRemapping.w);
+	float2 HZB_dUV = float2( deltaUV.x * HZBRemapping.x, deltaUV.y * HZBRemapping.y );
+
+	UNROLL
+	for(int i=1; i<=GTAO_NUMTAPS; i++)
 	{
-		float MipLevel = 0;
+		if(i==2)
+		{
+			HZB_dUV += HZB_dUV;
+			DeltaVS_XYLen	+= DeltaVS_XYLen;
+			MipLevel++;
+		}
+
+		if(i==GTAO_NUMTAPS/4)
+		{
+			HZB_dUV += HZB_dUV;
+			DeltaVS_XYLen	+= DeltaVS_XYLen;
+			MipLevel++;
+		}
+		else if(i==GTAO_NUMTAPS/2)
+		{
+			HZB_dUV += HZB_dUV;
+			DeltaVS_XYLen	+= DeltaVS_XYLen;
+			MipLevel++;
+		}
 
+		float DeviceZ		= Texture2DSampleLevel( PostprocessInput1, PostprocessInput1Sampler, HZB_UV, MipLevel).r;
+		float SceneDepth    = 1.0f / (DeviceZ * View.InvDeviceZToWorldZTransform[2] - View.InvDeviceZToWorldZTransform[3]);
+
+		float DiffDepth		= (CurrDepth - SceneDepth);
+		float distSQ		= (CurrVS_XYLen*CurrVS_XYLen) + (DiffDepth*DiffDepth);
+
+		HZB_UV				+= HZB_dUV;
+		CurrVS_XYLen		+= DeltaVS_XYLen;
+
+		// At this point we could early out if DistSQ is greater than falloff
+		BRANCH
+		if(distSQ < FALLOFF_END_SQ)
+		{
+			float ooDist    = rsqrt(distSQ);
+			float cosh		= ooDist * DiffDepth;
+
+			float falloff	= Falloff(distSQ);
+			BiggestAngle	= max(BiggestAngle, cosh - falloff);
+		}
+	}	
+
+
+#else
+	for(int i=1; i<=GTAO_NUMTAPS; i++)
+	{
 		if(i==3)
 		{
 			deltaUV += deltaUV;
-			MipLevel = 1;
+			MipLevel++;
+		}
+
+		if(i==GTAO_NUMTAPS/4)
+		{
+			deltaUV += deltaUV;
+			MipLevel++;
 		}
-		else if(i==5)
+		else if(i==GTAO_NUMTAPS/2)
 		{
 			deltaUV += deltaUV;
-			MipLevel = 2;
+			MipLevel++;
 		}
 
 		UV				+= deltaUV;
 		
-		float3 SampVS   = GetViewSpacePosFromHZB(UV,MipLevel, GtaoCommon);
+		float3 SampVS   = GetViewSpacePosFromHZB(UV,MipLevel);
 		float3 DiffV	= (SampVS - ViewSpacePosition) ;
 		float distSQ	= dot(DiffV, DiffV);
 
@@ -1063,11 +1188,12 @@ float SearchForLargestAngle(float2 inUV, float2 deltaUV, float3 ViewSpacePositio
 		{
 			float ooDist    = rsqrt(distSQ);
 			float cosh		= ooDist * dot(DiffV,ViewDir );
-			float falloff = Falloff(distSQ, GtaoCommon);
-			BiggestAngle = max(BiggestAngle, cosh - falloff);
+			float falloff	= Falloff(distSQ);
+			BiggestAngle	= max(BiggestAngle, cosh - falloff);
 		}
-
 	}
+#endif	
+
 
 	return BiggestAngle;
 }
@@ -1077,11 +1203,15 @@ float SearchForLargestAngle(float2 inUV, float2 deltaUV, float3 ViewSpacePositio
 float2 SearchAxisForAngles(float2 UV, float3 Dir, float3 ViewDir, float3 ViewSpacePos, float Scale)
 {
 	// TODO - Fix aspect ratio
-	float RandScale = 1.0f;//GTAOParams[0].y;
-	float2 DeltaUV	= (RandScale * Scale) * float2(View.BufferSizeAndInvSize.z * Dir.x, View.BufferSizeAndInvSize.w * Dir.y);
+	float2 DeltaUV	= Scale * float2(View.BufferSizeAndInvSize.z * Dir.x, View.BufferSizeAndInvSize.w * Dir.y);
 
-	float Ang1 = SearchForLargestAngle( UV,   DeltaUV, ViewSpacePos, ViewDir);
-	float Ang2 = SearchForLargestAngle( UV,  -DeltaUV, ViewSpacePos, ViewDir);
+	float O =  0.0f;//GTAOParams[0].z; 
+
+	float2 StartUV = UV + (DeltaUV*O);
+	float Ang1 = SearchForLargestAngle( StartUV,  DeltaUV, ViewSpacePos, ViewDir);
+
+	StartUV = UV - (DeltaUV*O);
+	float Ang2 = SearchForLargestAngle( StartUV, -DeltaUV, ViewSpacePos, ViewDir);
 
 	return float2(Ang1, Ang2);
 }
@@ -1093,20 +1223,27 @@ float ComputeInnerIntegral(float2 Angles, float3 ScreenDir, float3 ViewDir, floa
 	float3 Perp				= cross(ViewDir, PlaneNormal);
 	float3 ProjNormal		= ViewSpaceNormal - PlaneNormal * dot(ViewSpaceNormal, PlaneNormal);
 
-	float ProjNormalMagSq	= dot(ProjNormal, ProjNormal);
-	float RecipMag			= rsqrt(ProjNormalMagSq);
-	float ProjNormalMag		= 1.0f / max(RecipMag, 1e-7);
+	float LenProjNormal		= length(ProjNormal);
+	//float ProjNormalMagSq	= dot(ProjNormal, ProjNormal);
+	float RecipMag			= 1.0f / (LenProjNormal+1e-6);
 	float CosAng			= dot(ProjNormal, Perp) * RecipMag;	
 	float Gamma				= acosFast(CosAng) - PI_HALF;				
 	float CosGamma			= dot(ProjNormal, ViewDir) * RecipMag;
 	float SinGamma  		= CosAng * -2.0f;					
 
-	// clamp to normal hemisphere
-	Angles.x = Gamma + max(-Angles.x - Gamma, -PI_HALF);
-	Angles.y = Gamma + min( Angles.y - Gamma,  PI_HALF);
+	const float ANGLEBIAS = 0;// PI_HALF/4.0f;
+
+	// clamp to normal hemisphere 
+	Angles.x = Gamma + max(-Angles.x - Gamma, -(PI_HALF-ANGLEBIAS) );
+	Angles.y = Gamma + min( Angles.y - Gamma,  (PI_HALF-ANGLEBIAS) );
+
+	float AO = ( 0.25 * ( (Angles.x * SinGamma + CosGamma - cos((2.0 * Angles.x) - Gamma)) +
+				  	      (Angles.y * SinGamma + CosGamma - cos((2.0 * Angles.y) - Gamma)) ));
 
-	float AO = (ProjNormalMag * 0.25 * ( (Angles.x * SinGamma + CosGamma - cos((2.0 * Angles.x) - Gamma)) +
-				  			             (Angles.y * SinGamma + CosGamma - cos((2.0 * Angles.y) - Gamma)) ));
+	AO = saturate(AO);
+
+	//const float Bias = 0.2;
+	//AO = saturate( (AO - Bias)  *  (1.0/(1.0-Bias))  );
 
 	// Fade out based on user defined distance
 	float Mul = ScreenSpaceAOParams[4].x;
@@ -1136,36 +1273,49 @@ void HorizonSearchAndInnerIntegratePSandCS(in float2 UV, out float4 OutColor)
 	}
 	
 	float2 ViewportUVToRandomUV = ScreenSpaceAOParams[1].xy;
-	const float2 TemporalOffset = ScreenSpaceAOParams[3].xy;
-	float2 RandomVec = (Texture2DSample(RandomNormalTexture, RandomNormalTextureSampler, TemporalOffset + UV * ViewportUVToRandomUV).rg * 2 - 1) ;
+	float TemporalCos = GTAOParams[0].x;
+	float TemporalSin = GTAOParams[0].y;
+	float2 TemporalOffset = 0;// GTAOParams[1].xy;
+	float2 RandomTexVec = (Texture2DSample(RandomNormalTexture, RandomNormalTextureSampler, TemporalOffset + UV * ViewportUVToRandomUV).rg * 2 - 1) ;
+	float2 RandomVec;
+	
+	// Rotate the random Vec
+	RandomVec.x = dot(RandomTexVec.xy, float2(TemporalCos, -TemporalSin ));
+	RandomVec.y = dot(RandomTexVec.xy, float2(TemporalSin, TemporalCos ));
 
 	float3 ViewSpacePos  = ScreenToViewPos(UV,SceneDepth);
 
 	// Get Normal and transform into viewspace
-	float3 WorldNormal		= GetGBufferData(UV, false).WorldNormal;								// TAP
-	float3 ViewSpaceNormal	= normalize(mul(WorldNormal, (float3x3)View.TranslatedWorldToView));
+	float3 ViewSpaceNormal = GetNormal(UV, ViewSpacePos, false);
+#if GTAO_OPT
 	float3 ViewDir			= -normalize(ViewSpacePos.xyz);
+#else
+	float3 ViewDir			= float3(0,0,-1);
+#endif
 
 	float Sum=0.0;
 
+	// Given the depth determine how many pixels away we are guranteed to be outside of the falloff distance.
+	float InvTanHalfFov			= ScreenSpaceAOParams[3].w;
+	float ScreenPixelsToCutoff  = View.BufferSizeAndInvSize.x * ((FALLOFF_END*(1.0f/InvTanHalfFov)) / ViewSpacePos.z) ; 
+
 	// Get radius based on sampling 
 	// Get radius based on distance - TODO - Need to use FOV and aspect ratio
-	float Scale = 1000.0 / ViewSpacePos.z;
+	float Scale = ScreenPixelsToCutoff /  GTAO_NUMTAPS;
 
+	// Clamp the scale. Too small and it won't be jumping by more than 1 pixel. Too much and it'll thrash the cache
 	Scale = max(Scale, 2.0f);
 
-	// TODO - Opt - Determine how many taps we'll step before we hit the attenuation radius.
-
 #if GTAO_REF
-	for(int i=0; i<NumPhi; i++)
+	for(int i=0; i<GTAO_NUMROTS; i++)
 	{
-		float phi			= ((float)i) *  (PI / NumPhi);
+		float phi			= ((float)i) *  (PI / GTAO_NUMROTS);
 		float3 ScreenDir	= float3(cos(phi), sin(phi), 0.0);
 #else
 		float3 ScreenDir	= float3(RandomVec.x, RandomVec.y, 0.0);
 #endif
 
-		float2 Angles = SearchAxisForAngles( UV,  ScreenDir,  ViewDir,  ViewSpacePos, Scale);
+		float2 Angles = SearchAxisForAngles( UV ,  ScreenDir,  ViewDir,  ViewSpacePos, Scale);
 
 		Angles.x = acosFast(Angles.x);
 		Angles.y = acosFast(Angles.y);
@@ -1175,7 +1325,14 @@ void HorizonSearchAndInnerIntegratePSandCS(in float2 UV, out float4 OutColor)
 	}
 #endif
 
-	OutColor.r = Sum   / (float)NumPhi;
+
+	float AmbientOcclusionIntensity = ScreenSpaceAOParams[0].w;
+	float AmbientOcclusionPower = ScreenSpaceAOParams[0].x;
+	// user adjust AO
+	Sum = Sum   / (float)GTAO_NUMROTS;
+	Sum = 1 - (1 - pow(abs(Sum), AmbientOcclusionPower)) * AmbientOcclusionIntensity;
+
+	OutColor.r = Sum ;
 
 	return;
 }
@@ -1250,6 +1407,12 @@ void GTAOInnerIntegratePSandCS(in float2 UV, out float4 OutColor)
 	Angles.y = acosFast(Angles.y);
 
 	float AO = ComputeInnerIntegral(Angles, ScreenDir, ViewDir, ViewSpaceNormal, SceneDepth);
+
+	float AmbientOcclusionIntensity = ScreenSpaceAOParams[0].w;
+	float AmbientOcclusionPower		= ScreenSpaceAOParams[0].x;
+	// user adjust AO
+	AO = 1 - (1 - pow(abs(AO), AmbientOcclusionPower)) * AmbientOcclusionIntensity;
+
 	OutColor.rgb = AO ;
 
 	return;
diff --git a/Engine/Source/Runtime/Renderer/Private/CompositionLighting/PostProcessAmbientOcclusion.cpp b/Engine/Source/Runtime/Renderer/Private/CompositionLighting/PostProcessAmbientOcclusion.cpp
--- a/Engine/Source/Runtime/Renderer/Private/CompositionLighting/PostProcessAmbientOcclusion.cpp
+++ b/Engine/Source/Runtime/Renderer/Private/CompositionLighting/PostProcessAmbientOcclusion.cpp
@@ -1214,17 +1214,38 @@ public:
 	{
 		const FFinalPostProcessSettings& Settings = View.FinalPostProcessSettings;
 
-		FVector4 Value[1];
-		static float CurrRand = 0.0f;
-		float NewRand = FMath::FRand();
-		if (fabs(NewRand - CurrRand) < 0.2f)
+		uint32 TemporalFrame = 0;
+		uint32 Frame = 0;
+		
+		const FSceneViewState* ViewState = static_cast<const FSceneViewState*>(View.State);
+
+		if (ViewState)
 		{
-			NewRand = CurrRand < 0.5f ? NewRand + 0.2f : NewRand - 0.2f;
+			TemporalFrame = ViewState->GetCurrentTemporalAASampleIndex();
+			Frame = ViewState->GetFrameIndex();
 		}
-		CurrRand = NewRand;
-		Value[0] = FVector4(NewRand, NewRand*0.5f + 1.0f, 0.0f, 0.0f);
 
-		SetShaderValueArray(RHICmdList, ShaderRHI, GTAOParams, Value, 1);
+		static int32 DiffS = Frame - TemporalFrame;
+
+		FVector4 GTAOParam[2];
+
+		const float Rots[6]		= { 60.0f, 300.0f, 180.0f, 240.0f, 120.0f, 0.0f };
+		const float Offsets[4]	= { 0.0f, 0.5f, 0.25f, 0.75f };
+
+		float TemporalAngle = Rots[TemporalFrame % 6] * (PI / 360.0f);
+
+		GTAOParam[0] = FVector4(cos(TemporalAngle), sin(TemporalAngle),
+									Offsets[(TemporalFrame / 6) % 4], Offsets[TemporalFrame % 4]);
+
+
+		FIntPoint RandomizationSize = GSystemTextures.GTAORandomization->GetDesc().Extent;
+
+		GTAOParam[1] = FVector4(((float)TemporalFrame) * (1.0f / (float)RandomizationSize.X),
+			((float)TemporalFrame) * (3.0f / (float)RandomizationSize.X),
+			(float)TemporalFrame, 0.0f);
+
+
+		SetShaderValueArray(RHICmdList, ShaderRHI, GTAOParams, GTAOParam, 2);
 	}
 
 	friend FArchive& operator<<(FArchive& Ar, FGTAOParameters& This);
@@ -1239,7 +1260,9 @@ FArchive& operator<<(FArchive& Ar, FGTAOParameters& This)
 	return Ar;
 }
 
-template<uint32 bComputeShader>
+
+
+template<uint32 bComputeShader, uint32 ShaderQuality>
 class FPostProcessGTAOHorizonSearchPSandCS : public FGlobalShader
 {
 	DECLARE_SHADER_TYPE(FPostProcessGTAOHorizonSearchPSandCS, Global);
@@ -1252,6 +1275,7 @@ class FPostProcessGTAOHorizonSearchPSandCS : public FGlobalShader
 	static void ModifyCompilationEnvironment(const FGlobalShaderPermutationParameters& Parameters, FShaderCompilerEnvironment& OutEnvironment)
 	{
 		FGlobalShader::ModifyCompilationEnvironment(Parameters, OutEnvironment);
+		OutEnvironment.SetDefine(TEXT("SHADER_QUALITY"), ShaderQuality);
 		OutEnvironment.SetDefine(TEXT("COMPUTE_SHADER"), bComputeShader);
 		if (bComputeShader)
 		{
@@ -1373,14 +1397,28 @@ public:
 		return bComputeShader ? TEXT("HorizonSearchCS") : TEXT("HorizonSearchPS");
 	}
 };
-IMPLEMENT_SHADER_TYPE2(FPostProcessGTAOHorizonSearchPSandCS<false>, SF_Pixel);
-IMPLEMENT_SHADER_TYPE2(FPostProcessGTAOHorizonSearchPSandCS<true>, SF_Compute);
 
 
-template<uint32 bComputeShader>
-class FPostProcessGTAOCommon
+// #define avoids a lot of code duplication
+#define VARIATION0(A) \
+	typedef FPostProcessGTAOHorizonSearchPSandCS<false,A> FPostProcessGTAOHorizonSearchPS##A; \
+	typedef FPostProcessGTAOHorizonSearchPSandCS<true ,A> FPostProcessGTAOHorizonSearchCS##A; \
+	IMPLEMENT_SHADER_TYPE2(FPostProcessGTAOHorizonSearchPS##A, SF_Pixel); \
+	IMPLEMENT_SHADER_TYPE2(FPostProcessGTAOHorizonSearchCS##A, SF_Compute);
+VARIATION0(0)
+VARIATION0(1)
+VARIATION0(2)
+VARIATION0(3)
+VARIATION0(4)
+#undef VARIATION0
+
+
+
+
+template<uint32 bComputeShader, uint32 ShaderQuality, uint32 UseNormalBuffer >
+class FPostProcessGTAOHorizonSearchAndInnerIntegratePSandCS : public FGlobalShader
 {
-public:
+	DECLARE_SHADER_TYPE(FPostProcessGTAOHorizonSearchAndInnerIntegratePSandCS, Global);
 
 	static bool ShouldCompilePermutation(const FGlobalShaderPermutationParameters& Parameters)
 	{
@@ -1391,30 +1429,15 @@ public:
 	{
 		FGlobalShader::ModifyCompilationEnvironment(Parameters, OutEnvironment);
 		OutEnvironment.SetDefine(TEXT("COMPUTE_SHADER"), bComputeShader);
+		OutEnvironment.SetDefine(TEXT("SHADER_QUALITY"), ShaderQuality);
+		OutEnvironment.SetDefine(TEXT("USE_NORMALBUFFER"), UseNormalBuffer);
+
 		if (bComputeShader)
 		{
 			OutEnvironment.SetDefine(TEXT("THREADGROUP_SIZEX"), 8);
 			OutEnvironment.SetDefine(TEXT("THREADGROUP_SIZEY"), 8);
 		}
 	}
-};
-
-
-template<uint32 bComputeShader>
-class FPostProcessGTAOHorizonSearchAndInnerIntegratePSandCS : public FGlobalShader, FPostProcessGTAOCommon<bComputeShader>
-{
-	DECLARE_SHADER_TYPE(FPostProcessGTAOHorizonSearchAndInnerIntegratePSandCS, Global);
-
-	static bool ShouldCompilePermutation(const FGlobalShaderPermutationParameters& Parameters)
-	{
-		return FPostProcessGTAOCommon<bComputeShader>::ShouldCompilePermutation(Parameters);
-	}
-
-	static void ModifyCompilationEnvironment(const FGlobalShaderPermutationParameters& Parameters, FShaderCompilerEnvironment& OutEnvironment)
-	{
-		FPostProcessGTAOCommon<bComputeShader>::ModifyCompilationEnvironment(Parameters, OutEnvironment);
-		
-	}
 
 	/** Default constructor. */
 	FPostProcessGTAOHorizonSearchAndInnerIntegratePSandCS() {}
@@ -1427,6 +1450,7 @@ public:
 	FShaderParameter			  HZBRemapping;
 	FShaderResourceParameter      RandomNormalTexture;
 	FShaderResourceParameter      RandomNormalTextureSampler;
+	FGTAOParameters				  GTAOParams;
 
 	/** Initialization constructor. */
 	FPostProcessGTAOHorizonSearchAndInnerIntegratePSandCS(const ShaderMetaType::CompiledShaderInitializerType& Initializer)
@@ -1438,6 +1462,7 @@ public:
 		HZBRemapping.Bind(Initializer.ParameterMap, TEXT("HZBRemapping"));
 		RandomNormalTexture.Bind(Initializer.ParameterMap, TEXT("RandomNormalTexture"));
 		RandomNormalTextureSampler.Bind(Initializer.ParameterMap, TEXT("RandomNormalTextureSampler"));
+		GTAOParams.Bind(Initializer.ParameterMap);
 
 		if (bComputeShader)
 		{
@@ -1469,6 +1494,7 @@ public:
 
 		PostprocessParameter.SetCS(ShaderRHI, Context, RHICmdList, TStaticSamplerState<SF_Point, AM_Clamp, AM_Clamp, AM_Clamp>::GetRHI());
 
+		GTAOParams.Set(RHICmdList, Context.View, ShaderRHI);
 		RHICmdList.SetUAVParameter(ShaderRHI, OutTexture.GetBaseIndex(), OutUAV);
 
 		FVector4 HZBRemappingValue = GetHZBRemapVal(Context, DestSize, InputTextureSize);
@@ -1488,6 +1514,7 @@ public:
 		FGlobalShader::SetParameters<FViewUniformShaderParameters>(Context.RHICmdList, ShaderRHI, Context.View.ViewUniformBuffer);
 		SceneTextureParameters.Set(Context.RHICmdList, ShaderRHI, Context.View.FeatureLevel, ESceneTextureSetupMode::All);
 		ScreenSpaceAOParams.Set(Context.RHICmdList, Context.View, ShaderRHI, InputTextureSize, FScreenSpaceAOParameters::ERandTexType::GTAO);
+		GTAOParams.Set(Context.RHICmdList, Context.View, ShaderRHI);
 
 		FVector4 HZBRemappingValue = GetHZBRemapVal(Context, DestSize, InputTextureSize);
 		SetShaderValue(Context.RHICmdList, ShaderRHI, HZBRemapping, HZBRemappingValue);
@@ -1509,7 +1536,7 @@ public:
 	virtual bool Serialize(FArchive& Ar) override
 	{
 		bool bShaderHasOutdatedParameters = FGlobalShader::Serialize(Ar);
-		Ar << PostprocessParameter << SceneTextureParameters << ScreenSpaceAOParams << OutTexture << HZBRemapping << RandomNormalTexture << RandomNormalTextureSampler;
+		Ar << PostprocessParameter << SceneTextureParameters << ScreenSpaceAOParams << OutTexture << GTAOParams <<HZBRemapping << RandomNormalTexture << RandomNormalTextureSampler;
 		return bShaderHasOutdatedParameters;
 	}
 
@@ -1523,8 +1550,24 @@ public:
 		return bComputeShader ? TEXT("HorizonSearchAndInnerIntegrateCS") : TEXT("HorizonSearchAndInnerIntegratePS");
 	}
 };
-IMPLEMENT_SHADER_TYPE2(FPostProcessGTAOHorizonSearchAndInnerIntegratePSandCS<false>, SF_Pixel);
-IMPLEMENT_SHADER_TYPE2(FPostProcessGTAOHorizonSearchAndInnerIntegratePSandCS<true>, SF_Compute);
+
+// #define avoids a lot of code duplication
+#define VARIATION0(A) \
+	typedef FPostProcessGTAOHorizonSearchAndInnerIntegratePSandCS<false,A,0> FPostProcessGTAOHorizonSearchAndInnerIntegratePS0##A; \
+	typedef FPostProcessGTAOHorizonSearchAndInnerIntegratePSandCS<false,A,1> FPostProcessGTAOHorizonSearchAndInnerIntegratePS1##A; \
+	typedef FPostProcessGTAOHorizonSearchAndInnerIntegratePSandCS<true ,A,0> FPostProcessGTAOHorizonSearchAndInnerIntegrateCS0##A; \
+	typedef FPostProcessGTAOHorizonSearchAndInnerIntegratePSandCS<true ,A,1> FPostProcessGTAOHorizonSearchAndInnerIntegrateCS1##A; \
+	IMPLEMENT_SHADER_TYPE2(FPostProcessGTAOHorizonSearchAndInnerIntegratePS0##A, SF_Pixel); \
+	IMPLEMENT_SHADER_TYPE2(FPostProcessGTAOHorizonSearchAndInnerIntegratePS1##A, SF_Pixel); \
+	IMPLEMENT_SHADER_TYPE2(FPostProcessGTAOHorizonSearchAndInnerIntegrateCS0##A, SF_Compute);\
+	IMPLEMENT_SHADER_TYPE2(FPostProcessGTAOHorizonSearchAndInnerIntegrateCS1##A, SF_Compute);
+VARIATION0(0)
+VARIATION0(1)
+VARIATION0(2)
+VARIATION0(3)
+VARIATION0(4)
+#undef VARIATION0
+
 
 
 template<uint32 bComputeShader>
@@ -1560,6 +1603,7 @@ public:
 	FShaderResourceParameter      RandomNormalTextureSampler;
 	FShaderResourceParameter      HorizonsTexture;
 	FShaderResourceParameter      HorizonsTextureSampler;
+	FGTAOParameters				  GTAOParams;
 
 
 	/** Initialization constructor. */
@@ -1573,6 +1617,7 @@ public:
 		RandomNormalTextureSampler.Bind(Initializer.ParameterMap, TEXT("RandomNormalTextureSampler"));
 		HorizonsTexture.Bind(Initializer.ParameterMap, TEXT("HorizonsTexture"));
 		HorizonsTextureSampler.Bind(Initializer.ParameterMap, TEXT("HorizonsTextureSampler"));
+		GTAOParams.Bind(Initializer.ParameterMap);
 
 		if (bComputeShader)
 		{
@@ -1583,7 +1628,7 @@ public:
 	virtual bool Serialize(FArchive& Ar) override
 	{
 		bool bShaderHasOutdatedParameters = FGlobalShader::Serialize(Ar);
-		Ar << PostprocessParameter << SceneTextureParameters << ScreenSpaceAOParams << OutTexture <<  RandomNormalTexture << RandomNormalTextureSampler;
+		Ar << PostprocessParameter << SceneTextureParameters << ScreenSpaceAOParams << OutTexture << GTAOParams <<  RandomNormalTexture << RandomNormalTextureSampler;
 		Ar << HorizonsTexture << HorizonsTextureSampler;
 
 		return bShaderHasOutdatedParameters;
@@ -1599,6 +1644,7 @@ public:
 		FGlobalShader::SetParameters<FViewUniformShaderParameters>(RHICmdList, ShaderRHI, Context.View.ViewUniformBuffer);
 		SceneTextureParameters.Set(RHICmdList, ShaderRHI, Context.View.FeatureLevel, ESceneTextureSetupMode::All);
 		ScreenSpaceAOParams.Set(RHICmdList, Context.View, ShaderRHI, InputTextureSize, FScreenSpaceAOParameters::ERandTexType::GTAO);
+		GTAOParams.Set(RHICmdList, Context.View, ShaderRHI);
 
 		PostprocessParameter.SetCS(ShaderRHI, Context, RHICmdList, TStaticSamplerState<SF_Point, AM_Clamp, AM_Clamp, AM_Clamp>::GetRHI());
 
@@ -1624,6 +1670,7 @@ public:
 		FGlobalShader::SetParameters<FViewUniformShaderParameters>(Context.RHICmdList, ShaderRHI, Context.View.ViewUniformBuffer);
 		SceneTextureParameters.Set(Context.RHICmdList, ShaderRHI, Context.View.FeatureLevel, ESceneTextureSetupMode::All);
 		ScreenSpaceAOParams.Set(Context.RHICmdList, Context.View, ShaderRHI, InputTextureSize, FScreenSpaceAOParameters::ERandTexType::GTAO);
+		GTAOParams.Set(Context.RHICmdList, Context.View, ShaderRHI);
 
 		const FSceneRenderTargetItem& GTAORandomization = GSystemTextures.GTAORandomization->GetRenderTargetItem();
 		SetTextureParameter(Context.RHICmdList, ShaderRHI, RandomNormalTexture, RandomNormalTextureSampler, TStaticSamplerState<SF_Point, AM_Wrap, AM_Wrap, AM_Wrap>::GetRHI(), GTAORandomization.ShaderResourceTexture);
@@ -1652,18 +1699,55 @@ public:
 	}
 };
 
-
 IMPLEMENT_SHADER_TYPE2(FPostProcessGTAOInnerIntegratePSandCS<false>, SF_Pixel);
 IMPLEMENT_SHADER_TYPE2(FPostProcessGTAOInnerIntegratePSandCS<true>, SF_Compute);
 
-
-
 FRCPassPostProcessAmbientOcclusion_HorizonSearchInnerIntegrate::FRCPassPostProcessAmbientOcclusion_HorizonSearchInnerIntegrate(const FSceneView& View, bool FinalOutput)
 	:bFinalOutput(FinalOutput)
 {
 
 }
 
+template <uint32 ShaderQuality, uint32 UseNormals>
+void FRCPassPostProcessAmbientOcclusion_HorizonSearchInnerIntegrate::DispatchCS(const FRenderingCompositePassContext& Context, FIntPoint DestSize, FIntPoint TexSize)
+{
+	TShaderMapRef<FPostProcessGTAOHorizonSearchAndInnerIntegratePSandCS<true, ShaderQuality, UseNormals>> ComputeShader(Context.GetShaderMap());
+	FSceneRenderTargets& SceneContext = FSceneRenderTargets::Get(Context.RHICmdList);
+	const FSceneRenderTargetItem& DestRenderTarget = bFinalOutput ? SceneContext.ScreenSpaceAO->GetRenderTargetItem() : PassOutputs[0].RequestSurface(Context);
+
+	Context.RHICmdList.SetComputeShader(ComputeShader->GetComputeShader());
+	ComputeShader->SetParametersCS(Context.RHICmdList, Context, DestSize, TexSize, DestRenderTarget.UAV);
+
+	uint32 GroupSizeX = FMath::DivideAndRoundUp(DestSize.X, 8);
+	uint32 GroupSizeY = FMath::DivideAndRoundUp(DestSize.Y, 8);
+	DispatchComputeShader(Context.RHICmdList, *ComputeShader, GroupSizeX, GroupSizeY, 1);
+
+	ComputeShader->UnsetParameters(Context.RHICmdList);
+}
+
+template <uint32 ShaderQuality, uint32 UseNormals>
+FShader* FRCPassPostProcessAmbientOcclusion_HorizonSearchInnerIntegrate::SetShaderPS(const FRenderingCompositePassContext& Context, FGraphicsPipelineStateInitializer& GraphicsPSOInit, FIntPoint DestSize)
+{
+	TShaderMapRef<FPostProcessVS> VertexShader(Context.GetShaderMap());
+	TShaderMapRef<FPostProcessGTAOHorizonSearchAndInnerIntegratePSandCS<false, ShaderQuality, UseNormals> > PixelShader(Context.GetShaderMap());
+
+	GraphicsPSOInit.BoundShaderState.VertexDeclarationRHI = GFilterVertexDeclaration.VertexDeclarationRHI;
+	GraphicsPSOInit.BoundShaderState.VertexShaderRHI = GETSAFERHISHADER_VERTEX(*VertexShader);
+	GraphicsPSOInit.BoundShaderState.PixelShaderRHI = GETSAFERHISHADER_PIXEL(*PixelShader);
+	SetGraphicsPipelineState(Context.RHICmdList, GraphicsPSOInit);
+
+	const FPooledRenderTargetDesc* InputDesc0 = GetInputDesc(ePId_Input0);
+	FSceneRenderTargets& SceneContext = FSceneRenderTargets::Get(Context.RHICmdList);
+	FIntPoint TexSize = InputDesc0 ? InputDesc0->Extent : SceneContext.GetBufferSizeXY();
+
+	VertexShader->SetParameters(Context);
+	PixelShader->SetParametersPS(Context, DestSize, TexSize);
+
+	return *VertexShader;
+}
+
+
+
 void FRCPassPostProcessAmbientOcclusion_HorizonSearchInnerIntegrate::Process(FRenderingCompositePassContext& Context)
 {
 	SCOPED_GPU_STAT(Context.RHICmdList, GTAO_HorizonSearchInnerIntegrate);
@@ -1680,6 +1764,8 @@ void FRCPassPostProcessAmbientOcclusion_HorizonSearchInnerIntegrate::Process(FRe
 
 	static const auto CVar = IConsoleManager::Get().FindTConsoleVariableDataInt(TEXT("r.AmbientOcclusion.Compute"));
 
+	const int32 ShaderQuality = FSSAOHelper::GetAmbientOcclusionShaderLevel(Context.View);
+
 	if (CVar->GetValueOnRenderThread() >= 1)
 	{
 		Context.RHICmdList.TransitionResource(EResourceTransitionAccess::ERWBarrier, EResourceTransitionPipeline::EGfxToCompute, DestRenderTarget.UAV, nullptr);
@@ -1687,15 +1773,25 @@ void FRCPassPostProcessAmbientOcclusion_HorizonSearchInnerIntegrate::Process(FRe
 		// Compute  Version
 		UnbindRenderTargets(Context.RHICmdList);
 
-		TShaderMapRef<FPostProcessGTAOHorizonSearchAndInnerIntegratePSandCS<true>> ComputeShader(Context.GetShaderMap());
-		Context.RHICmdList.SetComputeShader(ComputeShader->GetComputeShader());
-		ComputeShader->SetParametersCS(Context.RHICmdList, Context, DestSize, TexSize, DestRenderTarget.UAV);
-
-		uint32 GroupSizeX = FMath::DivideAndRoundUp(DestSize.X, 8);
-		uint32 GroupSizeY = FMath::DivideAndRoundUp(DestSize.Y, 8);
-		DispatchComputeShader(Context.RHICmdList, *ComputeShader, GroupSizeX, GroupSizeY, 1);
-
-		ComputeShader->UnsetParameters(Context.RHICmdList);
+		// 0..4, 0:low 4:high
+		switch(ShaderQuality)
+		{
+		case 0:
+			DispatchCS<0,0>(Context, DestSize, TexSize);
+			break;
+		case 1:
+			DispatchCS<1, 0>(Context, DestSize, TexSize);
+			break;
+		case 2:
+			DispatchCS<2, 0>(Context, DestSize, TexSize);
+			break;
+		case 3:
+			DispatchCS<3, 0>(Context, DestSize, TexSize);
+			break;
+		case 4:
+			DispatchCS<4, 0>(Context, DestSize, TexSize);
+			break;
+		}
 	}
 	else
 	{
@@ -1706,21 +1802,36 @@ void FRCPassPostProcessAmbientOcclusion_HorizonSearchInnerIntegrate::Process(FRe
 			// PS Version
 			Context.SetViewportAndCallRHI(ViewRect);
 
-			TShaderMapRef<FPostProcessGTAOHorizonSearchAndInnerIntegratePSandCS<false> > PixelShader(Context.GetShaderMap());
-			TShaderMapRef<FPostProcessVS> VertexShader(Context.GetShaderMap());
-
 			FGraphicsPipelineStateInitializer GraphicsPSOInit;
 			Context.RHICmdList.ApplyCachedRenderTargets(GraphicsPSOInit);
-			GraphicsPSOInit.BlendState = TStaticBlendState<>::GetRHI();
-			GraphicsPSOInit.RasterizerState = TStaticRasterizerState<>::GetRHI();
-			GraphicsPSOInit.DepthStencilState = TStaticDepthStencilState<false, CF_Always>::GetRHI();
-			GraphicsPSOInit.PrimitiveType = PT_TriangleList;
-			GraphicsPSOInit.BoundShaderState.VertexDeclarationRHI = GFilterVertexDeclaration.VertexDeclarationRHI;
-			GraphicsPSOInit.BoundShaderState.VertexShaderRHI = GETSAFERHISHADER_VERTEX(*VertexShader);
-			GraphicsPSOInit.BoundShaderState.PixelShaderRHI = GETSAFERHISHADER_PIXEL(*PixelShader);
-			SetGraphicsPipelineState(Context.RHICmdList, GraphicsPSOInit);
 
-			PixelShader->SetParametersPS(Context, DestSize, TexSize);
+			// set the state
+			GraphicsPSOInit.BlendState			= TStaticBlendState<>::GetRHI();
+			GraphicsPSOInit.RasterizerState		= TStaticRasterizerState<>::GetRHI();
+			GraphicsPSOInit.DepthStencilState	= TStaticDepthStencilState<false, CF_Always>::GetRHI();
+			GraphicsPSOInit.PrimitiveType		= PT_TriangleList;
+			GraphicsPSOInit.bDepthBounds		= false;
+
+			FShader* VertexShader = 0;
+
+			switch (ShaderQuality)
+			{
+			case 0:
+				VertexShader = SetShaderPS<0, 0>(Context, GraphicsPSOInit, DestSize);
+				break;
+			case 1:
+				VertexShader = SetShaderPS<1, 0>(Context, GraphicsPSOInit, DestSize);
+				break;
+			case 2:
+				VertexShader = SetShaderPS<2, 0>(Context, GraphicsPSOInit, DestSize);
+				break;
+			case 3:
+				VertexShader = SetShaderPS<3, 0>(Context, GraphicsPSOInit, DestSize);
+				break;
+			case 4:
+				VertexShader = SetShaderPS<4, 0>(Context, GraphicsPSOInit, DestSize);
+				break;
+			}
 
 			DrawRectangle(
 				Context.RHICmdList,
@@ -1730,7 +1841,7 @@ void FRCPassPostProcessAmbientOcclusion_HorizonSearchInnerIntegrate::Process(FRe
 				ViewRect.Width(), ViewRect.Height(),
 				ViewRect.Size(),
 				TexSize,
-				*VertexShader,
+				VertexShader,
 				EDRF_UseTriangleOptimization);
 		}
 		Context.RHICmdList.EndRenderPass();
@@ -1849,7 +1960,7 @@ FPooledRenderTargetDesc FRCPassPostProcessAmbientOcclusion_InnerIntegrate::Compu
 
 	Ret = GetInput(ePId_Input0)->GetOutput()->RenderTargetDesc;
 	Ret.Reset();
-	Ret.Format = PF_B8G8R8A8;
+	Ret.Format = PF_G8;
 
 	Ret.ClearValue = FClearValueBinding::None;
 	Ret.TargetableFlags &= ~TexCreate_DepthStencilTargetable;
@@ -1868,9 +1979,52 @@ FRCPassPostProcessAmbientOcclusion_HorizonSearch::FRCPassPostProcessAmbientOcclu
 
 
 
+template <uint32 ShaderQuality>
+void FRCPassPostProcessAmbientOcclusion_HorizonSearch::DispatchCS(const FRenderingCompositePassContext& Context, FIntPoint DestSize, FIntPoint TexSize)
+{
+	FSceneRenderTargets& SceneContext = FSceneRenderTargets::Get(Context.RHICmdList);
+	const FSceneRenderTargetItem& DestRenderTarget = (AOType == ESSAOType::EAsyncCS) ? SceneContext.ScreenSpaceGTAOHorizons->GetRenderTargetItem() : PassOutputs[0].RequestSurface(Context);
+
+	FRHIAsyncComputeCommandListImmediate& RHICmdListComputeImmediate = FRHICommandListExecutor::GetImmediateAsyncComputeCommandList();
+
+	TShaderMapRef<FPostProcessGTAOHorizonSearchPSandCS<true, ShaderQuality>> ComputeShader(Context.GetShaderMap());
+	RHICmdListComputeImmediate.SetComputeShader(ComputeShader->GetComputeShader());
+	ComputeShader->SetParametersCS(RHICmdListComputeImmediate, Context, DestSize, TexSize, DestRenderTarget.UAV);
+
+	uint32 GroupSizeX = FMath::DivideAndRoundUp(DestSize.X, 8);
+	uint32 GroupSizeY = FMath::DivideAndRoundUp(DestSize.Y, 8);
+	DispatchComputeShader(RHICmdListComputeImmediate, *ComputeShader, GroupSizeX, GroupSizeY, 1);
+
+	ComputeShader->UnsetParameters(RHICmdListComputeImmediate);
+}
+
+
+template <uint32 ShaderQuality>
+FShader* FRCPassPostProcessAmbientOcclusion_HorizonSearch::SetShaderPS(const FRenderingCompositePassContext& Context, FGraphicsPipelineStateInitializer& GraphicsPSOInit, FIntPoint DestSize)
+{
+	TShaderMapRef<FPostProcessVS> VertexShader(Context.GetShaderMap());
+	TShaderMapRef<FPostProcessGTAOHorizonSearchPSandCS<false, ShaderQuality> > PixelShader(Context.GetShaderMap());
+
+	GraphicsPSOInit.BoundShaderState.VertexDeclarationRHI = GFilterVertexDeclaration.VertexDeclarationRHI;
+	GraphicsPSOInit.BoundShaderState.VertexShaderRHI = GETSAFERHISHADER_VERTEX(*VertexShader);
+	GraphicsPSOInit.BoundShaderState.PixelShaderRHI = GETSAFERHISHADER_PIXEL(*PixelShader);
+	SetGraphicsPipelineState(Context.RHICmdList, GraphicsPSOInit);
+
+	const FPooledRenderTargetDesc* InputDesc0 = GetInputDesc(ePId_Input0);
+	FSceneRenderTargets& SceneContext = FSceneRenderTargets::Get(Context.RHICmdList);
+	FIntPoint TexSize = InputDesc0 ? InputDesc0->Extent : SceneContext.GetBufferSizeXY();
+
+	VertexShader->SetParameters(Context);
+	PixelShader->SetParametersPS(Context, DestSize, TexSize);
+
+	return *VertexShader;
+}
+
+
 
 void FRCPassPostProcessAmbientOcclusion_HorizonSearch::Process(FRenderingCompositePassContext& Context)
 {
+
 	SCOPED_GPU_STAT(Context.RHICmdList, GTAO_HorizonSearch);
 	const FViewInfo& View = Context.View;
 	FIntRect ViewRect = View.ViewRect;
@@ -1885,6 +2039,7 @@ void FRCPassPostProcessAmbientOcclusion_HorizonSearch::Process(FRenderingComposi
 	const FSceneRenderTargetItem& DestRenderTarget = (AOType == ESSAOType::EAsyncCS) ? SceneContext.ScreenSpaceGTAOHorizons->GetRenderTargetItem()  : PassOutputs[0].RequestSurface(Context);
 
 	static const auto CVar = IConsoleManager::Get().FindTConsoleVariableDataInt(TEXT("r.AmbientOcclusion.Compute"));
+	const int32 ShaderQuality = FSSAOHelper::GetAmbientOcclusionShaderLevel(Context.View);
 
 	if (AOType == ESSAOType::EAsyncCS)
 	{
@@ -1901,39 +2056,64 @@ void FRCPassPostProcessAmbientOcclusion_HorizonSearch::Process(FRenderingComposi
 		// Compute  Version
 		UnbindRenderTargets(Context.RHICmdList);
 
-		TShaderMapRef<FPostProcessGTAOHorizonSearchPSandCS<true>> ComputeShader(Context.GetShaderMap());
-		RHICmdListComputeImmediate.SetComputeShader(ComputeShader->GetComputeShader());
-		ComputeShader->SetParametersCS(RHICmdListComputeImmediate, Context, DestSize, TexSize, DestRenderTarget.UAV);
-
-		uint32 GroupSizeX = FMath::DivideAndRoundUp(DestSize.X, 8);
-		uint32 GroupSizeY = FMath::DivideAndRoundUp(DestSize.Y, 8);
-		DispatchComputeShader(RHICmdListComputeImmediate, *ComputeShader, GroupSizeX, GroupSizeY, 1);
-
-		ComputeShader->UnsetParameters(RHICmdListComputeImmediate);
+		// 0..4, 0:low 4:high
+		switch (ShaderQuality)
+		{
+		case 0:
+			DispatchCS<0>(Context, DestSize, TexSize);
+			break;
+		case 1:
+			DispatchCS<1>(Context, DestSize, TexSize);
+			break;
+		case 2:
+			DispatchCS<2>(Context, DestSize, TexSize);
+			break;
+		case 3:
+			DispatchCS<3>(Context, DestSize, TexSize);
+			break;
+		case 4:
+			DispatchCS<4>(Context, DestSize, TexSize);
+			break;
+		}
 	}
 	else
 	{
-		// Pixel  Version
 		FRHIRenderPassInfo RPInfo(DestRenderTarget.TargetableTexture, ERenderTargetActions::Load_Store);
-		Context.RHICmdList.BeginRenderPass(RPInfo, TEXT("AmbientOcclusionSetup"));
+		Context.RHICmdList.BeginRenderPass(RPInfo, TEXT("GTAOHorizonSearch"));
 		{
+			// PS Version
 			Context.SetViewportAndCallRHI(ViewRect);
 
-			TShaderMapRef<FPostProcessGTAOHorizonSearchPSandCS<false> > PixelShader(Context.GetShaderMap());
-			TShaderMapRef<FPostProcessVS> VertexShader(Context.GetShaderMap());
-
 			FGraphicsPipelineStateInitializer GraphicsPSOInit;
 			Context.RHICmdList.ApplyCachedRenderTargets(GraphicsPSOInit);
+
+			// set the state
 			GraphicsPSOInit.BlendState = TStaticBlendState<>::GetRHI();
 			GraphicsPSOInit.RasterizerState = TStaticRasterizerState<>::GetRHI();
 			GraphicsPSOInit.DepthStencilState = TStaticDepthStencilState<false, CF_Always>::GetRHI();
 			GraphicsPSOInit.PrimitiveType = PT_TriangleList;
-			GraphicsPSOInit.BoundShaderState.VertexDeclarationRHI = GFilterVertexDeclaration.VertexDeclarationRHI;
-			GraphicsPSOInit.BoundShaderState.VertexShaderRHI = GETSAFERHISHADER_VERTEX(*VertexShader);
-			GraphicsPSOInit.BoundShaderState.PixelShaderRHI = GETSAFERHISHADER_PIXEL(*PixelShader);
-			SetGraphicsPipelineState(Context.RHICmdList, GraphicsPSOInit);
+			GraphicsPSOInit.bDepthBounds = false;
 
-			PixelShader->SetParametersPS(Context, DestSize, TexSize);
+			FShader* VertexShader = 0;
+
+			switch (ShaderQuality)
+			{
+			case 0:
+				VertexShader = SetShaderPS<0>(Context, GraphicsPSOInit, DestSize);
+				break;
+			case 1:
+				VertexShader = SetShaderPS<1>(Context, GraphicsPSOInit, DestSize);
+				break;
+			case 2:
+				VertexShader = SetShaderPS<2>(Context, GraphicsPSOInit, DestSize);
+				break;
+			case 3:
+				VertexShader = SetShaderPS<3>(Context, GraphicsPSOInit, DestSize);
+				break;
+			case 4:
+				VertexShader = SetShaderPS<4>(Context, GraphicsPSOInit, DestSize);
+				break;
+			}
 
 			DrawRectangle(
 				Context.RHICmdList,
@@ -1943,7 +2123,7 @@ void FRCPassPostProcessAmbientOcclusion_HorizonSearch::Process(FRenderingComposi
 				ViewRect.Width(), ViewRect.Height(),
 				ViewRect.Size(),
 				TexSize,
-				*VertexShader,
+				VertexShader,
 				EDRF_UseTriangleOptimization);
 		}
 		Context.RHICmdList.EndRenderPass();
diff --git a/Engine/Source/Runtime/Renderer/Private/CompositionLighting/PostProcessAmbientOcclusion.h b/Engine/Source/Runtime/Renderer/Private/CompositionLighting/PostProcessAmbientOcclusion.h
--- a/Engine/Source/Runtime/Renderer/Private/CompositionLighting/PostProcessAmbientOcclusion.h
+++ b/Engine/Source/Runtime/Renderer/Private/CompositionLighting/PostProcessAmbientOcclusion.h
@@ -155,6 +155,12 @@ public:
 	virtual void Release() override { delete this; }
 	virtual FPooledRenderTargetDesc ComputeOutputDesc(EPassOutputId InPassOutputId) const override;
 
+	template <uint32 ShaderQuality>
+	void DispatchCS(const FRenderingCompositePassContext& Context, FIntPoint DestSize, FIntPoint TexSize);
+
+	template <uint32 ShaderQuality>
+	FShader* SetShaderPS(const FRenderingCompositePassContext& Context, FGraphicsPipelineStateInitializer& GraphicsPSOInit, FIntPoint DestSize);
+
 private:
 	const ESSAOType AOType;
 };
@@ -187,6 +193,12 @@ public:
 	virtual void Release() override { delete this; }
 	virtual FPooledRenderTargetDesc ComputeOutputDesc(EPassOutputId InPassOutputId) const override;
 
+	template <uint32 ShaderQuality, uint32 UseNormals>
+	void DispatchCS(const FRenderingCompositePassContext& Context, FIntPoint DestSize, FIntPoint TexSize);
+
+	template <uint32 ShaderQuality, uint32 UseNormals>
+	FShader* SetShaderPS(const FRenderingCompositePassContext& Context, FGraphicsPipelineStateInitializer& GraphicsPSOInit, FIntPoint DestSize);
+
 private:
 	const bool bFinalOutput;
 
diff --git a/Engine/Source/Runtime/Renderer/Private/SystemTextures.cpp b/Engine/Source/Runtime/Renderer/Private/SystemTextures.cpp
--- a/Engine/Source/Runtime/Renderer/Private/SystemTextures.cpp
+++ b/Engine/Source/Runtime/Renderer/Private/SystemTextures.cpp
@@ -395,13 +395,12 @@ void FSystemTextures::InitializeFeatureLevelDependentTextures(FRHICommandListImm
 						float Angle  = (PI/16.0f)  * ((((x + y) & 0x3) << 2) + (x & 0x3));
 						float Step	 = 0.5f + ((1.0f/8.0f) * ((y - x) & 0x3));
 
-						float ScaleCos = FMath::Cos(Angle) * Step;
-						float ScaleSin = FMath::Sin(Angle) * Step;
+						float ScaleCos = FMath::Cos(Angle) ;
+						float ScaleSin = FMath::Sin(Angle) ;
 			
 						Dest[0] = (uint8_t)(ScaleCos*127.5f + 127.5f);
 						Dest[1] = (uint8_t)(ScaleSin*127.5f + 127.5f);
 					}
-
 				}
 			}
 			RHICmdList.UnlockTexture2D((FTexture2DRHIRef&)GTAORandomization->GetRenderTargetItem().ShaderResourceTexture, 0, false);
