Fix a crash on accessing invalid vertex buffer

diff --git a/Engine/Source/Runtime/Renderer/Private/VT/RuntimeVirtualTextureRender.cpp b/Engine/Source/Runtime/Renderer/Private/VT/RuntimeVirtualTextureRender.cpp
--- a/Engine/Source/Runtime/Renderer/Private/VT/RuntimeVirtualTextureRender.cpp
+++ b/Engine/Source/Runtime/Renderer/Private/VT/RuntimeVirtualTextureRender.cpp
@@ -351,12 +351,15 @@ namespace RuntimeVirtualTexture
 				}
 			}
 
+			FPrimitiveSceneInfo* PrimitiveSceneInfo = Scene->Primitives[PrimitiveIndex];
+
 			// Calculate Lod for current mip
 			const float AreaRatio = 2.f * SphereBounds.W * RcpWorldSize;
-			const int32 LodIndex = FMath::Clamp<int32>((int32)LodInfo.LodBias - FMath::FloorToInt(FMath::Log2(AreaRatio)), LodInfo.MinLod, LodInfo.MaxLod);
+			const int32 CurFirstLODIdx = PrimitiveSceneInfo->Proxy->GetCurrentFirstLODIdx_RenderThread();
+			const int32 MinLODIdx = FMath::Max((int32)LodInfo.MinLod, CurFirstLODIdx);
+			const int32 LodIndex = FMath::Clamp<int32>((int32)LodInfo.LodBias - FMath::FloorToInt(FMath::Log2(AreaRatio)), MinLODIdx, LodInfo.MaxLod);
 
 			// Process meshes
-			FPrimitiveSceneInfo* PrimitiveSceneInfo = Scene->Primitives[PrimitiveIndex];
 			for (int32 MeshIndex = 0; MeshIndex < PrimitiveSceneInfo->StaticMeshes.Num(); ++MeshIndex)
 			{
 				FStaticMeshBatchRelevance const& StaticMeshRelevance = PrimitiveSceneInfo->StaticMeshRelevances[MeshIndex];
