New SingleLayerWater rendering tech.
Explanations and todos there https://docs.google.com/document/d/12wltU2nsYNnabLw1H3MdPPd2Uj_myoiQSNysXp39vbc

diff --git a/Engine/Plugins/Runtime/Oculus/OculusVR/Source/OculusMR/Private/OculusMR_BoundaryMeshComponent.cpp b/Engine/Plugins/Runtime/Oculus/OculusVR/Source/OculusMR/Private/OculusMR_BoundaryMeshComponent.cpp
--- a/Engine/Plugins/Runtime/Oculus/OculusVR/Source/OculusMR/Private/OculusMR_BoundaryMeshComponent.cpp
+++ b/Engine/Plugins/Runtime/Oculus/OculusVR/Source/OculusMR/Private/OculusMR_BoundaryMeshComponent.cpp
@@ -190,6 +190,7 @@ public:
 		Result.bOpaqueRelevance = true;
 		Result.bUsesSceneDepth = false;
 		Result.bUsesSkyMaterial = false;
+		Result.bUsesSingleLayerWaterMaterial = false;
 		Result.bRenderCustomDepth = false;
 		return Result;
 	}
diff --git a/Engine/Shaders/Private/BasePassPixelShader.usf b/Engine/Shaders/Private/BasePassPixelShader.usf
--- a/Engine/Shaders/Private/BasePassPixelShader.usf
+++ b/Engine/Shaders/Private/BasePassPixelShader.usf
@@ -14,6 +14,12 @@
 	#define EyeAdaptationStruct OpaqueBasePass
 #endif
 
+#if MATERIAL_IS_SINGLE_LAYER_WATER
+#define PreIntegratedGF			OpaqueBasePass.PreIntegratedGFTexture
+#define PreIntegratedGFSampler	OpaqueBasePass.PreIntegratedGFSampler
+#include "ParticipatingMediaCommon.ush" 
+#endif
+
 #include "SHCommon.ush"
 #include "/Engine/Generated/Material.ush"
 #include "BasePassCommon.ush"
@@ -63,7 +69,7 @@ float NormalCurvatureToRoughness(float3 WorldNormal)
 #endif
 
 #include "ShadingModelsMaterial.ush"
-#if MATERIAL_SHADINGMODEL_HAIR || SIMPLE_FORWARD_DIRECTIONAL_LIGHT
+#if MATERIAL_SHADINGMODEL_HAIR || SIMPLE_FORWARD_DIRECTIONAL_LIGHT || MATERIAL_IS_SINGLE_LAYER_WATER
 #include "ShadingModels.ush"
 #endif
 
@@ -75,8 +81,8 @@ float NormalCurvatureToRoughness(float3 WorldNormal)
 #define EDITOR_ALPHA2COVERAGE (USE_EDITOR_COMPOSITING && SUPPORTS_PIXEL_COVERAGE)
 #define POST_PROCESS_SUBSURFACE ((MATERIAL_SHADINGMODEL_SUBSURFACE_PROFILE || MATERIAL_SHADINGMODEL_EYE) && USES_GBUFFER)
 
-#if TRANSLUCENCY_LIGHTING_SURFACE_LIGHTINGVOLUME || TRANSLUCENCY_LIGHTING_SURFACE_FORWARDSHADING || FORWARD_SHADING
-#include "ForwardLightingCommon.ush" 
+#if TRANSLUCENCY_LIGHTING_SURFACE_LIGHTINGVOLUME || TRANSLUCENCY_LIGHTING_SURFACE_FORWARDSHADING || FORWARD_SHADING || MATERIAL_IS_SINGLE_LAYER_WATER
+#include "ForwardLightingCommon.ush"
 #endif
 
 #if !FORWARD_SHADING
@@ -519,7 +525,7 @@ void GetPrecomputedIndirectLightingAndSkyLight(
 	#endif
 }
 
-#if SIMPLE_FORWARD_DIRECTIONAL_LIGHT
+#if SIMPLE_FORWARD_DIRECTIONAL_LIGHT || MATERIAL_IS_SINGLE_LAYER_WATER
 
 float3 GetSimpleForwardLightingDirectionalLight(FGBufferData GBuffer, float3 DiffuseColor, float3 SpecularColor, float Roughness, float3 WorldNormal, float3 CameraVector)
 {
@@ -883,6 +889,8 @@ void FPixelShaderInOut_MainPS(
 #endif
 
 	half Opacity = GetMaterialOpacity(PixelMaterialInputs);
+	const float BaseMaterialCoverageOverWater = Opacity;
+	const float WaterVisibility = 1.0 - BaseMaterialCoverageOverWater;
 
 	float3 VolumetricLightmapBrickTextureUVs;
 
@@ -1059,11 +1067,20 @@ ShadingModel);
 
 		DiffuseColor += (DiffuseIndirectLighting * DiffuseColorForIndirect + SubsurfaceIndirectLighting * SubsurfaceColor) * AOMultiBounce( GBuffer.BaseColor, DiffOcclusion );
 
+#if MATERIAL_IS_SINGLE_LAYER_WATER
+		// Fade out diffuse as this will be hanlded by the single scattering lighting. when over the water surface.
+		// We do this after the just above GetPrecomputedIndirectLightingAndSkyLight to keep ambient lighting avialable.
+		// We also keep the SpecularColor for sun/water interactions
+		const float3 WaterIndirectDiffuse = DiffuseIndirectLighting;	// Raw unaffected illuminance
+		GBuffer.DiffuseColor *= BaseMaterialCoverageOverWater;
+		DiffuseColor *= BaseMaterialCoverageOverWater;
+#endif
+
 		#if TRANSLUCENCY_PERVERTEX_FORWARD_SHADING
 
 			Color += BasePassInterpolants.VertexDiffuseLighting * GBuffer.DiffuseColor;
 
-		#elif FORWARD_SHADING || TRANSLUCENCY_LIGHTING_SURFACE_FORWARDSHADING || TRANSLUCENCY_LIGHTING_SURFACE_LIGHTINGVOLUME
+		#elif FORWARD_SHADING || TRANSLUCENCY_LIGHTING_SURFACE_FORWARDSHADING || TRANSLUCENCY_LIGHTING_SURFACE_LIGHTINGVOLUME || MATERIAL_IS_SINGLE_LAYER_WATER
 			uint GridIndex = 0;
 
 			#if INSTANCED_STEREO
@@ -1075,12 +1092,14 @@ ShadingModel);
 			#if FEATURE_LEVEL >= FEATURE_LEVEL_SM5
 				GridIndex = ComputeLightGridCellIndex((uint2)(MaterialParameters.SvPosition.xy - ResolvedView.ViewRectMin.xy), MaterialParameters.SvPosition.w, EyeIndex);
 
-				#if FORWARD_SHADING || TRANSLUCENCY_LIGHTING_SURFACE_FORWARDSHADING
+				#if FORWARD_SHADING || TRANSLUCENCY_LIGHTING_SURFACE_FORWARDSHADING || MATERIAL_IS_SINGLE_LAYER_WATER
 					const float Dither = InterleavedGradientNoise(MaterialParameters.SvPosition.xy, View.StateFrameIndexMod8);
-					Color += GetForwardDirectLighting(GridIndex, MaterialParameters.AbsoluteWorldPosition, MaterialParameters.CameraVector, GBuffer, NearestResolvedDepthScreenUV, MaterialParameters.PrimitiveId, EyeIndex, Dither);
+					float3 ForwardDirectLighting = GetForwardDirectLighting(GridIndex, MaterialParameters.AbsoluteWorldPosition, MaterialParameters.CameraVector, GBuffer, NearestResolvedDepthScreenUV, MaterialParameters.PrimitiveId, EyeIndex, Dither);
+					Color += ForwardDirectLighting;
 				#endif
 			#endif
-			#if !(MATERIAL_SINGLE_SHADINGMODEL && MATERIAL_SHADINGMODEL_HAIR)
+				// No IBL for water in deferred: that is skipped because it is done in the water composite pass. It should however be applied when using forward shading in order to get reflection without the water composite pass.
+			#if !(MATERIAL_SINGLE_SHADINGMODEL && MATERIAL_SHADINGMODEL_HAIR) && (!MATERIAL_IS_SINGLE_LAYER_WATER || FORWARD_SHADING) 
 				if (GBuffer.ShadingModelID != SHADINGMODELID_HAIR)
 				{
 					int SingleCaptureIndex = GetPrimitiveData(MaterialParameters.PrimitiveId).SingleCaptureIndex;
@@ -1191,6 +1210,143 @@ ShadingModel);
 #endif
 	Color += Emissive;
 
+
+
+#if MATERIAL_IS_SINGLE_LAYER_WATER
+	// Fade out the material contribution over to water contribution according to alpha.
+	const bool CameraIsUnderWater = false;
+	
+	float2 WaterNDC = MaterialParameters.ScreenPosition.xy / MaterialParameters.ScreenPosition.w;
+	float2 ViewportUV = WaterNDC * ResolvedView.ScreenPositionScaleBias.xy + ResolvedView.ScreenPositionScaleBias.wz;
+	float SceneCustomDepth = OpaqueBasePass.SceneCustomDepthTexture.SampleLevel(OpaqueBasePass.SceneCustomDepthSampler, ViewportUV, 0).r;
+	if (SceneCustomDepth > MaterialParameters.SvPosition.z)
+	{
+		clip(-1); // Custom depth Z is used to clip water to avoid it being visible inside boats for instance.
+	}
+
+	float WaterDepth = ConvertFromDeviceZ(MaterialParameters.SvPosition.z);
+	if (WaterVisibility > 0.0f) // DeltaDepth > 0.0f should not be needed as we ARE on a water surface TODO simplify that
+	{
+
+		float PixelSceneDepth = ConvertFromDeviceZ(OpaqueBasePass.SceneDepthWithoutSingleLayerWaterTexture.SampleLevel(OpaqueBasePass.SceneDepthWithoutSingleLayerWaterSampler, ViewportUV, 0).r);
+		float DeltaDepth = CameraIsUnderWater ? WaterDepth : PixelSceneDepth - WaterDepth;
+
+#if 1
+		//////////////////////////////////////////////////////////////////////////////////////////////////// BEGIN Refration TODO make common with DistoAccumulatePS
+		float2 Refraction = GetMaterialRefraction(PixelMaterialInputs);// X = IOR, Y = Refraction Bias
+		float3 ViewNormal = normalize(TransformWorldVectorToView(MaterialParameters.WorldNormal));
+#if REFRACTION_USE_PIXEL_NORMAL_OFFSET
+		half3 ViewVertexNormal = TransformTangentVectorToView(MaterialParameters, float3(0, 0, 1));
+		// Treat Refraction of 1 as no refraction, to be consistent with IOR mode
+		float2 ViewportUVDistortion = (ViewVertexNormal.xy - ViewNormal.xy) * (Refraction.x - 1);
+#else
+		// we assume the camera is in air
+		float AirIOR = 1.0f;
+		float2 ViewportUVDistortion = ViewNormal.xy * (Refraction.x - AirIOR);
+#endif
+		float2 BufferUVDistortion = ViewportUVDistortion * ResolvedView.ViewSizeAndInvSize.xy * ResolvedView.BufferSizeAndInvSize.zw;
+		// When ISR is enabled we store two FOVs in the DistortionParams and compute the aspect ratio here
+		const float4 DistortionParams = 1.0f;
+		const float OffsetFudgeFactor = 1.0f;
+		float InvTanHalfFov = DistortionParams.x;
+		float Ratio = ResolvedView.ViewSizeAndInvSize.y * ResolvedView.ViewSizeAndInvSize.w;
+		// Fix for Fov and aspect.
+		float2 FovFix = float2(InvTanHalfFov, Ratio*InvTanHalfFov);
+		BufferUVDistortion *= 0.2 * DistortionParams.zw * float2(OffsetFudgeFactor, -OffsetFudgeFactor) * FovFix;
+		////////////////////////////////////////////////////////////////////////////////////////////////////  END Refration
+
+		// TODO make distortion fading common with DistoAccumulatePS (right now 30cm to deltadepth fade out and 10m to distance fade out)
+		BufferUVDistortion *= saturate(DeltaDepth * 1.0 / 30.0f) * (1 - saturate(WaterDepth * 1.0 / 2000.0f)); 
+
+		float SceneDeviceZ = OpaqueBasePass.SceneDepthWithoutSingleLayerWaterTexture.SampleLevel(OpaqueBasePass.SceneDepthWithoutSingleLayerWaterSampler, saturate(ViewportUV + BufferUVDistortion), 0).r;
+		float SceneDepth = ConvertFromDeviceZ(SceneDeviceZ);
+		if (SceneDepth < WaterDepth)
+		{
+			SceneDepth = PixelSceneDepth;
+		}
+		else
+		{
+			ViewportUV = saturate(ViewportUV + BufferUVDistortion);
+		}
+		DeltaDepth = CameraIsUnderWater ? WaterDepth : SceneDepth - WaterDepth;
+#endif
+
+		const float3 ScatteringCoeff = max(0.0f, GetSingleLayerWaterMaterialOutput0(MaterialParameters));
+		const float3 AbsorptionCoeff = max(0.0f, GetSingleLayerWaterMaterialOutput1(MaterialParameters));
+		const float PhaseG = clamp(GetSingleLayerWaterMaterialOutput2(MaterialParameters), -1.0f, 1.0f);
+		const float PhaseValue = SchlickPhase(PhaseG, dot(ResolvedView.DirectionalLightDirection.xyz, MaterialParameters.CameraVector));
+		const float3 ExtinctionCoeff = ScatteringCoeff + AbsorptionCoeff;
+
+		// We also apply transmittance from light to under water surface. However, the scene has been lit by many sources already.
+		// So the transmittance toabove surface is simply approximated using the travel distance from the scene pixel to the water top, assuming a flat water surface.
+		// We also use the pixel SvPosition instead of the scene one pre refraction/distortion to avoid those extra ALUs.
+		// We cannot combine this transmittance with the transmittance from view because this would change the behavior of the analytical integration of light scattering integration.
+		const float3 SceneWorldPos = SvPositionToWorld(float4(MaterialParameters.SvPosition.xy, SceneDeviceZ, 1.0));
+		const float DistanceFromScenePixelToWaterTop = max(0.0, MaterialParameters.AbsoluteWorldPosition.z - SceneWorldPos.z);
+		const float3 SceneColor = exp(-DistanceFromScenePixelToWaterTop * ExtinctionCoeff) * (USE_PREEXPOSURE ? ResolvedView.OneOverPreExposure : 1.0f)
+			* OpaqueBasePass.SceneLuminanceWithoutSingleLayerWaterTexture.SampleLevel(OpaqueBasePass.SceneLuminanceWithoutSingleLayerWaterSampler, ViewportUV, 0).rgb;
+
+		float3 SunEi = ResolvedView.DirectionalLightColor.rgb * PI;	// times PI because it is divided by PI on CPU (=luminance) and we want illuminance here.  
+		float3 SunScattLi = PhaseValue * SunEi;
+		float3 AmbScattLi = PhaseValue * WaterIndirectDiffuse * PI;	// WaterIndirectDiffuse is luminance? So we need to multiply by PI to get illuminance. That we can then get through the phase function.
+#if 1
+		float DirectionalLightShadow = ComputeDirectionalLightDynamicShadowing(MaterialParameters.AbsoluteWorldPosition, GBuffer.Depth);
+
+		const float3 OpticalDepth = ExtinctionCoeff * DeltaDepth;
+		float3 Transmittance = exp(-OpticalDepth);
+		float3 ScatteredLuminance = ScatteringCoeff * (AmbScattLi + SunScattLi * DirectionalLightShadow);
+		ScatteredLuminance = (ScatteredLuminance - ScatteredLuminance * Transmittance) / ExtinctionCoeff;
+#else
+		float3 Transmittance = 1.0f;
+		float3 ScatteredLuminance = 0.0f;
+		const float RayMarchMaxDistance = min(DeltaDepth, 200.0f);  // 20 meters
+		const float RayMarchStepSize = RayMarchMaxDistance / 10.0f; // Less samples wil lresult in a bit brighter look due to TransmittanceToLightThroughWater being 1 on a longer first sample. Would need it part of analiytical integration
+		const float ShadowDither = RayMarchStepSize * GBufferDither;
+		for (float s = 0.0f; s < RayMarchMaxDistance; s += RayMarchStepSize)
+		{
+			// Only jitter shadow map sampling to not lose energy on first sample
+			float Shadow = ComputeDirectionalLightDynamicShadowing(MaterialParameters.AbsoluteWorldPosition - (s + ShadowDither)*MaterialParameters.CameraVector, GBuffer.Depth);
+
+			float3 WP = MaterialParameters.AbsoluteWorldPosition - s * MaterialParameters.CameraVector;
+			float WaterHeightAboveSample = max(0.0, MaterialParameters.AbsoluteWorldPosition.z - WP.z);
+			float3 TransmittanceToLightThroughWater = 1.0;												// no self shadow, same energy as above analytical solution
+			//float3 TransmittanceToLightThroughWater = exp(-ExtinctionCoeff * WaterHeightAboveSample);	// self shadow as transmittance to water level, close to reference, depends a bit on sample count due to first sample being critical for dense medium
+
+			float3 SampleTransmittance = exp(-ExtinctionCoeff * RayMarchStepSize); // Constant
+			float3 SS = (ScatteringCoeff * TransmittanceToLightThroughWater * (SunScattLi * Shadow + AmbScattLi));
+			ScatteredLuminance += Transmittance * (SS - SS * SampleTransmittance) / ExtinctionCoeff;
+			Transmittance *= SampleTransmittance;
+		}
+
+		// The rest of the medium
+		const float3 OpticalDepth2 = ExtinctionCoeff * max(0.0, DeltaDepth - RayMarchMaxDistance);
+		if (any(OpticalDepth2 > 0.0f))
+		{
+			float3 Transmittance2 = exp(-OpticalDepth2);
+			float3 ScatteredLuminance2 = ScatteringCoeff * (SunScattLi + AmbScattLi);
+			ScatteredLuminance += Transmittance * (ScatteredLuminance2 - ScatteredLuminance2 * Transmittance2) / ExtinctionCoeff;
+			Transmittance *= Transmittance2;
+		}
+#endif
+
+		// Apply Fresnel effect
+		const float3 N = MaterialParameters.WorldNormal;
+		const float3 V = MaterialParameters.CameraVector;
+		const float NoV = max(0.0, dot(N, V));
+		float3 EnvBrdf = EnvBRDF(GBuffer.SpecularColor, GBuffer.Roughness, NoV);
+		ScatteredLuminance *= CameraIsUnderWater ? 1.0 : (1.0 - EnvBrdf); // Under water is less visibule due to Frensel effect
+		Transmittance      *= CameraIsUnderWater ? (1.0 - EnvBrdf) : 1.0; // Above	"		"		"		"		"
+
+		// Add single in-scattering apply colored transmittance to scene color
+		Color += ScatteredLuminance * WaterVisibility;
+		Color += SceneColor * Transmittance * WaterVisibility ;
+	}
+
+
+#endif // MATERIAL_IS_SINGLE_LAYER_WATER
+
+
+
 	#if MATERIAL_DOMAIN_POSTPROCESS
 		#if MATERIAL_OUTPUT_OPACITY_AS_ALPHA
 			Out.MRT[0] = half4(Color, Opacity);
diff --git a/Engine/Shaders/Private/ForwardLightingCommon.ush b/Engine/Shaders/Private/ForwardLightingCommon.ush
--- a/Engine/Shaders/Private/ForwardLightingCommon.ush
+++ b/Engine/Shaders/Private/ForwardLightingCommon.ush
@@ -132,7 +132,8 @@ float3 GetForwardDirectLighting(uint GridIndex, float3 WorldPosition, float3 Cam
 		// Static shadowing uses ShadowMapChannel, dynamic shadows are packed into light attenuation using PreviewShadowMapChannel
 		LightData.ShadowMapChannelMask = UnpackShadowMapChannelMask(DirectionalLightData.DirectionalLightShadowMapChannelMask);
 
-		#if MATERIALBLENDING_SOLID || MATERIALBLENDING_MASKED
+		// We want to force the directional light shadow when using water material to see shadow on the water. This could be an option later.
+		#if ((MATERIALBLENDING_SOLID || MATERIALBLENDING_MASKED) && !MATERIAL_IS_SINGLE_LAYER_WATER)
 			float4 PreviewShadowMapChannelMask = UnpackShadowMapChannelMask(DirectionalLightData.DirectionalLightShadowMapChannelMask >> 4);
 			float DynamicShadowing = dot(PreviewShadowMapChannelMask, DynamicShadowFactors);
 
diff --git a/Engine/Shaders/Private/MobileBasePassPixelShader.usf b/Engine/Shaders/Private/MobileBasePassPixelShader.usf
--- a/Engine/Shaders/Private/MobileBasePassPixelShader.usf
+++ b/Engine/Shaders/Private/MobileBasePassPixelShader.usf
@@ -22,6 +22,10 @@
 #include "PlanarReflectionShared.ush"
 #include "MobileGGX.ush"
 
+#if MATERIAL_IS_SINGLE_LAYER_WATER
+#include "ParticipatingMediaCommon.ush"
+#endif
+
 // Add fading CSM plane:
 #define FADE_CSM 1
 
@@ -386,6 +390,8 @@ void Main(
 	GetMaterialCoverageAndClipping(MaterialParameters, PixelMaterialInputs);
 
 	half Opacity = GetMaterialOpacity(PixelMaterialInputs);
+	const float BaseMaterialCoverageOverWater = Opacity;
+	const float WaterVisibility = 1.0 - BaseMaterialCoverageOverWater;
 
 	// Store the results in local variables and reuse instead of calling the functions multiple times.
 	// Store the results in local variables and reuse instead of calling the functions multiple times.
@@ -430,6 +436,13 @@ void Main(
 #endif
 	}
 #endif
+
+#if MATERIAL_IS_SINGLE_LAYER_WATER
+	// Fade out diffuse as this will be hanlded by the single scattering lighting. when over the water surface.
+	// We keep the SpecularColor for sun/water interactions
+	float3 WaterIndirectDiffuse = 0;
+	DiffuseColor *= BaseMaterialCoverageOverWater;
+#endif
 	
 	half3 Color = 0;
 	half IndirectIrradiance = 1;
@@ -571,7 +584,11 @@ void Main(
 		#if ENABLE_SKY_LIGHT
 			//@mw todo
 			// TODO: Also need to do specular.
-			Color += GetSkySHDiffuseSimple(MaterialParameters.WorldNormal) * ResolvedView.SkyLightColor.rgb * DiffuseColor * MaterialAO;    
+			float3 SkyDiffuseLighting = GetSkySHDiffuseSimple(MaterialParameters.WorldNormal);
+		#if MATERIAL_IS_SINGLE_LAYER_WATER
+			WaterIndirectDiffuse += SkyDiffuseLighting;
+		#endif
+			Color += SkyDiffuseLighting * ResolvedView.SkyLightColor.rgb * DiffuseColor * MaterialAO;
 		#endif
 
 		#if MAX_DYNAMIC_POINT_LIGHTS > 0
@@ -637,7 +654,68 @@ void Main(
 		 
 	half3 Emissive = GetMaterialEmissive(PixelMaterialInputs);
 
-	Color += Emissive;
+	Color += Emissive; 
+
+
+
+#if MATERIAL_IS_SINGLE_LAYER_WATER
+	// Fade out the material contribution over to water contribution according to alpha.
+	const bool CameraIsUnderWater = false;
+
+	float WaterDepth = ConvertFromDeviceZ(MaterialParameters.SvPosition.z);
+	if (WaterVisibility > 0.0f) // DeltaDepth > 0.0f should not be needed as we ARE on a water surface TODO simplify that
+	{
+		float2 NDC = MaterialParameters.ScreenPosition.xy / MaterialParameters.ScreenPosition.w;
+		float2 ViewportUV = NDC * ResolvedView.ScreenPositionScaleBias.xy + ResolvedView.ScreenPositionScaleBias.wz;
+
+		float PixelSceneDeviceZ = MobileSceneTextures.SceneDepthTexture.SampleLevel(MobileSceneTextures.SceneDepthTextureSampler, ViewportUV, 0).r;
+		float PixelSceneDepth = ConvertFromDeviceZ(PixelSceneDeviceZ);
+		float DeltaDepth = CameraIsUnderWater ? WaterDepth : PixelSceneDepth - WaterDepth;
+
+
+		const float3 ScatteringCoeff = max(0.0f, GetSingleLayerWaterMaterialOutput0(MaterialParameters));
+		const float3 AbsorptionCoeff = max(0.0f, GetSingleLayerWaterMaterialOutput1(MaterialParameters));
+		const float PhaseG = clamp(GetSingleLayerWaterMaterialOutput2(MaterialParameters), -1.0f, 1.0f);
+		const float PhaseValue = SchlickPhase(PhaseG, dot(ResolvedView.DirectionalLightDirection.xyz, MaterialParameters.CameraVector));
+		const float3 ExtinctionCoeff = ScatteringCoeff + AbsorptionCoeff;
+		const float3 OpticalDepth = ExtinctionCoeff * DeltaDepth;
+		float3 Transmittance = exp(-OpticalDepth);
+
+		// We also apply transmittance from light to under water surface. However, the scene has been lit by many sources already.
+		// So the transmittance toabove surface is simply approximated using the travel distance from the scene pixel to the water top, assuming a flat water surface.
+		// We also use the pixel SvPosition instead of the scene one pre refraction/distortion to avoid those extra ALUs.
+		// We cannot combine this transmittance with the transmittance from view because this would change the behavior of the analytical integration of light scattering integration.
+		const float3 SceneWorldPos = SvPositionToWorld(float4(MaterialParameters.SvPosition.xy, PixelSceneDeviceZ, 1.0));
+		const float DistanceFromScenePixelToWaterTop = max(0.0, MaterialParameters.AbsoluteWorldPosition.z - SceneWorldPos.z);
+		const float FromSurfaceToLightTransmittance = exp(-DistanceFromScenePixelToWaterTop * ExtinctionCoeff);
+
+		const float DirectionalLightShadow = Shadow;
+
+		float3 SunEi = ResolvedView.DirectionalLightColor.rgb * PI;	// times PI because it is divided by PI on CPU (=luminance) and we want illuminance here.  
+		float3 SunScattLi = PhaseValue * SunEi;
+		float3 AmbScattLi = PhaseValue * WaterIndirectDiffuse * PI;	// WaterIndirectDiffuse is luminance? So we need to multiply by PI to get illuminance. That we can then get through the phase function.
+
+		float3 ScatteredLuminance = ScatteringCoeff * (AmbScattLi + SunScattLi * DirectionalLightShadow);
+		ScatteredLuminance = (ScatteredLuminance - ScatteredLuminance * Transmittance) / ExtinctionCoeff;
+
+		// Apply Fresnel effect
+		float3 EnvBrdf = SpecularColor; // SpecularColor is not F0 as in BasePassPixelShader but EnvBRDF.
+		ScatteredLuminance *= CameraIsUnderWater ? 1.0 : (1.0 - EnvBrdf); // Under water is less visibule due to Frensel effect
+		Transmittance *= CameraIsUnderWater ? (1.0 - EnvBrdf) : 1.0; // Above	"		"		"		"		"
+
+		// Add single in-scattering apply colored transmittance to scene color
+		Color += ScatteredLuminance * WaterVisibility;
+//		Color += SceneColor * Transmittance * WaterVisibility;
+
+		Opacity = 1.0 - ((1.0 - Opacity) * dot(Transmittance*FromSurfaceToLightTransmittance, float3(1.0 / 3.0, 1.0 / 3.0, 1.0 / 3.0)));
+	}
+#endif
+
+
+
+
+
+
 
 	half4 VertexFog = half4(0, 0, 0, 1);
 
@@ -654,7 +732,8 @@ void Main(
 	Color = lerp(Color, DiffuseColor + SpecularColor, ResolvedView.UnlitViewmodeMask);
 #endif
 
-	#if MATERIALBLENDING_ALPHACOMPOSITE
+	// On mobile, water (an opaque material) is rendered as trnaslucent with forced premultiplied alpha blending (see MobileBasePass::SetTranslucentRenderState)
+	#if MATERIALBLENDING_ALPHACOMPOSITE || MATERIAL_IS_SINGLE_LAYER_WATER
 		OutColor = half4(Color * VertexFog.a + VertexFog.rgb * Opacity, Opacity);
 	#elif MATERIALBLENDING_ALPHAHOLDOUT
 		// not implemented for holdout
diff --git a/Engine/Shaders/Private/ParticipatingMediaCommon.ush b/Engine/Shaders/Private/ParticipatingMediaCommon.ush
new file mode 100644
--- /dev/null
+++ b/Engine/Shaders/Private/ParticipatingMediaCommon.ush
@@ -0,0 +1,38 @@
+/*=============================================================================
+	ParticipatingMediaCommon.ush
+=============================================================================*/
+
+#include "Common.ush"
+
+
+
+float IsotropicPhase()
+{
+	return 1.0f / (4.0f * PI);
+}
+
+float HenyeyGreensteinPhase(float g, float CosTheta)
+{
+	g = -g;
+	return (1.0f - g * g) / (4.0f * PI * pow(1.0f + g * g - 2.0f * g * CosTheta, 1.5f));
+}
+
+float RaleighPhase(float CosTheta)
+{
+	return 3.0f * (1.0f + CosTheta * CosTheta) / (16.0f * PI);
+}
+
+// Schlick phase function approximating henyey-greenstein
+float SchlickPhaseFromK(float k, float CosTheta)
+{
+	const float SchlickPhaseFactor = 1.0f + k * CosTheta;
+	const float PhaseValue = (1.0f - k * k) / (4.0f * PI * SchlickPhaseFactor * SchlickPhaseFactor);
+	return PhaseValue;
+}
+float SchlickPhase(float g, float CosTheta)
+{
+	const float k = 1.55f * g - 0.55f * g * g * g;
+	return SchlickPhaseFromK(k, CosTheta);
+}
+
+
diff --git a/Engine/Shaders/Private/SingleLayerWaterComposite.usf b/Engine/Shaders/Private/SingleLayerWaterComposite.usf
new file mode 100644
--- /dev/null
+++ b/Engine/Shaders/Private/SingleLayerWaterComposite.usf
@@ -0,0 +1,139 @@
+// Copyright 1998-2019 Epic Games, Inc. All Rights Reserved.
+
+/*=============================================================================
+	WaterComposite - functionality to apply local cubemaps.
+=============================================================================*/
+
+#include "Common.ush"
+#include "DeferredShadingCommon.ush"  
+#include "BRDF.ush"
+#include "ReflectionEnvironmentShared.ush"
+#include "ShadingModels.ush"
+#include "SceneTextureParameters.ush"
+
+
+#define ENABLE_SKY_LIGHT 1
+
+#define REFLECTION_COMPOSITE_USE_BLENDED_REFLECTION_CAPTURES 0
+#define REFLECTION_COMPOSITE_SUPPORT_SKYLIGHT_BLEND 0
+#include "ReflectionEnvironmentComposite.ush"
+
+
+
+Texture2D ScreenSpaceReflectionsTexture;
+SamplerState ScreenSpaceReflectionsSampler;
+
+Texture2D SceneNoWaterDepthTexture;
+SamplerState SceneNoWaterDepthSampler;
+
+
+
+void SingleLayerWaterCompositePS(
+	in float4 SvPosition : SV_Position,
+	out float4 OutColor : SV_Target0)
+{
+	float2 BufferUV = SvPositionToBufferUV(SvPosition);
+	float2 ScreenPosition = SvPositionToScreenPosition(SvPosition).xy;
+
+
+	// Sample scene textures.
+	FGBufferData GBuffer = GetGBufferDataFromSceneTextures(BufferUV);
+	uint ShadingModelID = GBuffer.ShadingModelID;
+	// No AO or DFAO
+
+	const float3 OneThird3 = float3(1.0 / 3.0f, 1.0 / 3.0f, 1.0 / 3.0f);
+
+	float3 WorldPos = SvPositionToWorld(float4(SvPosition.xy, 1.0, SvPosition.w));
+	const bool CameraIsUnderWater = false;
+
+	float4 LuminanceTransmittance = float4(0, 0, 0, 1);
+	float3 TransmittanceToScene = 1.0f;
+
+#if USE_PREEXPOSURE
+	const float PreExposure = View.PreExposure;
+#else
+	const float PreExposure = 1.f;
+#endif
+	if (ShadingModelID != SHADINGMODELID_UNLIT)
+	{
+		const float OpaqueDepth = ConvertFromDeviceZ(SceneNoWaterDepthTexture.SampleLevel(SceneNoWaterDepthSampler, BufferUV, 0).r);
+		const float WaterDepth = GBuffer.Depth;
+		float DeltaDepth = CameraIsUnderWater ? WaterDepth :  OpaqueDepth - WaterDepth;	// Inverted depth
+		if (DeltaDepth > 0.0)
+		{
+			// Compute the sky reflection contribution
+			float3 Reflection = 0.0f;
+
+			// TODO MAKE COMMON WITH ReflectionEnvironment?
+			float3 WorldPosition = mul(float4(ScreenPosition * GBuffer.Depth, GBuffer.Depth, 1), View.ScreenToWorld).xyz;
+			float3 CameraToPixel = normalize(WorldPosition - View.WorldCameraOrigin);
+			float3 ReflectionVector = reflect(CameraToPixel, GBuffer.WorldNormal);
+			float IndirectIrradiance = 0.0;// GBuffer.IndirectIrradiance;
+			float IndirectSpecularOcclusion = 1.0f;
+			float3 ExtraIndirectSpecular = 0;
+			uint CaptureDataStartIndex = 0;
+			uint NumCulledReflectionCaptures = 0;
+			//
+			float3 N = GBuffer.WorldNormal;
+			float3 V = -CameraToPixel;
+			float3 R = 2 * dot(V, N) * N - V;
+			float NoV = saturate(dot(N, V));
+			// Point lobe in off-specular peak direction
+			R = GetOffSpecularPeakReflectionDir(N, R, GBuffer.Roughness);
+			Reflection += PreExposure * CompositeReflectionCapturesAndSkylight(
+				1,
+				WorldPosition,
+				R,
+				GBuffer.Roughness,
+				IndirectIrradiance,
+				IndirectSpecularOcclusion,
+				ExtraIndirectSpecular,
+				NumCulledReflectionCaptures,
+				CaptureDataStartIndex,
+				0);
+
+			// Then combine reflection with SSR
+			float4 SSR = Texture2DSample(ScreenSpaceReflectionsTexture, ScreenSpaceReflectionsSampler, BufferUV);
+			Reflection = SSR.rgb + Reflection * (1 - SSR.a);
+
+			// Apply the BRDF reflection integral to reflection. 
+			// BRDF/Fresnel is already applied for the under water part in the scene.
+			float3 EnvBrdf = EnvBRDF(GBuffer.SpecularColor, GBuffer.Roughness, NoV);
+			Reflection *= EnvBrdf;
+
+			// Soft fading near shor to avoid seeing triangles.
+			const float ShoreOpacity = saturate(DeltaDepth * 0.02);
+			Reflection.rgb *= ShoreOpacity;
+
+			// Apply transmittance to reflection if under water
+			if (CameraIsUnderWater)
+			{
+				TransmittanceToScene *= 1.0 - EnvBrdf;
+
+				// Using default under water material: compute transmittance and scattering
+				float3 WaterMediumScattering = 0.0f;
+				float3 WaterMediumTransmittance = float3(0.1, 0.1, 0.8);
+
+				Reflection *= WaterMediumTransmittance;
+				TransmittanceToScene *= WaterMediumTransmittance;
+			}
+
+			LuminanceTransmittance = float4(Reflection, dot(OneThird3, TransmittanceToScene));
+		}
+		else
+		{
+			clip(-1); // Since this shader does not write to depth or stencil it should still benefit from EarlyZ (See AMD depth-in-depth documentation)
+		}
+	}
+	else
+	{
+		clip(-1); // Since this shader does not write to depth or stencil it should still benefit from EarlyZ (See AMD depth-in-depth documentation)
+	}
+
+	// TODO use dual source blending to apply TransmittanceToScene
+	OutColor = LuminanceTransmittance;
+}
+
+
+
+
diff --git a/Engine/Shaders/Private/VolumetricFog.usf b/Engine/Shaders/Private/VolumetricFog.usf
--- a/Engine/Shaders/Private/VolumetricFog.usf
+++ b/Engine/Shaders/Private/VolumetricFog.usf
@@ -14,6 +14,7 @@
 #include "VolumeLightingCommon.ush"
 #include "VolumetricLightmapShared.ush"
 #include "ForwardShadowingCommon.ush"
+#include "ParticipatingMediaCommon.ush"
 
 RWTexture3D<float4> RWVBufferA;
 RWTexture3D<float4> RWVBufferB;
@@ -103,28 +104,6 @@ void MaterialSetupCS(
 	}
 }
 
-float IsotropicPhase()
-{
-	return 1 / (4 * PI);
-}
-
-float HenyeyGreensteinPhase(float g, float CosTheta)
-{
-	g = -g;
-	return (1 - g * g) / (4 * PI * pow(1 + g * g - 2 * g * CosTheta, 1.5f));
-}
-
-float SchlickPhase(float k, float CosTheta)
-{
-	float Inner = (1 + k * CosTheta);
-	return (1 - k * k) / (4 * PI * Inner * Inner);
-}
-
-float RaleighPhase(float CosTheta)
-{
-	return 3.0f * (1.0f + CosTheta * CosTheta) / (16.0f * PI);
-}
-
 // Positive g = forward scattering
 // Zero g = isotropic
 // Negative g = backward scattering
diff --git a/Engine/Source/Runtime/Engine/Classes/Materials/MaterialExpressionSingleLayerWaterMaterialOutput.h b/Engine/Source/Runtime/Engine/Classes/Materials/MaterialExpressionSingleLayerWaterMaterialOutput.h
new file mode 100644
--- /dev/null
+++ b/Engine/Source/Runtime/Engine/Classes/Materials/MaterialExpressionSingleLayerWaterMaterialOutput.h
@@ -0,0 +1,41 @@
+// Copyright 1998-2019 Epic Games, Inc. All Rights Reserved.
+
+#pragma once
+
+#include "CoreMinimal.h"
+#include "Materials/MaterialExpressionCustomOutput.h"
+#include "UObject/ObjectMacros.h"
+#include "MaterialExpressionSingleLayerWaterMaterialOutput.generated.h"
+
+/** Material output expression for writing single layer water volume material properties. */
+UCLASS(collapsecategories, hidecategories = Object)
+class UMaterialExpressionSingleLayerWaterMaterialOutput : public UMaterialExpressionCustomOutput
+{
+	GENERATED_UCLASS_BODY()
+
+	/** Input for scattering coefficient describing how light scatter around and is absorbed. Valid range is [0,+inf[. */
+	UPROPERTY()
+	FExpressionInput ScatteringCoefficients;
+
+	/** Input for scattering coefficient describing how light bounce is absorbed. Valid range is [0,+inf[. */
+	UPROPERTY()
+	FExpressionInput AbsorptionCoefficients;
+
+	/** Input for phase function 'g' parameter describing how much forward(g<0) or backward (g>0) light scatter around. Valid range is [-1,1]. */
+	UPROPERTY()
+	FExpressionInput PhaseG;
+
+public:
+#if WITH_EDITOR
+	//~ Begin UMaterialExpression Interface
+	virtual int32 Compile(class FMaterialCompiler* Compiler, int32 OutputIndex) override;
+	virtual void GetCaption(TArray<FString>& OutCaptions) const override;
+	//~ End UMaterialExpression Interface
+#endif
+
+	//~ Begin UMaterialExpressionCustomOutput Interface
+	virtual int32 GetNumOutputs() const override;
+	virtual FString GetFunctionName() const override;
+	virtual FString GetDisplayName() const override;
+	//~ End UMaterialExpressionCustomOutput Interface
+};
diff --git a/Engine/Source/Runtime/Engine/Classes/Materials/MaterialInterface.h b/Engine/Source/Runtime/Engine/Classes/Materials/MaterialInterface.h
--- a/Engine/Source/Runtime/Engine/Classes/Materials/MaterialInterface.h
+++ b/Engine/Source/Runtime/Engine/Classes/Materials/MaterialInterface.h
@@ -78,6 +78,7 @@ struct ENGINE_API FMaterialRelevance
 	uint8 bTranslucentSurfaceLighting : 1;
 	uint8 bUsesSceneDepth : 1;
 	uint8 bUsesSkyMaterial : 1;
+	uint8 bUsesSingleLayerWaterMaterial : 1;
 	uint8 bHasVolumeMaterialDomain : 1;
 	uint8 bUsesDistanceCullFade : 1;
 
diff --git a/Engine/Source/Runtime/Engine/Private/Materials/HLSLMaterialTranslator.h b/Engine/Source/Runtime/Engine/Private/Materials/HLSLMaterialTranslator.h
--- a/Engine/Source/Runtime/Engine/Private/Materials/HLSLMaterialTranslator.h
+++ b/Engine/Source/Runtime/Engine/Private/Materials/HLSLMaterialTranslator.h
@@ -17,6 +17,7 @@
 #include "Materials/Material.h"
 #include "Materials/MaterialExpressionMaterialFunctionCall.h"
 #include "Materials/MaterialFunctionInstance.h"
+#include "Materials/MaterialExpressionSingleLayerWaterMaterialOutput.h"
 #include "MaterialCompiler.h"
 #include "RenderUtils.h"
 #include "EngineGlobals.h"
@@ -948,7 +949,11 @@ public:
 			Chunk[MP_CustomData1]					= Material->CompilePropertyAndSetMaterialProperty(MP_CustomData1		,this);
 			Chunk[MP_AmbientOcclusion]				= Material->CompilePropertyAndSetMaterialProperty(MP_AmbientOcclusion	,this);
 
-			if (IsTranslucentBlendMode(BlendMode))
+			// We parse the node graph to see if a water material is used. 
+			// We cannot use output node registration from custom output node for water because CompileCustomOutputs is called below and we need to enabled some input before hand.
+			const bool bNodeGraphIsUsingSingleLayerWaterMaterialOutput = Material->IsUsingSingleLayerWaterMaterialOutput();
+
+			if (IsTranslucentBlendMode(BlendMode) || bNodeGraphIsUsingSingleLayerWaterMaterialOutput)
 			{
 				int32 UserRefraction = ForceCast(Material->CompilePropertyAndSetMaterialProperty(MP_Refraction, this), MCT_Float1);
 				int32 RefractionDepthBias = ForceCast(ScalarParameter(FName(TEXT("RefractionDepthBias")), Material->GetRefractionDepthBiasValue()), MCT_Float1);
@@ -1073,6 +1078,7 @@ ResourcesString = TEXT("");
 			MaterialCompilationOutput.bModifiesMeshPosition = bUsesPixelDepthOffset || bUsesWorldPositionOffset;
 			MaterialCompilationOutput.bUsesWorldPositionOffset = bUsesWorldPositionOffset;
 			MaterialCompilationOutput.bUsesPixelDepthOffset = bUsesPixelDepthOffset;
+			MaterialCompilationOutput.bUsesSingleLayerWaterMaterialOutput = bNodeGraphIsUsingSingleLayerWaterMaterialOutput;
 
 			// Fully rough if we have a roughness code chunk and it's constant and evaluates to 1.
 			bIsFullyRough = Chunk[MP_Roughness] != INDEX_NONE && IsMaterialPropertyUsed(MP_Roughness, Chunk[MP_Roughness], FLinearColor(1, 0, 0, 0), 1) == false;
@@ -1150,6 +1156,11 @@ ResourcesString = TEXT("");
 				Errorf(TEXT("Sky materials must be opaque and unlit. They are expected to completely replace the background."));
 			}
 
+			if (bNodeGraphIsUsingSingleLayerWaterMaterialOutput && (!MaterialShadingModels.IsLit() || BlendMode != BLEND_Opaque || !MaterialShadingModels.HasOnlyShadingModel(MSM_DefaultLit)))
+			{
+				Errorf(TEXT("Water materials must be opaque and lit and they only support the DefaultLit shading model."));
+			}
+
 			bool bDBufferAllowed = IsUsingDBuffers(Platform);
 			bool bDBufferBlendMode = IsDBufferDecalBlendMode((EDecalBlendMode)Material->GetDecalBlendMode());
 
@@ -1400,6 +1411,7 @@ ResourcesString = TEXT("");
 		OutEnvironment.SetDefine(TEXT("NEEDS_PARTICLE_TRANSFORM"), bUsesParticleTransform);
 		OutEnvironment.SetDefine(TEXT("USES_TRANSFORM_VECTOR"), bUsesTransformVector);
 		OutEnvironment.SetDefine(TEXT("WANT_PIXEL_DEPTH_OFFSET"), bUsesPixelDepthOffset);
+		OutEnvironment.SetDefine(TEXT("MATERIAL_IS_SINGLE_LAYER_WATER"), Material->IsUsingSingleLayerWaterMaterialOutput());
 		if (IsMetalPlatform(InPlatform))
 		{
 			OutEnvironment.SetDefine(TEXT("USES_WORLD_POSITION_OFFSET"), bUsesWorldPositionOffset);
diff --git a/Engine/Source/Runtime/Engine/Private/Materials/Material.cpp b/Engine/Source/Runtime/Engine/Private/Materials/Material.cpp
--- a/Engine/Source/Runtime/Engine/Private/Materials/Material.cpp
+++ b/Engine/Source/Runtime/Engine/Private/Materials/Material.cpp
@@ -43,6 +43,7 @@
 #include "Materials/MaterialFunction.h"
 #include "Materials/MaterialFunctionInstance.h"
 #include "Materials/MaterialExpressionMaterialFunctionCall.h"
+#include "Materials/MaterialExpressionSingleLayerWaterMaterialOutput.h"
 #include "SceneManagement.h"
 #include "Materials/MaterialUniformExpressions.h"
 #include "Engine/SubsurfaceProfile.h"
@@ -156,7 +157,8 @@ int32 FMaterialResource::CompilePropertyAndSetMaterialProperty(EMaterialProperty
 		case MP_OpacityMask:
 			// Force basic opaque surfaces to skip masked/translucent-only attributes.
 			// Some features can force the material to create a masked variant which unintentionally runs this dormant code
-			if (GetMaterialDomain() != MD_Surface || GetBlendMode() != BLEND_Opaque || (GetShadingModels().IsLit() && !GetShadingModels().HasShadingModel(MSM_DefaultLit)))
+			if (GetMaterialDomain() != MD_Surface || GetBlendMode() != BLEND_Opaque || (GetShadingModels().IsLit() && !GetShadingModels().HasShadingModel(MSM_DefaultLit)) 
+				|| MaterialInterface->GetMaterial()->HasAnyExpressionsInMaterialAndFunctionsOfType<UMaterialExpressionSingleLayerWaterMaterialOutput>())
 			{
 				Ret = MaterialInterface->CompileProperty(Compiler, Property);
 			}
@@ -902,6 +904,7 @@ UMaterial::UMaterial(const FObjectInitializer& ObjectInitializer)
 	bCastRayTracedShadows = true;
 	bUseTranslucencyVertexFog = true;
 	bIsSky = false;
+	bUsedWithWater = false;
 	BlendableLocation = BL_AfterTonemapping;
 	BlendablePriority = 0;
 	BlendableOutputAlpha = false;
@@ -5952,7 +5955,8 @@ static bool IsPropertyActive_Internal(EMaterialProperty InProperty,
 	bool bBlendableOutputAlpha,
 	bool bHasTessellation,
 	bool bHasRefraction,
-	bool bUsesShadingModelFromMaterialExpression)
+	bool bUsesShadingModelFromMaterialExpression,
+	bool bIsSingleLayerWater)
 {
 	if (Domain == MD_PostProcess)
 	{
@@ -6137,10 +6141,10 @@ static bool IsPropertyActive_Internal(EMaterialProperty InProperty,
 		Active = false;
 		break;
 	case MP_Refraction:
-		Active = bIsTranslucentBlendMode && BlendMode != BLEND_AlphaHoldout && BlendMode != BLEND_Modulate;
+		Active = (bIsTranslucentBlendMode && BlendMode != BLEND_AlphaHoldout && BlendMode != BLEND_Modulate) || bIsSingleLayerWater;
 		break;
 	case MP_Opacity:
-		Active = bIsTranslucentBlendMode && BlendMode != BLEND_Modulate;
+		Active = (bIsTranslucentBlendMode && BlendMode != BLEND_Modulate) || bIsSingleLayerWater;
 		if (IsSubsurfaceShadingModel(ShadingModels))
 		{
 			Active = true;
@@ -6207,6 +6211,7 @@ bool UMaterial::IsPropertyActive(EMaterialProperty InProperty) const
 #if WITH_EDITOR
 bool UMaterial::IsPropertyActiveInEditor(EMaterialProperty InProperty) const
 {
+	bool bUsesSingleLayerWaterMaterial = HasAnyExpressionsInMaterialAndFunctionsOfType<UMaterialExpressionSingleLayerWaterMaterialOutput>();
 	// explicitly DON'T use getters for BlendMode/ShadingModel...these getters may return an optimized value
 	// we want the actual value that's been set by the user in the material editor
 	return IsPropertyActive_Internal(InProperty,
@@ -6218,12 +6223,14 @@ bool UMaterial::IsPropertyActiveInEditor(EMaterialProperty InProperty) const
 		BlendableOutputAlpha,
 		D3D11TessellationMode != MTM_NoTessellation,
 		Refraction.IsConnected(),
-		IsShadingModelFromMaterialExpression());
+		IsShadingModelFromMaterialExpression(),
+		bUsesSingleLayerWaterMaterial);
 }
 #endif // WITH_EDITOR
 
 bool UMaterial::IsPropertyActiveInDerived(EMaterialProperty InProperty, const UMaterialInterface* DerivedMaterial) const
 {
+	bool bUsesSingleLayerWaterMaterial = HasAnyExpressionsInMaterialAndFunctionsOfType<UMaterialExpressionSingleLayerWaterMaterialOutput>();
 	return IsPropertyActive_Internal(InProperty,
 		MaterialDomain,
 		DerivedMaterial->GetBlendMode(),
@@ -6233,7 +6240,8 @@ bool UMaterial::IsPropertyActiveInDerived(EMaterialProperty InProperty, const UM
 		BlendableOutputAlpha,
 		D3D11TessellationMode != MTM_NoTessellation,
 		Refraction.IsConnected(),
-		DerivedMaterial->IsShadingModelFromMaterialExpression());
+		DerivedMaterial->IsShadingModelFromMaterialExpression(),
+		bUsesSingleLayerWaterMaterial);
 }
 
 #if WITH_EDITORONLY_DATA
diff --git a/Engine/Source/Runtime/Engine/Private/Materials/MaterialExpressions.cpp b/Engine/Source/Runtime/Engine/Private/Materials/MaterialExpressions.cpp
--- a/Engine/Source/Runtime/Engine/Private/Materials/MaterialExpressions.cpp
+++ b/Engine/Source/Runtime/Engine/Private/Materials/MaterialExpressions.cpp
@@ -172,6 +172,7 @@
 #include "Materials/MaterialExpressionScreenPosition.h"
 #include "Materials/MaterialExpressionShadingModel.h"
 #include "Materials/MaterialExpressionSine.h"
+#include "Materials/MaterialExpressionSingleLayerWaterMaterialOutput.h"
 #include "Materials/MaterialExpressionSobol.h"
 #include "Materials/MaterialExpressionSpeedTree.h"
 #include "Materials/MaterialExpressionSphereMask.h"
@@ -16844,4 +16845,86 @@ uint32 UMaterialExpressionShadingModel::GetOutputType(int32 OutputIndex)
 	return MCT_ShadingModel;
 }
 #endif // WITH_EDITOR
+
+UMaterialExpressionSingleLayerWaterMaterialOutput::UMaterialExpressionSingleLayerWaterMaterialOutput(const FObjectInitializer& ObjectInitializer)
+	: Super(ObjectInitializer)
+{
+	// Structure to hold one-time initialization
+	struct FConstructorStatics
+	{
+		FText NAME_Water;
+		FConstructorStatics()
+			: NAME_Water(LOCTEXT("Water", "Water"))
+		{
+		}
+	};
+	static FConstructorStatics ConstructorStatics;
+
+#if WITH_EDITORONLY_DATA
+	MenuCategories.Add(ConstructorStatics.NAME_Water);
+#endif
+
+#if WITH_EDITOR
+	Outputs.Reset();
+#endif
+}
+
+#if WITH_EDITOR
+
+int32 UMaterialExpressionSingleLayerWaterMaterialOutput::Compile(class FMaterialCompiler* Compiler, int32 OutputIndex)
+{
+	int32 CodeInput = INDEX_NONE;
+
+	if (!ScatteringCoefficients.IsConnected() && !AbsorptionCoefficients.IsConnected() && !PhaseG.IsConnected())
+	{
+		Compiler->Error(TEXT("No inputs to Single Layer Water Material."));
+	}
+
+	// Generates function names GetSingleLayerWaterMaterialOutput{index} used in BasePixelShader.usf.
+	if (OutputIndex == 0)
+	{
+		CodeInput = ScatteringCoefficients.IsConnected() ? ScatteringCoefficients.Compile(Compiler) : Compiler->Constant3(0.f, 0.f, 0.f);
+	}
+	else if (OutputIndex == 1)
+	{
+		CodeInput = AbsorptionCoefficients.IsConnected() ? AbsorptionCoefficients.Compile(Compiler) : Compiler->Constant3(0.f, 0.f, 0.f);
+	}
+	else if (OutputIndex == 2)
+	{
+		CodeInput = PhaseG.IsConnected() ? PhaseG.Compile(Compiler) : Compiler->Constant(0.f);
+	}
+
+	return Compiler->CustomOutput(this, OutputIndex, CodeInput);
+}
+
+void UMaterialExpressionSingleLayerWaterMaterialOutput::GetCaption(TArray<FString>& OutCaptions) const
+{
+	OutCaptions.Add(FString(TEXT("Single Layer Water Material")));
+}
+
+#endif // WITH_EDITOR
+
+int32 UMaterialExpressionSingleLayerWaterMaterialOutput::GetNumOutputs() const
+{
+	return 3;
+}
+
+FString UMaterialExpressionSingleLayerWaterMaterialOutput::GetFunctionName() const
+{
+	return TEXT("GetSingleLayerWaterMaterialOutput");
+}
+
+FString UMaterialExpressionSingleLayerWaterMaterialOutput::GetDisplayName() const
+{
+	return TEXT("Single Layer Water Material");
+}
+
+
+
+
+
+
+
+
+
 #undef LOCTEXT_NAMESPACE
diff --git a/Engine/Source/Runtime/Engine/Private/Materials/MaterialInterface.cpp b/Engine/Source/Runtime/Engine/Private/Materials/MaterialInterface.cpp
--- a/Engine/Source/Runtime/Engine/Private/Materials/MaterialInterface.cpp
+++ b/Engine/Source/Runtime/Engine/Private/Materials/MaterialInterface.cpp
@@ -52,6 +52,7 @@ void FMaterialRelevance::SetPrimitiveViewRelevance(FPrimitiveViewRelevance& OutV
 	OutViewRelevance.bTranslucentSurfaceLighting = bTranslucentSurfaceLighting;
 	OutViewRelevance.bUsesSceneDepth = bUsesSceneDepth;
 	OutViewRelevance.bUsesSkyMaterial = bUsesSkyMaterial;
+	OutViewRelevance.bUsesSingleLayerWaterMaterial = bUsesSingleLayerWaterMaterial;
 	OutViewRelevance.bHasVolumeMaterialDomain = bHasVolumeMaterialDomain;
 }
 
@@ -122,8 +123,12 @@ FMaterialRelevance UMaterialInterface::GetRelevance_Internal(const UMaterial* Ma
 			return FMaterialRelevance();
 		}
 
+		const bool bIsMobile = InFeatureLevel <= ERHIFeatureLevel::ES3_1;
+		const bool bUsesSingleLayerWaterMaterial = MaterialResource->MaterialUsesSingleLayerWater_GameThread();
+		const bool IsSinglePassWaterTranslucent = bIsMobile && bUsesSingleLayerWaterMaterial;
+
 		const EBlendMode BlendMode = (EBlendMode)GetBlendMode();
-		const bool bIsTranslucent = IsTranslucentBlendMode(BlendMode);
+		const bool bIsTranslucent = IsTranslucentBlendMode(BlendMode) || IsSinglePassWaterTranslucent; // We want meshes with water materials to be scheduled for translucent pass on mobile.
 
 		EMaterialDomain Domain = (EMaterialDomain)MaterialResource->GetMaterialDomain();
 		bool bDecal = (Domain == MD_DeferredDecal);
@@ -142,7 +147,7 @@ FMaterialRelevance UMaterialInterface::GetRelevance_Internal(const UMaterial* Ma
 		{
 			// Check whether the material can be drawn in the separate translucency pass as per FMaterialResource::IsTranslucencyAfterDOFEnabled and IsMobileSeparateTranslucencyEnabled
 			bool bSupportsSeparateTranclucency = Material->MaterialDomain != MD_UI && Material->MaterialDomain != MD_DeferredDecal;
-			bool bMaterialSeparateTranclucency = bSupportsSeparateTranclucency && (InFeatureLevel > ERHIFeatureLevel::ES3_1 ? Material->bEnableSeparateTranslucency : Material->bEnableMobileSeparateTranslucency);
+			bool bMaterialSeparateTranclucency = bSupportsSeparateTranclucency && (bIsMobile ? Material->bEnableMobileSeparateTranslucency : Material->bEnableSeparateTranslucency );
 			
 			MaterialRelevance.bOpaque = !bIsTranslucent;
 			MaterialRelevance.bMasked = IsMasked();
@@ -161,6 +166,7 @@ FMaterialRelevance UMaterialInterface::GetRelevance_Internal(const UMaterial* Ma
 			MaterialRelevance.bHasVolumeMaterialDomain = MaterialResource->IsVolumetricPrimitive();
 			MaterialRelevance.bUsesDistanceCullFade = MaterialResource->MaterialUsesDistanceCullFade_GameThread();
 			MaterialRelevance.bUsesSkyMaterial = Material->bIsSky;
+			MaterialRelevance.bUsesSingleLayerWaterMaterial = bUsesSingleLayerWaterMaterial;
 		}
 		return MaterialRelevance;
 	}
diff --git a/Engine/Source/Runtime/Engine/Private/Materials/MaterialShared.cpp b/Engine/Source/Runtime/Engine/Private/Materials/MaterialShared.cpp
--- a/Engine/Source/Runtime/Engine/Private/Materials/MaterialShared.cpp
+++ b/Engine/Source/Runtime/Engine/Private/Materials/MaterialShared.cpp
@@ -23,6 +23,7 @@
 #include "Materials/MaterialExpressionShadingModel.h"
 #include "Materials/MaterialExpressionReroute.h"
 #include "Materials/MaterialExpressionRuntimeVirtualTextureOutput.h"
+#include "Materials/MaterialExpressionSingleLayerWaterMaterialOutput.h"
 #include "ShaderCompiler.h"
 #include "MaterialCompiler.h"
 #include "MeshMaterialShaderType.h"
@@ -515,7 +516,8 @@ void FMaterialCompilationOutput::Serialize(FArchive& Ar)
 	PackedFlags |= (bUsesGlobalDistanceField		<< 3);
 	PackedFlags |= (bUsesPixelDepthOffset			<< 4);
 	PackedFlags |= (bUsesDistanceCullFade			<< 5);
-	PackedFlags |= (bHasRuntimeVirtualTextureOutput	<< 6);
+	PackedFlags |= (bHasRuntimeVirtualTextureOutput << 6);
+	PackedFlags |= (bUsesSingleLayerWaterMaterialOutput << 7);
 
 	Ar << PackedFlags;
 
@@ -525,7 +527,8 @@ void FMaterialCompilationOutput::Serialize(FArchive& Ar)
 	bUsesGlobalDistanceField		= (PackedFlags >> 3) & 1;
 	bUsesPixelDepthOffset			= (PackedFlags >> 4) & 1;
 	bUsesDistanceCullFade			= (PackedFlags >> 5) & 1;
-	bHasRuntimeVirtualTextureOutput	= (PackedFlags >> 6) & 1;
+	bHasRuntimeVirtualTextureOutput = (PackedFlags >> 6) & 1;
+	bUsesSingleLayerWaterMaterialOutput = (PackedFlags >> 7) & 1;
 }
 
 void FMaterial::GetShaderMapId(EShaderPlatform Platform, FMaterialShaderMapId& OutId) const
@@ -845,6 +848,19 @@ bool FMaterial::MaterialModifiesMeshPosition_GameThread() const
 	return bUsesWPO || GetTessellationMode() != MTM_NoTessellation;
 }
 
+bool FMaterial::MaterialUsesSingleLayerWater_RenderThread() const
+{
+	check(IsInParallelRenderingThread());
+	return RenderingThreadShaderMap ? RenderingThreadShaderMap->UsesSingleLayerWaterMaterialOutput() : false;
+}
+
+bool FMaterial::MaterialUsesSingleLayerWater_GameThread() const
+{
+	check(IsInGameThread());
+	FMaterialShaderMap* ShaderMap = GameThreadShaderMap.GetReference();
+	return ShaderMap ? ShaderMap->UsesSingleLayerWaterMaterialOutput() : false;
+}
+
 bool FMaterial::MaterialMayModifyMeshPosition() const
 {
 	// Conservative estimate when called before material translation has occurred. 
@@ -1449,6 +1465,11 @@ UMaterialInterface* FMaterialResource::GetMaterialInterface() const
 	return MaterialInstance ? (UMaterialInterface*)MaterialInstance : (UMaterialInterface*)Material;
 }
 
+bool FMaterialResource::IsUsingSingleLayerWaterMaterialOutput() const
+{
+	return GetMaterialInterface() ? GetMaterialInterface()->GetMaterial()->HasAnyExpressionsInMaterialAndFunctionsOfType<UMaterialExpressionSingleLayerWaterMaterialOutput>() : false;
+}
+
 #if WITH_EDITOR
 void FMaterialResource::NotifyCompilationFinished()
 {
diff --git a/Engine/Source/Runtime/Engine/Private/StaticMeshRender.cpp b/Engine/Source/Runtime/Engine/Private/StaticMeshRender.cpp
--- a/Engine/Source/Runtime/Engine/Private/StaticMeshRender.cpp
+++ b/Engine/Source/Runtime/Engine/Private/StaticMeshRender.cpp
@@ -1010,7 +1010,8 @@ void FStaticMeshSceneProxy::DrawStaticElements(FStaticPrimitiveDrawInterface* PD
 							&& !Material->IsTwoSided()
 							&& !IsTranslucentBlendMode(Material->GetBlendMode())
 							&& !Material->MaterialModifiesMeshPosition_RenderThread()
-							&& Material->GetMaterialDomain() == MD_Surface;
+							&& Material->GetMaterialDomain() == MD_Surface
+							&& !Material->MaterialUsesSingleLayerWater_RenderThread();
 
 						bAllSectionsCastShadow &= Section.bCastShadow;
 					}
diff --git a/Engine/Source/Runtime/Engine/Public/MaterialShared.h b/Engine/Source/Runtime/Engine/Public/MaterialShared.h
--- a/Engine/Source/Runtime/Engine/Public/MaterialShared.h
+++ b/Engine/Source/Runtime/Engine/Public/MaterialShared.h
@@ -554,6 +554,7 @@ public:
 		bUsesWorldPositionOffset(false),
 		bUsesGlobalDistanceField(false),
 		bUsesPixelDepthOffset(false),
+		bUsesSingleLayerWaterMaterialOutput(false),
 		bUsesDistanceCullFade(false),
 		bHasRuntimeVirtualTextureOutput(false)
 	{}
@@ -601,6 +602,9 @@ public:
 	/** true if the material writes a pixel depth offset */
 	uint8 bUsesPixelDepthOffset : 1;
 
+	/** true if the material writes a pixel depth offset */
+	uint8 bUsesSingleLayerWaterMaterialOutput : 1;
+
 	/** true if the material uses distance cull fade */
 	uint8 bUsesDistanceCullFade : 1;
 
@@ -1108,6 +1112,7 @@ public:
 	bool UsesVelocitySceneTexture() const { return MaterialCompilationOutput.UsesVelocitySceneTexture(); }
 	bool UsesDistanceCullFade() const { return MaterialCompilationOutput.bUsesDistanceCullFade; }
 	bool HasRuntimeVirtualTextureOutput() const { return MaterialCompilationOutput.bHasRuntimeVirtualTextureOutput; }
+	bool UsesSingleLayerWaterMaterialOutput() const { return MaterialCompilationOutput.bUsesSingleLayerWaterMaterialOutput; }
 #if WITH_EDITOR
 	uint32 GetNumUsedUVScalars() const { return MaterialCompilationOutput.NumUsedUVScalars; }
 	uint32 GetNumUsedCustomInterpolatorScalars() const { return MaterialCompilationOutput.NumUsedCustomInterpolatorScalars; }
@@ -1519,6 +1524,7 @@ public:
 	virtual uint32 GetStencilCompare() const { return 0; }
 	virtual bool HasRuntimeVirtualTextureOutput() const { return false; }
 	virtual bool CastsRayTracedShadows() const { return true; }
+	virtual bool IsUsingSingleLayerWaterMaterialOutput() const { return false; };
 
 	/**
 	 * Should shaders compiled for this material be saved to disk?
@@ -1605,6 +1611,10 @@ public:
 	ENGINE_API bool MaterialModifiesMeshPosition_RenderThread() const;
 	ENGINE_API bool MaterialModifiesMeshPosition_GameThread() const;
 
+	/** Does the material use a single layer water. */
+	ENGINE_API bool MaterialUsesSingleLayerWater_RenderThread() const;
+	ENGINE_API bool MaterialUsesSingleLayerWater_GameThread() const;
+
 	/** Does the material use a pixel depth offset. */
 	ENGINE_API bool MaterialUsesPixelDepthOffset() const;
 
@@ -2272,7 +2282,8 @@ public:
 	ENGINE_API virtual bool ComputeFogPerPixel() const override;
 	ENGINE_API virtual bool HasRuntimeVirtualTextureOutput() const override;
 	ENGINE_API virtual bool CastsRayTracedShadows() const override;
-	ENGINE_API virtual UMaterialInterface* GetMaterialInterface() const override;
+	ENGINE_API virtual bool IsUsingSingleLayerWaterMaterialOutput() const override;
+	ENGINE_API  virtual UMaterialInterface* GetMaterialInterface() const override;
 	/**
 	 * Should shaders compiled for this material be saved to disk?
 	 */
diff --git a/Engine/Source/Runtime/Engine/Public/PrimitiveViewRelevance.h b/Engine/Source/Runtime/Engine/Public/PrimitiveViewRelevance.h
--- a/Engine/Source/Runtime/Engine/Public/PrimitiveViewRelevance.h
+++ b/Engine/Source/Runtime/Engine/Public/PrimitiveViewRelevance.h
@@ -77,6 +77,8 @@ struct FPrimitiveViewRelevance
 	uint32 bUsesSceneDepth : 1;
 	/** Whether the primitive has materials that read the scene depth. */
 	uint32 bUsesSkyMaterial : 1;
+	/** Whether the primitive uses a single layer water material. */
+	uint32 bUsesSingleLayerWaterMaterial : 1;
 	/** Whether the view use custom data. */
 	uint32 bUseCustomViewData : 1;
 
diff --git a/Engine/Source/Runtime/RenderCore/Private/RenderCore.cpp b/Engine/Source/Runtime/RenderCore/Private/RenderCore.cpp
--- a/Engine/Source/Runtime/RenderCore/Private/RenderCore.cpp
+++ b/Engine/Source/Runtime/RenderCore/Private/RenderCore.cpp
@@ -56,6 +56,7 @@ DEFINE_STAT(STAT_DynamicPrimitiveDrawTime);
 DEFINE_STAT(STAT_StaticDrawListDrawTime);
 DEFINE_STAT(STAT_BasePassDrawTime);
 DEFINE_STAT(STAT_DepthDrawTime);
+DEFINE_STAT(STAT_WaterPassDrawTime);
 DEFINE_STAT(STAT_DynamicShadowSetupTime);
 DEFINE_STAT(STAT_RenderQueryResultTime);
 // Use 'stat initviews' to get more detail
diff --git a/Engine/Source/Runtime/RenderCore/Public/RenderCore.h b/Engine/Source/Runtime/RenderCore/Public/RenderCore.h
--- a/Engine/Source/Runtime/RenderCore/Public/RenderCore.h
+++ b/Engine/Source/Runtime/RenderCore/Public/RenderCore.h
@@ -27,6 +27,7 @@ DECLARE_CYCLE_STAT_EXTERN(TEXT("FinishRenderViewTarget"),STAT_FinishRenderViewTa
 DECLARE_CYCLE_STAT_EXTERN(TEXT("RenderVelocities"),STAT_RenderVelocities,STATGROUP_SceneRendering, RENDERCORE_API);
 DECLARE_CYCLE_STAT_EXTERN(TEXT("Depth drawing"),STAT_DepthDrawTime,STATGROUP_SceneRendering, RENDERCORE_API);
 DECLARE_CYCLE_STAT_EXTERN(TEXT("Base pass drawing"),STAT_BasePassDrawTime,STATGROUP_SceneRendering, RENDERCORE_API);
+DECLARE_CYCLE_STAT_EXTERN(TEXT("Water pass drawing"),STAT_WaterPassDrawTime,STATGROUP_SceneRendering, RENDERCORE_API);
 DECLARE_CYCLE_STAT_EXTERN(TEXT("Lighting drawing"),STAT_LightingDrawTime,STATGROUP_SceneRendering, RENDERCORE_API);
 DECLARE_CYCLE_STAT_EXTERN(TEXT("Proj Shadow drawing"),STAT_ProjectedShadowDrawTime,STATGROUP_SceneRendering, RENDERCORE_API);
 DECLARE_CYCLE_STAT_EXTERN(TEXT("Reflective Shadow map drawing"),STAT_ReflectiveShadowMapDrawTime,STATGROUP_SceneRendering, RENDERCORE_API);
diff --git a/Engine/Source/Runtime/Renderer/Private/BasePassRendering.cpp b/Engine/Source/Runtime/Renderer/Private/BasePassRendering.cpp
--- a/Engine/Source/Runtime/Renderer/Private/BasePassRendering.cpp
+++ b/Engine/Source/Runtime/Renderer/Private/BasePassRendering.cpp
@@ -578,6 +578,16 @@ void CreateOpaqueBasePassUniformBuffer(
 		}
 	}
 
+	// Single Layer Water
+	BasePassParameters.SceneLuminanceWithoutSingleLayerWaterTexture = (SceneRenderTargets.SceneColorWithoutSingleLayerWater.IsValid() ? SceneRenderTargets.SceneColorWithoutSingleLayerWater : GSystemTextures.BlackDummy)->GetRenderTargetItem().TargetableTexture;
+	BasePassParameters.SceneLuminanceWithoutSingleLayerWaterSampler = TStaticSamplerState<SF_Bilinear>::GetRHI();
+	BasePassParameters.SceneDepthWithoutSingleLayerWaterTexture = (SceneRenderTargets.SceneDepthZWithoutSingleLayerWater.IsValid() ? SceneRenderTargets.SceneDepthZWithoutSingleLayerWater : GSystemTextures.BlackDummy)->GetRenderTargetItem().TargetableTexture;
+	BasePassParameters.SceneDepthWithoutSingleLayerWaterSampler = TStaticSamplerState<SF_Point>::GetRHI();
+	BasePassParameters.PreIntegratedGFTexture = GSystemTextures.PreintegratedGF->GetRenderTargetItem().ShaderResourceTexture;
+	BasePassParameters.PreIntegratedGFSampler = TStaticSamplerState<SF_Bilinear, AM_Clamp, AM_Clamp, AM_Clamp>::GetRHI();
+	BasePassParameters.SceneCustomDepthTexture = (SceneRenderTargets.CustomDepth.IsValid() ? SceneRenderTargets.CustomDepth : GSystemTextures.BlackDummy)->GetRenderTargetItem().ShaderResourceTexture;
+	BasePassParameters.SceneCustomDepthSampler = TStaticSamplerState<SF_Point>::GetRHI();
+
 	// Misc
 	BasePassParameters.EyeAdaptation = GetEyeAdaptation(View);
 
@@ -1226,7 +1236,8 @@ void FBasePassMeshProcessor::AddMeshBatch(const FMeshBatch& RESTRICT MeshBatch,
 		if (bShouldDraw
 			&& (!PrimitiveSceneProxy || PrimitiveSceneProxy->ShouldRenderInMainPass())
 			&& ShouldIncludeDomainInMeshPass(Material.GetMaterialDomain())
-			&& !Material.IsSky())
+			&& !Material.IsSky()
+			&& !Material.MaterialUsesSingleLayerWater_RenderThread())
 		{
 			// Check for a cached light-map.
 			const bool bIsLitMaterial = ShadingModels.IsLit();
diff --git a/Engine/Source/Runtime/Renderer/Private/BasePassRendering.h b/Engine/Source/Runtime/Renderer/Private/BasePassRendering.h
--- a/Engine/Source/Runtime/Renderer/Private/BasePassRendering.h
+++ b/Engine/Source/Runtime/Renderer/Private/BasePassRendering.h
@@ -31,6 +31,7 @@ class FScene;
 template<typename TBufferStruct> class TUniformBufferRef;
 
 class FViewInfo;
+class UMaterialExpressionSingleLayerWaterMaterialOutput;
 
 /** Whether to allow the indirect lighting cache to be applied to dynamic objects. */
 extern int32 GIndirectLightingCache;
@@ -70,6 +71,15 @@ BEGIN_GLOBAL_SHADER_PARAMETER_STRUCT(FOpaqueBasePassUniformParameters,)
 	SHADER_PARAMETER_TEXTURE(Texture2D, DBufferCTexture)
 	SHADER_PARAMETER_SAMPLER(SamplerState, DBufferCTextureSampler)
 	SHADER_PARAMETER_TEXTURE(Texture2D<uint>, DBufferRenderMask)
+	// Single Layer Water
+	SHADER_PARAMETER_TEXTURE(Texture2D, SceneLuminanceWithoutSingleLayerWaterTexture)
+	SHADER_PARAMETER_SAMPLER(SamplerState, SceneLuminanceWithoutSingleLayerWaterSampler)
+	SHADER_PARAMETER_TEXTURE(Texture2D, SceneDepthWithoutSingleLayerWaterTexture)
+	SHADER_PARAMETER_SAMPLER(SamplerState, SceneDepthWithoutSingleLayerWaterSampler)
+	SHADER_PARAMETER_TEXTURE(Texture2D, PreIntegratedGFTexture)
+	SHADER_PARAMETER_SAMPLER(SamplerState, PreIntegratedGFSampler)
+	SHADER_PARAMETER_TEXTURE(Texture2D, SceneCustomDepthTexture)
+	SHADER_PARAMETER_SAMPLER(SamplerState, SceneCustomDepthSampler)
 	// Misc
 	SHADER_PARAMETER_TEXTURE(Texture2D, EyeAdaptation)
 END_GLOBAL_SHADER_PARAMETER_STRUCT()
@@ -487,6 +497,8 @@ public:
 		static const auto SupportStationarySkylight = IConsoleManager::Get().FindTConsoleVariableDataInt(TEXT("r.SupportStationarySkylight"));
 		static const auto SupportAllShaderPermutations = IConsoleManager::Get().FindTConsoleVariableDataInt(TEXT("r.SupportAllShaderPermutations"));
 
+		const bool IsSingleLayerWater = Parameters.Material->IsUsingSingleLayerWaterMaterialOutput();
+
 		const bool bTranslucent = IsTranslucentBlendMode(Parameters.Material->GetBlendMode());
 		const bool bForceAllPermutations = SupportAllShaderPermutations && SupportAllShaderPermutations->GetValueOnAnyThread() != 0;
 		const bool bProjectSupportsStationarySkylight = !SupportStationarySkylight || SupportStationarySkylight->GetValueOnAnyThread() != 0 || bForceAllPermutations;
@@ -495,6 +507,7 @@ public:
 			//translucent materials need to compile skylight support to support MOVABLE skylights also.
 			|| bTranslucent
 			// Some lightmap policies (eg Simple Forward) always require skylight support
+			|| IsSingleLayerWater
 			|| LightMapPolicyType::RequiresSkylight()
 			|| ((bProjectSupportsStationarySkylight || IsForwardShadingEnabled(Parameters.Platform)) && Parameters.Material->GetShadingModels().IsLit());
 		return bCacheShaders
@@ -659,4 +672,5 @@ ENUM_CLASS_FLAGS(FBasePassMeshProcessor::EFlags);
 
 extern void SetDepthStencilStateForBasePass(FMeshPassProcessorRenderState& DrawRenderState, ERHIFeatureLevel::Type FeatureLevel, const FMeshBatch& Mesh, const FPrimitiveSceneProxy* PrimitiveSceneProxy, bool bEnableReceiveDecalOutput, bool bUseDebugViewPS, FRHIDepthStencilState* LodFadeOverrideDepthStencilState);
 extern void SetupBasePassState(FExclusiveDepthStencil::Type BasePassDepthStencilAccess, const bool bShaderComplexity, FMeshPassProcessorRenderState& DrawRenderState);
-extern FMeshDrawCommandSortKey CalculateTranslucentMeshStaticSortKey(const FPrimitiveSceneProxy* RESTRICT PrimitiveSceneProxy, uint16 MeshIdInPrimitive);
\ No newline at end of file
+extern FMeshDrawCommandSortKey CalculateTranslucentMeshStaticSortKey(const FPrimitiveSceneProxy* RESTRICT PrimitiveSceneProxy, uint16 MeshIdInPrimitive);
+
diff --git a/Engine/Source/Runtime/Renderer/Private/DeferredShadingRenderer.cpp b/Engine/Source/Runtime/Renderer/Private/DeferredShadingRenderer.cpp
--- a/Engine/Source/Runtime/Renderer/Private/DeferredShadingRenderer.cpp
+++ b/Engine/Source/Runtime/Renderer/Private/DeferredShadingRenderer.cpp
@@ -7,6 +7,7 @@
 #include "DeferredShadingRenderer.h"
 #include "VelocityRendering.h"
 #include "AtmosphereRendering.h"
+#include "SingleLayerWaterRendering.h"
 #include "SkyAtmosphereRendering.h"
 #include "ScenePrivate.h"
 #include "ScreenRendering.h"
@@ -176,6 +177,7 @@ DECLARE_GPU_STAT(SortLights);
 DECLARE_GPU_STAT(PostRenderOpsFX);
 DECLARE_GPU_STAT(HZB);
 DECLARE_GPU_STAT_NAMED(Unaccounted, TEXT("[unaccounted]"));
+DECLARE_GPU_STAT(WaterRendering);
 
 const TCHAR* GetDepthPassReason(bool bDitheredLODTransitionsUseStencil, EShaderPlatform ShaderPlatform)
 {
@@ -380,6 +382,7 @@ DECLARE_CYCLE_STAT(TEXT("BasePass"), STAT_CLM_BasePass, STATGROUP_CommandListMar
 DECLARE_CYCLE_STAT(TEXT("AfterBasePass"), STAT_CLM_AfterBasePass, STATGROUP_CommandListMarkers);
 DECLARE_CYCLE_STAT(TEXT("Lighting"), STAT_CLM_Lighting, STATGROUP_CommandListMarkers);
 DECLARE_CYCLE_STAT(TEXT("AfterLighting"), STAT_CLM_AfterLighting, STATGROUP_CommandListMarkers);
+DECLARE_CYCLE_STAT(TEXT("WaterPass"), STAT_CLM_WaterPass, STATGROUP_CommandListMarkers);
 DECLARE_CYCLE_STAT(TEXT("Translucency"), STAT_CLM_Translucency, STATGROUP_CommandListMarkers);
 DECLARE_CYCLE_STAT(TEXT("RenderDistortion"), STAT_CLM_RenderDistortion, STATGROUP_CommandListMarkers);
 DECLARE_CYCLE_STAT(TEXT("AfterTranslucency"), STAT_CLM_AfterTranslucency, STATGROUP_CommandListMarkers);
@@ -1781,6 +1784,30 @@ void FDeferredShadingSceneRenderer::Render(FRHICommandListImmediate& RHICmdList)
 
 	checkSlow(RHICmdList.IsOutsideRenderPass());
 
+	
+	const bool bShouldRenderSingleLayerWater = ShouldRenderSingleLayerWater(Views, ViewFamily.EngineShowFlags);
+	if(bShouldRenderSingleLayerWater)
+	{
+		RHICmdList.SetCurrentStat(GET_STATID(STAT_CLM_WaterPass));
+		SCOPED_DRAW_EVENTF(RHICmdList, WaterRendering, TEXT("WaterRendering"));
+		SCOPED_GPU_STAT(RHICmdList, WaterRendering);
+
+		RHICmdList.TransitionResource(EResourceTransitionAccess::EWritable,  SceneContext.GetSceneDepthSurface());
+
+		const FExclusiveDepthStencil::Type WaterPassDepthStencilAccess = FExclusiveDepthStencil::DepthWrite_StencilWrite;
+		SceneContext.BeginRenderingWaterGBuffer(RHICmdList, WaterPassDepthStencilAccess, ViewFamily.EngineShowFlags.ShaderComplexity);
+		RenderSingleLayerWaterPass(RHICmdList, WaterPassDepthStencilAccess);
+		SceneContext.FinishWaterGBufferPassAndResolve(RHICmdList);
+
+		SceneContext.ResolveSceneDepthTexture(RHICmdList, FResolveRect(0, 0, FamilySize.X, FamilySize.Y));
+		SceneContext.ResolveSceneDepthToAuxiliaryTexture(RHICmdList);
+		RHICmdList.TransitionResource(EResourceTransitionAccess::EReadable, SceneContext.GetSceneDepthSurface());
+
+		RenderSingleLayerWaterSSR(RHICmdList);
+		ServiceLocalQueue();
+	}
+
+
 	FLightShaftsOutput LightShaftOutput;
 
 	// Draw Lightshafts
diff --git a/Engine/Source/Runtime/Renderer/Private/DeferredShadingRenderer.h b/Engine/Source/Runtime/Renderer/Private/DeferredShadingRenderer.h
--- a/Engine/Source/Runtime/Renderer/Private/DeferredShadingRenderer.h
+++ b/Engine/Source/Runtime/Renderer/Private/DeferredShadingRenderer.h
@@ -102,6 +102,17 @@ public:
 	*/
 	bool RenderBasePass(FRHICommandListImmediate& RHICmdList, FExclusiveDepthStencil::Type BasePassDepthStencilAccess, IPooledRenderTarget* ForwardScreenSpaceShadowMask, bool bParallelBasePass, bool bRenderLightmapDensity);
 
+
+	/**
+	* Renders the scene's single layer water base pass
+	* @return true if anything was rendered
+	*/
+	bool RenderSingleLayerWaterPass(FRHICommandListImmediate& RHICmdList, FExclusiveDepthStencil::Type WaterPassDepthStencilAccess);
+	/** Renders the water draw pass for a given View. */
+	bool RenderSingleLayerWaterPassView(FRHICommandListImmediate& RHICmdList, FViewInfo& View, const FMeshPassProcessorRenderState& InDrawRenderState);
+	/** Render, denoise and composite the scene SSR and under water effect.*/
+	void RenderSingleLayerWaterSSR(FRHICommandListImmediate& RHICmdList);
+
 	/** Finishes the view family rendering. */
 	void RenderFinish(FRHICommandListImmediate& RHICmdList);
 
diff --git a/Engine/Source/Runtime/Renderer/Private/DepthRendering.cpp b/Engine/Source/Runtime/Renderer/Private/DepthRendering.cpp
--- a/Engine/Source/Runtime/Renderer/Private/DepthRendering.cpp
+++ b/Engine/Source/Runtime/Renderer/Private/DepthRendering.cpp
@@ -784,7 +784,8 @@ void FDepthPassMeshProcessor::AddMeshBatch(const FMeshBatch& RESTRICT MeshBatch,
 		if (!bIsTranslucent
 			&& (!PrimitiveSceneProxy || PrimitiveSceneProxy->ShouldRenderInMainPass())
 			&& ShouldIncludeDomainInMeshPass(Material.GetMaterialDomain())
-			&& !Material.IsSky())
+			&& !Material.IsSky()
+			&& !Material.MaterialUsesSingleLayerWater_RenderThread())
 		{
 			if (BlendMode == BLEND_Opaque
 				&& MeshBatch.VertexFactory->SupportsPositionOnlyStream()
diff --git a/Engine/Source/Runtime/Renderer/Private/LightGridInjection.cpp b/Engine/Source/Runtime/Renderer/Private/LightGridInjection.cpp
--- a/Engine/Source/Runtime/Renderer/Private/LightGridInjection.cpp
+++ b/Engine/Source/Runtime/Renderer/Private/LightGridInjection.cpp
@@ -308,7 +308,7 @@ void FDeferredShadingSceneRenderer::ComputeLightGrid(FRHICommandListImmediate& R
 		for (int32 ViewIndex = 0; ViewIndex < Views.Num(); ViewIndex++)
 		{
 			const FViewInfo& View = Views[ViewIndex];
-			bAnyViewUsesForwardLighting |= View.bTranslucentSurfaceLighting || ShouldRenderVolumetricFog();
+			bAnyViewUsesForwardLighting |= View.bTranslucentSurfaceLighting || ShouldRenderVolumetricFog() || View.bHasSingleLayerWaterMaterial;
 		}
 
 		const bool bCullLightsToGrid = GLightCullingQuality 
diff --git a/Engine/Source/Runtime/Renderer/Private/LightMapDensityRendering.cpp b/Engine/Source/Runtime/Renderer/Private/LightMapDensityRendering.cpp
--- a/Engine/Source/Runtime/Renderer/Private/LightMapDensityRendering.cpp
+++ b/Engine/Source/Runtime/Renderer/Private/LightMapDensityRendering.cpp
@@ -254,7 +254,8 @@ void FLightmapDensityMeshProcessor::AddMeshBatch(const FMeshBatch& RESTRICT Mesh
 		const bool bAllowLowQualityLightMaps = (!SupportLowQualityLightmapsVar) || (SupportLowQualityLightmapsVar->GetValueOnAnyThread() != 0);
 
 		if ((!bTranslucentBlendMode || ViewIfDynamicMeshCommand->Family->EngineShowFlags.Wireframe)
-			&& !Material->IsSky())
+			&& !Material->IsSky()
+			&& !Material->MaterialUsesSingleLayerWater_RenderThread())
 		{
 			if (!bMaterialMasked && !Material->MaterialModifiesMeshPosition_RenderThread())
 			{
diff --git a/Engine/Source/Runtime/Renderer/Private/MobileBasePass.cpp b/Engine/Source/Runtime/Renderer/Private/MobileBasePass.cpp
--- a/Engine/Source/Runtime/Renderer/Private/MobileBasePass.cpp
+++ b/Engine/Source/Runtime/Renderer/Private/MobileBasePass.cpp
@@ -407,7 +407,13 @@ void MobileBasePass::SetTranslucentRenderState(FMeshPassProcessorRenderState& Dr
 			DrawRenderState.SetBlendState(TStaticBlendState<CW_RGBA, BO_Add, BF_Zero, BF_InverseSourceAlpha, BO_Add, BF_One, BF_InverseSourceAlpha>::GetRHI());
 			break;
 		default:
-			check(0);
+			if (Material.MaterialUsesSingleLayerWater_RenderThread())
+			{
+				// Single layer water is an opaque marerial rendered as translucent on Mobile. We force pre-multiplied alpha to achieve water depth based transmittance.
+				DrawRenderState.SetBlendState(TStaticBlendState<CW_RGB, BO_Add, BF_One, BF_InverseSourceAlpha, BO_Add, BF_Zero, BF_InverseSourceAlpha>::GetRHI());
+			}
+			else
+				check(0);
 		};
 	}
 	else
@@ -564,10 +570,11 @@ void FMobileBasePassMeshProcessor::AddMeshBatch(const FMeshBatch& RESTRICT MeshB
 	const EBlendMode BlendMode = Material.GetBlendMode();
 	const FMaterialShadingModelField ShadingModels = Material.GetShadingModels();
 	const bool bIsTranslucent = IsTranslucentBlendMode(BlendMode);
+	const bool bUsesWaterMaterial = Material.MaterialUsesSingleLayerWater_RenderThread(); // Water goes into the translucent pass
 	
 	if (bTranslucentBasePass)
 	{
-		bool bShouldDraw = bIsTranslucent && 
+		bool bShouldDraw = (bIsTranslucent || bUsesWaterMaterial) &&
 		(TranslucencyPassType == ETranslucencyPass::TPT_AllTranslucency
 		|| (TranslucencyPassType == ETranslucencyPass::TPT_StandardTranslucency && !Material.IsMobileSeparateTranslucencyEnabled())
 		|| (TranslucencyPassType == ETranslucencyPass::TPT_TranslucencyAfterDOF && Material.IsMobileSeparateTranslucencyEnabled()));
@@ -583,7 +590,7 @@ void FMobileBasePassMeshProcessor::AddMeshBatch(const FMeshBatch& RESTRICT MeshB
 	else
 	{
 		// opaque materials.
-		if (!bIsTranslucent)
+		if (!bIsTranslucent && !bUsesWaterMaterial)
 		{
 			const FLightSceneInfo* MobileDirectionalLight = MobileBasePass::GetDirectionalLightInfo(Scene, PrimitiveSceneProxy);
 			ELightMapPolicyType LightmapPolicyType = MobileBasePass::SelectMeshLightmapPolicy(Scene, MeshBatch, PrimitiveSceneProxy, MobileDirectionalLight, ShadingModels, bCanReceiveCSM, FeatureLevel);
diff --git a/Engine/Source/Runtime/Renderer/Private/PostProcess/SceneRenderTargets.cpp b/Engine/Source/Runtime/Renderer/Private/PostProcess/SceneRenderTargets.cpp
--- a/Engine/Source/Runtime/Renderer/Private/PostProcess/SceneRenderTargets.cpp
+++ b/Engine/Source/Runtime/Renderer/Private/PostProcess/SceneRenderTargets.cpp
@@ -232,7 +232,7 @@ FSceneRenderTargets::FSceneRenderTargets(const FViewInfo& View, const FSceneRend
 	: LightAttenuation(GRenderTargetPool.MakeSnapshot(SnapshotSource.LightAttenuation))
 	, LightAccumulation(GRenderTargetPool.MakeSnapshot(SnapshotSource.LightAccumulation))
 	, DirectionalOcclusion(GRenderTargetPool.MakeSnapshot(SnapshotSource.DirectionalOcclusion))
-	, SceneDepthZ(GRenderTargetPool.MakeSnapshot(SnapshotSource.SceneDepthZ))	
+	, SceneDepthZ(GRenderTargetPool.MakeSnapshot(SnapshotSource.SceneDepthZ))
 	, SceneVelocity(GRenderTargetPool.MakeSnapshot(SnapshotSource.SceneVelocity))
 	, LightingChannels(GRenderTargetPool.MakeSnapshot(SnapshotSource.LightingChannels))
 	, SceneAlphaCopy(GRenderTargetPool.MakeSnapshot(SnapshotSource.SceneAlphaCopy))
@@ -243,6 +243,8 @@ FSceneRenderTargets::FSceneRenderTargets(const FViewInfo& View, const FSceneRend
 	, GBufferC(GRenderTargetPool.MakeSnapshot(SnapshotSource.GBufferC))
 	, GBufferD(GRenderTargetPool.MakeSnapshot(SnapshotSource.GBufferD))
 	, GBufferE(GRenderTargetPool.MakeSnapshot(SnapshotSource.GBufferE))
+	, SceneColorWithoutSingleLayerWater(GRenderTargetPool.MakeSnapshot(SnapshotSource.SceneColorWithoutSingleLayerWater))
+	, SceneDepthZWithoutSingleLayerWater(GRenderTargetPool.MakeSnapshot(SnapshotSource.SceneDepthZWithoutSingleLayerWater))
 	, DBufferA(GRenderTargetPool.MakeSnapshot(SnapshotSource.DBufferA))
 	, DBufferB(GRenderTargetPool.MakeSnapshot(SnapshotSource.DBufferB))
 	, DBufferC(GRenderTargetPool.MakeSnapshot(SnapshotSource.DBufferC))
@@ -1000,6 +1002,70 @@ void FSceneRenderTargets::FinishGBufferPassAndResolve(FRHICommandListImmediate&
 	QuadOverdrawIndex = INDEX_NONE;
 }
 
+void FSceneRenderTargets::BeginRenderingWaterGBuffer(FRHICommandList& RHICmdList, FExclusiveDepthStencil::Type DepthStencilAccess, bool bBindQuadOverdrawBuffers)
+{
+	check(RHICmdList.IsOutsideRenderPass());
+
+	SCOPED_DRAW_EVENT(RHICmdList, BeginRenderingWaterGBuffer);
+	check(CurrentFeatureLevel >= ERHIFeatureLevel::SM4)
+
+	const EShaderPlatform ShaderPlatform = GetFeatureLevelShaderPlatform(CurrentFeatureLevel);
+	const bool bUseGBuffer = IsUsingGBuffers(ShaderPlatform);
+
+	// Allocate required buffers
+	{
+		FPooledRenderTargetDesc Desc(FPooledRenderTargetDesc::Create2DDesc(BufferSize, PF_DepthStencil, DefaultDepthClear, TexCreate_None, TexCreate_DepthStencilTargetable | TexCreate_ShaderResource | TexCreate_InputAttachmentRead, false));
+		Desc.NumSamples = GetNumSceneColorMSAASamples(CurrentFeatureLevel);
+		Desc.Flags |= GFastVRamConfig.SceneDepth;
+		GRenderTargetPool.FindFreeElement(RHICmdList, Desc, SceneDepthZWithoutSingleLayerWater, TEXT("SceneDepthZWithoutSingleLayerWater"), true);
+	}
+	{
+		FPooledRenderTargetDesc Desc(FPooledRenderTargetDesc::Create2DDesc(BufferSize, GetSceneColorFormat() , DefaultColorClear, TexCreate_None, TexCreate_RenderTargetable, false));
+		Desc.NumSamples = GetNumSceneColorMSAASamples(CurrentFeatureLevel);
+		GRenderTargetPool.FindFreeElement(RHICmdList, Desc, SceneColorWithoutSingleLayerWater, TEXT("SceneColorWithoutSingleLayerWater"));
+	}
+
+	// Copy and save textures without water informations for later composition operations
+	{
+		FRHICopyTextureInfo DepthCopyInfo;
+		RHICmdList.CopyTexture(SceneDepthZ->GetRenderTargetItem().ShaderResourceTexture, SceneDepthZWithoutSingleLayerWater->GetRenderTargetItem().ShaderResourceTexture, DepthCopyInfo);
+		RHICmdList.CopyTexture(GetSceneColorSurface(), SceneColorWithoutSingleLayerWater->GetRenderTargetItem().ShaderResourceTexture, DepthCopyInfo);
+	}
+
+	// Create MRT
+	int32 VelocityRTIndex = -1;
+	FRHIRenderPassInfo RPInfo;
+	FillGBufferRenderPassInfo(ERenderTargetLoadAction::ELoad, RPInfo, VelocityRTIndex);
+	// Set a dummy Scene color RT to avoid gbuffer to stomp HDR scene color we want to blend over
+	RPInfo.ColorRenderTargets[0].Action = MakeRenderTargetActions(ERenderTargetLoadAction::ENoAction, ERenderTargetStoreAction::ENoAction);
+	RPInfo.ColorRenderTargets[0].RenderTarget = GetSceneColorSurface();
+
+	// Stencil always has to be store or certain VK drivers will leave the attachment in an undefined state.
+	RPInfo.DepthStencilRenderTarget.Action = MakeDepthStencilTargetActions(MakeRenderTargetActions(ERenderTargetLoadAction::ELoad, ERenderTargetStoreAction::EStore), MakeRenderTargetActions(ERenderTargetLoadAction::ELoad, ERenderTargetStoreAction::EStore));
+	RPInfo.DepthStencilRenderTarget.DepthStencilTarget = (const FTexture2DRHIRef&)SceneDepthZ->GetRenderTargetItem().TargetableTexture;
+	RPInfo.DepthStencilRenderTarget.ExclusiveDepthStencil = DepthStencilAccess;
+
+	// Set other UAVs
+	const bool bClearQuadOverdrawBuffers = false;
+	SetQuadOverdrawUAV(RHICmdList, bBindQuadOverdrawBuffers, bClearQuadOverdrawBuffers, RPInfo);
+	if (UseVirtualTexturing(CurrentFeatureLevel) && !bBindQuadOverdrawBuffers)
+	{
+		BindVirtualTextureFeedbackUAV(RPInfo);
+	}
+
+	// Begin the pass
+	RHICmdList.BeginRenderPass(RPInfo, TEXT("WaterGBuffer"));
+
+	// Needs to be called after we start a renderpass in order for the color/depth decompress/expand to be executed on the next write-to-read barrier/transition.
+	RHICmdList.BindClearMRTValues(true, true, false);
+}
+
+void FSceneRenderTargets::FinishWaterGBufferPassAndResolve(FRHICommandListImmediate& RHICmdList)
+{
+	// Same as the GBuffer for now (also same resolves)
+	FinishGBufferPassAndResolve(RHICmdList);
+}
+
 int32 FSceneRenderTargets::GetNumGBufferTargets() const
 {
 	int32 NumGBufferTargets = 1;
@@ -2696,6 +2762,9 @@ void FSceneRenderTargets::ReleaseAllTargets()
 	MobileCustomStencil.SafeRelease();
 	CustomStencilSRV.SafeRelease();
 
+	SceneColorWithoutSingleLayerWater.SafeRelease();
+	SceneDepthZWithoutSingleLayerWater.SafeRelease();
+
 	for (int32 i = 0; i < ARRAY_COUNT(OptionalShadowDepthColor); i++)
 	{
 		OptionalShadowDepthColor[i].SafeRelease();
diff --git a/Engine/Source/Runtime/Renderer/Private/PostProcess/SceneRenderTargets.h b/Engine/Source/Runtime/Renderer/Private/PostProcess/SceneRenderTargets.h
--- a/Engine/Source/Runtime/Renderer/Private/PostProcess/SceneRenderTargets.h
+++ b/Engine/Source/Runtime/Renderer/Private/PostProcess/SceneRenderTargets.h
@@ -229,6 +229,9 @@ public:
 	void BeginRenderingGBuffer(FRHICommandList& RHICmdList, ERenderTargetLoadAction ColorLoadAction, ERenderTargetLoadAction DepthLoadAction, FExclusiveDepthStencil::Type DepthStencilAccess, bool bBindQuadOverdrawBuffers, bool bClearQuadOverdrawBuffers = false, const FLinearColor& ClearColor = FLinearColor(0, 0, 0, 1), bool bIsWireframe=false);
 	void FinishGBufferPassAndResolve(FRHICommandListImmediate& RHICmdList);
 
+	void BeginRenderingWaterGBuffer(FRHICommandList& RHICmdList, FExclusiveDepthStencil::Type DepthStencilAccess, bool bBindQuadOverdrawBuffers);
+	void FinishWaterGBufferPassAndResolve(FRHICommandListImmediate& RHICmdList);
+
 	/**
 	 * Sets the scene color target and restores its contents if necessary
 	 */
@@ -537,6 +540,10 @@ public:
 	TRefCountPtr<IPooledRenderTarget> GBufferD;
 	TRefCountPtr<IPooledRenderTarget> GBufferE;
 
+	// Water rendering extra buffers
+	TRefCountPtr<IPooledRenderTarget> SceneColorWithoutSingleLayerWater;
+	TRefCountPtr<IPooledRenderTarget> SceneDepthZWithoutSingleLayerWater;
+
 	// DBuffer: For decals before base pass (only temporarily available after early z pass and until base pass)
 	TRefCountPtr<IPooledRenderTarget> DBufferA;
 	TRefCountPtr<IPooledRenderTarget> DBufferB;
diff --git a/Engine/Source/Runtime/Renderer/Private/PrimitiveSceneInfo.cpp b/Engine/Source/Runtime/Renderer/Private/PrimitiveSceneInfo.cpp
--- a/Engine/Source/Runtime/Renderer/Private/PrimitiveSceneInfo.cpp
+++ b/Engine/Source/Runtime/Renderer/Private/PrimitiveSceneInfo.cpp
@@ -79,11 +79,13 @@ public:
 			const bool bSupportsCachingMeshDrawCommands = SupportsCachingMeshDrawCommands(PrimitiveSceneProxy, *StaticMesh, FeatureLevel);
 
 			bool bUseSkyMaterial = Mesh.MaterialRenderProxy->GetMaterial(FeatureLevel)->IsSky();
+			bool bUseSingleLayerWaterMaterial = Mesh.MaterialRenderProxy->GetMaterial(FeatureLevel)->MaterialUsesSingleLayerWater_RenderThread();
 			FStaticMeshBatchRelevance* StaticMeshRelevance = new(PrimitiveSceneInfo->StaticMeshRelevances) FStaticMeshBatchRelevance(
 				*StaticMesh, 
 				ScreenSize, 
 				bSupportsCachingMeshDrawCommands,
-				bUseSkyMaterial
+				bUseSkyMaterial,
+				bUseSingleLayerWaterMaterial
 			);
 		}
 	}
diff --git a/Engine/Source/Runtime/Renderer/Private/SceneCore.h b/Engine/Source/Runtime/Renderer/Private/SceneCore.h
--- a/Engine/Source/Runtime/Renderer/Private/SceneCore.h
+++ b/Engine/Source/Runtime/Renderer/Private/SceneCore.h
@@ -168,7 +168,7 @@ private:
 class FStaticMeshBatchRelevance
 {
 public:
-	FStaticMeshBatchRelevance(const FStaticMeshBatch& StaticMesh, float InScreenSize, bool InbSupportsCachingMeshDrawCommands, bool InbUseSkyMaterial)
+	FStaticMeshBatchRelevance(const FStaticMeshBatch& StaticMesh, float InScreenSize, bool InbSupportsCachingMeshDrawCommands, bool InbUseSkyMaterial, bool bInUseSingleLayerWaterMaterial)
 		: Id(StaticMesh.Id)
 		, ScreenSize(InScreenSize)
 		, CommandInfosBase(0)
@@ -182,6 +182,7 @@ public:
 		, bUseForDepthPass(StaticMesh.bUseForDepthPass)
 		, bUseAsOccluder(StaticMesh.bUseAsOccluder)
 		, bUseSkyMaterial(InbUseSkyMaterial)
+		, bUseSingleLayerWaterMaterial(bInUseSingleLayerWaterMaterial)
 		, bRenderToVirtualTexture(StaticMesh.bRenderToVirtualTexture)
 		, RuntimeVirtualTextureMaterialType(StaticMesh.RuntimeVirtualTextureMaterialType)
 		, bSupportsCachingMeshDrawCommands(InbSupportsCachingMeshDrawCommands)
@@ -220,6 +221,7 @@ public:
 	uint8 bUseForDepthPass	: 1; // Whether it can be used in depth pass.
 	uint8 bUseAsOccluder	: 1; // User hint whether it's a good occluder.
 	uint8 bUseSkyMaterial	: 1; // Whether this batch uses a Sky material or not.
+	uint8 bUseSingleLayerWaterMaterial : 1; // Whether this batch uses a water material or not.
 
 	/** Whether the mesh batch can be used for rendering to a virtual texture. */
 	uint8 bRenderToVirtualTexture : 1;
diff --git a/Engine/Source/Runtime/Renderer/Private/SceneRendering.cpp b/Engine/Source/Runtime/Renderer/Private/SceneRendering.cpp
--- a/Engine/Source/Runtime/Renderer/Private/SceneRendering.cpp
+++ b/Engine/Source/Runtime/Renderer/Private/SceneRendering.cpp
@@ -55,6 +55,7 @@
 #include "IXRCamera.h"
 #include "IXRCamera.h"
 #include "DiaphragmDOF.h" 
+#include "SingleLayerWaterRendering.h" 
 
 /*-----------------------------------------------------------------------------
 	Globals
@@ -813,6 +814,7 @@ void FViewInfo::Init()
 	bDisableDistanceBasedFadeTransitions = false;	
 	ShadingModelMaskInView = 0;
 	bSceneHasSkyMaterial = 0;
+	bHasSingleLayerWaterMaterial = 0;
 
 	NumVisibleStaticMeshElements = 0;
 	PrecomputedVisibilityData = 0;
@@ -1133,7 +1135,8 @@ void FViewInfo::SetupUniformBufferParameters(
 	{
 		if (Scene->SimpleDirectionalLight)
 		{
-			ViewUniformShaderParameters.DirectionalLightColor = Scene->SimpleDirectionalLight->Proxy->GetColor() / PI;
+			
+			ViewUniformShaderParameters.DirectionalLightColor = Scene->SimpleDirectionalLight->Proxy->GetTransmittanceFactor() * Scene->SimpleDirectionalLight->Proxy->GetColor() / PI;
 			ViewUniformShaderParameters.DirectionalLightDirection = -Scene->SimpleDirectionalLight->Proxy->GetDirection();
 		}
 		else
@@ -2639,8 +2642,12 @@ void FSceneRenderer::RenderFinish(FRHICommandListImmediate& RHICmdList)
 				bMobileShowVertexFogWarning = true;
 			}
 		}
+
+		const bool bSingleLayerWaterWarning = ShouldRenderSingleLayerWaterSkippedRenderEditorNotification(Views);
 		
-		const bool bAnyWarning = bShowPrecomputedVisibilityWarning || bShowGlobalClipPlaneWarning || bShowAtmosphericFogWarning || bShowSkylightWarning || bShowPointLightWarning || bShowDFAODisabledWarning || bShowShadowedLightOverflowWarning || bShowMobileDynamicCSMWarning || bShowMobileLowQualityLightmapWarning || bShowMobileMovableDirectionalLightWarning || bMobileShowVertexFogWarning || bShowSkinCacheOOM;
+		const bool bAnyWarning = bShowPrecomputedVisibilityWarning || bShowGlobalClipPlaneWarning || bShowAtmosphericFogWarning || bShowSkylightWarning || bShowPointLightWarning 
+			|| bShowDFAODisabledWarning || bShowShadowedLightOverflowWarning || bShowMobileDynamicCSMWarning || bShowMobileLowQualityLightmapWarning || bShowMobileMovableDirectionalLightWarning
+			|| bMobileShowVertexFogWarning || bShowSkinCacheOOM || bSingleLayerWaterWarning;
 
 		for(int32 ViewIndex = 0;ViewIndex < Views.Num();ViewIndex++)
 		{	
@@ -2768,6 +2775,14 @@ void FSceneRenderer::RenderFinish(FRHICommandListImmediate& RHICmdList)
 						Canvas.DrawShadowedText(10, Y, Message, GetStatsFont(), FLinearColor(0.8, 1.0, 0.2, 1.0));
 						Y += 14;
 					}
+
+					if (bSingleLayerWaterWarning)
+					{
+						static const FText Message = NSLOCTEXT("Renderer", "SingleLayerWater", "r.SingleLayerWater rendering is disabled with a view containing meshe(s) using water material. Meshes are not visible.");
+						Canvas.DrawShadowedText(10, Y, Message, GetStatsFont(), FLinearColor(1.0, 0.05, 0.05, 1.0));
+						Y += 14;
+					}
+
 					Canvas.Flush_RenderThread(RHICmdList);
 				}
 				
diff --git a/Engine/Source/Runtime/Renderer/Private/SceneRendering.h b/Engine/Source/Runtime/Renderer/Private/SceneRendering.h
--- a/Engine/Source/Runtime/Renderer/Private/SceneRendering.h
+++ b/Engine/Source/Runtime/Renderer/Private/SceneRendering.h
@@ -1025,6 +1025,10 @@ public:
 	 * This is used to skip the sky rendering part during the SkyAtmosphere pass on non mobile platforms.
 	 */
 	uint32 bSceneHasSkyMaterial : 1;
+	/**
+	 * true if the scene has at least one mesh with a material tagged as water visible in a view.
+	 */
+	uint32 bHasSingleLayerWaterMaterial : 1;
 	/** Bitmask of all shading models used by primitives in this view */
 	uint16 ShadingModelMaskInView;
 
diff --git a/Engine/Source/Runtime/Renderer/Private/SceneVisibility.cpp b/Engine/Source/Runtime/Renderer/Private/SceneVisibility.cpp
--- a/Engine/Source/Runtime/Renderer/Private/SceneVisibility.cpp
+++ b/Engine/Source/Runtime/Renderer/Private/SceneVisibility.cpp
@@ -1892,6 +1892,7 @@ struct FRelevancePacket
 	bool bTranslucentSurfaceLighting;
 	bool bUsesSceneDepth;
 	bool bSceneHasSkyMaterial;
+	bool bHasSingleLayerWaterMaterial;
 
 	FRelevancePacket(
 		FRHICommandListImmediate& InRHICmdList,
@@ -1929,6 +1930,7 @@ struct FRelevancePacket
 		, bTranslucentSurfaceLighting(false)
 		, bUsesSceneDepth(false)
 		, bSceneHasSkyMaterial(false)
+		, bHasSingleLayerWaterMaterial(false)
 	{
 	}
 
@@ -1942,6 +1944,7 @@ struct FRelevancePacket
 	{
 		CombinedShadingModelMask = 0;
 		bSceneHasSkyMaterial = 0;
+		bHasSingleLayerWaterMaterial = 0;
 		bUsesGlobalDistanceField = false;
 		bUsesLightingChannels = false;
 		bTranslucentSurfaceLighting = false;
@@ -2040,6 +2043,7 @@ struct FRelevancePacket
 			bTranslucentSurfaceLighting |= ViewRelevance.bTranslucentSurfaceLighting;
 			bUsesSceneDepth |= ViewRelevance.bUsesSceneDepth;
 			bSceneHasSkyMaterial |= ViewRelevance.bUsesSkyMaterial;
+			bHasSingleLayerWaterMaterial |= ViewRelevance.bUsesSingleLayerWaterMaterial;
 
 			if (ViewRelevance.bRenderCustomDepth)
 			{
@@ -2238,6 +2242,11 @@ struct FRelevancePacket
 									// Not needed on Mobile path as in this case everything goes into the regular base pass
 									DrawCommandPacket.AddCommandsForMesh(PrimitiveIndex, PrimitiveSceneInfo, StaticMeshRelevance, StaticMesh, Scene, bCanCache, EMeshPass::SkyPass);
 								}
+								else if(StaticMeshRelevance.bUseSingleLayerWaterMaterial)
+								{
+									// Not needed on Mobile path as in this case everything goes into the regular base pass
+									DrawCommandPacket.AddCommandsForMesh(PrimitiveIndex, PrimitiveSceneInfo, StaticMeshRelevance, StaticMesh, Scene, bCanCache, EMeshPass::SingleLayerWaterPass);
+								}
 
 								if (ViewRelevance.bRenderCustomDepth)
 								{
@@ -2375,6 +2384,7 @@ struct FRelevancePacket
 		WriteView.bTranslucentSurfaceLighting |= bTranslucentSurfaceLighting;
 		WriteView.bUsesSceneDepth |= bUsesSceneDepth;
 		WriteView.bSceneHasSkyMaterial |= bSceneHasSkyMaterial;
+		WriteView.bHasSingleLayerWaterMaterial |= bHasSingleLayerWaterMaterial;
 		VisibleDynamicPrimitivesWithSimpleLights.AppendTo(WriteView.VisibleDynamicPrimitivesWithSimpleLights);
 		WriteView.NumVisibleDynamicPrimitives += NumVisibleDynamicPrimitives;
 		WriteView.NumVisibleDynamicEditorPrimitives += NumVisibleDynamicEditorPrimitives;
diff --git a/Engine/Source/Runtime/Renderer/Private/ShadowDepthRendering.cpp b/Engine/Source/Runtime/Renderer/Private/ShadowDepthRendering.cpp
--- a/Engine/Source/Runtime/Renderer/Private/ShadowDepthRendering.cpp
+++ b/Engine/Source/Runtime/Renderer/Private/ShadowDepthRendering.cpp
@@ -1854,7 +1854,8 @@ void FShadowDepthPassMeshProcessor::AddMeshBatch(const FMeshBatch& RESTRICT Mesh
 
 		if ((bShouldCastShadow || (bReflectiveShadowmap && (Material.ShouldInjectEmissiveIntoLPV() || Material.ShouldBlockGI())))
 			&& ShouldIncludeDomainInMeshPass(Material.GetMaterialDomain())
-			&& !Material.IsSky())
+			&& !Material.IsSky()
+			&& !Material.MaterialUsesSingleLayerWater_RenderThread())
 		{
 			const FMaterialRenderProxy* EffectiveMaterialRenderProxy = &MaterialRenderProxy;
 			const FMaterial* EffectiveMaterial = &Material;
diff --git a/Engine/Source/Runtime/Renderer/Private/SingleLayerWaterRendering.cpp b/Engine/Source/Runtime/Renderer/Private/SingleLayerWaterRendering.cpp
new file mode 100644
--- /dev/null
+++ b/Engine/Source/Runtime/Renderer/Private/SingleLayerWaterRendering.cpp
@@ -0,0 +1,415 @@
+// Copyright 1998-2019 Epic Games, Inc. All Rights Reserved.
+
+/*=============================================================================
+	SingleLayerWaterRendering.cpp: Water pass rendering implementation.
+=============================================================================*/
+
+#include "SingleLayerWaterRendering.h"
+#include "BasePassRendering.h"
+#include "DeferredShadingRenderer.h"
+#include "MeshPassProcessor.inl"
+#include "PixelShaderUtils.h"
+#include "PostProcess/PostProcessSubsurface.h"
+#include "PostProcessTemporalAA.h"
+#include "RenderGraph.h"
+#include "ScenePrivate.h"
+#include "SceneRendering.h"
+#include "ScreenSpaceRayTracing.h"
+#include "SceneTextureParameters.h"
+
+
+
+DECLARE_GPU_STAT(SingleLayerWater);
+
+
+
+static TAutoConsoleVariable<int32> CVarSingleLayerWater(
+	TEXT("r.SingleLayerWater"), 0,
+	TEXT("Enable the single water renderring system."),
+	ECVF_RenderThreadSafe | ECVF_Scalability);
+
+
+static TAutoConsoleVariable<int32> CVarSingleLayerWaterSSR(
+	TEXT("r.SingleLayerWater.SSR"), 1,
+	TEXT("Enable SSR for the single water renderring system."),
+	ECVF_RenderThreadSafe | ECVF_Scalability);
+
+
+static TAutoConsoleVariable<int32> CVarSingleLayerWaterSSRTAA(
+	TEXT("r.SingleLayerWater.SSRTAA"), 1,
+	TEXT("Enable SSR denoising using TAA for the single water renderring system."),
+	ECVF_RenderThreadSafe | ECVF_Scalability);
+
+
+
+//////////////////////////////////////////////////////////////////////////
+
+
+
+static bool ShouldRenderSingleLayerWater(const FViewInfo& View)
+{
+	return View.bHasSingleLayerWaterMaterial;
+}
+
+bool ShouldRenderSingleLayerWater(const TArray<FViewInfo>& Views, const FEngineShowFlags& EngineShowFlags)
+{
+	if (CVarSingleLayerWater.GetValueOnRenderThread() > 0) // && EngineShowFlags.Water)
+	{
+		for (const FViewInfo& View : Views)
+		{
+			if (ShouldRenderSingleLayerWater(View))
+			{
+				return true;
+			}
+		}
+	}
+	return false;
+}
+
+bool ShouldRenderSingleLayerWaterSkippedRenderEditorNotification(const TArray<FViewInfo>& Views)
+{
+	if (CVarSingleLayerWater.GetValueOnRenderThread() <= 0)
+	{
+		for (const FViewInfo& View : Views)
+		{
+			if (ShouldRenderSingleLayerWater(View))
+			{
+				return true;
+			}
+		}
+	}
+	return false;
+}
+
+
+
+//////////////////////////////////////////////////////////////////////////
+
+
+
+FSingleLayerWaterPassMeshProcessor::FSingleLayerWaterPassMeshProcessor(const FScene* Scene, const FSceneView* InViewIfDynamicMeshCommand, const FMeshPassProcessorRenderState& InPassDrawRenderState, FMeshPassDrawListContext* InDrawListContext)
+	: FMeshPassProcessor(Scene, Scene->GetFeatureLevel(), InViewIfDynamicMeshCommand, InDrawListContext)
+	, PassDrawRenderState(InPassDrawRenderState)
+{
+	// Not needed for now
+	//FRHIBlendState* WaterBlendState = TStaticBlendState<CW_RGBA, BO_Add, BF_One, BF_InverseSourceAlpha>::GetRHI();
+	//PassDrawRenderState.SetBlendState(WaterBlendState);
+}
+
+void FSingleLayerWaterPassMeshProcessor::AddMeshBatch(const FMeshBatch& RESTRICT MeshBatch, uint64 BatchElementMask, const FPrimitiveSceneProxy* RESTRICT PrimitiveSceneProxy, int32 StaticMeshId)
+{
+	// Determine the mesh's material and blend mode.
+	const FMaterialRenderProxy* FallbackMaterialRenderProxyPtr = nullptr;
+	const FMaterial& Material = MeshBatch.MaterialRenderProxy->GetMaterialWithFallback(FeatureLevel, FallbackMaterialRenderProxyPtr);
+
+	if (Material.MaterialUsesSingleLayerWater_RenderThread())
+	{
+		const ERasterizerFillMode MeshFillMode = ComputeMeshFillMode(MeshBatch, Material);
+		const ERasterizerCullMode MeshCullMode = ComputeMeshCullMode(MeshBatch, Material);
+		const FMaterialRenderProxy& MaterialRenderProxy = FallbackMaterialRenderProxyPtr ? *FallbackMaterialRenderProxyPtr : *MeshBatch.MaterialRenderProxy;
+		Process(MeshBatch, BatchElementMask, PrimitiveSceneProxy, StaticMeshId, MaterialRenderProxy, Material, MeshFillMode, MeshCullMode);
+	}
+}
+
+void FSingleLayerWaterPassMeshProcessor::Process(
+	const FMeshBatch& MeshBatch,
+	uint64 BatchElementMask,
+	const FPrimitiveSceneProxy* RESTRICT PrimitiveSceneProxy,
+	int32 StaticMeshId,
+	const FMaterialRenderProxy& RESTRICT MaterialRenderProxy,
+	const FMaterial& RESTRICT MaterialResource,
+	ERasterizerFillMode MeshFillMode,
+	ERasterizerCullMode MeshCullMode)
+{
+	FUniformLightMapPolicy NoLightmapPolicy(LMP_NO_LIGHTMAP);
+	typedef FUniformLightMapPolicy LightMapPolicyType;
+	TMeshProcessorShaders<
+		TBasePassVertexShaderPolicyParamType<LightMapPolicyType>,
+		FBaseHS,
+		FBaseDS,
+		TBasePassPixelShaderPolicyParamType<LightMapPolicyType>> WaterPassShaders;
+
+	const FVertexFactory* VertexFactory = MeshBatch.VertexFactory;
+	const bool bRenderSkylight = true;
+	const bool bRenderAtmosphericFog = false;
+	GetBasePassShaders<LightMapPolicyType>(
+		MaterialResource,
+		VertexFactory->GetType(),
+		NoLightmapPolicy,
+		FeatureLevel,
+		bRenderAtmosphericFog,
+		bRenderSkylight,
+		WaterPassShaders.HullShader,
+		WaterPassShaders.DomainShader,
+		WaterPassShaders.VertexShader,
+		WaterPassShaders.PixelShader
+		);
+
+	TBasePassShaderElementData<LightMapPolicyType> ShaderElementData(nullptr);
+	ShaderElementData.InitializeMeshMaterialData(ViewIfDynamicMeshCommand, PrimitiveSceneProxy, MeshBatch, StaticMeshId, false);
+
+	const FMeshDrawCommandSortKey SortKey = CalculateMeshStaticSortKey(WaterPassShaders.VertexShader, WaterPassShaders.PixelShader);
+
+	BuildMeshDrawCommands(
+		MeshBatch,
+		BatchElementMask,
+		PrimitiveSceneProxy,
+		MaterialRenderProxy,
+		MaterialResource,
+		PassDrawRenderState,
+		WaterPassShaders,
+		MeshFillMode,
+		MeshCullMode,
+		SortKey,
+		EMeshPassFeatures::Default,
+		ShaderElementData);
+}
+
+FMeshPassProcessor* CreateSingleLayerWaterPassProcessor(const FScene* Scene, const FSceneView* InViewIfDynamicMeshCommand, FMeshPassDrawListContext* InDrawListContext)
+{
+	FMeshPassProcessorRenderState DrawRenderState(Scene->UniformBuffers.ViewUniformBuffer, Scene->UniformBuffers.OpaqueBasePassUniformBuffer);
+	DrawRenderState.SetInstancedViewUniformBuffer(Scene->UniformBuffers.InstancedViewUniformBuffer);
+
+	// Make sure depth write is enabled.
+	FExclusiveDepthStencil::Type BasePassDepthStencilAccess_DepthWrite = FExclusiveDepthStencil::Type(Scene->DefaultBasePassDepthStencilAccess | FExclusiveDepthStencil::DepthWrite);
+	SetupBasePassState(BasePassDepthStencilAccess_DepthWrite, false, DrawRenderState);
+
+	return new(FMemStack::Get()) FSingleLayerWaterPassMeshProcessor(Scene, InViewIfDynamicMeshCommand, DrawRenderState, InDrawListContext);
+}
+
+FRegisterPassProcessorCreateFunction RegisterSingleLayerWaterPass(&CreateSingleLayerWaterPassProcessor, EShadingPath::Deferred, EMeshPass::SingleLayerWaterPass, EMeshPassFlags::MainView);
+
+
+
+//////////////////////////////////////////////////////////////////////////
+
+
+
+class FSingleLayerWaterCompositePS : public FGlobalShader
+{
+	DECLARE_GLOBAL_SHADER(FSingleLayerWaterCompositePS);
+	SHADER_USE_PARAMETER_STRUCT(FSingleLayerWaterCompositePS, FGlobalShader)
+
+	class FScreenSpaceReflections : SHADER_PERMUTATION_BOOL("SCREEN_SPACE_REFLECTION");
+	using FPermutationDomain = TShaderPermutationDomain<FScreenSpaceReflections>;
+
+
+	BEGIN_SHADER_PARAMETER_STRUCT(FParameters, )
+		SHADER_PARAMETER_RDG_TEXTURE(Texture2D, ScreenSpaceReflectionsTexture)
+		SHADER_PARAMETER_SAMPLER(SamplerState, ScreenSpaceReflectionsSampler)
+		SHADER_PARAMETER_TEXTURE(Texture2D, PreIntegratedGF)
+		SHADER_PARAMETER_SAMPLER(SamplerState, PreIntegratedGFSampler)
+		SHADER_PARAMETER_RDG_TEXTURE(Texture2D, SceneNoWaterDepthTexture)
+		SHADER_PARAMETER_SAMPLER(SamplerState,  SceneNoWaterDepthSampler)
+		SHADER_PARAMETER_STRUCT_INCLUDE(FSceneTextureParameters, SceneTextures)	// Water scene texture
+		SHADER_PARAMETER_STRUCT_INCLUDE(FSceneTextureSamplerParameters, SceneTextureSamplers)
+		SHADER_PARAMETER_STRUCT_REF(FViewUniformShaderParameters, ViewUniformBuffer)
+		SHADER_PARAMETER_STRUCT_REF(FReflectionCaptureShaderData, ReflectionCaptureData)
+		SHADER_PARAMETER_STRUCT_REF(FReflectionUniformParameters, ReflectionsParameters)
+		RENDER_TARGET_BINDING_SLOTS()
+	END_SHADER_PARAMETER_STRUCT()
+
+	static FPermutationDomain RemapPermutation(FPermutationDomain PermutationVector)
+	{
+		return PermutationVector;
+	}
+
+	static bool ShouldCompilePermutation(const FGlobalShaderPermutationParameters& Parameters)
+	{
+		if (!IsFeatureLevelSupported(Parameters.Platform, ERHIFeatureLevel::SM4))
+		{
+			return false;
+		}
+		return true;
+	}
+
+	static void ModifyCompilationEnvironment(const FGlobalShaderPermutationParameters& Parameters, FShaderCompilerEnvironment& OutEnvironment)
+	{
+		FGlobalShader::ModifyCompilationEnvironment(Parameters, OutEnvironment);
+	}
+
+};
+
+IMPLEMENT_GLOBAL_SHADER(FSingleLayerWaterCompositePS, "/Engine/Private/SingleLayerWaterComposite.usf", "SingleLayerWaterCompositePS", SF_Pixel);
+
+void FDeferredShadingSceneRenderer::RenderSingleLayerWaterSSR(FRHICommandListImmediate& RHICmdList)
+{
+	if (CVarSingleLayerWater.GetValueOnRenderThread() <= 0)
+	{
+		return;
+	}
+
+	FRDGBuilder GraphBuilder(RHICmdList);
+	const FSceneRenderTargets& SceneContext = FSceneRenderTargets::Get(GraphBuilder.RHICmdList);
+	FRDGTextureRef SceneColorTexture = GraphBuilder.RegisterExternalTexture(SceneContext.GetSceneColor());
+
+	for (int32 ViewIndex = 0; ViewIndex < Views.Num(); ViewIndex++)
+	{
+		FViewInfo& View = Views[ViewIndex];
+
+		FRDGTextureRef ReflectionsColor = nullptr;
+
+		if (!ShouldRenderScreenSpaceReflections(View))
+		{
+			continue;
+		}
+
+		FSceneTextureParameters SceneTextures;
+		SetupSceneTextureParameters(GraphBuilder, &SceneTextures);
+
+		{
+			const bool bEnableSSR = CVarSingleLayerWaterSSR.GetValueOnRenderThread();
+			if (bEnableSSR)
+			{
+				// RUN SSR
+				// Uses the water GBuffer (depth, ABCDEF) to know how to start tracing.
+				// The water scene depth is used to know where to start tracing.
+				// Then it uses the scene HZB for the ray casting process.
+
+				FRDGTextureRef CurrentSceneColor = GraphBuilder.RegisterExternalTexture(SceneContext.GetSceneColor());
+
+				IScreenSpaceDenoiser::FReflectionsInputs DenoiserInputs;
+				IScreenSpaceDenoiser::FReflectionsRayTracingConfig RayTracingConfig;
+				ESSRQuality SSRQuality;
+				GetSSRQualityForView(View, &SSRQuality, &RayTracingConfig);
+
+				RDG_EVENT_SCOPE(GraphBuilder, "Water ScreenSpaceReflections(Quality=%d)", int32(SSRQuality));
+
+				const bool bDenoise = false;
+				RenderScreenSpaceReflections(
+					GraphBuilder, SceneTextures, CurrentSceneColor, View, SSRQuality, bDenoise, &DenoiserInputs);
+
+				ReflectionsColor = DenoiserInputs.Color;
+
+				if (CVarSingleLayerWaterSSRTAA.GetValueOnRenderThread())
+				{
+					check(View.ViewState);
+					FTAAPassParameters TAASettings(View);
+					TAASettings.Pass = ETAAPassConfig::ScreenSpaceReflections;
+					TAASettings.SceneColorInput = DenoiserInputs.Color;
+
+					FTAAOutputs TAAOutputs = TAASettings.AddTemporalAAPass(
+						GraphBuilder,
+						SceneTextures, View,
+						View.PrevViewInfo.SSRHistory,
+						&View.ViewState->PrevFrameViewInfo.SSRHistory);
+
+					ReflectionsColor = TAAOutputs.SceneColor;
+				}
+			}
+
+			// COMPOSITE
+			{
+				FSingleLayerWaterCompositePS::FPermutationDomain PermutationVector;
+				PermutationVector.Set<FSingleLayerWaterCompositePS::FScreenSpaceReflections>(bEnableSSR);
+				TShaderMapRef<FSingleLayerWaterCompositePS> PixelShader(View.ShaderMap, PermutationVector);
+
+				FSingleLayerWaterCompositePS::FParameters* PassParameters = GraphBuilder.AllocParameters<FSingleLayerWaterCompositePS::FParameters>();
+				PassParameters->ScreenSpaceReflectionsTexture = ReflectionsColor ? ReflectionsColor : GraphBuilder.RegisterExternalTexture(GSystemTextures.BlackDummy);
+				PassParameters->ScreenSpaceReflectionsSampler = TStaticSamplerState<SF_Point>::GetRHI();
+				PassParameters->PreIntegratedGF = GSystemTextures.PreintegratedGF->GetRenderTargetItem().ShaderResourceTexture;
+				PassParameters->PreIntegratedGFSampler = TStaticSamplerState<SF_Bilinear, AM_Clamp, AM_Clamp, AM_Clamp>::GetRHI();
+				PassParameters->SceneNoWaterDepthTexture = GraphBuilder.RegisterExternalTexture(SceneContext.SceneDepthZWithoutSingleLayerWater ? SceneContext.SceneDepthZWithoutSingleLayerWater : GSystemTextures.BlackDummy);
+				PassParameters->SceneNoWaterDepthSampler = TStaticSamplerState<SF_Point>::GetRHI();
+				PassParameters->SceneTextures = SceneTextures;
+				SetupSceneTextureSamplers(&PassParameters->SceneTextureSamplers);
+				PassParameters->ViewUniformBuffer = View.ViewUniformBuffer;
+				PassParameters->ReflectionCaptureData = View.ReflectionCaptureUniformBuffer;
+				{
+					FReflectionUniformParameters ReflectionUniformParameters;
+					SetupReflectionUniformParameters(View, ReflectionUniformParameters);
+					PassParameters->ReflectionsParameters = CreateUniformBufferImmediate(ReflectionUniformParameters, UniformBuffer_SingleDraw);
+				}
+				PassParameters->RenderTargets[0] = FRenderTargetBinding(SceneColorTexture, ERenderTargetLoadAction::ELoad, ERenderTargetStoreAction::EStore);
+				ClearUnusedGraphResources(*PixelShader, PassParameters);
+
+
+				GraphBuilder.AddPass(
+					RDG_EVENT_NAME("WaterComposite %dx%d", View.ViewRect.Width(), View.ViewRect.Height()),
+					PassParameters,
+					ERDGPassFlags::Raster,
+					[PassParameters, &View, PixelShader](FRHICommandList& InRHICmdList)
+				{
+					InRHICmdList.SetViewport(View.ViewRect.Min.X, View.ViewRect.Min.Y, 0.0f, View.ViewRect.Max.X, View.ViewRect.Max.Y, 1.0f);
+
+					FGraphicsPipelineStateInitializer GraphicsPSOInit;
+					FPixelShaderUtils::InitFullscreenPipelineState(InRHICmdList, View.ShaderMap, *PixelShader, GraphicsPSOInit);
+
+					// Premultiplied alpha where alpha is transmittance.
+					GraphicsPSOInit.BlendState = TStaticBlendState<CW_RGB, BO_Add, BF_One, BF_SourceAlpha>::GetRHI(); 
+
+					SetGraphicsPipelineState(InRHICmdList, GraphicsPSOInit);
+					SetShaderParameters(InRHICmdList, *PixelShader, PixelShader->GetPixelShader(), *PassParameters);
+					FPixelShaderUtils::DrawFullscreenTriangle(InRHICmdList);
+				});
+			}
+		}
+
+	}
+
+	TRefCountPtr<IPooledRenderTarget> OutSceneColor;
+	GraphBuilder.QueueTextureExtraction(SceneColorTexture, &OutSceneColor);		// Should not be needed...
+
+	GraphBuilder.Execute();
+
+	ResolveSceneColor(RHICmdList);
+}
+
+bool FDeferredShadingSceneRenderer::RenderSingleLayerWaterPass(FRHICommandListImmediate& RHICmdList, FExclusiveDepthStencil::Type WaterPassDepthStencilAccess)
+{
+	CSV_SCOPED_TIMING_STAT_EXCLUSIVE(RenderSingleLayerWaterPass);
+	SCOPED_NAMED_EVENT(FDeferredShadingSceneRenderer_RenderSingleLayerWaterPass, FColor::Emerald);
+
+	bool bDirty = false;
+
+	RHICmdList.AutomaticCacheFlushAfterComputeShader(false);
+
+	{
+		SCOPED_DRAW_EVENT(RHICmdList, SingleLayerWater);
+		SCOPE_CYCLE_COUNTER(STAT_WaterPassDrawTime);
+		SCOPED_GPU_STAT(RHICmdList, SingleLayerWater);
+
+		// Must have an open renderpass before getting here in single threaded mode.
+		check(RHICmdList.IsInsideRenderPass());
+
+		for (int32 ViewIndex = 0; ViewIndex < Views.Num(); ViewIndex++)
+		{
+			SCOPED_CONDITIONAL_DRAW_EVENTF(RHICmdList, EventView, Views.Num() > 1, TEXT("View%d"), ViewIndex);
+			FViewInfo& View = Views[ViewIndex];
+			SCOPED_GPU_MASK(RHICmdList, View.GPUMask);
+
+			TUniformBufferRef<FOpaqueBasePassUniformParameters> WaterPassUniformBuffer;
+			IPooledRenderTarget* WhiteForwardScreenSpaceShadowMask = GSystemTextures.WhiteDummy;
+			CreateOpaqueBasePassUniformBuffer(RHICmdList, View, WhiteForwardScreenSpaceShadowMask, WaterPassUniformBuffer);
+
+			FMeshPassProcessorRenderState DrawRenderState(View, WaterPassUniformBuffer);
+			SetupBasePassState(WaterPassDepthStencilAccess, ViewFamily.EngineShowFlags.ShaderComplexity, DrawRenderState);
+
+			const bool bShouldRenderView = View.ShouldRenderView();
+			if (bShouldRenderView)
+			{
+				Scene->UniformBuffers.UpdateViewUniformBuffer(View);
+
+				bDirty |= RenderSingleLayerWaterPassView(RHICmdList, View, DrawRenderState);
+			}
+		}
+	}
+
+	RHICmdList.AutomaticCacheFlushAfterComputeShader(true);
+	RHICmdList.FlushComputeShaderCache();
+
+	return bDirty;
+}
+
+bool FDeferredShadingSceneRenderer::RenderSingleLayerWaterPassView(FRHICommandListImmediate& RHICmdList, FViewInfo& View, const FMeshPassProcessorRenderState& InDrawRenderState)
+{
+	bool bDirty = false;
+	FMeshPassProcessorRenderState DrawRenderState(InDrawRenderState);
+	SetupBasePassView(RHICmdList, View, this);
+
+	View.ParallelMeshDrawCommandPasses[EMeshPass::SingleLayerWaterPass].DispatchDraw(nullptr, RHICmdList);
+
+	return bDirty;
+}
+
+
diff --git a/Engine/Source/Runtime/Renderer/Private/SingleLayerWaterRendering.h b/Engine/Source/Runtime/Renderer/Private/SingleLayerWaterRendering.h
new file mode 100644
--- /dev/null
+++ b/Engine/Source/Runtime/Renderer/Private/SingleLayerWaterRendering.h
@@ -0,0 +1,47 @@
+// Copyright 1998-2019 Epic Games, Inc. All Rights Reserved.
+
+/*=============================================================================
+	SingleLayerWaterRendering.h: Water pass rendering implementation.
+=============================================================================*/
+
+#pragma once
+
+#include "CoreMinimal.h"
+#include "RendererInterface.h"
+#include "MeshPassProcessor.h"
+#include "Containers\Array.h"
+
+
+
+class FViewInfo;
+
+
+
+class FSingleLayerWaterPassMeshProcessor : public FMeshPassProcessor
+{
+public:
+
+	FSingleLayerWaterPassMeshProcessor(const FScene* Scene, const FSceneView* InViewIfDynamicMeshCommand, const FMeshPassProcessorRenderState& InPassDrawRenderState, FMeshPassDrawListContext* InDrawListContext);
+
+	virtual void AddMeshBatch(const FMeshBatch& RESTRICT MeshBatch, uint64 BatchElementMask, const FPrimitiveSceneProxy* RESTRICT PrimitiveSceneProxy, int32 StaticMeshId = -1) override final;
+
+	FMeshPassProcessorRenderState PassDrawRenderState;
+
+private:
+	void Process(
+		const FMeshBatch& MeshBatch,
+		uint64 BatchElementMask,
+		const FPrimitiveSceneProxy* RESTRICT PrimitiveSceneProxy,
+		int32 StaticMeshId,
+		const FMaterialRenderProxy& RESTRICT MaterialRenderProxy,
+		const FMaterial& RESTRICT MaterialResource,
+		ERasterizerFillMode MeshFillMode,
+		ERasterizerCullMode MeshCullMode);
+};
+
+
+
+bool ShouldRenderSingleLayerWater(const TArray<FViewInfo>& Views, const FEngineShowFlags& EngineShowFlags);
+bool ShouldRenderSingleLayerWaterSkippedRenderEditorNotification(const TArray<FViewInfo>& Views);
+
+
diff --git a/Engine/Source/Runtime/Renderer/Private/VelocityRendering.cpp b/Engine/Source/Runtime/Renderer/Private/VelocityRendering.cpp
--- a/Engine/Source/Runtime/Renderer/Private/VelocityRendering.cpp
+++ b/Engine/Source/Runtime/Renderer/Private/VelocityRendering.cpp
@@ -573,7 +573,8 @@ void FVelocityMeshProcessor::AddMeshBatch(const FMeshBatch& RESTRICT MeshBatch,
 		const ERasterizerCullMode MeshCullMode = ComputeMeshCullMode(MeshBatch, *Material);
 
 		if ((BlendMode == BLEND_Opaque || BlendMode == BLEND_Masked)
-			&& !Material->IsSky())
+			&& !Material->IsSky()
+			&& !Material->MaterialUsesSingleLayerWater_RenderThread()) // No velocity for water as of now
 		{
 			if (Material->WritesEveryPixel() && !Material->IsTwoSided() && !Material->MaterialModifiesMeshPosition_RenderThread())
 			{
diff --git a/Engine/Source/Runtime/Renderer/Public/MeshPassProcessor.h b/Engine/Source/Runtime/Renderer/Public/MeshPassProcessor.h
--- a/Engine/Source/Runtime/Renderer/Public/MeshPassProcessor.h
+++ b/Engine/Source/Runtime/Renderer/Public/MeshPassProcessor.h
@@ -20,6 +20,7 @@ namespace EMeshPass
 		DepthPass,
 		BasePass,
 		SkyPass,
+		SingleLayerWaterPass,
 		CSMShadowDepth,
 		Distortion,
 		Velocity,
@@ -52,6 +53,7 @@ inline const TCHAR* GetMeshPassName(EMeshPass::Type MeshPass)
 	case EMeshPass::DepthPass: return TEXT("DepthPass");
 	case EMeshPass::BasePass: return TEXT("BasePass");
 	case EMeshPass::SkyPass: return TEXT("SkyPass");
+	case EMeshPass::SingleLayerWaterPass: return TEXT("SingleLayerWaterPass");
 	case EMeshPass::CSMShadowDepth: return TEXT("CSMShadowDepth");
 	case EMeshPass::Distortion: return TEXT("Distortion");
 	case EMeshPass::Velocity: return TEXT("Velocity");
