Made some Opaque pass logic in a common function for all materials

diff --git a/Engine/Source/Runtime/Engine/Private/Materials/HLSLMaterialTranslator.h b/Engine/Source/Runtime/Engine/Private/Materials/HLSLMaterialTranslator.h
--- a/Engine/Source/Runtime/Engine/Private/Materials/HLSLMaterialTranslator.h
+++ b/Engine/Source/Runtime/Engine/Private/Materials/HLSLMaterialTranslator.h
@@ -1160,6 +1160,10 @@ ResourcesString = TEXT("");
 			{
 				Errorf(TEXT("Water materials must be opaque and lit and they only support the DefaultLit shading model."));
 			}
+			if (bNodeGraphIsUsingSingleLayerWaterMaterialOutput && Material->IsDistorted())
+			{
+				Errorf(TEXT("Water materials must not use Unreal built-in distortion."));
+			}
 
 			bool bDBufferAllowed = IsUsingDBuffers(Platform);
 			bool bDBufferBlendMode = IsDBufferDecalBlendMode((EDecalBlendMode)Material->GetDecalBlendMode());
diff --git a/Engine/Source/Runtime/Engine/Public/MaterialShared.h b/Engine/Source/Runtime/Engine/Public/MaterialShared.h
--- a/Engine/Source/Runtime/Engine/Public/MaterialShared.h
+++ b/Engine/Source/Runtime/Engine/Public/MaterialShared.h
@@ -2748,3 +2748,10 @@ bool ReloadMaterialResource(
 	ERHIFeatureLevel::Type FeatureLevel,
 	EMaterialQualityLevel::Type QualityLevel);
 #endif
+
+inline bool ShouldIncludeMaterialInDefaultOpaquePass(const FMaterial& Material)
+{
+	return !Material.IsSky()
+		&& !Material.MaterialUsesSingleLayerWater_RenderThread();
+}
+
diff --git a/Engine/Source/Runtime/Renderer/Private/BasePassRendering.cpp b/Engine/Source/Runtime/Renderer/Private/BasePassRendering.cpp
--- a/Engine/Source/Runtime/Renderer/Private/BasePassRendering.cpp
+++ b/Engine/Source/Runtime/Renderer/Private/BasePassRendering.cpp
@@ -1236,8 +1236,7 @@ void FBasePassMeshProcessor::AddMeshBatch(const FMeshBatch& RESTRICT MeshBatch,
 		if (bShouldDraw
 			&& (!PrimitiveSceneProxy || PrimitiveSceneProxy->ShouldRenderInMainPass())
 			&& ShouldIncludeDomainInMeshPass(Material.GetMaterialDomain())
-			&& !Material.IsSky()
-			&& !Material.MaterialUsesSingleLayerWater_RenderThread())
+			&& ShouldIncludeMaterialInDefaultOpaquePass(Material))
 		{
 			// Check for a cached light-map.
 			const bool bIsLitMaterial = ShadingModels.IsLit();
diff --git a/Engine/Source/Runtime/Renderer/Private/DepthRendering.cpp b/Engine/Source/Runtime/Renderer/Private/DepthRendering.cpp
--- a/Engine/Source/Runtime/Renderer/Private/DepthRendering.cpp
+++ b/Engine/Source/Runtime/Renderer/Private/DepthRendering.cpp
@@ -784,8 +784,7 @@ void FDepthPassMeshProcessor::AddMeshBatch(const FMeshBatch& RESTRICT MeshBatch,
 		if (!bIsTranslucent
 			&& (!PrimitiveSceneProxy || PrimitiveSceneProxy->ShouldRenderInMainPass())
 			&& ShouldIncludeDomainInMeshPass(Material.GetMaterialDomain())
-			&& !Material.IsSky()
-			&& !Material.MaterialUsesSingleLayerWater_RenderThread())
+			&& ShouldIncludeMaterialInDefaultOpaquePass(Material))
 		{
 			if (BlendMode == BLEND_Opaque
 				&& MeshBatch.VertexFactory->SupportsPositionOnlyStream()
diff --git a/Engine/Source/Runtime/Renderer/Private/ShadowDepthRendering.cpp b/Engine/Source/Runtime/Renderer/Private/ShadowDepthRendering.cpp
--- a/Engine/Source/Runtime/Renderer/Private/ShadowDepthRendering.cpp
+++ b/Engine/Source/Runtime/Renderer/Private/ShadowDepthRendering.cpp
@@ -1854,8 +1854,7 @@ void FShadowDepthPassMeshProcessor::AddMeshBatch(const FMeshBatch& RESTRICT Mesh
 
 		if ((bShouldCastShadow || (bReflectiveShadowmap && (Material.ShouldInjectEmissiveIntoLPV() || Material.ShouldBlockGI())))
 			&& ShouldIncludeDomainInMeshPass(Material.GetMaterialDomain())
-			&& !Material.IsSky()
-			&& !Material.MaterialUsesSingleLayerWater_RenderThread())
+			&& ShouldIncludeMaterialInDefaultOpaquePass(Material))
 		{
 			const FMaterialRenderProxy* EffectiveMaterialRenderProxy = &MaterialRenderProxy;
 			const FMaterial* EffectiveMaterial = &Material;
