Remove water code from scene render target.

diff --git a/Engine/Source/Runtime/Renderer/Private/BasePassRendering.cpp b/Engine/Source/Runtime/Renderer/Private/BasePassRendering.cpp
--- a/Engine/Source/Runtime/Renderer/Private/BasePassRendering.cpp
+++ b/Engine/Source/Runtime/Renderer/Private/BasePassRendering.cpp
@@ -514,7 +514,9 @@ void SetupSharedBasePassParameters(
 void CreateOpaqueBasePassUniformBuffer(
 	FRHICommandListImmediate& RHICmdList, 
 	const FViewInfo& View,
-	IPooledRenderTarget* ForwardScreenSpaceShadowMask, 
+	IPooledRenderTarget* ForwardScreenSpaceShadowMask,
+	IPooledRenderTarget* SceneLuminanceWithoutSingleLayerWater,
+	IPooledRenderTarget* SceneDepthWithoutSingleLayerWaterTexture,
 	TUniformBufferRef<FOpaqueBasePassUniformParameters>& BasePassUniformBuffer)
 {
 	FSceneRenderTargets& SceneRenderTargets = FSceneRenderTargets::Get(RHICmdList);
@@ -579,9 +581,9 @@ void CreateOpaqueBasePassUniformBuffer(
 	}
 
 	// Single Layer Water
-	BasePassParameters.SceneLuminanceWithoutSingleLayerWaterTexture = (SceneRenderTargets.SceneColorWithoutSingleLayerWater.IsValid() ? SceneRenderTargets.SceneColorWithoutSingleLayerWater : GSystemTextures.BlackDummy)->GetRenderTargetItem().TargetableTexture;
+	BasePassParameters.SceneLuminanceWithoutSingleLayerWaterTexture = (SceneLuminanceWithoutSingleLayerWater ? SceneLuminanceWithoutSingleLayerWater->GetRenderTargetItem().TargetableTexture : GSystemTextures.BlackDummy->GetRenderTargetItem().TargetableTexture);
 	BasePassParameters.SceneLuminanceWithoutSingleLayerWaterSampler = TStaticSamplerState<SF_Bilinear>::GetRHI();
-	BasePassParameters.SceneDepthWithoutSingleLayerWaterTexture = (SceneRenderTargets.SceneDepthZWithoutSingleLayerWater.IsValid() ? SceneRenderTargets.SceneDepthZWithoutSingleLayerWater : GSystemTextures.BlackDummy)->GetRenderTargetItem().TargetableTexture;
+	BasePassParameters.SceneDepthWithoutSingleLayerWaterTexture = (SceneDepthWithoutSingleLayerWaterTexture ? SceneDepthWithoutSingleLayerWaterTexture->GetRenderTargetItem().TargetableTexture : GSystemTextures.BlackDummy->GetRenderTargetItem().TargetableTexture);
 	BasePassParameters.SceneDepthWithoutSingleLayerWaterSampler = TStaticSamplerState<SF_Point>::GetRHI();
 	BasePassParameters.PreIntegratedGFTexture = GSystemTextures.PreintegratedGF->GetRenderTargetItem().ShaderResourceTexture;
 	BasePassParameters.PreIntegratedGFSampler = TStaticSamplerState<SF_Bilinear, AM_Clamp, AM_Clamp, AM_Clamp>::GetRHI();
@@ -646,7 +648,7 @@ bool FDeferredShadingSceneRenderer::RenderBasePass(FRHICommandListImmediate& RHI
 				SCOPED_GPU_MASK(RHICmdList, View.GPUMask);
 
 				TUniformBufferRef<FOpaqueBasePassUniformParameters> BasePassUniformBuffer;
-				CreateOpaqueBasePassUniformBuffer(RHICmdList, View, ForwardScreenSpaceShadowMask, BasePassUniformBuffer);
+				CreateOpaqueBasePassUniformBuffer(RHICmdList, View, ForwardScreenSpaceShadowMask, nullptr, nullptr, BasePassUniformBuffer);
 
 				FMeshPassProcessorRenderState DrawRenderState(View, BasePassUniformBuffer);
 
@@ -687,7 +689,7 @@ bool FDeferredShadingSceneRenderer::RenderBasePass(FRHICommandListImmediate& RHI
 				SCOPED_GPU_MASK(RHICmdList, View.GPUMask);
 
 				TUniformBufferRef<FOpaqueBasePassUniformParameters> BasePassUniformBuffer;
-				CreateOpaqueBasePassUniformBuffer(RHICmdList, View, ForwardScreenSpaceShadowMask, BasePassUniformBuffer);
+				CreateOpaqueBasePassUniformBuffer(RHICmdList, View, ForwardScreenSpaceShadowMask, nullptr, nullptr, BasePassUniformBuffer);
 
 				FMeshPassProcessorRenderState DrawRenderState(View, BasePassUniformBuffer);
 
diff --git a/Engine/Source/Runtime/Renderer/Private/BasePassRendering.h b/Engine/Source/Runtime/Renderer/Private/BasePassRendering.h
--- a/Engine/Source/Runtime/Renderer/Private/BasePassRendering.h
+++ b/Engine/Source/Runtime/Renderer/Private/BasePassRendering.h
@@ -117,7 +117,9 @@ extern void SetupSharedBasePassParameters(
 extern void CreateOpaqueBasePassUniformBuffer(
 	FRHICommandListImmediate& RHICmdList, 
 	const FViewInfo& View,
-	IPooledRenderTarget* ForwardScreenSpaceShadowMask, 
+	IPooledRenderTarget* ForwardScreenSpaceShadowMask,
+	IPooledRenderTarget* SceneLuminanceWithoutSingleLayerWater,
+	IPooledRenderTarget* SceneDepthWithoutSingleLayerWaterTexture,
 	TUniformBufferRef<FOpaqueBasePassUniformParameters>& BasePassUniformBuffer);
 
 extern void CreateTranslucentBasePassUniformBuffer(
diff --git a/Engine/Source/Runtime/Renderer/Private/DeferredShadingRenderer.cpp b/Engine/Source/Runtime/Renderer/Private/DeferredShadingRenderer.cpp
--- a/Engine/Source/Runtime/Renderer/Private/DeferredShadingRenderer.cpp
+++ b/Engine/Source/Runtime/Renderer/Private/DeferredShadingRenderer.cpp
@@ -1794,16 +1794,18 @@ void FDeferredShadingSceneRenderer::Render(FRHICommandListImmediate& RHICmdList)
 
 		RHICmdList.TransitionResource(EResourceTransitionAccess::EWritable,  SceneContext.GetSceneDepthSurface());
 
+		FSingleLayerWaterPassData SingleLayerWaterPassData;
+
 		const FExclusiveDepthStencil::Type WaterPassDepthStencilAccess = FExclusiveDepthStencil::DepthWrite_StencilWrite;
-		SceneContext.BeginRenderingWaterGBuffer(RHICmdList, WaterPassDepthStencilAccess, ViewFamily.EngineShowFlags.ShaderComplexity);
-		RenderSingleLayerWaterPass(RHICmdList, WaterPassDepthStencilAccess);
-		SceneContext.FinishWaterGBufferPassAndResolve(RHICmdList);
+		BeginRenderingWaterGBuffer(RHICmdList, SingleLayerWaterPassData, WaterPassDepthStencilAccess, ViewFamily.EngineShowFlags.ShaderComplexity);
+		RenderSingleLayerWaterPass(RHICmdList, SingleLayerWaterPassData, WaterPassDepthStencilAccess);
+		FinishWaterGBufferPassAndResolve(RHICmdList);
 
 		SceneContext.ResolveSceneDepthTexture(RHICmdList, FResolveRect(0, 0, FamilySize.X, FamilySize.Y));
 		SceneContext.ResolveSceneDepthToAuxiliaryTexture(RHICmdList);
 		RHICmdList.TransitionResource(EResourceTransitionAccess::EReadable, SceneContext.GetSceneDepthSurface());
 
-		RenderSingleLayerWaterSSR(RHICmdList);
+		RenderSingleLayerWaterSSR(RHICmdList, SingleLayerWaterPassData);
 		ServiceLocalQueue();
 	}
 
diff --git a/Engine/Source/Runtime/Renderer/Private/DeferredShadingRenderer.h b/Engine/Source/Runtime/Renderer/Private/DeferredShadingRenderer.h
--- a/Engine/Source/Runtime/Renderer/Private/DeferredShadingRenderer.h
+++ b/Engine/Source/Runtime/Renderer/Private/DeferredShadingRenderer.h
@@ -21,11 +21,12 @@ enum class ERayTracingPrimaryRaysFlag : uint32;
 #endif
 
 class FSceneTextureParameters;
-
 class FDistanceFieldAOParameters;
 class UStaticMeshComponent;
 class FExponentialHeightFogSceneInfo;
 
+struct FSingleLayerWaterPassData;
+
 class FLightShaftsOutput
 {
 public:
@@ -102,16 +103,19 @@ public:
 	*/
 	bool RenderBasePass(FRHICommandListImmediate& RHICmdList, FExclusiveDepthStencil::Type BasePassDepthStencilAccess, IPooledRenderTarget* ForwardScreenSpaceShadowMask, bool bParallelBasePass, bool bRenderLightmapDensity);
 
-
+	/** Begin the water GBuffer rendering pass*/
+	void BeginRenderingWaterGBuffer(FRHICommandList& RHICmdList, FSingleLayerWaterPassData& PassData, FExclusiveDepthStencil::Type DepthStencilAccess, bool bBindQuadOverdrawBuffers);
+	/** End the water GBuffer rendering pass*/
+	void FinishWaterGBufferPassAndResolve(FRHICommandListImmediate& RHICmdList);
 	/**
 	* Renders the scene's single layer water base pass
 	* @return true if anything was rendered
 	*/
-	bool RenderSingleLayerWaterPass(FRHICommandListImmediate& RHICmdList, FExclusiveDepthStencil::Type WaterPassDepthStencilAccess);
+	bool RenderSingleLayerWaterPass(FRHICommandListImmediate& RHICmdList, FSingleLayerWaterPassData& PassData, FExclusiveDepthStencil::Type WaterPassDepthStencilAccess);
 	/** Renders the water draw pass for a given View. */
-	bool RenderSingleLayerWaterPassView(FRHICommandListImmediate& RHICmdList, FViewInfo& View, const FMeshPassProcessorRenderState& InDrawRenderState);
+	bool RenderSingleLayerWaterPassView(FRHICommandListImmediate& RHICmdList, FViewInfo& View, FSingleLayerWaterPassData& PassData, const FMeshPassProcessorRenderState& InDrawRenderState);
 	/** Render, denoise and composite the scene SSR and under water effect.*/
-	void RenderSingleLayerWaterSSR(FRHICommandListImmediate& RHICmdList);
+	void RenderSingleLayerWaterSSR(FRHICommandListImmediate& RHICmdList, FSingleLayerWaterPassData& PassData);
 
 	/** Finishes the view family rendering. */
 	void RenderFinish(FRHICommandListImmediate& RHICmdList);
diff --git a/Engine/Source/Runtime/Renderer/Private/PostProcess/PostProcessCompositeEditorPrimitives.cpp b/Engine/Source/Runtime/Renderer/Private/PostProcess/PostProcessCompositeEditorPrimitives.cpp
--- a/Engine/Source/Runtime/Renderer/Private/PostProcess/PostProcessCompositeEditorPrimitives.cpp
+++ b/Engine/Source/Runtime/Renderer/Private/PostProcess/PostProcessCompositeEditorPrimitives.cpp
@@ -565,7 +565,7 @@ void FRCPassPostProcessCompositeEditorPrimitives::Process(FRenderingCompositePas
 
 			if (bDeferredBasePass)
 			{
-				CreateOpaqueBasePassUniformBuffer(Context.RHICmdList, EditorView, nullptr, OpaqueBasePassUniformBuffer);
+				CreateOpaqueBasePassUniformBuffer(Context.RHICmdList, EditorView, nullptr, nullptr, nullptr, OpaqueBasePassUniformBuffer);
 				BasePassUniformBuffer = OpaqueBasePassUniformBuffer;
 			}
 			else
diff --git a/Engine/Source/Runtime/Renderer/Private/PostProcess/SceneRenderTargets.cpp b/Engine/Source/Runtime/Renderer/Private/PostProcess/SceneRenderTargets.cpp
--- a/Engine/Source/Runtime/Renderer/Private/PostProcess/SceneRenderTargets.cpp
+++ b/Engine/Source/Runtime/Renderer/Private/PostProcess/SceneRenderTargets.cpp
@@ -243,8 +243,6 @@ FSceneRenderTargets::FSceneRenderTargets(const FViewInfo& View, const FSceneRend
 	, GBufferC(GRenderTargetPool.MakeSnapshot(SnapshotSource.GBufferC))
 	, GBufferD(GRenderTargetPool.MakeSnapshot(SnapshotSource.GBufferD))
 	, GBufferE(GRenderTargetPool.MakeSnapshot(SnapshotSource.GBufferE))
-	, SceneColorWithoutSingleLayerWater(GRenderTargetPool.MakeSnapshot(SnapshotSource.SceneColorWithoutSingleLayerWater))
-	, SceneDepthZWithoutSingleLayerWater(GRenderTargetPool.MakeSnapshot(SnapshotSource.SceneDepthZWithoutSingleLayerWater))
 	, DBufferA(GRenderTargetPool.MakeSnapshot(SnapshotSource.DBufferA))
 	, DBufferB(GRenderTargetPool.MakeSnapshot(SnapshotSource.DBufferB))
 	, DBufferC(GRenderTargetPool.MakeSnapshot(SnapshotSource.DBufferC))
@@ -1002,70 +1000,6 @@ void FSceneRenderTargets::FinishGBufferPassAndResolve(FRHICommandListImmediate&
 	QuadOverdrawIndex = INDEX_NONE;
 }
 
-void FSceneRenderTargets::BeginRenderingWaterGBuffer(FRHICommandList& RHICmdList, FExclusiveDepthStencil::Type DepthStencilAccess, bool bBindQuadOverdrawBuffers)
-{
-	check(RHICmdList.IsOutsideRenderPass());
-
-	SCOPED_DRAW_EVENT(RHICmdList, BeginRenderingWaterGBuffer);
-	check(CurrentFeatureLevel >= ERHIFeatureLevel::SM4)
-
-	const EShaderPlatform ShaderPlatform = GetFeatureLevelShaderPlatform(CurrentFeatureLevel);
-	const bool bUseGBuffer = IsUsingGBuffers(ShaderPlatform);
-
-	// Allocate required buffers
-	{
-		FPooledRenderTargetDesc Desc(FPooledRenderTargetDesc::Create2DDesc(BufferSize, PF_DepthStencil, DefaultDepthClear, TexCreate_None, TexCreate_DepthStencilTargetable | TexCreate_ShaderResource | TexCreate_InputAttachmentRead, false));
-		Desc.NumSamples = GetNumSceneColorMSAASamples(CurrentFeatureLevel);
-		Desc.Flags |= GFastVRamConfig.SceneDepth;
-		GRenderTargetPool.FindFreeElement(RHICmdList, Desc, SceneDepthZWithoutSingleLayerWater, TEXT("SceneDepthZWithoutSingleLayerWater"), true);
-	}
-	{
-		FPooledRenderTargetDesc Desc(FPooledRenderTargetDesc::Create2DDesc(BufferSize, GetSceneColorFormat() , DefaultColorClear, TexCreate_None, TexCreate_RenderTargetable, false));
-		Desc.NumSamples = GetNumSceneColorMSAASamples(CurrentFeatureLevel);
-		GRenderTargetPool.FindFreeElement(RHICmdList, Desc, SceneColorWithoutSingleLayerWater, TEXT("SceneColorWithoutSingleLayerWater"));
-	}
-
-	// Copy and save textures without water informations for later composition operations
-	{
-		FRHICopyTextureInfo DepthCopyInfo;
-		RHICmdList.CopyTexture(SceneDepthZ->GetRenderTargetItem().ShaderResourceTexture, SceneDepthZWithoutSingleLayerWater->GetRenderTargetItem().ShaderResourceTexture, DepthCopyInfo);
-		RHICmdList.CopyTexture(GetSceneColorSurface(), SceneColorWithoutSingleLayerWater->GetRenderTargetItem().ShaderResourceTexture, DepthCopyInfo);
-	}
-
-	// Create MRT
-	int32 VelocityRTIndex = -1;
-	FRHIRenderPassInfo RPInfo;
-	FillGBufferRenderPassInfo(ERenderTargetLoadAction::ELoad, RPInfo, VelocityRTIndex);
-	// Set a dummy Scene color RT to avoid gbuffer to stomp HDR scene color we want to blend over
-	RPInfo.ColorRenderTargets[0].Action = MakeRenderTargetActions(ERenderTargetLoadAction::ENoAction, ERenderTargetStoreAction::ENoAction);
-	RPInfo.ColorRenderTargets[0].RenderTarget = GetSceneColorSurface();
-
-	// Stencil always has to be store or certain VK drivers will leave the attachment in an undefined state.
-	RPInfo.DepthStencilRenderTarget.Action = MakeDepthStencilTargetActions(MakeRenderTargetActions(ERenderTargetLoadAction::ELoad, ERenderTargetStoreAction::EStore), MakeRenderTargetActions(ERenderTargetLoadAction::ELoad, ERenderTargetStoreAction::EStore));
-	RPInfo.DepthStencilRenderTarget.DepthStencilTarget = (const FTexture2DRHIRef&)SceneDepthZ->GetRenderTargetItem().TargetableTexture;
-	RPInfo.DepthStencilRenderTarget.ExclusiveDepthStencil = DepthStencilAccess;
-
-	// Set other UAVs
-	const bool bClearQuadOverdrawBuffers = false;
-	SetQuadOverdrawUAV(RHICmdList, bBindQuadOverdrawBuffers, bClearQuadOverdrawBuffers, RPInfo);
-	if (UseVirtualTexturing(CurrentFeatureLevel) && !bBindQuadOverdrawBuffers)
-	{
-		BindVirtualTextureFeedbackUAV(RPInfo);
-	}
-
-	// Begin the pass
-	RHICmdList.BeginRenderPass(RPInfo, TEXT("WaterGBuffer"));
-
-	// Needs to be called after we start a renderpass in order for the color/depth decompress/expand to be executed on the next write-to-read barrier/transition.
-	RHICmdList.BindClearMRTValues(true, true, false);
-}
-
-void FSceneRenderTargets::FinishWaterGBufferPassAndResolve(FRHICommandListImmediate& RHICmdList)
-{
-	// Same as the GBuffer for now (also same resolves)
-	FinishGBufferPassAndResolve(RHICmdList);
-}
-
 int32 FSceneRenderTargets::GetNumGBufferTargets() const
 {
 	int32 NumGBufferTargets = 1;
@@ -2762,9 +2696,6 @@ void FSceneRenderTargets::ReleaseAllTargets()
 	MobileCustomStencil.SafeRelease();
 	CustomStencilSRV.SafeRelease();
 
-	SceneColorWithoutSingleLayerWater.SafeRelease();
-	SceneDepthZWithoutSingleLayerWater.SafeRelease();
-
 	for (int32 i = 0; i < ARRAY_COUNT(OptionalShadowDepthColor); i++)
 	{
 		OptionalShadowDepthColor[i].SafeRelease();
diff --git a/Engine/Source/Runtime/Renderer/Private/PostProcess/SceneRenderTargets.h b/Engine/Source/Runtime/Renderer/Private/PostProcess/SceneRenderTargets.h
--- a/Engine/Source/Runtime/Renderer/Private/PostProcess/SceneRenderTargets.h
+++ b/Engine/Source/Runtime/Renderer/Private/PostProcess/SceneRenderTargets.h
@@ -229,9 +229,6 @@ public:
 	void BeginRenderingGBuffer(FRHICommandList& RHICmdList, ERenderTargetLoadAction ColorLoadAction, ERenderTargetLoadAction DepthLoadAction, FExclusiveDepthStencil::Type DepthStencilAccess, bool bBindQuadOverdrawBuffers, bool bClearQuadOverdrawBuffers = false, const FLinearColor& ClearColor = FLinearColor(0, 0, 0, 1), bool bIsWireframe=false);
 	void FinishGBufferPassAndResolve(FRHICommandListImmediate& RHICmdList);
 
-	void BeginRenderingWaterGBuffer(FRHICommandList& RHICmdList, FExclusiveDepthStencil::Type DepthStencilAccess, bool bBindQuadOverdrawBuffers);
-	void FinishWaterGBufferPassAndResolve(FRHICommandListImmediate& RHICmdList);
-
 	/**
 	 * Sets the scene color target and restores its contents if necessary
 	 */
@@ -285,12 +282,17 @@ public:
 		DefaultColorClear = ColorClear;
 	}
 
+	FClearValueBinding GetDefaultColorClear() const
+	{
+		return DefaultColorClear;
+	}
+
 	void SetDefaultDepthClear(const FClearValueBinding DepthClear)
 	{
 		DefaultDepthClear = DepthClear;
 	}
 
-	FClearValueBinding GetDefaultDepthClear()
+	FClearValueBinding GetDefaultDepthClear() const
 	{
 		return DefaultDepthClear;
 	}
@@ -540,10 +542,6 @@ public:
 	TRefCountPtr<IPooledRenderTarget> GBufferD;
 	TRefCountPtr<IPooledRenderTarget> GBufferE;
 
-	// Water rendering extra buffers
-	TRefCountPtr<IPooledRenderTarget> SceneColorWithoutSingleLayerWater;
-	TRefCountPtr<IPooledRenderTarget> SceneDepthZWithoutSingleLayerWater;
-
 	// DBuffer: For decals before base pass (only temporarily available after early z pass and until base pass)
 	TRefCountPtr<IPooledRenderTarget> DBufferA;
 	TRefCountPtr<IPooledRenderTarget> DBufferB;
@@ -638,6 +636,9 @@ public:
 	/** Allocates render targets for use with the deferred shading path. */
 	// Temporarily Public to call from DefferedShaderRenderer to attempt recovery from a crash until cause is found.
 	void AllocateDeferredShadingPathRenderTargets(FRHICommandListImmediate& RHICmdList, const int32 NumViews = 1);
+
+	/** Fills the given FRenderPassInfo with the current GBuffer */
+	int32 FillGBufferRenderPassInfo(ERenderTargetLoadAction ColorLoadAction, FRHIRenderPassInfo& OutRenderPassInfo, int32& OutVelocityRTIndex);
 private:
 
 	/** Allocates render targets for use with the current shading path. */
@@ -682,9 +683,6 @@ private:
 	/** Gets all GBuffers to use.  Returns the number actually used. */
 	int32 GetGBufferRenderTargets(ERenderTargetLoadAction ColorLoadAction, FRHIRenderTargetView OutRenderTargets[MaxSimultaneousRenderTargets], int32& OutVelocityRTIndex);
 
-	/** Fills the given FRenderPassInfo with the current GBuffer */
-	int32 FillGBufferRenderPassInfo(ERenderTargetLoadAction ColorLoadAction, FRHIRenderPassInfo& OutRenderPassInfo, int32& OutVelocityRTIndex);
-
 	ESceneColorFormatType GetSceneColorFormatType() const
 	{
 		if (CurrentShadingPath == EShadingPath::Mobile)
diff --git a/Engine/Source/Runtime/Renderer/Private/Renderer.cpp b/Engine/Source/Runtime/Renderer/Private/Renderer.cpp
--- a/Engine/Source/Runtime/Renderer/Private/Renderer.cpp
+++ b/Engine/Source/Runtime/Renderer/Private/Renderer.cpp
@@ -296,7 +296,7 @@ void FRendererModule::DrawTileMesh(FRHICommandListImmediate& RHICmdList, FMeshPa
 				if (ShadingPath == EShadingPath::Deferred)
 				{
 					TUniformBufferRef<FOpaqueBasePassUniformParameters> OpaqueBasePassUniformBuffer;
-					CreateOpaqueBasePassUniformBuffer(RHICmdList, View, nullptr, OpaqueBasePassUniformBuffer);
+					CreateOpaqueBasePassUniformBuffer(RHICmdList, View, nullptr, nullptr, nullptr, OpaqueBasePassUniformBuffer);
 					BasePassUniformBuffer = OpaqueBasePassUniformBuffer;
 					
 					DrawRenderState.SetPassUniformBuffer(BasePassUniformBuffer);
diff --git a/Engine/Source/Runtime/Renderer/Private/SingleLayerWaterRendering.cpp b/Engine/Source/Runtime/Renderer/Private/SingleLayerWaterRendering.cpp
--- a/Engine/Source/Runtime/Renderer/Private/SingleLayerWaterRendering.cpp
+++ b/Engine/Source/Runtime/Renderer/Private/SingleLayerWaterRendering.cpp
@@ -10,6 +10,7 @@
 #include "MeshPassProcessor.inl"
 #include "PixelShaderUtils.h"
 #include "PostProcess/PostProcessSubsurface.h"
+#include "PostProcess/SceneRenderTargets.h"
 #include "PostProcessTemporalAA.h"
 #include "RenderGraph.h"
 #include "ScenePrivate.h"
@@ -232,7 +233,81 @@ class FSingleLayerWaterCompositePS : public FGlobalShader
 
 IMPLEMENT_GLOBAL_SHADER(FSingleLayerWaterCompositePS, "/Engine/Private/SingleLayerWaterComposite.usf", "SingleLayerWaterCompositePS", SF_Pixel);
 
-void FDeferredShadingSceneRenderer::RenderSingleLayerWaterSSR(FRHICommandListImmediate& RHICmdList)
+
+
+//////////////////////////////////////////////////////////////////////////
+
+
+
+void FDeferredShadingSceneRenderer::BeginRenderingWaterGBuffer(FRHICommandList& RHICmdList, FSingleLayerWaterPassData& PassData, FExclusiveDepthStencil::Type DepthStencilAccess, bool bBindQuadOverdrawBuffers)
+{
+	SCOPED_DRAW_EVENT(RHICmdList, BeginRenderingWaterGBuffer);
+
+	check(RHICmdList.IsOutsideRenderPass());
+	FSceneRenderTargets& SceneContext = FSceneRenderTargets::Get(RHICmdList);
+
+	const ERHIFeatureLevel::Type CurrentFeatureLevel = SceneContext.GetCurrentFeatureLevel();
+	const bool bUseGBuffer = IsUsingGBuffers(ShaderPlatform);
+	check(CurrentFeatureLevel >= ERHIFeatureLevel::SM4);
+
+
+	// Allocate required buffers
+	{
+		FPooledRenderTargetDesc Desc(FPooledRenderTargetDesc::Create2DDesc(SceneContext.GetBufferSizeXY(), PF_DepthStencil, SceneContext.GetDefaultDepthClear(), TexCreate_None, TexCreate_DepthStencilTargetable | TexCreate_ShaderResource | TexCreate_InputAttachmentRead, false));
+		Desc.NumSamples = SceneContext.GetNumSceneColorMSAASamples(CurrentFeatureLevel);
+		Desc.Flags |= GFastVRamConfig.SceneDepth;
+		GRenderTargetPool.FindFreeElement(RHICmdList, Desc, PassData.SceneDepthZWithoutSingleLayerWater, TEXT("SceneDepthZWithoutSingleLayerWater"), true);
+	}
+	{
+		FPooledRenderTargetDesc Desc(FPooledRenderTargetDesc::Create2DDesc(SceneContext.GetBufferSizeXY(), SceneContext.GetSceneColorFormat(), SceneContext.GetDefaultColorClear(), TexCreate_None, TexCreate_RenderTargetable, false));
+		Desc.NumSamples = SceneContext.GetNumSceneColorMSAASamples(CurrentFeatureLevel);
+		GRenderTargetPool.FindFreeElement(RHICmdList, Desc, PassData.SceneColorWithoutSingleLayerWater, TEXT("SceneColorWithoutSingleLayerWater"));
+	}
+
+	// Copy and save textures without water informations for later composition operations
+	{
+		FRHICopyTextureInfo DepthCopyInfo;
+		RHICmdList.CopyTexture(SceneContext.SceneDepthZ->GetRenderTargetItem().ShaderResourceTexture, PassData.SceneDepthZWithoutSingleLayerWater->GetRenderTargetItem().ShaderResourceTexture, DepthCopyInfo);
+		RHICmdList.CopyTexture(SceneContext.GetSceneColorSurface(), PassData.SceneColorWithoutSingleLayerWater->GetRenderTargetItem().ShaderResourceTexture, DepthCopyInfo);
+	}
+
+	// Create MRT
+	int32 VelocityRTIndex = -1;
+	FRHIRenderPassInfo RPInfo;
+	SceneContext.FillGBufferRenderPassInfo(ERenderTargetLoadAction::ELoad, RPInfo, VelocityRTIndex);
+	// Set a dummy Scene color RT to avoid gbuffer to stomp HDR scene color we want to blend over
+	RPInfo.ColorRenderTargets[0].Action = MakeRenderTargetActions(ERenderTargetLoadAction::ENoAction, ERenderTargetStoreAction::ENoAction);
+	RPInfo.ColorRenderTargets[0].RenderTarget = SceneContext.GetSceneColorSurface();
+
+	// Stencil always has to be store or certain VK drivers will leave the attachment in an undefined state.
+	RPInfo.DepthStencilRenderTarget.Action = MakeDepthStencilTargetActions(MakeRenderTargetActions(ERenderTargetLoadAction::ELoad, ERenderTargetStoreAction::EStore), MakeRenderTargetActions(ERenderTargetLoadAction::ELoad, ERenderTargetStoreAction::EStore));
+	RPInfo.DepthStencilRenderTarget.DepthStencilTarget = (const FTexture2DRHIRef&)SceneContext.SceneDepthZ->GetRenderTargetItem().TargetableTexture;
+	RPInfo.DepthStencilRenderTarget.ExclusiveDepthStencil = DepthStencilAccess;
+
+	// Set other UAVs
+	const bool bClearQuadOverdrawBuffers = false;
+	SceneContext.SetQuadOverdrawUAV(RHICmdList, bBindQuadOverdrawBuffers, bClearQuadOverdrawBuffers, RPInfo);
+	if (UseVirtualTexturing(CurrentFeatureLevel) && !bBindQuadOverdrawBuffers)
+	{
+		SceneContext.BindVirtualTextureFeedbackUAV(RPInfo);
+	}
+
+	// Begin the pass
+	RHICmdList.BeginRenderPass(RPInfo, TEXT("WaterGBuffer"));
+
+	// Needs to be called after we start a renderpass in order for the color/depth decompress/expand to be executed on the next write-to-read barrier/transition.
+	RHICmdList.BindClearMRTValues(true, true, false);
+}
+
+void FDeferredShadingSceneRenderer::FinishWaterGBufferPassAndResolve(FRHICommandListImmediate& RHICmdList)
+{
+	// Same as the GBuffer for now (also same resolves)
+	FSceneRenderTargets& SceneContext = FSceneRenderTargets::Get(RHICmdList);
+	SceneContext.FinishGBufferPassAndResolve(RHICmdList);
+}
+
+
+void FDeferredShadingSceneRenderer::RenderSingleLayerWaterSSR(FRHICommandListImmediate& RHICmdList, FSingleLayerWaterPassData& PassData)
 {
 	if (CVarSingleLayerWater.GetValueOnRenderThread() <= 0)
 	{
@@ -240,7 +315,7 @@ void FDeferredShadingSceneRenderer::RenderSingleLayerWaterSSR(FRHICommandListImm
 	}
 
 	FRDGBuilder GraphBuilder(RHICmdList);
-	const FSceneRenderTargets& SceneContext = FSceneRenderTargets::Get(GraphBuilder.RHICmdList);
+	FSceneRenderTargets& SceneContext = FSceneRenderTargets::Get(GraphBuilder.RHICmdList);
 	FRDGTextureRef SceneColorTexture = GraphBuilder.RegisterExternalTexture(SceneContext.GetSceneColor());
 
 	for (int32 ViewIndex = 0; ViewIndex < Views.Num(); ViewIndex++)
@@ -309,7 +384,7 @@ void FDeferredShadingSceneRenderer::RenderSingleLayerWaterSSR(FRHICommandListImm
 				PassParameters->ScreenSpaceReflectionsSampler = TStaticSamplerState<SF_Point>::GetRHI();
 				PassParameters->PreIntegratedGF = GSystemTextures.PreintegratedGF->GetRenderTargetItem().ShaderResourceTexture;
 				PassParameters->PreIntegratedGFSampler = TStaticSamplerState<SF_Bilinear, AM_Clamp, AM_Clamp, AM_Clamp>::GetRHI();
-				PassParameters->SceneNoWaterDepthTexture = GraphBuilder.RegisterExternalTexture(SceneContext.SceneDepthZWithoutSingleLayerWater ? SceneContext.SceneDepthZWithoutSingleLayerWater : GSystemTextures.BlackDummy);
+				PassParameters->SceneNoWaterDepthTexture = GraphBuilder.RegisterExternalTexture(PassData.SceneDepthZWithoutSingleLayerWater ? PassData.SceneDepthZWithoutSingleLayerWater : GSystemTextures.BlackDummy);
 				PassParameters->SceneNoWaterDepthSampler = TStaticSamplerState<SF_Point>::GetRHI();
 				PassParameters->SceneTextures = SceneTextures;
 				SetupSceneTextureSamplers(&PassParameters->SceneTextureSamplers);
@@ -355,7 +430,7 @@ void FDeferredShadingSceneRenderer::RenderSingleLayerWaterSSR(FRHICommandListImm
 	ResolveSceneColor(RHICmdList);
 }
 
-bool FDeferredShadingSceneRenderer::RenderSingleLayerWaterPass(FRHICommandListImmediate& RHICmdList, FExclusiveDepthStencil::Type WaterPassDepthStencilAccess)
+bool FDeferredShadingSceneRenderer::RenderSingleLayerWaterPass(FRHICommandListImmediate& RHICmdList, FSingleLayerWaterPassData& PassData, FExclusiveDepthStencil::Type WaterPassDepthStencilAccess)
 {
 	CSV_SCOPED_TIMING_STAT_EXCLUSIVE(RenderSingleLayerWaterPass);
 	SCOPED_NAMED_EVENT(FDeferredShadingSceneRenderer_RenderSingleLayerWaterPass, FColor::Emerald);
@@ -380,7 +455,7 @@ bool FDeferredShadingSceneRenderer::RenderSingleLayerWaterPass(FRHICommandListIm
 
 			TUniformBufferRef<FOpaqueBasePassUniformParameters> WaterPassUniformBuffer;
 			IPooledRenderTarget* WhiteForwardScreenSpaceShadowMask = GSystemTextures.WhiteDummy;
-			CreateOpaqueBasePassUniformBuffer(RHICmdList, View, WhiteForwardScreenSpaceShadowMask, WaterPassUniformBuffer);
+			CreateOpaqueBasePassUniformBuffer(RHICmdList, View, WhiteForwardScreenSpaceShadowMask, PassData.SceneColorWithoutSingleLayerWater, PassData.SceneDepthZWithoutSingleLayerWater, WaterPassUniformBuffer);
 
 			FMeshPassProcessorRenderState DrawRenderState(View, WaterPassUniformBuffer);
 			SetupBasePassState(WaterPassDepthStencilAccess, ViewFamily.EngineShowFlags.ShaderComplexity, DrawRenderState);
@@ -390,7 +465,7 @@ bool FDeferredShadingSceneRenderer::RenderSingleLayerWaterPass(FRHICommandListIm
 			{
 				Scene->UniformBuffers.UpdateViewUniformBuffer(View);
 
-				bDirty |= RenderSingleLayerWaterPassView(RHICmdList, View, DrawRenderState);
+				bDirty |= RenderSingleLayerWaterPassView(RHICmdList, View, PassData, DrawRenderState);
 			}
 		}
 	}
@@ -401,7 +476,7 @@ bool FDeferredShadingSceneRenderer::RenderSingleLayerWaterPass(FRHICommandListIm
 	return bDirty;
 }
 
-bool FDeferredShadingSceneRenderer::RenderSingleLayerWaterPassView(FRHICommandListImmediate& RHICmdList, FViewInfo& View, const FMeshPassProcessorRenderState& InDrawRenderState)
+bool FDeferredShadingSceneRenderer::RenderSingleLayerWaterPassView(FRHICommandListImmediate& RHICmdList, FViewInfo& View, FSingleLayerWaterPassData& PassData, const FMeshPassProcessorRenderState& InDrawRenderState)
 {
 	bool bDirty = false;
 	FMeshPassProcessorRenderState DrawRenderState(InDrawRenderState);
diff --git a/Engine/Source/Runtime/Renderer/Private/SingleLayerWaterRendering.h b/Engine/Source/Runtime/Renderer/Private/SingleLayerWaterRendering.h
--- a/Engine/Source/Runtime/Renderer/Private/SingleLayerWaterRendering.h
+++ b/Engine/Source/Runtime/Renderer/Private/SingleLayerWaterRendering.h
@@ -16,6 +16,12 @@
 class FViewInfo;
 
 
+struct FSingleLayerWaterPassData
+{
+	TRefCountPtr<IPooledRenderTarget> SceneColorWithoutSingleLayerWater;
+	TRefCountPtr<IPooledRenderTarget> SceneDepthZWithoutSingleLayerWater;
+};
+
 
 class FSingleLayerWaterPassMeshProcessor : public FMeshPassProcessor
 {
