Fixed 2019 compilation warning

yoan.stamant

diff --git a/Engine/Source/Runtime/Renderer/Private/SingleLayerWaterRendering.cpp b/Engine/Source/Runtime/Renderer/Private/SingleLayerWaterRendering.cpp
--- a/Engine/Source/Runtime/Renderer/Private/SingleLayerWaterRendering.cpp
+++ b/Engine/Source/Runtime/Renderer/Private/SingleLayerWaterRendering.cpp
@@ -333,7 +333,7 @@ void FDeferredShadingSceneRenderer::RenderSingleLayerWaterSSR(FRHICommandListImm
 		SetupSceneTextureParameters(GraphBuilder, &SceneTextures);
 
 		{
-			const bool bEnableSSR = CVarSingleLayerWaterSSR.GetValueOnRenderThread();
+			const bool bEnableSSR = CVarSingleLayerWaterSSR.GetValueOnRenderThread() != 0;
 			if (bEnableSSR)
 			{
 				// RUN SSR
