Renamed water rendering CVAR.

diff --git a/Engine/Source/Runtime/Renderer/Private/SingleLayerWaterRendering.cpp b/Engine/Source/Runtime/Renderer/Private/SingleLayerWaterRendering.cpp
--- a/Engine/Source/Runtime/Renderer/Private/SingleLayerWaterRendering.cpp
+++ b/Engine/Source/Runtime/Renderer/Private/SingleLayerWaterRendering.cpp
@@ -24,20 +24,20 @@ DECLARE_GPU_STAT(SingleLayerWater);
 
 
 
-static TAutoConsoleVariable<int32> CVarSingleLayerWater(
-	TEXT("r.SingleLayerWater"), 0,
+static TAutoConsoleVariable<int32> CVarWaterSingleLayer(
+	TEXT("r.Water.SingleLayer"), 0,
 	TEXT("Enable the single water renderring system."),
 	ECVF_RenderThreadSafe | ECVF_Scalability);
 
 
-static TAutoConsoleVariable<int32> CVarSingleLayerWaterSSR(
-	TEXT("r.SingleLayerWater.SSR"), 1,
+static TAutoConsoleVariable<int32> CVarWaterSingleLayerSSR(
+	TEXT("r.Water.SingleLayer.SSR"), 1,
 	TEXT("Enable SSR for the single water renderring system."),
 	ECVF_RenderThreadSafe | ECVF_Scalability);
 
 
-static TAutoConsoleVariable<int32> CVarSingleLayerWaterSSRTAA(
-	TEXT("r.SingleLayerWater.SSRTAA"), 1,
+static TAutoConsoleVariable<int32> CVarWaterSingleLayerSSRTAA(
+	TEXT("r.Water.SingleLayer.SSRTAA"), 1,
 	TEXT("Enable SSR denoising using TAA for the single water renderring system."),
 	ECVF_RenderThreadSafe | ECVF_Scalability);
 
@@ -54,7 +54,7 @@ static bool ShouldRenderSingleLayerWater(const FViewInfo& View)
 
 bool ShouldRenderSingleLayerWater(const TArray<FViewInfo>& Views, const FEngineShowFlags& EngineShowFlags)
 {
-	if (CVarSingleLayerWater.GetValueOnRenderThread() > 0) // && EngineShowFlags.Water)
+	if (CVarWaterSingleLayer.GetValueOnRenderThread() > 0) // && EngineShowFlags.Water)
 	{
 		for (const FViewInfo& View : Views)
 		{
@@ -69,7 +69,7 @@ bool ShouldRenderSingleLayerWater(const TArray<FViewInfo>& Views, const FEngineS
 
 bool ShouldRenderSingleLayerWaterSkippedRenderEditorNotification(const TArray<FViewInfo>& Views)
 {
-	if (CVarSingleLayerWater.GetValueOnRenderThread() <= 0)
+	if (CVarWaterSingleLayer.GetValueOnRenderThread() <= 0)
 	{
 		for (const FViewInfo& View : Views)
 		{
@@ -309,7 +309,7 @@ void FDeferredShadingSceneRenderer::FinishWaterGBufferPassAndResolve(FRHICommand
 
 void FDeferredShadingSceneRenderer::RenderSingleLayerWaterSSR(FRHICommandListImmediate& RHICmdList, FSingleLayerWaterPassData& PassData)
 {
-	if (CVarSingleLayerWater.GetValueOnRenderThread() <= 0)
+	if (CVarWaterSingleLayer.GetValueOnRenderThread() <= 0)
 	{
 		return;
 	}
@@ -333,7 +333,7 @@ void FDeferredShadingSceneRenderer::RenderSingleLayerWaterSSR(FRHICommandListImm
 		SetupSceneTextureParameters(GraphBuilder, &SceneTextures);
 
 		{
-			const bool bEnableSSR = CVarSingleLayerWaterSSR.GetValueOnRenderThread() != 0;
+			const bool bEnableSSR = CVarWaterSingleLayerSSR.GetValueOnRenderThread() != 0;
 			if (bEnableSSR)
 			{
 				// RUN SSR
@@ -356,7 +356,7 @@ void FDeferredShadingSceneRenderer::RenderSingleLayerWaterSSR(FRHICommandListImm
 
 				ReflectionsColor = DenoiserInputs.Color;
 
-				if (CVarSingleLayerWaterSSRTAA.GetValueOnRenderThread())
+				if (CVarWaterSingleLayerSSRTAA.GetValueOnRenderThread())
 				{
 					check(View.ViewState);
 					FTAAPassParameters TAASettings(View);
