Fixed on screen warning string to match new cvar name...

diff --git a/Engine/Source/Runtime/Renderer/Private/SceneRendering.cpp b/Engine/Source/Runtime/Renderer/Private/SceneRendering.cpp
--- a/Engine/Source/Runtime/Renderer/Private/SceneRendering.cpp
+++ b/Engine/Source/Runtime/Renderer/Private/SceneRendering.cpp
@@ -2778,7 +2778,7 @@ void FSceneRenderer::RenderFinish(FRHICommandListImmediate& RHICmdList)
 
 					if (bSingleLayerWaterWarning)
 					{
-						static const FText Message = NSLOCTEXT("Renderer", "SingleLayerWater", "r.SingleLayerWater rendering is disabled with a view containing meshe(s) using water material. Meshes are not visible.");
+						static const FText Message = NSLOCTEXT("Renderer", "SingleLayerWater", "r.Water.SingleLayer rendering is disabled with a view containing meshe(s) using water material. Meshes are not visible.");
 						Canvas.DrawShadowedText(10, Y, Message, GetStatsFont(), FLinearColor(1.0, 0.05, 0.05, 1.0));
 						Y += 14;
 					}
