Fixed merge compilation issues

diff --git a/Engine/Source/Runtime/Renderer/Private/SingleLayerWaterRendering.cpp b/Engine/Source/Runtime/Renderer/Private/SingleLayerWaterRendering.cpp
--- a/Engine/Source/Runtime/Renderer/Private/SingleLayerWaterRendering.cpp
+++ b/Engine/Source/Runtime/Renderer/Private/SingleLayerWaterRendering.cpp
@@ -217,7 +217,7 @@ class FSingleLayerWaterCompositePS : public FGlobalShader
 
 	static bool ShouldCompilePermutation(const FGlobalShaderPermutationParameters& Parameters)
 	{
-		if (!IsFeatureLevelSupported(Parameters.Platform, ERHIFeatureLevel::SM4))
+		if (!IsFeatureLevelSupported(Parameters.Platform, ERHIFeatureLevel::SM5))
 		{
 			return false;
 		}
@@ -248,7 +248,7 @@ void FDeferredShadingSceneRenderer::BeginRenderingWaterGBuffer(FRHICommandList&
 
 	const ERHIFeatureLevel::Type CurrentFeatureLevel = SceneContext.GetCurrentFeatureLevel();
 	const bool bUseGBuffer = IsUsingGBuffers(ShaderPlatform);
-	check(CurrentFeatureLevel >= ERHIFeatureLevel::SM4);
+	check(CurrentFeatureLevel >= ERHIFeatureLevel::SM5);
 
 
 	// Allocate required buffers
@@ -363,9 +363,11 @@ void FDeferredShadingSceneRenderer::RenderSingleLayerWaterSSR(FRHICommandListImm
 					TAASettings.Pass = ETAAPassConfig::ScreenSpaceReflections;
 					TAASettings.SceneColorInput = DenoiserInputs.Color;
 
-					FTAAOutputs TAAOutputs = TAASettings.AddTemporalAAPass(
+					FTAAOutputs TAAOutputs = AddTemporalAAPass(
 						GraphBuilder,
-						SceneTextures, View,
+						SceneTextures, 
+						View,
+						TAASettings,
 						View.PrevViewInfo.SSRHistory,
 						&View.ViewState->PrevFrameViewInfo.SSRHistory);
 
@@ -395,7 +397,7 @@ void FDeferredShadingSceneRenderer::RenderSingleLayerWaterSSR(FRHICommandListImm
 					SetupReflectionUniformParameters(View, ReflectionUniformParameters);
 					PassParameters->ReflectionsParameters = CreateUniformBufferImmediate(ReflectionUniformParameters, UniformBuffer_SingleDraw);
 				}
-				PassParameters->RenderTargets[0] = FRenderTargetBinding(SceneColorTexture, ERenderTargetLoadAction::ELoad, ERenderTargetStoreAction::EStore);
+				PassParameters->RenderTargets[0] = FRenderTargetBinding(SceneColorTexture, ERenderTargetLoadAction::ELoad);
 				ClearUnusedGraphResources(*PixelShader, PassParameters);
 
 
