Adding relevance for the dynamic path of the new singlelayerwater pass

diff --git a/Engine/Source/Runtime/Renderer/Private/SceneVisibility.cpp b/Engine/Source/Runtime/Renderer/Private/SceneVisibility.cpp
--- a/Engine/Source/Runtime/Renderer/Private/SceneVisibility.cpp
+++ b/Engine/Source/Runtime/Renderer/Private/SceneVisibility.cpp
@@ -2699,6 +2699,12 @@ void ComputeDynamicMeshRelevance(EShadingPath ShadingPath, bool bAddLightmapDens
 			PassMask.Set(EMeshPass::Velocity);
 			View.NumVisibleDynamicMeshElements[EMeshPass::Velocity] += NumElements;
 		}
+
+		if (ViewRelevance.bUsesSingleLayerWaterMaterial)
+		{
+			PassMask.Set(EMeshPass::SingleLayerWaterPass);
+			View.NumVisibleDynamicMeshElements[EMeshPass::SingleLayerWaterPass] += NumElements;
+		}
 	}
 
 	if (ViewRelevance.HasTranslucency()
