Eanble single layer water by default

#rnx

diff --git a/Engine/Source/Runtime/Renderer/Private/SingleLayerWaterRendering.cpp b/Engine/Source/Runtime/Renderer/Private/SingleLayerWaterRendering.cpp
--- a/Engine/Source/Runtime/Renderer/Private/SingleLayerWaterRendering.cpp
+++ b/Engine/Source/Runtime/Renderer/Private/SingleLayerWaterRendering.cpp
@@ -25,7 +25,7 @@ DECLARE_GPU_STAT(SingleLayerWater);
 
 
 static TAutoConsoleVariable<int32> CVarWaterSingleLayer(
-	TEXT("r.Water.SingleLayer"), 0,
+	TEXT("r.Water.SingleLayer"), 1,
 	TEXT("Enable the single water renderring system."),
 	ECVF_RenderThreadSafe | ECVF_Scalability);
 
