Slightly hacky fix for RVT rendering before GPUScene PrimitiveData is ready
Need to investigate further what triggers this to occur

diff --git a/Engine/Source/Runtime/Renderer/Private/VT/RuntimeVirtualTextureRender.cpp b/Engine/Source/Runtime/Renderer/Private/VT/RuntimeVirtualTextureRender.cpp
--- a/Engine/Source/Runtime/Renderer/Private/VT/RuntimeVirtualTextureRender.cpp
+++ b/Engine/Source/Runtime/Renderer/Private/VT/RuntimeVirtualTextureRender.cpp
@@ -652,7 +652,8 @@ namespace RuntimeVirtualTexture
 		// Test scene is loaded and has been updated once by main rendering passes
 		// This function gets called on the main thread, so accessing scene frame number is not strictly thread safe, but we can probably
 		// assume that frame number is always increasing, and so the test is conservative
-		return Scene != nullptr && Scene->GetRenderScene() != nullptr && Scene->GetRenderScene()->GetFrameNumber() > 1;
+		return Scene != nullptr && Scene->GetRenderScene() != nullptr && Scene->GetRenderScene()->GetFrameNumber() > 1 &&
+			Scene->GetRenderScene()->GPUScene.PrimitiveBuffer.SRV != nullptr; // todo[vt]: Shouldn't need this low level check, but sometimes this isn't created yet and so we crash in rendering.
 	}
 
 	void RenderPage(
