Better handling of water refelction options

diff --git a/Engine/Source/Runtime/Renderer/Private/DeferredShadingRenderer.cpp b/Engine/Source/Runtime/Renderer/Private/DeferredShadingRenderer.cpp
--- a/Engine/Source/Runtime/Renderer/Private/DeferredShadingRenderer.cpp
+++ b/Engine/Source/Runtime/Renderer/Private/DeferredShadingRenderer.cpp
@@ -1805,7 +1805,7 @@ void FDeferredShadingSceneRenderer::Render(FRHICommandListImmediate& RHICmdList)
 		SceneContext.ResolveSceneDepthToAuxiliaryTexture(RHICmdList);
 		RHICmdList.TransitionResource(EResourceTransitionAccess::EReadable, SceneContext.GetSceneDepthSurface());
 
-		RenderSingleLayerWaterSSR(RHICmdList, SingleLayerWaterPassData);
+		RenderSingleLayerWaterReflections(RHICmdList, SingleLayerWaterPassData);
 		ServiceLocalQueue();
 	}
 
diff --git a/Engine/Source/Runtime/Renderer/Private/DeferredShadingRenderer.h b/Engine/Source/Runtime/Renderer/Private/DeferredShadingRenderer.h
--- a/Engine/Source/Runtime/Renderer/Private/DeferredShadingRenderer.h
+++ b/Engine/Source/Runtime/Renderer/Private/DeferredShadingRenderer.h
@@ -115,7 +115,7 @@ public:
 	/** Renders the water draw pass for a given View. */
 	bool RenderSingleLayerWaterPassView(FRHICommandListImmediate& RHICmdList, FViewInfo& View, FSingleLayerWaterPassData& PassData, const FMeshPassProcessorRenderState& InDrawRenderState);
 	/** Render, denoise and composite the scene SSR and under water effect.*/
-	void RenderSingleLayerWaterSSR(FRHICommandListImmediate& RHICmdList, FSingleLayerWaterPassData& PassData);
+	void RenderSingleLayerWaterReflections(FRHICommandListImmediate& RHICmdList, FSingleLayerWaterPassData& PassData);
 
 	/** Finishes the view family rendering. */
 	void RenderFinish(FRHICommandListImmediate& RHICmdList);
diff --git a/Engine/Source/Runtime/Renderer/Private/SingleLayerWaterRendering.cpp b/Engine/Source/Runtime/Renderer/Private/SingleLayerWaterRendering.cpp
--- a/Engine/Source/Runtime/Renderer/Private/SingleLayerWaterRendering.cpp
+++ b/Engine/Source/Runtime/Renderer/Private/SingleLayerWaterRendering.cpp
@@ -30,6 +30,12 @@ static TAutoConsoleVariable<int32> CVarWaterSingleLayer(
 	ECVF_RenderThreadSafe | ECVF_Scalability);
 
 
+static TAutoConsoleVariable<int32> CVarWaterSingleLayerReflection(
+	TEXT("r.Water.SingleLayer.Reflection"), 1,
+	TEXT("Enable reflection rendering on water."),
+	ECVF_RenderThreadSafe | ECVF_Scalability);
+
+
 static TAutoConsoleVariable<int32> CVarWaterSingleLayerSSR(
 	TEXT("r.Water.SingleLayer.SSR"), 1,
 	TEXT("Enable SSR for the single water renderring system."),
@@ -307,9 +313,9 @@ void FDeferredShadingSceneRenderer::FinishWaterGBufferPassAndResolve(FRHICommand
 }
 
 
-void FDeferredShadingSceneRenderer::RenderSingleLayerWaterSSR(FRHICommandListImmediate& RHICmdList, FSingleLayerWaterPassData& PassData)
+void FDeferredShadingSceneRenderer::RenderSingleLayerWaterReflections(FRHICommandListImmediate& RHICmdList, FSingleLayerWaterPassData& PassData)
 {
-	if (CVarWaterSingleLayer.GetValueOnRenderThread() <= 0)
+	if (CVarWaterSingleLayer.GetValueOnRenderThread() <= 0 || CVarWaterSingleLayerReflection.GetValueOnRenderThread() <= 0)
 	{
 		return;
 	}
@@ -324,104 +330,96 @@ void FDeferredShadingSceneRenderer::RenderSingleLayerWaterSSR(FRHICommandListImm
 
 		FRDGTextureRef ReflectionsColor = nullptr;
 
-		if (!ShouldRenderScreenSpaceReflections(View))
-		{
-			continue;
-		}
-
 		FSceneTextureParameters SceneTextures;
 		SetupSceneTextureParameters(GraphBuilder, &SceneTextures);
 
+		const bool bEnableSSR = CVarWaterSingleLayerSSR.GetValueOnRenderThread() != 0 && ShouldRenderScreenSpaceReflections(View);
+		if (bEnableSSR)
 		{
-			const bool bEnableSSR = CVarWaterSingleLayerSSR.GetValueOnRenderThread() != 0;
-			if (bEnableSSR)
+			// RUN SSR
+			// Uses the water GBuffer (depth, ABCDEF) to know how to start tracing.
+			// The water scene depth is used to know where to start tracing.
+			// Then it uses the scene HZB for the ray casting process.
+
+			FRDGTextureRef CurrentSceneColor = GraphBuilder.RegisterExternalTexture(SceneContext.GetSceneColor());
+
+			IScreenSpaceDenoiser::FReflectionsInputs DenoiserInputs;
+			IScreenSpaceDenoiser::FReflectionsRayTracingConfig RayTracingConfig;
+			ESSRQuality SSRQuality;
+			GetSSRQualityForView(View, &SSRQuality, &RayTracingConfig);
+
+			RDG_EVENT_SCOPE(GraphBuilder, "Water ScreenSpaceReflections(Quality=%d)", int32(SSRQuality));
+
+			const bool bDenoise = false;
+			RenderScreenSpaceReflections(
+				GraphBuilder, SceneTextures, CurrentSceneColor, View, SSRQuality, bDenoise, &DenoiserInputs);
+
+			ReflectionsColor = DenoiserInputs.Color;
+
+			if (CVarWaterSingleLayerSSRTAA.GetValueOnRenderThread())
 			{
-				// RUN SSR
-				// Uses the water GBuffer (depth, ABCDEF) to know how to start tracing.
-				// The water scene depth is used to know where to start tracing.
-				// Then it uses the scene HZB for the ray casting process.
-
-				FRDGTextureRef CurrentSceneColor = GraphBuilder.RegisterExternalTexture(SceneContext.GetSceneColor());
-
-				IScreenSpaceDenoiser::FReflectionsInputs DenoiserInputs;
-				IScreenSpaceDenoiser::FReflectionsRayTracingConfig RayTracingConfig;
-				ESSRQuality SSRQuality;
-				GetSSRQualityForView(View, &SSRQuality, &RayTracingConfig);
-
-				RDG_EVENT_SCOPE(GraphBuilder, "Water ScreenSpaceReflections(Quality=%d)", int32(SSRQuality));
-
-				const bool bDenoise = false;
-				RenderScreenSpaceReflections(
-					GraphBuilder, SceneTextures, CurrentSceneColor, View, SSRQuality, bDenoise, &DenoiserInputs);
-
-				ReflectionsColor = DenoiserInputs.Color;
-
-				if (CVarWaterSingleLayerSSRTAA.GetValueOnRenderThread())
-				{
-					check(View.ViewState);
-					FTAAPassParameters TAASettings(View);
-					TAASettings.Pass = ETAAPassConfig::ScreenSpaceReflections;
-					TAASettings.SceneColorInput = DenoiserInputs.Color;
-
-					FTAAOutputs TAAOutputs = AddTemporalAAPass(
-						GraphBuilder,
-						SceneTextures, 
-						View,
-						TAASettings,
-						View.PrevViewInfo.SSRHistory,
-						&View.ViewState->PrevFrameViewInfo.SSRHistory);
-
-					ReflectionsColor = TAAOutputs.SceneColor;
-				}
+				check(View.ViewState);
+				FTAAPassParameters TAASettings(View);
+				TAASettings.Pass = ETAAPassConfig::ScreenSpaceReflections;
+				TAASettings.SceneColorInput = DenoiserInputs.Color;
+
+				FTAAOutputs TAAOutputs = AddTemporalAAPass(
+					GraphBuilder,
+					SceneTextures, 
+					View,
+					TAASettings,
+					View.PrevViewInfo.SSRHistory,
+					&View.ViewState->PrevFrameViewInfo.SSRHistory);
+
+				ReflectionsColor = TAAOutputs.SceneColor;
 			}
+		}
 
-			// COMPOSITE
+		// Composite reflections on water
+		{
+			FSingleLayerWaterCompositePS::FPermutationDomain PermutationVector;
+			PermutationVector.Set<FSingleLayerWaterCompositePS::FScreenSpaceReflections>(bEnableSSR);
+			TShaderMapRef<FSingleLayerWaterCompositePS> PixelShader(View.ShaderMap, PermutationVector);
+
+			FSingleLayerWaterCompositePS::FParameters* PassParameters = GraphBuilder.AllocParameters<FSingleLayerWaterCompositePS::FParameters>();
+			PassParameters->ScreenSpaceReflectionsTexture = ReflectionsColor ? ReflectionsColor : GraphBuilder.RegisterExternalTexture(GSystemTextures.BlackDummy);
+			PassParameters->ScreenSpaceReflectionsSampler = TStaticSamplerState<SF_Point>::GetRHI();
+			PassParameters->PreIntegratedGF = GSystemTextures.PreintegratedGF->GetRenderTargetItem().ShaderResourceTexture;
+			PassParameters->PreIntegratedGFSampler = TStaticSamplerState<SF_Bilinear, AM_Clamp, AM_Clamp, AM_Clamp>::GetRHI();
+			PassParameters->SceneNoWaterDepthTexture = GraphBuilder.RegisterExternalTexture(PassData.SceneDepthZWithoutSingleLayerWater ? PassData.SceneDepthZWithoutSingleLayerWater : GSystemTextures.BlackDummy);
+			PassParameters->SceneNoWaterDepthSampler = TStaticSamplerState<SF_Point>::GetRHI();
+			PassParameters->SceneTextures = SceneTextures;
+			SetupSceneTextureSamplers(&PassParameters->SceneTextureSamplers);
+			PassParameters->ViewUniformBuffer = View.ViewUniformBuffer;
+			PassParameters->ReflectionCaptureData = View.ReflectionCaptureUniformBuffer;
 			{
-				FSingleLayerWaterCompositePS::FPermutationDomain PermutationVector;
-				PermutationVector.Set<FSingleLayerWaterCompositePS::FScreenSpaceReflections>(bEnableSSR);
-				TShaderMapRef<FSingleLayerWaterCompositePS> PixelShader(View.ShaderMap, PermutationVector);
-
-				FSingleLayerWaterCompositePS::FParameters* PassParameters = GraphBuilder.AllocParameters<FSingleLayerWaterCompositePS::FParameters>();
-				PassParameters->ScreenSpaceReflectionsTexture = ReflectionsColor ? ReflectionsColor : GraphBuilder.RegisterExternalTexture(GSystemTextures.BlackDummy);
-				PassParameters->ScreenSpaceReflectionsSampler = TStaticSamplerState<SF_Point>::GetRHI();
-				PassParameters->PreIntegratedGF = GSystemTextures.PreintegratedGF->GetRenderTargetItem().ShaderResourceTexture;
-				PassParameters->PreIntegratedGFSampler = TStaticSamplerState<SF_Bilinear, AM_Clamp, AM_Clamp, AM_Clamp>::GetRHI();
-				PassParameters->SceneNoWaterDepthTexture = GraphBuilder.RegisterExternalTexture(PassData.SceneDepthZWithoutSingleLayerWater ? PassData.SceneDepthZWithoutSingleLayerWater : GSystemTextures.BlackDummy);
-				PassParameters->SceneNoWaterDepthSampler = TStaticSamplerState<SF_Point>::GetRHI();
-				PassParameters->SceneTextures = SceneTextures;
-				SetupSceneTextureSamplers(&PassParameters->SceneTextureSamplers);
-				PassParameters->ViewUniformBuffer = View.ViewUniformBuffer;
-				PassParameters->ReflectionCaptureData = View.ReflectionCaptureUniformBuffer;
-				{
-					FReflectionUniformParameters ReflectionUniformParameters;
-					SetupReflectionUniformParameters(View, ReflectionUniformParameters);
-					PassParameters->ReflectionsParameters = CreateUniformBufferImmediate(ReflectionUniformParameters, UniformBuffer_SingleDraw);
-				}
-				PassParameters->RenderTargets[0] = FRenderTargetBinding(SceneColorTexture, ERenderTargetLoadAction::ELoad);
-				ClearUnusedGraphResources(*PixelShader, PassParameters);
-
-
-				GraphBuilder.AddPass(
-					RDG_EVENT_NAME("WaterComposite %dx%d", View.ViewRect.Width(), View.ViewRect.Height()),
-					PassParameters,
-					ERDGPassFlags::Raster,
-					[PassParameters, &View, PixelShader](FRHICommandList& InRHICmdList)
-				{
-					InRHICmdList.SetViewport(View.ViewRect.Min.X, View.ViewRect.Min.Y, 0.0f, View.ViewRect.Max.X, View.ViewRect.Max.Y, 1.0f);
-
-					FGraphicsPipelineStateInitializer GraphicsPSOInit;
-					FPixelShaderUtils::InitFullscreenPipelineState(InRHICmdList, View.ShaderMap, *PixelShader, GraphicsPSOInit);
-
-					// Premultiplied alpha where alpha is transmittance.
-					GraphicsPSOInit.BlendState = TStaticBlendState<CW_RGB, BO_Add, BF_One, BF_SourceAlpha>::GetRHI(); 
-
-					SetGraphicsPipelineState(InRHICmdList, GraphicsPSOInit);
-					SetShaderParameters(InRHICmdList, *PixelShader, PixelShader->GetPixelShader(), *PassParameters);
-					FPixelShaderUtils::DrawFullscreenTriangle(InRHICmdList);
-				});
+				FReflectionUniformParameters ReflectionUniformParameters;
+				SetupReflectionUniformParameters(View, ReflectionUniformParameters);
+				PassParameters->ReflectionsParameters = CreateUniformBufferImmediate(ReflectionUniformParameters, UniformBuffer_SingleDraw);
 			}
-		}
+			PassParameters->RenderTargets[0] = FRenderTargetBinding(SceneColorTexture, ERenderTargetLoadAction::ELoad);
+			ClearUnusedGraphResources(*PixelShader, PassParameters);
+
 
+			GraphBuilder.AddPass(
+				RDG_EVENT_NAME("Water Composite %dx%d", View.ViewRect.Width(), View.ViewRect.Height()),
+				PassParameters,
+				ERDGPassFlags::Raster,
+				[PassParameters, &View, PixelShader](FRHICommandList& InRHICmdList)
+			{
+				InRHICmdList.SetViewport(View.ViewRect.Min.X, View.ViewRect.Min.Y, 0.0f, View.ViewRect.Max.X, View.ViewRect.Max.Y, 1.0f);
+
+				FGraphicsPipelineStateInitializer GraphicsPSOInit;
+				FPixelShaderUtils::InitFullscreenPipelineState(InRHICmdList, View.ShaderMap, *PixelShader, GraphicsPSOInit);
+
+				// Premultiplied alpha where alpha is transmittance.
+				GraphicsPSOInit.BlendState = TStaticBlendState<CW_RGB, BO_Add, BF_One, BF_SourceAlpha>::GetRHI(); 
+
+				SetGraphicsPipelineState(InRHICmdList, GraphicsPSOInit);
+				SetShaderParameters(InRHICmdList, *PixelShader, PixelShader->GetPixelShader(), *PassParameters);
+				FPixelShaderUtils::DrawFullscreenTriangle(InRHICmdList);
+			});
+		}
 	}
 
 	TRefCountPtr<IPooledRenderTarget> OutSceneColor;
