Make distortion code common between post process and mobile..

diff --git a/Engine/Shaders/Private/BasePassPixelShader.usf b/Engine/Shaders/Private/BasePassPixelShader.usf
--- a/Engine/Shaders/Private/BasePassPixelShader.usf
+++ b/Engine/Shaders/Private/BasePassPixelShader.usf
@@ -1228,7 +1228,7 @@ ShadingModel);
 			OpaqueBasePass.SceneDepthWithoutSingleLayerWaterTexture, OpaqueBasePass.SceneDepthWithoutSingleLayerWaterSampler,
 			OpaqueBasePass.SceneLuminanceWithoutSingleLayerWaterTexture, OpaqueBasePass.SceneLuminanceWithoutSingleLayerWaterSampler,
 			OpaqueBasePass.SceneCustomDepthTexture, OpaqueBasePass.SceneCustomDepthSampler,
-			Specular,
+			Specular, OpaqueBasePass.DistortionParams,
 			SunIlluminance, WaterDiffuseIndirectIlluminance, EnvBrdf,
 			CameraIsUnderWater, WaterVisibility).Luminance;
 	}
diff --git a/Engine/Shaders/Private/DistortAccumulatePS.usf b/Engine/Shaders/Private/DistortAccumulatePS.usf
--- a/Engine/Shaders/Private/DistortAccumulatePS.usf
+++ b/Engine/Shaders/Private/DistortAccumulatePS.usf
@@ -18,11 +18,9 @@
 #include "SceneTexturesCommon.ush"
 #include "/Engine/Generated/Material.ush"
 #include "/Engine/Generated/VertexFactory.ush"
+#include "DistortionCommon.ush"
 
-//A fudge factor scale to bring values close to what they would have been under usual circumstances prior to this change.
-static const float OffsetFudgeFactor = 0.00023;
 
-static const half DistortionScaleBias = 4.0f;
 
 /** output distortion offsets as color so they can be accumulated (via blending) */
 void Main(
@@ -56,47 +54,17 @@ void Main(
 	// material distortion offset
 	half3 Normal = GetMaterialNormal(MaterialParameters, PixelMaterialInputs);
 
-#if MATERIAL_TANGENTSPACENORMAL
-	half3 ViewNormal = normalize(TransformTangentVectorToView(MaterialParameters, Normal));
-#else
-	half3 ViewNormal = normalize(TransformWorldVectorToView(Normal));
-#endif
-
-	// X = IOR, Y = Refraction Bias
-	float2 Refraction = GetMaterialRefraction(PixelMaterialInputs);
-	 	
-#if REFRACTION_USE_PIXEL_NORMAL_OFFSET
-	half3 ViewVertexNormal = TransformTangentVectorToView(MaterialParameters, float3(0, 0, 1));
-	// Treat Refraction of 1 as no refraction, to be consistent with IOR mode
-	float2 ViewportUVDistortion = (ViewVertexNormal.xy - ViewNormal.xy) * (Refraction.x - 1);
-#else
-	// we assume the camera is in air
-	float AirIOR = 1.0f;
-	float2 ViewportUVDistortion = ViewNormal.xy * (Refraction.x - AirIOR);
-#endif
-
-	float2 BufferUVDistortion = ViewportUVDistortion * ResolvedView.ViewSizeAndInvSize.xy * ResolvedView.BufferSizeAndInvSize.zw;
-
-	//clip if the distortion distance (squared) is too small to be noticed
-	//this will result in a less expensive apply pass since the clipped pixels won't set stencil to pass
-	clip(dot(BufferUVDistortion, BufferUVDistortion) - .00001);
-
 	// Prevent silhouettes from geometry that is in front of distortion from being seen in the distortion 
 	float2 NDC = (MaterialParameters.ScreenPosition.xy / MaterialParameters.ScreenPosition.w);
 	float2 ScreenUV = NDC * ResolvedView.ScreenPositionScaleBias.xy + ResolvedView.ScreenPositionScaleBias.wz;
 
-	// When ISR is enabled we store two FOVs in the DistortionParams and compute the aspect ratio here
-#if INSTANCED_STEREO
-	float InvTanHalfFov = DistortionParams[Interpolants.EyeIndex];
-    float Ratio = DistortionParams.z / DistortionParams.w;
-#else
-	float InvTanHalfFov = DistortionParams.x;
-    float Ratio = DistortionParams.y;
-#endif
+	// X = IOR, Y = Refraction Bias
+	float2 Refraction = GetMaterialRefraction(PixelMaterialInputs);
 
-	// Fix for Fov and aspect.
- 	float2 FovFix = float2(InvTanHalfFov,Ratio*InvTanHalfFov);
-	BufferUVDistortion *= DistortionParams.zw * float2(OffsetFudgeFactor,-OffsetFudgeFactor) * FovFix;
+	// Compute UV distortion
+	float2 BufferUVDistortion = ComputeBufferUVDistortion(
+		MaterialParameters, PixelMaterialInputs, ResolvedView,
+		Normal, DistortionParams, ScreenUV, Refraction, true);
 
 	// Sample depth at distortion offset
 	float2 DistortBufferUV = ScreenUV + BufferUVDistortion;
@@ -107,17 +75,9 @@ void Main(
 #else
 	float DistortSceneDepth = CalcSceneDepth(DistortBufferUV);
 #endif // (ES2_PROFILE || ES3_1_PROFILE) 
-	
-	// Soft thresholding 
-	float Bias = -Refraction.y;
-	float Range = clamp(abs(Bias * 0.5f), 0, 50);
-	float Z = DistortSceneDepth;
-	float ZCompare = MaterialParameters.ScreenPosition.w;
-	float InvWidth = 1.0f / max(1.0f, Range );
-	BufferUVDistortion *= saturate((Z - ZCompare) * InvWidth + Bias);
-
-	//Scale up for better precision in low/subtle refractions at the expense of artefacts at higher refraction.
-	BufferUVDistortion*= DistortionScaleBias;
+
+	// Post process UV distortion according to depth
+	PostProcessUVDistortion(MaterialParameters, PixelMaterialInputs, DistortSceneDepth, BufferUVDistortion, Refraction);
 
 	// store positive and negative offsets separately
 	float2 PosOffset = max(BufferUVDistortion,0);
diff --git a/Engine/Shaders/Private/DistortionCommon.ush b/Engine/Shaders/Private/DistortionCommon.ush
new file mode 100644
--- /dev/null
+++ b/Engine/Shaders/Private/DistortionCommon.ush
@@ -0,0 +1,74 @@
+/*=============================================================================
+	DistortionCommon.ush
+=============================================================================*/
+
+#include "Common.ush"
+
+
+float2 ComputeBufferUVDistortion(
+	in FMaterialPixelParameters MaterialParameters, in FPixelMaterialInputs PixelMaterialInputs, in ViewState ResolvedView,
+	in half3 Normal, in float4 DistortionParams, in float2 ScreenUV, in float2 Refraction, in bool TryToClip)
+{
+#if MATERIAL_TANGENTSPACENORMAL
+	half3 ViewNormal = normalize(TransformTangentVectorToView(MaterialParameters, Normal));
+#else
+	half3 ViewNormal = normalize(TransformWorldVectorToView(Normal));
+#endif
+
+#if REFRACTION_USE_PIXEL_NORMAL_OFFSET
+	half3 ViewVertexNormal = TransformTangentVectorToView(MaterialParameters, float3(0, 0, 1));
+	// Treat Refraction of 1 as no refraction, to be consistent with IOR mode
+	float2 ViewportUVDistortion = (ViewVertexNormal.xy - ViewNormal.xy) * (Refraction.x - 1);
+#else
+	// we assume the camera is in air
+	float AirIOR = 1.0f;
+	float2 ViewportUVDistortion = ViewNormal.xy * (Refraction.x - AirIOR);
+#endif
+
+	float2 BufferUVDistortion = ViewportUVDistortion * ResolvedView.ViewSizeAndInvSize.xy * ResolvedView.BufferSizeAndInvSize.zw;
+
+	//clip if the distortion distance (squared) is too small to be noticed
+	//this will result in a less expensive apply pass since the clipped pixels won't set stencil to pass
+	if (TryToClip)
+	{
+		clip(dot(BufferUVDistortion, BufferUVDistortion) - .00001);
+	}
+
+	// When ISR is enabled we store two FOVs in the DistortionParams and compute the aspect ratio here
+#if INSTANCED_STEREO
+	float InvTanHalfFov = DistortionParams[Interpolants.EyeIndex];
+	float Ratio = DistortionParams.z / DistortionParams.w;
+#else
+	float InvTanHalfFov = DistortionParams.x;
+	float Ratio = DistortionParams.y;
+#endif
+
+	// Fix for Fov and aspect.
+	float2 FovFix = float2(InvTanHalfFov, Ratio*InvTanHalfFov);
+	//A fudge factor scale to bring values close to what they would have been under usual circumstances prior to this change.
+	const float OffsetFudgeFactor = 0.00023;
+	BufferUVDistortion *= DistortionParams.zw * float2(OffsetFudgeFactor, -OffsetFudgeFactor) * FovFix;
+
+	return BufferUVDistortion;
+}
+
+
+void PostProcessUVDistortion(
+	in FMaterialPixelParameters MaterialParameters, in FPixelMaterialInputs PixelMaterialInputs, 
+	in float DistortSceneDepth,	inout float2 BufferUVDistortion, in float2 Refraction)
+{
+	// Soft thresholding 
+	float Bias = -Refraction.y;
+	float Range = clamp(abs(Bias * 0.5f), 0, 50);
+	float Z = DistortSceneDepth;
+	float ZCompare = MaterialParameters.ScreenPosition.w;
+	float InvWidth = 1.0f / max(1.0f, Range);
+	BufferUVDistortion *= saturate((Z - ZCompare) * InvWidth + Bias);
+
+	//Scale up for better precision in low/subtle refractions at the expense of artefacts at higher refraction.
+	static const half DistortionScaleBias = 4.0f;
+	BufferUVDistortion *= DistortionScaleBias;
+}
+
+
+
diff --git a/Engine/Shaders/Private/MobileBasePassPixelShader.usf b/Engine/Shaders/Private/MobileBasePassPixelShader.usf
--- a/Engine/Shaders/Private/MobileBasePassPixelShader.usf
+++ b/Engine/Shaders/Private/MobileBasePassPixelShader.usf
@@ -664,10 +664,11 @@ void Main(
 		const float3 WaterDiffuseIndirectIlluminance = WaterDiffuseIndirectLuminance * PI;	// DiffuseIndirectLighting is luminance. So we need to multiply by PI to get illuminance.
 		const float3 EnvBrdf = SpecularColor; // SpecularColor is not F0 as in BasePassPixelShader, it is EnvBRDF.
 
+		const float4 NullDistortionParams = 1.0f;
 		WaterVolumeLightingOutput WaterLighting = EvaluateWaterVolumeLighting(
 			MaterialParameters, PixelMaterialInputs, ResolvedView,
 			MobileSceneTextures.SceneDepthTexture, MobileSceneTextures.SceneDepthTextureSampler,
-			Shadow, Specular,
+			Shadow, Specular, NullDistortionParams,
 			SunIlluminance, WaterDiffuseIndirectIlluminance, EnvBrdf,
 			CameraIsUnderWater, WaterVisibility);
 
diff --git a/Engine/Shaders/Private/SingleLayerWaterCommon.ush b/Engine/Shaders/Private/SingleLayerWaterCommon.ush
--- a/Engine/Shaders/Private/SingleLayerWaterCommon.ush
+++ b/Engine/Shaders/Private/SingleLayerWaterCommon.ush
@@ -5,6 +5,7 @@
 #include "Common.ush"
 #include "ParticipatingMediaCommon.ush" 
 #include "ShadingCommon.ush" 
+#include "DistortionCommon.ush"
 
 
 #ifndef SIMPLE_SINGLE_LAYER_WATER
@@ -28,7 +29,7 @@ WaterVolumeLightingOutput EvaluateWaterVolumeLighting(
 	Texture2D SceneLuminanceWithoutSingleLayerWaterTexture, SamplerState SceneLuminanceWithoutSingleLayerWaterSampler,
 	Texture2D SceneCustomDepthTexture, SamplerState SceneCustomDepthSampler,
 #endif // SIMPLE_SINGLE_LAYER_WATER
-	float Specular,
+	float Specular, const float4 DistortionParams,
 	float3 SunIlluminance, float3 AmbiantIlluminance, float3 EnvBrdf,
 	bool CameraIsUnderWater, float WaterVisibility)
 {
@@ -61,36 +62,18 @@ WaterVolumeLightingOutput EvaluateWaterVolumeLighting(
 
 #else  // !SIMPLE_SINGLE_LAYER_WATER
 
-		//////////////////////////////////////////////////////////////////////////////////////////////////// BEGIN Refration TODO make common with DistoAccumulatePS
-		float2 Refraction = GetMaterialRefraction(PixelMaterialInputs);// X = IOR, Y = Refraction Bias
-		float3 ViewNormal = normalize(TransformWorldVectorToView(MaterialParameters.WorldNormal));
-#if REFRACTION_USE_PIXEL_NORMAL_OFFSET
-		half3 ViewVertexNormal = TransformTangentVectorToView(MaterialParameters, float3(0, 0, 1));
-		// Treat Refraction of 1 as no refraction, to be consistent with IOR mode
-		float2 ViewportUVDistortion = (ViewVertexNormal.xy - ViewNormal.xy) * (Refraction.x - 1);
-#else
-		// we assume the camera is in air
-		float AirIOR = 1.0f;
-		float2 ViewportUVDistortion = ViewNormal.xy * (Refraction.x - AirIOR);
-#endif
-		float2 BufferUVDistortion = ViewportUVDistortion * ResolvedView.ViewSizeAndInvSize.xy * ResolvedView.BufferSizeAndInvSize.zw;
-		// When ISR is enabled we store two FOVs in the DistortionParams and compute the aspect ratio here
-		const float4 DistortionParams = 1.0f;
-		const float OffsetFudgeFactor = 1.0f;
-		float InvTanHalfFov = DistortionParams.x;
-		float Ratio = ResolvedView.ViewSizeAndInvSize.y * ResolvedView.ViewSizeAndInvSize.w;
-		// Fix for Fov and aspect.
-		float2 FovFix = float2(InvTanHalfFov, Ratio*InvTanHalfFov);
-		BufferUVDistortion *= 0.2 * DistortionParams.zw * float2(OffsetFudgeFactor, -OffsetFudgeFactor) * FovFix;
-		////////////////////////////////////////////////////////////////////////////////////////////////////  END Refration
-
-		// TODO make distortion fading common with DistoAccumulatePS (right now 30cm to BehindWaterDeltaDepth fade out and 10m to distance fade out)
-		BufferUVDistortion *= saturate((PixelSceneDepth - WaterDepth) * 1.0 / 30.0f) * (1 - saturate(WaterDepth * 1.0 / 2000.0f));
-
-
+		float2 Refraction = GetMaterialRefraction(PixelMaterialInputs); // X = IOR, Y = Refraction Bias
+		float2 BufferUVDistortion = ComputeBufferUVDistortion(
+			MaterialParameters, PixelMaterialInputs, ResolvedView,
+			MaterialParameters.WorldNormal, DistortionParams, ViewportUV, Refraction, false);
+		PostProcessUVDistortion(MaterialParameters, PixelMaterialInputs, PixelSceneDepth, BufferUVDistortion, Refraction);
+		// Also reduce refraction according to water thickness
+		BufferUVDistortion *= saturate((PixelSceneDepth - WaterDepth) * 1.0 / 30.0f);// *(1 - saturate(WaterDepth * 1.0 / 2000.0f));
 
+		// Clamp distorted UVs to not sample black outside the sub part of the buffer. Big quality increase.
+		const float2 DistortedUv = clamp(ViewportUV + BufferUVDistortion, 0.0, (ResolvedView.ViewSizeAndInvSize.xy-0.5f) * ResolvedView.BufferSizeAndInvSize.zw);
 
-		float SceneDeviceZ = SceneDepthWithoutSingleLayerWaterTexture.SampleLevel(SceneDepthWithoutSingleLayerWaterSampler, saturate(ViewportUV + BufferUVDistortion), 0).r;
+		float SceneDeviceZ = SceneDepthWithoutSingleLayerWaterTexture.SampleLevel(SceneDepthWithoutSingleLayerWaterSampler, DistortedUv, 0).r;
 		float SceneDepth = ConvertFromDeviceZ(SceneDeviceZ);
 		if (SceneDepth < WaterDepth)
 		{
@@ -98,7 +81,7 @@ WaterVolumeLightingOutput EvaluateWaterVolumeLighting(
 		}
 		else
 		{
-			ViewportUV = saturate(ViewportUV + BufferUVDistortion);
+			ViewportUV = DistortedUv;
 		}
 
 #endif // !SIMPLE_SINGLE_LAYER_WATER
diff --git a/Engine/Source/Runtime/Renderer/Private/BasePassRendering.cpp b/Engine/Source/Runtime/Renderer/Private/BasePassRendering.cpp
--- a/Engine/Source/Runtime/Renderer/Private/BasePassRendering.cpp
+++ b/Engine/Source/Runtime/Renderer/Private/BasePassRendering.cpp
@@ -6,6 +6,7 @@
 
 #include "BasePassRendering.h"
 #include "DeferredShadingRenderer.h"
+#include "DistortionRendering.h"
 #include "DynamicPrimitiveDrawing.h"
 #include "ScenePrivate.h"
 #include "MeshPassProcessor.inl"
@@ -589,6 +590,7 @@ void CreateOpaqueBasePassUniformBuffer(
 	BasePassParameters.PreIntegratedGFSampler = TStaticSamplerState<SF_Bilinear, AM_Clamp, AM_Clamp, AM_Clamp>::GetRHI();
 	BasePassParameters.SceneCustomDepthTexture = (SceneRenderTargets.CustomDepth.IsValid() ? SceneRenderTargets.CustomDepth : GSystemTextures.BlackDummy)->GetRenderTargetItem().ShaderResourceTexture;
 	BasePassParameters.SceneCustomDepthSampler = TStaticSamplerState<SF_Point>::GetRHI();
+	SetupDistortionParams(BasePassParameters.DistortionParams, View);
 
 	// Misc
 	BasePassParameters.EyeAdaptation = GetEyeAdaptation(View);
diff --git a/Engine/Source/Runtime/Renderer/Private/BasePassRendering.h b/Engine/Source/Runtime/Renderer/Private/BasePassRendering.h
--- a/Engine/Source/Runtime/Renderer/Private/BasePassRendering.h
+++ b/Engine/Source/Runtime/Renderer/Private/BasePassRendering.h
@@ -80,6 +80,7 @@ BEGIN_GLOBAL_SHADER_PARAMETER_STRUCT(FOpaqueBasePassUniformParameters,)
 	SHADER_PARAMETER_SAMPLER(SamplerState, PreIntegratedGFSampler)
 	SHADER_PARAMETER_TEXTURE(Texture2D, SceneCustomDepthTexture)
 	SHADER_PARAMETER_SAMPLER(SamplerState, SceneCustomDepthSampler)
+	SHADER_PARAMETER(FVector4, DistortionParams)
 	// Misc
 	SHADER_PARAMETER_TEXTURE(Texture2D, EyeAdaptation)
 END_GLOBAL_SHADER_PARAMETER_STRUCT()
diff --git a/Engine/Source/Runtime/Renderer/Private/DistortionRendering.cpp b/Engine/Source/Runtime/Renderer/Private/DistortionRendering.cpp
--- a/Engine/Source/Runtime/Renderer/Private/DistortionRendering.cpp
+++ b/Engine/Source/Runtime/Renderer/Private/DistortionRendering.cpp
@@ -41,16 +41,13 @@ static TAutoConsoleVariable<int32> CVarDisableDistortion(
 IMPLEMENT_GLOBAL_SHADER_PARAMETER_STRUCT(FDistortionPassUniformParameters, "DistortionPass");
 IMPLEMENT_GLOBAL_SHADER_PARAMETER_STRUCT(FMobileDistortionPassUniformParameters, "MobileDistortionPass");
 
-void SetupDistortionPassUniformBuffer(FRHICommandListImmediate& RHICmdList, const FViewInfo& View, FDistortionPassUniformParameters& DistortionPassParameters)
+void SetupDistortionParams(FVector4& DistortionParams, const FViewInfo& View)
 {
-	FSceneRenderTargets& SceneRenderTargets = FSceneRenderTargets::Get(RHICmdList);
-	SetupSceneTextureUniformParameters(SceneRenderTargets, View.FeatureLevel, ESceneTextureSetupMode::All, DistortionPassParameters.SceneTextures);
-
 	float Ratio = View.UnscaledViewRect.Width() / (float)View.UnscaledViewRect.Height();
-	DistortionPassParameters.DistortionParams.X = View.ViewMatrices.GetProjectionMatrix().M[0][0];
-	DistortionPassParameters.DistortionParams.Y = Ratio;
-	DistortionPassParameters.DistortionParams.Z = (float)View.UnscaledViewRect.Width();
-	DistortionPassParameters.DistortionParams.W = (float)View.UnscaledViewRect.Height();
+	DistortionParams.X = View.ViewMatrices.GetProjectionMatrix().M[0][0];
+	DistortionParams.Y = Ratio;
+	DistortionParams.Z = (float)View.UnscaledViewRect.Width();
+	DistortionParams.W = (float)View.UnscaledViewRect.Height();
 
 	// When ISR is enabled we store two FOVs in the distortion parameters and compute the aspect ratio in the shader instead.
 	if ((View.IsInstancedStereoPass() || View.bIsMobileMultiViewEnabled) && View.Family->Views.Num() > 0)
@@ -59,30 +56,22 @@ void SetupDistortionPassUniformBuffer(FRHICommandListImmediate& RHICmdList, cons
 		const EStereoscopicPass StereoPassIndex = (View.StereoPass != eSSP_FULL) ? eSSP_RIGHT_EYE : eSSP_FULL;
 
 		const FViewInfo& InstancedView = static_cast<const FViewInfo&>(View.Family->GetStereoEyeView(StereoPassIndex));
-		DistortionPassParameters.DistortionParams.Y = InstancedView.ViewMatrices.GetProjectionMatrix().M[0][0];
+		DistortionParams.Y = InstancedView.ViewMatrices.GetProjectionMatrix().M[0][0];
 	}
 }
 
+void SetupDistortionPassUniformBuffer(FRHICommandListImmediate& RHICmdList, const FViewInfo& View, FDistortionPassUniformParameters& DistortionPassParameters)
+{
+	FSceneRenderTargets& SceneRenderTargets = FSceneRenderTargets::Get(RHICmdList);
+	SetupSceneTextureUniformParameters(SceneRenderTargets, View.FeatureLevel, ESceneTextureSetupMode::All, DistortionPassParameters.SceneTextures);
+	SetupDistortionParams(DistortionPassParameters.DistortionParams, View);
+}
+
 void SetupMobileDistortionPassUniformBuffer(FRHICommandListImmediate& RHICmdList, const FViewInfo& View, FMobileDistortionPassUniformParameters& DistortionPassParameters)
 {
 	FSceneRenderTargets& SceneRenderTargets = FSceneRenderTargets::Get(RHICmdList);
 	SetupMobileSceneTextureUniformParameters(SceneRenderTargets, View.FeatureLevel, true, DistortionPassParameters.SceneTextures);
-
-	float Ratio = View.UnscaledViewRect.Width() / (float)View.UnscaledViewRect.Height();
-	DistortionPassParameters.DistortionParams.X = View.ViewMatrices.GetProjectionMatrix().M[0][0];
-	DistortionPassParameters.DistortionParams.Y = Ratio;
-	DistortionPassParameters.DistortionParams.Z = (float)View.UnscaledViewRect.Width();
-	DistortionPassParameters.DistortionParams.W = (float)View.UnscaledViewRect.Height();
-
-	// When ISR is enabled we store two FOVs in the distortion parameters and compute the aspect ratio in the shader instead.
-	if ((View.IsInstancedStereoPass() || View.bIsMobileMultiViewEnabled) && View.Family->Views.Num() > 0)
-	{
-		// When drawing the left eye in a stereo scene, copy the right eye view values into the instanced view uniform buffer.
-		const EStereoscopicPass StereoPassIndex = (View.StereoPass != eSSP_FULL) ? eSSP_RIGHT_EYE : eSSP_FULL;
-
-		const FViewInfo& InstancedView = static_cast<const FViewInfo&>(View.Family->GetStereoEyeView(StereoPassIndex));
-		DistortionPassParameters.DistortionParams.Y = InstancedView.ViewMatrices.GetProjectionMatrix().M[0][0];
-	}
+	SetupDistortionParams(DistortionPassParameters.DistortionParams, View);
 }
 
 /**
diff --git a/Engine/Source/Runtime/Renderer/Private/DistortionRendering.h b/Engine/Source/Runtime/Renderer/Private/DistortionRendering.h
--- a/Engine/Source/Runtime/Renderer/Private/DistortionRendering.h
+++ b/Engine/Source/Runtime/Renderer/Private/DistortionRendering.h
@@ -70,6 +70,9 @@ BEGIN_GLOBAL_SHADER_PARAMETER_STRUCT(FMobileDistortionPassUniformParameters, )
 	SHADER_PARAMETER(FVector4, DistortionParams)
 END_GLOBAL_SHADER_PARAMETER_STRUCT()
 
+
+extern void SetupDistortionParams(FVector4& DistortionParams, const FViewInfo& View);
+
 extern void SetupMobileDistortionPassUniformBuffer(FRHICommandListImmediate& RHICmdList, const FViewInfo& View, FMobileDistortionPassUniformParameters& DistortionPassParameters);
 
 class FDistortionMeshProcessor : public FMeshPassProcessor
