Give symbols more unique names to avoid conflict with other system's symbols
#rnx

diff --git a/Engine/Source/Runtime/Renderer/Private/GpuDebugRendering.cpp b/Engine/Source/Runtime/Renderer/Private/GpuDebugRendering.cpp
--- a/Engine/Source/Runtime/Renderer/Private/GpuDebugRendering.cpp
+++ b/Engine/Source/Runtime/Renderer/Private/GpuDebugRendering.cpp
@@ -41,9 +41,9 @@ FGpuDebugPrimitiveBuffers AllocateGpuDebugPrimitiveBuffers(FRHICommandListImmedi
 	return Output;
 }
 
-class FLine { public: FVector Start; FVector End; FLinearColor Color; };
-typedef TArray<FLine> FLineArray;
-static FLineArray ReadGPUDebugPrimitives(FRHICommandListImmediate& RHICmdList, FGpuDebugPrimitiveBuffers& DebugPrimitiveBuffer)
+class FGpuDebugLine { public: FVector Start; FVector End; FLinearColor Color; };
+typedef TArray<FGpuDebugLine> FGpuDebugLineArray;
+static FGpuDebugLineArray ReadGPUDebugPrimitives(FRHICommandListImmediate& RHICmdList, FGpuDebugPrimitiveBuffers& DebugPrimitiveBuffer)
 {
 	// Count
 	uint32 PointCount = 0;
@@ -70,7 +70,7 @@ static FLineArray ReadGPUDebugPrimitives(FRHICommandListImmediate& RHICmdList, F
 	}
 
 	// Primitives
-	FLineArray Output;
+	FGpuDebugLineArray Output;
 	if (PointCount > 0)
 	{
 		auto ToColor = [](uint32 ColorIndex)
@@ -127,7 +127,7 @@ void BindGpuDebugPrimitiveBuffers(FRHIRenderPassInfo& RPInfo, FGpuDebugPrimitive
 
 void DrawGpuDebugPrimitives(FRHICommandListImmediate& RHICmdList, TArray<FViewInfo>& Views, FGpuDebugPrimitiveBuffers& DebugPrimitiveBuffer)
 {
-	const FLineArray DebugLines = ReadGPUDebugPrimitives(RHICmdList, DebugPrimitiveBuffer);
+	const FGpuDebugLineArray DebugLines = ReadGPUDebugPrimitives(RHICmdList, DebugPrimitiveBuffer);
 
 	uint32 ViewIndex = 0;
 	for (FViewInfo& View : Views)
@@ -136,7 +136,7 @@ void DrawGpuDebugPrimitives(FRHICommandListImmediate& RHICmdList, TArray<FViewIn
 
 		for(int32 i=0; i < DebugLines.Num(); ++i)
 		{
-			const FLine& Line = DebugLines[i];
+			const FGpuDebugLine& Line = DebugLines[i];
 			PDI.DrawLine(Line.Start, Line.End, Line.Color, 0);
 		}
 
