Fix for composite pass being used in Forward.

diff --git a/Engine/Source/Runtime/Renderer/Private/SingleLayerWaterRendering.cpp b/Engine/Source/Runtime/Renderer/Private/SingleLayerWaterRendering.cpp
--- a/Engine/Source/Runtime/Renderer/Private/SingleLayerWaterRendering.cpp
+++ b/Engine/Source/Runtime/Renderer/Private/SingleLayerWaterRendering.cpp
@@ -320,6 +320,16 @@ void FDeferredShadingSceneRenderer::RenderSingleLayerWaterReflections(FRHIComman
 		return;
 	}
 
+	bool AllViewAreForward = true;
+	for (int32 ViewIndex = 0; ViewIndex < Views.Num(); ViewIndex++)
+	{
+		AllViewAreForward &= IsAnyForwardShadingEnabled(Views[ViewIndex].GetShaderPlatform());
+	}
+	if (AllViewAreForward)
+	{
+		return; // No SSR or composite needed in Forward for anyview so quick return. Reflections are applied in the WaterGBuffer pass.
+	}
+
 	FRDGBuilder GraphBuilder(RHICmdList);
 	FSceneRenderTargets& SceneContext = FSceneRenderTargets::Get(GraphBuilder.RHICmdList);
 	FRDGTextureRef SceneColorTexture = GraphBuilder.RegisterExternalTexture(SceneContext.GetSceneColor());
@@ -328,6 +338,11 @@ void FDeferredShadingSceneRenderer::RenderSingleLayerWaterReflections(FRHIComman
 	{
 		FViewInfo& View = Views[ViewIndex];
 
+		if (IsAnyForwardShadingEnabled(View.GetShaderPlatform()))
+		{
+			continue; // No SSR or composite needed in forward views.
+		}
+
 		FRDGTextureRef ReflectionsColor = nullptr;
 
 		FSceneTextureParameters SceneTextures;
