Improves performance of LightScatteringCS by up to ~40%.

diff --git a/Engine/Shaders/Private/VolumetricFog.usf b/Engine/Shaders/Private/VolumetricFog.usf
--- a/Engine/Shaders/Private/VolumetricFog.usf
+++ b/Engine/Shaders/Private/VolumetricFog.usf
@@ -71,6 +71,18 @@ float GlobalExtinctionScale;
 #define THREADGROUP_SIZE 1
 #endif
 
+#ifndef THREADGROUP_SIZE_X
+#define THREADGROUP_SIZE_X THREADGROUP_SIZE
+#endif
+
+#ifndef THREADGROUP_SIZE_Y
+#define THREADGROUP_SIZE_Y THREADGROUP_SIZE
+#endif
+
+#ifndef THREADGROUP_SIZE_Z
+#define THREADGROUP_SIZE_Z THREADGROUP_SIZE
+#endif
+
 [numthreads(THREADGROUP_SIZE, THREADGROUP_SIZE, THREADGROUP_SIZE)]
 void MaterialSetupCS( 
 	uint3 GroupId : SV_GroupID,
@@ -165,6 +177,10 @@ uint HistoryMissSuperSampleCount;
 float PhaseG;
 float InverseSquaredLightDistanceBiasScale;
 
+#ifndef HISTORY_MISS_SUPER_SAMPLE_COUNT
+#define HISTORY_MISS_SUPER_SAMPLE_COUNT 1
+#endif
+
 #ifndef INVERSE_SQUARED_FALLOFF
 #define INVERSE_SQUARED_FALLOFF 1
 #endif
@@ -428,7 +444,7 @@ float UseHeightFogColors;
 float UseDirectionalLightShadowing;
 float StaticLightingScatteringIntensity;
 
-[numthreads(THREADGROUP_SIZE, THREADGROUP_SIZE, THREADGROUP_SIZE)]
+[numthreads(THREADGROUP_SIZE_X, THREADGROUP_SIZE_Y, THREADGROUP_SIZE_Z)]
 void LightScatteringCS(
 	uint3 GroupId : SV_GroupID,
 	uint3 DispatchThreadId : SV_DispatchThreadID,
@@ -450,10 +466,10 @@ void LightScatteringCS(
 
 	// Supersample if the history was outside the camera frustum
 	// The compute shader is dispatched with extra threads, make sure those don't supersample
-	NumSuperSamples = HistoryAlpha < .001f && all(GridCoordinate < VolumetricFog.GridSizeInt) ? HistoryMissSuperSampleCount : 1;
+	NumSuperSamples = HistoryAlpha < .001f && all(GridCoordinate < VolumetricFog.GridSizeInt) ? HISTORY_MISS_SUPER_SAMPLE_COUNT : 1;
 
 #endif
-
+	UNROLL
 	for (uint SampleIndex = 0; SampleIndex < NumSuperSamples; SampleIndex++)
 	{
 		float3 CellOffset = FrameJitterOffsets[SampleIndex].xyz;
@@ -542,7 +558,18 @@ void LightScatteringCS(
 			LightScattering += (StaticLightingScatteringIntensity / PI) * max(DotSH(IrradianceSH, RotatedHGZonalHarmonic), 0);
 		}
 #endif
-		
+	}
+
+	for (uint SampleIndex = 0; SampleIndex < NumSuperSamples; SampleIndex++)
+	{
+		float3 CellOffset = FrameJitterOffsets[SampleIndex].xyz;
+		//float CellOffset = .5f;
+
+		float SceneDepth;
+		float3 WorldPosition = ComputeCellWorldPosition(GridCoordinate, CellOffset, SceneDepth);
+		float CameraVectorLength = length(WorldPosition - View.WorldCameraOrigin);
+		float3 CameraVector = (WorldPosition - View.WorldCameraOrigin) / CameraVectorLength;
+
 		uint GridIndex = ComputeLightGridCellIndex(GridCoordinate.xy * VolumetricFog.FogGridToPixelXY, SceneDepth, 0);
 		const FCulledLightsGridData CulledLightsGrid = GetCulledLightsGrid(GridIndex, 0);
 
diff --git a/Engine/Source/Runtime/Core/Public/Math/IntVector.h b/Engine/Source/Runtime/Core/Public/Math/IntVector.h
--- a/Engine/Source/Runtime/Core/Public/Math/IntVector.h
+++ b/Engine/Source/Runtime/Core/Public/Math/IntVector.h
@@ -275,6 +275,7 @@ public:
 	 * @return A new divided int point.
 	 */
 	static FIntVector DivideAndRoundUp( FIntVector lhs, int32 Divisor );
+	static FIntVector DivideAndRoundUp( FIntVector lhs, FIntVector Divisor );
 
 	/**
 	 * Gets the number of components a point has.
@@ -478,6 +479,12 @@ FORCEINLINE FIntVector FIntVector::DivideAndRoundUp( FIntVector lhs, int32 Divis
 	return FIntVector(FMath::DivideAndRoundUp(lhs.X, Divisor), FMath::DivideAndRoundUp(lhs.Y, Divisor), FMath::DivideAndRoundUp(lhs.Z, Divisor));
 }
 
+FORCEINLINE FIntVector FIntVector::DivideAndRoundUp(FIntVector lhs, FIntVector Divisor)
+{
+	return FIntVector(FMath::DivideAndRoundUp(lhs.X, Divisor.X), FMath::DivideAndRoundUp(lhs.Y, Divisor.Y), FMath::DivideAndRoundUp(lhs.Z, Divisor.Z));
+}
+
+
 FORCEINLINE int32 FIntVector::GetMax() const
 {
 	return FMath::Max(FMath::Max(X, Y), Z);
diff --git a/Engine/Source/Runtime/Renderer/Private/VolumetricFog.cpp b/Engine/Source/Runtime/Renderer/Private/VolumetricFog.cpp
--- a/Engine/Source/Runtime/Renderer/Private/VolumetricFog.cpp
+++ b/Engine/Source/Runtime/Renderer/Private/VolumetricFog.cpp
@@ -121,7 +121,10 @@ FVector VolumetricFogTemporalRandom(uint32 FrameNumber)
 	return RandomOffsetValue;
 }
 
-uint32 VolumetricFogGridInjectionGroupSize = 4;
+static const uint32 VolumetricFogGridInjectionGroupSize  = 4;
+static const uint32 VolumetricFogLightScatteringGroupSizeX = 8;
+static const uint32 VolumetricFogLightScatteringGroupSizeY = 8;
+static const uint32 VolumetricFogLightScatteringGroupSizeZ = 1;
 
 class FVolumetricFogMaterialSetupCS : public FGlobalShader
 {
@@ -556,8 +559,10 @@ class TVolumetricFogLightScatteringCS : public FGlobalShader
 
 	class FTemporalReprojection			: SHADER_PERMUTATION_BOOL("USE_TEMPORAL_REPROJECTION");
 	class FDistanceFieldSkyOcclusion	: SHADER_PERMUTATION_BOOL("DISTANCE_FIELD_SKY_OCCLUSION");
+	class FSuperSampleCount				: SHADER_PERMUTATION_RANGE_INT("HISTORY_MISS_SUPER_SAMPLE_COUNT",1,16);
 
 	using FPermutationDomain = TShaderPermutationDomain<
+		FSuperSampleCount,
 		FTemporalReprojection,
 		FDistanceFieldSkyOcclusion >;
 
@@ -580,7 +585,9 @@ public:
 	static void ModifyCompilationEnvironment(const FGlobalShaderPermutationParameters& Parameters, FShaderCompilerEnvironment& OutEnvironment)
 	{
 		FGlobalShader::ModifyCompilationEnvironment(Parameters, OutEnvironment);
-		OutEnvironment.SetDefine(TEXT("THREADGROUP_SIZE"), VolumetricFogGridInjectionGroupSize);
+		OutEnvironment.SetDefine(TEXT("THREADGROUP_SIZE_X"), VolumetricFogLightScatteringGroupSizeX);
+		OutEnvironment.SetDefine(TEXT("THREADGROUP_SIZE_Y"), VolumetricFogLightScatteringGroupSizeY);
+		OutEnvironment.SetDefine(TEXT("THREADGROUP_SIZE_Z"), VolumetricFogLightScatteringGroupSizeZ);
 		FVolumetricFogIntegrationParameters::ModifyCompilationEnvironment(Parameters.Platform, OutEnvironment);
 		FForwardLightingParameters::ModifyCompilationEnvironment(Parameters.Platform, OutEnvironment);
 	}
@@ -1111,15 +1118,17 @@ void FDeferredShadingSceneRenderer::ComputeVolumetricFog(FRHICommandListImmediat
 				TVolumetricFogLightScatteringCS::FPermutationDomain PermutationVector;
 				PermutationVector.Set< TVolumetricFogLightScatteringCS::FTemporalReprojection >(bUseTemporalReprojection);
 				PermutationVector.Set< TVolumetricFogLightScatteringCS::FDistanceFieldSkyOcclusion >(bUseDistanceFieldSkyOcclusion);
+				PermutationVector.Set< TVolumetricFogLightScatteringCS::FSuperSampleCount >(GVolumetricFogHistoryMissSupersampleCount);
 
 				auto ComputeShader = View.ShaderMap->GetShader< TVolumetricFogLightScatteringCS >(PermutationVector);
 				ClearUnusedGraphResources(ComputeShader, PassParameters);
 
 				GraphBuilder.AddPass(
-					RDG_EVENT_NAME("LightScattering %dx%dx%d %s %s",
+					RDG_EVENT_NAME("LightScattering %dx%dx%d SS:%d %s %s",
 						VolumetricFogGridSize.X,
 						VolumetricFogGridSize.Y,
 						VolumetricFogGridSize.Z,
+						GVolumetricFogHistoryMissSupersampleCount,
 						bUseDistanceFieldSkyOcclusion ? TEXT("DFAO") : TEXT(""),
 						PassParameters->LightFunctionTexture ? TEXT("LF") : TEXT("")),
 					PassParameters,
@@ -1127,7 +1136,7 @@ void FDeferredShadingSceneRenderer::ComputeVolumetricFog(FRHICommandListImmediat
 					[PassParameters, ComputeShader, &View, this, FogInfo, bUseTemporalReprojection, VolumetricFogGridSize, IntegrationData, bUseDirectionalLightShadowing, bUseDistanceFieldSkyOcclusion, LightFunctionWorldToShadow](FRHICommandListImmediate& RHICmdList)
 				{
 					UnbindRenderTargets(RHICmdList);
-					const FIntVector NumGroups = FIntVector::DivideAndRoundUp(VolumetricFogGridSize, VolumetricFogGridInjectionGroupSize);
+					const FIntVector NumGroups = FIntVector::DivideAndRoundUp(VolumetricFogGridSize, FIntVector(VolumetricFogLightScatteringGroupSizeX, VolumetricFogLightScatteringGroupSizeY, VolumetricFogLightScatteringGroupSizeZ));
 
 					RHICmdList.SetComputeShader(ComputeShader->GetComputeShader());
 
@@ -1183,4 +1192,4 @@ void FDeferredShadingSceneRenderer::ComputeVolumetricFog(FRHICommandListImmediat
 
 		GraphBuilder.Execute();
 	}
-}
\ No newline at end of file
+}
