Added underwater translucency pass.

diff --git a/Engine/Source/Runtime/Engine/Classes/Materials/Material.h b/Engine/Source/Runtime/Engine/Classes/Materials/Material.h
--- a/Engine/Source/Runtime/Engine/Classes/Materials/Material.h
+++ b/Engine/Source/Runtime/Engine/Classes/Materials/Material.h
@@ -494,6 +494,10 @@ public:
 	UPROPERTY(EditAnywhere, Category=Translucency, meta=(DisplayName = "Render After DOF"), AdvancedDisplay)
 	uint8 bEnableSeparateTranslucency : 1;
 
+	/** Indicates that the material should be rendered in the underwater pass. */
+	UPROPERTY(EditAnywhere, Category=Translucency, meta=(DisplayName = "Render Under Water"), AdvancedDisplay)
+	uint8 bEnableRenderUnderWater : 1;
+
 	/**
 	 * Indicates that the material should be rendered using responsive anti-aliasing. Improves sharpness of small moving particles such as sparks.
 	 * Only use for small moving features because it will cause aliasing of the background.
diff --git a/Engine/Source/Runtime/Engine/Classes/Materials/MaterialInterface.h b/Engine/Source/Runtime/Engine/Classes/Materials/MaterialInterface.h
--- a/Engine/Source/Runtime/Engine/Classes/Materials/MaterialInterface.h
+++ b/Engine/Source/Runtime/Engine/Classes/Materials/MaterialInterface.h
@@ -67,6 +67,7 @@ struct ENGINE_API FMaterialRelevance
 	uint8 bMasked : 1;
 	uint8 bDistortion : 1;
 	uint8 bSeparateTranslucency : 1; // Translucency After DOF
+	uint8 bUnderWaterTranslucency : 1;
 	uint8 bNormalTranslucency : 1;
 	uint8 bUsesSceneColorCopy : 1;
 	uint8 bDisableOffscreenRendering : 1; // Blend Modulate
diff --git a/Engine/Source/Runtime/Engine/Private/Materials/HLSLMaterialTranslator.h b/Engine/Source/Runtime/Engine/Private/Materials/HLSLMaterialTranslator.h
--- a/Engine/Source/Runtime/Engine/Private/Materials/HLSLMaterialTranslator.h
+++ b/Engine/Source/Runtime/Engine/Private/Materials/HLSLMaterialTranslator.h
@@ -1167,6 +1167,11 @@ ResourcesString = TEXT("");
 				}
 			}
 
+			if (IsTranslucentBlendMode(BlendMode) && Material->IsTranslucencyUnderWaterEnabled() && (Material->IsMobileSeparateTranslucencyEnabled() || Material->IsTranslucencyAfterDOFEnabled()))
+			{
+				Errorf(TEXT("A material cannot be sent to a separate translucent pass and under water at the same time."));
+			}
+
 			bool bDBufferAllowed = IsUsingDBuffers(Platform);
 			bool bDBufferBlendMode = IsDBufferDecalBlendMode((EDecalBlendMode)Material->GetDecalBlendMode());
 
diff --git a/Engine/Source/Runtime/Engine/Private/Materials/Material.cpp b/Engine/Source/Runtime/Engine/Private/Materials/Material.cpp
--- a/Engine/Source/Runtime/Engine/Private/Materials/Material.cpp
+++ b/Engine/Source/Runtime/Engine/Private/Materials/Material.cpp
@@ -4287,6 +4287,7 @@ bool UMaterial::CanEditChange(const UProperty* InProperty) const
 		}
 
 		if (PropertyName == GET_MEMBER_NAME_STRING_CHECKED(UMaterial, bEnableSeparateTranslucency)
+			|| PropertyName == GET_MEMBER_NAME_STRING_CHECKED(UMaterial, bEnableRenderUnderWater)
 			|| PropertyName == GET_MEMBER_NAME_STRING_CHECKED(UMaterial, bEnableResponsiveAA)
 			|| PropertyName == GET_MEMBER_NAME_STRING_CHECKED(UMaterial, bScreenSpaceReflections)
 			|| PropertyName == GET_MEMBER_NAME_STRING_CHECKED(UMaterial, bContactShadows)
diff --git a/Engine/Source/Runtime/Engine/Private/Materials/MaterialInterface.cpp b/Engine/Source/Runtime/Engine/Private/Materials/MaterialInterface.cpp
--- a/Engine/Source/Runtime/Engine/Private/Materials/MaterialInterface.cpp
+++ b/Engine/Source/Runtime/Engine/Private/Materials/MaterialInterface.cpp
@@ -42,6 +42,7 @@ void FMaterialRelevance::SetPrimitiveViewRelevance(FPrimitiveViewRelevance& OutV
 	OutViewRelevance.bOutputsTranslucentVelocityRelevance = bOutputsTranslucentVelocity;
 	OutViewRelevance.bDistortionRelevance = bDistortion;
 	OutViewRelevance.bSeparateTranslucencyRelevance = bSeparateTranslucency;
+	OutViewRelevance.bUnderWaterTranslucencyRelevance = bUnderWaterTranslucency;
 	OutViewRelevance.bNormalTranslucencyRelevance = bNormalTranslucency;
 	OutViewRelevance.bUsesSceneColorCopy = bUsesSceneColorCopy;
 	OutViewRelevance.bDisableOffscreenRendering = bDisableOffscreenRendering;
@@ -147,12 +148,14 @@ FMaterialRelevance UMaterialInterface::GetRelevance_Internal(const UMaterial* Ma
 		{
 			// Check whether the material can be drawn in the separate translucency pass as per FMaterialResource::IsTranslucencyAfterDOFEnabled and IsMobileSeparateTranslucencyEnabled
 			bool bSupportsSeparateTranclucency = Material->MaterialDomain != MD_UI && Material->MaterialDomain != MD_DeferredDecal;
-			bool bMaterialSeparateTranclucency = bSupportsSeparateTranclucency && (bIsMobile ? Material->bEnableMobileSeparateTranslucency : Material->bEnableSeparateTranslucency );
+			bool bMaterialSeparateTranclucency = bSupportsSeparateTranclucency && (bIsMobile ? Material->bEnableMobileSeparateTranslucency : Material->bEnableSeparateTranslucency);
+			bool bMaterialUnderWaterTranslucency = bSupportsSeparateTranclucency && Material->bEnableRenderUnderWater;
 			
 			MaterialRelevance.bOpaque = !bIsTranslucent;
 			MaterialRelevance.bMasked = IsMasked();
 			MaterialRelevance.bDistortion = MaterialResource->IsDistorted();
 			MaterialRelevance.bSeparateTranslucency = bIsTranslucent && bMaterialSeparateTranclucency;
+			MaterialRelevance.bUnderWaterTranslucency = bIsTranslucent && bMaterialUnderWaterTranslucency;
 			MaterialRelevance.bNormalTranslucency = bIsTranslucent && !bMaterialSeparateTranclucency;
 			MaterialRelevance.bDisableDepthTest = bIsTranslucent && Material->bDisableDepthTest;		
 			MaterialRelevance.bUsesSceneColorCopy = bIsTranslucent && MaterialResource->RequiresSceneColorCopy_GameThread();
diff --git a/Engine/Source/Runtime/Engine/Private/Materials/MaterialShared.cpp b/Engine/Source/Runtime/Engine/Private/Materials/MaterialShared.cpp
--- a/Engine/Source/Runtime/Engine/Private/Materials/MaterialShared.cpp
+++ b/Engine/Source/Runtime/Engine/Private/Materials/MaterialShared.cpp
@@ -1216,6 +1216,11 @@ bool FMaterialResource::IsTranslucencyAfterDOFEnabled() const
 	return Material->bEnableSeparateTranslucency && !IsUIMaterial() && !IsDeferredDecal();
 }
 
+bool FMaterialResource::IsTranslucencyUnderWaterEnabled() const
+{
+	return Material->bEnableRenderUnderWater && !IsUIMaterial() && !IsDeferredDecal();
+}
+
 bool FMaterialResource::IsMobileSeparateTranslucencyEnabled() const
 {
 	return Material->bEnableMobileSeparateTranslucency && !IsUIMaterial() && !IsDeferredDecal();
diff --git a/Engine/Source/Runtime/Engine/Public/MaterialShared.h b/Engine/Source/Runtime/Engine/Public/MaterialShared.h
--- a/Engine/Source/Runtime/Engine/Public/MaterialShared.h
+++ b/Engine/Source/Runtime/Engine/Public/MaterialShared.h
@@ -1493,6 +1493,7 @@ public:
 	virtual float GetTranslucentSelfShadowSecondOpacity() const { return 1.0f; }
 	virtual float GetTranslucentBackscatteringExponent() const { return 1.0f; }
 	virtual bool IsTranslucencyAfterDOFEnabled() const { return false; }
+	virtual bool IsTranslucencyUnderWaterEnabled() const { return false; }
 	virtual bool IsMobileSeparateTranslucencyEnabled() const { return false; }
 	virtual FLinearColor GetTranslucentMultipleScatteringExtinction() const { return FLinearColor::White; }
 	virtual float GetTranslucentShadowStartOffset() const { return 0.0f; }
@@ -2250,6 +2251,7 @@ public:
 	ENGINE_API virtual float GetTranslucentSelfShadowSecondOpacity() const override;
 	ENGINE_API virtual float GetTranslucentBackscatteringExponent() const override;
 	ENGINE_API virtual bool IsTranslucencyAfterDOFEnabled() const override;
+	ENGINE_API virtual bool IsTranslucencyUnderWaterEnabled() const override;
 	ENGINE_API virtual bool IsMobileSeparateTranslucencyEnabled() const override;
 	ENGINE_API virtual FLinearColor GetTranslucentMultipleScatteringExtinction() const override;
 	ENGINE_API virtual float GetTranslucentShadowStartOffset() const override;
diff --git a/Engine/Source/Runtime/Engine/Public/PrimitiveViewRelevance.h b/Engine/Source/Runtime/Engine/Public/PrimitiveViewRelevance.h
--- a/Engine/Source/Runtime/Engine/Public/PrimitiveViewRelevance.h
+++ b/Engine/Source/Runtime/Engine/Public/PrimitiveViewRelevance.h
@@ -30,6 +30,8 @@ struct FPrimitiveViewRelevance
 	uint32 bSeparateTranslucencyRelevance : 1;
 	/** The primitive has one or more elements that have normal translucency. */
 	uint32 bNormalTranslucencyRelevance : 1;
+	/** The primitive has one or more elements that have normal translucency. */
+	uint32 bUnderWaterTranslucencyRelevance : 1;
 	/** For translucent primitives reading the scene color. */
 	uint32 bUsesSceneColorCopy : 1;
 	/** For primitive that can't render in offscreen buffers (blend modulate). */
@@ -91,7 +93,7 @@ struct FPrimitiveViewRelevance
 
 	bool HasTranslucency() const 
 	{
-		return bSeparateTranslucencyRelevance || bNormalTranslucencyRelevance;
+		return bSeparateTranslucencyRelevance || bNormalTranslucencyRelevance || bUnderWaterTranslucencyRelevance;
 	}
 
 	/** Default constructor */
diff --git a/Engine/Source/Runtime/Renderer/Private/BasePassRendering.cpp b/Engine/Source/Runtime/Renderer/Private/BasePassRendering.cpp
--- a/Engine/Source/Runtime/Renderer/Private/BasePassRendering.cpp
+++ b/Engine/Source/Runtime/Renderer/Private/BasePassRendering.cpp
@@ -1217,11 +1217,15 @@ void FBasePassMeshProcessor::AddMeshBatch(const FMeshBatch& RESTRICT MeshBatch,
 				switch (TranslucencyPassType)
 				{
 				case ETranslucencyPass::TPT_StandardTranslucency:
-					bShouldDraw = !Material.IsTranslucencyAfterDOFEnabled();
+					bShouldDraw = !Material.IsTranslucencyAfterDOFEnabled() && !Material.IsTranslucencyUnderWaterEnabled();
 					break;
 
 				case ETranslucencyPass::TPT_TranslucencyAfterDOF:
-					bShouldDraw = Material.IsTranslucencyAfterDOFEnabled();
+					bShouldDraw = Material.IsTranslucencyAfterDOFEnabled() && !Material.IsTranslucencyUnderWaterEnabled();
+					break;
+
+				case ETranslucencyPass::TPT_TranslucencyUnderWater:
+					bShouldDraw = Material.IsTranslucencyUnderWaterEnabled();
 					break;
 
 				case ETranslucencyPass::TPT_AllTranslucency:
@@ -1603,6 +1607,17 @@ FMeshPassProcessor* CreateTranslucencyAfterDOFProcessor(const FScene* Scene, con
 	return new(FMemStack::Get()) FBasePassMeshProcessor(Scene, Scene->GetFeatureLevel(), InViewIfDynamicMeshCommand, PassDrawRenderState, InDrawListContext, Flags, ETranslucencyPass::TPT_TranslucencyAfterDOF);
 }
 
+FMeshPassProcessor* CreateTranslucencyUnderWaterProcessor(const FScene* Scene, const FSceneView* InViewIfDynamicMeshCommand, FMeshPassDrawListContext* InDrawListContext)
+{
+	FMeshPassProcessorRenderState PassDrawRenderState(Scene->UniformBuffers.ViewUniformBuffer, Scene->UniformBuffers.TranslucentBasePassUniformBuffer);
+	PassDrawRenderState.SetInstancedViewUniformBuffer(Scene->UniformBuffers.InstancedViewUniformBuffer);
+	PassDrawRenderState.SetDepthStencilState(TStaticDepthStencilState<false, CF_DepthNearOrEqual>::GetRHI());
+
+	const FBasePassMeshProcessor::EFlags Flags = FBasePassMeshProcessor::EFlags::CanUseDepthStencil;
+
+	return new(FMemStack::Get()) FBasePassMeshProcessor(Scene, Scene->GetFeatureLevel(), InViewIfDynamicMeshCommand, PassDrawRenderState, InDrawListContext, Flags, ETranslucencyPass::TPT_TranslucencyUnderWater);
+}
+
 FMeshPassProcessor* CreateTranslucencyAllPassProcessor(const FScene* Scene, const FSceneView* InViewIfDynamicMeshCommand, FMeshPassDrawListContext* InDrawListContext)
 {
 	FMeshPassProcessorRenderState PassDrawRenderState(Scene->UniformBuffers.ViewUniformBuffer, Scene->UniformBuffers.TranslucentBasePassUniformBuffer);
@@ -1617,4 +1632,5 @@ FMeshPassProcessor* CreateTranslucencyAllPassProcessor(const FScene* Scene, cons
 FRegisterPassProcessorCreateFunction RegisterBasePass(&CreateBasePassProcessor, EShadingPath::Deferred, EMeshPass::BasePass, EMeshPassFlags::CachedMeshCommands | EMeshPassFlags::MainView);
 FRegisterPassProcessorCreateFunction RegisterTranslucencyStandardPass(&CreateTranslucencyStandardPassProcessor, EShadingPath::Deferred, EMeshPass::TranslucencyStandard, EMeshPassFlags::MainView);
 FRegisterPassProcessorCreateFunction RegisterTranslucencyAfterDOFPass(&CreateTranslucencyAfterDOFProcessor, EShadingPath::Deferred, EMeshPass::TranslucencyAfterDOF, EMeshPassFlags::MainView);
+FRegisterPassProcessorCreateFunction RegisterTranslucencyUnderWaterPass(&CreateTranslucencyUnderWaterProcessor, EShadingPath::Deferred, EMeshPass::TranslucencyUnderWater, EMeshPassFlags::MainView);
 FRegisterPassProcessorCreateFunction RegisterTranslucencyAllPass(&CreateTranslucencyAllPassProcessor, EShadingPath::Deferred, EMeshPass::TranslucencyAll, EMeshPassFlags::MainView);
\ No newline at end of file
diff --git a/Engine/Source/Runtime/Renderer/Private/DeferredShadingRenderer.cpp b/Engine/Source/Runtime/Renderer/Private/DeferredShadingRenderer.cpp
--- a/Engine/Source/Runtime/Renderer/Private/DeferredShadingRenderer.cpp
+++ b/Engine/Source/Runtime/Renderer/Private/DeferredShadingRenderer.cpp
@@ -1792,19 +1792,30 @@ void FDeferredShadingSceneRenderer::Render(FRHICommandListImmediate& RHICmdList)
 		SCOPED_DRAW_EVENTF(RHICmdList, WaterRendering, TEXT("WaterRendering"));
 		SCOPED_GPU_STAT(RHICmdList, WaterRendering);
 
-		RHICmdList.TransitionResource(EResourceTransitionAccess::EWritable,  SceneContext.GetSceneDepthSurface());
+		// First render the under water particles. We only support a black readable scene color for now.
+		// TODO: this should be done after the water GBuffer pass when the camera is under water.
+		IPooledRenderTarget* DummyReadBackSceneTexture = GSystemTextures.BlackDummy;
+		RenderTranslucency(RHICmdList, ETranslucencyPass::TPT_TranslucencyUnderWater, DummyReadBackSceneTexture);
 
+		// Copy the texture to be available for the water surface to refrace
 		FSingleLayerWaterPassData SingleLayerWaterPassData;
+		CopySingleLayerWaterTextures(RHICmdList, SingleLayerWaterPassData);
+
+		// Make the Depth texture writable since the water GBuffer pass will update it
+		RHICmdList.TransitionResource(EResourceTransitionAccess::EWritable,  SceneContext.GetSceneDepthSurface());
 
+		// Render the GBuffer pass, updating the GBuffer and also writing lit water in the scene.
 		const FExclusiveDepthStencil::Type WaterPassDepthStencilAccess = FExclusiveDepthStencil::DepthWrite_StencilWrite;
 		BeginRenderingWaterGBuffer(RHICmdList, SingleLayerWaterPassData, WaterPassDepthStencilAccess, ViewFamily.EngineShowFlags.ShaderComplexity);
 		RenderSingleLayerWaterPass(RHICmdList, SingleLayerWaterPassData, WaterPassDepthStencilAccess);
 		FinishWaterGBufferPassAndResolve(RHICmdList);
 
+		// Resolves the depth texture back to readable for SSR and later passes.1
 		SceneContext.ResolveSceneDepthTexture(RHICmdList, FResolveRect(0, 0, FamilySize.X, FamilySize.Y));
 		SceneContext.ResolveSceneDepthToAuxiliaryTexture(RHICmdList);
 		RHICmdList.TransitionResource(EResourceTransitionAccess::EReadable, SceneContext.GetSceneDepthSurface());
 
+		// If supported render SSR, the composite pass in non deferred and/or under water effect.
 		RenderSingleLayerWaterReflections(RHICmdList, SingleLayerWaterPassData);
 		ServiceLocalQueue();
 	}
@@ -1956,6 +1967,18 @@ void FDeferredShadingSceneRenderer::Render(FRHICommandListImmediate& RHICmdList)
 			// Disable UAV cache flushing so we have optimal VT feedback performance.
 			RHICmdList.AutomaticCacheFlushAfterComputeShader(false);
 
+			if (!bShouldRenderSingleLayerWater)
+			{
+				// We render under water transparent separately even when water is disabled (simple, and avoid checking if water is enabled everywhere)
+				// This translucency pass is cheap and fast to render but could be improved 
+				//  - it is always full screen todya.
+				//  - particles do not have correct under water fog apply on them. it is the backgroung fog that is applied on them.
+				//  - particle are not cut out by the top part of the depth layer
+				// This could be resolved by rendering particles in an offscreenbuffer with corret fog and applied in the SingleLayerWater gbuffer pass.
+				// A depth buffer from previous frame could also be sampled to clip particle.
+				RenderTranslucency(RHICmdList, ETranslucencyPass::TPT_TranslucencyUnderWater, SceneColorCopy);
+			}
+
 			if (ViewFamily.AllowTranslucencyAfterDOF())
 			{
 				RenderTranslucency(RHICmdList, ETranslucencyPass::TPT_StandardTranslucency, SceneColorCopy);
diff --git a/Engine/Source/Runtime/Renderer/Private/DeferredShadingRenderer.h b/Engine/Source/Runtime/Renderer/Private/DeferredShadingRenderer.h
--- a/Engine/Source/Runtime/Renderer/Private/DeferredShadingRenderer.h
+++ b/Engine/Source/Runtime/Renderer/Private/DeferredShadingRenderer.h
@@ -103,6 +103,8 @@ public:
 	*/
 	bool RenderBasePass(FRHICommandListImmediate& RHICmdList, FExclusiveDepthStencil::Type BasePassDepthStencilAccess, IPooledRenderTarget* ForwardScreenSpaceShadowMask, bool bParallelBasePass, bool bRenderLightmapDensity);
 
+	/** Copy the scene color and depth buffer for water refraction*/
+	void CopySingleLayerWaterTextures(FRHICommandList& RHICmdList, FSingleLayerWaterPassData& PassData);
 	/** Begin the water GBuffer rendering pass*/
 	void BeginRenderingWaterGBuffer(FRHICommandList& RHICmdList, FSingleLayerWaterPassData& PassData, FExclusiveDepthStencil::Type DepthStencilAccess, bool bBindQuadOverdrawBuffers);
 	/** End the water GBuffer rendering pass*/
diff --git a/Engine/Source/Runtime/Renderer/Private/MeshDrawCommands.cpp b/Engine/Source/Runtime/Renderer/Private/MeshDrawCommands.cpp
--- a/Engine/Source/Runtime/Renderer/Private/MeshDrawCommands.cpp
+++ b/Engine/Source/Runtime/Renderer/Private/MeshDrawCommands.cpp
@@ -968,6 +968,7 @@ void FParallelMeshDrawCommandPass::DispatchPassSetup(
 	{
 		case EMeshPass::TranslucencyStandard: TaskContext.TranslucencyPass = ETranslucencyPass::TPT_StandardTranslucency; break;
 		case EMeshPass::TranslucencyAfterDOF: TaskContext.TranslucencyPass = ETranslucencyPass::TPT_TranslucencyAfterDOF; break;
+		case EMeshPass::TranslucencyUnderWater: TaskContext.TranslucencyPass = ETranslucencyPass::TPT_TranslucencyUnderWater; break;
 		case EMeshPass::TranslucencyAll: TaskContext.TranslucencyPass = ETranslucencyPass::TPT_AllTranslucency; break;
 		case EMeshPass::MobileInverseOpacity: TaskContext.TranslucencyPass = ETranslucencyPass::TPT_StandardTranslucency; break;
 	}
diff --git a/Engine/Source/Runtime/Renderer/Private/MobileBasePass.cpp b/Engine/Source/Runtime/Renderer/Private/MobileBasePass.cpp
--- a/Engine/Source/Runtime/Renderer/Private/MobileBasePass.cpp
+++ b/Engine/Source/Runtime/Renderer/Private/MobileBasePass.cpp
@@ -577,7 +577,8 @@ void FMobileBasePassMeshProcessor::AddMeshBatch(const FMeshBatch& RESTRICT MeshB
 		bool bShouldDraw = (bIsTranslucent || bUsesWaterMaterial) &&
 		(TranslucencyPassType == ETranslucencyPass::TPT_AllTranslucency
 		|| (TranslucencyPassType == ETranslucencyPass::TPT_StandardTranslucency && !Material.IsMobileSeparateTranslucencyEnabled())
-		|| (TranslucencyPassType == ETranslucencyPass::TPT_TranslucencyAfterDOF && Material.IsMobileSeparateTranslucencyEnabled()));
+		|| (TranslucencyPassType == ETranslucencyPass::TPT_TranslucencyAfterDOF && Material.IsMobileSeparateTranslucencyEnabled())
+		|| (TranslucencyPassType == ETranslucencyPass::TPT_TranslucencyUnderWater && Material.IsTranslucencyUnderWaterEnabled()));
 							
 		if (bShouldDraw)
 		{
diff --git a/Engine/Source/Runtime/Renderer/Private/SceneVisibility.cpp b/Engine/Source/Runtime/Renderer/Private/SceneVisibility.cpp
--- a/Engine/Source/Runtime/Renderer/Private/SceneVisibility.cpp
+++ b/Engine/Source/Runtime/Renderer/Private/SceneVisibility.cpp
@@ -2013,7 +2013,14 @@ struct FRelevancePacket
 
 			if (bTranslucentRelevance && !bEditorRelevance && ViewRelevance.bRenderInMainPass)
 			{
-				if (View.Family->AllowTranslucencyAfterDOF())
+				if (ViewRelevance.bUnderWaterTranslucencyRelevance)
+				{
+					// Translucency under water do not need to use scene color resolve (a black texture is provided) and we do not need to render to another off-screen buffer.
+					const bool bUsesSceneColorCopy = false;
+					const bool bDisableOffscreenRendering = true;
+					TranslucentPrimCount.Add(ETranslucencyPass::TPT_TranslucencyUnderWater, bUsesSceneColorCopy, bDisableOffscreenRendering);
+				}
+				else if (View.Family->AllowTranslucencyAfterDOF())
 				{
 					if (ViewRelevance.bNormalTranslucencyRelevance)
 					{
@@ -2307,6 +2314,11 @@ struct FRelevancePacket
 								{
 									DrawCommandPacket.AddCommandsForMesh(PrimitiveIndex, PrimitiveSceneInfo, StaticMeshRelevance, StaticMesh, Scene, bCanCache, EMeshPass::TranslucencyAfterDOF);
 								}
+
+								if (ViewRelevance.bUnderWaterTranslucencyRelevance)
+								{
+									DrawCommandPacket.AddCommandsForMesh(PrimitiveIndex, PrimitiveSceneInfo, StaticMeshRelevance, StaticMesh, Scene, bCanCache, EMeshPass::TranslucencyUnderWater);
+								}
 							}
 							else
 							{
@@ -2724,6 +2736,12 @@ void ComputeDynamicMeshRelevance(EShadingPath ShadingPath, bool bAddLightmapDens
 				PassMask.Set(EMeshPass::TranslucencyAfterDOF);
 				View.NumVisibleDynamicMeshElements[EMeshPass::TranslucencyAfterDOF] += NumElements;
 			}
+
+			if (ViewRelevance.bUnderWaterTranslucencyRelevance)
+			{
+				PassMask.Set(EMeshPass::TranslucencyUnderWater);
+				View.NumVisibleDynamicMeshElements[EMeshPass::TranslucencyUnderWater] += NumElements;
+			}
 		}
 		else
 		{
diff --git a/Engine/Source/Runtime/Renderer/Private/SingleLayerWaterRendering.cpp b/Engine/Source/Runtime/Renderer/Private/SingleLayerWaterRendering.cpp
--- a/Engine/Source/Runtime/Renderer/Private/SingleLayerWaterRendering.cpp
+++ b/Engine/Source/Runtime/Renderer/Private/SingleLayerWaterRendering.cpp
@@ -245,17 +245,11 @@ IMPLEMENT_GLOBAL_SHADER(FSingleLayerWaterCompositePS, "/Engine/Private/SingleLay
 
 
 
-void FDeferredShadingSceneRenderer::BeginRenderingWaterGBuffer(FRHICommandList& RHICmdList, FSingleLayerWaterPassData& PassData, FExclusiveDepthStencil::Type DepthStencilAccess, bool bBindQuadOverdrawBuffers)
+void FDeferredShadingSceneRenderer::CopySingleLayerWaterTextures(FRHICommandList& RHICmdList, FSingleLayerWaterPassData& PassData)
 {
-	SCOPED_DRAW_EVENT(RHICmdList, BeginRenderingWaterGBuffer);
-
 	check(RHICmdList.IsOutsideRenderPass());
 	FSceneRenderTargets& SceneContext = FSceneRenderTargets::Get(RHICmdList);
-
 	const ERHIFeatureLevel::Type CurrentFeatureLevel = SceneContext.GetCurrentFeatureLevel();
-	const bool bUseGBuffer = IsUsingGBuffers(ShaderPlatform);
-	check(CurrentFeatureLevel >= ERHIFeatureLevel::SM5);
-
 
 	// Allocate required buffers
 	{
@@ -276,6 +270,18 @@ void FDeferredShadingSceneRenderer::BeginRenderingWaterGBuffer(FRHICommandList&
 		RHICmdList.CopyTexture(SceneContext.SceneDepthZ->GetRenderTargetItem().ShaderResourceTexture, PassData.SceneDepthZWithoutSingleLayerWater->GetRenderTargetItem().ShaderResourceTexture, DepthCopyInfo);
 		RHICmdList.CopyTexture(SceneContext.GetSceneColorSurface(), PassData.SceneColorWithoutSingleLayerWater->GetRenderTargetItem().ShaderResourceTexture, DepthCopyInfo);
 	}
+}
+
+void FDeferredShadingSceneRenderer::BeginRenderingWaterGBuffer(FRHICommandList& RHICmdList, FSingleLayerWaterPassData& PassData, FExclusiveDepthStencil::Type DepthStencilAccess, bool bBindQuadOverdrawBuffers)
+{
+	SCOPED_DRAW_EVENT(RHICmdList, BeginRenderingWaterGBuffer);
+
+	check(RHICmdList.IsOutsideRenderPass());
+	FSceneRenderTargets& SceneContext = FSceneRenderTargets::Get(RHICmdList);
+
+	const ERHIFeatureLevel::Type CurrentFeatureLevel = SceneContext.GetCurrentFeatureLevel();
+	const bool bUseGBuffer = IsUsingGBuffers(ShaderPlatform);
+	check(CurrentFeatureLevel >= ERHIFeatureLevel::SM5);
 
 	// Create MRT
 	int32 VelocityRTIndex = -1;
diff --git a/Engine/Source/Runtime/Renderer/Private/TranslucencyPass.h b/Engine/Source/Runtime/Renderer/Private/TranslucencyPass.h
--- a/Engine/Source/Runtime/Renderer/Private/TranslucencyPass.h
+++ b/Engine/Source/Runtime/Renderer/Private/TranslucencyPass.h
@@ -13,6 +13,7 @@ namespace ETranslucencyPass
 	{
 		TPT_StandardTranslucency,
 		TPT_TranslucencyAfterDOF,
+		TPT_TranslucencyUnderWater,
 
 		/** Drawing all translucency, regardless of separate or standard.  Used when drawing translucency outside of the main renderer, eg FRendererModule::DrawTile. */
 		TPT_AllTranslucency,
diff --git a/Engine/Source/Runtime/Renderer/Private/TranslucentRendering.cpp b/Engine/Source/Runtime/Renderer/Private/TranslucentRendering.cpp
--- a/Engine/Source/Runtime/Renderer/Private/TranslucentRendering.cpp
+++ b/Engine/Source/Runtime/Renderer/Private/TranslucentRendering.cpp
@@ -91,6 +91,7 @@ EMeshPass::Type TranslucencyPassToMeshPass(ETranslucencyPass::Type TranslucencyP
 	{
 	case ETranslucencyPass::TPT_StandardTranslucency: TranslucencyMeshPass = EMeshPass::TranslucencyStandard; break;
 	case ETranslucencyPass::TPT_TranslucencyAfterDOF: TranslucencyMeshPass = EMeshPass::TranslucencyAfterDOF; break;
+	case ETranslucencyPass::TPT_TranslucencyUnderWater: TranslucencyMeshPass = EMeshPass::TranslucencyUnderWater; break;
 	case ETranslucencyPass::TPT_AllTranslucency: TranslucencyMeshPass = EMeshPass::TranslucencyAll; break;
 	}
 
@@ -964,7 +965,8 @@ void FDeferredShadingSceneRenderer::RenderTranslucency(FRHICommandListImmediate&
 		}
 		else
 		{
-			SceneContext.BeginRenderingTranslucency(RHICmdList, View, *this, ViewIndex == 0 || View.Family->bMultiGPUForkAndJoin);
+			bool bAllowBufferClear = TranslucencyPass != ETranslucencyPass::TPT_TranslucencyUnderWater; // We render in an existing scene color in this case so do not clear.
+			SceneContext.BeginRenderingTranslucency(RHICmdList, View, *this, (ViewIndex == 0 || View.Family->bMultiGPUForkAndJoin) && bAllowBufferClear);
 			DrawRenderState.SetDepthStencilState(TStaticDepthStencilState<false, CF_DepthNearOrEqual>::GetRHI());
 
 			if (bUseParallel && !ViewFamily.UseDebugViewPS())
diff --git a/Engine/Source/Runtime/Renderer/Public/MeshPassProcessor.h b/Engine/Source/Runtime/Renderer/Public/MeshPassProcessor.h
--- a/Engine/Source/Runtime/Renderer/Public/MeshPassProcessor.h
+++ b/Engine/Source/Runtime/Renderer/Public/MeshPassProcessor.h
@@ -26,6 +26,7 @@ namespace EMeshPass
 		Velocity,
 		TranslucencyStandard,
 		TranslucencyAfterDOF,
+		TranslucencyUnderWater,
 		TranslucencyAll, /** Drawing all translucency, regardless of separate or standard.  Used when drawing translucency outside of the main renderer, eg FRendererModule::DrawTile. */
 		LightmapDensity,
 		DebugViewMode, /** Any of EDebugViewShaderMode */
@@ -59,6 +60,7 @@ inline const TCHAR* GetMeshPassName(EMeshPass::Type MeshPass)
 	case EMeshPass::Velocity: return TEXT("Velocity");
 	case EMeshPass::TranslucencyStandard: return TEXT("TranslucencyStandard");
 	case EMeshPass::TranslucencyAfterDOF: return TEXT("TranslucencyAfterDOF");
+	case EMeshPass::TranslucencyUnderWater: return TEXT("TranslucencyUnderWater");
 	case EMeshPass::TranslucencyAll: return TEXT("TranslucencyAll");
 	case EMeshPass::LightmapDensity: return TEXT("LightmapDensity");
 	case EMeshPass::DebugViewMode: return TEXT("DebugViewMode");
