Generate thumbnails for world height RVTs

diff --git a/Engine/Shaders/Private/VirtualTextureCompress.usf b/Engine/Shaders/Private/VirtualTextureCompress.usf
--- a/Engine/Shaders/Private/VirtualTextureCompress.usf
+++ b/Engine/Shaders/Private/VirtualTextureCompress.usf
@@ -250,3 +250,18 @@ void CopyNormalSpecularExCS(
 	OutCopyTexture0[ThreadId.xy + DestRect.xy] = float4(BaseColor, NormalXY.x);
 	OutCopyTexture1[ThreadId.xy + DestRect.xy] = float4(RoughnessSpecularNormalZ, NormalXY.y);
 }
+
+/** Copy world height (without compression)
+  * This path is used when want to copy the world height data into a thumbnail. */
+[numthreads(8,8,1)]
+void CopyWorldHeightCS(
+	uint3 ThreadId : SV_DispatchThreadID 
+	)
+{
+	uint2 SamplePos = ThreadId.xy;
+	if (any(ThreadId.xy >= DestRect.zw))
+		return;
+
+	float WorldHeight = RenderTexture0[SamplePos].r;
+	OutCopyTexture0[ThreadId.xy + DestRect.xy] = float4(WorldHeight.xxx, 1.f);
+}
diff --git a/Engine/Source/Editor/VirtualTexturingEditor/Private/RuntimeVirtualTextureThumbnailRenderer.cpp b/Engine/Source/Editor/VirtualTexturingEditor/Private/RuntimeVirtualTextureThumbnailRenderer.cpp
--- a/Engine/Source/Editor/VirtualTexturingEditor/Private/RuntimeVirtualTextureThumbnailRenderer.cpp
+++ b/Engine/Source/Editor/VirtualTexturingEditor/Private/RuntimeVirtualTextureThumbnailRenderer.cpp
@@ -39,18 +39,14 @@ bool URuntimeVirtualTextureThumbnailRenderer::CanVisualizeAsset(UObject* Object)
 	URuntimeVirtualTexture* RuntimeVirtualTexture = Cast<URuntimeVirtualTexture>(Object);
 	if (RuntimeVirtualTexture->GetEnabled())
 	{
-		//todo[vt]: Support thumbnails for ERuntimeVirtualTextureMaterialType::WorldHeight (which requires render to PF_G16)
-		if (RuntimeVirtualTexture->GetMaterialType() != ERuntimeVirtualTextureMaterialType::WorldHeight)
+		// We need a matching URuntimeVirtualTextureComponent in a Scene to be able to render a thumbnail
+		URuntimeVirtualTextureComponent* RuntimeVirtualTextureComponent = FindComponent(RuntimeVirtualTexture);
+		if (RuntimeVirtualTextureComponent != nullptr)
 		{
-			// We need a matching URuntimeVirtualTextureComponent in a Scene to be able to render a thumbnail
-			URuntimeVirtualTextureComponent* RuntimeVirtualTextureComponent = FindComponent(RuntimeVirtualTexture);
-			if (RuntimeVirtualTextureComponent != nullptr)
+			FSceneInterface* Scene = RuntimeVirtualTextureComponent->GetScene();
+			if (Scene != nullptr && RuntimeVirtualTexture::IsSceneReadyToRender(Scene->GetRenderScene()))
 			{
-				FSceneInterface* Scene = RuntimeVirtualTextureComponent->GetScene();
-				if (Scene != nullptr && RuntimeVirtualTexture::IsSceneReadyToRender(Scene->GetRenderScene()))
-				{
-					return true;
-				}
+				return true;
 			}
 		}
 	}
diff --git a/Engine/Source/Runtime/Renderer/Private/VT/RuntimeVirtualTextureRender.cpp b/Engine/Source/Runtime/Renderer/Private/VT/RuntimeVirtualTextureRender.cpp
--- a/Engine/Source/Runtime/Renderer/Private/VT/RuntimeVirtualTextureRender.cpp
+++ b/Engine/Source/Runtime/Renderer/Private/VT/RuntimeVirtualTextureRender.cpp
@@ -492,6 +492,9 @@ namespace RuntimeVirtualTexture
 	typedef FShader_VirtualTextureCompress_CS< ERuntimeVirtualTextureMaterialType::BaseColor_Normal_Specular_Ex, true > FShader_VirtualTextureCompress_BaseColorNormalSpecularEx_CopyOnly_CS;
 	IMPLEMENT_SHADER_TYPE(template<>, FShader_VirtualTextureCompress_BaseColorNormalSpecularEx_CopyOnly_CS, TEXT("/Engine/Private/VirtualTextureCompress.usf"), TEXT("CopyNormalSpecularExCS"), SF_Compute);
 
+	typedef FShader_VirtualTextureCompress_CS< ERuntimeVirtualTextureMaterialType::WorldHeight, true > FShader_VirtualTextureCompress_WorldHeight_CopyOnly_CS;
+	IMPLEMENT_SHADER_TYPE(template<>, FShader_VirtualTextureCompress_WorldHeight_CopyOnly_CS, TEXT("/Engine/Private/VirtualTextureCompress.usf"), TEXT("CopyWorldHeightCS"), SF_Compute);
+
 
 	/** Set up the BC compression pass for the specific MaterialType. */
 	template< ERuntimeVirtualTextureMaterialType MaterialType, bool CopyOnly = false >
@@ -534,7 +537,8 @@ namespace RuntimeVirtualTexture
 	{
 		const FIntVector GroupCount((TextureSize.X + 7) / 8, (TextureSize.Y + 7) / 8, 1);
 
-		// Only needed for BaseColor_Normal_Specular where we need to pack normal and specular into one VT layer
+		// Needed for BaseColor_Normal_Specular where we need to pack normal and specular into one VT layer
+		// Needed for WorldHeight thumbnails where we need to copy G16 to the RGBA8 thumbnail
 		switch (MaterialType)
 		{
 		case ERuntimeVirtualTextureMaterialType::BaseColor_Normal_Specular:
@@ -543,6 +547,9 @@ namespace RuntimeVirtualTexture
 		case ERuntimeVirtualTextureMaterialType::BaseColor_Normal_Specular_Ex:
 			AddCompressOrCopyPass<ERuntimeVirtualTextureMaterialType::BaseColor_Normal_Specular_Ex, true>(GraphBuilder, FeatureLevel, Parameters, GroupCount);
 			break;
+		case ERuntimeVirtualTextureMaterialType::WorldHeight:
+			AddCompressOrCopyPass<ERuntimeVirtualTextureMaterialType::WorldHeight, true>(GraphBuilder, FeatureLevel, Parameters, GroupCount);
+			break;
 		}
 	}
 
@@ -556,6 +563,7 @@ namespace RuntimeVirtualTexture
 			bRenderPass = OutputTexture0 != nullptr;
 			bCompressPass = bRenderPass && (OutputTexture0->GetFormat() == PF_DXT1 || OutputTexture0->GetFormat() == PF_DXT5 || OutputTexture0->GetFormat() == PF_BC5);
 			bCopyPass = bRenderPass && !bCompressPass && (MaterialType == ERuntimeVirtualTextureMaterialType::BaseColor_Normal_Specular || MaterialType == ERuntimeVirtualTextureMaterialType::BaseColor_Normal_Specular_Ex);
+			bCopyPass |= bRenderPass && MaterialType == ERuntimeVirtualTextureMaterialType::WorldHeight && OutputTexture0->GetFormat() == PF_B8G8R8A8;
 		
 			switch (MaterialType)
 			{
@@ -621,6 +629,10 @@ namespace RuntimeVirtualTexture
 				{
 					OutputAlias0 = RenderTexture0 = GraphBuilder.CreateTexture(FPooledRenderTargetDesc::Create2DDesc(TextureSize, PF_G16, FClearValueBinding::Black, TexCreate_None, TexCreate_RenderTargetable, false), TEXT("RenderTexture0"));
 				}
+				if (bCopyPass)
+				{
+					OutputAlias0 = CopyTexture0 = GraphBuilder.CreateTexture(FPooledRenderTargetDesc::Create2DDesc(TextureSize, PF_B8G8R8A8, FClearValueBinding::None, TexCreate_None, TexCreate_UAV, false), TEXT("CopyTexture0"));
+				}
 				break;
 			}
 		}
