Fix virtual texture physical layer indexing bug

diff --git a/Engine/Source/Runtime/Renderer/Private/VT/VirtualTextureSystem.cpp b/Engine/Source/Runtime/Renderer/Private/VT/VirtualTextureSystem.cpp
--- a/Engine/Source/Runtime/Renderer/Private/VT/VirtualTextureSystem.cpp
+++ b/Engine/Source/Runtime/Renderer/Private/VT/VirtualTextureSystem.cpp
@@ -1747,7 +1747,7 @@ void FVirtualTextureSystem::SubmitRequests(FRHICommandListImmediate& RHICmdList,
 			uint32 ProducerTextureLayerMask = 0;
 			for (uint32 ProducerLayerIndex = 0; ProducerLayerIndex < Producer.GetNumTextureLayers(); ++ProducerLayerIndex)
 			{
-				if (ProducerPhysicalGroupMask & (1 << Producer.GetPhysicalGroupIndexForTextureLayer(ProducerTextureLayerMask)))
+				if (ProducerPhysicalGroupMask & (1 << Producer.GetPhysicalGroupIndexForTextureLayer(ProducerLayerIndex)))
 				{
 					ProducerTextureLayerMask |= (1 << ProducerLayerIndex);
 				}
