Make VT pools respect budget size
Use floor instead of ceil so that budget is our upper memory limit

diff --git a/Engine/Source/Runtime/Renderer/Private/VT/VirtualTexturePhysicalSpace.cpp b/Engine/Source/Runtime/Renderer/Private/VT/VirtualTexturePhysicalSpace.cpp
--- a/Engine/Source/Runtime/Renderer/Private/VT/VirtualTexturePhysicalSpace.cpp
+++ b/Engine/Source/Runtime/Renderer/Private/VT/VirtualTexturePhysicalSpace.cpp
@@ -30,7 +30,7 @@ FVirtualTexturePhysicalSpace::FVirtualTexturePhysicalSpace(const FVTPhysicalSpac
 	}
 	const uint32 MaxTiles = (uint32)(PoolSizeInBytes / TileSizeBytes);
 	
-	TextureSizeInTiles = FMath::CeilToInt(FMath::Sqrt((float)MaxTiles));
+	TextureSizeInTiles = FMath::FloorToInt(FMath::Sqrt((float)MaxTiles));
 	if (bForce16BitPageTable)
 	{
 		// 16 bit page tables support max size of 64x64 (4096 tiles)
