RoboMerge - manual merge

Added tile categorisation for water reflection composite and SSR pass.
Also removes the water SSR dedicated TAA pass when not enabled using the correct check (not enabled on consoles)

diff --git a/Engine/Shaders/Private/SSRT/SSRTReflections.usf b/Engine/Shaders/Private/SSRT/SSRTReflections.usf
--- a/Engine/Shaders/Private/SSRT/SSRTReflections.usf
+++ b/Engine/Shaders/Private/SSRT/SSRTReflections.usf
@@ -89,11 +89,11 @@ void ScreenSpaceReflectionsStencilPS(
 	#endif
 }
 
-void ScreenSpaceReflectionsPS(
-	float4 SvPosition : SV_POSITION
-	, out float4 OutColor : SV_Target0
+void ScreenSpaceReflections(
+	float4 SvPosition
+	, out float4 OutColor
 #if SSR_OUTPUT_FOR_DENOISER
-	, out float4 OutClosestHitDistance : SV_Target1
+	, out float4 OutClosestHitDistance
 #endif
 )
 {
@@ -355,4 +355,59 @@ void ScreenSpaceReflectionsPS(
 		OutClosestHitDistance = ClosestHitDistanceSqr > 0 ? ClampToHalfFloatRange(sqrt(ClosestHitDistanceSqr)) : ClosestHitDistanceSqr;
 	}
 	#endif
-}
\ No newline at end of file
+}
+
+void ScreenSpaceReflectionsPS(
+	float4 SvPosition : SV_POSITION
+	, out float4 OutColor : SV_Target0
+#if SSR_OUTPUT_FOR_DENOISER
+	, out float4 OutClosestHitDistance : SV_Target1
+#endif
+)
+{
+	ScreenSpaceReflections(SvPosition, OutColor
+#if SSR_OUTPUT_FOR_DENOISER
+		,OutClosestHitDistance
+#endif
+	);
+}
+
+
+
+#if defined(TILE_COMPUTE_SSR)
+
+Buffer<uint> TileListDataBuffer;
+
+RWTexture2D<float4> OutColorUAV;
+RWTexture2D<float4> OutClosestHitDistanceUAV;
+
+#define WATER_TILE_SIZE 8
+[numthreads(WATER_TILE_SIZE, WATER_TILE_SIZE, 1)]
+void ScreenSpaceReflectionsCS(/*uint3 ThreadId : SV_DispatchThreadID,*/ uint3 GroupId : SV_GroupID, uint3 GroupThreadId : SV_GroupThreadID)
+{
+	uint TileData = TileListDataBuffer[GroupId.x];
+	uint3 ThreadId = uint3((TileData & 0x0000FFFF)*WATER_TILE_SIZE + GroupThreadId.x, (TileData >> 16)*WATER_TILE_SIZE + GroupThreadId.y, 0);
+
+	if (int(ThreadId.x) < View.ViewSizeAndInvSize.x && int(ThreadId.y) < View.ViewSizeAndInvSize.y)
+	{
+		uint2 PixelPos = ThreadId.xy + View.ViewRectMin.xy;
+		float4 SvPosition = float4(float2(PixelPos) + 0.5f, 0.5, 1.0f);
+
+		float4 OutColorVar; 
+		float4 OutClosestHitDistanceVar;
+		ScreenSpaceReflections(SvPosition, OutColorVar
+#if SSR_OUTPUT_FOR_DENOISER
+			,OutClosestHitDistanceVar 
+#endif
+		);
+
+		OutColorUAV[PixelPos] = OutColorVar;
+#if SSR_OUTPUT_FOR_DENOISER
+		OutClosestHitDistanceUAV[PixelPos] = OutClosestHitDistanceVar;
+#endif
+	}
+}
+
+#endif
+
+
diff --git a/Engine/Shaders/Private/SingleLayerWaterComposite.usf b/Engine/Shaders/Private/SingleLayerWaterComposite.usf
--- a/Engine/Shaders/Private/SingleLayerWaterComposite.usf
+++ b/Engine/Shaders/Private/SingleLayerWaterComposite.usf
@@ -26,27 +26,26 @@ SamplerState ScreenSpaceReflectionsSampler;
 Texture2D SceneNoWaterDepthTexture;
 SamplerState SceneNoWaterDepthSampler;
 
-
-
-void SingleLayerWaterCompositePS(
-	in float4 SvPosition : SV_Position,
-	out float4 OutColor : SV_Target0)
+struct SingleLayerWaterCompositeOutput
 {
-	float2 BufferUV = SvPositionToBufferUV(SvPosition);
-	float2 ScreenPosition = SvPositionToScreenPosition(SvPosition).xy;
-
+	float4 LuminanceTransmittance;
+	float  Clip;
+};
 
+SingleLayerWaterCompositeOutput SingleLayerWaterComposite(float2 BufferUV, float2 ScreenPosition)
+{
 	// Sample scene textures.
 	FGBufferData GBuffer = GetGBufferDataFromSceneTextures(BufferUV);
 	uint ShadingModelID = GBuffer.ShadingModelID;
 	// No AO or DFAO
 
-	const float3 OneThird3 = float3(1.0 / 3.0f, 1.0 / 3.0f, 1.0 / 3.0f);
-
-	float3 WorldPos = SvPositionToWorld(float4(SvPosition.xy, 1.0, SvPosition.w));
+	const float3 OneThird3 = float3(1.0f / 3.0f, 1.0f / 3.0f, 1.0f / 3.0f);
 	const bool CameraIsUnderWater = false;
 
-	float4 LuminanceTransmittance = float4(0, 0, 0, 1);
+	SingleLayerWaterCompositeOutput Output;
+	Output.LuminanceTransmittance = float4(0.0f, 0.0f, 0.0f, 1.0f);
+	Output.Clip = 1.0f; // no clipping
+
 	float3 TransmittanceToScene = 1.0f;
 
 #if USE_PREEXPOSURE
@@ -54,11 +53,11 @@ void SingleLayerWaterCompositePS(
 #else
 	const float PreExposure = 1.f;
 #endif
-	if (ShadingModelID != SHADINGMODELID_UNLIT)
+	if (ShadingModelID == SHADINGMODELID_SINGLELAYERWATER)
 	{
 		const float OpaqueDepth = ConvertFromDeviceZ(SceneNoWaterDepthTexture.SampleLevel(SceneNoWaterDepthSampler, BufferUV, 0).r);
 		const float WaterDepth = GBuffer.Depth;
-		float DeltaDepth = CameraIsUnderWater ? WaterDepth :  OpaqueDepth - WaterDepth;	// Inverted depth
+		float DeltaDepth = CameraIsUnderWater ? WaterDepth : OpaqueDepth - WaterDepth;	// Inverted depth
 		if (DeltaDepth > 0.0)
 		{
 			// Compute the sky reflection contribution
@@ -118,22 +117,184 @@ void SingleLayerWaterCompositePS(
 				TransmittanceToScene *= WaterMediumTransmittance;
 			}
 
-			LuminanceTransmittance = float4(Reflection, dot(OneThird3, TransmittanceToScene));
+			Output.LuminanceTransmittance = float4(Reflection, dot(OneThird3, TransmittanceToScene));
 		}
 		else
 		{
-			clip(-1); // Since this shader does not write to depth or stencil it should still benefit from EarlyZ (See AMD depth-in-depth documentation)
+			Output.Clip = -1.0f;
 		}
 	}
 	else
 	{
-		clip(-1); // Since this shader does not write to depth or stencil it should still benefit from EarlyZ (See AMD depth-in-depth documentation)
+		Output.Clip = -1.0f;
+	}
+
+	return Output;
+}
+
+
+
+#ifdef TILE_CATERGORISATION_SHADER
+
+
+unsigned int TiledViewWidth;
+unsigned int TiledViewHeight;
+float TiledViewWidthInv;
+float TiledViewHeightInv;
+
+RWStructuredBuffer<uint> DispatchIndirectDataUAV;
+RWStructuredBuffer<uint> WaterTileListDataUAV;
+
+groupshared bool ContainsWater[WATER_TILE_SIZE*WATER_TILE_SIZE];
+
+[numthreads(WATER_TILE_SIZE, WATER_TILE_SIZE, 1)]
+void WaterTileCatergorisationCS(uint3 ThreadId : SV_DispatchThreadID, uint3 GroupId : SV_GroupID, uint3 GroupThreadId : SV_GroupThreadID)
+{
+	if (ThreadId.x == 0 && ThreadId.y == 0 && GroupThreadId.x == 0 && GroupThreadId.y == 0)
+	{
+		// TODO compute clear myself
+		DispatchIndirectDataUAV[1] = 1;
+		DispatchIndirectDataUAV[2] = 1;
+	}
+
+#if 0
+	// Slow reference path
+	const int LinearIndex = GroupThreadId.y * WATER_TILE_SIZE + GroupThreadId.x;
+	if (LinearIndex < 1)
+	{
+		bool bContainsWater = false;
+		for (uint i = 0; i < WATER_TILE_SIZE; ++i)
+		{
+			for (uint j = 0; j < WATER_TILE_SIZE; ++j)
+			{
+				uint2 PixelPos = (ThreadId.xy + uint2(i, j)) + View.ViewRectMin.xy;
+				float2 BufferUV = (PixelPos + 0.5f) * View.BufferSizeAndInvSize.zw;
+				FGBufferData GBuffer = GetGBufferDataFromSceneTextures(BufferUV);
+
+				bContainsWater = bContainsWater || GBuffer.ShadingModelID == SHADINGMODELID_SINGLELAYERWATER;
+			}
+		}
+		if (bContainsWater)
+		{
+			uint WriteToIndex;
+			InterlockedAdd(DispatchIndirectDataUAV[0], 1, WriteToIndex);
+			WaterTileListDataUAV[WriteToIndex] = GroupId.x | (GroupId.y << 16); // assumes 16bit is enought to represent a tiled resolution up to 65,535 :)
+		}
+	}
+
+#else
+
+	bool bContainsWater = false;
+	if(ThreadId.x < uint(View.BufferSizeAndInvSize.x) && ThreadId.y < uint(View.BufferSizeAndInvSize.y))
+	{
+		uint2 PixelPos = ThreadId.xy + View.ViewRectMin.xy;
+		float2 BufferUV = (PixelPos + 0.5f) * View.BufferSizeAndInvSize.zw;
+		FGBufferData GBuffer = GetGBufferDataFromSceneTextures(BufferUV);
+		bContainsWater = GBuffer.ShadingModelID == SHADINGMODELID_SINGLELAYERWATER;
+	}
+
+	const int LinearIndex = GroupThreadId.y * WATER_TILE_SIZE + GroupThreadId.x;
+	ContainsWater[LinearIndex] = bContainsWater;
+	GroupMemoryBarrierWithGroupSync();
+
+	if (LinearIndex < 32) // 8*8 = 64 elements to merge
+	{
+		ContainsWater[LinearIndex] = ContainsWater[LinearIndex] || ContainsWater[LinearIndex + 32];
+	}
+	GroupMemoryBarrierWithGroupSync();
+	if (LinearIndex < 16)
+	{
+		ContainsWater[LinearIndex] = ContainsWater[LinearIndex] || ContainsWater[LinearIndex + 16];
 	}
+	GroupMemoryBarrierWithGroupSync();
+
+	// The smallest wave size is 16 on Intel hardware. So now we can do simple math operations without group sync.
+	// EDIT: for some reason group sync is needed until the end, otherwise some pixels are missing...
+
+	if (LinearIndex < 8)
+	{
+		ContainsWater[LinearIndex] = ContainsWater[LinearIndex] || ContainsWater[LinearIndex + 8];
+	}
+	GroupMemoryBarrierWithGroupSync();
+	if (LinearIndex < 4)
+	{
+		ContainsWater[LinearIndex] = ContainsWater[LinearIndex] || ContainsWater[LinearIndex + 4];
+	}
+	GroupMemoryBarrierWithGroupSync();
+	if (LinearIndex < 2)
+	{
+		ContainsWater[LinearIndex] = ContainsWater[LinearIndex] || ContainsWater[LinearIndex + 2];
+	}
+	GroupMemoryBarrierWithGroupSync();
+	if (LinearIndex < 1 && (ContainsWater[LinearIndex] || ContainsWater[LinearIndex + 1]))
+	{
+		uint WriteToIndex;  
+		InterlockedAdd(DispatchIndirectDataUAV[0], 1, WriteToIndex);
+		WaterTileListDataUAV[WriteToIndex] = GroupId.x | (GroupId.y << 16); // assumes 16bit is enought to represent a tiled resolution up to 65,535 :)
+	}
+#endif
+} 
+
+
+
+#elif defined(TILED_COMPOSITE_SHADER)
+
+
+
+unsigned int TiledViewWidth;
+unsigned int TiledViewHeight;
+
+Buffer<uint> WaterTileListData;
+
+RWTexture2D<float3> SceneColorUAV;
+
+[numthreads(WATER_TILE_SIZE, WATER_TILE_SIZE, 1)]
+void WaterTiledCategorisationCS(/*uint3 ThreadId : SV_DispatchThreadID,*/ uint3 GroupId : SV_GroupID, uint3 GroupThreadId : SV_GroupThreadID)
+{
+	uint TileData = WaterTileListData[GroupId.x];
+	uint3 ThreadId = uint3((TileData & 0x0000FFFF)*WATER_TILE_SIZE + GroupThreadId.x, (TileData >> 16)*WATER_TILE_SIZE + GroupThreadId.y, 0);
+
+	if (int(ThreadId.x) < View.ViewSizeAndInvSize.x && int(ThreadId.y) < View.ViewSizeAndInvSize.y)
+	{
+		uint2 PixelPos = ThreadId.xy + View.ViewRectMin.xy;
+		float2 BufferUV = (PixelPos + 0.5) * View.BufferSizeAndInvSize.zw;
+		float2 ScreenPosition = ViewportUVToScreenPos((ThreadId.xy + 0.5f) * View.ViewSizeAndInvSize.zw);
+
+		SingleLayerWaterCompositeOutput Result = SingleLayerWaterComposite(BufferUV, ScreenPosition);
+
+		if (Result.Clip >= 0)
+		{
+			// TODO colored transmittance when we have a under water medium to apply
+			float3 SceneColor = SceneColorUAV[PixelPos];
+			//SceneColor = Result.LuminanceTransmittance.rgb + Result.LuminanceTransmittance.a * SceneColor  +0.2*float3(sin((TileData & 0x0000FFFF)*WATER_TILE_SIZE), sin((TileData >> 16)*WATER_TILE_SIZE), 0.0f);
+			SceneColor = Result.LuminanceTransmittance.rgb + Result.LuminanceTransmittance.a * SceneColor;
+			SceneColorUAV[PixelPos] = SceneColor;
+		}
+	}
+}
+
+
+
+#else
+
+
+
+void SingleLayerWaterCompositePS(
+	in float4 SvPosition : SV_Position,
+	out float4 OutColor : SV_Target0)
+{
+	float2 BufferUV = SvPositionToBufferUV(SvPosition);
+	float2 ScreenPosition = SvPositionToScreenPosition(SvPosition).xy;
 
 	// TODO use dual source blending to apply TransmittanceToScene
-	OutColor = LuminanceTransmittance;
+	SingleLayerWaterCompositeOutput Result = SingleLayerWaterComposite(BufferUV, ScreenPosition);
+
+	clip(Result.Clip); // Since this shader does not write to depth or stencil it should still benefit from EarlyZ (See AMD depth-in-depth documentation)
+	OutColor = Result.LuminanceTransmittance;
 }
 
 
 
+#endif
+
 
diff --git a/Engine/Source/Runtime/Renderer/Private/ScreenSpaceRayTracing.cpp b/Engine/Source/Runtime/Renderer/Private/ScreenSpaceRayTracing.cpp
--- a/Engine/Source/Runtime/Renderer/Private/ScreenSpaceRayTracing.cpp
+++ b/Engine/Source/Runtime/Renderer/Private/ScreenSpaceRayTracing.cpp
@@ -212,6 +212,16 @@ BEGIN_SHADER_PARAMETER_STRUCT(FSSRCommonParameters, )
 	SHADER_PARAMETER_STRUCT_REF(FViewUniformShaderParameters, ViewUniformBuffer)
 END_SHADER_PARAMETER_STRUCT()
 
+BEGIN_SHADER_PARAMETER_STRUCT(FSSRPassCommonParameters, )
+	SHADER_PARAMETER(FVector4, HZBUvFactorAndInvFactor)
+	SHADER_PARAMETER(FVector4, PrevScreenPositionScaleBias)
+	SHADER_PARAMETER(float, PrevSceneColorPreExposureCorrection)
+	SHADER_PARAMETER_RDG_TEXTURE(Texture2D, SceneColor)
+	SHADER_PARAMETER_SAMPLER(SamplerState, SceneColorSampler)
+	SHADER_PARAMETER_RDG_TEXTURE(Texture2D, HZB)
+	SHADER_PARAMETER_SAMPLER(SamplerState, HZBSampler)
+	SHADER_PARAMETER_RDG_TEXTURE_UAV(RWTexture<float4>, ScreenSpaceRayTracingDebugOutput)
+END_SHADER_PARAMETER_STRUCT()
 
 class FSSRQualityDim : SHADER_PERMUTATION_ENUM_CLASS("SSR_QUALITY", ESSRQuality);
 class FSSROutputForDenoiser : SHADER_PERMUTATION_BOOL("SSR_OUTPUT_FOR_DENOISER");
@@ -322,20 +332,36 @@ class FScreenSpaceReflectionsPS : public FGlobalShader
 	
 	BEGIN_SHADER_PARAMETER_STRUCT(FParameters, )
 		SHADER_PARAMETER_STRUCT_INCLUDE(FSSRCommonParameters, CommonParameters)
+		SHADER_PARAMETER_STRUCT_INCLUDE(FSSRPassCommonParameters, SSRPassCommonParameter)
+		RENDER_TARGET_BINDING_SLOTS()
+	END_SHADER_PARAMETER_STRUCT()
+};
 
-		SHADER_PARAMETER(FVector4, HZBUvFactorAndInvFactor)
-		SHADER_PARAMETER(FVector4, PrevScreenPositionScaleBias)
-		SHADER_PARAMETER(float, PrevSceneColorPreExposureCorrection)
-		
-		SHADER_PARAMETER_RDG_TEXTURE(Texture2D, SceneColor)
-		SHADER_PARAMETER_SAMPLER(SamplerState, SceneColorSampler)
-		
-		SHADER_PARAMETER_RDG_TEXTURE(Texture2D, HZB)
-		SHADER_PARAMETER_SAMPLER(SamplerState, HZBSampler)
+class FScreenSpaceReflectionsCS : public FGlobalShader
+{
+	DECLARE_GLOBAL_SHADER(FScreenSpaceReflectionsCS);
+	SHADER_USE_PARAMETER_STRUCT(FScreenSpaceReflectionsCS, FGlobalShader);
 
-		SHADER_PARAMETER_RDG_TEXTURE_UAV(RWTexture<float4>, ScreenSpaceRayTracingDebugOutput)
-		
-		RENDER_TARGET_BINDING_SLOTS()
+	using FPermutationDomain = FScreenSpaceReflectionsPS::FPermutationDomain;
+
+	static bool ShouldCompilePermutation(const FGlobalShaderPermutationParameters& Parameters)
+	{
+		return FScreenSpaceReflectionsPS::ShouldCompilePermutation(Parameters);
+	}
+
+	static void ModifyCompilationEnvironment(const FGlobalShaderPermutationParameters& Parameters, FShaderCompilerEnvironment& OutEnvironment)
+	{
+		OutEnvironment.SetDefine(TEXT("TILE_COMPUTE_SSR"), 1.0f);
+		FGlobalShader::ModifyCompilationEnvironment(Parameters, OutEnvironment);
+	}
+
+	BEGIN_SHADER_PARAMETER_STRUCT(FParameters, )
+		SHADER_PARAMETER_STRUCT_INCLUDE(FSSRCommonParameters, CommonParameters)
+		SHADER_PARAMETER_STRUCT_INCLUDE(FSSRPassCommonParameters, SSRPassCommonParameter)
+		SHADER_PARAMETER_RDG_BUFFER_SRV(Buffer<uint>, TileListDataBuffer)
+		SHADER_PARAMETER_RDG_TEXTURE_UAV(RWTexture2D<float3>, OutColorUAV)
+		SHADER_PARAMETER_RDG_TEXTURE_UAV(RWTexture2D<float3>, OutClosestHitDistanceUAV)
+		SHADER_PARAMETER_RDG_BUFFER(Buffer<uint>, IndirectDispatchParameters)	// Not used in shader but need to be reference in the parameter list
 	END_SHADER_PARAMETER_STRUCT()
 };
 
@@ -382,6 +408,7 @@ class FScreenSpaceDiffuseIndirectCS : public FGlobalShader
 IMPLEMENT_GLOBAL_SHADER(FSSRTPrevFrameReductionCS, "/Engine/Private/SSRT/SSRTPrevFrameReduction.usf", "MainCS", SF_Compute);
 IMPLEMENT_GLOBAL_SHADER(FSSRTDiffuseTileClassificationCS, "/Engine/Private/SSRT/SSRTTileClassification.usf", "MainCS", SF_Compute);
 IMPLEMENT_GLOBAL_SHADER(FScreenSpaceReflectionsPS,        "/Engine/Private/SSRT/SSRTReflections.usf", "ScreenSpaceReflectionsPS", SF_Pixel);
+IMPLEMENT_GLOBAL_SHADER(FScreenSpaceReflectionsCS,        "/Engine/Private/SSRT/SSRTReflections.usf", "ScreenSpaceReflectionsCS", SF_Compute);
 IMPLEMENT_GLOBAL_SHADER(FScreenSpaceReflectionsStencilPS, "/Engine/Private/SSRT/SSRTReflections.usf", "ScreenSpaceReflectionsStencilPS", SF_Pixel);
 IMPLEMENT_GLOBAL_SHADER(FScreenSpaceDiffuseIndirectCS, "/Engine/Private/SSRT/SSRTDiffuseIndirect.usf", "MainCS", SF_Compute);
 
@@ -510,7 +537,8 @@ void RenderScreenSpaceReflections(
 	const FViewInfo& View,
 	ESSRQuality SSRQuality,
 	bool bDenoiser,
-	IScreenSpaceDenoiser::FReflectionsInputs* DenoiserInputs)
+	IScreenSpaceDenoiser::FReflectionsInputs* DenoiserInputs,
+	FTiledScreenSpaceReflection* TiledScreenSpaceReflection)
 {
 	FRDGTextureRef InputColor = CurrentSceneColor;
 	if (SSRQuality != ESSRQuality::VisualizeSSR)
@@ -525,14 +553,14 @@ void RenderScreenSpaceReflections(
 		}
 	}
 
-	const bool SSRStencilPrePass = CVarSSRStencil.GetValueOnRenderThread() != 0 && SSRQuality != ESSRQuality::VisualizeSSR;
+	const bool SSRStencilPrePass = CVarSSRStencil.GetValueOnRenderThread() != 0 && SSRQuality != ESSRQuality::VisualizeSSR && TiledScreenSpaceReflection == nullptr;
 	
 	// Alloc inputs for denoising.
 	{
 		FRDGTextureDesc Desc = FPooledRenderTargetDesc::Create2DDesc(
 			FSceneRenderTargets::Get_FrameConstantsOnly().GetBufferSizeXY(),
 			PF_FloatRGBA, FClearValueBinding(FLinearColor(0, 0, 0, 0)),
-			TexCreate_None, TexCreate_RenderTargetable | TexCreate_ShaderResource,
+			TexCreate_None, TexCreate_RenderTargetable | TexCreate_ShaderResource | TexCreate_UAV,
 			false);
 
 		Desc.AutoWritable = false;
@@ -555,6 +583,12 @@ void RenderScreenSpaceReflections(
 	CommonParameters.ViewUniformBuffer = View.ViewUniformBuffer;
 	CommonParameters.SceneTextures = SceneTextures;
 	SetupSceneTextureSamplers(&CommonParameters.SceneTextureSamplers);
+	// Pipe down a mid grey texture when not using TAA's history to avoid wrongly reprojecting current scene color as if previous frame's TAA history.
+	if (InputColor == CurrentSceneColor)
+	{
+		// Technically should be 32767.0f / 65535.0f to perfectly null out DecodeVelocityFromTexture(), but 0.5f is good enough.
+		CommonParameters.SceneTextures.SceneVelocityBuffer = GraphBuilder.RegisterExternalTexture(GSystemTextures.MidGreyDummy);
+	}
 	
 	FRenderTargetBindingSlots RenderTargets;
 	RenderTargets[0] = FRenderTargetBinding(DenoiserInputs->Color, ERenderTargetLoadAction::ENoAction);
@@ -607,13 +641,8 @@ void RenderScreenSpaceReflections(
 	}
 
 	// Adds SSR pass.
+	auto SetSSRParameters = [&](auto* PassParameters)
 	{
-		FScreenSpaceReflectionsPS::FPermutationDomain PermutationVector;
-		PermutationVector.Set<FSSRQualityDim>(SSRQuality);
-		PermutationVector.Set<FSSROutputForDenoiser>(bDenoiser);
-
-		FScreenSpaceReflectionsPS::FParameters* PassParameters = GraphBuilder.AllocParameters<FScreenSpaceReflectionsPS::FParameters>();
-		PassParameters->CommonParameters = CommonParameters;
 		{
 			const FVector2D HZBUvFactor(
 				float(View.ViewRect.Width()) / float(2 * View.HZBMipmap0Size.X),
@@ -643,28 +672,32 @@ void RenderScreenSpaceReflections(
 				-ViewportExtent.Y * 0.5f * InvBufferSize.Y,
 				(ViewportExtent.X * 0.5f + ViewportOffset.X) * InvBufferSize.X,
 				(ViewportExtent.Y * 0.5f + ViewportOffset.Y) * InvBufferSize.Y);
+
+			PassParameters->ScreenSpaceRayTracingDebugOutput = CreateScreenSpaceRayTracingDebugUAV(GraphBuilder, DenoiserInputs->Color->Desc, TEXT("DebugSSR"), true);
 		}
 		PassParameters->PrevSceneColorPreExposureCorrection = InputColor != CurrentSceneColor ? View.PreExposure / View.PrevViewInfo.SceneColorPreExposure : 1.0f;
-		
-		// Pipe down a mid grey texture when not using TAA's history to avoid wrongly reprojecting current scene color as if previous frame's TAA history.
-		if (InputColor == CurrentSceneColor)
-		{
-			// Technically should be 32767.0f / 65535.0f to perfectly null out DecodeVelocityFromTexture(), but 0.5f is good enough.
-			PassParameters->CommonParameters.SceneTextures.SceneVelocityBuffer = GraphBuilder.RegisterExternalTexture(GSystemTextures.MidGreyDummy);
-		}
 
 		PassParameters->SceneColor = InputColor;
 		PassParameters->SceneColorSampler = TStaticSamplerState<SF_Point>::GetRHI();
-		
+
 		PassParameters->HZB = GraphBuilder.RegisterExternalTexture(View.HZB);
 		PassParameters->HZBSampler = TStaticSamplerState<SF_Point>::GetRHI();
+	};
+
+	if (TiledScreenSpaceReflection == nullptr)
+	{
+		FScreenSpaceReflectionsPS::FPermutationDomain PermutationVector;
+		PermutationVector.Set<FSSRQualityDim>(SSRQuality);
+		PermutationVector.Set<FSSROutputForDenoiser>(bDenoiser);
 
-		PassParameters->ScreenSpaceRayTracingDebugOutput = CreateScreenSpaceRayTracingDebugUAV(GraphBuilder, DenoiserInputs->Color->Desc, TEXT("DebugSSR"), true);
+		FScreenSpaceReflectionsPS::FParameters* PassParameters = GraphBuilder.AllocParameters<FScreenSpaceReflectionsPS::FParameters>();
+		PassParameters->CommonParameters = CommonParameters;
+		SetSSRParameters(&PassParameters->SSRPassCommonParameter);
 		PassParameters->RenderTargets = RenderTargets;
 
 		TShaderMapRef<FScreenSpaceReflectionsPS> PixelShader(View.ShaderMap, PermutationVector);
 		ClearUnusedGraphResources(*PixelShader, PassParameters);
-		
+
 		GraphBuilder.AddPass(
 			RDG_EVENT_NAME("SSR RayMarch(Quality=%d RayPerPixel=%d%s) %dx%d",
 				SSRQuality, RayTracingConfigs.RayCountPerPixel, bDenoiser ? TEXT(" DenoiserOutput") : TEXT(""),
@@ -676,7 +709,7 @@ void RenderScreenSpaceReflections(
 			SCOPED_GPU_STAT(RHICmdList, ScreenSpaceReflections);
 			RHICmdList.SetViewport(View.ViewRect.Min.X, View.ViewRect.Min.Y, 0.0f, View.ViewRect.Max.X, View.ViewRect.Max.Y, 1.0f);
 			RHICmdList.SetStencilRef(0x80);
-		
+
 			FGraphicsPipelineStateInitializer GraphicsPSOInit;
 			FPixelShaderUtils::InitFullscreenPipelineState(RHICmdList, View.ShaderMap, *PixelShader, /* out */ GraphicsPSOInit);
 			if (SSRStencilPrePass)
@@ -691,6 +724,33 @@ void RenderScreenSpaceReflections(
 			FPixelShaderUtils::DrawFullscreenTriangle(RHICmdList);
 		});
 	}
+	else
+	{
+		check(TiledScreenSpaceReflection->TileSize == 8);
+		FScreenSpaceReflectionsCS::FPermutationDomain PermutationVector;
+		PermutationVector.Set<FSSRQualityDim>(SSRQuality);
+		PermutationVector.Set<FSSROutputForDenoiser>(bDenoiser);
+
+		FScreenSpaceReflectionsCS::FParameters* PassParameters = GraphBuilder.AllocParameters<FScreenSpaceReflectionsCS::FParameters>();
+		PassParameters->CommonParameters = CommonParameters;
+		SetSSRParameters(&PassParameters->SSRPassCommonParameter);
+		PassParameters->TileListDataBuffer = TiledScreenSpaceReflection->TileListStructureBufferSRV;
+		PassParameters->IndirectDispatchParameters = TiledScreenSpaceReflection->DispatchIndirectParametersBuffer;
+		PassParameters->OutColorUAV = GraphBuilder.CreateUAV(DenoiserInputs->Color);
+		if (bDenoiser)
+		{
+			PassParameters->OutClosestHitDistanceUAV = GraphBuilder.CreateUAV(DenoiserInputs->RayHitDistance);
+		}
+
+		TShaderMapRef<FScreenSpaceReflectionsCS> ComputeShader(View.ShaderMap, PermutationVector);
+		ClearUnusedGraphResources(*ComputeShader, PassParameters);
+
+		FComputeShaderUtils::AddPass(GraphBuilder,
+			RDG_EVENT_NAME("SSR RayMarch(Quality=%d RayPerPixel=%d%s) %dx%d",
+				SSRQuality, RayTracingConfigs.RayCountPerPixel, bDenoiser ? TEXT(" DenoiserOutput") : TEXT(""),
+				View.ViewRect.Width(), View.ViewRect.Height()),
+			*ComputeShader, PassParameters, TiledScreenSpaceReflection->DispatchIndirectParametersBuffer, 0);
+	}
 } // RenderScreenSpaceReflections()
 
 void RenderScreenSpaceDiffuseIndirect(
diff --git a/Engine/Source/Runtime/Renderer/Private/ScreenSpaceRayTracing.h b/Engine/Source/Runtime/Renderer/Private/ScreenSpaceRayTracing.h
--- a/Engine/Source/Runtime/Renderer/Private/ScreenSpaceRayTracing.h
+++ b/Engine/Source/Runtime/Renderer/Private/ScreenSpaceRayTracing.h
@@ -20,6 +20,16 @@ enum class ESSRQuality
 	MAX
 };
 
+struct FTiledScreenSpaceReflection
+{
+	FRDGBufferRef TileListDataBuffer;
+	FRDGBufferRef DispatchIndirectParametersBuffer;
+	FRDGBufferUAVRef DispatchIndirectParametersBufferUAV;
+	FRDGBufferUAVRef TileListStructureBufferUAV;
+	FRDGBufferSRVRef TileListStructureBufferSRV;
+	uint32 TileSize;
+};
+
 bool ShouldRenderScreenSpaceReflections(const FViewInfo& View);
 
 bool ShouldRenderScreenSpaceDiffuseIndirect(const FViewInfo& View);
@@ -35,7 +45,8 @@ void RenderScreenSpaceReflections(
 	const FViewInfo& View,
 	ESSRQuality SSRQuality,
 	bool bDenoiser,
-	IScreenSpaceDenoiser::FReflectionsInputs* DenoiserInputs);
+	IScreenSpaceDenoiser::FReflectionsInputs* DenoiserInputs,
+	FTiledScreenSpaceReflection* TiledScreenSpaceReflection = nullptr);
 
 void RenderScreenSpaceDiffuseIndirect(
 	FRDGBuilder& GraphBuilder, 
diff --git a/Engine/Source/Runtime/Renderer/Private/SingleLayerWaterRendering.cpp b/Engine/Source/Runtime/Renderer/Private/SingleLayerWaterRendering.cpp
--- a/Engine/Source/Runtime/Renderer/Private/SingleLayerWaterRendering.cpp
+++ b/Engine/Source/Runtime/Renderer/Private/SingleLayerWaterRendering.cpp
@@ -36,6 +36,12 @@ static TAutoConsoleVariable<int32> CVarWaterSingleLayerReflection(
 	ECVF_RenderThreadSafe | ECVF_Scalability);
 
 
+static TAutoConsoleVariable<int32> CVarWaterSingleLayerTiledComposite(
+	TEXT("r.Water.SingleLayer.TiledComposite"), 1,
+	TEXT("Enable tiled optimisation of the water reflection rendering."),
+	ECVF_RenderThreadSafe | ECVF_Scalability);
+
+
 static TAutoConsoleVariable<int32> CVarWaterSingleLayerSSR(
 	TEXT("r.Water.SingleLayer.SSR"), 1,
 	TEXT("Enable SSR for the single water renderring system."),
@@ -192,27 +198,31 @@ FRegisterPassProcessorCreateFunction RegisterSingleLayerWaterPass(&CreateSingleL
 
 
 
+BEGIN_SHADER_PARAMETER_STRUCT(FSingleLayerWaterCommonShaderParameters, )
+	SHADER_PARAMETER_RDG_TEXTURE(Texture2D, ScreenSpaceReflectionsTexture)
+	SHADER_PARAMETER_SAMPLER(SamplerState, ScreenSpaceReflectionsSampler)
+	SHADER_PARAMETER_TEXTURE(Texture2D, PreIntegratedGF)
+	SHADER_PARAMETER_SAMPLER(SamplerState, PreIntegratedGFSampler)
+	SHADER_PARAMETER_RDG_TEXTURE(Texture2D, SceneNoWaterDepthTexture)
+	SHADER_PARAMETER_SAMPLER(SamplerState, SceneNoWaterDepthSampler)
+	SHADER_PARAMETER_STRUCT_INCLUDE(FSceneTextureParameters, SceneTextures)	// Water scene texture
+	SHADER_PARAMETER_STRUCT_INCLUDE(FSceneTextureSamplerParameters, SceneTextureSamplers)
+	SHADER_PARAMETER_STRUCT_REF(FViewUniformShaderParameters, ViewUniformBuffer)
+	SHADER_PARAMETER_STRUCT_REF(FReflectionCaptureShaderData, ReflectionCaptureData)
+	SHADER_PARAMETER_STRUCT_REF(FReflectionUniformParameters, ReflectionsParameters)
+END_SHADER_PARAMETER_STRUCT()
+
+class FSingleLayerWaterScreenSpaceReflections : SHADER_PERMUTATION_BOOL("SCREEN_SPACE_REFLECTION");
+
 class FSingleLayerWaterCompositePS : public FGlobalShader
 {
 	DECLARE_GLOBAL_SHADER(FSingleLayerWaterCompositePS);
 	SHADER_USE_PARAMETER_STRUCT(FSingleLayerWaterCompositePS, FGlobalShader)
 
-	class FScreenSpaceReflections : SHADER_PERMUTATION_BOOL("SCREEN_SPACE_REFLECTION");
-	using FPermutationDomain = TShaderPermutationDomain<FScreenSpaceReflections>;
-
+	using FPermutationDomain = TShaderPermutationDomain<FSingleLayerWaterScreenSpaceReflections>;
 
 	BEGIN_SHADER_PARAMETER_STRUCT(FParameters, )
-		SHADER_PARAMETER_RDG_TEXTURE(Texture2D, ScreenSpaceReflectionsTexture)
-		SHADER_PARAMETER_SAMPLER(SamplerState, ScreenSpaceReflectionsSampler)
-		SHADER_PARAMETER_TEXTURE(Texture2D, PreIntegratedGF)
-		SHADER_PARAMETER_SAMPLER(SamplerState, PreIntegratedGFSampler)
-		SHADER_PARAMETER_RDG_TEXTURE(Texture2D, SceneNoWaterDepthTexture)
-		SHADER_PARAMETER_SAMPLER(SamplerState,  SceneNoWaterDepthSampler)
-		SHADER_PARAMETER_STRUCT_INCLUDE(FSceneTextureParameters, SceneTextures)	// Water scene texture
-		SHADER_PARAMETER_STRUCT_INCLUDE(FSceneTextureSamplerParameters, SceneTextureSamplers)
-		SHADER_PARAMETER_STRUCT_REF(FViewUniformShaderParameters, ViewUniformBuffer)
-		SHADER_PARAMETER_STRUCT_REF(FReflectionCaptureShaderData, ReflectionCaptureData)
-		SHADER_PARAMETER_STRUCT_REF(FReflectionUniformParameters, ReflectionsParameters)
+		SHADER_PARAMETER_STRUCT_INCLUDE(FSingleLayerWaterCommonShaderParameters, CommonParameters)
 		RENDER_TARGET_BINDING_SLOTS()
 	END_SHADER_PARAMETER_STRUCT()
 
@@ -239,6 +249,95 @@ class FSingleLayerWaterCompositePS : public FGlobalShader
 
 IMPLEMENT_GLOBAL_SHADER(FSingleLayerWaterCompositePS, "/Engine/Private/SingleLayerWaterComposite.usf", "SingleLayerWaterCompositePS", SF_Pixel);
 
+class FWaterTileCategorisationCS : public FGlobalShader
+{
+	DECLARE_GLOBAL_SHADER(FWaterTileCategorisationCS);
+	SHADER_USE_PARAMETER_STRUCT(FWaterTileCategorisationCS, FGlobalShader)
+
+	using FPermutationDomain = TShaderPermutationDomain<>;
+
+	static int32 GetTileSize()
+	{
+		return 8;
+	}
+
+	BEGIN_SHADER_PARAMETER_STRUCT(FParameters, )
+		SHADER_PARAMETER_STRUCT_INCLUDE(FSingleLayerWaterCommonShaderParameters, CommonParameters)
+		SHADER_PARAMETER(uint32, TiledViewWidth)
+		SHADER_PARAMETER(uint32, TiledViewHeight)
+		SHADER_PARAMETER(float, TiledViewWidthInv)
+		SHADER_PARAMETER(float, TiledViewHeightInv)
+		SHADER_PARAMETER_RDG_BUFFER_UAV(RWBuffer<uint>, DispatchIndirectDataUAV)
+		SHADER_PARAMETER_RDG_BUFFER_UAV(RWBuffer<uint>, WaterTileListDataUAV)
+	END_SHADER_PARAMETER_STRUCT()
+
+	static FPermutationDomain RemapPermutation(FPermutationDomain PermutationVector)
+	{
+		return PermutationVector;
+	}
+
+	static bool ShouldCompilePermutation(const FGlobalShaderPermutationParameters& Parameters)
+	{
+		if (!IsFeatureLevelSupported(Parameters.Platform, ERHIFeatureLevel::SM5) || !RHISupportsDrawIndirect(Parameters.Platform))
+		{
+			return false;
+		}
+		return true;
+	}
+
+	static void ModifyCompilationEnvironment(const FGlobalShaderPermutationParameters& Parameters, FShaderCompilerEnvironment& OutEnvironment)
+	{
+		OutEnvironment.SetDefine(TEXT("TILE_CATERGORISATION_SHADER"), 1.0f);
+		OutEnvironment.SetDefine(TEXT("WATER_TILE_SIZE"), GetTileSize());
+		FGlobalShader::ModifyCompilationEnvironment(Parameters, OutEnvironment);
+	}
+
+};
+
+IMPLEMENT_GLOBAL_SHADER(FWaterTileCategorisationCS, "/Engine/Private/SingleLayerWaterComposite.usf", "WaterTileCatergorisationCS", SF_Compute);
+
+// Disabled water composition due to non 32bits UAV operations
+/*class FWaterTiledCompositeCS : public FGlobalShader
+{
+	DECLARE_GLOBAL_SHADER(FWaterTiledCompositeCS);
+	SHADER_USE_PARAMETER_STRUCT(FWaterTiledCompositeCS, FGlobalShader)
+
+	using FPermutationDomain = TShaderPermutationDomain<FSingleLayerWaterScreenSpaceReflections>;
+
+	BEGIN_SHADER_PARAMETER_STRUCT(FParameters, )
+		SHADER_PARAMETER_STRUCT_INCLUDE(FSingleLayerWaterCommonShaderParameters, CommonParameters)
+		SHADER_PARAMETER(uint32, TiledViewWidth)
+		SHADER_PARAMETER(uint32, TiledViewHeight)
+		SHADER_PARAMETER_RDG_BUFFER_SRV(Buffer<uint>, WaterTileListData)
+		SHADER_PARAMETER_UAV(RWTexture2D<float3>, SceneColorUAV)
+		SHADER_PARAMETER_RDG_BUFFER(Buffer<uint>, IndirectDispatchParameters)	// Not used in shader but need to be reference in the parameter list
+	END_SHADER_PARAMETER_STRUCT()
+
+	static FPermutationDomain RemapPermutation(FPermutationDomain PermutationVector)
+	{
+		return PermutationVector;
+	}
+
+	static bool ShouldCompilePermutation(const FGlobalShaderPermutationParameters& Parameters)
+	{
+		if (!IsFeatureLevelSupported(Parameters.Platform, ERHIFeatureLevel::SM5) || !RHISupportsDrawIndirect(Parameters.Platform))
+		{
+			return false;
+		}
+		return true;
+	}
+
+	static void ModifyCompilationEnvironment(const FGlobalShaderPermutationParameters& Parameters, FShaderCompilerEnvironment& OutEnvironment)
+	{
+		OutEnvironment.SetDefine(TEXT("TILED_COMPOSITE_SHADER"), 1.0f);
+		OutEnvironment.SetDefine(TEXT("WATER_TILE_SIZE"), FWaterTileCategorisationCS::GetTileSize());
+		FGlobalShader::ModifyCompilationEnvironment(Parameters, OutEnvironment);
+	}
+
+};
+
+IMPLEMENT_GLOBAL_SHADER(FWaterTiledCompositeCS, "/Engine/Private/SingleLayerWaterComposite.usf", "WaterTiledCategorisationCS", SF_Compute);*/
+
 
 
 //////////////////////////////////////////////////////////////////////////
@@ -354,6 +453,59 @@ void FDeferredShadingSceneRenderer::RenderSingleLayerWaterReflections(FRHIComman
 		FSceneTextureParameters SceneTextures;
 		SetupSceneTextureParameters(GraphBuilder, &SceneTextures);
 
+		auto SetCommonParameters = [&](FSingleLayerWaterCommonShaderParameters& Parameters)
+		{
+			Parameters.ScreenSpaceReflectionsTexture = ReflectionsColor ? ReflectionsColor : GraphBuilder.RegisterExternalTexture(GSystemTextures.BlackDummy);
+			Parameters.ScreenSpaceReflectionsSampler = TStaticSamplerState<SF_Point>::GetRHI();
+			Parameters.PreIntegratedGF = GSystemTextures.PreintegratedGF->GetRenderTargetItem().ShaderResourceTexture;
+			Parameters.PreIntegratedGFSampler = TStaticSamplerState<SF_Bilinear, AM_Clamp, AM_Clamp, AM_Clamp>::GetRHI();
+			Parameters.SceneNoWaterDepthTexture = GraphBuilder.RegisterExternalTexture(PassData.SceneDepthZWithoutSingleLayerWater ? PassData.SceneDepthZWithoutSingleLayerWater : GSystemTextures.BlackDummy);
+			Parameters.SceneNoWaterDepthSampler = TStaticSamplerState<SF_Point>::GetRHI();
+			Parameters.SceneTextures = SceneTextures;
+			SetupSceneTextureSamplers(&Parameters.SceneTextureSamplers);
+			Parameters.ViewUniformBuffer = View.ViewUniformBuffer;
+			Parameters.ReflectionCaptureData = View.ReflectionCaptureUniformBuffer;
+			{
+				FReflectionUniformParameters ReflectionUniformParameters;
+				SetupReflectionUniformParameters(View, ReflectionUniformParameters);
+				Parameters.ReflectionsParameters = CreateUniformBufferImmediate(ReflectionUniformParameters, UniformBuffer_SingleDraw);
+			}
+		};
+
+		const bool bRunTiled = RHISupportsDrawIndirect(View.GetShaderPlatform()) && CVarWaterSingleLayerTiledComposite.GetValueOnRenderThread();
+		FTiledScreenSpaceReflection TiledScreenSpaceReflection = {nullptr, nullptr, nullptr, nullptr, nullptr, 8};
+		FIntVector ViewRes(View.ViewRect.Width(), View.ViewRect.Height(), 1);
+		FIntVector TiledViewRes = FIntVector::DivideAndRoundUp(ViewRes, TiledScreenSpaceReflection.TileSize);
+		if (bRunTiled)
+		{
+			TiledScreenSpaceReflection.DispatchIndirectParametersBuffer = GraphBuilder.CreateBuffer(FRDGBufferDesc::CreateIndirectDesc<FRHIDispatchIndirectParameters>(), TEXT("WaterIndirectDrawParameters"));
+			TiledScreenSpaceReflection.DispatchIndirectParametersBufferUAV = GraphBuilder.CreateUAV(TiledScreenSpaceReflection.DispatchIndirectParametersBuffer);
+			FRDGBufferDesc TileListStructuredBufferDesc = FRDGBufferDesc::CreateStructuredDesc(4, TiledViewRes.X * TiledViewRes.Y); // one uint32 element per tile
+			TiledScreenSpaceReflection.TileListDataBuffer = GraphBuilder.CreateBuffer(TileListStructuredBufferDesc, TEXT("WaterTileList"));
+			TiledScreenSpaceReflection.TileListStructureBufferUAV = GraphBuilder.CreateUAV(TiledScreenSpaceReflection.TileListDataBuffer);
+			TiledScreenSpaceReflection.TileListStructureBufferSRV = GraphBuilder.CreateSRV(TiledScreenSpaceReflection.TileListDataBuffer);
+
+			// Clear DispatchIndirectParametersBuffer
+			AddClearUAVPass(GraphBuilder, TiledScreenSpaceReflection.DispatchIndirectParametersBufferUAV, 0);
+
+			// Categorization based on SHADING_MODEL_ID
+			{
+				FWaterTileCategorisationCS::FPermutationDomain PermutationVector;
+				TShaderMapRef<FWaterTileCategorisationCS> ComputeShader(View.ShaderMap, PermutationVector);
+
+				FWaterTileCategorisationCS::FParameters* PassParameters = GraphBuilder.AllocParameters<FWaterTileCategorisationCS::FParameters>();
+				SetCommonParameters(PassParameters->CommonParameters);
+				PassParameters->TiledViewWidth = TiledViewRes.X;
+				PassParameters->TiledViewHeight = TiledViewRes.Y;
+				PassParameters->TiledViewWidthInv = 1.0f / float(TiledViewRes.X);
+				PassParameters->TiledViewHeightInv = 1.0f / float(TiledViewRes.Y);
+				PassParameters->DispatchIndirectDataUAV = TiledScreenSpaceReflection.DispatchIndirectParametersBufferUAV;
+				PassParameters->WaterTileListDataUAV = TiledScreenSpaceReflection.TileListStructureBufferUAV;
+
+				FComputeShaderUtils::AddPass(GraphBuilder, RDG_EVENT_NAME("WaterTileCategorisation"), *ComputeShader, PassParameters, TiledViewRes);
+			}
+		}
+
 		const bool bEnableSSR = CVarWaterSingleLayerSSR.GetValueOnRenderThread() != 0 && ShouldRenderScreenSpaceReflections(View);
 		if (bEnableSSR)
 		{
@@ -373,11 +525,11 @@ void FDeferredShadingSceneRenderer::RenderSingleLayerWaterReflections(FRHIComman
 
 			const bool bDenoise = false;
 			RenderScreenSpaceReflections(
-				GraphBuilder, SceneTextures, CurrentSceneColor, View, SSRQuality, bDenoise, &DenoiserInputs);
+				GraphBuilder, SceneTextures, CurrentSceneColor, View, SSRQuality, bDenoise, &DenoiserInputs, bRunTiled ? &TiledScreenSpaceReflection : nullptr);
 
 			ReflectionsColor = DenoiserInputs.Color;
 
-			if (CVarWaterSingleLayerSSRTAA.GetValueOnRenderThread())
+			if (CVarWaterSingleLayerSSRTAA.GetValueOnRenderThread() && IsSSRTemporalPassRequired(View)) // TAA pass is an option
 			{
 				check(View.ViewState);
 				FTAAPassParameters TAASettings(View);
@@ -397,27 +549,33 @@ void FDeferredShadingSceneRenderer::RenderSingleLayerWaterReflections(FRHIComman
 		}
 
 		// Composite reflections on water
+/*		if (bRunTiled)									// Disabled water composition due to non 32bits UAV operations
+		{
+			// Render Tiled composite CS shader
+			{
+				FWaterTiledCompositeCS::FPermutationDomain PermutationVector;
+				PermutationVector.Set<FSingleLayerWaterScreenSpaceReflections>(bEnableSSR);
+				TShaderMapRef<FWaterTiledCompositeCS> ComputeShader(View.ShaderMap, PermutationVector);
+
+				FWaterTiledCompositeCS::FParameters* PassParameters = GraphBuilder.AllocParameters<FWaterTiledCompositeCS::FParameters>();
+				SetCommonParameters(PassParameters->CommonParameters);
+				PassParameters->TiledViewWidth = TiledViewRes.X;
+				PassParameters->TiledViewHeight = TiledViewRes.Y;
+				PassParameters->WaterTileListData = TiledScreenSpaceReflection.TileListStructureBufferSRV;
+				PassParameters->SceneColorUAV = SceneContext.GetSceneColorTextureUAV();
+				PassParameters->IndirectDispatchParameters = TiledScreenSpaceReflection.DispatchIndirectParametersBuffer;
+				
+				FComputeShaderUtils::AddPass(GraphBuilder, RDG_EVENT_NAME("WaterTiledComposite"), *ComputeShader, PassParameters, TiledScreenSpaceReflection.DispatchIndirectParametersBuffer, 0);
+			}
+		}
+		else*/
 		{
 			FSingleLayerWaterCompositePS::FPermutationDomain PermutationVector;
-			PermutationVector.Set<FSingleLayerWaterCompositePS::FScreenSpaceReflections>(bEnableSSR);
+			PermutationVector.Set<FSingleLayerWaterScreenSpaceReflections>(bEnableSSR);
 			TShaderMapRef<FSingleLayerWaterCompositePS> PixelShader(View.ShaderMap, PermutationVector);
 
 			FSingleLayerWaterCompositePS::FParameters* PassParameters = GraphBuilder.AllocParameters<FSingleLayerWaterCompositePS::FParameters>();
-			PassParameters->ScreenSpaceReflectionsTexture = ReflectionsColor ? ReflectionsColor : GraphBuilder.RegisterExternalTexture(GSystemTextures.BlackDummy);
-			PassParameters->ScreenSpaceReflectionsSampler = TStaticSamplerState<SF_Point>::GetRHI();
-			PassParameters->PreIntegratedGF = GSystemTextures.PreintegratedGF->GetRenderTargetItem().ShaderResourceTexture;
-			PassParameters->PreIntegratedGFSampler = TStaticSamplerState<SF_Bilinear, AM_Clamp, AM_Clamp, AM_Clamp>::GetRHI();
-			PassParameters->SceneNoWaterDepthTexture = GraphBuilder.RegisterExternalTexture(PassData.SceneDepthZWithoutSingleLayerWater ? PassData.SceneDepthZWithoutSingleLayerWater : GSystemTextures.BlackDummy);
-			PassParameters->SceneNoWaterDepthSampler = TStaticSamplerState<SF_Point>::GetRHI();
-			PassParameters->SceneTextures = SceneTextures;
-			SetupSceneTextureSamplers(&PassParameters->SceneTextureSamplers);
-			PassParameters->ViewUniformBuffer = View.ViewUniformBuffer;
-			PassParameters->ReflectionCaptureData = View.ReflectionCaptureUniformBuffer;
-			{
-				FReflectionUniformParameters ReflectionUniformParameters;
-				SetupReflectionUniformParameters(View, ReflectionUniformParameters);
-				PassParameters->ReflectionsParameters = CreateUniformBufferImmediate(ReflectionUniformParameters, UniformBuffer_SingleDraw);
-			}
+			SetCommonParameters(PassParameters->CommonParameters);
 			PassParameters->RenderTargets[0] = FRenderTargetBinding(SceneColorTexture, ERenderTargetLoadAction::ELoad);
 			ClearUnusedGraphResources(*PixelShader, PassParameters);
 
