Support for YCoCg encoding of the base color in runtime virtual textures
Gives much better quality result but encoding is slower and there is a 25% memory increase in the physical texture

diff --git a/Engine/Shaders/Private/BCCompressionCommon.ush b/Engine/Shaders/Private/BCCompressionCommon.ush
--- a/Engine/Shaders/Private/BCCompressionCommon.ush
+++ b/Engine/Shaders/Private/BCCompressionCommon.ush
@@ -176,7 +176,7 @@ uint2 CompressBC1Block( in float3 Block[16] )
 	return uint2((MinColor565 << 16) | MaxColor565, Indices);
 }
 
-// Compress a BC1 block that will be sampled as BC1_UNORM
+// Compress a BC1 block that will be sampled as sRGB
 // We expect linear colors as input and convert internally
 uint2 CompressBC1Block_SRGB( in float3 Block[16] )
 {
@@ -225,7 +225,7 @@ uint4 CompressBC3Block( in float3 BlockRGB[16], in float BlockA[16] )
 	return uint4((MinAlphaUint << 8) | MaxAlphaUint | AlphaIndices.x, AlphaIndices.y, (MinColor565 << 16) | MaxColor565, ColorIndices);
 }
 
-// Compress a BC3 block that will be sampled as BC3_UNORM
+// Compress a BC3 block that will be sampled as sRGB
 // We expect linear colors as input and convert internally
 uint4 CompressBC3Block_SRGB( in float3 BlockRGB[16], in float BlockA[16] )
 {
@@ -299,3 +299,177 @@ uint4 CompressBC5Block( in float BlockU[16], in float BlockV[16] )
 
 	return uint4((MinUUint << 8) | MaxUUint | IndicesU.x, IndicesU.y, (MinVUint << 8) | MaxVUint | IndicesV.x, IndicesV.y);
 }
+
+// Convert RGB linear color to YCoCG
+float3 RGB2YCoCg(float3 RGB)
+{
+	float Y = (RGB.r + 2.f * RGB.g + RGB.b) * 0.25f;
+	float Co = ((2.f * RGB.r - 2.f * RGB.b) * 0.25f + 128.f / 255.f);
+	float Cg = ((-RGB.r + 2.f * RGB.g - RGB.b) * 0.25f + 128.f / 255.f);
+	return float3(Y, Co, Cg);
+}
+
+// Get a single scale factor to use for a YCoCg color block
+// This increases precision at the expense of potential blending artifacts across blocks
+float2 GetYCoCgScale(float2 MinCoCg, float2 MaxCoCg)
+{
+	MinCoCg = abs(MinCoCg - 128.f / 255.f);
+	MaxCoCg = abs(MaxCoCg - 128.f / 255.f);
+
+	float MaxComponent = max(max(MinCoCg.x, MinCoCg.y), max(MaxCoCg.x, MaxCoCg.y));
+
+	return (MaxComponent < 32.f / 255.f) ? float2(4.f, 0.25f) : (MaxComponent < 64.f / 255.f) ? float2(2.f, 0.5f) : float2(1.f, 1.f);
+}
+
+void ApplyYCoCgScale(inout float2 MinCoCg, inout float2 MaxCoCg, float Scale)
+{
+	MinCoCg = (MinCoCg - 128.f / 255.f) * Scale + 128.f / 255.f;
+	MaxCoCg = (MaxCoCg - 128.f / 255.f) * Scale + 128.f / 255.f;
+}
+
+// Inset the CoCg bounding end points
+void InsetCoCgEndPoints(inout float2 MinCoCg, inout float2 MaxCoGg)
+{
+	float2 Inset = (MaxCoGg - MinCoCg - (8.f / 255.f)) / 16.f;
+	MinCoCg = saturate(MinCoCg + Inset);
+	MaxCoGg = saturate(MaxCoGg - Inset);
+}
+
+// Inset the luminance end points
+void InsetLumaEndPoints(inout float MinY, inout float MaxY)
+{
+	float Inset = (MaxY - MinY - (16.f / 255.f)) / 32.f;
+	MinY = saturate(MinY + Inset);
+	MaxY = saturate(MaxY - Inset);
+}
+
+// Select the 2 min/max end points from the CoCg bounding rectangle based on the block contents 
+void AdjustMinMaxDiagonalYCoCg(const float3 Block[16], inout float2 MinCoCg, inout float2 MaxCoGg)
+{
+	float2 MidCoCg = (MaxCoGg + MinCoCg) * 0.5;
+
+	float Sum = 0.f;
+	for (int i = 0; i < 16; i++)
+	{
+		float2 Diff = Block[i].yz - MidCoCg;
+		Sum += Diff.x * Diff.y;
+	}
+	if (Sum < 0.f)
+	{
+		float Temp = MaxCoGg.y;
+		MaxCoGg.y = MinCoCg.y;
+		MinCoCg.y = Temp;
+	}
+}
+
+uint CoCgToUint565(inout float2 CoCg)
+{
+	float2 Scale = float2(31.f, 63.f);
+	float2 ColorScaled = round(saturate(CoCg) * Scale);
+	uint ColorPacked = ((uint)ColorScaled.r << 11) | ((uint)ColorScaled.g << 5);
+	CoCg = ColorScaled / Scale;
+
+	return ColorPacked;
+}
+
+// Calculate the final packed indices for the CoCg part of a color block
+uint GetPackedCoCgBlockIndices(in float3 Block[16], in float2 MinCoCg, in float2 MaxCoCg)
+{
+	uint PackedIndices = 0;
+
+	// Project onto max->min color vector and segment into range [0,3]
+	float2 Range = MinCoCg - MaxCoCg;
+	float Scale = 3.f / dot(Range, Range);
+	float2 ScaledRange = Range * Scale;
+	float Bias = (dot(MaxCoCg, MaxCoCg) - dot(MaxCoCg, MinCoCg)) * Scale;
+
+	for (int i = 15; i >= 0; --i)
+	{
+		// Compute the distance index for this element
+		uint Index = round(dot(Block[i].yz, ScaledRange) + Bias);
+		// Convert distance index into the BC index
+		Index += (Index > 0) - (3 * (Index == 3));
+		// OR into the final PackedIndices
+		PackedIndices = (PackedIndices << 2) | Index;
+	}
+
+	return PackedIndices;
+}
+
+// Calculate the final packed indices for the Luma part of a color block
+uint2 GetPackedLumaBlockIndices(in float3 Block[16], in float MinAlpha, in float MaxAlpha)
+{
+	uint2 PackedIndices = 0;
+
+	// Segment alpha max->min into range [0,7]
+	float Range = MinAlpha - MaxAlpha;
+	float Scale = 7.f / Range;
+	float Bias = -Scale * MaxAlpha;
+
+	uint i = 0;
+	// The first 5 elements of the block will go into the top 16 bits of the x component
+	for (i = 0; i < 5; ++i)
+	{
+		// Compute the distance index for this element
+		uint Index = round(Block[i].x * Scale + Bias);
+		// Convert distance index into the BC index
+		Index += (Index > 0) - (7 * (Index == 7));
+		// OR into the final PackedIndices
+		PackedIndices.x |= (Index << (i * 3 + 16));
+	}
+
+	// The 6th element is split across the x and y components
+	{
+		i = 5;
+		uint Index = round(Block[i].x * Scale + Bias);
+		Index += (Index > 0) - (7 * (Index == 7));
+		PackedIndices.x |= (Index << 31);
+		PackedIndices.y |= (Index >> 1);
+	}
+
+	// The rest of the elements fill the y component
+	for (i = 6; i < 16; ++i)
+	{
+		uint Index = round(Block[i].x * Scale + Bias);
+		Index += (Index > 0) - (7 * (Index == 7 ? 1 : 0));
+		PackedIndices.y |= (Index << (i * 3 - 16));
+	}
+
+	return PackedIndices;
+}
+
+// Convert a linear RGB block to YCoCg and compress a BC3 block 
+uint4 CompressBC3BlockYCoCg(in float3 Block[16])
+{
+	for (int i = 0; i < 16; ++i)
+	{
+		Block[i] = RGB2YCoCg(Block[i]);
+	}
+
+	float3 MinColor, MaxColor;
+	GetMinMax(Block, MinColor, MaxColor);
+
+	AdjustMinMaxDiagonalYCoCg(Block, MinColor.yz, MaxColor.yz);
+
+	float2 Scale = GetYCoCgScale(MinColor.yz, MaxColor.yz);
+
+	ApplyYCoCgScale(MinColor.yz, MaxColor.yz, Scale.x);
+
+	InsetCoCgEndPoints(MinColor.yz, MaxColor.yz);
+
+	uint MinColor565 = CoCgToUint565(MinColor.yz) | ((uint)Scale.x - 1);
+	uint MaxColor565 = CoCgToUint565(MaxColor.yz) | ((uint)Scale.x - 1);
+
+	ApplyYCoCgScale(MinColor.yz, MaxColor.yz, Scale.y);
+
+	uint CoCgIndices = GetPackedCoCgBlockIndices(Block, MinColor.yz, MaxColor.yz);
+
+	InsetLumaEndPoints(MinColor.x, MaxColor.x);
+
+	uint MinLumaUint = round(MinColor.x * 255.f);
+	uint MaxLumaUint = round(MaxColor.x * 255.f);
+
+	uint2 Indices = GetPackedLumaBlockIndices(Block, MinColor.x, MaxColor.x);
+
+	return uint4((MinLumaUint << 8) | MaxLumaUint | Indices.x, Indices.y, (MinColor565 << 16) | MaxColor565, CoCgIndices);
+}
diff --git a/Engine/Shaders/Private/VirtualTextureCommon.ush b/Engine/Shaders/Private/VirtualTextureCommon.ush
--- a/Engine/Shaders/Private/VirtualTextureCommon.ush
+++ b/Engine/Shaders/Private/VirtualTextureCommon.ush
@@ -534,6 +534,7 @@ MaterialFloat4 TextureVirtualSampleLevel(
 	return NumLevels-log2(vPagesWideInLevel);
 }*/
 
+
 /**
 	Helper function to convert world space to virtual texture UV space
 	todo[vt]: For large worlds we will probably want to use position relative to Camera to avoid precision issues
@@ -544,6 +545,16 @@ float2 VirtualTextureWorldToUV(in float3 WorldPos, in float3 Origin, in float3 U
 	return saturate(float2(dot(P, U), dot(P, V)));
 }
 
+/** Unpack color from YCoCg stored in BC3. */
+float3 VirtualTextureUnpackBaseColorYCoCg(in float4 PackedValue)
+{
+	float Y = PackedValue.a;
+	float Scale = 1.f / ((255.f / 8.f) * PackedValue.b + 1.f);
+	float Co = (PackedValue.r - 128.f / 255.f) * Scale;
+	float Cg = (PackedValue.g - 128.f / 255.f) * Scale;
+	return float3(Y + Co - Cg, Y + Cg, Y - Co - Cg);
+}
+
 /** Generic normal unpack funtion. */
 float3 VirtualTextureUnpackNormal(in float2 PackedXY, in float PackedSignZ)
 {
@@ -562,7 +573,6 @@ float3 VirtualTextureUnpackNormalBC3(in float4 PackedValue)
 /** Unpack normal from two BC3 textures. */
 float3 VirtualTextureUnpackNormalBC3BC3(in float4 PackedValue0, in float4 PackedValue1)
 {
-
 	return VirtualTextureUnpackNormal(float2(PackedValue0.a, PackedValue1.a), PackedValue1.b);
 }
 
@@ -572,6 +582,12 @@ float3 VirtualTextureUnpackNormalBC5(in float4 PackedValue)
 	return VirtualTextureUnpackNormal(PackedValue.rg, 1.f);
 }
 
+/** Unpack normal from BC5 and BC1 textures. */
+float3 VirtualTextureUnpackNormalBC5BC1(in float4 PackedValue0, in float4 PackedValue1)
+{
+	return VirtualTextureUnpackNormal(float2(PackedValue0.x, PackedValue0.y), PackedValue1.b);
+}
+
 /** Unpack 16 bit height (with some hardcoded range scale/bias). */
 float VirtualTextureUnpackHeightR16(in float4 PackedValue)
 {
diff --git a/Engine/Shaders/Private/VirtualTextureCompress.usf b/Engine/Shaders/Private/VirtualTextureCompress.usf
--- a/Engine/Shaders/Private/VirtualTextureCompress.usf
+++ b/Engine/Shaders/Private/VirtualTextureCompress.usf
@@ -19,8 +19,10 @@ SamplerState TextureSampler2;
 RWTexture2D<uint2> OutCompressTexture0_u2;
 RWTexture2D<uint4> OutCompressTexture0_u4;
 RWTexture2D<uint4> OutCompressTexture1;
+RWTexture2D<uint2> OutCompressTexture2;
 RWTexture2D<float4> OutCopyTexture0;
 RWTexture2D<float4> OutCopyTexture1;
+RWTexture2D<float4> OutCopyTexture2;
 
 // Read a 4x4 color block ready for BC1 compression
 void ReadBlockRGB( Texture2D<float4> SourceTexture, SamplerState TextureSampler, float2 UV, float2 TexelUVSize, out float3 Block[16] )
@@ -251,6 +253,61 @@ void CopyNormalSpecularExCS(
 	OutCopyTexture1[ThreadId.xy + DestRect.xy] = float4(RoughnessSpecularNormalZ, NormalXY.y);
 }
 
+/** Compress base color, normal, roughness and specular to BC3, BC5, BC1 target set. */
+[numthreads(8,8,1)]
+void CompressBaseColorNormalSpecularYCoCgCS(
+	uint3 ThreadId : SV_DispatchThreadID 
+	)
+{
+	uint2 SamplePos = ThreadId.xy * 4;
+	if (any(SamplePos >= DestRect.zw))
+		return;
+
+	float2 TexelUVSize = 1.f / float2(DestRect.zw);
+	float2 SampleUV = (float2(SamplePos) + 0.5f) * TexelUVSize;
+
+	float3 BlockBaseColor[16];
+	ReadBlockRGB(RenderTexture0, TextureSampler0, SampleUV, TexelUVSize, BlockBaseColor);
+
+	float BlockNormalX[16];
+	float BlockNormalY[16];
+	ReadBlockRG(RenderTexture1, TextureSampler1, SampleUV, TexelUVSize, BlockNormalX, BlockNormalY);
+
+	float3 BlockSpecular[16];
+	ReadBlockRGB(RenderTexture2, TextureSampler2, SampleUV, TexelUVSize, BlockSpecular);
+
+	for (int i=0; i<16; i++)
+	{
+		BlockSpecular[i].z = round(BlockSpecular[i].z);
+	}
+
+	OutCompressTexture0_u4[ThreadId.xy + DestRect.xy] = CompressBC3BlockYCoCg(BlockBaseColor);
+	OutCompressTexture1[ThreadId.xy + DestRect.xy] = CompressBC5Block(BlockNormalX, BlockNormalY);
+	OutCompressTexture2[ThreadId.xy + DestRect.xy] = CompressBC1Block(BlockSpecular);
+}
+
+/** Copy normal X and Y to the alpha channels of the final target textures (without compression)
+  * This path is used when we disable texture compression, because we need to keep the same final channel layout. */
+[numthreads(8,8,1)]
+void CopyNormalSpecularYCoCgCS(
+	uint3 ThreadId : SV_DispatchThreadID 
+	)
+{
+	uint2 SamplePos = ThreadId.xy;
+	if (any(ThreadId.xy >= DestRect.zw))
+		return;
+
+	float3 YCoCg = RGB2YCoCg(RenderTexture0[SamplePos].xyz);
+	float2 NormalXY = RenderTexture1[SamplePos].xy;
+	float3 RoughnessSpecularNormalZ = RenderTexture2[SamplePos].xyz;
+	
+	RoughnessSpecularNormalZ.z = round(RoughnessSpecularNormalZ.z);
+	
+	OutCopyTexture0[ThreadId.xy + DestRect.xy] = float4(YCoCg.yz, 0, YCoCg.x);
+	OutCopyTexture1[ThreadId.xy + DestRect.xy] = float4(NormalXY.xy, 0, 1);
+	OutCopyTexture2[ThreadId.xy + DestRect.xy] = float4(RoughnessSpecularNormalZ, 1);
+}
+
 /** Copy world height (without compression)
   * This path is used when want to copy the world height data into a thumbnail. */
 [numthreads(8,8,1)]
diff --git a/Engine/Source/Editor/VirtualTexturingEditor/Private/RuntimeVirtualTextureBuild.cpp b/Engine/Source/Editor/VirtualTexturingEditor/Private/RuntimeVirtualTextureBuild.cpp
--- a/Engine/Source/Editor/VirtualTexturingEditor/Private/RuntimeVirtualTextureBuild.cpp
+++ b/Engine/Source/Editor/VirtualTexturingEditor/Private/RuntimeVirtualTextureBuild.cpp
@@ -202,6 +202,7 @@ namespace RuntimeVirtualTexture
 						true,
 						RenderTileResources.GetRenderTarget(0), TileBox,
 						RenderTileResources.GetRenderTarget(1), TileBox,
+						RenderTileResources.GetRenderTarget(2), TileBox,
 						Transform,
 						UVRange,
 						RenderLevel,
diff --git a/Engine/Source/Editor/VirtualTexturingEditor/Private/RuntimeVirtualTextureThumbnailRenderer.cpp b/Engine/Source/Editor/VirtualTexturingEditor/Private/RuntimeVirtualTextureThumbnailRenderer.cpp
--- a/Engine/Source/Editor/VirtualTexturingEditor/Private/RuntimeVirtualTextureThumbnailRenderer.cpp
+++ b/Engine/Source/Editor/VirtualTexturingEditor/Private/RuntimeVirtualTextureThumbnailRenderer.cpp
@@ -83,6 +83,7 @@ void URuntimeVirtualTextureThumbnailRenderer::Draw(UObject* Object, int32 X, int
 			true,
 			RenderTarget->GetRenderTargetTexture(), DestRect,
 			nullptr, DestRect,
+			nullptr, DestRect,
 			Transform,
 			FBox2D(FVector2D(0, 0), FVector2D(1, 1)),
 			MaxLevel,
diff --git a/Engine/Source/Runtime/Engine/Private/Materials/HLSLMaterialTranslator.h b/Engine/Source/Runtime/Engine/Private/Materials/HLSLMaterialTranslator.h
--- a/Engine/Source/Runtime/Engine/Private/Materials/HLSLMaterialTranslator.h
+++ b/Engine/Source/Runtime/Engine/Private/Materials/HLSLMaterialTranslator.h
@@ -5174,11 +5174,12 @@ protected:
 		return AddInlinedCodeChunk(MCT_Float2, *SampleCode, *GetParameterCode(WorldPositionIndex), *GetParameterCode(P0), *GetParameterCode(P1), *GetParameterCode(P2));
 	}
 
-	virtual int32 VirtualTextureUnpack(int32 CodeIndex0, int32 CodeIndex1, EVirtualTextureUnpackType UnpackType) override
+	virtual int32 VirtualTextureUnpack(int32 CodeIndex0, int32 CodeIndex1, int32 CodeIndex2, EVirtualTextureUnpackType UnpackType) override
 	{
-		if (UnpackType == EVirtualTextureUnpackType::BaseColor)
+		if (UnpackType == EVirtualTextureUnpackType::BaseColorYCoCg)
 		{
-			return CodeIndex0 == INDEX_NONE ? INDEX_NONE : ComponentMask(CodeIndex0, 1, 1, 1, 0);
+			FString	SampleCode(TEXT("VirtualTextureUnpackBaseColorYCoCg(%s)"));
+			return CodeIndex0 == INDEX_NONE ? INDEX_NONE : AddCodeChunk(MCT_Float3, *SampleCode, *GetParameterCode(CodeIndex0));
 		}
 		else if (UnpackType == EVirtualTextureUnpackType::NormalBC3)
 		{
@@ -5195,17 +5196,10 @@ protected:
 			FString	SampleCode(TEXT("VirtualTextureUnpackNormalBC3BC3(%s, %s)"));
 			return CodeIndex0 == INDEX_NONE || CodeIndex1 == INDEX_NONE ? INDEX_NONE : AddCodeChunk(MCT_Float3, *SampleCode, *GetParameterCode(CodeIndex0), *GetParameterCode(CodeIndex1));
 		}
-		else if (UnpackType == EVirtualTextureUnpackType::SpecularR8)
+		else if (UnpackType == EVirtualTextureUnpackType::NormalBC5BC1)
 		{
-			return CodeIndex1 == INDEX_NONE ? INDEX_NONE : ComponentMask(CodeIndex1, 1, 0, 0, 0);
-		}
-		else if (UnpackType == EVirtualTextureUnpackType::RoughnessG8)
-		{
-			return CodeIndex1 == INDEX_NONE ? INDEX_NONE : ComponentMask(CodeIndex1, 0, 1, 0, 0);
-		}
-		else if (UnpackType == EVirtualTextureUnpackType::RoughnessB8)
-		{
-			return CodeIndex1 == INDEX_NONE ? INDEX_NONE : ComponentMask(CodeIndex1, 0, 0, 1, 0);
+			FString	SampleCode(TEXT("VirtualTextureUnpackNormalBC5BC1(%s, %s)"));
+			return CodeIndex0 == INDEX_NONE || CodeIndex1 == INDEX_NONE ? INDEX_NONE : AddCodeChunk(MCT_Float3, *SampleCode, *GetParameterCode(CodeIndex1), *GetParameterCode(CodeIndex2));
 		}
 		else if (UnpackType == EVirtualTextureUnpackType::HeightR16)
 		{
diff --git a/Engine/Source/Runtime/Engine/Private/Materials/MaterialExpressions.cpp b/Engine/Source/Runtime/Engine/Private/Materials/MaterialExpressions.cpp
--- a/Engine/Source/Runtime/Engine/Private/Materials/MaterialExpressions.cpp
+++ b/Engine/Source/Runtime/Engine/Private/Materials/MaterialExpressions.cpp
@@ -2168,6 +2168,8 @@ int32 UMaterialExpressionRuntimeVirtualTextureSample::Compile(class FMaterialCom
 
 	// Calculate the virtual texture layer and sampling/unpacking functions for this output
 	// Fallback to a sensible default value if the output isn't valid for the bound virtual texture
+	uint32 UnpackTarget = 0;
+	uint32 UnpackMask = 0;
 	EVirtualTextureUnpackType UnpackType = EVirtualTextureUnpackType::None;
 
 	bool bIsBaseColorValid = false;
@@ -2181,6 +2183,7 @@ int32 UMaterialExpressionRuntimeVirtualTextureSample::Compile(class FMaterialCom
 	case ERuntimeVirtualTextureMaterialType::BaseColor_Normal: bIsBaseColorValid = bIsNormalValid = true; break;
 	case ERuntimeVirtualTextureMaterialType::BaseColor_Normal_Specular:
 	case ERuntimeVirtualTextureMaterialType::BaseColor_Normal_Specular_Ex: bIsBaseColorValid = bIsNormalValid = bIsSpecularValid = true; break;
+	case ERuntimeVirtualTextureMaterialType::BaseColor_Normal_Specular_YCoCg: bIsBaseColorValid = bIsNormalValid = bIsSpecularValid = true; break;
 	case ERuntimeVirtualTextureMaterialType::WorldHeight: bIsWorldHeightValid = true; break;
 	}
 
@@ -2189,7 +2192,11 @@ int32 UMaterialExpressionRuntimeVirtualTextureSample::Compile(class FMaterialCom
 	case 0: 
 		if ((bIsParameter || bIsVirtualTextureValid) && bIsBaseColorValid)
 		{
-			UnpackType = EVirtualTextureUnpackType::BaseColor;
+			switch (MaterialType)
+			{
+			case ERuntimeVirtualTextureMaterialType::BaseColor_Normal_Specular_YCoCg: UnpackType = EVirtualTextureUnpackType::BaseColorYCoCg; break;
+			default: UnpackTarget = 0; UnpackMask = 0x7; break;
+			}
 		}
 		else
 		{
@@ -2199,7 +2206,12 @@ int32 UMaterialExpressionRuntimeVirtualTextureSample::Compile(class FMaterialCom
 	case 1:
 		if ((bIsParameter || bIsVirtualTextureValid) && bIsSpecularValid)
 		{
-			UnpackType = EVirtualTextureUnpackType::SpecularR8;
+			switch (MaterialType)
+			{
+			case ERuntimeVirtualTextureMaterialType::BaseColor_Normal_Specular:
+			case ERuntimeVirtualTextureMaterialType::BaseColor_Normal_Specular_Ex:  UnpackTarget = 1; UnpackMask = 0x1; break;
+			case ERuntimeVirtualTextureMaterialType::BaseColor_Normal_Specular_YCoCg: UnpackTarget = 2; UnpackMask = 0x1; break;
+			}
 		}
 		else
 		{
@@ -2211,8 +2223,9 @@ int32 UMaterialExpressionRuntimeVirtualTextureSample::Compile(class FMaterialCom
 		{
 			switch (MaterialType)
 			{
-			case ERuntimeVirtualTextureMaterialType::BaseColor_Normal_Specular: UnpackType = EVirtualTextureUnpackType::RoughnessB8; break;
-			case ERuntimeVirtualTextureMaterialType::BaseColor_Normal_Specular_Ex: UnpackType = EVirtualTextureUnpackType::RoughnessG8; break;
+			case ERuntimeVirtualTextureMaterialType::BaseColor_Normal_Specular: UnpackTarget = 1; UnpackMask = 0x4; break;
+			case ERuntimeVirtualTextureMaterialType::BaseColor_Normal_Specular_Ex: UnpackTarget = 1; UnpackMask = 0x2; break;
+			case ERuntimeVirtualTextureMaterialType::BaseColor_Normal_Specular_YCoCg: UnpackTarget = 2; UnpackMask = 0x2; break;
 			}
 		}
 		else
@@ -2228,6 +2241,7 @@ int32 UMaterialExpressionRuntimeVirtualTextureSample::Compile(class FMaterialCom
 			case ERuntimeVirtualTextureMaterialType::BaseColor_Normal: UnpackType = EVirtualTextureUnpackType::NormalBC5; break;
 			case ERuntimeVirtualTextureMaterialType::BaseColor_Normal_Specular: UnpackType = EVirtualTextureUnpackType::NormalBC3; break;
 			case ERuntimeVirtualTextureMaterialType::BaseColor_Normal_Specular_Ex: UnpackType = EVirtualTextureUnpackType::NormalBC3BC3; break;
+			case ERuntimeVirtualTextureMaterialType::BaseColor_Normal_Specular_YCoCg: UnpackType = EVirtualTextureUnpackType::NormalBC5BC1; break;
 			}
 		}
 		else
@@ -2250,7 +2264,7 @@ int32 UMaterialExpressionRuntimeVirtualTextureSample::Compile(class FMaterialCom
 	}
 	
 	// Compile the texture object references
-	enum { MAX_RVT_LAYERS = 2 };
+	enum { MAX_RVT_LAYERS = 3 };
 	const int32 TextureLayerCount = URuntimeVirtualTexture::GetLayerCount(MaterialType);
 	check(TextureLayerCount <= MAX_RVT_LAYERS);
 
@@ -2326,7 +2340,15 @@ int32 UMaterialExpressionRuntimeVirtualTextureSample::Compile(class FMaterialCom
 	}
 
 	// Compile any unpacking code
-	const int32 UnpackCodeIndex = Compiler->VirtualTextureUnpack(SampleCodeIndex[0], SampleCodeIndex[1], UnpackType);
+	int32 UnpackCodeIndex = INDEX_NONE;
+	if (UnpackType != EVirtualTextureUnpackType::None)
+	{
+		UnpackCodeIndex = Compiler->VirtualTextureUnpack(SampleCodeIndex[0], SampleCodeIndex[1], SampleCodeIndex[2], UnpackType);
+	}
+	else
+	{
+		UnpackCodeIndex = SampleCodeIndex[UnpackTarget] == INDEX_NONE ? INDEX_NONE : Compiler->ComponentMask(SampleCodeIndex[UnpackTarget], UnpackMask & 1, (UnpackMask >> 1) & 1, (UnpackMask >> 2) & 1, (UnpackMask >> 3) & 1);
+	}
 	return UnpackCodeIndex;
 }
 
diff --git a/Engine/Source/Runtime/Engine/Private/VT/RuntimeVirtualTexture.cpp b/Engine/Source/Runtime/Engine/Private/VT/RuntimeVirtualTexture.cpp
--- a/Engine/Source/Runtime/Engine/Private/VT/RuntimeVirtualTexture.cpp
+++ b/Engine/Source/Runtime/Engine/Private/VT/RuntimeVirtualTexture.cpp
@@ -258,6 +258,8 @@ int32 URuntimeVirtualTexture::GetLayerCount(ERuntimeVirtualTextureMaterialType I
 	case ERuntimeVirtualTextureMaterialType::BaseColor_Normal_Specular:
 	case ERuntimeVirtualTextureMaterialType::BaseColor_Normal_Specular_Ex:
 		return 2;
+	case ERuntimeVirtualTextureMaterialType::BaseColor_Normal_Specular_YCoCg:
+		return 3;
 	default:
 		break;
 	}
@@ -283,6 +285,7 @@ EPixelFormat URuntimeVirtualTexture::GetLayerFormat(int32 LayerIndex) const
 		case ERuntimeVirtualTextureMaterialType::BaseColor_Normal_Specular:
 			return bCompressTextures ? PF_DXT1 : PF_B8G8R8A8;
 		case ERuntimeVirtualTextureMaterialType::BaseColor_Normal_Specular_Ex:
+		case ERuntimeVirtualTextureMaterialType::BaseColor_Normal_Specular_YCoCg:
 			return bCompressTextures ? PF_DXT5 : PF_B8G8R8A8;
 		case ERuntimeVirtualTextureMaterialType::WorldHeight:
 			return PF_G16;
@@ -295,12 +298,21 @@ EPixelFormat URuntimeVirtualTexture::GetLayerFormat(int32 LayerIndex) const
 		switch (MaterialType)
 		{
 		case ERuntimeVirtualTextureMaterialType::BaseColor_Normal:
+		case ERuntimeVirtualTextureMaterialType::BaseColor_Normal_Specular_YCoCg:
 			return bCompressTextures ? PF_BC5 : PF_B8G8R8A8;
 		case ERuntimeVirtualTextureMaterialType::BaseColor_Normal_Specular:
 		case ERuntimeVirtualTextureMaterialType::BaseColor_Normal_Specular_Ex:
 			return bCompressTextures ? PF_DXT5 : PF_B8G8R8A8;
-		case ERuntimeVirtualTextureMaterialType::BaseColor:
-		case ERuntimeVirtualTextureMaterialType::WorldHeight:
+		default:
+			break;
+		}
+	}
+	else if (LayerIndex == 2)
+	{
+		switch (MaterialType)
+		{
+		case ERuntimeVirtualTextureMaterialType::BaseColor_Normal_Specular_YCoCg:
+			return bCompressTextures ? PF_DXT1 : PF_B8G8R8A8;
 		default:
 			break;
 		}
@@ -321,6 +333,7 @@ bool URuntimeVirtualTexture::IsLayerSRGB(int32 LayerIndex) const
 	case ERuntimeVirtualTextureMaterialType::BaseColor_Normal_Specular_Ex:
 		// Only BaseColor layer is sRGB
 		return LayerIndex == 0;
+	case ERuntimeVirtualTextureMaterialType::BaseColor_Normal_Specular_YCoCg:
 	case ERuntimeVirtualTextureMaterialType::WorldHeight:
 		return false;
 	default:
@@ -491,7 +504,7 @@ void URuntimeVirtualTexture::InitializeStreamingTexture(uint32 InSizeX, uint32 I
 
 	StreamingTexture->BuildHash = GetStreamingTextureBuildHash();
 
-	enum { MAX_RVT_LAYERS = 2 };
+	enum { MAX_RVT_LAYERS = 3 };
 	const int32 LayerCount = GetLayerCount();
 	check(LayerCount <= MAX_RVT_LAYERS);
 	ETextureSourceFormat LayerFormats[MAX_RVT_LAYERS];
diff --git a/Engine/Source/Runtime/Engine/Public/MaterialCompiler.h b/Engine/Source/Runtime/Engine/Public/MaterialCompiler.h
--- a/Engine/Source/Runtime/Engine/Public/MaterialCompiler.h
+++ b/Engine/Source/Runtime/Engine/Public/MaterialCompiler.h
@@ -31,13 +31,11 @@ enum EMaterialForceCastFlags
 enum class EVirtualTextureUnpackType
 {
 	None,
-	BaseColor,
+	BaseColorYCoCg,
 	NormalBC3,
 	NormalBC5,
 	NormalBC3BC3,
-	SpecularR8,
-	RoughnessG8,
-	RoughnessB8,
+	NormalBC5BC1,
 	HeightR16,
 };
 
@@ -211,7 +209,7 @@ public:
 	virtual int32 VirtualTextureUniform(int32 TextureIndex, int32 VectorIndex) = 0;
 	virtual int32 VirtualTextureUniform(FName ParameterName, int32 TextureIndex, int32 VectorIndex) = 0;
 	virtual int32 VirtualTextureWorldToUV(int32 WorldPositionIndex, int32 P0, int32 P1, int32 P2) = 0;
-	virtual int32 VirtualTextureUnpack(int32 CodeIndex0, int32 CodeIndex1, EVirtualTextureUnpackType UnpackType) = 0;
+	virtual int32 VirtualTextureUnpack(int32 CodeIndex0, int32 CodeIndex1, int32 CodeIndex2, EVirtualTextureUnpackType UnpackType) = 0;
 
 	virtual int32 ExternalTexture(const FGuid& ExternalTextureGuid) = 0;
 	virtual int32 ExternalTexture(UTexture* InTexture, int32& TextureReferenceIndex) = 0;
@@ -498,7 +496,7 @@ public:
 	virtual int32 VirtualTextureUniform(int32 TextureIndex, int32 VectorIndex) override { return Compiler->VirtualTextureUniform(TextureIndex, VectorIndex); }
 	virtual int32 VirtualTextureUniform(FName ParameterName, int32 TextureIndex, int32 VectorIndex) override { return Compiler->VirtualTextureUniform(ParameterName, TextureIndex, VectorIndex); }
 	virtual int32 VirtualTextureWorldToUV(int32 WorldPositionIndex, int32 P0, int32 P1, int32 P2) override { return Compiler->VirtualTextureWorldToUV(WorldPositionIndex, P0, P1, P2); }
-	virtual int32 VirtualTextureUnpack(int32 CodeIndex0, int32 CodeIndex1, EVirtualTextureUnpackType UnpackType) override { return Compiler->VirtualTextureUnpack(CodeIndex0, CodeIndex1, UnpackType); }
+	virtual int32 VirtualTextureUnpack(int32 CodeIndex0, int32 CodeIndex1, int32 CodeIndex2, EVirtualTextureUnpackType UnpackType) override { return Compiler->VirtualTextureUnpack(CodeIndex0, CodeIndex1, CodeIndex2, UnpackType); }
 
 	virtual int32 ExternalTexture(const FGuid& ExternalTextureGuid) override { return Compiler->ExternalTexture(ExternalTextureGuid); }
 	virtual int32 ExternalTexture(UTexture* InTexture, int32& TextureReferenceIndex) override { return Compiler->ExternalTexture(InTexture, TextureReferenceIndex); }
diff --git a/Engine/Source/Runtime/Engine/Public/VT/RuntimeVirtualTextureEnum.h b/Engine/Source/Runtime/Engine/Public/VT/RuntimeVirtualTextureEnum.h
--- a/Engine/Source/Runtime/Engine/Public/VT/RuntimeVirtualTextureEnum.h
+++ b/Engine/Source/Runtime/Engine/Public/VT/RuntimeVirtualTextureEnum.h
@@ -17,6 +17,7 @@ enum class ERuntimeVirtualTextureMaterialType : uint8
 	BaseColor_Normal UMETA(DisplayName = "Base Color, Normal"),
 	BaseColor_Normal_Specular UMETA(DisplayName = "[BC1|BC3] Base Color, Normal, Roughness, Specular"),
 	BaseColor_Normal_Specular_Ex UMETA(DisplayName = "[BC3|BC3] Base Color, Normal, Roughness, Specular"),
+	BaseColor_Normal_Specular_YCoCg UMETA(DisplayName = "[BC3|BC5|BC1] Base Color, Normal, Roughness, Specular"),
 	WorldHeight UMETA(DisplayName = "World Height"),
 
 	Count UMETA(Hidden),
diff --git a/Engine/Source/Runtime/Renderer/Private/VT/RuntimeVirtualTextureProducer.cpp b/Engine/Source/Runtime/Renderer/Private/VT/RuntimeVirtualTextureProducer.cpp
--- a/Engine/Source/Runtime/Renderer/Private/VT/RuntimeVirtualTextureProducer.cpp
+++ b/Engine/Source/Runtime/Renderer/Private/VT/RuntimeVirtualTextureProducer.cpp
@@ -65,6 +65,9 @@ void FRuntimeVirtualTextureFinalizer::Finalize(FRHICommandListImmediate& RHICmdL
 		const FVector2D DestinationBoxStart1(Entry.DestX1 * TileSize, Entry.DestY1 * TileSize);
 		const FBox2D DestinationBox1(DestinationBoxStart1, DestinationBoxStart1 + FVector2D(TileSize, TileSize));
 
+		const FVector2D DestinationBoxStart2(Entry.DestX2 * TileSize, Entry.DestY2 * TileSize);
+		const FBox2D DestinationBox2(DestinationBoxStart2, DestinationBoxStart2 + FVector2D(TileSize, TileSize));
+
 		const float X = (float)FMath::ReverseMortonCode2(Entry.vAddress);
 		const float Y = (float)FMath::ReverseMortonCode2(Entry.vAddress >> 1);
 		const float DivisorX = (float)Desc.BlockWidthInTiles / (float)(1 << Entry.vLevel);
@@ -85,6 +88,8 @@ void FRuntimeVirtualTextureFinalizer::Finalize(FRHICommandListImmediate& RHICmdL
 			DestinationBox0, 
 			Entry.Texture1,
 			DestinationBox1,
+			Entry.Texture2,
+			DestinationBox2,
 			UVToWorld,
 			UVRange, 
 			Entry.vLevel,
@@ -138,7 +143,7 @@ IVirtualTextureFinalizer* FRuntimeVirtualTextureProducer::ProducePageData(
 	// This can be almost always be avoided by setting up the physical pools correctly for the application's needs.
 	// If we can't avoid partial layer masks then we could look at ways to handle it more efficiently (right now we render all layers even for these partial requests).
 
-	//todo[vt]: Add support for more than two layers
+	//todo[vt]: Tidy up logic now we have 3 layers
 	if (TargetLayers[0].TextureRHI != nullptr)
 	{
 		Tile.Texture0 = TargetLayers[0].TextureRHI->GetTexture2D();
@@ -153,6 +158,13 @@ IVirtualTextureFinalizer* FRuntimeVirtualTextureProducer::ProducePageData(
 		Tile.DestY1 = TargetLayers[1].pPageLocation.Y;
 	}
 
+	if (TargetLayers[2].TextureRHI != nullptr)
+	{
+		Tile.Texture2 = TargetLayers[2].TextureRHI->GetTexture2D();
+		Tile.DestX2 = TargetLayers[2].pPageLocation.X;
+		Tile.DestY2 = TargetLayers[2].pPageLocation.Y;
+	}
+
 	Finalizer.InitProducer(ProducerHandle);
 	Finalizer.AddTile(Tile);
 
diff --git a/Engine/Source/Runtime/Renderer/Private/VT/RuntimeVirtualTextureProducer.h b/Engine/Source/Runtime/Renderer/Private/VT/RuntimeVirtualTextureProducer.h
--- a/Engine/Source/Runtime/Renderer/Private/VT/RuntimeVirtualTextureProducer.h
+++ b/Engine/Source/Runtime/Renderer/Private/VT/RuntimeVirtualTextureProducer.h
@@ -22,10 +22,13 @@ public:
 	{
 		FRHITexture2D* Texture0 = nullptr;
 		FRHITexture2D* Texture1 = nullptr;
+		FRHITexture2D* Texture2 = nullptr;
 		int32 DestX0 = 0;
 		int32 DestY0 = 0;
 		int32 DestX1 = 0;
 		int32 DestY1 = 0;
+		int32 DestX2 = 0;
+		int32 DestY2 = 0;
 		uint32 vAddress = 0;
 		uint8 vLevel = 0;
 	};
diff --git a/Engine/Source/Runtime/Renderer/Private/VT/RuntimeVirtualTextureRender.cpp b/Engine/Source/Runtime/Renderer/Private/VT/RuntimeVirtualTextureRender.cpp
--- a/Engine/Source/Runtime/Renderer/Private/VT/RuntimeVirtualTextureRender.cpp
+++ b/Engine/Source/Runtime/Renderer/Private/VT/RuntimeVirtualTextureRender.cpp
@@ -267,6 +267,7 @@ namespace RuntimeVirtualTexture
 						break;
 					case ERuntimeVirtualTextureMaterialType::BaseColor_Normal_Specular:
 					case ERuntimeVirtualTextureMaterialType::BaseColor_Normal_Specular_Ex:
+					case ERuntimeVirtualTextureMaterialType::BaseColor_Normal_Specular_YCoCg:
 						Process<FMaterialPolicy_BaseColorNormalSpecular>(MeshBatch, BatchElementMask, StaticMeshId, PrimitiveSceneProxy, MaterialRenderProxy, Material);
 						break;
 					case ERuntimeVirtualTextureMaterialType::WorldHeight:
@@ -440,8 +441,10 @@ namespace RuntimeVirtualTexture
 			SHADER_PARAMETER_RDG_TEXTURE_UAV(RWTexture<uint2>, OutCompressTexture0_u2)
 			SHADER_PARAMETER_RDG_TEXTURE_UAV(RWTexture<uint4>, OutCompressTexture0_u4)
 			SHADER_PARAMETER_RDG_TEXTURE_UAV(RWTexture<uint4>, OutCompressTexture1)
+			SHADER_PARAMETER_RDG_TEXTURE_UAV(RWTexture<uint2>, OutCompressTexture2)
 			SHADER_PARAMETER_RDG_TEXTURE_UAV(RWTexture<float4>, OutCopyTexture0)
 			SHADER_PARAMETER_RDG_TEXTURE_UAV(RWTexture<float4>, OutCopyTexture1)
+			SHADER_PARAMETER_RDG_TEXTURE_UAV(RWTexture<float4>, OutCopyTexture2)
 		END_SHADER_PARAMETER_STRUCT()
 
 		static bool ShouldCompilePermutation(FGlobalShaderPermutationParameters const& Parameters)
@@ -486,12 +489,18 @@ namespace RuntimeVirtualTexture
 	typedef FShader_VirtualTextureCompress_CS< ERuntimeVirtualTextureMaterialType::BaseColor_Normal_Specular_Ex, false > FShader_VirtualTextureCompress_BaseColorNormalSpecularEx_CS;
 	IMPLEMENT_SHADER_TYPE(template<>, FShader_VirtualTextureCompress_BaseColorNormalSpecularEx_CS, TEXT("/Engine/Private/VirtualTextureCompress.usf"), TEXT("CompressBaseColorNormalSpecularExCS"), SF_Compute);
 
+	typedef FShader_VirtualTextureCompress_CS< ERuntimeVirtualTextureMaterialType::BaseColor_Normal_Specular_YCoCg, false > FShader_VirtualTextureCompress_BaseColorNormalSpecularYCoCg_CS;
+	IMPLEMENT_SHADER_TYPE(template<>, FShader_VirtualTextureCompress_BaseColorNormalSpecularYCoCg_CS, TEXT("/Engine/Private/VirtualTextureCompress.usf"), TEXT("CompressBaseColorNormalSpecularYCoCgCS"), SF_Compute);
+
 	typedef FShader_VirtualTextureCompress_CS< ERuntimeVirtualTextureMaterialType::BaseColor_Normal_Specular, true > FShader_VirtualTextureCompress_BaseColorNormalSpecular_CopyOnly_CS;
 	IMPLEMENT_SHADER_TYPE(template<>, FShader_VirtualTextureCompress_BaseColorNormalSpecular_CopyOnly_CS, TEXT("/Engine/Private/VirtualTextureCompress.usf"), TEXT("CopyNormalSpecularCS"), SF_Compute);
 
 	typedef FShader_VirtualTextureCompress_CS< ERuntimeVirtualTextureMaterialType::BaseColor_Normal_Specular_Ex, true > FShader_VirtualTextureCompress_BaseColorNormalSpecularEx_CopyOnly_CS;
 	IMPLEMENT_SHADER_TYPE(template<>, FShader_VirtualTextureCompress_BaseColorNormalSpecularEx_CopyOnly_CS, TEXT("/Engine/Private/VirtualTextureCompress.usf"), TEXT("CopyNormalSpecularExCS"), SF_Compute);
 
+	typedef FShader_VirtualTextureCompress_CS< ERuntimeVirtualTextureMaterialType::BaseColor_Normal_Specular_YCoCg, true > FShader_VirtualTextureCompress_BaseColorNormalSpecularYCoCg_CopyOnly_CS;
+	IMPLEMENT_SHADER_TYPE(template<>, FShader_VirtualTextureCompress_BaseColorNormalSpecularYCoCg_CopyOnly_CS, TEXT("/Engine/Private/VirtualTextureCompress.usf"), TEXT("CopyNormalSpecularYCoCgCS"), SF_Compute);
+
 	typedef FShader_VirtualTextureCompress_CS< ERuntimeVirtualTextureMaterialType::WorldHeight, true > FShader_VirtualTextureCompress_WorldHeight_CopyOnly_CS;
 	IMPLEMENT_SHADER_TYPE(template<>, FShader_VirtualTextureCompress_WorldHeight_CopyOnly_CS, TEXT("/Engine/Private/VirtualTextureCompress.usf"), TEXT("CopyWorldHeightCS"), SF_Compute);
 
@@ -529,6 +538,9 @@ namespace RuntimeVirtualTexture
 		case ERuntimeVirtualTextureMaterialType::BaseColor_Normal_Specular_Ex:
 			AddCompressOrCopyPass<ERuntimeVirtualTextureMaterialType::BaseColor_Normal_Specular_Ex>(GraphBuilder, FeatureLevel, Parameters, GroupCount);
 			break;
+		case ERuntimeVirtualTextureMaterialType::BaseColor_Normal_Specular_YCoCg:
+			AddCompressOrCopyPass<ERuntimeVirtualTextureMaterialType::BaseColor_Normal_Specular_YCoCg>(GraphBuilder, FeatureLevel, Parameters, GroupCount);
+			break;
 		}
 	}
 
@@ -547,6 +559,9 @@ namespace RuntimeVirtualTexture
 		case ERuntimeVirtualTextureMaterialType::BaseColor_Normal_Specular_Ex:
 			AddCompressOrCopyPass<ERuntimeVirtualTextureMaterialType::BaseColor_Normal_Specular_Ex, true>(GraphBuilder, FeatureLevel, Parameters, GroupCount);
 			break;
+		case ERuntimeVirtualTextureMaterialType::BaseColor_Normal_Specular_YCoCg:
+			AddCompressOrCopyPass<ERuntimeVirtualTextureMaterialType::BaseColor_Normal_Specular_YCoCg, true>(GraphBuilder, FeatureLevel, Parameters, GroupCount);
+			break;
 		case ERuntimeVirtualTextureMaterialType::WorldHeight:
 			AddCompressOrCopyPass<ERuntimeVirtualTextureMaterialType::WorldHeight, true>(GraphBuilder, FeatureLevel, Parameters, GroupCount);
 			break;
@@ -562,7 +577,9 @@ namespace RuntimeVirtualTexture
 		{
 			bRenderPass = OutputTexture0 != nullptr;
 			bCompressPass = bRenderPass && (OutputTexture0->GetFormat() == PF_DXT1 || OutputTexture0->GetFormat() == PF_DXT5 || OutputTexture0->GetFormat() == PF_BC5);
-			bCopyPass = bRenderPass && !bCompressPass && (MaterialType == ERuntimeVirtualTextureMaterialType::BaseColor_Normal_Specular || MaterialType == ERuntimeVirtualTextureMaterialType::BaseColor_Normal_Specular_Ex);
+			bCopyPass = bRenderPass && !bCompressPass && MaterialType == ERuntimeVirtualTextureMaterialType::BaseColor_Normal_Specular;
+			bCopyPass |= bRenderPass && !bCompressPass && MaterialType == ERuntimeVirtualTextureMaterialType::BaseColor_Normal_Specular_Ex;
+			bCopyPass |= bRenderPass && !bCompressPass && MaterialType == ERuntimeVirtualTextureMaterialType::BaseColor_Normal_Specular_YCoCg;
 			bCopyPass |= bRenderPass && MaterialType == ERuntimeVirtualTextureMaterialType::WorldHeight && OutputTexture0->GetFormat() == PF_B8G8R8A8;
 		
 			switch (MaterialType)
@@ -585,7 +602,7 @@ namespace RuntimeVirtualTexture
 				}
 				if (bCompressPass)
 				{
-					OutputAlias0 = CompressTexture0_u2 = GraphBuilder.CreateTexture(FPooledRenderTargetDesc::Create2DDesc(TextureSize / 4, PF_R32G32_UINT, FClearValueBinding::None, TexCreate_None, TexCreate_UAV, false), TEXT("CompressTexture0"));
+					OutputAlias0 = CompressTexture0_u4 = GraphBuilder.CreateTexture(FPooledRenderTargetDesc::Create2DDesc(TextureSize / 4, PF_R32G32B32A32_UINT, FClearValueBinding::None, TexCreate_None, TexCreate_UAV, false), TEXT("CompressTexture0"));
 					OutputAlias1 = CompressTexture1 = GraphBuilder.CreateTexture(FPooledRenderTargetDesc::Create2DDesc(TextureSize / 4, PF_R32G32B32A32_UINT, FClearValueBinding::None, TexCreate_None, TexCreate_UAV, false), TEXT("CompressTexture1"));
 				}
 				break;
@@ -624,6 +641,26 @@ namespace RuntimeVirtualTexture
 					OutputAlias1 = CopyTexture1 = GraphBuilder.CreateTexture(FPooledRenderTargetDesc::Create2DDesc(TextureSize, PF_B8G8R8A8, FClearValueBinding::None, TexCreate_None, TexCreate_UAV, false), TEXT("CopyTexture1"));
 				}
 				break;
+			case ERuntimeVirtualTextureMaterialType::BaseColor_Normal_Specular_YCoCg:
+				if (bRenderPass)
+				{
+					OutputAlias0 = RenderTexture0 = GraphBuilder.CreateTexture(FPooledRenderTargetDesc::Create2DDesc(TextureSize, PF_B8G8R8A8, FClearValueBinding::Black, TexCreate_SRGB, TexCreate_RenderTargetable, false), TEXT("RenderTexture0"));
+					RenderTexture1 = GraphBuilder.CreateTexture(FPooledRenderTargetDesc::Create2DDesc(TextureSize, PF_B8G8R8A8, FClearValueBinding::DefaultNormal8Bit, TexCreate_None, TexCreate_RenderTargetable, false), TEXT("RenderTexture1"));
+					RenderTexture2 = GraphBuilder.CreateTexture(FPooledRenderTargetDesc::Create2DDesc(TextureSize, PF_B8G8R8A8, FClearValueBinding::DefaultNormal8Bit, TexCreate_None, TexCreate_RenderTargetable, false), TEXT("RenderTexture2"));
+				}
+				if (bCompressPass)
+				{
+					OutputAlias0 = CompressTexture0_u4 = GraphBuilder.CreateTexture(FPooledRenderTargetDesc::Create2DDesc(TextureSize / 4, PF_R32G32B32A32_UINT, FClearValueBinding::None, TexCreate_None, TexCreate_UAV, false), TEXT("CompressTexture0"));
+					OutputAlias1 = CompressTexture1 = GraphBuilder.CreateTexture(FPooledRenderTargetDesc::Create2DDesc(TextureSize / 4, PF_R32G32B32A32_UINT, FClearValueBinding::None, TexCreate_None, TexCreate_UAV, false), TEXT("CompressTexture1"));
+					OutputAlias2 = CompressTexture2 = GraphBuilder.CreateTexture(FPooledRenderTargetDesc::Create2DDesc(TextureSize / 4, PF_R32G32_UINT, FClearValueBinding::None, TexCreate_None, TexCreate_UAV, false), TEXT("CompressTexture2"));
+				}
+				if (bCopyPass)
+				{
+					OutputAlias0 = CopyTexture0 = GraphBuilder.CreateTexture(FPooledRenderTargetDesc::Create2DDesc(TextureSize, PF_B8G8R8A8, FClearValueBinding::None, TexCreate_None, TexCreate_UAV, false), TEXT("CopyTexture0"));
+					OutputAlias1 = CopyTexture1 = GraphBuilder.CreateTexture(FPooledRenderTargetDesc::Create2DDesc(TextureSize, PF_B8G8R8A8, FClearValueBinding::None, TexCreate_None, TexCreate_UAV, false), TEXT("CopyTexture1"));
+					OutputAlias2 = CopyTexture2 = GraphBuilder.CreateTexture(FPooledRenderTargetDesc::Create2DDesc(TextureSize, PF_B8G8R8A8, FClearValueBinding::None, TexCreate_None, TexCreate_UAV, false), TEXT("CopyTexture2"));
+				}
+				break;
 			case ERuntimeVirtualTextureMaterialType::WorldHeight:
 				if (bRenderPass)
 				{
@@ -649,13 +686,16 @@ namespace RuntimeVirtualTexture
 		FRDGTextureRef CompressTexture0_u2 = nullptr;
 		FRDGTextureRef CompressTexture0_u4 = nullptr;
 		FRDGTextureRef CompressTexture1 = nullptr;
+		FRDGTextureRef CompressTexture2 = nullptr;
 		FRDGTextureRef CopyTexture0 = nullptr;
 		FRDGTextureRef CopyTexture1 = nullptr;
+		FRDGTextureRef CopyTexture2 = nullptr;
 
 		/** Aliases to one of the render/compress/copy textures. This is what we will Copy into the final physical texture. */
 		//todo[vt]: On platforms that support direct aliasing we can not set these and compress direct to the final destination
 		FRDGTextureRef OutputAlias0 = nullptr;
 		FRDGTextureRef OutputAlias1 = nullptr;
+		FRDGTextureRef OutputAlias2 = nullptr;
 	};
 
 
@@ -678,6 +718,8 @@ namespace RuntimeVirtualTexture
 		FBox2D const& DestBox0,
 		FRHITexture2D* OutputTexture1,
 		FBox2D const& DestBox1,
+		FRHITexture2D* OutputTexture2, 
+		FBox2D const& DestBox2,
 		FTransform const& UVToWorld,
 		FBox2D const& UVRange,
 		uint8 vLevel,
@@ -775,6 +817,7 @@ namespace RuntimeVirtualTexture
 			PassParameters->OutCompressTexture0_u2 = GraphSetup.CompressTexture0_u2 ? GraphBuilder.CreateUAV(FRDGTextureUAVDesc(GraphSetup.CompressTexture0_u2)) : nullptr;
 			PassParameters->OutCompressTexture0_u4 = GraphSetup.CompressTexture0_u4 ? GraphBuilder.CreateUAV(FRDGTextureUAVDesc(GraphSetup.CompressTexture0_u4)) : nullptr;
 			PassParameters->OutCompressTexture1 = GraphSetup.CompressTexture1 ? GraphBuilder.CreateUAV(FRDGTextureUAVDesc(GraphSetup.CompressTexture1)) : nullptr;
+			PassParameters->OutCompressTexture2 = GraphSetup.CompressTexture2 ? GraphBuilder.CreateUAV(FRDGTextureUAVDesc(GraphSetup.CompressTexture2)) : nullptr;
 
 			AddCompressPass(GraphBuilder, View->GetFeatureLevel(), PassParameters, TextureSize, MaterialType);
 		}
@@ -792,6 +835,7 @@ namespace RuntimeVirtualTexture
 			PassParameters->TextureSampler2 = TStaticSamplerState<SF_Point, AM_Clamp, AM_Clamp, AM_Clamp>::GetRHI();
 			PassParameters->OutCopyTexture0 = GraphSetup.CopyTexture0 ? GraphBuilder.CreateUAV(FRDGTextureUAVDesc(GraphSetup.CopyTexture0)) : nullptr;
 			PassParameters->OutCopyTexture1 = GraphSetup.CopyTexture1 ? GraphBuilder.CreateUAV(FRDGTextureUAVDesc(GraphSetup.CopyTexture1)) : nullptr;
+			PassParameters->OutCopyTexture2 = GraphSetup.CopyTexture2 ? GraphBuilder.CreateUAV(FRDGTextureUAVDesc(GraphSetup.CopyTexture2)) : nullptr;
 
 			AddCopyPass(GraphBuilder, View->GetFeatureLevel(), PassParameters, TextureSize, MaterialType);
 		}
@@ -813,6 +857,14 @@ namespace RuntimeVirtualTexture
 			GraphOutputSize1 = GraphSetup.OutputAlias1->Desc.GetSize();
 		}
 
+		TRefCountPtr<IPooledRenderTarget> GraphOutputTexture2;
+		FIntVector GraphOutputSize2;
+		if (GraphSetup.OutputAlias2 != nullptr)
+		{
+			GraphBuilder.QueueTextureExtraction(GraphSetup.OutputAlias2, &GraphOutputTexture2);
+			GraphOutputSize2 = GraphSetup.OutputAlias2->Desc.GetSize();
+		}
+
 		// Execute the graph
 		GraphBuilder.Execute();
 
@@ -834,6 +886,15 @@ namespace RuntimeVirtualTexture
 
 			RHICmdList.CopyTexture(GraphOutputTexture1->GetRenderTargetItem().ShaderResourceTexture->GetTexture2D(), OutputTexture1->GetTexture2D(), Info);
 		}
+
+		if (GraphSetup.OutputAlias2 != nullptr && OutputTexture2 != nullptr)
+		{
+			FRHICopyTextureInfo Info;
+			Info.Size = GraphOutputSize2;
+			Info.DestPosition = FIntVector(DestBox2.Min.X, DestBox2.Min.Y, 0);
+
+			RHICmdList.CopyTexture(GraphOutputTexture2->GetRenderTargetItem().ShaderResourceTexture->GetTexture2D(), OutputTexture2->GetTexture2D(), Info);
+		}
 	}
 
 
diff --git a/Engine/Source/Runtime/Renderer/Public/VT/RuntimeVirtualTextureRender.h b/Engine/Source/Runtime/Renderer/Public/VT/RuntimeVirtualTextureRender.h
--- a/Engine/Source/Runtime/Renderer/Public/VT/RuntimeVirtualTextureRender.h
+++ b/Engine/Source/Runtime/Renderer/Public/VT/RuntimeVirtualTextureRender.h
@@ -37,6 +37,7 @@ namespace RuntimeVirtualTexture
 		bool bClearTextures,
 		FRHITexture2D* Texture0, FBox2D const& DestBox0,
 		FRHITexture2D* Texture1, FBox2D const& DestBox1,
+		FRHITexture2D* Texture2, FBox2D const& DestBox2,
 		FTransform const& UVToWorld,
 		FBox2D const& UVRange,
 		uint8 vLevel,
