Added "r.ShaderComplexity.CacheShaders" to allow caching of the shader complexity viewmode shaders to avoid latency caused by the on-the-fly compilation.

diff --git a/Engine/Source/Runtime/Renderer/Private/DebugViewModeRendering.cpp b/Engine/Source/Runtime/Renderer/Private/DebugViewModeRendering.cpp
--- a/Engine/Source/Runtime/Renderer/Private/DebugViewModeRendering.cpp
+++ b/Engine/Source/Runtime/Renderer/Private/DebugViewModeRendering.cpp
@@ -29,7 +29,6 @@ IMPLEMENT_GLOBAL_SHADER_PARAMETER_STRUCT(FDebugViewModePassPassUniformParameters
 
 #if !(UE_BUILD_SHIPPING || UE_BUILD_TEST)
 
-
 void SetupDebugViewModePassUniformBuffer(FSceneRenderTargets& SceneContext, ERHIFeatureLevel::Type FeatureLevel, FDebugViewModePassPassUniformParameters& PassParameters)
 {
 	SetupSceneTextureUniformParameters(SceneContext, FeatureLevel, ESceneTextureSetupMode::None, PassParameters.SceneTextures);
@@ -53,6 +52,25 @@ IMPLEMENT_MATERIAL_SHADER_TYPE(,FDebugViewModeDS,TEXT("/Engine/Private/DebugView
 
 ENGINE_API bool GetDebugViewMaterial(const UMaterialInterface* InMaterialInterface, EDebugViewShaderMode InDebugViewMode, ERHIFeatureLevel::Type InFeatureLevel,const FMaterialRenderProxy*& OutMaterialRenderProxy, const FMaterial*& OutMaterial);
 
+
+bool FDebugViewModeVS::ShouldCompilePermutation(const FMeshMaterialShaderPermutationParameters& Parameters)
+{
+	if (AllowDebugViewVSDSHS(Parameters.Platform))
+	{
+		// If it comes from FDebugViewModeMaterialProxy, compile it.
+		if (Parameters.Material->GetFriendlyName().Contains(TEXT("DebugViewMode")))
+		{
+			return true;
+		}
+		// Otherwise we only cache it if this for the shader complexity.
+		else if (GCacheShaderComplexityShaders)
+		{
+			return !FDebugViewModeInterface::AllowFallbackToDefaultMaterial(Parameters.Material) || Parameters.Material->IsDefaultMaterial();
+		}
+	}
+	return false;
+}
+
 extern FRenderingCompositeOutputRef AddTemporalAADebugViewPass(FPostprocessContext& Context);
 
 void FDeferredShadingSceneRenderer::DoDebugViewModePostProcessing(FRHICommandListImmediate& RHICmdList, const FViewInfo& View, TRefCountPtr<IPooledRenderTarget>& VelocityRT)
@@ -306,8 +324,6 @@ FDebugViewModeMeshProcessor::FDebugViewModeMeshProcessor(
 
 void FDebugViewModeMeshProcessor::AddMeshBatch(const FMeshBatch& RESTRICT MeshBatch, uint64 BatchElementMask, const FPrimitiveSceneProxy* RESTRICT PrimitiveSceneProxy, int32 StaticMeshId)
 {
-	const FMaterialRenderProxy* MaterialRenderProxy = nullptr;
-	const FMaterial* Material = nullptr;
 	const FMaterial* BatchMaterial = MeshBatch.MaterialRenderProxy->GetMaterialNoFallback(FeatureLevel);
 
 	if (!DebugViewModeInterface || !BatchMaterial)
@@ -320,7 +336,21 @@ void FDebugViewModeMeshProcessor::AddMeshBatch(const FMeshBatch& RESTRICT MeshBa
 	{
 		ResolvedMaterial = UMaterial::GetDefaultMaterial(MD_Surface);
 	}
-	if (!GetDebugViewMaterial(ResolvedMaterial, DebugViewMode, FeatureLevel, MaterialRenderProxy, Material))
+
+	const FMaterialRenderProxy* MaterialRenderProxy = nullptr;
+	const FMaterial* Material = nullptr;
+
+	if (DebugViewMode == DVSM_ShaderComplexity && GCacheShaderComplexityShaders)
+	{
+		Material = ResolvedMaterial->GetMaterialResource(FeatureLevel);
+		MaterialRenderProxy  = ResolvedMaterial->GetRenderProxy();
+
+		if (!Material || !MaterialRenderProxy || !Material->HasValidGameThreadShaderMap() ||  !Material->GetRenderingThreadShaderMap())
+		{
+			return;
+		}
+	}
+	else if (!GetDebugViewMaterial(ResolvedMaterial, DebugViewMode, FeatureLevel, MaterialRenderProxy, Material))
 	{
 		return;
 	}
diff --git a/Engine/Source/Runtime/Renderer/Private/DebugViewModeRendering.h b/Engine/Source/Runtime/Renderer/Private/DebugViewModeRendering.h
--- a/Engine/Source/Runtime/Renderer/Private/DebugViewModeRendering.h
+++ b/Engine/Source/Runtime/Renderer/Private/DebugViewModeRendering.h
@@ -86,10 +86,7 @@ protected:
 
 public:
 
-	static bool ShouldCompilePermutation(const FMeshMaterialShaderPermutationParameters& Parameters)
-	{
-		return AllowDebugViewVSDSHS(Parameters.Platform) && Parameters.Material->GetFriendlyName().Contains(TEXT("DebugViewMode"));
-	}
+	static bool ShouldCompilePermutation(const FMeshMaterialShaderPermutationParameters& Parameters);
 
 	void GetShaderBindings(
 		const FScene* Scene,
diff --git a/Engine/Source/Runtime/Renderer/Private/ShaderComplexityRendering.cpp b/Engine/Source/Runtime/Renderer/Private/ShaderComplexityRendering.cpp
--- a/Engine/Source/Runtime/Renderer/Private/ShaderComplexityRendering.cpp
+++ b/Engine/Source/Runtime/Renderer/Private/ShaderComplexityRendering.cpp
@@ -10,6 +10,14 @@ ShaderComplexityRendering.cpp: Contains definitions for rendering the shader com
 
 #if !(UE_BUILD_SHIPPING || UE_BUILD_TEST)
 
+extern int32 GCacheShaderComplexityShaders = 0;
+static FAutoConsoleVariableRef CVarNiagaraAllowTickBeforeRender(
+	TEXT("r.ShaderComplexity.CacheShaders"),
+	GCacheShaderComplexityShaders,
+	TEXT("If non zero, store the shader complexity shaders in the material shader map, to prevent compile on-the-fly lag. (default=0)"),
+	ECVF_ReadOnly
+);
+
 int32 GShaderComplexityBaselineForwardVS = 134;
 static FAutoConsoleVariableRef CVarShaderComplexityBaselineForwardVS(
 	TEXT("r.ShaderComplexity.Baseline.Forward.VS"),
diff --git a/Engine/Source/Runtime/Renderer/Private/ShaderComplexityRendering.h b/Engine/Source/Runtime/Renderer/Private/ShaderComplexityRendering.h
--- a/Engine/Source/Runtime/Renderer/Private/ShaderComplexityRendering.h
+++ b/Engine/Source/Runtime/Renderer/Private/ShaderComplexityRendering.h
@@ -20,6 +20,8 @@ struct FMeshDrawingRenderState;
 
 #if !(UE_BUILD_SHIPPING || UE_BUILD_TEST)
 
+extern int32 GCacheShaderComplexityShaders;
+
 extern int32 GShaderComplexityBaselineForwardVS;
 extern int32 GShaderComplexityBaselineForwardPS;
 extern int32 GShaderComplexityBaselineForwardUnlitPS;
@@ -27,6 +29,7 @@ extern int32 GShaderComplexityBaselineDeferredVS;
 extern int32 GShaderComplexityBaselineDeferredPS;
 extern int32 GShaderComplexityBaselineDeferredUnlitPS;
 
+
 template <bool bQuadComplexity>
 class TComplexityAccumulatePS : public FDebugViewModePS
 {
@@ -36,7 +39,20 @@ public:
 	static bool ShouldCompilePermutation(const FMeshMaterialShaderPermutationParameters& Parameters)
 	{
 		// See FDebugViewModeMaterialProxy::GetFriendlyName()
-		return AllowDebugViewShaderMode(bQuadComplexity ? DVSM_QuadComplexity : DVSM_ShaderComplexity, Parameters.Platform, Parameters.Material->GetFeatureLevel()) && Parameters.Material->GetFriendlyName().Contains(TEXT("ComplexityAccumulate"));
+		if (AllowDebugViewShaderMode(bQuadComplexity ? DVSM_QuadComplexity : DVSM_ShaderComplexity, Parameters.Platform, Parameters.Material->GetFeatureLevel()))
+		{
+			// If it comes from FDebugViewModeMaterialProxy, compile it.
+			if (Parameters.Material->GetFriendlyName().Contains(TEXT("ComplexityAccumulate")))
+			{
+				return true;
+			}
+			// Otherwise we only cache it if this for the shader complexity.
+			else if (GCacheShaderComplexityShaders)
+			{
+				return !FDebugViewModeInterface::AllowFallbackToDefaultMaterial(Parameters.Material) || Parameters.Material->IsDefaultMaterial();
+			}
+		}
+		return false;
 	}
 
 	TComplexityAccumulatePS(const ShaderMetaType::CompiledShaderInitializerType& Initializer):
