Fixed compilation error around GCacheShaderComplexityShaders

diff --git a/Engine/Source/Runtime/Renderer/Private/ShaderComplexityRendering.cpp b/Engine/Source/Runtime/Renderer/Private/ShaderComplexityRendering.cpp
--- a/Engine/Source/Runtime/Renderer/Private/ShaderComplexityRendering.cpp
+++ b/Engine/Source/Runtime/Renderer/Private/ShaderComplexityRendering.cpp
@@ -10,7 +10,7 @@ ShaderComplexityRendering.cpp: Contains definitions for rendering the shader com
 
 #if !(UE_BUILD_SHIPPING || UE_BUILD_TEST)
 
-extern int32 GCacheShaderComplexityShaders = 0;
+int32 GCacheShaderComplexityShaders = 0;
 static FAutoConsoleVariableRef CVarNiagaraAllowTickBeforeRender(
 	TEXT("r.ShaderComplexity.CacheShaders"),
 	GCacheShaderComplexityShaders,
