Single Layer Water optimization. Added refraction downsample factor. By default set to 2, which is half resolution.

diff --git a/Engine/Shaders/Private/BasePassPixelShader.usf b/Engine/Shaders/Private/BasePassPixelShader.usf
--- a/Engine/Shaders/Private/BasePassPixelShader.usf
+++ b/Engine/Shaders/Private/BasePassPixelShader.usf
@@ -1228,6 +1228,7 @@ ShadingModel);
 			OpaqueBasePass.SceneDepthWithoutSingleLayerWaterTexture, OpaqueBasePass.SceneDepthWithoutSingleLayerWaterSampler,
 			OpaqueBasePass.SceneLuminanceWithoutSingleLayerWaterTexture, OpaqueBasePass.SceneLuminanceWithoutSingleLayerWaterSampler,
 			OpaqueBasePass.SceneCustomDepthTexture, OpaqueBasePass.SceneCustomDepthSampler,
+			OpaqueBasePass.SceneWithoutSingleLayerWaterValidUVRect,
 			Specular, OpaqueBasePass.DistortionParams,
 			SunIlluminance, WaterDiffuseIndirectIlluminance, EnvBrdf,
 			CameraIsUnderWater, WaterVisibility).Luminance;
diff --git a/Engine/Shaders/Private/SingleLayerWaterCommon.ush b/Engine/Shaders/Private/SingleLayerWaterCommon.ush
--- a/Engine/Shaders/Private/SingleLayerWaterCommon.ush
+++ b/Engine/Shaders/Private/SingleLayerWaterCommon.ush
@@ -46,6 +46,7 @@ WaterVolumeLightingOutput EvaluateWaterVolumeLighting(
 	Texture2D SceneDepthWithoutSingleLayerWaterTexture, SamplerState SceneDepthWithoutSingleLayerWaterSampler,
 	Texture2D SceneLuminanceWithoutSingleLayerWaterTexture, SamplerState SceneLuminanceWithoutSingleLayerWaterSampler,
 	Texture2D SceneCustomDepthTexture, SamplerState SceneCustomDepthSampler,
+	float4 SceneWithoutSingleLayerWaterValidUVRect,
 #endif // SIMPLE_SINGLE_LAYER_WATER
 	float Specular, const float4 DistortionParams,
 	float3 SunIlluminance, float3 AmbiantIlluminance, float3 EnvBrdf,
@@ -91,9 +92,9 @@ WaterVolumeLightingOutput EvaluateWaterVolumeLighting(
 		BufferUVDistortion *= saturate((PixelSceneDepth - WaterDepth) * 1.0 / 30.0f);// *(1 - saturate(WaterDepth * 1.0 / 2000.0f));
 
 		// Clamp distorted UVs to not sample black outside the sub part of the buffer. Big quality increase.
-		const float2 DistortedUv = clamp(ViewportUV + BufferUVDistortion, 0.0, (ResolvedView.ViewSizeAndInvSize.xy-0.5f) * ResolvedView.BufferSizeAndInvSize.zw);
+		const float2 DistortedUV = clamp(ViewportUV + BufferUVDistortion, SceneWithoutSingleLayerWaterValidUVRect.xy, SceneWithoutSingleLayerWaterValidUVRect.zw);
 
-		float SceneDeviceZ = SceneDepthWithoutSingleLayerWaterTexture.SampleLevel(SceneDepthWithoutSingleLayerWaterSampler, DistortedUv, 0).r;
+		float SceneDeviceZ = SceneDepthWithoutSingleLayerWaterTexture.SampleLevel(SceneDepthWithoutSingleLayerWaterSampler, DistortedUV, 0).r;
 		float SceneDepth = ConvertFromDeviceZ(SceneDeviceZ);
 		if (SceneDepth < WaterDepth)
 		{
@@ -101,7 +102,7 @@ WaterVolumeLightingOutput EvaluateWaterVolumeLighting(
 		}
 		else
 		{
-			ViewportUV = DistortedUv;
+			ViewportUV = DistortedUV;
 		}
 
 #endif // !SIMPLE_SINGLE_LAYER_WATER
diff --git a/Engine/Shaders/Private/SingleLayerWaterComposite.usf b/Engine/Shaders/Private/SingleLayerWaterComposite.usf
--- a/Engine/Shaders/Private/SingleLayerWaterComposite.usf
+++ b/Engine/Shaders/Private/SingleLayerWaterComposite.usf
@@ -297,4 +297,20 @@ void SingleLayerWaterCompositePS(
 
 #endif
 
+Texture2D<float4> SceneColorCopyDownsampleTexture;
+SamplerState      SceneColorCopyDownsampleSampler;
+Texture2D<float4> SceneDepthCopyDownsampleTexture;
+SamplerState      SceneDepthCopyDownsampleSampler;
+float2 SVPositionToSourceTextureUV;
+
+void WaterRefractionCopyPS(
+	in float4 SVPosition : SV_Position,
+	out float4 OutColor : SV_Target0,
+	out float OutDepth : SV_DEPTH)
+{
+	float2 BufferUV = SVPosition.xy * SVPositionToSourceTextureUV;
 
+	// Downsample and average all 2x2 quads
+	OutColor = Texture2DSampleLevel(SceneColorCopyDownsampleTexture, SceneColorCopyDownsampleSampler, BufferUV, 0);
+	OutDepth = Texture2DSampleLevel(SceneDepthCopyDownsampleTexture, SceneDepthCopyDownsampleSampler, BufferUV, 0).r;
+}
\ No newline at end of file
diff --git a/Engine/Shaders/Public/ShaderVersion.ush b/Engine/Shaders/Public/ShaderVersion.ush
--- a/Engine/Shaders/Public/ShaderVersion.ush
+++ b/Engine/Shaders/Public/ShaderVersion.ush
@@ -3,6 +3,6 @@
 // in Platform.ush (which should be included in any shader) it allows to invalidate the shader DDC.
 // 
 // If you are merging streams and there is a conflict with this GUID you should make a new GUID rather than taking one or the other.
-// GUID = <Guid("BC534643-94F5-441C-93FE-9D8A0073798B")>
+// GUID = 919A4804-C336-4DCC-BC7B-22C1054255EF
 
 
diff --git a/Engine/Source/Runtime/Renderer/Private/BasePassRendering.cpp b/Engine/Source/Runtime/Renderer/Private/BasePassRendering.cpp
--- a/Engine/Source/Runtime/Renderer/Private/BasePassRendering.cpp
+++ b/Engine/Source/Runtime/Renderer/Private/BasePassRendering.cpp
@@ -516,6 +516,7 @@ void CreateOpaqueBasePassUniformBuffer(
 	FRHICommandListImmediate& RHICmdList, 
 	const FViewInfo& View,
 	IPooledRenderTarget* ForwardScreenSpaceShadowMask,
+	FVector4* SceneWithoutSingleLayerWaterValidUVRect,
 	IPooledRenderTarget* SceneLuminanceWithoutSingleLayerWater,
 	IPooledRenderTarget* SceneDepthWithoutSingleLayerWaterTexture,
 	TUniformBufferRef<FOpaqueBasePassUniformParameters>& BasePassUniformBuffer)
@@ -582,6 +583,7 @@ void CreateOpaqueBasePassUniformBuffer(
 	}
 
 	// Single Layer Water
+	BasePassParameters.SceneWithoutSingleLayerWaterValidUVRect = SceneWithoutSingleLayerWaterValidUVRect ? *SceneWithoutSingleLayerWaterValidUVRect : FVector4(0.0f, 0.0f, 1.0f, 1.0f);
 	BasePassParameters.SceneLuminanceWithoutSingleLayerWaterTexture = (SceneLuminanceWithoutSingleLayerWater ? SceneLuminanceWithoutSingleLayerWater->GetRenderTargetItem().TargetableTexture : GSystemTextures.BlackDummy->GetRenderTargetItem().TargetableTexture);
 	BasePassParameters.SceneLuminanceWithoutSingleLayerWaterSampler = TStaticSamplerState<SF_Bilinear>::GetRHI();
 	BasePassParameters.SceneDepthWithoutSingleLayerWaterTexture = (SceneDepthWithoutSingleLayerWaterTexture ? SceneDepthWithoutSingleLayerWaterTexture->GetRenderTargetItem().TargetableTexture : GSystemTextures.BlackDummy->GetRenderTargetItem().TargetableTexture);
@@ -650,7 +652,7 @@ bool FDeferredShadingSceneRenderer::RenderBasePass(FRHICommandListImmediate& RHI
 				SCOPED_GPU_MASK(RHICmdList, View.GPUMask);
 
 				TUniformBufferRef<FOpaqueBasePassUniformParameters> BasePassUniformBuffer;
-				CreateOpaqueBasePassUniformBuffer(RHICmdList, View, ForwardScreenSpaceShadowMask, nullptr, nullptr, BasePassUniformBuffer);
+				CreateOpaqueBasePassUniformBuffer(RHICmdList, View, ForwardScreenSpaceShadowMask, nullptr, nullptr, nullptr, BasePassUniformBuffer);
 
 				FMeshPassProcessorRenderState DrawRenderState(View, BasePassUniformBuffer);
 
@@ -691,7 +693,7 @@ bool FDeferredShadingSceneRenderer::RenderBasePass(FRHICommandListImmediate& RHI
 				SCOPED_GPU_MASK(RHICmdList, View.GPUMask);
 
 				TUniformBufferRef<FOpaqueBasePassUniformParameters> BasePassUniformBuffer;
-				CreateOpaqueBasePassUniformBuffer(RHICmdList, View, ForwardScreenSpaceShadowMask, nullptr, nullptr, BasePassUniformBuffer);
+				CreateOpaqueBasePassUniformBuffer(RHICmdList, View, ForwardScreenSpaceShadowMask, nullptr, nullptr, nullptr, BasePassUniformBuffer);
 
 				FMeshPassProcessorRenderState DrawRenderState(View, BasePassUniformBuffer);
 
diff --git a/Engine/Source/Runtime/Renderer/Private/BasePassRendering.h b/Engine/Source/Runtime/Renderer/Private/BasePassRendering.h
--- a/Engine/Source/Runtime/Renderer/Private/BasePassRendering.h
+++ b/Engine/Source/Runtime/Renderer/Private/BasePassRendering.h
@@ -80,6 +80,7 @@ BEGIN_GLOBAL_SHADER_PARAMETER_STRUCT(FOpaqueBasePassUniformParameters,)
 	SHADER_PARAMETER_SAMPLER(SamplerState, PreIntegratedGFSampler)
 	SHADER_PARAMETER_TEXTURE(Texture2D, SceneCustomDepthTexture)
 	SHADER_PARAMETER_SAMPLER(SamplerState, SceneCustomDepthSampler)
+	SHADER_PARAMETER(FVector4, SceneWithoutSingleLayerWaterValidUVRect)
 	SHADER_PARAMETER(FVector4, DistortionParams)
 	// Misc
 	SHADER_PARAMETER_TEXTURE(Texture2D, EyeAdaptation)
@@ -119,6 +120,7 @@ extern void CreateOpaqueBasePassUniformBuffer(
 	FRHICommandListImmediate& RHICmdList, 
 	const FViewInfo& View,
 	IPooledRenderTarget* ForwardScreenSpaceShadowMask,
+	FVector4* SceneWithoutSingleLayerWaterValidUVRect,
 	IPooledRenderTarget* SceneLuminanceWithoutSingleLayerWater,
 	IPooledRenderTarget* SceneDepthWithoutSingleLayerWaterTexture,
 	TUniformBufferRef<FOpaqueBasePassUniformParameters>& BasePassUniformBuffer);
diff --git a/Engine/Source/Runtime/Renderer/Private/DeferredShadingRenderer.h b/Engine/Source/Runtime/Renderer/Private/DeferredShadingRenderer.h
--- a/Engine/Source/Runtime/Renderer/Private/DeferredShadingRenderer.h
+++ b/Engine/Source/Runtime/Renderer/Private/DeferredShadingRenderer.h
@@ -104,7 +104,7 @@ public:
 	bool RenderBasePass(FRHICommandListImmediate& RHICmdList, FExclusiveDepthStencil::Type BasePassDepthStencilAccess, IPooledRenderTarget* ForwardScreenSpaceShadowMask, bool bParallelBasePass, bool bRenderLightmapDensity);
 
 	/** Copy the scene color and depth buffer for water refraction*/
-	void CopySingleLayerWaterTextures(FRHICommandList& RHICmdList, FSingleLayerWaterPassData& PassData);
+	void CopySingleLayerWaterTextures(FRHICommandListImmediate& RHICmdList, FSingleLayerWaterPassData& PassData);
 	/** Begin the water GBuffer rendering pass*/
 	void BeginRenderingWaterGBuffer(FRHICommandList& RHICmdList, FSingleLayerWaterPassData& PassData, FExclusiveDepthStencil::Type DepthStencilAccess, bool bBindQuadOverdrawBuffers);
 	/** End the water GBuffer rendering pass*/
diff --git a/Engine/Source/Runtime/Renderer/Private/PostProcess/PostProcessCompositeEditorPrimitives.cpp b/Engine/Source/Runtime/Renderer/Private/PostProcess/PostProcessCompositeEditorPrimitives.cpp
--- a/Engine/Source/Runtime/Renderer/Private/PostProcess/PostProcessCompositeEditorPrimitives.cpp
+++ b/Engine/Source/Runtime/Renderer/Private/PostProcess/PostProcessCompositeEditorPrimitives.cpp
@@ -565,7 +565,7 @@ void FRCPassPostProcessCompositeEditorPrimitives::Process(FRenderingCompositePas
 
 			if (bDeferredBasePass)
 			{
-				CreateOpaqueBasePassUniformBuffer(Context.RHICmdList, EditorView, nullptr, nullptr, nullptr, OpaqueBasePassUniformBuffer);
+				CreateOpaqueBasePassUniformBuffer(Context.RHICmdList, EditorView, nullptr, nullptr, nullptr, nullptr, OpaqueBasePassUniformBuffer);
 				BasePassUniformBuffer = OpaqueBasePassUniformBuffer;
 			}
 			else
diff --git a/Engine/Source/Runtime/Renderer/Private/Renderer.cpp b/Engine/Source/Runtime/Renderer/Private/Renderer.cpp
--- a/Engine/Source/Runtime/Renderer/Private/Renderer.cpp
+++ b/Engine/Source/Runtime/Renderer/Private/Renderer.cpp
@@ -296,7 +296,7 @@ void FRendererModule::DrawTileMesh(FRHICommandListImmediate& RHICmdList, FMeshPa
 				if (ShadingPath == EShadingPath::Deferred)
 				{
 					TUniformBufferRef<FOpaqueBasePassUniformParameters> OpaqueBasePassUniformBuffer;
-					CreateOpaqueBasePassUniformBuffer(RHICmdList, View, nullptr, nullptr, nullptr, OpaqueBasePassUniformBuffer);
+					CreateOpaqueBasePassUniformBuffer(RHICmdList, View, nullptr, nullptr, nullptr, nullptr, OpaqueBasePassUniformBuffer);
 					BasePassUniformBuffer = OpaqueBasePassUniformBuffer;
 					
 					DrawRenderState.SetPassUniformBuffer(BasePassUniformBuffer);
diff --git a/Engine/Source/Runtime/Renderer/Private/SingleLayerWaterRendering.cpp b/Engine/Source/Runtime/Renderer/Private/SingleLayerWaterRendering.cpp
--- a/Engine/Source/Runtime/Renderer/Private/SingleLayerWaterRendering.cpp
+++ b/Engine/Source/Runtime/Renderer/Private/SingleLayerWaterRendering.cpp
@@ -41,6 +41,13 @@ static TAutoConsoleVariable<int32> CVarWaterSingleLayerTiledComposite(
 	TEXT("Enable tiled optimisation of the water reflection rendering."),
 	ECVF_RenderThreadSafe | ECVF_Scalability);
 
+int32 GSingleLayerWaterRefractionDownsampleFactor = 2;
+FAutoConsoleVariableRef CVarWaterSingleLayerRefractionDownsampleFactor(
+	TEXT("r.Water.SingleLayer.RefractionDownsampleFactor"),
+	GSingleLayerWaterRefractionDownsampleFactor,
+	TEXT("Resolution divider for the water refraction buffer."),
+	ECVF_Scalability | ECVF_RenderThreadSafe
+);
 
 static TAutoConsoleVariable<int32> CVarWaterSingleLayerSSR(
 	TEXT("r.Water.SingleLayer.SSR"), 1,
@@ -338,37 +345,98 @@ IMPLEMENT_GLOBAL_SHADER(FWaterTileCategorisationCS, "/Engine/Private/SingleLayer
 
 IMPLEMENT_GLOBAL_SHADER(FWaterTiledCompositeCS, "/Engine/Private/SingleLayerWaterComposite.usf", "WaterTiledCategorisationCS", SF_Compute);*/
 
+class FWaterRefractionCopyPS : public FGlobalShader
+{
+	DECLARE_GLOBAL_SHADER(FWaterRefractionCopyPS);
+	SHADER_USE_PARAMETER_STRUCT(FWaterRefractionCopyPS, FGlobalShader);
+
+	BEGIN_SHADER_PARAMETER_STRUCT(FParameters, )
+		SHADER_PARAMETER_STRUCT_REF(FViewUniformShaderParameters, ViewUniformBuffer)
+		SHADER_PARAMETER_RDG_TEXTURE(Texture2D, SceneColorCopyDownsampleTexture)
+		SHADER_PARAMETER_SAMPLER(SamplerState, SceneColorCopyDownsampleSampler)
+		SHADER_PARAMETER_RDG_TEXTURE(Texture2D, SceneDepthCopyDownsampleTexture)
+		SHADER_PARAMETER_SAMPLER(SamplerState, SceneDepthCopyDownsampleSampler)
+		SHADER_PARAMETER(FVector2D, SVPositionToSourceTextureUV)
+		RENDER_TARGET_BINDING_SLOTS()
+	END_SHADER_PARAMETER_STRUCT()
+
+	static bool ShouldCompilePermutation(const FGlobalShaderPermutationParameters& Parameters)
+	{
+		return IsFeatureLevelSupported(Parameters.Platform, ERHIFeatureLevel::SM5);
+	}
+
+	static void ModifyCompilationEnvironment(const FGlobalShaderPermutationParameters& Parameters, FShaderCompilerEnvironment& OutEnvironment)
+	{
+		FGlobalShader::ModifyCompilationEnvironment(Parameters, OutEnvironment);
+	}
+};
+
+IMPLEMENT_GLOBAL_SHADER(FWaterRefractionCopyPS, "/Engine/Private/SingleLayerWaterComposite.usf", "WaterRefractionCopyPS", SF_Pixel);
+
 
 
 //////////////////////////////////////////////////////////////////////////
 
 
 
-void FDeferredShadingSceneRenderer::CopySingleLayerWaterTextures(FRHICommandList& RHICmdList, FSingleLayerWaterPassData& PassData)
+void FDeferredShadingSceneRenderer::CopySingleLayerWaterTextures(FRHICommandListImmediate& RHICmdList, FSingleLayerWaterPassData& PassData)
 {
 	check(RHICmdList.IsOutsideRenderPass());
 	FSceneRenderTargets& SceneContext = FSceneRenderTargets::Get(RHICmdList);
 	const ERHIFeatureLevel::Type CurrentFeatureLevel = SceneContext.GetCurrentFeatureLevel();
+	const int32 RefractionDownsampleFactor = FMath::Clamp(GSingleLayerWaterRefractionDownsampleFactor, 1, 8);
+	const FIntPoint RefractionResolution = FIntPoint::DivideAndRoundDown(SceneContext.GetBufferSizeXY(), RefractionDownsampleFactor);
 
 	// Allocate required buffers
 	{
-		FPooledRenderTargetDesc Desc(FPooledRenderTargetDesc::Create2DDesc(SceneContext.GetBufferSizeXY(), PF_DepthStencil, SceneContext.GetDefaultDepthClear(), TexCreate_None, TexCreate_DepthStencilTargetable | TexCreate_ShaderResource | TexCreate_InputAttachmentRead, false));
+		FPooledRenderTargetDesc Desc(FPooledRenderTargetDesc::Create2DDesc(RefractionResolution, PF_DepthStencil, SceneContext.GetDefaultDepthClear(), TexCreate_None, TexCreate_DepthStencilTargetable | TexCreate_ShaderResource | TexCreate_InputAttachmentRead, false));
 		Desc.NumSamples = SceneContext.GetNumSceneColorMSAASamples(CurrentFeatureLevel);
 		Desc.Flags |= GFastVRamConfig.SceneDepth;
 		GRenderTargetPool.FindFreeElement(RHICmdList, Desc, PassData.SceneDepthZWithoutSingleLayerWater, TEXT("SceneDepthZWithoutSingleLayerWater"), true);
 	}
 	{
-		FPooledRenderTargetDesc Desc(FPooledRenderTargetDesc::Create2DDesc(SceneContext.GetBufferSizeXY(), SceneContext.GetSceneColorFormat(), SceneContext.GetDefaultColorClear(), TexCreate_None, TexCreate_RenderTargetable, false));
+		FPooledRenderTargetDesc Desc(FPooledRenderTargetDesc::Create2DDesc(RefractionResolution, SceneContext.GetSceneColorFormat(), SceneContext.GetDefaultColorClear(), TexCreate_None, TexCreate_RenderTargetable, false));
 		Desc.NumSamples = SceneContext.GetNumSceneColorMSAASamples(CurrentFeatureLevel);
 		GRenderTargetPool.FindFreeElement(RHICmdList, Desc, PassData.SceneColorWithoutSingleLayerWater, TEXT("SceneColorWithoutSingleLayerWater"));
 	}
 
-	// Copy and save textures without water informations for later composition operations
-	{
-		FRHICopyTextureInfo DepthCopyInfo;
-		RHICmdList.CopyTexture(SceneContext.SceneDepthZ->GetRenderTargetItem().ShaderResourceTexture, PassData.SceneDepthZWithoutSingleLayerWater->GetRenderTargetItem().ShaderResourceTexture, DepthCopyInfo);
-		RHICmdList.CopyTexture(SceneContext.GetSceneColorSurface(), PassData.SceneColorWithoutSingleLayerWater->GetRenderTargetItem().ShaderResourceTexture, DepthCopyInfo);
-	}
+	// For now support only the 1st view
+	const FViewInfo& View = Views[0];
+
+	FRDGBuilder GraphBuilder(RHICmdList);
+	FRDGTextureRef TargetColorTexture = GraphBuilder.RegisterExternalTexture(PassData.SceneColorWithoutSingleLayerWater);
+	FRDGTextureRef TargetDepthTexture = GraphBuilder.RegisterExternalTexture(PassData.SceneDepthZWithoutSingleLayerWater);
+
+	FWaterRefractionCopyPS::FParameters* PassParameters = GraphBuilder.AllocParameters<FWaterRefractionCopyPS::FParameters>();
+	PassParameters->ViewUniformBuffer = View.ViewUniformBuffer;
+	PassParameters->SceneColorCopyDownsampleTexture = GraphBuilder.RegisterExternalTexture(SceneContext.GetSceneColor());
+	PassParameters->SceneColorCopyDownsampleSampler = TStaticSamplerState<SF_Bilinear, AM_Clamp, AM_Clamp, AM_Clamp>::GetRHI();
+	PassParameters->SceneDepthCopyDownsampleTexture = GraphBuilder.RegisterExternalTexture(SceneContext.SceneDepthZ);
+	PassParameters->SceneDepthCopyDownsampleSampler = TStaticSamplerState<SF_Bilinear, AM_Clamp, AM_Clamp, AM_Clamp>::GetRHI();
+	PassParameters->SVPositionToSourceTextureUV = FVector2D(RefractionDownsampleFactor / float(SceneContext.GetBufferSizeXY().X), RefractionDownsampleFactor / float(SceneContext.GetBufferSizeXY().Y));
+	PassParameters->RenderTargets[0] = FRenderTargetBinding(TargetColorTexture, ERenderTargetLoadAction::ELoad, ERenderTargetStoreAction::EStore);
+	PassParameters->RenderTargets.DepthStencil = FDepthStencilBinding(TargetDepthTexture, ERenderTargetLoadAction::ELoad, ERenderTargetStoreAction::EStore, FExclusiveDepthStencil::DepthWrite_StencilNop);
+
+	auto PixelShader = View.ShaderMap->GetShader<FWaterRefractionCopyPS>();
+
+	const FIntRect RefractionViewRect = FIntRect(FIntPoint(0, 0), FIntPoint::DivideAndRoundDown(View.ViewRect.Size(), RefractionDownsampleFactor));
+
+	PassData.SceneWithoutSingleLayerWaterValidUVRect.X = (RefractionViewRect.Min.X + 0.5f) / RefractionResolution.Y;
+	PassData.SceneWithoutSingleLayerWaterValidUVRect.Y = (RefractionViewRect.Min.Y + 0.5f) / RefractionResolution.Y;
+	PassData.SceneWithoutSingleLayerWaterValidUVRect.Z = (RefractionViewRect.Max.X - 0.5f) / RefractionResolution.Y;
+	PassData.SceneWithoutSingleLayerWaterValidUVRect.W = (RefractionViewRect.Max.Y - 0.5f) / RefractionResolution.Y;
+
+	FPixelShaderUtils::AddFullscreenPass(
+		GraphBuilder,
+		View.ShaderMap,
+		RDG_EVENT_NAME("Water Refraction Copy"),
+		PixelShader,
+		PassParameters,
+		RefractionViewRect,
+		nullptr,
+		nullptr,
+		TStaticDepthStencilState<true, CF_Always>::GetRHI());
+	GraphBuilder.Execute();
 }
 
 void FDeferredShadingSceneRenderer::BeginRenderingWaterGBuffer(FRHICommandList& RHICmdList, FSingleLayerWaterPassData& PassData, FExclusiveDepthStencil::Type DepthStencilAccess, bool bBindQuadOverdrawBuffers)
@@ -634,7 +702,13 @@ bool FDeferredShadingSceneRenderer::RenderSingleLayerWaterPass(FRHICommandListIm
 
 			TUniformBufferRef<FOpaqueBasePassUniformParameters> WaterPassUniformBuffer;
 			IPooledRenderTarget* WhiteForwardScreenSpaceShadowMask = GSystemTextures.WhiteDummy;
-			CreateOpaqueBasePassUniformBuffer(RHICmdList, View, WhiteForwardScreenSpaceShadowMask, PassData.SceneColorWithoutSingleLayerWater, PassData.SceneDepthZWithoutSingleLayerWater, WaterPassUniformBuffer);
+			CreateOpaqueBasePassUniformBuffer(RHICmdList, 
+				View, 
+				WhiteForwardScreenSpaceShadowMask, 
+				&PassData.SceneWithoutSingleLayerWaterValidUVRect,
+				PassData.SceneColorWithoutSingleLayerWater,
+				PassData.SceneDepthZWithoutSingleLayerWater,
+				WaterPassUniformBuffer);
 
 			FMeshPassProcessorRenderState DrawRenderState(View, WaterPassUniformBuffer);
 			SetupBasePassState(WaterPassDepthStencilAccess, ViewFamily.EngineShowFlags.ShaderComplexity, DrawRenderState);
diff --git a/Engine/Source/Runtime/Renderer/Private/SingleLayerWaterRendering.h b/Engine/Source/Runtime/Renderer/Private/SingleLayerWaterRendering.h
--- a/Engine/Source/Runtime/Renderer/Private/SingleLayerWaterRendering.h
+++ b/Engine/Source/Runtime/Renderer/Private/SingleLayerWaterRendering.h
@@ -20,6 +20,7 @@ struct FSingleLayerWaterPassData
 {
 	TRefCountPtr<IPooledRenderTarget> SceneColorWithoutSingleLayerWater;
 	TRefCountPtr<IPooledRenderTarget> SceneDepthZWithoutSingleLayerWater;
+	FVector4 SceneWithoutSingleLayerWaterValidUVRect;
 };
 
 
