Swiching DispatchIndirect to DrawIndirect + tile when possible for all platform to use tiling even if triangle and not rect.
That still result in a more adaptative cost. Also PS makes it roughly 10% faster when rect primitive are supported.

diff --git a/Engine/Shaders/Private/SSRT/SSRTReflections.usf b/Engine/Shaders/Private/SSRT/SSRTReflections.usf
--- a/Engine/Shaders/Private/SSRT/SSRTReflections.usf
+++ b/Engine/Shaders/Private/SSRT/SSRTReflections.usf
@@ -375,41 +375,3 @@ void ScreenSpaceReflectionsPS(
 }
 
 
-
-#if defined(TILE_COMPUTE_SSR)
-
-Buffer<uint> TileListDataBuffer;
-
-RWTexture2D<float4> OutColorUAV;
-RWTexture2D<float4> OutClosestHitDistanceUAV;
-
-#define WATER_TILE_SIZE 8
-[numthreads(WATER_TILE_SIZE, WATER_TILE_SIZE, 1)]
-void ScreenSpaceReflectionsCS(/*uint3 ThreadId : SV_DispatchThreadID,*/ uint3 GroupId : SV_GroupID, uint3 GroupThreadId : SV_GroupThreadID)
-{
-	uint TileData = TileListDataBuffer[GroupId.x];
-	uint3 ThreadId = uint3((TileData & 0x0000FFFF)*WATER_TILE_SIZE + GroupThreadId.x, (TileData >> 16)*WATER_TILE_SIZE + GroupThreadId.y, 0);
-
-	if (int(ThreadId.x) < View.ViewSizeAndInvSize.x && int(ThreadId.y) < View.ViewSizeAndInvSize.y)
-	{
-		uint2 PixelPos = ThreadId.xy + View.ViewRectMin.xy;
-		float4 SvPosition = float4(float2(PixelPos) + 0.5f, 0.5, 1.0f);
-
-		float4 OutColorVar; 
-		float4 OutClosestHitDistanceVar;
-		ScreenSpaceReflections(SvPosition, OutColorVar
-#if SSR_OUTPUT_FOR_DENOISER
-			,OutClosestHitDistanceVar 
-#endif
-		);
-
-		OutColorUAV[PixelPos] = OutColorVar;
-#if SSR_OUTPUT_FOR_DENOISER
-		OutClosestHitDistanceUAV[PixelPos] = OutClosestHitDistanceVar;
-#endif
-	}
-}
-
-#endif
-
-
diff --git a/Engine/Shaders/Private/SingleLayerWaterComposite.usf b/Engine/Shaders/Private/SingleLayerWaterComposite.usf
--- a/Engine/Shaders/Private/SingleLayerWaterComposite.usf
+++ b/Engine/Shaders/Private/SingleLayerWaterComposite.usf
@@ -137,35 +137,34 @@ SingleLayerWaterCompositeOutput SingleLayerWaterComposite(float2 BufferUV, float
 #ifdef TILE_CATERGORISATION_SHADER
 
 
-uint TiledViewWidth;
-uint TiledViewHeight;
-float TiledViewWidthInv;
-float TiledViewHeightInv;
+uint VertexCountPerInstanceIndirect;
 
 RWStructuredBuffer<uint> DispatchIndirectDataUAV;
 RWStructuredBuffer<uint> WaterTileListDataUAV;
 
-groupshared bool ContainsWater[WATER_TILE_SIZE*WATER_TILE_SIZE];
+groupshared bool ContainsWater[WORK_TILE_SIZE*WORK_TILE_SIZE];
 
-[numthreads(WATER_TILE_SIZE, WATER_TILE_SIZE, 1)]
+[numthreads(WORK_TILE_SIZE, WORK_TILE_SIZE, 1)]
 void WaterTileCatergorisationCS(uint3 ThreadId : SV_DispatchThreadID, uint3 GroupId : SV_GroupID, uint3 GroupThreadId : SV_GroupThreadID)
 {
 	if (ThreadId.x == 0 && ThreadId.y == 0 && GroupThreadId.x == 0 && GroupThreadId.y == 0)
 	{
 		// TODO compute clear myself
-		DispatchIndirectDataUAV[1] = 1;
-		DispatchIndirectDataUAV[2] = 1;
+		DispatchIndirectDataUAV[0] = VertexCountPerInstanceIndirect;		// VertexCountPerInstance
+		//DispatchIndirectDataUAV[1]		// InstanceCount			already cleared to 0
+		//DispatchIndirectDataUAV[2] = 0;	// StartVertexLocation			"			"
+		//DispatchIndirectDataUAV[3] = 0;	// StartInstanceLocation		"			"
 	}
 
 #if 0
 	// Slow reference path
-	const int LinearIndex = GroupThreadId.y * WATER_TILE_SIZE + GroupThreadId.x;
+	const int LinearIndex = GroupThreadId.y * WORK_TILE_SIZE + GroupThreadId.x;
 	if (LinearIndex < 1)
 	{
 		bool bContainsWater = false;
-		for (uint i = 0; i < WATER_TILE_SIZE; ++i)
+		for (uint i = 0; i < WORK_TILE_SIZE; ++i)
 		{
-			for (uint j = 0; j < WATER_TILE_SIZE; ++j)
+			for (uint j = 0; j < WORK_TILE_SIZE; ++j)
 			{
 				uint2 PixelPos = (ThreadId.xy + uint2(i, j)) + View.ViewRectMin.xy;
 				float2 BufferUV = (PixelPos + 0.5f) * View.BufferSizeAndInvSize.zw;
@@ -193,7 +192,7 @@ void WaterTileCatergorisationCS(uint3 ThreadId : SV_DispatchThreadID, uint3 Grou
 		bContainsWater = GBuffer.ShadingModelID == SHADINGMODELID_SINGLELAYERWATER;
 	}
 
-	const int LinearIndex = GroupThreadId.y * WATER_TILE_SIZE + GroupThreadId.x;
+	const int LinearIndex = GroupThreadId.y * WORK_TILE_SIZE + GroupThreadId.x;
 	ContainsWater[LinearIndex] = bContainsWater;
 	GroupMemoryBarrierWithGroupSync();
 
@@ -229,7 +228,7 @@ void WaterTileCatergorisationCS(uint3 ThreadId : SV_DispatchThreadID, uint3 Grou
 	if (LinearIndex < 1 && (ContainsWater[LinearIndex] || ContainsWater[LinearIndex + 1]))
 	{
 		uint WriteToIndex;  
-		InterlockedAdd(DispatchIndirectDataUAV[0], 1, WriteToIndex);
+		InterlockedAdd(DispatchIndirectDataUAV[1], 1, WriteToIndex);
 		WaterTileListDataUAV[WriteToIndex] = GroupId.x | (GroupId.y << 16); // assumes 16bit is enought to represent a tiled resolution up to 65,535 :)
 	}
 #endif
@@ -237,44 +236,28 @@ void WaterTileCatergorisationCS(uint3 ThreadId : SV_DispatchThreadID, uint3 Grou
 
 
 
-#elif defined(TILED_COMPOSITE_SHADER)
+#elif defined(TILE_VERTEX_SHADER)
 
 
+StructuredBuffer<uint> TileListData; 
 
-uint TiledViewWidth;
-uint TiledViewHeight;
-
-Buffer<uint> WaterTileListData;
-
-RWTexture2D<float3> SceneColorUAV;
-
-[numthreads(WATER_TILE_SIZE, WATER_TILE_SIZE, 1)]
-void WaterTiledCategorisationCS(/*uint3 ThreadId : SV_DispatchThreadID,*/ uint3 GroupId : SV_GroupID, uint3 GroupThreadId : SV_GroupThreadID)
+void WaterTileVS(
+	in uint InstanceId : SV_InstanceID,
+	in uint VertexId : SV_VertexID,
+	out float4 Position : SV_POSITION) 
 {
-	uint TileData = WaterTileListData[GroupId.x];
-	uint3 ThreadId = uint3((TileData & 0x0000FFFF)*WATER_TILE_SIZE + GroupThreadId.x, (TileData >> 16)*WATER_TILE_SIZE + GroupThreadId.y, 0);
-
-	if (int(ThreadId.x) < View.ViewSizeAndInvSize.x && int(ThreadId.y) < View.ViewSizeAndInvSize.y)
-	{
-		uint2 PixelPos = ThreadId.xy + View.ViewRectMin.xy;
-		float2 BufferUV = (PixelPos + 0.5) * View.BufferSizeAndInvSize.zw;
-		float2 ScreenPosition = ViewportUVToScreenPos((ThreadId.xy + 0.5f) * View.ViewSizeAndInvSize.zw);
+	uint TileData = TileListData[InstanceId.x];
+	uint2 TileOrigin = uint2((TileData & 0x0000FFFF) * WORK_TILE_SIZE, (TileData >> 16) * WORK_TILE_SIZE);
 
-		SingleLayerWaterCompositeOutput Result = SingleLayerWaterComposite(BufferUV, ScreenPosition);
+	uint2 TileVertex = TileOrigin + View.ViewRectMin.xy;
+	TileVertex.x += VertexId == 1 || VertexId == 2 || VertexId == 4 ? WORK_TILE_SIZE : 0;
+	TileVertex.y += VertexId == 2 || VertexId == 4 || VertexId == 5 ? WORK_TILE_SIZE : 0;
 
-		if (Result.Clip >= 0)
-		{
-			// TODO colored transmittance when we have a under water medium to apply
-			float3 SceneColor = SceneColorUAV[PixelPos];
-			//SceneColor = Result.LuminanceTransmittance.rgb + Result.LuminanceTransmittance.a * SceneColor  +0.2*float3(sin((TileData & 0x0000FFFF)*WATER_TILE_SIZE), sin((TileData >> 16)*WATER_TILE_SIZE), 0.0f);
-			SceneColor = Result.LuminanceTransmittance.rgb + Result.LuminanceTransmittance.a * SceneColor;
-			SceneColorUAV[PixelPos] = SceneColor;
-		}
-	}
+	// View port is set on the view rect. So no offset are needed.
+	Position = float4(float2(TileVertex) * View.ViewSizeAndInvSize.zw * float2(2.0f, -2.0f) + float2(-1.0, 1.0f), 0.5f, 1.0f);
 }
 
 
-
 #else
 
 
diff --git a/Engine/Source/Runtime/Renderer/Private/ScreenSpaceRayTracing.cpp b/Engine/Source/Runtime/Renderer/Private/ScreenSpaceRayTracing.cpp
--- a/Engine/Source/Runtime/Renderer/Private/ScreenSpaceRayTracing.cpp
+++ b/Engine/Source/Runtime/Renderer/Private/ScreenSpaceRayTracing.cpp
@@ -333,36 +333,38 @@ class FScreenSpaceReflectionsPS : public FGlobalShader
 	BEGIN_SHADER_PARAMETER_STRUCT(FParameters, )
 		SHADER_PARAMETER_STRUCT_INCLUDE(FSSRCommonParameters, CommonParameters)
 		SHADER_PARAMETER_STRUCT_INCLUDE(FSSRPassCommonParameters, SSRPassCommonParameter)
+		SHADER_PARAMETER_RDG_BUFFER(Buffer<uint>, IndirectDrawParameter)			// FScreenSpaceReflectionsTileVS
+		SHADER_PARAMETER_RDG_BUFFER_SRV(StructuredBuffer<uint>, TileListData)		// FScreenSpaceReflectionsTileVS
 		RENDER_TARGET_BINDING_SLOTS()
 	END_SHADER_PARAMETER_STRUCT()
 };
 
-class FScreenSpaceReflectionsCS : public FGlobalShader
+// This is duplicated from FWaterTileVS because vertex shader should share Parameters structure for everything to be registered correctly in a RDG pass.
+class FScreenSpaceReflectionsTileVS : public FGlobalShader
 {
-	DECLARE_GLOBAL_SHADER(FScreenSpaceReflectionsCS);
-	SHADER_USE_PARAMETER_STRUCT(FScreenSpaceReflectionsCS, FGlobalShader);
+	DECLARE_GLOBAL_SHADER(FScreenSpaceReflectionsTileVS);
+	SHADER_USE_PARAMETER_STRUCT(FScreenSpaceReflectionsTileVS, FGlobalShader);
 
-	using FPermutationDomain = FScreenSpaceReflectionsPS::FPermutationDomain;
+	using FPermutationDomain = TShaderPermutationDomain<>;
+
+	using FParameters = FScreenSpaceReflectionsPS::FParameters; // Sharing parameters for proper registration with RDG
+
+	static FPermutationDomain RemapPermutation(FPermutationDomain PermutationVector)
+	{
+		return PermutationVector;
+	}
 
 	static bool ShouldCompilePermutation(const FGlobalShaderPermutationParameters& Parameters)
 	{
-		return FScreenSpaceReflectionsPS::ShouldCompilePermutation(Parameters);
+		return IsFeatureLevelSupported(Parameters.Platform, ERHIFeatureLevel::SM5);
 	}
 
 	static void ModifyCompilationEnvironment(const FGlobalShaderPermutationParameters& Parameters, FShaderCompilerEnvironment& OutEnvironment)
 	{
-		OutEnvironment.SetDefine(TEXT("TILE_COMPUTE_SSR"), 1.0f);
+		OutEnvironment.SetDefine(TEXT("TILE_VERTEX_SHADER"), 1.0f);
+		OutEnvironment.SetDefine(TEXT("WORK_TILE_SIZE"), 8);
 		FGlobalShader::ModifyCompilationEnvironment(Parameters, OutEnvironment);
 	}
-
-	BEGIN_SHADER_PARAMETER_STRUCT(FParameters, )
-		SHADER_PARAMETER_STRUCT_INCLUDE(FSSRCommonParameters, CommonParameters)
-		SHADER_PARAMETER_STRUCT_INCLUDE(FSSRPassCommonParameters, SSRPassCommonParameter)
-		SHADER_PARAMETER_RDG_BUFFER_SRV(Buffer<uint>, TileListDataBuffer)
-		SHADER_PARAMETER_RDG_TEXTURE_UAV(RWTexture2D<float3>, OutColorUAV)
-		SHADER_PARAMETER_RDG_TEXTURE_UAV(RWTexture2D<float3>, OutClosestHitDistanceUAV)
-		SHADER_PARAMETER_RDG_BUFFER(Buffer<uint>, IndirectDispatchParameters)	// Not used in shader but need to be reference in the parameter list
-	END_SHADER_PARAMETER_STRUCT()
 };
 
 class FScreenSpaceDiffuseIndirectCS : public FGlobalShader
@@ -408,7 +410,7 @@ class FScreenSpaceDiffuseIndirectCS : public FGlobalShader
 IMPLEMENT_GLOBAL_SHADER(FSSRTPrevFrameReductionCS, "/Engine/Private/SSRT/SSRTPrevFrameReduction.usf", "MainCS", SF_Compute);
 IMPLEMENT_GLOBAL_SHADER(FSSRTDiffuseTileClassificationCS, "/Engine/Private/SSRT/SSRTTileClassification.usf", "MainCS", SF_Compute);
 IMPLEMENT_GLOBAL_SHADER(FScreenSpaceReflectionsPS,        "/Engine/Private/SSRT/SSRTReflections.usf", "ScreenSpaceReflectionsPS", SF_Pixel);
-IMPLEMENT_GLOBAL_SHADER(FScreenSpaceReflectionsCS,        "/Engine/Private/SSRT/SSRTReflections.usf", "ScreenSpaceReflectionsCS", SF_Compute);
+IMPLEMENT_GLOBAL_SHADER(FScreenSpaceReflectionsTileVS,    "/Engine/Private/SingleLayerWaterComposite.usf", "WaterTileVS", SF_Vertex);
 IMPLEMENT_GLOBAL_SHADER(FScreenSpaceReflectionsStencilPS, "/Engine/Private/SSRT/SSRTReflections.usf", "ScreenSpaceReflectionsStencilPS", SF_Pixel);
 IMPLEMENT_GLOBAL_SHADER(FScreenSpaceDiffuseIndirectCS, "/Engine/Private/SSRT/SSRTDiffuseIndirect.usf", "MainCS", SF_Compute);
 
@@ -684,18 +686,19 @@ void RenderScreenSpaceReflections(
 		PassParameters->HZBSampler = TStaticSamplerState<SF_Point>::GetRHI();
 	};
 
-	if (TiledScreenSpaceReflection == nullptr)
-	{
-		FScreenSpaceReflectionsPS::FPermutationDomain PermutationVector;
-		PermutationVector.Set<FSSRQualityDim>(SSRQuality);
-		PermutationVector.Set<FSSROutputForDenoiser>(bDenoiser);
+	FScreenSpaceReflectionsPS::FPermutationDomain PermutationVector;
+	PermutationVector.Set<FSSRQualityDim>(SSRQuality);
+	PermutationVector.Set<FSSROutputForDenoiser>(bDenoiser);
 
-		FScreenSpaceReflectionsPS::FParameters* PassParameters = GraphBuilder.AllocParameters<FScreenSpaceReflectionsPS::FParameters>();
-		PassParameters->CommonParameters = CommonParameters;
-		SetSSRParameters(&PassParameters->SSRPassCommonParameter);
-		PassParameters->RenderTargets = RenderTargets;
+	FScreenSpaceReflectionsPS::FParameters* PassParameters = GraphBuilder.AllocParameters<FScreenSpaceReflectionsPS::FParameters>();
+	PassParameters->CommonParameters = CommonParameters;
+	SetSSRParameters(&PassParameters->SSRPassCommonParameter);
+	PassParameters->RenderTargets = RenderTargets;
 
-		TShaderMapRef<FScreenSpaceReflectionsPS> PixelShader(View.ShaderMap, PermutationVector);
+	TShaderMapRef<FScreenSpaceReflectionsPS> PixelShader(View.ShaderMap, PermutationVector);
+
+	if (TiledScreenSpaceReflection == nullptr)
+	{
 		ClearUnusedGraphResources(*PixelShader, PassParameters);
 
 		GraphBuilder.AddPass(
@@ -726,30 +729,48 @@ void RenderScreenSpaceReflections(
 	}
 	else
 	{
-		check(TiledScreenSpaceReflection->TileSize == 8);
-		FScreenSpaceReflectionsCS::FPermutationDomain PermutationVector;
-		PermutationVector.Set<FSSRQualityDim>(SSRQuality);
-		PermutationVector.Set<FSSROutputForDenoiser>(bDenoiser);
+		check(TiledScreenSpaceReflection->TileSize == 8); // WORK_TILE_SIZE
 
-		FScreenSpaceReflectionsCS::FParameters* PassParameters = GraphBuilder.AllocParameters<FScreenSpaceReflectionsCS::FParameters>();
-		PassParameters->CommonParameters = CommonParameters;
-		SetSSRParameters(&PassParameters->SSRPassCommonParameter);
-		PassParameters->TileListDataBuffer = TiledScreenSpaceReflection->TileListStructureBufferSRV;
-		PassParameters->IndirectDispatchParameters = TiledScreenSpaceReflection->DispatchIndirectParametersBuffer;
-		PassParameters->OutColorUAV = GraphBuilder.CreateUAV(DenoiserInputs->Color);
-		if (bDenoiser)
-		{
-			PassParameters->OutClosestHitDistanceUAV = GraphBuilder.CreateUAV(DenoiserInputs->RayHitDistance);
-		}
+		FScreenSpaceReflectionsTileVS::FPermutationDomain VsPermutationVector;
+		TShaderMapRef<FScreenSpaceReflectionsTileVS> VertexShader(View.ShaderMap, VsPermutationVector);
 
-		TShaderMapRef<FScreenSpaceReflectionsCS> ComputeShader(View.ShaderMap, PermutationVector);
-		ClearUnusedGraphResources(*ComputeShader, PassParameters);
+		PassParameters->TileListData = TiledScreenSpaceReflection->TileListStructureBufferSRV;
+		PassParameters->IndirectDrawParameter = TiledScreenSpaceReflection->DispatchIndirectParametersBuffer;
+
+		ValidateShaderParameters(*VertexShader, *PassParameters);
+		ValidateShaderParameters(*PixelShader, *PassParameters);
 
-		FComputeShaderUtils::AddPass(GraphBuilder,
+		GraphBuilder.AddPass(
 			RDG_EVENT_NAME("SSR RayMarch(Quality=%d RayPerPixel=%d%s) %dx%d",
 				SSRQuality, RayTracingConfigs.RayCountPerPixel, bDenoiser ? TEXT(" DenoiserOutput") : TEXT(""),
 				View.ViewRect.Width(), View.ViewRect.Height()),
-			*ComputeShader, PassParameters, TiledScreenSpaceReflection->DispatchIndirectParametersBuffer, 0);
+			PassParameters,
+			ERDGPassFlags::Raster,
+			[PassParameters, &View, VertexShader, PixelShader, SSRStencilPrePass](FRHICommandList& RHICmdList)
+		{
+			SCOPED_GPU_STAT(RHICmdList, ScreenSpaceReflections);
+			RHICmdList.SetViewport(View.ViewRect.Min.X, View.ViewRect.Min.Y, 0.0f, View.ViewRect.Max.X, View.ViewRect.Max.Y, 1.0f);
+			RHICmdList.SetStencilRef(0x80);
+
+			FGraphicsPipelineStateInitializer GraphicsPSOInit;
+			FPixelShaderUtils::InitFullscreenPipelineState(RHICmdList, View.ShaderMap, *PixelShader, /* out */ GraphicsPSOInit);
+			if (SSRStencilPrePass)
+			{
+				// Clobers the stencil to pixel that should not compute SSR
+				GraphicsPSOInit.DepthStencilState = TStaticDepthStencilState<false, CF_Always, true, CF_Equal, SO_Keep, SO_Keep, SO_Keep>::GetRHI();
+			}
+			GraphicsPSOInit.PrimitiveType = GRHISupportsRectTopology ? PT_RectList : PT_TriangleList;
+			GraphicsPSOInit.BoundShaderState.VertexDeclarationRHI = GEmptyVertexDeclaration.VertexDeclarationRHI;
+			GraphicsPSOInit.BoundShaderState.VertexShaderRHI = VertexShader->GetVertexShader();
+			GraphicsPSOInit.BoundShaderState.PixelShaderRHI = PixelShader->GetPixelShader();
+
+			SetGraphicsPipelineState(RHICmdList, GraphicsPSOInit);
+			SetShaderParameters(RHICmdList, *VertexShader, VertexShader->GetVertexShader(), *PassParameters);
+			SetShaderParameters(RHICmdList, *PixelShader, PixelShader->GetPixelShader(), *PassParameters);
+
+			PassParameters->IndirectDrawParameter->MarkResourceAsUsed();
+			RHICmdList.DrawPrimitiveIndirect(PassParameters->IndirectDrawParameter->GetIndirectRHICallBuffer(), 0);
+		});
 	}
 } // RenderScreenSpaceReflections()
 
diff --git a/Engine/Source/Runtime/Renderer/Private/SingleLayerWaterRendering.cpp b/Engine/Source/Runtime/Renderer/Private/SingleLayerWaterRendering.cpp
--- a/Engine/Source/Runtime/Renderer/Private/SingleLayerWaterRendering.cpp
+++ b/Engine/Source/Runtime/Renderer/Private/SingleLayerWaterRendering.cpp
@@ -230,6 +230,8 @@ class FSingleLayerWaterCompositePS : public FGlobalShader
 
 	BEGIN_SHADER_PARAMETER_STRUCT(FParameters, )
 		SHADER_PARAMETER_STRUCT_INCLUDE(FSingleLayerWaterCommonShaderParameters, CommonParameters)
+		SHADER_PARAMETER_RDG_BUFFER(Buffer<uint>, IndirectDrawParameter)	// Not used in shader but need to be reference in the parameter list
+		SHADER_PARAMETER_RDG_BUFFER_SRV(StructuredBuffer<uint>, TileListData)
 		RENDER_TARGET_BINDING_SLOTS()
 	END_SHADER_PARAMETER_STRUCT()
 
@@ -270,10 +272,7 @@ class FWaterTileCategorisationCS : public FGlobalShader
 
 	BEGIN_SHADER_PARAMETER_STRUCT(FParameters, )
 		SHADER_PARAMETER_STRUCT_INCLUDE(FSingleLayerWaterCommonShaderParameters, CommonParameters)
-		SHADER_PARAMETER(uint32, TiledViewWidth)
-		SHADER_PARAMETER(uint32, TiledViewHeight)
-		SHADER_PARAMETER(float, TiledViewWidthInv)
-		SHADER_PARAMETER(float, TiledViewHeightInv)
+		SHADER_PARAMETER(uint32, VertexCountPerInstanceIndirect)
 		SHADER_PARAMETER_RDG_BUFFER_UAV(RWBuffer<uint>, DispatchIndirectDataUAV)
 		SHADER_PARAMETER_RDG_BUFFER_UAV(RWBuffer<uint>, WaterTileListDataUAV)
 	END_SHADER_PARAMETER_STRUCT()
@@ -295,7 +294,7 @@ class FWaterTileCategorisationCS : public FGlobalShader
 	static void ModifyCompilationEnvironment(const FGlobalShaderPermutationParameters& Parameters, FShaderCompilerEnvironment& OutEnvironment)
 	{
 		OutEnvironment.SetDefine(TEXT("TILE_CATERGORISATION_SHADER"), 1.0f);
-		OutEnvironment.SetDefine(TEXT("WATER_TILE_SIZE"), GetTileSize());
+		OutEnvironment.SetDefine(TEXT("WORK_TILE_SIZE"), GetTileSize());
 		FGlobalShader::ModifyCompilationEnvironment(Parameters, OutEnvironment);
 	}
 
@@ -303,22 +302,14 @@ class FWaterTileCategorisationCS : public FGlobalShader
 
 IMPLEMENT_GLOBAL_SHADER(FWaterTileCategorisationCS, "/Engine/Private/SingleLayerWaterComposite.usf", "WaterTileCatergorisationCS", SF_Compute);
 
-// Disabled water composition due to non 32bits UAV operations
-/*class FWaterTiledCompositeCS : public FGlobalShader
+class FWaterTileVS : public FGlobalShader
 {
-	DECLARE_GLOBAL_SHADER(FWaterTiledCompositeCS);
-	SHADER_USE_PARAMETER_STRUCT(FWaterTiledCompositeCS, FGlobalShader)
+	DECLARE_GLOBAL_SHADER(FWaterTileVS);
+	SHADER_USE_PARAMETER_STRUCT(FWaterTileVS, FGlobalShader);
 
-	using FPermutationDomain = TShaderPermutationDomain<FSingleLayerWaterScreenSpaceReflections>;
+	using FPermutationDomain = TShaderPermutationDomain<>;
 
-	BEGIN_SHADER_PARAMETER_STRUCT(FParameters, )
-		SHADER_PARAMETER_STRUCT_INCLUDE(FSingleLayerWaterCommonShaderParameters, CommonParameters)
-		SHADER_PARAMETER(uint32, TiledViewWidth)
-		SHADER_PARAMETER(uint32, TiledViewHeight)
-		SHADER_PARAMETER_RDG_BUFFER_SRV(Buffer<uint>, WaterTileListData)
-		SHADER_PARAMETER_UAV(RWTexture2D<float3>, SceneColorUAV)
-		SHADER_PARAMETER_RDG_BUFFER(Buffer<uint>, IndirectDispatchParameters)	// Not used in shader but need to be reference in the parameter list
-	END_SHADER_PARAMETER_STRUCT()
+	using FParameters = FSingleLayerWaterCompositePS::FParameters; // Sharing parameters for proper registration with RDG
 
 	static FPermutationDomain RemapPermutation(FPermutationDomain PermutationVector)
 	{
@@ -327,23 +318,18 @@ IMPLEMENT_GLOBAL_SHADER(FWaterTileCategorisationCS, "/Engine/Private/SingleLayer
 
 	static bool ShouldCompilePermutation(const FGlobalShaderPermutationParameters& Parameters)
 	{
-		if (!IsFeatureLevelSupported(Parameters.Platform, ERHIFeatureLevel::SM5) || !RHISupportsDrawIndirect(Parameters.Platform))
-		{
-			return false;
-		}
-		return true;
+		return IsFeatureLevelSupported(Parameters.Platform, ERHIFeatureLevel::SM5);
 	}
 
 	static void ModifyCompilationEnvironment(const FGlobalShaderPermutationParameters& Parameters, FShaderCompilerEnvironment& OutEnvironment)
 	{
-		OutEnvironment.SetDefine(TEXT("TILED_COMPOSITE_SHADER"), 1.0f);
-		OutEnvironment.SetDefine(TEXT("WATER_TILE_SIZE"), FWaterTileCategorisationCS::GetTileSize());
+		OutEnvironment.SetDefine(TEXT("TILE_VERTEX_SHADER"), 1.0f);
+		OutEnvironment.SetDefine(TEXT("WORK_TILE_SIZE"), FWaterTileCategorisationCS::GetTileSize());
 		FGlobalShader::ModifyCompilationEnvironment(Parameters, OutEnvironment);
 	}
-
 };
 
-IMPLEMENT_GLOBAL_SHADER(FWaterTiledCompositeCS, "/Engine/Private/SingleLayerWaterComposite.usf", "WaterTiledCategorisationCS", SF_Compute);*/
+IMPLEMENT_GLOBAL_SHADER(FWaterTileVS, "/Engine/Private/SingleLayerWaterComposite.usf", "WaterTileVS", SF_Vertex);
 
 class FWaterRefractionCopyPS : public FGlobalShader
 {
@@ -414,8 +400,8 @@ void FDeferredShadingSceneRenderer::CopySingleLayerWaterTextures(FRHICommandList
 	PassParameters->SceneDepthCopyDownsampleTexture = GraphBuilder.RegisterExternalTexture(SceneContext.SceneDepthZ);
 	PassParameters->SceneDepthCopyDownsampleSampler = TStaticSamplerState<SF_Bilinear, AM_Clamp, AM_Clamp, AM_Clamp>::GetRHI();
 	PassParameters->SVPositionToSourceTextureUV = FVector2D(RefractionDownsampleFactor / float(SceneContext.GetBufferSizeXY().X), RefractionDownsampleFactor / float(SceneContext.GetBufferSizeXY().Y));
-	PassParameters->RenderTargets[0] = FRenderTargetBinding(TargetColorTexture, ERenderTargetLoadAction::ELoad, ERenderTargetStoreAction::EStore);
-	PassParameters->RenderTargets.DepthStencil = FDepthStencilBinding(TargetDepthTexture, ERenderTargetLoadAction::ELoad, ERenderTargetStoreAction::EStore, FExclusiveDepthStencil::DepthWrite_StencilNop);
+	PassParameters->RenderTargets[0] = FRenderTargetBinding(TargetColorTexture, ERenderTargetLoadAction::ELoad);
+	PassParameters->RenderTargets.DepthStencil = FDepthStencilBinding(TargetDepthTexture, ERenderTargetLoadAction::ELoad, ERenderTargetLoadAction::ENoAction, FExclusiveDepthStencil::DepthWrite_StencilNop);
 
 	auto PixelShader = View.ShaderMap->GetShader<FWaterRefractionCopyPS>();
 
@@ -546,10 +532,10 @@ void FDeferredShadingSceneRenderer::RenderSingleLayerWaterReflections(FRHIComman
 		FIntVector TiledViewRes = FIntVector::DivideAndRoundUp(ViewRes, TiledScreenSpaceReflection.TileSize);
 		if (bRunTiled)
 		{
-			TiledScreenSpaceReflection.DispatchIndirectParametersBuffer = GraphBuilder.CreateBuffer(FRDGBufferDesc::CreateIndirectDesc<FRHIDispatchIndirectParameters>(), TEXT("WaterIndirectDrawParameters"));
+			TiledScreenSpaceReflection.DispatchIndirectParametersBuffer = GraphBuilder.CreateBuffer(FRDGBufferDesc::CreateIndirectDesc<FRHIDrawIndirectParameters>(), TEXT("WaterIndirectDrawParameters"));
 			TiledScreenSpaceReflection.DispatchIndirectParametersBufferUAV = GraphBuilder.CreateUAV(TiledScreenSpaceReflection.DispatchIndirectParametersBuffer);
 			FRDGBufferDesc TileListStructuredBufferDesc = FRDGBufferDesc::CreateStructuredDesc(4, TiledViewRes.X * TiledViewRes.Y); // one uint32 element per tile
-			TiledScreenSpaceReflection.TileListDataBuffer = GraphBuilder.CreateBuffer(TileListStructuredBufferDesc, TEXT("WaterTileList"));
+			TiledScreenSpaceReflection.TileListDataBuffer = GraphBuilder.CreateBuffer(TileListStructuredBufferDesc, TEXT("TileListDataBuffer"));
 			TiledScreenSpaceReflection.TileListStructureBufferUAV = GraphBuilder.CreateUAV(TiledScreenSpaceReflection.TileListDataBuffer);
 			TiledScreenSpaceReflection.TileListStructureBufferSRV = GraphBuilder.CreateSRV(TiledScreenSpaceReflection.TileListDataBuffer);
 
@@ -563,10 +549,7 @@ void FDeferredShadingSceneRenderer::RenderSingleLayerWaterReflections(FRHIComman
 
 				FWaterTileCategorisationCS::FParameters* PassParameters = GraphBuilder.AllocParameters<FWaterTileCategorisationCS::FParameters>();
 				SetCommonParameters(PassParameters->CommonParameters);
-				PassParameters->TiledViewWidth = TiledViewRes.X;
-				PassParameters->TiledViewHeight = TiledViewRes.Y;
-				PassParameters->TiledViewWidthInv = 1.0f / float(TiledViewRes.X);
-				PassParameters->TiledViewHeightInv = 1.0f / float(TiledViewRes.Y);
+				PassParameters->VertexCountPerInstanceIndirect = GRHISupportsRectTopology ? 3 : 6;
 				PassParameters->DispatchIndirectDataUAV = TiledScreenSpaceReflection.DispatchIndirectParametersBufferUAV;
 				PassParameters->WaterTileListDataUAV = TiledScreenSpaceReflection.TileListStructureBufferUAV;
 
@@ -617,54 +600,65 @@ void FDeferredShadingSceneRenderer::RenderSingleLayerWaterReflections(FRHIComman
 		}
 
 		// Composite reflections on water
-/*		if (bRunTiled)									// Disabled water composition due to non 32bits UAV operations
-		{
-			// Render Tiled composite CS shader
-			{
-				FWaterTiledCompositeCS::FPermutationDomain PermutationVector;
-				PermutationVector.Set<FSingleLayerWaterScreenSpaceReflections>(bEnableSSR);
-				TShaderMapRef<FWaterTiledCompositeCS> ComputeShader(View.ShaderMap, PermutationVector);
-
-				FWaterTiledCompositeCS::FParameters* PassParameters = GraphBuilder.AllocParameters<FWaterTiledCompositeCS::FParameters>();
-				SetCommonParameters(PassParameters->CommonParameters);
-				PassParameters->TiledViewWidth = TiledViewRes.X;
-				PassParameters->TiledViewHeight = TiledViewRes.Y;
-				PassParameters->WaterTileListData = TiledScreenSpaceReflection.TileListStructureBufferSRV;
-				PassParameters->SceneColorUAV = SceneContext.GetSceneColorTextureUAV();
-				PassParameters->IndirectDispatchParameters = TiledScreenSpaceReflection.DispatchIndirectParametersBuffer;
-				
-				FComputeShaderUtils::AddPass(GraphBuilder, RDG_EVENT_NAME("WaterTiledComposite"), *ComputeShader, PassParameters, TiledScreenSpaceReflection.DispatchIndirectParametersBuffer, 0);
-			}
-		}
-		else*/
 		{
 			FSingleLayerWaterCompositePS::FPermutationDomain PermutationVector;
 			PermutationVector.Set<FSingleLayerWaterScreenSpaceReflections>(bEnableSSR);
 			TShaderMapRef<FSingleLayerWaterCompositePS> PixelShader(View.ShaderMap, PermutationVector);
 
+			FWaterTileVS::FPermutationDomain VsPermutationVector;
+			TShaderMapRef<FWaterTileVS> VertexShader(View.ShaderMap, VsPermutationVector);
+
 			FSingleLayerWaterCompositePS::FParameters* PassParameters = GraphBuilder.AllocParameters<FSingleLayerWaterCompositePS::FParameters>();
 			SetCommonParameters(PassParameters->CommonParameters);
+			PassParameters->TileListData = TiledScreenSpaceReflection.TileListStructureBufferSRV;
+			PassParameters->IndirectDrawParameter = TiledScreenSpaceReflection.DispatchIndirectParametersBuffer;
 			PassParameters->RenderTargets[0] = FRenderTargetBinding(SceneColorTexture, ERenderTargetLoadAction::ELoad);
-			ClearUnusedGraphResources(*PixelShader, PassParameters);
 
+			if (bRunTiled)
+			{
+				ValidateShaderParameters(*VertexShader, *PassParameters);
+			}
+			ValidateShaderParameters(*PixelShader, *PassParameters);
 
 			GraphBuilder.AddPass(
 				RDG_EVENT_NAME("Water Composite %dx%d", View.ViewRect.Width(), View.ViewRect.Height()),
 				PassParameters,
 				ERDGPassFlags::Raster,
-				[PassParameters, &View, PixelShader](FRHICommandList& InRHICmdList)
+				[PassParameters, &View, TiledScreenSpaceReflection, VertexShader, PixelShader, bRunTiled](FRHICommandList& InRHICmdList)
 			{
 				InRHICmdList.SetViewport(View.ViewRect.Min.X, View.ViewRect.Min.Y, 0.0f, View.ViewRect.Max.X, View.ViewRect.Max.Y, 1.0f);
 
-				FGraphicsPipelineStateInitializer GraphicsPSOInit;
-				FPixelShaderUtils::InitFullscreenPipelineState(InRHICmdList, View.ShaderMap, *PixelShader, GraphicsPSOInit);
-
-				// Premultiplied alpha where alpha is transmittance.
-				GraphicsPSOInit.BlendState = TStaticBlendState<CW_RGB, BO_Add, BF_One, BF_SourceAlpha>::GetRHI(); 
-
-				SetGraphicsPipelineState(InRHICmdList, GraphicsPSOInit);
-				SetShaderParameters(InRHICmdList, *PixelShader, PixelShader->GetPixelShader(), *PassParameters);
-				FPixelShaderUtils::DrawFullscreenTriangle(InRHICmdList);
+				if (bRunTiled)
+				{
+					FGraphicsPipelineStateInitializer GraphicsPSOInit;
+					InRHICmdList.ApplyCachedRenderTargets(GraphicsPSOInit);
+					GraphicsPSOInit.PrimitiveType = GRHISupportsRectTopology ? PT_RectList : PT_TriangleList;
+					GraphicsPSOInit.BlendState = TStaticBlendState<CW_RGB, BO_Add, BF_One, BF_SourceAlpha>::GetRHI();
+					GraphicsPSOInit.DepthStencilState = TStaticDepthStencilState<false>::GetRHI();
+					GraphicsPSOInit.RasterizerState = TStaticRasterizerState<FM_Solid,CM_None>::GetRHI();
+					GraphicsPSOInit.BoundShaderState.VertexDeclarationRHI = GEmptyVertexDeclaration.VertexDeclarationRHI;
+					GraphicsPSOInit.BoundShaderState.VertexShaderRHI = VertexShader->GetVertexShader();
+					GraphicsPSOInit.BoundShaderState.PixelShaderRHI = PixelShader->GetPixelShader();
+					SetGraphicsPipelineState(InRHICmdList, GraphicsPSOInit);
+
+					SetShaderParameters(InRHICmdList, *PixelShader, PixelShader->GetPixelShader(), *PassParameters);
+					SetShaderParameters(InRHICmdList, *VertexShader, VertexShader->GetVertexShader(), *PassParameters);
+					
+					PassParameters->IndirectDrawParameter->MarkResourceAsUsed();
+					InRHICmdList.DrawPrimitiveIndirect(PassParameters->IndirectDrawParameter->GetIndirectRHICallBuffer(), 0);
+				}
+				else
+				{
+					FGraphicsPipelineStateInitializer GraphicsPSOInit;
+					FPixelShaderUtils::InitFullscreenPipelineState(InRHICmdList, View.ShaderMap, *PixelShader, GraphicsPSOInit);
+
+					// Premultiplied alpha where alpha is transmittance.
+					GraphicsPSOInit.BlendState = TStaticBlendState<CW_RGB, BO_Add, BF_One, BF_SourceAlpha>::GetRHI();
+
+					SetGraphicsPipelineState(InRHICmdList, GraphicsPSOInit);
+					SetShaderParameters(InRHICmdList, *PixelShader, PixelShader->GetPixelShader(), *PassParameters);
+					FPixelShaderUtils::DrawFullscreenTriangle(InRHICmdList);
+				}
 			});
 		}
 	}
