Single Layer Water optimization. Packed refraction color and depth into a single texture.

diff --git a/Engine/Shaders/Private/BasePassPixelShader.usf b/Engine/Shaders/Private/BasePassPixelShader.usf
--- a/Engine/Shaders/Private/BasePassPixelShader.usf
+++ b/Engine/Shaders/Private/BasePassPixelShader.usf
@@ -1225,8 +1225,7 @@ ShadingModel);
 
 		Color += EvaluateWaterVolumeLighting(
 			MaterialParameters, PixelMaterialInputs, ResolvedView,
-			OpaqueBasePass.SceneDepthWithoutSingleLayerWaterTexture, OpaqueBasePass.SceneDepthWithoutSingleLayerWaterSampler,
-			OpaqueBasePass.SceneLuminanceWithoutSingleLayerWaterTexture, OpaqueBasePass.SceneLuminanceWithoutSingleLayerWaterSampler,
+			OpaqueBasePass.SceneColorAndDepthWithoutSingleLayerWaterTexture, OpaqueBasePass.SceneColorAndDepthWithoutSingleLayerWaterSampler,
 			OpaqueBasePass.SceneCustomDepthTexture, OpaqueBasePass.SceneCustomDepthSampler,
 			OpaqueBasePass.SceneWithoutSingleLayerWaterValidUVRect,
 			Specular, OpaqueBasePass.DistortionParams,
diff --git a/Engine/Shaders/Private/SingleLayerWaterCommon.ush b/Engine/Shaders/Private/SingleLayerWaterCommon.ush
--- a/Engine/Shaders/Private/SingleLayerWaterCommon.ush
+++ b/Engine/Shaders/Private/SingleLayerWaterCommon.ush
@@ -7,6 +7,8 @@
 #include "ShadingCommon.ush" 
 #include "DistortionCommon.ush"
 
+// Scale for encoding depth into SceneColorAndDepthWithoutSingleLayerWaterTexture
+#define SINGLE_LAYER_WATER_DEPTH_SCALE 100.0f
 
 #ifndef SIMPLE_SINGLE_LAYER_WATER
 #define SIMPLE_SINGLE_LAYER_WATER 0
@@ -43,8 +45,7 @@ WaterVolumeLightingOutput EvaluateWaterVolumeLighting(
 #if SIMPLE_SINGLE_LAYER_WATER
 	float DirectionalLightShadow,
 #else
-	Texture2D SceneDepthWithoutSingleLayerWaterTexture, SamplerState SceneDepthWithoutSingleLayerWaterSampler,
-	Texture2D SceneLuminanceWithoutSingleLayerWaterTexture, SamplerState SceneLuminanceWithoutSingleLayerWaterSampler,
+	Texture2D SceneColorAndDepthWithoutSingleLayerWaterTexture, SamplerState SceneColorAndDepthWithoutSingleLayerWaterSampler,
 	Texture2D SceneCustomDepthTexture, SamplerState SceneCustomDepthSampler,
 	float4 SceneWithoutSingleLayerWaterValidUVRect,
 #endif // SIMPLE_SINGLE_LAYER_WATER
@@ -80,8 +81,7 @@ WaterVolumeLightingOutput EvaluateWaterVolumeLighting(
 
 #else  // !SIMPLE_SINGLE_LAYER_WATER
 
-		float PixelSceneDeviceZ = SceneDepthWithoutSingleLayerWaterTexture.SampleLevel(SceneDepthWithoutSingleLayerWaterSampler, ViewportUV, 0).r;
-		float PixelSceneDepth = ConvertFromDeviceZ(PixelSceneDeviceZ);
+		float PixelSceneDepth = SceneColorAndDepthWithoutSingleLayerWaterTexture.SampleLevel(SceneColorAndDepthWithoutSingleLayerWaterSampler, ViewportUV, 0).w * SINGLE_LAYER_WATER_DEPTH_SCALE;
 		
 		float2 Refraction = GetMaterialRefraction(PixelMaterialInputs); // X = IOR, Y = Refraction Bias
 		float2 BufferUVDistortion = ComputeBufferUVDistortion(
@@ -94,8 +94,8 @@ WaterVolumeLightingOutput EvaluateWaterVolumeLighting(
 		// Clamp distorted UVs to not sample black outside the sub part of the buffer. Big quality increase.
 		const float2 DistortedUV = clamp(ViewportUV + BufferUVDistortion, SceneWithoutSingleLayerWaterValidUVRect.xy, SceneWithoutSingleLayerWaterValidUVRect.zw);
 
-		float SceneDeviceZ = SceneDepthWithoutSingleLayerWaterTexture.SampleLevel(SceneDepthWithoutSingleLayerWaterSampler, DistortedUV, 0).r;
-		float SceneDepth = ConvertFromDeviceZ(SceneDeviceZ);
+		float SceneDepth = SceneColorAndDepthWithoutSingleLayerWaterTexture.SampleLevel(SceneColorAndDepthWithoutSingleLayerWaterSampler, DistortedUV, 0).w * SINGLE_LAYER_WATER_DEPTH_SCALE;
+
 		if (SceneDepth < WaterDepth)
 		{
 			SceneDepth = PixelSceneDepth; // The pixel we try to sample is closer than the water pixel: falback to no distortion
@@ -133,6 +133,7 @@ WaterVolumeLightingOutput EvaluateWaterVolumeLighting(
 		// We also apply transmittance from light to under water surface. However, the scene has been lit by many sources already.
 		// So the transmittance toabove surface is simply approximated using the travel distance from the scene pixel to the water top, assuming a flat water surface.
 		// We cannot combine this transmittance with the transmittance from view because this would change the behavior of the analytical integration of light scattering integration.
+		const float SceneDeviceZ = ConvertToDeviceZ(SceneDepth);
 		const float3 BehindWaterSceneWorldPos = SvPositionToWorld(float4(MaterialParameters.SvPosition.xy, SceneDeviceZ, 1.0));
 		const float DistanceFromScenePixelToWaterTop = max(0.0, MaterialParameters.AbsoluteWorldPosition.z - BehindWaterSceneWorldPos.z);
 		const float3 MeanTransmittanceToLightSources = exp(-DistanceFromScenePixelToWaterTop * ExtinctionCoeff);
@@ -141,7 +142,7 @@ WaterVolumeLightingOutput EvaluateWaterVolumeLighting(
 		const float3 BehindWaterSceneLuminance = 0.0f; // Cannot read back the scene color in this case
 #else
 		// We use the pixel SvPosition instead of the scene one pre refraction/distortion to avoid those extra ALUs.
-		float3 BehindWaterSceneLuminance = SceneLuminanceWithoutSingleLayerWaterTexture.SampleLevel(SceneLuminanceWithoutSingleLayerWaterSampler, ViewportUV, 0).rgb;
+		float3 BehindWaterSceneLuminance = SceneColorAndDepthWithoutSingleLayerWaterTexture.SampleLevel(SceneColorAndDepthWithoutSingleLayerWaterSampler, ViewportUV, 0).rgb;
 		BehindWaterSceneLuminance = MeanTransmittanceToLightSources * (USE_PREEXPOSURE ? ResolvedView.OneOverPreExposure : 1.0f) * BehindWaterSceneLuminance;
 #endif
 
diff --git a/Engine/Shaders/Private/SingleLayerWaterComposite.usf b/Engine/Shaders/Private/SingleLayerWaterComposite.usf
--- a/Engine/Shaders/Private/SingleLayerWaterComposite.usf
+++ b/Engine/Shaders/Private/SingleLayerWaterComposite.usf
@@ -11,6 +11,8 @@
 #include "ShadingModels.ush"
 #include "SceneTextureParameters.ush"
 
+// Scale for encoding depth into SceneColorAndDepthWithoutSingleLayerWaterTexture
+#define SINGLE_LAYER_WATER_DEPTH_SCALE 100.0f
 
 #define ENABLE_SKY_LIGHT 1
 
@@ -23,8 +25,8 @@
 Texture2D ScreenSpaceReflectionsTexture;
 SamplerState ScreenSpaceReflectionsSampler;
 
-Texture2D SceneNoWaterDepthTexture;
-SamplerState SceneNoWaterDepthSampler;
+Texture2D SceneNoWaterColorAndDepthTexture;
+SamplerState SceneNoWaterColorAndDepthSampler;
 
 struct SingleLayerWaterCompositeOutput
 {
@@ -55,8 +57,9 @@ SingleLayerWaterCompositeOutput SingleLayerWaterComposite(float2 BufferUV, float
 #endif
 	if (ShadingModelID == SHADINGMODELID_SINGLELAYERWATER)
 	{
-		const float OpaqueDepth = ConvertFromDeviceZ(SceneNoWaterDepthTexture.SampleLevel(SceneNoWaterDepthSampler, BufferUV, 0).r);
-		const float WaterDepth = GBuffer.Depth;
+		float OpaqueDepth = SceneNoWaterColorAndDepthTexture.SampleLevel(SceneNoWaterColorAndDepthSampler, BufferUV, 0).w * SINGLE_LAYER_WATER_DEPTH_SCALE;
+		float WaterDepth = GBuffer.Depth;
+
 		float DeltaDepth = CameraIsUnderWater ? WaterDepth : OpaqueDepth - WaterDepth;	// Inverted depth
 		if (DeltaDepth > 0.0)
 		{
@@ -101,7 +104,7 @@ SingleLayerWaterCompositeOutput SingleLayerWaterComposite(float2 BufferUV, float
 			Reflection *= EnvBrdf;
 
 			// Soft fading near shor to avoid seeing triangles.
-			const float ShoreOpacity = saturate(DeltaDepth * 0.02);
+			const float ShoreOpacity = saturate(DeltaDepth * 0.02f);
 			Reflection.rgb *= ShoreOpacity;
 
 			// Apply transmittance to reflection if under water
@@ -288,12 +291,11 @@ float2 SVPositionToSourceTextureUV;
 
 void WaterRefractionCopyPS(
 	in float4 SVPosition : SV_Position,
-	out float4 OutColor : SV_Target0,
-	out float OutDepth : SV_DEPTH)
+	out float4 OutColorAndDepth : SV_Target0)
 {
 	float2 BufferUV = SVPosition.xy * SVPositionToSourceTextureUV;
 
 	// Downsample and average all 2x2 quads
-	OutColor = Texture2DSampleLevel(SceneColorCopyDownsampleTexture, SceneColorCopyDownsampleSampler, BufferUV, 0);
-	OutDepth = Texture2DSampleLevel(SceneDepthCopyDownsampleTexture, SceneDepthCopyDownsampleSampler, BufferUV, 0).r;
+	OutColorAndDepth.xyz = Texture2DSampleLevel(SceneColorCopyDownsampleTexture, SceneColorCopyDownsampleSampler, BufferUV, 0).xyz;
+	OutColorAndDepth.w = ConvertFromDeviceZ(Texture2DSampleLevel(SceneDepthCopyDownsampleTexture, SceneDepthCopyDownsampleSampler, BufferUV, 0).r) / SINGLE_LAYER_WATER_DEPTH_SCALE;
 }
\ No newline at end of file
diff --git a/Engine/Shaders/Public/ShaderVersion.ush b/Engine/Shaders/Public/ShaderVersion.ush
--- a/Engine/Shaders/Public/ShaderVersion.ush
+++ b/Engine/Shaders/Public/ShaderVersion.ush
@@ -3,6 +3,4 @@
 // in Platform.ush (which should be included in any shader) it allows to invalidate the shader DDC.
 // 
 // If you are merging streams and there is a conflict with this GUID you should make a new GUID rather than taking one or the other.
-// GUID = 919A4804-C336-4DCC-BC7B-22C1054255EF
-
-
+// GUID = CE7B9FFB-2E88-43DA-AD0A-C147887A37ED
diff --git a/Engine/Source/Runtime/Renderer/Private/BasePassRendering.cpp b/Engine/Source/Runtime/Renderer/Private/BasePassRendering.cpp
--- a/Engine/Source/Runtime/Renderer/Private/BasePassRendering.cpp
+++ b/Engine/Source/Runtime/Renderer/Private/BasePassRendering.cpp
@@ -517,8 +517,7 @@ void CreateOpaqueBasePassUniformBuffer(
 	const FViewInfo& View,
 	IPooledRenderTarget* ForwardScreenSpaceShadowMask,
 	FVector4* SceneWithoutSingleLayerWaterValidUVRect,
-	IPooledRenderTarget* SceneLuminanceWithoutSingleLayerWater,
-	IPooledRenderTarget* SceneDepthWithoutSingleLayerWaterTexture,
+	IPooledRenderTarget* SceneColorAndDepthWithoutSingleLayerWater,
 	TUniformBufferRef<FOpaqueBasePassUniformParameters>& BasePassUniformBuffer)
 {
 	FSceneRenderTargets& SceneRenderTargets = FSceneRenderTargets::Get(RHICmdList);
@@ -584,10 +583,8 @@ void CreateOpaqueBasePassUniformBuffer(
 
 	// Single Layer Water
 	BasePassParameters.SceneWithoutSingleLayerWaterValidUVRect = SceneWithoutSingleLayerWaterValidUVRect ? *SceneWithoutSingleLayerWaterValidUVRect : FVector4(0.0f, 0.0f, 1.0f, 1.0f);
-	BasePassParameters.SceneLuminanceWithoutSingleLayerWaterTexture = (SceneLuminanceWithoutSingleLayerWater ? SceneLuminanceWithoutSingleLayerWater->GetRenderTargetItem().TargetableTexture : GSystemTextures.BlackDummy->GetRenderTargetItem().TargetableTexture);
-	BasePassParameters.SceneLuminanceWithoutSingleLayerWaterSampler = TStaticSamplerState<SF_Bilinear>::GetRHI();
-	BasePassParameters.SceneDepthWithoutSingleLayerWaterTexture = (SceneDepthWithoutSingleLayerWaterTexture ? SceneDepthWithoutSingleLayerWaterTexture->GetRenderTargetItem().TargetableTexture : GSystemTextures.BlackDummy->GetRenderTargetItem().TargetableTexture);
-	BasePassParameters.SceneDepthWithoutSingleLayerWaterSampler = TStaticSamplerState<SF_Point>::GetRHI();
+	BasePassParameters.SceneColorAndDepthWithoutSingleLayerWaterTexture = (SceneColorAndDepthWithoutSingleLayerWater ? SceneColorAndDepthWithoutSingleLayerWater->GetRenderTargetItem().TargetableTexture : GSystemTextures.BlackDummy->GetRenderTargetItem().TargetableTexture);
+	BasePassParameters.SceneColorAndDepthWithoutSingleLayerWaterSampler = TStaticSamplerState<SF_Bilinear>::GetRHI();
 	BasePassParameters.PreIntegratedGFTexture = GSystemTextures.PreintegratedGF->GetRenderTargetItem().ShaderResourceTexture;
 	BasePassParameters.PreIntegratedGFSampler = TStaticSamplerState<SF_Bilinear, AM_Clamp, AM_Clamp, AM_Clamp>::GetRHI();
 	BasePassParameters.SceneCustomDepthTexture = (SceneRenderTargets.CustomDepth.IsValid() ? SceneRenderTargets.CustomDepth : GSystemTextures.BlackDummy)->GetRenderTargetItem().ShaderResourceTexture;
@@ -652,7 +649,7 @@ bool FDeferredShadingSceneRenderer::RenderBasePass(FRHICommandListImmediate& RHI
 				SCOPED_GPU_MASK(RHICmdList, View.GPUMask);
 
 				TUniformBufferRef<FOpaqueBasePassUniformParameters> BasePassUniformBuffer;
-				CreateOpaqueBasePassUniformBuffer(RHICmdList, View, ForwardScreenSpaceShadowMask, nullptr, nullptr, nullptr, BasePassUniformBuffer);
+				CreateOpaqueBasePassUniformBuffer(RHICmdList, View, ForwardScreenSpaceShadowMask, nullptr, nullptr, BasePassUniformBuffer);
 
 				FMeshPassProcessorRenderState DrawRenderState(View, BasePassUniformBuffer);
 
@@ -693,7 +690,7 @@ bool FDeferredShadingSceneRenderer::RenderBasePass(FRHICommandListImmediate& RHI
 				SCOPED_GPU_MASK(RHICmdList, View.GPUMask);
 
 				TUniformBufferRef<FOpaqueBasePassUniformParameters> BasePassUniformBuffer;
-				CreateOpaqueBasePassUniformBuffer(RHICmdList, View, ForwardScreenSpaceShadowMask, nullptr, nullptr, nullptr, BasePassUniformBuffer);
+				CreateOpaqueBasePassUniformBuffer(RHICmdList, View, ForwardScreenSpaceShadowMask, nullptr, nullptr, BasePassUniformBuffer);
 
 				FMeshPassProcessorRenderState DrawRenderState(View, BasePassUniformBuffer);
 
diff --git a/Engine/Source/Runtime/Renderer/Private/BasePassRendering.h b/Engine/Source/Runtime/Renderer/Private/BasePassRendering.h
--- a/Engine/Source/Runtime/Renderer/Private/BasePassRendering.h
+++ b/Engine/Source/Runtime/Renderer/Private/BasePassRendering.h
@@ -72,10 +72,8 @@ BEGIN_GLOBAL_SHADER_PARAMETER_STRUCT(FOpaqueBasePassUniformParameters,)
 	SHADER_PARAMETER_SAMPLER(SamplerState, DBufferCTextureSampler)
 	SHADER_PARAMETER_TEXTURE(Texture2D<uint>, DBufferRenderMask)
 	// Single Layer Water
-	SHADER_PARAMETER_TEXTURE(Texture2D, SceneLuminanceWithoutSingleLayerWaterTexture)
-	SHADER_PARAMETER_SAMPLER(SamplerState, SceneLuminanceWithoutSingleLayerWaterSampler)
-	SHADER_PARAMETER_TEXTURE(Texture2D, SceneDepthWithoutSingleLayerWaterTexture)
-	SHADER_PARAMETER_SAMPLER(SamplerState, SceneDepthWithoutSingleLayerWaterSampler)
+	SHADER_PARAMETER_TEXTURE(Texture2D, SceneColorAndDepthWithoutSingleLayerWaterTexture)
+	SHADER_PARAMETER_SAMPLER(SamplerState, SceneColorAndDepthWithoutSingleLayerWaterSampler)
 	SHADER_PARAMETER_TEXTURE(Texture2D, PreIntegratedGFTexture)
 	SHADER_PARAMETER_SAMPLER(SamplerState, PreIntegratedGFSampler)
 	SHADER_PARAMETER_TEXTURE(Texture2D, SceneCustomDepthTexture)
@@ -121,8 +119,7 @@ extern void CreateOpaqueBasePassUniformBuffer(
 	const FViewInfo& View,
 	IPooledRenderTarget* ForwardScreenSpaceShadowMask,
 	FVector4* SceneWithoutSingleLayerWaterValidUVRect,
-	IPooledRenderTarget* SceneLuminanceWithoutSingleLayerWater,
-	IPooledRenderTarget* SceneDepthWithoutSingleLayerWaterTexture,
+	IPooledRenderTarget* SceneColorAndDepthWithoutSingleLayerWater,
 	TUniformBufferRef<FOpaqueBasePassUniformParameters>& BasePassUniformBuffer);
 
 extern void CreateTranslucentBasePassUniformBuffer(
diff --git a/Engine/Source/Runtime/Renderer/Private/PostProcess/PostProcessCompositeEditorPrimitives.cpp b/Engine/Source/Runtime/Renderer/Private/PostProcess/PostProcessCompositeEditorPrimitives.cpp
--- a/Engine/Source/Runtime/Renderer/Private/PostProcess/PostProcessCompositeEditorPrimitives.cpp
+++ b/Engine/Source/Runtime/Renderer/Private/PostProcess/PostProcessCompositeEditorPrimitives.cpp
@@ -565,7 +565,7 @@ void FRCPassPostProcessCompositeEditorPrimitives::Process(FRenderingCompositePas
 
 			if (bDeferredBasePass)
 			{
-				CreateOpaqueBasePassUniformBuffer(Context.RHICmdList, EditorView, nullptr, nullptr, nullptr, nullptr, OpaqueBasePassUniformBuffer);
+				CreateOpaqueBasePassUniformBuffer(Context.RHICmdList, EditorView, nullptr, nullptr, nullptr, OpaqueBasePassUniformBuffer);
 				BasePassUniformBuffer = OpaqueBasePassUniformBuffer;
 			}
 			else
diff --git a/Engine/Source/Runtime/Renderer/Private/Renderer.cpp b/Engine/Source/Runtime/Renderer/Private/Renderer.cpp
--- a/Engine/Source/Runtime/Renderer/Private/Renderer.cpp
+++ b/Engine/Source/Runtime/Renderer/Private/Renderer.cpp
@@ -296,7 +296,7 @@ void FRendererModule::DrawTileMesh(FRHICommandListImmediate& RHICmdList, FMeshPa
 				if (ShadingPath == EShadingPath::Deferred)
 				{
 					TUniformBufferRef<FOpaqueBasePassUniformParameters> OpaqueBasePassUniformBuffer;
-					CreateOpaqueBasePassUniformBuffer(RHICmdList, View, nullptr, nullptr, nullptr, nullptr, OpaqueBasePassUniformBuffer);
+					CreateOpaqueBasePassUniformBuffer(RHICmdList, View, nullptr, nullptr, nullptr, OpaqueBasePassUniformBuffer);
 					BasePassUniformBuffer = OpaqueBasePassUniformBuffer;
 					
 					DrawRenderState.SetPassUniformBuffer(BasePassUniformBuffer);
diff --git a/Engine/Source/Runtime/Renderer/Private/SingleLayerWaterRendering.cpp b/Engine/Source/Runtime/Renderer/Private/SingleLayerWaterRendering.cpp
--- a/Engine/Source/Runtime/Renderer/Private/SingleLayerWaterRendering.cpp
+++ b/Engine/Source/Runtime/Renderer/Private/SingleLayerWaterRendering.cpp
@@ -210,8 +210,8 @@ BEGIN_SHADER_PARAMETER_STRUCT(FSingleLayerWaterCommonShaderParameters, )
 	SHADER_PARAMETER_SAMPLER(SamplerState, ScreenSpaceReflectionsSampler)
 	SHADER_PARAMETER_TEXTURE(Texture2D, PreIntegratedGF)
 	SHADER_PARAMETER_SAMPLER(SamplerState, PreIntegratedGFSampler)
-	SHADER_PARAMETER_RDG_TEXTURE(Texture2D, SceneNoWaterDepthTexture)
-	SHADER_PARAMETER_SAMPLER(SamplerState, SceneNoWaterDepthSampler)
+	SHADER_PARAMETER_RDG_TEXTURE(Texture2D, SceneNoWaterColorAndDepthTexture)
+	SHADER_PARAMETER_SAMPLER(SamplerState, SceneNoWaterColorAndDepthSampler)
 	SHADER_PARAMETER_STRUCT_INCLUDE(FSceneTextureParameters, SceneTextures)	// Water scene texture
 	SHADER_PARAMETER_STRUCT_INCLUDE(FSceneTextureSamplerParameters, SceneTextureSamplers)
 	SHADER_PARAMETER_STRUCT_REF(FViewUniformShaderParameters, ViewUniformBuffer)
@@ -373,25 +373,15 @@ void FDeferredShadingSceneRenderer::CopySingleLayerWaterTextures(FRHICommandList
 	const int32 RefractionDownsampleFactor = FMath::Clamp(GSingleLayerWaterRefractionDownsampleFactor, 1, 8);
 	const FIntPoint RefractionResolution = FIntPoint::DivideAndRoundDown(SceneContext.GetBufferSizeXY(), RefractionDownsampleFactor);
 
-	// Allocate required buffers
-	{
-		FPooledRenderTargetDesc Desc(FPooledRenderTargetDesc::Create2DDesc(RefractionResolution, PF_DepthStencil, SceneContext.GetDefaultDepthClear(), TexCreate_None, TexCreate_DepthStencilTargetable | TexCreate_ShaderResource | TexCreate_InputAttachmentRead, false));
-		Desc.NumSamples = SceneContext.GetNumSceneColorMSAASamples(CurrentFeatureLevel);
-		Desc.Flags |= GFastVRamConfig.SceneDepth;
-		GRenderTargetPool.FindFreeElement(RHICmdList, Desc, PassData.SceneDepthZWithoutSingleLayerWater, TEXT("SceneDepthZWithoutSingleLayerWater"), true);
-	}
-	{
-		FPooledRenderTargetDesc Desc(FPooledRenderTargetDesc::Create2DDesc(RefractionResolution, SceneContext.GetSceneColorFormat(), SceneContext.GetDefaultColorClear(), TexCreate_None, TexCreate_RenderTargetable, false));
-		Desc.NumSamples = SceneContext.GetNumSceneColorMSAASamples(CurrentFeatureLevel);
-		GRenderTargetPool.FindFreeElement(RHICmdList, Desc, PassData.SceneColorWithoutSingleLayerWater, TEXT("SceneColorWithoutSingleLayerWater"));
-	}
+	FPooledRenderTargetDesc Desc(FPooledRenderTargetDesc::Create2DDesc(RefractionResolution, PF_FloatRGBA, SceneContext.GetDefaultColorClear(), TexCreate_None, TexCreate_RenderTargetable, false));
+	Desc.NumSamples = SceneContext.GetNumSceneColorMSAASamples(CurrentFeatureLevel);
+	GRenderTargetPool.FindFreeElement(RHICmdList, Desc, PassData.SceneColorAndDepthWithoutSingleLayerWater, TEXT("SceneColorAndDepthWithoutSingleLayerWater"));
 
 	// For now support only the 1st view
 	const FViewInfo& View = Views[0];
 
 	FRDGBuilder GraphBuilder(RHICmdList);
-	FRDGTextureRef TargetColorTexture = GraphBuilder.RegisterExternalTexture(PassData.SceneColorWithoutSingleLayerWater);
-	FRDGTextureRef TargetDepthTexture = GraphBuilder.RegisterExternalTexture(PassData.SceneDepthZWithoutSingleLayerWater);
+	FRDGTextureRef TargetColorTexture = GraphBuilder.RegisterExternalTexture(PassData.SceneColorAndDepthWithoutSingleLayerWater);
 
 	FWaterRefractionCopyPS::FParameters* PassParameters = GraphBuilder.AllocParameters<FWaterRefractionCopyPS::FParameters>();
 	PassParameters->ViewUniformBuffer = View.ViewUniformBuffer;
@@ -401,7 +391,6 @@ void FDeferredShadingSceneRenderer::CopySingleLayerWaterTextures(FRHICommandList
 	PassParameters->SceneDepthCopyDownsampleSampler = TStaticSamplerState<SF_Bilinear, AM_Clamp, AM_Clamp, AM_Clamp>::GetRHI();
 	PassParameters->SVPositionToSourceTextureUV = FVector2D(RefractionDownsampleFactor / float(SceneContext.GetBufferSizeXY().X), RefractionDownsampleFactor / float(SceneContext.GetBufferSizeXY().Y));
 	PassParameters->RenderTargets[0] = FRenderTargetBinding(TargetColorTexture, ERenderTargetLoadAction::ELoad);
-	PassParameters->RenderTargets.DepthStencil = FDepthStencilBinding(TargetDepthTexture, ERenderTargetLoadAction::ELoad, ERenderTargetLoadAction::ENoAction, FExclusiveDepthStencil::DepthWrite_StencilNop);
 
 	auto PixelShader = View.ShaderMap->GetShader<FWaterRefractionCopyPS>();
 
@@ -418,10 +407,7 @@ void FDeferredShadingSceneRenderer::CopySingleLayerWaterTextures(FRHICommandList
 		RDG_EVENT_NAME("Water Refraction Copy"),
 		PixelShader,
 		PassParameters,
-		RefractionViewRect,
-		nullptr,
-		nullptr,
-		TStaticDepthStencilState<true, CF_Always>::GetRHI());
+		RefractionViewRect);
 	GraphBuilder.Execute();
 }
 
@@ -513,8 +499,8 @@ void FDeferredShadingSceneRenderer::RenderSingleLayerWaterReflections(FRHIComman
 			Parameters.ScreenSpaceReflectionsSampler = TStaticSamplerState<SF_Point>::GetRHI();
 			Parameters.PreIntegratedGF = GSystemTextures.PreintegratedGF->GetRenderTargetItem().ShaderResourceTexture;
 			Parameters.PreIntegratedGFSampler = TStaticSamplerState<SF_Bilinear, AM_Clamp, AM_Clamp, AM_Clamp>::GetRHI();
-			Parameters.SceneNoWaterDepthTexture = GraphBuilder.RegisterExternalTexture(PassData.SceneDepthZWithoutSingleLayerWater ? PassData.SceneDepthZWithoutSingleLayerWater : GSystemTextures.BlackDummy);
-			Parameters.SceneNoWaterDepthSampler = TStaticSamplerState<SF_Point>::GetRHI();
+			Parameters.SceneNoWaterColorAndDepthTexture = GraphBuilder.RegisterExternalTexture(PassData.SceneColorAndDepthWithoutSingleLayerWater ? PassData.SceneColorAndDepthWithoutSingleLayerWater : GSystemTextures.BlackDummy);
+			Parameters.SceneNoWaterColorAndDepthSampler = TStaticSamplerState<SF_Point>::GetRHI();
 			Parameters.SceneTextures = SceneTextures;
 			SetupSceneTextureSamplers(&Parameters.SceneTextureSamplers);
 			Parameters.ViewUniformBuffer = View.ViewUniformBuffer;
@@ -700,8 +686,7 @@ bool FDeferredShadingSceneRenderer::RenderSingleLayerWaterPass(FRHICommandListIm
 				View, 
 				WhiteForwardScreenSpaceShadowMask, 
 				&PassData.SceneWithoutSingleLayerWaterValidUVRect,
-				PassData.SceneColorWithoutSingleLayerWater,
-				PassData.SceneDepthZWithoutSingleLayerWater,
+				PassData.SceneColorAndDepthWithoutSingleLayerWater,
 				WaterPassUniformBuffer);
 
 			FMeshPassProcessorRenderState DrawRenderState(View, WaterPassUniformBuffer);
diff --git a/Engine/Source/Runtime/Renderer/Private/SingleLayerWaterRendering.h b/Engine/Source/Runtime/Renderer/Private/SingleLayerWaterRendering.h
--- a/Engine/Source/Runtime/Renderer/Private/SingleLayerWaterRendering.h
+++ b/Engine/Source/Runtime/Renderer/Private/SingleLayerWaterRendering.h
@@ -18,8 +18,7 @@ class FViewInfo;
 
 struct FSingleLayerWaterPassData
 {
-	TRefCountPtr<IPooledRenderTarget> SceneColorWithoutSingleLayerWater;
-	TRefCountPtr<IPooledRenderTarget> SceneDepthZWithoutSingleLayerWater;
+	TRefCountPtr<IPooledRenderTarget> SceneColorAndDepthWithoutSingleLayerWater;
 	FVector4 SceneWithoutSingleLayerWaterValidUVRect;
 };
 
