Parallelize SingleLayerWaterRendering

#rnx

diff --git a/Engine/Source/Runtime/Renderer/Private/DeferredShadingRenderer.cpp b/Engine/Source/Runtime/Renderer/Private/DeferredShadingRenderer.cpp
--- a/Engine/Source/Runtime/Renderer/Private/DeferredShadingRenderer.cpp
+++ b/Engine/Source/Runtime/Renderer/Private/DeferredShadingRenderer.cpp
@@ -105,6 +105,13 @@ static TAutoConsoleVariable<int32> CVarParallelBasePass(
 	ECVF_RenderThreadSafe
 );
 
+static TAutoConsoleVariable<int32> CVarParallelSingleLayerWaterPass(
+	TEXT("r.ParallelSingleLayerWaterPass"),
+	1,
+	TEXT("Toggles parallel single layer water pass rendering. Parallel rendering must be enabled for this to have an effect."),
+	ECVF_RenderThreadSafe
+);
+
 static int32 GRayTracing = 0;
 static TAutoConsoleVariable<int32> CVarRayTracing(
 	TEXT("r.RayTracing"),
@@ -1806,9 +1813,18 @@ void FDeferredShadingSceneRenderer::Render(FRHICommandListImmediate& RHICmdList)
 
 		// Render the GBuffer pass, updating the GBuffer and also writing lit water in the scene.
 		const FExclusiveDepthStencil::Type WaterPassDepthStencilAccess = FExclusiveDepthStencil::DepthWrite_StencilWrite;
-		BeginRenderingWaterGBuffer(RHICmdList, SingleLayerWaterPassData, WaterPassDepthStencilAccess, ViewFamily.EngineShowFlags.ShaderComplexity);
-		RenderSingleLayerWaterPass(RHICmdList, SingleLayerWaterPassData, WaterPassDepthStencilAccess);
-		FinishWaterGBufferPassAndResolve(RHICmdList);
+		const bool bDoParallelSingleLayerWater = GRHICommandList.UseParallelAlgorithms() && CVarParallelSingleLayerWaterPass.GetValueOnRenderThread()==1;
+		if (!bDoParallelSingleLayerWater)
+		{
+			BeginRenderingWaterGBuffer(RHICmdList, WaterPassDepthStencilAccess, ViewFamily.EngineShowFlags.ShaderComplexity, ShaderPlatform);
+		}
+		
+		RenderSingleLayerWaterPass(RHICmdList, SingleLayerWaterPassData, WaterPassDepthStencilAccess, bDoParallelSingleLayerWater);
+		if (bDoParallelSingleLayerWater)
+		{
+			BeginRenderingWaterGBuffer(RHICmdList, WaterPassDepthStencilAccess, ViewFamily.EngineShowFlags.ShaderComplexity, ShaderPlatform);
+		}
+        FinishWaterGBufferPassAndResolve(RHICmdList);
 
 		// Resolves the depth texture back to readable for SSR and later passes.1
 		SceneContext.ResolveSceneDepthTexture(RHICmdList, FResolveRect(0, 0, FamilySize.X, FamilySize.Y));
diff --git a/Engine/Source/Runtime/Renderer/Private/DeferredShadingRenderer.h b/Engine/Source/Runtime/Renderer/Private/DeferredShadingRenderer.h
--- a/Engine/Source/Runtime/Renderer/Private/DeferredShadingRenderer.h
+++ b/Engine/Source/Runtime/Renderer/Private/DeferredShadingRenderer.h
@@ -106,16 +106,16 @@ public:
 	/** Copy the scene color and depth buffer for water refraction*/
 	void CopySingleLayerWaterTextures(FRHICommandListImmediate& RHICmdList, FSingleLayerWaterPassData& PassData);
 	/** Begin the water GBuffer rendering pass*/
-	void BeginRenderingWaterGBuffer(FRHICommandList& RHICmdList, FSingleLayerWaterPassData& PassData, FExclusiveDepthStencil::Type DepthStencilAccess, bool bBindQuadOverdrawBuffers);
+	static void BeginRenderingWaterGBuffer(FRHICommandList& RHICmdList, FExclusiveDepthStencil::Type DepthStencilAccess, bool bBindQuadOverdrawBuffers, EShaderPlatform InShaderPlatform);
 	/** End the water GBuffer rendering pass*/
 	void FinishWaterGBufferPassAndResolve(FRHICommandListImmediate& RHICmdList);
 	/**
 	* Renders the scene's single layer water base pass
 	* @return true if anything was rendered
 	*/
-	bool RenderSingleLayerWaterPass(FRHICommandListImmediate& RHICmdList, FSingleLayerWaterPassData& PassData, FExclusiveDepthStencil::Type WaterPassDepthStencilAccess);
+	bool RenderSingleLayerWaterPass(FRHICommandListImmediate& RHICmdList, FSingleLayerWaterPassData& PassData, FExclusiveDepthStencil::Type WaterPassDepthStencilAccess, bool bParallel);
 	/** Renders the water draw pass for a given View. */
-	bool RenderSingleLayerWaterPassView(FRHICommandListImmediate& RHICmdList, FViewInfo& View, FSingleLayerWaterPassData& PassData, const FMeshPassProcessorRenderState& InDrawRenderState);
+	bool RenderSingleLayerWaterPassView(FRHICommandListImmediate& RHICmdList, FViewInfo& View, FSingleLayerWaterPassData& PassData, const FMeshPassProcessorRenderState& InDrawRenderState, bool bParallel);
 	/** Render, denoise and composite the scene SSR and under water effect.*/
 	void RenderSingleLayerWaterReflections(FRHICommandListImmediate& RHICmdList, FSingleLayerWaterPassData& PassData);
 
diff --git a/Engine/Source/Runtime/Renderer/Private/SingleLayerWaterRendering.cpp b/Engine/Source/Runtime/Renderer/Private/SingleLayerWaterRendering.cpp
--- a/Engine/Source/Runtime/Renderer/Private/SingleLayerWaterRendering.cpp
+++ b/Engine/Source/Runtime/Renderer/Private/SingleLayerWaterRendering.cpp
@@ -61,6 +61,19 @@ static TAutoConsoleVariable<int32> CVarWaterSingleLayerSSRTAA(
 	ECVF_RenderThreadSafe | ECVF_Scalability);
 
 
+static TAutoConsoleVariable<int32> CVarRHICmdSingleLayerWaterDeferredContexts(
+	TEXT("r.RHICmdSingleLayerWaterDeferredContexts"),
+	1,
+	TEXT("True to use deferred contexts to parallelize single layer water command list execution."));
+
+
+static TAutoConsoleVariable<int32> CVarRHICmdFlushRenderThreadTasksSingleLayerWater(
+	TEXT("r.RHICmdFlushRenderThreadTasksBasePass"),
+	0,
+	TEXT("Wait for completion of parallel render thread tasks at the end of Single layer water. A more granular version of r.RHICmdFlushRenderThreadTasks. If either r.RHICmdFlushRenderThreadTasks or r.RHICmdFlushRenderThreadTasksSingleLayerWater is > 0 we will flush."));
+
+
+
 
 //////////////////////////////////////////////////////////////////////////
 
@@ -363,6 +376,43 @@ IMPLEMENT_GLOBAL_SHADER(FWaterRefractionCopyPS, "/Engine/Private/SingleLayerWate
 
 //////////////////////////////////////////////////////////////////////////
 
+DECLARE_CYCLE_STAT(TEXT("WaterSingleLayer"), STAT_CLP_WaterSingleLayerPass, STATGROUP_ParallelCommandListMarkers); 
+
+class FWaterSingleLayerPassParallelCommandListSet : public FParallelCommandListSet
+{
+public:
+	FExclusiveDepthStencil::Type PassDepthStencilAccess;
+
+	FWaterSingleLayerPassParallelCommandListSet(
+		const FViewInfo& InView,
+		FRHICommandListImmediate& InParentCmdList,
+		bool bInParallelExecute,
+		bool bInCreateSceneContext,
+		const FSceneRenderer* InSceneRenderer,
+		FExclusiveDepthStencil::Type InPassDepthStencilAccess,
+		const FMeshPassProcessorRenderState& InDrawRenderState)
+		: FParallelCommandListSet(GET_STATID(STAT_CLP_WaterSingleLayerPass), InView, InSceneRenderer, InParentCmdList, bInParallelExecute, bInCreateSceneContext, InDrawRenderState)
+		, PassDepthStencilAccess(InPassDepthStencilAccess)
+	{
+	}
+
+	virtual ~FWaterSingleLayerPassParallelCommandListSet()
+	{
+		Dispatch();
+	}
+
+	virtual void SetStateOnCommandList(FRHICommandList& CmdList) override
+	{
+		FParallelCommandListSet::SetStateOnCommandList(CmdList);
+
+		FDeferredShadingSceneRenderer::BeginRenderingWaterGBuffer(CmdList, PassDepthStencilAccess, SceneRenderer->ViewFamily.EngineShowFlags.ShaderComplexity, SceneRenderer->ShaderPlatform);
+		SetupBasePassView(CmdList, View, SceneRenderer);
+	}
+};
+
+//////////////////////////////////////////////////////////////////////////
+
+
 
 
 void FDeferredShadingSceneRenderer::CopySingleLayerWaterTextures(FRHICommandListImmediate& RHICmdList, FSingleLayerWaterPassData& PassData)
@@ -411,7 +461,7 @@ void FDeferredShadingSceneRenderer::CopySingleLayerWaterTextures(FRHICommandList
 	GraphBuilder.Execute();
 }
 
-void FDeferredShadingSceneRenderer::BeginRenderingWaterGBuffer(FRHICommandList& RHICmdList, FSingleLayerWaterPassData& PassData, FExclusiveDepthStencil::Type DepthStencilAccess, bool bBindQuadOverdrawBuffers)
+void FDeferredShadingSceneRenderer::BeginRenderingWaterGBuffer(FRHICommandList& RHICmdList, FExclusiveDepthStencil::Type DepthStencilAccess, bool bBindQuadOverdrawBuffers, EShaderPlatform InShaderPlatform)
 {
 	SCOPED_DRAW_EVENT(RHICmdList, BeginRenderingWaterGBuffer);
 
@@ -419,7 +469,7 @@ void FDeferredShadingSceneRenderer::BeginRenderingWaterGBuffer(FRHICommandList&
 	FSceneRenderTargets& SceneContext = FSceneRenderTargets::Get(RHICmdList);
 
 	const ERHIFeatureLevel::Type CurrentFeatureLevel = SceneContext.GetCurrentFeatureLevel();
-	const bool bUseGBuffer = IsUsingGBuffers(ShaderPlatform);
+	const bool bUseGBuffer = IsUsingGBuffers(InShaderPlatform);
 	check(CurrentFeatureLevel >= ERHIFeatureLevel::SM5);
 
 	// Create MRT
@@ -657,9 +707,9 @@ void FDeferredShadingSceneRenderer::RenderSingleLayerWaterReflections(FRHIComman
 	ResolveSceneColor(RHICmdList);
 }
 
-bool FDeferredShadingSceneRenderer::RenderSingleLayerWaterPass(FRHICommandListImmediate& RHICmdList, FSingleLayerWaterPassData& PassData, FExclusiveDepthStencil::Type WaterPassDepthStencilAccess)
+bool FDeferredShadingSceneRenderer::RenderSingleLayerWaterPass(FRHICommandListImmediate& RHICmdList, FSingleLayerWaterPassData& PassData, FExclusiveDepthStencil::Type WaterPassDepthStencilAccess, bool bParallel)
 {
-	CSV_SCOPED_TIMING_STAT_EXCLUSIVE(RenderSingleLayerWaterPass);
+	CSV_SCOPED_TIMING_STAT_EXCLUSIVE(WaterPassRenderSingleLayer);
 	SCOPED_NAMED_EVENT(FDeferredShadingSceneRenderer_RenderSingleLayerWaterPass, FColor::Emerald);
 
 	bool bDirty = false;
@@ -671,8 +721,11 @@ bool FDeferredShadingSceneRenderer::RenderSingleLayerWaterPass(FRHICommandListIm
 		SCOPE_CYCLE_COUNTER(STAT_WaterPassDrawTime);
 		SCOPED_GPU_STAT(RHICmdList, SingleLayerWater);
 
-		// Must have an open renderpass before getting here in single threaded mode.
-		check(RHICmdList.IsInsideRenderPass());
+		if (!bParallel)
+		{
+			// Must have an open renderpass before getting here in single threaded mode.
+			check(RHICmdList.IsInsideRenderPass());
+		}
 
 		for (int32 ViewIndex = 0; ViewIndex < Views.Num(); ViewIndex++)
 		{
@@ -697,7 +750,7 @@ bool FDeferredShadingSceneRenderer::RenderSingleLayerWaterPass(FRHICommandListIm
 			{
 				Scene->UniformBuffers.UpdateViewUniformBuffer(View);
 
-				bDirty |= RenderSingleLayerWaterPassView(RHICmdList, View, PassData, DrawRenderState);
+				bDirty |= RenderSingleLayerWaterPassView(RHICmdList, View, PassData, DrawRenderState, bParallel);
 			}
 		}
 	}
@@ -708,15 +761,28 @@ bool FDeferredShadingSceneRenderer::RenderSingleLayerWaterPass(FRHICommandListIm
 	return bDirty;
 }
 
-bool FDeferredShadingSceneRenderer::RenderSingleLayerWaterPassView(FRHICommandListImmediate& RHICmdList, FViewInfo& View, FSingleLayerWaterPassData& PassData, const FMeshPassProcessorRenderState& InDrawRenderState)
+bool FDeferredShadingSceneRenderer::RenderSingleLayerWaterPassView(FRHICommandListImmediate& RHICmdList, FViewInfo& View, FSingleLayerWaterPassData& PassData, const FMeshPassProcessorRenderState& InDrawRenderState, bool bParallel)
 {
-	bool bDirty = false;
-	FMeshPassProcessorRenderState DrawRenderState(InDrawRenderState);
-	SetupBasePassView(RHICmdList, View, this);
+	if (!bParallel)
+	{
+	    SetupBasePassView(RHICmdList, View, this);
+		View.ParallelMeshDrawCommandPasses[EMeshPass::SingleLayerWaterPass].DispatchDraw(nullptr, RHICmdList);
+	}
+	else
+	{
+		FWaterSingleLayerPassParallelCommandListSet ParallelSet
+		(
+			View,
+			RHICmdList,
+			CVarRHICmdSingleLayerWaterDeferredContexts.GetValueOnRenderThread() > 0,
+			CVarRHICmdFlushRenderThreadTasksSingleLayerWater.GetValueOnRenderThread() == 0 && CVarRHICmdFlushRenderThreadTasks.GetValueOnRenderThread() == 0,
+			this,
+			FExclusiveDepthStencil::DepthWrite_StencilWrite,
+			InDrawRenderState
+		);
 
-	View.ParallelMeshDrawCommandPasses[EMeshPass::SingleLayerWaterPass].DispatchDraw(nullptr, RHICmdList);
+		View.ParallelMeshDrawCommandPasses[EMeshPass::SingleLayerWaterPass].DispatchDraw(&ParallelSet, RHICmdList);
+	}
 
-	return bDirty;
+	return View.ParallelMeshDrawCommandPasses[EMeshPass::SingleLayerWaterPass].HasAnyDraw();
 }
-
-
