Fixed console object name that duplicates an existing one

#rnx

diff --git a/Engine/Source/Runtime/Renderer/Private/SingleLayerWaterRendering.cpp b/Engine/Source/Runtime/Renderer/Private/SingleLayerWaterRendering.cpp
--- a/Engine/Source/Runtime/Renderer/Private/SingleLayerWaterRendering.cpp
+++ b/Engine/Source/Runtime/Renderer/Private/SingleLayerWaterRendering.cpp
@@ -68,7 +68,7 @@ static TAutoConsoleVariable<int32> CVarRHICmdSingleLayerWaterDeferredContexts(
 
 
 static TAutoConsoleVariable<int32> CVarRHICmdFlushRenderThreadTasksSingleLayerWater(
-	TEXT("r.RHICmdFlushRenderThreadTasksBasePass"),
+	TEXT("r.RHICmdFlushRenderThreadTasksSingleLayerWater"),
 	0,
 	TEXT("Wait for completion of parallel render thread tasks at the end of Single layer water. A more granular version of r.RHICmdFlushRenderThreadTasks. If either r.RHICmdFlushRenderThreadTasks or r.RHICmdFlushRenderThreadTasksSingleLayerWater is > 0 we will flush."));
 
