Updated GTAO code.
Runs at half res with explicit spatial and temporal filters.

diff --git a/Engine/Shaders/Private/PostProcessAmbientOcclusion.usf b/Engine/Shaders/Private/PostProcessAmbientOcclusion.usf
--- a/Engine/Shaders/Private/PostProcessAmbientOcclusion.usf
+++ b/Engine/Shaders/Private/PostProcessAmbientOcclusion.usf
@@ -954,6 +954,17 @@ float4 GTAOParams[2];
 RWTexture2D<float2> HorizonOutTexture;
 #endif
 
+Texture2D		HistoryTexture;
+SamplerState	HistoryTextureSampler;
+
+Texture2D		DepthHistoryTexture;
+SamplerState	DepthHistoryTextureSampler;
+
+
+float4 PrevScreenPositionScaleBias;
+float4 FilterParams;
+
+
 
 float Falloff(float distSq)
 {
@@ -967,16 +978,6 @@ float Falloff(float distSq)
 //	return  2.0 * saturate(distSq* Falloff_Scale + Falloff_Bias   );
 }
 
-
-float3 ScreenToViewPosX(float2 UV, float SceneDepth)
-{
-	float3 ViewPos = ScreenToViewPos(UV, SceneDepth);
-
-	//return float3( ViewPos.x, ViewPos.z, ViewPos.y);
-	return float3( ViewPos.x, ViewPos.y, ViewPos.z);
-}
-
-
 float3 GetViewSpacePosFromHZB(float2 UV, float MipLevel)
 {
 	float SceneDepth = GetDepthFromAOInput(UV);
@@ -987,7 +988,7 @@ float3 GetViewSpacePosFromHZB(float2 UV, float MipLevel)
 	SceneDepth = Texture2DSampleLevel( PostprocessInput1, PostprocessInput1Sampler, HZBUV, MipLevel).r;
 	SceneDepth = ConvertFromDeviceZ(SceneDepth);
 
-	return ScreenToViewPosX(UV, SceneDepth);
+	return ScreenToViewPos(UV, SceneDepth);
 }
 
 
@@ -1012,8 +1013,8 @@ float3 GetNormal(float2 UV, float3 ViewSpacePosMid, bool bConstructFromDepth)
 		float ZRight		= ConvertFromDeviceZ(DeviceZ + DeviceZDdx);
 		float ZDown			= ConvertFromDeviceZ(DeviceZ + DeviceZDdy);
 
-		float3 Right		= ScreenToViewPosX(UV+ float2( DeltaUV.x,      0.0f) , ZRight)-ViewSpacePosMid;
-		float3 Down			= ScreenToViewPosX(UV+ float2(      0.0f, DeltaUV.y) , ZDown) -ViewSpacePosMid;
+		float3 Right		= ScreenToViewPos(UV+ float2( DeltaUV.x,      0.0f) , ZRight)-ViewSpacePosMid;
+		float3 Down			= ScreenToViewPos(UV+ float2(      0.0f, DeltaUV.y) , ZDown) -ViewSpacePosMid;
 		
 		ViewSpaceNormal = normalize(cross(Right, Down));
 
@@ -1067,7 +1068,7 @@ float SearchForLargestAngleRef(float2 inUV, float2 deltaUV, float3 ViewSpacePosi
 	// Precompute as much stuff upfront
 	float CurrDepth			= GetLinearDepthProj(UV);
 	UV						+= deltaUV;
-	float2 DeltaVS_XY	    = ScreenToViewPosX(UV, CurrDepth).xy - ViewSpacePosition.xy;
+	float2 DeltaVS_XY	    = ScreenToViewPos(UV, CurrDepth).xy - ViewSpacePosition.xy;
 	float  DeltaVS_XYLen    = sqrt(dot( DeltaVS_XY.xy, DeltaVS_XY.xy));
 	float  CurrVS_XYLen		= DeltaVS_XYLen;
 
@@ -1103,7 +1104,7 @@ float SearchForLargestAngleRef(float2 inUV, float2 deltaUV, float3 ViewSpacePosi
 		UV					+= deltaUV;
 	
 		float SceneDepth	= GetDepthFromAOInput(UV);
-		float3 SampVS		= ScreenToViewPosX(UV, SceneDepth);
+		float3 SampVS		= ScreenToViewPos(UV, SceneDepth);
 	
 		float3 DiffV		= (SampVS - ViewSpacePosition) ;
 		float distSQ		= dot(DiffV, DiffV);
@@ -1150,7 +1151,7 @@ float SearchForLargestAngle(float2 inUV, float2 deltaUV, float3 ViewSpacePositio
 	// Precompute as much stuff upfront
 	float CurrDepth			= GetLinearDepthProj(UV);
 	UV						+= deltaUV;
-	float2 DeltaVS_XY	    = ScreenToViewPosX(UV, CurrDepth).xy - ViewSpacePosition.xy;
+	float2 DeltaVS_XY	    = ScreenToViewPos(UV, CurrDepth).xy - ViewSpacePosition.xy;
 	float  DeltaVS_XYLen    = sqrt(dot( DeltaVS_XY.xy, DeltaVS_XY.xy));
 	float  CurrVS_XYLen		= DeltaVS_XYLen;
 
@@ -1319,36 +1320,43 @@ float ComputeInnerIntegral(float2 Angles, float3 ScreenDir, float3 ViewDir, floa
 
 /*
 *
-* HORIZON SEARCH AND INNER INTEGRATE
+* HORIZON SEARCH AND INNER INTEGRATE COMBINED
 *
 */
 
-void HorizonSearchAndInnerIntegratePSandCS(in float2 UV, out float4 OutColor)
+float2 GetRandomVector(float2 UV )
 {
-	OutColor = 0;
-
-	float SceneDepth	 = GetDepthFromAOInput(UV);		
-	if(SceneDepth > ScreenSpaceAOParams[4].w)	
-	{		
-		OutColor = 1;
-		return;
-	}
-	
 	float2 ViewportUVToRandomUV = ScreenSpaceAOParams[1].xy;
 	float TemporalCos = GTAOParams[0].x;
 	float TemporalSin = GTAOParams[0].y;
-	float2 TemporalOffset = 0;// GTAOParams[1].xy;
-	float2 RandomTexVec = (Texture2DSample(RandomNormalTexture, RandomNormalTextureSampler, TemporalOffset + UV * ViewportUVToRandomUV).rg * 2 - 1) ;
+	float2 RandomTexVec = (Texture2DSample(RandomNormalTexture, RandomNormalTextureSampler, UV * ViewportUVToRandomUV).rg * 2 - 1) ;
 	float2 RandomVec;
-	
+
 	// Rotate the random Vec
 	RandomVec.x = dot(RandomTexVec.xy, float2(TemporalCos, -TemporalSin ));
 	RandomVec.y = dot(RandomTexVec.xy, float2(TemporalSin, TemporalCos ));
 
-	float3 ViewSpacePos  = ScreenToViewPosX(UV,SceneDepth);
+	return RandomVec;
+}
+
+
+
+void GTAOCombinedPSandCS(in float2 UV, out float4 OutColor)
+{
+	OutColor = 0;
+
+	float SceneDepth = GetDepthFromAOInput(UV);		
+	if(SceneDepth > ScreenSpaceAOParams[4].w)	
+	{		
+		OutColor = 1;
+		return;
+	}
+	
+	float2 RandomVec		= GetRandomVector(UV);
+	float3 ViewSpacePos		= ScreenToViewPos(UV,SceneDepth);
 
 	// Get Normal and transform into viewspace
-	float3 ViewSpaceNormal = GetNormal(UV, ViewSpacePos, false);
+	float3 ViewSpaceNormal	= GetNormal(UV, ViewSpacePos, false);
 #if GTAO_OPT
 	float3 ViewDir			= float3(0,0,-1);
 #else
@@ -1360,9 +1368,6 @@ void HorizonSearchAndInnerIntegratePSandCS(in float2 UV, out float4 OutColor)
 	// Given the depth determine how many pixels away we are guranteed to be outside of the falloff distance.
 	float InvTanHalfFov			= ScreenSpaceAOParams[3].w;
 	float ScreenPixelsToCutoff  = View.BufferSizeAndInvSize.x * ((FALLOFF_END*0.5) / ViewSpacePos.z) ; 
-
-	// Get radius based on sampling 
-	// Get radius based on distance - TODO - Need to use FOV and aspect ratio
 	float Scale = ScreenPixelsToCutoff /  GTAO_NUMTAPS;
 
 	// Clamp the scale. Too small and it won't be jumping by more than 1 pixel. Too much and it'll thrash the cache
@@ -1388,7 +1393,7 @@ void HorizonSearchAndInnerIntegratePSandCS(in float2 UV, out float4 OutColor)
 #endif
 
 	float AmbientOcclusionIntensity = ScreenSpaceAOParams[0].w;
-	float AmbientOcclusionPower = ScreenSpaceAOParams[0].x;
+	float AmbientOcclusionPower		= ScreenSpaceAOParams[0].x;
 	// user adjust AO
 	Sum = Sum   / (float)GTAO_NUMROTS;
 	Sum = 1 - (1 - pow(abs(Sum), AmbientOcclusionPower)) * AmbientOcclusionIntensity;
@@ -1399,15 +1404,15 @@ void HorizonSearchAndInnerIntegratePSandCS(in float2 UV, out float4 OutColor)
 }
 
 
-void HorizonSearchAndInnerIntegratePS(in noperspective float4 UVAndScreenPos : TEXCOORD0, out float4 OutColor : SV_Target0)
+void GTAOCombinedPS(in noperspective float4 UVAndScreenPos : TEXCOORD0, out float4 OutColor : SV_Target0)
 {
-	HorizonSearchAndInnerIntegratePSandCS(UVAndScreenPos.xy, OutColor);	
+	GTAOCombinedPSandCS(UVAndScreenPos.xy, OutColor);	
 }
 
 
 #if COMPUTE_SHADER
 [numthreads(THREADGROUP_SIZEX, THREADGROUP_SIZEY, 1)]
-void HorizonSearchAndInnerIntegrateCS(
+void GTAOCombinedCS(
 	uint2 GroupId : SV_GroupID,
 	uint2 DispatchThreadId : SV_DispatchThreadID,
 	uint2 GroupThreadId : SV_GroupThreadID) 
@@ -1417,7 +1422,7 @@ void HorizonSearchAndInnerIntegrateCS(
 	float2 PixelCenter	= (float2)PixelPos + float2(0.5, 0.5);
 	float2 BufferUV		= SvPositionToBufferUV(float4(PixelCenter, 0, 0));
 
-	HorizonSearchAndInnerIntegratePSandCS(BufferUV , OutColor);
+	GTAOCombinedPSandCS(BufferUV , OutColor);
 
 	OutTexture[PixelPos] = OutColor.r;
 }
@@ -1428,25 +1433,22 @@ void HorizonSearchAndInnerIntegrateCS(
 
 /*
 *
-*
 *	JUST INNER INTEGRATE
 *
-*
 */
 
 Texture2D		HorizonsTexture;
 SamplerState	HorizonsTextureSampler;
+float4			InnerIntegrateParams;
 
-void GTAOInnerIntegratePSandCS(in float2 UV, out float4 OutColor)
+
+float GTAOInnerIntegratePSandCS(in float2 UV)
 {
 	// Read the angles buffer
-	OutColor = 0;
-
 	float SceneDepth	 = GetDepthFromAOInput(UV);		
 	if(SceneDepth > ScreenSpaceAOParams[4].w)	
 	{		
-		OutColor = 1;
-		return;
+		return 1;
 	}
 
 	float2 Angles			= Texture2DSample(HorizonsTexture, HorizonsTextureSampler, UV).xy;	// Angles computed from previous pass
@@ -1459,7 +1461,7 @@ void GTAOInnerIntegratePSandCS(in float2 UV, out float4 OutColor)
 	float3 ScreenDir			= float3(RandomVec.x, RandomVec.y, 0.0);
 
 	// ViewspacePos and Normal
-	float3 ViewSpacePos		= ScreenToViewPosX(UV, SceneDepth);
+	float3 ViewSpacePos		= ScreenToViewPos(UV, SceneDepth);
 	float3 WorldNormal		= GetGBufferData(UV, false).WorldNormal;
 	float3 ViewSpaceNormal	= normalize(mul(WorldNormal, (float3x3)View.TranslatedWorldToView));
 	float3 ViewDir			= -normalize(ViewSpacePos.xyz);	// TODO - This is a function of UV only.
@@ -1474,16 +1476,15 @@ void GTAOInnerIntegratePSandCS(in float2 UV, out float4 OutColor)
 	// user adjust AO
 	AO = 1 - (1 - pow(abs(AO), AmbientOcclusionPower)) * AmbientOcclusionIntensity;
 
-	OutColor.rgb = AO ;
-
-	return;
+	return AO ;
 }
 
 
 
 void GTAOInnerIntegratePS(in noperspective float4 UVAndScreenPos : TEXCOORD0, out float4 OutColor : SV_Target0)
 {
-	GTAOInnerIntegratePSandCS(UVAndScreenPos.xy, OutColor);	
+	float AO = 	GTAOInnerIntegratePSandCS(UVAndScreenPos.xy);	
+	OutColor = AO;
 }
 
 
@@ -1494,14 +1495,17 @@ void GTAOInnerIntegrateCS(
 	uint2 DispatchThreadId : SV_DispatchThreadID,
 	uint2 GroupThreadId : SV_GroupThreadID) 
 {
-	float4 OutColor		= 0;
+
+	float DownSampleFactor = InnerIntegrateParams.x;
+	float  OutColor		= 0;
 	int2   PixelPos		= DispatchThreadId + ScreenSpaceAOParams[5].zw; 
 	float2 PixelCenter	= (float2)PixelPos + float2(0.5, 0.5);
+	PixelCenter		   *= DownSampleFactor;
 	float2 BufferUV		= SvPositionToBufferUV(float4(PixelCenter, 0, 0));
 
-	GTAOInnerIntegratePSandCS(BufferUV , OutColor);
+	float AO = GTAOInnerIntegratePSandCS(BufferUV);
 	
-	OutTexture[PixelPos] = OutColor.r;
+	OutTexture[PixelPos] = AO;
 }
 #endif
 
@@ -1516,6 +1520,9 @@ void GTAOInnerIntegrateCS(
 *
 */
 
+float4 HorizonSearchParams;
+
+
 float2 HorizonSearchPSandCS(in float2 UV)
 {
 	float2 OutHorizons = 0;
@@ -1527,17 +1534,20 @@ float2 HorizonSearchPSandCS(in float2 UV)
 		return OutHorizons;
 	}
 
-	float2 ViewportUVToRandomUV = ScreenSpaceAOParams[1].xy;
-	const float2 TemporalOffset = ScreenSpaceAOParams[3].xy;
-	float2 RandomVec			= (Texture2DSample(RandomNormalTexture, RandomNormalTextureSampler, TemporalOffset + UV * ViewportUVToRandomUV).rg * 2 - 1) ;
+	float2 RandomVec			= GetRandomVector(UV);
 	float3 ScreenDir			= float3(RandomVec.x, RandomVec.y, 0.0);
 	
-	float3 ViewSpacePos			= ScreenToViewPosX(UV, SceneDepth);
+	float3 ViewSpacePos			= ScreenToViewPos(UV, SceneDepth);
 	float3 ViewDir				= -normalize(ViewSpacePos.xyz);
 
-	// Get radius based on distance - TODO - Need to use FOV and aspect ratio
-	float Scale = 1000.0 / ViewSpacePos.z;
-	Scale = max(Scale, 1.0f);
+	// Given the depth determine how many pixels away we are guranteed to be outside of the falloff distance.
+	float InvTanHalfFov			= ScreenSpaceAOParams[3].w;
+	float ScreenPixelsToCutoff  = View.BufferSizeAndInvSize.x * ((FALLOFF_END*0.5) / ViewSpacePos.z) ; 
+	float Scale = ScreenPixelsToCutoff /  GTAO_NUMTAPS;
+
+	// Clamp the scale. Too small and it won't be jumping by more than 1 pixel. Too much and it'll thrash the cache
+	Scale = clamp(Scale, 2.0, 8.0);
+
 
 	float2 Angles = SearchAxisForAngles( UV,  ScreenDir,  ViewDir,  ViewSpacePos, Scale);
 	OutHorizons.rg = Angles*0.5+0.5;
@@ -1549,8 +1559,7 @@ float2 HorizonSearchPSandCS(in float2 UV)
 void HorizonSearchPS(in noperspective float4 UVAndScreenPos : TEXCOORD0, float4 SvPosition : SV_POSITION, out float4 OutColor : SV_Target0)
 {
 	float2 Horizons = HorizonSearchPSandCS(UVAndScreenPos.xy);	
-	OutColor = float4(Horizons.xy,0,0);;
-
+	OutColor = float4(Horizons.xy,0,0);
 }
 
 
@@ -1561,9 +1570,12 @@ void HorizonSearchCS(
 	uint2 DispatchThreadId : SV_DispatchThreadID,
 	uint2 GroupThreadId : SV_GroupThreadID) 
 {
+	float DownSampleFactor = HorizonSearchParams.x;
+
 	float2 OutColor		= 0;
 	int2   PixelPos		= DispatchThreadId + ScreenSpaceAOParams[5].zw; 
 	float2 PixelCenter	= (float2)PixelPos + float2(0.5, 0.5);
+	PixelCenter		   *= DownSampleFactor;
 	float2 BufferUV		= SvPositionToBufferUV(float4(PixelCenter, 0, 0));
 
 	float2 Horizons = HorizonSearchPSandCS(BufferUV);
@@ -1571,3 +1583,163 @@ void HorizonSearchCS(
 	HorizonOutTexture[PixelPos] = Horizons;
 }
 #endif
+
+
+
+
+float2 ReprojectPos(float2 UV, float Depth)
+{
+	// Given a UV reproject where this was in the previous frame
+	// Camera motion for pixel (in ScreenPos space).
+	float2 ThisScreen = (UV.xy - View.ScreenPositionScaleBias.wz) / View.ScreenPositionScaleBias.xy;
+	float4 ThisClip = float4( ThisScreen, Depth, 1 );
+	float4 PrevClip = mul( ThisClip, View.ClipToPrevClip );
+	float2 PrevScreen = PrevClip.xy / PrevClip.w;
+	float2 PrevUV = PrevScreen.xy * PrevScreenPositionScaleBias.xy + PrevScreenPositionScaleBias.zw;
+
+	PrevUV += PostprocessInput0Size.zw*0.25;
+
+
+	// TODO - Take into account velocity
+//	float2 Velocity = Texture.SampleLevel( Sampler, UV, 0 ).xy;
+//	if( Velocity.x > 0.0 )
+//	{
+//		PrevScreen = ThisClip.xy - DecodeVelocityFromTexture( Velocity );
+//	}
+
+
+	return PrevUV;
+}
+
+void GTAOTemporalFilterPSAndCS(float2 UV, out float OutAO, out float OutDepth)
+{
+	float NewAO			= Texture2DSample(PostprocessInput0, PostprocessInput0Sampler, UV).r;
+
+	float SceneDepth	= GetDepthFromAOInput(UV);	// FIXME - Get from HZB Buffer
+	float DeviceZ		= LookupDeviceZ( UV);
+
+	float2 PrevUV		= ReprojectPos( UV,  DeviceZ);
+	float  HistoryAO	= Texture2DSample(HistoryTexture, HistoryTextureSampler, PrevUV).r;
+
+	float CurrDepth		= SceneDepth;
+	float PrevDepth		= Texture2DSample(DepthHistoryTexture, DepthHistoryTextureSampler, PrevUV).r;
+
+	float Blend			= 0.15;
+
+	if( abs(CurrDepth-PrevDepth) > 100)
+		Blend =1;
+
+	if((PrevUV.x <= 0.0) || (PrevUV.y <= 0.0) || (PrevUV.x >= 1.0) || (PrevUV.y >= 1.0))
+	{
+		Blend=1;
+	}
+
+	// Do a neighbourhood clamp 
+	float2 dUV = PostprocessInput0Size.zw*2;
+	float AONeighbour1 = Texture2DSample(PostprocessInput0, PostprocessInput0Sampler, UV + float2(-dUV.x,-dUV.y) ).r;
+	float AONeighbour2 = Texture2DSample(PostprocessInput0, PostprocessInput0Sampler, UV + float2(-dUV.x, dUV.y) ).r;
+	float AONeighbour3 = Texture2DSample(PostprocessInput0, PostprocessInput0Sampler, UV + float2( dUV.x,-dUV.y) ).r;
+	float AONeighbour4 = Texture2DSample(PostprocessInput0, PostprocessInput0Sampler, UV + float2( dUV.x, dUV.y) ).r;
+
+	float MinAO = min(NewAO, min(min(AONeighbour1, AONeighbour2), min(AONeighbour3, AONeighbour4)));
+	float MaxAO = max(NewAO, max(max(AONeighbour1, AONeighbour2), max(AONeighbour3, AONeighbour4)));
+
+	float BlurredAO = 0.2 * (NewAO + AONeighbour1 + AONeighbour2 + AONeighbour3 + AONeighbour4);
+
+	OutAO = (NewAO*Blend) + (HistoryAO* (1-Blend));
+	OutAO = clamp(OutAO, MinAO, MaxAO);
+
+
+	OutDepth = CurrDepth;
+}
+
+void GTAOTemporalFilterPS(in noperspective float4 UVAndScreenPos : TEXCOORD0, float4 SvPosition : SV_POSITION, out float4 OutColor : SV_Target0, out float4 OutColor1 : SV_Target1)
+{
+	float OutAO, OutDepth;
+	
+	GTAOTemporalFilterPSAndCS(UVAndScreenPos.xy, OutAO, OutDepth);
+	OutColor	= OutAO;
+	OutColor1 = OutDepth;
+}
+
+#if COMPUTE_SHADER
+[numthreads(THREADGROUP_SIZEX, THREADGROUP_SIZEY, 1)]
+void GTAOTemporalFilterCS(
+	uint2 GroupId : SV_GroupID,
+	uint2 DispatchThreadId : SV_DispatchThreadID,
+	uint2 GroupThreadId : SV_GroupThreadID) 
+{
+	float OutColor		= 0;
+	float OutDepth		= 0;
+	int2   PixelPos		= DispatchThreadId + ScreenSpaceAOParams[5].zw; 
+	float2 PixelCenter	= (float2)PixelPos + float2(0.5, 0.5);
+	float2 BufferUV		= SvPositionToBufferUV(float4(PixelCenter, 0, 0));
+
+	GTAOTemporalFilterPSAndCS(BufferUV, OutColor, OutDepth);
+
+	OutTexture[PixelPos] = OutColor;
+}
+#endif
+
+void GTAOSpatialFilterPS(in noperspective float4 UVAndScreenPos : TEXCOORD0,  out float4 OutColor : SV_Target0)
+{
+	OutColor = Texture2DSample(PostprocessInput0, PostprocessInput0Sampler, UVAndScreenPos.xy);
+}
+
+
+#if COMPUTE_SHADER
+[numthreads(THREADGROUP_SIZEX, THREADGROUP_SIZEY, 1)]
+void GTAOSpatialFilterCS(
+	uint2 GroupId : SV_GroupID,
+	uint2 DispatchThreadId : SV_DispatchThreadID,
+	uint2 GroupThreadId : SV_GroupThreadID) 
+{
+	BRANCH
+	if (any(DispatchThreadId >= (uint2)FilterParams.zw))
+	{
+		return;
+	}
+
+	uint2 DestPixelPos = FilterParams.xy + DispatchThreadId;
+	float2 SampleUV = DestPixelPos * PostprocessInput0Size.zw;
+
+	// Use a 4x4 box filter because the random texture is tiled 4x4
+	float Result;
+	Result = PostprocessInput0.SampleLevel(PostprocessInput0Sampler, SampleUV, 0).r;
+	Result += PostprocessInput0.SampleLevel(PostprocessInput0Sampler, SampleUV, 0, int2(2, 0)).r;
+	Result += PostprocessInput0.SampleLevel(PostprocessInput0Sampler, SampleUV, 0, int2(0, 2)).r;
+	Result += PostprocessInput0.SampleLevel(PostprocessInput0Sampler, SampleUV, 0, int2(2, 2)).r;
+
+	OutTexture[DestPixelPos] = Result * 0.25;
+}
+#endif
+
+
+
+float4 GTAOUpsamplePSAndCS(float2 UV)
+{
+	return Texture2DSample(PostprocessInput0, PostprocessInput0Sampler, UV);
+}
+
+
+void GTAOUpsamplePS(in noperspective float4 UVAndScreenPos : TEXCOORD0,  out float4 OutColor : SV_Target0)
+{
+	OutColor = GTAOUpsamplePSAndCS(UVAndScreenPos.xy);
+}
+
+#if COMPUTE_SHADER
+[numthreads(THREADGROUP_SIZEX, THREADGROUP_SIZEY, 1)]
+void GTAOUpsampleCS(
+	uint2 GroupId : SV_GroupID,
+	uint2 DispatchThreadId : SV_DispatchThreadID,
+	uint2 GroupThreadId : SV_GroupThreadID) 
+{
+	float2 OutColor		= 0;
+	int2   PixelPos		= DispatchThreadId + ScreenSpaceAOParams[5].zw; 
+	float2 PixelCenter	= (float2)PixelPos + float2(0.5, 0.5);
+	float2 BufferUV		= SvPositionToBufferUV(float4(PixelCenter, 0, 0));
+
+	float4 Color = GTAOUpsamplePSAndCS(BufferUV);
+	OutTexture[PixelPos] = Color;
+}
+#endif
\ No newline at end of file
diff --git a/Engine/Source/Runtime/Renderer/Private/CompositionLighting/CompositionLighting.cpp b/Engine/Source/Runtime/Renderer/Private/CompositionLighting/CompositionLighting.cpp
--- a/Engine/Source/Runtime/Renderer/Private/CompositionLighting/CompositionLighting.cpp
+++ b/Engine/Source/Runtime/Renderer/Private/CompositionLighting/CompositionLighting.cpp
@@ -32,6 +32,14 @@ static TAutoConsoleVariable<int32> CVarSSAOSmoothPass(
 	TEXT("Whether to smooth SSAO output when TAA is disabled"),
 	ECVF_RenderThreadSafe | ECVF_Scalability);
 
+static TAutoConsoleVariable<int32> CVarGTAODownsample(
+	TEXT("r.GTAO.Downsample"),
+	1,
+	TEXT("Perform GTAO at Halfres \n ")
+	TEXT("0: Off \n ")
+	TEXT("1: On (default)\n "),
+	ECVF_RenderThreadSafe | ECVF_Scalability);
+
 bool IsAmbientCubemapPassRequired(const FSceneView& View)
 {
 	FScene* Scene = (FScene*)View.Family->Scene;
@@ -110,9 +118,27 @@ static FRenderingCompositeOutputRef AddPostProcessingGTAOAsyncHorizonSearch(FRHI
 {
 	FRenderingCompositePass* FinalOutputPass;
 
+
+	FSceneRenderTargets& SceneContext = FSceneRenderTargets::Get(Context.RHICmdList);
+
+	FIntPoint BufferSize	 = SceneContext.GetBufferSizeXY();
+	FIntPoint HalfBufferSize = FIntPoint::DivideAndRoundUp(BufferSize, 2);
+	FPooledRenderTargetDesc Desc(FPooledRenderTargetDesc::Create2DDesc(HalfBufferSize, PF_R8G8, FClearValueBinding::White, TexCreate_None, TexCreate_RenderTargetable, false));
+	if (SceneContext.GetCurrentFeatureLevel() >= ERHIFeatureLevel::SM5)
+	{
+		Desc.TargetableFlags |= TexCreate_UAV;
+	}
+	GRenderTargetPool.FindFreeElement(RHICmdList, Desc, SceneContext.ScreenSpaceGTAOHorizons, TEXT("ScreenSpaceGTAOHorizons"));
+
+	uint32 DownsampleFactor = 1;
+
+	if (CVarGTAODownsample.GetValueOnRenderThread() > 0)
+	{
+		DownsampleFactor = 2;
+	}
+	
 	FRenderingCompositePass* HZBInput = Context.Graph.RegisterPass(new FRCPassPostProcessInput(const_cast<FViewInfo&>(Context.View).HZB));
-		
-	FRenderingCompositePass* AmbientOcclusionHorizonSearch = Context.Graph.RegisterPass(new(FMemStack::Get()) FRCPassPostProcessAmbientOcclusion_HorizonSearch(Context.View, ESSAOType::EAsyncCS));
+	FRenderingCompositePass* AmbientOcclusionHorizonSearch = Context.Graph.RegisterPass(new(FMemStack::Get()) FRCPassPostProcessAmbientOcclusion_HorizonSearch(Context.View, DownsampleFactor, ESSAOType::EAsyncCS));
 
 	AmbientOcclusionHorizonSearch->SetInput(ePId_Input0, Context.SceneDepth);
 	AmbientOcclusionHorizonSearch->SetInput(ePId_Input1, HZBInput);
@@ -128,24 +154,55 @@ static FRenderingCompositeOutputRef AddPostProcessingGTAOCombined(FRHICommandLis
 {
 	FRenderingCompositePass* FinalOutputPass;
 
-	const bool bNeedSmoothingPass = CVarSSAOSmoothPass.GetValueOnRenderThread()!=0;
-
 	FRenderingCompositePass* HZBInput = Context.Graph.RegisterPass(new FRCPassPostProcessInput(const_cast<FViewInfo&>(Context.View).HZB));
 
-	FRenderingCompositePass* AmbientOcclusionGTAO = Context.Graph.RegisterPass(new(FMemStack::Get()) FRCPassPostProcessAmbientOcclusion_HorizonSearchInnerIntegrate(Context.View, !bNeedSmoothingPass));
-	AmbientOcclusionGTAO->SetInput(ePId_Input0, Context.SceneDepth);
-	AmbientOcclusionGTAO->SetInput(ePId_Input1, HZBInput);
-	FinalOutputPass = AmbientOcclusionGTAO;
+	uint32 DownsampleFactor = 1;
+
+	if (CVarGTAODownsample.GetValueOnRenderThread() > 0)
+	{
+		DownsampleFactor = 2;
+	}
+
+	{
+		FRenderingCompositePass* AmbientOcclusionGTAO = Context.Graph.RegisterPass(new(FMemStack::Get()) FRCPassPostProcessAmbientOcclusion_GTAOCombined(Context.View, DownsampleFactor, false));
+		AmbientOcclusionGTAO->SetInput(ePId_Input0, Context.SceneDepth);
+		AmbientOcclusionGTAO->SetInput(ePId_Input1, HZBInput);
+		FinalOutputPass = AmbientOcclusionGTAO;
+	}
 
 	FSceneRenderTargets& SceneContext = FSceneRenderTargets::Get(RHICmdList);
 	SceneContext.bScreenSpaceAOIsValid = true;
 
-	if (bNeedSmoothingPass)
+	FSceneViewState* ViewState = Context.View.ViewState;
+
+	// Add spatial Filter
 	{
-		FRenderingCompositePass* SSAOSmoothPass;
-		SSAOSmoothPass = Context.Graph.RegisterPass(new (FMemStack::Get()) FRCPassPostProcessAmbientOcclusionSmooth(ESSAOType::ECS, true));
-		SSAOSmoothPass->SetInput(ePId_Input0, FinalOutputPass);
-		FinalOutputPass = SSAOSmoothPass;
+		FRenderingCompositePass* SpatialPass;
+		SpatialPass = Context.Graph.RegisterPass(new (FMemStack::Get()) FRCPassPostProcessAmbientOcclusion_GTAO_SpatialFilter(Context.View));
+		SpatialPass->SetInput(ePId_Input0, FinalOutputPass);
+		SpatialPass->SetInput(ePId_Input1, HZBInput);
+		FinalOutputPass = SpatialPass;
+	}
+
+	if (ViewState)
+	{
+		// Add temporal filter
+		FRenderingCompositePass* TemporalPass;
+		TemporalPass = Context.Graph.RegisterPass(new (FMemStack::Get()) FRCPassPostProcessAmbientOcclusion_GTAO_TemporalFilter(Context.View, DownsampleFactor,
+			Context.View.PrevViewInfo.GTAOHistory,
+			&ViewState->PrevFrameViewInfo.GTAOHistory));
+
+		TemporalPass->SetInput(ePId_Input0, FinalOutputPass);
+		FinalOutputPass = TemporalPass;
+
+	}
+		
+
+	{
+		FRenderingCompositePass* UpsamplePass;
+		UpsamplePass = Context.Graph.RegisterPass(new (FMemStack::Get()) FRCPassPostProcessAmbientOcclusion_GTAO_Upsample());
+		UpsamplePass->SetInput(ePId_Input0, FinalOutputPass);
+		FinalOutputPass = UpsamplePass;
 	}
 
 	Context.FinalOutput = FRenderingCompositeOutputRef(FinalOutputPass);
@@ -157,21 +214,52 @@ static FRenderingCompositeOutputRef AddPostProcessingGTAOIntegration(FRHICommand
 {
 	FRenderingCompositePass* FinalOutputPass;
 
-	const bool bNeedSmoothingPass = CVarSSAOSmoothPass.GetValueOnRenderThread()!=0;
+	FRenderingCompositePass* HZBInput = Context.Graph.RegisterPass(new FRCPassPostProcessInput(const_cast<FViewInfo&>(Context.View).HZB));
+
+	uint32 DownsampleFactor = 1;
+
+	if (CVarGTAODownsample.GetValueOnRenderThread() > 0)
+	{
+		DownsampleFactor = 2;
+	}
 
-	FRenderingCompositePass* AmbientOcclusionInnerIntegrate = Context.Graph.RegisterPass(new(FMemStack::Get()) FRCPassPostProcessAmbientOcclusion_InnerIntegrate(Context.View, !bNeedSmoothingPass));
+	FRenderingCompositePass* AmbientOcclusionInnerIntegrate = Context.Graph.RegisterPass(new(FMemStack::Get()) FRCPassPostProcessAmbientOcclusion_InnerIntegrate(Context.View, DownsampleFactor, false));
 	AmbientOcclusionInnerIntegrate->SetInput(ePId_Input0, Context.SceneDepth);
 	FinalOutputPass = AmbientOcclusionInnerIntegrate;
 
 	FSceneRenderTargets& SceneContext = FSceneRenderTargets::Get(RHICmdList);
 	SceneContext.bScreenSpaceAOIsValid = true;
 
-	if (bNeedSmoothingPass)
+	FSceneViewState* ViewState = Context.View.ViewState;
+
+
+	// Add spatial Filter
 	{
-		FRenderingCompositePass* SSAOSmoothPass;
-		SSAOSmoothPass = Context.Graph.RegisterPass(new (FMemStack::Get()) FRCPassPostProcessAmbientOcclusionSmooth(ESSAOType::ECS, true));
-		SSAOSmoothPass->SetInput(ePId_Input0, FinalOutputPass);
-		FinalOutputPass = SSAOSmoothPass;
+		FRenderingCompositePass* SpatialPass;
+		SpatialPass = Context.Graph.RegisterPass(new (FMemStack::Get()) FRCPassPostProcessAmbientOcclusion_GTAO_SpatialFilter(Context.View));
+		SpatialPass->SetInput(ePId_Input0, FinalOutputPass);
+		SpatialPass->SetInput(ePId_Input1, HZBInput);
+		FinalOutputPass = SpatialPass;
+	}
+
+	if (ViewState)
+	{
+		// Add temporal filter
+		FRenderingCompositePass* TemporalPass;
+		TemporalPass = Context.Graph.RegisterPass(new (FMemStack::Get()) FRCPassPostProcessAmbientOcclusion_GTAO_TemporalFilter(Context.View, DownsampleFactor,
+			Context.View.PrevViewInfo.GTAOHistory,
+			&ViewState->PrevFrameViewInfo.GTAOHistory));
+
+		TemporalPass->SetInput(ePId_Input0, FinalOutputPass);
+		FinalOutputPass = TemporalPass;
+
+	}
+
+	{
+		FRenderingCompositePass* UpsamplePass;
+		UpsamplePass = Context.Graph.RegisterPass(new (FMemStack::Get()) FRCPassPostProcessAmbientOcclusion_GTAO_Upsample());
+		UpsamplePass->SetInput(ePId_Input0, FinalOutputPass);
+		FinalOutputPass = UpsamplePass;
 	}
 
 	Context.FinalOutput = FRenderingCompositeOutputRef(FinalOutputPass);
@@ -446,6 +534,8 @@ void FCompositionLighting::ProcessAfterBasePass(FRHICommandListImmediate& RHICmd
 
 		CompositeContext.Process(Context.FinalOutput.GetPass(), TEXT("CompositionLighting_AfterBasePass"));
 	}
+
+	SceneContext.ScreenSpaceGTAOHorizons.SafeRelease();
 }
 
 
diff --git a/Engine/Source/Runtime/Renderer/Private/CompositionLighting/PostProcessAmbientOcclusion.cpp b/Engine/Source/Runtime/Renderer/Private/CompositionLighting/PostProcessAmbientOcclusion.cpp
--- a/Engine/Source/Runtime/Renderer/Private/CompositionLighting/PostProcessAmbientOcclusion.cpp
+++ b/Engine/Source/Runtime/Renderer/Private/CompositionLighting/PostProcessAmbientOcclusion.cpp
@@ -20,9 +20,12 @@ DECLARE_GPU_STAT_NAMED(SSAOSetup, TEXT("ScreenSpace AO Setup") );
 DECLARE_GPU_STAT_NAMED(SSAO, TEXT("ScreenSpace AO") );
 DECLARE_GPU_STAT_NAMED(BasePassAO, TEXT("BasePass AO") );
 DECLARE_GPU_STAT_NAMED(SSAOSmooth, TEXT("SSAO smooth"));
-DECLARE_GPU_STAT_NAMED(GTAO_HorizonSearch, TEXT("GTAO HorizonSearch"));
-DECLARE_GPU_STAT_NAMED(GTAO_HorizonSearchInnerIntegrate, TEXT("GTAO HorizonSetupAndInnerIntegrate"));
+DECLARE_GPU_STAT_NAMED(GTAO_HorizonSearch,	TEXT("GTAO HorizonSearch"));
+DECLARE_GPU_STAT_NAMED(GTAO_Combined,		TEXT("GTAO Combined"));
 DECLARE_GPU_STAT_NAMED(GTAO_InnerIntegrate, TEXT("GTAO InnerIntegrate"));
+DECLARE_GPU_STAT_NAMED(GTAO_TemporalFilter, TEXT("GTAO Temportal Filter"));
+DECLARE_GPU_STAT_NAMED(GTAO_SpatialFilter,	TEXT("GTAO Spatial Filter"));
+DECLARE_GPU_STAT_NAMED(GTAO_Upsample,		TEXT("GTAO Upsample"));
 
 // Tile size for the AmbientOcclusion compute shader, tweaked for 680 GTX. */
 // see GCN Performance Tip 21 http://developer.amd.com/wordpress/media/2013/05/GCNPerformanceTweets.pdf 
@@ -92,12 +95,20 @@ static TAutoConsoleVariable<int32> CVarAmbientOcclusionDepthBoundsTest(
 
 static TAutoConsoleVariable<int32> CVarAmbientOcclusionMethod(
 	TEXT("r.AmbientOcclusion.Method"),
-	0,
+	1,
 	TEXT("Select between SSAO methods \n ")
 	TEXT("0: SSAO (default)\n ")
 	TEXT("1: GTAO\n "),
 	ECVF_RenderThreadSafe | ECVF_Scalability);
 
+static TAutoConsoleVariable<int32> CVarGTAOTemporalFilter(
+	TEXT("r.GTAO.TemporalFilter"),
+	1,
+	TEXT("Enable Temporal Filter for GTAO \n ")
+	TEXT("0: Off \n ")
+	TEXT("1: On (default)\n "),
+	ECVF_RenderThreadSafe | ECVF_Scalability);
+
 
 float FSSAOHelper::GetAmbientOcclusionQualityRT(const FSceneView& View)
 {
@@ -226,7 +237,7 @@ EGTAOType FSSAOHelper::GetGTAOPassType(const FViewInfo& View)
 
 	if (Method == 1)
 	{
-		if (IsAmbientOcclusionAsyncCompute(View, 1))
+		if (IsAmbientOcclusionAsyncCompute(View, 1) && GSupportsEfficientAsyncCompute)
 		{
 			return EGTAOType::ESplitAsync;
 		}
@@ -1221,12 +1232,11 @@ public:
 
 		if (ViewState)
 		{
-			TemporalFrame = ViewState->GetCurrentTemporalAASampleIndex();
-			Frame = ViewState->GetFrameIndex();
+			TemporalFrame	= ViewState->GetCurrentTemporalAASampleIndex();
+			Frame			= ViewState->GetFrameIndex();
 		}
-
-		static int32 DiffS = Frame - TemporalFrame;
-
+			
+		TemporalFrame = Frame;
 		FVector4 GTAOParam[2];
 
 		const float Rots[6]		= { 60.0f, 300.0f, 180.0f, 240.0f, 120.0f, 0.0f };
@@ -1234,8 +1244,7 @@ public:
 
 		float TemporalAngle = Rots[TemporalFrame % 6] * (PI / 360.0f);
 
-		GTAOParam[0] = FVector4(cos(TemporalAngle), sin(TemporalAngle),
-									Offsets[(TemporalFrame / 6) % 4], Offsets[TemporalFrame % 4]);
+		GTAOParam[0] = FVector4(cos(TemporalAngle), sin(TemporalAngle), Offsets[(TemporalFrame / 6) % 4], Offsets[TemporalFrame % 4]);
 
 
 		FIntPoint RandomizationSize = GSystemTextures.GTAORandomization->GetDesc().Extent;
@@ -1296,6 +1305,7 @@ public:
 	FShaderResourceParameter      RandomNormalTexture;
 	FShaderResourceParameter      RandomNormalTextureSampler;
 	FShaderParameter			  HZBRemapping;
+	FShaderParameter			  HorizonSearchParams;
 
 	/** Initialization constructor. */
 	FPostProcessGTAOHorizonSearchPSandCS(const ShaderMetaType::CompiledShaderInitializerType& Initializer)
@@ -1306,6 +1316,7 @@ public:
 		ScreenSpaceAOParams.Bind(Initializer.ParameterMap);
 		GTAOParams.Bind(Initializer.ParameterMap);
 		HZBRemapping.Bind(Initializer.ParameterMap, TEXT("HZBRemapping"));
+		HorizonSearchParams.Bind(Initializer.ParameterMap, TEXT("HorizonSearchParams"));
 		RandomNormalTexture.Bind(Initializer.ParameterMap, TEXT("RandomNormalTexture"));
 		RandomNormalTextureSampler.Bind(Initializer.ParameterMap, TEXT("RandomNormalTextureSampler"));
 
@@ -1323,23 +1334,20 @@ public:
 		const FVector2D UVScaleFactor(float(DestSize.X) / float(InputTextureSize.X),
 			float(DestSize.Y) / float(InputTextureSize.Y));
 
-		return FVector4(HZBScaleFactor.X / UVScaleFactor.X, HZBScaleFactor.Y / UVScaleFactor.Y, 0.0f, 0.0f);
+		return FVector4(HZBScaleFactor.X , HZBScaleFactor.Y , 0.0f, 0.0f);
 	}
 
-
 	template <typename TRHICmdList>
-	void SetParametersCS(TRHICmdList& RHICmdList, const FRenderingCompositePassContext& Context, FIntPoint DestSize, FIntPoint InputTextureSize, FRHIUnorderedAccessView *OutUAV)
+	void SetParametersCS(TRHICmdList& RHICmdList, const FRenderingCompositePassContext& Context, FIntPoint DestSize, FIntPoint InputTextureSize, uint32 DownScaleFactor, FRHIUnorderedAccessView *OutUAV)
 	{
 		const FFinalPostProcessSettings& Settings = Context.View.FinalPostProcessSettings;
 		FRHIComputeShader* ShaderRHI = GetComputeShader();
 
 		FGlobalShader::SetParameters<FViewUniformShaderParameters>(RHICmdList, ShaderRHI, Context.View.ViewUniformBuffer);
 		SceneTextureParameters.Set(RHICmdList, ShaderRHI, Context.View.FeatureLevel, ESceneTextureSetupMode::All);
-		ScreenSpaceAOParams.Set(RHICmdList, Context.View, ShaderRHI, InputTextureSize, FScreenSpaceAOParameters::ERandTexType::GTAO);
+		ScreenSpaceAOParams.Set(RHICmdList, Context.View, ShaderRHI, DestSize, FScreenSpaceAOParameters::ERandTexType::GTAO);
 
 		GTAOParams.Set(RHICmdList, Context.View, ShaderRHI);
-
-
 		PostprocessParameter.SetCS(ShaderRHI, Context, RHICmdList, TStaticSamplerState<SF_Point, AM_Clamp, AM_Clamp, AM_Clamp>::GetRHI());
 
 		RHICmdList.SetUAVParameter(ShaderRHI, HorizonOutTexture.GetBaseIndex(), OutUAV);
@@ -1347,6 +1355,9 @@ public:
 		FVector4 HZBRemappingValue = GetHZBRemapVal(Context, DestSize, InputTextureSize);
 		SetShaderValue(RHICmdList, ShaderRHI, HZBRemapping, HZBRemappingValue);
 
+		FVector4 HorizonSearchParamsValue = FVector4((float)DownScaleFactor, 0.0f, 0.0f, 0.0f);
+		SetShaderValue(RHICmdList, ShaderRHI, HorizonSearchParams, HorizonSearchParamsValue);
+
 		const FSceneRenderTargetItem& GTAORandomization = GSystemTextures.GTAORandomization->GetRenderTargetItem();
 		SetTextureParameter(RHICmdList, ShaderRHI, RandomNormalTexture, RandomNormalTextureSampler, TStaticSamplerState<SF_Point, AM_Wrap, AM_Wrap, AM_Wrap>::GetRHI(), GTAORandomization.ShaderResourceTexture);
 	}
@@ -1360,7 +1371,7 @@ public:
 
 		FGlobalShader::SetParameters<FViewUniformShaderParameters>(Context.RHICmdList, ShaderRHI, Context.View.ViewUniformBuffer);
 		SceneTextureParameters.Set(Context.RHICmdList, ShaderRHI, Context.View.FeatureLevel, ESceneTextureSetupMode::All);
-		ScreenSpaceAOParams.Set(Context.RHICmdList, Context.View, ShaderRHI, InputTextureSize, FScreenSpaceAOParameters::ERandTexType::GTAO);
+		ScreenSpaceAOParams.Set(Context.RHICmdList, Context.View, ShaderRHI, DestSize, FScreenSpaceAOParameters::ERandTexType::GTAO);
 		GTAOParams.Set(Context.RHICmdList, Context.View, ShaderRHI);
 
 		FVector4 HZBRemappingValue = GetHZBRemapVal(Context, DestSize, InputTextureSize);
@@ -1383,7 +1394,7 @@ public:
 	virtual bool Serialize(FArchive& Ar) override
 	{
 		bool bShaderHasOutdatedParameters = FGlobalShader::Serialize(Ar);
-		Ar << PostprocessParameter << SceneTextureParameters << ScreenSpaceAOParams << GTAOParams << HorizonOutTexture << HZBRemapping << RandomNormalTexture << RandomNormalTextureSampler;
+		Ar << PostprocessParameter << SceneTextureParameters << ScreenSpaceAOParams << GTAOParams << HorizonOutTexture << HZBRemapping << RandomNormalTexture << RandomNormalTextureSampler << HorizonSearchParams;
 		return bShaderHasOutdatedParameters;
 	}
 
@@ -1416,9 +1427,9 @@ VARIATION0(4)
 
 
 template<uint32 bComputeShader, uint32 ShaderQuality, uint32 UseNormalBuffer >
-class FPostProcessGTAOHorizonSearchAndInnerIntegratePSandCS : public FGlobalShader
+class FPostProcessGTAOCombinedPSandCS : public FGlobalShader
 {
-	DECLARE_SHADER_TYPE(FPostProcessGTAOHorizonSearchAndInnerIntegratePSandCS, Global);
+	DECLARE_SHADER_TYPE(FPostProcessGTAOCombinedPSandCS, Global);
 
 	static bool ShouldCompilePermutation(const FGlobalShaderPermutationParameters& Parameters)
 	{
@@ -1440,7 +1451,7 @@ class FPostProcessGTAOHorizonSearchAndInnerIntegratePSandCS : public FGlobalShad
 	}
 
 	/** Default constructor. */
-	FPostProcessGTAOHorizonSearchAndInnerIntegratePSandCS() {}
+	FPostProcessGTAOCombinedPSandCS() {}
 
 public:
 	FSceneTextureShaderParameters SceneTextureParameters;
@@ -1453,7 +1464,7 @@ public:
 	FGTAOParameters				  GTAOParams;
 
 	/** Initialization constructor. */
-	FPostProcessGTAOHorizonSearchAndInnerIntegratePSandCS(const ShaderMetaType::CompiledShaderInitializerType& Initializer)
+	FPostProcessGTAOCombinedPSandCS(const ShaderMetaType::CompiledShaderInitializerType& Initializer)
 		: FGlobalShader(Initializer)
 	{
 		PostprocessParameter.Bind(Initializer.ParameterMap);
@@ -1475,10 +1486,11 @@ public:
 		const FVector2D HZBScaleFactor(float(Context.View.ViewRect.Width()) / float(2 * Context.View.HZBMipmap0Size.X),
 			float(Context.View.ViewRect.Height()) / float(2 * Context.View.HZBMipmap0Size.Y));
 
-		const FVector2D UVScaleFactor(float(DestSize.X) / float(InputTextureSize.X),
-			float(DestSize.Y) / float(InputTextureSize.Y));
+		const FVector2D UVScaleFactor(
+			float(InputTextureSize.X) / float(Context.View.ViewRect.Width()),
+			float(InputTextureSize.Y) / float(Context.View.ViewRect.Height()));
 
-		return FVector4(HZBScaleFactor.X / UVScaleFactor.X, HZBScaleFactor.Y / UVScaleFactor.Y, 0.0f, 0.0f);
+		return FVector4(HZBScaleFactor.X * UVScaleFactor.X, HZBScaleFactor.Y / UVScaleFactor.Y, 0.0f, 0.0f);
 	}
 
 
@@ -1513,7 +1525,7 @@ public:
 
 		FGlobalShader::SetParameters<FViewUniformShaderParameters>(Context.RHICmdList, ShaderRHI, Context.View.ViewUniformBuffer);
 		SceneTextureParameters.Set(Context.RHICmdList, ShaderRHI, Context.View.FeatureLevel, ESceneTextureSetupMode::All);
-		ScreenSpaceAOParams.Set(Context.RHICmdList, Context.View, ShaderRHI, InputTextureSize, FScreenSpaceAOParameters::ERandTexType::GTAO);
+		ScreenSpaceAOParams.Set(Context.RHICmdList, Context.View, ShaderRHI, DestSize, FScreenSpaceAOParameters::ERandTexType::GTAO);
 		GTAOParams.Set(Context.RHICmdList, Context.View, ShaderRHI);
 
 		FVector4 HZBRemappingValue = GetHZBRemapVal(Context, DestSize, InputTextureSize);
@@ -1547,20 +1559,20 @@ public:
 
 	static const TCHAR* GetFunctionName()
 	{
-		return bComputeShader ? TEXT("HorizonSearchAndInnerIntegrateCS") : TEXT("HorizonSearchAndInnerIntegratePS");
+		return bComputeShader ? TEXT("GTAOCombinedCS") : TEXT("GTAOCombinedPS");
 	}
 };
 
 // #define avoids a lot of code duplication
 #define VARIATION0(A) \
-	typedef FPostProcessGTAOHorizonSearchAndInnerIntegratePSandCS<false,A,0> FPostProcessGTAOHorizonSearchAndInnerIntegratePS0##A; \
-	typedef FPostProcessGTAOHorizonSearchAndInnerIntegratePSandCS<false,A,1> FPostProcessGTAOHorizonSearchAndInnerIntegratePS1##A; \
-	typedef FPostProcessGTAOHorizonSearchAndInnerIntegratePSandCS<true ,A,0> FPostProcessGTAOHorizonSearchAndInnerIntegrateCS0##A; \
-	typedef FPostProcessGTAOHorizonSearchAndInnerIntegratePSandCS<true ,A,1> FPostProcessGTAOHorizonSearchAndInnerIntegrateCS1##A; \
-	IMPLEMENT_SHADER_TYPE2(FPostProcessGTAOHorizonSearchAndInnerIntegratePS0##A, SF_Pixel); \
-	IMPLEMENT_SHADER_TYPE2(FPostProcessGTAOHorizonSearchAndInnerIntegratePS1##A, SF_Pixel); \
-	IMPLEMENT_SHADER_TYPE2(FPostProcessGTAOHorizonSearchAndInnerIntegrateCS0##A, SF_Compute);\
-	IMPLEMENT_SHADER_TYPE2(FPostProcessGTAOHorizonSearchAndInnerIntegrateCS1##A, SF_Compute);
+	typedef FPostProcessGTAOCombinedPSandCS<false,A,0> FPostProcessGTAOCombinedPSandCSPS0##A; \
+	typedef FPostProcessGTAOCombinedPSandCS<false,A,1> FPostProcessGTAOCombinedPSandCSPS1##A; \
+	typedef FPostProcessGTAOCombinedPSandCS<true ,A,0> FPostProcessGTAOCombinedPSandCSCS0##A; \
+	typedef FPostProcessGTAOCombinedPSandCS<true ,A,1> FPostProcessGTAOCombinedPSandCSCS1##A; \
+	IMPLEMENT_SHADER_TYPE2(FPostProcessGTAOCombinedPSandCSPS0##A, SF_Pixel); \
+	IMPLEMENT_SHADER_TYPE2(FPostProcessGTAOCombinedPSandCSPS1##A, SF_Pixel); \
+	IMPLEMENT_SHADER_TYPE2(FPostProcessGTAOCombinedPSandCSCS0##A, SF_Compute);\
+	IMPLEMENT_SHADER_TYPE2(FPostProcessGTAOCombinedPSandCSCS1##A, SF_Compute);
 VARIATION0(0)
 VARIATION0(1)
 VARIATION0(2)
@@ -1604,6 +1616,7 @@ public:
 	FShaderResourceParameter      HorizonsTexture;
 	FShaderResourceParameter      HorizonsTextureSampler;
 	FGTAOParameters				  GTAOParams;
+	FShaderParameter			  InnerIntegrateParams;
 
 
 	/** Initialization constructor. */
@@ -1618,6 +1631,7 @@ public:
 		HorizonsTexture.Bind(Initializer.ParameterMap, TEXT("HorizonsTexture"));
 		HorizonsTextureSampler.Bind(Initializer.ParameterMap, TEXT("HorizonsTextureSampler"));
 		GTAOParams.Bind(Initializer.ParameterMap);
+		InnerIntegrateParams.Bind(Initializer.ParameterMap, TEXT("InnerIntegrateParams"));
 
 		if (bComputeShader)
 		{
@@ -1629,14 +1643,14 @@ public:
 	{
 		bool bShaderHasOutdatedParameters = FGlobalShader::Serialize(Ar);
 		Ar << PostprocessParameter << SceneTextureParameters << ScreenSpaceAOParams << OutTexture << GTAOParams <<  RandomNormalTexture << RandomNormalTextureSampler;
-		Ar << HorizonsTexture << HorizonsTextureSampler;
+		Ar << HorizonsTexture << HorizonsTextureSampler << InnerIntegrateParams;
 
 		return bShaderHasOutdatedParameters;
 	}
 
 
 	template <typename TRHICmdList>
-	void SetParametersCS(TRHICmdList& RHICmdList, const FRenderingCompositePassContext& Context, FIntPoint DestSize, FIntPoint InputTextureSize, FRHIUnorderedAccessView *OutUAV)
+	void SetParametersCS(TRHICmdList& RHICmdList, const FRenderingCompositePassContext& Context, uint32 DownScaleFactor, FIntPoint DestSize, FIntPoint InputTextureSize, FRHIUnorderedAccessView *OutUAV)
 	{
 		const FFinalPostProcessSettings& Settings = Context.View.FinalPostProcessSettings;
 		FRHIComputeShader* ShaderRHI = GetComputeShader();
@@ -1658,8 +1672,11 @@ public:
 		SetTextureParameter(RHICmdList, ShaderRHI, HorizonsTexture, HorizonsTextureSampler, TStaticSamplerState<SF_Point, AM_Wrap, AM_Wrap, AM_Wrap>::GetRHI(), 
 			SceneContext.ScreenSpaceGTAOHorizons->GetRenderTargetItem().ShaderResourceTexture);
 
+		FVector4 InnerIntegrateParamsValue = FVector4((float)DownScaleFactor, 0.0f, 0.0f, 0.0f);
+		SetShaderValue(RHICmdList, ShaderRHI, InnerIntegrateParams, InnerIntegrateParamsValue);
 	}
 
+
 	void SetParametersPS(const FRenderingCompositePassContext& Context, FIntPoint DestSize, FIntPoint InputTextureSize)
 	{
 		const FFinalPostProcessSettings& Settings = Context.View.FinalPostProcessSettings;
@@ -1702,16 +1719,17 @@ public:
 IMPLEMENT_SHADER_TYPE2(FPostProcessGTAOInnerIntegratePSandCS<false>, SF_Pixel);
 IMPLEMENT_SHADER_TYPE2(FPostProcessGTAOInnerIntegratePSandCS<true>, SF_Compute);
 
-FRCPassPostProcessAmbientOcclusion_HorizonSearchInnerIntegrate::FRCPassPostProcessAmbientOcclusion_HorizonSearchInnerIntegrate(const FSceneView& View, bool FinalOutput)
-	:bFinalOutput(FinalOutput)
+FRCPassPostProcessAmbientOcclusion_GTAOCombined::FRCPassPostProcessAmbientOcclusion_GTAOCombined(const FSceneView& View, uint32 inDownScaleFactor, bool FinalOutput)
+	:bFinalOutput(FinalOutput),
+	DownScaleFactor(inDownScaleFactor)
 {
 
 }
 
 template <uint32 ShaderQuality, uint32 UseNormals>
-void FRCPassPostProcessAmbientOcclusion_HorizonSearchInnerIntegrate::DispatchCS(const FRenderingCompositePassContext& Context, FIntPoint DestSize, FIntPoint TexSize)
+void FRCPassPostProcessAmbientOcclusion_GTAOCombined::DispatchCS(const FRenderingCompositePassContext& Context, FIntPoint DestSize, FIntPoint TexSize)
 {
-	TShaderMapRef<FPostProcessGTAOHorizonSearchAndInnerIntegratePSandCS<true, ShaderQuality, UseNormals>> ComputeShader(Context.GetShaderMap());
+	TShaderMapRef<FPostProcessGTAOCombinedPSandCS<true, ShaderQuality, UseNormals>> ComputeShader(Context.GetShaderMap());
 	FSceneRenderTargets& SceneContext = FSceneRenderTargets::Get(Context.RHICmdList);
 	const FSceneRenderTargetItem& DestRenderTarget = bFinalOutput ? SceneContext.ScreenSpaceAO->GetRenderTargetItem() : PassOutputs[0].RequestSurface(Context);
 
@@ -1726,10 +1744,10 @@ void FRCPassPostProcessAmbientOcclusion_HorizonSearchInnerIntegrate::DispatchCS(
 }
 
 template <uint32 ShaderQuality, uint32 UseNormals>
-FShader* FRCPassPostProcessAmbientOcclusion_HorizonSearchInnerIntegrate::SetShaderPS(const FRenderingCompositePassContext& Context, FGraphicsPipelineStateInitializer& GraphicsPSOInit, FIntPoint DestSize)
+FShader* FRCPassPostProcessAmbientOcclusion_GTAOCombined::SetShaderPS(const FRenderingCompositePassContext& Context, FGraphicsPipelineStateInitializer& GraphicsPSOInit, FIntPoint DestSize)
 {
 	TShaderMapRef<FPostProcessVS> VertexShader(Context.GetShaderMap());
-	TShaderMapRef<FPostProcessGTAOHorizonSearchAndInnerIntegratePSandCS<false, ShaderQuality, UseNormals> > PixelShader(Context.GetShaderMap());
+	TShaderMapRef<FPostProcessGTAOCombinedPSandCS<false, ShaderQuality, UseNormals> > PixelShader(Context.GetShaderMap());
 
 	GraphicsPSOInit.BoundShaderState.VertexDeclarationRHI = GFilterVertexDeclaration.VertexDeclarationRHI;
 	GraphicsPSOInit.BoundShaderState.VertexShaderRHI = GETSAFERHISHADER_VERTEX(*VertexShader);
@@ -1748,19 +1766,24 @@ FShader* FRCPassPostProcessAmbientOcclusion_HorizonSearchInnerIntegrate::SetShad
 
 
 
-void FRCPassPostProcessAmbientOcclusion_HorizonSearchInnerIntegrate::Process(FRenderingCompositePassContext& Context)
+void FRCPassPostProcessAmbientOcclusion_GTAOCombined::Process(FRenderingCompositePassContext& Context)
 {
-	SCOPED_GPU_STAT(Context.RHICmdList, GTAO_HorizonSearchInnerIntegrate);
+	SCOPED_GPU_STAT(Context.RHICmdList, GTAO_Combined);
 	const FViewInfo& View = Context.View;
-	FIntRect ViewRect = View.ViewRect;
-	FIntRect SrcRect = Context.View.ViewRect;
-	FIntRect DestRect = SrcRect;
-	FIntPoint DestSize(DestRect.Width(), DestRect.Height());
-	const FPooledRenderTargetDesc* InputDesc0 = GetInputDesc(ePId_Input0);
+
+	// Get Size of destination
 	FSceneRenderTargets& SceneContext = FSceneRenderTargets::Get(Context.RHICmdList);
-	FIntPoint TexSize = InputDesc0 ? InputDesc0->Extent : SceneContext.GetBufferSizeXY();
+	const FSceneRenderTargetItem& DestRenderTarget = PassOutputs[0].RequestSurface(Context);
+	const FPooledRenderTargetDesc* InputDesc0 = GetInputDesc(ePId_Input0);
 
-	const FSceneRenderTargetItem& DestRenderTarget = bFinalOutput ? SceneContext.ScreenSpaceAO->GetRenderTargetItem() : PassOutputs[0].RequestSurface(Context);
+	// Get the size of the input and output sizes of the buffers.
+	FIntPoint InputTexSize	= InputDesc0 ? InputDesc0->Extent : SceneContext.GetBufferSizeXY();
+	FIntPoint OutputTexSize = PassOutputs[0].RenderTargetDesc.Extent; 
+	
+	FIntRect InputViewRect  = View.ViewRect;
+	FIntRect OutputViewRect = InputViewRect / DownScaleFactor;
+
+	FIntRect OutputFullRect(FIntPoint(0,0), OutputTexSize);
 
 	static const auto CVar = IConsoleManager::Get().FindTConsoleVariableDataInt(TEXT("r.AmbientOcclusion.Compute"));
 
@@ -1777,19 +1800,19 @@ void FRCPassPostProcessAmbientOcclusion_HorizonSearchInnerIntegrate::Process(FRe
 		switch(ShaderQuality)
 		{
 		case 0:
-			DispatchCS<0,0>(Context, DestSize, TexSize);
+			DispatchCS<0,0>(Context, OutputTexSize, InputTexSize);
 			break;
 		case 1:
-			DispatchCS<1, 0>(Context, DestSize, TexSize);
+			DispatchCS<1, 0>(Context, OutputTexSize, InputTexSize);
 			break;
 		case 2:
-			DispatchCS<2, 0>(Context, DestSize, TexSize);
+			DispatchCS<2, 0>(Context, OutputTexSize, InputTexSize);
 			break;
 		case 3:
-			DispatchCS<3, 0>(Context, DestSize, TexSize);
+			DispatchCS<3, 0>(Context, OutputTexSize, InputTexSize);
 			break;
 		case 4:
-			DispatchCS<4, 0>(Context, DestSize, TexSize);
+			DispatchCS<4, 0>(Context, OutputTexSize, InputTexSize);
 			break;
 		}
 	}
@@ -1800,7 +1823,8 @@ void FRCPassPostProcessAmbientOcclusion_HorizonSearchInnerIntegrate::Process(FRe
 		Context.RHICmdList.BeginRenderPass(RPInfo, TEXT("AmbientOcclusionSetup"));
 		{
 			// PS Version
-			Context.SetViewportAndCallRHI(ViewRect);
+			Context.SetViewportAndCallRHI(OutputFullRect);
+			DrawClearQuad(Context.RHICmdList, FLinearColor::White);
 
 			FGraphicsPipelineStateInitializer GraphicsPSOInit;
 			Context.RHICmdList.ApplyCachedRenderTargets(GraphicsPSOInit);
@@ -1817,30 +1841,28 @@ void FRCPassPostProcessAmbientOcclusion_HorizonSearchInnerIntegrate::Process(FRe
 			switch (ShaderQuality)
 			{
 			case 0:
-				VertexShader = SetShaderPS<0, 0>(Context, GraphicsPSOInit, DestSize);
+				VertexShader = SetShaderPS<0, 0>(Context, GraphicsPSOInit, OutputTexSize);
 				break;
 			case 1:
-				VertexShader = SetShaderPS<1, 0>(Context, GraphicsPSOInit, DestSize);
+				VertexShader = SetShaderPS<1, 0>(Context, GraphicsPSOInit, OutputTexSize);
 				break;
 			case 2:
-				VertexShader = SetShaderPS<2, 0>(Context, GraphicsPSOInit, DestSize);
+				VertexShader = SetShaderPS<2, 0>(Context, GraphicsPSOInit, OutputTexSize);
 				break;
 			case 3:
-				VertexShader = SetShaderPS<3, 0>(Context, GraphicsPSOInit, DestSize);
+				VertexShader = SetShaderPS<3, 0>(Context, GraphicsPSOInit, OutputTexSize);
 				break;
 			case 4:
-				VertexShader = SetShaderPS<4, 0>(Context, GraphicsPSOInit, DestSize);
+				VertexShader = SetShaderPS<4, 0>(Context, GraphicsPSOInit, OutputTexSize);
 				break;
 			}
 
 			DrawRectangle(
 				Context.RHICmdList,
-				0, 0,
-				ViewRect.Width(), ViewRect.Height(),
-				ViewRect.Min.X, ViewRect.Min.Y,
-				ViewRect.Width(), ViewRect.Height(),
-				ViewRect.Size(),
-				TexSize,
+				OutputViewRect.Min.X, OutputViewRect.Min.Y, OutputViewRect.Width(), OutputViewRect.Height(),
+				InputViewRect.Min.X, InputViewRect.Min.Y, InputViewRect.Width(), InputViewRect.Height(),
+				OutputTexSize,
+				InputTexSize,
 				VertexShader,
 				EDRF_UseTriangleOptimization);
 		}
@@ -1851,46 +1873,46 @@ void FRCPassPostProcessAmbientOcclusion_HorizonSearchInnerIntegrate::Process(FRe
 }
 
 
-FPooledRenderTargetDesc FRCPassPostProcessAmbientOcclusion_HorizonSearchInnerIntegrate::ComputeOutputDesc(EPassOutputId InPassOutputId) const
+FPooledRenderTargetDesc FRCPassPostProcessAmbientOcclusion_GTAOCombined::ComputeOutputDesc(EPassOutputId InPassOutputId) const
 {
 	FPooledRenderTargetDesc Ret;
 
 	Ret = GetInput(ePId_Input0)->GetOutput()->RenderTargetDesc;
 	Ret.Reset();
-	Ret.Format = PF_B8G8R8A8;
+	Ret.Format = PF_G8;
+	Ret.Extent = FIntPoint::DivideAndRoundUp(Ret.Extent, DownScaleFactor);
 
 	Ret.ClearValue = FClearValueBinding::None;
 	Ret.TargetableFlags &= ~TexCreate_DepthStencilTargetable;
 	Ret.TargetableFlags |= TexCreate_RenderTargetable | TexCreate_ShaderResource;
-	Ret.DebugName = TEXT("GTAOInnerHorizonInnerIntegrate");
+	Ret.DebugName = TEXT("GTAOCombined");
 
 	return Ret;
 }
 
 
-
-
-FRCPassPostProcessAmbientOcclusion_InnerIntegrate::FRCPassPostProcessAmbientOcclusion_InnerIntegrate(const FSceneView& View, bool FinalOutput)
-	:bFinalOutput(FinalOutput)
+FRCPassPostProcessAmbientOcclusion_InnerIntegrate::FRCPassPostProcessAmbientOcclusion_InnerIntegrate(const FSceneView& View, uint32 inDownScaleFactor, bool FinalOutput)
+	:bFinalOutput(FinalOutput),
+	DownScaleFactor(inDownScaleFactor)
 {
-
 }
 
-
-
-
 void FRCPassPostProcessAmbientOcclusion_InnerIntegrate::Process(FRenderingCompositePassContext& Context)
 {
 	SCOPED_GPU_STAT(Context.RHICmdList, GTAO_InnerIntegrate);
 	const FViewInfo& View = Context.View;
-	FIntRect ViewRect = View.ViewRect;
-	FIntRect SrcRect = Context.View.ViewRect;
-	FIntRect DestRect = SrcRect;
-	FIntPoint DestSize(DestRect.Width(), DestRect.Height());
-	const FPooledRenderTargetDesc* InputDesc0 = GetInputDesc(ePId_Input0);
+
 
 	FSceneRenderTargets& SceneContext = FSceneRenderTargets::Get(Context.RHICmdList);
-	FIntPoint TexSize = InputDesc0 ? InputDesc0->Extent : SceneContext.GetBufferSizeXY();
+	const FPooledRenderTargetDesc* InputDesc0 = GetInputDesc(ePId_Input0);
+
+	FIntPoint InputTexSize = InputDesc0->Extent;
+	FIntPoint OutputTexSize = PassOutputs[0].RenderTargetDesc.Extent;
+	FIntRect OutputFullRect(FIntPoint(0, 0), OutputTexSize);
+
+	FIntRect InputViewRect = View.ViewRect / DownScaleFactor;
+	FIntRect OutputViewRect = InputViewRect;
+
 
 	const FSceneRenderTargetItem& DestRenderTarget = bFinalOutput ? SceneContext.ScreenSpaceAO->GetRenderTargetItem() : PassOutputs[0].RequestSurface(Context);
 	static const auto CVar = IConsoleManager::Get().FindTConsoleVariableDataInt(TEXT("r.AmbientOcclusion.Compute"));
@@ -1904,10 +1926,10 @@ void FRCPassPostProcessAmbientOcclusion_InnerIntegrate::Process(FRenderingCompos
 
 		TShaderMapRef<FPostProcessGTAOInnerIntegratePSandCS<true>> ComputeShader(Context.GetShaderMap());
 		Context.RHICmdList.SetComputeShader(ComputeShader->GetComputeShader());
-		ComputeShader->SetParametersCS(Context.RHICmdList, Context, DestSize, TexSize, DestRenderTarget.UAV);
+		ComputeShader->SetParametersCS(Context.RHICmdList, Context, DownScaleFactor, OutputTexSize, InputTexSize, DestRenderTarget.UAV);
 
-		uint32 GroupSizeX = FMath::DivideAndRoundUp(DestSize.X, 8);
-		uint32 GroupSizeY = FMath::DivideAndRoundUp(DestSize.Y, 8);
+		uint32 GroupSizeX = FMath::DivideAndRoundUp(InputViewRect.Width(), 8);
+		uint32 GroupSizeY = FMath::DivideAndRoundUp(InputViewRect.Height(), 8);
 		DispatchComputeShader(Context.RHICmdList, *ComputeShader, GroupSizeX, GroupSizeY, 1);
 
 		ComputeShader->UnsetParameters(Context.RHICmdList);
@@ -1918,7 +1940,7 @@ void FRCPassPostProcessAmbientOcclusion_InnerIntegrate::Process(FRenderingCompos
 		FRHIRenderPassInfo RPInfo(DestRenderTarget.TargetableTexture, ERenderTargetActions::Load_Store);
 		Context.RHICmdList.BeginRenderPass(RPInfo, TEXT("AmbientOcclusionSetup"));
 		{
-			Context.SetViewportAndCallRHI(ViewRect);
+			Context.SetViewportAndCallRHI(InputViewRect);
 
 			TShaderMapRef<FPostProcessGTAOInnerIntegratePSandCS<false> > PixelShader(Context.GetShaderMap());
 			TShaderMapRef<FPostProcessVS> VertexShader(Context.GetShaderMap());
@@ -1934,16 +1956,16 @@ void FRCPassPostProcessAmbientOcclusion_InnerIntegrate::Process(FRenderingCompos
 			GraphicsPSOInit.BoundShaderState.PixelShaderRHI = GETSAFERHISHADER_PIXEL(*PixelShader);
 			SetGraphicsPipelineState(Context.RHICmdList, GraphicsPSOInit);
 
-			PixelShader->SetParametersPS(Context, DestSize, TexSize);
+			PixelShader->SetParametersPS(Context, OutputTexSize, InputTexSize);
 
 			DrawRectangle(
 				Context.RHICmdList,
 				0, 0,
-				ViewRect.Width(), ViewRect.Height(),
-				ViewRect.Min.X, ViewRect.Min.Y,
-				ViewRect.Width(), ViewRect.Height(),
-				ViewRect.Size(),
-				TexSize,
+				InputViewRect.Width(), InputViewRect.Height(),
+				InputViewRect.Min.X, InputViewRect.Min.Y,
+				InputViewRect.Width(), InputViewRect.Height(),
+				InputViewRect.Size(),
+				OutputTexSize,
 				*VertexShader,
 				EDRF_UseTriangleOptimization);
 		}
@@ -1962,6 +1984,8 @@ FPooledRenderTargetDesc FRCPassPostProcessAmbientOcclusion_InnerIntegrate::Compu
 	Ret.Reset();
 	Ret.Format = PF_G8;
 
+	Ret.Extent = FIntPoint::DivideAndRoundUp(Ret.Extent, DownScaleFactor);
+
 	Ret.ClearValue = FClearValueBinding::None;
 	Ret.TargetableFlags &= ~TexCreate_DepthStencilTargetable;
 	Ret.TargetableFlags |= TexCreate_UAV;
@@ -1971,16 +1995,16 @@ FPooledRenderTargetDesc FRCPassPostProcessAmbientOcclusion_InnerIntegrate::Compu
 }
 
 
-
-FRCPassPostProcessAmbientOcclusion_HorizonSearch::FRCPassPostProcessAmbientOcclusion_HorizonSearch(const FSceneView& View, const ESSAOType InAOType)
-	: AOType(InAOType)
+FRCPassPostProcessAmbientOcclusion_HorizonSearch::FRCPassPostProcessAmbientOcclusion_HorizonSearch(const FSceneView& View, uint32 InDownScaleFactor, const ESSAOType InAOType)
+	: AOType(InAOType),
+	DownScaleFactor(InDownScaleFactor)
 {
 }
 
 
 
 template <uint32 ShaderQuality>
-void FRCPassPostProcessAmbientOcclusion_HorizonSearch::DispatchCS(const FRenderingCompositePassContext& Context, FIntPoint DestSize, FIntPoint TexSize)
+void FRCPassPostProcessAmbientOcclusion_HorizonSearch::DispatchCS(const FRenderingCompositePassContext& Context, FIntRect ViewRect, FIntPoint DestSize, FIntPoint TexSize)
 {
 	FSceneRenderTargets& SceneContext = FSceneRenderTargets::Get(Context.RHICmdList);
 	const FSceneRenderTargetItem& DestRenderTarget = (AOType == ESSAOType::EAsyncCS) ? SceneContext.ScreenSpaceGTAOHorizons->GetRenderTargetItem() : PassOutputs[0].RequestSurface(Context);
@@ -1989,16 +2013,15 @@ void FRCPassPostProcessAmbientOcclusion_HorizonSearch::DispatchCS(const FRenderi
 
 	TShaderMapRef<FPostProcessGTAOHorizonSearchPSandCS<true, ShaderQuality>> ComputeShader(Context.GetShaderMap());
 	RHICmdListComputeImmediate.SetComputeShader(ComputeShader->GetComputeShader());
-	ComputeShader->SetParametersCS(RHICmdListComputeImmediate, Context, DestSize, TexSize, DestRenderTarget.UAV);
+	ComputeShader->SetParametersCS(RHICmdListComputeImmediate, Context, DestSize, TexSize, DownScaleFactor, DestRenderTarget.UAV);
 
-	uint32 GroupSizeX = FMath::DivideAndRoundUp(DestSize.X, 8);
-	uint32 GroupSizeY = FMath::DivideAndRoundUp(DestSize.Y, 8);
+	uint32 GroupSizeX = FMath::DivideAndRoundUp(ViewRect.Width() , 8);
+	uint32 GroupSizeY = FMath::DivideAndRoundUp(ViewRect.Height(), 8);
 	DispatchComputeShader(RHICmdListComputeImmediate, *ComputeShader, GroupSizeX, GroupSizeY, 1);
 
 	ComputeShader->UnsetParameters(RHICmdListComputeImmediate);
 }
 
-
 template <uint32 ShaderQuality>
 FShader* FRCPassPostProcessAmbientOcclusion_HorizonSearch::SetShaderPS(const FRenderingCompositePassContext& Context, FGraphicsPipelineStateInitializer& GraphicsPSOInit, FIntPoint DestSize)
 {
@@ -2024,17 +2047,18 @@ FShader* FRCPassPostProcessAmbientOcclusion_HorizonSearch::SetShaderPS(const FRe
 
 void FRCPassPostProcessAmbientOcclusion_HorizonSearch::Process(FRenderingCompositePassContext& Context)
 {
-
 	SCOPED_GPU_STAT(Context.RHICmdList, GTAO_HorizonSearch);
 	const FViewInfo& View = Context.View;
-	FIntRect ViewRect = View.ViewRect;
-	FIntRect SrcRect = Context.View.ViewRect;
-	FIntRect DestRect = SrcRect;
-	FIntPoint DestSize(DestRect.Width(), DestRect.Height());
-	const FPooledRenderTargetDesc* InputDesc0 = GetInputDesc(ePId_Input0);
+
 	FSceneRenderTargets& SceneContext = FSceneRenderTargets::Get(Context.RHICmdList);
-	FIntPoint TexSize = InputDesc0 ? InputDesc0->Extent : SceneContext.GetBufferSizeXY();
+	const FPooledRenderTargetDesc* InputDesc0 = GetInputDesc(ePId_Input0);
 
+	FIntPoint InputTexSize = InputDesc0->Extent;
+	FIntPoint OutputTexSize = PassOutputs[0].RenderTargetDesc.Extent;
+	FIntRect OutputFullRect(FIntPoint(0, 0), OutputTexSize);
+
+	FIntRect InputViewRect = View.ViewRect / DownScaleFactor;
+	FIntRect OutputViewRect = InputViewRect;
 
 	const FSceneRenderTargetItem& DestRenderTarget = (AOType == ESSAOType::EAsyncCS) ? SceneContext.ScreenSpaceGTAOHorizons->GetRenderTargetItem()  : PassOutputs[0].RequestSurface(Context);
 
@@ -2060,19 +2084,19 @@ void FRCPassPostProcessAmbientOcclusion_HorizonSearch::Process(FRenderingComposi
 		switch (ShaderQuality)
 		{
 		case 0:
-			DispatchCS<0>(Context, DestSize, TexSize);
+			DispatchCS<0>(Context, InputViewRect, OutputTexSize, InputTexSize);
 			break;
 		case 1:
-			DispatchCS<1>(Context, DestSize, TexSize);
+			DispatchCS<1>(Context, InputViewRect, OutputTexSize, InputTexSize);
 			break;
 		case 2:
-			DispatchCS<2>(Context, DestSize, TexSize);
+			DispatchCS<2>(Context, InputViewRect, OutputTexSize, InputTexSize);
 			break;
 		case 3:
-			DispatchCS<3>(Context, DestSize, TexSize);
+			DispatchCS<3>(Context, InputViewRect, OutputTexSize, InputTexSize);
 			break;
 		case 4:
-			DispatchCS<4>(Context, DestSize, TexSize);
+			DispatchCS<4>(Context, InputViewRect, OutputTexSize, InputTexSize);
 			break;
 		}
 	}
@@ -2082,7 +2106,7 @@ void FRCPassPostProcessAmbientOcclusion_HorizonSearch::Process(FRenderingComposi
 		Context.RHICmdList.BeginRenderPass(RPInfo, TEXT("GTAOHorizonSearch"));
 		{
 			// PS Version
-			Context.SetViewportAndCallRHI(ViewRect);
+			Context.SetViewportAndCallRHI(InputViewRect);
 
 			FGraphicsPipelineStateInitializer GraphicsPSOInit;
 			Context.RHICmdList.ApplyCachedRenderTargets(GraphicsPSOInit);
@@ -2099,30 +2123,30 @@ void FRCPassPostProcessAmbientOcclusion_HorizonSearch::Process(FRenderingComposi
 			switch (ShaderQuality)
 			{
 			case 0:
-				VertexShader = SetShaderPS<0>(Context, GraphicsPSOInit, DestSize);
+				VertexShader = SetShaderPS<0>(Context, GraphicsPSOInit, OutputTexSize);
 				break;
 			case 1:
-				VertexShader = SetShaderPS<1>(Context, GraphicsPSOInit, DestSize);
+				VertexShader = SetShaderPS<1>(Context, GraphicsPSOInit, OutputTexSize);
 				break;
 			case 2:
-				VertexShader = SetShaderPS<2>(Context, GraphicsPSOInit, DestSize);
+				VertexShader = SetShaderPS<2>(Context, GraphicsPSOInit, OutputTexSize);
 				break;
 			case 3:
-				VertexShader = SetShaderPS<3>(Context, GraphicsPSOInit, DestSize);
+				VertexShader = SetShaderPS<3>(Context, GraphicsPSOInit, OutputTexSize);
 				break;
 			case 4:
-				VertexShader = SetShaderPS<4>(Context, GraphicsPSOInit, DestSize);
+				VertexShader = SetShaderPS<4>(Context, GraphicsPSOInit, OutputTexSize);
 				break;
 			}
 
 			DrawRectangle(
 				Context.RHICmdList,
 				0, 0,
-				ViewRect.Width(), ViewRect.Height(),
-				ViewRect.Min.X, ViewRect.Min.Y,
-				ViewRect.Width(), ViewRect.Height(),
-				ViewRect.Size(),
-				TexSize,
+				InputViewRect.Width(), InputViewRect.Height(),
+				InputViewRect.Min.X,   InputViewRect.Min.Y,
+				InputViewRect.Width(), InputViewRect.Height(),
+				InputViewRect.Size(),
+				InputTexSize,
 				VertexShader,
 				EDRF_UseTriangleOptimization);
 		}
@@ -2141,6 +2165,8 @@ FPooledRenderTargetDesc FRCPassPostProcessAmbientOcclusion_HorizonSearch::Comput
 	Ret.Reset();
 	Ret.Format = PF_B8G8R8A8;
 
+	Ret.Extent = FIntPoint::DivideAndRoundUp(Ret.Extent, 2);
+
 	Ret.ClearValue = FClearValueBinding::None;
 	Ret.TargetableFlags &= ~TexCreate_DepthStencilTargetable;
 	Ret.TargetableFlags |= TexCreate_UAV;
@@ -2150,4 +2176,626 @@ FPooledRenderTargetDesc FRCPassPostProcessAmbientOcclusion_HorizonSearch::Comput
 	Ret.DebugName = TEXT("GTAOInnerHorizon");
 
 	return Ret;
-}
\ No newline at end of file
+}
+
+template<uint32 bComputeShader>
+class FPostProcessGTAOTemporalFilterPSandCS : public FGlobalShader
+{
+	DECLARE_SHADER_TYPE(FPostProcessGTAOTemporalFilterPSandCS, Global);
+
+	static bool ShouldCompilePermutation(const FGlobalShaderPermutationParameters& Parameters)
+	{
+		return IsFeatureLevelSupported(Parameters.Platform, ERHIFeatureLevel::SM5);
+	}
+
+	static void ModifyCompilationEnvironment(const FGlobalShaderPermutationParameters& Parameters, FShaderCompilerEnvironment& OutEnvironment)
+	{
+		FGlobalShader::ModifyCompilationEnvironment(Parameters, OutEnvironment);
+		OutEnvironment.SetDefine(TEXT("COMPUTE_SHADER"), bComputeShader);
+		if (bComputeShader)
+		{
+			OutEnvironment.SetDefine(TEXT("THREADGROUP_SIZEX"), 8);
+			OutEnvironment.SetDefine(TEXT("THREADGROUP_SIZEY"), 8);
+		}
+	}
+
+	/** Default constructor. */
+	FPostProcessGTAOTemporalFilterPSandCS() {}
+
+public:
+	FSceneTextureShaderParameters	SceneTextureParameters;
+	FPostProcessPassParameters		PostprocessParameter;
+	FShaderResourceParameter		HistoryTexture;
+	FShaderResourceParameter		HistoryTextureSampler;
+
+	FShaderResourceParameter		DepthHistoryTexture;
+	FShaderResourceParameter		DepthHistoryTextureSampler;
+
+
+	FShaderParameter				PrevScreenPositionScaleBias;
+	FScreenSpaceAOParameters		ScreenSpaceAOParams;
+	FShaderParameter OutTexture;
+
+
+	/** Initialization constructor. */
+	FPostProcessGTAOTemporalFilterPSandCS(const ShaderMetaType::CompiledShaderInitializerType& Initializer)
+		: FGlobalShader(Initializer)
+	{
+		PostprocessParameter.Bind(Initializer.ParameterMap);
+		SceneTextureParameters.Bind(Initializer);
+
+		HistoryTexture.Bind(Initializer.ParameterMap,				TEXT("HistoryTexture"));
+		HistoryTextureSampler.Bind(Initializer.ParameterMap,		TEXT("HistoryTextureSampler"));
+
+		DepthHistoryTexture.Bind(Initializer.ParameterMap,			TEXT("DepthHistoryTexture"));
+		DepthHistoryTextureSampler.Bind(Initializer.ParameterMap,	TEXT("DepthHistoryTextureSampler"));
+
+		PrevScreenPositionScaleBias.Bind(Initializer.ParameterMap, TEXT("PrevScreenPositionScaleBias"));
+		OutTexture.Bind(Initializer.ParameterMap, TEXT("OutTexture"));
+
+		ScreenSpaceAOParams.Bind(Initializer.ParameterMap);
+	}
+
+	template <typename TRHICmdList>
+	void SetParametersCS(TRHICmdList& RHICmdList, const FRenderingCompositePassContext& Context, FRHIUnorderedAccessView *OutUAV, const FGTAOTAAHistory& InputHistory)
+	{
+		const FFinalPostProcessSettings& Settings = Context.View.FinalPostProcessSettings;
+		FRHIComputeShader* ShaderRHI = GetComputeShader();
+
+		FGlobalShader::SetParameters<FViewUniformShaderParameters>(RHICmdList, ShaderRHI, Context.View.ViewUniformBuffer);
+		SceneTextureParameters.Set(RHICmdList, ShaderRHI, Context.View.FeatureLevel, ESceneTextureSetupMode::All);
+		PostprocessParameter.SetCS(ShaderRHI, Context, RHICmdList, TStaticSamplerState<SF_Point, AM_Clamp, AM_Clamp, AM_Clamp>::GetRHI());
+
+		ScreenSpaceAOParams.Set(RHICmdList, Context.View, ShaderRHI, InputTextureSize, FScreenSpaceAOParameters::ERandTexType::GTAO);
+
+		FIntPoint ViewportOffset = InputHistory.ViewportRect.Min;
+		FIntPoint ViewportExtent = InputHistory.ViewportRect.Size();
+		FIntPoint BufferSize = InputHistory.ReferenceBufferSize;
+
+		FVector4 PrevScreenPositionScaleBiasValue = FVector4(
+			 ViewportExtent.X * 0.5f / BufferSize.X,
+			-ViewportExtent.Y * 0.5f / BufferSize.Y,
+			(ViewportExtent.X * 0.5f + ViewportOffset.X) / BufferSize.X,
+			(ViewportExtent.Y * 0.5f + ViewportOffset.Y) / BufferSize.Y);
+
+
+		SetShaderValue(RHICmdList, ShaderRHI, PrevScreenPositionScaleBias, PrevScreenPositionScaleBiasValue);
+
+
+		// Set history shader parameters.
+		if (InputHistory.IsValid())
+		{
+			SetTextureParameter(RHICmdList, ShaderRHI, HistoryTexture, HistoryTextureSampler,
+				TStaticSamplerState<SF_Bilinear>::GetRHI(),
+				InputHistory.RT[0]->GetRenderTargetItem().TargetableTexture);
+		}
+		else
+		{
+			// Need to bind a white dummy
+			SetTextureParameter( RHICmdList, ShaderRHI, HistoryTexture, HistoryTextureSampler,
+				TStaticSamplerState<SF_Bilinear>::GetRHI(),
+				GSystemTextures.WhiteDummy->GetRenderTargetItem().ShaderResourceTexture);
+		}
+		RHICmdList.SetUAVParameter(ShaderRHI, OutTexture.GetBaseIndex(), OutUAV);
+	}
+
+	void SetParametersPS(const FRenderingCompositePassContext& Context, FIntPoint DestSize, FIntPoint InputTextureSize, const FGTAOTAAHistory& InputHistory)
+	{
+		const FFinalPostProcessSettings& Settings = Context.View.FinalPostProcessSettings;
+		FRHIPixelShader* ShaderRHI = GetPixelShader();
+
+		PostprocessParameter.SetPS(Context.RHICmdList, ShaderRHI, Context, TStaticSamplerState<SF_Point, AM_Clamp, AM_Clamp, AM_Clamp>::GetRHI());
+
+		FGlobalShader::SetParameters<FViewUniformShaderParameters>(Context.RHICmdList, ShaderRHI, Context.View.ViewUniformBuffer);
+		SceneTextureParameters.Set(Context.RHICmdList, ShaderRHI, Context.View.FeatureLevel, ESceneTextureSetupMode::All);
+
+		FIntPoint ViewportOffset = InputHistory.ViewportRect.Min;
+		FIntPoint ViewportExtent = InputHistory.ViewportRect.Size();
+		FIntPoint BufferSize = InputHistory.ReferenceBufferSize;
+
+		FVector4 PrevScreenPositionScaleBiasValue = FVector4(
+			ViewportExtent.X * 0.5f / BufferSize.X,
+			-ViewportExtent.Y * 0.5f / BufferSize.Y,
+			(ViewportExtent.X * 0.5f + ViewportOffset.X) / BufferSize.X,
+			(ViewportExtent.Y * 0.5f + ViewportOffset.Y) / BufferSize.Y);
+
+		SetShaderValue(Context.RHICmdList, ShaderRHI, PrevScreenPositionScaleBias, PrevScreenPositionScaleBiasValue);
+
+		if (InputHistory.IsValid())
+		{
+			SetTextureParameter(Context.RHICmdList, ShaderRHI, HistoryTexture, HistoryTextureSampler,
+				TStaticSamplerState<SF_Bilinear,AM_Border, AM_Border, AM_Border, 0,0, 0xffffffff >::GetRHI(),
+				InputHistory.RT[0]->GetRenderTargetItem().TargetableTexture);
+
+
+			SetTextureParameter(Context.RHICmdList, ShaderRHI, DepthHistoryTexture, DepthHistoryTextureSampler,
+				TStaticSamplerState<SF_Bilinear>::GetRHI(),
+				InputHistory.Depth[0]->GetRenderTargetItem().TargetableTexture);
+		}
+		else
+		{
+			// Need to bind a white dummy
+			SetTextureParameter(Context.RHICmdList, ShaderRHI, HistoryTexture, HistoryTextureSampler,
+				TStaticSamplerState<SF_Bilinear>::GetRHI(),
+				GSystemTextures.WhiteDummy->GetRenderTargetItem().ShaderResourceTexture);
+
+		}
+	}
+
+	template <typename TRHICmdList>
+	void UnsetParameters(TRHICmdList& RHICmdList)
+	{
+		FRHIComputeShader* ShaderRHI = GetComputeShader();
+		RHICmdList.SetUAVParameter(ShaderRHI, OutTexture.GetBaseIndex(), nullptr);
+	}
+
+
+	// FShader interface.
+	virtual bool Serialize(FArchive& Ar) override
+	{
+		bool bShaderHasOutdatedParameters = FGlobalShader::Serialize(Ar);
+		Ar << DepthHistoryTexture << DepthHistoryTextureSampler << HistoryTexture << HistoryTextureSampler << PostprocessParameter << SceneTextureParameters << PrevScreenPositionScaleBias << OutTexture;
+		return bShaderHasOutdatedParameters;
+	}
+
+	static const TCHAR* GetSourceFilename()
+	{
+		return TEXT("/Engine/Private/PostProcessAmbientOcclusion.usf");
+	}
+
+	static const TCHAR* GetFunctionName()
+	{
+		return bComputeShader ? TEXT("GTAOTemporalFilterCS") : TEXT("GTAOTemporalFilterPS");
+	}
+};
+
+IMPLEMENT_SHADER_TYPE2(FPostProcessGTAOTemporalFilterPSandCS<false>, SF_Pixel); 
+IMPLEMENT_SHADER_TYPE2(FPostProcessGTAOTemporalFilterPSandCS<true>, SF_Compute);
+
+
+FRCPassPostProcessAmbientOcclusion_GTAO_TemporalFilter::FRCPassPostProcessAmbientOcclusion_GTAO_TemporalFilter(const FSceneView& View, uint32 InDownScaleFactor,
+	const FGTAOTAAHistory& InInputHistory,
+	FGTAOTAAHistory* OutOutputHistory)
+	: InputHistory(InInputHistory),
+	OutputHistory(OutOutputHistory),
+	DownScaleFactor(InDownScaleFactor)
+{
+}
+
+
+template <uint32 ShaderQuality>
+void FRCPassPostProcessAmbientOcclusion_GTAO_TemporalFilter::DispatchCS(const FRenderingCompositePassContext& Context, FIntRect OutputViewRect, FIntPoint OutputTexSize)
+{
+	FSceneRenderTargets& SceneContext = FSceneRenderTargets::Get(Context.RHICmdList);
+	const FSceneRenderTargetItem& DestRenderTarget = PassOutputs[0].RequestSurface(Context);
+
+	FRHIAsyncComputeCommandListImmediate& RHICmdListComputeImmediate = FRHICommandListExecutor::GetImmediateAsyncComputeCommandList();
+
+	TShaderMapRef<FPostProcessGTAOTemporalFilterPSandCS<true>> ComputeShader(Context.GetShaderMap());
+	RHICmdListComputeImmediate.SetComputeShader(ComputeShader->GetComputeShader());
+	ComputeShader->SetParametersCS(RHICmdListComputeImmediate, Context, DestRenderTarget.UAV, InputHistory);
+
+	uint32 GroupSizeX = FMath::DivideAndRoundUp(OutputViewRect.Max.X, 8);
+	uint32 GroupSizeY = FMath::DivideAndRoundUp(OutputViewRect.Max.Y, 8);
+	DispatchComputeShader(RHICmdListComputeImmediate, *ComputeShader, GroupSizeX, GroupSizeY, 1);
+
+	ComputeShader->UnsetParameters(RHICmdListComputeImmediate);
+}
+
+
+void FRCPassPostProcessAmbientOcclusion_GTAO_TemporalFilter::Process(FRenderingCompositePassContext& Context)
+{
+	SCOPED_GPU_STAT(Context.RHICmdList, GTAO_TemporalFilter);
+	const FViewInfo& View = Context.View;
+
+	const FPooledRenderTargetDesc* InputDesc0 = GetInputDesc(ePId_Input0);
+	FSceneRenderTargets& SceneContext = FSceneRenderTargets::Get(Context.RHICmdList);
+
+	const FSceneRenderTargetItem& DestRenderTarget0 = PassOutputs[0].RequestSurface(Context);
+	const FSceneRenderTargetItem& DestRenderTarget1 = PassOutputs[1].RequestSurface(Context);
+
+	FIntPoint InputTexSize  = InputDesc0->Extent;
+	FIntPoint OutputTexSize = PassOutputs[0].RenderTargetDesc.Extent;
+	FIntRect OutputFullRect(FIntPoint(0, 0), OutputTexSize);
+	
+	FIntRect InputViewRect = View.ViewRect / DownScaleFactor;
+	FIntRect OutputViewRect = InputViewRect ;
+
+	// Whether to use camera cut shader permutation or not.
+	const bool bCameraCut = !InputHistory.IsValid() || View.bCameraCut;
+
+	OutputHistory->SafeRelease();
+	OutputHistory->RT[0]	= PassOutputs[0].PooledRenderTarget;
+	OutputHistory->Depth[0] = PassOutputs[1].PooledRenderTarget;
+	OutputHistory->ViewportRect = OutputViewRect;
+	OutputHistory->ReferenceBufferSize = OutputTexSize;
+
+	
+	{
+		//FRHIRenderPassInfo RPInfo(DestRenderTarget.TargetableTexture, ERenderTargetActions::Load_Store);
+		FRHITexture* RenderTargets[2] =
+		{
+			DestRenderTarget0.TargetableTexture,
+			DestRenderTarget1.TargetableTexture
+		};
+
+		Context.RHICmdList.TransitionResource(EResourceTransitionAccess::EWritable, DestRenderTarget0.TargetableTexture);
+		Context.RHICmdList.TransitionResource(EResourceTransitionAccess::EWritable, DestRenderTarget1.TargetableTexture);
+
+		FRHIRenderPassInfo RPInfo(2, RenderTargets, ERenderTargetActions::Load_Store);
+
+		Context.RHICmdList.BeginRenderPass(RPInfo, TEXT("GTAO_TemporalFilter"));
+		{
+			Context.SetViewportAndCallRHI(OutputFullRect);
+
+			DrawClearQuad(Context.RHICmdList, FLinearColor::White);
+
+
+			FGraphicsPipelineStateInitializer GraphicsPSOInit;
+			Context.RHICmdList.ApplyCachedRenderTargets(GraphicsPSOInit);
+
+			TShaderMapRef<FPostProcessVS> VertexShader(Context.GetShaderMap());
+			TShaderMapRef<FPostProcessGTAOTemporalFilterPSandCS<false>> PixelShader(Context.GetShaderMap());
+
+			// set the state
+			GraphicsPSOInit.BlendState = TStaticBlendState<>::GetRHI();
+			GraphicsPSOInit.RasterizerState = TStaticRasterizerState<>::GetRHI();
+			GraphicsPSOInit.DepthStencilState = TStaticDepthStencilState<false, CF_Always>::GetRHI();
+			GraphicsPSOInit.PrimitiveType = PT_TriangleList;
+			GraphicsPSOInit.bDepthBounds = false;
+
+			GraphicsPSOInit.BoundShaderState.VertexDeclarationRHI = GFilterVertexDeclaration.VertexDeclarationRHI;
+			GraphicsPSOInit.BoundShaderState.VertexShaderRHI = GETSAFERHISHADER_VERTEX(*VertexShader);
+			GraphicsPSOInit.BoundShaderState.PixelShaderRHI = GETSAFERHISHADER_PIXEL(*PixelShader);
+			SetGraphicsPipelineState(Context.RHICmdList, GraphicsPSOInit);
+
+			VertexShader->SetParameters(Context);
+			PixelShader->SetParametersPS(Context, OutputTexSize, InputTexSize, InputHistory);
+
+			DrawRectangle(
+				Context.RHICmdList,
+				OutputViewRect.Min.X, OutputViewRect.Min.Y, OutputViewRect.Width(), OutputViewRect.Height(),
+				InputViewRect.Min.X, InputViewRect.Min.Y, InputViewRect.Width(), InputViewRect.Height(),
+				OutputTexSize,
+				InputTexSize,
+				*VertexShader,
+				EDRF_UseTriangleOptimization);
+
+		}
+		Context.RHICmdList.EndRenderPass();
+	}
+	Context.RHICmdList.TransitionResource(EResourceTransitionAccess::EReadable, DestRenderTarget0.TargetableTexture);
+	Context.RHICmdList.TransitionResource(EResourceTransitionAccess::EReadable, DestRenderTarget1.TargetableTexture);
+}
+
+
+FPooledRenderTargetDesc FRCPassPostProcessAmbientOcclusion_GTAO_TemporalFilter::ComputeOutputDesc(EPassOutputId InPassOutputId) const
+{
+	FPooledRenderTargetDesc Ret;
+
+	Ret = GetInput(ePId_Input0)->GetOutput()->RenderTargetDesc;
+	Ret.Reset();
+	if (InPassOutputId == 0)
+	{
+		Ret.Format = PF_G8;
+	}
+	else
+	{
+		Ret.Format = PF_R32_FLOAT;
+	}
+
+
+	Ret.ClearValue = FClearValueBinding::None;
+	Ret.TargetableFlags &= ~TexCreate_DepthStencilTargetable;
+	Ret.TargetableFlags |= TexCreate_UAV;
+	Ret.TargetableFlags |= TexCreate_RenderTargetable | TexCreate_ShaderResource;
+
+	Ret.DebugName		= TEXT("GTAOTemporalAccumulate");
+	return Ret;
+}
+
+
+template<uint32 bComputeShader>
+class FPostProcessGTAOSpatialFilterPSandCS : public FGlobalShader
+{
+	DECLARE_SHADER_TYPE(FPostProcessGTAOSpatialFilterPSandCS, Global);
+
+	static bool ShouldCompilePermutation(const FGlobalShaderPermutationParameters& Parameters)
+	{
+		return IsFeatureLevelSupported(Parameters.Platform, ERHIFeatureLevel::SM5);
+	}
+
+	static void ModifyCompilationEnvironment(const FGlobalShaderPermutationParameters& Parameters, FShaderCompilerEnvironment& OutEnvironment)
+	{
+		FGlobalShader::ModifyCompilationEnvironment(Parameters, OutEnvironment);
+		OutEnvironment.SetDefine(TEXT("COMPUTE_SHADER"), bComputeShader);
+		if (bComputeShader)
+		{
+			OutEnvironment.SetDefine(TEXT("THREADGROUP_SIZEX"), 8);
+			OutEnvironment.SetDefine(TEXT("THREADGROUP_SIZEY"), 8);
+		}
+	}
+
+	FPostProcessGTAOSpatialFilterPSandCS() {}
+	FPostProcessPassParameters	PostprocessParameter;
+	FShaderParameter			OutTexture;
+	FShaderParameter			FilterParams;
+
+public:
+
+	/** Initialization constructor. */
+	FPostProcessGTAOSpatialFilterPSandCS(const ShaderMetaType::CompiledShaderInitializerType& Initializer)
+		: FGlobalShader(Initializer)
+	{
+		PostprocessParameter.Bind(Initializer.ParameterMap);
+		OutTexture.Bind(Initializer.ParameterMap, TEXT("OutTexture"));
+		FilterParams.Bind(Initializer.ParameterMap, TEXT("FilterParams"));
+	}
+
+	// FShader interface.
+	virtual bool Serialize(FArchive& Ar) override
+	{
+		bool bShaderHasOutdatedParameters = FGlobalShader::Serialize(Ar);
+		Ar << PostprocessParameter << OutTexture << FilterParams;
+		return bShaderHasOutdatedParameters;
+	}
+
+	void SetParametersPS(const FRenderingCompositePassContext& Context)
+	{
+		const FFinalPostProcessSettings& Settings = Context.View.FinalPostProcessSettings;
+		FRHIPixelShader* ShaderRHI = GetPixelShader();
+
+		FGlobalShader::SetParameters<FViewUniformShaderParameters>(Context.RHICmdList, ShaderRHI, Context.View.ViewUniformBuffer);
+
+		PostprocessParameter.SetPS(Context.RHICmdList, ShaderRHI, Context, TStaticSamplerState<SF_Bilinear, AM_Clamp, AM_Clamp, AM_Clamp>::GetRHI());
+	}
+
+
+	template <typename TRHICmdList>
+	void SetParametersCS(TRHICmdList& RHICmdList, const FRenderingCompositePassContext& Context, FIntPoint InputTextureSize, const FIntRect& OutputRect, FRHIUnorderedAccessView *OutUAV)
+	{
+		const FFinalPostProcessSettings& Settings = Context.View.FinalPostProcessSettings;
+		FRHIComputeShader* ShaderRHI = GetComputeShader();
+
+		FGlobalShader::SetParameters<FViewUniformShaderParameters>(RHICmdList, ShaderRHI, Context.View.ViewUniformBuffer);
+		PostprocessParameter.SetCS(ShaderRHI, Context, RHICmdList, TStaticSamplerState<SF_Bilinear, AM_Clamp, AM_Clamp, AM_Clamp>::GetRHI());
+
+		RHICmdList.SetUAVParameter(ShaderRHI, OutTexture.GetBaseIndex(), OutUAV);
+
+		FVector4 FilterParamsValue(OutputRect.Min.X, OutputRect.Min.Y, OutputRect.Width(), OutputRect.Height());
+		SetShaderValue(RHICmdList, ShaderRHI, FilterParams, FilterParamsValue);
+	}
+
+	template <typename TRHICmdList>
+	void UnsetParameters(TRHICmdList& RHICmdList)
+	{
+		FRHIComputeShader* ShaderRHI = GetComputeShader();
+		RHICmdList.SetUAVParameter(ShaderRHI, OutTexture.GetBaseIndex(), nullptr);
+	}
+
+
+	static const TCHAR* GetSourceFilename()
+	{
+		return TEXT("/Engine/Private/PostProcessAmbientOcclusion.usf");
+	}
+
+	static const TCHAR* GetFunctionName()
+	{
+		return bComputeShader ? TEXT("GTAOSpatialFilterCS") : TEXT("GTAOSpatialFilterPS");
+	}
+};
+
+IMPLEMENT_SHADER_TYPE2(FPostProcessGTAOSpatialFilterPSandCS<true>,  SF_Compute);
+IMPLEMENT_SHADER_TYPE2(FPostProcessGTAOSpatialFilterPSandCS<false>, SF_Pixel);
+
+FRCPassPostProcessAmbientOcclusion_GTAO_SpatialFilter::FRCPassPostProcessAmbientOcclusion_GTAO_SpatialFilter(const FSceneView& View)
+{
+}
+
+void FRCPassPostProcessAmbientOcclusion_GTAO_SpatialFilter::Process(FRenderingCompositePassContext& Context)
+{
+	SCOPED_GPU_STAT(Context.RHICmdList, GTAO_SpatialFilter);
+	const FViewInfo& View = Context.View;
+
+	const FPooledRenderTargetDesc* InputDesc0 = GetInputDesc(ePId_Input0);
+	FIntPoint InputTexSize = InputDesc0->Extent;
+	FSceneRenderTargets& SceneContext = FSceneRenderTargets::Get(Context.RHICmdList);
+	FIntPoint OutputTexSize = SceneContext.GetBufferSizeXY();
+	FIntRect OutputFullRect(FIntPoint(0, 0), OutputTexSize);
+
+	int DownScaleFactor = 2;
+
+	FIntRect InputViewRect = View.ViewRect / DownScaleFactor;
+	FIntRect OutputViewRect = InputViewRect;
+
+	const FSceneRenderTargetItem& DestRenderTarget = PassOutputs[0].RequestSurface(Context);
+
+	// Compute  Version
+	Context.RHICmdList.TransitionResource(EResourceTransitionAccess::ERWBarrier, EResourceTransitionPipeline::EGfxToCompute, DestRenderTarget.UAV, nullptr);
+
+	UnbindRenderTargets(Context.RHICmdList);
+
+	TShaderMapRef<FPostProcessGTAOSpatialFilterPSandCS<true>> ComputeShader(Context.GetShaderMap());
+	Context.RHICmdList.SetComputeShader(ComputeShader->GetComputeShader());
+	ComputeShader->SetParametersCS(Context.RHICmdList, Context, InputTexSize, OutputViewRect, DestRenderTarget.UAV);
+
+	uint32 GroupSizeX = FMath::DivideAndRoundUp(OutputViewRect.Width(), 8);
+	uint32 GroupSizeY = FMath::DivideAndRoundUp(OutputViewRect.Height(), 8);
+	DispatchComputeShader(Context.RHICmdList, *ComputeShader, GroupSizeX, GroupSizeY, 1);
+
+	ComputeShader->UnsetParameters(Context.RHICmdList);
+}
+
+
+
+FPooledRenderTargetDesc FRCPassPostProcessAmbientOcclusion_GTAO_SpatialFilter::ComputeOutputDesc(EPassOutputId InPassOutputId) const
+{
+	FPooledRenderTargetDesc Ret;
+
+	Ret = GetInput(ePId_Input0)->GetOutput()->RenderTargetDesc;
+	Ret.Reset();
+	Ret.Format = PF_G8;
+
+	Ret.ClearValue = FClearValueBinding::None;
+	Ret.TargetableFlags &= ~TexCreate_DepthStencilTargetable;
+	Ret.TargetableFlags |= TexCreate_UAV;
+	Ret.TargetableFlags |= TexCreate_RenderTargetable | TexCreate_ShaderResource;
+
+	Ret.DebugName = TEXT("GTAOFilter");
+	return Ret;
+}
+
+
+
+template<uint32 bComputeShader>
+class FPostProcessGTAOUpsamplePSandCS : public FGlobalShader
+{
+	DECLARE_SHADER_TYPE(FPostProcessGTAOUpsamplePSandCS, Global);
+
+	static bool ShouldCompilePermutation(const FGlobalShaderPermutationParameters& Parameters)
+	{
+		return IsFeatureLevelSupported(Parameters.Platform, ERHIFeatureLevel::SM5);
+	}
+
+	static void ModifyCompilationEnvironment(const FGlobalShaderPermutationParameters& Parameters, FShaderCompilerEnvironment& OutEnvironment)
+	{
+		FGlobalShader::ModifyCompilationEnvironment(Parameters, OutEnvironment);
+		OutEnvironment.SetDefine(TEXT("COMPUTE_SHADER"), bComputeShader);
+		if (bComputeShader)
+		{
+			OutEnvironment.SetDefine(TEXT("THREADGROUP_SIZEX"), 8);
+			OutEnvironment.SetDefine(TEXT("THREADGROUP_SIZEY"), 8);
+		}
+	}
+
+	FPostProcessGTAOUpsamplePSandCS() {}
+	FPostProcessPassParameters	  PostprocessParameter;
+	FShaderParameter OutTexture;
+	FScreenSpaceAOParameters ScreenSpaceAOParams;
+
+public:
+
+	/** Initialization constructor. */
+	FPostProcessGTAOUpsamplePSandCS(const ShaderMetaType::CompiledShaderInitializerType& Initializer)
+		: FGlobalShader(Initializer)
+	{
+		PostprocessParameter.Bind(Initializer.ParameterMap);
+		OutTexture.Bind(Initializer.ParameterMap, TEXT("OutTexture"));
+		ScreenSpaceAOParams.Bind(Initializer.ParameterMap);
+	}
+
+	// FShader interface.
+	virtual bool Serialize(FArchive& Ar) override
+	{
+		bool bShaderHasOutdatedParameters = FGlobalShader::Serialize(Ar);
+		Ar << PostprocessParameter << OutTexture << ScreenSpaceAOParams;
+		return bShaderHasOutdatedParameters;
+	}
+
+	void SetParametersPS(const FRenderingCompositePassContext& Context)
+	{
+		const FFinalPostProcessSettings& Settings = Context.View.FinalPostProcessSettings;
+		FRHIPixelShader* ShaderRHI = GetPixelShader();
+
+		FGlobalShader::SetParameters<FViewUniformShaderParameters>(Context.RHICmdList, ShaderRHI, Context.View.ViewUniformBuffer);
+
+		PostprocessParameter.SetPS(Context.RHICmdList, ShaderRHI, Context, TStaticSamplerState<SF_Bilinear, AM_Clamp, AM_Clamp, AM_Clamp>::GetRHI());
+	}
+
+	static const TCHAR* GetSourceFilename()
+	{
+		return TEXT("/Engine/Private/PostProcessAmbientOcclusion.usf");
+	}
+
+	static const TCHAR* GetFunctionName()
+	{
+		return bComputeShader ? TEXT("GTAOUpsampleCS") : TEXT("GTAOUpsamplePS");
+	}
+};
+
+IMPLEMENT_SHADER_TYPE2(FPostProcessGTAOUpsamplePSandCS<true>, SF_Compute);
+IMPLEMENT_SHADER_TYPE2(FPostProcessGTAOUpsamplePSandCS<false>, SF_Pixel);
+
+FRCPassPostProcessAmbientOcclusion_GTAO_Upsample::FRCPassPostProcessAmbientOcclusion_GTAO_Upsample()
+{
+}
+
+void FRCPassPostProcessAmbientOcclusion_GTAO_Upsample::Process(FRenderingCompositePassContext& Context)
+{
+	SCOPED_GPU_STAT(Context.RHICmdList, GTAO_Upsample);
+	const FViewInfo& View = Context.View;
+
+	const FPooledRenderTargetDesc* InputDesc0 = GetInputDesc(ePId_Input0);
+	FSceneRenderTargets& SceneContext = FSceneRenderTargets::Get(Context.RHICmdList);
+	const FSceneRenderTargetItem& DestRenderTarget = SceneContext.ScreenSpaceAO->GetRenderTargetItem();
+
+	// Do a simple upsample filter from source to dest
+	FIntPoint InputTexSize = InputDesc0->Extent;
+	FIntPoint OutputTexSize = SceneContext.GetBufferSizeXY();
+	FIntRect OutputFullRect(FIntPoint(0, 0), OutputTexSize);
+
+	int DownScaleFactor = 2;
+
+	FIntRect InputViewRect = View.ViewRect / DownScaleFactor;
+	FIntRect OutputViewRect = View.ViewRect;
+
+	Context.RHICmdList.TransitionResource(EResourceTransitionAccess::EWritable, DestRenderTarget.TargetableTexture);
+
+	FRHIRenderPassInfo RPInfo(DestRenderTarget.TargetableTexture, ERenderTargetActions::Load_Store);
+	Context.RHICmdList.BeginRenderPass(RPInfo, TEXT("GTAO_Filter"));
+	{
+		Context.SetViewportAndCallRHI(OutputFullRect);
+
+		FGraphicsPipelineStateInitializer GraphicsPSOInit;
+		Context.RHICmdList.ApplyCachedRenderTargets(GraphicsPSOInit);
+
+		TShaderMapRef<FPostProcessVS> VertexShader(Context.GetShaderMap());
+		TShaderMapRef<FPostProcessGTAOUpsamplePSandCS<false>> PixelShader(Context.GetShaderMap());
+
+		// set the state
+		GraphicsPSOInit.BlendState = TStaticBlendState<>::GetRHI();
+		GraphicsPSOInit.RasterizerState = TStaticRasterizerState<>::GetRHI();
+		GraphicsPSOInit.DepthStencilState = TStaticDepthStencilState<false, CF_Always>::GetRHI();
+		GraphicsPSOInit.PrimitiveType = PT_TriangleList;
+		GraphicsPSOInit.bDepthBounds = false;
+
+		GraphicsPSOInit.BoundShaderState.VertexDeclarationRHI = GFilterVertexDeclaration.VertexDeclarationRHI;
+		GraphicsPSOInit.BoundShaderState.VertexShaderRHI = GETSAFERHISHADER_VERTEX(*VertexShader);
+		GraphicsPSOInit.BoundShaderState.PixelShaderRHI = GETSAFERHISHADER_PIXEL(*PixelShader);
+		SetGraphicsPipelineState(Context.RHICmdList, GraphicsPSOInit);
+
+		VertexShader->SetParameters(Context);
+		PixelShader->SetParametersPS(Context);
+
+		DrawRectangle(
+			Context.RHICmdList,
+			OutputViewRect.Min.X, OutputViewRect.Min.Y, OutputViewRect.Width(), OutputViewRect.Height(),
+			InputViewRect.Min.X, InputViewRect.Min.Y, InputViewRect.Width(), InputViewRect.Height(),
+			OutputTexSize,
+			InputTexSize,
+			*VertexShader,
+			EDRF_UseTriangleOptimization);
+
+	}
+	Context.RHICmdList.EndRenderPass();
+
+	Context.RHICmdList.TransitionResource(EResourceTransitionAccess::EReadable, DestRenderTarget.TargetableTexture);
+}
+
+FPooledRenderTargetDesc FRCPassPostProcessAmbientOcclusion_GTAO_Upsample::ComputeOutputDesc(EPassOutputId InPassOutputId) const
+{
+	FPooledRenderTargetDesc Ret;
+
+	Ret = GetInput(ePId_Input0)->GetOutput()->RenderTargetDesc;
+	Ret.Reset();
+	Ret.Format = PF_G8;
+
+	Ret.ClearValue = FClearValueBinding::None;
+	Ret.TargetableFlags &= ~TexCreate_DepthStencilTargetable;
+	Ret.TargetableFlags |= TexCreate_UAV;
+	Ret.TargetableFlags |= TexCreate_RenderTargetable | TexCreate_ShaderResource;
+
+	Ret.DebugName = TEXT("GTAOFilter");
+	return Ret;
+}
+
diff --git a/Engine/Source/Runtime/Renderer/Private/CompositionLighting/PostProcessAmbientOcclusion.h b/Engine/Source/Runtime/Renderer/Private/CompositionLighting/PostProcessAmbientOcclusion.h
--- a/Engine/Source/Runtime/Renderer/Private/CompositionLighting/PostProcessAmbientOcclusion.h
+++ b/Engine/Source/Runtime/Renderer/Private/CompositionLighting/PostProcessAmbientOcclusion.h
@@ -148,7 +148,7 @@ class FRCPassPostProcessAmbientOcclusion_HorizonSearch : public TRenderingCompos
 {
 public:
 
-	FRCPassPostProcessAmbientOcclusion_HorizonSearch(const FSceneView& View, const ESSAOType AOType );
+	FRCPassPostProcessAmbientOcclusion_HorizonSearch(const FSceneView& View, uint32 DownScaleFactor, const ESSAOType AOType );
 
 	// interface FRenderingCompositePass ---------
 	virtual void Process(FRenderingCompositePassContext& Context) override;
@@ -156,13 +156,14 @@ public:
 	virtual FPooledRenderTargetDesc ComputeOutputDesc(EPassOutputId InPassOutputId) const override;
 
 	template <uint32 ShaderQuality>
-	void DispatchCS(const FRenderingCompositePassContext& Context, FIntPoint DestSize, FIntPoint TexSize);
+	void DispatchCS(const FRenderingCompositePassContext& Context, FIntRect ViewRect, FIntPoint DestSize, FIntPoint TexSize);
 
 	template <uint32 ShaderQuality>
 	FShader* SetShaderPS(const FRenderingCompositePassContext& Context, FGraphicsPipelineStateInitializer& GraphicsPSOInit, FIntPoint DestSize);
 
 private:
 	const ESSAOType AOType;
+	uint32		DownScaleFactor;
 };
 
 
@@ -170,7 +171,7 @@ class FRCPassPostProcessAmbientOcclusion_InnerIntegrate : public TRenderingCompo
 {
 public:
 
-	FRCPassPostProcessAmbientOcclusion_InnerIntegrate(const FSceneView& View, bool FinalOutput);
+	FRCPassPostProcessAmbientOcclusion_InnerIntegrate(const FSceneView& View, uint32 DownScaleFactor, bool FinalOutput);
 
 	// interface FRenderingCompositePass ---------
 	virtual void Process(FRenderingCompositePassContext& Context) override;
@@ -179,14 +180,15 @@ public:
 
 private:
 	const bool bFinalOutput;
+	uint32		DownScaleFactor;
 };
 
 
-class FRCPassPostProcessAmbientOcclusion_HorizonSearchInnerIntegrate : public TRenderingCompositePassBase<2, 1>
+class FRCPassPostProcessAmbientOcclusion_GTAOCombined : public TRenderingCompositePassBase<2, 1>
 {
 public:
 
-	FRCPassPostProcessAmbientOcclusion_HorizonSearchInnerIntegrate(const FSceneView& View, bool FinalOutput);
+	FRCPassPostProcessAmbientOcclusion_GTAOCombined(const FSceneView& View, uint32 DownScaleFactor, bool FinalOutput);
 
 	// interface FRenderingCompositePass ---------
 	virtual void Process(FRenderingCompositePassContext& Context) override;
@@ -201,5 +203,59 @@ public:
 
 private:
 	const bool bFinalOutput;
+	uint32		DownScaleFactor;
+};
+
+
+
+
+class FRCPassPostProcessAmbientOcclusion_GTAO_TemporalFilter : public TRenderingCompositePassBase<1, 2>
+{
+public:
+
+	FRCPassPostProcessAmbientOcclusion_GTAO_TemporalFilter(const FSceneView& View, uint32 DownScaleFactor, const FGTAOTAAHistory& InInputHistory, FGTAOTAAHistory* OutOutputHistory);
+
+	// interface FRenderingCompositePass ---------
+	virtual void Process(FRenderingCompositePassContext& Context) override;
+	virtual void Release() override { delete this; }
+	virtual FPooledRenderTargetDesc ComputeOutputDesc(EPassOutputId InPassOutputId) const override;
 
+	template <uint32 ShaderQuality>
+	void DispatchCS(const FRenderingCompositePassContext& Context, FIntRect OutputViewRect, FIntPoint OutputTexSize);
+
+private:
+
+	const FGTAOTAAHistory& InputHistory;
+	FGTAOTAAHistory* OutputHistory;
+	uint32		DownScaleFactor;
 };
+
+class FRCPassPostProcessAmbientOcclusion_GTAO_SpatialFilter : public TRenderingCompositePassBase<2, 1>
+{
+public:
+
+	FRCPassPostProcessAmbientOcclusion_GTAO_SpatialFilter(const FSceneView& View);
+
+	// interface FRenderingCompositePass ---------
+	virtual void Process(FRenderingCompositePassContext& Context) override;
+	virtual void Release() override { delete this; }
+	virtual FPooledRenderTargetDesc ComputeOutputDesc(EPassOutputId InPassOutputId) const override;
+
+private:
+
+};
+
+class FRCPassPostProcessAmbientOcclusion_GTAO_Upsample : public TRenderingCompositePassBase<1, 1>
+{
+public:
+
+	FRCPassPostProcessAmbientOcclusion_GTAO_Upsample();
+
+	// interface FRenderingCompositePass ---------
+	virtual void Process(FRenderingCompositePassContext& Context) override;
+	virtual void Release() override { delete this; }
+	virtual FPooledRenderTargetDesc ComputeOutputDesc(EPassOutputId InPassOutputId) const override;
+
+private:
+
+};
\ No newline at end of file
diff --git a/Engine/Source/Runtime/Renderer/Private/PostProcess/SceneRenderTargets.cpp b/Engine/Source/Runtime/Renderer/Private/PostProcess/SceneRenderTargets.cpp
--- a/Engine/Source/Runtime/Renderer/Private/PostProcess/SceneRenderTargets.cpp
+++ b/Engine/Source/Runtime/Renderer/Private/PostProcess/SceneRenderTargets.cpp
@@ -2342,19 +2342,6 @@ void FSceneRenderTargets::AllocateDeferredShadingPathRenderTargets(FRHICommandLi
 			GRenderTargetPool.FindFreeElement(RHICmdList, Desc, ScreenSpaceAO, TEXT("ScreenSpaceAO"), true, ERenderTargetTransience::NonTransient);
 		}
 
-		// Create the screen space Horizon texture for Async GTAO
-		{
-			FPooledRenderTargetDesc Desc(FPooledRenderTargetDesc::Create2DDesc(BufferSize, PF_R8G8, FClearValueBinding::White, TexCreate_None, TexCreate_RenderTargetable, false));
-			Desc.Flags |= GFastVRamConfig.ScreenSpaceAO;
-
-			if (CurrentFeatureLevel >= ERHIFeatureLevel::SM5)
-			{
-				Desc.TargetableFlags |= TexCreate_UAV;
-			}
-			GRenderTargetPool.FindFreeElement(RHICmdList, Desc, ScreenSpaceGTAOHorizons, TEXT("ScreenSpaceGTAOHorizons"), true, ERenderTargetTransience::NonTransient);
-		}
-		
-		
 		{
 			const int32 TranslucencyLightingVolumeDim = GetTranslucencyLightingVolumeDim();
 
@@ -2688,6 +2675,7 @@ void FSceneRenderTargets::ReleaseAllTargets()
 	DBufferB.SafeRelease();
 	DBufferC.SafeRelease();
 	ScreenSpaceAO.SafeRelease();
+	ScreenSpaceGTAOHorizons.SafeRelease();
 	QuadOverdrawBuffer.SafeRelease();
 	LightAttenuation.SafeRelease();
 	LightAccumulation.SafeRelease();
diff --git a/Engine/Source/Runtime/Renderer/Private/SceneRendering.h b/Engine/Source/Runtime/Renderer/Private/SceneRendering.h
--- a/Engine/Source/Runtime/Renderer/Private/SceneRendering.h
+++ b/Engine/Source/Runtime/Renderer/Private/SceneRendering.h
@@ -744,6 +744,42 @@ struct FScreenSpaceDenoiserHistory
 	}
 };
 
+
+
+// Structure in charge of storing all information about GTAO history.
+struct FGTAOTAAHistory
+{
+	// Number of render target in the history.
+	static constexpr uint32 kRenderTargetCount = 2;
+
+	// Render targets holding's pixel history.
+	//  scene color's RGBA are in RT[0].
+	TRefCountPtr<IPooledRenderTarget> RT[kRenderTargetCount];
+	TRefCountPtr<IPooledRenderTarget> Depth[kRenderTargetCount];
+
+	// Reference size of RT. Might be different than RT's actual size to handle down res.
+	FIntPoint ReferenceBufferSize;
+
+	// Viewport coordinate of the history in RT according to ReferenceBufferSize.
+	FIntRect ViewportRect;
+
+	void SafeRelease()
+	{
+		for (uint32 i = 0; i < kRenderTargetCount; i++)
+		{
+			RT[i].SafeRelease();
+			Depth[i].SafeRelease();
+		}
+	}
+
+	bool IsValid() const
+	{
+		return RT[0].IsValid();
+	}
+};
+
+
+
 // Structure that hold all information related to previous frame.
 struct FPreviousViewInfo
 {
@@ -778,6 +814,9 @@ struct FPreviousViewInfo
 	// History for the ambient occlusion
 	FScreenSpaceDenoiserHistory AmbientOcclusionHistory;
 
+	// History for GTAO
+	FGTAOTAAHistory				 GTAOHistory;
+
 	// History for global illumination
 	FScreenSpaceDenoiserHistory DiffuseIndirectHistory;
 
