GTAO. Turn off new method by default

diff --git a/Engine/Source/Runtime/Renderer/Private/CompositionLighting/PostProcessAmbientOcclusion.cpp b/Engine/Source/Runtime/Renderer/Private/CompositionLighting/PostProcessAmbientOcclusion.cpp
--- a/Engine/Source/Runtime/Renderer/Private/CompositionLighting/PostProcessAmbientOcclusion.cpp
+++ b/Engine/Source/Runtime/Renderer/Private/CompositionLighting/PostProcessAmbientOcclusion.cpp
@@ -95,7 +95,7 @@ static TAutoConsoleVariable<int32> CVarAmbientOcclusionDepthBoundsTest(
 
 static TAutoConsoleVariable<int32> CVarAmbientOcclusionMethod(
 	TEXT("r.AmbientOcclusion.Method"),
-	1,
+	0,
 	TEXT("Select between SSAO methods \n ")
 	TEXT("0: SSAO (default)\n ")
 	TEXT("1: GTAO\n "),
