Water refraction downsample - pick furthermost depth to minimize downsampling artifacts

diff --git a/Engine/Shaders/Private/SingleLayerWaterComposite.usf b/Engine/Shaders/Private/SingleLayerWaterComposite.usf
--- a/Engine/Shaders/Private/SingleLayerWaterComposite.usf
+++ b/Engine/Shaders/Private/SingleLayerWaterComposite.usf
@@ -297,5 +297,14 @@ void WaterRefractionCopyPS(
 
 	// Downsample and average all 2x2 quads
 	OutColorAndDepth.xyz = Texture2DSampleLevel(SceneColorCopyDownsampleTexture, SceneColorCopyDownsampleSampler, BufferUV, 0).xyz;
-	OutColorAndDepth.w = ConvertFromDeviceZ(Texture2DSampleLevel(SceneDepthCopyDownsampleTexture, SceneDepthCopyDownsampleSampler, BufferUV, 0).r) / SINGLE_LAYER_WATER_DEPTH_SCALE;
+
+	float DeviceDepthZ = 0.0f;
+#if DOWNSAMPLE_REFRACTION
+	// Pick furthermost depth to minimize downsampling artifacts (UE4 is using reverse-z)
+	float4 Depth4 = SceneDepthCopyDownsampleTexture.Gather(SceneDepthCopyDownsampleSampler, BufferUV);
+	DeviceDepthZ = min(min(min(Depth4.x, Depth4.y), Depth4.z), Depth4.w);
+#else
+	DeviceDepthZ = Texture2DSampleLevel(SceneDepthCopyDownsampleTexture, SceneDepthCopyDownsampleSampler, BufferUV, 0).x;
+#endif
+	OutColorAndDepth.w = ConvertFromDeviceZ(DeviceDepthZ) / SINGLE_LAYER_WATER_DEPTH_SCALE;
 }
\ No newline at end of file
diff --git a/Engine/Source/Runtime/Renderer/Private/SingleLayerWaterRendering.cpp b/Engine/Source/Runtime/Renderer/Private/SingleLayerWaterRendering.cpp
--- a/Engine/Source/Runtime/Renderer/Private/SingleLayerWaterRendering.cpp
+++ b/Engine/Source/Runtime/Renderer/Private/SingleLayerWaterRendering.cpp
@@ -359,6 +359,10 @@ class FWaterRefractionCopyPS : public FGlobalShader
 		RENDER_TARGET_BINDING_SLOTS()
 	END_SHADER_PARAMETER_STRUCT()
 
+	class FDownsampleRefraction : SHADER_PERMUTATION_BOOL("DOWNSAMPLE_REFRACTION");
+
+	using FPermutationDomain = TShaderPermutationDomain<FDownsampleRefraction>;
+
 	static bool ShouldCompilePermutation(const FGlobalShaderPermutationParameters& Parameters)
 	{
 		return IsFeatureLevelSupported(Parameters.Platform, ERHIFeatureLevel::SM5);
@@ -442,7 +446,9 @@ void FDeferredShadingSceneRenderer::CopySingleLayerWaterTextures(FRHICommandList
 	PassParameters->SVPositionToSourceTextureUV = FVector2D(RefractionDownsampleFactor / float(SceneContext.GetBufferSizeXY().X), RefractionDownsampleFactor / float(SceneContext.GetBufferSizeXY().Y));
 	PassParameters->RenderTargets[0] = FRenderTargetBinding(TargetColorTexture, ERenderTargetLoadAction::ELoad);
 
-	auto PixelShader = View.ShaderMap->GetShader<FWaterRefractionCopyPS>();
+	FWaterRefractionCopyPS::FPermutationDomain PermutationVector;
+	PermutationVector.Set<FWaterRefractionCopyPS::FDownsampleRefraction>(RefractionDownsampleFactor > 1);
+	auto PixelShader = View.ShaderMap->GetShader<FWaterRefractionCopyPS>(PermutationVector);
 
 	const FIntRect RefractionViewRect = FIntRect(FIntPoint(0, 0), FIntPoint::DivideAndRoundDown(View.ViewRect.Size(), RefractionDownsampleFactor));
 
