Removing mulstiscattering lut small format that was deemed usless and result in different look on mobile with Foirtnite non-earth alike atmosphere.

diff --git a/Engine/Shaders/Private/SkyAtmosphere.usf b/Engine/Shaders/Private/SkyAtmosphere.usf
--- a/Engine/Shaders/Private/SkyAtmosphere.usf
+++ b/Engine/Shaders/Private/SkyAtmosphere.usf
@@ -60,10 +60,6 @@
  // FASTSKY mapping is done based on Light0 as main light
 #define FASTSKY_LIGHT_INDEX 0
 
-// This (arbitrary) multi scattering storage scale is used to make sure the LUT values are in an acceptable range when using r.SkyAtmosphere.MultiScatteringLUT.UseSmallFormat is 1.
-#define MULTISCATTERING_STORE_SCALE (4.0f * PI * 5.0f)
-#define MULTISCATTERING_STORE_SCALE_INV 1.0f/MULTISCATTERING_STORE_SCALE
-
 // Only available for shaders ran per view, so this constant should not be accessed in common code such as IntegrateSingleScatteredLuminance for instance.
 float AerialPerspectiveStartDepth;
 
@@ -370,7 +366,7 @@ float3 GetMultipleScattering(float3 WorlPos, float ViewZenithCosAngle)
 	float2 UV = saturate(float2(ViewZenithCosAngle*0.5f + 0.5f, (length(WorlPos) - Atmosphere.BottomRadius) / (Atmosphere.TopRadius - Atmosphere.BottomRadius)));
 	// We do no apply UV transform to sub range here as it has minimal impact.
 	float3 MultiScatteredLuminance = MultiScatteredLuminanceLutTexture.SampleLevel(MultiScatteredLuminanceLutTextureSampler, UV, 0).rgb;
-	return MultiScatteredLuminance * MULTISCATTERING_STORE_SCALE_INV;
+	return MultiScatteredLuminance;
 }
 
 
@@ -1011,7 +1007,7 @@ void RenderMultiScatteredLuminanceLutCS(uint3 ThreadId : SV_DispatchThreadID)
 
 	// MultipleScatteringFactor can be applied here because the LUT is compute every frame
 	// L is pre-exposed since InScatteredLuminance is computed from pre-exposed sun light. So multi-scattering contribution is pre-exposed.
-	MultiScatteredLuminanceLutUAV[int2(PixPos)] = L * Atmosphere.MultiScatteringFactor * MULTISCATTERING_STORE_SCALE;
+	MultiScatteredLuminanceLutUAV[int2(PixPos)] = L * Atmosphere.MultiScatteringFactor;
 }
 
 
diff --git a/Engine/Source/Runtime/Renderer/Private/SkyAtmosphereRendering.cpp b/Engine/Source/Runtime/Renderer/Private/SkyAtmosphereRendering.cpp
--- a/Engine/Source/Runtime/Renderer/Private/SkyAtmosphereRendering.cpp
+++ b/Engine/Source/Runtime/Renderer/Private/SkyAtmosphereRendering.cpp
@@ -171,11 +171,6 @@ static TAutoConsoleVariable<float> CVarSkyAtmosphereMultiScatteringLUTHighQualit
 	TEXT("The when enabled, 64 samples are used instead of 2, resulting in a more accurate multi scattering approximation (but also more expenssive).\n"),
 	ECVF_RenderThreadSafe | ECVF_Scalability);
 
-static TAutoConsoleVariable<int32> CVarSkyAtmosphereMultiScatteringLUTUseSmallFormat(
-	TEXT("r.SkyAtmosphere.MultiScatteringLUT.UseSmallFormat"), 0,
-	TEXT("If true, the multi-scattering LUT will use a small R8BG8B8A8 format to store data at lower quality."),
-	ECVF_RenderThreadSafe | ECVF_Scalability);
-
 static TAutoConsoleVariable<float> CVarSkyAtmosphereMultiScatteringLUTWidth(
 	TEXT("r.SkyAtmosphere.MultiScatteringLUT.Width"), 32,
 	TEXT(""),
@@ -918,7 +913,6 @@ void InitSkyAtmosphereForScene(FRHICommandListImmediate& RHICmdList, FScene* Sce
 		// Initialise per scene/atmosphere resources
 		//
 		const bool TranstmittanceLUTUseSmallFormat = CVarSkyAtmosphereTransmittanceLUTUseSmallFormat.GetValueOnRenderThread() > 0;
-		const bool MultiScatteringLUTUseSmallFormat = CVarSkyAtmosphereMultiScatteringLUTUseSmallFormat.GetValueOnRenderThread() > 0;
 
 		TRefCountPtr<IPooledRenderTarget>& TransmittanceLutTexture = SkyInfo.GetTransmittanceLutTexture();
 		Desc = FPooledRenderTargetDesc::Create2DDesc(
@@ -929,7 +923,7 @@ void InitSkyAtmosphereForScene(FRHICommandListImmediate& RHICmdList, FScene* Sce
 		TRefCountPtr<IPooledRenderTarget>& MultiScatteredLuminanceLutTexture = SkyInfo.GetMultiScatteredLuminanceLutTexture();
 		Desc = FPooledRenderTargetDesc::Create2DDesc(
 			FIntPoint(MultiScatteredLuminanceLutWidth, MultiScatteredLuminanceLutHeight),
-			MultiScatteringLUTUseSmallFormat ? TextureLUTSmallFormat : TextureLUTFormat, FClearValueBinding::None, TexCreate_HideInVisualizeTexture, TexCreate_ShaderResource | TexCreate_UAV, false);
+			TextureLUTFormat, FClearValueBinding::None, TexCreate_HideInVisualizeTexture, TexCreate_ShaderResource | TexCreate_UAV, false);
 		GRenderTargetPool.FindFreeElement(RHICmdList, Desc, MultiScatteredLuminanceLutTexture, TEXT("MultiScatteredLuminanceLutTexture"), true, ERenderTargetTransience::Transient);
 
 		if (CVarSkyAtmosphereDistantSkyLightLUT.GetValueOnRenderThread() > 0)
