Fixing ShouldCompileTiledShadersForPlatform linker error

diff --git a/Engine/Source/Runtime/Renderer/Private/ScreenSpaceRayTracing.cpp b/Engine/Source/Runtime/Renderer/Private/ScreenSpaceRayTracing.cpp
--- a/Engine/Source/Runtime/Renderer/Private/ScreenSpaceRayTracing.cpp
+++ b/Engine/Source/Runtime/Renderer/Private/ScreenSpaceRayTracing.cpp
@@ -112,8 +112,6 @@ bool IsSSRTemporalPassRequired(const FViewInfo& View)
 	return View.AntiAliasingMethod != AAM_TemporalAA || CVarSSRTemporal.GetValueOnRenderThread() != 0;
 }
 
-bool ShouldCompileTiledShadersForPlatform(EShaderPlatform ShaderPlatform);
-
 namespace
 {
 
@@ -358,6 +356,7 @@ class FScreenSpaceReflectionsTileVS : public FGlobalShader
 
 	static bool ShouldCompilePermutation(const FGlobalShaderPermutationParameters& Parameters)
 	{
+		extern bool ShouldCompileTiledShadersForPlatform(EShaderPlatform ShaderPlatform);
 		return ShouldCompileTiledShadersForPlatform(Parameters.Platform);
 	}
 
diff --git a/Engine/Source/Runtime/Renderer/Private/SingleLayerWaterRendering.cpp b/Engine/Source/Runtime/Renderer/Private/SingleLayerWaterRendering.cpp
--- a/Engine/Source/Runtime/Renderer/Private/SingleLayerWaterRendering.cpp
+++ b/Engine/Source/Runtime/Renderer/Private/SingleLayerWaterRendering.cpp
@@ -116,6 +116,7 @@ bool ShouldRenderSingleLayerWaterSkippedRenderEditorNotification(const TArray<FV
 
 bool ShouldCompileTiledShadersForPlatform(EShaderPlatform ShaderPlatform)
 {
+	// Temporary disable tiled water shaders until compilation issues are resolved.
 	return IsFeatureLevelSupported(ShaderPlatform, ERHIFeatureLevel::SM5) && RHISupportsDrawIndirect(ShaderPlatform)
 		&& !IsSwitchPlatform(ShaderPlatform) && !IsVulkanPlatform(ShaderPlatform); // Switch does not use tiling and Vulkan gives error with WaterTileCatergorisationCS usage of atomic.
 }
