Added debug CVar r.Water.SingleLayer.RefractionFullPrecision to be able to see what artifacts come from reduced water refraction depth precision

diff --git a/Engine/Source/Runtime/Renderer/Private/SingleLayerWaterRendering.cpp b/Engine/Source/Runtime/Renderer/Private/SingleLayerWaterRendering.cpp
--- a/Engine/Source/Runtime/Renderer/Private/SingleLayerWaterRendering.cpp
+++ b/Engine/Source/Runtime/Renderer/Private/SingleLayerWaterRendering.cpp
@@ -49,6 +49,14 @@ FAutoConsoleVariableRef CVarWaterSingleLayerRefractionDownsampleFactor(
 	ECVF_Scalability | ECVF_RenderThreadSafe
 );
 
+int32 GSingleLayerWaterRefractionFullPrecision = 0;
+FAutoConsoleVariableRef CVarWaterSingleLayerRefractionFullPrecision(
+	TEXT("r.Water.SingleLayer.RefractionFullPrecision"),
+	GSingleLayerWaterRefractionFullPrecision,
+	TEXT("Whether to pack refraction color and depth in RGBA Float32. To be used as a debug option to find issues with refraction depth precision."),
+	ECVF_Scalability | ECVF_RenderThreadSafe
+);
+
 static TAutoConsoleVariable<int32> CVarWaterSingleLayerSSR(
 	TEXT("r.Water.SingleLayer.SSR"), 1,
 	TEXT("Enable SSR for the single water renderring system."),
@@ -426,7 +434,7 @@ void FDeferredShadingSceneRenderer::CopySingleLayerWaterTextures(FRHICommandList
 	const int32 RefractionDownsampleFactor = FMath::Clamp(GSingleLayerWaterRefractionDownsampleFactor, 1, 8);
 	const FIntPoint RefractionResolution = FIntPoint::DivideAndRoundDown(SceneContext.GetBufferSizeXY(), RefractionDownsampleFactor);
 
-	FPooledRenderTargetDesc Desc(FPooledRenderTargetDesc::Create2DDesc(RefractionResolution, PF_FloatRGBA, SceneContext.GetDefaultColorClear(), TexCreate_None, TexCreate_RenderTargetable, false));
+	FPooledRenderTargetDesc Desc(FPooledRenderTargetDesc::Create2DDesc(RefractionResolution, GSingleLayerWaterRefractionFullPrecision ? PF_A32B32G32R32F : PF_FloatRGBA, SceneContext.GetDefaultColorClear(), TexCreate_None, TexCreate_RenderTargetable, false));
 	Desc.NumSamples = SceneContext.GetNumSceneColorMSAASamples(CurrentFeatureLevel);
 	GRenderTargetPool.FindFreeElement(RHICmdList, Desc, PassData.SceneColorAndDepthWithoutSingleLayerWater, TEXT("SceneColorAndDepthWithoutSingleLayerWater"));
 
