Fixing build on platforms which use clang.

diff --git a/Engine/Source/Runtime/Renderer/Private/ScreenSpaceRayTracing.cpp b/Engine/Source/Runtime/Renderer/Private/ScreenSpaceRayTracing.cpp
--- a/Engine/Source/Runtime/Renderer/Private/ScreenSpaceRayTracing.cpp
+++ b/Engine/Source/Runtime/Renderer/Private/ScreenSpaceRayTracing.cpp
@@ -112,6 +112,8 @@ bool IsSSRTemporalPassRequired(const FViewInfo& View)
 	return View.AntiAliasingMethod != AAM_TemporalAA || CVarSSRTemporal.GetValueOnRenderThread() != 0;
 }
 
+bool ShouldCompileTiledShadersForPlatform(EShaderPlatform ShaderPlatform);
+
 namespace
 {
 
@@ -356,8 +358,7 @@ class FScreenSpaceReflectionsTileVS : public FGlobalShader
 
 	static bool ShouldCompilePermutation(const FGlobalShaderPermutationParameters& Parameters)
 	{
-		extern bool ShouldCompileTiledShadersForPlatform(EShaderPlatform ShaderPlatform);
-		return ShouldCompileTiledShadersForPlatform(Parameters.Platform);
+		return ::ShouldCompileTiledShadersForPlatform(Parameters.Platform);
 	}
 
 	static void ModifyCompilationEnvironment(const FGlobalShaderPermutationParameters& Parameters, FShaderCompilerEnvironment& OutEnvironment)
