Sort masked geometry to the end of prepass to make maximum use of postZ
Measured ~.3ms average GPU savings on CL 9016249

diff --git a/Engine/Source/Runtime/Renderer/Private/DepthRendering.cpp b/Engine/Source/Runtime/Renderer/Private/DepthRendering.cpp
--- a/Engine/Source/Runtime/Renderer/Private/DepthRendering.cpp
+++ b/Engine/Source/Runtime/Renderer/Private/DepthRendering.cpp
@@ -44,6 +44,14 @@ static TAutoConsoleVariable<int32> CVarRHICmdFlushRenderThreadTasksPrePass(
 	0,
 	TEXT("Wait for completion of parallel render thread tasks at the end of the pre pass.  A more granular version of r.RHICmdFlushRenderThreadTasks. If either r.RHICmdFlushRenderThreadTasks or r.RHICmdFlushRenderThreadTasksPrePass is > 0 we will flush."));
 
+static int32 GEarlyZSortMasked = 1;
+static FAutoConsoleVariableRef CVarSortPrepassMasked(
+	TEXT("r.EarlyZSortMasked"),
+	GEarlyZSortMasked,
+	TEXT("Sort EarlyZ masked draws to the end of the draw order.\n"),
+	ECVF_Default
+);
+
 const TCHAR* GetDepthDrawingModeString(EDepthDrawingMode Mode)
 {
 	switch (Mode)
@@ -683,11 +691,30 @@ bool FDeferredShadingSceneRenderer::RenderPrePassHMD(FRHICommandListImmediate& R
 	return true;
 }
 
+FMeshDrawCommandSortKey CalculateDepthPassMeshStaticSortKey(EBlendMode BlendMode, const FMeshMaterialShader* VertexShader, const FMeshMaterialShader* PixelShader)
+{
+	FMeshDrawCommandSortKey SortKey;
+	if (GEarlyZSortMasked)
+	{
+		SortKey.BasePass.VertexShaderHash = PointerHash(VertexShader) & 0xFFFF;
+		SortKey.BasePass.PixelShaderHash = PointerHash(PixelShader);
+		SortKey.BasePass.Masked = BlendMode == EBlendMode::BLEND_Masked ? 1 : 0;
+	}
+	else
+	{
+		SortKey.Generic.VertexShaderHash = PointerHash(VertexShader);
+		SortKey.Generic.PixelShaderHash = PointerHash(PixelShader);
+	}
+	
+	return SortKey;
+}
+
 template<bool bPositionOnly>
 void FDepthPassMeshProcessor::Process(
 	const FMeshBatch& RESTRICT MeshBatch,
 	uint64 BatchElementMask,
 	int32 StaticMeshId,
+	EBlendMode BlendMode,
 	const FPrimitiveSceneProxy* RESTRICT PrimitiveSceneProxy,
 	const FMaterialRenderProxy& RESTRICT MaterialRenderProxy,
 	const FMaterial& RESTRICT MaterialResource,
@@ -723,7 +750,7 @@ void FDepthPassMeshProcessor::Process(
 	FDepthOnlyShaderElementData ShaderElementData(0.0f);
 	ShaderElementData.InitializeMeshMaterialData(ViewIfDynamicMeshCommand, PrimitiveSceneProxy, MeshBatch, StaticMeshId, true);
 
-	const FMeshDrawCommandSortKey SortKey = CalculateMeshStaticSortKey(DepthPassShaders.VertexShader, DepthPassShaders.PixelShader);
+	const FMeshDrawCommandSortKey SortKey = CalculateDepthPassMeshStaticSortKey(BlendMode, DepthPassShaders.VertexShader, DepthPassShaders.PixelShader);
 
 	BuildMeshDrawCommands(
 		MeshBatch,
@@ -793,7 +820,7 @@ void FDepthPassMeshProcessor::AddMeshBatch(const FMeshBatch& RESTRICT MeshBatch,
 			{
 				const FMaterialRenderProxy& DefaultProxy = *UMaterial::GetDefaultMaterial(MD_Surface)->GetRenderProxy();
 				const FMaterial& DefaultMaterial = *DefaultProxy.GetMaterial(FeatureLevel);
-				Process<true>(MeshBatch, BatchElementMask, StaticMeshId, PrimitiveSceneProxy, DefaultProxy, DefaultMaterial, MeshFillMode, MeshCullMode);
+				Process<true>(MeshBatch, BatchElementMask, StaticMeshId, BlendMode, PrimitiveSceneProxy, DefaultProxy, DefaultMaterial, MeshFillMode, MeshCullMode);
 			}
 			else
 			{
@@ -811,7 +838,7 @@ void FDepthPassMeshProcessor::AddMeshBatch(const FMeshBatch& RESTRICT MeshBatch,
 						EffectiveMaterial = EffectiveMaterialRenderProxy->GetMaterial(FeatureLevel);
 					}
 
-					Process<false>(MeshBatch, BatchElementMask, StaticMeshId, PrimitiveSceneProxy, *EffectiveMaterialRenderProxy, *EffectiveMaterial, MeshFillMode, MeshCullMode);
+					Process<false>(MeshBatch, BatchElementMask, StaticMeshId, BlendMode, PrimitiveSceneProxy, *EffectiveMaterialRenderProxy, *EffectiveMaterial, MeshFillMode, MeshCullMode);
 				}
 			}
 		}
diff --git a/Engine/Source/Runtime/Renderer/Private/DepthRendering.h b/Engine/Source/Runtime/Renderer/Private/DepthRendering.h
--- a/Engine/Source/Runtime/Renderer/Private/DepthRendering.h
+++ b/Engine/Source/Runtime/Renderer/Private/DepthRendering.h
@@ -222,6 +222,7 @@ private:
 		const FMeshBatch& MeshBatch,
 		uint64 BatchElementMask,
 		int32 StaticMeshId,
+		EBlendMode BlendMode,
 		const FPrimitiveSceneProxy* RESTRICT PrimitiveSceneProxy,
 		const FMaterialRenderProxy& RESTRICT MaterialRenderProxy,
 		const FMaterial& RESTRICT MaterialResource,
