Copying //UE4/Dev-Physics to Dev-Main (//UE4/Dev-Main) [at] 9098278

diff --git a/Engine/Source/Runtime/Renderer/Private/DeferredShadingRenderer.h b/Engine/Source/Runtime/Renderer/Private/DeferredShadingRenderer.h
--- a/Engine/Source/Runtime/Renderer/Private/DeferredShadingRenderer.h
+++ b/Engine/Source/Runtime/Renderer/Private/DeferredShadingRenderer.h
@@ -224,7 +224,7 @@ private:
 	void RenderDiffuseIndirectAndAmbientOcclusion(FRHICommandListImmediate& RHICmdList);
 
 	/** Renders sky lighting and reflections that can be done in a deferred pass. */
-	void RenderDeferredReflectionsAndSkyLighting(FRHICommandListImmediate& RHICmdList, TRefCountPtr<IPooledRenderTarget>& DynamicBentNormalAO, TRefCountPtr<IPooledRenderTarget>& VelocityRT);
+	void RenderDeferredReflectionsAndSkyLighting(FRHICommandListImmediate& RHICmdList, TRefCountPtr<IPooledRenderTarget>& DynamicBentNormalAO, TRefCountPtr<IPooledRenderTarget>& VelocityRT, const struct FHairStrandsDatas* HairDatas);
 
 	/** Computes DFAO, modulates it to scene color (which is assumed to contain diffuse indirect lighting), and stores the output bent normal for use occluding specular. */
 	void RenderDFAOAsIndirectShadowing(
@@ -283,7 +283,7 @@ private:
 	void GatherAndSortLights(FSortedLightSetSceneInfo& OutSortedLights);
 
 	/** Renders the scene's lighting. */
-	void RenderLights(FRHICommandListImmediate& RHICmdList, FSortedLightSetSceneInfo &SortedLightSet);
+	void RenderLights(FRHICommandListImmediate& RHICmdList, FSortedLightSetSceneInfo &SortedLightSet, const struct FHairStrandsDatas* HairDatas);
 
 	/** Renders an array of lights for the stationary light overlap viewmode. */
 	void RenderLightArrayForOverlapViewmode(FRHICommandListImmediate& RHICmdList, const TSparseArray<FLightSceneInfoCompact>& LightArray);
@@ -367,7 +367,7 @@ private:
 	void RenderCapsuleShadowsForMovableSkylight(FRHICommandListImmediate& RHICmdList, TRefCountPtr<IPooledRenderTarget>& BentNormalOutput) const;
 
 	/** Render deferred projections of shadows for a given light into the light attenuation buffer. */
-	bool RenderShadowProjections(FRHICommandListImmediate& RHICmdList, const FLightSceneInfo* LightSceneInfo, IPooledRenderTarget* ScreenShadowMaskTexture, bool& bInjectedTranslucentVolume);
+	bool RenderShadowProjections(FRHICommandListImmediate& RHICmdList, const FLightSceneInfo* LightSceneInfo, IPooledRenderTarget* ScreenShadowMaskTexture, IPooledRenderTarget* ScreenShadowMaskSubPixelTexture, const struct FHairStrandsDatas* HairDatas, bool& bInjectedTranslucentVolume);
 
 	/** Render shadow projections when forward rendering. */
 	void RenderForwardShadingShadowProjections(FRHICommandListImmediate& RHICmdList, TRefCountPtr<IPooledRenderTarget>& ForwardScreenSpaceShadowMask);
@@ -400,7 +400,7 @@ private:
 	  * @param LightIndex The light's index into FScene::Lights
 	  * @return true if anything got rendered
 	  */
-	void RenderLight(FRHICommandList& RHICmdList, const FLightSceneInfo* LightSceneInfo, IPooledRenderTarget* ScreenShadowMaskTexture, bool bRenderOverlap, bool bIssueDrawEvent);
+	void RenderLight(FRHICommandList& RHICmdList, const FLightSceneInfo* LightSceneInfo, IPooledRenderTarget* ScreenShadowMaskTexture, IPooledRenderTarget* ScreenShadowMaskSubPixelTexture, struct FHairStrandsTransmittanceMaskData* TransmittanceMaskData, const struct FHairStrandsVisibilityViews* InHairVisibilityViews, bool bRenderOverlap, bool bIssueDrawEvent);
 
 	/** Renders an array of simple lights using standard deferred shading. */
 	void RenderSimpleLightsStandardDeferred(FRHICommandListImmediate& RHICmdList, const FSimpleLightArray& SimpleLights);
@@ -509,7 +509,9 @@ private:
 		const FViewInfo& View,
 		const FLightSceneInfo& LightSceneInfo,
 		const IScreenSpaceDenoiser::FShadowRayTracingConfig& RayTracingConfig,
-		IScreenSpaceDenoiser::EShadowRequirements DenoiserRequirements,
+		const IScreenSpaceDenoiser::EShadowRequirements DenoiserRequirements,
+		const bool bSubPixelShadowMask,
+		FRDGTextureRef HairCategorizationTexture,
 		FRDGTextureRef* OutShadowMask,
 		FRDGTextureRef* OutRayHitDistance);
 
diff --git a/Engine/Source/Runtime/Renderer/Private/HairStrands/HairStrandsCluster.cpp b/Engine/Source/Runtime/Renderer/Private/HairStrands/HairStrandsCluster.cpp
new file mode 100644
--- /dev/null
+++ b/Engine/Source/Runtime/Renderer/Private/HairStrands/HairStrandsCluster.cpp
@@ -0,0 +1,76 @@
+// Copyright 1998-2019 Epic Games, Inc. All Rights Reserved.
+
+
+#include "HairStrandsCluster.h"
+#include "HairStrandsUtils.h"
+#include "SceneRendering.h"
+#include "SceneManagement.h"
+
+FHairStrandsClusterViews CreateHairStrandsClusters(
+	FRHICommandListImmediate& RHICmdList,
+	const FScene* Scene,
+	const TArray<FViewInfo>& Views)
+{
+	FHairStrandsClusterViews PrimitivesClusterViews;
+	for (int32 ViewIndex = 0; ViewIndex < Views.Num(); ViewIndex++)
+	{
+		const FViewInfo& View = Views[ViewIndex];
+		if (View.Family)
+		{
+			FHairStrandsClusterDatas& PrimitivesClusters = PrimitivesClusterViews.Views.AddDefaulted_GetRef();
+
+			if (View.HairStrandsMeshElements.Num() == 0)
+			{
+				continue;
+			}
+
+			// Spatially clusters all hair primitives within the same area, and render a single DOM for them
+			uint32 ClusterId = 0;
+			auto UpdateCluster = [&PrimitivesClusters, &View, &ClusterId](const FMeshBatchAndRelevance* MeshBatchAndRelevance, const FPrimitiveSceneProxy* Proxy)
+			{
+				const FBoxSphereBounds& PrimitiveBounds = Proxy->GetBounds();
+
+				bool bClusterFound = false;
+				for (FHairStrandsClusterData& Cluster : PrimitivesClusters.Datas)
+				{
+					const bool bIntersect = FBoxSphereBounds::SpheresIntersect(Cluster.Bounds, PrimitiveBounds);
+					if (bIntersect)
+					{
+						Cluster.Bounds = Union(Cluster.Bounds, PrimitiveBounds);
+
+						if (MeshBatchAndRelevance)
+						{
+							Cluster.PrimitivesInfos.Add(*MeshBatchAndRelevance);
+						}
+						bClusterFound = true;
+						break;
+					}
+				}
+
+				if (!bClusterFound)
+				{
+					FHairStrandsClusterData Cluster;
+					Cluster.ClusterId = ClusterId++;
+					if (MeshBatchAndRelevance)
+					{
+						Cluster.PrimitivesInfos.Add(*MeshBatchAndRelevance);
+					}
+					Cluster.Bounds = PrimitiveBounds;
+					PrimitivesClusters.Datas.Add(Cluster);
+				}
+			};
+
+			for (const FMeshBatchAndRelevance& MeshBatchAndRelevance : View.HairStrandsMeshElements)
+			{
+				UpdateCluster(&MeshBatchAndRelevance, MeshBatchAndRelevance.PrimitiveSceneProxy);
+			}
+
+			for (FHairStrandsClusterData& Cluster : PrimitivesClusters.Datas)
+			{
+				Cluster.ScreenRect = ComputeProjectedScreenRect(Cluster.Bounds.GetBox(), View);
+			}
+		}
+	}
+
+	return PrimitivesClusterViews;
+}
diff --git a/Engine/Source/Runtime/Renderer/Private/HairStrands/HairStrandsCluster.h b/Engine/Source/Runtime/Renderer/Private/HairStrands/HairStrandsCluster.h
new file mode 100644
--- /dev/null
+++ b/Engine/Source/Runtime/Renderer/Private/HairStrands/HairStrandsCluster.h
@@ -0,0 +1,45 @@
+// Copyright 1998-2019 Epic Games, Inc. All Rights Reserved.
+
+/*=============================================================================
+	HairStrandsCluster.h: Hair strands cluster computation implementation.
+=============================================================================*/
+
+#pragma once
+
+#include "CoreMinimal.h"
+#include "RendererInterface.h"
+#include "HairStrandsVoxelization.h"
+#include "SceneManagement.h"
+
+/// Hair cluster infos (i.e. group of mesh having hair material on them)
+struct FHairStrandsClusterData
+{
+	typedef TArray<FMeshBatchAndRelevance, SceneRenderingAllocator> TPrimitiveInfos;
+
+	FVector GetMinBound() const { return VoxelResources.MinAABB; }
+	FVector GetMaxBound() const { return VoxelResources.MaxAABB; }
+	uint32  GetResolution() const { return VoxelResources.DensityTexture ? VoxelResources.DensityTexture->GetDesc().Extent.X : 0; }
+
+	FHairStrandsVoxelResources VoxelResources;
+	TPrimitiveInfos PrimitivesInfos;
+	FBoxSphereBounds Bounds;
+	FIntRect ScreenRect;
+	uint32 ClusterId;
+};
+
+/// Store all hair strandscluster infos for a given view
+struct FHairStrandsClusterDatas
+{
+	TArray<FHairStrandsClusterData, SceneRenderingAllocator> Datas;
+};
+
+/// Store all hair strands cluster info for all views
+struct FHairStrandsClusterViews
+{
+	TArray<FHairStrandsClusterDatas, SceneRenderingAllocator> Views;
+};
+
+FHairStrandsClusterViews CreateHairStrandsClusters(
+	FRHICommandListImmediate& RHICmdList,
+	const FScene* Scene,
+	const TArray<FViewInfo>& Views);
\ No newline at end of file
diff --git a/Engine/Source/Runtime/Renderer/Private/HairStrands/HairStrandsComposition.cpp b/Engine/Source/Runtime/Renderer/Private/HairStrands/HairStrandsComposition.cpp
new file mode 100644
--- /dev/null
+++ b/Engine/Source/Runtime/Renderer/Private/HairStrands/HairStrandsComposition.cpp
@@ -0,0 +1,258 @@
+// Copyright 1998-2019 Epic Games, Inc. All Rights Reserved.
+
+#include "HairStrandsComposition.h"
+#include "Shader.h"
+#include "GlobalShader.h"
+#include "ShaderParameters.h"
+#include "ShaderParameterStruct.h"
+#include "ShaderParameterMacros.h"
+#include "RenderGraphUtils.h"
+#include "PostProcessing.h"
+#include "HairStrandsRendering.h"
+
+/////////////////////////////////////////////////////////////////////////////////////////
+
+static int32 GHairFastResolveVelocityThreshold = 1;
+static FAutoConsoleVariableRef CVarHairFastResolveVelocityThreshold(TEXT("r.HairStrands.VelocityThreshold"), GHairFastResolveVelocityThreshold, TEXT("Threshold value (in pixel) above which a pixel is forced to be resolve with responsive AA (in order to avoid smearing). Default is 3."));
+
+///////////////////////////////////////////////////////////////////////////////////////////////////
+class FHairVisibilityComposeSubPixelPS : public FGlobalShader
+{
+	DECLARE_GLOBAL_SHADER(FHairVisibilityComposeSubPixelPS);
+	SHADER_USE_PARAMETER_STRUCT(FHairVisibilityComposeSubPixelPS, FGlobalShader);
+
+	BEGIN_SHADER_PARAMETER_STRUCT(FParameters, )
+		SHADER_PARAMETER_RDG_TEXTURE(Texture2D, SubPixelColorTexture)
+		SHADER_PARAMETER_RDG_TEXTURE(Texture2D, HairVisibilityDepthTexture)
+		RENDER_TARGET_BINDING_SLOTS()
+	END_SHADER_PARAMETER_STRUCT()
+
+public:
+	static bool ShouldCompilePermutation(const FGlobalShaderPermutationParameters& Parameters) { return Parameters.Platform == SP_PCD3D_SM5; }
+};
+
+IMPLEMENT_GLOBAL_SHADER(FHairVisibilityComposeSubPixelPS, "/Engine/Private/HairStrands/HairStrandsVisibilityComposeSubPixelPS.usf", "SubColorPS", SF_Pixel);
+
+static void AddHairVisibilityComposeSubPixelPass(
+	FRDGBuilder& GraphBuilder,
+	const FViewInfo& View,
+	const FRDGTextureRef& SubPixelColorTexture,
+	const FRDGTextureRef& HairVisibilityDepthTexture,
+	FRDGTextureRef& OutColorTexture,
+	FRDGTextureRef& OutDepthTexture)
+{
+	FHairVisibilityComposeSubPixelPS::FParameters* Parameters = GraphBuilder.AllocParameters<FHairVisibilityComposeSubPixelPS::FParameters>();
+	Parameters->SubPixelColorTexture = SubPixelColorTexture;
+	Parameters->HairVisibilityDepthTexture = HairVisibilityDepthTexture;
+	Parameters->RenderTargets[0] = FRenderTargetBinding(OutColorTexture, ERenderTargetLoadAction::ELoad);
+	Parameters->RenderTargets.DepthStencil = FDepthStencilBinding(
+		OutDepthTexture,
+		ERenderTargetLoadAction::ELoad,
+		ERenderTargetLoadAction::ENoAction,
+		FExclusiveDepthStencil::DepthWrite_StencilNop);
+	
+	TShaderMapRef<FPostProcessVS> VertexShader(View.ShaderMap);
+	TShaderMapRef<FHairVisibilityComposeSubPixelPS> PixelShader(View.ShaderMap);
+	const TShaderMap<FGlobalShaderType>* GlobalShaderMap = View.ShaderMap;
+	const FIntRect Viewport = View.ViewRect;
+	const FIntPoint Resolution = OutColorTexture->Desc.Extent;
+	const FViewInfo* CapturedView = &View;
+
+	{
+		ClearUnusedGraphResources(*PixelShader, Parameters);
+
+		GraphBuilder.AddPass(
+			RDG_EVENT_NAME("HairStrandsVisibilityComposeSubSPixel"),
+			Parameters,
+			ERDGPassFlags::Raster,
+			[Parameters, VertexShader, PixelShader, Viewport, Resolution, CapturedView](FRHICommandList& RHICmdList)
+		{
+			FGraphicsPipelineStateInitializer GraphicsPSOInit;
+			RHICmdList.ApplyCachedRenderTargets(GraphicsPSOInit);
+			GraphicsPSOInit.BlendState = TStaticBlendState<CW_RGBA, BO_Add, BF_One, BF_InverseSourceAlpha, BO_Add, BF_One, BF_Zero>::GetRHI();
+			GraphicsPSOInit.RasterizerState = TStaticRasterizerState<>::GetRHI();
+
+			// Write stencil value for partially covered pixel, in order to run responsive AA on them
+			GraphicsPSOInit.DepthStencilState = TStaticDepthStencilState<true, CF_DepthNearOrEqual>::GetRHI();
+			GraphicsPSOInit.BoundShaderState.VertexDeclarationRHI = GFilterVertexDeclaration.VertexDeclarationRHI;
+			GraphicsPSOInit.BoundShaderState.VertexShaderRHI = GETSAFERHISHADER_VERTEX(*VertexShader);
+			GraphicsPSOInit.BoundShaderState.PixelShaderRHI = GETSAFERHISHADER_PIXEL(*PixelShader);
+			GraphicsPSOInit.PrimitiveType = PT_TriangleList;
+			SetGraphicsPipelineState(RHICmdList, GraphicsPSOInit);
+
+			VertexShader->SetParameters(RHICmdList, CapturedView->ViewUniformBuffer);
+			RHICmdList.SetViewport(Viewport.Min.X, Viewport.Min.Y, 0.0f, Viewport.Max.X, Viewport.Max.Y, 1.0f);
+			SetShaderParameters(RHICmdList, *PixelShader, PixelShader->GetPixelShader(), *Parameters);
+			DrawRectangle(
+				RHICmdList,
+				0, 0,
+				Viewport.Width(), Viewport.Height(),
+				Viewport.Min.X, Viewport.Min.Y,
+				Viewport.Width(), Viewport.Height(),
+				Viewport.Size(),
+				Resolution,
+				*VertexShader,
+				EDRF_UseTriangleOptimization);
+		});
+	}
+}
+
+///////////////////////////////////////////////////////////////////////////////////////////////////
+class FHairVisibilityFastResolvePS : public FGlobalShader
+{
+	DECLARE_GLOBAL_SHADER(FHairVisibilityFastResolvePS);
+	SHADER_USE_PARAMETER_STRUCT(FHairVisibilityFastResolvePS, FGlobalShader);
+
+	BEGIN_SHADER_PARAMETER_STRUCT(FParameters, )
+		SHADER_PARAMETER(float, VelocityThreshold)
+		SHADER_PARAMETER_RDG_TEXTURE(Texture2D, HairVisibilityVelocityTexture)
+		SHADER_PARAMETER_RDG_TEXTURE(Texture2D, HairVisibilityDepthTexture)
+		RENDER_TARGET_BINDING_SLOTS()
+	END_SHADER_PARAMETER_STRUCT()
+
+public:
+	static bool ShouldCompilePermutation(const FGlobalShaderPermutationParameters& Parameters) { return Parameters.Platform == SP_PCD3D_SM5; }
+};
+
+IMPLEMENT_GLOBAL_SHADER(FHairVisibilityFastResolvePS, "/Engine/Private/HairStrands/HairStrandsVisibilityComposeSubPixelPS.usf", "FastResolvePS", SF_Pixel);
+
+static void AddHairVisibilityFastResolvePass(
+	FRDGBuilder& GraphBuilder,
+	const FViewInfo& View,
+	const FRDGTextureRef& HairVisibilityVelocityTexture,
+	const FRDGTextureRef& HairVisibilityDepthTexture,
+	FRDGTextureRef& OutDepthTexture)
+{
+	const FIntPoint Resolution = OutDepthTexture->Desc.Extent;
+	FRDGTextureRef DummyTexture;
+	{
+		FRDGTextureDesc Desc;
+		Desc.Extent = Resolution;
+		Desc.Depth = 0;
+		Desc.Format = PF_R8G8B8A8;
+		Desc.NumMips = 1;
+		Desc.NumSamples = 1;
+		Desc.Flags = TexCreate_None;
+		Desc.TargetableFlags = TexCreate_RenderTargetable | TexCreate_ShaderResource;
+		Desc.ClearValue = FClearValueBinding(0);
+		DummyTexture = GraphBuilder.CreateTexture(Desc, TEXT("HairDummyTexture"));
+	}
+
+	FVector2D PixelVelocity(1.f / (Resolution.X * 2), 1.f / (Resolution.Y * 2));
+	const float VelocityThreshold = FMath::Clamp(GHairFastResolveVelocityThreshold, 0, 512) * FMath::Min(PixelVelocity.X, PixelVelocity.Y);
+
+	FHairVisibilityFastResolvePS::FParameters* Parameters = GraphBuilder.AllocParameters<FHairVisibilityFastResolvePS::FParameters>();
+	Parameters->HairVisibilityVelocityTexture = HairVisibilityVelocityTexture;
+	Parameters->HairVisibilityDepthTexture = HairVisibilityDepthTexture;
+	Parameters->VelocityThreshold = VelocityThreshold;
+	Parameters->RenderTargets[0] = FRenderTargetBinding(DummyTexture, ERenderTargetLoadAction::ENoAction);
+	Parameters->RenderTargets.DepthStencil = FDepthStencilBinding(
+		OutDepthTexture,
+		ERenderTargetLoadAction::ELoad,
+		ERenderTargetLoadAction::ELoad,
+		FExclusiveDepthStencil::DepthNop_StencilWrite);
+
+	TShaderMapRef<FPostProcessVS> VertexShader(View.ShaderMap);
+	TShaderMapRef<FHairVisibilityFastResolvePS> PixelShader(View.ShaderMap);
+	const TShaderMap<FGlobalShaderType>* GlobalShaderMap = View.ShaderMap;
+	const FIntRect Viewport = View.ViewRect;
+	
+	const FViewInfo* CapturedView = &View;
+
+	{
+		ClearUnusedGraphResources(*PixelShader, Parameters);
+
+		GraphBuilder.AddPass(
+			RDG_EVENT_NAME("HairStrandsVisibilityMarkTAAFastResolve"),
+			Parameters,
+			ERDGPassFlags::Raster,
+			[Parameters, VertexShader, PixelShader, Viewport, Resolution, CapturedView](FRHICommandList& RHICmdList)
+		{
+			FGraphicsPipelineStateInitializer GraphicsPSOInit;
+			RHICmdList.ApplyCachedRenderTargets(GraphicsPSOInit);
+			GraphicsPSOInit.BlendState = TStaticBlendState<CW_RGBA, BO_Add, BF_One, BF_Zero, BO_Add, BF_One, BF_Zero>::GetRHI();
+			GraphicsPSOInit.RasterizerState = TStaticRasterizerState<>::GetRHI();	
+			GraphicsPSOInit.DepthStencilState = TStaticDepthStencilState<
+				false, CF_Always, 
+				true, CF_Always, SO_Keep, SO_Keep, SO_Replace,
+				false, CF_Always, SO_Keep, SO_Keep, SO_Keep,
+				STENCIL_TEMPORAL_RESPONSIVE_AA_MASK, STENCIL_TEMPORAL_RESPONSIVE_AA_MASK>::GetRHI();
+			GraphicsPSOInit.BoundShaderState.VertexDeclarationRHI = GFilterVertexDeclaration.VertexDeclarationRHI;
+			GraphicsPSOInit.BoundShaderState.VertexShaderRHI = GETSAFERHISHADER_VERTEX(*VertexShader);
+			GraphicsPSOInit.BoundShaderState.PixelShaderRHI = GETSAFERHISHADER_PIXEL(*PixelShader);
+			GraphicsPSOInit.PrimitiveType = PT_TriangleList;
+			SetGraphicsPipelineState(RHICmdList, GraphicsPSOInit);
+
+			VertexShader->SetParameters(RHICmdList, CapturedView->ViewUniformBuffer);
+			RHICmdList.SetStencilRef(STENCIL_TEMPORAL_RESPONSIVE_AA_MASK);
+			RHICmdList.SetViewport(Viewport.Min.X, Viewport.Min.Y, 0.0f, Viewport.Max.X, Viewport.Max.Y, 1.0f);
+			SetShaderParameters(RHICmdList, *PixelShader, PixelShader->GetPixelShader(), *Parameters);
+			DrawRectangle(
+				RHICmdList,
+				0, 0,
+				Viewport.Width(), Viewport.Height(),
+				Viewport.Min.X, Viewport.Min.Y,
+				Viewport.Width(), Viewport.Height(),
+				Viewport.Size(),
+				Resolution,
+				*VertexShader,
+				EDRF_UseTriangleOptimization);
+		});
+	}
+}
+
+///////////////////////////////////////////////////////////////////////////////////////////////////
+
+void RenderHairComposeSubPixel(
+	FRHICommandListImmediate& RHICmdList,
+	const TArray<FViewInfo>& Views,
+	const FHairStrandsDatas* HairDatas)
+{
+	if (!HairDatas || HairDatas->HairVisibilityViews.HairDatas.Num() == 0)
+		return;
+
+	const FHairStrandsVisibilityViews& HairVisibilityViews = HairDatas->HairVisibilityViews;
+
+	DECLARE_GPU_STAT(HairStrandsComposeSubPixel);
+	SCOPED_DRAW_EVENT(RHICmdList, HairStrandsComposeSubPixel);
+	SCOPED_GPU_STAT(RHICmdList, HairStrandsComposeSubPixel);
+
+	FSceneRenderTargets& SceneTargets = FSceneRenderTargets::Get(RHICmdList);
+	FRDGBuilder GraphBuilder(RHICmdList);
+
+	FRDGTextureRef SceneColorSubPixelTexture = GraphBuilder.RegisterExternalTexture(SceneTargets.SceneColorSubPixel, TEXT("SceneColorSubPixelTexture"));
+	FRDGTextureRef SceneColorTexture = GraphBuilder.RegisterExternalTexture(SceneTargets.GetSceneColor(), TEXT("SceneColorTexture"));
+	FRDGTextureRef SceneColorDepth = GraphBuilder.RegisterExternalTexture(SceneTargets.SceneDepthZ, TEXT("SceneDepthTexture"));
+	for (int32 ViewIndex = 0; ViewIndex < Views.Num(); ViewIndex++)
+	{
+		const FViewInfo& View = Views[ViewIndex];
+		if (View.Family)
+		{
+			if (ViewIndex < HairVisibilityViews.HairDatas.Num() && HairVisibilityViews.HairDatas[ViewIndex].DepthTexture)
+			{
+				FRDGTextureRef HairVisibilityDepthTexture = GraphBuilder.RegisterExternalTexture(HairVisibilityViews.HairDatas[ViewIndex].DepthTexture, TEXT("HairVisibilityDepthTexture"));
+				FRDGTextureRef HairVisibilityVelocityTexture = GraphBuilder.RegisterExternalTexture(HairVisibilityViews.HairDatas[ViewIndex].VelocityTexture, TEXT("HairVisibilityVelocityTexture"));
+				
+				// #hair_todo : compose partially covered hair with transparent surface: this can be done by 
+				// rendering quad(s) covering the hair at the correct depth. This will be sorted with other 
+				// transparent surface, which should make the overall sorting workable
+				AddHairVisibilityComposeSubPixelPass(
+					GraphBuilder,
+					View,
+					SceneColorSubPixelTexture,
+					HairVisibilityDepthTexture,
+					SceneColorTexture,
+					SceneColorDepth);
+
+				AddHairVisibilityFastResolvePass(
+					GraphBuilder,
+					View,
+					HairVisibilityVelocityTexture,
+					HairVisibilityDepthTexture,
+					SceneColorDepth);
+			}
+		}
+	}
+
+	GraphBuilder.Execute();
+}
diff --git a/Engine/Source/Runtime/Renderer/Private/HairStrands/HairStrandsComposition.h b/Engine/Source/Runtime/Renderer/Private/HairStrands/HairStrandsComposition.h
new file mode 100644
--- /dev/null
+++ b/Engine/Source/Runtime/Renderer/Private/HairStrands/HairStrandsComposition.h
@@ -0,0 +1,15 @@
+// Copyright 1998-2019 Epic Games, Inc. All Rights Reserved.
+
+/*=============================================================================
+	HairStrandsComposition.h: Hair strands pixel composition implementation.
+=============================================================================*/
+
+#pragma once
+
+#include "CoreMinimal.h"
+#include "RendererInterface.h"
+
+void RenderHairComposeSubPixel(
+	FRHICommandListImmediate& RHICmdList, 
+	const TArray<FViewInfo>& Views,
+	const struct FHairStrandsDatas* HairDatas);
diff --git a/Engine/Source/Runtime/Renderer/Private/HairStrands/HairStrandsDebug.cpp b/Engine/Source/Runtime/Renderer/Private/HairStrands/HairStrandsDebug.cpp
new file mode 100644
--- /dev/null
+++ b/Engine/Source/Runtime/Renderer/Private/HairStrands/HairStrandsDebug.cpp
@@ -0,0 +1,666 @@
+// Copyright 1998-2019 Epic Games, Inc. All Rights Reserved.
+
+#include "HairStrandsDebug.h"
+#include "HairStrandsInterface.h"
+#include "HairStrandsCluster.h"
+#include "HairStrandsDeepShadow.h"
+#include "HairStrandsUtils.h"
+#include "HairStrandsVoxelization.h"
+#include "HairStrandsRendering.h"
+#include "HairStrandsVisibility.h"
+
+#include "Shader.h"
+#include "GlobalShader.h"
+#include "ShaderParameters.h"
+#include "ShaderParameterStruct.h"
+#include "RenderGraphUtils.h"
+#include "PostProcessing.h"
+#include "SceneTextureParameters.h"
+#include "DynamicPrimitiveDrawing.h"
+#include "RenderTargetTemp.h"
+#include "CanvasTypes.h"
+
+static int32 GDeepShadowDebugIndex = 0;
+static float GDeepShadowDebugScale = 20;
+
+static FAutoConsoleVariableRef CVarDeepShadowDebugDomIndex(TEXT("r.HairStrands.DeepShadow.DebugDOMIndex"), GDeepShadowDebugIndex, TEXT("Index of the DOM texture to draw"));
+static FAutoConsoleVariableRef CVarDeepShadowDebugDomScale(TEXT("r.HairStrands.DeepShadow.DebugDOMScale"), GDeepShadowDebugScale, TEXT("Scaling value for the DeepOpacityMap when drawing the deep shadow stats"));
+
+static int32 GDeepShadowStats = 0;
+static FAutoConsoleVariableRef CVarDeepShadowStats(TEXT("r.HairStrands.DeepShadow.DrawStats"), GDeepShadowStats, TEXT("Draw various stats/debug mode about hair rendering"));
+
+static int32 GHairDebugRenderMode = 0;
+static FAutoConsoleVariableRef CVarDebugPhysicsStrand(TEXT("r.HairStrands.DebugRenderMode"), GHairDebugRenderMode, TEXT("Render debug mode for hair strands. 0:off, 1:simulation strands, 2:render strands with colored simulation strands influence, 3:hair UV, 4:hair root UV, 5: hair seed, 6: dimensions"));
+
+
+enum class EHairDebugMode : uint8
+{
+	None,
+	LightBounds,
+	DeepOpacityMaps,
+	ClusterScreenRect,
+	SamplePerPixel,
+	CoverageType,
+	TAAResolveType,
+	VoxelsDensity,
+	VoxelsTangent,
+	VoxelsBaseColor,
+	VoxelsRoughness
+};
+
+static EHairDebugMode GetHairDebugMode()
+{
+	switch (GDeepShadowStats)
+	{
+	case 0:  return EHairDebugMode::None;
+	case 1:  return EHairDebugMode::LightBounds;
+	case 2:  return EHairDebugMode::ClusterScreenRect;
+	case 3:  return EHairDebugMode::DeepOpacityMaps;
+	case 4:  return EHairDebugMode::SamplePerPixel;
+	case 5:  return EHairDebugMode::TAAResolveType;
+	case 6:  return EHairDebugMode::CoverageType;
+	case 7:  return EHairDebugMode::VoxelsDensity;
+	case 8:  return EHairDebugMode::VoxelsTangent;
+	case 9:  return EHairDebugMode::VoxelsBaseColor;
+	case 10: return EHairDebugMode::VoxelsRoughness;
+	default: return EHairDebugMode::None;
+	};
+}
+
+static const TCHAR* ToString(EHairDebugMode DebugMode)
+{
+	switch (DebugMode)
+	{
+	case EHairDebugMode::None: return TEXT("None");
+	case EHairDebugMode::LightBounds: return TEXT("All DOMs light bounds");
+	case EHairDebugMode::ClusterScreenRect: return TEXT("Screen projected clusters");
+	case EHairDebugMode::DeepOpacityMaps: return TEXT("Deep opacity maps");
+	case EHairDebugMode::SamplePerPixel: return TEXT("Sub-pixel sample count");
+	case EHairDebugMode::TAAResolveType: return TEXT("TAA resolve type (regular/responsive)");
+	case EHairDebugMode::CoverageType: return TEXT("Type of hair coverage (full/partial)");
+	case EHairDebugMode::VoxelsDensity: return TEXT("Hair density volume");
+	case EHairDebugMode::VoxelsTangent: return TEXT("Hair tangent volume");
+	case EHairDebugMode::VoxelsBaseColor: return TEXT("Hair base color volume");
+	case EHairDebugMode::VoxelsRoughness: return TEXT("Hair roughness volume");
+	default: return TEXT("None");
+	};
+}
+
+
+EHairStrandsDebugMode GetHairStrandsDebugMode()
+{
+	switch (GHairDebugRenderMode)
+	{
+	case 0:  return EHairStrandsDebugMode::None;
+	case 1:  return EHairStrandsDebugMode::SimHairStrands;
+	case 2:  return EHairStrandsDebugMode::RenderHairStrands;
+	case 3:  return EHairStrandsDebugMode::RenderHairRootUV;
+	case 4:  return EHairStrandsDebugMode::RenderHairUV;
+	case 5:  return EHairStrandsDebugMode::RenderHairSeed;
+	case 6:  return EHairStrandsDebugMode::RenderHairDimension;
+	case 7:  return EHairStrandsDebugMode::RenderHairRadiusVariation;
+	default: return EHairStrandsDebugMode::None;
+	};
+}
+
+static const TCHAR* ToString(EHairStrandsDebugMode DebugMode)
+{
+	switch (DebugMode)
+	{
+	case EHairStrandsDebugMode::None						: return TEXT("None");
+	case EHairStrandsDebugMode::SimHairStrands				: return TEXT("Simulation strands");
+	case EHairStrandsDebugMode::RenderHairStrands			: return TEXT("Rendering strands influences");
+	case EHairStrandsDebugMode::RenderHairRootUV			: return TEXT("Roots UV");
+	case EHairStrandsDebugMode::RenderHairUV				: return TEXT("Hair UV");
+	case EHairStrandsDebugMode::RenderHairSeed				: return TEXT("Hair seed");
+	case EHairStrandsDebugMode::RenderHairDimension			: return TEXT("Hair dimensions");
+	case EHairStrandsDebugMode::RenderHairRadiusVariation	: return TEXT("Hair radius variation");
+	default													: return TEXT("None");
+	};
+}
+
+
+///////////////////////////////////////////////////////////////////////////////////////////////////
+class FHairDebugPS : public FGlobalShader
+{
+	DECLARE_GLOBAL_SHADER(FHairDebugPS);
+	SHADER_USE_PARAMETER_STRUCT(FHairDebugPS, FGlobalShader);
+
+	class FDebugMode : SHADER_PERMUTATION_INT("PERMUTATION_DEBUG_MODE", 3);
+	using FPermutationDomain = TShaderPermutationDomain<FDebugMode>;
+
+	BEGIN_SHADER_PARAMETER_STRUCT(FParameters, )
+		SHADER_PARAMETER(FVector2D, OutputResolution)
+		SHADER_PARAMETER(uint32, FastResolveMask)
+		SHADER_PARAMETER_RDG_TEXTURE(Texture2D, CategorizationTexture)
+		SHADER_PARAMETER_SRV(Texture2D, DepthStencilTexture)
+		SHADER_PARAMETER_SAMPLER(SamplerState, LinearSampler)
+		SHADER_PARAMETER_STRUCT_REF(FViewUniformShaderParameters, ViewUniformBuffer)
+		RENDER_TARGET_BINDING_SLOTS()
+		END_SHADER_PARAMETER_STRUCT()
+
+public:
+	static bool ShouldCompilePermutation(const FGlobalShaderPermutationParameters& Parameters) { return Parameters.Platform == SP_PCD3D_SM5; }
+};
+
+IMPLEMENT_GLOBAL_SHADER(FHairDebugPS, "/Engine/Private/HairStrands/HairStrandsDebug.usf", "MainPS", SF_Pixel);
+
+static void AddDebugHairPass(
+	FRDGBuilder& GraphBuilder,
+	const FViewInfo* View,
+	const EHairDebugMode InDebugMode,
+	const TRefCountPtr<IPooledRenderTarget>& InCategorizationTexture,
+	const FShaderResourceViewRHIRef& InDepthStencilTexture,
+	FRDGTextureRef& OutTarget)
+{
+	check(OutTarget);
+	check(InDebugMode == EHairDebugMode::TAAResolveType || InDebugMode == EHairDebugMode::SamplePerPixel || InDebugMode == EHairDebugMode::CoverageType);
+
+	if (!InCategorizationTexture) return;
+	if (InDebugMode == EHairDebugMode::TAAResolveType && !InDepthStencilTexture) return;
+
+	FRDGTextureRef CategorizationTexture = InCategorizationTexture ? GraphBuilder.RegisterExternalTexture(InCategorizationTexture, TEXT("CategorizationTexture")) : nullptr;
+
+	const FIntRect Viewport = View->ViewRect;
+	const FIntPoint Resolution(Viewport.Width(), Viewport.Height());
+
+	FHairDebugPS::FParameters* Parameters = GraphBuilder.AllocParameters<FHairDebugPS::FParameters>();
+	Parameters->OutputResolution = Resolution;
+	Parameters->FastResolveMask = STENCIL_TEMPORAL_RESPONSIVE_AA_MASK;
+	Parameters->CategorizationTexture = CategorizationTexture;
+	Parameters->DepthStencilTexture = InDepthStencilTexture;
+	Parameters->LinearSampler = TStaticSamplerState<SF_Bilinear, AM_Clamp, AM_Clamp, AM_Clamp>::GetRHI();
+	Parameters->RenderTargets[0] = FRenderTargetBinding(OutTarget, ERenderTargetLoadAction::ELoad, 0);
+	TShaderMapRef<FPostProcessVS> VertexShader(View->ShaderMap);
+
+	FHairDebugPS::FPermutationDomain PermutationVector;
+	uint32 DebugPermutation = 0;
+	switch (InDebugMode)
+	{
+		case EHairDebugMode::SamplePerPixel:	DebugPermutation = 0; break;
+		case EHairDebugMode::CoverageType:		DebugPermutation = 1; break;
+		case EHairDebugMode::TAAResolveType:	DebugPermutation = 2; break;
+	};
+	PermutationVector.Set<FHairDebugPS::FDebugMode>(DebugPermutation);
+	TShaderMapRef<FHairDebugPS> PixelShader(View->ShaderMap, PermutationVector);
+
+	ClearUnusedGraphResources(*PixelShader, Parameters);
+
+	GraphBuilder.AddPass(
+		RDG_EVENT_NAME("HairStrandsDebug"),
+		Parameters,
+		ERDGPassFlags::Raster,
+		[Parameters, VertexShader, PixelShader, Viewport, Resolution, View](FRHICommandList& RHICmdList)
+	{
+		FGraphicsPipelineStateInitializer GraphicsPSOInit;
+		RHICmdList.ApplyCachedRenderTargets(GraphicsPSOInit);
+		GraphicsPSOInit.BlendState = TStaticBlendState<CW_RGBA, BO_Add, BF_One, BF_Zero, BO_Add, BF_One, BF_Zero>::GetRHI();
+		GraphicsPSOInit.RasterizerState = TStaticRasterizerState<>::GetRHI();
+		GraphicsPSOInit.DepthStencilState = TStaticDepthStencilState<false, CF_Always>::GetRHI();
+		GraphicsPSOInit.BoundShaderState.VertexDeclarationRHI = GFilterVertexDeclaration.VertexDeclarationRHI;
+		GraphicsPSOInit.BoundShaderState.VertexShaderRHI = GETSAFERHISHADER_VERTEX(*VertexShader);
+		GraphicsPSOInit.BoundShaderState.PixelShaderRHI = GETSAFERHISHADER_PIXEL(*PixelShader);
+		GraphicsPSOInit.PrimitiveType = PT_TriangleList;
+		SetGraphicsPipelineState(RHICmdList, GraphicsPSOInit);
+
+		VertexShader->SetParameters(RHICmdList, View->ViewUniformBuffer);
+		RHICmdList.SetViewport(Viewport.Min.X, Viewport.Min.Y, 0.0f, Viewport.Max.X, Viewport.Max.Y, 1.0f);
+		SetShaderParameters(RHICmdList, *PixelShader, PixelShader->GetPixelShader(), *Parameters);
+
+		DrawRectangle(
+			RHICmdList,
+			0, 0,
+			Viewport.Width(), Viewport.Height(),
+			Viewport.Min.X, Viewport.Min.Y,
+			Viewport.Width(), Viewport.Height(),
+			Viewport.Size(),
+			Resolution,
+			*VertexShader,
+			EDRF_UseTriangleOptimization);
+	});
+}
+
+///////////////////////////////////////////////////////////////////////////////////////////////////
+class FDeepShadowVisualizePS : public FGlobalShader
+{
+	DECLARE_GLOBAL_SHADER(FDeepShadowVisualizePS);
+	SHADER_USE_PARAMETER_STRUCT(FDeepShadowVisualizePS, FGlobalShader);
+
+	class FOutputType : SHADER_PERMUTATION_INT("PERMUTATION_OUTPUT_TYPE", 2);
+	using FPermutationDomain = TShaderPermutationDomain<FOutputType>;
+
+	BEGIN_SHADER_PARAMETER_STRUCT(FParameters, )
+		SHADER_PARAMETER(float, DomScale)
+		SHADER_PARAMETER(FVector2D, DomAtlasOffset)
+		SHADER_PARAMETER(FVector2D, DomAtlasScale)
+		SHADER_PARAMETER(FVector2D, OutputResolution)
+		SHADER_PARAMETER(FVector2D, InvOutputResolution)
+		SHADER_PARAMETER(FIntVector4, HairViewRect)
+
+		SHADER_PARAMETER_RDG_TEXTURE(Texture2D, DeepShadowDepthTexture)
+		SHADER_PARAMETER_RDG_TEXTURE(Texture2D, DeepShadowLayerTexture)
+
+		SHADER_PARAMETER_SAMPLER(SamplerState, LinearSampler)
+
+		SHADER_PARAMETER_STRUCT_REF(FViewUniformShaderParameters, ViewUniformBuffer)
+		RENDER_TARGET_BINDING_SLOTS()
+		END_SHADER_PARAMETER_STRUCT()
+
+public:
+	static bool ShouldCompilePermutation(const FGlobalShaderPermutationParameters& Parameters) { return Parameters.Platform == SP_PCD3D_SM5; }
+};
+
+IMPLEMENT_GLOBAL_SHADER(FDeepShadowVisualizePS, "/Engine/Private/HairStrands/HairStrandsDeepShadowDebug.usf", "VisualizeDomPS", SF_Pixel);
+
+static void AddDebugDeepShadowTexturePass(
+	FRDGBuilder& GraphBuilder,
+	const FViewInfo* View,
+	const FIntRect& HairViewRect,
+	const FHairStrandsDeepShadowData* ShadowData,
+	FRDGTextureRef& OutTarget)
+{
+	check(OutTarget);
+
+	FRDGTextureRef DeepShadowDepthTexture = nullptr;
+	FRDGTextureRef DeepShadowLayerTexture = nullptr;
+	FIntPoint AtlasResolution(0, 0);
+	FVector2D AltasOffset(0, 0);
+	FVector2D AltasScale(0, 0);
+	if (ShadowData)
+	{
+		DeepShadowDepthTexture = GraphBuilder.RegisterExternalTexture(ShadowData->DepthTexture, TEXT("DOMDepthTexture"));
+		DeepShadowLayerTexture = GraphBuilder.RegisterExternalTexture(ShadowData->LayersTexture, TEXT("DOMLayerTexture"));
+
+		AtlasResolution = FIntPoint(DeepShadowDepthTexture->Desc.Extent.X, DeepShadowDepthTexture->Desc.Extent.Y);
+		AltasOffset = FVector2D(ShadowData->AtlasRect.Min.X / float(AtlasResolution.X), ShadowData->AtlasRect.Min.Y / float(AtlasResolution.Y));
+		AltasScale = FVector2D((ShadowData->AtlasRect.Max.X - ShadowData->AtlasRect.Min.X) / float(AtlasResolution.X), (ShadowData->AtlasRect.Max.Y - ShadowData->AtlasRect.Min.Y) / float(AtlasResolution.Y));
+	}
+
+	const FIntRect Viewport = View->ViewRect;
+	const FIntPoint Resolution(Viewport.Width(), Viewport.Height());
+
+	FDeepShadowVisualizePS::FParameters* Parameters = GraphBuilder.AllocParameters<FDeepShadowVisualizePS::FParameters>();
+	Parameters->DomScale = GDeepShadowDebugScale;
+	Parameters->DomAtlasOffset = AltasOffset;
+	Parameters->DomAtlasScale = AltasScale;
+	Parameters->OutputResolution = Resolution;
+	Parameters->InvOutputResolution = FVector2D(1.f / Resolution.X, 1.f / Resolution.Y);
+	Parameters->DeepShadowDepthTexture = DeepShadowDepthTexture;
+	Parameters->DeepShadowLayerTexture = DeepShadowLayerTexture;
+	Parameters->LinearSampler = TStaticSamplerState<SF_Bilinear, AM_Clamp, AM_Clamp, AM_Clamp>::GetRHI();
+	Parameters->HairViewRect = FIntVector4(HairViewRect.Min.X, HairViewRect.Min.Y, HairViewRect.Width(), HairViewRect.Height());
+	Parameters->RenderTargets[0] = FRenderTargetBinding(OutTarget, ERenderTargetLoadAction::ELoad, 0);
+	TShaderMapRef<FPostProcessVS> VertexShader(View->ShaderMap);
+	FDeepShadowVisualizePS::FPermutationDomain PermutationVector;
+	PermutationVector.Set<FDeepShadowVisualizePS::FOutputType>(ShadowData ? 0 : 1);
+	TShaderMapRef<FDeepShadowVisualizePS> PixelShader(View->ShaderMap, PermutationVector);
+
+	ClearUnusedGraphResources(*PixelShader, Parameters);
+
+	GraphBuilder.AddPass(
+		ShadowData ? RDG_EVENT_NAME("DebugDeepShadowTexture") : RDG_EVENT_NAME("DebugHairViewRect"),
+		Parameters,
+		ERDGPassFlags::Raster,
+		[Parameters, VertexShader, PixelShader, Viewport, Resolution, View](FRHICommandList& RHICmdList)
+	{
+		FGraphicsPipelineStateInitializer GraphicsPSOInit;
+		RHICmdList.ApplyCachedRenderTargets(GraphicsPSOInit);
+		GraphicsPSOInit.BlendState = TStaticBlendState<CW_RGBA, BO_Add, BF_One, BF_Zero, BO_Add, BF_One, BF_Zero>::GetRHI();
+		GraphicsPSOInit.RasterizerState = TStaticRasterizerState<>::GetRHI();
+		GraphicsPSOInit.DepthStencilState = TStaticDepthStencilState<false, CF_Always>::GetRHI();
+		GraphicsPSOInit.BoundShaderState.VertexDeclarationRHI = GFilterVertexDeclaration.VertexDeclarationRHI;
+		GraphicsPSOInit.BoundShaderState.VertexShaderRHI = GETSAFERHISHADER_VERTEX(*VertexShader);
+		GraphicsPSOInit.BoundShaderState.PixelShaderRHI = GETSAFERHISHADER_PIXEL(*PixelShader);
+		GraphicsPSOInit.PrimitiveType = PT_TriangleList;
+		SetGraphicsPipelineState(RHICmdList, GraphicsPSOInit);
+
+		VertexShader->SetParameters(RHICmdList, View->ViewUniformBuffer);
+		RHICmdList.SetViewport(Viewport.Min.X, Viewport.Min.Y, 0.0f, Viewport.Max.X, Viewport.Max.Y, 1.0f);
+		SetShaderParameters(RHICmdList, *PixelShader, PixelShader->GetPixelShader(), *Parameters);
+
+		DrawRectangle(
+			RHICmdList,
+			0, 0,
+			Viewport.Width(), Viewport.Height(),
+			Viewport.Min.X, Viewport.Min.Y,
+			Viewport.Width(), Viewport.Height(),
+			Viewport.Size(),
+			Resolution,
+			*VertexShader,
+			EDRF_UseTriangleOptimization);
+	});
+}
+
+///////////////////////////////////////////////////////////////////////////////////////////////////
+class FVoxelRaymarchingPS : public FGlobalShader
+{
+	DECLARE_GLOBAL_SHADER(FVoxelRaymarchingPS);
+	SHADER_USE_PARAMETER_STRUCT(FVoxelRaymarchingPS, FGlobalShader);
+
+	class FDebugMode : SHADER_PERMUTATION_INT("PERMUTATION_DEBUG_MODE", 4);
+	using FPermutationDomain = TShaderPermutationDomain<FDebugMode>;
+
+	BEGIN_SHADER_PARAMETER_STRUCT(FParameters, )
+		SHADER_PARAMETER_STRUCT_INCLUDE(FSceneTextureParameters, SceneTextures)
+
+		SHADER_PARAMETER(FVector, VoxelMinAABB)
+		SHADER_PARAMETER(uint32, VoxelResolution)
+		SHADER_PARAMETER(FVector, VoxelMaxAABB)
+		SHADER_PARAMETER(float, DensityIsoline)
+		SHADER_PARAMETER(float, VoxelDensityScale)
+		SHADER_PARAMETER(FVector2D, OutputResolution)
+
+		SHADER_PARAMETER_RDG_TEXTURE(Texture3D, DensityTexture)
+		SHADER_PARAMETER_RDG_TEXTURE(Texture3D, TangentXTexture)
+		SHADER_PARAMETER_RDG_TEXTURE(Texture3D, TangentYTexture)
+		SHADER_PARAMETER_RDG_TEXTURE(Texture3D, TangentZTexture)
+		SHADER_PARAMETER_RDG_TEXTURE(Texture3D, MaterialTexture)
+		SHADER_PARAMETER_SAMPLER(SamplerState, LinearSampler)
+
+		SHADER_PARAMETER_STRUCT_REF(FViewUniformShaderParameters, ViewUniformBuffer)
+		RENDER_TARGET_BINDING_SLOTS()
+		END_SHADER_PARAMETER_STRUCT()
+
+public:
+	static bool ShouldCompilePermutation(const FGlobalShaderPermutationParameters& Parameters) { return Parameters.Platform == SP_PCD3D_SM5; }
+};
+
+IMPLEMENT_GLOBAL_SHADER(FVoxelRaymarchingPS, "/Engine/Private/HairStrands/HairStrandsVoxelRayMarching.usf", "MainPS", SF_Pixel);
+
+static void AddVoxelRaymarchingPass(
+	FRDGBuilder& GraphBuilder,
+	const FViewInfo& View,
+	const EHairDebugMode DebugMode,
+	const FHairStrandsClusterDatas& ClusterDatas,
+	FRDGTextureRef& OutputTexture)
+{
+	check(DebugMode == EHairDebugMode::VoxelsDensity || DebugMode == EHairDebugMode::VoxelsTangent || DebugMode == EHairDebugMode::VoxelsBaseColor || DebugMode == EHairDebugMode::VoxelsRoughness);
+
+	FSceneTextureParameters SceneTextures;
+	SetupSceneTextureParameters(GraphBuilder, &SceneTextures);
+
+	const FIntPoint Resolution(OutputTexture->Desc.Extent);
+	for (const FHairStrandsClusterData& ClusterData : ClusterDatas.Datas)
+	{
+		if (DebugMode == EHairDebugMode::VoxelsDensity && !ClusterData.VoxelResources.DensityTexture)
+			return;
+
+		if (DebugMode == EHairDebugMode::VoxelsTangent && (!ClusterData.VoxelResources.TangentXTexture || !ClusterData.VoxelResources.TangentYTexture || !ClusterData.VoxelResources.TangentZTexture))
+			return;
+
+		if ((DebugMode == EHairDebugMode::VoxelsBaseColor || DebugMode == EHairDebugMode::VoxelsRoughness) && !ClusterData.VoxelResources.MaterialTexture)
+			return;
+
+		const FRDGTextureRef VoxelDensityTexture = GraphBuilder.RegisterExternalTexture(ClusterData.VoxelResources.DensityTexture ? ClusterData.VoxelResources.DensityTexture : GSystemTextures.BlackDummy, TEXT("HairVoxelDensityTexture"));
+		const FRDGTextureRef VoxelTangentXTexture = GraphBuilder.RegisterExternalTexture(ClusterData.VoxelResources.TangentXTexture ? ClusterData.VoxelResources.TangentXTexture : GSystemTextures.BlackDummy, TEXT("HairVoxelTangentXTexture"));
+		const FRDGTextureRef VoxelTangentYTexture = GraphBuilder.RegisterExternalTexture(ClusterData.VoxelResources.TangentYTexture ? ClusterData.VoxelResources.TangentYTexture : GSystemTextures.BlackDummy, TEXT("HairVoxelTangentYTexture"));
+		const FRDGTextureRef VoxelTangentZTexture = GraphBuilder.RegisterExternalTexture(ClusterData.VoxelResources.TangentZTexture ? ClusterData.VoxelResources.TangentZTexture : GSystemTextures.BlackDummy, TEXT("HairVoxelTangentZTexture"));
+		const FRDGTextureRef VoxelMaterialTexture = GraphBuilder.RegisterExternalTexture(ClusterData.VoxelResources.MaterialTexture ? ClusterData.VoxelResources.MaterialTexture : GSystemTextures.BlackDummy, TEXT("HairVoxelMaterialTexture"));
+
+		FVoxelRaymarchingPS::FParameters* Parameters = GraphBuilder.AllocParameters<FVoxelRaymarchingPS::FParameters>();
+		Parameters->ViewUniformBuffer = View.ViewUniformBuffer;
+		Parameters->OutputResolution = Resolution;
+		Parameters->SceneTextures = SceneTextures;
+		Parameters->DensityTexture = VoxelDensityTexture;
+		Parameters->TangentXTexture = VoxelTangentXTexture;
+		Parameters->TangentYTexture = VoxelTangentYTexture;
+		Parameters->TangentZTexture = VoxelTangentZTexture;
+		Parameters->MaterialTexture = VoxelMaterialTexture;
+		Parameters->VoxelMinAABB = ClusterData.GetMinBound();
+		Parameters->VoxelMaxAABB = ClusterData.GetMaxBound();
+		Parameters->VoxelResolution = ClusterData.GetResolution();
+		Parameters->VoxelDensityScale = GetHairStrandsVoxelizationDensityScale();
+		Parameters->DensityIsoline = 1;
+		Parameters->LinearSampler = TStaticSamplerState<SF_Trilinear, AM_Clamp, AM_Clamp, AM_Clamp>::GetRHI();
+		Parameters->RenderTargets[0] = FRenderTargetBinding(OutputTexture, ERenderTargetLoadAction::ELoad);
+
+		const FIntPoint OutputResolution = SceneTextures.SceneDepthBuffer->Desc.Extent;
+		TShaderMapRef<FPostProcessVS> VertexShader(View.ShaderMap);
+
+		FVoxelRaymarchingPS::FPermutationDomain PermutationVector;
+		uint32 DebugPermutation = 0;
+		switch (DebugMode)
+		{
+		case EHairDebugMode::VoxelsDensity:		DebugPermutation = 0; break;
+		case EHairDebugMode::VoxelsTangent:		DebugPermutation = 1; break;
+		case EHairDebugMode::VoxelsBaseColor:	DebugPermutation = 2; break;
+		case EHairDebugMode::VoxelsRoughness:	DebugPermutation = 3; break;
+		};
+		PermutationVector.Set<FVoxelRaymarchingPS::FDebugMode>(DebugPermutation);
+
+		TShaderMapRef<FVoxelRaymarchingPS> PixelShader(View.ShaderMap, PermutationVector);
+		const TShaderMap<FGlobalShaderType>* GlobalShaderMap = View.ShaderMap;
+		const FIntRect Viewport = View.ViewRect;
+		const FViewInfo* CapturedView = &View;
+
+		ClearUnusedGraphResources(*PixelShader, Parameters);
+
+		GraphBuilder.AddPass(
+			RDG_EVENT_NAME("HairStrandsVoxelRaymarching"),
+			Parameters,
+			ERDGPassFlags::Raster,
+			[Parameters, VertexShader, PixelShader, Viewport, Resolution, CapturedView](FRHICommandList& RHICmdList)
+		{
+			FGraphicsPipelineStateInitializer GraphicsPSOInit;
+			RHICmdList.ApplyCachedRenderTargets(GraphicsPSOInit);
+			GraphicsPSOInit.BlendState = TStaticBlendState<CW_RGBA, BO_Add, BF_One, BF_Zero, BO_Add, BF_One, BF_Zero>::GetRHI();
+			GraphicsPSOInit.RasterizerState = TStaticRasterizerState<>::GetRHI();
+			GraphicsPSOInit.DepthStencilState = TStaticDepthStencilState<false, CF_Always>::GetRHI();
+
+			GraphicsPSOInit.BoundShaderState.VertexDeclarationRHI = GFilterVertexDeclaration.VertexDeclarationRHI;
+			GraphicsPSOInit.BoundShaderState.VertexShaderRHI = GETSAFERHISHADER_VERTEX(*VertexShader);
+			GraphicsPSOInit.BoundShaderState.PixelShaderRHI = GETSAFERHISHADER_PIXEL(*PixelShader);
+			GraphicsPSOInit.PrimitiveType = PT_TriangleList;
+			SetGraphicsPipelineState(RHICmdList, GraphicsPSOInit);
+
+			VertexShader->SetParameters(RHICmdList, CapturedView->ViewUniformBuffer);
+			RHICmdList.SetViewport(Viewport.Min.X, Viewport.Min.Y, 0.0f, Viewport.Max.X, Viewport.Max.Y, 1.0f);
+			SetShaderParameters(RHICmdList, *PixelShader, PixelShader->GetPixelShader(), *Parameters);
+
+			DrawRectangle(
+				RHICmdList,
+				0, 0,
+				Viewport.Width(), Viewport.Height(),
+				Viewport.Min.X, Viewport.Min.Y,
+				Viewport.Width(), Viewport.Height(),
+				Viewport.Size(),
+				Resolution,
+				*VertexShader,
+				EDRF_UseTriangleOptimization);
+		});
+	}
+}
+
+///////////////////////////////////////////////////////////////////////////////////////////////////
+
+void RenderHairStrandsDebugInfo(FRHICommandListImmediate& RHICmdList, TArray<FViewInfo>& Views, const FHairStrandsDatas* HairDatas)
+{
+	const float YStep = 14;
+	const float ColumnWidth = 200;
+
+	if (Views.Num() == 0 || !HairDatas)
+		return;
+
+	const FHairStrandsDeepShadowViews& InDomViews = HairDatas->DeepShadowViews;
+	const FHairStrandsClusterViews& InClusterViews = HairDatas->HairClusterPerViews;
+
+	// Only render debug information for the main view
+	const uint32 ViewIndex = 0;
+	FViewInfo& View = Views[ViewIndex];
+	const FSceneViewFamily& ViewFamily = *(View.Family);
+	FSceneRenderTargets& SceneTargets = FSceneRenderTargets::Get(RHICmdList);
+
+	// Debug mode name only
+	const EHairStrandsDebugMode StrandsDebugMode = GetHairStrandsDebugMode();
+	const EHairDebugMode HairDebugMode = GetHairDebugMode();
+
+	if (HairDebugMode == EHairDebugMode::DeepOpacityMaps)
+	{
+		const uint32 DomIndex = GDeepShadowDebugIndex;
+		TRefCountPtr<IPooledRenderTarget> DepthTexture;
+		TRefCountPtr<IPooledRenderTarget> LayerTexture;
+		const FHairStrandsDeepShadowDatas& DeepShadoDatas = InDomViews.Views[ViewIndex];
+		const bool bIsValid = DomIndex < uint32(DeepShadoDatas.Datas.Num());
+		if (bIsValid)
+		{				
+			DepthTexture = DeepShadoDatas.Datas[DomIndex].DepthTexture;
+			LayerTexture = DeepShadoDatas.Datas[DomIndex].LayersTexture;
+		}
+
+		if (DepthTexture && LayerTexture)
+		{
+			const FHairStrandsDeepShadowData& DeepShadoData = DeepShadoDatas.Datas[DomIndex];
+			FRDGBuilder GraphBuilder(RHICmdList);
+			FRDGTextureRef SceneColorTexture = GraphBuilder.RegisterExternalTexture(SceneTargets.GetSceneColor(), TEXT("SceneColorTexture"));
+			AddDebugDeepShadowTexturePass(GraphBuilder, &View, FIntRect(), &DeepShadoData, SceneColorTexture);
+			GraphBuilder.Execute();
+		}
+	}
+
+	// View Rect
+	if (IsHairStrandsViewRectOptimEnable() && HairDebugMode == EHairDebugMode::ClusterScreenRect)
+	{
+		FRDGBuilder GraphBuilder(RHICmdList);
+		FRDGTextureRef SceneColorTexture = GraphBuilder.RegisterExternalTexture(SceneTargets.GetSceneColor(), TEXT("SceneColorTexture"));
+		if (ViewIndex < uint32(InClusterViews.Views.Num()))
+		{
+			const FHairStrandsClusterDatas& ClusterDatas = InClusterViews.Views[ViewIndex];
+			for (const FHairStrandsClusterData& ClusterData : ClusterDatas.Datas)
+			{
+				AddDebugDeepShadowTexturePass(GraphBuilder, &View, ClusterData.ScreenRect, nullptr, SceneColorTexture);
+			}
+
+			const FIntRect TotalRect = ComputeVisibleHairStrandsClustersRect(View.ViewRect, ClusterDatas);
+			AddDebugDeepShadowTexturePass(GraphBuilder, &View, TotalRect, nullptr, SceneColorTexture);
+		}
+		GraphBuilder.Execute();
+	}
+	
+	const bool bIsVoxelMode = HairDebugMode == EHairDebugMode::VoxelsDensity || HairDebugMode == EHairDebugMode::VoxelsTangent || HairDebugMode == EHairDebugMode::VoxelsBaseColor || HairDebugMode == EHairDebugMode::VoxelsRoughness;
+
+	// Render Frustum for all lights & clusters
+	{
+		FViewElementPDI ShadowFrustumPDI(&View, nullptr, nullptr);
+
+		// All DOMs
+		if (HairDebugMode == EHairDebugMode::LightBounds && ViewIndex < uint32(InDomViews.Views.Num()))
+		{
+			const FHairStrandsDeepShadowDatas& DOMs = InDomViews.Views[ViewIndex];
+			for (const FHairStrandsDeepShadowData& DomData : DOMs.Datas)
+			{
+				DrawFrustumWireframe(&ShadowFrustumPDI, DomData.WorldToLightTransform.Inverse(), FColor::Emerald, 0);
+				DrawWireBox(&ShadowFrustumPDI, DomData.Bounds.GetBox(), FColor::Yellow, 0);
+			}
+		}
+
+		// Current DOM
+		if (HairDebugMode == EHairDebugMode::DeepOpacityMaps && ViewIndex < uint32(InDomViews.Views.Num()))
+		{
+			const int32 CurrentIndex = FMath::Max(0, GDeepShadowDebugIndex);
+			const FHairStrandsDeepShadowDatas& DOMs = InDomViews.Views[ViewIndex];
+			if (CurrentIndex < DOMs.Datas.Num())
+			{
+				DrawFrustumWireframe(&ShadowFrustumPDI, DOMs.Datas[CurrentIndex].WorldToLightTransform.Inverse(), FColor::Emerald, 0);
+				DrawWireBox(&ShadowFrustumPDI, DOMs.Datas[CurrentIndex].Bounds.GetBox(), FColor::Yellow, 0);
+			}
+		}
+
+		// Voxelization
+		if (bIsVoxelMode && ViewIndex < uint32(InClusterViews.Views.Num()))
+		{
+			const FHairStrandsClusterDatas& ClusterDatas = InClusterViews.Views[ViewIndex];
+			for (const FHairStrandsClusterData& ClusterData : ClusterDatas.Datas)
+			{
+				DrawFrustumWireframe(&ShadowFrustumPDI, ClusterData.VoxelResources.WorldToClip.Inverse(), FColor::Purple, 0);
+
+				const FBox VoxelizationBox(ClusterData.GetMinBound(), ClusterData.GetMaxBound());
+				DrawWireBox(&ShadowFrustumPDI, VoxelizationBox, FColor::Red, 0);
+			}
+		}
+	}
+
+	if (HairDebugMode == EHairDebugMode::TAAResolveType || HairDebugMode == EHairDebugMode::SamplePerPixel || HairDebugMode == EHairDebugMode::CoverageType)
+	{
+		FRDGBuilder GraphBuilder(RHICmdList);
+		FRDGTextureRef SceneColorTexture = GraphBuilder.RegisterExternalTexture(SceneTargets.GetSceneColor(), TEXT("SceneColorTexture"));
+		if (ViewIndex < uint32(HairDatas->HairVisibilityViews.HairDatas.Num()))
+		{
+			const FHairStrandsVisibilityData& VisibilityData = HairDatas->HairVisibilityViews.HairDatas[ViewIndex];
+			AddDebugHairPass(GraphBuilder, &View, HairDebugMode, VisibilityData.CategorizationTexture, SceneTargets.SceneStencilSRV, SceneColorTexture);
+		}
+
+		GraphBuilder.Execute();
+	}
+
+	if (bIsVoxelMode)
+	{
+		if (ViewIndex < uint32(InClusterViews.Views.Num()))
+		{
+			FRDGBuilder GraphBuilder(RHICmdList);
+			FRDGTextureRef SceneColorTexture = GraphBuilder.RegisterExternalTexture(SceneTargets.GetSceneColor(), TEXT("SceneColorTexture"));
+			const FHairStrandsClusterDatas& ClusterDatas = InClusterViews.Views[ViewIndex];
+			AddVoxelRaymarchingPass(GraphBuilder, View, HairDebugMode, ClusterDatas, SceneColorTexture);
+			GraphBuilder.Execute();
+		}
+	}
+
+	// Text
+	if (HairDebugMode == EHairDebugMode::LightBounds || HairDebugMode == EHairDebugMode::DeepOpacityMaps)
+	{
+		FRenderTargetTemp TempRenderTarget(View, (const FTexture2DRHIRef&)SceneTargets.GetSceneColor()->GetRenderTargetItem().TargetableTexture);
+		FCanvas Canvas(&TempRenderTarget, NULL, ViewFamily.CurrentRealTime, ViewFamily.CurrentWorldTime, ViewFamily.DeltaWorldTime, View.FeatureLevel);
+
+		const FHairStrandsDeepShadowDatas& ViewData = InDomViews.Views[ViewIndex];
+		const uint32 DomTextureIndex = GDeepShadowDebugIndex;
+
+		const FIntPoint AtlasResolution = ViewData.Datas.Num() > 0 && ViewData.Datas[0].DepthTexture ? ViewData.Datas[0].DepthTexture->GetDesc().Extent : FIntPoint(0, 0);
+		float X = 20;
+		float Y = 38;
+
+		FLinearColor DebugColor(1, 1, 0);
+		FString Line;
+
+		const FVector4 HairComponent = GetHairComponents();
+		const uint32 bHairR = HairComponent.X > 0 ? 1 : 0;
+		const uint32 bHairTT = HairComponent.Y > 0 ? 1 : 0;
+		const uint32 bHairTRT = HairComponent.Z > 0 ? 1 : 0;
+		const uint32 bHairGlobalScattering = FMath::FloorToInt(HairComponent.W / 10.0f) > 0 ? 1 : 0;
+		const uint32 bHairLocalScattering  = FMath::Frac(HairComponent.W / 10.0f)*10.f > 0 ? 1 : 0;
+
+		Line = FString::Printf(TEXT("Hair Components : (R=%d, TT=%d, TRT=%d, GS=%d, LS=%d)"), bHairR, bHairTT, bHairTRT, bHairGlobalScattering, bHairLocalScattering);
+		Canvas.DrawShadowedString(X, Y += YStep, *Line, GetStatsFont(), DebugColor);
+		Line = FString::Printf(TEXT("----------------------------------------------------------------"));				Canvas.DrawShadowedString(X, Y += YStep, *Line, GetStatsFont(), DebugColor);
+		Line = FString::Printf(TEXT("Debug strands mode : %s"), ToString(StrandsDebugMode));							Canvas.DrawShadowedString(X, Y += YStep, *Line, GetStatsFont(), DebugColor);
+		Line = FString::Printf(TEXT("Voxelization : %s"), IsHairStrandsVoxelizationEnable() ? TEXT("On") : TEXT("Off"));		Canvas.DrawShadowedString(X, Y += YStep, *Line, GetStatsFont(), DebugColor);
+		Line = FString::Printf(TEXT("View rect optim.: %s"), IsHairStrandsViewRectOptimEnable() ? TEXT("On") : TEXT("Off"));	Canvas.DrawShadowedString(X, Y += YStep, *Line, GetStatsFont(), DebugColor);
+		Line = FString::Printf(TEXT("----------------------------------------------------------------"));				Canvas.DrawShadowedString(X, Y += YStep, *Line, GetStatsFont(), DebugColor);
+		Line = FString::Printf(TEXT("DOM Atlas resolution : %d/%d"), AtlasResolution.X, AtlasResolution.Y);				Canvas.DrawShadowedString(X, Y += YStep, *Line, GetStatsFont(), DebugColor);
+		Line = FString::Printf(TEXT("DOM Cluster count : %d"), ViewData.Datas.Num());									Canvas.DrawShadowedString(X, Y += YStep, *Line, GetStatsFont(), DebugColor);
+		Line = FString::Printf(TEXT("DOM Texture Index : %d/%d"), DomTextureIndex, ViewData.Datas.Num());				Canvas.DrawShadowedString(X, Y += YStep, *Line, GetStatsFont(), DebugColor);
+
+		uint32 BoundIndex = 0;
+		for (const FHairStrandsDeepShadowData& DomData : ViewData.Datas)
+		{
+			Line = FString::Printf(TEXT(" %d - Bound Radus: %f.2m (%dx%d)"), BoundIndex++, DomData.Bounds.GetSphere().W / 10.f, DomData.ShadowResolution.X, DomData.ShadowResolution.Y);
+			Canvas.DrawShadowedString(X, Y += YStep, *Line, GetStatsFont(), DebugColor);
+		}
+
+		Canvas.Flush_RenderThread(RHICmdList);
+	}
+
+	if (StrandsDebugMode != EHairStrandsDebugMode::None || HairDebugMode != EHairDebugMode::None)
+	{
+		float X = 40;
+		float Y = View.ViewRect.Height() - YStep * 3.f;
+		FString Line;
+		if (StrandsDebugMode != EHairStrandsDebugMode::None)
+			Line = FString::Printf(TEXT("Hair Debug mode - %s"), ToString(StrandsDebugMode));
+		else if (HairDebugMode != EHairDebugMode::None)
+			Line = FString::Printf(TEXT("Hair Debug mode - %s"), ToString(HairDebugMode));
+
+		FRenderTargetTemp TempRenderTarget(View, (const FTexture2DRHIRef&)SceneTargets.GetSceneColor()->GetRenderTargetItem().TargetableTexture);
+		FCanvas Canvas(&TempRenderTarget, NULL, ViewFamily.CurrentRealTime, ViewFamily.CurrentWorldTime, ViewFamily.DeltaWorldTime, View.FeatureLevel);
+		Canvas.DrawShadowedString(X, Y += YStep, *Line, GetStatsFont(), FLinearColor(1, 1, 0));
+		Canvas.Flush_RenderThread(RHICmdList);
+	}
+}
diff --git a/Engine/Source/Runtime/Renderer/Private/HairStrands/HairStrandsDebug.h b/Engine/Source/Runtime/Renderer/Private/HairStrands/HairStrandsDebug.h
new file mode 100644
--- /dev/null
+++ b/Engine/Source/Runtime/Renderer/Private/HairStrands/HairStrandsDebug.h
@@ -0,0 +1,15 @@
+// Copyright 1998-2019 Epic Games, Inc. All Rights Reserved.
+
+/*=============================================================================
+	HairStrandsDebug.h: Hair strands debug display.
+=============================================================================*/
+
+#pragma once
+
+#include "CoreMinimal.h"
+#include "RendererInterface.h"
+
+void RenderHairStrandsDebugInfo(
+	FRHICommandListImmediate& RHICmdList,
+	TArray<FViewInfo>& Views,
+	const struct FHairStrandsDatas* HairDatas);
\ No newline at end of file
diff --git a/Engine/Source/Runtime/Renderer/Private/HairStrands/HairStrandsDeepShadow.cpp b/Engine/Source/Runtime/Renderer/Private/HairStrands/HairStrandsDeepShadow.cpp
new file mode 100644
--- /dev/null
+++ b/Engine/Source/Runtime/Renderer/Private/HairStrands/HairStrandsDeepShadow.cpp
@@ -0,0 +1,276 @@
+// Copyright 1998-2019 Epic Games, Inc. All Rights Reserved.
+
+#include "HairStrandsDeepShadow.h"
+#include "HairStrandsRasterCommon.h"
+#include "HairStrandsCluster.h"
+#include "HairStrandsUtils.h"
+#include "LightSceneInfo.h"
+#include "ScenePrivate.h"
+
+// this is temporary until we split the voxelize and DOM path
+static int32 GDeepShadowResolution = 2048;
+static FAutoConsoleVariableRef CVarDeepShadowResolution(TEXT("r.HairStrands.DeepShadow.Resolution"), GDeepShadowResolution, TEXT("Shadow resolution for Deep Opacity Map rendering. (default = 2048)"));
+
+///////////////////////////////////////////////////////////////////////////////////////////////////
+
+typedef TArray<const FLightSceneInfo*, SceneRenderingAllocator> FLightSceneInfos;
+typedef TArray<FLightSceneInfos, SceneRenderingAllocator> FLightSceneInfosArray;
+
+static FLightSceneInfosArray GetVisibleDeepShadowLights(const FScene* Scene, const TArray<FViewInfo>& Views)
+{
+	// Collect all visible lights per view
+	FLightSceneInfosArray VisibleLightsPerView;
+	VisibleLightsPerView.SetNum(Views.Num());
+	for (TSparseArray<FLightSceneInfoCompact>::TConstIterator LightIt(Scene->Lights); LightIt; ++LightIt)
+	{
+		const FLightSceneInfoCompact& LightSceneInfoCompact = *LightIt;
+		const FLightSceneInfo* const LightSceneInfo = LightSceneInfoCompact.LightSceneInfo;
+
+		if (!LightSceneInfo->ShouldRenderLightViewIndependent())
+			continue;
+
+		// Check if the light is visible in any of the views.
+		for (int32 ViewIndex = 0; ViewIndex < Views.Num(); ViewIndex++)
+		{
+			const bool bIsCompatible = LightSceneInfo->ShouldRenderLight(Views[ViewIndex]) && LightSceneInfo->Proxy->CastsHairStrandsDeepShadow();
+			if (!bIsCompatible)
+				continue;
+
+			VisibleLightsPerView[ViewIndex].Add(LightSceneInfo);
+		}
+	}
+
+	return VisibleLightsPerView;
+}
+
+///////////////////////////////////////////////////////////////////////////////////////////////////
+
+static void RenderDeepShadowFrontDepth(
+	FRHICommandList& RHICmdList,
+	const FScene* Scene,
+	const FViewInfo* ViewInfo,
+	const FHairStrandsClusterData::TPrimitiveInfos& PrimitiveSceneInfo,
+	const FHairStrandsLightDesc& LightDesc,
+	const FIntRect& AtlasRect,
+	const bool bClearOutput,
+	TRefCountPtr<IPooledRenderTarget>& outShadowDepthRT)
+{
+	DECLARE_GPU_STAT(HairStrandsDeepShadowFrontDepth);
+	SCOPED_DRAW_EVENT(RHICmdList, HairStrandsDeepShadowFrontDepth);
+	SCOPED_GPU_STAT(RHICmdList, HairStrandsDeepShadowFrontDepth);
+
+	EDepthStencilTargetActions DepthStencilAction = MakeDepthStencilTargetActions(
+		MakeRenderTargetActions(bClearOutput ? ERenderTargetLoadAction::EClear : ERenderTargetLoadAction::ELoad, ERenderTargetStoreAction::EStore),
+		MakeRenderTargetActions(ERenderTargetLoadAction::ENoAction, ERenderTargetStoreAction::ENoAction));
+	FRHIRenderPassInfo RPInfo(outShadowDepthRT->GetRenderTargetItem().TargetableTexture, DepthStencilAction, nullptr, FExclusiveDepthStencil::DepthWrite_StencilNop);
+	RHICmdList.BeginRenderPass(RPInfo, TEXT("DeepShadowDepth"));
+	RasterHairStrands(
+		RHICmdList,
+		Scene,
+		ViewInfo,
+		PrimitiveSceneInfo,
+		LightDesc,
+		EHairStrandsRasterPassType::FrontDepth,
+		nullptr,
+		AtlasRect);
+	RHICmdList.EndRenderPass();
+}
+
+///////////////////////////////////////////////////////////////////////////////////////////////////
+
+static void RenderDeepShadowLayers(
+	FRHICommandListImmediate& RHICmdList,
+	const FScene* Scene,
+	const FViewInfo* ViewInfo,
+	const FHairStrandsClusterData::TPrimitiveInfos& PrimitiveSceneInfo,
+	const FHairStrandsLightDesc& LightDesc,
+	const FIntRect& AtlasRect,
+	FRHITexture* ShadowDepthTexture,
+	const bool bClearOutput,
+	TRefCountPtr<IPooledRenderTarget>& outDeepShadowLayersRT)
+{
+	DECLARE_GPU_STAT(HairStrandsDeepShadowLayers);
+	SCOPED_DRAW_EVENT(RHICmdList, HairStrandsDeepShadowLayers);
+	SCOPED_GPU_STAT(RHICmdList, HairStrandsDeepShadowLayers);
+
+	FRHIRenderPassInfo RPInfo(outDeepShadowLayersRT->GetRenderTargetItem().TargetableTexture, bClearOutput ? ERenderTargetActions::Clear_Store : ERenderTargetActions::Load_Store);
+	RHICmdList.BeginRenderPass(RPInfo, TEXT("DeepShadowLayers"));
+	RasterHairStrands(
+		RHICmdList,
+		Scene,
+		ViewInfo,
+		PrimitiveSceneInfo,
+		LightDesc,
+		EHairStrandsRasterPassType::DeepOpacityMap,
+		ShadowDepthTexture,
+		AtlasRect);
+	RHICmdList.EndRenderPass();
+}
+
+///////////////////////////////////////////////////////////////////////////////////////////////////
+
+FHairStrandsDeepShadowViews RenderHairStrandsDeepShadows(
+	FRHICommandListImmediate& RHICmdList,
+	const FScene* Scene,
+	const TArray<FViewInfo>& Views,
+	const FHairStrandsClusterViews& DeepShadowClusterViews)
+{
+	QUICK_SCOPE_CYCLE_COUNTER(STAT_CLM_RenderDeepShadow);
+	DECLARE_GPU_STAT(HairStrandsDeepShadow);
+	SCOPED_DRAW_EVENT(RHICmdList, HairStrandsDeepShadow);
+	SCOPED_GPU_STAT(RHICmdList, HairStrandsDeepShadow);
+
+	FLightSceneInfosArray VisibleLightsPerView = GetVisibleDeepShadowLights(Scene, Views);
+
+	// Compute the number of DOM which need to be created and insert default value
+	FHairStrandsDeepShadowViews DeepShadowsPerView;
+	uint32 DOMSlotCount = 0;
+	for (int32 ViewIndex = 0; ViewIndex < Views.Num(); ViewIndex++)
+	{
+		const FViewInfo& View = Views[ViewIndex];
+		if (View.Family)
+		{
+			const FHairStrandsClusterDatas& DeepShadowClusterDatas = DeepShadowClusterViews.Views[ViewIndex];
+			DeepShadowsPerView.Views.AddDefaulted();
+
+			if (DeepShadowClusterDatas.Datas.Num() == 0 || 
+				VisibleLightsPerView[ViewIndex].Num() == 0 ||
+				IsHairStrandsForVoxelTransmittanceAndShadowEnable()) 
+			{
+				continue; 
+			}
+
+			for (int32 ClusterIt = 0; ClusterIt < DeepShadowClusterDatas.Datas.Num(); ++ClusterIt)
+			{
+				const FHairStrandsClusterData& Cluster = DeepShadowClusterDatas.Datas[ClusterIt];
+				const FBoxSphereBounds ClusterBounds = Cluster.Bounds;
+
+				for (const FLightSceneInfo* LightInfo : VisibleLightsPerView[ViewIndex])
+				{
+					const FLightSceneProxy* LightProxy = LightInfo->Proxy;
+					if (!LightProxy->AffectsBounds(ClusterBounds))
+					{
+						continue;
+					}
+
+					DOMSlotCount++;
+				}
+			}
+		}
+	}
+
+	if (DOMSlotCount == 0)
+		return DeepShadowsPerView;
+
+	const uint32 AtlasSlotX = FGenericPlatformMath::CeilToInt(FMath::Sqrt(DOMSlotCount));
+	const FIntPoint AtlasSlotDimension(AtlasSlotX, AtlasSlotX == DOMSlotCount ? 1 : AtlasSlotX);
+	const FIntPoint AtlasSlotResolution(GDeepShadowResolution, GDeepShadowResolution);
+	const FIntPoint AtlasResolution(AtlasSlotResolution.X * AtlasSlotDimension.X, AtlasSlotResolution.Y * AtlasSlotDimension.Y);
+	TRefCountPtr<IPooledRenderTarget> FrontDepthAtlasTexture;
+	TRefCountPtr<IPooledRenderTarget> DeepShadowLayersAtlasTexture;
+
+	// Create Atlas resources for DOM
+	bool bClearFrontDepthAtlasTexture = true;
+	bool bClearLayerAtlasTexture = true;
+	{
+		{
+			FPooledRenderTargetDesc ShadowDesc(FPooledRenderTargetDesc::Create2DDesc(AtlasResolution, PF_DepthStencil, FClearValueBinding::DepthFar, TexCreate_None, TexCreate_DepthStencilTargetable, false));
+			GRenderTargetPool.FindFreeElement(RHICmdList, ShadowDesc, FrontDepthAtlasTexture, TEXT("ShadowDepth"));
+		}
+
+		{
+			FPooledRenderTargetDesc Desc(FPooledRenderTargetDesc::Create2DDesc(AtlasResolution, PF_FloatRGBA, FClearValueBinding::Transparent, TexCreate_None, TexCreate_RenderTargetable, false));
+			GRenderTargetPool.FindFreeElement(RHICmdList, Desc, DeepShadowLayersAtlasTexture, TEXT("DeepShadowLayers"));
+		}
+	}
+
+	// @todo_hair: could we share the rendering of the deep shadow across all view? Is it needed?
+	// @todo_hair: need to clarify what how view are/should be managed. At the moment this is 
+	// confusing as it seems there are many views but in the renderLight code, only one shadow mask 
+	// texture is created for all lights.
+	uint32 AtlasSlotIndex = 0;
+	for (int32 ViewIndex = 0; ViewIndex < Views.Num(); ViewIndex++)
+	{
+		const FViewInfo& View = Views[ViewIndex];
+		if (View.Family)
+		{
+			FHairStrandsDeepShadowDatas& DeepShadowDatas = DeepShadowsPerView.Views[ViewIndex];
+
+			const FHairStrandsClusterDatas& DeepShadowClusterDatas = DeepShadowClusterViews.Views[ViewIndex];
+			for (int32 ClusterIt = 0; ClusterIt < DeepShadowClusterDatas.Datas.Num(); ++ClusterIt)
+			{
+				// List of all the light in the scene.
+				uint32 LightLocalIndex = 0;
+				for (const FLightSceneInfo* LightInfo : VisibleLightsPerView[ViewIndex])
+				{
+					const FHairStrandsClusterData& Cluster = DeepShadowClusterDatas.Datas[ClusterIt];
+					FBoxSphereBounds ClusterBounds = Cluster.Bounds;
+
+					const FLightSceneProxy* LightProxy = LightInfo->Proxy;
+					if (!LightProxy->AffectsBounds(ClusterBounds))
+					{
+						continue;
+					}
+
+					const ELightComponentType LightType = (ELightComponentType)LightProxy->GetLightType();
+
+					FMinHairRadiusAtDepth1 MinStrandRadiusAtDepth1;
+					const FIntPoint AtlasRectOffset(
+						(AtlasSlotIndex % AtlasSlotDimension.X) * AtlasSlotResolution.X,
+						(AtlasSlotIndex / AtlasSlotDimension.X) * AtlasSlotResolution.Y);
+
+					// Note: LightPosition.W is used in the transmittance mask shader to differentiate between directional and local lights.
+					FHairStrandsDeepShadowData& DomData = DeepShadowDatas.Datas.AddZeroed_GetRef();
+					ComputeWorldToLightClip(DomData.WorldToLightTransform, MinStrandRadiusAtDepth1, ClusterBounds, *LightProxy, LightType, AtlasSlotResolution);
+					DomData.LightDirection = LightProxy->GetDirection();
+					DomData.LightPosition = FVector4(LightProxy->GetPosition(), LightType == ELightComponentType::LightType_Directional ? 0 : 1);
+					DomData.LightLuminance = LightProxy->GetColor();
+					DomData.LightType = LightType;
+					DomData.LightId = LightInfo->Id;
+					DomData.ShadowResolution = AtlasSlotResolution;
+					DomData.Bounds = ClusterBounds;
+					DomData.AtlasRect = FIntRect(AtlasRectOffset, AtlasRectOffset + AtlasSlotResolution);
+					DomData.ClusterId = Cluster.ClusterId;
+					DomData.DepthTexture = FrontDepthAtlasTexture;
+					DomData.LayersTexture = DeepShadowLayersAtlasTexture;
+					AtlasSlotIndex++;
+
+					FHairStrandsLightDesc LightDesc;
+					LightDesc.bIsOrtho = LightType == ELightComponentType::LightType_Directional;
+					LightDesc.LightDirection = DomData.LightDirection;
+					LightDesc.LightPosition = DomData.LightPosition;
+					LightDesc.WorldToLightClipTransform = DomData.WorldToLightTransform;
+					LightDesc.MinStrandRadiusAtDepth1 = MinStrandRadiusAtDepth1;
+
+					RenderDeepShadowFrontDepth(
+						RHICmdList,
+						Scene,
+						&View,
+						Cluster.PrimitivesInfos,
+						LightDesc,
+						DomData.AtlasRect,
+						bClearFrontDepthAtlasTexture,
+						DomData.DepthTexture);
+
+					RenderDeepShadowLayers(
+						RHICmdList,
+						Scene,
+						&View,
+						Cluster.PrimitivesInfos,
+						LightDesc,
+						DomData.AtlasRect,
+						DomData.DepthTexture->GetRenderTargetItem().TargetableTexture.GetReference(),
+						bClearLayerAtlasTexture,
+						DomData.LayersTexture);
+
+					bClearFrontDepthAtlasTexture = false;
+					bClearLayerAtlasTexture = false;
+
+					++LightLocalIndex;
+				}
+			}
+		}
+	}
+
+	return DeepShadowsPerView;
+}
diff --git a/Engine/Source/Runtime/Renderer/Private/HairStrands/HairStrandsDeepShadow.h b/Engine/Source/Runtime/Renderer/Private/HairStrands/HairStrandsDeepShadow.h
new file mode 100644
--- /dev/null
+++ b/Engine/Source/Runtime/Renderer/Private/HairStrands/HairStrandsDeepShadow.h
@@ -0,0 +1,50 @@
+// Copyright 1998-2019 Epic Games, Inc. All Rights Reserved.
+
+/*=============================================================================
+	HairStrandsDeepShadow.h: Hair strands deep shadow implementation.
+=============================================================================*/
+
+#pragma once
+
+#include "CoreMinimal.h"
+#include "RendererInterface.h"
+#include "SceneTypes.h"
+
+/// Hold deep shadow information for a given light.
+struct FHairStrandsDeepShadowData
+{
+	static const uint32 MaxClusterCount = 16u;
+
+	TRefCountPtr<IPooledRenderTarget> DepthTexture;
+	TRefCountPtr<IPooledRenderTarget> LayersTexture;
+	FMatrix WorldToLightTransform;
+	FIntRect AtlasRect;
+	uint32 ClusterId = ~0;
+
+	FIntPoint ShadowResolution = FIntPoint::ZeroValue;
+	uint32 LightId = ~0;
+	ELightComponentType LightType = LightType_MAX;
+	FVector  LightDirection;
+	FVector4 LightPosition;
+	FLinearColor LightLuminance;
+
+	FBoxSphereBounds Bounds;
+};
+
+/// Store all deep shadows infos for a given view
+struct FHairStrandsDeepShadowDatas
+{
+	TArray<FHairStrandsDeepShadowData, SceneRenderingAllocator> Datas;
+};
+
+/// Store all deep shadows info for all views
+struct FHairStrandsDeepShadowViews
+{
+	TArray<FHairStrandsDeepShadowDatas, SceneRenderingAllocator> Views;
+};
+
+FHairStrandsDeepShadowViews RenderHairStrandsDeepShadows(
+	FRHICommandListImmediate& RHICmdList,
+	const class FScene* Scene,
+	const TArray<FViewInfo>& Views,
+	const struct FHairStrandsClusterViews& ClusterViews);
\ No newline at end of file
diff --git a/Engine/Source/Runtime/Renderer/Private/HairStrands/HairStrandsEnvironment.cpp b/Engine/Source/Runtime/Renderer/Private/HairStrands/HairStrandsEnvironment.cpp
new file mode 100644
--- /dev/null
+++ b/Engine/Source/Runtime/Renderer/Private/HairStrands/HairStrandsEnvironment.cpp
@@ -0,0 +1,217 @@
+// Copyright 1998-2019 Epic Games, Inc. All Rights Reserved.
+
+
+#include "HairStrandsEnvironment.h"
+#include "Shader.h"
+#include "GlobalShader.h"
+#include "ShaderParameters.h"
+#include "ShaderParameterStruct.h"
+#include "SceneTextureParameters.h"
+#include "SceneRenderTargetParameters.h"
+#include "RenderGraph.h"
+#include "RenderGraphUtils.h"
+#include "SceneRendering.h"
+#include "PixelShaderUtils.h"
+#include "SystemTextures.h"
+#include "HairStrandsRendering.h"
+#include "ReflectionEnvironment.h"
+#include "ScenePrivate.h"
+
+///////////////////////////////////////////////////////////////////////////////////////////////////
+
+static int32 GHairSkylightingEnable = 1;
+static FAutoConsoleVariableRef CVarHairSkylightingEnable(TEXT("r.HairStrands.SkyLightingEnable"), GHairSkylightingEnable, TEXT("Enable sky lighting on hair."));
+
+static int32 GHairSkyAOEnable = 0;
+static FAutoConsoleVariableRef CVarHairSkyAOEnable(TEXT("r.HairStrands.SkyAOEnable"), GHairSkyAOEnable, TEXT("Enable (sky) AO on hair."));
+
+static float GHairSkylightingConeAngle = 3;
+static FAutoConsoleVariableRef CVarHairSkylightingConeAngle(TEXT("r.HairStrands.SkyLightingConeAngle"), GHairSkylightingConeAngle, TEXT("Cone angle for tracing sky lighting on hair."));
+
+static bool GetHairStrandsSkyLightingEnable() { return GHairSkylightingEnable > 0; }
+static bool GetHairStrandsSkyAOEnable() { return GHairSkyAOEnable > 0; }
+static float GetHairStrandsSkyLightingConeAngle() { return FMath::Max(0.f, GHairSkylightingConeAngle); }
+
+///////////////////////////////////////////////////////////////////////////////////////////////////
+
+DECLARE_GPU_STAT_NAMED(HairStrandsReflectionEnvironment, TEXT("Hair Strands Reflection Environment"));
+
+class FHairEnvironmentLightingPS : public FGlobalShader
+{
+	DECLARE_GLOBAL_SHADER(FHairEnvironmentLightingPS);
+	SHADER_USE_PARAMETER_STRUCT(FHairEnvironmentLightingPS, FGlobalShader)
+
+
+		class FRenderMode : SHADER_PERMUTATION_INT("PERMUTATION_RENDER_MODE", 2);
+	using FPermutationDomain = TShaderPermutationDomain<FRenderMode>;
+	static bool ShouldCompilePermutation(const FGlobalShaderPermutationParameters& Parameters)
+	{
+		return IsHairStrandsSupported(Parameters.Platform);;
+	}
+
+	BEGIN_SHADER_PARAMETER_STRUCT(FParameters, )
+
+		SHADER_PARAMETER(FVector, Voxel_MinAABB)
+		SHADER_PARAMETER(FVector, Voxel_MaxAABB)
+		SHADER_PARAMETER(uint32, Voxel_Resolution)
+		SHADER_PARAMETER(float, Voxel_DensityScale)
+		SHADER_PARAMETER(float, Voxel_DepthBiasScale)
+		SHADER_PARAMETER(float, Voxel_TanConeAngle)
+
+		SHADER_PARAMETER_TEXTURE(Texture2D, PreIntegratedGF)
+		SHADER_PARAMETER_SAMPLER(SamplerState, PreIntegratedGFSampler)
+
+		SHADER_PARAMETER_TEXTURE(Texture2D, HairCategorizationTexture)
+		SHADER_PARAMETER_TEXTURE(Texture2D, HairVisibilityNodeOffsetAndCount)
+		SHADER_PARAMETER_SRV(Buffer, HairVisibilityNodeData)
+
+		SHADER_PARAMETER_STRUCT_INCLUDE(FSceneTextureParameters, SceneTextures)
+		SHADER_PARAMETER_STRUCT_INCLUDE(FSceneTextureSamplerParameters, SceneTextureSamplers)
+
+		SHADER_PARAMETER_RDG_TEXTURE(TEXTURE3D, HairEnergyLUTTexture)
+		SHADER_PARAMETER_RDG_TEXTURE(TEXTURE3D, HairScatteringLUTTexture)
+		SHADER_PARAMETER_SAMPLER(SamplerState, HairLUTSampler)
+
+		SHADER_PARAMETER_RDG_TEXTURE(Texture3D, Voxel_DensityTexture)
+
+		SHADER_PARAMETER_STRUCT_REF(FViewUniformShaderParameters, ViewUniformBuffer)
+		SHADER_PARAMETER_STRUCT_REF(FReflectionUniformParameters, ReflectionsParameters)
+		SHADER_PARAMETER_STRUCT_REF(FReflectionCaptureShaderData, ReflectionCaptureData)
+		SHADER_PARAMETER_STRUCT_REF(FForwardLightData, ForwardLightData)
+
+		RENDER_TARGET_BINDING_SLOTS()
+		END_SHADER_PARAMETER_STRUCT()
+}; // FHairEnvironmentLightingPS
+
+IMPLEMENT_GLOBAL_SHADER(FHairEnvironmentLightingPS, "/Engine/Private/HairStrands/HairStrandsEnvironmentLighting.usf", "MainPS", SF_Pixel);
+
+void RenderHairStrandsEnvironmentLighting(
+	FRDGBuilder& GraphBuilder,
+	const uint32 ViewIndex,
+	const TArray<FViewInfo>& Views, 
+	const FHairStrandsDatas* HairDatas,
+	FRDGTextureRef SceneColorTexture,
+	FRDGTextureRef SceneColorSubPixelTexture)
+{
+	const FViewInfo& View = Views[ViewIndex];
+	const FHairStrandsVisibilityData* HairVisibilityData = HairDatas ? &HairDatas->HairVisibilityViews.HairDatas[ViewIndex] : nullptr;
+	const bool bRenderHairLighting = HairVisibilityData && HairVisibilityData->NodeIndex && HairVisibilityData->NodeDataSRV && (GetHairStrandsSkyLightingEnable() || GetHairStrandsSkyAOEnable());
+	if (!bRenderHairLighting)
+	{
+		return;
+	}
+
+	// @hair_todo: 
+	// * Add support for : BentNormal, global AO, SSR, DFSO
+	// * Add local reflection probe, current take into account only the sky lighting
+
+	FSceneTextureParameters SceneTextures;
+	SetupSceneTextureParameters(GraphBuilder, &SceneTextures);
+
+	RDG_GPU_STAT_SCOPE(GraphBuilder, HairStrandsReflectionEnvironment);
+
+	enum class EEnvRenderMode : uint8
+	{
+		Lighting,
+		AO
+	};
+	const FHairLUT InHairLUT = GetHairLUT(GraphBuilder.RHICmdList, View);
+
+	// Setup the parameters of the shader.
+	auto RenderEnv = [&](EEnvRenderMode RenderMode)
+	{
+		// Render the reflection environment with tiled deferred culling
+		const bool bHasBoxCaptures = (View.NumBoxReflectionCaptures > 0);
+		const bool bHasSphereCaptures = (View.NumSphereReflectionCaptures > 0);
+		FHairEnvironmentLightingPS::FParameters* PassParameters = GraphBuilder.AllocParameters<FHairEnvironmentLightingPS::FParameters>();
+				
+		PassParameters->HairEnergyLUTTexture = GraphBuilder.RegisterExternalTexture(InHairLUT.Textures[HairLUTType_MeanEnergy], TEXT("HairMeanEnergyLUTTexture"));;
+		PassParameters->HairScatteringLUTTexture = GraphBuilder.RegisterExternalTexture(InHairLUT.Textures[HairLUTType_DualScattering], TEXT("HairScatteringEnergyLUTTexture"));
+		PassParameters->HairLUTSampler = TStaticSamplerState<SF_Bilinear, AM_Clamp, AM_Clamp, AM_Clamp>::GetRHI();
+
+		if (HairDatas && ViewIndex < uint32(HairDatas->HairClusterPerViews.Views.Num()) && HairDatas->HairClusterPerViews.Views[ViewIndex].Datas.Num() > 0)
+		{
+			const uint32 ClusterIndex = 0; // Support only one cluster at the moment
+			const FHairStrandsClusterData& ClusterData = HairDatas->HairClusterPerViews.Views[ViewIndex].Datas[ClusterIndex];
+
+			PassParameters->Voxel_MinAABB = ClusterData.GetMinBound();
+			PassParameters->Voxel_MaxAABB = ClusterData.GetMaxBound();
+			PassParameters->Voxel_Resolution = ClusterData.GetResolution();
+			PassParameters->Voxel_DensityTexture = GraphBuilder.RegisterExternalTexture(ClusterData.VoxelResources.DensityTexture);
+			PassParameters->Voxel_DensityScale = GetHairStrandsVoxelizationDensityScale();
+			PassParameters->Voxel_DepthBiasScale = GetHairStrandsVoxelizationDepthBiasScale();
+			PassParameters->Voxel_TanConeAngle = FMath::Tan(FMath::DegreesToRadians(GetHairStrandsSkyLightingConeAngle()));
+		}
+		else
+		{
+			PassParameters->Voxel_DensityTexture = GraphBuilder.RegisterExternalTexture(GSystemTextures.BlackDummy);
+		}
+
+
+		PassParameters->PreIntegratedGF = GSystemTextures.PreintegratedGF->GetRenderTargetItem().ShaderResourceTexture;
+		PassParameters->PreIntegratedGFSampler = TStaticSamplerState<SF_Bilinear, AM_Clamp, AM_Clamp, AM_Clamp>::GetRHI();
+
+		PassParameters->SceneTextures = SceneTextures;
+		SetupSceneTextureSamplers(&PassParameters->SceneTextureSamplers);
+
+		PassParameters->ViewUniformBuffer = View.ViewUniformBuffer;
+		PassParameters->ReflectionCaptureData = View.ReflectionCaptureUniformBuffer;
+		{
+			FReflectionUniformParameters ReflectionUniformParameters;
+			SetupReflectionUniformParameters(View, ReflectionUniformParameters);
+			PassParameters->ReflectionsParameters = CreateUniformBufferImmediate(ReflectionUniformParameters, UniformBuffer_SingleDraw);
+		}
+		PassParameters->ForwardLightData = View.ForwardLightingResources->ForwardLightDataUniformBuffer;
+
+
+		// Bind hair data
+		if (HairVisibilityData)
+		{
+			PassParameters->HairCategorizationTexture = HairVisibilityData->CategorizationTexture->GetRenderTargetItem().ShaderResourceTexture;
+			PassParameters->HairVisibilityNodeOffsetAndCount = HairVisibilityData->NodeIndex->GetRenderTargetItem().ShaderResourceTexture;
+			PassParameters->HairVisibilityNodeData = HairVisibilityData->NodeDataSRV;
+		}
+
+		PassParameters->RenderTargets[0] = FRenderTargetBinding(SceneColorTexture, ERenderTargetLoadAction::ELoad);
+		if (RenderMode == EEnvRenderMode::Lighting)
+			PassParameters->RenderTargets[1] = FRenderTargetBinding(SceneColorSubPixelTexture, ERenderTargetLoadAction::ELoad);
+
+		FHairEnvironmentLightingPS::FPermutationDomain PermutationVector;
+		PermutationVector.Set<FHairEnvironmentLightingPS::FRenderMode>(RenderMode == EEnvRenderMode::AO ? 1 : 0);
+		TShaderMapRef<FHairEnvironmentLightingPS> PixelShader(View.ShaderMap, PermutationVector);
+		ClearUnusedGraphResources(*PixelShader, PassParameters);
+
+		GraphBuilder.AddPass(
+			RDG_EVENT_NAME("HairStrandsEnvironment %dx%d", View.ViewRect.Width(), View.ViewRect.Height()),
+			PassParameters,
+			ERDGPassFlags::Raster,
+			[PassParameters, &View, PixelShader, bRenderHairLighting, RenderMode](FRHICommandList& InRHICmdList)
+		{
+			InRHICmdList.SetViewport(View.ViewRect.Min.X, View.ViewRect.Min.Y, 0.0f, View.ViewRect.Max.X, View.ViewRect.Max.Y, 1.0f);
+
+			FGraphicsPipelineStateInitializer GraphicsPSOInit;
+			FPixelShaderUtils::InitFullscreenPipelineState(InRHICmdList, View.ShaderMap, *PixelShader, GraphicsPSOInit);
+
+			if (RenderMode == EEnvRenderMode::AO)
+				GraphicsPSOInit.BlendState = TStaticBlendState<
+				CW_RGBA, BO_Add, BF_Zero, BF_SourceAlpha, BO_Add, BF_Zero, BF_One
+				>::GetRHI();
+			else //if (RenderMode == EEnvRenderMode::Lighting)
+				GraphicsPSOInit.BlendState = TStaticBlendState<
+				CW_RGBA, BO_Add, BF_One, BF_One, BO_Add, BF_One, BF_One,
+				CW_RGBA, BO_Add, BF_One, BF_One, BO_Max, BF_SourceAlpha, BF_DestAlpha
+				>::GetRHI();
+
+			SetGraphicsPipelineState(InRHICmdList, GraphicsPSOInit);
+			SetShaderParameters(InRHICmdList, *PixelShader, PixelShader->GetPixelShader(), *PassParameters);
+			FPixelShaderUtils::DrawFullscreenTriangle(InRHICmdList);
+		});
+	}; // hair
+
+	// @hair_todo: move hair AO actually into the AO pass.
+	if (GetHairStrandsSkyAOEnable())
+		RenderEnv(EEnvRenderMode::AO);
+
+	if (GetHairStrandsSkyLightingEnable())
+		RenderEnv(EEnvRenderMode::Lighting);
+}
diff --git a/Engine/Source/Runtime/Renderer/Private/HairStrands/HairStrandsEnvironment.h b/Engine/Source/Runtime/Renderer/Private/HairStrands/HairStrandsEnvironment.h
new file mode 100644
--- /dev/null
+++ b/Engine/Source/Runtime/Renderer/Private/HairStrands/HairStrandsEnvironment.h
@@ -0,0 +1,19 @@
+// Copyright 1998-2019 Epic Games, Inc. All Rights Reserved.
+
+/*=============================================================================
+	HairStrandsEnvironment.h: Hair strands environment lighting.
+=============================================================================*/
+
+#pragma once
+
+#include "CoreMinimal.h"
+#include "RendererInterface.h"
+#include "RenderGraphResources.h"
+#
+void RenderHairStrandsEnvironmentLighting(
+	FRDGBuilder& GraphBuilder,
+	const uint32 ViewIndex,
+	const TArray<FViewInfo>& Views,
+	const struct FHairStrandsDatas* HairDatas,
+	FRDGTextureRef SceneColorTexture,
+	FRDGTextureRef SceneColorSubPixelTexture);
\ No newline at end of file
diff --git a/Engine/Source/Runtime/Renderer/Private/HairStrands/HairStrandsInterface.cpp b/Engine/Source/Runtime/Renderer/Private/HairStrands/HairStrandsInterface.cpp
new file mode 100644
--- /dev/null
+++ b/Engine/Source/Runtime/Renderer/Private/HairStrands/HairStrandsInterface.cpp
@@ -0,0 +1,69 @@
+// Copyright 1998-2019 Epic Games, Inc. All Rights Reserved.
+
+/*=============================================================================
+	HairStrandsInterface.h: Hair manager implementation.
+=============================================================================*/
+
+#include "HairStrandsInterface.h"
+#include "HairStrandsUtils.h"
+
+DEFINE_LOG_CATEGORY_STATIC(LogHairRendering, Log, All);
+
+static int32 GHairStrandsRenderingEnable = 1;
+static FAutoConsoleVariableRef CVarHairStrandsRenderingEnable(TEXT("r.HairStrands.Enable"), GHairStrandsRenderingEnable, TEXT("Enable/Disable hair strands rendering"));
+
+struct FHairStrandsManager
+{
+	struct Element
+	{
+		uint64 Id;
+		FHairStrandsInterpolation Data;
+	};
+	TArray<Element> Elements;
+};
+
+FHairStrandsManager GHairManager;
+
+void RegisterHairStrands(uint64 Id, const FHairStrandsInterpolation& E)
+{ 
+	for (int32 Index = 0; Index < GHairManager.Elements.Num(); ++Index)
+	{
+		if (GHairManager.Elements[Index].Id == Id)
+		{
+			// Component already registered. This should not happen.
+			UE_LOG(LogHairRendering, Warning, TEXT("Component already register. This should't happen. Please report this to a rendering engineer."))
+			return;
+		}
+	}
+
+	GHairManager.Elements.Add({ Id, E });
+}
+
+FHairStrandsInterpolation UnregisterHairStrands(uint64 Id)
+{
+	for (int32 Index=0;Index< GHairManager.Elements.Num();++Index)
+	{
+		if (GHairManager.Elements[Index].Id == Id)
+		{
+			FHairStrandsInterpolation Out = GHairManager.Elements[Index].Data;
+			GHairManager.Elements[Index] = GHairManager.Elements[GHairManager.Elements.Num()-1];
+			GHairManager.Elements.SetNum(GHairManager.Elements.Num() - 1);
+			return Out;
+		}
+	}
+
+	return FHairStrandsInterpolation();
+}
+
+void RunHairStrandsInterpolation(FRHICommandListImmediate& RHICmdList)
+{
+	for (FHairStrandsManager::Element& E : GHairManager.Elements)
+	{
+		if (E.Data.Input && E.Data.Output && E.Data.Function)
+		{
+			E.Data.Function(RHICmdList, E.Data.Input, E.Data.Output);
+		}
+	}
+}
+
+bool IsHairStrandsEnable(EShaderPlatform Platform) { return IsHairStrandsSupported(Platform) && GHairStrandsRenderingEnable == 1 && GHairManager.Elements.Num() > 0; }
\ No newline at end of file
diff --git a/Engine/Source/Runtime/Renderer/Private/HairStrands/HairStrandsLUT.cpp b/Engine/Source/Runtime/Renderer/Private/HairStrands/HairStrandsLUT.cpp
new file mode 100644
--- /dev/null
+++ b/Engine/Source/Runtime/Renderer/Private/HairStrands/HairStrandsLUT.cpp
@@ -0,0 +1,109 @@
+// Copyright 1998-2019 Epic Games, Inc. All Rights Reserved.
+
+#include "HairStrandsLUT.h"
+#include "HairStrandsUtils.h"
+#include "Shader.h"
+#include "GlobalShader.h"
+#include "ShaderParameters.h"
+#include "ShaderParameterStruct.h"
+#include "RenderGraphUtils.h"
+#include "SceneRendering.h"
+#include "SystemTextures.h"
+
+static int32 GHairLUTIncidentAngleCount = 64;
+static int32 GHairLUTRoughnessCount = 64;
+static int32 GHairLUTAbsorptionCount = 16;
+static int32 GHairLUTSampleCountScale = 1;
+static FAutoConsoleVariableRef CVarHairLUTIncidentAngleCount(TEXT("r.HairStrands.HairLUT.IncidentAngleCount"), GHairLUTIncidentAngleCount, TEXT("Change the number of slices of the hair LUT for the incident angle axis"));
+static FAutoConsoleVariableRef CVarHairLUTRoughnessCount(TEXT("r.HairStrands.HairLUT.RoughnessCount"), GHairLUTRoughnessCount, TEXT("Change the number of slices of the hair LUT for the roughness axis"));
+static FAutoConsoleVariableRef CVarHairLUTAbsorptionCount(TEXT("r.HairStrands.HairLUT.AbsorptionCount"), GHairLUTAbsorptionCount, TEXT("Change the number of slices of the hair LUT for the absorption axis"));
+static FAutoConsoleVariableRef CVarHairLUTSampleCount(TEXT("r.HairStrands.HairLUT.SampleCountScale"), GHairLUTSampleCountScale, TEXT("Change the number of sample used for computing the hair LUT. This is a multiplier, default is 1."));
+
+/////////////////////////////////////////////////////////////////////////////////////////
+
+class FHairLUTCS : public FGlobalShader
+{
+	DECLARE_GLOBAL_SHADER(FHairLUTCS);
+	SHADER_USE_PARAMETER_STRUCT(FHairLUTCS, FGlobalShader);
+
+	class FLUTType : SHADER_PERMUTATION_INT("PERMUTATION_LUT_TYPE", HairLUTTypeCount);
+	using FPermutationDomain = TShaderPermutationDomain<FLUTType>;
+
+	BEGIN_SHADER_PARAMETER_STRUCT(FParameters, )
+		SHADER_PARAMETER(uint32, AbsorptionCount)
+		SHADER_PARAMETER(uint32, RoughnessCount)
+		SHADER_PARAMETER(uint32, ThetaCount)
+		SHADER_PARAMETER(uint32, SampleCountScale)
+		SHADER_PARAMETER(FIntVector, OutputResolution)
+		SHADER_PARAMETER_RDG_TEXTURE_UAV(RWTexture2D, OutputColor)
+	END_SHADER_PARAMETER_STRUCT()
+
+public:
+	static bool ShouldCompilePermutation(const FGlobalShaderPermutationParameters& Parameters) { return IsHairStrandsSupported(Parameters.Platform); }
+};
+
+IMPLEMENT_GLOBAL_SHADER(FHairLUTCS, "/Engine/Private/HairStrands/HairStrandsLUT.usf", "MainCS", SF_Compute);
+
+static FRDGTextureRef AddHairLUTPass(
+	FRDGBuilder& GraphBuilder,
+	const FViewInfo& View,
+	const FHairLUTType LUTType)
+{
+	const FIntVector OutputResolution(GHairLUTIncidentAngleCount, GHairLUTRoughnessCount, GHairLUTAbsorptionCount);
+
+	FRDGTextureDesc OutputDesc;
+	OutputDesc.Extent.X = OutputResolution.X;
+	OutputDesc.Extent.Y = OutputResolution.Y;
+	OutputDesc.Depth = OutputResolution.Z;
+	OutputDesc.Format = PF_FloatRGBA;
+	OutputDesc.NumMips = 1;
+	OutputDesc.Flags = TexCreate_ShaderResource;
+	OutputDesc.TargetableFlags = TexCreate_UAV;
+	FRDGTextureRef HairLUTTexture = GraphBuilder.CreateTexture(OutputDesc, TEXT("HairLUT"));
+
+	FHairLUTCS::FParameters* Parameters = GraphBuilder.AllocParameters<FHairLUTCS::FParameters>();
+	Parameters->OutputColor = GraphBuilder.CreateUAV(FRDGTextureUAVDesc(HairLUTTexture, 0));
+	Parameters->ThetaCount = OutputResolution.X;
+	Parameters->RoughnessCount = OutputResolution.Y;
+	Parameters->AbsorptionCount = OutputResolution.Z;
+	Parameters->SampleCountScale = GHairLUTSampleCountScale;
+	Parameters->OutputResolution = OutputResolution;
+
+	FHairLUTCS::FPermutationDomain PermutationVector;
+	PermutationVector.Set<FHairLUTCS::FLUTType>(LUTType);
+
+	TShaderMapRef<FHairLUTCS> ComputeShader(View.ShaderMap, PermutationVector);
+	FComputeShaderUtils::AddPass(
+		GraphBuilder,
+		RDG_EVENT_NAME("HairStrandsLUT"),
+		*ComputeShader,
+		Parameters,
+		FComputeShaderUtils::GetGroupCount(OutputResolution, FIntVector(FComputeShaderUtils::kGolden2DGroupSize)));
+
+	return HairLUTTexture;
+}
+
+FHairLUT GetHairLUT(FRHICommandListImmediate& RHICmdList, const FViewInfo& View)
+{
+	// Lazy LUT generation
+	const bool bNeedGenerate = 
+		GSystemTextures.HairLUT0.GetReference() == nullptr || GSystemTextures.HairLUT1.GetReference() == nullptr ||
+		GSystemTextures.HairLUT0.GetReference()->GetRenderTargetItem().ShaderResourceTexture->GetSizeXYZ() != FIntVector(GHairLUTIncidentAngleCount, GHairLUTRoughnessCount, GHairLUTAbsorptionCount);
+	if (bNeedGenerate)
+	{
+		FRDGBuilder GraphBuilder(RHICmdList);
+
+		FRDGTextureRef HairDualScatteringLUTTexture = AddHairLUTPass(GraphBuilder, View, HairLUTType_DualScattering);
+		GraphBuilder.QueueTextureExtraction(HairDualScatteringLUTTexture, &GSystemTextures.HairLUT0);
+
+		FRDGTextureRef HairMeanEnergyLUTTexture		= AddHairLUTPass(GraphBuilder, View, HairLUTType_MeanEnergy);		
+		GraphBuilder.QueueTextureExtraction(HairMeanEnergyLUTTexture, &GSystemTextures.HairLUT1);
+
+		GraphBuilder.Execute();
+	}
+	FHairLUT HairLUTData;
+	HairLUTData.Textures[HairLUTType_DualScattering] = GSystemTextures.HairLUT0;
+	HairLUTData.Textures[HairLUTType_MeanEnergy] = GSystemTextures.HairLUT1;
+
+	return HairLUTData;
+}
diff --git a/Engine/Source/Runtime/Renderer/Private/HairStrands/HairStrandsLUT.h b/Engine/Source/Runtime/Renderer/Private/HairStrands/HairStrandsLUT.h
new file mode 100644
--- /dev/null
+++ b/Engine/Source/Runtime/Renderer/Private/HairStrands/HairStrandsLUT.h
@@ -0,0 +1,25 @@
+// Copyright 1998-2019 Epic Games, Inc. All Rights Reserved.
+
+/*=============================================================================
+	HairRendering.h: Hair strands LUT generation.
+=============================================================================*/
+
+#pragma once
+
+#include "CoreMinimal.h"
+#include "RendererInterface.h"
+
+enum FHairLUTType
+{
+	HairLUTType_DualScattering,
+	HairLUTType_MeanEnergy,
+	HairLUTTypeCount
+};
+
+struct FHairLUT
+{
+	TRefCountPtr<IPooledRenderTarget> Textures[HairLUTTypeCount];
+};
+
+/// Returns Hair LUTs. LUTs are generated on demand.
+FHairLUT GetHairLUT(FRHICommandListImmediate& RHICmdList, const FViewInfo& View);
