Copying //UE4/Dev-Physics to Dev-Main (//UE4/Dev-Main) [at] 9098278

diff --git a/Engine/Source/Runtime/Renderer/Private/HairStrands/HairStrandsRasterCommon.cpp b/Engine/Source/Runtime/Renderer/Private/HairStrands/HairStrandsRasterCommon.cpp
new file mode 100644
--- /dev/null
+++ b/Engine/Source/Runtime/Renderer/Private/HairStrands/HairStrandsRasterCommon.cpp
@@ -0,0 +1,460 @@
+// Copyright 1998-2019 Epic Games, Inc. All Rights Reserved.
+
+#include "HairStrandsRasterCommon.h"
+#include "HairStrandsUtils.h"
+#include "PrimitiveSceneProxy.h"
+#include "Shader.h"
+#include "MeshMaterialShader.h"
+#include "ShaderParameters.h"
+#include "ShaderParameterStruct.h"
+#include "MeshPassProcessor.h"
+#include "MeshPassProcessor.inl"
+#include "ScenePrivate.h"
+
+// this is temporary until we split the voxelize and DOM path
+
+
+/////////////////////////////////////////////////////////////////////////////////////////
+
+BEGIN_GLOBAL_SHADER_PARAMETER_STRUCT(FDeepShadowPassUniformParameters, )
+	//SHADER_PARAMETER_STRUCT(FSceneTexturesUniformParameters, SceneTextures)
+	SHADER_PARAMETER(FMatrix, WorldToClipMatrix)
+	SHADER_PARAMETER(FVector, LightDirection)
+	SHADER_PARAMETER(FVector4, SliceValue)
+	SHADER_PARAMETER(uint32, bIsOrtho)
+	SHADER_PARAMETER(FIntRect, AtlasRect)
+	SHADER_PARAMETER(float, MinStrandRadius)
+	SHADER_PARAMETER(FVector, VoxelMinAABB)
+	SHADER_PARAMETER(FVector, VoxelMaxAABB)
+	SHADER_PARAMETER(uint32, VoxelResolution)
+	SHADER_PARAMETER_TEXTURE(Texture2D<float>, FrontDepthTexture)
+END_GLOBAL_SHADER_PARAMETER_STRUCT()
+IMPLEMENT_GLOBAL_SHADER_PARAMETER_STRUCT(FDeepShadowPassUniformParameters, "DeepShadowPass");
+
+
+/////////////////////////////////////////////////////////////////////////////////////////
+
+class FDeepShadowDepthMeshVS : public FMeshMaterialShader
+{
+	DECLARE_SHADER_TYPE(FDeepShadowDepthMeshVS, MeshMaterial);
+
+protected:
+
+	FDeepShadowDepthMeshVS(const ShaderMetaType::CompiledShaderInitializerType& Initializer)
+		: FMeshMaterialShader(Initializer)
+	{
+		ERHIFeatureLevel::Type FeatureLevel = GetMaxSupportedFeatureLevel((EShaderPlatform)Initializer.Target.Platform);
+		check(FSceneInterface::GetShadingPath(FeatureLevel) != EShadingPath::Mobile);
+		// deferred
+		PassUniformBuffer.Bind(Initializer.ParameterMap, FDeepShadowPassUniformParameters::StaticStructMetadata.GetShaderVariableName());
+	}
+
+	FDeepShadowDepthMeshVS() {}
+
+	static bool ShouldCompilePermutation(const FMeshMaterialShaderPermutationParameters& Parameters)
+	{
+		return IsCompatibleWithHairStrands(Parameters.Material, Parameters.Platform);
+	}
+
+	static void ModifyCompilationEnvironment(const FMaterialShaderPermutationParameters& Parameters, FShaderCompilerEnvironment& OutEnvironment)
+	{
+		FMaterialShader::ModifyCompilationEnvironment(Parameters, OutEnvironment);
+		OutEnvironment.SetDefine(TEXT("MESH_RENDER_MODE"), 0);
+	}
+};
+
+IMPLEMENT_MATERIAL_SHADER_TYPE(, FDeepShadowDepthMeshVS, TEXT("/Engine/Private/HairStrands/HairStrandsDeepShadowVS.usf"), TEXT("Main"), SF_Vertex);
+
+/////////////////////////////////////////////////////////////////////////////////////////
+
+class FDeepShadowDomMeshVS : public FMeshMaterialShader
+{
+	DECLARE_SHADER_TYPE(FDeepShadowDomMeshVS, MeshMaterial);
+
+protected:
+
+	FDeepShadowDomMeshVS(const ShaderMetaType::CompiledShaderInitializerType& Initializer)
+		: FMeshMaterialShader(Initializer)
+	{
+		ERHIFeatureLevel::Type FeatureLevel = GetMaxSupportedFeatureLevel((EShaderPlatform)Initializer.Target.Platform);
+		check(FSceneInterface::GetShadingPath(FeatureLevel) != EShadingPath::Mobile);
+		// deferred
+		PassUniformBuffer.Bind(Initializer.ParameterMap, FDeepShadowPassUniformParameters::StaticStructMetadata.GetShaderVariableName());
+	}
+
+	FDeepShadowDomMeshVS() {}
+
+	static bool ShouldCompilePermutation(const FMeshMaterialShaderPermutationParameters& Parameters)
+	{
+		return IsCompatibleWithHairStrands(Parameters.Material, Parameters.Platform);
+	}
+
+	static void ModifyCompilationEnvironment(const FMaterialShaderPermutationParameters& Parameters, FShaderCompilerEnvironment& OutEnvironment)
+	{
+		FMaterialShader::ModifyCompilationEnvironment(Parameters, OutEnvironment);
+		OutEnvironment.SetDefine(TEXT("MESH_RENDER_MODE"), 1);
+	}
+};
+
+IMPLEMENT_MATERIAL_SHADER_TYPE(, FDeepShadowDomMeshVS, TEXT("/Engine/Private/HairStrands/HairStrandsDeepShadowVS.usf"), TEXT("Main"), SF_Vertex);
+
+/////////////////////////////////////////////////////////////////////////////////////////
+
+template<bool bVoxelizeMaterial>
+class FVoxelMeshVS : public FMeshMaterialShader
+{
+	DECLARE_SHADER_TYPE(FVoxelMeshVS, MeshMaterial);
+
+protected:
+
+	FVoxelMeshVS(const ShaderMetaType::CompiledShaderInitializerType& Initializer)
+		: FMeshMaterialShader(Initializer)
+	{
+		ERHIFeatureLevel::Type FeatureLevel = GetMaxSupportedFeatureLevel((EShaderPlatform)Initializer.Target.Platform);
+		check(FSceneInterface::GetShadingPath(FeatureLevel) != EShadingPath::Mobile);
+		// deferred
+		PassUniformBuffer.Bind(Initializer.ParameterMap, FDeepShadowPassUniformParameters::StaticStructMetadata.GetShaderVariableName());
+	}
+
+	FVoxelMeshVS() {}
+
+	static bool ShouldCompilePermutation(const FMeshMaterialShaderPermutationParameters& Parameters)
+	{
+		return IsCompatibleWithHairStrands(Parameters.Material, Parameters.Platform);
+	}
+
+	static void ModifyCompilationEnvironment(const FMaterialShaderPermutationParameters& Parameters, FShaderCompilerEnvironment& OutEnvironment)
+	{
+		FMaterialShader::ModifyCompilationEnvironment(Parameters, OutEnvironment);
+		OutEnvironment.SetDefine(TEXT("MESH_RENDER_MODE"), 2);
+		OutEnvironment.SetDefine(TEXT("SUPPORT_TANGENT_PROPERTY"), bVoxelizeMaterial ? 1 : 0);
+		OutEnvironment.SetDefine(TEXT("SUPPORT_MATERIAL_PROPERTY"), bVoxelizeMaterial ? 1 : 0);
+	}
+};
+IMPLEMENT_MATERIAL_SHADER_TYPE(template<>, FVoxelMeshVS<false>, TEXT("/Engine/Private/HairStrands/HairStrandsDeepShadowVS.usf"), TEXT("Main"), SF_Vertex);
+IMPLEMENT_MATERIAL_SHADER_TYPE(template<>, FVoxelMeshVS<true>, TEXT("/Engine/Private/HairStrands/HairStrandsDeepShadowVS.usf"), TEXT("Main"), SF_Vertex);
+
+/////////////////////////////////////////////////////////////////////////////////////////
+
+class FDeepShadowDepthMeshPS : public FMeshMaterialShader
+{
+	DECLARE_SHADER_TYPE(FDeepShadowDepthMeshPS, MeshMaterial);
+
+public:
+
+	FDeepShadowDepthMeshPS(const ShaderMetaType::CompiledShaderInitializerType& Initializer)
+		: FMeshMaterialShader(Initializer)
+	{
+		ERHIFeatureLevel::Type FeatureLevel = GetMaxSupportedFeatureLevel((EShaderPlatform)Initializer.Target.Platform);
+		check(FSceneInterface::GetShadingPath(FeatureLevel) != EShadingPath::Mobile);
+		PassUniformBuffer.Bind(Initializer.ParameterMap, FDeepShadowPassUniformParameters::StaticStructMetadata.GetShaderVariableName());
+	}
+
+	FDeepShadowDepthMeshPS() {}
+
+	static bool ShouldCompilePermutation(const FMeshMaterialShaderPermutationParameters& Parameters)
+	{
+		return IsCompatibleWithHairStrands(Parameters.Material, Parameters.Platform);
+	}
+};
+IMPLEMENT_MATERIAL_SHADER_TYPE(, FDeepShadowDepthMeshPS, TEXT("/Engine/Private/HairStrands/HairStrandsDeepShadowPS.usf"), TEXT("MainDepth"), SF_Pixel);
+
+/////////////////////////////////////////////////////////////////////////////////////////
+
+class FDeepShadowDomMeshPS : public FMeshMaterialShader
+{
+	DECLARE_SHADER_TYPE(FDeepShadowDomMeshPS, MeshMaterial);
+
+public:
+
+	FDeepShadowDomMeshPS(const ShaderMetaType::CompiledShaderInitializerType& Initializer)
+		: FMeshMaterialShader(Initializer)
+	{
+		ERHIFeatureLevel::Type FeatureLevel = GetMaxSupportedFeatureLevel((EShaderPlatform)Initializer.Target.Platform);
+		check(FSceneInterface::GetShadingPath(FeatureLevel) != EShadingPath::Mobile);
+		PassUniformBuffer.Bind(Initializer.ParameterMap, FDeepShadowPassUniformParameters::StaticStructMetadata.GetShaderVariableName());
+	}
+
+	FDeepShadowDomMeshPS() {}
+
+	static bool ShouldCompilePermutation(const FMeshMaterialShaderPermutationParameters& Parameters)
+	{
+		return IsCompatibleWithHairStrands(Parameters.Material, Parameters.Platform);
+	}
+
+	virtual bool Serialize(FArchive& Ar) override
+	{
+		bool bShaderHasOutdatedParameters = FMeshMaterialShader::Serialize(Ar);
+		return bShaderHasOutdatedParameters;
+	}
+};
+IMPLEMENT_MATERIAL_SHADER_TYPE(, FDeepShadowDomMeshPS, TEXT("/Engine/Private/HairStrands/HairStrandsDeepShadowPS.usf"), TEXT("MainDom"), SF_Pixel);
+
+/////////////////////////////////////////////////////////////////////////////////////////
+
+template<bool bVoxelizeMaterial>
+class FVoxelMeshPS : public FMeshMaterialShader
+{
+	DECLARE_SHADER_TYPE(FVoxelMeshPS, MeshMaterial);
+
+public:
+
+	FVoxelMeshPS(const ShaderMetaType::CompiledShaderInitializerType& Initializer)
+		: FMeshMaterialShader(Initializer)
+	{
+		ERHIFeatureLevel::Type FeatureLevel = GetMaxSupportedFeatureLevel((EShaderPlatform)Initializer.Target.Platform);
+		check(FSceneInterface::GetShadingPath(FeatureLevel) != EShadingPath::Mobile);
+		PassUniformBuffer.Bind(Initializer.ParameterMap, FDeepShadowPassUniformParameters::StaticStructMetadata.GetShaderVariableName());
+		DensityTexture.Bind(Initializer.ParameterMap, TEXT("DensityTexture"));
+		if (bVoxelizeMaterial)
+		{
+			TangentXTexture.Bind(Initializer.ParameterMap, TEXT("TangentXTexture"));
+			TangentYTexture.Bind(Initializer.ParameterMap, TEXT("TangentYTexture"));
+			TangentZTexture.Bind(Initializer.ParameterMap, TEXT("TangentZTexture"));
+			MaterialTexture.Bind(Initializer.ParameterMap, TEXT("MaterialTexture"));
+		}
+	}
+
+	FVoxelMeshPS() {}
+
+	static bool ShouldCompilePermutation(const FMeshMaterialShaderPermutationParameters& Parameters)
+	{
+		return IsCompatibleWithHairStrands(Parameters.Material, Parameters.Platform);
+	}
+
+	static void ModifyCompilationEnvironment(const FMaterialShaderPermutationParameters& Parameters, FShaderCompilerEnvironment& OutEnvironment)
+	{
+		FMaterialShader::ModifyCompilationEnvironment(Parameters, OutEnvironment);
+		OutEnvironment.SetDefine(TEXT("SUPPORT_TANGENT_PROPERTY"), bVoxelizeMaterial ? 1 : 0);
+		OutEnvironment.SetDefine(TEXT("SUPPORT_MATERIAL_PROPERTY"), bVoxelizeMaterial ? 1 : 0);
+	}
+
+	virtual bool Serialize(FArchive& Ar) override
+	{
+		bool bShaderHasOutdatedParameters = FMeshMaterialShader::Serialize(Ar);
+		Ar << DensityTexture;
+		if (bVoxelizeMaterial)
+		{
+			Ar << TangentXTexture;
+			Ar << TangentYTexture;
+			Ar << TangentZTexture;
+			Ar << MaterialTexture;
+		}
+		return bShaderHasOutdatedParameters;
+	}
+private:
+	FRWShaderParameter DensityTexture;
+	FRWShaderParameter TangentXTexture;
+	FRWShaderParameter TangentYTexture;
+	FRWShaderParameter TangentZTexture;
+	FRWShaderParameter MaterialTexture;
+};
+IMPLEMENT_MATERIAL_SHADER_TYPE(template<>, FVoxelMeshPS<true>, TEXT("/Engine/Private/HairStrands/HairStrandsDeepShadowPS.usf"), TEXT("MainVoxel"), SF_Pixel);
+IMPLEMENT_MATERIAL_SHADER_TYPE(template<>, FVoxelMeshPS<false>, TEXT("/Engine/Private/HairStrands/HairStrandsDeepShadowPS.usf"), TEXT("MainVoxel"), SF_Pixel);
+
+/////////////////////////////////////////////////////////////////////////////////////////
+
+class FDeepShadowMeshProcessor : public FMeshPassProcessor
+{
+public:
+
+	FDeepShadowMeshProcessor(
+		const FScene* Scene,
+		const FSceneView* InViewIfDynamicMeshCommand,
+		const FMeshPassProcessorRenderState& InPassDrawRenderState,
+		FDynamicPassMeshDrawListContext* InDrawListContext,
+		const EHairStrandsRasterPassType PType);
+
+	virtual void AddMeshBatch(const FMeshBatch& RESTRICT MeshBatch, uint64 BatchElementMask, const FPrimitiveSceneProxy* RESTRICT PrimitiveSceneProxy, int32 StaticMeshId = -1) override final;
+
+private:
+	template<typename VertexShaderType, typename PixelShaderType>
+	void Process(
+		const FMeshBatch& MeshBatch,
+		uint64 BatchElementMask,
+		const FPrimitiveSceneProxy* RESTRICT PrimitiveSceneProxy,
+		int32 StaticMeshId,
+		const FMaterialRenderProxy& RESTRICT MaterialRenderProxy,
+		const FMaterial& RESTRICT MaterialResource,
+		ERasterizerFillMode MeshFillMode,
+		ERasterizerCullMode MeshCullMode);
+
+	EHairStrandsRasterPassType DeepShadowPassType;
+	FMeshPassProcessorRenderState PassDrawRenderState;
+};
+
+void FDeepShadowMeshProcessor::AddMeshBatch(const FMeshBatch& RESTRICT MeshBatch, uint64 BatchElementMask, const FPrimitiveSceneProxy* RESTRICT PrimitiveSceneProxy, int32 StaticMeshId)
+{
+	// Determine the mesh's material and blend mode.
+	const FMaterialRenderProxy* FallbackMaterialRenderProxyPtr = nullptr;
+	const FMaterial& Material = MeshBatch.MaterialRenderProxy->GetMaterialWithFallback(FeatureLevel, FallbackMaterialRenderProxyPtr);
+	const bool bIsCompatible = IsCompatibleWithHairStrands(&Material, FeatureLevel);
+
+	if (bIsCompatible
+		&& (!PrimitiveSceneProxy || PrimitiveSceneProxy->ShouldRenderInMainPass())
+		&& ShouldIncludeDomainInMeshPass(Material.GetMaterialDomain()))
+	{
+		const FMaterialRenderProxy& MaterialRenderProxy = FallbackMaterialRenderProxyPtr ? *FallbackMaterialRenderProxyPtr : *MeshBatch.MaterialRenderProxy;
+		const ERasterizerFillMode MeshFillMode = ComputeMeshFillMode(MeshBatch, Material);
+		const ERasterizerCullMode MeshCullMode = DeepShadowPassType == EHairStrandsRasterPassType::FrontDepth ? ComputeMeshCullMode(MeshBatch, Material) : CM_None;
+
+		if (DeepShadowPassType == EHairStrandsRasterPassType::FrontDepth)
+			Process<FDeepShadowDepthMeshVS, FDeepShadowDepthMeshPS>(MeshBatch, BatchElementMask, PrimitiveSceneProxy, StaticMeshId, MaterialRenderProxy, Material, MeshFillMode, MeshCullMode);
+		else if (DeepShadowPassType == EHairStrandsRasterPassType::DeepOpacityMap)
+			Process<FDeepShadowDomMeshVS, FDeepShadowDomMeshPS>(MeshBatch, BatchElementMask, PrimitiveSceneProxy, StaticMeshId, MaterialRenderProxy, Material, MeshFillMode, MeshCullMode);
+		else if (DeepShadowPassType == EHairStrandsRasterPassType::Voxelization)
+			Process<FVoxelMeshVS<false>, FVoxelMeshPS<false>>(MeshBatch, BatchElementMask, PrimitiveSceneProxy, StaticMeshId, MaterialRenderProxy, Material, MeshFillMode, MeshCullMode);
+		else if (DeepShadowPassType == EHairStrandsRasterPassType::VoxelizationMaterial)
+			Process<FVoxelMeshVS<true>, FVoxelMeshPS<true>>(MeshBatch, BatchElementMask, PrimitiveSceneProxy, StaticMeshId, MaterialRenderProxy, Material, MeshFillMode, MeshCullMode);
+	}
+}
+
+// Vertex is either FDeepShadowDepthMeshVS, FDeepShadowDomMeshVS, or FVoxelMeshVS
+// Pixel  is either FDeepShadowDepthMeshPS, FDeepShadowDomMeshPS, or FVoxelMeshPS
+template<typename VertexShaderType, typename PixelShaderType>
+void FDeepShadowMeshProcessor::Process(
+	const FMeshBatch& MeshBatch,
+	uint64 BatchElementMask,
+	const FPrimitiveSceneProxy* RESTRICT PrimitiveSceneProxy,
+	int32 StaticMeshId,
+	const FMaterialRenderProxy& RESTRICT MaterialRenderProxy,
+	const FMaterial& RESTRICT MaterialResource,
+	ERasterizerFillMode MeshFillMode,
+	ERasterizerCullMode MeshCullMode)
+{
+	const FVertexFactory* VertexFactory = MeshBatch.VertexFactory;
+
+	TMeshProcessorShaders<
+		VertexShaderType,
+		FMeshMaterialShader,
+		FMeshMaterialShader,
+		PixelShaderType> PassShaders;
+	{
+		const EMaterialTessellationMode MaterialTessellationMode = MaterialResource.GetTessellationMode();
+		FVertexFactoryType* VertexFactoryType = VertexFactory->GetType();
+		PassShaders.DomainShader = nullptr;
+		PassShaders.HullShader = nullptr;
+		PassShaders.VertexShader = MaterialResource.GetShader<VertexShaderType>(VertexFactoryType);
+		PassShaders.PixelShader = MaterialResource.GetShader<PixelShaderType>(VertexFactoryType);
+	}
+
+	FMeshPassProcessorRenderState DrawRenderState(PassDrawRenderState);
+
+	FMeshMaterialShaderElementData ShaderElementData;
+	ShaderElementData.InitializeMeshMaterialData(ViewIfDynamicMeshCommand, PrimitiveSceneProxy, MeshBatch, StaticMeshId, false);
+
+	BuildMeshDrawCommands(
+		MeshBatch,
+		BatchElementMask,
+		PrimitiveSceneProxy,
+		MaterialRenderProxy,
+		MaterialResource,
+		DrawRenderState,
+		PassShaders,
+		MeshFillMode,
+		MeshCullMode,
+		FMeshDrawCommandSortKey::Default,
+		EMeshPassFeatures::Default,
+		ShaderElementData);
+}
+
+FDeepShadowMeshProcessor::FDeepShadowMeshProcessor(
+	const FScene* Scene,
+	const FSceneView* InViewIfDynamicMeshCommand,
+	const FMeshPassProcessorRenderState& InPassDrawRenderState,
+	FDynamicPassMeshDrawListContext* InDrawListContext,
+	const EHairStrandsRasterPassType PType)
+	: FMeshPassProcessor(Scene, Scene->GetFeatureLevel(), InViewIfDynamicMeshCommand, InDrawListContext)
+	, DeepShadowPassType(PType)
+	, PassDrawRenderState(InPassDrawRenderState)
+{
+}
+
+/////////////////////////////////////////////////////////////////////////////////////////
+
+void RasterHairStrands(
+	FRHICommandList& RHICmdList,
+	const FScene* Scene,
+	const FViewInfo* ViewInfo,
+	const FHairStrandsClusterData::TPrimitiveInfos& PrimitiveSceneInfos,
+	const FHairStrandsLightDesc& LightDesc,
+	const EHairStrandsRasterPassType ShadowPassType,
+	FRHITexture* DeepShadowDepthTexture,
+	const FIntRect& AtlasRect,
+	const FVector MinAABB,
+	const FVector MaxAABB)
+{
+	check(RHICmdList.IsInsideRenderPass());
+	check(IsInRenderingThread());
+
+	SCOPE_CYCLE_COUNTER(STAT_RenderPerObjectShadowDepthsTime);
+
+	TUniformBufferRef<FDeepShadowPassUniformParameters> PassUniformBuffer;
+	{
+		if (ShadowPassType == EHairStrandsRasterPassType::DeepOpacityMap)
+		{
+			check(DeepShadowDepthTexture);
+		}
+
+		FDeepShadowPassUniformParameters PassParameters;
+		PassParameters.WorldToClipMatrix = LightDesc.WorldToLightClipTransform;
+		PassParameters.LightDirection = LightDesc.LightDirection;
+		PassParameters.SliceValue = FVector4(1, 1, 1, 1);
+		PassParameters.FrontDepthTexture = DeepShadowDepthTexture;
+		PassParameters.MinStrandRadius = LightDesc.MinStrandRadiusAtDepth1.Primary;
+		PassParameters.bIsOrtho = LightDesc.bIsOrtho ? 1u : 0u;
+		PassParameters.AtlasRect = AtlasRect;
+		PassParameters.VoxelMinAABB = MinAABB;
+		PassParameters.VoxelMaxAABB = MaxAABB;
+		PassParameters.VoxelResolution = AtlasRect.Width();
+		PassUniformBuffer = CreateUniformBufferImmediate(PassParameters, UniformBuffer_SingleDraw, EUniformBufferValidation::None);
+	}
+
+	FMeshPassProcessorRenderState DrawRenderState(*ViewInfo, PassUniformBuffer);
+	{
+		RHICmdList.SetViewport(AtlasRect.Min.X, AtlasRect.Min.Y, 0.0f, AtlasRect.Max.X, AtlasRect.Max.Y, 1.0f);
+
+		if (ShadowPassType == EHairStrandsRasterPassType::DeepOpacityMap)
+		{
+			DrawRenderState.SetBlendState(TStaticBlendState<
+				CW_RGBA, BO_Add, BF_One, BF_One, BO_Add, BF_One, BF_One,
+				CW_RGBA, BO_Add, BF_One, BF_One, BO_Add, BF_One, BF_One>::GetRHI());
+			DrawRenderState.SetDepthStencilState(TStaticDepthStencilState<false, CF_Always>::GetRHI());
+		}
+		else if (ShadowPassType == EHairStrandsRasterPassType::FrontDepth)
+		{
+			DrawRenderState.SetBlendState(TStaticBlendState<
+				CW_RGBA, BO_Add, BF_One, BF_Zero, BO_Add, BF_One, BF_Zero,
+				CW_RGBA, BO_Add, BF_One, BF_Zero, BO_Add, BF_One, BF_Zero>::GetRHI());
+			DrawRenderState.SetDepthStencilState(TStaticDepthStencilState<true, CF_DepthNearOrEqual>::GetRHI());
+		}
+		else if (ShadowPassType == EHairStrandsRasterPassType::Voxelization || ShadowPassType == EHairStrandsRasterPassType::VoxelizationMaterial)
+		{
+			DrawRenderState.SetBlendState(TStaticBlendState<
+				CW_RGBA, BO_Add, BF_One, BF_Zero, BO_Add, BF_One, BF_Zero>::GetRHI());
+			DrawRenderState.SetDepthStencilState(TStaticDepthStencilState<false, CF_Always>::GetRHI());
+		}
+
+		FDynamicMeshDrawCommandStorage DynamicMeshDrawCommandStorage; // << Were would thid be stored?
+		FMeshCommandOneFrameArray VisibleMeshDrawCommands;
+		FGraphicsMinimalPipelineStateSet GraphicsMinimalPipelineStateSet;
+		FDynamicPassMeshDrawListContext ShadowContext(DynamicMeshDrawCommandStorage, VisibleMeshDrawCommands, GraphicsMinimalPipelineStateSet);
+
+
+		FDeepShadowMeshProcessor DeepShadowMeshProcessor(Scene, ViewInfo /* is a SceneView */, DrawRenderState, &ShadowContext, ShadowPassType);
+
+		for (const FMeshBatchAndRelevance& MeshBatchAndRelevance : PrimitiveSceneInfos)
+		{
+			const FMeshBatch& MeshBatch = *MeshBatchAndRelevance.Mesh;
+			const uint64 BatchElementMask = ~0ull;
+			DeepShadowMeshProcessor.AddMeshBatch(MeshBatch, BatchElementMask, MeshBatchAndRelevance.PrimitiveSceneProxy);
+		}
+
+		if (VisibleMeshDrawCommands.Num() > 0)
+		{
+			FRHIVertexBuffer* PrimitiveIdVertexBuffer = nullptr;
+			SortAndMergeDynamicPassMeshDrawCommands(ViewInfo->GetFeatureLevel(), VisibleMeshDrawCommands, DynamicMeshDrawCommandStorage, PrimitiveIdVertexBuffer, 1);
+			SubmitMeshDrawCommands(VisibleMeshDrawCommands, GraphicsMinimalPipelineStateSet, PrimitiveIdVertexBuffer, 0, false, 1, RHICmdList);
+		}
+	}
+}
diff --git a/Engine/Source/Runtime/Renderer/Private/HairStrands/HairStrandsRasterCommon.h b/Engine/Source/Runtime/Renderer/Private/HairStrands/HairStrandsRasterCommon.h
new file mode 100644
--- /dev/null
+++ b/Engine/Source/Runtime/Renderer/Private/HairStrands/HairStrandsRasterCommon.h
@@ -0,0 +1,32 @@
+// Copyright 1998-2019 Epic Games, Inc. All Rights Reserved.
+
+/*=============================================================================
+	HairRendering.h: Hair rendering implementation.
+=============================================================================*/
+
+#pragma once
+
+#include "CoreMinimal.h"
+#include "RendererInterface.h"
+#include "HairStrandsCluster.h"
+
+enum class EHairStrandsRasterPassType : uint8
+{
+	FrontDepth,
+	DeepOpacityMap,
+	Voxelization,
+	VoxelizationMaterial,
+};
+
+// to be removed
+void RasterHairStrands(
+	FRHICommandList& RHICmdList,
+	const FScene* Scene,
+	const FViewInfo* ViewInfo,
+	const FHairStrandsClusterData::TPrimitiveInfos& PrimitiveSceneInfos,
+	const struct FHairStrandsLightDesc& LightDesc,
+	const EHairStrandsRasterPassType ShadowPassType,
+	FRHITexture* DeepShadowDepthTexture,
+	const FIntRect& AtlasRect,
+	const FVector MinAABB = FVector::ZeroVector,
+	const FVector MaxAABB = FVector::ZeroVector);
\ No newline at end of file
diff --git a/Engine/Source/Runtime/Renderer/Private/HairStrands/HairStrandsRendering.cpp b/Engine/Source/Runtime/Renderer/Private/HairStrands/HairStrandsRendering.cpp
new file mode 100644
--- /dev/null
+++ b/Engine/Source/Runtime/Renderer/Private/HairStrands/HairStrandsRendering.cpp
@@ -0,0 +1,3 @@
+// Copyright 1998-2019 Epic Games, Inc. All Rights Reserved.
+
+#include "HairStrandsRendering.h"
\ No newline at end of file
diff --git a/Engine/Source/Runtime/Renderer/Private/HairStrands/HairStrandsRendering.h b/Engine/Source/Runtime/Renderer/Private/HairStrands/HairStrandsRendering.h
new file mode 100644
--- /dev/null
+++ b/Engine/Source/Runtime/Renderer/Private/HairStrands/HairStrandsRendering.h
@@ -0,0 +1,29 @@
+// Copyright 1998-2019 Epic Games, Inc. All Rights Reserved.
+
+/*=============================================================================
+	HairRendering.h: Hair rendering implementation.
+=============================================================================*/
+
+#pragma once
+
+#include "CoreMinimal.h"
+#include "RendererInterface.h"
+#include "HairStrandsUtils.h"
+#include "HairStrandsCluster.h"
+#include "HairStrandsLUT.h"
+#include "HairStrandsDeepShadow.h"
+#include "HairStrandsVoxelization.h"
+#include "HairStrandsVisibility.h"
+#include "HairStrandsTransmittance.h"
+#include "HairStrandsEnvironment.h"
+#include "HairStrandsComposition.h"
+#include "HairStrandsDebug.h"
+#include "HairStrandsInterface.h"
+
+/// Hold all the hair strands data
+struct FHairStrandsDatas
+{
+	FHairStrandsDeepShadowViews DeepShadowViews;
+	FHairStrandsVisibilityViews HairVisibilityViews;
+	FHairStrandsClusterViews HairClusterPerViews;
+};
diff --git a/Engine/Source/Runtime/Renderer/Private/HairStrands/HairStrandsTransmittance.cpp b/Engine/Source/Runtime/Renderer/Private/HairStrands/HairStrandsTransmittance.cpp
new file mode 100644
--- /dev/null
+++ b/Engine/Source/Runtime/Renderer/Private/HairStrands/HairStrandsTransmittance.cpp
@@ -0,0 +1,635 @@
+// Copyright 1998-2019 Epic Games, Inc. All Rights Reserved.
+
+#include "HairStrandsTransmittance.h"
+#include "HairStrandsCluster.h"
+#include "HairStrandsLUT.h"
+#include "HairStrandsDeepShadow.h"
+#include "HairStrandsVoxelization.h"
+#include "HairStrandsRendering.h"
+
+#include "Shader.h"
+#include "GlobalShader.h"
+#include "ShaderParameters.h"
+#include "ShaderParameterStruct.h"
+#include "SceneTextureParameters.h"
+#include "RenderGraphUtils.h"
+#include "PostProcessing.h"
+
+///////////////////////////////////////////////////////////////////////////////////////////////////
+
+static int32 GDeepShadowDebugMode = 0;
+static FAutoConsoleVariableRef CVarDeepShadowDebugMode(TEXT("r.HairStrands.DeepShadow.DebugMode"), GDeepShadowDebugMode, TEXT("Color debug mode for deep shadow"));
+static uint32 GetDeepShadowDebugMode() { return uint32(FMath::Max(0, GDeepShadowDebugMode)); }
+
+static int32 GDeepShadowKernelType = 2; // 0:linear, 1:PCF_2x2, 2: PCF_6x4, 3:PCSS
+static float GDeepShadowKernelAperture = 1;
+static FAutoConsoleVariableRef CVarDeepShadowKernelType(TEXT("r.HairStrands.DeepShadow.KernelType"), GDeepShadowKernelType, TEXT("Set the type of kernel used for evaluating hair transmittance, 0:linear, 1:PCF_2x2, 2: PCF_6x4, 3:PCSS"));
+static FAutoConsoleVariableRef CVarDeepShadowKernelAperture(TEXT("r.HairStrands.DeepShadow.KernelAperture"), GDeepShadowKernelAperture, TEXT("Set the aperture angle, in degree, used by the kernel for evaluating the hair transmittance when using PCSS kernel"));
+
+static uint32 GetDeepShadowKernelType() { return uint32(FMath::Max(0, GDeepShadowKernelType)); }
+static float GetDeepShadowKernelAperture() { return GDeepShadowKernelAperture; }
+
+static int32 GStrandHairShadowMaskKernelType = 2;
+static FAutoConsoleVariableRef GVarDeepShadowShadowMaskKernelType(TEXT("r.HairStrands.DeepShadow.ShadowMaskKernelType"), GStrandHairShadowMaskKernelType, TEXT("Set the kernel type for filtering shadow cast by hair on opaque geometry (0:2x2, 1:4x4, 2:Gaussian8, 3:Gaussian16). Default is 0"));
+
+static float GDeepShadowDensityScale = 2;	// Default is arbitrary, based on Mike asset
+static float GDeepShadowDepthBiasScale = 2;	// Default is arbitrary, based on content test
+static FAutoConsoleVariableRef CVarDeepShadowDensityScale(TEXT("r.HairStrands.DeepShadow.DensityScale"), GDeepShadowDensityScale, TEXT("Set density scale for compensating the lack of hair fiber in an asset"));
+static FAutoConsoleVariableRef CVarDeepShadowDepthBiasScale(TEXT("r.HairStrands.DeepShadow.DepthBiasScale"), GDeepShadowDepthBiasScale, TEXT("Set depth bias scale for transmittance computation"));
+
+float GetDeepShadowDensityScale() { return FMath::Max(0.0f, GDeepShadowDensityScale); }
+float GetDeepShadowDepthBiasScale() { return FMath::Max(0.0f, GDeepShadowDepthBiasScale); }
+///////////////////////////////////////////////////////////////////////////////////////////////////
+
+static bool HasDeepShadowData(const FLightSceneInfo* LightSceneInfo, const FHairStrandsDeepShadowDatas& InDatas)
+{
+	for (const FHairStrandsDeepShadowData& DomData : InDatas.Datas)
+	{
+		if (DomData.LightId == LightSceneInfo->Id)
+			return true;
+	}
+
+	return false;
+}
+
+
+///////////////////////////////////////////////////////////////////////////////////////////////////
+// Transmittance Mask from deep shadow
+
+enum FHairTransmittanceType
+{
+	FHairTransmittanceType_DeepShadow,
+	FHairTransmittanceType_Voxel,
+	FHairTransmittanceTypeCount
+};
+
+class FDeepTransmittanceMaskCS : public FGlobalShader
+{
+	DECLARE_GLOBAL_SHADER(FDeepTransmittanceMaskCS);
+	SHADER_USE_PARAMETER_STRUCT(FDeepTransmittanceMaskCS, FGlobalShader);
+
+	class FTransmittanceType : SHADER_PERMUTATION_INT("PERMUTATION_TRANSMITTANCE_TYPE", FHairTransmittanceTypeCount);
+	class FTransmittanceGroupSize : SHADER_PERMUTATION_INT("PERMUTATION_GROUP_SIZE", 2);
+	using FPermutationDomain = TShaderPermutationDomain<FTransmittanceType, FTransmittanceGroupSize>;
+
+	BEGIN_SHADER_PARAMETER_STRUCT(FParameters, )
+		SHADER_PARAMETER_STRUCT_INCLUDE(FSceneTextureParameters, SceneTextures)
+
+		SHADER_PARAMETER_ARRAY(FIntVector4, DeepShadow_AtlasSlotOffsets, [FHairStrandsDeepShadowData::MaxClusterCount])
+		SHADER_PARAMETER_ARRAY(FMatrix, DeepShadow_WorldToLightTransforms, [FHairStrandsDeepShadowData::MaxClusterCount])
+		SHADER_PARAMETER(FIntPoint, DeepShadow_Resolution)
+		SHADER_PARAMETER(FVector, LightDirection)
+		SHADER_PARAMETER(FVector4, LightPosition)
+		SHADER_PARAMETER(float, DepthBiasScale)
+		SHADER_PARAMETER(float, DensityScale)
+		SHADER_PARAMETER(float, DeepShadow_KernelAperture)
+		SHADER_PARAMETER(uint32, DeepShadow_KernelType)
+		SHADER_PARAMETER(uint32, DeepShadow_DebugMode)
+		SHADER_PARAMETER(FMatrix, DeepShadow_ShadowToWorld)
+
+		SHADER_PARAMETER_ARRAY(FVector4, Voxel_MinAABBs, [FHairStrandsDeepShadowData::MaxClusterCount])
+		SHADER_PARAMETER_ARRAY(FVector4, Voxel_MaxAABBs, [FHairStrandsDeepShadowData::MaxClusterCount])
+		SHADER_PARAMETER(uint32, Voxel_Resolution)
+		SHADER_PARAMETER_RDG_TEXTURE(Texture3D, Voxel_DensityTexture0)
+		SHADER_PARAMETER_RDG_TEXTURE(Texture3D, Voxel_DensityTexture1)
+		SHADER_PARAMETER_RDG_TEXTURE(Texture3D, Voxel_DensityTexture2)
+		SHADER_PARAMETER_RDG_TEXTURE(Texture3D, Voxel_DensityTexture3)
+		SHADER_PARAMETER_RDG_TEXTURE(Texture3D, Voxel_DensityTexture4)
+		SHADER_PARAMETER_RDG_TEXTURE(Texture3D, Voxel_DensityTexture5)
+		SHADER_PARAMETER_RDG_TEXTURE(Texture3D, Voxel_DensityTexture6)
+		SHADER_PARAMETER_RDG_TEXTURE(Texture3D, Voxel_DensityTexture7)
+
+		SHADER_PARAMETER_RDG_TEXTURE(Texture2D, DeepShadow_FrontDepthTexture)
+		SHADER_PARAMETER_RDG_TEXTURE(Texture2D, DeepShadow_DomTexture)
+		SHADER_PARAMETER_RDG_TEXTURE(Texture3D, HairLUTTexture)
+		SHADER_PARAMETER_RDG_BUFFER_SRV(StructuredBuffer, HairVisibilityNodeData)
+		SHADER_PARAMETER_RDG_BUFFER_SRV(StructuredBuffer, HairVisibilityNodeCoord)
+		SHADER_PARAMETER_RDG_BUFFER(StructuredBuffer, IndirectArgsBuffer)
+
+		SHADER_PARAMETER_RDG_BUFFER_UAV(StructuredBuffer, OutputColor)
+		SHADER_PARAMETER_SAMPLER(SamplerState, LinearSampler)
+		SHADER_PARAMETER_SAMPLER(SamplerState, ShadowSampler)
+
+		SHADER_PARAMETER_STRUCT_REF(FViewUniformShaderParameters, ViewUniformBuffer)
+		END_SHADER_PARAMETER_STRUCT()
+
+public:
+	static bool ShouldCompilePermutation(const FGlobalShaderPermutationParameters& Parameters) { return IsHairStrandsSupported(Parameters.Platform); }
+};
+
+IMPLEMENT_GLOBAL_SHADER(FDeepTransmittanceMaskCS, "/Engine/Private/HairStrands/HairStrandsDeepTransmittanceMask.usf", "MainCS", SF_Compute);
+
+struct FDeepShadowTransmittanceParams
+{
+	FIntVector4 DeepShadow_AtlasSlotOffsets[FHairStrandsDeepShadowData::MaxClusterCount];
+	FMatrix DeepShadow_WorldToLightTransforms[FHairStrandsDeepShadowData::MaxClusterCount];
+	FIntPoint DeepShadow_Resolution = FIntPoint(0, 0);
+	FVector LightDirection = FVector::ZeroVector;
+	FVector4 LightPosition = FVector4(0, 0, 0, 0);
+	float DepthBiasScale = 0;
+	float DensityScale = 0;
+	FMatrix DeepShadow_ShadowToWorld = FMatrix::Identity;
+
+	FRDGTextureRef DeepShadow_FrontDepthTexture = nullptr;
+	FRDGTextureRef DeepShadow_DomTexture = nullptr;
+
+	FRDGBufferRef HairVisibilityNodeData = nullptr;
+	FRDGBufferRef HairVisibilityNodeCoord = nullptr;
+
+	FRDGTextureRef Cluster_DensityTextures[FHairStrandsDeepShadowData::MaxClusterCount];
+	FVector4 Cluster_MinAABBs[FHairStrandsDeepShadowData::MaxClusterCount];
+	FVector4 Cluster_MaxAABBs[FHairStrandsDeepShadowData::MaxClusterCount];
+	uint32  Cluster_Resolution;
+};
+
+// Transmittance mask
+static FRDGBufferRef AddDeepShadowTransmittanceMaskPass(
+	FRDGBuilder& GraphBuilder,
+	const FSceneTextureParameters& SceneTextures,
+	const FViewInfo& View,
+	const FHairTransmittanceType TransmittanceType,
+	const FDeepShadowTransmittanceParams& Params,
+	const uint32 NodeGroupSize,
+	FRDGTextureRef HairLUTTexture,
+	FRDGBufferRef IndirectArgsBuffer)
+{
+	FRDGBufferRef OutBuffer = GraphBuilder.CreateBuffer(FRDGBufferDesc::CreateStructuredDesc(
+		4 * sizeof(float),
+		Params.HairVisibilityNodeData->Desc.NumElements),
+		TEXT("HairTransmittanceNodeData"));
+
+	FDeepTransmittanceMaskCS::FParameters* Parameters = GraphBuilder.AllocParameters<FDeepTransmittanceMaskCS::FParameters>();
+	Parameters->ViewUniformBuffer = View.ViewUniformBuffer;
+	Parameters->SceneTextures = SceneTextures;
+	Parameters->HairLUTTexture = HairLUTTexture;
+	Parameters->DeepShadow_FrontDepthTexture = Params.DeepShadow_FrontDepthTexture;
+	Parameters->DeepShadow_DomTexture = Params.DeepShadow_DomTexture;
+	Parameters->LinearSampler = TStaticSamplerState<SF_Bilinear, AM_Clamp, AM_Clamp, AM_Clamp>::GetRHI();
+	Parameters->OutputColor = GraphBuilder.CreateUAV(OutBuffer);
+	Parameters->DeepShadow_Resolution = Params.DeepShadow_Resolution;
+	Parameters->LightDirection = Params.LightDirection;
+	Parameters->LightPosition = Params.LightPosition;
+	Parameters->DepthBiasScale = Params.DepthBiasScale;
+	Parameters->DensityScale = Params.DensityScale;
+	Parameters->DeepShadow_KernelAperture = GetDeepShadowKernelAperture();
+	Parameters->DeepShadow_KernelType = GetDeepShadowKernelType();
+	Parameters->DeepShadow_DebugMode = GetDeepShadowDebugMode();
+	Parameters->DeepShadow_ShadowToWorld = Params.DeepShadow_ShadowToWorld;
+	Parameters->IndirectArgsBuffer = IndirectArgsBuffer;
+
+	memcpy(&(Parameters->DeepShadow_AtlasSlotOffsets[0]), Params.DeepShadow_AtlasSlotOffsets, sizeof(FIntVector4) * FHairStrandsDeepShadowData::MaxClusterCount);
+	memcpy(&(Parameters->DeepShadow_WorldToLightTransforms[0]), Params.DeepShadow_WorldToLightTransforms, sizeof(FMatrix) * FHairStrandsDeepShadowData::MaxClusterCount);
+
+	memcpy(&(Parameters->Voxel_MinAABBs[0]), Params.Cluster_MinAABBs, sizeof(FVector4) * FHairStrandsDeepShadowData::MaxClusterCount);
+	memcpy(&(Parameters->Voxel_MaxAABBs[0]), Params.Cluster_MaxAABBs, sizeof(FVector4) * FHairStrandsDeepShadowData::MaxClusterCount);
+	Parameters->Voxel_Resolution = Params.Cluster_Resolution;
+	Parameters->Voxel_DensityTexture0 = Params.Cluster_DensityTextures[0];
+	Parameters->Voxel_DensityTexture1 = Params.Cluster_DensityTextures[1];
+	Parameters->Voxel_DensityTexture2 = Params.Cluster_DensityTextures[2];
+	Parameters->Voxel_DensityTexture3 = Params.Cluster_DensityTextures[3];
+	Parameters->Voxel_DensityTexture4 = Params.Cluster_DensityTextures[4];
+	Parameters->Voxel_DensityTexture5 = Params.Cluster_DensityTextures[5];
+	Parameters->Voxel_DensityTexture6 = Params.Cluster_DensityTextures[6];
+	Parameters->Voxel_DensityTexture7 = Params.Cluster_DensityTextures[7];
+
+	Parameters->HairVisibilityNodeData = GraphBuilder.CreateSRV(Params.HairVisibilityNodeData);
+	Parameters->HairVisibilityNodeCoord = GraphBuilder.CreateSRV(Params.HairVisibilityNodeCoord);
+
+	check(NodeGroupSize == 64 || NodeGroupSize == 32);
+	FDeepTransmittanceMaskCS::FPermutationDomain PermutationVector;
+	PermutationVector.Set<FDeepTransmittanceMaskCS::FTransmittanceType>(TransmittanceType);
+	PermutationVector.Set<FDeepTransmittanceMaskCS::FTransmittanceGroupSize>(NodeGroupSize == 64 ? 0 : (NodeGroupSize == 32 ? 1 : 2));
+
+	TShaderMapRef<FDeepTransmittanceMaskCS> ComputeShader(View.ShaderMap, PermutationVector);
+	FComputeShaderUtils::AddPass(
+		GraphBuilder,
+		RDG_EVENT_NAME("HairStrandsTransmittanceMask"),
+		*ComputeShader,
+		Parameters,
+		IndirectArgsBuffer,
+		0);
+
+	return OutBuffer;
+}
+
+///////////////////////////////////////////////////////////////////////////////////////////////////
+// Opaque Mask from deep shadow
+
+enum FHairOpaqueMaskType
+{
+	FHairOpaqueMaskType_DeepShadow,
+	FHairOpaqueMaskType_Voxel,
+	FHairOpaqueMaskTypeCount
+};
+
+class FDeepShadowMaskPS : public FGlobalShader
+{
+	DECLARE_GLOBAL_SHADER(FDeepShadowMaskPS);
+	SHADER_USE_PARAMETER_STRUCT(FDeepShadowMaskPS, FGlobalShader);
+
+	class FOpaqueMaskType : SHADER_PERMUTATION_INT("PERMUTATION_OPAQUEMASK_TYPE", FHairOpaqueMaskTypeCount);
+	class FKernelType : SHADER_PERMUTATION_INT("PERMUTATION_KERNEL_TYPE", 4);
+	using FPermutationDomain = TShaderPermutationDomain<FOpaqueMaskType, FKernelType>;
+
+	BEGIN_SHADER_PARAMETER_STRUCT(FParameters, )
+		SHADER_PARAMETER_STRUCT_INCLUDE(FSceneTextureParameters, SceneTextures)
+
+		SHADER_PARAMETER(FIntPoint, DeepShadow_SlotOffset)
+		SHADER_PARAMETER(FIntPoint, DeepShadow_SlotResolution)
+		SHADER_PARAMETER(FMatrix, DeepShadow_WorldToLightTransform)
+		SHADER_PARAMETER(uint32, DeepShadow_IsWholeSceneLight)
+
+		SHADER_PARAMETER(FVector4, Voxel_LightPosition)
+		SHADER_PARAMETER(FVector, Voxel_LightDirection)
+		SHADER_PARAMETER(float, Voxel_DensityScale)
+		SHADER_PARAMETER(FVector, Voxel_MinAABB)
+		SHADER_PARAMETER(uint32, Voxel_Resolution)
+		SHADER_PARAMETER(FVector, Voxel_MaxAABB)
+
+		SHADER_PARAMETER_RDG_TEXTURE(Texture3D, Voxel_DensityTexture)
+		SHADER_PARAMETER_RDG_TEXTURE(Texture2D, DeepShadow_FrontDepthTexture)
+		SHADER_PARAMETER_RDG_TEXTURE(Texture2D, CategorizationTexture)
+		SHADER_PARAMETER_SAMPLER(SamplerState, LinearSampler)
+		SHADER_PARAMETER_SAMPLER(SamplerState, ShadowSampler)
+
+		SHADER_PARAMETER_STRUCT_REF(FViewUniformShaderParameters, ViewUniformBuffer)
+		RENDER_TARGET_BINDING_SLOTS()
+		END_SHADER_PARAMETER_STRUCT()
+
+public:
+	static bool ShouldCompilePermutation(const FGlobalShaderPermutationParameters& Parameters) { return IsHairStrandsSupported(Parameters.Platform); }
+};
+
+IMPLEMENT_GLOBAL_SHADER(FDeepShadowMaskPS, "/Engine/Private/HairStrands/HairStrandsDeepShadowMask.usf", "MainPS", SF_Pixel);
+
+
+struct FDeepShadowOpaqueParams
+{
+	FRDGTextureRef	CategorizationTexture = nullptr;
+
+	FMatrix			DeepShadow_WorldToLightTransform;
+	FIntRect		DeepShadow_AtlasRect;
+	FRDGTextureRef	DeepShadow_FrontDepthTexture = nullptr;
+	bool			DeepShadow_bIsWholeSceneLight = false;
+
+	FVector			Voxel_LightDirection = FVector::ZeroVector;
+	FVector4		Voxel_LightPosition = FVector4(0, 0, 0, 0);
+	float			Voxel_DensityScale = 0;
+	FRDGTextureRef	Voxel_DensityTexture = nullptr;
+	FVector			Voxel_MinAABB;
+	FVector			Voxel_MaxAABB;
+	uint32			Voxel_Resolution;
+};
+
+// Opaque mask
+static void AddDeepShadowOpaqueMaskPass(
+	FRDGBuilder& GraphBuilder,
+	const FSceneTextureParameters& SceneTextures,
+	const FViewInfo& View,
+	const FHairOpaqueMaskType HairOpaqueMaskType,
+	const FDeepShadowOpaqueParams& Params,
+	FRDGTextureRef& OutShadowMask)
+{
+	check(OutShadowMask);
+
+	FDeepShadowMaskPS::FParameters* Parameters = GraphBuilder.AllocParameters<FDeepShadowMaskPS::FParameters>();
+	Parameters->ViewUniformBuffer = View.ViewUniformBuffer;
+	Parameters->SceneTextures = SceneTextures;
+	Parameters->DeepShadow_WorldToLightTransform = Params.DeepShadow_WorldToLightTransform;
+	Parameters->DeepShadow_FrontDepthTexture = Params.DeepShadow_FrontDepthTexture;
+	Parameters->CategorizationTexture = Params.CategorizationTexture;
+	Parameters->LinearSampler = TStaticSamplerState<SF_Bilinear, AM_Clamp, AM_Clamp, AM_Clamp>::GetRHI();
+	Parameters->ShadowSampler = TStaticSamplerState<SF_Point, AM_Clamp, AM_Clamp, AM_Clamp>::GetRHI();
+	Parameters->DeepShadow_IsWholeSceneLight = Params.DeepShadow_bIsWholeSceneLight ? 1 : 0;
+	Parameters->RenderTargets[0] = FRenderTargetBinding(OutShadowMask, ERenderTargetLoadAction::ELoad);
+	Parameters->DeepShadow_SlotOffset = FIntPoint(Params.DeepShadow_AtlasRect.Min.X, Params.DeepShadow_AtlasRect.Min.Y);
+	Parameters->DeepShadow_SlotResolution = FIntPoint(Params.DeepShadow_AtlasRect.Max.X - Params.DeepShadow_AtlasRect.Min.X, Params.DeepShadow_AtlasRect.Max.Y - Params.DeepShadow_AtlasRect.Min.Y);
+
+	Parameters->Voxel_LightPosition = Params.Voxel_LightPosition;
+	Parameters->Voxel_LightDirection = Params.Voxel_LightDirection;
+	Parameters->Voxel_DensityScale = Params.Voxel_DensityScale;
+	Parameters->Voxel_MinAABB = Params.Voxel_MinAABB;
+	Parameters->Voxel_Resolution = Params.Voxel_Resolution;
+	Parameters->Voxel_MaxAABB = Params.Voxel_MaxAABB;
+
+	Parameters->Voxel_DensityTexture = Params.Voxel_DensityTexture;
+
+	FDeepShadowMaskPS::FPermutationDomain PermutationVector;
+	PermutationVector.Set<FDeepShadowMaskPS::FOpaqueMaskType>(HairOpaqueMaskType);
+	PermutationVector.Set<FDeepShadowMaskPS::FKernelType>(FMath::Clamp(GStrandHairShadowMaskKernelType, 0, 3));
+
+	const FIntPoint OutputResolution = SceneTextures.SceneDepthBuffer->Desc.Extent;
+	TShaderMapRef<FPostProcessVS> VertexShader(View.ShaderMap);
+	TShaderMapRef<FDeepShadowMaskPS> PixelShader(View.ShaderMap, PermutationVector);
+	const TShaderMap<FGlobalShaderType>* GlobalShaderMap = View.ShaderMap;
+	const FIntRect Viewport = View.ViewRect;
+	const FViewInfo* CapturedView = &View;
+
+	{
+		ClearUnusedGraphResources(*PixelShader, Parameters);
+		FIntPoint Resolution = OutShadowMask->Desc.Extent;
+
+		GraphBuilder.AddPass(
+			RDG_EVENT_NAME("HairStrandsShadowMask"),
+			Parameters,
+			ERDGPassFlags::Raster,
+			[Parameters, VertexShader, PixelShader, Viewport, Resolution, CapturedView](FRHICommandList& RHICmdList)
+		{
+			FGraphicsPipelineStateInitializer GraphicsPSOInit;
+			RHICmdList.ApplyCachedRenderTargets(GraphicsPSOInit);
+			GraphicsPSOInit.BlendState = TStaticBlendState<CW_RGBA, BO_Min, BF_One, BF_One, BO_Min, BF_One, BF_One>::GetRHI(); // Min Operator
+			GraphicsPSOInit.RasterizerState = TStaticRasterizerState<>::GetRHI();
+			GraphicsPSOInit.DepthStencilState = TStaticDepthStencilState<false, CF_Always>::GetRHI();
+
+			GraphicsPSOInit.BoundShaderState.VertexDeclarationRHI = GFilterVertexDeclaration.VertexDeclarationRHI;
+			GraphicsPSOInit.BoundShaderState.VertexShaderRHI = GETSAFERHISHADER_VERTEX(*VertexShader);
+			GraphicsPSOInit.BoundShaderState.PixelShaderRHI = GETSAFERHISHADER_PIXEL(*PixelShader);
+			GraphicsPSOInit.PrimitiveType = PT_TriangleList;
+			SetGraphicsPipelineState(RHICmdList, GraphicsPSOInit);
+
+			VertexShader->SetParameters(RHICmdList, CapturedView->ViewUniformBuffer);
+			RHICmdList.SetViewport(Viewport.Min.X, Viewport.Min.Y, 0.0f, Viewport.Max.X, Viewport.Max.Y, 1.0f);
+			SetShaderParameters(RHICmdList, *PixelShader, PixelShader->GetPixelShader(), *Parameters);
+
+			DrawRectangle(
+				RHICmdList,
+				0, 0,
+				Viewport.Width(), Viewport.Height(),
+				Viewport.Min.X, Viewport.Min.Y,
+				Viewport.Width(), Viewport.Height(),
+				Viewport.Size(),
+				Resolution,
+				*VertexShader,
+				EDRF_UseTriangleOptimization);
+		});
+	}
+}
+
+///////////////////////////////////////////////////////////////////////////////////////////////////
+
+static FHairStrandsTransmittanceMaskData RenderHairStrandsTransmittanceMask(
+	FRHICommandListImmediate& RHICmdList,
+	const FViewInfo& View,
+	const FLightSceneInfo* LightSceneInfo,
+	const FHairStrandsDeepShadowDatas& DeepShadowDatas,
+	const FHairStrandsClusterDatas& DeepClusterDatas,
+	const FHairStrandsVisibilityData& VisibilityData)
+{
+	if (DeepClusterDatas.Datas.Num() == 0)
+		return FHairStrandsTransmittanceMaskData();
+
+	if (!HasDeepShadowData(LightSceneInfo, DeepShadowDatas) && !IsHairStrandsVoxelizationEnable())
+		return FHairStrandsTransmittanceMaskData();
+
+	DECLARE_GPU_STAT(HairStrandsTransmittanceMask);
+	SCOPED_DRAW_EVENT(RHICmdList, HairStrandsTransmittanceMask);
+	SCOPED_GPU_STAT(RHICmdList, HairStrandsTransmittanceMask);
+
+	const FHairLUT InHairLUT = GetHairLUT(RHICmdList, View);
+
+	FSceneRenderTargets& SceneTargets = FSceneRenderTargets::Get(RHICmdList);
+	FRDGBuilder GraphBuilder(RHICmdList);
+
+	// Note: GbufferB.a store the shading model on the 4 lower bits (MATERIAL_SHADINGMODEL_HAIR)
+	FSceneTextureParameters SceneTextures;
+	SetupSceneTextureParameters(GraphBuilder, &SceneTextures);
+
+	FRDGTextureRef HairLUTTexture = GraphBuilder.RegisterExternalTexture(InHairLUT.Textures[HairLUTType_DualScattering], TEXT("HairLUTTexture"));
+	FRDGBufferRef NodeIndirectArgBuffer = GraphBuilder.RegisterExternalBuffer(VisibilityData.NodeIndirectArg, TEXT("HairNodeIndirectArgBuffer"));
+
+	FDeepShadowTransmittanceParams Params;
+	Params.HairVisibilityNodeData = GraphBuilder.RegisterExternalBuffer(VisibilityData.NodeData, TEXT("HairVisibilityNodeData"));
+	Params.HairVisibilityNodeCoord = GraphBuilder.RegisterExternalBuffer(VisibilityData.NodeCoord, TEXT("HairVisibilityNodeCoord"));
+	Params.DensityScale = GetDeepShadowDensityScale();
+	memset(Params.DeepShadow_AtlasSlotOffsets, 0, sizeof(Params.DeepShadow_AtlasSlotOffsets));
+	memset(Params.DeepShadow_WorldToLightTransforms, 0, sizeof(Params.DeepShadow_WorldToLightTransforms));
+
+	FRDGBufferRef OutShadowMask = nullptr;
+	bool bHasFoundLight = false;
+	if (!IsHairStrandsForVoxelTransmittanceAndShadowEnable())
+	{
+		for (const FHairStrandsDeepShadowData& DeepShadowData : DeepShadowDatas.Datas)
+		{
+			if (DeepShadowData.LightId == LightSceneInfo->Id)
+			{
+				bHasFoundLight = true;
+				if (!Params.DeepShadow_FrontDepthTexture)
+				{
+					Params.DeepShadow_FrontDepthTexture = GraphBuilder.RegisterExternalTexture(DeepShadowData.DepthTexture, TEXT("DeepShadow_FrontDepthTexture"));
+					Params.DeepShadow_DomTexture = GraphBuilder.RegisterExternalTexture(DeepShadowData.LayersTexture, TEXT("DeepShadow_DomTexture"));
+				}
+				Params.DeepShadow_Resolution = DeepShadowData.ShadowResolution;
+				Params.LightDirection = DeepShadowData.LightDirection;
+				Params.LightPosition = DeepShadowData.LightPosition;
+				Params.DepthBiasScale = GetDeepShadowDepthBiasScale();
+				Params.DeepShadow_AtlasSlotOffsets[DeepShadowData.ClusterId] = FIntVector4(DeepShadowData.AtlasRect.Min.X, DeepShadowData.AtlasRect.Min.Y, DeepShadowData.AtlasRect.Width(), DeepShadowData.AtlasRect.Height());
+				Params.DeepShadow_WorldToLightTransforms[DeepShadowData.ClusterId] = DeepShadowData.WorldToLightTransform;
+			}
+		}
+
+		if (bHasFoundLight)
+		{
+			check(Params.DeepShadow_FrontDepthTexture);
+			check(Params.DeepShadow_DomTexture);
+			OutShadowMask = AddDeepShadowTransmittanceMaskPass(
+				GraphBuilder,
+				SceneTextures,
+				View,
+				FHairTransmittanceType_DeepShadow,
+				Params,
+				VisibilityData.NodeGroupSize,
+				HairLUTTexture,
+				NodeIndirectArgBuffer);
+		}
+	}
+
+	if (!bHasFoundLight && IsHairStrandsVoxelizationEnable())
+	{
+		Params.Cluster_Resolution = 0;
+		memset(Params.Cluster_MinAABBs, 0, sizeof(FVector4) * FHairStrandsDeepShadowData::MaxClusterCount);
+		memset(Params.Cluster_MaxAABBs, 0, sizeof(FVector4) * FHairStrandsDeepShadowData::MaxClusterCount);
+
+		FRDGTextureRef DefaultDensityTexture = GraphBuilder.RegisterExternalTexture(GSystemTextures.BlackDummy, TEXT("Voxel_DefaultDensityTexture"));
+		for (uint32 TexIt = 0; TexIt < FHairStrandsDeepShadowData::MaxClusterCount; ++TexIt)
+		{
+			Params.Cluster_DensityTextures[TexIt] = DefaultDensityTexture;
+		}
+
+		Params.DensityScale = GetHairStrandsVoxelizationDensityScale();
+		Params.DepthBiasScale = GetHairStrandsVoxelizationDepthBiasScale();
+		Params.LightDirection = LightSceneInfo->Proxy->GetDirection();
+		Params.LightPosition = FVector4(LightSceneInfo->Proxy->GetPosition(), LightSceneInfo->Proxy->GetLightType() == ELightComponentType::LightType_Directional ? 0 : 1);
+		for (const FHairStrandsClusterData& ClusterData : DeepClusterDatas.Datas)
+		{
+			Params.Cluster_Resolution = ClusterData.GetResolution();
+			Params.Cluster_MinAABBs[ClusterData.ClusterId] = ClusterData.GetMinBound();
+			Params.Cluster_MaxAABBs[ClusterData.ClusterId] = ClusterData.GetMaxBound();
+			if (ClusterData.VoxelResources.DensityTexture)
+				Params.Cluster_DensityTextures[ClusterData.ClusterId] = GraphBuilder.RegisterExternalTexture(ClusterData.VoxelResources.DensityTexture);
+		}
+
+		OutShadowMask = AddDeepShadowTransmittanceMaskPass(
+			GraphBuilder,
+			SceneTextures,
+			View,
+			FHairTransmittanceType_Voxel,
+			Params,
+			VisibilityData.NodeGroupSize,
+			HairLUTTexture,
+			NodeIndirectArgBuffer);
+	}
+
+	FHairStrandsTransmittanceMaskData OutTransmittanceMaskData;
+	GraphBuilder.QueueBufferExtraction(OutShadowMask, &OutTransmittanceMaskData.TransmittanceMask, FRDGResourceState::EAccess::Read, FRDGResourceState::EPipeline::Graphics);
+
+	// #RDG_todo/#Hair_todo
+	// Keep an extra ref to keep the buffer alive until the .Execute() function. The issue arive by the fact the indirect buffer is never 
+	// explicitly referenced in the graph. So its reference count is never incremented, and this makes it culled during the graph dependency 
+	// walk.
+	TRefCountPtr<FPooledRDGBuffer> DummyNodeIndirectArg;
+	GraphBuilder.QueueBufferExtraction(NodeIndirectArgBuffer, &DummyNodeIndirectArg, FRDGResourceState::EAccess::Read, FRDGResourceState::EPipeline::Compute);
+
+	GraphBuilder.Execute();
+	OutTransmittanceMaskData.TransmittanceMaskSRV = RHICreateShaderResourceView(OutTransmittanceMaskData.TransmittanceMask->StructuredBuffer);
+
+	return OutTransmittanceMaskData;
+}
+
+FHairStrandsTransmittanceMaskData RenderHairStrandsTransmittanceMask(
+	FRHICommandListImmediate& RHICmdList,
+	const TArray<FViewInfo>& Views,
+	const FLightSceneInfo* LightSceneInfo,
+	const FHairStrandsDatas* HairDatas)
+{
+	FHairStrandsTransmittanceMaskData TransmittanceMaskData;
+	if (HairDatas)
+	{
+		for (int32 ViewIndex = 0; ViewIndex < Views.Num(); ViewIndex++)
+		{
+			const FViewInfo& View = Views[ViewIndex];
+			const FHairStrandsDeepShadowDatas& InDeepShadowDatas = HairDatas->DeepShadowViews.Views[ViewIndex];
+			const FHairStrandsVisibilityData& InHairVisibilityData = HairDatas->HairVisibilityViews.HairDatas[ViewIndex];
+			const FHairStrandsClusterDatas& InDeepClusterDatas = HairDatas->HairClusterPerViews.Views[ViewIndex];
+
+			TransmittanceMaskData = RenderHairStrandsTransmittanceMask(RHICmdList, View, LightSceneInfo, InDeepShadowDatas, InDeepClusterDatas, InHairVisibilityData);
+		}
+	}
+	return TransmittanceMaskData;
+}
+
+///////////////////////////////////////////////////////////////////////////////////////////////////
+
+static void RenderHairStrandsShadowMask(
+	FRHICommandListImmediate& RHICmdList,
+	const FViewInfo& View,
+	const FLightSceneInfo* LightSceneInfo,
+	IPooledRenderTarget* ScreenShadowMaskTexture,
+	const FHairStrandsDeepShadowDatas& DeepShadowDatas,
+	const FHairStrandsVisibilityData& InVisibilityData,
+	const FHairStrandsClusterDatas& InClusterDatas)
+{
+	if (InClusterDatas.Datas.Num() == 0)
+		return;
+
+	if (!HasDeepShadowData(LightSceneInfo, DeepShadowDatas) && !IsHairStrandsVoxelizationEnable())
+		return;
+
+	DECLARE_GPU_STAT(HairStrandsOpaqueMask);
+	SCOPED_DRAW_EVENT(RHICmdList, HairStrandsOpaqueMask);
+	SCOPED_GPU_STAT(RHICmdList, HairStrandsOpaqueMask);
+
+	FSceneRenderTargets& SceneTargets = FSceneRenderTargets::Get(RHICmdList);
+	FRDGBuilder GraphBuilder(RHICmdList);
+
+	// Note: GbufferB.a store the shading model on the 4 lower bits (MATERIAL_SHADINGMODEL_HAIR)
+	FSceneTextureParameters SceneTextures;
+	SetupSceneTextureParameters(GraphBuilder, &SceneTextures);
+
+	FRDGTextureRef Categorization = RegisterExternalTextureWithFallback(GraphBuilder, InVisibilityData.CategorizationTexture, GSystemTextures.BlackDummy);
+	FRDGTextureRef OutShadowMask = GraphBuilder.RegisterExternalTexture(ScreenShadowMaskTexture, TEXT("ScreenShadowMaskTexture"));
+
+	bool bHasDeepShadow = false;
+	if (!IsHairStrandsForVoxelTransmittanceAndShadowEnable())
+	{
+		for (const FHairStrandsDeepShadowData& DomData : DeepShadowDatas.Datas)
+		{
+			if (DomData.LightId != LightSceneInfo->Id)
+				continue;
+
+			bHasDeepShadow = true;
+			FRDGTextureRef DeepShadowDepth = GraphBuilder.RegisterExternalTexture(DomData.DepthTexture, TEXT("DeepShadowDepthTexture"));
+			const bool bIsWholeSceneLight = LightSceneInfo->Proxy->GetLightType() == LightType_Directional;
+
+			FDeepShadowOpaqueParams Params;
+			Params.CategorizationTexture = Categorization;
+			Params.DeepShadow_WorldToLightTransform = DomData.WorldToLightTransform;
+			Params.DeepShadow_AtlasRect = DomData.AtlasRect;
+			Params.DeepShadow_FrontDepthTexture = DeepShadowDepth;
+			Params.DeepShadow_bIsWholeSceneLight = bIsWholeSceneLight;
+
+			AddDeepShadowOpaqueMaskPass(
+				GraphBuilder,
+				SceneTextures,
+				View,
+				FHairOpaqueMaskType_DeepShadow,
+				Params,
+				OutShadowMask);
+		}
+	}
+
+	// Code is disabled for now until we have the full DOM/voxel fallback logic
+	// If there is no deep shadow for this light, fallback on the voxel representation
+	if (!bHasDeepShadow && IsHairStrandsVoxelizationEnable())
+	{
+		for (const FHairStrandsClusterData& ClusterData : InClusterDatas.Datas)
+		{
+			const bool bIsWholeSceneLight = LightSceneInfo->Proxy->GetLightType() == LightType_Directional;
+
+			FDeepShadowOpaqueParams Params;
+			Params.CategorizationTexture = Categorization;
+			Params.Voxel_Resolution = ClusterData.GetResolution();
+			Params.Voxel_MinAABB = ClusterData.GetMinBound();
+			Params.Voxel_MaxAABB = ClusterData.GetMaxBound();
+			Params.Voxel_DensityTexture = GraphBuilder.RegisterExternalTexture(ClusterData.VoxelResources.DensityTexture ? ClusterData.VoxelResources.DensityTexture : GSystemTextures.WhiteDummy, TEXT("Voxel_DensityTexture"));
+			Params.Voxel_DensityScale = GetDeepShadowDensityScale();
+			Params.Voxel_LightDirection = LightSceneInfo->Proxy->GetDirection();
+			Params.Voxel_LightPosition = FVector4(LightSceneInfo->Proxy->GetPosition(), LightSceneInfo->Proxy->GetLightType() == ELightComponentType::LightType_Directional ? 0 : 1);
+
+			AddDeepShadowOpaqueMaskPass(
+				GraphBuilder,
+				SceneTextures,
+				View,
+				FHairOpaqueMaskType_Voxel,
+				Params,
+				OutShadowMask);
+		}
+	}
+
+	TRefCountPtr<IPooledRenderTarget> LocalOutput = GSystemTextures.BlackDummy;
+	GraphBuilder.QueueTextureExtraction(OutShadowMask, &LocalOutput);
+
+	GraphBuilder.Execute();
+}
+
+void RenderHairStrandsShadowMask(
+	FRHICommandListImmediate& RHICmdList,
+	const TArray<FViewInfo>& Views,
+	const FLightSceneInfo* LightSceneInfo,
+	IPooledRenderTarget* ScreenShadowMaskTexture,
+	const FHairStrandsDatas* HairDatas)
+{
+	const FHairStrandsDeepShadowViews& DeepShadowViews = HairDatas->DeepShadowViews;
+	const FHairStrandsVisibilityViews& HairVisibilityViews = HairDatas->HairVisibilityViews;
+	const FHairStrandsClusterViews& HairClusterViews = HairDatas->HairClusterPerViews;
+
+	for (int32 ViewIndex = 0; ViewIndex < Views.Num(); ++ViewIndex)
+	{
+		if (ViewIndex < DeepShadowViews.Views.Num() && ViewIndex < HairVisibilityViews.HairDatas.Num() && ViewIndex < HairClusterViews.Views.Num())
+		{
+			const FHairStrandsVisibilityData& HairVisibilityData = HairVisibilityViews.HairDatas[ViewIndex];
+			const FHairStrandsDeepShadowDatas& DeepShadowDatas = DeepShadowViews.Views[ViewIndex];
+			const FHairStrandsClusterDatas& ClusterDatas = HairClusterViews.Views[ViewIndex];
+			RenderHairStrandsShadowMask(RHICmdList, Views[ViewIndex], LightSceneInfo, ScreenShadowMaskTexture, DeepShadowDatas, HairVisibilityData, ClusterDatas);
+		}
+	}
+}
diff --git a/Engine/Source/Runtime/Renderer/Private/HairStrands/HairStrandsTransmittance.h b/Engine/Source/Runtime/Renderer/Private/HairStrands/HairStrandsTransmittance.h
new file mode 100644
--- /dev/null
+++ b/Engine/Source/Runtime/Renderer/Private/HairStrands/HairStrandsTransmittance.h
@@ -0,0 +1,32 @@
+// Copyright 1998-2019 Epic Games, Inc. All Rights Reserved.
+
+/*=============================================================================
+	HairStrandsTransmittance.h: Hair strands transmittance evaluation.
+=============================================================================*/
+
+#pragma once
+
+#include "CoreMinimal.h"
+#include "RendererInterface.h"
+#include "RenderGraphResources.h"
+
+struct FHairStrandsTransmittanceMaskData
+{
+	TRefCountPtr<FPooledRDGBuffer>	TransmittanceMask;
+	FShaderResourceViewRHIRef		TransmittanceMaskSRV;
+};
+
+/// Write opaque hair shadow onto screen shadow mask to have fine hair details cast onto opaque geometries
+void RenderHairStrandsShadowMask(
+	FRHICommandListImmediate& RHICmdList,
+	const TArray<FViewInfo>& Views,
+	const class FLightSceneInfo* LightSceneInfo,
+	IPooledRenderTarget* ScreenShadowMaskTexture,
+	const struct FHairStrandsDatas* Hairdatas);
+
+/// Write hair transmittance onto screen shadow mask
+FHairStrandsTransmittanceMaskData RenderHairStrandsTransmittanceMask(
+	FRHICommandListImmediate& RHICmdList,
+	const TArray<FViewInfo>& Views,
+	const class FLightSceneInfo* LightSceneInfo,
+	const struct FHairStrandsDatas* Hairdatas);
\ No newline at end of file
diff --git a/Engine/Source/Runtime/Renderer/Private/HairStrands/HairStrandsUtils.cpp b/Engine/Source/Runtime/Renderer/Private/HairStrands/HairStrandsUtils.cpp
new file mode 100644
--- /dev/null
+++ b/Engine/Source/Runtime/Renderer/Private/HairStrands/HairStrandsUtils.cpp
@@ -0,0 +1,198 @@
+// Copyright 1998-2019 Epic Games, Inc. All Rights Reserved.
+
+#include "HairStrandsUtils.h"
+#include "ScenePrivate.h"
+#include "HairStrandsCluster.h"
+
+static float GHairR = 1;
+static float GHairTT = 1;
+static float GHairTRT = 1;
+static float GHairGlobalScattering = 1;
+static float GHairLocalScattering = 1;
+static FAutoConsoleVariableRef CVarHairR(TEXT("r.HairStrands.Components.R"), GHairR, TEXT("Enable/disable hair BSDF component R"));
+static FAutoConsoleVariableRef CVarHairTT(TEXT("r.HairStrands.Components.TT"), GHairTT, TEXT("Enable/disable hair BSDF component TT"));
+static FAutoConsoleVariableRef CVarHairTRT(TEXT("r.HairStrands.Components.TRT"), GHairTRT, TEXT("Enable/disable hair BSDF component TRT"));
+static FAutoConsoleVariableRef CVarHairGlobalScattering(TEXT("r.HairStrands.Components.GlobalScattering"), GHairGlobalScattering, TEXT("Enable/disable hair BSDF component global scattering"));
+static FAutoConsoleVariableRef CVarHairLocalScattering(TEXT("r.HairStrands.Components.LocalScattering"), GHairLocalScattering, TEXT("Enable/disable hair BSDF component local scattering"));
+
+static float GStrandHairRasterizationScale = 0.5f; // For no AA without TAA, a good value is: 1.325f (Empirical)
+static FAutoConsoleVariableRef CVarStrandHairRasterizationScale(TEXT("r.HairStrands.RasterizationScale"), GStrandHairRasterizationScale, TEXT("Rasterization scale to snap strand to pixel"));
+
+static float GStrandHairVelocityRasterizationScale = 1.5f; // Tuned based on heavy motion example (e.g., head shaking)
+static FAutoConsoleVariableRef CVarStrandHairMaxRasterizationScale(TEXT("r.HairStrands.VelocityRasterizationScale"), GStrandHairVelocityRasterizationScale, TEXT("Rasterization scale to snap strand to pixel under high velocity"));
+
+static float GStrandHairShadowRasterizationScale = 1.0f;
+static FAutoConsoleVariableRef CVarStrandHairShadowRasterizationScale(TEXT("r.HairStrands.ShadowRasterizationScale"), GStrandHairShadowRasterizationScale, TEXT("Rasterization scale to snap strand to pixel in shadow view"));
+
+static float GDeepShadowAABBScale = 1.0f;
+static FAutoConsoleVariableRef CVarDeepShadowAABBScale(TEXT("r.HairStrands.DeepShadow.AABBScale"), GDeepShadowAABBScale, TEXT("Scaling value for loosing/tighting deep shadow bounding volume"));
+
+static int32 GHairVisibilityRectOptimEnable = 1;
+static FAutoConsoleVariableRef CVarHairVisibilityRectOptimEnable(TEXT("r.HairStrands.RectLightingOptim"), GHairVisibilityRectOptimEnable, TEXT("Hair Visibility use projected view rect to light only relevant pixels"));
+
+float SampleCountToSubPixelSize(uint32 SamplePerPixelCount)
+{
+	float Scale = 1;
+	switch (SamplePerPixelCount)
+	{
+	case 1: Scale = 1.f; break;
+	case 4: Scale = 8.f / 16.f; break;
+	case 8: Scale = 4.f / 16.f; break;
+	}
+	return Scale;
+}
+FVector4 GetHairComponents()
+{
+	return FVector4(
+		GHairR > 0 ? 1 : 0,
+		GHairTT > 0 ? 1 : 0,
+		GHairTRT > 0 ? 1 : 0,
+		(GHairGlobalScattering > 0 ? 1 : 0) + (GHairLocalScattering > 0 ? 10 : 0));
+}
+
+FMinHairRadiusAtDepth1 ComputeMinStrandRadiusAtDepth1(
+	const FIntPoint& Resolution,
+	const float FOV,
+	const uint32 SampleCount,
+	const float OverrideStrandHairRasterizationScale)
+{
+	auto InternalMinRadiusAtDepth1 = [Resolution, FOV, SampleCount](float RasterizationScale)
+	{
+		const float DiameterToRadius = 0.5f;
+		const float SubPixelScale = SampleCountToSubPixelSize(SampleCount);
+		const float vFOV = FMath::DegreesToRadians(FOV);
+		const float StrandDiameterAtDepth1 = FMath::Tan(vFOV * 0.5f) / (0.5f * Resolution.Y) * SubPixelScale;
+		return DiameterToRadius * RasterizationScale * StrandDiameterAtDepth1;
+	};
+
+	FMinHairRadiusAtDepth1 Out;
+
+	// Scales strand to covers a bit more than a pixel and insure at least one sample point is hit
+	const float PrimaryRasterizationScale = OverrideStrandHairRasterizationScale > 0 ? OverrideStrandHairRasterizationScale : GStrandHairRasterizationScale;
+	const float VelocityRasterizationScale = OverrideStrandHairRasterizationScale > 0 ? OverrideStrandHairRasterizationScale : GStrandHairVelocityRasterizationScale;
+	Out.Primary = InternalMinRadiusAtDepth1(PrimaryRasterizationScale);
+	Out.Velocity = InternalMinRadiusAtDepth1(VelocityRasterizationScale);
+
+	return Out;
+}
+
+void ComputeWorldToLightClip(
+	FMatrix& WorldToClipTransform,
+	FMinHairRadiusAtDepth1& MinStrandRadiusAtDepth1,
+	const FBoxSphereBounds& PrimitivesBounds,
+	const FLightSceneProxy& LightProxy,
+	const ELightComponentType LightType,
+	const FIntPoint& ShadowResolution)
+{
+	const FSphere SphereBound = PrimitivesBounds.GetSphere();
+	const float SphereRadius = SphereBound.W * GDeepShadowAABBScale;
+	const FVector LightPosition = LightProxy.GetPosition();
+	const float MinZ = FMath::Max(0.1f, FVector::Distance(LightPosition, SphereBound.Center)) - SphereBound.W;
+	const float MaxZ = FMath::Max(0.2f, FVector::Distance(LightPosition, SphereBound.Center)) + SphereBound.W;
+
+	const float StrandHairRasterizationScale = GStrandHairShadowRasterizationScale ? GStrandHairShadowRasterizationScale : GStrandHairRasterizationScale;
+
+	MinStrandRadiusAtDepth1 = FMinHairRadiusAtDepth1();
+	WorldToClipTransform = FMatrix::Identity;
+	if (LightType == LightType_Directional)
+	{
+		const FVector LightDirection = LightProxy.GetDirection();
+		FReversedZOrthoMatrix OrthoMatrix(SphereRadius, SphereRadius, 1.f / (2 * SphereRadius), 0);
+		FLookAtMatrix LookAt(SphereBound.Center - LightDirection * SphereRadius, SphereBound.Center, FVector(0, 0, 1));
+		WorldToClipTransform = LookAt * OrthoMatrix;
+		MinStrandRadiusAtDepth1.Primary = StrandHairRasterizationScale * SphereRadius / FMath::Min(ShadowResolution.X, ShadowResolution.Y);
+		MinStrandRadiusAtDepth1.Velocity = MinStrandRadiusAtDepth1.Primary;
+	}
+	else if (LightType == LightType_Spot || LightType == LightType_Point)
+	{
+		const FVector LightDirection = LightPosition - PrimitivesBounds.GetSphere().Center;
+		const float SphereDistance = FVector::Distance(LightPosition, SphereBound.Center);
+		const float HalfFov = asin(SphereRadius / SphereDistance);
+
+		FReversedZPerspectiveMatrix ProjMatrix(HalfFov, 1, 1, MinZ, MaxZ);
+		FLookAtMatrix WorldToLight(LightPosition, SphereBound.Center, FVector(0, 0, 1));
+		WorldToClipTransform = WorldToLight * ProjMatrix;
+		MinStrandRadiusAtDepth1 = ComputeMinStrandRadiusAtDepth1(ShadowResolution, 2 * HalfFov, 1, StrandHairRasterizationScale);
+	}
+	else if (LightType == LightType_Rect)
+	{
+		const FVector LightDirection = LightPosition - PrimitivesBounds.GetSphere().Center;
+		const float SphereDistance = FVector::Distance(LightPosition, SphereBound.Center);
+		const float HalfFov = asin(SphereRadius / SphereDistance);
+
+		FReversedZPerspectiveMatrix ProjMatrix(HalfFov, 1, 1, MinZ, MaxZ);
+		FLookAtMatrix WorldToLight(LightPosition, SphereBound.Center, FVector(0, 0, 1));
+		WorldToClipTransform = WorldToLight * ProjMatrix;
+		MinStrandRadiusAtDepth1 = ComputeMinStrandRadiusAtDepth1(ShadowResolution, 2 * HalfFov, 1, StrandHairRasterizationScale);
+	}
+}
+
+FIntRect ComputeProjectedScreenRect(const FBox& B, const FViewInfo& View)
+{
+	FVector2D MinP( FLT_MAX,  FLT_MAX);
+	FVector2D MaxP(-FLT_MAX, -FLT_MAX);
+	FVector Vertices[8] =
+	{
+		FVector(B.Min),
+		FVector(B.Min.X, B.Min.Y, B.Max.Z),
+		FVector(B.Min.X, B.Max.Y, B.Min.Z),
+		FVector(B.Max.X, B.Min.Y, B.Min.Z),
+		FVector(B.Max.X, B.Max.Y, B.Min.Z),
+		FVector(B.Max.X, B.Min.Y, B.Max.Z),
+		FVector(B.Min.X, B.Max.Y, B.Max.Z),
+		FVector(B.Max)
+	};
+
+	for (uint32 i = 0; i < 8; ++i)
+	{
+		FVector2D P;
+		if (View.WorldToPixel(Vertices[i], P))
+		{
+			MinP.X = FMath::Min(MinP.X, P.X);
+			MinP.Y = FMath::Min(MinP.Y, P.Y);
+			MaxP.X = FMath::Max(MaxP.X, P.X);
+			MaxP.Y = FMath::Max(MaxP.Y, P.Y);
+		}
+	}
+
+	FIntRect OutRect;
+	OutRect.Min = FIntPoint(FMath::FloorToInt(MinP.X), FMath::FloorToInt(MinP.Y));
+	OutRect.Max = FIntPoint(FMath::CeilToInt(MaxP.X),  FMath::CeilToInt(MaxP.Y));
+
+	// Clamp to screen rect
+	OutRect.Min.X = FMath::Max(View.ViewRect.Min.X, OutRect.Min.X);
+	OutRect.Min.Y = FMath::Max(View.ViewRect.Min.Y, OutRect.Min.Y);
+	OutRect.Max.X = FMath::Min(View.ViewRect.Max.X, OutRect.Max.X);
+	OutRect.Max.Y = FMath::Min(View.ViewRect.Max.Y, OutRect.Max.Y);
+
+	return OutRect;
+}
+
+
+FIntRect ComputeVisibleHairStrandsClustersRect(const FIntRect& ViewRect, const FHairStrandsClusterDatas& ClusterDatas)
+{
+	FIntRect TotalRect(INT_MAX, INT_MAX, -INT_MAX, -INT_MAX);
+	if (IsHairStrandsViewRectOptimEnable())
+	{
+		for (const FHairStrandsClusterData& ClusterData : ClusterDatas.Datas)
+		{
+			TotalRect.Union(ClusterData.ScreenRect);
+		}
+	}
+	else
+	{
+		TotalRect = ViewRect;
+	}
+
+	return TotalRect;
+}
+
+bool IsHairStrandsViewRectOptimEnable()
+{
+	return GHairVisibilityRectOptimEnable > 0;
+}
+
+bool IsHairStrandsSupported(const EShaderPlatform Platform)
+{
+	return IsD3DPlatform(Platform, false) && IsPCPlatform(Platform) && GetMaxSupportedFeatureLevel(Platform) == ERHIFeatureLevel::SM5;
+}
\ No newline at end of file
diff --git a/Engine/Source/Runtime/Renderer/Private/HairStrands/HairStrandsUtils.h b/Engine/Source/Runtime/Renderer/Private/HairStrands/HairStrandsUtils.h
new file mode 100644
--- /dev/null
+++ b/Engine/Source/Runtime/Renderer/Private/HairStrands/HairStrandsUtils.h
@@ -0,0 +1,41 @@
+// Copyright 1998-2019 Epic Games, Inc. All Rights Reserved.
+
+/*=============================================================================
+	HairStrandsUtils.h: Hair strands utils.
+=============================================================================*/
+
+#pragma once
+
+#include "CoreMinimal.h"
+#include "RendererInterface.h"
+#include "HairStrandsInterface.h"
+#include "SceneTypes.h"
+
+struct FHairStrandsLightDesc
+{
+	FMatrix WorldToLightClipTransform;
+	FVector LightDirection;
+	FVector LightPosition;
+	FMinHairRadiusAtDepth1 MinStrandRadiusAtDepth1;
+	bool bIsOrtho;
+};
+
+FIntRect ComputeProjectedScreenRect(const FBox& B, const FViewInfo& View);
+
+void ComputeWorldToLightClip(
+	FMatrix& WorldToClipTransform,
+	FMinHairRadiusAtDepth1& MinStrandRadiusAtDepth1,
+	const FBoxSphereBounds& PrimitivesBounds,
+	const class FLightSceneProxy& LightProxy,
+	const ELightComponentType LightType,
+	const FIntPoint& ShadowResolution);
+
+FVector4 GetHairComponents();
+
+float SampleCountToSubPixelSize(uint32 SamplePerPixelCount);
+
+FIntRect ComputeVisibleHairStrandsClustersRect(const FIntRect& ViewRect, const struct FHairStrandsClusterDatas& ClusterDatas);
+
+bool IsHairStrandsViewRectOptimEnable();
+
+RENDERER_API bool IsHairStrandsSupported(const EShaderPlatform Platform);
\ No newline at end of file
