Copying //UE4/Dev-Physics to Dev-Main (//UE4/Dev-Main) [at] 9098278

diff --git a/Engine/Source/Runtime/Renderer/Private/HairStrands/HairStrandsVisibility.cpp b/Engine/Source/Runtime/Renderer/Private/HairStrands/HairStrandsVisibility.cpp
new file mode 100644
--- /dev/null
+++ b/Engine/Source/Runtime/Renderer/Private/HairStrands/HairStrandsVisibility.cpp
@@ -0,0 +1,1189 @@
+// Copyright 1998-2019 Epic Games, Inc. All Rights Reserved.
+
+#include "HairStrandsVisibility.h"
+#include "HairStrandsCluster.h"
+#include "HairStrandsUtils.h"
+#include "HairStrandsInterface.h"
+
+#include "Shader.h"
+#include "GlobalShader.h"
+#include "ShaderParameters.h"
+#include "ShaderParameterStruct.h"
+#include "SceneTextureParameters.h"
+#include "RenderGraphUtils.h"
+#include "PostProcessing.h"
+#include "MeshPassProcessor.inl"
+
+DECLARE_GPU_STAT(HairStrandsVisibility);
+
+/////////////////////////////////////////////////////////////////////////////////////////
+
+BEGIN_GLOBAL_SHADER_PARAMETER_STRUCT(FHairVisibilityPassUniformParameters, )
+SHADER_PARAMETER(float, MinStrandRadius_Primary)
+SHADER_PARAMETER(float, MinStrandRadius_Velocity)
+SHADER_PARAMETER(float, HairStrandsVelocityScale)
+SHADER_PARAMETER_TEXTURE(Texture2D<float>, MainDepthTexture)
+END_GLOBAL_SHADER_PARAMETER_STRUCT()
+IMPLEMENT_GLOBAL_SHADER_PARAMETER_STRUCT(FHairVisibilityPassUniformParameters, "HairVisibilityPass");
+
+/////////////////////////////////////////////////////////////////////////////////////////
+
+static int32 GHairVisibilitySampleCount = 8;
+static FAutoConsoleVariableRef CVarHairVisibilitySampleCount(TEXT("r.HairStrands.VisibilitySampleCount"), GHairVisibilitySampleCount, TEXT("Hair strands visibility sample count"));
+
+static int32 GHairClearVisibilityBuffer = 0;
+static FAutoConsoleVariableRef CVarHairClearVisibilityBuffer(TEXT("r.HairStrands.VisibilityClear"), GHairClearVisibilityBuffer, TEXT("Clear hair strands visibility buffer"));
+
+static int32 GHairVelocityType = 1; // default is 
+static FAutoConsoleVariableRef CVarHairVelocityType(TEXT("r.HairStrands.VelocityType"), GHairVelocityType, TEXT("Type of velocity filtering (0:avg, 1:closest, 2:max). Default is 1."));
+
+static int32 GHairVelocityMagnitudeScale = 100; // Tuned by eye, based on heavy motion (strong head shack)
+static FAutoConsoleVariableRef CVarHairVelocityMagnitudeScale(TEXT("r.HairStrands.VelocityMagnitudeScale"), GHairVelocityMagnitudeScale, TEXT("Velocity magnitude (in pixel) at which a hair will reach its pic velocity-rasterization-scale under motion to reduce aliasing. Default is 100."));
+
+/////////////////////////////////////////////////////////////////////////////////////////
+
+enum EHairVisibilityRenderMode
+{
+	HairVisibilityRenderMode_MSAA,
+	HairVisibilityRenderModeCount
+};
+
+enum EHairVisibilityVendor
+{
+	HairVisibilityVendor_AMD,
+	HairVisibilityVendor_NVIDIA,
+	HairVisibilityVendor_INTEL,
+	HairVisibilityVendorCount
+};
+
+EHairVisibilityRenderMode GetHairVisibilityRenderMode()
+{
+	return HairVisibilityRenderMode_MSAA;
+}
+
+uint32 GetHairVisibilitySampleCount()
+{
+	return GetHairVisibilityRenderMode() == HairVisibilityRenderMode_MSAA ? FMath::Clamp(GHairVisibilitySampleCount,1,16) : 1;
+}
+
+static bool IsCompatibleWithHairVisibility(const FMeshMaterialShaderPermutationParameters& Parameters)
+{
+	return IsCompatibleWithHairStrands(Parameters.Material, Parameters.Platform);
+}
+
+static EHairVisibilityVendor GetVendor()
+{
+	return IsRHIDeviceAMD() ? HairVisibilityVendor_AMD : (IsRHIDeviceNVIDIA() ? HairVisibilityVendor_NVIDIA : HairVisibilityVendor_INTEL);
+}
+
+static uint32 GetVendorOptimalGroupSize1D()
+{
+	switch (GetVendor())
+	{
+	case HairVisibilityVendor_AMD:		return 64;
+	case HairVisibilityVendor_NVIDIA:	return 32;
+	case HairVisibilityVendor_INTEL:	return 64;
+	default:							return 64;
+	}
+}
+
+static FIntPoint GetVendorOptimalGroupSize2D()
+{
+	switch (GetVendor())
+	{
+	case HairVisibilityVendor_AMD:		return FIntPoint(8, 8);
+	case HairVisibilityVendor_NVIDIA:	return FIntPoint(8, 4);
+	case HairVisibilityVendor_INTEL:	return FIntPoint(8, 8);
+	default:							return FIntPoint(8, 8);
+	}
+}
+
+class FHairVisibilityVS : public FMeshMaterialShader
+{
+	DECLARE_SHADER_TYPE(FHairVisibilityVS,MeshMaterial);
+
+protected:
+
+	FHairVisibilityVS(const ShaderMetaType::CompiledShaderInitializerType& Initializer)
+		: FMeshMaterialShader(Initializer)
+	{
+		ERHIFeatureLevel::Type FeatureLevel = GetMaxSupportedFeatureLevel((EShaderPlatform)Initializer.Target.Platform);
+		check(FSceneInterface::GetShadingPath(FeatureLevel) != EShadingPath::Mobile);
+		// deferred
+		PassUniformBuffer.Bind(Initializer.ParameterMap, FHairVisibilityPassUniformParameters::StaticStructMetadata.GetShaderVariableName());
+	}
+
+	FHairVisibilityVS() {}
+
+	static bool ShouldCompilePermutation(const FMeshMaterialShaderPermutationParameters& Parameters)
+	{
+		return IsCompatibleWithHairVisibility(Parameters);
+	}
+};
+
+class FHairVisibilityHS : public FBaseHS
+{
+	DECLARE_SHADER_TYPE(FHairVisibilityHS,MeshMaterial);
+
+protected:
+
+	FHairVisibilityHS(const ShaderMetaType::CompiledShaderInitializerType& Initializer)
+		: FBaseHS(Initializer)
+	{
+		PassUniformBuffer.Bind(Initializer.ParameterMap, FHairVisibilityPassUniformParameters::StaticStructMetadata.GetShaderVariableName());
+	}
+
+	FHairVisibilityHS() {}
+
+	static bool ShouldCompilePermutation(const FMeshMaterialShaderPermutationParameters& Parameters)
+	{
+		return FBaseHS::ShouldCompilePermutation(Parameters) && IsCompatibleWithHairVisibility(Parameters);
+	}
+};
+
+class FHairVisibilityDS : public FBaseDS
+{
+	DECLARE_SHADER_TYPE(FHairVisibilityDS,MeshMaterial);
+
+protected:
+
+	FHairVisibilityDS(const ShaderMetaType::CompiledShaderInitializerType& Initializer)
+		: FBaseDS(Initializer)
+	{
+		PassUniformBuffer.Bind(Initializer.ParameterMap, FHairVisibilityPassUniformParameters::StaticStructMetadata.GetShaderVariableName());
+	}
+
+	FHairVisibilityDS() {}
+
+	static bool ShouldCompilePermutation(const FMeshMaterialShaderPermutationParameters& Parameters)
+	{
+		return FBaseDS::ShouldCompilePermutation(Parameters) && IsCompatibleWithHairVisibility(Parameters);
+	}
+};
+
+IMPLEMENT_MATERIAL_SHADER_TYPE(, FHairVisibilityVS, TEXT("/Engine/Private/HairStrands/HairStrandsVisibilityVS.usf"), TEXT("Main"), SF_Vertex);
+IMPLEMENT_MATERIAL_SHADER_TYPE(, FHairVisibilityHS, TEXT("/Engine/Private/HairStrands/HairStrandsVisibilityVS.usf"), TEXT("MainHull"), SF_Hull);
+IMPLEMENT_MATERIAL_SHADER_TYPE(, FHairVisibilityDS, TEXT("/Engine/Private/HairStrands/HairStrandsVisibilityVS.usf"), TEXT("MainDomain"), SF_Domain);
+
+/////////////////////////////////////////////////////////////////////////////////////////
+
+
+class FHairVisibilityShaderElementData : public FMeshMaterialShaderElementData
+{
+public:
+	FHairVisibilityShaderElementData(uint32 InHairClusterId) : HairClusterId(InHairClusterId) { }
+	uint32 HairClusterId;
+};
+
+
+template<EHairVisibilityRenderMode RenderMode>
+class FHairVisibilityPS : public FMeshMaterialShader
+{
+	DECLARE_SHADER_TYPE(FHairVisibilityPS, MeshMaterial);
+
+public:
+
+	FHairVisibilityPS(const ShaderMetaType::CompiledShaderInitializerType& Initializer)
+		:	FMeshMaterialShader(Initializer)
+	{
+		ERHIFeatureLevel::Type FeatureLevel = GetMaxSupportedFeatureLevel((EShaderPlatform)Initializer.Target.Platform);
+		check(FSceneInterface::GetShadingPath(FeatureLevel) != EShadingPath::Mobile);
+		PassUniformBuffer.Bind(Initializer.ParameterMap, FHairVisibilityPassUniformParameters::StaticStructMetadata.GetShaderVariableName());
+		HairVisibilityPass_HairClusterIndex.Bind(Initializer.ParameterMap, TEXT("HairVisibilityPass_HairClusterIndex"));
+	}
+
+	FHairVisibilityPS() {}
+
+	static bool ShouldCompilePermutation(const FMeshMaterialShaderPermutationParameters& Parameters)
+	{
+		return IsCompatibleWithHairStrands(Parameters.Material, Parameters.Platform);
+	}
+
+	static void ModifyCompilationEnvironment(const FMaterialShaderPermutationParameters& Parameters, FShaderCompilerEnvironment& OutEnvironment)	
+	{
+		FMaterialShader::ModifyCompilationEnvironment(Parameters, OutEnvironment);
+		const uint32 RenderModeValue = RenderMode == HairVisibilityRenderMode_MSAA ? 0 : 1;
+		OutEnvironment.SetDefine(TEXT("HAIR_RENDER_MODE"), RenderModeValue);
+	}
+
+	virtual bool Serialize(FArchive& Ar) override
+	{
+		bool bShaderHasOutdatedParameters = FMeshMaterialShader::Serialize(Ar);
+		Ar << HairVisibilityPass_HairClusterIndex;
+		return bShaderHasOutdatedParameters;
+	}
+
+	void GetShaderBindings(
+		const FScene* Scene,
+		ERHIFeatureLevel::Type FeatureLevel,
+		const FPrimitiveSceneProxy* PrimitiveSceneProxy,
+		const FMaterialRenderProxy& MaterialRenderProxy,
+		const FMaterial& Material,
+		const FMeshPassProcessorRenderState& DrawRenderState,
+		const FHairVisibilityShaderElementData& ShaderElementData,
+		FMeshDrawSingleShaderBindings& ShaderBindings) const
+	{
+		FMeshMaterialShader::GetShaderBindings(Scene, FeatureLevel, PrimitiveSceneProxy, MaterialRenderProxy, Material, DrawRenderState, ShaderElementData, ShaderBindings);
+		ShaderBindings.Add(HairVisibilityPass_HairClusterIndex, ShaderElementData.HairClusterId);
+	}
+
+	FShaderParameter HairVisibilityPass_HairClusterIndex;
+};
+IMPLEMENT_MATERIAL_SHADER_TYPE(template<>, FHairVisibilityPS<HairVisibilityRenderMode_MSAA>, TEXT("/Engine/Private/HairStrands/HairStrandsVisibilityPS.usf"), TEXT("MainVisibility"), SF_Pixel);
+
+/////////////////////////////////////////////////////////////////////////////////////////
+
+class FHairVisibilityProcessor : public FMeshPassProcessor
+{
+public:
+	FHairVisibilityProcessor(
+		const FScene* Scene,
+		const FSceneView* InViewIfDynamicMeshCommand,
+		const FMeshPassProcessorRenderState& InPassDrawRenderState,
+		const EHairVisibilityRenderMode InRenderMode,
+		FDynamicPassMeshDrawListContext* InDrawListContext);
+
+	virtual void AddMeshBatch(const FMeshBatch& RESTRICT MeshBatch, uint64 BatchElementMask, const FPrimitiveSceneProxy* RESTRICT PrimitiveSceneProxy, int32 StaticMeshId = -1) override final;
+	void AddMeshBatch(const FMeshBatch& RESTRICT MeshBatch, uint64 BatchElementMask, const FPrimitiveSceneProxy* RESTRICT PrimitiveSceneProxy, int32 StaticMeshId, uint32 HairClusterId);
+
+private:
+	template<EHairVisibilityRenderMode RenderMode>
+	void Process(
+		const FMeshBatch& MeshBatch,
+		uint64 BatchElementMask,
+		const FPrimitiveSceneProxy* RESTRICT PrimitiveSceneProxy,
+		int32 StaticMeshId,
+		const FMaterialRenderProxy& RESTRICT MaterialRenderProxy,
+		const FMaterial& RESTRICT MaterialResource,
+		const uint32 HairClusterId,
+		ERasterizerFillMode MeshFillMode,
+		ERasterizerCullMode MeshCullMode);
+
+	const EHairVisibilityRenderMode RenderMode;
+	FMeshPassProcessorRenderState PassDrawRenderState;
+};
+
+void FHairVisibilityProcessor::AddMeshBatch(const FMeshBatch& RESTRICT MeshBatch, uint64 BatchElementMask, const FPrimitiveSceneProxy* RESTRICT PrimitiveSceneProxy, int32 StaticMeshId)
+{
+	AddMeshBatch(MeshBatch, BatchElementMask, PrimitiveSceneProxy, StaticMeshId, 0);
+}
+
+void FHairVisibilityProcessor::AddMeshBatch(const FMeshBatch& RESTRICT MeshBatch, uint64 BatchElementMask, const FPrimitiveSceneProxy* RESTRICT PrimitiveSceneProxy, int32 StaticMeshId, uint32 HairClusterId)
+{
+	// Determine the mesh's material and blend mode.
+	const FMaterialRenderProxy* FallbackMaterialRenderProxyPtr = nullptr;
+	const FMaterial& Material = MeshBatch.MaterialRenderProxy->GetMaterialWithFallback(FeatureLevel, FallbackMaterialRenderProxyPtr);
+	const bool bIsCompatible = IsCompatibleWithHairStrands(&Material, FeatureLevel);
+
+	if (bIsCompatible
+		&& (!PrimitiveSceneProxy || PrimitiveSceneProxy->ShouldRenderInMainPass())
+		&& ShouldIncludeDomainInMeshPass(Material.GetMaterialDomain()))
+	{
+		const FMaterialRenderProxy& MaterialRenderProxy = FallbackMaterialRenderProxyPtr ? *FallbackMaterialRenderProxyPtr : *MeshBatch.MaterialRenderProxy;
+		const ERasterizerFillMode MeshFillMode = ComputeMeshFillMode(MeshBatch, Material);
+		const ERasterizerCullMode MeshCullMode = ComputeMeshCullMode(MeshBatch, Material);
+		if (RenderMode == HairVisibilityRenderMode_MSAA)
+			Process<HairVisibilityRenderMode_MSAA>(MeshBatch, BatchElementMask, PrimitiveSceneProxy, StaticMeshId, MaterialRenderProxy, Material, HairClusterId, MeshFillMode, MeshCullMode);
+	}
+}
+
+template<EHairVisibilityRenderMode TRenderMode>
+void FHairVisibilityProcessor::Process(
+	const FMeshBatch& MeshBatch, 
+	uint64 BatchElementMask, 
+	const FPrimitiveSceneProxy* RESTRICT PrimitiveSceneProxy,
+	int32 StaticMeshId,
+	const FMaterialRenderProxy& RESTRICT MaterialRenderProxy,
+	const FMaterial& RESTRICT MaterialResource,
+	const uint32 HairClusterId,
+	ERasterizerFillMode MeshFillMode,
+	ERasterizerCullMode MeshCullMode)
+{
+	const FVertexFactory* VertexFactory = MeshBatch.VertexFactory;
+
+	TMeshProcessorShaders<
+		FHairVisibilityVS,
+		FHairVisibilityHS,
+		FHairVisibilityDS,
+		FHairVisibilityPS<TRenderMode>> PassShaders;
+	{
+		const EMaterialTessellationMode MaterialTessellationMode = MaterialResource.GetTessellationMode();
+		FVertexFactoryType* VertexFactoryType = VertexFactory->GetType();
+
+		const bool bNeedsHSDS = RHISupportsTessellation(GShaderPlatformForFeatureLevel[FeatureLevel])
+			&& VertexFactoryType->SupportsTessellationShaders()
+			&& MaterialTessellationMode != MTM_NoTessellation;
+
+		if (bNeedsHSDS)
+		{
+			PassShaders.DomainShader = MaterialResource.GetShader<FHairVisibilityDS>(VertexFactoryType);
+			PassShaders.HullShader = MaterialResource.GetShader<FHairVisibilityHS>(VertexFactoryType);
+		}
+
+		PassShaders.VertexShader = MaterialResource.GetShader<FHairVisibilityVS>(VertexFactoryType);
+		PassShaders.PixelShader = MaterialResource.GetShader<FHairVisibilityPS<TRenderMode>>(VertexFactoryType);
+	}
+
+	FMeshPassProcessorRenderState DrawRenderState(PassDrawRenderState);
+
+	FHairVisibilityShaderElementData ShaderElementData(HairClusterId);
+	ShaderElementData.InitializeMeshMaterialData(ViewIfDynamicMeshCommand, PrimitiveSceneProxy, MeshBatch, StaticMeshId, false);
+
+	BuildMeshDrawCommands(
+		MeshBatch,
+		BatchElementMask,
+		PrimitiveSceneProxy,
+		MaterialRenderProxy,
+		MaterialResource,
+		DrawRenderState,
+		PassShaders,
+		MeshFillMode,
+		MeshCullMode,
+		FMeshDrawCommandSortKey::Default,
+		EMeshPassFeatures::Default,
+		ShaderElementData);
+}
+
+FHairVisibilityProcessor::FHairVisibilityProcessor(
+	const FScene* Scene,
+	const FSceneView* InViewIfDynamicMeshCommand,
+	const FMeshPassProcessorRenderState& InPassDrawRenderState,
+	const EHairVisibilityRenderMode InRenderMode,
+	FDynamicPassMeshDrawListContext* InDrawListContext)
+	: FMeshPassProcessor(Scene, Scene->GetFeatureLevel(), InViewIfDynamicMeshCommand, InDrawListContext)
+	, RenderMode(InRenderMode)
+	, PassDrawRenderState(InPassDrawRenderState)
+{
+}
+
+///////////////////////////////////////////////////////////////////////////////////////////////////
+// Clear uint texture
+class FClearUIntGraphicPS : public FGlobalShader
+{
+	DECLARE_GLOBAL_SHADER(FClearUIntGraphicPS);
+	SHADER_USE_PARAMETER_STRUCT(FClearUIntGraphicPS, FGlobalShader);
+
+	BEGIN_SHADER_PARAMETER_STRUCT(FParameters, )
+		SHADER_PARAMETER(uint32, ClearValue)
+		RENDER_TARGET_BINDING_SLOTS()
+	END_SHADER_PARAMETER_STRUCT()
+
+public:
+	static bool ShouldCompilePermutation(const FGlobalShaderPermutationParameters& Parameters) { return IsHairStrandsSupported(Parameters.Platform); }
+};
+
+IMPLEMENT_GLOBAL_SHADER(FClearUIntGraphicPS, "/Engine/Private/HairStrands/HairStrandsVisibilityClearPS.usf", "ClearPS", SF_Pixel);
+
+// Opaque mask
+static void AddClearGraphicPass(
+	FRDGBuilder& GraphBuilder,
+	FRDGEventName&& PassName,
+	const FViewInfo* View,
+	const uint32 ClearValue,
+	FRDGTextureRef& OutTarget)
+{
+	check(OutTarget);
+
+	FClearUIntGraphicPS::FParameters* Parameters = GraphBuilder.AllocParameters<FClearUIntGraphicPS::FParameters>();
+	Parameters->ClearValue = ClearValue;
+	Parameters->RenderTargets[0] = FRenderTargetBinding(OutTarget, ERenderTargetLoadAction::ENoAction, 0);
+
+	TShaderMapRef<FPostProcessVS> VertexShader(View->ShaderMap);
+	TShaderMapRef<FClearUIntGraphicPS> PixelShader(View->ShaderMap);
+	const FIntRect Viewport = View->ViewRect;
+	const FIntPoint Resolution = OutTarget->Desc.Extent;
+
+	ClearUnusedGraphResources(*PixelShader, Parameters);
+
+	GraphBuilder.AddPass(
+		Forward<FRDGEventName>(PassName),
+		Parameters,
+		ERDGPassFlags::Raster,
+		[Parameters, VertexShader, PixelShader, Viewport, Resolution, View](FRHICommandList& RHICmdList)
+	{
+		FGraphicsPipelineStateInitializer GraphicsPSOInit;
+		RHICmdList.ApplyCachedRenderTargets(GraphicsPSOInit);
+		GraphicsPSOInit.BlendState = TStaticBlendState<CW_RGBA, BO_Add, BF_One, BF_Zero, BO_Add, BF_One, BF_Zero>::GetRHI();
+		GraphicsPSOInit.RasterizerState = TStaticRasterizerState<>::GetRHI();
+		GraphicsPSOInit.DepthStencilState = TStaticDepthStencilState<false, CF_Always>::GetRHI();
+		GraphicsPSOInit.BoundShaderState.VertexDeclarationRHI = GFilterVertexDeclaration.VertexDeclarationRHI;
+		GraphicsPSOInit.BoundShaderState.VertexShaderRHI = GETSAFERHISHADER_VERTEX(*VertexShader);
+		GraphicsPSOInit.BoundShaderState.PixelShaderRHI = GETSAFERHISHADER_PIXEL(*PixelShader);
+		GraphicsPSOInit.PrimitiveType = PT_TriangleList;
+		SetGraphicsPipelineState(RHICmdList, GraphicsPSOInit);
+
+		VertexShader->SetParameters(RHICmdList, View->ViewUniformBuffer);
+		RHICmdList.SetViewport(Viewport.Min.X, Viewport.Min.Y, 0.0f, Viewport.Max.X, Viewport.Max.Y, 1.0f);
+		SetShaderParameters(RHICmdList, *PixelShader, PixelShader->GetPixelShader(), *Parameters);
+
+		DrawRectangle(
+			RHICmdList,
+			0, 0,
+			Viewport.Width(),Viewport.Height(),
+			Viewport.Min.X,  Viewport.Min.Y,
+			Viewport.Width(),Viewport.Height(),
+			Viewport.Size(),
+			Resolution,
+			*VertexShader,
+			EDRF_UseTriangleOptimization);
+	});
+}
+
+BEGIN_SHADER_PARAMETER_STRUCT(FClearUAVTextureParameters, )
+	SHADER_PARAMETER_RDG_TEXTURE_UAV(RWTexture2D, TextureUAV)
+END_SHADER_PARAMETER_STRUCT()
+
+void AddClearUAVPass(
+	FRDGBuilder& GraphBuilder,
+	FRDGEventName&& PassName,
+	FRDGTextureRef Texture,
+	uint32 Value)
+{	
+	FClearUAVTextureParameters* Parameters = GraphBuilder.AllocParameters< FClearUAVTextureParameters >();
+	Parameters->TextureUAV = GraphBuilder.CreateUAV(Texture);
+
+	GraphBuilder.AddPass(
+		Forward<FRDGEventName>(PassName),
+		Parameters,
+		ERDGPassFlags::Compute,
+		[Parameters, Texture, Value](FRHICommandList& RHICmdList)
+	{
+		uint32 ClearValue[4] = { Value, Value, Value, Value };
+		FRHIUnorderedAccessView* GlobalCounterUAV = Parameters->TextureUAV->GetRHI();
+		RHICmdList.ClearTinyUAV(GlobalCounterUAV, ClearValue);
+	});
+}
+
+///////////////////////////////////////////////////////////////////////////////////////////////////
+// Copy dispatch count into an indirect buffer 
+class FCopyIndirectBufferCS : public FGlobalShader
+{
+	DECLARE_GLOBAL_SHADER(FCopyIndirectBufferCS);
+	SHADER_USE_PARAMETER_STRUCT(FCopyIndirectBufferCS, FGlobalShader);
+
+	BEGIN_SHADER_PARAMETER_STRUCT(FParameters, )
+		SHADER_PARAMETER(uint32, ThreadGroupSize)
+		SHADER_PARAMETER_RDG_TEXTURE(Texture2D, CounterTexture)
+		SHADER_PARAMETER_RDG_BUFFER_UAV(RWBuffer, OutArgBuffer)
+	END_SHADER_PARAMETER_STRUCT()
+
+public:
+	static bool ShouldCompilePermutation(const FGlobalShaderPermutationParameters& Parameters) { return IsHairStrandsSupported(Parameters.Platform); }
+};
+
+IMPLEMENT_GLOBAL_SHADER(FCopyIndirectBufferCS, "/Engine/Private/HairStrands/HairStrandsVisibilityCopyIndirectArg.usf", "CopyCS", SF_Compute);
+
+static FRDGBufferRef AddCopyIndirectArgPass(
+	FRDGBuilder& GraphBuilder,
+	const FViewInfo* View,
+	const uint32 ThreadGroupSize,
+	FRDGTextureRef CounterTexture)
+{
+	check(CounterTexture);
+
+	FRDGBufferRef OutBuffer = GraphBuilder.CreateBuffer(FRDGBufferDesc::CreateIndirectDesc<FRHIDispatchIndirectParameters>(), TEXT("HairVisibilityIndirectArgBuffer"));
+
+	FCopyIndirectBufferCS::FParameters* Parameters = GraphBuilder.AllocParameters<FCopyIndirectBufferCS::FParameters>();
+	Parameters->ThreadGroupSize = 32;
+	Parameters->CounterTexture = CounterTexture;
+	Parameters->OutArgBuffer = GraphBuilder.CreateUAV(OutBuffer);
+
+	TShaderMapRef<FCopyIndirectBufferCS> ComputeShader(View->ShaderMap);
+
+	FComputeShaderUtils::AddPass(
+		GraphBuilder,
+		RDG_EVENT_NAME("HairStrandsVisbilityCopyIndirectArgs"),
+		*ComputeShader,
+		Parameters,
+		FIntVector(1,1,1));
+
+	return OutBuffer;
+}
+
+///////////////////////////////////////////////////////////////////////////////////////////////////
+class FHairVisibilityPrimitiveIdCompactionCS : public FGlobalShader
+{
+	DECLARE_GLOBAL_SHADER(FHairVisibilityPrimitiveIdCompactionCS);
+	SHADER_USE_PARAMETER_STRUCT(FHairVisibilityPrimitiveIdCompactionCS, FGlobalShader);
+
+	class FVendor   : SHADER_PERMUTATION_INT("PERMUTATION_VENDOR", HairVisibilityVendorCount);
+	class FVelocity : SHADER_PERMUTATION_INT("PERMUTATION_VELOCITY", 4);
+	using FPermutationDomain = TShaderPermutationDomain<FVendor, FVelocity>;
+
+	BEGIN_SHADER_PARAMETER_STRUCT(FParameters, )
+		SHADER_PARAMETER(FIntPoint, OutputResolution)
+		SHADER_PARAMETER(uint32, MaxNodeCount)
+		SHADER_PARAMETER(uint32, HairVisibilitySampleCount)
+		SHADER_PARAMETER(FIntPoint, ResolutionOffset)
+
+		SHADER_PARAMETER_RDG_TEXTURE(Texture2D, MSAA_DepthTexture)
+		SHADER_PARAMETER_RDG_TEXTURE(Texture2D, MSAA_IDTexture)
+		SHADER_PARAMETER_RDG_TEXTURE(Texture2D, MSAA_MaterialTexture)
+		SHADER_PARAMETER_RDG_TEXTURE(Texture2D, MSAA_AttributeTexture)
+		SHADER_PARAMETER_RDG_TEXTURE(Texture2D, MSAA_VelocityTexture)
+		
+		SHADER_PARAMETER_RDG_TEXTURE_UAV(Texture2D, OutCompactNodeCounter)
+		SHADER_PARAMETER_RDG_TEXTURE_UAV(Texture2D, OutCompactNodeIndex)
+		SHADER_PARAMETER_RDG_TEXTURE_UAV(Texture2D, OutCategorizationTexture)
+		SHADER_PARAMETER_RDG_BUFFER_UAV(StructuredBuffer, OutCompactNodeData)
+		SHADER_PARAMETER_RDG_BUFFER_UAV(StructuredBuffer, OutCompactNodeCoord)
+		SHADER_PARAMETER_RDG_TEXTURE_UAV(Texture2D, OutVelocityTexture)
+
+		SHADER_PARAMETER_STRUCT_REF(FSceneTexturesUniformParameters, SceneTexturesStruct)
+		END_SHADER_PARAMETER_STRUCT()
+
+public:
+	static bool ShouldCompilePermutation(const FGlobalShaderPermutationParameters& Parameters) { return IsHairStrandsSupported(Parameters.Platform); }
+};
+
+IMPLEMENT_GLOBAL_SHADER(FHairVisibilityPrimitiveIdCompactionCS, "/Engine/Private/HairStrands/HairStrandsVisibilityCompaction.usf", "MainCS", SF_Compute);
+
+static void AddHairVisibilityPrimitiveIdCompactionPass(
+	FRDGBuilder& GraphBuilder,
+	const FViewInfo& View,
+	const FHairStrandsClusterDatas& ClusterDatas,
+	const uint32 NodeGroupSize,
+	const FRDGTextureRef& MSAA_DepthTexture,
+	const FRDGTextureRef& MSAA_IDTexture,
+	const FRDGTextureRef& MSAA_MaterialTexture,
+	const FRDGTextureRef& MSAA_AttributeTexture,
+	const FRDGTextureRef& MSAA_VelocityTexture,
+	FRDGTextureRef& OutCompactNodeIndex,
+	FRDGBufferRef& OutCompactNodeData,
+	FRDGBufferRef& OutCompactNodeCoord,
+	FRDGTextureRef& OutCategorizationTexture,
+	FRDGTextureRef& OutVelocityTexture,
+	FRDGBufferRef& OutIndirectArgsBuffer)
+{
+	check(MSAA_DepthTexture);
+	check(MSAA_IDTexture);
+	check(MSAA_MaterialTexture);
+	check(MSAA_AttributeTexture);
+
+	const FIntPoint Resolution = MSAA_DepthTexture->Desc.Extent;
+	
+	FRDGTextureRef CompactCounter;
+	{
+		FRDGTextureDesc Desc;
+		Desc.Extent.X = 1;
+		Desc.Extent.Y = 1;
+		Desc.Depth = 0;
+		Desc.Format = PF_R32_UINT;
+		Desc.NumMips = 1;
+		Desc.NumSamples = 1;
+		Desc.Flags = TexCreate_None;
+		Desc.TargetableFlags = TexCreate_UAV;
+		Desc.ClearValue = FClearValueBinding(0);
+		CompactCounter = GraphBuilder.CreateTexture(Desc, TEXT("HairVisibilityCompactCounter"));
+	}
+
+	{
+		FRDGTextureDesc Desc;
+		Desc.Extent = Resolution;
+		Desc.Depth = 0;
+		Desc.Format = PF_R32_UINT;
+		Desc.NumMips = 1;
+		Desc.NumSamples = 1;
+		Desc.Flags = TexCreate_None;
+		Desc.TargetableFlags = TexCreate_UAV | TexCreate_ShaderResource;
+		Desc.ClearValue = FClearValueBinding(0);
+		OutCompactNodeIndex = GraphBuilder.CreateTexture(Desc, TEXT("HairVisibilityCompactNodeIndex"));
+	}
+
+	{
+		FRDGTextureDesc OutputDesc;
+		OutputDesc.Extent = Resolution;
+		OutputDesc.Format = PF_R16G16B16A16_UINT;
+		OutputDesc.NumMips = 1;
+		OutputDesc.TargetableFlags = TexCreate_UAV;
+		OutCategorizationTexture = GraphBuilder.CreateTexture(OutputDesc, TEXT("CategorizationTexture"));
+	}
+	
+	AddClearUAVPass(GraphBuilder, RDG_EVENT_NAME("HairStrandsClearCompactionCounter"), CompactCounter, 0);
+	AddClearUAVPass(GraphBuilder, RDG_EVENT_NAME("HairStrandsClearCompactionOffsetAndCount"), OutCompactNodeIndex, 0);
+	AddClearUAVPass(GraphBuilder, RDG_EVENT_NAME("HairStrandsClearCategorizationTexture"), OutCategorizationTexture, 0);
+
+	const uint32 HairVisibilitySampleCount = GetHairVisibilitySampleCount();
+	const uint32 SampleCount = FMath::RoundUpToPowerOfTwo(HairVisibilitySampleCount);
+	const uint32 MaxNodeCount = Resolution.X * Resolution.Y * SampleCount;
+	{
+		struct NodeData
+		{
+			uint32 Depth;
+			uint32 PrimitiveId_ClusterId;
+			uint32 Tangent_Coverage;
+			uint32 BaseColor_Roughness;
+			uint32 Specular;
+		};
+
+		OutCompactNodeData = GraphBuilder.CreateBuffer(FRDGBufferDesc::CreateStructuredDesc(sizeof(NodeData), MaxNodeCount), TEXT("HairVisibilityPrimitiveIdCompactNodeData"));
+	}
+
+	{
+		// Pixel coord of the node. Stored as 2*uint16, packed into a single uint32
+		OutCompactNodeCoord = GraphBuilder.CreateBuffer(FRDGBufferDesc::CreateStructuredDesc(sizeof(uint32), MaxNodeCount), TEXT("HairVisibilityPrimitiveIdCompactNodeCoord"));
+	}
+
+	FSceneRenderTargets& SceneContext = FSceneRenderTargets::Get(GraphBuilder.RHICmdList);
+	FSceneTexturesUniformParameters SceneTextures;
+	SetupSceneTextureUniformParameters(SceneContext, View.FeatureLevel, ESceneTextureSetupMode::All, SceneTextures);
+
+	const bool bWriteOutVelocity = OutVelocityTexture != nullptr;
+	FHairVisibilityPrimitiveIdCompactionCS::FPermutationDomain PermutationVector;
+	PermutationVector.Set<FHairVisibilityPrimitiveIdCompactionCS::FVendor>(GetVendor());
+	PermutationVector.Set<FHairVisibilityPrimitiveIdCompactionCS::FVelocity>(bWriteOutVelocity ? FMath::Clamp(GHairVelocityType+1, 0, 3) : 0);
+
+	FHairVisibilityPrimitiveIdCompactionCS::FParameters* Parameters = GraphBuilder.AllocParameters<FHairVisibilityPrimitiveIdCompactionCS::FParameters>();
+	Parameters->MSAA_DepthTexture = MSAA_DepthTexture;
+	Parameters->MSAA_IDTexture = MSAA_IDTexture;
+	Parameters->MSAA_MaterialTexture = MSAA_MaterialTexture;
+	Parameters->MSAA_AttributeTexture = MSAA_AttributeTexture;
+	Parameters->OutputResolution = Resolution;
+	Parameters->MaxNodeCount = MaxNodeCount;
+	Parameters->HairVisibilitySampleCount = HairVisibilitySampleCount;
+	Parameters->SceneTexturesStruct = CreateUniformBufferImmediate(SceneTextures, EUniformBufferUsage::UniformBuffer_SingleDraw);
+	Parameters->OutCompactNodeCounter = GraphBuilder.CreateUAV(CompactCounter);
+	Parameters->OutCompactNodeIndex = GraphBuilder.CreateUAV(OutCompactNodeIndex);
+	Parameters->OutCompactNodeData = GraphBuilder.CreateUAV(OutCompactNodeData);
+	Parameters->OutCompactNodeCoord = GraphBuilder.CreateUAV(OutCompactNodeCoord);
+	Parameters->OutCategorizationTexture = GraphBuilder.CreateUAV(OutCategorizationTexture);
+
+	if (bWriteOutVelocity)
+	{
+		Parameters->MSAA_VelocityTexture = MSAA_VelocityTexture;
+		Parameters->OutVelocityTexture = GraphBuilder.CreateUAV(OutVelocityTexture);
+	}
+ 
+	FIntRect TotalRect = ComputeVisibleHairStrandsClustersRect(View.ViewRect, ClusterDatas);
+
+	// Snap the rect onto thread group boundary
+	const FIntPoint GroupSize = GetVendorOptimalGroupSize2D();
+	TotalRect.Min.X = FMath::FloorToInt(float(TotalRect.Min.X) / float(GroupSize.X)) * GroupSize.X;
+	TotalRect.Min.Y = FMath::FloorToInt(float(TotalRect.Min.Y) / float(GroupSize.Y)) * GroupSize.Y;
+	TotalRect.Max.X = FMath::CeilToInt(float(TotalRect.Max.X) / float(GroupSize.X)) * GroupSize.X;
+	TotalRect.Max.Y = FMath::CeilToInt(float(TotalRect.Max.Y) / float(GroupSize.Y)) * GroupSize.Y;
+	
+	FIntPoint RectResolution(TotalRect.Width(), TotalRect.Height());
+	Parameters->ResolutionOffset = FIntPoint(TotalRect.Min.X, TotalRect.Min.Y);
+
+	TShaderMapRef<FHairVisibilityPrimitiveIdCompactionCS> ComputeShader(View.ShaderMap, PermutationVector);
+	FComputeShaderUtils::AddPass(
+		GraphBuilder,
+		RDG_EVENT_NAME("HairStrandsVisibilityCompaction"),
+		*ComputeShader,
+		Parameters,
+		FComputeShaderUtils::GetGroupCount(RectResolution, GroupSize));
+
+	OutIndirectArgsBuffer = AddCopyIndirectArgPass(GraphBuilder, &View, NodeGroupSize, CompactCounter);
+}
+
+///////////////////////////////////////////////////////////////////////////////////////////////////
+class FHairVisibilityFillOpaqueDepthPS : public FGlobalShader
+{
+	DECLARE_GLOBAL_SHADER(FHairVisibilityFillOpaqueDepthPS);
+	SHADER_USE_PARAMETER_STRUCT(FHairVisibilityFillOpaqueDepthPS, FGlobalShader);
+
+	BEGIN_SHADER_PARAMETER_STRUCT(FParameters, )
+		SHADER_PARAMETER_RDG_TEXTURE(Texture2D, SceneDepthTexture)
+		SHADER_PARAMETER_RDG_TEXTURE(Texture2D, VisibilityDepthTexture)
+		SHADER_PARAMETER_RDG_TEXTURE(Texture2D, VisibilityIDTexture)
+
+		SHADER_PARAMETER_STRUCT_REF(FViewUniformShaderParameters, ViewUniformBuffer)
+		RENDER_TARGET_BINDING_SLOTS()
+		END_SHADER_PARAMETER_STRUCT()
+
+public:
+	static bool ShouldCompilePermutation(const FGlobalShaderPermutationParameters& Parameters) { return IsHairStrandsSupported(Parameters.Platform); }
+};
+
+IMPLEMENT_GLOBAL_SHADER(FHairVisibilityFillOpaqueDepthPS, "/Engine/Private/HairStrands/HairStrandsVisibilityFillOpaqueDepthPS.usf", "MainPS", SF_Pixel);
+
+static FRDGTextureRef AddHairVisibilityFillOpaqueDepth(
+	FRDGBuilder& GraphBuilder,
+	const FViewInfo& View,
+	const FIntPoint& Resolution,
+	const FHairStrandsClusterDatas& ClusterDatas,
+	const FRDGTextureRef& SceneDepthTexture)
+{
+	FRDGTextureRef OutVisibilityDepthTexture;
+	{
+		const uint32 MSAASampleCount = FMath::RoundUpToPowerOfTwo(FMath::Clamp(GHairVisibilitySampleCount, 1, 16));
+
+		FRDGTextureDesc Desc;
+		Desc.Extent.X = Resolution.X;
+		Desc.Extent.Y = Resolution.Y;
+		Desc.Depth = 0;
+		Desc.Format = PF_DepthStencil;
+		Desc.NumMips = 1;
+		Desc.NumSamples = MSAASampleCount;
+		Desc.Flags = TexCreate_None;
+		Desc.TargetableFlags = TexCreate_DepthStencilTargetable | TexCreate_ShaderResource;
+		Desc.ClearValue = FClearValueBinding::DepthFar;
+		Desc.bForceSharedTargetAndShaderResource = true;
+		OutVisibilityDepthTexture = GraphBuilder.CreateTexture(Desc, TEXT("HairVisibilityDepthTexture"));
+	}
+
+	FHairVisibilityFillOpaqueDepthPS::FParameters* Parameters = GraphBuilder.AllocParameters<FHairVisibilityFillOpaqueDepthPS::FParameters>();
+	Parameters->ViewUniformBuffer = View.ViewUniformBuffer;
+	Parameters->SceneDepthTexture = SceneDepthTexture;
+	Parameters->RenderTargets.DepthStencil = FDepthStencilBinding(
+		OutVisibilityDepthTexture,
+		ERenderTargetLoadAction::EClear,
+		ERenderTargetLoadAction::ENoAction,
+		FExclusiveDepthStencil::DepthWrite_StencilNop);
+
+	TShaderMapRef<FPostProcessVS> VertexShader(View.ShaderMap);
+	TShaderMapRef<FHairVisibilityFillOpaqueDepthPS> PixelShader(View.ShaderMap);
+	const TShaderMap<FGlobalShaderType>* GlobalShaderMap = View.ShaderMap;
+	const FIntRect Viewport = View.ViewRect;
+	const FViewInfo* CapturedView = &View;
+
+	TArray<FIntRect> ClusterRects;
+	for (const FHairStrandsClusterData& ClusterData : ClusterDatas.Datas)
+	{
+		ClusterRects.Add(ClusterData.ScreenRect);
+	}
+
+	{
+		ClearUnusedGraphResources(*PixelShader, Parameters);
+
+		GraphBuilder.AddPass(
+			RDG_EVENT_NAME("HairStrandsVisibilityFillOpaqueDepth"),
+			Parameters,
+			ERDGPassFlags::Raster,
+			[Parameters, VertexShader, PixelShader, ClusterRects, Viewport, Resolution, CapturedView](FRHICommandList& RHICmdList)
+		{
+			FGraphicsPipelineStateInitializer GraphicsPSOInit;
+			RHICmdList.ApplyCachedRenderTargets(GraphicsPSOInit);
+			GraphicsPSOInit.BlendState = TStaticBlendState<CW_RGBA, BO_Add, BF_One, BF_Zero, BO_Add, BF_One, BF_Zero>::GetRHI();
+			GraphicsPSOInit.RasterizerState = TStaticRasterizerState<>::GetRHI();
+			GraphicsPSOInit.DepthStencilState = TStaticDepthStencilState<true, CF_DepthNearOrEqual>::GetRHI();
+
+			GraphicsPSOInit.BoundShaderState.VertexDeclarationRHI = GFilterVertexDeclaration.VertexDeclarationRHI;
+			GraphicsPSOInit.BoundShaderState.VertexShaderRHI = GETSAFERHISHADER_VERTEX(*VertexShader);
+			GraphicsPSOInit.BoundShaderState.PixelShaderRHI = GETSAFERHISHADER_PIXEL(*PixelShader);
+			GraphicsPSOInit.PrimitiveType = PT_TriangleList;
+			SetGraphicsPipelineState(RHICmdList, GraphicsPSOInit);
+
+			VertexShader->SetParameters(RHICmdList, CapturedView->ViewUniformBuffer);
+			SetShaderParameters(RHICmdList, *PixelShader, PixelShader->GetPixelShader(), *Parameters);
+
+			for (const FIntRect& ViewRect : ClusterRects)
+			{
+				RHICmdList.SetViewport(ViewRect.Min.X, ViewRect.Min.Y, 0.0f, ViewRect.Max.X, ViewRect.Max.Y, 1.0f);
+				DrawRectangle(
+					RHICmdList,
+					0, 0,
+					Viewport.Width(), Viewport.Height(),
+					Viewport.Min.X,   Viewport.Min.Y,
+					Viewport.Width(), Viewport.Height(),
+					Viewport.Size(),
+					Resolution,
+					*VertexShader,
+					EDRF_UseTriangleOptimization);
+			}
+		});
+	}
+
+
+	return OutVisibilityDepthTexture;
+}
+
+///////////////////////////////////////////////////////////////////////////////////////////////////
+
+struct FVisibilityMSAAPassParameters
+{
+	BEGIN_SHADER_PARAMETER_STRUCT(FParameters, )
+		RENDER_TARGET_BINDING_SLOTS()
+	END_SHADER_PARAMETER_STRUCT()
+};
+
+template <typename TPassParameters>
+static void AddHairVisibilityCommonPass(
+	FRDGBuilder& GraphBuilder,
+	const FScene* Scene,
+	const FViewInfo* ViewInfo,
+	const FHairStrandsClusterDatas& ClusterDatas,
+	const EHairVisibilityRenderMode RenderMode,
+	TPassParameters* PassParameters)
+{
+	GraphBuilder.AddPass(
+		RDG_EVENT_NAME("HairStrandsVisibilityPass"),
+		PassParameters,
+		ERDGPassFlags::Raster,
+		[PassParameters, Scene = Scene, ViewInfo, &ClusterDatas, RenderMode](FRHICommandListImmediate& RHICmdList)
+	{
+		check(RHICmdList.IsInsideRenderPass());
+		check(IsInRenderingThread());
+
+		const FIntPoint Resolution(ViewInfo->ViewRect.Width(), ViewInfo->ViewRect.Height());
+		TUniformBufferRef<FHairVisibilityPassUniformParameters> PassUniformBuffer;
+		{
+			FVector2D PixelVelocity(1.f / (Resolution.X * 2), 1.f / (Resolution.Y * 2));
+			const float VelocityMagnitudeScale = FMath::Clamp(GHairVelocityMagnitudeScale, 0, 512) * FMath::Min(PixelVelocity.X, PixelVelocity.Y);
+
+			// Set the sample count to one as we want the size of the pixel
+			FHairVisibilityPassUniformParameters PassUniformParameters;
+			FMinHairRadiusAtDepth1 MinHairRadius = ComputeMinStrandRadiusAtDepth1(FIntPoint(ViewInfo->UnconstrainedViewRect.Width(), ViewInfo->UnconstrainedViewRect.Height()), ViewInfo->FOV, 1, 0);
+			PassUniformParameters.MinStrandRadius_Primary = MinHairRadius.Primary;
+			PassUniformParameters.MinStrandRadius_Velocity = MinHairRadius.Velocity;
+			PassUniformParameters.HairStrandsVelocityScale = VelocityMagnitudeScale;
+			PassUniformBuffer = CreateUniformBufferImmediate(PassUniformParameters, UniformBuffer_SingleDraw, EUniformBufferValidation::None);
+		}
+
+		FMeshPassProcessorRenderState DrawRenderState(*ViewInfo, PassUniformBuffer);
+		{
+			RHICmdList.SetViewport(0, 0, 0.0f, Resolution.X, Resolution.Y, 1.0f);
+			DrawRenderState.SetBlendState(TStaticBlendState<
+				CW_RGBA, BO_Add, BF_One, BF_Zero, BO_Add, BF_One, BF_Zero,
+				CW_RGBA, BO_Add, BF_One, BF_Zero, BO_Add, BF_One, BF_Zero>::GetRHI());
+			if (RenderMode == HairVisibilityRenderMode_MSAA)
+				DrawRenderState.SetDepthStencilState(TStaticDepthStencilState<true, CF_DepthNearOrEqual>::GetRHI());
+
+			FDynamicMeshDrawCommandStorage DynamicMeshDrawCommandStorage; // << Were would thid be stored?
+			// Use  void DrawDynamicMeshPass(const FSceneView& View, FRHICommandList& RHICmdList, const LambdaType& BuildPassProcessorLambda)
+			FMeshCommandOneFrameArray VisibleMeshDrawCommands;
+			FGraphicsMinimalPipelineStateSet PipelineStateSet;
+			FDynamicPassMeshDrawListContext ShadowContext(DynamicMeshDrawCommandStorage, VisibleMeshDrawCommands, PipelineStateSet);
+			FHairVisibilityProcessor MeshProcessor(Scene, ViewInfo /* is a SceneView */, DrawRenderState, RenderMode, &ShadowContext);
+
+			for (const FHairStrandsClusterData& ClusterData : ClusterDatas.Datas)
+			{
+				for (const FMeshBatchAndRelevance& MeshBatchAndRelevance : ClusterData.PrimitivesInfos)
+				{
+					const FMeshBatch& MeshBatch = *MeshBatchAndRelevance.Mesh;
+					const uint64 BatchElementMask = ~0ull;
+					MeshProcessor.AddMeshBatch(MeshBatch, BatchElementMask, MeshBatchAndRelevance.PrimitiveSceneProxy, -1, ClusterData.ClusterId);
+				}
+			}
+
+			if (VisibleMeshDrawCommands.Num() > 0)
+			{
+				FRHIVertexBuffer* PrimitiveIdVertexBuffer = nullptr;
+				SortAndMergeDynamicPassMeshDrawCommands(ViewInfo->GetFeatureLevel(), VisibleMeshDrawCommands, DynamicMeshDrawCommandStorage, PrimitiveIdVertexBuffer, 1);
+				SubmitMeshDrawCommands(VisibleMeshDrawCommands, PipelineStateSet, PrimitiveIdVertexBuffer, 0, false, 1, RHICmdList);
+			}
+		}
+	});
+}
+
+static void AddHairVisibilityMSAAPass(
+	FRDGBuilder& GraphBuilder,
+	const FScene* Scene,
+	const FViewInfo* ViewInfo,
+	const FHairStrandsClusterDatas& ClusterDatas,
+	const FIntPoint& Resolution,
+	FRDGTextureRef& VisibilityIdTexture,
+	FRDGTextureRef& VisibilityMaterialTexture,
+	FRDGTextureRef& VisibilityAttributeTexture,
+	FRDGTextureRef& VisibilityVelocityTexture,
+	FRDGTextureRef& VisibilityDepthTexture)
+{
+	const uint32 MSAASampleCount = FMath::RoundUpToPowerOfTwo(FMath::Clamp(GHairVisibilitySampleCount, 1, 16));
+
+	{
+		FRDGTextureDesc Desc;
+		Desc.Extent.X = Resolution.X;
+		Desc.Extent.Y = Resolution.Y;
+		Desc.Depth = 0;
+		Desc.Format = PF_R16G16B16A16_UINT;
+		Desc.NumMips = 1;
+		Desc.NumSamples = MSAASampleCount;
+		Desc.Flags = TexCreate_None;
+		Desc.TargetableFlags = TexCreate_RenderTargetable | TexCreate_ShaderResource;
+		Desc.bForceSharedTargetAndShaderResource = true;
+		VisibilityIdTexture = GraphBuilder.CreateTexture(Desc, TEXT("HairVisibilityIDTexture"));
+	}
+
+	{
+		FRDGTextureDesc Desc;
+		Desc.Extent.X = Resolution.X;
+		Desc.Extent.Y = Resolution.Y;
+		Desc.Depth = 0;
+		Desc.Format = PF_R8G8B8A8;
+		Desc.NumMips = 1;
+		Desc.NumSamples = MSAASampleCount;
+		Desc.Flags = TexCreate_None;
+		Desc.TargetableFlags = TexCreate_RenderTargetable | TexCreate_ShaderResource;
+		Desc.ClearValue = FClearValueBinding(FLinearColor(0, 0, 0, 0));
+		Desc.bForceSharedTargetAndShaderResource = true;
+		VisibilityMaterialTexture = GraphBuilder.CreateTexture(Desc, TEXT("HairVisibilityMaterialTexture"));
+	}
+
+	{
+		FRDGTextureDesc Desc;
+		Desc.Extent.X = Resolution.X;
+		Desc.Extent.Y = Resolution.Y;
+		Desc.Depth = 0;
+		Desc.Format = PF_R8G8B8A8;
+		Desc.NumMips = 1;
+		Desc.NumSamples = MSAASampleCount;
+		Desc.Flags = TexCreate_None;
+		Desc.TargetableFlags = TexCreate_RenderTargetable | TexCreate_ShaderResource;
+		Desc.ClearValue = FClearValueBinding(FLinearColor(0, 0, 0, 0));
+		Desc.bForceSharedTargetAndShaderResource = true;
+		VisibilityAttributeTexture = GraphBuilder.CreateTexture(Desc, TEXT("HairVisibilityAttributeTexture"));
+	}
+
+	{
+		FRDGTextureDesc Desc;
+		Desc.Extent.X = Resolution.X;
+		Desc.Extent.Y = Resolution.Y;
+		Desc.Depth = 0;
+		Desc.Format = PF_G16R16;
+		Desc.NumMips = 1;
+		Desc.NumSamples = MSAASampleCount;
+		Desc.Flags = TexCreate_None;
+		Desc.TargetableFlags = TexCreate_RenderTargetable | TexCreate_ShaderResource;
+		Desc.ClearValue = FClearValueBinding(FLinearColor(0, 0, 0, 0));
+		Desc.bForceSharedTargetAndShaderResource = true;
+		VisibilityVelocityTexture = GraphBuilder.CreateTexture(Desc, TEXT("HairVisibilityVelocityTexture"));
+	}
+	AddClearGraphicPass(GraphBuilder, RDG_EVENT_NAME("HairStrandsClearVisibilityMSAAIdTexture"), ViewInfo, 0xFFFFFFFF, VisibilityIdTexture);
+
+	// Manually clear RTs as using the Clear action on the RT, issue a global clean on all targets, while still need a special clear 
+	// for the PrimitiveId buffer
+	// const ERenderTargetLoadAction LoadAction = GHairClearVisibilityBuffer ? ERenderTargetLoadAction::EClear : ERenderTargetLoadAction::ENoAction;
+	ERenderTargetLoadAction LoadAction = ERenderTargetLoadAction::ENoAction;
+	if (GHairClearVisibilityBuffer)
+	{
+		LoadAction = ERenderTargetLoadAction::ELoad;
+		AddClearGraphicPass(GraphBuilder, RDG_EVENT_NAME("HairStrandsClearVisibilityMSAAMaterial"), ViewInfo, 0, VisibilityMaterialTexture);
+		AddClearGraphicPass(GraphBuilder, RDG_EVENT_NAME("HairStrandsClearVisibilityMSAAAttribute"), ViewInfo, 0, VisibilityAttributeTexture);
+		AddClearGraphicPass(GraphBuilder, RDG_EVENT_NAME("HairStrandsClearVisibilityMSAAVelocity"), ViewInfo, 0, VisibilityVelocityTexture);
+	}
+
+	FVisibilityMSAAPassParameters::FParameters* PassParameters = GraphBuilder.AllocParameters<FVisibilityMSAAPassParameters::FParameters>();
+	PassParameters->RenderTargets[0] = FRenderTargetBinding(VisibilityIdTexture, ERenderTargetLoadAction::ELoad, 0);
+	PassParameters->RenderTargets[1] = FRenderTargetBinding(VisibilityMaterialTexture,  LoadAction, 0);
+	PassParameters->RenderTargets[2] = FRenderTargetBinding(VisibilityAttributeTexture, LoadAction, 0);
+	PassParameters->RenderTargets[3] = FRenderTargetBinding(VisibilityVelocityTexture,  LoadAction, 0);
+
+	PassParameters->RenderTargets.DepthStencil = FDepthStencilBinding(
+		VisibilityDepthTexture,
+		ERenderTargetLoadAction::ELoad,
+		ERenderTargetLoadAction::ENoAction,
+		FExclusiveDepthStencil::DepthWrite_StencilNop);
+	AddHairVisibilityCommonPass(GraphBuilder, Scene, ViewInfo, ClusterDatas, HairVisibilityRenderMode_MSAA, PassParameters);
+}
+
+
+///////////////////////////////////////////////////////////////////////////////////////////////////
+class FHairVisibilityDepthPS : public FGlobalShader
+{
+	DECLARE_GLOBAL_SHADER(FHairVisibilityDepthPS);
+	SHADER_USE_PARAMETER_STRUCT(FHairVisibilityDepthPS, FGlobalShader);
+
+	BEGIN_SHADER_PARAMETER_STRUCT(FParameters, )
+		SHADER_PARAMETER_RDG_TEXTURE(Texture2D, HairVisibilityDepthTexture)
+		RENDER_TARGET_BINDING_SLOTS()
+	END_SHADER_PARAMETER_STRUCT()
+
+public:
+	static bool ShouldCompilePermutation(const FGlobalShaderPermutationParameters& Parameters) { return IsHairStrandsSupported(Parameters.Platform); }
+};
+
+IMPLEMENT_GLOBAL_SHADER(FHairVisibilityDepthPS, "/Engine/Private/HairStrands/HairStrandsVisibilityDepthPS.usf", "MainPS", SF_Pixel);
+
+static void AddHairVisibilityColorAndDepthPatchPass(
+	FRDGBuilder& GraphBuilder,
+	const FViewInfo& View,
+	const FRDGTextureRef& VisibilityDepthTexture,
+	FRDGTextureRef& OutGBufferBTexture,
+	FRDGTextureRef& OutColorTexture,
+	FRDGTextureRef& OutDepthTexture)
+{
+	FHairVisibilityDepthPS::FParameters* Parameters = GraphBuilder.AllocParameters<FHairVisibilityDepthPS::FParameters>();
+	Parameters->HairVisibilityDepthTexture = VisibilityDepthTexture;
+	Parameters->RenderTargets[0] = FRenderTargetBinding(OutGBufferBTexture, ERenderTargetLoadAction::ELoad);
+	Parameters->RenderTargets[1] = FRenderTargetBinding(OutColorTexture, ERenderTargetLoadAction::ELoad);
+	Parameters->RenderTargets.DepthStencil = FDepthStencilBinding(
+		OutDepthTexture,
+		ERenderTargetLoadAction::ELoad,
+		ERenderTargetLoadAction::ELoad,
+		FExclusiveDepthStencil::DepthWrite_StencilNop);
+
+	TShaderMapRef<FPostProcessVS> VertexShader(View.ShaderMap);
+	TShaderMapRef<FHairVisibilityDepthPS> PixelShader(View.ShaderMap);
+	const TShaderMap<FGlobalShaderType>* GlobalShaderMap = View.ShaderMap;
+	const FIntRect Viewport = View.ViewRect;
+	const FIntPoint Resolution = OutDepthTexture->Desc.Extent;
+	const FViewInfo* CapturedView = &View;
+
+	{
+		ClearUnusedGraphResources(*PixelShader, Parameters);
+
+		GraphBuilder.AddPass(
+			RDG_EVENT_NAME("HairStrandsVisibilityWriteColorAndDepth"),
+			Parameters,
+			ERDGPassFlags::Raster,
+			[Parameters, VertexShader, PixelShader, Viewport, Resolution, CapturedView](FRHICommandList& RHICmdList)
+		{
+			FGraphicsPipelineStateInitializer GraphicsPSOInit;
+			RHICmdList.ApplyCachedRenderTargets(GraphicsPSOInit);
+			GraphicsPSOInit.BlendState = TStaticBlendState<CW_RGBA, BO_Add, BF_One, BF_Zero, BO_Add, BF_One, BF_Zero>::GetRHI();
+			GraphicsPSOInit.RasterizerState = TStaticRasterizerState<>::GetRHI();
+			GraphicsPSOInit.DepthStencilState = TStaticDepthStencilState<true, CF_Greater>::GetRHI();
+
+			GraphicsPSOInit.BoundShaderState.VertexDeclarationRHI = GFilterVertexDeclaration.VertexDeclarationRHI;
+			GraphicsPSOInit.BoundShaderState.VertexShaderRHI = GETSAFERHISHADER_VERTEX(*VertexShader);
+			GraphicsPSOInit.BoundShaderState.PixelShaderRHI = GETSAFERHISHADER_PIXEL(*PixelShader);
+			GraphicsPSOInit.PrimitiveType = PT_TriangleList;
+			SetGraphicsPipelineState(RHICmdList, GraphicsPSOInit);
+
+			VertexShader->SetParameters(RHICmdList, CapturedView->ViewUniformBuffer);
+			RHICmdList.SetViewport(Viewport.Min.X, Viewport.Min.Y, 0.0f, Viewport.Max.X, Viewport.Max.Y, 1.0f);
+			SetShaderParameters(RHICmdList, *PixelShader, PixelShader->GetPixelShader(), *Parameters);
+			DrawRectangle(
+				RHICmdList,
+				0, 0,
+				Viewport.Width(), Viewport.Height(),
+				Viewport.Min.X, Viewport.Min.Y,
+				Viewport.Width(), Viewport.Height(),
+				Viewport.Size(),
+				Resolution,
+				*VertexShader,
+				EDRF_UseTriangleOptimization);
+		});
+	}
+}
+
+///////////////////////////////////////////////////////////////////////////////////////////////////
+
+FHairStrandsVisibilityViews RenderHairStrandsVisibilityBuffer(
+	FRHICommandListImmediate& RHICmdList, 
+	const FScene* Scene, 
+	const TArray<FViewInfo>& Views, 
+	TRefCountPtr<IPooledRenderTarget> InSceneGBufferBTexture,
+	TRefCountPtr<IPooledRenderTarget> InSceneColorTexture,
+	TRefCountPtr<IPooledRenderTarget> InSceneDepthTexture,
+	TRefCountPtr<IPooledRenderTarget> InSceneVelocityTexture,
+	const FHairStrandsClusterViews& ClusterViews)
+{
+	QUICK_SCOPE_CYCLE_COUNTER(STAT_CLM_RenderHairStrandsVisibility);
+	SCOPED_DRAW_EVENT(RHICmdList, HairStrandsVisibility);
+	SCOPED_GPU_STAT(RHICmdList, HairStrandsVisibility);
+
+	FHairStrandsVisibilityViews Output;
+
+	for (int32 ViewIndex = 0; ViewIndex < Views.Num(); ViewIndex++)
+	{
+		const FViewInfo& View = Views[ViewIndex];
+		if (View.Family)
+		{
+			FHairStrandsVisibilityData& VisibilityData = Output.HairDatas.AddDefaulted_GetRef();
+			VisibilityData.NodeGroupSize = GetVendorOptimalGroupSize1D();
+			const FHairStrandsClusterDatas& ClusterDatas = ClusterViews.Views[ViewIndex];
+
+			if (ClusterDatas.Datas.Num() == 0)
+				continue;
+
+			// Use the scene color for computing target resolution as the View.ViewRect, 
+			// doesn't include the actual resolution padding which make buffer size 
+			// mismatch, and create artifact (e.g. velocity computation)
+			const FIntPoint Resolution = InSceneColorTexture->GetDesc().Extent;
+
+
+			FRDGBuilder GraphBuilder(RHICmdList);
+			FSceneRenderTargets& SceneContext = FSceneRenderTargets::Get(RHICmdList);
+			FRDGTextureRef SceneGBufferBTexture = GraphBuilder.RegisterExternalTexture(InSceneGBufferBTexture, TEXT("SceneGBufferBTexture"));
+			FRDGTextureRef SceneColorTexture = GraphBuilder.RegisterExternalTexture(InSceneColorTexture, TEXT("SceneColorTexture"));
+			FRDGTextureRef SceneDepthTexture = GraphBuilder.RegisterExternalTexture(InSceneDepthTexture, TEXT("SceneDepthTexture"));
+			FRDGTextureRef SceneVelocityTexture = InSceneVelocityTexture ? GraphBuilder.RegisterExternalTexture(InSceneVelocityTexture, TEXT("SceneVelocityTexture")) : nullptr;
+			
+			const EHairVisibilityRenderMode RenderMode = GetHairVisibilityRenderMode();
+
+			struct FRDGMsaaVisibilityResources
+			{
+				FRDGTextureRef DepthTexture;
+				FRDGTextureRef IdTexture;
+				FRDGTextureRef MaterialTexture;
+				FRDGTextureRef AttributeTexture;
+				FRDGTextureRef VelocityTexture;
+			} MsaaVisibilityResources;
+
+			if (RenderMode == HairVisibilityRenderMode_MSAA)
+			{
+				
+				MsaaVisibilityResources.DepthTexture = AddHairVisibilityFillOpaqueDepth(
+					GraphBuilder,
+					View,
+					Resolution,
+					ClusterDatas,
+					SceneDepthTexture);
+
+				AddHairVisibilityMSAAPass(
+					GraphBuilder, 
+					Scene, 
+					&View, 
+					ClusterDatas, 
+					Resolution,
+					MsaaVisibilityResources.IdTexture, 
+					MsaaVisibilityResources.MaterialTexture, 
+					MsaaVisibilityResources.AttributeTexture,
+					MsaaVisibilityResources.VelocityTexture,
+					MsaaVisibilityResources.DepthTexture); 
+
+				// This is used when compaction is not enabled.
+				GraphBuilder.QueueTextureExtraction(MsaaVisibilityResources.IdTexture, &VisibilityData.IDTexture);
+				GraphBuilder.QueueTextureExtraction(MsaaVisibilityResources.MaterialTexture, &VisibilityData.MaterialTexture);
+				GraphBuilder.QueueTextureExtraction(MsaaVisibilityResources.AttributeTexture, &VisibilityData.AttributeTexture);
+				GraphBuilder.QueueTextureExtraction(MsaaVisibilityResources.VelocityTexture, &VisibilityData.VelocityTexture);
+				GraphBuilder.QueueTextureExtraction(MsaaVisibilityResources.DepthTexture, &VisibilityData.DepthTexture);
+				
+				{
+					FRDGTextureRef CompactNodeIndex;
+					FRDGBufferRef CompactNodeData;
+					FRDGBufferRef CompactNodeCoord;
+					FRDGTextureRef CategorizationTexture;
+					FRDGBufferRef IndirectArgsBuffer;
+
+					AddHairVisibilityPrimitiveIdCompactionPass(
+						GraphBuilder,
+						View,
+						ClusterDatas,
+						VisibilityData.NodeGroupSize,
+						MsaaVisibilityResources.DepthTexture,
+						MsaaVisibilityResources.IdTexture,
+						MsaaVisibilityResources.MaterialTexture,
+						MsaaVisibilityResources.AttributeTexture,
+						MsaaVisibilityResources.VelocityTexture,
+						CompactNodeIndex,
+						CompactNodeData,
+						CompactNodeCoord,
+						CategorizationTexture,
+						SceneVelocityTexture,
+						IndirectArgsBuffer);
+					GraphBuilder.QueueTextureExtraction(CompactNodeIndex,		&VisibilityData.NodeIndex);
+					GraphBuilder.QueueTextureExtraction(CategorizationTexture,	&VisibilityData.CategorizationTexture);
+					GraphBuilder.QueueBufferExtraction(CompactNodeData,			&VisibilityData.NodeData,					FRDGResourceState::EAccess::Read, FRDGResourceState::EPipeline::Graphics);
+					GraphBuilder.QueueBufferExtraction(CompactNodeCoord,		&VisibilityData.NodeCoord,					FRDGResourceState::EAccess::Read, FRDGResourceState::EPipeline::Graphics);
+					GraphBuilder.QueueBufferExtraction(IndirectArgsBuffer,		&VisibilityData.NodeIndirectArg,			FRDGResourceState::EAccess::Read, FRDGResourceState::EPipeline::Compute);
+				}
+			}
+		
+			// For fully covered pixels, write: 
+			// * black color into the scene color
+			// * closest depth
+			// * unlit shading model ID 
+			AddHairVisibilityColorAndDepthPatchPass(
+				GraphBuilder,
+				View,
+				MsaaVisibilityResources.DepthTexture,
+				SceneGBufferBTexture,
+				SceneColorTexture,
+				SceneDepthTexture);
+
+			GraphBuilder.Execute();
+
+			// #hair_todo: is there a better way to get SRV view of a RDG buffer? should work as long as there is not reuse between the pass
+			if (VisibilityData.NodeData)
+				VisibilityData.NodeDataSRV = RHICreateShaderResourceView(VisibilityData.NodeData->StructuredBuffer);
+
+			if (VisibilityData.NodeCoord)
+				VisibilityData.NodeCoordSRV = RHICreateShaderResourceView(VisibilityData.NodeCoord->StructuredBuffer);
+		}
+	}
+
+	return Output;
+}
diff --git a/Engine/Source/Runtime/Renderer/Private/HairStrands/HairStrandsVisibility.h b/Engine/Source/Runtime/Renderer/Private/HairStrands/HairStrandsVisibility.h
new file mode 100644
--- /dev/null
+++ b/Engine/Source/Runtime/Renderer/Private/HairStrands/HairStrandsVisibility.h
@@ -0,0 +1,44 @@
+// Copyright 1998-2019 Epic Games, Inc. All Rights Reserved.
+
+/*=============================================================================
+	HairVisibilityRendering.h: Hair strands visibility buffer implementation.
+=============================================================================*/
+
+#pragma once
+
+#include "CoreMinimal.h"
+#include "RendererInterface.h"
+#include "RenderGraphResources.h"
+
+struct FHairStrandsVisibilityData
+{
+	TRefCountPtr<IPooledRenderTarget> DepthTexture;
+	TRefCountPtr<IPooledRenderTarget> IDTexture;
+	TRefCountPtr<IPooledRenderTarget> MaterialTexture;
+	TRefCountPtr<IPooledRenderTarget> AttributeTexture;
+	TRefCountPtr<IPooledRenderTarget> VelocityTexture;
+	TRefCountPtr<IPooledRenderTarget> CategorizationTexture;
+
+	TRefCountPtr<IPooledRenderTarget> NodeIndex;
+	TRefCountPtr<FPooledRDGBuffer>	  NodeData;
+	FShaderResourceViewRHIRef		  NodeDataSRV;
+	TRefCountPtr<FPooledRDGBuffer>	  NodeCoord;
+	FShaderResourceViewRHIRef		  NodeCoordSRV;
+	TRefCountPtr<FPooledRDGBuffer>	  NodeIndirectArg;
+	uint32							  NodeGroupSize = 0;
+};
+
+struct FHairStrandsVisibilityViews
+{
+	TArray<FHairStrandsVisibilityData, SceneRenderingAllocator> HairDatas;
+};
+
+FHairStrandsVisibilityViews RenderHairStrandsVisibilityBuffer(
+	FRHICommandListImmediate& RHICmdList,
+	const class FScene* Scene,
+	const TArray<FViewInfo>& Views,
+	TRefCountPtr<IPooledRenderTarget> GBufferBTexture,
+	TRefCountPtr<IPooledRenderTarget> ColorTexture,
+	TRefCountPtr<IPooledRenderTarget> DepthTexture,
+	TRefCountPtr<IPooledRenderTarget> VelocityTexture,
+	const struct FHairStrandsClusterViews& ClusterViews);
diff --git a/Engine/Source/Runtime/Renderer/Private/HairStrands/HairStrandsVoxelization.cpp b/Engine/Source/Runtime/Renderer/Private/HairStrands/HairStrandsVoxelization.cpp
new file mode 100644
--- /dev/null
+++ b/Engine/Source/Runtime/Renderer/Private/HairStrands/HairStrandsVoxelization.cpp
@@ -0,0 +1,380 @@
+// Copyright 1998-2019 Epic Games, Inc. All Rights Reserved.
+
+#include "HairStrandsVoxelization.h"
+#include "HairStrandsRasterCommon.h"
+#include "HairStrandsCluster.h"
+#include "HairStrandsUtils.h"
+#include "Shader.h"
+#include "GlobalShader.h"
+#include "ShaderParameters.h"
+#include "ShaderParameterStruct.h"
+#include "SceneTextureParameters.h"
+#include "RenderGraphUtils.h"
+#include "PostProcessing.h"
+
+static float GStrandHairVoxelizationRasterizationScale = 1.0f;
+static FAutoConsoleVariableRef CVarStrandHairVoxelizationRasterizationScale(TEXT("r.HairStrands.VoxelizationRasterizationScale"), GStrandHairVoxelizationRasterizationScale, TEXT("Rasterization scale to snap strand to pixel for voxelization"));
+
+static int32 GHairVoxelizationResolution = 256;
+static FAutoConsoleVariableRef CVarGHairVoxelizationResolution(TEXT("r.HairStrands.Voxelization.Resolution"), GHairVoxelizationResolution, TEXT("Change the resolution of the voxelization volume for hair strands"));
+
+static int32 GHairVoxelizationEnable = 1;
+static FAutoConsoleVariableRef CVarGHairVoxelizationEnable(TEXT("r.HairStrands.Voxelization"), GHairVoxelizationEnable, TEXT("Enable hair voxelization for transmittance evaluation"));
+
+static int32 GHairVoxelizationMaterialEnable = 0;
+static FAutoConsoleVariableRef CVarHairVoxelizationMaterialEnable(TEXT("r.HairStrands.VoxelizationMaterial"), GHairVoxelizationMaterialEnable, TEXT("Enable hair material voxelization for LOD evaluation"));
+
+static float GHairVoxelizationAABBScale = 1.0f;
+static FAutoConsoleVariableRef CVarHairVoxelizationAABBScale(TEXT("r.HairStrands.Voxelization.AABBScale"), GHairVoxelizationAABBScale, TEXT("Scale the hair cluster bounding box"));
+
+static float GHairVoxelizationDensityScale = 1.0f;
+static float GHairVoxelizationDepthBiasScale = 3.0f;
+static FAutoConsoleVariableRef CVarHairVoxelizationDensityScale(TEXT("r.HairStrands.Voxelization.DensityScale"), GHairVoxelizationDensityScale, TEXT("Scale the hair density when computing voxel transmittance. Default value is 2 (arbitraty)"));
+static FAutoConsoleVariableRef CVarHairVoxelizationDepthBiasScale(TEXT("r.HairStrands.Voxelization.DepthBiasScale"), GHairVoxelizationDepthBiasScale, TEXT("Set depth bias for voxel ray marching. Offset the origin position towards the light"));
+
+static int32 GHairVoxelInjectOpaqueDepthEnable = 1;
+static FAutoConsoleVariableRef CVarHairVoxelInjectOpaqueDepthEnable(TEXT("r.HairStrands.Voxelization.InjectOpaqueDepth"), GHairVoxelInjectOpaqueDepthEnable, TEXT("Inject opaque geometry depth into the voxel volume for acting as occluder."));
+
+float GetHairStrandsVoxelizationDensityScale() { return FMath::Max(0.0f, GHairVoxelizationDensityScale); }
+float GetHairStrandsVoxelizationDepthBiasScale() { return FMath::Max(0.0f, GHairVoxelizationDepthBiasScale); }
+
+static int32 GHairForVoxelTransmittanceAndShadow = 0;
+static FAutoConsoleVariableRef CVarHairForVoxelTransmittanceAndShadow(TEXT("r.HairStrands.Voxelization.ForceTransmittanceAndShadow"), GHairForVoxelTransmittanceAndShadow, TEXT("For transmittance and shadow to be computed with density volume. This requires voxelization is enabled."));
+
+bool IsHairStrandsVoxelizationEnable()
+{
+	return GHairVoxelizationEnable > 0;
+}
+
+bool IsHairStrandsForVoxelTransmittanceAndShadowEnable()
+{
+	return IsHairStrandsVoxelizationEnable() && GHairForVoxelTransmittanceAndShadow > 0;
+}
+
+///////////////////////////////////////////////////////////////////////////////////////////////////
+class FVoxelInjectOpaquePS : public FGlobalShader
+{
+	DECLARE_GLOBAL_SHADER(FVoxelInjectOpaquePS);
+	SHADER_USE_PARAMETER_STRUCT(FVoxelInjectOpaquePS, FGlobalShader);
+
+	BEGIN_SHADER_PARAMETER_STRUCT(FParameters, )
+		SHADER_PARAMETER_STRUCT_INCLUDE(FSceneTextureParameters, SceneTextures)
+
+		SHADER_PARAMETER(FVector, VoxelMinAABB)
+		SHADER_PARAMETER(uint32, VoxelResolution)
+		SHADER_PARAMETER(FVector, VoxelMaxAABB)
+		SHADER_PARAMETER(FVector2D, OutputResolution)
+		SHADER_PARAMETER(FVector2D, SceneDepthResolution)
+
+		SHADER_PARAMETER_RDG_TEXTURE(Texture2D, SceneDepthTexture)
+		SHADER_PARAMETER_RDG_TEXTURE_UAV(RWTexture3D, DensityTexture)
+		SHADER_PARAMETER_SAMPLER(SamplerState, LinearSampler)
+
+		SHADER_PARAMETER_STRUCT_REF(FViewUniformShaderParameters, ViewUniformBuffer)
+		RENDER_TARGET_BINDING_SLOTS()
+	END_SHADER_PARAMETER_STRUCT()
+
+public:
+	static bool ShouldCompilePermutation(const FGlobalShaderPermutationParameters& Parameters) { return IsHairStrandsSupported(Parameters.Platform); }
+};
+
+IMPLEMENT_GLOBAL_SHADER(FVoxelInjectOpaquePS, "/Engine/Private/HairStrands/HairStrandsVoxelOpaque.usf", "MainPS", SF_Pixel);
+
+static void AddVoxelInjectOpaquePass(
+	FRDGBuilder& GraphBuilder,
+	const FViewInfo& View,
+	FHairStrandsClusterData& Cluster)
+{
+	FSceneTextureParameters SceneTextures;
+	SetupSceneTextureParameters(GraphBuilder, &SceneTextures);
+
+	// #hair_todo: change this to a CS. PS was for easing debugging
+	const uint32 LinearizeRes = FMath::Sqrt(Cluster.VoxelResources.DensityTexture->GetDesc().Depth);
+
+	FIntPoint Resolution;
+	Resolution.X = Cluster.VoxelResources.DensityTexture->GetDesc().Extent.X * LinearizeRes;
+	Resolution.Y = Cluster.VoxelResources.DensityTexture->GetDesc().Extent.Y * LinearizeRes;
+
+	if (!Cluster.VoxelResources.DensityTexture)
+		return;
+
+	const FRDGTextureRef VoxelDensityTexture = GraphBuilder.RegisterExternalTexture(Cluster.VoxelResources.DensityTexture, TEXT("HairVoxelDensityTexture"));
+
+	FRDGTextureDesc OutputDesc;
+	OutputDesc.Extent.X = Resolution.X;
+	OutputDesc.Extent.Y = Resolution.Y;
+	OutputDesc.Depth = 0;
+	OutputDesc.Format = PF_FloatRGBA;
+	OutputDesc.NumMips = 1;
+	OutputDesc.Flags = 0;
+	OutputDesc.TargetableFlags = TexCreate_RenderTargetable;
+	FRDGTextureRef DummyTexture = GraphBuilder.CreateTexture(OutputDesc, TEXT("HairVoxelInjectDepth")); // Dummy texture for debugging. Convert this pass into a compute shader.
+
+	FVoxelInjectOpaquePS::FParameters* Parameters = GraphBuilder.AllocParameters<FVoxelInjectOpaquePS::FParameters>();
+	Parameters->ViewUniformBuffer = View.ViewUniformBuffer;
+	Parameters->OutputResolution = Resolution;
+	Parameters->SceneDepthResolution = SceneTextures.SceneDepthBuffer->Desc.Extent;
+	Parameters->SceneDepthTexture = SceneTextures.SceneDepthBuffer;
+	Parameters->SceneTextures = SceneTextures;
+	Parameters->DensityTexture = GraphBuilder.CreateUAV(VoxelDensityTexture);
+	Parameters->VoxelMinAABB = Cluster.GetMinBound();
+	Parameters->VoxelMaxAABB = Cluster.GetMaxBound();
+	Parameters->VoxelResolution = Cluster.GetResolution();
+	Parameters->LinearSampler = TStaticSamplerState<SF_Trilinear, AM_Clamp, AM_Clamp, AM_Clamp>::GetRHI();
+	Parameters->RenderTargets[0] = FRenderTargetBinding(DummyTexture, ERenderTargetLoadAction::EClear);
+
+	const FIntPoint OutputResolution = SceneTextures.SceneDepthBuffer->Desc.Extent;
+	TShaderMapRef<FPostProcessVS> VertexShader(View.ShaderMap);
+	TShaderMapRef<FVoxelInjectOpaquePS> PixelShader(View.ShaderMap);
+	const TShaderMap<FGlobalShaderType>* GlobalShaderMap = View.ShaderMap;
+	const FIntRect Viewport(0, 0, Resolution.X, Resolution.Y);
+	const FViewInfo* CapturedView = &View;
+	
+	ClearUnusedGraphResources(*PixelShader, Parameters);
+
+	GraphBuilder.AddPass(
+		RDG_EVENT_NAME("HairStrandsVoxelInjectOpaque"),
+		Parameters,
+		ERDGPassFlags::Raster,
+		[Parameters, VertexShader, PixelShader, Viewport, Resolution, CapturedView](FRHICommandList& RHICmdList)
+	{
+		FGraphicsPipelineStateInitializer GraphicsPSOInit;
+		RHICmdList.ApplyCachedRenderTargets(GraphicsPSOInit);
+		GraphicsPSOInit.BlendState = TStaticBlendState<CW_RGBA, BO_Add, BF_One, BF_Zero, BO_Add, BF_One, BF_Zero>::GetRHI();
+		GraphicsPSOInit.RasterizerState = TStaticRasterizerState<>::GetRHI();
+		GraphicsPSOInit.DepthStencilState = TStaticDepthStencilState<false, CF_Always>::GetRHI();
+
+		GraphicsPSOInit.BoundShaderState.VertexDeclarationRHI = GFilterVertexDeclaration.VertexDeclarationRHI;
+		GraphicsPSOInit.BoundShaderState.VertexShaderRHI = GETSAFERHISHADER_VERTEX(*VertexShader);
+		GraphicsPSOInit.BoundShaderState.PixelShaderRHI = GETSAFERHISHADER_PIXEL(*PixelShader);
+		GraphicsPSOInit.PrimitiveType = PT_TriangleList;
+		SetGraphicsPipelineState(RHICmdList, GraphicsPSOInit);
+
+		VertexShader->SetParameters(RHICmdList, CapturedView->ViewUniformBuffer);
+		RHICmdList.SetViewport(Viewport.Min.X, Viewport.Min.Y, 0.0f, Viewport.Max.X, Viewport.Max.Y, 1.0f);
+		SetShaderParameters(RHICmdList, *PixelShader, PixelShader->GetPixelShader(), *Parameters);
+
+		DrawRectangle(
+			RHICmdList,
+			0, 0,
+			Viewport.Width(), Viewport.Height(),
+			Viewport.Min.X, Viewport.Min.Y,
+			Viewport.Width(), Viewport.Height(),
+			Viewport.Size(),
+			Resolution,
+			*VertexShader,
+			EDRF_UseTriangleOptimization);
+	});
+}
+
+///////////////////////////////////////////////////////////////////////////////////////////////////
+class FVoxelGenerateMipCS : public FGlobalShader
+{
+	DECLARE_GLOBAL_SHADER(FVoxelGenerateMipCS);
+	SHADER_USE_PARAMETER_STRUCT(FVoxelGenerateMipCS, FGlobalShader);
+
+	BEGIN_SHADER_PARAMETER_STRUCT(FParameters, )
+		SHADER_PARAMETER_STRUCT_INCLUDE(FSceneTextureParameters, SceneTextures)
+
+		SHADER_PARAMETER(uint32, VoxelResolution)
+		SHADER_PARAMETER(uint32, SourceMip)
+		SHADER_PARAMETER(uint32, TargetMip)
+
+		SHADER_PARAMETER_RDG_TEXTURE_SRV(Texture3D, InDensityTexture)
+		SHADER_PARAMETER_RDG_TEXTURE_UAV(RWTexture3D, OutDensityTexture)
+
+	END_SHADER_PARAMETER_STRUCT()
+
+public:
+	static bool ShouldCompilePermutation(const FGlobalShaderPermutationParameters& Parameters) { return IsHairStrandsSupported(Parameters.Platform); }
+};
+
+IMPLEMENT_GLOBAL_SHADER(FVoxelGenerateMipCS, "/Engine/Private/HairStrands/HairStrandsVoxelMip.usf", "MainCS", SF_Compute);
+
+static void AddVoxelGenerateMipPass(
+	FRDGBuilder& GraphBuilder,
+	const FViewInfo& View,
+	FHairStrandsClusterData& Cluster)
+{
+	if (!Cluster.VoxelResources.DensityTexture)
+		return;
+
+	TRefCountPtr<IPooledRenderTarget> DensityTexture = Cluster.VoxelResources.DensityTexture;
+	check(DensityTexture->GetDesc().Extent.X == DensityTexture->GetDesc().Extent.Y);
+
+	const uint32 VoxelResolution = DensityTexture->GetDesc().Extent.X;
+	const uint32 MipCount = DensityTexture->GetDesc().NumMips;
+	FRDGTextureRef VoxelDensityTexture = GraphBuilder.RegisterExternalTexture(DensityTexture, TEXT("HairVoxelDensityTexture"));
+	for (uint32 MipIt = 0; MipIt < MipCount-1; ++MipIt)
+	{
+		FVoxelGenerateMipCS::FParameters* Parameters = GraphBuilder.AllocParameters<FVoxelGenerateMipCS::FParameters>();
+		Parameters->InDensityTexture = GraphBuilder.CreateSRV(FRDGTextureSRVDesc::CreateForMipLevel(VoxelDensityTexture, MipIt));
+		Parameters->OutDensityTexture = GraphBuilder.CreateUAV(FRDGTextureUAVDesc(VoxelDensityTexture, MipIt+1));
+		Parameters->VoxelResolution = VoxelResolution;
+		Parameters->SourceMip = MipIt;
+		Parameters->TargetMip = MipIt+1;
+
+		uint32 SourceResolution = VoxelResolution >> (MipIt);
+		uint32 TargetResolution = VoxelResolution >> (MipIt + 1);
+
+		TShaderMapRef<FVoxelGenerateMipCS> ComputeShader(View.ShaderMap);
+		const TShaderMap<FGlobalShaderType>* GlobalShaderMap = View.ShaderMap;
+		const FIntVector DispatchCount = FComputeShaderUtils::GetGroupCount(FIntVector(TargetResolution, TargetResolution, TargetResolution), FIntVector(4, 4, 4));
+		
+		ClearUnusedGraphResources(*ComputeShader, Parameters);
+		GraphBuilder.AddPass(
+			RDG_EVENT_NAME("HairStrandsVoxelMip"),
+			Parameters,
+			ERDGPassFlags::Compute | ERDGPassFlags::GenerateMips,
+			[Parameters, ComputeShader, DispatchCount](FRHICommandList& RHICmdList)
+		{
+			FComputeShaderUtils::Dispatch(RHICmdList, *ComputeShader, *Parameters, DispatchCount);
+		});
+	}
+}
+
+static void RenderVoxelization(
+	FRHICommandListImmediate& RHICmdList,
+	const FScene* Scene,
+	const FViewInfo* ViewInfo,
+	FHairStrandsClusterData& Cluster)
+{
+	const FHairStrandsClusterData::TPrimitiveInfos& PrimitiveSceneInfo = Cluster.PrimitivesInfos;
+	const FBoxSphereBounds& Bounds = Cluster.Bounds;
+	FHairStrandsVoxelResources& VoxelResources = Cluster.VoxelResources;
+
+	DECLARE_GPU_STAT(HairStrandsVoxelization);
+	SCOPED_DRAW_EVENT(RHICmdList, HairStrandsVoxelization);
+	SCOPED_GPU_STAT(RHICmdList, HairStrandsVoxelization);
+
+	const uint32 ResolutionDim = FMath::RoundUpToPowerOfTwo(GHairVoxelizationResolution);
+	const EPixelFormat Format = PF_R32_UINT;
+	const FIntPoint Resolution(ResolutionDim, ResolutionDim);
+	uint32 MipCount = 1;
+	{
+		uint32 CurrentResolution = ResolutionDim;
+		while (CurrentResolution > 4)
+		{
+			++MipCount;
+			CurrentResolution = CurrentResolution >> 1;
+		}
+	}
+
+	// todo remove this dummy texture
+	TRefCountPtr<IPooledRenderTarget> DummyTexture;
+	{
+		FPooledRenderTargetDesc Desc(FPooledRenderTargetDesc::Create2DDesc(Resolution, PF_R32_UINT, FClearValueBinding::Black, TexCreate_None, TexCreate_RenderTargetable, false));
+		GRenderTargetPool.FindFreeElement(RHICmdList, Desc, DummyTexture, TEXT("DummyTexture"));
+	}
+
+	const bool bVoxelizeMaterial = GHairVoxelizationMaterialEnable > 0;
+	{
+		FPooledRenderTargetDesc Desc(FPooledRenderTargetDesc::CreateVolumeDesc(ResolutionDim, ResolutionDim, ResolutionDim, PF_R32_UINT, FClearValueBinding::Black, TexCreate_None, TexCreate_UAV | TexCreate_ShaderResource, false, MipCount));
+		GRenderTargetPool.FindFreeElement(RHICmdList, Desc, VoxelResources.DensityTexture, TEXT("DensityTexture"));
+		if (bVoxelizeMaterial)
+		{
+			GRenderTargetPool.FindFreeElement(RHICmdList, Desc, VoxelResources.TangentXTexture, TEXT("TangentXTexture"));
+			GRenderTargetPool.FindFreeElement(RHICmdList, Desc, VoxelResources.TangentYTexture, TEXT("TangentYTexture"));
+			GRenderTargetPool.FindFreeElement(RHICmdList, Desc, VoxelResources.TangentZTexture, TEXT("TangentZTexture"));
+			GRenderTargetPool.FindFreeElement(RHICmdList, Desc, VoxelResources.MaterialTexture, TEXT("MaterialTexture"));
+		}
+
+		uint32 ClearValues[4] = { 0, 0, 0, 0 };
+		RHICmdList.ClearTinyUAV(VoxelResources.DensityTexture->GetRenderTargetItem().UAV, ClearValues);
+		if (bVoxelizeMaterial)
+		{
+			RHICmdList.ClearTinyUAV(VoxelResources.TangentXTexture->GetRenderTargetItem().UAV, ClearValues);
+			RHICmdList.ClearTinyUAV(VoxelResources.TangentYTexture->GetRenderTargetItem().UAV, ClearValues);
+			RHICmdList.ClearTinyUAV(VoxelResources.TangentZTexture->GetRenderTargetItem().UAV, ClearValues);
+			RHICmdList.ClearTinyUAV(VoxelResources.MaterialTexture->GetRenderTargetItem().UAV, ClearValues);
+		}
+	}
+
+	FRHIRenderPassInfo RPInfo(DummyTexture->GetRenderTargetItem().TargetableTexture, ERenderTargetActions::Load_Store);
+	RPInfo.UAVIndex = 0;
+	RPInfo.NumUAVs = bVoxelizeMaterial ? 5 : 1;
+	RPInfo.UAVs[0] = VoxelResources.DensityTexture->GetRenderTargetItem().UAV;
+	if (bVoxelizeMaterial)
+	{
+		RPInfo.UAVs[1] = VoxelResources.TangentXTexture->GetRenderTargetItem().UAV;
+		RPInfo.UAVs[2] = VoxelResources.TangentYTexture->GetRenderTargetItem().UAV;
+		RPInfo.UAVs[3] = VoxelResources.TangentZTexture->GetRenderTargetItem().UAV;
+		RPInfo.UAVs[4] = VoxelResources.MaterialTexture->GetRenderTargetItem().UAV;
+	}
+
+	const FIntRect ViewportRect(0, 0, ResolutionDim, ResolutionDim);
+
+	FHairStrandsLightDesc LightDesc;
+	LightDesc.bIsOrtho = true;
+	LightDesc.LightDirection = FVector(1, 0, 0);
+	LightDesc.LightPosition = FVector(0, 0, 0);
+	{
+		const FSphere SphereBound = Bounds.GetSphere();
+		const float SphereRadius = SphereBound.W * GHairVoxelizationAABBScale;
+
+		FReversedZOrthoMatrix OrthoMatrix(SphereRadius, SphereRadius, 1.f / (2 * SphereRadius), 0);
+		FLookAtMatrix LookAt(SphereBound.Center - LightDesc.LightDirection * SphereRadius, SphereBound.Center, FVector(0, 0, 1));
+		LightDesc.WorldToLightClipTransform = LookAt * OrthoMatrix;
+		LightDesc.MinStrandRadiusAtDepth1.Primary = GStrandHairVoxelizationRasterizationScale * SphereRadius / FMath::Min(Resolution.X, Resolution.Y);
+		LightDesc.MinStrandRadiusAtDepth1.Velocity = LightDesc.MinStrandRadiusAtDepth1.Primary;
+
+		VoxelResources.WorldToClip = LightDesc.WorldToLightClipTransform;
+		VoxelResources.MinAABB = Bounds.GetSphere().Center - SphereRadius;
+		VoxelResources.MaxAABB = Bounds.GetSphere().Center + SphereRadius;
+	}
+
+	RHICmdList.BeginRenderPass(RPInfo, TEXT("HairVoxelization"));
+	RasterHairStrands(
+		RHICmdList,
+		Scene,
+		ViewInfo,
+		PrimitiveSceneInfo,
+		LightDesc,
+		bVoxelizeMaterial ? EHairStrandsRasterPassType::VoxelizationMaterial : EHairStrandsRasterPassType::Voxelization,
+		nullptr,
+		ViewportRect,
+		VoxelResources.MinAABB,
+		VoxelResources.MaxAABB);
+	RHICmdList.EndRenderPass();
+
+	if (GHairVoxelInjectOpaqueDepthEnable)
+	{
+		FRDGBuilder GraphBuilder(RHICmdList);
+		AddVoxelInjectOpaquePass(GraphBuilder, *ViewInfo, Cluster);
+		GraphBuilder.Execute();
+	}
+
+	{
+		FRDGBuilder GraphBuilder(RHICmdList);
+		AddVoxelGenerateMipPass(GraphBuilder, *ViewInfo, Cluster);
+		GraphBuilder.Execute();
+	}
+}
+
+void VoxelizeHairStrands(
+	FRHICommandListImmediate& RHICmdList, 
+	const FScene* Scene, 
+	const TArray<FViewInfo>& Views,
+	FHairStrandsClusterViews& DeepShadowClusterViews)
+{
+	if (!IsHairStrandsVoxelizationEnable())
+		return;
+
+	FHairStrandsClusterViews PrimitivesClusterViews;
+	for (int32 ViewIndex = 0; ViewIndex < Views.Num(); ViewIndex++)
+	{			
+		if (ViewIndex >= DeepShadowClusterViews.Views.Num())
+			continue;
+
+		const FViewInfo& View = Views[ViewIndex];
+		FHairStrandsClusterDatas& ClusterDatas = DeepShadowClusterViews.Views[ViewIndex];
+		for (FHairStrandsClusterData& Cluster : ClusterDatas.Datas)
+		{
+			RenderVoxelization(
+				RHICmdList,
+				Scene,
+				&View,
+				Cluster);
+		}
+	}
+}
diff --git a/Engine/Source/Runtime/Renderer/Private/HairStrands/HairStrandsVoxelization.h b/Engine/Source/Runtime/Renderer/Private/HairStrands/HairStrandsVoxelization.h
new file mode 100644
--- /dev/null
+++ b/Engine/Source/Runtime/Renderer/Private/HairStrands/HairStrandsVoxelization.h
@@ -0,0 +1,34 @@
+// Copyright 1998-2019 Epic Games, Inc. All Rights Reserved.
+
+/*=============================================================================
+	HairStrandsVoxelization.h: Hair voxelization implementation.
+=============================================================================*/
+
+#pragma once
+
+#include "CoreMinimal.h"
+#include "RendererInterface.h"
+
+struct FHairStrandsVoxelResources
+{
+	TRefCountPtr<IPooledRenderTarget> DensityTexture;
+	TRefCountPtr<IPooledRenderTarget> TangentXTexture;
+	TRefCountPtr<IPooledRenderTarget> TangentYTexture;
+	TRefCountPtr<IPooledRenderTarget> TangentZTexture;
+	TRefCountPtr<IPooledRenderTarget> MaterialTexture;
+	FMatrix WorldToClip;
+	FVector MinAABB;
+	FVector MaxAABB;
+};
+
+/// Global enable/disable for hair voxelization
+bool IsHairStrandsVoxelizationEnable();
+bool IsHairStrandsForVoxelTransmittanceAndShadowEnable();
+float GetHairStrandsVoxelizationDensityScale();
+float GetHairStrandsVoxelizationDepthBiasScale();
+
+void VoxelizeHairStrands(
+	FRHICommandListImmediate& RHICmdList,
+	const class FScene* Scene,
+	const TArray<FViewInfo>& Views,
+	struct FHairStrandsClusterViews& DeepShadowClusterViews);
\ No newline at end of file
diff --git a/Engine/Source/Runtime/Renderer/Private/LightGridInjection.cpp b/Engine/Source/Runtime/Renderer/Private/LightGridInjection.cpp
--- a/Engine/Source/Runtime/Renderer/Private/LightGridInjection.cpp
+++ b/Engine/Source/Runtime/Renderer/Private/LightGridInjection.cpp
@@ -792,7 +792,7 @@ void FDeferredShadingSceneRenderer::RenderForwardShadingShadowProjections(FRHICo
 
 				if (VisibleLightInfo.ShadowsToProject.Num() > 0)
 				{
-					FSceneRenderer::RenderShadowProjections(RHICmdList, LightSceneInfo, ForwardScreenSpaceShadowMask, true, false);
+					FSceneRenderer::RenderShadowProjections(RHICmdList, LightSceneInfo, ForwardScreenSpaceShadowMask, nullptr, true, false, nullptr);
 				}
 
 				RenderCapsuleDirectShadows(RHICmdList, *LightSceneInfo, ForwardScreenSpaceShadowMask, VisibleLightInfo.CapsuleShadowsToProject, true);
