Copying //UE4/Dev-Physics to Dev-Main (//UE4/Dev-Main) [at] 9098278

diff --git a/Engine/Source/Runtime/Renderer/Private/LightRendering.cpp b/Engine/Source/Runtime/Renderer/Private/LightRendering.cpp
--- a/Engine/Source/Runtime/Renderer/Private/LightRendering.cpp
+++ b/Engine/Source/Runtime/Renderer/Private/LightRendering.cpp
@@ -17,6 +17,7 @@
 #include "VisualizeTexture.h"
 #include "RayTracing/RaytracingOptions.h"
 #include "SceneTextureParameters.h"
+#include "HairStrands/HairStrandsRendering.h"
 
 // ENABLE_DEBUG_DISCARD_PROP is used to test the lighting code by allowing to discard lights to see how performance scales
 // It ought never to be enabled in a shipping build, and is probably only really useful when woring on the shading code.
@@ -200,6 +201,20 @@ IMPLEMENT_SHADER_TYPE(template<>,TDeferredLightVS<false>,TEXT("/Engine/Private/D
 IMPLEMENT_SHADER_TYPE(template<>,TDeferredLightVS<true>,TEXT("/Engine/Private/DeferredLightVertexShaders.usf"),TEXT("RadialVertexMain"),SF_Vertex);
 
 
+struct FRenderLightParams
+{
+	// Precompute transmittance
+	FShaderResourceViewRHIRef DeepShadow_TransmittanceMaskBuffer = nullptr;
+	uint32 DeepShadow_TransmittanceMaskBufferMaxCount = 0;
+	bool bUseTransmittanceData = false;
+	
+	// Visibility buffer data
+	IPooledRenderTarget* HairVisibilityNodeOffsetAndCount = nullptr;
+	FShaderResourceViewRHIRef HairVisibilityNodeDataSRV = nullptr;
+
+	IPooledRenderTarget* ScreenShadowMaskSubPixelTexture = nullptr;
+};
+
 enum class ELightSourceShape
 {
 	Directional,
@@ -221,6 +236,7 @@ class FDeferredLightPS : public FGlobalShader
 	class FVisualizeCullingDim	: SHADER_PERMUTATION_BOOL("VISUALIZE_LIGHT_CULLING");
 	class FLightingChannelsDim	: SHADER_PERMUTATION_BOOL("USE_LIGHTING_CHANNELS");
 	class FTransmissionDim		: SHADER_PERMUTATION_BOOL("USE_TRANSMISSION");
+	class FHairLighting			: SHADER_PERMUTATION_BOOL("USE_HAIR_LIGHTING");
 
 	using FPermutationDomain = TShaderPermutationDomain<
 		FSourceShapeDim,
@@ -229,7 +245,8 @@ class FDeferredLightPS : public FGlobalShader
 		FInverseSquaredDim,
 		FVisualizeCullingDim,
 		FLightingChannelsDim,
-		FTransmissionDim >;
+		FTransmissionDim,
+		FHairLighting>;
 
 	static bool ShouldCompilePermutation(const FGlobalShaderPermutationParameters& Parameters)
 	{
@@ -284,23 +301,40 @@ class FDeferredLightPS : public FGlobalShader
 		LightingChannelsSampler.Bind(Initializer.ParameterMap, TEXT("LightingChannelsSampler"));
 		TransmissionProfilesTexture.Bind(Initializer.ParameterMap, TEXT("SSProfilesTexture"));
 		TransmissionProfilesLinearSampler.Bind(Initializer.ParameterMap, TEXT("TransmissionProfilesLinearSampler"));
+
+		HairTransmittanceBuffer.Bind(Initializer.ParameterMap, TEXT("HairTransmittanceBuffer"));
+		HairTransmittanceBufferMaxCount.Bind(Initializer.ParameterMap, TEXT("HairTransmittanceBufferMaxCount"));
+		ScreenShadowMaskSubPixelTexture.Bind(Initializer.ParameterMap, TEXT("ScreenShadowMaskSubPixelTexture")); // TODO hook the shader itself
+
+		HairLUTTexture.Bind(Initializer.ParameterMap, TEXT("HairLUTTexture"));
+		HairLUTSampler.Bind(Initializer.ParameterMap, TEXT("HairLUTSampler"));
+		HairComponents.Bind(Initializer.ParameterMap, TEXT("HairComponents"));
+		HairTransmittanceEnable.Bind(Initializer.ParameterMap, TEXT("HairTransmittanceEnable"));
+
+		HairVisibilityNodeOffsetAndCount.Bind(Initializer.ParameterMap, TEXT("HairVisibilityNodeOffsetAndCount"));
+		HairVisibilityNodeData.Bind(Initializer.ParameterMap, TEXT("HairVisibilityNodeData"));
 	}
 
 	FDeferredLightPS()
 	{}
 
 public:
-	void SetParameters(FRHICommandList& RHICmdList, const FSceneView& View, const FLightSceneInfo* LightSceneInfo, IPooledRenderTarget* ScreenShadowMaskTexture)
+	void SetParameters(
+		FRHICommandList& RHICmdList, 
+		const FSceneView& View, 
+		const FLightSceneInfo* LightSceneInfo, 
+		IPooledRenderTarget* ScreenShadowMaskTexture, 
+		FRenderLightParams* RenderLightParams)
 	{
 		FRHIPixelShader* ShaderRHI = GetPixelShader();
-		SetParametersBase(RHICmdList, ShaderRHI, View, ScreenShadowMaskTexture, LightSceneInfo->Proxy->GetIESTextureResource());
+		SetParametersBase(RHICmdList, ShaderRHI, View, ScreenShadowMaskTexture, LightSceneInfo->Proxy->GetIESTextureResource(), RenderLightParams);
 		SetDeferredLightParameters(RHICmdList, ShaderRHI, GetUniformBufferParameter<FDeferredLightUniformStruct>(), LightSceneInfo, View);
 	}
 
 	void SetParametersSimpleLight(FRHICommandList& RHICmdList, const FSceneView& View, const FSimpleLightEntry& SimpleLight, const FSimpleLightPerViewEntry& SimpleLightPerViewData)
 	{
 		FRHIPixelShader* ShaderRHI = GetPixelShader();
-		SetParametersBase(RHICmdList, ShaderRHI, View, NULL, NULL);
+		SetParametersBase(RHICmdList, ShaderRHI, View, nullptr, nullptr, nullptr);
 		SetSimpleDeferredLightParameters(RHICmdList, ShaderRHI, GetUniformBufferParameter<FDeferredLightUniformStruct>(), SimpleLight, SimpleLightPerViewData, View);
 	}
 
@@ -320,12 +354,31 @@ public:
 		Ar << LightingChannelsSampler;
 		Ar << TransmissionProfilesTexture;
 		Ar << TransmissionProfilesLinearSampler;
+
+		Ar << HairTransmittanceBuffer;
+		Ar << HairTransmittanceBufferMaxCount;
+
+		Ar << HairVisibilityNodeOffsetAndCount;
+		Ar << HairVisibilityNodeData;
+		Ar << ScreenShadowMaskSubPixelTexture;
+
+		Ar << HairLUTTexture;
+		Ar << HairLUTSampler;
+		Ar << HairComponents;
+		Ar << HairTransmittanceEnable;
+
 		return bShaderHasOutdatedParameters;
 	}
 
 private:
 
-	void SetParametersBase(FRHICommandList& RHICmdList, FRHIPixelShader* ShaderRHI, const FSceneView& View, IPooledRenderTarget* ScreenShadowMaskTexture, FTexture* IESTextureResource)
+	void SetParametersBase(
+		FRHICommandList& RHICmdList, 
+		FRHIPixelShader* ShaderRHI, 
+		const FSceneView& View, 
+		IPooledRenderTarget* ScreenShadowMaskTexture, 
+		FTexture* IESTextureResource, 
+		FRenderLightParams* RenderLightParams)
 	{
 		FGlobalShader::SetParameters<FViewUniformShaderParameters>(RHICmdList, ShaderRHI,View.ViewUniformBuffer);
 		SceneTextureParameters.Set(RHICmdList, ShaderRHI, View.FeatureLevel, ESceneTextureSetupMode::All);
@@ -409,6 +462,93 @@ private:
 				TStaticSamplerState<SF_Bilinear, AM_Clamp, AM_Clamp, AM_Clamp>::GetRHI(),
 				Item.ShaderResourceTexture);
 		}
+
+		if (HairTransmittanceBuffer.IsBound())
+		{
+			const uint32 TransmittanceBufferMaxCount = RenderLightParams ? RenderLightParams->DeepShadow_TransmittanceMaskBufferMaxCount : 0;
+			SetShaderValue(
+				RHICmdList,
+				ShaderRHI,
+				HairTransmittanceBufferMaxCount,
+				TransmittanceBufferMaxCount);
+			if (RenderLightParams && RenderLightParams->DeepShadow_TransmittanceMaskBuffer)
+			{
+				SetSRVParameter(RHICmdList, ShaderRHI, HairTransmittanceBuffer, RenderLightParams->DeepShadow_TransmittanceMaskBuffer);
+			}
+		}
+
+		if (ScreenShadowMaskSubPixelTexture.IsBound())
+		{
+			if (RenderLightParams && RenderLightParams->ScreenShadowMaskSubPixelTexture)
+			{
+				SetTextureParameter(
+					RHICmdList,
+					ShaderRHI,
+					ScreenShadowMaskSubPixelTexture,
+					LightAttenuationTextureSampler,
+					TStaticSamplerState<SF_Point, AM_Wrap, AM_Wrap, AM_Wrap>::GetRHI(),
+					(RenderLightParams && RenderLightParams->ScreenShadowMaskSubPixelTexture) ? RenderLightParams->ScreenShadowMaskSubPixelTexture->GetRenderTargetItem().ShaderResourceTexture : GWhiteTexture->TextureRHI);
+			}
+		}
+
+		if (HairVisibilityNodeOffsetAndCount.IsBound())
+		{
+			if (RenderLightParams && RenderLightParams->HairVisibilityNodeOffsetAndCount)
+			{
+				SetTextureParameter(
+					RHICmdList,
+					ShaderRHI,
+					HairVisibilityNodeOffsetAndCount,
+					LightAttenuationTextureSampler,
+					TStaticSamplerState<SF_Point, AM_Wrap, AM_Wrap, AM_Wrap>::GetRHI(),
+					RenderLightParams->HairVisibilityNodeOffsetAndCount->GetRenderTargetItem().TargetableTexture);
+			}
+		}
+
+		if (HairVisibilityNodeData.IsBound())
+		{
+			if (RenderLightParams && RenderLightParams->HairVisibilityNodeDataSRV)
+			{
+				FShaderResourceViewRHIRef SRV = RenderLightParams->HairVisibilityNodeDataSRV;
+				SetSRVParameter(
+					RHICmdList, 
+					ShaderRHI, 
+					HairVisibilityNodeData, 
+					SRV);
+			}
+		}
+
+		if (HairLUTTexture.IsBound())
+		{
+			IPooledRenderTarget* HairLUTTextureResource = GSystemTextures.HairLUT0;
+			SetTextureParameter(
+				RHICmdList,
+				ShaderRHI,
+				HairLUTTexture,
+				HairLUTSampler,
+				TStaticSamplerState<SF_Bilinear, AM_Clamp, AM_Clamp, AM_Clamp>::GetRHI(),
+				HairLUTTextureResource ? HairLUTTextureResource->GetRenderTargetItem().ShaderResourceTexture : GBlackVolumeTexture->TextureRHI);
+		}
+
+		if (HairComponents.IsBound())
+		{
+			FVector4 InHairComponents = GetHairComponents();
+			SetShaderValue(
+				RHICmdList,
+				ShaderRHI,
+				HairComponents,
+				InHairComponents);
+		}
+
+		if (HairTransmittanceEnable.IsBound())
+		{
+			uint32 bIsDeepShadowEnable = (RenderLightParams && RenderLightParams->bUseTransmittanceData) ? 1 : 0;
+			SetShaderValue(
+				RHICmdList,
+				ShaderRHI,
+				HairTransmittanceEnable,
+				bIsDeepShadowEnable);
+		}
 	}
 
 	FSceneTextureShaderParameters SceneTextureParameters;
@@ -424,7 +564,18 @@ private:
 	FShaderResourceParameter LightingChannelsSampler;
 	FShaderResourceParameter TransmissionProfilesTexture;
 	FShaderResourceParameter TransmissionProfilesLinearSampler;
-};
+
+	FShaderResourceParameter HairTransmittanceBuffer;
+	FShaderParameter HairTransmittanceBufferMaxCount;
+	FShaderResourceParameter HairVisibilityNodeOffsetAndCount;
+	FShaderResourceParameter HairVisibilityNodeData;
+	FShaderResourceParameter ScreenShadowMaskSubPixelTexture;
+
+	FShaderResourceParameter HairLUTTexture;
+	FShaderResourceParameter HairLUTSampler;
+	FShaderParameter HairComponents;
+	FShaderParameter HairTransmittanceEnable;
+	};
 
 IMPLEMENT_GLOBAL_SHADER(FDeferredLightPS, "/Engine/Private/DeferredLightPixelShaders.usf", "DeferredLightPixelMain", SF_Pixel);
 
@@ -742,8 +893,24 @@ void FDeferredShadingSceneRenderer::GatherAndSortLights(FSortedLightSetSceneInfo
 
 
 /** Renders the scene's lighting. */
-void FDeferredShadingSceneRenderer::RenderLights(FRHICommandListImmediate& RHICmdList, FSortedLightSetSceneInfo &SortedLightSet)
+void FDeferredShadingSceneRenderer::RenderLights(FRHICommandListImmediate& RHICmdList, FSortedLightSetSceneInfo &SortedLightSet, const FHairStrandsDatas* HairDatas)
 {
+	const bool bUseHairLighting = HairDatas != nullptr;
+	const FHairStrandsDeepShadowViews* InDeepShadowViews = HairDatas ? &HairDatas->DeepShadowViews : nullptr;
+	const FHairStrandsVisibilityViews* InHairVisibilityViews = HairDatas ? &HairDatas->HairVisibilityViews : nullptr;
+	auto HasHairStrandsClusters = [bUseHairLighting, HairDatas](int32 ViewIndex)
+	{
+		return HairDatas && ViewIndex < HairDatas->HairClusterPerViews.Views.Num() && HairDatas->HairClusterPerViews.Views[ViewIndex].Datas.Num() > 0;
+	};
+
+	auto GetHairStrandsCategorizationTexture = [HairDatas](int32 ViewIndex, FRDGBuilder& GraphBuilder)
+	{
+		FRDGTextureRef HairCategorizationTexture = nullptr;
+		if (HairDatas && ViewIndex < HairDatas->HairVisibilityViews.HairDatas.Num() && HairDatas->HairVisibilityViews.HairDatas[ViewIndex].CategorizationTexture)
+			HairCategorizationTexture = GraphBuilder.RegisterExternalTexture(HairDatas->HairVisibilityViews.HairDatas[ViewIndex].CategorizationTexture);
+		return HairCategorizationTexture;
+	};
+
 	check(RHICmdList.IsOutsideRenderPass());
 
 	SCOPED_NAMED_EVENT(FDeferredShadingSceneRenderer_RenderLights, FColor::Emerald);
@@ -772,6 +939,14 @@ void FDeferredShadingSceneRenderer::RenderLights(FRHICommandListImmediate& RHICm
 			RHICmdList.WaitComputeFence(TranslucencyLightingVolumeClearEndFence);
 		}
 
+		// Clear sub pixel color when there is no light
+		if (bUseHairLighting)
+		{
+			FRHIRenderPassInfo RPInfo(SceneContext.SceneColorSubPixel->GetRenderTargetItem().TargetableTexture, MakeRenderTargetActions(ERenderTargetLoadAction::EClear, ERenderTargetStoreAction::EStore));
+			RHICmdList.BeginRenderPass(RPInfo, TEXT("ClearSubSpixelColor"));
+			RHICmdList.EndRenderPass();
+		}
+
 		if(ViewFamily.EngineShowFlags.DirectLighting)
 		{
 			SCOPED_DRAW_EVENT(RHICmdList, NonShadowedLights);
@@ -837,7 +1012,7 @@ void FDeferredShadingSceneRenderer::RenderLights(FRHICommandListImmediate& RHICm
 					const FLightSceneInfo* const LightSceneInfo = SortedLightInfo.LightSceneInfo;
 
 					// Render the light to the scene color buffer, using a 1x1 white texture as input
-					RenderLight(RHICmdList, LightSceneInfo, NULL, false, false);
+					RenderLight(RHICmdList, LightSceneInfo, nullptr, nullptr, nullptr, nullptr, false, false);
 				}
 
 				SceneContext.FinishRenderingSceneColor(RHICmdList);
@@ -945,6 +1120,7 @@ void FDeferredShadingSceneRenderer::RenderLights(FRHICommandListImmediate& RHICm
 			const IScreenSpaceDenoiser* DenoiserToUse = DenoiserMode == 1 ? DefaultDenoiser : GScreenSpaceDenoiser;
 
 			TArray<TRefCountPtr<IPooledRenderTarget>> PreprocessedShadowMaskTextures;
+			TArray<TRefCountPtr<IPooledRenderTarget>> PreprocessedShadowMaskSubPixelTextures;
 
 			const int32 MaxDenoisingBatchSize = FMath::Clamp(CVarMaxShadowDenoisingBatchSize.GetValueOnRenderThread(), 1, IScreenSpaceDenoiser::kMaxBatchSize);
 			const int32 MaxRTShadowBatchSize = CVarMaxShadowRayTracingBatchSize.GetValueOnRenderThread();
@@ -956,7 +1132,8 @@ void FDeferredShadingSceneRenderer::RenderLights(FRHICommandListImmediate& RHICm
 			// for denoising perf.
 			if (RHI_RAYTRACING && dDoShadowBatching)
 			{
-				FViewInfo& View = Views[0];
+				const uint32 ViewIndex = 0;
+				FViewInfo& View = Views[ViewIndex];
 
 				// Allocate PreprocessedShadowMaskTextures once so QueueTextureExtraction can deferred write.
 				{
@@ -970,11 +1147,17 @@ void FDeferredShadingSceneRenderer::RenderLights(FRHICommandListImmediate& RHICm
 				}
 
 				PreprocessedShadowMaskTextures.SetNum(SortedLights.Num());
+
+				if (HasHairStrandsClusters(ViewIndex))
+				{ 
+					PreprocessedShadowMaskSubPixelTextures.SetNum(SortedLights.Num());
+				}
 			} // if (RHI_RAYTRACING)
 
 			bool bDirectLighting = ViewFamily.EngineShowFlags.DirectLighting;
 
 			TRefCountPtr<IPooledRenderTarget> ScreenShadowMaskTexture;
+			TRefCountPtr<IPooledRenderTarget> ScreenShadowMaskSubPixelTexture;
 
 			// Draw shadowed and light function lights
 			for (int32 LightIndex = AttenuationLightStart; LightIndex < SortedLights.Num(); LightIndex++)
@@ -989,12 +1172,15 @@ void FDeferredShadingSceneRenderer::RenderLights(FRHICommandListImmediate& RHICm
 				bool bDrawPreviewIndicator = ViewFamily.EngineShowFlags.PreviewShadowsIndicator && !LightSceneInfo.IsPrecomputedLightingValid() && LightSceneInfo.Proxy->HasStaticShadowing();
 				bool bInjectedTranslucentVolume = false;
 				bool bUsedShadowMaskTexture = false;
+				const bool bDrawDeepShadow = bDrawShadows && bUseHairLighting;
 
 				FScopeCycleCounter Context(LightSceneInfo.Proxy->GetStatId());
 
 				if ((bDrawShadows || bDrawLightFunction || bDrawPreviewIndicator) && !ScreenShadowMaskTexture.IsValid())
 				{
 					SceneContext.AllocateScreenShadowMask(RHICmdList, ScreenShadowMaskTexture);
+					if (bUseHairLighting)
+						SceneContext.AllocateScreenShadowMask(RHICmdList, ScreenShadowMaskSubPixelTexture);
 				}
 
 				FString LightNameWithLevel;
@@ -1010,7 +1196,8 @@ void FDeferredShadingSceneRenderer::RenderLights(FRHICommandListImmediate& RHICm
 					// Inline ray traced shadow batching, launches shadow batches when needed
 					// reduces memory overhead while keeping shadows batched to optimize costs
 					{
-						FViewInfo& View = Views[0];
+						const uint32 ViewIndex = 0;
+						FViewInfo& View = Views[ViewIndex];
 
 						IScreenSpaceDenoiser::FShadowRayTracingConfig RayTracingConfig;
 						RayTracingConfig.RayCountPerPixel = LightSceneInfo.Proxy->GetSamplesPerPixel();
@@ -1110,6 +1297,7 @@ void FDeferredShadingSceneRenderer::RenderLights(FRHICommandListImmediate& RHICm
 								}
 
 								// Ray trace the shadow.
+								FRDGTextureRef HairCategorizationTexture = GetHairStrandsCategorizationTexture(ViewIndex, GraphBuilder);
 								FRDGTextureRef ShadowMask;
 								FRDGTextureRef RayHitDistance;
 								{
@@ -1124,10 +1312,34 @@ void FDeferredShadingSceneRenderer::RenderLights(FRHICommandListImmediate& RHICm
 										BatchLightSceneInfo,
 										BatchRayTracingConfig,
 										DenoiserRequirements,
+										false,
+										HairCategorizationTexture,
 										&ShadowMask,
 										&RayHitDistance);
 								}
 
+								// Ray trace the shadow cast by opaque geometries on to hair strands geometries
+								// Note: No denoiser is required on this output, as the hair strands are geometrically noisy, and so hard to denoise
+								if (HasHairStrandsClusters(ViewIndex))
+								{
+									FRDGTextureRef SubPixelShadowMask = nullptr;
+									FRDGTextureRef SubPixelRayHitDistance = nullptr;
+									RenderRayTracingShadows(
+										GraphBuilder,
+										SceneTextures,
+										View,
+										LightSceneInfo,
+										RayTracingConfig,
+										IScreenSpaceDenoiser::EShadowRequirements::Bailout,
+										true,
+										HairCategorizationTexture,
+										&SubPixelShadowMask,
+										&SubPixelRayHitDistance);
+									TRefCountPtr<IPooledRenderTarget>* RefDestination = &PreprocessedShadowMaskSubPixelTextures[LightIndex - AttenuationLightStart];
+									check(*RefDestination == nullptr);
+									GraphBuilder.QueueTextureExtraction(SubPixelShadowMask, RefDestination);
+								}
+
 								bool bBatchFull = false;
 
 								if (bRequiresDenoiser)
@@ -1183,72 +1395,104 @@ void FDeferredShadingSceneRenderer::RenderLights(FRHICommandListImmediate& RHICm
 
 					if (RHI_RAYTRACING && PreprocessedShadowMaskTextures.Num() > 0 && PreprocessedShadowMaskTextures[LightIndex - AttenuationLightStart])
 					{
-						ScreenShadowMaskTexture = PreprocessedShadowMaskTextures[LightIndex - AttenuationLightStart];
-						PreprocessedShadowMaskTextures[LightIndex - AttenuationLightStart] = nullptr;
+						const uint32 ShadowMaskIndex = LightIndex - AttenuationLightStart;
+						ScreenShadowMaskTexture = PreprocessedShadowMaskTextures[ShadowMaskIndex];
+						PreprocessedShadowMaskTextures[ShadowMaskIndex] = nullptr;
+
+						// Subp-ixel shadow for hair strands geometries
+						if (bUseHairLighting && ShadowMaskIndex < uint32(PreprocessedShadowMaskSubPixelTextures.Num()))
+						{
+							ScreenShadowMaskSubPixelTexture = PreprocessedShadowMaskSubPixelTextures[ShadowMaskIndex];
+							PreprocessedShadowMaskSubPixelTextures[ShadowMaskIndex] = nullptr;
+						}
 					}
 					else if (OcclusionType == FLightOcclusionType::Raytraced)
 					{
-						FViewInfo& View = Views[0];
+						const uint32 ViewIndex = 0;
+						FViewInfo& View = Views[ViewIndex];
 
-							IScreenSpaceDenoiser::FShadowRayTracingConfig RayTracingConfig;
-							RayTracingConfig.RayCountPerPixel = LightSceneInfo.Proxy->GetSamplesPerPixel();
+						IScreenSpaceDenoiser::FShadowRayTracingConfig RayTracingConfig;
+						RayTracingConfig.RayCountPerPixel = LightSceneInfo.Proxy->GetSamplesPerPixel();
 
-							IScreenSpaceDenoiser::EShadowRequirements DenoiserRequirements = IScreenSpaceDenoiser::EShadowRequirements::Bailout;
-							if (DenoiserMode != 0 && LightRequiresDenosier(LightSceneInfo))
-							{
-								DenoiserRequirements = DenoiserToUse->GetShadowRequirements(View, LightSceneInfo, RayTracingConfig);
-							}
+						IScreenSpaceDenoiser::EShadowRequirements DenoiserRequirements = IScreenSpaceDenoiser::EShadowRequirements::Bailout;
+						if (DenoiserMode != 0 && LightRequiresDenosier(LightSceneInfo))
+						{
+							DenoiserRequirements = DenoiserToUse->GetShadowRequirements(View, LightSceneInfo, RayTracingConfig);
+						}
 
 						FRDGBuilder GraphBuilder(RHICmdList);
 
 						FSceneTextureParameters SceneTextures;
 						SetupSceneTextureParameters(GraphBuilder, &SceneTextures);
 
+						FRDGTextureRef HairCategorizationTexture = GetHairStrandsCategorizationTexture(ViewIndex, GraphBuilder);
 						FRDGTextureRef ShadowMask;
 						FRDGTextureRef RayHitDistance;
-							RenderRayTracingShadows(
-								GraphBuilder,
-								SceneTextures,
-								View,
-								LightSceneInfo,
-								RayTracingConfig,
-								DenoiserRequirements,
+						RenderRayTracingShadows(
+							GraphBuilder,
+							SceneTextures,
+							View,
+							LightSceneInfo,
+							RayTracingConfig,
+							DenoiserRequirements,
+							false,
+							HairCategorizationTexture,
 							&ShadowMask,
 							&RayHitDistance);
 
-							if (DenoiserRequirements != IScreenSpaceDenoiser::EShadowRequirements::Bailout)
-							{
-								TStaticArray<IScreenSpaceDenoiser::FShadowParameters, IScreenSpaceDenoiser::kMaxBatchSize> InputParameters;
-								TStaticArray<IScreenSpaceDenoiser::FShadowPenumbraOutputs, IScreenSpaceDenoiser::kMaxBatchSize> Outputs;
+						if (DenoiserRequirements != IScreenSpaceDenoiser::EShadowRequirements::Bailout)
+						{
+							TStaticArray<IScreenSpaceDenoiser::FShadowParameters, IScreenSpaceDenoiser::kMaxBatchSize> InputParameters;
+							TStaticArray<IScreenSpaceDenoiser::FShadowPenumbraOutputs, IScreenSpaceDenoiser::kMaxBatchSize> Outputs;
 
-								InputParameters[0].InputTextures.Penumbra = ShadowMask;
-								InputParameters[0].InputTextures.ClosestOccluder = RayHitDistance;
-								InputParameters[0].LightSceneInfo = &LightSceneInfo;
-								InputParameters[0].RayTracingConfig = RayTracingConfig;
+							InputParameters[0].InputTextures.Penumbra = ShadowMask;
+							InputParameters[0].InputTextures.ClosestOccluder = RayHitDistance;
+							InputParameters[0].LightSceneInfo = &LightSceneInfo;
+							InputParameters[0].RayTracingConfig = RayTracingConfig;
 
-								int32 InputParameterCount = 1;
+							int32 InputParameterCount = 1;
 
-								RDG_EVENT_SCOPE(GraphBuilder, "%s%s(Shadow BatchSize=%d) %dx%d",
-									DenoiserToUse != DefaultDenoiser ? TEXT("ThirdParty ") : TEXT(""),
-									DenoiserToUse->GetDebugName(),
-									InputParameterCount,
-									View.ViewRect.Width(), View.ViewRect.Height());
+							RDG_EVENT_SCOPE(GraphBuilder, "%s%s(Shadow BatchSize=%d) %dx%d",
+								DenoiserToUse != DefaultDenoiser ? TEXT("ThirdParty ") : TEXT(""),
+								DenoiserToUse->GetDebugName(),
+								InputParameterCount,
+								View.ViewRect.Width(), View.ViewRect.Height());
 
-								DenoiserToUse->DenoiseMonochromaticShadows(
-									GraphBuilder,
-									View,
-									&View.PrevViewInfo,
-									SceneTextures,
-									InputParameters,
-									InputParameterCount,
-									Outputs);
+							DenoiserToUse->DenoiseMonochromaticShadows(
+								GraphBuilder,
+								View,
+								&View.PrevViewInfo,
+								SceneTextures,
+								InputParameters,
+								InputParameterCount,
+								Outputs);
 
-								GraphBuilder.QueueTextureExtraction(Outputs[0].DiffusePenumbra, &ScreenShadowMaskTexture);
-							}
-							else
-							{
-								GraphBuilder.QueueTextureExtraction(ShadowMask, &ScreenShadowMaskTexture);
-							}
+							GraphBuilder.QueueTextureExtraction(Outputs[0].DiffusePenumbra, &ScreenShadowMaskTexture);
+						}
+						else
+						{
+							GraphBuilder.QueueTextureExtraction(ShadowMask, &ScreenShadowMaskTexture);
+						}
+
+						// Ray trace the shadow cast by opaque geometries on to hair strands geometries
+						// Note: No denoiser is required on this output, as the hair strands are geometrically noisy, and so hard to denoise
+						if (HasHairStrandsClusters(ViewIndex))
+						{
+							FRDGTextureRef SubPixelShadowMask = nullptr;
+							FRDGTextureRef SubPixelRayHitDistance = nullptr;
+							RenderRayTracingShadows(
+								GraphBuilder,
+								SceneTextures,
+								View,
+								LightSceneInfo,
+								RayTracingConfig,
+								IScreenSpaceDenoiser::EShadowRequirements::Bailout,
+								true,
+								HairCategorizationTexture,
+								&SubPixelShadowMask,
+								&SubPixelRayHitDistance);
+							GraphBuilder.QueueTextureExtraction(SubPixelShadowMask, &ScreenShadowMaskSubPixelTexture);
+						}
 
 						GraphBuilder.Execute();
 					}
@@ -1260,51 +1504,57 @@ void FDeferredShadingSceneRenderer::RenderLights(FRHICommandListImmediate& RHICm
 							View.HeightfieldLightingViewInfo.ClearShadowing(View, RHICmdList, LightSceneInfo);
 						}
 					
-						// Clear light attenuation for local lights with a quad covering their extents
-						const bool bClearLightScreenExtentsOnly = CVarAllowClearLightSceneExtentsOnly.GetValueOnRenderThread() && SortedLightInfo.SortKey.Fields.LightType != LightType_Directional;
-						// All shadows render with min blending
-						bool bClearToWhite = !bClearLightScreenExtentsOnly;
-
-						FRHIRenderPassInfo RPInfo(ScreenShadowMaskTexture->GetRenderTargetItem().TargetableTexture, ERenderTargetActions::Load_Store);
-						RPInfo.DepthStencilRenderTarget.Action = MakeDepthStencilTargetActions(ERenderTargetActions::Load_DontStore, ERenderTargetActions::Load_Store);
-						RPInfo.DepthStencilRenderTarget.DepthStencilTarget = SceneContext.GetSceneDepthSurface();
-						RPInfo.DepthStencilRenderTarget.ExclusiveDepthStencil = FExclusiveDepthStencil::DepthRead_StencilWrite;
-						if (bClearToWhite)
-						{
-							RPInfo.ColorRenderTargets[0].Action = ERenderTargetActions::Clear_Store;
-						}
-
-						TransitionRenderPassTargets(RHICmdList, RPInfo);
-						RHICmdList.BeginRenderPass(RPInfo, TEXT("ClearScreenShadowMask"));
-						if (bClearLightScreenExtentsOnly)
+						auto ClearShadowMask = [&](TRefCountPtr<IPooledRenderTarget>& InScreenShadowMaskTexture)
 						{
-							SCOPED_DRAW_EVENT(RHICmdList, ClearQuad);
-
-							for (int32 ViewIndex = 0; ViewIndex < Views.Num(); ViewIndex++)
+							// Clear light attenuation for local lights with a quad covering their extents
+							const bool bClearLightScreenExtentsOnly = CVarAllowClearLightSceneExtentsOnly.GetValueOnRenderThread() && SortedLightInfo.SortKey.Fields.LightType != LightType_Directional;
+							// All shadows render with min blending
+							bool bClearToWhite = !bClearLightScreenExtentsOnly;
+
+							FRHIRenderPassInfo RPInfo(InScreenShadowMaskTexture->GetRenderTargetItem().TargetableTexture, ERenderTargetActions::Load_Store);
+							RPInfo.DepthStencilRenderTarget.Action = MakeDepthStencilTargetActions(ERenderTargetActions::Load_DontStore, ERenderTargetActions::Load_Store);
+							RPInfo.DepthStencilRenderTarget.DepthStencilTarget = SceneContext.GetSceneDepthSurface();
+							RPInfo.DepthStencilRenderTarget.ExclusiveDepthStencil = FExclusiveDepthStencil::DepthRead_StencilWrite;
+							if (bClearToWhite)
 							{
-								const FViewInfo& View = Views[ViewIndex];
-								FIntRect ScissorRect;
+								RPInfo.ColorRenderTargets[0].Action = ERenderTargetActions::Clear_Store;
+							}
 
-								if (!LightSceneInfo.Proxy->GetScissorRect(ScissorRect, View, View.ViewRect))
+							TransitionRenderPassTargets(RHICmdList, RPInfo);
+							RHICmdList.BeginRenderPass(RPInfo, TEXT("ClearScreenShadowMask"));
+							if (bClearLightScreenExtentsOnly)
+							{
+								SCOPED_DRAW_EVENT(RHICmdList, ClearQuad);
+	
+								for (int32 ViewIndex = 0; ViewIndex < Views.Num(); ViewIndex++)
 								{
-									ScissorRect = View.ViewRect;
-								}
+									const FViewInfo& View = Views[ViewIndex];
+									FIntRect ScissorRect;
 
-								if (ScissorRect.Min.X < ScissorRect.Max.X && ScissorRect.Min.Y < ScissorRect.Max.Y)
-								{
-									RHICmdList.SetViewport(ScissorRect.Min.X, ScissorRect.Min.Y, 0.0f, ScissorRect.Max.X, ScissorRect.Max.Y, 1.0f);
-									DrawClearQuad(RHICmdList, true, FLinearColor(1, 1, 1, 1), false, 0, false, 0);
-								}
-								else
-								{
-									LightSceneInfo.Proxy->GetScissorRect(ScissorRect, View, View.ViewRect);
+									if (!LightSceneInfo.Proxy->GetScissorRect(ScissorRect, View, View.ViewRect))
+									{
+										ScissorRect = View.ViewRect;
+									}
+
+									if (ScissorRect.Min.X < ScissorRect.Max.X && ScissorRect.Min.Y < ScissorRect.Max.Y)
+									{
+										RHICmdList.SetViewport(ScissorRect.Min.X, ScissorRect.Min.Y, 0.0f, ScissorRect.Max.X, ScissorRect.Max.Y, 1.0f);
+										DrawClearQuad(RHICmdList, true, FLinearColor(1, 1, 1, 1), false, 0, false, 0);
+									}
+									else
+									{
+										LightSceneInfo.Proxy->GetScissorRect(ScissorRect, View, View.ViewRect);
+									}
 								}
 							}
-						}
+							RHICmdList.EndRenderPass();
+						};
 
-						RHICmdList.EndRenderPass();
+						ClearShadowMask(ScreenShadowMaskTexture);
+						if (ScreenShadowMaskSubPixelTexture)
+							ClearShadowMask(ScreenShadowMaskSubPixelTexture);
 
-						RenderShadowProjections(RHICmdList, &LightSceneInfo, ScreenShadowMaskTexture, bInjectedTranslucentVolume);
+						RenderShadowProjections(RHICmdList, &LightSceneInfo, ScreenShadowMaskTexture, ScreenShadowMaskSubPixelTexture, HairDatas, bInjectedTranslucentVolume);
 					}
 
 					bUsedShadowMaskTexture = true;
@@ -1339,6 +1589,8 @@ void FDeferredShadingSceneRenderer::RenderLights(FRHICommandListImmediate& RHICm
 				if (bUsedShadowMaskTexture)
 				{
 					RHICmdList.CopyToResolveTarget(ScreenShadowMaskTexture->GetRenderTargetItem().TargetableTexture, ScreenShadowMaskTexture->GetRenderTargetItem().ShaderResourceTexture, FResolveParams(FResolveRect()));
+					if (ScreenShadowMaskSubPixelTexture)
+						RHICmdList.CopyToResolveTarget(ScreenShadowMaskSubPixelTexture->GetRenderTargetItem().TargetableTexture, ScreenShadowMaskSubPixelTexture->GetRenderTargetItem().ShaderResourceTexture, FResolveParams(FResolveRect()));
 				}
 
 				if(bDirectLighting && !bInjectedTranslucentVolume)
@@ -1352,6 +1604,14 @@ void FDeferredShadingSceneRenderer::RenderLights(FRHICommandListImmediate& RHICm
 				}
 
 				GVisualizeTexture.SetCheckPoint(RHICmdList, ScreenShadowMaskTexture);
+				if (ScreenShadowMaskSubPixelTexture)
+					GVisualizeTexture.SetCheckPoint(RHICmdList, ScreenShadowMaskSubPixelTexture);
+
+				FHairStrandsTransmittanceMaskData TransmittanceMaskData;
+				if (bDrawDeepShadow)
+				{
+					TransmittanceMaskData = RenderHairStrandsTransmittanceMask(RHICmdList, Views, &LightSceneInfo, HairDatas);
+				}
 
 				if (ShouldRenderRayTracingStochasticRectLight(LightSceneInfo))
 				{
@@ -1363,12 +1623,13 @@ void FDeferredShadingSceneRenderer::RenderLights(FRHICommandListImmediate& RHICm
 				}
 				else
 				{
-					SceneContext.BeginRenderingSceneColor(RHICmdList, ESimpleRenderTargetMode::EExistingColorAndDepth, FExclusiveDepthStencil::DepthRead_StencilWrite);
+					SCOPED_DRAW_EVENT(RHICmdList, StandardDeferredLighting);
+					SceneContext.BeginRenderingSceneColor(RHICmdList, ESimpleRenderTargetMode::EExistingColorAndDepth, FExclusiveDepthStencil::DepthRead_StencilWrite, true, true);
 
-					// Render the light to the scene color buffer, conditionally using the attenuation buffer or a 1x1 white texture as input
+					// Render the light to the scene color buffer, conditionally using the attenuation buffer or a 1x1 white texture as input 
 					if (bDirectLighting)
 					{
-						RenderLight(RHICmdList, &LightSceneInfo, ScreenShadowMaskTexture, false, true);
+						RenderLight(RHICmdList, &LightSceneInfo, ScreenShadowMaskTexture, ScreenShadowMaskSubPixelTexture, &TransmittanceMaskData, InHairVisibilityViews, false, true);
 					}
 
 					SceneContext.FinishRenderingSceneColor(RHICmdList);
@@ -1405,7 +1666,7 @@ void FDeferredShadingSceneRenderer::RenderLightArrayForOverlapViewmode(FRHIComma
 			&& !LightSceneInfo->Proxy->HasStaticLighting()
 			&& LightSceneInfo->Proxy->CastsStaticShadow())
 		{
-			RenderLight(RHICmdList, LightSceneInfo, NULL, true, false);
+			RenderLight(RHICmdList, LightSceneInfo, nullptr, nullptr, nullptr, nullptr, true, false);
 		}
 	}
 }
@@ -1471,6 +1732,7 @@ static void SetShaderTemplLightingSimple(
 	PermutationVector.Set< FDeferredLightPS::FVisualizeCullingDim >( View.Family->EngineShowFlags.VisualizeLightCulling );
 	PermutationVector.Set< FDeferredLightPS::FLightingChannelsDim >( false );
 	PermutationVector.Set< FDeferredLightPS::FTransmissionDim >( false );
+	PermutationVector.Set< FDeferredLightPS::FHairLighting>( false );
 
 	TShaderMapRef< FDeferredLightPS > PixelShader( View.ShaderMap, PermutationVector );
 	GraphicsPSOInit.BoundShaderState.VertexDeclarationRHI = GetVertexDeclarationFVector4();
@@ -1518,7 +1780,7 @@ void CalculateLightNearFarDepthFromBounds(const FViewInfo& View, const FSphere &
  * @param LightIndex The light's index into FScene::Lights
  * @return true if anything got rendered
  */
-void FDeferredShadingSceneRenderer::RenderLight(FRHICommandList& RHICmdList, const FLightSceneInfo* LightSceneInfo, IPooledRenderTarget* ScreenShadowMaskTexture, bool bRenderOverlap, bool bIssueDrawEvent)
+void FDeferredShadingSceneRenderer::RenderLight(FRHICommandList& RHICmdList, const FLightSceneInfo* LightSceneInfo, IPooledRenderTarget* ScreenShadowMaskTexture, IPooledRenderTarget* ScreenShadowMaskSubPixelTexture, FHairStrandsTransmittanceMaskData* InTransmittanceMaskData, const FHairStrandsVisibilityViews* InHairVisibilityViews, bool bRenderOverlap, bool bIssueDrawEvent)
 {
 	SCOPE_CYCLE_COUNTER(STAT_DirectLightRenderingTime);
 	INC_DWORD_STAT(STAT_NumLightsUsingStandardDeferred);
@@ -1527,8 +1789,19 @@ void FDeferredShadingSceneRenderer::RenderLight(FRHICommandList& RHICmdList, con
 	FGraphicsPipelineStateInitializer GraphicsPSOInit;
 	RHICmdList.ApplyCachedRenderTargets(GraphicsPSOInit);
 
-	// Use additive blending for color
-	GraphicsPSOInit.BlendState = TStaticBlendState<CW_RGBA, BO_Add, BF_One, BF_One, BO_Add, BF_One, BF_One>::GetRHI();
+	const bool bHairRenderingEnabled = ScreenShadowMaskSubPixelTexture && InTransmittanceMaskData && InHairVisibilityViews && (LightSceneInfo->Proxy->CastsHairStrandsDeepShadow() || IsHairStrandsVoxelizationEnable());
+	if (bHairRenderingEnabled)
+	{
+		GraphicsPSOInit.BlendState = TStaticBlendState<
+			CW_RGBA, BO_Add, BF_One, BF_One, BO_Add, BF_One, BF_One,
+			CW_RGBA, BO_Add, BF_One, BF_One, BO_Max, BF_SourceAlpha, BF_DestAlpha
+		>::GetRHI();
+	}
+	else
+	{
+		GraphicsPSOInit.BlendState = TStaticBlendState<CW_RGBA, BO_Add, BF_One, BF_One, BO_Add, BF_One, BF_One>::GetRHI();
+	}
+
 	GraphicsPSOInit.PrimitiveType = PT_TriangleList;
 
 	const FSphere LightBounds = LightSceneInfo->Proxy->GetBoundingSphere();
@@ -1554,6 +1827,21 @@ void FDeferredShadingSceneRenderer::RenderLight(FRHICommandList& RHICmdList, con
 		// Set the device viewport for the view.
 		RHICmdList.SetViewport(View.ViewRect.Min.X, View.ViewRect.Min.Y, 0.0f, View.ViewRect.Max.X, View.ViewRect.Max.Y, 1.0f);
 
+		FRenderLightParams RenderLightParams;
+		if (bHairRenderingEnabled)
+		{
+			RenderLightParams.bUseTransmittanceData = bHairRenderingEnabled;
+			RenderLightParams.DeepShadow_TransmittanceMaskBuffer = InTransmittanceMaskData ? InTransmittanceMaskData->TransmittanceMaskSRV : nullptr;
+			RenderLightParams.DeepShadow_TransmittanceMaskBufferMaxCount = InTransmittanceMaskData && InTransmittanceMaskData->TransmittanceMask ? InTransmittanceMaskData->TransmittanceMask->Desc.NumElements : 0;
+			RenderLightParams.ScreenShadowMaskSubPixelTexture = ScreenShadowMaskSubPixelTexture;
+			if (InHairVisibilityViews && ViewIndex < InHairVisibilityViews->HairDatas.Num())
+			{
+				const FHairStrandsVisibilityData& HairVisibilityData = InHairVisibilityViews->HairDatas[ViewIndex];
+				RenderLightParams.HairVisibilityNodeOffsetAndCount = HairVisibilityData.NodeIndex;
+				RenderLightParams.HairVisibilityNodeDataSRV = HairVisibilityData.NodeDataSRV;
+			}
+		}
+
 		if (LightSceneInfo->Proxy->GetLightType() == LightType_Directional)
 		{
 			// Turn DBT back off
@@ -1581,6 +1869,7 @@ void FDeferredShadingSceneRenderer::RenderLight(FRHICommandList& RHICmdList, con
 				PermutationVector.Set< FDeferredLightPS::FVisualizeCullingDim >( View.Family->EngineShowFlags.VisualizeLightCulling );
 				PermutationVector.Set< FDeferredLightPS::FLightingChannelsDim >( View.bUsesLightingChannels );
 				PermutationVector.Set< FDeferredLightPS::FTransmissionDim >( bTransmission );
+				PermutationVector.Set< FDeferredLightPS::FHairLighting>(bHairRenderingEnabled);
 
 				TShaderMapRef< FDeferredLightPS > PixelShader( View.ShaderMap, PermutationVector );
 				GraphicsPSOInit.BoundShaderState.VertexDeclarationRHI = GFilterVertexDeclaration.VertexDeclarationRHI;
@@ -1588,7 +1877,7 @@ void FDeferredShadingSceneRenderer::RenderLight(FRHICommandList& RHICmdList, con
 				GraphicsPSOInit.BoundShaderState.PixelShaderRHI = GETSAFERHISHADER_PIXEL(*PixelShader);
 
 				SetGraphicsPipelineState(RHICmdList, GraphicsPSOInit);
-				PixelShader->SetParameters(RHICmdList, View, LightSceneInfo, ScreenShadowMaskTexture);
+				PixelShader->SetParameters(RHICmdList, View, LightSceneInfo, ScreenShadowMaskTexture, bHairRenderingEnabled ? &RenderLightParams : nullptr);
 			}
 
 			VertexShader->SetParameters(RHICmdList, View, LightSceneInfo);
@@ -1608,7 +1897,8 @@ void FDeferredShadingSceneRenderer::RenderLight(FRHICommandList& RHICmdList, con
 		else
 		{
 			// Use DBT to allow work culling on shadow lights
-			GraphicsPSOInit.bDepthBounds = (GSupportsDepthBoundsTest && GAllowDepthBoundsTest != 0);
+			// Disable depth bound when hair rendering is enabled as this rejects partially covered pixel write (with opaque background)
+			GraphicsPSOInit.bDepthBounds = !bHairRenderingEnabled && (GSupportsDepthBoundsTest && GAllowDepthBoundsTest != 0);
 
 			TShaderMapRef<TDeferredLightVS<true> > VertexShader(View.ShaderMap);
 
@@ -1634,6 +1924,7 @@ void FDeferredShadingSceneRenderer::RenderLight(FRHICommandList& RHICmdList, con
 				PermutationVector.Set< FDeferredLightPS::FVisualizeCullingDim >( View.Family->EngineShowFlags.VisualizeLightCulling );
 				PermutationVector.Set< FDeferredLightPS::FLightingChannelsDim >( View.bUsesLightingChannels );
 				PermutationVector.Set< FDeferredLightPS::FTransmissionDim >( bTransmission );
+				PermutationVector.Set< FDeferredLightPS::FHairLighting>(bHairRenderingEnabled);
 
 				TShaderMapRef< FDeferredLightPS > PixelShader( View.ShaderMap, PermutationVector );
 				GraphicsPSOInit.BoundShaderState.VertexDeclarationRHI = GetVertexDeclarationFVector4();
@@ -1641,13 +1932,13 @@ void FDeferredShadingSceneRenderer::RenderLight(FRHICommandList& RHICmdList, con
 				GraphicsPSOInit.BoundShaderState.PixelShaderRHI = GETSAFERHISHADER_PIXEL(*PixelShader);
 
 				SetGraphicsPipelineState(RHICmdList, GraphicsPSOInit);
-				PixelShader->SetParameters(RHICmdList, View, LightSceneInfo, ScreenShadowMaskTexture);
+				PixelShader->SetParameters(RHICmdList, View, LightSceneInfo, ScreenShadowMaskTexture, &RenderLightParams);
 			}
 
 			VertexShader->SetParameters(RHICmdList, View, LightSceneInfo);
 
 			// Use DBT to allow work culling on shadow lights
-			if (GSupportsDepthBoundsTest && GAllowDepthBoundsTest != 0)
+			if (GraphicsPSOInit.bDepthBounds)
 			{
 				// Can use the depth bounds test to skip work for pixels which won't be touched by the light (i.e outside the depth range)
 				float NearDepth = 1.f;
diff --git a/Engine/Source/Runtime/Renderer/Private/PostProcess/SceneRenderTargets.cpp b/Engine/Source/Runtime/Renderer/Private/PostProcess/SceneRenderTargets.cpp
--- a/Engine/Source/Runtime/Renderer/Private/PostProcess/SceneRenderTargets.cpp
+++ b/Engine/Source/Runtime/Renderer/Private/PostProcess/SceneRenderTargets.cpp
@@ -27,6 +27,7 @@
 #include "VT/VirtualTextureSystem.h"
 #include "VT/VirtualTextureFeedback.h"
 #include "VisualizeTexture.h"
+#include "GpuDebugRendering.h"
 
 static TAutoConsoleVariable<int32> CVarRSMResolution(
 	TEXT("r.LPV.RSMResolution"),
@@ -231,6 +232,7 @@ static void SnapshotArray(TArray<TRefCountPtr<IPooledRenderTarget>, TInlineAlloc
 FSceneRenderTargets::FSceneRenderTargets(const FViewInfo& View, const FSceneRenderTargets& SnapshotSource)
 	: LightAttenuation(GRenderTargetPool.MakeSnapshot(SnapshotSource.LightAttenuation))
 	, LightAccumulation(GRenderTargetPool.MakeSnapshot(SnapshotSource.LightAccumulation))
+	, SceneColorSubPixel(GRenderTargetPool.MakeSnapshot(SnapshotSource.SceneColorSubPixel))
 	, DirectionalOcclusion(GRenderTargetPool.MakeSnapshot(SnapshotSource.DirectionalOcclusion))
 	, SceneDepthZ(GRenderTargetPool.MakeSnapshot(SnapshotSource.SceneDepthZ))
 	, SceneVelocity(GRenderTargetPool.MakeSnapshot(SnapshotSource.SceneVelocity))
@@ -671,7 +673,7 @@ void FSceneRenderTargets::Allocate(FRHICommandListImmediate& RHICmdList, const F
 	AllocateRenderTargets(RHICmdList, ViewFamily.Views.Num());
 }
 
-void FSceneRenderTargets::BeginRenderingSceneColor(FRHICommandList& RHICmdList, ESimpleRenderTargetMode RenderTargetMode/*=EUninitializedColorExistingDepth*/, FExclusiveDepthStencil DepthStencilAccess, bool bTransitionWritable)
+void FSceneRenderTargets::BeginRenderingSceneColor(FRHICommandList& RHICmdList, ESimpleRenderTargetMode RenderTargetMode/*=EUninitializedColorExistingDepth*/, FExclusiveDepthStencil DepthStencilAccess, bool bTransitionWritable, bool bBindSubPixelRT)
 {
 	check(RHICmdList.IsOutsideRenderPass());
 
@@ -687,6 +689,12 @@ void FSceneRenderTargets::BeginRenderingSceneColor(FRHICommandList& RHICmdList,
 	RPInfo.DepthStencilRenderTarget.DepthStencilTarget = GetSceneDepthSurface();
 	RPInfo.DepthStencilRenderTarget.ExclusiveDepthStencil = DepthStencilAccess;
 
+	if (bBindSubPixelRT)
+	{
+		RPInfo.ColorRenderTargets[1].Action = MakeRenderTargetActions(ERenderTargetLoadAction::ELoad, ERenderTargetStoreAction::EStore);
+		RPInfo.ColorRenderTargets[1].RenderTarget = SceneColorSubPixel->GetRenderTargetItem().TargetableTexture;
+	}
+
 	if (bTransitionWritable)
 	{
 		TransitionRenderPassTargets(RHICmdList, RPInfo);
@@ -1298,6 +1306,17 @@ TRefCountPtr<IPooledRenderTarget>& FSceneRenderTargets::GetSceneColor()
 	return GetSceneColorForCurrentShadingPath();
 }
 
+TRefCountPtr<IPooledRenderTarget>& FSceneRenderTargets::GetSceneColorSubPixel()
+{
+	if (!GetSceneColorForCurrentShadingPath())
+	{
+		return GSystemTextures.BlackDummy;
+	}
+
+	return SceneColorSubPixel;
+}
+
+
 void FSceneRenderTargets::SetSceneColor(IPooledRenderTarget* In)
 {
 	check(CurrentShadingPath < EShadingPath::Num);
@@ -2268,6 +2287,13 @@ void FSceneRenderTargets::AllocateScreenShadowMask(FRHICommandList& RHICmdList,
 	GRenderTargetPool.FindFreeElement(RHICmdList, Desc, ScreenShadowMaskTexture, TEXT("ScreenShadowMaskTexture"));
 }
 
+void FSceneRenderTargets::AllocateScreenTransmittanceMask(FRHICommandList& RHICmdList, TRefCountPtr<IPooledRenderTarget>& ScreenTransmittanceMaskTexture)
+{
+	FPooledRenderTargetDesc Desc(FPooledRenderTargetDesc::Create2DDesc(GetBufferSizeXY(), PF_FloatRGBA, FClearValueBinding::White, TexCreate_None, TexCreate_RenderTargetable | TexCreate_UAV, false));
+	Desc.NumSamples = GetNumSceneColorMSAASamples(GetCurrentFeatureLevel());
+	GRenderTargetPool.FindFreeElement(RHICmdList, Desc, ScreenTransmittanceMaskTexture, TEXT("ScreenTransmittanceMaskTexture"));
+}
+
 const FTexture2DRHIRef& FSceneRenderTargets::GetOptionalShadowDepthColorSurface(FRHICommandList& RHICmdList, int32 Width, int32 Height) const
 {
 	// Look for matching resolution
@@ -2435,6 +2461,12 @@ void FSceneRenderTargets::AllocateDeferredShadingPathRenderTargets(FRHICommandLi
 		GRenderTargetPool.FindFreeElement(RHICmdList, Desc, LightAccumulation, TEXT("LightAccumulation"), true, ERenderTargetTransience::NonTransient);
 	}
 
+	if (CurrentFeatureLevel >= ERHIFeatureLevel::SM5)
+	{
+		FPooledRenderTargetDesc Desc(FPooledRenderTargetDesc::Create2DDesc(BufferSize, PF_FloatRGBA,  FClearValueBinding::Transparent, TexCreate_None, TexCreate_RenderTargetable, false));
+		GRenderTargetPool.FindFreeElement(RHICmdList, Desc, SceneColorSubPixel, TEXT("SceneColorSubPixel"), true);
+	}
+
 	AllocateDebugViewModeTargets(RHICmdList);
 
 	if (bAllocateVelocityGBuffer)
@@ -2679,6 +2711,7 @@ void FSceneRenderTargets::ReleaseAllTargets()
 	QuadOverdrawBuffer.SafeRelease();
 	LightAttenuation.SafeRelease();
 	LightAccumulation.SafeRelease();
+	SceneColorSubPixel.SafeRelease();
 	DirectionalOcclusion.SafeRelease();
 	CustomDepth.SafeRelease();
 	MobileCustomStencil.SafeRelease();
diff --git a/Engine/Source/Runtime/Renderer/Private/PostProcess/SceneRenderTargets.h b/Engine/Source/Runtime/Renderer/Private/PostProcess/SceneRenderTargets.h
--- a/Engine/Source/Runtime/Renderer/Private/PostProcess/SceneRenderTargets.h
+++ b/Engine/Source/Runtime/Renderer/Private/PostProcess/SceneRenderTargets.h
@@ -232,7 +232,7 @@ public:
 	/**
 	 * Sets the scene color target and restores its contents if necessary
 	 */
-	void BeginRenderingSceneColor(FRHICommandList& FRHICommandListImmediate, ESimpleRenderTargetMode RenderTargetMode = ESimpleRenderTargetMode::EUninitializedColorExistingDepth, FExclusiveDepthStencil DepthStencilAccess = FExclusiveDepthStencil::DepthWrite_StencilWrite, bool bTransitionWritable = true);
+	void BeginRenderingSceneColor(FRHICommandList& FRHICommandListImmediate, ESimpleRenderTargetMode RenderTargetMode = ESimpleRenderTargetMode::EUninitializedColorExistingDepth, FExclusiveDepthStencil DepthStencilAccess = FExclusiveDepthStencil::DepthWrite_StencilWrite, bool bTransitionWritable = true, bool bBindSubPixelRT = false);
 	void FinishRenderingSceneColor(FRHICommandList& RHICmdList);
 
 	// @return true: call FinishRenderingCustomDepth after rendering, false: don't render it, feature is disabled
@@ -456,6 +456,7 @@ public:
 	const TRefCountPtr<IPooledRenderTarget>& GetSceneColor() const;
 
 	TRefCountPtr<IPooledRenderTarget>& GetSceneColor();
+	TRefCountPtr<IPooledRenderTarget>& GetSceneColorSubPixel();
 
 	EPixelFormat GetSceneColorFormat(ERHIFeatureLevel::Type InFeatureLevel) const;
 	EPixelFormat GetSceneColorFormat() const;
@@ -496,6 +497,7 @@ public:
 	void AllocateDebugViewModeTargets(FRHICommandList& RHICmdList);
 
 	void AllocateScreenShadowMask(FRHICommandList& RHICmdList, TRefCountPtr<IPooledRenderTarget>& ScreenShadowMaskTexture);
+	void AllocateScreenTransmittanceMask(FRHICommandList& RHICmdList, TRefCountPtr<IPooledRenderTarget>& ScreenTransmittanceMaskTexture);
 
 	TRefCountPtr<IPooledRenderTarget>& GetReflectionBrightnessTarget();
 
@@ -518,6 +520,9 @@ public:
 	// Light Accumulation is a high precision scratch pad matching the size of the scene color buffer used by many passes.
 	TRefCountPtr<IPooledRenderTarget> LightAccumulation;
 
+	// Secondary scene color target for storing subpixel color details (e.g. hair pixels)
+	TRefCountPtr<IPooledRenderTarget> SceneColorSubPixel;
+
 	// Reflection Environment: Bringing back light accumulation buffer to apply indirect reflections
 	TRefCountPtr<IPooledRenderTarget> DirectionalOcclusion;
 	
diff --git a/Engine/Source/Runtime/Renderer/Private/PrimitiveSceneInfo.cpp b/Engine/Source/Runtime/Renderer/Private/PrimitiveSceneInfo.cpp
--- a/Engine/Source/Runtime/Renderer/Private/PrimitiveSceneInfo.cpp
+++ b/Engine/Source/Runtime/Renderer/Private/PrimitiveSceneInfo.cpp
@@ -85,7 +85,8 @@ public:
 				ScreenSize, 
 				bSupportsCachingMeshDrawCommands,
 				bUseSkyMaterial,
-				bUseSingleLayerWaterMaterial
+				bUseSingleLayerWaterMaterial,
+				FeatureLevel
 			);
 		}
 	}
diff --git a/Engine/Source/Runtime/Renderer/Private/RayTracing/RayTracingShadows.cpp b/Engine/Source/Runtime/Renderer/Private/RayTracing/RayTracingShadows.cpp
--- a/Engine/Source/Runtime/Renderer/Private/RayTracing/RayTracingShadows.cpp
+++ b/Engine/Source/Runtime/Renderer/Private/RayTracing/RayTracingShadows.cpp
@@ -54,8 +54,9 @@ class FOcclusionRGS : public FGlobalShader
 	class FLightTypeDim : SHADER_PERMUTATION_INT("LIGHT_TYPE", LightType_MAX);
 	class FDenoiserOutputDim : SHADER_PERMUTATION_INT("DIM_DENOISER_OUTPUT", 4);
 	class FEnableTwoSidedGeometryDim : SHADER_PERMUTATION_BOOL("ENABLE_TWO_SIDED_GEOMETRY");
+	class FHairLighting : SHADER_PERMUTATION_INT("USE_HAIR_LIGHTING", 2);
 
-	using FPermutationDomain = TShaderPermutationDomain<FLightTypeDim, FDenoiserOutputDim, FEnableTwoSidedGeometryDim>;
+	using FPermutationDomain = TShaderPermutationDomain<FLightTypeDim, FDenoiserOutputDim, FEnableTwoSidedGeometryDim, FHairLighting>;
 
 	static bool ShouldCompilePermutation(const FGlobalShaderPermutationParameters& Parameters)
 	{
@@ -66,12 +67,14 @@ class FOcclusionRGS : public FGlobalShader
 		SHADER_PARAMETER(uint32, SamplesPerPixel)
 		SHADER_PARAMETER(float, NormalBias)
 		SHADER_PARAMETER(uint32, LightingChannelMask)
+		SHADER_PARAMETER(uint32, ShadowMaskType)
 		SHADER_PARAMETER(FIntRect, LightScissor)
 		SHADER_PARAMETER(FIntPoint, PixelOffset)
 
 		SHADER_PARAMETER_STRUCT(FLightShaderParameters, Light)
 		SHADER_PARAMETER_STRUCT_INCLUDE(FSceneTextureParameters, SceneTextures)
 
+		SHADER_PARAMETER_RDG_TEXTURE(Texture2D, HairCategorizationTexture)
 		SHADER_PARAMETER_SRV(RaytracingAccelerationStructure, TLAS)
 		SHADER_PARAMETER_RDG_TEXTURE_UAV(RWTexture2D<float4>, RWOcclusionMaskUAV)
 		SHADER_PARAMETER_RDG_TEXTURE_UAV(RWTexture2D<float>, RWRayDistanceUAV)
@@ -98,6 +101,7 @@ void FDeferredShadingSceneRenderer::PrepareRayTracingShadows(const FViewInfo& Vi
 		IScreenSpaceDenoiser::EShadowRequirements::PenumbraAndClosestOccluder,
 	};
 
+	for (int32 HairLighting = 0; HairLighting < 2; ++HairLighting)
 	for (int32 LightType = 0; LightType < LightType_MAX; ++LightType)
 	{
 		for (IScreenSpaceDenoiser::EShadowRequirements DenoiserRequirement : DenoiserRequirements)
@@ -106,6 +110,7 @@ void FDeferredShadingSceneRenderer::PrepareRayTracingShadows(const FViewInfo& Vi
 			PermutationVector.Set<FOcclusionRGS::FLightTypeDim>(LightType);
 			PermutationVector.Set<FOcclusionRGS::FDenoiserOutputDim>((int32)DenoiserRequirement);
 			PermutationVector.Set<FOcclusionRGS::FEnableTwoSidedGeometryDim>(CVarRayTracingShadowsEnableTwoSidedGeometry.GetValueOnRenderThread() != 0);
+			PermutationVector.Set<FOcclusionRGS::FHairLighting>(HairLighting);
 
 			TShaderMapRef<FOcclusionRGS> RayGenerationShader(View.ShaderMap, PermutationVector);
 			OutRayGenShaders.Add(RayGenerationShader->GetRayTracingShader());
@@ -121,7 +126,9 @@ void FDeferredShadingSceneRenderer::RenderRayTracingShadows(
 	const FViewInfo& View,
 	const FLightSceneInfo& LightSceneInfo,
 	const IScreenSpaceDenoiser::FShadowRayTracingConfig& RayTracingConfig,
-	IScreenSpaceDenoiser::EShadowRequirements DenoiserRequirements,
+	const IScreenSpaceDenoiser::EShadowRequirements DenoiserRequirements,
+	const bool bSubPixelShadowMask,
+	FRDGTextureRef HairCategorizationTexture,
 	FRDGTextureRef* OutShadowMask,
 	FRDGTextureRef* OutRayHitDistance)
 #if RHI_RAYTRACING
@@ -178,6 +185,7 @@ void FDeferredShadingSceneRenderer::RenderRayTracingShadows(
 
 	// Ray generation pass for shadow occlusion.
 	{
+		const bool bUseHairLighting = HairCategorizationTexture != nullptr;
 		FOcclusionRGS::FParameters* PassParameters = GraphBuilder.AllocParameters<FOcclusionRGS::FParameters>();
 		PassParameters->RWOcclusionMaskUAV = GraphBuilder.CreateUAV(FRDGTextureUAVDesc(ScreenShadowMaskTexture));
 		PassParameters->RWRayDistanceUAV = GraphBuilder.CreateUAV(FRDGTextureUAVDesc(RayDistanceTexture));
@@ -190,6 +198,12 @@ void FDeferredShadingSceneRenderer::RenderRayTracingShadows(
 		PassParameters->SceneTextures = SceneTextures;
 		PassParameters->LightScissor = ScissorRect;
 		PassParameters->PixelOffset = PixelOffset;
+		PassParameters->ShadowMaskType = bUseHairLighting ? (bSubPixelShadowMask ? 1 : 0) : 0;
+		PassParameters->HairCategorizationTexture = HairCategorizationTexture;
+		if (!PassParameters->HairCategorizationTexture)
+		{
+			PassParameters->HairCategorizationTexture = GraphBuilder.RegisterExternalTexture(GSystemTextures.BlackDummy);
+		}
 		
 		FOcclusionRGS::FPermutationDomain PermutationVector;
 		PermutationVector.Set<FOcclusionRGS::FLightTypeDim>(LightSceneProxy->GetLightType());
@@ -211,6 +225,7 @@ void FDeferredShadingSceneRenderer::RenderRayTracingShadows(
 			PermutationVector.Set<FOcclusionRGS::FDenoiserOutputDim>(0);
 		}
 		PermutationVector.Set<FOcclusionRGS::FEnableTwoSidedGeometryDim>(CVarRayTracingShadowsEnableTwoSidedGeometry.GetValueOnRenderThread() != 0);
+		PermutationVector.Set<FOcclusionRGS::FHairLighting>(bUseHairLighting? 1 : 0);
 
 		TShaderMapRef<FOcclusionRGS> RayGenerationShader(GetGlobalShaderMap(FeatureLevel), PermutationVector);
 
diff --git a/Engine/Source/Runtime/Renderer/Private/ReflectionEnvironment.cpp b/Engine/Source/Runtime/Renderer/Private/ReflectionEnvironment.cpp
--- a/Engine/Source/Runtime/Renderer/Private/ReflectionEnvironment.cpp
+++ b/Engine/Source/Runtime/Renderer/Private/ReflectionEnvironment.cpp
@@ -35,6 +35,8 @@
 #include "RayTracing/RaytracingOptions.h"
 #include "RenderGraph.h"
 #include "PixelShaderUtils.h"
+#include "SceneTextureParameters.h"
+#include "HairStrands/HairStrandsRendering.h"
 
 DECLARE_GPU_STAT_NAMED(ReflectionEnvironment, TEXT("Reflection Environment"));
 DECLARE_GPU_STAT_NAMED(RayTracingReflections, TEXT("Ray Tracing Reflections"));
@@ -587,7 +589,7 @@ bool FDeferredShadingSceneRenderer::ShouldDoReflectionEnvironment() const
 		&& ViewFamily.EngineShowFlags.ReflectionEnvironment;
 }
 
-void FDeferredShadingSceneRenderer::RenderDeferredReflectionsAndSkyLighting(FRHICommandListImmediate& RHICmdList, TRefCountPtr<IPooledRenderTarget>& DynamicBentNormalAO, TRefCountPtr<IPooledRenderTarget>& VelocityRT)
+void FDeferredShadingSceneRenderer::RenderDeferredReflectionsAndSkyLighting(FRHICommandListImmediate& RHICmdList, TRefCountPtr<IPooledRenderTarget>& DynamicBentNormalAO, TRefCountPtr<IPooledRenderTarget>& VelocityRT, const FHairStrandsDatas* HairDatas)
 {
 	check(RHICmdList.IsOutsideRenderPass());
 
@@ -645,16 +647,22 @@ void FDeferredShadingSceneRenderer::RenderDeferredReflectionsAndSkyLighting(FRHI
 	FRDGBuilder GraphBuilder(RHICmdList);
 
 	FRDGTextureRef SceneColorTexture = GraphBuilder.RegisterExternalTexture(SceneContext.GetSceneColor());
+	FRDGTextureRef SceneColorSubPixelTexture = nullptr;
+	if (HairDatas)
+	{
+		SceneColorSubPixelTexture = GraphBuilder.RegisterExternalTexture(SceneContext.GetSceneColorSubPixel());
+	}
 	FRDGTextureRef AmbientOcclusionTexture = GraphBuilder.RegisterExternalTexture(SceneContext.bScreenSpaceAOIsValid ? SceneContext.ScreenSpaceAO : GSystemTextures.WhiteDummy);
 	FRDGTextureRef DynamicBentNormalAOTexture = GraphBuilder.RegisterExternalTexture(DynamicBentNormalAO ? DynamicBentNormalAO : GSystemTextures.WhiteDummy);
 
 	FSceneTextureParameters SceneTextures;
 	SetupSceneTextureParameters(GraphBuilder, &SceneTextures);
 
+	uint32 ViewIndex = 0;
 	for (FViewInfo& View : Views)
 	{
+		const uint32 CurrentViewIndex = ViewIndex++;
 		const bool bRayTracedReflections = ShouldRenderRayTracingReflections(View);
-
 		const bool bScreenSpaceReflections = !bRayTracedReflections && ShouldRenderScreenSpaceReflections(View);
 
 		const bool bComposePlanarReflections = !bRayTracedReflections && HasDeferredPlanarReflections(View);
@@ -857,6 +865,8 @@ void FDeferredShadingSceneRenderer::RenderDeferredReflectionsAndSkyLighting(FRHI
 
 			PassParameters->RenderTargets[0] = FRenderTargetBinding(SceneColorTexture, ERenderTargetLoadAction::ELoad);
 
+			// Bind hair data
+			const bool bCheckerboardSubsurfaceRendering = IsSubsurfaceCheckerboardFormat(SceneContext.GetSceneColorFormat());
 			auto PermutationVector = FReflectionEnvironmentSkyLightingPS::BuildPermutationVector(
 				View, bHasBoxCaptures, bHasSphereCaptures, DynamicBentNormalAO != NULL, bSkyLight, bDynamicSkyLight, bApplySkyShadowing, bRayTracedReflections);
 
@@ -867,7 +877,7 @@ void FDeferredShadingSceneRenderer::RenderDeferredReflectionsAndSkyLighting(FRHI
 				RDG_EVENT_NAME("ReflectionEnvironmentAndSky %dx%d", View.ViewRect.Width(), View.ViewRect.Height()),
 				PassParameters,
 				ERDGPassFlags::Raster,
-				[PassParameters, &View, PixelShader](FRHICommandList& InRHICmdList)
+				[PassParameters, &View, PixelShader, bCheckerboardSubsurfaceRendering](FRHICommandList& InRHICmdList)
 			{
 				InRHICmdList.SetViewport(View.ViewRect.Min.X, View.ViewRect.Min.Y, 0.0f, View.ViewRect.Max.X, View.ViewRect.Max.Y, 1.0f);
 
@@ -882,7 +892,6 @@ void FDeferredShadingSceneRenderer::RenderDeferredReflectionsAndSkyLighting(FRHI
 				}
 				else
 				{
-					const bool bCheckerboardSubsurfaceRendering = IsSubsurfaceCheckerboardFormat(PassParameters->RenderTargets[0].GetTexture()->Desc.Format);
 					if (bCheckerboardSubsurfaceRendering)
 					{
 						GraphicsPSOInit.BlendState = TStaticBlendState<CW_RGB, BO_Add, BF_One, BF_One>::GetRHI();
@@ -898,11 +907,23 @@ void FDeferredShadingSceneRenderer::RenderDeferredReflectionsAndSkyLighting(FRHI
 				FPixelShaderUtils::DrawFullscreenTriangle(InRHICmdList);
 			});
 		} // if (bRequiresApply)
+
+
+		if (HairDatas)
+		{
+			RenderHairStrandsEnvironmentLighting(GraphBuilder, CurrentViewIndex, Views, HairDatas, SceneColorTexture, SceneColorSubPixelTexture);
+		}
 	} // for (FViewInfo& View : Views)
 	
 	TRefCountPtr<IPooledRenderTarget> OutSceneColor;
 	GraphBuilder.QueueTextureExtraction(SceneColorTexture, &OutSceneColor);
 
+	TRefCountPtr<IPooledRenderTarget> OutSceneSubpixelColor;
+	if (SceneColorSubPixelTexture)
+	{
+		GraphBuilder.QueueTextureExtraction(SceneColorSubPixelTexture, &OutSceneSubpixelColor);
+	}
+
 	GraphBuilder.Execute();
 	
 	ResolveSceneColor(RHICmdList);
diff --git a/Engine/Source/Runtime/Renderer/Private/SceneCore.h b/Engine/Source/Runtime/Renderer/Private/SceneCore.h
--- a/Engine/Source/Runtime/Renderer/Private/SceneCore.h
+++ b/Engine/Source/Runtime/Renderer/Private/SceneCore.h
@@ -168,7 +168,7 @@ private:
 class FStaticMeshBatchRelevance
 {
 public:
-	FStaticMeshBatchRelevance(const FStaticMeshBatch& StaticMesh, float InScreenSize, bool InbSupportsCachingMeshDrawCommands, bool InbUseSkyMaterial, bool bInUseSingleLayerWaterMaterial)
+	FStaticMeshBatchRelevance(const FStaticMeshBatch& StaticMesh, float InScreenSize, bool InbSupportsCachingMeshDrawCommands, bool InbUseSkyMaterial, bool bInUseSingleLayerWaterMaterial, ERHIFeatureLevel::Type FeatureLevel)
 		: Id(StaticMesh.Id)
 		, ScreenSize(InScreenSize)
 		, CommandInfosBase(0)
@@ -183,6 +183,7 @@ public:
 		, bUseAsOccluder(StaticMesh.bUseAsOccluder)
 		, bUseSkyMaterial(InbUseSkyMaterial)
 		, bUseSingleLayerWaterMaterial(bInUseSingleLayerWaterMaterial)
+		, bUseHairStrands(StaticMesh.UseForHairStrands(FeatureLevel))
 		, bRenderToVirtualTexture(StaticMesh.bRenderToVirtualTexture)
 		, RuntimeVirtualTextureMaterialType(StaticMesh.RuntimeVirtualTextureMaterialType)
 		, bSupportsCachingMeshDrawCommands(InbSupportsCachingMeshDrawCommands)
@@ -222,6 +223,7 @@ public:
 	uint8 bUseAsOccluder	: 1; // User hint whether it's a good occluder.
 	uint8 bUseSkyMaterial	: 1; // Whether this batch uses a Sky material or not.
 	uint8 bUseSingleLayerWaterMaterial : 1; // Whether this batch uses a water material or not.
+	uint8 bUseHairStrands	: 1; // Whether it contains hair strands geometry.
 
 	/** Whether the mesh batch can be used for rendering to a virtual texture. */
 	uint8 bRenderToVirtualTexture : 1;
diff --git a/Engine/Source/Runtime/Renderer/Private/SceneRendering.h b/Engine/Source/Runtime/Renderer/Private/SceneRendering.h
--- a/Engine/Source/Runtime/Renderer/Private/SceneRendering.h
+++ b/Engine/Source/Runtime/Renderer/Private/SceneRendering.h
@@ -962,6 +962,9 @@ public:
 	/* [PrimitiveIndex] = end index index in DynamicMeshElements[], to support GetDynamicMeshElementRange(). Contains valid values only for visible primitives with bDynamicRelevance. */
 	TArray<uint32, SceneRenderingAllocator> DynamicMeshEndIndices;
 
+	/** Hair strands dynamic mesh element. */
+	TArray<FMeshBatchAndRelevance, SceneRenderingAllocator> HairStrandsMeshElements;
+
 	/* Mesh pass relevance for gathered dynamic mesh elements. */
 	TArray<FMeshPassMask, SceneRenderingAllocator> DynamicMeshElementsPassRelevance;
 
@@ -1594,7 +1597,7 @@ protected:
 
 	void SetupMeshPass(FViewInfo& View, FExclusiveDepthStencil::Type BasePassDepthStencilAccess, FViewCommands& ViewCommands);
 
-	bool RenderShadowProjections(FRHICommandListImmediate& RHICmdList, const FLightSceneInfo* LightSceneInfo, IPooledRenderTarget* ScreenShadowMaskTexture, bool bProjectingForForwardShading, bool bMobileModulatedProjections);
+	bool RenderShadowProjections(FRHICommandListImmediate& RHICmdList, const FLightSceneInfo* LightSceneInfo, IPooledRenderTarget* ScreenShadowMaskTexture, IPooledRenderTarget* ScreenShadowMaskSubPixelTexture, bool bProjectingForForwardShading, bool bMobileModulatedProjections, const struct FHairStrandsVisibilityViews* InHairVisibilityViews);
 
 	/** Finds a matching cached preshadow, if one exists. */
 	TRefCountPtr<FProjectedShadowInfo> GetCachedPreshadow(
diff --git a/Engine/Source/Runtime/Renderer/Private/SceneVisibility.cpp b/Engine/Source/Runtime/Renderer/Private/SceneVisibility.cpp
--- a/Engine/Source/Runtime/Renderer/Private/SceneVisibility.cpp
+++ b/Engine/Source/Runtime/Renderer/Private/SceneVisibility.cpp
@@ -36,6 +36,7 @@
 #include "GPUScene.h"
 #include "TranslucentRendering.h"
 #include "Async/ParallelFor.h"
+#include "HairStrands/HairStrandsRendering.h"
 #include "RectLightSceneProxy.h"
 #include "Math/Halton.h"
 
@@ -1968,6 +1969,13 @@ struct FRelevancePacket
 			const bool bEditorSelectionRelevance = ViewRelevance.bEditorStaticSelectionRelevance;
 			const bool bTranslucentRelevance = ViewRelevance.HasTranslucency();
 
+			const bool bHairStrandsEnabled = ViewRelevance.bHairStrandsRelevance && IsHairStrandsEnable(Scene->GetShaderPlatform());
+			if (!bEditorRelevance && bHairStrandsEnabled)
+			{
+				++NumVisibleDynamicPrimitives;
+				OutHasDynamicMeshElementsMasks[BitIndex] |= ViewBit;
+			}
+
 			if (View.bIsReflectionCapture && !PrimitiveSceneInfo->Proxy->IsVisibleInReflectionCaptures())
 			{
 				NotDrawRelevant.AddPrim(BitIndex);
@@ -2043,7 +2051,7 @@ struct FRelevancePacket
 					bHasDistortionPrimitives = true;
 				}
 			}
-
+			
 			CombinedShadingModelMask |= ViewRelevance.ShadingModelMaskRelevance;
 			bUsesGlobalDistanceField |= ViewRelevance.bUsesGlobalDistanceField;
 			bUsesLightingChannels |= ViewRelevance.bUsesLightingChannels;
@@ -2768,6 +2776,14 @@ void ComputeDynamicMeshRelevance(EShadingPath ShadingPath, bool bAddLightmapDens
 		PassMask.Set(EMeshPass::EditorSelection);
 		View.NumVisibleDynamicMeshElements[EMeshPass::EditorSelection] += NumElements;
 	}
+
+	// Hair strands are not rendered into the base pass (bRenderInMainPass=0) and so this 
+	// adds a special pass for allowing hair strands to be selectable.
+	if (View.bAllowTranslucentPrimitivesInHitProxy && ViewRelevance.bHairStrandsRelevance)
+	{
+		PassMask.Set(EMeshPass::HitProxy);
+		View.NumVisibleDynamicMeshElements[EMeshPass::HitProxy] += NumElements;
+	}
 #endif
 
 	if (ViewRelevance.bHasVolumeMaterialDomain)
@@ -2786,6 +2802,14 @@ void ComputeDynamicMeshRelevance(EShadingPath ShadingPath, bool bAddLightmapDens
 		BatchAndProxy.Proxy = MeshBatch.PrimitiveSceneProxy;
 		BatchAndProxy.SortKey = MeshBatch.PrimitiveSceneProxy->GetTranslucencySortPriority();
 	}
+	
+	const bool bIsHairStrandsCompatible = ViewRelevance.bHairStrandsRelevance && IsHairStrandsEnable(View.GetShaderPlatform());
+	if (bIsHairStrandsCompatible)
+	{
+		View.HairStrandsMeshElements.AddUninitialized(1);
+		FMeshBatchAndRelevance& BatchAndProxy = View.HairStrandsMeshElements.Last();
+		BatchAndProxy = MeshBatch;
+	}
 }
 
 void FSceneRenderer::GatherDynamicMeshElements(
diff --git a/Engine/Source/Runtime/Renderer/Private/ShadowRendering.cpp b/Engine/Source/Runtime/Renderer/Private/ShadowRendering.cpp
--- a/Engine/Source/Runtime/Renderer/Private/ShadowRendering.cpp
+++ b/Engine/Source/Runtime/Renderer/Private/ShadowRendering.cpp
@@ -13,6 +13,7 @@
 #include "ScenePrivate.h"
 #include "PipelineStateCache.h"
 #include "ClearQuad.h"
+#include "HairStrands/HairStrandsRendering.h"
 
 ///////////////////////////////////////////////////////////////////////////////////////////////////
 // Directional light
@@ -225,37 +226,43 @@ IMPLEMENT_SHADER_TYPE(,FShadowVolumeBoundProjectionVS,TEXT("/Engine/Private/Shad
  * Implementations for TShadowProjectionPS.  
  */
 #if !UE_BUILD_DOCS
-#define IMPLEMENT_SHADOW_PROJECTION_PIXEL_SHADER(Quality,UseFadePlane,UseTransmission) \
-	typedef TShadowProjectionPS<Quality, UseFadePlane, false, UseTransmission> FShadowProjectionPS##Quality##UseFadePlane##UseTransmission; \
-	IMPLEMENT_SHADER_TYPE(template<>,FShadowProjectionPS##Quality##UseFadePlane##UseTransmission,TEXT("/Engine/Private/ShadowProjectionPixelShader.usf"),TEXT("Main"),SF_Pixel);
+#define IMPLEMENT_SHADOW_PROJECTION_PIXEL_SHADER(Quality,UseFadePlane,UseTransmission, SupportSubPixel) \
+	typedef TShadowProjectionPS<Quality, UseFadePlane, false, UseTransmission, SupportSubPixel> FShadowProjectionPS##Quality##UseFadePlane##UseTransmission##SupportSubPixel; \
+	IMPLEMENT_SHADER_TYPE(template<>,FShadowProjectionPS##Quality##UseFadePlane##UseTransmission##SupportSubPixel,TEXT("/Engine/Private/ShadowProjectionPixelShader.usf"),TEXT("Main"),SF_Pixel);
 
 // Projection shaders without the distance fade, with different quality levels.
-IMPLEMENT_SHADOW_PROJECTION_PIXEL_SHADER(1,false,false);
-IMPLEMENT_SHADOW_PROJECTION_PIXEL_SHADER(2,false,false);
-IMPLEMENT_SHADOW_PROJECTION_PIXEL_SHADER(3,false,false);
-IMPLEMENT_SHADOW_PROJECTION_PIXEL_SHADER(4,false,false);
-IMPLEMENT_SHADOW_PROJECTION_PIXEL_SHADER(5,false,false);
-
-IMPLEMENT_SHADOW_PROJECTION_PIXEL_SHADER(1,false,true);
-IMPLEMENT_SHADOW_PROJECTION_PIXEL_SHADER(2,false,true);
-IMPLEMENT_SHADOW_PROJECTION_PIXEL_SHADER(3,false,true);
-IMPLEMENT_SHADOW_PROJECTION_PIXEL_SHADER(4,false,true);
-IMPLEMENT_SHADOW_PROJECTION_PIXEL_SHADER(5,false,true);
+IMPLEMENT_SHADOW_PROJECTION_PIXEL_SHADER(1,false,false,false);
+IMPLEMENT_SHADOW_PROJECTION_PIXEL_SHADER(2,false,false,false);
+IMPLEMENT_SHADOW_PROJECTION_PIXEL_SHADER(3,false,false,false);
+IMPLEMENT_SHADOW_PROJECTION_PIXEL_SHADER(4,false,false,false);
+IMPLEMENT_SHADOW_PROJECTION_PIXEL_SHADER(5,false,false,false);
+
+IMPLEMENT_SHADOW_PROJECTION_PIXEL_SHADER(1,false,true,false);
+IMPLEMENT_SHADOW_PROJECTION_PIXEL_SHADER(2,false,true,false);
+IMPLEMENT_SHADOW_PROJECTION_PIXEL_SHADER(3,false,true,false);
+IMPLEMENT_SHADOW_PROJECTION_PIXEL_SHADER(4,false,true,false);
+IMPLEMENT_SHADOW_PROJECTION_PIXEL_SHADER(5,false,true,false);
 
 // Projection shaders with the distance fade, with different quality levels.
-IMPLEMENT_SHADOW_PROJECTION_PIXEL_SHADER(1,true,false);
-IMPLEMENT_SHADOW_PROJECTION_PIXEL_SHADER(2,true,false);
-IMPLEMENT_SHADOW_PROJECTION_PIXEL_SHADER(3,true,false);
-IMPLEMENT_SHADOW_PROJECTION_PIXEL_SHADER(4,true,false);
-IMPLEMENT_SHADOW_PROJECTION_PIXEL_SHADER(5,true,false);
-
-IMPLEMENT_SHADOW_PROJECTION_PIXEL_SHADER(1,true,true);
-IMPLEMENT_SHADOW_PROJECTION_PIXEL_SHADER(2,true,true);
-IMPLEMENT_SHADOW_PROJECTION_PIXEL_SHADER(3,true,true);
-IMPLEMENT_SHADOW_PROJECTION_PIXEL_SHADER(4,true,true);
-IMPLEMENT_SHADOW_PROJECTION_PIXEL_SHADER(5,true,true);
-
-
+IMPLEMENT_SHADOW_PROJECTION_PIXEL_SHADER(1,true,false,false);
+IMPLEMENT_SHADOW_PROJECTION_PIXEL_SHADER(2,true,false,false);
+IMPLEMENT_SHADOW_PROJECTION_PIXEL_SHADER(3,true,false,false);
+IMPLEMENT_SHADOW_PROJECTION_PIXEL_SHADER(4,true,false,false);
+IMPLEMENT_SHADOW_PROJECTION_PIXEL_SHADER(5,true,false,false);
+
+IMPLEMENT_SHADOW_PROJECTION_PIXEL_SHADER(1,true,true,false);
+IMPLEMENT_SHADOW_PROJECTION_PIXEL_SHADER(2,true,true,false);
+IMPLEMENT_SHADOW_PROJECTION_PIXEL_SHADER(3,true,true,false);
+IMPLEMENT_SHADOW_PROJECTION_PIXEL_SHADER(4,true,true,false);
+IMPLEMENT_SHADOW_PROJECTION_PIXEL_SHADER(5,true,true,false);
+
+// Projection shaders without the distance fade, without transmission, with Sub-PixelSupport with different quality levels
+IMPLEMENT_SHADOW_PROJECTION_PIXEL_SHADER(1, false, false, true);
+IMPLEMENT_SHADOW_PROJECTION_PIXEL_SHADER(2, false, false, true);
+IMPLEMENT_SHADOW_PROJECTION_PIXEL_SHADER(3, false, false, true);
+IMPLEMENT_SHADOW_PROJECTION_PIXEL_SHADER(4, false, false, true);
+IMPLEMENT_SHADOW_PROJECTION_PIXEL_SHADER(5, false, false, true);
+													   
 #undef IMPLEMENT_SHADOW_PROJECTION_PIXEL_SHADER
 #endif
 
@@ -307,12 +314,34 @@ IMPLEMENT_SHADOW_PROJECTION_PIXEL_SHADER(5, true);
 #undef IMPLEMENT_SHADOW_PROJECTION_PIXEL_SHADER
 
 static void GetShadowProjectionShaders(
-	int32 Quality, const FViewInfo& View, const FProjectedShadowInfo* ShadowInfo, bool bMobileModulatedProjections,
+	int32 Quality, const FViewInfo& View, const FProjectedShadowInfo* ShadowInfo, bool bMobileModulatedProjections, bool bSubPixelSupport,
 	FShadowProjectionVertexShaderInterface** OutShadowProjVS, FShadowProjectionPixelShaderInterface** OutShadowProjPS)
 {
 	check(!*OutShadowProjVS);
 	check(!*OutShadowProjPS);
 
+	if (bSubPixelSupport)
+	{
+		check(!bMobileModulatedProjections);
+
+		if (ShadowInfo->IsWholeSceneDirectionalShadow())
+			*OutShadowProjVS = View.ShaderMap->GetShader<FShadowProjectionNoTransformVS>();
+		else
+			*OutShadowProjVS = View.ShaderMap->GetShader<FShadowVolumeBoundProjectionVS>();
+
+		switch (Quality)
+		{
+		case 1: *OutShadowProjPS = View.ShaderMap->GetShader<TShadowProjectionPS<1, false, false, false, true> >(); break;
+		case 2: *OutShadowProjPS = View.ShaderMap->GetShader<TShadowProjectionPS<2, false, false, false, true> >(); break;
+		case 3: *OutShadowProjPS = View.ShaderMap->GetShader<TShadowProjectionPS<3, false, false, false, true> >(); break;
+		case 4: *OutShadowProjPS = View.ShaderMap->GetShader<TShadowProjectionPS<4, false, false, false, true> >(); break;
+		case 5: *OutShadowProjPS = View.ShaderMap->GetShader<TShadowProjectionPS<5, false, false, false, true> >(); break;
+		default:
+			check(0);
+		}
+		return;
+	}
+
 	if (ShadowInfo->bTranslucentShadow)
 	{
 		*OutShadowProjVS = View.ShaderMap->GetShader<FShadowVolumeBoundProjectionVS>();
@@ -837,7 +866,7 @@ void FProjectedShadowInfo::SetupProjectionStencilMask(
 	}
 }
 
-void FProjectedShadowInfo::RenderProjection(FRHICommandListImmediate& RHICmdList, int32 ViewIndex, const FViewInfo* View, const FSceneRenderer* SceneRender, bool bProjectingForForwardShading, bool bMobileModulatedProjections) const
+void FProjectedShadowInfo::RenderProjection(FRHICommandListImmediate& RHICmdList, int32 ViewIndex, const FViewInfo* View, const FSceneRenderer* SceneRender, bool bProjectingForForwardShading, bool bMobileModulatedProjections, const FHairStrandsVisibilityData* HairVisibilityData) const
 {
 #if WANTS_DRAW_MESH_EVENTS
 	FString EventName;
@@ -953,7 +982,8 @@ void FProjectedShadowInfo::RenderProjection(FRHICommandListImmediate& RHICmdList
 		FShadowProjectionVertexShaderInterface* ShadowProjVS = nullptr;
 		FShadowProjectionPixelShaderInterface* ShadowProjPS = nullptr;
 
-		GetShadowProjectionShaders(LocalQuality, *View, this, bMobileModulatedProjections, &ShadowProjVS, &ShadowProjPS);
+		const bool bSubPixelSupport = HairVisibilityData != nullptr;
+		GetShadowProjectionShaders(LocalQuality, *View, this, bMobileModulatedProjections, bSubPixelSupport, &ShadowProjVS, &ShadowProjPS);
 
 		GraphicsPSOInit.BoundShaderState.VertexDeclarationRHI = GetVertexDeclarationFVector4();
 		GraphicsPSOInit.BoundShaderState.VertexShaderRHI = GETSAFERHISHADER_VERTEX(ShadowProjVS);
@@ -969,7 +999,7 @@ void FProjectedShadowInfo::RenderProjection(FRHICommandListImmediate& RHICmdList
 		RHICmdList.SetStencilRef(0);
 
 		ShadowProjVS->SetParameters(RHICmdList, *View, this);
-		ShadowProjPS->SetParameters(RHICmdList, ViewIndex, *View, this);
+		ShadowProjPS->SetParameters(RHICmdList, ViewIndex, *View, HairVisibilityData, this);
 	}
 
 	if (IsWholeSceneDirectionalShadow())
@@ -1015,7 +1045,7 @@ static void SetPointLightShaderTempl(FRHICommandList& RHICmdList, FGraphicsPipel
 	SetGraphicsPipelineState(RHICmdList, GraphicsPSOInit);
 	
 	VertexShader->SetParameters(RHICmdList, View, ShadowInfo);
-	PixelShader->SetParameters(RHICmdList, ViewIndex, View, ShadowInfo);
+	PixelShader->SetParameters(RHICmdList, ViewIndex, View, nullptr, ShadowInfo);
 }
 
 /** Render one pass point light shadow projections. */
@@ -1432,7 +1462,7 @@ bool FDeferredShadingSceneRenderer::InjectReflectiveShadowMaps(FRHICommandListIm
 	return true;
 }
 
-bool FSceneRenderer::RenderShadowProjections(FRHICommandListImmediate& RHICmdList, const FLightSceneInfo* LightSceneInfo, IPooledRenderTarget* ScreenShadowMaskTexture, bool bProjectingForForwardShading, bool bMobileModulatedProjections)
+bool FSceneRenderer::RenderShadowProjections(FRHICommandListImmediate& RHICmdList, const FLightSceneInfo* LightSceneInfo, IPooledRenderTarget* ScreenShadowMaskTexture, IPooledRenderTarget* ScreenShadowMaskSubPixelTexture, bool bProjectingForForwardShading, bool bMobileModulatedProjections, const FHairStrandsVisibilityViews* InHairVisibilityViews)
 {
 	FVisibleLightInfo& VisibleLightInfo = VisibleLightInfos[LightSceneInfo->Id];
 	FSceneRenderTargets& SceneContext = FSceneRenderTargets::Get(RHICmdList);
@@ -1457,72 +1487,91 @@ bool FSceneRenderer::RenderShadowProjections(FRHICommandListImmediate& RHICmdLis
 
 	if (NormalShadows.Num() > 0)
 	{
-		// Render normal shadows
-		if (bMobileModulatedProjections)
-		{
-			// part of scene color rendering pass
-			check(RHICmdList.IsInsideRenderPass());
-		}
-		else
+		auto RenderShadowMask = [&](const FHairStrandsVisibilityViews* HairVisibilityViews)
 		{
-			check(RHICmdList.IsOutsideRenderPass());
-			// Normal deferred shadows render to the shadow mask
-			FRHIRenderPassInfo RPInfo(ScreenShadowMaskTexture->GetRenderTargetItem().TargetableTexture, ERenderTargetActions::Load_Store);
-			RPInfo.DepthStencilRenderTarget.Action = MakeDepthStencilTargetActions(ERenderTargetActions::Load_DontStore, ERenderTargetActions::Load_Store);
-			RPInfo.DepthStencilRenderTarget.DepthStencilTarget = SceneContext.GetSceneDepthSurface();
-			RPInfo.DepthStencilRenderTarget.ExclusiveDepthStencil = FExclusiveDepthStencil::DepthRead_StencilWrite;
-
-			TransitionRenderPassTargets(RHICmdList, RPInfo);
-			RHICmdList.BeginRenderPass(RPInfo, TEXT("RenderShadowProjection"));
-		}
-
-		for (int32 ViewIndex = 0; ViewIndex < Views.Num(); ViewIndex++)
-		{
-			SCOPED_CONDITIONAL_DRAW_EVENTF(RHICmdList, EventView, Views.Num() > 1, TEXT("View%d"), ViewIndex);
+			for (int32 ViewIndex = 0; ViewIndex < Views.Num(); ViewIndex++)
+			{
+				SCOPED_CONDITIONAL_DRAW_EVENTF(RHICmdList, EventView, Views.Num() > 1, TEXT("View%d"), ViewIndex);
 
-			const FViewInfo& View = Views[ViewIndex];
+				const FViewInfo& View = Views[ViewIndex];
 
-			// Set the device viewport for the view.
-			RHICmdList.SetViewport(View.ViewRect.Min.X, View.ViewRect.Min.Y, 0.0f, View.ViewRect.Max.X, View.ViewRect.Max.Y, 1.0f);
+				const FHairStrandsVisibilityData* HairVisibilityData = nullptr;
+				if (HairVisibilityViews)
+				{
+					HairVisibilityData = &(HairVisibilityViews->HairDatas[ViewIndex]);
+				}
 
-			// Set the light's scissor rectangle.
-			LightSceneInfo->Proxy->SetScissorRect(RHICmdList, View, View.ViewRect);
+				// Set the device viewport for the view.
+				RHICmdList.SetViewport(View.ViewRect.Min.X, View.ViewRect.Min.Y, 0.0f, View.ViewRect.Max.X, View.ViewRect.Max.Y, 1.0f);
 
-			Scene->UniformBuffers.UpdateViewUniformBuffer(View);
+				// Set the light's scissor rectangle.
+				LightSceneInfo->Proxy->SetScissorRect(RHICmdList, View, View.ViewRect);
 
-			// Project the shadow depth buffers onto the scene.
-			for (int32 ShadowIndex = 0; ShadowIndex < NormalShadows.Num(); ShadowIndex++)
-			{
-				FProjectedShadowInfo* ProjectedShadowInfo = NormalShadows[ShadowIndex];
+				Scene->UniformBuffers.UpdateViewUniformBuffer(View);
 
-				if (ProjectedShadowInfo->bAllocated)
+				// Project the shadow depth buffers onto the scene.
+				for (int32 ShadowIndex = 0; ShadowIndex < NormalShadows.Num(); ShadowIndex++)
 				{
-					// Only project the shadow if it's large enough in this particular view (split screen, etc... may have shadows that are large in one view but irrelevantly small in others)
-					if (ProjectedShadowInfo->FadeAlphas[ViewIndex] > 1.0f / 256.0f)
+					FProjectedShadowInfo* ProjectedShadowInfo = NormalShadows[ShadowIndex];
+
+					if (ProjectedShadowInfo->bAllocated)
 					{
-						if (ProjectedShadowInfo->bOnePassPointLightShadow)
+						// Only project the shadow if it's large enough in this particular view (split screen, etc... may have shadows that are large in one view but irrelevantly small in others)
+						if (ProjectedShadowInfo->FadeAlphas[ViewIndex] > 1.0f / 256.0f)
 						{
-							ProjectedShadowInfo->RenderOnePassPointLightProjection(RHICmdList, ViewIndex, View, bProjectingForForwardShading);
-						}
-						else
-						{
-							ProjectedShadowInfo->RenderProjection(RHICmdList, ViewIndex, &View, this, bProjectingForForwardShading, bMobileModulatedProjections);
+							if (ProjectedShadowInfo->bOnePassPointLightShadow)
+							{
+								ProjectedShadowInfo->RenderOnePassPointLightProjection(RHICmdList, ViewIndex, View, bProjectingForForwardShading);
+							}
+							else
+							{
+								ProjectedShadowInfo->RenderProjection(RHICmdList, ViewIndex, &View, this, bProjectingForForwardShading, bMobileModulatedProjections, HairVisibilityData);
+							}
 						}
 					}
 				}
 			}
-		}
-
-		// Reset the scissor rectangle.
-		RHICmdList.SetScissorRect(false, 0, 0, 0, 0);
+		};
 
+		// Render normal shadows
 		if (bMobileModulatedProjections)
 		{
 			// part of scene color rendering pass
 			check(RHICmdList.IsInsideRenderPass());
+			RenderShadowMask(nullptr);
+			// Reset the scissor rectangle.
+			RHICmdList.SetScissorRect(false, 0, 0, 0, 0);
 		}
 		else
 		{
+			check(RHICmdList.IsOutsideRenderPass());
+			// Normal deferred shadows render to the shadow mask
+			FRHIRenderPassInfo RPInfo(ScreenShadowMaskTexture->GetRenderTargetItem().TargetableTexture, ERenderTargetActions::Load_Store);
+			RPInfo.DepthStencilRenderTarget.Action = MakeDepthStencilTargetActions(ERenderTargetActions::Load_DontStore, ERenderTargetActions::Load_Store);
+			RPInfo.DepthStencilRenderTarget.DepthStencilTarget = SceneContext.GetSceneDepthSurface();
+			RPInfo.DepthStencilRenderTarget.ExclusiveDepthStencil = FExclusiveDepthStencil::DepthRead_StencilWrite;
+
+			TransitionRenderPassTargets(RHICmdList, RPInfo);
+			RHICmdList.BeginRenderPass(RPInfo, TEXT("RenderShadowProjection"));
+			RenderShadowMask(nullptr);
+			RHICmdList.SetScissorRect(false, 0, 0, 0, 0);
+			RHICmdList.EndRenderPass();
+		}
+
+		// SubPixelShadow
+		if (!bMobileModulatedProjections && ScreenShadowMaskSubPixelTexture && InHairVisibilityViews)
+		{			
+			check(RHICmdList.IsOutsideRenderPass());
+			// Normal deferred shadows render to the shadow mask
+			FRHIRenderPassInfo RPInfo(ScreenShadowMaskSubPixelTexture->GetRenderTargetItem().TargetableTexture, ERenderTargetActions::Load_Store);
+			RPInfo.DepthStencilRenderTarget.Action = MakeDepthStencilTargetActions(ERenderTargetActions::Load_DontStore, ERenderTargetActions::Load_Store);
+			RPInfo.DepthStencilRenderTarget.DepthStencilTarget = SceneContext.GetSceneDepthSurface();
+			RPInfo.DepthStencilRenderTarget.ExclusiveDepthStencil = FExclusiveDepthStencil::DepthRead_StencilWrite;
+
+			TransitionRenderPassTargets(RHICmdList, RPInfo);
+			RHICmdList.BeginRenderPass(RPInfo, TEXT("RenderShadowProjectionSubPixel"));
+			RenderShadowMask(InHairVisibilityViews);
+			RHICmdList.SetScissorRect(false, 0, 0, 0, 0);
 			RHICmdList.EndRenderPass();
 		}
 	}
@@ -1564,7 +1613,7 @@ bool FSceneRenderer::RenderShadowProjections(FRHICommandListImmediate& RHICmdLis
 	return true;
 }
 	
-bool FDeferredShadingSceneRenderer::RenderShadowProjections(FRHICommandListImmediate& RHICmdList, const FLightSceneInfo* LightSceneInfo, IPooledRenderTarget* ScreenShadowMaskTexture, bool& bInjectedTranslucentVolume)
+bool FDeferredShadingSceneRenderer::RenderShadowProjections(FRHICommandListImmediate& RHICmdList, const FLightSceneInfo* LightSceneInfo, IPooledRenderTarget* ScreenShadowMaskTexture, IPooledRenderTarget* ScreenShadowMaskSubPixelTexture, const FHairStrandsDatas* HairDatas, bool& bInjectedTranslucentVolume)
 {
 	SCOPED_NAMED_EVENT(FDeferredShadingSceneRenderer_RenderShadowProjections, FColor::Emerald);
 	SCOPE_CYCLE_COUNTER(STAT_ProjectedShadowDrawTime);
@@ -1575,7 +1624,7 @@ bool FDeferredShadingSceneRenderer::RenderShadowProjections(FRHICommandListImmed
 
 	FVisibleLightInfo& VisibleLightInfo = VisibleLightInfos[LightSceneInfo->Id];
 
-	FSceneRenderer::RenderShadowProjections(RHICmdList, LightSceneInfo, ScreenShadowMaskTexture, false, false);
+	FSceneRenderer::RenderShadowProjections(RHICmdList, LightSceneInfo, ScreenShadowMaskTexture, ScreenShadowMaskSubPixelTexture, false, false, HairDatas ? &HairDatas->HairVisibilityViews : nullptr);
 
 	checkSlow(RHICmdList.IsOutsideRenderPass());
 	for (int32 ShadowIndex = 0; ShadowIndex < VisibleLightInfo.ShadowsToProject.Num(); ShadowIndex++)
@@ -1637,6 +1686,12 @@ bool FDeferredShadingSceneRenderer::RenderShadowProjections(FRHICommandListImmed
 		}
 	}
 
+	// Inject deep shadow mask
+	if (HairDatas)
+	{
+		RenderHairStrandsShadowMask(RHICmdList, Views, LightSceneInfo, ScreenShadowMaskTexture, HairDatas);
+	}
+
 	return true;
 }
 
@@ -1657,7 +1712,7 @@ void FMobileSceneRenderer::RenderModulatedShadowProjections(FRHICommandListImmed
 		{
 			TArray<FProjectedShadowInfo*, SceneRenderingAllocator> Shadows;
 			SCOPE_CYCLE_COUNTER(STAT_ProjectedShadowDrawTime);
-			FSceneRenderer::RenderShadowProjections(RHICmdList, LightSceneInfo, NULL, false, true);
+			FSceneRenderer::RenderShadowProjections(RHICmdList, LightSceneInfo, nullptr, nullptr, false, true, nullptr);
 		}
 	}
 }
diff --git a/Engine/Source/Runtime/Renderer/Private/ShadowRendering.h b/Engine/Source/Runtime/Renderer/Private/ShadowRendering.h
--- a/Engine/Source/Runtime/Renderer/Private/ShadowRendering.h
+++ b/Engine/Source/Runtime/Renderer/Private/ShadowRendering.h
@@ -28,6 +28,7 @@
 #include "ShaderParameterUtils.h"
 #include "LightRendering.h"
 #include "LightPropagationVolume.h"
+#include "HairStrands/HairStrandsRendering.h"
 
 ENGINE_API IPooledRenderTarget* GetSubsufaceProfileTexture_RT(FRHICommandListImmediate& RHICmdList);
 
@@ -322,6 +323,9 @@ public:
 	/** Whether turn on back-lighting transmission. */
 	uint32 bTransmission : 1;
 
+	/** Whether turn on hair strands deep shadow. */
+	uint32 bHairStrandsDeepShadow : 1;
+
 	/** View projection matrices for each cubemap face, used by one pass point light shadows. */
 	TArray<FMatrix> OnePassShadowViewProjectionMatrices;
 
@@ -401,7 +405,7 @@ public:
 	/**
 	 * Projects the shadow onto the scene for a particular view.
 	 */
-	void RenderProjection(FRHICommandListImmediate& RHICmdList, int32 ViewIndex, const class FViewInfo* View, const class FSceneRenderer* SceneRender, bool bProjectingForForwardShading, bool bMobile) const;
+	void RenderProjection(FRHICommandListImmediate& RHICmdList, int32 ViewIndex, const class FViewInfo* View, const class FSceneRenderer* SceneRender, bool bProjectingForForwardShading, bool bMobile, const struct FHairStrandsVisibilityData* HairVisibilityData) const;
 
 	void BeginRenderRayTracedDistanceFieldProjection(FRHICommandListImmediate& RHICmdList, const FViewInfo& View);
 
@@ -752,8 +756,8 @@ public:
 		FRHICommandList& RHICmdList, 
 		int32 ViewIndex,
 		const FSceneView& View,
-		const FProjectedShadowInfo* ShadowInfo
-		)
+		const FHairStrandsVisibilityData* HairVisibilityData,
+		const FProjectedShadowInfo* ShadowInfo)
 	{ 
 		FGlobalShader::SetParameters<FViewUniformShaderParameters>(RHICmdList, GetPixelShader(), View.ViewUniformBuffer);
 	}
@@ -779,9 +783,10 @@ public:
 		InvFadePlaneLength.Bind(ParameterMap,TEXT("InvFadePlaneLength"));
 		ShadowTileOffsetAndSizeParam.Bind(ParameterMap, TEXT("ShadowTileOffsetAndSize"));
 		LightPositionOrDirection.Bind(ParameterMap, TEXT("LightPositionOrDirection"));
+		HairCategorizationTexture.Bind(ParameterMap, TEXT("HairCategorizationTexture"));
 	}
 
-	void Set(FRHICommandList& RHICmdList, FShader* Shader, const FSceneView& View, const FProjectedShadowInfo* ShadowInfo)
+	void Set(FRHICommandList& RHICmdList, FShader* Shader, const FSceneView& View, const FProjectedShadowInfo* ShadowInfo, const FHairStrandsVisibilityData* HairVisibilityData)
 	{
 		FRHIPixelShader* ShaderRHI = Shader->GetPixelShader();
 
@@ -862,7 +867,6 @@ public:
 			SetShaderValue(RHICmdList, ShaderRHI, InvFadePlaneLength, 1.0f / ShadowInfo->CascadeSettings.FadePlaneLength);
 		}
 
-
 		if (LightPositionOrDirection.IsBound())
 		{
 			const FVector LightDirection = ShadowInfo->GetLightSceneInfo().Proxy->GetDirection();
@@ -870,6 +874,11 @@ public:
 			const bool bIsDirectional = ShadowInfo->GetLightSceneInfo().Proxy->GetLightType() == LightType_Directional;
 			SetShaderValue(RHICmdList, ShaderRHI, LightPositionOrDirection, bIsDirectional ? FVector4(LightDirection,0) : FVector4(LightPosition,1));
 		}
+
+		if (HairVisibilityData && HairVisibilityData->CategorizationTexture)
+		{
+			SetTextureParameter(RHICmdList, ShaderRHI, HairCategorizationTexture, HairVisibilityData->CategorizationTexture->GetRenderTargetItem().ShaderResourceTexture);
+		}
 	}
 
 	/** Serializer. */
@@ -886,6 +895,7 @@ public:
 		Ar << P.InvFadePlaneLength;
 		Ar << P.ShadowTileOffsetAndSizeParam;
 		Ar << P.LightPositionOrDirection;
+		Ar << P.HairCategorizationTexture;
 		return Ar;
 	}
 
@@ -902,13 +912,14 @@ private:
 	FShaderParameter InvFadePlaneLength;
 	FShaderParameter ShadowTileOffsetAndSizeParam;
 	FShaderParameter LightPositionOrDirection;
+	FShaderResourceParameter HairCategorizationTexture;
 };
 
 /**
  * TShadowProjectionPS
  * A pixel shader for projecting a shadow depth buffer onto the scene.  Used with any light type casting normal shadows.
  */
-template<uint32 Quality, bool bUseFadePlane = false, bool bModulatedShadows = false, bool bUseTransmission = false>
+template<uint32 Quality, bool bUseFadePlane = false, bool bModulatedShadows = false, bool bUseTransmission = false, bool SubPixelShadow = false>
 class TShadowProjectionPS : public FShadowProjectionPixelShaderInterface
 {
 	DECLARE_SHADER_TYPE(TShadowProjectionPS,Global);
@@ -947,6 +958,7 @@ public:
 	{
 		FShadowProjectionPixelShaderInterface::ModifyCompilationEnvironment(Parameters,OutEnvironment);
 		OutEnvironment.SetDefine(TEXT("SHADOW_QUALITY"), Quality);
+		OutEnvironment.SetDefine(TEXT("SUBPIXEL_SHADOW"), (uint32)(SubPixelShadow ? 1 : 0));
 		OutEnvironment.SetDefine(TEXT("USE_FADE_PLANE"), (uint32)(bUseFadePlane ? 1 : 0));
 		OutEnvironment.SetDefine(TEXT("USE_TRANSMISSION"), (uint32)(bUseTransmission ? 1 : 0));
 	}
@@ -957,17 +969,17 @@ public:
 	 * @param ShadowInfo - projected shadow info for a single light
 	 */
 	virtual void SetParameters(
-		FRHICommandList& RHICmdList, 
+		FRHICommandList& RHICmdList,
 		int32 ViewIndex,
 		const FSceneView& View,
-		const FProjectedShadowInfo* ShadowInfo
-		) override
+		const FHairStrandsVisibilityData* HairVisibilityData,
+		const FProjectedShadowInfo* ShadowInfo) override
 	{
 		FRHIPixelShader* ShaderRHI = GetPixelShader();
 
-		FShadowProjectionPixelShaderInterface::SetParameters(RHICmdList, ViewIndex,View,ShadowInfo);
+		FShadowProjectionPixelShaderInterface::SetParameters(RHICmdList, ViewIndex, View, HairVisibilityData, ShadowInfo);
 
-		ProjectionParameters.Set(RHICmdList, this, View, ShadowInfo);
+		ProjectionParameters.Set(RHICmdList, this, View, ShadowInfo, HairVisibilityData);
 		const FLightSceneProxy& LightProxy = *(ShadowInfo->GetLightSceneInfo().Proxy);
 
 		SetShaderValue(RHICmdList, ShaderRHI, ShadowFadeFraction, ShadowInfo->FadeAlphas[ViewIndex] );
@@ -1059,9 +1071,10 @@ public:
 		FRHICommandList& RHICmdList,
 		int32 ViewIndex,
 		const FSceneView& View,
+		const FHairStrandsVisibilityData* HairVisibilityData,
 		const FProjectedShadowInfo* ShadowInfo) override
 	{
-		TShadowProjectionPS<Quality, false, true>::SetParameters(RHICmdList, ViewIndex, View, ShadowInfo);
+		TShadowProjectionPS<Quality, false, true>::SetParameters(RHICmdList, ViewIndex, View, HairVisibilityData, ShadowInfo);
 		FRHIPixelShader* ShaderRHI = this->GetPixelShader();
 		SetShaderValue(RHICmdList, ShaderRHI, ModulatedShadowColorParameter, ShadowInfo->GetLightSceneInfo().Proxy->GetModulatedShadowColor());
 	}
@@ -1137,9 +1150,10 @@ public:
 		FRHICommandList& RHICmdList, 
 		int32 ViewIndex,
 		const FSceneView& View,
+		const FHairStrandsVisibilityData* HairVisibilityData,
 		const FProjectedShadowInfo* ShadowInfo) override
 	{
-		TShadowProjectionPS<Quality>::SetParameters(RHICmdList, ViewIndex, View, ShadowInfo);
+		TShadowProjectionPS<Quality>::SetParameters(RHICmdList, ViewIndex, View, HairVisibilityData, ShadowInfo);
 
 		FTranslucentSelfShadowUniformParameters TranslucentSelfShadowUniformParameters;
 		SetupTranslucentSelfShadowUniformParameters(ShadowInfo, TranslucentSelfShadowUniformParameters);
@@ -1306,12 +1320,12 @@ public:
 		FRHICommandList& RHICmdList, 
 		int32 ViewIndex,
 		const FSceneView& View,
-		const FProjectedShadowInfo* ShadowInfo
-		)
+		const FHairStrandsVisibilityData* HairVisibilityData,
+		const FProjectedShadowInfo* ShadowInfo)
 	{
 		FRHIPixelShader* ShaderRHI = GetPixelShader();
 
-		FGlobalShader::SetParameters<FViewUniformShaderParameters>(RHICmdList, ShaderRHI,View.ViewUniformBuffer);
+		FGlobalShader::SetParameters<FViewUniformShaderParameters>(RHICmdList, ShaderRHI, View.ViewUniformBuffer);
 
 		SceneTextureParameters.Set(RHICmdList, ShaderRHI, View.FeatureLevel, ESceneTextureSetupMode::All);
 		OnePassShadowParameters.Set(RHICmdList, ShaderRHI, ShadowInfo);
@@ -1435,9 +1449,10 @@ public:
 		FRHICommandList& RHICmdList,
 		int32 ViewIndex,
 		const FSceneView& View,
+		const FHairStrandsVisibilityData* HairVisibilityData,
 		const FProjectedShadowInfo* ShadowInfo) override
 	{
-		TShadowProjectionPS<Quality, bUseFadePlane>::SetParameters(RHICmdList, ViewIndex, View, ShadowInfo);
+		TShadowProjectionPS<Quality, bUseFadePlane>::SetParameters(RHICmdList, ViewIndex, View, HairVisibilityData, ShadowInfo);
 
 		FRHIPixelShader* ShaderRHI = this->GetPixelShader();
 
@@ -1502,11 +1517,12 @@ public:
 		FRHICommandList& RHICmdList,
 		int32 ViewIndex,
 		const FSceneView& View,
+		const FHairStrandsVisibilityData* HairVisibilityData,
 		const FProjectedShadowInfo* ShadowInfo) override
 	{
 		check(ShadowInfo->GetLightSceneInfo().Proxy->GetLightType() == LightType_Spot);
 
-		TShadowProjectionPS<Quality, bUseFadePlane>::SetParameters(RHICmdList, ViewIndex, View, ShadowInfo);
+		TShadowProjectionPS<Quality, bUseFadePlane>::SetParameters(RHICmdList, ViewIndex, View, HairVisibilityData, ShadowInfo);
 
 		FRHIPixelShader* ShaderRHI = this->GetPixelShader();
 
diff --git a/Engine/Source/Runtime/Renderer/Private/ShadowSetup.cpp b/Engine/Source/Runtime/Renderer/Private/ShadowSetup.cpp
--- a/Engine/Source/Runtime/Renderer/Private/ShadowSetup.cpp
+++ b/Engine/Source/Runtime/Renderer/Private/ShadowSetup.cpp
@@ -514,6 +514,7 @@ FProjectedShadowInfo::FProjectedShadowInfo()
 	, bSelfShadowOnly(false)
 	, bPerObjectOpaqueShadow(false)
 	, bTransmission(false)
+	, bHairStrandsDeepShadow(false)
 	, LightSceneInfo(0)
 	, ParentSceneInfo(0)
 	, NumDynamicSubjectMeshElements(0)
@@ -554,6 +555,7 @@ bool FProjectedShadowInfo::SetupPerObjectProjection(
 	bPreShadow = bInPreShadow;
 	bSelfShadowOnly = InParentSceneInfo->Proxy->CastsSelfShadowOnly();
 	bTransmission = InLightSceneInfo->Proxy->Transmission();
+	bHairStrandsDeepShadow = InLightSceneInfo->Proxy->CastsHairStrandsDeepShadow();
 
 	check(!bRayTracedDistanceField);
 
@@ -675,6 +677,7 @@ void FProjectedShadowInfo::SetupWholeSceneProjection(
 	bRayTracedDistanceField = Initializer.bRayTracedDistanceField;
 	bWholeSceneShadow = true;
 	bTransmission = InLightSceneInfo->Proxy->Transmission();
+	bHairStrandsDeepShadow = InLightSceneInfo->Proxy->CastsHairStrandsDeepShadow();
 	bReflectiveShadowmap = bInReflectiveShadowMap; 
 	BorderSize = InBorderSize;
 
diff --git a/Engine/Source/Runtime/Renderer/Private/SystemTextures.cpp b/Engine/Source/Runtime/Renderer/Private/SystemTextures.cpp
--- a/Engine/Source/Runtime/Renderer/Private/SystemTextures.cpp
+++ b/Engine/Source/Runtime/Renderer/Private/SystemTextures.cpp
@@ -596,6 +596,8 @@ void FSystemTextures::ReleaseDynamicRHI()
 	SSAORandomization.SafeRelease();
 	GTAORandomization.SafeRelease();
 	PreintegratedGF.SafeRelease();
+	HairLUT0.SafeRelease();
+	HairLUT1.SafeRelease();
 	LTCMat.SafeRelease();
 	LTCAmp.SafeRelease();
 	MaxFP16Depth.SafeRelease();
diff --git a/Engine/Source/Runtime/Renderer/Private/VelocityRendering.cpp b/Engine/Source/Runtime/Renderer/Private/VelocityRendering.cpp
--- a/Engine/Source/Runtime/Renderer/Private/VelocityRendering.cpp
+++ b/Engine/Source/Runtime/Renderer/Private/VelocityRendering.cpp
@@ -464,7 +464,7 @@ FPooledRenderTargetDesc FVelocityRendering::GetRenderTargetDesc()
 {
 	const FIntPoint BufferSize = FSceneRenderTargets::Get_FrameConstantsOnly().GetBufferSizeXY();
 	const FIntPoint VelocityBufferSize = BufferSize;		// full resolution so we can reuse the existing full res z buffer
-	return FPooledRenderTargetDesc(FPooledRenderTargetDesc::Create2DDesc(VelocityBufferSize, PF_G16R16, FClearValueBinding::Transparent, TexCreate_None, TexCreate_RenderTargetable, false));
+	return FPooledRenderTargetDesc(FPooledRenderTargetDesc::Create2DDesc(VelocityBufferSize, PF_G16R16, FClearValueBinding::Transparent, TexCreate_None, TexCreate_RenderTargetable | TexCreate_UAV, false));
 }
 
 bool FVelocityRendering::BasePassCanOutputVelocity(EShaderPlatform ShaderPlatform)
diff --git a/Engine/Source/Runtime/Renderer/Public/HairStrandsInterface.h b/Engine/Source/Runtime/Renderer/Public/HairStrandsInterface.h
new file mode 100644
--- /dev/null
+++ b/Engine/Source/Runtime/Renderer/Public/HairStrandsInterface.h
@@ -0,0 +1,59 @@
+// Copyright 1998-2019 Epic Games, Inc. All Rights Reserved.
+
+/*=============================================================================
+	HairStrandsInterface.h: public interface for hair strands rendering.
+=============================================================================*/
+
+#pragma once
+
+#include "CoreMinimal.h"
+#include "RendererInterface.h"
+
+enum class EHairStrandsDebugMode : uint8
+{
+	None,
+	SimHairStrands,
+	RenderHairStrands,
+	RenderHairUV,
+	RenderHairRootUV,
+	RenderHairSeed,
+	RenderHairDimension,
+	RenderHairRadiusVariation,
+	Count
+};
+
+/// Return the active debug view mode
+RENDERER_API EHairStrandsDebugMode GetHairStrandsDebugMode();
+
+/// Return the number of sample subsample count used for the visibility pass
+RENDERER_API uint32 GetHairVisibilitySampleCount();
+
+struct FMinHairRadiusAtDepth1
+{
+	float Primary = 1;
+	float Velocity = 1;
+};
+
+/// Compute the strand radius at a distance of 1 meter
+RENDERER_API FMinHairRadiusAtDepth1 ComputeMinStrandRadiusAtDepth1(
+	const FIntPoint& Resolution,
+	const float FOV,
+	const uint32 SampleCount,
+	const float OverrideStrandHairRasterizationScale);
+
+typedef void (*THairStrandsInterpolationFunction)(FRHICommandListImmediate& RHICmdList, struct FHairStrandsInterpolationInput* Input, struct FHairStrandsInterpolationOutput* Output);
+
+struct FHairStrandsInterpolation
+{
+	struct FHairStrandsInterpolationInput* Input = nullptr;
+	struct FHairStrandsInterpolationOutput* Output = nullptr;
+	THairStrandsInterpolationFunction Function = nullptr;
+};
+
+RENDERER_API void RegisterHairStrands(uint64 Id, const FHairStrandsInterpolation& E);
+RENDERER_API FHairStrandsInterpolation UnregisterHairStrands(uint64 Id);
+void RunHairStrandsInterpolation(FRHICommandListImmediate& RHICmdList);
+
+
+RENDERER_API bool IsHairStrandsSupported(const EShaderPlatform Platform);
+bool IsHairStrandsEnable(EShaderPlatform Platform);
\ No newline at end of file
diff --git a/Engine/Source/Runtime/Renderer/Public/SceneRenderTargetParameters.h b/Engine/Source/Runtime/Renderer/Public/SceneRenderTargetParameters.h
--- a/Engine/Source/Runtime/Renderer/Public/SceneRenderTargetParameters.h
+++ b/Engine/Source/Runtime/Renderer/Public/SceneRenderTargetParameters.h
@@ -8,6 +8,7 @@
 
 #include "CoreMinimal.h"
 #include "ShaderParameters.h"
+#include "SceneInterface.h"
 #include "MaterialShared.h"
 
 class FSceneView;
diff --git a/Engine/Source/Runtime/Windows/D3D11RHI/Private/D3D11Device.cpp b/Engine/Source/Runtime/Windows/D3D11RHI/Private/D3D11Device.cpp
--- a/Engine/Source/Runtime/Windows/D3D11RHI/Private/D3D11Device.cpp
+++ b/Engine/Source/Runtime/Windows/D3D11RHI/Private/D3D11Device.cpp
@@ -150,6 +150,8 @@ FD3D11DynamicRHI::FD3D11DynamicRHI(IDXGIFactory1* InDXGIFactory1,D3D_FEATURE_LEV
 		GPixelFormats[PF_X24_G8].PlatformFormat = DXGI_FORMAT_X32_TYPELESS_G8X24_UINT;
 		GPixelFormats[PF_X24_G8].BlockBytes = 5;
 	}
+	GPixelFormats[ PF_DepthStencil	].Supported = true;
+	GPixelFormats[ PF_X24_G8		].Supported = true;
 	GPixelFormats[ PF_ShadowDepth	].PlatformFormat	= DXGI_FORMAT_R16_TYPELESS;
 	GPixelFormats[ PF_ShadowDepth	].BlockBytes		= 2;
 	GPixelFormats[ PF_ShadowDepth	].Supported			= true;
