* Fixed indirect dispatch of tiles for water on X1 by removing usage of RHISupportsDrawIndirect as it is incorrectly returning false on X1.
* Fixed X1 missing shader on startup. FWaterRefractionCopyPS was incorrectly checking for RHISupportsDrawIndirect, which also incorrectly returns false and prevents this mandatory shader from compiling
* Water refraction - separated RGBA16F color and depth into 2 targets (R11G10B11F - color, F16F - depth). This will help with upcoming ESRAM optimizations.
* Water refraction - all refraction depth is now point sampled. This fixes artifacts on object silhouettes.

diff --git a/Engine/Shaders/Private/BasePassPixelShader.usf b/Engine/Shaders/Private/BasePassPixelShader.usf
--- a/Engine/Shaders/Private/BasePassPixelShader.usf
+++ b/Engine/Shaders/Private/BasePassPixelShader.usf
@@ -1226,7 +1226,8 @@ ShadingModel);
 
 		Color += EvaluateWaterVolumeLighting(
 			MaterialParameters, PixelMaterialInputs, ResolvedView,
-			OpaqueBasePass.SceneColorAndDepthWithoutSingleLayerWaterTexture, OpaqueBasePass.SceneColorAndDepthWithoutSingleLayerWaterSampler,
+			OpaqueBasePass.SceneColorWithoutSingleLayerWaterTexture, OpaqueBasePass.SceneColorWithoutSingleLayerWaterSampler,
+			OpaqueBasePass.SceneDepthWithoutSingleLayerWaterTexture, OpaqueBasePass.SceneDepthWithoutSingleLayerWaterSampler,
 			OpaqueBasePass.SceneCustomDepthTexture, OpaqueBasePass.SceneCustomDepthSampler,
 			OpaqueBasePass.SceneWithoutSingleLayerWaterValidUVRect,
 			Specular, OpaqueBasePass.DistortionParams,
diff --git a/Engine/Shaders/Private/SingleLayerWaterCommon.ush b/Engine/Shaders/Private/SingleLayerWaterCommon.ush
--- a/Engine/Shaders/Private/SingleLayerWaterCommon.ush
+++ b/Engine/Shaders/Private/SingleLayerWaterCommon.ush
@@ -7,7 +7,7 @@
 #include "ShadingCommon.ush" 
 #include "DistortionCommon.ush"
 
-// Scale for encoding depth into SceneColorAndDepthWithoutSingleLayerWaterTexture
+// Scale for encoding depth into SceneDepthWithoutSingleLayerWaterTexture
 #define SINGLE_LAYER_WATER_DEPTH_SCALE 100.0f
 
 #ifndef SIMPLE_SINGLE_LAYER_WATER
@@ -45,7 +45,8 @@ WaterVolumeLightingOutput EvaluateWaterVolumeLighting(
 #if SIMPLE_SINGLE_LAYER_WATER
 	float DirectionalLightShadow,
 #else
-	Texture2D SceneColorAndDepthWithoutSingleLayerWaterTexture, SamplerState SceneColorAndDepthWithoutSingleLayerWaterSampler,
+	Texture2D SceneColorWithoutSingleLayerWaterTexture, SamplerState SceneColorWithoutSingleLayerWaterSampler,
+	Texture2D SceneDepthWithoutSingleLayerWaterTexture, SamplerState SceneDepthWithoutSingleLayerWaterSampler,
 	Texture2D SceneCustomDepthTexture, SamplerState SceneCustomDepthSampler,
 	float4 SceneWithoutSingleLayerWaterValidUVRect,
 #endif // SIMPLE_SINGLE_LAYER_WATER
@@ -81,7 +82,7 @@ WaterVolumeLightingOutput EvaluateWaterVolumeLighting(
 
 #else  // !SIMPLE_SINGLE_LAYER_WATER
 
-		float PixelSceneDepth = SceneColorAndDepthWithoutSingleLayerWaterTexture.SampleLevel(SceneColorAndDepthWithoutSingleLayerWaterSampler, ViewportUV, 0).w * SINGLE_LAYER_WATER_DEPTH_SCALE;
+		float PixelSceneDepth = SceneDepthWithoutSingleLayerWaterTexture.SampleLevel(SceneDepthWithoutSingleLayerWaterSampler, ViewportUV, 0).x * SINGLE_LAYER_WATER_DEPTH_SCALE;
 		
 		float2 Refraction = GetMaterialRefraction(PixelMaterialInputs); // X = IOR, Y = Refraction Bias
 		float2 BufferUVDistortion = ComputeBufferUVDistortion(
@@ -94,7 +95,7 @@ WaterVolumeLightingOutput EvaluateWaterVolumeLighting(
 		// Clamp distorted UVs to not sample black outside the sub part of the buffer. Big quality increase.
 		const float2 DistortedUV = clamp(ViewportUV + BufferUVDistortion, SceneWithoutSingleLayerWaterValidUVRect.xy, SceneWithoutSingleLayerWaterValidUVRect.zw);
 
-		float SceneDepth = SceneColorAndDepthWithoutSingleLayerWaterTexture.SampleLevel(SceneColorAndDepthWithoutSingleLayerWaterSampler, DistortedUV, 0).w * SINGLE_LAYER_WATER_DEPTH_SCALE;
+		float SceneDepth = SceneDepthWithoutSingleLayerWaterTexture.SampleLevel(SceneDepthWithoutSingleLayerWaterSampler, DistortedUV, 0).x * SINGLE_LAYER_WATER_DEPTH_SCALE;
 
 		if (SceneDepth < WaterDepth)
 		{
@@ -143,7 +144,7 @@ WaterVolumeLightingOutput EvaluateWaterVolumeLighting(
 		const float3 BehindWaterSceneLuminance = 0.0f; // Cannot read back the scene color in this case
 #else
 		// We use the pixel SvPosition instead of the scene one pre refraction/distortion to avoid those extra ALUs.
-		float3 BehindWaterSceneLuminance = SceneColorAndDepthWithoutSingleLayerWaterTexture.SampleLevel(SceneColorAndDepthWithoutSingleLayerWaterSampler, ViewportUV, 0).rgb;
+		float3 BehindWaterSceneLuminance = SceneColorWithoutSingleLayerWaterTexture.SampleLevel(SceneColorWithoutSingleLayerWaterSampler, ViewportUV, 0).rgb;
 		BehindWaterSceneLuminance = MeanTransmittanceToLightSources * (USE_PREEXPOSURE ? ResolvedView.OneOverPreExposure : 1.0f) * BehindWaterSceneLuminance;
 #endif
 
diff --git a/Engine/Shaders/Private/SingleLayerWaterComposite.usf b/Engine/Shaders/Private/SingleLayerWaterComposite.usf
--- a/Engine/Shaders/Private/SingleLayerWaterComposite.usf
+++ b/Engine/Shaders/Private/SingleLayerWaterComposite.usf
@@ -11,7 +11,7 @@
 #include "ShadingModels.ush"
 #include "SceneTextureParameters.ush"
 
-// Scale for encoding depth into SceneColorAndDepthWithoutSingleLayerWaterTexture
+// Scale for encoding depth into SceneDepthWithoutSingleLayerWaterTexture
 #define SINGLE_LAYER_WATER_DEPTH_SCALE 100.0f
 
 #define ENABLE_SKY_LIGHT 1
@@ -25,8 +25,8 @@
 Texture2D ScreenSpaceReflectionsTexture;
 SamplerState ScreenSpaceReflectionsSampler;
 
-Texture2D SceneNoWaterColorAndDepthTexture;
-SamplerState SceneNoWaterColorAndDepthSampler;
+Texture2D SceneNoWaterDepthTexture;
+SamplerState SceneNoWaterDepthSampler;
 
 struct SingleLayerWaterCompositeOutput
 {
@@ -57,7 +57,7 @@ SingleLayerWaterCompositeOutput SingleLayerWaterComposite(float2 BufferUV, float
 #endif
 	if (ShadingModelID == SHADINGMODELID_SINGLELAYERWATER)
 	{
-		float OpaqueDepth = SceneNoWaterColorAndDepthTexture.SampleLevel(SceneNoWaterColorAndDepthSampler, BufferUV, 0).w * SINGLE_LAYER_WATER_DEPTH_SCALE;
+		float OpaqueDepth = SceneNoWaterDepthTexture.SampleLevel(SceneNoWaterDepthSampler, BufferUV, 0).x * SINGLE_LAYER_WATER_DEPTH_SCALE;
 		float WaterDepth = GBuffer.Depth;
 
 		float DeltaDepth = CameraIsUnderWater ? WaterDepth : OpaqueDepth - WaterDepth;	// Inverted depth
@@ -291,12 +291,14 @@ float2 SVPositionToSourceTextureUV;
 
 void WaterRefractionCopyPS(
 	in float4 SVPosition : SV_Position,
-	out float4 OutColorAndDepth : SV_Target0)
+	out float4 OutSceneColor : SV_Target0,
+	out float OutSceneDepth : SV_Target1)
 {
 	float2 BufferUV = SVPosition.xy * SVPositionToSourceTextureUV;
 
 	// Downsample and average all 2x2 quads
-	OutColorAndDepth.xyz = Texture2DSampleLevel(SceneColorCopyDownsampleTexture, SceneColorCopyDownsampleSampler, BufferUV, 0).xyz;
+	OutSceneColor.xyz = Texture2DSampleLevel(SceneColorCopyDownsampleTexture, SceneColorCopyDownsampleSampler, BufferUV, 0).xyz;
+	OutSceneColor.w = 1.0f;
 
 	float DeviceDepthZ = 0.0f;
 #if DOWNSAMPLE_REFRACTION
@@ -306,5 +308,5 @@ void WaterRefractionCopyPS(
 #else
 	DeviceDepthZ = Texture2DSampleLevel(SceneDepthCopyDownsampleTexture, SceneDepthCopyDownsampleSampler, BufferUV, 0).x;
 #endif
-	OutColorAndDepth.w = ConvertFromDeviceZ(DeviceDepthZ) / SINGLE_LAYER_WATER_DEPTH_SCALE;
+	OutSceneDepth = ConvertFromDeviceZ(DeviceDepthZ) / SINGLE_LAYER_WATER_DEPTH_SCALE;
 }
\ No newline at end of file
diff --git a/Engine/Source/Runtime/Renderer/Private/BasePassRendering.cpp b/Engine/Source/Runtime/Renderer/Private/BasePassRendering.cpp
--- a/Engine/Source/Runtime/Renderer/Private/BasePassRendering.cpp
+++ b/Engine/Source/Runtime/Renderer/Private/BasePassRendering.cpp
@@ -517,7 +517,8 @@ void CreateOpaqueBasePassUniformBuffer(
 	const FViewInfo& View,
 	IPooledRenderTarget* ForwardScreenSpaceShadowMask,
 	FVector4* SceneWithoutSingleLayerWaterValidUVRect,
-	IPooledRenderTarget* SceneColorAndDepthWithoutSingleLayerWater,
+	IPooledRenderTarget* SceneColorWithoutSingleLayerWater,
+	IPooledRenderTarget* SceneDepthWithoutSingleLayerWater,
 	TUniformBufferRef<FOpaqueBasePassUniformParameters>& BasePassUniformBuffer)
 {
 	FSceneRenderTargets& SceneRenderTargets = FSceneRenderTargets::Get(RHICmdList);
@@ -583,8 +584,10 @@ void CreateOpaqueBasePassUniformBuffer(
 
 	// Single Layer Water
 	BasePassParameters.SceneWithoutSingleLayerWaterValidUVRect = SceneWithoutSingleLayerWaterValidUVRect ? *SceneWithoutSingleLayerWaterValidUVRect : FVector4(0.0f, 0.0f, 1.0f, 1.0f);
-	BasePassParameters.SceneColorAndDepthWithoutSingleLayerWaterTexture = (SceneColorAndDepthWithoutSingleLayerWater ? SceneColorAndDepthWithoutSingleLayerWater->GetRenderTargetItem().TargetableTexture : GSystemTextures.BlackDummy->GetRenderTargetItem().TargetableTexture);
-	BasePassParameters.SceneColorAndDepthWithoutSingleLayerWaterSampler = TStaticSamplerState<SF_Bilinear>::GetRHI();
+	BasePassParameters.SceneColorWithoutSingleLayerWaterTexture = (SceneColorWithoutSingleLayerWater ? SceneColorWithoutSingleLayerWater->GetRenderTargetItem().TargetableTexture : GSystemTextures.BlackDummy->GetRenderTargetItem().TargetableTexture);
+	BasePassParameters.SceneColorWithoutSingleLayerWaterSampler = TStaticSamplerState<SF_Bilinear>::GetRHI();
+	BasePassParameters.SceneDepthWithoutSingleLayerWaterTexture = (SceneDepthWithoutSingleLayerWater ? SceneDepthWithoutSingleLayerWater->GetRenderTargetItem().TargetableTexture : GSystemTextures.BlackDummy->GetRenderTargetItem().TargetableTexture);
+	BasePassParameters.SceneDepthWithoutSingleLayerWaterSampler = TStaticSamplerState<SF_Point>::GetRHI();
 	BasePassParameters.PreIntegratedGFTexture = GSystemTextures.PreintegratedGF->GetRenderTargetItem().ShaderResourceTexture;
 	BasePassParameters.PreIntegratedGFSampler = TStaticSamplerState<SF_Bilinear, AM_Clamp, AM_Clamp, AM_Clamp>::GetRHI();
 	BasePassParameters.SceneCustomDepthTexture = (SceneRenderTargets.CustomDepth.IsValid() ? SceneRenderTargets.CustomDepth : GSystemTextures.BlackDummy)->GetRenderTargetItem().ShaderResourceTexture;
@@ -649,7 +652,7 @@ bool FDeferredShadingSceneRenderer::RenderBasePass(FRHICommandListImmediate& RHI
 				SCOPED_GPU_MASK(RHICmdList, View.GPUMask);
 
 				TUniformBufferRef<FOpaqueBasePassUniformParameters> BasePassUniformBuffer;
-				CreateOpaqueBasePassUniformBuffer(RHICmdList, View, ForwardScreenSpaceShadowMask, nullptr, nullptr, BasePassUniformBuffer);
+				CreateOpaqueBasePassUniformBuffer(RHICmdList, View, ForwardScreenSpaceShadowMask, nullptr, nullptr, nullptr, BasePassUniformBuffer);
 
 				FMeshPassProcessorRenderState DrawRenderState(View, BasePassUniformBuffer);
 
@@ -690,7 +693,7 @@ bool FDeferredShadingSceneRenderer::RenderBasePass(FRHICommandListImmediate& RHI
 				SCOPED_GPU_MASK(RHICmdList, View.GPUMask);
 
 				TUniformBufferRef<FOpaqueBasePassUniformParameters> BasePassUniformBuffer;
-				CreateOpaqueBasePassUniformBuffer(RHICmdList, View, ForwardScreenSpaceShadowMask, nullptr, nullptr, BasePassUniformBuffer);
+				CreateOpaqueBasePassUniformBuffer(RHICmdList, View, ForwardScreenSpaceShadowMask, nullptr, nullptr, nullptr, BasePassUniformBuffer);
 
 				FMeshPassProcessorRenderState DrawRenderState(View, BasePassUniformBuffer);
 
diff --git a/Engine/Source/Runtime/Renderer/Private/BasePassRendering.h b/Engine/Source/Runtime/Renderer/Private/BasePassRendering.h
--- a/Engine/Source/Runtime/Renderer/Private/BasePassRendering.h
+++ b/Engine/Source/Runtime/Renderer/Private/BasePassRendering.h
@@ -72,8 +72,10 @@ BEGIN_GLOBAL_SHADER_PARAMETER_STRUCT(FOpaqueBasePassUniformParameters,)
 	SHADER_PARAMETER_SAMPLER(SamplerState, DBufferCTextureSampler)
 	SHADER_PARAMETER_TEXTURE(Texture2D<uint>, DBufferRenderMask)
 	// Single Layer Water
-	SHADER_PARAMETER_TEXTURE(Texture2D, SceneColorAndDepthWithoutSingleLayerWaterTexture)
-	SHADER_PARAMETER_SAMPLER(SamplerState, SceneColorAndDepthWithoutSingleLayerWaterSampler)
+	SHADER_PARAMETER_TEXTURE(Texture2D, SceneColorWithoutSingleLayerWaterTexture)
+	SHADER_PARAMETER_SAMPLER(SamplerState, SceneColorWithoutSingleLayerWaterSampler)
+	SHADER_PARAMETER_TEXTURE(Texture2D, SceneDepthWithoutSingleLayerWaterTexture)
+	SHADER_PARAMETER_SAMPLER(SamplerState, SceneDepthWithoutSingleLayerWaterSampler)
 	SHADER_PARAMETER_TEXTURE(Texture2D, PreIntegratedGFTexture)
 	SHADER_PARAMETER_SAMPLER(SamplerState, PreIntegratedGFSampler)
 	SHADER_PARAMETER_TEXTURE(Texture2D, SceneCustomDepthTexture)
@@ -119,7 +121,8 @@ extern void CreateOpaqueBasePassUniformBuffer(
 	const FViewInfo& View,
 	IPooledRenderTarget* ForwardScreenSpaceShadowMask,
 	FVector4* SceneWithoutSingleLayerWaterValidUVRect,
-	IPooledRenderTarget* SceneColorAndDepthWithoutSingleLayerWater,
+	IPooledRenderTarget* SceneColorWithoutSingleLayerWater,
+	IPooledRenderTarget* SceneDepthWithoutSingleLayerWater,
 	TUniformBufferRef<FOpaqueBasePassUniformParameters>& BasePassUniformBuffer);
 
 extern void CreateTranslucentBasePassUniformBuffer(
diff --git a/Engine/Source/Runtime/Renderer/Private/PostProcess/PostProcessCompositeEditorPrimitives.cpp b/Engine/Source/Runtime/Renderer/Private/PostProcess/PostProcessCompositeEditorPrimitives.cpp
--- a/Engine/Source/Runtime/Renderer/Private/PostProcess/PostProcessCompositeEditorPrimitives.cpp
+++ b/Engine/Source/Runtime/Renderer/Private/PostProcess/PostProcessCompositeEditorPrimitives.cpp
@@ -565,7 +565,7 @@ void FRCPassPostProcessCompositeEditorPrimitives::Process(FRenderingCompositePas
 
 			if (bDeferredBasePass)
 			{
-				CreateOpaqueBasePassUniformBuffer(Context.RHICmdList, EditorView, nullptr, nullptr, nullptr, OpaqueBasePassUniformBuffer);
+				CreateOpaqueBasePassUniformBuffer(Context.RHICmdList, EditorView, nullptr, nullptr, nullptr, nullptr, OpaqueBasePassUniformBuffer);
 				BasePassUniformBuffer = OpaqueBasePassUniformBuffer;
 			}
 			else
diff --git a/Engine/Source/Runtime/Renderer/Private/Renderer.cpp b/Engine/Source/Runtime/Renderer/Private/Renderer.cpp
--- a/Engine/Source/Runtime/Renderer/Private/Renderer.cpp
+++ b/Engine/Source/Runtime/Renderer/Private/Renderer.cpp
@@ -296,7 +296,7 @@ void FRendererModule::DrawTileMesh(FRHICommandListImmediate& RHICmdList, FMeshPa
 				if (ShadingPath == EShadingPath::Deferred)
 				{
 					TUniformBufferRef<FOpaqueBasePassUniformParameters> OpaqueBasePassUniformBuffer;
-					CreateOpaqueBasePassUniformBuffer(RHICmdList, View, nullptr, nullptr, nullptr, OpaqueBasePassUniformBuffer);
+					CreateOpaqueBasePassUniformBuffer(RHICmdList, View, nullptr, nullptr, nullptr, nullptr, OpaqueBasePassUniformBuffer);
 					BasePassUniformBuffer = OpaqueBasePassUniformBuffer;
 					
 					DrawRenderState.SetPassUniformBuffer(BasePassUniformBuffer);
diff --git a/Engine/Source/Runtime/Renderer/Private/ScreenSpaceRayTracing.cpp b/Engine/Source/Runtime/Renderer/Private/ScreenSpaceRayTracing.cpp
--- a/Engine/Source/Runtime/Renderer/Private/ScreenSpaceRayTracing.cpp
+++ b/Engine/Source/Runtime/Renderer/Private/ScreenSpaceRayTracing.cpp
@@ -112,7 +112,7 @@ bool IsSSRTemporalPassRequired(const FViewInfo& View)
 	return View.AntiAliasingMethod != AAM_TemporalAA || CVarSSRTemporal.GetValueOnRenderThread() != 0;
 }
 
-bool ShouldCompileTiledShadersForPlatform(EShaderPlatform ShaderPlatform);
+bool UseSingleLayerWaterIndirectDraw(EShaderPlatform ShaderPlatform);
 
 namespace
 {
@@ -358,7 +358,7 @@ class FScreenSpaceReflectionsTileVS : public FGlobalShader
 
 	static bool ShouldCompilePermutation(const FGlobalShaderPermutationParameters& Parameters)
 	{
-		return ::ShouldCompileTiledShadersForPlatform(Parameters.Platform);
+		return ::UseSingleLayerWaterIndirectDraw(Parameters.Platform);
 	}
 
 	static void ModifyCompilationEnvironment(const FGlobalShaderPermutationParameters& Parameters, FShaderCompilerEnvironment& OutEnvironment)
diff --git a/Engine/Source/Runtime/Renderer/Private/SingleLayerWaterRendering.cpp b/Engine/Source/Runtime/Renderer/Private/SingleLayerWaterRendering.cpp
--- a/Engine/Source/Runtime/Renderer/Private/SingleLayerWaterRendering.cpp
+++ b/Engine/Source/Runtime/Renderer/Private/SingleLayerWaterRendering.cpp
@@ -53,7 +53,7 @@ int32 GSingleLayerWaterRefractionFullPrecision = 0;
 FAutoConsoleVariableRef CVarWaterSingleLayerRefractionFullPrecision(
 	TEXT("r.Water.SingleLayer.RefractionFullPrecision"),
 	GSingleLayerWaterRefractionFullPrecision,
-	TEXT("Whether to pack refraction color and depth in RGBA Float32. To be used as a debug option to find issues with refraction depth precision."),
+	TEXT("Whether to pack refraction depth in a Float32 (instead of Float16). To be used as a debug option to find issues with refraction depth precision."),
 	ECVF_Scalability | ECVF_RenderThreadSafe
 );
 
@@ -122,10 +122,9 @@ bool ShouldRenderSingleLayerWaterSkippedRenderEditorNotification(const TArray<FV
 	return false;
 }
 
-bool ShouldCompileTiledShadersForPlatform(EShaderPlatform ShaderPlatform)
+bool UseSingleLayerWaterIndirectDraw(EShaderPlatform ShaderPlatform)
 {
-	// Temporary disable tiled water shaders until compilation issues are resolved.
-	return IsFeatureLevelSupported(ShaderPlatform, ERHIFeatureLevel::SM5) && RHISupportsDrawIndirect(ShaderPlatform)
+	return IsFeatureLevelSupported(ShaderPlatform, ERHIFeatureLevel::SM5)
 		&& !IsSwitchPlatform(ShaderPlatform) && !IsVulkanPlatform(ShaderPlatform); // Switch does not use tiling and Vulkan gives error with WaterTileCatergorisationCS usage of atomic.
 }
 
@@ -238,8 +237,8 @@ BEGIN_SHADER_PARAMETER_STRUCT(FSingleLayerWaterCommonShaderParameters, )
 	SHADER_PARAMETER_SAMPLER(SamplerState, ScreenSpaceReflectionsSampler)
 	SHADER_PARAMETER_TEXTURE(Texture2D, PreIntegratedGF)
 	SHADER_PARAMETER_SAMPLER(SamplerState, PreIntegratedGFSampler)
-	SHADER_PARAMETER_RDG_TEXTURE(Texture2D, SceneNoWaterColorAndDepthTexture)
-	SHADER_PARAMETER_SAMPLER(SamplerState, SceneNoWaterColorAndDepthSampler)
+	SHADER_PARAMETER_RDG_TEXTURE(Texture2D, SceneNoWaterDepthTexture)
+	SHADER_PARAMETER_SAMPLER(SamplerState, SceneNoWaterDepthSampler)
 	SHADER_PARAMETER_STRUCT_INCLUDE(FSceneTextureParameters, SceneTextures)	// Water scene texture
 	SHADER_PARAMETER_STRUCT_INCLUDE(FSceneTextureSamplerParameters, SceneTextureSamplers)
 	SHADER_PARAMETER_STRUCT_REF(FViewUniformShaderParameters, ViewUniformBuffer)
@@ -308,7 +307,7 @@ class FWaterTileCategorisationCS : public FGlobalShader
 
 	static bool ShouldCompilePermutation(const FGlobalShaderPermutationParameters& Parameters)
 	{
-		return ShouldCompileTiledShadersForPlatform(Parameters.Platform);
+		return UseSingleLayerWaterIndirectDraw(Parameters.Platform);
 	}
 
 	static void ModifyCompilationEnvironment(const FGlobalShaderPermutationParameters& Parameters, FShaderCompilerEnvironment& OutEnvironment)
@@ -338,7 +337,7 @@ class FWaterTileVS : public FGlobalShader
 
 	static bool ShouldCompilePermutation(const FGlobalShaderPermutationParameters& Parameters)
 	{
-		return ShouldCompileTiledShadersForPlatform(Parameters.Platform);
+		return UseSingleLayerWaterIndirectDraw(Parameters.Platform);
 	}
 
 	static void ModifyCompilationEnvironment(const FGlobalShaderPermutationParameters& Parameters, FShaderCompilerEnvironment& OutEnvironment)
@@ -372,7 +371,7 @@ class FWaterRefractionCopyPS : public FGlobalShader
 
 	static bool ShouldCompilePermutation(const FGlobalShaderPermutationParameters& Parameters)
 	{
-		return ShouldCompileTiledShadersForPlatform(Parameters.Platform);
+		return IsFeatureLevelSupported(Parameters.Platform, ERHIFeatureLevel::SM5);
 	}
 
 	static void ModifyCompilationEnvironment(const FGlobalShaderPermutationParameters& Parameters, FShaderCompilerEnvironment& OutEnvironment)
@@ -434,15 +433,18 @@ void FDeferredShadingSceneRenderer::CopySingleLayerWaterTextures(FRHICommandList
 	const int32 RefractionDownsampleFactor = FMath::Clamp(GSingleLayerWaterRefractionDownsampleFactor, 1, 8);
 	const FIntPoint RefractionResolution = FIntPoint::DivideAndRoundDown(SceneContext.GetBufferSizeXY(), RefractionDownsampleFactor);
 
-	FPooledRenderTargetDesc Desc(FPooledRenderTargetDesc::Create2DDesc(RefractionResolution, GSingleLayerWaterRefractionFullPrecision ? PF_A32B32G32R32F : PF_FloatRGBA, SceneContext.GetDefaultColorClear(), TexCreate_None, TexCreate_RenderTargetable, false));
-	Desc.NumSamples = SceneContext.GetNumSceneColorMSAASamples(CurrentFeatureLevel);
-	GRenderTargetPool.FindFreeElement(RHICmdList, Desc, PassData.SceneColorAndDepthWithoutSingleLayerWater, TEXT("SceneColorAndDepthWithoutSingleLayerWater"));
+	FPooledRenderTargetDesc ColorDesc(FPooledRenderTargetDesc::Create2DDesc(RefractionResolution, SceneContext.GetSceneColorFormat(), SceneContext.GetDefaultColorClear(), TexCreate_None, TexCreate_RenderTargetable, false));
+	GRenderTargetPool.FindFreeElement(RHICmdList, ColorDesc, PassData.SceneColorWithoutSingleLayerWater, TEXT("SceneColorWithoutSingleLayerWater"));
+
+	FPooledRenderTargetDesc DepthDesc(FPooledRenderTargetDesc::Create2DDesc(RefractionResolution, GSingleLayerWaterRefractionFullPrecision ? PF_R32_FLOAT : PF_R16F_FILTER, SceneContext.GetDefaultColorClear(), TexCreate_None, TexCreate_RenderTargetable, false));
+	GRenderTargetPool.FindFreeElement(RHICmdList, DepthDesc, PassData.SceneDepthWithoutSingleLayerWater, TEXT("SceneDepthWithoutSingleLayerWater"));
 
 	// For now support only the 1st view
 	const FViewInfo& View = Views[0];
 
 	FRDGBuilder GraphBuilder(RHICmdList);
-	FRDGTextureRef TargetColorTexture = GraphBuilder.RegisterExternalTexture(PassData.SceneColorAndDepthWithoutSingleLayerWater);
+	FRDGTextureRef TargetColorTexture = GraphBuilder.RegisterExternalTexture(PassData.SceneColorWithoutSingleLayerWater);
+	FRDGTextureRef TargetSceneDepthTexture = GraphBuilder.RegisterExternalTexture(PassData.SceneDepthWithoutSingleLayerWater);
 
 	FWaterRefractionCopyPS::FParameters* PassParameters = GraphBuilder.AllocParameters<FWaterRefractionCopyPS::FParameters>();
 	PassParameters->ViewUniformBuffer = View.ViewUniformBuffer;
@@ -452,6 +454,7 @@ void FDeferredShadingSceneRenderer::CopySingleLayerWaterTextures(FRHICommandList
 	PassParameters->SceneDepthCopyDownsampleSampler = TStaticSamplerState<SF_Bilinear, AM_Clamp, AM_Clamp, AM_Clamp>::GetRHI();
 	PassParameters->SVPositionToSourceTextureUV = FVector2D(RefractionDownsampleFactor / float(SceneContext.GetBufferSizeXY().X), RefractionDownsampleFactor / float(SceneContext.GetBufferSizeXY().Y));
 	PassParameters->RenderTargets[0] = FRenderTargetBinding(TargetColorTexture, ERenderTargetLoadAction::ELoad);
+	PassParameters->RenderTargets[1] = FRenderTargetBinding(TargetSceneDepthTexture, ERenderTargetLoadAction::ELoad);
 
 	FWaterRefractionCopyPS::FPermutationDomain PermutationVector;
 	PermutationVector.Set<FWaterRefractionCopyPS::FDownsampleRefraction>(RefractionDownsampleFactor > 1);
@@ -562,8 +565,8 @@ void FDeferredShadingSceneRenderer::RenderSingleLayerWaterReflections(FRHIComman
 			Parameters.ScreenSpaceReflectionsSampler = TStaticSamplerState<SF_Point>::GetRHI();
 			Parameters.PreIntegratedGF = GSystemTextures.PreintegratedGF->GetRenderTargetItem().ShaderResourceTexture;
 			Parameters.PreIntegratedGFSampler = TStaticSamplerState<SF_Bilinear, AM_Clamp, AM_Clamp, AM_Clamp>::GetRHI();
-			Parameters.SceneNoWaterColorAndDepthTexture = GraphBuilder.RegisterExternalTexture(PassData.SceneColorAndDepthWithoutSingleLayerWater ? PassData.SceneColorAndDepthWithoutSingleLayerWater : GSystemTextures.BlackDummy);
-			Parameters.SceneNoWaterColorAndDepthSampler = TStaticSamplerState<SF_Point>::GetRHI();
+			Parameters.SceneNoWaterDepthTexture = GraphBuilder.RegisterExternalTexture(PassData.SceneDepthWithoutSingleLayerWater ? PassData.SceneDepthWithoutSingleLayerWater : GSystemTextures.BlackDummy);
+			Parameters.SceneNoWaterDepthSampler = TStaticSamplerState<SF_Point>::GetRHI();
 			Parameters.SceneTextures = SceneTextures;
 			SetupSceneTextureSamplers(&Parameters.SceneTextureSamplers);
 			Parameters.ViewUniformBuffer = View.ViewUniformBuffer;
@@ -575,7 +578,7 @@ void FDeferredShadingSceneRenderer::RenderSingleLayerWaterReflections(FRHIComman
 			}
 		};
 
-		const bool bRunTiled = ShouldCompileTiledShadersForPlatform(View.GetShaderPlatform()) && CVarWaterSingleLayerTiledComposite.GetValueOnRenderThread();
+		const bool bRunTiled = UseSingleLayerWaterIndirectDraw(View.GetShaderPlatform()) && CVarWaterSingleLayerTiledComposite.GetValueOnRenderThread();
 		FTiledScreenSpaceReflection TiledScreenSpaceReflection = {nullptr, nullptr, nullptr, nullptr, nullptr, 8};
 		FIntVector ViewRes(View.ViewRect.Width(), View.ViewRect.Height(), 1);
 		FIntVector TiledViewRes = FIntVector::DivideAndRoundUp(ViewRes, TiledScreenSpaceReflection.TileSize);
@@ -606,7 +609,7 @@ void FDeferredShadingSceneRenderer::RenderSingleLayerWaterReflections(FRHIComman
 			}
 		}
 
-		const bool bEnableSSR = CVarWaterSingleLayerSSR.GetValueOnRenderThread() != 0 && ShouldRenderScreenSpaceReflections(View);
+ 		const bool bEnableSSR = CVarWaterSingleLayerSSR.GetValueOnRenderThread() != 0 && ShouldRenderScreenSpaceReflections(View);
 		if (bEnableSSR)
 		{
 			// RUN SSR
@@ -757,7 +760,8 @@ bool FDeferredShadingSceneRenderer::RenderSingleLayerWaterPass(FRHICommandListIm
 				View, 
 				WhiteForwardScreenSpaceShadowMask, 
 				&PassData.SceneWithoutSingleLayerWaterValidUVRect,
-				PassData.SceneColorAndDepthWithoutSingleLayerWater,
+				PassData.SceneColorWithoutSingleLayerWater,
+				PassData.SceneDepthWithoutSingleLayerWater,
 				WaterPassUniformBuffer);
 
 			FMeshPassProcessorRenderState DrawRenderState(View, WaterPassUniformBuffer);
diff --git a/Engine/Source/Runtime/Renderer/Private/SingleLayerWaterRendering.h b/Engine/Source/Runtime/Renderer/Private/SingleLayerWaterRendering.h
--- a/Engine/Source/Runtime/Renderer/Private/SingleLayerWaterRendering.h
+++ b/Engine/Source/Runtime/Renderer/Private/SingleLayerWaterRendering.h
@@ -11,18 +11,15 @@
 #include "MeshPassProcessor.h"
 #include "Containers/Array.h"
 
-
-
 class FViewInfo;
 
-
 struct FSingleLayerWaterPassData
 {
-	TRefCountPtr<IPooledRenderTarget> SceneColorAndDepthWithoutSingleLayerWater;
+	TRefCountPtr<IPooledRenderTarget> SceneColorWithoutSingleLayerWater;
+	TRefCountPtr<IPooledRenderTarget> SceneDepthWithoutSingleLayerWater;
 	FVector4 SceneWithoutSingleLayerWaterValidUVRect;
 };
 
-
 class FSingleLayerWaterPassMeshProcessor : public FMeshPassProcessor
 {
 public:
