Add RHIT stat for particle and atmosphere simulation

diff --git a/Engine/Source/Runtime/Renderer/Private/MobileShadingRenderer.cpp b/Engine/Source/Runtime/Renderer/Private/MobileShadingRenderer.cpp
--- a/Engine/Source/Runtime/Renderer/Private/MobileShadingRenderer.cpp
+++ b/Engine/Source/Runtime/Renderer/Private/MobileShadingRenderer.cpp
@@ -80,6 +80,7 @@ DECLARE_CYCLE_STAT(TEXT("Occlusion"), STAT_CLMM_Occlusion, STATGROUP_CommandList
 DECLARE_CYCLE_STAT(TEXT("Post"), STAT_CLMM_Post, STATGROUP_CommandListMarkers);
 DECLARE_CYCLE_STAT(TEXT("Translucency"), STAT_CLMM_Translucency, STATGROUP_CommandListMarkers);
 DECLARE_CYCLE_STAT(TEXT("Shadows"), STAT_CLMM_Shadows, STATGROUP_CommandListMarkers);
+DECLARE_CYCLE_STAT(TEXT("SceneSimulation"), STAT_CLMM_SceneSim, STATGROUP_CommandListMarkers);
 
 FGlobalDynamicIndexBuffer FMobileSceneRenderer::DynamicIndexBuffer;
 FGlobalDynamicVertexBuffer FMobileSceneRenderer::DynamicVertexBuffer;
@@ -401,6 +402,8 @@ void FMobileSceneRenderer::Render(FRHICommandListImmediate& RHICmdList)
 	DynamicReadBuffer.Commit();
 	RHICmdList.ImmediateFlush(EImmediateFlushType::DispatchToRHIThread);
 
+	RHICmdList.SetCurrentStat(GET_STATID(STAT_CLMM_SceneSim));
+
 	// Generate the Sky/Atmosphere look up tables
 	// Do compute work first, before any graphics work
 	const bool bShouldRenderSkyAtmosphere = ShouldRenderSkyAtmosphere(Scene, ViewFamily.EngineShowFlags);
