Water refraction - changed PF_R16F_FILTER to PF_R16F as it is always point sampled (if there is any different between those two on any RHI...)

diff --git a/Engine/Source/Runtime/Renderer/Private/SingleLayerWaterRendering.cpp b/Engine/Source/Runtime/Renderer/Private/SingleLayerWaterRendering.cpp
--- a/Engine/Source/Runtime/Renderer/Private/SingleLayerWaterRendering.cpp
+++ b/Engine/Source/Runtime/Renderer/Private/SingleLayerWaterRendering.cpp
@@ -436,7 +436,7 @@ void FDeferredShadingSceneRenderer::CopySingleLayerWaterTextures(FRHICommandList
 	FPooledRenderTargetDesc ColorDesc(FPooledRenderTargetDesc::Create2DDesc(RefractionResolution, SceneContext.GetSceneColorFormat(), SceneContext.GetDefaultColorClear(), TexCreate_None, TexCreate_RenderTargetable, false));
 	GRenderTargetPool.FindFreeElement(RHICmdList, ColorDesc, PassData.SceneColorWithoutSingleLayerWater, TEXT("SceneColorWithoutSingleLayerWater"));
 
-	FPooledRenderTargetDesc DepthDesc(FPooledRenderTargetDesc::Create2DDesc(RefractionResolution, GSingleLayerWaterRefractionFullPrecision ? PF_R32_FLOAT : PF_R16F_FILTER, SceneContext.GetDefaultColorClear(), TexCreate_None, TexCreate_RenderTargetable, false));
+	FPooledRenderTargetDesc DepthDesc(FPooledRenderTargetDesc::Create2DDesc(RefractionResolution, GSingleLayerWaterRefractionFullPrecision ? PF_R32_FLOAT : PF_R16F, SceneContext.GetDefaultColorClear(), TexCreate_None, TexCreate_RenderTargetable, false));
 	GRenderTargetPool.FindFreeElement(RHICmdList, DepthDesc, PassData.SceneDepthWithoutSingleLayerWater, TEXT("SceneDepthWithoutSingleLayerWater"));
 
 	// For now support only the 1st view
