Disable access to scene textures during depth rendering on Mobile

diff --git a/Engine/Source/Runtime/Renderer/Private/DepthRendering.h b/Engine/Source/Runtime/Renderer/Private/DepthRendering.h
--- a/Engine/Source/Runtime/Renderer/Private/DepthRendering.h
+++ b/Engine/Source/Runtime/Renderer/Private/DepthRendering.h
@@ -157,6 +157,17 @@ public:
 		BindSceneTextureUniformBufferDependentOnShadingPath(Initializer, PassUniformBuffer);
 	}
 
+	static void ModifyCompilationEnvironment(const FMaterialShaderPermutationParameters& Parameters, FShaderCompilerEnvironment& OutEnvironment)
+	{
+		FMeshMaterialShader::ModifyCompilationEnvironment(Parameters, OutEnvironment);
+		
+		if (IsMobilePlatform(Parameters.Platform))
+		{
+			// No access to scene textures during depth rendering on mobile
+			OutEnvironment.SetDefine(TEXT("SCENE_TEXTURES_DISABLED"), 1u);
+		}
+	}
+
 	FDepthOnlyPS() {}
 
 	void GetShaderBindings(
