GPU feature level support check (this allows HTML5 runtime to work)

#rnx

diff --git a/Engine/Source/Runtime/Renderer/Private/SkyAtmosphereRendering.cpp b/Engine/Source/Runtime/Renderer/Private/SkyAtmosphereRendering.cpp
--- a/Engine/Source/Runtime/Renderer/Private/SkyAtmosphereRendering.cpp
+++ b/Engine/Source/Runtime/Renderer/Private/SkyAtmosphereRendering.cpp
@@ -827,6 +827,10 @@ public:
 
 	virtual void InitRHI() override
 	{
+		if ( ! RHISupportsComputeShaders(GMaxRHIShaderPlatform) )
+		{
+			return;
+		}
 		const uint32 GroupSize = GetSampletCount();
 		const float GroupSizeInv = 1.0f / float(GroupSize);
 
@@ -858,7 +862,10 @@ public:
 
 	virtual void ReleaseRHI()
 	{
-		UniformSphereSamplesBuffer.Release();
+		if ( RHISupportsComputeShaders(GMaxRHIShaderPlatform) )
+		{
+			UniformSphereSamplesBuffer.Release();
+		}
 	}
 };
 TGlobalResource<FUniformSphereSamplesBuffer> GUniformSphereSamplesBuffer;
