Added half res scene color as SSR input. It is controlled by r.SSR.HalfResSceneColor and enabled by default on lower scalability settings. Greatly improves performance in bandwidth bound scenarios while retaining very similar visuals. On base X1 it brings SSR time from 1.15ms to 0.66ms.

diff --git a/Engine/Config/BaseScalability.ini b/Engine/Config/BaseScalability.ini
--- a/Engine/Config/BaseScalability.ini
+++ b/Engine/Config/BaseScalability.ini
@@ -379,6 +379,7 @@ r.Streaming.MaxEffectiveScreenSize=0
 r.TranslucencyLightingVolumeDim=24
 r.RefractionQuality=0
 r.SSR.Quality=0
+r.SSR.HalfResSceneColor=1
 r.SceneColorFormat=3
 r.DetailMode=0
 r.TranslucencyVolumeBlur=0
@@ -407,6 +408,7 @@ fx.NiagaraGlobalSpawnCountScale=0.125
 r.TranslucencyLightingVolumeDim=32
 r.RefractionQuality=0
 r.SSR.Quality=0
+r.SSR.HalfResSceneColor=1
 r.SceneColorFormat=3
 r.DetailMode=1
 r.TranslucencyVolumeBlur=0
@@ -435,6 +437,7 @@ fx.NiagaraGlobalSpawnCountScale=0.25
 r.TranslucencyLightingVolumeDim=48
 r.RefractionQuality=2
 r.SSR.Quality=2
+r.SSR.HalfResSceneColor=1
 r.SceneColorFormat=3
 r.DetailMode=1
 r.TranslucencyVolumeBlur=1
@@ -463,6 +466,7 @@ fx.NiagaraGlobalSpawnCountScale=0.5
 r.TranslucencyLightingVolumeDim=64
 r.RefractionQuality=2
 r.SSR.Quality=3
+r.SSR.HalfResSceneColor=0
 r.SceneColorFormat=4
 r.DetailMode=2
 r.TranslucencyVolumeBlur=1
@@ -491,6 +495,7 @@ fx.NiagaraGlobalSpawnCountScale=1.0
 r.TranslucencyLightingVolumeDim=64
 r.RefractionQuality=2
 r.SSR.Quality=4
+r.SSR.HalfResSceneColor=0
 r.SceneColorFormat=4
 r.DetailMode=2
 r.TranslucencyVolumeBlur=1
diff --git a/Engine/Source/Runtime/Renderer/Private/PostProcess/PostProcessing.cpp b/Engine/Source/Runtime/Renderer/Private/PostProcess/PostProcessing.cpp
--- a/Engine/Source/Runtime/Renderer/Private/PostProcess/PostProcessing.cpp
+++ b/Engine/Source/Runtime/Renderer/Private/PostProcess/PostProcessing.cpp
@@ -1206,6 +1206,12 @@ void FPostProcessing::Process(FRHICommandListImmediate& RHICmdList, const FViewI
 					HalfResolutionSceneColorViewRect = Outputs.Viewport;
 				}
 
+				if (!View.bStatePrevViewInfoIsReadOnly && HalfResolutionSceneColorTexture)
+				{
+					FSceneViewState* ViewState = (FSceneViewState*)View.State;
+					GraphBuilder.QueueTextureExtraction(HalfResolutionSceneColorTexture, &ViewState->PrevFrameViewInfo.HalfResTemporalAAHistory);
+				}
+
 				FEyeAdaptationParameters EyeAdaptationParameters = GetEyeAdaptationParameters(View, ERHIFeatureLevel::SM5);
 
 				// Default the new eye adaptation to the last one in case it's not generated this frame.
diff --git a/Engine/Source/Runtime/Renderer/Private/SceneRendering.h b/Engine/Source/Runtime/Renderer/Private/SceneRendering.h
--- a/Engine/Source/Runtime/Renderer/Private/SceneRendering.h
+++ b/Engine/Source/Runtime/Renderer/Private/SceneRendering.h
@@ -798,6 +798,9 @@ struct FPreviousViewInfo
 	// Temporal AA result of last frame
 	FTemporalAAHistory TemporalAAHistory;
 
+	// Half resolution version temporal AA result of last frame
+	TRefCountPtr<IPooledRenderTarget> HalfResTemporalAAHistory;
+
 	// Temporal AA history for diaphragm DOF.
 	FTemporalAAHistory DOFSetupHistory;
 	
diff --git a/Engine/Source/Runtime/Renderer/Private/ScreenSpaceRayTracing.cpp b/Engine/Source/Runtime/Renderer/Private/ScreenSpaceRayTracing.cpp
--- a/Engine/Source/Runtime/Renderer/Private/ScreenSpaceRayTracing.cpp
+++ b/Engine/Source/Runtime/Renderer/Private/ScreenSpaceRayTracing.cpp
@@ -21,6 +21,14 @@ static TAutoConsoleVariable<int32> CVarSSRQuality(
 	TEXT(" 4: very high (likely too slow for real-time)"),
 	ECVF_Scalability | ECVF_RenderThreadSafe);
 
+int32 GSSRHalfResSceneColor = 0;
+FAutoConsoleVariableRef CVarSSRHalfResSceneColor(
+	TEXT("r.SSR.HalfResSceneColor"),
+	GSSRHalfResSceneColor,
+	TEXT("Use half res scene color as input for SSR. Improves performance without much of a visual quality loss."),
+	ECVF_Scalability | ECVF_RenderThreadSafe
+);
+
 static TAutoConsoleVariable<int32> CVarSSRTemporal(
 	TEXT("r.SSR.Temporal"),
 	0,
@@ -551,6 +559,10 @@ void RenderScreenSpaceReflections(
 		{
 			InputColor = GraphBuilder.RegisterExternalTexture(View.PrevViewInfo.CustomSSRInput);
 		}
+		else if (GSSRHalfResSceneColor && View.PrevViewInfo.HalfResTemporalAAHistory.IsValid())
+		{
+			InputColor = GraphBuilder.RegisterExternalTexture(View.PrevViewInfo.HalfResTemporalAAHistory);
+		}
 		else if (View.PrevViewInfo.TemporalAAHistory.IsValid())
 		{
 			InputColor = GraphBuilder.RegisterExternalTexture(View.PrevViewInfo.TemporalAAHistory.RT[0]);
@@ -682,7 +694,7 @@ void RenderScreenSpaceReflections(
 		PassParameters->PrevSceneColorPreExposureCorrection = InputColor != CurrentSceneColor ? View.PreExposure / View.PrevViewInfo.SceneColorPreExposure : 1.0f;
 
 		PassParameters->SceneColor = InputColor;
-		PassParameters->SceneColorSampler = TStaticSamplerState<SF_Point>::GetRHI();
+		PassParameters->SceneColorSampler = GSSRHalfResSceneColor ? TStaticSamplerState<SF_Bilinear>::GetRHI() : TStaticSamplerState<SF_Point>::GetRHI();
 
 		PassParameters->HZB = GraphBuilder.RegisterExternalTexture(View.HZB);
 		PassParameters->HZBSampler = TStaticSamplerState<SF_Point>::GetRHI();
