Deprecating ARRAY_COUNT and changing it to UE_ARRAY_COUNT.

diff --git a/Engine/Source/Runtime/OpenGLDrv/Private/Android/AndroidEGL.cpp b/Engine/Source/Runtime/OpenGLDrv/Private/Android/AndroidEGL.cpp
--- a/Engine/Source/Runtime/OpenGLDrv/Private/Android/AndroidEGL.cpp
+++ b/Engine/Source/Runtime/OpenGLDrv/Private/Android/AndroidEGL.cpp
@@ -70,7 +70,7 @@ void VerifyEGLResult(EGLint ErrorCode, const TCHAR* Msg1, const TCHAR* Msg2, con
 			TEXT("UNKNOWN EGL ERROR")
 		};
 
-		uint32 ErrorIndex = FMath::Min<uint32>(ErrorCode - EGL_SUCCESS, ARRAY_COUNT(EGLErrorStrings) - 1);
+		uint32 ErrorIndex = FMath::Min<uint32>(ErrorCode - EGL_SUCCESS, UE_ARRAY_COUNT(EGLErrorStrings) - 1);
 		UE_LOG(LogRHI, Warning, TEXT("%s(%u): %s%s failed with error %s (0x%x)"),
 			Filename, Line, Msg1, Msg2, EGLErrorStrings[ErrorIndex], ErrorCode);
 		check(0);
diff --git a/Engine/Source/Runtime/OpenGLDrv/Private/OpenGLShaders.cpp b/Engine/Source/Runtime/OpenGLDrv/Private/OpenGLShaders.cpp
--- a/Engine/Source/Runtime/OpenGLDrv/Private/OpenGLShaders.cpp
+++ b/Engine/Source/Runtime/OpenGLDrv/Private/OpenGLShaders.cpp
@@ -5184,7 +5184,7 @@ void FOpenGLProgramBinaryCache::CompilePendingShaders(const FOpenGLLinkedProgram
 {
 	if (CachePtr)
 	{
-		for (int32 StageIdx = 0; StageIdx < ARRAY_COUNT(Config.Shaders); ++StageIdx)
+		for (int32 StageIdx = 0; StageIdx < UE_ARRAY_COUNT(Config.Shaders); ++StageIdx)
 		{
 			GLuint ShaderResource = Config.Shaders[StageIdx].Resource;
 			FPendingShaderCode* PendingShaderCodePtr = CachePtr->ShadersPendingCompilation.Find(ShaderResource);
diff --git a/Engine/Source/Runtime/OpenGLDrv/Private/OpenGLUtil.cpp b/Engine/Source/Runtime/OpenGLDrv/Private/OpenGLUtil.cpp
--- a/Engine/Source/Runtime/OpenGLDrv/Private/OpenGLUtil.cpp
+++ b/Engine/Source/Runtime/OpenGLDrv/Private/OpenGLUtil.cpp
@@ -25,7 +25,7 @@ void VerifyOpenGLResult(GLenum ErrorCode, const TCHAR* Msg1, const TCHAR* Msg2,
 			TEXT("UNKNOWN ERROR")
 		};
 
-		uint32 ErrorIndex = FMath::Min<uint32>(ErrorCode - GL_INVALID_ENUM, ARRAY_COUNT(ErrorStrings) - 1);
+		uint32 ErrorIndex = FMath::Min<uint32>(ErrorCode - GL_INVALID_ENUM, UE_ARRAY_COUNT(ErrorStrings) - 1);
 		UE_LOG(LogRHI,Fatal,TEXT("%s(%u): %s%s failed with error %s (0x%x)"),
 			Filename,Line,Msg1,Msg2,ErrorStrings[ErrorIndex],ErrorCode);
 	}
diff --git a/Engine/Source/Runtime/OpenGLDrv/Public/OpenGLDrv.h b/Engine/Source/Runtime/OpenGLDrv/Public/OpenGLDrv.h
--- a/Engine/Source/Runtime/OpenGLDrv/Public/OpenGLDrv.h
+++ b/Engine/Source/Runtime/OpenGLDrv/Public/OpenGLDrv.h
@@ -1168,7 +1168,7 @@ private:
 	{
 		if (FOpenGL::NeedsVertexAttribRemapTable())
 		{
-			check(VertexAttributeIndex < ARRAY_COUNT(PendingState.BoundShaderState->GetVertexShader()->Bindings.VertexAttributeRemap));
+			check(VertexAttributeIndex < UE_ARRAY_COUNT(PendingState.BoundShaderState->GetVertexShader()->Bindings.VertexAttributeRemap));
 			VertexAttributeIndex = PendingState.BoundShaderState->GetVertexShader()->Bindings.VertexAttributeRemap[VertexAttributeIndex];
 		}
 		check(VertexAttributeIndex < NUM_OPENGL_VERTEX_STREAMS); // check that this attribute has remaped correctly.
@@ -1179,7 +1179,7 @@ private:
 	{
 		if (FOpenGL::NeedsVertexAttribRemapTable())
 		{
-			check(VertexAttributeIndex < ARRAY_COUNT(Bindings.VertexAttributeRemap));
+			check(VertexAttributeIndex < UE_ARRAY_COUNT(Bindings.VertexAttributeRemap));
 			VertexAttributeIndex = Bindings.VertexAttributeRemap[VertexAttributeIndex];
 		}
 		check(VertexAttributeIndex < NUM_OPENGL_VERTEX_STREAMS); // check that this attribute has remaped correctly.
diff --git a/Engine/Source/Runtime/OpenGLDrv/Public/OpenGLShaderResources.h b/Engine/Source/Runtime/OpenGLDrv/Public/OpenGLShaderResources.h
--- a/Engine/Source/Runtime/OpenGLDrv/Public/OpenGLShaderResources.h
+++ b/Engine/Source/Runtime/OpenGLDrv/Public/OpenGLShaderResources.h
@@ -208,7 +208,7 @@ inline FArchive& operator<<(FArchive& Ar, FOpenGLShaderBindings& Bindings)
 	Ar << Bindings.NumUniformBuffers;
 	Ar << Bindings.NumUAVs;
 	Ar << Bindings.bFlattenUB;
-	for (uint32 i = 0; i < ARRAY_COUNT(Bindings.VertexAttributeRemap); i++)
+	for (uint32 i = 0; i < UE_ARRAY_COUNT(Bindings.VertexAttributeRemap); i++)
 	{
 		Ar << Bindings.VertexAttributeRemap[i];
 	}
diff --git a/Engine/Source/Runtime/PakFile/Public/IPlatformFilePak.h b/Engine/Source/Runtime/PakFile/Public/IPlatformFilePak.h
--- a/Engine/Source/Runtime/PakFile/Public/IPlatformFilePak.h
+++ b/Engine/Source/Runtime/PakFile/Public/IPlatformFilePak.h
@@ -2477,7 +2477,7 @@ struct FPakSignatureFile
 	{
 		ChunkHashes = InChunkHashes;
 		DecryptedHash = ComputeCurrentMasterHash();
-		FRSA::EncryptPrivate(TArrayView<const uint8>(DecryptedHash.Hash, ARRAY_COUNT(FSHAHash::Hash)), EncryptedHash, InKey);
+		FRSA::EncryptPrivate(TArrayView<const uint8>(DecryptedHash.Hash, UE_ARRAY_COUNT(FSHAHash::Hash)), EncryptedHash, InKey);
 	}
 
 	/**
@@ -2514,7 +2514,7 @@ struct FPakSignatureFile
 		{
 			TArray<uint8> Decrypted;
 			int32 BytesDecrypted = FRSA::DecryptPublic(EncryptedHash, Decrypted, InKey);
-			if (BytesDecrypted == ARRAY_COUNT(FSHAHash::Hash))
+			if (BytesDecrypted == UE_ARRAY_COUNT(FSHAHash::Hash))
 			{
 				FSHAHash CurrentHash = ComputeCurrentMasterHash();
 				if (FMemory::Memcmp(Decrypted.GetData(), CurrentHash.Hash, Decrypted.Num()) == 0)
diff --git a/Engine/Source/Runtime/Portal/LauncherCheck/Private/LauncherCheckModule.cpp b/Engine/Source/Runtime/Portal/LauncherCheck/Private/LauncherCheckModule.cpp
--- a/Engine/Source/Runtime/Portal/LauncherCheck/Private/LauncherCheckModule.cpp
+++ b/Engine/Source/Runtime/Portal/LauncherCheck/Private/LauncherCheckModule.cpp
@@ -104,7 +104,7 @@ private:
 		FString EncodedExePath;
 
 		static const TCHAR* Delims[] = { TEXT("/") };
-		static const int32 NumDelims = ARRAY_COUNT(Delims);
+		static const int32 NumDelims = UE_ARRAY_COUNT(Delims);
 
 		// Get the path to the executable, and encode it ('cos spaces, colons, etc break things)
 		FString ExePath = FPlatformProcess::BaseDir();
diff --git a/Engine/Source/Runtime/RHI/Private/RHI.cpp b/Engine/Source/Runtime/RHI/Private/RHI.cpp
--- a/Engine/Source/Runtime/RHI/Private/RHI.cpp
+++ b/Engine/Source/Runtime/RHI/Private/RHI.cpp
@@ -582,11 +582,11 @@ FName FeatureLevelNames[] =
 	FName(TEXT("SM5")),
 };
 
-static_assert(ARRAY_COUNT(FeatureLevelNames) == ERHIFeatureLevel::Num, "Missing entry from feature level names.");
+static_assert(UE_ARRAY_COUNT(FeatureLevelNames) == ERHIFeatureLevel::Num, "Missing entry from feature level names.");
 
 RHI_API bool GetFeatureLevelFromName(FName Name, ERHIFeatureLevel::Type& OutFeatureLevel)
 {
-	for (int32 NameIndex = 0; NameIndex < ARRAY_COUNT(FeatureLevelNames); NameIndex++)
+	for (int32 NameIndex = 0; NameIndex < UE_ARRAY_COUNT(FeatureLevelNames); NameIndex++)
 	{
 		if (FeatureLevelNames[NameIndex] == Name)
 		{
@@ -601,8 +601,8 @@ RHI_API bool GetFeatureLevelFromName(FName Name, ERHIFeatureLevel::Type& OutFeat
 
 RHI_API void GetFeatureLevelName(ERHIFeatureLevel::Type InFeatureLevel, FString& OutName)
 {
-	check(InFeatureLevel < ARRAY_COUNT(FeatureLevelNames));
-	if (InFeatureLevel < ARRAY_COUNT(FeatureLevelNames))
+	check(InFeatureLevel < UE_ARRAY_COUNT(FeatureLevelNames));
+	if (InFeatureLevel < UE_ARRAY_COUNT(FeatureLevelNames))
 	{
 		FeatureLevelNames[(int32)InFeatureLevel].ToString(OutName);
 	}
@@ -615,8 +615,8 @@ RHI_API void GetFeatureLevelName(ERHIFeatureLevel::Type InFeatureLevel, FString&
 static FName InvalidFeatureLevelName(TEXT("InvalidFeatureLevel"));
 RHI_API void GetFeatureLevelName(ERHIFeatureLevel::Type InFeatureLevel, FName& OutName)
 {
-	check(InFeatureLevel < ARRAY_COUNT(FeatureLevelNames));
-	if (InFeatureLevel < ARRAY_COUNT(FeatureLevelNames))
+	check(InFeatureLevel < UE_ARRAY_COUNT(FeatureLevelNames));
+	if (InFeatureLevel < UE_ARRAY_COUNT(FeatureLevelNames))
 	{
 		OutName = FeatureLevelNames[(int32)InFeatureLevel];
 	}
@@ -634,11 +634,11 @@ FName ShadingPathNames[] =
 	FName(TEXT("Mobile")),
 };
 
-static_assert(ARRAY_COUNT(ShadingPathNames) == ERHIShadingPath::Num, "Missing entry from shading path names.");
+static_assert(UE_ARRAY_COUNT(ShadingPathNames) == ERHIShadingPath::Num, "Missing entry from shading path names.");
 
 RHI_API bool GetShadingPathFromName(FName Name, ERHIShadingPath::Type& OutShadingPath)
 {
-	for (int32 NameIndex = 0; NameIndex < ARRAY_COUNT(ShadingPathNames); NameIndex++)
+	for (int32 NameIndex = 0; NameIndex < UE_ARRAY_COUNT(ShadingPathNames); NameIndex++)
 	{
 		if (ShadingPathNames[NameIndex] == Name)
 		{
@@ -653,8 +653,8 @@ RHI_API bool GetShadingPathFromName(FName Name, ERHIShadingPath::Type& OutShadin
 
 RHI_API void GetShadingPathName(ERHIShadingPath::Type InShadingPath, FString& OutName)
 {
-	check(InShadingPath < ARRAY_COUNT(ShadingPathNames));
-	if (InShadingPath < ARRAY_COUNT(ShadingPathNames))
+	check(InShadingPath < UE_ARRAY_COUNT(ShadingPathNames));
+	if (InShadingPath < UE_ARRAY_COUNT(ShadingPathNames))
 	{
 		ShadingPathNames[(int32)InShadingPath].ToString(OutName);
 	}
@@ -667,8 +667,8 @@ RHI_API void GetShadingPathName(ERHIShadingPath::Type InShadingPath, FString& Ou
 static FName InvalidShadingPathName(TEXT("InvalidShadingPath"));
 RHI_API void GetShadingPathName(ERHIShadingPath::Type InShadingPath, FName& OutName)
 {
-	check(InShadingPath < ARRAY_COUNT(ShadingPathNames));
-	if (InShadingPath < ARRAY_COUNT(ShadingPathNames))
+	check(InShadingPath < UE_ARRAY_COUNT(ShadingPathNames));
+	if (InShadingPath < UE_ARRAY_COUNT(ShadingPathNames))
 	{
 		OutName = ShadingPathNames[(int32)InShadingPath];
 	}
diff --git a/Engine/Source/Runtime/Renderer/Private/CapsuleShadowRendering.cpp b/Engine/Source/Runtime/Renderer/Private/CapsuleShadowRendering.cpp
--- a/Engine/Source/Runtime/Renderer/Private/CapsuleShadowRendering.cpp
+++ b/Engine/Source/Runtime/Renderer/Private/CapsuleShadowRendering.cpp
@@ -158,7 +158,7 @@ public:
 
 		FRHIUnorderedAccessView* OutUAVs[1];
 		OutUAVs[0] = ComputedLightDirectionDataBuffer.UAV;
-		RHICmdList.TransitionResources(EResourceTransitionAccess::EWritable, EResourceTransitionPipeline::EComputeToCompute, OutUAVs, ARRAY_COUNT(OutUAVs));
+		RHICmdList.TransitionResources(EResourceTransitionAccess::EWritable, EResourceTransitionPipeline::EComputeToCompute, OutUAVs, UE_ARRAY_COUNT(OutUAVs));
 
 		SetShaderValue(RHICmdList, ShaderRHI, NumLightDirectionData, NumLightDirectionDataValue);
 		SetSRVParameter(RHICmdList, ShaderRHI, LightDirectionData, LightDirectionDataSRV);
@@ -178,7 +178,7 @@ public:
 
 		FRHIUnorderedAccessView* OutUAVs[1];
 		OutUAVs[0] = ComputedLightDirectionDataBuffer.UAV;
-		RHICmdList.TransitionResources(EResourceTransitionAccess::EReadable, EResourceTransitionPipeline::EComputeToCompute, OutUAVs, ARRAY_COUNT(OutUAVs) );
+		RHICmdList.TransitionResources(EResourceTransitionAccess::EReadable, EResourceTransitionPipeline::EComputeToCompute, OutUAVs, UE_ARRAY_COUNT(OutUAVs) );
 	}
 
 	// FShader interface.
@@ -341,7 +341,7 @@ public:
 			OutUAVs[1] = TileIntersectionCountsBuffer->UAV;
 		}
 		
-		RHICmdList.TransitionResources(EResourceTransitionAccess::ERWBarrier, EResourceTransitionPipeline::EComputeToCompute, OutUAVs, TileIntersectionCountsBuffer ? ARRAY_COUNT(OutUAVs) : 1);
+		RHICmdList.TransitionResources(EResourceTransitionAccess::ERWBarrier, EResourceTransitionPipeline::EComputeToCompute, OutUAVs, TileIntersectionCountsBuffer ? UE_ARRAY_COUNT(OutUAVs) : 1);
 
 		if (ShadowingType == ShapeShadow_MovableSkylightTiledCulling)
 		{
@@ -448,7 +448,7 @@ public:
 			OutUAVs[1] = TileIntersectionCountsBuffer->UAV;
 		}
 		
-		RHICmdList.TransitionResources(EResourceTransitionAccess::EReadable, EResourceTransitionPipeline::EComputeToCompute, OutUAVs, TileIntersectionCountsBuffer ? ARRAY_COUNT(OutUAVs) : 1);
+		RHICmdList.TransitionResources(EResourceTransitionAccess::EReadable, EResourceTransitionPipeline::EComputeToCompute, OutUAVs, TileIntersectionCountsBuffer ? UE_ARRAY_COUNT(OutUAVs) : 1);
 	}
 
 	// FShader interface.
diff --git a/Engine/Source/Runtime/Renderer/Private/CompositionLighting/PostProcessDeferredDecals.cpp b/Engine/Source/Runtime/Renderer/Private/CompositionLighting/PostProcessDeferredDecals.cpp
--- a/Engine/Source/Runtime/Renderer/Private/CompositionLighting/PostProcessDeferredDecals.cpp
+++ b/Engine/Source/Runtime/Renderer/Private/CompositionLighting/PostProcessDeferredDecals.cpp
@@ -287,7 +287,7 @@ bool RenderPreStencil(FRenderingCompositePassContext& Context, const FMatrix& Co
 	Context.RHICmdList.SetStreamSource(0, GetUnitCubeVertexBuffer(), 0);
 
 	// Render decal mask
-	Context.RHICmdList.DrawIndexedPrimitive(GetUnitCubeIndexBuffer(), 0, 0, 8, 0, ARRAY_COUNT(GCubeIndices) / 3, 1);
+	Context.RHICmdList.DrawIndexedPrimitive(GetUnitCubeIndexBuffer(), 0, 0, 8, 0, UE_ARRAY_COUNT(GCubeIndices) / 3, 1);
 
 	return true;
 }
@@ -756,7 +756,7 @@ void FRCPassPostProcessDeferredDecals::Process(FRenderingCompositePassContext& C
 					FDecalRendering::SetShader(RHICmdList, GraphicsPSOInit, View, DecalData, CurrentStage, FrustumComponentToClip);
 					RHICmdList.SetStencilRef(StencilRef);
 
-					RHICmdList.DrawIndexedPrimitive(GetUnitCubeIndexBuffer(), 0, 0, 8, 0, ARRAY_COUNT(GCubeIndices) / 3, 1);
+					RHICmdList.DrawIndexedPrimitive(GetUnitCubeIndexBuffer(), 0, 0, 8, 0, UE_ARRAY_COUNT(GCubeIndices) / 3, 1);
 					RenderTargetManager.bGufferADirty |= (RenderTargetManager.TargetsToResolve[FDecalRenderTargetManager::GBufferAIndex] != nullptr);
 				}
 
diff --git a/Engine/Source/Runtime/Renderer/Private/ConvertToUniformMesh.cpp b/Engine/Source/Runtime/Renderer/Private/ConvertToUniformMesh.cpp
--- a/Engine/Source/Runtime/Renderer/Private/ConvertToUniformMesh.cpp
+++ b/Engine/Source/Runtime/Renderer/Private/ConvertToUniformMesh.cpp
@@ -363,7 +363,7 @@ public:
 
 		FRHIUnorderedAccessView* UniformMeshUAVs[1];
 		UniformMeshUAVs[0] = Scene->DistanceFieldSceneData.SurfelBuffers->Surfels.UAV;
-		RHICmdList.TransitionResources(EResourceTransitionAccess::ERWBarrier, EResourceTransitionPipeline::EComputeToCompute, UniformMeshUAVs, ARRAY_COUNT(UniformMeshUAVs));
+		RHICmdList.TransitionResources(EResourceTransitionAccess::ERWBarrier, EResourceTransitionPipeline::EComputeToCompute, UniformMeshUAVs, UE_ARRAY_COUNT(UniformMeshUAVs));
 
 		SurfelBufferParameters.Set(RHICmdList, ShaderRHI, *Scene->DistanceFieldSceneData.SurfelBuffers, *Scene->DistanceFieldSceneData.InstancedSurfelBuffers);
 		
@@ -381,7 +381,7 @@ public:
 		const FScene* Scene = (const FScene*)View.Family->Scene;
 		FRHIUnorderedAccessView* UniformMeshUAVs[1];
 		UniformMeshUAVs[0] = Scene->DistanceFieldSceneData.SurfelBuffers->Surfels.UAV;
-		RHICmdList.TransitionResources(EResourceTransitionAccess::EReadable, EResourceTransitionPipeline::EComputeToCompute, UniformMeshUAVs, ARRAY_COUNT(UniformMeshUAVs));
+		RHICmdList.TransitionResources(EResourceTransitionAccess::EReadable, EResourceTransitionPipeline::EComputeToCompute, UniformMeshUAVs, UE_ARRAY_COUNT(UniformMeshUAVs));
 	}
 
 	virtual bool Serialize(FArchive& Ar) override
diff --git a/Engine/Source/Runtime/Renderer/Private/CustomDepthRendering.cpp b/Engine/Source/Runtime/Renderer/Private/CustomDepthRendering.cpp
--- a/Engine/Source/Runtime/Renderer/Private/CustomDepthRendering.cpp
+++ b/Engine/Source/Runtime/Renderer/Private/CustomDepthRendering.cpp
@@ -88,7 +88,7 @@ void FCustomDepthPassMeshProcessor::AddMeshBatch(const FMeshBatch& RESTRICT Mesh
 				TStaticDepthStencilState<true, CF_DepthNearOrEqual, true, CF_Always, SO_Keep, SO_Replace, SO_Replace, false, CF_Always, SO_Keep, SO_Keep, SO_Keep, 255, 64>::GetRHI(),
 				TStaticDepthStencilState<true, CF_DepthNearOrEqual, true, CF_Always, SO_Keep, SO_Replace, SO_Replace, false, CF_Always, SO_Keep, SO_Keep, SO_Keep, 255, 128>::GetRHI()
 			};
-			checkSlow(EStencilMask::SM_Count == ARRAY_COUNT(StencilStates));
+			checkSlow(EStencilMask::SM_Count == UE_ARRAY_COUNT(StencilStates));
 
 			PassDrawRenderState.SetDepthStencilState(StencilStates[(int32)PrimitiveSceneProxy->GetStencilWriteMask()]);
 			PassDrawRenderState.SetStencilRef(CustomDepthStencilValue);
diff --git a/Engine/Source/Runtime/Renderer/Private/DistanceFieldAmbientOcclusion.cpp b/Engine/Source/Runtime/Renderer/Private/DistanceFieldAmbientOcclusion.cpp
--- a/Engine/Source/Runtime/Renderer/Private/DistanceFieldAmbientOcclusion.cpp
+++ b/Engine/Source/Runtime/Renderer/Private/DistanceFieldAmbientOcclusion.cpp
@@ -205,18 +205,18 @@ float TemporalHalton2( int32 Index, int32 Base )
 
 void GetSpacedVectors(uint32 FrameNumber, TArray<FVector, TInlineAllocator<9> >& OutVectors)
 {
-	OutVectors.Empty(ARRAY_COUNT(SpacedVectors9));
+	OutVectors.Empty(UE_ARRAY_COUNT(SpacedVectors9));
 
 	if (GAOSampleSet == 0)
 	{
-		for (int32 i = 0; i < ARRAY_COUNT(SpacedVectors9); i++)
+		for (int32 i = 0; i < UE_ARRAY_COUNT(SpacedVectors9); i++)
 		{
 			OutVectors.Add(SpacedVectors9[i]);
 		}
 	}
 	else
 	{
-		for (int32 i = 0; i < ARRAY_COUNT(RelaxedSpacedVectors9); i++)
+		for (int32 i = 0; i < UE_ARRAY_COUNT(RelaxedSpacedVectors9); i++)
 		{
 			OutVectors.Add(RelaxedSpacedVectors9[i]);
 		}
@@ -240,7 +240,7 @@ void GetSpacedVectors(uint32 FrameNumber, TArray<FVector, TInlineAllocator<9> >&
 }
 
 // Cone half angle derived from each cone covering an equal solid angle
-float GAOConeHalfAngle = FMath::Acos(1 - 1.0f / (float)ARRAY_COUNT(SpacedVectors9));
+float GAOConeHalfAngle = FMath::Acos(1 - 1.0f / (float)UE_ARRAY_COUNT(SpacedVectors9));
 
 // Number of AO sample positions along each cone
 // Must match shader code
@@ -497,14 +497,14 @@ void GenerateBestSpacedVectors()
 	{
 		bGenerated = true;
 
-		FVector OriginalSpacedVectors9[ARRAY_COUNT(SpacedVectors9)];
+		FVector OriginalSpacedVectors9[UE_ARRAY_COUNT(SpacedVectors9)];
 
-		for (int32 i = 0; i < ARRAY_COUNT(OriginalSpacedVectors9); i++)
+		for (int32 i = 0; i < UE_ARRAY_COUNT(OriginalSpacedVectors9); i++)
 		{
 			OriginalSpacedVectors9[i] = SpacedVectors9[i];
 		}
 
-		float CosHalfAngle = 1 - 1.0f / (float)ARRAY_COUNT(OriginalSpacedVectors9);
+		float CosHalfAngle = 1 - 1.0f / (float)UE_ARRAY_COUNT(OriginalSpacedVectors9);
 		// Used to prevent self-shadowing on a plane
 		float AngleBias = .03f;
 		float MinAngle = FMath::Acos(CosHalfAngle) + AngleBias;
@@ -513,11 +513,11 @@ void GenerateBestSpacedVectors()
 		// Relaxation iterations by repulsion
 		for (int32 Iteration = 0; Iteration < 10000; Iteration++)
 		{
-			for (int32 i = 0; i < ARRAY_COUNT(OriginalSpacedVectors9); i++)
+			for (int32 i = 0; i < UE_ARRAY_COUNT(OriginalSpacedVectors9); i++)
 			{
 				FVector Force(0.0f, 0.0f, 0.0f);
 
-				for (int32 j = 0; j < ARRAY_COUNT(OriginalSpacedVectors9); j++)
+				for (int32 j = 0; j < UE_ARRAY_COUNT(OriginalSpacedVectors9); j++)
 				{
 					if (i != j)
 					{
@@ -539,7 +539,7 @@ void GenerateBestSpacedVectors()
 			}
 		}
 
-		for (int32 i = 0; i < ARRAY_COUNT(OriginalSpacedVectors9); i++)
+		for (int32 i = 0; i < UE_ARRAY_COUNT(OriginalSpacedVectors9); i++)
 		{
 			UE_LOG(LogDistanceField, Log, TEXT("FVector(%f, %f, %f),"), OriginalSpacedVectors9[i].X, OriginalSpacedVectors9[i].Y, OriginalSpacedVectors9[i].Z);
 		}
@@ -557,7 +557,7 @@ void GenerateBestSpacedVectors()
 		// HemisphereSolidAngle = 2 * PI
 		// ConeSolidAngle = 2 * PI * (1 - cos(ConeHalfAngle))
 		// cos(ConeHalfAngle) = 1 - 1 / NumCones
-		float CosHalfAngle = 1 - 1.0f / (float)ARRAY_COUNT(BestSpacedVectors9);
+		float CosHalfAngle = 1 - 1.0f / (float)UE_ARRAY_COUNT(BestSpacedVectors9);
 		// Prevent self-intersection in sample set
 		float MinAngle = FMath::Acos(CosHalfAngle);
 		float MinZ = FMath::Sin(MinAngle);
@@ -566,9 +566,9 @@ void GenerateBestSpacedVectors()
 		// Super slow random brute force search
 		for (int i = 0; i < 1000000; i++)
 		{
-			FVector CandidateSpacedVectors[ARRAY_COUNT(BestSpacedVectors9)];
+			FVector CandidateSpacedVectors[UE_ARRAY_COUNT(BestSpacedVectors9)];
 
-			for (int j = 0; j < ARRAY_COUNT(CandidateSpacedVectors); j++)
+			for (int j = 0; j < UE_ARRAY_COUNT(CandidateSpacedVectors); j++)
 			{
 				FVector NewSample;
 
@@ -598,7 +598,7 @@ void GenerateBestSpacedVectors()
 
 				bool bIntersects = false;
 
-				for (int j = 0; j < ARRAY_COUNT(CandidateSpacedVectors); j++)
+				for (int j = 0; j < UE_ARRAY_COUNT(CandidateSpacedVectors); j++)
 				{
 					if (FVector::DotProduct(CandidateSpacedVectors[j], NewSample) > CosHalfAngle)
 					{
@@ -614,7 +614,7 @@ void GenerateBestSpacedVectors()
 			{
 				BestCoverage = Coverage;
 
-				for (int j = 0; j < ARRAY_COUNT(CandidateSpacedVectors); j++)
+				for (int j = 0; j < UE_ARRAY_COUNT(CandidateSpacedVectors); j++)
 				{
 					BestSpacedVectors9[j] = CandidateSpacedVectors[j];
 				}
diff --git a/Engine/Source/Runtime/Renderer/Private/DistanceFieldGlobalIllumination.cpp b/Engine/Source/Runtime/Renderer/Private/DistanceFieldGlobalIllumination.cpp
--- a/Engine/Source/Runtime/Renderer/Private/DistanceFieldGlobalIllumination.cpp
+++ b/Engine/Source/Runtime/Renderer/Private/DistanceFieldGlobalIllumination.cpp
@@ -152,7 +152,7 @@ public:
 		FRHIUnorderedAccessView* OutUAVs[2];
 		OutUAVs[0] = GVPLResources.VPLParameterBuffer.UAV;
 		OutUAVs[1] = GVPLResources.VPLData.UAV;
-		RHICmdList.TransitionResources(EResourceTransitionAccess::ERWBarrier, EResourceTransitionPipeline::EComputeToCompute, OutUAVs, ARRAY_COUNT(OutUAVs));
+		RHICmdList.TransitionResources(EResourceTransitionAccess::ERWBarrier, EResourceTransitionPipeline::EComputeToCompute, OutUAVs, UE_ARRAY_COUNT(OutUAVs));
 
 		VPLParameterBuffer.SetBuffer(RHICmdList, ShaderRHI, GVPLResources.VPLParameterBuffer);
 		VPLData.SetBuffer(RHICmdList, ShaderRHI, GVPLResources.VPLData);
@@ -175,7 +175,7 @@ public:
 		FRHIUnorderedAccessView* OutUAVs[2];
 		OutUAVs[0] = GVPLResources.VPLParameterBuffer.UAV;
 		OutUAVs[1] = GVPLResources.VPLData.UAV;
-		RHICmdList.TransitionResources(EResourceTransitionAccess::EReadable, EResourceTransitionPipeline::EComputeToCompute, OutUAVs, ARRAY_COUNT(OutUAVs));
+		RHICmdList.TransitionResources(EResourceTransitionAccess::EReadable, EResourceTransitionPipeline::EComputeToCompute, OutUAVs, UE_ARRAY_COUNT(OutUAVs));
 	}
 
 	// FShader interface.
@@ -314,7 +314,7 @@ public:
 		FRHIUnorderedAccessView* OutUAVs[2];
 		OutUAVs[0] = GCulledVPLResources.VPLParameterBuffer.UAV;
 		OutUAVs[1] = GCulledVPLResources.VPLData.UAV;
-		RHICmdList.TransitionResources(EResourceTransitionAccess::ERWBarrier, EResourceTransitionPipeline::EComputeToCompute, OutUAVs, ARRAY_COUNT(OutUAVs));
+		RHICmdList.TransitionResources(EResourceTransitionAccess::ERWBarrier, EResourceTransitionPipeline::EComputeToCompute, OutUAVs, UE_ARRAY_COUNT(OutUAVs));
 
 		CulledVPLParameterBuffer.SetBuffer(RHICmdList, ShaderRHI, GCulledVPLResources.VPLParameterBuffer);
 		CulledVPLData.SetBuffer(RHICmdList, ShaderRHI, GCulledVPLResources.VPLData);
diff --git a/Engine/Source/Runtime/Renderer/Private/DistanceFieldLightingShared.h b/Engine/Source/Runtime/Renderer/Private/DistanceFieldLightingShared.h
--- a/Engine/Source/Runtime/Renderer/Private/DistanceFieldLightingShared.h
+++ b/Engine/Source/Runtime/Renderer/Private/DistanceFieldLightingShared.h
@@ -163,7 +163,7 @@ public:
 			FRHIUnorderedAccessView* OutUAVs[2];
 			OutUAVs[0] = ObjectBuffers.Bounds.UAV;
 			OutUAVs[1] = ObjectBuffers.Data.UAV;
-			RHICmdList.TransitionResources(EResourceTransitionAccess::ERWBarrier, EResourceTransitionPipeline::EComputeToCompute, OutUAVs, ARRAY_COUNT(OutUAVs));
+			RHICmdList.TransitionResources(EResourceTransitionAccess::ERWBarrier, EResourceTransitionPipeline::EComputeToCompute, OutUAVs, UE_ARRAY_COUNT(OutUAVs));
 		}
 
 		SceneObjectBounds.SetBuffer(RHICmdList, ShaderRHI, ObjectBuffers.Bounds);
@@ -197,7 +197,7 @@ public:
 			FRHIUnorderedAccessView* OutUAVs[2];
 			OutUAVs[0] = ObjectBuffers.Bounds.UAV;
 			OutUAVs[1] = ObjectBuffers.Data.UAV;
-			RHICmdList.TransitionResources(EResourceTransitionAccess::EReadable, EResourceTransitionPipeline::EComputeToCompute, OutUAVs, ARRAY_COUNT(OutUAVs));
+			RHICmdList.TransitionResources(EResourceTransitionAccess::EReadable, EResourceTransitionPipeline::EComputeToCompute, OutUAVs, UE_ARRAY_COUNT(OutUAVs));
 		}
 	}
 
diff --git a/Engine/Source/Runtime/Renderer/Private/DistanceFieldObjectCulling.cpp b/Engine/Source/Runtime/Renderer/Private/DistanceFieldObjectCulling.cpp
--- a/Engine/Source/Runtime/Renderer/Private/DistanceFieldObjectCulling.cpp
+++ b/Engine/Source/Runtime/Renderer/Private/DistanceFieldObjectCulling.cpp
@@ -143,7 +143,7 @@ public:
 		OutUAVs[3] = GAOCulledObjectBuffers.Buffers.BoxBounds.UAV;
 		OutUAVs[4] = Scene->DistanceFieldSceneData.GetCurrentObjectBuffers()->Data.UAV;
 		OutUAVs[5] = Scene->DistanceFieldSceneData.GetCurrentObjectBuffers()->Bounds.UAV;
-		RHICmdList.TransitionResources(EResourceTransitionAccess::ERWBarrier, EResourceTransitionPipeline::EComputeToCompute, OutUAVs, ARRAY_COUNT(OutUAVs));
+		RHICmdList.TransitionResources(EResourceTransitionAccess::ERWBarrier, EResourceTransitionPipeline::EComputeToCompute, OutUAVs, UE_ARRAY_COUNT(OutUAVs));
 
 		FRHIComputeShader* ShaderRHI = GetComputeShader();
 		FGlobalShader::SetParameters<FViewUniformShaderParameters>(RHICmdList, ShaderRHI, View.ViewUniformBuffer);
@@ -170,7 +170,7 @@ public:
 		OutUAVs[3] = GAOCulledObjectBuffers.Buffers.BoxBounds.UAV;
 		OutUAVs[4] = Scene->DistanceFieldSceneData.GetCurrentObjectBuffers()->Data.UAV;
 		OutUAVs[5] = Scene->DistanceFieldSceneData.GetCurrentObjectBuffers()->Bounds.UAV;		
-		RHICmdList.TransitionResources(EResourceTransitionAccess::ERWBarrier, EResourceTransitionPipeline::EComputeToCompute, OutUAVs, ARRAY_COUNT(OutUAVs));
+		RHICmdList.TransitionResources(EResourceTransitionAccess::ERWBarrier, EResourceTransitionPipeline::EComputeToCompute, OutUAVs, UE_ARRAY_COUNT(OutUAVs));
 	}
 
 	virtual bool Serialize(FArchive& Ar)
@@ -274,7 +274,7 @@ public:
 		FRHIUnorderedAccessView* OutUAVs[2];
 		OutUAVs[0] = TileIntersectionResources->TileConeAxisAndCos.UAV;
 		OutUAVs[1] = TileIntersectionResources->TileConeDepthRanges.UAV;
-		RHICmdList.TransitionResources(EResourceTransitionAccess::ERWBarrier, EResourceTransitionPipeline::EComputeToCompute, OutUAVs, ARRAY_COUNT(OutUAVs));
+		RHICmdList.TransitionResources(EResourceTransitionAccess::ERWBarrier, EResourceTransitionPipeline::EComputeToCompute, OutUAVs, UE_ARRAY_COUNT(OutUAVs));
 
 		TileConeAxisAndCos.SetBuffer(RHICmdList, ShaderRHI, TileIntersectionResources->TileConeAxisAndCos);
 		TileConeDepthRanges.SetBuffer(RHICmdList, ShaderRHI, TileIntersectionResources->TileConeDepthRanges);
@@ -303,7 +303,7 @@ public:
 		FRHIUnorderedAccessView* OutUAVs[2];
 		OutUAVs[0] = TileIntersectionResources->TileConeAxisAndCos.UAV;
 		OutUAVs[1] = TileIntersectionResources->TileConeDepthRanges.UAV;
-		RHICmdList.TransitionResources(EResourceTransitionAccess::EReadable, EResourceTransitionPipeline::EComputeToCompute, OutUAVs, ARRAY_COUNT(OutUAVs));
+		RHICmdList.TransitionResources(EResourceTransitionAccess::EReadable, EResourceTransitionPipeline::EComputeToCompute, OutUAVs, UE_ARRAY_COUNT(OutUAVs));
 	}
 
 	virtual bool Serialize(FArchive& Ar)
diff --git a/Engine/Source/Runtime/Renderer/Private/DistanceFieldObjectManagement.cpp b/Engine/Source/Runtime/Renderer/Private/DistanceFieldObjectManagement.cpp
--- a/Engine/Source/Runtime/Renderer/Private/DistanceFieldObjectManagement.cpp
+++ b/Engine/Source/Runtime/Renderer/Private/DistanceFieldObjectManagement.cpp
@@ -198,7 +198,7 @@ public:
 		FRHIUnorderedAccessView* OutUAVs[2];
 		OutUAVs[0] = ObjectBuffers.Bounds.UAV;
 		OutUAVs[1] = ObjectBuffers.Data.UAV;
-		RHICmdList.TransitionResources(EResourceTransitionAccess::EReadable, EResourceTransitionPipeline::EComputeToCompute, OutUAVs, ARRAY_COUNT(OutUAVs));
+		RHICmdList.TransitionResources(EResourceTransitionAccess::EReadable, EResourceTransitionPipeline::EComputeToCompute, OutUAVs, UE_ARRAY_COUNT(OutUAVs));
 	}
 
 	virtual bool Serialize(FArchive& Ar) override
@@ -256,7 +256,7 @@ public:
 		FRHIUnorderedAccessView* OutUAVs[2];
 		OutUAVs[0] = ObjectBuffersDest.Bounds.UAV;
 		OutUAVs[1] = ObjectBuffersDest.Data.UAV;
-		RHICmdList.TransitionResources(EResourceTransitionAccess::ERWBarrier, EResourceTransitionPipeline::EComputeToCompute, OutUAVs, ARRAY_COUNT(OutUAVs));
+		RHICmdList.TransitionResources(EResourceTransitionAccess::ERWBarrier, EResourceTransitionPipeline::EComputeToCompute, OutUAVs, UE_ARRAY_COUNT(OutUAVs));
 
 		CopyObjectBounds.SetBuffer(RHICmdList, ShaderRHI, ObjectBuffersDest.Bounds);
 		CopyObjectData.SetBuffer(RHICmdList, ShaderRHI, ObjectBuffersDest.Data);
@@ -273,7 +273,7 @@ public:
 		FRHIUnorderedAccessView* OutUAVs[2];
 		OutUAVs[0] = ObjectBuffersDest.Bounds.UAV;
 		OutUAVs[1] = ObjectBuffersDest.Data.UAV;
-		RHICmdList.TransitionResources(EResourceTransitionAccess::EReadable, EResourceTransitionPipeline::EComputeToCompute, OutUAVs, ARRAY_COUNT(OutUAVs));
+		RHICmdList.TransitionResources(EResourceTransitionAccess::EReadable, EResourceTransitionPipeline::EComputeToCompute, OutUAVs, UE_ARRAY_COUNT(OutUAVs));
 	}
 
 	virtual bool Serialize(FArchive& Ar) override
@@ -330,7 +330,7 @@ public:
 		FRHIUnorderedAccessView* OutUAVs[2];
 		OutUAVs[0] = SurfelBuffersDest.InterpolatedVertexData.UAV;
 		OutUAVs[1] = SurfelBuffersDest.Surfels.UAV;		
-		RHICmdList.TransitionResources(EResourceTransitionAccess::ERWBarrier, EResourceTransitionPipeline::EComputeToCompute, OutUAVs, ARRAY_COUNT(OutUAVs));
+		RHICmdList.TransitionResources(EResourceTransitionAccess::ERWBarrier, EResourceTransitionPipeline::EComputeToCompute, OutUAVs, UE_ARRAY_COUNT(OutUAVs));
 
 		CopyInterpolatedVertexData.SetBuffer(RHICmdList, ShaderRHI, SurfelBuffersDest.InterpolatedVertexData);
 		CopySurfelData.SetBuffer(RHICmdList, ShaderRHI, SurfelBuffersDest.Surfels);
@@ -347,7 +347,7 @@ public:
 		FRHIUnorderedAccessView* OutUAVs[2];
 		OutUAVs[0] = SurfelBuffersDest.InterpolatedVertexData.UAV;
 		OutUAVs[1] = SurfelBuffersDest.Surfels.UAV;
-		RHICmdList.TransitionResources(EResourceTransitionAccess::EReadable, EResourceTransitionPipeline::EComputeToCompute, OutUAVs, ARRAY_COUNT(OutUAVs));
+		RHICmdList.TransitionResources(EResourceTransitionAccess::EReadable, EResourceTransitionPipeline::EComputeToCompute, OutUAVs, UE_ARRAY_COUNT(OutUAVs));
 	}
 
 	virtual bool Serialize(FArchive& Ar) override
diff --git a/Engine/Source/Runtime/Renderer/Private/DistanceFieldShadowing.cpp b/Engine/Source/Runtime/Renderer/Private/DistanceFieldShadowing.cpp
--- a/Engine/Source/Runtime/Renderer/Private/DistanceFieldShadowing.cpp
+++ b/Engine/Source/Runtime/Renderer/Private/DistanceFieldShadowing.cpp
@@ -143,11 +143,11 @@ public:
 		OutUAVs[1] = GShadowCulledObjectBuffers.Buffers.Bounds.UAV;
 		OutUAVs[2] = GShadowCulledObjectBuffers.Buffers.Data.UAV;
 		OutUAVs[3] = GShadowCulledObjectBuffers.Buffers.BoxBounds.UAV;		
-		RHICmdList.TransitionResources(EResourceTransitionAccess::ERWBarrier, EResourceTransitionPipeline::EComputeToCompute, OutUAVs, ARRAY_COUNT(OutUAVs));
+		RHICmdList.TransitionResources(EResourceTransitionAccess::ERWBarrier, EResourceTransitionPipeline::EComputeToCompute, OutUAVs, UE_ARRAY_COUNT(OutUAVs));
 
 		CulledObjectParameters.Set(RHICmdList, ShaderRHI, GShadowCulledObjectBuffers.Buffers);
 
-		SetShaderValue(RHICmdList, ShaderRHI, ObjectBoundingGeometryIndexCount, ARRAY_COUNT(GCubeIndices));
+		SetShaderValue(RHICmdList, ShaderRHI, ObjectBoundingGeometryIndexCount, UE_ARRAY_COUNT(GCubeIndices));
 		SetShaderValue(RHICmdList, ShaderRHI, WorldToShadow, WorldToShadowValue);
 		SetShaderValue(RHICmdList, ShaderRHI, ShadowBoundingSphere, ShadowBoundingSphereValue);
 
@@ -172,7 +172,7 @@ public:
 		OutUAVs[1] = GShadowCulledObjectBuffers.Buffers.Bounds.UAV;
 		OutUAVs[2] = GShadowCulledObjectBuffers.Buffers.Data.UAV;
 		OutUAVs[3] = GShadowCulledObjectBuffers.Buffers.BoxBounds.UAV;
-		RHICmdList.TransitionResources(EResourceTransitionAccess::EReadable, EResourceTransitionPipeline::EComputeToCompute, OutUAVs, ARRAY_COUNT(OutUAVs));
+		RHICmdList.TransitionResources(EResourceTransitionAccess::EReadable, EResourceTransitionPipeline::EComputeToCompute, OutUAVs, UE_ARRAY_COUNT(OutUAVs));
 	}
 
 	virtual bool Serialize(FArchive& Ar) override
diff --git a/Engine/Source/Runtime/Renderer/Private/GPUBenchmark.cpp b/Engine/Source/Runtime/Renderer/Private/GPUBenchmark.cpp
--- a/Engine/Source/Runtime/Renderer/Private/GPUBenchmark.cpp
+++ b/Engine/Source/Runtime/Renderer/Private/GPUBenchmark.cpp
@@ -454,7 +454,7 @@ void RendererGPUBenchmark(FRHICommandListImmediate& RHICmdList, FSynthBenchmarkR
 	{
 		// larger number means more accuracy but slower, some slower GPUs might timeout with a number to large
 		const uint32 IterationCount = 70;
-		const uint32 MethodCount = ARRAY_COUNT(InOut.GPUStats);
+		const uint32 MethodCount = UE_ARRAY_COUNT(InOut.GPUStats);
 
 		enum class EMethodType
 		{
@@ -484,10 +484,10 @@ void RendererGPUBenchmark(FRHICommandListImmediate& RHICmdList, FSynthBenchmarkR
 			{ TEXT("VertThroughPut2"),  1.0f / 1.767f,  TEXT("s/GigaVert"), 0.0f, EMethodType::Vertex }, // TODO: Set weights
 		};
 
-		static_assert(ARRAY_COUNT(Methods) == ARRAY_COUNT(InOut.GPUStats), "Benchmark methods descriptor array lengths should match.");
+		static_assert(UE_ARRAY_COUNT(Methods) == UE_ARRAY_COUNT(InOut.GPUStats), "Benchmark methods descriptor array lengths should match.");
 
 		// Initialize the GPU benchmark stats
-		for (int32 Index = 0; Index < ARRAY_COUNT(Methods); ++Index)
+		for (int32 Index = 0; Index < UE_ARRAY_COUNT(Methods); ++Index)
 		{
 			auto& Method = Methods[Index];
 			InOut.GPUStats[Index] = FSynthBenchmarkStat(Method.Desc, Method.IndexNormalizedTime, Method.ValueType, Method.Weight);
diff --git a/Engine/Source/Runtime/Renderer/Private/GlobalDistanceField.cpp b/Engine/Source/Runtime/Renderer/Private/GlobalDistanceField.cpp
--- a/Engine/Source/Runtime/Renderer/Private/GlobalDistanceField.cpp
+++ b/Engine/Source/Runtime/Renderer/Private/GlobalDistanceField.cpp
@@ -191,7 +191,7 @@ public:
 		OutUAVs[1] = GGlobalDistanceFieldCulledObjectBuffers.Buffers.Bounds.UAV;
 		OutUAVs[2] = GGlobalDistanceFieldCulledObjectBuffers.Buffers.Data.UAV;
 		OutUAVs[3] = GGlobalDistanceFieldCulledObjectBuffers.Buffers.BoxBounds.UAV;
-		RHICmdList.TransitionResources(EResourceTransitionAccess::ERWBarrier, EResourceTransitionPipeline::EComputeToCompute, OutUAVs, ARRAY_COUNT(OutUAVs));
+		RHICmdList.TransitionResources(EResourceTransitionAccess::ERWBarrier, EResourceTransitionPipeline::EComputeToCompute, OutUAVs, UE_ARRAY_COUNT(OutUAVs));
 
 		CulledObjectParameters.Set(RHICmdList, ShaderRHI, GGlobalDistanceFieldCulledObjectBuffers.Buffers);
 
diff --git a/Engine/Source/Runtime/Renderer/Private/HeightfieldLighting.cpp b/Engine/Source/Runtime/Renderer/Private/HeightfieldLighting.cpp
--- a/Engine/Source/Runtime/Renderer/Private/HeightfieldLighting.cpp
+++ b/Engine/Source/Runtime/Renderer/Private/HeightfieldLighting.cpp
@@ -475,7 +475,7 @@ void FHeightfieldLightingViewInfo::SetupVisibleHeightfields(const FViewInfo& Vie
 
 					RHICmdList.SetViewport(0, 0, 0.0f, LightingAtlasSize.X, LightingAtlasSize.Y, 1.0f);
 
-					FRHIRenderPassInfo RPInfo(ARRAY_COUNT(RenderTargets), RenderTargets, ERenderTargetActions::Clear_Store);
+					FRHIRenderPassInfo RPInfo(UE_ARRAY_COUNT(RenderTargets), RenderTargets, ERenderTargetActions::Clear_Store);
 					RHICmdList.BeginRenderPass(RPInfo, TEXT("SetupHeightfields"));
 					{
 						FGraphicsPipelineStateInitializer GraphicsPSOInit;
@@ -513,7 +513,7 @@ void FHeightfieldLightingViewInfo::SetupVisibleHeightfields(const FViewInfo& Vie
 					}
 					RHICmdList.EndRenderPass();
 
-					RHICmdList.TransitionResources(EResourceTransitionAccess::EReadable, RenderTargets, ARRAY_COUNT(RenderTargets));
+					RHICmdList.TransitionResources(EResourceTransitionAccess::EReadable, RenderTargets, UE_ARRAY_COUNT(RenderTargets));
 				}
 			}
 		}
diff --git a/Engine/Source/Runtime/Renderer/Private/LightMapRendering.cpp b/Engine/Source/Runtime/Renderer/Private/LightMapRendering.cpp
--- a/Engine/Source/Runtime/Renderer/Private/LightMapRendering.cpp
+++ b/Engine/Source/Runtime/Renderer/Private/LightMapRendering.cpp
@@ -193,7 +193,7 @@ void InterpolateVolumetricLightmap(
 	
 	auto ReadSHCoefficient = [&BrickTextureCoordinate, &VolumetricLightmapData, &AmbientVector](uint32 CoefficientIndex)
 	{
-		check(CoefficientIndex < ARRAY_COUNT(VolumetricLightmapData.BrickData.SHCoefficients));
+		check(CoefficientIndex < UE_ARRAY_COUNT(VolumetricLightmapData.BrickData.SHCoefficients));
 
 		// Undo normalization done in FIrradianceBrickData::SetFromVolumeLightingSample
 		const FLinearColor SHDenormalizationScales0(
diff --git a/Engine/Source/Runtime/Renderer/Private/MeshPassProcessor.cpp b/Engine/Source/Runtime/Renderer/Private/MeshPassProcessor.cpp
--- a/Engine/Source/Runtime/Renderer/Private/MeshPassProcessor.cpp
+++ b/Engine/Source/Runtime/Renderer/Private/MeshPassProcessor.cpp
@@ -108,7 +108,7 @@ void FMeshDrawShaderBindings::SetShaderBindings(
 	for (int32 UniformBufferIndex = 0; UniformBufferIndex < NumUniformBuffers; UniformBufferIndex++)
 	{
 		FShaderParameterInfo Parameter = UniformBufferParameters[UniformBufferIndex];
-		checkSlow(Parameter.BaseIndex < ARRAY_COUNT(ShaderBindingState.UniformBuffers));
+		checkSlow(Parameter.BaseIndex < UE_ARRAY_COUNT(ShaderBindingState.UniformBuffers));
 		FRHIUniformBuffer* UniformBuffer = UniformBufferBindings[UniformBufferIndex];
 
 		if (UniformBuffer != ShaderBindingState.UniformBuffers[Parameter.BaseIndex])
@@ -126,7 +126,7 @@ void FMeshDrawShaderBindings::SetShaderBindings(
 	for (int32 SamplerIndex = 0; SamplerIndex < NumTextureSamplers; SamplerIndex++)
 	{
 		FShaderParameterInfo Parameter = TextureSamplerParameters[SamplerIndex];
-		checkSlow(Parameter.BaseIndex < ARRAY_COUNT(ShaderBindingState.Samplers));
+		checkSlow(Parameter.BaseIndex < UE_ARRAY_COUNT(ShaderBindingState.Samplers));
 		FRHISamplerState* Sampler = (FRHISamplerState*)SamplerBindings[SamplerIndex];
 
 		if (Sampler != ShaderBindingState.Samplers[Parameter.BaseIndex])
@@ -145,7 +145,7 @@ void FMeshDrawShaderBindings::SetShaderBindings(
 	for (uint32 SRVIndex = 0; SRVIndex < NumSRVs; SRVIndex++)
 	{
 		FShaderParameterInfo Parameter = SRVParameters[SRVIndex];
-		checkSlow(Parameter.BaseIndex < ARRAY_COUNT(ShaderBindingState.SRVs));
+		checkSlow(Parameter.BaseIndex < UE_ARRAY_COUNT(ShaderBindingState.SRVs));
 
 		uint32 TypeByteIndex = SRVIndex / 8;
 		uint32 TypeBitIndex = SRVIndex - TypeByteIndex;
@@ -290,9 +290,9 @@ void FMeshDrawShaderBindings::SetRayTracingShaderBindingsForHitGroup(
 	for (int32 UniformBufferIndex = 0; UniformBufferIndex < NumUniformBuffers; UniformBufferIndex++)
 	{
 		FShaderParameterInfo Parameter = UniformBufferParameters[UniformBufferIndex];
-		checkSlow(Parameter.BaseIndex < ARRAY_COUNT(LocalUniformBuffers));
+		checkSlow(Parameter.BaseIndex < UE_ARRAY_COUNT(LocalUniformBuffers));
 		const FRHIUniformBuffer* UniformBuffer = UniformBufferBindings[UniformBufferIndex];
-		if (Parameter.BaseIndex < ARRAY_COUNT(LocalUniformBuffers))
+		if (Parameter.BaseIndex < UE_ARRAY_COUNT(LocalUniformBuffers))
 		{
 			LocalUniformBuffers[Parameter.BaseIndex] = UniformBuffer;
 			MaxUniformBufferUsed = FMath::Max((int32)Parameter.BaseIndex, MaxUniformBufferUsed);
@@ -439,13 +439,13 @@ public:
 		StencilRef = -1;
 
 		// Vertex streams must be reset if PSO changes.
-		for (int32 VertexStreamIndex = 0; VertexStreamIndex < ARRAY_COUNT(VertexStreams); ++VertexStreamIndex)
+		for (int32 VertexStreamIndex = 0; VertexStreamIndex < UE_ARRAY_COUNT(VertexStreams); ++VertexStreamIndex)
 		{
 			VertexStreams[VertexStreamIndex].VertexBuffer = nullptr;
 		}
 
 		// Shader bindings must be reset if PSO changes
-		for (int32 FrequencyIndex = 0; FrequencyIndex < ARRAY_COUNT(ShaderBindings); FrequencyIndex++)
+		for (int32 FrequencyIndex = 0; FrequencyIndex < UE_ARRAY_COUNT(ShaderBindings); FrequencyIndex++)
 		{
 			FShaderBindingState& RESTRICT ShaderBinding = ShaderBindings[FrequencyIndex];
 
@@ -706,7 +706,7 @@ void FMeshDrawShaderBindings::Release()
 	}
 #endif
 
-	if (Size > ARRAY_COUNT(InlineStorage))
+	if (Size > UE_ARRAY_COUNT(InlineStorage))
 	{
 		delete[] HeapData;
 	}
diff --git a/Engine/Source/Runtime/Renderer/Private/MobileDecalRendering.cpp b/Engine/Source/Runtime/Renderer/Private/MobileDecalRendering.cpp
--- a/Engine/Source/Runtime/Renderer/Private/MobileDecalRendering.cpp
+++ b/Engine/Source/Runtime/Renderer/Private/MobileDecalRendering.cpp
@@ -132,7 +132,7 @@ void FMobileSceneRenderer::RenderDecals(FRHICommandListImmediate& RHICmdList)
 				// Set shader params
 				FDecalRendering::SetShader(RHICmdList, GraphicsPSOInit, View, DecalData, DRS_Mobile, FrustumComponentToClip);
 			
-				RHICmdList.DrawIndexedPrimitive(GetUnitCubeIndexBuffer(), 0, 0, 8, 0, ARRAY_COUNT(GCubeIndices) / 3, 1);
+				RHICmdList.DrawIndexedPrimitive(GetUnitCubeIndexBuffer(), 0, 0, 8, 0, UE_ARRAY_COUNT(GCubeIndices) / 3, 1);
 			}
 		}
 	}
diff --git a/Engine/Source/Runtime/Renderer/Private/MobileReflectionEnvironmentCapture.cpp b/Engine/Source/Runtime/Renderer/Private/MobileReflectionEnvironmentCapture.cpp
--- a/Engine/Source/Runtime/Renderer/Private/MobileReflectionEnvironmentCapture.cpp
+++ b/Engine/Source/Runtime/Renderer/Private/MobileReflectionEnvironmentCapture.cpp
@@ -253,7 +253,7 @@ namespace MobileReflectionEnvironmentCapture
 					SetGraphicsPipelineState(RHICmdList, GraphicsPSOInit);
 
 					FLinearColor UnusedColors[1] = { FLinearColor::Black };
-					PixelShader->SetColors(RHICmdList, UnusedColors, ARRAY_COUNT(UnusedColors));
+					PixelShader->SetColors(RHICmdList, UnusedColors, UE_ARRAY_COUNT(UnusedColors));
 
 					DrawRectangle(
 						RHICmdList,
diff --git a/Engine/Source/Runtime/Renderer/Private/MobileShadingRenderer.cpp b/Engine/Source/Runtime/Renderer/Private/MobileShadingRenderer.cpp
--- a/Engine/Source/Runtime/Renderer/Private/MobileShadingRenderer.cpp
+++ b/Engine/Source/Runtime/Renderer/Private/MobileShadingRenderer.cpp
@@ -967,7 +967,7 @@ void FMobileSceneRenderer::UpdateDirectionalLightUniformBuffers(FRHICommandListI
 {
 	bool bDynamicShadows = ViewFamily.EngineShowFlags.DynamicShadows;
 	// Fill in the other entries based on the lights
-	for (int32 ChannelIdx = 0; ChannelIdx < ARRAY_COUNT(Scene->MobileDirectionalLights); ChannelIdx++)
+	for (int32 ChannelIdx = 0; ChannelIdx < UE_ARRAY_COUNT(Scene->MobileDirectionalLights); ChannelIdx++)
 	{
 		FMobileDirectionalLightShaderParameters Params;
 		SetupMobileDirectionalLightUniformParameters(*Scene, View, VisibleLightInfos, ChannelIdx, bDynamicShadows, Params);
@@ -996,7 +996,7 @@ void FMobileSceneRenderer::CreateDirectionalLightUniformBuffers(FViewInfo& View)
 	// First array entry is used for primitives with no lighting channel set
 	View.MobileDirectionalLightUniformBuffers[0] = TUniformBufferRef<FMobileDirectionalLightShaderParameters>::CreateUniformBufferImmediate(FMobileDirectionalLightShaderParameters(), UniformBuffer_SingleFrame);
 	// Fill in the other entries based on the lights
-	for (int32 ChannelIdx = 0; ChannelIdx < ARRAY_COUNT(Scene->MobileDirectionalLights); ChannelIdx++)
+	for (int32 ChannelIdx = 0; ChannelIdx < UE_ARRAY_COUNT(Scene->MobileDirectionalLights); ChannelIdx++)
 	{
 		FMobileDirectionalLightShaderParameters Params;
 		SetupMobileDirectionalLightUniformParameters(*Scene, View, VisibleLightInfos, ChannelIdx, bDynamicShadows, Params);
diff --git a/Engine/Source/Runtime/Renderer/Private/PlanarReflectionRendering.cpp b/Engine/Source/Runtime/Renderer/Private/PlanarReflectionRendering.cpp
--- a/Engine/Source/Runtime/Renderer/Private/PlanarReflectionRendering.cpp
+++ b/Engine/Source/Runtime/Renderer/Private/PlanarReflectionRendering.cpp
@@ -91,7 +91,7 @@ void SetupPlanarReflectionUniformParameters(const class FSceneView& View, const
 		// Instanced stereo needs both view's values available at once
 		if (ReflectionSceneProxy->bIsStereo || View.Family->Views.Num() == 1)
 		{
-			static_assert(ARRAY_COUNT(ReflectionSceneProxy->ProjectionWithExtraFOV) == 2 
+			static_assert(UE_ARRAY_COUNT(ReflectionSceneProxy->ProjectionWithExtraFOV) == 2 
 				&& GPlanarReflectionUniformMaxReflectionViews == 2, "Code assumes max 2 planar reflection views.");
 
 			OutParameters.ProjectionWithExtraFOV[0] = ReflectionSceneProxy->ProjectionWithExtraFOV[0];
diff --git a/Engine/Source/Runtime/Renderer/Private/PostProcess/DiaphragmDOF.cpp b/Engine/Source/Runtime/Renderer/Private/PostProcess/DiaphragmDOF.cpp
--- a/Engine/Source/Runtime/Renderer/Private/PostProcess/DiaphragmDOF.cpp
+++ b/Engine/Source/Runtime/Renderer/Private/PostProcess/DiaphragmDOF.cpp
@@ -324,7 +324,7 @@ const TCHAR* GetEventName(EDiaphragmDOFLayerProcessing e)
 		TEXT("FocusOnly"),
 	};
 	int32 i = int32(e);
-	check(i < ARRAY_COUNT(kArray));
+	check(i < UE_ARRAY_COUNT(kArray));
 	return kArray[i];
 }
 
@@ -335,7 +335,7 @@ const TCHAR* GetEventName(EDiaphragmDOFPostfilterMethod e)
 		TEXT("Max3x3"),
 	};
 	int32 i = int32(e) - 1;
-	check(i < ARRAY_COUNT(kArray));
+	check(i < UE_ARRAY_COUNT(kArray));
 	return kArray[i];
 }
 
@@ -347,7 +347,7 @@ const TCHAR* GetEventName(EDiaphragmDOFBokehSimulation e)
 		TEXT("Generic"),
 	};
 	int32 i = int32(e);
-	check(i < ARRAY_COUNT(kArray));
+	check(i < UE_ARRAY_COUNT(kArray));
 	return kArray[i];
 }
 
@@ -359,7 +359,7 @@ const TCHAR* GetEventName(EDiaphragmDOFBokehLUTFormat e)
 		TEXT("Gather"),
 	};
 	int32 i = int32(e);
-	check(i < ARRAY_COUNT(kArray));
+	check(i < UE_ARRAY_COUNT(kArray));
 	return kArray[i];
 }
 
@@ -372,7 +372,7 @@ const TCHAR* GetEventName(EDiaphragmDOFGatherQuality e)
 		TEXT("Cinematic"),
 	};
 	int32 i = int32(e);
-	check(i < ARRAY_COUNT(kArray));
+	check(i < UE_ARRAY_COUNT(kArray));
 	return kArray[i];
 }
 
@@ -384,7 +384,7 @@ const TCHAR* GetEventName(EDiaphragmDOFDilateCocMode e)
 		TEXT("MinAbs"),
 	};
 	int32 i = int32(e);
-	check(i < ARRAY_COUNT(kArray));
+	check(i < UE_ARRAY_COUNT(kArray));
 	return kArray[i];
 }
 
@@ -1710,7 +1710,7 @@ FRDGTextureRef DiaphragmDOF::AddPasses(
 			SampleRadiusCount[0] = CurrentConvolutionRadius;
 
 			// If the theoric radius is too big, setup more dilate passes.
-			for (int32 i = 1; i < ARRAY_COUNT(SampleDistanceMultiplier); i++)
+			for (int32 i = 1; i < UE_ARRAY_COUNT(SampleDistanceMultiplier); i++)
 			{
 				if (MaximumTileDilation <= CurrentConvolutionRadius)
 				{
diff --git a/Engine/Source/Runtime/Renderer/Private/PostProcess/PostProcessMaterial.cpp b/Engine/Source/Runtime/Renderer/Private/PostProcess/PostProcessMaterial.cpp
--- a/Engine/Source/Runtime/Renderer/Private/PostProcess/PostProcessMaterial.cpp
+++ b/Engine/Source/Runtime/Renderer/Private/PostProcess/PostProcessMaterial.cpp
@@ -146,7 +146,7 @@ FRHIDepthStencilState* GetMaterialStencilState(const FMaterial* Material)
 		TStaticDepthStencilState<false, CF_Always, true, CF_Never>::GetRHI(),
 		TStaticDepthStencilState<false, CF_Always, true, CF_Always>::GetRHI(),
 	};
-	static_assert(EMaterialStencilCompare::MSC_Count == ARRAY_COUNT(StencilStates), "Ensure that all EMaterialStencilCompare values are accounted for.");
+	static_assert(EMaterialStencilCompare::MSC_Count == UE_ARRAY_COUNT(StencilStates), "Ensure that all EMaterialStencilCompare values are accounted for.");
 
 	check(Material);
 
@@ -172,7 +172,7 @@ FRHIBlendState* GetMaterialBlendState(const FMaterial* Material)
 		TStaticBlendState<CW_RGBA, BO_Add, BF_One, BF_InverseSourceAlpha, BO_Add, BF_One, BF_InverseSourceAlpha>::GetRHI(),
 		TStaticBlendState<CW_RGBA, BO_Add, BF_Zero, BF_InverseSourceAlpha, BO_Add, BF_Zero, BF_InverseSourceAlpha>::GetRHI(),
 	};
-	static_assert(EBlendMode::BLEND_MAX == ARRAY_COUNT(BlendStates), "Ensure that all EBlendMode values are accounted for.");
+	static_assert(EBlendMode::BLEND_MAX == UE_ARRAY_COUNT(BlendStates), "Ensure that all EBlendMode values are accounted for.");
 
 	check(Material);
 
diff --git a/Engine/Source/Runtime/Renderer/Private/PostProcess/PostProcessTemporalAA.cpp b/Engine/Source/Runtime/Renderer/Private/PostProcess/PostProcessTemporalAA.cpp
--- a/Engine/Source/Runtime/Renderer/Private/PostProcess/PostProcessTemporalAA.cpp
+++ b/Engine/Source/Runtime/Renderer/Private/PostProcess/PostProcessTemporalAA.cpp
@@ -348,8 +348,8 @@ const TCHAR* const kTAAPassNames[] = {
 };
 
 
-static_assert(ARRAY_COUNT(kTAAOutputNames) == int32(ETAAPassConfig::MAX), "Missing TAA output name.");
-static_assert(ARRAY_COUNT(kTAAPassNames) == int32(ETAAPassConfig::MAX), "Missing TAA pass name.");
+static_assert(UE_ARRAY_COUNT(kTAAOutputNames) == int32(ETAAPassConfig::MAX), "Missing TAA output name.");
+static_assert(UE_ARRAY_COUNT(kTAAPassNames) == int32(ETAAPassConfig::MAX), "Missing TAA pass name.");
 
 
 FIntPoint FTAAPassParameters::GetOutputExtent() const
diff --git a/Engine/Source/Runtime/Renderer/Private/PostProcess/PostProcessing.cpp b/Engine/Source/Runtime/Renderer/Private/PostProcess/PostProcessing.cpp
--- a/Engine/Source/Runtime/Renderer/Private/PostProcess/PostProcessing.cpp
+++ b/Engine/Source/Runtime/Renderer/Private/PostProcess/PostProcessing.cpp
@@ -214,7 +214,7 @@ public:
 			TEXT("Scene(1/32)"),
 			TEXT("Scene(1/64)")
 		};
-		static_assert(ARRAY_COUNT(PassNames) == StageCount, "PassNames size must equal StageCount");
+		static_assert(UE_ARRAY_COUNT(PassNames) == StageCount, "PassNames size must equal StageCount");
 
 		// The first stage is the input.
 		Textures[0] = HalfResolutionSceneTexture;
@@ -411,8 +411,8 @@ FBloomOutputs AddBloomPass(
 				6  // Q5
 			};
 
-			static_assert(ARRAY_COUNT(BloomStages) == BloomQualityCountMax, "Array must be one less than the number of bloom quality entries.");
-			static_assert(ARRAY_COUNT(BloomQualityToSceneDownsampleStage) == BloomQualityCountMax, "Array must be one less than the number of bloom quality entries.");
+			static_assert(UE_ARRAY_COUNT(BloomStages) == BloomQualityCountMax, "Array must be one less than the number of bloom quality entries.");
+			static_assert(UE_ARRAY_COUNT(BloomQualityToSceneDownsampleStage) == BloomQualityCountMax, "Array must be one less than the number of bloom quality entries.");
 
 			// Use bloom quality to select the number of downsample stages to use for bloom.
 			const uint32 BloomStageCount = BloomQualityToSceneDownsampleStage[BloomQualityIndex];
diff --git a/Engine/Source/Runtime/Renderer/Private/PostProcess/SceneRenderTargets.cpp b/Engine/Source/Runtime/Renderer/Private/PostProcess/SceneRenderTargets.cpp
--- a/Engine/Source/Runtime/Renderer/Private/PostProcess/SceneRenderTargets.cpp
+++ b/Engine/Source/Runtime/Renderer/Private/PostProcess/SceneRenderTargets.cpp
@@ -318,7 +318,7 @@ inline const TCHAR* GetSceneColorTargetName(EShadingPath ShadingPath)
 		TEXT("SceneColorMobile"), 
 		TEXT("SceneColorDeferred")
 	};
-	check((uint32)ShadingPath < ARRAY_COUNT(SceneColorNames));
+	check((uint32)ShadingPath < UE_ARRAY_COUNT(SceneColorNames));
 	return SceneColorNames[(uint32)ShadingPath];
 }
 
@@ -2298,7 +2298,7 @@ const FTexture2DRHIRef& FSceneRenderTargets::GetOptionalShadowDepthColorSurface(
 {
 	// Look for matching resolution
 	int32 EmptySlot = -1;
-	for (int32 Index = 0; Index < ARRAY_COUNT(OptionalShadowDepthColor); Index++)
+	for (int32 Index = 0; Index < UE_ARRAY_COUNT(OptionalShadowDepthColor); Index++)
 	{
 		if (OptionalShadowDepthColor[Index])
 		{
@@ -2717,17 +2717,17 @@ void FSceneRenderTargets::ReleaseAllTargets()
 	MobileCustomStencil.SafeRelease();
 	CustomStencilSRV.SafeRelease();
 
-	for (int32 i = 0; i < ARRAY_COUNT(OptionalShadowDepthColor); i++)
+	for (int32 i = 0; i < UE_ARRAY_COUNT(OptionalShadowDepthColor); i++)
 	{
 		OptionalShadowDepthColor[i].SafeRelease();
 	}
 
-	for (int32 i = 0; i < ARRAY_COUNT(ReflectionColorScratchCubemap); i++)
+	for (int32 i = 0; i < UE_ARRAY_COUNT(ReflectionColorScratchCubemap); i++)
 	{
 		ReflectionColorScratchCubemap[i].SafeRelease();
 	}
 
-	for (int32 i = 0; i < ARRAY_COUNT(DiffuseIrradianceScratchCubemap); i++)
+	for (int32 i = 0; i < UE_ARRAY_COUNT(DiffuseIrradianceScratchCubemap); i++)
 	{
 		DiffuseIrradianceScratchCubemap[i].SafeRelease();
 	}
diff --git a/Engine/Source/Runtime/Renderer/Private/ReflectionEnvironmentCapture.cpp b/Engine/Source/Runtime/Renderer/Private/ReflectionEnvironmentCapture.cpp
--- a/Engine/Source/Runtime/Renderer/Private/ReflectionEnvironmentCapture.cpp
+++ b/Engine/Source/Runtime/Renderer/Private/ReflectionEnvironmentCapture.cpp
@@ -372,7 +372,7 @@ void FilterReflectionEnvironment(FRHICommandListImmediate& RHICmdList, ERHIFeatu
 		SetGraphicsPipelineState(RHICmdList, GraphicsPSOInit);
 
 		FLinearColor UnusedColors[1] = { FLinearColor::Black };
-		PixelShader->SetColors(RHICmdList, UnusedColors, ARRAY_COUNT(UnusedColors));
+		PixelShader->SetColors(RHICmdList, UnusedColors, UE_ARRAY_COUNT(UnusedColors));
 
 		DrawRectangle(
 			RHICmdList,
diff --git a/Engine/Source/Runtime/Renderer/Private/RendererScene.cpp b/Engine/Source/Runtime/Renderer/Private/RendererScene.cpp
--- a/Engine/Source/Runtime/Renderer/Private/RendererScene.cpp
+++ b/Engine/Source/Runtime/Renderer/Private/RendererScene.cpp
@@ -183,7 +183,7 @@ FSceneViewState::FSceneViewState()
 
 	HeightfieldLightingAtlas = NULL;
 
-	for (int32 CascadeIndex = 0; CascadeIndex < ARRAY_COUNT(TranslucencyLightingCacheAllocations); CascadeIndex++)
+	for (int32 CascadeIndex = 0; CascadeIndex < UE_ARRAY_COUNT(TranslucencyLightingCacheAllocations); CascadeIndex++)
 	{
 		TranslucencyLightingCacheAllocations[CascadeIndex] = NULL;
 	}
@@ -253,7 +253,7 @@ FSceneViewState::~FSceneViewState()
 	CachedVisibilityChunk = NULL;
 	ShadowOcclusionQueryMaps.Reset();
 
-	for (int32 CascadeIndex = 0; CascadeIndex < ARRAY_COUNT(TranslucencyLightingCacheAllocations); CascadeIndex++)
+	for (int32 CascadeIndex = 0; CascadeIndex < UE_ARRAY_COUNT(TranslucencyLightingCacheAllocations); CascadeIndex++)
 	{
 		delete TranslucencyLightingCacheAllocations[CascadeIndex];
 	}
@@ -905,7 +905,7 @@ void FPersistentUniformBuffers::Initialize()
 	MobileDistortionPassUniformBuffer = TUniformBufferRef<FMobileDistortionPassUniformParameters>::CreateUniformBufferImmediate(MobileDistortionPassUniformParameters, UniformBuffer_MultiFrame, EUniformBufferValidation::None);
 	
 	FMobileDirectionalLightShaderParameters MobileDirectionalLightShaderParameters = {};
-	for (int32 Index = 0; Index < ARRAY_COUNT(MobileDirectionalLightUniformBuffers); ++Index)
+	for (int32 Index = 0; Index < UE_ARRAY_COUNT(MobileDirectionalLightUniformBuffers); ++Index)
 	{
 		MobileDirectionalLightUniformBuffers[Index] = TUniformBufferRef<FMobileDirectionalLightShaderParameters>::CreateUniformBufferImmediate(MobileDirectionalLightShaderParameters, UniformBuffer_MultiFrame, EUniformBufferValidation::None);
 	}
@@ -2536,7 +2536,7 @@ void FScene::RemoveLightSceneInfo_RenderThread(FLightSceneInfo* LightSceneInfo)
 #endif
 
 		    // check MobileDirectionalLights
-		    for (int32 LightChannelIdx = 0; LightChannelIdx < ARRAY_COUNT(MobileDirectionalLights); LightChannelIdx++)
+		    for (int32 LightChannelIdx = 0; LightChannelIdx < UE_ARRAY_COUNT(MobileDirectionalLights); LightChannelIdx++)
 		    {
 			    if (LightSceneInfo == MobileDirectionalLights[LightChannelIdx])
 			    {
diff --git a/Engine/Source/Runtime/Renderer/Private/SceneRendering.cpp b/Engine/Source/Runtime/Renderer/Private/SceneRendering.cpp
--- a/Engine/Source/Runtime/Renderer/Private/SceneRendering.cpp
+++ b/Engine/Source/Runtime/Renderer/Private/SceneRendering.cpp
@@ -1713,7 +1713,7 @@ FViewInfo* FViewInfo::CreateSnapshot() const
 	TUniformBufferRef<FViewUniformShaderParameters> NullViewUniformBuffer;
 	FMemory::Memcpy(Result->ViewUniformBuffer, NullViewUniformBuffer); 
 	TUniformBufferRef<FMobileDirectionalLightShaderParameters> NullMobileDirectionalLightUniformBuffer;
-	for (size_t i = 0; i < ARRAY_COUNT(Result->MobileDirectionalLightUniformBuffers); i++)
+	for (size_t i = 0; i < UE_ARRAY_COUNT(Result->MobileDirectionalLightUniformBuffers); i++)
 	{
 		// This memcpy is necessary to clear the reference from the memcpy of the whole of this -> Result without releasing the pointer
 		// But what we really want is the null buffer.
@@ -3408,7 +3408,7 @@ static void RenderViewFamily_RenderThread(FRHICommandListImmediate& RHICmdList,
 		}
 
 		// Only reset per-frame scene state once all views have processed their frame, including those in planar reflections
-		for (int32 CacheType = 0; CacheType < ARRAY_COUNT(SceneRenderer->Scene->DistanceFieldSceneData.PrimitiveModifiedBounds); CacheType++)
+		for (int32 CacheType = 0; CacheType < UE_ARRAY_COUNT(SceneRenderer->Scene->DistanceFieldSceneData.PrimitiveModifiedBounds); CacheType++)
 		{
 			SceneRenderer->Scene->DistanceFieldSceneData.PrimitiveModifiedBounds[CacheType].Reset();
 		}
diff --git a/Engine/Source/Runtime/Renderer/Private/SceneSoftwareOcclusion.cpp b/Engine/Source/Runtime/Renderer/Private/SceneSoftwareOcclusion.cpp
--- a/Engine/Source/Runtime/Renderer/Private/SceneSoftwareOcclusion.cpp
+++ b/Engine/Source/Runtime/Renderer/Private/SceneSoftwareOcclusion.cpp
@@ -940,7 +940,7 @@ static const ENamedThreads::Type ThreadNameMap[] =
 
 static ENamedThreads::Type GetOcclusionThreadName()
 {
-	int32 Index = FMath::Clamp<int32>(GSOThreadName, 0, ARRAY_COUNT(ThreadNameMap)-1);
+	int32 Index = FMath::Clamp<int32>(GSOThreadName, 0, UE_ARRAY_COUNT(ThreadNameMap)-1);
 	return ThreadNameMap[Index];
 }
 
diff --git a/Engine/Source/Runtime/Renderer/Private/SceneVisibility.cpp b/Engine/Source/Runtime/Renderer/Private/SceneVisibility.cpp
--- a/Engine/Source/Runtime/Renderer/Private/SceneVisibility.cpp
+++ b/Engine/Source/Runtime/Renderer/Private/SceneVisibility.cpp
@@ -2959,7 +2959,7 @@ void FSceneRenderer::PreVisibilityFrameSetup(FRHICommandListImmediate& RHICmdLis
 			{
 				RollingRemoveIndex = 0;
 				RollingPassShrinkIndex++;
-				if (RollingPassShrinkIndex >= ARRAY_COUNT(Scene->CachedDrawLists))
+				if (RollingPassShrinkIndex >= UE_ARRAY_COUNT(Scene->CachedDrawLists))
 				{
 					RollingPassShrinkIndex = 0;
 				}
@@ -3129,7 +3129,7 @@ void FSceneRenderer::PreVisibilityFrameSetup(FRHICommandListImmediate& RHICmdLis
 			{
 				float SamplesX[] = { -8.0f/16.0f, 0.0/16.0f };
 				float SamplesY[] = { /* - */ 0.0f/16.0f, 8.0/16.0f };
-				check(TemporalAASamples == ARRAY_COUNT(SamplesX));
+				check(TemporalAASamples == UE_ARRAY_COUNT(SamplesX));
 				SampleX = SamplesX[ TemporalSampleIndex ];
 				SampleY = SamplesY[ TemporalSampleIndex ];
 			}
@@ -3150,7 +3150,7 @@ void FSceneRenderer::PreVisibilityFrameSetup(FRHICommandListImmediate& RHICmdLis
 				//   .S
 				float SamplesX[] = { -4.0f/16.0f, 4.0/16.0f };
 				float SamplesY[] = { -4.0f/16.0f, 4.0/16.0f };
-				check(TemporalAASamples == ARRAY_COUNT(SamplesX));
+				check(TemporalAASamples == UE_ARRAY_COUNT(SamplesX));
 				SampleX = SamplesX[ TemporalSampleIndex ];
 				SampleY = SamplesY[ TemporalSampleIndex ];
 			}
@@ -3163,7 +3163,7 @@ void FSceneRenderer::PreVisibilityFrameSetup(FRHICommandListImmediate& RHICmdLis
 				// Rolling circle pattern (A,B,C).
 				float SamplesX[] = { -2.0f/3.0f,  2.0/3.0f,  0.0/3.0f };
 				float SamplesY[] = { -2.0f/3.0f,  0.0/3.0f,  2.0/3.0f };
-				check(TemporalAASamples == ARRAY_COUNT(SamplesX));
+				check(TemporalAASamples == UE_ARRAY_COUNT(SamplesX));
 				SampleX = SamplesX[ TemporalSampleIndex ];
 				SampleY = SamplesY[ TemporalSampleIndex ];
 			}
@@ -3178,7 +3178,7 @@ void FSceneRenderer::PreVisibilityFrameSetup(FRHICommandListImmediate& RHICmdLis
 				// Rolling circle pattern (N,E,S,W).
 				float SamplesX[] = { -2.0f/16.0f,  6.0/16.0f, 2.0/16.0f, -6.0/16.0f };
 				float SamplesY[] = { -6.0f/16.0f, -2.0/16.0f, 6.0/16.0f,  2.0/16.0f };
-				check(TemporalAASamples == ARRAY_COUNT(SamplesX));
+				check(TemporalAASamples == UE_ARRAY_COUNT(SamplesX));
 				SampleX = SamplesX[ TemporalSampleIndex ];
 				SampleY = SamplesY[ TemporalSampleIndex ];
 			}
@@ -3192,7 +3192,7 @@ void FSceneRenderer::PreVisibilityFrameSetup(FRHICommandListImmediate& RHICmdLis
 				// Rolling circle pattern (N,E,S,W).
 				float SamplesX[] = {  0.0f/2.0f,  1.0/2.0f,  0.0/2.0f, -1.0/2.0f };
 				float SamplesY[] = { -1.0f/2.0f,  0.0/2.0f,  1.0/2.0f,  0.0/2.0f };
-				check(TemporalAASamples == ARRAY_COUNT(SamplesX));
+				check(TemporalAASamples == UE_ARRAY_COUNT(SamplesX));
 				SampleX = SamplesX[ TemporalSampleIndex ];
 				SampleY = SamplesY[ TemporalSampleIndex ];
 			}
diff --git a/Engine/Source/Runtime/Renderer/Private/ScreenSpaceDenoise.cpp b/Engine/Source/Runtime/Renderer/Private/ScreenSpaceDenoise.cpp
--- a/Engine/Source/Runtime/Renderer/Private/ScreenSpaceDenoise.cpp
+++ b/Engine/Source/Runtime/Renderer/Private/ScreenSpaceDenoise.cpp
@@ -589,11 +589,11 @@ const TCHAR* const kDenoiserOutputResourceNames[] = {
 	nullptr,
 };
 
-static_assert(ARRAY_COUNT(kReconstructionResourceNames) == int32(ESignalProcessing::MAX) * kMaxBufferProcessingCount, "You forgot me!");
-static_assert(ARRAY_COUNT(kRejectionPreConvolutionResourceNames) == int32(ESignalProcessing::MAX) * kMaxBufferProcessingCount, "You forgot me!");
-static_assert(ARRAY_COUNT(kTemporalAccumulationResourceNames) == int32(ESignalProcessing::MAX) * kMaxBufferProcessingCount, "You forgot me!");
-static_assert(ARRAY_COUNT(kHistoryConvolutionResourceNames) == int32(ESignalProcessing::MAX) * kMaxBufferProcessingCount, "You forgot me!");
-static_assert(ARRAY_COUNT(kDenoiserOutputResourceNames) == int32(ESignalProcessing::MAX) * kMaxBufferProcessingCount, "You forgot me!");
+static_assert(UE_ARRAY_COUNT(kReconstructionResourceNames) == int32(ESignalProcessing::MAX) * kMaxBufferProcessingCount, "You forgot me!");
+static_assert(UE_ARRAY_COUNT(kRejectionPreConvolutionResourceNames) == int32(ESignalProcessing::MAX) * kMaxBufferProcessingCount, "You forgot me!");
+static_assert(UE_ARRAY_COUNT(kTemporalAccumulationResourceNames) == int32(ESignalProcessing::MAX) * kMaxBufferProcessingCount, "You forgot me!");
+static_assert(UE_ARRAY_COUNT(kHistoryConvolutionResourceNames) == int32(ESignalProcessing::MAX) * kMaxBufferProcessingCount, "You forgot me!");
+static_assert(UE_ARRAY_COUNT(kDenoiserOutputResourceNames) == int32(ESignalProcessing::MAX) * kMaxBufferProcessingCount, "You forgot me!");
 
 
 /** Returns whether should compole pipeline for a given shader platform.*/
diff --git a/Engine/Source/Runtime/Renderer/Private/ShadowDepthRendering.cpp b/Engine/Source/Runtime/Renderer/Private/ShadowDepthRendering.cpp
--- a/Engine/Source/Runtime/Renderer/Private/ShadowDepthRendering.cpp
+++ b/Engine/Source/Runtime/Renderer/Private/ShadowDepthRendering.cpp
@@ -1630,7 +1630,7 @@ void FSceneRenderer::RenderShadowDepthMaps(FRHICommandListImmediate& RHICmdList)
 			ColorTarget1.TargetableTexture
 		};
 
-		FRHIRenderPassInfo RPInfo(ARRAY_COUNT(RenderTargetArray), RenderTargetArray, ERenderTargetActions::Load_Store);
+		FRHIRenderPassInfo RPInfo(UE_ARRAY_COUNT(RenderTargetArray), RenderTargetArray, ERenderTargetActions::Load_Store);
 		TransitionRenderPassTargets(RHICmdList, RPInfo);
 		RHICmdList.BeginRenderPass(RPInfo, TEXT("RenderTranslucencyDepths"));
 		{
@@ -1703,23 +1703,23 @@ void FSceneRenderer::RenderShadowDepthMaps(FRHICommandListImmediate& RHICmdList)
 				Uavs[2] = LightPropagationVolume->GetVplListBufferUav();
 				Uavs[3] = LightPropagationVolume->GetVplListHeadBufferUav();
 
-				FRHIRenderPassInfo RPInfo(ARRAY_COUNT(RenderTargets), RenderTargets, ERenderTargetActions::Load_Store);
+				FRHIRenderPassInfo RPInfo(UE_ARRAY_COUNT(RenderTargets), RenderTargets, ERenderTargetActions::Load_Store);
 				RPInfo.DepthStencilRenderTarget.Action = EDepthStencilTargetActions::LoadDepthStencil_StoreDepthStencil;
 				RPInfo.DepthStencilRenderTarget.DepthStencilTarget = DepthTarget.TargetableTexture;
 				RPInfo.DepthStencilRenderTarget.ExclusiveDepthStencil = FExclusiveDepthStencil::DepthWrite_StencilWrite;
 
 				// Set starting UAV bind index
-				RPInfo.UAVIndex = ARRAY_COUNT(RenderTargets);
-				RPInfo.NumUAVs = ARRAY_COUNT(Uavs);
+				RPInfo.UAVIndex = UE_ARRAY_COUNT(RenderTargets);
+				RPInfo.NumUAVs = UE_ARRAY_COUNT(Uavs);
 				for (int32 Index = 0; Index < RPInfo.NumUAVs; Index++)
 				{
 					RPInfo.UAVs[Index] = Uavs[Index];
 				}
 
-				InRHICmdList.TransitionResources(EResourceTransitionAccess::ERWBarrier, EResourceTransitionPipeline::EGfxToGfx, Uavs, ARRAY_COUNT(Uavs));
+				InRHICmdList.TransitionResources(EResourceTransitionAccess::ERWBarrier, EResourceTransitionPipeline::EGfxToGfx, Uavs, UE_ARRAY_COUNT(Uavs));
 				InRHICmdList.BeginRenderPass(RPInfo, TEXT("ShadowAtlas"));
 
-				ProjectedShadowInfo->ClearDepth(InRHICmdList, this, ARRAY_COUNT(RenderTargets), RenderTargets, DepthTarget.TargetableTexture, bPerformClear);
+				ProjectedShadowInfo->ClearDepth(InRHICmdList, this, UE_ARRAY_COUNT(RenderTargets), RenderTargets, DepthTarget.TargetableTexture, bPerformClear);
 			};
 
 			{
@@ -1748,7 +1748,7 @@ void FSceneRenderer::RenderShadowDepthMaps(FRHICommandListImmediate& RHICmdList)
 				FRHIUnorderedAccessView* UavsToReadable[2];
 				UavsToReadable[0] = LightPropagationVolume->GetGvListBufferUav();
 				UavsToReadable[1] = LightPropagationVolume->GetGvListHeadBufferUav();
-				RHICmdList.TransitionResources(EResourceTransitionAccess::EReadable, EResourceTransitionPipeline::EGfxToGfx, UavsToReadable, ARRAY_COUNT(UavsToReadable));
+				RHICmdList.TransitionResources(EResourceTransitionAccess::EReadable, EResourceTransitionPipeline::EGfxToGfx, UavsToReadable, UE_ARRAY_COUNT(UavsToReadable));
 			}
 			checkSlow(RHICmdList.IsOutsideRenderPass());
 		}
diff --git a/Engine/Source/Runtime/Renderer/Private/ShadowSetup.cpp b/Engine/Source/Runtime/Renderer/Private/ShadowSetup.cpp
--- a/Engine/Source/Runtime/Renderer/Private/ShadowSetup.cpp
+++ b/Engine/Source/Runtime/Renderer/Private/ShadowSetup.cpp
@@ -4195,7 +4195,7 @@ void FSceneRenderer::InitDynamicShadows(FRHICommandListImmediate& RHICmdList, FG
 					// Allow movable and stationary lights to create CSM, or static lights that are unbuilt
 					if ((!LightSceneInfo->Proxy->HasStaticLighting() && LightSceneInfoCompact.bCastDynamicShadow) || bCreateShadowToPreviewStaticLight)
 					{
-						static_assert(ARRAY_COUNT(Scene->MobileDirectionalLights) == 3, "All array entries for MobileDirectionalLights must be checked");
+						static_assert(UE_ARRAY_COUNT(Scene->MobileDirectionalLights) == 3, "All array entries for MobileDirectionalLights must be checked");
 						if( !bMobile ||
 							((LightSceneInfo->Proxy->UseCSMForDynamicObjects() || LightSceneInfo->Proxy->IsMovable()) 
 								// Mobile uses the scene's MobileDirectionalLights only for whole scene shadows.
diff --git a/Engine/Source/Runtime/Renderer/Private/SurfelTree.cpp b/Engine/Source/Runtime/Renderer/Private/SurfelTree.cpp
--- a/Engine/Source/Runtime/Renderer/Private/SurfelTree.cpp
+++ b/Engine/Source/Runtime/Renderer/Private/SurfelTree.cpp
@@ -227,7 +227,7 @@ public:
 
 		FRHIUnorderedAccessView* UniformMeshUAVs[1];
 		UniformMeshUAVs[0] = Scene->DistanceFieldSceneData.SurfelBuffers->InterpolatedVertexData.UAV;
-		RHICmdList.TransitionResources(EResourceTransitionAccess::ERWBarrier, EResourceTransitionPipeline::EComputeToCompute, UniformMeshUAVs, ARRAY_COUNT(UniformMeshUAVs));
+		RHICmdList.TransitionResources(EResourceTransitionAccess::ERWBarrier, EResourceTransitionPipeline::EComputeToCompute, UniformMeshUAVs, UE_ARRAY_COUNT(UniformMeshUAVs));
 
 		SurfelBufferParameters.Set(RHICmdList, ShaderRHI, *Scene->DistanceFieldSceneData.SurfelBuffers, *Scene->DistanceFieldSceneData.InstancedSurfelBuffers);
 		
@@ -250,7 +250,7 @@ public:
 
 		FRHIUnorderedAccessView* UniformMeshUAVs[1];
 		UniformMeshUAVs[0] = Scene->DistanceFieldSceneData.SurfelBuffers->InterpolatedVertexData.UAV;
-		RHICmdList.TransitionResources(EResourceTransitionAccess::EReadable, EResourceTransitionPipeline::EComputeToCompute, UniformMeshUAVs, ARRAY_COUNT(UniformMeshUAVs));
+		RHICmdList.TransitionResources(EResourceTransitionAccess::EReadable, EResourceTransitionPipeline::EComputeToCompute, UniformMeshUAVs, UE_ARRAY_COUNT(UniformMeshUAVs));
 	}
 
 	virtual bool Serialize(FArchive& Ar) override
@@ -305,7 +305,7 @@ void GenerateSurfelRepresentation(FRHICommandListImmediate& RHICmdList, FSceneRe
 			FRHIUnorderedAccessView* UniformMeshUAVs[2];
 			UniformMeshUAVs[0] = UniformMeshBuffers->TriangleAreas.UAV;
 			UniformMeshUAVs[1] = UniformMeshBuffers->TriangleCDFs.UAV;
-			RHICmdList.TransitionResources(EResourceTransitionAccess::ERWBarrier, EResourceTransitionPipeline::EComputeToCompute, UniformMeshUAVs, ARRAY_COUNT(UniformMeshUAVs));
+			RHICmdList.TransitionResources(EResourceTransitionAccess::ERWBarrier, EResourceTransitionPipeline::EComputeToCompute, UniformMeshUAVs, UE_ARRAY_COUNT(UniformMeshUAVs));
 
 			{
 				TShaderMapRef<FComputeTriangleAreasCS> ComputeShader(GetGlobalShaderMap(View.GetFeatureLevel()));
diff --git a/Engine/Source/Runtime/Renderer/Private/TranslucentLighting.cpp b/Engine/Source/Runtime/Renderer/Private/TranslucentLighting.cpp
--- a/Engine/Source/Runtime/Renderer/Private/TranslucentLighting.cpp
+++ b/Engine/Source/Runtime/Renderer/Private/TranslucentLighting.cpp
@@ -556,7 +556,7 @@ void FProjectedShadowInfo::RenderTranslucencyDepths(FRHICommandList& RHICmdList,
 			);
 
 		FLinearColor ClearColors[2] = {FLinearColor(0,0,0,0), FLinearColor(0,0,0,0)};
-		DrawClearQuadMRT(RHICmdList, true, ARRAY_COUNT(ClearColors), ClearColors, false, 1.0f, false, 0);
+		DrawClearQuadMRT(RHICmdList, true, UE_ARRAY_COUNT(ClearColors), ClearColors, false, 1.0f, false, 0);
 
 		// Set the viewport for the shadow.
 		RHICmdList.SetViewport(
@@ -1183,7 +1183,7 @@ void FDeferredShadingSceneRenderer::ClearTranslucentVolumePerObjectShadowing(FRH
 		ClearColors[0] = FLinearColor(1, 1, 1, 1);
 		ClearColors[1] = FLinearColor(1, 1, 1, 1);
 
-		FSceneRenderTargets::ClearVolumeTextures<ARRAY_COUNT(RenderTargets)>(RHICmdList, FeatureLevel, RenderTargets, ClearColors);
+		FSceneRenderTargets::ClearVolumeTextures<UE_ARRAY_COUNT(RenderTargets)>(RHICmdList, FeatureLevel, RenderTargets, ClearColors);
 	}
 }
 
@@ -1475,7 +1475,7 @@ static void InjectTranslucentLightArray(FRHICommandListImmediate& RHICmdList, co
 		RenderTargets[0] = RT0->GetRenderTargetItem().TargetableTexture;
 		RenderTargets[1] = RT1->GetRenderTargetItem().TargetableTexture;
 
-		FRHIRenderPassInfo RPInfo(ARRAY_COUNT(RenderTargets), RenderTargets, ERenderTargetActions::Load_Store);
+		FRHIRenderPassInfo RPInfo(UE_ARRAY_COUNT(RenderTargets), RenderTargets, ERenderTargetActions::Load_Store);
 		TransitionRenderPassTargets(RHICmdList, RPInfo);
 
 		RHICmdList.BeginRenderPass(RPInfo, TEXT("InjectTranslucentLightArray"));
@@ -1693,7 +1693,7 @@ void FDeferredShadingSceneRenderer::InjectSimpleTranslucentVolumeLightingArray(F
 			RenderTargets[0] = RT0->GetRenderTargetItem().TargetableTexture;
 			RenderTargets[1] = RT1->GetRenderTargetItem().TargetableTexture;
 
-			FRHIRenderPassInfo RPInfo(ARRAY_COUNT(RenderTargets), RenderTargets, ERenderTargetActions::Load_Store);
+			FRHIRenderPassInfo RPInfo(UE_ARRAY_COUNT(RenderTargets), RenderTargets, ERenderTargetActions::Load_Store);
 			TransitionRenderPassTargets(RHICmdList, RPInfo);
 			RHICmdList.BeginRenderPass(RPInfo, TEXT("InjectSimpleTranslucentVolumeLightingArray"));
 			{
@@ -1818,7 +1818,7 @@ void FDeferredShadingSceneRenderer::FilterTranslucentVolumeLighting(FRHICommandL
 				}
 				RHICmdList.TransitionResources(EResourceTransitionAccess::EReadable, Inputs, 2);
 
-				FRHIRenderPassInfo RPInfo(ARRAY_COUNT(RenderTargets), RenderTargets, ERenderTargetActions::Load_Store);
+				FRHIRenderPassInfo RPInfo(UE_ARRAY_COUNT(RenderTargets), RenderTargets, ERenderTargetActions::Load_Store);
 				TransitionRenderPassTargets(RHICmdList, RPInfo);
 				RHICmdList.BeginRenderPass(RPInfo, TEXT("FilterTranslucentVolumeLighting"));
 				{
diff --git a/Engine/Source/Runtime/Renderer/Private/VT/VirtualTextureSpace.cpp b/Engine/Source/Runtime/Renderer/Private/VT/VirtualTextureSpace.cpp
--- a/Engine/Source/Runtime/Renderer/Private/VT/VirtualTextureSpace.cpp
+++ b/Engine/Source/Runtime/Renderer/Private/VT/VirtualTextureSpace.cpp
@@ -272,7 +272,7 @@ void FVirtualTextureSpace::AllocateTextures(FRHICommandList& RHICmdList)
 	if (bNeedToAllocatePageTable)
 	{
 		const TCHAR* TextureNames[] = { TEXT("PageTable_0"), TEXT("PageTable_1") };
-		static_assert(ARRAY_COUNT(TextureNames) == TextureCapacity, "");
+		static_assert(UE_ARRAY_COUNT(TextureNames) == TextureCapacity, "");
 
 		for (uint32 TextureIndex = 0u; TextureIndex < GetNumPageTableTextures(); ++TextureIndex)
 		{
diff --git a/Engine/Source/Runtime/Renderer/Private/VolumeLighting.h b/Engine/Source/Runtime/Renderer/Private/VolumeLighting.h
--- a/Engine/Source/Runtime/Renderer/Private/VolumeLighting.h
+++ b/Engine/Source/Runtime/Renderer/Private/VolumeLighting.h
@@ -87,7 +87,7 @@ public:
 
 		SetShaderValue(RHICmdList, ShaderRHI, ShadowInjectParams, ShadowInjectParamValue);
 
-		SetShaderValueArray(RHICmdList, ShaderRHI, ClippingPlanes, Planes, ARRAY_COUNT(Planes));
+		SetShaderValueArray(RHICmdList, ShaderRHI, ClippingPlanes, Planes, UE_ARRAY_COUNT(Planes));
 	
 		ELightComponentType LightType = (ELightComponentType)LightSceneInfo->Proxy->GetLightType();
 
diff --git a/Engine/Source/Runtime/Renderer/Public/MeshPassProcessor.h b/Engine/Source/Runtime/Renderer/Public/MeshPassProcessor.h
--- a/Engine/Source/Runtime/Renderer/Public/MeshPassProcessor.h
+++ b/Engine/Source/Runtime/Renderer/Public/MeshPassProcessor.h
@@ -390,7 +390,7 @@ private:
 
 		Size = InSize;
 
-		if (InSize > ARRAY_COUNT(InlineStorage))
+		if (InSize > UE_ARRAY_COUNT(InlineStorage))
 		{
 			HeapData = new uint8[InSize];
 		}
@@ -408,12 +408,12 @@ private:
 
 	uint8* GetData()
 	{
-		return Size <= ARRAY_COUNT(InlineStorage) ? &InlineStorage[0] : HeapData;
+		return Size <= UE_ARRAY_COUNT(InlineStorage) ? &InlineStorage[0] : HeapData;
 	}
 
 	const uint8* GetData() const
 	{
-		return Size <= ARRAY_COUNT(InlineStorage) ? &InlineStorage[0] : HeapData;
+		return Size <= UE_ARRAY_COUNT(InlineStorage) ? &InlineStorage[0] : HeapData;
 	}
 
 	RENDERER_API void CopyFrom(const FMeshDrawShaderBindings& Other);
diff --git a/Engine/Source/Runtime/SignalProcessing/Private/Filter.cpp b/Engine/Source/Runtime/SignalProcessing/Private/Filter.cpp
--- a/Engine/Source/Runtime/SignalProcessing/Private/Filter.cpp
+++ b/Engine/Source/Runtime/SignalProcessing/Private/Filter.cpp
@@ -552,8 +552,8 @@ namespace Audio
 		, Alpha(0.0f)
 		, PassBandGainCompensation(0.0f)
 	{
-		FMemory::Memzero(Factors, sizeof(float) * ARRAY_COUNT(Beta));
-		FMemory::Memzero(Beta, sizeof(float) * ARRAY_COUNT(Beta));
+		FMemory::Memzero(Factors, sizeof(float) * UE_ARRAY_COUNT(Beta));
+		FMemory::Memzero(Beta, sizeof(float) * UE_ARRAY_COUNT(Beta));
 	}
 
 	FLadderFilter::~FLadderFilter()
diff --git a/Engine/Source/Runtime/SlateCore/Private/Fonts/SlateTextShaper.cpp b/Engine/Source/Runtime/SlateCore/Private/Fonts/SlateTextShaper.cpp
--- a/Engine/Source/Runtime/SlateCore/Private/Fonts/SlateTextShaper.cpp
+++ b/Engine/Source/Runtime/SlateCore/Private/Fonts/SlateTextShaper.cpp
@@ -610,7 +610,7 @@ void FSlateTextShaper::PerformHarfBuzzTextShaping(const TCHAR* InText, const int
 			const hb_feature_t HarfBuzzFeatures[] = {
 				{ HB_TAG('k','e','r','n'), bHasKerning, 0, uint32(-1) }
 			};
-			const int32 HarfBuzzFeaturesCount = ARRAY_COUNT(HarfBuzzFeatures);
+			const int32 HarfBuzzFeaturesCount = UE_ARRAY_COUNT(HarfBuzzFeatures);
 
 			uint32 GlyphFlags = 0;
 			SlateFontRendererUtils::AppendGlyphFlags(*HarfBuzzTextSequenceEntry.FaceAndMemory, *HarfBuzzTextSequenceEntry.FontDataPtr, GlyphFlags);
diff --git a/Engine/Source/Runtime/Sockets/Private/Android/SocketSubsystemAndroid.cpp b/Engine/Source/Runtime/Sockets/Private/Android/SocketSubsystemAndroid.cpp
--- a/Engine/Source/Runtime/Sockets/Private/Android/SocketSubsystemAndroid.cpp
+++ b/Engine/Source/Runtime/Sockets/Private/Android/SocketSubsystemAndroid.cpp
@@ -145,7 +145,7 @@ TSharedRef<FInternetAddr> FSocketSubsystemAndroid::GetLocalHostAddr(FOutputDevic
 			FMemory::Memzero(&CellularAddress, sizeof(sockaddr_storage));
 			FMemory::Memzero(&OtherAddress, sizeof(sockaddr_storage));
 
-			for (int32 IdxReq = 0; IdxReq < ARRAY_COUNT(IfReqs); ++IdxReq)
+			for (int32 IdxReq = 0; IdxReq < UE_ARRAY_COUNT(IfReqs); ++IdxReq)
 			{
 				// Cache the address information, as the following flag lookup will 
 				// write into the ifr_addr field.
diff --git a/Engine/Source/Runtime/Sockets/Private/SocketSubsystem.cpp b/Engine/Source/Runtime/Sockets/Private/SocketSubsystem.cpp
--- a/Engine/Source/Runtime/Sockets/Private/SocketSubsystem.cpp
+++ b/Engine/Source/Runtime/Sockets/Private/SocketSubsystem.cpp
@@ -341,7 +341,7 @@ TSharedRef<FInternetAddr> ISocketSubsystem::GetLocalHostAddr(FOutputDevice& Out,
 bool ISocketSubsystem::GetMultihomeAddress(TSharedRef<FInternetAddr>& Addr)
 {
 	TCHAR Home[256] = TEXT("");
-	if (FParse::Value(FCommandLine::Get(), TEXT("MULTIHOME="), Home, ARRAY_COUNT(Home)))
+	if (FParse::Value(FCommandLine::Get(), TEXT("MULTIHOME="), Home, UE_ARRAY_COUNT(Home)))
 	{
 		TSharedPtr<FInternetAddr> MultiHomeQuery = GetAddressFromString(Home);
 		if (Home == NULL || !MultiHomeQuery.IsValid())
diff --git a/Engine/Source/Runtime/SynthBenchmark/Private/SynthBenchmarkPrivate.cpp b/Engine/Source/Runtime/SynthBenchmark/Private/SynthBenchmarkPrivate.cpp
--- a/Engine/Source/Runtime/SynthBenchmark/Private/SynthBenchmarkPrivate.cpp
+++ b/Engine/Source/Runtime/SynthBenchmark/Private/SynthBenchmarkPrivate.cpp
@@ -112,7 +112,7 @@ void FSynthBenchmark::Run(FSynthBenchmarkResults& InOut, bool bGPUBenchmark, flo
 	InOut.CPUStats[1] = FSynthBenchmarkStat(TEXT("Fractal"), 0.0286f, TEXT("s/Run"), 1.5f);
 	InOut.CPUStats[1].SetMeasuredTime(RunBenchmark(WorkScale, FractalBenchmark));
 
-	for(uint32 i = 0; i < ARRAY_COUNT(InOut.CPUStats); ++i)
+	for(uint32 i = 0; i < UE_ARRAY_COUNT(InOut.CPUStats); ++i)
 	{
 		UE_LOG(LogSynthBenchmark, Display, TEXT("         ... %f %s '%s'"), InOut.CPUStats[i].GetNormalizedTime(), InOut.CPUStats[i].GetValueType(), InOut.CPUStats[i].GetDesc());
 	}
@@ -130,7 +130,7 @@ void FSynthBenchmark::Run(FSynthBenchmarkResults& InOut, bool bGPUBenchmark, flo
 	UE_LOG(LogSynthBenchmark, Display, TEXT("  NumberOfCores (physical): %d"), FPlatformMisc::NumberOfCores());
 	UE_LOG(LogSynthBenchmark, Display, TEXT("  NumberOfCores (logical): %d"), FPlatformMisc::NumberOfCoresIncludingHyperthreads());
 
-	for (uint32 MethodId = 0; MethodId < ARRAY_COUNT(InOut.CPUStats); ++MethodId)
+	for (uint32 MethodId = 0; MethodId < UE_ARRAY_COUNT(InOut.CPUStats); ++MethodId)
 	{
 		UE_LOG(LogSynthBenchmark, Display, TEXT("  CPU Perf Index %d: %.1f (weight %.2f)"), MethodId, InOut.CPUStats[MethodId].ComputePerfIndex(), InOut.CPUStats[MethodId].GetWeight());
 	}
@@ -218,7 +218,7 @@ void FSynthBenchmark::Run(FSynthBenchmarkResults& InOut, bool bGPUBenchmark, flo
 				PrintGPUStats(InOut.GPUStats, TEXT(""));
 				UE_LOG(LogSynthBenchmark, Display, TEXT(""));
 
-				for (uint32 MethodId = 0; MethodId < ARRAY_COUNT(InOut.GPUStats); ++MethodId)
+				for (uint32 MethodId = 0; MethodId < UE_ARRAY_COUNT(InOut.GPUStats); ++MethodId)
 				{
 					UE_LOG(LogSynthBenchmark, Display, TEXT("  GPU Perf Index %d: %.1f (weight %.2f)"), MethodId, InOut.GPUStats[MethodId].ComputePerfIndex(), InOut.GPUStats[MethodId].GetWeight());
 				}
diff --git a/Engine/Source/Runtime/TimeManagement/Private/TimeManagementTests.cpp b/Engine/Source/Runtime/TimeManagement/Private/TimeManagementTests.cpp
--- a/Engine/Source/Runtime/TimeManagement/Private/TimeManagementTests.cpp
+++ b/Engine/Source/Runtime/TimeManagement/Private/TimeManagementTests.cpp
@@ -32,7 +32,7 @@ namespace
 IMPLEMENT_SIMPLE_AUTOMATION_TEST(FFrameTimeComparisonTest, "System.Core.Time.Comparison", EAutomationTestFlags::ApplicationContextMask | EAutomationTestFlags::SmokeFilter | EAutomationTestFlags::HighPriority)
 bool FFrameTimeComparisonTest::RunTest(const FString& Parameters)
 {
-	const int32 NumFrames = ARRAY_COUNT(TestTimes);
+	const int32 NumFrames = UE_ARRAY_COUNT(TestTimes);
 
 	for (int32 Index = 0; Index < NumFrames; ++Index)
 	{
@@ -83,7 +83,7 @@ bool FFrameTimeComparisonTest::RunTest(const FString& Parameters)
 IMPLEMENT_SIMPLE_AUTOMATION_TEST(FFrameTimeToSecondsTest, "System.Core.Time.ToSeconds", EAutomationTestFlags::ApplicationContextMask | EAutomationTestFlags::SmokeFilter | EAutomationTestFlags::HighPriority)
 bool FFrameTimeToSecondsTest::RunTest(const FString& Parameters)
 {
-	const int32 NumFrames = ARRAY_COUNT(TestTimes);
+	const int32 NumFrames = UE_ARRAY_COUNT(TestTimes);
 
 	FFrameRate TestRate = FCommonFrameRates::FPS_60();
 	double ExpectedSeconds[] = {
@@ -116,7 +116,7 @@ bool FFrameTimeToSecondsTest::RunTest(const FString& Parameters)
 IMPLEMENT_SIMPLE_AUTOMATION_TEST(FFrameTimeAdditionTest, "System.Core.Time.Addition", EAutomationTestFlags::ApplicationContextMask | EAutomationTestFlags::SmokeFilter | EAutomationTestFlags::HighPriority)
 bool FFrameTimeAdditionTest::RunTest(const FString& Parameters)
 {
-	const int32 NumFrames = ARRAY_COUNT(TestTimes);
+	const int32 NumFrames = UE_ARRAY_COUNT(TestTimes);
 
 	// Test adding a positive FrameTime with a small sub frame
 	{
@@ -239,7 +239,7 @@ bool FFrameTimeAdditionTest::RunTest(const FString& Parameters)
 IMPLEMENT_SIMPLE_AUTOMATION_TEST(FFrameTimeSubtractionTest, "System.Core.Time.Subtraction", EAutomationTestFlags::ApplicationContextMask | EAutomationTestFlags::SmokeFilter | EAutomationTestFlags::HighPriority)
 bool FFrameTimeSubtractionTest::RunTest(const FString& Parameters)
 {
-	const int32 NumFrames = ARRAY_COUNT(TestTimes);
+	const int32 NumFrames = UE_ARRAY_COUNT(TestTimes);
 
 	// Test subtracting a positive FrameTime with a small sub frame
 	{
@@ -363,7 +363,7 @@ bool FFrameTimeSubtractionTest::RunTest(const FString& Parameters)
 IMPLEMENT_SIMPLE_AUTOMATION_TEST(FFrameTimeConversionTest, "System.Core.Time.Conversion", EAutomationTestFlags::ApplicationContextMask | EAutomationTestFlags::SmokeFilter | EAutomationTestFlags::HighPriority)
 bool FFrameTimeConversionTest::RunTest(const FString& Parameters)
 {
-	const int32 NumFrames = ARRAY_COUNT(TestTimes);
+	const int32 NumFrames = UE_ARRAY_COUNT(TestTimes);
 
 	{
 		FFrameRate SrcRate = FCommonFrameRates::FPS_60();
@@ -469,7 +469,7 @@ bool FFrameRateMultiplesTest::RunTest(const FString& Parameters)
 		FFrameRate(24000,1)
 	};
 
-	const int32 NumRates = ARRAY_COUNT(TestRates);
+	const int32 NumRates = UE_ARRAY_COUNT(TestRates);
 
 	{
 		bool IsMultipleOf[] = {
diff --git a/Engine/Source/Runtime/TraceLog/Private/Trace/Writer.cpp b/Engine/Source/Runtime/TraceLog/Private/Trace/Writer.cpp
--- a/Engine/Source/Runtime/TraceLog/Private/Trace/Writer.cpp
+++ b/Engine/Source/Runtime/TraceLog/Private/Trace/Writer.cpp
@@ -646,7 +646,7 @@ static void Writer_ControlRecv()
 					continue;
 				}
 
-				if (ArgC < ARRAY_COUNT(ArgV))
+				if (ArgC < UE_ARRAY_COUNT(ArgV))
 				{
 					ArgV[ArgC] = Cursor;
 					++ArgC;
diff --git a/Engine/Source/Runtime/VulkanRHI/Private/VulkanIndexBuffer.cpp b/Engine/Source/Runtime/VulkanRHI/Private/VulkanIndexBuffer.cpp
--- a/Engine/Source/Runtime/VulkanRHI/Private/VulkanIndexBuffer.cpp
+++ b/Engine/Source/Runtime/VulkanRHI/Private/VulkanIndexBuffer.cpp
@@ -105,7 +105,7 @@ FVulkanResourceMultiBuffer::FVulkanResourceMultiBuffer(FVulkanDevice* InDevice,
 			}
 
 			NumBuffers = bDynamic ? NUM_BUFFERS : 1;
-			check(NumBuffers <= ARRAY_COUNT(Buffers));
+			check(NumBuffers <= UE_ARRAY_COUNT(Buffers));
 
 			for (uint32 Index = 0; Index < NumBuffers; ++Index)
 			{
diff --git a/Engine/Source/Runtime/VulkanRHI/Private/VulkanLayers.cpp b/Engine/Source/Runtime/VulkanRHI/Private/VulkanLayers.cpp
--- a/Engine/Source/Runtime/VulkanRHI/Private/VulkanLayers.cpp
+++ b/Engine/Source/Runtime/VulkanRHI/Private/VulkanLayers.cpp
@@ -633,7 +633,7 @@ void FVulkanDevice::GetDeviceExtensionsAndLayers(VkPhysicalDevice Gpu, EGpuVendo
 		}
 	}
 
-	for (uint32 Index = 0; Index < ARRAY_COUNT(GDeviceExtensions) && GDeviceExtensions[Index] != nullptr; ++Index)
+	for (uint32 Index = 0; Index < UE_ARRAY_COUNT(GDeviceExtensions) && GDeviceExtensions[Index] != nullptr; ++Index)
 	{
 		if (ListContains(AvailableExtensions, GDeviceExtensions[Index]))
 		{
diff --git a/Engine/Source/Runtime/VulkanRHI/Private/VulkanMemory.cpp b/Engine/Source/Runtime/VulkanRHI/Private/VulkanMemory.cpp
--- a/Engine/Source/Runtime/VulkanRHI/Private/VulkanMemory.cpp
+++ b/Engine/Source/Runtime/VulkanRHI/Private/VulkanMemory.cpp
@@ -1598,7 +1598,7 @@ namespace VulkanRHI
 		uint64 AllocBinnedTotal = 0;
 		uint64 UsedLargeTotal = 0;
 		uint64 AllocLargeTotal = 0;
-		for (int32 PoolSizeIndex = 0; PoolSizeIndex < ARRAY_COUNT(UsedBufferAllocations); PoolSizeIndex++)
+		for (int32 PoolSizeIndex = 0; PoolSizeIndex < UE_ARRAY_COUNT(UsedBufferAllocations); PoolSizeIndex++)
 		{
 			TArray<FBufferAllocation*>& UsedAllocations = UsedBufferAllocations[PoolSizeIndex];
 			TArray<FBufferAllocation*>& FreeAllocations = FreeBufferAllocations[PoolSizeIndex];
diff --git a/Engine/Source/Runtime/VulkanRHI/Private/VulkanPipeline.cpp b/Engine/Source/Runtime/VulkanRHI/Private/VulkanPipeline.cpp
--- a/Engine/Source/Runtime/VulkanRHI/Private/VulkanPipeline.cpp
+++ b/Engine/Source/Runtime/VulkanRHI/Private/VulkanPipeline.cpp
@@ -979,12 +979,12 @@ void FVulkanPipelineStateCacheManager::FGfxPipelineEntry::FRenderTargets::ReadFr
 			New->ReadFrom(Source[Index]);
 		}
 	};
-	CopyAttachmentRefs(ColorAttachments, RTLayout.ColorReferences, ARRAY_COUNT(RTLayout.ColorReferences));
-	CopyAttachmentRefs(ResolveAttachments, RTLayout.ResolveReferences, ARRAY_COUNT(RTLayout.ResolveReferences));
+	CopyAttachmentRefs(ColorAttachments, RTLayout.ColorReferences, UE_ARRAY_COUNT(RTLayout.ColorReferences));
+	CopyAttachmentRefs(ResolveAttachments, RTLayout.ResolveReferences, UE_ARRAY_COUNT(RTLayout.ResolveReferences));
 	DepthStencil.ReadFrom(RTLayout.DepthStencilReference);
 
-	Descriptions.AddZeroed(ARRAY_COUNT(RTLayout.Desc));
-	for (int32 Index = 0; Index < ARRAY_COUNT(RTLayout.Desc); ++Index)
+	Descriptions.AddZeroed(UE_ARRAY_COUNT(RTLayout.Desc));
+	for (int32 Index = 0; Index < UE_ARRAY_COUNT(RTLayout.Desc); ++Index)
 	{
 		Descriptions[Index].ReadFrom(RTLayout.Desc[Index]);
 	}
@@ -1013,11 +1013,11 @@ void FVulkanPipelineStateCacheManager::FGfxPipelineEntry::FRenderTargets::WriteI
 			Source[Index].WriteInto(*Dest);
 		}
 	};
-	CopyAttachmentRefs(ColorAttachments, Out.ColorReferences, ARRAY_COUNT(Out.ColorReferences));
-	CopyAttachmentRefs(ResolveAttachments, Out.ResolveReferences, ARRAY_COUNT(Out.ResolveReferences));
+	CopyAttachmentRefs(ColorAttachments, Out.ColorReferences, UE_ARRAY_COUNT(Out.ColorReferences));
+	CopyAttachmentRefs(ResolveAttachments, Out.ResolveReferences, UE_ARRAY_COUNT(Out.ResolveReferences));
 	DepthStencil.WriteInto(Out.DepthStencilReference);
 
-	for (int32 Index = 0; Index < ARRAY_COUNT(Out.Desc); ++Index)
+	for (int32 Index = 0; Index < UE_ARRAY_COUNT(Out.Desc); ++Index)
 	{
 		Descriptions[Index].WriteInto(Out.Desc[Index]);
 	}
@@ -1061,7 +1061,7 @@ FArchive& operator << (FArchive& Ar, FVulkanPipelineStateCacheManager::FGfxPipel
 
 	Ar << Entry.DepthStencil;
 
-	for (int32 Index = 0; Index < ARRAY_COUNT(Entry.ShaderHashes.Stages); ++Index)
+	for (int32 Index = 0; Index < UE_ARRAY_COUNT(Entry.ShaderHashes.Stages); ++Index)
 	{
 		Ar << Entry.ShaderHashes.Stages[Index];
 	}
diff --git a/Engine/Source/Runtime/VulkanRHI/Private/VulkanVertexDeclaration.cpp b/Engine/Source/Runtime/VulkanRHI/Private/VulkanVertexDeclaration.cpp
--- a/Engine/Source/Runtime/VulkanRHI/Private/VulkanVertexDeclaration.cpp
+++ b/Engine/Source/Runtime/VulkanRHI/Private/VulkanVertexDeclaration.cpp
@@ -113,7 +113,7 @@ void FVulkanVertexInputStateInfo::Generate(FVulkanVertexDeclaration* VertexDecla
 	BindingsNum = 0;
 	BindingToStream.Reset();
 	StreamToBinding.Reset();
-	for (int32 i=0; i<ARRAY_COUNT(Bindings); i++)
+	for (int32 i=0; i<UE_ARRAY_COUNT(Bindings); i++)
 	{
 		if (!((1<<i) & BindingsMask))
 		{
@@ -129,7 +129,7 @@ void FVulkanVertexInputStateInfo::Generate(FVulkanVertexDeclaration* VertexDecla
 	}
 
 	// Clean originally placed bindings
-	FMemory::Memset(Bindings + BindingsNum, 0, sizeof(Bindings[0]) * (ARRAY_COUNT(Bindings)-BindingsNum));
+	FMemory::Memset(Bindings + BindingsNum, 0, sizeof(Bindings[0]) * (UE_ARRAY_COUNT(Bindings)-BindingsNum));
 
 	// Attributes are expected to be uninitialized/empty
 	check(AttributesNum == 0);
diff --git a/Engine/Source/Runtime/Windows/D3D11RHI/Private/D3D11Shaders.cpp b/Engine/Source/Runtime/Windows/D3D11RHI/Private/D3D11Shaders.cpp
--- a/Engine/Source/Runtime/Windows/D3D11RHI/Private/D3D11Shaders.cpp
+++ b/Engine/Source/Runtime/Windows/D3D11RHI/Private/D3D11Shaders.cpp
@@ -323,7 +323,7 @@ FD3D11BoundShaderState::FD3D11BoundShaderState(
 	bShaderNeedsGlobalConstantBuffer[SF_Pixel] = InPixelShader ? InPixelShader->bShaderNeedsGlobalConstantBuffer : false;
 	bShaderNeedsGlobalConstantBuffer[SF_Geometry] = InGeometryShader ? InGeometryShader->bShaderNeedsGlobalConstantBuffer : false;
 
-	static_assert(ARRAY_COUNT(bShaderNeedsGlobalConstantBuffer) == SF_NumStandardFrequencies, "EShaderFrequency size should match with array count of bShaderNeedsGlobalConstantBuffer.");
+	static_assert(UE_ARRAY_COUNT(bShaderNeedsGlobalConstantBuffer) == SF_NumStandardFrequencies, "EShaderFrequency size should match with array count of bShaderNeedsGlobalConstantBuffer.");
 }
 
 FD3D11BoundShaderState::~FD3D11BoundShaderState()
diff --git a/Engine/Source/Runtime/Windows/D3D11RHI/Private/Windows/WindowsD3D11Device.cpp b/Engine/Source/Runtime/Windows/D3D11RHI/Private/Windows/WindowsD3D11Device.cpp
--- a/Engine/Source/Runtime/Windows/D3D11RHI/Private/Windows/WindowsD3D11Device.cpp
+++ b/Engine/Source/Runtime/Windows/D3D11RHI/Private/Windows/WindowsD3D11Device.cpp
@@ -268,7 +268,7 @@ static bool SafeTestD3D11CreateDevice(IDXGIAdapter* Adapter,D3D_FEATURE_LEVEL Mi
 	
 	// Trim to allowed feature levels
 	int32 FirstAllowedFeatureLevel = 0;
-	int32 NumAllowedFeatureLevels = ARRAY_COUNT(RequestedFeatureLevels);
+	int32 NumAllowedFeatureLevels = UE_ARRAY_COUNT(RequestedFeatureLevels);
 	int32 LastAllowedFeatureLevel = NumAllowedFeatureLevels - 1;
 	
 	while (FirstAllowedFeatureLevel < NumAllowedFeatureLevels)
diff --git a/Engine/Source/Runtime/Windows/XAudio2/Private/XAudio2Source.cpp b/Engine/Source/Runtime/Windows/XAudio2/Private/XAudio2Source.cpp
--- a/Engine/Source/Runtime/Windows/XAudio2/Private/XAudio2Source.cpp
+++ b/Engine/Source/Runtime/Windows/XAudio2/Private/XAudio2Source.cpp
@@ -1729,14 +1729,14 @@ void FSpatializationHelper::Init()
 	VolumeCurvePoint[0].DSPSetting = 1.0f;
 	VolumeCurvePoint[1].Distance = 1.0f;
 	VolumeCurvePoint[1].DSPSetting = 1.0f;
-	VolumeCurve.PointCount = ARRAY_COUNT( VolumeCurvePoint );
+	VolumeCurve.PointCount = UE_ARRAY_COUNT( VolumeCurvePoint );
 	VolumeCurve.pPoints = VolumeCurvePoint;
 
 	ReverbVolumeCurvePoint[0].Distance = 0.0f;
 	ReverbVolumeCurvePoint[0].DSPSetting = 0.5f;
 	ReverbVolumeCurvePoint[1].Distance = 1.0f;
 	ReverbVolumeCurvePoint[1].DSPSetting = 0.5f;
-	ReverbVolumeCurve.PointCount = ARRAY_COUNT( ReverbVolumeCurvePoint );
+	ReverbVolumeCurve.PointCount = UE_ARRAY_COUNT( ReverbVolumeCurvePoint );
 	ReverbVolumeCurve.pPoints = ReverbVolumeCurvePoint;
 
 	Emitter.pVolumeCurve = &VolumeCurve;
@@ -1748,7 +1748,7 @@ void FSpatializationHelper::Init()
 	Emitter.DopplerScaler = 1.0f;
 
 	// Zero the matrix coefficients
-	FMemory::Memzero(MatrixCoefficients, sizeof(float)*ARRAY_COUNT(MatrixCoefficients));
+	FMemory::Memzero(MatrixCoefficients, sizeof(float)*UE_ARRAY_COUNT(MatrixCoefficients));
 
 	DSPSettings.SrcChannelCount = UE4_XAUDIO3D_INPUTCHANNELS;
 	DSPSettings.DstChannelCount = SPEAKER_COUNT;
@@ -1835,7 +1835,7 @@ void FSpatializationHelper::DumpSpatializationState() const
 
 	// DSPSettings
 	UE_LOG(LogXAudio2, Log, TEXT("  DSPSettings"));
-	FLocal::DumpChannelArray(TEXT("    "), TEXT("pMatrixCoefficients"), ARRAY_COUNT(MatrixCoefficients), DSPSettings.pMatrixCoefficients);
+	FLocal::DumpChannelArray(TEXT("    "), TEXT("pMatrixCoefficients"), UE_ARRAY_COUNT(MatrixCoefficients), DSPSettings.pMatrixCoefficients);
 	UE_LOG(LogXAudio2, Log, TEXT("    SrcChannelCount: %u"), DSPSettings.SrcChannelCount);
 	UE_LOG(LogXAudio2, Log, TEXT("    DstChannelCount: %u"), DSPSettings.DstChannelCount);
 	UE_LOG(LogXAudio2, Log, TEXT("    LPFDirectCoefficient: %f"), DSPSettings.LPFDirectCoefficient);
@@ -1906,7 +1906,7 @@ void FSpatializationHelper::DumpSpatializationState() const
 	FLocal::DumpChannelArray(TEXT("  "), TEXT("EmitterAzimuths"), UE4_XAUDIO3D_INPUTCHANNELS, EmitterAzimuths);
 
 	// MatrixCoefficients
-	FLocal::DumpChannelArray(TEXT("  "), TEXT("MatrixCoefficients"), ARRAY_COUNT(MatrixCoefficients), MatrixCoefficients);
+	FLocal::DumpChannelArray(TEXT("  "), TEXT("MatrixCoefficients"), UE_ARRAY_COUNT(MatrixCoefficients), MatrixCoefficients);
 }
 
 void FSpatializationHelper::CalculateDolbySurroundRate( const FVector& OrientFront, const FVector& ListenerPosition, const FVector& EmitterPosition, float OmniRadius, float* OutVolumes  )
@@ -1961,7 +1961,7 @@ void FSpatializationHelper::CalculateDolbySurroundRate( const FVector& OrientFro
 			//	*NaNorINF, SpeakerIndex, OmniRadius, DSPSettings.pMatrixCoefficients[SpeakerIndex]);
 #endif
 			// Zero the coefficients so we don't continue getting bad values
-			FMemory::Memzero(MatrixCoefficients, sizeof(float)*ARRAY_COUNT(MatrixCoefficients));
+			FMemory::Memzero(MatrixCoefficients, sizeof(float)*UE_ARRAY_COUNT(MatrixCoefficients));
 		}
 #endif
 	}
