Updates and bugfixes to GTAO code
defaults to off. use r.AmbientOcclusion.Method to turn on

diff --git a/Engine/Shaders/Private/PostProcessAmbientOcclusion.usf b/Engine/Shaders/Private/PostProcessAmbientOcclusion.usf
--- a/Engine/Shaders/Private/PostProcessAmbientOcclusion.usf
+++ b/Engine/Shaders/Private/PostProcessAmbientOcclusion.usf
@@ -939,7 +939,7 @@ void MainSSAOSmoothCS(
 
 
 
-float4 GTAOParams[2];
+float4 GTAOParams[3];
 
 #define FALLOFF_START		100.0
 #define FALLOFF_END			200.0
@@ -1420,7 +1420,7 @@ void GTAOCombinedCS(
 	float4 OutColor		= 0;
 	int2   PixelPos		= DispatchThreadId + ScreenSpaceAOParams[5].zw; 
 	float2 PixelCenter	= (float2)PixelPos + float2(0.5, 0.5);
-	float2 BufferUV		= SvPositionToBufferUV(float4(PixelCenter, 0, 0));
+	float2 BufferUV		= PixelCenter.xy * GTAOParams[2].zw;
 
 	GTAOCombinedPSandCS(BufferUV , OutColor);
 
@@ -1455,10 +1455,8 @@ float GTAOInnerIntegratePSandCS(in float2 UV)
 	Angles				    = Angles *2-1;
 
 	// Get Angle
-	float2 ViewportUVToRandomUV = ScreenSpaceAOParams[1].xy;
-	const float2 TemporalOffset = ScreenSpaceAOParams[3].xy;
-	float2 RandomVec			= (Texture2DSample(RandomNormalTexture, RandomNormalTextureSampler, TemporalOffset + UV * ViewportUVToRandomUV).rg * 2 - 1) ;
-	float3 ScreenDir			= float3(RandomVec.x, RandomVec.y, 0.0);
+	float2 RandomVec		= GetRandomVector(UV);
+	float3 ScreenDir		= float3(RandomVec.x, RandomVec.y, 0.0);
 
 	// ViewspacePos and Normal
 	float3 ViewSpacePos		= ScreenToViewPos(UV, SceneDepth);
@@ -1500,8 +1498,7 @@ void GTAOInnerIntegrateCS(
 	float  OutColor		= 0;
 	int2   PixelPos		= DispatchThreadId + ScreenSpaceAOParams[5].zw; 
 	float2 PixelCenter	= (float2)PixelPos + float2(0.5, 0.5);
-	PixelCenter		   *= DownSampleFactor;
-	float2 BufferUV		= SvPositionToBufferUV(float4(PixelCenter, 0, 0));
+	float2 BufferUV		= PixelCenter.xy * GTAOParams[2].zw;
 
 	float AO = GTAOInnerIntegratePSandCS(BufferUV);
 	
@@ -1575,8 +1572,7 @@ void HorizonSearchCS(
 	float2 OutColor		= 0;
 	int2   PixelPos		= DispatchThreadId + ScreenSpaceAOParams[5].zw; 
 	float2 PixelCenter	= (float2)PixelPos + float2(0.5, 0.5);
-	PixelCenter		   *= DownSampleFactor;
-	float2 BufferUV		= SvPositionToBufferUV(float4(PixelCenter, 0, 0));
+	float2 BufferUV		= PixelCenter.xy * GTAOParams[2].zw;
 
 	float2 Horizons = HorizonSearchPSandCS(BufferUV);
 	
@@ -1599,14 +1595,14 @@ float2 ReprojectPos(float2 UV, float Depth)
 
 	PrevUV += PostprocessInput0Size.zw*0.25;
 
+	float2 Velocity		= GetGBufferData(UV, false).Velocity.xy;
 
 	// TODO - Take into account velocity
-//	float2 Velocity = Texture.SampleLevel( Sampler, UV, 0 ).xy;
-//	if( Velocity.x > 0.0 )
-//	{
-//		PrevScreen = ThisClip.xy - DecodeVelocityFromTexture( Velocity );
-//	}
-
+	if( Velocity.x > 0.0 )
+	{
+		PrevScreen = ThisClip.xy - DecodeVelocityFromTexture( Velocity );
+	}
+	
 
 	return PrevUV;
 }
diff --git a/Engine/Source/Runtime/Renderer/Private/CompositionLighting/CompositionLighting.cpp b/Engine/Source/Runtime/Renderer/Private/CompositionLighting/CompositionLighting.cpp
--- a/Engine/Source/Runtime/Renderer/Private/CompositionLighting/CompositionLighting.cpp
+++ b/Engine/Source/Runtime/Renderer/Private/CompositionLighting/CompositionLighting.cpp
@@ -40,6 +40,30 @@ static TAutoConsoleVariable<int32> CVarGTAODownsample(
 	TEXT("1: On (default)\n "),
 	ECVF_RenderThreadSafe | ECVF_Scalability);
 
+static TAutoConsoleVariable<int32> CVarGTAOTemporalFilter(
+	TEXT("r.GTAO.TemporalFilter"),
+	1,
+	TEXT("Enable Temporal Filter for GTAO \n ")
+	TEXT("0: Off \n ")
+	TEXT("1: On (default)\n "),
+	ECVF_RenderThreadSafe | ECVF_Scalability);
+
+static TAutoConsoleVariable<int32> CVarGTAOSpatialFilter(
+	TEXT("r.GTAO.SpatialFilter"),
+	1,
+	TEXT("Enable Spatial Filter for GTAO \n ")
+	TEXT("0: Off \n ")
+	TEXT("1: On (default)\n "),
+	ECVF_RenderThreadSafe | ECVF_Scalability);
+
+static TAutoConsoleVariable<int32> CVarGTAOCombined(
+	TEXT("r.GTAO.Combined"),
+	1,
+	TEXT("Enable Spatial Filter for GTAO \n ")
+	TEXT("0: Off \n ")
+	TEXT("1: On (default)\n "),
+	ECVF_RenderThreadSafe | ECVF_Scalability);
+
 bool IsAmbientCubemapPassRequired(const FSceneView& View)
 {
 	FScene* Scene = (FScene*)View.Family->Scene;
@@ -120,23 +144,17 @@ static FRenderingCompositeOutputRef AddPostProcessingGTAOAsyncHorizonSearch(FRHI
 
 
 	FSceneRenderTargets& SceneContext = FSceneRenderTargets::Get(Context.RHICmdList);
+	uint32 DownsampleFactor = CVarGTAODownsample.GetValueOnRenderThread() > 0 ? 2 : 1;
 
 	FIntPoint BufferSize	 = SceneContext.GetBufferSizeXY();
-	FIntPoint HalfBufferSize = FIntPoint::DivideAndRoundUp(BufferSize, 2);
-	FPooledRenderTargetDesc Desc(FPooledRenderTargetDesc::Create2DDesc(HalfBufferSize, PF_R8G8, FClearValueBinding::White, TexCreate_None, TexCreate_RenderTargetable, false));
+	FIntPoint HorizonBufferSize = FIntPoint::DivideAndRoundUp(BufferSize, DownsampleFactor);
+	FPooledRenderTargetDesc Desc(FPooledRenderTargetDesc::Create2DDesc(HorizonBufferSize, PF_R8G8, FClearValueBinding::White, TexCreate_None, TexCreate_RenderTargetable, false));
 	if (SceneContext.GetCurrentFeatureLevel() >= ERHIFeatureLevel::SM5)
 	{
 		Desc.TargetableFlags |= TexCreate_UAV;
 	}
 	GRenderTargetPool.FindFreeElement(RHICmdList, Desc, SceneContext.ScreenSpaceGTAOHorizons, TEXT("ScreenSpaceGTAOHorizons"));
 
-	uint32 DownsampleFactor = 1;
-
-	if (CVarGTAODownsample.GetValueOnRenderThread() > 0)
-	{
-		DownsampleFactor = 2;
-	}
-	
 	FRenderingCompositePass* HZBInput = Context.Graph.RegisterPass(new FRCPassPostProcessInput(const_cast<FViewInfo&>(Context.View).HZB));
 	FRenderingCompositePass* AmbientOcclusionHorizonSearch = Context.Graph.RegisterPass(new(FMemStack::Get()) FRCPassPostProcessAmbientOcclusion_HorizonSearch(Context.View, DownsampleFactor, ESSAOType::EAsyncCS));
 
@@ -156,35 +174,57 @@ static FRenderingCompositeOutputRef AddPostProcessingGTAOCombined(FRHICommandLis
 
 	FRenderingCompositePass* HZBInput = Context.Graph.RegisterPass(new FRCPassPostProcessInput(const_cast<FViewInfo&>(Context.View).HZB));
 
-	uint32 DownsampleFactor = 1;
-
-	if (CVarGTAODownsample.GetValueOnRenderThread() > 0)
-	{
-		DownsampleFactor = 2;
-	}
+	uint32 DownsampleFactor = CVarGTAODownsample.GetValueOnRenderThread() > 0 ? 2 : 1;
+	FSceneRenderTargets& SceneContext = FSceneRenderTargets::Get(RHICmdList);
 
+	if(CVarGTAOCombined.GetValueOnRenderThread() ==1)
 	{
 		FRenderingCompositePass* AmbientOcclusionGTAO = Context.Graph.RegisterPass(new(FMemStack::Get()) FRCPassPostProcessAmbientOcclusion_GTAOCombined(Context.View, DownsampleFactor, false));
 		AmbientOcclusionGTAO->SetInput(ePId_Input0, Context.SceneDepth);
 		AmbientOcclusionGTAO->SetInput(ePId_Input1, HZBInput);
 		FinalOutputPass = AmbientOcclusionGTAO;
 	}
+	else
+	{
+
+		FIntPoint BufferSize = SceneContext.GetBufferSizeXY();
+		FIntPoint HorizonBufferSize = FIntPoint::DivideAndRoundUp(BufferSize, DownsampleFactor);
+		FPooledRenderTargetDesc Desc(FPooledRenderTargetDesc::Create2DDesc(HorizonBufferSize, PF_R8G8, FClearValueBinding::White, TexCreate_None, TexCreate_RenderTargetable, false));
+		if (SceneContext.GetCurrentFeatureLevel() >= ERHIFeatureLevel::SM5)
+		{
+			Desc.TargetableFlags |= TexCreate_UAV;
+		}
+		GRenderTargetPool.FindFreeElement(RHICmdList, Desc, SceneContext.ScreenSpaceGTAOHorizons, TEXT("ScreenSpaceGTAOHorizons"));
+
+		FRenderingCompositePass* AmbientOcclusionHorizonSearch = Context.Graph.RegisterPass(new(FMemStack::Get()) FRCPassPostProcessAmbientOcclusion_HorizonSearch(Context.View, DownsampleFactor, ESSAOType::ECS));
+
+		AmbientOcclusionHorizonSearch->SetInput(ePId_Input0, Context.SceneDepth);
+		AmbientOcclusionHorizonSearch->SetInput(ePId_Input1, HZBInput);
+
+		FinalOutputPass = AmbientOcclusionHorizonSearch;
+
+		FRenderingCompositePass* AmbientOcclusionInnerIntegrate = Context.Graph.RegisterPass(new(FMemStack::Get()) FRCPassPostProcessAmbientOcclusion_InnerIntegrate(Context.View, DownsampleFactor, false));
+		AmbientOcclusionInnerIntegrate->SetInput(ePId_Input0, Context.SceneDepth);
+		AmbientOcclusionInnerIntegrate->SetInput(ePId_Input1, FinalOutputPass);
+		FinalOutputPass = AmbientOcclusionInnerIntegrate;
+	}
 
-	FSceneRenderTargets& SceneContext = FSceneRenderTargets::Get(RHICmdList);
 	SceneContext.bScreenSpaceAOIsValid = true;
 
 	FSceneViewState* ViewState = Context.View.ViewState;
 
 	// Add spatial Filter
+	if (CVarGTAOSpatialFilter.GetValueOnRenderThread() == 1)
 	{
 		FRenderingCompositePass* SpatialPass;
-		SpatialPass = Context.Graph.RegisterPass(new (FMemStack::Get()) FRCPassPostProcessAmbientOcclusion_GTAO_SpatialFilter(Context.View));
+		SpatialPass = Context.Graph.RegisterPass(new (FMemStack::Get()) FRCPassPostProcessAmbientOcclusion_GTAO_SpatialFilter(Context.View, DownsampleFactor));
 		SpatialPass->SetInput(ePId_Input0, FinalOutputPass);
 		SpatialPass->SetInput(ePId_Input1, HZBInput);
 		FinalOutputPass = SpatialPass;
 	}
 
-	if (ViewState)
+	
+	if (ViewState && CVarGTAOTemporalFilter.GetValueOnRenderThread() == 1)
 	{
 		// Add temporal filter
 		FRenderingCompositePass* TemporalPass;
@@ -196,11 +236,10 @@ static FRenderingCompositeOutputRef AddPostProcessingGTAOCombined(FRHICommandLis
 		FinalOutputPass = TemporalPass;
 
 	}
-		
-
+	
 	{
 		FRenderingCompositePass* UpsamplePass;
-		UpsamplePass = Context.Graph.RegisterPass(new (FMemStack::Get()) FRCPassPostProcessAmbientOcclusion_GTAO_Upsample());
+		UpsamplePass = Context.Graph.RegisterPass(new (FMemStack::Get()) FRCPassPostProcessAmbientOcclusion_GTAO_Upsample(Context.View, DownsampleFactor));
 		UpsamplePass->SetInput(ePId_Input0, FinalOutputPass);
 		FinalOutputPass = UpsamplePass;
 	}
@@ -215,13 +254,7 @@ static FRenderingCompositeOutputRef AddPostProcessingGTAOIntegration(FRHICommand
 	FRenderingCompositePass* FinalOutputPass;
 
 	FRenderingCompositePass* HZBInput = Context.Graph.RegisterPass(new FRCPassPostProcessInput(const_cast<FViewInfo&>(Context.View).HZB));
-
-	uint32 DownsampleFactor = 1;
-
-	if (CVarGTAODownsample.GetValueOnRenderThread() > 0)
-	{
-		DownsampleFactor = 2;
-	}
+	uint32 DownsampleFactor = CVarGTAODownsample.GetValueOnRenderThread() > 0 ? 2 : 1;
 
 	FRenderingCompositePass* AmbientOcclusionInnerIntegrate = Context.Graph.RegisterPass(new(FMemStack::Get()) FRCPassPostProcessAmbientOcclusion_InnerIntegrate(Context.View, DownsampleFactor, false));
 	AmbientOcclusionInnerIntegrate->SetInput(ePId_Input0, Context.SceneDepth);
@@ -232,19 +265,21 @@ static FRenderingCompositeOutputRef AddPostProcessingGTAOIntegration(FRHICommand
 
 	FSceneViewState* ViewState = Context.View.ViewState;
 
-
 	// Add spatial Filter
+	if (CVarGTAOSpatialFilter.GetValueOnRenderThread() == 1)
 	{
 		FRenderingCompositePass* SpatialPass;
-		SpatialPass = Context.Graph.RegisterPass(new (FMemStack::Get()) FRCPassPostProcessAmbientOcclusion_GTAO_SpatialFilter(Context.View));
+		SpatialPass = Context.Graph.RegisterPass(new (FMemStack::Get()) FRCPassPostProcessAmbientOcclusion_GTAO_SpatialFilter(Context.View, DownsampleFactor));
 		SpatialPass->SetInput(ePId_Input0, FinalOutputPass);
 		SpatialPass->SetInput(ePId_Input1, HZBInput);
 		FinalOutputPass = SpatialPass;
 	}
 
-	if (ViewState)
+//	bool bNeedsUpsample = DownsampleFactor != 1;
+
+	// Add temporal filter
+	if (ViewState && CVarGTAOTemporalFilter.GetValueOnRenderThread() == 1)
 	{
-		// Add temporal filter
 		FRenderingCompositePass* TemporalPass;
 		TemporalPass = Context.Graph.RegisterPass(new (FMemStack::Get()) FRCPassPostProcessAmbientOcclusion_GTAO_TemporalFilter(Context.View, DownsampleFactor,
 			Context.View.PrevViewInfo.GTAOHistory,
@@ -254,10 +289,9 @@ static FRenderingCompositeOutputRef AddPostProcessingGTAOIntegration(FRHICommand
 		FinalOutputPass = TemporalPass;
 
 	}
-
 	{
 		FRenderingCompositePass* UpsamplePass;
-		UpsamplePass = Context.Graph.RegisterPass(new (FMemStack::Get()) FRCPassPostProcessAmbientOcclusion_GTAO_Upsample());
+		UpsamplePass = Context.Graph.RegisterPass(new (FMemStack::Get()) FRCPassPostProcessAmbientOcclusion_GTAO_Upsample(Context.View, DownsampleFactor));
 		UpsamplePass->SetInput(ePId_Input0, FinalOutputPass);
 		FinalOutputPass = UpsamplePass;
 	}
diff --git a/Engine/Source/Runtime/Renderer/Private/CompositionLighting/PostProcessAmbientOcclusion.cpp b/Engine/Source/Runtime/Renderer/Private/CompositionLighting/PostProcessAmbientOcclusion.cpp
--- a/Engine/Source/Runtime/Renderer/Private/CompositionLighting/PostProcessAmbientOcclusion.cpp
+++ b/Engine/Source/Runtime/Renderer/Private/CompositionLighting/PostProcessAmbientOcclusion.cpp
@@ -101,14 +101,6 @@ static TAutoConsoleVariable<int32> CVarAmbientOcclusionMethod(
 	TEXT("1: GTAO\n "),
 	ECVF_RenderThreadSafe | ECVF_Scalability);
 
-static TAutoConsoleVariable<int32> CVarGTAOTemporalFilter(
-	TEXT("r.GTAO.TemporalFilter"),
-	1,
-	TEXT("Enable Temporal Filter for GTAO \n ")
-	TEXT("0: Off \n ")
-	TEXT("1: On (default)\n "),
-	ECVF_RenderThreadSafe | ECVF_Scalability);
-
 
 float FSSAOHelper::GetAmbientOcclusionQualityRT(const FSceneView& View)
 {
@@ -1221,7 +1213,7 @@ public:
 
 	//@param TRHICmdList could be async compute or compute dispatch, so template on commandlist type.
 	template<typename ShaderRHIParamRef, typename TRHICmdList>
-	void Set(TRHICmdList& RHICmdList, const FViewInfo& View, const ShaderRHIParamRef ShaderRHI) const
+	void Set(TRHICmdList& RHICmdList, const FViewInfo& View, FIntPoint DestSize, const ShaderRHIParamRef ShaderRHI) const
 	{
 		const FFinalPostProcessSettings& Settings = View.FinalPostProcessSettings;
 
@@ -1237,7 +1229,8 @@ public:
 		}
 			
 		TemporalFrame = Frame;
-		FVector4 GTAOParam[2];
+		const int ArraySize = 3;
+		FVector4 GTAOParam[ArraySize];
 
 		const float Rots[6]		= { 60.0f, 300.0f, 180.0f, 240.0f, 120.0f, 0.0f };
 		const float Offsets[4]	= { 0.0f, 0.5f, 0.25f, 0.75f };
@@ -1249,12 +1242,13 @@ public:
 
 		FIntPoint RandomizationSize = GSystemTextures.GTAORandomization->GetDesc().Extent;
 
-		GTAOParam[1] = FVector4(((float)TemporalFrame) * (1.0f / (float)RandomizationSize.X),
-			((float)TemporalFrame) * (3.0f / (float)RandomizationSize.X),
-			(float)TemporalFrame, 0.0f);
+		GTAOParam[1] = FVector4(0.0f, 0.0f, 0.0f, 0.0f);
 
+		float Fx = float(DestSize.X);
+		float Fy = float(DestSize.Y);
+		GTAOParam[2] = FVector4(Fx, Fy, 1.0f / Fx, 1.0f / Fy);
 
-		SetShaderValueArray(RHICmdList, ShaderRHI, GTAOParams, GTAOParam, 2);
+		SetShaderValueArray(RHICmdList, ShaderRHI, GTAOParams, GTAOParam, ArraySize);
 	}
 
 	friend FArchive& operator<<(FArchive& Ar, FGTAOParameters& This);
@@ -1270,6 +1264,14 @@ FArchive& operator<<(FArchive& Ar, FGTAOParameters& This)
 }
 
 
+static FVector4 GetHZBValue(const FViewInfo& View)
+{
+	return FVector4(float(View.ViewRect.Width()) / float(2 * View.HZBMipmap0Size.X),
+					float(View.ViewRect.Height()) / float(2 * View.HZBMipmap0Size.Y), 0.0f, 0.0f);
+
+	
+}
+
 
 template<uint32 bComputeShader, uint32 ShaderQuality>
 class FPostProcessGTAOHorizonSearchPSandCS : public FGlobalShader
@@ -1328,13 +1330,10 @@ public:
 
 	FVector4 GetHZBRemapVal(const FRenderingCompositePassContext& Context, FIntPoint DestSize, FIntPoint InputTextureSize)
 	{
-		const FVector2D HZBScaleFactor(float(Context.View.ViewRect.Width()) / float(2 * Context.View.HZBMipmap0Size.X),
-			float(Context.View.ViewRect.Height()) / float(2 * Context.View.HZBMipmap0Size.Y));
-
-		const FVector2D UVScaleFactor(float(DestSize.X) / float(InputTextureSize.X),
-			float(DestSize.Y) / float(InputTextureSize.Y));
+		const FVector2D HZBScaleFactor(float(InputTextureSize.X) / float(2 * Context.View.HZBMipmap0Size.X),
+			float(InputTextureSize.Y) / float(2 * Context.View.HZBMipmap0Size.Y));
 
-		return FVector4(HZBScaleFactor.X , HZBScaleFactor.Y , 0.0f, 0.0f);
+		return FVector4(HZBScaleFactor.X, HZBScaleFactor.Y, 0.0f, 0.0f);
 	}
 
 	template <typename TRHICmdList>
@@ -1347,7 +1346,8 @@ public:
 		SceneTextureParameters.Set(RHICmdList, ShaderRHI, Context.View.FeatureLevel, ESceneTextureSetupMode::All);
 		ScreenSpaceAOParams.Set(RHICmdList, Context.View, ShaderRHI, DestSize, FScreenSpaceAOParameters::ERandTexType::GTAO);
 
-		GTAOParams.Set(RHICmdList, Context.View, ShaderRHI);
+		GTAOParams.Set(RHICmdList, Context.View, DestSize, ShaderRHI);
+
 		PostprocessParameter.SetCS(ShaderRHI, Context, RHICmdList, TStaticSamplerState<SF_Point, AM_Clamp, AM_Clamp, AM_Clamp>::GetRHI());
 
 		RHICmdList.SetUAVParameter(ShaderRHI, HorizonOutTexture.GetBaseIndex(), OutUAV);
@@ -1372,7 +1372,7 @@ public:
 		FGlobalShader::SetParameters<FViewUniformShaderParameters>(Context.RHICmdList, ShaderRHI, Context.View.ViewUniformBuffer);
 		SceneTextureParameters.Set(Context.RHICmdList, ShaderRHI, Context.View.FeatureLevel, ESceneTextureSetupMode::All);
 		ScreenSpaceAOParams.Set(Context.RHICmdList, Context.View, ShaderRHI, DestSize, FScreenSpaceAOParameters::ERandTexType::GTAO);
-		GTAOParams.Set(Context.RHICmdList, Context.View, ShaderRHI);
+		GTAOParams.Set(Context.RHICmdList, Context.View, DestSize, ShaderRHI);
 
 		FVector4 HZBRemappingValue = GetHZBRemapVal(Context, DestSize, InputTextureSize);
 		SetShaderValue(Context.RHICmdList, ShaderRHI, HZBRemapping, HZBRemappingValue);
@@ -1481,32 +1481,29 @@ public:
 		}
 	}
 
+
 	FVector4 GetHZBRemapVal(const FRenderingCompositePassContext& Context, FIntPoint DestSize, FIntPoint InputTextureSize)
 	{
-		const FVector2D HZBScaleFactor(float(Context.View.ViewRect.Width()) / float(2 * Context.View.HZBMipmap0Size.X),
-			float(Context.View.ViewRect.Height()) / float(2 * Context.View.HZBMipmap0Size.Y));
+		const FVector2D HZBScaleFactor(float(InputTextureSize.X) / float(2 * Context.View.HZBMipmap0Size.X),
+			float(InputTextureSize.Y) / float(2 * Context.View.HZBMipmap0Size.Y));
 
-		const FVector2D UVScaleFactor(
-			float(InputTextureSize.X) / float(Context.View.ViewRect.Width()),
-			float(InputTextureSize.Y) / float(Context.View.ViewRect.Height()));
-
-		return FVector4(HZBScaleFactor.X * UVScaleFactor.X, HZBScaleFactor.Y / UVScaleFactor.Y, 0.0f, 0.0f);
+		return FVector4(HZBScaleFactor.X, HZBScaleFactor.Y, 0.0f, 0.0f);
 	}
 
 
 	template <typename TRHICmdList>
-	void SetParametersCS(TRHICmdList& RHICmdList, const FRenderingCompositePassContext& Context, FIntPoint DestSize, FIntPoint InputTextureSize, FRHIUnorderedAccessView *OutUAV)
+	void SetParametersCS(TRHICmdList& RHICmdList, const FRenderingCompositePassContext& Context, FIntPoint DestSize, FIntPoint InputTextureSize, uint32 DownScaleFactor, FRHIUnorderedAccessView *OutUAV)
 	{
 		const FFinalPostProcessSettings& Settings = Context.View.FinalPostProcessSettings;
 		FRHIComputeShader* ShaderRHI = GetComputeShader();
 
 		FGlobalShader::SetParameters<FViewUniformShaderParameters>(RHICmdList, ShaderRHI, Context.View.ViewUniformBuffer);
 		SceneTextureParameters.Set(RHICmdList, ShaderRHI, Context.View.FeatureLevel, ESceneTextureSetupMode::All);
-		ScreenSpaceAOParams.Set(RHICmdList, Context.View, ShaderRHI, InputTextureSize, FScreenSpaceAOParameters::ERandTexType::GTAO);
+		ScreenSpaceAOParams.Set(RHICmdList, Context.View, ShaderRHI, DestSize, FScreenSpaceAOParameters::ERandTexType::GTAO);
 
 		PostprocessParameter.SetCS(ShaderRHI, Context, RHICmdList, TStaticSamplerState<SF_Point, AM_Clamp, AM_Clamp, AM_Clamp>::GetRHI());
 
-		GTAOParams.Set(RHICmdList, Context.View, ShaderRHI);
+		GTAOParams.Set(RHICmdList, Context.View, DestSize, ShaderRHI);
 		RHICmdList.SetUAVParameter(ShaderRHI, OutTexture.GetBaseIndex(), OutUAV);
 
 		FVector4 HZBRemappingValue = GetHZBRemapVal(Context, DestSize, InputTextureSize);
@@ -1526,7 +1523,7 @@ public:
 		FGlobalShader::SetParameters<FViewUniformShaderParameters>(Context.RHICmdList, ShaderRHI, Context.View.ViewUniformBuffer);
 		SceneTextureParameters.Set(Context.RHICmdList, ShaderRHI, Context.View.FeatureLevel, ESceneTextureSetupMode::All);
 		ScreenSpaceAOParams.Set(Context.RHICmdList, Context.View, ShaderRHI, DestSize, FScreenSpaceAOParameters::ERandTexType::GTAO);
-		GTAOParams.Set(Context.RHICmdList, Context.View, ShaderRHI);
+		GTAOParams.Set(Context.RHICmdList, Context.View, DestSize, ShaderRHI);
 
 		FVector4 HZBRemappingValue = GetHZBRemapVal(Context, DestSize, InputTextureSize);
 		SetShaderValue(Context.RHICmdList, ShaderRHI, HZBRemapping, HZBRemappingValue);
@@ -1658,7 +1655,7 @@ public:
 		FGlobalShader::SetParameters<FViewUniformShaderParameters>(RHICmdList, ShaderRHI, Context.View.ViewUniformBuffer);
 		SceneTextureParameters.Set(RHICmdList, ShaderRHI, Context.View.FeatureLevel, ESceneTextureSetupMode::All);
 		ScreenSpaceAOParams.Set(RHICmdList, Context.View, ShaderRHI, InputTextureSize, FScreenSpaceAOParameters::ERandTexType::GTAO);
-		GTAOParams.Set(RHICmdList, Context.View, ShaderRHI);
+		GTAOParams.Set(RHICmdList, Context.View, DestSize, ShaderRHI);
 
 		PostprocessParameter.SetCS(ShaderRHI, Context, RHICmdList, TStaticSamplerState<SF_Point, AM_Clamp, AM_Clamp, AM_Clamp>::GetRHI());
 
@@ -1677,7 +1674,7 @@ public:
 	}
 
 
-	void SetParametersPS(const FRenderingCompositePassContext& Context, FIntPoint DestSize, FIntPoint InputTextureSize)
+	void SetParametersPS(const FRenderingCompositePassContext& Context, FIntPoint DestSize)
 	{
 		const FFinalPostProcessSettings& Settings = Context.View.FinalPostProcessSettings;
 		FRHIPixelShader* ShaderRHI = GetPixelShader();
@@ -1686,8 +1683,8 @@ public:
 
 		FGlobalShader::SetParameters<FViewUniformShaderParameters>(Context.RHICmdList, ShaderRHI, Context.View.ViewUniformBuffer);
 		SceneTextureParameters.Set(Context.RHICmdList, ShaderRHI, Context.View.FeatureLevel, ESceneTextureSetupMode::All);
-		ScreenSpaceAOParams.Set(Context.RHICmdList, Context.View, ShaderRHI, InputTextureSize, FScreenSpaceAOParameters::ERandTexType::GTAO);
-		GTAOParams.Set(Context.RHICmdList, Context.View, ShaderRHI);
+		ScreenSpaceAOParams.Set(Context.RHICmdList, Context.View, ShaderRHI, DestSize, FScreenSpaceAOParameters::ERandTexType::GTAO);
+		GTAOParams.Set(Context.RHICmdList, Context.View, DestSize, ShaderRHI);
 
 		const FSceneRenderTargetItem& GTAORandomization = GSystemTextures.GTAORandomization->GetRenderTargetItem();
 		SetTextureParameter(Context.RHICmdList, ShaderRHI, RandomNormalTexture, RandomNormalTextureSampler, TStaticSamplerState<SF_Point, AM_Wrap, AM_Wrap, AM_Wrap>::GetRHI(), GTAORandomization.ShaderResourceTexture);
@@ -1719,27 +1716,26 @@ public:
 IMPLEMENT_SHADER_TYPE2(FPostProcessGTAOInnerIntegratePSandCS<false>, SF_Pixel);
 IMPLEMENT_SHADER_TYPE2(FPostProcessGTAOInnerIntegratePSandCS<true>, SF_Compute);
 
-FRCPassPostProcessAmbientOcclusion_GTAOCombined::FRCPassPostProcessAmbientOcclusion_GTAOCombined(const FSceneView& View, uint32 inDownScaleFactor, bool FinalOutput)
+FRCPassPostProcessAmbientOcclusion_GTAOCombined::FRCPassPostProcessAmbientOcclusion_GTAOCombined(const FSceneView& View, uint32 InDownScaleFactor, bool FinalOutput)
 	:bFinalOutput(FinalOutput),
-	DownScaleFactor(inDownScaleFactor)
+	DownScaleFactor(InDownScaleFactor)
 {
 
 }
 
 template <uint32 ShaderQuality, uint32 UseNormals>
-void FRCPassPostProcessAmbientOcclusion_GTAOCombined::DispatchCS(const FRenderingCompositePassContext& Context, FIntPoint DestSize, FIntPoint TexSize)
+void FRCPassPostProcessAmbientOcclusion_GTAOCombined::DispatchCS(const FRenderingCompositePassContext& Context, FIntRect ViewRect, FIntPoint DestSize, FIntPoint TexSize)
 {
 	TShaderMapRef<FPostProcessGTAOCombinedPSandCS<true, ShaderQuality, UseNormals>> ComputeShader(Context.GetShaderMap());
 	FSceneRenderTargets& SceneContext = FSceneRenderTargets::Get(Context.RHICmdList);
-	const FSceneRenderTargetItem& DestRenderTarget = bFinalOutput ? SceneContext.ScreenSpaceAO->GetRenderTargetItem() : PassOutputs[0].RequestSurface(Context);
+	const FSceneRenderTargetItem& DestRenderTarget = PassOutputs[0].RequestSurface(Context);
 
 	Context.RHICmdList.SetComputeShader(ComputeShader->GetComputeShader());
-	ComputeShader->SetParametersCS(Context.RHICmdList, Context, DestSize, TexSize, DestRenderTarget.UAV);
+	ComputeShader->SetParametersCS(Context.RHICmdList, Context, DestSize, TexSize, DownScaleFactor, DestRenderTarget.UAV);
 
-	uint32 GroupSizeX = FMath::DivideAndRoundUp(DestSize.X, 8);
-	uint32 GroupSizeY = FMath::DivideAndRoundUp(DestSize.Y, 8);
+	uint32 GroupSizeX = FMath::DivideAndRoundUp(ViewRect.Width(), 8);
+	uint32 GroupSizeY = FMath::DivideAndRoundUp(ViewRect.Height(), 8);
 	DispatchComputeShader(Context.RHICmdList, *ComputeShader, GroupSizeX, GroupSizeY, 1);
-
 	ComputeShader->UnsetParameters(Context.RHICmdList);
 }
 
@@ -1777,14 +1773,11 @@ void FRCPassPostProcessAmbientOcclusion_GTAOCombined::Process(FRenderingComposit
 	const FPooledRenderTargetDesc* InputDesc0 = GetInputDesc(ePId_Input0);
 
 	// Get the size of the input and output sizes of the buffers.
-	FIntPoint InputTexSize	= InputDesc0 ? InputDesc0->Extent : SceneContext.GetBufferSizeXY();
-	FIntPoint OutputTexSize = PassOutputs[0].RenderTargetDesc.Extent; 
-	
-	FIntRect InputViewRect  = View.ViewRect;
-	FIntRect OutputViewRect = InputViewRect / DownScaleFactor;
-
-	FIntRect OutputFullRect(FIntPoint(0,0), OutputTexSize);
+	FIntPoint InputTexSize = InputDesc0 ? InputDesc0->Extent : SceneContext.GetBufferSizeXY();
+	FIntPoint OutputTexSize = PassOutputs[0].RenderTargetDesc.Extent;
+	FIntRect ViewRect = View.ViewRect / DownScaleFactor;
 
+	   
 	static const auto CVar = IConsoleManager::Get().FindTConsoleVariableDataInt(TEXT("r.AmbientOcclusion.Compute"));
 
 	const int32 ShaderQuality = FSSAOHelper::GetAmbientOcclusionShaderLevel(Context.View);
@@ -1800,19 +1793,19 @@ void FRCPassPostProcessAmbientOcclusion_GTAOCombined::Process(FRenderingComposit
 		switch(ShaderQuality)
 		{
 		case 0:
-			DispatchCS<0,0>(Context, OutputTexSize, InputTexSize);
+			DispatchCS<0,0>(Context, ViewRect, OutputTexSize, InputTexSize);
 			break;
 		case 1:
-			DispatchCS<1, 0>(Context, OutputTexSize, InputTexSize);
+			DispatchCS<1, 0>(Context, ViewRect, OutputTexSize, InputTexSize);
 			break;
 		case 2:
-			DispatchCS<2, 0>(Context, OutputTexSize, InputTexSize);
+			DispatchCS<2, 0>(Context, ViewRect, OutputTexSize, InputTexSize);
 			break;
 		case 3:
-			DispatchCS<3, 0>(Context, OutputTexSize, InputTexSize);
+			DispatchCS<3, 0>(Context, ViewRect, OutputTexSize, InputTexSize);
 			break;
 		case 4:
-			DispatchCS<4, 0>(Context, OutputTexSize, InputTexSize);
+			DispatchCS<4, 0>(Context, ViewRect, OutputTexSize, InputTexSize);
 			break;
 		}
 	}
@@ -1823,7 +1816,7 @@ void FRCPassPostProcessAmbientOcclusion_GTAOCombined::Process(FRenderingComposit
 		Context.RHICmdList.BeginRenderPass(RPInfo, TEXT("AmbientOcclusionSetup"));
 		{
 			// PS Version
-			Context.SetViewportAndCallRHI(OutputFullRect);
+			Context.SetViewportAndCallRHI(ViewRect);
 			DrawClearQuad(Context.RHICmdList, FLinearColor::White);
 
 			FGraphicsPipelineStateInitializer GraphicsPSOInit;
@@ -1857,14 +1850,17 @@ void FRCPassPostProcessAmbientOcclusion_GTAOCombined::Process(FRenderingComposit
 				break;
 			}
 
+
 			DrawRectangle(
 				Context.RHICmdList,
-				OutputViewRect.Min.X, OutputViewRect.Min.Y, OutputViewRect.Width(), OutputViewRect.Height(),
-				InputViewRect.Min.X, InputViewRect.Min.Y, InputViewRect.Width(), InputViewRect.Height(),
+				0, 0,ViewRect.Width(), ViewRect.Height(),
+				ViewRect.Min.X, ViewRect.Min.Y, ViewRect.Width(), ViewRect.Height(),
+				ViewRect.Size(),
 				OutputTexSize,
-				InputTexSize,
 				VertexShader,
 				EDRF_UseTriangleOptimization);
+
+
 		}
 		Context.RHICmdList.EndRenderPass();
 	}
@@ -1879,21 +1875,22 @@ FPooledRenderTargetDesc FRCPassPostProcessAmbientOcclusion_GTAOCombined::Compute
 
 	Ret = GetInput(ePId_Input0)->GetOutput()->RenderTargetDesc;
 	Ret.Reset();
-	Ret.Format = PF_G8;
+	Ret.Format = PF_G8; 
 	Ret.Extent = FIntPoint::DivideAndRoundUp(Ret.Extent, DownScaleFactor);
 
 	Ret.ClearValue = FClearValueBinding::None;
 	Ret.TargetableFlags &= ~TexCreate_DepthStencilTargetable;
 	Ret.TargetableFlags |= TexCreate_RenderTargetable | TexCreate_ShaderResource;
+	Ret.TargetableFlags |= TexCreate_UAV;
 	Ret.DebugName = TEXT("GTAOCombined");
 
 	return Ret;
 }
 
 
-FRCPassPostProcessAmbientOcclusion_InnerIntegrate::FRCPassPostProcessAmbientOcclusion_InnerIntegrate(const FSceneView& View, uint32 inDownScaleFactor, bool FinalOutput)
+FRCPassPostProcessAmbientOcclusion_InnerIntegrate::FRCPassPostProcessAmbientOcclusion_InnerIntegrate(const FSceneView& View, uint32 InDownScaleFactor, bool FinalOutput)
 	:bFinalOutput(FinalOutput),
-	DownScaleFactor(inDownScaleFactor)
+	DownScaleFactor(InDownScaleFactor)
 {
 }
 
@@ -1906,18 +1903,14 @@ void FRCPassPostProcessAmbientOcclusion_InnerIntegrate::Process(FRenderingCompos
 	FSceneRenderTargets& SceneContext = FSceneRenderTargets::Get(Context.RHICmdList);
 	const FPooledRenderTargetDesc* InputDesc0 = GetInputDesc(ePId_Input0);
 
-	FIntPoint InputTexSize = InputDesc0->Extent;
-	FIntPoint OutputTexSize = PassOutputs[0].RenderTargetDesc.Extent;
-	FIntRect OutputFullRect(FIntPoint(0, 0), OutputTexSize);
-
-	FIntRect InputViewRect = View.ViewRect / DownScaleFactor;
-	FIntRect OutputViewRect = InputViewRect;
+	FIntPoint TexSize = PassOutputs[0].RenderTargetDesc.Extent;
 
+	FIntRect ViewRect = View.ViewRect / DownScaleFactor;
 
 	const FSceneRenderTargetItem& DestRenderTarget = bFinalOutput ? SceneContext.ScreenSpaceAO->GetRenderTargetItem() : PassOutputs[0].RequestSurface(Context);
 	static const auto CVar = IConsoleManager::Get().FindTConsoleVariableDataInt(TEXT("r.AmbientOcclusion.Compute"));
 
-	if (CVar->GetValueOnRenderThread() >= 1)
+	if (0)// CVar->GetValueOnRenderThread() >= 1)
 	{
 		// Compute  Version
 		Context.RHICmdList.TransitionResource(EResourceTransitionAccess::ERWBarrier, EResourceTransitionPipeline::EGfxToCompute, DestRenderTarget.UAV, nullptr);
@@ -1926,10 +1919,10 @@ void FRCPassPostProcessAmbientOcclusion_InnerIntegrate::Process(FRenderingCompos
 
 		TShaderMapRef<FPostProcessGTAOInnerIntegratePSandCS<true>> ComputeShader(Context.GetShaderMap());
 		Context.RHICmdList.SetComputeShader(ComputeShader->GetComputeShader());
-		ComputeShader->SetParametersCS(Context.RHICmdList, Context, DownScaleFactor, OutputTexSize, InputTexSize, DestRenderTarget.UAV);
+		ComputeShader->SetParametersCS(Context.RHICmdList, Context, DownScaleFactor, TexSize, TexSize, DestRenderTarget.UAV);
 
-		uint32 GroupSizeX = FMath::DivideAndRoundUp(InputViewRect.Width(), 8);
-		uint32 GroupSizeY = FMath::DivideAndRoundUp(InputViewRect.Height(), 8);
+		uint32 GroupSizeX = FMath::DivideAndRoundUp(ViewRect.Width(), 8);
+		uint32 GroupSizeY = FMath::DivideAndRoundUp(ViewRect.Height(), 8);
 		DispatchComputeShader(Context.RHICmdList, *ComputeShader, GroupSizeX, GroupSizeY, 1);
 
 		ComputeShader->UnsetParameters(Context.RHICmdList);
@@ -1940,7 +1933,7 @@ void FRCPassPostProcessAmbientOcclusion_InnerIntegrate::Process(FRenderingCompos
 		FRHIRenderPassInfo RPInfo(DestRenderTarget.TargetableTexture, ERenderTargetActions::Load_Store);
 		Context.RHICmdList.BeginRenderPass(RPInfo, TEXT("AmbientOcclusionSetup"));
 		{
-			Context.SetViewportAndCallRHI(InputViewRect);
+			Context.SetViewportAndCallRHI(ViewRect);
 
 			TShaderMapRef<FPostProcessGTAOInnerIntegratePSandCS<false> > PixelShader(Context.GetShaderMap());
 			TShaderMapRef<FPostProcessVS> VertexShader(Context.GetShaderMap());
@@ -1956,16 +1949,14 @@ void FRCPassPostProcessAmbientOcclusion_InnerIntegrate::Process(FRenderingCompos
 			GraphicsPSOInit.BoundShaderState.PixelShaderRHI = GETSAFERHISHADER_PIXEL(*PixelShader);
 			SetGraphicsPipelineState(Context.RHICmdList, GraphicsPSOInit);
 
-			PixelShader->SetParametersPS(Context, OutputTexSize, InputTexSize);
+			PixelShader->SetParametersPS(Context, TexSize);
 
 			DrawRectangle(
 				Context.RHICmdList,
-				0, 0,
-				InputViewRect.Width(), InputViewRect.Height(),
-				InputViewRect.Min.X, InputViewRect.Min.Y,
-				InputViewRect.Width(), InputViewRect.Height(),
-				InputViewRect.Size(),
-				OutputTexSize,
+				0, 0, ViewRect.Width(), ViewRect.Height(),
+				ViewRect.Min.X, ViewRect.Min.Y, ViewRect.Width(), ViewRect.Height(),
+				ViewRect.Size(),
+				TexSize,
 				*VertexShader,
 				EDRF_UseTriangleOptimization);
 		}
@@ -2007,7 +1998,7 @@ template <uint32 ShaderQuality>
 void FRCPassPostProcessAmbientOcclusion_HorizonSearch::DispatchCS(const FRenderingCompositePassContext& Context, FIntRect ViewRect, FIntPoint DestSize, FIntPoint TexSize)
 {
 	FSceneRenderTargets& SceneContext = FSceneRenderTargets::Get(Context.RHICmdList);
-	const FSceneRenderTargetItem& DestRenderTarget = (AOType == ESSAOType::EAsyncCS) ? SceneContext.ScreenSpaceGTAOHorizons->GetRenderTargetItem() : PassOutputs[0].RequestSurface(Context);
+	const FSceneRenderTargetItem& DestRenderTarget = SceneContext.ScreenSpaceGTAOHorizons->GetRenderTargetItem();
 
 	FRHIAsyncComputeCommandListImmediate& RHICmdListComputeImmediate = FRHICommandListExecutor::GetImmediateAsyncComputeCommandList();
 
@@ -2015,7 +2006,7 @@ void FRCPassPostProcessAmbientOcclusion_HorizonSearch::DispatchCS(const FRenderi
 	RHICmdListComputeImmediate.SetComputeShader(ComputeShader->GetComputeShader());
 	ComputeShader->SetParametersCS(RHICmdListComputeImmediate, Context, DestSize, TexSize, DownScaleFactor, DestRenderTarget.UAV);
 
-	uint32 GroupSizeX = FMath::DivideAndRoundUp(ViewRect.Width() , 8);
+	uint32 GroupSizeX = FMath::DivideAndRoundUp(ViewRect.Width(), 8);
 	uint32 GroupSizeY = FMath::DivideAndRoundUp(ViewRect.Height(), 8);
 	DispatchComputeShader(RHICmdListComputeImmediate, *ComputeShader, GroupSizeX, GroupSizeY, 1);
 
@@ -2053,15 +2044,14 @@ void FRCPassPostProcessAmbientOcclusion_HorizonSearch::Process(FRenderingComposi
 	FSceneRenderTargets& SceneContext = FSceneRenderTargets::Get(Context.RHICmdList);
 	const FPooledRenderTargetDesc* InputDesc0 = GetInputDesc(ePId_Input0);
 
-	FIntPoint InputTexSize = InputDesc0->Extent;
-
+	// Get the size of the input and output sizes of the buffers.
+	FIntPoint InputTexSize = InputDesc0 ? InputDesc0->Extent : SceneContext.GetBufferSizeXY();
 	FIntPoint OutputTexSize = PassOutputs[0].RenderTargetDesc.Extent;
-	FIntRect OutputFullRect(FIntPoint(0, 0), OutputTexSize);
+	FIntRect ViewRect = View.ViewRect / DownScaleFactor;;
 
-	FIntRect InputViewRect = View.ViewRect / DownScaleFactor;
-	FIntRect OutputViewRect = InputViewRect;
 
-	const FSceneRenderTargetItem& DestRenderTarget = (AOType == ESSAOType::EAsyncCS) ? SceneContext.ScreenSpaceGTAOHorizons->GetRenderTargetItem()  : PassOutputs[0].RequestSurface(Context);
+	//const FSceneRenderTargetItem& DestRenderTarget = (AOType == ESSAOType::EAsyncCS) ? SceneContext.ScreenSpaceGTAOHorizons->GetRenderTargetItem()  : PassOutputs[0].RequestSurface(Context);
+	const FSceneRenderTargetItem& DestRenderTarget = SceneContext.ScreenSpaceGTAOHorizons->GetRenderTargetItem();
 
 	static const auto CVar = IConsoleManager::Get().FindTConsoleVariableDataInt(TEXT("r.AmbientOcclusion.Compute"));
 	const int32 ShaderQuality = FSSAOHelper::GetAmbientOcclusionShaderLevel(Context.View);
@@ -2085,19 +2075,19 @@ void FRCPassPostProcessAmbientOcclusion_HorizonSearch::Process(FRenderingComposi
 		switch (ShaderQuality)
 		{
 		case 0:
-			DispatchCS<0>(Context, InputViewRect, OutputTexSize, InputTexSize);
+			DispatchCS<0>(Context, ViewRect, OutputTexSize, InputTexSize);
 			break;
 		case 1:
-			DispatchCS<1>(Context, InputViewRect, OutputTexSize, InputTexSize);
+			DispatchCS<1>(Context, ViewRect, OutputTexSize, InputTexSize);
 			break;
 		case 2:
-			DispatchCS<2>(Context, InputViewRect, OutputTexSize, InputTexSize);
+			DispatchCS<2>(Context, ViewRect, OutputTexSize, InputTexSize);
 			break;
 		case 3:
-			DispatchCS<3>(Context, InputViewRect, OutputTexSize, InputTexSize);
+			DispatchCS<3>(Context, ViewRect, OutputTexSize, InputTexSize);
 			break;
 		case 4:
-			DispatchCS<4>(Context, InputViewRect, OutputTexSize, InputTexSize);
+			DispatchCS<4>(Context, ViewRect, OutputTexSize, InputTexSize);
 			break;
 		}
 	}
@@ -2107,7 +2097,7 @@ void FRCPassPostProcessAmbientOcclusion_HorizonSearch::Process(FRenderingComposi
 		Context.RHICmdList.BeginRenderPass(RPInfo, TEXT("GTAOHorizonSearch"));
 		{
 			// PS Version
-			Context.SetViewportAndCallRHI(InputViewRect);
+			Context.SetViewportAndCallRHI(ViewRect);
 
 			FGraphicsPipelineStateInitializer GraphicsPSOInit;
 			Context.RHICmdList.ApplyCachedRenderTargets(GraphicsPSOInit);
@@ -2142,12 +2132,10 @@ void FRCPassPostProcessAmbientOcclusion_HorizonSearch::Process(FRenderingComposi
 
 			DrawRectangle(
 				Context.RHICmdList,
-				0, 0,
-				InputViewRect.Width(), InputViewRect.Height(),
-				InputViewRect.Min.X,   InputViewRect.Min.Y,
-				InputViewRect.Width(), InputViewRect.Height(),
-				InputViewRect.Size(),
-				InputTexSize,
+				0, 0, ViewRect.Width(), ViewRect.Height(),
+				ViewRect.Min.X, ViewRect.Min.Y, ViewRect.Width(), ViewRect.Height(),
+				ViewRect.Size(),
+				OutputTexSize,
 				VertexShader,
 				EDRF_UseTriangleOptimization);
 		}
@@ -2166,7 +2154,7 @@ FPooledRenderTargetDesc FRCPassPostProcessAmbientOcclusion_HorizonSearch::Comput
 	Ret.Reset();
 	Ret.Format = PF_B8G8R8A8;
 
-	Ret.Extent = FIntPoint::DivideAndRoundUp(Ret.Extent, 2);
+	Ret.Extent = FIntPoint::DivideAndRoundUp(Ret.Extent, DownScaleFactor);
 
 	Ret.ClearValue = FClearValueBinding::None;
 	Ret.TargetableFlags &= ~TexCreate_DepthStencilTargetable;
@@ -2174,7 +2162,7 @@ FPooledRenderTargetDesc FRCPassPostProcessAmbientOcclusion_HorizonSearch::Comput
 	Ret.TargetableFlags |= TexCreate_RenderTargetable | TexCreate_ShaderResource;
 
 
-	Ret.DebugName = TEXT("GTAOInnerHorizon");
+	Ret.DebugName = TEXT("GTAOHorizonSearch");
 
 	return Ret;
 }
@@ -2547,7 +2535,8 @@ public:
 IMPLEMENT_SHADER_TYPE2(FPostProcessGTAOSpatialFilterPSandCS<true>,  SF_Compute);
 IMPLEMENT_SHADER_TYPE2(FPostProcessGTAOSpatialFilterPSandCS<false>, SF_Pixel);
 
-FRCPassPostProcessAmbientOcclusion_GTAO_SpatialFilter::FRCPassPostProcessAmbientOcclusion_GTAO_SpatialFilter(const FSceneView& View)
+FRCPassPostProcessAmbientOcclusion_GTAO_SpatialFilter::FRCPassPostProcessAmbientOcclusion_GTAO_SpatialFilter(const FSceneView& View, uint32 InDownScaleFactor)
+	:DownScaleFactor(InDownScaleFactor)
 {
 }
 
@@ -2562,8 +2551,6 @@ void FRCPassPostProcessAmbientOcclusion_GTAO_SpatialFilter::Process(FRenderingCo
 	FIntPoint OutputTexSize = SceneContext.GetBufferSizeXY();
 	FIntRect OutputFullRect(FIntPoint(0, 0), OutputTexSize);
 
-	int DownScaleFactor = 2;
-
 	FIntRect InputViewRect = View.ViewRect / DownScaleFactor;
 	FIntRect OutputViewRect = InputViewRect;
 
@@ -2675,7 +2662,8 @@ public:
 IMPLEMENT_SHADER_TYPE2(FPostProcessGTAOUpsamplePSandCS<true>, SF_Compute);
 IMPLEMENT_SHADER_TYPE2(FPostProcessGTAOUpsamplePSandCS<false>, SF_Pixel);
 
-FRCPassPostProcessAmbientOcclusion_GTAO_Upsample::FRCPassPostProcessAmbientOcclusion_GTAO_Upsample()
+FRCPassPostProcessAmbientOcclusion_GTAO_Upsample::FRCPassPostProcessAmbientOcclusion_GTAO_Upsample(const FSceneView& View, uint32 InDownScaleFactor)
+	:DownScaleFactor(InDownScaleFactor)
 {
 }
 
@@ -2693,8 +2681,6 @@ void FRCPassPostProcessAmbientOcclusion_GTAO_Upsample::Process(FRenderingComposi
 	FIntPoint OutputTexSize = SceneContext.GetBufferSizeXY();
 	FIntRect OutputFullRect(FIntPoint(0, 0), OutputTexSize);
 
-	int DownScaleFactor = 2;
-
 	FIntRect InputViewRect = View.ViewRect / DownScaleFactor;
 	FIntRect OutputViewRect = View.ViewRect;
 
diff --git a/Engine/Source/Runtime/Renderer/Private/CompositionLighting/PostProcessAmbientOcclusion.h b/Engine/Source/Runtime/Renderer/Private/CompositionLighting/PostProcessAmbientOcclusion.h
--- a/Engine/Source/Runtime/Renderer/Private/CompositionLighting/PostProcessAmbientOcclusion.h
+++ b/Engine/Source/Runtime/Renderer/Private/CompositionLighting/PostProcessAmbientOcclusion.h
@@ -196,7 +196,7 @@ public:
 	virtual FPooledRenderTargetDesc ComputeOutputDesc(EPassOutputId InPassOutputId) const override;
 
 	template <uint32 ShaderQuality, uint32 UseNormals>
-	void DispatchCS(const FRenderingCompositePassContext& Context, FIntPoint DestSize, FIntPoint TexSize);
+	void DispatchCS(const FRenderingCompositePassContext& Context, FIntRect ViewRect, FIntPoint DestSize, FIntPoint TexSize);
 
 	template <uint32 ShaderQuality, uint32 UseNormals>
 	FShader* SetShaderPS(const FRenderingCompositePassContext& Context, FGraphicsPipelineStateInitializer& GraphicsPSOInit, FIntPoint DestSize);
@@ -234,7 +234,7 @@ class FRCPassPostProcessAmbientOcclusion_GTAO_SpatialFilter : public TRenderingC
 {
 public:
 
-	FRCPassPostProcessAmbientOcclusion_GTAO_SpatialFilter(const FSceneView& View);
+	FRCPassPostProcessAmbientOcclusion_GTAO_SpatialFilter(const FSceneView& View, uint32 DownScaleFactor);
 
 	// interface FRenderingCompositePass ---------
 	virtual void Process(FRenderingCompositePassContext& Context) override;
@@ -242,6 +242,7 @@ public:
 	virtual FPooledRenderTargetDesc ComputeOutputDesc(EPassOutputId InPassOutputId) const override;
 
 private:
+	uint32		DownScaleFactor;
 
 };
 
@@ -249,7 +250,7 @@ class FRCPassPostProcessAmbientOcclusion_GTAO_Upsample : public TRenderingCompos
 {
 public:
 
-	FRCPassPostProcessAmbientOcclusion_GTAO_Upsample();
+	FRCPassPostProcessAmbientOcclusion_GTAO_Upsample(const FSceneView& View, uint32 DownScaleFactor);
 
 	// interface FRenderingCompositePass ---------
 	virtual void Process(FRenderingCompositePassContext& Context) override;
@@ -257,5 +258,6 @@ public:
 	virtual FPooledRenderTargetDesc ComputeOutputDesc(EPassOutputId InPassOutputId) const override;
 
 private:
+	uint32		DownScaleFactor;
 
 };
\ No newline at end of file
