Fix for velocity reproject calculations on AO. Handle camera Cuts. Use unclamped frame index as it has a different period than standard TAA

diff --git a/Engine/Shaders/Private/PostProcessAmbientOcclusion.usf b/Engine/Shaders/Private/PostProcessAmbientOcclusion.usf
--- a/Engine/Shaders/Private/PostProcessAmbientOcclusion.usf
+++ b/Engine/Shaders/Private/PostProcessAmbientOcclusion.usf
@@ -1581,8 +1581,9 @@ void HorizonSearchCS(
 #endif
 
 
-
-
+Texture2D SceneVelocityTexture;
+SamplerState SceneVelocityTextureSampler;
+float4 BlendParams;
 float2 ReprojectPos(float2 UV, float Depth)
 {
 	// Given a UV reproject where this was in the previous frame
@@ -1591,18 +1592,16 @@ float2 ReprojectPos(float2 UV, float Depth)
 	float4 ThisClip = float4( ThisScreen, Depth, 1 );
 	float4 PrevClip = mul( ThisClip, View.ClipToPrevClip );
 	float2 PrevScreen = PrevClip.xy / PrevClip.w;
-	float2 PrevUV = PrevScreen.xy * PrevScreenPositionScaleBias.xy + PrevScreenPositionScaleBias.zw;
-
-	PrevUV += PostprocessInput0Size.zw*0.25;
 
-	float2 Velocity		= GetGBufferData(UV, false).Velocity.xy;
 
-	// TODO - Take into account velocity
+	float2 Velocity	=  Texture2DSampleLevel(SceneVelocityTexture, SceneVelocityTextureSampler, UV,0).rg;
 	if( Velocity.x > 0.0 )
 	{
 		PrevScreen = ThisClip.xy - DecodeVelocityFromTexture( Velocity );
 	}
 	
+	float2 PrevUV = PrevScreen.xy * PrevScreenPositionScaleBias.xy + PrevScreenPositionScaleBias.zw;
+	PrevUV += PostprocessInput0Size.zw*0.25;
 
 	return PrevUV;
 }
@@ -1631,7 +1630,7 @@ void GTAOTemporalFilterPSAndCS(float2 UV, out float OutAO, out float OutDepth)
 	}
 
 	// Do a neighbourhood clamp 
-	float2 dUV = PostprocessInput0Size.zw*2;
+	float2 dUV = PostprocessInput0Size.zw*3;
 	float AONeighbour1 = Texture2DSample(PostprocessInput0, PostprocessInput0Sampler, UV + float2(-dUV.x,-dUV.y) ).r;
 	float AONeighbour2 = Texture2DSample(PostprocessInput0, PostprocessInput0Sampler, UV + float2(-dUV.x, dUV.y) ).r;
 	float AONeighbour3 = Texture2DSample(PostprocessInput0, PostprocessInput0Sampler, UV + float2( dUV.x,-dUV.y) ).r;
@@ -1639,12 +1638,21 @@ void GTAOTemporalFilterPSAndCS(float2 UV, out float OutAO, out float OutDepth)
 
 	float MinAO = min(NewAO, min(min(AONeighbour1, AONeighbour2), min(AONeighbour3, AONeighbour4)));
 	float MaxAO = max(NewAO, max(max(AONeighbour1, AONeighbour2), max(AONeighbour3, AONeighbour4)));
+	float HalfAO = (MaxAO - MinAO) *0.5;
+	float MidAO  = (MinAO + HalfAO);
+	const float ExpandNeighbour = 1.0;
+	MinAO = saturate(MidAO - (HalfAO * ExpandNeighbour));
+	MaxAO = saturate(MidAO + (HalfAO * ExpandNeighbour));
 
 	float BlurredAO = 0.2 * (NewAO + AONeighbour1 + AONeighbour2 + AONeighbour3 + AONeighbour4);
 
 	OutAO = (NewAO*Blend) + (HistoryAO* (1-Blend));
 	OutAO = clamp(OutAO, MinAO, MaxAO);
 
+	if(BlendParams.x > 0.5 )
+	{
+		OutAO = NewAO;
+	}
 
 	OutDepth = CurrDepth;
 }
diff --git a/Engine/Source/Runtime/Engine/Public/SceneManagement.h b/Engine/Source/Runtime/Engine/Public/SceneManagement.h
--- a/Engine/Source/Runtime/Engine/Public/SceneManagement.h
+++ b/Engine/Source/Runtime/Engine/Public/SceneManagement.h
@@ -206,6 +206,8 @@ public:
 	//
 	virtual uint32 GetCurrentTemporalAASampleIndex() const = 0;
 
+	virtual uint32 GetCurrentUnclampedTemporalAASampleIndex() const = 0;
+
 	virtual void SetSequencerState(ESequencerState InSequencerState) = 0;
 
 	virtual ESequencerState GetSequencerState() = 0;
diff --git a/Engine/Source/Runtime/Renderer/Private/CompositionLighting/PostProcessAmbientOcclusion.cpp b/Engine/Source/Runtime/Renderer/Private/CompositionLighting/PostProcessAmbientOcclusion.cpp
--- a/Engine/Source/Runtime/Renderer/Private/CompositionLighting/PostProcessAmbientOcclusion.cpp
+++ b/Engine/Source/Runtime/Renderer/Private/CompositionLighting/PostProcessAmbientOcclusion.cpp
@@ -1218,17 +1218,14 @@ public:
 		const FFinalPostProcessSettings& Settings = View.FinalPostProcessSettings;
 
 		uint32 TemporalFrame = 0;
-		uint32 Frame = 0;
 		
 		const FSceneViewState* ViewState = static_cast<const FSceneViewState*>(View.State);
 
 		if (ViewState)
 		{
-			TemporalFrame	= ViewState->GetCurrentTemporalAASampleIndex();
-			Frame			= ViewState->GetFrameIndex();
+			TemporalFrame	= ViewState->GetCurrentUnclampedTemporalAASampleIndex();
 		}
 			
-		TemporalFrame = Frame;
 		const int ArraySize = 3;
 		FVector4 GTAOParam[ArraySize];
 
@@ -2200,8 +2197,11 @@ public:
 	FShaderResourceParameter		DepthHistoryTexture;
 	FShaderResourceParameter		DepthHistoryTextureSampler;
 
+	FShaderResourceParameter		SceneVelocityTexture;
+	FShaderResourceParameter		SceneVelocityTextureSampler;
 
 	FShaderParameter				PrevScreenPositionScaleBias;
+	FShaderParameter				BlendParams;
 	FScreenSpaceAOParameters		ScreenSpaceAOParams;
 	FShaderParameter OutTexture;
 
@@ -2219,14 +2219,18 @@ public:
 		DepthHistoryTexture.Bind(Initializer.ParameterMap,			TEXT("DepthHistoryTexture"));
 		DepthHistoryTextureSampler.Bind(Initializer.ParameterMap,	TEXT("DepthHistoryTextureSampler"));
 
-		PrevScreenPositionScaleBias.Bind(Initializer.ParameterMap, TEXT("PrevScreenPositionScaleBias"));
+		SceneVelocityTexture.Bind(Initializer.ParameterMap,			TEXT("SceneVelocityTexture"));
+		SceneVelocityTextureSampler.Bind(Initializer.ParameterMap,	TEXT("SceneVelocityTextureSampler"));
+
+		PrevScreenPositionScaleBias.Bind(Initializer.ParameterMap,  TEXT("PrevScreenPositionScaleBias"));
+		BlendParams.Bind(Initializer.ParameterMap,					TEXT("BlendParams"));
 		OutTexture.Bind(Initializer.ParameterMap, TEXT("OutTexture"));
 
 		ScreenSpaceAOParams.Bind(Initializer.ParameterMap);
 	}
 
 
-	void SetParametersPS(const FRenderingCompositePassContext& Context, FIntPoint DestSize, FIntPoint InputTextureSize, const FGTAOTAAHistory& InputHistory)
+	void SetParametersPS(const FRenderingCompositePassContext& Context, FIntPoint DestSize, FIntPoint InputTextureSize, bool bCameraCut, const FGTAOTAAHistory& InputHistory, TRefCountPtr<IPooledRenderTarget>& VelocityRT)
 	{
 		const FFinalPostProcessSettings& Settings = Context.View.FinalPostProcessSettings;
 		FRHIPixelShader* ShaderRHI = GetPixelShader();
@@ -2248,6 +2252,10 @@ public:
 
 		SetShaderValue(Context.RHICmdList, ShaderRHI, PrevScreenPositionScaleBias, PrevScreenPositionScaleBiasValue);
 
+		FVector4 BlendParamsValue = FVector4(bCameraCut ? 1.0f : 0.0f, 0.0f, 0.0f, 0.0f);
+		SetShaderValue(Context.RHICmdList, ShaderRHI, BlendParams, BlendParamsValue);
+
+
 		if (InputHistory.IsValid())
 		{
 			SetTextureParameter(Context.RHICmdList, ShaderRHI, HistoryTexture, HistoryTextureSampler,
@@ -2267,6 +2275,10 @@ public:
 				GSystemTextures.WhiteDummy->GetRenderTargetItem().ShaderResourceTexture);
 
 		}
+
+		SetTextureParameter(Context.RHICmdList, ShaderRHI, SceneVelocityTexture, SceneVelocityTextureSampler,
+				TStaticSamplerState<SF_Point>::GetRHI(), VelocityRT->GetRenderTargetItem().ShaderResourceTexture);
+
 	}
 
 	template <typename TRHICmdList>
@@ -2282,6 +2294,7 @@ public:
 	{
 		bool bShaderHasOutdatedParameters = FGlobalShader::Serialize(Ar);
 		Ar << DepthHistoryTexture << DepthHistoryTextureSampler << HistoryTexture << HistoryTextureSampler << PostprocessParameter << SceneTextureParameters << PrevScreenPositionScaleBias << OutTexture;
+		Ar << SceneVelocityTexture << SceneVelocityTextureSampler << BlendParams;
 		return bShaderHasOutdatedParameters;
 	}
 
@@ -2397,7 +2410,8 @@ void FRCPassPostProcessAmbientOcclusion_GTAO_TemporalFilter::Process(FRenderingC
 			SetGraphicsPipelineState(Context.RHICmdList, GraphicsPSOInit);
 
 			VertexShader->SetParameters(Context);
-			PixelShader->SetParametersPS(Context, OutputTexSize, InputTexSize, InputHistory);
+			PixelShader->SetParametersPS(Context, OutputTexSize, InputTexSize, bCameraCut, InputHistory,
+											bCameraCut ? GSystemTextures.BlackDummy : SceneContext.SceneVelocity);
 
 			DrawRectangle(
 				Context.RHICmdList,
diff --git a/Engine/Source/Runtime/Renderer/Private/RendererScene.cpp b/Engine/Source/Runtime/Renderer/Private/RendererScene.cpp
--- a/Engine/Source/Runtime/Renderer/Private/RendererScene.cpp
+++ b/Engine/Source/Runtime/Renderer/Private/RendererScene.cpp
@@ -163,6 +163,7 @@ FSceneViewState::FSceneViewState()
 	CachedVisibilityChunkIndex = INDEX_NONE;
 	MIDUsedCount = 0;
 	TemporalAASampleIndex = 0;
+	TemporalAASampleIndexUnclamped = 0;
 	FrameIndex = 0;
 	DistanceFieldTemporalSampleIndex = 0;
 	AOTileIntersectionResources = NULL;
diff --git a/Engine/Source/Runtime/Renderer/Private/ScenePrivate.h b/Engine/Source/Runtime/Renderer/Private/ScenePrivate.h
--- a/Engine/Source/Runtime/Renderer/Private/ScenePrivate.h
+++ b/Engine/Source/Runtime/Renderer/Private/ScenePrivate.h
@@ -843,6 +843,9 @@ public:
 	// if TemporalAA is on this cycles through 0..TemporalAASampleCount-1, ResetViewState() puts it back to 0
 	int8 TemporalAASampleIndex;
 
+	// if TemporalAA is on this cycles through 0..Onwards, ResetViewState() puts it back to 0
+	uint32 TemporalAASampleIndexUnclamped;
+
 	// counts up by one each frame, warped in 0..7 range, ResetViewState() puts it back to 0
 	uint32 FrameIndex;
 	
@@ -990,6 +993,10 @@ public:
 		return TemporalAASampleIndex;
 	}
 
+	virtual uint32 GetCurrentUnclampedTemporalAASampleIndex() const
+	{
+		return TemporalAASampleIndexUnclamped;
+	}
 	// Returns the index of the frame with a desired power of two modulus.
 	inline uint32 GetFrameIndex(uint32 Pow2Modulus) const
 	{
@@ -1007,6 +1014,7 @@ public:
 	virtual void ResetViewState()
 	{
 		TemporalAASampleIndex = 0;
+		TemporalAASampleIndexUnclamped = 0;
 		FrameIndex = 0;
 		DistanceFieldTemporalSampleIndex = 0;
 		PreExposure = 1.f;
diff --git a/Engine/Source/Runtime/Renderer/Private/SceneVisibility.cpp b/Engine/Source/Runtime/Renderer/Private/SceneVisibility.cpp
--- a/Engine/Source/Runtime/Renderer/Private/SceneVisibility.cpp
+++ b/Engine/Source/Runtime/Renderer/Private/SceneVisibility.cpp
@@ -3111,7 +3111,7 @@ void FSceneRenderer::PreVisibilityFrameSetup(FRHICommandListImmediate& RHICmdLis
 			}
 
 			// Compute the new sample index in the temporal sequence.
-			int32 TemporalSampleIndex = ViewState->TemporalAASampleIndex + 1;
+			int32 TemporalSampleIndex			= ViewState->TemporalAASampleIndex + 1;
 			if(TemporalSampleIndex >= TemporalAASamples || View.bCameraCut)
 			{
 				TemporalSampleIndex = 0;
@@ -3120,7 +3120,8 @@ void FSceneRenderer::PreVisibilityFrameSetup(FRHICommandListImmediate& RHICmdLis
 			// Updates view state.
 			if (!View.bStatePrevViewInfoIsReadOnly && !bFreezeTemporalSequences)
 			{
-				ViewState->TemporalAASampleIndex = TemporalSampleIndex;
+				ViewState->TemporalAASampleIndex		  = TemporalSampleIndex;
+				ViewState->TemporalAASampleIndexUnclamped = ViewState->TemporalAASampleIndexUnclamped+1;
 			}
 
 			// Choose sub pixel sample coordinate in the temporal sequence.
