Prepare for batched runtime virtual texture updates

diff --git a/Engine/Source/Editor/VirtualTexturingEditor/Private/RuntimeVirtualTextureBuild.cpp b/Engine/Source/Editor/VirtualTexturingEditor/Private/RuntimeVirtualTextureBuild.cpp
--- a/Engine/Source/Editor/VirtualTexturingEditor/Private/RuntimeVirtualTextureBuild.cpp
+++ b/Engine/Source/Editor/VirtualTexturingEditor/Private/RuntimeVirtualTextureBuild.cpp
@@ -193,21 +193,25 @@ namespace RuntimeVirtualTexture
 						RHICmdList.TransitionResource(EResourceTransitionAccess::EWritable, RenderTileResources.GetRenderTarget(Layer));
 					}
 
-					// Render a single tile
-					RuntimeVirtualTexture::RenderPage(
-						RHICmdList,
-						Scene->GetRenderScene(),
-						1 << VirtualTextureSceneIndex,
-						MaterialType,
-						true,
-						RenderTileResources.GetRenderTarget(0), TileBox,
-						RenderTileResources.GetRenderTarget(1), TileBox,
-						RenderTileResources.GetRenderTarget(2), TileBox,
-						Transform,
-						UVRange,
-						RenderLevel,
-						MaxLevel,
-						DebugType);
+					RuntimeVirtualTexture::FRenderPageBatchDesc Desc;
+					Desc.Scene = Scene->GetRenderScene();
+					Desc.RuntimeVirtualTextureMask = 1 << VirtualTextureSceneIndex;
+					Desc.UVToWorld = Transform;
+					Desc.MaterialType = MaterialType;
+					Desc.MaxLevel = MaxLevel;
+					Desc.bClearTextures = true;
+					Desc.DebugType = DebugType;
+					Desc.NumPageDescs = 1;
+					Desc.Textures[0] = RenderTileResources.GetRenderTarget(0);
+					Desc.Textures[1] = RenderTileResources.GetRenderTarget(1);
+					Desc.Textures[2] = RenderTileResources.GetRenderTarget(2);
+					Desc.PageDescs[0].DestBox[0] = TileBox;
+					Desc.PageDescs[0].DestBox[1] = TileBox;
+					Desc.PageDescs[0].DestBox[2] = TileBox;
+					Desc.PageDescs[0].UVRange = UVRange;
+					Desc.PageDescs[0].vLevel = RenderLevel;
+
+					RuntimeVirtualTexture::RenderPages(RHICmdList, Desc);
 
 					// Transition render targets for copying
 					for (int32 Layer = 0; Layer < NumLayers; Layer++)
diff --git a/Engine/Source/Editor/VirtualTexturingEditor/Private/RuntimeVirtualTextureThumbnailRenderer.cpp b/Engine/Source/Editor/VirtualTexturingEditor/Private/RuntimeVirtualTextureThumbnailRenderer.cpp
--- a/Engine/Source/Editor/VirtualTexturingEditor/Private/RuntimeVirtualTextureThumbnailRenderer.cpp
+++ b/Engine/Source/Editor/VirtualTexturingEditor/Private/RuntimeVirtualTextureThumbnailRenderer.cpp
@@ -61,7 +61,7 @@ void URuntimeVirtualTextureThumbnailRenderer::Draw(UObject* Object, int32 X, int
 	FSceneInterface* Scene = RuntimeVirtualTextureComponent != nullptr ? RuntimeVirtualTextureComponent->GetScene() : nullptr;
 	check(Scene != nullptr);
 
-	const FBox2D DestRect = FBox2D(FVector2D(X, Y), FVector2D(Width, Height));
+	const FBox2D DestBox = FBox2D(FVector2D(X, Y), FVector2D(Width, Height));
 	const FTransform Transform = RuntimeVirtualTextureComponent->GetVirtualTextureTransform();
 	const uint32 VirtualTextureSceneIndex = RuntimeVirtualTexture::GetRuntimeVirtualTextureSceneIndex_GameThread(RuntimeVirtualTextureComponent);
 	const ERuntimeVirtualTextureMaterialType MaterialType = RuntimeVirtualTexture->GetMaterialType();
@@ -71,23 +71,26 @@ void URuntimeVirtualTextureThumbnailRenderer::Draw(UObject* Object, int32 X, int
 	const int32 MaxLevel = (int32)FMath::CeilLogTwo(FMath::Max(VTDesc.BlockWidthInTiles, VTDesc.BlockHeightInTiles));
 
 	ENQUEUE_RENDER_COMMAND(BakeStreamingTextureTileCommand)(
-		[Scene, VirtualTextureSceneIndex, MaterialType, RenderTarget, DestRect, Transform, MaxLevel](FRHICommandListImmediate& RHICmdList)
+		[Scene, VirtualTextureSceneIndex, MaterialType, RenderTarget, DestBox, Transform, MaxLevel](FRHICommandListImmediate& RHICmdList)
 	{
 		FMaterialRenderProxy::UpdateDeferredCachedUniformExpressions();
 
-		RuntimeVirtualTexture::RenderPage(
-			RHICmdList,
-			Scene->GetRenderScene(),
-			1 << VirtualTextureSceneIndex,
-			MaterialType,
-			true,
-			RenderTarget->GetRenderTargetTexture(), DestRect,
-			nullptr, DestRect,
-			nullptr, DestRect,
-			Transform,
-			FBox2D(FVector2D(0, 0), FVector2D(1, 1)),
-			MaxLevel,
-			MaxLevel,
-			ERuntimeVirtualTextureDebugType::None);
+		RuntimeVirtualTexture::FRenderPageBatchDesc Desc;
+		Desc.Scene = Scene->GetRenderScene();
+		Desc.RuntimeVirtualTextureMask = 1 << VirtualTextureSceneIndex;
+		Desc.UVToWorld = Transform;
+		Desc.MaterialType = MaterialType;
+		Desc.MaxLevel = MaxLevel;
+		Desc.bClearTextures = true;
+		Desc.DebugType = ERuntimeVirtualTextureDebugType::None;
+		Desc.NumPageDescs = 1;
+		Desc.Textures[0] = RenderTarget->GetRenderTargetTexture();
+		Desc.Textures[1] = nullptr;
+		Desc.Textures[2] = nullptr;
+		Desc.PageDescs[0].DestBox[0] = DestBox;
+		Desc.PageDescs[0].UVRange = FBox2D(FVector2D(0, 0), FVector2D(1, 1));
+		Desc.PageDescs[0].vLevel = MaxLevel;
+
+		RuntimeVirtualTexture::RenderPages(RHICmdList, Desc);
 	});
 }
diff --git a/Engine/Source/Runtime/Engine/Private/Materials/MaterialExpressions.cpp b/Engine/Source/Runtime/Engine/Private/Materials/MaterialExpressions.cpp
--- a/Engine/Source/Runtime/Engine/Private/Materials/MaterialExpressions.cpp
+++ b/Engine/Source/Runtime/Engine/Private/Materials/MaterialExpressions.cpp
@@ -2265,12 +2265,11 @@ int32 UMaterialExpressionRuntimeVirtualTextureSample::Compile(class FMaterialCom
 	}
 	
 	// Compile the texture object references
-	enum { MAX_RVT_LAYERS = 3 };
 	const int32 TextureLayerCount = URuntimeVirtualTexture::GetLayerCount(MaterialType);
-	check(TextureLayerCount <= MAX_RVT_LAYERS);
+	check(TextureLayerCount <= RuntimeVirtualTexture::MaxTextureLayers);
 
-	int32 TextureCodeIndex[MAX_RVT_LAYERS] = { INDEX_NONE };
-	int32 TextureReferenceIndex[MAX_RVT_LAYERS] = { INDEX_NONE };
+	int32 TextureCodeIndex[RuntimeVirtualTexture::MaxTextureLayers] = { INDEX_NONE };
+	int32 TextureReferenceIndex[RuntimeVirtualTexture::MaxTextureLayers] = { INDEX_NONE };
 	for (int32 TexureLayerIndex = 0; TexureLayerIndex < TextureLayerCount; TexureLayerIndex++)
 	{
 		const int32 PageTableLayerIndex = bSinglePhysicalSpace ? 0 : TexureLayerIndex;
@@ -2328,7 +2327,7 @@ int32 UMaterialExpressionRuntimeVirtualTextureSample::Compile(class FMaterialCom
 	}
 
 	// Compile the texture sample code
-	int32 SampleCodeIndex[MAX_RVT_LAYERS] = { INDEX_NONE };
+	int32 SampleCodeIndex[RuntimeVirtualTexture::MaxTextureLayers] = { INDEX_NONE };
 	for (int32 TexureLayerIndex = 0; TexureLayerIndex < TextureLayerCount; TexureLayerIndex++)
 	{
 		SampleCodeIndex[TexureLayerIndex] = Compiler->TextureSample(
diff --git a/Engine/Source/Runtime/Engine/Private/VT/RuntimeVirtualTexture.cpp b/Engine/Source/Runtime/Engine/Private/VT/RuntimeVirtualTexture.cpp
--- a/Engine/Source/Runtime/Engine/Private/VT/RuntimeVirtualTexture.cpp
+++ b/Engine/Source/Runtime/Engine/Private/VT/RuntimeVirtualTexture.cpp
@@ -504,10 +504,9 @@ void URuntimeVirtualTexture::InitializeStreamingTexture(uint32 InSizeX, uint32 I
 
 	StreamingTexture->BuildHash = GetStreamingTextureBuildHash();
 
-	enum { MAX_RVT_LAYERS = 3 };
 	const int32 LayerCount = GetLayerCount();
-	check(LayerCount <= MAX_RVT_LAYERS);
-	ETextureSourceFormat LayerFormats[MAX_RVT_LAYERS];
+	check(LayerCount <= RuntimeVirtualTexture::MaxTextureLayers);
+	ETextureSourceFormat LayerFormats[RuntimeVirtualTexture::MaxTextureLayers];
 
 	for (int32 Layer = 0; Layer < LayerCount; Layer++)
 	{
diff --git a/Engine/Source/Runtime/Engine/Public/MeshBatch.h b/Engine/Source/Runtime/Engine/Public/MeshBatch.h
--- a/Engine/Source/Runtime/Engine/Public/MeshBatch.h
+++ b/Engine/Source/Runtime/Engine/Public/MeshBatch.h
@@ -231,7 +231,7 @@ struct FMeshBatch
 	/** Whether the mesh batch can be rendered to virtual textures. */
 	uint32 bRenderToVirtualTexture : 1;
 	/** What virtual texture material type this mesh batch should be rendered with. */
-	uint32 RuntimeVirtualTextureMaterialType : ERuntimeVirtualTextureMaterialType_NumBits;
+	uint32 RuntimeVirtualTextureMaterialType : RuntimeVirtualTexture::MaterialType_NumBits;
 
 	// can be NULL
 	const FLightCacheInterface* LCI;
diff --git a/Engine/Source/Runtime/Engine/Public/VT/RuntimeVirtualTextureEnum.h b/Engine/Source/Runtime/Engine/Public/VT/RuntimeVirtualTextureEnum.h
--- a/Engine/Source/Runtime/Engine/Public/VT/RuntimeVirtualTextureEnum.h
+++ b/Engine/Source/Runtime/Engine/Public/VT/RuntimeVirtualTextureEnum.h
@@ -5,6 +5,9 @@
 #include "CoreMinimal.h"
 #include "RuntimeVirtualTextureEnum.generated.h"
 
+/** Maximum number of texture layers we will have in a runtime virtual texture. Increase if we add a ERuntimeVirtualTextureMaterialType with more layers. */
+namespace RuntimeVirtualTexture { enum { MaxTextureLayers = 3 }; }
+
 /** 
  * Enumeration of virtual texture stack layouts to support. 
  * Extend this enumeration with other layouts as required. For example we will probably want to add a displacement texture option.
@@ -23,8 +26,8 @@ enum class ERuntimeVirtualTextureMaterialType : uint8
 	Count UMETA(Hidden),
 };
 
-enum { ERuntimeVirtualTextureMaterialType_NumBits = 3 };
-static_assert((uint32)ERuntimeVirtualTextureMaterialType::Count <= (1 << (uint32)ERuntimeVirtualTextureMaterialType_NumBits), "NumBits is too small");
+namespace RuntimeVirtualTexture { enum { MaterialType_NumBits = 3 }; }
+static_assert((uint32)ERuntimeVirtualTextureMaterialType::Count <= (1 << (uint32)RuntimeVirtualTexture::MaterialType_NumBits), "NumBits is too small");
 
 /** Enumeration of main pass behaviors when rendering to a runtime virtual texture. */
 UENUM()
diff --git a/Engine/Source/Runtime/Renderer/Private/SceneCore.h b/Engine/Source/Runtime/Renderer/Private/SceneCore.h
--- a/Engine/Source/Runtime/Renderer/Private/SceneCore.h
+++ b/Engine/Source/Runtime/Renderer/Private/SceneCore.h
@@ -228,7 +228,7 @@ public:
 	/** Whether the mesh batch can be used for rendering to a virtual texture. */
 	uint8 bRenderToVirtualTexture : 1;
 	/** What virtual texture material type this mesh batch should be rendered with. */
-	uint8 RuntimeVirtualTextureMaterialType : ERuntimeVirtualTextureMaterialType_NumBits;
+	uint8 RuntimeVirtualTextureMaterialType : RuntimeVirtualTexture::MaterialType_NumBits;
 
 	/** Cached from vertex factory to avoid dereferencing VF in InitViews. */
 	uint8 bSupportsCachingMeshDrawCommands : 1;
diff --git a/Engine/Source/Runtime/Renderer/Private/VT/RuntimeVirtualTextureProducer.cpp b/Engine/Source/Runtime/Renderer/Private/VT/RuntimeVirtualTextureProducer.cpp
--- a/Engine/Source/Runtime/Renderer/Private/VT/RuntimeVirtualTextureProducer.cpp
+++ b/Engine/Source/Runtime/Renderer/Private/VT/RuntimeVirtualTextureProducer.cpp
@@ -55,18 +55,24 @@ void FRuntimeVirtualTextureFinalizer::AddTile(FTileEntry& Tile)
 
 void FRuntimeVirtualTextureFinalizer::Finalize(FRHICommandListImmediate& RHICmdList)
 {
-	for (auto Entry : Tiles)
+	RuntimeVirtualTexture::FRenderPageBatchDesc RenderPageBatchDesc;
+	RenderPageBatchDesc.Scene = Scene->GetRenderScene();
+	RenderPageBatchDesc.RuntimeVirtualTextureMask = RuntimeVirtualTextureMask;
+	RenderPageBatchDesc.UVToWorld = UVToWorld;
+	RenderPageBatchDesc.MaterialType = MaterialType;
+	RenderPageBatchDesc.MaxLevel = Desc.MaxLevel;
+	RenderPageBatchDesc.bClearTextures = bClearTextures;
+	RenderPageBatchDesc.DebugType = ERuntimeVirtualTextureDebugType::None;
+	
+	for (int LayerIndex = 0; LayerIndex < RuntimeVirtualTexture::MaxTextureLayers; ++LayerIndex)
 	{
-		const int32 TileSize = Desc.TileSize + 2 * Desc.TileBorderSize;
-		
-		const FVector2D DestinationBoxStart0(Entry.DestX0 * TileSize, Entry.DestY0 * TileSize);
-		const FBox2D DestinationBox0(DestinationBoxStart0, DestinationBoxStart0 + FVector2D(TileSize, TileSize));
-
-		const FVector2D DestinationBoxStart1(Entry.DestX1 * TileSize, Entry.DestY1 * TileSize);
-		const FBox2D DestinationBox1(DestinationBoxStart1, DestinationBoxStart1 + FVector2D(TileSize, TileSize));
+		RenderPageBatchDesc.Textures[LayerIndex] = Tiles[0].Textures[LayerIndex];
+	}
 
-		const FVector2D DestinationBoxStart2(Entry.DestX2 * TileSize, Entry.DestY2 * TileSize);
-		const FBox2D DestinationBox2(DestinationBoxStart2, DestinationBoxStart2 + FVector2D(TileSize, TileSize));
+	int32 BatchSize = 0;
+	for (auto Entry : Tiles)
+	{
+		RuntimeVirtualTexture::FRenderPageDesc& RenderPageDesc = RenderPageBatchDesc.PageDescs[BatchSize];
 
 		const float X = (float)FMath::ReverseMortonCode2(Entry.vAddress);
 		const float Y = (float)FMath::ReverseMortonCode2(Entry.vAddress >> 1);
@@ -78,23 +84,35 @@ void FRuntimeVirtualTextureFinalizer::Finalize(FRHICommandListImmediate& RHICmdL
 		const FVector2D UVBorder = UVSize * ((float)Desc.TileBorderSize / (float)Desc.TileSize);
 		const FBox2D UVRange(UV - UVBorder, UV + UVSize + UVBorder);
 
-		RuntimeVirtualTexture::RenderPage(
-			RHICmdList,
-			Scene->GetRenderScene(),
-			RuntimeVirtualTextureMask,
-			MaterialType,
-			bClearTextures,
-			Entry.Texture0, 
-			DestinationBox0, 
-			Entry.Texture1,
-			DestinationBox1,
-			Entry.Texture2,
-			DestinationBox2,
-			UVToWorld,
-			UVRange, 
-			Entry.vLevel,
-			Desc.MaxLevel,
-			ERuntimeVirtualTextureDebugType::None);
+		RenderPageDesc.vLevel = Entry.vLevel;
+		RenderPageDesc.UVRange = UVRange;
+
+		const int32 TileSize = Desc.TileSize + 2 * Desc.TileBorderSize;
+		for (int LayerIndex = 0; LayerIndex < RuntimeVirtualTexture::MaxTextureLayers; ++LayerIndex)
+		{
+			const FVector2D DestinationBoxStart0(Entry.DestPos[LayerIndex].X * TileSize, Entry.DestPos[LayerIndex].Y * TileSize);
+			RenderPageDesc.DestBox[LayerIndex] = FBox2D(DestinationBoxStart0, DestinationBoxStart0 + FVector2D(TileSize, TileSize));
+		}
+
+		bool bBreakBatchForTextures = false;
+		for (int LayerIndex = 0; LayerIndex < RuntimeVirtualTexture::MaxTextureLayers; ++LayerIndex)
+		{
+			// This should never happen which is why we don't bother pre-sorting to maximize batch size
+			bBreakBatchForTextures |= (RenderPageBatchDesc.Textures[LayerIndex] != Entry.Textures[LayerIndex]);
+		}
+
+		if (++BatchSize == RuntimeVirtualTexture::EMaxRenderPageBatch || bBreakBatchForTextures)
+		{
+			RenderPageBatchDesc.NumPageDescs = BatchSize;
+			RuntimeVirtualTexture::RenderPages(RHICmdList, RenderPageBatchDesc);
+			BatchSize = 0;
+		}
+	}
+
+	if (BatchSize > 0)
+	{
+		RenderPageBatchDesc.NumPageDescs = BatchSize;
+		RuntimeVirtualTexture::RenderPages(RHICmdList, RenderPageBatchDesc);
 	}
 
 	Tiles.SetNumUnsafeInternal(0);
@@ -143,26 +161,13 @@ IVirtualTextureFinalizer* FRuntimeVirtualTextureProducer::ProducePageData(
 	// This can be almost always be avoided by setting up the physical pools correctly for the application's needs.
 	// If we can't avoid partial layer masks then we could look at ways to handle it more efficiently (right now we render all layers even for these partial requests).
 
-	//todo[vt]: Tidy up logic now we have 3 layers
-	if (TargetLayers[0].TextureRHI != nullptr)
-	{
-		Tile.Texture0 = TargetLayers[0].TextureRHI->GetTexture2D();
-		Tile.DestX0 = TargetLayers[0].pPageLocation.X;
-		Tile.DestY0 = TargetLayers[0].pPageLocation.Y;
-	}
-
-	if (TargetLayers[1].TextureRHI != nullptr)
-	{
-		Tile.Texture1 = TargetLayers[1].TextureRHI->GetTexture2D();
-		Tile.DestX1 = TargetLayers[1].pPageLocation.X;
-		Tile.DestY1 = TargetLayers[1].pPageLocation.Y;
-	}
-
-	if (TargetLayers[2].TextureRHI != nullptr)
+	for (int LayerIndex = 0; LayerIndex < RuntimeVirtualTexture::MaxTextureLayers; ++LayerIndex)
 	{
-		Tile.Texture2 = TargetLayers[2].TextureRHI->GetTexture2D();
-		Tile.DestX2 = TargetLayers[2].pPageLocation.X;
-		Tile.DestY2 = TargetLayers[2].pPageLocation.Y;
+		if (TargetLayers[LayerIndex].TextureRHI != nullptr)
+		{
+			Tile.Textures[LayerIndex] = TargetLayers[LayerIndex].TextureRHI->GetTexture2D();
+			Tile.DestPos[LayerIndex] = FIntPoint(TargetLayers[LayerIndex].pPageLocation.X, TargetLayers[LayerIndex].pPageLocation.Y);
+		}
 	}
 
 	Finalizer.InitProducer(ProducerHandle);
diff --git a/Engine/Source/Runtime/Renderer/Private/VT/RuntimeVirtualTextureProducer.h b/Engine/Source/Runtime/Renderer/Private/VT/RuntimeVirtualTextureProducer.h
--- a/Engine/Source/Runtime/Renderer/Private/VT/RuntimeVirtualTextureProducer.h
+++ b/Engine/Source/Runtime/Renderer/Private/VT/RuntimeVirtualTextureProducer.h
@@ -5,6 +5,7 @@
 #include "CoreMinimal.h"
 #include "UObject/ObjectMacros.h"
 #include "VirtualTexturing.h"
+#include "VT/RuntimeVirtualTextureEnum.h"
 
 enum class ERuntimeVirtualTextureMaterialType : uint8;
 class FRHITexture2D;
@@ -20,15 +21,8 @@ public:
 	/** A description for a single tile to render. */
 	struct FTileEntry
 	{
-		FRHITexture2D* Texture0 = nullptr;
-		FRHITexture2D* Texture1 = nullptr;
-		FRHITexture2D* Texture2 = nullptr;
-		int32 DestX0 = 0;
-		int32 DestY0 = 0;
-		int32 DestX1 = 0;
-		int32 DestY1 = 0;
-		int32 DestX2 = 0;
-		int32 DestY2 = 0;
+		FRHITexture2D* Textures[RuntimeVirtualTexture::MaxTextureLayers] = { nullptr };
+		FIntPoint DestPos[RuntimeVirtualTexture::MaxTextureLayers];
 		uint32 vAddress = 0;
 		uint8 vLevel = 0;
 	};
diff --git a/Engine/Source/Runtime/Renderer/Private/VT/RuntimeVirtualTextureRender.cpp b/Engine/Source/Runtime/Renderer/Private/VT/RuntimeVirtualTextureRender.cpp
--- a/Engine/Source/Runtime/Renderer/Private/VT/RuntimeVirtualTextureRender.cpp
+++ b/Engine/Source/Runtime/Renderer/Private/VT/RuntimeVirtualTextureRender.cpp
@@ -573,7 +573,7 @@ namespace RuntimeVirtualTexture
 	struct FRenderGraphSetup
 	{
 		//todo[vt]: Add flag to disable the clear render target behavior and win some performance when we can. This could be driven a UI on the VT or the VT Plane?
-		FRenderGraphSetup(FRDGBuilder& GraphBuilder, ERuntimeVirtualTextureMaterialType MaterialType, FRHITexture2D* OutputTexture0, FRHITexture2D* OutputTexture1, FIntPoint TextureSize)
+		FRenderGraphSetup(FRDGBuilder& GraphBuilder, ERuntimeVirtualTextureMaterialType MaterialType, FRHITexture2D* OutputTexture0, FIntPoint TextureSize)
 		{
 			bRenderPass = OutputTexture0 != nullptr;
 			bCompressPass = bRenderPass && (OutputTexture0->GetFormat() == PF_DXT1 || OutputTexture0->GetFormat() == PF_DXT5 || OutputTexture0->GetFormat() == PF_BC5);
@@ -782,7 +782,7 @@ namespace RuntimeVirtualTexture
 		// Build graph
 		FMemMark Mark(FMemStack::Get());
 		FRDGBuilder GraphBuilder(RHICmdList);
-		FRenderGraphSetup GraphSetup(GraphBuilder, MaterialType, OutputTexture0, OutputTexture1, TextureSize);
+		FRenderGraphSetup GraphSetup(GraphBuilder, MaterialType, OutputTexture0, TextureSize);
 
 		// Draw Pass
 		if (GraphSetup.bRenderPass)
@@ -897,6 +897,33 @@ namespace RuntimeVirtualTexture
 		}
 	}
 
+	void RenderPages(FRHICommandListImmediate& RHICmdList, FRenderPageBatchDesc const& InDesc)
+	{
+		CSV_SCOPED_TIMING_STAT_EXCLUSIVE(STAT_VirtualTextureSystem_RVT_RenderPages);
+		SCOPED_DRAW_EVENT(RHICmdList, RuntimeVirtualTextureRenderPages);
+		check(InDesc.NumPageDescs <= EMaxRenderPageBatch);
+
+		for (int32 PageIndex = 0; PageIndex < InDesc.NumPageDescs; ++PageIndex)
+		{
+			FRenderPageDesc const& PageDesc = InDesc.PageDescs[PageIndex];
+
+			RenderPage(
+				RHICmdList,
+				InDesc.Scene,
+				InDesc.RuntimeVirtualTextureMask,
+				InDesc.MaterialType,
+				InDesc.bClearTextures,
+				InDesc.Textures[0], PageDesc.DestBox[0],
+				InDesc.Textures[1], PageDesc.DestBox[1],
+				InDesc.Textures[2], PageDesc.DestBox[2],
+				InDesc.UVToWorld,
+				PageDesc.UVRange,
+				PageDesc.vLevel,
+				InDesc.MaxLevel,
+				InDesc.DebugType);
+		}
+	}
+
 
 	uint32 GetRuntimeVirtualTextureSceneIndex_GameThread(class URuntimeVirtualTextureComponent* InComponent)
 	{
diff --git a/Engine/Source/Runtime/Renderer/Public/VT/RuntimeVirtualTextureRender.h b/Engine/Source/Runtime/Renderer/Public/VT/RuntimeVirtualTextureRender.h
--- a/Engine/Source/Runtime/Renderer/Public/VT/RuntimeVirtualTextureRender.h
+++ b/Engine/Source/Runtime/Renderer/Public/VT/RuntimeVirtualTextureRender.h
@@ -3,9 +3,8 @@
 #pragma once
 
 #include "CoreMinimal.h"
+#include "VT/RuntimeVirtualTextureEnum.h"
 
-enum class ERuntimeVirtualTextureDebugType;
-enum class ERuntimeVirtualTextureMaterialType : uint8;
 class FRHICommandListImmediate;
 class FRHITexture2D;
 class FScene;
@@ -22,25 +21,36 @@ namespace RuntimeVirtualTexture
 	RENDERER_API uint32 GetRuntimeVirtualTextureSceneIndex_GameThread(URuntimeVirtualTextureComponent* InComponent);
 #endif
 
+	/** Enum for our maximum RenderPages() batch size. */
+	enum { EMaxRenderPageBatch = 8 };
+
+	/** A single page description. Multiple of these can be placed in a single FRenderPageBatchDesc batch description. */
+	struct FRenderPageDesc
+	{
+		uint8 vLevel;
+		FBox2D UVRange;
+		FBox2D DestBox[RuntimeVirtualTexture::MaxTextureLayers];
+	};
+
+	/** A description of a batch of pages to be rendered with a single call to RenderPages(). */
+	struct FRenderPageBatchDesc
+	{
+		FScene* Scene;
+		uint32 RuntimeVirtualTextureMask;
+		FTransform UVToWorld;
+		ERuntimeVirtualTextureMaterialType MaterialType;
+		uint8 MaxLevel;
+		bool bClearTextures;
+		ERuntimeVirtualTextureDebugType DebugType;
+
+		int32 NumPageDescs;
+		FRHITexture2D* Textures[RuntimeVirtualTexture::MaxTextureLayers];
+		FRenderPageDesc PageDescs[EMaxRenderPageBatch];
+	};
+
 	/** Returns true if the FScene is initialized for rendering to runtime virtual textures. */
 	RENDERER_API bool IsSceneReadyToRender(FScene* Scene);
 
-	/**
-	 * Render a single page of a virtual texture with a given material.
-	 * todo[vt]: Likely to be more optimal to batch several pages at a time and share setup/visibility/render targets.
-	 */
-	RENDERER_API void RenderPage(
-		FRHICommandListImmediate& RHICmdList,
-		FScene* Scene,
-		uint32 RuntimeVirtualTextureMask,
-		ERuntimeVirtualTextureMaterialType MaterialType,
-		bool bClearTextures,
-		FRHITexture2D* Texture0, FBox2D const& DestBox0,
-		FRHITexture2D* Texture1, FBox2D const& DestBox1,
-		FRHITexture2D* Texture2, FBox2D const& DestBox2,
-		FTransform const& UVToWorld,
-		FBox2D const& UVRange,
-		uint8 vLevel,
-		uint8 MaxLevel,
-		ERuntimeVirtualTextureDebugType DebugType);
+	/** Render a batch of pages for a runtime virtual texture. */
+	RENDERER_API void RenderPages(FRHICommandListImmediate& RHICmdList, FRenderPageBatchDesc const& InDesc);
 }
