Support passing UAV of physical texture to RuntimeVirtualtexture::RenderPages
This will allow compressing direct to physical texture on platforms that support it
Some small RHI changes will have to be approved before turning on though

diff --git a/Engine/Source/Editor/UnrealEd/Private/Texture2DPreview.cpp b/Engine/Source/Editor/UnrealEd/Private/Texture2DPreview.cpp
--- a/Engine/Source/Editor/UnrealEd/Private/Texture2DPreview.cpp
+++ b/Engine/Source/Editor/UnrealEd/Private/Texture2DPreview.cpp
@@ -70,7 +70,7 @@ public:
 			FVirtualTexture2DResource* VirtualTextureValue = (FVirtualTexture2DResource*)TextureValue;
 			IAllocatedVirtualTexture* AllocatedVT = VirtualTextureValue->AcquireAllocatedVT();
 	
-			FRHIShaderResourceView* PhysicalView = AllocatedVT->GetPhysicalTextureView((uint32)LayerIndex, TextureValue->bSRGB);
+			FRHIShaderResourceView* PhysicalView = AllocatedVT->GetPhysicalTextureSRV((uint32)LayerIndex, TextureValue->bSRGB);
 			SetSRVParameter(RHICmdList, GetPixelShader(), InTexture, PhysicalView);
 			SetSamplerParameter(RHICmdList, GetPixelShader(), InTextureSampler, VirtualTextureValue->SamplerStateRHI);
 
diff --git a/Engine/Source/Editor/VirtualTexturingEditor/Private/RuntimeVirtualTextureBuild.cpp b/Engine/Source/Editor/VirtualTexturingEditor/Private/RuntimeVirtualTextureBuild.cpp
--- a/Engine/Source/Editor/VirtualTexturingEditor/Private/RuntimeVirtualTextureBuild.cpp
+++ b/Engine/Source/Editor/VirtualTexturingEditor/Private/RuntimeVirtualTextureBuild.cpp
@@ -202,9 +202,9 @@ namespace RuntimeVirtualTexture
 					Desc.bClearTextures = true;
 					Desc.DebugType = DebugType;
 					Desc.NumPageDescs = 1;
-					Desc.Textures[0] = RenderTileResources.GetRenderTarget(0);
-					Desc.Textures[1] = RenderTileResources.GetRenderTarget(1);
-					Desc.Textures[2] = RenderTileResources.GetRenderTarget(2);
+					Desc.Targets[0].Texture = RenderTileResources.GetRenderTarget(0);
+					Desc.Targets[1].Texture = RenderTileResources.GetRenderTarget(1);
+					Desc.Targets[2].Texture = RenderTileResources.GetRenderTarget(2);
 					Desc.PageDescs[0].DestBox[0] = TileBox;
 					Desc.PageDescs[0].DestBox[1] = TileBox;
 					Desc.PageDescs[0].DestBox[2] = TileBox;
diff --git a/Engine/Source/Editor/VirtualTexturingEditor/Private/RuntimeVirtualTextureThumbnailRenderer.cpp b/Engine/Source/Editor/VirtualTexturingEditor/Private/RuntimeVirtualTextureThumbnailRenderer.cpp
--- a/Engine/Source/Editor/VirtualTexturingEditor/Private/RuntimeVirtualTextureThumbnailRenderer.cpp
+++ b/Engine/Source/Editor/VirtualTexturingEditor/Private/RuntimeVirtualTextureThumbnailRenderer.cpp
@@ -84,9 +84,7 @@ void URuntimeVirtualTextureThumbnailRenderer::Draw(UObject* Object, int32 X, int
 		Desc.bClearTextures = true;
 		Desc.DebugType = ERuntimeVirtualTextureDebugType::None;
 		Desc.NumPageDescs = 1;
-		Desc.Textures[0] = RenderTarget->GetRenderTargetTexture();
-		Desc.Textures[1] = nullptr;
-		Desc.Textures[2] = nullptr;
+		Desc.Targets[0].Texture = RenderTarget->GetRenderTargetTexture();
 		Desc.PageDescs[0].DestBox[0] = DestBox;
 		Desc.PageDescs[0].UVRange = FBox2D(FVector2D(0, 0), FVector2D(1, 1));
 		Desc.PageDescs[0].vLevel = MaxLevel;
diff --git a/Engine/Source/Runtime/Engine/Private/Materials/MaterialUniformExpressions.cpp b/Engine/Source/Runtime/Engine/Private/Materials/MaterialUniformExpressions.cpp
--- a/Engine/Source/Runtime/Engine/Private/Materials/MaterialUniformExpressions.cpp
+++ b/Engine/Source/Runtime/Engine/Private/Materials/MaterialUniformExpressions.cpp
@@ -869,7 +869,7 @@ void FUniformExpressionSet::FillUniformBuffer(const FMaterialRenderContext& Mate
 					const IAllocatedVirtualTexture* AllocatedVT = UniformExpressionCache.AllocatedVTs[StackAndLayerIndex.StackIndex];
 					if (AllocatedVT != nullptr)
 					{
-						FRHIShaderResourceView* PhysicalViewRHI = AllocatedVT->GetPhysicalTextureView(StackAndLayerIndex.LayerIndex, VTResource->bSRGB);
+						FRHIShaderResourceView* PhysicalViewRHI = AllocatedVT->GetPhysicalTextureSRV(StackAndLayerIndex.LayerIndex, VTResource->bSRGB);
 						if (PhysicalViewRHI)
 						{
 							*ResourceTablePhysicalTexturePtr = PhysicalViewRHI;
@@ -891,7 +891,7 @@ void FUniformExpressionSet::FillUniformBuffer(const FMaterialRenderContext& Mate
 					if (AllocatedVT != nullptr)
 					{
 						const int32 LayerIndex = UniformVirtualTextureExpressions[ExpressionIndex]->GetTextureLayerIndex();
-						FRHIShaderResourceView* PhysicalViewRHI = AllocatedVT->GetPhysicalTextureView(LayerIndex, Texture->IsLayerSRGB(LayerIndex));
+						FRHIShaderResourceView* PhysicalViewRHI = AllocatedVT->GetPhysicalTextureSRV(LayerIndex, Texture->IsLayerSRGB(LayerIndex));
 						if (PhysicalViewRHI != nullptr)
 						{
 							*ResourceTablePhysicalTexturePtr = PhysicalViewRHI;
diff --git a/Engine/Source/Runtime/RenderCore/Public/VirtualTexturing.h b/Engine/Source/Runtime/RenderCore/Public/VirtualTexturing.h
--- a/Engine/Source/Runtime/RenderCore/Public/VirtualTexturing.h
+++ b/Engine/Source/Runtime/RenderCore/Public/VirtualTexturing.h
@@ -97,7 +97,6 @@ struct FVTProducerDescription
 	
 	bool bPersistentHighestMip = true;
 	bool bContinuousUpdate = false;
-	bool bCreateRenderTarget = false;
 	
 	uint32 TileSize = 0u;
 	uint32 TileBorderSize = 0u;
@@ -188,10 +187,10 @@ struct FVTRequestPageResult
 /** Describes a location to write a single layer of a VT tile */
 struct FVTProduceTargetLayer
 {
-	/** The texture to write to */
+	/** The texture to write to. */
 	FRHITexture* TextureRHI = nullptr;
-
-	TRefCountPtr<struct IPooledRenderTarget> PooledRenderTarget;
+	/** The UAV to write to. This may be nullptr if no suitable UAV can be created for the texture format.  */
+	FRHIUnorderedAccessView* UnorderedAccessViewRHI = nullptr;
 
 	/** Location within the texture to write */
 	FIntVector pPageLocation;
@@ -298,7 +297,7 @@ public:
 
 	virtual FRHITexture* GetPageTableTexture(uint32 InPageTableIndex) const = 0;
 	virtual FRHITexture* GetPhysicalTexture(uint32 InLayerIndex) const = 0;
-	virtual FRHIShaderResourceView* GetPhysicalTextureView(uint32 InLayerIndex, bool bSRGB) const = 0;
+	virtual FRHIShaderResourceView* GetPhysicalTextureSRV(uint32 InLayerIndex, bool bSRGB) const = 0;
 	virtual uint32 GetPhysicalTextureSize(uint32 InLayerIndex) const = 0;
 	virtual uint32 GetNumPageTableTextures() const = 0;
 
diff --git a/Engine/Source/Runtime/Renderer/Private/VT/AllocatedVirtualTexture.cpp b/Engine/Source/Runtime/Renderer/Private/VT/AllocatedVirtualTexture.cpp
--- a/Engine/Source/Runtime/Renderer/Private/VT/AllocatedVirtualTexture.cpp
+++ b/Engine/Source/Runtime/Renderer/Private/VT/AllocatedVirtualTexture.cpp
@@ -231,12 +231,12 @@ FRHITexture* FAllocatedVirtualTexture::GetPhysicalTexture(uint32 InLayerIndex) c
 	return nullptr;
 }
 
-FRHIShaderResourceView* FAllocatedVirtualTexture::GetPhysicalTextureView(uint32 InLayerIndex, bool bSRGB) const
+FRHIShaderResourceView* FAllocatedVirtualTexture::GetPhysicalTextureSRV(uint32 InLayerIndex, bool bSRGB) const
 {
 	if (InLayerIndex < Description.NumTextureLayers)
 	{
 		const FVirtualTexturePhysicalSpace* PhysicalSpace = UniquePageTableLayers[TextureLayers[InLayerIndex].UniquePageTableLayerIndex].PhysicalSpace;
-		return PhysicalSpace ? PhysicalSpace->GetPhysicalTextureView(TextureLayers[InLayerIndex].PhysicalTextureIndex, bSRGB) : nullptr;
+		return PhysicalSpace ? PhysicalSpace->GetPhysicalTextureSRV(TextureLayers[InLayerIndex].PhysicalTextureIndex, bSRGB) : nullptr;
 	}
 	return nullptr;
 }
diff --git a/Engine/Source/Runtime/Renderer/Private/VT/AllocatedVirtualTexture.h b/Engine/Source/Runtime/Renderer/Private/VT/AllocatedVirtualTexture.h
--- a/Engine/Source/Runtime/Renderer/Private/VT/AllocatedVirtualTexture.h
+++ b/Engine/Source/Runtime/Renderer/Private/VT/AllocatedVirtualTexture.h
@@ -38,7 +38,7 @@ public:
 	// begin IAllocatedVirtualTexture
 	virtual FRHITexture* GetPageTableTexture(uint32 InPageTableIndex) const override;
 	virtual FRHITexture* GetPhysicalTexture(uint32 InLayerIndex) const override;
-	virtual FRHIShaderResourceView* GetPhysicalTextureView(uint32 InLayerIndex, bool bSRGB) const override;
+	virtual FRHIShaderResourceView* GetPhysicalTextureSRV(uint32 InLayerIndex, bool bSRGB) const override;
 	virtual uint32 GetPhysicalTextureSize(uint32 InLayerIndex) const override;
 	virtual uint32 GetNumPageTableTextures() const override;
 	virtual void GetPackedPageTableUniform(FUintVector4* OutUniform, bool bApplyBlockScale) const override;
diff --git a/Engine/Source/Runtime/Renderer/Private/VT/RuntimeVirtualTextureProducer.cpp b/Engine/Source/Runtime/Renderer/Private/VT/RuntimeVirtualTextureProducer.cpp
--- a/Engine/Source/Runtime/Renderer/Private/VT/RuntimeVirtualTextureProducer.cpp
+++ b/Engine/Source/Runtime/Renderer/Private/VT/RuntimeVirtualTextureProducer.cpp
@@ -66,7 +66,8 @@ void FRuntimeVirtualTextureFinalizer::Finalize(FRHICommandListImmediate& RHICmdL
 	
 	for (int LayerIndex = 0; LayerIndex < RuntimeVirtualTexture::MaxTextureLayers; ++LayerIndex)
 	{
-		RenderPageBatchDesc.Textures[LayerIndex] = Tiles[0].Textures[LayerIndex];
+		RenderPageBatchDesc.Targets[LayerIndex].Texture = Tiles[0].Targets[LayerIndex].TextureRHI != nullptr ? Tiles[0].Targets[LayerIndex].TextureRHI->GetTexture2D() : nullptr;
+		RenderPageBatchDesc.Targets[LayerIndex].UAV = Tiles[0].Targets[LayerIndex].UnorderedAccessViewRHI;
 	}
 
 	int32 BatchSize = 0;
@@ -90,15 +91,15 @@ void FRuntimeVirtualTextureFinalizer::Finalize(FRHICommandListImmediate& RHICmdL
 		const int32 TileSize = Desc.TileSize + 2 * Desc.TileBorderSize;
 		for (int LayerIndex = 0; LayerIndex < RuntimeVirtualTexture::MaxTextureLayers; ++LayerIndex)
 		{
-			const FVector2D DestinationBoxStart0(Entry.DestPos[LayerIndex].X * TileSize, Entry.DestPos[LayerIndex].Y * TileSize);
+			const FVector2D DestinationBoxStart0(Entry.Targets[LayerIndex].pPageLocation.X * TileSize, Entry.Targets[LayerIndex].pPageLocation.Y * TileSize);
 			RenderPageDesc.DestBox[LayerIndex] = FBox2D(DestinationBoxStart0, DestinationBoxStart0 + FVector2D(TileSize, TileSize));
 		}
 
 		bool bBreakBatchForTextures = false;
 		for (int LayerIndex = 0; LayerIndex < RuntimeVirtualTexture::MaxTextureLayers; ++LayerIndex)
 		{
-			// This should never happen which is why we don't bother pre-sorting to maximize batch size
-			bBreakBatchForTextures |= (RenderPageBatchDesc.Textures[LayerIndex] != Entry.Textures[LayerIndex]);
+			// This should never happen which is why we don't bother sorting to maximize batch size
+			bBreakBatchForTextures |= (RenderPageBatchDesc.Targets[LayerIndex].Texture != Entry.Targets[LayerIndex].TextureRHI);
 		}
 
 		if (++BatchSize == RuntimeVirtualTexture::EMaxRenderPageBatch || bBreakBatchForTextures)
@@ -107,6 +108,15 @@ void FRuntimeVirtualTextureFinalizer::Finalize(FRHICommandListImmediate& RHICmdL
 			RuntimeVirtualTexture::RenderPages(RHICmdList, RenderPageBatchDesc);
 			BatchSize = 0;
 		}
+
+		if (bBreakBatchForTextures)
+		{
+			for (int LayerIndex = 0; LayerIndex < RuntimeVirtualTexture::MaxTextureLayers; ++LayerIndex)
+			{
+				RenderPageBatchDesc.Targets[LayerIndex].Texture = Tiles[0].Targets[LayerIndex].TextureRHI != nullptr ? Tiles[0].Targets[LayerIndex].TextureRHI->GetTexture2D() : nullptr;
+				RenderPageBatchDesc.Targets[LayerIndex].UAV = Tiles[0].Targets[LayerIndex].UnorderedAccessViewRHI;
+			}
+		}
 	}
 
 	if (BatchSize > 0)
@@ -156,17 +166,15 @@ IVirtualTextureFinalizer* FRuntimeVirtualTextureProducer::ProducePageData(
 	Tile.vAddress = vAddress;
 	Tile.vLevel = vLevel;
 
-	//todo[vt]: 
 	// Partial layer masks can happen when one layer has more physical space available so that old pages are evicted at different rates.
-	// This can be almost always be avoided by setting up the physical pools correctly for the application's needs.
-	// If we can't avoid partial layer masks then we could look at ways to handle it more efficiently (right now we render all layers even for these partial requests).
+	// We currently render all layers even for these partial requests. That might be considered inefficient?
+	// But since the problem is avoided by setting bSinglePhysicalSpace on the URuntimeVirtualTexture we can live with it.
 
 	for (int LayerIndex = 0; LayerIndex < RuntimeVirtualTexture::MaxTextureLayers; ++LayerIndex)
 	{
 		if (TargetLayers[LayerIndex].TextureRHI != nullptr)
 		{
-			Tile.Textures[LayerIndex] = TargetLayers[LayerIndex].TextureRHI->GetTexture2D();
-			Tile.DestPos[LayerIndex] = FIntPoint(TargetLayers[LayerIndex].pPageLocation.X, TargetLayers[LayerIndex].pPageLocation.Y);
+			Tile.Targets[LayerIndex] = TargetLayers[LayerIndex];
 		}
 	}
 
diff --git a/Engine/Source/Runtime/Renderer/Private/VT/RuntimeVirtualTextureProducer.h b/Engine/Source/Runtime/Renderer/Private/VT/RuntimeVirtualTextureProducer.h
--- a/Engine/Source/Runtime/Renderer/Private/VT/RuntimeVirtualTextureProducer.h
+++ b/Engine/Source/Runtime/Renderer/Private/VT/RuntimeVirtualTextureProducer.h
@@ -21,8 +21,7 @@ public:
 	/** A description for a single tile to render. */
 	struct FTileEntry
 	{
-		FRHITexture2D* Textures[RuntimeVirtualTexture::MaxTextureLayers] = { nullptr };
-		FIntPoint DestPos[RuntimeVirtualTexture::MaxTextureLayers];
+		FVTProduceTargetLayer Targets[RuntimeVirtualTexture::MaxTextureLayers];
 		uint32 vAddress = 0;
 		uint8 vLevel = 0;
 	};
diff --git a/Engine/Source/Runtime/Renderer/Private/VT/RuntimeVirtualTextureRender.cpp b/Engine/Source/Runtime/Renderer/Private/VT/RuntimeVirtualTextureRender.cpp
--- a/Engine/Source/Runtime/Renderer/Private/VT/RuntimeVirtualTextureRender.cpp
+++ b/Engine/Source/Runtime/Renderer/Private/VT/RuntimeVirtualTextureRender.cpp
@@ -715,10 +715,13 @@ namespace RuntimeVirtualTexture
 		ERuntimeVirtualTextureMaterialType MaterialType,
 		bool bClearTextures,
 		FRHITexture2D* OutputTexture0,
+		FRHIUnorderedAccessView* OutputUAV0,
 		FBox2D const& DestBox0,
 		FRHITexture2D* OutputTexture1,
+		FRHIUnorderedAccessView* OutputUAV1,
 		FBox2D const& DestBox1,
 		FRHITexture2D* OutputTexture2, 
+		FRHIUnorderedAccessView* OutputUAV2,
 		FBox2D const& DestBox2,
 		FTransform const& UVToWorld,
 		FBox2D const& UVRange,
@@ -913,9 +916,9 @@ namespace RuntimeVirtualTexture
 				InDesc.RuntimeVirtualTextureMask,
 				InDesc.MaterialType,
 				InDesc.bClearTextures,
-				InDesc.Textures[0], PageDesc.DestBox[0],
-				InDesc.Textures[1], PageDesc.DestBox[1],
-				InDesc.Textures[2], PageDesc.DestBox[2],
+				InDesc.Targets[0].Texture, InDesc.Targets[0].UAV, PageDesc.DestBox[0],
+				InDesc.Targets[1].Texture, InDesc.Targets[1].UAV, PageDesc.DestBox[1],
+				InDesc.Targets[2].Texture, InDesc.Targets[2].UAV, PageDesc.DestBox[2],
 				InDesc.UVToWorld,
 				PageDesc.UVRange,
 				PageDesc.vLevel,
diff --git a/Engine/Source/Runtime/Renderer/Private/VT/VirtualTexturePhysicalSpace.cpp b/Engine/Source/Runtime/Renderer/Private/VT/VirtualTexturePhysicalSpace.cpp
--- a/Engine/Source/Runtime/Renderer/Private/VT/VirtualTexturePhysicalSpace.cpp
+++ b/Engine/Source/Runtime/Renderer/Private/VT/VirtualTexturePhysicalSpace.cpp
@@ -7,6 +7,9 @@
 #include "Stats/Stats.h"
 #include "VT/VirtualTexturePoolConfig.h"
 
+// Will turn this on when RHI support is reviewed and submitted
+#define VIRTUALTEXTURE_UAV_ALIASING 0
+
 FVirtualTexturePhysicalSpace::FVirtualTexturePhysicalSpace(const FVTPhysicalSpaceDescription& InDesc, uint16 InID)
 	: Description(InDesc)
 	, NumRefs(0u)
@@ -57,30 +60,76 @@ FVirtualTexturePhysicalSpace::~FVirtualTexturePhysicalSpace()
 {
 }
 
+EPixelFormat GetUnorderedAccessViewFormat(EPixelFormat InFormat)
+{
+#if VIRTUALTEXTURE_UAV_ALIASING
+	// Use alias formats for compressed textures on APIs where that is possible
+	// This allows us to compress runtime data directly to the physical texture
+	const bool bUAVAliasForCompressedTextures = GRHISupportsUAVFormatAliasing;
+#else
+	const bool bUAVAliasForCompressedTextures = false;
+#endif
+
+	switch (InFormat)
+	{
+	case PF_DXT1: 
+	case PF_BC4: 
+		return bUAVAliasForCompressedTextures ? PF_R32G32_UINT : PF_Unknown;
+	case PF_DXT3: 
+	case PF_DXT5: 
+	case PF_BC5: 
+		return bUAVAliasForCompressedTextures ? PF_R32G32B32A32_UINT : PF_Unknown;
+	}
+
+	return InFormat;
+}
+
 void FVirtualTexturePhysicalSpace::InitRHI()
 {
 	FRHICommandListImmediate& RHICmdList = FRHICommandListExecutor::GetImmediateCommandList();
 
 	for (int32 Layer = 0; Layer < Description.NumLayers; ++Layer)
 	{
+		const EPixelFormat FormatSRV = Description.Format[Layer];
+		const EPixelFormat FormatUAV = GetUnorderedAccessViewFormat(FormatSRV);
+		const bool bCreateUAV = FormatUAV != PF_Unknown;
+		const bool bCreateAliasedUAV = bCreateUAV && (FormatUAV != FormatSRV);
+
+		// Allocate physical texture from the render target pool
 		const uint32 TextureSize = GetTextureSize();
 		const FPooledRenderTargetDesc Desc = FPooledRenderTargetDesc::Create2DDesc(
 			FIntPoint(TextureSize, TextureSize),
-			Description.Format[Layer],
+			FormatSRV,
 			FClearValueBinding::None,
 			TexCreate_None,
-			TexCreate_ShaderResource | (Description.bCreateRenderTarget ? (TexCreate_RenderTargetable | TexCreate_UAV) : 0),
+			bCreateUAV ? TexCreate_UAV : TexCreate_None,
 			false);
 
 		GRenderTargetPool.FindFreeElement(RHICmdList, Desc, PooledRenderTarget[Layer], TEXT("PhysicalTexture"));
 		FRHITexture* TextureRHI = PooledRenderTarget[Layer]->GetRenderTargetItem().ShaderResourceTexture;
 
-		// Create sRGB/non-sRGB views into the physical texture
-		FRHITextureSRVCreateInfo ViewInfo;
-		TextureView[Layer] = RHICreateShaderResourceView(TextureRHI, ViewInfo);
+		// Create sRGB and non-sRGB shader resource views into the physical texture
+		FRHITextureSRVCreateInfo SRVCreateInfo;
+		SRVCreateInfo.Format = FormatSRV;
+		TextureSRV[Layer] = RHICreateShaderResourceView(TextureRHI, SRVCreateInfo);
+
+		SRVCreateInfo.SRGBOverride = SRGBO_ForceEnable;
+		TextureSRV_SRGB[Layer] = RHICreateShaderResourceView(TextureRHI, SRVCreateInfo);
 
-		ViewInfo.SRGBOverride = SRGBO_ForceEnable;
-		TextureSRGBView[Layer] = RHICreateShaderResourceView(TextureRHI, ViewInfo);
+		if (bCreateUAV)
+		{
+#if VIRTUALTEXTURE_UAV_ALIASING
+			if (bCreateAliasedUAV)
+			{
+				// Specific API for format aliasing. Maybe RHI will unify this in future...
+				TextureUAV[Layer] = RHICreateUnorderedAccessView(TextureRHI, 0, FormatUAV);
+			}
+			else
+#endif
+			{
+				TextureUAV[Layer] = RHICreateUnorderedAccessView(TextureRHI);
+			}
+		}
 	}
 }
 
@@ -89,8 +138,9 @@ void FVirtualTexturePhysicalSpace::ReleaseRHI()
 	for (int32 Layer = 0; Layer < Description.NumLayers; ++Layer)
 	{
 		GRenderTargetPool.FreeUnusedResource(PooledRenderTarget[Layer]);
-		TextureView[Layer].SafeRelease();
-		TextureSRGBView[Layer].SafeRelease();
+		TextureSRV[Layer].SafeRelease();
+		TextureSRV_SRGB[Layer].SafeRelease();
+		TextureUAV[Layer].SafeRelease();
 	}
 }
 
diff --git a/Engine/Source/Runtime/Renderer/Private/VT/VirtualTexturePhysicalSpace.h b/Engine/Source/Runtime/Renderer/Private/VT/VirtualTexturePhysicalSpace.h
--- a/Engine/Source/Runtime/Renderer/Private/VT/VirtualTexturePhysicalSpace.h
+++ b/Engine/Source/Runtime/Renderer/Private/VT/VirtualTexturePhysicalSpace.h
@@ -15,7 +15,6 @@ struct FVTPhysicalSpaceDescription
 	uint8 NumLayers;
 	TEnumAsByte<EPixelFormat> Format[VIRTUALTEXTURE_SPACE_MAXLAYERS];
 	bool bContinuousUpdate;
-	bool bCreateRenderTarget;
 };
 
 inline bool operator==(const FVTPhysicalSpaceDescription& Lhs, const FVTPhysicalSpaceDescription& Rhs)
@@ -23,8 +22,7 @@ inline bool operator==(const FVTPhysicalSpaceDescription& Lhs, const FVTPhysical
 	if (Lhs.TileSize != Rhs.TileSize || 
 		Lhs.NumLayers != Rhs.NumLayers || 
 		Lhs.Dimensions != Rhs.Dimensions || 
-		Lhs.bContinuousUpdate != Rhs.bContinuousUpdate || 
-		Lhs.bCreateRenderTarget != Rhs.bCreateRenderTarget)
+		Lhs.bContinuousUpdate != Rhs.bContinuousUpdate)
 	{
 		return false;
 	}
@@ -78,16 +76,14 @@ public:
 		return PooledRenderTarget[Layer]->GetRenderTargetItem().ShaderResourceTexture;
 	}
 
-	TRefCountPtr<IPooledRenderTarget> GetPhysicalTexturePooledRenderTarget(int32 Layer) const
+	FRHIShaderResourceView* GetPhysicalTextureSRV(int32 Layer, bool bSRGB) const
 	{
-		check(PooledRenderTarget[Layer].IsValid());
-		check(Description.bCreateRenderTarget);
-		return PooledRenderTarget[Layer];
+		return bSRGB ? TextureSRV_SRGB[Layer] : TextureSRV[Layer];
 	}
 
-	FRHIShaderResourceView* GetPhysicalTextureView(int32 Layer, bool bSRGB) const
+	FRHIUnorderedAccessView* GetPhysicalTextureUAV(int32 Layer) const
 	{
-		return bSRGB ? TextureSRGBView[Layer] : TextureView[Layer];
+		return TextureUAV[Layer];
 	}
 
 #if STATS
@@ -104,8 +100,9 @@ private:
 	FVTPhysicalSpaceDescription Description;
 	FTexturePagePool Pool;
 	TRefCountPtr<IPooledRenderTarget> PooledRenderTarget[VIRTUALTEXTURE_SPACE_MAXLAYERS];
-	FShaderResourceViewRHIRef TextureView[VIRTUALTEXTURE_SPACE_MAXLAYERS];
-	FShaderResourceViewRHIRef TextureSRGBView[VIRTUALTEXTURE_SPACE_MAXLAYERS];
+	FShaderResourceViewRHIRef TextureSRV[VIRTUALTEXTURE_SPACE_MAXLAYERS];
+	FShaderResourceViewRHIRef TextureSRV_SRGB[VIRTUALTEXTURE_SPACE_MAXLAYERS];
+	FUnorderedAccessViewRHIRef TextureUAV[VIRTUALTEXTURE_SPACE_MAXLAYERS];
 
 	uint32 TextureSizeInTiles;
 	uint32 NumRefs;
diff --git a/Engine/Source/Runtime/Renderer/Private/VT/VirtualTextureProducer.cpp b/Engine/Source/Runtime/Renderer/Private/VT/VirtualTextureProducer.cpp
--- a/Engine/Source/Runtime/Renderer/Private/VT/VirtualTextureProducer.cpp
+++ b/Engine/Source/Runtime/Renderer/Private/VT/VirtualTextureProducer.cpp
@@ -63,7 +63,6 @@ FVirtualTextureProducerHandle FVirtualTextureProducerCollection::RegisterProduce
 	PhysicalSpaceDesc.Dimensions = InDesc.Dimensions;
 	PhysicalSpaceDesc.TileSize = InDesc.TileSize + InDesc.TileBorderSize * 2u;
 	PhysicalSpaceDesc.bContinuousUpdate = InDesc.bContinuousUpdate;
-	PhysicalSpaceDesc.bCreateRenderTarget = InDesc.bCreateRenderTarget;
 
 	Entry.Producer.PhysicalGroups.SetNum(InDesc.NumPhysicalGroups);
 	for (uint32 PhysicalGroupIndex = 0u; PhysicalGroupIndex < InDesc.NumPhysicalGroups; ++PhysicalGroupIndex)
@@ -188,7 +187,7 @@ uint32 FVirtualTextureProducerCollection::RemoveAllCallbacks(const void* Baton)
 			Callback.OwnerHandle = FVirtualTextureProducerHandle();
 
 			// If callback is already pending, we can't move it back to free list, or we risk corrupting the pending list while it's being iterated
-			// Setting DestroyedFunction tol nullptr here will ensure callback is no longer invoked, and it will be moved to free list later when it's removed from pending list
+			// Setting DestroyedFunction to nullptr here will ensure callback is no longer invoked, and it will be moved to free list later when it's removed from pending list
 			if (!Callback.bPending)
 			{
 				Callback.PackedFlags = 0u;
diff --git a/Engine/Source/Runtime/Renderer/Private/VT/VirtualTextureSystem.cpp b/Engine/Source/Runtime/Renderer/Private/VT/VirtualTextureSystem.cpp
--- a/Engine/Source/Runtime/Renderer/Private/VT/VirtualTextureSystem.cpp
+++ b/Engine/Source/Runtime/Renderer/Private/VT/VirtualTextureSystem.cpp
@@ -1669,10 +1669,7 @@ void FVirtualTextureSystem::SubmitRequestsFromLocalTileList(const TSet<FVirtualT
 					if (Producer.GetPhysicalGroupIndexForTextureLayer(ProducerLayerIndex) == ProducerPhysicalGroupIndex)
 					{
 						ProduceTarget[ProducerLayerIndex].TextureRHI = PhysicalSpace->GetPhysicalTexture(PhysicalLocalTextureIndex);
-						if (PhysicalSpace->GetDescription().bCreateRenderTarget)
-						{
-							ProduceTarget[ProducerLayerIndex].PooledRenderTarget = PhysicalSpace->GetPhysicalTexturePooledRenderTarget(PhysicalLocalTextureIndex);
-						}
+						ProduceTarget[ProducerLayerIndex].UnorderedAccessViewRHI = PhysicalSpace->GetPhysicalTextureUAV(PhysicalLocalTextureIndex);
 						ProduceTarget[ProducerLayerIndex].pPageLocation = PhysicalSpace->GetPhysicalLocation(pAddress);
 						LayerMask |= 1 << ProducerLayerIndex;
 						PhysicalLocalTextureIndex++;
@@ -1790,17 +1787,13 @@ void FVirtualTextureSystem::SubmitRequests(FRHICommandListImmediate& RHICmdList,
 							const uint32 pAddress = PagePool.Alloc(this, Frame, ProducerHandle, ProducerPhysicalGroupIndex, TileToLoad.Local_vAddress, TileToLoad.Local_vLevel, bLockTile);
 							check(pAddress != ~0u);
 
-							//JM Loop for textures
 							int32 PhysicalLocalTextureIndex = 0;
 							for (uint32 ProducerLayerIndex = 0u; ProducerLayerIndex < Producer.GetNumTextureLayers(); ++ProducerLayerIndex)
 							{
 								if (Producer.GetPhysicalGroupIndexForTextureLayer(ProducerLayerIndex) == ProducerPhysicalGroupIndex)
 								{
 									ProduceTarget[ProducerLayerIndex].TextureRHI = PhysicalSpace->GetPhysicalTexture(PhysicalLocalTextureIndex);
-									if (PhysicalSpace->GetDescription().bCreateRenderTarget)
-									{
-										ProduceTarget[ProducerLayerIndex].PooledRenderTarget = PhysicalSpace->GetPhysicalTexturePooledRenderTarget(PhysicalLocalTextureIndex);
-									}
+									ProduceTarget[ProducerLayerIndex].UnorderedAccessViewRHI = PhysicalSpace->GetPhysicalTextureUAV(PhysicalLocalTextureIndex);
 									ProduceTarget[ProducerLayerIndex].pPageLocation = PhysicalSpace->GetPhysicalLocation(pAddress);
 									
 									PhysicalLocalTextureIndex++;
@@ -1845,13 +1838,13 @@ void FVirtualTextureSystem::SubmitRequests(FRHICommandListImmediate& RHICmdList,
 								TEXT("%s missing tile: LayerMask: %X, Layer %d, vAddress %06X, vLevel %d"),
 								*Producer.GetName().ToString(), ProducerPhysicalGroupMask, ProducerPhysicalGroupIndex, TileToLoad.Local_vAddress, TileToLoad.Local_vLevel);
 							
-							// JM for all layers
 							int32 PhysicalLocalTextureIndex = 0;
 							for (uint32 ProducerLayerIndex = 0u; ProducerLayerIndex < Producer.GetNumTextureLayers(); ++ProducerLayerIndex)
 							{
 								if (Producer.GetPhysicalGroupIndexForTextureLayer(ProducerLayerIndex) == ProducerPhysicalGroupIndex)
 								{
 									ProduceTarget[ProducerLayerIndex].TextureRHI = PhysicalSpace->GetPhysicalTexture(PhysicalLocalTextureIndex);
+									ProduceTarget[ProducerLayerIndex].UnorderedAccessViewRHI = PhysicalSpace->GetPhysicalTextureUAV(PhysicalLocalTextureIndex);
 									ProduceTarget[ProducerLayerIndex].pPageLocation = PhysicalSpace->GetPhysicalLocation(pAddress);
 									PhysicalLocalTextureIndex++;
 								}
diff --git a/Engine/Source/Runtime/Renderer/Public/VT/RuntimeVirtualTextureRender.h b/Engine/Source/Runtime/Renderer/Public/VT/RuntimeVirtualTextureRender.h
--- a/Engine/Source/Runtime/Renderer/Public/VT/RuntimeVirtualTextureRender.h
+++ b/Engine/Source/Runtime/Renderer/Public/VT/RuntimeVirtualTextureRender.h
@@ -7,6 +7,7 @@
 
 class FRHICommandListImmediate;
 class FRHITexture2D;
+class FRHIUnorderedAccessView;
 class FScene;
 class URuntimeVirtualTextureComponent;
 
@@ -24,31 +25,54 @@ namespace RuntimeVirtualTexture
 	/** Enum for our maximum RenderPages() batch size. */
 	enum { EMaxRenderPageBatch = 8 };
 
+	/** Structure containing a texture layer target description for a call for RenderPages(). */
+	struct FRenderPageTarget
+	{
+		/** Physical texture to render to. */
+		FRHITexture2D* Texture = nullptr;
+		/** Unordered access view of physical texture to render to. If this exists we can render to it directly instead of using RHICopyTexture(). */
+		FRHIUnorderedAccessView* UAV = nullptr;
+	};
+
 	/** A single page description. Multiple of these can be placed in a single FRenderPageBatchDesc batch description. */
 	struct FRenderPageDesc
 	{
+		/** vLevel to render at. */
 		uint8 vLevel;
+		/** UV range to render in virtual texture space. */
 		FBox2D UVRange;
+		/** Destination box to render in texel space of the target physical texture. */
 		FBox2D DestBox[RuntimeVirtualTexture::MaxTextureLayers];
 	};
 
 	/** A description of a batch of pages to be rendered with a single call to RenderPages(). */
 	struct FRenderPageBatchDesc
 	{
+		/** Scene to use when rendering the batch. */
 		FScene* Scene;
+		/** Mask created from the target runtime virtual texture's index within the scene. */
 		uint32 RuntimeVirtualTextureMask;
+		/** Virtual texture UV space to world space transform. */
 		FTransform UVToWorld;
+		/** Material type of the runtime virtual texture that we are rendering. */
 		ERuntimeVirtualTextureMaterialType MaterialType;
+		/** Max mip level of the runtime virtual texture that we are rendering. */
 		uint8 MaxLevel;
+		/** Set to true to clear before rendering. */
 		bool bClearTextures;
+		/** Debug visualization to render with. */
 		ERuntimeVirtualTextureDebugType DebugType;
 
+		/** Physical texture targets to render to. */
+		FRenderPageTarget Targets[RuntimeVirtualTexture::MaxTextureLayers];
+		
+		/** Number of pages to render. */
 		int32 NumPageDescs;
-		FRHITexture2D* Textures[RuntimeVirtualTexture::MaxTextureLayers];
+		/** Page descriptions for each page in the batch. */
 		FRenderPageDesc PageDescs[EMaxRenderPageBatch];
 	};
 
-	/** Returns true if the FScene is initialized for rendering to runtime virtual textures. */
+	/** Returns true if the FScene is initialized for rendering to runtime virtual textures. Always check this before calling RenderPages(). */
 	RENDERER_API bool IsSceneReadyToRender(FScene* Scene);
 
 	/** Render a batch of pages for a runtime virtual texture. */
