Sanitized reciprocal operation to avoid NaNs.

diff --git a/Engine/Source/Runtime/Renderer/Private/SceneRendering.cpp b/Engine/Source/Runtime/Renderer/Private/SceneRendering.cpp
--- a/Engine/Source/Runtime/Renderer/Private/SceneRendering.cpp
+++ b/Engine/Source/Runtime/Renderer/Private/SceneRendering.cpp
@@ -1051,14 +1051,18 @@ void SetupPrecomputedVolumetricLightmapUniformBufferParameters(const FScene* Sce
 		ViewUniformShaderParameters.SkyBentNormalBrickTexture = OrBlack3DIfNull(VolumetricLightmapData->BrickData.SkyBentNormal.Texture);
 		ViewUniformShaderParameters.DirectionalLightShadowingBrickTexture = OrBlack3DIfNull(VolumetricLightmapData->BrickData.DirectionalLightShadowing.Texture);
 
-		const FBox& VolumeBounds = VolumetricLightmapData->GetBounds();
-		const FVector InvVolumeSize = FVector(1.0f) / VolumeBounds.GetSize();
+		const FBox VolumeBounds = VolumetricLightmapData->GetBounds();
+		const FVector VolumeSize = VolumeBounds.GetSize();
+		const FVector InvVolumeSize = VolumeSize.Reciprocal();
+
+		const FVector BrickDimensions(VolumetricLightmapData->BrickDataDimensions);
+		const FVector InvBrickDimensions = BrickDimensions.Reciprocal();
 
 		ViewUniformShaderParameters.VolumetricLightmapWorldToUVScale = InvVolumeSize;
 		ViewUniformShaderParameters.VolumetricLightmapWorldToUVAdd = -VolumeBounds.Min * InvVolumeSize;
 		ViewUniformShaderParameters.VolumetricLightmapIndirectionTextureSize = FVector(VolumetricLightmapData->IndirectionTextureDimensions);
 		ViewUniformShaderParameters.VolumetricLightmapBrickSize = VolumetricLightmapData->BrickSize;
-		ViewUniformShaderParameters.VolumetricLightmapBrickTexelSize = FVector(1.0f, 1.0f, 1.0f) / FVector(VolumetricLightmapData->BrickDataDimensions);
+		ViewUniformShaderParameters.VolumetricLightmapBrickTexelSize = InvBrickDimensions;
 	}
 	else
 	{
