Fixed missing resource barrier from render target to shader readable state for scene color and depth copy textures.

diff --git a/Engine/Source/Runtime/Renderer/Private/SingleLayerWaterRendering.cpp b/Engine/Source/Runtime/Renderer/Private/SingleLayerWaterRendering.cpp
--- a/Engine/Source/Runtime/Renderer/Private/SingleLayerWaterRendering.cpp
+++ b/Engine/Source/Runtime/Renderer/Private/SingleLayerWaterRendering.cpp
@@ -475,6 +475,9 @@ void FDeferredShadingSceneRenderer::CopySingleLayerWaterTextures(FRHICommandList
 		PassParameters,
 		RefractionViewRect);
 	GraphBuilder.Execute();
+
+	RHICmdList.TransitionResource(EResourceTransitionAccess::EReadable, PassData.SceneColorWithoutSingleLayerWater->GetRenderTargetItem().TargetableTexture.GetReference());
+	RHICmdList.TransitionResource(EResourceTransitionAccess::EReadable, PassData.SceneDepthWithoutSingleLayerWater->GetRenderTargetItem().TargetableTexture.GetReference());
 }
 
 void FDeferredShadingSceneRenderer::BeginRenderingWaterGBuffer(FRHICommandList& RHICmdList, FExclusiveDepthStencil::Type DepthStencilAccess, bool bBindQuadOverdrawBuffers, EShaderPlatform InShaderPlatform)
