Fix for non-power-2 UDIM virtual textures

diff --git a/Engine/Source/Runtime/Renderer/Private/VT/VirtualTextureSystem.cpp b/Engine/Source/Runtime/Renderer/Private/VT/VirtualTextureSystem.cpp
--- a/Engine/Source/Runtime/Renderer/Private/VT/VirtualTextureSystem.cpp
+++ b/Engine/Source/Runtime/Renderer/Private/VT/VirtualTextureSystem.cpp
@@ -1415,8 +1415,8 @@ void FVirtualTextureSystem::GatherRequestsTask(const FGatherRequestsParameters&
 				const uint32 ProducerMipHeightInTiles = FMath::Max(Producer->GetHeightInTiles() >> Local_vLevel, 1u);
 				uint32 Local_vTileX = FMath::ReverseMortonCode2(Local_vAddress);
 				uint32 Local_vTileY = FMath::ReverseMortonCode2(Local_vAddress >> 1);
-				Local_vTileX &= (ProducerMipWidthInTiles - 1u);
-				Local_vTileY &= (ProducerMipHeightInTiles - 1u);
+				Local_vTileX %= ProducerMipWidthInTiles;
+				Local_vTileY %= ProducerMipHeightInTiles;
 				Local_vAddress = FMath::MortonCode2(Local_vTileX) | (FMath::MortonCode2(Local_vTileY) << 1);
 			}
 
