Fix for water rendering order issue on mobile. (now it will automatically mimic other renderer order)

diff --git a/Engine/Source/Runtime/Renderer/Private/MobileBasePass.cpp b/Engine/Source/Runtime/Renderer/Private/MobileBasePass.cpp
--- a/Engine/Source/Runtime/Renderer/Private/MobileBasePass.cpp
+++ b/Engine/Source/Runtime/Renderer/Private/MobileBasePass.cpp
@@ -643,6 +643,8 @@ void FMobileBasePassMeshProcessor::Process(
 	if (bTranslucentBasePass)
 	{
 		SortKey = CalculateTranslucentMeshStaticSortKey(PrimitiveSceneProxy, MeshBatch.MeshIdInPrimitive);
+		// We always want water to be rendered first on mobile in order to mimic other renderers where it is opaque. We shift the other priorities by 1.
+		SortKey.Translucent.Priority = ShadingModels.HasShadingModel(MSM_SingleLayerWater) ? uint16(0) : uint16(FMath::Clamp(uint32(SortKey.Translucent.Priority) + 1, 0u, uint32(USHRT_MAX)));
 	}
 	else
 	{
