Fixed ensure in RTGI denoiser

diff --git a/Engine/Source/Runtime/Renderer/Private/ScreenSpaceDenoise.cpp b/Engine/Source/Runtime/Renderer/Private/ScreenSpaceDenoise.cpp
--- a/Engine/Source/Runtime/Renderer/Private/ScreenSpaceDenoise.cpp
+++ b/Engine/Source/Runtime/Renderer/Private/ScreenSpaceDenoise.cpp
@@ -1053,7 +1053,7 @@ static void DenoiseSignalAtConstantPixelDensity(
 		check(Union == View.ViewRect);
 	}
 
-	ensure(Settings.InputResolutionFraction == 1.0f || Settings.InputResolutionFraction == 0.5f);
+	ensure(Settings.InputResolutionFraction == 1.0f || Settings.InputResolutionFraction == 0.5f || Settings.InputResolutionFraction == 0.25f);
 	
 	auto GetResourceNames = [&](const TCHAR* const ResourceNames[])
 	{
