Fixed random corruption on the edges of the screen. It was caused by sampling refraction color and depth data outside of valid region. Fix involves proper clamping all refraction color and depth reads. Additionally replaced UVRect with MaxUV, as at the moment we never need to clamp lower range.

diff --git a/Engine/Shaders/Private/BasePassPixelShader.usf b/Engine/Shaders/Private/BasePassPixelShader.usf
--- a/Engine/Shaders/Private/BasePassPixelShader.usf
+++ b/Engine/Shaders/Private/BasePassPixelShader.usf
@@ -1238,7 +1238,7 @@ ShadingModel);
 			MaterialParameters, PixelMaterialInputs, ResolvedView,
 			OpaqueBasePass.SceneColorWithoutSingleLayerWaterTexture, SingleLayerWaterSceneColorSampler,
 			OpaqueBasePass.SceneDepthWithoutSingleLayerWaterTexture, SingleLayerWaterSceneDepthSampler,
-			OpaqueBasePass.SceneWithoutSingleLayerWaterValidUVRect,
+			OpaqueBasePass.SceneWithoutSingleLayerWaterMaxUV,
 			Specular, OpaqueBasePass.DistortionParams,
 			SunIlluminance, WaterDiffuseIndirectIlluminance, EnvBrdf,
 			CameraIsUnderWater, WaterVisibility).Luminance;
diff --git a/Engine/Shaders/Private/SingleLayerWaterCommon.ush b/Engine/Shaders/Private/SingleLayerWaterCommon.ush
--- a/Engine/Shaders/Private/SingleLayerWaterCommon.ush
+++ b/Engine/Shaders/Private/SingleLayerWaterCommon.ush
@@ -47,7 +47,7 @@ WaterVolumeLightingOutput EvaluateWaterVolumeLighting(
 #else
 	Texture2D SceneColorWithoutSingleLayerWaterTexture, SamplerState SceneColorWithoutSingleLayerWaterSampler,
 	Texture2D SceneDepthWithoutSingleLayerWaterTexture, SamplerState SceneDepthWithoutSingleLayerWaterSampler,
-	float4 SceneWithoutSingleLayerWaterValidUVRect,
+	float2 SceneWithoutSingleLayerWaterMaxUV,
 #endif // SIMPLE_SINGLE_LAYER_WATER
 	float Specular, const float4 DistortionParams,
 	float3 SunIlluminance, float3 AmbiantIlluminance, float3 EnvBrdf,
@@ -73,6 +73,8 @@ WaterVolumeLightingOutput EvaluateWaterVolumeLighting(
 
 #else  // !SIMPLE_SINGLE_LAYER_WATER
 
+		ViewportUV = min(ViewportUV, SceneWithoutSingleLayerWaterMaxUV);
+
 		float PixelSceneDepth = SceneDepthWithoutSingleLayerWaterTexture.SampleLevel(SceneDepthWithoutSingleLayerWaterSampler, ViewportUV, 0).x * SINGLE_LAYER_WATER_DEPTH_SCALE;
 		
 		float2 Refraction = GetMaterialRefraction(PixelMaterialInputs); // X = IOR, Y = Refraction Bias
@@ -83,8 +85,7 @@ WaterVolumeLightingOutput EvaluateWaterVolumeLighting(
 		// Also reduce refraction according to water thickness
 		BufferUVDistortion *= saturate((PixelSceneDepth - WaterDepth) * 1.0 / 30.0f);// *(1 - saturate(WaterDepth * 1.0 / 2000.0f));
 
-		// Clamp distorted UVs to not sample black outside the sub part of the buffer. Big quality increase.
-		const float2 DistortedUV = clamp(ViewportUV + BufferUVDistortion, SceneWithoutSingleLayerWaterValidUVRect.xy, SceneWithoutSingleLayerWaterValidUVRect.zw);
+		const float2 DistortedUV = min(ViewportUV + BufferUVDistortion, SceneWithoutSingleLayerWaterMaxUV);
 
 		float SceneDepth = SceneDepthWithoutSingleLayerWaterTexture.SampleLevel(SceneDepthWithoutSingleLayerWaterSampler, DistortedUV, 0).x * SINGLE_LAYER_WATER_DEPTH_SCALE;
 
diff --git a/Engine/Shaders/Private/SingleLayerWaterComposite.usf b/Engine/Shaders/Private/SingleLayerWaterComposite.usf
--- a/Engine/Shaders/Private/SingleLayerWaterComposite.usf
+++ b/Engine/Shaders/Private/SingleLayerWaterComposite.usf
@@ -21,13 +21,14 @@
 #include "ReflectionEnvironmentComposite.ush"
 
 
-
 Texture2D ScreenSpaceReflectionsTexture;
 SamplerState ScreenSpaceReflectionsSampler;
 
 Texture2D SceneNoWaterDepthTexture;
 SamplerState SceneNoWaterDepthSampler;
 
+float2 SceneNoWaterMaxUV;
+
 struct SingleLayerWaterCompositeOutput
 {
 	float4 LuminanceTransmittance;
@@ -57,7 +58,9 @@ SingleLayerWaterCompositeOutput SingleLayerWaterComposite(float2 BufferUV, float
 #endif
 	if (ShadingModelID == SHADINGMODELID_SINGLELAYERWATER)
 	{
-		float OpaqueDepth = SceneNoWaterDepthTexture.SampleLevel(SceneNoWaterDepthSampler, BufferUV, 0).x * SINGLE_LAYER_WATER_DEPTH_SCALE;
+		const float2 SceneNoWaterUV = min(BufferUV, SceneNoWaterMaxUV);
+
+		float OpaqueDepth = SceneNoWaterDepthTexture.SampleLevel(SceneNoWaterDepthSampler, SceneNoWaterUV, 0).x * SINGLE_LAYER_WATER_DEPTH_SCALE;
 		float WaterDepth = GBuffer.Depth;
 
 		float DeltaDepth = CameraIsUnderWater ? WaterDepth : OpaqueDepth - WaterDepth;	// Inverted depth
diff --git a/Engine/Source/Runtime/Renderer/Private/BasePassRendering.cpp b/Engine/Source/Runtime/Renderer/Private/BasePassRendering.cpp
--- a/Engine/Source/Runtime/Renderer/Private/BasePassRendering.cpp
+++ b/Engine/Source/Runtime/Renderer/Private/BasePassRendering.cpp
@@ -516,7 +516,7 @@ void CreateOpaqueBasePassUniformBuffer(
 	FRHICommandListImmediate& RHICmdList, 
 	const FViewInfo& View,
 	IPooledRenderTarget* ForwardScreenSpaceShadowMask,
-	FVector4* SceneWithoutSingleLayerWaterValidUVRect,
+	FVector2D* SceneWithoutSingleLayerWaterMaxUV,
 	IPooledRenderTarget* SceneColorWithoutSingleLayerWater,
 	IPooledRenderTarget* SceneDepthWithoutSingleLayerWater,
 	TUniformBufferRef<FOpaqueBasePassUniformParameters>& BasePassUniformBuffer)
@@ -583,7 +583,7 @@ void CreateOpaqueBasePassUniformBuffer(
 	}
 
 	// Single Layer Water
-	BasePassParameters.SceneWithoutSingleLayerWaterValidUVRect = SceneWithoutSingleLayerWaterValidUVRect ? *SceneWithoutSingleLayerWaterValidUVRect : FVector4(0.0f, 0.0f, 1.0f, 1.0f);
+	BasePassParameters.SceneWithoutSingleLayerWaterMaxUV = SceneWithoutSingleLayerWaterMaxUV ? *SceneWithoutSingleLayerWaterMaxUV : FVector2D(1.0f, 1.0f);
 	BasePassParameters.SceneColorWithoutSingleLayerWaterTexture = (SceneColorWithoutSingleLayerWater ? SceneColorWithoutSingleLayerWater->GetRenderTargetItem().TargetableTexture : GSystemTextures.BlackDummy->GetRenderTargetItem().TargetableTexture);
 	BasePassParameters.SceneColorWithoutSingleLayerWaterSampler = TStaticSamplerState<SF_Bilinear>::GetRHI();
 	BasePassParameters.SceneDepthWithoutSingleLayerWaterTexture = (SceneDepthWithoutSingleLayerWater ? SceneDepthWithoutSingleLayerWater->GetRenderTargetItem().TargetableTexture : GSystemTextures.BlackDummy->GetRenderTargetItem().TargetableTexture);
diff --git a/Engine/Source/Runtime/Renderer/Private/BasePassRendering.h b/Engine/Source/Runtime/Renderer/Private/BasePassRendering.h
--- a/Engine/Source/Runtime/Renderer/Private/BasePassRendering.h
+++ b/Engine/Source/Runtime/Renderer/Private/BasePassRendering.h
@@ -78,7 +78,7 @@ BEGIN_GLOBAL_SHADER_PARAMETER_STRUCT(FOpaqueBasePassUniformParameters,)
 	SHADER_PARAMETER_SAMPLER(SamplerState, SceneDepthWithoutSingleLayerWaterSampler)
 	SHADER_PARAMETER_TEXTURE(Texture2D, PreIntegratedGFTexture)
 	SHADER_PARAMETER_SAMPLER(SamplerState, PreIntegratedGFSampler)
-	SHADER_PARAMETER(FVector4, SceneWithoutSingleLayerWaterValidUVRect)
+	SHADER_PARAMETER(FVector2D, SceneWithoutSingleLayerWaterMaxUV)
 	SHADER_PARAMETER(FVector4, DistortionParams)
 	// Misc
 	SHADER_PARAMETER_TEXTURE(Texture2D, EyeAdaptation)
@@ -118,7 +118,7 @@ extern void CreateOpaqueBasePassUniformBuffer(
 	FRHICommandListImmediate& RHICmdList, 
 	const FViewInfo& View,
 	IPooledRenderTarget* ForwardScreenSpaceShadowMask,
-	FVector4* SceneWithoutSingleLayerWaterValidUVRect,
+	FVector2D* SceneWithoutSingleLayerWaterMaxUV,
 	IPooledRenderTarget* SceneColorWithoutSingleLayerWater,
 	IPooledRenderTarget* SceneDepthWithoutSingleLayerWater,
 	TUniformBufferRef<FOpaqueBasePassUniformParameters>& BasePassUniformBuffer);
diff --git a/Engine/Source/Runtime/Renderer/Private/SingleLayerWaterRendering.cpp b/Engine/Source/Runtime/Renderer/Private/SingleLayerWaterRendering.cpp
--- a/Engine/Source/Runtime/Renderer/Private/SingleLayerWaterRendering.cpp
+++ b/Engine/Source/Runtime/Renderer/Private/SingleLayerWaterRendering.cpp
@@ -239,6 +239,7 @@ BEGIN_SHADER_PARAMETER_STRUCT(FSingleLayerWaterCommonShaderParameters, )
 	SHADER_PARAMETER_SAMPLER(SamplerState, PreIntegratedGFSampler)
 	SHADER_PARAMETER_RDG_TEXTURE(Texture2D, SceneNoWaterDepthTexture)
 	SHADER_PARAMETER_SAMPLER(SamplerState, SceneNoWaterDepthSampler)
+	SHADER_PARAMETER(FVector2D, SceneNoWaterMaxUV)
 	SHADER_PARAMETER_STRUCT_INCLUDE(FSceneTextureParameters, SceneTextures)	// Water scene texture
 	SHADER_PARAMETER_STRUCT_INCLUDE(FSceneTextureSamplerParameters, SceneTextureSamplers)
 	SHADER_PARAMETER_STRUCT_REF(FViewUniformShaderParameters, ViewUniformBuffer)
@@ -462,10 +463,8 @@ void FDeferredShadingSceneRenderer::CopySingleLayerWaterTextures(FRHICommandList
 
 	const FIntRect RefractionViewRect = FIntRect(FIntPoint(0, 0), FIntPoint::DivideAndRoundDown(View.ViewRect.Size(), RefractionDownsampleFactor));
 
-	PassData.SceneWithoutSingleLayerWaterValidUVRect.X = (RefractionViewRect.Min.X + 0.5f) / RefractionResolution.Y;
-	PassData.SceneWithoutSingleLayerWaterValidUVRect.Y = (RefractionViewRect.Min.Y + 0.5f) / RefractionResolution.Y;
-	PassData.SceneWithoutSingleLayerWaterValidUVRect.Z = (RefractionViewRect.Max.X - 0.5f) / RefractionResolution.Y;
-	PassData.SceneWithoutSingleLayerWaterValidUVRect.W = (RefractionViewRect.Max.Y - 0.5f) / RefractionResolution.Y;
+	PassData.SceneWithoutSingleLayerWaterMaxUV.X = (RefractionViewRect.Max.X - 0.5f) / RefractionResolution.X;
+	PassData.SceneWithoutSingleLayerWaterMaxUV.Y = (RefractionViewRect.Max.Y - 0.5f) / RefractionResolution.Y;
 
 	FPixelShaderUtils::AddFullscreenPass(
 		GraphBuilder,
@@ -570,6 +569,7 @@ void FDeferredShadingSceneRenderer::RenderSingleLayerWaterReflections(FRHIComman
 			Parameters.PreIntegratedGFSampler = TStaticSamplerState<SF_Bilinear, AM_Clamp, AM_Clamp, AM_Clamp>::GetRHI();
 			Parameters.SceneNoWaterDepthTexture = GraphBuilder.RegisterExternalTexture(PassData.SceneDepthWithoutSingleLayerWater ? PassData.SceneDepthWithoutSingleLayerWater : GSystemTextures.BlackDummy);
 			Parameters.SceneNoWaterDepthSampler = TStaticSamplerState<SF_Point>::GetRHI();
+			Parameters.SceneNoWaterMaxUV = PassData.SceneWithoutSingleLayerWaterMaxUV;
 			Parameters.SceneTextures = SceneTextures;
 			SetupSceneTextureSamplers(&Parameters.SceneTextureSamplers);
 			Parameters.ViewUniformBuffer = View.ViewUniformBuffer;
@@ -762,7 +762,7 @@ bool FDeferredShadingSceneRenderer::RenderSingleLayerWaterPass(FRHICommandListIm
 			CreateOpaqueBasePassUniformBuffer(RHICmdList, 
 				View, 
 				WhiteForwardScreenSpaceShadowMask, 
-				&PassData.SceneWithoutSingleLayerWaterValidUVRect,
+				&PassData.SceneWithoutSingleLayerWaterMaxUV,
 				PassData.SceneColorWithoutSingleLayerWater,
 				PassData.SceneDepthWithoutSingleLayerWater,
 				WaterPassUniformBuffer);
diff --git a/Engine/Source/Runtime/Renderer/Private/SingleLayerWaterRendering.h b/Engine/Source/Runtime/Renderer/Private/SingleLayerWaterRendering.h
--- a/Engine/Source/Runtime/Renderer/Private/SingleLayerWaterRendering.h
+++ b/Engine/Source/Runtime/Renderer/Private/SingleLayerWaterRendering.h
@@ -17,7 +17,7 @@ struct FSingleLayerWaterPassData
 {
 	TRefCountPtr<IPooledRenderTarget> SceneColorWithoutSingleLayerWater;
 	TRefCountPtr<IPooledRenderTarget> SceneDepthWithoutSingleLayerWater;
-	FVector4 SceneWithoutSingleLayerWaterValidUVRect;
+	FVector2D SceneWithoutSingleLayerWaterMaxUV;
 };
 
 class FSingleLayerWaterPassMeshProcessor : public FMeshPassProcessor
