Fixes DOF artifact when TAAU is enabled and SP > 100

diff --git a/Engine/Source/Runtime/Renderer/Private/PostProcess/DiaphragmDOF.cpp b/Engine/Source/Runtime/Renderer/Private/PostProcess/DiaphragmDOF.cpp
--- a/Engine/Source/Runtime/Renderer/Private/PostProcess/DiaphragmDOF.cpp
+++ b/Engine/Source/Runtime/Renderer/Private/PostProcess/DiaphragmDOF.cpp
@@ -1498,7 +1498,7 @@ FRDGTextureRef DiaphragmDOF::AddPasses(
 		else if (bOutputHalfResolution)
 		{
 			PermutationVector.Set<FDiaphragmDOFSetupCS::FOutputResDivisor>(2);
-			PassViewSize = PreprocessViewSize;
+			PassViewSize = GatheringViewSize;
 			CocRadiusBasis = PreprocessViewSize.X;
 		}
 		else
