Fixed STAT_StaticMeshTriangles to report static mesh batches.

#6261

diff --git a/Engine/Source/Runtime/Renderer/Private/SceneVisibility.cpp b/Engine/Source/Runtime/Renderer/Private/SceneVisibility.cpp
--- a/Engine/Source/Runtime/Renderer/Private/SceneVisibility.cpp
+++ b/Engine/Source/Runtime/Renderer/Private/SceneVisibility.cpp
@@ -13,6 +13,7 @@
 #include "Async/TaskGraphInterfaces.h"
 #include "EngineDefines.h"
 #include "EngineGlobals.h"
+#include "EngineStats.h"
 #include "RHIDefinitions.h"
 #include "SceneTypes.h"
 #include "SceneInterface.h"
@@ -2294,6 +2295,8 @@ struct FRelevancePacket
 								}
 
 								++NumVisibleStaticMeshElements;
+
+								INC_DWORD_STAT_BY(STAT_StaticMeshTriangles, StaticMesh.GetNumPrimitives());
 							}
 
 							bNeedsBatchVisibility = true;
