Fix for Windows 7 crash

diff --git a/Engine/Source/Runtime/Renderer/Private/VT/VirtualTexturePhysicalSpace.cpp b/Engine/Source/Runtime/Renderer/Private/VT/VirtualTexturePhysicalSpace.cpp
--- a/Engine/Source/Runtime/Renderer/Private/VT/VirtualTexturePhysicalSpace.cpp
+++ b/Engine/Source/Runtime/Renderer/Private/VT/VirtualTexturePhysicalSpace.cpp
@@ -92,8 +92,7 @@ void FVirtualTexturePhysicalSpace::InitRHI()
 	{
 		const EPixelFormat FormatSRV = Description.Format[Layer];
 		const EPixelFormat FormatUAV = GetUnorderedAccessViewFormat(FormatSRV);
-		const bool bCreateUAV = FormatUAV != PF_Unknown;
-		const bool bCreateAliasedUAV = bCreateUAV && (FormatUAV != FormatSRV);
+		const bool bCreateAliasedUAV = (FormatUAV != PF_Unknown) && (FormatUAV != FormatSRV);
 
 		// Allocate physical texture from the render target pool
 		const uint32 TextureSize = GetTextureSize();
@@ -102,7 +101,7 @@ void FVirtualTexturePhysicalSpace::InitRHI()
 			FormatSRV,
 			FClearValueBinding::None,
 			TexCreate_None,
-			bCreateUAV ? TexCreate_UAV : TexCreate_None,
+			bCreateAliasedUAV ? TexCreate_UAV : TexCreate_None,
 			false);
 
 		GRenderTargetPool.FindFreeElement(RHICmdList, Desc, PooledRenderTarget[Layer], TEXT("PhysicalTexture"));
@@ -116,20 +115,12 @@ void FVirtualTexturePhysicalSpace::InitRHI()
 		SRVCreateInfo.SRGBOverride = SRGBO_ForceEnable;
 		TextureSRV_SRGB[Layer] = RHICreateShaderResourceView(TextureRHI, SRVCreateInfo);
 
-		if (bCreateUAV)
-		{
 #if VIRTUALTEXTURE_UAV_ALIASING
-			if (bCreateAliasedUAV)
-			{
-				// Specific API for format aliasing. Maybe RHI will unify this in future...
-				TextureUAV[Layer] = RHICreateUnorderedAccessView(TextureRHI, 0, FormatUAV);
-			}
-			else
-#endif
-			{
-				TextureUAV[Layer] = RHICreateUnorderedAccessView(TextureRHI);
-			}
+		if (bCreateAliasedUAV)
+		{
+			TextureUAV[Layer] = RHICreateUnorderedAccessView(TextureRHI, 0, FormatUAV);
 		}
+#endif
 	}
 }
 
