Add LLM tags for VirtualTextureSystem
From shelf taken from ChrisB

diff --git a/Engine/Source/Runtime/Core/Private/HAL/LowLevelMemTracker.cpp b/Engine/Source/Runtime/Core/Private/HAL/LowLevelMemTracker.cpp
--- a/Engine/Source/Runtime/Core/Private/HAL/LowLevelMemTracker.cpp
+++ b/Engine/Source/Runtime/Core/Private/HAL/LowLevelMemTracker.cpp
@@ -70,6 +70,7 @@ DECLARE_LLM_MEMORY_STAT(TEXT("Shaders"), STAT_ShadersLLM, STATGROUP_LLMFULL);
 DECLARE_LLM_MEMORY_STAT(TEXT("PSO"), STAT_PSOLLM, STATGROUP_LLMFULL);
 DECLARE_LLM_MEMORY_STAT(TEXT("Textures"), STAT_TexturesLLM, STATGROUP_LLMFULL);
 DECLARE_LLM_MEMORY_STAT(TEXT("TextureMetaData"), STAT_TextureMetaDataLLM, STATGROUP_LLMFULL);
+DECLARE_LLM_MEMORY_STAT(TEXT("VirtualTextureSystem"), STAT_VirtualTextureSystemLLM, STATGROUP_LLMFULL);
 DECLARE_LLM_MEMORY_STAT(TEXT("Render Targets"), STAT_RenderTargetsLLM, STATGROUP_LLMFULL);
 DECLARE_LLM_MEMORY_STAT(TEXT("SceneRender"), STAT_SceneRenderLLM, STATGROUP_LLMFULL);
 DECLARE_LLM_MEMORY_STAT(TEXT("RHIMisc"), STAT_RHIMiscLLM, STATGROUP_LLMFULL);
diff --git a/Engine/Source/Runtime/Core/Public/HAL/LowLevelMemTracker.h b/Engine/Source/Runtime/Core/Public/HAL/LowLevelMemTracker.h
--- a/Engine/Source/Runtime/Core/Public/HAL/LowLevelMemTracker.h
+++ b/Engine/Source/Runtime/Core/Public/HAL/LowLevelMemTracker.h
@@ -171,6 +171,7 @@ enum class ELLMTagSet : uint8
 	macro(PSO,									"PSO",							GET_STATFNAME(STAT_PSOLLM),									GET_STATFNAME(STAT_EngineSummaryLLM),			-1)\
 	macro(Textures,								"Textures",						GET_STATFNAME(STAT_TexturesLLM),							GET_STATFNAME(STAT_TexturesSummaryLLM),			-1)\
 	macro(TextureMetaData,						"TextureMetaData",				GET_STATFNAME(STAT_TextureMetaDataLLM),						GET_STATFNAME(STAT_TexturesSummaryLLM),			-1)\
+	macro(VirtualTextureSystem,					"VirtualTextureSystem",			GET_STATFNAME(STAT_VirtualTextureSystemLLM),				GET_STATFNAME(STAT_TexturesSummaryLLM),			-1)\
 	macro(RenderTargets,						"RenderTargets",				GET_STATFNAME(STAT_RenderTargetsLLM),						GET_STATFNAME(STAT_EngineSummaryLLM),			-1)\
 	macro(SceneRender,							"SceneRender",					GET_STATFNAME(STAT_SceneRenderLLM),							GET_STATFNAME(STAT_EngineSummaryLLM),			-1)\
 	macro(RHIMisc,								"RHIMisc",						GET_STATFNAME(STAT_RHIMiscLLM),								GET_STATFNAME(STAT_EngineSummaryLLM),			-1)\
diff --git a/Engine/Source/Runtime/Renderer/Private/VT/VirtualTextureSystem.cpp b/Engine/Source/Runtime/Renderer/Private/VT/VirtualTextureSystem.cpp
--- a/Engine/Source/Runtime/Renderer/Private/VT/VirtualTextureSystem.cpp
+++ b/Engine/Source/Runtime/Renderer/Private/VT/VirtualTextureSystem.cpp
@@ -562,6 +562,8 @@ FVirtualTextureProducer* FVirtualTextureSystem::FindProducer(const FVirtualTextu
 
 FVirtualTextureSpace* FVirtualTextureSystem::AcquireSpace(const FVTSpaceDescription& InDesc, uint32 InSizeNeeded)
 {
+	LLM_SCOPE(ELLMTag::VirtualTextureSystem);
+
 	// If InDesc requests a private space, don't reuse any existing spaces
 	if (!InDesc.bPrivateSpace)
 	{
@@ -610,6 +612,8 @@ void FVirtualTextureSystem::ReleaseSpace(FVirtualTextureSpace* Space)
 
 FVirtualTexturePhysicalSpace* FVirtualTextureSystem::AcquirePhysicalSpace(const FVTPhysicalSpaceDescription& InDesc)
 {
+	LLM_SCOPE(ELLMTag::VirtualTextureSystem);
+
 	for (int i = 0; i < PhysicalSpaces.Num(); ++i)
 	{
 		FVirtualTexturePhysicalSpace* PhysicalSpace = PhysicalSpaces[i];
@@ -1646,6 +1650,8 @@ void FVirtualTextureSystem::GatherRequestsTask(const FGatherRequestsParameters&
 
 void FVirtualTextureSystem::SubmitRequestsFromLocalTileList(const TSet<FVirtualTextureLocalTile>& LocalTileList, EVTProducePageFlags Flags, FRHICommandListImmediate& RHICmdList, ERHIFeatureLevel::Type FeatureLevel)
 {
+	LLM_SCOPE(ELLMTag::VirtualTextureSystem);
+
 	for (const FVirtualTextureLocalTile& Tile : LocalTileList)
 	{
 		const FVirtualTextureProducerHandle ProducerHandle = Tile.GetProducerHandle();
@@ -1724,6 +1730,8 @@ void FVirtualTextureSystem::SubmitPreMappedRequests(FRHICommandListImmediate& RH
 
 void FVirtualTextureSystem::SubmitRequests(FRHICommandListImmediate& RHICmdList, ERHIFeatureLevel::Type FeatureLevel, FMemStack& MemStack, FUniqueRequestList* RequestList, bool bAsync)
 {
+	LLM_SCOPE(ELLMTag::VirtualTextureSystem);
+
 	FMemMark Mark(MemStack);
 
 	// Allocate space to hold the physical address we allocate for each page load (1 page per layer per request)
@@ -2040,6 +2048,8 @@ void FVirtualTextureSystem::SubmitRequests(FRHICommandListImmediate& RHICmdList,
 
 void FVirtualTextureSystem::AllocateResources(FRHICommandListImmediate& RHICmdList, ERHIFeatureLevel::Type FeatureLevel)
 {
+	LLM_SCOPE(ELLMTag::VirtualTextureSystem);
+
 	for (uint32 ID = 0; ID < MaxSpaces; ID++)
 	{
 		if (Spaces[ID])
