defer deletion of HitproxyArray

diff --git a/Engine/Source/Runtime/Renderer/Private/RendererScene.cpp b/Engine/Source/Runtime/Renderer/Private/RendererScene.cpp
--- a/Engine/Source/Runtime/Renderer/Private/RendererScene.cpp
+++ b/Engine/Source/Runtime/Renderer/Private/RendererScene.cpp
@@ -3919,10 +3919,15 @@ void FScene::UpdateAllPrimitiveSceneInfos(FRHICommandListImmediate& RHICmdList)
 		{
 			// free the primitive scene proxy.
 			delete PrimitiveSceneInfo->Proxy;
-			delete PrimitiveSceneInfo;
 			// Delete the PrimitiveSceneInfo on the game thread after the rendering thread has processed its removal.
 			// This must be done on the game thread because the hit proxy references (and possibly other members) need to be freed on the game thread.
-			//BeginCleanup(PrimitiveSceneInfo);
+			struct DeferDeleteHitProxies : FDeferredCleanupInterface
+			{
+				DeferDeleteHitProxies(TArray<TRefCountPtr<HHitProxy>>&& InHitProxies) : HitProxies(MoveTemp(InHitProxies)) {}
+				TArray<TRefCountPtr<HHitProxy>> HitProxies;
+			};
+			BeginCleanup(new DeferDeleteHitProxies(MoveTemp(PrimitiveSceneInfo->HitProxies)));
+			delete PrimitiveSceneInfo;
 		}
 	}
 
