Add stat for free virtual texture pages

diff --git a/Engine/Source/Runtime/Renderer/Private/VT/VirtualTextureSystem.cpp b/Engine/Source/Runtime/Renderer/Private/VT/VirtualTextureSystem.cpp
--- a/Engine/Source/Runtime/Renderer/Private/VT/VirtualTextureSystem.cpp
+++ b/Engine/Source/Runtime/Renderer/Private/VT/VirtualTextureSystem.cpp
@@ -2,20 +2,20 @@
 
 #include "VirtualTextureSystem.h"
 
-#include "TexturePagePool.h"
-#include "VirtualTextureSpace.h"
-#include "VirtualTexturePhysicalSpace.h"
 #include "AllocatedVirtualTexture.h"
-#include "VirtualTextureFeedback.h"
-#include "VirtualTexturing.h"
-#include "UniquePageList.h"
-#include "UniqueRequestList.h"
-#include "Stats/Stats.h"
-#include "ScenePrivate.h"
-#include "SceneUtils.h"
 #include "HAL/IConsoleManager.h"
 #include "PostProcess/SceneRenderTargets.h"
-
+#include "ProfilingDebugging/CsvProfiler.h"
+#include "ScenePrivate.h"
+#include "SceneUtils.h"
+#include "Stats/Stats.h"
+#include "TexturePagePool.h"
+#include "UniquePageList.h"
+#include "UniqueRequestList.h"
+#include "VirtualTextureFeedback.h"
+#include "VirtualTexturePhysicalSpace.h"
+#include "VirtualTextureSpace.h"
+#include "VirtualTexturing.h"
 
 DECLARE_CYCLE_STAT(TEXT("Feedback Analysis"), STAT_FeedbackAnalysis, STATGROUP_VirtualTexturing);
 DECLARE_CYCLE_STAT(TEXT("VirtualTextureSystem Update"), STAT_VirtualTextureSystem_Update, STATGROUP_VirtualTexturing);
@@ -37,6 +37,7 @@ DECLARE_DWORD_COUNTER_STAT(TEXT("Num page visible"), STAT_NumPageVisible, STATGR
 DECLARE_DWORD_COUNTER_STAT(TEXT("Num page visible resident"), STAT_NumPageVisibleResident, STATGROUP_VirtualTexturing);
 DECLARE_DWORD_COUNTER_STAT(TEXT("Num page visible not resident"), STAT_NumPageVisibleNotResident, STATGROUP_VirtualTexturing);
 DECLARE_DWORD_COUNTER_STAT(TEXT("Num page prefetch"), STAT_NumPagePrefetch, STATGROUP_VirtualTexturing);
+DECLARE_DWORD_COUNTER_STAT(TEXT("Num page free"), STAT_NumPageFree, STATGROUP_VirtualTexturing);
 DECLARE_DWORD_COUNTER_STAT(TEXT("Num page update"), STAT_NumPageUpdate, STATGROUP_VirtualTexturing);
 DECLARE_DWORD_COUNTER_STAT(TEXT("Num continuous page update"), STAT_NumContinuousPageUpdate, STATGROUP_VirtualTexturing);
 
@@ -1604,6 +1605,7 @@ void FVirtualTextureSystem::GatherRequestsTask(const FGatherRequestsParameters&
 		}
 	}
 
+	uint32 TotalAllocatedPages = 0u;
 	for (uint32 PhysicalSpaceID = 0u; PhysicalSpaceID < (uint32)PhysicalSpaces.Num(); ++PhysicalSpaceID)
 	{
 		if (PhysicalSpaces[PhysicalSpaceID] == nullptr)
@@ -1612,6 +1614,8 @@ void FVirtualTextureSystem::GatherRequestsTask(const FGatherRequestsParameters&
 		}
 
 		FVirtualTexturePhysicalSpace* RESTRICT PhysicalSpace = GetPhysicalSpace(PhysicalSpaceID);
+		TotalAllocatedPages += PhysicalSpace->GetNumTiles();
+
 		FPageUpdateBuffer& RESTRICT Buffer = PageUpdateBuffers[PhysicalSpaceID];
 
 		if (Buffer.WorkingSetSize > 0u)
@@ -1646,6 +1650,10 @@ void FVirtualTextureSystem::GatherRequestsTask(const FGatherRequestsParameters&
 	INC_DWORD_STAT_BY(STAT_NumPageVisibleResident, NumResidentPages);
 	INC_DWORD_STAT_BY(STAT_NumPageVisibleNotResident, NumNonResidentPages);
 	INC_DWORD_STAT_BY(STAT_NumPagePrefetch, NumPrefetchPages);
+	INC_DWORD_STAT_BY(STAT_NumPageFree, TotalAllocatedPages - NumResidentPages);
+
+	const float PhysicalPoolUsage = TotalAllocatedPages > 0 ? (float)NumResidentPages / (float)TotalAllocatedPages : 0.f;
+	CSV_CUSTOM_STAT_GLOBAL(VirtualTexturePageUsage, PhysicalPoolUsage, ECsvCustomStatOp::Set);
 }
 
 void FVirtualTextureSystem::SubmitRequestsFromLocalTileList(const TSet<FVirtualTextureLocalTile>& LocalTileList, EVTProducePageFlags Flags, FRHICommandListImmediate& RHICmdList, ERHIFeatureLevel::Type FeatureLevel)
