Memory leak fix

diff --git a/Engine/Source/Runtime/Renderer/Private/VT/VirtualTextureSystem.cpp b/Engine/Source/Runtime/Renderer/Private/VT/VirtualTextureSystem.cpp
--- a/Engine/Source/Runtime/Renderer/Private/VT/VirtualTextureSystem.cpp
+++ b/Engine/Source/Runtime/Renderer/Private/VT/VirtualTextureSystem.cpp
@@ -607,7 +607,7 @@ void FVirtualTextureSystem::ReleaseSpace(FVirtualTextureSpace* Space)
 		// This can only happen on render thread, so we can call ReleaseResource() directly and then delete the pointer immediately
 		DEC_MEMORY_STAT_BY(STAT_TotalPagetableMemory, Space->GetSizeInBytes());
 		Space->ReleaseResource();
-		Spaces[Space->GetID()].Release();
+		Spaces[Space->GetID()].Reset();
 	}
 }
 
