Copying //UE4/Dev-Physics to Dev-Main (//UE4/Dev-Main) [at] 9420715

diff --git a/Engine/Source/Runtime/Engine/Public/GPUSkinCache.h b/Engine/Source/Runtime/Engine/Public/GPUSkinCache.h
--- a/Engine/Source/Runtime/Engine/Public/GPUSkinCache.h
+++ b/Engine/Source/Runtime/Engine/Public/GPUSkinCache.h
@@ -103,6 +103,8 @@ public:
 	ENGINE_API FGPUSkinCache(bool bInRequiresMemoryLimit);
 	ENGINE_API ~FGPUSkinCache();
 
+	struct FCachedGeometrySection GetCachedGeometry(FGPUSkinCacheEntry* InOutEntry, uint32 SectionId);
+
 	void ProcessEntry(FRHICommandListImmediate& RHICmdList, FGPUBaseSkinVertexFactory* VertexFactory,
 		FGPUSkinPassthroughVertexFactory* TargetVertexFactory, const FSkelMeshRenderSection& BatchElement, FSkeletalMeshObjectGPUSkin* Skin,
 		const FMorphVertexBuffer* MorphVertexBuffer, const FSkeletalMeshVertexClothBuffer* ClothVertexBuffer, const FClothSimulData* SimData,
diff --git a/Engine/Source/Runtime/Engine/Public/SkeletalRenderPublic.h b/Engine/Source/Runtime/Engine/Public/SkeletalRenderPublic.h
--- a/Engine/Source/Runtime/Engine/Public/SkeletalRenderPublic.h
+++ b/Engine/Source/Runtime/Engine/Public/SkeletalRenderPublic.h
@@ -43,6 +43,31 @@ enum class EPreviousBoneTransformUpdateMode
 	DuplicateCurrentToPrevious,
 };
 
+struct FCachedGeometrySection
+{
+	FRHIShaderResourceView* PositionBuffer = nullptr;
+	FRHIShaderResourceView* IndexBuffer = nullptr;
+	uint32 NumPrimitives = 0;
+	uint32 VertexBaseIndex = 0;
+	uint32 IndexBaseIndex = 0;
+	uint32 TotalVertexCount = 0;
+	uint32 TotalIndexCount = 0;
+	uint32 SectionIndex = 0;
+	int32 LODIndex = 0;
+};
+
+struct FCachedGeometry
+{
+	int32 LODIndex = 0;
+	TArray<FCachedGeometrySection> Sections;
+
+	void Reset()
+	{
+		LODIndex = 0;
+		Sections.Empty();
+	}
+};
+
 /**
 * Interface for mesh rendering data
 */
@@ -197,6 +222,9 @@ public:
 	/** Get the skeletal mesh resource for which this mesh object was created. */
 	FORCEINLINE FSkeletalMeshRenderData& GetSkeletalMeshRenderData() const { return *SkeletalMeshRenderData; }
 
+	/** Set callback data for register/unregister/update events  */
+	virtual void SetCallbackData(struct FSkeletalMeshObjectCallbackData& CallbackData) {};
+
 #if RHI_RAYTRACING
 	/** Retrieve ray tracing geometry from the underlying mesh object */
 	virtual FRayTracingGeometry* GetRayTracingGeometry() { return nullptr; }
@@ -204,6 +232,8 @@ public:
 	virtual FRWBuffer* GetRayTracingDynamicVertexBuffer() { return nullptr; }
 #endif // RHI_RAYTRACING
 
+	virtual FCachedGeometry GetCachedGeometry() const { return FCachedGeometry(); };
+
 	/** Called to notify clothing data that component transform has changed */
 	virtual void RefreshClothingTransforms(const FMatrix& InNewLocalToWorld, uint32 FrameNumber) {};
 
diff --git a/Engine/Source/Runtime/Renderer/Private/DeferredShadingRenderer.cpp b/Engine/Source/Runtime/Renderer/Private/DeferredShadingRenderer.cpp
--- a/Engine/Source/Runtime/Renderer/Private/DeferredShadingRenderer.cpp
+++ b/Engine/Source/Runtime/Renderer/Private/DeferredShadingRenderer.cpp
@@ -1157,12 +1157,7 @@ void FDeferredShadingSceneRenderer::Render(FRHICommandListImmediate& RHICmdList)
 		RHICmdList.SetCurrentStat(GET_STATID(STAT_CLM_FXPreRender));
 		Scene->FXSystem->PreRender(RHICmdList, &Views[0].GlobalDistanceFieldInfo.ParameterData, Views[0].AllowGPUParticleUpdate());
 	}
-
-	if (IsHairStrandsEnable(Scene->GetShaderPlatform()))
-	{
-		RunHairStrandsInterpolation(RHICmdList);
-	}
-
+	
 	bool bDidAfterTaskWork = false;
 	auto AfterTasksAreStarted = [&bDidAfterTaskWork, bDoInitViewAftersPrepass, this, &RHICmdList, &ILCTaskData, &UpdateViewCustomDataEvents]()
 	{
@@ -1187,9 +1182,15 @@ void FDeferredShadingSceneRenderer::Render(FRHICommandListImmediate& RHICmdList)
 		}
 	};
 
-	if (FGPUSkinCache* GPUSkinCache = Scene->GetGPUSkinCache())
+	RunGPUSkinCacheTransition(RHICmdList, Scene, EGPUSkinCacheTransition::Renderer);
+
+	// Interpolation needs to happen after the skin cache run as there is a dependency 
+	// on the skin cache output.
+	if (IsHairStrandsEnable(Scene->GetShaderPlatform()) && Views.Num() > 0)
 	{
-		GPUSkinCache->TransitionAllToReadable(RHICmdList);
+		const EWorldType::Type WorldType = Views[0].Family->Scene->GetWorld()->WorldType;
+		auto ShaderMap = GetGlobalShaderMap(FeatureLevel);
+		RunHairStrandsInterpolation(RHICmdList, WorldType, ShaderMap, EHairStrandsInterpolationType::RenderStrands);
 	}
 
 	// Before starting the render, all async task for the Custom data must be completed
diff --git a/Engine/Source/Runtime/Renderer/Private/HairStrands/HairStrandsDebug.cpp b/Engine/Source/Runtime/Renderer/Private/HairStrands/HairStrandsDebug.cpp
--- a/Engine/Source/Runtime/Renderer/Private/HairStrands/HairStrandsDebug.cpp
+++ b/Engine/Source/Runtime/Renderer/Private/HairStrands/HairStrandsDebug.cpp
@@ -8,6 +8,8 @@
 #include "HairStrandsVoxelization.h"
 #include "HairStrandsRendering.h"
 #include "HairStrandsVisibility.h"
+#include "HairStrandsInterface.h"
+#include "HairStrandsMeshProjection.h"
 
 #include "Shader.h"
 #include "GlobalShader.h"
@@ -26,16 +28,33 @@ static float GDeepShadowDebugScale = 20;
 static FAutoConsoleVariableRef CVarDeepShadowDebugDomIndex(TEXT("r.HairStrands.DeepShadow.DebugDOMIndex"), GDeepShadowDebugIndex, TEXT("Index of the DOM texture to draw"));
 static FAutoConsoleVariableRef CVarDeepShadowDebugDomScale(TEXT("r.HairStrands.DeepShadow.DebugDOMScale"), GDeepShadowDebugScale, TEXT("Scaling value for the DeepOpacityMap when drawing the deep shadow stats"));
 
-static int32 GDeepShadowStats = 0;
-static FAutoConsoleVariableRef CVarDeepShadowStats(TEXT("r.HairStrands.DeepShadow.DrawStats"), GDeepShadowStats, TEXT("Draw various stats/debug mode about hair rendering"));
+static int32 GHairStrandsDebugMode = 0;
+static FAutoConsoleVariableRef CVarDeepShadowStats(TEXT("r.HairStrands.DebugMode"), GHairStrandsDebugMode, TEXT("Draw various stats/debug mode about hair rendering"));
+
+static int32 GHairStrandsDebugStrandsMode = 0;
+static FAutoConsoleVariableRef CVarDebugPhysicsStrand(TEXT("r.HairStrands.StrandsMode"), GHairStrandsDebugStrandsMode, TEXT("Render debug mode for hair strands. 0:off, 1:simulation strands, 2:render strands with colored simulation strands influence, 3:hair UV, 4:hair root UV, 5: hair seed, 6: dimensions"));
 
-static int32 GHairDebugRenderMode = 0;
-static FAutoConsoleVariableRef CVarDebugPhysicsStrand(TEXT("r.HairStrands.DebugRenderMode"), GHairDebugRenderMode, TEXT("Render debug mode for hair strands. 0:off, 1:simulation strands, 2:render strands with colored simulation strands influence, 3:hair UV, 4:hair root UV, 5: hair seed, 6: dimensions"));
+static int32 GHairDebugMeshProjection_SkinCacheMesh = 0;
+static int32 GHairDebugMeshProjection_HairRestTriangles = 0;
+static int32 GHairDebugMeshProjection_HairRestFrames = 0;
+static int32 GHairDebugMeshProjection_HairDeformedTriangles = 0;
+static int32 GHairDebugMeshProjection_HairDeformedFrames = 0;
+static FAutoConsoleVariableRef CVarHairDebugMeshProjection_SkinCacheMesh(TEXT("r.HairStrands.MeshProjection.DebugSkinCache"),						GHairDebugMeshProjection_SkinCacheMesh, TEXT("Render debug mes projection"));
+static FAutoConsoleVariableRef CVarHairDebugMeshProjection_HairRestTriangles(TEXT("r.HairStrands.MeshProjection.DebugHairRestTriangles"),			GHairDebugMeshProjection_HairRestTriangles, TEXT("Render debug mes projection"));
+static FAutoConsoleVariableRef CVarHairDebugMeshProjection_HairRestFrames(TEXT("r.HairStrands.MeshProjection.DebugHairRestFrames"),					GHairDebugMeshProjection_HairRestFrames, TEXT("Render debug mes projection"));
+static FAutoConsoleVariableRef CVarHairDebugMeshProjection_HairDeformedTriangles(TEXT("r.HairStrands.MeshProjection.DebugHairDeformedTriangles"),	GHairDebugMeshProjection_HairDeformedTriangles, TEXT("Render debug mes projection"));
+static FAutoConsoleVariableRef CVarHairDebugMeshProjection_HairDeformedFrames(TEXT("r.HairStrands.MeshProjection.DebugHairDeformedFrames"),			GHairDebugMeshProjection_HairDeformedFrames, TEXT("Render debug mes projection"));
 
 
+// Helper functions for accessing interpolation data for debug purpose.
+// Definitions is in HairStrandsInterface.cpp
+void GetGroomInterpolationData(const EWorldType::Type WorldType, FHairStrandsProjectionMeshData& OutGeometries);
+void GetGroomInterpolationData(const EWorldType::Type WorldType, FHairStrandsProjectionHairDatas& OutHairData, TArray<int32>& OutLODIndices);
+
 enum class EHairDebugMode : uint8
 {
 	None,
+	ClusterData,
 	LightBounds,
 	DeepOpacityMaps,
 	ClusterScreenRect,
@@ -45,24 +64,27 @@ enum class EHairDebugMode : uint8
 	VoxelsDensity,
 	VoxelsTangent,
 	VoxelsBaseColor,
-	VoxelsRoughness
+	VoxelsRoughness,
+	MeshProjection
 };
 
 static EHairDebugMode GetHairDebugMode()
 {
-	switch (GDeepShadowStats)
+	switch (GHairStrandsDebugMode)
 	{
 	case 0:  return EHairDebugMode::None;
-	case 1:  return EHairDebugMode::LightBounds;
-	case 2:  return EHairDebugMode::ClusterScreenRect;
-	case 3:  return EHairDebugMode::DeepOpacityMaps;
-	case 4:  return EHairDebugMode::SamplePerPixel;
-	case 5:  return EHairDebugMode::TAAResolveType;
-	case 6:  return EHairDebugMode::CoverageType;
-	case 7:  return EHairDebugMode::VoxelsDensity;
-	case 8:  return EHairDebugMode::VoxelsTangent;
-	case 9:  return EHairDebugMode::VoxelsBaseColor;
-	case 10: return EHairDebugMode::VoxelsRoughness;
+	case 1:  return EHairDebugMode::ClusterData;
+	case 2:  return EHairDebugMode::LightBounds;
+	case 3:  return EHairDebugMode::ClusterScreenRect;
+	case 4:  return EHairDebugMode::DeepOpacityMaps;
+	case 5:  return EHairDebugMode::SamplePerPixel;
+	case 6:  return EHairDebugMode::TAAResolveType;
+	case 7:  return EHairDebugMode::CoverageType;
+	case 8:  return EHairDebugMode::VoxelsDensity;
+	case 9:  return EHairDebugMode::VoxelsTangent;
+	case 10:  return EHairDebugMode::VoxelsBaseColor;
+	case 11: return EHairDebugMode::VoxelsRoughness;
+	case 12: return EHairDebugMode::MeshProjection;
 	default: return EHairDebugMode::None;
 	};
 }
@@ -72,6 +94,7 @@ static const TCHAR* ToString(EHairDebugMode DebugMode)
 	switch (DebugMode)
 	{
 	case EHairDebugMode::None: return TEXT("None");
+	case EHairDebugMode::ClusterData: return TEXT("Cluster info");
 	case EHairDebugMode::LightBounds: return TEXT("All DOMs light bounds");
 	case EHairDebugMode::ClusterScreenRect: return TEXT("Screen projected clusters");
 	case EHairDebugMode::DeepOpacityMaps: return TEXT("Deep opacity maps");
@@ -82,14 +105,15 @@ static const TCHAR* ToString(EHairDebugMode DebugMode)
 	case EHairDebugMode::VoxelsTangent: return TEXT("Hair tangent volume");
 	case EHairDebugMode::VoxelsBaseColor: return TEXT("Hair base color volume");
 	case EHairDebugMode::VoxelsRoughness: return TEXT("Hair roughness volume");
+	case EHairDebugMode::MeshProjection: return TEXT("Hair mesh projection");
 	default: return TEXT("None");
 	};
 }
 
 
-EHairStrandsDebugMode GetHairStrandsDebugMode()
+EHairStrandsDebugMode GetHairStrandsDebugStrandsMode()
 {
-	switch (GHairDebugRenderMode)
+	switch (GHairStrandsDebugStrandsMode)
 	{
 	case 0:  return EHairStrandsDebugMode::None;
 	case 1:  return EHairStrandsDebugMode::SimHairStrands;
@@ -140,7 +164,7 @@ class FHairDebugPS : public FGlobalShader
 		END_SHADER_PARAMETER_STRUCT()
 
 public:
-	static bool ShouldCompilePermutation(const FGlobalShaderPermutationParameters& Parameters) { return Parameters.Platform == SP_PCD3D_SM5; }
+	static bool ShouldCompilePermutation(const FGlobalShaderPermutationParameters& Parameters) { return IsHairStrandsSupported(Parameters.Platform); }
 };
 
 IMPLEMENT_GLOBAL_SHADER(FHairDebugPS, "/Engine/Private/HairStrands/HairStrandsDebug.usf", "MainPS", SF_Pixel);
@@ -247,7 +271,7 @@ class FDeepShadowVisualizePS : public FGlobalShader
 		END_SHADER_PARAMETER_STRUCT()
 
 public:
-	static bool ShouldCompilePermutation(const FGlobalShaderPermutationParameters& Parameters) { return Parameters.Platform == SP_PCD3D_SM5; }
+	static bool ShouldCompilePermutation(const FGlobalShaderPermutationParameters& Parameters) { return IsHairStrandsSupported(Parameters.Platform); }
 };
 
 IMPLEMENT_GLOBAL_SHADER(FDeepShadowVisualizePS, "/Engine/Private/HairStrands/HairStrandsDeepShadowDebug.usf", "VisualizeDomPS", SF_Pixel);
@@ -362,7 +386,7 @@ class FVoxelRaymarchingPS : public FGlobalShader
 		END_SHADER_PARAMETER_STRUCT()
 
 public:
-	static bool ShouldCompilePermutation(const FGlobalShaderPermutationParameters& Parameters) { return Parameters.Platform == SP_PCD3D_SM5; }
+	static bool ShouldCompilePermutation(const FGlobalShaderPermutationParameters& Parameters) { return IsHairStrandsSupported(Parameters.Platform); }
 };
 
 IMPLEMENT_GLOBAL_SHADER(FVoxelRaymarchingPS, "/Engine/Private/HairStrands/HairStrandsVoxelRayMarching.usf", "MainPS", SF_Pixel);
@@ -470,9 +494,363 @@ static void AddVoxelRaymarchingPass(
 		});
 	}
 }
+	
+///////////////////////////////////////////////////////////////////////////////////////////////////
+BEGIN_SHADER_PARAMETER_STRUCT(FHairProjectionMeshDebugParameters, )
+	SHADER_PARAMETER(FMatrix, LocalToWorld)
+	SHADER_PARAMETER(uint32, VertexOffset)
+	SHADER_PARAMETER(uint32, IndexOffset)
+	SHADER_PARAMETER(uint32, MaxIndexCount)
+	SHADER_PARAMETER(uint32, MaxVertexCount)
+	SHADER_PARAMETER(FVector2D, OutputResolution)
+	SHADER_PARAMETER_SRV(StructuredBuffer, InputIndexBuffer)
+	SHADER_PARAMETER_SRV(StructuredBuffer, InputVertexBuffer)
+	SHADER_PARAMETER_STRUCT_REF(FViewUniformShaderParameters, ViewUniformBuffer)
+	RENDER_TARGET_BINDING_SLOTS()
+END_SHADER_PARAMETER_STRUCT()
+
+class FHairProjectionMeshDebug : public FGlobalShader
+{
+public:
+	static bool ShouldCompilePermutation(const FGlobalShaderPermutationParameters& Parameters)
+	{
+		return IsHairStrandsSupported(Parameters.Platform);
+	}
+
+
+	FHairProjectionMeshDebug() = default;
+	FHairProjectionMeshDebug(const CompiledShaderInitializerType& Initializer) : FGlobalShader(Initializer) {}
+};
+
+
+class FHairProjectionMeshDebugVS : public FHairProjectionMeshDebug
+{
+public:
+	class FInputType : SHADER_PERMUTATION_INT("PERMUTATION_INPUT_TYPE", 2);
+	using FPermutationDomain = TShaderPermutationDomain<FInputType>;
+
+	DECLARE_GLOBAL_SHADER(FHairProjectionMeshDebugVS);
+	SHADER_USE_PARAMETER_STRUCT(FHairProjectionMeshDebugVS, FHairProjectionMeshDebug);
+
+	BEGIN_SHADER_PARAMETER_STRUCT(FParameters, )		
+		SHADER_PARAMETER_STRUCT_INCLUDE(FHairProjectionMeshDebugParameters, Pass)
+	END_SHADER_PARAMETER_STRUCT()
+};
+
+class FHairProjectionMeshDebugPS : public FHairProjectionMeshDebug
+{
+	DECLARE_GLOBAL_SHADER(FHairProjectionMeshDebugPS);
+	SHADER_USE_PARAMETER_STRUCT(FHairProjectionMeshDebugPS, FHairProjectionMeshDebug);
+
+	BEGIN_SHADER_PARAMETER_STRUCT(FParameters, )
+		SHADER_PARAMETER_STRUCT_INCLUDE(FHairProjectionMeshDebugParameters, Pass)
+	END_SHADER_PARAMETER_STRUCT()
+};
+
+IMPLEMENT_GLOBAL_SHADER(FHairProjectionMeshDebugVS, "/Engine/Private/HairStrands/HairStrandsMeshProjectionMeshDebug.usf", "MainVS", SF_Vertex);
+IMPLEMENT_GLOBAL_SHADER(FHairProjectionMeshDebugPS, "/Engine/Private/HairStrands/HairStrandsMeshProjectionMeshDebug.usf", "MainPS", SF_Pixel);
+
+static void AddDebugProjectionMeshPass(
+	FRDGBuilder& GraphBuilder,
+	const FViewInfo* View,
+	const bool bClearDepth,
+	FHairStrandsProjectionMeshSection& MeshSectionData,
+	FRDGTextureRef& ColorTexture, 
+	FRDGTextureRef& DepthTexture)
+{	
+	const EPrimitiveType PrimitiveType = PT_TriangleList;
+	const bool bHasIndexBuffer = MeshSectionData.IndexBuffer != nullptr;
+	const uint32 PrimitiveCount = MeshSectionData.NumPrimitives;
+
+	if (!MeshSectionData.PositionBuffer || PrimitiveCount==0)
+		return;
+
+	const FIntRect Viewport = View->ViewRect;
+	const FIntPoint Resolution(Viewport.Width(), Viewport.Height());
+
+	FHairProjectionMeshDebugParameters* Parameters = GraphBuilder.AllocParameters<FHairProjectionMeshDebugParameters>();
+	Parameters->LocalToWorld = MeshSectionData.LocalToWorld.ToMatrixWithScale();
+	Parameters->OutputResolution = Resolution;
+	Parameters->VertexOffset = MeshSectionData.VertexBaseIndex;
+	Parameters->IndexOffset = MeshSectionData.IndexBaseIndex;
+	Parameters->MaxIndexCount = MeshSectionData.TotalIndexCount;
+	Parameters->MaxVertexCount = MeshSectionData.TotalVertexCount;
+	Parameters->InputIndexBuffer  = MeshSectionData.IndexBuffer;
+	Parameters->InputVertexBuffer = MeshSectionData.PositionBuffer;
+	Parameters->ViewUniformBuffer = View->ViewUniformBuffer;
+	Parameters->RenderTargets[0] = FRenderTargetBinding(ColorTexture, ERenderTargetLoadAction::ELoad, 0);
+	Parameters->RenderTargets.DepthStencil = FDepthStencilBinding(DepthTexture, bClearDepth ? ERenderTargetLoadAction::EClear : ERenderTargetLoadAction::ELoad, ERenderTargetLoadAction::ENoAction, FExclusiveDepthStencil::DepthWrite_StencilNop);
+
+	FHairProjectionMeshDebugVS::FPermutationDomain PermutationVector;
+	PermutationVector.Set<FHairProjectionMeshDebugVS::FInputType>(bHasIndexBuffer ? 1 : 0);
+
+	TShaderMapRef<FHairProjectionMeshDebugVS> VertexShader(View->ShaderMap, PermutationVector);
+	TShaderMapRef<FHairProjectionMeshDebugPS> PixelShader(View->ShaderMap);
+
+	FHairProjectionMeshDebugVS::FParameters VSParameters;
+	VSParameters.Pass = *Parameters;
+	FHairProjectionMeshDebugPS::FParameters PSParameters;
+	PSParameters.Pass = *Parameters;
+		
+	GraphBuilder.AddPass(
+		RDG_EVENT_NAME("HairStrandsMeshProjectionMeshDebug"),
+		Parameters,
+		ERDGPassFlags::Raster,
+		[Parameters, VSParameters, PSParameters, VertexShader, PixelShader, Viewport, Resolution, View, PrimitiveCount, PrimitiveType](FRHICommandList& RHICmdList)
+	{
+
+		RHICmdList.SetViewport(
+			Viewport.Min.X,
+			Viewport.Min.Y,
+			0.0f,
+			Viewport.Max.X,
+			Viewport.Max.Y,
+			1.0f);
+
+		// Apply additive blending pipeline state.
+		FGraphicsPipelineStateInitializer GraphicsPSOInit;
+		RHICmdList.ApplyCachedRenderTargets(GraphicsPSOInit);
+		GraphicsPSOInit.BlendState = TStaticBlendState<CW_RGBA, BO_Add, BF_One, BF_Zero, BO_Add, BF_One, BF_Zero>::GetRHI();
+		GraphicsPSOInit.RasterizerState = TStaticRasterizerState<>::GetRHI();
+		GraphicsPSOInit.DepthStencilState = TStaticDepthStencilState<true, CF_DepthNearOrEqual>::GetRHI();
+		GraphicsPSOInit.BoundShaderState.VertexDeclarationRHI = GEmptyVertexDeclaration.VertexDeclarationRHI;
+		GraphicsPSOInit.BoundShaderState.VertexShaderRHI = GETSAFERHISHADER_VERTEX(*VertexShader);
+		GraphicsPSOInit.BoundShaderState.PixelShaderRHI = GETSAFERHISHADER_PIXEL(*PixelShader);
+		GraphicsPSOInit.PrimitiveType = PrimitiveType;
+		SetGraphicsPipelineState(RHICmdList, GraphicsPSOInit);
+
+		SetShaderParameters(RHICmdList, *VertexShader, VertexShader->GetVertexShader(), VSParameters);
+		SetShaderParameters(RHICmdList, *PixelShader, PixelShader->GetPixelShader(), PSParameters);
+
+		// Emit an instanced quad draw call on the order of the number of pixels on the screen.	
+		RHICmdList.DrawPrimitive(0, PrimitiveCount, 1);
+	});
+}
+
+///////////////////////////////////////////////////////////////////////////////////////////////////
+BEGIN_SHADER_PARAMETER_STRUCT(FHairProjectionHairDebugParameters, )
+	SHADER_PARAMETER(FVector2D, OutputResolution)
+	SHADER_PARAMETER(uint32, MaxRootCount)
+	SHADER_PARAMETER(uint32, DeformedFrameEnable)
+	SHADER_PARAMETER(FMatrix, RootLocalToWorld)
+
+	SHADER_PARAMETER_SRV(StructuredBuffer, RootTrianglePosition0)
+	SHADER_PARAMETER_SRV(StructuredBuffer, RootTrianglePosition2)
+	SHADER_PARAMETER_SRV(StructuredBuffer, RootTrianglePosition1)
+
+	SHADER_PARAMETER_SRV(StructuredBuffer, RestPosition0Buffer)
+	SHADER_PARAMETER_SRV(StructuredBuffer, RestPosition1Buffer)
+	SHADER_PARAMETER_SRV(StructuredBuffer, RestPosition2Buffer)
+
+	SHADER_PARAMETER_SRV(StructuredBuffer, DeformedPosition0Buffer)
+	SHADER_PARAMETER_SRV(StructuredBuffer, DeformedPosition1Buffer)
+	SHADER_PARAMETER_SRV(StructuredBuffer, DeformedPosition2Buffer)
+
+	// Change for actual frame data (stored or computed only)
+	SHADER_PARAMETER_SRV(StructuredBuffer, RootPositionBuffer)
+	SHADER_PARAMETER_SRV(StructuredBuffer, RootNormalBuffer)
+
+	SHADER_PARAMETER_STRUCT_REF(FViewUniformShaderParameters, ViewUniformBuffer)
+	RENDER_TARGET_BINDING_SLOTS()
+END_SHADER_PARAMETER_STRUCT()
+
+class FHairProjectionHairDebug : public FGlobalShader
+{
+public:
+	static bool ShouldCompilePermutation(const FGlobalShaderPermutationParameters& Parameters)
+	{
+		return IsHairStrandsSupported(Parameters.Platform);
+	}
+
+	FHairProjectionHairDebug() = default;
+	FHairProjectionHairDebug(const CompiledShaderInitializerType& Initializer) : FGlobalShader(Initializer) {}
+};
+
+
+class FHairProjectionHairDebugVS : public FHairProjectionHairDebug
+{
+public:
+	class FInputType : SHADER_PERMUTATION_INT("PERMUTATION_INPUT_TYPE", 2);
+	using FPermutationDomain = TShaderPermutationDomain<FInputType>;
+
+	DECLARE_GLOBAL_SHADER(FHairProjectionHairDebugVS);
+	SHADER_USE_PARAMETER_STRUCT(FHairProjectionHairDebugVS, FHairProjectionHairDebug);
+
+	BEGIN_SHADER_PARAMETER_STRUCT(FParameters, )
+		SHADER_PARAMETER_STRUCT_INCLUDE(FHairProjectionHairDebugParameters, Pass)
+		END_SHADER_PARAMETER_STRUCT()
+};
+
+class FHairProjectionHairDebugPS : public FHairProjectionHairDebug
+{
+	DECLARE_GLOBAL_SHADER(FHairProjectionHairDebugPS);
+	SHADER_USE_PARAMETER_STRUCT(FHairProjectionHairDebugPS, FHairProjectionHairDebug);
+
+	BEGIN_SHADER_PARAMETER_STRUCT(FParameters, )
+		SHADER_PARAMETER_STRUCT_INCLUDE(FHairProjectionHairDebugParameters, Pass)
+		END_SHADER_PARAMETER_STRUCT()
+};
+
+IMPLEMENT_GLOBAL_SHADER(FHairProjectionHairDebugVS, "/Engine/Private/HairStrands/HairStrandsMeshProjectionHairDebug.usf", "MainVS", SF_Vertex);
+IMPLEMENT_GLOBAL_SHADER(FHairProjectionHairDebugPS, "/Engine/Private/HairStrands/HairStrandsMeshProjectionHairDebug.usf", "MainPS", SF_Pixel);
+
+enum class EDebugProjectionHairType
+{
+	HairFrame,
+	HairTriangle,
+};
+
+static void AddDebugProjectionHairPass(
+	FRDGBuilder& GraphBuilder,
+	const FViewInfo* View,
+	const bool bClearDepth,
+	const EDebugProjectionHairType GeometryType,
+	const HairStrandsTriangleType PoseType,
+	const int32 LODIndex,
+	const FHairStrandsProjectionHairData& HairData,
+	FRDGTextureRef& ColorTarget,
+	FRDGTextureRef& DepthTexture)
+{
+	const EPrimitiveType PrimitiveType = GeometryType == EDebugProjectionHairType::HairFrame ? PT_LineList : PT_TriangleList;
+	const uint32 PrimitiveCount = HairData.RootCount;
+
+	if (PrimitiveCount == 0 || LODIndex < 0 || LODIndex >= HairData.LODDatas.Num())
+		return;
+
+	if (EDebugProjectionHairType::HairFrame == GeometryType && (!HairData.RootPositionBuffer || !HairData.RootNormalBuffer))
+		return;
+
+	if (EDebugProjectionHairType::HairTriangle == GeometryType && HairStrandsTriangleType::RestPose == PoseType && (
+		!HairData.LODDatas[LODIndex].RestRootTrianglePosition0Buffer ||
+		!HairData.LODDatas[LODIndex].RestRootTrianglePosition1Buffer ||
+		!HairData.LODDatas[LODIndex].RestRootTrianglePosition2Buffer))
+		return;
+
+	if (EDebugProjectionHairType::HairTriangle == GeometryType && HairStrandsTriangleType::DeformedPose == PoseType && (
+		!HairData.LODDatas[LODIndex].DeformedRootTrianglePosition0Buffer ||
+		!HairData.LODDatas[LODIndex].DeformedRootTrianglePosition1Buffer ||
+		!HairData.LODDatas[LODIndex].DeformedRootTrianglePosition2Buffer))
+		return;
+
+	if (EDebugProjectionHairType::HairFrame == GeometryType && (
+		!HairData.LODDatas[LODIndex].RestRootTrianglePosition0Buffer ||
+		!HairData.LODDatas[LODIndex].RestRootTrianglePosition1Buffer ||
+		!HairData.LODDatas[LODIndex].RestRootTrianglePosition2Buffer ||
+		!HairData.LODDatas[LODIndex].DeformedRootTrianglePosition0Buffer ||
+		!HairData.LODDatas[LODIndex].DeformedRootTrianglePosition1Buffer ||
+		!HairData.LODDatas[LODIndex].DeformedRootTrianglePosition2Buffer))
+		return;
+
+	const FIntRect Viewport = View->ViewRect;
+	const FIntPoint Resolution(Viewport.Width(), Viewport.Height());
+
+	FHairProjectionHairDebugParameters* Parameters = GraphBuilder.AllocParameters<FHairProjectionHairDebugParameters>();
+	Parameters->OutputResolution = Resolution;
+	Parameters->MaxRootCount = HairData.RootCount;
+
+	if (EDebugProjectionHairType::HairFrame == GeometryType)
+	{
+		Parameters->RootLocalToWorld = HairData.LocalToWorld.ToMatrixWithScale();
+		Parameters->RootPositionBuffer = HairData.RootPositionBuffer;
+		Parameters->RootNormalBuffer = HairData.RootNormalBuffer;
+		Parameters->DeformedFrameEnable = PoseType == HairStrandsTriangleType::DeformedPose;
+		//if (PoseType == HairStrandsTriangleType::DeformedPose)
+		{
+			Parameters->RestPosition0Buffer = HairData.LODDatas[LODIndex].RestRootTrianglePosition0Buffer->SRV;
+			Parameters->RestPosition1Buffer = HairData.LODDatas[LODIndex].RestRootTrianglePosition1Buffer->SRV;
+			Parameters->RestPosition2Buffer = HairData.LODDatas[LODIndex].RestRootTrianglePosition2Buffer->SRV;
+
+			Parameters->DeformedPosition0Buffer = HairData.LODDatas[LODIndex].DeformedRootTrianglePosition0Buffer->SRV;
+			Parameters->DeformedPosition1Buffer = HairData.LODDatas[LODIndex].DeformedRootTrianglePosition1Buffer->SRV;
+			Parameters->DeformedPosition2Buffer = HairData.LODDatas[LODIndex].DeformedRootTrianglePosition2Buffer->SRV;
+		}
+	}
+
+	if (EDebugProjectionHairType::HairTriangle == GeometryType)
+	{
+		if (PoseType == HairStrandsTriangleType::RestPose)
+		{
+			Parameters->RootTrianglePosition0 = HairData.LODDatas[LODIndex].RestRootTrianglePosition0Buffer->SRV;
+			Parameters->RootTrianglePosition1 = HairData.LODDatas[LODIndex].RestRootTrianglePosition1Buffer->SRV;
+			Parameters->RootTrianglePosition2 = HairData.LODDatas[LODIndex].RestRootTrianglePosition2Buffer->SRV;
+		}
+
+		if (PoseType == HairStrandsTriangleType::DeformedPose)
+		{
+			Parameters->RootTrianglePosition0 = HairData.LODDatas[LODIndex].DeformedRootTrianglePosition0Buffer->SRV;
+			Parameters->RootTrianglePosition1 = HairData.LODDatas[LODIndex].DeformedRootTrianglePosition1Buffer->SRV;
+			Parameters->RootTrianglePosition2 = HairData.LODDatas[LODIndex].DeformedRootTrianglePosition2Buffer->SRV;
+		}
+	}
+
+	Parameters->ViewUniformBuffer = View->ViewUniformBuffer;
+	Parameters->RenderTargets[0] = FRenderTargetBinding(ColorTarget, ERenderTargetLoadAction::ELoad, 0);
+	Parameters->RenderTargets.DepthStencil = FDepthStencilBinding(DepthTexture, bClearDepth ? ERenderTargetLoadAction::EClear : ERenderTargetLoadAction::ELoad, ERenderTargetLoadAction::ENoAction, FExclusiveDepthStencil::DepthWrite_StencilNop);
+
+	FHairProjectionHairDebugVS::FPermutationDomain PermutationVector;
+	PermutationVector.Set<FHairProjectionHairDebugVS::FInputType>(PrimitiveType == PT_LineList ? 0 : 1);
+
+	TShaderMapRef<FHairProjectionHairDebugVS> VertexShader(View->ShaderMap, PermutationVector);
+	TShaderMapRef<FHairProjectionHairDebugPS> PixelShader(View->ShaderMap);
+
+	FHairProjectionHairDebugVS::FParameters VSParameters;
+	VSParameters.Pass = *Parameters;
+	FHairProjectionHairDebugPS::FParameters PSParameters;
+	PSParameters.Pass = *Parameters;
+
+	GraphBuilder.AddPass(
+		RDG_EVENT_NAME("HairStrandsMeshProjectionHairDebug"),
+		Parameters,
+		ERDGPassFlags::Raster,
+		[Parameters, VSParameters, PSParameters, VertexShader, PixelShader, Viewport, Resolution, View, PrimitiveCount, PrimitiveType](FRHICommandList& RHICmdList)
+	{
+
+		RHICmdList.SetViewport(
+			Viewport.Min.X,
+			Viewport.Min.Y,
+			0.0f,
+			Viewport.Max.X,
+			Viewport.Max.Y,
+			1.0f);
+
+		// Apply additive blending pipeline state.
+		FGraphicsPipelineStateInitializer GraphicsPSOInit;
+		RHICmdList.ApplyCachedRenderTargets(GraphicsPSOInit);
+		GraphicsPSOInit.BlendState = TStaticBlendState<CW_RGBA, BO_Add, BF_One, BF_Zero, BO_Add, BF_One, BF_Zero>::GetRHI();
+		GraphicsPSOInit.RasterizerState = TStaticRasterizerState<>::GetRHI();
+		GraphicsPSOInit.DepthStencilState = TStaticDepthStencilState<true, CF_DepthNearOrEqual>::GetRHI();
+		GraphicsPSOInit.BoundShaderState.VertexDeclarationRHI = GEmptyVertexDeclaration.VertexDeclarationRHI;
+		GraphicsPSOInit.BoundShaderState.VertexShaderRHI = GETSAFERHISHADER_VERTEX(*VertexShader);
+		GraphicsPSOInit.BoundShaderState.PixelShaderRHI = GETSAFERHISHADER_PIXEL(*PixelShader);
+		GraphicsPSOInit.PrimitiveType = PrimitiveType;
+		SetGraphicsPipelineState(RHICmdList, GraphicsPSOInit);
+
+		SetShaderParameters(RHICmdList, *VertexShader, VertexShader->GetVertexShader(), VSParameters);
+		SetShaderParameters(RHICmdList, *PixelShader, PixelShader->GetPixelShader(), PSParameters);
+
+		// Emit an instanced quad draw call on the order of the number of pixels on the screen.	
+		RHICmdList.DrawPrimitive(0, PrimitiveCount, 1);
+	});
+}
 
 ///////////////////////////////////////////////////////////////////////////////////////////////////
 
+const TCHAR* ToString(EWorldType::Type Type)
+{
+	switch (Type)
+	{
+		case EWorldType::None			: return TEXT("None");
+		case EWorldType::Game			: return TEXT("Game");
+		case EWorldType::Editor			: return TEXT("Editor");
+		case EWorldType::PIE			: return TEXT("PIE");
+		case EWorldType::EditorPreview	: return TEXT("EditorPreview");
+		case EWorldType::GamePreview	: return TEXT("GamePreview");
+		case EWorldType::GameRPC		: return TEXT("GameRPC");
+		case EWorldType::Inactive		: return TEXT("Inactive");
+		default							: return TEXT("Unknown");
+	}
+}
+
 void RenderHairStrandsDebugInfo(FRHICommandListImmediate& RHICmdList, TArray<FViewInfo>& Views, const FHairStrandsDatas* HairDatas)
 {
 	const float YStep = 14;
@@ -491,9 +869,78 @@ void RenderHairStrandsDebugInfo(FRHICommandListImmediate& RHICmdList, TArray<FVi
 	FSceneRenderTargets& SceneTargets = FSceneRenderTargets::Get(RHICmdList);
 
 	// Debug mode name only
-	const EHairStrandsDebugMode StrandsDebugMode = GetHairStrandsDebugMode();
+	const EHairStrandsDebugMode StrandsDebugMode = GetHairStrandsDebugStrandsMode();
 	const EHairDebugMode HairDebugMode = GetHairDebugMode();
 
+	if (HairDebugMode == EHairDebugMode::ClusterData)
+	{
+		FViewElementPDI ShadowFrustumPDI(&View, nullptr, nullptr);
+		const FHairStrandsClusterDatas& ClusterDatas = InClusterViews.Views[ViewIndex];
+		for (const FHairStrandsClusterData& ClusterData : ClusterDatas.Datas)
+		{
+			const FBox ClusterBox(ClusterData.GetMinBound(), ClusterData.GetMaxBound());
+			DrawWireBox(&ShadowFrustumPDI, ClusterBox, FColor::Red, 0);
+		}
+
+		FRenderTargetTemp TempRenderTarget(View, (const FTexture2DRHIRef&)SceneTargets.GetSceneColor()->GetRenderTargetItem().TargetableTexture);
+		FCanvas Canvas(&TempRenderTarget, NULL, ViewFamily.CurrentRealTime, ViewFamily.CurrentWorldTime, ViewFamily.DeltaWorldTime, View.FeatureLevel);
+
+		float X = 20;
+		float Y = 38;
+		FLinearColor InactiveColor(0.5, 0.5, 0.5);
+		FLinearColor DebugColor(1, 1, 0);
+		FString Line;
+
+		//const FVector4 HairComponent = GetHairComponents();
+		//const uint32 bHairR = HairComponent.X > 0 ? 1 : 0;
+		//const uint32 bHairTT = HairComponent.Y > 0 ? 1 : 0;
+		//const uint32 bHairTRT = HairComponent.Z > 0 ? 1 : 0;
+		//const uint32 bHairGlobalScattering = FMath::FloorToInt(HairComponent.W / 10.0f) > 0 ? 1 : 0;
+		//const uint32 bHairLocalScattering = FMath::Frac(HairComponent.W / 10.0f)*10.f > 0 ? 1 : 0;
+
+		const FHairStrandsDebugInfos DebugInfos = GetHairStandsDebugInfos();
+
+		Line = FString::Printf(TEXT("----------------------------------------------------------------"));				
+		Canvas.DrawShadowedString(X, Y += YStep, *Line, GetStatsFont(), DebugColor);
+		
+		Line = FString::Printf(TEXT("Registered component count : %d"), DebugInfos.Num());									
+		Canvas.DrawShadowedString(X, Y += YStep, *Line, GetStatsFont(), DebugColor);
+
+		for (const FHairStrandsDebugInfo& DebugInfo : DebugInfos)
+		{
+			check(ViewFamily.Scene && ViewFamily.Scene->GetWorld());
+			const bool bIsActive = DebugInfo.WorldType == ViewFamily.Scene->GetWorld()->WorldType;
+
+			Line = FString::Printf(TEXT(" * Id:%d | WorldType:%s | Group count : %d"), DebugInfo.Id, ToString(DebugInfo.WorldType), DebugInfo.GroupInfos.Num());
+			Canvas.DrawShadowedString(X, Y += YStep, *Line, GetStatsFont(), bIsActive ? DebugColor : InactiveColor);
+
+			for (const FHairStrandsDebugInfo::FGroupInfo& DebugGroupInfo : DebugInfo.GroupInfos)
+			{
+				Line = FString::Printf(TEXT("        |> CurveCount : %d | VertexCount : %d | MaxRadius : %f | MaxLength : %f | Skinned: %s | LOD count : %d"),
+					DebugGroupInfo.CurveCount,
+					DebugGroupInfo.VertexCount,
+					DebugGroupInfo.MaxRadius,
+					DebugGroupInfo.MaxLength,
+					DebugGroupInfo.bHasSkinInterpolation ? TEXT("True") : TEXT("False"),
+					DebugGroupInfo.LODCount);
+				Canvas.DrawShadowedString(X, Y += YStep, *Line, GetStatsFont(), bIsActive ? DebugColor : InactiveColor);
+			}
+		}
+
+		Line = FString::Printf(TEXT("----------------------------------------------------------------"));
+		Canvas.DrawShadowedString(X, Y += YStep, *Line, GetStatsFont(), DebugColor);
+
+		Line = FString::Printf(TEXT("Cluster count : %d"), ClusterDatas.Datas.Num());
+		Canvas.DrawShadowedString(X, Y += YStep, *Line, GetStatsFont(), DebugColor);
+		for (const FHairStrandsClusterData& ClusterData : ClusterDatas.Datas)
+		{
+			Line = FString::Printf(TEXT(" %d - Bound Radus: %f.2m (%dx%d)"), ClusterData.ClusterId, ClusterData.Bounds.GetSphere().W);
+			Canvas.DrawShadowedString(X, Y += YStep, *Line, GetStatsFont(), DebugColor);
+		}
+
+		Canvas.Flush_RenderThread(RHICmdList);
+	}
+
 	if (HairDebugMode == EHairDebugMode::DeepOpacityMaps)
 	{
 		const uint32 DomIndex = GDeepShadowDebugIndex;
@@ -604,6 +1051,76 @@ void RenderHairStrandsDebugInfo(FRHICommandListImmediate& RHICmdList, TArray<FVi
 		}
 	}
 
+	if (HairDebugMode == EHairDebugMode::MeshProjection)
+	{
+		const EWorldType::Type WorldType = View.Family->Scene->GetWorld()->WorldType;
+
+		FRDGBuilder GraphBuilder(RHICmdList);
+		FRDGTextureRef SceneColorTexture = GraphBuilder.RegisterExternalTexture(SceneTargets.GetSceneColor(), TEXT("SceneColorTexture"));
+		if (ViewIndex < uint32(HairDatas->HairVisibilityViews.HairDatas.Num()))
+		{
+			FHairStrandsProjectionMeshData MeshProjectionData;
+			GetGroomInterpolationData(WorldType, MeshProjectionData);
+
+			bool bClearDepth = true;
+			FRDGTextureRef DepthTexture;
+			{
+				FRDGTextureDesc Desc;
+				Desc.Extent = SceneColorTexture->Desc.Extent;
+				Desc.Depth = 0;
+				Desc.Format = PF_DepthStencil;
+				Desc.NumMips = 1;
+				Desc.NumSamples = 1;
+				Desc.Flags = TexCreate_None;
+				Desc.TargetableFlags = TexCreate_DepthStencilTargetable;
+				Desc.ClearValue = FClearValueBinding::DepthFar;
+				Desc.bForceSharedTargetAndShaderResource = true;
+				DepthTexture = GraphBuilder.CreateTexture(Desc, TEXT("HairInterpolationDepthTexture"));
+			}
+
+			if (GHairDebugMeshProjection_SkinCacheMesh > 0)
+			{
+				for (FHairStrandsProjectionMeshSection& Section : MeshProjectionData.Sections)
+				{
+					AddDebugProjectionMeshPass(GraphBuilder, &View, bClearDepth, Section, SceneColorTexture, DepthTexture);
+					bClearDepth = false;
+				}
+			}
+
+			FHairStrandsProjectionHairDatas HairProjectionDatas;
+			TArray<int32> HairLODIndices;
+			GetGroomInterpolationData(WorldType, HairProjectionDatas, HairLODIndices);
+			check(HairProjectionDatas.HairData.Num() == HairLODIndices.Num());
+			for (int32 HairIndex=0; HairIndex < HairProjectionDatas.HairData.Num(); ++HairIndex)
+			{
+				const FHairStrandsProjectionHairData& Data = HairProjectionDatas.HairData[HairIndex];
+				const int32 LODIndex = HairLODIndices[HairIndex];
+
+				if (GHairDebugMeshProjection_HairRestTriangles > 0)
+				{
+					AddDebugProjectionHairPass(GraphBuilder, &View, bClearDepth, EDebugProjectionHairType::HairTriangle, HairStrandsTriangleType::RestPose, LODIndex, Data, SceneColorTexture, DepthTexture);
+					bClearDepth = false;
+				}
+				if (GHairDebugMeshProjection_HairRestFrames > 0)
+				{
+					AddDebugProjectionHairPass(GraphBuilder, &View, bClearDepth, EDebugProjectionHairType::HairFrame, HairStrandsTriangleType::RestPose, LODIndex, Data, SceneColorTexture, DepthTexture);
+					bClearDepth = false;
+				}
+				if (GHairDebugMeshProjection_HairDeformedTriangles > 0)
+				{
+					AddDebugProjectionHairPass(GraphBuilder, &View, bClearDepth, EDebugProjectionHairType::HairTriangle, HairStrandsTriangleType::DeformedPose, LODIndex, Data, SceneColorTexture, DepthTexture);
+					bClearDepth = false;
+				}
+				if (GHairDebugMeshProjection_HairDeformedFrames > 0)
+				{
+					AddDebugProjectionHairPass(GraphBuilder, &View, bClearDepth, EDebugProjectionHairType::HairFrame, HairStrandsTriangleType::DeformedPose, LODIndex, Data, SceneColorTexture, DepthTexture);
+					bClearDepth = false;
+				}
+			}
+		}
+		GraphBuilder.Execute();
+	}
+
 	// Text
 	if (HairDebugMode == EHairDebugMode::LightBounds || HairDebugMode == EHairDebugMode::DeepOpacityMaps)
 	{
diff --git a/Engine/Source/Runtime/Renderer/Private/HairStrands/HairStrandsInterface.cpp b/Engine/Source/Runtime/Renderer/Private/HairStrands/HairStrandsInterface.cpp
--- a/Engine/Source/Runtime/Renderer/Private/HairStrands/HairStrandsInterface.cpp
+++ b/Engine/Source/Runtime/Renderer/Private/HairStrands/HairStrandsInterface.cpp
@@ -5,30 +5,99 @@
 =============================================================================*/
 
 #include "HairStrandsInterface.h"
-#include "HairStrandsUtils.h"
+#include "HairStrandsRendering.h"
+#include "HairStrandsMeshProjection.h"
+
+#include "SkeletalRenderPublic.h"
+#include "SceneRendering.h"
 
 DEFINE_LOG_CATEGORY_STATIC(LogHairRendering, Log, All);
 
 static int32 GHairStrandsRenderingEnable = 1;
 static FAutoConsoleVariableRef CVarHairStrandsRenderingEnable(TEXT("r.HairStrands.Enable"), GHairStrandsRenderingEnable, TEXT("Enable/Disable hair strands rendering"));
 
+////////////////////////////////////////////////////////////////////////////////////////////////////////////////////////
+
+inline FHairStrandsProjectionMeshSection ConvertMeshSection(const FCachedGeometrySection& In, const FTransform& InTransform)
+{
+	FHairStrandsProjectionMeshSection Out;
+	Out.IndexBuffer = In.IndexBuffer;
+	Out.PositionBuffer = In.PositionBuffer;
+	Out.TotalVertexCount = In.TotalVertexCount;
+	Out.TotalIndexCount = In.TotalIndexCount;
+	Out.VertexBaseIndex = In.VertexBaseIndex;
+	Out.IndexBaseIndex = In.IndexBaseIndex;
+	Out.NumPrimitives = In.NumPrimitives;
+	Out.SectionIndex = In.SectionIndex;
+	Out.LODIndex = In.LODIndex;
+	Out.LocalToWorld = InTransform;
+	return Out;
+}
+
+////////////////////////////////////////////////////////////////////////////////////////////////////////////////////////
+// Runtime execution order (on the render thread):
+//  * Register
+//  * For each frame
+//		* Update
+//		* AddProjectionQuery (Opt)
+//		* Project (Opt)
+//		* Update triangles information for dynamic meshes
+//		* RunHairStrandsInterpolation (Interpolation callback)
+//  * UnRegister
+////////////////////////////////////////////////////////////////////////////////////////////////////////////////////////
+
 struct FHairStrandsManager
 {
 	struct Element
 	{
-		uint64 Id;
-		FHairStrandsInterpolation Data;
+		uint64 Id = 0;
+		EWorldType::Type WorldType = EWorldType::None;
+		FHairStrandsDebugInfo::FGroupInfos DebugGroupInfos;
+		FHairStrandsInterpolationData InterpolationData;
+		FHairStrandsProjectionHairData RenProjectionHairData;
+		FHairStrandsProjectionHairData SimProjectionHairData;
+		FCachedGeometry CachedGeometry;
+		FTransform SkeletalLocalToWorld;
+		FVector SkeletalWorldBoundsCenter;
+		const FSkeletalMeshObject* MeshObject = nullptr;
+		int32 FrameLODIndex = -1;
 	};
 	TArray<Element> Elements;
+
+	struct ProjectionQuery
+	{
+		uint64 Id = 0;
+		EWorldType::Type WorldType = EWorldType::None;
+		int32 LODIndex = -1;
+		FVector RestRootCenter = FVector::ZeroVector;
+		bool bProcessed = false;
+	};
+	TArray<ProjectionQuery> ProjectionsQueries;
 };
 
 FHairStrandsManager GHairManager;
 
-void RegisterHairStrands(uint64 Id, const FHairStrandsInterpolation& E)
-{ 
+void AddHairStrandsProjectionQuery(
+	FRHICommandListImmediate& RHICmdList,
+	uint64 Id,
+	EWorldType::Type WorldType,
+	int32 LODIndex,
+	const FVector& RestRootCenter)
+{
+	GHairManager.ProjectionsQueries.Add({ Id, WorldType, LODIndex, RestRootCenter, false });
+}
+
+void RegisterHairStrands(
+	uint64 Id, 
+	EWorldType::Type WorldType,
+	const FHairStrandsInterpolationData& InterpolationData, 
+	const FHairStrandsProjectionHairData& RenProjectionData,
+ 	const FHairStrandsProjectionHairData& SimProjectionData,
+	const FHairStrandsDebugInfo::FGroupInfos& GroupInfos)
+{
 	for (int32 Index = 0; Index < GHairManager.Elements.Num(); ++Index)
 	{
-		if (GHairManager.Elements[Index].Id == Id)
+		if (GHairManager.Elements[Index].Id == Id && GHairManager.Elements[Index].WorldType == WorldType)
 		{
 			// Component already registered. This should not happen.
 			UE_LOG(LogHairRendering, Warning, TEXT("Component already register. This should't happen. Please report this to a rendering engineer."))
@@ -36,34 +105,253 @@ void RegisterHairStrands(uint64 Id, const FHairStrandsInterpolation& E)
 		}
 	}
 
-	GHairManager.Elements.Add({ Id, E });
+	FHairStrandsManager::Element& E =  GHairManager.Elements.AddDefaulted_GetRef();
+	E.Id = Id;
+	E.WorldType = WorldType;
+	E.InterpolationData = InterpolationData;
+	E.RenProjectionHairData = RenProjectionData;
+	E.RenProjectionHairData.LocalToWorld = FTransform::Identity;
+	E.SimProjectionHairData = SimProjectionData;
+	E.SimProjectionHairData.LocalToWorld = FTransform::Identity;
+	E.SkeletalLocalToWorld = FTransform::Identity;
+	E.SkeletalWorldBoundsCenter = FVector::ZeroVector;
+	E.DebugGroupInfos = GroupInfos;
+}
+
+bool UpdateHairStrands(
+	uint64 Id, 
+	EWorldType::Type WorldType, 
+	const FTransform& HairLocalToWorld, 
+	const FTransform& SkeletalLocalToWorld,
+	const FVector& SkeletalWorldBoundsCenter)
+{
+	for (FHairStrandsManager::Element& E : GHairManager.Elements)
+	{
+		if (E.Id != Id || E.WorldType != WorldType)
+			continue;
+
+		E.RenProjectionHairData.LocalToWorld = HairLocalToWorld;
+		E.SimProjectionHairData.LocalToWorld = HairLocalToWorld;
+		E.SkeletalLocalToWorld = SkeletalLocalToWorld;
+		E.SkeletalWorldBoundsCenter = SkeletalWorldBoundsCenter;
+		return true;
+	}
+
+	return false;
+}
+
+bool UpdateHairStrands(
+	uint64 Id, 
+	EWorldType::Type NewWorldType)
+{
+	for (FHairStrandsManager::Element& E : GHairManager.Elements)
+	{
+		if (E.Id != Id)
+			continue;
+
+		E.WorldType = NewWorldType;
+		return true;
+	}
+
+	return false;
+}
+
+bool UpdateHairStrands(
+	uint64 Id,
+	EWorldType::Type WorldType,
+	const FTransform& HairLocalToWorld,
+	const FHairStrandsProjectionHairData& RenProjectionData,
+	const FHairStrandsProjectionHairData& SimProjectionData)
+{
+	for (FHairStrandsManager::Element& E : GHairManager.Elements)
+	{
+		if (E.Id != Id || E.WorldType != WorldType)
+			continue;
+
+		E.RenProjectionHairData = RenProjectionData;
+		E.RenProjectionHairData.LocalToWorld = HairLocalToWorld;
+		E.SimProjectionHairData = SimProjectionData;
+		E.SimProjectionHairData.LocalToWorld = HairLocalToWorld;
+		return true;
+	}
+
+	return false;
+}
+
+bool UpdateHairStrands(
+	uint64 Id,
+	EWorldType::Type WorldType,
+	const FSkeletalMeshObject* MeshObject)
+{
+	for (FHairStrandsManager::Element& E : GHairManager.Elements)
+	{
+		if (E.Id != Id || E.WorldType != WorldType)
+			continue;
+
+		E.MeshObject = MeshObject;
+		return true;
+	}
+
+	return false;
 }
 
-FHairStrandsInterpolation UnregisterHairStrands(uint64 Id)
+void UnregisterHairStrands(uint64 Id, EWorldType::Type WorldType)
 {
 	for (int32 Index=0;Index< GHairManager.Elements.Num();++Index)
 	{
-		if (GHairManager.Elements[Index].Id == Id)
+		if (GHairManager.Elements[Index].Id == Id && GHairManager.Elements[Index].WorldType == WorldType)
 		{
-			FHairStrandsInterpolation Out = GHairManager.Elements[Index].Data;
+			FHairStrandsInterpolationData Out = GHairManager.Elements[Index].InterpolationData;
 			GHairManager.Elements[Index] = GHairManager.Elements[GHairManager.Elements.Num()-1];
 			GHairManager.Elements.SetNum(GHairManager.Elements.Num() - 1);
-			return Out;
+		}
+	}
+}
+
+void RunHairStrandsInterpolation(FRHICommandListImmediate& RHICmdList, EWorldType::Type WorldType, TShaderMap<FGlobalShaderType>* ShaderMap, EHairStrandsInterpolationType Type)
+{
+	check(IsInRenderingThread());
+
+	// Update geometry cached based on GPU Skin output
+	for (FHairStrandsManager::Element& E : GHairManager.Elements)
+	{
+		if (E.WorldType != WorldType)
+			continue;
+
+		E.CachedGeometry = E.MeshObject ? E.MeshObject->GetCachedGeometry() : FCachedGeometry();
+	}
+
+	// Process projection queries
+	for (FHairStrandsManager::ProjectionQuery& Q : GHairManager.ProjectionsQueries)
+	{	
+		for (FHairStrandsManager::Element& E : GHairManager.Elements)
+		{
+			if (E.Id != Q.Id || E.WorldType != Q.WorldType || E.CachedGeometry.Sections.Num() == 0 || Q.bProcessed)
+				continue;
+	
+			FHairStrandsProjectionMeshData MeshData;
+			for (FCachedGeometrySection Section : E. CachedGeometry.Sections)
+			{
+				check(Section.LODIndex == Q.LODIndex);
+				MeshData.Sections.Add(ConvertMeshSection(Section, E.SkeletalLocalToWorld));
+			}
+
+			//if (EHairStrandsInterpolationType::RenderStrands == Type)
+			{
+				E.SimProjectionHairData.LODDatas[Q.LODIndex].RestRootCenter = Q.RestRootCenter;
+				ProjectHairStrandsOntoMesh(RHICmdList, ShaderMap, Q.LODIndex, MeshData, E.SimProjectionHairData);
+				UpdateHairStrandsMeshTriangles(RHICmdList, ShaderMap, Q.LODIndex, HairStrandsTriangleType::RestPose, MeshData, E.SimProjectionHairData);
+			}
+
+			//if (EHairStrandsInterpolationType::RenderStrands == Type)
+			{
+				E.RenProjectionHairData.LODDatas[Q.LODIndex].RestRootCenter = Q.RestRootCenter;
+				ProjectHairStrandsOntoMesh(RHICmdList, ShaderMap, Q.LODIndex, MeshData, E.RenProjectionHairData);
+				UpdateHairStrandsMeshTriangles(RHICmdList, ShaderMap, Q.LODIndex, HairStrandsTriangleType::RestPose, MeshData, E.RenProjectionHairData);
+			}
+			Q.bProcessed = true;
+		}
+	}
+
+	// Update dynamic mesh triangles
+	for (FHairStrandsManager::Element& E : GHairManager.Elements)
+	{
+		E.FrameLODIndex = -1;
+		if (E.WorldType != WorldType || E.CachedGeometry.Sections.Num() == 0)
+			continue;
+	
+		FHairStrandsProjectionMeshData MeshData;
+		for (FCachedGeometrySection Section : E.CachedGeometry.Sections)
+		{
+			// Ensure all mesh's sections have the same LOD index
+			if (E.FrameLODIndex < 0) E.FrameLODIndex = Section.LODIndex;
+			check(E.FrameLODIndex == Section.LODIndex);
+	
+			MeshData.Sections.Add(ConvertMeshSection(Section, E.SkeletalLocalToWorld));
+		}
+
+		if (EHairStrandsInterpolationType::RenderStrands == Type && 0 <= E.FrameLODIndex && E.FrameLODIndex < E.RenProjectionHairData.LODDatas.Num() && E.RenProjectionHairData.LODDatas[E.FrameLODIndex].bIsValid)
+		{
+			E.RenProjectionHairData.LODDatas[E.FrameLODIndex].DeformedRootCenter = E.SkeletalWorldBoundsCenter;
+			UpdateHairStrandsMeshTriangles(RHICmdList, ShaderMap, E.FrameLODIndex, HairStrandsTriangleType::DeformedPose, MeshData, E.RenProjectionHairData);
+		}
+
+		if (EHairStrandsInterpolationType::SimulationStrands == Type && 0 <= E.FrameLODIndex && E.FrameLODIndex < E.SimProjectionHairData.LODDatas.Num() && E.SimProjectionHairData.LODDatas[E.FrameLODIndex].bIsValid)
+		{
+			E.SimProjectionHairData.LODDatas[E.FrameLODIndex].DeformedRootCenter = E.SkeletalWorldBoundsCenter;
+			UpdateHairStrandsMeshTriangles(RHICmdList, ShaderMap, E.FrameLODIndex, HairStrandsTriangleType::DeformedPose, MeshData, E.SimProjectionHairData);
 		}
 	}
 
-	return FHairStrandsInterpolation();
+	// Hair interpolation
+	if (EHairStrandsInterpolationType::RenderStrands == Type)
+	{
+		for (FHairStrandsManager::Element& E : GHairManager.Elements)
+		{
+			if (E.WorldType != WorldType)
+				continue;
+
+			if (E.InterpolationData.Input && E.InterpolationData.Output && E.InterpolationData.Function)
+			{
+				E.InterpolationData.Function(RHICmdList, E.InterpolationData.Input, E.InterpolationData.Output, E.RenProjectionHairData, E.FrameLODIndex);
+			}
+		}
+	}
+}
+
+void GetGroomInterpolationData(const EWorldType::Type WorldType, FHairStrandsProjectionMeshData& OutGeometries)
+{
+	for (FHairStrandsManager::Element& E : GHairManager.Elements)
+	{
+		if (E.WorldType != WorldType)
+			continue;
+
+		for (FCachedGeometrySection Section : E.CachedGeometry.Sections)
+		{		
+			OutGeometries.Sections.Add(ConvertMeshSection(Section, E.SkeletalLocalToWorld));
+		}
+	}
+}
+
+void GetGroomInterpolationData(const EWorldType::Type WorldType, FHairStrandsProjectionHairDatas& Out, TArray<int32>& OutLODIndices)
+{
+	for (FHairStrandsManager::Element& E : GHairManager.Elements)
+	{
+		if (E.WorldType != WorldType)
+			continue;
+
+		const bool bHasDynamicMesh = E.CachedGeometry.Sections.Num() > 0;
+		if (bHasDynamicMesh)
+		{
+			Out.HairData.Add(E.RenProjectionHairData);
+			OutLODIndices.Add(E.FrameLODIndex);
+		}
+	}
 }
 
-void RunHairStrandsInterpolation(FRHICommandListImmediate& RHICmdList)
+FHairStrandsDebugInfos GetHairStandsDebugInfos()
 {
+	FHairStrandsDebugInfos Infos;
 	for (FHairStrandsManager::Element& E : GHairManager.Elements)
 	{
-		if (E.Data.Input && E.Data.Output && E.Data.Function)
+		FHairStrandsDebugInfo& Info = Infos.AddDefaulted_GetRef();
+		Info.Id = E.Id;
+		Info.WorldType = E.WorldType;
+		Info.GroupInfos = E.DebugGroupInfos;
+		for (FHairStrandsDebugInfo::FGroupInfo& GroupInfo : Info.GroupInfos)
 		{
-			E.Data.Function(RHICmdList, E.Data.Input, E.Data.Output);
+			GroupInfo.LODCount = E.RenProjectionHairData.LODDatas.Num();
+			GroupInfo.bHasSkinInterpolation = E.RenProjectionHairData.LODDatas.Num() > 0;
 		}
 	}
+
+	return Infos;
 }
 
-bool IsHairStrandsEnable(EShaderPlatform Platform) { return IsHairStrandsSupported(Platform) && GHairStrandsRenderingEnable == 1 && GHairManager.Elements.Num() > 0; }
\ No newline at end of file
+bool IsHairStrandsEnable(EShaderPlatform Platform) 
+{ 
+	return 
+		IsHairStrandsSupported(Platform) && 
+		GHairStrandsRenderingEnable == 1 && 
+		GHairManager.Elements.Num() > 0; 
+}
\ No newline at end of file
diff --git a/Engine/Source/Runtime/Renderer/Private/HairStrands/HairStrandsMeshProjection.cpp b/Engine/Source/Runtime/Renderer/Private/HairStrands/HairStrandsMeshProjection.cpp
new file mode 100644
--- /dev/null
+++ b/Engine/Source/Runtime/Renderer/Private/HairStrands/HairStrandsMeshProjection.cpp
@@ -0,0 +1,242 @@
+// Copyright 1998-2019 Epic Games, Inc. All Rights Reserved.
+
+#include "HairStrandsMeshProjection.h"
+#include "MeshMaterialShader.h"
+#include "ScenePrivate.h"
+#include "RayTracingDynamicGeometryCollection.h"
+#include "MeshPassProcessor.h"
+#include "RenderGraphUtils.h"
+
+static int32 GHairProjectionMaxSpinLockCount = 4096;
+static FAutoConsoleVariableRef CVarHairProjectionMaxSpinLockCount(TEXT("r.HairStrands.Projection.MaxSpinLockCount"), GHairProjectionMaxSpinLockCount, TEXT("Change the spin lock count for writing hair projection data"));
+
+///////////////////////////////////////////////////////////////////////////////////////////////////
+class FHairMeshProjectionCS : public FGlobalShader
+{
+	DECLARE_GLOBAL_SHADER(FHairMeshProjectionCS);
+	SHADER_USE_PARAMETER_STRUCT(FHairMeshProjectionCS, FGlobalShader);
+
+	BEGIN_SHADER_PARAMETER_STRUCT(FParameters, )
+		SHADER_PARAMETER(uint32, bClear)
+		SHADER_PARAMETER(uint32, MaxRootCount)
+		SHADER_PARAMETER(uint32, MaxSpinLockCount)
+		SHADER_PARAMETER(FMatrix, RootLocalToWorld)
+
+		SHADER_PARAMETER(FMatrix, MeshLocalToWorld)
+		SHADER_PARAMETER(uint32, MeshNumPrimitives)
+		SHADER_PARAMETER(uint32, MeshSectionIndex)
+		SHADER_PARAMETER(uint32, MeshMaxIndexCount)
+		SHADER_PARAMETER(uint32, MeshMaxVertexCount)
+		SHADER_PARAMETER(uint32, MeshIndexOffset)
+		SHADER_PARAMETER_SRV(Buffer, MeshIndexBuffer)
+		SHADER_PARAMETER_SRV(Buffer, MeshPositionBuffer)
+
+		SHADER_PARAMETER_SRV(Buffer, RootPositionBuffer)
+		SHADER_PARAMETER_SRV(Buffer, RootNormalBuffer)
+
+		SHADER_PARAMETER_UAV(StructuredBuffer, OutRootTriangleIndex)
+		SHADER_PARAMETER_UAV(StructuredBuffer, OutRootTriangleBarycentrics)
+		SHADER_PARAMETER_RDG_BUFFER_UAV(StructuredBuffer, OutRootTriangleDistance)
+	END_SHADER_PARAMETER_STRUCT()
+
+public:
+	static bool ShouldCompilePermutation(const FGlobalShaderPermutationParameters& Parameters) { return Parameters.Platform == SP_PCD3D_SM5; }
+};
+
+IMPLEMENT_GLOBAL_SHADER(FHairMeshProjectionCS, "/Engine/Private/HairStrands/HairStrandsMeshProjection.usf", "MainCS", SF_Compute);
+
+static void AddHairStrandMeshProjectionPass(
+	FRDGBuilder& GraphBuilder,
+	TShaderMap<FGlobalShaderType>* ShaderMap,
+	const bool bClear,
+	const int32 LODIndex,
+	const FHairStrandsProjectionMeshSection& MeshSectionData,
+	const FHairStrandsProjectionHairData& RootData,
+	FRDGBufferRef RootDistanceBuffer)
+{
+	if (!RootData.RootPositionBuffer ||
+		!RootData.RootNormalBuffer ||
+		LODIndex < 0 || LODIndex >= RootData.LODDatas.Num() ||
+		!RootData.LODDatas[LODIndex].RootTriangleIndexBuffer ||
+		!RootData.LODDatas[LODIndex].RootTriangleBarycentricBuffer ||
+		!MeshSectionData.IndexBuffer ||
+		!MeshSectionData.PositionBuffer ||
+		MeshSectionData.TotalIndexCount == 0 ||
+		MeshSectionData.TotalVertexCount == 0)
+	{
+		return;
+	}
+
+	// Change the transform to be just a relative transform from the mesh
+	FTransform RootTranform = RootData.LocalToWorld;
+	FTransform MeshTranform = MeshSectionData.LocalToWorld;
+	RootTranform.SetLocation(RootTranform.GetLocation() - MeshTranform.GetLocation());
+	MeshTranform.SetLocation(FVector(0, 0, 0));
+
+	FHairMeshProjectionCS::FParameters* Parameters = GraphBuilder.AllocParameters<FHairMeshProjectionCS::FParameters>();
+	Parameters->bClear				= bClear ? 1 : 0;
+	Parameters->MaxRootCount		= RootData.RootCount;
+	Parameters->MaxSpinLockCount	= FMath::Clamp(GHairProjectionMaxSpinLockCount, 0, 100000);
+	Parameters->RootPositionBuffer	= RootData.RootPositionBuffer;
+	Parameters->RootNormalBuffer	= RootData.RootNormalBuffer;
+	Parameters->RootLocalToWorld	= RootTranform.ToMatrixWithScale();
+
+	Parameters->MeshLocalToWorld	= MeshTranform.ToMatrixWithScale();
+	Parameters->MeshNumPrimitives	= MeshSectionData.NumPrimitives;
+	Parameters->MeshSectionIndex	= MeshSectionData.SectionIndex;
+	Parameters->MeshMaxIndexCount	= MeshSectionData.TotalIndexCount;
+	Parameters->MeshMaxVertexCount	= MeshSectionData.TotalVertexCount;
+	Parameters->MeshIndexOffset		= MeshSectionData.IndexBaseIndex;
+	Parameters->MeshIndexBuffer		= MeshSectionData.IndexBuffer;
+	Parameters->MeshPositionBuffer	= MeshSectionData.PositionBuffer;
+
+	Parameters->OutRootTriangleIndex		= RootData.LODDatas[LODIndex].RootTriangleIndexBuffer->UAV;
+	Parameters->OutRootTriangleBarycentrics = RootData.LODDatas[LODIndex].RootTriangleBarycentricBuffer->UAV;
+	Parameters->OutRootTriangleDistance		= GraphBuilder.CreateUAV(FRDGBufferUAVDesc(RootDistanceBuffer, PF_R32_FLOAT));
+
+	const FIntVector DispatchGroupCount = FComputeShaderUtils::GetGroupCount(RootData.RootCount, 32);
+	check(DispatchGroupCount.X < 65536);
+	TShaderMapRef<FHairMeshProjectionCS> ComputeShader(ShaderMap);
+	FComputeShaderUtils::AddPass(
+		GraphBuilder,
+		RDG_EVENT_NAME("HairStrandsMeshProjection"),
+		*ComputeShader,
+		Parameters,
+		DispatchGroupCount);
+}
+
+void ProjectHairStrandsOntoMesh(
+	FRHICommandListImmediate& RHICmdList, 
+	TShaderMap<FGlobalShaderType>* ShaderMap,
+	const int32 LODIndex,
+	FHairStrandsProjectionMeshData& ProjectionMeshData, 
+	FHairStrandsProjectionHairData& ProjectionHairData)
+{
+	if (LODIndex < 0 || LODIndex >= ProjectionHairData.LODDatas.Num())
+		return;
+
+	FRDGBuilder GraphBuilder(RHICmdList);
+
+	FRDGBufferRef RootDistanceBuffer = GraphBuilder.CreateBuffer(FRDGBufferDesc::CreateStructuredDesc(sizeof(float),  ProjectionHairData.RootCount),	TEXT("HairStrandsTriangleDistance"));
+
+	bool ClearDistance = true;
+	for (FHairStrandsProjectionMeshSection& MeshSection : ProjectionMeshData.Sections)
+	{
+		check(ProjectionHairData.LODDatas[LODIndex].LODIndex == LODIndex);
+		AddHairStrandMeshProjectionPass(GraphBuilder, ShaderMap, ClearDistance, LODIndex, MeshSection, ProjectionHairData, RootDistanceBuffer);
+		ProjectionHairData.LODDatas[LODIndex].bIsValid = true;
+		ClearDistance = false;
+	}
+	GraphBuilder.Execute();
+}
+
+
+
+///////////////////////////////////////////////////////////////////////////////////////////////////
+class FHairUpdateMeshTriangleCS : public FGlobalShader
+{
+	DECLARE_GLOBAL_SHADER(FHairUpdateMeshTriangleCS);
+	SHADER_USE_PARAMETER_STRUCT(FHairUpdateMeshTriangleCS, FGlobalShader);
+
+	BEGIN_SHADER_PARAMETER_STRUCT(FParameters, )
+		SHADER_PARAMETER(FVector, RootTriangleWorldCenter)
+		SHADER_PARAMETER(uint32, MaxRootCount)
+		SHADER_PARAMETER(FMatrix, RootLocalToWorld)
+
+		SHADER_PARAMETER(FMatrix, MeshLocalToWorld)
+		SHADER_PARAMETER(uint32, MeshSectionIndex)
+		SHADER_PARAMETER(uint32, MeshMaxIndexCount)
+		SHADER_PARAMETER(uint32, MeshMaxVertexCount)
+		SHADER_PARAMETER(uint32, MeshIndexOffset)
+		SHADER_PARAMETER_SRV(Buffer, MeshIndexBuffer)
+		SHADER_PARAMETER_SRV(Buffer, MeshPositionBuffer)
+
+		SHADER_PARAMETER_SRV(Buffer, RootTriangleIndex)
+		SHADER_PARAMETER_UAV(StructuredBuffer, OutRootTrianglePosition0)
+		SHADER_PARAMETER_UAV(StructuredBuffer, OutRootTrianglePosition1)
+		SHADER_PARAMETER_UAV(StructuredBuffer, OutRootTrianglePosition2)
+	END_SHADER_PARAMETER_STRUCT()
+
+public:
+	static bool ShouldCompilePermutation(const FGlobalShaderPermutationParameters& Parameters) { return Parameters.Platform == SP_PCD3D_SM5; }
+};
+
+IMPLEMENT_GLOBAL_SHADER(FHairUpdateMeshTriangleCS, "/Engine/Private/HairStrands/HairStrandsMeshUpdate.usf", "MainCS", SF_Compute);
+
+static void AddHairStrandUpdateMeshTrianglesPass(
+	FRDGBuilder& GraphBuilder,
+	TShaderMap<FGlobalShaderType>* ShaderMap,
+	const int32 LODIndex,
+	const HairStrandsTriangleType Type,
+	const FHairStrandsProjectionMeshSection& MeshSectionData,
+	FHairStrandsProjectionHairData& RootData)	
+{
+	if (RootData.RootCount == 0 || LODIndex < 0 || LODIndex >= RootData.LODDatas.Num())
+	{
+		return;
+	}
+	FHairStrandsProjectionHairData::LODData& LODData = RootData.LODDatas[LODIndex];
+	check(LODData.LODIndex == LODIndex);
+
+	// Change the transform to be just a relative transform from the mesh
+	FTransform RootTranform = RootData.LocalToWorld;
+	FTransform MeshTranform = MeshSectionData.LocalToWorld;
+	RootTranform.SetLocation(RootTranform.GetLocation() - MeshTranform.GetLocation());
+	MeshTranform.SetLocation(FVector(0, 0, 0));
+
+	FHairUpdateMeshTriangleCS::FParameters* Parameters = GraphBuilder.AllocParameters<FHairUpdateMeshTriangleCS::FParameters>();
+	Parameters->MaxRootCount		= RootData.RootCount;
+	Parameters->RootTriangleIndex	= LODData.RootTriangleIndexBuffer->SRV;
+	Parameters->RootLocalToWorld	= RootTranform.ToMatrixWithScale();
+	
+	Parameters->MeshLocalToWorld	= MeshTranform.ToMatrixWithScale();
+	Parameters->MeshSectionIndex	= MeshSectionData.SectionIndex;
+	Parameters->MeshMaxIndexCount	= MeshSectionData.TotalIndexCount;
+	Parameters->MeshMaxVertexCount	= MeshSectionData.TotalVertexCount;
+	Parameters->MeshIndexOffset		= MeshSectionData.IndexBaseIndex;
+	Parameters->MeshIndexBuffer		= MeshSectionData.IndexBuffer;
+	Parameters->MeshPositionBuffer	= MeshSectionData.PositionBuffer;
+
+	if (Type == HairStrandsTriangleType::RestPose)
+	{
+		Parameters->RootTriangleWorldCenter = FVector::ZeroVector; // TODO  LODData.RestRootCenter;
+		Parameters->OutRootTrianglePosition0 = LODData.RestRootTrianglePosition0Buffer->UAV;
+		Parameters->OutRootTrianglePosition1 = LODData.RestRootTrianglePosition1Buffer->UAV;
+		Parameters->OutRootTrianglePosition2 = LODData.RestRootTrianglePosition2Buffer->UAV;
+	}
+	else if (Type == HairStrandsTriangleType::DeformedPose)
+	{
+		Parameters->RootTriangleWorldCenter = FVector::ZeroVector; // TODO LODData.DeformedRootCenter;
+		Parameters->OutRootTrianglePosition0 = LODData.DeformedRootTrianglePosition0Buffer->UAV;
+		Parameters->OutRootTrianglePosition1 = LODData.DeformedRootTrianglePosition1Buffer->UAV;
+		Parameters->OutRootTrianglePosition2 = LODData.DeformedRootTrianglePosition2Buffer->UAV;
+	}
+
+	const FIntVector DispatchGroupCount = FComputeShaderUtils::GetGroupCount(RootData.RootCount, 32);
+	check(DispatchGroupCount.X < 65536);
+	TShaderMapRef<FHairUpdateMeshTriangleCS> ComputeShader(ShaderMap);
+	FComputeShaderUtils::AddPass(
+		GraphBuilder,
+		RDG_EVENT_NAME("HairStrandsTriangleMeshUpdate"),
+		*ComputeShader,
+		Parameters,
+		DispatchGroupCount);
+}
+
+void UpdateHairStrandsMeshTriangles(
+	FRHICommandListImmediate& RHICmdList,
+	TShaderMap<FGlobalShaderType>* ShaderMap,
+	const int32 LODIndex,
+	const HairStrandsTriangleType Type,
+	FHairStrandsProjectionMeshData& ProjectionMeshData,
+	FHairStrandsProjectionHairData& ProjectionHairData)
+{
+	FRDGBuilder GraphBuilder(RHICmdList);
+
+	for (FHairStrandsProjectionMeshSection& MeshSection : ProjectionMeshData.Sections)
+	{
+		AddHairStrandUpdateMeshTrianglesPass(GraphBuilder, ShaderMap, LODIndex, Type, MeshSection, ProjectionHairData);
+	}
+	
+	GraphBuilder.Execute();
+}
\ No newline at end of file
diff --git a/Engine/Source/Runtime/Renderer/Private/HairStrands/HairStrandsMeshProjection.h b/Engine/Source/Runtime/Renderer/Private/HairStrands/HairStrandsMeshProjection.h
new file mode 100644
--- /dev/null
+++ b/Engine/Source/Runtime/Renderer/Private/HairStrands/HairStrandsMeshProjection.h
@@ -0,0 +1,30 @@
+// Copyright 1998-2019 Epic Games, Inc. All Rights Reserved.
+
+#pragma once
+
+#include "CoreMinimal.h"
+#include "RendererInterface.h"
+#include "Shader.h"
+
+/* Project hair strands onto a LOD mesh */
+void ProjectHairStrandsOntoMesh(
+	FRHICommandListImmediate& RHICmdList, 
+	TShaderMap<FGlobalShaderType>* ShaderMap,
+	const int32 LODIndex,
+	struct FHairStrandsProjectionMeshData& ProjectionMeshData, 
+	struct FHairStrandsProjectionHairData& ProjectionHairData);
+
+enum class HairStrandsTriangleType
+{
+	RestPose,
+	DeformedPose,
+};
+
+/* Update the triangles information on which hair stands have been projected */
+void UpdateHairStrandsMeshTriangles(
+	FRHICommandListImmediate& RHICmdList,
+	TShaderMap<FGlobalShaderType>* ShaderMap,
+	const int32 LODIndex,
+	const HairStrandsTriangleType Type,
+	struct FHairStrandsProjectionMeshData& ProjectionMeshData,
+	struct FHairStrandsProjectionHairData& ProjectionHairData);
\ No newline at end of file
diff --git a/Engine/Source/Runtime/Renderer/Private/HairStrands/HairStrandsRendering.h b/Engine/Source/Runtime/Renderer/Private/HairStrands/HairStrandsRendering.h
--- a/Engine/Source/Runtime/Renderer/Private/HairStrands/HairStrandsRendering.h
+++ b/Engine/Source/Runtime/Renderer/Private/HairStrands/HairStrandsRendering.h
@@ -8,6 +8,7 @@
 
 #include "CoreMinimal.h"
 #include "RendererInterface.h"
+#include "Shader.h"
 #include "HairStrandsUtils.h"
 #include "HairStrandsCluster.h"
 #include "HairStrandsLUT.h"
@@ -27,3 +28,15 @@ struct FHairStrandsDatas
 	FHairStrandsVisibilityViews HairVisibilityViews;
 	FHairStrandsClusterViews HairClusterPerViews;
 };
+
+enum class EHairStrandsInterpolationType
+{
+	RenderStrands,
+	SimulationStrands
+};
+
+void RunHairStrandsInterpolation(
+	FRHICommandListImmediate& RHICmdList, 
+	EWorldType::Type WorldType, 
+	TShaderMap<FGlobalShaderType>* ShaderMap, 
+	EHairStrandsInterpolationType Type);
\ No newline at end of file
diff --git a/Engine/Source/Runtime/Renderer/Private/LightRendering.cpp b/Engine/Source/Runtime/Renderer/Private/LightRendering.cpp
--- a/Engine/Source/Runtime/Renderer/Private/LightRendering.cpp
+++ b/Engine/Source/Runtime/Renderer/Private/LightRendering.cpp
@@ -274,6 +274,11 @@ class FDeferredLightPS : public FGlobalShader
 			}
 		}
 
+		if (PermutationVector.Get<FHairLighting>() && !IsHairStrandsSupported(Parameters.Platform))
+		{
+			return false;
+		}
+
 		/*if( PermutationVector.Get< FVisualizeCullingDim >() && (
 			PermutationVector.Get< FSourceShapeDim >() == ELightSourceShape::Rect ||
 			PermutationVector.Get< FIESProfileDim >() ||
diff --git a/Engine/Source/Runtime/Renderer/Private/SceneRendering.cpp b/Engine/Source/Runtime/Renderer/Private/SceneRendering.cpp
--- a/Engine/Source/Runtime/Renderer/Private/SceneRendering.cpp
+++ b/Engine/Source/Runtime/Renderer/Private/SceneRendering.cpp
@@ -4123,3 +4123,20 @@ FRHITexture* FSceneRenderer::GetMultiViewSceneColor(const FSceneRenderTargets& S
 		return static_cast<FTextureRHIRef>(ViewFamily.RenderTarget->GetRenderTargetTexture());
 	}
 }
+
+void RunGPUSkinCacheTransition(FRHICommandList& RHICmdList, FScene* Scene, EGPUSkinCacheTransition Type)
+{
+	// * When hair strands is disabled, the skin cache sync point run later 
+	//   during the deferred render pass
+	// * When hair strands is enabled, the skin cache sync point is run earlier, during 
+	//   the init views pass, as the output of the skin cached is used by Niagara
+	const bool bHairStrandsEnabled = IsHairStrandsEnable(Scene->GetShaderPlatform());
+	const bool bRun = bHairStrandsEnabled && EGPUSkinCacheTransition::FrameSetup == Type;
+	if (bRun)
+	{
+		if (FGPUSkinCache* GPUSkinCache = Scene->GetGPUSkinCache())
+		{
+			GPUSkinCache->TransitionAllToReadable(RHICmdList);
+		}
+	}
+}
\ No newline at end of file
diff --git a/Engine/Source/Runtime/Renderer/Private/SceneRendering.h b/Engine/Source/Runtime/Renderer/Private/SceneRendering.h
--- a/Engine/Source/Runtime/Renderer/Private/SceneRendering.h
+++ b/Engine/Source/Runtime/Renderer/Private/SceneRendering.h
@@ -1989,4 +1989,13 @@ private:
 extern FFastVramConfig GFastVRamConfig;
 
 extern bool UseCachedMeshDrawCommands();
-extern bool IsDynamicInstancingEnabled(ERHIFeatureLevel::Type FeatureLevel);
\ No newline at end of file
+extern bool IsDynamicInstancingEnabled(ERHIFeatureLevel::Type FeatureLevel);
+
+enum class EGPUSkinCacheTransition
+{
+	FrameSetup,
+	Renderer,
+};
+
+/* Run GPU skin cache resource transitions */
+void RunGPUSkinCacheTransition(class FRHICommandList& RHICmdList, class FScene* Scene, EGPUSkinCacheTransition Type);
\ No newline at end of file
diff --git a/Engine/Source/Runtime/Renderer/Private/SceneVisibility.cpp b/Engine/Source/Runtime/Renderer/Private/SceneVisibility.cpp
--- a/Engine/Source/Runtime/Renderer/Private/SceneVisibility.cpp
+++ b/Engine/Source/Runtime/Renderer/Private/SceneVisibility.cpp
@@ -2977,6 +2977,15 @@ void FSceneRenderer::PreVisibilityFrameSetup(FRHICommandListImmediate& RHICmdLis
 		}
 	}
 
+	RunGPUSkinCacheTransition(RHICmdList, Scene, EGPUSkinCacheTransition::FrameSetup);
+
+	if (IsHairStrandsEnable(Scene->GetShaderPlatform()) && Views.Num() > 0)
+	{
+		const EWorldType::Type WorldType = Views[0].Family->Scene->GetWorld()->WorldType;
+		auto ShaderMap = GetGlobalShaderMap(FeatureLevel);
+		RunHairStrandsInterpolation(RHICmdList, WorldType, ShaderMap, EHairStrandsInterpolationType::SimulationStrands);
+	}
+
 	// Notify the FX system that the scene is about to perform visibility checks.
 
 	if (Scene->FXSystem && Views.IsValidIndex(0))
diff --git a/Engine/Source/Runtime/Renderer/Public/HairStrandsInterface.h b/Engine/Source/Runtime/Renderer/Public/HairStrandsInterface.h
--- a/Engine/Source/Runtime/Renderer/Public/HairStrandsInterface.h
+++ b/Engine/Source/Runtime/Renderer/Public/HairStrandsInterface.h
@@ -8,6 +8,11 @@
 
 #include "CoreMinimal.h"
 #include "RendererInterface.h"
+#include "Containers/Array.h"
+#include "Engine/EngineTypes.h"
+
+////////////////////////////////////////////////////////////////////////////////////////////////////
+// Misc/Helpers
 
 enum class EHairStrandsDebugMode : uint8
 {
@@ -23,7 +28,7 @@ enum class EHairStrandsDebugMode : uint8
 };
 
 /// Return the active debug view mode
-RENDERER_API EHairStrandsDebugMode GetHairStrandsDebugMode();
+RENDERER_API EHairStrandsDebugMode GetHairStrandsDebugStrandsMode();
 
 /// Return the number of sample subsample count used for the visibility pass
 RENDERER_API uint32 GetHairVisibilitySampleCount();
@@ -41,19 +46,143 @@ RENDERER_API FMinHairRadiusAtDepth1 ComputeMinStrandRadiusAtDepth1(
 	const uint32 SampleCount,
 	const float OverrideStrandHairRasterizationScale);
 
-typedef void (*THairStrandsInterpolationFunction)(FRHICommandListImmediate& RHICmdList, struct FHairStrandsInterpolationInput* Input, struct FHairStrandsInterpolationOutput* Output);
+////////////////////////////////////////////////////////////////////////////////////////////////////
+// Hair/Mesh projection & interpolation
+typedef void (*THairStrandsInterpolationFunction)(
+	FRHICommandListImmediate& RHICmdList, 
+	struct FHairStrandsInterpolationInput* Input, 
+	struct FHairStrandsInterpolationOutput* Output, 
+	struct FHairStrandsProjectionHairData& HairData,
+	int32 LODIndex);
 
-struct FHairStrandsInterpolation
+struct FHairStrandsInterpolationData
 {
 	struct FHairStrandsInterpolationInput* Input = nullptr;
 	struct FHairStrandsInterpolationOutput* Output = nullptr;
 	THairStrandsInterpolationFunction Function = nullptr;
 };
 
-RENDERER_API void RegisterHairStrands(uint64 Id, const FHairStrandsInterpolation& E);
-RENDERER_API FHairStrandsInterpolation UnregisterHairStrands(uint64 Id);
-void RunHairStrandsInterpolation(FRHICommandListImmediate& RHICmdList);
+struct FRWBuffer;
+
+struct FHairStrandsProjectionHairData
+{
+	FRHIShaderResourceView* RootPositionBuffer = nullptr;
+	FRHIShaderResourceView* RootNormalBuffer = nullptr;
+	FRWBuffer* VertexToCurveIndexBuffer = nullptr;
+
+	struct LODData
+	{
+		bool bIsValid = false;
+		int32 LODIndex = -1;
+		FRWBuffer* RootTriangleIndexBuffer = nullptr;
+		FRWBuffer* RootTriangleBarycentricBuffer = nullptr;
+
+		// Rest root triangles' positions are relative to root center (for preserving precision)
+		FVector RestRootCenter = FVector::ZeroVector;
+		FRWBuffer* RestRootTrianglePosition0Buffer = nullptr;
+		FRWBuffer* RestRootTrianglePosition1Buffer = nullptr;
+		FRWBuffer* RestRootTrianglePosition2Buffer = nullptr;
+
+		// Deformed root triangles' positions are relative to root center (for preserving precision)
+		FVector DeformedRootCenter = FVector::ZeroVector;
+		FRWBuffer* DeformedRootTrianglePosition0Buffer = nullptr;
+		FRWBuffer* DeformedRootTrianglePosition1Buffer = nullptr;
+		FRWBuffer* DeformedRootTrianglePosition2Buffer = nullptr;
+	};
+	TArray<LODData> LODDatas;
+
+	uint32 RootCount = 0;
+	FTransform LocalToWorld = FTransform::Identity;
+};
+
+struct FHairStrandsProjectionMeshSection
+{
+	FTransform LocalToWorld;
+	FRHIShaderResourceView* PositionBuffer = nullptr;
+	FRHIShaderResourceView* IndexBuffer = nullptr;
+	uint32 NumPrimitives = 0;
+	uint32 VertexBaseIndex = 0;
+	uint32 IndexBaseIndex = 0;
+	uint32 TotalVertexCount = 0;
+	uint32 TotalIndexCount = 0;
+	uint32 SectionIndex = 0;
+	int32 LODIndex = 0;
+};
+
+struct FHairStrandsProjectionMeshData
+{
+	TArray<FHairStrandsProjectionMeshSection> Sections;
+};
+
+struct FHairStrandsProjectionHairDatas
+{
+	TArray<FHairStrandsProjectionHairData> HairData;
+};
+
+RENDERER_API void AddHairStrandsProjectionQuery(FRHICommandListImmediate& RHICmdList, uint64 Id, EWorldType::Type WorldType, int32 LODIndex, const FVector& RestRootCenter);
+
+////////////////////////////////////////////////////////////////////////////////////////////////////
+// Debug infos
+
+struct FHairStrandsDebugInfo
+{
+	uint64 Id = 0;
+	EWorldType::Type WorldType = EWorldType::None;
+
+	struct FGroupInfo
+	{
+		float MaxRadius = 0;
+		float MaxLength = 0;
+		uint32 VertexCount = 0;
+		uint32 CurveCount = 0;
+
+		bool bHasSkinInterpolation = false;
+		uint32 LODCount = 0;
+	};
+	typedef TArray<FGroupInfo> FGroupInfos;
+	FGroupInfos GroupInfos;
+};
+
+typedef TArray<FHairStrandsDebugInfo> FHairStrandsDebugInfos;
+FHairStrandsDebugInfos GetHairStandsDebugInfos();
+
+////////////////////////////////////////////////////////////////////////////////////////////////////
+// Registrations
+
+RENDERER_API void RegisterHairStrands(
+	uint64 Id,
+	EWorldType::Type WorldType,
+	const FHairStrandsInterpolationData& E,
+	const FHairStrandsProjectionHairData& RenProjection,
+	const FHairStrandsProjectionHairData& SimProjection,
+	const FHairStrandsDebugInfo::FGroupInfos& GroupInfos);
+
+RENDERER_API void UnregisterHairStrands(
+	uint64 Id, 
+	EWorldType::Type WorldType);
+
+RENDERER_API bool UpdateHairStrands(
+	uint64 Id,
+	EWorldType::Type NewWorldType);
+
+RENDERER_API bool UpdateHairStrands(
+	uint64 Id,
+	EWorldType::Type WorldType,
+	const FTransform& HairLocalToWorld,
+	const FHairStrandsProjectionHairData& RenProjection,
+	const FHairStrandsProjectionHairData& SimProjection);
+
+RENDERER_API bool UpdateHairStrands(
+	uint64 Id,
+	EWorldType::Type WorldType,
+	const FTransform& HairLocalToWorld,
+	const FTransform& MeshLocalToWorld,
+	const FVector& SkeletalWorldBoundsCenter);
 
+RENDERER_API bool UpdateHairStrands(
+	uint64 Id, 
+	EWorldType::Type WorldType, 
+	const class FSkeletalMeshObject* MeshObject);
 
 RENDERER_API bool IsHairStrandsSupported(const EShaderPlatform Platform);
-bool IsHairStrandsEnable(EShaderPlatform Platform);
\ No newline at end of file
+bool IsHairStrandsEnable(EShaderPlatform Platform);
