Water optimization for switch
 - no dynamic lights, only sun and sky,
 - no custom depth read for clipping
 - no distortion, no colored transmittance on background

diff --git a/Engine/Shaders/Private/BasePassPixelShader.usf b/Engine/Shaders/Private/BasePassPixelShader.usf
--- a/Engine/Shaders/Private/BasePassPixelShader.usf
+++ b/Engine/Shaders/Private/BasePassPixelShader.usf
@@ -89,9 +89,24 @@ float NormalCurvatureToRoughness(float3 WorldNormal)
 #endif
 
 #if  MATERIAL_SHADINGMODEL_SINGLELAYERWATER
-#include "SingleLayerWaterCommon.ush"
+
+#if SINGLE_LAYER_WATER_SIMPLE_FORWARD
+
+	// Use mobile like simple shading
+	#define SIMPLE_SINGLE_LAYER_WATER 1
+
+	// Change the blend mode here from opaque to pre-multiplied-alpha for simplicity
+	#undef MATERIALBLENDING_ALPHACOMPOSITE
+	#define MATERIALBLENDING_ALPHACOMPOSITE 1
+	#undef MATERIALBLENDING_SOLID
+	#define MATERIALBLENDING_SOLID 0
+
 #endif
 
+#include "SingleLayerWaterCommon.ush"
+
+#endif // MATERIAL_SHADINGMODEL_SINGLELAYERWATER
+
 #if !FORWARD_SHADING
 
 void GetVolumeLightingNonDirectional(float4 AmbientLightingVector, float3 DiffuseColor, inout float3 InterpolatedLighting, out float4 VolumeLighting)
@@ -1235,6 +1250,21 @@ ShadingModel);
 		const float3 V = MaterialParameters.CameraVector;
 		const float3 EnvBrdf = EnvBRDF(GBuffer.SpecularColor, GBuffer.Roughness, max(0.0, dot(N, V)));
 
+#if SINGLE_LAYER_WATER_SIMPLE_FORWARD 
+		const float4 NullDistortionParams = 1.0f;
+		WaterVolumeLightingOutput WaterLighting = EvaluateWaterVolumeLighting(
+			MaterialParameters, PixelMaterialInputs, ResolvedView,
+			DirectionalLightShadow,
+			OpaqueBasePass.SceneDepthWithoutSingleLayerWaterTexture, SingleLayerWaterSceneDepthSampler, // Scene depth texture
+			Specular, NullDistortionParams,
+			SunIlluminance, WaterDiffuseIndirectIlluminance, EnvBrdf,
+			CameraIsUnderWater, WaterVisibility);
+
+		// Add water luminance contribution
+		Color += WaterLighting.Luminance;
+		// Combine top layer opacity with water transmittance (grey scale)
+		Opacity = 1.0 - ((1.0 - Opacity) * dot(WaterLighting.WaterToSceneToLightTransmittance, float3(1.0 / 3.0, 1.0 / 3.0, 1.0 / 3.0)));
+#else
 		Color += EvaluateWaterVolumeLighting(
 			MaterialParameters, PixelMaterialInputs, ResolvedView,
 			DirectionalLightShadow,
@@ -1244,6 +1274,7 @@ ShadingModel);
 			Specular, OpaqueBasePass.DistortionParams,
 			SunIlluminance, WaterDiffuseIndirectIlluminance, EnvBrdf,
 			CameraIsUnderWater, WaterVisibility).Luminance;
+#endif
 	}
 #endif // MATERIAL_SHADINGMODEL_SINGLELAYERWATER
 
diff --git a/Engine/Shaders/Private/Definitions.usf b/Engine/Shaders/Private/Definitions.usf
--- a/Engine/Shaders/Private/Definitions.usf
+++ b/Engine/Shaders/Private/Definitions.usf
@@ -158,3 +158,7 @@
 #ifndef DISABLE_FORWARD_LOCAL_LIGHTS
 #define DISABLE_FORWARD_LOCAL_LIGHTS					0
 #endif
+
+#ifndef SINGLE_LAYER_WATER_SIMPLE_FORWARD
+#define SINGLE_LAYER_WATER_SIMPLE_FORWARD				0
+#endif
diff --git a/Engine/Shaders/Private/SingleLayerWaterCommon.ush b/Engine/Shaders/Private/SingleLayerWaterCommon.ush
--- a/Engine/Shaders/Private/SingleLayerWaterCommon.ush
+++ b/Engine/Shaders/Private/SingleLayerWaterCommon.ush
@@ -43,6 +43,9 @@ bool WaterRefract(float3 V, float3 N, float Eta, out float3 R)
 WaterVolumeLightingOutput EvaluateWaterVolumeLighting(
 	FMaterialPixelParameters MaterialParameters, FPixelMaterialInputs PixelMaterialInputs, ViewState ResolvedView,
 	float DirectionalLightShadow,
+#if SINGLE_LAYER_WATER_SIMPLE_FORWARD
+	Texture2D SceneDepthWithoutSingleLayerWaterTexture, SamplerState SceneDepthWithoutSingleLayerWaterSampler,
+#endif
 #if !SIMPLE_SINGLE_LAYER_WATER
 	Texture2D SceneColorWithoutSingleLayerWaterTexture, SamplerState SceneColorWithoutSingleLayerWaterSampler,
 	Texture2D SceneDepthWithoutSingleLayerWaterTexture, SamplerState SceneDepthWithoutSingleLayerWaterSampler,
@@ -64,8 +67,14 @@ WaterVolumeLightingOutput EvaluateWaterVolumeLighting(
 	{
 		float WaterDepth = ConvertFromDeviceZ(MaterialParameters.SvPosition.z);
 #if SIMPLE_SINGLE_LAYER_WATER
+
+#if SINGLE_LAYER_WATER_SIMPLE_FORWARD 
+		float PixelSceneDepth = SceneDepthWithoutSingleLayerWaterTexture.SampleLevel(SceneDepthWithoutSingleLayerWaterSampler, ViewportUV, 0).x * SINGLE_LAYER_WATER_DEPTH_SCALE;
+		float PixelSceneDeviceZ = ConvertToDeviceZ(PixelSceneDepth);
+#else
 		float PixelSceneDeviceZ = LookupDeviceZ(ViewportUV);
 		float PixelSceneDepth = ConvertFromDeviceZ(PixelSceneDeviceZ);
+#endif
 
 #if FEATURE_LEVEL <= FEATURE_LEVEL_ES3_1
 		// Only mobile run the SIMPLE_SINGLE_LAYER_WATER path today.
diff --git a/Engine/Shaders/Private/SingleLayerWaterComposite.usf b/Engine/Shaders/Private/SingleLayerWaterComposite.usf
--- a/Engine/Shaders/Private/SingleLayerWaterComposite.usf
+++ b/Engine/Shaders/Private/SingleLayerWaterComposite.usf
@@ -294,14 +294,19 @@ float2 SVPositionToSourceTextureUV;
 
 void WaterRefractionCopyPS(
 	in float4 SVPosition : SV_Position,
-	out float4 OutSceneColor : SV_Target0,
-	out float OutSceneDepth : SV_Target1)
+	out float OutSceneDepth : SV_Target0
+#if DOWNSAMPLE_COLOR
+	, out float4 OutSceneColor : SV_Target1
+#endif
+)
 {
 	float2 BufferUV = SVPosition.xy * SVPositionToSourceTextureUV;
 
+#if DOWNSAMPLE_COLOR
 	// Downsample and average all 2x2 quads
 	OutSceneColor.xyz = Texture2DSampleLevel(SceneColorCopyDownsampleTexture, SceneColorCopyDownsampleSampler, BufferUV, 0).xyz;
 	OutSceneColor.w = 1.0f;
+#endif
 
 	float DeviceDepthZ = 0.0f;
 #if DOWNSAMPLE_REFRACTION
diff --git a/Engine/Source/Runtime/Engine/Classes/Materials/Material.h b/Engine/Source/Runtime/Engine/Classes/Materials/Material.h
--- a/Engine/Source/Runtime/Engine/Classes/Materials/Material.h
+++ b/Engine/Source/Runtime/Engine/Classes/Materials/Material.h
@@ -1796,6 +1796,12 @@ public:
 	bool HasFlippedCoordinates();
 #endif //WITH_EDITORONLY_DATA
 
+	// This is to have switch use the simple single layer water shading similar to mobile: no dynamic lights, only sun and sky, no distortion, no colored transmittance on background, no custom depth read.
+	static bool SingleLayerWaterUsesSimpleShading(EShaderPlatform ShaderPlatform) 
+	{
+		return  IsSwitchPlatform(ShaderPlatform) && IsForwardShadingEnabled(ShaderPlatform);
+	}
+
 private:
 #if WITH_EDITOR
 	static FMaterialCompilationFinished MaterialCompilationFinishedEvent;
diff --git a/Engine/Source/Runtime/Engine/Private/Materials/HLSLMaterialTranslator.h b/Engine/Source/Runtime/Engine/Private/Materials/HLSLMaterialTranslator.h
--- a/Engine/Source/Runtime/Engine/Private/Materials/HLSLMaterialTranslator.h
+++ b/Engine/Source/Runtime/Engine/Private/Materials/HLSLMaterialTranslator.h
@@ -1527,6 +1527,10 @@ ResourcesString = TEXT("");
 			{
 				OutEnvironment.SetDefine(TEXT("DISABLE_FORWARD_LOCAL_LIGHTS"), TEXT("1"));
 			}
+			if (ShadingModels.HasShadingModel(MSM_SingleLayerWater) && UMaterial::SingleLayerWaterUsesSimpleShading(InPlatform))
+			{
+				OutEnvironment.SetDefine(TEXT("SINGLE_LAYER_WATER_SIMPLE_FORWARD"), TEXT("1"));
+			}
 
 			if (NumSetMaterials == 1)
 			{
diff --git a/Engine/Source/Runtime/Renderer/Private/SingleLayerWaterRendering.cpp b/Engine/Source/Runtime/Renderer/Private/SingleLayerWaterRendering.cpp
--- a/Engine/Source/Runtime/Renderer/Private/SingleLayerWaterRendering.cpp
+++ b/Engine/Source/Runtime/Renderer/Private/SingleLayerWaterRendering.cpp
@@ -138,9 +138,12 @@ FSingleLayerWaterPassMeshProcessor::FSingleLayerWaterPassMeshProcessor(const FSc
 	: FMeshPassProcessor(Scene, Scene->GetFeatureLevel(), InViewIfDynamicMeshCommand, InDrawListContext)
 	, PassDrawRenderState(InPassDrawRenderState)
 {
-	// Not needed for now
-	//FRHIBlendState* WaterBlendState = TStaticBlendState<CW_RGBA, BO_Add, BF_One, BF_InverseSourceAlpha>::GetRHI();
-	//PassDrawRenderState.SetBlendState(WaterBlendState);
+	if (UMaterial::SingleLayerWaterUsesSimpleShading(Scene->GetShaderPlatform()))
+	{
+		// Force non opaque, pre multiplied alpha, transparent blend mode because water is going to be blended against scene color (no distortion from texture scene color).
+		FRHIBlendState* ForwardSimpleWaterBlendState = TStaticBlendState<CW_RGBA, BO_Add, BF_One, BF_InverseSourceAlpha>::GetRHI();
+		PassDrawRenderState.SetBlendState(ForwardSimpleWaterBlendState);
+	}
 }
 
 void FSingleLayerWaterPassMeshProcessor::AddMeshBatch(const FMeshBatch& RESTRICT MeshBatch, uint64 BatchElementMask, const FPrimitiveSceneProxy* RESTRICT PrimitiveSceneProxy, int32 StaticMeshId)
@@ -367,8 +370,9 @@ class FWaterRefractionCopyPS : public FGlobalShader
 	END_SHADER_PARAMETER_STRUCT()
 
 	class FDownsampleRefraction : SHADER_PERMUTATION_BOOL("DOWNSAMPLE_REFRACTION");
+	class FDownsampleColor : SHADER_PERMUTATION_BOOL("DOWNSAMPLE_COLOR");
 
-	using FPermutationDomain = TShaderPermutationDomain<FDownsampleRefraction>;
+	using FPermutationDomain = TShaderPermutationDomain<FDownsampleRefraction, FDownsampleColor>;
 
 	static bool ShouldCompilePermutation(const FGlobalShaderPermutationParameters& Parameters)
 	{
@@ -428,14 +432,20 @@ public:
 
 void FDeferredShadingSceneRenderer::CopySingleLayerWaterTextures(FRHICommandListImmediate& RHICmdList, FSingleLayerWaterPassData& PassData)
 {
+	const bool bSingleLayerWaterSimpleShading = UMaterial::SingleLayerWaterUsesSimpleShading(Scene->GetShaderPlatform());
+	bool bCopyColor = !bSingleLayerWaterSimpleShading;
+
 	check(RHICmdList.IsOutsideRenderPass());
 	FSceneRenderTargets& SceneContext = FSceneRenderTargets::Get(RHICmdList);
 	const ERHIFeatureLevel::Type CurrentFeatureLevel = SceneContext.GetCurrentFeatureLevel();
 	const int32 RefractionDownsampleFactor = FMath::Clamp(GSingleLayerWaterRefractionDownsampleFactor, 1, 8);
 	const FIntPoint RefractionResolution = FIntPoint::DivideAndRoundDown(SceneContext.GetBufferSizeXY(), RefractionDownsampleFactor);
 
-	FPooledRenderTargetDesc ColorDesc(FPooledRenderTargetDesc::Create2DDesc(RefractionResolution, SceneContext.GetSceneColorFormat(), SceneContext.GetDefaultColorClear(), TexCreate_None, TexCreate_RenderTargetable, false));
-	GRenderTargetPool.FindFreeElement(RHICmdList, ColorDesc, PassData.SceneColorWithoutSingleLayerWater, TEXT("SceneColorWithoutSingleLayerWater"));
+	if (!bCopyColor)
+	{
+		FPooledRenderTargetDesc ColorDesc(FPooledRenderTargetDesc::Create2DDesc(RefractionResolution, SceneContext.GetSceneColorFormat(), SceneContext.GetDefaultColorClear(), TexCreate_None, TexCreate_RenderTargetable, false));
+		GRenderTargetPool.FindFreeElement(RHICmdList, ColorDesc, PassData.SceneColorWithoutSingleLayerWater, TEXT("SceneColorWithoutSingleLayerWater"));
+	}
 
 	FPooledRenderTargetDesc DepthDesc(FPooledRenderTargetDesc::Create2DDesc(RefractionResolution, GSingleLayerWaterRefractionFullPrecision ? PF_R32_FLOAT : PF_R16F, SceneContext.GetDefaultColorClear(), TexCreate_None, TexCreate_RenderTargetable, false));
 	GRenderTargetPool.FindFreeElement(RHICmdList, DepthDesc, PassData.SceneDepthWithoutSingleLayerWater, TEXT("SceneDepthWithoutSingleLayerWater"));
@@ -444,8 +454,6 @@ void FDeferredShadingSceneRenderer::CopySingleLayerWaterTextures(FRHICommandList
 	const FViewInfo& View = Views[0];
 
 	FRDGBuilder GraphBuilder(RHICmdList);
-	FRDGTextureRef TargetColorTexture = GraphBuilder.RegisterExternalTexture(PassData.SceneColorWithoutSingleLayerWater);
-	FRDGTextureRef TargetSceneDepthTexture = GraphBuilder.RegisterExternalTexture(PassData.SceneDepthWithoutSingleLayerWater);
 
 	FWaterRefractionCopyPS::FParameters* PassParameters = GraphBuilder.AllocParameters<FWaterRefractionCopyPS::FParameters>();
 	PassParameters->ViewUniformBuffer = View.ViewUniformBuffer;
@@ -457,8 +465,17 @@ void FDeferredShadingSceneRenderer::CopySingleLayerWaterTextures(FRHICommandList
 	PassParameters->RenderTargets[0] = FRenderTargetBinding(TargetColorTexture, ERenderTargetLoadAction::ELoad);
 	PassParameters->RenderTargets[1] = FRenderTargetBinding(TargetSceneDepthTexture, ERenderTargetLoadAction::ELoad);
 
+	FRDGTextureRef TargetSceneDepthTexture = GraphBuilder.RegisterExternalTexture(PassData.SceneDepthWithoutSingleLayerWater);
+	PassParameters->RenderTargets[0] = FRenderTargetBinding(TargetSceneDepthTexture, ERenderTargetLoadAction::ELoad, ERenderTargetStoreAction::EStore);
+	if (!bCopyColor)
+	{
+		FRDGTextureRef TargetColorTexture = GraphBuilder.RegisterExternalTexture(PassData.SceneColorWithoutSingleLayerWater);
+		PassParameters->RenderTargets[1] = FRenderTargetBinding(TargetColorTexture, ERenderTargetLoadAction::ELoad, ERenderTargetStoreAction::EStore);
+	}
+
 	FWaterRefractionCopyPS::FPermutationDomain PermutationVector;
 	PermutationVector.Set<FWaterRefractionCopyPS::FDownsampleRefraction>(RefractionDownsampleFactor > 1);
+	PermutationVector.Set<FWaterRefractionCopyPS::FDownsampleColor>(bCopyColor);
 	auto PixelShader = View.ShaderMap->GetShader<FWaterRefractionCopyPS>(PermutationVector);
 
 	const FIntRect RefractionViewRect = FIntRect(FIntPoint(0, 0), FIntPoint::DivideAndRoundDown(View.ViewRect.Size(), RefractionDownsampleFactor));
@@ -475,7 +492,10 @@ void FDeferredShadingSceneRenderer::CopySingleLayerWaterTextures(FRHICommandList
 		RefractionViewRect);
 	GraphBuilder.Execute();
 
-	RHICmdList.TransitionResource(EResourceTransitionAccess::EReadable, PassData.SceneColorWithoutSingleLayerWater->GetRenderTargetItem().TargetableTexture.GetReference());
+	if (!bCopyColor)
+	{
+		RHICmdList.TransitionResource(EResourceTransitionAccess::EReadable, PassData.SceneColorWithoutSingleLayerWater->GetRenderTargetItem().TargetableTexture.GetReference());
+	}
 	RHICmdList.TransitionResource(EResourceTransitionAccess::EReadable, PassData.SceneDepthWithoutSingleLayerWater->GetRenderTargetItem().TargetableTexture.GetReference());
 }
 
@@ -763,7 +783,7 @@ bool FDeferredShadingSceneRenderer::RenderSingleLayerWaterPass(FRHICommandListIm
 				View, 
 				WhiteForwardScreenSpaceShadowMask, 
 				&PassData.SceneWithoutSingleLayerWaterMaxUV,
-				PassData.SceneColorWithoutSingleLayerWater,
+				PassData.SceneColorWithoutSingleLayerWater.IsValid() ? PassData.SceneColorWithoutSingleLayerWater : GSystemTextures.BlackDummy,
 				PassData.SceneDepthWithoutSingleLayerWater,
 				WaterPassUniformBuffer);
 
