Fixed compilation

diff --git a/Engine/Source/Runtime/Engine/Classes/Materials/Material.h b/Engine/Source/Runtime/Engine/Classes/Materials/Material.h
--- a/Engine/Source/Runtime/Engine/Classes/Materials/Material.h
+++ b/Engine/Source/Runtime/Engine/Classes/Materials/Material.h
@@ -1796,9 +1796,6 @@ public:
 	bool HasFlippedCoordinates();
 #endif //WITH_EDITORONLY_DATA
 
-	// This is to have switch use the simple single layer water shading similar to mobile: no dynamic lights, only sun and sky, no distortion, no colored transmittance on background, no custom depth read.
-	static bool SingleLayerWaterUsesSimpleShading(EShaderPlatform ShaderPlatform);
-
 private:
 #if WITH_EDITOR
 	static FMaterialCompilationFinished MaterialCompilationFinishedEvent;
diff --git a/Engine/Source/Runtime/Engine/Private/Materials/HLSLMaterialTranslator.h b/Engine/Source/Runtime/Engine/Private/Materials/HLSLMaterialTranslator.h
--- a/Engine/Source/Runtime/Engine/Private/Materials/HLSLMaterialTranslator.h
+++ b/Engine/Source/Runtime/Engine/Private/Materials/HLSLMaterialTranslator.h
@@ -1527,7 +1527,10 @@ ResourcesString = TEXT("");
 			{
 				OutEnvironment.SetDefine(TEXT("DISABLE_FORWARD_LOCAL_LIGHTS"), TEXT("1"));
 			}
-			if (ShadingModels.HasShadingModel(MSM_SingleLayerWater) && UMaterial::SingleLayerWaterUsesSimpleShading(InPlatform))
+
+			// This is to have switch use the simple single layer water shading similar to mobile: no dynamic lights, only sun and sky, no distortion, no colored transmittance on background, no custom depth read.
+			const bool bSingleLayerWaterUsesSimpleShading = IsSwitchPlatform(InPlatform) && IsForwardShadingEnabled(InPlatform);
+			if (ShadingModels.HasShadingModel(MSM_SingleLayerWater) && bSingleLayerWaterUsesSimpleShading)
 			{
 				OutEnvironment.SetDefine(TEXT("SINGLE_LAYER_WATER_SIMPLE_FORWARD"), TEXT("1"));
 			}
diff --git a/Engine/Source/Runtime/Engine/Private/Materials/Material.cpp b/Engine/Source/Runtime/Engine/Private/Materials/Material.cpp
--- a/Engine/Source/Runtime/Engine/Private/Materials/Material.cpp
+++ b/Engine/Source/Runtime/Engine/Private/Materials/Material.cpp
@@ -6319,11 +6319,6 @@ bool UMaterial::HasFlippedCoordinates()
 }
 #endif //WITH_EDITORONLY_DATA
 
-bool UMaterial::SingleLayerWaterUsesSimpleShading(EShaderPlatform ShaderPlatform)
-{
-	return  IsSwitchPlatform(ShaderPlatform) && IsForwardShadingEnabled(ShaderPlatform);
-}
-
 void UMaterial::GetLightingGuidChain(bool bIncludeTextures, TArray<FGuid>& OutGuids) const
 {
 #if WITH_EDITORONLY_DATA
diff --git a/Engine/Source/Runtime/Renderer/Private/SingleLayerWaterRendering.cpp b/Engine/Source/Runtime/Renderer/Private/SingleLayerWaterRendering.cpp
--- a/Engine/Source/Runtime/Renderer/Private/SingleLayerWaterRendering.cpp
+++ b/Engine/Source/Runtime/Renderer/Private/SingleLayerWaterRendering.cpp
@@ -92,6 +92,12 @@ static bool ShouldRenderSingleLayerWater(const FViewInfo& View)
 	return View.bHasSingleLayerWaterMaterial;
 }
 
+// This is to have switch use the simple single layer water shading similar to mobile: no dynamic lights, only sun and sky, no distortion, no colored transmittance on background, no custom depth read.
+bool SingleLayerWaterUsesSimpleShading(EShaderPlatform ShaderPlatform)
+{
+	return  IsSwitchPlatform(ShaderPlatform) && IsForwardShadingEnabled(ShaderPlatform);
+}
+
 bool ShouldRenderSingleLayerWater(const TArray<FViewInfo>& Views, const FEngineShowFlags& EngineShowFlags)
 {
 	if (CVarWaterSingleLayer.GetValueOnRenderThread() > 0) // && EngineShowFlags.Water)
@@ -138,7 +144,8 @@ FSingleLayerWaterPassMeshProcessor::FSingleLayerWaterPassMeshProcessor(const FSc
 	: FMeshPassProcessor(Scene, Scene->GetFeatureLevel(), InViewIfDynamicMeshCommand, InDrawListContext)
 	, PassDrawRenderState(InPassDrawRenderState)
 {
-	if (UMaterial::SingleLayerWaterUsesSimpleShading(Scene->GetShaderPlatform()))
+	const bool bSingleLayerWaterUsesSimpleShading = SingleLayerWaterUsesSimpleShading(Scene->GetShaderPlatform());
+	if (bSingleLayerWaterUsesSimpleShading)
 	{
 		// Force non opaque, pre multiplied alpha, transparent blend mode because water is going to be blended against scene color (no distortion from texture scene color).
 		FRHIBlendState* ForwardSimpleWaterBlendState = TStaticBlendState<CW_RGBA, BO_Add, BF_One, BF_InverseSourceAlpha>::GetRHI();
@@ -432,8 +439,8 @@ public:
 
 void FDeferredShadingSceneRenderer::CopySingleLayerWaterTextures(FRHICommandListImmediate& RHICmdList, FSingleLayerWaterPassData& PassData)
 {
-	const bool bSingleLayerWaterSimpleShading = UMaterial::SingleLayerWaterUsesSimpleShading(Scene->GetShaderPlatform());
-	bool bCopyColor = !bSingleLayerWaterSimpleShading;
+	const bool bSingleLayerWaterUsesSimpleShading = SingleLayerWaterUsesSimpleShading(Scene->GetShaderPlatform());
+	bool bCopyColor = !bSingleLayerWaterUsesSimpleShading;
 
 	check(RHICmdList.IsOutsideRenderPass());
 	FSceneRenderTargets& SceneContext = FSceneRenderTargets::Get(RHICmdList);
