Shelve from Sebastien.Hillaire: Fixed water refraction copy logic.

diff --git a/Engine/Source/Runtime/Renderer/Private/SingleLayerWaterRendering.cpp b/Engine/Source/Runtime/Renderer/Private/SingleLayerWaterRendering.cpp
--- a/Engine/Source/Runtime/Renderer/Private/SingleLayerWaterRendering.cpp
+++ b/Engine/Source/Runtime/Renderer/Private/SingleLayerWaterRendering.cpp
@@ -448,7 +448,7 @@ void FDeferredShadingSceneRenderer::CopySingleLayerWaterTextures(FRHICommandList
 	const int32 RefractionDownsampleFactor = FMath::Clamp(GSingleLayerWaterRefractionDownsampleFactor, 1, 8);
 	const FIntPoint RefractionResolution = FIntPoint::DivideAndRoundDown(SceneContext.GetBufferSizeXY(), RefractionDownsampleFactor);
 
-	if (!bCopyColor)
+	if (bCopyColor)
 	{
 		FPooledRenderTargetDesc ColorDesc(FPooledRenderTargetDesc::Create2DDesc(RefractionResolution, SceneContext.GetSceneColorFormat(), SceneContext.GetDefaultColorClear(), TexCreate_None, TexCreate_RenderTargetable, false));
 		GRenderTargetPool.FindFreeElement(RHICmdList, ColorDesc, PassData.SceneColorWithoutSingleLayerWater, TEXT("SceneColorWithoutSingleLayerWater"));
@@ -474,7 +474,7 @@ void FDeferredShadingSceneRenderer::CopySingleLayerWaterTextures(FRHICommandList
 
 	FRDGTextureRef TargetSceneDepthTexture = GraphBuilder.RegisterExternalTexture(PassData.SceneDepthWithoutSingleLayerWater);
 	PassParameters->RenderTargets[0] = FRenderTargetBinding(TargetSceneDepthTexture, ERenderTargetLoadAction::ELoad, ERenderTargetStoreAction::EStore);
-	if (!bCopyColor)
+	if (bCopyColor)
 	{
 		FRDGTextureRef TargetColorTexture = GraphBuilder.RegisterExternalTexture(PassData.SceneColorWithoutSingleLayerWater);
 		PassParameters->RenderTargets[1] = FRenderTargetBinding(TargetColorTexture, ERenderTargetLoadAction::ELoad, ERenderTargetStoreAction::EStore);
@@ -499,7 +499,7 @@ void FDeferredShadingSceneRenderer::CopySingleLayerWaterTextures(FRHICommandList
 		RefractionViewRect);
 	GraphBuilder.Execute();
 
-	if (!bCopyColor)
+	if (bCopyColor)
 	{
 		RHICmdList.TransitionResource(EResourceTransitionAccess::EReadable, PassData.SceneColorWithoutSingleLayerWater->GetRenderTargetItem().TargetableTexture.GetReference());
 	}
