Fixing compilation after a merge.

diff --git a/Engine/Source/Runtime/Renderer/Private/SingleLayerWaterRendering.cpp b/Engine/Source/Runtime/Renderer/Private/SingleLayerWaterRendering.cpp
--- a/Engine/Source/Runtime/Renderer/Private/SingleLayerWaterRendering.cpp
+++ b/Engine/Source/Runtime/Renderer/Private/SingleLayerWaterRendering.cpp
@@ -469,8 +469,6 @@ void FDeferredShadingSceneRenderer::CopySingleLayerWaterTextures(FRHICommandList
 	PassParameters->SceneDepthCopyDownsampleTexture = GraphBuilder.RegisterExternalTexture(SceneContext.SceneDepthZ);
 	PassParameters->SceneDepthCopyDownsampleSampler = TStaticSamplerState<SF_Bilinear, AM_Clamp, AM_Clamp, AM_Clamp>::GetRHI();
 	PassParameters->SVPositionToSourceTextureUV = FVector2D(RefractionDownsampleFactor / float(SceneContext.GetBufferSizeXY().X), RefractionDownsampleFactor / float(SceneContext.GetBufferSizeXY().Y));
-	PassParameters->RenderTargets[0] = FRenderTargetBinding(TargetColorTexture, ERenderTargetLoadAction::ELoad);
-	PassParameters->RenderTargets[1] = FRenderTargetBinding(TargetSceneDepthTexture, ERenderTargetLoadAction::ELoad);
 
 	FRDGTextureRef TargetSceneDepthTexture = GraphBuilder.RegisterExternalTexture(PassData.SceneDepthWithoutSingleLayerWater);
 	PassParameters->RenderTargets[0] = FRenderTargetBinding(TargetSceneDepthTexture, ERenderTargetLoadAction::ELoad, ERenderTargetStoreAction::EStore);
