When translucency isn't rendered skip blending seperate translucency and scene color, since separate translucency is blank.

diff --git a/Engine/Source/Runtime/Renderer/Private/PostProcess/PostProcessing.cpp b/Engine/Source/Runtime/Renderer/Private/PostProcess/PostProcessing.cpp
--- a/Engine/Source/Runtime/Renderer/Private/PostProcess/PostProcessing.cpp
+++ b/Engine/Source/Runtime/Renderer/Private/PostProcess/PostProcessing.cpp
@@ -1098,7 +1098,7 @@ void FPostProcessing::Process(FRHICommandListImmediate& RHICmdList, const FViewI
 					}
 
 					// DOF passes were not added, therefore need to compose Separate translucency manually.
-					if (LocalSceneColorTexture == SceneColorTexture && SeparateTranslucencyTexture)
+					if (LocalSceneColorTexture == SceneColorTexture && SeparateTranslucencyTexture && View.Family->EngineShowFlags.Translucency)
 					{
 						LocalSceneColorTexture = AddSeparateTranslucencyCompositionPass(GraphBuilder, View, SceneColorTexture, SeparateTranslucencyTexture);
 					}
@@ -1416,7 +1416,7 @@ void FPostProcessing::Process(FRHICommandListImmediate& RHICmdList, const FViewI
 					FRDGTextureRef LocalSeparateTranslucency = Pass->CreateRDGTextureForOptionalInput(GraphBuilder, ePId_Input1, TEXT("SeparateTranslucency"));
 
 					FRDGTextureRef NewSceneColor = SceneColor;
-					if (LocalSeparateTranslucency)
+					if (LocalSeparateTranslucency && InContext.View.Family->EngineShowFlags.Translucency)
 					{
 						NewSceneColor = AddSeparateTranslucencyCompositionPass(GraphBuilder, InContext.View, SceneColor, LocalSeparateTranslucency);
 					}
