Fix compilation error
#rnx

diff --git a/Engine/Source/Runtime/Renderer/Private/SingleLayerWaterRendering.cpp b/Engine/Source/Runtime/Renderer/Private/SingleLayerWaterRendering.cpp
--- a/Engine/Source/Runtime/Renderer/Private/SingleLayerWaterRendering.cpp
+++ b/Engine/Source/Runtime/Renderer/Private/SingleLayerWaterRendering.cpp
@@ -471,11 +471,11 @@ void FDeferredShadingSceneRenderer::CopySingleLayerWaterTextures(FRHICommandList
 	PassParameters->SVPositionToSourceTextureUV = FVector2D(RefractionDownsampleFactor / float(SceneContext.GetBufferSizeXY().X), RefractionDownsampleFactor / float(SceneContext.GetBufferSizeXY().Y));
 
 	FRDGTextureRef TargetSceneDepthTexture = GraphBuilder.RegisterExternalTexture(PassData.SceneDepthWithoutSingleLayerWater);
-	PassParameters->RenderTargets[0] = FRenderTargetBinding(TargetSceneDepthTexture, ERenderTargetLoadAction::ELoad, ERenderTargetStoreAction::EStore);
+	PassParameters->RenderTargets[0] = FRenderTargetBinding(TargetSceneDepthTexture, ERenderTargetLoadAction::ELoad, (uint8)ERenderTargetStoreAction::EStore);
 	if (bCopyColor)
 	{
 		FRDGTextureRef TargetColorTexture = GraphBuilder.RegisterExternalTexture(PassData.SceneColorWithoutSingleLayerWater);
-		PassParameters->RenderTargets[1] = FRenderTargetBinding(TargetColorTexture, ERenderTargetLoadAction::ELoad, ERenderTargetStoreAction::EStore);
+		PassParameters->RenderTargets[1] = FRenderTargetBinding(TargetColorTexture, ERenderTargetLoadAction::ELoad, (uint8)ERenderTargetStoreAction::EStore);
 	}
 
 	FWaterRefractionCopyPS::FPermutationDomain PermutationVector;
