* Water composite pass now disables reflection capture sampling if none of those are in the scene. This fixes artifacts when sampling random data and also speeds up a bit water composite pass.
* Fixed reflection capture apply on water. SingleLayerWayerComposite.usf calls CompositeReflectionCapturesAndSkylight with a 0 reflection capture Id and set #define REFLECTION_COMPOSITE_USE_BLENDED_REFLECTION_CAPTURES 0. This causes always to sample from a default reflection capture 0. Changed REFLECTION_COMPOSITE_USE_BLENDED_REFLECTION_CAPTURES to 1 in order to blend multiple reflection captures on top of the water in FEATURE_LEVEL_SM5.
* Refactored FSingleLayerWaterCompositePS::Params into FWaterCompositeParameters with VS and PS parameters inside. This allows to call ClearUnusedGraphResources and remove unused parameters from the draw call.

diff --git a/Engine/Shaders/Private/SingleLayerWaterComposite.usf b/Engine/Shaders/Private/SingleLayerWaterComposite.usf
--- a/Engine/Shaders/Private/SingleLayerWaterComposite.usf
+++ b/Engine/Shaders/Private/SingleLayerWaterComposite.usf
@@ -1,7 +1,7 @@
 // Copyright 1998-2019 Epic Games, Inc. All Rights Reserved.
 
 /*=============================================================================
-	WaterComposite - functionality to apply local cubemaps.
+	SingleLayerWaterComposite - final step of the single layer water
 =============================================================================*/
 
 #include "Common.ush"
@@ -16,7 +16,7 @@
 
 #define ENABLE_SKY_LIGHT 1
 
-#define REFLECTION_COMPOSITE_USE_BLENDED_REFLECTION_CAPTURES 0
+#define REFLECTION_COMPOSITE_USE_BLENDED_REFLECTION_CAPTURES (FEATURE_LEVEL >= FEATURE_LEVEL_SM5)
 #define REFLECTION_COMPOSITE_SUPPORT_SKYLIGHT_BLEND 0
 #include "ReflectionEnvironmentComposite.ush"
 
diff --git a/Engine/Source/Runtime/Renderer/Private/SingleLayerWaterRendering.cpp b/Engine/Source/Runtime/Renderer/Private/SingleLayerWaterRendering.cpp
--- a/Engine/Source/Runtime/Renderer/Private/SingleLayerWaterRendering.cpp
+++ b/Engine/Source/Runtime/Renderer/Private/SingleLayerWaterRendering.cpp
@@ -257,23 +257,21 @@ BEGIN_SHADER_PARAMETER_STRUCT(FSingleLayerWaterCommonShaderParameters, )
 	SHADER_PARAMETER_STRUCT_REF(FReflectionUniformParameters, ReflectionsParameters)
 END_SHADER_PARAMETER_STRUCT()
 
-class FSingleLayerWaterScreenSpaceReflections : SHADER_PERMUTATION_BOOL("SCREEN_SPACE_REFLECTION");
-
 class FSingleLayerWaterCompositePS : public FGlobalShader
 {
 	DECLARE_GLOBAL_SHADER(FSingleLayerWaterCompositePS);
 	SHADER_USE_PARAMETER_STRUCT(FSingleLayerWaterCompositePS, FGlobalShader)
 
-		using FPermutationDomain = TShaderPermutationDomain<FSingleLayerWaterScreenSpaceReflections>;
+	class FScreenSpaceReflections : SHADER_PERMUTATION_BOOL("SCREEN_SPACE_REFLECTION");
+	class FHasBoxCaptures : SHADER_PERMUTATION_BOOL("REFLECTION_COMPOSITE_HAS_BOX_CAPTURES");
+	class FHasSphereCaptures : SHADER_PERMUTATION_BOOL("REFLECTION_COMPOSITE_HAS_SPHERE_CAPTURES");
+	using FPermutationDomain = TShaderPermutationDomain<FScreenSpaceReflections, FHasBoxCaptures, FHasSphereCaptures>;
 
 	BEGIN_SHADER_PARAMETER_STRUCT(FParameters, )
 		SHADER_PARAMETER_STRUCT_INCLUDE(FSingleLayerWaterCommonShaderParameters, CommonParameters)
-		SHADER_PARAMETER_RDG_BUFFER(Buffer<uint>, IndirectDrawParameter)	// Not used in shader but need to be reference in the parameter list
-		SHADER_PARAMETER_RDG_BUFFER_SRV(StructuredBuffer<uint>, TileListData)
-		RENDER_TARGET_BINDING_SLOTS()
 	END_SHADER_PARAMETER_STRUCT()
 
-		static FPermutationDomain RemapPermutation(FPermutationDomain PermutationVector)
+	static FPermutationDomain RemapPermutation(FPermutationDomain PermutationVector)
 	{
 		return PermutationVector;
 	}
@@ -339,7 +337,10 @@ class FWaterTileVS : public FGlobalShader
 
 	using FPermutationDomain = TShaderPermutationDomain<>;
 
-	using FParameters = FSingleLayerWaterCompositePS::FParameters; // Sharing parameters for proper registration with RDG
+	BEGIN_SHADER_PARAMETER_STRUCT(FParameters, )
+		SHADER_PARAMETER_STRUCT_REF(FViewUniformShaderParameters, ViewUniformBuffer)
+		SHADER_PARAMETER_RDG_BUFFER_SRV(StructuredBuffer<uint>, TileListData)
+	END_SHADER_PARAMETER_STRUCT()
 
 	static FPermutationDomain RemapPermutation(FPermutationDomain PermutationVector)
 	{
@@ -550,6 +551,12 @@ void FDeferredShadingSceneRenderer::FinishWaterGBufferPassAndResolve(FRHICommand
 	SceneContext.FinishGBufferPassAndResolve(RHICmdList);
 }
 
+BEGIN_SHADER_PARAMETER_STRUCT(FWaterCompositeParameters, )
+	SHADER_PARAMETER_STRUCT_INCLUDE(FWaterTileVS::FParameters, VS)
+	SHADER_PARAMETER_STRUCT_INCLUDE(FSingleLayerWaterCompositePS::FParameters, PS)
+	SHADER_PARAMETER_RDG_BUFFER(Buffer<uint>, IndirectDrawParameter)
+	RENDER_TARGET_BINDING_SLOTS()
+END_SHADER_PARAMETER_STRUCT()
 
 void FDeferredShadingSceneRenderer::RenderSingleLayerWaterReflections(FRHICommandListImmediate& RHICmdList, FSingleLayerWaterPassData& PassData)
 {
@@ -681,22 +688,34 @@ void FDeferredShadingSceneRenderer::RenderSingleLayerWaterReflections(FRHIComman
 
 		// Composite reflections on water
 		{
+			const bool bHasBoxCaptures = (View.NumBoxReflectionCaptures > 0);
+			const bool bHasSphereCaptures = (View.NumSphereReflectionCaptures > 0);
+
 			FSingleLayerWaterCompositePS::FPermutationDomain PermutationVector;
-			PermutationVector.Set<FSingleLayerWaterScreenSpaceReflections>(bEnableSSR);
+			PermutationVector.Set<FSingleLayerWaterCompositePS::FScreenSpaceReflections>(bEnableSSR);
+			PermutationVector.Set<FSingleLayerWaterCompositePS::FHasBoxCaptures>(bHasBoxCaptures);
+			PermutationVector.Set<FSingleLayerWaterCompositePS::FHasSphereCaptures>(bHasSphereCaptures);
 			TShaderMapRef<FSingleLayerWaterCompositePS> PixelShader(View.ShaderMap, PermutationVector);
 
-			FSingleLayerWaterCompositePS::FParameters* PassParameters = GraphBuilder.AllocParameters<FSingleLayerWaterCompositePS::FParameters>();
-			SetCommonParameters(PassParameters->CommonParameters);
-			PassParameters->TileListData = TiledScreenSpaceReflection.TileListStructureBufferSRV;
+			FWaterCompositeParameters* PassParameters = GraphBuilder.AllocParameters<FWaterCompositeParameters>();
+
+			PassParameters->VS.ViewUniformBuffer = View.ViewUniformBuffer;
+			PassParameters->VS.TileListData = TiledScreenSpaceReflection.TileListStructureBufferSRV;
+
+			SetCommonParameters(PassParameters->PS.CommonParameters);
+
 			PassParameters->IndirectDrawParameter = TiledScreenSpaceReflection.DispatchIndirectParametersBuffer;
 			PassParameters->RenderTargets[0] = FRenderTargetBinding(SceneColorTexture, ERenderTargetLoadAction::ELoad);
 
-			ValidateShaderParameters(*PixelShader, *PassParameters);
+			ValidateShaderParameters(*PixelShader, PassParameters->PS);
+			ClearUnusedGraphResources(*PixelShader, &PassParameters->PS);
+
 			if (bRunTiled)
 			{
 				FWaterTileVS::FPermutationDomain VsPermutationVector;
 				TShaderMapRef<FWaterTileVS> VertexShader(View.ShaderMap, VsPermutationVector);
-				ValidateShaderParameters(*VertexShader, *PassParameters);
+				ValidateShaderParameters(*VertexShader, PassParameters->VS);
+				ClearUnusedGraphResources(*VertexShader, &PassParameters->VS);
 
 				GraphBuilder.AddPass(
 					RDG_EVENT_NAME("Water Composite %dx%d", View.ViewRect.Width(), View.ViewRect.Height()),
@@ -717,10 +736,9 @@ void FDeferredShadingSceneRenderer::RenderSingleLayerWaterReflections(FRHIComman
 					GraphicsPSOInit.BoundShaderState.PixelShaderRHI = PixelShader->GetPixelShader();
 					SetGraphicsPipelineState(InRHICmdList, GraphicsPSOInit);
 
-					SetShaderParameters(InRHICmdList, *PixelShader, PixelShader->GetPixelShader(), *PassParameters);
-					SetShaderParameters(InRHICmdList, *VertexShader, VertexShader->GetVertexShader(), *PassParameters);
+					SetShaderParameters(InRHICmdList, *VertexShader, VertexShader->GetVertexShader(), PassParameters->VS);
+					SetShaderParameters(InRHICmdList, *PixelShader, PixelShader->GetPixelShader(), PassParameters->PS);
 
-					PassParameters->IndirectDrawParameter->MarkResourceAsUsed();
 					InRHICmdList.DrawPrimitiveIndirect(PassParameters->IndirectDrawParameter->GetIndirectRHICallBuffer(), 0);
 				});
 			}
@@ -741,7 +759,7 @@ void FDeferredShadingSceneRenderer::RenderSingleLayerWaterReflections(FRHIComman
 					GraphicsPSOInit.BlendState = TStaticBlendState<CW_RGB, BO_Add, BF_One, BF_SourceAlpha>::GetRHI();
 
 					SetGraphicsPipelineState(InRHICmdList, GraphicsPSOInit);
-					SetShaderParameters(InRHICmdList, *PixelShader, PixelShader->GetPixelShader(), *PassParameters);
+					SetShaderParameters(InRHICmdList, *PixelShader, PixelShader->GetPixelShader(), PassParameters->PS);
 					FPixelShaderUtils::DrawFullscreenTriangle(InRHICmdList);
 				});
 			}
