Save ~0.2ms GPU on PS4 by moving stencil dither LOD pattern fill to compute, and then running it async alongside GPU particle simulation. Also manually integrated CL 8940152 - Added GRHISupportsDepthUAV flag specifying if the RHI supports depth target unordered access views. Based on this flag, SceneDepthZ now sets TexCreate_UAV if supported (needed for stencil optimization).

diff --git a/Engine/Shaders/Private/DitheredTransitionStencil.usf b/Engine/Shaders/Private/DitheredTransitionStencil.usf
--- a/Engine/Shaders/Private/DitheredTransitionStencil.usf
+++ b/Engine/Shaders/Private/DitheredTransitionStencil.usf
@@ -1,24 +1,58 @@
 // Copyright 1998-2019 Epic Games, Inc. All Rights Reserved.
 
-/*=============================================================================
+/*=============================================================================================================
 	DitheredTransitionStencil.usf: Fills the stencil buffer with the current dithered transition mask.
-=============================================================================*/
+=============================================================================================================*/
 
 #include "Common.ush"
 
 float DitheredTransitionFactor;
 
+// .xy = Viewport offset
+//  .z = Stencil masked value
+//  .w = Stencil cleared value
+uint4 StencilOffsetAndValues;
+
+RWTexture2D<uint> StencilOutput;
+
+half GetDitherValue(float2 InPosition)
+{
+	float RandCos = cos(dot(floor(InPosition), float2(347.83451793f,3343.28371963f)));
+	float RandomVal = frac(RandCos * 1000.f);
+	half RetVal = (DitheredTransitionFactor < 0.f) ?
+		(DitheredTransitionFactor + 1.f > RandomVal) :
+		(DitheredTransitionFactor < RandomVal);
+	return (RetVal - 0.001f);
+}
+
 void Main(
 	FScreenVertexOutput Interpolators)
 {
 	float4 SvPosition = Interpolators.Position;
 	if (abs(DitheredTransitionFactor) > 0.001f)
 	{
-		float RandCos = cos(dot(floor(SvPosition.xy), float2(347.83451793f,3343.28371963f)));
-		float RandomVal = frac(RandCos * 1000.f);
-		half RetVal = (DitheredTransitionFactor < 0.f) ?
-			(DitheredTransitionFactor + 1.f > RandomVal) :
-			(DitheredTransitionFactor < RandomVal);
-		clip(RetVal - 0.001f);
+		clip(GetDitherValue(SvPosition.xy));
 	}
 }
+
+[numthreads(8, 8, 1)]
+void MainCS(uint3 PixelPos : SV_DispatchThreadID)
+{
+	// This should assumes no pre-existing values in the stencil buffer, so any pixels
+	// that pass the transition factor test will set the sandbox bit, otherwise the
+	// default clear value should be exported (skipping the normal stencil target clear).
+	uint StencilValue = (StencilOffsetAndValues.w & 0xFF);
+
+	if (abs(DitheredTransitionFactor) > 0.001f)
+	{
+		if (GetDitherValue(float2(PixelPos.xy)) >= 0.0f)
+		{
+			// The clip() instruction kills any pixels < 0, so we flip the test here
+			// and change from default clear to masked value in the case where the test
+			// passes.
+			StencilValue = (StencilOffsetAndValues.z & 0xFF);
+		}
+	}
+
+	StencilOutput[StencilOffsetAndValues.xy + PixelPos.xy] = StencilValue;
+}
\ No newline at end of file
diff --git a/Engine/Source/Runtime/RHI/Private/RHI.cpp b/Engine/Source/Runtime/RHI/Private/RHI.cpp
--- a/Engine/Source/Runtime/RHI/Private/RHI.cpp
+++ b/Engine/Source/Runtime/RHI/Private/RHI.cpp
@@ -472,6 +472,7 @@ bool GHardwareHiddenSurfaceRemoval = false;
 bool GRHISupportsAsyncTextureCreation = false;
 bool GRHISupportsQuadTopology = false;
 bool GRHISupportsRectTopology = false;
+bool GRHISupportsDepthUAV = false;
 bool GRHISupportsAtomicUInt64 = false;
 bool GSupportsParallelRenderingTasksWithSeparateRHIThread = true;
 bool GRHIThreadNeedsKicking = false;
diff --git a/Engine/Source/Runtime/RHI/Public/DynamicRHI.h b/Engine/Source/Runtime/RHI/Public/DynamicRHI.h
--- a/Engine/Source/Runtime/RHI/Public/DynamicRHI.h
+++ b/Engine/Source/Runtime/RHI/Public/DynamicRHI.h
@@ -944,6 +944,11 @@ public:
 		return nullptr;
 	}
 
+	virtual FTexture2DRHIRef RHIGetStencilTexture(FRHITexture* SourceTextureRHI)
+	{
+		return nullptr;
+	}
+
 	virtual void RHIAliasTextureResources(FRHITexture* DestTexture, FRHITexture* SrcTexture)
 	{
 		checkNoEntry();
@@ -1402,6 +1407,11 @@ FORCEINLINE FTexture2DRHIRef RHIGetFMaskTexture(FRHITexture* SourceTextureRHI)
 	return GDynamicRHI->RHIGetFMaskTexture(SourceTextureRHI);
 }
 
+FORCEINLINE FTexture2DRHIRef RHIGetStencilTexture(FRHITexture* SourceTextureRHI)
+{
+	return GDynamicRHI->RHIGetStencilTexture(SourceTextureRHI);
+}
+
 FORCEINLINE void RHIAdvanceFrameFence()
 {
 	return GDynamicRHI->RHIAdvanceFrameFence();
diff --git a/Engine/Source/Runtime/RHI/Public/RHI.h b/Engine/Source/Runtime/RHI/Public/RHI.h
--- a/Engine/Source/Runtime/RHI/Public/RHI.h
+++ b/Engine/Source/Runtime/RHI/Public/RHI.h
@@ -354,6 +354,9 @@ extern RHI_API bool GSupportsDepthBoundsTest;
 /** True if the RHI and current hardware support a render target write mask */
 extern RHI_API bool GSupportsRenderTargetWriteMask;
 
+/** True if the RHI supports depth target unordered access views. */
+extern RHI_API bool GRHISupportsDepthUAV;
+
 /** True if the RHI and current hardware supports efficient AsyncCompute (by default we assume false and later we can enable this for more hardware) */
 extern RHI_API bool GSupportsEfficientAsyncCompute;
 
diff --git a/Engine/Source/Runtime/RHI/Public/RHIValidation.h b/Engine/Source/Runtime/RHI/Public/RHIValidation.h
--- a/Engine/Source/Runtime/RHI/Public/RHIValidation.h
+++ b/Engine/Source/Runtime/RHI/Public/RHIValidation.h
@@ -965,6 +965,11 @@ public:
 		return RHI->RHIGetFMaskTexture(SourceTextureRHI);
 	}
 
+	virtual FTexture2DRHIRef RHIGetStencilTexture(FRHITexture* SourceTextureRHI) override final
+	{
+		return RHI->RHIGetStencilTexture(SourceTextureRHI);
+	}
+
 	// Only relevant with an RHI thread, this advances the backbuffer for the purpose of GetViewportBackBuffer
 	// FlushType: Thread safe
 	virtual void RHIAdvanceFrameForGetViewportBackBuffer(FRHIViewport* Viewport) override final
diff --git a/Engine/Source/Runtime/Renderer/Private/DeferredShadingRenderer.cpp b/Engine/Source/Runtime/Renderer/Private/DeferredShadingRenderer.cpp
--- a/Engine/Source/Runtime/Renderer/Private/DeferredShadingRenderer.cpp
+++ b/Engine/Source/Runtime/Renderer/Private/DeferredShadingRenderer.cpp
@@ -49,6 +49,15 @@ static TAutoConsoleVariable<int32> CVarStencilForLODDither(
 	TEXT("Forces a full prepass when enabled."),
 	ECVF_RenderThreadSafe | ECVF_ReadOnly);
 
+static TAutoConsoleVariable<int32> CVarStencilLODDitherMode(
+	TEXT("r.StencilLODMode"),
+	2,
+	TEXT("Specifies the dither LOD stencil mode.\n")
+	TEXT(" 0: Graphics pass.\n")
+	TEXT(" 1: Compute pass (on supported platforms).\n")
+	TEXT(" 2: Compute async pass (on supported platforms)."),
+	ECVF_RenderThreadSafe);
+
 TAutoConsoleVariable<int32> CVarCustomDepthOrder(
 	TEXT("r.CustomDepth.Order"),
 	1,	
@@ -966,7 +975,7 @@ void FDeferredShadingSceneRenderer::Render(FRHICommandListImmediate& RHICmdList)
 
 	// Find the visible primitives.
 	RHICmdList.ImmediateFlush(EImmediateFlushType::DispatchToRHIThread);
-
+	
 	bool bDoInitViewAftersPrepass = false;
 	{
 		SCOPED_GPU_STAT(RHICmdList, VisibilityCommands);
@@ -1133,6 +1142,7 @@ void FDeferredShadingSceneRenderer::Render(FRHICommandListImmediate& RHICmdList)
 
 	
 	const bool bIsOcclusionTesting = DoOcclusionQueries(FeatureLevel) && (!bIsWireframe || bIsViewFrozen || bHasViewParent);
+	const bool bNeedsPrePass = NeedsPrePass(this);
 
 	// Dynamic vertex and index buffers need to be committed before rendering.
 	GEngine->GetPreRenderDelegate().Broadcast();
@@ -1150,6 +1160,72 @@ void FDeferredShadingSceneRenderer::Render(FRHICommandListImmediate& RHICmdList)
 		}
 	}
 
+	StencilLODMode = CVarStencilLODDitherMode.GetValueOnRenderThread();
+	if (!GRHISupportsDepthUAV)
+	{
+		// RHI doesn't support depth/stencil UAVs - enforce graphics path
+		StencilLODMode = 0;
+	}
+	else if (StencilLODMode == 2 && !GSupportsEfficientAsyncCompute)
+	{
+		// Async compute is not supported, fall back to compute path (on graphics queue)
+		StencilLODMode = 1;
+	}
+	else if (IsHMDHiddenAreaMaskActive())
+	{
+		// Unsupported mode for compute path - enforce graphics path on VR
+		StencilLODMode = 0;
+	}
+
+	const bool bStencilLODCompute = (StencilLODMode == 1 || StencilLODMode == 2);
+	const bool bStencilLODComputeAsync = StencilLODMode == 2;
+
+	FComputeFenceRHIRef AsyncDitherLODStartFence;
+	FComputeFenceRHIRef AsyncDitherLODEndFence;
+
+	FTexture2DRHIRef StencilTexture;
+	FUnorderedAccessViewRHIRef StencilTextureUAV;
+	if (bStencilLODCompute && bDitheredLODTransitionsUseStencil)
+	{
+		// Either compute pass will happen prior to the prepass, and the
+		// stencil clear will be skipped there.
+
+		StencilTexture = GDynamicRHI->RHIGetStencilTexture(SceneContext.GetSceneDepthSurface());
+		StencilTextureUAV = GDynamicRHI->RHICreateUnorderedAccessView(StencilTexture, 0 /* Mip Level */);
+
+		RHICmdList.TransitionResource(EResourceTransitionAccess::ERWBarrier, EResourceTransitionPipeline::EGfxToCompute, StencilTextureUAV);
+
+		if (bStencilLODComputeAsync)
+		{
+			static FName AsyncDitherLODStartFenceName(TEXT("AsyncDitherLODStartFence"));
+			static FName AsyncDitherLODEndFenceName(TEXT("AsyncDitherLODEndFence"));
+			AsyncDitherLODStartFence = RHICmdList.CreateComputeFence(AsyncDitherLODStartFenceName);
+			AsyncDitherLODEndFence = RHICmdList.CreateComputeFence(AsyncDitherLODEndFenceName);
+
+			FRHIAsyncComputeCommandListImmediate& RHICmdListComputeImmediate = FRHICommandListExecutor::GetImmediateAsyncComputeCommandList();
+
+			RHICmdList.TransitionResource(EResourceTransitionAccess::ERWBarrier, EResourceTransitionPipeline::EComputeToCompute, nullptr, AsyncDitherLODStartFence);
+			RHICmdListComputeImmediate.WaitComputeFence(AsyncDitherLODStartFence);
+
+			PreRenderDitherFill(RHICmdListComputeImmediate, SceneContext, StencilTextureUAV);
+
+			RHICmdListComputeImmediate.TransitionResources(EResourceTransitionAccess::ERWBarrier, EResourceTransitionPipeline::EComputeToCompute, nullptr, 0, AsyncDitherLODEndFence);
+			FRHIAsyncComputeCommandListImmediate::ImmediateDispatch(RHICmdListComputeImmediate);
+		}
+		else
+		{
+			PreRenderDitherFill(RHICmdList, SceneContext, StencilTextureUAV);
+		}
+	}
+
+	// TODO: Move to async compute with proper RDG support.
+	const bool bShouldRenderSkyAtmosphere = ShouldRenderSkyAtmosphere(Scene, ViewFamily.EngineShowFlags);
+	if (bShouldRenderSkyAtmosphere)
+	{
+		// Generate the Sky/Atmosphere look up tables
+		RenderSkyAtmosphereLookUpTables(RHICmdList);
+	}
+
 	// Notify the FX system that the scene is about to be rendered.
 	if (Scene->FXSystem && Views.IsValidIndex(0))
 	{
@@ -1157,7 +1233,12 @@ void FDeferredShadingSceneRenderer::Render(FRHICommandListImmediate& RHICmdList)
 		RHICmdList.SetCurrentStat(GET_STATID(STAT_CLM_FXPreRender));
 		Scene->FXSystem->PreRender(RHICmdList, &Views[0].GlobalDistanceFieldInfo.ParameterData, Views[0].AllowGPUParticleUpdate());
 	}
-	
+
+	if (bStencilLODComputeAsync)
+	{
+		RHICmdList.WaitComputeFence(AsyncDitherLODEndFence);
+	}
+
 	bool bDidAfterTaskWork = false;
 	auto AfterTasksAreStarted = [&bDidAfterTaskWork, bDoInitViewAftersPrepass, this, &RHICmdList, &ILCTaskData, &UpdateViewCustomDataEvents]()
 	{
@@ -1201,20 +1282,12 @@ void FDeferredShadingSceneRenderer::Render(FRHICommandListImmediate& RHICmdList)
 		FTaskGraphInterface::Get().WaitUntilTasksComplete(UpdateViewCustomDataEvents, ENamedThreads::GetRenderThread());
 	}
 
-	// Generate the Sky/Atmosphere look up tables
-	const bool bShouldRenderSkyAtmosphere = ShouldRenderSkyAtmosphere(Scene, ViewFamily.EngineShowFlags);
-	if (bShouldRenderSkyAtmosphere)
-	{
-		RenderSkyAtmosphereLookUpTables(RHICmdList);
-	}
-
 	checkSlow(RHICmdList.IsOutsideRenderPass());
 
 	// The Z-prepass
 
 	// Draw the scene pre-pass / early z pass, populating the scene depth buffer and HiZ
 	GRenderTargetPool.AddPhaseEvent(TEXT("EarlyZPass"));
-	const bool bNeedsPrePass = NeedsPrePass(this);
 	bool bDepthWasCleared;
 	if (bNeedsPrePass)
 	{
@@ -2160,7 +2233,7 @@ void FDeferredShadingSceneRenderer::Render(FRHICommandListImmediate& RHICmdList)
 			}
 			else
 			{
-				GPostProcessing.Process(RHICmdList, Views[ViewIndex], SceneContext.SceneVelocity);
+				GPostProcessing.Process(RHICmdList, Views[ ViewIndex ], SceneContext.SceneVelocity);
 			}
 		}
 
diff --git a/Engine/Source/Runtime/Renderer/Private/DeferredShadingRenderer.h b/Engine/Source/Runtime/Renderer/Private/DeferredShadingRenderer.h
--- a/Engine/Source/Runtime/Renderer/Private/DeferredShadingRenderer.h
+++ b/Engine/Source/Runtime/Renderer/Private/DeferredShadingRenderer.h
@@ -45,6 +45,7 @@ public:
 	EDepthDrawingMode EarlyZPassMode;
 	bool bEarlyZPassMovable;
 	bool bDitheredLODTransitionsUseStencil;
+	int32 StencilLODMode = 0;
 	
 	FComputeFenceRHIRef TranslucencyLightingVolumeClearEndFence;
 
@@ -194,6 +195,9 @@ private:
 	*/
 	bool PreRenderPrePass(FRHICommandListImmediate& RHICmdList);
 
+	void PreRenderDitherFill(FRHIAsyncComputeCommandListImmediate& RHICmdList, FSceneRenderTargets& SceneContext, FRHIUnorderedAccessView* StencilTextureUAV);
+	void PreRenderDitherFill(FRHICommandListImmediate& RHICmdList, FSceneRenderTargets& SceneContext, FRHIUnorderedAccessView* StencilTextureUAV);
+
 	void RenderPrePassEditorPrimitives(FRHICommandList& RHICmdList, const FViewInfo& View, const FMeshPassProcessorRenderState& DrawRenderState, EDepthDrawingMode DepthDrawingMode, bool bRespectUseAsOccluderFlag);
 
 	/**
diff --git a/Engine/Source/Runtime/Renderer/Private/DepthRendering.cpp b/Engine/Source/Runtime/Renderer/Private/DepthRendering.cpp
--- a/Engine/Source/Runtime/Renderer/Private/DepthRendering.cpp
+++ b/Engine/Source/Runtime/Renderer/Private/DepthRendering.cpp
@@ -323,6 +323,7 @@ bool FDeferredShadingSceneRenderer::RenderPrePassViewParallel(const FViewInfo& V
 class FDitheredTransitionStencilPS : public FGlobalShader
 {
 	DECLARE_SHADER_TYPE(FDitheredTransitionStencilPS, Global);
+
 public:
 
 	static bool ShouldCompilePermutation(const FGlobalShaderPermutationParameters& Parameters)
@@ -333,7 +334,7 @@ public:
 	FDitheredTransitionStencilPS(const ShaderMetaType::CompiledShaderInitializerType& Initializer)
 	: FGlobalShader(Initializer)
 	{
-		DitheredTransitionFactorParameter.Bind(Initializer.ParameterMap, TEXT("DitheredTransitionFactor"));
+		DitheredTransitionFactorParameter.Bind(Initializer.ParameterMap, TEXT("DitheredTransitionFactor"), EShaderParameterFlags::SPF_Mandatory);
 	}
 
 	FDitheredTransitionStencilPS()
@@ -357,25 +358,158 @@ public:
 
 	FShaderParameter DitheredTransitionFactorParameter;
 };
-
 IMPLEMENT_SHADER_TYPE(, FDitheredTransitionStencilPS, TEXT("/Engine/Private/DitheredTransitionStencil.usf"), TEXT("Main"), SF_Pixel);
 
+/** A compute shader used to fill the stencil buffer with the current dithered transition mask. */
+class FDitheredTransitionStencilCS : public FGlobalShader
+{
+	DECLARE_SHADER_TYPE(FDitheredTransitionStencilCS, Global);
+
+public:
+
+	static bool ShouldCompilePermutation(const FGlobalShaderPermutationParameters& Parameters)
+	{
+		return IsFeatureLevelSupported(Parameters.Platform, ERHIFeatureLevel::SM5);
+	}
+
+	FDitheredTransitionStencilCS(const ShaderMetaType::CompiledShaderInitializerType& Initializer)
+		: FGlobalShader(Initializer)
+	{
+		DitheredTransitionFactorParameter.Bind(Initializer.ParameterMap, TEXT("DitheredTransitionFactor"), EShaderParameterFlags::SPF_Mandatory);
+		StencilOffsetAndValuesParameter.Bind(Initializer.ParameterMap, TEXT("StencilOffsetAndValues"), EShaderParameterFlags::SPF_Mandatory);
+		StencilOutputParameter.Bind(Initializer.ParameterMap, TEXT("StencilOutput"), EShaderParameterFlags::SPF_Mandatory);
+	}
+
+	FDitheredTransitionStencilCS()
+	{
+	}
+
+	template <typename TRHICmdList>
+	void SetParameters(TRHICmdList& RHICmdList, const FSceneView& View, FRHIUnorderedAccessView* StencilOutputUAV, FIntPoint BufferSizeXY, FIntPoint ViewOffsetXY, uint32 StencilValue)
+	{
+		FRHIComputeShader* ComputeShader = GetComputeShader();
+
+		FGlobalShader::SetParameters<FViewUniformShaderParameters>(RHICmdList, ComputeShader, View.ViewUniformBuffer);
+
+		const float DitherFactor = View.GetTemporalLODTransition();
+		SetShaderValue(RHICmdList, ComputeShader, DitheredTransitionFactorParameter, DitherFactor);
+
+		const uint32 MaskedValue = (StencilValue & 0xFF);
+		const uint32 ClearedValue = 0;
+
+		const FIntVector4 StencilOffsetAndValues(
+			ViewOffsetXY.X,
+			ViewOffsetXY.Y,
+			int32(MaskedValue),
+			int32(ClearedValue));
+
+		SetShaderValue(RHICmdList, ComputeShader, StencilOffsetAndValuesParameter, StencilOffsetAndValues);
+		SetUAVParameter(RHICmdList, ComputeShader, StencilOutputParameter, StencilOutputUAV);
+	}
+
+	template <typename TRHICmdList>
+	void UnsetParameters(TRHICmdList& RHICmdList)
+	{
+		FRHIComputeShader* ComputeShader = GetComputeShader();
+		if (StencilOutputParameter.IsBound())
+		{
+			RHICmdList.SetUAVParameter(ComputeShader, StencilOutputParameter.GetBaseIndex(), nullptr);
+		}
+	}
+
+	virtual bool Serialize(FArchive& Ar) override
+	{
+		bool bShaderHasOutdatedParameters = FGlobalShader::Serialize(Ar);
+		Ar << DitheredTransitionFactorParameter << StencilOffsetAndValuesParameter << StencilOutputParameter;
+		return bShaderHasOutdatedParameters;
+	}
+
+	FShaderParameter DitheredTransitionFactorParameter;
+	FShaderParameter StencilOffsetAndValuesParameter;
+	FShaderResourceParameter StencilOutputParameter;
+};
+IMPLEMENT_SHADER_TYPE(, FDitheredTransitionStencilCS, TEXT("/Engine/Private/DitheredTransitionStencil.usf"), TEXT("MainCS"), SF_Compute);
+
 /** Possibly do the FX prerender and setup the prepass*/
 bool FDeferredShadingSceneRenderer::PreRenderPrePass(FRHICommandListImmediate& RHICmdList)
 {
-	SCOPED_GPU_MASK(RHICmdList, FRHIGPUMask::All()); // Required otherwise emulatestereo gets broken.
+	SCOPED_GPU_MASK(RHICmdList, FRHIGPUMask::All()); // Required otherwise emulate stereo gets broken.
 
 	RHICmdList.SetCurrentStat(GET_STATID(STAT_CLM_PrePass));
+
 	// RenderPrePassHMD clears the depth buffer. If this changes we must change RenderPrePass to maintain the correct behavior!
 	bool bDepthWasCleared = RenderPrePassHMD(RHICmdList);
 
 	FSceneRenderTargets& SceneContext = FSceneRenderTargets::Get(RHICmdList);
 
-	SceneContext.BeginRenderingPrePass(RHICmdList, !bDepthWasCleared);
+	// Both compute approaches run earlier, so skip clearing stencil here, just load existing.
+	const bool bNoStencilClear = bDitheredLODTransitionsUseStencil && (StencilLODMode == 1 || StencilLODMode == 2);
+
+	SceneContext.BeginRenderingPrePass(RHICmdList, !bDepthWasCleared, !bNoStencilClear);
 	bDepthWasCleared = true;
 
-	// Dithered transition stencil mask fill
-	if (bDitheredLODTransitionsUseStencil)
+	// Dithered transition stencil mask fill (graphics path)
+	if (bDitheredLODTransitionsUseStencil && StencilLODMode == 0)
+	{
+		PreRenderDitherFill(RHICmdList, SceneContext, nullptr);
+	}
+
+	// Need to close the render pass here since we may call BeginRenderingPrePass later
+	RHICmdList.EndRenderPass();
+
+	return bDepthWasCleared;
+}
+
+void FDeferredShadingSceneRenderer::PreRenderDitherFill(FRHIAsyncComputeCommandListImmediate& RHICmdList, FSceneRenderTargets& SceneContext, FRHIUnorderedAccessView* StencilTextureUAV)
+{
+	SCOPED_GPU_EVENT(RHICmdList, DitheredStencilPrePass);
+
+	FIntPoint BufferSizeXY = SceneContext.GetBufferSizeXY();
+
+	for (int32 ViewIndex = 0; ViewIndex < Views.Num(); ++ViewIndex)
+	{
+		SCOPED_CONDITIONAL_GPU_EVENTF(RHICmdList, EventView, Views.Num() > 1, TEXT("View%d"), ViewIndex);
+
+		FViewInfo& View = Views[ViewIndex];
+
+		TShaderMapRef<FDitheredTransitionStencilCS> ComputeShader(View.ShaderMap);
+		RHICmdList.SetComputeShader(ComputeShader->GetComputeShader());
+		ComputeShader->SetParameters(RHICmdList, View, StencilTextureUAV, BufferSizeXY, View.ViewRect.Min, STENCIL_SANDBOX_MASK);
+		const int32 SubWidth = FMath::Min(BufferSizeXY.X, View.ViewRect.Width());
+		const int32 SubHeight = FMath::Min(BufferSizeXY.Y, View.ViewRect.Height());
+		check(SubWidth > 0 && SubHeight > 0);
+
+		DispatchComputeShader(RHICmdList, *ComputeShader, FMath::DivideAndRoundUp(SubWidth, 8), FMath::DivideAndRoundUp(SubHeight, 8), 1);
+		ComputeShader->UnsetParameters(RHICmdList);
+	}
+}
+
+
+void FDeferredShadingSceneRenderer::PreRenderDitherFill(FRHICommandListImmediate& RHICmdList, FSceneRenderTargets& SceneContext, FRHIUnorderedAccessView* StencilTextureUAV)
+{
+	SCOPED_DRAW_EVENT(RHICmdList, DitheredStencilPrePass);
+
+	FIntPoint BufferSizeXY = SceneContext.GetBufferSizeXY();
+	if (StencilLODMode == 1 || StencilLODMode == 2)
+	{
+		for (int32 ViewIndex = 0; ViewIndex < Views.Num(); ++ViewIndex)
+		{
+			SCOPED_CONDITIONAL_DRAW_EVENTF(RHICmdList, EventView, Views.Num() > 1, TEXT("View%d"), ViewIndex);
+
+			FViewInfo& View = Views[ViewIndex];
+
+			TShaderMapRef<FDitheredTransitionStencilCS> ComputeShader(View.ShaderMap);
+			RHICmdList.SetComputeShader(ComputeShader->GetComputeShader());
+			ComputeShader->SetParameters(RHICmdList, View, StencilTextureUAV, BufferSizeXY, View.ViewRect.Min, STENCIL_SANDBOX_MASK);
+			const int32 SubWidth = FMath::Min(BufferSizeXY.X, View.ViewRect.Width());
+			const int32 SubHeight = FMath::Min(BufferSizeXY.Y, View.ViewRect.Height());
+			check(SubWidth > 0 && SubHeight > 0);
+
+			DispatchComputeShader(RHICmdList, *ComputeShader, FMath::DivideAndRoundUp(SubWidth, 8), FMath::DivideAndRoundUp(SubHeight, 8), 1);
+			ComputeShader->UnsetParameters(RHICmdList);
+		}
+	}
+	else
 	{
 		FGraphicsPipelineStateInitializer GraphicsPSOInit;
 		RHICmdList.ApplyCachedRenderTargets(GraphicsPSOInit);
@@ -386,14 +520,13 @@ bool FDeferredShadingSceneRenderer::PreRenderPrePass(FRHICommandListImmediate& R
 			false, CF_Always, SO_Keep, SO_Keep, SO_Keep,
 			STENCIL_SANDBOX_MASK, STENCIL_SANDBOX_MASK>::GetRHI();
 
-		SCOPED_DRAW_EVENT(RHICmdList, DitheredStencilPrePass);
-		FIntPoint BufferSizeXY = SceneContext.GetBufferSizeXY();
-
 		for (int32 ViewIndex = 0; ViewIndex < Views.Num(); ++ViewIndex)
 		{
+			FViewInfo& View = Views[ViewIndex];
+
+			SCOPED_GPU_MASK(RHICmdList, View.GPUMask);
 			SCOPED_CONDITIONAL_DRAW_EVENTF(RHICmdList, EventView, Views.Num() > 1, TEXT("View%d"), ViewIndex);
 
-			FViewInfo& View = Views[ViewIndex];
 			RHICmdList.SetViewport(View.ViewRect.Min.X, View.ViewRect.Min.Y, 0.0f, View.ViewRect.Max.X, View.ViewRect.Max.Y, 1.0f);
 
 			// Set shaders, states
@@ -423,10 +556,6 @@ bool FDeferredShadingSceneRenderer::PreRenderPrePass(FRHICommandListImmediate& R
 				EDRF_UseTriangleOptimization);
 		}
 	}
-	// Need to close the renderpass here since we may call BeginRenderingPrePass later
-	RHICmdList.EndRenderPass();
-
-	return bDepthWasCleared;
 }
 
 void FDeferredShadingSceneRenderer::RenderPrePassEditorPrimitives(FRHICommandList& RHICmdList, const FViewInfo& View, const FMeshPassProcessorRenderState& DrawRenderState, EDepthDrawingMode DepthDrawingMode, bool bRespectUseAsOccluderFlag) 
diff --git a/Engine/Source/Runtime/Renderer/Private/PostProcess/SceneRenderTargets.cpp b/Engine/Source/Runtime/Renderer/Private/PostProcess/SceneRenderTargets.cpp
--- a/Engine/Source/Runtime/Renderer/Private/PostProcess/SceneRenderTargets.cpp
+++ b/Engine/Source/Runtime/Renderer/Private/PostProcess/SceneRenderTargets.cpp
@@ -1454,7 +1454,7 @@ void FSceneRenderTargets::FinishRenderingCustomDepth(FRHICommandListImmediate& R
 	bCustomDepthIsValid = true;
 }
 
-void FSceneRenderTargets::BeginRenderingPrePass(FRHICommandList& RHICmdList, bool bPerformClear)
+void FSceneRenderTargets::BeginRenderingPrePass(FRHICommandList& RHICmdList, bool bPerformClear, bool bStencilClear)
 {
 	check(RHICmdList.IsOutsideRenderPass());
 
@@ -1471,8 +1471,9 @@ void FSceneRenderTargets::BeginRenderingPrePass(FRHICommandList& RHICmdList, boo
 	{				
 		// Clear the depth buffer.
 		// Note, this is a reversed Z depth surface, so 0.0f is the far plane.
-		RPInfo.DepthStencilRenderTarget.Action = MakeDepthStencilTargetActions(ERenderTargetActions::Clear_Store, ERenderTargetActions::Clear_Store);
-		bSceneDepthCleared = true;	
+		ERenderTargetActions StencilAction = bStencilClear ? ERenderTargetActions::Clear_Store : ERenderTargetActions::Load_Store;
+		RPInfo.DepthStencilRenderTarget.Action = MakeDepthStencilTargetActions(ERenderTargetActions::Clear_Store, StencilAction);
+		bSceneDepthCleared = true;
 	}
 	else
 	{
@@ -2223,8 +2224,11 @@ void FSceneRenderTargets::AllocateCommonDepthTargets(FRHICommandList& RHICmdList
 		FTexture2DRHIRef DepthTex, SRTex;
 		bHMDAllocatedDepthTarget = StereoRenderTargetManager && StereoRenderTargetManager->AllocateDepthTexture(0, BufferSize.X, BufferSize.Y, PF_X24_G8, 0, TexCreate_None, TexCreate_DepthStencilTargetable, DepthTex, SRTex, GetNumSceneColorMSAASamples(CurrentFeatureLevel));
 
+		// Allow UAV depth?
+		const uint32 DepthUAVFlag = GRHISupportsDepthUAV ? TexCreate_UAV : 0;
+
 		// Create a texture to store the resolved scene depth, and a render-targetable surface to hold the unresolved scene depth.
-		FPooledRenderTargetDesc Desc(FPooledRenderTargetDesc::Create2DDesc(BufferSize, PF_DepthStencil, DefaultDepthClear, TexCreate_None, TexCreate_DepthStencilTargetable | TexCreate_ShaderResource | TexCreate_InputAttachmentRead, false));
+		FPooledRenderTargetDesc Desc(FPooledRenderTargetDesc::Create2DDesc(BufferSize, PF_DepthStencil, DefaultDepthClear, TexCreate_None, TexCreate_DepthStencilTargetable | TexCreate_ShaderResource | TexCreate_InputAttachmentRead | DepthUAVFlag, false));
 		Desc.NumSamples = GetNumSceneColorMSAASamples(CurrentFeatureLevel);
 		Desc.Flags |= GFastVRamConfig.SceneDepth;
 
diff --git a/Engine/Source/Runtime/Renderer/Private/PostProcess/SceneRenderTargets.h b/Engine/Source/Runtime/Renderer/Private/PostProcess/SceneRenderTargets.h
--- a/Engine/Source/Runtime/Renderer/Private/PostProcess/SceneRenderTargets.h
+++ b/Engine/Source/Runtime/Renderer/Private/PostProcess/SceneRenderTargets.h
@@ -268,7 +268,7 @@ public:
 	void ResolveSceneDepthTexture(FRHICommandList& RHICmdList, const FResolveRect& ResolveRect);
 	void ResolveSceneDepthToAuxiliaryTexture(FRHICommandList& RHICmdList);
 
-	void BeginRenderingPrePass(FRHICommandList& RHICmdList, bool bPerformClear);
+	void BeginRenderingPrePass(FRHICommandList& RHICmdList, bool bPerformClear, bool bStencilClear = true);
 	void FinishRenderingPrePass(FRHICommandListImmediate& RHICmdList);
 
 	void BeginRenderingSceneAlphaCopy(FRHICommandListImmediate& RHICmdList);
