Better CSV stat to track percentage of virtual texture physical space is being used

diff --git a/Engine/Source/Runtime/Renderer/Private/VT/TexturePagePool.cpp b/Engine/Source/Runtime/Renderer/Private/VT/TexturePagePool.cpp
--- a/Engine/Source/Runtime/Renderer/Private/VT/TexturePagePool.cpp
+++ b/Engine/Source/Runtime/Renderer/Private/VT/TexturePagePool.cpp
@@ -75,7 +75,6 @@ void FTexturePagePool::UnmapAllPagesForSpace(FVirtualTextureSystem* System, uint
 
 void FTexturePagePool::EvictPages(FVirtualTextureSystem* System, const FVirtualTextureProducerHandle& ProducerHandle)
 {
-
 	for (uint32 pAddress = 0u; pAddress < NumPages; ++pAddress)
 	{
 		const FPageEntry& PageEntry = Pages[pAddress];
@@ -281,6 +280,29 @@ void FTexturePagePool::UpdateUsage(uint32 Frame, uint16 pAddress)
 	}
 }
 
+uint32 FTexturePagePool::GetNumVisiblePages(uint32 Frame) const
+{
+	uint32 Count = 0;
+	for (uint32 i = 0; i < NumPages; ++i)
+	{
+		if (FreeHeap.IsPresent(i))
+		{
+			uint32 Key = FreeHeap.GetKey(i);
+			if ((Key >> 4) > Frame)
+			{
+				Count ++;
+			}
+		}
+		else
+		{
+			// Consider all locked pages as visible
+			Count++;
+		}
+	}
+
+	return Count;
+}
+
 void FTexturePagePool::MapPage(FVirtualTextureSpace* Space, FVirtualTexturePhysicalSpace* PhysicalSpace, uint8 PageTableLayerIndex, uint8 vLogSize, uint32 vAddress, uint8 vLevel, uint16 pAddress)
 {
 	check(pAddress < NumPages);
diff --git a/Engine/Source/Runtime/Renderer/Private/VT/TexturePagePool.h b/Engine/Source/Runtime/Renderer/Private/VT/TexturePagePool.h
--- a/Engine/Source/Runtime/Renderer/Private/VT/TexturePagePool.h
+++ b/Engine/Source/Runtime/Renderer/Private/VT/TexturePagePool.h
@@ -119,6 +119,11 @@ public:
 	 */
 	void		UpdateUsage(uint32 Frame, uint16 pAddress);
 
+	/**
+	 * Returns the number of pages marked as used since a given frame. 
+	 */
+	uint32		GetNumVisiblePages(uint32 Frame) const;
+
 	/**
 	* Map the physical address to a specific virtual address.
 	*/
diff --git a/Engine/Source/Runtime/Renderer/Private/VT/VirtualTextureSystem.cpp b/Engine/Source/Runtime/Renderer/Private/VT/VirtualTextureSystem.cpp
--- a/Engine/Source/Runtime/Renderer/Private/VT/VirtualTextureSystem.cpp
+++ b/Engine/Source/Runtime/Renderer/Private/VT/VirtualTextureSystem.cpp
@@ -18,27 +18,26 @@
 #include "VT/VirtualTextureScalability.h"
 #include "VT/VirtualTextureSpace.h"
 
-DECLARE_CYCLE_STAT(TEXT("Feedback Analysis"), STAT_FeedbackAnalysis, STATGROUP_VirtualTexturing);
 DECLARE_CYCLE_STAT(TEXT("VirtualTextureSystem Update"), STAT_VirtualTextureSystem_Update, STATGROUP_VirtualTexturing);
 
-DECLARE_CYCLE_STAT(TEXT("Page Table Updates"), STAT_PageTableUpdates, STATGROUP_VirtualTexturing);
 DECLARE_CYCLE_STAT(TEXT("Gather Requests"), STAT_ProcessRequests_Gather, STATGROUP_VirtualTexturing);
 DECLARE_CYCLE_STAT(TEXT("Sort Requests"), STAT_ProcessRequests_Sort, STATGROUP_VirtualTexturing);
 DECLARE_CYCLE_STAT(TEXT("Submit Requests"), STAT_ProcessRequests_Submit, STATGROUP_VirtualTexturing);
 DECLARE_CYCLE_STAT(TEXT("Map Requests"), STAT_ProcessRequests_Map, STATGROUP_VirtualTexturing);
 DECLARE_CYCLE_STAT(TEXT("Finalize Requests"), STAT_ProcessRequests_Finalize, STATGROUP_VirtualTexturing);
-
 DECLARE_CYCLE_STAT(TEXT("Merge Unique Pages"), STAT_ProcessRequests_MergePages, STATGROUP_VirtualTexturing);
 DECLARE_CYCLE_STAT(TEXT("Merge Requests"), STAT_ProcessRequests_MergeRequests, STATGROUP_VirtualTexturing);
 DECLARE_CYCLE_STAT(TEXT("Submit Tasks"), STAT_ProcessRequests_SubmitTasks, STATGROUP_VirtualTexturing);
 
+DECLARE_CYCLE_STAT(TEXT("Feedback Analysis"), STAT_FeedbackAnalysis, STATGROUP_VirtualTexturing);
+DECLARE_CYCLE_STAT(TEXT("Page Table Updates"), STAT_PageTableUpdates, STATGROUP_VirtualTexturing);
 DECLARE_CYCLE_STAT(TEXT("Flush Cache"), STAT_FlushCache, STATGROUP_VirtualTexturing);
+DECLARE_CYCLE_STAT(TEXT("Update Stats"), STAT_UpdateStats, STATGROUP_VirtualTexturing);
 
 DECLARE_DWORD_COUNTER_STAT(TEXT("Num page visible"), STAT_NumPageVisible, STATGROUP_VirtualTexturing);
 DECLARE_DWORD_COUNTER_STAT(TEXT("Num page visible resident"), STAT_NumPageVisibleResident, STATGROUP_VirtualTexturing);
 DECLARE_DWORD_COUNTER_STAT(TEXT("Num page visible not resident"), STAT_NumPageVisibleNotResident, STATGROUP_VirtualTexturing);
 DECLARE_DWORD_COUNTER_STAT(TEXT("Num page prefetch"), STAT_NumPagePrefetch, STATGROUP_VirtualTexturing);
-DECLARE_DWORD_COUNTER_STAT(TEXT("Num page free"), STAT_NumPageFree, STATGROUP_VirtualTexturing);
 DECLARE_DWORD_COUNTER_STAT(TEXT("Num page update"), STAT_NumPageUpdate, STATGROUP_VirtualTexturing);
 DECLARE_DWORD_COUNTER_STAT(TEXT("Num continuous page update"), STAT_NumContinuousPageUpdate, STATGROUP_VirtualTexturing);
 
@@ -1105,6 +1104,8 @@ void FVirtualTextureSystem::Update(FRHICommandListImmediate& RHICmdList, ERHIFea
 	// Submit the merged requests
 	SubmitRequests(RHICmdList, FeatureLevel, MemStack, MergedRequestList, true);
 
+	UpdateCSVStats();
+
 	ReleasePendingSpaces();
 }
 
@@ -1599,7 +1600,6 @@ void FVirtualTextureSystem::GatherRequestsTask(const FGatherRequestsParameters&
 		}
 	}
 
-	uint32 TotalAllocatedPages = 0u;
 	for (uint32 PhysicalSpaceID = 0u; PhysicalSpaceID < (uint32)PhysicalSpaces.Num(); ++PhysicalSpaceID)
 	{
 		if (PhysicalSpaces[PhysicalSpaceID] == nullptr)
@@ -1608,8 +1608,6 @@ void FVirtualTextureSystem::GatherRequestsTask(const FGatherRequestsParameters&
 		}
 
 		FVirtualTexturePhysicalSpace* RESTRICT PhysicalSpace = GetPhysicalSpace(PhysicalSpaceID);
-		TotalAllocatedPages += PhysicalSpace->GetNumTiles();
-
 		FPageUpdateBuffer& RESTRICT Buffer = PageUpdateBuffers[PhysicalSpaceID];
 
 		if (Buffer.WorkingSetSize > 0u)
@@ -1644,10 +1642,31 @@ void FVirtualTextureSystem::GatherRequestsTask(const FGatherRequestsParameters&
 	INC_DWORD_STAT_BY(STAT_NumPageVisibleResident, NumResidentPages);
 	INC_DWORD_STAT_BY(STAT_NumPageVisibleNotResident, NumNonResidentPages);
 	INC_DWORD_STAT_BY(STAT_NumPagePrefetch, NumPrefetchPages);
-	INC_DWORD_STAT_BY(STAT_NumPageFree, TotalAllocatedPages - NumResidentPages);
+}
 
-	const float PhysicalPoolUsage = TotalAllocatedPages > 0 ? (float)NumResidentPages / (float)TotalAllocatedPages : 0.f;
+
+void FVirtualTextureSystem::UpdateCSVStats() const
+{
+#if CSV_PROFILER
+	SCOPE_CYCLE_COUNTER(STAT_UpdateStats);
+
+	uint32 TotalPages = 0;
+	uint32 CurrentPages = 0;
+	const uint32 AgeTolerance = 5; // Include some tolerance/smoothing for previous frames
+	for (int i = 0; i < PhysicalSpaces.Num(); ++i)
+	{
+		FVirtualTexturePhysicalSpace* PhysicalSpace = PhysicalSpaces[i].Get();
+		if (PhysicalSpace)
+		{
+			FTexturePagePool const& PagePool = PhysicalSpace->GetPagePool();
+			TotalPages += PagePool.GetNumPages();
+			CurrentPages += PagePool.GetNumVisiblePages(Frame > AgeTolerance ? Frame - AgeTolerance : 0);
+		}
+	}
+
+	const float PhysicalPoolUsage = TotalPages > 0 ? (float)CurrentPages / (float)TotalPages : 0.f;
 	CSV_CUSTOM_STAT_GLOBAL(VirtualTexturePageUsage, PhysicalPoolUsage, ECsvCustomStatOp::Set);
+#endif
 }
 
 void FVirtualTextureSystem::SubmitRequestsFromLocalTileList(const TSet<FVirtualTextureLocalTile>& LocalTileList, EVTProducePageFlags Flags, FRHICommandListImmediate& RHICmdList, ERHIFeatureLevel::Type FeatureLevel)
diff --git a/Engine/Source/Runtime/Renderer/Private/VT/VirtualTextureSystem.h b/Engine/Source/Runtime/Renderer/Private/VT/VirtualTextureSystem.h
--- a/Engine/Source/Runtime/Renderer/Private/VT/VirtualTextureSystem.h
+++ b/Engine/Source/Runtime/Renderer/Private/VT/VirtualTextureSystem.h
@@ -91,6 +91,8 @@ private:
 	void GatherRequestsTask(const FGatherRequestsParameters& Parameters);
 	void FeedbackAnalysisTask(const FFeedbackAnalysisParameters& Parameters);
 
+	void UpdateCSVStats() const;
+
 	uint32	Frame;
 
 	static const uint32 MaxSpaces = 16;
