Fix for rendering of 16bit height on platforms where output RT needs specifying in shader compile
Also touch the relevant shader to force recompile

diff --git a/Engine/Source/Runtime/Renderer/Private/VT/RuntimeVirtualTextureRender.cpp b/Engine/Source/Runtime/Renderer/Private/VT/RuntimeVirtualTextureRender.cpp
--- a/Engine/Source/Runtime/Renderer/Private/VT/RuntimeVirtualTextureRender.cpp
+++ b/Engine/Source/Runtime/Renderer/Private/VT/RuntimeVirtualTextureRender.cpp
@@ -120,7 +120,7 @@ namespace RuntimeVirtualTexture
 		static void ModifyCompilationEnvironment(FShaderCompilerEnvironment& OutEnvironment)
 		{
 			OutEnvironment.SetDefine(TEXT("OUT_WORLDHEIGHT"), 1);
-			OutEnvironment.SetRenderTargetOutputFormat(0, PF_G16);
+			OutEnvironment.SetRenderTargetOutputFormat(0, PF_R32_FLOAT);
 		}
 
 		static FRHIBlendState* GetBlendState()
