Water Shading - fixed mismatching structured buffer and buffer bindings.
This fixes crashes on Mac and following -d3ddebug error when dispaching FWaterTileCategorisationCS:
D3D11 ERROR: ID3D11DeviceContext::Dispatch: The Unordered Access View in slot 0 of the Compute Shader unit is not a Structured Buffer while the shader expects a Buffer with a structure stride of 4. This mismatch is invalid if the shader actually uses the view (e.g. it is not skipped due to shader code branching).

diff --git a/Engine/Shaders/Private/SingleLayerWaterComposite.usf b/Engine/Shaders/Private/SingleLayerWaterComposite.usf
--- a/Engine/Shaders/Private/SingleLayerWaterComposite.usf
+++ b/Engine/Shaders/Private/SingleLayerWaterComposite.usf
@@ -145,8 +145,8 @@ SingleLayerWaterCompositeOutput SingleLayerWaterComposite(float2 BufferUV, float
 
 uint VertexCountPerInstanceIndirect;
 
-RWStructuredBuffer<uint> DispatchIndirectDataUAV;
-RWStructuredBuffer<uint> WaterTileListDataUAV;
+RWBuffer<uint> DispatchIndirectDataUAV;
+RWBuffer<uint> WaterTileListDataUAV;
 
 groupshared bool ContainsWater[WORK_TILE_SIZE*WORK_TILE_SIZE];
 
@@ -245,7 +245,7 @@ void WaterTileCatergorisationCS(uint3 ThreadId : SV_DispatchThreadID, uint3 Grou
 #elif defined(TILE_VERTEX_SHADER)
 
 
-StructuredBuffer<uint> TileListData; //
+Buffer<uint> TileListData;
 
 void WaterTileVS(
 	in uint InstanceId : SV_InstanceID,
diff --git a/Engine/Source/Runtime/Renderer/Private/SingleLayerWaterRendering.cpp b/Engine/Source/Runtime/Renderer/Private/SingleLayerWaterRendering.cpp
--- a/Engine/Source/Runtime/Renderer/Private/SingleLayerWaterRendering.cpp
+++ b/Engine/Source/Runtime/Renderer/Private/SingleLayerWaterRendering.cpp
@@ -339,7 +339,7 @@ class FWaterTileVS : public FGlobalShader
 
 	BEGIN_SHADER_PARAMETER_STRUCT(FParameters, )
 		SHADER_PARAMETER_STRUCT_REF(FViewUniformShaderParameters, ViewUniformBuffer)
-		SHADER_PARAMETER_RDG_BUFFER_SRV(StructuredBuffer<uint>, TileListData)
+		SHADER_PARAMETER_RDG_BUFFER_SRV(Buffer<uint>, TileListData)
 	END_SHADER_PARAMETER_STRUCT()
 
 	static FPermutationDomain RemapPermutation(FPermutationDomain PermutationVector)
@@ -621,10 +621,9 @@ void FDeferredShadingSceneRenderer::RenderSingleLayerWaterReflections(FRHIComman
 		{
 			TiledScreenSpaceReflection.DispatchIndirectParametersBuffer = GraphBuilder.CreateBuffer(FRDGBufferDesc::CreateIndirectDesc<FRHIDrawIndirectParameters>(), TEXT("WaterIndirectDrawParameters"));
 			TiledScreenSpaceReflection.DispatchIndirectParametersBufferUAV = GraphBuilder.CreateUAV(TiledScreenSpaceReflection.DispatchIndirectParametersBuffer);
-			FRDGBufferDesc TileListStructuredBufferDesc = FRDGBufferDesc::CreateStructuredDesc(4, TiledViewRes.X * TiledViewRes.Y); // one uint32 element per tile
-			TiledScreenSpaceReflection.TileListDataBuffer = GraphBuilder.CreateBuffer(TileListStructuredBufferDesc, TEXT("TileListDataBuffer"));
-			TiledScreenSpaceReflection.TileListStructureBufferUAV = GraphBuilder.CreateUAV(TiledScreenSpaceReflection.TileListDataBuffer);
-			TiledScreenSpaceReflection.TileListStructureBufferSRV = GraphBuilder.CreateSRV(TiledScreenSpaceReflection.TileListDataBuffer);
+			TiledScreenSpaceReflection.TileListDataBuffer = GraphBuilder.CreateBuffer(FRDGBufferDesc::CreateBufferDesc(sizeof(uint32), TiledViewRes.X * TiledViewRes.Y), TEXT("TileListDataBuffer"));
+			TiledScreenSpaceReflection.TileListStructureBufferUAV = GraphBuilder.CreateUAV(TiledScreenSpaceReflection.TileListDataBuffer, PF_R32_UINT);
+			TiledScreenSpaceReflection.TileListStructureBufferSRV = GraphBuilder.CreateSRV(TiledScreenSpaceReflection.TileListDataBuffer, PF_R32_UINT);
 
 			// Clear DispatchIndirectParametersBuffer
 			AddClearUAVPass(GraphBuilder, TiledScreenSpaceReflection.DispatchIndirectParametersBufferUAV, 0);
