Fix incoming changes in light of CL# 9306968
#rnx

diff --git a/Engine/Source/Runtime/Renderer/Private/VT/VirtualTextureSystem.cpp b/Engine/Source/Runtime/Renderer/Private/VT/VirtualTextureSystem.cpp
--- a/Engine/Source/Runtime/Renderer/Private/VT/VirtualTextureSystem.cpp
+++ b/Engine/Source/Runtime/Renderer/Private/VT/VirtualTextureSystem.cpp
@@ -216,7 +216,7 @@ FVirtualTextureSystem::~FVirtualTextureSystem()
 			BeginReleaseResource(Space);
 		}
 	}
-	for(int i = 0; i < PhysicalSpaces.Num(); ++i)
+	for(int32 i = 0; i < PhysicalSpaces.Num(); ++i)
 	{
 		FVirtualTexturePhysicalSpace* PhysicalSpace = PhysicalSpaces[i];
 		if (PhysicalSpace)
@@ -288,7 +288,7 @@ void FVirtualTextureSystem::DumpFromConsole()
 
 void FVirtualTextureSystem::ListPhysicalPoolsFromConsole()
 {
-	for(int i = 0; i < PhysicalSpaces.Num(); ++i)
+	for(int32 i = 0; i < PhysicalSpaces.Num(); ++i)
 	{
 		if (PhysicalSpaces[i])
 		{
@@ -608,7 +608,7 @@ FVirtualTexturePhysicalSpace* FVirtualTextureSystem::AcquirePhysicalSpace(const
 {
 	LLM_SCOPE(ELLMTag::VirtualTextureSystem);
 
-	for (int i = 0; i < PhysicalSpaces.Num(); ++i)
+	for (int32 i = 0; i < PhysicalSpaces.Num(); ++i)
 	{
 		FVirtualTexturePhysicalSpace* PhysicalSpace = PhysicalSpaces[i];
 		if (PhysicalSpace && PhysicalSpace->GetDescription() == InDesc)
@@ -620,7 +620,7 @@ FVirtualTexturePhysicalSpace* FVirtualTextureSystem::AcquirePhysicalSpace(const
 	uint32 ID = PhysicalSpaces.Num();
 	check(ID <= 0x0fff);
 
-	for (int i = 0; i < PhysicalSpaces.Num(); ++i)
+	for (int32 i = 0; i < PhysicalSpaces.Num(); ++i)
 	{
 		if (!PhysicalSpaces[i])
 		{
@@ -878,7 +878,7 @@ void FVirtualTextureSystem::Update(FRHICommandListImmediate& RHICmdList, ERHIFea
 		SCOPE_CYCLE_COUNTER(STAT_FlushCache);
 		INC_DWORD_STAT_BY(STAT_NumFlushCache, 1);
 
-		for (int i = 0; i < PhysicalSpaces.Num(); ++i)
+		for (int32 i = 0; i < PhysicalSpaces.Num(); ++i)
 		{
 			FVirtualTexturePhysicalSpace* PhysicalSpace = PhysicalSpaces[i];
 			if (PhysicalSpace)
@@ -1653,9 +1653,9 @@ void FVirtualTextureSystem::UpdateCSVStats() const
 	uint32 TotalPages = 0;
 	uint32 CurrentPages = 0;
 	const uint32 AgeTolerance = 5; // Include some tolerance/smoothing for previous frames
-	for (int i = 0; i < PhysicalSpaces.Num(); ++i)
+	for (int32 i = 0; i < PhysicalSpaces.Num(); ++i)
 	{
-		FVirtualTexturePhysicalSpace* PhysicalSpace = PhysicalSpaces[i].Get();
+		FVirtualTexturePhysicalSpace* PhysicalSpace = PhysicalSpaces[i];
 		if (PhysicalSpace)
 		{
 			FTexturePagePool const& PagePool = PhysicalSpace->GetPagePool();
