Changing r.RayTracing.GlobalIllumination.DiffuseThreshold default from 0.05 to 0.01.

diff --git a/Engine/Source/Runtime/Renderer/Private/RayTracing/RayTracingGlobalIllumination.cpp b/Engine/Source/Runtime/Renderer/Private/RayTracing/RayTracingGlobalIllumination.cpp
--- a/Engine/Source/Runtime/Renderer/Private/RayTracing/RayTracingGlobalIllumination.cpp
+++ b/Engine/Source/Runtime/Renderer/Private/RayTracing/RayTracingGlobalIllumination.cpp
@@ -61,7 +61,7 @@ static FAutoConsoleVariableRef CVarRayTracingGlobalIlluminationNextEventEstimati
 	TEXT("NOTE: This parameter is experimental")
 );
 
-static float GRayTracingGlobalIlluminationDiffuseThreshold = 0.05;
+static float GRayTracingGlobalIlluminationDiffuseThreshold = 0.01;
 static FAutoConsoleVariableRef CVarRayTracingGlobalIlluminationDiffuseThreshold(
 	TEXT("r.RayTracing.GlobalIllumination.DiffuseThreshold"),
 	GRayTracingGlobalIlluminationDiffuseThreshold,
