Increased the camera above ground offset to make sure it is always above the ground and not in the virtual planet occluding rays. (some lfoating point issues was happening)
Rerpoduced and fixed on PC.

diff --git a/Engine/Source/Runtime/Renderer/Private/SceneRendering.cpp b/Engine/Source/Runtime/Renderer/Private/SceneRendering.cpp
--- a/Engine/Source/Runtime/Renderer/Private/SceneRendering.cpp
+++ b/Engine/Source/Runtime/Renderer/Private/SceneRendering.cpp
@@ -1305,8 +1305,8 @@ void FViewInfo::SetupUniformBufferParameters(
 		AtmosphereLightDataClearStartIndex = NUM_ATMOSPHERE_LIGHTS;	// Do not clear any atmosphere light data, this component sets everything it needs
 
 		// The constants below should match the one in SkyAtmosphereCommon.ush
-		const float SkyUnitToCm = 1.0f / 0.00001f;
-		const float PlanetRadiusOffset = 0.001f;
+		const float SkyUnitToCm = 1.0f / 0.00001f;	// Kilometers to Centimeters
+		const float PlanetRadiusOffset = 0.01f;		// Always force to be 10 meters above the ground/sea level (to always see the sky and not be under the virtual planet occluding ray tracing)
 
 		const float Offset = PlanetRadiusOffset * SkyUnitToCm;
 		const float BottomRadiusWorld = AtmosphereSetup.BottomRadius * SkyUnitToCm;
