Fixed render graph assert in ray tracing sky pass in multiview mode

diff --git a/Engine/Source/Runtime/Renderer/Private/RayTracing/RaytracingSkylight.cpp b/Engine/Source/Runtime/Renderer/Private/RayTracing/RaytracingSkylight.cpp
--- a/Engine/Source/Runtime/Renderer/Private/RayTracing/RaytracingSkylight.cpp
+++ b/Engine/Source/Runtime/Renderer/Private/RayTracing/RaytracingSkylight.cpp
@@ -982,6 +982,9 @@ void FDeferredShadingSceneRenderer::RenderRayTracingSkyLight(
 		RayDistanceTexture = GraphBuilder.CreateTexture(Desc, TEXT("RayTracingSkyLightHitDistance"));
 	}
 
+	FRDGTextureUAV* SkyLightkUAV = GraphBuilder.CreateUAV(SkyLightTexture);
+	FRDGTextureUAV* RayDistanceUAV = GraphBuilder.CreateUAV(RayDistanceTexture);
+
 	// Sampling state
 	FHaltonPrimes HaltonPrimes;
 	InitializeHaltonPrimes(Scene->HaltonPrimesResource, HaltonPrimes);
@@ -1005,22 +1008,6 @@ void FDeferredShadingSceneRenderer::RenderRayTracingSkyLight(
 		SubsurfaceProfileRT = GSystemTextures.BlackDummy;
 	}
 
-	// View-invariant parameters
-	FRayTracingSkyLightRGS::FParameters *PassParameters = GraphBuilder.AllocParameters<FRayTracingSkyLightRGS::FParameters>();
-	PassParameters->RWOcclusionMaskUAV = GraphBuilder.CreateUAV(SkyLightTexture);
-	PassParameters->RWRayDistanceUAV = GraphBuilder.CreateUAV(RayDistanceTexture);
-	PassParameters->SkyLightData = CreateUniformBufferImmediate(SkyLightData, EUniformBufferUsage::UniformBuffer_SingleDraw);
-	PassParameters->HaltonPrimes = CreateUniformBufferImmediate(HaltonPrimes, EUniformBufferUsage::UniformBuffer_SingleDraw);
-	PassParameters->BlueNoise = CreateUniformBufferImmediate(BlueNoise, EUniformBufferUsage::UniformBuffer_SingleDraw);
-	PassParameters->SkyLightVisibilityRaysDimensions = SkyLightVisibilityRaysDimensions;
-	if (CVarRayTracingSkyLightDecoupleSampleGeneration.GetValueOnRenderThread() == 1)
-	{
-		PassParameters->SkyLightVisibilityRays = GraphBuilder.CreateSRV(SkyLightVisibilityRaysBuffer, EPixelFormat::PF_R32_UINT);
-	}
-	PassParameters->SSProfilesTexture = GraphBuilder.RegisterExternalTexture(SubsurfaceProfileRT);
-	PassParameters->TransmissionProfilesLinearSampler = TStaticSamplerState<SF_Bilinear, AM_Clamp, AM_Clamp, AM_Clamp>::GetRHI();
-	PassParameters->SceneTextures = SceneTextures;
-
 	for (int32 ViewIndex = 0; ViewIndex < Views.Num(); ++ViewIndex)
 	{
 		FViewInfo& View = Views[ViewIndex];
@@ -1034,7 +1021,21 @@ void FDeferredShadingSceneRenderer::RenderRayTracingSkyLight(
 		FHaltonIteration HaltonIteration;
 		InitializeHaltonSequenceIteration(HaltonSequenceIteration, HaltonIteration);
 
-		//  View-dependent parameters
+		FRayTracingSkyLightRGS::FParameters *PassParameters = GraphBuilder.AllocParameters<FRayTracingSkyLightRGS::FParameters>();
+		PassParameters->RWOcclusionMaskUAV = SkyLightkUAV;
+		PassParameters->RWRayDistanceUAV = RayDistanceUAV;
+		PassParameters->SkyLightData = CreateUniformBufferImmediate(SkyLightData, EUniformBufferUsage::UniformBuffer_SingleDraw);
+		PassParameters->HaltonPrimes = CreateUniformBufferImmediate(HaltonPrimes, EUniformBufferUsage::UniformBuffer_SingleDraw);
+		PassParameters->BlueNoise = CreateUniformBufferImmediate(BlueNoise, EUniformBufferUsage::UniformBuffer_SingleDraw);
+		PassParameters->SkyLightVisibilityRaysDimensions = SkyLightVisibilityRaysDimensions;
+		if (CVarRayTracingSkyLightDecoupleSampleGeneration.GetValueOnRenderThread() == 1)
+		{
+			PassParameters->SkyLightVisibilityRays = GraphBuilder.CreateSRV(SkyLightVisibilityRaysBuffer, EPixelFormat::PF_R32_UINT);
+		}
+		PassParameters->SSProfilesTexture = GraphBuilder.RegisterExternalTexture(SubsurfaceProfileRT);
+		PassParameters->TransmissionProfilesLinearSampler = TStaticSamplerState<SF_Bilinear, AM_Clamp, AM_Clamp, AM_Clamp>::GetRHI();
+		PassParameters->SceneTextures = SceneTextures;
+
 		PassParameters->TLAS = View.RayTracingScene.RayTracingSceneRHI->GetShaderResourceView();
 		PassParameters->ViewUniformBuffer = View.ViewUniformBuffer;
 		PassParameters->HaltonIteration = CreateUniformBufferImmediate(HaltonIteration, EUniformBufferUsage::UniformBuffer_SingleDraw);
