Don't compile RSM shaders when a project isn't using LPVs

diff --git a/Engine/Source/Runtime/Engine/Classes/Engine/RendererSettings.h b/Engine/Source/Runtime/Engine/Classes/Engine/RendererSettings.h
--- a/Engine/Source/Runtime/Engine/Classes/Engine/RendererSettings.h
+++ b/Engine/Source/Runtime/Engine/Classes/Engine/RendererSettings.h
@@ -755,6 +755,12 @@ class ENGINE_API URendererSettings : public UDeveloperSettings
 		ConfigRestartRequired = true))
 		uint32 bSupportMaterialLayers : 1;
 
+	UPROPERTY(config, EditAnywhere, Category = Lighting, meta = (
+		ConsoleVariable = "r.LightPropagationVolume", DisplayName = "Light Propagation Volumes",
+		ToolTip = "Whether to allow the usage and compilation of Light Propagation Volumes.",
+		ConfigRestartRequired = true))
+		uint32 bLPV : 1;
+
 public:
 
 	//~ Begin UObject Interface
diff --git a/Engine/Source/Runtime/Engine/Private/RendererSettings.cpp b/Engine/Source/Runtime/Engine/Private/RendererSettings.cpp
--- a/Engine/Source/Runtime/Engine/Private/RendererSettings.cpp
+++ b/Engine/Source/Runtime/Engine/Private/RendererSettings.cpp
@@ -73,6 +73,7 @@ URendererSettings::URendererSettings(const FObjectInitializer& ObjectInitializer
 	GPUSimulationTextureSizeX = 1024;
 	GPUSimulationTextureSizeY = 1024;
 	bEnableRayTracing = 0;
+	bLPV = 1;
 }
 
 void URendererSettings::PostInitProperties()
diff --git a/Engine/Source/Runtime/Renderer/Private/ShadowDepthRendering.cpp b/Engine/Source/Runtime/Renderer/Private/ShadowDepthRendering.cpp
--- a/Engine/Source/Runtime/Renderer/Private/ShadowDepthRendering.cpp
+++ b/Engine/Source/Runtime/Renderer/Private/ShadowDepthRendering.cpp
@@ -223,12 +223,20 @@ public:
 		// Compile for special engine materials.
 		if (bRenderReflectiveShadowMap)
 		{
-			// Reflective shadow map shaders must be compiled for every material because they access the material normal
-			return !bUsePositionOnlyStream
-				// Don't render ShadowDepth for translucent unlit materials, unless we're injecting emissive
-				&& (Material->ShouldCastDynamicShadows() || Material->ShouldInjectEmissiveIntoLPV()
-					|| Material->ShouldBlockGI())
-				&& IsFeatureLevelSupported(Platform, ERHIFeatureLevel::SM5);
+			static const auto SupportLPV = IConsoleManager::Get().FindTConsoleVariableDataInt(TEXT("r.LightPropagationVolume"));
+			if (SupportLPV && SupportLPV->GetValueOnAnyThread() == 0)
+			{
+				return false;
+			}
+			else
+			{
+				// Reflective shadow map shaders must be compiled for every material because they access the material normal
+				return !bUsePositionOnlyStream
+					// Don't render ShadowDepth for translucent unlit materials, unless we're injecting emissive
+					&& (Material->ShouldCastDynamicShadows() || Material->ShouldInjectEmissiveIntoLPV()
+						|| Material->ShouldBlockGI())
+					&& IsFeatureLevelSupported(Platform, ERHIFeatureLevel::SM5);
+			}
 		}
 		else
 		{
@@ -280,9 +288,17 @@ public:
 
 	static bool ShouldCompilePermutation(const FMeshMaterialShaderPermutationParameters& Parameters)
 	{
-		// Re-use ShouldCache from vertex shader
-		return FBaseHS::ShouldCompilePermutation(Parameters)
-			&& TShadowDepthVS<ShaderMode, bRenderReflectiveShadowMap, false>::ShouldCompilePermutation(Parameters);
+		static const auto SupportLPV = IConsoleManager::Get().FindTConsoleVariableDataInt(TEXT("r.LightPropagationVolume"));
+		if (SupportLPV && SupportLPV->GetValueOnAnyThread() == 0)
+		{
+			return false;
+		}
+		else
+		{
+			// Re-use ShouldCache from vertex shader
+			return FBaseHS::ShouldCompilePermutation(Parameters)
+				&& TShadowDepthVS<ShaderMode, bRenderReflectiveShadowMap, false>::ShouldCompilePermutation(Parameters);
+		}
 	}
 
 	static void ModifyCompilationEnvironment(const FMaterialShaderPermutationParameters& Parameters, FShaderCompilerEnvironment& OutEnvironment)
@@ -322,9 +338,17 @@ public:
 
 	static bool ShouldCompilePermutation(const FMeshMaterialShaderPermutationParameters& Parameters)
 	{
-		// Re-use ShouldCache from vertex shader
-		return FBaseDS::ShouldCompilePermutation(Parameters)
-			&& TShadowDepthVS<ShaderMode, bRenderReflectiveShadowMap, false>::ShouldCompilePermutation(Parameters);
+		static const auto SupportLPV = IConsoleManager::Get().FindTConsoleVariableDataInt(TEXT("r.LightPropagationVolume"));
+		if (SupportLPV && SupportLPV->GetValueOnAnyThread() == 0)
+		{
+			return false;
+		}
+		else
+		{
+			// Re-use ShouldCache from vertex shader
+			return FBaseDS::ShouldCompilePermutation(Parameters)
+				&& TShadowDepthVS<ShaderMode, bRenderReflectiveShadowMap, false>::ShouldCompilePermutation(Parameters);
+		}
 	}
 
 	static void ModifyCompilationEnvironment(const FMaterialShaderPermutationParameters& Parameters, FShaderCompilerEnvironment& OutEnvironment)
@@ -485,12 +509,20 @@ public:
 
 		if (bRenderReflectiveShadowMap)
 		{
-			//Note: This logic needs to stay in sync with OverrideWithDefaultMaterialForShadowDepth!
-			// Reflective shadow map shaders must be compiled for every material because they access the material normal
-			return
-				// Only compile one pass point light shaders for feature levels >= SM4
-				(Material->ShouldCastDynamicShadows() || Material->ShouldInjectEmissiveIntoLPV() || Material->ShouldBlockGI())
-				&& IsFeatureLevelSupported(Platform, ERHIFeatureLevel::SM5);
+			static const auto SupportLPV = IConsoleManager::Get().FindTConsoleVariableDataInt(TEXT("r.LightPropagationVolume"));
+			if (SupportLPV && SupportLPV->GetValueOnAnyThread() == 0)
+			{
+				return false;
+			}
+			else
+			{
+				//Note: This logic needs to stay in sync with OverrideWithDefaultMaterialForShadowDepth!
+				// Reflective shadow map shaders must be compiled for every material because they access the material normal
+				return
+					// Only compile one pass point light shaders for feature levels >= SM4
+					(Material->ShouldCastDynamicShadows() || Material->ShouldInjectEmissiveIntoLPV() || Material->ShouldBlockGI())
+					&& IsFeatureLevelSupported(Platform, ERHIFeatureLevel::SM5);
+			}
 		}
 		else
 		{
