Fix bug caused by CL 9451345 in DeferredShadingRendering.cpp. If bDitheredLODTransitionsUseStencil was false, AsyncDitherLODEndFence would be null, which would then be passed into the RHI causing the RHI to dereference a nullptr.

diff --git a/Engine/Source/Runtime/Renderer/Private/DeferredShadingRenderer.cpp b/Engine/Source/Runtime/Renderer/Private/DeferredShadingRenderer.cpp
--- a/Engine/Source/Runtime/Renderer/Private/DeferredShadingRenderer.cpp
+++ b/Engine/Source/Runtime/Renderer/Private/DeferredShadingRenderer.cpp
@@ -1237,7 +1237,7 @@ void FDeferredShadingSceneRenderer::Render(FRHICommandListImmediate& RHICmdList)
 		Scene->FXSystem->PreRender(RHICmdList, &Views[0].GlobalDistanceFieldInfo.ParameterData, Views[0].AllowGPUParticleUpdate());
 	}
 
-	if (bStencilLODComputeAsync)
+	if (AsyncDitherLODEndFence)
 	{
 		RHICmdList.WaitComputeFence(AsyncDitherLODEndFence);
 	}
