Fixed division by zero => NaN result

#rnx

diff --git a/Engine/Source/Runtime/Renderer/Private/VolumeLighting.h b/Engine/Source/Runtime/Renderer/Private/VolumeLighting.h
--- a/Engine/Source/Runtime/Renderer/Private/VolumeLighting.h
+++ b/Engine/Source/Runtime/Renderer/Private/VolumeLighting.h
@@ -74,7 +74,7 @@ public:
 			// far cascade plane
 			if(InnerSplitIndex != CascadeCount - 1)
 			{
-				ShadowInjectParamValue.Y = 1.0f / ShadowCascadeSettings.SplitFarFadeRegion;
+				ShadowInjectParamValue.Y = 1.0f / (ShadowCascadeSettings.SplitFarFadeRegion == 0.0f ? 0.0001f : ShadowCascadeSettings.SplitFarFadeRegion);
 				Planes[1] = FVector4((FVector)(ShadowCascadeSettings.FarFrustumPlane), -ShadowCascadeSettings.FarFrustumPlane.W);
 			}
 
