Added water depth override for Android as a workaround for broken scene depth reads on some devices. Controlled with r.Water.ConstantWaterDepth

diff --git a/Engine/Shaders/Private/SingleLayerWaterCommon.ush b/Engine/Shaders/Private/SingleLayerWaterCommon.ush
--- a/Engine/Shaders/Private/SingleLayerWaterCommon.ush
+++ b/Engine/Shaders/Private/SingleLayerWaterCommon.ush
@@ -89,6 +89,15 @@ WaterVolumeLightingOutput EvaluateWaterVolumeLighting(
 			PixelSceneDepth = lerp(PixelSceneDepth, FadeBehindWaterDepth, saturate((WaterDepth - FadeStartDistance) / FadeLenght));
 			PixelSceneDeviceZ = ConvertToDeviceZ(PixelSceneDepth);
 		}
+		
+		// Water depth override hack for Android devices
+	#if !COMPILER_METAL
+		if (View.ConstantWaterDepth >= 0.0f)
+		{
+			PixelSceneDepth = WaterDepth + View.ConstantWaterDepth;
+			PixelSceneDeviceZ = ConvertToDeviceZ(PixelSceneDepth);
+		}
+	#endif
 #endif
 
 		float SceneDeviceZ = PixelSceneDeviceZ;
diff --git a/Engine/Source/Runtime/Engine/Public/SceneView.h b/Engine/Source/Runtime/Engine/Public/SceneView.h
--- a/Engine/Source/Runtime/Engine/Public/SceneView.h
+++ b/Engine/Source/Runtime/Engine/Public/SceneView.h
@@ -724,6 +724,7 @@ enum ETranslucencyVolumeCascade
 	VIEW_UNIFORM_BUFFER_MEMBER_ARRAY(FVector4, XRPassthroughCameraUVs, [2]) \
 	VIEW_UNIFORM_BUFFER_MEMBER(int32, FarShadowStaticMeshLODBias) \
 	VIEW_UNIFORM_BUFFER_MEMBER(float, MinRoughness) \
+	VIEW_UNIFORM_BUFFER_MEMBER(float, ConstantWaterDepth) \
 
 #define VIEW_UNIFORM_BUFFER_MEMBER(type, identifier) \
 	SHADER_PARAMETER(type, identifier)
diff --git a/Engine/Source/Runtime/Renderer/Private/SceneRendering.cpp b/Engine/Source/Runtime/Renderer/Private/SceneRendering.cpp
--- a/Engine/Source/Runtime/Renderer/Private/SceneRendering.cpp
+++ b/Engine/Source/Runtime/Renderer/Private/SceneRendering.cpp
@@ -192,6 +192,11 @@ static TAutoConsoleVariable<int32> CVarViewRectUseScreenBottom(
 	ECVF_RenderThreadSafe
 );
 
+static TAutoConsoleVariable<float> CVarConstantWaterDepth(
+	TEXT("r.Water.ConstantWaterDepth"),
+	-1.0f,
+	TEXT("Setting this to a negative value means disabled and reading from scene depth texture. Specified in unreal units."),
+	ECVF_Scalability | ECVF_RenderThreadSafe);
 
 #if !(UE_BUILD_SHIPPING || UE_BUILD_TEST)
 static TAutoConsoleVariable<float> CVarGeneralPurposeTweak(
@@ -1492,6 +1497,8 @@ void FViewInfo::SetupUniformBufferParameters(
 		ViewUniformShaderParameters.GeneralPurposeTweak = Value;
 	}
 
+	ViewUniformShaderParameters.ConstantWaterDepth = CVarConstantWaterDepth.GetValueOnRenderThread();
+
 	ViewUniformShaderParameters.DemosaicVposOffset = 0.0f;
 	{
 		ViewUniformShaderParameters.DemosaicVposOffset = CVarDemosaicVposOffset.GetValueOnRenderThread();
