Expose an anisotropy setting for virtual texture only and have a base setting that looks better on PC low texture settings.
We can tweak this back down live if necessay, but I don't think it will be.
Return the global anisotropy setting back to 0 as it was before my last attempted fix.

diff --git a/Engine/Config/BaseScalability.ini b/Engine/Config/BaseScalability.ini
--- a/Engine/Config/BaseScalability.ini
+++ b/Engine/Config/BaseScalability.ini
@@ -329,6 +329,7 @@ r.Streaming.AmortizeCPUToGPUCopy=1
 r.Streaming.MaxNumTexturesToStreamPerFrame=1
 r.Streaming.Boost=0.3
 r.MaxAnisotropy=0
+r.VT.MaxAnisotropy=4
 r.Streaming.LimitPoolSizeToVRAM=1
 r.Streaming.PoolSize=400
 r.Streaming.MaxEffectiveScreenSize=0
@@ -339,6 +340,7 @@ r.Streaming.AmortizeCPUToGPUCopy=0
 r.Streaming.MaxNumTexturesToStreamPerFrame=0
 r.Streaming.Boost=1
 r.MaxAnisotropy=2
+r.VT.MaxAnisotropy=4
 r.Streaming.LimitPoolSizeToVRAM=1
 r.Streaming.PoolSize=600
 r.Streaming.MaxEffectiveScreenSize=0
@@ -349,6 +351,7 @@ r.Streaming.AmortizeCPUToGPUCopy=0
 r.Streaming.MaxNumTexturesToStreamPerFrame=0
 r.Streaming.Boost=1
 r.MaxAnisotropy=4
+r.VT.MaxAnisotropy=8
 r.Streaming.LimitPoolSizeToVRAM=1
 r.Streaming.PoolSize=800
 r.Streaming.MaxEffectiveScreenSize=0
@@ -359,6 +362,7 @@ r.Streaming.AmortizeCPUToGPUCopy=0
 r.Streaming.MaxNumTexturesToStreamPerFrame=0
 r.Streaming.Boost=1
 r.MaxAnisotropy=8
+r.VT.MaxAnisotropy=8
 r.Streaming.LimitPoolSizeToVRAM=0
 r.Streaming.PoolSize=1000
 r.Streaming.MaxEffectiveScreenSize=0
@@ -369,6 +373,7 @@ r.Streaming.AmortizeCPUToGPUCopy=0
 r.Streaming.MaxNumTexturesToStreamPerFrame=0
 r.Streaming.Boost=1
 r.MaxAnisotropy=8
+r.VT.MaxAnisotropy=8
 r.Streaming.LimitPoolSizeToVRAM=0
 r.Streaming.PoolSize=3000
 r.Streaming.MaxEffectiveScreenSize=0
diff --git a/Engine/Source/Runtime/Engine/Private/VT/VirtualTextureScalability.cpp b/Engine/Source/Runtime/Engine/Private/VT/VirtualTextureScalability.cpp
--- a/Engine/Source/Runtime/Engine/Private/VT/VirtualTextureScalability.cpp
+++ b/Engine/Source/Runtime/Engine/Private/VT/VirtualTextureScalability.cpp
@@ -41,15 +41,24 @@ namespace VirtualTextureScalability
 		ECVF_RenderThreadSafe | ECVF_Scalability
 	);
 
+	static TAutoConsoleVariable<int32> CVarVTMaxAnisotropy(
+		TEXT("r.VT.MaxAnisotropy"),
+		8,
+		TEXT("MaxAnisotropy setting for Virtual Texture sampling."),
+		ECVF_RenderThreadSafe | ECVF_Scalability
+	);
+	
 
 	/** Track changes and apply to relevant systems. This allows us to dynamically change the scalability settings. */
 	static void OnUpdate()
 	{
 		const float PoolSizeScale = CVarVTPoolSizeScale.GetValueOnGameThread();
 		const float TileCountBias = CVarVTTileCountBias.GetValueOnGameThread();
+		const float MaxAnisotropy = CVarVTMaxAnisotropy.GetValueOnGameThread();
 
 		static float LastPoolSizeScale = PoolSizeScale;
 		static float LastTileCountBias = TileCountBias;
+		static float LastMaxAnisotropy = MaxAnisotropy;
 
 		bool bUpdate = false;
 		if (LastPoolSizeScale != PoolSizeScale)
@@ -62,6 +71,11 @@ namespace VirtualTextureScalability
 			LastTileCountBias = TileCountBias;
 			bUpdate = true;
 		}
+		if (LastMaxAnisotropy != MaxAnisotropy)
+		{
+			LastMaxAnisotropy = MaxAnisotropy;
+			bUpdate = true;
+		}
 
 		if (bUpdate)
 		{
@@ -119,4 +133,9 @@ namespace VirtualTextureScalability
 	{
 		return CVarVTTileCountBias.GetValueOnAnyThread();
 	}
+
+	int32 GetMaxAnisotropy()
+	{
+		return CVarVTMaxAnisotropy.GetValueOnAnyThread();
+	}
 }
diff --git a/Engine/Source/Runtime/Engine/Public/VT/VirtualTextureScalability.h b/Engine/Source/Runtime/Engine/Public/VT/VirtualTextureScalability.h
--- a/Engine/Source/Runtime/Engine/Public/VT/VirtualTextureScalability.h
+++ b/Engine/Source/Runtime/Engine/Public/VT/VirtualTextureScalability.h
@@ -12,4 +12,9 @@ namespace VirtualTextureScalability
 	ENGINE_API float GetPoolSizeScale();
 	/** Get resolution bias for runtime virtual textures. */
 	ENGINE_API int32 GetRuntimeVirtualTextureSizeBias();
+	/**
+	 * Get maximum anisotropy when virtual texture sampling. 
+	 * This is also clamped per virtual texture according to the tile border size.
+	 */
+	ENGINE_API int32 GetMaxAnisotropy();
 }
diff --git a/Engine/Source/Runtime/Renderer/Private/VT/AllocatedVirtualTexture.cpp b/Engine/Source/Runtime/Renderer/Private/VT/AllocatedVirtualTexture.cpp
--- a/Engine/Source/Runtime/Renderer/Private/VT/AllocatedVirtualTexture.cpp
+++ b/Engine/Source/Runtime/Renderer/Private/VT/AllocatedVirtualTexture.cpp
@@ -1,9 +1,10 @@
 // Copyright 1998-2019 Epic Games, Inc. All Rights Reserved.
 
 #include "AllocatedVirtualTexture.h"
-#include "VirtualTextureSystem.h"
-#include "VirtualTextureSpace.h"
-#include "VirtualTexturePhysicalSpace.h"
+#include "VT/VirtualTextureScalability.h"
+#include "VT/VirtualTextureSystem.h"
+#include "VT/VirtualTextureSpace.h"
+#include "VT/VirtualTexturePhysicalSpace.h"
 
 FAllocatedVirtualTexture::FAllocatedVirtualTexture(FVirtualTextureSystem* InSystem,
 	uint32 InFrame,
@@ -268,8 +269,6 @@ static inline uint32 BitcastFloatToUInt32(float v)
 
 void FAllocatedVirtualTexture::GetPackedPageTableUniform(FUintVector4* OutUniform, bool bApplyBlockScale) const
 {
-	static const auto MaxAnisotropyCVar = IConsoleManager::Get().FindTConsoleVariableDataInt(TEXT("r.MaxAnisotropy"));
-
 	const uint32 vPageX = FMath::ReverseMortonCode2(VirtualAddress);
 	const uint32 vPageY = FMath::ReverseMortonCode2(VirtualAddress >> 1);
 	const uint32 vPageSize = GetVirtualTileSize();
@@ -278,7 +277,7 @@ void FAllocatedVirtualTexture::GetPackedPageTableUniform(FUintVector4* OutUnifor
 	const uint32 HeightInPages = GetHeightInTiles();
 	const uint32 vPageTableMipBias = FMath::FloorLog2(vPageSize);
 
-	const uint32 MaxAnisotropy = FMath::Clamp<int32>(MaxAnisotropyCVar->GetValueOnRenderThread(), 1, PageBorderSize);
+	const uint32 MaxAnisotropy = FMath::Clamp<int32>(VirtualTextureScalability::GetMaxAnisotropy(), 1, PageBorderSize);
 	const uint32 MaxAnisotropyLog2 = FMath::FloorLog2(MaxAnisotropy);
 
 	// make sure everything fits in the allocated number of bits
