Integrating CL#9566776 from //UE4/Dev-Main to Dev-VirtualProduction (//UE4/Dev-VirtualProduction) to fix a crash when loading VirtualCamera sample project

diff --git a/Engine/Source/Runtime/Renderer/Private/ShadowDepthRendering.cpp b/Engine/Source/Runtime/Renderer/Private/ShadowDepthRendering.cpp
--- a/Engine/Source/Runtime/Renderer/Private/ShadowDepthRendering.cpp
+++ b/Engine/Source/Runtime/Renderer/Private/ShadowDepthRendering.cpp
@@ -288,17 +288,9 @@ public:
 
 	static bool ShouldCompilePermutation(const FMeshMaterialShaderPermutationParameters& Parameters)
 	{
-		static const auto SupportLPV = IConsoleManager::Get().FindTConsoleVariableDataInt(TEXT("r.LightPropagationVolume"));
-		if (SupportLPV && SupportLPV->GetValueOnAnyThread() == 0)
-		{
-			return false;
-		}
-		else
-		{
-			// Re-use ShouldCache from vertex shader
-			return FBaseHS::ShouldCompilePermutation(Parameters)
-				&& TShadowDepthVS<ShaderMode, bRenderReflectiveShadowMap, false>::ShouldCompilePermutation(Parameters);
-		}
+		// Re-use ShouldCache from vertex shader
+		return FBaseHS::ShouldCompilePermutation(Parameters)
+			&& TShadowDepthVS<ShaderMode, bRenderReflectiveShadowMap, false>::ShouldCompilePermutation(Parameters);
 	}
 
 	static void ModifyCompilationEnvironment(const FMaterialShaderPermutationParameters& Parameters, FShaderCompilerEnvironment& OutEnvironment)
@@ -338,17 +330,9 @@ public:
 
 	static bool ShouldCompilePermutation(const FMeshMaterialShaderPermutationParameters& Parameters)
 	{
-		static const auto SupportLPV = IConsoleManager::Get().FindTConsoleVariableDataInt(TEXT("r.LightPropagationVolume"));
-		if (SupportLPV && SupportLPV->GetValueOnAnyThread() == 0)
-		{
-			return false;
-		}
-		else
-		{
-			// Re-use ShouldCache from vertex shader
-			return FBaseDS::ShouldCompilePermutation(Parameters)
-				&& TShadowDepthVS<ShaderMode, bRenderReflectiveShadowMap, false>::ShouldCompilePermutation(Parameters);
-		}
+		// Re-use ShouldCache from vertex shader
+		return FBaseDS::ShouldCompilePermutation(Parameters)
+			&& TShadowDepthVS<ShaderMode, bRenderReflectiveShadowMap, false>::ShouldCompilePermutation(Parameters);
 	}
 
 	static void ModifyCompilationEnvironment(const FMaterialShaderPermutationParameters& Parameters, FShaderCompilerEnvironment& OutEnvironment)
