Round-robin occlusion queries for stereoscopic rendering
With the vr.RoundRobinOcclusion flag turned on, each frame will only render occlusion queries for one eye with an alternating scheme (ie odd frames only left eye, even frames only right eye).
Additionally, the interface for class FPrimitiveOcclusionHistory has been changed along with a slight modification in its implementation. With round-robin occlusion, when extracing history entries we perform a search through the entire history to look for the oldest entry in a primitive's occlusion history.

diff --git a/Engine/Source/Runtime/Engine/Classes/Engine/RendererSettings.h b/Engine/Source/Runtime/Engine/Classes/Engine/RendererSettings.h
--- a/Engine/Source/Runtime/Engine/Classes/Engine/RendererSettings.h
+++ b/Engine/Source/Runtime/Engine/Classes/Engine/RendererSettings.h
@@ -537,6 +537,12 @@ class ENGINE_API URendererSettings : public UDeveloperSettings
 		ConfigRestartRequired = true))
 		uint32 bMonoscopicFarField : 1;
 
+	UPROPERTY(config, EditAnywhere, Category = VR, meta = (
+		ConsoleVariable = "vr.RoundRobinOcclusion", DisplayName = "Round Robin Occlusion Queries",
+		ToolTip = "Enable round-robin scheduling of occlusion queries for VR.",
+		ConfigRestartRequired = false))
+	uint32 bRoundRobinOcclusion : 1;
+
 	UPROPERTY(config, EditAnywhere, Category = Experimental, meta = (
 		ConsoleVariable = "vr.ODSCapture", DisplayName = "Omni-directional Stereo Capture",
 		ToolTip = "Enable Omni-directional Stereo Capture.",
diff --git a/Engine/Source/Runtime/Renderer/Private/RendererScene.cpp b/Engine/Source/Runtime/Renderer/Private/RendererScene.cpp
--- a/Engine/Source/Runtime/Renderer/Private/RendererScene.cpp
+++ b/Engine/Source/Runtime/Renderer/Private/RendererScene.cpp
@@ -138,6 +138,8 @@ FSceneViewState::FSceneViewState()
 
 	bIsStereoView = false;
 
+	bRoundRobinOcclusionEnabled = false;
+
 	HeightfieldLightingAtlas = NULL;
 
 	for (int32 CascadeIndex = 0; CascadeIndex < ARRAY_COUNT(TranslucencyLightingCacheAllocations); CascadeIndex++)
diff --git a/Engine/Source/Runtime/Renderer/Private/ScenePrivate.h b/Engine/Source/Runtime/Renderer/Private/ScenePrivate.h
--- a/Engine/Source/Runtime/Renderer/Private/ScenePrivate.h
+++ b/Engine/Source/Runtime/Renderer/Private/ScenePrivate.h
@@ -105,6 +105,38 @@ public:
 	/** whether or not this primitive was grouped the last time it was queried */
 	bool bGroupedQuery[FOcclusionQueryHelpers::MaxBufferedOcclusionFrames];
 
+private:
+	/**
+	 *	Whether or not we need to linearly search the history for a past entry. Scanning may be necessary if for every frame there
+	 *	is a hole in PendingOcclusionQueryFrames in the same spot (ex. if for every frame PendingOcclusionQueryFrames[1] is null).
+	 *	This could lead to overdraw for the frames that attempt to read these holes by getting back nothing every time.
+	 *	This can occur when round robin occlusion queries are turned on while NumBufferedFrames is even.
+	 */
+	bool bNeedsScanOnRead;
+
+	/**
+	 *	Scan for the oldest non-stale (<= LagTolerance frames old) in the occlusion history by examining their corresponding frame numbers.
+	 *	Conditions where this is needed to get a query for read-back are described for bNeedsScanOnRead.
+	 *	Returns -1 if no such query exists in the occlusion history.
+	 */
+	FORCEINLINE int32 ScanOldestNonStaleQueryIndex(uint32 FrameNumber, int32 NumBufferedFrames, int32 LagTolerance) const
+	{
+		uint32 OldestFrame = UINT32_MAX;
+		int32 OldestQueryIndex = -1;
+		for (int Index = 0; Index < NumBufferedFrames; ++Index)
+		{
+			const uint32 ThisFrameNumber = PendingOcclusionQueryFrames[Index];
+			const int32 LaggedFrames = FrameNumber - ThisFrameNumber;
+			if (PendingOcclusionQuery[Index].IsValid() && LaggedFrames <= LagTolerance && ThisFrameNumber < OldestFrame)
+			{
+				OldestFrame = ThisFrameNumber;
+				OldestQueryIndex = Index;
+			}
+		}
+		return OldestQueryIndex;
+	}
+
+public:
 	/** Initialization constructor. */
 	FORCEINLINE FPrimitiveOcclusionHistory(FPrimitiveComponentId InPrimitiveId, int32 SubQuery)
 		: PrimitiveId(InPrimitiveId)
@@ -118,6 +150,7 @@ public:
 		, BecameEligibleForQueryCooldown(0)
 		, WasOccludedLastFrame(false)
 		, OcclusionStateWasDefiniteLastFrame(false)
+		, bNeedsScanOnRead(false)
 	{
 		for (int32 Index = 0; Index < FOcclusionQueryHelpers::MaxBufferedOcclusionFrames; Index++)
 		{
@@ -137,6 +170,7 @@ public:
 		, BecameEligibleForQueryCooldown(0)
 		, WasOccludedLastFrame(false)
 		, OcclusionStateWasDefiniteLastFrame(false)
+		, bNeedsScanOnRead(false)
 	{
 		for (int32 Index = 0; Index < FOcclusionQueryHelpers::MaxBufferedOcclusionFrames; Index++)
 		{
@@ -165,32 +199,46 @@ public:
 	FORCEINLINE void ReleaseQuery(TOcclusionQueryPool& Pool, uint32 FrameNumber, int32 NumBufferedFrames)
 	{
 		const uint32 QueryIndex = FOcclusionQueryHelpers::GetQueryLookupIndex(FrameNumber, NumBufferedFrames);
-		if (PendingOcclusionQuery[QueryIndex].GetReference())
+		if (PendingOcclusionQuery[QueryIndex].IsValid())
 		{
 			Pool.ReleaseQuery(PendingOcclusionQuery[QueryIndex]);
 		}
 	}
 
-	FORCEINLINE FRenderQueryRHIParamRef GetPastQuery(uint32 FrameNumber, int32 NumBufferedFrames, bool& bOutGrouped)
+	FORCEINLINE FRenderQueryRHIParamRef GetQueryForEviction(uint32 FrameNumber, int32 NumBufferedFrames) const
 	{
-		// Get the oldest occlusion query
 		const uint32 QueryIndex = FOcclusionQueryHelpers::GetQueryLookupIndex(FrameNumber, NumBufferedFrames);
-		bOutGrouped = false;
-		if (PendingOcclusionQuery[QueryIndex].GetReference() && PendingOcclusionQueryFrames[QueryIndex] == FrameNumber - uint32(NumBufferedFrames))
+		if (PendingOcclusionQuery[QueryIndex].IsValid())
 		{
-			bOutGrouped = bGroupedQuery[QueryIndex];
 			return PendingOcclusionQuery[QueryIndex].GetReference();
 		}
 		return nullptr;
 	}
 
-	FORCEINLINE void SetCurrentQuery(uint32 FrameNumber, FRenderQueryRHIParamRef NewQuery, int32 NumBufferedFrames, bool bGrouped)
+
+	FORCEINLINE FRenderQueryRHIParamRef GetQueryForReading(uint32 FrameNumber, int32 NumBufferedFrames, int32 LagTolerance, bool& bOutGrouped) const
+	{
+		const int32 OldestQueryIndex = bNeedsScanOnRead ? ScanOldestNonStaleQueryIndex(FrameNumber, NumBufferedFrames, LagTolerance)
+														: FOcclusionQueryHelpers::GetQueryLookupIndex(FrameNumber, NumBufferedFrames);
+		const int32 LaggedFrames = FrameNumber - PendingOcclusionQueryFrames[OldestQueryIndex];
+		if (OldestQueryIndex == -1 || !PendingOcclusionQuery[OldestQueryIndex].IsValid() || LaggedFrames > LagTolerance)
+		{
+			bOutGrouped = false;
+			return nullptr;
+		}
+		bOutGrouped = bGroupedQuery[OldestQueryIndex];
+		return PendingOcclusionQuery[OldestQueryIndex];
+	}
+
+	FORCEINLINE void SetCurrentQuery(uint32 FrameNumber, FRenderQueryRHIParamRef NewQuery, int32 NumBufferedFrames, bool bGrouped, bool bNeedsScan)
 	{
 		// Get the current occlusion query
 		const uint32 QueryIndex = FOcclusionQueryHelpers::GetQueryIssueIndex(FrameNumber, NumBufferedFrames);
 		PendingOcclusionQuery[QueryIndex] = NewQuery;
 		PendingOcclusionQueryFrames[QueryIndex] = FrameNumber;
 		bGroupedQuery[QueryIndex] = bGrouped;
+
+		bNeedsScanOnRead = bNeedsScan;
 	}
 
 	FORCEINLINE uint32 LastQuerySubmitFrame() const
@@ -750,6 +798,9 @@ private:
 	// whether this view is a stereo counterpart to a primary view
 	bool bIsStereoView;
 
+	// The whether or not round-robin occlusion querying is enabled for this view
+	bool bRoundRobinOcclusionEnabled;
+
 public:
 
 	// Previous frame's view info to use.
@@ -995,6 +1046,16 @@ public:
 	 */
 	void TrimOcclusionHistory(float CurrentTime, float MinHistoryTime, float MinQueryTime, int32 FrameNumber);
 
+	inline void UpdateRoundRobin(const bool bUseRoundRobin)
+	{
+		bRoundRobinOcclusionEnabled = bUseRoundRobin;
+	}
+
+	inline bool IsRoundRobinEnabled() const
+	{
+		return bRoundRobinOcclusionEnabled;
+	}
+
 	/**
 	 * Checks whether a shadow is occluded this frame.
 	 * @param Primitive - The shadow subject.
diff --git a/Engine/Source/Runtime/Renderer/Private/SceneRendering.cpp b/Engine/Source/Runtime/Renderer/Private/SceneRendering.cpp
--- a/Engine/Source/Runtime/Renderer/Private/SceneRendering.cpp
+++ b/Engine/Source/Runtime/Renderer/Private/SceneRendering.cpp
@@ -129,6 +129,12 @@ static TAutoConsoleVariable<int32> CVarMonoscopicFarFieldMode(
 	TEXT(", 4 mono far field only"),
 	ECVF_Scalability | ECVF_RenderThreadSafe);
 
+static TAutoConsoleVariable<int32> CVarRoundRobinOcclusion(
+	TEXT("vr.RoundRobinOcclusion"),
+	0,
+	TEXT("0 to disable round-robin occlusion queries for stereo rendering (default), 1 to enable."),
+	ECVF_Scalability | ECVF_RenderThreadSafe);
+
 static TAutoConsoleVariable<int32> CVarUsePreExposure(
 	TEXT("r.UsePreExposure"),
 	0,
diff --git a/Engine/Source/Runtime/Renderer/Private/SceneVisibility.cpp b/Engine/Source/Runtime/Renderer/Private/SceneVisibility.cpp
--- a/Engine/Source/Runtime/Renderer/Private/SceneVisibility.cpp
+++ b/Engine/Source/Runtime/Renderer/Private/SceneVisibility.cpp
@@ -605,6 +605,7 @@ struct FVisForPrimParams
 		, HZBBoundsToAdd(OutHZBBounds)
 		, QueriesToAdd(OutQueriesToRun)	
 		, SubIsOccluded(OutSubIsOccluded)
+		, bNeedsScanOnRead(false)
 	{
 
 	}
@@ -645,6 +646,9 @@ struct FVisForPrimParams
 	bool bSubmitQueries;
 	bool bHZBOcclusion;	
 
+	// Whether the entries written into the history need to be read using a scan search (see FPrimitiveOcclusionHistory::bNeedsScanOnRead)
+	bool bNeedsScanOnRead;
+
 	//occlusion history to insert into.  In parallel these will be all merged back into the view's history on the main thread.
 	//use TChunkedArray so pointers to the new FPrimitiveOcclusionHistory's won't change if the array grows.	
 	TArray<FPrimitiveOcclusionHistory>*		InsertPrimitiveOcclusionHistory;
@@ -675,12 +679,31 @@ static void FetchVisibilityForPrimitives_Range(FVisForPrimParams& Params)
 
 	FSceneViewState* ViewState = (FSceneViewState*)View.State;
 	const int32 NumBufferedFrames = FOcclusionQueryHelpers::GetNumBufferedFrames(Scene->GetFeatureLevel());
-	const bool bClearQueries = !View.Family->EngineShowFlags.HitProxies;
+	bool bClearQueries = !View.Family->EngineShowFlags.HitProxies;
 	const float CurrentRealTime = View.Family->CurrentRealTime;
 	uint32 OcclusionFrameCounter = ViewState->OcclusionFrameCounter;
 	FRenderQueryPool& OcclusionQueryPool = ViewState->OcclusionQueryPool;
 	FHZBOcclusionTester& HZBOcclusionTests = ViewState->HZBOcclusionTests;
-	
+
+	int32 ReadBackLagTolerance = NumBufferedFrames;
+
+	const bool bIsStereoView = View.StereoPass == eSSP_LEFT_EYE || View.StereoPass == eSSP_RIGHT_EYE;
+	const bool bUseRoundRobinOcclusion = bIsStereoView && !View.bIsSceneCapture && View.ViewState->IsRoundRobinEnabled();
+	if (bUseRoundRobinOcclusion)
+	{
+		// We don't allow clearing of a history entry if we do not also submit an occlusion query to replace the deleted one
+		// as we want to keep the history as full as possible
+		bClearQueries &= bSubmitQueries;
+
+		// However, if this frame happens to be the first frame, then we clear anyway since in the first frame we should not be
+		// reading past queries
+		bClearQueries |= View.bIgnoreExistingQueries;
+
+		// Round-robin occlusion culling involves reading frames that could be twice as stale as without round-robin
+		ReadBackLagTolerance = NumBufferedFrames * 2;
+	}
+	// Round robin occlusion culling can make holes in the occlusion history which would require scanning the history when reading
+	Params.bNeedsScanOnRead = bUseRoundRobinOcclusion;
 
 	TSet<FPrimitiveOcclusionHistory, FPrimitiveOcclusionHistoryKeyFuncs>& ViewPrimitiveOcclusionHistory = ViewState->PrimitiveOcclusionHistorySet;
 	TArray<FPrimitiveOcclusionHistory>* InsertPrimitiveOcclusionHistory = Params.InsertPrimitiveOcclusionHistory;
@@ -821,7 +844,7 @@ static void FetchVisibilityForPrimitives_Range(FVisForPrimParams& Params)
 						// Read the occlusion query results.
 						uint64 NumSamples = 0;
 						bool bGrouped = false;
-						FRenderQueryRHIParamRef PastQuery = PrimitiveOcclusionHistory->GetPastQuery(OcclusionFrameCounter, NumBufferedFrames, bGrouped);
+						FRenderQueryRHIParamRef PastQuery = PrimitiveOcclusionHistory->GetQueryForReading(OcclusionFrameCounter, NumBufferedFrames, ReadBackLagTolerance, bGrouped);
 						if (PastQuery)
 						{
 							//int32 RefCount = PastQuery.GetReference()->GetRefCount();
@@ -903,9 +926,7 @@ static void FetchVisibilityForPrimitives_Range(FVisForPrimParams& Params)
 					}
 					else
 					{
-						bool bGrouped = false;
-						FRenderQueryRHIParamRef Query = PrimitiveOcclusionHistory->GetPastQuery(OcclusionFrameCounter, NumBufferedFrames, bGrouped);
-						if (Query)
+						if (PrimitiveOcclusionHistory->GetQueryForEviction(OcclusionFrameCounter, NumBufferedFrames))
 						{
 							QueriesToRelease->Add(PrimitiveOcclusionHistory);							
 						}
@@ -1032,7 +1053,8 @@ static void FetchVisibilityForPrimitives_Range(FVisForPrimParams& Params)
 											View.GroupedOcclusionQueries.BatchPrimitive(BoundOrigin, BoundExtent) :
 											View.IndividualOcclusionQueries.BatchPrimitive(BoundOrigin, BoundExtent),
 											NumBufferedFrames,
-											bGroupedQuery
+											bGroupedQuery,
+											Params.bNeedsScanOnRead
 										);
 									}
 								}
@@ -1321,7 +1343,8 @@ static int32 FetchVisibilityForPrimitives(const FScene* Scene, FViewInfo& View,
 						View.GroupedOcclusionQueries.BatchPrimitive(RunQueriesIter->BoundsOrigin, RunQueriesIter->BoundsExtent) :
 						View.IndividualOcclusionQueries.BatchPrimitive(RunQueriesIter->BoundsOrigin, RunQueriesIter->BoundsExtent),
 						NumBufferedFrames,
-						RunQueriesIter->bGroupedQuery
+						RunQueriesIter->bGroupedQuery,
+						Params[i].bNeedsScanOnRead
 						);
 				}
 			}
@@ -1405,7 +1428,8 @@ static int32 FetchVisibilityForPrimitives(const FScene* Scene, FViewInfo& View,
 					PrimitiveOcclusionHistory->SetCurrentQuery(OcclusionFrameCounter,
 						View.IndividualOcclusionQueries.BatchPrimitive(RunQueriesIter->BoundsOrigin, RunQueriesIter->BoundsExtent),
 						NumBufferedFrames,
-						false
+						false,
+						Params.bNeedsScanOnRead
 					);
 				}
 			}
@@ -1432,7 +1456,8 @@ static int32 FetchVisibilityForPrimitives(const FScene* Scene, FViewInfo& View,
 					PrimitiveOcclusionHistory->SetCurrentQuery(OcclusionFrameCounter,
 						View.IndividualOcclusionQueries.BatchPrimitive(RunQueriesIter->BoundsOrigin, RunQueriesIter->BoundsExtent),
 						NumBufferedFrames,
-						false
+						false,
+						Params.bNeedsScanOnRead
 					);
 				}
 			}
@@ -1523,8 +1548,21 @@ static int32 OcclusionCull(FRHICommandListImmediate& RHICmdList, const FScene* S
 				check(!ViewState->HZBOcclusionTests.IsValidFrame(ViewState->OcclusionFrameCounter));
 				ViewState->HZBOcclusionTests.MapResults(RHICmdList);
 			}
-			
-			NumOccludedPrimitives += FetchVisibilityForPrimitives(Scene, View, bSubmitQueries, bHZBOcclusion);			
+ 
+			// Perform round-robin occlusion queries
+			if (View.ViewState->IsRoundRobinEnabled() &&
+				!View.bIsSceneCapture && // We only round-robin on the main renderer (not scene captures)
+				!View.bIgnoreExistingQueries && // We do not alternate occlusion queries when we want to refresh the occlusion history
+				(View.StereoPass == eSSP_LEFT_EYE || View.StereoPass == eSSP_RIGHT_EYE)) // Only relevant to stereo views
+			{
+				// For even frames, prevent left eye from occlusion querying
+				// For odd frames, prevent right eye from occlusion querying
+				const bool FrameParity = View.ViewState->PrevFrameNumber % 2;
+				bSubmitQueries &= (FrameParity && View.StereoPass == eSSP_LEFT_EYE ||
+								  !FrameParity && View.StereoPass == eSSP_RIGHT_EYE);
+			}
+
+			NumOccludedPrimitives += FetchVisibilityForPrimitives(Scene, View, bSubmitQueries, bHZBOcclusion);
 
 			if( bHZBOcclusion )
 			{
@@ -2694,6 +2732,16 @@ void FSceneRenderer::PreVisibilityFrameSetup(FRHICommandListImmediate& RHICmdLis
 				View.bIgnoreExistingQueries = true;
 				View.bDisableDistanceBasedFadeTransitions = true;
 			}
+
+			// Turn on/off round-robin occlusion querying in the ViewState
+			static const auto CVarRROCC = IConsoleManager::Get().FindTConsoleVariableDataInt(TEXT("vr.RoundRobinOcclusion"));
+			const bool bEnableRoundRobin = CVarRROCC ? (CVarRROCC->GetValueOnAnyThread() != false) : false;
+			if (bEnableRoundRobin != ViewState->IsRoundRobinEnabled())
+			{
+				ViewState->UpdateRoundRobin(bEnableRoundRobin);
+				View.bIgnoreExistingQueries = true;
+			}
+
 			ViewState->PrevViewMatrixForOcclusionQuery = View.ViewMatrices.GetViewMatrix();
 			ViewState->PrevViewOriginForOcclusionQuery = View.ViewMatrices.GetViewOrigin();
 				
