Fixed some compiler warnings

diff --git a/Engine/Source/Runtime/Renderer/Private/SceneVisibility.cpp b/Engine/Source/Runtime/Renderer/Private/SceneVisibility.cpp
--- a/Engine/Source/Runtime/Renderer/Private/SceneVisibility.cpp
+++ b/Engine/Source/Runtime/Renderer/Private/SceneVisibility.cpp
@@ -600,12 +600,12 @@ struct FVisForPrimParams
 		, NumToProcess(InNumToProcess)
 		, bSubmitQueries(bInSubmitQueries)
 		, bHZBOcclusion(bInHZBOcclusion)		
+		, bNeedsScanOnRead(false)
 		, InsertPrimitiveOcclusionHistory(OutOcclusionHistory)
 		, QueriesToRelease(OutQueriesToRelease)
 		, HZBBoundsToAdd(OutHZBBounds)
 		, QueriesToAdd(OutQueriesToRun)	
 		, SubIsOccluded(OutSubIsOccluded)
-		, bNeedsScanOnRead(false)
 	{
 
 	}
@@ -1558,8 +1558,8 @@ static int32 OcclusionCull(FRHICommandListImmediate& RHICmdList, const FScene* S
 				// For even frames, prevent left eye from occlusion querying
 				// For odd frames, prevent right eye from occlusion querying
 				const bool FrameParity = View.ViewState->PrevFrameNumber % 2;
-				bSubmitQueries &= (FrameParity && View.StereoPass == eSSP_LEFT_EYE ||
-								  !FrameParity && View.StereoPass == eSSP_RIGHT_EYE);
+				bSubmitQueries &= (FrameParity && View.StereoPass == eSSP_LEFT_EYE) ||
+								  (!FrameParity && View.StereoPass == eSSP_RIGHT_EYE);
 			}
 
 			NumOccludedPrimitives += FetchVisibilityForPrimitives(Scene, View, bSubmitQueries, bHZBOcclusion);
