#DEV-VR - Fixing CIS error (forcing value to bool)

diff --git a/Engine/Source/Runtime/Renderer/Private/SceneVisibility.cpp b/Engine/Source/Runtime/Renderer/Private/SceneVisibility.cpp
--- a/Engine/Source/Runtime/Renderer/Private/SceneVisibility.cpp
+++ b/Engine/Source/Runtime/Renderer/Private/SceneVisibility.cpp
@@ -1557,7 +1557,7 @@ static int32 OcclusionCull(FRHICommandListImmediate& RHICmdList, const FScene* S
 			{
 				// For even frames, prevent left eye from occlusion querying
 				// For odd frames, prevent right eye from occlusion querying
-				const bool FrameParity = View.ViewState->PrevFrameNumber % 2;
+				const bool FrameParity = ((View.ViewState->PrevFrameNumber & 0x01) == 1);
 				bSubmitQueries &= (FrameParity && View.StereoPass == eSSP_LEFT_EYE) ||
 								  (!FrameParity && View.StereoPass == eSSP_RIGHT_EYE);
 			}
