Use the resolved depth texture as a temporary stencil buffer for decals.
This fixes a sample count inconsistency between the render target color/depth attachments.

diff --git a/Engine/Source/Runtime/Renderer/Private/CompositionLighting/PostProcessDeferredDecals.cpp b/Engine/Source/Runtime/Renderer/Private/CompositionLighting/PostProcessDeferredDecals.cpp
--- a/Engine/Source/Runtime/Renderer/Private/CompositionLighting/PostProcessDeferredDecals.cpp
+++ b/Engine/Source/Runtime/Renderer/Private/CompositionLighting/PostProcessDeferredDecals.cpp
@@ -1059,7 +1059,9 @@ void FDecalRenderTargetManager::SetRenderTargetMode(FDecalRenderingCommon::ERend
 		TargetsToResolve[DBufferAIndex] = SceneContext.DBufferA->GetRenderTargetItem().TargetableTexture;
 		TargetsToResolve[DBufferBIndex] = SceneContext.DBufferB->GetRenderTargetItem().TargetableTexture;
 		TargetsToResolve[DBufferCIndex] = SceneContext.DBufferC->GetRenderTargetItem().TargetableTexture;
-		SetRenderTargets(RHICmdList, 3, &TargetsToResolve[DBufferAIndex], SceneContext.GetSceneDepthSurface(), ESimpleRenderTargetMode::EExistingColorAndDepth, FExclusiveDepthStencil::DepthRead_StencilWrite, TargetsToTransitionWritable[CurrentRenderTargetMode]);
+
+		// We re-use the SceneDepthTexture here as a temporary stencil that's ensured to be single-sampled
+		SetRenderTargets(RHICmdList, 3, &TargetsToResolve[DBufferAIndex], SceneContext.GetSceneDepthTexture(), ESimpleRenderTargetMode::EExistingColorAndDepth, FExclusiveDepthStencil::DepthRead_StencilWrite, TargetsToTransitionWritable[CurrentRenderTargetMode]);
 		break;
 
 	case FDecalRenderingCommon::RTM_AmbientOcclusion:
