DVR - Fix for crash when going into immersive mode on mobile rendering preview

diff --git a/Engine/Source/Runtime/Renderer/Private/MobileShadingRenderer.cpp b/Engine/Source/Runtime/Renderer/Private/MobileShadingRenderer.cpp
--- a/Engine/Source/Runtime/Renderer/Private/MobileShadingRenderer.cpp
+++ b/Engine/Source/Runtime/Renderer/Private/MobileShadingRenderer.cpp
@@ -86,6 +86,28 @@ TUniformBufferRef<FMobileDirectionalLightShaderParameters>& GetNullMobileDirecti
 	return NullLightParams;
 }
 
+void FMobileSceneRenderer::PrepareViewVisibilityLists()
+{
+	// Prepare view's visibility lists.
+	// TODO: only do this when CSM + static is required.
+	for (auto& View : Views)
+	{
+		FMobileCSMVisibilityInfo& MobileCSMVisibilityInfo = View.MobileCSMVisibilityInfo;
+		// Init list of primitives that can receive Dynamic CSM.
+		MobileCSMVisibilityInfo.MobilePrimitiveCSMReceiverVisibilityMap.Init(false, View.PrimitiveVisibilityMap.Num());
+
+		ensure(View.StaticMeshVisibilityMap.Num());
+
+		// Init static mesh visibility info for CSM drawlist
+		MobileCSMVisibilityInfo.MobileCSMStaticMeshVisibilityMap.Init(false, View.StaticMeshVisibilityMap.Num());
+		MobileCSMVisibilityInfo.MobileCSMStaticBatchVisibility.AddZeroed(View.StaticMeshBatchVisibility.Num());
+
+		// Init static mesh visibility info for default drawlist that excludes meshes in CSM only drawlist.
+		MobileCSMVisibilityInfo.MobileNonCSMStaticMeshVisibilityMap = View.StaticMeshVisibilityMap;
+		MobileCSMVisibilityInfo.MobileNonCSMStaticBatchVisibility = View.StaticMeshBatchVisibility;
+	}
+}
+
 /**
  * Initialize scene's views.
  * Check visibility, sort translucent items, etc.
@@ -108,7 +130,12 @@ void FMobileSceneRenderer::InitViews(FRHICommandListImmediate& RHICmdList)
 	if (bDynamicShadows && !IsSimpleForwardShadingEnabled(GetFeatureLevelShaderPlatform(FeatureLevel)))
 	{
 		// Setup dynamic shadows.
-		InitDynamicShadows(RHICmdList);		
+		InitDynamicShadows(RHICmdList);
+	}
+	else
+	{
+		// TODO: only do this when CSM + static is required.
+		PrepareViewVisibilityLists();
 	}
 
 	// if we kicked off ILC update via task, wait and finalize.
diff --git a/Engine/Source/Runtime/Renderer/Private/SceneRendering.h b/Engine/Source/Runtime/Renderer/Private/SceneRendering.h
--- a/Engine/Source/Runtime/Renderer/Private/SceneRendering.h
+++ b/Engine/Source/Runtime/Renderer/Private/SceneRendering.h
@@ -1702,6 +1702,8 @@ protected:
 	/** Finds the visible dynamic shadows for each view. */
 	void InitDynamicShadows(FRHICommandListImmediate& RHICmdList);
 
+	void PrepareViewVisibilityLists();
+
 	/** Build visibility lists on CSM receivers and non-csm receivers. */
 	void BuildCSMVisibilityState(FLightSceneInfo* LightSceneInfo);
 
diff --git a/Engine/Source/Runtime/Renderer/Private/ShadowSetupMobile.cpp b/Engine/Source/Runtime/Renderer/Private/ShadowSetupMobile.cpp
--- a/Engine/Source/Runtime/Renderer/Private/ShadowSetupMobile.cpp
+++ b/Engine/Source/Runtime/Renderer/Private/ShadowSetupMobile.cpp
@@ -267,6 +267,7 @@ static void VisualizeMobileDynamicCSMSubjectCapsules(FViewInfo& View, FLightScen
 		break;
 	}
 }
+
 /** Finds the visible dynamic shadows for each view. */
 void FMobileSceneRenderer::InitDynamicShadows(FRHICommandListImmediate& RHICmdList)
 {
@@ -301,22 +302,7 @@ void FMobileSceneRenderer::InitDynamicShadows(FRHICommandListImmediate& RHICmdLi
 
 	FSceneRenderer::InitDynamicShadows(RHICmdList);
 
-	// Prepare view's visibility lists.
-	// TODO: only do this when CSM + static is required.
-	for (auto& View : Views)
-	{
-		FMobileCSMVisibilityInfo& MobileCSMVisibilityInfo = View.MobileCSMVisibilityInfo;
-		// Init list of primitives that can receive Dynamic CSM.
-		MobileCSMVisibilityInfo.MobilePrimitiveCSMReceiverVisibilityMap.Init(false, View.PrimitiveVisibilityMap.Num());
-
-		// Init static mesh visibility info for CSM drawlist
-		MobileCSMVisibilityInfo.MobileCSMStaticMeshVisibilityMap.Init(false, View.StaticMeshVisibilityMap.Num());
-		MobileCSMVisibilityInfo.MobileCSMStaticBatchVisibility.AddZeroed(View.StaticMeshBatchVisibility.Num());
-
-		// Init static mesh visibility info for default drawlist that excludes meshes in CSM only drawlist.
-		MobileCSMVisibilityInfo.MobileNonCSMStaticMeshVisibilityMap = View.StaticMeshVisibilityMap;
-		MobileCSMVisibilityInfo.MobileNonCSMStaticBatchVisibility = View.StaticMeshBatchVisibility;
-	}
+	PrepareViewVisibilityLists();
 
 	for (FLightSceneInfo* MobileDirectionalLightSceneInfo : Scene->MobileDirectionalLights)
 	{
