DVR - Remove potentially bogus ensure

diff --git a/Engine/Source/Runtime/Renderer/Private/MobileShadingRenderer.cpp b/Engine/Source/Runtime/Renderer/Private/MobileShadingRenderer.cpp
--- a/Engine/Source/Runtime/Renderer/Private/MobileShadingRenderer.cpp
+++ b/Engine/Source/Runtime/Renderer/Private/MobileShadingRenderer.cpp
@@ -96,8 +96,6 @@ void FMobileSceneRenderer::PrepareViewVisibilityLists()
 		// Init list of primitives that can receive Dynamic CSM.
 		MobileCSMVisibilityInfo.MobilePrimitiveCSMReceiverVisibilityMap.Init(false, View.PrimitiveVisibilityMap.Num());
 
-		ensure(View.StaticMeshVisibilityMap.Num());
-
 		// Init static mesh visibility info for CSM drawlist
 		MobileCSMVisibilityInfo.MobileCSMStaticMeshVisibilityMap.Init(false, View.StaticMeshVisibilityMap.Num());
 		MobileCSMVisibilityInfo.MobileCSMStaticBatchVisibility.AddZeroed(View.StaticMeshBatchVisibility.Num());
