Fix merge error in CL 4485106

diff --git a/Engine/Source/Runtime/Renderer/Private/CompositionLighting/PostProcessDeferredDecals.cpp b/Engine/Source/Runtime/Renderer/Private/CompositionLighting/PostProcessDeferredDecals.cpp
--- a/Engine/Source/Runtime/Renderer/Private/CompositionLighting/PostProcessDeferredDecals.cpp
+++ b/Engine/Source/Runtime/Renderer/Private/CompositionLighting/PostProcessDeferredDecals.cpp
@@ -934,6 +934,7 @@ void FRCPassPostProcessDeferredDecals::Process(FRenderingCompositePassContext& C
 					RHICmdList.DrawIndexedPrimitive(GetUnitCubeIndexBuffer(), PT_TriangleList, 0, 0, 8, 0, ARRAY_COUNT(GCubeIndices) / 3, 1);
 					RenderTargetManager.bGufferADirty |= (RenderTargetManager.TargetsToResolve[FDecalRenderTargetManager::GBufferAIndex] != nullptr);
 				}
+			}
 
 			if ((Scene.Decals.Num() > 0) && CurrentStage == DRS_Emissive)
 			{
