Bringing over 4.21.1 hotfix for multi-view view visibility.

diff --git a/Engine/Source/Runtime/Renderer/Private/SceneVisibility.cpp b/Engine/Source/Runtime/Renderer/Private/SceneVisibility.cpp
--- a/Engine/Source/Runtime/Renderer/Private/SceneVisibility.cpp
+++ b/Engine/Source/Runtime/Renderer/Private/SceneVisibility.cpp
@@ -2835,6 +2835,8 @@ void FSceneRenderer::ComputeViewVisibility(FRHICommandListImmediate& RHICmdList)
 		HasDynamicEditorMeshElementsMasks.AddZeroed(NumPrimitives);
 	}
 
+	const bool bIsInstancedStereo = (Views.Num() > 0) ? (Views[0].IsInstancedStereoPass() || Views[0].bIsMobileMultiViewEnabled) : false;
+
 	uint8 ViewBit = 0x1;
 	for (int32 ViewIndex = 0; ViewIndex < Views.Num(); ++ViewIndex, ViewBit <<= 1)
 	{
@@ -3060,7 +3062,7 @@ void FSceneRenderer::ComputeViewVisibility(FRHICommandListImmediate& RHICmdList)
 		}
 
 		// ISR views can't compute relevance until all views are frustum culled
-		if (!View.IsInstancedStereoPass())
+		if (!bIsInstancedStereo)
 		{
 			SCOPE_CYCLE_COUNTER(STAT_ViewRelevance);
 			ComputeAndMarkRelevanceForViewParallel(RHICmdList, Scene, View, ViewBit, HasDynamicMeshElementsMasks, HasDynamicEditorMeshElementsMasks, HasViewCustomDataMasks);
@@ -3093,7 +3095,7 @@ void FSceneRenderer::ComputeViewVisibility(FRHICommandListImmediate& RHICmdList)
 		View.bSceneHasDecals = (Scene->Decals.Num() > 0);
 	}
 
-	if (Views.Num() > 1 && Views[0].IsInstancedStereoPass())
+	if ((Views.Num() > 1) && bIsInstancedStereo)
 	{
 		// Ensure primitives from the right-eye view are visible in the left-eye (instanced) view
 		FSceneBitArray& LeftView = Views[0].PrimitiveVisibilityMap;
@@ -3114,7 +3116,7 @@ void FSceneRenderer::ComputeViewVisibility(FRHICommandListImmediate& RHICmdList)
 	ViewBit = 0x1;
 	for (FViewInfo& View : Views)
 	{
-		if (View.IsInstancedStereoPass())
+		if (bIsInstancedStereo)
 		{
 			SCOPE_CYCLE_COUNTER(STAT_ViewRelevance);
 			ComputeAndMarkRelevanceForViewParallel(RHICmdList, Scene, View, ViewBit, HasDynamicMeshElementsMasks, HasDynamicEditorMeshElementsMasks, HasViewCustomDataMasks);
