Exposed a symbol for cross module access

diff --git a/Engine/Source/Runtime/Renderer/Private/PostProcess/SceneFilterRendering.h b/Engine/Source/Runtime/Renderer/Private/PostProcess/SceneFilterRendering.h
--- a/Engine/Source/Runtime/Renderer/Private/PostProcess/SceneFilterRendering.h
+++ b/Engine/Source/Runtime/Renderer/Private/PostProcess/SceneFilterRendering.h
@@ -15,7 +15,7 @@
 #include "ShaderParameterUtils.h"
 
 /** Uniform buffer for computing the vertex positional and UV adjustments in the vertex shader. */
-BEGIN_GLOBAL_SHADER_PARAMETER_STRUCT( FDrawRectangleParameters,)
+BEGIN_GLOBAL_SHADER_PARAMETER_STRUCT( FDrawRectangleParameters, RENDERER_API)
 	SHADER_PARAMETER( FVector4, PosScaleBias )
 	SHADER_PARAMETER( FVector4, UVScaleBias )
 	SHADER_PARAMETER( FVector4, InvTargetSizeAndTextureSize )
