Minor improvements to quad stereo support

diff --git a/Engine/Source/Runtime/Engine/Public/StereoRendering.h b/Engine/Source/Runtime/Engine/Public/StereoRendering.h
--- a/Engine/Source/Runtime/Engine/Public/StereoRendering.h
+++ b/Engine/Source/Runtime/Engine/Public/StereoRendering.h
@@ -19,7 +19,9 @@ enum EStereoscopicPass
 	eSSP_FULL,
 	eSSP_LEFT_EYE,
 	eSSP_RIGHT_EYE,
-	eSSP_MONOSCOPIC_EYE
+	eSSP_MONOSCOPIC_EYE,
+	eSSP_LEFT_EYE_SIDE,
+	eSSP_RIGHT_EYE_SIDE
 };
 
 class IStereoRendering
diff --git a/Engine/Source/Runtime/Renderer/Private/SceneRendering.cpp b/Engine/Source/Runtime/Renderer/Private/SceneRendering.cpp
--- a/Engine/Source/Runtime/Renderer/Private/SceneRendering.cpp
+++ b/Engine/Source/Runtime/Renderer/Private/SceneRendering.cpp
@@ -1373,7 +1373,7 @@ void FViewInfo::SetupUniformBufferParameters(
 	extern FVector GetReflectionEnvironmentRoughnessMixingScaleBiasAndLargestWeight();
 	ViewUniformShaderParameters.ReflectionEnvironmentRoughnessMixingScaleBiasAndLargestWeight = GetReflectionEnvironmentRoughnessMixingScaleBiasAndLargestWeight();
 
-	ViewUniformShaderParameters.StereoPassIndex = (StereoPass <= eSSP_LEFT_EYE) ? 0 : (StereoPass == eSSP_RIGHT_EYE) ? 1 : StereoPass - eSSP_MONOSCOPIC_EYE + 1;
+	ViewUniformShaderParameters.StereoPassIndex = GEngine->StereoRenderingDevice ? GEngine->StereoRenderingDevice->GetViewIndexForPass(StereoPass) : 0;
 	ViewUniformShaderParameters.StereoIPD = StereoIPD;
 	
 	ViewUniformShaderParameters.PreIntegratedBRDF = GEngine->PreIntegratedSkinBRDFTexture->Resource->TextureRHI;
