Fix clamping bounds on depth-texture downsampling to the right-side of the source texture instead of the view width (was causing the right eye depth to not be sampled in VR).

diff --git a/Engine/Source/Runtime/Renderer/Private/DeferredShadingRenderer.cpp b/Engine/Source/Runtime/Renderer/Private/DeferredShadingRenderer.cpp
--- a/Engine/Source/Runtime/Renderer/Private/DeferredShadingRenderer.cpp
+++ b/Engine/Source/Runtime/Renderer/Private/DeferredShadingRenderer.cpp
@@ -2158,7 +2158,7 @@ void FDeferredShadingSceneRenderer::DownsampleDepthSurface(FRHICommandList& RHIC
 
 	SetGraphicsPipelineState(RHICmdList, GraphicsPSOInit);
 
-	PixelShader->SetParameters(RHICmdList, View, bUseMaxDepth, View.ViewRect.Size());
+	PixelShader->SetParameters(RHICmdList, View, bUseMaxDepth, View.ViewRect.Max);
 	const uint32 DownsampledX = FMath::TruncToInt(View.ViewRect.Min.X * ScaleFactor);
 	const uint32 DownsampledY = FMath::TruncToInt(View.ViewRect.Min.Y * ScaleFactor);
 	const uint32 DownsampledSizeX = FMath::TruncToInt(View.ViewRect.Width() * ScaleFactor);
