Fixes mobile preview when in sRGB mode and no editor elements are composited (G key press)

diff --git a/Engine/Source/Runtime/Renderer/Private/SceneRendering.cpp b/Engine/Source/Runtime/Renderer/Private/SceneRendering.cpp
--- a/Engine/Source/Runtime/Renderer/Private/SceneRendering.cpp
+++ b/Engine/Source/Runtime/Renderer/Private/SceneRendering.cpp
@@ -2930,7 +2930,8 @@ bool FSceneRenderer::ShouldCompositeEditorPrimitives(const FViewInfo& View)
 		    || View.TopViewMeshElements.Num() 
 		    || View.BatchedViewElements.HasPrimsToDraw() 
 		    || View.TopBatchedViewElements.HasPrimsToDraw() 
-		    || View.NumVisibleDynamicEditorPrimitives > 0)
+		    || View.NumVisibleDynamicEditorPrimitives > 0
+			|| IsMobileColorsRGB())
 	    {
 		    return true;
 	    }
