Allow MRMesh to be rendered exclusively in the depth prepass
Adds a depth prepass to the mobile renderer that's disabled on tiled GPUs

diff --git a/Engine/Source/Runtime/Engine/Classes/Components/PrimitiveComponent.h b/Engine/Source/Runtime/Engine/Classes/Components/PrimitiveComponent.h
--- a/Engine/Source/Runtime/Engine/Classes/Components/PrimitiveComponent.h
+++ b/Engine/Source/Runtime/Engine/Classes/Components/PrimitiveComponent.h
@@ -322,6 +322,10 @@ public:
 	UPROPERTY(EditAnywhere, AdvancedDisplay, BlueprintReadOnly, Category = Rendering)
 	uint8 bRenderInMainPass:1;
 
+	/** If true, this component will be rendered in the depth pass even if it's not rendered in the main pass */
+	UPROPERTY(EditAnywhere, AdvancedDisplay, BlueprintReadOnly, Category = Rendering)
+	uint8 bRenderInDepthPass:1;
+
 	/** Whether the primitive receives decals. */
 	UPROPERTY(EditAnywhere, AdvancedDisplay, BlueprintReadOnly, Category=Rendering)
 	uint8 bReceivesDecals:1;
diff --git a/Engine/Source/Runtime/Engine/Private/Components/PrimitiveComponent.cpp b/Engine/Source/Runtime/Engine/Private/Components/PrimitiveComponent.cpp
--- a/Engine/Source/Runtime/Engine/Private/Components/PrimitiveComponent.cpp
+++ b/Engine/Source/Runtime/Engine/Private/Components/PrimitiveComponent.cpp
@@ -329,6 +329,7 @@ UPrimitiveComponent::UPrimitiveComponent(const FObjectInitializer& ObjectInitial
 	bVisibleInReflectionCaptures = true;
 	bVisibleInRayTracing = true;
 	bRenderInMainPass = true;
+	bRenderInDepthPass = true;
 	VisibilityId = INDEX_NONE;
 #if WITH_EDITORONLY_DATA
 	CanBeCharacterBase_DEPRECATED = ECB_Yes;
diff --git a/Engine/Source/Runtime/Engine/Private/PrimitiveSceneProxy.cpp b/Engine/Source/Runtime/Engine/Private/PrimitiveSceneProxy.cpp
--- a/Engine/Source/Runtime/Engine/Private/PrimitiveSceneProxy.cpp
+++ b/Engine/Source/Runtime/Engine/Private/PrimitiveSceneProxy.cpp
@@ -100,6 +100,7 @@ FPrimitiveSceneProxy::FPrimitiveSceneProxy(const UPrimitiveComponent* InComponen
 ,	bStaticLighting(InComponent->HasStaticLighting())
 ,	bVisibleInReflectionCaptures(InComponent->bVisibleInReflectionCaptures)
 ,	bVisibleInRayTracing(InComponent->bVisibleInRayTracing)
+,	bRenderInDepthPass(InComponent->bRenderInDepthPass)
 ,	bRenderInMainPass(InComponent->bRenderInMainPass)
 ,	bRequiresVisibleLevelToRender(false)
 ,	bIsComponentLevelVisible(false)
diff --git a/Engine/Source/Runtime/Engine/Public/PrimitiveSceneProxy.h b/Engine/Source/Runtime/Engine/Public/PrimitiveSceneProxy.h
--- a/Engine/Source/Runtime/Engine/Public/PrimitiveSceneProxy.h
+++ b/Engine/Source/Runtime/Engine/Public/PrimitiveSceneProxy.h
@@ -492,6 +492,7 @@ public:
 	inline bool IsVisibleInReflectionCaptures() const { return bVisibleInReflectionCaptures; }
 	inline bool IsVisibleInRayTracing() const { return bVisibleInRayTracing; }
 	inline bool ShouldRenderInMainPass() const { return bRenderInMainPass; }
+	inline bool ShouldRenderInDepthPass() const { return bRenderInMainPass || bRenderInDepthPass; }
 	inline bool IsCollisionEnabled() const { return bCollisionEnabled; }
 	inline bool IsHovered() const { return bHovered; }
 	inline bool IsOwnedBy(const AActor* Actor) const { return Owners.Find(Actor) != INDEX_NONE; }
@@ -820,6 +821,9 @@ private:
 	/** If true, this component will be visible in ray tracing effects. Turning this off will remove it from ray traced reflections, shadows, etc. */
 	uint8 bVisibleInRayTracing : 1;
 
+	/** If true this primitive Renders in the depthPass */
+	uint8 bRenderInDepthPass : 1;
+
 	/** If true this primitive Renders in the mainPass */
 	uint8 bRenderInMainPass : 1;
 
diff --git a/Engine/Source/Runtime/Engine/Public/PrimitiveViewRelevance.h b/Engine/Source/Runtime/Engine/Public/PrimitiveViewRelevance.h
--- a/Engine/Source/Runtime/Engine/Public/PrimitiveViewRelevance.h
+++ b/Engine/Source/Runtime/Engine/Public/PrimitiveViewRelevance.h
@@ -51,6 +51,8 @@ struct FPrimitiveViewRelevance
 	uint32 bVelocityRelevance : 1;
 	/** The primitive should render to the custom depth pass. */
 	uint32 bRenderCustomDepth : 1;
+	/** The primitive should render to the depth prepass even if it's not rendered in the main pass. */
+	uint32 bRenderInDepthPass : 1;
 	/** The primitive should render to the base pass / normal depth / velocity rendering. */
 	uint32 bRenderInMainPass : 1;
 	/** The primitive has materials using the volume domain. */
diff --git a/Engine/Source/Runtime/MRMesh/Private/MRMeshComponent.cpp b/Engine/Source/Runtime/MRMesh/Private/MRMeshComponent.cpp
--- a/Engine/Source/Runtime/MRMesh/Private/MRMeshComponent.cpp
+++ b/Engine/Source/Runtime/MRMesh/Private/MRMeshComponent.cpp
@@ -403,9 +403,7 @@ private:
 							BatchElement.IndexBuffer = &Section->IndexBuffer;
 							Mesh.bWireframe = bUseWireframe;
 							Mesh.bUseAsOccluder = bEnableOcclusion;
-
-// Waiting for Jules' changes
-//							Mesh.bUseForDepthPass = bEnableOcclusion;
+							Mesh.bUseForDepthPass = bEnableOcclusion;
 
 							Mesh.VertexFactory = &Section->VertexFactory;
 							Mesh.MaterialRenderProxy = MaterialProxy;
@@ -461,10 +459,7 @@ private:
 		Result.bDynamicRelevance = true;
 		// If there is a material set that is not the default material, then this wants to be rendered in the main pass
 		Result.bRenderInMainPass = (bUseWireframe || MaterialToUse != UMaterial::GetDefaultMaterial(MD_Surface)) && ShouldRenderInMainPass();
-
-// Waiting for Jules' changes
-//		Result.bRenderInDepthPass = bEnableOcclusion;
-
+		Result.bRenderInDepthPass = bEnableOcclusion;
 		Result.bUsesLightingChannels = GetLightingChannelMask() != GetDefaultLightingChannelMask();
 		Result.bRenderCustomDepth = ShouldRenderCustomDepth();
 		Result.bSeparateTranslucencyRelevance = MaterialToUse->GetMaterial()->bEnableSeparateTranslucency;
diff --git a/Engine/Source/Runtime/Renderer/Private/DepthRendering.cpp b/Engine/Source/Runtime/Renderer/Private/DepthRendering.cpp
--- a/Engine/Source/Runtime/Renderer/Private/DepthRendering.cpp
+++ b/Engine/Source/Runtime/Renderer/Private/DepthRendering.cpp
@@ -782,7 +782,7 @@ void FDepthPassMeshProcessor::AddMeshBatch(const FMeshBatch& RESTRICT MeshBatch,
 		const bool bIsTranslucent = IsTranslucentBlendMode(BlendMode);
 
 		if (!bIsTranslucent
-			&& (!PrimitiveSceneProxy || PrimitiveSceneProxy->ShouldRenderInMainPass())
+			&& (!PrimitiveSceneProxy || PrimitiveSceneProxy->ShouldRenderInDepthPass())
 			&& ShouldIncludeDomainInMeshPass(Material.GetMaterialDomain()))
 		{
 			if (BlendMode == BLEND_Opaque
diff --git a/Engine/Source/Runtime/Renderer/Private/SceneVisibility.cpp b/Engine/Source/Runtime/Renderer/Private/SceneVisibility.cpp
--- a/Engine/Source/Runtime/Renderer/Private/SceneVisibility.cpp
+++ b/Engine/Source/Runtime/Renderer/Private/SceneVisibility.cpp
@@ -2207,7 +2207,7 @@ struct FRelevancePacket
 					if (ViewRelevance.bDrawRelevance)
 					{
 						if ((StaticMeshRelevance.bUseForMaterial || StaticMeshRelevance.bUseAsOccluder) 
-							&& (ViewRelevance.bRenderInMainPass || ViewRelevance.bRenderCustomDepth) 
+							&& (ViewRelevance.bRenderInMainPass || ViewRelevance.bRenderCustomDepth || ViewRelevance.bRenderInDepthPass)
 							&& !bHiddenByHLODFade)
 						{
 							if (StaticMeshRelevance.bUseForDepthPass && bDrawDepthOnly)
@@ -2216,7 +2216,7 @@ struct FRelevancePacket
 							}
 
 							// Mark static mesh as visible for rendering
-							if (StaticMeshRelevance.bUseForMaterial)
+							if (StaticMeshRelevance.bUseForMaterial && (ViewRelevance.bRenderInMainPass || ViewRelevance.bRenderCustomDepth))
 							{
 								DrawCommandPacket.AddCommandsForMesh(PrimitiveIndex, PrimitiveSceneInfo, StaticMeshRelevance, StaticMesh, Scene, bCanCache, EMeshPass::BasePass);
 								MarkMask |= EMarkMaskBits::StaticMeshVisibilityMapMask;
@@ -2622,58 +2622,61 @@ void ComputeDynamicMeshRelevance(EShadingPath ShadingPath, bool bAddLightmapDens
 {
 	const int32 NumElements = MeshBatch.Mesh->Elements.Num();
 
-	if (ViewRelevance.bDrawRelevance && (ViewRelevance.bRenderInMainPass || ViewRelevance.bRenderCustomDepth))
+	if (ViewRelevance.bDrawRelevance && (ViewRelevance.bRenderInMainPass || ViewRelevance.bRenderCustomDepth || ViewRelevance.bRenderInDepthPass))
 	{
 		PassMask.Set(EMeshPass::DepthPass);
 		View.NumVisibleDynamicMeshElements[EMeshPass::DepthPass] += NumElements;
 
-		PassMask.Set(EMeshPass::BasePass);
-		View.NumVisibleDynamicMeshElements[EMeshPass::BasePass] += NumElements;
-
-		if (ShadingPath == EShadingPath::Mobile)
+		if (ViewRelevance.bRenderInMainPass || ViewRelevance.bRenderCustomDepth)
 		{
-			PassMask.Set(EMeshPass::MobileBasePassCSM);
-			View.NumVisibleDynamicMeshElements[EMeshPass::MobileBasePassCSM] += NumElements;
-		}
+			PassMask.Set(EMeshPass::BasePass);
+			View.NumVisibleDynamicMeshElements[EMeshPass::BasePass] += NumElements;
 
-		if (ViewRelevance.bRenderCustomDepth)
-		{
-			PassMask.Set(EMeshPass::CustomDepth);
-			View.NumVisibleDynamicMeshElements[EMeshPass::CustomDepth] += NumElements;
-		}
+			if (ShadingPath == EShadingPath::Mobile)
+			{
+				PassMask.Set(EMeshPass::MobileBasePassCSM);
+				View.NumVisibleDynamicMeshElements[EMeshPass::MobileBasePassCSM] += NumElements;
+			}
 
-		if (bAddLightmapDensityCommands)
-		{
-			PassMask.Set(EMeshPass::LightmapDensity);
-			View.NumVisibleDynamicMeshElements[EMeshPass::LightmapDensity] += NumElements;
-		}
+			if (ViewRelevance.bRenderCustomDepth)
+			{
+				PassMask.Set(EMeshPass::CustomDepth);
+				View.NumVisibleDynamicMeshElements[EMeshPass::CustomDepth] += NumElements;
+			}
+
+			if (bAddLightmapDensityCommands)
+			{
+				PassMask.Set(EMeshPass::LightmapDensity);
+				View.NumVisibleDynamicMeshElements[EMeshPass::LightmapDensity] += NumElements;
+			}
 #if !(UE_BUILD_SHIPPING || UE_BUILD_TEST)
-		else if (View.Family->UseDebugViewPS())
-		{
-			PassMask.Set(EMeshPass::DebugViewMode);
-			View.NumVisibleDynamicMeshElements[EMeshPass::DebugViewMode] += NumElements;
-		}
+			else if (View.Family->UseDebugViewPS())
+			{
+				PassMask.Set(EMeshPass::DebugViewMode);
+				View.NumVisibleDynamicMeshElements[EMeshPass::DebugViewMode] += NumElements;
+			}
 #endif
 
 #if WITH_EDITOR
-		if (View.bAllowTranslucentPrimitivesInHitProxy)
-		{
-			PassMask.Set(EMeshPass::HitProxy);
-			View.NumVisibleDynamicMeshElements[EMeshPass::HitProxy] += NumElements;
-		}
-		else
-		{
-			PassMask.Set(EMeshPass::HitProxyOpaqueOnly);
-			View.NumVisibleDynamicMeshElements[EMeshPass::HitProxyOpaqueOnly] += NumElements;
-		}
+			if (View.bAllowTranslucentPrimitivesInHitProxy)
+			{
+				PassMask.Set(EMeshPass::HitProxy);
+				View.NumVisibleDynamicMeshElements[EMeshPass::HitProxy] += NumElements;
+			}
+			else
+			{
+				PassMask.Set(EMeshPass::HitProxyOpaqueOnly);
+				View.NumVisibleDynamicMeshElements[EMeshPass::HitProxyOpaqueOnly] += NumElements;
+			}
 #endif
 
-		if (ViewRelevance.bVelocityRelevance
+			if (ViewRelevance.bVelocityRelevance
 				&& FVelocityRendering::PrimitiveHasVelocity(View.GetFeatureLevel(), PrimitiveSceneInfo)
 				&& FVelocityRendering::PrimitiveHasVelocityForView(View, Bounds.BoxSphereBounds, PrimitiveSceneInfo))
-		{
-			PassMask.Set(EMeshPass::Velocity);
-			View.NumVisibleDynamicMeshElements[EMeshPass::Velocity] += NumElements;
+			{
+				PassMask.Set(EMeshPass::Velocity);
+				View.NumVisibleDynamicMeshElements[EMeshPass::Velocity] += NumElements;
+			}
 		}
 	}
 
