Fix dynamic decals rendering multiple times in left eye (this shader does not support instanced stereo, so disable instancing in this case).

diff --git a/Engine/Source/Runtime/Renderer/Private/CompositionLighting/PostProcessMeshDecals.cpp b/Engine/Source/Runtime/Renderer/Private/CompositionLighting/PostProcessMeshDecals.cpp
--- a/Engine/Source/Runtime/Renderer/Private/CompositionLighting/PostProcessMeshDecals.cpp
+++ b/Engine/Source/Runtime/Renderer/Private/CompositionLighting/PostProcessMeshDecals.cpp
@@ -365,7 +365,6 @@ void DrawDecalMeshCommands(FRenderingCompositePassContext& Context, EDecalRender
 	Context.SetViewportAndCallRHI(Context.View.ViewRect);
 	RHICmdList.ApplyCachedRenderTargets(GraphicsPSOInit);
 
-
 	DrawDynamicMeshPass(View, RHICmdList,
 		[&View, CurrentDecalStage, RenderTargetMode](FDynamicPassMeshDrawListContext* DynamicMeshPassContext)
 	{
@@ -384,7 +383,7 @@ void DrawDecalMeshCommands(FRenderingCompositePassContext& Context, EDecalRender
 
 			PassMeshProcessor.AddMeshBatch(*Mesh, DefaultBatchElementMask, PrimitiveSceneProxy);
 		}
-	});
+	}, true);
 }
 
 void RenderMeshDecals(FRenderingCompositePassContext& Context, EDecalRenderStage CurrentDecalStage)
diff --git a/Engine/Source/Runtime/Renderer/Public/MeshPassProcessor.inl b/Engine/Source/Runtime/Renderer/Public/MeshPassProcessor.inl
--- a/Engine/Source/Runtime/Renderer/Public/MeshPassProcessor.inl
+++ b/Engine/Source/Runtime/Renderer/Public/MeshPassProcessor.inl
@@ -139,7 +139,7 @@ void FMeshPassProcessor::BuildMeshDrawCommands(
 * Does many dynamic allocations - do not use for game rendering.
 */
 template<typename LambdaType>
-void DrawDynamicMeshPass(const FSceneView& View, FRHICommandList& RHICmdList, const LambdaType& BuildPassProcessorLambda)
+void DrawDynamicMeshPass(const FSceneView& View, FRHICommandList& RHICmdList, const LambdaType& BuildPassProcessorLambda, bool bForceStereoInstancingOff = false)
 {
 	FDynamicMeshDrawCommandStorage DynamicMeshDrawCommandStorage;
 	FMeshCommandOneFrameArray VisibleMeshDrawCommands;
@@ -150,6 +150,6 @@ void DrawDynamicMeshPass(const FSceneView& View, FRHICommandList& RHICmdList, co
 	BuildPassProcessorLambda(&DynamicMeshPassContext);
 
 	// We assume all dynamic passes are in stereo if it is enabled in the view, so we apply ISR to them
-	const uint32 InstanceFactor = View.IsInstancedStereoPass() ? 2 : 1;
+	const uint32 InstanceFactor = (!bForceStereoInstancingOff && View.IsInstancedStereoPass()) ? 2 : 1;
 	DrawDynamicMeshPassPrivate(View, RHICmdList, VisibleMeshDrawCommands, DynamicMeshDrawCommandStorage, GraphicsMinimalPipelineStateSet, InstanceFactor);
 }
