ShaderPrint debugger to help with compute shader debugging

diff --git a/Engine/Shaders/Private/MiniFontCommon.ush b/Engine/Shaders/Private/MiniFontCommon.ush
--- a/Engine/Shaders/Private/MiniFontCommon.ush
+++ b/Engine/Shaders/Private/MiniFontCommon.ush
@@ -1,15 +1,14 @@
 // Copyright 1998-2019 Epic Games, Inc. All Rights Reserved.
 
 /*=============================================================================
-	TinyFontCommon.usf: TinyFont shared functions and structures.
+	MiniFontCommon.usf: TinyFont shared functions and structures.
 =============================================================================*/
 
-#ifndef __TINY_FONT_COMMON__
-#define __TINY_FONT_COMMON__
+#pragma once
 
 Texture2D MiniFontTexture;
 
-// useful with PrintCharacter
+// CharacterIDs for the MiniFont
 #define _0_ 0
 #define _1_ 1
 #define _2_ 2
@@ -46,6 +45,19 @@ Texture2D MiniFontTexture;
 #define _X_ 33
 #define _Y_ 34
 #define _Z_ 35
+#define _MINUS_ 59
+#define _COMMA_ 60
+#define _DOT_ 61
+#define _PLUS_ 62
+#define _SPC_ 63
+
+// sample from the MiniFont texture
+// @param CharacterID - the character to print
+// @param Position - the position within the 8x8 character
+float SampleMiniFont(int CharacterID, uint2 Position)
+{
+	return MiniFontTexture.Load(int3(CharacterID * 8 + Position.x, Position.y, 0)).r;
+}
 
 // for printf debugging in the shader
 // @param LeftTop - is advanced, in pixels
@@ -55,7 +67,7 @@ void PrintCharacter(int2 PixelPos, inout float3 OutColor, float3 FontColor, inou
 
 	FLATTEN if (Rel.x < 8 && Rel.y < 8)
 	{
-		OutColor = lerp(OutColor, FontColor, MiniFontTexture.Load(int3(CharacterID * 8 + Rel.x, Rel.y, 0)).r);
+		OutColor = lerp(OutColor, FontColor, SampleMiniFont(CharacterID, Rel));
 	}
 
 	LeftTop.x += 8;
@@ -76,7 +88,6 @@ uint ExtractDigitFromFloat(float Number, float DigitValue)
 #endif
 }
 
-
 // for printf debugging in the shader, has to be positive
 // outputs a float number in the form: xxx.yyy
 // @param LeftTop - in pixels
@@ -84,17 +95,12 @@ void PrintFloat(int2 PixelPos, inout float3 OutColor, float3 FontColor, int2 Lef
 {
 	int2 Cursor = LeftTop;
 
-	float BorderDistance = ComputeDistanceToRect(PixelPos, LeftTop, int2(7, 1) * 8 - 1);
-
-	// black border around number
-	//	OutColor = lerp(0, OutColor, saturate(BorderDistance - 2));
-
 	// before period
 	PrintCharacter(PixelPos, OutColor, FontColor, Cursor, ExtractDigitFromFloat(Number, 100));
 	PrintCharacter(PixelPos, OutColor, FontColor, Cursor, ExtractDigitFromFloat(Number, 10));
 	PrintCharacter(PixelPos, OutColor, FontColor, Cursor, ExtractDigitFromFloat(Number, 1));
 	// period
-	PrintCharacter(PixelPos, OutColor, FontColor, Cursor, 512 / 8 - 3);
+	PrintCharacter(PixelPos, OutColor, FontColor, Cursor, _DOT_);
 	// after period
 	PrintCharacter(PixelPos, OutColor, FontColor, Cursor, ExtractDigitFromFloat(Number, 0.1));
 	PrintCharacter(PixelPos, OutColor, FontColor, Cursor, ExtractDigitFromFloat(Number, 0.01));
@@ -108,11 +114,6 @@ void PrintFloatNoFraction(int2 PixelPos, inout float3 OutColor, float3 FontColor
 {
 	int2 Cursor = LeftTop;
 
-	float BorderDistance = ComputeDistanceToRect(PixelPos, LeftTop, int2(7, 1) * 8 - 1);
-
-	// black border around number
-	//	OutColor = lerp(0, OutColor, saturate(BorderDistance - 2));
-
 	// before period
 	PrintCharacter(PixelPos, OutColor, FontColor, Cursor, ExtractDigitFromFloat(Number, 10000));
 	PrintCharacter(PixelPos, OutColor, FontColor, Cursor, ExtractDigitFromFloat(Number, 1000));
@@ -129,11 +130,6 @@ void PrintSmallFloat(int2 PixelPos, inout float3 OutColor, float3 FontColor, int
 	Number = abs(Number) + 0.05; // Round up first digit
 	int2 Cursor = LeftTop;
 
-	float BorderDistance = ComputeDistanceToRect(PixelPos, LeftTop, int2(7, 1) * 8 - 1);
-
-	// black border around number
-	//	OutColor = lerp(0, OutColor, saturate(BorderDistance - 2));
-
 	// before period
 	FLATTEN
 	if (Number >= 10)
@@ -142,10 +138,7 @@ void PrintSmallFloat(int2 PixelPos, inout float3 OutColor, float3 FontColor, int
 	}
 	PrintCharacter(PixelPos, OutColor, FontColor, Cursor, ExtractDigitFromFloat(Number, 1));
 	// period
-	PrintCharacter(PixelPos, OutColor, FontColor, Cursor, 512 / 8 - 3);
+	PrintCharacter(PixelPos, OutColor, FontColor, Cursor, _DOT_);
 	// after period
 	PrintCharacter(PixelPos, OutColor, FontColor, Cursor, ExtractDigitFromFloat(Number, 0.1));
 }
-
-
-#endif // __TINY_FONT_COMMON__
diff --git a/Engine/Shaders/Private/ShaderPrintCommon.ush b/Engine/Shaders/Private/ShaderPrintCommon.ush
new file mode 100644
--- /dev/null
+++ b/Engine/Shaders/Private/ShaderPrintCommon.ush
@@ -0,0 +1,461 @@
+/*=============================================================================
+	ShaderPrintCommon.ush:
+	Include this to be able to call ShaderPrint() from arbitrary shaders
+=============================================================================*/
+
+#pragma once
+
+// Include the MiniFont symbol definitions
+#include "MiniFontCommon.ush"
+
+// --------------------------------------------------------------------------
+// Uniforms and structures
+
+// Needs to match C++ code in ShaderPrint.cpp
+struct ShaderPrintItem
+{
+	float2 ScreenPos; // Position in normalized coordinates
+	int Value;        // Cast to value or symbol
+	int Type;         // SHADER_PRINT_TYPE_* defines how to read Value
+};
+
+RWStructuredBuffer<ShaderPrintItem> RWValuesBuffer;
+
+// Content of ShaderPrintItem.Type defines what to cast ShaderPrintItem.Value as
+#define SHADER_PRINT_TYPE_SYMBOL 0
+#define SHADER_PRINT_TYPE_FLOAT  1
+#define SHADER_PRINT_TYPE_INT    2
+#define SHADER_PRINT_TYPE_UINT   3
+#define SHADER_PRINT_TYPE_HEX    4
+
+#include "/Engine/Generated/UniformBuffers/ShaderPrint.ush"
+//float4 FontSize;      // .xy character size, .zw character spacing (in normalized coordinates)
+//int MaxValueCount;    // maximum number of values that we can print in one frame
+//int MaxSymbolCount;   // maximum number of symbols (after conversion from characters) that we can print in one frame
+
+static const float2 ShaderPrintCursorStart = float2(0.05f, 0.05f);
+static const float ShaderPrintTabCount = 12.f;
+
+// 'Global' values for tracking printing state (per thread)
+static float2 ShaderPrintCursorPos = ShaderPrintCursorStart;
+static bool ShaderPrintFilterEnable = true;
+
+// --------------------------------------------------------------------------
+// Internal helpers
+
+void ShaderPrint_Internal(in ShaderPrintItem Item)
+{
+	// If MaxValueCount is 0 then we don't reset the buffer counter so early out here
+	if (ShaderPrint.MaxValueCount == 0)
+	{
+		return;
+	}
+
+	// Buffer counter is stored in first element .Value
+	int IndexToStore = 0;
+	InterlockedAdd(RWValuesBuffer[0].Value, 1, IndexToStore);
+
+	// Prevent writing off the buffer
+	// Note that counter still increases so need clamp when reading it in later passes
+	if (IndexToStore >= ShaderPrint.MaxValueCount)
+	{
+		return;
+	}
+
+	// Because counter is in first element, read/writes need to be offset by 1
+	RWValuesBuffer[IndexToStore + 1] = Item;
+}
+
+void ShaderPrint_Internal(in float2 ScreenPos, in int Value, in int Type)
+{
+	ShaderPrintItem Item;
+	Item.ScreenPos = ScreenPos;
+	Item.Value = Value;
+	Item.Type = Type;
+	ShaderPrint_Internal(Item);
+}
+
+// --------------------------------------------------------------------------
+// Symbol printing
+
+float2 ShaderPrintSymbol(in float2 ScreenPos, in int Symbol)
+{
+	if (ShaderPrintFilterEnable)
+	{
+		ShaderPrint_Internal(ScreenPos, Symbol, SHADER_PRINT_TYPE_SYMBOL);
+		ScreenPos.x += ShaderPrint.FontSize.z;
+	}
+	return ScreenPos;
+}
+
+void ShaderPrintSymbol(in int symbol)
+{
+	ShaderPrintCursorPos = ShaderPrintSymbol(ShaderPrintCursorPos, symbol);
+}
+
+float2 ShaderPrintHelloWorld(in float2 ScreenPos)
+{
+	if (ShaderPrintFilterEnable)
+	{
+		ScreenPos = ShaderPrintSymbol(ScreenPos, _H_);
+		ScreenPos = ShaderPrintSymbol(ScreenPos, _E_);
+		ScreenPos = ShaderPrintSymbol(ScreenPos, _L_);
+		ScreenPos = ShaderPrintSymbol(ScreenPos, _L_);
+		ScreenPos = ShaderPrintSymbol(ScreenPos, _O_);
+		ScreenPos = ShaderPrintSymbol(ScreenPos, _SPC_);
+		ScreenPos = ShaderPrintSymbol(ScreenPos, _W_);
+		ScreenPos = ShaderPrintSymbol(ScreenPos, _O_);
+		ScreenPos = ShaderPrintSymbol(ScreenPos, _R_);
+		ScreenPos = ShaderPrintSymbol(ScreenPos, _L_);
+		ScreenPos = ShaderPrintSymbol(ScreenPos, _D_);
+		ScreenPos = ShaderPrintSymbol(ScreenPos, _SPC_);
+	}
+	return ScreenPos;
+}
+
+void ShaderPrintHelloWorld()
+{
+	ShaderPrintCursorPos = ShaderPrintHelloWorld(ShaderPrintCursorPos);
+}
+
+// --------------------------------------------------------------------------
+// Value printing
+
+float2 ShaderPrint(in float2 ScreenPos, in float Value, in int Type = SHADER_PRINT_TYPE_FLOAT)
+{
+	if (ShaderPrintFilterEnable)
+	{
+		ShaderPrint_Internal(ScreenPos, asint(Value), Type);
+		ScreenPos.x += ShaderPrint.FontSize.z * ShaderPrintTabCount;
+	}
+	return ScreenPos;
+}
+
+void ShaderPrint(in float Value, in int Type = SHADER_PRINT_TYPE_FLOAT)
+{
+	ShaderPrintCursorPos = ShaderPrint(ShaderPrintCursorPos, Value);
+}
+
+float2 ShaderPrint(in float2 ScreenPos, in int Value, in int Type = SHADER_PRINT_TYPE_INT)
+{
+	if (ShaderPrintFilterEnable)
+	{
+		ShaderPrint_Internal(ScreenPos, Value, Type);
+		ScreenPos.x += ShaderPrint.FontSize.z * ShaderPrintTabCount;
+	}
+	return ScreenPos;
+}
+
+void ShaderPrint(in int Value, in int Type = SHADER_PRINT_TYPE_INT)
+{
+	ShaderPrintCursorPos = ShaderPrint(ShaderPrintCursorPos, Value);
+}
+
+float2 ShaderPrint(in float2 ScreenPos, in uint Value, in int Type = SHADER_PRINT_TYPE_UINT)
+{
+	if (ShaderPrintFilterEnable)
+	{
+		ShaderPrint_Internal(ScreenPos, asint(Value), Type);
+		ScreenPos.x += ShaderPrint.FontSize.z * ShaderPrintTabCount;
+	}
+	return ScreenPos;
+}
+
+void ShaderPrint(in uint Value, in int Type = SHADER_PRINT_TYPE_UINT)
+{
+	ShaderPrintCursorPos = ShaderPrint(ShaderPrintCursorPos, Value, Type);
+}
+
+float2 ShaderPrint(in float2 ScreenPos, in bool Value)
+{
+	if (ShaderPrintFilterEnable)
+	{
+		ScreenPos = ShaderPrint(ScreenPos, Value ? 1U : 0U);
+	}
+	return ScreenPos;
+}
+
+void ShaderPrint(in bool Value)
+{
+	ShaderPrintCursorPos = ShaderPrint(ShaderPrintCursorPos, Value);
+}
+
+float2 ShaderPrint(in float2 ScreenPos, in float2 Value)
+{
+	if (ShaderPrintFilterEnable)
+	{
+		ScreenPos = ShaderPrint(ScreenPos, Value.x);
+		ScreenPos = ShaderPrint(ScreenPos, Value.y);
+	}
+	return ScreenPos;
+}
+
+void ShaderPrint(in float2 Value)
+{
+	ShaderPrintCursorPos = ShaderPrint(ShaderPrintCursorPos, Value);
+}
+
+float2 ShaderPrint(in float2 ScreenPos, in int2 Value)
+{
+	if (ShaderPrintFilterEnable)
+	{
+		ScreenPos = ShaderPrint(ScreenPos, Value.x);
+		ScreenPos = ShaderPrint(ScreenPos, Value.y);
+	}
+	return ScreenPos;
+}
+
+void ShaderPrint(in int2 Value)
+{
+	ShaderPrintCursorPos = ShaderPrint(ShaderPrintCursorPos, Value);
+}
+
+float2 ShaderPrint(in float2 ScreenPos, in uint2 Value)
+{
+	if (ShaderPrintFilterEnable)
+	{
+		ScreenPos = ShaderPrint(ScreenPos, (int)Value.x);
+		ScreenPos = ShaderPrint(ScreenPos, (int)Value.y);
+	}
+	return ScreenPos;
+}
+
+void ShaderPrint(in uint2 Value)
+{
+	ShaderPrintCursorPos = ShaderPrint(ShaderPrintCursorPos, Value);
+}
+
+float2 ShaderPrint(in float2 ScreenPos, in float3 Value)
+{
+	if (ShaderPrintFilterEnable)
+	{
+		ScreenPos = ShaderPrint(ScreenPos, Value.x);
+		ScreenPos = ShaderPrint(ScreenPos, Value.y);
+		ScreenPos = ShaderPrint(ScreenPos, Value.z);
+	}
+	return ScreenPos;
+}
+
+void ShaderPrint(in float3 Value)
+{
+	ShaderPrintCursorPos = ShaderPrint(ShaderPrintCursorPos, Value);
+}
+
+float2 ShaderPrint(in float2 ScreenPos, in int3 Value)
+{
+	if (ShaderPrintFilterEnable)
+	{
+		ScreenPos = ShaderPrint(ScreenPos, Value.x);
+		ScreenPos = ShaderPrint(ScreenPos, Value.y);
+		ScreenPos = ShaderPrint(ScreenPos, Value.z);
+	}
+	return ScreenPos;
+}
+
+void ShaderPrint(in int3 Value)
+{
+	ShaderPrintCursorPos = ShaderPrint(ShaderPrintCursorPos, Value);
+}
+
+float2 ShaderPrint(in float2 ScreenPos, in uint3 Value)
+{
+	if (ShaderPrintFilterEnable)
+	{
+		ScreenPos = ShaderPrint(ScreenPos, (int)Value.x);
+		ScreenPos = ShaderPrint(ScreenPos, (int)Value.y);
+		ScreenPos = ShaderPrint(ScreenPos, (int)Value.z);
+	}
+	return ScreenPos;
+}
+
+void ShaderPrint(in uint3 Value)
+{
+	ShaderPrintCursorPos = ShaderPrint(ShaderPrintCursorPos, Value);
+}
+
+float2 ShaderPrint(in float2 ScreenPos, in float4 Value)
+{
+	if (ShaderPrintFilterEnable)
+	{
+		ScreenPos = ShaderPrint(ScreenPos, Value.x);
+		ScreenPos = ShaderPrint(ScreenPos, Value.y);
+		ScreenPos = ShaderPrint(ScreenPos, Value.z);
+		ScreenPos = ShaderPrint(ScreenPos, Value.w);
+	}
+	return ScreenPos;
+}
+
+void ShaderPrint(in float4 Value)
+{
+	ShaderPrintCursorPos = ShaderPrint(ShaderPrintCursorPos, Value);
+}
+
+float2 ShaderPrint(in float2 ScreenPos, in int4 Value)
+{
+	if (ShaderPrintFilterEnable)
+	{
+		ScreenPos = ShaderPrint(ScreenPos, Value.x);
+		ScreenPos = ShaderPrint(ScreenPos, Value.y);
+		ScreenPos = ShaderPrint(ScreenPos, Value.z);
+		ScreenPos = ShaderPrint(ScreenPos, Value.w);
+	}
+	return ScreenPos;
+}
+
+void ShaderPrint(in int4 Value)
+{
+	ShaderPrintCursorPos = ShaderPrint(ShaderPrintCursorPos, Value);
+}
+
+float2 ShaderPrint(in float2 ScreenPos, in uint4 Value)
+{
+	if (ShaderPrintFilterEnable)
+	{
+		ScreenPos = ShaderPrint(ScreenPos, (int)Value.x);
+		ScreenPos = ShaderPrint(ScreenPos, (int)Value.y);
+		ScreenPos = ShaderPrint(ScreenPos, (int)Value.z);
+		ScreenPos = ShaderPrint(ScreenPos, (int)Value.w);
+	}
+	return ScreenPos;
+}
+
+void ShaderPrint(in uint4 Value)
+{
+	ShaderPrintCursorPos = ShaderPrint(ShaderPrintCursorPos, Value);
+}
+
+float2 ShaderPrint(in float2 ScreenPos, in float3x3 Value)
+{
+	if (ShaderPrintFilterEnable)
+	{
+		ShaderPrint(ScreenPos, Value[0]);
+		ScreenPos += float2(0.f, ShaderPrint.FontSize.w);
+		ShaderPrint(ScreenPos, Value[1]);
+		ScreenPos += float2(0.f, ShaderPrint.FontSize.w);
+		ShaderPrint(ScreenPos, Value[2]);
+		ScreenPos += float2(0.f, ShaderPrint.FontSize.w);
+	}
+	return ScreenPos;
+}
+
+void ShaderPrint(in float3x3 Value)
+{
+	ShaderPrintCursorPos = ShaderPrint(ShaderPrintCursorPos, Value);
+}
+
+float2 ShaderPrint(in float2 ScreenPos, in float4x3 Value)
+{
+	if (ShaderPrintFilterEnable)
+	{
+		ShaderPrint(ScreenPos, float4(Value[0], 0.f));
+		ScreenPos += float2(0.f, ShaderPrint.FontSize.w);
+		ShaderPrint(ScreenPos, float4(Value[1], 0.f));
+		ScreenPos += float2(0.f, ShaderPrint.FontSize.w);
+		ShaderPrint(ScreenPos, float4(Value[2], 0.f));
+		ScreenPos += float2(0.f, ShaderPrint.FontSize.w);
+		ShaderPrint(ScreenPos, float4(Value[3], 1.f));
+		ScreenPos += float2(0.f, ShaderPrint.FontSize.w);
+	}
+	return ScreenPos;
+}
+
+void ShaderPrint(in float4x3 Value)
+{
+	ShaderPrintCursorPos = ShaderPrint(ShaderPrintCursorPos, Value);
+}
+
+float2 ShaderPrint(in float2 ScreenPos, in float4x4 Value)
+{
+	if (ShaderPrintFilterEnable)
+	{
+		ShaderPrint(ScreenPos, Value[0]);
+		ScreenPos += float2(0.f, ShaderPrint.FontSize.w);
+		ShaderPrint(ScreenPos, Value[1]);
+		ScreenPos += float2(0.f, ShaderPrint.FontSize.w);
+		ShaderPrint(ScreenPos, Value[2]);
+		ScreenPos += float2(0.f, ShaderPrint.FontSize.w);
+		ShaderPrint(ScreenPos, Value[3]);
+		ScreenPos += float2(0.f, ShaderPrint.FontSize.w);
+	}
+	return ScreenPos;
+}
+
+void ShaderPrint(in float4x4 Value)
+{
+	ShaderPrintCursorPos = ShaderPrint(ShaderPrintCursorPos, Value);
+}
+
+// --------------------------------------------------------------------------
+// Formating helpers
+
+float2 ShaderPrintSpace(in float2 ScreenPos, in float Count = 1)
+{
+	if (ShaderPrintFilterEnable)
+	{
+		ScreenPos.x += ShaderPrint.FontSize.z * Count;
+	}
+	return ScreenPos;
+
+}
+
+void ShaderPrintSpace(in float Count = 1)
+{
+	ShaderPrintCursorPos = ShaderPrintSpace(ShaderPrintCursorPos, Count);
+}
+
+float2 ShaderPrintNewline(in float2 ScreenPos)
+{
+	if (ShaderPrintFilterEnable)
+	{
+		ScreenPos.x = ShaderPrintCursorStart.x;
+		ScreenPos.y += ShaderPrint.FontSize.w;
+	}
+	return ScreenPos;
+}
+
+void ShaderPrintNewline()
+{
+	ShaderPrintCursorPos = ShaderPrintNewline(ShaderPrintCursorPos);
+}
+
+float2 ShaderPrintGetCursorPos()
+{
+	return ShaderPrintCursorPos;
+}
+
+float2 ShaderPrintSetCursorPos(in float2 ScreenPos)
+{
+	float2 PrevScreenPos = ShaderPrintCursorPos;
+	ShaderPrintCursorPos = ScreenPos;
+	return PrevScreenPos;
+}
+
+// --------------------------------------------------------------------------
+// Filter helpers
+// Use these to restrict ShaderPrint to a subset of active threads
+// The filter functions can be called once at the start of the shader to filter all later printing
+//
+// Example use cases would be:
+//
+// ShaderPrintFilter(all(SvPosition.xy == float2(100, 100)));
+// to debug a single pixel
+//
+// ShaderPrintFilter(all(DispatchThreadId == uint3(0, 0, 0)));
+// to debug a single compute shader thread
+//
+// ShaderPrintFilterOneThread()
+// to debug a single random active thread
+// Note that ShaderPrintFilterOneThread() only works for one shader invocation in a view since it relies on a global memory location
+
+void ShaderPrintFilter(bool bFilter)
+{
+	ShaderPrintFilterEnable = bFilter;
+}
+
+void ShaderPrintFilterOneThread()
+{
+	// Atomic flag is stored in first element .Type Value
+	uint PrevValue;
+	InterlockedCompareExchange(RWValuesBuffer[0].Type, 0, 1, PrevValue);
+	ShaderPrintFilter(PrevValue == 0);
+}
diff --git a/Engine/Shaders/Private/ShaderPrintDraw.usf b/Engine/Shaders/Private/ShaderPrintDraw.usf
new file mode 100644
--- /dev/null
+++ b/Engine/Shaders/Private/ShaderPrintDraw.usf
@@ -0,0 +1,260 @@
+/*=============================================================================
+	ShaderPrintDraw.usf:
+	Shaders to gather contents of the debug value buffer, convert it to 
+	symbols and draw each symbol as an instanced quad
+=============================================================================*/
+
+#include "Common.ush"	
+#include "ShaderPrintCommon.ush"
+
+StructuredBuffer<ShaderPrintItem> ValuesBuffer;
+RWStructuredBuffer<ShaderPrintItem> RWSymbolsBuffer;
+StructuredBuffer<ShaderPrintItem> SymbolsBuffer;
+RWBuffer<uint> RWIndirectDispatchArgsBuffer;
+RWBuffer<uint> RWIndirectDrawArgsBuffer;
+
+
+// --------------------------------------------------------------------------
+// Clear the output value buffer
+// Needs to be done at the beginning of each view
+// Other shaders should only print debug values after this pass
+
+[numthreads(1, 1, 1)]
+void InitValueBufferCS()
+{
+	// Clear .Value in the first buffer element which is (ab)used as the buffer counter
+	RWValuesBuffer[0].Value = 0;
+	// Clear .Type in the first buffer element which is (ab)used as the ShaderPrintFilterOneThread() atomic value
+	RWValuesBuffer[0].Type = 0;
+}
+
+
+// --------------------------------------------------------------------------
+// Fill the indirect params buffer ready for the value->symbol compute pass
+// Also clear the symbol buffer ready to be filled
+
+[numthreads(1, 1, 1)]
+void BuildIndirectDispatchArgsCS()
+{
+	uint ValueCount = min(ShaderPrint.MaxValueCount, ValuesBuffer[0].Value);
+
+	RWIndirectDispatchArgsBuffer[0] = (ValueCount + 63)/64;
+	RWIndirectDispatchArgsBuffer[1] = 1;
+	RWIndirectDispatchArgsBuffer[2] = 1;
+
+	// Clear .Value in the first buffer element which is (ab)used as the buffer counter
+	RWSymbolsBuffer[0].Value = 0;
+}
+
+
+// --------------------------------------------------------------------------
+// Read the values buffer and convert to the symbols buffer
+// This is required for converting a single number to a full set of symbols
+
+void AddSymbol(in ShaderPrintItem Item)
+{
+	// Buffer counter is stored in first element .Value
+	int IndexToStore = 0;
+	InterlockedAdd(RWSymbolsBuffer[0].Value, 1, IndexToStore);
+	
+	// Prevent writing off the buffer
+	// Note that counter still increases so need to clamp when reading it in later passes
+	if (IndexToStore >= ShaderPrint.MaxSymbolCount)
+	{
+		return;
+	}
+
+	// Because counter is in first element, read/writes need to be offset by 1
+	RWSymbolsBuffer[IndexToStore + 1] = Item;
+}
+
+void AddSymbol(inout float2 ScreenPos, in int Symbol)
+{
+	ShaderPrintItem Item;
+	Item.ScreenPos = ScreenPos;
+	Item.Value = Symbol;
+	Item.Type = 0;
+	AddSymbol(Item);
+
+	ScreenPos.x += ShaderPrint.FontSize.z;
+}
+
+int GetDecimalSymbol(in int Digit)
+{
+	Digit = clamp(Digit, 0, 9);
+	return (Digit + _0_);
+}
+
+int GetHexSymbol(in int Digit)
+{
+	Digit = clamp(Digit, 0, 15);
+	return Digit < 10 ? Digit + _0_ : Digit - 10 + _A_;
+}
+
+void AddUIntSymbols(inout float2 ScreenPos, in uint Value)
+{
+	const uint MaxSymbols = 10; // Symbols required for MAX_UINT
+	uint SymbolCount = 1;
+	uint Divisor = 1;
+	uint TestValue = Value;
+	for (uint i = 0; i < MaxSymbols-1; ++i)
+	{
+		TestValue /= 10;
+		if (TestValue > 0)
+		{
+			SymbolCount ++;
+			Divisor *= 10;
+		}
+	}
+	
+	for (uint i = 0; i < SymbolCount; ++i)
+	{
+		uint Digit = Value / Divisor;
+		AddSymbol(ScreenPos, GetDecimalSymbol((int)Digit));
+		Value = Value - (Digit * Divisor);
+		Divisor /= 10;
+	}
+}
+
+void AddIntSymbols(inout float2 ScreenPos, in int Value)
+{
+	if (Value < 0)
+	{
+	   AddSymbol(ScreenPos, _MINUS_);
+	   Value = -Value;
+	}
+
+	AddUIntSymbols(ScreenPos, (uint)Value);
+}
+
+void AddFloatSymbols(inout float2 ScreenPos, in float Value)
+{
+	if (isinf(Value))
+	{
+		AddSymbol(ScreenPos, _I_);
+		AddSymbol(ScreenPos, _N_);
+		AddSymbol(ScreenPos, _F_);
+		return;
+	}
+
+	if (isnan(Value))
+	{
+		AddSymbol(ScreenPos, _N_);
+		AddSymbol(ScreenPos, _A_);
+		AddSymbol(ScreenPos, _N_);
+		return;
+	}
+
+	if (Value < 0)
+	{
+	   AddSymbol(ScreenPos, _MINUS_);
+	   Value = -Value;
+	}
+
+	AddIntSymbols(ScreenPos, (int)Value);
+	AddSymbol(ScreenPos, _DOT_);
+
+	for (int i = 5; i > 0; --i)
+	{
+		Value = frac(Value);
+		if (Value > 0.f)
+		{
+			Value *= 10.f;
+			AddSymbol(ScreenPos, GetDecimalSymbol((int)Value));
+		}
+	}
+}
+
+void AddHexSymbols(inout float2 ScreenPos, uint Value)
+{
+    AddSymbol(ScreenPos, _0_);
+    AddSymbol(ScreenPos, _X_);
+    
+	for (int i = 7; i >= 0; --i)
+    {
+        uint Digit = (Value >> (i*4)) & 0xf;
+        AddSymbol(ScreenPos, GetHexSymbol((int)Digit));
+    }
+}
+
+[numthreads(64, 1, 1)]
+void BuildSymbolBufferCS(
+	uint3 DispatchThreadId : SV_DispatchThreadID
+	)
+{
+	uint ValueCount = min(ShaderPrint.MaxValueCount, ValuesBuffer[0].Value);
+	if (DispatchThreadId.x >= ValueCount)
+	{
+		return;
+	}
+
+	ShaderPrintItem Value = ValuesBuffer[DispatchThreadId.x + 1];
+	
+	if (Value.Type == SHADER_PRINT_TYPE_SYMBOL)
+		AddSymbol(Value.ScreenPos, Value.Value);
+	else if (Value.Type == SHADER_PRINT_TYPE_FLOAT)
+		AddFloatSymbols(Value.ScreenPos, asfloat(Value.Value));
+	else if (Value.Type == SHADER_PRINT_TYPE_INT)
+		AddIntSymbols(Value.ScreenPos, Value.Value);
+	else if (Value.Type == SHADER_PRINT_TYPE_UINT)
+		AddUIntSymbols(Value.ScreenPos, asuint(Value.Value));
+	else if (Value.Type == SHADER_PRINT_TYPE_HEX)
+		AddHexSymbols(Value.ScreenPos, asuint(Value.Value));
+}
+
+
+// --------------------------------------------------------------------------
+// Fill the indirect params buffer ready for the instanced draw pass
+
+[numthreads(1, 1, 1)]
+void BuildIndirectDrawArgsCS()
+{
+	uint SymbolCount = min(ShaderPrint.MaxSymbolCount, SymbolsBuffer[0].Value);
+
+	RWIndirectDrawArgsBuffer[0] = 6;
+	RWIndirectDrawArgsBuffer[1] = SymbolCount * 2; // Main glyph + drop shadow
+	RWIndirectDrawArgsBuffer[2] = 0;
+	RWIndirectDrawArgsBuffer[3] = 0;
+}
+
+
+// --------------------------------------------------------------------------
+// Instanced draw call to render each symbol
+
+void DrawSymbolsVS(
+	in uint InstanceId : SV_InstanceID,
+	in uint VertexId : SV_VertexID,
+	out float4 OutPosition : SV_POSITION,
+	out float2 OutUV : TEXCOORD0,
+	out int OutValue : TEXCOORD1,
+	out int OutShadow : TEXCOORD2
+	)
+{
+	bool bDropShadow = (InstanceId & 1) == 0;
+	OutShadow = bDropShadow ? 1 : 0;
+	float2 DropShadowOffset = bDropShadow ? ShaderPrint.FontSize.xy * float2(0.2f, -0.2f) : 0;
+	
+	InstanceId = InstanceId >> 1;
+	ShaderPrintItem Symbol = SymbolsBuffer[InstanceId + 1];
+	OutValue = Symbol.Value;
+
+	OutUV = float2(((VertexId + 1) / 3) & 1, VertexId & 1);
+	
+	float2 VertPos = float2(OutUV.x, 1.f - OutUV.y) * 2.f - 1.f;
+	float2 InstancePos = float2(Symbol.ScreenPos.x, 1.f - Symbol.ScreenPos.y) * 2.f - 1.f;
+	OutPosition.xy = VertPos.xy * ShaderPrint.FontSize.xy + InstancePos + DropShadowOffset;
+	OutPosition.zw = float2(0, 1);
+}
+
+void DrawSymbolsPS(
+	in float4 InPosition : SV_POSITION,
+	in noperspective float2 InUV : TEXCOORD0,
+	in nointerpolation int InValue : TEXCOORD1,
+	in nointerpolation int InShadow : TEXCOORD2,
+	out float4 OutColor : SV_Target0
+	)
+{
+	float Alpha = SampleMiniFont(InValue, (uint2)floor(InUV * 8.f));
+	float3 Color = InShadow == 0 ? Alpha : 0;
+	OutColor = float4(Color, Alpha);
+}
diff --git a/Engine/Source/Runtime/Renderer/Private/DeferredShadingRenderer.cpp b/Engine/Source/Runtime/Renderer/Private/DeferredShadingRenderer.cpp
--- a/Engine/Source/Runtime/Renderer/Private/DeferredShadingRenderer.cpp
+++ b/Engine/Source/Runtime/Renderer/Private/DeferredShadingRenderer.cpp
@@ -34,6 +34,7 @@
 #include "SceneViewFamilyBlackboard.h"
 #include "ScreenSpaceDenoise.h"
 #include "RayTracing/RaytracingOptions.h"
+#include "ShaderPrint.h"
 
 TAutoConsoleVariable<int32> CVarEarlyZPass(
 	TEXT("r.EarlyZPass"),
@@ -1101,6 +1102,11 @@ void FDeferredShadingSceneRenderer::Render(FRHICommandListImmediate& RHICmdList)
 
 	UpdateGPUScene(RHICmdList, *Scene);
 
+	for (int32 ViewIndex = 0; ViewIndex < Views.Num(); ViewIndex++)
+	{
+		ShaderPrint::BeginView(RHICmdList, Views[ViewIndex]);
+	}
+
 	for (int32 ViewIndex = 0; ViewIndex < Views.Num(); ViewIndex++)
 	{
 		UploadDynamicPrimitiveShaderDataForView(RHICmdList, *Scene, Views[ViewIndex]);
@@ -2165,8 +2171,8 @@ void FDeferredShadingSceneRenderer::Render(FRHICommandListImmediate& RHICmdList)
 			}
 			else
 			{
-			GPostProcessing.Process(RHICmdList, Views[ ViewIndex ], VelocityRT);
-		}
+				GPostProcessing.Process(RHICmdList, Views[ ViewIndex ], VelocityRT);
+			}
 		}
 
 		// End of frame, we don't need it anymore
@@ -2181,6 +2187,11 @@ void FDeferredShadingSceneRenderer::Render(FRHICommandListImmediate& RHICmdList)
 		SceneContext.AdjustGBufferRefCount(RHICmdList, -1);
 	}
 
+	for (int32 ViewIndex = 0; ViewIndex < Views.Num(); ViewIndex++)
+	{
+		ShaderPrint::EndView(Views[ViewIndex]);
+	}
+
 #if WITH_MGPU
 	DoCrossGPUTransfers(RHICmdList, RenderTargetGPUMask);
 #endif
diff --git a/Engine/Source/Runtime/Renderer/Private/PostProcess/PostProcessShaderPrint.cpp b/Engine/Source/Runtime/Renderer/Private/PostProcess/PostProcessShaderPrint.cpp
new file mode 100644
--- /dev/null
+++ b/Engine/Source/Runtime/Renderer/Private/PostProcess/PostProcessShaderPrint.cpp
@@ -0,0 +1,28 @@
+// Copyright 1998-2019 Epic Games, Inc. All Rights Reserved.
+
+#include "PostProcess/PostProcessShaderPrint.h"
+
+#include "ShaderPrint.h"
+
+bool FRCPassPostProcessShaderPrint::IsEnabled()
+{
+	return ShaderPrint::IsEnabled();
+}
+
+void FRCPassPostProcessShaderPrint::Process(FRenderingCompositePassContext& Context)
+{
+	const FRenderingCompositeOutputRef* Input = GetInput(ePId_Input0);
+	if (Input == nullptr)
+		return;
+
+	ShaderPrint::DrawView(Context.RHICmdList, Context.View, Input->GetOutput()->PooledRenderTarget);
+	PassOutputs[0].RenderTargetDesc = Input->GetOutput()->RenderTargetDesc;
+	PassOutputs[0].PooledRenderTarget = Input->GetOutput()->PooledRenderTarget;
+}
+
+FPooledRenderTargetDesc FRCPassPostProcessShaderPrint::ComputeOutputDesc(EPassOutputId InPassOutputId) const
+{
+	FPooledRenderTargetDesc Ret = GetInput(ePId_Input0)->GetOutput()->RenderTargetDesc;
+	Ret.DebugName = TEXT("PostProcessShaderPrint");
+	return Ret;
+}
diff --git a/Engine/Source/Runtime/Renderer/Private/PostProcess/PostProcessShaderPrint.h b/Engine/Source/Runtime/Renderer/Private/PostProcess/PostProcessShaderPrint.h
new file mode 100644
--- /dev/null
+++ b/Engine/Source/Runtime/Renderer/Private/PostProcess/PostProcessShaderPrint.h
@@ -0,0 +1,19 @@
+// Copyright 1998-2019 Epic Games, Inc. All Rights Reserved.
+
+#pragma once
+
+#include "CoreMinimal.h"
+#include "RendererInterface.h"
+#include "PostProcess/RenderingCompositionGraph.h"
+
+// PostProcess node to manage the ShaderPrint debugging output
+class FRCPassPostProcessShaderPrint : public TRenderingCompositePassBase<1, 1>
+{
+public:
+	static bool IsEnabled();
+
+	virtual bool FrameBufferBlendingWithInput0() const override { return true; }
+	virtual void Process(FRenderingCompositePassContext& Context) override;
+	virtual void Release() override { delete this; }
+	virtual FPooledRenderTargetDesc ComputeOutputDesc(EPassOutputId InPassOutputId) const override;
+};
diff --git a/Engine/Source/Runtime/Renderer/Private/PostProcess/PostProcessing.cpp b/Engine/Source/Runtime/Renderer/Private/PostProcess/PostProcessing.cpp
--- a/Engine/Source/Runtime/Renderer/Private/PostProcess/PostProcessing.cpp
+++ b/Engine/Source/Runtime/Renderer/Private/PostProcess/PostProcessing.cpp
@@ -41,6 +41,7 @@
 #include "PostProcess/PostProcessHMD.h"
 #include "PostProcess/PostProcessVisualizeComplexity.h"
 #include "PostProcess/PostProcessCompositeEditorPrimitives.h"
+#include "PostProcess/PostProcessShaderPrint.h"
 #include "CompositionLighting/PostProcessPassThrough.h"
 #include "PostProcess/PostProcessTestImage.h"
 #include "HighResScreenshot.h"
@@ -2108,6 +2109,13 @@ void FPostProcessing::Process(FRHICommandListImmediate& RHICmdList, const FViewI
 			Context.FinalOutput = FRenderingCompositeOutputRef(Node);
 		}
 
+		if (FRCPassPostProcessShaderPrint::IsEnabled())
+		{
+			FRenderingCompositePass* Node = Context.Graph.RegisterPass(new(FMemStack::Get()) FRCPassPostProcessShaderPrint());
+			Node->SetInput(ePId_Input0, FRenderingCompositeOutputRef(Context.FinalOutput));
+			Context.FinalOutput = FRenderingCompositeOutputRef(Node);
+		}
+
 		AddHighResScreenshotMask(Context, SeparateTranslucency);
 
 		FIntPoint PrimaryUpscaleViewSize = Context.View.GetSecondaryViewRectSize();
diff --git a/Engine/Source/Runtime/Renderer/Private/SceneRendering.h b/Engine/Source/Runtime/Renderer/Private/SceneRendering.h
--- a/Engine/Source/Runtime/Renderer/Private/SceneRendering.h
+++ b/Engine/Source/Runtime/Renderer/Private/SceneRendering.h
@@ -1048,6 +1048,8 @@ public:
 	FShaderResourceViewRHIRef PrimitiveSceneDataOverrideSRV;
 	FShaderResourceViewRHIRef LightmapSceneDataOverrideSRV;
 
+	FRWBufferStructured ShaderPrintValueBuffer;
+
 #if RHI_RAYTRACING
 	TArray<FRayTracingGeometryInstance, SceneRenderingAllocator> RayTracingGeometryInstances;
 
diff --git a/Engine/Source/Runtime/Renderer/Private/ShaderPrint.cpp b/Engine/Source/Runtime/Renderer/Private/ShaderPrint.cpp
new file mode 100644
--- /dev/null
+++ b/Engine/Source/Runtime/Renderer/Private/ShaderPrint.cpp
@@ -0,0 +1,439 @@
+// Copyright 1998-2019 Epic Games, Inc. All Rights Reserved.
+
+#include "ShaderPrint.h"
+#include "ShaderPrintParameters.h"
+
+#include "CommonRenderResources.h"
+#include "Containers/DynamicRHIResourceArray.h"
+#include "Engine/Engine.h"
+#include "GlobalShader.h"
+#include "PipelineStateCache.h"
+#include "RenderGraphBuilder.h"
+#include "SceneRendering.h"
+#include "SystemTextures.h"
+
+
+namespace ShaderPrint
+{
+	// Console variables
+	static TAutoConsoleVariable<int32> CVarEnable(
+		TEXT("r.ShaderPrintEnable"),
+		0,
+		TEXT("ShaderPrint debugging toggle.\n"),
+		ECVF_Cheat | ECVF_RenderThreadSafe);
+
+	static TAutoConsoleVariable<int32> CVarFontSize(
+		TEXT("r.ShaderPrintFontSize"),
+		16,
+		TEXT("ShaderPrint font size.\n"),
+		ECVF_Cheat | ECVF_RenderThreadSafe);
+
+	static TAutoConsoleVariable<int32> CVarFontSpacingX(
+		TEXT("r.ShaderPrintFontSpacingX"),
+		0,
+		TEXT("ShaderPrint horizontal spacing between symbols.\n"),
+		ECVF_Cheat | ECVF_RenderThreadSafe);
+
+	static TAutoConsoleVariable<int32> CVarFontSpacingY(
+		TEXT("r.ShaderPrintFontSpacingY"),
+		8,
+		TEXT("ShaderPrint vertical spacing between symbols.\n"),
+		ECVF_Cheat | ECVF_RenderThreadSafe);
+
+	static TAutoConsoleVariable<int32> CVarMaxValueCount(
+		TEXT("r.ShaderPrintMaxValueCount"),
+		2000,
+		TEXT("ShaderPrint output buffer size.\n"),
+		ECVF_Cheat | ECVF_RenderThreadSafe);
+
+
+	// Structure used by shader buffers to store values and symbols
+	struct ShaderPrintItem
+	{
+		FVector2D ScreenPos;
+		int32 Value;
+		int32 Type;
+	};
+
+	// Get value buffer size
+	// Note that if the ShaderPrint system is disabled we still want to bind a minimal buffer
+	int32 GetMaxValueCount()
+	{
+		int32 MaxValueCount = FMath::Max(CVarMaxValueCount.GetValueOnRenderThread(), 0);
+		return IsEnabled() ? MaxValueCount : 0;
+	}
+
+	// Get symbol buffer size
+	// This is some multiple of the value buffer size to allow for maximum value->symbol expansion
+	int32 GetMaxSymbolCount()
+	{
+		return GetMaxValueCount() * 12;
+	}
+
+	// ShaderPrint uniform buffer
+	IMPLEMENT_GLOBAL_SHADER_PARAMETER_STRUCT(FUniformBufferParameters, "ShaderPrint");
+	typedef TUniformBufferRef<FUniformBufferParameters> FUniformBufferRef;
+
+	// Fill the uniform buffer parameters
+	void SetUniformBufferParameters(FViewInfo const& View, FUniformBufferParameters& OutParameters)
+	{
+		const float FontWidth = (float)FMath::Max(CVarFontSize.GetValueOnRenderThread(), 1) / (float)FMath::Max(View.UnconstrainedViewRect.Size().X, 1);
+		const float FontHeight = (float)FMath::Max(CVarFontSize.GetValueOnRenderThread(), 1) / (float)FMath::Max(View.UnconstrainedViewRect.Size().Y, 1);
+		const float SpaceWidth = (float)FMath::Max(CVarFontSpacingX.GetValueOnRenderThread(), 1) / (float)FMath::Max(View.UnconstrainedViewRect.Size().X, 1);
+		const float SpaceHeight = (float)FMath::Max(CVarFontSpacingY.GetValueOnRenderThread(), 1) / (float)FMath::Max(View.UnconstrainedViewRect.Size().Y, 1);
+
+		OutParameters.FontSize = FVector4(FontWidth, FontHeight, SpaceWidth + FontWidth, SpaceHeight + FontHeight);
+
+		OutParameters.MaxValueCount = GetMaxValueCount();
+		OutParameters.MaxSymbolCount = GetMaxSymbolCount();
+	}
+
+	// Return a uniform buffer with values filled and with single frame lifetime
+	FUniformBufferRef CreateUniformBuffer(FViewInfo const& View)
+	{
+		FUniformBufferParameters Parameters;
+		SetUniformBufferParameters(View, Parameters);
+		return FUniformBufferRef::CreateUniformBufferImmediate(Parameters, UniformBuffer_SingleFrame);
+	}
+
+	// Fill the FShaderParameters parameters
+	void SetParameters(FViewInfo const& View, FShaderParameters& OutParameters)
+	{
+		OutParameters.UniformBufferParameters = CreateUniformBuffer(View);
+		OutParameters.RWValuesBuffer = View.ShaderPrintValueBuffer.UAV;
+	}
+
+	// FShaderParametersLegacy implementation
+	void FShaderParametersLegacy::Bind(FShaderParameterMap const& ParameterMap)
+	{
+		UniformBufferParameter.Bind(ParameterMap, TEXT("ShaderPrint"));
+		ValuesBufferParameter.Bind(ParameterMap, TEXT("ValuesBuffer"));
+	}
+
+	template<typename TShaderRHIParamRef>
+	void SetShaderParameters(FShaderParametersLegacy const* P, FRHICommandListImmediate& RHICmdList, TShaderRHIParamRef const& ShaderRHI, FViewInfo const& View)
+	{
+		SetUniformBufferParameter(RHICmdList, ShaderRHI, P->UniformBufferParameter, CreateUniformBuffer(View));
+		P->ValuesBufferParameter.SetBuffer(RHICmdList, ShaderRHI, View.ShaderPrintValueBuffer);
+	}
+
+	void FShaderParametersLegacy::SetParameters(FRHICommandListImmediate& RHICmdList, FVertexShaderRHIParamRef const& ShaderRHI, FViewInfo const& View)
+	{
+		SetShaderParameters(this, RHICmdList, ShaderRHI, View);
+	}
+
+	void FShaderParametersLegacy::SetParameters(FRHICommandListImmediate& RHICmdList, FPixelShaderRHIParamRef const& ShaderRHI, FViewInfo const& View)
+	{
+		SetShaderParameters(this, RHICmdList, ShaderRHI, View);
+	}
+
+	void FShaderParametersLegacy::SetParameters(FRHICommandListImmediate& RHICmdList, FComputeShaderRHIParamRef const& ShaderRHI, FViewInfo const& View)
+	{
+		SetShaderParameters(this, RHICmdList, ShaderRHI, View);
+	}
+
+	void FShaderParametersLegacy::UnsetUAV(FRHICommandListImmediate& RHICmdList, FComputeShaderRHIParamRef const& ShaderRHI)
+	{
+		ValuesBufferParameter.UnsetUAV(RHICmdList, ShaderRHI);
+	}
+
+	FArchive& operator<<(FArchive& Ar, FShaderParametersLegacy& P)
+	{
+		Ar << P.UniformBufferParameter << P.ValuesBufferParameter;
+		return Ar;
+	}
+
+	// Shader to initialize the output value buffer
+	class FShaderInitValueBufferCS : public FGlobalShader
+	{
+	public:
+		DECLARE_SHADER_TYPE(FShaderInitValueBufferCS, Global);
+		SHADER_USE_PARAMETER_STRUCT(FShaderInitValueBufferCS, FGlobalShader);
+
+		BEGIN_SHADER_PARAMETER_STRUCT(FParameters, )
+			SHADER_PARAMETER_UAV(RWStructuredBuffer<ShaderPrintItem>, RWValuesBuffer)
+		END_SHADER_PARAMETER_STRUCT()
+
+		static bool ShouldCompilePermutation(FGlobalShaderPermutationParameters const& Parameters)
+		{
+			return IsFeatureLevelSupported(Parameters.Platform, ERHIFeatureLevel::SM4);
+		}
+	};
+
+	IMPLEMENT_SHADER_TYPE(, FShaderInitValueBufferCS, TEXT("/Engine/Private/ShaderPrintDraw.usf"), TEXT("InitValueBufferCS"), SF_Compute)
+
+	// Shader to fill the indirect parameter arguments ready for the value->symbol compute pass
+	class FShaderBuildIndirectDispatchArgsCS : public FGlobalShader
+	{
+	public:
+		DECLARE_SHADER_TYPE(FShaderBuildIndirectDispatchArgsCS, Global);
+		SHADER_USE_PARAMETER_STRUCT(FShaderBuildIndirectDispatchArgsCS, FGlobalShader);
+
+		BEGIN_SHADER_PARAMETER_STRUCT(FParameters, )
+			SHADER_PARAMETER_STRUCT_REF(FUniformBufferParameters, UniformBufferParameters)
+			SHADER_PARAMETER_SRV(StructuredBuffer<ShaderPrintItem>, ValuesBuffer)
+			SHADER_PARAMETER_RDG_BUFFER_UAV(RWStructuredBuffer<ShaderPrintItem>, RWSymbolsBuffer)
+			SHADER_PARAMETER_RDG_BUFFER_UAV(RWStructuredBuffer<uint>, RWIndirectDispatchArgsBuffer)
+		END_SHADER_PARAMETER_STRUCT()
+
+		static bool ShouldCompilePermutation(FGlobalShaderPermutationParameters const& Parameters)
+		{
+			return IsFeatureLevelSupported(Parameters.Platform, ERHIFeatureLevel::SM4);
+		}
+	};
+
+	IMPLEMENT_SHADER_TYPE(, FShaderBuildIndirectDispatchArgsCS, TEXT("/Engine/Private/ShaderPrintDraw.usf"), TEXT("BuildIndirectDispatchArgsCS"), SF_Compute)
+
+	// Shader to read the values buffer and convert to the symbols buffer
+	class FShaderBuildSymbolBufferCS : public FGlobalShader
+	{
+	public:
+		DECLARE_SHADER_TYPE(FShaderBuildSymbolBufferCS, Global);
+		SHADER_USE_PARAMETER_STRUCT(FShaderBuildSymbolBufferCS, FGlobalShader);
+
+		BEGIN_SHADER_PARAMETER_STRUCT(FParameters, )
+			SHADER_PARAMETER_STRUCT_REF(FUniformBufferParameters, UniformBufferParameters)
+			SHADER_PARAMETER_SRV(StructuredBuffer<ShaderPrintItem>, ValuesBuffer)
+			SHADER_PARAMETER_RDG_BUFFER_UAV(RWStructuredBuffer<ShaderPrintItem>, RWSymbolsBuffer)
+			SHADER_PARAMETER_RDG_BUFFER(StructuredBuffer<uint>, IndirectDispatchArgsBuffer)
+		END_SHADER_PARAMETER_STRUCT()
+
+		static bool ShouldCompilePermutation(FGlobalShaderPermutationParameters const& Parameters)
+		{
+			return IsFeatureLevelSupported(Parameters.Platform, ERHIFeatureLevel::SM4);
+		}
+	};
+
+	IMPLEMENT_SHADER_TYPE(, FShaderBuildSymbolBufferCS, TEXT("/Engine/Private/ShaderPrintDraw.usf"), TEXT("BuildSymbolBufferCS"), SF_Compute)
+
+	// Shader to fill the indirect parameter arguments ready for draw pass
+	class FShaderBuildIndirectDrawArgsCS : public FGlobalShader
+	{
+	public:
+		DECLARE_SHADER_TYPE(FShaderBuildIndirectDrawArgsCS, Global);
+		SHADER_USE_PARAMETER_STRUCT(FShaderBuildIndirectDrawArgsCS, FGlobalShader);
+
+		BEGIN_SHADER_PARAMETER_STRUCT(FParameters, )
+			SHADER_PARAMETER_STRUCT_REF(FUniformBufferParameters, UniformBufferParameters)
+			SHADER_PARAMETER_RDG_BUFFER_SRV(StructuredBuffer<ShaderPrintItem>, SymbolsBuffer)
+			SHADER_PARAMETER_RDG_BUFFER_UAV(RWStructuredBuffer<uint>, RWIndirectDrawArgsBuffer)
+		END_SHADER_PARAMETER_STRUCT()
+
+		static bool ShouldCompilePermutation(FGlobalShaderPermutationParameters const& Parameters)
+		{
+			return IsFeatureLevelSupported(Parameters.Platform, ERHIFeatureLevel::SM4);
+		}
+	};
+
+	IMPLEMENT_SHADER_TYPE(, FShaderBuildIndirectDrawArgsCS, TEXT("/Engine/Private/ShaderPrintDraw.usf"), TEXT("BuildIndirectDrawArgsCS"), SF_Compute)
+
+	// Shader for draw pass to render each symbol
+	class FShaderDrawSymbols : public FGlobalShader
+	{
+	public:
+		SHADER_USE_PARAMETER_STRUCT(FShaderDrawSymbols, FGlobalShader);
+
+		BEGIN_SHADER_PARAMETER_STRUCT(FParameters, )
+			RENDER_TARGET_BINDING_SLOTS()
+			SHADER_PARAMETER_STRUCT_REF(FUniformBufferParameters, UniformBufferParameters)
+			SHADER_PARAMETER_TEXTURE(Texture2D, MiniFontTexture)
+			SHADER_PARAMETER_RDG_BUFFER_SRV(StructuredBuffer<ShaderPrintItem>, SymbolsBuffer)
+			SHADER_PARAMETER_RDG_BUFFER(StructuredBuffer<uint>, IndirectDrawArgsBuffer)
+		END_SHADER_PARAMETER_STRUCT()
+
+		static bool ShouldCompilePermutation(FGlobalShaderPermutationParameters const& Parameters)
+		{
+			return IsFeatureLevelSupported(Parameters.Platform, ERHIFeatureLevel::SM4);
+		}
+	};
+
+	class FShaderDrawSymbolsVS : public FShaderDrawSymbols
+	{
+	public:
+		DECLARE_SHADER_TYPE(FShaderDrawSymbolsVS, Global);
+
+		FShaderDrawSymbolsVS()
+		{}
+
+		FShaderDrawSymbolsVS(const ShaderMetaType::CompiledShaderInitializerType& Initializer)
+			: FShaderDrawSymbols(Initializer)
+		{}
+	};
+
+	IMPLEMENT_SHADER_TYPE(, FShaderDrawSymbolsVS, TEXT("/Engine/Private/ShaderPrintDraw.usf"), TEXT("DrawSymbolsVS"), SF_Vertex)
+
+	class FShaderDrawSymbolsPS : public FShaderDrawSymbols
+	{
+	public:
+		DECLARE_SHADER_TYPE(FShaderDrawSymbolsPS, Global);
+
+		FShaderDrawSymbolsPS()
+		{}
+
+		FShaderDrawSymbolsPS(const ShaderMetaType::CompiledShaderInitializerType& Initializer)
+			: FShaderDrawSymbols(Initializer)
+		{}
+	};
+
+	IMPLEMENT_SHADER_TYPE(, FShaderDrawSymbolsPS, TEXT("/Engine/Private/ShaderPrintDraw.usf"), TEXT("DrawSymbolsPS"), SF_Pixel)
+
+
+	bool IsEnabled()
+	{
+		return CVarEnable.GetValueOnAnyThread();
+	}
+
+	void BeginView(FRHICommandListImmediate& RHICmdList, FViewInfo& View)
+	{
+		// Initialize output buffer and store in the view info
+		// Values buffer contains Count + 1 elements. The first element is only used as a counter.
+		View.ShaderPrintValueBuffer.Initialize(sizeof(ShaderPrintItem), GetMaxValueCount() + 1, 0U, TEXT("ShaderPrintValueBuffer"));
+
+		// Early out if system is disabled
+		// Note that we still prepared a minimal ShaderPrintValueBuffer 
+		// This is in case some debug shader code is still active (we don't want an unbound buffer!)
+		if (!IsEnabled())
+		{
+			return;
+		}
+
+		SCOPED_DRAW_EVENT(RHICmdList, ShaderPrintBeginView);
+
+		// Clear the output buffer internal counter ready for use
+		const ERHIFeatureLevel::Type FeatureLevel = GMaxRHIFeatureLevel;
+		TShaderMap<FGlobalShaderType>* GlobalShaderMap = GetGlobalShaderMap(FeatureLevel);
+
+		TShaderMapRef< FShaderInitValueBufferCS > ComputeShader(GlobalShaderMap);
+
+		FShaderInitValueBufferCS::FParameters Parameters;
+		Parameters.RWValuesBuffer = View.ShaderPrintValueBuffer.UAV;
+
+		FComputeShaderUtils::Dispatch(RHICmdList, *ComputeShader, Parameters, FIntVector(1, 1, 1));
+	}
+
+	void DrawView(FRHICommandListImmediate& RHICmdList, FViewInfo const& View, TRefCountPtr<IPooledRenderTarget>& OutputTarget)
+	{
+		// Early out if system is disabled
+		if (!IsEnabled())
+		{
+			return;
+		}
+
+		SCOPED_DRAW_EVENT(RHICmdList, ShaderPrintDrawView);
+
+		// Build graph
+		FRDGBuilder GraphBuilder(RHICmdList);
+
+		// Initialize graph managed resources
+		// Symbols buffer contains Count + 1 elements. The first element is only used as a counter.
+		FRDGBufferRef SymbolBuffer = GraphBuilder.CreateBuffer(FRDGBufferDesc::CreateStructuredDesc(sizeof(ShaderPrintItem), GetMaxSymbolCount() + 1), TEXT("ShaderPrintSymbolBuffer"));
+		FRDGBufferRef IndirectDispatchArgsBuffer = GraphBuilder.CreateBuffer(FRDGBufferDesc::CreateIndirectDesc(4), TEXT("ShaderPrintIndirectDispatchArgs"));
+		FRDGBufferRef IndirectDrawArgsBuffer = GraphBuilder.CreateBuffer(FRDGBufferDesc::CreateIndirectDesc(5), TEXT("ShaderPrintIndirectDrawArgs"));
+		FRDGTextureRef OutputTexture = GraphBuilder.RegisterExternalTexture(OutputTarget);
+
+		// Non graph managed resources
+		FUniformBufferRef UniformBuffer = CreateUniformBuffer(View);
+		FShaderResourceViewRHIRef ValuesBuffer = View.ShaderPrintValueBuffer.SRV;
+		FTextureRHIRef FontTexture = GEngine->MiniFontTexture != nullptr ? GEngine->MiniFontTexture->Resource->TextureRHI : GSystemTextures.BlackDummy->GetRenderTargetItem().ShaderResourceTexture;;
+
+		const ERHIFeatureLevel::Type FeatureLevel = GMaxRHIFeatureLevel;
+		TShaderMap<FGlobalShaderType>* GlobalShaderMap = GetGlobalShaderMap(FeatureLevel);
+
+		// BuildIndirectDispatchArgs
+		{
+			typedef FShaderBuildIndirectDispatchArgsCS SHADER;
+			TShaderMapRef<SHADER> ComputeShader(GlobalShaderMap);
+
+			SHADER::FParameters* PassParameters = GraphBuilder.AllocParameters<SHADER::FParameters>();
+			PassParameters->UniformBufferParameters = UniformBuffer;
+			PassParameters->ValuesBuffer = ValuesBuffer;
+			PassParameters->RWSymbolsBuffer = GraphBuilder.CreateUAV(SymbolBuffer, EPixelFormat::PF_R32_UINT);
+			PassParameters->RWIndirectDispatchArgsBuffer = GraphBuilder.CreateUAV(IndirectDispatchArgsBuffer, EPixelFormat::PF_R32_UINT);
+
+			FComputeShaderUtils::AddPass(
+				GraphBuilder, 
+				RDG_EVENT_NAME("BuildIndirectDispatchArgs"), 
+				*ComputeShader, PassParameters,
+				FIntVector(1, 1, 1));
+		}
+
+		// BuildSymbolBuffer
+		{
+			typedef FShaderBuildSymbolBufferCS SHADER;
+			TShaderMapRef<SHADER> ComputeShader(GlobalShaderMap);
+
+			SHADER::FParameters* PassParameters = GraphBuilder.AllocParameters<SHADER::FParameters>();
+			PassParameters->UniformBufferParameters = UniformBuffer;
+			PassParameters->ValuesBuffer = ValuesBuffer;
+			PassParameters->RWSymbolsBuffer = GraphBuilder.CreateUAV(SymbolBuffer, EPixelFormat::PF_R32_UINT);
+			PassParameters->IndirectDispatchArgsBuffer = IndirectDispatchArgsBuffer;
+
+			FComputeShaderUtils::AddPass(
+				GraphBuilder,
+				RDG_EVENT_NAME("BuildSymbolBuffer"),
+				*ComputeShader, PassParameters,
+				IndirectDispatchArgsBuffer, 0);
+		}
+
+		// BuildIndirectDrawArgs
+		{
+			typedef FShaderBuildIndirectDrawArgsCS SHADER;
+			TShaderMapRef<SHADER> ComputeShader(GlobalShaderMap);
+
+			SHADER::FParameters* PassParameters = GraphBuilder.AllocParameters<SHADER::FParameters>();
+			PassParameters->UniformBufferParameters = UniformBuffer;
+			PassParameters->SymbolsBuffer = GraphBuilder.CreateSRV(SymbolBuffer);
+			PassParameters->RWIndirectDrawArgsBuffer = GraphBuilder.CreateUAV(IndirectDrawArgsBuffer, EPixelFormat::PF_R32_UINT);
+
+			FComputeShaderUtils::AddPass(
+				GraphBuilder,
+				RDG_EVENT_NAME("BuildIndirectDrawArgs"),
+				*ComputeShader, PassParameters,
+				FIntVector(1, 1, 1));
+		}
+
+		// DrawSymbols
+		{
+			typedef FShaderDrawSymbols SHADER;
+			TShaderMapRef< FShaderDrawSymbolsVS > VertexShader(GlobalShaderMap);
+			TShaderMapRef< FShaderDrawSymbolsPS > PixelShader(GlobalShaderMap);
+
+			SHADER::FParameters* PassParameters = GraphBuilder.AllocParameters<SHADER::FParameters>();
+			PassParameters->RenderTargets[0] = FRenderTargetBinding(OutputTexture, ERenderTargetLoadAction::ENoAction, ERenderTargetStoreAction::ENoAction);
+			PassParameters->UniformBufferParameters = UniformBuffer;
+			PassParameters->MiniFontTexture = FontTexture;
+			PassParameters->SymbolsBuffer = GraphBuilder.CreateSRV(SymbolBuffer);
+			PassParameters->IndirectDrawArgsBuffer = IndirectDrawArgsBuffer;
+
+			GraphBuilder.AddPass(
+				RDG_EVENT_NAME("DrawSymbols"),
+				PassParameters,
+				ERenderGraphPassFlags::None,
+				[VertexShader, PixelShader, PassParameters](FRHICommandListImmediate& RHICmdListImmediate)
+			{
+				FGraphicsPipelineStateInitializer GraphicsPSOInit;
+				RHICmdListImmediate.ApplyCachedRenderTargets(GraphicsPSOInit);
+				GraphicsPSOInit.DepthStencilState = TStaticDepthStencilState<false, CF_Always>::GetRHI();
+				GraphicsPSOInit.BlendState = TStaticBlendState<CW_RGBA, BO_Add, BF_One, BF_InverseSourceAlpha, BO_Add, BF_Zero, BF_One>::GetRHI();
+				GraphicsPSOInit.RasterizerState = TStaticRasterizerState<>::GetRHI();
+				GraphicsPSOInit.PrimitiveType = PT_TriangleList;
+				GraphicsPSOInit.BoundShaderState.VertexDeclarationRHI = GetVertexDeclarationFVector4();
+				GraphicsPSOInit.BoundShaderState.VertexShaderRHI = GETSAFERHISHADER_VERTEX(*VertexShader);
+				GraphicsPSOInit.BoundShaderState.PixelShaderRHI = GETSAFERHISHADER_PIXEL(*PixelShader);
+				SetGraphicsPipelineState(RHICmdListImmediate, GraphicsPSOInit);
+
+				SetShaderParameters(RHICmdListImmediate, *VertexShader, VertexShader->GetVertexShader(), *PassParameters);
+				SetShaderParameters(RHICmdListImmediate, *PixelShader, PixelShader->GetPixelShader(), *PassParameters);
+
+				RHICmdListImmediate.DrawIndexedPrimitiveIndirect(GTwoTrianglesIndexBuffer.IndexBufferRHI, PassParameters->IndirectDrawArgsBuffer->GetIndirectRHICallBuffer(), 0);
+			});
+		}
+
+		GraphBuilder.Execute();
+	}
+
+	void EndView(FViewInfo& View)
+	{
+		View.ShaderPrintValueBuffer.Release();
+	}
+}
diff --git a/Engine/Source/Runtime/Renderer/Private/ShaderPrint.h b/Engine/Source/Runtime/Renderer/Private/ShaderPrint.h
new file mode 100644
--- /dev/null
+++ b/Engine/Source/Runtime/Renderer/Private/ShaderPrint.h
@@ -0,0 +1,31 @@
+// Copyright 1998-2019 Epic Games, Inc. All Rights Reserved.
+
+// Use ShaderPrint to debug print from any shader type.
+// Call BeginView() at the start of a view to init the debug output buffer.
+// Call DrawView() to composite the debug data to the final render target.
+// Call EndView() at the end of the view to tidy up.
+
+// The shader HLSL code should include the ShaderPrintCommon.ush and use the ShaderPrint*() functions.
+// The shader C++ code needs to bind the buffers required to capture the debug print. See ShaderPrintParameters.h for this. 
+
+#pragma once
+
+#include "CoreMinimal.h"
+#include "Templates/RefCounting.h"
+
+class FRHICommandListImmediate;
+class FViewInfo;
+struct IPooledRenderTarget;
+
+namespace ShaderPrint
+{
+	// Have we enabled the ShaderPrint system?
+	bool IsEnabled();
+
+	// Allocate the debug print buffer associated with the view
+	void BeginView(FRHICommandListImmediate& RHICmdList, FViewInfo& View);
+	// Draw info from the debug print buffer to the given output target
+	void DrawView(FRHICommandListImmediate& RHICmdList, FViewInfo const& View, TRefCountPtr<IPooledRenderTarget>& OutputTarget);
+	// Release the debug print buffer associated with the view
+	void EndView(FViewInfo& View);
+}
diff --git a/Engine/Source/Runtime/Renderer/Private/ShaderPrintParameters.h b/Engine/Source/Runtime/Renderer/Private/ShaderPrintParameters.h
new file mode 100644
--- /dev/null
+++ b/Engine/Source/Runtime/Renderer/Private/ShaderPrintParameters.h
@@ -0,0 +1,64 @@
+// Copyright 1998-2019 Epic Games, Inc. All Rights Reserved.
+
+// The ShaderPrint system uses a RWBuffer to capture any debug print from a shader.
+// This means that the buffer needs to be bound for the shader you wish to debug.
+// It would be ideal if that was automatic (maybe by having a fixed bind point for the buffer and binding it for the entire view).
+// But for now you need to manually add binding information to your FShader class.
+// To do this either:
+// (i) Use SHADER_PARAMETER_STRUCT_INCLUDE(ShaderPrint::FShaderParameters) in your FShader::FParameters declaration and call SetParameters().
+// (ii) Put an FShaderParametersLegacy as a member of your FShader and add calls into Bind()/SetParameters().
+
+// Also it seems that we can only bind a RWBuffer to compute shaders right now. Fixing this would allow us to use this system from all shader stages.
+
+#pragma once
+
+#include "CoreMinimal.h"
+#include "ShaderParameterMacros.h"
+#include "ShaderParameters.h"
+
+class FArchive;
+class FRHICommandListImmediate;
+class FShaderParameterMap;
+class FViewInfo;
+
+namespace ShaderPrint
+{
+	// ShaderPrint uniform buffer layout
+	BEGIN_GLOBAL_SHADER_PARAMETER_STRUCT(FUniformBufferParameters, )
+		SHADER_PARAMETER(FVector4, FontSize)
+		SHADER_PARAMETER(int32, MaxValueCount)
+		SHADER_PARAMETER(int32, MaxSymbolCount)
+	END_GLOBAL_SHADER_PARAMETER_STRUCT()
+
+	// ShaderPrint parameter struct declaration
+	BEGIN_SHADER_PARAMETER_STRUCT(FShaderParameters, )
+		SHADER_PARAMETER_STRUCT_REF(FUniformBufferParameters, UniformBufferParameters)
+		SHADER_PARAMETER_UAV(RWStructuredBuffer<ShaderPrintItem>, RWValuesBuffer)
+	END_SHADER_PARAMETER_STRUCT()
+
+	// Call this to fill the FShaderParameters
+	void SetParameters(FViewInfo const& View, FShaderParameters& OutParameters);
+
+
+	// Legacy parameter binding helper for an FShader that doesn't use FParameters style parameter declaration
+	struct FShaderParametersLegacy
+	{
+		// Call on shader construction to bind
+		void Bind(FShaderParameterMap const& ParameterMap);
+
+		// Call to set parameters
+		void SetParameters(FRHICommandListImmediate& RHICmdList, FVertexShaderRHIParamRef const& ShaderRHI, FViewInfo const& View);
+		void SetParameters(FRHICommandListImmediate& RHICmdList, FPixelShaderRHIParamRef const& ShaderRHI, FViewInfo const& View);
+		void SetParameters(FRHICommandListImmediate& RHICmdList, FComputeShaderRHIParamRef const& ShaderRHI, FViewInfo const& View);
+
+		// Call to unbind the UAV
+		void UnsetUAV(FRHICommandListImmediate& RHICmdList, FComputeShaderRHIParamRef const& ShaderRHI);
+
+		// Serializer
+		friend FArchive& operator<<(FArchive& Ar, FShaderParametersLegacy& P);
+
+		// Parameters
+		FShaderUniformBufferParameter UniformBufferParameter;
+		FRWShaderParameter ValuesBufferParameter;
+	};
+}
