DR - hlslcc doesn't support templated complex types

diff --git a/Engine/Source/Runtime/Renderer/Private/ShaderPrint.cpp b/Engine/Source/Runtime/Renderer/Private/ShaderPrint.cpp
--- a/Engine/Source/Runtime/Renderer/Private/ShaderPrint.cpp
+++ b/Engine/Source/Runtime/Renderer/Private/ShaderPrint.cpp
@@ -156,7 +156,7 @@ namespace ShaderPrint
 
 		static bool ShouldCompilePermutation(FGlobalShaderPermutationParameters const& Parameters)
 		{
-			return IsFeatureLevelSupported(Parameters.Platform, ERHIFeatureLevel::SM4);
+			return IsFeatureLevelSupported(Parameters.Platform, ERHIFeatureLevel::SM4) && !IsHlslccShaderPlatform(Parameters.Platform);
 		}
 	};
 
@@ -178,7 +178,7 @@ namespace ShaderPrint
 
 		static bool ShouldCompilePermutation(FGlobalShaderPermutationParameters const& Parameters)
 		{
-			return IsFeatureLevelSupported(Parameters.Platform, ERHIFeatureLevel::SM4);
+			return IsFeatureLevelSupported(Parameters.Platform, ERHIFeatureLevel::SM4) && !IsHlslccShaderPlatform(Parameters.Platform);
 		}
 	};
 
@@ -200,7 +200,7 @@ namespace ShaderPrint
 
 		static bool ShouldCompilePermutation(FGlobalShaderPermutationParameters const& Parameters)
 		{
-			return IsFeatureLevelSupported(Parameters.Platform, ERHIFeatureLevel::SM4);
+			return IsFeatureLevelSupported(Parameters.Platform, ERHIFeatureLevel::SM4) && !IsHlslccShaderPlatform(Parameters.Platform);
 		}
 	};
 
@@ -221,7 +221,7 @@ namespace ShaderPrint
 
 		static bool ShouldCompilePermutation(FGlobalShaderPermutationParameters const& Parameters)
 		{
-			return IsFeatureLevelSupported(Parameters.Platform, ERHIFeatureLevel::SM4);
+			return IsFeatureLevelSupported(Parameters.Platform, ERHIFeatureLevel::SM4) && !IsHlslccShaderPlatform(Parameters.Platform);
 		}
 	};
 
@@ -243,7 +243,7 @@ namespace ShaderPrint
 
 		static bool ShouldCompilePermutation(FGlobalShaderPermutationParameters const& Parameters)
 		{
-			return IsFeatureLevelSupported(Parameters.Platform, ERHIFeatureLevel::SM4);
+			return IsFeatureLevelSupported(Parameters.Platform, ERHIFeatureLevel::SM4) && !IsHlslccShaderPlatform(Parameters.Platform);
 		}
 	};
 
