Distance field support for landscape holes
Jira UE-57322

diff --git a/Engine/Shaders/Private/GlobalDistanceField.usf b/Engine/Shaders/Private/GlobalDistanceField.usf
--- a/Engine/Shaders/Private/GlobalDistanceField.usf
+++ b/Engine/Shaders/Private/GlobalDistanceField.usf
@@ -350,26 +350,30 @@ void CompositeHeightfieldsIntoGlobalDistanceFieldCS(
 
 	if (all(HeightfieldUV > MinMaxHeightfieldUV.xy) && all(HeightfieldUV < MinMaxHeightfieldUV.zw))
 	{
-		float3 WorldHeightfieldNormal;
-		float3 WorldHeightfieldShadingPosition = GetHeightfieldWorldPositionAndNormal(ValidHeightfieldIndex, LocalPosition.xy, HeightfieldUV, WorldHeightfieldNormal);
-
 		if (all(DispatchThreadId.xy < UpdateRegionSize.xy))
 		{
-			// Compute distance for all Z values of the update region
-			LOOP
-			for (uint ZIndex = 0; ZIndex < UpdateRegionSize.z; ZIndex++)
+			float3 WorldHeightfieldNormal;
+			float HeightfieldVisibility;
+			float3 WorldHeightfieldShadingPosition = GetHeightfieldWorldPositionAndNormal(ValidHeightfieldIndex, LocalPosition.xy, HeightfieldUV, WorldHeightfieldNormal, HeightfieldVisibility);
+
+			if (HeightfieldVisibility > 0.5f) // Skip holes in the heightfield
 			{
-				float3 WorldPosition;
-				WorldPosition.xy = LookupWorldPosition.xy;
-				WorldPosition.z = UpdateRegionVolumeMin.z + ZIndex * UpdateRegionVolumeStep;
+				// Compute distance for all Z values of the update region
+				LOOP
+				for (uint ZIndex = 0; ZIndex < UpdateRegionSize.z; ZIndex++)
+				{
+					float3 WorldPosition;
+					WorldPosition.xy = LookupWorldPosition.xy;
+					WorldPosition.z = UpdateRegionVolumeMin.z + ZIndex * UpdateRegionVolumeStep;
 
-				// Project the vertical height vector onto the normal of the heightfield directly below the point we are computing the distance field for, use the perpendicular distance
-				float DistanceToHeightfieldPlane = dot(WorldHeightfieldNormal, WorldPosition - WorldHeightfieldShadingPosition);
-				float MinDistance = clamp(DistanceToHeightfieldPlane, -AOGlobalMaxSphereQueryRadius, AOGlobalMaxSphereQueryRadius);
+					// Project the vertical height vector onto the normal of the heightfield directly below the point we are computing the distance field for, use the perpendicular distance
+					float DistanceToHeightfieldPlane = dot(WorldHeightfieldNormal, WorldPosition - WorldHeightfieldShadingPosition);
+					float MinDistance = clamp(DistanceToHeightfieldPlane, -AOGlobalMaxSphereQueryRadius, AOGlobalMaxSphereQueryRadius);
 
-				float3 GlobalUV = ComputeGlobalUV(WorldPosition, ClipmapIndex);
-				uint3 WriteCoordinate = (uint3)(frac(GlobalUV) * GlobalVolumeDimension);
-				RWGlobalDistanceFieldTexture[WriteCoordinate] = min(MinDistance, RWGlobalDistanceFieldTexture[WriteCoordinate]);
+					float3 GlobalUV = ComputeGlobalUV(WorldPosition, ClipmapIndex);
+					uint3 WriteCoordinate = (uint3)(frac(GlobalUV) * GlobalVolumeDimension);
+					RWGlobalDistanceFieldTexture[WriteCoordinate] = min(MinDistance, RWGlobalDistanceFieldTexture[WriteCoordinate]);
+				}
 			}
 		}
 	}
diff --git a/Engine/Shaders/Private/HeightfieldLightingShared.ush b/Engine/Shaders/Private/HeightfieldLightingShared.ush
--- a/Engine/Shaders/Private/HeightfieldLightingShared.ush
+++ b/Engine/Shaders/Private/HeightfieldLightingShared.ush
@@ -8,11 +8,13 @@ Texture2D HeightfieldTexture;
 SamplerState HeightfieldSampler;
 Texture2D DiffuseColorTexture;
 SamplerState DiffuseColorSampler;
+Texture2D VisibilityTexture;
+SamplerState VisibilitySampler;
 uint NumHeightfields;
 Buffer<float4> HeightfieldDescriptions;
 
 // In float4's, must match C++
-#define HEIGHTFIELD_DATA_STRIDE 10
+#define HEIGHTFIELD_DATA_STRIDE 12
 
 float4 GetHeightfieldScaleBias(uint Index, out bool bHasSubsections)
 {
@@ -38,29 +40,39 @@ float2 GetHeightfieldSize(uint Index)
 	return HeightfieldDescriptions[Index * HEIGHTFIELD_DATA_STRIDE + 3].xy;
 }
 
-float2 GetGlobalTexelSize(uint Index)
+float2 GetHeightfieldSizeInv(uint Index)
 {
 	return HeightfieldDescriptions[Index * HEIGHTFIELD_DATA_STRIDE + 3].zw;
 }
 
+float2 GetGlobalTexelSize(uint Index)
+{
+	return HeightfieldDescriptions[Index * HEIGHTFIELD_DATA_STRIDE + 4].xy;
+}
+
 float4x4 GetWorldToLocal(uint Index)
 {
-	float4 M0 = HeightfieldDescriptions[Index * HEIGHTFIELD_DATA_STRIDE + 4];
-	float4 M1 = HeightfieldDescriptions[Index * HEIGHTFIELD_DATA_STRIDE + 5];
-	float4 M2 = HeightfieldDescriptions[Index * HEIGHTFIELD_DATA_STRIDE + 6];
+	float4 M0 = HeightfieldDescriptions[Index * HEIGHTFIELD_DATA_STRIDE + 5];
+	float4 M1 = HeightfieldDescriptions[Index * HEIGHTFIELD_DATA_STRIDE + 6];
+	float4 M2 = HeightfieldDescriptions[Index * HEIGHTFIELD_DATA_STRIDE + 7];
 
 	return transpose(float4x4(M0, M1, M2, float4(0.0f, 0.0f, 0.0f, 1.0f)));
 }
 
 float4x4 GetLocalToWorld(uint Index)
 {
-	float4 M0 = HeightfieldDescriptions[Index * HEIGHTFIELD_DATA_STRIDE + 7];
-	float4 M1 = HeightfieldDescriptions[Index * HEIGHTFIELD_DATA_STRIDE + 8];
-	float4 M2 = HeightfieldDescriptions[Index * HEIGHTFIELD_DATA_STRIDE + 9];
+	float4 M0 = HeightfieldDescriptions[Index * HEIGHTFIELD_DATA_STRIDE + 8];
+	float4 M1 = HeightfieldDescriptions[Index * HEIGHTFIELD_DATA_STRIDE + 9];
+	float4 M2 = HeightfieldDescriptions[Index * HEIGHTFIELD_DATA_STRIDE + 10];
 
 	return transpose(float4x4(M0, M1, M2, float4(0.0f, 0.0f, 0.0f, 1.0f)));
 }
 
+float4 GetVisibilityMask(uint Index)
+{
+	return HeightfieldDescriptions[Index * HEIGHTFIELD_DATA_STRIDE + 11];
+}
+
 float2 GetHeightfieldUV(uint HeightfieldIndex, float2 LocalPosition, out float4 MinMaxHeightfieldUV)
 {
 	bool bHasSubsections;
@@ -86,7 +98,12 @@ float2 GetHeightfieldUV(uint HeightfieldIndex, float2 LocalPosition)
 	return GetHeightfieldUV(HeightfieldIndex, LocalPosition, MinMaxHeightfieldUV);
 }
 
-float3 GetHeightfieldWorldPositionAndNormal(uint HeightfieldIndex, float2 LocalPosition, float2 HeightfieldUV, out float3 WorldNormal)
+float2 GetVisibilityUV(uint HeightfieldIndex, float2 LocalPosition)
+{
+	return LocalPosition * GetHeightfieldSizeInv(HeightfieldIndex);
+}
+
+float3 GetHeightfieldWorldPositionAndNormal(uint HeightfieldIndex, float2 LocalPosition, float2 HeightfieldUV, out float3 WorldNormal, out float Visibility)
 {
 	//@todo - don't call GetHeightfieldScaleBias twice
 	bool bHasSubsections;
@@ -98,6 +115,10 @@ float3 GetHeightfieldWorldPositionAndNormal(uint HeightfieldIndex, float2 LocalP
 	float2 SampleNormal = float2(SampleValue.b, SampleValue.a) * float2(2.0,2.0) - float2(1.0,1.0);
 	WorldNormal = float3(SampleNormal, sqrt(max(1.0 - dot(SampleNormal, SampleNormal), 0.0)));
 
+	float2 VisibilityUV = GetVisibilityUV(HeightfieldIndex, LocalPosition);
+	float4 SampleVisibility = Texture2DSampleLevel(VisibilityTexture, VisibilitySampler, VisibilityUV, 0);
+	Visibility = 1.f - dot(SampleVisibility, GetVisibilityMask(HeightfieldIndex));
+
 	float3 HeightfieldShadingPosition = float3(LocalPosition, HeightfieldHeight);
 	float3 WorldHeightfieldShadingPosition = mul(float4(HeightfieldShadingPosition, 1), GetLocalToWorld(HeightfieldIndex)).xyz;
 
@@ -107,5 +128,6 @@ float3 GetHeightfieldWorldPositionAndNormal(uint HeightfieldIndex, float2 LocalP
 float3 GetHeightfieldWorldPosition(uint HeightfieldIndex, float2 LocalPosition, float2 HeightfieldUV)
 {
 	float3 WorldNormal;
-	return GetHeightfieldWorldPositionAndNormal(HeightfieldIndex, LocalPosition, HeightfieldUV, WorldNormal);
+	float Visibility;
+	return GetHeightfieldWorldPositionAndNormal(HeightfieldIndex, LocalPosition, HeightfieldUV, WorldNormal, Visibility);
 }
diff --git a/Engine/Source/Runtime/Engine/Public/PrimitiveSceneProxy.h b/Engine/Source/Runtime/Engine/Public/PrimitiveSceneProxy.h
--- a/Engine/Source/Runtime/Engine/Public/PrimitiveSceneProxy.h
+++ b/Engine/Source/Runtime/Engine/Public/PrimitiveSceneProxy.h
@@ -103,12 +103,13 @@ public:
 	FMatrix LocalToWorld;
 	FVector2D LightingAtlasLocation;
 	FIntRect HeightfieldRect;
-
 	int32 NumSubsections;
 	FVector4 SubsectionScaleAndBias;
+	int32 VisibilityChannel;
 
 	FHeightfieldComponentDescription(const FMatrix& InLocalToWorld) :
-		LocalToWorld(InLocalToWorld)
+		LocalToWorld(InLocalToWorld), 
+		VisibilityChannel(-1)
 	{}
 };
 
@@ -280,9 +281,11 @@ public:
 
 	virtual bool HeightfieldHasPendingStreaming() const { return false; }
 
-	virtual void GetHeightfieldRepresentation(UTexture2D*& OutHeightmapTexture, UTexture2D*& OutDiffuseColorTexture, FHeightfieldComponentDescription& OutDescription)
+	virtual void GetHeightfieldRepresentation(UTexture2D*& OutHeightmapTexture, UTexture2D*& OutDiffuseColorTexture, UTexture2D*& OutVisibilityTexture, FHeightfieldComponentDescription& OutDescription)
 	{
-		OutHeightmapTexture = NULL;
+		OutHeightmapTexture = nullptr;
+		OutDiffuseColorTexture = nullptr;
+		OutVisibilityTexture = nullptr;
 	}
 
 	/**
diff --git a/Engine/Source/Runtime/Landscape/Private/LandscapeRender.cpp b/Engine/Source/Runtime/Landscape/Private/LandscapeRender.cpp
--- a/Engine/Source/Runtime/Landscape/Private/LandscapeRender.cpp
+++ b/Engine/Source/Runtime/Landscape/Private/LandscapeRender.cpp
@@ -4395,7 +4395,7 @@ bool FLandscapeComponentSceneProxy::HeightfieldHasPendingStreaming() const
 	return HeightmapTexture && HeightmapTexture->bHasStreamingUpdatePending;
 }
 
-void FLandscapeComponentSceneProxy::GetHeightfieldRepresentation(UTexture2D*& OutHeightmapTexture, UTexture2D*& OutDiffuseColorTexture, FHeightfieldComponentDescription& OutDescription)
+void FLandscapeComponentSceneProxy::GetHeightfieldRepresentation(UTexture2D*& OutHeightmapTexture, UTexture2D*& OutDiffuseColorTexture, UTexture2D*& OutVisibilityTexture, FHeightfieldComponentDescription& OutDescription)
 {
 	OutHeightmapTexture = HeightmapTexture;
 	OutDiffuseColorTexture = BaseColorForGITexture;
@@ -4412,6 +4412,18 @@ void FLandscapeComponentSceneProxy::GetHeightfieldRepresentation(UTexture2D*& Ou
 	OutDescription.NumSubsections = NumSubsections;
 
 	OutDescription.SubsectionScaleAndBias = FVector4(SubsectionSizeQuads, SubsectionSizeQuads, HeightmapSubsectionOffsetU, HeightmapSubsectionOffsetV);
+
+	OutVisibilityTexture = nullptr;
+	OutDescription.VisibilityChannel = -1;
+	for (auto& Allocation : LandscapeComponent->WeightmapLayerAllocations)
+	{
+		if (Allocation.LayerInfo == ALandscapeProxy::VisibilityLayer)
+		{
+			OutVisibilityTexture = WeightmapTextures[Allocation.WeightmapTextureIndex];
+			OutDescription.VisibilityChannel = Allocation.WeightmapTextureChannel;
+			break;
+		}
+	}
 }
 
 void FLandscapeComponentSceneProxy::GetLCIs(FLCIArray& LCIs)
diff --git a/Engine/Source/Runtime/Landscape/Public/LandscapeRender.h b/Engine/Source/Runtime/Landscape/Public/LandscapeRender.h
--- a/Engine/Source/Runtime/Landscape/Public/LandscapeRender.h
+++ b/Engine/Source/Runtime/Landscape/Public/LandscapeRender.h
@@ -708,7 +708,7 @@ public:
 
 	virtual bool HeightfieldHasPendingStreaming() const override;
 
-	virtual void GetHeightfieldRepresentation(UTexture2D*& OutHeightmapTexture, UTexture2D*& OutDiffuseColorTexture, FHeightfieldComponentDescription& OutDescription) override;
+	virtual void GetHeightfieldRepresentation(UTexture2D*& OutHeightmapTexture, UTexture2D*& OutDiffuseColorTexture, UTexture2D*& OutVisibilityTexture, FHeightfieldComponentDescription& OutDescription) override;
 
 	virtual void GetLCIs(FLCIArray& LCIs) override;
 
diff --git a/Engine/Source/Runtime/Renderer/Private/GlobalDistanceField.cpp b/Engine/Source/Runtime/Renderer/Private/GlobalDistanceField.cpp
--- a/Engine/Source/Runtime/Renderer/Private/GlobalDistanceField.cpp
+++ b/Engine/Source/Runtime/Renderer/Private/GlobalDistanceField.cpp
@@ -634,6 +634,7 @@ public:
 		int32 ClipmapIndexValue,
 		const FVolumeUpdateRegion& UpdateRegion,
 		UTexture2D* HeightfieldTextureValue,
+		UTexture2D* VisibilityTextureValue,
 		int32 NumHeightfieldsValue)
 	{
 		FComputeShaderRHIParamRef ShaderRHI = GetComputeShader();
@@ -653,7 +654,7 @@ public:
 		SetShaderValue(RHICmdList, ShaderRHI, AOGlobalMaxSphereQueryRadius, GlobalMaxSphereQueryRadius);
 
 		HeightfieldDescriptionParameters.Set(RHICmdList, ShaderRHI, GetHeightfieldDescriptionsSRV(), NumHeightfieldsValue);
-		HeightfieldTextureParameters.Set(RHICmdList, ShaderRHI, HeightfieldTextureValue, NULL);
+		HeightfieldTextureParameters.Set(RHICmdList, ShaderRHI, HeightfieldTextureValue, NULL, VisibilityTextureValue);
 	}
 
 	void UnsetParameters(FRHICommandList& RHICmdList, const FGlobalDistanceFieldInfo& GlobalDistanceFieldInfo, const FGlobalDistanceFieldClipmap& Clipmap)
@@ -735,8 +736,9 @@ void FHeightfieldLightingViewInfo::CompositeHeightfieldsIntoGlobalDistanceField(
 			{
 				UTexture2D* HeightfieldTexture = NULL;
 				UTexture2D* DiffuseColorTexture = NULL;
+				UTexture2D* VisibilityTexture = NULL;
 				FHeightfieldComponentDescription NewComponentDescription(HeightfieldPrimitive->Proxy->GetLocalToWorld());
-				HeightfieldPrimitive->Proxy->GetHeightfieldRepresentation(HeightfieldTexture, DiffuseColorTexture, NewComponentDescription);
+				HeightfieldPrimitive->Proxy->GetHeightfieldRepresentation(HeightfieldTexture, DiffuseColorTexture, VisibilityTexture, NewComponentDescription);
 
 				if (HeightfieldTexture && HeightfieldTexture->Resource->TextureRHI)
 				{
@@ -751,7 +753,7 @@ void FHeightfieldLightingViewInfo::CompositeHeightfieldsIntoGlobalDistanceField(
 						UpdateRegionHeightfield.Rect.Union(NewComponentDescription.HeightfieldRect);
 					}
 
-					TArray<FHeightfieldComponentDescription>& ComponentDescriptions = UpdateRegionHeightfield.ComponentDescriptions.FindOrAdd(FHeightfieldComponentTextures(HeightfieldTexture, DiffuseColorTexture));
+					TArray<FHeightfieldComponentDescription>& ComponentDescriptions = UpdateRegionHeightfield.ComponentDescriptions.FindOrAdd(FHeightfieldComponentTextures(HeightfieldTexture, DiffuseColorTexture, VisibilityTexture));
 					ComponentDescriptions.Add(NewComponentDescription);
 				}
 			}
@@ -768,10 +770,11 @@ void FHeightfieldLightingViewInfo::CompositeHeightfieldsIntoGlobalDistanceField(
 					UploadHeightfieldDescriptions(HeightfieldDescriptions, FVector2D(1, 1), 1.0f / UpdateRegionHeightfield.DownsampleFactor);
 
 					UTexture2D* HeightfieldTexture = It.Key().HeightAndNormal;
+					UTexture2D* VisibilityTexture = It.Key().Visibility;
 
 					TShaderMapRef<FCompositeHeightfieldsIntoGlobalDistanceFieldCS> ComputeShader(View.ShaderMap);
 					RHICmdList.SetComputeShader(ComputeShader->GetComputeShader());
-					ComputeShader->SetParameters(RHICmdList, Scene, View, GlobalMaxSphereQueryRadius, GlobalDistanceFieldInfo, Clipmap, ClipmapIndexValue, UpdateRegion, HeightfieldTexture, HeightfieldDescriptions.Num());
+					ComputeShader->SetParameters(RHICmdList, Scene, View, GlobalMaxSphereQueryRadius, GlobalDistanceFieldInfo, Clipmap, ClipmapIndexValue, UpdateRegion, HeightfieldTexture, VisibilityTexture, HeightfieldDescriptions.Num());
 
 					//@todo - match typical update sizes.  Camera movement creates narrow slabs.
 					const uint32 NumGroupsX = FMath::DivideAndRoundUp<int32>(UpdateRegion.CellsSize.X, HeightfieldCompositeTileSize);
diff --git a/Engine/Source/Runtime/Renderer/Private/HeightfieldLighting.cpp b/Engine/Source/Runtime/Renderer/Private/HeightfieldLighting.cpp
--- a/Engine/Source/Runtime/Renderer/Private/HeightfieldLighting.cpp
+++ b/Engine/Source/Runtime/Renderer/Private/HeightfieldLighting.cpp
@@ -313,7 +313,7 @@ public:
 		const FPixelShaderRHIParamRef ShaderRHI = GetPixelShader();
 
 		FGlobalShader::SetParameters<FViewUniformShaderParameters>(RHICmdList, ShaderRHI, View.ViewUniformBuffer);
-		HeightfieldTextureParameters.Set(RHICmdList, ShaderRHI, HeightfieldTextureValue, DiffuseColorTextureValue);
+		HeightfieldTextureParameters.Set(RHICmdList, ShaderRHI, HeightfieldTextureValue, DiffuseColorTextureValue, nullptr);
 	}
 	// FShader interface.
 	virtual bool Serialize(FArchive& Ar) override
@@ -408,8 +408,10 @@ void FHeightfieldLightingViewInfo::SetupVisibleHeightfields(const FViewInfo& Vie
 			{
 				UTexture2D* HeightfieldTexture = NULL;
 				UTexture2D* DiffuseColorTexture = NULL;
+				UTexture2D* VisibilityTexture = NULL;
+
 				FHeightfieldComponentDescription NewComponentDescription(HeightfieldPrimitive->Proxy->GetLocalToWorld());
-				HeightfieldPrimitive->Proxy->GetHeightfieldRepresentation(HeightfieldTexture, DiffuseColorTexture, NewComponentDescription);
+				HeightfieldPrimitive->Proxy->GetHeightfieldRepresentation(HeightfieldTexture, DiffuseColorTexture, VisibilityTexture, NewComponentDescription);
 
 				if (HeightfieldTexture && HeightfieldTexture->Resource->TextureRHI)
 				{
@@ -425,7 +427,7 @@ void FHeightfieldLightingViewInfo::SetupVisibleHeightfields(const FViewInfo& Vie
 						Heightfield.Rect.Union(NewComponentDescription.HeightfieldRect);
 					}
 
-					TArray<FHeightfieldComponentDescription>& ComponentDescriptions = Heightfield.ComponentDescriptions.FindOrAdd(FHeightfieldComponentTextures(HeightfieldTexture, DiffuseColorTexture));
+					TArray<FHeightfieldComponentDescription>& ComponentDescriptions = Heightfield.ComponentDescriptions.FindOrAdd(FHeightfieldComponentTextures(HeightfieldTexture, DiffuseColorTexture, VisibilityTexture));
 					ComponentDescriptions.Add(NewComponentDescription);
 				}
 			}
@@ -534,8 +536,10 @@ void FHeightfieldLightingViewInfo::SetupHeightfieldsForScene(const FScene& Scene
 
 			UTexture2D* HeightfieldTexture = NULL;
 			UTexture2D* DiffuseColorTexture = NULL;
+			UTexture2D* VisibilityTexture = NULL;
+
 			FHeightfieldComponentDescription NewComponentDescription(HeightfieldPrimitive->Proxy->GetLocalToWorld());
-			HeightfieldPrimitive->Proxy->GetHeightfieldRepresentation(HeightfieldTexture, DiffuseColorTexture, NewComponentDescription);
+			HeightfieldPrimitive->Proxy->GetHeightfieldRepresentation(HeightfieldTexture, DiffuseColorTexture, VisibilityTexture, NewComponentDescription);
 
 			if (HeightfieldTexture && HeightfieldTexture->Resource->TextureRHI)
 			{
@@ -551,7 +555,7 @@ void FHeightfieldLightingViewInfo::SetupHeightfieldsForScene(const FScene& Scene
 					Heightfield.Rect.Union(NewComponentDescription.HeightfieldRect);
 				}
 
-				TArray<FHeightfieldComponentDescription>& ComponentDescriptions = Heightfield.ComponentDescriptions.FindOrAdd(FHeightfieldComponentTextures(HeightfieldTexture, DiffuseColorTexture));
+				TArray<FHeightfieldComponentDescription>& ComponentDescriptions = Heightfield.ComponentDescriptions.FindOrAdd(FHeightfieldComponentTextures(HeightfieldTexture, DiffuseColorTexture, VisibilityTexture));
 				ComponentDescriptions.Add(NewComponentDescription);
 			}
 		}
@@ -567,7 +571,7 @@ public:
 	FHeightfieldDescriptionsResource()
 	{
 		// In float4's, must match usf
-		const int32 HEIGHTFIELD_DATA_STRIDE = 10;
+		const int32 HEIGHTFIELD_DATA_STRIDE = 12;
 
 		Data.Format = PF_A32B32G32R32F;
 		Data.Stride = HEIGHTFIELD_DATA_STRIDE;
@@ -620,7 +624,8 @@ void UploadHeightfieldDescriptions(const TArray<FHeightfieldComponentDescription
 
 		HeightfieldDescriptionData.Add(LightingUVScaleBias);
 
-		HeightfieldDescriptionData.Add(FVector4(Description.HeightfieldRect.Size().X, Description.HeightfieldRect.Size().Y, InvLightingAtlasSize.X, InvLightingAtlasSize.Y));
+		HeightfieldDescriptionData.Add(FVector4(Description.HeightfieldRect.Size().X, Description.HeightfieldRect.Size().Y, 1.f/Description.HeightfieldRect.Size().X, 1.f/Description.HeightfieldRect.Size().Y));
+		HeightfieldDescriptionData.Add(FVector4(InvLightingAtlasSize.X, InvLightingAtlasSize.Y, 0.f, 0.f));
 
 		const FMatrix LocalToWorldT = Description.LocalToWorld.GetTransposed();
 		const FMatrix WorldToLocalT = Description.LocalToWorld.Inverse().GetTransposed();
@@ -632,6 +637,13 @@ void UploadHeightfieldDescriptions(const TArray<FHeightfieldComponentDescription
 		HeightfieldDescriptionData.Add(*(FVector4*)&LocalToWorldT.M[0]);
 		HeightfieldDescriptionData.Add(*(FVector4*)&LocalToWorldT.M[1]);
 		HeightfieldDescriptionData.Add(*(FVector4*)&LocalToWorldT.M[2]);
+
+		FVector4 ChannelMask(0.f, 0.f, 0.f, 0.f);
+		if (Description.VisibilityChannel >= 0 && Description.VisibilityChannel < 4)
+		{
+			ChannelMask.Component(Description.VisibilityChannel) = 1.f;
+		}
+		HeightfieldDescriptionData.Add(ChannelMask);
 	}
 
 	check(HeightfieldDescriptionData.Num() % GHeightfieldDescriptions.Data.Stride == 0);
@@ -1289,7 +1301,7 @@ public:
 		AOParameters.Set(RHICmdList, ShaderRHI, Parameters);
 		ScreenGridParameters.Set(RHICmdList, ShaderRHI, View, DistanceFieldNormal);
 		HeightfieldDescriptionParameters.Set(RHICmdList, ShaderRHI, GetHeightfieldDescriptionsSRV(), NumHeightfieldsValue);
-		HeightfieldTextureParameters.Set(RHICmdList, ShaderRHI, HeightfieldTextureValue, NULL);
+		HeightfieldTextureParameters.Set(RHICmdList, ShaderRHI, HeightfieldTextureValue, nullptr, nullptr);
 
 		RHICmdList.TransitionResource(EResourceTransitionAccess::ERWBarrier, EResourceTransitionPipeline::EComputeToCompute, ScreenGridResources.ScreenGridConeVisibility.UAV);
 		ScreenGridConeVisibility.SetBuffer(RHICmdList, ShaderRHI, ScreenGridResources.ScreenGridConeVisibility);
diff --git a/Engine/Source/Runtime/Renderer/Private/HeightfieldLighting.h b/Engine/Source/Runtime/Renderer/Private/HeightfieldLighting.h
--- a/Engine/Source/Runtime/Renderer/Private/HeightfieldLighting.h
+++ b/Engine/Source/Runtime/Renderer/Private/HeightfieldLighting.h
@@ -54,14 +54,15 @@ class FHeightfieldComponentTextures
 {
 public:
 
-	FHeightfieldComponentTextures(UTexture2D* InHeightAndNormal, UTexture2D* InDiffuseColor) :
+	FHeightfieldComponentTextures(UTexture2D* InHeightAndNormal, UTexture2D* InDiffuseColor, UTexture2D* InVisibility) :
 		HeightAndNormal(InHeightAndNormal),
-		DiffuseColor(InDiffuseColor)
+		DiffuseColor(InDiffuseColor),
+		Visibility(InVisibility)
 	{}
 
 	FORCEINLINE bool operator==(FHeightfieldComponentTextures Other) const
 	{
-		return HeightAndNormal == Other.HeightAndNormal && DiffuseColor == Other.DiffuseColor;
+		return HeightAndNormal == Other.HeightAndNormal && DiffuseColor == Other.DiffuseColor && Visibility == Other.Visibility;
 	}
 
 	FORCEINLINE friend uint32 GetTypeHash(FHeightfieldComponentTextures ComponentTextures)
@@ -71,6 +72,7 @@ public:
 
 	UTexture2D* HeightAndNormal;
 	UTexture2D* DiffuseColor;
+	UTexture2D* Visibility;
 };
 
 class FHeightfieldDescription
@@ -181,6 +183,8 @@ public:
 		HeightfieldSampler.Bind(ParameterMap, TEXT("HeightfieldSampler"));
 		DiffuseColorTexture.Bind(ParameterMap, TEXT("DiffuseColorTexture"));
 		DiffuseColorSampler.Bind(ParameterMap, TEXT("DiffuseColorSampler"));
+		VisibilityTexture.Bind(ParameterMap, TEXT("VisibilityTexture"));
+		VisibilitySampler.Bind(ParameterMap, TEXT("VisibilitySampler"));
 	}
 
 	friend FArchive& operator<<(FArchive& Ar, FHeightfieldTextureParameters& Parameters)
@@ -189,11 +193,13 @@ public:
 		Ar << Parameters.HeightfieldSampler;
 		Ar << Parameters.DiffuseColorTexture;
 		Ar << Parameters.DiffuseColorSampler;
+		Ar << Parameters.VisibilityTexture;
+		Ar << Parameters.VisibilitySampler;
 		return Ar;
 	}
 
 	template<typename ShaderRHIParamRef>
-	void Set(FRHICommandList& RHICmdList, const ShaderRHIParamRef ShaderRHI, UTexture2D* HeightfieldTextureValue, UTexture2D* DiffuseColorTextureValue)
+	void Set(FRHICommandList& RHICmdList, const ShaderRHIParamRef ShaderRHI, UTexture2D* HeightfieldTextureValue, UTexture2D* DiffuseColorTextureValue, UTexture2D* VisibilityTextureValue)
 	{
 		//@todo - shouldn't filter the heightfield, it's packed
 		SetTextureParameter(RHICmdList, ShaderRHI, HeightfieldTexture, HeightfieldSampler, TStaticSamplerState<SF_Bilinear>::GetRHI(), HeightfieldTextureValue->Resource->TextureRHI);
@@ -206,6 +212,15 @@ public:
 		{
 			SetTextureParameter(RHICmdList, ShaderRHI, DiffuseColorTexture, DiffuseColorSampler, TStaticSamplerState<SF_Bilinear>::GetRHI(), GBlackTexture->TextureRHI);
 		}
+
+		if (VisibilityTextureValue)
+		{
+			SetTextureParameter(RHICmdList, ShaderRHI, VisibilityTexture, VisibilitySampler, TStaticSamplerState<SF_Bilinear>::GetRHI(), VisibilityTextureValue->Resource->TextureRHI);
+		}
+		else
+		{
+			SetTextureParameter(RHICmdList, ShaderRHI, VisibilityTexture, VisibilitySampler, TStaticSamplerState<SF_Bilinear>::GetRHI(), GBlackTexture->TextureRHI);
+		}
 	}
 
 private:
@@ -213,4 +228,6 @@ private:
 	FShaderResourceParameter HeightfieldSampler;
 	FShaderResourceParameter DiffuseColorTexture;
 	FShaderResourceParameter DiffuseColorSampler;
+	FShaderResourceParameter VisibilityTexture;
+	FShaderResourceParameter VisibilitySampler;
 };
