Fixes for decal fade in/out blueprint API

diff --git a/Engine/Source/Runtime/Renderer/Private/RendererScene.cpp b/Engine/Source/Runtime/Renderer/Private/RendererScene.cpp
--- a/Engine/Source/Runtime/Renderer/Private/RendererScene.cpp
+++ b/Engine/Source/Runtime/Renderer/Private/RendererScene.cpp
@@ -1871,8 +1871,8 @@ void FScene::UpdateDecalFadeOutTime(UDecalComponent* Decal)
 			}
 			else
 			{
-				Proxy->InvFadeInDuration = -1.0f;
-				Proxy->FadeInStartDelayNormalized = 1.0f;
+				Proxy->InvFadeDuration = -1.0f;
+				Proxy->FadeStartDelayNormalized = 1.0f;
 			}
 		});
 	}
@@ -1893,7 +1893,7 @@ void FScene::UpdateDecalFadeInTime(UDecalComponent* Decal)
 			if (DecalFadeDuration > 0.0f)
 			{
 				Proxy->InvFadeInDuration = 1.0f / DecalFadeDuration;
-				Proxy->FadeInStartDelayNormalized = (CurrentTime + DecalFadeDuration) * -Proxy->InvFadeInDuration;
+				Proxy->FadeInStartDelayNormalized = (CurrentTime + DecalFadeStartDelay) * -Proxy->InvFadeInDuration;
 			}
 			else
 			{
