DR - Fix warning
#rnx

diff --git a/Engine/Source/Runtime/Renderer/Private/ScreenSpaceDiffuseIndirect.cpp b/Engine/Source/Runtime/Renderer/Private/ScreenSpaceDiffuseIndirect.cpp
--- a/Engine/Source/Runtime/Renderer/Private/ScreenSpaceDiffuseIndirect.cpp
+++ b/Engine/Source/Runtime/Renderer/Private/ScreenSpaceDiffuseIndirect.cpp
@@ -149,11 +149,11 @@ void RenderScreenSpaceDiffuseIndirect( FRHICommandListImmediate& RHICmdList, FVi
 		RDG_EVENT_NAME( "ScreenSpaceDiffuseIndirect %dx%d", View.ViewRect.Width(), View.ViewRect.Height() ),
 		PassParameters,
 		ERenderGraphPassFlags::None,
-		[ PassParameters, &View, Quality ]( FRHICommandListImmediate& RHICmdList )
+		[ PassParameters, &View, Quality ]( FRHICommandListImmediate& InRHICmdList )
 		{
-			SCOPED_GPU_STAT(RHICmdList, ScreenSpaceDiffuseIndirect);
+			SCOPED_GPU_STAT(InRHICmdList, ScreenSpaceDiffuseIndirect);
 
-			RHICmdList.SetViewport( View.ViewRect.Min.X, View.ViewRect.Min.Y, 0.0f, View.ViewRect.Max.X, View.ViewRect.Max.Y, 1.0f );
+			InRHICmdList.SetViewport( View.ViewRect.Min.X, View.ViewRect.Min.Y, 0.0f, View.ViewRect.Max.X, View.ViewRect.Max.Y, 1.0f );
 
 			FScreenSpaceDiffuseIndirectPS::FPermutationDomain PermutationVector;
 			PermutationVector.Set< FScreenSpaceDiffuseIndirectPS::FQualityDim >( Quality );
@@ -163,7 +163,7 @@ void RenderScreenSpaceDiffuseIndirect( FRHICommandListImmediate& RHICmdList, FVi
 
 			{
 				FGraphicsPipelineStateInitializer GraphicsPSOInit;
-				RHICmdList.ApplyCachedRenderTargets( GraphicsPSOInit );
+				InRHICmdList.ApplyCachedRenderTargets( GraphicsPSOInit );
 
 				GraphicsPSOInit.BlendState = TStaticBlendState< CW_RGB, BO_Add, BF_One, BF_SourceAlpha >::GetRHI();
 				GraphicsPSOInit.DepthStencilState = TStaticDepthStencilState<false, CF_Always>::GetRHI();
@@ -173,19 +173,19 @@ void RenderScreenSpaceDiffuseIndirect( FRHICommandListImmediate& RHICmdList, FVi
 				GraphicsPSOInit.BoundShaderState.VertexShaderRHI = GETSAFERHISHADER_VERTEX( VertexShader );
 				GraphicsPSOInit.BoundShaderState.PixelShaderRHI = GETSAFERHISHADER_PIXEL( PixelShader );
 
-				SetGraphicsPipelineState( RHICmdList, GraphicsPSOInit );
+				SetGraphicsPipelineState( InRHICmdList, GraphicsPSOInit );
 			}
 
-			SetShaderParameters( RHICmdList, PixelShader, PixelShader->GetPixelShader(), *PassParameters );
+			SetShaderParameters( InRHICmdList, PixelShader, PixelShader->GetPixelShader(), *PassParameters );
 
 			DrawPostProcessPass(
-				RHICmdList,
+				InRHICmdList,
 				0, 0,
 				View.ViewRect.Width(), View.ViewRect.Height(),
 				View.ViewRect.Min.X, View.ViewRect.Min.Y,
 				View.ViewRect.Width(), View.ViewRect.Height(),
 				View.ViewRect.Size(),
-				FSceneRenderTargets::Get( RHICmdList ).GetBufferSizeXY(),
+				FSceneRenderTargets::Get( InRHICmdList ).GetBufferSizeXY(),
 				VertexShader,
 				View.StereoPass,
 				false,
