Disable auto-exposure for distance field visualization modes.

diff --git a/Engine/Source/Runtime/Renderer/Private/PostProcess/PostProcessEyeAdaptation.cpp b/Engine/Source/Runtime/Renderer/Private/PostProcess/PostProcessEyeAdaptation.cpp
--- a/Engine/Source/Runtime/Renderer/Private/PostProcess/PostProcessEyeAdaptation.cpp
+++ b/Engine/Source/Runtime/Renderer/Private/PostProcess/PostProcessEyeAdaptation.cpp
@@ -87,7 +87,12 @@ inline static void ComputeEyeAdaptationValues(const ERHIFeatureLevel::Type MinFe
 	float AverageLuminanceScale = 1.f;
 
 	// When any of those flags are set, make sure the tonemapper uses an exposure of 1.
-	if (!EngineShowFlags.Lighting || (EngineShowFlags.VisualizeBuffer && View.CurrentBufferVisualizationMode != NAME_None) || View.Family->UseDebugViewPS() || EngineShowFlags.RayTracingDebug )
+	if (!EngineShowFlags.Lighting 
+		|| (EngineShowFlags.VisualizeBuffer && View.CurrentBufferVisualizationMode != NAME_None) 
+		|| View.Family->UseDebugViewPS() 
+		|| EngineShowFlags.RayTracingDebug 
+		|| EngineShowFlags.VisualizeDistanceFieldAO
+		|| EngineShowFlags.VisualizeGlobalDistanceField)
 	{
 		LocalExposureMultipler = 1.f;
 	}
