Fix pre-exposure not working in certain view type.

Change IsRich view check to a custom flag check to compute the pre-exposure value when needed. This is necessary for having correct pre-exposure value in certain view like in the shader preview window. This issue was reported in a UDN ticket.

diff --git a/Engine/Source/Runtime/Renderer/Private/PostProcess/PostProcessEyeAdaptation.cpp b/Engine/Source/Runtime/Renderer/Private/PostProcess/PostProcessEyeAdaptation.cpp
--- a/Engine/Source/Runtime/Renderer/Private/PostProcess/PostProcessEyeAdaptation.cpp
+++ b/Engine/Source/Runtime/Renderer/Private/PostProcess/PostProcessEyeAdaptation.cpp
@@ -452,7 +452,23 @@ void FSceneViewState::UpdatePreExposure(FViewInfo& View)
 	{
 		const FSceneViewFamily& ViewFamily = *View.Family;
 
-		if (!IsRichView(ViewFamily) /*&& !ViewFamily.EngineShowFlags.VisualizeHDR */&& !ViewFamily.EngineShowFlags.VisualizeBloom && ViewFamily.EngineShowFlags.PostProcessing && ViewFamily.bResolveScene)
+		// One could use the IsRichView functionality to check if we need to update pre-exposure, 
+		// but this is too limiting for certain view. For instance shader preview doesn't have 
+		// volumetric lighting enabled, which makes the view be flagged as rich, and not updating 
+		// the pre-exposition value.
+		const bool bUpdatePreExposure =
+			ViewFamily.EngineShowFlags.Lighting &&
+			ViewFamily.EngineShowFlags.PostProcessing &&
+			ViewFamily.bResolveScene &&
+			!ViewFamily.EngineShowFlags.LightMapDensity &&
+			!ViewFamily.EngineShowFlags.StationaryLightOverlap &&
+			!ViewFamily.EngineShowFlags.LightComplexity &&
+			!ViewFamily.EngineShowFlags.LODColoration &&
+			!ViewFamily.EngineShowFlags.HLODColoration &&
+			!ViewFamily.EngineShowFlags.LevelColoration &&
+			!ViewFamily.EngineShowFlags.VisualizeBloom;
+
+		if (bUpdatePreExposure)
 		{
 			const float PreExposureOverride = CVarEyeAdaptationPreExposureOverride.GetValueOnRenderThread();
 			const float LastExposure = View.GetLastEyeAdaptationExposure();
