* Add basic Shader/GPU driven debug rendering
 * Add cursor postion to View constant buffer.

Simple API: allocate, bind (for PS), draw. Support only wire-quad/cross/line drawing. Currently readback on CPU, but should drawing could/should be done directly on the GPU.

diff --git a/Engine/Shaders/Private/GpuDebugRendering.ush b/Engine/Shaders/Private/GpuDebugRendering.ush
new file mode 100644
--- /dev/null
+++ b/Engine/Shaders/Private/GpuDebugRendering.ush
@@ -0,0 +1,69 @@
+#pragma once
+
+#define DEBUG_GPU_RENDERING 1
+
+#if DEBUG_GPU_RENDERING
+
+#define COLOR_RED		0
+#define COLOR_GREEN		1
+#define COLOR_BLUE		2
+#define COLOR_YELLOW	3
+#define COLOR_WHITE		4
+
+// TODO abstract/add binding configuration based on define 
+// Note: for pixel shader, UAV are in the output space, if one RT is bound, first available UAV is 1
+RWTexture2D<uint>	OutputDebugPrimitiveCount	: register(u1);
+RWTexture2D<float4>	OutputDebugPrimitive		: register(u2);
+
+static bool	bGpuDebugEnable					= false;
+static const uint MaxDebugPrimitiveCount	= 1024;		// Match GpuDebugRednering.cpp :: AllocateGpuDebugPrimitiveBuffers()
+
+void SetGpuDebugEnable(bool Value) 	{ bGpuDebugEnable = Value; }
+bool IsGpuDebugEnable() 			{ return bGpuDebugEnable; }
+
+void AddLine(float3 StartPoint, float3 EndPoint, uint Color)
+{
+	uint Entry = -2;
+    InterlockedAdd(OutputDebugPrimitiveCount[uint2(0, 0)], 2, Entry);
+	if (Entry < MaxDebugPrimitiveCount+2)
+    {
+		OutputDebugPrimitive[uint2(Entry + 0, 0)] = float4(StartPoint, Color);
+		OutputDebugPrimitive[uint2(Entry + 1, 0)] = float4(EndPoint, Color);
+    }
+}
+
+void AddQuad(float3 P0, float3 P1, float3 P2, float3 P3, uint Color)
+{
+	uint Entry = -2;
+    InterlockedAdd(OutputDebugPrimitiveCount[uint2(0, 0)], 8, Entry);
+	if (Entry < MaxDebugPrimitiveCount+8)
+    {
+		OutputDebugPrimitive[uint2(Entry + 0, 0)] = float4(P0, Color);
+		OutputDebugPrimitive[uint2(Entry + 1, 0)] = float4(P1, Color);
+		OutputDebugPrimitive[uint2(Entry + 2, 0)] = float4(P1, Color);
+		OutputDebugPrimitive[uint2(Entry + 3, 0)] = float4(P2, Color);
+		OutputDebugPrimitive[uint2(Entry + 4, 0)] = float4(P2, Color);
+		OutputDebugPrimitive[uint2(Entry + 5, 0)] = float4(P3, Color);
+		OutputDebugPrimitive[uint2(Entry + 6, 0)] = float4(P3, Color);
+		OutputDebugPrimitive[uint2(Entry + 7, 0)] = float4(P0, Color);
+    }
+}
+
+void AddCross(float3 Center, float Size, uint Color)
+{
+	uint Entry = -2;
+    InterlockedAdd(OutputDebugPrimitiveCount[uint2(0, 0)], 6, Entry);
+	if (Entry < MaxDebugPrimitiveCount+6)
+    {
+		OutputDebugPrimitive[uint2(Entry + 0, 0)] = float4(Center+float3(-Size,0,0), Color);
+		OutputDebugPrimitive[uint2(Entry + 1, 0)] = float4(Center+float3( Size,0,0), Color);
+
+		OutputDebugPrimitive[uint2(Entry + 2, 0)] = float4(Center+float3(0,-Size,0), Color);
+		OutputDebugPrimitive[uint2(Entry + 3, 0)] = float4(Center+float3(0, Size,0), Color);
+
+		OutputDebugPrimitive[uint2(Entry + 4, 0)] = float4(Center+float3(0,0,-Size), Color);
+		OutputDebugPrimitive[uint2(Entry + 5, 0)] = float4(Center+float3(0,0, Size), Color);
+    }
+}
+
+#endif
\ No newline at end of file
diff --git a/Engine/Source/Runtime/Engine/Private/SceneView.cpp b/Engine/Source/Runtime/Engine/Private/SceneView.cpp
--- a/Engine/Source/Runtime/Engine/Private/SceneView.cpp
+++ b/Engine/Source/Runtime/Engine/Private/SceneView.cpp
@@ -49,6 +49,13 @@ static TAutoConsoleVariable<float> CVarSSRMaxRoughness(
 
 #if !(UE_BUILD_SHIPPING || UE_BUILD_TEST)
 
+static TAutoConsoleVariable<int32> CVarFreezeMouseCursor(
+	TEXT("r.FreezeMouseCursor"),
+	0,
+	TEXT("Free the mouse cursor position, for passes which use it to display debug information.\n")
+	TEXT("0: default\n")
+	TEXT("1: freeze mouse cursor position at current location"),
+	ECVF_Cheat);
 
 static TAutoConsoleVariable<int32> CVarShadowFreezeCamera(
 	TEXT("r.Shadow.FreezeCamera"),
@@ -2219,6 +2226,12 @@ void FSceneView::SetupCommonViewUniformBufferParameters(
 	ViewUniformShaderParameters.GameTime = Family->CurrentWorldTime;
 	ViewUniformShaderParameters.RealTime = Family->CurrentRealTime;
 	ViewUniformShaderParameters.DeltaTime = Family->DeltaWorldTime;
+	static FIntPoint LockedCursorPos = CursorPos;
+	if (CVarFreezeMouseCursor.GetValueOnRenderThread() == 0 && CursorPos.X >= 0 && CursorPos.Y >= 0)
+	{
+		LockedCursorPos = CursorPos;
+	}
+	ViewUniformShaderParameters.CursorPosition = LockedCursorPos;
 	ViewUniformShaderParameters.Random = FMath::Rand();
 	ViewUniformShaderParameters.FrameNumber = Family->FrameNumber;
 
diff --git a/Engine/Source/Runtime/Engine/Public/SceneView.h b/Engine/Source/Runtime/Engine/Public/SceneView.h
--- a/Engine/Source/Runtime/Engine/Public/SceneView.h
+++ b/Engine/Source/Runtime/Engine/Public/SceneView.h
@@ -688,6 +688,7 @@ enum ETranslucencyVolumeCascade
 	VIEW_UNIFORM_BUFFER_MEMBER(float, GlobalVolumeDimension) \
 	VIEW_UNIFORM_BUFFER_MEMBER(float, GlobalVolumeTexelSize) \
 	VIEW_UNIFORM_BUFFER_MEMBER(float, MaxGlobalDistance) \
+	VIEW_UNIFORM_BUFFER_MEMBER(FIntPoint, CursorPosition) \
 	VIEW_UNIFORM_BUFFER_MEMBER(float, bCheckerboardSubsurfaceProfileRendering) \
 	VIEW_UNIFORM_BUFFER_MEMBER(FVector, VolumetricFogInvGridSize) \
 	VIEW_UNIFORM_BUFFER_MEMBER(FVector, VolumetricFogGridZParams) \
diff --git a/Engine/Source/Runtime/Renderer/Private/GpuDebugRendering.cpp b/Engine/Source/Runtime/Renderer/Private/GpuDebugRendering.cpp
new file mode 100644
--- /dev/null
+++ b/Engine/Source/Runtime/Renderer/Private/GpuDebugRendering.cpp
@@ -0,0 +1,145 @@
+// Copyright 1998-2019 Epic Games, Inc. All Rights Reserved.
+
+#include "GpuDebugRendering.h"
+#include "DynamicPrimitiveDrawing.h"
+#include "RenderGraphUtils.h"
+#include "RenderTargetPool.h"
+
+FGpuDebugPrimitiveBuffers AllocateGpuDebugPrimitiveBuffers(FRHICommandListImmediate& RHICmdList)
+{
+	FGpuDebugPrimitiveBuffers Output;
+	{
+		FIntPoint Resolution(1, 1);
+		EPixelFormat Format = PF_R32_UINT;
+		{
+			FPooledRenderTargetDesc Desc(FPooledRenderTargetDesc::Create2DDesc(Resolution, Format, FClearValueBinding::Black, TexCreate_None, TexCreate_UAV, false));
+			GRenderTargetPool.FindFreeElement(RHICmdList, Desc, Output.DebugPrimitiveCountTexture, TEXT("DebugPrimitiveCountTexture"));
+		}
+		{
+			FPooledRenderTargetDesc Desc(FPooledRenderTargetDesc::Create2DDesc(Resolution, Format, FClearValueBinding::None, TexCreate_CPUReadback, TexCreate_None, false));
+			GRenderTargetPool.FindFreeElement(RHICmdList, Desc, Output.DebugPrimitiveCountStagingTexture, TEXT("DebugPrimitiveCountStagingTexture"));
+		}
+	}
+
+	{
+		FIntPoint Resolution(1024, 1);
+		EPixelFormat Format = PF_A32B32G32R32F;
+		{
+			FPooledRenderTargetDesc Desc(FPooledRenderTargetDesc::Create2DDesc(Resolution, Format, FClearValueBinding::Black, TexCreate_None, TexCreate_UAV, false));
+			GRenderTargetPool.FindFreeElement(RHICmdList, Desc, Output.DebugPrimitiveTexture, TEXT("DebugPrimitiveTexture"));
+		}
+		{
+			FPooledRenderTargetDesc Desc(FPooledRenderTargetDesc::Create2DDesc(Resolution, Format, FClearValueBinding::None, TexCreate_CPUReadback, TexCreate_None, false));
+			GRenderTargetPool.FindFreeElement(RHICmdList, Desc, Output.DebugPrimitiveStagingTexture, TEXT("DebugPrimitiveStagingTexture"));
+		}
+	}
+
+	uint32 ClearValues[4] = { 0, 0, 0, 0 };
+	RHICmdList.ClearTinyUAV(Output.DebugPrimitiveCountTexture->GetRenderTargetItem().UAV, ClearValues);
+	RHICmdList.ClearTinyUAV(Output.DebugPrimitiveTexture->GetRenderTargetItem().UAV, ClearValues);
+
+	return Output;
+}
+
+struct FLine { FVector Start; FVector End; FLinearColor Color; };
+typedef TArray<FLine> FLineArray;
+static FLineArray ReadGPUDebugPrimitives(FRHICommandListImmediate& RHICmdList, FGpuDebugPrimitiveBuffers& DebugPrimitiveBuffer)
+{
+	// Count
+	uint32 PointCount = 0;
+	{
+		FTextureRHIRef SourceTexture = DebugPrimitiveBuffer.DebugPrimitiveCountTexture->GetRenderTargetItem().TargetableTexture;
+		FTextureRHIRef StagingTexture = DebugPrimitiveBuffer.DebugPrimitiveCountStagingTexture->GetRenderTargetItem().ShaderResourceTexture;
+
+		// Transfer memory GPU -> CPU
+		RHICmdList.CopyToResolveTarget(SourceTexture, StagingTexture, FResolveParams());
+		if (StagingTexture.IsValid())
+		{
+			const uint32* CountData = nullptr;
+			int32 BufferWidth = 0, BufferHeight = 0;
+			RHICmdList.MapStagingSurface(StagingTexture, *(void**)&CountData, BufferWidth, BufferHeight);
+
+			if (CountData)
+			{
+				PointCount = CountData[0];
+				PointCount = FMath::Min(PointCount, 1024u);
+			}
+			RHICmdList.UnmapStagingSurface(StagingTexture);
+		}
+
+	}
+
+	// Primitives
+	FLineArray Output;
+	if (PointCount > 0)
+	{
+		auto ToColor = [](uint32 ColorIndex)
+		{
+			FLinearColor Color;
+			switch (ColorIndex)
+			{
+			case 0:		return FLinearColor::Red;
+			case 1:		return FLinearColor::Green;
+			case 2:		return FLinearColor::Blue;
+			case 3:		return FLinearColor::Yellow;
+			default:	return FLinearColor::White;
+			}
+		};
+
+		FTextureRHIRef SourceTexture = DebugPrimitiveBuffer.DebugPrimitiveTexture->GetRenderTargetItem().TargetableTexture;
+		FTextureRHIRef StagingTexture = DebugPrimitiveBuffer.DebugPrimitiveStagingTexture->GetRenderTargetItem().ShaderResourceTexture;
+
+		// Transfer memory GPU -> CPU
+		RHICmdList.CopyToResolveTarget(SourceTexture, StagingTexture, FResolveParams());
+		if (StagingTexture.IsValid())
+		{	
+			const FVector4* PrimitiveData = nullptr;
+			int32 BufferWidth = 0, BufferHeight = 0;
+			RHICmdList.MapStagingSurface(StagingTexture, *(void**)&PrimitiveData, BufferWidth, BufferHeight);
+			
+			if (PrimitiveData)
+			{
+				for (uint32 i = 0; i < PointCount; i += 2)
+				{
+					const FVector P0 = PrimitiveData[i+0];
+					const FVector P1 = PrimitiveData[i+1];
+					const FLinearColor Color = ToColor(uint32(PrimitiveData[i + 0].W));
+					Output.Add({ P0, P1, Color });
+				}
+			}
+			RHICmdList.UnmapStagingSurface(StagingTexture);
+		}
+
+	}
+	return Output;
+}
+
+void BindGpuDebugPrimitiveBuffers(FRHIRenderPassInfo& RPInfo, FGpuDebugPrimitiveBuffers& DebugPrimitiveBuffer, uint32 UAVIndex)
+{
+	// Currently this value is hardcoded into the shader. Once the index is configurable in the shader, this check should be remove
+	check(UAVIndex == 0); 
+
+	RPInfo.UAVIndex = UAVIndex;
+	RPInfo.NumUAVs = 2;
+	RPInfo.UAVs[UAVIndex+0] = DebugPrimitiveBuffer.DebugPrimitiveCountTexture->GetRenderTargetItem().UAV;
+	RPInfo.UAVs[UAVIndex+1] = DebugPrimitiveBuffer.DebugPrimitiveTexture->GetRenderTargetItem().UAV;
+}
+
+void DrawGpuDebugPrimitives(FRHICommandListImmediate& RHICmdList, TArray<FViewInfo>& Views, FGpuDebugPrimitiveBuffers& DebugPrimitiveBuffer)
+{
+	const FLineArray DebugLines = ReadGPUDebugPrimitives(RHICmdList, DebugPrimitiveBuffer);
+
+	uint32 ViewIndex = 0;
+	for (FViewInfo& View : Views)
+	{
+		FViewElementPDI PDI(&View, nullptr, nullptr);
+
+		for(int32 i=0; i < DebugLines.Num(); ++i)
+		{
+			const FLine& Line = DebugLines[i];
+			PDI.DrawLine(Line.Start, Line.End, Line.Color, 0);
+		}
+
+		ViewIndex++;
+	}
+}
diff --git a/Engine/Source/Runtime/Renderer/Private/GpuDebugRendering.h b/Engine/Source/Runtime/Renderer/Private/GpuDebugRendering.h
new file mode 100644
--- /dev/null
+++ b/Engine/Source/Runtime/Renderer/Private/GpuDebugRendering.h
@@ -0,0 +1,19 @@
+// Copyright 1998-2019 Epic Games, Inc. All Rights Reserved.
+
+#pragma once
+
+#include "RendererInterface.h"
+#include "RHICommandList.h"
+
+struct FGpuDebugPrimitiveBuffers
+{
+	TRefCountPtr<IPooledRenderTarget> DebugPrimitiveCountTexture;
+	TRefCountPtr<IPooledRenderTarget> DebugPrimitiveCountStagingTexture;
+
+	TRefCountPtr<IPooledRenderTarget> DebugPrimitiveTexture;
+	TRefCountPtr<IPooledRenderTarget> DebugPrimitiveStagingTexture;
+};
+
+FGpuDebugPrimitiveBuffers AllocateGpuDebugPrimitiveBuffers(FRHICommandListImmediate& RHICmdList);
+void BindGpuDebugPrimitiveBuffers(FRHIRenderPassInfo& RPInfo, FGpuDebugPrimitiveBuffers& DebugPrimitiveBuffer, uint32 UAVIndex = 0);
+void DrawGpuDebugPrimitives(FRHICommandListImmediate& Context, TArray<FViewInfo>& Views, FGpuDebugPrimitiveBuffers& DebugPrimitiveBuffer);
\ No newline at end of file
