Fix a compilation issue related to vertex factor API change and recent merge down.

diff --git a/Engine/Source/Runtime/Renderer/Private/RayTracing/RayTracingMaterialHitShaders.cpp b/Engine/Source/Runtime/Renderer/Private/RayTracing/RayTracingMaterialHitShaders.cpp
--- a/Engine/Source/Runtime/Renderer/Private/RayTracing/RayTracingMaterialHitShaders.cpp
+++ b/Engine/Source/Runtime/Renderer/Private/RayTracing/RayTracingMaterialHitShaders.cpp
@@ -206,7 +206,7 @@ void FRayTracingMeshProcessor::BuildRayTracingMeshCommands(
 	SharedCommand.bOpaque = MaterialResource.GetBlendMode() == EBlendMode::BLEND_Opaque;
 
 	FVertexInputStreamArray VertexStreams;
-	VertexFactory->GetStreams(ERHIFeatureLevel::SM5, VertexStreams);
+	VertexFactory->GetStreams(ERHIFeatureLevel::SM5, EVertexInputStreamType::Default, VertexStreams);
 
 	if (PassShaders.RayHitGroupShader)
 	{
@@ -226,7 +226,7 @@ void FRayTracingMeshProcessor::BuildRayTracingMeshCommands(
 			if (PassShaders.RayHitGroupShader)
 			{
 				FMeshDrawSingleShaderBindings RayHitGroupShaderBindings = RayTracingMeshCommand.ShaderBindings.GetSingleShaderBindings(SF_RayHitGroup);
-				PassShaders.RayHitGroupShader->GetElementShaderBindings(Scene, ViewIfDynamicMeshCommand, VertexFactory, false, FeatureLevel, PrimitiveSceneProxy, MeshBatch, BatchElement, ShaderElementData, RayHitGroupShaderBindings, VertexStreams);
+				PassShaders.RayHitGroupShader->GetElementShaderBindings(Scene, ViewIfDynamicMeshCommand, VertexFactory, EVertexInputStreamType::Default, FeatureLevel, PrimitiveSceneProxy, MeshBatch, BatchElement, ShaderElementData, RayHitGroupShaderBindings, VertexStreams);
 			}
 
 			int32 GeometrySegmentIndex = MeshBatch.SegmentIndex + BatchElementIndex;
