Fixed pre-exposure merge issue.

diff --git a/Engine/Source/Runtime/Renderer/Private/PostProcess/PostProcessEyeAdaptation.cpp b/Engine/Source/Runtime/Renderer/Private/PostProcess/PostProcessEyeAdaptation.cpp
--- a/Engine/Source/Runtime/Renderer/Private/PostProcess/PostProcessEyeAdaptation.cpp
+++ b/Engine/Source/Runtime/Renderer/Private/PostProcess/PostProcessEyeAdaptation.cpp
@@ -440,6 +440,22 @@ float FRCPassPostProcessEyeAdaptation::GetFixedExposure(const FViewInfo& View)
 void FSceneViewState::UpdatePreExposure(FViewInfo& View)
 {
 	const FSceneViewFamily& ViewFamily = *View.Family;
+	// One could use the IsRichView functionality to check if we need to update pre-exposure, 
+	// but this is too limiting for certain view. For instance shader preview doesn't have 
+	// volumetric lighting enabled, which makes the view be flagged as rich, and not updating 
+	// the pre-exposition value.
+	const bool bIsPreExposureRelevant =
+		ViewFamily.EngineShowFlags.Lighting &&
+		ViewFamily.EngineShowFlags.PostProcessing &&
+		ViewFamily.bResolveScene &&
+		!ViewFamily.EngineShowFlags.LightMapDensity &&
+		!ViewFamily.EngineShowFlags.StationaryLightOverlap &&
+		!ViewFamily.EngineShowFlags.LightComplexity &&
+		!ViewFamily.EngineShowFlags.LODColoration &&
+		!ViewFamily.EngineShowFlags.HLODColoration &&
+		!ViewFamily.EngineShowFlags.LevelColoration &&
+		!ViewFamily.EngineShowFlags.VisualizeBloom;
+
 
 	PreExposure = 1.f;
 	bUpdateLastExposure = false;
@@ -452,27 +468,9 @@ void FSceneViewState::UpdatePreExposure(FViewInfo& View)
 			PreExposure = FRCPassPostProcessEyeAdaptation::GetFixedExposure(View);
 		}
 	}
-	else if (!IsRichView(ViewFamily) /*&& !ViewFamily.EngineShowFlags.VisualizeHDR */&& !ViewFamily.EngineShowFlags.VisualizeBloom && ViewFamily.EngineShowFlags.PostProcessing && ViewFamily.bResolveScene)
-	{
-		const FSceneViewFamily& ViewFamily = *View.Family;
-
-		// One could use the IsRichView functionality to check if we need to update pre-exposure, 
-		// but this is too limiting for certain view. For instance shader preview doesn't have 
-		// volumetric lighting enabled, which makes the view be flagged as rich, and not updating 
-		// the pre-exposition value.
-		const bool bUpdatePreExposure =
-			ViewFamily.EngineShowFlags.Lighting &&
-			ViewFamily.EngineShowFlags.PostProcessing &&
-			ViewFamily.bResolveScene &&
-			!ViewFamily.EngineShowFlags.LightMapDensity &&
-			!ViewFamily.EngineShowFlags.StationaryLightOverlap &&
-			!ViewFamily.EngineShowFlags.LightComplexity &&
-			!ViewFamily.EngineShowFlags.LODColoration &&
-			!ViewFamily.EngineShowFlags.HLODColoration &&
-			!ViewFamily.EngineShowFlags.LevelColoration &&
-			!ViewFamily.EngineShowFlags.VisualizeBloom;
-
-		if (bUpdatePreExposure)
+	else if (bIsPreExposureRelevant)
+	{
+		if (UsePreExposure(View.GetShaderPlatform()))
 		{
 			const float PreExposureOverride = CVarEyeAdaptationPreExposureOverride.GetValueOnRenderThread();
 			const float LastExposure = View.GetLastEyeAdaptationExposure();
