Fixed some errors in the eye adaptation readback logic that made the engine read invalid pre-exposure entries.
Enabled pre-exposure by default on all platform that have "r.EyeAdaptationQuality 1"

diff --git a/Engine/Source/Runtime/Renderer/Private/PostProcess/PostProcessEyeAdaptation.cpp b/Engine/Source/Runtime/Renderer/Private/PostProcess/PostProcessEyeAdaptation.cpp
--- a/Engine/Source/Runtime/Renderer/Private/PostProcess/PostProcessEyeAdaptation.cpp
+++ b/Engine/Source/Runtime/Renderer/Private/PostProcess/PostProcessEyeAdaptation.cpp
@@ -445,6 +445,7 @@ void FSceneViewState::UpdatePreExposure(FViewInfo& View)
 	// volumetric lighting enabled, which makes the view be flagged as rich, and not updating 
 	// the pre-exposition value.
 	const bool bIsPreExposureRelevant =
+		ViewFamily.EngineShowFlags.EyeAdaptation && // Controls whether scene luminance is computed at all.
 		ViewFamily.EngineShowFlags.Lighting &&
 		ViewFamily.EngineShowFlags.PostProcessing &&
 		ViewFamily.bResolveScene &&
diff --git a/Engine/Source/Runtime/Renderer/Private/ScenePrivate.h b/Engine/Source/Runtime/Renderer/Private/ScenePrivate.h
--- a/Engine/Source/Runtime/Renderer/Private/ScenePrivate.h
+++ b/Engine/Source/Runtime/Renderer/Private/ScenePrivate.h
@@ -1166,7 +1166,7 @@ public:
 	/** Swaps the double-buffer targets used in eye adaptation */
 	void SwapEyeAdaptationRTs()
 	{
-		EyeAdaptationRTManager.SwapRTs(bUpdateLastExposure);
+		EyeAdaptationRTManager.SwapRTs(bUpdateLastExposure && bValidEyeAdaptation);
 	}
 
 	bool HasValidEyeAdaptation() const
diff --git a/Engine/Source/Runtime/Renderer/Private/SceneRendering.cpp b/Engine/Source/Runtime/Renderer/Private/SceneRendering.cpp
--- a/Engine/Source/Runtime/Renderer/Private/SceneRendering.cpp
+++ b/Engine/Source/Runtime/Renderer/Private/SceneRendering.cpp
@@ -163,8 +163,8 @@ static TAutoConsoleVariable<int32> CVarRoundRobinOcclusion(
 
 static TAutoConsoleVariable<int32> CVarUsePreExposure(
 	TEXT("r.UsePreExposure"),
-	0,
-	TEXT("0 to disable pre-exposure (default), 1 to enable."),
+	1,
+	TEXT("0 to disable pre-exposure, 1 to enable (default)."),
 	ECVF_ReadOnly | ECVF_RenderThreadSafe);
 
 static TAutoConsoleVariable<int32> CVarODSCapture(
