Subsurface scattering post process passes now use new RenderGraph API.

 - Cleaned up surrounding code so that subsurface post process code is referenced from one place.
 - Cleaned up shaders to be more consistent and to better support the new parameter struct system.
 - Refactored shaders to use new permutation API and shader parameter APIs.

diff --git a/Engine/Shaders/Private/Common.ush b/Engine/Shaders/Private/Common.ush
--- a/Engine/Shaders/Private/Common.ush
+++ b/Engine/Shaders/Private/Common.ush
@@ -1597,7 +1597,11 @@ struct FWriteToSliceGeometryOutput
 #define USE_WORLD_POSITION_EXCLUDING_SHADER_OFFSETS	(NEEDS_WORLD_POSITION_EXCLUDING_SHADER_OFFSETS && FEATURE_LEVEL >= FEATURE_LEVEL_ES3_1)
 
 /** Used for calculating vertex positions and UVs when drawing with DrawRectangle */
-void DrawRectangle( in float4 InPosition, in float2 InTexCoord, out float4 OutPosition, out float2 OutTexCoord)
+void DrawRectangle(
+	in float4 InPosition,
+	in float2 InTexCoord,
+	out float4 OutPosition,
+	out float2 OutTexCoord)
 {
 	OutPosition = InPosition;
 	OutPosition.xy = -1.0f + 2.0f * (DrawRectangleParameters.PosScaleBias.zw + (InPosition.xy * DrawRectangleParameters.PosScaleBias.xy)) * DrawRectangleParameters.InvTargetSizeAndTextureSize.xy;
@@ -1605,8 +1609,19 @@ void DrawRectangle( in float4 InPosition, in float2 InTexCoord, out float4 OutPo
 	OutTexCoord.xy = (DrawRectangleParameters.UVScaleBias.zw + (InTexCoord.xy * DrawRectangleParameters.UVScaleBias.xy)) * DrawRectangleParameters.InvTargetSizeAndTextureSize.zw;
 }
 
+/** Helper variant for vertex shaders which need a separate output for SV_POSITION and packed UV / Screen position */
+void DrawRectangle(
+	in float4 InPosition,
+	in float2 InTexCoord,
+	out float4 OutPosition,
+	out float4 OutUVAndScreenPos)
+{
+	DrawRectangle(InPosition, InTexCoord, OutPosition, OutUVAndScreenPos.xy);
+	OutUVAndScreenPos.zw = OutPosition.xy;
+}
+
 /** Used for calculating vertex positions when drawing with DrawRectangle */
-void DrawRectangle( in float4 InPosition, out float4 OutPosition)
+void DrawRectangle(in float4 InPosition, out float4 OutPosition)
 {
 	OutPosition = InPosition;
 	OutPosition.xy = -1.0f + 2.0f * (DrawRectangleParameters.PosScaleBias.zw + (InPosition.xy * DrawRectangleParameters.PosScaleBias.xy)) * DrawRectangleParameters.InvTargetSizeAndTextureSize.xy;
diff --git a/Engine/Shaders/Private/PostProcessBloom.usf b/Engine/Shaders/Private/PostProcessBloom.usf
--- a/Engine/Shaders/Private/PostProcessBloom.usf
+++ b/Engine/Shaders/Private/PostProcessBloom.usf
@@ -12,18 +12,6 @@
 // xyz:Bloom1Tint.rgb, w:unused
 float4 ColorScale1;
 
-// vertex shader entry point
-void MainPostprocessCommonVS(
-	in float4 InPosition : ATTRIBUTE0,
-	in float2 InTexCoord : ATTRIBUTE1,
-	out noperspective float4 OutTexCoord : TEXCOORD0,
-	out float4 OutPosition : SV_POSITION
-	)
-{
-	DrawRectangle(InPosition, InTexCoord, OutPosition, OutTexCoord.xy);
-	OutTexCoord.zw = OutPosition.xy;
-}
-
 void VisualizeBloomSetupPS(
 	noperspective float4 UVAndScreenPos : TEXCOORD0,
 	float4 SvPosition : SV_POSITION, 
@@ -132,12 +120,12 @@ void VisualizeBloomOverlayPS(
 void MainVS(
 	float4 InPosition : ATTRIBUTE0,
 	float2 InTexCoord : ATTRIBUTE1,
-	out noperspective float4 OutTexCoord : TEXCOORD0,
+	out noperspective float4 OutUVAndScreenPos : TEXCOORD0,
 	out nointerpolation float OutExposureScale : TEXCOORD1,
 	out float4 OutPosition : SV_POSITION
 	)
 {
-	MainPostprocessCommonVS(InPosition, InTexCoord, OutTexCoord, OutPosition);
+	DrawRectangle(InPosition, InTexCoord, OutPosition, OutUVAndScreenPos);
 
 #if FEATURE_LEVEL >= FEATURE_LEVEL_SM5
 	// texture can be GWhiteTexture which is 1x1. It's important we don't read outside bounds.
diff --git a/Engine/Shaders/Private/PostProcessCommon.ush b/Engine/Shaders/Private/PostProcessCommon.ush
--- a/Engine/Shaders/Private/PostProcessCommon.ush
+++ b/Engine/Shaders/Private/PostProcessCommon.ush
@@ -72,7 +72,6 @@ float4 ViewportSize;
 uint4 ViewportRect;
 
 SamplerState BilinearTextureSampler0;
-SamplerState BilinearTextureSampler1;
 
 // transforms ScreenPos to the screen pixel position, not the viewport local position (*xy +wz)
 // (0,0) left top, (width-1,height-1) at right bottom, in pixels (without offset to reach the pixel center)
diff --git a/Engine/Shaders/Private/PostProcessSubsurface.usf b/Engine/Shaders/Private/PostProcessSubsurface.usf
--- a/Engine/Shaders/Private/PostProcessSubsurface.usf
+++ b/Engine/Shaders/Private/PostProcessSubsurface.usf
@@ -13,46 +13,57 @@
 #include "MiniFontCommon.ush"
 
 // x:Radius*DistanceToProjectionWindow/KernelSize*0.5, y:DistanceToProjectionWindow, zw: unused
-float4 SSSParams;
+float4 SubsurfaceParams;
 
-// 0: (testing, could be faster), 1: higher quality
-#define FULLRES_WHERE_POSSIBLE 1
+Texture2D SubsurfaceInput0_Texture;
+Texture2D SubsurfaceInput1_Texture;
 
-#ifndef RECOMBINE_QUALITY
-	#define RECOMBINE_QUALITY 0
-#endif
+SamplerState SubsurfaceInput0_Sampler;
+SamplerState SubsurfaceInput1_Sampler;
 
-// ------------------------------------------
+float4 SubsurfaceInput0_Size;
+float4 SubsurfaceInput1_Size;
 
-// setup for "SeparableSSS.ush"
+#define SUBSURFACE_DIRECTION_HORIZONTAL 0
+#define SUBSURFACE_DIRECTION_VERTICAL 1
 
-float ComputeMaskFromDepthInAlpha(float Alpha)
-{
-	return Alpha > 0;
-}
+// Controls the quality (number of samples) of the blur kernel.
+#define SUBSURFACE_QUALITY_LOW 0
+#define SUBSURFACE_QUALITY_MEDIUM 1
+#define SUBSURFACE_QUALITY_HIGH 2
 
-// can be optimized
-float GetSubsurfaceStrength(float2 UV)
-{
-	FScreenSpaceData ScreenSpaceData = GetScreenSpaceData(UV);
+// Full resolution recombine.
+#define SUBSURFACE_RECOMBINE_MODE_FULLRES 0
 
-	float Mask = UseSubsurfaceProfile(ScreenSpaceData.GBuffer.ShadingModelID) ? 1.0f : 0.0f;
-	return Mask * ScreenSpaceData.GBuffer.CustomData.a;
-}
+// Half resolution recombine.
+#define SUBSURFACE_RECOMBINE_MODE_HALFRES 1
+
+// Just reconstruct the lighting (needed for scalability).
+#define SUBSURFACE_RECOMBINE_MODE_NO_SCATTERING 2
+
+// Controls the quality of lighting reconstruction.
+#define SUBSURFACE_RECOMBINE_QUALITY_LOW 0
+#define SUBSURFACE_RECOMBINE_QUALITY_HIGH 1
+
+#ifndef SUBSURFACE_RECOMBINE_QUALITY
+	#define SUBSURFACE_RECOMBINE_QUALITY SUBSURFACE_RECOMBINE_QUALITY_LOW
+#endif
 
-// 0:full res (slower but reference) / 1:faster half res with some quality loss
-// #define HALF_RES 1
+#if SUBSURFACE_RECOMBINE_QUALITY == SUBSURFACE_RECOMBINE_QUALITY_HALFRES
+	#define SUBSURFACE_HALFRES 1
+#endif
 
-// set by C++ for SubsurfacePS
-// #define SSS_DIRECTION 0/1
+//=============================================================================
+// setup for "SeparableSSS.ush"
+//=============================================================================
 
-#if SSS_SAMPLESET == 0
+#if SUBSURFACE_QUALITY == SUBSURFACE_QUALITY_LOW
 	#define	SSSS_N_KERNELWEIGHTCOUNT SSSS_KERNEL2_SIZE
 	#define	SSSS_N_KERNELWEIGHTOFFSET SSSS_KERNEL2_OFFSET
-#elif SSS_SAMPLESET == 1
+#elif SUBSURFACE_QUALITY == SUBSURFACE_QUALITY_MEDIUM
 	#define	SSSS_N_KERNELWEIGHTCOUNT SSSS_KERNEL1_SIZE
 	#define	SSSS_N_KERNELWEIGHTOFFSET SSSS_KERNEL1_OFFSET
-#else // SSS_SAMPLESET == 2
+#else // SUBSURFACE_QUALITY == SUBSURFACE_QUALITY_HIGH
 	#define	SSSS_N_KERNELWEIGHTCOUNT SSSS_KERNEL0_SIZE
 	#define	SSSS_N_KERNELWEIGHTOFFSET SSSS_KERNEL0_OFFSET
 #endif
@@ -61,72 +72,67 @@ float GetSubsurfaceStrength(float2 UV)
 // 1: no color bleeding in z direction
 #define SSSS_FOLLOW_SURFACE 1
 
-// @param In needs to be in 0..1 range 
-// @return 0..1
-float4 Quantize8Bit(float4 In)
+float GetMaskFromDepthInAlpha(float Alpha)
 {
-	return floor(In * 255.999f) / 255.0f;
+	return Alpha > 0;
 }
 
-// @return .rgb is the weight for color channel, .a is the sample location
-float4 GetKernel(uint SampleIndex, uint SubsurfaceProfileInt)
+// can be optimized
+float GetSubsurfaceStrength(float2 UV)
 {
-	const float4 TableMax = float4(1, 1, 1, SUBSURFACE_KERNEL_SIZE);
-
-	// profiled on NV670 fullscreen, one pass (total: 2 pass), samples: 13+1+13, large object filling the screen
+	FScreenSpaceData ScreenSpaceData = GetScreenSpaceData(UV);
 
-	// texture lookup
-	float4 Value = ActualSSProfilesTexture.Load(int3(SampleIndex, SubsurfaceProfileInt, 0)) * TableMax;		// 0.88ms for 8bit
-	// simulated  8 bit lookup
-//	float4 Value = Quantize8Bit(kernel[SampleIndex] / TableMax) * TableMax;			// 1.33ms same look as texture
-	// float reference
-//	float4 Value = kernel[SampleIndex];												// 0.85ms, quite a bit different in look, need to use more than 8 bit?
+	float Mask = UseSubsurfaceProfile(ScreenSpaceData.GBuffer.ShadingModelID) ? 1.0f : 0.0f;
+	return Mask * ScreenSpaceData.GBuffer.CustomData.a;
+}
 
-	// debug if the kernal was a box filter
-//	Value.rgb = 1.0f / (SSSS_N_KERNELWEIGHTCOUNT * 2 - 1);
+// @return .rgb is the weight for color channel, .a is the sample location
+float4 GetSubsurfaceProfileKernel(uint SampleIndex, uint SubsurfaceProfileInt)
+{
+	const float4 TableMax = float4(1, 1, 1, SUBSURFACE_KERNEL_SIZE);
 
-	return Value;
+	return ActualSSProfilesTexture.Load(int3(SampleIndex, SubsurfaceProfileInt, 0)) * TableMax;
 }
 
-float GetProfileRadiusScale(FGBufferData GBufferData)
+float GetSubsurfaceProfileRadiusScale(FGBufferData GBufferData)
 {
 	// 0..255, which SubSurface profile to pick
 	uint SubsurfaceProfileInt = ExtractSubsurfaceProfileInt(GBufferData);
 
-	return GetKernel(SSSS_N_KERNELWEIGHTOFFSET + SSSS_N_KERNELWEIGHTCOUNT - 1, SubsurfaceProfileInt).a;
+	return GetSubsurfaceProfileKernel(SSSS_N_KERNELWEIGHTOFFSET + SSSS_N_KERNELWEIGHTCOUNT - 1, SubsurfaceProfileInt).a;
 }
 
-float3 GetProfileSubsurfaceColor(FGBufferData GBufferData)
+float3 GetSubsurfaceProfileColor(FGBufferData GBufferData)
 {
 	// 0..255, which SubSurface profile to pick
 	uint SubsurfaceProfileInt = ExtractSubsurfaceProfileInt(GBufferData);
 
-	return GetKernel(SSSS_SUBSURFACE_COLOR_OFFSET, SubsurfaceProfileInt).rgb;
+	return GetSubsurfaceProfileKernel(SSSS_SUBSURFACE_COLOR_OFFSET, SubsurfaceProfileInt).rgb;
 }
 
-float4 SampleSetup(float2 BufferUV)
+uint GetSubsurfaceProfileId(float2 BufferUV)
 {
-	#if MANUALLY_CLAMP_UV
-		BufferUV = clamp(BufferUV, View.BufferBilinearUVMinMax.xy, View.BufferBilinearUVMinMax.zw);
-	#endif
-	return Texture2DSampleLevel(PostprocessInput0, PostprocessInput0Sampler, BufferUV, 0);
-}
-
-uint SampleSetupProfileId(float2 BufferUV)
-{
-#if MANUALLY_CLAMP_UV
+#if SUBSURFACE_MANUALLY_CLAMP_UV
 	BufferUV = clamp(BufferUV, View.BufferBilinearUVMinMax.xy, View.BufferBilinearUVMinMax.zw);
 #endif
 	FGBufferData LocalGBufferData = GetGBufferData(BufferUV);
 	return ExtractSubsurfaceProfileInt(LocalGBufferData);
 }
 
-float3 GetProfileBoundaryColorBleed(FGBufferData GBufferData)
+float3 GetSubsurfaceProfileBoundaryColorBleed(FGBufferData GBufferData)
 {
 	// 0..255, which SubSurface profile to pick
 	uint SubsurfaceProfileInt = ExtractSubsurfaceProfileInt(GBufferData);
 
-	return GetKernel(SSSS_BOUNDARY_COLOR_BLEED_OFFSET, SubsurfaceProfileInt).rgb;
+	return GetSubsurfaceProfileKernel(SSSS_BOUNDARY_COLOR_BLEED_OFFSET, SubsurfaceProfileInt).rgb;
+}
+
+float4 GetSceneColor(float2 BufferUV)
+{
+#if SUBSURFACE_MANUALLY_CLAMP_UV
+	BufferUV = clamp(BufferUV, View.BufferBilinearUVMinMax.xy, View.BufferBilinearUVMinMax.zw);
+#endif
+	return Texture2DSampleLevel(SubsurfaceInput0_Texture, SubsurfaceInput0_Sampler, BufferUV, 0);
 }
 
 // from [REDACTED_PROFILE_URL]/separable-sss/tree/master/Demo
@@ -134,18 +140,17 @@ float3 GetProfileBoundaryColorBleed(FGBufferData GBufferData)
 // [REDACTED_PROFILE_URL]
 #include "SeparableSSS.ush"
 
-// ------------------------------------------
+//=============================================================================
 
 bool InUnitBox(float2 UV)
 {
 	return UV.x >= 0 && UV.y >= 0 && UV.y < 1 && UV.y < 1;
 }
 
-
 // @return 0=don't blend in, 1:fully blend in
 float ComputeFullResLerp(FScreenSpaceData ScreenSpaceData, float2 UVSceneColor, float4 FullResInputSize)
 {
-	float SSSScaleX = SSSParams.x;
+	float SSSScaleX = SubsurfaceParams.x;
 
 	float scale = SSSScaleX / CalcSceneDepth(UVSceneColor);
 
@@ -154,21 +159,19 @@ float ComputeFullResLerp(FScreenSpaceData ScreenSpaceData, float2 UVSceneColor,
 	// Calculate the final step to fetch the surrounding pixels:
 	float finalStep = scale * HorizontalScaler;
 
-	finalStep *= GetProfileRadiusScale(ScreenSpaceData.GBuffer);
+	finalStep *= GetSubsurfaceProfileRadiusScale(ScreenSpaceData.GBuffer);
 
 	float PixelSizeRadius = finalStep / (FullResInputSize.z * 0.5f);
 
 	// tweaked for skin, a more flat kernel might need a smaller value, around 2 seems reasonable because we do half res
 	const float PixelSize = 4.0f;
 
-	float Ret = 1.0f;
+	float Ret = saturate(PixelSizeRadius - PixelSize);
 
-	//
-	Ret *= saturate(PixelSizeRadius - PixelSize);
 	// opacity allows to scale the radius - at some point we should fade in the full resolution, we don't have a masking other than that.
 	Ret *= saturate(ScreenSpaceData.GBuffer.CustomData.a * 10);
-	// todo: Subsurface has some non scatter contribution - all that should come from the Full res
 
+	// todo: Subsurface has some non scatter contribution - all that should come from the Full res
 	return Ret;
 }
 
@@ -177,7 +180,7 @@ void VisualizePS(in noperspective float4 UVAndScreenPos : TEXCOORD0, float4 SvPo
 {
 	float2 UV = UVAndScreenPos.xy;
 
-	OutColor = Texture2DSample(PostprocessInput0, PostprocessInput0Sampler, UV);
+	OutColor = Texture2DSample(SubsurfaceInput0_Texture, SubsurfaceInput0_Sampler, UV);
 
 	int2 PixelPos = (int2)SvPosition.xy;
 
@@ -238,7 +241,7 @@ void VisualizePS(in noperspective float4 UVAndScreenPos : TEXCOORD0, float4 SvPo
 			int2 LeftTop = (PixelPos / 8) * 8;
 			PrintCharacter(PixelPos, OutColor.rgb, float3(1, 1, 1), LeftTop, SubsurfaceProfileInt);
 
-			OutColor.rgb *= ComputeFullResLerp(ScreenSpaceData, UV, PostprocessInput0Size);
+			OutColor.rgb *= ComputeFullResLerp(ScreenSpaceData, UV, SubsurfaceInput0_Size);
 		}
 	}
 
@@ -257,13 +260,13 @@ half3 LookupSceneColor(float2 SceneUV, int2 PixelOffset)
 #if ES2_PROFILE && COMPILER_GLSL_ES2
 	// slower but always works
 	// to prevent "error: Texture offset not supported on GLSL ES"
-	return Texture2DSample(PostprocessInput0, PostprocessInput0Sampler, SceneUV + PixelOffset * PostprocessInput0Size.zw).rgb;
+	return Texture2DSample(SubsurfaceInput0_Texture, SubsurfaceInput0_Sampler, SceneUV + PixelOffset * SubsurfaceInput0_Size.zw).rgb;
 #else
 	// faster
-	return PostprocessInput0.SampleLevel(PostprocessInput0Sampler, SceneUV, 0, PixelOffset).rgb;
+	return SubsurfaceInput0_Texture.SampleLevel(SubsurfaceInput0_Sampler, SceneUV, 0, PixelOffset).rgb;
 #endif
 }
-
+ 
 // @param UVSceneColor for the full res rendertarget (BufferSize) e.g. SceneColor or GBuffers
 // @param ReconstructMethod 0/1/2/3 (should be a literal constant to allow compiler optimizations)
 SDiffuseAndSpecular ReconstructLighting(float2 UVSceneColor, uint ReconstructMethod)
@@ -277,7 +280,7 @@ SDiffuseAndSpecular ReconstructLighting(float2 UVSceneColor, uint ReconstructMet
 
 	  // todo: We could alternate the diagonal with TemporalAA or even only only 1 sample for low spec or 4 for high spec
 
-	  float3 Quant0 = PostprocessInput0.SampleLevel(PostprocessInput0Sampler, UVSceneColor, 0).rgb;
+	  float3 Quant0 = SubsurfaceInput0_Texture.SampleLevel(SubsurfaceInput0_Sampler, UVSceneColor, 0).rgb;
 
 	  // todo: expose as scalability setting (can be evaluate best without TemporalAA)
 	  // 0:fast but pattern can appear, 1:better, 2: even better, 3: best but expensive
@@ -331,7 +334,7 @@ SDiffuseAndSpecular ReconstructLighting(float2 UVSceneColor, uint ReconstructMet
 #else // SUBSURFACE_PROFILE_CHECKERBOARD
 	{
 	  // If we're not doing checkerboard encoding, we just need to read a single pixel and decode (combined diffuse/spec in RGB)
-	  float4 CenterSample = PostprocessInput0.SampleLevel(PostprocessInput0Sampler, UVSceneColor, 0);
+	  float4 CenterSample = SubsurfaceInput0_Texture.SampleLevel(SubsurfaceInput0_Sampler, UVSceneColor, 0);
 	  float3 CombinedColor = CenterSample.rgb;
 	  float DiffuseLuminance = CenterSample.a;
   
@@ -374,67 +377,42 @@ void SetupPS(in noperspective float4 UVAndScreenPos : TEXCOORD0, out float4 OutC
 {
 	float2 BufferUV = UVAndScreenPos.xy;
 
-#if HALF_RES
+#if SUBSURFACE_HALF_RES
 	// order aligned with Gather() hardware implementation
 	// RGB: color*A, A:weight 0 if no subsurface scattering
 
-	float4 A = SetupSubsurfaceForOnePixel(min(BufferUV + float2(-0.5,  0.5f) * PostprocessInput0Size.zw, View.BufferBilinearUVMinMax.zw));
-	float4 B = SetupSubsurfaceForOnePixel(min(BufferUV + float2( 0.5,  0.5f) * PostprocessInput0Size.zw, View.BufferBilinearUVMinMax.zw));
-	float4 C = SetupSubsurfaceForOnePixel(min(BufferUV + float2( 0.5, -0.5f) * PostprocessInput0Size.zw, View.BufferBilinearUVMinMax.zw));
-	float4 D = SetupSubsurfaceForOnePixel(min(BufferUV + float2(-0.5, -0.5f) * PostprocessInput0Size.zw, View.BufferBilinearUVMinMax.zw));
+	float4 A = SetupSubsurfaceForOnePixel(min(BufferUV + float2(-0.5,  0.5f) * SubsurfaceInput0_Size.zw, View.BufferBilinearUVMinMax.zw));
+	float4 B = SetupSubsurfaceForOnePixel(min(BufferUV + float2( 0.5,  0.5f) * SubsurfaceInput0_Size.zw, View.BufferBilinearUVMinMax.zw));
+	float4 C = SetupSubsurfaceForOnePixel(min(BufferUV + float2( 0.5, -0.5f) * SubsurfaceInput0_Size.zw, View.BufferBilinearUVMinMax.zw));
+	float4 D = SetupSubsurfaceForOnePixel(min(BufferUV + float2(-0.5, -0.5f) * SubsurfaceInput0_Size.zw, View.BufferBilinearUVMinMax.zw));
 
 	float4 Sum = (A + B) + (C + D);
 
 	float Div = 1.0f / max(Sum.a, 0.00001f);
 
 	OutColor.rgb = Sum.rgb * Div;
-	
-	float4 FourDepth = GatherSceneDepth(BufferUV, PostprocessInput0Size.zw);
+
+	float4 FourDepth = GatherSceneDepth(BufferUV, SubsurfaceInput0_Size.zw);
 
 	// average all valid depth values to a single one
 	float SingleDepth = dot(FourDepth, float4(A.a, B.a, C.a, D.a)) * Div;
 
 	OutColor.a = SingleDepth;
-#else // HALF_RES
+#else // SUBSURFACE_HALF_RES
 	OutColor = SetupSubsurfaceForOnePixel(BufferUV);
-	if( OutColor.a > 0){
-		OutColor.a = CalcSceneDepth(BufferUV);
-	}
-#endif // HALF_RES
-
-#if 0
-	float SSSScaleX = SSSParams.x;
-
-	float scale = SSSScaleX / SingleDepth;
-//	float scale = SSSScaleX / min(min(FourDepth.x, FourDepth.y), min(FourDepth.z, FourDepth.w));
-
-	float HorizontalScaler = SUBSURFACE_RADIUS_SCALE;
-
-	// Calculate the final step to fetch the surrounding pixels:
-	float finalStep = scale * HorizontalScaler;
-	
-	FScreenSpaceData ScreenSpaceData = GetScreenSpaceData(BufferUV);
-
-	finalStep *= GetProfileRadiusScale(ScreenSpaceData.GBuffer);
-
-	// in full resolution
-	float PixelFracationThreshold = 3.0f;
-
-	// *0.5f as we read in half res in the blur pass
-	if(finalStep < PixelFracationThreshold * PostprocessInput0Size.z * 0.5f)
+	if(OutColor.a > 0)
 	{
-		// very small radius doesn't require work
-		OutColor = 0;
+		OutColor.a = CalcSceneDepth(BufferUV);
 	}
-#endif
+#endif // SUBSURFACE_HALF_RES
 }
 
 // input0 is created by the SetupPS shader
 void MainPS(noperspective float4 UVAndScreenPos : TEXCOORD0, out float4 OutColor : SV_Target0)
 {
 	float2 BufferUV = UVAndScreenPos.xy;
-	
-#if SSS_DIRECTION == 0
+
+#if SUBSURFACE_DIRECTION == 0
 	// horizontal
 	float2 ViewportDirectionUV = float2(1, 0) * SUBSURFACE_RADIUS_SCALE;
 #else
@@ -442,20 +420,18 @@ void MainPS(noperspective float4 UVAndScreenPos : TEXCOORD0, out float4 OutColor
 	float2 ViewportDirectionUV = float2(0, 1) * SUBSURFACE_RADIUS_SCALE * (View.ViewSizeAndInvSize.x * View.ViewSizeAndInvSize.w);
 #endif
 
-	#if MANUALLY_CLAMP_UV
-		ViewportDirectionUV *= (View.ViewSizeAndInvSize.x * View.BufferSizeAndInvSize.z);
-	#endif
+#if SUBSURFACE_MANUALLY_CLAMP_UV
+	ViewportDirectionUV *= (View.ViewSizeAndInvSize.x * View.BufferSizeAndInvSize.z);
+#endif
 
 	OutColor = SSSSBlurPS(BufferUV, ViewportDirectionUV, false);
 
-#if SSS_DIRECTION == 1
+#if SUBSURFACE_DIRECTION == SUBSURFACE_DIRECTION_VERTICAL
 	// second pass prepares the setup from the recombine pass which doesn't need depth but wants to reconstruct the color
-	OutColor.a = ComputeMaskFromDepthInAlpha(OutColor.a);
+	OutColor.a = GetMaskFromDepthInAlpha(OutColor.a);
 #endif
 }
 
-
-
 // Recombines the half res Subsurface filtered lighting contribution (upsampled and renormalized with the alpha)
 // with the SceneColor.
 void SubsurfaceRecombinePS(noperspective float4 UVAndScreenPos : TEXCOORD0, float4 SvPosition : SV_POSITION, out float4 OutColor : SV_Target0)
@@ -466,57 +442,52 @@ void SubsurfaceRecombinePS(noperspective float4 UVAndScreenPos : TEXCOORD0, floa
 	
 	if (!UseSubsurfaceProfile(ScreenSpaceData.GBuffer.ShadingModelID))
 	{
-		OutColor = Texture2DSample(PostprocessInput0, PostprocessInput0Sampler, BufferUV);
+		OutColor = Texture2DSample(SubsurfaceInput0_Texture, SubsurfaceInput0_Sampler, BufferUV);
 		return;
 	}
 
-	// bilinear filtering
-	float3 SSSColor;
-	{
-		float4 SSSColorWithAlpha = Texture2DSample(PostprocessInput1, PostprocessInput1Sampler, BufferUV);
-
-		// renormalize to dilate RGB to fix half res upsampling artifacts
-		SSSColor = SSSColorWithAlpha.rgb / max(SSSColorWithAlpha.a, 0.00001f);
-	}
-
+	float3 SSSColor = float3(0, 0, 0);
 	float LerpFactor = 1;
 
-#if HALF_RES
-#if 1
+#if SUBSURFACE_RECOMBINE_MODE != SUBSURFACE_RECOMBINE_MODE_NO_SCATTERING
+#if SUBSURFACE_HALF_RES
 	// fade out subsurface scattering if radius is too small to be more crips (not blend with half resolution)
 	// minor quality improvement (faces are more detailed in distance)
-	LerpFactor = ComputeFullResLerp(ScreenSpaceData, BufferUV, PostprocessInput1Size);
+	LerpFactor = ComputeFullResLerp(ScreenSpaceData, BufferUV, SubsurfaceInput1_Size);
+#endif // SUBSURFACE_HALF_RES
 
-	// hack to debug if the fade is happening at the same time the VisualizeSSS shows it
-//	SSSColor = 0;
-#endif
-#endif // HALF_RES
+	{
+		float4 SSSColorWithAlpha = Texture2DSample(SubsurfaceInput1_Texture, SubsurfaceInput1_Sampler, BufferUV);
 
-#if !RECOMBINE_SUBSURFACESCATTER
+		// renormalize to dilate RGB to fix half res upsampling artifacts
+		SSSColor = SSSColorWithAlpha.rgb / max(SSSColorWithAlpha.a, 0.00001f);
+	}
+#else // SUBSURFACE_RECOMBINE_MODE == SUBSURFACE_RECOMBINE_MODE_NO_SCATTERING
 	// Scalability requests no Scatter, but we still need to reconstruct a color
 	LerpFactor = 0;
-#endif
+#endif // SUBSURFACE_RECOMBINE_MODE
 
 	// we multiply the base color later in to get more crips human skin textures (scanned data always has Subsurface included)
 	float3 StoredBaseColor = ScreenSpaceData.GBuffer.StoredBaseColor;
 	float StoredSpecular = ScreenSpaceData.GBuffer.StoredSpecular;
 
-	uint ReconstructMethod = RECOMBINE_QUALITY ? 3 : 1;
+	uint ReconstructMethod = (SUBSURFACE_RECOMBINE_QUALITY == SUBSURFACE_RECOMBINE_QUALITY_HIGH) ? 3 : 1;
 
 	SDiffuseAndSpecular DiffuseAndSpecular = ReconstructLighting(BufferUV, ReconstructMethod);
 	
 	float3 ExtractedNonSubsurface = DiffuseAndSpecular.Specular;
 
 	// asset specific color
-	float3 SubsurfaceColor = GetProfileSubsurfaceColor(ScreenSpaceData.GBuffer);
-	float3 FadedSubsurfaceColor = GetProfileSubsurfaceColor(ScreenSpaceData.GBuffer) * LerpFactor;
-
-	// hack to debug if the fade is happening at the same time the VisualizeSSS shows it
-//	SSSColor = float3(0,1,0);
+	float3 SubsurfaceColor = GetSubsurfaceProfileColor(ScreenSpaceData.GBuffer);
+	float3 FadedSubsurfaceColor = SubsurfaceColor * LerpFactor;
 
 	// combine potentially half res with full res
 	float3 SubsurfaceLighting = lerp(DiffuseAndSpecular.Diffuse, SSSColor, FadedSubsurfaceColor);
 
-
 	OutColor = float4(SubsurfaceLighting * StoredBaseColor + ExtractedNonSubsurface, 0); 
 }
+
+void SubsurfaceViewportCopyPS(noperspective float4 InUV : TEXCOORD0, out float4 OutColor : SV_Target0)
+{
+	OutColor = Texture2DSample(SubsurfaceInput0_Texture, SubsurfaceInput0_Sampler, InUV.xy);
+}
diff --git a/Engine/Shaders/Private/ScreenPass.usf b/Engine/Shaders/Private/ScreenPass.usf
new file mode 100644
--- /dev/null
+++ b/Engine/Shaders/Private/ScreenPass.usf
@@ -0,0 +1,12 @@
+// Copyright 1998-2019 Epic Games, Inc. All Rights Reserved.
+
+#include "Common.ush"
+
+void ScreenPassVS(
+	in float4 InPosition : ATTRIBUTE0,
+	in float2 InTexCoord : ATTRIBUTE1,
+	out noperspective float4 OutUVAndScreenPos : TEXCOORD0,
+	out float4 OutPosition : SV_POSITION)
+{
+	DrawRectangle(InPosition, InTexCoord, OutPosition, OutUVAndScreenPos);
+}
\ No newline at end of file
diff --git a/Engine/Shaders/Private/SeparableSSS.ush b/Engine/Shaders/Private/SeparableSSS.ush
--- a/Engine/Shaders/Private/SeparableSSS.ush
+++ b/Engine/Shaders/Private/SeparableSSS.ush
@@ -111,15 +111,10 @@
 //-----------------------------------------------------------------------------
 // Porting Functions
 #define SSSSTexture2D Texture2D
-#define SSSSSampleLevelZero(tex, coord) tex.SampleLevel(LinearSampler, coord, 0)
-#define SSSSSampleLevelZeroPoint(tex, coord) tex.SampleLevel(PointSampler, coord, 0)
 #define SSSSSampleShadowmap(coord)	float4(0,0,0,0)	// todo: not used yet
-#define SSSSSampleMasking(texcoord)	(HasSubsurfaceProfile(GetScreenSpaceData(texcoord).GBuffer.ShadingModelID))	// todo: masking doesn't work with filtering ideally put in alpha to get it filtered
-#define SSSSSampleSceneColor(coord) SampleSetup(coord)
-#define SSSSSampleSceneColorPoint(coord) SampleSetup(coord) // todo: could be made point
-#define SSSSSampleProfileId(coord) SampleSetupProfileId(coord)
-#define SSSSSampleLevelZeroOffset(tex, coord, offset) tex.SampleLevel(LinearSampler, coord, 0, offset)
-#define SSSSSampleOffset(tex, coord, offset) SSSSSampleLevelZeroOffset(tex, coord, offset)
+#define SSSSSampleSceneColor(coord) GetSceneColor(coord)
+#define SSSSSampleSceneColorPoint(coord) GetSceneColor(coord) // todo: could be made point
+#define SSSSSampleProfileId(coord) GetSubsurfaceProfileId(coord)
 #define SSSSLerp(a, b, t) lerp(a, b, t)
 #define SSSSMad(a, b, c) mad(a, b, c)
 #define SSSSMul(v, m) mul(v, m)
@@ -195,7 +190,7 @@ float4 SSSSBlurPS(float2 BufferUV, float2 dir, bool initStencil)
 	// we store the depth in alpha
 	float OutDepth = colorM.a;
 
-	colorM.a = ComputeMaskFromDepthInAlpha(colorM.a);
+	colorM.a = GetMaskFromDepthInAlpha(colorM.a);
 
 	// we don't need to process pixels that are not part of the subsurafce scattering (optimization, also prevents divide by later on)
 	BRANCH if(!colorM.a)
@@ -212,7 +207,7 @@ float4 SSSSBlurPS(float2 BufferUV, float2 dir, bool initStencil)
     if (initStencil) // (Checked in compile time, it's optimized away)
         if (SSSStrength < 1 / 256.0f) discard;
 
-	float SSSScaleX = SSSParams.x;
+	float SSSScaleX = SubsurfaceParams.x;
 
 	float scale = SSSScaleX / OutDepth;
 
@@ -234,23 +229,23 @@ float4 SSSSBlurPS(float2 BufferUV, float2 dir, bool initStencil)
 	float3 colorInvDiv = 0.00001f;
 
 	// center sample
-	colorInvDiv += GetKernel(SSSS_N_KERNELWEIGHTOFFSET, SubsurfaceProfileInt).rgb;
-	colorAccum = colorM.rgb * GetKernel(SSSS_N_KERNELWEIGHTOFFSET, SubsurfaceProfileInt).rgb;
+	colorInvDiv += GetSubsurfaceProfileKernel(SSSS_N_KERNELWEIGHTOFFSET, SubsurfaceProfileInt).rgb;
+	colorAccum = colorM.rgb * GetSubsurfaceProfileKernel(SSSS_N_KERNELWEIGHTOFFSET, SubsurfaceProfileInt).rgb;
 
-	float3 BoundaryColorBleed = GetProfileBoundaryColorBleed(GBufferData);
+	float3 BoundaryColorBleed = GetSubsurfaceProfileBoundaryColorBleed(GBufferData);
 
 	// Accumulate the other samples:
     SSSS_UNROLL
 	for (int i = 1; i < SSSS_N_KERNELWEIGHTCOUNT; i++) 
 	{
 		// Kernel.a = 0..SUBSURFACE_KERNEL_SIZE (radius)
-		half4 Kernel = GetKernel(SSSS_N_KERNELWEIGHTOFFSET + i, SubsurfaceProfileInt);
+		half4 Kernel = GetSubsurfaceProfileKernel(SSSS_N_KERNELWEIGHTOFFSET + i, SubsurfaceProfileInt);
 
 		float4 LocalAccum = 0;
 
 		float2 UVOffset = Kernel.a * finalStep;
 		// The kernel is symtrical, we want to use that property.
-		// Half the GetKernel() calls (more expensive if done by texture)
+		// Half the GetSubsurfaceProfileKernel() calls (more expensive if done by texture)
 		// Half the weighting computations (saves 3mul per lookup sample)
 		SSSS_UNROLL
 		for (int Side = -1; Side <= 1; Side += 2)
@@ -263,18 +258,17 @@ float4 SSSSBlurPS(float2 BufferUV, float2 dir, bool initStencil)
 
 
 			float LocalDepth = color.a;
-			
-			color.a = ComputeMaskFromDepthInAlpha(color.a);
+			color.a = GetMaskFromDepthInAlpha(color.a);
 
 #if SSSS_FOLLOW_SURFACE == 1
 			// If the difference in depth is huge, we weight the sample less or not at all
-			float s = saturate(12000.0f / 400000 * SSSParams.y *
-	//        float s = saturate(300.0f/400000 * SSSParams.y *
+			float s = saturate(12000.0f / 400000 * SubsurfaceParams.y *
+	//        float s = saturate(300.0f/400000 * SubsurfaceParams.y *
 				abs(OutDepth - LocalDepth));
 
 			color.a *= 1 - s;
 #endif
-			// approximation, ideally we would reconstruct the mask with ComputeMaskFromDepthInAlpha() and do manual bilinear filter
+			// approximation, ideally we would reconstruct the mask with GetMaskFromDepthInAlpha() and do manual bilinear filter
 			// needed?
 			color.rgb *= color.a * ColorTint;
 
diff --git a/Engine/Source/Runtime/Engine/Private/UserInterface/Canvas.cpp b/Engine/Source/Runtime/Engine/Private/UserInterface/Canvas.cpp
--- a/Engine/Source/Runtime/Engine/Private/UserInterface/Canvas.cpp
+++ b/Engine/Source/Runtime/Engine/Private/UserInterface/Canvas.cpp
@@ -638,11 +638,8 @@ FCanvas::~FCanvas()
 	}
 }
 
-void FCanvas::Flush_RenderThread(FRHICommandListImmediate& RHICmdList, bool bForce)
+void FCanvas::Flush_RenderThread(FRHICommandListImmediate& RHICmdList, bool bForce, bool bInsideRenderPass)
 {
-	// This renderpass is self contained.
-	check(RHICmdList.IsOutsideRenderPass());
-
 	SCOPE_CYCLE_COUNTER(STAT_Canvas_FlushTime);
 
 	if (!(AllowedModes&Allow_Flush) && !bForce)
@@ -660,9 +657,7 @@ void FCanvas::Flush_RenderThread(FRHICommandListImmediate& RHICmdList, bool bFor
 	}
 
 	// Update the font cache with new text before elements are drawn
-	{
-		FEngineFontServices::Get().UpdateCache();
-	}
+	FEngineFontServices::Get().UpdateCache();
 
 	// FCanvasSortElement compare class
 	struct FCompareFCanvasSortElement
@@ -676,22 +671,32 @@ void FCanvas::Flush_RenderThread(FRHICommandListImmediate& RHICmdList, bool bFor
 	SortedElements.Sort(FCompareFCanvasSortElement());
 
 	SCOPED_DRAW_EVENT(RHICmdList, CanvasFlush);
-	const FTexture2DRHIRef& RenderTargetTexture = RenderTarget->GetRenderTargetTexture();
 
-	check(IsValidRef(RenderTargetTexture));
+	if (!bInsideRenderPass)
+	{
+		const FTexture2DRHIRef& RenderTargetTexture = RenderTarget->GetRenderTargetTexture();
 
-	FRHIRenderPassInfo RPInfo(RenderTargetTexture, ERenderTargetActions::Load_Store);
-	
-	// Set the RHI render target.
-	if (IsUsingInternalTexture())
+		check(IsValidRef(RenderTargetTexture));
+		check(RHICmdList.IsOutsideRenderPass());
+
+		FRHIRenderPassInfo RPInfo(RenderTargetTexture, ERenderTargetActions::Load_Store);
+
+		// Set the RHI render target.
+		if (IsUsingInternalTexture())
+		{
+			RPInfo.ColorRenderTargets[0].Action = ERenderTargetActions::Clear_Store;
+		}
+
+		RHICmdList.BeginRenderPass(RPInfo, TEXT("CanvasRenderThread"));
+	}
+	else
 	{
-		RPInfo.ColorRenderTargets[0].Action = ERenderTargetActions::Clear_Store;
+		check(RHICmdList.IsInsideRenderPass());
 	}
-	
-	RHICmdList.BeginRenderPass(RPInfo, TEXT("CanvasRenderThread"));
+
 	{
+		// Disable depth test & writes
 		FMeshPassProcessorRenderState DrawRenderState;
-		// disable depth test & writes
 		DrawRenderState.SetDepthStencilState(TStaticDepthStencilState<false, CF_Always>::GetRHI());
 
 		if (ViewRect.Area() <= 0)
@@ -738,7 +743,11 @@ void FCanvas::Flush_RenderThread(FRHICommandListImmediate& RHICmdList, bool bFor
 			LastElementIndex = INDEX_NONE;
 		}
 	}
-	RHICmdList.EndRenderPass();	
+
+	if (!bInsideRenderPass)
+	{
+		RHICmdList.EndRenderPass();
+	}
 }
 
 void FCanvas::Flush_GameThread(bool bForce)
diff --git a/Engine/Source/Runtime/Engine/Public/CanvasTypes.h b/Engine/Source/Runtime/Engine/Public/CanvasTypes.h
--- a/Engine/Source/Runtime/Engine/Public/CanvasTypes.h
+++ b/Engine/Source/Runtime/Engine/Public/CanvasTypes.h
@@ -270,8 +270,10 @@ public:
 	* Sends a message to the rendering thread to draw the batched elements.
 	* @param RHICmdList - command list to use
 	* @param bForce - force the flush even if Allow_Flush is not enabled
+	* @param bInsideRenderPass - Set to true if flushing inside a render pass (e.g. Render Graph pass).
+	*	This will skip creating a render pass internally, and assert if the command list is not in a render pass.
 	*/
-	ENGINE_API void Flush_RenderThread(FRHICommandListImmediate& RHICmdList, bool bForce = false);
+	ENGINE_API void Flush_RenderThread(FRHICommandListImmediate& RHICmdList, bool bForce = false, bool bInsideRenderPass = false);
 
 	/**
 	* Sends a message to the rendering thread to draw the batched elements.
diff --git a/Engine/Source/Runtime/RenderCore/Private/RenderGraphBuilder.cpp b/Engine/Source/Runtime/RenderCore/Private/RenderGraphBuilder.cpp
--- a/Engine/Source/Runtime/RenderCore/Private/RenderGraphBuilder.cpp
+++ b/Engine/Source/Runtime/RenderCore/Private/RenderGraphBuilder.cpp
@@ -351,11 +351,8 @@ void FRDGBuilder::ValidatePass(const FRenderGraphPass* Pass) const
 					RenderTarget.GetTexture()->Name);
 
 				// TODO(RDG): should only be done when there is a store action.
-				if (!Texture->bHasEverBeenProduced)
-				{
-					Texture->bHasEverBeenProduced = true;
-					Texture->DebugFirstProducer = Pass;
-				}
+				Texture->bHasEverBeenProduced = true;
+				Texture->DebugFirstProducer = Pass;
 			}
 
 			bFoundRTBound = bFoundRTBound || IsBoundAsReadable(RenderTarget.GetTexture(), ParameterStruct);
@@ -951,7 +948,7 @@ void FRDGBuilder::AllocateAndTransitionPassResources(const FRenderGraphPass* Pas
 				#if RENDER_GRAPH_DEBUGGING
 				{
 					check(Buffer->bHasEverBeenProduced);
-				}	
+				}
 				#endif
 				check(Buffer->PooledBuffer);
 
@@ -978,7 +975,7 @@ void FRDGBuilder::AllocateAndTransitionPassResources(const FRenderGraphPass* Pas
 				#if RENDER_GRAPH_DEBUGGING
 				{
 					check(SRV->Desc.Buffer->bHasEverBeenProduced);
-				}	
+				}
 				#endif
 				check(SRV->Desc.Buffer->PooledBuffer);
 				
diff --git a/Engine/Source/Runtime/RenderCore/Public/RenderGraphResources.h b/Engine/Source/Runtime/RenderCore/Public/RenderGraphResources.h
--- a/Engine/Source/Runtime/RenderCore/Public/RenderGraphResources.h
+++ b/Engine/Source/Runtime/RenderCore/Public/RenderGraphResources.h
@@ -42,7 +42,7 @@ using FRDGBufferUAVRef = const FRDGBufferUAV*;
 /** Generic graph resource. */
 class RENDERCORE_API FRDGResource
 {
-private:
+protected:
 	/** Pointer on the RHI resource once allocated, that the RHI can dereferenced according to IsUniformBufferResourceIndirectionType(). */
 	union
 	{
@@ -89,8 +89,6 @@ private:
 #endif
 
 	friend class FRDGBuilder;
-
-
 	/** Friendship over parameter settings for bIsActuallyUsedByPass. It means only this parameter settings can be
 	 * used for render graph resource, to force the useless dependency tracking to just work.
 	 */
@@ -117,10 +115,8 @@ using FRDGTextureDesc = FPooledRenderTargetDesc;
 class RENDERCORE_API FRDGTexture : public FRDGResource
 {
 public:
-	//TODO(RDG): using FDesc = FPooledRenderTargetDesc;
-
 	/** Descriptor of the graph tracked texture. */
-	const FPooledRenderTargetDesc Desc;
+	const FRDGTextureDesc Desc;
 
 	/** Returns the allocated pooled render target. Must only be called within a pass's lambda. */
 	inline IPooledRenderTarget* GetPooledRenderTarget() const
@@ -129,11 +125,12 @@ public:
 		return PooledRenderTarget;
 	}
 
+
 	/** Returns the allocated RHI texture. Must only be called within a pass's lambda. */
 	inline FTextureRHIParamRef GetRHITexture() const
 	{
-		check(PooledRenderTarget);
-		return PooledRenderTarget->GetRenderTargetItem().ShaderResourceTexture;
+		check(CachedRHI.Texture);
+		return CachedRHI.Texture;
 	}
 
 private:
@@ -145,7 +142,7 @@ private:
 	FRDGTexture(const TCHAR* DebugName, const FPooledRenderTargetDesc& InDesc)
 		: FRDGResource(DebugName)
 		, Desc(InDesc)
-	{ }
+	{}
 
 	friend class FRDGBuilder;
 
@@ -158,7 +155,7 @@ private:
 #endif
 };
 
-/** Decsriptor for render graph tracked SRV. */
+/** Descriptor for render graph tracked SRV. */
 class RENDERCORE_API FRDGTextureSRVDesc
 {
 public:
@@ -189,7 +186,7 @@ private:
 	friend class FRDGBuilder;
 };
 
-/** Decsriptor for render graph tracked UAV. */
+/** Descriptor for render graph tracked UAV. */
 class RENDERCORE_API FRDGTextureUAVDesc
 {
 public:
diff --git a/Engine/Source/Runtime/Renderer/Private/CompositionLighting/CompositionLighting.cpp b/Engine/Source/Runtime/Renderer/Private/CompositionLighting/CompositionLighting.cpp
--- a/Engine/Source/Runtime/Renderer/Private/CompositionLighting/CompositionLighting.cpp
+++ b/Engine/Source/Runtime/Renderer/Private/CompositionLighting/CompositionLighting.cpp
@@ -27,34 +27,6 @@ DECLARE_GPU_STAT_NAMED(CompositionPreLighting, TEXT("Composition PreLighting") )
 DECLARE_GPU_STAT_NAMED(CompositionLpvIndirect, TEXT("Composition LpvIndirect") );
 DECLARE_GPU_STAT_NAMED(CompositionPostLighting, TEXT("Composition PostLighting") );
 
-// -------------------------------------------------------
-
-static TAutoConsoleVariable<float> CVarSSSScale(
-	TEXT("r.SSS.Scale"),
-	1.0f,
-	TEXT("Affects the Screen space subsurface scattering pass")
-	TEXT("(use shadingmodel SubsurfaceProfile, get near to the object as the default)\n")
-	TEXT("is human skin which only scatters about 1.2cm)\n")
-	TEXT(" 0: off (if there is no object on the screen using this pass it should automatically disable the post process pass)\n")
-	TEXT("<1: scale scatter radius down (for testing)\n")
-	TEXT(" 1: use given radius form the Subsurface scattering asset (default)\n")
-	TEXT(">1: scale scatter radius up (for testing)"),
-	ECVF_Scalability | ECVF_RenderThreadSafe);
-
-static TAutoConsoleVariable<int32> CVarSSSHalfRes(
-	TEXT("r.SSS.HalfRes"),
-	1,
-	TEXT(" 0: full quality (not optimized, as reference)\n")
-	TEXT(" 1: parts of the algorithm runs in half resolution which is lower quality but faster (default)"),
-	ECVF_RenderThreadSafe | ECVF_Scalability);
-
-static TAutoConsoleVariable<int32> CVarSubsurfaceScattering(
-	TEXT("r.SubsurfaceScattering"),
-	1,
-	TEXT(" 0: disabled\n")
-	TEXT(" 1: enabled (default)"),
-	ECVF_RenderThreadSafe | ECVF_Scalability);
-
 static TAutoConsoleVariable<int32> CVarSSAOSmoothPass(
 	TEXT("r.AmbientOcclusion.Compute.Smooth"),
 	1,
@@ -438,79 +410,6 @@ void FCompositionLighting::ProcessLpvIndirect(FRHICommandListImmediate& RHICmdLi
 	CompositeContext.Process(Context.FinalOutput.GetPass(), TEXT("CompositionLighting"));
 }
 
-void FCompositionLighting::ProcessAfterLighting(FRHICommandListImmediate& RHICmdList, FViewInfo& View)
-{
-	check(IsInRenderingThread());
-	FSceneRenderTargets& SceneContext = FSceneRenderTargets::Get(RHICmdList);
-
-	{
-		FMemMark Mark(FMemStack::Get());
-		FRenderingCompositePassContext CompositeContext(RHICmdList, View);
-		FPostprocessContext Context(RHICmdList, CompositeContext.Graph, View);
-		FRenderingCompositeOutputRef AmbientOcclusion;
-
-		// Screen Space Subsurface Scattering
-		{
-			float Radius = CVarSSSScale.GetValueOnRenderThread();
-			bool bSimpleDynamicLighting = IsAnyForwardShadingEnabled(View.GetShaderPlatform());
-			bool bScreenSpaceSubsurfacePassNeeded = ((View.ShadingModelMaskInView & GetUseSubsurfaceProfileShadingModelMask()) != 0) && IsSubsurfacePostprocessRequired();
-			bool bSubsurfaceAllowed = CVarSubsurfaceScattering.GetValueOnRenderThread() == 1;
-
-			if (bScreenSpaceSubsurfacePassNeeded 
-				&& !bSimpleDynamicLighting 
-				&& bSubsurfaceAllowed)
-			{
-				bool bHalfRes = CVarSSSHalfRes.GetValueOnRenderThread() != 0;
-				bool bSingleViewportMode = View.Family->Views.Num() == 1;
-
-				if(Radius > 0 && View.Family->EngineShowFlags.SubsurfaceScattering)
-				{
-					FRenderingCompositePass* PassSetup = Context.Graph.RegisterPass(new(FMemStack::Get()) FRCPassPostProcessSubsurfaceSetup(View, bHalfRes));
-					PassSetup->SetInput(ePId_Input0, Context.FinalOutput);
-
-					FRenderingCompositePass* PassX = Context.Graph.RegisterPass(new(FMemStack::Get()) FRCPassPostProcessSubsurface(0, bHalfRes));
-					PassX->SetInput(ePId_Input0, PassSetup);
-
-					FRenderingCompositePass* PassY = Context.Graph.RegisterPass(new(FMemStack::Get()) FRCPassPostProcessSubsurface(1, bHalfRes));
-					PassY->SetInput(ePId_Input0, PassX);
-					PassY->SetInput(ePId_Input1, PassSetup);
-
-					// full res composite pass, no blurring (Radius=0), replaces SceneColor
-					FRenderingCompositePass* RecombinePass = Context.Graph.RegisterPass(new(FMemStack::Get()) FRCPassPostProcessSubsurfaceRecombine(bHalfRes, bSingleViewportMode));
-					RecombinePass->SetInput(ePId_Input0, Context.FinalOutput);
-					RecombinePass->SetInput(ePId_Input1, PassY);
-					RecombinePass->SetInput(ePId_Input2, PassSetup);
-					Context.FinalOutput = FRenderingCompositeOutputRef(RecombinePass);
-				}
-				else
-				{
-					// needed for Scalability
-					FRenderingCompositePass* RecombinePass = Context.Graph.RegisterPass(new(FMemStack::Get()) FRCPassPostProcessSubsurfaceRecombine(bHalfRes, bSingleViewportMode));
-					RecombinePass->SetInput(ePId_Input0, Context.FinalOutput);
-					Context.FinalOutput = FRenderingCompositeOutputRef(RecombinePass);
-				}
-			}
-		}
-
-		// The graph setup should be finished before this line ----------------------------------------
-
-		SCOPED_DRAW_EVENT(RHICmdList, CompositionAfterLighting);
-		SCOPED_GPU_STAT(RHICmdList, CompositionPostLighting);
-
-		// we don't replace the final element with the scenecolor because this is what those passes should do by themself
-
-		CompositeContext.Process(Context.FinalOutput.GetPass(), TEXT("CompositionLighting"));
-	}
-
-	// We only release the after the last view was processed (SplitScreen)
-	if(View.Family->Views[View.Family->Views.Num() - 1] == &View)
-	{
-		// The RT should be released as early as possible to allow sharing of that memory for other purposes.
-		// This becomes even more important with some limited VRam (XBoxOne).
-		SceneContext.SetLightAttenuation(0);
-	}
-}
-
 bool FCompositionLighting::CanProcessAsyncSSAO(TArray<FViewInfo>& Views)
 {
 	bool bAnyAsyncSSAO = true;
@@ -612,11 +511,4 @@ void FCompositionLighting::GfxWaitForAsyncSSAO(FRHICommandListImmediate& RHICmdL
 		RHICmdList.WaitComputeFence(AsyncSSAOFence);
 		AsyncSSAOFence = nullptr;
 	}
-}
-
-bool FCompositionLighting::IsSubsurfacePostprocessRequired() const
-{
-	const bool bSSSEnabled = CVarSubsurfaceScattering->GetInt() != 0;
-	const bool bSSSScaleEnabled = CVarSSSScale.GetValueOnAnyThread() > 0.0f;
-	return (bSSSEnabled && bSSSScaleEnabled);	
-}
+}
\ No newline at end of file
diff --git a/Engine/Source/Runtime/Renderer/Private/CompositionLighting/CompositionLighting.h b/Engine/Source/Runtime/Renderer/Private/CompositionLighting/CompositionLighting.h
--- a/Engine/Source/Runtime/Renderer/Private/CompositionLighting/CompositionLighting.h
+++ b/Engine/Source/Runtime/Renderer/Private/CompositionLighting/CompositionLighting.h
@@ -22,14 +22,10 @@ public:
 	// only call if LPV is enabled
 	void ProcessLpvIndirect(FRHICommandListImmediate& RHICmdList, FViewInfo& View);
 
-	void ProcessAfterLighting(FRHICommandListImmediate& RHICmdList, FViewInfo& View);
-
 	bool CanProcessAsyncSSAO(TArray<FViewInfo>& Views);
 	void ProcessAsyncSSAO(FRHICommandListImmediate& RHICmdList, TArray<FViewInfo>& Views);
 	void GfxWaitForAsyncSSAO(FRHICommandListImmediate& RHICmdList);
 
-	bool IsSubsurfacePostprocessRequired() const;
-
 private:
 	void PrepareAsyncSSAO(FRHICommandListImmediate& RHICmdList, TArray<FViewInfo>& Views);
 	void FinishAsyncSSAO(FRHICommandListImmediate& RHICmdList);
diff --git a/Engine/Source/Runtime/Renderer/Private/CompositionLighting/PostProcessPassThrough.cpp b/Engine/Source/Runtime/Renderer/Private/CompositionLighting/PostProcessPassThrough.cpp
--- a/Engine/Source/Runtime/Renderer/Private/CompositionLighting/PostProcessPassThrough.cpp
+++ b/Engine/Source/Runtime/Renderer/Private/CompositionLighting/PostProcessPassThrough.cpp
@@ -161,61 +161,4 @@ FPooledRenderTargetDesc FRCPassPostProcessPassThrough::ComputeOutputDesc(EPassOu
 	Ret.DebugName = TEXT("PassThrough");
 
 	return Ret;
-}
-
-void CopyOverOtherViewportsIfNeeded(FRenderingCompositePassContext& Context, const FSceneView& ExcludeView)
-{
-	const FViewInfo& View = Context.View;
-	const FSceneViewFamily* ViewFamily = View.Family;
-
-	FGraphicsPipelineStateInitializer GraphicsPSOInit;
-	Context.RHICmdList.ApplyCachedRenderTargets(GraphicsPSOInit);
-	GraphicsPSOInit.BlendState = TStaticBlendState<>::GetRHI();
-	GraphicsPSOInit.RasterizerState = TStaticRasterizerState<>::GetRHI();
-	GraphicsPSOInit.DepthStencilState = TStaticDepthStencilState<false, CF_Always>::GetRHI();
-
-	// only for multiple views we need to do this
-	if(ViewFamily->Views.Num())
-	{
-		SCOPED_DRAW_EVENT(Context.RHICmdList, CopyOverOtherViewportsIfNeeded);
-
-		TShaderMapRef<FPostProcessVS> VertexShader(Context.GetShaderMap());
-		TShaderMapRef<FPostProcessPassThroughPS> PixelShader(Context.GetShaderMap());
-
-		GraphicsPSOInit.BoundShaderState.VertexDeclarationRHI = GFilterVertexDeclaration.VertexDeclarationRHI;
-		GraphicsPSOInit.BoundShaderState.VertexShaderRHI = GETSAFERHISHADER_VERTEX(*VertexShader);
-		GraphicsPSOInit.BoundShaderState.PixelShaderRHI = GETSAFERHISHADER_PIXEL(*PixelShader);
-		GraphicsPSOInit.PrimitiveType = PT_TriangleList;
-
-		SetGraphicsPipelineState(Context.RHICmdList, GraphicsPSOInit);
-
-		VertexShader->SetParameters(Context);
-		PixelShader->SetParameters(Context.RHICmdList, Context);
-
-		FSceneRenderTargets& SceneContext = FSceneRenderTargets::Get(Context.RHICmdList);
-
-		FIntPoint Size = SceneContext.GetBufferSizeXY();
-	
-		for(uint32 ViewId = 0, ViewCount = ViewFamily->Views.Num(); ViewId < ViewCount; ++ViewId)
-		{
-			const FViewInfo* LocalView = static_cast<const FViewInfo*>(ViewFamily->Views[ViewId]);
-			
-			if(LocalView != &ExcludeView)
-			{
-				FIntRect Rect = LocalView->ViewRect;
-
-				DrawPostProcessPass(Context.RHICmdList,
-					Rect.Min.X, Rect.Min.Y,
-					Rect.Width(), Rect.Height(),
-					Rect.Min.X, Rect.Min.Y,
-					Rect.Width(), Rect.Height(),
-					Size,
-					Size,
-					*VertexShader,
-					LocalView->StereoPass, 
-					Context.HasHmdMesh(),
-					EDRF_UseTriangleOptimization);
-			}
-		}
-	}
-}
+}
\ No newline at end of file
diff --git a/Engine/Source/Runtime/Renderer/Private/CompositionLighting/PostProcessPassThrough.h b/Engine/Source/Runtime/Renderer/Private/CompositionLighting/PostProcessPassThrough.h
--- a/Engine/Source/Runtime/Renderer/Private/CompositionLighting/PostProcessPassThrough.h
+++ b/Engine/Source/Runtime/Renderer/Private/CompositionLighting/PostProcessPassThrough.h
@@ -66,13 +66,4 @@ private:
 	IPooledRenderTarget* Dest;
 	bool bAdditiveBlend;
 	FPooledRenderTargetDesc NewDesc;
-};
-
-// useful to replace the SceneColor and leave the other views e.g. when doing stereo rendering, assumed there is no scaling on the size (input and output is BufferSize)
-// call before SetViewport()
-// assumes this to be called before
-//  SetRenderTarget()
-//  Context.RHICmdList.SetBlendState(TStaticBlendState<>::GetRHI());
-//  Context.RHICmdList.SetRasterizerState(TStaticRasterizerState<>::GetRHI());
-// 	Context.RHICmdList.SetDepthStencilState(TStaticDepthStencilState<false, CF_Always>::GetRHI());
-void CopyOverOtherViewportsIfNeeded(FRenderingCompositePassContext& Context, const FSceneView& ExcludeView);
+};
\ No newline at end of file
diff --git a/Engine/Source/Runtime/Renderer/Private/DeferredShadingRenderer.cpp b/Engine/Source/Runtime/Renderer/Private/DeferredShadingRenderer.cpp
--- a/Engine/Source/Runtime/Renderer/Private/DeferredShadingRenderer.cpp
+++ b/Engine/Source/Runtime/Renderer/Private/DeferredShadingRenderer.cpp
@@ -11,6 +11,7 @@
 #include "ScreenRendering.h"
 #include "PostProcess/SceneFilterRendering.h"
 #include "PostProcess/ScreenSpaceReflections.h"
+#include "PostProcess/PostProcessSubsurface.h"
 #include "CompositionLighting/CompositionLighting.h"
 #include "FXSystem.h"
 #include "OneColorShader.h"
@@ -1766,13 +1767,41 @@ void FDeferredShadingSceneRenderer::Render(FRHICommandListImmediate& RHICmdList)
 
 		ServiceLocalQueue();
 
-		// Post-lighting composition lighting stage
-		// e.g. ScreenSpaceSubsurfaceScattering
-		for(int32 ViewIndex = 0; ViewIndex < Views.Num(); ++ViewIndex)
-		{	
-			SCOPED_CONDITIONAL_DRAW_EVENTF(RHICmdList, EventView,Views.Num() > 1, TEXT("View%d"), ViewIndex);
-			GCompositionLighting.ProcessAfterLighting(RHICmdList, Views[ViewIndex]);
+		{
+			FSceneRenderTargets& SceneRenderTargets = FSceneRenderTargets::Get(RHICmdList);
+
+			FRDGBuilder GraphBuilder(RHICmdList);
+
+			FRDGTextureRef SceneColorTexture = GraphBuilder.RegisterExternalTexture(SceneRenderTargets.GetSceneColor(), TEXT("SceneColor"));
+
+			// Post-lighting composition lighting stage
+			// e.g. ScreenSpaceSubsurfaceScattering
+			for (int32 ViewIndex = 0; ViewIndex < Views.Num(); ++ViewIndex)
+			{
+				RDG_EVENT_SCOPE(GraphBuilder, "SubsurfaceScattering(ViewId=%d)", ViewIndex);
+
+				const FViewInfo& View = Views[ViewIndex];
+
+				if (IsSubsurfaceRequiredForView(View))
+				{
+					FScreenPassContext* Context = FScreenPassContext::Create(RHICmdList, View);
+
+					SceneColorTexture = ComputeSubsurface(GraphBuilder, Context, SceneColorTexture);
+				}
+			}
+
+			// Extract the result texture out and re-assign it to the scene render targets blackboard.
+			TRefCountPtr<IPooledRenderTarget> SceneColorTarget;
+			GraphBuilder.QueueTextureExtraction(SceneColorTexture, &SceneColorTarget, false);
+			GraphBuilder.Execute();
+
+			SceneRenderTargets.SetSceneColor(SceneColorTarget);
+
+			// The RT should be released as early as possible to allow sharing of that memory for other purposes.
+			// This becomes even more important with some limited VRam (XBoxOne).
+			SceneRenderTargets.SetLightAttenuation(nullptr);
 		}
+
 #if RHI_RAYTRACING
 		if (SkyLightRT)
 		{
diff --git a/Engine/Source/Runtime/Renderer/Private/PostProcess/PostProcessSubsurface.cpp b/Engine/Source/Runtime/Renderer/Private/PostProcess/PostProcessSubsurface.cpp
--- a/Engine/Source/Runtime/Renderer/Private/PostProcess/PostProcessSubsurface.cpp
+++ b/Engine/Source/Runtime/Renderer/Private/PostProcess/PostProcessSubsurface.cpp
@@ -5,371 +5,178 @@
 =============================================================================*/
 
 #include "PostProcess/PostProcessSubsurface.h"
-#include "EngineGlobals.h"
+#include "PostProcess/SceneRenderTargets.h"
 #include "Engine/SubsurfaceProfile.h"
-#include "StaticBoundShaderState.h"
 #include "CanvasTypes.h"
-#include "UnrealEngine.h"
-#include "SceneUtils.h"
-#include "RenderTargetPool.h"
 #include "RenderTargetTemp.h"
-#include "PostProcess/SceneRenderTargets.h"
-#include "PostProcess/SceneFilterRendering.h"
-#include "SceneRenderTargetParameters.h"
-#include "PostProcess/PostProcessing.h"
-#include "CompositionLighting/PostProcessPassThrough.h"
-#include "ClearQuad.h"
-#include "PipelineStateCache.h"
-#include "VisualizeTexture.h"
 
 ENGINE_API const IPooledRenderTarget* GetSubsufaceProfileTexture_RT(FRHICommandListImmediate& RHICmdList);
 
+namespace
+{
+	TAutoConsoleVariable<int32> CVarSubsurfaceScattering(
+		TEXT("r.SubsurfaceScattering"),
+		1,
+		TEXT(" 0: disabled\n")
+		TEXT(" 1: enabled (default)"),
+		ECVF_RenderThreadSafe | ECVF_Scalability);
+
+	TAutoConsoleVariable<float> CVarSSSScale(
+		TEXT("r.SSS.Scale"),
+		1.0f,
+		TEXT("Affects the Screen space subsurface scattering pass")
+		TEXT("(use shadingmodel SubsurfaceProfile, get near to the object as the default)\n")
+		TEXT("is human skin which only scatters about 1.2cm)\n")
+		TEXT(" 0: off (if there is no object on the screen using this pass it should automatically disable the post process pass)\n")
+		TEXT("<1: scale scatter radius down (for testing)\n")
+		TEXT(" 1: use given radius form the Subsurface scattering asset (default)\n")
+		TEXT(">1: scale scatter radius up (for testing)"),
+		ECVF_Scalability | ECVF_RenderThreadSafe);
+
+	TAutoConsoleVariable<int32> CVarSSSHalfRes(
+		TEXT("r.SSS.HalfRes"),
+		1,
+		TEXT(" 0: full quality (not optimized, as reference)\n")
+		TEXT(" 1: parts of the algorithm runs in half resolution which is lower quality but faster (default)"),
+		ECVF_RenderThreadSafe | ECVF_Scalability);
+
+	TAutoConsoleVariable<int32> CVarSSSQuality(
+		TEXT("r.SSS.Quality"),
+		0,
+		TEXT("Defines the quality of the recombine pass when using the SubsurfaceScatteringProfile shading model\n")
+		TEXT(" 0: low (faster, default)\n")
+		TEXT(" 1: high (sharper details but slower)\n")
+		TEXT("-1: auto, 1 if TemporalAA is disabled (without TemporalAA the quality is more noticable)"),
+		ECVF_RenderThreadSafe | ECVF_Scalability);
+
+	TAutoConsoleVariable<int32> CVarSSSFilter(
+		TEXT("r.SSS.Filter"),
+		1,
+		TEXT("Defines the filter method for Screenspace Subsurface Scattering feature.\n")
+		TEXT(" 0: point filter (useful for testing, could be cleaner)\n")
+		TEXT(" 1: bilinear filter"),
+		ECVF_RenderThreadSafe | ECVF_Scalability);
+
+	TAutoConsoleVariable<int32> CVarSSSSampleSet(
+		TEXT("r.SSS.SampleSet"),
+		2,
+		TEXT("Defines how many samples we use for Screenspace Subsurface Scattering feature.\n")
+		TEXT(" 0: lowest quality (6*2+1)\n")
+		TEXT(" 1: medium quality (9*2+1)\n")
+		TEXT(" 2: high quality (13*2+1) (default)"),
+		ECVF_RenderThreadSafe | ECVF_Scalability);
+
+	TAutoConsoleVariable<int32> CVarSSSCheckerboard(
+		TEXT("r.SSS.Checkerboard"),
+		2,
+		TEXT("Enables or disables checkerboard rendering for subsurface profile rendering.\n")
+		TEXT("This is necessary if SceneColor does not include a floating point alpha channel (e.g 32-bit formats)\n")
+		TEXT(" 0: Disabled (high quality) \n")
+		TEXT(" 1: Enabled (low quality). Surface lighting will be at reduced resolution.\n")
+		TEXT(" 2: Automatic. Non-checkerboard lighting will be applied if we have a suitable rendertarget format\n"),
+		ECVF_RenderThreadSafe);
+}
+
+enum class ESubsurfaceMode : uint32
+{
+	// Performs a full resolution scattering filter.
+	FullRes,
 
-static TAutoConsoleVariable<int32> CVarSSSQuality(
-	TEXT("r.SSS.Quality"),
-	0,
-	TEXT("Defines the quality of the recombine pass when using the SubsurfaceScatteringProfile shading model\n")
-	TEXT(" 0: low (faster, default)\n")
-	TEXT(" 1: high (sharper details but slower)\n")
-	TEXT("-1: auto, 1 if TemporalAA is disabled (without TemporalAA the quality is more noticable)"),
-	ECVF_RenderThreadSafe  | ECVF_Scalability);
-
-static TAutoConsoleVariable<int32> CVarSSSFilter(
-	TEXT("r.SSS.Filter"),
-	1,
-	TEXT("Defines the filter method for Screenspace Subsurface Scattering feature.\n")
-	TEXT(" 0: point filter (useful for testing, could be cleaner)\n")
-	TEXT(" 1: bilinear filter"),
-	ECVF_RenderThreadSafe  | ECVF_Scalability);
-
-static TAutoConsoleVariable<int32> CVarSSSSampleSet(
-	TEXT("r.SSS.SampleSet"),
-	2,
-	TEXT("Defines how many samples we use for Screenspace Subsurface Scattering feature.\n")
-	TEXT(" 0: lowest quality (6*2+1)\n")
-	TEXT(" 1: medium quality (9*2+1)\n")
-	TEXT(" 2: high quality (13*2+1) (default)"),
-	ECVF_RenderThreadSafe  | ECVF_Scalability);
-
-static TAutoConsoleVariable<int32> CVarCheckerboardSubsurfaceProfileRendering(
-	TEXT("r.SSS.Checkerboard"),
-	2,
-	TEXT("Enables or disables checkerboard rendering for subsurface profile rendering.\n")
-	TEXT("This is necessary if SceneColor does not include a floating point alpha channel (e.g 32-bit formats)\n")
-	TEXT(" 0: Disabled (high quality) \n")
-	TEXT(" 1: Enabled (low quality). Surface lighting will be at reduced resolution.\n")
-	TEXT(" 2: Automatic. Non-checkerboard lighting will be applied if we have a suitable rendertarget format\n"),
-	ECVF_RenderThreadSafe
-	);
-
-// -------------------------------------------------------------
+	// Performs a half resolution scattering filter.
+	HalfRes,
 
+	// Reconstructs lighting, but does not perform scattering.
+	Bypass,
+
+	MAX
+};
+
+// Returns the [0, N] clamped value of the 'r.SSS.Scale' CVar.
 float GetSubsurfaceRadiusScale()
 {
 	static const auto CVar = IConsoleManager::Get().FindTConsoleVariableDataFloat(TEXT("r.SSS.Scale"));
 	check(CVar);
-	float Ret = CVar->GetValueOnRenderThread();
 
-	return FMath::Max(0.0f, Ret);
+	return FMath::Max(0.0f, CVar->GetValueOnRenderThread());
 }
 
-// -------------------------------------------------------------
-
-const bool FRCPassPostProcessSubsurface::RequiresCheckerboardSubsurfaceRendering(EPixelFormat SceneColorFormat)
+// Returns the current subsurface mode required by the current view.
+ESubsurfaceMode GetSubsurfaceModeForView(const FViewInfo& View)
 {
-	int CVarValue = CVarCheckerboardSubsurfaceProfileRendering.GetValueOnRenderThread();
-	if (CVarValue == 0)
-	{
-		return false;
-	}
-	else if ( CVarValue == 1 )
-	{
-		return true;
-	}
-	else if (CVarValue == 2)
-	{
-		switch (SceneColorFormat)
-		{
-		case PF_A32B32G32R32F:
-		case PF_FloatRGBA:
-			return false;
-		default:
-			return true;
-		}
-	}
-	return true;
-}
-
-// -------------------------------------------------------------
+	const float Radius = GetSubsurfaceRadiusScale();
+	const bool bShowSubsurfaceScattering = Radius > 0 && View.Family->EngineShowFlags.SubsurfaceScattering;
 
-/** Shared shader parameters needed for screen space subsurface scattering. */
-class FSubsurfaceParameters
-{
-public:
-	void Bind(const FShaderParameterMap& ParameterMap)
-	{
-		SSSParams.Bind(ParameterMap, TEXT("SSSParams"));
-		SSProfilesTexture.Bind(ParameterMap, TEXT("SSProfilesTexture"));
-	}
-
-	void SetParameters(FRHICommandList& RHICmdList, const FPixelShaderRHIParamRef& ShaderRHI, const FRenderingCompositePassContext& Context) const
+	if (bShowSubsurfaceScattering)
 	{
+		const bool bHalfRes = CVarSSSHalfRes.GetValueOnRenderThread() != 0;
+		if (bHalfRes)
 		{
-			// from Separabale.usf: float distanceToProjectionWindow = 1.0 / tan(0.5 * radians(SSSS_FOVY))
-			// can be extracted out of projection matrix
-
-			FSceneRenderTargets& SceneContext = FSceneRenderTargets::Get(RHICmdList);
-
-			// Calculate the sssWidth scale (1.0 for a unit plane sitting on the projection window):
-			float DistanceToProjectionWindow = Context.View.ViewMatrices.GetProjectionMatrix().M[0][0]; 
-
-			float SSSScaleZ = DistanceToProjectionWindow * GetSubsurfaceRadiusScale();
-
-			// * 0.5f: hacked in 0.5 - -1..1 to 0..1 but why this isn't in demo code?
-			float SSSScaleX = SSSScaleZ / SUBSURFACE_KERNEL_SIZE * 0.5f;
-
-			FVector4 ColorScale(SSSScaleX, SSSScaleZ, 0, 0);
-			SetShaderValue(Context.RHICmdList, ShaderRHI, SSSParams, ColorScale);
+			return ESubsurfaceMode::HalfRes;
 		}
-
+		else
 		{
-			const IPooledRenderTarget* PooledRT = GetSubsufaceProfileTexture_RT(Context.RHICmdList);
-
-			if(!PooledRT)
-			{
-				// no subsurface profile was used yet
-				PooledRT = GSystemTextures.BlackDummy;
-			}
-
-			const FSceneRenderTargetItem& Item = PooledRT->GetRenderTargetItem();
-
-			SetTextureParameter(Context.RHICmdList, ShaderRHI, SSProfilesTexture, Item.ShaderResourceTexture);
+			return ESubsurfaceMode::FullRes;
 		}
 	}
-
-	friend FArchive& operator<<(FArchive& Ar,FSubsurfaceParameters& P)
+	else
 	{
-		Ar << P.SSSParams << P.SSProfilesTexture;
-		return Ar;
+		return ESubsurfaceMode::Bypass;
 	}
+}
 
-private:
-	FShaderParameter SSSParams;
-	FShaderResourceParameter SSProfilesTexture;
-};
-
-// ---------------------------------------------
-
-
-/**
- * Encapsulates the post processing subsurface scattering pixel shader.
- */
-class FPostProcessSubsurfaceVisualizePS : public FGlobalShader
+// Returns the SS profile texture with a black fallback texture if none exists yet.
+FTextureRHIRef GetSubsurfaceProfileTexture(FRHICommandListImmediate& RHICmdList)
 {
-	DECLARE_SHADER_TYPE(FPostProcessSubsurfaceVisualizePS , Global);
+	const IPooledRenderTarget* ProfileTextureTarget = GetSubsufaceProfileTexture_RT(RHICmdList);
 
-	static bool ShouldCompilePermutation(const FGlobalShaderPermutationParameters& Parameters)
+	if (!ProfileTextureTarget)
 	{
-		return IsFeatureLevelSupported(Parameters.Platform, ERHIFeatureLevel::SM4);
-	}
-
-	static void ModifyCompilationEnvironment(const FGlobalShaderPermutationParameters& Parameters, FShaderCompilerEnvironment& OutEnvironment)
-	{
-		FGlobalShader::ModifyCompilationEnvironment(Parameters, OutEnvironment);
-		OutEnvironment.SetDefine(TEXT("SUBSURFACE_RADIUS_SCALE"), SUBSURFACE_RADIUS_SCALE);
-		OutEnvironment.SetDefine(TEXT("SUBSURFACE_KERNEL_SIZE"), SUBSURFACE_KERNEL_SIZE);
+		// No subsurface profile was used yet
+		ProfileTextureTarget = GSystemTextures.BlackDummy;
 	}
 
-	/** Default constructor. */
-	FPostProcessSubsurfaceVisualizePS () {}
-
-public:
-	FPostProcessPassParameters PostprocessParameter;
-	FSceneTextureShaderParameters SceneTextureParameters;
-	FShaderResourceParameter MiniFontTexture;
-	FSubsurfaceParameters SubsurfaceParameters;
-
-	/** Initialization constructor. */
-	FPostProcessSubsurfaceVisualizePS(const ShaderMetaType::CompiledShaderInitializerType& Initializer)
-		: FGlobalShader(Initializer)
-	{
-		PostprocessParameter.Bind(Initializer.ParameterMap);
-		SceneTextureParameters.Bind(Initializer);
-		MiniFontTexture.Bind(Initializer.ParameterMap, TEXT("MiniFontTexture"));
-		SubsurfaceParameters.Bind(Initializer.ParameterMap);
-	}
-
-	template <typename TRHICmdList>
-	void SetParameters(TRHICmdList& RHICmdList, const FRenderingCompositePassContext& Context)
-	{
-		const FFinalPostProcessSettings& Settings = Context.View.FinalPostProcessSettings;
-		const FPixelShaderRHIParamRef ShaderRHI = GetPixelShader();
-
-		FGlobalShader::SetParameters<FViewUniformShaderParameters>(RHICmdList, ShaderRHI, Context.View.ViewUniformBuffer);
-		PostprocessParameter.SetPS(RHICmdList, ShaderRHI, Context, TStaticSamplerState<SF_Point, AM_Clamp, AM_Clamp, AM_Clamp>::GetRHI());
-		SceneTextureParameters.Set(RHICmdList, ShaderRHI, Context.View.FeatureLevel, ESceneTextureSetupMode::All);
-		SetTextureParameter(RHICmdList, ShaderRHI, MiniFontTexture, GEngine->MiniFontTexture ? GEngine->MiniFontTexture->Resource->TextureRHI : GSystemTextures.WhiteDummy->GetRenderTargetItem().TargetableTexture);
-		SubsurfaceParameters.SetParameters(RHICmdList, ShaderRHI, Context);
-	}
-
-	// FShader interface.
-	virtual bool Serialize(FArchive& Ar) override
-	{
-		bool bShaderHasOutdatedParameters = FGlobalShader::Serialize(Ar);
-		Ar << PostprocessParameter << SceneTextureParameters << MiniFontTexture << SubsurfaceParameters;
-		return bShaderHasOutdatedParameters;
-	}
-
-	static const TCHAR* GetSourceFilename()
-	{
-		return TEXT("/Engine/Private/PostProcessSubsurface.usf");
-	}
-
-	static const TCHAR* GetFunctionName()
-	{
-		return TEXT("VisualizePS");
-	}
-};
-
-IMPLEMENT_SHADER_TYPE3(FPostProcessSubsurfaceVisualizePS, SF_Pixel);
-
-void SetSubsurfaceVisualizeShader(const FRenderingCompositePassContext& Context)
-{
-	FGraphicsPipelineStateInitializer GraphicsPSOInit;
-	Context.RHICmdList.ApplyCachedRenderTargets(GraphicsPSOInit);
-	GraphicsPSOInit.BlendState = TStaticBlendState<>::GetRHI();
-	GraphicsPSOInit.RasterizerState = TStaticRasterizerState<>::GetRHI();
-	GraphicsPSOInit.DepthStencilState = TStaticDepthStencilState<false, CF_Always>::GetRHI();
-
-	TShaderMapRef<FPostProcessVS> VertexShader(Context.GetShaderMap());
-	TShaderMapRef<FPostProcessSubsurfaceVisualizePS> PixelShader(Context.GetShaderMap());
-
-	GraphicsPSOInit.BoundShaderState.VertexDeclarationRHI = GFilterVertexDeclaration.VertexDeclarationRHI;
-	GraphicsPSOInit.BoundShaderState.VertexShaderRHI = GETSAFERHISHADER_VERTEX(*VertexShader);
-	GraphicsPSOInit.BoundShaderState.PixelShaderRHI = GETSAFERHISHADER_PIXEL(*PixelShader);
-	GraphicsPSOInit.PrimitiveType = PT_TriangleList;
-
-	SetGraphicsPipelineState(Context.RHICmdList, GraphicsPSOInit);
-
-	PixelShader->SetParameters(Context.RHICmdList, Context);
-	VertexShader->SetParameters(Context);
+	return ProfileTextureTarget->GetRenderTargetItem().ShaderResourceTexture;
 }
 
-FRCPassPostProcessSubsurfaceVisualize::FRCPassPostProcessSubsurfaceVisualize(FRHICommandList& RHICmdList)
+// Returns a half-scaled size rounded to an even multiple of two (but clamped to 1x1 minimum).
+FIntPoint GetHalfSize(FIntPoint Size)
 {
-	// we need the GBuffer, we release it Process()
-	FSceneRenderTargets::Get(RHICmdList).AdjustGBufferRefCount(RHICmdList, 1);
+	Size = FIntPoint::DivideAndRoundUp(Size, 2);
+	Size.X = FMath::Max(1, Size.X);
+	Size.Y = FMath::Max(1, Size.Y);
+	return Size;
 }
 
-void FRCPassPostProcessSubsurfaceVisualize::Process(FRenderingCompositePassContext& Context)
+// Returns a half-scaled rect, with the max rounded to the nearest multiple of two.
+FIntRect GetHalfRect(FIntRect Rect)
 {
-	SCOPED_DRAW_EVENT(Context.RHICmdList, SubsurfaceVisualize);
-
-	const FPooledRenderTargetDesc* InputDesc = GetInputDesc(ePId_Input0);
-
-	if(!InputDesc)
-	{
-		// input is not hooked up correctly
-		return;
-	}
-
-	const FViewInfo& View = Context.View;
-	const FSceneViewFamily& ViewFamily = *(View.Family);
-
-	FIntPoint SrcSize = InputDesc->Extent;
-	FIntPoint DestSize = PassOutputs[0].RenderTargetDesc.Extent;
-
-	FSceneRenderTargets& SceneContext = FSceneRenderTargets::Get(Context.RHICmdList);
-	// e.g. 4 means the input texture is 4x smaller than the buffer size
-	uint32 ScaleFactor = SceneContext.GetBufferSizeXY().X / SrcSize.X;
-
-	FIntRect SrcRect = View.ViewRect / ScaleFactor;
-	FIntRect DestRect = SrcRect;
-
-	const FSceneRenderTargetItem& DestRenderTarget = PassOutputs[0].RequestSurface(Context);
-
-	// Set the view family's render target/viewport.
-	// #todo-renderpasses Use Clear? Load maintains previous behavior.
-	FRHIRenderPassInfo RPInfo(DestRenderTarget.TargetableTexture, ERenderTargetActions::Load_Store);
-	Context.RHICmdList.BeginRenderPass(RPInfo, TEXT("SubsurfaceVisualize"));
-	{
-		// is optimized away if possible (RT size=view size, )
-		DrawClearQuad(Context.RHICmdList, true, FLinearColor::Black, false, 0, false, 0, PassOutputs[0].RenderTargetDesc.Extent, DestRect);
-
-		Context.SetViewportAndCallRHI(0, 0, 0.0f, DestSize.X, DestSize.Y, 1.0f);
-
-		SetSubsurfaceVisualizeShader(Context);
-
-		// Draw a quad mapping scene color to the view's render target
-		TShaderMapRef<FPostProcessVS> VertexShader(Context.GetShaderMap());
-		DrawRectangle(
-			Context.RHICmdList,
-			DestRect.Min.X, DestRect.Min.Y,
-			DestRect.Width(), DestRect.Height(),
-			SrcRect.Min.X, SrcRect.Min.Y,
-			SrcRect.Width(), SrcRect.Height(),
-			DestSize,
-			SrcSize,
-			*VertexShader,
-			EDRF_UseTriangleOptimization);
-	}
-	Context.RHICmdList.EndRenderPass();
-
-	{
-		FRenderTargetTemp TempRenderTarget(View, (const FTexture2DRHIRef&)DestRenderTarget.TargetableTexture);
-		FCanvas Canvas(&TempRenderTarget, NULL, ViewFamily.CurrentRealTime, ViewFamily.CurrentWorldTime, ViewFamily.DeltaWorldTime, Context.GetFeatureLevel());
-
-		float X = 30;
-		float Y = 28;
-		const float YStep = 14;
-
-		FString Line;
-
-		Line = FString::Printf(TEXT("Visualize Screen Space Subsurface Scattering"));
-		Canvas.DrawShadowedString(X, Y += YStep, *Line, GetStatsFont(), FLinearColor(1, 1, 1));
-
-		Y += YStep;
-
-		uint32 Index = 0;
-		while (GSubsurfaceProfileTextureObject.GetEntryString(Index++, Line))
-		{
-			Canvas.DrawShadowedString(X, Y += YStep, *Line, GetStatsFont(), FLinearColor(1, 1, 1));
-		}
-
-		Canvas.Flush_RenderThread(Context.RHICmdList);
-	}
-
-	Context.RHICmdList.CopyToResolveTarget(DestRenderTarget.TargetableTexture, DestRenderTarget.ShaderResourceTexture, FResolveParams());
-	
-
-	// we no longer need the GBuffer
-	SceneContext.AdjustGBufferRefCount(Context.RHICmdList, -1);
+	return FIntRect(Rect.Min / 2, GetHalfSize(Rect.Max));
 }
 
-FPooledRenderTargetDesc FRCPassPostProcessSubsurfaceVisualize::ComputeOutputDesc(EPassOutputId InPassOutputId) const
-{
-	FPooledRenderTargetDesc Ret = FSceneRenderTargets::Get_FrameConstantsOnly().GetSceneColor()->GetDesc();
-	Ret.Flags &= ~(TexCreate_FastVRAM | TexCreate_Transient);
-	Ret.Reset();
-	Ret.DebugName = TEXT("SubsurfaceVisualize");
-	// alpha is used to store depth and renormalize (alpha==0 means there is no subsurface scattering)
-	Ret.Format = PF_FloatRGBA;
+// Set of common shader parameters shared by all subsurface shaders.
+BEGIN_SHADER_PARAMETER_STRUCT(FSubsurfaceParameters, )
+	SHADER_PARAMETER(FVector4, SubsurfaceParams)
+	SHADER_PARAMETER_STRUCT_INCLUDE(FScreenPassCommonParameters, ScreenPassCommonParameters)
+	SHADER_PARAMETER_TEXTURE(Texture2D, SSProfilesTexture)
+END_SHADER_PARAMETER_STRUCT()
 
-	return Ret;
+FSubsurfaceParameters GetSubsurfaceCommonParameters(FRHICommandListImmediate& RHICmdList, FScreenPassContextRef Context)
+{
+	const float DistanceToProjectionWindow = Context->View.ViewMatrices.GetProjectionMatrix().M[0][0];
+	const float SSSScaleZ = DistanceToProjectionWindow * GetSubsurfaceRadiusScale();
+	const float SSSScaleX = SSSScaleZ / SUBSURFACE_KERNEL_SIZE * 0.5f;
+
+	FSubsurfaceParameters Parameters;
+	Parameters.ScreenPassCommonParameters = Context->ScreenPassCommonParameters;
+	Parameters.SubsurfaceParams = FVector4(SSSScaleX, SSSScaleZ, 0, 0);
+	Parameters.SSProfilesTexture = GetSubsurfaceProfileTexture(RHICmdList);
+	return Parameters;
 }
 
-
-// ---------------------------------------------
-
-/**
- * Encapsulates the post processing subsurface scattering pixel shader.
- * @param HalfRes 0:to full res, 1:to half res
- */
-template <uint32 HalfRes, uint32 Checkerboard>
-class FPostProcessSubsurfaceSetupPS : public FGlobalShader
+// Base class for a subsurface shader.
+class FSubsurfaceShader : public FGlobalShader
 {
-	DECLARE_SHADER_TYPE(FPostProcessSubsurfaceSetupPS , Global);
-
+public:
 	static bool ShouldCompilePermutation(const FGlobalShaderPermutationParameters& Parameters)
 	{
 		return IsFeatureLevelSupported(Parameters.Platform, ERHIFeatureLevel::SM4);
@@ -378,706 +185,521 @@ class FPostProcessSubsurfaceSetupPS : public FGlobalShader
 	static void ModifyCompilationEnvironment(const FGlobalShaderPermutationParameters& Parameters, FShaderCompilerEnvironment& OutEnvironment)
 	{
 		FGlobalShader::ModifyCompilationEnvironment(Parameters, OutEnvironment);
-		OutEnvironment.SetDefine(TEXT("HALF_RES"), HalfRes);
 		OutEnvironment.SetDefine(TEXT("SUBSURFACE_RADIUS_SCALE"), SUBSURFACE_RADIUS_SCALE);
 		OutEnvironment.SetDefine(TEXT("SUBSURFACE_KERNEL_SIZE"), SUBSURFACE_KERNEL_SIZE);
-		OutEnvironment.SetDefine(TEXT("SUBSURFACE_PROFILE_CHECKERBOARD"), Checkerboard);
 	}
 
-	/** Default constructor. */
-	FPostProcessSubsurfaceSetupPS () {}
-
-public:
-	FPostProcessPassParameters PostprocessParameter;
-	FSceneTextureShaderParameters SceneTextureParameters;
-	FShaderResourceParameter MiniFontTexture;
-	FSubsurfaceParameters SubsurfaceParameters;
-
-	/** Initialization constructor. */
-	FPostProcessSubsurfaceSetupPS(const ShaderMetaType::CompiledShaderInitializerType& Initializer)
+	FSubsurfaceShader() = default;
+	FSubsurfaceShader(const ShaderMetaType::CompiledShaderInitializerType& Initializer)
 		: FGlobalShader(Initializer)
-	{
-		PostprocessParameter.Bind(Initializer.ParameterMap);
-		SceneTextureParameters.Bind(Initializer);
-		SubsurfaceParameters.Bind(Initializer.ParameterMap);
-	}
-
-	void SetParameters(const FRenderingCompositePassContext& Context)
-	{
-		const FFinalPostProcessSettings& Settings = Context.View.FinalPostProcessSettings;
-		const FPixelShaderRHIParamRef ShaderRHI = GetPixelShader();
-
-		FGlobalShader::SetParameters<FViewUniformShaderParameters>(Context.RHICmdList, ShaderRHI, Context.View.ViewUniformBuffer);
-		PostprocessParameter.SetPS(Context.RHICmdList, ShaderRHI, Context, TStaticSamplerState<SF_Point, AM_Clamp, AM_Clamp, AM_Clamp>::GetRHI());
-		SceneTextureParameters.Set(Context.RHICmdList, ShaderRHI, Context.View.FeatureLevel, ESceneTextureSetupMode::All);
-		SubsurfaceParameters.SetParameters(Context.RHICmdList, ShaderRHI, Context);
-	}
-
-	// FShader interface.
-	virtual bool Serialize(FArchive& Ar) override
-	{
-		bool bShaderHasOutdatedParameters = FGlobalShader::Serialize(Ar);
-		Ar << PostprocessParameter << SceneTextureParameters << SubsurfaceParameters;
-		return bShaderHasOutdatedParameters;
-	}
-
-	static const TCHAR* GetSourceFilename()
-	{
-		return TEXT("/Engine/Private/PostProcessSubsurface.usf");
-	}
-
-	static const TCHAR* GetFunctionName()
-	{
-		return TEXT("SetupPS");
-	}
+	{}
 };
 
-// #define avoids a lot of code duplication
-#define VARIATION1(A)	VARIATION2(A,0)	VARIATION2(A,1)
-#define VARIATION2(A,B) typedef FPostProcessSubsurfaceSetupPS<A,B> FPostProcessSubsurfaceSetupPS##A##B; \
-	IMPLEMENT_SHADER_TYPE2(FPostProcessSubsurfaceSetupPS##A##B, SF_Pixel);
-	VARIATION1(0) VARIATION1(1)
-#undef VARIATION1
-#undef VARIATION2
-
-
-template <uint32 HalfRes, uint32 Checkerboard>
-void SetSubsurfaceSetupShader(const FRenderingCompositePassContext& Context)
+// Encapsulates the post processing subsurface scattering pixel shader.
+class FSubsurfaceVisualizePS : public FSubsurfaceShader
 {
-	FGraphicsPipelineStateInitializer GraphicsPSOInit;
-	Context.RHICmdList.ApplyCachedRenderTargets(GraphicsPSOInit);
-	GraphicsPSOInit.BlendState = TStaticBlendState<>::GetRHI();
-	GraphicsPSOInit.RasterizerState = TStaticRasterizerState<>::GetRHI();
-	GraphicsPSOInit.DepthStencilState = TStaticDepthStencilState<false, CF_Always>::GetRHI();
-
-	TShaderMapRef<FPostProcessVS> VertexShader(Context.GetShaderMap());
-	TShaderMapRef<FPostProcessSubsurfaceSetupPS<HalfRes, Checkerboard> > PixelShader(Context.GetShaderMap());
-
-	GraphicsPSOInit.BoundShaderState.VertexDeclarationRHI = GFilterVertexDeclaration.VertexDeclarationRHI;
-	GraphicsPSOInit.BoundShaderState.VertexShaderRHI = GETSAFERHISHADER_VERTEX(*VertexShader);
-	GraphicsPSOInit.BoundShaderState.PixelShaderRHI = GETSAFERHISHADER_PIXEL(*PixelShader);
-	GraphicsPSOInit.PrimitiveType = PT_TriangleList;
-
-	SetGraphicsPipelineState(Context.RHICmdList, GraphicsPSOInit);
-
-	PixelShader->SetParameters(Context);
-	VertexShader->SetParameters(Context);
-}
+public:
+	DECLARE_GLOBAL_SHADER(FSubsurfaceVisualizePS);
+	SHADER_USE_PARAMETER_STRUCT(FSubsurfaceVisualizePS, FSubsurfaceShader);
+
+	BEGIN_SHADER_PARAMETER_STRUCT(FParameters, )
+		SHADER_PARAMETER_STRUCT_INCLUDE(FSubsurfaceParameters, Subsurface)
+		SHADER_PARAMETER_STRUCT(FScreenPassInput, SubsurfaceInput0)
+		SHADER_PARAMETER_TEXTURE(Texture2D, MiniFontTexture)
+		RENDER_TARGET_BINDING_SLOTS()
+	END_SHADER_PARAMETER_STRUCT()
+};
 
-// --------------------------------------
+IMPLEMENT_GLOBAL_SHADER(FSubsurfaceVisualizePS, "/Engine/Private/PostProcessSubsurface.usf", "VisualizePS", SF_Pixel);
 
-FRCPassPostProcessSubsurfaceSetup::FRCPassPostProcessSubsurfaceSetup(FViewInfo& View, bool bInHalfRes)
-	: ViewRect(View.ViewRect)
-	, bHalfRes(bInHalfRes)
+// Encapsulates the post processing subsurface scattering pixel shader.
+class FSubsurfaceSetupPS : public FSubsurfaceShader
 {
-}
+public:
+	DECLARE_GLOBAL_SHADER(FSubsurfaceSetupPS);
+	SHADER_USE_PARAMETER_STRUCT(FSubsurfaceSetupPS, FSubsurfaceShader)
+
+	BEGIN_SHADER_PARAMETER_STRUCT(FParameters, )
+		SHADER_PARAMETER_STRUCT_INCLUDE(FSubsurfaceParameters, Subsurface)
+		SHADER_PARAMETER_STRUCT(FScreenPassInput, SubsurfaceInput0)
+		RENDER_TARGET_BINDING_SLOTS()
+	END_SHADER_PARAMETER_STRUCT()
+
+	class FDimensionHalfRes : SHADER_PERMUTATION_BOOL("SUBSURFACE_HALF_RES");
+	class FDimensionCheckerboard : SHADER_PERMUTATION_BOOL("SUBSURFACE_PROFILE_CHECKERBOARD");
+	using FPermutationDomain = TShaderPermutationDomain<FDimensionHalfRes, FDimensionCheckerboard>;
+};
 
-void FRCPassPostProcessSubsurfaceSetup::Process(FRenderingCompositePassContext& Context)
+IMPLEMENT_GLOBAL_SHADER(FSubsurfaceSetupPS, "/Engine/Private/PostProcessSubsurface.usf", "SetupPS", SF_Pixel);
+
+// Shader for the SSS separable blur.
+class FSubsurfacePS : public FSubsurfaceShader
 {
-	SCOPED_DRAW_EVENT(Context.RHICmdList, SubsurfaceSetup);
+public:
+	DECLARE_GLOBAL_SHADER(FSubsurfacePS);
+	SHADER_USE_PARAMETER_STRUCT(FSubsurfacePS, FSubsurfaceShader);
 
-	const FPooledRenderTargetDesc* InputDesc = GetInputDesc(ePId_Input0);
+	BEGIN_SHADER_PARAMETER_STRUCT(FParameters, )
+		SHADER_PARAMETER_STRUCT_INCLUDE(FSubsurfaceParameters, Subsurface)
+		SHADER_PARAMETER_STRUCT(FScreenPassInput, SubsurfaceInput0)
+		RENDER_TARGET_BINDING_SLOTS()
+	END_SHADER_PARAMETER_STRUCT()
 
-	if(!InputDesc)
+	// Direction of the 1D separable filter.
+	enum class EDirection : uint32
 	{
-		// input is not hooked up correctly
-		return;
-	}
-
-	FSceneRenderTargets& SceneContext = FSceneRenderTargets::Get(Context.RHICmdList);
-	const bool bCheckerboard = FRCPassPostProcessSubsurface::RequiresCheckerboardSubsurfaceRendering( SceneContext.GetSceneColorFormat() );
-	const FViewInfo& View = Context.View;
-	const FSceneViewFamily& ViewFamily = *(View.Family);
-
-	FIntPoint SrcSize = InputDesc->Extent;
-	FIntPoint DestSize = PassOutputs[0].RenderTargetDesc.Extent;
+		Horizontal,
+		Vertical,
+		MAX
+	};
 
-	int32 ScaleFactor = bHalfRes ? 2 : 1;
-
-	FIntRect SrcRect = View.ViewRect;
-
-	FIntRect DestRect;
-	DestRect.Min = View.ViewRect.Min / ScaleFactor;
-	ensure(DestRect.Min * ScaleFactor == View.ViewRect.Min);
-	DestRect.Max.X = FMath::DivideAndRoundUp(View.ViewRect.Max.X, ScaleFactor);
-	DestRect.Max.Y = FMath::DivideAndRoundUp(View.ViewRect.Max.Y, ScaleFactor);
-
-	const FSceneRenderTargetItem& DestRenderTarget = PassOutputs[0].RequestSurface(Context);
+	// Controls the quality (number of samples) of the blur kernel.
+	enum class EQuality : uint32
+	{
+		Low,
+		Medium,
+		High,
+		MAX
+	};
 
-	ERenderTargetLoadAction LoadAction = Context.GetLoadActionForRenderTarget(DestRenderTarget);
+	class FDimensionDirection : SHADER_PERMUTATION_ENUM_CLASS("SUBSURFACE_DIRECTION", EDirection);
+	class FDimensionQuality : SHADER_PERMUTATION_ENUM_CLASS("SUBSURFACE_QUALITY", EQuality);
+	class FDimensionManuallyClampUV : SHADER_PERMUTATION_BOOL("SUBSURFACE_MANUALLY_CLAMP_UV");
+	using FPermutationDomain = TShaderPermutationDomain<FDimensionDirection, FDimensionQuality, FDimensionManuallyClampUV>;
 
-	FRHIRenderPassInfo RPInfo(DestRenderTarget.TargetableTexture, MakeRenderTargetActions(LoadAction, ERenderTargetStoreAction::EStore), DestRenderTarget.ShaderResourceTexture);
-	Context.RHICmdList.BeginRenderPass(RPInfo, TEXT("SubsurfaceSetup"));
+	// Returns the sampler state based on the requested SSS filter CVar setting.
+	static FSamplerStateRHIParamRef GetSamplerState()
 	{
-		Context.SetViewportAndCallRHI(DestRect);
-
-		if (bHalfRes)
+		if (CVarSSSFilter.GetValueOnRenderThread())
 		{
-			if (bCheckerboard)
-			{
-				SetSubsurfaceSetupShader<1, 1>(Context);
-			}
-			else
-			{
-				SetSubsurfaceSetupShader<1, 0>(Context);
-			}
+			return TStaticSamplerState<SF_Bilinear, AM_Border, AM_Border, AM_Border>::GetRHI();
 		}
 		else
 		{
-			if (bCheckerboard)
-			{
-				SetSubsurfaceSetupShader<0, 1>(Context);
-			}
-			else
-			{
-				SetSubsurfaceSetupShader<0, 0>(Context);
-			}
+			return TStaticSamplerState<SF_Point, AM_Border, AM_Border, AM_Border>::GetRHI();
 		}
-
-		// Draw a quad mapping scene color to the view's render target
-		TShaderMapRef<FPostProcessVS> VertexShader(Context.GetShaderMap());
-
-		// Align up, so downsample always pickups correct pixel from the full-res buffer.
-		FIntPoint TargetSize = SrcRect.Size();
-		TargetSize.X = Align(TargetSize.X, ScaleFactor);
-		TargetSize.Y = Align(TargetSize.Y, ScaleFactor);
-
-		DrawPostProcessPass(
-			Context.RHICmdList,
-			0, 0,
-			SrcRect.Width(), SrcRect.Height(),
-			SrcRect.Min.X, SrcRect.Min.Y,
-			SrcRect.Width(), SrcRect.Height(),
-			TargetSize,
-			SrcSize,
-			*VertexShader,
-			View.StereoPass,
-			Context.HasHmdMesh(),
-			EDRF_UseTriangleOptimization);
-
 	}
-	Context.RHICmdList.EndRenderPass();
 
-	FTextureRHIParamRef TexturesToReadable[1];
-	TexturesToReadable[0] = DestRenderTarget.ShaderResourceTexture;
-	Context.RHICmdList.TransitionResources(EResourceTransitionAccess::EReadable, TexturesToReadable, ARRAY_COUNT(TexturesToReadable));
-}
-
-FPooledRenderTargetDesc FRCPassPostProcessSubsurfaceSetup::ComputeOutputDesc(EPassOutputId InPassOutputId) const
-{
-	FPooledRenderTargetDesc Ret = FSceneRenderTargets::Get_FrameConstantsOnly().GetSceneColor()->GetDesc();
-	Ret.Flags &= ~(TexCreate_FastVRAM | TexCreate_Transient);
-	Ret.Reset();
-	Ret.DebugName = TEXT("SubsurfaceSetup");
-	// alpha is used to store depth and renormalize (alpha==0 means there is no subsurface scattering)
-	Ret.Format = PF_FloatRGBA;
-
-	if(bHalfRes)
+	// Returns the SSS quality level requested by the SSS SampleSet CVar setting.
+	static EQuality GetQuality()
 	{
-		Ret.Extent = FIntPoint::DivideAndRoundUp(Ret.Extent, 2);
-		Ret.Extent.X = FMath::Max(1, Ret.Extent.X);
-		Ret.Extent.Y = FMath::Max(1, Ret.Extent.Y);
+		return static_cast<FSubsurfacePS::EQuality>(
+			FMath::Clamp(
+				CVarSSSSampleSet.GetValueOnRenderThread(),
+				static_cast<int32>(FSubsurfacePS::EQuality::Low),
+				static_cast<int32>(FSubsurfacePS::EQuality::High)));
 	}
-	
-	return Ret;
-}
+};
 
+IMPLEMENT_GLOBAL_SHADER(FSubsurfacePS, "/Engine/Private/PostProcessSubsurface.usf", "MainPS", SF_Pixel);
 
-/** Encapsulates the post processing subsurface pixel shader. */
-// @param Direction 0: horizontal, 1:vertical
-// @param SampleSet 0:low, 1:med, 2:high
-template <uint32 Direction, uint32 SampleSet, uint32 ManuallyClampUV>
-class TPostProcessSubsurfacePS : public FGlobalShader
+// Encapsulates the post processing subsurface recombine pixel shader.
+class FSubsurfaceRecombinePS : public FSubsurfaceShader
 {
-	DECLARE_SHADER_TYPE(TPostProcessSubsurfacePS, Global);
+	DECLARE_GLOBAL_SHADER(FSubsurfaceRecombinePS);
+	SHADER_USE_PARAMETER_STRUCT(FSubsurfaceRecombinePS, FSubsurfaceShader);
 
-	static bool ShouldCompilePermutation(const FGlobalShaderPermutationParameters& Parameters)
-	{
-		return IsFeatureLevelSupported(Parameters.Platform, ERHIFeatureLevel::SM4);
-	}
+	BEGIN_SHADER_PARAMETER_STRUCT(FParameters, )
+		SHADER_PARAMETER_STRUCT_INCLUDE(FSubsurfaceParameters, Subsurface)
+		SHADER_PARAMETER_STRUCT(FScreenPassInput, SubsurfaceInput0)
+		SHADER_PARAMETER_STRUCT(FScreenPassInput, SubsurfaceInput1)
+		RENDER_TARGET_BINDING_SLOTS()
+	END_SHADER_PARAMETER_STRUCT();
 
-	static void ModifyCompilationEnvironment(const FGlobalShaderPermutationParameters& Parameters, FShaderCompilerEnvironment& OutEnvironment)
+	// Controls the quality of lighting reconstruction.
+	enum class EQuality : uint32
 	{
-		FGlobalShader::ModifyCompilationEnvironment(Parameters,OutEnvironment);
-		OutEnvironment.SetDefine(TEXT("SSS_DIRECTION"), Direction);
-		OutEnvironment.SetDefine(TEXT("SSS_SAMPLESET"), SampleSet);
-		OutEnvironment.SetDefine(TEXT("SUBSURFACE_RADIUS_SCALE"), SUBSURFACE_RADIUS_SCALE);
-		OutEnvironment.SetDefine(TEXT("SUBSURFACE_KERNEL_SIZE"), SUBSURFACE_KERNEL_SIZE);
-		OutEnvironment.SetDefine(TEXT("MANUALLY_CLAMP_UV"), ManuallyClampUV);
-	}
+		Low,
+		High,
+		MAX
+	};
 
-	/** Default constructor. */
-	TPostProcessSubsurfacePS() {}
+	class FDimensionMode : SHADER_PERMUTATION_ENUM_CLASS("SUBSURFACE_RECOMBINE_MODE", ESubsurfaceMode);
+	class FDimensionQuality : SHADER_PERMUTATION_ENUM_CLASS("SUBSURFACE_RECOMBINE_QUALITY", EQuality);
+	class FDimensionCheckerboard : SHADER_PERMUTATION_BOOL("SUBSURFACE_PROFILE_CHECKERBOARD");
+	using FPermutationDomain = TShaderPermutationDomain<FDimensionMode, FDimensionQuality, FDimensionCheckerboard>;
 
-public:
-	FPostProcessPassParameters PostprocessParameter;
-	FSceneTextureShaderParameters SceneTextureParameters;
-	FSubsurfaceParameters SubsurfaceParameters;
-
-	/** Initialization constructor. */
-	TPostProcessSubsurfacePS(const ShaderMetaType::CompiledShaderInitializerType& Initializer)
-		: FGlobalShader(Initializer)
+	// Returns the Recombine quality level requested by the SSS Quality CVar setting.
+	static EQuality GetQuality(const FViewInfo& View)
 	{
-		PostprocessParameter.Bind(Initializer.ParameterMap);
-		SceneTextureParameters.Bind(Initializer);
-		SubsurfaceParameters.Bind(Initializer.ParameterMap);
-	}
-
-	// FShader interface.
-	virtual bool Serialize(FArchive& Ar) override
-	{
-		bool bShaderHasOutdatedParameters = FGlobalShader::Serialize(Ar);
-		Ar << PostprocessParameter << SceneTextureParameters << SubsurfaceParameters;
-		return bShaderHasOutdatedParameters;
-	}
-
-	void SetParameters(const FRenderingCompositePassContext& Context)
-	{
-		const FPixelShaderRHIParamRef ShaderRHI = GetPixelShader();
+		const uint32 QualityCVar = CVarSSSQuality.GetValueOnRenderThread();
 
-		FGlobalShader::SetParameters<FViewUniformShaderParameters>(Context.RHICmdList, ShaderRHI, Context.View.ViewUniformBuffer);
-		SceneTextureParameters.Set(Context.RHICmdList, ShaderRHI, Context.View.FeatureLevel, ESceneTextureSetupMode::All);
+		// Quality is forced to high when the CVar is set to 'auto' and TAA is NOT enabled.
+		// TAA improves quality through temporal filtering, making it less necessary to use
+		// high quality mode.
+		const bool bUseHighQuality = (QualityCVar == -1 && View.AntiAliasingMethod != AAM_TemporalAA);
 
-		if(CVarSSSFilter.GetValueOnRenderThread())
+		if (QualityCVar == 1 || bUseHighQuality)
 		{
-			PostprocessParameter.SetPS(Context.RHICmdList, ShaderRHI, Context, TStaticSamplerState<SF_Bilinear, AM_Border, AM_Border, AM_Border>::GetRHI());
+			return EQuality::High;
 		}
 		else
 		{
-			PostprocessParameter.SetPS(Context.RHICmdList, ShaderRHI, Context, TStaticSamplerState<SF_Point,AM_Border,AM_Border,AM_Border>::GetRHI());
+			return EQuality::Low;
 		}
-		SubsurfaceParameters.SetParameters(Context.RHICmdList, ShaderRHI, Context);
 	}
+};
 
-	static const TCHAR* GetSourceFilename()
-	{
-		return TEXT("/Engine/Private/PostProcessSubsurface.usf");
-	}
+IMPLEMENT_GLOBAL_SHADER(FSubsurfaceRecombinePS, "/Engine/Private/PostProcessSubsurface.usf", "SubsurfaceRecombinePS", SF_Pixel);
 
-	static const TCHAR* GetFunctionName()
+// Encapsulates a simple copy pixel shader.
+class FSubsurfaceViewportCopyPS : public FSubsurfaceShader
+{
+	DECLARE_GLOBAL_SHADER(FSubsurfaceViewportCopyPS);
+	SHADER_USE_PARAMETER_STRUCT(FSubsurfaceViewportCopyPS, FSubsurfaceShader);
+
+	BEGIN_SHADER_PARAMETER_STRUCT(FParameters, )
+		SHADER_PARAMETER_STRUCT(FScreenPassInput, SubsurfaceInput0)
+		RENDER_TARGET_BINDING_SLOTS()
+	END_SHADER_PARAMETER_STRUCT();
+
+	static bool ShouldCompilePermutation(const FGlobalShaderPermutationParameters& Parameters)
 	{
-		return TEXT("MainPS");
+		return IsFeatureLevelSupported(Parameters.Platform, ERHIFeatureLevel::SM4);
 	}
 };
 
-// #define avoids a lot of code duplication
-#define VARIATION1(A)		VARIATION2(A,0)			VARIATION2(A,1)			VARIATION2(A,2)
-#define VARIATION2(A, B)	VARIATION3(A,B,0)		VARIATION3(A,B,1)
-#define VARIATION3(A, B, C) typedef TPostProcessSubsurfacePS<A, B, C> TPostProcessSubsurfacePS##A##B##C; \
-	IMPLEMENT_SHADER_TYPE2(TPostProcessSubsurfacePS##A##B##C, SF_Pixel);
-	VARIATION1(0) VARIATION1(1) VARIATION1(2)
-#undef VARIATION1
-#undef VARIATION2
-#undef VARIATION3
-
+IMPLEMENT_GLOBAL_SHADER(FSubsurfaceViewportCopyPS, "/Engine/Private/PostProcessSubsurface.usf", "SubsurfaceViewportCopyPS", SF_Pixel);
 
-FRCPassPostProcessSubsurface::FRCPassPostProcessSubsurface(uint32 InDirection, bool bInHalfRes)
-	: Direction(InDirection) 
-	, bHalfRes(bInHalfRes)
+bool IsSubsurfaceEnabled()
 {
-	check(InDirection < 2);
+	const bool bEnabled = CVarSubsurfaceScattering.GetValueOnAnyThread() != 0;
+	const bool bHasScale = CVarSSSScale.GetValueOnAnyThread() > 0.0f;
+	return (bEnabled && bHasScale);
 }
 
-template <uint32 Direction, uint32 SampleSet, uint32 ManuallyClampUV>
-void SetSubsurfaceShader(const FRenderingCompositePassContext& Context, TShaderMapRef<FPostProcessVS> &VertexShader)
+bool IsSubsurfaceRequiredForView(const FViewInfo& View)
 {
-	FGraphicsPipelineStateInitializer GraphicsPSOInit;
-	Context.RHICmdList.ApplyCachedRenderTargets(GraphicsPSOInit);
-	GraphicsPSOInit.BlendState = TStaticBlendState<>::GetRHI();
-	GraphicsPSOInit.RasterizerState = TStaticRasterizerState<>::GetRHI();
-	GraphicsPSOInit.DepthStencilState = TStaticDepthStencilState<false, CF_Always>::GetRHI();
-
-	TShaderMapRef<TPostProcessSubsurfacePS<Direction, SampleSet, ManuallyClampUV> > PixelShader(Context.GetShaderMap());
-
-	GraphicsPSOInit.BoundShaderState.VertexDeclarationRHI = GFilterVertexDeclaration.VertexDeclarationRHI;
-	GraphicsPSOInit.BoundShaderState.VertexShaderRHI = GETSAFERHISHADER_VERTEX(*VertexShader);
-	GraphicsPSOInit.BoundShaderState.PixelShaderRHI = GETSAFERHISHADER_PIXEL(*PixelShader);
-	GraphicsPSOInit.PrimitiveType = PT_TriangleList;
-
-	SetGraphicsPipelineState(Context.RHICmdList, GraphicsPSOInit);
-
-	PixelShader->SetParameters(Context);
-	VertexShader->SetParameters(Context);
+	const bool bSimpleDynamicLighting = IsAnyForwardShadingEnabled(View.GetShaderPlatform());
+	const bool bSubsurfaceEnabled = IsSubsurfaceEnabled();
+	const bool bViewHasSubsurfaceMaterials = ((View.ShadingModelMaskInView & GetUseSubsurfaceProfileShadingModelMask()) != 0);
+	return (bSubsurfaceEnabled && bViewHasSubsurfaceMaterials && !bSimpleDynamicLighting);
 }
 
-// 0:horizontal, 1: vertical
-template <uint32 Direction>
-void SetSubsurfaceShaderSampleSet(const FRenderingCompositePassContext& Context, TShaderMapRef<FPostProcessVS> &VertexShader, uint32 SampleSet)
+bool IsSubsurfaceCheckerboardFormat(EPixelFormat SceneColorFormat)
 {
-	FIntRect SrcRect = Context.View.ViewRect;
-	FIntPoint SrcSize = FSceneRenderTargets::Get(Context.RHICmdList).GetBufferSizeXY();
-
-	if (SrcRect.Min == FIntPoint::ZeroValue && SrcRect.Max == SrcSize)
+	int CVarValue = CVarSSSCheckerboard.GetValueOnRenderThread();
+	if (CVarValue == 0)
 	{
-		switch (SampleSet)
-		{
-		case 0: SetSubsurfaceShader<Direction, 0, 0>(Context, VertexShader); break;
-		case 1: SetSubsurfaceShader<Direction, 1, 0>(Context, VertexShader); break;
-		case 2: SetSubsurfaceShader<Direction, 2, 0>(Context, VertexShader); break;
-		default: check(0);
-		}
+		return false;
 	}
-	else
+	else if (CVarValue == 1)
+	{
+		return true;
+	}
+	else if (CVarValue == 2)
 	{
-		switch (SampleSet)
+		switch (SceneColorFormat)
 		{
-		case 0: SetSubsurfaceShader<Direction, 0, 1>(Context, VertexShader); break;
-		case 1: SetSubsurfaceShader<Direction, 1, 1>(Context, VertexShader); break;
-		case 2: SetSubsurfaceShader<Direction, 2, 1>(Context, VertexShader); break;
-		default: check(0);
+		case PF_A32B32G32R32F:
+		case PF_FloatRGBA:
+			return false;
+		default:
+			return true;
 		}
 	}
+	return true;
 }
 
-void FRCPassPostProcessSubsurface::Process(FRenderingCompositePassContext& Context)
+FRDGTextureRef ComputeSubsurface(
+	FRDGBuilder& GraphBuilder,
+	FScreenPassContextRef Context,
+	FRDGTextureRef SceneTexture)
 {
-	const FPooledRenderTargetDesc* InputDesc = GetInputDesc(ePId_Input0);
+	check(Context);
+	check(SceneTexture);
 
-	check(InputDesc);
+	const FPooledRenderTargetDesc& SceneTextureDesc = SceneTexture->Desc;
+	const ESubsurfaceMode SubsurfaceMode = GetSubsurfaceModeForView(Context->View);
+	const bool bHalfRes = (SubsurfaceMode == ESubsurfaceMode::HalfRes);
 
-	{
-		const IPooledRenderTarget* PooledRT = GetSubsufaceProfileTexture_RT(Context.RHICmdList);
+	// Viewport rect mapped onto the scene texture. Not necessarily a full screen mapping (e.g. VR).
+	const FIntRect ViewportRectFinal = Context->ViewportRect;
 
-		check(PooledRT);
+	// Viewport for intermediate passes, which may be half resolution depending on the pass settings.
+	const FIntRect ViewportRectIntermediate = bHalfRes ? GetHalfRect(ViewportRectFinal) : ViewportRectFinal;
 
-		// for debugging
-		GVisualizeTexture.SetCheckPoint(Context.RHICmdList, PooledRT);
-	}
-
-	const FViewInfo& View = Context.View;
-	const FSceneViewFamily& ViewFamily = *(View.Family);
-
-	FIntPoint SrcSize = InputDesc->Extent;
-	FIntPoint DestSize = PassOutputs[0].RenderTargetDesc.Extent;
+	// Size of the final scene texture.
+	const FIntPoint TextureSizeFinal = SceneTextureDesc.Extent;
 
-	check(DestSize.X);
-	check(DestSize.Y);
-	check(SrcSize.X);
-	check(SrcSize.Y);
+	// Size of intermediate textures, which may be half resolution depending on pass settings.
+	const FIntPoint TextureSizeIntermediate = bHalfRes ? GetHalfSize(TextureSizeFinal) : TextureSizeFinal;
 
-	int32 ScaleFactor = bHalfRes ? 2 : 1;
+	// Description shared by all intermediate pass textures.
+	const FRDGTextureDesc TextureDescIntermediate = FRDGTextureDesc::Create2DDesc(
+		TextureSizeIntermediate,
+		PF_FloatRGBA,
+		FClearValueBinding(),
+		TexCreate_None,
+		TexCreate_RenderTargetable | TexCreate_ShaderResource,
+		false);
 
-	FIntRect SrcRect;
-	SrcRect.Min = View.ViewRect.Min / ScaleFactor;
-	ensure(SrcRect.Min * ScaleFactor == View.ViewRect.Min);
-	SrcRect.Max.X = FMath::DivideAndRoundUp(View.ViewRect.Max.X, ScaleFactor);
-	SrcRect.Max.Y = FMath::DivideAndRoundUp(View.ViewRect.Max.Y, ScaleFactor);
-	FIntRect DestRect = SrcRect;
+	// Common subsurface parameters shared by all passes.
+	const FSubsurfaceParameters SubsurfaceCommonParameters = GetSubsurfaceCommonParameters(GraphBuilder.RHICmdList, Context);
 
-	TRefCountPtr<IPooledRenderTarget> NewSceneColor;
+	// Texture handles used by subsurface passes.
+	FRDGTextureRef SetupTexture = SceneTexture;
+	FRDGTextureRef SubsurfaceTextureX = nullptr;
+	FRDGTextureRef SubsurfaceTextureY = nullptr;
+	FRDGTextureRef RecombineTexture = nullptr;
 
-	const FSceneRenderTargetItem& DestRenderTarget = PassOutputs[0].RequestSurface(Context);
+	// Sampler handles used by subsurface passes.
+	FSamplerStateRHIParamRef PointClampSampler = TStaticSamplerState<SF_Point, AM_Clamp, AM_Clamp, AM_Clamp>::GetRHI();
+	FSamplerStateRHIParamRef BilinearBorderState = TStaticSamplerState<SF_Bilinear, AM_Border, AM_Border, AM_Border>::GetRHI();
 
-	ERenderTargetLoadAction LoadAction = Context.GetLoadActionForRenderTarget(DestRenderTarget);
+	// Whether to use checkerboard rendering for subsurface passes (predicated on CVar and format).
+	const bool bCheckerboard = IsSubsurfaceCheckerboardFormat(SceneTextureDesc.Format);
 
-	FRHIRenderPassInfo RPInfo(DestRenderTarget.TargetableTexture, MakeRenderTargetActions(LoadAction, ERenderTargetStoreAction::EStore), DestRenderTarget.ShaderResourceTexture);
-	Context.RHICmdList.BeginRenderPass(RPInfo, TEXT("Subsurface"));
+	// When in bypass mode the setup and convolution passes are skipped, but lighting
+	// reconstruction is still performed in the recombine pass.
+	if (SubsurfaceMode != ESubsurfaceMode::Bypass)
 	{
-		Context.SetViewportAndCallRHI(DestRect);
-
-		TShaderMapRef<FPostProcessVS> VertexShader(Context.GetShaderMap());
-
-		SCOPED_DRAW_EVENTF(Context.RHICmdList, SubsurfacePass, TEXT("SubsurfaceDirection#%d"), Direction);
+		SetupTexture = GraphBuilder.CreateTexture(TextureDescIntermediate, TEXT("SubsurfaceSetupTexture"));
 
-		uint32 SampleSet = FMath::Clamp(CVarSSSSampleSet.GetValueOnRenderThread(), 0, 2);
-
-		if (Direction == 0)
+		// Setup pass outputs the diffuse scene color and depth in preparation for the scatter passes.
 		{
-			SetSubsurfaceShaderSampleSet<0>(Context, VertexShader, SampleSet);
+			FSubsurfaceSetupPS::FParameters* PassParameters = GraphBuilder.AllocParameters<FSubsurfaceSetupPS::FParameters>();
+			PassParameters->Subsurface = SubsurfaceCommonParameters;
+			PassParameters->RenderTargets[0] = FRenderTargetBinding(SetupTexture, ERenderTargetLoadAction::ENoAction, ERenderTargetStoreAction::EStore);
+			PassParameters->SubsurfaceInput0 = GetScreenPassInputParameters(SceneTexture, PointClampSampler);
+
+			// Reading from the final target; writing to intermediate target.
+			const FIntRect ViewportRect = ViewportRectIntermediate;
+			const FIntRect TextureRect = ViewportRectFinal;
+			const FIntPoint TextureSize = TextureSizeFinal;
+
+			FSubsurfaceSetupPS::FPermutationDomain PixelShaderPermutationVector;
+			PixelShaderPermutationVector.Set<FSubsurfaceSetupPS::FDimensionHalfRes>(bHalfRes);
+			PixelShaderPermutationVector.Set<FSubsurfaceSetupPS::FDimensionCheckerboard>(bCheckerboard);
+			TShaderMapRef<FSubsurfaceSetupPS> PixelShader(Context->ShaderMap, PixelShaderPermutationVector);
+
+			AddDrawScreenPass(GraphBuilder, RDG_EVENT_NAME("SubsurfaceSetup"), Context, ViewportRect, TextureRect, TextureSize, *PixelShader, PassParameters);
 		}
-		else
-		{
-			SetSubsurfaceShaderSampleSet<1>(Context, VertexShader, SampleSet);
-		}
-
-		DrawPostProcessPass(
-			Context.RHICmdList,
-			0, 0,
-			DestRect.Width(), DestRect.Height(),
-			SrcRect.Min.X, SrcRect.Min.Y,
-			SrcRect.Width(), SrcRect.Height(),
-			DestRect.Size(),
-			SrcSize,
-			*VertexShader,
-			View.StereoPass,
-			Context.HasHmdMesh(),
-			EDRF_UseTriangleOptimization);
-	}
-	Context.RHICmdList.EndRenderPass();
 
-	FTextureRHIParamRef TexturesToReadable[1];
-	TexturesToReadable[0] = DestRenderTarget.ShaderResourceTexture;
-	Context.RHICmdList.TransitionResources(EResourceTransitionAccess::EReadable, TexturesToReadable, ARRAY_COUNT(TexturesToReadable));
-}
+		SubsurfaceTextureX = GraphBuilder.CreateTexture(TextureDescIntermediate, TEXT("SubsurfaceTextureX"));
+		SubsurfaceTextureY = GraphBuilder.CreateTexture(TextureDescIntermediate, TEXT("SubsurfaceTextureY"));
 
+		FSamplerStateRHIParamRef SubsurfaceSamplerState = FSubsurfacePS::GetSamplerState();
+		const FSubsurfacePS::EQuality SubsurfaceQuality = FSubsurfacePS::GetQuality();
 
-FPooledRenderTargetDesc FRCPassPostProcessSubsurface::ComputeOutputDesc(EPassOutputId InPassOutputId) const
-{
-	FPooledRenderTargetDesc Ret = GetInput(ePId_Input0)->GetOutput()->RenderTargetDesc;
-
-	Ret.Reset();
-	Ret.DebugName = (Direction == 0) ? TEXT("SubsurfaceX") : TEXT("SubsurfaceY");
+		static const TCHAR* SubsurfacePassNameByDirection[] =
+		{
+			TEXT("SubsurfaceX"),
+			TEXT("SubsurfaceY")
+		};
 
-	return Ret;
-}
+		FRDGTextureRef SubsurfacePassOutputByDirection[] =
+		{
+			SubsurfaceTextureX,
+			SubsurfaceTextureY
+		};
 
+		FRDGTextureRef SubsurfacePassInputByDirection[] =
+		{
+			SetupTexture,
+			SubsurfaceTextureX
+		};
 
+		// Horizontal / Vertical scattering passes using a separable filter.
+		for (uint32 DirectionIndex = 0; DirectionIndex < static_cast<uint32>(FSubsurfacePS::EDirection::MAX); ++DirectionIndex)
+		{
+			const auto Direction = static_cast<FSubsurfacePS::EDirection>(DirectionIndex);
 
+			FSubsurfacePS::FParameters* PassParameters = GraphBuilder.AllocParameters<FSubsurfacePS::FParameters>();
+			PassParameters->Subsurface = SubsurfaceCommonParameters;
+			PassParameters->RenderTargets[0] = FRenderTargetBinding(SubsurfacePassOutputByDirection[DirectionIndex], ERenderTargetLoadAction::ENoAction, ERenderTargetStoreAction::EStore);
+			PassParameters->SubsurfaceInput0 = GetScreenPassInputParameters(SubsurfacePassInputByDirection[DirectionIndex], SubsurfaceSamplerState);
 
+			// Read / Write between intermediate targets.
+			const FIntRect ViewportRect = ViewportRectIntermediate;
+			const FIntRect TextureRect = ViewportRectIntermediate;
+			const FIntPoint TextureSize = TextureSizeIntermediate;
 
+			// If we are sampling from a subset of the texture (e.g. stereo rendering), we have to manually clamp UVs.
+			const bool bManuallyClampUV = TextureSize != TextureRect.Size();
 
-/** Encapsulates the post processing subsurface recombine pixel shader. */
-// @param RecombineMode 0:fullres, 1: halfres, 2:no scattering, just reconstruct the lighting (needed for scalability)
-// @param RecombineQuality 0:low..1:high
-template <uint32 RecombineMode, uint32 RecombineQuality, uint32 Checkerboard>
-class TPostProcessSubsurfaceRecombinePS : public FGlobalShader
-{
-	DECLARE_SHADER_TYPE(TPostProcessSubsurfaceRecombinePS, Global);
+			FSubsurfacePS::FPermutationDomain PixelShaderPermutationVector;
+			PixelShaderPermutationVector.Set<FSubsurfacePS::FDimensionDirection>(Direction);
+			PixelShaderPermutationVector.Set<FSubsurfacePS::FDimensionQuality>(SubsurfaceQuality);
+			PixelShaderPermutationVector.Set<FSubsurfacePS::FDimensionManuallyClampUV>(bManuallyClampUV);
+			TShaderMapRef<FSubsurfacePS> PixelShader(Context->ShaderMap, PixelShaderPermutationVector);
 
-	static bool ShouldCompilePermutation(const FGlobalShaderPermutationParameters& Parameters)
-	{
-		return IsFeatureLevelSupported(Parameters.Platform, ERHIFeatureLevel::SM4);
+			AddDrawScreenPass(GraphBuilder, FRDGEventName(SubsurfacePassNameByDirection[DirectionIndex]), Context, ViewportRect, TextureRect, TextureSize, *PixelShader, PassParameters);
+		}
 	}
 
-	static void ModifyCompilationEnvironment(const FGlobalShaderPermutationParameters& Parameters, FShaderCompilerEnvironment& OutEnvironment)
+	RecombineTexture = GraphBuilder.CreateTexture(SceneTextureDesc, TEXT("SubsurfaceRecombine"));
+
+	// If multiple views exist (e.g. VR) we need to copy all other viewports from the scene texture
+	// to the recombine target so we don't lose them.
+	if (Context->ViewFamily.Views.Num())
 	{
-		FGlobalShader::ModifyCompilationEnvironment(Parameters,OutEnvironment);
-		OutEnvironment.SetDefine(TEXT("RECOMBINE_QUALITY"), RecombineQuality);
-		OutEnvironment.SetDefine(TEXT("HALF_RES"), (uint32)(RecombineMode == 1));
-		OutEnvironment.SetDefine(TEXT("RECOMBINE_SUBSURFACESCATTER"), (uint32)(RecombineMode != 2));
-		OutEnvironment.SetDefine(TEXT("SUBSURFACE_RADIUS_SCALE"), SUBSURFACE_RADIUS_SCALE);
-		OutEnvironment.SetDefine(TEXT("SUBSURFACE_KERNEL_SIZE"), SUBSURFACE_KERNEL_SIZE);
-		OutEnvironment.SetDefine(TEXT("SUBSURFACE_PROFILE_CHECKERBOARD"), Checkerboard);
-	}
+		FSubsurfaceViewportCopyPS::FParameters* PassParameters = GraphBuilder.AllocParameters<FSubsurfaceViewportCopyPS::FParameters>();
+		PassParameters->RenderTargets[0] = FRenderTargetBinding(RecombineTexture, ERenderTargetLoadAction::ENoAction, ERenderTargetStoreAction::EStore);
+		PassParameters->SubsurfaceInput0 = GetScreenPassInputParameters(SceneTexture, PointClampSampler);
 
-	/** Default constructor. */
-	TPostProcessSubsurfaceRecombinePS() {}
+		const FIntPoint TextureSize = SceneTextureDesc.Extent;
 
-public:
-	FPostProcessPassParameters PostprocessParameter;
-	FSceneTextureShaderParameters SceneTextureParameters;
-	FSubsurfaceParameters SubsurfaceParameters;
+		TShaderMapRef<FSubsurfaceViewportCopyPS> PixelShader(Context->ShaderMap);
 
-	/** Initialization constructor. */
-	TPostProcessSubsurfaceRecombinePS(const ShaderMetaType::CompiledShaderInitializerType& Initializer)
-		: FGlobalShader(Initializer)
-	{
-		PostprocessParameter.Bind(Initializer.ParameterMap);
-		SceneTextureParameters.Bind(Initializer);
-		SubsurfaceParameters.Bind(Initializer.ParameterMap);
-	}
+		GraphBuilder.AddPass(
+			RDG_EVENT_NAME("SubsurfaceViewportCopy"),
+			PassParameters,
+			ERenderGraphPassFlags::None,
+			[Context, PixelShader, TextureSize, PassParameters](FRHICommandListImmediate& RHICmdList)
+		{
+			const FSceneViewFamily& ViewFamily = Context->ViewFamily;
 
-	// FShader interface.
-	virtual bool Serialize(FArchive& Ar) override
-	{
-		bool bShaderHasOutdatedParameters = FGlobalShader::Serialize(Ar);
-		Ar << PostprocessParameter << SceneTextureParameters << SubsurfaceParameters;
-		return bShaderHasOutdatedParameters;
-	}
+			for (uint32 ViewId = 0, ViewCount = ViewFamily.Views.Num(); ViewId < ViewCount; ++ViewId)
+			{
+				const FViewInfo* LocalView = static_cast<const FViewInfo*>(ViewFamily.Views[ViewId]);
 
-	void SetParameters(const FRenderingCompositePassContext& Context)
-	{
-		const FPixelShaderRHIParamRef ShaderRHI = GetPixelShader();
+				// Skip the view we are currently processing.
+				if (LocalView != &Context->View)
+				{
+					const FIntRect Rect = LocalView->ViewRect;
 
-		FGlobalShader::SetParameters<FViewUniformShaderParameters>(Context.RHICmdList, ShaderRHI, Context.View.ViewUniformBuffer);
-		SceneTextureParameters.Set(Context.RHICmdList, ShaderRHI, Context.View.FeatureLevel, ESceneTextureSetupMode::All);
-		PostprocessParameter.SetPS(Context.RHICmdList, ShaderRHI, Context, TStaticSamplerState<SF_Bilinear, AM_Border, AM_Border, AM_Border>::GetRHI());
-		SubsurfaceParameters.SetParameters(Context.RHICmdList, ShaderRHI, Context);
+					DrawScreenPass(RHICmdList, Context, Rect, Rect, TextureSize, *PixelShader, *PassParameters);
+				}
+			}
+		});
 	}
 
-	static const TCHAR* GetSourceFilename()
+	// Recombines scattering result with scene color.
 	{
-		return TEXT("/Engine/Private/PostProcessSubsurface.usf");
-	}
+		FSubsurfaceRecombinePS::FParameters* PassParameters = GraphBuilder.AllocParameters<FSubsurfaceRecombinePS::FParameters>();
+		PassParameters->Subsurface = SubsurfaceCommonParameters;
+		PassParameters->RenderTargets[0] = FRenderTargetBinding(RecombineTexture, ERenderTargetLoadAction::ELoad, ERenderTargetStoreAction::EStore);
+		PassParameters->SubsurfaceInput0 = GetScreenPassInputParameters(SceneTexture, BilinearBorderState);
 
-	static const TCHAR* GetFunctionName()
-	{
-		return TEXT("SubsurfaceRecombinePS");
-	}
-};
+		// Scattering output target is only used when scattering is enabled.
+		if (SubsurfaceMode != ESubsurfaceMode::Bypass)
+		{
+			PassParameters->SubsurfaceInput1 = GetScreenPassInputParameters(SubsurfaceTextureY, BilinearBorderState);
+		}
 
-// #define avoids a lot of code duplication
-#define VARIATION1(A)		VARIATION2(A,0)			VARIATION2(A,1)		
-#define VARIATION2(A, B)	VARIATION3(A,B,0)		VARIATION3(A,B,1)
-#define VARIATION3(A, B, C) typedef TPostProcessSubsurfaceRecombinePS<A, B, C> TPostProcessSubsurfaceRecombinePS##A##B##C; \
-	IMPLEMENT_SHADER_TYPE2(TPostProcessSubsurfaceRecombinePS##A##B##C, SF_Pixel);
-VARIATION1(0) VARIATION1(1) VARIATION1(2)
-#undef VARIATION1
-#undef VARIATION2
-#undef VARIATION3
-
-// @param RecombineMode 0:fullres, 1: halfres, 2:no scattering, just reconstruct the lighting (needed for scalability)
-// @param RecombineQuality 0:low..1:high
-template <uint32 RecombineMode, uint32 RecombineQuality, uint32 Checkerboard>
-void SetSubsurfaceRecombineShader(const FRenderingCompositePassContext& Context, TShaderMapRef<FPostProcessVS> &VertexShader)
-{
-	FGraphicsPipelineStateInitializer GraphicsPSOInit;
-	Context.RHICmdList.ApplyCachedRenderTargets(GraphicsPSOInit);
-	GraphicsPSOInit.BlendState = TStaticBlendState<>::GetRHI();
-	GraphicsPSOInit.RasterizerState = TStaticRasterizerState<>::GetRHI();
-	GraphicsPSOInit.DepthStencilState = TStaticDepthStencilState<false, CF_Always>::GetRHI();
+		// Read from intermediate target; write to final target.
+		const FIntRect ViewportRect = ViewportRectFinal;
+		const FIntRect TextureRect = ViewportRectIntermediate;
+		const FIntPoint TextureSize = TextureSizeIntermediate;
 
-	TShaderMapRef<TPostProcessSubsurfaceRecombinePS<RecombineMode, RecombineQuality, Checkerboard> > PixelShader(Context.GetShaderMap());
+		const FSubsurfaceRecombinePS::EQuality RecombineQuality = FSubsurfaceRecombinePS::GetQuality(Context->View);
 
-	GraphicsPSOInit.BoundShaderState.VertexDeclarationRHI = GFilterVertexDeclaration.VertexDeclarationRHI;
-	GraphicsPSOInit.BoundShaderState.VertexShaderRHI = GETSAFERHISHADER_VERTEX(*VertexShader);
-	GraphicsPSOInit.BoundShaderState.PixelShaderRHI = GETSAFERHISHADER_PIXEL(*PixelShader);
-	GraphicsPSOInit.PrimitiveType = PT_TriangleList;
+		FSubsurfaceRecombinePS::FPermutationDomain PixelShaderPermutationVector;
+		PixelShaderPermutationVector.Set<FSubsurfaceRecombinePS::FDimensionMode>(SubsurfaceMode);
+		PixelShaderPermutationVector.Set<FSubsurfaceRecombinePS::FDimensionQuality>(RecombineQuality);
+		PixelShaderPermutationVector.Set<FSubsurfaceRecombinePS::FDimensionCheckerboard>(bCheckerboard);
+		TShaderMapRef<FSubsurfaceRecombinePS> PixelShader(Context->ShaderMap, PixelShaderPermutationVector);
 
-	SetGraphicsPipelineState(Context.RHICmdList, GraphicsPSOInit);
+		AddDrawScreenPass(GraphBuilder, RDG_EVENT_NAME("SubsurfaceRecombine"), Context, ViewportRect, TextureRect, TextureSize, *PixelShader, PassParameters);
+	}
 
-	PixelShader->SetParameters(Context);
-	VertexShader->SetParameters(Context);
+	return RecombineTexture;
 }
 
-FRCPassPostProcessSubsurfaceRecombine::FRCPassPostProcessSubsurfaceRecombine(bool bInHalfRes, bool bInSingleViewportMode)
-	: bHalfRes(bInHalfRes)
-	, bSingleViewportMode(bInSingleViewportMode)
+FRDGTextureRef VisualizeSubsurface(
+	FRDGBuilder& GraphBuilder,
+	FScreenPassContextRef Context,
+	FRDGTextureRef SceneTexture)
 {
-}
+	check(Context);
+	check(SceneTexture);
 
-void FRCPassPostProcessSubsurfaceRecombine::Process(FRenderingCompositePassContext& Context)
-{
-	FSceneRenderTargets& SceneContext = FSceneRenderTargets::Get(Context.RHICmdList);
-	const FPooledRenderTargetDesc* InputDesc = GetInputDesc(ePId_Input0);
+	FSamplerStateRHIParamRef PointClampState = TStaticSamplerState<SF_Point, AM_Clamp, AM_Clamp, AM_Clamp>::GetRHI();
 
-	check(InputDesc);
+	FRDGTextureRef VisualizeTexture = GraphBuilder.CreateTexture(SceneTexture->Desc, TEXT("SubsurfaceVisualize"));
 
-	const FViewInfo& View = Context.View;
-	const FSceneViewFamily& ViewFamily = *(View.Family);
+	FSubsurfaceVisualizePS::FParameters* PassParameters = GraphBuilder.AllocParameters<FSubsurfaceVisualizePS::FParameters>();
+	PassParameters->Subsurface = GetSubsurfaceCommonParameters(GraphBuilder.RHICmdList, Context);
+	PassParameters->RenderTargets[0] = FRenderTargetBinding(VisualizeTexture, ERenderTargetLoadAction::EClear, ERenderTargetStoreAction::EStore);
+	PassParameters->SubsurfaceInput0 = GetScreenPassInputParameters(SceneTexture, PointClampState);
+	PassParameters->MiniFontTexture = GetMiniFontTexture();
 
-	FIntPoint SrcSize = InputDesc->Extent;
-	FIntPoint DestSize = SceneContext.GetBufferSizeXY();
+	GraphBuilder.AddPass(
+		RDG_EVENT_NAME("SubsurfaceVisualize"),
+		PassParameters,
+		ERenderGraphPassFlags::None,
+		[Context, SceneTexture, VisualizeTexture, PassParameters](FRHICommandListImmediate& RHICmdList)
+	{
+		const FIntRect ViewportRect = Context->ViewportRect;
+		const FIntRect TextureRect = ViewportRect;
+		const FIntPoint TextureSize = SceneTexture->Desc.Extent;
 
-	check(DestSize.X);
-	check(DestSize.Y);
-	check(SrcSize.X);
-	check(SrcSize.Y);
+		TShaderMapRef<FSubsurfaceVisualizePS> PixelShader(Context->ShaderMap);
 
-	FIntRect SrcRect = View.ViewRect;
-	FIntRect DestRect = SrcRect;
+		DrawScreenPass(RHICmdList, Context, ViewportRect, TextureRect, TextureSize, *PixelShader, *PassParameters);
 
-	TRefCountPtr<IPooledRenderTarget>& SceneColor = SceneContext.GetSceneColor();
+		// Draw debug text
+		{
+			const FSceneViewFamily& ViewFamily = Context->ViewFamily;
+			FRenderTargetTemp TempRenderTarget(static_cast<FTexture2DRHIParamRef>(VisualizeTexture->GetRHITexture()), TextureSize);
+			FCanvas Canvas(&TempRenderTarget, nullptr, ViewFamily.CurrentRealTime, ViewFamily.CurrentWorldTime, ViewFamily.DeltaWorldTime, Context->View.GetFeatureLevel());
 
-	const FSceneRenderTargetItem& DestRenderTarget = PassOutputs[0].RequestSurface(Context);
+			float X = 30;
+			float Y = 28;
+			const float YStep = 14;
 
-	ERenderTargetLoadAction LoadAction = Context.GetLoadActionForRenderTarget(DestRenderTarget);
+			FString Line = FString::Printf(TEXT("Visualize Screen Space Subsurface Scattering"));
+			Canvas.DrawShadowedString(X, Y += YStep, *Line, GetStatsFont(), FLinearColor(1, 1, 1));
 
-	FRHIRenderPassInfo RPInfo(DestRenderTarget.TargetableTexture, MakeRenderTargetActions(LoadAction, ERenderTargetStoreAction::EStore), DestRenderTarget.ShaderResourceTexture);
-	Context.RHICmdList.BeginRenderPass(RPInfo, TEXT("SubsurfaceRecombine"));
-	{
+			Y += YStep;
 
-		CopyOverOtherViewportsIfNeeded(Context, View);
+			uint32 Index = 0;
+			while (GSubsurfaceProfileTextureObject.GetEntryString(Index++, Line))
+			{
+				Canvas.DrawShadowedString(X, Y += YStep, *Line, GetStatsFont(), FLinearColor(1, 1, 1));
+			}
 
-		Context.SetViewportAndCallRHI(DestRect);
+			const bool bFlush = false;
+			const bool bInsideRenderPass = true;
+			Canvas.Flush_RenderThread(RHICmdList, bFlush, bInsideRenderPass);
+		}
+	});
 
-		TShaderMapRef<FPostProcessVS> VertexShader(Context.GetShaderMap());
+	return VisualizeTexture;
+}
 
-		uint32 QualityCVar = CVarSSSQuality.GetValueOnRenderThread();
-		const bool bCheckerboard = FRCPassPostProcessSubsurface::RequiresCheckerboardSubsurfaceRendering(SceneContext.GetSceneColorFormat());
+FSubsurfaceVisualizeCompositePass::FSubsurfaceVisualizeCompositePass(FRHICommandList& RHICmdList)
+{
+	// we need the GBuffer, we release it Process()
+	FSceneRenderTargets::Get(RHICmdList).AdjustGBufferRefCount(RHICmdList, 1);
+}
 
-		// 0:low / 1:high
-		uint32 RecombineQuality = 0;
-		{
-			if (QualityCVar == -1)
-			{
-				RecombineQuality = (View.AntiAliasingMethod == AAM_TemporalAA) ? 0 : 1;
-			}
-			else if (QualityCVar == 1)
-			{
-				RecombineQuality = 1;
-			}
-		}
+void FSubsurfaceVisualizeCompositePass::Process(FRenderingCompositePassContext& CompositePassContext)
+{
+	FRHICommandListImmediate& RHICmdList = CompositePassContext.RHICmdList;
 
-		// needed for Scalability
-		// 0:fullres, 1: halfres, 2:no scattering, just reconstruct the lighting (needed for scalability)
-		uint32 RecombineMode = 2;
+	FRDGBuilder GraphBuilder(RHICmdList);
 
-		if (GetInput(ePId_Input1)->IsValid())
-		{
-			RecombineMode = bHalfRes ? 1 : 0;
-		}
+	FScreenPassContext* Context = FScreenPassContext::Create(RHICmdList, CompositePassContext.View);
 
-		SCOPED_DRAW_EVENTF(Context.RHICmdList, SubsurfacePassRecombine, TEXT("SubsurfacePassRecombine Mode:%d Quality:%d"), RecombineMode, RecombineQuality);
+	FRDGTextureRef InputTexture = CreateRDGTextureForInput(GraphBuilder, ePId_Input0, TEXT("SceneColor"), eFC_0000);
 
-		{
-			if (bCheckerboard)
-			{
-				if (RecombineQuality == 0)
-				{
-					switch (RecombineMode)
-					{
-					case 0: SetSubsurfaceRecombineShader<0, 0, 1>(Context, VertexShader); break;
-					case 1: SetSubsurfaceRecombineShader<1, 0, 1>(Context, VertexShader); break;
-					case 2: SetSubsurfaceRecombineShader<2, 0, 1>(Context, VertexShader); break;
-					default: check(0);
-					}
-				}
-				else
-				{
-					switch (RecombineMode)
-					{
-					case 0: SetSubsurfaceRecombineShader<0, 1, 1>(Context, VertexShader); break;
-					case 1: SetSubsurfaceRecombineShader<1, 1, 1>(Context, VertexShader); break;
-					case 2: SetSubsurfaceRecombineShader<2, 1, 1>(Context, VertexShader); break;
-					default: check(0);
-					}
-				}
-			}
-			else
-			{
-				if (RecombineQuality == 0)
-				{
-					switch (RecombineMode)
-					{
-					case 0: SetSubsurfaceRecombineShader<0, 0, 0>(Context, VertexShader); break;
-					case 1: SetSubsurfaceRecombineShader<1, 0, 0>(Context, VertexShader); break;
-					case 2: SetSubsurfaceRecombineShader<2, 0, 0>(Context, VertexShader); break;
-					default: check(0);
-					}
-				}
-				else
-				{
-					switch (RecombineMode)
-					{
-					case 0: SetSubsurfaceRecombineShader<0, 1, 0>(Context, VertexShader); break;
-					case 1: SetSubsurfaceRecombineShader<1, 1, 0>(Context, VertexShader); break;
-					case 2: SetSubsurfaceRecombineShader<2, 1, 0>(Context, VertexShader); break;
-					default: check(0);
-					}
-				}
-			}
-		}
+	FRDGTextureRef OutputTexture = VisualizeSubsurface(GraphBuilder, Context, InputTexture);
 
-		DrawPostProcessPass(
-			Context.RHICmdList,
-			0, 0,
-			DestRect.Width(), DestRect.Height(),
-			SrcRect.Min.X, SrcRect.Min.Y,
-			SrcRect.Width(), SrcRect.Height(),
-			DestRect.Size(),
-			SrcSize,
-			*VertexShader,
-			View.StereoPass,
-			Context.HasHmdMesh(),
-			EDRF_UseTriangleOptimization);
-	}
-	Context.RHICmdList.EndRenderPass();
+	ExtractRDGTextureForOutput(GraphBuilder, ePId_Output0, OutputTexture);
 
-	// replace the current SceneColor with this one
-	SceneContext.SetSceneColor(PassOutputs[0].PooledRenderTarget);
-	PassOutputs[0].PooledRenderTarget.SafeRelease();
-}
+	TRefCountPtr<IPooledRenderTarget> OutputTarget;
+	GraphBuilder.QueueTextureExtraction(OutputTexture, &OutputTarget);
 
+	GraphBuilder.Execute();
 
-FPooledRenderTargetDesc FRCPassPostProcessSubsurfaceRecombine::ComputeOutputDesc(EPassOutputId InPassOutputId) const
-{
-	FPooledRenderTargetDesc Ret = GetInput(ePId_Input0)->GetOutput()->RenderTargetDesc;
+	check(OutputTarget);
 
-	Ret.Reset();
-	Ret.DebugName = TEXT("SceneColorSubsurface");
+	RHICmdList.CopyToResolveTarget(
+		OutputTarget->GetRenderTargetItem().TargetableTexture,
+		OutputTarget->GetRenderTargetItem().ShaderResourceTexture,
+		FResolveParams());
 
-	// we replace the HDR SceneColor with this one
-	return Ret;
+	FSceneRenderTargets::Get(RHICmdList).AdjustGBufferRefCount(RHICmdList, -1);
 }
 
+FPooledRenderTargetDesc FSubsurfaceVisualizeCompositePass::ComputeOutputDesc(EPassOutputId InPassOutputId) const
+{
+	FPooledRenderTargetDesc Ret = FSceneRenderTargets::Get_FrameConstantsOnly().GetSceneColor()->GetDesc();
+	Ret.Flags &= ~(TexCreate_FastVRAM | TexCreate_Transient);
+	Ret.Reset();
+	Ret.DebugName = TEXT("SubsurfaceVisualize");
+	Ret.Format = PF_FloatRGBA;
+	return Ret;
+}
\ No newline at end of file
diff --git a/Engine/Source/Runtime/Renderer/Private/PostProcess/PostProcessSubsurface.h b/Engine/Source/Runtime/Renderer/Private/PostProcess/PostProcessSubsurface.h
--- a/Engine/Source/Runtime/Renderer/Private/PostProcess/PostProcessSubsurface.h
+++ b/Engine/Source/Runtime/Renderer/Private/PostProcess/PostProcessSubsurface.h
@@ -6,90 +6,40 @@
 
 #pragma once
 
-#include "CoreMinimal.h"
-#include "RendererInterface.h"
-#include "PostProcess/RenderingCompositionGraph.h"
-
-class FViewInfo;
-
-// ePId_Input0: SceneColor
-// derives from TRenderingCompositePassBase<InputCount, OutputCount> 
-// uses some GBuffer attributes
-// alpha is unused
-class FRCPassPostProcessSubsurfaceVisualize : public TRenderingCompositePassBase<1, 1>
-{
-public:
-	FRCPassPostProcessSubsurfaceVisualize(FRHICommandList& RHICmdList);
-
-	// interface FRenderingCompositePass ---------
-	virtual void Process(FRenderingCompositePassContext& Context) override;
-	virtual FPooledRenderTargetDesc ComputeOutputDesc(EPassOutputId InPassOutputId) const override;
-	virtual void Release() override { delete this; }
-};
-
-// ePId_Input0: SceneColor
-// derives from TRenderingCompositePassBase<InputCount, OutputCount> 
-// uses some GBuffer attributes
-// alpha is unused
-class FRCPassPostProcessSubsurfaceSetup : public TRenderingCompositePassBase<1, 1>
-{
-public:
-	// constructor
-	FRCPassPostProcessSubsurfaceSetup(FViewInfo& View, bool bInHalfRes);
-
-	// interface FRenderingCompositePass ---------
-	virtual void Process(FRenderingCompositePassContext& Context) override;
-	virtual FPooledRenderTargetDesc ComputeOutputDesc(EPassOutputId InPassOutputId) const override;
-	virtual void Release() override { delete this; }
-
-	FIntRect ViewRect;
-	bool bHalfRes;
-};
-
-
-// derives from TRenderingCompositePassBase<InputCount, OutputCount>
-// ePId_Input0: SceneColor (horizontal blur) or the pass before (vertical blur)
-// ePId_Input1: optional Setup pass (only for InDirection==1)
-// modifies SceneColor, uses some GBuffer attributes
-class FRCPassPostProcessSubsurface : public TRenderingCompositePassBase<2, 1>
+#include "ScreenPass.h"
+#include "RenderingCompositionGraph.h"
+
+// Returns whether subsurface scattering is globally enabled.
+bool IsSubsurfaceEnabled();
+
+// Returns whether subsurface scattering is required for the provided view.
+bool IsSubsurfaceRequiredForView(const FViewInfo& View);
+	
+// Returns whether checkerboard rendering is enabled for the provided format.
+bool IsSubsurfaceCheckerboardFormat(EPixelFormat SceneColorFormat);
+
+// Computes subsurface scattering on the scene texture and produces a new scene texture as output.
+FRDGTextureRef ComputeSubsurface(
+	FRDGBuilder& GraphBuilder,
+	FScreenPassContextRef Context,
+	FRDGTextureRef SceneTexture);
+
+// Visualizes subsurface scattering profiles by overlaying an image on the provided scene texture.
+// Produces a new scene texture as output.
+FRDGTextureRef VisualizeSubsurface(
+	FRDGBuilder& GraphBuilder,
+	FScreenPassContextRef Context,
+	FRDGTextureRef SceneTexture);
+
+// An adapter to connect the new Render Graph implementation to the legacy Composition Graph.
+class FSubsurfaceVisualizeCompositePass : public TRenderingCompositePassBase<1, 1>
 {
 public:
-	// constructor
-	// @param InDirection 0:horizontal/1:vertical
-	FRCPassPostProcessSubsurface(uint32 InDirection, bool bInHalfRes);
-
-	static const bool RequiresCheckerboardSubsurfaceRendering(EPixelFormat SceneColorFormat);
-
-	// interface FRenderingCompositePass ---------
-
-	virtual void Process(FRenderingCompositePassContext& Context) override;
-	virtual void Release() override { delete this; }
-	virtual FPooledRenderTargetDesc ComputeOutputDesc(EPassOutputId InPassOutputId) const override;
+	FSubsurfaceVisualizeCompositePass(FRHICommandList& CmdList);
+	virtual ~FSubsurfaceVisualizeCompositePass() = default;
 
 private:
-	// 0:horizontal/1:vertical
-	uint32 Direction;
-	bool bHalfRes;
-};
-
-
-// derives from TRenderingCompositePassBase<InputCount, OutputCount>
-// ePId_Input0: SceneColor before Screen Space Subsurface input
-// ePId_Input1: optional output from FRCPassPostProcessSubsurface (if not present we do cheap reconstruction for Scalability)
-// ePId_Input2: optional SubsurfaceSetup, can be half res
-// modifies SceneColor, uses some GBuffer attributes
-class FRCPassPostProcessSubsurfaceRecombine : public TRenderingCompositePassBase<3, 1>
-{
-public:
-	// constructor
-	FRCPassPostProcessSubsurfaceRecombine(bool bInHalfRes, bool bInSingleViewportMode);
-
-	// interface FRenderingCompositePass ---------
-
-	virtual void Process(FRenderingCompositePassContext& Context) override;
-	virtual void Release() override { delete this; }
-	virtual FPooledRenderTargetDesc ComputeOutputDesc(EPassOutputId InPassOutputId) const override;
-
-	bool bHalfRes;
-	bool bSingleViewportMode;
-};
+	void Process(FRenderingCompositePassContext& Context) override;
+	FPooledRenderTargetDesc ComputeOutputDesc(EPassOutputId InPassOutputId) const override;
+	void Release() override { delete this; }
+};
\ No newline at end of file
diff --git a/Engine/Source/Runtime/Renderer/Private/PostProcess/PostProcessing.cpp b/Engine/Source/Runtime/Renderer/Private/PostProcess/PostProcessing.cpp
--- a/Engine/Source/Runtime/Renderer/Private/PostProcess/PostProcessing.cpp
+++ b/Engine/Source/Runtime/Renderer/Private/PostProcess/PostProcessing.cpp
@@ -43,18 +43,18 @@
 #include "PostProcess/PostProcessVisualizeComplexity.h"
 #include "PostProcess/PostProcessCompositeEditorPrimitives.h"
 #include "PostProcess/PostProcessShaderPrint.h"
-#include "CompositionLighting/PostProcessPassThrough.h"
 #include "PostProcess/PostProcessTestImage.h"
-#include "HighResScreenshot.h"
+#include "PostProcess/PostProcessFFTBloom.h"
+#include "PostProcess/PostProcessStreamingAccuracyLegend.h"
 #include "PostProcess/PostProcessSubsurface.h"
 #include "PostProcess/PostProcessMorpheus.h"
+#include "CompositionLighting/PostProcessPassThrough.h"
+#include "CompositionLighting/PostProcessLpvIndirect.h"
+#include "HighResScreenshot.h"
 #include "IHeadMountedDisplay.h"
 #include "IXRTrackingSystem.h"
 #include "BufferVisualizationData.h"
-#include "CompositionLighting/PostProcessLpvIndirect.h"
-#include "PostProcess/PostProcessStreamingAccuracyLegend.h"
 #include "DeferredShadingRenderer.h"
-#include "PostProcess/PostProcessFFTBloom.h"
 #include "MobileSeparateTranslucencyPass.h"
 #include "MobileDistortionPass.h"
 
@@ -177,7 +177,7 @@ static TAutoConsoleVariable<int32> CVarAlphaChannel(
 	TEXT(" 1: enabled in linear color space;\n")
 	TEXT(" 2: same as 1, but also enable it through the tonemapper. Compositing after the tonemapper is incorrect, as their is no meaning to tonemap the alpha channel. This is only meant to be use exclusively for broadcasting hardware that does not support linear color space compositing and tonemapping."),
 	ECVF_ReadOnly);
-	
+
 static TAutoConsoleVariable<int32> CVarPostProcessingPreferCompute(
 	TEXT("r.PostProcessing.PreferCompute"),
 	0,
@@ -228,9 +228,6 @@ static TAutoConsoleVariable<float> CVarTemporalAAHistorySP(
 	ECVF_RenderThreadSafe
 );
 
-
-IMPLEMENT_SHADER_TYPE(,FPostProcessVS,TEXT("/Engine/Private/PostProcessBloom.usf"),TEXT("MainPostprocessCommonVS"),SF_Vertex);
-
 static bool HasPostProcessMaterial(FPostprocessContext& Context, EBlendableLocation InLocation);
 
 // -------------------------------------------------------
@@ -2136,7 +2133,7 @@ void FPostProcessing::Process(FRHICommandListImmediate& RHICmdList, const FViewI
 			ensureMsgf(!bUnscaledFinalOutput, TEXT("VisualizeSSS is incompatible with unscaled output."));
 
 			// the setup pass also does visualization, based on EngineShowFlags.VisualizeSSS
-			FRenderingCompositePass* PassVisualize = Context.Graph.RegisterPass(new(FMemStack::Get()) FRCPassPostProcessSubsurfaceVisualize(RHICmdList));
+			FRenderingCompositePass* PassVisualize = Context.Graph.RegisterPass(new(FMemStack::Get()) FSubsurfaceVisualizeCompositePass(Context.RHICmdList));
 			PassVisualize->SetInput(ePId_Input0, Context.FinalOutput);
 			Context.FinalOutput = FRenderingCompositeOutputRef(PassVisualize);
 		}
@@ -2922,4 +2919,4 @@ void FPostProcessing::ProcessPlanarReflection(FRHICommandListImmediate& RHICmdLi
 bool FPostProcessing::HasAlphaChannelSupport()
 {
 	return CVarAlphaChannel.GetValueOnAnyThread() != 0;
-}
+}
\ No newline at end of file
diff --git a/Engine/Source/Runtime/Renderer/Private/PostProcess/PostProcessing.h b/Engine/Source/Runtime/Renderer/Private/PostProcess/PostProcessing.h
--- a/Engine/Source/Runtime/Renderer/Private/PostProcess/PostProcessing.h
+++ b/Engine/Source/Runtime/Renderer/Private/PostProcess/PostProcessing.h
@@ -7,16 +7,27 @@
 #pragma once
 
 #include "CoreMinimal.h"
-#include "Shader.h"
-#include "GlobalShader.h"
 #include "PostProcess/RenderingCompositionGraph.h"
 #include "SceneRendering.h"
+#include "ScreenPass.h"
+
+// For compatibility with composition graph passes until they are ported to Render Graph.
+class FPostProcessVS : public FScreenPassVS
+{
+public:
+	FPostProcessVS() = default;
+	FPostProcessVS(const ShaderMetaType::CompiledShaderInitializerType& Initializer)
+		: FScreenPassVS(Initializer)
+	{}
+
+	void SetParameters(const FRenderingCompositePassContext&) {}
+	void SetParameters(FRHICommandList&, const FUniformBufferRHIParamRef) {}
+};
 
 /** The context used to setup a post-process pass. */
 class FPostprocessContext
 {
 public:
-
 	FPostprocessContext(FRHICommandListImmediate& InRHICmdList, FRenderingCompositionGraph& InGraph, const FViewInfo& InView);
 
 	FRHICommandListImmediate& RHICmdList;
@@ -31,48 +42,12 @@ public:
 	FRenderingCompositeOutputRef FinalOutput;
 };
 
-/** Encapsulates the post processing vertex shader. */
-class FPostProcessVS : public FGlobalShader
-{
-	DECLARE_SHADER_TYPE(FPostProcessVS,Global);
-
-	static bool ShouldCompilePermutation(const FGlobalShaderPermutationParameters& Parameters)
-	{
-		return true;
-	}
-
-	/** Default constructor. */
-	FPostProcessVS() {}
-
-	/** to have a similar interface as all other shaders */
-	void SetParameters(const FRenderingCompositePassContext& Context)
-	{
-		FGlobalShader::SetParameters<FViewUniformShaderParameters>(Context.RHICmdList, GetVertexShader(), Context.View.ViewUniformBuffer);
-	}
-
-	void SetParameters(FRHICommandList& RHICmdList, const FUniformBufferRHIParamRef ViewUniformBuffer)
-	{
-		FGlobalShader::SetParameters<FViewUniformShaderParameters>(RHICmdList, GetVertexShader(), ViewUniformBuffer);
-	}
-
-public:
-
-	/** Initialization constructor. */
-	FPostProcessVS(const ShaderMetaType::CompiledShaderInitializerType& Initializer)
-		: FGlobalShader(Initializer)
-	{
-	}
-};
-
-
-
 /**
  * The center for all post processing activities.
  */
 class FPostProcessing
 {
 public:
-
 	bool AllowFullPostProcessing(const FViewInfo& View, ERHIFeatureLevel::Type FeatureLevel);
 
 	void RegisterHMDPostprocessPass(FPostprocessContext& Context, const FEngineShowFlags& EngineShowFlags) const;
@@ -80,8 +55,6 @@ public:
 	// @param VelocityRT only valid if motion blur is supported
 	void Process(FRHICommandListImmediate& RHICmdList, const FViewInfo& View, TRefCountPtr<IPooledRenderTarget>& VelocityRT);
 
-
-
 	void ProcessES2(FRHICommandListImmediate& RHICmdList, FScene* Scene, const FViewInfo& View);
 
 	void ProcessPlanarReflection(FRHICommandListImmediate& RHICmdList, FViewInfo& View, TRefCountPtr<IPooledRenderTarget>& VelocityRT, TRefCountPtr<IPooledRenderTarget>& OutFilteredSceneColor);
diff --git a/Engine/Source/Runtime/Renderer/Private/PostProcess/RenderingCompositionGraph.cpp b/Engine/Source/Runtime/Renderer/Private/PostProcess/RenderingCompositionGraph.cpp
--- a/Engine/Source/Runtime/Renderer/Private/PostProcess/RenderingCompositionGraph.cpp
+++ b/Engine/Source/Runtime/Renderer/Private/PostProcess/RenderingCompositionGraph.cpp
@@ -18,6 +18,7 @@
 #include "PostProcess/SceneRenderTargets.h"
 #include "SceneRendering.h"
 #include "VisualizeTexture.h"
+#include "ScreenPass.h"
 
 #include "IImageWrapper.h"
 #include "ImageWriteQueue.h"
@@ -161,6 +162,34 @@ void CompositionGraph_OnStartFrame()
 #endif
 }
 
+const TRefCountPtr<IPooledRenderTarget>& GetFallbackTarget(EFallbackColor FallbackColor)
+{
+	switch (FallbackColor)
+	{
+	case eFC_0000: return GSystemTextures.BlackDummy;
+	case eFC_0001: return GSystemTextures.BlackAlphaOneDummy;
+	case eFC_1111: return GSystemTextures.WhiteDummy;
+	default:
+		ensure(!"Unhandled enum in EFallbackColor");
+		static const TRefCountPtr<IPooledRenderTarget> NullTarget;
+		return NullTarget;
+	}
+}
+
+const FTextureRHIRef& GetFallbackTexture(EFallbackColor FallbackColor)
+{
+	const TRefCountPtr<IPooledRenderTarget>& Target = GetFallbackTarget(FallbackColor);
+	if (Target)
+	{
+		return Target->GetRenderTargetItem().ShaderResourceTexture;
+	}
+	else
+	{
+		static const FTextureRHIRef NullTexture;
+		return NullTexture;
+	}
+}
+
 FRenderingCompositePassContext::FRenderingCompositePassContext(FRHICommandListImmediate& InRHICmdList, const FViewInfo& InView)
 	: View(InView)
 	, SceneColorViewRect(InView.ViewRect)
@@ -197,14 +226,7 @@ void FRenderingCompositePassContext::Process(const TArray<FRenderingCompositePas
 	check(!bWasProcessed);
 
 	bWasProcessed = true;
-
-	// query if we have a custom HMD post process mesh to use
-	static const auto* const HiddenAreaMaskCVar = IConsoleManager::Get().FindTConsoleVariableDataInt(TEXT("vr.HiddenAreaMask"));
-	bHasHmdMesh = (HiddenAreaMaskCVar != nullptr &&
-		HiddenAreaMaskCVar->GetValueOnRenderThread() == 1 &&
-		GEngine &&
-		GEngine->XRSystem.IsValid() && GEngine->XRSystem->GetHMDDevice() &&
-		GEngine->XRSystem->GetHMDDevice()->HasVisibleAreaMesh());
+	bHasHmdMesh = IsHMDHiddenAreaMaskActive();
 
 	if (TargetedRoots.Num() == 0)
 	{
@@ -945,16 +967,7 @@ void FPostProcessPassParameters::Set(
 		SetShaderValue(RHICmdList, ShaderRHI, SceneColorBufferUVViewport, SceneColorBufferUVViewportValue);
 	}
 
-	IPooledRenderTarget* FallbackTexture = 0;
-	
-	switch(FallbackColor)
-	{
-		case eFC_0000: FallbackTexture = GSystemTextures.BlackDummy; break;
-		case eFC_0001: FallbackTexture = GSystemTextures.BlackAlphaOneDummy; break;
-		case eFC_1111: FallbackTexture = GSystemTextures.WhiteDummy; break;
-		default:
-			ensure(!"Unhandled enum in EFallbackColor");
-	}
+	const FTextureRHIRef& FallbackTexture = GetFallbackTexture(FallbackColor);
 
 	// ePId_Input0, ePId_Input1, ...
 	for(uint32 Id = 0; Id < (uint32)ePId_Input_MAX; ++Id)
@@ -1011,7 +1024,7 @@ void FPostProcessPassParameters::Set(
 		{
 			// if the input is not there but the shader request it we give it at least some data to avoid d3ddebug errors and shader permutations
 			// to make features optional we use default black for additive passes without shader permutations
-			SetTextureParameter(RHICmdList, ShaderRHI, PostprocessInputParameter[Id], PostprocessInputParameterSampler[Id], LocalFilter, FallbackTexture->GetRenderTargetItem().TargetableTexture);
+			SetTextureParameter(RHICmdList, ShaderRHI, PostprocessInputParameter[Id], PostprocessInputParameterSampler[Id], LocalFilter, FallbackTexture);
 
 			FVector4 Dummy(1, 1, 1, 1);
 			SetShaderValue(RHICmdList, ShaderRHI, PostprocessInputSizeParameter[Id], Dummy);
@@ -1149,3 +1162,37 @@ FString FRenderingCompositePass::ConstructDebugName()
 
 	return Name;
 }
+
+FRDGTextureRef FRenderingCompositePass::CreateRDGTextureForInput(
+	FRDGBuilder& GraphBuilder,
+	EPassInputId InputId,
+	const TCHAR* InputName,
+	EFallbackColor FallbackColor)
+{
+	TRefCountPtr<IPooledRenderTarget> RenderTarget;
+
+	if (const FRenderingCompositeOutputRef* OutputRef = GetInput(InputId))
+	{
+		if (FRenderingCompositeOutput* Input = OutputRef->GetOutput())
+		{
+			RenderTarget = Input->RequestInput();
+		}
+	}
+
+	return RegisterExternalTextureWithFallback(
+		GraphBuilder,
+		RenderTarget,
+		GetFallbackTarget(FallbackColor),
+		InputName);
+}
+
+void FRenderingCompositePass::ExtractRDGTextureForOutput(
+	FRDGBuilder& GraphBuilder,
+	EPassOutputId OutputId,
+	FRDGTextureRef Texture)
+{
+	if (FRenderingCompositeOutput* Output = GetOutput(OutputId))
+	{
+		GraphBuilder.QueueTextureExtraction(Texture, &Output->PooledRenderTarget, false);
+	}
+}
\ No newline at end of file
diff --git a/Engine/Source/Runtime/Renderer/Private/PostProcess/RenderingCompositionGraph.h b/Engine/Source/Runtime/Renderer/Private/PostProcess/RenderingCompositionGraph.h
--- a/Engine/Source/Runtime/Renderer/Private/PostProcess/RenderingCompositionGraph.h
+++ b/Engine/Source/Runtime/Renderer/Private/PostProcess/RenderingCompositionGraph.h
@@ -337,6 +337,24 @@ struct FRenderingCompositePass
 
 	FString ConstructDebugName();
 
+	/**
+	 * Registers a RDG texture for the provided input. Returns a fallback color if the requested
+	 * input is null. Used for compatibility when porting to RDG.
+	 */
+	FRDGTextureRef CreateRDGTextureForInput(
+		FRDGBuilder& GraphBuilder,
+		EPassInputId InputId,
+		const TCHAR* InputName,
+		EFallbackColor FallbackColor);
+
+	/**
+	 * Registers a RDG texture to be extracted to the assigned output during graph execution.
+	 */
+	void ExtractRDGTextureForOutput(
+		FRDGBuilder& GraphBuilder,
+		EPassOutputId OutputId,
+		FRDGTextureRef Texture);
+
 	/**
 	 * Convenience method, is using other virtual methods.
 	 * @return 0 if there is an error
diff --git a/Engine/Source/Runtime/Renderer/Private/ReflectionEnvironment.cpp b/Engine/Source/Runtime/Renderer/Private/ReflectionEnvironment.cpp
--- a/Engine/Source/Runtime/Renderer/Private/ReflectionEnvironment.cpp
+++ b/Engine/Source/Runtime/Renderer/Private/ReflectionEnvironment.cpp
@@ -879,7 +879,7 @@ void FDeferredShadingSceneRenderer::RenderDeferredReflectionsAndSkyLighting(FRHI
 			}
 			else
 			{
-				const bool bCheckerboardSubsurfaceRendering = FRCPassPostProcessSubsurface::RequiresCheckerboardSubsurfaceRendering(SceneContext.GetSceneColorFormat());
+				const bool bCheckerboardSubsurfaceRendering = IsSubsurfaceCheckerboardFormat(SceneContext.GetSceneColorFormat());
 				if (bCheckerboardSubsurfaceRendering)
 				{
 					GraphicsPSOInit.BlendState = TStaticBlendState<CW_RGB, BO_Add, BF_One, BF_One>::GetRHI();
diff --git a/Engine/Source/Runtime/Renderer/Private/SceneRendering.cpp b/Engine/Source/Runtime/Renderer/Private/SceneRendering.cpp
--- a/Engine/Source/Runtime/Renderer/Private/SceneRendering.cpp
+++ b/Engine/Source/Runtime/Renderer/Private/SceneRendering.cpp
@@ -22,6 +22,7 @@
 #include "PostProcess/SceneFilterRendering.h"
 #include "PostProcess/RenderingCompositionGraph.h"
 #include "PostProcess/PostProcessEyeAdaptation.h"
+#include "PostProcess/PostProcessSubsurface.h"
 #include "CompositionLighting/CompositionLighting.h"
 #include "LegacyScreenPercentageDriver.h"
 #include "SceneViewExtension.h"
@@ -1096,7 +1097,7 @@ void FViewInfo::SetupUniformBufferParameters(
 		InPrevViewMatrices
 	);
 
-	const bool bCheckerboardSubsurfaceRendering = FRCPassPostProcessSubsurface::RequiresCheckerboardSubsurfaceRendering(SceneContext.GetSceneColorFormat());
+	const bool bCheckerboardSubsurfaceRendering = IsSubsurfaceCheckerboardFormat(SceneContext.GetSceneColorFormat());
 	ViewUniformShaderParameters.bCheckerboardSubsurfaceProfileRendering = bCheckerboardSubsurfaceRendering ? 1.0f : 0.0f;
 
 	ViewUniformShaderParameters.IndirectLightingCacheShowFlag = Family->EngineShowFlags.IndirectLightingCache;
@@ -1309,7 +1310,7 @@ void FViewInfo::SetupUniformBufferParameters(
 	ViewUniformShaderParameters.DepthOfFieldScale = FinalPostProcessSettings.DepthOfFieldScale;
 	ViewUniformShaderParameters.DepthOfFieldFocalLength = 50.0f;
 
-	ViewUniformShaderParameters.bSubsurfacePostprocessEnabled = GCompositionLighting.IsSubsurfacePostprocessRequired() ? 1.0f : 0.0f;
+	ViewUniformShaderParameters.bSubsurfacePostprocessEnabled = IsSubsurfaceEnabled() ? 1.0f : 0.0f;
 
 	{
 		// This is the CVar default
diff --git a/Engine/Source/Runtime/Renderer/Private/ScreenPass.cpp b/Engine/Source/Runtime/Renderer/Private/ScreenPass.cpp
new file mode 100644
--- /dev/null
+++ b/Engine/Source/Runtime/Renderer/Private/ScreenPass.cpp
@@ -0,0 +1,95 @@
+// Copyright 1998-2019 Epic Games, Inc. All Rights Reserved.
+
+#include "ScreenPass.h"
+#include "EngineGlobals.h"
+#include "ScenePrivate.h"
+#include "RendererModule.h"
+#include "IXRTrackingSystem.h"
+#include "IHeadMountedDisplay.h"
+
+IMPLEMENT_GLOBAL_SHADER(FScreenPassVS, "/Engine/Private/ScreenPass.usf", "ScreenPassVS", SF_Vertex);
+
+const FTextureRHIRef& GetMiniFontTexture()
+{
+	if (GEngine->MiniFontTexture)
+	{
+		return GEngine->MiniFontTexture->Resource->TextureRHI;
+	}
+	else
+	{
+		return GSystemTextures.WhiteDummy->GetRenderTargetItem().TargetableTexture;
+	}
+}
+
+bool IsHMDHiddenAreaMaskActive()
+{
+	// Query if we have a custom HMD post process mesh to use
+	static const auto* const HiddenAreaMaskCVar = IConsoleManager::Get().FindTConsoleVariableDataInt(TEXT("vr.HiddenAreaMask"));
+
+	return (HiddenAreaMaskCVar != nullptr &&
+		HiddenAreaMaskCVar->GetValueOnRenderThread() == 1 &&
+		GEngine &&
+		GEngine->XRSystem.IsValid() && GEngine->XRSystem->GetHMDDevice() &&
+		GEngine->XRSystem->GetHMDDevice()->HasVisibleAreaMesh());
+}
+
+FScreenPassCommonParameters GetScreenPassCommonParameters(FRHICommandListImmediate& RHICmdList, const FViewInfo& View)
+{
+	FSceneRenderTargets& SceneContext = FSceneRenderTargets::Get(RHICmdList);
+
+	const FIntRect ViewportRect = View.ViewRect;
+	const FIntPoint ViewportOffset = ViewportRect.Min;
+	const FIntPoint ViewportExtent = ViewportRect.Size();
+
+	FScreenPassCommonParameters Parameters;
+	Parameters.ViewportRect = ViewportRect;
+	Parameters.ViewportSize = FVector4(
+		ViewportExtent.X,
+		ViewportExtent.Y,
+		1.0f / ViewportExtent.X,
+		1.0f / ViewportExtent.Y);
+
+	Parameters.ScreenPosToPixelValue = FVector4(
+		ViewportExtent.X * 0.5f,
+		-ViewportExtent.Y * 0.5f,
+		ViewportExtent.X * 0.5f - 0.5f + ViewportOffset.X,
+		ViewportExtent.Y * 0.5f - 0.5f + ViewportOffset.Y);
+
+	Parameters.BilinearTextureSampler0 = TStaticSamplerState<SF_Bilinear>::GetRHI();
+
+	Parameters.ViewUniformBuffer = View.ViewUniformBuffer;
+	Parameters.SceneUniformBuffer = CreateSceneTextureUniformBuffer(
+		SceneContext, View.FeatureLevel, ESceneTextureSetupMode::All, EUniformBufferUsage::UniformBuffer_SingleFrame);
+
+	return Parameters;
+}
+
+FScreenPassInput GetScreenPassInputParameters(FRDGTextureRef Texture, FSamplerStateRHIParamRef SamplerState)
+{
+	check(Texture);
+	check(SamplerState);
+
+	const FVector2D Size(Texture->Desc.Extent.X, Texture->Desc.Extent.Y);
+
+	FScreenPassInput Input;
+	Input.Size = FVector4(Size.X, Size.Y, 1.0f / Size.X, 1.0f / Size.Y);
+	Input.Texture = Texture;
+	Input.Sampler = SamplerState;
+	return Input;
+}
+
+FScreenPassContext* FScreenPassContext::Create(FRHICommandListImmediate& RHICmdList, const FViewInfo& InView)
+{
+	return new(FMemStack::Get()) FScreenPassContext(RHICmdList, InView);
+}
+
+FScreenPassContext::FScreenPassContext(FRHICommandListImmediate& RHICmdList, const FViewInfo& InView)
+	: View(InView)
+	, ViewFamily(*View.Family)
+	, ViewportRect(View.ViewRect)
+	, StereoPass(View.StereoPass)
+	, bHasHMDMask(IsHMDHiddenAreaMaskActive())
+	, ShaderMap(View.ShaderMap)
+	, ScreenPassVS(View.ShaderMap)
+	, ScreenPassCommonParameters(GetScreenPassCommonParameters(RHICmdList, View))
+{}
\ No newline at end of file
diff --git a/Engine/Source/Runtime/Renderer/Private/ScreenPass.h b/Engine/Source/Runtime/Renderer/Private/ScreenPass.h
new file mode 100644
--- /dev/null
+++ b/Engine/Source/Runtime/Renderer/Private/ScreenPass.h
@@ -0,0 +1,176 @@
+// Copyright 1998-2019 Epic Games, Inc. All Rights Reserved.
+
+#pragma once
+
+#include "CoreMinimal.h"
+#include "GlobalShader.h"
+#include "RenderGraph.h"
+#include "PostProcess/SceneFilterRendering.h"
+
+// The vertex shader used by DrawScreenPass to draw a rectangle.
+class FScreenPassVS : public FGlobalShader
+{
+public:
+	DECLARE_GLOBAL_SHADER(FScreenPassVS);
+
+	static bool ShouldCompilePermutation(const FGlobalShaderPermutationParameters&)
+	{
+		return true;
+	}
+
+	FScreenPassVS() = default;
+	FScreenPassVS(const ShaderMetaType::CompiledShaderInitializerType& Initializer)
+		: FGlobalShader(Initializer)
+	{}
+};
+
+// Returns whether a HMD hidden area mask is being used for VR.
+bool IsHMDHiddenAreaMaskActive();
+
+// Returns the global engine mini font texture.
+const FTextureRHIRef& GetMiniFontTexture();
+
+// Defines the common set of parameters for a screen space pass.
+BEGIN_SHADER_PARAMETER_STRUCT(FScreenPassCommonParameters, )
+	SHADER_PARAMETER(FIntRect, ViewportRect)
+	SHADER_PARAMETER(FVector4, ViewportSize)
+	SHADER_PARAMETER(FVector4, ScreenPosToPixelValue)
+	SHADER_PARAMETER_STRUCT_REF(FSceneTexturesUniformParameters, SceneUniformBuffer)
+	SHADER_PARAMETER_STRUCT_REF(FViewUniformShaderParameters, ViewUniformBuffer)
+	SHADER_PARAMETER_SAMPLER(SamplerState, BilinearTextureSampler0)
+END_SHADER_PARAMETER_STRUCT()
+
+FScreenPassCommonParameters GetScreenPassCommonParameters(FRHICommandListImmediate& RHICmdList, const FViewInfo& View);
+
+// Defines shader parameters for a single texture input to a screen space pass.
+BEGIN_SHADER_PARAMETER_STRUCT(FScreenPassInput, )
+	SHADER_PARAMETER(FVector4, Size)
+	SHADER_PARAMETER_RDG_TEXTURE(Texture2D, Texture)
+	SHADER_PARAMETER_SAMPLER(SamplerState, Sampler)
+END_SHADER_PARAMETER_STRUCT()
+
+FScreenPassInput GetScreenPassInputParameters(FRDGTextureRef Texture, FSamplerStateRHIParamRef SamplerState);
+
+/**
+ * The context used for screen pass operations. Extracts and holds common state required
+ * by most screen passes from the provided View instance. Exists to reduce pointer
+ * chasing and reduce function parameter sizes as well as share immutable common state.
+ */
+class FScreenPassContext
+{
+public:
+	// Creates an instance of the context on the stack specifically to survive through the lifetime of the render graph.
+	static FScreenPassContext* Create(FRHICommandListImmediate& RHICmdList, const FViewInfo& InView);
+
+	// The current view instance being processed.
+	const FViewInfo& View;
+
+	// The current view family instance being processed.
+	const FSceneViewFamily& ViewFamily;
+
+	// The viewport rect for the view being processed.
+	const FIntRect ViewportRect;
+
+	// VR - Which stereo pass is being rendered.
+	const EStereoscopicPass StereoPass;
+
+	// VR - Whether a HMD hidden area mask is being used for VR.
+	const bool bHasHMDMask;
+
+	// The global shader map for the current view being processed.
+	const TShaderMap<FGlobalShaderType>* ShaderMap;
+
+	// The vertex shader used by draw screen pass. Cached here to avoid many lookups.
+	const TShaderMapRef<FScreenPassVS> ScreenPassVS;
+
+	// Common screen space parameters, filled at context creation time.
+	const FScreenPassCommonParameters ScreenPassCommonParameters;
+
+private:
+	FScreenPassContext(FRHICommandListImmediate& RHICmdList, const FViewInfo& InView);
+	FScreenPassContext(const FScreenPassContext&) = delete;
+};
+
+using FScreenPassContextRef = const FScreenPassContext*;
+
+/**
+ * Draws a full-viewport triangle with the provided pixel shader type. The destination full-viewport triangle
+ * and interpolated source UV coordinates are derived from the viewport and texture rectangles, respectively.
+ * @param ViewportRect The rectangle, in pixels, of the viewport drawn onto the destination render target.
+ * @param TextureRect The rectangle, in pixels, of the source texture mapped to the [0, 1]x[0, 1] UV space.
+ * @param TextureSize The total size, in pixels, of the source texture.
+ * @param PixelShader The pixel shader instance assigned to the pipeline state.
+ * @param PixelShaderParameters The parameter block assigned to the pixel shader prior to draw.
+ */
+template<typename TPixelShaderType>
+void DrawScreenPass(
+	FRHICommandListImmediate& RHICmdList,
+	FScreenPassContextRef Context,
+	FIntRect ViewportRect,
+	FIntRect TextureRect,
+	FIntPoint TextureSize,
+	TPixelShaderType* PixelShader,
+	const typename TPixelShaderType::FParameters& PixelShaderParameters)
+{
+	const FIntPoint ViewportSize = ViewportRect.Size();
+	const FPixelShaderRHIParamRef PixelShaderRHI = GETSAFERHISHADER_PIXEL(PixelShader);
+	FScreenPassVS* VertexShader = *(Context->ScreenPassVS);
+
+	FGraphicsPipelineStateInitializer GraphicsPSOInit;
+	RHICmdList.ApplyCachedRenderTargets(GraphicsPSOInit);
+	GraphicsPSOInit.BlendState = TStaticBlendState<>::GetRHI();
+	GraphicsPSOInit.RasterizerState = TStaticRasterizerState<>::GetRHI();
+	GraphicsPSOInit.DepthStencilState = TStaticDepthStencilState<false, CF_Always>::GetRHI();
+	GraphicsPSOInit.BoundShaderState.VertexDeclarationRHI = GFilterVertexDeclaration.VertexDeclarationRHI;
+	GraphicsPSOInit.BoundShaderState.VertexShaderRHI = GETSAFERHISHADER_VERTEX(VertexShader);
+	GraphicsPSOInit.BoundShaderState.PixelShaderRHI = PixelShaderRHI;
+	GraphicsPSOInit.PrimitiveType = PT_TriangleList;
+
+	RHICmdList.SetViewport(ViewportRect.Min.X, ViewportRect.Min.Y, 0.0f, ViewportRect.Max.X, ViewportRect.Max.Y, 1.0f);
+	SetGraphicsPipelineState(RHICmdList, GraphicsPSOInit);
+	SetShaderParameters(RHICmdList, PixelShader, PixelShaderRHI, PixelShaderParameters);
+
+	DrawPostProcessPass(
+		RHICmdList,
+		0, 0, ViewportSize.X, ViewportSize.Y,
+		TextureRect.Min.X, TextureRect.Min.Y, TextureRect.Width(), TextureRect.Height(),
+		ViewportSize,
+		TextureSize,
+		VertexShader,
+		Context->StereoPass,
+		Context->bHasHMDMask,
+		EDRF_UseTriangleOptimization);
+}
+
+/**
+ * Adds a new Render Graph pass which internally calls @ref DrawScreenPass.
+ *
+ * Parameters are forwarded to @ref DrawScreenPass.
+ */
+template <typename TPixelShaderType>
+inline void AddDrawScreenPass(
+	FRDGBuilder& GraphBuilder,
+	FRDGEventName&& PassName,
+	FScreenPassContextRef Context,
+	FIntRect ViewportRect,
+	FIntRect TextureRect,
+	FIntPoint TextureSize,
+	TPixelShaderType* PixelShader,
+	typename TPixelShaderType::FParameters* PixelShaderParameters)
+{
+	check(Context);
+	check(PixelShader);
+	check(PixelShaderParameters);
+
+	ClearUnusedGraphResources(PixelShader, PixelShaderParameters);
+
+	GraphBuilder.AddPass(
+		Forward<FRDGEventName>(PassName),
+		PixelShaderParameters,
+		ERenderGraphPassFlags::None,
+		[Context, ViewportRect, TextureRect, TextureSize, PixelShader, PixelShaderParameters]
+		(FRHICommandListImmediate& RHICmdList)
+	{
+		DrawScreenPass(RHICmdList, Context, ViewportRect, TextureRect, TextureSize, PixelShader, *PixelShaderParameters);
+	});
+}
\ No newline at end of file
