Fixed compile errors on Linux due to missing header includes

diff --git a/Engine/Source/Runtime/Renderer/Private/ScreenPass.h b/Engine/Source/Runtime/Renderer/Private/ScreenPass.h
--- a/Engine/Source/Runtime/Renderer/Private/ScreenPass.h
+++ b/Engine/Source/Runtime/Renderer/Private/ScreenPass.h
@@ -3,7 +3,9 @@
 #pragma once
 
 #include "CoreMinimal.h"
+#include "CommonRenderResources.h"
 #include "GlobalShader.h"
+#include "PipelineStateCache.h"
 #include "RenderGraph.h"
 #include "PostProcess/SceneFilterRendering.h"
 
