Fixed wrong usage of runtime flag in post process pipelines generating shaders.
We now have a permutation for that feature, only if we know if the target shader platform can support required volume texture features we need (read, render to). Otherwise we only build the non volume texture path.
Clarified pipeline and runtime checks.
Max shader permutation has been bumped.

diff --git a/Engine/Source/Runtime/Engine/Private/ShaderCompiler/ShaderCompiler.cpp b/Engine/Source/Runtime/Engine/Private/ShaderCompiler/ShaderCompiler.cpp
--- a/Engine/Source/Runtime/Engine/Private/ShaderCompiler/ShaderCompiler.cpp
+++ b/Engine/Source/Runtime/Engine/Private/ShaderCompiler/ShaderCompiler.cpp
@@ -3835,7 +3835,7 @@ void VerifyGlobalShaders(EShaderPlatform Platform, bool bLoadedFromCacheFile)
 		}
 
 		ensureMsgf(
-			PermutationCountToCompile < 200,
+			PermutationCountToCompile < 353,	// ToneMapper today can go up to 353 permutations
 			TEXT("Global shader %s has %i permutation: probably more that it needs."),
 			GlobalShaderType->GetName(), PermutationCountToCompile);
 
diff --git a/Engine/Source/Runtime/Renderer/Private/PostProcess/PostProcessCombineLUTs.cpp b/Engine/Source/Runtime/Renderer/Private/PostProcess/PostProcessCombineLUTs.cpp
--- a/Engine/Source/Runtime/Renderer/Private/PostProcess/PostProcessCombineLUTs.cpp
+++ b/Engine/Source/Runtime/Renderer/Private/PostProcess/PostProcessCombineLUTs.cpp
@@ -59,9 +59,14 @@ static TAutoConsoleVariable<int32> CVarMobileTonemapperFilm(
 
 // false:use 256x16 texture / true:use volume texture (faster, requires geometry shader)
 // USE_VOLUME_LUT: needs to be the same for C++ and HLSL
-bool UseVolumeTextureLUT(EShaderPlatform Platform) 
+bool PipelineVolumeTextureLUTMayBeSupportedAtRuntime(EShaderPlatform Platform)
 {
-	return (IsFeatureLevelSupported(Platform,ERHIFeatureLevel::SM4) && GSupportsVolumeTextureRendering && (Platform != SP_METAL_MRT && Platform != SP_METAL_MRT_MAC) && (RHISupportsGeometryShaders(Platform) || RHISupportsVertexShaderLayer(Platform)));
+	// This is used to know if the target shader platform does not support required volume texture features we need for sure (read, render to).
+	return (IsFeatureLevelSupported(Platform, ERHIFeatureLevel::SM4) && (Platform != SP_METAL_MRT && Platform != SP_METAL_MRT_MAC) && (RHISupportsGeometryShaders(Platform) || RHISupportsVertexShaderLayer(Platform)));
+}
+bool RuntimeVolumeTextureLUTSupported(EShaderPlatform Platform)	// Only valid to be called during runtime
+{
+	return PipelineVolumeTextureLUTMayBeSupportedAtRuntime(Platform) && GSupportsVolumeTextureRendering;
 }
 
 // including the neutral one at index 0
@@ -544,11 +549,18 @@ FArchive& operator<<(FArchive& Ar, FColorRemapShaderParameters& P)
 template<uint32 BlendCount>
 class FLUTBlenderPS : public FGlobalShader
 {
-	DECLARE_SHADER_TYPE(FLUTBlenderPS,Global);
+	DECLARE_GLOBAL_SHADER(FLUTBlenderPS);
+
 public:
 
+	class FBlenderUseVolumeLut : SHADER_PERMUTATION_BOOL("USE_VOLUME_LUT");
+	using FPermutationDomain = TShaderPermutationDomain<FBlenderUseVolumeLut>;
+
 	static bool ShouldCompilePermutation(const FGlobalShaderPermutationParameters& Parameters)
 	{
+		FPermutationDomain PermutationVector(Parameters.PermutationId);
+		if (!PipelineVolumeTextureLUTMayBeSupportedAtRuntime(Parameters.Platform) && PermutationVector.Get<FBlenderUseVolumeLut>())
+			return false; // We know that this platform does not support volume textures feature needed for our post process, so do not compile the permutation.
 		return true;
 	}
 
@@ -570,7 +582,6 @@ public:
 		FGlobalShader::ModifyCompilationEnvironment(Parameters, OutEnvironment);
 
 		OutEnvironment.SetDefine(TEXT("BLENDCOUNT"), BlendCount);
-		OutEnvironment.SetDefine(TEXT("USE_VOLUME_LUT"), UseVolumeTextureLUT(Parameters.Platform));
 	}
 
 	virtual bool Serialize(FArchive& Ar) override
@@ -586,12 +597,11 @@ private: // ---------------------------------------------------
 	FCombineLUTsShaderParameters<BlendCount> CombineLUTsShaderParameters;
 };
 
-
-IMPLEMENT_SHADER_TYPE(template<>,FLUTBlenderPS<1>,TEXT("/Engine/Private/PostProcessCombineLUTs.usf"),TEXT("MainPS"),SF_Pixel);
-IMPLEMENT_SHADER_TYPE(template<>,FLUTBlenderPS<2>,TEXT("/Engine/Private/PostProcessCombineLUTs.usf"),TEXT("MainPS"),SF_Pixel);
-IMPLEMENT_SHADER_TYPE(template<>,FLUTBlenderPS<3>,TEXT("/Engine/Private/PostProcessCombineLUTs.usf"),TEXT("MainPS"),SF_Pixel);
-IMPLEMENT_SHADER_TYPE(template<>,FLUTBlenderPS<4>,TEXT("/Engine/Private/PostProcessCombineLUTs.usf"),TEXT("MainPS"),SF_Pixel);
-IMPLEMENT_SHADER_TYPE(template<>,FLUTBlenderPS<5>,TEXT("/Engine/Private/PostProcessCombineLUTs.usf"),TEXT("MainPS"),SF_Pixel);
+IMPLEMENT_GLOBAL_SHADER(FLUTBlenderPS<1>, "/Engine/Private/PostProcessCombineLUTs.usf", "MainPS", SF_Pixel);
+IMPLEMENT_GLOBAL_SHADER(FLUTBlenderPS<2>, "/Engine/Private/PostProcessCombineLUTs.usf", "MainPS", SF_Pixel);
+IMPLEMENT_GLOBAL_SHADER(FLUTBlenderPS<3>, "/Engine/Private/PostProcessCombineLUTs.usf", "MainPS", SF_Pixel);
+IMPLEMENT_GLOBAL_SHADER(FLUTBlenderPS<4>, "/Engine/Private/PostProcessCombineLUTs.usf", "MainPS", SF_Pixel);
+IMPLEMENT_GLOBAL_SHADER(FLUTBlenderPS<5>, "/Engine/Private/PostProcessCombineLUTs.usf", "MainPS", SF_Pixel);
 
 /**
 * A compute shader for blending multiple LUTs together
@@ -600,11 +610,18 @@ IMPLEMENT_SHADER_TYPE(template<>,FLUTBlenderPS<5>,TEXT("/Engine/Private/PostProc
 template<uint32 BlendCount>
 class FLUTBlenderCS : public FGlobalShader
 {
-	DECLARE_SHADER_TYPE(FLUTBlenderCS,Global);
+	DECLARE_GLOBAL_SHADER(FLUTBlenderCS);
 
 public:
+
+	class FBlenderUseVolumeLut : SHADER_PERMUTATION_BOOL("USE_VOLUME_LUT");
+	using FPermutationDomain = TShaderPermutationDomain<FBlenderUseVolumeLut>;
+
 	static bool ShouldCompilePermutation(const FGlobalShaderPermutationParameters& Parameters)
 	{
+		FPermutationDomain PermutationVector(Parameters.PermutationId);
+		if (!PipelineVolumeTextureLUTMayBeSupportedAtRuntime(Parameters.Platform) && PermutationVector.Get<FBlenderUseVolumeLut>())
+			return false; // We know that this platform does not support volume textures feature needed for our post process, so do not compile the permutation.
 		return IsFeatureLevelSupported(Parameters.Platform, ERHIFeatureLevel::SM5);
 	}
 
@@ -616,7 +633,6 @@ public:
 		
 		// PS params
 		OutEnvironment.SetDefine(TEXT("BLENDCOUNT"), BlendCount);
-		OutEnvironment.SetDefine(TEXT("USE_VOLUME_LUT"), UseVolumeTextureLUT(Parameters.Platform));
 	}
 
 	FLUTBlenderCS(const ShaderMetaType::CompiledShaderInitializerType& Initializer)
@@ -672,11 +688,11 @@ private: // ---------------------------------------------------
 	FCombineLUTsShaderParameters<BlendCount> CombineLUTsShaderParameters;
 };
 
-IMPLEMENT_SHADER_TYPE(template<>,FLUTBlenderCS<1>,TEXT("/Engine/Private/PostProcessCombineLUTs.usf"),TEXT("MainCS"),SF_Compute);
-IMPLEMENT_SHADER_TYPE(template<>,FLUTBlenderCS<2>,TEXT("/Engine/Private/PostProcessCombineLUTs.usf"),TEXT("MainCS"),SF_Compute);
-IMPLEMENT_SHADER_TYPE(template<>,FLUTBlenderCS<3>,TEXT("/Engine/Private/PostProcessCombineLUTs.usf"),TEXT("MainCS"),SF_Compute);
-IMPLEMENT_SHADER_TYPE(template<>,FLUTBlenderCS<4>,TEXT("/Engine/Private/PostProcessCombineLUTs.usf"),TEXT("MainCS"),SF_Compute);
-IMPLEMENT_SHADER_TYPE(template<>,FLUTBlenderCS<5>,TEXT("/Engine/Private/PostProcessCombineLUTs.usf"),TEXT("MainCS"),SF_Compute);
+IMPLEMENT_GLOBAL_SHADER(FLUTBlenderCS<1>, "/Engine/Private/PostProcessCombineLUTs.usf", "MainCS", SF_Compute);
+IMPLEMENT_GLOBAL_SHADER(FLUTBlenderCS<2>, "/Engine/Private/PostProcessCombineLUTs.usf", "MainCS", SF_Compute);
+IMPLEMENT_GLOBAL_SHADER(FLUTBlenderCS<3>, "/Engine/Private/PostProcessCombineLUTs.usf", "MainCS", SF_Compute);
+IMPLEMENT_GLOBAL_SHADER(FLUTBlenderCS<4>, "/Engine/Private/PostProcessCombineLUTs.usf", "MainCS", SF_Compute);
+IMPLEMENT_GLOBAL_SHADER(FLUTBlenderCS<5>, "/Engine/Private/PostProcessCombineLUTs.usf", "MainCS", SF_Compute);
 
 template <typename TRHICommandList>
 static void SetLUTBlenderShader(FRenderingCompositePassContext& Context, TRHICommandList& RHICmdList, uint32 BlendCount, FTexture* Texture[], float Weights[], const FVolumeBounds& VolumeBounds, bool bUseTriangleStrip)
@@ -700,7 +716,9 @@ static void SetLUTBlenderShader(FRenderingCompositePassContext& Context, TRHICom
 #define CASE_COUNT(BlendCount) \
 	case BlendCount: \
 	{ \
-		TShaderMapRef<FLUTBlenderPS<BlendCount> > PixelShader(ShaderMap); \
+		FLUTBlenderPS<BlendCount>::FPermutationDomain PermutationVector; \
+		PermutationVector.Set<FLUTBlenderPS<BlendCount>::FBlenderUseVolumeLut>(RuntimeVolumeTextureLUTSupported(Context.GetShaderPlatform())); \
+		TShaderMapRef<FLUTBlenderPS<BlendCount>> PixelShader(ShaderMap, PermutationVector); \
 		LocalPixelShader = *PixelShader;\
 	}; \
 	break;
@@ -717,7 +735,7 @@ static void SetLUTBlenderShader(FRenderingCompositePassContext& Context, TRHICom
 		//		UE_LOG(LogRenderer, Fatal,TEXT("Invalid number of samples: %u"),BlendCount);
 	}
 #undef CASE_COUNT
-	if(UseVolumeTextureLUT(Context.View.GetShaderPlatform()))
+	if(RuntimeVolumeTextureLUTSupported(Context.View.GetShaderPlatform()))
 	{
 		TShaderMapRef<FWriteToSliceVS> VertexShader(ShaderMap);
 		TOptionalShaderMapRef<FWriteToSliceGS> GeometryShader(ShaderMap);
@@ -751,8 +769,10 @@ static void SetLUTBlenderShader(FRenderingCompositePassContext& Context, TRHICom
 #define CASE_COUNT(BlendCount) \
 	case BlendCount: \
 	{ \
-	TShaderMapRef<FLUTBlenderPS<BlendCount> > PixelShader(ShaderMap); \
-	PixelShader->SetParameters(RHICmdList, View, Texture, Weights); \
+		FLUTBlenderPS<BlendCount>::FPermutationDomain PermutationVector; \
+		PermutationVector.Set<FLUTBlenderPS<BlendCount>::FBlenderUseVolumeLut>(RuntimeVolumeTextureLUTSupported(Context.GetShaderPlatform())); \
+		TShaderMapRef<FLUTBlenderPS<BlendCount>> PixelShader(ShaderMap, PermutationVector); \
+		PixelShader->SetParameters(RHICmdList, View, Texture, Weights); \
 	}; \
 	break;
 
@@ -901,7 +921,7 @@ void FRCPassPostProcessCombineLUTs::Process(FRenderingCompositePassContext& Cont
 	SCOPED_DRAW_EVENTF(Context.RHICmdList, PostProcessCombineLUTs, TEXT("PostProcessCombineLUTs%s [%d] %dx%dx%d"),
 		bIsComputePass?TEXT("Compute"):TEXT(""), LocalCount, GLUTSize, GLUTSize, GLUTSize);
 
-	const bool bUseVolumeTextureLUT = UseVolumeTextureLUT(ShaderPlatform);
+	const bool bUseVolumeTextureLUT = RuntimeVolumeTextureLUTSupported(ShaderPlatform);
 	// for a 3D texture, the viewport is 16x16 (per slice), for a 2D texture, it's unwrapped to 256x16
 	FIntPoint DestSize(bUseVolumeTextureLUT ? GLUTSize : GLUTSize * GLUTSize, GLUTSize);
 
@@ -998,14 +1018,18 @@ void FRCPassPostProcessCombineLUTs::DispatchCS(TRHICmdList& RHICmdList, FRenderi
 {
 	auto ShaderMap = Context.GetShaderMap();
 
+	const bool bRuntimeVolumeTextureLUTSupported = RuntimeVolumeTextureLUTSupported(ShaderPlatform);
+
 	FIntPoint DestSize(DestRect.Width(), DestRect.Height());
 	uint32 GroupSizeXY = FMath::DivideAndRoundUp(DestSize.X, GCombineLUTsComputeTileSize);
-	uint32 GroupSizeZ = UseVolumeTextureLUT(ShaderPlatform) ? GroupSizeXY : 1;
+	uint32 GroupSizeZ = bRuntimeVolumeTextureLUTSupported ? GroupSizeXY : 1;
 
 #define DISPATCH_CASE(A)																			\
 	case A:																							\
 	{																								\
-		TShaderMapRef<FLUTBlenderCS<A>> ComputeShader(ShaderMap);									\
+		FLUTBlenderCS<A>::FPermutationDomain PermutationVector;										\
+		PermutationVector.Set<FLUTBlenderCS<A>::FBlenderUseVolumeLut>(bRuntimeVolumeTextureLUTSupported);\
+		TShaderMapRef<FLUTBlenderCS<A>> ComputeShader(ShaderMap, PermutationVector);				\
 		RHICmdList.SetComputeShader(ComputeShader->GetComputeShader());								\
 		ComputeShader->SetParameters(RHICmdList, Context, DestSize, DestUAV, Textures, Weights);	\
 		DispatchComputeShader(RHICmdList, *ComputeShader, GroupSizeXY, GroupSizeXY, GroupSizeZ);	\
@@ -1049,7 +1073,7 @@ FPooledRenderTargetDesc FRCPassPostProcessCombineLUTs::ComputeOutputDesc(EPassOu
 		
 		Ret = FPooledRenderTargetDesc::Create2DDesc(FIntPoint(GLUTSize * GLUTSize, GLUTSize), LUTPixelFormat, FClearValueBinding::Transparent, TexCreate_None, TexCreate_RenderTargetable | TexCreate_ShaderResource, false);
 		
-		if(UseVolumeTextureLUT(ShaderPlatform))
+		if(RuntimeVolumeTextureLUTSupported(ShaderPlatform))
 		{
 			Ret.Extent = FIntPoint(GLUTSize, GLUTSize);
 			Ret.Depth = GLUTSize;
diff --git a/Engine/Source/Runtime/Renderer/Private/PostProcess/PostProcessCombineLUTs.h b/Engine/Source/Runtime/Renderer/Private/PostProcess/PostProcessCombineLUTs.h
--- a/Engine/Source/Runtime/Renderer/Private/PostProcess/PostProcessCombineLUTs.h
+++ b/Engine/Source/Runtime/Renderer/Private/PostProcess/PostProcessCombineLUTs.h
@@ -14,7 +14,8 @@
 class UTexture;
 class FFinalPostProcessSettings;
 
-bool UseVolumeTextureLUT(EShaderPlatform Platform);
+bool PipelineVolumeTextureLUTMayBeSupportedAtRuntime(EShaderPlatform Platform);
+bool RuntimeVolumeTextureLUTSupported(EShaderPlatform Platform);
 
 // derives from TRenderingCompositePassBase<InputCount, OutputCount>
 class FRCPassPostProcessCombineLUTs : public TRenderingCompositePassBase<0, 1>
diff --git a/Engine/Source/Runtime/Renderer/Private/PostProcess/PostProcessTonemap.cpp b/Engine/Source/Runtime/Renderer/Private/PostProcess/PostProcessTonemap.cpp
--- a/Engine/Source/Runtime/Renderer/Private/PostProcess/PostProcessTonemap.cpp
+++ b/Engine/Source/Runtime/Renderer/Private/PostProcess/PostProcessTonemap.cpp
@@ -122,6 +122,7 @@ class FTonemapperGrainIntensityDim : SHADER_PERMUTATION_BOOL("USE_GRAIN_INTENSIT
 class FTonemapperVignetteDim       : SHADER_PERMUTATION_BOOL("USE_VIGNETTE");
 class FTonemapperSharpenDim        : SHADER_PERMUTATION_BOOL("USE_SHARPEN");
 class FTonemapperGrainJitterDim    : SHADER_PERMUTATION_BOOL("USE_GRAIN_JITTER");
+class FTonemapperUseVolumeLut      : SHADER_PERMUTATION_BOOL("USE_VOLUME_LUT");
 
 using FCommonDomain = TShaderPermutationDomain<
 	FTonemapperBloomDim,
@@ -129,7 +130,8 @@ using FCommonDomain = TShaderPermutationDomain<
 	FTonemapperGrainIntensityDim,
 	FTonemapperVignetteDim,
 	FTonemapperSharpenDim,
-	FTonemapperGrainJitterDim>;
+	FTonemapperGrainJitterDim,
+	FTonemapperUseVolumeLut>;
 
 FORCEINLINE bool ShouldCompileCommonPermutation(const FCommonDomain& PermutationVector)
 {
@@ -168,6 +170,8 @@ FCommonDomain BuildCommonPermutationDomain(const FViewInfo& View, bool bGammaOnl
 	PermutationVector.Set<FTonemapperGrainJitterDim>(Settings.GrainJitter > 0.0f);
 	PermutationVector.Set<FTonemapperSharpenDim>(CVarTonemapperSharpen.GetValueOnRenderThread() > 0.0f);
 
+	PermutationVector.Set<FTonemapperUseVolumeLut>(RuntimeVolumeTextureLUTSupported(View.GetShaderPlatform()));
+
 	return PermutationVector;
 }
 
@@ -217,7 +221,7 @@ FDesktopDomain RemapPermutation(FDesktopDomain PermutationVector)
 	return PermutationVector;
 }
 
-bool ShouldCompileDesktopPermutation(FDesktopDomain PermutationVector)
+bool ShouldCompileDesktopPermutation(EShaderPlatform Platform, FDesktopDomain PermutationVector)
 {
 	auto CommonPermutationVector = PermutationVector.Get<FCommonDomain>();
 
@@ -237,6 +241,9 @@ bool ShouldCompileDesktopPermutation(FDesktopDomain PermutationVector)
 			!PermutationVector.Get<FTonemapperGrainQuantizationDim>();
 	}
 
+	if (!PipelineVolumeTextureLUTMayBeSupportedAtRuntime(Platform) && CommonPermutationVector.Get<FTonemapperUseVolumeLut>())
+		return false; // We know that this platform does not support volume textures feature needed for our post process, so do not compile the permutation.
+
 	return true;
 }
 
@@ -694,12 +701,11 @@ class FPostProcessTonemapPS : public FGlobalShader
 		{
 			return false;
 		}
-		return TonemapperPermutation::ShouldCompileDesktopPermutation(FPermutationDomain(Parameters.PermutationId));
+		return TonemapperPermutation::ShouldCompileDesktopPermutation(Parameters.Platform, FPermutationDomain(Parameters.PermutationId));
 	}
 
 	static void ModifyCompilationEnvironment(const FGlobalShaderPermutationParameters& Parameters, FShaderCompilerEnvironment& OutEnvironment)
 	{
-		OutEnvironment.SetDefine(TEXT("USE_VOLUME_LUT"), UseVolumeTextureLUT(Parameters.Platform));
 	}
 
 	/** Default constructor. */
@@ -769,14 +775,13 @@ class FPostProcessTonemapCS : public FGlobalShader
 
 		FPermutationDomain PermutationVector(Parameters.PermutationId);
 
-		return TonemapperPermutation::ShouldCompileDesktopPermutation(PermutationVector.Get<TonemapperPermutation::FDesktopDomain>());
+		return TonemapperPermutation::ShouldCompileDesktopPermutation(Parameters.Platform, PermutationVector.Get<TonemapperPermutation::FDesktopDomain>());
 	}	
 	
 	static void ModifyCompilationEnvironment(const FGlobalShaderPermutationParameters& Parameters, FShaderCompilerEnvironment& OutEnvironment)
 	{
 		OutEnvironment.SetDefine(TEXT("THREADGROUP_SIZEX"), GTonemapComputeTileSizeX);
 		OutEnvironment.SetDefine(TEXT("THREADGROUP_SIZEY"), GTonemapComputeTileSizeY);
-		OutEnvironment.SetDefine(TEXT("USE_VOLUME_LUT"), UseVolumeTextureLUT(Parameters.Platform));
 	}
 
 	/** Default constructor. */
