Fixed compile error in some configs due to FRDGEventName construction in ComputeSubsurface()

diff --git a/Engine/Source/Runtime/Renderer/Private/PostProcess/PostProcessSubsurface.cpp b/Engine/Source/Runtime/Renderer/Private/PostProcess/PostProcessSubsurface.cpp
--- a/Engine/Source/Runtime/Renderer/Private/PostProcess/PostProcessSubsurface.cpp
+++ b/Engine/Source/Runtime/Renderer/Private/PostProcess/PostProcessSubsurface.cpp
@@ -482,10 +482,10 @@ FRDGTextureRef ComputeSubsurface(
 		FSamplerStateRHIParamRef SubsurfaceSamplerState = FSubsurfacePS::GetSamplerState();
 		const FSubsurfacePS::EQuality SubsurfaceQuality = FSubsurfacePS::GetQuality();
 
-		static const TCHAR* SubsurfacePassNameByDirection[] =
+		static const FRDGEventName SubsurfacePassNameByDirection[] =
 		{
-			TEXT("SubsurfaceX"),
-			TEXT("SubsurfaceY")
+			RDG_EVENT_NAME("SubsurfaceX"),
+			RDG_EVENT_NAME("SubsurfaceY")
 		};
 
 		FRDGTextureRef SubsurfacePassOutputByDirection[] =
