Fixed compile errors on Linux due to changes in FLUTBlenderCS/PS.

These shaders now use a combination of old and new permutation systems, however IMPLEMENT_GLOBAL_SHADER macro was not designed for that. Ideally, old system should be completely removed and only new system should be used.

diff --git a/Engine/Source/Runtime/Renderer/Private/PostProcess/PostProcessCombineLUTs.cpp b/Engine/Source/Runtime/Renderer/Private/PostProcess/PostProcessCombineLUTs.cpp
--- a/Engine/Source/Runtime/Renderer/Private/PostProcess/PostProcessCombineLUTs.cpp
+++ b/Engine/Source/Runtime/Renderer/Private/PostProcess/PostProcessCombineLUTs.cpp
@@ -559,8 +559,10 @@ public:
 	static bool ShouldCompilePermutation(const FGlobalShaderPermutationParameters& Parameters)
 	{
 		FPermutationDomain PermutationVector(Parameters.PermutationId);
-		if (!PipelineVolumeTextureLUTMayBeSupportedAtRuntime(Parameters.Platform) && PermutationVector.Get<FBlenderUseVolumeLut>())
+		if (!PipelineVolumeTextureLUTMayBeSupportedAtRuntime(Parameters.Platform) && PermutationVector.template Get<FBlenderUseVolumeLut>())
+		{
 			return false; // We know that this platform does not support volume textures feature needed for our post process, so do not compile the permutation.
+		}
 		return true;
 	}
 
@@ -597,11 +599,11 @@ private: // ---------------------------------------------------
 	FCombineLUTsShaderParameters<BlendCount> CombineLUTsShaderParameters;
 };
 
-IMPLEMENT_GLOBAL_SHADER(FLUTBlenderPS<1>, "/Engine/Private/PostProcessCombineLUTs.usf", "MainPS", SF_Pixel);
-IMPLEMENT_GLOBAL_SHADER(FLUTBlenderPS<2>, "/Engine/Private/PostProcessCombineLUTs.usf", "MainPS", SF_Pixel);
-IMPLEMENT_GLOBAL_SHADER(FLUTBlenderPS<3>, "/Engine/Private/PostProcessCombineLUTs.usf", "MainPS", SF_Pixel);
-IMPLEMENT_GLOBAL_SHADER(FLUTBlenderPS<4>, "/Engine/Private/PostProcessCombineLUTs.usf", "MainPS", SF_Pixel);
-IMPLEMENT_GLOBAL_SHADER(FLUTBlenderPS<5>, "/Engine/Private/PostProcessCombineLUTs.usf", "MainPS", SF_Pixel);
+template<> IMPLEMENT_GLOBAL_SHADER(FLUTBlenderPS<1>, "/Engine/Private/PostProcessCombineLUTs.usf", "MainPS", SF_Pixel);
+template<> IMPLEMENT_GLOBAL_SHADER(FLUTBlenderPS<2>, "/Engine/Private/PostProcessCombineLUTs.usf", "MainPS", SF_Pixel);
+template<> IMPLEMENT_GLOBAL_SHADER(FLUTBlenderPS<3>, "/Engine/Private/PostProcessCombineLUTs.usf", "MainPS", SF_Pixel);
+template<> IMPLEMENT_GLOBAL_SHADER(FLUTBlenderPS<4>, "/Engine/Private/PostProcessCombineLUTs.usf", "MainPS", SF_Pixel);
+template<> IMPLEMENT_GLOBAL_SHADER(FLUTBlenderPS<5>, "/Engine/Private/PostProcessCombineLUTs.usf", "MainPS", SF_Pixel);
 
 /**
 * A compute shader for blending multiple LUTs together
@@ -620,8 +622,10 @@ public:
 	static bool ShouldCompilePermutation(const FGlobalShaderPermutationParameters& Parameters)
 	{
 		FPermutationDomain PermutationVector(Parameters.PermutationId);
-		if (!PipelineVolumeTextureLUTMayBeSupportedAtRuntime(Parameters.Platform) && PermutationVector.Get<FBlenderUseVolumeLut>())
+		if (!PipelineVolumeTextureLUTMayBeSupportedAtRuntime(Parameters.Platform) && PermutationVector.template Get<FBlenderUseVolumeLut>())
+		{
 			return false; // We know that this platform does not support volume textures feature needed for our post process, so do not compile the permutation.
+		}
 		return IsFeatureLevelSupported(Parameters.Platform, ERHIFeatureLevel::SM5);
 	}
 
@@ -688,11 +692,11 @@ private: // ---------------------------------------------------
 	FCombineLUTsShaderParameters<BlendCount> CombineLUTsShaderParameters;
 };
 
-IMPLEMENT_GLOBAL_SHADER(FLUTBlenderCS<1>, "/Engine/Private/PostProcessCombineLUTs.usf", "MainCS", SF_Compute);
-IMPLEMENT_GLOBAL_SHADER(FLUTBlenderCS<2>, "/Engine/Private/PostProcessCombineLUTs.usf", "MainCS", SF_Compute);
-IMPLEMENT_GLOBAL_SHADER(FLUTBlenderCS<3>, "/Engine/Private/PostProcessCombineLUTs.usf", "MainCS", SF_Compute);
-IMPLEMENT_GLOBAL_SHADER(FLUTBlenderCS<4>, "/Engine/Private/PostProcessCombineLUTs.usf", "MainCS", SF_Compute);
-IMPLEMENT_GLOBAL_SHADER(FLUTBlenderCS<5>, "/Engine/Private/PostProcessCombineLUTs.usf", "MainCS", SF_Compute);
+template<> IMPLEMENT_GLOBAL_SHADER(FLUTBlenderCS<1>, "/Engine/Private/PostProcessCombineLUTs.usf", "MainCS", SF_Compute);
+template<> IMPLEMENT_GLOBAL_SHADER(FLUTBlenderCS<2>, "/Engine/Private/PostProcessCombineLUTs.usf", "MainCS", SF_Compute);
+template<> IMPLEMENT_GLOBAL_SHADER(FLUTBlenderCS<3>, "/Engine/Private/PostProcessCombineLUTs.usf", "MainCS", SF_Compute);
+template<> IMPLEMENT_GLOBAL_SHADER(FLUTBlenderCS<4>, "/Engine/Private/PostProcessCombineLUTs.usf", "MainCS", SF_Compute);
+template<> IMPLEMENT_GLOBAL_SHADER(FLUTBlenderCS<5>, "/Engine/Private/PostProcessCombineLUTs.usf", "MainCS", SF_Compute);
 
 template <typename TRHICommandList>
 static void SetLUTBlenderShader(FRenderingCompositePassContext& Context, TRHICommandList& RHICmdList, uint32 BlendCount, FTexture* Texture[], float Weights[], const FVolumeBounds& VolumeBounds, bool bUseTriangleStrip)
