Removing completely old style shader permutation for new one.
All post autotest ran succesfully and tested locally with up to 5 blended LUT.

diff --git a/Engine/Source/Runtime/Renderer/Private/PostProcess/PostProcessCombineLUTs.cpp b/Engine/Source/Runtime/Renderer/Private/PostProcess/PostProcessCombineLUTs.cpp
--- a/Engine/Source/Runtime/Renderer/Private/PostProcess/PostProcessCombineLUTs.cpp
+++ b/Engine/Source/Runtime/Renderer/Private/PostProcess/PostProcessCombineLUTs.cpp
@@ -97,15 +97,20 @@ struct FColorTransform
 FCombineLUTsShaderParameters
 -----------------------------------------------------------------------------*/
 
-template <uint32 BlendCount>
 class FCombineLUTsShaderParameters
 {
 public:
-	FCombineLUTsShaderParameters() {}
+	FCombineLUTsShaderParameters() : BlendCount(0)
+	{
+		// Called when serializing
+	}
 
-	FCombineLUTsShaderParameters(const FShaderParameterMap& ParameterMap)
-		: ColorRemapShaderParameters(ParameterMap)
+	FCombineLUTsShaderParameters(const FShaderParameterMap& ParameterMap, uint32 BlendCountParam)
+		: BlendCount(BlendCountParam)
+		, ColorRemapShaderParameters(ParameterMap)
 	{
+		check(BlendCount>0 && BlendCount <= GMaxLUTBlendCount);
+
 		// Suppress static code analysis warnings about a potentially ill-defined loop. BlendCount > 0 is valid.
 		CA_SUPPRESS(6294)
 
@@ -176,7 +181,7 @@ public:
 	}
 	
 	template <typename TRHICmdList, typename TRHIShader>
-	void Set(TRHICmdList& RHICmdList, const TRHIShader ShaderRHI, const FSceneView& View, FTexture* Textures[BlendCount], float Weights[BlendCount])
+	void Set(TRHICmdList& RHICmdList, const TRHIShader ShaderRHI, const FSceneView& View, FTexture** Textures, float* Weights)
 	{
 		const FPostProcessSettings& Settings = View.FinalPostProcessSettings;
 		const FSceneViewFamily& ViewFamily = *(View.Family);
@@ -402,7 +407,9 @@ public:
 
 	friend FArchive& operator<<(FArchive& Ar,FCombineLUTsShaderParameters& P)
 	{
-		for(uint32 i = 0; i < BlendCount; ++i)
+		Ar << P.BlendCount;
+
+		for(uint32 i = 0; i < P.BlendCount; ++i)
 		{
 			Ar << P.TextureParameter[i];
 			Ar << P.TextureParameterSampler[i];
@@ -426,6 +433,7 @@ public:
 	}
 
 	// [0] is not used as it's the neutral one we do in the shader
+	uint32 BlendCount;
 	FShaderResourceParameter TextureParameter[GMaxLUTBlendCount];
 	FShaderResourceParameter TextureParameterSampler[GMaxLUTBlendCount];
 	FShaderParameter WeightsParameter;
@@ -543,10 +551,7 @@ FArchive& operator<<(FArchive& Ar, FColorRemapShaderParameters& P)
 
 /**
 * A pixel shader for blending multiple LUT to one
-*
-* @param BlendCount >0
 */
-template<uint32 BlendCount>
 class FLUTBlenderPS : public FGlobalShader
 {
 	DECLARE_GLOBAL_SHADER(FLUTBlenderPS);
@@ -554,12 +559,13 @@ class FLUTBlenderPS : public FGlobalShader
 public:
 
 	class FBlenderUseVolumeLut : SHADER_PERMUTATION_BOOL("USE_VOLUME_LUT");
-	using FPermutationDomain = TShaderPermutationDomain<FBlenderUseVolumeLut>;
+	class FBlendCount : SHADER_PERMUTATION_RANGE_INT("BLENDCOUNT", 1, 5);
+	using FPermutationDomain = TShaderPermutationDomain<FBlenderUseVolumeLut, FBlendCount>;
 
 	static bool ShouldCompilePermutation(const FGlobalShaderPermutationParameters& Parameters)
 	{
 		FPermutationDomain PermutationVector(Parameters.PermutationId);
-		if (!PipelineVolumeTextureLUTMayBeSupportedAtRuntime(Parameters.Platform) && PermutationVector.template Get<FBlenderUseVolumeLut>())
+		if (!PipelineVolumeTextureLUTMayBeSupportedAtRuntime(Parameters.Platform) && PermutationVector.Get<FBlenderUseVolumeLut>())
 		{
 			return false; // We know that this platform does not support volume textures feature needed for our post process, so do not compile the permutation.
 		}
@@ -568,12 +574,12 @@ public:
 
 	FLUTBlenderPS(const ShaderMetaType::CompiledShaderInitializerType& Initializer)
 		: FGlobalShader(Initializer) 
-		, CombineLUTsShaderParameters(Initializer.ParameterMap)
+		, CombineLUTsShaderParameters(Initializer.ParameterMap, FPermutationDomain(Initializer.PermutationId).Get<FBlendCount>())
 	{
 	}
 	
 	template <typename TRHICommandList>
-	void SetParameters(TRHICommandList& RHICmdList, const FSceneView& View, FTexture* Textures[BlendCount], float Weights[BlendCount])
+	void SetParameters(TRHICommandList& RHICmdList, const FSceneView& View, FTexture** Textures, float* Weights)
 	{
 		const FPixelShaderRHIParamRef ShaderRHI = GetPixelShader();
 		CombineLUTsShaderParameters.Set(RHICmdList, ShaderRHI, View, Textures, Weights);
@@ -582,8 +588,6 @@ public:
 	static void ModifyCompilationEnvironment(const FGlobalShaderPermutationParameters& Parameters, FShaderCompilerEnvironment& OutEnvironment)
 	{
 		FGlobalShader::ModifyCompilationEnvironment(Parameters, OutEnvironment);
-
-		OutEnvironment.SetDefine(TEXT("BLENDCOUNT"), BlendCount);
 	}
 
 	virtual bool Serialize(FArchive& Ar) override
@@ -596,20 +600,14 @@ public:
 private: // ---------------------------------------------------
 	FLUTBlenderPS() {}
 
-	FCombineLUTsShaderParameters<BlendCount> CombineLUTsShaderParameters;
+	FCombineLUTsShaderParameters CombineLUTsShaderParameters;
 };
 
-template<> IMPLEMENT_GLOBAL_SHADER(FLUTBlenderPS<1>, "/Engine/Private/PostProcessCombineLUTs.usf", "MainPS", SF_Pixel);
-template<> IMPLEMENT_GLOBAL_SHADER(FLUTBlenderPS<2>, "/Engine/Private/PostProcessCombineLUTs.usf", "MainPS", SF_Pixel);
-template<> IMPLEMENT_GLOBAL_SHADER(FLUTBlenderPS<3>, "/Engine/Private/PostProcessCombineLUTs.usf", "MainPS", SF_Pixel);
-template<> IMPLEMENT_GLOBAL_SHADER(FLUTBlenderPS<4>, "/Engine/Private/PostProcessCombineLUTs.usf", "MainPS", SF_Pixel);
-template<> IMPLEMENT_GLOBAL_SHADER(FLUTBlenderPS<5>, "/Engine/Private/PostProcessCombineLUTs.usf", "MainPS", SF_Pixel);
+IMPLEMENT_GLOBAL_SHADER(FLUTBlenderPS, "/Engine/Private/PostProcessCombineLUTs.usf", "MainPS", SF_Pixel);
 
 /**
 * A compute shader for blending multiple LUTs together
-* @param BlendCount [1 - 5]
 */
-template<uint32 BlendCount>
 class FLUTBlenderCS : public FGlobalShader
 {
 	DECLARE_GLOBAL_SHADER(FLUTBlenderCS);
@@ -617,7 +615,8 @@ class FLUTBlenderCS : public FGlobalShader
 public:
 
 	class FBlenderUseVolumeLut : SHADER_PERMUTATION_BOOL("USE_VOLUME_LUT");
-	using FPermutationDomain = TShaderPermutationDomain<FBlenderUseVolumeLut>;
+	class FBlendCount : SHADER_PERMUTATION_RANGE_INT("BLENDCOUNT", 1, 5);
+	using FPermutationDomain = TShaderPermutationDomain<FBlenderUseVolumeLut, FBlendCount>;
 
 	static bool ShouldCompilePermutation(const FGlobalShaderPermutationParameters& Parameters)
 	{
@@ -634,14 +633,11 @@ public:
 		// CS params
 		FGlobalShader::ModifyCompilationEnvironment(Parameters, OutEnvironment);
 		OutEnvironment.SetDefine(TEXT("THREADGROUP_SIZE"), GCombineLUTsComputeTileSize);
-		
-		// PS params
-		OutEnvironment.SetDefine(TEXT("BLENDCOUNT"), BlendCount);
 	}
 
 	FLUTBlenderCS(const ShaderMetaType::CompiledShaderInitializerType& Initializer)
 		: FGlobalShader(Initializer) 
-		, CombineLUTsShaderParameters(Initializer.ParameterMap)
+		, CombineLUTsShaderParameters(Initializer.ParameterMap, FPermutationDomain(Initializer.PermutationId).Get<FBlendCount>())
 	{
 		// CS params
 		OutComputeTex.Bind(Initializer.ParameterMap, TEXT("OutComputeTex"));
@@ -649,7 +645,7 @@ public:
 	}
 
 	template <typename TRHICmdList>
-	void SetParameters(TRHICmdList& RHICmdList, const FRenderingCompositePassContext& Context, const FIntPoint& DestSize, FUnorderedAccessViewRHIParamRef DestUAV, FTexture* Textures[BlendCount], float Weights[BlendCount])
+	void SetParameters(TRHICmdList& RHICmdList, const FRenderingCompositePassContext& Context, const FIntPoint& DestSize, FUnorderedAccessViewRHIParamRef DestUAV, FTexture** Textures, float* Weights)
 	{
 		const FComputeShaderRHIParamRef ShaderRHI = GetComputeShader();
 
@@ -689,14 +685,10 @@ private: // ---------------------------------------------------
 	FShaderParameter CombineLUTsComputeParams;
 
 	// PS params
-	FCombineLUTsShaderParameters<BlendCount> CombineLUTsShaderParameters;
+	FCombineLUTsShaderParameters CombineLUTsShaderParameters;
 };
 
-template<> IMPLEMENT_GLOBAL_SHADER(FLUTBlenderCS<1>, "/Engine/Private/PostProcessCombineLUTs.usf", "MainCS", SF_Compute);
-template<> IMPLEMENT_GLOBAL_SHADER(FLUTBlenderCS<2>, "/Engine/Private/PostProcessCombineLUTs.usf", "MainCS", SF_Compute);
-template<> IMPLEMENT_GLOBAL_SHADER(FLUTBlenderCS<3>, "/Engine/Private/PostProcessCombineLUTs.usf", "MainCS", SF_Compute);
-template<> IMPLEMENT_GLOBAL_SHADER(FLUTBlenderCS<4>, "/Engine/Private/PostProcessCombineLUTs.usf", "MainCS", SF_Compute);
-template<> IMPLEMENT_GLOBAL_SHADER(FLUTBlenderCS<5>, "/Engine/Private/PostProcessCombineLUTs.usf", "MainCS", SF_Compute);
+IMPLEMENT_GLOBAL_SHADER(FLUTBlenderCS, "/Engine/Private/PostProcessCombineLUTs.usf", "MainCS", SF_Compute);
 
 template <typename TRHICommandList>
 static void SetLUTBlenderShader(FRenderingCompositePassContext& Context, TRHICommandList& RHICmdList, uint32 BlendCount, FTexture* Texture[], float Weights[], const FVolumeBounds& VolumeBounds, bool bUseTriangleStrip)
@@ -710,35 +702,16 @@ static void SetLUTBlenderShader(FRenderingCompositePassContext& Context, TRHICom
 
 	check(BlendCount > 0);
 
-	FShader* LocalPixelShader = 0;
 	const FViewInfo& View = Context.View;
 
 	const auto FeatureLevel = Context.GetFeatureLevel();
 	auto ShaderMap = Context.GetShaderMap();
 
-	// A macro to handle setting the filter shader for a specific number of samples.
-#define CASE_COUNT(BlendCount) \
-	case BlendCount: \
-	{ \
-		FLUTBlenderPS<BlendCount>::FPermutationDomain PermutationVector; \
-		PermutationVector.Set<FLUTBlenderPS<BlendCount>::FBlenderUseVolumeLut>(RuntimeVolumeTextureLUTSupported(Context.GetShaderPlatform())); \
-		TShaderMapRef<FLUTBlenderPS<BlendCount>> PixelShader(ShaderMap, PermutationVector); \
-		LocalPixelShader = *PixelShader;\
-	}; \
-	break;
-
-	switch(BlendCount)
-	{
-		// starts at 1 as we always have at least the neutral one
-		CASE_COUNT(1);
-		CASE_COUNT(2);
-		CASE_COUNT(3);
-		CASE_COUNT(4);
-		CASE_COUNT(5);
-		//	default:
-		//		UE_LOG(LogRenderer, Fatal,TEXT("Invalid number of samples: %u"),BlendCount);
-	}
-#undef CASE_COUNT
+	FLUTBlenderPS::FPermutationDomain PermutationVector;
+	PermutationVector.Set<FLUTBlenderPS::FBlenderUseVolumeLut>(RuntimeVolumeTextureLUTSupported(Context.GetShaderPlatform()));
+	PermutationVector.Set<FLUTBlenderPS::FBlendCount>(BlendCount);
+	TShaderMapRef<FLUTBlenderPS> PixelShader(ShaderMap, PermutationVector);
+
 	if(RuntimeVolumeTextureLUTSupported(Context.View.GetShaderPlatform()))
 	{
 		TShaderMapRef<FWriteToSliceVS> VertexShader(ShaderMap);
@@ -750,7 +723,7 @@ static void SetLUTBlenderShader(FRenderingCompositePassContext& Context, TRHICom
 #if PLATFORM_SUPPORTS_GEOMETRY_SHADERS
 		GraphicsPSOInit.BoundShaderState.GeometryShaderRHI = GETSAFERHISHADER_GEOMETRY(*GeometryShader);
 #endif
-		GraphicsPSOInit.BoundShaderState.PixelShaderRHI = GETSAFERHISHADER_PIXEL(LocalPixelShader);
+		GraphicsPSOInit.BoundShaderState.PixelShaderRHI = GETSAFERHISHADER_PIXEL(*PixelShader);
 		SetGraphicsPipelineState(RHICmdList, GraphicsPSOInit);
 
 		VertexShader->SetParameters(RHICmdList, VolumeBounds, FIntVector(VolumeBounds.MaxX - VolumeBounds.MinX));
@@ -765,33 +738,13 @@ static void SetLUTBlenderShader(FRenderingCompositePassContext& Context, TRHICom
 
 		GraphicsPSOInit.BoundShaderState.VertexDeclarationRHI = GFilterVertexDeclaration.VertexDeclarationRHI;
 		GraphicsPSOInit.BoundShaderState.VertexShaderRHI = GETSAFERHISHADER_VERTEX(*VertexShader);
-		GraphicsPSOInit.BoundShaderState.PixelShaderRHI = GETSAFERHISHADER_PIXEL(LocalPixelShader);
+		GraphicsPSOInit.BoundShaderState.PixelShaderRHI = GETSAFERHISHADER_PIXEL(*PixelShader);
 		SetGraphicsPipelineState(RHICmdList, GraphicsPSOInit);
 
 		VertexShader->SetParameters(Context);
 	}
-#define CASE_COUNT(BlendCount) \
-	case BlendCount: \
-	{ \
-		FLUTBlenderPS<BlendCount>::FPermutationDomain PermutationVector; \
-		PermutationVector.Set<FLUTBlenderPS<BlendCount>::FBlenderUseVolumeLut>(RuntimeVolumeTextureLUTSupported(Context.GetShaderPlatform())); \
-		TShaderMapRef<FLUTBlenderPS<BlendCount>> PixelShader(ShaderMap, PermutationVector); \
-		PixelShader->SetParameters(RHICmdList, View, Texture, Weights); \
-	}; \
-	break;
-
-	switch(BlendCount)
-	{
-		// starts at 1 as we always have at least the neutral one
-		CASE_COUNT(1);
-		CASE_COUNT(2);
-		CASE_COUNT(3);
-		CASE_COUNT(4);
-		CASE_COUNT(5);
-		//	default:
-		//		UE_LOG(LogRenderer, Fatal,TEXT("Invalid number of samples: %u"),BlendCount);
-	}
-#undef CASE_COUNT
+
+	PixelShader->SetParameters(RHICmdList, View, &Texture[0], &Weights[0]);
 }
 
 
@@ -1028,31 +981,14 @@ void FRCPassPostProcessCombineLUTs::DispatchCS(TRHICmdList& RHICmdList, FRenderi
 	uint32 GroupSizeXY = FMath::DivideAndRoundUp(DestSize.X, GCombineLUTsComputeTileSize);
 	uint32 GroupSizeZ = bRuntimeVolumeTextureLUTSupported ? GroupSizeXY : 1;
 
-#define DISPATCH_CASE(A)																			\
-	case A:																							\
-	{																								\
-		FLUTBlenderCS<A>::FPermutationDomain PermutationVector;										\
-		PermutationVector.Set<FLUTBlenderCS<A>::FBlenderUseVolumeLut>(bRuntimeVolumeTextureLUTSupported);\
-		TShaderMapRef<FLUTBlenderCS<A>> ComputeShader(ShaderMap, PermutationVector);				\
-		RHICmdList.SetComputeShader(ComputeShader->GetComputeShader());								\
-		ComputeShader->SetParameters(RHICmdList, Context, DestSize, DestUAV, Textures, Weights);	\
-		DispatchComputeShader(RHICmdList, *ComputeShader, GroupSizeXY, GroupSizeXY, GroupSizeZ);	\
-		ComputeShader->UnsetParameters(RHICmdList);													\
-	}																								\
-	break;
-
-	switch(BlendCount)
-	{
-	DISPATCH_CASE(1)
-	DISPATCH_CASE(2)
-	DISPATCH_CASE(3)
-	DISPATCH_CASE(4)
-	DISPATCH_CASE(5)
-	default:
-		check(0);
-	}
-
-#undef DISPATCH_CASE
+	FLUTBlenderCS::FPermutationDomain PermutationVector;
+	PermutationVector.Set<FLUTBlenderCS::FBlenderUseVolumeLut>(bRuntimeVolumeTextureLUTSupported);
+	PermutationVector.Set<FLUTBlenderCS::FBlendCount>(BlendCount);
+	TShaderMapRef<FLUTBlenderCS> ComputeShader(ShaderMap, PermutationVector);
+	RHICmdList.SetComputeShader(ComputeShader->GetComputeShader());
+	ComputeShader->SetParameters(RHICmdList, Context, DestSize, DestUAV, Textures, Weights);
+	DispatchComputeShader(RHICmdList, *ComputeShader, GroupSizeXY, GroupSizeXY, GroupSizeZ);
+	ComputeShader->UnsetParameters(RHICmdList);
 }
 
 FPooledRenderTargetDesc FRCPassPostProcessCombineLUTs::ComputeOutputDesc(EPassOutputId InPassOutputId) const
