Temporary disabling pre-exposure through the override value, while figuring out how to fix bloom.

diff --git a/Engine/Source/Runtime/Renderer/Private/PostProcess/PostProcessEyeAdaptation.cpp b/Engine/Source/Runtime/Renderer/Private/PostProcess/PostProcessEyeAdaptation.cpp
--- a/Engine/Source/Runtime/Renderer/Private/PostProcess/PostProcessEyeAdaptation.cpp
+++ b/Engine/Source/Runtime/Renderer/Private/PostProcess/PostProcessEyeAdaptation.cpp
@@ -19,7 +19,7 @@ RENDERCORE_API bool UsePreExposure(EShaderPlatform Platform);
 // Initialize the static CVar
 TAutoConsoleVariable<float> CVarEyeAdaptationPreExposureOverride(
 	TEXT("r.EyeAdaptation.PreExposureOverride"),
-	0,
+	1,
 	TEXT("Overide the scene pre-exposure by a custom value. \n")
 	TEXT("= 0 : No override\n")
 	TEXT("> 0 : Override PreExposure\n"),
