Allows ComputeOutputDesc() to return FPooledRenderTargetDesc() and let have ExtractRDGTextureForOutput() actually set the descriptor from the FRDGTexture to avoid decsriptor setup code duplication.

diff --git a/Engine/Source/Runtime/Renderer/Private/PostProcess/RenderingCompositionGraph.cpp b/Engine/Source/Runtime/Renderer/Private/PostProcess/RenderingCompositionGraph.cpp
--- a/Engine/Source/Runtime/Renderer/Private/PostProcess/RenderingCompositionGraph.cpp
+++ b/Engine/Source/Runtime/Renderer/Private/PostProcess/RenderingCompositionGraph.cpp
@@ -345,33 +345,6 @@ void FRenderingCompositionGraph::RecursivelyGatherDependencies(FRenderingComposi
 			RecursivelyGatherDependencies(OutputRefItPass);
 		}
 	}
-
-	// the pass is asked what the intermediate surface/texture format needs to be for all its outputs.
-	for(uint32 OutputId = 0; ; ++OutputId)
-	{
-		EPassOutputId PassOutputId = (EPassOutputId)(OutputId);
-		FRenderingCompositeOutput* Output = Pass->GetOutput(PassOutputId);
-
-		if(!Output)
-		{
-			break;
-		}
-
-		Output->RenderTargetDesc = Pass->ComputeOutputDesc(PassOutputId);
-
-		// Allow format overrides for high-precision work
-		static const auto CVarPostProcessingColorFormat = IConsoleManager::Get().FindTConsoleVariableDataInt(TEXT("r.PostProcessingColorFormat"));
-
-		if (CVarPostProcessingColorFormat->GetValueOnRenderThread() == 1)
-		{
-			if (Output->RenderTargetDesc.Format == PF_FloatRGBA ||
-				Output->RenderTargetDesc.Format == PF_FloatRGB ||
-				Output->RenderTargetDesc.Format == PF_FloatR11G11B10)
-			{
-				Output->RenderTargetDesc.Format = PF_A32B32G32R32F;
-			}
-		}
-	}
 }
 
 TUniquePtr<FImagePixelData> FRenderingCompositionGraph::GetDumpOutput(FRenderingCompositePassContext& Context, FIntRect SourceRect, FRenderingCompositeOutput* Output) const
@@ -537,6 +510,44 @@ void FRenderingCompositionGraph::RecursivelyProcess(const FRenderingCompositeOut
 			}
 		}
 	}
+	
+	// Requests the output render target descriptors.
+	for(uint32 OutputId = 0; ; ++OutputId)
+	{
+		EPassOutputId PassOutputId = (EPassOutputId)(OutputId);
+		FRenderingCompositeOutput* PassOutput = Pass->GetOutput(PassOutputId);
+
+		if(!PassOutput)
+		{
+			break;
+		}
+
+		PassOutput->RenderTargetDesc = Pass->ComputeOutputDesc(PassOutputId);
+
+		// Allow format overrides for high-precision work
+		static const auto CVarPostProcessingColorFormat = IConsoleManager::Get().FindTConsoleVariableDataInt(TEXT("r.PostProcessingColorFormat"));
+
+		if (CVarPostProcessingColorFormat->GetValueOnRenderThread() == 1)
+		{
+			if (PassOutput->RenderTargetDesc.Format == PF_FloatRGBA ||
+				PassOutput->RenderTargetDesc.Format == PF_FloatRGB ||
+				PassOutput->RenderTargetDesc.Format == PF_FloatR11G11B10)
+			{
+				PassOutput->RenderTargetDesc.Format = PF_A32B32G32R32F;
+			}
+		}
+	}
+
+	// Execute the pass straight away to have any update on the output decsriptors ExtractRDGTextureForOutput() might do.
+	{
+		Context.Pass = Pass;
+		Context.SetViewportInvalid();
+
+		// then process the pass itself
+		check(!Context.RHICmdList.IsInsideRenderPass());
+		Pass->Process(Context);
+		check(!Context.RHICmdList.IsInsideRenderPass());
+	}
 
 #if !(UE_BUILD_SHIPPING || UE_BUILD_TEST)
 	if(ShouldDebugCompositionGraph())
@@ -709,14 +720,6 @@ void FRenderingCompositionGraph::RecursivelyProcess(const FRenderingCompositeOut
 	}
 #endif
 
-	Context.Pass = Pass;
-	Context.SetViewportInvalid();
-
-	// then process the pass itself
-	check(!Context.RHICmdList.IsInsideRenderPass());
-	Pass->Process(Context);
-	check(!Context.RHICmdList.IsInsideRenderPass());
-
 	// for VisualizeTexture and output buffer dumping
 #if !(UE_BUILD_SHIPPING || UE_BUILD_TEST)
 	{
@@ -1193,6 +1196,7 @@ void FRenderingCompositePass::ExtractRDGTextureForOutput(
 {
 	if (FRenderingCompositeOutput* Output = GetOutput(OutputId))
 	{
+		Output->RenderTargetDesc = Texture->Desc;
 		GraphBuilder.QueueTextureExtraction(Texture, &Output->PooledRenderTarget, false);
 	}
-}
\ No newline at end of file
+}
