Refactors TAA to RDG

diff --git a/Engine/Shaders/Private/PostProcessTemporalAA.usf b/Engine/Shaders/Private/PostProcessTemporalAA.usf
--- a/Engine/Shaders/Private/PostProcessTemporalAA.usf
+++ b/Engine/Shaders/Private/PostProcessTemporalAA.usf
@@ -5,7 +5,6 @@
 =============================================================================*/
 
 #include "Common.ush"
-#include "PostProcessCommon.ush"
 #include "Random.ush"
 #include "EyeAdaptationCommon.ush"
 #include "TextureSampling.ush"
@@ -469,6 +468,18 @@ float4 MaxViewportUVAndSvPositionToViewportUV;
 float2 ScreenPosAbsMax;
 float4 ScreenPosToHistoryBufferUV;
 
+float4 InputSceneColorSize;
+Texture2D InputSceneColor;
+SamplerState InputSceneColorSampler;
+Texture2D InputSceneMetadata;
+SamplerState InputSceneMetadataSampler;
+
+Texture2D SceneDepthBuffer;
+SamplerState SceneDepthBufferSampler;
+
+Texture2D SceneVelocityBuffer;
+SamplerState SceneVelocityBufferSampler;
+
 Texture2D HistoryBuffer0;
 Texture2D HistoryBuffer1;
 SamplerState HistoryBuffer0Sampler;
@@ -476,6 +487,9 @@ SamplerState HistoryBuffer1Sampler;
 float4 HistoryBufferSize;
 float4 HistoryBufferUVMinMax;
 
+float4 OutputViewportSize;
+float4 OutputViewportRect;
+
 #if COMPUTESHADER
 RWTexture2D<float4> OutComputeTex0;
 
@@ -500,8 +514,8 @@ float4 InputViewSize;
 // Temporal jitter at the pixel scale.
 float2 TemporalJitterPixels;
 
-// x = screen percentage; y = upscale factor = 1 / screen percentage.
-float2 ScreenPercentageAndUpscaleFactor;
+float ScreenPercentage;
+float UpscaleFactor; // = 1 / ScreenPercentage
 
 #endif // AA_UPSAMPLE
 
@@ -840,12 +854,12 @@ float GetSceneColorHdrWeight(
 		UNROLL
 		for (uint i = 0; i < kNeighborsCount; i++)
 		{
-			InputParams.CachedNeighbors0[i] = TransformCurrentFrameSceneColor(PostprocessInput0.SampleLevel(
-				PostprocessInput0Sampler, InputParams.NearestBufferUV, 0, kOffsets3x3[i]));
+			InputParams.CachedNeighbors0[i] = TransformCurrentFrameSceneColor(InputSceneColor.SampleLevel(
+				InputSceneColorSampler, InputParams.NearestBufferUV, 0, kOffsets3x3[i]));
 	
 			#if HISTORY_RENDER_TARGETS == 2
-				InputParams.CachedNeighbors1[i] = PostprocessInput1.SampleLevel(
-					PostprocessInput1Sampler, InputParams.NearestBufferUV, 0, kOffsets3x3[i]).r;
+				InputParams.CachedNeighbors1[i] = InputSceneMetadata.SampleLevel(
+					InputSceneMetadataSampler, InputParams.NearestBufferUV, 0, kOffsets3x3[i]).r;
 			#endif
 		}
 	}
@@ -1005,7 +1019,7 @@ float2 GetGroupThread0InputPixelCoord(in FTAAInputParameters InputParams)
 	float2 Thread0SvPosition = InputParams.GroupId * uint2(THREADGROUP_SIZEX, THREADGROUP_SIZEY) + 0.5;
 		
 	// Output pixel's viewport UV group thread index 0.
-	float2 Thread0ViewportUV = Thread0SvPosition * ViewportSize.zw;
+	float2 Thread0ViewportUV = Thread0SvPosition * OutputViewportSize.zw;
 
 	// Pixel coordinate of the center of output pixel O in the input viewport.
 	float2 Thread0PPCo = Thread0ViewportUV * InputViewSize.xy + TemporalJitterPixels;
@@ -1030,7 +1044,7 @@ uint2 GetGroupTileTexelOffset(in FTAAInputParameters InputParams, uint TileBorde
 	}
 	#else // !AA_UPSAMPLE
 	{
-		return ViewportRect.xy + InputParams.GroupId * uint2(THREADGROUP_SIZEX, THREADGROUP_SIZEY) - TileBorderSize;
+		return OutputViewportRect.xy + InputParams.GroupId * uint2(THREADGROUP_SIZEX, THREADGROUP_SIZEY) - TileBorderSize;
 	}
 	#endif
 }
@@ -1043,7 +1057,7 @@ uint GetTileArrayIndexFromPixelOffset(in FTAAInputParameters InputParams, int2 P
 		const float2 RowMultiplier = float2(1, TileBorderSize * 2 + LDS_BASE_TILE_WIDTH);
 
 		float2 Thread0PPCk = GetGroupThread0InputPixelCoord(InputParams);
-		float2 PPCk = InputParams.NearestBufferUV * PostprocessInput0Size.xy;
+		float2 PPCk = InputParams.NearestBufferUV * InputSceneColorSize.xy;
 
 		float2 TilePos = floor(PPCk) - floor(Thread0PPCk);
 		return uint(dot(TilePos, RowMultiplier) + dot(float2(PixelOffset) + float(TileBorderSize), RowMultiplier));
@@ -1083,7 +1097,7 @@ void PrecacheInputSceneDepthToLDS(in FTAAInputParameters InputParams)
 				(i != LoadCount - 1) ||
 				(LDS_DEPTH_ARRAY_SIZE % THREADGROUP_TOTAL) == 0)
 			{
-				GroupSharedArray[LinearGroupThreadId] = SceneTexturesStruct.SceneDepthTexture.Load(uint3(TexelLocation, 0)).x;
+				GroupSharedArray[LinearGroupThreadId] = SceneDepthBuffer.Load(uint3(TexelLocation, 0)).x;
 			}
 			LinearGroupThreadId += THREADGROUP_TOTAL;
 		}
@@ -1112,8 +1126,8 @@ void PrecacheInputSceneDepthToLDS(in FTAAInputParameters InputParams)
 				(i != LoadCount - 1) ||
 				((LDS_DEPTH_ARRAY_SIZE / 4) % THREADGROUP_TOTAL) == 0)
 			{
-				float2 UV = float2(TexelLocation + 0.5) * PostprocessInput0Size.zw;
-				float4 Depth = SceneTexturesStruct.SceneDepthTexture.Gather(SceneTexturesStruct.SceneDepthTextureSampler, UV);
+				float2 UV = float2(TexelLocation + 0.5) * InputSceneColorSize.zw;
+				float4 Depth = SceneDepthBuffer.Gather(SceneDepthBufferSampler, UV);
 				GroupSharedArray[DestI + 1 * LDS_DEPTH_TILE_WIDTH + 0] = Depth.x;
 				GroupSharedArray[DestI + 1 * LDS_DEPTH_TILE_WIDTH + 1] = Depth.y;
 				GroupSharedArray[DestI + 0 * LDS_DEPTH_TILE_WIDTH + 1] = Depth.z;
@@ -1189,15 +1203,15 @@ void PrecacheInputSceneColorToLDS(in FTAAInputParameters InputParams)
 			(LDS_COLOR_ARRAY_SIZE % THREADGROUP_TOTAL) == 0)
 		{
 			#if LDS_FLOAT_UV
-				float2 BufferUV = TexelLocation * PostprocessInput0Size.zw;
+				float2 BufferUV = TexelLocation * InputSceneColorSize.zw;
 
 				#if AA_BORDER == 2
 					BufferUV = clamp(BufferUV, View.BufferBilinearUVMinMax.xy, View.BufferBilinearUVMinMax.zw);
 				#endif
 
-				float4 RawColor = PostprocessInput0.SampleLevel(PostprocessInput0Sampler, BufferUV, 0);
+				float4 RawColor = InputSceneColor.SampleLevel(InputSceneColorSampler, BufferUV, 0);
 			#else
-				float4 RawColor = PostprocessInput0.Load(uint3(TexelLocation, 0));
+				float4 RawColor = InputSceneColor.Load(uint3(TexelLocation, 0));
 			#endif
 
 			float4 Color = TransformCurrentFrameSceneColor(RawColor);
@@ -1312,7 +1326,7 @@ void PrecacheInputSceneDepth(in FTAAInputParameters InputParams)
 // Sample scene color.
 float SampleCachedSceneDepthTexture(in FTAAInputParameters InputParams, int2 PixelOffset)
 {
-	return SceneTexturesStruct.SceneDepthTexture.SampleLevel(SceneTexturesStruct.SceneDepthTextureSampler, InputParams.NearestBufferUV, 0, PixelOffset).r;
+	return SceneDepthBuffer.SampleLevel(SceneDepthBufferSampler, InputParams.NearestBufferUV, 0, PixelOffset).r;
 }
 
 
@@ -1335,16 +1349,16 @@ FTAASceneColorSample SampleCachedSceneColorTexture(
 	int2 PixelOffset)
 {
 	FTAASceneColorSample Sample;
-	Sample.Color = TransformCurrentFrameSceneColor(PostprocessInput0.SampleLevel(
-		PostprocessInput0Sampler, InputParams.NearestBufferUV, 0, PixelOffset));
+	Sample.Color = TransformCurrentFrameSceneColor(InputSceneColor.SampleLevel(
+		InputSceneColorSampler, InputParams.NearestBufferUV, 0, PixelOffset));
 
 	Sample.CocRadius = 0;
 
 	#if AA_HISTORY_PAYLOAD == HISTORY_PAYLOAD_RGB_COC || AA_HISTORY_PAYLOAD == HISTORY_PAYLOAD_RGB_COC_LEGACY
 		Sample.CocRadius = Sample.Color.a;
 	#elif AA_HISTORY_PAYLOAD == HISTORY_PAYLOAD_RGB_OPACITY_COC
-		Sample.CocRadius = PostprocessInput1.SampleLevel(
-			PostprocessInput1Sampler, InputParams.NearestBufferUV, 0, PixelOffset).r;
+		Sample.CocRadius = InputSceneMetadata.SampleLevel(
+			InputSceneMetadataSampler, InputParams.NearestBufferUV, 0, PixelOffset).r;
 	#endif
 
 	Sample.HdrWeight = GetSceneColorHdrWeight(InputParams, Sample.Color);
@@ -1362,7 +1376,7 @@ FTAASceneColorSample SampleCachedSceneColorTexture(
 // Returns the weight of a pixels at a coordinate <PixelDelta> from the PDF highest point.
 float ComputeSampleWeigth(in FTAAIntermediaryResult IntermediaryResult, float2 PixelDelta)
 {
-	float u2 = ScreenPercentageAndUpscaleFactor.y * ScreenPercentageAndUpscaleFactor.y;
+	float u2 = UpscaleFactor * UpscaleFactor;
 
 	// The point of InvFilterScaleFactor is to blur current frame scene color when upscaling.
 	// Therefore there is no need to do it when downscaling.
@@ -1391,7 +1405,7 @@ float ComputeSampleWeigth(in FTAAIntermediaryResult IntermediaryResult, float2 P
 // Returns the weight of a pixels at a coordinate <PixelDelta> from the PDF highest point.
 float ComputePixelWeigth(in FTAAIntermediaryResult IntermediaryResult, float2 PixelDelta)
 {
-	float u2 = ScreenPercentageAndUpscaleFactor.y * ScreenPercentageAndUpscaleFactor.y;
+	float u2 = UpscaleFactor * UpscaleFactor;
 
 	// The point of InvFilterScaleFactor is to blur current frame scene color when upscaling.
 	// Therefore there is no need to do it when downscaling.
@@ -1595,7 +1609,7 @@ void FilterCurrentFrameInputSamples(
 			#error AA_BORDER == 1 does not work for TAA upsample because also needs to be done in color clamping box.
 		#else
 			// Use unfiltered for 1 pixel border.
-			float2 TestPos = abs(InputParams.ScreenPos) + ViewportSize.zw * 2.0;
+			float2 TestPos = abs(InputParams.ScreenPos) + OutputViewportSize.zw * 2.0;
 			
 			if(max(TestPos.x, TestPos.y) >= 1.0) 
 			{
@@ -1671,7 +1685,7 @@ void ComputeNeighborhoodBoundingbox(
 		NeighborMax = Neighbors[4];
 		
 		// Reduce distance threshold as upsacale factor increase to reduce ghosting.
-		float DistthresholdLerp = ScreenPercentageAndUpscaleFactor.y - 1;
+		float DistthresholdLerp = UpscaleFactor - 1;
 		float DistThreshold = lerp(1.51, 1.3, DistthresholdLerp);
 
 		#if AA_SAMPLES == 9
@@ -1902,8 +1916,8 @@ FTAAHistoryPayload TemporalAASample(uint2 GroupId, uint2 GroupThreadId, uint Gro
 		{
 			const uint kResponsiveStencilMask = 1 << 3;
 			
-			int2 SceneStencilUV = (int2)trunc(InputParams.NearestBufferUV * PostprocessInput0Size.xy);
-			uint SceneStencilRef = SceneTexturesStruct.SceneStencilTexture.Load(int3(SceneStencilUV, 0)) STENCIL_COMPONENT_SWIZZLE;
+			int2 SceneStencilUV = (int2)trunc(InputParams.NearestBufferUV * InputSceneColorSize.xy);
+			uint SceneStencilRef = SceneDepthBuffer.Load(int3(SceneStencilUV, 0)) STENCIL_COMPONENT_SWIZZLE;
 
 			InputParams.bIsResponsiveAAPixel = (SceneStencilRef & kResponsiveStencilMask) ? 1.f : 0.f;
 		}
@@ -1924,8 +1938,8 @@ FTAAHistoryPayload TemporalAASample(uint2 GroupId, uint2 GroupThreadId, uint Gro
 			// Pixel coordinate of the center of the nearest top left input pixel T.
 			float2 PPCt = floor(PPCo - 0.5) + 0.5;
 		
-			InputParams.NearestBufferUV = PostprocessInput0Size.zw * (InputViewMin + PPCk);
-			InputParams.NearestTopLeftBufferUV = PostprocessInput0Size.zw * (InputViewMin + PPCt);
+			InputParams.NearestBufferUV = InputSceneColorSize.zw * (InputViewMin + PPCk);
+			InputParams.NearestTopLeftBufferUV = InputSceneColorSize.zw * (InputViewMin + PPCt);
 
 			// TODO: because use nearest sampler, can be faster in this computation.
 		}
@@ -1981,9 +1995,9 @@ FTAAHistoryPayload TemporalAASample(uint2 GroupId, uint2 GroupThreadId, uint Gro
 				// This is offset for reading from velocity texture.
 				// This supports half or fractional resolution velocity textures.
 				// With the assumption that UV position scales between velocity and color.
-				VelocityOffset = DepthOffset * PostprocessInput0Size.zw;
+				VelocityOffset = DepthOffset * InputSceneColorSize.zw;
 				// This is [0 to 1] flipped in Y.
-				//PosN.xy = ScreenPos + DepthOffset * ViewportSize.zw * 2.0;
+				//PosN.xy = ScreenPos + DepthOffset * OutputViewportSize.zw * 2.0;
 				PosN.z = DepthsXYZW;
 			}
 		#else // !HAS_INVERTED_Z_BUFFER
@@ -2005,18 +2019,18 @@ FTAAHistoryPayload TemporalAASample(uint2 GroupId, uint2 GroupThreadId, uint Gro
 		float2 PrevScreen = PrevClip.xy / PrevClip.w;
 		float2 BackN = PosN.xy - PrevScreen;
 
-		float2 BackTemp = BackN * ViewportSize.xy;
+		float2 BackTemp = BackN * OutputViewportSize.xy;
 
 		#if AA_DYNAMIC
 		{
 			float2 VelocityN;
-			VelocityN = PostprocessInput2.SampleLevel(PostprocessInput2Sampler, InputParams.NearestBufferUV + VelocityOffset, 0).xy;
+			VelocityN = SceneVelocityBuffer.SampleLevel(SceneVelocityBufferSampler, InputParams.NearestBufferUV + VelocityOffset, 0).xy;
 			bool DynamicN = VelocityN.x > 0.0;
 			if(DynamicN)
 			{
 				BackN = DecodeVelocityFromTexture(VelocityN);
 			}
-			BackTemp = BackN * ViewportSize.xy;
+			BackTemp = BackN * OutputViewportSize.xy;
 		}
 		#endif
 
@@ -2076,11 +2090,11 @@ FTAAHistoryPayload TemporalAASample(uint2 GroupId, uint2 GroupThreadId, uint Gro
 			#error AA_DYNAMIC_ANTIGHOST requires AA_DYNAMIC
 		#endif
 		// TODO: try a 2x2 for AA_UPSAMPLE
-		bool Dynamic1 = PostprocessInput2.SampleLevel(PostprocessInput2Sampler, InputParams.NearestBufferUV, 0, int2( 0, -1)).x > 0;
-		bool Dynamic3 = PostprocessInput2.SampleLevel(PostprocessInput2Sampler, InputParams.NearestBufferUV, 0, int2(-1,  0)).x > 0;
-		Dynamic4 = PostprocessInput2.SampleLevel(PostprocessInput2Sampler, InputParams.NearestBufferUV, 0).x > 0;
-		bool Dynamic5 = PostprocessInput2.SampleLevel(PostprocessInput2Sampler, InputParams.NearestBufferUV, 0, int2( 1,  0)).x > 0;
-		bool Dynamic7 = PostprocessInput2.SampleLevel(PostprocessInput2Sampler, InputParams.NearestBufferUV, 0, int2( 0,  1)).x > 0;
+		bool Dynamic1 = SceneVelocityBuffer.SampleLevel(SceneVelocityBufferSampler, InputParams.NearestBufferUV, 0, int2( 0, -1)).x > 0;
+		bool Dynamic3 = SceneVelocityBuffer.SampleLevel(SceneVelocityBufferSampler, InputParams.NearestBufferUV, 0, int2(-1,  0)).x > 0;
+		Dynamic4 = SceneVelocityBuffer.SampleLevel(SceneVelocityBufferSampler, InputParams.NearestBufferUV, 0).x > 0;
+		bool Dynamic5 = SceneVelocityBuffer.SampleLevel(SceneVelocityBufferSampler, InputParams.NearestBufferUV, 0, int2( 1,  0)).x > 0;
+		bool Dynamic7 = SceneVelocityBuffer.SampleLevel(SceneVelocityBufferSampler, InputParams.NearestBufferUV, 0, int2( 0,  1)).x > 0;
 
 		bool Dynamic = Dynamic1 || Dynamic3 || Dynamic4 || Dynamic5 || Dynamic7;
 		IgnoreHistory = IgnoreHistory || !Dynamic && History.Color.a > 0;
@@ -2111,10 +2125,10 @@ FTAAHistoryPayload TemporalAASample(uint2 GroupId, uint2 GroupThreadId, uint Gro
 
 		// Blur input according to input pixel velocity to reduce blocky filtering cause by narrow filter on low screen percentage.
 		// Multiplied by upscale factor because Velocity is in output resolution based pixel velocity.
-		IntermediaryResult.InvFilterScaleFactor -= (Velocity * ScreenPercentageAndUpscaleFactor.y) * 0.1;
+		IntermediaryResult.InvFilterScaleFactor -= (Velocity * UpscaleFactor) * 0.1;
 
 		// Set a minimal filtering scale to screen percentage to not unecessarily blur the input more than the screen percentage.
-		IntermediaryResult.InvFilterScaleFactor = max(IntermediaryResult.InvFilterScaleFactor, ScreenPercentageAndUpscaleFactor.x);
+		IntermediaryResult.InvFilterScaleFactor = max(IntermediaryResult.InvFilterScaleFactor, ScreenPercentage);
 
 		FilterCurrentFrameInputSamples(
 			InputParams,
@@ -2257,7 +2271,7 @@ void MainCS(
 	uint2 GroupThreadId : SV_GroupThreadID,
 	uint GroupThreadIndex : SV_GroupIndex)
 {
-	float2 ViewportUV = (float2(DispatchThreadId) + 0.5f) * ViewportSize.zw;
+	float2 ViewportUV = (float2(DispatchThreadId) + 0.5f) * OutputViewportSize.zw;
 	
 	#if AA_LOWER_RESOLUTION
 	{
@@ -2288,8 +2302,8 @@ void MainCS(
 	}
 	#endif
 
-	uint2 PixelPos = DispatchThreadId + ViewportRect.xy;
-	if (all(PixelPos < ViewportRect.zw))
+	uint2 PixelPos = DispatchThreadId + OutputViewportRect.xy;
+	if (all(PixelPos < OutputViewportRect.zw))
 	{
 		OutComputeTex0[PixelPos] = OutColor0;
 
@@ -2330,13 +2344,8 @@ void MainCS(
 #endif
 
 
-void MainPS( 
-	in noperspective float2 BufferUV : TEXCOORD0, 
-	in noperspective float3 InExposureScaleVignette : TEXCOORD1, 
-	in noperspective float4 Unused1 : TEXCOORD2,
-	in noperspective float2 Unused2 : TEXCOORD3,
-	in noperspective float2 Unused3 : TEXCOORD4,
-	float4 SvPosition : SV_POSITION, // after all interpolators
+void MainPS(
+	float4 SvPosition : SV_POSITION,
 	#if HISTORY_RENDER_TARGETS == 2
 		out float4 OutColor0 : SV_Target0,
 		out float4 OutColor1 : SV_Target1
@@ -2345,11 +2354,11 @@ void MainPS(
 	#endif
 	)
 {
-	float2 ViewportUV = BufferUVToViewportUV(BufferUV);
+	float2 ViewportUV = (SvPosition.xy - OutputViewportRect.xy) * OutputViewportSize.zw;
 
 	#if AA_LOWER_RESOLUTION
 	{
-		ViewportUV = (SvPosition.xy - ViewportRect.xy) * MaxViewportUVAndSvPositionToViewportUV.zw;
+		ViewportUV = (SvPosition.xy - OutputViewportRect.xy) * MaxViewportUVAndSvPositionToViewportUV.zw;
 		ViewportUV = min(ViewportUV, MaxViewportUVAndSvPositionToViewportUV.xy);
 	}
 	#endif
@@ -2357,7 +2366,7 @@ void MainPS(
 	uint2 GroupId = int2(0, 0);
 	uint2 GroupThreadId = int2(0, 0);
 	uint GroupThreadIndex = 0;
-	float FrameExposureScale = InExposureScaleVignette.x;
+	float FrameExposureScale = EyeAdaptationLookup();
 
 	FTAAHistoryPayload OutputPayload = TemporalAASample(GroupId, GroupThreadId, GroupThreadIndex, ViewportUV, FrameExposureScale);
 
diff --git a/Engine/Shaders/Private/Tools/FullscreenVertexShader.usf b/Engine/Shaders/Private/Tools/FullscreenVertexShader.usf
--- a/Engine/Shaders/Private/Tools/FullscreenVertexShader.usf
+++ b/Engine/Shaders/Private/Tools/FullscreenVertexShader.usf
@@ -8,10 +8,8 @@
 
 void MainVS(
 	float2 InPosition : ATTRIBUTE0,
-	float2 InUV       : ATTRIBUTE1,
-	out FScreenVertexOutput Output
-	)
+	float2 InUV       : ATTRIBUTE1, // TODO: kill
+	out float4 Position : SV_POSITION)
 {
-	Output.UV = InUV;
-	Output.Position = float4(InPosition.x * 2.0 - 1.0, 1.0 - 2.0 * InPosition.y, 0, 1);
+	Position = float4(InPosition.x * 2.0 - 1.0, 1.0 - 2.0 * InPosition.y, 0, 1);
 }
diff --git a/Engine/Shaders/Private/Tools/VisualizeTexture.usf b/Engine/Shaders/Private/Tools/VisualizeTexture.usf
--- a/Engine/Shaders/Private/Tools/VisualizeTexture.usf
+++ b/Engine/Shaders/Private/Tools/VisualizeTexture.usf
@@ -36,10 +36,10 @@ float3 TextureExtent;
 
 
 void VisualizeTexturePS(
-	FScreenVertexOutput Input,
+	float4 SvPosition : SV_POSITION,
 	out float4 OutColor : SV_Target0)
 {
-	float2 UV = Input.UV;
+	float2 UV = SvPosition / TextureExtent.xy;
 
 	// Longitude Latitude unwrapping
 	float3 SpherePos;
diff --git a/Engine/Source/Runtime/RenderCore/Private/CommonRenderResources.cpp b/Engine/Source/Runtime/RenderCore/Private/CommonRenderResources.cpp
--- a/Engine/Source/Runtime/RenderCore/Private/CommonRenderResources.cpp
+++ b/Engine/Source/Runtime/RenderCore/Private/CommonRenderResources.cpp
@@ -14,7 +14,7 @@ TGlobalResource<FEmptyVertexDeclaration> GEmptyVertexDeclaration;
 TGlobalResource<FScreenRectangleVertexBuffer> GScreenRectangleVertexBuffer;
 TGlobalResource<FScreenRectangleIndexBuffer> GScreenRectangleIndexBuffer;
 
-IMPLEMENT_GLOBAL_SHADER(FVisualizeTextureVS, "/Engine/Private/Tools/FullscreenVertexShader.usf", "MainVS", SF_Vertex);
+IMPLEMENT_GLOBAL_SHADER(FScreenVertexShaderVS, "/Engine/Private/Tools/FullscreenVertexShader.usf", "MainVS", SF_Vertex);
 
 
 void FScreenRectangleVertexBuffer::InitRHI()
diff --git a/Engine/Source/Runtime/RenderCore/Private/PixelShaderUtils.cpp b/Engine/Source/Runtime/RenderCore/Private/PixelShaderUtils.cpp
--- a/Engine/Source/Runtime/RenderCore/Private/PixelShaderUtils.cpp
+++ b/Engine/Source/Runtime/RenderCore/Private/PixelShaderUtils.cpp
@@ -38,3 +38,23 @@ void FPixelShaderUtils::DrawFullscreenQuad(FRHICommandList& RHICmdList, uint32 I
 		/*NumPrimitives=*/ 2,
 		/*NumInstances=*/ InstanceCount);
 }
+
+// static
+void FPixelShaderUtils::InitFullscreenPipelineState(
+	FRHICommandList& RHICmdList,
+	const TShaderMap<FGlobalShaderType>* GlobalShaderMap,
+	const FShader* PixelShader,
+	FGraphicsPipelineStateInitializer& GraphicsPSOInit)
+{
+	TShaderMapRef<FScreenVertexShaderVS> VertexShader(GlobalShaderMap);
+		
+	RHICmdList.ApplyCachedRenderTargets(GraphicsPSOInit);
+	GraphicsPSOInit.BlendState = TStaticBlendState<>::GetRHI();
+	GraphicsPSOInit.RasterizerState = TStaticRasterizerState<>::GetRHI();
+	GraphicsPSOInit.DepthStencilState = TStaticDepthStencilState<false, CF_Always>::GetRHI();
+
+	GraphicsPSOInit.BoundShaderState.VertexDeclarationRHI = GFilterVertexDeclaration.VertexDeclarationRHI;
+	GraphicsPSOInit.BoundShaderState.VertexShaderRHI = GETSAFERHISHADER_VERTEX(*VertexShader);
+	GraphicsPSOInit.BoundShaderState.PixelShaderRHI = GETSAFERHISHADER_PIXEL(PixelShader);
+	GraphicsPSOInit.PrimitiveType = PT_TriangleList;
+}
diff --git a/Engine/Source/Runtime/RenderCore/Public/CommonRenderResources.h b/Engine/Source/Runtime/RenderCore/Public/CommonRenderResources.h
--- a/Engine/Source/Runtime/RenderCore/Public/CommonRenderResources.h
+++ b/Engine/Source/Runtime/RenderCore/Public/CommonRenderResources.h
@@ -91,11 +91,12 @@ public:
 extern RENDERCORE_API TGlobalResource<FScreenRectangleIndexBuffer> GScreenRectangleIndexBuffer;
 
 
-/** Vertex shader to draw a full screen quad that works on all platforms. */
-class RENDERCORE_API FVisualizeTextureVS : public FGlobalShader
+/** Vertex shader to draw a screen quad that works on all platforms. Does not have any shader parameters.
+ * The pixel shader should just use SV_Position. */
+class RENDERCORE_API FScreenVertexShaderVS : public FGlobalShader
 {
-	DECLARE_GLOBAL_SHADER(FVisualizeTextureVS);
-	SHADER_USE_PARAMETER_STRUCT(FVisualizeTextureVS, FGlobalShader);
+	DECLARE_GLOBAL_SHADER(FScreenVertexShaderVS);
+	SHADER_USE_PARAMETER_STRUCT(FScreenVertexShaderVS, FGlobalShader);
 
 	static bool ShouldCompilePermutation(const FGlobalShaderPermutationParameters& Parameters) {
 		return true;
diff --git a/Engine/Source/Runtime/RenderCore/Public/PixelShaderUtils.h b/Engine/Source/Runtime/RenderCore/Public/PixelShaderUtils.h
--- a/Engine/Source/Runtime/RenderCore/Public/PixelShaderUtils.h
+++ b/Engine/Source/Runtime/RenderCore/Public/PixelShaderUtils.h
@@ -23,6 +23,12 @@ struct RENDERCORE_API FPixelShaderUtils
 	/** Draw a two triangle on the entire viewport. */
 	static void DrawFullscreenQuad(FRHICommandList& RHICmdList, uint32 InstanceCount = 1);
 
+	/** Initialize a pipeline state object initializer with almost all the basics required to do a full viewport pass. */
+	static void InitFullscreenPipelineState(
+		FRHICommandList& RHICmdList,
+		const TShaderMap<FGlobalShaderType>* GlobalShaderMap,
+		const FShader* PixelShader,
+		FGraphicsPipelineStateInitializer& GraphicsPSOInit);
 
 	/** Dispatch a full screen pixel shader to rhi command list with its parameters. */
 	template<typename TShaderClass>
@@ -33,23 +39,13 @@ struct RENDERCORE_API FPixelShaderUtils
 		const typename TShaderClass::FParameters& Parameters,
 		const FIntRect& Viewport)
 	{
-		CA_ASSUME(PixelShader);
-		TShaderMapRef<FVisualizeTextureVS> VertexShader(GlobalShaderMap);
-
+		check(PixelShader);
+		RHICmdList.SetViewport(Viewport.Min.X, Viewport.Min.Y, 0.0f, Viewport.Max.X, Viewport.Max.Y, 1.0f);
+		
 		FGraphicsPipelineStateInitializer GraphicsPSOInit;
-		RHICmdList.ApplyCachedRenderTargets(GraphicsPSOInit);
-		GraphicsPSOInit.BlendState = TStaticBlendState<>::GetRHI();
-		GraphicsPSOInit.RasterizerState = TStaticRasterizerState<>::GetRHI();
-		GraphicsPSOInit.DepthStencilState = TStaticDepthStencilState<false, CF_Always>::GetRHI();
-
-		GraphicsPSOInit.BoundShaderState.VertexDeclarationRHI = GFilterVertexDeclaration.VertexDeclarationRHI;
-		GraphicsPSOInit.BoundShaderState.VertexShaderRHI = GETSAFERHISHADER_VERTEX(*VertexShader);
-		GraphicsPSOInit.BoundShaderState.PixelShaderRHI = GETSAFERHISHADER_PIXEL(PixelShader);
-		GraphicsPSOInit.PrimitiveType = PT_TriangleList;
+		InitFullscreenPipelineState(RHICmdList, GlobalShaderMap, PixelShader, /* out */ GraphicsPSOInit);
 		SetGraphicsPipelineState(RHICmdList, GraphicsPSOInit);
 
-		RHICmdList.SetViewport(Viewport.Min.X, Viewport.Min.Y, 0.0f, Viewport.Max.X, Viewport.Max.Y, 1.0f);
-
 		SetShaderParameters(RHICmdList, PixelShader, PixelShader->GetPixelShader(), Parameters);
 
 		DrawFullscreenTriangle(RHICmdList);
@@ -65,7 +61,7 @@ struct RENDERCORE_API FPixelShaderUtils
 		typename TShaderClass::FParameters* Parameters,
 		const FIntRect& Viewport)
 	{
-		CA_ASSUME(PixelShader != nullptr);
+		check(PixelShader);
 		ClearUnusedGraphResources(PixelShader, Parameters);
 
 		GraphBuilder.AddPass(
diff --git a/Engine/Source/Runtime/Renderer/Private/PostProcess/DiaphragmDOFGraph.cpp b/Engine/Source/Runtime/Renderer/Private/PostProcess/DiaphragmDOFGraph.cpp
--- a/Engine/Source/Runtime/Renderer/Private/PostProcess/DiaphragmDOFGraph.cpp
+++ b/Engine/Source/Runtime/Renderer/Private/PostProcess/DiaphragmDOFGraph.cpp
@@ -246,8 +246,9 @@ bool DiaphragmDOF::WireSceneColorPasses(FPostprocessContext& Context, const FRen
 		if (Context.View.PrimaryScreenPercentageMethod == EPrimaryScreenPercentageMethod::TemporalUpscale)
 		{
 			TAAParameters.Pass = ETAAPassConfig::DiaphragmDOFUpsampling;
+			TAAParameters.bIsComputePass = true;
 		}
-
+		
 		TAAParameters.SetupViewRect(Context.View, PrefilteringResolutionDivisor);
 		TAAParameters.TopLeftCornerViewRects();
 
diff --git a/Engine/Source/Runtime/Renderer/Private/PostProcess/PostProcessTemporalAA.cpp b/Engine/Source/Runtime/Renderer/Private/PostProcess/PostProcessTemporalAA.cpp
--- a/Engine/Source/Runtime/Renderer/Private/PostProcess/PostProcessTemporalAA.cpp
+++ b/Engine/Source/Runtime/Renderer/Private/PostProcess/PostProcessTemporalAA.cpp
@@ -18,6 +18,12 @@
 #include "PipelineStateCache.h"
 #include "PostProcessing.h"
 
+
+#include "RenderGraph.h"
+#include "SceneViewFamilyBlackboard.h"
+#include "PixelShaderUtils.h"
+
+
 const int32 GTemporalAATileSizeX = 8;
 const int32 GTemporalAATileSizeY = 8;
 
@@ -51,6 +57,7 @@ static TAutoConsoleVariable<int32> CVarTemporalAAUpsampleFiltered(
 	TEXT("Use filtering to fetch color history during TamporalAA upsampling (see AA_FILTERED define in TAA shader). Disabling this makes TAAU faster, but lower quality. "),
 	ECVF_Scalability | ECVF_RenderThreadSafe);
 
+
 static float CatmullRom( float x )
 {
 	float ax = FMath::Abs(x);
@@ -61,230 +68,46 @@ static float CatmullRom( float x )
 }
 
 
-struct FTemporalAAParameters
-{
-public:
-	FPostProcessPassParameters PostprocessParameter;
-	FSceneTextureShaderParameters SceneTextureParameters;
-	FShaderParameter SampleWeights;
-	FShaderParameter PlusWeights;
-	FShaderParameter DitherScale;
-	FShaderParameter VelocityScaling;
-	FShaderParameter CurrentFrameWeight;
-	FShaderParameter ScreenPosAbsMax;
-	FShaderParameter ScreenPosToHistoryBufferUV;
-	FShaderResourceParameter HistoryBuffer[FTemporalAAHistory::kRenderTargetCount];
-	FShaderResourceParameter HistoryBufferSampler[FTemporalAAHistory::kRenderTargetCount];
-	FShaderParameter HistoryBufferSize;
-	FShaderParameter HistoryBufferUVMinMax;
-	FShaderParameter MaxViewportUVAndSvPositionToViewportUV;
-	FShaderParameter PreExposureSettings;
-	FShaderParameter ViewportUVToInputBufferUV;
-
-	void Bind(const FShader::CompiledShaderInitializerType& Initializer)
-	{
-		const FShaderParameterMap& ParameterMap = Initializer.ParameterMap;
-		PostprocessParameter.Bind(ParameterMap);
-		SceneTextureParameters.Bind(Initializer);
-		SampleWeights.Bind(ParameterMap, TEXT("SampleWeights"));
-		PlusWeights.Bind(ParameterMap, TEXT("PlusWeights"));
-		DitherScale.Bind(ParameterMap, TEXT("DitherScale"));
-		VelocityScaling.Bind(ParameterMap, TEXT("VelocityScaling"));
-		CurrentFrameWeight.Bind(ParameterMap, TEXT("CurrentFrameWeight"));
-		ScreenPosAbsMax.Bind(ParameterMap, TEXT("ScreenPosAbsMax"));
-		ScreenPosToHistoryBufferUV.Bind(ParameterMap, TEXT("ScreenPosToHistoryBufferUV"));
-		HistoryBuffer[0].Bind(ParameterMap, TEXT("HistoryBuffer0"));
-		HistoryBuffer[1].Bind(ParameterMap, TEXT("HistoryBuffer1"));
-		HistoryBufferSampler[0].Bind(ParameterMap, TEXT("HistoryBuffer0Sampler"));
-		HistoryBufferSampler[1].Bind(ParameterMap, TEXT("HistoryBuffer1Sampler"));
-		HistoryBufferSize.Bind(ParameterMap, TEXT("HistoryBufferSize"));
-		HistoryBufferUVMinMax.Bind(ParameterMap, TEXT("HistoryBufferUVMinMax"));
-		MaxViewportUVAndSvPositionToViewportUV.Bind(ParameterMap, TEXT("MaxViewportUVAndSvPositionToViewportUV"));
-		PreExposureSettings.Bind(ParameterMap, TEXT("PreExposureSettings"));
-		ViewportUVToInputBufferUV.Bind(ParameterMap, TEXT("ViewportUVToInputBufferUV"));
-	}
-
-	void Serialize(FArchive& Ar)
-	{
-		Ar << PostprocessParameter << SceneTextureParameters ;
-		Ar << SampleWeights << PlusWeights << DitherScale << VelocityScaling << CurrentFrameWeight << ScreenPosAbsMax << ScreenPosToHistoryBufferUV << HistoryBuffer[0] << HistoryBuffer[1] << HistoryBufferSampler[0] << HistoryBufferSampler[1] << HistoryBufferSize << HistoryBufferUVMinMax;
-		Ar << MaxViewportUVAndSvPositionToViewportUV << PreExposureSettings << ViewportUVToInputBufferUV;
-	}
-
-	template <typename TRHICmdList, typename TShaderRHIParamRef>
-	void SetParameters(
-		TRHICmdList& RHICmdList,
-		const TShaderRHIParamRef ShaderRHI,
-		const FRenderingCompositePassContext& Context,
-		const FTemporalAAHistory& InputHistory,
-		const FTAAPassParameters& PassParameters,
-		bool bUseDither,
-		const FIntPoint& SrcSize)
-	{
-		SceneTextureParameters.Set(RHICmdList, ShaderRHI, Context.View.FeatureLevel, ESceneTextureSetupMode::All);
-
-		float ResDivisor = PassParameters.ResolutionDivisor;
-		float ResDivisorInv = 1.0f / ResDivisor;
-
-		// PS params
-		{
-			float JitterX = Context.View.TemporalJitterPixels.X;
-			float JitterY = Context.View.TemporalJitterPixels.Y;
-
-			static const float SampleOffsets[9][2] =
-			{
-				{ -1.0f, -1.0f },
-				{  0.0f, -1.0f },
-				{  1.0f, -1.0f },
-				{ -1.0f,  0.0f },
-				{  0.0f,  0.0f },
-				{  1.0f,  0.0f },
-				{ -1.0f,  1.0f },
-				{  0.0f,  1.0f },
-				{  1.0f,  1.0f },
-			};
-
-			float FilterSize = CVarTemporalAAFilterSize.GetValueOnRenderThread();
-			int32 bCatmullRom = CVarTemporalAACatmullRom.GetValueOnRenderThread();
-
-			float Weights[9];
-			float WeightsPlus[5];
-			float TotalWeight = 0.0f;
-			float TotalWeightLow = 0.0f;
-			float TotalWeightPlus = 0.0f;
-			for (int32 i = 0; i < 9; i++)
-			{
-				float PixelOffsetX = SampleOffsets[i][0] - JitterX * ResDivisorInv;
-				float PixelOffsetY = SampleOffsets[i][1] - JitterY * ResDivisorInv;
-
-				PixelOffsetX /= FilterSize;
-				PixelOffsetY /= FilterSize;
-
-				if (bCatmullRom)
-				{
-					Weights[i] = CatmullRom(PixelOffsetX) * CatmullRom(PixelOffsetY);
-					TotalWeight += Weights[i];
-				}
-				else
-				{
-					// Normal distribution, Sigma = 0.47
-					Weights[i] = FMath::Exp(-2.29f * (PixelOffsetX * PixelOffsetX + PixelOffsetY * PixelOffsetY));
-					TotalWeight += Weights[i];
-				}
-			}
-
-			WeightsPlus[0] = Weights[1];
-			WeightsPlus[1] = Weights[3];
-			WeightsPlus[2] = Weights[4];
-			WeightsPlus[3] = Weights[5];
-			WeightsPlus[4] = Weights[7];
-			TotalWeightPlus = Weights[1] + Weights[3] + Weights[4] + Weights[5] + Weights[7];
-
-			for (int32 i = 0; i < 9; i++)
-			{
-				SetShaderValue(RHICmdList, ShaderRHI, SampleWeights, Weights[i] / TotalWeight, i);
-			}
-
-			for (int32 i = 0; i < 5; i++)
-			{
-				SetShaderValue(RHICmdList, ShaderRHI, PlusWeights, WeightsPlus[i] / TotalWeightPlus, i);
-			}
-		}
-
-		SetShaderValue(RHICmdList, ShaderRHI, DitherScale, bUseDither ? 1.0f : 0.0f);
-
-		const bool bIgnoreVelocity = (Context.View.ViewState && Context.View.ViewState->bSequencerIsPaused);
-		SetShaderValue(RHICmdList, ShaderRHI, VelocityScaling, bIgnoreVelocity ? 0.0f : 1.0f);
-
-		SetShaderValue(RHICmdList, ShaderRHI, CurrentFrameWeight, CVarTemporalAACurrentFrameWeight.GetValueOnRenderThread());
-
-		// Set history shader parameters.
-		if (InputHistory.IsValid())
-		{
-			FIntPoint ReferenceViewportOffset = InputHistory.ViewportRect.Min;
-			FIntPoint ReferenceViewportExtent = InputHistory.ViewportRect.Size();
-			FIntPoint ReferenceBufferSize = InputHistory.ReferenceBufferSize;
-
-			float InvReferenceBufferSizeX = 1.f / float(InputHistory.ReferenceBufferSize.X);
-			float InvReferenceBufferSizeY = 1.f / float(InputHistory.ReferenceBufferSize.Y);
-
-			FVector4 ScreenPosToPixelValue(
-				ReferenceViewportExtent.X * 0.5f * InvReferenceBufferSizeX,
-				-ReferenceViewportExtent.Y * 0.5f * InvReferenceBufferSizeY,
-				(ReferenceViewportExtent.X * 0.5f + ReferenceViewportOffset.X) * InvReferenceBufferSizeX,
-				(ReferenceViewportExtent.Y * 0.5f + ReferenceViewportOffset.Y) * InvReferenceBufferSizeY);
-			SetShaderValue(RHICmdList, ShaderRHI, ScreenPosToHistoryBufferUV, ScreenPosToPixelValue);
-
-			FIntPoint ViewportOffset = ReferenceViewportOffset / PassParameters.ResolutionDivisor;
-			FIntPoint ViewportExtent = FIntPoint::DivideAndRoundUp(ReferenceViewportExtent, PassParameters.ResolutionDivisor);
-			FIntPoint BufferSize = ReferenceBufferSize / PassParameters.ResolutionDivisor;
-
-			FVector2D ScreenPosAbsMaxValue(1.0f - 1.0f / float(ViewportExtent.X), 1.0f - 1.0f / float(ViewportExtent.Y));
-			SetShaderValue(RHICmdList, ShaderRHI, ScreenPosAbsMax, ScreenPosAbsMaxValue);
-
-			float InvBufferSizeX = 1.f / float(BufferSize.X);
-			float InvBufferSizeY = 1.f / float(BufferSize.Y);
-
-			FVector4 HistoryBufferUVMinMaxValue(
-				(ViewportOffset.X + 0.5f) * InvBufferSizeX,
-				(ViewportOffset.Y + 0.5f) * InvBufferSizeY,
-				(ViewportOffset.X + ViewportExtent.X - 0.5f) * InvBufferSizeX,
-				(ViewportOffset.Y + ViewportExtent.Y - 0.5f) * InvBufferSizeY);
-
-			SetShaderValue(RHICmdList, ShaderRHI, HistoryBufferUVMinMax, HistoryBufferUVMinMaxValue);
-
-			FVector4 HistoryBufferSizeValue(BufferSize.X, BufferSize.Y, InvBufferSizeX, InvBufferSizeY);
-			SetShaderValue(RHICmdList, ShaderRHI, HistoryBufferSize, HistoryBufferSizeValue);
-
-			for (uint32 i = 0; i < FTemporalAAHistory::kRenderTargetCount; i++)
-			{
-				if (InputHistory.RT[i].IsValid())
-				{
-					SetTextureParameter(
-						RHICmdList, ShaderRHI,
-						HistoryBuffer[i], HistoryBufferSampler[i],
-						TStaticSamplerState<SF_Bilinear>::GetRHI(),
-						InputHistory.RT[i]->GetRenderTargetItem().ShaderResourceTexture);
-				}
-			}
-		}
-
-		{
-			FVector4 MaxViewportUVAndSvPositionToViewportUVValue(
-				(PassParameters.OutputViewRect.Width() - 0.5f * ResDivisor) / float(PassParameters.OutputViewRect.Width()),
-				(PassParameters.OutputViewRect.Height() - 0.5f * ResDivisor) / float(PassParameters.OutputViewRect.Height()),
-				ResDivisor / float(PassParameters.OutputViewRect.Width()),
-				ResDivisor / float(PassParameters.OutputViewRect.Height()));
-
-			SetShaderValue(RHICmdList, ShaderRHI, MaxViewportUVAndSvPositionToViewportUV, MaxViewportUVAndSvPositionToViewportUVValue);
-		}
-
-		// Pre-exposure, One over Pre-exposure, History pre-exposure, History one over pre-exposure.
-		// DOF settings must preserve scene color range.
-		FVector4 PreExposureSettingsValue(1.f, 1.f, 1.f, 1.f);
-		if (PassParameters.Pass == ETAAPassConfig::Main)
-		{
-			PreExposureSettingsValue.X = Context.View.PreExposure;
-			PreExposureSettingsValue.Y = 1.f / FMath::Max<float>(SMALL_NUMBER, Context.View.PreExposure);
-			PreExposureSettingsValue.Z = InputHistory.IsValid() ? InputHistory.SceneColorPreExposure : Context.View.PreExposure;
-			PreExposureSettingsValue.W  = 1.f / FMath::Max<float>(SMALL_NUMBER, PreExposureSettingsValue.Z);
-		}
-		SetShaderValue(RHICmdList, ShaderRHI, PreExposureSettings, PreExposureSettingsValue);
-
-		{
-			float InvSizeX = 1.0f / float(SrcSize.X);
-			float InvSizeY = 1.0f / float(SrcSize.Y);
-			FVector4 ViewportUVToBufferUVValue(
-				ResDivisorInv * PassParameters.InputViewRect.Width() * InvSizeX,
-				ResDivisorInv * PassParameters.InputViewRect.Height() * InvSizeY,
-				ResDivisorInv * PassParameters.InputViewRect.Min.X * InvSizeX,
-				ResDivisorInv * PassParameters.InputViewRect.Min.Y * InvSizeY);
-
-			SetShaderValue(RHICmdList, ShaderRHI, ViewportUVToInputBufferUV, ViewportUVToBufferUVValue);
-		}
-		
-	}
-};
+BEGIN_SHADER_PARAMETER_STRUCT(FTAAShaderParameters,)
+	SHADER_PARAMETER(FVector4, PreExposureSettings)
+	SHADER_PARAMETER(FVector4, ViewportUVToInputBufferUV)
+	SHADER_PARAMETER(FVector4, MaxViewportUVAndSvPositionToViewportUV)
+	SHADER_PARAMETER(FVector2D, ScreenPosAbsMax)
+	SHADER_PARAMETER(float, CurrentFrameWeight)
+	
+	SHADER_PARAMETER_ARRAY(float, SampleWeights, [9])
+	SHADER_PARAMETER_ARRAY(float, PlusWeights, [5])
+
+	SHADER_PARAMETER(FVector4, InputSceneColorSize)
+	SHADER_PARAMETER(FVector4, OutputViewportSize)
+	SHADER_PARAMETER(FVector4, OutputViewportRect)
+
+	// History parameters
+	SHADER_PARAMETER(FVector4, HistoryBufferSize)
+	SHADER_PARAMETER(FVector4, HistoryBufferUVMinMax)
+	SHADER_PARAMETER(FVector4, ScreenPosToHistoryBufferUV)
+	
+	SHADER_PARAMETER_STRUCT_INCLUDE(FSceneViewFamilyBlackboard, SceneBlackboard)
+	
+	SHADER_PARAMETER_RDG_TEXTURE(Texture2D, EyeAdaptation)
+
+	// Inputs
+	SHADER_PARAMETER_RDG_TEXTURE(Texture2D, InputSceneColor)
+	SHADER_PARAMETER_SAMPLER(SamplerState, InputSceneColorSampler)
+	SHADER_PARAMETER_RDG_TEXTURE(Texture2D, InputSceneMetadata)
+	SHADER_PARAMETER_SAMPLER(SamplerState, InputSceneMetadataSampler)
+
+	// History resourrces
+	SHADER_PARAMETER_RDG_TEXTURE(Texture2D, HistoryBuffer0)
+	SHADER_PARAMETER_SAMPLER(SamplerState, HistoryBuffer0Sampler)
+	SHADER_PARAMETER_RDG_TEXTURE(Texture2D, HistoryBuffer1)
+	SHADER_PARAMETER_SAMPLER(SamplerState, HistoryBuffer1Sampler)
+	
+	SHADER_PARAMETER_SAMPLER(SamplerState, SceneDepthBufferSampler)
+	SHADER_PARAMETER_SAMPLER(SamplerState, SceneVelocityBufferSampler)
+
+	SHADER_PARAMETER_STRUCT_REF(FViewUniformShaderParameters, ViewUniformBuffer)
+END_SHADER_PARAMETER_STRUCT()
 
 
 // ---------------------------------------------------- Shader permutation dimensions
@@ -305,15 +128,21 @@ class FTAADownsampleDim : SHADER_PERMUTATION_BOOL("TAA_DOWNSAMPLE");
 
 // ---------------------------------------------------- Shaders
 
-class FPostProcessTemporalAAPS : public FGlobalShader
+class FTemporalAAPS : public FGlobalShader
 {
-	DECLARE_GLOBAL_SHADER(FPostProcessTemporalAAPS);
+	DECLARE_GLOBAL_SHADER(FTemporalAAPS);
+	SHADER_USE_PARAMETER_STRUCT(FTemporalAAPS, FGlobalShader);
 
 	using FPermutationDomain = TShaderPermutationDomain<
 		FTAAPassConfigDim,
 		FTAAFastDim,
 		FTAAResponsiveDim,
 		FTAACameraCutDim>;
+	
+	BEGIN_SHADER_PARAMETER_STRUCT(FParameters, )
+		SHADER_PARAMETER_STRUCT_INCLUDE(FTAAShaderParameters, CommonParameters)
+		RENDER_TARGET_BINDING_SLOTS()
+	END_SHADER_PARAMETER_STRUCT()
 
 	static bool ShouldCompilePermutation(const FGlobalShaderPermutationParameters& Parameters)
 	{
@@ -346,52 +175,12 @@ class FPostProcessTemporalAAPS : public FGlobalShader
 	{
 		return PermutationVector.Get<FTAAPassConfigDim>() == ETAAPassConfig::Main;
 	}
+}; // class FTemporalAAPS
 
-	/** Default constructor. */
-	FPostProcessTemporalAAPS() {}
-
-	/** Initialization constructor. */
-	FPostProcessTemporalAAPS(const ShaderMetaType::CompiledShaderInitializerType& Initializer)
-		: FGlobalShader(Initializer)
-	{
-		Parameter.Bind(Initializer);
-	}
-
-	// FShader interface.
-	virtual bool Serialize(FArchive& Ar) override
-	{
-		bool bShaderHasOutdatedParameters = FGlobalShader::Serialize(Ar);
-		Parameter.Serialize(Ar);
-		return bShaderHasOutdatedParameters;
-	}
-
-	template <typename TRHICmdList>
-	void SetParameters(
-		TRHICmdList& RHICmdList,
-		const FRenderingCompositePassContext& Context,
-		const FTemporalAAHistory& InputHistory,
-		const FTAAPassParameters& PassParameters,
-		bool bUseDither,
-		const FIntPoint& SrcSize)
-	{
-		const FPixelShaderRHIParamRef ShaderRHI = GetPixelShader();
-
-		FGlobalShader::SetParameters<FViewUniformShaderParameters>(RHICmdList, ShaderRHI, Context.View.ViewUniformBuffer);
-		
-		Parameter.PostprocessParameter.SetPS(RHICmdList, ShaderRHI, Context);
-
-		Parameter.SetParameters(RHICmdList, ShaderRHI, Context, InputHistory, PassParameters, bUseDither, SrcSize);
-	}
-
-	FTemporalAAParameters Parameter;
-};
-
-IMPLEMENT_GLOBAL_SHADER(FPostProcessTemporalAAPS, "/Engine/Private/PostProcessTemporalAA.usf", "MainPS", SF_Pixel);
-
-
-class FPostProcessTemporalAACS : public FGlobalShader
+class FTemporalAACS : public FGlobalShader
 {
-	DECLARE_GLOBAL_SHADER(FPostProcessTemporalAACS);
+	DECLARE_GLOBAL_SHADER(FTemporalAACS);
+	SHADER_USE_PARAMETER_STRUCT(FTemporalAACS, FGlobalShader);
 
 	using FPermutationDomain = TShaderPermutationDomain<
 		FTAAPassConfigDim,
@@ -401,6 +190,21 @@ class FPostProcessTemporalAACS : public FGlobalShader
 		FTAAUpsampleFilteredDim,
 		FTAADownsampleDim>;
 
+	BEGIN_SHADER_PARAMETER_STRUCT(FParameters, )
+		SHADER_PARAMETER_STRUCT_INCLUDE(FTAAShaderParameters, CommonParameters)
+		
+		// Temporal upsample specific parameters.
+		SHADER_PARAMETER(FVector4, InputViewSize)
+		SHADER_PARAMETER(FVector2D, InputViewMin)
+		SHADER_PARAMETER(FVector2D, TemporalJitterPixels)
+		SHADER_PARAMETER(float, ScreenPercentage)
+		SHADER_PARAMETER(float, UpscaleFactor)
+
+		SHADER_PARAMETER_RDG_TEXTURE_UAV(Texture2D, OutComputeTex0)
+		SHADER_PARAMETER_RDG_TEXTURE_UAV(Texture2D, OutComputeTex1)
+		SHADER_PARAMETER_RDG_TEXTURE_UAV(Texture2D, OutComputeTexDownsampled)
+	END_SHADER_PARAMETER_STRUCT()
+
 	static bool ShouldCompilePermutation(const FGlobalShaderPermutationParameters& Parameters)
 	{
 		FPermutationDomain PermutationVector(Parameters.PermutationId);
@@ -484,185 +288,81 @@ class FPostProcessTemporalAACS : public FGlobalShader
 		OutEnvironment.SetDefine(TEXT("THREADGROUP_SIZEX"), GTemporalAATileSizeX);
 		OutEnvironment.SetDefine(TEXT("THREADGROUP_SIZEY"), GTemporalAATileSizeY);
 	}
+}; // class FTemporalAACS
 
-	/** Default constructor. */
-	FPostProcessTemporalAACS() {}
 
-	/** Initialization constructor. */
-	FPostProcessTemporalAACS(const ShaderMetaType::CompiledShaderInitializerType& Initializer)
-		: FGlobalShader(Initializer)
-	{
-		Parameter.Bind(Initializer);
-
-		EyeAdaptation.Bind(Initializer.ParameterMap, TEXT("EyeAdaptation"));
-		OutComputeTex0.Bind(Initializer.ParameterMap, TEXT("OutComputeTex0"));
-		OutComputeTex1.Bind(Initializer.ParameterMap, TEXT("OutComputeTex1"));
-		OutComputeTexDownsampled.Bind(Initializer.ParameterMap, TEXT("OutComputeTexDownsampled"));
-		InputViewMin.Bind(Initializer.ParameterMap, TEXT("InputViewMin"));
-		InputViewSize.Bind(Initializer.ParameterMap, TEXT("InputViewSize"));
-		TemporalJitterPixels.Bind(Initializer.ParameterMap, TEXT("TemporalJitterPixels"));
-		ScreenPercentageAndUpscaleFactor.Bind(Initializer.ParameterMap, TEXT("ScreenPercentageAndUpscaleFactor"));
-	}
-
-	// FShader interface.
-	virtual bool Serialize(FArchive& Ar) override
-	{
-		bool bShaderHasOutdatedParameters = FGlobalShader::Serialize(Ar);
-		Parameter.Serialize(Ar);
-		Ar << EyeAdaptation << OutComputeTex0 << OutComputeTex1 << OutComputeTexDownsampled << InputViewMin << InputViewSize << TemporalJitterPixels << ScreenPercentageAndUpscaleFactor;
-		return bShaderHasOutdatedParameters;
-	}
-
-	template <typename TRHICmdList>
-	void SetParameters(
-		TRHICmdList& RHICmdList,
-		const FRenderingCompositePassContext& Context,
-		const FTemporalAAHistory& InputHistory,
-		const FTAAPassParameters& PassParameters,
-		const FIntPoint& DestSize,
-		const FSceneRenderTargetItem* DestRenderTarget[2],
-		FUnorderedAccessViewRHIParamRef DestDownsampledUAV,
-		const FIntPoint& SrcSize,
-		bool bUseDither,
-		FTextureRHIParamRef EyeAdaptationTex)
-	{
-		const FComputeShaderRHIParamRef ShaderRHI = GetComputeShader();
-		FSceneViewState* ViewState = (FSceneViewState*)Context.View.State;
-		FGlobalShader::SetParameters<FViewUniformShaderParameters>(RHICmdList, ShaderRHI, Context.View.ViewUniformBuffer);
+IMPLEMENT_GLOBAL_SHADER(FTemporalAAPS, "/Engine/Private/PostProcessTemporalAA.usf", "MainPS", SF_Pixel);
+IMPLEMENT_GLOBAL_SHADER(FTemporalAACS, "/Engine/Private/PostProcessTemporalAA.usf", "MainCS", SF_Compute);
 
-		// CS params
-		Parameter.PostprocessParameter.SetCS(ShaderRHI, Context, RHICmdList);
 
-		RHICmdList.SetUAVParameter(ShaderRHI, OutComputeTex0.GetBaseIndex(), DestRenderTarget[0]->UAV);
-		if (DestRenderTarget[1])
-			RHICmdList.SetUAVParameter(ShaderRHI, OutComputeTex1.GetBaseIndex(), DestRenderTarget[1]->UAV);
-
-		if (DestDownsampledUAV)
-		{
-			RHICmdList.SetUAVParameter(ShaderRHI, OutComputeTexDownsampled.GetBaseIndex(), DestDownsampledUAV);
-		}
-
-		// VS params
-		SetTextureParameter(RHICmdList, ShaderRHI, EyeAdaptation, EyeAdaptationTex);
-
-		Parameter.SetParameters(RHICmdList, ShaderRHI, Context, InputHistory, PassParameters, bUseDither, SrcSize);
+static void SetupSampleWeightParameters(FTAAShaderParameters* OutTAAParameters, const FTAAPassParameters& PassParameters, FVector2D TemporalJitterPixels)
+{
+	float JitterX = TemporalJitterPixels.X;
+	float JitterY = TemporalJitterPixels.Y;
+	float ResDivisorInv = 1.0f / float(PassParameters.ResolutionDivisor);
 
-		// Temporal AA upscale specific params.
+	static const float SampleOffsets[9][2] =
+	{
+		{ -1.0f, -1.0f },
+		{  0.0f, -1.0f },
+		{  1.0f, -1.0f },
+		{ -1.0f,  0.0f },
+		{  0.0f,  0.0f },
+		{  1.0f,  0.0f },
+		{ -1.0f,  1.0f },
+		{  0.0f,  1.0f },
+		{  1.0f,  1.0f },
+	};
+
+	float FilterSize = CVarTemporalAAFilterSize.GetValueOnRenderThread();
+	int32 bCatmullRom = CVarTemporalAACatmullRom.GetValueOnRenderThread();
+
+	// Compute 3x3 weights
+	{
+		float TotalWeight = 0.0f;
+		for (int32 i = 0; i < 9; i++)
 		{
-			float InputViewSizeInvScale = PassParameters.ResolutionDivisor;
-			float InputViewSizeScale = 1.0f / InputViewSizeInvScale;
+			float PixelOffsetX = SampleOffsets[i][0] - JitterX * ResDivisorInv;
+			float PixelOffsetY = SampleOffsets[i][1] - JitterY * ResDivisorInv;
 
-			SetShaderValue(RHICmdList, ShaderRHI, TemporalJitterPixels, InputViewSizeScale * Context.View.TemporalJitterPixels);
-			SetShaderValue(RHICmdList, ShaderRHI, ScreenPercentageAndUpscaleFactor, FVector2D(
-				float(PassParameters.InputViewRect.Width()) / float(PassParameters.OutputViewRect.Width()),
-				float(PassParameters.OutputViewRect.Width()) / float(PassParameters.InputViewRect.Width())));
+			PixelOffsetX /= FilterSize;
+			PixelOffsetY /= FilterSize;
 
-			SetShaderValue(RHICmdList, ShaderRHI, InputViewMin, InputViewSizeScale * FVector2D(PassParameters.InputViewRect.Min.X, PassParameters.InputViewRect.Min.Y));
-			SetShaderValue(RHICmdList, ShaderRHI, InputViewSize, FVector4(
-				InputViewSizeScale * PassParameters.InputViewRect.Width(), InputViewSizeScale * PassParameters.InputViewRect.Height(),
-				InputViewSizeInvScale / PassParameters.InputViewRect.Width(), InputViewSizeInvScale / PassParameters.InputViewRect.Height()));
+			if (bCatmullRom)
+			{
+				OutTAAParameters->SampleWeights[i] = CatmullRom(PixelOffsetX) * CatmullRom(PixelOffsetY);
+				TotalWeight += OutTAAParameters->SampleWeights[i];
+			}
+			else
+			{
+				// Normal distribution, Sigma = 0.47
+				OutTAAParameters->SampleWeights[i] = FMath::Exp(-2.29f * (PixelOffsetX * PixelOffsetX + PixelOffsetY * PixelOffsetY));
+				TotalWeight += OutTAAParameters->SampleWeights[i];
+			}
 		}
+	
+		for (int32 i = 0; i < 9; i++)
+			OutTAAParameters->SampleWeights[i] /= TotalWeight;
 	}
 
-	template <typename TRHICmdList>
-	void UnsetParameters(TRHICmdList& RHICmdList)
+	// Compute 3x3 + weights.
 	{
-		const FComputeShaderRHIParamRef ShaderRHI = GetComputeShader();
-		if (OutComputeTex0.IsBound())
-			RHICmdList.SetUAVParameter(ShaderRHI, OutComputeTex0.GetBaseIndex(), NULL);
-		if (OutComputeTex1.IsBound())
-			RHICmdList.SetUAVParameter(ShaderRHI, OutComputeTex1.GetBaseIndex(), NULL);
-		if (OutComputeTexDownsampled.IsBound())
-			RHICmdList.SetUAVParameter(ShaderRHI, OutComputeTexDownsampled.GetBaseIndex(), NULL);
+		OutTAAParameters->PlusWeights[0] = OutTAAParameters->SampleWeights[1];
+		OutTAAParameters->PlusWeights[1] = OutTAAParameters->SampleWeights[3];
+		OutTAAParameters->PlusWeights[2] = OutTAAParameters->SampleWeights[4];
+		OutTAAParameters->PlusWeights[3] = OutTAAParameters->SampleWeights[5];
+		OutTAAParameters->PlusWeights[4] = OutTAAParameters->SampleWeights[7];
+		float TotalWeightPlus = (
+			OutTAAParameters->SampleWeights[1] +
+			OutTAAParameters->SampleWeights[3] +
+			OutTAAParameters->SampleWeights[4] +
+			OutTAAParameters->SampleWeights[5] +
+			OutTAAParameters->SampleWeights[7]);
+	
+		for (int32 i = 0; i < 5; i++)
+			OutTAAParameters->PlusWeights[i] /= TotalWeightPlus;
 	}
-
-	FTemporalAAParameters Parameter;
-	FShaderResourceParameter EyeAdaptation;
-	FShaderParameter TemporalAAComputeParams;
-	FShaderParameter OutComputeTex0;
-	FShaderParameter OutComputeTex1;
-	FShaderParameter OutComputeTexDownsampled;
-	FShaderParameter InputViewMin;
-	FShaderParameter InputViewSize;
-	FShaderParameter TemporalJitterPixels;
-	FShaderParameter ScreenPercentageAndUpscaleFactor;
-};
-
-IMPLEMENT_GLOBAL_SHADER(FPostProcessTemporalAACS, "/Engine/Private/PostProcessTemporalAA.usf", "MainCS", SF_Compute);
-
-static inline void TransitionPixelPassResources(FRenderingCompositePassContext& Context)
-{
-	TShaderMapRef< FPostProcessTonemapVS > VertexShader(Context.GetShaderMap());
-	VertexShader->TransitionResources(Context);
-}
-
-void DrawPixelPassTemplate(
-	FRenderingCompositePassContext& Context,
-	const FPostProcessTemporalAAPS::FPermutationDomain& PermutationVector,
-	const FIntPoint SrcSize,
-	const FIntRect ViewRect,
-	const FTemporalAAHistory& InputHistory,
-	const FTAAPassParameters& PassParameters,
-	const bool bUseDither,
-	FDepthStencilStateRHIParamRef DepthStencilState)
-{
-	FGraphicsPipelineStateInitializer GraphicsPSOInit;
-	Context.RHICmdList.ApplyCachedRenderTargets(GraphicsPSOInit);
-	GraphicsPSOInit.BlendState = TStaticBlendState<>::GetRHI();
-	GraphicsPSOInit.RasterizerState = TStaticRasterizerState<>::GetRHI();
-	GraphicsPSOInit.DepthStencilState = DepthStencilState;
-
-	TShaderMapRef< FPostProcessTonemapVS > VertexShader(Context.GetShaderMap());
-	TShaderMapRef< FPostProcessTemporalAAPS > PixelShader(Context.GetShaderMap(), PermutationVector);
-
-	GraphicsPSOInit.BoundShaderState.VertexDeclarationRHI = GFilterVertexDeclaration.VertexDeclarationRHI;
-	GraphicsPSOInit.BoundShaderState.VertexShaderRHI = GETSAFERHISHADER_VERTEX(*VertexShader);
-	GraphicsPSOInit.BoundShaderState.PixelShaderRHI = GETSAFERHISHADER_PIXEL(*PixelShader);
-	GraphicsPSOInit.PrimitiveType = PT_TriangleList;
-
-	SetGraphicsPipelineState(Context.RHICmdList, GraphicsPSOInit);
-	VertexShader->SetVS(Context);
-	PixelShader->SetParameters(Context.RHICmdList, Context, InputHistory, PassParameters, bUseDither, SrcSize);
-
-	DrawRectangle(
-		Context.RHICmdList,
-		0, 0,
-		ViewRect.Width(), ViewRect.Height(),
-		ViewRect.Min.X, ViewRect.Min.Y,
-		ViewRect.Width(), ViewRect.Height(),
-		ViewRect.Size(),
-		SrcSize,
-		*VertexShader,
-		EDRF_UseTriangleOptimization);
-}
-
-template<class TRHICmdList>
-void DispatchCSTemplate(
-	TRHICmdList& RHICmdList,
-	FRenderingCompositePassContext& Context,
-	const FPostProcessTemporalAACS::FPermutationDomain& PermutationVector,
-	const FTemporalAAHistory& InputHistory,
-	const FTAAPassParameters& PassParameters,
-	const FIntPoint& SrcSize,
-	const FSceneRenderTargetItem* DestRenderTarget[2],
-	FUnorderedAccessViewRHIParamRef DestDownsampledUAV,
-	const bool bUseDither,
-	FTextureRHIParamRef EyeAdaptationTex)
-{
-	auto ShaderMap = Context.GetShaderMap();
-	TShaderMapRef<FPostProcessTemporalAACS> ComputeShader(ShaderMap, PermutationVector);
-
-	RHICmdList.SetComputeShader(ComputeShader->GetComputeShader());
-
-	FIntPoint DestSize = FIntPoint::DivideAndRoundUp(PassParameters.OutputViewRect.Size(), PassParameters.ResolutionDivisor);
-	ComputeShader->SetParameters(RHICmdList, Context, InputHistory, PassParameters, DestSize, DestRenderTarget, DestDownsampledUAV, SrcSize, bUseDither, EyeAdaptationTex);
-
-	uint32 GroupSizeX = FMath::DivideAndRoundUp(DestSize.X, GTemporalAATileSizeX);
-	uint32 GroupSizeY = FMath::DivideAndRoundUp(DestSize.Y, GTemporalAATileSizeY);
-	DispatchComputeShader(RHICmdList, *ComputeShader, GroupSizeX, GroupSizeY, 1);
-
-	ComputeShader->UnsetParameters(RHICmdList);
-}
+} // SetupSampleWeightParameters()
 
 
 DECLARE_GPU_STAT(TAA)
@@ -695,95 +395,241 @@ static_assert(ARRAY_COUNT(kTAAOutputNames) == int32(ETAAPassConfig::MAX), "Missi
 static_assert(ARRAY_COUNT(kTAAPassNames) == int32(ETAAPassConfig::MAX), "Missing TAA pass name.");
 
 
-FRCPassPostProcessTemporalAA::FRCPassPostProcessTemporalAA(
-	const FPostprocessContext& Context,
-	const FTAAPassParameters& InParameters,
-	const FTemporalAAHistory& InInputHistory,
-	FTemporalAAHistory* OutOutputHistory)
-	: Parameters(InParameters)
-	, OutputExtent(0, 0)
-	, InputHistory(InInputHistory)
-	, OutputHistory(OutOutputHistory)
+FIntPoint FTAAPassParameters::GetOutputExtent() const
 {
-	bIsComputePass = Parameters.bIsComputePass;
-	bPreferAsyncCompute = false;
-	bPreferAsyncCompute &= (GNumAlternateFrameRenderingGroups == 1); // Can't handle multi-frame updates on async pipe
+	check(Validate());
+	check(SceneColorInput);
 
-	bDownsamplePossible = Parameters.bDownsample &&
-		bIsComputePass &&
-		IsMainTAAConfig(Parameters.Pass);
+	FIntPoint InputExtent = SceneColorInput->Desc.Extent;
 
-	if (IsTAAUpsamplingConfig(Parameters.Pass))
-	{
-		bIsComputePass = true;
+	if (!IsTAAUpsamplingConfig(Pass))
+		return InputExtent;
 
-		check(Parameters.OutputViewRect.Min == FIntPoint::ZeroValue);
-		FIntPoint PrimaryUpscaleViewSize = Parameters.OutputViewRect.Size();
-		FIntPoint QuantizedPrimaryUpscaleViewSize;
-		QuantizeSceneBufferSize(PrimaryUpscaleViewSize, QuantizedPrimaryUpscaleViewSize);
+	check(OutputViewRect.Min == FIntPoint::ZeroValue);
+	FIntPoint PrimaryUpscaleViewSize = FIntPoint::DivideAndRoundUp(OutputViewRect.Size(), ResolutionDivisor);
+	FIntPoint QuantizedPrimaryUpscaleViewSize;
+	QuantizeSceneBufferSize(PrimaryUpscaleViewSize, QuantizedPrimaryUpscaleViewSize);
 
-		FSceneRenderTargets& SceneContext = FSceneRenderTargets::Get(Context.RHICmdList);
-		OutputExtent.X = FMath::Max(SceneContext.GetBufferSizeXY().X, QuantizedPrimaryUpscaleViewSize.X);
-		OutputExtent.Y = FMath::Max(SceneContext.GetBufferSizeXY().Y, QuantizedPrimaryUpscaleViewSize.Y);
-	}
+	return FIntPoint(
+		FMath::Max(InputExtent.X, QuantizedPrimaryUpscaleViewSize.X),
+		FMath::Max(InputExtent.Y, QuantizedPrimaryUpscaleViewSize.Y));
 }
 
-void FRCPassPostProcessTemporalAA::Process(FRenderingCompositePassContext& Context)
+bool FTAAPassParameters::Validate() const
 {
-	AsyncEndFence = FComputeFenceRHIRef();
+	if (IsTAAUpsamplingConfig(Pass))
+	{
+		check(OutputViewRect.Min == FIntPoint::ZeroValue);
+	}
+	else
+	{
+		check(InputViewRect == OutputViewRect);
+	}
+	return true;
+}
 
-	FSceneRenderTargets& SceneContext = FSceneRenderTargets::Get(Context.RHICmdList);
 
-	const FPooledRenderTargetDesc* InputDesc = GetInputDesc(ePId_Input0);
-	FIntPoint SrcSize = InputDesc->Extent;
+FTAAOutputs FTAAPassParameters::AddTemporalAAPass(
+	FRDGBuilder& GraphBuilder,
+	const FSceneViewFamilyBlackboard& SceneBlackboard,
+	const FViewInfo& View,
+	const FTemporalAAHistory& InputHistory,
+	FTemporalAAHistory* OutputHistory) const
+{
+	check(Validate());
 
 	// Number of render target in TAA history.
-	const int32 RenderTargetCount = IsDOFTAAConfig(Parameters.Pass) && FPostProcessing::HasAlphaChannelSupport() ? 2 : 1;
+	const int32 RenderTargetCount = IsDOFTAAConfig(Pass) && FPostProcessing::HasAlphaChannelSupport() ? 2 : 1;
+	
+	// Whether this is main TAA pass;
+	bool bIsMainPass = IsMainTAAConfig(Pass);
+	
+	// Whether to use camera cut shader permutation or not.
+	const bool bCameraCut = !InputHistory.IsValid() || View.bCameraCut;
 
-	const FSceneRenderTargetItem* DestRenderTarget[2] = {nullptr, nullptr};
-	DestRenderTarget[0] = &PassOutputs[0].RequestSurface(Context);
-	if (RenderTargetCount == 2)
-		DestRenderTarget[1] = &PassOutputs[1].RequestSurface(Context);
+	FIntPoint OutputExtent = GetOutputExtent();
 
-	const FSceneRenderTargetItem& DestDownsampled = bDownsamplePossible ? PassOutputs[2].RequestSurface(Context) : FSceneRenderTargetItem();
+	// Src rectangle.
+	FIntRect SrcRect = InputViewRect;
+	FIntRect DestRect = OutputViewRect;
+	FIntRect PracticableSrcRect = FIntRect::DivideAndRoundUp(SrcRect, ResolutionDivisor);
+	FIntRect PracticableDestRect = FIntRect::DivideAndRoundUp(DestRect, ResolutionDivisor);
 
-	// Whether this is main TAA pass;
-	bool bIsMainPass = IsMainTAAConfig(Parameters.Pass);
+	// Name of the pass.
+	const TCHAR* PassName = kTAAPassNames[static_cast<int32>(Pass)];
 
-	// Whether to use camera cut shader permutation or not.
-	const bool bCameraCut = !InputHistory.IsValid() || Context.View.bCameraCut;
+	// Setups common shader parameters
+	FTAAShaderParameters CommonShaderParameters;
+	{
+		FIntPoint SrcSize = SceneColorInput->Desc.Extent;
 
-	// Whether to use responsive stencil test.
-	bool bUseResponsiveStencilTest = Parameters.Pass == ETAAPassConfig::Main && !bIsComputePass && !bCameraCut;
+		if (!IsTAAUpsamplingConfig(Pass))
+			SetupSampleWeightParameters(&CommonShaderParameters, *this, View.TemporalJitterPixels);
+	
+		float ResDivisor = ResolutionDivisor;
+		float ResDivisorInv = 1.0f / ResDivisor;
 
-	// Only use dithering if we are outputting to a low precision format
-	const bool bUseDither = PassOutputs[0].RenderTargetDesc.Format != PF_FloatRGBA && bIsMainPass;
+		CommonShaderParameters.ViewUniformBuffer = View.ViewUniformBuffer;
+		CommonShaderParameters.CurrentFrameWeight = CVarTemporalAACurrentFrameWeight.GetValueOnRenderThread();
 
-	// Src rectangle.
-	FIntRect SrcRect = Parameters.InputViewRect;
+		CommonShaderParameters.SceneBlackboard = SceneBlackboard;
+		CommonShaderParameters.SceneDepthBufferSampler = TStaticSamplerState<SF_Point>::GetRHI();
+		CommonShaderParameters.SceneVelocityBufferSampler = TStaticSamplerState<SF_Point>::GetRHI();
+
+		// Input buffer shader parameters
+		{
+			CommonShaderParameters.InputSceneColorSize = FVector4(
+				SceneColorInput->Desc.Extent.X,
+				SceneColorInput->Desc.Extent.Y,
+				1.0f / float(SceneColorInput->Desc.Extent.X),
+				1.0f / float(SceneColorInput->Desc.Extent.Y));
+			CommonShaderParameters.InputSceneColor = SceneColorInput;
+			CommonShaderParameters.InputSceneColorSampler = TStaticSamplerState<SF_Point>::GetRHI();
+			CommonShaderParameters.InputSceneMetadata = SceneMetadataInput;
+			CommonShaderParameters.InputSceneMetadataSampler = TStaticSamplerState<SF_Point>::GetRHI();
+		}
 
-	// Dest rectangle is same as source rectangle, unless Upsampling.
-	FIntRect DestRect = Parameters.OutputViewRect;
-	check(IsTAAUpsamplingConfig(Parameters.Pass) || SrcRect == DestRect);
+		CommonShaderParameters.OutputViewportSize = FVector4(
+			PracticableDestRect.Width(), PracticableDestRect.Height(), 1.0f / float(PracticableDestRect.Width()), 1.0f / float(PracticableDestRect.Height()));
+		CommonShaderParameters.OutputViewportRect = FVector4(PracticableDestRect.Min.X, PracticableDestRect.Min.Y, PracticableDestRect.Max.X, PracticableDestRect.Max.Y);
 
-	// Name of the pass.
-	const TCHAR* PassName = kTAAPassNames[static_cast<int32>(Parameters.Pass)];
+		// Set history shader parameters.
+		{
+			FIntPoint ReferenceViewportOffset = InputHistory.ViewportRect.Min;
+			FIntPoint ReferenceViewportExtent = InputHistory.ViewportRect.Size();
+			FIntPoint ReferenceBufferSize = InputHistory.ReferenceBufferSize;
+
+			float InvReferenceBufferSizeX = 1.f / float(InputHistory.ReferenceBufferSize.X);
+			float InvReferenceBufferSizeY = 1.f / float(InputHistory.ReferenceBufferSize.Y);
+
+			CommonShaderParameters.ScreenPosToHistoryBufferUV = FVector4(
+				ReferenceViewportExtent.X * 0.5f * InvReferenceBufferSizeX,
+				-ReferenceViewportExtent.Y * 0.5f * InvReferenceBufferSizeY,
+				(ReferenceViewportExtent.X * 0.5f + ReferenceViewportOffset.X) * InvReferenceBufferSizeX,
+				(ReferenceViewportExtent.Y * 0.5f + ReferenceViewportOffset.Y) * InvReferenceBufferSizeY);
+
+			FIntPoint ViewportOffset = ReferenceViewportOffset / ResolutionDivisor;
+			FIntPoint ViewportExtent = FIntPoint::DivideAndRoundUp(ReferenceViewportExtent, ResolutionDivisor);
+			FIntPoint BufferSize = ReferenceBufferSize / ResolutionDivisor;
+
+			CommonShaderParameters.ScreenPosAbsMax = FVector2D(1.0f - 1.0f / float(ViewportExtent.X), 1.0f - 1.0f / float(ViewportExtent.Y));
+
+			float InvBufferSizeX = 1.f / float(BufferSize.X);
+			float InvBufferSizeY = 1.f / float(BufferSize.Y);
+
+			CommonShaderParameters.HistoryBufferUVMinMax = FVector4(
+				(ViewportOffset.X + 0.5f) * InvBufferSizeX,
+				(ViewportOffset.Y + 0.5f) * InvBufferSizeY,
+				(ViewportOffset.X + ViewportExtent.X - 0.5f) * InvBufferSizeX,
+				(ViewportOffset.Y + ViewportExtent.Y - 0.5f) * InvBufferSizeY);
+
+			CommonShaderParameters.HistoryBufferSize = FVector4(BufferSize.X, BufferSize.Y, InvBufferSizeX, InvBufferSizeY);
+
+			if (InputHistory.IsValid())
+			{
+				CommonShaderParameters.HistoryBuffer0 = GraphBuilder.RegisterExternalTexture(InputHistory.RT[0]);
+				if (InputHistory.RT[1].IsValid())
+					CommonShaderParameters.HistoryBuffer1 = GraphBuilder.RegisterExternalTexture(InputHistory.RT[1]);
+			}
+
+			CommonShaderParameters.HistoryBuffer0Sampler = TStaticSamplerState<SF_Bilinear>::GetRHI();
+			CommonShaderParameters.HistoryBuffer1Sampler = TStaticSamplerState<SF_Bilinear>::GetRHI();
+		}
+
+		CommonShaderParameters.MaxViewportUVAndSvPositionToViewportUV = FVector4(
+			(PracticableDestRect.Width() - 0.5f * ResDivisor) / float(PracticableDestRect.Width()),
+			(PracticableDestRect.Height() - 0.5f * ResDivisor) / float(PracticableDestRect.Height()),
+			ResDivisor / float(DestRect.Width()),
+			ResDivisor / float(DestRect.Height()));
+
+		// Pre-exposure, One over Pre-exposure, History pre-exposure, History one over pre-exposure.
+		// DOF settings must preserve scene color range.
+		CommonShaderParameters.PreExposureSettings = FVector4(1.f, 1.f, 1.f, 1.f);
+		if (Pass == ETAAPassConfig::Main)
+		{
+			CommonShaderParameters.PreExposureSettings.X = View.PreExposure;
+			CommonShaderParameters.PreExposureSettings.Y = 1.f / FMath::Max<float>(SMALL_NUMBER, View.PreExposure);
+			CommonShaderParameters.PreExposureSettings.Z = InputHistory.IsValid() ? InputHistory.SceneColorPreExposure : View.PreExposure;
+			CommonShaderParameters.PreExposureSettings.W  = 1.f / FMath::Max<float>(SMALL_NUMBER, CommonShaderParameters.PreExposureSettings.Z);
+		}
+
+		{
+			float InvSizeX = 1.0f / float(SrcSize.X);
+			float InvSizeY = 1.0f / float(SrcSize.Y);
+			CommonShaderParameters.ViewportUVToInputBufferUV = FVector4(
+				ResDivisorInv * InputViewRect.Width() * InvSizeX,
+				ResDivisorInv * InputViewRect.Height() * InvSizeY,
+				ResDivisorInv * InputViewRect.Min.X * InvSizeX,
+				ResDivisorInv * InputViewRect.Min.Y * InvSizeY);
+		}
+		
+		CommonShaderParameters.EyeAdaptation = GetEyeAdaptationTexture(GraphBuilder, View);
+	}
+
+	// Create outputs
+	FTAAOutputs Outputs;
+	{
+		FRDGTextureDesc SceneColorDesc = SceneColorInput->Desc;
+		{
+			SceneColorDesc.Flags &= ~(TexCreate_FastVRAM | TexCreate_Transient);
+			SceneColorDesc.Reset();
+			//regardless of input type, PF_FloatRGBA is required to properly accumulate between frames for a good result.
+			SceneColorDesc.Format = PF_FloatRGBA;
+			SceneColorDesc.AutoWritable = false;
+			SceneColorDesc.TargetableFlags &= ~(TexCreate_RenderTargetable | TexCreate_UAV);
+			SceneColorDesc.TargetableFlags |= bIsComputePass ? TexCreate_UAV : TexCreate_RenderTargetable;
+
+			// Need a UAV to resource transition from gfx to compute.
+			if (IsDOFTAAConfig(Pass))
+			{
+				SceneColorDesc.TargetableFlags |= TexCreate_UAV;
+			}
+
+			Outputs.SceneColor = GraphBuilder.CreateTexture(SceneColorDesc, kTAAOutputNames[static_cast<int32>(Pass)]);
+		}
+
+		if (RenderTargetCount == 2)
+		{
+			Outputs.SceneMetadata = GraphBuilder.CreateTexture(SceneColorDesc, kTAAOutputNames[static_cast<int32>(Pass)]);
+		}
+
+		if (bDownsample)
+		{
+			check(bIsComputePass);
+			
+			FRDGTextureDesc HalfResSceneColorDesc = SceneColorInput->Desc;
+			HalfResSceneColorDesc.Flags &= ~(TexCreate_FastVRAM);
+			HalfResSceneColorDesc.Reset();
+
+			if (DownsampleOverrideFormat != PF_Unknown)
+			{
+				HalfResSceneColorDesc.Format = DownsampleOverrideFormat;
+			}
 
-	// Stats.
-	SCOPED_GPU_STAT(Context.RHICmdList, TAA);
+			HalfResSceneColorDesc.AutoWritable = false;
+			HalfResSceneColorDesc.TargetableFlags &= ~TexCreate_RenderTargetable;
+			HalfResSceneColorDesc.TargetableFlags |= TexCreate_UAV;
+
+			HalfResSceneColorDesc.Extent = FIntPoint::DivideAndRoundUp(HalfResSceneColorDesc.Extent, 2);
+			HalfResSceneColorDesc.Extent.X = FMath::Max(1, HalfResSceneColorDesc.Extent.X);
+			HalfResSceneColorDesc.Extent.Y = FMath::Max(1, HalfResSceneColorDesc.Extent.Y);
+			
+			Outputs.DownsampledSceneColor = GraphBuilder.CreateTexture(HalfResSceneColorDesc, TEXT("SceneColorHalfRes"));
+		}
+	}
 
 	if (bIsComputePass)
 	{
-		FPostProcessTemporalAACS::FPermutationDomain PermutationVector;
-		PermutationVector.Set<FTAAPassConfigDim>(Parameters.Pass);
-		PermutationVector.Set<FTAAFastDim>(Parameters.bUseFast);
-		PermutationVector.Set<FTAACameraCutDim>(!Context.View.PrevViewInfo.TemporalAAHistory.IsValid());
-		PermutationVector.Set<FTAADownsampleDim>(DestDownsampled.IsValid());
+		FTemporalAACS::FPermutationDomain PermutationVector;
+		PermutationVector.Set<FTAAPassConfigDim>(Pass);
+		PermutationVector.Set<FTAAFastDim>(bUseFast);
+		PermutationVector.Set<FTAACameraCutDim>(!View.PrevViewInfo.TemporalAAHistory.IsValid());
+		PermutationVector.Set<FTAADownsampleDim>(bDownsample);
 		PermutationVector.Set<FTAAUpsampleFilteredDim>(true);
 
-		if (IsTAAUpsamplingConfig(Parameters.Pass))
+		if (IsTAAUpsamplingConfig(Pass))
 		{
-			const bool bUpsampleFiltered = CVarTemporalAAUpsampleFiltered.GetValueOnRenderThread() != 0 || Parameters.Pass != ETAAPassConfig::MainUpsampling;
+			const bool bUpsampleFiltered = CVarTemporalAAUpsampleFiltered.GetValueOnRenderThread() != 0 || Pass != ETAAPassConfig::MainUpsampling;
 			PermutationVector.Set<FTAAUpsampleFilteredDim>(bUpsampleFiltered);
 
 			// If screen percentage > 100% on X or Y axes, then use screen percentage range = 2 shader permutation to disable LDS caching.
@@ -796,7 +642,7 @@ void FRCPassPostProcessTemporalAA::Process(FRenderingCompositePassContext& Conte
 			else if (SrcRect.Width() * 100 < 50 * DestRect.Width() &&
 				SrcRect.Height() * 100 < 50 * DestRect.Height())
 			{
-				check(Parameters.Pass == ETAAPassConfig::MainSuperSampling);
+				check(Pass == ETAAPassConfig::MainSuperSampling);
 				PermutationVector.Set<FTAAScreenPercentageDim>(3);
 			}
 			// If screen percentage < 71% on X and Y axes, then use screen percentage range = 1 shader permutation to have smaller LDS caching.
@@ -807,145 +653,116 @@ void FRCPassPostProcessTemporalAA::Process(FRenderingCompositePassContext& Conte
 			}
 		}
 
-		FIntRect PracticableSrcRect = FIntRect::DivideAndRoundUp(SrcRect, Parameters.ResolutionDivisor);
-		FIntRect PracticableDestRect = FIntRect::DivideAndRoundUp(DestRect, Parameters.ResolutionDivisor);
+		FTemporalAACS::FParameters* PassParameters = GraphBuilder.AllocParameters<FTemporalAACS::FParameters>();
+		PassParameters->CommonParameters = CommonShaderParameters;
 
-		SCOPED_DRAW_EVENTF(Context.RHICmdList, TemporalAA, TEXT("TAA %s CS%s %dx%d -> %dx%d"),
-			PassName, Parameters.bUseFast ? TEXT(" Fast") : TEXT(""), PracticableSrcRect.Width(), PracticableSrcRect.Height(), PracticableDestRect.Width(), PracticableDestRect.Height());
-
-		// Common setup
-		// #todo-renderpass remove once everything is renderpasses
-		UnbindRenderTargets(Context.RHICmdList);
-		Context.SetViewportAndCallRHI(PracticableDestRect, 0.0f, 1.0f);
-		
-		static FName AsyncEndFenceName(TEXT("AsyncTemporalAAEndFence"));
-		AsyncEndFence = Context.RHICmdList.CreateComputeFence(AsyncEndFenceName);
-
-		FTextureRHIRef EyeAdaptationTex = GWhiteTexture->TextureRHI;
-		if (Context.View.HasValidEyeAdaptation())
+		// Temporal upsample specific shader parameters.
 		{
-			EyeAdaptationTex = Context.View.GetEyeAdaptation(Context.RHICmdList)->GetRenderTargetItem().TargetableTexture;
-		}
-
-		FUnorderedAccessViewRHIParamRef UAVs[2];
-		UAVs[0] = DestRenderTarget[0]->UAV;
-		if (RenderTargetCount == 2)
-			UAVs[1] = DestRenderTarget[1]->UAV;
+			// Temporal AA upscale specific params.
+			float InputViewSizeInvScale = ResolutionDivisor;
+			float InputViewSizeScale = 1.0f / InputViewSizeInvScale;
 
-		if (IsAsyncComputePass())
-		{
-			// Async path
-			FRHIAsyncComputeCommandListImmediate& RHICmdListComputeImmediate = FRHICommandListExecutor::GetImmediateAsyncComputeCommandList();
-			{
-				SCOPED_COMPUTE_EVENT(RHICmdListComputeImmediate, AsyncTemporalAA);
-				WaitForInputPassComputeFences(RHICmdListComputeImmediate);
-					
-				RHICmdListComputeImmediate.TransitionResources(
-					EResourceTransitionAccess::ERWBarrier, EResourceTransitionPipeline::EGfxToCompute,
-					UAVs, RenderTargetCount);
-
-				if (DestDownsampled.IsValid())
-				{
-					RHICmdListComputeImmediate.TransitionResource(EResourceTransitionAccess::ERWBarrier, EResourceTransitionPipeline::EGfxToCompute, DestDownsampled.UAV);
-				}
-
-				DispatchCSTemplate(
-					RHICmdListComputeImmediate, Context, PermutationVector,
-					InputHistory, Parameters, SrcSize, DestRenderTarget, DestDownsampled.UAV, bUseDither, EyeAdaptationTex);
-
-				RHICmdListComputeImmediate.TransitionResources(
-					EResourceTransitionAccess::EReadable, EResourceTransitionPipeline::EComputeToGfx,
-					UAVs, RenderTargetCount, AsyncEndFence);
-				if (DestDownsampled.IsValid())
-				{
-					RHICmdListComputeImmediate.TransitionResource(EResourceTransitionAccess::EReadable, EResourceTransitionPipeline::EComputeToGfx, DestDownsampled.UAV);
-				}
-			}
-			FRHIAsyncComputeCommandListImmediate::ImmediateDispatch(RHICmdListComputeImmediate);
+			PassParameters->TemporalJitterPixels = InputViewSizeScale * View.TemporalJitterPixels;
+			PassParameters->ScreenPercentage = float(InputViewRect.Width()) / float(OutputViewRect.Width());
+			PassParameters->UpscaleFactor = float(OutputViewRect.Width()) / float(InputViewRect.Width());
+			PassParameters->InputViewMin = InputViewSizeScale * FVector2D(InputViewRect.Min.X, InputViewRect.Min.Y);
+			PassParameters->InputViewSize = FVector4(
+				InputViewSizeScale * InputViewRect.Width(), InputViewSizeScale * InputViewRect.Height(),
+				InputViewSizeInvScale / InputViewRect.Width(), InputViewSizeInvScale / InputViewRect.Height());
 		}
-		else
-		{
-			// Direct path
-			WaitForInputPassComputeFences(Context.RHICmdList);
-			Context.RHICmdList.BeginUpdateMultiFrameResource(DestRenderTarget[0]->ShaderResourceTexture);
-			if (RenderTargetCount == 2)
-				Context.RHICmdList.BeginUpdateMultiFrameResource(DestRenderTarget[1]->ShaderResourceTexture);
 
-			Context.RHICmdList.TransitionResources(
-				EResourceTransitionAccess::ERWBarrier, EResourceTransitionPipeline::EGfxToCompute,
-				UAVs, RenderTargetCount);
-
-			if (DestDownsampled.IsValid())
-			{
-				Context.RHICmdList.TransitionResource(EResourceTransitionAccess::ERWBarrier, EResourceTransitionPipeline::EGfxToCompute, DestDownsampled.UAV);
-			}
-
-			DispatchCSTemplate(
-				Context.RHICmdList, Context, PermutationVector,
-				InputHistory, Parameters, SrcSize, DestRenderTarget, DestDownsampled.UAV, bUseDither, EyeAdaptationTex);
-
-			Context.RHICmdList.TransitionResources(
-				EResourceTransitionAccess::EReadable, EResourceTransitionPipeline::EComputeToGfx,
-				UAVs, RenderTargetCount, AsyncEndFence);
-
-			if (DestDownsampled.IsValid())
-			{
-				Context.RHICmdList.TransitionResource(EResourceTransitionAccess::EReadable, EResourceTransitionPipeline::EComputeToGfx, DestDownsampled.UAV);
-			}
-
-			Context.RHICmdList.EndUpdateMultiFrameResource(DestRenderTarget[0]->ShaderResourceTexture);
-			if (RenderTargetCount == 2)
-				Context.RHICmdList.EndUpdateMultiFrameResource(DestRenderTarget[1]->ShaderResourceTexture);
+		// UAVs
+		{
+			PassParameters->OutComputeTex0 = GraphBuilder.CreateUAV(Outputs.SceneColor);
+			if (Outputs.SceneMetadata)
+				PassParameters->OutComputeTex1 = GraphBuilder.CreateUAV(Outputs.SceneMetadata);
+			if (Outputs.DownsampledSceneColor)
+				PassParameters->OutComputeTexDownsampled = GraphBuilder.CreateUAV(Outputs.DownsampledSceneColor);
 		}
+		
+#if 0 // TODO.
+		Context.RHICmdList.BeginUpdateMultiFrameResource(DestRenderTarget[0]->ShaderResourceTexture);
+		if (RenderTargetCount == 2)
+			Context.RHICmdList.BeginUpdateMultiFrameResource(DestRenderTarget[1]->ShaderResourceTexture);
+#endif
+
+		TShaderMapRef<FTemporalAACS> ComputeShader(View.ShaderMap, PermutationVector);
+		ClearUnusedGraphResources(*ComputeShader, PassParameters);
+		GraphBuilder.AddPass(
+			RDG_EVENT_NAME("TAA %s CS%s %dx%d -> %dx%d",
+				PassName, bUseFast ? TEXT(" Fast") : TEXT(""),
+				PracticableSrcRect.Width(), PracticableSrcRect.Height(),
+				PracticableDestRect.Width(), PracticableDestRect.Height()),
+			PassParameters,
+			ERenderGraphPassFlags::Compute,
+			[PassParameters, ComputeShader, PracticableDestRect](FRHICommandList& RHICmdList)
+		{
+			SCOPED_GPU_STAT(RHICmdList, TAA); // TODO(RDG) find a way to integrate GPU stat groups into RDG so can use FComputeShaderUtils::AddPass()
+			FComputeShaderUtils::Dispatch(
+				RHICmdList, *ComputeShader, *PassParameters,
+				FComputeShaderUtils::GetGroupCount(PracticableDestRect.Size(), GTemporalAATileSizeX));
+		});
 	}
 	else
 	{
-		check(!IsTAAUpsamplingConfig(Parameters.Pass));
-
-		FIntRect ViewRect = FIntRect::DivideAndRoundUp(DestRect, Parameters.ResolutionDivisor);
-		FIntPoint DestSize = PassOutputs[0].RenderTargetDesc.Extent;
-
-		SCOPED_DRAW_EVENTF(Context.RHICmdList, TemporalAA, TEXT("TAA %s PS%s %dx%d"),
-			PassName, Parameters.bUseFast ? TEXT(" Fast") : TEXT(""), ViewRect.Width(), ViewRect.Height());
-
-		WaitForInputPassComputeFences(Context.RHICmdList);
-
-		// Inform MultiGPU systems that we're starting to update this resource
-		Context.RHICmdList.BeginUpdateMultiFrameResource(DestRenderTarget[0]->ShaderResourceTexture);
-
-		// make sure we transition resources before we begin the render pass on Vulkan (which happens when we call SetRenderTargets)
-		TransitionPixelPassResources(Context);
-
-		// Setup render targets.
+		check(!IsTAAUpsamplingConfig(Pass));
+
+		// Whether to use responsive stencil test.
+		bool bUseResponsiveStencilTest = Pass == ETAAPassConfig::Main && !bIsComputePass && !bCameraCut;
+	
+		FTemporalAAPS::FParameters* PassParameters = GraphBuilder.AllocParameters<FTemporalAAPS::FParameters>();
+		PassParameters->CommonParameters = CommonShaderParameters;
+		PassParameters->RenderTargets[0] = FRenderTargetBinding(
+			Outputs.SceneColor,
+			ERenderTargetLoadAction::ENoAction,
+			ERenderTargetStoreAction::EStore);
 		
-		// Inform MultiGPU systems that we're starting to update this resource
-		Context.RHICmdList.BeginUpdateMultiFrameResource(DestRenderTarget[0]->ShaderResourceTexture);
-		FRHIRenderPassInfo RPInfo(DestRenderTarget[0]->TargetableTexture, ERenderTargetActions::DontLoad_Store);
-
-		if (RenderTargetCount == 2)
+		if (bUseResponsiveStencilTest)
 		{
-			RPInfo.ColorRenderTargets[1].RenderTarget = DestRenderTarget[1]->TargetableTexture;
-			RPInfo.ColorRenderTargets[1].Action = ERenderTargetActions::DontLoad_Store;
-			RPInfo.ColorRenderTargets[1].ArraySlice = -1;
-			RPInfo.ColorRenderTargets[1].MipIndex = 0;
-
-			Context.RHICmdList.BeginUpdateMultiFrameResource(DestRenderTarget[1]->ShaderResourceTexture);
+			PassParameters->RenderTargets.DepthStencil.Texture = SceneBlackboard.SceneDepthBuffer;
+			PassParameters->RenderTargets.DepthStencil.DepthLoadAction = ERenderTargetLoadAction::ENoAction;
+			PassParameters->RenderTargets.DepthStencil.DepthStoreAction = ERenderTargetStoreAction::ENoAction;
+			PassParameters->RenderTargets.DepthStencil.StencilLoadAction = ERenderTargetLoadAction::ELoad;
+			PassParameters->RenderTargets.DepthStencil.StencilStoreAction = ERenderTargetStoreAction::EStore;
+			PassParameters->RenderTargets.DepthStencil.DepthStencilAccess = FExclusiveDepthStencil::DepthRead_StencilWrite;
 		}
+		
+		FTemporalAAPS::FPermutationDomain BasePermutationVector;
+		BasePermutationVector.Set<FTAAPassConfigDim>(Pass);
+		BasePermutationVector.Set<FTAAFastDim>(bUseFast);
+		BasePermutationVector.Set<FTAACameraCutDim>(bCameraCut);
 
-		RPInfo.DepthStencilRenderTarget.DepthStencilTarget = SceneContext.GetSceneDepthTexture();
-		RPInfo.DepthStencilRenderTarget.ResolveTarget = nullptr;
-		RPInfo.DepthStencilRenderTarget.Action = MakeDepthStencilTargetActions(ERenderTargetActions::DontLoad_DontStore, ERenderTargetActions::Load_Store);
-		RPInfo.DepthStencilRenderTarget.ExclusiveDepthStencil = FExclusiveDepthStencil::DepthRead_StencilWrite;
+		TShaderMapRef<FTemporalAAPS> PixelShader(View.ShaderMap, BasePermutationVector);
+		ClearUnusedGraphResources(*PixelShader, PassParameters);
 
-		Context.RHICmdList.BeginRenderPass(RPInfo, TEXT("TemporalAA"));
+		GraphBuilder.AddPass(
+			RDG_EVENT_NAME("TAA %s PS%s %dx%d",
+				PassName, bUseFast ? TEXT(" Fast") : TEXT(""),
+				PracticableDestRect.Width(), PracticableDestRect.Height()),
+			PassParameters,
+			ERenderGraphPassFlags::None,
+			[PassParameters, &View, PracticableDestRect, bCameraCut, BasePermutationVector, bUseResponsiveStencilTest](FRHICommandList& RHICmdList)
 		{
-			Context.SetViewportAndCallRHI(ViewRect);
+			SCOPED_GPU_STAT(RHICmdList, TAA);
+			RHICmdList.SetViewport(PracticableDestRect.Min.X, PracticableDestRect.Min.Y, 0.0f, PracticableDestRect.Max.X, PracticableDestRect.Max.Y, 1.0f);
+
+			FTemporalAAPS::FPermutationDomain PermutationVector = BasePermutationVector;
 
-			FPostProcessTemporalAAPS::FPermutationDomain PermutationVector;
-			PermutationVector.Set<FTAAPassConfigDim>(Parameters.Pass);
-			PermutationVector.Set<FTAAFastDim>(Parameters.bUseFast);
-			PermutationVector.Set<FTAACameraCutDim>(bCameraCut);
+			// Lambda to draw pixel shader.
+			auto DrawTAAPixelShader = [&](FDepthStencilStateRHIParamRef DepthStencilState)
+			{
+				TShaderMapRef<FTemporalAAPS> PixelShader(View.ShaderMap, PermutationVector);
 
+				FGraphicsPipelineStateInitializer GraphicsPSOInit;
+				FPixelShaderUtils::InitFullscreenPipelineState(RHICmdList, View.ShaderMap, *PixelShader, GraphicsPSOInit);
+				GraphicsPSOInit.DepthStencilState = DepthStencilState;
+				SetGraphicsPipelineState(RHICmdList, GraphicsPSOInit);
+
+				SetShaderParameters(RHICmdList, *PixelShader, PixelShader->GetPixelShader(), *PassParameters);
+				
+				FPixelShaderUtils::DrawFullscreenTriangle(RHICmdList);
+			};
+	
 			if (bUseResponsiveStencilTest)
 			{
 				// Normal temporal feedback
@@ -956,10 +773,7 @@ void FRCPassPostProcessTemporalAA::Process(FRenderingCompositePassContext& Conte
 					false, CF_Always, SO_Keep, SO_Keep, SO_Keep,
 					STENCIL_TEMPORAL_RESPONSIVE_AA_MASK, STENCIL_TEMPORAL_RESPONSIVE_AA_MASK>::GetRHI();
 
-				DrawPixelPassTemplate(
-					Context, PermutationVector, SrcSize, ViewRect,
-					InputHistory, Parameters, bUseDither,
-					DepthStencilState);
+				DrawTAAPixelShader(DepthStencilState);
 
 				// Responsive feedback for tagged pixels
 				// Draw to pixels where stencil != 0
@@ -970,142 +784,87 @@ void FRCPassPostProcessTemporalAA::Process(FRenderingCompositePassContext& Conte
 					STENCIL_TEMPORAL_RESPONSIVE_AA_MASK, STENCIL_TEMPORAL_RESPONSIVE_AA_MASK>::GetRHI();
 
 				PermutationVector.Set<FTAAResponsiveDim>(true);
-				DrawPixelPassTemplate(
-					Context, PermutationVector, SrcSize, ViewRect,
-					InputHistory, Parameters, bUseDither,
-					DepthStencilState);
+				DrawTAAPixelShader(DepthStencilState);
 			}
 			else
 			{
-				DrawPixelPassTemplate(
-					Context, PermutationVector, SrcSize, ViewRect,
-					InputHistory, Parameters, bUseDither,
-					TStaticDepthStencilState<false, CF_Always>::GetRHI());
-			}
-
-			if (RenderTargetCount == 2)
-			{
-				Context.RHICmdList.EndUpdateMultiFrameResource(DestRenderTarget[1]->ShaderResourceTexture);
+				DrawTAAPixelShader(TStaticDepthStencilState<false, CF_Always>::GetRHI());
 			}
-		}
-		Context.RHICmdList.EndRenderPass();
-		Context.RHICmdList.CopyToResolveTarget(DestRenderTarget[0]->TargetableTexture, DestRenderTarget[0]->ShaderResourceTexture, FResolveParams());
-
-		if (RenderTargetCount == 2)
-		{
-			Context.RHICmdList.CopyToResolveTarget(DestRenderTarget[1]->TargetableTexture, DestRenderTarget[1]->ShaderResourceTexture, FResolveParams());
-		}
-
-		if (IsDOFTAAConfig(Parameters.Pass))
-		{
-			Context.RHICmdList.TransitionResource(EResourceTransitionAccess::EReadable, EResourceTransitionPipeline::EGfxToCompute, DestRenderTarget[0]->UAV);
-
-			if (RenderTargetCount == 2)
-			{
-				Context.RHICmdList.TransitionResource(EResourceTransitionAccess::EReadable, EResourceTransitionPipeline::EGfxToCompute, DestRenderTarget[1]->UAV);
-			}
-		}
-
-		// Inform MultiGPU systems that we've finished with this texture for this frame
-		Context.RHICmdList.EndUpdateMultiFrameResource(DestRenderTarget[0]->ShaderResourceTexture);
+		});
 	}
-
-	if (!Context.View.bViewStateIsReadOnly)
+	
+	if (!View.bViewStateIsReadOnly)
 	{
 		OutputHistory->SafeRelease();
-		OutputHistory->RT[0] = PassOutputs[0].PooledRenderTarget;
-		OutputHistory->ViewportRect = DestRect;
-		OutputHistory->ReferenceBufferSize = FSceneRenderTargets::Get(Context.RHICmdList).GetBufferSizeXY();
-		OutputHistory->SceneColorPreExposure = Context.View.PreExposure;
 
-		if (OutputExtent.X > 0)
-		{
-			OutputHistory->ReferenceBufferSize = OutputExtent;
-		}
-	}
+		GraphBuilder.QueueTextureExtraction(Outputs.SceneColor, &OutputHistory->RT[0]);
+		if (Outputs.SceneMetadata)
+			GraphBuilder.QueueTextureExtraction(Outputs.SceneMetadata, &OutputHistory->RT[1]);
 
-	// Changes the view rectangle of the scene color and reference buffer size when doing temporal upsample for the
-	// following passes to still work.
-	if (Parameters.Pass == ETAAPassConfig::MainUpsampling || Parameters.Pass == ETAAPassConfig::MainSuperSampling)
-	{
-		Context.SceneColorViewRect = DestRect;
-		Context.ReferenceBufferSize = OutputExtent;
+		OutputHistory->ViewportRect = DestRect;
+		OutputHistory->ReferenceBufferSize = OutputExtent * ResolutionDivisor;
+		OutputHistory->SceneColorPreExposure = View.PreExposure;
 	}
-}
-
-FPooledRenderTargetDesc FRCPassPostProcessTemporalAA::ComputeOutputDesc(EPassOutputId InPassOutputId) const
-{
-	FPooledRenderTargetDesc Ret;
 
-	switch (InPassOutputId)
-	{
-	case ePId_Output0: // main color output
-	case ePId_Output1:
-
-		Ret = GetInput(ePId_Input0)->GetOutput()->RenderTargetDesc;
-		Ret.Flags &= ~(TexCreate_FastVRAM | TexCreate_Transient);
-		Ret.Reset();
-		//regardless of input type, PF_FloatRGBA is required to properly accumulate between frames for a good result.
-		Ret.Format = PF_FloatRGBA;
-		Ret.DebugName = kTAAOutputNames[static_cast<int32>(Parameters.Pass)];
-		Ret.AutoWritable = false;
-		Ret.TargetableFlags &= ~(TexCreate_RenderTargetable | TexCreate_UAV);
-		Ret.TargetableFlags |= bIsComputePass ? TexCreate_UAV : TexCreate_RenderTargetable;
-
-		if (OutputExtent.X > 0)
-		{
-			check(OutputExtent.X % Parameters.ResolutionDivisor == 0);
-			check(OutputExtent.Y % Parameters.ResolutionDivisor == 0);
-			Ret.Extent = OutputExtent / Parameters.ResolutionDivisor;
-		}
-
-		// Need a UAV to resource transition from gfx to compute.
-		if (IsDOFTAAConfig(Parameters.Pass))
-		{
-			Ret.TargetableFlags |= TexCreate_UAV;
-		}
+	return Outputs;
+} // AddTemporalAAPass()
 
-		break;
 
-	case ePId_Output2: // downsampled color output
-
-		if (!bDownsamplePossible)
-		{
-			break;
-		}
-
-		check(bIsComputePass);
-
-		Ret = GetInput(ePId_Input0)->GetOutput()->RenderTargetDesc;
-		Ret.Flags &= ~(TexCreate_FastVRAM);
-		Ret.Reset();
+FRCPassPostProcessTemporalAA::FRCPassPostProcessTemporalAA(
+	const FPostprocessContext& Context,
+	const FTAAPassParameters& InParameters,
+	const FTemporalAAHistory& InInputHistory,
+	FTemporalAAHistory* OutOutputHistory)
+	: SavedParameters(InParameters)
+	, InputHistory(InInputHistory)
+	, OutputHistory(OutOutputHistory)
+{
+	check(SavedParameters.Validate());
 
-		if (Parameters.DownsampleOverrideFormat != PF_Unknown)
-		{
-			Ret.Format = Parameters.DownsampleOverrideFormat;
-		}
+	bIsComputePass = SavedParameters.bIsComputePass;
+	bPreferAsyncCompute = false;
+}
 
-		Ret.DebugName = TEXT("SceneColorHalfRes");
-		Ret.AutoWritable = false;
-		Ret.TargetableFlags &= ~TexCreate_RenderTargetable;
-		Ret.TargetableFlags |= TexCreate_UAV;
+void FRCPassPostProcessTemporalAA::Process(FRenderingCompositePassContext& Context)
+{
+	WaitForInputPassComputeFences(Context.RHICmdList);
 
-		if (OutputExtent.X > 0)
-		{
-			check(OutputExtent.X % Parameters.ResolutionDivisor == 0);
-			check(OutputExtent.Y % Parameters.ResolutionDivisor == 0);
-			Ret.Extent = OutputExtent / Parameters.ResolutionDivisor;
-		}
+	AsyncEndFence = FComputeFenceRHIRef();
+	
+	FRDGBuilder GraphBuilder(Context.RHICmdList);
 
-		Ret.Extent = FIntPoint::DivideAndRoundUp(Ret.Extent, 2);
-		Ret.Extent.X = FMath::Max(1, Ret.Extent.X);
-		Ret.Extent.Y = FMath::Max(1, Ret.Extent.Y);
+	FSceneRenderTargets& SceneContext = FSceneRenderTargets::Get(Context.RHICmdList);
+	FSceneViewFamilyBlackboard SceneBlackboard;
+	SetupSceneViewFamilyBlackboard(GraphBuilder, &SceneBlackboard);
+	
+	FTAAPassParameters Parameters = SavedParameters;
+	Parameters.SceneColorInput = CreateRDGTextureForInput(GraphBuilder, ePId_Input0, TEXT("SceneColor"), eFC_0000);
+	Parameters.SceneMetadataInput = CreateRDGTextureForInput(GraphBuilder, ePId_Input1, TEXT("SceneColor"), eFC_0000);
+
+	FTAAOutputs Outputs = Parameters.AddTemporalAAPass(
+		GraphBuilder,
+		SceneBlackboard, Context.View,
+		InputHistory, /* out */ OutputHistory);
+		
+	ExtractRDGTextureForOutput(GraphBuilder, ePId_Output0, Outputs.SceneColor);
+	if (Outputs.SceneMetadata)
+		ExtractRDGTextureForOutput(GraphBuilder, ePId_Output1, Outputs.SceneMetadata);
+	if (Outputs.DownsampledSceneColor)
+		ExtractRDGTextureForOutput(GraphBuilder, ePId_Output2, Outputs.DownsampledSceneColor);
 
-		break;
+	GraphBuilder.Execute();
 
-	default:
-		check(false);
+	// Changes the view rectangle of the scene color and reference buffer size when doing temporal upsample for the
+	// following passes to still work.
+	if (Parameters.Pass == ETAAPassConfig::MainUpsampling || Parameters.Pass == ETAAPassConfig::MainSuperSampling)
+	{
+		Context.SceneColorViewRect = Parameters.OutputViewRect;
+		Context.ReferenceBufferSize = Parameters.GetOutputExtent() * Parameters.ResolutionDivisor;
 	}
+}
 
-	return Ret;
+FPooledRenderTargetDesc FRCPassPostProcessTemporalAA::ComputeOutputDesc(EPassOutputId InPassOutputId) const
+{
+	// ExtractRDGTextureForOutput() is doing this work for us already.
+	return FPooledRenderTargetDesc();
 }
diff --git a/Engine/Source/Runtime/Renderer/Private/PostProcess/PostProcessTemporalAA.h b/Engine/Source/Runtime/Renderer/Private/PostProcess/PostProcessTemporalAA.h
--- a/Engine/Source/Runtime/Renderer/Private/PostProcess/PostProcessTemporalAA.h
+++ b/Engine/Source/Runtime/Renderer/Private/PostProcess/PostProcessTemporalAA.h
@@ -7,10 +7,18 @@
 #pragma once
 
 #include "CoreMinimal.h"
+#include "RenderGraph.h"
+
+// TODO: kill these includes once FRCPassPostProcessTemporalAA is gone.
 #include "RendererInterface.h"
 #include "RenderingCompositionGraph.h"
 
 
+class FViewInfo;
+class FSceneViewFamilyBlackboard;
+struct FTemporalAAHistory;
+
+
 /** Lists of TAA configurations. */
 enum class ETAAPassConfig
 {
@@ -43,40 +51,56 @@ static FORCEINLINE bool IsDOFTAAConfig(ETAAPassConfig Pass)
 }
 
 
+/** GPU Output of the TAA pass. */
+struct FTAAOutputs
+{
+	// Anti aliased scene color.
+	// Can have alpha channel, or CoC for DOF.
+	FRDGTexture* SceneColor = nullptr;
+
+	// Optional information that get anti aliased, such as separate CoC for DOF.
+	FRDGTexture* SceneMetadata = nullptr;
+
+	// Optional scene color output at half the resolution.
+	FRDGTexture* DownsampledSceneColor = nullptr;
+};
+
+
 /** Configuration of TAA. */
 struct FTAAPassParameters
 {
 	// TAA pass to run.
-	ETAAPassConfig Pass;
+	ETAAPassConfig Pass = ETAAPassConfig::Main;
 
 	// Whether to use the faster shader permutation.
-	bool bUseFast;
+	bool bUseFast = false;
 
 	// Whether to do compute or not.
-	bool bIsComputePass;
+	bool bIsComputePass = false;
 
 	// Whether downsampled (box filtered, half resolution) frame should be written out.
 	// Only used when bIsComputePass is true.
-	bool bDownsample;
-	EPixelFormat DownsampleOverrideFormat;
+	bool bDownsample = false;
+	EPixelFormat DownsampleOverrideFormat = PF_Unknown;
 
 	// Viewport rectangle of the input and output of TAA at ResolutionDivisor == 1.
 	FIntRect InputViewRect;
 	FIntRect OutputViewRect;
 
 	// Resolution divisor.
-	int32 ResolutionDivisor;
+	int32 ResolutionDivisor = 1;
+
+	// Anti aliased scene color.
+	// Can have alpha channel, or CoC for DOF.
+	FRDGTexture* SceneColorInput = nullptr;
+
+	// Optional information that get anti aliased, such as separate CoC for DOF.
+	FRDGTexture* SceneMetadataInput = nullptr;
 
 
 	FTAAPassParameters(const FViewInfo& View)
-		: Pass(ETAAPassConfig::Main)
-		, bUseFast(false)
-		, bIsComputePass(false)
-		, bDownsample(false)
-		, DownsampleOverrideFormat(PF_Unknown)
-		, InputViewRect(View.ViewRect)
+		: InputViewRect(View.ViewRect)
 		, OutputViewRect(View.ViewRect)
-		, ResolutionDivisor(1)
 	{ }
 
 
@@ -107,14 +131,25 @@ struct FTAAPassParameters
 		OutputViewRect.Max -= OutputViewRect.Min;
 		OutputViewRect.Min = FIntPoint::ZeroValue;
 	}
-};
+	
+	/** Returns the texture resolution that will be outputed. */
+	FIntPoint GetOutputExtent() const;
+
+	/** Validate the settings of TAA, to make sure there is no issue. */
+	bool Validate() const;
 
+	
+	/** Apply a temporal AA pass. */
+	FTAAOutputs AddTemporalAAPass(
+		FRDGBuilder& GraphBuilder,
+		const FSceneViewFamilyBlackboard& SceneBlackboard,
+		const FViewInfo& View,
+		const FTemporalAAHistory& InputHistory,
+		FTemporalAAHistory* OutputHistory) const;
+}; // struct FTAAPassParameters
 
-// ePId_Input0: Full Res Scene color (point)
-// ePId_Input2: Velocity (point)
-// ePId_Output0: Antialiased color
-// ePId_Output1: Downsampled antialiased color (only when FTAAConfig::bDownsample is true)
-// derives from TRenderingCompositePassBase<InputCount, OutputCount> 
+
+// DEPRECATED. Use FTAAPassParameters::AddTemporalAAPass() instead.
 class FRCPassPostProcessTemporalAA : public TRenderingCompositePassBase<3, 3>
 {
 public:
@@ -131,17 +166,11 @@ public:
 
 	virtual FComputeFenceRHIParamRef GetComputePassEndFence() const override { return AsyncEndFence; }
 
-	bool IsDownsamplePossible() const { return bDownsamplePossible; }
-
 private:
-	const FTAAPassParameters Parameters;
-
-	FIntPoint OutputExtent;
+	const FTAAPassParameters SavedParameters;
 
 	FComputeFenceRHIRef AsyncEndFence;
 
 	const FTemporalAAHistory& InputHistory;
 	FTemporalAAHistory* OutputHistory;
-
-	bool bDownsamplePossible = false;
 };
diff --git a/Engine/Source/Runtime/Renderer/Private/PostProcess/PostProcessing.cpp b/Engine/Source/Runtime/Renderer/Private/PostProcess/PostProcessing.cpp
--- a/Engine/Source/Runtime/Renderer/Private/PostProcess/PostProcessing.cpp
+++ b/Engine/Source/Runtime/Renderer/Private/PostProcess/PostProcessing.cpp
@@ -949,7 +949,7 @@ static void AddTemporalAA( FPostprocessContext& Context, FRenderingCompositeOutp
 	TemporalAAPass->SetInput( ePId_Input2, VelocityInput );
 	Context.FinalOutput = FRenderingCompositeOutputRef( TemporalAAPass );
 
-	if (OutSceneColorHalfRes && TemporalAAPass->IsDownsamplePossible())
+	if (OutSceneColorHalfRes && Parameters.bDownsample)
 	{
 		*OutSceneColorHalfRes = FRenderingCompositeOutputRef(TemporalAAPass, ePId_Output2);
 	}
