Clustered deferred shading implementation.
- Added clustered deferred pixel shader & supporting code (using render graph)
- GetSubsufaceProfileTexture_RT now returns non-const pointer to make it work with RegisterExternalTexture
- Added FSortedLightSetSceneInfo to store the sorted lights and ranges, moved light sorting to before light grid injection.
  - Moved light sorting to own function: FDeferredShadingSceneRenderer::GatherAndSortLights
  - Added Simple Lights to the sorted light array.
  - Lights supporting clustered deferred now have their own sub-range (which include rect lights without a source texture)
- Added ranges of light indices to the FForwardLightData such that the shader can skip any lights that are not supported by clustered deferred (but they are there for other operations)
- Added flag to say if the lights were injected in the light grid.
- Added 'DummyRectLightSourceTexture' to FForwardLightData to use as the fallback rect light source texture instead of the clearly invalid DeferredLightData
- Refactored light grid injection to use RenderGraph and more convenient parameter setting.
- Added light type to local light data stored for light grid (2 bits packed into uint used for shadow map channel mask / lighting channel mask)
- Fix for light complexity that was influenced by translucent rendering (in an incorrect way).
- Added "r.DebugLightDiscardProp" / GDebugLightDiscardProp that allows discarding a proportion of the visible lights, to check performance scaling.
- Reorganized tiled deferred shader to reduce register pressure.
- Light grid build:
  - Added view-space light data to speed up light grid build
  - Added slightly improved cone culling test and re-ordererd the tests for better performance.
- Removed dead code: declaration of 'FClusteredLightsSceneInfo'
- Added Preexposure fix (also fixed exposure for light culling visualization)

diff --git a/Engine/Shaders/Private/ClusteredDeferredShadingPixelShader.usf b/Engine/Shaders/Private/ClusteredDeferredShadingPixelShader.usf
new file mode 100644
--- /dev/null
+++ b/Engine/Shaders/Private/ClusteredDeferredShadingPixelShader.usf
@@ -0,0 +1,282 @@
+// Copyright 1998-2019 Epic Games, Inc. All Rights Reserved.
+
+/*=============================================================================
+	ClusteredDeferredShadingPixelShader - applies all local lights in the
+	LightGrid in a full-screen pass, not directional ones as they are done
+	in the traditional deferred passes.
+=============================================================================*/
+
+// Following the lead of Tiled Deferred
+#define SUPPORT_CONTACT_SHADOWS 0
+#define USE_IES_PROFILE 0
+#define NON_DIRECTIONAL_DIRECT_LIGHTING 0
+
+// Used in RectLight.ush, and requires a per-light texture which we cannot support at present.
+// If the textures are atlased or something like that in the future it should be ok to turn on.
+#define USE_SOURCE_TEXTURE 0
+
+#include "Common.ush"
+#include "DeferredShadingCommon.ush"  
+#include "BRDF.ush"
+#include "ShadingModels.ush"
+#include "LightGridCommon.ush"
+#include "DeferredLightingCommon.ush"
+
+
+// Causes the shader to run one pass over the light list per shading model, which cuts down peak register pressure.
+// Moves some overhead out of the inner loop (testing for shading model) which ought to be a good thing also.
+#define USE_PASS_PER_SHADING_MODEL 1
+
+// Extracts out the simple lights into their own loop that is executed first, independently of shading model.
+#define ENABLE_SEPARATE_SIMPLE_LIGHT_LOOP 0
+
+/**
+ * Adds local lighting using the light grid, does not apply directional lights, as they are done elsewhere.
+ * Does not support dynamic shadows, as these require the per-light shadow mask.
+ */
+float3 GetLightGridLocalLighting(const FCulledLightsGridData LightGridData, float3 WorldPosition, float3 CameraVector, FScreenSpaceData ScreenSpaceData, float2 ScreenUV, uint EyeIndex, float Dither, uint FirstNonSimpleLightIndex)
+{
+	float3 DirectLighting = 0;
+
+	// TODO: Not part of LightGrid ligth data, does it need to be added?
+	float SpecularScale = 1.0f;
+
+	// Limit max to ForwardLightData.NumLocalLights.
+	// This prevents GPU hangs when the PS tries to read from uninitialized NumCulledLightsGrid buffer
+	const uint NumLocalLights = min(LightGridData.NumLocalLights, GetNumLocalLights(EyeIndex));
+
+	LOOP
+	for (uint LocalLightListIndex = FirstNonSimpleLightIndex; LocalLightListIndex < NumLocalLights; LocalLightListIndex++)
+	{
+		const FLocalLightData LocalLight = GetLocalLightData(LightGridData.DataStartIndex + LocalLightListIndex, EyeIndex);
+
+		// The lights are sorted such that all that support clustered deferred are at the beginning, there might be others
+		// (e.g., lights with dynamic shadows) so we break out when the condition fails.
+		if (!LocalLight.bClusteredDeferredSupported)
+		{
+			break;
+		}
+
+		FDeferredLightData LightData = (FDeferredLightData)0;
+		LightData.Position = LocalLight.LightPositionAndInvRadius.xyz;
+		LightData.InvRadius = LocalLight.LightPositionAndInvRadius.w;
+
+		// extra-early out since we know light grid is sloppy and all lights in list are radial (have a range)
+		// appears useless
+		float invLightRadiusSq = LightData.InvRadius*LightData.InvRadius;
+		if (length2(WorldPosition - LightData.Position) * invLightRadiusSq > 1.0f)
+		{
+			continue;
+		}
+
+		LightData.Color = LocalLight.LightColorAndFalloffExponent.xyz;
+		//LightData.Color.xy = 0.0f;
+		LightData.FalloffExponent = LocalLight.LightColorAndFalloffExponent.w;
+		LightData.Direction = LocalLight.LightDirectionAndShadowMask.xyz;
+		LightData.SpotAngles = LocalLight.SpotAnglesAndSourceRadiusPacked.xy;
+		LightData.SourceRadius = LocalLight.SpotAnglesAndSourceRadiusPacked.z;
+		LightData.SourceLength = f16tof32(asuint(LocalLight.SpotAnglesAndSourceRadiusPacked.w));
+		LightData.Tangent = LocalLight.LightTangentAndSoftSourceRadius.xyz;
+		LightData.SoftSourceRadius = LocalLight.LightTangentAndSoftSourceRadius.w;
+		LightData.bInverseSquared = LightData.FalloffExponent == 0;
+		// Only directional is not radial, and this loop is only non-directional lights (LIGHT_SOURCE_SHAPE > 0 in trad) 
+		LightData.bRadialLight = true;
+		// Setting it to true follows how it is setup in standard deferred, where only the directional lights are not doing this (LIGHT_SOURCE_SHAPE > 0)
+		//LightData.bSpotLight = true;
+		//LightData.bSpotLight = LightData.SpotAngles.x > -2.0f;
+		LightData.SpecularScale = SpecularScale;
+
+		uint LightTypeAndPackedShadowMapChannelMask = asuint(LocalLight.LightDirectionAndShadowMask.w);
+
+		// bits [17:16] really
+		uint LightType = LightTypeAndPackedShadowMapChannelMask >> 16;
+
+		// TODO: not sure if this check is even needed, standard deferred always seems to set it to true?
+		LightData.bSpotLight = (LightType == LIGHT_TYPE_SPOT);
+
+		// NOTE: supporting the rect-light code path adds ~30% overhead on NVIDIA (RTX 2060) even when not used! 
+		// Presumed register pressure increase again, not profiled in detail.
+		// Rect lights are now disabled by the sorting, so the shader should not handle them, for now.
+		LightData.bRectLight = false;// (LightType == LIGHT_TYPE_RECT);
+
+		// bits [7:0]
+		LightData.ShadowedBits = (LightTypeAndPackedShadowMapChannelMask & 0xFF) != 0 ? 1 : 0;
+		// Static shadowing uses ShadowMapChannel, dynamic shadows are packed into light attenuation using PreviewShadowMapChannel
+		LightData.ShadowMapChannelMask = UnpackShadowMapChannelMask(LightTypeAndPackedShadowMapChannelMask);
+
+		// This matches tiled deferred, since there is no dynamic shadow on the lights used for clustered shading.
+		float4 LightAttenuation = 1.0f;
+
+#if 1
+		FRectTexture RectTexture = InitRectTexture(ForwardLightData.DummyRectLightSourceTexture);
+		// NOTE: uses AO data like tiled deferred (as opposed to forward path)
+		// TODO: Passing in 0,0 as SVPos, what is meaning of this? Only used if 'REFERENCE_QUALITY' is on.
+		DirectLighting += GetDynamicLighting(WorldPosition, CameraVector, ScreenSpaceData.GBuffer, ScreenSpaceData.AmbientOcclusion, ScreenSpaceData.GBuffer.ShadingModelID, LightData, LightAttenuation, Dither, uint2(0, 0), RectTexture).xyz;
+
+#else // 1
+		FSimpleDeferredLightData SimpleLightData = (FSimpleDeferredLightData)0;
+		SimpleLightData.Position = LightData.Position;
+		SimpleLightData.InvRadius = LightData.InvRadius;
+		SimpleLightData.Color = LightData.Color;
+		SimpleLightData.FalloffExponent = LightData.FalloffExponent;
+		SimpleLightData.bInverseSquared = LightData.bInverseSquared;
+		DirectLighting += GetSimpleDynamicLighting(WorldPosition, CameraVector, ScreenSpaceData.GBuffer.WorldNormal, ScreenSpaceData.AmbientOcclusion, ScreenSpaceData.GBuffer.DiffuseColor, ScreenSpaceData.GBuffer.SpecularColor, ScreenSpaceData.GBuffer.Roughness, SimpleLightData);
+#endif //
+
+	}
+
+	// For debugging
+	//DirectLighting = LightGridData.NumLocalLights / (float)ForwardLightData.MaxCulledLightsPerCell;
+	return DirectLighting;
+}
+
+#if USE_PASS_PER_SHADING_MODEL
+
+// NOTE: this is a macro since we want to be sure the 'ScreenSpaceData.GBuffer.ShadingModelID = ShadingModelId' is compile time constant
+//       The rest sort of follows.
+// NOTE 2: The screen-space data is loaded inside the if to try to minimize register pressure, but not clear if it is respected.
+#define GET_LIGHT_GRID_LOCAL_LIGHTING_SINGLE_SM(ShadingModelId, PixelShadingModelId, /*inout float4*/ CompositedLighting, ScreenUV, LightGridData, Dither, FirstNonSimpleLightIndex) \
+{	\
+	BRANCH	\
+	if (PixelShadingModelId == ShadingModelId)	\
+	{	\
+		float2 ScreenPosition = UVAndScreenPos.zw;	\
+		FScreenSpaceData ScreenSpaceData = GetScreenSpaceData(ScreenUV);	\
+		float SceneDepth = CalcSceneDepth(ScreenUV);	\
+		float3 WorldPosition = mul(float4(ScreenPosition * SceneDepth, SceneDepth, 1), View.ScreenToWorld).xyz;	\
+		float3 CameraVector = normalize(WorldPosition - View.WorldCameraOrigin);	\
+		ScreenSpaceData.GBuffer.ShadingModelID = ShadingModelId;	\
+		CompositedLighting += GetLightGridLocalLighting(LightGridData, WorldPosition, CameraVector, ScreenSpaceData, ScreenUV, 0, Dither, FirstNonSimpleLightIndex); \
+	}	\
+}
+
+#endif // USE_PASS_PER_SHADING_MODEL
+
+
+
+void ClusteredShadingPixelShader(
+	in noperspective float4 UVAndScreenPos : TEXCOORD0,
+	in float4 SvPosition : SV_Position,
+	out float4 OutColor : SV_Target0)
+{
+	float2 ScreenUV = UVAndScreenPos.xy;
+
+	uint PixelShadingModelID = GetScreenSpaceData(ScreenUV).GBuffer.ShadingModelID;
+
+	// TODO: This is how it was in TiledDeferred
+	//const float Dither = 0.5f;
+	// BUT: this is how it is in trad. deferred, assumed better:
+	const float Dither = InterleavedGradientNoise(SvPosition.xy, View.StateFrameIndexMod8);
+
+	// ?
+	const uint EyeIndex = 0;
+
+	float SceneDepth = CalcSceneDepth(ScreenUV);
+	float2 LocalPosition = SvPosition.xy - View.ViewRectMin.xy;
+	uint GridIndex = ComputeLightGridCellIndex(uint2(LocalPosition.x, LocalPosition.y), SceneDepth);
+	const FCulledLightsGridData CulledLightGridData = GetCulledLightsGrid(GridIndex, EyeIndex);
+#if 0
+	const LightGridData GridData = GetLightGridData(EyeIndex);
+	uint ZSlice = (uint)(max(0, log2(SceneDepth * GridData.LightGridZParams.x + GridData.LightGridZParams.y) * GridData.LightGridZParams.z));
+	ZSlice = min(ZSlice, (uint)(GridData.CulledGridSize.z - 1));
+	uint3 GridCoordinate = uint3(uint2(LocalPosition.x, LocalPosition.y) >> GridData.LightGridPixelSizeShift, ZSlice);
+
+	if (ZSlice % 2)
+	{
+		OutColor = float(ZSlice) / float(GridData.CulledGridSize.z - 1);
+	}
+	else
+	{
+		OutColor = 1.0f - float(ZSlice) / float(GridData.CulledGridSize.z - 1);
+	}
+
+	if (GridCoordinate.x % 2)
+	{
+		OutColor.xyz = 1.0f - OutColor.xyz;
+	}
+	if (GridCoordinate.y % 2)
+	{
+		OutColor.xyz = 1.0f - OutColor.xyz;
+	}
+	OutColor.w = 0.0f;
+	return;
+#endif // 0
+	// NOTE: this early out helps with the case where there are no lights in the grid cell.
+	if (CulledLightGridData.NumLocalLights == 0 || PixelShadingModelID == SHADINGMODELID_UNLIT)
+	{
+		OutColor = 0;
+		return;
+	}
+
+	float3 CompositedLighting = 0;
+
+#if ENABLE_SEPARATE_SIMPLE_LIGHT_LOOP
+
+	uint FirstNonSimpleLightIndex = 0;
+	{
+		float2 ScreenPosition = UVAndScreenPos.zw;
+		FScreenSpaceData ScreenSpaceData = GetScreenSpaceData(ScreenUV);
+		float3 WorldPosition = mul(float4(ScreenPosition * SceneDepth, SceneDepth, 1), View.ScreenToWorld).xyz;
+		float3 CameraVector = normalize(WorldPosition - View.WorldCameraOrigin);
+
+		// Limit max to ForwardLightData.NumLocalLights.
+		// This prevents GPU hangs when the PS tries to read from uninitialized NumCulledLightsGrid buffer
+		const uint NumLocalLights = min(CulledLightGridData.NumLocalLights, GetNumLocalLights(EyeIndex));
+
+		uint LocalLightListIndex = 0;
+
+		LOOP
+		for (; LocalLightListIndex < NumLocalLights; LocalLightListIndex++)
+		{
+			const FLocalLightData LocalLight = GetLocalLightData(CulledLightGridData.DataStartIndex + LocalLightListIndex, EyeIndex);
+			if (!LocalLight.bIsSimpleLight)
+			{
+				break;
+			}
+
+			FSimpleDeferredLightData SimpleLightData = (FSimpleDeferredLightData)0;
+			SimpleLightData.Position = LocalLight.LightPositionAndInvRadius.xyz;
+			SimpleLightData.InvRadius = LocalLight.LightPositionAndInvRadius.w;
+			SimpleLightData.Color = LocalLight.LightColorAndFalloffExponent.xyz;
+			SimpleLightData.FalloffExponent = LocalLight.LightColorAndFalloffExponent.w;
+			SimpleLightData.bInverseSquared = SimpleLightData.FalloffExponent == 0;
+
+			CompositedLighting += GetSimpleDynamicLighting(WorldPosition, CameraVector, ScreenSpaceData.GBuffer.WorldNormal, ScreenSpaceData.AmbientOcclusion, ScreenSpaceData.GBuffer.DiffuseColor, ScreenSpaceData.GBuffer.SpecularColor, ScreenSpaceData.GBuffer.Roughness, SimpleLightData);
+		}
+		FirstNonSimpleLightIndex = LocalLightListIndex;
+	}
+#else // !ENABLE_SEPARATE_SIMPLE_LIGHT_LOOP
+	uint FirstNonSimpleLightIndex = 0;
+
+#endif // ENABLE_SEPARATE_SIMPLE_LIGHT_LOOP
+#if USE_PASS_PER_SHADING_MODEL
+
+	GET_LIGHT_GRID_LOCAL_LIGHTING_SINGLE_SM(SHADINGMODELID_DEFAULT_LIT, PixelShadingModelID, CompositedLighting, ScreenUV, CulledLightGridData, Dither, FirstNonSimpleLightIndex);
+	GET_LIGHT_GRID_LOCAL_LIGHTING_SINGLE_SM(SHADINGMODELID_SUBSURFACE, PixelShadingModelID, CompositedLighting, ScreenUV, CulledLightGridData, Dither, FirstNonSimpleLightIndex);
+	GET_LIGHT_GRID_LOCAL_LIGHTING_SINGLE_SM(SHADINGMODELID_PREINTEGRATED_SKIN, PixelShadingModelID, CompositedLighting, ScreenUV, CulledLightGridData, Dither, FirstNonSimpleLightIndex);
+	GET_LIGHT_GRID_LOCAL_LIGHTING_SINGLE_SM(SHADINGMODELID_CLEAR_COAT, PixelShadingModelID, CompositedLighting, ScreenUV, CulledLightGridData, Dither, FirstNonSimpleLightIndex);
+	GET_LIGHT_GRID_LOCAL_LIGHTING_SINGLE_SM(SHADINGMODELID_SUBSURFACE_PROFILE, PixelShadingModelID, CompositedLighting, ScreenUV, CulledLightGridData, Dither, FirstNonSimpleLightIndex);
+	GET_LIGHT_GRID_LOCAL_LIGHTING_SINGLE_SM(SHADINGMODELID_TWOSIDED_FOLIAGE, PixelShadingModelID, CompositedLighting, ScreenUV, CulledLightGridData, Dither, FirstNonSimpleLightIndex);
+	GET_LIGHT_GRID_LOCAL_LIGHTING_SINGLE_SM(SHADINGMODELID_HAIR, PixelShadingModelID, CompositedLighting, ScreenUV, CulledLightGridData, Dither, FirstNonSimpleLightIndex);
+	GET_LIGHT_GRID_LOCAL_LIGHTING_SINGLE_SM(SHADINGMODELID_CLOTH, PixelShadingModelID, CompositedLighting, ScreenUV, CulledLightGridData, Dither, FirstNonSimpleLightIndex);
+	GET_LIGHT_GRID_LOCAL_LIGHTING_SINGLE_SM(SHADINGMODELID_EYE, PixelShadingModelID, CompositedLighting, ScreenUV, CulledLightGridData, Dither, FirstNonSimpleLightIndex);
+#else // !USE_PASS_PER_SHADING_MODEL
+
+	float2 ScreenPosition = UVAndScreenPos.zw;
+	FScreenSpaceData ScreenSpaceData = GetScreenSpaceData(ScreenUV);
+	float3 WorldPosition = mul(float4(ScreenPosition * SceneDepth, SceneDepth, 1), View.ScreenToWorld).xyz;
+	float3 CameraVector = normalize(WorldPosition - View.WorldCameraOrigin);
+
+	uint NumCulledEntryIndex = (ForwardLightData.NumGridCells + GridIndex) * NUM_CULLED_LIGHTS_GRID_STRIDE;
+
+	CompositedLighting += GetLightGridLocalLighting(CulledLightGridData, WorldPosition, CameraVector, ScreenSpaceData, ScreenUV, 0, Dither, FirstNonSimpleLightIndex);
+#endif // USE_PASS_PER_SHADING_MODEL
+	
+#if USE_PREEXPOSURE && !VISUALIZE_LIGHT_CULLING
+	CompositedLighting *= View.PreExposure;
+#endif	 // USE_PREEXPOSURE && !VISUALIZE_LIGHT_CULLING
+	OutColor.rgb = CompositedLighting;
+
+	OutColor.w = 0.0f;
+}
+
diff --git a/Engine/Shaders/Private/ForwardLightingCommon.ush b/Engine/Shaders/Private/ForwardLightingCommon.ush
--- a/Engine/Shaders/Private/ForwardLightingCommon.ush
+++ b/Engine/Shaders/Private/ForwardLightingCommon.ush
@@ -1,7 +1,7 @@
 // Copyright 1998-2019 Epic Games, Inc. All Rights Reserved.
 
 /*=============================================================================
-	ForwardLightingCommon.usf
+	ForwardLightingCommon.ush
 =============================================================================*/
 
 #define NON_DIRECTIONAL_DIRECT_LIGHTING (TRANSLUCENCY_LIGHTING_VOLUMETRIC_NONDIRECTIONAL || TRANSLUCENCY_LIGHTING_VOLUMETRIC_PERVERTEX_NONDIRECTIONAL)
@@ -80,11 +80,6 @@ float4 GetForwardDynamicShadowFactors(float2 ScreenUV)
 	return Square(Value);
 }
 
-float4 UnpackShadowMapChannelMask(uint ShadowMapChannelMaskPacked)
-{
-	return float4((ShadowMapChannelMaskPacked & 1) ? 1.0f : 0.0f, (ShadowMapChannelMaskPacked & 2) ? 1.0f : 0.0f, (ShadowMapChannelMaskPacked & 4) ? 1.0f : 0.0f, (ShadowMapChannelMaskPacked & 8) ? 1.0f : 0.0f);
-}
-
 float GetIndirectOcclusion(float2 ScreenUV, FGBufferData GBufferData)
 {
 	float IndirectOcclusion;
@@ -120,7 +115,7 @@ float3 GetForwardDirectLighting(uint GridIndex, float3 WorldPosition, float3 Cam
 	uint LightingChannelMask = GetPrimitiveData(PrimitiveId).LightingChannelMask;
 
 	const FDirectionalLightData DirectionalLightData = GetDirectionalLightData(EyeIndex);
-	FRectTexture RectTexture = InitRectTexture(DeferredLightUniforms.SourceTexture);
+	FRectTexture RectTexture = InitRectTexture(ForwardLightData.DummyRectLightSourceTexture);
 
 	BRANCH
 	if (DirectionalLightData.HasDirectionalLight)
@@ -196,17 +191,23 @@ float3 GetForwardDirectLighting(uint GridIndex, float3 WorldPosition, float3 Cam
 		LightData.bSpotLight = LightData.SpotAngles.x > -2.0f;
 		LightData.SpecularScale = SpecularScale;
 
-		uint PackedShadowMapChannelMask = asuint(LocalLight.LightDirectionAndShadowMask.w);
-		LightData.ShadowedBits = (PackedShadowMapChannelMask & 0xFF) != 0 ? 1 : 0;
+		// LightType=bits[17:16], LightingChannelMask=[15:8], DynShadowMask=[7:4] ShadowMapChannelMask=[3:0]
+		uint LightTypeAndPackedShadowMapChannelMask = asuint(LocalLight.LightDirectionAndShadowMask.w);
+
+		// TODO: currently LightType is ignored and the code always sets 'LightData.bRectLight' to false (initialization)
+		//       but the light grid adds rect lights as well so this probably leads to artifacts.
+		uint LightType = LightTypeAndPackedShadowMapChannelMask >> 16;
+
+		LightData.ShadowedBits = (LightTypeAndPackedShadowMapChannelMask & 0xFF) != 0 ? 1 : 0;
 		// Static shadowing uses ShadowMapChannel, dynamic shadows are packed into light attenuation using PreviewShadowMapChannel
-		LightData.ShadowMapChannelMask = UnpackShadowMapChannelMask(PackedShadowMapChannelMask);
-		float4 PreviewShadowMapChannelMask = UnpackShadowMapChannelMask(PackedShadowMapChannelMask >> 4);
+		LightData.ShadowMapChannelMask = UnpackShadowMapChannelMask(LightTypeAndPackedShadowMapChannelMask);
+		float4 PreviewShadowMapChannelMask = UnpackShadowMapChannelMask(LightTypeAndPackedShadowMapChannelMask >> 4);
 		float DynamicShadowing = dot(PreviewShadowMapChannelMask, DynamicShadowFactors);
 		float4 LightAttenuation = float4(1, 1, DynamicShadowing.x, DynamicShadowing.x);
 		float3 NewLighting = GetDynamicLighting(WorldPosition, -CameraVector, GBufferData, 1, GBufferData.ShadingModelID, LightData, LightAttenuation, Dither, uint2(0,0), RectTexture).xyz;
 
 		FLATTEN
-		if ((PackedShadowMapChannelMask >> 8) & LightingChannelMask)
+		if ((LightTypeAndPackedShadowMapChannelMask >> 8) & LightingChannelMask)
 		{
 			DirectLighting += NewLighting;
 		}
diff --git a/Engine/Shaders/Private/LightGridCommon.ush b/Engine/Shaders/Private/LightGridCommon.ush
--- a/Engine/Shaders/Private/LightGridCommon.ush
+++ b/Engine/Shaders/Private/LightGridCommon.ush
@@ -1,7 +1,16 @@
+// Copyright 1998-2019 Epic Games, Inc. All Rights Reserved.
+
 /*=============================================================================
-	LightGridCommon.usf
+	LightGridCommon.ush
 =============================================================================*/
 
+// TODO: This is clearly not the right place for these defines! Ideally ought to be exposed from engine somehow...
+#define LIGHT_TYPE_DIRECTIONAL		0 
+#define LIGHT_TYPE_POINT			1 
+#define LIGHT_TYPE_SPOT				2 
+#define LIGHT_TYPE_RECT				3 
+#define LIGHT_TYPE_MAX				4 
+
 #if INSTANCED_STEREO
 	#if MATERIALBLENDING_ANY_TRANSLUCENT
 		#define ForwardLightDataISR TranslucentBasePass.Shared.ForwardISR
@@ -154,6 +163,10 @@ struct FLocalLightData
 	float4 SpotAnglesAndSourceRadiusPacked;
     float4 LightDirectionAndShadowMask;
     float4 LightTangentAndSoftSourceRadius;
+	/** Flag is true if clustered deferred is supported for this light. They are always first / together.*/
+	bool bClusteredDeferredSupported; 
+	/** Flag is true if it is a simple light. They are always first / together.*/
+	bool bIsSimpleLight;
 };
 
 FLocalLightData GetLocalLightData(uint GridIndex, uint EyeIndex)
@@ -165,21 +178,23 @@ FLocalLightData GetLocalLightData(uint GridIndex, uint EyeIndex)
 	if (EyeIndex == 0)
 	{
 #endif
-
 		uint LocalLightIndex = ForwardLightData.CulledLightDataGrid[GridIndex];
 		uint LocalLightBaseIndex = LocalLightIndex * LOCAL_LIGHT_DATA_STRIDE;
+		Result.bClusteredDeferredSupported = LocalLightIndex < ForwardLightData.ClusteredDeferredSupportedEndIndex;
+		Result.bIsSimpleLight = LocalLightIndex < ForwardLightData.SimpleLightsEndIndex;
 		Result.LightPositionAndInvRadius = ForwardLightData.ForwardLocalLightBuffer[LocalLightBaseIndex + 0];
 		Result.LightColorAndFalloffExponent = ForwardLightData.ForwardLocalLightBuffer[LocalLightBaseIndex + 1];
 		Result.LightDirectionAndShadowMask = ForwardLightData.ForwardLocalLightBuffer[LocalLightBaseIndex + 2];
 		Result.SpotAnglesAndSourceRadiusPacked = ForwardLightData.ForwardLocalLightBuffer[LocalLightBaseIndex + 3];
 		Result.LightTangentAndSoftSourceRadius = ForwardLightData.ForwardLocalLightBuffer[LocalLightBaseIndex + 4];
-
 #if INSTANCED_STEREO
 	}
 	else
 	{
 		uint LocalLightIndex = ForwardLightDataISR.CulledLightDataGrid[GridIndex];
 		uint LocalLightBaseIndex = LocalLightIndex * LOCAL_LIGHT_DATA_STRIDE;
+		Result.bClusteredDeferredSupported = LocalLightIndex < ForwardLightDataISR.ClusteredDeferredSupportedEndIndex;
+		Result.bIsSimpleLight = LocalLightIndex < ForwardLightDataISR.SimpleLightsEndIndex;
 		Result.LightPositionAndInvRadius = ForwardLightDataISR.ForwardLocalLightBuffer[LocalLightBaseIndex + 0];
 		Result.LightColorAndFalloffExponent = ForwardLightDataISR.ForwardLocalLightBuffer[LocalLightBaseIndex + 1];
 		Result.LightDirectionAndShadowMask = ForwardLightDataISR.ForwardLocalLightBuffer[LocalLightBaseIndex + 2];
@@ -190,3 +205,11 @@ FLocalLightData GetLocalLightData(uint GridIndex, uint EyeIndex)
 
 	return Result;
 }
+
+/**
+ * Helper to unpack the shadow mask part of the FLocalLightData::LightDirectionAndShadowMask.w
+ */
+float4 UnpackShadowMapChannelMask(uint ShadowMapChannelMaskPacked)
+{
+	return float4((ShadowMapChannelMaskPacked & 1) ? 1.0f : 0.0f, (ShadowMapChannelMaskPacked & 2) ? 1.0f : 0.0f, (ShadowMapChannelMaskPacked & 4) ? 1.0f : 0.0f, (ShadowMapChannelMaskPacked & 8) ? 1.0f : 0.0f);
+}
diff --git a/Engine/Shaders/Private/LightGridInjection.usf b/Engine/Shaders/Private/LightGridInjection.usf
--- a/Engine/Shaders/Private/LightGridInjection.usf
+++ b/Engine/Shaders/Private/LightGridInjection.usf
@@ -12,7 +12,14 @@ RWBuffer<uint> RWNextCulledLightLink;
 RWBuffer<uint> RWStartOffsetGrid;
 RWBuffer<uint> RWCulledLightLinks;
 
-float ComputeDepthFromZSlice(uint ZSlice)
+#if ENABLE_LIGHT_CULLING_VIEW_SPACE_BUILD_DATA
+Buffer<float4> LightViewSpacePositionAndRadius;
+Buffer<float4> LightViewSpaceDirAndPreprocAngle;
+#endif // ENABLE_LIGHT_CULLING_VIEW_SPACE_BUILD_DATA
+
+#define REFINE_SPOTLIGHT_BOUNDS 1
+
+float ComputeCellNearViewDepthFromZSlice(uint ZSlice)
 {
 	float SliceDepth = (exp2(ZSlice / ForwardLightData.LightGridZParams.z) - ForwardLightData.LightGridZParams.y) / ForwardLightData.LightGridZParams.x;
 
@@ -46,8 +53,8 @@ void ComputeCellViewAABB(uint3 GridCoordinate, out float3 ViewTileMin, out float
 	float2 UnitPlaneTileMin = GridCoordinate.xy * TileSize + UnitPlaneMin;
 	float2 UnitPlaneTileMax = (GridCoordinate.xy + 1) * TileSize + UnitPlaneMin;
 
-	float MinTileZ = ComputeDepthFromZSlice(GridCoordinate.z);
-	float MaxTileZ = ComputeDepthFromZSlice(GridCoordinate.z + 1);
+	float MinTileZ = ComputeCellNearViewDepthFromZSlice(GridCoordinate.z);
+	float MaxTileZ = ComputeCellNearViewDepthFromZSlice(GridCoordinate.z + 1);
 
 	float MinTileDeviceZ = ConvertToDeviceZ(MinTileZ);
 	float4 MinDepthCorner0 = mul(float4(UnitPlaneTileMin.x, UnitPlaneTileMin.y, MinTileDeviceZ, 1), View.ClipToView);
@@ -104,6 +111,37 @@ bool IntersectConeWithSphere(float3 ConeVertex, float3 ConeAxis, float ConeRadiu
 	return !(bSphereTooFarFromCone || bSpherePastConeEnd || bSphereBehindVertex);
 }
 
+
+
+bool AabbOutsidePlane(float3 center, float3 extents, float4 plane)
+{
+	float dist = dot(float4(center, 1.0), plane);
+	float radius = dot(extents, abs(plane.xyz));
+
+	return dist > radius;
+}
+
+
+/**
+ * Approximate cone / aabb test that creates a single separating plane that lies in the cone on the side facing the centre of the Aabb
+ * Returns false if the Aabb is outside (entirely on the positive side) of this plane.
+ * Returns true otherwise, only works for 'acute angled' cones, where the angle is < 90 degrees.
+ * Is approximate, in that it can yield false negatives, i.e., that an Aabb may be actually outside, but the test still returns false.
+ * Since the intended use is to cull light cones, this is acceptable whereas false positives would cause glitches.
+ */
+bool IsAabbOutsideInfiniteAcuteConeApprox(float3 ConeVertex, float3 ConeAxis, float TanConeAngle, float3 AabbCentre, float3 AabbExt)
+{
+	// 1. find plane (well, base) in which normal lies, and which is perpendicular to axis and centre of aabb.
+	float3 D = AabbCentre - ConeVertex;
+
+	// perpendicular to cone axis in plane of cone axis and aabb centre.
+	float3 M = -normalize(cross(cross(D, ConeAxis), ConeAxis));
+	float3 N = -TanConeAngle * ConeAxis + M;
+	float4 Plane = float4(N, 0.0);
+
+	return AabbOutsidePlane(D, AabbExt, Plane);
+}
+
 [numthreads(THREADGROUP_SIZE, THREADGROUP_SIZE, THREADGROUP_SIZE)]
 void LightGridInjectionCS(
 	uint3 GroupId : SV_GroupID,
@@ -134,43 +172,41 @@ void LightGridInjectionCS(
 		for (uint LocalLightIndex = 0; LocalLightIndex < ForwardLightData.NumLocalLights; LocalLightIndex++)
 		{
 			uint LocalLightBaseIndex = LocalLightIndex * LOCAL_LIGHT_DATA_STRIDE;
+
+#if ENABLE_LIGHT_CULLING_VIEW_SPACE_BUILD_DATA
+			float4 LightPositionAndRadius = LightViewSpacePositionAndRadius[LocalLightIndex];
+			float3 ViewSpaceLightPosition = LightPositionAndRadius.xyz;
+			float LightRadius = LightPositionAndRadius.w;
+			float4 LightPositionAndInvRadius = ForwardLightData.ForwardLocalLightBuffer[LocalLightBaseIndex + 0];
+#else // !ENABLE_LIGHT_CULLING_VIEW_SPACE_BUILD_DATA
 			float4 LightPositionAndInvRadius = ForwardLightData.ForwardLocalLightBuffer[LocalLightBaseIndex + 0];
 			float LightRadius = 1.0f / LightPositionAndInvRadius.w;
-			bool bPassSpotlightTest = true;
 
-			#define REFINE_SPOTLIGHT_BOUNDS 1
-			#if REFINE_SPOTLIGHT_BOUNDS
-				float CosOuterCone = ForwardLightData.ForwardLocalLightBuffer[LocalLightBaseIndex + 3].x;
+			float3 ViewSpaceLightPosition = mul(float4(LightPositionAndInvRadius.xyz + View.PreViewTranslation.xyz, 1), View.TranslatedWorldToView).xyz;
+#endif // ENABLE_LIGHT_CULLING_VIEW_SPACE_BUILD_DATA
 
-				if (CosOuterCone > -2.0f)
+			float BoxDistanceSq = ComputeSquaredDistanceFromBoxToPoint(ViewTileCenter, ViewTileExtent, ViewSpaceLightPosition);
+
+			if (BoxDistanceSq < LightRadius * LightRadius)
+			{
+#if REFINE_SPOTLIGHT_BOUNDS
+				bool bPassSpotlightTest = true;
 				{
-					float2 CosSinAngle = float2(CosOuterCone, sqrt(1 - CosOuterCone * CosOuterCone));
-					float3 LightDirection = ForwardLightData.ForwardLocalLightBuffer[LocalLightBaseIndex + 2].xyz;
-
-					bPassSpotlightTest = IntersectConeWithSphere(LightPositionAndInvRadius.xyz, LightDirection, LightRadius, CosSinAngle, WorldTileBoundingSphere);
-					
-					// minimal bounding sphere encompassing given cone, based on: https://bartwronski.com/2017/04/13/cull-that-cone/
-					const float COS_PI_OVER_4 = 0.707107f; // Cos(Pi/4);
-					if (CosSinAngle.x < COS_PI_OVER_4)
-					{
-						LightPositionAndInvRadius.xyz -= LightDirection * LightRadius * CosSinAngle.x;
-						LightRadius = LightRadius * CosSinAngle.y;
-					}
-					else
+					float4 ViewSpaceDirAndPreprocAngle = LightViewSpaceDirAndPreprocAngle[LocalLightIndex];
+					float TanConeAngle = ViewSpaceDirAndPreprocAngle.w;
+
+					// Set to 0 for non-acute cones, or non-spot lights.
+					if (TanConeAngle > 0.0f)
 					{
-						LightRadius = LightRadius / (2.0f * CosSinAngle.x);
-						LightPositionAndInvRadius.xyz -= LightDirection * LightRadius;
+						float3 ViewSpaceLightDirection = -ViewSpaceDirAndPreprocAngle.xyz;
+						bPassSpotlightTest = !IsAabbOutsideInfiniteAcuteConeApprox(ViewSpaceLightPosition, ViewSpaceLightDirection, TanConeAngle, ViewTileCenter, ViewTileExtent);
 					}
 				}
-			#endif
-
-			float3 ViewSpaceLightPosition = mul(float4(LightPositionAndInvRadius.xyz + View.PreViewTranslation.xyz, 1), View.TranslatedWorldToView).xyz;
-
-			float BoxDistanceSq = ComputeSquaredDistanceFromBoxToPoint(ViewTileCenter, ViewTileExtent, ViewSpaceLightPosition);
+				if (bPassSpotlightTest)
+#endif // REFINE_SPOTLIGHT_BOUNDS
+				{
 
-			if (BoxDistanceSq < LightRadius * LightRadius && bPassSpotlightTest)
-			{
-				#if USE_LINKED_CULL_LIST
+#if USE_LINKED_CULL_LIST
 					uint NextLink;
 					InterlockedAdd(RWNextCulledLightLink[0], 1U, NextLink);
 
@@ -182,7 +218,7 @@ void LightGridInjectionCS(
 						RWCulledLightLinks[NextLink * LIGHT_LINK_STRIDE + 1] = PreviousLink;
 					}
 
-				#else
+#else
 					uint CulledLightIndex;
 					InterlockedAdd(RWNumCulledLightsGrid[GridIndex * NUM_CULLED_LIGHTS_GRID_STRIDE + 0], 1U, CulledLightIndex);
 					RWNumCulledLightsGrid[GridIndex * NUM_CULLED_LIGHTS_GRID_STRIDE + 1] = GridIndex * ForwardLightData.MaxCulledLightsPerCell;
@@ -191,7 +227,8 @@ void LightGridInjectionCS(
 					{
 						RWCulledLightDataGrid[GridIndex * ForwardLightData.MaxCulledLightsPerCell + CulledLightIndex] = LocalLightIndex;
 					}
-				#endif
+#endif
+				}
 			}
 		}
 		
diff --git a/Engine/Shaders/Private/TiledDeferredLightShaders.usf b/Engine/Shaders/Private/TiledDeferredLightShaders.usf
--- a/Engine/Shaders/Private/TiledDeferredLightShaders.usf
+++ b/Engine/Shaders/Private/TiledDeferredLightShaders.usf
@@ -5,10 +5,20 @@
 =============================================================================*/
 
 #define SUPPORT_CONTACT_SHADOWS 0
+#define USE_SOURCE_TEXTURE 0
 
 #include "Common.ush"
 #include "DeferredLightingCommon.ush"
 
+/**
+ * Hoist shading model selection out to the main function, then repeat light loop with fixed light model and let early outs deal with it.
+ * This should let the compiler defer loading values from the GBuffer until they are needed and thus reduce the catastrophic 
+ * resister pressure.
+ * If the compiler refuses, then move the GBuffer data load into the per-shading model branches also.
+ * Also use an outermost early out check to skip the whole pass.
+ */
+#define USE_PASS_PER_SHADING_MODEL 1
+
 #define THREADGROUP_TOTALSIZE (THREADGROUP_SIZEX * THREADGROUP_SIZEY)
 
 /** View rect min in xy, max in zw. */
@@ -46,6 +56,72 @@ groupshared uint TileSimpleLightIndices[MAX_LIGHTS];
 Texture2D<float4> InTexture;
 RWTexture2D<float4> RWOutTexture;
 
+void ExecuteShadingLoops(inout float4 CompositedLighting, FScreenSpaceData ScreenSpaceData, uint NumLightsAffectingTile, uint NumSimpleLightsAffectingTile, float3 CameraVector, float3 WorldPosition)
+{
+	LOOP
+	for (uint TileLightIndex = 0; TileLightIndex < NumLightsAffectingTile; TileLightIndex++) 
+	{
+		uint LightIndex = TileLightIndices[TileLightIndex];
+
+		FDeferredLightData LightData = (FDeferredLightData)0;
+		LightData.Position = TiledDeferred.LightPositionAndInvRadius[LightIndex].xyz;
+		LightData.InvRadius = TiledDeferred.LightPositionAndInvRadius[LightIndex].w;
+		LightData.Color = TiledDeferred.LightColorAndFalloffExponent[LightIndex].xyz;
+		LightData.FalloffExponent = TiledDeferred.LightColorAndFalloffExponent[LightIndex].w;
+		{
+			float4 Value = TiledDeferred2.LightDirectionAndSpotlightMaskAndSpecularScale[LightIndex];
+			LightData.Direction = Value.xyz;
+			LightData.bSpotLight = Value.w > 0;
+			LightData.SpecularScale = abs(Value.w);
+		}
+		LightData.SpotAngles = TiledDeferred2.SpotAnglesAndSourceRadiusAndSimpleLighting[LightIndex].xy;
+		LightData.SourceRadius = TiledDeferred2.SpotAnglesAndSourceRadiusAndSimpleLighting[LightIndex].z;
+		// don't support lights with non-0 length in tiled-deferred path, so make sure it stays zero-initialized
+		//LightData.SourceLength = TiledDeferred2.SpotAnglesAndSourceRadiusAndSimpleLighting[LightIndex].w;
+		LightData.ShadowMapChannelMask = TiledDeferred2.ShadowMapChannelMask[LightIndex];
+		LightData.bInverseSquared = LightData.FalloffExponent == 0;
+		// Only radial lights supported with tiled deferred
+		LightData.bRadialLight = true;
+		// The only type of shadowing supported for lights using tiled is static shadowing, so the light should only compute shadowing if it has static shadowing
+		LightData.ShadowedBits = dot(LightData.ShadowMapChannelMask, float4(1, 1, 1, 1));
+	
+		// Lights requiring light attenuation are not supported tiled for now
+		// WARNING: DeferredLightUniforms.SourceTexture is not set here, but compiles out because of bRectLight being false and USE_SOURCE_TEXTURE == 0.
+		FRectTexture RectTexture = InitRectTexture(DeferredLightUniforms.SourceTexture);
+		CompositedLighting += GetDynamicLighting(WorldPosition, CameraVector, ScreenSpaceData.GBuffer, ScreenSpaceData.AmbientOcclusion, ScreenSpaceData.GBuffer.ShadingModelID, LightData, float4(1, 1, 1, 1), 0.5, uint2(0,0), RectTexture);
+	}
+
+	// The simple shading model does not depend on ShadingModelID, so use it anytime the material is lit
+	LOOP
+	for (uint TileLightIndex2 = 0; TileLightIndex2 < NumSimpleLightsAffectingTile; TileLightIndex2++) 
+	{
+		uint LightIndex = TileSimpleLightIndices[TileLightIndex2];
+
+		FSimpleDeferredLightData LightData = (FSimpleDeferredLightData)0;
+		LightData.Position = TiledDeferred.LightPositionAndInvRadius[LightIndex].xyz;
+		LightData.InvRadius = TiledDeferred.LightPositionAndInvRadius[LightIndex].w;
+		LightData.Color = TiledDeferred.LightColorAndFalloffExponent[LightIndex].xyz;
+		LightData.FalloffExponent = TiledDeferred.LightColorAndFalloffExponent[LightIndex].w;
+		LightData.bInverseSquared = LightData.FalloffExponent == 0;
+						
+		// todo: doesn't support ScreenSpaceSubsurfaceScattering yet (not using alpha)
+		CompositedLighting.rgb += GetSimpleDynamicLighting(WorldPosition, CameraVector, ScreenSpaceData.GBuffer.WorldNormal, ScreenSpaceData.AmbientOcclusion, ScreenSpaceData.GBuffer.DiffuseColor, ScreenSpaceData.GBuffer.SpecularColor, ScreenSpaceData.GBuffer.Roughness, LightData);
+	}
+}
+
+
+#define EXECUTE_SHADING_LOOPS_SINGLE_SM(ShadingModelId, PixelShadingModelId, /*inout float4*/ CompositedLighting, PixelPos, NumLightsAffectingTile, NumSimpleLightsAffectingTile, CameraVector, WorldPosition) \
+{	\
+	BRANCH	\
+	if (PixelShadingModelId == ShadingModelId)	\
+	{	\
+		FScreenSpaceData ScreenSpaceData = GetScreenSpaceDataUint(PixelPos); \
+		ScreenSpaceData.GBuffer.ShadingModelID = ShadingModelId;	\
+		ExecuteShadingLoops(CompositedLighting, ScreenSpaceData, NumLightsAffectingTile, NumSimpleLightsAffectingTile, CameraVector, WorldPosition);	\
+	}	\
+}
+
+
 [numthreads(THREADGROUP_SIZEX, THREADGROUP_SIZEY, 1)]
 void TiledDeferredLightingMain(
 	uint3 GroupId : SV_GroupID,
@@ -172,9 +248,6 @@ void TiledDeferredLightingMain(
 	float2 ClipPosition = ScreenPosition * (View.ViewToClip[3][3] < 1.0f ? SceneDepth : 1.0f);
 	float3 WorldPosition = mul(float4(ClipPosition, SceneDepth, 1), View.ScreenToWorld).xyz;
 
-	// RGB accumulated RGB HDR color, A: luminance for screenspace subsurface scattering
-	float4 CompositedLighting = 0;
-
 #define DO_CULLING_AND_SHADING 1
 #if DO_CULLING_AND_SHADING
 
@@ -239,68 +312,47 @@ void TiledDeferredLightingMain(
 	uint NumLightsAffectingTile = TileNumLights;
 	uint NumSimpleLightsAffectingTile = TileNumSimpleLights;
 	
+
+#if USE_PASS_PER_SHADING_MODEL
+	// Now we could trick the compiler to put this in a branch by setting an always 1 int variable and decorate with BRANCH
+	// this might force it to not load the whole thing, but you never know with shader compilers.
+	uint PixelShadingModelID = GetScreenSpaceDataUint(PixelPos).GBuffer.ShadingModelID;
+#else // !USE_PASS_PER_SHADING_MODEL
 	// Lookup GBuffer properties once per pixel 
 	FScreenSpaceData ScreenSpaceData = GetScreenSpaceDataUint(PixelPos);
-	FGBufferData InGBufferData = ScreenSpaceData.GBuffer; 
-	float3 CameraVector = normalize(WorldPosition - View.WorldCameraOrigin);
-		
-	BRANCH
-	if (InGBufferData.ShadingModelID != SHADINGMODELID_UNLIT)
-	{
-        FRectTexture RectTexture = InitRectTexture(DeferredLightUniforms.SourceTexture);
+	uint PixelShadingModelID = ScreenSpaceData.GBuffer.ShadingModelID;
 
-		LOOP
-		for (uint TileLightIndex = 0; TileLightIndex < NumLightsAffectingTile; TileLightIndex++) 
-		{
-			uint LightIndex = TileLightIndices[TileLightIndex];
-
-			FDeferredLightData LightData = (FDeferredLightData)0;
-			LightData.Position = TiledDeferred.LightPositionAndInvRadius[LightIndex].xyz;
-			LightData.InvRadius = TiledDeferred.LightPositionAndInvRadius[LightIndex].w;
-			LightData.Color = TiledDeferred.LightColorAndFalloffExponent[LightIndex].xyz;
-			LightData.FalloffExponent = TiledDeferred.LightColorAndFalloffExponent[LightIndex].w;
-			{
-				float4 Value = TiledDeferred2.LightDirectionAndSpotlightMaskAndSpecularScale[LightIndex];
-				LightData.Direction = Value.xyz;
-				LightData.bSpotLight = Value.w > 0;
-				LightData.SpecularScale = abs(Value.w);
-			}
-			LightData.SpotAngles = TiledDeferred2.SpotAnglesAndSourceRadiusAndSimpleLighting[LightIndex].xy;
-			LightData.SourceRadius = TiledDeferred2.SpotAnglesAndSourceRadiusAndSimpleLighting[LightIndex].z;
-			// don't support lights with non-0 length in tiled-deferred path, so make sure it stays zero-initialized
-			//LightData.SourceLength = TiledDeferred2.SpotAnglesAndSourceRadiusAndSimpleLighting[LightIndex].w;
-			LightData.ShadowMapChannelMask = TiledDeferred2.ShadowMapChannelMask[LightIndex];
-			LightData.bInverseSquared = LightData.FalloffExponent == 0;
-			// Only radial lights supported with tiled deferred
-			LightData.bRadialLight = true;
-			// The only type of shadowing supported for lights using tiled is static shadowing, so the light should only compute shadowing if it has static shadowing
-			LightData.ShadowedBits = dot(LightData.ShadowMapChannelMask, float4(1, 1, 1, 1));
+#endif // USE_PASS_PER_SHADING_MODEL
 	
-			// Lights requiring light attenuation are not supported tiled for now
-			CompositedLighting += GetDynamicLighting(WorldPosition, CameraVector, ScreenSpaceData.GBuffer, ScreenSpaceData.AmbientOcclusion, InGBufferData.ShadingModelID, LightData, float4(1, 1, 1, 1), 0.5, uint2(0,0), RectTexture);
-		}
 
-		// The simple shading model does not depend on ShadingModelID, so use it anytime the material is lit
-		LOOP
-		for (uint TileLightIndex = 0; TileLightIndex < NumSimpleLightsAffectingTile; TileLightIndex++) 
-		{
-			uint LightIndex = TileSimpleLightIndices[TileLightIndex];
-
-			FSimpleDeferredLightData LightData = (FSimpleDeferredLightData)0;
-			LightData.Position = TiledDeferred.LightPositionAndInvRadius[LightIndex].xyz;
-			LightData.InvRadius = TiledDeferred.LightPositionAndInvRadius[LightIndex].w;
-			LightData.Color = TiledDeferred.LightColorAndFalloffExponent[LightIndex].xyz;
-			LightData.FalloffExponent = TiledDeferred.LightColorAndFalloffExponent[LightIndex].w;
-			LightData.bInverseSquared = LightData.FalloffExponent == 0;
-						
-			// todo: doesn't support ScreenSpaceSubsurfaceScattering yet (not using alpha)
-			CompositedLighting.rgb += GetSimpleDynamicLighting(WorldPosition, CameraVector, ScreenSpaceData.GBuffer.WorldNormal, ScreenSpaceData.AmbientOcclusion, ScreenSpaceData.GBuffer.DiffuseColor, ScreenSpaceData.GBuffer.SpecularColor, ScreenSpaceData.GBuffer.Roughness, LightData);
-		}
+	// NOTE: This definition of 'CameraVector' appears to be the reverse of what 
+	// is in forward shading parts e.g, as defined in 'FMaterialPixelParameters'.
+
+	float3 CameraVector = normalize(WorldPosition - View.WorldCameraOrigin);
+
+	// RGB accumulated RGB HDR color, A: luminance for screenspace subsurface scattering
+	float4 CompositedLighting = 0;
+	//BRANCH
+	if (PixelShadingModelID != SHADINGMODELID_UNLIT)
+	{
+#if USE_PASS_PER_SHADING_MODEL
+		EXECUTE_SHADING_LOOPS_SINGLE_SM(SHADINGMODELID_DEFAULT_LIT, PixelShadingModelID, CompositedLighting, PixelPos, NumLightsAffectingTile, NumSimpleLightsAffectingTile, CameraVector, WorldPosition);
+		EXECUTE_SHADING_LOOPS_SINGLE_SM(SHADINGMODELID_SUBSURFACE, PixelShadingModelID, CompositedLighting, PixelPos, NumLightsAffectingTile, NumSimpleLightsAffectingTile, CameraVector, WorldPosition);
+		EXECUTE_SHADING_LOOPS_SINGLE_SM(SHADINGMODELID_PREINTEGRATED_SKIN, PixelShadingModelID, CompositedLighting, PixelPos, NumLightsAffectingTile, NumSimpleLightsAffectingTile, CameraVector, WorldPosition);
+		EXECUTE_SHADING_LOOPS_SINGLE_SM(SHADINGMODELID_CLEAR_COAT, PixelShadingModelID, CompositedLighting, PixelPos, NumLightsAffectingTile, NumSimpleLightsAffectingTile, CameraVector, WorldPosition);
+		EXECUTE_SHADING_LOOPS_SINGLE_SM(SHADINGMODELID_SUBSURFACE_PROFILE, PixelShadingModelID, CompositedLighting, PixelPos, NumLightsAffectingTile, NumSimpleLightsAffectingTile, CameraVector, WorldPosition);
+		EXECUTE_SHADING_LOOPS_SINGLE_SM(SHADINGMODELID_TWOSIDED_FOLIAGE, PixelShadingModelID, CompositedLighting, PixelPos, NumLightsAffectingTile, NumSimpleLightsAffectingTile, CameraVector, WorldPosition);
+		EXECUTE_SHADING_LOOPS_SINGLE_SM(SHADINGMODELID_HAIR, PixelShadingModelID, CompositedLighting, PixelPos, NumLightsAffectingTile, NumSimpleLightsAffectingTile, CameraVector, WorldPosition);
+		EXECUTE_SHADING_LOOPS_SINGLE_SM(SHADINGMODELID_CLOTH, PixelShadingModelID, CompositedLighting, PixelPos, NumLightsAffectingTile, NumSimpleLightsAffectingTile, CameraVector, WorldPosition);
+		EXECUTE_SHADING_LOOPS_SINGLE_SM(SHADINGMODELID_EYE, PixelShadingModelID, CompositedLighting, PixelPos, NumLightsAffectingTile, NumSimpleLightsAffectingTile, CameraVector, WorldPosition);
+#else // !USE_PASS_PER_SHADING_MODEL
+		ExecuteShadingLoops(CompositedLighting, ScreenSpaceData, NumLightsAffectingTile, NumSimpleLightsAffectingTile, CameraVector, WorldPosition);
+#endif // USE_PASS_PER_SHADING_MODEL
 	}
 
-#if USE_PREEXPOSURE
+#if USE_PREEXPOSURE && !VISUALIZE_LIGHT_CULLING
 	CompositedLighting *= View.PreExposure;
-#endif
+#endif // USE_PREEXPOSURE && !VISUALIZE_LIGHT_CULLING
 
 #endif
 
diff --git a/Engine/Source/Runtime/Engine/Classes/Engine/SubsurfaceProfile.h b/Engine/Source/Runtime/Engine/Classes/Engine/SubsurfaceProfile.h
--- a/Engine/Source/Runtime/Engine/Classes/Engine/SubsurfaceProfile.h
+++ b/Engine/Source/Runtime/Engine/Classes/Engine/SubsurfaceProfile.h
@@ -154,7 +154,7 @@ public:
 	void UpdateProfile(int32 AllocationId, const FSubsurfaceProfileStruct Settings);
 
 	// @return can be 0 if there is no SubsurfaceProfile
-	const struct IPooledRenderTarget* GetTexture(FRHICommandListImmediate& RHICmdList);
+	struct IPooledRenderTarget* GetTexture(FRHICommandListImmediate& RHICmdList);
 
 
 	//~ Begin FRenderResource Interface.
diff --git a/Engine/Source/Runtime/Engine/Private/Rendering/SubsurfaceProfile.cpp b/Engine/Source/Runtime/Engine/Private/Rendering/SubsurfaceProfile.cpp
--- a/Engine/Source/Runtime/Engine/Private/Rendering/SubsurfaceProfile.cpp
+++ b/Engine/Source/Runtime/Engine/Private/Rendering/SubsurfaceProfile.cpp
@@ -97,7 +97,7 @@ void FSubsurfaceProfileTexture::UpdateProfile(int32 AllocationId, const FSubsurf
 	GSSProfiles.SafeRelease();
 }
 
-const IPooledRenderTarget* FSubsurfaceProfileTexture::GetTexture(FRHICommandListImmediate& RHICmdList)
+IPooledRenderTarget* FSubsurfaceProfileTexture::GetTexture(FRHICommandListImmediate& RHICmdList)
 {
 	if (!GSSProfiles)
 	{
@@ -356,7 +356,7 @@ void FSubsurfaceProfileTexture::Dump()
 
 
 
-ENGINE_API const IPooledRenderTarget* GetSubsufaceProfileTexture_RT(FRHICommandListImmediate& RHICmdList)
+ENGINE_API IPooledRenderTarget* GetSubsufaceProfileTexture_RT(FRHICommandListImmediate& RHICmdList)
 {
 	check(IsInRenderingThread());
 
diff --git a/Engine/Source/Runtime/Renderer/Private/ClusteredDeferredShadingPass.cpp b/Engine/Source/Runtime/Renderer/Private/ClusteredDeferredShadingPass.cpp
new file mode 100644
--- /dev/null
+++ b/Engine/Source/Runtime/Renderer/Private/ClusteredDeferredShadingPass.cpp
@@ -0,0 +1,180 @@
+// Copyright 1998-2019 Epic Games, Inc. All Rights Reserved.
+
+/*=============================================================================
+	TiledDeferredLightRendering.cpp: Implementation of tiled deferred shading
+=============================================================================*/
+
+#include "CoreMinimal.h"
+#include "Stats/Stats.h"
+#include "HAL/IConsoleManager.h"
+#include "EngineGlobals.h"
+#include "RHI.h"
+#include "UniformBuffer.h"
+#include "ShaderParameters.h"
+#include "RendererInterface.h"
+#include "Shader.h"
+#include "SceneUtils.h"
+#include "RHIStaticStates.h"
+#include "PostProcess/SceneRenderTargets.h"
+#include "LightSceneInfo.h"
+#include "GlobalShader.h"
+#include "SceneRenderTargetParameters.h"
+#include "DeferredShadingRenderer.h"
+#include "ScenePrivate.h"
+
+#include "SceneFilterRendering.h"
+#include "PostProcessing.h"
+
+// TODO: This looks like a strange and hacky way to get hold of this resource, but it seems to work.
+ENGINE_API IPooledRenderTarget* GetSubsufaceProfileTexture_RT(FRHICommandListImmediate& RHICmdList);
+
+// This is used to switch on and off the clustered deferred shading implementation, that uses the light grid to perform shading.
+int32 GUseClusteredDeferredShading = 0;
+static FAutoConsoleVariableRef CVarUseClusteredDeferredShading(
+	TEXT("r.UseClusteredDeferredShading"),
+	GUseClusteredDeferredShading,
+	TEXT("Toggle use of clustered deferred shading for lights that support it. 0 is off (default), 1 is on (also required is SM5 to actually turn on)."),
+	ECVF_RenderThreadSafe
+);
+
+DECLARE_GPU_STAT_NAMED(ClusteredShading, TEXT("Clustered Shading"));
+
+
+
+bool FDeferredShadingSceneRenderer::ShouldUseClusteredDeferredShading() const
+{
+	// The feature level is the same as in the shader compile conditions below, maybe we don't need SM5?
+	// NOTE: should also take into account the conditions for building the light grid, since these 
+	//       shaders might have another feature level.
+	return GUseClusteredDeferredShading != 0 && Scene->GetFeatureLevel() >= ERHIFeatureLevel::SM5;
+}
+
+bool FDeferredShadingSceneRenderer::AreClusteredLightsInLightGrid() const
+{
+	return bClusteredShadingLightsInLightGrid;
+}
+
+
+/**
+ * Clustered deferred shading shader, used in a full-screen pass to apply all lights in the light grid.
+ */
+class FClusteredShadingPS : public FGlobalShader
+{
+	DECLARE_GLOBAL_SHADER(FClusteredShadingPS);
+	SHADER_USE_PARAMETER_STRUCT(FClusteredShadingPS, FGlobalShader)
+
+	class FVisualizeLightCullingDim : SHADER_PERMUTATION_BOOL("VISUALIZE_LIGHT_CULLING");
+	using FPermutationDomain = TShaderPermutationDomain<FVisualizeLightCullingDim>;
+
+	BEGIN_SHADER_PARAMETER_STRUCT(FParameters, )
+		RENDER_TARGET_BINDING_SLOTS()
+		SHADER_PARAMETER_STRUCT_REF(FForwardLightData, Forward)
+		SHADER_PARAMETER_STRUCT_REF(FViewUniformShaderParameters, View)
+		SHADER_PARAMETER_STRUCT_REF(FSceneTexturesUniformParameters, SceneTextures)
+
+		SHADER_PARAMETER_RDG_TEXTURE(Texture2D, LTCMatTexture)
+		SHADER_PARAMETER_SAMPLER(SamplerState, LTCMatSampler)
+
+		SHADER_PARAMETER_RDG_TEXTURE(Texture2D, LTCAmpTexture)
+		SHADER_PARAMETER_SAMPLER(SamplerState, LTCAmpSampler)
+
+		SHADER_PARAMETER_RDG_TEXTURE(Texture2D, SSProfilesTexture)
+		SHADER_PARAMETER_SAMPLER(SamplerState, TransmissionProfilesLinearSampler)
+
+	END_SHADER_PARAMETER_STRUCT()
+
+	static bool ShouldCompilePermutation(const FGlobalShaderPermutationParameters& Parameters)
+	{
+		// OLATODO: what level do we actually need for this?
+		return IsFeatureLevelSupported(Parameters.Platform, ERHIFeatureLevel::SM5);
+	}
+
+	static void ModifyCompilationEnvironment(const FGlobalShaderPermutationParameters& Parameters, FShaderCompilerEnvironment& OutEnvironment)
+	{
+		FGlobalShader::ModifyCompilationEnvironment(Parameters, OutEnvironment);
+		FForwardLightingParameters::ModifyCompilationEnvironment(Parameters.Platform, OutEnvironment);
+	}
+};
+
+IMPLEMENT_GLOBAL_SHADER(FClusteredShadingPS, "/Engine/Private/ClusteredDeferredShadingPixelShader.usf", "ClusteredShadingPixelShader", SF_Pixel);
+
+void FDeferredShadingSceneRenderer::AddClusteredDeferredShadingPass(FRHICommandListImmediate& RHICmdList, const FSortedLightSetSceneInfo &SortedLightsSet)
+{
+	check(GUseClusteredDeferredShading);
+
+	const int32 NumLightsToRender = SortedLightsSet.ClusteredSupportedEnd;
+
+	if (NumLightsToRender > 0)
+	{
+		FRDGBuilder GraphBuilder(RHICmdList);
+
+		// TODO: are these going to measure what I think if placed here? Or should they be inside the lambda?
+		SCOPED_GPU_STAT(RHICmdList, ClusteredShading);
+		SCOPED_DRAW_EVENTF(RHICmdList, ClusteredShading, TEXT("ClusteredShading"));
+
+		FSceneRenderTargets& SceneContext = FSceneRenderTargets::Get(RHICmdList);
+
+		for (int32 ViewIndex = 0, Num = Views.Num(); ViewIndex < Num; ViewIndex++)
+		{
+			FViewInfo& View = Views[ViewIndex];
+
+			FClusteredShadingPS::FParameters *PassParameters = GraphBuilder.AllocParameters<FClusteredShadingPS::FParameters>();
+
+			PassParameters->View = View.ViewUniformBuffer;
+			// NOTE: wonder why the scene textures uniform buffer is not pre prepared somewhere? Used a lot...
+			PassParameters->SceneTextures = CreateSceneTextureUniformBufferSingleDraw(RHICmdList, ESceneTextureSetupMode::All, View.FeatureLevel);
+			// NOTE: This doesn't contain any RDG mention, so does this enforce dependencies or is this a later feature?
+			PassParameters->Forward = View.ForwardLightingResources->ForwardLightDataUniformBuffer;
+
+			PassParameters->LTCMatTexture = GraphBuilder.RegisterExternalTexture(GSystemTextures.LTCMat);
+			PassParameters->LTCMatSampler = TStaticSamplerState<SF_Bilinear, AM_Clamp, AM_Clamp, AM_Clamp>::GetRHI();
+			PassParameters->LTCAmpTexture = GraphBuilder.RegisterExternalTexture(GSystemTextures.LTCAmp);
+			PassParameters->LTCAmpSampler = TStaticSamplerState<SF_Bilinear, AM_Clamp, AM_Clamp, AM_Clamp>::GetRHI();
+			PassParameters->SSProfilesTexture = RegisterExternalTextureWithFallback(GraphBuilder, GetSubsufaceProfileTexture_RT(RHICmdList), GSystemTextures.BlackDummy);
+			PassParameters->TransmissionProfilesLinearSampler = TStaticSamplerState<SF_Bilinear, AM_Clamp, AM_Clamp, AM_Clamp>::GetRHI();
+
+			PassParameters->RenderTargets[0] = FRenderTargetBinding(GraphBuilder.RegisterExternalTexture(SceneContext.GetSceneColor()), ERenderTargetLoadAction::ENoAction, ERenderTargetStoreAction::ENoAction);
+			// NOTE: if we wanted to get depth/stencil testing happening we'd need to add that here too, at the moment we don't.
+
+			GraphBuilder.AddPass(
+				RDG_EVENT_NAME("ClusteredDeferredShading, #Lights: %d", NumLightsToRender),
+				PassParameters,
+				ERenderGraphPassFlags::None,
+				[PassParameters, &View, &SceneContext](FRHICommandListImmediate& RHICmdList)
+			{
+				TShaderMapRef<FPostProcessVS> VertexShader(View.ShaderMap);
+
+				FClusteredShadingPS::FPermutationDomain PermutationVector;
+				PermutationVector.Set<FClusteredShadingPS::FVisualizeLightCullingDim>(View.Family->EngineShowFlags.VisualizeLightCulling);
+				TShaderMapRef<FClusteredShadingPS> PixelShader(View.ShaderMap, PermutationVector);
+				{
+					FGraphicsPipelineStateInitializer GraphicsPSOInit;
+
+					// Additive blend to accumulate lighting contributions.
+					GraphicsPSOInit.BlendState = TStaticBlendState<CW_RGBA, BO_Add, BF_One, BF_One, BO_Add, BF_One, BF_One>::GetRHI();
+
+					GraphicsPSOInit.RasterizerState = TStaticRasterizerState<FM_Solid, CM_None>::GetRHI();
+					GraphicsPSOInit.DepthStencilState = TStaticDepthStencilState<false, CF_Always>::GetRHI();
+
+					GraphicsPSOInit.BoundShaderState.VertexDeclarationRHI = GFilterVertexDeclaration.VertexDeclarationRHI;
+					GraphicsPSOInit.BoundShaderState.VertexShaderRHI = GETSAFERHISHADER_VERTEX(*VertexShader);
+					GraphicsPSOInit.BoundShaderState.PixelShaderRHI = GETSAFERHISHADER_PIXEL(*PixelShader);
+					GraphicsPSOInit.PrimitiveType = PT_TriangleList;
+
+					SetGraphicsPipelineState(RHICmdList, GraphicsPSOInit, EApplyRendertargetOption::ForceApply);
+				}
+				RHICmdList.SetViewport(View.ViewRect.Min.X, View.ViewRect.Min.Y, 0.0f, View.ViewRect.Max.X, View.ViewRect.Max.Y, 1.0f);
+
+				SetShaderParameters(RHICmdList, *PixelShader, PixelShader->GetPixelShader(), *PassParameters);
+
+				DrawRectangle(RHICmdList, 0, 0, View.ViewRect.Width(), View.ViewRect.Height(),
+					View.ViewRect.Min.X, View.ViewRect.Min.Y, View.ViewRect.Width(), View.ViewRect.Height(),
+					FIntPoint(View.ViewRect.Width(), View.ViewRect.Height()), SceneContext.GetBufferSizeXY(), *VertexShader);
+			});
+
+		}
+		// NOTE: I have not added any queue extraction thingo or suchlike, should it? I assume that once the render targets are setup
+		//       to use RDG they'll get tracked correctly.
+		GraphBuilder.Execute();
+	}
+}
diff --git a/Engine/Source/Runtime/Renderer/Private/DeferredShadingRenderer.cpp b/Engine/Source/Runtime/Renderer/Private/DeferredShadingRenderer.cpp
--- a/Engine/Source/Runtime/Renderer/Private/DeferredShadingRenderer.cpp
+++ b/Engine/Source/Runtime/Renderer/Private/DeferredShadingRenderer.cpp
@@ -201,6 +201,7 @@ FDeferredShadingSceneRenderer::FDeferredShadingSceneRenderer(const FSceneViewFam
 	: FSceneRenderer(InViewFamily, HitProxyConsumer)
 	, EarlyZPassMode(Scene ? Scene->EarlyZPassMode : DDM_None)
 	, bEarlyZPassMovable(Scene ? Scene->bEarlyZPassMovable : false)
+	, bClusteredShadingLightsInLightGrid(false)
 {
 	static const auto StencilLODDitherCVar = IConsoleManager::Get().FindTConsoleVariableDataInt(TEXT("r.StencilForLODDither"));
 	bDitheredLODTransitionsUseStencil = StencilLODDitherCVar->GetValueOnAnyThread() != 0;
@@ -1139,7 +1140,10 @@ void FDeferredShadingSceneRenderer::Render(FRHICommandListImmediate& RHICmdList)
 	SceneContext.ResolveSceneDepthTexture(RHICmdList, FResolveRect(0, 0, FamilySize.X, FamilySize.Y));
 	SceneContext.ResolveSceneDepthToAuxiliaryTexture(RHICmdList);
 
-	ComputeLightGrid(RHICmdList, bComputeLightGrid);
+	// NOTE: The ordering of the lights is used to select sub-sets for different purposes, e.g., those that support clustered deferred.
+	FSortedLightSetSceneInfo SortedLightSet;
+	GatherAndSortLights(SortedLightSet);
+	ComputeLightGrid(RHICmdList, bComputeLightGrid, SortedLightSet);
 
 	if (bUseGBuffer || IsSimpleForwardShadingEnabled(ShaderPlatform))
 	{
@@ -1717,7 +1721,7 @@ void FDeferredShadingSceneRenderer::Render(FRHICommandListImmediate& RHICmdList)
 		}
 
 		RHICmdList.SetCurrentStat(GET_STATID(STAT_CLM_Lighting));
-		RenderLights(RHICmdList);
+		RenderLights(RHICmdList, SortedLightSet);
 		RHICmdList.SetCurrentStat(GET_STATID(STAT_CLM_AfterLighting));
 		ServiceLocalQueue();
 
@@ -1918,7 +1922,7 @@ void FDeferredShadingSceneRenderer::Render(FRHICommandListImmediate& RHICmdList)
 	GRenderTargetPool.AddPhaseEvent(TEXT("Translucency"));
 
 	// Draw translucency.
-	if (bCanOverlayRayTracingOutput && ViewFamily.EngineShowFlags.Translucency)
+	if (bCanOverlayRayTracingOutput && ViewFamily.EngineShowFlags.Translucency && !ViewFamily.EngineShowFlags.VisualizeLightCulling)
 	{
 		SCOPE_CYCLE_COUNTER(STAT_TranslucencyDrawTime);
 
diff --git a/Engine/Source/Runtime/Renderer/Private/DeferredShadingRenderer.h b/Engine/Source/Runtime/Renderer/Private/DeferredShadingRenderer.h
--- a/Engine/Source/Runtime/Renderer/Private/DeferredShadingRenderer.h
+++ b/Engine/Source/Runtime/Renderer/Private/DeferredShadingRenderer.h
@@ -69,8 +69,11 @@ public:
 	 */
 	bool RenderPrePassViewParallel(const FViewInfo& View, FRHICommandListImmediate& ParentCmdList, const FMeshPassProcessorRenderState& DrawRenderState, TFunctionRef<void()> AfterTasksAreStarted, bool bDoPrePre);
 
-	/** Culls local lights to a grid in frustum space.  Needed for forward shading or translucency using the Surface lighting mode. */
-	void ComputeLightGrid(FRHICommandListImmediate& RHICmdList, bool bNeedLightGrid);
+	/** 
+	 * Culls local lights and reflection probes to a grid in frustum space, builds one light list and grid per view in the current Views.  
+	 * Needed for forward shading or translucency using the Surface lighting mode, and clustered deferred shading. 
+	 */
+	void ComputeLightGrid(FRHICommandListImmediate& RHICmdList, bool bNeedLightGrid, FSortedLightSetSceneInfo &SortedLightSet);
 
 	/** Renders the basepass for a given View, in parallel */
 	void RenderBasePassViewParallel(FViewInfo& View, FRHICommandListImmediate& ParentCmdList, FExclusiveDepthStencil::Type BasePassDepthStencilAccess, const FMeshPassProcessorRenderState& InDrawRenderState);
@@ -222,13 +225,34 @@ private:
 	bool CanUseTiledDeferred() const;
 
 	/** Whether to use tiled deferred shading given a number of lights that support it. */
-	bool ShouldUseTiledDeferred(int32 NumUnshadowedLights, int32 NumSimpleLights) const;
+	bool ShouldUseTiledDeferred(int32 NumTiledDeferredLights) const;
+
+	/** 
+	 * True if the 'r.UseClusteredDeferredShading' flag is 1 and sufficient feature level. 
+	 * NOTE: When true it takes precedence over the TiledDeferred path, since they handle the same lights.
+	 */
+	bool ShouldUseClusteredDeferredShading() const;
+
+	/**
+	 * Have the requisite lights been injected into the light grid, AKA can we run the shading pass?
+	 */
+	bool AreClusteredLightsInLightGrid() const;
+
 
-	/** Renders the lights in SortedLights in the range [0, NumUnshadowedLights) using tiled deferred shading. */
-	void RenderTiledDeferredLighting(FRHICommandListImmediate& RHICmdList, const TArray<FSortedLightSceneInfo, SceneRenderingAllocator>& SortedLights, int32 NumUnshadowedLights, const FSimpleLightArray& SimpleLights);
+	/** Add a clustered deferred shading lighting render pass.	Note: in the future it should take the RenderGraph builder as argument */
+	void AddClusteredDeferredShadingPass(FRHICommandListImmediate& RHICmdList, const FSortedLightSetSceneInfo &SortedLightsSet);
+
+	/** Renders the lights in SortedLights in the range [TiledDeferredLightsStart, TiledDeferredLightsEnd) using tiled deferred shading. */
+	void RenderTiledDeferredLighting(FRHICommandListImmediate& RHICmdList, const TArray<FSortedLightSceneInfo, SceneRenderingAllocator>& SortedLights, int32 TiledDeferredLightsStart, int32 TiledDeferredLightsEnd, const FSimpleLightArray& SimpleLights);
+
+	/**
+	 * Rounds up lights and sorts them according to what type of renderer supports them. The result is stored in OutSortedLights 
+	 * NOTE: Also extracts the SimpleLights AND adds them to the sorted range (first sub-range). 
+	 */
+	void GatherAndSortLights(FSortedLightSetSceneInfo& OutSortedLights);
 
 	/** Renders the scene's lighting. */
-	void RenderLights(FRHICommandListImmediate& RHICmdList);
+	void RenderLights(FRHICommandListImmediate& RHICmdList, FSortedLightSetSceneInfo &SortedLightSet);
 
 	/** Renders an array of lights for the stationary light overlap viewmode. */
 	void RenderLightArrayForOverlapViewmode(FRHICommandListImmediate& RHICmdList, const TSparseArray<FLightSceneInfoCompact>& LightArray);
@@ -369,7 +393,7 @@ private:
 	void InjectTranslucentVolumeLighting(FRHICommandListImmediate& RHICmdList, const FLightSceneInfo& LightSceneInfo, const FProjectedShadowInfo* InProjectedShadowInfo, const FViewInfo& View, int32 ViewIndex);
 
 	/** Accumulates direct lighting for an array of unshadowed lights. */
-	void InjectTranslucentVolumeLightingArray(FRHICommandListImmediate& RHICmdList, const TArray<FSortedLightSceneInfo, SceneRenderingAllocator>& SortedLights, int32 NumLights);
+	void InjectTranslucentVolumeLightingArray(FRHICommandListImmediate& RHICmdList, const TArray<FSortedLightSceneInfo, SceneRenderingAllocator>& SortedLights, int32 FirstLightIndex, int32 LightsEndIndex);
 
 	/** Accumulates direct lighting for simple lights. */
 	void InjectSimpleTranslucentVolumeLightingArray(FRHICommandListImmediate& RHICmdList, const FSimpleLightArray& SimpleLights, const FViewInfo& View, const int32 ViewIndex);
@@ -539,6 +563,9 @@ private:
 	static void PreparePathTracing(const FViewInfo& View, TArray<FRayTracingShaderRHIParamRef>& OutRayGenShaders);
 
 #endif // RHI_RAYTRACING
+
+	/** Set to true if the lights needed for clustered shading have been injected in the light grid (set in ComputeLightGrid). */
+	bool bClusteredShadingLightsInLightGrid;
 };
 
 DECLARE_CYCLE_STAT_EXTERN(TEXT("PrePass"), STAT_CLM_PrePass, STATGROUP_CommandListMarkers, );
\ No newline at end of file
diff --git a/Engine/Source/Runtime/Renderer/Private/LightGridInjection.cpp b/Engine/Source/Runtime/Renderer/Private/LightGridInjection.cpp
--- a/Engine/Source/Runtime/Renderer/Private/LightGridInjection.cpp
+++ b/Engine/Source/Runtime/Renderer/Private/LightGridInjection.cpp
@@ -73,6 +73,35 @@ FAutoConsoleVariableRef CVarLightCullingQuality(
 	ECVF_RenderThreadSafe
 );
 
+
+// TODO move to render graph utils
+BEGIN_SHADER_PARAMETER_STRUCT(FClearUAVParameters, )
+	SHADER_PARAMETER_RDG_TEXTURE_UAV(	RWTexture2D,	TextureUAV)
+	SHADER_PARAMETER_RDG_BUFFER_UAV(	RWBuffer<uint>,	BufferUAV)
+END_SHADER_PARAMETER_STRUCT()
+
+FORCEINLINE_DEBUGGABLE void AddPass_ClearUAV(
+	FRDGBuilder& GraphBuilder,
+	FRDGEventName&& PassName,
+	FRDGBufferUAVRef BufferUAV,
+	uint32 Value)
+{
+	FClearUAVParameters* Parameters = GraphBuilder.AllocParameters< FClearUAVParameters >();
+	Parameters->BufferUAV = BufferUAV;
+
+	GraphBuilder.AddPass(
+		Forward<FRDGEventName>(PassName),
+		Parameters,
+		ERenderGraphPassFlags::Compute,
+		[&Parameters, BufferUAV, Value](FRHICommandList& RHICmdList)
+		{
+			ClearUAV( RHICmdList, BufferUAV->CachedRHI.UAV, BufferUAV->Desc.Buffer->Desc.GetTotalNumBytes(), Value );
+		} );
+}
+
+
+
+
 /** A minimal forwarding lighting setup. */
 class FMinimalDummyForwardLightingResources : public FRenderResource
 {
@@ -146,6 +175,8 @@ FForwardLightData::FForwardLightData()
 	DirectionalLightStaticShadowmap = GBlackTexture->TextureRHI;
 	StaticShadowmapSampler = TStaticSamplerState<SF_Bilinear,AM_Clamp,AM_Clamp,AM_Clamp>::GetRHI();
 
+	DummyRectLightSourceTexture = GWhiteTexture->TextureRHI;
+
 	ForwardLocalLightBuffer = nullptr;
 	NumCulledLightsGrid = nullptr;
 	CulledLightDataGrid = nullptr;
@@ -160,123 +191,30 @@ typedef uint16 FLightIndexType;
 // UINT_MAX indexable light limit
 typedef uint32 FLightIndexType32;
 
-/**  */
-class FForwardCullingParameters
-{
-public:
-
-	static void ModifyCompilationEnvironment(EShaderPlatform Platform, FShaderCompilerEnvironment& OutEnvironment)
-	{
-		OutEnvironment.SetDefine(TEXT("LIGHT_LINK_STRIDE"), LightLinkStride);
-	}
-
-	void Bind(const FShaderParameterMap& ParameterMap)
-	{
-		NextCulledLightLink.Bind(ParameterMap, TEXT("NextCulledLightLink"));
-		StartOffsetGrid.Bind(ParameterMap, TEXT("StartOffsetGrid"));
-		CulledLightLinks.Bind(ParameterMap, TEXT("CulledLightLinks"));
-		NextCulledLightData.Bind(ParameterMap, TEXT("NextCulledLightData"));
-		NumCulledLightsGrid.Bind(ParameterMap, TEXT("NumCulledLightsGrid"));
-		CulledLightDataGrid.Bind(ParameterMap, TEXT("CulledLightDataGrid"));
-	}
-
-	template<typename ShaderRHIParamRef>
-	void Set(
-		FRHICommandList& RHICmdList,
-		const ShaderRHIParamRef& ShaderRHI,
-		const FForwardLightingCullingResources& ForwardLightingCullingResources,
-		const FForwardLightingViewResources& ViewResources)
-	{
-		NextCulledLightLink.SetBuffer(RHICmdList, ShaderRHI, ForwardLightingCullingResources.NextCulledLightLink);
-		StartOffsetGrid.SetBuffer(RHICmdList, ShaderRHI, ForwardLightingCullingResources.StartOffsetGrid);
-		CulledLightLinks.SetBuffer(RHICmdList, ShaderRHI, ForwardLightingCullingResources.CulledLightLinks);
-		NextCulledLightData.SetBuffer(RHICmdList, ShaderRHI, ForwardLightingCullingResources.NextCulledLightData);
-		NumCulledLightsGrid.SetBuffer(RHICmdList, ShaderRHI, ViewResources.NumCulledLightsGrid);
-		CulledLightDataGrid.SetBuffer(RHICmdList, ShaderRHI, ViewResources.CulledLightDataGrid);
-	}
-
-	template<typename ShaderRHIParamRef>
-	void UnsetParameters(
-		FRHICommandList& RHICmdList,
-		const ShaderRHIParamRef& ShaderRHI,
-		const FForwardLightingCullingResources& ForwardLightingCullingResources,
-		const FForwardLightingViewResources& ViewResources)
-	{
-		NextCulledLightLink.UnsetUAV(RHICmdList, ShaderRHI);
-		StartOffsetGrid.UnsetUAV(RHICmdList, ShaderRHI);
-		CulledLightLinks.UnsetUAV(RHICmdList, ShaderRHI);
-		NextCulledLightData.UnsetUAV(RHICmdList, ShaderRHI);
-		NumCulledLightsGrid.UnsetUAV(RHICmdList, ShaderRHI);
-		CulledLightDataGrid.UnsetUAV(RHICmdList, ShaderRHI);
-
-		TArray<FUnorderedAccessViewRHIParamRef, TInlineAllocator<4>> OutUAVs;
-
-		if (NextCulledLightLink.IsUAVBound())
-		{
-			OutUAVs.Add(ForwardLightingCullingResources.NextCulledLightLink.UAV);
-		}
-
-		if (StartOffsetGrid.IsUAVBound())
-		{
-			OutUAVs.Add(ForwardLightingCullingResources.StartOffsetGrid.UAV);
-		}
-
-		if (CulledLightLinks.IsUAVBound())
-		{
-			OutUAVs.Add(ForwardLightingCullingResources.CulledLightLinks.UAV);
-		}
-
-		if (NextCulledLightData.IsUAVBound())
-		{
-			OutUAVs.Add(ForwardLightingCullingResources.NextCulledLightData.UAV);
-		}
-
-		if (NumCulledLightsGrid.IsUAVBound())
-		{
-			OutUAVs.Add(ViewResources.NumCulledLightsGrid.UAV);
-		}
-
-		if (CulledLightDataGrid.IsUAVBound())
-		{
-			OutUAVs.Add(ViewResources.CulledLightDataGrid.UAV);
-		}
-
-		if (OutUAVs.Num() > 0)
-		{
-			RHICmdList.TransitionResources(EResourceTransitionAccess::EReadable, EResourceTransitionPipeline::EComputeToGfx, OutUAVs.GetData(), OutUAVs.Num());
-		}
-	}
-
-	/** Serializer. */
-	friend FArchive& operator<<(FArchive& Ar, FForwardCullingParameters& P)
-	{
-		Ar << P.NextCulledLightLink;
-		Ar << P.StartOffsetGrid;
-		Ar << P.CulledLightLinks;
-		Ar << P.NextCulledLightData;
-		Ar << P.NumCulledLightsGrid;
-		Ar << P.CulledLightDataGrid;
-		return Ar;
-	}
-
-private:
-
-	FRWShaderParameter NextCulledLightLink;
-	FRWShaderParameter StartOffsetGrid;
-	FRWShaderParameter CulledLightLinks;
-	FRWShaderParameter NextCulledLightData;
-	FRWShaderParameter NumCulledLightsGrid;
-	FRWShaderParameter CulledLightDataGrid;
-};
-
-
 uint32 LightGridInjectionGroupSize = 4;
 
-template<bool bLightLinkedListCulling>
-class TLightGridInjectionCS : public FGlobalShader
+
+class FLightGridInjectionCS : public FGlobalShader
 {
-	DECLARE_SHADER_TYPE(TLightGridInjectionCS, Global)
+	DECLARE_GLOBAL_SHADER(FLightGridInjectionCS);
+	SHADER_USE_PARAMETER_STRUCT(FLightGridInjectionCS, FGlobalShader)
 public:
+	class FUseLinkedListDim : SHADER_PERMUTATION_BOOL("USE_LINKED_CULL_LIST");
+	using FPermutationDomain = TShaderPermutationDomain<FUseLinkedListDim>;
+
+	BEGIN_SHADER_PARAMETER_STRUCT(FParameters, )
+		//SHADER_PARAMETER_STRUCT_INCLUDE(FLightGridInjectionCommonParameters, CommonParameters)
+		SHADER_PARAMETER_STRUCT_REF(FReflectionCaptureShaderData, ReflectionCapture)
+		SHADER_PARAMETER_STRUCT_REF(FForwardLightData, Forward)
+		SHADER_PARAMETER_STRUCT_REF(FViewUniformShaderParameters, View)
+		SHADER_PARAMETER_UAV(RWBuffer<uint>, RWNumCulledLightsGrid)
+		SHADER_PARAMETER_UAV(RWBuffer<uint>, RWCulledLightDataGrid)
+		SHADER_PARAMETER_RDG_BUFFER_UAV(RWBuffer<uint>, RWNextCulledLightLink)
+		SHADER_PARAMETER_RDG_BUFFER_UAV(RWBuffer<uint>, RWStartOffsetGrid)
+		SHADER_PARAMETER_RDG_BUFFER_UAV(RWBuffer<uint>, RWCulledLightLinks)
+		SHADER_PARAMETER_SRV(StrongTypedBuffer<float4>, LightViewSpacePositionAndRadius)
+		SHADER_PARAMETER_SRV(StrongTypedBuffer<float4>, LightViewSpaceDirAndPreprocAngle)
+	END_SHADER_PARAMETER_STRUCT()
 
 	static bool ShouldCompilePermutation(const FGlobalShaderPermutationParameters& Parameters)
 	{
@@ -288,54 +226,29 @@ public:
 		FGlobalShader::ModifyCompilationEnvironment(Parameters, OutEnvironment);
 		OutEnvironment.SetDefine(TEXT("THREADGROUP_SIZE"), LightGridInjectionGroupSize);
 		FForwardLightingParameters::ModifyCompilationEnvironment(Parameters.Platform, OutEnvironment);
-		FForwardCullingParameters::ModifyCompilationEnvironment(Parameters.Platform, OutEnvironment);
-		OutEnvironment.SetDefine(TEXT("USE_LINKED_CULL_LIST"), bLightLinkedListCulling);
-	}
-
-	TLightGridInjectionCS(const ShaderMetaType::CompiledShaderInitializerType& Initializer)
-		: FGlobalShader(Initializer)
-	{
-		ForwardCullingParameters.Bind(Initializer.ParameterMap);
-	}
-
-	TLightGridInjectionCS()
-	{
-	}
-
-	void SetParameters(FRHICommandList& RHICmdList, const FViewInfo& View, const FForwardLightingCullingResources& ForwardLightingCullingResources)
-	{
-		FComputeShaderRHIParamRef ShaderRHI = GetComputeShader();
-		FGlobalShader::SetParameters<FViewUniformShaderParameters>(RHICmdList, ShaderRHI, View.ViewUniformBuffer);
-		ForwardCullingParameters.Set(RHICmdList, ShaderRHI, ForwardLightingCullingResources, *View.ForwardLightingResources);
-
-		SetUniformBufferParameter(RHICmdList, ShaderRHI, GetUniformBufferParameter<FForwardLightData>(), View.ForwardLightingResources->ForwardLightDataUniformBuffer);
-		SetUniformBufferParameter(RHICmdList, ShaderRHI, GetUniformBufferParameter<FReflectionCaptureShaderData>(), View.ReflectionCaptureUniformBuffer);
-	}
-
-	void UnsetParameters(FRHICommandList& RHICmdList, const FViewInfo& View, const FForwardLightingCullingResources& ForwardLightingCullingResources)
-	{
-		ForwardCullingParameters.UnsetParameters(RHICmdList, GetComputeShader(), ForwardLightingCullingResources, *View.ForwardLightingResources);
-	}
-
-	virtual bool Serialize(FArchive& Ar)
-	{
-		bool bShaderHasOutdatedParameters = FGlobalShader::Serialize(Ar);
-		Ar << ForwardCullingParameters;
-		return bShaderHasOutdatedParameters;
+		OutEnvironment.SetDefine(TEXT("LIGHT_LINK_STRIDE"), LightLinkStride);
+		OutEnvironment.SetDefine(TEXT("ENABLE_LIGHT_CULLING_VIEW_SPACE_BUILD_DATA"), ENABLE_LIGHT_CULLING_VIEW_SPACE_BUILD_DATA);
 	}
-
-private:
-
-	FForwardCullingParameters ForwardCullingParameters;
 };
 
-IMPLEMENT_SHADER_TYPE(template<>, TLightGridInjectionCS<true>, TEXT("/Engine/Private/LightGridInjection.usf"), TEXT("LightGridInjectionCS"), SF_Compute);
-IMPLEMENT_SHADER_TYPE(template<>, TLightGridInjectionCS<false>, TEXT("/Engine/Private/LightGridInjection.usf"), TEXT("LightGridInjectionCS"), SF_Compute);
+IMPLEMENT_GLOBAL_SHADER(FLightGridInjectionCS, "/Engine/Private/LightGridInjection.usf", "LightGridInjectionCS", SF_Compute);
+
 
 class FLightGridCompactCS : public FGlobalShader
 {
-	DECLARE_SHADER_TYPE(FLightGridCompactCS, Global)
+	DECLARE_GLOBAL_SHADER(FLightGridCompactCS)
+	SHADER_USE_PARAMETER_STRUCT(FLightGridCompactCS, FGlobalShader)
 public:
+	BEGIN_SHADER_PARAMETER_STRUCT(FParameters, )
+		SHADER_PARAMETER_STRUCT_REF(FForwardLightData, Forward)
+		SHADER_PARAMETER_STRUCT_REF(FViewUniformShaderParameters, View)
+		SHADER_PARAMETER_UAV(RWBuffer<uint>, RWNumCulledLightsGrid)
+		SHADER_PARAMETER_UAV(RWBuffer<uint>, RWCulledLightDataGrid)
+		SHADER_PARAMETER_RDG_BUFFER_UAV(RWBuffer<uint>, RWNextCulledLightData)
+		SHADER_PARAMETER_RDG_BUFFER_SRV(Buffer<uint>, StartOffsetGrid)
+		SHADER_PARAMETER_RDG_BUFFER_SRV(Buffer<uint>, CulledLightLinks)
+
+	END_SHADER_PARAMETER_STRUCT()
 
 	static bool ShouldCompilePermutation(const FGlobalShaderPermutationParameters& Parameters)
 	{
@@ -347,48 +260,26 @@ public:
 		FGlobalShader::ModifyCompilationEnvironment(Parameters, OutEnvironment);
 		OutEnvironment.SetDefine(TEXT("THREADGROUP_SIZE"), LightGridInjectionGroupSize);
 		FForwardLightingParameters::ModifyCompilationEnvironment(Parameters.Platform, OutEnvironment);
-		FForwardCullingParameters::ModifyCompilationEnvironment(Parameters.Platform, OutEnvironment);
+		OutEnvironment.SetDefine(TEXT("LIGHT_LINK_STRIDE"), LightLinkStride);
 		OutEnvironment.SetDefine(TEXT("MAX_CAPTURES"), GMaxNumReflectionCaptures);
+		OutEnvironment.SetDefine(TEXT("ENABLE_LIGHT_CULLING_VIEW_SPACE_BUILD_DATA"), ENABLE_LIGHT_CULLING_VIEW_SPACE_BUILD_DATA);
 	}
+};
 
-	FLightGridCompactCS(const ShaderMetaType::CompiledShaderInitializerType& Initializer)
-		: FGlobalShader(Initializer)
-	{
-		ForwardCullingParameters.Bind(Initializer.ParameterMap);
-	}
-
-	FLightGridCompactCS()
-	{
-	}
-
-	void SetParameters(FRHICommandList& RHICmdList, const FViewInfo& View, const FForwardLightingCullingResources& ForwardLightingCullingResources)
-	{
-		FComputeShaderRHIParamRef ShaderRHI = GetComputeShader();
-		FGlobalShader::SetParameters<FViewUniformShaderParameters>(RHICmdList, ShaderRHI, View.ViewUniformBuffer);
-		ForwardCullingParameters.Set(RHICmdList, ShaderRHI, ForwardLightingCullingResources, *View.ForwardLightingResources);
-
-		SetUniformBufferParameter(RHICmdList, ShaderRHI, GetUniformBufferParameter<FForwardLightData>(), View.ForwardLightingResources->ForwardLightDataUniformBuffer);
-		SetUniformBufferParameter(RHICmdList, ShaderRHI, GetUniformBufferParameter<FReflectionCaptureShaderData>(), View.ReflectionCaptureUniformBuffer);
-	}
-
-	void UnsetParameters(FRHICommandList& RHICmdList, const FViewInfo& View, const FForwardLightingCullingResources& ForwardLightingCullingResources)
-	{
-		ForwardCullingParameters.UnsetParameters(RHICmdList, GetComputeShader(), ForwardLightingCullingResources, *View.ForwardLightingResources);
-	}
+IMPLEMENT_GLOBAL_SHADER(FLightGridCompactCS, "/Engine/Private/LightGridInjection.usf", "LightGridCompactCS", SF_Compute);
 
-	virtual bool Serialize(FArchive& Ar)
+/**
+ */
+FORCEINLINE float GetTanRadAngleOrZero(float coneAngle)
+{
+	if (coneAngle < PI / 2.001f)
 	{
-		bool bShaderHasOutdatedParameters = FGlobalShader::Serialize(Ar);
-		Ar << ForwardCullingParameters;
-		return bShaderHasOutdatedParameters;
+		return FMath::Tan(coneAngle);
 	}
 
-private:
-
-	FForwardCullingParameters ForwardCullingParameters;
-};
+	return 0.0f;
+}
 
-IMPLEMENT_SHADER_TYPE(, FLightGridCompactCS, TEXT("/Engine/Private/LightGridInjection.usf"), TEXT("LightGridCompactCS"), SF_Compute);
 
 FVector GetLightGridZParams(float NearPlane, float FarPlane)
 {
@@ -411,7 +302,25 @@ FVector GetLightGridZParams(float NearPlane, float FarPlane)
 	return FVector(B, O, S);
 }
 
-void FDeferredShadingSceneRenderer::ComputeLightGrid(FRHICommandListImmediate& RHICmdList, bool bNeedLightGrid)
+// TODO: Might already exist as utility somewhere
+template <typename T>
+void UpdateDynamicVector4BufferData(const TArray<T, SceneRenderingAllocator> &DataArray, FDynamicReadBuffer &Buffer)
+{
+	const uint32 NumBytesRequired = DataArray.Num() * DataArray.GetTypeSize();
+
+	if (Buffer.NumBytes < NumBytesRequired)
+	{
+		Buffer.Release();
+		Buffer.Initialize(sizeof(FVector4), NumBytesRequired / sizeof(FVector4), PF_A32B32G32R32F, BUF_Volatile);
+	}
+
+	Buffer.Lock();
+	FPlatformMemory::Memcpy(Buffer.MappedBuffer, DataArray.GetData(), DataArray.Num() * DataArray.GetTypeSize());
+	Buffer.Unlock();
+}
+
+
+void FDeferredShadingSceneRenderer::ComputeLightGrid(FRHICommandListImmediate& RHICmdList, bool bNeedLightGrid, FSortedLightSetSceneInfo &SortedLightSet)
 {
 	if (!bNeedLightGrid || FeatureLevel < ERHIFeatureLevel::SM5)
 	{
@@ -439,14 +348,12 @@ void FDeferredShadingSceneRenderer::ComputeLightGrid(FRHICommandListImmediate& R
 			bAnyViewUsesForwardLighting |= View.bTranslucentSurfaceLighting || ShouldRenderVolumetricFog();
 		}
 
-		const bool bCullLightsToGrid = GLightCullingQuality && (ViewFamily.EngineShowFlags.DirectLighting && (IsForwardShadingEnabled(ShaderPlatform) || bAnyViewUsesForwardLighting || IsRayTracingEnabled()));
+		const bool bCullLightsToGrid = GLightCullingQuality 
+			&& (ViewFamily.EngineShowFlags.DirectLighting 
+			&& (IsForwardShadingEnabled(ShaderPlatform) || bAnyViewUsesForwardLighting || IsRayTracingEnabled() || ShouldUseClusteredDeferredShading()));
 			   
-		FSimpleLightArray SimpleLights;
-
-		if (bCullLightsToGrid)
-		{
-			GatherSimpleLights(ViewFamily, Views, SimpleLights);
-		}
+		// Store this flag if lights are injected in the grids, check with 'AreClusteredLightsInLightGrid()'
+		bClusteredShadingLightsInLightGrid = bCullLightsToGrid;
 
 		for (int32 ViewIndex = 0; ViewIndex < Views.Num(); ViewIndex++)
 		{
@@ -455,21 +362,80 @@ void FDeferredShadingSceneRenderer::ComputeLightGrid(FRHICommandListImmediate& R
 			ForwardLightData = FForwardLightData();
 
 			TArray<FForwardLocalLightData, SceneRenderingAllocator> ForwardLocalLightData;
+#if ENABLE_LIGHT_CULLING_VIEW_SPACE_BUILD_DATA
+			// TODO: In the future these should be generated by CPU-side render graph passes that fill the GPU buffers only when actually needed.
+			// TODO: Should use 'Empty' to pre-allocate enough space for all lights?
+			TArray<FVector4, SceneRenderingAllocator> ViewSpacePosAndRadiusData;
+			TArray<FVector4, SceneRenderingAllocator> ViewSpaceDirAndPreprocAngleData;
+#endif // ENABLE_LIGHT_CULLING_VIEW_SPACE_BUILD_DATA
 			float FurthestLight = 1000;
 
+			// Track the end markers for different types
+			int32 SimpleLightsEnd = 0;
+			int32 ClusteredSupportedEnd = 0;
+
 			if (bCullLightsToGrid)
 			{
-				ForwardLocalLightData.Empty(Scene->Lights.Num() + SimpleLights.InstanceData.Num());
+				// Simple lights are copied without view dependent checks, so same in and out
+				SimpleLightsEnd = SortedLightSet.SimpleLightsEnd;
+				// 1. insert simple lights
+				if (SimpleLightsEnd > 0)
+				{
+					const FSimpleLightArray &SimpleLights = SortedLightSet.SimpleLights;
+
+					// Pack both values into a single float to keep float4 alignment
+					const FFloat16 SimpleLightSourceLength16f = FFloat16(0);
+					FLightingChannels SimpleLightLightingChannels;
+					// Put simple lights in all lighting channels
+					SimpleLightLightingChannels.bChannel0 = SimpleLightLightingChannels.bChannel1 = SimpleLightLightingChannels.bChannel2 = true;
+					const uint32 SimpleLightLightingChannelMask = GetLightingChannelMaskForStruct(SimpleLightLightingChannels);
+
+					// Now using the sorted lights, and keep track of ranges as we go.
+					for (int SortedIndex = 0; SortedIndex < SortedLightSet.SimpleLightsEnd; ++SortedIndex)
+					{
+						check(SortedLightSet.SortedLights[SortedIndex].LightSceneInfo == nullptr);
+						check(!SortedLightSet.SortedLights[SortedIndex].SortKey.Fields.bIsNotSimpleLight);
+
+						int32 SimpleLightIndex = SortedLightSet.SortedLights[SortedIndex].SimpleLightIndex;
+
+						ForwardLocalLightData.AddUninitialized(1);
+						FForwardLocalLightData& LightData = ForwardLocalLightData.Last();
+
+						const FSimpleLightEntry& SimpleLight = SimpleLights.InstanceData[SimpleLightIndex];
+						const FSimpleLightPerViewEntry& SimpleLightPerViewData = SimpleLights.GetViewDependentData(SimpleLightIndex, ViewIndex, Views.Num());
+						LightData.LightPositionAndInvRadius = FVector4(SimpleLightPerViewData.Position, 1.0f / FMath::Max(SimpleLight.Radius, KINDA_SMALL_NUMBER));
+						LightData.LightColorAndFalloffExponent = FVector4(SimpleLight.Color, SimpleLight.Exponent);
+
+						// No shadowmap channels for simple lights
+						uint32 ShadowMapChannelMask = 0;
+						ShadowMapChannelMask |= SimpleLightLightingChannelMask << 8;
+
+						LightData.LightDirectionAndShadowMapChannelMask = FVector4(FVector(1, 0, 0), *((float*)&ShadowMapChannelMask));
+
+						// Pack both values into a single float to keep float4 alignment
+						const FFloat16 VolumetricScatteringIntensity16f = FFloat16(SimpleLight.VolumetricScatteringIntensity);
+						const uint32 PackedWInt = ((uint32)SimpleLightSourceLength16f.Encoded) | ((uint32)VolumetricScatteringIntensity16f.Encoded << 16);
+
+						LightData.SpotAnglesAndSourceRadiusPacked = FVector4(-2, 1, 0, *(float*)&PackedWInt);
+						LightData.LightTangentAndSoftSourceRadius = FVector4(1.0f, 0.0f, 0.0f, 0.0f);
 
-				for (TSparseArray<FLightSceneInfoCompact>::TConstIterator LightIt(Scene->Lights); LightIt; ++LightIt)
+#if ENABLE_LIGHT_CULLING_VIEW_SPACE_BUILD_DATA
+						FVector4 ViewSpacePosAndRadius(View.ViewMatrices.GetViewMatrix().TransformPosition(SimpleLightPerViewData.Position), SimpleLight.Radius);
+						ViewSpacePosAndRadiusData.Add(ViewSpacePosAndRadius);
+						ViewSpaceDirAndPreprocAngleData.AddZeroed();
+#endif // ENABLE_LIGHT_CULLING_VIEW_SPACE_BUILD_DATA
+					}
+				}
+				const TArray<FSortedLightSceneInfo, SceneRenderingAllocator>& SortedLights = SortedLightSet.SortedLights;
+				ClusteredSupportedEnd = SimpleLightsEnd;
+				// Next add all the other lights, track the end index for clustered supporting lights
+				for (int SortedIndex = SimpleLightsEnd; SortedIndex < SortedLights.Num(); ++SortedIndex)
 				{
-					const FLightSceneInfoCompact& LightSceneInfoCompact = *LightIt;
-					const FLightSceneInfo* const LightSceneInfo = LightSceneInfoCompact.LightSceneInfo;
-					const FVisibleLightInfo& VisibleLightInfo = VisibleLightInfos[LightIt.GetIndex()];
+					const FSortedLightSceneInfo& SortedLightInfo = SortedLights[SortedIndex];
+					const FLightSceneInfo* const LightSceneInfo = SortedLightInfo.LightSceneInfo;
 					const FLightSceneProxy* LightProxy = LightSceneInfo->Proxy;
 
-					if (LightSceneInfo->ShouldRenderLightViewIndependent()
-						&& LightSceneInfo->ShouldRenderLight(View)
+					if (LightSceneInfo->ShouldRenderLight(View)
 						// Reflection override skips direct specular because it tends to be blindingly bright with a perfectly smooth surface
 						&& !ViewFamily.EngineShowFlags.ReflectionOverride)
 					{
@@ -496,7 +462,7 @@ void FDeferredShadingSceneRenderer::ComputeLightGrid(FRHICommandListImmediate& R
 						}
 
 						// Static shadowing uses ShadowMapChannel, dynamic shadows are packed into light attenuation using DynamicShadowMapChannel
-						uint32 ShadowMapChannelMaskPacked =
+						uint32 LightTypeAndShadowMapChannelMaskPacked =
 							(ShadowMapChannel == 0 ? 1 : 0) |
 							(ShadowMapChannel == 1 ? 2 : 0) |
 							(ShadowMapChannel == 2 ? 4 : 0) |
@@ -506,21 +472,28 @@ void FDeferredShadingSceneRenderer::ComputeLightGrid(FRHICommandListImmediate& R
 							(DynamicShadowMapChannel == 2 ? 64 : 0) |
 							(DynamicShadowMapChannel == 3 ? 128 : 0);
 
-						ShadowMapChannelMaskPacked |= LightProxy->GetLightingChannelMask() << 8;
+						LightTypeAndShadowMapChannelMaskPacked |= LightProxy->GetLightingChannelMask() << 8;
+						// pack light type in this uint32 as well
+						LightTypeAndShadowMapChannelMaskPacked |= SortedLightInfo.SortKey.Fields.LightType << 16;
 
-						if( ( LightSceneInfoCompact.LightType == LightType_Point && ViewFamily.EngineShowFlags.PointLights ) ||
-							( LightSceneInfoCompact.LightType == LightType_Spot  && ViewFamily.EngineShowFlags.SpotLights ) ||
-							( LightSceneInfoCompact.LightType == LightType_Rect  && ViewFamily.EngineShowFlags.RectLights ) )
+						if ((SortedLightInfo.SortKey.Fields.LightType == LightType_Point && ViewFamily.EngineShowFlags.PointLights) ||
+							(SortedLightInfo.SortKey.Fields.LightType == LightType_Spot && ViewFamily.EngineShowFlags.SpotLights) ||
+							(SortedLightInfo.SortKey.Fields.LightType == LightType_Rect && ViewFamily.EngineShowFlags.RectLights))
 						{
 							ForwardLocalLightData.AddUninitialized(1);
 							FForwardLocalLightData& LightData = ForwardLocalLightData.Last();
 
+							// Track the last one to support clustered deferred
+							if (!SortedLightInfo.SortKey.Fields.bClusteredDeferredNotSupported)
+							{
+								ClusteredSupportedEnd = FMath::Max(ClusteredSupportedEnd, ForwardLocalLightData.Num());
+							}
 							const float LightFade = GetLightFadeFactor(View, LightProxy);
 							LightParameters.Color *= LightFade;
 
 							LightData.LightPositionAndInvRadius = FVector4(LightParameters.Position, LightParameters.InvRadius);
 							LightData.LightColorAndFalloffExponent = FVector4(LightParameters.Color, LightParameters.FalloffExponent);
-							LightData.LightDirectionAndShadowMapChannelMask = FVector4(LightParameters.Direction, *((float*)&ShadowMapChannelMaskPacked));
+							LightData.LightDirectionAndShadowMapChannelMask = FVector4(LightParameters.Direction, *((float*)&LightTypeAndShadowMapChannelMaskPacked));
 
 							LightData.SpotAnglesAndSourceRadiusPacked = FVector4(LightParameters.SpotAngles.X, LightParameters.SpotAngles.Y, LightParameters.SourceRadius, 0);
 
@@ -543,14 +516,25 @@ void FDeferredShadingSceneRenderer::ComputeLightGrid(FRHICommandListImmediate& R
 							const FSphere BoundingSphere = LightProxy->GetBoundingSphere();
 							const float Distance = View.ViewMatrices.GetViewMatrix().TransformPosition(BoundingSphere.Center).Z + BoundingSphere.W;
 							FurthestLight = FMath::Max(FurthestLight, Distance);
+
+#if ENABLE_LIGHT_CULLING_VIEW_SPACE_BUILD_DATA
+							// Note: inverting radius twice seems stupid (but done in shader anyway otherwise)
+							FVector4 ViewSpacePosAndRadius(View.ViewMatrices.GetViewMatrix().TransformPosition(LightParameters.Position), 1.0f / LightParameters.InvRadius);
+							ViewSpacePosAndRadiusData.Add(ViewSpacePosAndRadius);
+
+							float PreProcAngle = SortedLightInfo.SortKey.Fields.LightType == LightType_Spot ? GetTanRadAngleOrZero(LightSceneInfo->Proxy->GetOuterConeAngle()) : 0.0f;
+
+							FVector4 ViewSpaceDirAndPreprocAngle(View.ViewMatrices.GetViewMatrix().TransformVector(LightParameters.Direction), PreProcAngle);
+							ViewSpaceDirAndPreprocAngleData.Add(ViewSpaceDirAndPreprocAngle);
+#endif // ENABLE_LIGHT_CULLING_VIEW_SPACE_BUILD_DATA
 						}
-						else if (LightSceneInfoCompact.LightType == LightType_Directional && ViewFamily.EngineShowFlags.DirectionalLights)
+						else if (SortedLightInfo.SortKey.Fields.LightType == LightType_Directional && ViewFamily.EngineShowFlags.DirectionalLights)
 						{
 							ForwardLightData.HasDirectionalLight = 1;
 							ForwardLightData.DirectionalLightColor = LightParameters.Color;
 							ForwardLightData.DirectionalLightVolumetricScatteringIntensity = LightProxy->GetVolumetricScatteringIntensity();
 							ForwardLightData.DirectionalLightDirection = LightParameters.Direction;
-							ForwardLightData.DirectionalLightShadowMapChannelMask = ShadowMapChannelMaskPacked;
+							ForwardLightData.DirectionalLightShadowMapChannelMask = LightTypeAndShadowMapChannelMaskPacked;
 
 							const FVector2D FadeParams = LightProxy->GetDirectionalLightDistanceFadeParameters(View.GetFeatureLevel(), LightSceneInfo->IsPrecomputedLightingValid(), View.MaxShadowCascades);
 
@@ -594,39 +578,9 @@ void FDeferredShadingSceneRenderer::ComputeLightGrid(FRHICommandListImmediate& R
 						}
 					}
 				}
-
-				// Pack both values into a single float to keep float4 alignment
-				const FFloat16 SimpleLightSourceLength16f = FFloat16(0);
-				FLightingChannels SimpleLightLightingChannels;
-				// Put simple lights in all lighting channels
-				SimpleLightLightingChannels.bChannel0 = SimpleLightLightingChannels.bChannel1 = SimpleLightLightingChannels.bChannel2 = true;
-				const uint32 SimpleLightLightingChannelMask = GetLightingChannelMaskForStruct(SimpleLightLightingChannels);
-
-				for (int32 SimpleLightIndex = 0; SimpleLightIndex < SimpleLights.InstanceData.Num(); SimpleLightIndex++)
-				{
-					ForwardLocalLightData.AddUninitialized(1);
-					FForwardLocalLightData& LightData = ForwardLocalLightData.Last();
-
-					const FSimpleLightEntry& SimpleLight = SimpleLights.InstanceData[SimpleLightIndex];
-					const FSimpleLightPerViewEntry& SimpleLightPerViewData = SimpleLights.GetViewDependentData(SimpleLightIndex, ViewIndex, Views.Num());
-					LightData.LightPositionAndInvRadius = FVector4(SimpleLightPerViewData.Position, 1.0f / FMath::Max(SimpleLight.Radius, KINDA_SMALL_NUMBER));
-					LightData.LightColorAndFalloffExponent = FVector4(SimpleLight.Color, SimpleLight.Exponent);
-
-					// No shadowmap channels for simple lights
-					uint32 ShadowMapChannelMask = 0;
-					ShadowMapChannelMask |= SimpleLightLightingChannelMask << 8;
-
-					LightData.LightDirectionAndShadowMapChannelMask = FVector4(FVector(1, 0, 0), *((float*)&ShadowMapChannelMask));
-
-					// Pack both values into a single float to keep float4 alignment
-					const FFloat16 VolumetricScatteringIntensity16f = FFloat16(SimpleLight.VolumetricScatteringIntensity);
-					const uint32 PackedWInt = ((uint32)SimpleLightSourceLength16f.Encoded) | ((uint32)VolumetricScatteringIntensity16f.Encoded << 16);
-
-					LightData.SpotAnglesAndSourceRadiusPacked = FVector4(-2, 1, 0, *(float*)&PackedWInt);
-					LightData.LightTangentAndSoftSourceRadius = FVector4(1.0f, 0.0f, 0.0f, 0.0f);
-				}
 			}
 
+
 			// Store off the number of lights before we add a fake entry
 			const int32 NumLocalLightsFinal = ForwardLocalLightData.Num();
 
@@ -636,28 +590,18 @@ void FDeferredShadingSceneRenderer::ComputeLightGrid(FRHICommandListImmediate& R
 				ForwardLocalLightData.AddZeroed();
 			}
 
-			{
-				const uint32 NumBytesRequired = ForwardLocalLightData.Num() * ForwardLocalLightData.GetTypeSize();
-
-				if (View.ForwardLightingResources->ForwardLocalLightBuffer.NumBytes < NumBytesRequired)
-				{
-					View.ForwardLightingResources->ForwardLocalLightBuffer.Release();
-					View.ForwardLightingResources->ForwardLocalLightBuffer.Initialize(sizeof(FVector4), NumBytesRequired / sizeof(FVector4), PF_A32B32G32R32F, BUF_Volatile);
-				}
-
-				ForwardLightData.ForwardLocalLightBuffer = View.ForwardLightingResources->ForwardLocalLightBuffer.SRV;
-				View.ForwardLightingResources->ForwardLocalLightBuffer.Lock();
-				FPlatformMemory::Memcpy(View.ForwardLightingResources->ForwardLocalLightBuffer.MappedBuffer, ForwardLocalLightData.GetData(), ForwardLocalLightData.Num() * ForwardLocalLightData.GetTypeSize());
-				View.ForwardLightingResources->ForwardLocalLightBuffer.Unlock();
-			}
+			UpdateDynamicVector4BufferData(ForwardLocalLightData, View.ForwardLightingResources->ForwardLocalLightBuffer);
 
 			const FIntPoint LightGridSizeXY = FIntPoint::DivideAndRoundUp(View.ViewRect.Size(), GLightGridPixelSize);
+			ForwardLightData.ForwardLocalLightBuffer = View.ForwardLightingResources->ForwardLocalLightBuffer.SRV;
 			ForwardLightData.NumLocalLights = NumLocalLightsFinal;
 			ForwardLightData.NumReflectionCaptures = View.NumBoxReflectionCaptures + View.NumSphereReflectionCaptures;
 			ForwardLightData.NumGridCells = LightGridSizeXY.X * LightGridSizeXY.Y * GLightGridSizeZ;
 			ForwardLightData.CulledGridSize = FIntVector(LightGridSizeXY.X, LightGridSizeXY.Y, GLightGridSizeZ);
 			ForwardLightData.MaxCulledLightsPerCell = GMaxCulledLightsPerCell;
 			ForwardLightData.LightGridPixelSizeShift = FMath::FloorLog2(GLightGridPixelSize);
+			ForwardLightData.SimpleLightsEndIndex = SimpleLightsEnd;
+			ForwardLightData.ClusteredDeferredSupportedEndIndex = ClusteredSupportedEnd;
 
 			// Clamp far plane to something reasonable
 			float FarPlane = FMath::Min(FMath::Max(FurthestLight, View.FurthestReflectionCaptureDistance), (float)HALF_WORLD_MAX / 5.0f);
@@ -676,6 +620,12 @@ void FDeferredShadingSceneRenderer::ComputeLightGrid(FRHICommandListImmediate& R
 					bWarned = true;
 				}
 			}
+
+#if ENABLE_LIGHT_CULLING_VIEW_SPACE_BUILD_DATA
+			const SIZE_T LightIndexTypeSize = CHANGE_LIGHTINDEXTYPE_SIZE && !bAllowFormatConversion ? sizeof(FLightIndexType32) : sizeof(FLightIndexType);
+			// Fuse these loops as I see no reason why not and we build some temporary data that is needed in the build pass and is 
+			// not needed to be stored permanently.
+#else // !ENABLE_LIGHT_CULLING_VIEW_SPACE_BUILD_DATA
 		}
 
 		const SIZE_T LightIndexTypeSize = CHANGE_LIGHTINDEXTYPE_SIZE && !bAllowFormatConversion ? sizeof(FLightIndexType32) : sizeof(FLightIndexType);
@@ -684,8 +634,9 @@ void FDeferredShadingSceneRenderer::ComputeLightGrid(FRHICommandListImmediate& R
 		{
 			FViewInfo& View = Views[ViewIndex];
 			FForwardLightData& ForwardLightData = View.ForwardLightingResources->ForwardLightData;
-
 			const FIntPoint LightGridSizeXY = FIntPoint::DivideAndRoundUp(View.ViewRect.Size(), GLightGridPixelSize);
+#endif // ENABLE_LIGHT_CULLING_VIEW_SPACE_BUILD_DATA
+
 			const int32 NumCells = LightGridSizeXY.X * LightGridSizeXY.Y * GLightGridSizeZ * NumCulledGridPrimitiveTypes;
 
 			if (View.ForwardLightingResources->NumCulledLightsGrid.NumBytes != NumCells * NumCulledLightsGridStride * sizeof(uint32))
@@ -711,36 +662,36 @@ void FDeferredShadingSceneRenderer::ComputeLightGrid(FRHICommandListImmediate& R
 			FForwardLightingCullingResources& ForwardLightingCullingResources = bShouldCacheTemporaryBuffers ? View.ViewState->ForwardLightingCullingResources : LocalCullingResources;
 
 			const uint32 CulledLightLinksElements = NumCells * GMaxCulledLightsPerCell * LightLinkStride;
-			if (ForwardLightingCullingResources.CulledLightLinks.NumBytes != (CulledLightLinksElements * sizeof(uint32))
-				|| (GFastVRamConfig.bDirty && ForwardLightingCullingResources.CulledLightLinks.NumBytes > 0))
+
+#if ENABLE_LIGHT_CULLING_VIEW_SPACE_BUILD_DATA
+			if (ViewSpacePosAndRadiusData.Num() == 0)
 			{
-				UE_CLOG(CulledLightLinksElements * sizeof(uint32) > 256llu * (1llu << 20llu), LogRenderer, Warning,
-					TEXT("Attempt to allocate large FRWBuffer (not supported by Metal): ForwardLightingCullingResources.CulledLightLinks %u Bytes, LightGridSize %dx%dx%d, NumCulledGridPrimitiveTypes %d, NumCells %d, GMaxCulledLightsPerCell %d, LightLinkStride %d, View Resolution %dx%d"),
-					CulledLightLinksElements * sizeof(uint32), LightGridSizeXY.X, LightGridSizeXY.Y, GLightGridSizeZ, NumCulledGridPrimitiveTypes, NumCells, GMaxCulledLightsPerCell, LightLinkStride, View.ViewRect.Size().X, View.ViewRect.Size().Y);
-
-				const uint32 FastVRamFlag = GFastVRamConfig.ForwardLightingCullingResources | (IsTransientResourceBufferAliasingEnabled() ? BUF_Transient : BUF_None);
-				ForwardLightingCullingResources.CulledLightLinks.Initialize(sizeof(uint32), CulledLightLinksElements, PF_R32_UINT, FastVRamFlag, TEXT("CulledLightLinks"));
-				ForwardLightingCullingResources.NextCulledLightLink.Initialize(sizeof(uint32), 1, PF_R32_UINT, FastVRamFlag, TEXT("NextCulledLightLink"));
-				ForwardLightingCullingResources.StartOffsetGrid.Initialize(sizeof(uint32), NumCells, PF_R32_UINT, FastVRamFlag, TEXT("StartOffsetGrid"));
-				ForwardLightingCullingResources.NextCulledLightData.Initialize(sizeof(uint32), 1, PF_R32_UINT, FastVRamFlag, TEXT("NextCulledLightData"));
+				// Make sure the buffer gets created even though we're not going to read from it in the shader, for platforms like PS4 that assert on null resources being bound
+				ViewSpacePosAndRadiusData.AddZeroed();
+				ViewSpaceDirAndPreprocAngleData.AddZeroed();
 			}
+			check(ViewSpacePosAndRadiusData.Num() == ForwardLocalLightData.Num());
+			check(ViewSpaceDirAndPreprocAngleData.Num() == ForwardLocalLightData.Num());
 
+			UpdateDynamicVector4BufferData(ViewSpacePosAndRadiusData, ForwardLightingCullingResources.ViewSpacePosAndRadiusData);
+			UpdateDynamicVector4BufferData(ViewSpaceDirAndPreprocAngleData, ForwardLightingCullingResources.ViewSpaceDirAndPreprocAngleData);
+#endif // ENABLE_LIGHT_CULLING_VIEW_SPACE_BUILD_DATA
+
+			// Used to pass to the GetDynamicLighting but not actually used, since USE_SOURCE_TEXTURE is 0
+			ForwardLightData.DummyRectLightSourceTexture = GWhiteTexture->TextureRHI;
 			ForwardLightData.NumCulledLightsGrid = View.ForwardLightingResources->NumCulledLightsGrid.SRV;
 			ForwardLightData.CulledLightDataGrid = View.ForwardLightingResources->CulledLightDataGrid.SRV;
 
 			View.ForwardLightingResources->ForwardLightDataUniformBuffer = TUniformBufferRef<FForwardLightData>::CreateUniformBufferImmediate(ForwardLightData, UniformBuffer_SingleFrame);
 
-			if (IsTransientResourceBufferAliasingEnabled())
-			{
-				// Acquire resources
-				ForwardLightingCullingResources.CulledLightLinks.AcquireTransientResource();
-				ForwardLightingCullingResources.NextCulledLightLink.AcquireTransientResource();
-				ForwardLightingCullingResources.StartOffsetGrid.AcquireTransientResource();
-				ForwardLightingCullingResources.NextCulledLightData.AcquireTransientResource();
-			}
-
 			const FIntVector NumGroups = FIntVector::DivideAndRoundUp(FIntVector(LightGridSizeXY.X, LightGridSizeXY.Y, GLightGridSizeZ), LightGridInjectionGroupSize);
 
+			// TODO: When external buffers and extract can be set up maybe use that to get these out instead.
+			TArray<FUnorderedAccessViewRHIParamRef, TInlineAllocator<2>> OutUAVs({
+				View.ForwardLightingResources->NumCulledLightsGrid.UAV,
+				View.ForwardLightingResources->CulledLightDataGrid.UAV });
+
+			RHICmdList.TransitionResources(EResourceTransitionAccess::EWritable, EResourceTransitionPipeline::EGfxToCompute, OutUAVs.GetData(), OutUAVs.Num());
 			{
 				SCOPED_DRAW_EVENTF(RHICmdList, CullLights, TEXT("CullLights %ux%ux%u NumLights %u NumCaptures %u"),
 					ForwardLightData.CulledGridSize.X,
@@ -749,60 +700,92 @@ void FDeferredShadingSceneRenderer::ComputeLightGrid(FRHICommandListImmediate& R
 					ForwardLightData.NumLocalLights,
 					ForwardLightData.NumReflectionCaptures);
 
-				TArray<FUnorderedAccessViewRHIParamRef, TInlineAllocator<6>> OutUAVs;
-				OutUAVs.Add(View.ForwardLightingResources->NumCulledLightsGrid.UAV);
-				OutUAVs.Add(View.ForwardLightingResources->CulledLightDataGrid.UAV);
-				OutUAVs.Add(ForwardLightingCullingResources.NextCulledLightLink.UAV);
-				OutUAVs.Add(ForwardLightingCullingResources.StartOffsetGrid.UAV);
-				OutUAVs.Add(ForwardLightingCullingResources.CulledLightLinks.UAV);
-				OutUAVs.Add(ForwardLightingCullingResources.NextCulledLightData.UAV);
-				RHICmdList.TransitionResources(EResourceTransitionAccess::EWritable, EResourceTransitionPipeline::EGfxToCompute, OutUAVs.GetData(), OutUAVs.Num());
-
-				if (GLightLinkedListCulling)
+				FRDGBuilder GraphBuilder(RHICmdList);
 				{
-					ClearUAV(RHICmdList, ForwardLightingCullingResources.StartOffsetGrid, 0xFFFFFFFF);
-					ClearUAV(RHICmdList, ForwardLightingCullingResources.NextCulledLightLink, 0);
-					ClearUAV(RHICmdList, ForwardLightingCullingResources.NextCulledLightData, 0);
-
-					TShaderMapRef<TLightGridInjectionCS<true> > ComputeShader(View.ShaderMap);
-					RHICmdList.SetComputeShader(ComputeShader->GetComputeShader());
-					ComputeShader->SetParameters(RHICmdList, View, ForwardLightingCullingResources);
-					DispatchComputeShader(RHICmdList, *ComputeShader, NumGroups.X, NumGroups.Y, NumGroups.Z);
-					ComputeShader->UnsetParameters(RHICmdList, View, ForwardLightingCullingResources);
-				}
-				else
-				{
-					ClearUAV(RHICmdList, View.ForwardLightingResources->NumCulledLightsGrid, 0);
+					RDG_EVENT_SCOPE(GraphBuilder, "CullLights %ux%ux%u NumLights %u NumCaptures %u",
+						ForwardLightData.CulledGridSize.X,
+						ForwardLightData.CulledGridSize.Y,
+						ForwardLightData.CulledGridSize.Z,
+						ForwardLightData.NumLocalLights,
+						ForwardLightData.NumReflectionCaptures);
+
+					// TODO: change to render graph?
+					//FRDGBufferRef NumCulledLightsGridBuffer = GraphBuilder.CreateBuffer(FRDGBufferDesc::CreateStructuredDesc(sizeof(uint32), NumCells * NumCulledLightsGridStride), TEXT("NumCulledLightsGrid"));
+					//FRDGBufferRef CulledLightDataGridBuffer = GraphBuilder.CreateBuffer(FRDGBufferDesc::CreateStructuredDesc(LightIndexTypeSize, NumCells * GMaxCulledLightsPerCell), TEXT("CulledLightDataGrid"));
+
+					FRDGBufferRef CulledLightLinksBuffer = GraphBuilder.CreateBuffer(FRDGBufferDesc::CreateStructuredDesc(sizeof(uint32), CulledLightLinksElements), TEXT("CulledLightLinks"));
+					FRDGBufferRef StartOffsetGridBuffer = GraphBuilder.CreateBuffer(FRDGBufferDesc::CreateStructuredDesc(sizeof(uint32), NumCells), TEXT("StartOffsetGrid"));
+					FRDGBufferRef NextCulledLightLinkBuffer = GraphBuilder.CreateBuffer(FRDGBufferDesc::CreateStructuredDesc(sizeof(uint32), 1), TEXT("NextCulledLightLink"));
+					FRDGBufferRef NextCulledLightDataBuffer = GraphBuilder.CreateBuffer(FRDGBufferDesc::CreateStructuredDesc(sizeof(uint32), 1), TEXT("NextCulledLightData"));
+
+#if ENABLE_LIGHT_CULLING_VIEW_SPACE_BUILD_DATA
+					// TODO: When CPU-side generative passes are added this can also use render graph
+					//FRDGBufferRef ViewSpacePosAndRadiusDataBuffer = GraphBuilder.CreateBuffer(FRDGBufferDesc::CreateStructuredDesc(sizeof(FVector4), ViewSpacePosAndRadiusData.Num()), TEXT("ViewSpacePosAndRadiusData"));
+					//FRDGBufferRef ViewSpaceDirAndPreprocAngleDataBuffer = GraphBuilder.CreateBuffer(FRDGBufferDesc::CreateStructuredDesc(sizeof(FVector4), ViewSpaceDirAndPreprocAngleData.Num()), TEXT("ViewSpacePosAndRadiusData"));
+#endif // ENABLE_LIGHT_CULLING_VIEW_SPACE_BUILD_DATA
+
+					FLightGridInjectionCS::FParameters *PassParameters = GraphBuilder.AllocParameters<FLightGridInjectionCS::FParameters>();
+
+					PassParameters->View = View.ViewUniformBuffer;
+					PassParameters->ReflectionCapture = View.ReflectionCaptureUniformBuffer;
+					PassParameters->Forward = View.ForwardLightingResources->ForwardLightDataUniformBuffer;
+					PassParameters->RWNumCulledLightsGrid = View.ForwardLightingResources->NumCulledLightsGrid.UAV;
+					PassParameters->RWCulledLightDataGrid = View.ForwardLightingResources->CulledLightDataGrid.UAV;
+					PassParameters->RWNextCulledLightLink = GraphBuilder.CreateUAV(NextCulledLightLinkBuffer);
+					PassParameters->RWStartOffsetGrid = GraphBuilder.CreateUAV(StartOffsetGridBuffer);
+					PassParameters->RWCulledLightLinks = GraphBuilder.CreateUAV(CulledLightLinksBuffer);
+
+#if ENABLE_LIGHT_CULLING_VIEW_SPACE_BUILD_DATA
+					// TODO: When CPU-side generative passes are added this can also use render graph
+					//PassParameters->LightViewSpacePositionAndRadius = GraphBuilder.CreateSRV(ViewSpacePosAndRadiusDataBuffer);
+					//PassParameters->LightViewSpaceDirAndPreprocAngle = GraphBuilder.CreateSRV(ViewSpaceDirAndPreprocAngleDataBuffer);
+					PassParameters->LightViewSpacePositionAndRadius = ForwardLightingCullingResources.ViewSpacePosAndRadiusData.SRV;
+					PassParameters->LightViewSpaceDirAndPreprocAngle = ForwardLightingCullingResources.ViewSpaceDirAndPreprocAngleData.SRV;
+#endif // ENABLE_LIGHT_CULLING_VIEW_SPACE_BUILD_DATA
+
+					FLightGridInjectionCS::FPermutationDomain PermutationVector;
+					PermutationVector.Set<FLightGridInjectionCS::FUseLinkedListDim>(GLightLinkedListCulling);
+					TShaderMapRef<FLightGridInjectionCS> ComputeShader(View.ShaderMap, PermutationVector);
+
+					if (GLightLinkedListCulling)
+					{
+						AddPass_ClearUAV(GraphBuilder, RDG_EVENT_NAME("Clear:StartOffsetGrid"), PassParameters->RWStartOffsetGrid, 0xFFFFFFFF);
+						AddPass_ClearUAV(GraphBuilder, RDG_EVENT_NAME("Clear:NextCulledLightLink"), PassParameters->RWNextCulledLightLink, 0);
+						AddPass_ClearUAV(GraphBuilder, RDG_EVENT_NAME("Clear:NextCulledLightData"), GraphBuilder.CreateUAV(NextCulledLightDataBuffer), 0);
+						FComputeShaderUtils::AddPass(GraphBuilder, RDG_EVENT_NAME("LightGridInject:LinkedList"), *ComputeShader, PassParameters, NumGroups);
 
-					TShaderMapRef<TLightGridInjectionCS<false> > ComputeShader(View.ShaderMap);
-					RHICmdList.SetComputeShader(ComputeShader->GetComputeShader());
-					ComputeShader->SetParameters(RHICmdList, View, ForwardLightingCullingResources);
-					DispatchComputeShader(RHICmdList, *ComputeShader, NumGroups.X, NumGroups.Y, NumGroups.Z);
-					ComputeShader->UnsetParameters(RHICmdList, View, ForwardLightingCullingResources);
-				}
-			}
 
-			if (GLightLinkedListCulling)
-			{
-				SCOPED_DRAW_EVENT(RHICmdList, Compact);
+						{
+							TShaderMapRef<FLightGridCompactCS> ComputeShaderCompact(View.ShaderMap);
+							FLightGridCompactCS::FParameters *PassParametersCompact = GraphBuilder.AllocParameters<FLightGridCompactCS::FParameters>();
+							PassParametersCompact->View = View.ViewUniformBuffer;
+							PassParametersCompact->Forward = View.ForwardLightingResources->ForwardLightDataUniformBuffer;
+
+							PassParametersCompact->CulledLightLinks = GraphBuilder.CreateSRV(CulledLightLinksBuffer);
+							PassParametersCompact->RWNumCulledLightsGrid = View.ForwardLightingResources->NumCulledLightsGrid.UAV;
+							PassParametersCompact->RWCulledLightDataGrid = View.ForwardLightingResources->CulledLightDataGrid.UAV;
+							PassParametersCompact->RWNextCulledLightData = GraphBuilder.CreateUAV(NextCulledLightDataBuffer);
+							PassParametersCompact->StartOffsetGrid = GraphBuilder.CreateSRV(StartOffsetGridBuffer);;
+
+							FComputeShaderUtils::AddPass(GraphBuilder, RDG_EVENT_NAME("CompactLinks"), *ComputeShaderCompact, PassParametersCompact, NumGroups);
+						}
+					}
+					else
+					{
+						ClearUAV(RHICmdList, View.ForwardLightingResources->NumCulledLightsGrid, 0);
+						FComputeShaderUtils::AddPass(GraphBuilder, RDG_EVENT_NAME("LightGridInject:NotLinkedList"), *ComputeShader, PassParameters, NumGroups);
+					}
+				}
+				GraphBuilder.Execute();
 
-				TShaderMapRef<FLightGridCompactCS> ComputeShader(View.ShaderMap);
-				RHICmdList.SetComputeShader(ComputeShader->GetComputeShader());
-				ComputeShader->SetParameters(RHICmdList, View, ForwardLightingCullingResources);
-				DispatchComputeShader(RHICmdList, *ComputeShader, NumGroups.X, NumGroups.Y, NumGroups.Z);
-				ComputeShader->UnsetParameters(RHICmdList, View, ForwardLightingCullingResources);
-			}
-			if (IsTransientResourceBufferAliasingEnabled())
-			{
-				ForwardLightingCullingResources.CulledLightLinks.DiscardTransientResource();
-				ForwardLightingCullingResources.NextCulledLightLink.DiscardTransientResource();
-				ForwardLightingCullingResources.StartOffsetGrid.DiscardTransientResource();
-				ForwardLightingCullingResources.NextCulledLightData.DiscardTransientResource();
+				RHICmdList.TransitionResources(EResourceTransitionAccess::EReadable, EResourceTransitionPipeline::EComputeToGfx, OutUAVs.GetData(), OutUAVs.Num());
 			}
 		}
 	}
 }
 
+
+
 void FDeferredShadingSceneRenderer::RenderForwardShadingShadowProjections(FRHICommandListImmediate& RHICmdList, TRefCountPtr<IPooledRenderTarget>& ForwardScreenSpaceShadowMask)
 {
 	check(RHICmdList.IsOutsideRenderPass());
diff --git a/Engine/Source/Runtime/Renderer/Private/LightRendering.cpp b/Engine/Source/Runtime/Renderer/Private/LightRendering.cpp
--- a/Engine/Source/Runtime/Renderer/Private/LightRendering.cpp
+++ b/Engine/Source/Runtime/Renderer/Private/LightRendering.cpp
@@ -5,6 +5,7 @@
 =============================================================================*/
 
 #include "LightRendering.h"
+#include "RendererModule.h"
 #include "DeferredShadingRenderer.h"
 #include "LightPropagationVolume.h"
 #include "ScenePrivate.h"
@@ -17,12 +18,20 @@
 #include "RayTracing/RaytracingOptions.h"
 #include "SceneViewFamilyBlackboard.h"
 
+// ENABLE_DEBUG_DISCARD_PROP is used to test the lighting code by allowing to discard lights to see how performance scales
+// It ought never to be enabled in a shipping build, and is probably only really useful when woring on the shading code.
+#if !(UE_BUILD_SHIPPING || UE_BUILD_TEST)
+	#define ENABLE_DEBUG_DISCARD_PROP 1
+#else // (UE_BUILD_SHIPPING || UE_BUILD_TEST)
+	#define ENABLE_DEBUG_DISCARD_PROP 0
+#endif // !(UE_BUILD_SHIPPING || UE_BUILD_TEST)
+
 DECLARE_GPU_STAT(Lights);
 
 IMPLEMENT_GLOBAL_SHADER_PARAMETER_STRUCT(FDeferredLightUniformStruct, "DeferredLightUniforms");
 
 extern int32 GUseTranslucentLightingVolumes;
-ENGINE_API const IPooledRenderTarget* GetSubsufaceProfileTexture_RT(FRHICommandListImmediate& RHICmdList);
+ENGINE_API IPooledRenderTarget* GetSubsufaceProfileTexture_RT(FRHICommandListImmediate& RHICmdList);
 
 
 static int32 GAllowDepthBoundsTest = 1;
@@ -67,6 +76,17 @@ static TAutoConsoleVariable<int32> CVarMaxShadowDenoisingBatchSize(
 	TEXT("Maximum number of shadow to denoise at the same time."),
 	ECVF_RenderThreadSafe);
 
+#if ENABLE_DEBUG_DISCARD_PROP
+static float GDebugLightDiscardProp = 0.0f;
+static FAutoConsoleVariableRef CVarDebugLightDiscardProp(
+	TEXT("r.DebugLightDiscardProp"),
+	GDebugLightDiscardProp,
+	TEXT("[0,1]: Proportion of lights to discard for debug/performance profiling purposes.")
+);
+#endif // ENABLE_DEBUG_DISCARD_PROP
+
+
+
 FLightOcclusionType GetLightOcclusionType(const FLightSceneProxy& Proxy)
 {
 #if RHI_RAYTRACING
@@ -523,38 +543,44 @@ static bool LightRequiresDenosier(const FLightSceneInfo& LightSceneInfo)
 }
 
 
-/** Renders the scene's lighting. */
-void FDeferredShadingSceneRenderer::RenderLights(FRHICommandListImmediate& RHICmdList)
-{
-	check(RHICmdList.IsOutsideRenderPass());
-
-	SCOPED_NAMED_EVENT(FDeferredShadingSceneRenderer_RenderLights, FColor::Emerald);
-	SCOPED_DRAW_EVENT(RHICmdList, DirectLighting);
-	SCOPED_GPU_STAT(RHICmdList, Lights);
-
-
-	bool bStencilBufferDirty = false;	// The stencil buffer should've been cleared to 0 already
-
-	SCOPE_CYCLE_COUNTER(STAT_LightingDrawTime);
-	SCOPE_CYCLE_COUNTER(STAT_LightRendering);
 
-	FSimpleLightArray SimpleLights;
+void FDeferredShadingSceneRenderer::GatherAndSortLights(FSortedLightSetSceneInfo& OutSortedLights)
+{
 	if (bAllowSimpleLights)
 	{
-		GatherSimpleLights(ViewFamily, Views, SimpleLights);
+		GatherSimpleLights(ViewFamily, Views, OutSortedLights.SimpleLights);
 	}
+	FSimpleLightArray &SimpleLights = OutSortedLights.SimpleLights;
+	TArray<FSortedLightSceneInfo, SceneRenderingAllocator> &SortedLights = OutSortedLights.SortedLights;
 
-	TArray<FSortedLightSceneInfo, SceneRenderingAllocator> SortedLights;
-	SortedLights.Empty(Scene->Lights.Num());
+	// NOTE: we allocate space also for simple lights such that they can be referenced in the same sorted range
+	SortedLights.Empty(Scene->Lights.Num() + SimpleLights.InstanceData.Num());
 
 	bool bDynamicShadows = ViewFamily.EngineShowFlags.DynamicShadows && GetShadowQuality() > 0;
-	
+
+#if ENABLE_DEBUG_DISCARD_PROP
+	int Total = Scene->Lights.Num() + SimpleLights.InstanceData.Num();
+	int NumToKeep = int(float(Total) * (1.0f - GDebugLightDiscardProp));
+	const float DebugDiscardStride = float(NumToKeep) / float(Total);
+	float DebugDiscardCounter = 0.0f;
+#endif // ENABLE_DEBUG_DISCARD_PROP
 	// Build a list of visible lights.
 	for (TSparseArray<FLightSceneInfoCompact>::TConstIterator LightIt(Scene->Lights); LightIt; ++LightIt)
 	{
 		const FLightSceneInfoCompact& LightSceneInfoCompact = *LightIt;
 		const FLightSceneInfo* const LightSceneInfo = LightSceneInfoCompact.LightSceneInfo;
 
+#if ENABLE_DEBUG_DISCARD_PROP
+		{
+			int PrevCounter = int(DebugDiscardCounter);
+			DebugDiscardCounter += DebugDiscardStride;
+			if (PrevCounter >= int(DebugDiscardCounter))
+			{
+				continue;
+			}
+		}
+#endif // ENABLE_DEBUG_DISCARD_PROP
+
 		if (LightSceneInfo->ShouldRenderLightViewIndependent()
 			// Reflection override skips direct specular because it tends to be blindingly bright with a perfectly smooth surface
 			&& !ViewFamily.EngineShowFlags.ReflectionOverride)
@@ -573,18 +599,59 @@ void FDeferredShadingSceneRenderer::RenderLights(FRHICommandListImmediate& RHICm
 					SortedLightInfo->SortKey.Fields.bLightFunction = ViewFamily.EngineShowFlags.LightFunctions && CheckForLightFunction(LightSceneInfo);
 					SortedLightInfo->SortKey.Fields.bUsesLightingChannels = Views[ViewIndex].bUsesLightingChannels && LightSceneInfo->Proxy->GetLightingChannelMask() != GetDefaultLightingChannelMask();
 
-					// tiled deferred lighting only supported for certain lights that don't use any additional features
-					const bool bTiledDeferredSupported = LightSceneInfo->Proxy->IsTiledDeferredLightingSupported() &&
+					// These are not simple lights.
+					SortedLightInfo->SortKey.Fields.bIsNotSimpleLight = 1;
+
+
+					// tiled and clustered deferred lighting only supported for certain lights that don't use any additional features
+					// And also that are not directional (mostly because it does'nt make so much sense to insert them into every grid cell in the universe)
+					// In the forward case one directional light gets put into its own variables, and in the deferred case it gets a full-screen pass.
+					// Usually it'll have shadows and stuff anyway.
+					// Rect lights are not supported as the performance impact is significant even if not used, for now, left for trad. deferred.
+					const bool bTiledOrClusteredDeferredSupported = 
 						!SortedLightInfo->SortKey.Fields.bTextureProfile &&
 						!SortedLightInfo->SortKey.Fields.bShadowed &&
 						!SortedLightInfo->SortKey.Fields.bLightFunction &&
-						!SortedLightInfo->SortKey.Fields.bUsesLightingChannels;
-					SortedLightInfo->SortKey.Fields.bTiledDeferredNotSupported = !bTiledDeferredSupported;
+						!SortedLightInfo->SortKey.Fields.bUsesLightingChannels
+						&& LightSceneInfoCompact.LightType != LightType_Directional
+						&& LightSceneInfoCompact.LightType != LightType_Rect;
+
+					SortedLightInfo->SortKey.Fields.bTiledDeferredNotSupported = !(bTiledOrClusteredDeferredSupported && LightSceneInfo->Proxy->IsTiledDeferredLightingSupported());
+
+					SortedLightInfo->SortKey.Fields.bClusteredDeferredNotSupported = !bTiledOrClusteredDeferredSupported;
 					break;
 				}
 			}
 		}
 	}
+	// Add the simple lights also
+	for (int32 SimpleLightIndex = 0; SimpleLightIndex < SimpleLights.InstanceData.Num(); SimpleLightIndex++)
+	{
+#if ENABLE_DEBUG_DISCARD_PROP
+		{
+			int PrevCounter = int(DebugDiscardCounter);
+			DebugDiscardCounter += DebugDiscardStride;
+			if (PrevCounter >= int(DebugDiscardCounter))
+			{
+				continue;
+			}
+		}
+#endif // ENABLE_DEBUG_DISCARD_PROP
+
+		FSortedLightSceneInfo* SortedLightInfo = new(SortedLights) FSortedLightSceneInfo(SimpleLightIndex);
+		SortedLightInfo->SortKey.Fields.LightType = LightType_Point;
+		SortedLightInfo->SortKey.Fields.bTextureProfile = 0;
+		SortedLightInfo->SortKey.Fields.bShadowed = 0;
+		SortedLightInfo->SortKey.Fields.bLightFunction = 0;
+		SortedLightInfo->SortKey.Fields.bUsesLightingChannels = 0;
+
+		// These are simple lights.
+		SortedLightInfo->SortKey.Fields.bIsNotSimpleLight = 0;
+
+		// Simple lights are ok to use with tiled and clustered deferred lighting
+		SortedLightInfo->SortKey.Fields.bTiledDeferredNotSupported = 0;
+		SortedLightInfo->SortKey.Fields.bClusteredDeferredNotSupported = 0;
+	}
 
 	// Sort non-shadowed, non-light function lights first to avoid render target switches.
 	struct FCompareFSortedLightSceneInfo
@@ -596,35 +663,80 @@ void FDeferredShadingSceneRenderer::RenderLights(FRHICommandListImmediate& RHICm
 	};
 	SortedLights.Sort( FCompareFSortedLightSceneInfo() );
 
+	// Scan and find ranges.
+	OutSortedLights.SimpleLightsEnd = SortedLights.Num();
+	OutSortedLights.TiledSupportedEnd = SortedLights.Num();
+	OutSortedLights.ClusteredSupportedEnd = SortedLights.Num();
+	OutSortedLights.AttenuationLightStart = SortedLights.Num();
+
+	// Iterate over all lights to be rendered and build ranges for tiled deferred and unshadowed lights
+	for (int32 LightIndex = 0; LightIndex < SortedLights.Num(); LightIndex++)
 	{
-		FSceneRenderTargets& SceneContext = FSceneRenderTargets::Get(RHICmdList);
+		const FSortedLightSceneInfo& SortedLightInfo = SortedLights[LightIndex];
+		const bool bDrawShadows = SortedLightInfo.SortKey.Fields.bShadowed;
+		const bool bDrawLightFunction = SortedLightInfo.SortKey.Fields.bLightFunction;
+		const bool bTextureLightProfile = SortedLightInfo.SortKey.Fields.bTextureProfile;
+		const bool bLightingChannels = SortedLightInfo.SortKey.Fields.bUsesLightingChannels;
 
-		int32 AttenuationLightStart = SortedLights.Num();
-		int32 SupportedByTiledDeferredLightEnd = SortedLights.Num();
+		if (SortedLightInfo.SortKey.Fields.bIsNotSimpleLight && OutSortedLights.SimpleLightsEnd == SortedLights.Num())
+		{
+			// Mark the first index to not be simple
+			OutSortedLights.SimpleLightsEnd = LightIndex;
+		}
 
-		// Iterate over all lights to be rendered and build ranges for tiled deferred and unshadowed lights
-		for (int32 LightIndex = 0; LightIndex < SortedLights.Num(); LightIndex++)
+		if (SortedLightInfo.SortKey.Fields.bTiledDeferredNotSupported && OutSortedLights.TiledSupportedEnd == SortedLights.Num())
 		{
-			const FSortedLightSceneInfo& SortedLightInfo = SortedLights[LightIndex];
-			const bool bDrawShadows = SortedLightInfo.SortKey.Fields.bShadowed;
-			const bool bDrawLightFunction = SortedLightInfo.SortKey.Fields.bLightFunction;
-			const bool bTextureLightProfile = SortedLightInfo.SortKey.Fields.bTextureProfile;
-			const bool bLightingChannels = SortedLightInfo.SortKey.Fields.bUsesLightingChannels;
+			// Mark the first index to not support tiled deferred
+			OutSortedLights.TiledSupportedEnd = LightIndex;
+		}
 
-			if (SortedLightInfo.SortKey.Fields.bTiledDeferredNotSupported && SupportedByTiledDeferredLightEnd == SortedLights.Num())
-			{
-				// Mark the first index to not support tiled deferred
-				SupportedByTiledDeferredLightEnd = LightIndex;
-			}
+		if (SortedLightInfo.SortKey.Fields.bClusteredDeferredNotSupported && OutSortedLights.ClusteredSupportedEnd == SortedLights.Num())
+		{
+			// Mark the first index to not support clustered deferred
+			OutSortedLights.ClusteredSupportedEnd = LightIndex;
+		}
 
-			if (bDrawShadows || bDrawLightFunction || bLightingChannels)
-			{
-				// Once we find a shadowed light, we can exit the loop, these lights should never support tiled deferred rendering either
-				check(SortedLightInfo.SortKey.Fields.bTiledDeferredNotSupported);
-				AttenuationLightStart = LightIndex;
-				break;
-			}
+		if (bDrawShadows || bDrawLightFunction || bLightingChannels)
+		{
+			// Once we find a shadowed light, we can exit the loop, these lights should never support tiled deferred rendering either
+			check(SortedLightInfo.SortKey.Fields.bTiledDeferredNotSupported);
+			OutSortedLights.AttenuationLightStart = LightIndex;
+			break;
 		}
+	}
+
+	// Make sure no obvious things went wrong!
+	check(OutSortedLights.TiledSupportedEnd >= OutSortedLights.SimpleLightsEnd);
+	check(OutSortedLights.ClusteredSupportedEnd >= OutSortedLights.TiledSupportedEnd);
+	check(OutSortedLights.AttenuationLightStart >= OutSortedLights.ClusteredSupportedEnd);
+}
+
+
+
+/** Renders the scene's lighting. */
+void FDeferredShadingSceneRenderer::RenderLights(FRHICommandListImmediate& RHICmdList, FSortedLightSetSceneInfo &SortedLightSet)
+{
+	check(RHICmdList.IsOutsideRenderPass());
+
+	SCOPED_NAMED_EVENT(FDeferredShadingSceneRenderer_RenderLights, FColor::Emerald);
+	SCOPED_DRAW_EVENT(RHICmdList, Lights);
+	SCOPED_GPU_STAT(RHICmdList, Lights);
+
+
+	bool bStencilBufferDirty = false;	// The stencil buffer should've been cleared to 0 already
+
+	SCOPE_CYCLE_COUNTER(STAT_LightingDrawTime);
+	SCOPE_CYCLE_COUNTER(STAT_LightRendering);
+
+	const FSimpleLightArray &SimpleLights = SortedLightSet.SimpleLights;
+	const TArray<FSortedLightSceneInfo, SceneRenderingAllocator> &SortedLights = SortedLightSet.SortedLights;
+	const int32 AttenuationLightStart = SortedLightSet.AttenuationLightStart;
+	const int32 SimpleLightsEnd = SortedLightSet.SimpleLightsEnd;
+
+	{
+		SCOPED_DRAW_EVENT(RHICmdList, DirectLighting);
+
+		FSceneRenderTargets& SceneContext = FSceneRenderTargets::Get(RHICmdList);
 		
 		if (GbEnableAsyncComputeTranslucencyLightingVolumeClear && GSupportsEfficientAsyncCompute)
 		{
@@ -637,11 +749,25 @@ void FDeferredShadingSceneRenderer::RenderLights(FRHICommandListImmediate& RHICm
 			SCOPED_DRAW_EVENT(RHICmdList, NonShadowedLights);
 			INC_DWORD_STAT_BY(STAT_NumUnshadowedLights, AttenuationLightStart);
 
-			int32 StandardDeferredStart = 0;
+			// Currently they have a special path anyway in case of standard deferred so always skip the simple lights
+			int32 StandardDeferredStart = SortedLightSet.SimpleLightsEnd;
 
-			bool bRenderSimpleLightsStandardDeferred = SimpleLights.InstanceData.Num() > 0;
+			bool bRenderSimpleLightsStandardDeferred = SortedLightSet.SimpleLights.InstanceData.Num() > 0;
+			
+			UE_CLOG(ShouldUseClusteredDeferredShading() && !AreClusteredLightsInLightGrid(), LogRenderer, Warning,
+				TEXT("Clustered deferred shading is enabled, but lights were not injected in grid, falling back to other methods (hint 'r.LightCulling.Quality' may cause this)."));
 
-			if (CanUseTiledDeferred())
+			// True if the clustered shading is enabled and the feature level is there, and that the light grid had lights injected.
+			if (ShouldUseClusteredDeferredShading() && AreClusteredLightsInLightGrid())
+			{
+				// Tell the trad. deferred that the clustered deferred capable lights are taken care of.
+				// This includes the simple lights
+				StandardDeferredStart = SortedLightSet.ClusteredSupportedEnd;
+				// Tell the trad. deferred that the simple lights are spoken for.
+				bRenderSimpleLightsStandardDeferred = false;
+				AddClusteredDeferredShadingPass(RHICmdList, SortedLightSet);
+			}
+			else if (CanUseTiledDeferred())
 			{
 				bool bAnyViewIsStereo = false;
 				for (int32 ViewIndex = 0; ViewIndex < Views.Num(); ++ViewIndex)
@@ -654,19 +780,19 @@ void FDeferredShadingSceneRenderer::RenderLights(FRHICommandListImmediate& RHICm
 				}
 
 				// Use tiled deferred shading on any unshadowed lights without a texture light profile
-				if (ShouldUseTiledDeferred(SupportedByTiledDeferredLightEnd, SimpleLights.InstanceData.Num()) && !bAnyViewIsStereo)
+				if (ShouldUseTiledDeferred(SortedLightSet.TiledSupportedEnd) && !bAnyViewIsStereo)
 				{
 					// Update the range that needs to be processed by standard deferred to exclude the lights done with tiled
-					StandardDeferredStart = SupportedByTiledDeferredLightEnd;
+					StandardDeferredStart = SortedLightSet.TiledSupportedEnd;
 					bRenderSimpleLightsStandardDeferred = false;
-					RenderTiledDeferredLighting(RHICmdList, SortedLights, SupportedByTiledDeferredLightEnd, SimpleLights);
+					RenderTiledDeferredLighting(RHICmdList, SortedLights, SortedLightSet.SimpleLightsEnd, SortedLightSet.TiledSupportedEnd, SimpleLights);
 				}
 			}
-			
+
 			if (bRenderSimpleLightsStandardDeferred)
 			{
 				SceneContext.BeginRenderingSceneColor(RHICmdList, ESimpleRenderTargetMode::EExistingColorAndDepth, FExclusiveDepthStencil::DepthRead_StencilWrite);
-				RenderSimpleLightsStandardDeferred(RHICmdList, SimpleLights);
+				RenderSimpleLightsStandardDeferred(RHICmdList, SortedLightSet.SimpleLights);
 				SceneContext.FinishRenderingSceneColor(RHICmdList);
 			}
 
@@ -693,9 +819,9 @@ void FDeferredShadingSceneRenderer::RenderLights(FRHICommandListImmediate& RHICm
 			{
 				if (AttenuationLightStart)
 				{
-					// Inject non-shadowed, non-light function lights in to the volume.
+					// Inject non-shadowed, non-simple, non-light function lights in to the volume.
 					SCOPED_DRAW_EVENT(RHICmdList, InjectNonShadowedTranslucentLighting);
-					InjectTranslucentVolumeLightingArray(RHICmdList, SortedLights, AttenuationLightStart);
+					InjectTranslucentVolumeLightingArray(RHICmdList, SortedLights, SimpleLightsEnd, AttenuationLightStart);
 				}
 
 				if(SimpleLights.InstanceData.Num() > 0)
@@ -745,7 +871,7 @@ void FDeferredShadingSceneRenderer::RenderLights(FRHICommandListImmediate& RHICm
 			// LPV Direct Light Injection
 			if ( bRenderedRSM )
 			{
-				for (int32 LightIndex = 0; LightIndex < SortedLights.Num(); LightIndex++)
+				for (int32 LightIndex = SimpleLightsEnd; LightIndex < SortedLights.Num(); LightIndex++)
 				{
 					const FSortedLightSceneInfo& SortedLightInfo = SortedLights[LightIndex];
 					const FLightSceneInfo* const LightSceneInfo = SortedLightInfo.LightSceneInfo;
diff --git a/Engine/Source/Runtime/Renderer/Private/LightSceneInfo.h b/Engine/Source/Runtime/Renderer/Private/LightSceneInfo.h
--- a/Engine/Source/Runtime/Renderer/Private/LightSceneInfo.h
+++ b/Engine/Source/Runtime/Renderer/Private/LightSceneInfo.h
@@ -12,6 +12,7 @@
 #include "GenericOctreePublic.h"
 #include "SceneManagement.h"
 #include "GenericOctree.h"
+#include "PrimitiveSceneProxy.h"
 #include "Templates/UniquePtr.h"
 
 class FLightPrimitiveInteraction;
@@ -82,23 +83,59 @@ struct FSortedLightSceneInfo
 			uint32 bShadowed : 1;
 			/** Whether the light uses lighting channels. */
 			uint32 bUsesLightingChannels : 1;
+			/** Whether the light is NOT a simple light - they always support tiled/clustered but may want to be selected separately. */
+			uint32 bIsNotSimpleLight : 1;
 			/** True if the light doesn't support tiled deferred, logic is inverted so that lights that DO support tiled deferred will sort first in list */
 			uint32 bTiledDeferredNotSupported : 1;
+			/** 
+			 * True if the light doesn't support clustered deferred, logic is inverted so that lights that DO support clustered deferred will sort first in list 
+			 * Super-set of lights supporting tiled, so the tiled lights will end up in the first part of this range.
+			 */
+			uint32 bClusteredDeferredNotSupported : 1;
 		} Fields;
 		/** Sort key bits packed into an integer. */
 		int32 Packed;
 	} SortKey;
 
 	const FLightSceneInfo* LightSceneInfo;
+	int32 SimpleLightIndex;
 
 	/** Initialization constructor. */
 	explicit FSortedLightSceneInfo(const FLightSceneInfo* InLightSceneInfo)
-		: LightSceneInfo(InLightSceneInfo)
+		: LightSceneInfo(InLightSceneInfo),
+		SimpleLightIndex(-1)
 	{
 		SortKey.Packed = 0;
+		SortKey.Fields.bIsNotSimpleLight = 1;
+	}
+
+	explicit FSortedLightSceneInfo(int32 InSimpleLightIndex)
+		: LightSceneInfo(nullptr),
+		SimpleLightIndex(InSimpleLightIndex)
+	{
+		SortKey.Packed = 0;
+		SortKey.Fields.bIsNotSimpleLight = 0;
 	}
 };
 
+/** 
+ * Stores info about sorted lights and ranges. 
+ * The sort-key in FSortedLightSceneInfo gives rise to the following order:
+ *  [SimpleLights,Tiled/Clustered,LightFunction/Shadow/LightChannels/TextureProfile]
+ */
+struct FSortedLightSetSceneInfo
+{
+	int SimpleLightsEnd;
+	int TiledSupportedEnd;
+	int ClusteredSupportedEnd;
+
+	/** First light with shadow map or */
+	int AttenuationLightStart;
+
+	FSimpleLightArray SimpleLights;
+	TArray<FSortedLightSceneInfo, SceneRenderingAllocator> SortedLights;
+};
+
 template <>
 struct TUseBitwiseSwap<FSortedLightSceneInfo>
 {
@@ -299,17 +336,3 @@ struct FLightOctreeSemantics
 		Element.BoundingSphereVector = VectorAdd(Element.BoundingSphereVector, OffsetReg);
 	}
 };
-
-/** Stores lighting information for the clustered forward shading path */
-struct FClusteredLightsSceneInfo
-{
-	FIntPoint TileSize;			// In pixels
-	FIntVector GridSize;		// In tiles (x,y) + slices (z)
-	FVector4 LightGridZParams;
-
-	// Index of the light in this array corresponds to the bit set in the grid.
-	TArray<FLightSceneInfoCompact> ClusteredLights;
-
-	// The light grid.  Size is >= GridSize.
-	FTexture3DRHIRef LightGridTex;
-};
diff --git a/Engine/Source/Runtime/Renderer/Private/PostProcess/PostProcessSubsurface.cpp b/Engine/Source/Runtime/Renderer/Private/PostProcess/PostProcessSubsurface.cpp
--- a/Engine/Source/Runtime/Renderer/Private/PostProcess/PostProcessSubsurface.cpp
+++ b/Engine/Source/Runtime/Renderer/Private/PostProcess/PostProcessSubsurface.cpp
@@ -10,7 +10,7 @@
 #include "CanvasTypes.h"
 #include "RenderTargetTemp.h"
 
-ENGINE_API const IPooledRenderTarget* GetSubsufaceProfileTexture_RT(FRHICommandListImmediate& RHICmdList);
+ENGINE_API IPooledRenderTarget* GetSubsufaceProfileTexture_RT(FRHICommandListImmediate& RHICmdList);
 
 namespace
 {
diff --git a/Engine/Source/Runtime/Renderer/Private/SceneRendering.h b/Engine/Source/Runtime/Renderer/Private/SceneRendering.h
--- a/Engine/Source/Runtime/Renderer/Private/SceneRendering.h
+++ b/Engine/Source/Runtime/Renderer/Private/SceneRendering.h
@@ -462,6 +462,8 @@ const int32 GMaxForwardShadowCascades = 4;
 	SHADER_PARAMETER(FVector4, DirectionalLightShadowmapAtlasBufferSize) \
 	SHADER_PARAMETER(float, DirectionalLightDepthBias) \
 	SHADER_PARAMETER(uint32, DirectionalLightUseStaticShadowing) \
+	SHADER_PARAMETER(uint32, SimpleLightsEndIndex) \
+	SHADER_PARAMETER(uint32, ClusteredDeferredSupportedEndIndex) \
 	SHADER_PARAMETER(FVector4, DirectionalLightStaticShadowBufferSize) \
 	SHADER_PARAMETER(FMatrix, DirectionalLightWorldToStaticShadow) \
 	SHADER_PARAMETER_TEXTURE(Texture2D, DirectionalLightShadowmapAtlas) \
@@ -470,7 +472,8 @@ const int32 GMaxForwardShadowCascades = 4;
 	SHADER_PARAMETER_SAMPLER(SamplerState, StaticShadowmapSampler) \
 	SHADER_PARAMETER_SRV(StrongTypedBuffer<float4>, ForwardLocalLightBuffer) \
 	SHADER_PARAMETER_SRV(StrongTypedBuffer<uint>, NumCulledLightsGrid) \
-	SHADER_PARAMETER_SRV(StrongTypedBuffer<uint>, CulledLightDataGrid) 
+	SHADER_PARAMETER_SRV(StrongTypedBuffer<uint>, CulledLightDataGrid) \
+	SHADER_PARAMETER_TEXTURE(Texture2D, DummyRectLightSourceTexture)
 
 BEGIN_GLOBAL_SHADER_PARAMETER_STRUCT_WITH_CONSTRUCTOR(FForwardLightData,)
 	FORWARD_GLOBAL_LIGHT_DATA_UNIFORM_BUFFER_MEMBER_TABLE
@@ -494,20 +497,22 @@ public:
 	}
 };
 
+#define ENABLE_LIGHT_CULLING_VIEW_SPACE_BUILD_DATA 1
+
 class FForwardLightingCullingResources
 {
 public:
-	FRWBuffer NextCulledLightLink;
-	FRWBuffer StartOffsetGrid;
-	FRWBuffer CulledLightLinks;
-	FRWBuffer NextCulledLightData;
 
+#if ENABLE_LIGHT_CULLING_VIEW_SPACE_BUILD_DATA
+	FDynamicReadBuffer ViewSpacePosAndRadiusData;
+	FDynamicReadBuffer ViewSpaceDirAndPreprocAngleData;
+#endif // ENABLE_LIGHT_CULLING_VIEW_SPACE_BUILD_DATA
 	void Release()
 	{
-		NextCulledLightLink.Release();
-		StartOffsetGrid.Release();
-		CulledLightLinks.Release();
-		NextCulledLightData.Release();
+#if ENABLE_LIGHT_CULLING_VIEW_SPACE_BUILD_DATA
+		ViewSpacePosAndRadiusData.Release();
+		ViewSpaceDirAndPreprocAngleData.Release();
+#endif // ENABLE_LIGHT_CULLING_VIEW_SPACE_BUILD_DATA
 	}
 };
 
diff --git a/Engine/Source/Runtime/Renderer/Private/ShadowRendering.h b/Engine/Source/Runtime/Renderer/Private/ShadowRendering.h
--- a/Engine/Source/Runtime/Renderer/Private/ShadowRendering.h
+++ b/Engine/Source/Runtime/Renderer/Private/ShadowRendering.h
@@ -29,7 +29,7 @@
 #include "LightRendering.h"
 #include "LightPropagationVolume.h"
 
-ENGINE_API const IPooledRenderTarget* GetSubsufaceProfileTexture_RT(FRHICommandListImmediate& RHICmdList);
+ENGINE_API IPooledRenderTarget* GetSubsufaceProfileTexture_RT(FRHICommandListImmediate& RHICmdList);
 
 class FPrimitiveSceneInfo;
 class FPrimitiveSceneProxy;
diff --git a/Engine/Source/Runtime/Renderer/Private/TiledDeferredLightRendering.cpp b/Engine/Source/Runtime/Renderer/Private/TiledDeferredLightRendering.cpp
--- a/Engine/Source/Runtime/Renderer/Private/TiledDeferredLightRendering.cpp
+++ b/Engine/Source/Runtime/Renderer/Private/TiledDeferredLightRendering.cpp
@@ -124,10 +124,11 @@ public:
 		const FViewInfo& View, 
 		int32 ViewIndex,
 		int32 NumViews,
-		const TArray<FSortedLightSceneInfo, SceneRenderingAllocator>& SortedLights, 
-		int32 NumLightsToRenderInSortedLights, 
-		const FSimpleLightArray& SimpleLights, 
-		int32 StartIndex, 
+		const TArray<FSortedLightSceneInfo, SceneRenderingAllocator>& SortedLights,
+		int32 TiledDeferredLightsStart, 
+		int32 TiledDeferredLightsEnd,
+		const FSimpleLightArray& SimpleLights,
+		int32 StartIndex,
 		int32 NumThisPass,
 		IPooledRenderTarget& InTextureValue,
 		IPooledRenderTarget& OutTextureValue)
@@ -162,6 +163,8 @@ public:
 			GSystemTextures.LTCAmp->GetRenderTargetItem().ShaderResourceTexture
 			);
 
+		const int32 NumLightsToRenderInSortedLights = TiledDeferredLightsEnd - TiledDeferredLightsStart;
+
 		if (TransmissionProfilesTexture.IsBound())
 		{
 			FSceneRenderTargets& SceneContext = FSceneRenderTargets::Get(RHICmdList);
@@ -193,7 +196,7 @@ public:
 		{
 			if (StartIndex + LightIndex < NumLightsToRenderInSortedLights)
 			{
-				const FSortedLightSceneInfo& SortedLightInfo = SortedLights[StartIndex + LightIndex];
+				const FSortedLightSceneInfo& SortedLightInfo = SortedLights[TiledDeferredLightsStart + StartIndex + LightIndex];
 				const FLightSceneInfo* const LightSceneInfo = SortedLightInfo.LightSceneInfo;
 
 				FLightShaderParameters LightParameters;
@@ -325,10 +328,10 @@ bool FDeferredShadingSceneRenderer::CanUseTiledDeferred() const
 	return GUseTiledDeferredShading != 0 && Scene->GetFeatureLevel() >= ERHIFeatureLevel::SM5;
 }
 
-bool FDeferredShadingSceneRenderer::ShouldUseTiledDeferred(int32 NumUnshadowedLights, int32 NumSimpleLights) const
+bool FDeferredShadingSceneRenderer::ShouldUseTiledDeferred(int32 NumTiledDeferredLights) const
 {
 	// Only use tiled deferred if there are enough unshadowed lights to justify the fixed cost, 
-	return (NumUnshadowedLights + NumSimpleLights >= GNumLightsBeforeUsingTiledDeferred);
+	return (NumTiledDeferredLights >= GNumLightsBeforeUsingTiledDeferred);
 }
 
 template <bool bVisualizeLightCulling>
@@ -337,10 +340,11 @@ static void SetShaderTemplTiledLighting(
 	const FViewInfo& View,
 	int32 ViewIndex,
 	int32 NumViews,
-	const TArray<FSortedLightSceneInfo, SceneRenderingAllocator>& SortedLights, 
-	int32 NumLightsToRenderInSortedLights, 
+	const TArray<FSortedLightSceneInfo, SceneRenderingAllocator>& SortedLights,
+	int32 TiledDeferredLightsStart, 
+	int32 TiledDeferredLightsEnd,
 	const FSimpleLightArray& SimpleLights,
-	int32 StartIndex, 
+	int32 StartIndex,
 	int32 NumThisPass,
 	IPooledRenderTarget& InTexture,
 	IPooledRenderTarget& OutTexture)
@@ -348,7 +352,7 @@ static void SetShaderTemplTiledLighting(
 	TShaderMapRef<FTiledDeferredLightingCS<bVisualizeLightCulling> > ComputeShader(View.ShaderMap);
 	RHICmdList.SetComputeShader(ComputeShader->GetComputeShader());
 
-	ComputeShader->SetParameters(RHICmdList, View, ViewIndex, NumViews, SortedLights, NumLightsToRenderInSortedLights, SimpleLights, StartIndex, NumThisPass, InTexture, OutTexture);
+	ComputeShader->SetParameters(RHICmdList, View, ViewIndex, NumViews, SortedLights, TiledDeferredLightsStart, TiledDeferredLightsEnd, SimpleLights, StartIndex, NumThisPass, InTexture, OutTexture);
 
 	uint32 GroupSizeX = (View.ViewRect.Size().X + GDeferredLightTileSizeX - 1) / GDeferredLightTileSizeX;
 	uint32 GroupSizeY = (View.ViewRect.Size().Y + GDeferredLightTileSizeY - 1) / GDeferredLightTileSizeY;
@@ -358,8 +362,10 @@ static void SetShaderTemplTiledLighting(
 }
 
 
-void FDeferredShadingSceneRenderer::RenderTiledDeferredLighting(FRHICommandListImmediate& RHICmdList, const TArray<FSortedLightSceneInfo, SceneRenderingAllocator>& SortedLights, int32 NumUnshadowedLights, const FSimpleLightArray& SimpleLights)
+void FDeferredShadingSceneRenderer::RenderTiledDeferredLighting(FRHICommandListImmediate& RHICmdList, const TArray<FSortedLightSceneInfo, SceneRenderingAllocator>& SortedLights, int32 TiledDeferredLightsStart, int32 TiledDeferredLightsEnd, const FSimpleLightArray& SimpleLights)
 {
+	const int32 NumUnshadowedLights = TiledDeferredLightsEnd - TiledDeferredLightsStart;
+
 	check(GUseTiledDeferredShading);
 	check(SortedLights.Num() >= NumUnshadowedLights);
 
@@ -405,11 +411,11 @@ void FDeferredShadingSceneRenderer::RenderTiledDeferredLighting(FRHICommandListI
 
 					if(View.Family->EngineShowFlags.VisualizeLightCulling)
 					{
-						SetShaderTemplTiledLighting<1>(RHICmdList, View, ViewIndex, Views.Num(), SortedLights, NumLightsToRenderInSortedLights, SimpleLights, StartIndex, NumThisPass, InTexture, *OutTexture);
+						SetShaderTemplTiledLighting<1>(RHICmdList, View, ViewIndex, Views.Num(), SortedLights, TiledDeferredLightsStart, TiledDeferredLightsEnd, SimpleLights, StartIndex, NumThisPass, InTexture, *OutTexture);
 					}
 					else
 					{
-						SetShaderTemplTiledLighting<0>(RHICmdList, View, ViewIndex, Views.Num(), SortedLights, NumLightsToRenderInSortedLights, SimpleLights, StartIndex, NumThisPass, InTexture, *OutTexture);
+						SetShaderTemplTiledLighting<0>(RHICmdList, View, ViewIndex, Views.Num(), SortedLights, TiledDeferredLightsStart, TiledDeferredLightsEnd, SimpleLights, StartIndex, NumThisPass, InTexture, *OutTexture);
 					}
 				}
 			}
diff --git a/Engine/Source/Runtime/Renderer/Private/TranslucentLighting.cpp b/Engine/Source/Runtime/Renderer/Private/TranslucentLighting.cpp
--- a/Engine/Source/Runtime/Renderer/Private/TranslucentLighting.cpp
+++ b/Engine/Source/Runtime/Renderer/Private/TranslucentLighting.cpp
@@ -1574,18 +1574,18 @@ void FDeferredShadingSceneRenderer::InjectTranslucentVolumeLighting(FRHICommandL
 	}
 }
 
-void FDeferredShadingSceneRenderer::InjectTranslucentVolumeLightingArray(FRHICommandListImmediate& RHICmdList, const TArray<FSortedLightSceneInfo, SceneRenderingAllocator>& SortedLights, int32 NumLights)
+void FDeferredShadingSceneRenderer::InjectTranslucentVolumeLightingArray(FRHICommandListImmediate& RHICmdList, const TArray<FSortedLightSceneInfo, SceneRenderingAllocator>& SortedLights, int32 FirstLightIndex, int32 LightsEndIndex)
 {
 	SCOPE_CYCLE_COUNTER(STAT_TranslucentInjectTime);
 
 	TArray<FTranslucentLightInjectionData, SceneRenderingAllocator> *LightInjectionData = new TArray<FTranslucentLightInjectionData, SceneRenderingAllocator>[Views.Num()];
 	for(int32 ViewIndex = 0; ViewIndex < Views.Num(); ++ViewIndex)
 	{
-		LightInjectionData[ViewIndex].Empty(NumLights);
+		LightInjectionData[ViewIndex].Empty(LightsEndIndex - FirstLightIndex);
 	}
 	
 
-	for (int32 LightIndex = 0; LightIndex < NumLights; LightIndex++)
+	for (int32 LightIndex = FirstLightIndex; LightIndex < LightsEndIndex; LightIndex++)
 	{
 		const FSortedLightSceneInfo& SortedLightInfo = SortedLights[LightIndex];
 		const FLightSceneInfo* const LightSceneInfo = SortedLightInfo.LightSceneInfo;
