Force reflection captures to use fully rough material. Done all in the base pass.
SkyEnv pass has now been removed when in reflection capture.
Factored duplicated code into a single function.

PR #5109: Force Reflection capture rendering to fully rough. (Contributed by [REDACTED_PERSON])

diff --git a/Engine/Shaders/Private/BRDF.ush b/Engine/Shaders/Private/BRDF.ush
--- a/Engine/Shaders/Private/BRDF.ush
+++ b/Engine/Shaders/Private/BRDF.ush
@@ -355,6 +355,14 @@ half EnvBRDFApproxNonmetal( half Roughness, half NoV )
 	return min( r.x * r.x, exp2( -9.28 * NoV ) ) * r.x + r.y;
 }
 
+void EnvBRDFApproxFullyRough(inout half3 DiffuseColor, inout half3 SpecularColor)
+{
+	// Factors derived from EnvBRDFApprox( SpecularColor, 1, 1 ) == SpecularColor * 0.4524 - 0.0024
+	DiffuseColor += SpecularColor * 0.45;
+	SpecularColor = 0;
+	// We do not modify Roughness here as this is done differently at different places.
+}
+
 
 float D_InvBlinn( float a2, float NoH )
 {
diff --git a/Engine/Shaders/Private/BasePassPixelShader.usf b/Engine/Shaders/Private/BasePassPixelShader.usf
--- a/Engine/Shaders/Private/BasePassPixelShader.usf
+++ b/Engine/Shaders/Private/BasePassPixelShader.usf
@@ -927,12 +927,13 @@ void FPixelShaderInOut_MainPS(
 	}
 #endif
 
-    #if FORCE_FULLY_ROUGH
-		// Factors derived from EnvBRDFApprox( SpecularColor, 1, 1 ) == SpecularColor * 0.4524 - 0.0024
-		GBuffer.DiffuseColor += GBuffer.SpecularColor * 0.45;
-		GBuffer.SpecularColor = 0;
-		GBuffer.Roughness = 1;
+    #if !FORCE_FULLY_ROUGH
+	if (View.RenderingReflectionCaptureMask) // Force material rendered in reflection capture to have an expanded albedo to try to be energy conservative (when specular is removed).
 	#endif
+	{
+		EnvBRDFApproxFullyRough(GBuffer.DiffuseColor, GBuffer.SpecularColor);
+		// When rendering reflection captures, GBuffer.Roughness is already forced to 1 using RoughnessOverrideParameter in GetMaterialRoughness.
+	}
 
 	float3 BentNormal = MaterialParameters.WorldNormal;
 			
@@ -1075,7 +1076,10 @@ void FPixelShaderInOut_MainPS(
 	#endif
 
 	#if !MATERIAL_SHADINGMODEL_UNLIT && USE_DEVELOPMENT_SHADERS
-		Color = lerp(Color, GBuffer.DiffuseColor + GBuffer.SpecularColor * 0.45, View.UnlitViewmodeMask);
+		float3 GBufferDiffuseColor = GBuffer.DiffuseColor;
+		float3 GBufferSpecularColor = GBuffer.SpecularColor;
+		EnvBRDFApproxFullyRough(GBufferDiffuseColor, GBufferSpecularColor);
+		Color = lerp(Color, GBufferDiffuseColor, View.UnlitViewmodeMask);
 	#endif
 
 	half3 Emissive = GetMaterialEmissive(PixelMaterialInputs);
diff --git a/Engine/Shaders/Private/ForwardLightingCommon.ush b/Engine/Shaders/Private/ForwardLightingCommon.ush
--- a/Engine/Shaders/Private/ForwardLightingCommon.ush
+++ b/Engine/Shaders/Private/ForwardLightingCommon.ush
@@ -340,11 +340,6 @@ half3 GetImageBasedReflectionLighting(FMaterialPixelParameters MaterialParameter
 		SingleCaptureIndex,
 		EyeIndex);
 
-	// Factors derived from EnvBRDFApprox( SpecularColor, 1, 1 ) == SpecularColor * 0.4524 - 0.0024
-	float3 SpecularBounce = 0.45f * SpecularColor * IndirectIrradiance;
-	// Replace reflection captures with indirect diffuse when we're rendering to a reflection capture, to avoid a feedback loop
-	SpecularIBL.rgb = lerp(SpecularIBL.rgb, SpecularBounce, ResolvedView.RenderingReflectionCaptureMask);
-
 #if MATERIAL_SSR && !FORWARD_SHADING
 	if( View.CameraCut == 0 )
 	{
diff --git a/Engine/Shaders/Private/MobileBasePassPixelShader.usf b/Engine/Shaders/Private/MobileBasePassPixelShader.usf
--- a/Engine/Shaders/Private/MobileBasePassPixelShader.usf
+++ b/Engine/Shaders/Private/MobileBasePassPixelShader.usf
@@ -400,9 +400,7 @@ void Main(
 	half Roughness = GetMaterialRoughness(PixelMaterialInputs);
 
 #if FULLY_ROUGH
-	// Factors derived from EnvBRDFApprox( SpecularColor, 1, 1 ) == SpecularColor * 0.4524 - 0.0024
-	DiffuseColor += SpecularColor * 0.45;
-	SpecularColor = 0;
+	EnvBRDFApproxFullyRough(DiffuseColor, SpecularColor);
 #else
 	half NoV = max( dot( MaterialParameters.WorldNormal, MaterialParameters.CameraVector ), 0 );
 
diff --git a/Engine/Shaders/Private/ReflectionEnvironmentPixelShader.usf b/Engine/Shaders/Private/ReflectionEnvironmentPixelShader.usf
--- a/Engine/Shaders/Private/ReflectionEnvironmentPixelShader.usf
+++ b/Engine/Shaders/Private/ReflectionEnvironmentPixelShader.usf
@@ -170,25 +170,6 @@ float3 SkyLightDiffuse(FScreenSpaceData ScreenSpaceData, float4 UVAndScreenPos,
 	return Lighting;
 }
 
-// Used during reflection captures and SM4 to get faked bounce light from specular surfaces
-float3 SpecularBounce(FScreenSpaceData ScreenSpaceData, float4 UVAndScreenPos)
-{
-	float2 UV = UVAndScreenPos.xy;
-
-	FGBufferData GBuffer = ScreenSpaceData.GBuffer;
-
-	float3 ScreenVector = normalize(mul(float4(UVAndScreenPos.zw, 1, 0), View.ScreenToWorld).xyz);
-
-	float3 N = GBuffer.WorldNormal;
-	float3 V = -ScreenVector;
-
-	float NoV = saturate(dot(N, V));
-
-	GBuffer.SpecularColor = EnvBRDFApprox(GBuffer.SpecularColor, GBuffer.Roughness, NoV);
-
-	return GBuffer.SpecularColor * GBuffer.IndirectIrradiance;
-}
-
 float3 ReflectionEnvironment(FScreenSpaceData ScreenSpaceData, float4 UVAndScreenPos, float4 SvPosition, float3 BentNormal)
 {
 #if USE_PREEXPOSURE
@@ -352,11 +333,7 @@ void ReflectionEnvironmentSkyLighting(
 	BRANCH
 	if (ShadingModelID != SHADINGMODELID_UNLIT && ShadingModelID != SHADINGMODELID_HAIR)
 	{
-#if SPECULAR_BOUNCE
-		OutColor.xyz += SpecularBounce(ScreenSpaceData, UVAndScreenPos);
-#else
 		OutColor.xyz += ReflectionEnvironment(ScreenSpaceData, UVAndScreenPos, SvPosition, BentNormal);
-#endif
 	}
 }
 
diff --git a/Engine/Source/Runtime/Renderer/Private/ReflectionEnvironment.cpp b/Engine/Source/Runtime/Renderer/Private/ReflectionEnvironment.cpp
--- a/Engine/Source/Runtime/Renderer/Private/ReflectionEnvironment.cpp
+++ b/Engine/Source/Runtime/Renderer/Private/ReflectionEnvironment.cpp
@@ -515,7 +515,6 @@ class FReflectionEnvironmentSkyLightingPS : public FGlobalShader
 	class FHasBoxCaptures			: SHADER_PERMUTATION_BOOL("REFLECTION_COMPOSITE_HAS_BOX_CAPTURES");
 	class FHasSphereCaptures		: SHADER_PERMUTATION_BOOL("REFLECTION_COMPOSITE_HAS_SPHERE_CAPTURES");
 	class FDFAOIndirectOcclusion	: SHADER_PERMUTATION_BOOL("SUPPORT_DFAO_INDIRECT_OCCLUSION");
-	class FSpecularBounce			: SHADER_PERMUTATION_BOOL("SPECULAR_BOUNCE");
 	class FSkyLight					: SHADER_PERMUTATION_BOOL("ENABLE_SKY_LIGHT");
 	class FDynamicSkyLight			: SHADER_PERMUTATION_BOOL("ENABLE_DYNAMIC_SKY_LIGHT");
 	class FSkyShadowing				: SHADER_PERMUTATION_BOOL("APPLY_SKY_SHADOWING");
@@ -525,7 +524,6 @@ class FReflectionEnvironmentSkyLightingPS : public FGlobalShader
 		FHasBoxCaptures,
 		FHasSphereCaptures,
 		FDFAOIndirectOcclusion,
-		FSpecularBounce,
 		FSkyLight,
 		FDynamicSkyLight,
 		FSkyShadowing,
@@ -533,15 +531,6 @@ class FReflectionEnvironmentSkyLightingPS : public FGlobalShader
 
 	static FPermutationDomain RemapPermutation(FPermutationDomain PermutationVector)
 	{
-		// Environment captures have simple specular bounce without reflection captures.
-		if (PermutationVector.Get<FSpecularBounce>())
-		{
-			PermutationVector.Set<FSkyLight>(false);
-			PermutationVector.Set<FDFAOIndirectOcclusion>(false);
-			PermutationVector.Set<FHasBoxCaptures>(false);
-			PermutationVector.Set<FHasSphereCaptures>(false);
-		}
-
 		// FSkyLightingDynamicSkyLight requires FSkyLightingSkyLight.
 		if (!PermutationVector.Get<FSkyLight>())
 		{
@@ -557,14 +546,13 @@ class FReflectionEnvironmentSkyLightingPS : public FGlobalShader
 		return PermutationVector;
 	}
 
-	static FPermutationDomain BuildPermutationVector(const FViewInfo& View, bool bBoxCapturesOnly, bool bSphereCapturesOnly, bool bSupportDFAOIndirectOcclusion, bool bSpecularBounce, bool bEnableSkyLight, bool bEnableDynamicSkyLight, bool bApplySkyShadowing, bool bRayTracedReflections)
+	static FPermutationDomain BuildPermutationVector(const FViewInfo& View, bool bBoxCapturesOnly, bool bSphereCapturesOnly, bool bSupportDFAOIndirectOcclusion, bool bEnableSkyLight, bool bEnableDynamicSkyLight, bool bApplySkyShadowing, bool bRayTracedReflections)
 	{
 		FPermutationDomain PermutationVector;
 
 		PermutationVector.Set<FHasBoxCaptures>(bBoxCapturesOnly);
 		PermutationVector.Set<FHasSphereCaptures>(bSphereCapturesOnly);
 		PermutationVector.Set<FDFAOIndirectOcclusion>(bSupportDFAOIndirectOcclusion);
-		PermutationVector.Set<FSpecularBounce>(bSpecularBounce);
 		PermutationVector.Set<FSkyLight>(bEnableSkyLight);
 		PermutationVector.Set<FDynamicSkyLight>(bEnableDynamicSkyLight);
 		PermutationVector.Set<FSkyShadowing>(bApplySkyShadowing);
@@ -724,6 +712,12 @@ void FDeferredShadingSceneRenderer::RenderDeferredReflectionsAndSkyLighting(FRHI
 		bAnyViewWithRaytracingReflections = bAnyViewWithRaytracingReflections || (View.FinalPostProcessSettings.ReflectionsType == EReflectionsType::RayTracing);
 	}
 
+	if (bReflectionCapture)
+	{
+		// if we are rendering a reflection capture then we can skip this pass entirely (no reflection and no sky contribution evaluated in this pass)
+		return;	
+	}
+
 	const bool bRayTracedReflections = IsRayTracingEnabled() && (GRayTracingReflections < 0 ? bAnyViewWithRaytracingReflections : GRayTracingReflections);
 
 	// The specular sky light contribution is also needed by RT Reflections as a fallback.
@@ -855,7 +849,7 @@ void FDeferredShadingSceneRenderer::RenderDeferredReflectionsAndSkyLighting(FRHI
 
 			TShaderMapRef<FPostProcessVS> VertexShader(View.ShaderMap);
 
-			auto PermutationVector = FReflectionEnvironmentSkyLightingPS::BuildPermutationVector(View, bHasBoxCaptures, bHasSphereCaptures, DynamicBentNormalAO != NULL, bReflectionCapture, bSkyLight, bDynamicSkyLight, bApplySkyShadowing, bRayTracedReflections);
+			auto PermutationVector = FReflectionEnvironmentSkyLightingPS::BuildPermutationVector(View, bHasBoxCaptures, bHasSphereCaptures, DynamicBentNormalAO != NULL, bSkyLight, bDynamicSkyLight, bApplySkyShadowing, bRayTracedReflections);
 
 			TShaderMapRef<FReflectionEnvironmentSkyLightingPS> PixelShader(View.ShaderMap, PermutationVector);
 
@@ -863,16 +857,12 @@ void FDeferredShadingSceneRenderer::RenderDeferredReflectionsAndSkyLighting(FRHI
 
 			SceneContext.BeginRenderingSceneColor(
 				RHICmdList,
-				bReflectionCapture ? ESimpleRenderTargetMode::EUninitializedColorExistingDepth : ESimpleRenderTargetMode::EExistingColorAndDepth,
+				ESimpleRenderTargetMode::EExistingColorAndDepth,
 				FExclusiveDepthStencil::DepthRead_StencilWrite);
 			RHICmdList.ApplyCachedRenderTargets(GraphicsPSOInit);
 
 			extern int32 GAOOverwriteSceneColor;
-			if (bReflectionCapture)
-			{
-				GraphicsPSOInit.BlendState = TStaticBlendState<CW_RGB, BO_Add, BF_One, BF_One>::GetRHI();
-			}
-			else if (GetReflectionEnvironmentCVar() == 2 || GAOOverwriteSceneColor)
+			if (GetReflectionEnvironmentCVar() == 2 || GAOOverwriteSceneColor)
 			{
 				// override scene color for debugging
 				GraphicsPSOInit.BlendState = TStaticBlendState<>::GetRHI();
@@ -904,31 +894,15 @@ void FDeferredShadingSceneRenderer::RenderDeferredReflectionsAndSkyLighting(FRHI
 
 			PixelShader->SetParameters(RHICmdList, View, ReflectionsColor->GetRenderTargetItem().ShaderResourceTexture, DynamicBentNormalAO);
 
-			if (bReflectionCapture)
-			{
-				DrawRectangle(
-					RHICmdList,
-					0, 0,
-					View.ViewRect.Width(), View.ViewRect.Height(),
-					0, 0,
-					View.ViewRect.Width(), View.ViewRect.Height(),
-					FIntPoint(View.ViewRect.Width(), View.ViewRect.Height()),
-					SceneContext.GetBufferSizeXY(),
-					*VertexShader,
-					EDRF_UseTriangleOptimization);
-			}
-			else
-			{
-				DrawRectangle(
-					RHICmdList,
-					0, 0,
-					View.ViewRect.Width(), View.ViewRect.Height(),
-					View.ViewRect.Min.X, View.ViewRect.Min.Y,
-					View.ViewRect.Width(), View.ViewRect.Height(),
-					FIntPoint(View.ViewRect.Width(), View.ViewRect.Height()),
-					SceneContext.GetBufferSizeXY(),
-					*VertexShader);
-			}
+			DrawRectangle(
+				RHICmdList,
+				0, 0,
+				View.ViewRect.Width(), View.ViewRect.Height(),
+				View.ViewRect.Min.X, View.ViewRect.Min.Y,
+				View.ViewRect.Width(), View.ViewRect.Height(),
+				FIntPoint(View.ViewRect.Width(), View.ViewRect.Height()),
+				SceneContext.GetBufferSizeXY(),
+				*VertexShader);
 
 			SceneContext.FinishRenderingSceneColor(RHICmdList);
 
