DR - Fix Vulkan crashing due to using a storage texel (ie structured buffer) on a shader that needed storage buffer (RWBuffer)
#rnx

diff --git a/Engine/Source/Runtime/Renderer/Private/LightGridInjection.cpp b/Engine/Source/Runtime/Renderer/Private/LightGridInjection.cpp
--- a/Engine/Source/Runtime/Renderer/Private/LightGridInjection.cpp
+++ b/Engine/Source/Runtime/Renderer/Private/LightGridInjection.cpp
@@ -713,10 +713,10 @@ void FDeferredShadingSceneRenderer::ComputeLightGrid(FRHICommandListImmediate& R
 					//FRDGBufferRef NumCulledLightsGridBuffer = GraphBuilder.CreateBuffer(FRDGBufferDesc::CreateStructuredDesc(sizeof(uint32), NumCells * NumCulledLightsGridStride), TEXT("NumCulledLightsGrid"));
 					//FRDGBufferRef CulledLightDataGridBuffer = GraphBuilder.CreateBuffer(FRDGBufferDesc::CreateStructuredDesc(LightIndexTypeSize, NumCells * GMaxCulledLightsPerCell), TEXT("CulledLightDataGrid"));
 
-					FRDGBufferRef CulledLightLinksBuffer = GraphBuilder.CreateBuffer(FRDGBufferDesc::CreateStructuredDesc(sizeof(uint32), CulledLightLinksElements), TEXT("CulledLightLinks"));
-					FRDGBufferRef StartOffsetGridBuffer = GraphBuilder.CreateBuffer(FRDGBufferDesc::CreateStructuredDesc(sizeof(uint32), NumCells), TEXT("StartOffsetGrid"));
-					FRDGBufferRef NextCulledLightLinkBuffer = GraphBuilder.CreateBuffer(FRDGBufferDesc::CreateStructuredDesc(sizeof(uint32), 1), TEXT("NextCulledLightLink"));
-					FRDGBufferRef NextCulledLightDataBuffer = GraphBuilder.CreateBuffer(FRDGBufferDesc::CreateStructuredDesc(sizeof(uint32), 1), TEXT("NextCulledLightData"));
+					FRDGBufferRef CulledLightLinksBuffer = GraphBuilder.CreateBuffer(FRDGBufferDesc::CreateBufferDesc(sizeof(uint32), CulledLightLinksElements), TEXT("CulledLightLinks"));
+					FRDGBufferRef StartOffsetGridBuffer = GraphBuilder.CreateBuffer(FRDGBufferDesc::CreateBufferDesc(sizeof(uint32), NumCells), TEXT("StartOffsetGrid"));
+					FRDGBufferRef NextCulledLightLinkBuffer = GraphBuilder.CreateBuffer(FRDGBufferDesc::CreateBufferDesc(sizeof(uint32), 1), TEXT("NextCulledLightLink"));
+					FRDGBufferRef NextCulledLightDataBuffer = GraphBuilder.CreateBuffer(FRDGBufferDesc::CreateBufferDesc(sizeof(uint32), 1), TEXT("NextCulledLightData"));
 
 #if ENABLE_LIGHT_CULLING_VIEW_SPACE_BUILD_DATA
 					// TODO: When CPU-side generative passes are added this can also use render graph
@@ -731,9 +731,9 @@ void FDeferredShadingSceneRenderer::ComputeLightGrid(FRHICommandListImmediate& R
 					PassParameters->Forward = View.ForwardLightingResources->ForwardLightDataUniformBuffer;
 					PassParameters->RWNumCulledLightsGrid = View.ForwardLightingResources->NumCulledLightsGrid.UAV;
 					PassParameters->RWCulledLightDataGrid = View.ForwardLightingResources->CulledLightDataGrid.UAV;
-					PassParameters->RWNextCulledLightLink = GraphBuilder.CreateUAV(NextCulledLightLinkBuffer);
-					PassParameters->RWStartOffsetGrid = GraphBuilder.CreateUAV(StartOffsetGridBuffer);
-					PassParameters->RWCulledLightLinks = GraphBuilder.CreateUAV(CulledLightLinksBuffer);
+					PassParameters->RWNextCulledLightLink = GraphBuilder.CreateUAV(NextCulledLightLinkBuffer, PF_R32_UINT);
+					PassParameters->RWStartOffsetGrid = GraphBuilder.CreateUAV(StartOffsetGridBuffer, PF_R32_UINT);
+					PassParameters->RWCulledLightLinks = GraphBuilder.CreateUAV(CulledLightLinksBuffer, PF_R32_UINT);
 
 #if ENABLE_LIGHT_CULLING_VIEW_SPACE_BUILD_DATA
 					// TODO: When CPU-side generative passes are added this can also use render graph
@@ -751,7 +751,7 @@ void FDeferredShadingSceneRenderer::ComputeLightGrid(FRHICommandListImmediate& R
 					{
 						AddPass_ClearUAV(GraphBuilder, RDG_EVENT_NAME("Clear:StartOffsetGrid"), PassParameters->RWStartOffsetGrid, 0xFFFFFFFF);
 						AddPass_ClearUAV(GraphBuilder, RDG_EVENT_NAME("Clear:NextCulledLightLink"), PassParameters->RWNextCulledLightLink, 0);
-						AddPass_ClearUAV(GraphBuilder, RDG_EVENT_NAME("Clear:NextCulledLightData"), GraphBuilder.CreateUAV(NextCulledLightDataBuffer), 0);
+						AddPass_ClearUAV(GraphBuilder, RDG_EVENT_NAME("Clear:NextCulledLightData"), GraphBuilder.CreateUAV(NextCulledLightDataBuffer, PF_R32_UINT), 0);
 						FComputeShaderUtils::AddPass(GraphBuilder, RDG_EVENT_NAME("LightGridInject:LinkedList"), *ComputeShader, PassParameters, NumGroups);
 
 
@@ -761,11 +761,11 @@ void FDeferredShadingSceneRenderer::ComputeLightGrid(FRHICommandListImmediate& R
 							PassParametersCompact->View = View.ViewUniformBuffer;
 							PassParametersCompact->Forward = View.ForwardLightingResources->ForwardLightDataUniformBuffer;
 
-							PassParametersCompact->CulledLightLinks = GraphBuilder.CreateSRV(CulledLightLinksBuffer);
+							PassParametersCompact->CulledLightLinks = GraphBuilder.CreateSRV(CulledLightLinksBuffer, PF_R32_UINT);
 							PassParametersCompact->RWNumCulledLightsGrid = View.ForwardLightingResources->NumCulledLightsGrid.UAV;
 							PassParametersCompact->RWCulledLightDataGrid = View.ForwardLightingResources->CulledLightDataGrid.UAV;
-							PassParametersCompact->RWNextCulledLightData = GraphBuilder.CreateUAV(NextCulledLightDataBuffer);
-							PassParametersCompact->StartOffsetGrid = GraphBuilder.CreateSRV(StartOffsetGridBuffer);;
+							PassParametersCompact->RWNextCulledLightData = GraphBuilder.CreateUAV(NextCulledLightDataBuffer, PF_R32_UINT);
+							PassParametersCompact->StartOffsetGrid = GraphBuilder.CreateSRV(StartOffsetGridBuffer, PF_R32_UINT);
 
 							FComputeShaderUtils::AddPass(GraphBuilder, RDG_EVENT_NAME("CompactLinks"), *ComputeShaderCompact, PassParametersCompact, NumGroups);
 						}
