Implements TRCPassForRDG to conveniently program part of the post processing graph with RDG.

diff --git a/Engine/Source/Runtime/Renderer/Private/PostProcess/RenderingCompositionGraph.h b/Engine/Source/Runtime/Renderer/Private/PostProcess/RenderingCompositionGraph.h
--- a/Engine/Source/Runtime/Renderer/Private/PostProcess/RenderingCompositionGraph.h
+++ b/Engine/Source/Runtime/Renderer/Private/PostProcess/RenderingCompositionGraph.h
@@ -498,6 +498,7 @@ private:
 	uint32 Dependencies;
 };
 
+
 //
 template <uint32 InputCount, uint32 OutputCount>
 struct TRenderingCompositePassBase :public FRenderingCompositePass
@@ -672,4 +673,57 @@ protected:
 	}
 };
 
+
+/** Utility to conveniently create a RDG graph within a post process graph.
+ *
+ * Example:
+ *	FRenderingCompositePass* DiaphragmDOFPass = Context.Graph.RegisterPass(
+ *		new(FMemStack::Get()) TRCPassForRDG<2, 1>([](FRenderingCompositePass* Pass, FRenderingCompositePassContext& Context)
+ *	{
+ *		FRDGBuilder GraphBuilder(Context.RHICmdList);
+ *
+ *		FRDGTextureRef SceneColor = Pass->CreateRDGTextureForInput(GraphBuilder, ePId_Input0, TEXT("SceneColor"), eFC_0000);
+ *		FRDGTextureRef SeparateTranslucency = Pass->CreateRDGTextureForInput(GraphBuilder, ePId_Input1, TEXT("SeparateTranslucency"), eFC_0000);
+ *
+ *		// ...
+ *
+ *		Pass->ExtractRDGTextureForOutput(GraphBuilder, ePId_Output0, NewSceneColor);
+ *		GraphBuilder.Execute();
+ *	}));
+ *	DiaphragmDOFPass->SetInput(ePId_Input0, Context.FinalOutput);
+ *	DiaphragmDOFPass->SetInput(ePId_Input1, SeparateTranslucency);
+ *	Context.FinalOutput = FRenderingCompositeOutputRef(DiaphragmDOFPass, ePId_Output0);
+ */
+template<uint32 InputCount, uint32 OutputCount>
+class TRCPassForRDG : public TRenderingCompositePassBase<InputCount, OutputCount>
+{
+public:
+	TRCPassForRDG(TFunction<void (FRenderingCompositePass*, FRenderingCompositePassContext&)>&& InProcessLambda)
+		: ProcessLambda(InProcessLambda)
+	{ }
+
+	// interface FRenderingCompositePass ---------
+	virtual void Process(FRenderingCompositePassContext& Context) override
+	{
+		WaitForInputPassComputeFences(Context.RHICmdList);
+
+		ProcessLambda(this, Context);
+	}
+
+	virtual void Release() override
+	{
+		delete this;
+	}
+
+	virtual FPooledRenderTargetDesc ComputeOutputDesc(EPassOutputId InPassOutputId) const
+	{
+		// ExtractRDGTextureForOutput() is doing this work for us already.
+		return FPooledRenderTargetDesc();
+	}
+
+private:
+	TFunction<void (FRenderingCompositePass*, FRenderingCompositePassContext&)> ProcessLambda;
+};
+
+
 void CompositionGraph_OnStartFrame();
