Makes clang happy

diff --git a/Engine/Source/Runtime/Renderer/Private/PostProcess/RenderingCompositionGraph.h b/Engine/Source/Runtime/Renderer/Private/PostProcess/RenderingCompositionGraph.h
--- a/Engine/Source/Runtime/Renderer/Private/PostProcess/RenderingCompositionGraph.h
+++ b/Engine/Source/Runtime/Renderer/Private/PostProcess/RenderingCompositionGraph.h
@@ -705,7 +705,7 @@ public:
 	// interface FRenderingCompositePass ---------
 	virtual void Process(FRenderingCompositePassContext& Context) override
 	{
-		WaitForInputPassComputeFences(Context.RHICmdList);
+		this->TRenderingCompositePassBase<InputCount, OutputCount>::WaitForInputPassComputeFences(Context.RHICmdList);
 
 		ProcessLambda(this, Context);
 	}
