Replicates 5708732 in TAA refactor to RDG: Fixes case where the TAA history can be nullptr

The shader parameter structure framework catches this issue with a failure in ValidateShaderParameters() with error message instead of a GPU hang. :)

diff --git a/Engine/Source/Runtime/Renderer/Private/PostProcess/PostProcessTemporalAA.cpp b/Engine/Source/Runtime/Renderer/Private/PostProcess/PostProcessTemporalAA.cpp
--- a/Engine/Source/Runtime/Renderer/Private/PostProcess/PostProcessTemporalAA.cpp
+++ b/Engine/Source/Runtime/Renderer/Private/PostProcess/PostProcessTemporalAA.cpp
@@ -531,6 +531,11 @@ FTAAOutputs FTAAPassParameters::AddTemporalAAPass(
 				if (InputHistory.RT[1].IsValid())
 					CommonShaderParameters.HistoryBuffer1 = GraphBuilder.RegisterExternalTexture(InputHistory.RT[1]);
 			}
+			else
+			{
+				CommonShaderParameters.HistoryBuffer0 = GraphBuilder.RegisterExternalTexture(GSystemTextures.BlackDummy);
+				CommonShaderParameters.HistoryBuffer1 = CommonShaderParameters.HistoryBuffer0;
+			}
 
 			CommonShaderParameters.HistoryBuffer0Sampler = TStaticSamplerState<SF_Bilinear>::GetRHI();
 			CommonShaderParameters.HistoryBuffer1Sampler = TStaticSamplerState<SF_Bilinear>::GetRHI();
