Cleaning of some remaining dxr_todo-s and removal of dead code that was never called. No functional changes.

diff --git a/Engine/Source/Runtime/Engine/Private/SkeletalRenderCPUSkin.cpp b/Engine/Source/Runtime/Engine/Private/SkeletalRenderCPUSkin.cpp
--- a/Engine/Source/Runtime/Engine/Private/SkeletalRenderCPUSkin.cpp
+++ b/Engine/Source/Runtime/Engine/Private/SkeletalRenderCPUSkin.cpp
@@ -400,22 +400,6 @@ void FSkeletalMeshObjectCPUSkin::FSkeletalMeshObjectLOD::ReleaseResources()
 	bResourcesInitialized = false;
 }
 
-#if RHI_RAYTRACING
-// #dxr_todo: this looks like dead code
-void FSkeletalMeshObjectCPUSkin::FSkeletalMeshObjectLOD::BuildRayTracingAccelerationStructure()
-{
-	if (RayTracingGeometry.Initializer.PositionVertexBuffer && RayTracingGeometry.Initializer.IndexBuffer)
-	{
-		ENQUEUE_RENDER_COMMAND(SkeletalRenderCPUSkinBuildRayTracingAccelerationStructure)(
-			[this](FRHICommandListImmediate& RHICmdList)
-		{
-			RHICmdList.BuildAccelerationStructure(RayTracingGeometry.RayTracingGeometryRHI);
-		}
-		);
-	}
-}
-#endif // RHI_RAYTRACING
-
 TArray<FTransform>* FSkeletalMeshObjectCPUSkin::GetComponentSpaceTransforms() const
 {
 #if !(UE_BUILD_SHIPPING || UE_BUILD_TEST)
diff --git a/Engine/Source/Runtime/Engine/Private/SkeletalRenderCPUSkin.h b/Engine/Source/Runtime/Engine/Private/SkeletalRenderCPUSkin.h
--- a/Engine/Source/Runtime/Engine/Private/SkeletalRenderCPUSkin.h
+++ b/Engine/Source/Runtime/Engine/Private/SkeletalRenderCPUSkin.h
@@ -205,11 +205,6 @@ private:
 		{
 			CumulativeResourceSize.AddUnknownMemoryBytes(StaticMeshVertexBuffer.GetResourceSize() + PositionVertexBuffer.GetStride() * PositionVertexBuffer.GetNumVertices());
 		}
-
-#if RHI_RAYTRACING
-		/** Builds ray tracing acceleration structures per LOD. */
-		void BuildRayTracingAccelerationStructure();
-#endif // RHI_RAYTRACING
 	};
 
 	/** Render data for each LOD */
diff --git a/Engine/Source/Runtime/Engine/Private/SkeletalRenderStatic.cpp b/Engine/Source/Runtime/Engine/Private/SkeletalRenderStatic.cpp
--- a/Engine/Source/Runtime/Engine/Private/SkeletalRenderStatic.cpp
+++ b/Engine/Source/Runtime/Engine/Private/SkeletalRenderStatic.cpp
@@ -186,18 +186,3 @@ void FSkeletalMeshObjectStatic::FSkeletalMeshObjectLOD::ReleaseResources()
 	bResourcesInitialized = false;
 }
 
-#if RHI_RAYTRACING
-// #dxr_todo: this looks like dead code
-void FSkeletalMeshObjectStatic::FSkeletalMeshObjectLOD::BuildRayTracingAccelerationStructure()
-{
-	if (RayTracingGeometry.Initializer.PositionVertexBuffer && RayTracingGeometry.Initializer.IndexBuffer)
-	{
-		ENQUEUE_RENDER_COMMAND(SkeletalRenderStaticBuildRayTracingAccelerationStructure)(
-			[this](FRHICommandListImmediate& RHICmdList)
-		{
-			RHICmdList.BuildAccelerationStructure(RayTracingGeometry.RayTracingGeometryRHI);
-		}
-		);
-	}
-}
-#endif // RHI_RAYTRACING
diff --git a/Engine/Source/Runtime/Engine/Private/SkeletalRenderStatic.h b/Engine/Source/Runtime/Engine/Private/SkeletalRenderStatic.h
--- a/Engine/Source/Runtime/Engine/Private/SkeletalRenderStatic.h
+++ b/Engine/Source/Runtime/Engine/Private/SkeletalRenderStatic.h
@@ -100,11 +100,6 @@ private:
  		{
 			CumulativeResourceSize.AddDedicatedSystemMemoryBytes(sizeof(*this));
  		}
-
-#if RHI_RAYTRACING
-		/** Builds ray tracing acceleration structures per LOD. */
-		void BuildRayTracingAccelerationStructure();
-#endif // RHI_RAYTRACING
 	};
 
 	/** Render data for each LOD */
diff --git a/Engine/Source/Runtime/Renderer/Private/RayTracing/RayTracingTranslucency.cpp b/Engine/Source/Runtime/Renderer/Private/RayTracing/RayTracingTranslucency.cpp
--- a/Engine/Source/Runtime/Renderer/Private/RayTracing/RayTracingTranslucency.cpp
+++ b/Engine/Source/Runtime/Renderer/Private/RayTracing/RayTracingTranslucency.cpp
@@ -213,7 +213,7 @@ public:
 		SceneTextureParameters.Set(RHICmdList, ShaderRHI, View.FeatureLevel, ESceneTextureSetupMode::All);
 
 		SetTextureParameter(RHICmdList, ShaderRHI, TranslucencyTextureParameter, TranslucencyTextureSamplerParameter, TStaticSamplerState<SF_Bilinear>::GetRHI(), TranslucencyTexture);
-		// #dxr_todo: Use hit-distance texture for denoising
+		// #dxr_todo: UE-72581 Use hit-distance texture for denoising
 	}
 
 	virtual bool Serialize(FArchive& Ar) override
@@ -242,7 +242,6 @@ void FDeferredShadingSceneRenderer::PrepareRayTracingTranslucency(const FViewInf
 
 void FDeferredShadingSceneRenderer::RenderRayTracingTranslucency(FRHICommandListImmediate& RHICmdList)
 {
-	//#dxr_todo: check DOF support, do we need to call RenderRayTracingTranslucency twice?
 	if (!ShouldRenderTranslucency(ETranslucencyPass::TPT_StandardTranslucency)
 		&& !ShouldRenderTranslucency(ETranslucencyPass::TPT_TranslucencyAfterDOF)
 		&& !ShouldRenderTranslucency(ETranslucencyPass::TPT_AllTranslucency)
@@ -264,7 +263,7 @@ void FDeferredShadingSceneRenderer::RenderRayTracingTranslucency(FRHICommandList
 		FSceneViewFamilyBlackboard SceneBlackboard;
 		SetupSceneViewFamilyBlackboard(GraphBuilder, &SceneBlackboard);
 
-		//#dxr_todo: do not use reflections denoiser structs but separated ones
+		//#dxr_todo: UE-72581 do not use reflections denoiser structs but separated ones
 		IScreenSpaceDenoiser::FReflectionsInputs DenoiserInputs;
 		float ResolutionFraction = 1.0f;
 		int32 TranslucencySPP = GRayTracingTranslucencySamplesPerPixel > -1 ? GRayTracingTranslucencySamplesPerPixel : View.FinalPostProcessSettings.RayTracingTranslucencySamplesPerPixel;
@@ -274,7 +273,7 @@ void FDeferredShadingSceneRenderer::RenderRayTracingTranslucency(FRHICommandList
 			View, &DenoiserInputs.Color, &DenoiserInputs.RayHitDistance,
 			TranslucencySPP, GRayTracingTranslucencyHeightFog, ResolutionFraction);
 
-		//#dxr_todo: denoise: replace DenoiserInputs with DenoiserOutputs in the following lines!
+		//#dxr_todo: UE-72581 : replace DenoiserInputs with DenoiserOutputs in the following lines!
 		TRefCountPtr<IPooledRenderTarget> TranslucencyColor = GSystemTextures.BlackDummy;
 		TRefCountPtr<IPooledRenderTarget> TranslucencyHitDistanceColor = GSystemTextures.BlackDummy;
 
