Fixed bloom issues auto-exposure.

diff --git a/Engine/Shaders/Private/PostProcessBloom.usf b/Engine/Shaders/Private/PostProcessBloom.usf
--- a/Engine/Shaders/Private/PostProcessBloom.usf
+++ b/Engine/Shaders/Private/PostProcessBloom.usf
@@ -159,7 +159,7 @@ float4 BloomSetupCommon(float2 UV, float Threshold, float ExposureScale)
 	// mask 0..1
 	half BloomAmount = saturate(BloomLuminance / 2.0f);
 
-	return float4(BloomAmount * LinearColor, 0);
+	return float4(BloomAmount * LinearColor, 0) * PreExposure;
 }
 
 // x:BloomThreshold, yz:unused, w:ExposureScale (useful if eyeadaptation is locked)
diff --git a/Engine/Shaders/Private/PostProcessTonemap.usf b/Engine/Shaders/Private/PostProcessTonemap.usf
--- a/Engine/Shaders/Private/PostProcessTonemap.usf
+++ b/Engine/Shaders/Private/PostProcessTonemap.usf
@@ -463,6 +463,8 @@ float4 TonemapCommonPS(
 				BloomUV = FullViewUV.xy;
 			#endif
 			float4 CombinedBloom = Texture2DSample(PostprocessInput1, PostprocessInput1Sampler, BloomUV);
+			CombinedBloom.rgb *= OneOverPreExposure; // Standard bloom path, convolution path has bloom already in scene color.
+
 			float2 DirtViewportUV = (SvPosition.xy - float2(ViewportRect.xy)) * ViewportSize.zw;
 			float3 BloomDirtMaskColor = Texture2DSample(BloomDirtMask.Mask, BloomDirtMask.MaskSampler, DirtViewportUV).rgb * BloomDirtMask.Tint.rgb;
 
diff --git a/Engine/Source/Runtime/Renderer/Private/PostProcess/PostProcessFFTBloom.cpp b/Engine/Source/Runtime/Renderer/Private/PostProcess/PostProcessFFTBloom.cpp
--- a/Engine/Source/Runtime/Renderer/Private/PostProcess/PostProcessFFTBloom.cpp
+++ b/Engine/Source/Runtime/Renderer/Private/PostProcess/PostProcessFFTBloom.cpp
@@ -801,8 +801,8 @@ FSceneRenderTargetItem* FRCPassFFTBloom::InitDomainAndGetKernel(FRenderingCompos
 	const float ClampedBloomConvolutionBufferScale = FMath::Clamp(PPSettings.BloomConvolutionBufferScale, 0.f, 1.f);
 	
 
-	// The pre-filter boost parameters for bright pixels
-	const FVector PreFilter(PPSettings.BloomConvolutionPreFilterMin, PPSettings.BloomConvolutionPreFilterMax, PPSettings.BloomConvolutionPreFilterMult);
+	// The pre-filter boost parameters for bright pixels. Because the Convolution PP work in pre-exposure space, the min and max needs adjustment.
+	const FVector PreFilter(PPSettings.BloomConvolutionPreFilterMin * View.PreExposure, PPSettings.BloomConvolutionPreFilterMax * View.PreExposure, PPSettings.BloomConvolutionPreFilterMult);
 
 	
 	// Clip the Kernel support (i.e. bloom size) to 100% the screen width 
@@ -952,8 +952,8 @@ bool FRCPassFFTBloom::ConvolveImageWithKernel(FRenderingCompositePassContext& Co
 
 	const FViewInfo& View = Context.View;
 	const auto& FinalPPSettings = View.FinalPostProcessSettings;
-	// The pre-filter boost parameters for bright pixels
-	const FVector PreFilter(FinalPPSettings.BloomConvolutionPreFilterMin, FinalPPSettings.BloomConvolutionPreFilterMax, FinalPPSettings.BloomConvolutionPreFilterMult);
+	// The pre-filter boost parameters for bright pixels. Because the Convolution PP work in pre-exposure space, the min and max needs adjustment.
+	const FVector PreFilter(FinalPPSettings.BloomConvolutionPreFilterMin * View.PreExposure, FinalPPSettings.BloomConvolutionPreFilterMax * View.PreExposure, FinalPPSettings.BloomConvolutionPreFilterMult);
 
 	const FLinearColor Tint(1, 1, 1, 1);
 
