DR - Fix Android warnings
#rnx

diff --git a/Engine/Source/Runtime/Renderer/Private/ClusteredDeferredShadingPass.cpp b/Engine/Source/Runtime/Renderer/Private/ClusteredDeferredShadingPass.cpp
--- a/Engine/Source/Runtime/Renderer/Private/ClusteredDeferredShadingPass.cpp
+++ b/Engine/Source/Runtime/Renderer/Private/ClusteredDeferredShadingPass.cpp
@@ -140,7 +140,7 @@ void FDeferredShadingSceneRenderer::AddClusteredDeferredShadingPass(FRHICommandL
 				RDG_EVENT_NAME("ClusteredDeferredShading, #Lights: %d", NumLightsToRender),
 				PassParameters,
 				ERenderGraphPassFlags::None,
-				[PassParameters, &View, &SceneContext](FRHICommandListImmediate& RHICmdList)
+				[PassParameters, &View, &SceneContext](FRHICommandListImmediate& InRHICmdList)
 			{
 				TShaderMapRef<FPostProcessVS> VertexShader(View.ShaderMap);
 
@@ -161,13 +161,13 @@ void FDeferredShadingSceneRenderer::AddClusteredDeferredShadingPass(FRHICommandL
 					GraphicsPSOInit.BoundShaderState.PixelShaderRHI = GETSAFERHISHADER_PIXEL(*PixelShader);
 					GraphicsPSOInit.PrimitiveType = PT_TriangleList;
 
-					SetGraphicsPipelineState(RHICmdList, GraphicsPSOInit, EApplyRendertargetOption::ForceApply);
+					SetGraphicsPipelineState(InRHICmdList, GraphicsPSOInit, EApplyRendertargetOption::ForceApply);
 				}
-				RHICmdList.SetViewport(View.ViewRect.Min.X, View.ViewRect.Min.Y, 0.0f, View.ViewRect.Max.X, View.ViewRect.Max.Y, 1.0f);
+				InRHICmdList.SetViewport(View.ViewRect.Min.X, View.ViewRect.Min.Y, 0.0f, View.ViewRect.Max.X, View.ViewRect.Max.Y, 1.0f);
 
-				SetShaderParameters(RHICmdList, *PixelShader, PixelShader->GetPixelShader(), *PassParameters);
+				SetShaderParameters(InRHICmdList, *PixelShader, PixelShader->GetPixelShader(), *PassParameters);
 
-				DrawRectangle(RHICmdList, 0, 0, View.ViewRect.Width(), View.ViewRect.Height(),
+				DrawRectangle(InRHICmdList, 0, 0, View.ViewRect.Width(), View.ViewRect.Height(),
 					View.ViewRect.Min.X, View.ViewRect.Min.Y, View.ViewRect.Width(), View.ViewRect.Height(),
 					FIntPoint(View.ViewRect.Width(), View.ViewRect.Height()), SceneContext.GetBufferSizeXY(), *VertexShader);
 			});
diff --git a/Engine/Source/Runtime/Renderer/Private/PostProcess/PostProcessTemporalAA.cpp b/Engine/Source/Runtime/Renderer/Private/PostProcess/PostProcessTemporalAA.cpp
--- a/Engine/Source/Runtime/Renderer/Private/PostProcess/PostProcessTemporalAA.cpp
+++ b/Engine/Source/Runtime/Renderer/Private/PostProcess/PostProcessTemporalAA.cpp
@@ -756,14 +756,14 @@ FTAAOutputs FTAAPassParameters::AddTemporalAAPass(
 			// Lambda to draw pixel shader.
 			auto DrawTAAPixelShader = [&](FDepthStencilStateRHIParamRef DepthStencilState)
 			{
-				TShaderMapRef<FTemporalAAPS> PixelShader(View.ShaderMap, PermutationVector);
+				TShaderMapRef<FTemporalAAPS> TAAPixelShader(View.ShaderMap, PermutationVector);
 
 				FGraphicsPipelineStateInitializer GraphicsPSOInit;
-				FPixelShaderUtils::InitFullscreenPipelineState(RHICmdList, View.ShaderMap, *PixelShader, GraphicsPSOInit);
+				FPixelShaderUtils::InitFullscreenPipelineState(RHICmdList, View.ShaderMap, *TAAPixelShader, GraphicsPSOInit);
 				GraphicsPSOInit.DepthStencilState = DepthStencilState;
 				SetGraphicsPipelineState(RHICmdList, GraphicsPSOInit);
 
-				SetShaderParameters(RHICmdList, *PixelShader, PixelShader->GetPixelShader(), *PassParameters);
+				SetShaderParameters(RHICmdList, *TAAPixelShader, TAAPixelShader->GetPixelShader(), *PassParameters);
 				
 				FPixelShaderUtils::DrawFullscreenTriangle(RHICmdList);
 			};
