Re-enabling pre-exposure through the override.

diff --git a/Engine/Source/Runtime/Renderer/Private/PostProcess/PostProcessEyeAdaptation.cpp b/Engine/Source/Runtime/Renderer/Private/PostProcess/PostProcessEyeAdaptation.cpp
--- a/Engine/Source/Runtime/Renderer/Private/PostProcess/PostProcessEyeAdaptation.cpp
+++ b/Engine/Source/Runtime/Renderer/Private/PostProcess/PostProcessEyeAdaptation.cpp
@@ -19,7 +19,7 @@ RENDERCORE_API bool UsePreExposure(EShaderPlatform Platform);
 // Initialize the static CVar
 TAutoConsoleVariable<float> CVarEyeAdaptationPreExposureOverride(
 	TEXT("r.EyeAdaptation.PreExposureOverride"),
-	1,
+	0,
 	TEXT("Overide the scene pre-exposure by a custom value. \n")
 	TEXT("= 0 : No override\n")
 	TEXT("> 0 : Override PreExposure\n"),
