Added a small comment as to why we are limited to GMaxNumReflectionCaptures = 341.

diff --git a/Engine/Source/Runtime/Renderer/Private/SceneRendering.h b/Engine/Source/Runtime/Renderer/Private/SceneRendering.h
--- a/Engine/Source/Runtime/Renderer/Private/SceneRendering.h
+++ b/Engine/Source/Runtime/Renderer/Private/SceneRendering.h
@@ -563,6 +563,7 @@ struct FMeshDecalBatch
 	}
 };
 
+// DX11 maximum 2d texture array size is D3D11_REQ_TEXTURE2D_ARRAY_AXIS_DIMENSION = 2048, and 2048/6 = 341.33.
 static const int32 GMaxNumReflectionCaptures = 341;
 
 /** Per-reflection capture data needed by the shader. */
