Fully disable ShaderPrint below SM5
Previously we tried to create a StructuredBuffer without checking platform support

diff --git a/Engine/Source/Runtime/Renderer/Private/PostProcess/PostProcessShaderPrint.cpp b/Engine/Source/Runtime/Renderer/Private/PostProcess/PostProcessShaderPrint.cpp
--- a/Engine/Source/Runtime/Renderer/Private/PostProcess/PostProcessShaderPrint.cpp
+++ b/Engine/Source/Runtime/Renderer/Private/PostProcess/PostProcessShaderPrint.cpp
@@ -4,9 +4,9 @@
 
 #include "ShaderPrint.h"
 
-bool FRCPassPostProcessShaderPrint::IsEnabled()
+bool FRCPassPostProcessShaderPrint::IsEnabled(FViewInfo const& View)
 {
-	return ShaderPrint::IsEnabled();
+	return ShaderPrint::IsEnabled() && ShaderPrint::IsSupported(View);
 }
 
 void FRCPassPostProcessShaderPrint::Process(FRenderingCompositePassContext& Context)
diff --git a/Engine/Source/Runtime/Renderer/Private/PostProcess/PostProcessShaderPrint.h b/Engine/Source/Runtime/Renderer/Private/PostProcess/PostProcessShaderPrint.h
--- a/Engine/Source/Runtime/Renderer/Private/PostProcess/PostProcessShaderPrint.h
+++ b/Engine/Source/Runtime/Renderer/Private/PostProcess/PostProcessShaderPrint.h
@@ -10,7 +10,7 @@
 class FRCPassPostProcessShaderPrint : public TRenderingCompositePassBase<1, 1>
 {
 public:
-	static bool IsEnabled();
+	static bool IsEnabled(FViewInfo const& View);
 
 	virtual bool FrameBufferBlendingWithInput0() const override { return true; }
 	virtual void Process(FRenderingCompositePassContext& Context) override;
diff --git a/Engine/Source/Runtime/Renderer/Private/PostProcess/PostProcessing.cpp b/Engine/Source/Runtime/Renderer/Private/PostProcess/PostProcessing.cpp
--- a/Engine/Source/Runtime/Renderer/Private/PostProcess/PostProcessing.cpp
+++ b/Engine/Source/Runtime/Renderer/Private/PostProcess/PostProcessing.cpp
@@ -2161,7 +2161,7 @@ void FPostProcessing::Process(FRHICommandListImmediate& RHICmdList, const FViewI
 			Context.FinalOutput = FRenderingCompositeOutputRef(Node);
 		}
 
-		if (FRCPassPostProcessShaderPrint::IsEnabled())
+		if (FRCPassPostProcessShaderPrint::IsEnabled(View))
 		{
 			FRenderingCompositePass* Node = Context.Graph.RegisterPass(new(FMemStack::Get()) FRCPassPostProcessShaderPrint());
 			Node->SetInput(ePId_Input0, FRenderingCompositeOutputRef(Context.FinalOutput));
diff --git a/Engine/Source/Runtime/Renderer/Private/ShaderPrint.cpp b/Engine/Source/Runtime/Renderer/Private/ShaderPrint.cpp
--- a/Engine/Source/Runtime/Renderer/Private/ShaderPrint.cpp
+++ b/Engine/Source/Runtime/Renderer/Private/ShaderPrint.cpp
@@ -143,6 +143,22 @@ namespace ShaderPrint
 		return Ar;
 	}
 
+	// Supported platforms
+	bool IsSupported(EShaderPlatform InShaderPlatform)
+	{
+		return RHISupportsComputeShaders(InShaderPlatform) && !IsHlslccShaderPlatform(InShaderPlatform);
+	}
+
+	bool IsSupported(FViewInfo const& View)
+	{
+		return IsSupported(View.GetShaderPlatform());
+	}
+
+	bool IsEnabled()
+	{
+		return CVarEnable.GetValueOnAnyThread();
+	}
+
 	// Shader to initialize the output value buffer
 	class FShaderInitValueBufferCS : public FGlobalShader
 	{
@@ -156,7 +172,7 @@ namespace ShaderPrint
 
 		static bool ShouldCompilePermutation(FGlobalShaderPermutationParameters const& Parameters)
 		{
-			return RHISupportsComputeShaders(Parameters.Platform) && !IsHlslccShaderPlatform(Parameters.Platform);
+			return IsSupported(Parameters.Platform);
 		}
 	};
 
@@ -178,7 +194,7 @@ namespace ShaderPrint
 
 		static bool ShouldCompilePermutation(FGlobalShaderPermutationParameters const& Parameters)
 		{
-			return RHISupportsComputeShaders(Parameters.Platform) && !IsHlslccShaderPlatform(Parameters.Platform);
+			return IsSupported(Parameters.Platform);
 		}
 	};
 
@@ -200,7 +216,7 @@ namespace ShaderPrint
 
 		static bool ShouldCompilePermutation(FGlobalShaderPermutationParameters const& Parameters)
 		{
-			return RHISupportsComputeShaders(Parameters.Platform) && !IsHlslccShaderPlatform(Parameters.Platform);
+			return IsSupported(Parameters.Platform);
 		}
 	};
 
@@ -221,7 +237,7 @@ namespace ShaderPrint
 
 		static bool ShouldCompilePermutation(FGlobalShaderPermutationParameters const& Parameters)
 		{
-			return RHISupportsComputeShaders(Parameters.Platform) && !IsHlslccShaderPlatform(Parameters.Platform);
+			return IsSupported(Parameters.Platform);
 		}
 	};
 
@@ -243,7 +259,7 @@ namespace ShaderPrint
 
 		static bool ShouldCompilePermutation(FGlobalShaderPermutationParameters const& Parameters)
 		{
-			return IsFeatureLevelSupported(Parameters.Platform, ERHIFeatureLevel::SM4) && !IsHlslccShaderPlatform(Parameters.Platform);
+			return IsSupported(Parameters.Platform);
 		}
 	};
 
@@ -278,13 +294,13 @@ namespace ShaderPrint
 	IMPLEMENT_SHADER_TYPE(, FShaderDrawSymbolsPS, TEXT("/Engine/Private/ShaderPrintDraw.usf"), TEXT("DrawSymbolsPS"), SF_Pixel)
 
 
-	bool IsEnabled()
-	{
-		return CVarEnable.GetValueOnAnyThread();
-	}
-
 	void BeginView(FRHICommandListImmediate& RHICmdList, FViewInfo& View)
 	{
+		if (!IsSupported(View))
+		{
+			return;
+		}
+
 		// Initialize output buffer and store in the view info
 		// Values buffer contains Count + 1 elements. The first element is only used as a counter.
 		View.ShaderPrintValueBuffer.Initialize(sizeof(ShaderPrintItem), GetMaxValueCount() + 1, 0U, TEXT("ShaderPrintValueBuffer"));
@@ -314,7 +330,7 @@ namespace ShaderPrint
 	void DrawView(FRHICommandListImmediate& RHICmdList, FViewInfo const& View, TRefCountPtr<IPooledRenderTarget>& OutputTarget)
 	{
 		// Early out if system is disabled
-		if (!IsEnabled())
+		if (!IsEnabled() || !IsSupported(View))
 		{
 			return;
 		}
diff --git a/Engine/Source/Runtime/Renderer/Private/ShaderPrint.h b/Engine/Source/Runtime/Renderer/Private/ShaderPrint.h
--- a/Engine/Source/Runtime/Renderer/Private/ShaderPrint.h
+++ b/Engine/Source/Runtime/Renderer/Private/ShaderPrint.h
@@ -19,6 +19,8 @@ struct IPooledRenderTarget;
 
 namespace ShaderPrint
 {
+	// Does the platform support the ShaderPrint system?
+	bool IsSupported(FViewInfo const& View);
 	// Have we enabled the ShaderPrint system?
 	bool IsEnabled();
 
