Rewrites Diaphragm DOF to RDG.

It's implemented as separate code path to have few CL in Dev-Rendering with both old and new in case there is suspition this refactored introduced a regression in the futur.

diff --git a/Engine/Shaders/Private/DiaphragmDOF/DOFBokehLUT.usf b/Engine/Shaders/Private/DiaphragmDOF/DOFBokehLUT.usf
--- a/Engine/Shaders/Private/DiaphragmDOF/DOFBokehLUT.usf
+++ b/Engine/Shaders/Private/DiaphragmDOF/DOFBokehLUT.usf
@@ -153,7 +153,11 @@ void RingCoordinateFromCartesianCoordinate(const float2 C, out float OutRingRadi
 
 //------------------------------------------------------- OUTPUTS
 
-RWTexture2D<float4>	Output0;
+#if !defined(NEW_CODE) // TODO: remove me
+#define BokehLUTOutput Output0
+#endif
+
+RWTexture2D<float4>	BokehLUTOutput;
 
 
 //------------------------------------------------------- ENTRY POINT
@@ -190,7 +194,7 @@ void BuildBokehLUTMainCS(uint2 DispatchThreadId : SV_DispatchThreadID)
 
 		float CocRadiusToBokehEdgeFactor = ComputeCocRadiusToBokehEdgeFactor(PixelAngle);
 
-		Output0[DispatchThreadId] = float4(CocRadiusToBokehEdgeFactor, 0, 0, 0);
+		BokehLUTOutput[DispatchThreadId] = float4(CocRadiusToBokehEdgeFactor, 0, 0, 0);
 	}
 	#elif DIM_LUT_FORMAT == LUT_FORMAT_FULL_RES_TO_COC_DISTANCE
 	{
@@ -203,7 +207,7 @@ void BuildBokehLUTMainCS(uint2 DispatchThreadId : SV_DispatchThreadID)
 		}
 
 		float PixelCocDistanceFromBokehCenter = FullResPixelDistanceToCocDistance(length(RelativePosFromCenter) / CocRadiusToBokehEdgeFactor);
-		Output0[DispatchThreadId] = float4(PixelCocDistanceFromBokehCenter, 0, 0, 0);
+		BokehLUTOutput[DispatchThreadId] = float4(PixelCocDistanceFromBokehCenter, 0, 0, 0);
 	}
 	#elif DIM_LUT_FORMAT == LUT_FORMAT_GATHER_SAMPLE_POS
 	{
@@ -259,7 +263,7 @@ void BuildBokehLUTMainCS(uint2 DispatchThreadId : SV_DispatchThreadID)
 		// Final position of the sample in kernel.
 		float2 SamplePos = KernelUnitPos * RingRadius;
 
-		Output0[DispatchThreadId] = float4(SamplePos, 0, 0);
+		BokehLUTOutput[DispatchThreadId] = float4(SamplePos, 0, 0);
 	}
 	#else
 		#error Unknown LUT format.
diff --git a/Engine/Shaders/Private/DiaphragmDOF/DOFCocTileDilate.usf b/Engine/Shaders/Private/DiaphragmDOF/DOFCocTileDilate.usf
--- a/Engine/Shaders/Private/DiaphragmDOF/DOFCocTileDilate.usf
+++ b/Engine/Shaders/Private/DiaphragmDOF/DOFCocTileDilate.usf
@@ -25,6 +25,23 @@ float fSampleOffsetMultipler = 1;
 float CocRadiusToBucketDistanceUpperBound;
 float BucketDistanceToCocRadius;
 
+Texture2D TileInput_Foreground;
+Texture2D TileInput_Background;
+
+Texture2D DilatedTileMinMax_Foreground;
+Texture2D DilatedTileMinMax_Background;
+
+#if defined(NEW_CODE) // TODO: remove me
+#define Output0 TileOutput_Foreground
+#define Output1 TileOutput_Background
+#else
+#define TileInput_Foreground PostprocessInput0
+#define TileInput_Background PostprocessInput1
+
+#define DilatedTileMinMax_Foreground PostprocessInput2
+#define DilatedTileMinMax_Background PostprocessInput3
+#endif
+
 
 //------------------------------------------------------- OUTPUTS.
 
@@ -57,7 +74,7 @@ void CocDilateMainCS(uint2 DispatchThreadId : SV_DispatchThreadID)
 
 	// Fetches the already dilated min foreground and max background COC radius.
 	#if DIM_DILATE_MODE == DILATE_MODE_MIN_ABS
-		FCocTileSample DilatedMinMaxSample = LoadCocTile(CocTileLayout, PostprocessInput2, PostprocessInput3, PixelPos);
+		FCocTileSample DilatedMinMaxSample = LoadCocTile(CocTileLayout, DilatedTileMinMax_Foreground, DilatedTileMinMax_Background, PixelPos);
 	#endif
 
 	// Iterate over all input tiles and assign them into their buckets.
@@ -128,7 +145,7 @@ void CocDilateMainCS(uint2 DispatchThreadId : SV_DispatchThreadID)
 					//ClosestDistance = fRadiusMultiplier * float((BucketId + 1) * COC_TILE_SIZE) - COC_TILE_SIZE;
 
 					// Sample tile.
-					const FCocTileSample NeightborSample = LoadCocTile(CocTileLayout, PostprocessInput0, PostprocessInput1, SampleIndex);
+					const FCocTileSample NeightborSample = LoadCocTile(CocTileLayout, TileInput_Foreground, TileInput_Background, SampleIndex);
 
 					OutSampleBuckets[BucketId].FgdMinCoc = min(
 						OutSampleBuckets[BucketId].FgdMinCoc,
@@ -164,7 +181,7 @@ void CocDilateMainCS(uint2 DispatchThreadId : SV_DispatchThreadID)
 	} // for (uint BucketId = 0; BucketId < Radius; BucketId++)
 	
 	// Sample the center bucket.
-	FCocTileSample OutSample = LoadCocTile(CocTileLayout, PostprocessInput0, PostprocessInput1, PixelPos);
+	FCocTileSample OutSample = LoadCocTile(CocTileLayout, TileInput_Foreground, TileInput_Background, PixelPos);
 
 	#if DIM_DILATE_MODE == DILATE_MODE_MIN_ABS
 	{
diff --git a/Engine/Shaders/Private/DiaphragmDOF/DOFCocTileFlatten.usf b/Engine/Shaders/Private/DiaphragmDOF/DOFCocTileFlatten.usf
--- a/Engine/Shaders/Private/DiaphragmDOF/DOFCocTileFlatten.usf
+++ b/Engine/Shaders/Private/DiaphragmDOF/DOFCocTileFlatten.usf
@@ -26,13 +26,25 @@
 
 float2 ThreadIdToBufferUV;
 float2 MaxBufferUV;
-float PreProcessingToProcessingCocRadiusFactor;
+
+Texture2D GatherInput_SceneColor;
+Texture2D GatherInput_SeparateCoc;
+
+#if defined(NEW_CODE) // TODO: remove me
+#define Output0 TileOutput_Foreground
+#define Output1 TileOutput_Background
+#else
+#define GatherInput_SceneColor PostprocessInput0
+#define GatherInput_SeparateCoc PostprocessInput0
+#endif
 
 
 //------------------------------------------------------- OUTPUTS.
 
 RWTexture2D<float4>	Output0;
 RWTexture2D<float4>	Output1;
+//RWTexture2D<float4>	TileOutput_Foreground;
+//RWTexture2D<float4>	TileOutput_Background;
 
 
 //------------------------------------------------------- LDS.
@@ -136,9 +148,9 @@ void CocFlattenMainCS(
 		
 		#if CONFIG_DOF_ALPHA
 			// The coc radius is within separate buffer's red channel.
-			float4 CocRadii = PostprocessInput0.GatherRed(GlobalPointClampedSampler, BufferUV, 0);
+			float4 CocRadii = GatherInput_SeparateCoc.GatherRed(GlobalPointClampedSampler, BufferUV, 0);
 		#else
-			float4 CocRadii = PostprocessInput0.GatherAlpha(GlobalPointClampedSampler, BufferUV, 0);
+			float4 CocRadii = GatherInput_SceneColor.GatherAlpha(GlobalPointClampedSampler, BufferUV, 0);
 		#endif
 
 		InitCocTileSampleFromCocRadius(CocRadii.x, TileSample);
@@ -157,9 +169,9 @@ void CocFlattenMainCS(
 
 		#if CONFIG_DOF_ALPHA
 			// The coc radius is within separate buffer's red channel.
-			float CocRadius = PostprocessInput0[PixelPos].r;
+			float CocRadius = GatherInput_SeparateCoc[PixelPos].r;
 		#else
-			float CocRadius = PostprocessInput0[PixelPos].w;
+			float CocRadius = GatherInput_SceneColor[PixelPos].a;
 		#endif
 
 		InitCocTileSampleFromCocRadius(CocRadius, TileSample);
diff --git a/Engine/Shaders/Private/DiaphragmDOF/DOFDownsample.usf b/Engine/Shaders/Private/DiaphragmDOF/DOFDownsample.usf
--- a/Engine/Shaders/Private/DiaphragmDOF/DOFDownsample.usf
+++ b/Engine/Shaders/Private/DiaphragmDOF/DOFDownsample.usf
@@ -12,6 +12,21 @@
 float2 MaxBufferUV;
 float OutputCocRadiusMultiplier;
 
+float4 GatherInputSize;
+Texture2D GatherInput_SceneColor;
+Texture2D GatherInput_SeparateCoc;
+
+#if defined(NEW_CODE) // TODO: remove me
+#define Output0 OutDownsampledGatherInput_SceneColor
+#define Output1 OutDownsampledGatherInput_SeparateCoc
+
+#else
+#define GatherInputSize PostprocessInput0Size
+#define GatherInput_SceneColor PostprocessInput0
+#define GatherInput_SeparateCoc PostprocessInput1
+
+#endif
+
 
 //------------------------------------------------------- OUTPUTS
 
@@ -29,7 +44,7 @@ void DownsampleCS(uint2 DispatchThreadId : SV_DispatchThreadID)
 {
 	// Compute scene buffer uv, but carefully to handle view size % 2 == 1,
 	// because this pass is running at half resolution.
-	float2 BufferUV = (2.0 * (DispatchThreadId + 0.5)) * PostprocessInput0Size.zw;
+	float2 BufferUV = (2.0 * (DispatchThreadId + 0.5)) * GatherInputSize.zw;
 	
 	// Sample scene color.
 	float4 SceneColors[4];
@@ -37,15 +52,15 @@ void DownsampleCS(uint2 DispatchThreadId : SV_DispatchThreadID)
 
 	for (uint i = 0; i < 4; i++)
 	{
-		float2 SampleUV = BufferUV + PostprocessInput0Size.zw * (0.5 * float2(kOffsetsCross3x3[i]));
+		float2 SampleUV = BufferUV + GatherInputSize.zw * (0.5 * float2(kOffsetsCross3x3[i]));
 
 		SampleUV = min(SampleUV, MaxBufferUV);
 
-		float4 RawSample0 = PostprocessInput0.SampleLevel(GlobalPointClampedSampler, SampleUV, 0);
+		float4 RawSample0 = GatherInput_SceneColor.SampleLevel(GlobalPointClampedSampler, SampleUV, 0);
 
 		#if CONFIG_DOF_ALPHA
 		{
-			float4 RawSample1 = PostprocessInput1.SampleLevel(GlobalPointClampedSampler, SampleUV, 0);
+			float4 RawSample1 = GatherInput_SeparateCoc.SampleLevel(GlobalPointClampedSampler, SampleUV, 0);
 
 			SceneColors[i] = RawSample0;
 			CocRadii[i] = RawSample1.r;
diff --git a/Engine/Shaders/Private/DiaphragmDOF/DOFGatherKernel.ush b/Engine/Shaders/Private/DiaphragmDOF/DOFGatherKernel.ush
--- a/Engine/Shaders/Private/DiaphragmDOF/DOFGatherKernel.ush
+++ b/Engine/Shaders/Private/DiaphragmDOF/DOFGatherKernel.ush
@@ -10,6 +10,23 @@
 #include "DOFGatherAccumulator.ush"
 
 
+//------------------------------------------------------- PARAMETERS
+
+float4 GatherInputSize;
+Texture2D GatherInput_SceneColor;
+Texture2D GatherInput_SeparateCoc;
+
+Texture2D BokehLUT;
+
+
+#if !defined(NEW_CODE) // TODO: remove me
+#define GatherInputSize PostprocessInput0Size
+#define GatherInput_SceneColor PostprocessInput0
+#define GatherInput_SeparateCoc PostprocessInput1
+#define BokehLUT PostprocessInput4
+#endif
+
+
 //------------------------------------------------------- GATHERING KERNEL
 
 // LARGEST_RINGS_FIRST controles the order the rings must be proceded.
@@ -188,16 +205,16 @@ FGatherSample FetchGatherSample(in FPrivateKernelSharedScope KernelScope, float2
 
 	if (KernelScope.bBilinearSampleInput)
 	{
-		RawSample0 = PostprocessInput0.SampleLevel(
+		RawSample0 = GatherInput_SceneColor.SampleLevel(
 			GlobalBilinearClampedSampler, SampleUV, KernelScope.GatherInputMipLevel);
-		RawSample1 = PostprocessInput1.SampleLevel(
+		RawSample1 = GatherInput_SeparateCoc.SampleLevel(
 			GlobalBilinearClampedSampler, SampleUV, KernelScope.GatherInputMipLevel);
 	}
 	else
 	{
-		RawSample0 = PostprocessInput0.SampleLevel(
+		RawSample0 = GatherInput_SceneColor.SampleLevel(
 			GlobalPointClampedSampler, SampleUV, KernelScope.GatherInputMipLevel);
-		RawSample1 = PostprocessInput1.SampleLevel(
+		RawSample1 = GatherInput_SeparateCoc.SampleLevel(
 			GlobalPointClampedSampler, SampleUV, KernelScope.GatherInputMipLevel);
 	}
 
@@ -281,7 +298,7 @@ void UpdateRuntimeKernelScope(
 		KernelRandomRange = 0;
 	}
 
-	KernelScope.KernelCenterUV = GatherParameters.InputBufferCenterUV + (RandomCircle * KernelRandomRange) * PostprocessInput0Size.zw;
+	KernelScope.KernelCenterUV = GatherParameters.InputBufferCenterUV + (RandomCircle * KernelRandomRange) * GatherInputSize.zw;
 
 	KernelScope.SampleCountToSamplePosition = GatherParameters.KernelRadius * UnitSampleRadius;
 
@@ -466,12 +483,12 @@ void GatherRingSamples(
 					0.5 + (0.5 - CONFIG_POINT_MIRROR_BOKEH * SquarePos) / float(BOKEH_LUT_SIZE),
 				};
 				
-				SamplePosOffset[0] = KernelScope.SampleCountToSamplePosition * PostprocessInput4.SampleLevel(GlobalBilinearWrappedSampler, LookUpUV[0], 0).xy;
+				SamplePosOffset[0] = KernelScope.SampleCountToSamplePosition * BokehLUT.SampleLevel(GlobalBilinearWrappedSampler, LookUpUV[0], 0).xy;
 
 				#if DIM_BOKEH_SIMULATION == BOKEH_SIMULATION_SIMMETRIC && 0
 					SamplePosOffset[1] = -SamplePosOffset[0];
 				#else
-					SamplePosOffset[1] = KernelScope.SampleCountToSamplePosition * PostprocessInput4.SampleLevel(
+					SamplePosOffset[1] = KernelScope.SampleCountToSamplePosition * BokehLUT.SampleLevel(
 						GlobalBilinearWrappedSampler, LookUpUV[1], 0).xy;
 				#endif
 				
@@ -559,7 +576,7 @@ void GatherRingSamples(
 						#error Unknown bokeh simulation.
 					#endif
 
-					float4 LookupSample = PostprocessInput4.SampleLevel(GlobalBilinearWrappedSampler, LookUpUV, 0);
+					float4 LookupSample = BokehLUT.SampleLevel(GlobalBilinearWrappedSampler, LookUpUV, 0);
 					float CocRadiusToBokehEdgeDistance = LookupSample.x;
 
 					Sample[k].CocRadius *= CocRadiusToBokehEdgeDistance;
@@ -638,7 +655,7 @@ uint GatherSamplesAndResolve(
 
 	// Set up initial compile time stuf in the scope of the kernel.
 	{
-		KernelScope.GatherInputInvSize = PostprocessInput0Size.zw;
+		KernelScope.GatherInputInvSize = GatherInputSize.zw;
 		KernelScope.bKenrelHasEnoughSamples = DIM_LAYER_PROCESSING == LAYER_PROCESSING_SLIGHT_OUT_OF_FOCUS;
 		KernelScope.bBilinearSampleInput = (
 			static_condition(GatherParameters.bFastGathering) &&
diff --git a/Engine/Shaders/Private/DiaphragmDOF/DOFGatherPass.usf b/Engine/Shaders/Private/DiaphragmDOF/DOFGatherPass.usf
--- a/Engine/Shaders/Private/DiaphragmDOF/DOFGatherPass.usf
+++ b/Engine/Shaders/Private/DiaphragmDOF/DOFGatherPass.usf
@@ -35,14 +35,26 @@ float2 ConsiderAbsCocRadiusAffineTransformation;
 
 float2 InputBufferUVToOutputPixel;
 
+Texture2D TileClassification_Foreground;
+Texture2D TileClassification_Background;
+
 
 //------------------------------------------------------- OUTPUTS
 
-RWTexture2D<float4>	Output0;
-RWTexture2D<float4>	Output1;
+#if !defined(NEW_CODE) // TODO: remove me
+#define ConvolutionOutput_SceneColor Output0
+#define ConvolutionOutput_SeparateAlpha Output1
+#define ScatterOcclusionOutput Output2
+
+#define TileClassification_Foreground PostprocessInput2
+#define TileClassification_Background PostprocessInput3
+#endif
+
+RWTexture2D<float4>	ConvolutionOutput_SceneColor;
+RWTexture2D<float4>	ConvolutionOutput_SeparateAlpha;
 
 #if CONFIG_GENERATE_SCATTER_OCCLUSION_BUFFER
-RWTexture2D<float4>	Output2;
+RWTexture2D<float4>	ScatterOcclusionOutput;
 #endif
 
 
@@ -79,20 +91,20 @@ void StoreAccumulatorOutput(in uint2 OutputPixelPosition, in FAccumulatorOutput
 	#if CONFIG_DOF_ALPHA
 	{
 		#if DIM_LAYER_PROCESSING == LAYER_PROCESSING_FOREGROUND_ONLY
-			Output0[OutputPixelPosition] = AccumulatorOutput.ForegroundColor * AccumulatorOutput.ForegroundAlpha;
-			Output1[OutputPixelPosition] = AccumulatorOutput.ForegroundAlpha;
+			ConvolutionOutput_SceneColor[OutputPixelPosition] = AccumulatorOutput.ForegroundColor * AccumulatorOutput.ForegroundAlpha;
+			ConvolutionOutput_SeparateAlpha[OutputPixelPosition] = AccumulatorOutput.ForegroundAlpha;
 
 		#elif DIM_LAYER_PROCESSING == LAYER_PROCESSING_FOREGROUND_HOLE_FILLING
-			Output0[OutputPixelPosition] = AccumulatorOutput.ForegroundHoleFillingColor;
-			Output1[OutputPixelPosition] = AccumulatorOutput.ForegroundHoleFillingAlpha;
+			ConvolutionOutput_SceneColor[OutputPixelPosition] = AccumulatorOutput.ForegroundHoleFillingColor;
+			ConvolutionOutput_SeparateAlpha[OutputPixelPosition] = AccumulatorOutput.ForegroundHoleFillingAlpha;
 
 		#elif DIM_LAYER_PROCESSING == LAYER_PROCESSING_BACKGROUND_ONLY
-			Output0[OutputPixelPosition] = AccumulatorOutput.BackgroundColor;
-			Output1[OutputPixelPosition] = AccumulatorOutput.BackgroundWeight > 0 ? 1 : 0;
+			ConvolutionOutput_SceneColor[OutputPixelPosition] = AccumulatorOutput.BackgroundColor;
+			ConvolutionOutput_SeparateAlpha[OutputPixelPosition] = AccumulatorOutput.BackgroundWeight > 0 ? 1 : 0;
 			//Output1[OutputPixelPosition] = float4(1, 0, 0, 0);
 		
 		#elif DIM_LAYER_PROCESSING == LAYER_PROCESSING_SLIGHT_OUT_OF_FOCUS
-			Output0[OutputPixelPosition] = float4(
+			ConvolutionOutput_SceneColor[OutputPixelPosition] = float4(
 				AccumulatorOutput.SlightFocusColor.rgb,
 				AccumulatorOutput.SlightFocusOpacity);
 		
@@ -104,22 +116,22 @@ void StoreAccumulatorOutput(in uint2 OutputPixelPosition, in FAccumulatorOutput
 	#else // !CONFIG_DOF_ALPHA
 	{
 		#if DIM_LAYER_PROCESSING == LAYER_PROCESSING_FOREGROUND_ONLY
-			Output0[OutputPixelPosition] = float4(
+			ConvolutionOutput_SceneColor[OutputPixelPosition] = float4(
 				AccumulatorOutput.ForegroundColor.rgb * AccumulatorOutput.ForegroundAlpha,
 				AccumulatorOutput.ForegroundAlpha);
 
 		#elif DIM_LAYER_PROCESSING == LAYER_PROCESSING_FOREGROUND_HOLE_FILLING
-			Output0[OutputPixelPosition] = float4(
+			ConvolutionOutput_SceneColor[OutputPixelPosition] = float4(
 				AccumulatorOutput.ForegroundHoleFillingColor.rgb,
 				AccumulatorOutput.ForegroundHoleFillingAlpha);
 
 		#elif DIM_LAYER_PROCESSING == LAYER_PROCESSING_BACKGROUND_ONLY
-			Output0[OutputPixelPosition] = float4(
+			ConvolutionOutput_SceneColor[OutputPixelPosition] = float4(
 				AccumulatorOutput.BackgroundColor.rgb,
 				AccumulatorOutput.BackgroundWeight > 0 ? 1 : 0);
 		
 		#elif DIM_LAYER_PROCESSING == LAYER_PROCESSING_SLIGHT_OUT_OF_FOCUS
-			Output0[OutputPixelPosition] = float4(
+			ConvolutionOutput_SceneColor[OutputPixelPosition] = float4(
 				AccumulatorOutput.SlightFocusColor.rgb,
 				AccumulatorOutput.SlightFocusOpacity);
 		
@@ -131,7 +143,7 @@ void StoreAccumulatorOutput(in uint2 OutputPixelPosition, in FAccumulatorOutput
 	#endif // !CONFIG_DOF_ALPHA
 
 	#if CONFIG_GENERATE_SCATTER_OCCLUSION_BUFFER
-		Output2[OutputPixelPosition] = OutScatterOcclusion;
+		ScatterOcclusionOutput[OutputPixelPosition] = OutScatterOcclusion;
 
 	#endif
 }
@@ -222,7 +234,7 @@ void GatherMainCS(
 		#else
 			const uint2 TilePos = GroupId;
 		#endif
-		const FCocTileSample CocTileInfos = LoadCocTile(COC_TILE_LAYOUT_FGD_SEP_BGD, PostprocessInput2, PostprocessInput3, int2(TilePos));
+		const FCocTileSample CocTileInfos = LoadCocTile(COC_TILE_LAYOUT_FGD_SEP_BGD, TileClassification_Foreground, TileClassification_Background, int2(TilePos));
 
 		TileSuggestion = InferGatherTileSuggestion(CocTileInfos,
 			/* LayerProcessing = */ DIM_LAYER_PROCESSING,
diff --git a/Engine/Shaders/Private/DiaphragmDOF/DOFHybridScatterPixelShader.usf b/Engine/Shaders/Private/DiaphragmDOF/DOFHybridScatterPixelShader.usf
--- a/Engine/Shaders/Private/DiaphragmDOF/DOFHybridScatterPixelShader.usf
+++ b/Engine/Shaders/Private/DiaphragmDOF/DOFHybridScatterPixelShader.usf
@@ -22,6 +22,19 @@
 
 float ScatteringScaling;
 
+Texture2D BokehLUT;
+
+float4 ScatterOcclusionSize;
+Texture2D ScatterOcclusion;
+
+
+#if !defined(NEW_CODE) // TODO: remove me
+#define BokehLUT PostprocessInput2
+#define ScatterOcclusionSize PostprocessInput3Size
+#define ScatterOcclusion PostprocessInput3
+#endif
+
+
 
 //------------------------------------------------------- ENTRY POINT
 
@@ -64,7 +77,7 @@ void ScatterMainPS(
 	#endif
 
 	#if DIM_SCATTER_OCCLUSION
-		float4 ScatterOcclusionSample = PostprocessInput3.SampleLevel(GlobalBilinearWrappedSampler, SvPosition.xy * PostprocessInput3Size.zw, 0);
+		float4 ScatterOcclusionSample = ScatterOcclusion.SampleLevel(GlobalBilinearWrappedSampler, SvPosition.xy * ScatterOcclusionSize.zw, 0);
 
 		float CocRadiusAvg = ScatterOcclusionSample.x;
 		float CocRadiusVariance = ScatterOcclusionSample.y;
@@ -86,7 +99,7 @@ void ScatterMainPS(
 			float2 LookUpUV = PixelOffset * (CONFIG_POINT_MIRROR_BOKEH * (0.5 - 1.0 / float(BOKEH_LUT_SIZE)) * SafeRcp(PixelDistance));
 
 			// TODO: might be faster without using view uniform buffer.
-			float4 LookupSample = PostprocessInput2.SampleLevel(GlobalBilinearWrappedSampler, LookUpUV, 0);
+			float4 LookupSample = BokehLUT.SampleLevel(GlobalBilinearWrappedSampler, LookUpUV, 0);
 			float CocRadiusToBokehEdgeFactor = LookupSample.x;
 
 			Intersection = saturate(AbsCocRadius[i] * CocRadiusToBokehEdgeFactor - PixelDistance + ANTI_ALIASING_FEATHER_OFFSET);
diff --git a/Engine/Shaders/Private/DiaphragmDOF/DOFPostfiltering.usf b/Engine/Shaders/Private/DiaphragmDOF/DOFPostfiltering.usf
--- a/Engine/Shaders/Private/DiaphragmDOF/DOFPostfiltering.usf
+++ b/Engine/Shaders/Private/DiaphragmDOF/DOFPostfiltering.usf
@@ -16,8 +16,30 @@
 float2 MaxInputBufferUV;
 float MinGatherRadius;
 
-RWTexture2D<float4>	Output0;
-RWTexture2D<float4>	Output1;
+
+float4 ConvolutionInputSize;
+Texture2D ConvolutionInput_SceneColor;
+Texture2D ConvolutionInput_SeparateAlpha;
+
+Texture2D TileClassification_Foreground;
+Texture2D TileClassification_Background;
+
+#if !defined(NEW_CODE) // TODO: remove me
+#define ConvolutionOutput_SceneColor Output0
+#define ConvolutionOutput_SeparateAlpha Output1
+
+#define ConvolutionInputSize PostprocessInput0Size
+#define ConvolutionInput_SceneColor PostprocessInput0
+#define ConvolutionInput_SeparateAlpha PostprocessInput1
+
+#define TileClassification_Foreground PostprocessInput2
+#define TileClassification_Background PostprocessInput3
+#endif
+
+
+RWTexture2D<float4>	ConvolutionOutput_SceneColor;
+RWTexture2D<float4>	ConvolutionOutput_SeparateAlpha;
+
 
 static const int2 kOffsets3x3[9] = {
 	int2(-1, -1),
@@ -150,7 +172,7 @@ void PostfilterMainCS(
 	uint2 DispatchThreadId : SV_DispatchThreadID,
 	uint2 GroupId : SV_GroupID)
 {
-	float2 InputBufferUV = (DispatchThreadId + 0.5) * PostprocessInput0Size.zw;
+	float2 InputBufferUV = (DispatchThreadId + 0.5) * ConvolutionInputSize.zw;
 	
 	BRANCH
 	if (any(DispatchThreadId >= ViewportRect.zw))
@@ -163,7 +185,11 @@ void PostfilterMainCS(
 	#if DIM_TILE_PERMUTATION
 	{
 		const uint2 TilePos = GroupId;
-		const FCocTileSample CocTileInfos = LoadCocTile(COC_TILE_LAYOUT_FGD_SEP_BGD, PostprocessInput2, PostprocessInput3, TilePos);
+		const FCocTileSample CocTileInfos = LoadCocTile(
+			COC_TILE_LAYOUT_FGD_SEP_BGD,
+			TileClassification_Foreground,
+			TileClassification_Background,
+			TilePos);
 
 		FGatheringTileSuggestion TileSuggestion = InferGatherTileSuggestion(CocTileInfos,
 			/* LayerProcessing = */ DIM_LAYER_PROCESSING,
@@ -179,43 +205,33 @@ void PostfilterMainCS(
 		uint2 OutputPixelPosition = DispatchThreadId;
 		float4 Foreground = float4(0, 0, 0, 0);
 		float4 Background = float4(0, 0, 0, 0);
+		float SeparateAlpha = 0;
 
 		#if DIM_LAYER_PROCESSING == LAYER_PROCESSING_FOREGROUND_ONLY
-			Output0[OutputPixelPosition] = Foreground;
+			ConvolutionOutput_SceneColor[OutputPixelPosition] = Foreground;
 
 		#elif DIM_LAYER_PROCESSING == LAYER_PROCESSING_BACKGROUND_ONLY
-			Output0[OutputPixelPosition] = Background;
+			ConvolutionOutput_SceneColor[OutputPixelPosition] = Background;
 
 		#else
-			Output0[OutputPixelPosition] = Foreground;
-			Output1[OutputPixelPosition] = Background;
+			#error Unknown layer processing.
+		#endif
+		
+		#if CONFIG_DOF_ALPHA
+			ConvolutionOutput_SeparateAlpha[OutputPixelPosition] = SeparateAlpha;
 		#endif
 	}
 	else
 	{
-		#if DIM_LAYER_PROCESSING == LAYER_PROCESSING_BACKGROUND_ONLY
-			float4 Foreground = 0;
-			float4 Background = PostfilterInput(PostprocessInput0, PostprocessInput0Size.zw, InputBufferUV);
-		#else
-			float4 Foreground = PostfilterInput(PostprocessInput0, PostprocessInput0Size.zw, InputBufferUV);
-			float4 Background = PostfilterInput(PostprocessInput1, PostprocessInput0Size.zw, InputBufferUV);
-		#endif
-
 		#if 1 // Lower VGPR footprint.
-			uint2 OutputPixelPosition = InputBufferUV * PostprocessInput0Size.xy;
+			uint2 OutputPixelPosition = InputBufferUV * ConvolutionInputSize.xy;
 		#else
 			uint2 OutputPixelPosition = DispatchThreadId;
 		#endif
-
-		#if DIM_LAYER_PROCESSING == LAYER_PROCESSING_FOREGROUND_ONLY
-			Output0[OutputPixelPosition] = Foreground;
-
-		#elif DIM_LAYER_PROCESSING == LAYER_PROCESSING_BACKGROUND_ONLY
-			Output0[OutputPixelPosition] = Background;
-
-		#else
-			Output0[OutputPixelPosition] = Foreground;
-			Output1[OutputPixelPosition] = Background;
+			
+		ConvolutionOutput_SceneColor[OutputPixelPosition] = PostfilterInput(ConvolutionInput_SceneColor, ConvolutionInputSize.zw, InputBufferUV);
+		#if CONFIG_DOF_ALPHA
+			ConvolutionOutput_SeparateAlpha[OutputPixelPosition] = PostfilterInput(ConvolutionInput_SeparateAlpha, ConvolutionInputSize.zw, InputBufferUV);
 		#endif
 	}
 }
diff --git a/Engine/Shaders/Private/DiaphragmDOF/DOFRecombine.usf b/Engine/Shaders/Private/DiaphragmDOF/DOFRecombine.usf
--- a/Engine/Shaders/Private/DiaphragmDOF/DOFRecombine.usf
+++ b/Engine/Shaders/Private/DiaphragmDOF/DOFRecombine.usf
@@ -134,7 +134,53 @@ float4 CocModelParameters;
 float2 DepthBlurParameters;
 float2 DOFBufferUVMax;
 
-RWTexture2D<float4>	Output0;
+Texture2D BokehLUT;
+
+Texture2D SceneColorInput;
+Texture2D SceneDepthTexture;
+Texture2D SceneSeparateCoc;
+Texture2D SceneSeparateTranslucency;
+
+float4 ConvolutionInputSize;
+
+Texture2D ForegroundConvolution_SceneColor;
+Texture2D ForegroundConvolution_SeparateAlpha;
+
+Texture2D ForegroundHoleFillingConvolution_SceneColor;
+Texture2D ForegroundHoleFillingConvolution_SeparateAlpha;
+
+Texture2D SlightOutOfFocusConvolution_SceneColor;
+Texture2D SlightOutOfFocusConvolution_SeparateAlpha;
+
+Texture2D BackgroundConvolution_SceneColor;
+Texture2D BackgroundConvolution_SeparateAlpha;
+
+
+#if !defined(NEW_CODE) // TODO: remove me
+
+#define SceneColorInput PostprocessInput0
+#define SceneSeparateCoc PostprocessInput1
+#define SceneSeparateTranslucency PostprocessInput2
+
+#define BokehLUT PostprocessInput10
+
+#define ForegroundConvolution_SceneColor PostprocessInput3
+#define ForegroundConvolution_SeparateAlpha PostprocessInput4
+
+#define ForegroundHoleFillingConvolution_SceneColor PostprocessInput5
+#define ForegroundHoleFillingConvolution_SeparateAlpha PostprocessInput6
+
+#define SlightOutOfFocusConvolution_SceneColor PostprocessInput9
+
+#define BackgroundConvolution_SceneColor PostprocessInput7
+#define BackgroundConvolution_SeparateAlpha PostprocessInput8
+
+#define SceneColorOutput Output0
+
+#endif
+
+
+RWTexture2D<float4>	SceneColorOutput;
 
 
 //------------------------------------------------------- INTERMEDIARY STRUCTURES
@@ -161,6 +207,15 @@ struct FRecombineInputParameters
 
 //------------------------------------------------------- FUNCTIONS
 
+float SampleWorldDepth(float2 BufferUV)
+{
+#if defined(NEW_CODE)
+	return ConvertFromDeviceZ(SceneDepthTexture.SampleLevel(GlobalPointClampedSampler, BufferUV, 0).x);
+#else
+	return CalcSceneDepth(BufferUV);
+#endif
+}
+
 // {0 to 1} output.
 float NoizNorm(float2 N, float X)
 {
@@ -526,7 +581,7 @@ void FetchAndAccumulateSamplePair(
 	float SampleDistance = FullResPixelDistanceToCocDistance(length(PixelOffset));
 
 	// Scene buffer offset.
-	float2 SceneBufferUVOffset = PixelOffset * PostprocessInput0Size.zw;
+	float2 SceneBufferUVOffset = PixelOffset * View.BufferSizeAndInvSize.zw;
 	
 	// Two sample to gather at same time.
 	FGatherSample Sample[2];
@@ -543,7 +598,7 @@ void FetchAndAccumulateSamplePair(
 			const float InvLutSize = rcp(float(BOKEH_LUT_SIZE));
 
 			float2 LookupUV = (0.5 + 0.5 * InvLutSize) + PixelOffset * (SampleSign * InvLutSize);
-			float4 LookupSample = PostprocessInput10.SampleLevel(GlobalPointClampedSampler, LookupUV, 0);
+			float4 LookupSample = BokehLUT.SampleLevel(GlobalPointClampedSampler, LookupUV, 0);
 			SampleDistance = LookupSample.x;
 		}
 		#endif
@@ -557,11 +612,11 @@ void FetchAndAccumulateSamplePair(
 		
 		// Fetch full res color and CocRadius.
 		#if CONFIG_FETCH_FULLRES_COC_FROM_ALPHA
-			Sample[k].Color = PostprocessInput0.SampleLevel(GlobalPointClampedSampler, BufferUV, 0);
+			Sample[k].Color = SceneColorInput.SampleLevel(GlobalPointClampedSampler, BufferUV, 0);
 			Sample[k].CocRadius = Sample[k].Color.a;
 		#else
-			Sample[k].Color = PostprocessInput0.SampleLevel(GlobalPointClampedSampler, BufferUV, 0);
-			Sample[k].CocRadius = DepthToHalfResCocRadius(CalcSceneDepth(BufferUV));
+			Sample[k].Color = SceneColorInput.SampleLevel(GlobalPointClampedSampler, BufferUV, 0);
+			Sample[k].CocRadius = DepthToHalfResCocRadius(SampleWorldDepth(BufferUV));
 		#endif
 		
 		// Convert scene color alpha from translucency to opacity.
@@ -657,7 +712,11 @@ void RecombineMainCS(
 	// Setup input parameters.
 	FRecombineInputParameters InputParameters;
 	{
+#if defined(NEW_CODE) // TODO: remove me
+		InputParameters.DOFBufferSize = ConvolutionInputSize;
+#else
 		InputParameters.DOFBufferSize = (DIM_LAYER_PROCESSING == LAYER_PROCESSING_BACKGROUND_ONLY) ? PostprocessInput7Size : PostprocessInput3Size;
+#endif
 
 		InputParameters.ViewportUV = (DispatchThreadId + 0.5) * ViewportSize.zw;
 		InputParameters.SceneBufferUV = ViewportUVToBufferUV(InputParameters.ViewportUV);
@@ -688,13 +747,13 @@ void RecombineMainCS(
 
 		#elif CONFIG_DOF_ALPHA
 			// Sample premultiplied RGBA foreground.
-			ForegroundColor = PostprocessInput3.SampleLevel(
+			ForegroundColor = ForegroundConvolution_SceneColor.SampleLevel(
 				GlobalBilinearClampedSampler, InputParameters.DOFBufferUV, 0);
-			ForegroundTranslucency = 1 - PostprocessInput4.SampleLevel(
+			ForegroundTranslucency = 1 - ForegroundConvolution_SeparateAlpha.SampleLevel(
 				GlobalBilinearClampedSampler, InputParameters.DOFBufferUV, 0).r;
 		#else
 			// Sample premultiplied RGBA foreground.
-			ForegroundColor = PostprocessInput3.SampleLevel(
+			ForegroundColor = ForegroundConvolution_SceneColor.SampleLevel(
 				GlobalBilinearClampedSampler, InputParameters.DOFBufferUV, 0);
 			ForegroundTranslucency = 1 - ForegroundColor.a;
 
@@ -704,13 +763,13 @@ void RecombineMainCS(
 	
 	// Get full res color and coc radius.
 	FGatherSample CenterSample;
-	CenterSample.Color = PostprocessInput0.SampleLevel(GlobalPointClampedSampler, InputParameters.SceneBufferUV, 0);
+	CenterSample.Color = SceneColorInput.SampleLevel(GlobalPointClampedSampler, InputParameters.SceneBufferUV, 0);
 	CenterSample.Intersection = 1.0;
 
 	#if CONFIG_FETCH_FULLRES_COC_FROM_ALPHA
 		CenterSample.CocRadius = CenterSample.Color.a;
 	#else
-		CenterSample.CocRadius = DepthToHalfResCocRadius(CalcSceneDepth(InputParameters.SceneBufferUV));
+		CenterSample.CocRadius = DepthToHalfResCocRadius(SampleWorldDepth(InputParameters.SceneBufferUV));
 	#endif
 	
 	// Convert scene color alpha from translucency to opacity.
@@ -743,7 +802,7 @@ void RecombineMainCS(
 			{
 				const float2 SamplePixelOffset = float2(kOffsetsCross3x3[j]) * CocDistanceToFullResPixelDistance(MAX_RECOMBINE_ABS_COC_RADIUS);
 
-				float2 SampleUVOffset = PostprocessInput0Size.zw * SamplePixelOffset;
+				float2 SampleUVOffset = View.BufferSizeAndInvSize.zw * SamplePixelOffset;
 				float2 SampleUV = InputParameters.SceneBufferUV + SampleUVOffset;
 			
 				if (CONFIG_CLAMP_SCENE_BUFFER_UV)
@@ -752,9 +811,9 @@ void RecombineMainCS(
 				}
 
 				#if CONFIG_FETCH_FULLRES_COC_FROM_ALPHA
-					float SampleCocRadius = PostprocessInput0.SampleLevel(GlobalPointClampedSampler, SampleUV, 0).a;
+					float SampleCocRadius = SceneColorInput.SampleLevel(GlobalPointClampedSampler, SampleUV, 0).a;
 				#else
-					float SampleCocRadius = DepthToHalfResCocRadius(CalcSceneDepth(SampleUV));
+					float SampleCocRadius = DepthToHalfResCocRadius(SampleWorldDepth(SampleUV));
 				#endif
 
 				float SampleAbsCocRadius = abs(SampleCocRadius);
@@ -898,7 +957,7 @@ void RecombineMainCS(
 					SampleUV = min(SampleUV, DOFBufferUVMax);
 				}
 
-				float4 StableSampleColor = PostprocessInput9.SampleLevel(GlobalPointClampedSampler, SampleUV, 0);
+				float4 StableSampleColor = SlightOutOfFocusConvolution_SceneColor.SampleLevel(GlobalPointClampedSampler, SampleUV, 0);
 
 				if (i == 0)
 				{
@@ -964,10 +1023,10 @@ void RecombineMainCS(
 	}
 	else if (CONFIG_COMPOSITING_METHOD == COMPOSITING_METHOD_BILINEAR_BKG)
 	{
-		BackgroundColor = PostprocessInput7.SampleLevel(GlobalBilinearClampedSampler, InputParameters.DOFBufferUV, 0);
+		BackgroundColor = BackgroundConvolution_SceneColor.SampleLevel(GlobalBilinearClampedSampler, InputParameters.DOFBufferUV, 0);
 
 		#if CONFIG_DOF_ALPHA
-			BackgroundValidity = PostprocessInput8.SampleLevel(GlobalBilinearClampedSampler, InputParameters.DOFBufferUV, 0).r;
+			BackgroundValidity = BackgroundConvolution_SeparateAlpha.SampleLevel(GlobalBilinearClampedSampler, InputParameters.DOFBufferUV, 0).r;
 		#else
 			//Background = float4(0, 0, 0, 1);
 			BackgroundValidity = BackgroundColor.a;
@@ -979,10 +1038,10 @@ void RecombineMainCS(
 		// Hole fill the background in output final scene color before composing foreground on top.
 		#if CONFIG_HOLE_FILLING_METHOD == HOLE_FILLING_METHOD_SEPARATE_GATHER
 		{
-			HoleFillingAdditiveColor = PostprocessInput5.SampleLevel(GlobalBilinearClampedSampler, InputParameters.DOFBufferUV, 0);
+			HoleFillingAdditiveColor = ForegroundHoleFillingConvolution_SceneColor.SampleLevel(GlobalBilinearClampedSampler, InputParameters.DOFBufferUV, 0);
 
 			#if CONFIG_DOF_ALPHA
-				HoleFillingTranslucency = PostprocessInput6.SampleLevel(GlobalBilinearClampedSampler, InputParameters.DOFBufferUV, 0).r;
+				HoleFillingTranslucency = ForegroundHoleFillingConvolution_SeparateAlpha.SampleLevel(GlobalBilinearClampedSampler, InputParameters.DOFBufferUV, 0).r;
 			#else
 				HoleFillingTranslucency = HoleFillingAdditiveColor.a;
 			#endif
@@ -1039,7 +1098,7 @@ void RecombineMainCS(
 	// Compose separate translucency.
 	if (1)
 	{
-		float4 SeparateTranslucencyColor = PostprocessInput2.SampleLevel(GlobalPointClampedSampler, InputParameters.SceneBufferUV, 0);
+		float4 SeparateTranslucencyColor = SceneSeparateTranslucency.SampleLevel(GlobalPointClampedSampler, InputParameters.SceneBufferUV, 0);
 		float SeparateTranslucency = SeparateTranslucencyColor.a;
 		SeparateTranslucencyColor.a = 1 - SeparateTranslucencyColor.a;
 
@@ -1093,7 +1152,7 @@ void RecombineMainCS(
 		//OutputFinalSceneColor = ForegroundColor;
 		//OutputFinalSceneColor = HoleFillingAdditiveColor;
 		//OutputFinalSceneColor = BackgroundColor.a;
-		//OutputFinalSceneColor = PostprocessInput9.SampleLevel(GlobalPointClampedSampler, InputParameters.DOFBufferUV, 0).a;
+		//OutputFinalSceneColor = SlightOutOfFocusConvolution_SceneColor.SampleLevel(GlobalPointClampedSampler, InputParameters.DOFBufferUV, 0).a;
 		//OutputFinalSceneColor = GatherBackgroundUnpremultipliedColor;
 		//OutputFinalSceneColor.a = GatherBackgroundTranslucency;
 	}
@@ -1111,7 +1170,7 @@ void RecombineMainCS(
 			OutputFinalSceneColor = GatherForegroundTranslucency;
 		
 		//OutputFinalSceneColor = GatherForegroundTranslucency;
-		//OutputFinalSceneColor = PostprocessInput7.SampleLevel(GlobalBilinearClampedSampler, InputParameters.DOFBufferUV, 0);
+		//OutputFinalSceneColor = BackgroundConvolution_SceneColor.SampleLevel(GlobalBilinearClampedSampler, InputParameters.DOFBufferUV, 0);
 
 
 		//if (View.GeneralPurposeTweak > 3.5)
@@ -1121,19 +1180,19 @@ void RecombineMainCS(
 
 		//OutputFinalSceneColor = BackgroundColor;
 		//OutputFinalSceneColor = PostprocessInput3.SampleLevel(GlobalPointClampedSampler, InputParameters.DOFBufferUV, 0);
-		//OutputFinalSceneColor = PostprocessInput5.SampleLevel(GlobalPointClampedSampler, InputParameters.DOFBufferUV, 0);
+		//OutputFinalSceneColor = ForegroundHoleFillingConvolution_SceneColor.SampleLevel(GlobalPointClampedSampler, InputParameters.DOFBufferUV, 0);
 	}
 	#endif
 	
 	#if 1 // Lower VGPR footprint.
-		uint2 OutputPixelPosition = InputParameters.SceneBufferUV * PostprocessInput0Size.xy;
+		uint2 OutputPixelPosition = InputParameters.SceneBufferUV * View.BufferSizeAndInvSize.xy;
 	#else
 		uint2 OutputPixelPosition = ViewportRect.xy + DispatchThreadId;
 	#endif
 
 	#if CONFIG_DOF_ALPHA
-		Output0[OutputPixelPosition] = OutputFinalSceneColor;
+		SceneColorOutput[OutputPixelPosition] = OutputFinalSceneColor;
 	#else
-		Output0[OutputPixelPosition] = float4(OutputFinalSceneColor.rgb, 0);
+		SceneColorOutput[OutputPixelPosition] = float4(OutputFinalSceneColor.rgb, 0);
 	#endif
 }
diff --git a/Engine/Shaders/Private/DiaphragmDOF/DOFReduce.usf b/Engine/Shaders/Private/DiaphragmDOF/DOFReduce.usf
--- a/Engine/Shaders/Private/DiaphragmDOF/DOFReduce.usf
+++ b/Engine/Shaders/Private/DiaphragmDOF/DOFReduce.usf
@@ -61,10 +61,36 @@
 
 uint MaxScatteringGroupCount;
 float2 MaxInputBufferUV;
-
 float MinScatteringCocRadius;
+float PreProcessingToProcessingCocRadiusFactor;
+
+float4 GatherInputSize;
+Texture2D GatherInput_SceneColor;
+Texture2D GatherInput_SeparateCoc;
+
+float4 QuarterResGatherInputSize;
+Texture2D QuarterResGatherInput_SceneColor;
 
-float2 PreProcessingToProcessingCocRadiusFactor;
+#if defined(NEW_CODE) // TODO: remove me
+#define Output0Mip0 OutputMips_0_SceneColor
+#define Output1Mip0 OutputMips_0_SeparateCoc
+#define Output0Mip1 OutputMips_1_SceneColor
+#define Output1Mip1 OutputMips_1_SeparateCoc
+#define Output0Mip2 OutputMips_2_SceneColor
+#define Output1Mip2 OutputMips_2_SeparateCoc
+#define Output0Mip3 OutputMips_3_SceneColor
+#define Output1Mip3 OutputMips_3_SeparateCoc
+
+#else
+
+#define GatherInputSize PostprocessInput0Size
+#define GatherInput_SceneColor PostprocessInput0
+#define GatherInput_SeparateCoc PostprocessInput1
+
+#define QuarterResGatherInputSize PostprocessInput2Size
+#define QuarterResGatherInput_SceneColor PostprocessInput2
+
+#endif
 
 
 //------------------------------------------------------- OUTPUTS
@@ -73,7 +99,6 @@ RWTexture2D<float4>	Output0Mip0;
 RWTexture2D<float4>	Output0Mip1;
 RWTexture2D<float4>	Output0Mip2;
 RWTexture2D<float4>	Output0Mip3;
-RWTexture2D<float4>	Output0Mip4;
 
 #if CONFIG_GATHER_INPUT_LAYOUT == GATHER_INPUT_LAYOUT_RGB_ALPHA_COC || CONFIG_GATHER_INPUT_LAYOUT == GATHER_INPUT_LAYOUT_RGB_SEPARATE_COC
 
@@ -81,7 +106,6 @@ RWTexture2D<float>	Output1Mip0;
 RWTexture2D<float>	Output1Mip1;
 RWTexture2D<float>	Output1Mip2;
 RWTexture2D<float>	Output1Mip3;
-RWTexture2D<float>	Output1Mip4;
 
 #endif
 
@@ -324,7 +348,7 @@ void ReduceCS(
 
 	uint2 DispatchThreadId = THREADGROUP_TILE_SIZE * GroupId + GroupThreadId;
 	
-	float2 BufferUV = (DispatchThreadId + 0.5) * PostprocessInput0Size.zw;
+	float2 BufferUV = (DispatchThreadId + 0.5) * GatherInputSize.zw;
 
 	if (true)
 	{
@@ -336,12 +360,12 @@ void ReduceCS(
 	float CocRadius = 0;
 	#if CONFIG_DOF_ALPHA
 	{
-		GatheredColor = PostprocessInput0.SampleLevel(GlobalPointClampedSampler, BufferUV, 0);
-		CocRadius = PostprocessInput1.SampleLevel(GlobalPointClampedSampler, BufferUV, 0).r;
+		GatheredColor = GatherInput_SceneColor.SampleLevel(GlobalPointClampedSampler, BufferUV, 0);
+		CocRadius = GatherInput_SeparateCoc.SampleLevel(GlobalPointClampedSampler, BufferUV, 0).r;
 	}
 	#else
 	{
-		GatheredColor = PostprocessInput0.SampleLevel(GlobalPointClampedSampler, BufferUV, 0);
+		GatheredColor = GatherInput_SceneColor.SampleLevel(GlobalPointClampedSampler, BufferUV, 0);
 		CocRadius = GatheredColor.a;
 	}
 	#endif
@@ -355,7 +379,7 @@ void ReduceCS(
 		const uint IndirectParameterSize = SCATTERING_INDIRECT_PARAMETER_SIZE;
 		
 		// Convert Coc radius from preprocessing basis to processing basis.
-		float ScatterCocRadius = CocRadius * PreProcessingToProcessingCocRadiusFactor.x;
+		float ScatterCocRadius = CocRadius * PreProcessingToProcessingCocRadiusFactor;
 
 		// Decides whether should hybrid scatter or not.
 		// It is fine to do it for sample outside viewport UV because they are discarded when the bokeh
@@ -363,8 +387,8 @@ void ReduceCS(
 		float ScatterFactor;
 		{
 			FHybridScatterInputs Parameters;
-			Parameters.NeighborhoodComparisonBuffer = PostprocessInput2;
-			Parameters.NeighborhoodComparisonBufferInvSize = PostprocessInput2Size.zw;
+			Parameters.NeighborhoodComparisonBuffer = QuarterResGatherInput_SceneColor;
+			Parameters.NeighborhoodComparisonBufferInvSize = QuarterResGatherInputSize.zw;
 			Parameters.NeighborhoodComparisonMaxBufferUV = MaxInputBufferUV;
 
 			Parameters.FrameExposureScale = FrameExposureScale;
diff --git a/Engine/Shaders/Private/DiaphragmDOF/DOFSetup.usf b/Engine/Shaders/Private/DiaphragmDOF/DOFSetup.usf
--- a/Engine/Shaders/Private/DiaphragmDOF/DOFSetup.usf
+++ b/Engine/Shaders/Private/DiaphragmDOF/DOFSetup.usf
@@ -26,6 +26,14 @@ float2 DepthBlurParameters;
 // x: full res basis; y: half res basis.
 float2 CocRadiusBasis;
 
+Texture2D SceneColorTexture;
+Texture2D SceneDepthTexture;
+
+#if defined(NEW_CODE) // TODO: remove me
+#define PostprocessInput0 SceneColorTexture
+#define PostprocessInput1 SceneDepthTexture
+#endif
+
 
 //------------------------------------------------------- OUTPUTS
 
diff --git a/Engine/Source/Runtime/Renderer/Private/PostProcess/DiaphragmDOF.cpp b/Engine/Source/Runtime/Renderer/Private/PostProcess/DiaphragmDOF.cpp
new file mode 100644
--- /dev/null
+++ b/Engine/Source/Runtime/Renderer/Private/PostProcess/DiaphragmDOF.cpp
@@ -0,0 +1,2264 @@
+// Copyright 1998-2019 Epic Games, Inc. All Rights Reserved.
+
+/*=============================================================================
+	DiaphragmDOFPasses.cpp: Implementations of all diaphragm DOF's passes.
+=============================================================================*/
+
+#include "PostProcess/DiaphragmDOF.h"
+#include "PostProcess/DiaphragmDOFPasses.h"
+#include "StaticBoundShaderState.h"
+#include "SceneUtils.h"
+#include "PostProcess/SceneRenderTargets.h"
+#include "PostProcess/SceneFilterRendering.h"
+#include "PostProcess/PostProcessBokehDOF.h"
+#include "SceneRenderTargetParameters.h"
+#include "PostProcess/PostProcessing.h"
+#include "DeferredShadingRenderer.h"
+#include "PipelineStateCache.h"
+#include "ScenePrivate.h"
+#include "ClearQuad.h"
+#include "SpriteIndexBuffer.h"
+#include "PostProcessTemporalAA.h"
+#include "SceneViewFamilyBlackboard.h"
+
+
+FORCEINLINE_DEBUGGABLE void AddPass_ClearUAV(
+	FRDGBuilder& GraphBuilder,
+	FRDGEventName&& PassName,
+	FRDGBufferUAVRef BufferUAV,
+	uint32 Value);
+
+
+// ---------------------------------------------------- Cvars
+
+TAutoConsoleVariable<int32> CVarAccumulatorQuality(
+	TEXT("r.DOF.Gather.AccumulatorQuality"),
+	1,
+	TEXT("Controles the quality of the gathering accumulator.\n"),
+	ECVF_Scalability | ECVF_RenderThreadSafe);
+
+TAutoConsoleVariable<int32> CVarEnableGatherBokehSettings(
+	TEXT("r.DOF.Gather.EnableBokehSettings"),
+	1,
+	TEXT("Whether to applies bokeh settings on foreground and background gathering.\n")
+	TEXT(" 0: Disable;\n 1: Enable (default)."),
+	ECVF_Scalability | ECVF_RenderThreadSafe);
+
+TAutoConsoleVariable<int32> CVarPostFilteringMethod(
+	TEXT("r.DOF.Gather.PostfilterMethod"),
+	1,
+	TEXT("Method to use to post filter a gather pass.\n")
+	TEXT(" 0: None;\n")
+	TEXT(" 1: Per RGB channel median 3x3 (default);\n")
+	TEXT(" 2: Per RGB channel max 3x3."),
+	ECVF_Scalability | ECVF_RenderThreadSafe);
+
+TAutoConsoleVariable<int32> CVarRingCount(
+	TEXT("r.DOF.Gather.RingCount"),
+	5,
+	TEXT("Number of rings for gathering kernels [[3; 5]]. Default to 5.\n"),
+	ECVF_Scalability | ECVF_RenderThreadSafe);
+
+
+TAutoConsoleVariable<int32> CVarHybridScatterForegroundMode(
+	TEXT("r.DOF.Scatter.ForegroundCompositing"),
+	1,
+	TEXT("Compositing mode of the foreground hybrid scattering.\n")
+	TEXT(" 0: Disabled;\n")
+	TEXT(" 1: Additive (default)."),
+	ECVF_Scalability | ECVF_RenderThreadSafe);
+
+TAutoConsoleVariable<int32> CVarHybridScatterBackgroundMode(
+	TEXT("r.DOF.Scatter.BackgroundCompositing"),
+	2,
+	TEXT("Compositing mode of the background hybrid scattering.\n")
+	TEXT(" 0: Disabled;\n")
+	TEXT(" 1: Additive;\n")
+	TEXT(" 2: Gather occlusion (default)."),
+	ECVF_Scalability | ECVF_RenderThreadSafe);
+
+TAutoConsoleVariable<int32> CVarEnableScatterBokehSettings(
+	TEXT("r.DOF.Scatter.EnableBokehSettings"),
+	1,
+	TEXT("Whether to enable bokeh settings on scattering.\n")
+	TEXT(" 0: Disable;\n 1: Enable (default)."),
+	ECVF_Scalability | ECVF_RenderThreadSafe);
+
+TAutoConsoleVariable<float> CVarScatterMinCocRadius(
+	TEXT("r.DOF.Scatter.MinCocRadius"),
+	3.0f,
+	TEXT("Minimal Coc radius required to be scattered (default = 3)."),
+	ECVF_Scalability | ECVF_RenderThreadSafe);
+
+TAutoConsoleVariable<float> CVarScatterMaxSpriteRatio(
+	TEXT("r.DOF.Scatter.MaxSpriteRatio"),
+	0.1f,
+	TEXT("Maximum ratio of scattered pixel quad as sprite, usefull to control DOF's scattering upperbound. ")
+	TEXT(" 1 will allow to scatter 100% pixel quads, whereas 0.2 will only allow 20% (default = 0.1)."),
+	ECVF_Scalability | ECVF_RenderThreadSafe);
+
+TAutoConsoleVariable<int32> CVarEnableRecombineBokehSettings(
+	TEXT("r.DOF.Recombine.EnableBokehSettings"),
+	1,
+	TEXT("Whether to applies bokeh settings on slight out of focus done in recombine pass.\n")
+	TEXT(" 0: Disable;\n 1: Enable (default)."),
+	ECVF_Scalability | ECVF_RenderThreadSafe);
+
+TAutoConsoleVariable<int32> CVarRecombineQuality(
+	TEXT("r.DOF.Recombine.Quality"),
+	2,
+	TEXT("Configures the quality of the recombine pass.\n")
+	TEXT(" 0: No slight out of focus;\n")
+	TEXT(" 1: Slight out of focus 24spp;\n")
+	TEXT(" 2: Slight out of focus 32spp (default)."),
+	ECVF_Scalability | ECVF_RenderThreadSafe);
+
+TAutoConsoleVariable<float> CVarMinimalFullresBlurRadius(
+	TEXT("r.DOF.Recombine.MinFullresBlurRadius"),
+	0.1f,
+	TEXT("Minimal blurring radius used in full resolution pixel width to actually do ")
+	TEXT("DOF  when slight out of focus is enabled (default = 0.1)."),
+	ECVF_RenderThreadSafe);
+
+TAutoConsoleVariable<int32> CVarDOFTemporalAAQuality(
+	TEXT("r.DOF.TemporalAAQuality"),
+	1,
+	TEXT("Quality of temporal AA pass done in DOF.\n")
+	TEXT(" 0: Faster but lower quality;")
+	TEXT(" 1: Higher quality pass (default)."),
+	ECVF_Scalability | ECVF_RenderThreadSafe);
+
+namespace // TODO encapsulate all cvars
+{
+
+TAutoConsoleVariable<float> CVarScatterNeighborCompareMaxColor(
+	TEXT("r.DOF2.Scatter.NeighborCompareMaxColor"),
+	10,
+	TEXT("Controles the linear color clamping upperbound applied before color of pixel and neighbors are compared.")
+	TEXT(" To low, and you may not scatter enough; to high you may scatter unnecessarily too much in highlights")
+	TEXT(" (Default: 10)."),
+	ECVF_RenderThreadSafe);
+
+} // namespace
+
+namespace
+{
+
+EDiaphragmDOFPostfilterMethod GetPostfilteringMethod()
+{
+	int32 i = CVarPostFilteringMethod.GetValueOnRenderThread();
+	if (i >= 0 && i < int32(EDiaphragmDOFPostfilterMethod::MAX))
+	{
+		return EDiaphragmDOFPostfilterMethod(i);
+	}
+	return EDiaphragmDOFPostfilterMethod::None;
+}
+
+
+enum class EHybridScatterMode
+{
+	Disabled,
+	Additive,
+	Occlusion,
+};
+
+
+}
+
+
+// ---------------------------------------------------- COMMON
+
+namespace
+{
+
+const int32 kDefaultGroupSize = 8;
+const int32 kCocTileSize = kDefaultGroupSize;
+
+FIntPoint CocTileGridSize(FIntPoint FullResSize)
+{
+	uint32 TilesX = FMath::DivideAndRoundUp(FullResSize.X, kCocTileSize);
+	uint32 TilesY = FMath::DivideAndRoundUp(FullResSize.Y, kCocTileSize);
+	return FIntPoint(TilesX, TilesY);
+}
+
+
+/** Returns the lower res's viewport from a given view size. */
+FIntRect GetLowerResViewport(const FIntRect& ViewRect, int32 ResDivisor)
+{
+	check(ResDivisor >= 1);
+	check(FMath::IsPowerOfTwo(ResDivisor));
+
+	FIntRect DestViewport;
+
+	// All diaphragm DOF's lower res viewports are top left cornered to only do a min(SampleUV, MaxUV) when doing convolution.
+	DestViewport.Min = FIntPoint::ZeroValue;
+
+	DestViewport.Max.X = FMath::DivideAndRoundUp(ViewRect.Width(), ResDivisor);
+	DestViewport.Max.Y = FMath::DivideAndRoundUp(ViewRect.Height(), ResDivisor);
+	return DestViewport;
+}
+
+
+const TCHAR* GetEventName(EDiaphragmDOFLayerProcessing e)
+{
+	static const TCHAR* const kArray[] = {
+		TEXT("FgdOnly"),
+		TEXT("FgdFill"),
+		TEXT("BgdOnly"),
+		TEXT("Fgd&Bgd"),
+		TEXT("FocusOnly"),
+	};
+	int32 i = int32(e);
+	check(i < ARRAY_COUNT(kArray));
+	return kArray[i];
+}
+
+const TCHAR* GetEventName(EDiaphragmDOFPostfilterMethod e)
+{
+	static const TCHAR* const kArray[] = {
+		TEXT("Median3x3"),
+		TEXT("Max3x3"),
+	};
+	int32 i = int32(e) - 1;
+	check(i < ARRAY_COUNT(kArray));
+	return kArray[i];
+}
+
+const TCHAR* GetEventName(EDiaphragmDOFBokehSimulation e)
+{
+	static const TCHAR* const kArray[] = {
+		TEXT("None"),
+		TEXT("Symmetric"),
+		TEXT("Generic"),
+	};
+	int32 i = int32(e);
+	check(i < ARRAY_COUNT(kArray));
+	return kArray[i];
+}
+
+const TCHAR* GetEventName(FRCPassDiaphragmDOFBuildBokehLUT::EFormat e)
+{
+	static const TCHAR* const kArray[] = {
+		TEXT("Scatter"),
+		TEXT("Recombine"),
+		TEXT("Gather"),
+	};
+	int32 i = int32(e);
+	check(i < ARRAY_COUNT(kArray));
+	return kArray[i];
+}
+
+const TCHAR* GetEventName(FRCPassDiaphragmDOFGather::EQualityConfig e)
+{
+	static const TCHAR* const kArray[] = {
+		TEXT("LowQ"),
+		TEXT("HighQ"),
+		TEXT("ScatterOcclusion"),
+		TEXT("Cinematic"),
+	};
+	int32 i = int32(e);
+	check(i < ARRAY_COUNT(kArray));
+	return kArray[i];
+}
+
+const TCHAR* GetEventName(FRCPassDiaphragmDOFDilateCoc::EMode e)
+{
+	static const TCHAR* const kArray[] = {
+		TEXT("StandAlone"),
+		TEXT("MinMax"),
+		TEXT("MinAbs"),
+	};
+	int32 i = int32(e);
+	check(i < ARRAY_COUNT(kArray));
+	return kArray[i];
+}
+
+// Returns X and Y for F(M) = saturate(M * X + Y) so that F(LowM) = 0 and F(HighM) = 1
+FVector2D GenerateSaturatedAffineTransformation(float LowM, float HighM)
+{
+	float X = 1.0f / (HighM - LowM);
+	return FVector2D(X, -X * LowM);
+}
+
+// Affine transformtations that always return 0 or 1.
+const FVector2D kContantlyPassingAffineTransformation(0, 1);
+const FVector2D kContantlyBlockingAffineTransformation(0, 0);
+
+
+/** Base shader class for diaphragm DOF. */
+class FDiaphragmDOFShader : public FGlobalShader
+{
+public:
+	static bool ShouldCompilePermutation(const FGlobalShaderPermutationParameters& Parameters)
+	{
+		return DiaphragmDOF::IsSupported(Parameters.Platform);
+	}
+
+	static void ModifyCompilationEnvironment(const FGlobalShaderPermutationParameters& Parameters, FShaderCompilerEnvironment& OutEnvironment)
+	{
+		OutEnvironment.SetDefine(TEXT("COC_TILE_SIZE"), kCocTileSize);
+		OutEnvironment.SetDefine(TEXT("NEW_CODE"), kCocTileSize); // TODO: kill that
+	}
+
+	FDiaphragmDOFShader() {}
+	FDiaphragmDOFShader(const ShaderMetaType::CompiledShaderInitializerType& Initializer)
+		: FGlobalShader(Initializer)
+	{ }
+};
+
+} // namespace
+
+
+// ---------------------------------------------------- Global resource
+
+namespace
+{
+
+class FDOFGlobalResource : public FRenderResource
+{
+public:
+	// Index buffer to have 4 vertex shader invocation per scatter group that is the most efficient in therms of vertex processing
+	// using  if RHI does not support rect list topology.
+	FSpriteIndexBuffer<16> ScatterIndexBuffer;
+
+
+	virtual void InitDynamicRHI() override
+	{
+		if (!GRHISupportsRectTopology)
+		{
+			ScatterIndexBuffer.InitRHI();
+		}
+	}
+
+	virtual void ReleaseDynamicRHI() override
+	{
+		if (!GRHISupportsRectTopology)
+		{
+			ScatterIndexBuffer.ReleaseRHI();
+		}
+	}
+};
+
+TGlobalResource<FDOFGlobalResource> GDOFGlobalResource;
+
+}
+
+
+// ---------------------------------------------------- Shader permutation dimensions
+
+namespace
+{
+
+class FDDOFDilateRadiusDim     : SHADER_PERMUTATION_RANGE_INT("DIM_DILATE_RADIUS", 1, 3);
+class FDDOFDilateModeDim       : SHADER_PERMUTATION_ENUM_CLASS("DIM_DILATE_MODE", FRCPassDiaphragmDOFDilateCoc::EMode);
+
+class FDDOFLayerProcessingDim  : SHADER_PERMUTATION_ENUM_CLASS("DIM_LAYER_PROCESSING", EDiaphragmDOFLayerProcessing);
+class FDDOFGatherRingCountDim  : SHADER_PERMUTATION_RANGE_INT("DIM_GATHER_RING_COUNT", FRCPassDiaphragmDOFGather::kMinRingCount, 3);
+class FDDOFGatherQualityDim    : SHADER_PERMUTATION_ENUM_CLASS("DIM_GATHER_QUALITY", FRCPassDiaphragmDOFGather::EQualityConfig);
+class FDDOFPostfilterMethodDim : SHADER_PERMUTATION_ENUM_CLASS("DIM_POSTFILTER_METHOD", EDiaphragmDOFPostfilterMethod);
+class FDDOFClampInputUVDim     : SHADER_PERMUTATION_BOOL("DIM_CLAMP_INPUT_UV");
+class FDDOFRGBColorBufferDim   : SHADER_PERMUTATION_BOOL("DIM_RGB_COLOR_BUFFER");
+
+class FDDOFBokehSimulationDim  : SHADER_PERMUTATION_ENUM_CLASS("DIM_BOKEH_SIMULATION", EDiaphragmDOFBokehSimulation);
+class FDDOFScatterOcclusionDim : SHADER_PERMUTATION_BOOL("DIM_SCATTER_OCCLUSION");
+
+} // namespace
+
+
+// ---------------------------------------------------- Shared shader parameters
+
+namespace
+{
+
+struct FDOFGatherInputDescs
+{
+	FRDGTextureDesc SceneColor;
+	FRDGTextureDesc SeparateCoc;
+};
+
+BEGIN_SHADER_PARAMETER_STRUCT(FDOFGatherInputTextures, )
+	SHADER_PARAMETER_RDG_TEXTURE(Texture2D, SceneColor)
+	SHADER_PARAMETER_RDG_TEXTURE(Texture2D, SeparateCoc)
+END_SHADER_PARAMETER_STRUCT()
+
+BEGIN_SHADER_PARAMETER_STRUCT(FDOFGatherInputUAVs, )
+	SHADER_PARAMETER_RDG_TEXTURE_UAV(RWTexture2D<float4>, SceneColor)
+	SHADER_PARAMETER_RDG_TEXTURE_UAV(RWTexture2D<float4>, SeparateCoc)
+END_SHADER_PARAMETER_STRUCT()
+
+FDOFGatherInputTextures CreateTextures(FRDGBuilder& GraphBuilder, const FDOFGatherInputDescs& Descs, const TCHAR* DebugName)
+{
+	FDOFGatherInputTextures Textures;
+	Textures.SceneColor = GraphBuilder.CreateTexture(Descs.SceneColor, DebugName);
+	if (Descs.SeparateCoc.Format != PF_Unknown)
+		Textures.SeparateCoc = GraphBuilder.CreateTexture(Descs.SeparateCoc, DebugName);
+	return Textures;
+}
+
+FDOFGatherInputUAVs CreateUAVs(FRDGBuilder& GraphBuilder, const FDOFGatherInputTextures& Textures, uint8 MipLevel = 0)
+{
+	FDOFGatherInputUAVs UAVs;
+	UAVs.SceneColor = GraphBuilder.CreateUAV(FRDGTextureUAVDesc(Textures.SceneColor, MipLevel));
+	if (Textures.SeparateCoc)
+		UAVs.SeparateCoc = GraphBuilder.CreateUAV(FRDGTextureUAVDesc(Textures.SeparateCoc, MipLevel));
+	return UAVs;
+}
+
+
+struct FDOFConvolutionDescs
+{
+	FRDGTextureDesc SceneColor;
+	FRDGTextureDesc SeparateAlpha;
+};
+
+BEGIN_SHADER_PARAMETER_STRUCT(FDOFConvolutionTextures, )
+	SHADER_PARAMETER_RDG_TEXTURE(Texture2D, SceneColor)
+	SHADER_PARAMETER_RDG_TEXTURE(Texture2D, SeparateAlpha)
+END_SHADER_PARAMETER_STRUCT()
+
+BEGIN_SHADER_PARAMETER_STRUCT(FDOFConvolutionUAVs, )
+	SHADER_PARAMETER_RDG_TEXTURE_UAV(RWTexture2D<float4>, SceneColor)
+	SHADER_PARAMETER_RDG_TEXTURE_UAV(RWTexture2D<float4>, SeparateAlpha)
+END_SHADER_PARAMETER_STRUCT()
+
+FDOFConvolutionTextures CreateTextures(FRDGBuilder& GraphBuilder, const FDOFConvolutionDescs& Descs, const TCHAR* DebugName)
+{
+	FDOFConvolutionTextures Textures;
+	Textures.SceneColor = GraphBuilder.CreateTexture(Descs.SceneColor, DebugName);
+	if (Descs.SeparateAlpha.Format != PF_Unknown)
+		Textures.SeparateAlpha = GraphBuilder.CreateTexture(Descs.SeparateAlpha, DebugName);
+	return Textures;
+}
+
+FDOFConvolutionUAVs CreateUAVs(FRDGBuilder& GraphBuilder, const FDOFConvolutionTextures& Textures)
+{
+	FDOFConvolutionUAVs UAVs;
+	UAVs.SceneColor = GraphBuilder.CreateUAV(Textures.SceneColor);
+	if (Textures.SeparateAlpha)
+		UAVs.SeparateAlpha = GraphBuilder.CreateUAV(Textures.SeparateAlpha);
+	return UAVs;
+}
+
+
+struct FDOFTileClassificationDescs
+{
+	FRDGTextureDesc Foreground;
+	FRDGTextureDesc Background;
+};
+
+BEGIN_SHADER_PARAMETER_STRUCT(FDOFTileClassificationTextures, )
+	SHADER_PARAMETER_RDG_TEXTURE(Texture2D, Foreground)
+	SHADER_PARAMETER_RDG_TEXTURE(Texture2D, Background)
+END_SHADER_PARAMETER_STRUCT()
+
+BEGIN_SHADER_PARAMETER_STRUCT(FDOFTileClassificationUAVs, )
+	SHADER_PARAMETER_RDG_TEXTURE_UAV(RWTexture2D<float4>, Foreground)
+	SHADER_PARAMETER_RDG_TEXTURE_UAV(RWTexture2D<float4>, Background)
+END_SHADER_PARAMETER_STRUCT()
+
+FDOFTileClassificationTextures CreateTextures(FRDGBuilder& GraphBuilder, const FDOFTileClassificationDescs& Descs, const TCHAR* DebugNames[2])
+{
+	FDOFTileClassificationTextures Textures;
+	Textures.Foreground = GraphBuilder.CreateTexture(Descs.Foreground, DebugNames[0]);
+	Textures.Background = GraphBuilder.CreateTexture(Descs.Background, DebugNames[1]);
+	return Textures;
+}
+
+FDOFTileClassificationUAVs CreateUAVs(FRDGBuilder& GraphBuilder, const FDOFTileClassificationTextures& Textures)
+{
+	FDOFTileClassificationUAVs UAVs;
+	UAVs.Foreground = GraphBuilder.CreateUAV(Textures.Foreground);
+	UAVs.Background = GraphBuilder.CreateUAV(Textures.Background);
+	return UAVs;
+}
+
+
+BEGIN_SHADER_PARAMETER_STRUCT(FDOFCommonShaderParameters, )
+	SHADER_PARAMETER_STRUCT_REF(FViewUniformShaderParameters, ViewUniformBuffer)
+END_SHADER_PARAMETER_STRUCT()
+
+
+BEGIN_SHADER_PARAMETER_STRUCT(FDOFCocModelShaderParameters, )
+	SHADER_PARAMETER(FVector4, CocModelParameters)
+	SHADER_PARAMETER(FVector2D, DepthBlurParameters)
+END_SHADER_PARAMETER_STRUCT()
+
+void SetCocModelParameters(
+	FDOFCocModelShaderParameters* OutParameters,
+	const DiaphragmDOF::FPhysicalCocModel& CocModel,
+	float CocRadiusBasis = 1.0f)
+{
+	OutParameters->CocModelParameters.X = CocRadiusBasis * CocModel.InfinityBackgroundCocRadius;
+	OutParameters->CocModelParameters.Y = CocRadiusBasis * CocModel.MinForegroundCocRadius;
+	OutParameters->CocModelParameters.Z = CocRadiusBasis * CocModel.MaxBackgroundCocRadius;
+	OutParameters->DepthBlurParameters.X = CocModel.DepthBlurExponent;
+	OutParameters->DepthBlurParameters.Y = CocRadiusBasis * CocModel.MaxDepthBlurRadius;
+}
+
+} // namespace
+
+
+// ---------------------------------------------------- Shaders
+
+namespace
+{
+
+class FDiaphragmDOFSetupCS : public FDiaphragmDOFShader
+{
+	DECLARE_GLOBAL_SHADER(FDiaphragmDOFSetupCS);
+	SHADER_USE_PARAMETER_STRUCT(FDiaphragmDOFSetupCS, FDiaphragmDOFShader);
+
+	class FOutputResDivisor : SHADER_PERMUTATION_INT("DIM_OUTPUT_RES_DIVISOR", 3);
+
+	using FPermutationDomain = TShaderPermutationDomain<FOutputResDivisor>;
+
+	BEGIN_SHADER_PARAMETER_STRUCT(FParameters, )
+		SHADER_PARAMETER_STRUCT_INCLUDE(FDOFCommonShaderParameters, CommonParameters)
+		SHADER_PARAMETER_STRUCT_INCLUDE(FDOFCocModelShaderParameters, CocModel)
+
+		SHADER_PARAMETER(FVector4, ViewportRect)
+		SHADER_PARAMETER(FVector2D, CocRadiusBasis) // TODO: decompose
+
+		SHADER_PARAMETER_RDG_TEXTURE(Texture2D, SceneColorTexture)
+		SHADER_PARAMETER_RDG_TEXTURE(Texture2D, SceneDepthTexture)
+
+		SHADER_PARAMETER_RDG_TEXTURE_UAV(RWTexture2D<float4>, Output0)
+		SHADER_PARAMETER_RDG_TEXTURE_UAV(RWTexture2D<float4>, Output1)
+		SHADER_PARAMETER_RDG_TEXTURE_UAV(RWTexture2D<float4>, Output2)
+	END_SHADER_PARAMETER_STRUCT()
+};
+
+class FDiaphragmDOFCocFlattenCS : public FDiaphragmDOFShader
+{
+	DECLARE_GLOBAL_SHADER(FDiaphragmDOFCocFlattenCS);
+	SHADER_USE_PARAMETER_STRUCT(FDiaphragmDOFCocFlattenCS, FDiaphragmDOFShader);
+
+	class FDoCocGather4 : SHADER_PERMUTATION_BOOL("DIM_DO_COC_GATHER4");
+
+	using FPermutationDomain = TShaderPermutationDomain<FDoCocGather4>;
+	
+	BEGIN_SHADER_PARAMETER_STRUCT(FParameters, )
+		SHADER_PARAMETER(FIntRect, ViewportRect)
+		SHADER_PARAMETER(FVector2D, ThreadIdToBufferUV)
+		SHADER_PARAMETER(FVector2D, MaxBufferUV)
+		
+		SHADER_PARAMETER_STRUCT_INCLUDE(FDOFCommonShaderParameters, CommonParameters)
+		SHADER_PARAMETER_STRUCT(FDOFGatherInputTextures, GatherInput)
+		SHADER_PARAMETER_STRUCT(FDOFTileClassificationUAVs, TileOutput)
+	END_SHADER_PARAMETER_STRUCT()
+
+	static_assert(FSceneViewScreenPercentageConfig::kMinTAAUpsampleResolutionFraction == 0.5f,
+		"Gather4 shader permutation assumes with min TAAU screen percentage = 50%.");
+	static_assert(FSceneViewScreenPercentageConfig::kMaxTAAUpsampleResolutionFraction == 2.0f,
+		"Gather4 shader permutation assumes with max TAAU screen percentage = 200%.");
+};
+
+class FDiaphragmDOFCocDilateCS : public FDiaphragmDOFShader
+{
+	DECLARE_GLOBAL_SHADER(FDiaphragmDOFCocDilateCS);
+	SHADER_USE_PARAMETER_STRUCT(FDiaphragmDOFCocDilateCS, FDiaphragmDOFShader);
+
+	using FPermutationDomain = TShaderPermutationDomain<FDDOFDilateRadiusDim, FDDOFDilateModeDim>;
+	
+	BEGIN_SHADER_PARAMETER_STRUCT(FParameters, )
+		SHADER_PARAMETER(FIntRect, ViewportRect)
+		SHADER_PARAMETER(int32, SampleOffsetMultipler)
+		SHADER_PARAMETER(float, fSampleOffsetMultipler)
+		SHADER_PARAMETER(float, CocRadiusToBucketDistanceUpperBound)
+		SHADER_PARAMETER(float, BucketDistanceToCocRadius)
+		
+		SHADER_PARAMETER_STRUCT_INCLUDE(FDOFCommonShaderParameters, CommonParameters)
+		SHADER_PARAMETER_STRUCT(FDOFTileClassificationTextures, TileInput)
+		SHADER_PARAMETER_STRUCT(FDOFTileClassificationTextures, DilatedTileMinMax)
+		SHADER_PARAMETER_STRUCT(FDOFTileClassificationUAVs, TileOutput)
+	END_SHADER_PARAMETER_STRUCT()
+};
+
+class FDiaphragmDOFDownsampleCS : public FDiaphragmDOFShader
+{
+	DECLARE_GLOBAL_SHADER(FDiaphragmDOFDownsampleCS);
+	SHADER_USE_PARAMETER_STRUCT(FDiaphragmDOFDownsampleCS, FDiaphragmDOFShader);
+
+	static bool ShouldCompilePermutation(const FGlobalShaderPermutationParameters& Parameters)
+	{
+		if (!FRCPassDiaphragmDOFHybridScatter::IsSupported(Parameters.Platform))
+		{
+			return false;
+		}
+		return FDiaphragmDOFShader::ShouldCompilePermutation(Parameters);
+	}
+	
+	BEGIN_SHADER_PARAMETER_STRUCT(FParameters, )
+		SHADER_PARAMETER(FIntRect, ViewportRect)
+		SHADER_PARAMETER(FVector2D, MaxBufferUV)
+		SHADER_PARAMETER(float, OutputCocRadiusMultiplier)
+
+		SHADER_PARAMETER(FVector4, GatherInputSize)
+		SHADER_PARAMETER_STRUCT(FDOFGatherInputTextures, GatherInput)
+
+		SHADER_PARAMETER_STRUCT_INCLUDE(FDOFCommonShaderParameters, CommonParameters)
+		SHADER_PARAMETER_STRUCT(FDOFGatherInputUAVs, OutDownsampledGatherInput)
+	END_SHADER_PARAMETER_STRUCT()
+};
+
+class FDiaphragmDOFReduceCS : public FDiaphragmDOFShader
+{
+	DECLARE_GLOBAL_SHADER(FDiaphragmDOFReduceCS);
+	SHADER_USE_PARAMETER_STRUCT(FDiaphragmDOFReduceCS, FDiaphragmDOFShader);
+
+	class FReduceMipCount : SHADER_PERMUTATION_RANGE_INT("DIM_REDUCE_MIP_COUNT", 2, 3);
+	class FHybridScatterForeground : SHADER_PERMUTATION_BOOL("DIM_HYBRID_SCATTER_FGD");
+	class FHybridScatterBackground : SHADER_PERMUTATION_BOOL("DIM_HYBRID_SCATTER_BGD");
+
+	using FPermutationDomain = TShaderPermutationDomain<
+		FReduceMipCount,
+		FHybridScatterForeground,
+		FHybridScatterBackground,
+		FDDOFRGBColorBufferDim>;
+
+	static bool ShouldCompilePermutation(const FGlobalShaderPermutationParameters& Parameters)
+	{
+		FPermutationDomain PermutationVector(Parameters.PermutationId);
+
+		// Do not compile storing Coc independently of RGB if not supported.
+		if (PermutationVector.Get<FDDOFRGBColorBufferDim>() && !FRCPassDiaphragmDOFGather::SupportRGBColorBuffer(Parameters.Platform))
+		{
+			return false;
+		}
+
+		if (!FRCPassDiaphragmDOFHybridScatter::IsSupported(Parameters.Platform))
+		{
+			if (PermutationVector.Get<FHybridScatterForeground>() || PermutationVector.Get<FHybridScatterBackground>())
+			{
+				return false;
+			}
+		}
+
+		return FDiaphragmDOFShader::ShouldCompilePermutation(Parameters);
+	}
+
+	BEGIN_SHADER_PARAMETER_STRUCT(FParameters, )
+		SHADER_PARAMETER(FIntRect, ViewportRect)
+		SHADER_PARAMETER(FVector2D, MaxInputBufferUV)
+		SHADER_PARAMETER(int32, MaxScatteringGroupCount)
+		SHADER_PARAMETER(float, PreProcessingToProcessingCocRadiusFactor)
+		SHADER_PARAMETER(float, MinScatteringCocRadius)
+		SHADER_PARAMETER(float, NeighborCompareMaxColor)
+		
+		SHADER_PARAMETER_RDG_TEXTURE(Texture2D, EyeAdaptation)
+		SHADER_PARAMETER_STRUCT_INCLUDE(FDOFCommonShaderParameters, CommonParameters)
+
+		SHADER_PARAMETER(FVector4, GatherInputSize)
+		SHADER_PARAMETER_STRUCT(FDOFGatherInputTextures, GatherInput)
+		
+		SHADER_PARAMETER(FVector4, QuarterResGatherInputSize)
+		SHADER_PARAMETER_STRUCT(FDOFGatherInputTextures, QuarterResGatherInput)
+
+		SHADER_PARAMETER_STRUCT_ARRAY(FDOFGatherInputUAVs, OutputMips, [FRCPassDiaphragmDOFReduce::kMaxMipLevelCount])
+		SHADER_PARAMETER_RDG_BUFFER_UAV(RWBuffer<uint>, OutScatterDrawIndirectParameters)
+		SHADER_PARAMETER_RDG_BUFFER_UAV(RWStructuredBuffer<float4>, OutForegroundScatterDrawList)
+		SHADER_PARAMETER_RDG_BUFFER_UAV(RWStructuredBuffer<float4>, OutBackgroundScatterDrawList)
+	END_SHADER_PARAMETER_STRUCT()
+};
+
+class FDiaphragmDOFScatterGroupPackCS : public FDiaphragmDOFShader
+{
+	DECLARE_GLOBAL_SHADER(FDiaphragmDOFScatterGroupPackCS);
+	SHADER_USE_PARAMETER_STRUCT(FDiaphragmDOFScatterGroupPackCS, FDiaphragmDOFShader);
+
+	static bool ShouldCompilePermutation(const FGlobalShaderPermutationParameters& Parameters)
+	{
+		if (!FRCPassDiaphragmDOFHybridScatter::IsSupported(Parameters.Platform))
+		{
+			return false;
+		}
+		return FDiaphragmDOFShader::ShouldCompilePermutation(Parameters);
+	}
+	
+	BEGIN_SHADER_PARAMETER_STRUCT(FParameters, )
+		SHADER_PARAMETER(int32, MaxScatteringGroupCount)
+
+		SHADER_PARAMETER_RDG_BUFFER_UAV(RWBuffer<uint>, OutScatterDrawIndirectParameters)
+		SHADER_PARAMETER_RDG_BUFFER_UAV(RWStructuredBuffer<float4>, OutForegroundScatterDrawList)
+		SHADER_PARAMETER_RDG_BUFFER_UAV(RWStructuredBuffer<float4>, OutBackgroundScatterDrawList)
+	END_SHADER_PARAMETER_STRUCT()
+};
+
+class FDiaphragmDOFBuildBokehLUTCS : public FDiaphragmDOFShader
+{
+	DECLARE_GLOBAL_SHADER(FDiaphragmDOFBuildBokehLUTCS);
+	SHADER_USE_PARAMETER_STRUCT(FDiaphragmDOFBuildBokehLUTCS, FDiaphragmDOFShader);
+
+	class FBokehSimulationDim : SHADER_PERMUTATION_BOOL("DIM_ROUND_BLADES");
+	class FLUTFormatDim : SHADER_PERMUTATION_ENUM_CLASS("DIM_LUT_FORMAT", FRCPassDiaphragmDOFBuildBokehLUT::EFormat);
+
+	using FPermutationDomain = TShaderPermutationDomain<FBokehSimulationDim, FLUTFormatDim>;
+	
+	BEGIN_SHADER_PARAMETER_STRUCT(FParameters, )
+		SHADER_PARAMETER(uint32, BladeCount)
+		SHADER_PARAMETER(float, DiaphragmRotation)
+		SHADER_PARAMETER(float, CocRadiusToCircumscribedRadius)
+		SHADER_PARAMETER(float, CocRadiusToIncircleRadius)
+		SHADER_PARAMETER(float, DiaphragmBladeRadius)
+		SHADER_PARAMETER(float, DiaphragmBladeCenterOffset)
+
+		SHADER_PARAMETER_RDG_TEXTURE_UAV(RWStructuredBuffer<float4>, BokehLUTOutput)
+	END_SHADER_PARAMETER_STRUCT()
+};
+
+class FDiaphragmDOFGatherCS : public FDiaphragmDOFShader
+{
+	DECLARE_GLOBAL_SHADER(FDiaphragmDOFGatherCS);
+	SHADER_USE_PARAMETER_STRUCT(FDiaphragmDOFGatherCS, FDiaphragmDOFShader);
+
+	using FPermutationDomain = TShaderPermutationDomain<
+		FDDOFLayerProcessingDim,
+		FDDOFGatherRingCountDim,
+		FDDOFBokehSimulationDim,
+		FDDOFGatherQualityDim,
+		FDDOFClampInputUVDim,
+		FDDOFRGBColorBufferDim>;
+
+	static FPermutationDomain RemapPermutation(FPermutationDomain PermutationVector)
+	{
+		// There is a lot of permutation, so no longer compile permutation.
+		if (1)
+		{
+			// Alway clamp input buffer UV.
+			PermutationVector.Set<FDDOFClampInputUVDim>(true);
+
+			// Alway simulate bokeh generically.
+			if (PermutationVector.Get<FDDOFBokehSimulationDim>() == EDiaphragmDOFBokehSimulation::SimmetricBokeh)
+			{
+				PermutationVector.Set<FDDOFBokehSimulationDim>(EDiaphragmDOFBokehSimulation::GenericBokeh);
+			}
+		}
+
+		// Slight out of focus only need 3 rings.
+		if (PermutationVector.Get<FDDOFLayerProcessingDim>() == EDiaphragmDOFLayerProcessing::SlightOutOfFocus)
+			PermutationVector.Set<FDDOFGatherRingCountDim>(FRCPassDiaphragmDOFGather::kMaxSlightOutOfFocusRingCount);
+
+		return PermutationVector;
+	}
+
+	static bool ShouldCompilePermutation(const FGlobalShaderPermutationParameters& Parameters)
+	{
+		FPermutationDomain PermutationVector(Parameters.PermutationId);
+
+		// Do not compile this permutation if we know this is going to be remapped.
+		if (RemapPermutation(PermutationVector) != PermutationVector)
+		{
+			return false;
+		}
+
+		// Some platforms might be to slow for even considering large number of gathering samples.
+		if (PermutationVector.Get<FDDOFGatherRingCountDim>() > FRCPassDiaphragmDOFGather::MaxRingCount(Parameters.Platform))
+		{
+			return false;
+		}
+
+		// Do not compile storing Coc independently of RGB.
+		if (PermutationVector.Get<FDDOFRGBColorBufferDim>() && !FRCPassDiaphragmDOFGather::SupportRGBColorBuffer(Parameters.Platform))
+		{
+			return false;
+		}
+
+		// No point compiling gather pass with hybrid scatter occlusion if the shader platform doesn't support if.
+		if (!FRCPassDiaphragmDOFHybridScatter::IsSupported(Parameters.Platform) &&
+			PermutationVector.Get<FDDOFGatherQualityDim>() == FRCPassDiaphragmDOFGather::EQualityConfig::HighQualityWithHybridScatterOcclusion) return false;
+
+		// Do not compile bokeh simulation shaders on platform that couldn't handle them anyway.
+		if (!FRCPassDiaphragmDOFGather::SupportsBokehSimmulation(Parameters.Platform) &&
+			PermutationVector.Get<FDDOFBokehSimulationDim>() != EDiaphragmDOFBokehSimulation::Disabled) return false;
+
+		if (PermutationVector.Get<FDDOFLayerProcessingDim>() == EDiaphragmDOFLayerProcessing::ForegroundOnly)
+		{
+			// Foreground does not support CocVariance output yet.
+			if (PermutationVector.Get<FDDOFGatherQualityDim>() == FRCPassDiaphragmDOFGather::EQualityConfig::HighQualityWithHybridScatterOcclusion) return false;
+
+			// Storing Coc independently of RGB is only supported for low gathering quality.
+			if (PermutationVector.Get<FDDOFRGBColorBufferDim>() &&
+				PermutationVector.Get<FDDOFGatherQualityDim>() != FRCPassDiaphragmDOFGather::EQualityConfig::LowQualityAccumulator)
+				return false;
+		}
+		else if (PermutationVector.Get<FDDOFLayerProcessingDim>() == EDiaphragmDOFLayerProcessing::ForegroundHoleFilling)
+		{
+			// Foreground hole filling does not need to output CocVariance, since this is the job of foreground pass.
+			if (PermutationVector.Get<FDDOFGatherQualityDim>() == FRCPassDiaphragmDOFGather::EQualityConfig::HighQualityWithHybridScatterOcclusion) return false;
+
+			// Foreground hole filling doesn't have lower quality accumulator.
+			if (PermutationVector.Get<FDDOFGatherQualityDim>() == FRCPassDiaphragmDOFGather::EQualityConfig::LowQualityAccumulator) return false;
+
+			// Foreground hole filling doesn't need cinematic quality.
+			if (PermutationVector.Get<FDDOFGatherQualityDim>() == FRCPassDiaphragmDOFGather::EQualityConfig::Cinematic) return false;
+
+			// No bokeh simulation on hole filling, always use euclidian closest distance to compute opacity alpha channel.
+			if (PermutationVector.Get<FDDOFBokehSimulationDim>() != EDiaphragmDOFBokehSimulation::Disabled) return false;
+
+			// Storing Coc independently of RGB is only supported for RecombineQuality == 0.
+			if (PermutationVector.Get<FDDOFRGBColorBufferDim>())
+				return false;
+		}
+		else if (PermutationVector.Get<FDDOFLayerProcessingDim>() == EDiaphragmDOFLayerProcessing::SlightOutOfFocus)
+		{
+			// Slight out of focus gather pass does not need large radius since only accumulating only
+			// abs(CocRadius) < kMaxSlightOutOfFocusRingCount.
+			if (PermutationVector.Get<FDDOFGatherRingCountDim>() > FRCPassDiaphragmDOFGather::kMaxSlightOutOfFocusRingCount) return false;
+
+			// Slight out of focus don't need to output CocVariance since there is no hybrid scattering.
+			if (PermutationVector.Get<FDDOFGatherQualityDim>() == FRCPassDiaphragmDOFGather::EQualityConfig::HighQualityWithHybridScatterOcclusion) return false;
+
+			// Slight out of focus filling can't have lower quality accumulator since it needs to brute force the focus areas.
+			if (PermutationVector.Get<FDDOFGatherQualityDim>() == FRCPassDiaphragmDOFGather::EQualityConfig::LowQualityAccumulator) return false;
+
+			// Slight out of focus doesn't have cinematic quality, yet.
+			if (PermutationVector.Get<FDDOFGatherQualityDim>() == FRCPassDiaphragmDOFGather::EQualityConfig::Cinematic) return false;
+
+			// Storing Coc independently of RGB is only supported for RecombineQuality == 0.
+			if (PermutationVector.Get<FDDOFRGBColorBufferDim>())
+				return false;
+		}
+		else if (PermutationVector.Get<FDDOFLayerProcessingDim>() == EDiaphragmDOFLayerProcessing::BackgroundOnly)
+		{
+			// There is no performance point doing high quality gathering without scattering occlusion.
+			if (PermutationVector.Get<FDDOFGatherQualityDim>() == FRCPassDiaphragmDOFGather::EQualityConfig::HighQuality) return false;
+
+			// Storing Coc independently of RGB is only supported for low gathering quality.
+			if (PermutationVector.Get<FDDOFRGBColorBufferDim>() &&
+				PermutationVector.Get<FDDOFGatherQualityDim>() != FRCPassDiaphragmDOFGather::EQualityConfig::LowQualityAccumulator)
+				return false;
+		}
+		else if (PermutationVector.Get<FDDOFLayerProcessingDim>() == EDiaphragmDOFLayerProcessing::ForegroundAndBackground)
+		{
+			// Gathering foreground and backrgound at same time is not supported yet.
+			return false;
+		}
+		else
+		{
+			check(0);
+		}
+
+		return FDiaphragmDOFShader::ShouldCompilePermutation(Parameters);
+	}
+
+	static void ModifyCompilationEnvironment(const FGlobalShaderPermutationParameters& Parameters, FShaderCompilerEnvironment& OutEnvironment)
+	{
+		// The gathering pass shader code gives a really hard time to the HLSL compiler. To improve
+		// iteration time on the shader, only pass down a /O1 instead of /O3.
+		if (Parameters.Platform == SP_PCD3D_SM5)
+		{
+			OutEnvironment.CompilerFlags.Add(CFLAG_StandardOptimization);
+		}
+
+		FDiaphragmDOFShader::ModifyCompilationEnvironment(Parameters, OutEnvironment);
+	}
+
+	
+	BEGIN_SHADER_PARAMETER_STRUCT(FParameters, )
+		SHADER_PARAMETER(FVector4, ViewportSize)
+		SHADER_PARAMETER(FIntRect, ViewportRect)
+		SHADER_PARAMETER(FVector2D, TemporalJitterPixels)
+		SHADER_PARAMETER(FVector2D, DispatchThreadIdToInputBufferUV)
+		SHADER_PARAMETER(FVector2D, ConsiderCocRadiusAffineTransformation0)
+		SHADER_PARAMETER(FVector2D, ConsiderCocRadiusAffineTransformation1)
+		SHADER_PARAMETER(FVector2D, ConsiderAbsCocRadiusAffineTransformation)
+		SHADER_PARAMETER(FVector2D, InputBufferUVToOutputPixel)
+		SHADER_PARAMETER(float, MipBias)
+		SHADER_PARAMETER(float, MaxRecombineAbsCocRadius)
+		
+		SHADER_PARAMETER_STRUCT_INCLUDE(FDOFCommonShaderParameters, CommonParameters)
+
+		SHADER_PARAMETER(FVector4, GatherInputSize)
+		SHADER_PARAMETER_STRUCT(FDOFGatherInputTextures, GatherInput)
+
+		SHADER_PARAMETER_STRUCT(FDOFTileClassificationTextures, TileClassification)
+		SHADER_PARAMETER_RDG_TEXTURE(Texture2D, BokehLUT)
+
+		SHADER_PARAMETER_STRUCT(FDOFConvolutionUAVs, ConvolutionOutput)
+		SHADER_PARAMETER_RDG_TEXTURE_UAV(RWTexture2D<float4>, ScatterOcclusionOutput)
+	END_SHADER_PARAMETER_STRUCT()
+}; // class FDiaphragmDOFGatherCS
+
+class FDiaphragmDOFPostfilterCS : public FDiaphragmDOFShader
+{
+	DECLARE_GLOBAL_SHADER(FDiaphragmDOFPostfilterCS);
+	SHADER_USE_PARAMETER_STRUCT(FDiaphragmDOFPostfilterCS, FDiaphragmDOFShader);
+
+	class FTileOptimization : SHADER_PERMUTATION_BOOL("DIM_TILE_PERMUTATION");
+
+	using FPermutationDomain = TShaderPermutationDomain<FDDOFLayerProcessingDim, FDDOFPostfilterMethodDim, FTileOptimization>;
+
+	static FPermutationDomain RemapPermutationVector(FPermutationDomain PermutationVector)
+	{
+		// Tile permutation optimisation is only for Max3x3 post filtering.
+		if (PermutationVector.Get<FDDOFPostfilterMethodDim>() != EDiaphragmDOFPostfilterMethod::RGBMax3x3)
+			PermutationVector.Set<FTileOptimization>(false);
+		return PermutationVector;
+	}
+
+	static bool ShouldCompilePermutation(const FGlobalShaderPermutationParameters& Parameters)
+	{
+		FPermutationDomain PermutationVector(Parameters.PermutationId);
+		if (RemapPermutationVector(PermutationVector) != PermutationVector) return false;
+		if (PermutationVector.Get<FDDOFPostfilterMethodDim>() == EDiaphragmDOFPostfilterMethod::None) return false;
+		if (PermutationVector.Get<FDDOFLayerProcessingDim>() != EDiaphragmDOFLayerProcessing::ForegroundOnly &&
+			PermutationVector.Get<FDDOFLayerProcessingDim>() != EDiaphragmDOFLayerProcessing::BackgroundOnly) return false;
+
+		return FDiaphragmDOFShader::ShouldCompilePermutation(Parameters);
+	}
+	
+	BEGIN_SHADER_PARAMETER_STRUCT(FParameters, )
+		SHADER_PARAMETER(FIntRect, ViewportRect)
+		SHADER_PARAMETER(FVector2D, MaxInputBufferUV)
+		SHADER_PARAMETER(float, MinGatherRadius)
+		
+		SHADER_PARAMETER_STRUCT_INCLUDE(FDOFCommonShaderParameters, CommonParameters)
+
+		SHADER_PARAMETER(FVector4, ConvolutionInputSize)
+		SHADER_PARAMETER_STRUCT(FDOFConvolutionTextures, ConvolutionInput)
+
+		SHADER_PARAMETER_STRUCT(FDOFTileClassificationTextures, TileClassification)
+		SHADER_PARAMETER_STRUCT(FDOFConvolutionUAVs, ConvolutionOutput)
+	END_SHADER_PARAMETER_STRUCT()
+}; // class FDiaphragmDOFPostfilterCS
+
+BEGIN_SHADER_PARAMETER_STRUCT(FDOFHybridScatterParameters, )
+	SHADER_PARAMETER(FVector4, ViewportSize)
+	SHADER_PARAMETER(float, CocRadiusToCircumscribedRadius)
+	SHADER_PARAMETER(float, ScatteringScaling)
+	
+	SHADER_PARAMETER_STRUCT_INCLUDE(FDOFCommonShaderParameters, CommonParameters)
+	
+	SHADER_PARAMETER(FVector4, ScatterOcclusionSize)
+	SHADER_PARAMETER_RDG_TEXTURE(Texture2D, ScatterOcclusion)
+
+	SHADER_PARAMETER_RDG_TEXTURE(Texture2D, BokehLUT)
+
+	SHADER_PARAMETER_RDG_BUFFER(Buffer<uint>, IndirectDrawParameter)
+	SHADER_PARAMETER_RDG_BUFFER_SRV(StructuredBuffer<float4>, ScatterDrawList)
+
+	RENDER_TARGET_BINDING_SLOTS()
+END_SHADER_PARAMETER_STRUCT()
+
+class FDiaphragmDOFHybridScatterVS : public FDiaphragmDOFShader
+{
+	DECLARE_GLOBAL_SHADER(FDiaphragmDOFHybridScatterVS);
+	SHADER_USE_PARAMETER_STRUCT(FDiaphragmDOFHybridScatterVS, FDiaphragmDOFShader);
+
+	using FPermutationDomain = TShaderPermutationDomain<>;
+	using FParameters = FDOFHybridScatterParameters;
+
+	static bool ShouldCompilePermutation(const FGlobalShaderPermutationParameters& Parameters)
+	{
+		if (!FRCPassDiaphragmDOFHybridScatter::IsSupported(Parameters.Platform))
+		{
+			return false;
+		}
+		return FDiaphragmDOFShader::ShouldCompilePermutation(Parameters);
+	}
+};
+
+class FDiaphragmDOFHybridScatterPS : public FDiaphragmDOFShader
+{
+	DECLARE_GLOBAL_SHADER(FDiaphragmDOFHybridScatterPS);
+	SHADER_USE_PARAMETER_STRUCT(FDiaphragmDOFHybridScatterPS, FDiaphragmDOFShader);
+
+	class FBokehSimulationDim : SHADER_PERMUTATION_BOOL("DIM_BOKEH_SIMULATION");
+
+	using FPermutationDomain = TShaderPermutationDomain<FDDOFLayerProcessingDim, FBokehSimulationDim, FDDOFScatterOcclusionDim>;
+	using FParameters = FDOFHybridScatterParameters;
+
+	static FPermutationDomain RemapPermutation(FPermutationDomain PermutationVector)
+	{
+		// Pixel shader are exactly the same between foreground and background when there is no bokeh LUT.
+		if (PermutationVector.Get<FDDOFLayerProcessingDim>() == EDiaphragmDOFLayerProcessing::BackgroundOnly && 
+			!PermutationVector.Get<FBokehSimulationDim>())
+		{
+			PermutationVector.Set<FDDOFLayerProcessingDim>(EDiaphragmDOFLayerProcessing::ForegroundOnly);
+		}
+
+		return PermutationVector;
+	}
+
+	static bool ShouldCompilePermutation(const FGlobalShaderPermutationParameters& Parameters)
+	{
+		if (!FRCPassDiaphragmDOFHybridScatter::IsSupported(Parameters.Platform))
+		{
+			return false;
+		}
+
+		FPermutationDomain PermutationVector(Parameters.PermutationId);
+
+		// Do not compile this permutation if it gets remapped at runtime.
+		if (RemapPermutation(PermutationVector) != PermutationVector)
+		{
+			return false;
+		}
+
+		if (PermutationVector.Get<FDDOFLayerProcessingDim>() != EDiaphragmDOFLayerProcessing::ForegroundOnly &&
+			PermutationVector.Get<FDDOFLayerProcessingDim>() != EDiaphragmDOFLayerProcessing::BackgroundOnly) return false;
+
+		return FDiaphragmDOFShader::ShouldCompilePermutation(Parameters);
+	}
+};
+
+class FDiaphragmDOFRecombineCS : public FDiaphragmDOFShader
+{
+	DECLARE_GLOBAL_SHADER(FDiaphragmDOFRecombineCS);
+	SHADER_USE_PARAMETER_STRUCT(FDiaphragmDOFRecombineCS, FDiaphragmDOFShader);
+
+	class FQualityDim : SHADER_PERMUTATION_INT("DIM_QUALITY", 3);
+
+	using FPermutationDomain = TShaderPermutationDomain<FDDOFLayerProcessingDim, FDDOFBokehSimulationDim, FQualityDim>;
+
+	static bool ShouldCompilePermutation(const FGlobalShaderPermutationParameters& Parameters)
+	{
+		FPermutationDomain PermutationVector(Parameters.PermutationId);
+		if (PermutationVector.Get<FDDOFLayerProcessingDim>() != EDiaphragmDOFLayerProcessing::ForegroundOnly &&
+		    PermutationVector.Get<FDDOFLayerProcessingDim>() != EDiaphragmDOFLayerProcessing::BackgroundOnly &&
+			PermutationVector.Get<FDDOFLayerProcessingDim>() != EDiaphragmDOFLayerProcessing::ForegroundAndBackground) return false;
+
+		// Do not compile bokeh simulation shaders on platform that couldn't handle them anyway.
+		if (!FRCPassDiaphragmDOFGather::SupportsBokehSimmulation(Parameters.Platform) &&
+			PermutationVector.Get<FDDOFBokehSimulationDim>() != EDiaphragmDOFBokehSimulation::Disabled) return false;
+
+		return FDiaphragmDOFShader::ShouldCompilePermutation(Parameters);
+	}
+	
+	BEGIN_SHADER_PARAMETER_STRUCT(FParameters, )
+		SHADER_PARAMETER_STRUCT_INCLUDE(FDOFCommonShaderParameters, CommonParameters)
+		SHADER_PARAMETER_STRUCT_INCLUDE(FDOFCocModelShaderParameters, CocModel)
+
+		SHADER_PARAMETER(FIntRect, ViewportRect)
+		SHADER_PARAMETER(FVector4, ViewportSize)
+		SHADER_PARAMETER(FVector2D, TemporalJitterPixels)
+		SHADER_PARAMETER(FVector2D, DOFBufferUVMax)
+		
+		SHADER_PARAMETER_RDG_TEXTURE(Texture2D, BokehLUT)
+
+		// Full res textures.
+		SHADER_PARAMETER_RDG_TEXTURE(Texture2D, SceneColorInput)
+		SHADER_PARAMETER_RDG_TEXTURE(Texture2D, SceneDepthTexture)
+		SHADER_PARAMETER_RDG_TEXTURE(Texture2D, SceneSeparateCoc)
+		SHADER_PARAMETER_RDG_TEXTURE(Texture2D, SceneSeparateTranslucency)
+
+		// Half res convolution textures.
+		SHADER_PARAMETER(FVector4, ConvolutionInputSize)
+		SHADER_PARAMETER_STRUCT(FDOFConvolutionTextures, ForegroundConvolution)
+		SHADER_PARAMETER_STRUCT(FDOFConvolutionTextures, ForegroundHoleFillingConvolution)
+		SHADER_PARAMETER_STRUCT(FDOFConvolutionTextures, SlightOutOfFocusConvolution)
+		SHADER_PARAMETER_STRUCT(FDOFConvolutionTextures, BackgroundConvolution)
+
+		SHADER_PARAMETER_RDG_TEXTURE_UAV(RWTexture<float4>, SceneColorOutput)
+	END_SHADER_PARAMETER_STRUCT()
+};
+
+IMPLEMENT_GLOBAL_SHADER(FDiaphragmDOFSetupCS,            "/Engine/Private/DiaphragmDOF/DOFSetup.usf",                     "SetupCS",                SF_Compute);
+IMPLEMENT_GLOBAL_SHADER(FDiaphragmDOFCocFlattenCS,       "/Engine/Private/DiaphragmDOF/DOFCocTileFlatten.usf",            "CocFlattenMainCS",       SF_Compute);
+IMPLEMENT_GLOBAL_SHADER(FDiaphragmDOFCocDilateCS,        "/Engine/Private/DiaphragmDOF/DOFCocTileDilate.usf",             "CocDilateMainCS",        SF_Compute);
+IMPLEMENT_GLOBAL_SHADER(FDiaphragmDOFDownsampleCS,       "/Engine/Private/DiaphragmDOF/DOFDownsample.usf",                "DownsampleCS",           SF_Compute);
+IMPLEMENT_GLOBAL_SHADER(FDiaphragmDOFReduceCS,           "/Engine/Private/DiaphragmDOF/DOFReduce.usf",                    "ReduceCS",               SF_Compute);
+IMPLEMENT_GLOBAL_SHADER(FDiaphragmDOFScatterGroupPackCS, "/Engine/Private/DiaphragmDOF/DOFHybridScatterCompilation.usf",  "ScatterGroupPackMainCS", SF_Compute);
+IMPLEMENT_GLOBAL_SHADER(FDiaphragmDOFBuildBokehLUTCS,    "/Engine/Private/DiaphragmDOF/DOFBokehLUT.usf",                  "BuildBokehLUTMainCS",    SF_Compute);
+IMPLEMENT_GLOBAL_SHADER(FDiaphragmDOFGatherCS,           "/Engine/Private/DiaphragmDOF/DOFGatherPass.usf",                "GatherMainCS",           SF_Compute);
+IMPLEMENT_GLOBAL_SHADER(FDiaphragmDOFPostfilterCS,       "/Engine/Private/DiaphragmDOF/DOFPostfiltering.usf",             "PostfilterMainCS",       SF_Compute);
+IMPLEMENT_GLOBAL_SHADER(FDiaphragmDOFHybridScatterVS,    "/Engine/Private/DiaphragmDOF/DOFHybridScatterVertexShader.usf", "ScatterMainVS",          SF_Vertex);
+IMPLEMENT_GLOBAL_SHADER(FDiaphragmDOFHybridScatterPS,    "/Engine/Private/DiaphragmDOF/DOFHybridScatterPixelShader.usf",  "ScatterMainPS",          SF_Pixel);
+IMPLEMENT_GLOBAL_SHADER(FDiaphragmDOFRecombineCS,        "/Engine/Private/DiaphragmDOF/DOFRecombine.usf",                 "RecombineMainCS",        SF_Compute);
+
+} // namespace
+
+FRDGTextureRef DiaphragmDOF::AddPasses(
+	FRDGBuilder& GraphBuilder,
+	const FSceneViewFamilyBlackboard& SceneBlackboard,
+	const FViewInfo& View,
+	FRDGTextureRef InputSceneColor,
+	FRDGTextureRef SceneSeparateTranslucency)
+{
+	if (View.Family->EngineShowFlags.VisualizeDOF)
+	{
+		// no need for this pass
+		return InputSceneColor;
+	}
+
+	// Format of the scene color.
+	EPixelFormat SceneColorFormat = InputSceneColor->Desc.Format;
+	
+	// Whether should process alpha channel of the scene or not.
+	const bool bProcessSceneAlpha = FPostProcessing::HasAlphaChannelSupport();
+
+	const EShaderPlatform ShaderPlatform = View.GetShaderPlatform();
+
+	// Number of sampling ring in the gathering kernel.
+	const int32 HalfResRingCount = FMath::Clamp(CVarRingCount.GetValueOnRenderThread(), FRCPassDiaphragmDOFGather::kMinRingCount, FRCPassDiaphragmDOFGather::MaxRingCount(ShaderPlatform));
+
+	// Post filtering method to do.
+	const EDiaphragmDOFPostfilterMethod PostfilterMethod = GetPostfilteringMethod();
+
+	// The mode for hybrid scattering.
+	const EHybridScatterMode FgdHybridScatteringMode = EHybridScatterMode(CVarHybridScatterForegroundMode.GetValueOnRenderThread());
+	const EHybridScatterMode BgdHybridScatteringMode = EHybridScatterMode(CVarHybridScatterBackgroundMode.GetValueOnRenderThread());
+	
+	const float MinScatteringCocRadius = FMath::Max(CVarScatterMinCocRadius.GetValueOnRenderThread(), FRCPassDiaphragmDOFHybridScatter::kMinCocRadius);
+
+	// Whether the platform support gather bokeh simmulation.
+	const bool bSupportGatheringBokehSimulation = FRCPassDiaphragmDOFGather::SupportsBokehSimmulation(ShaderPlatform);
+
+	// Whether should use shade permutation that does lower quality accumulation.
+	// TODO: this is becoming a mess.
+	const bool bUseLowAccumulatorQuality = CVarAccumulatorQuality.GetValueOnRenderThread() == 0;
+	const bool bUseCinematicAccumulatorQuality = CVarAccumulatorQuality.GetValueOnRenderThread() == 2;
+
+	// Setting for scattering budget upper bound.
+	const float MaxScatteringRatio = FMath::Clamp(CVarScatterMaxSpriteRatio.GetValueOnRenderThread(), 0.0f, 1.0f);
+
+	// Slight out of focus is not supporting with DOF's TAA upsampling, because of the brute force
+	// kernel used in GatherCS for slight out of focus stability buffer.
+	const bool bSupportsSlightOutOfFocus = View.PrimaryScreenPercentageMethod != EPrimaryScreenPercentageMethod::TemporalUpscale;
+
+	// Quality setting for the recombine pass.
+	const int32 RecombineQuality = bSupportsSlightOutOfFocus ? FMath::Clamp(CVarRecombineQuality.GetValueOnRenderThread(), 0, FRCPassDiaphragmDOFRecombine::kMaxQuality) : 0;
+
+	// Resolution divisor.
+	// TODO: Exposes lower resolution divisor?
+	const int32 PrefilteringResolutionDivisor = 2;
+
+	// Minimal absolute Coc radius to spawn a gather pass. Blurring radius under this are considered not great looking.
+	// This is assuming the pass is opacity blending with a ramp from 1 to 2. This can not be exposed as a cvar,
+	// because the slight out focus's lower res pass uses for full res convolution stability depends on this.
+	const float kMinimalAbsGatherPassCocRadius = 1;
+
+	// Minimal CocRadius to wire lower res gathering passes.
+	const float BackgroundCocRadiusMaximumForUniquePass = HalfResRingCount * 4.0; // TODO: polish that.
+
+	// Whether the recombine pass does slight out of focus convolution.
+	const bool bRecombineDoesSlightOutOfFocus = RecombineQuality > 0;
+
+	// Whether the recombine pass wants separate input buffer for foreground hole filling.
+	const bool bRecombineDoesSeparateForegroundHoleFilling = RecombineQuality > 0;
+
+	// Compute the required blurring radius to actually perform depth of field, that depends on whether
+	// doing slight out of focus convolution.
+	const float MinRequiredBlurringRadius = bRecombineDoesSlightOutOfFocus
+		? (CVarMinimalFullresBlurRadius.GetValueOnRenderThread() * 0.5f)
+		: kMinimalAbsGatherPassCocRadius;
+
+	// Whether to use R11G11B10 + separate C0C buffer.
+	const bool bRGBBufferSeparateCocBuffer = (
+		SceneColorFormat == PF_FloatR11G11B10 &&
+
+		// Can't use FloatR11G11B10 if also need to support alpha channel.
+		!bProcessSceneAlpha &&
+
+		// This is just to get the number of shader permutation down.
+		RecombineQuality == 0 &&
+		bUseLowAccumulatorQuality &&
+		FRCPassDiaphragmDOFGather::SupportRGBColorBuffer(ShaderPlatform));
+
+
+	// Derives everything needed from the view.
+	FSceneViewState* const ViewState = View.ViewState;
+
+	FPhysicalCocModel CocModel;
+	CocModel.Compile(View);
+
+	FBokehModel BokehModel;
+	BokehModel.Compile(View);
+
+	// Prepare preprocessing TAA pass.
+	FTAAPassParameters TAAParameters(View);
+	{
+		TAAParameters.Pass = ETAAPassConfig::DiaphragmDOF;
+
+		// When using dynamic resolution, the blur introduce by TAA's history resolution changes is quite noticeable on DOF.
+		// Therefore we switch for a temporal upsampling technic to maintain the same history resolution.
+		if (View.PrimaryScreenPercentageMethod == EPrimaryScreenPercentageMethod::TemporalUpscale)
+		{
+			TAAParameters.Pass = ETAAPassConfig::DiaphragmDOFUpsampling;
+			TAAParameters.bIsComputePass = true;
+		}
+		
+		TAAParameters.SetupViewRect(View, PrefilteringResolutionDivisor);
+		TAAParameters.TopLeftCornerViewRects();
+
+		TAAParameters.bUseFast = CVarDOFTemporalAAQuality.GetValueOnRenderThread() == 0;
+	}
+
+	// Size of the view in GatherColorSetup.
+	FIntPoint FullResViewSize = View.ViewRect.Size();
+	FIntPoint PreprocessViewSize = FIntPoint::DivideAndRoundUp(FullResViewSize, PrefilteringResolutionDivisor);
+	FIntPoint GatheringViewSize = PreprocessViewSize;
+
+	if (View.AntiAliasingMethod == AAM_TemporalAA && ViewState)
+	{
+		PreprocessViewSize = FIntPoint::DivideAndRoundUp(TAAParameters.OutputViewRect.Size(), PrefilteringResolutionDivisor);
+	}
+
+	const float PreProcessingToProcessingCocRadiusFactor = float(GatheringViewSize.X) / float(PreprocessViewSize.X);
+
+	const float MaxBackgroundCocRadius = CocModel.ComputeViewMaxBackgroundCocRadius(GatheringViewSize.X);
+	const float MinForegroundCocRadius = CocModel.ComputeViewMinForegroundCocRadius(GatheringViewSize.X);
+	const float AbsMaxForegroundCocRadius = FMath::Abs(MinForegroundCocRadius);
+	const float MaxBluringRadius = FMath::Max(MaxBackgroundCocRadius, AbsMaxForegroundCocRadius);
+
+	// Whether should hybrid scatter for foreground and background.
+	bool bForegroundHybridScattering = FgdHybridScatteringMode != EHybridScatterMode::Disabled && AbsMaxForegroundCocRadius > MinScatteringCocRadius && MaxScatteringRatio > 0.0f;
+	bool bBackgroundHybridScattering = BgdHybridScatteringMode != EHybridScatterMode::Disabled && MaxBackgroundCocRadius > MinScatteringCocRadius && MaxScatteringRatio > 0.0f;
+
+	if (!FRCPassDiaphragmDOFHybridScatter::IsSupported(ShaderPlatform))
+	{
+		bForegroundHybridScattering = false;
+		bBackgroundHybridScattering = false;
+	}
+
+	// Compute the reference buffer size for PrefilteringResolutionDivisor.
+	FIntPoint RefBufferSize = FIntPoint::DivideAndRoundUp(InputSceneColor->Desc.Extent, PrefilteringResolutionDivisor);
+	
+	EDiaphragmDOFBokehSimulation BokehSimulation = EDiaphragmDOFBokehSimulation::Disabled;
+	if (BokehModel.BokehShape != EBokehShape::Circle)
+	{
+		BokehSimulation = (BokehModel.DiaphragmBladeCount % 2)
+			? EDiaphragmDOFBokehSimulation::GenericBokeh
+			: EDiaphragmDOFBokehSimulation::SimmetricBokeh;
+	}
+
+	// If the max blurring radius is too small, do not wire any passes.
+	if (MaxBluringRadius < MinRequiredBlurringRadius)
+	{
+		return InputSceneColor;
+	}
+
+	RDG_EVENT_SCOPE(GraphBuilder, "DiaphragmDOF");
+
+	bool bGatherForeground = AbsMaxForegroundCocRadius > kMinimalAbsGatherPassCocRadius;
+	
+	// Setup all the desciptors.
+	FRDGTextureDesc FullResDesc;
+	{
+		FullResDesc = InputSceneColor->Desc;
+
+		// Reset so that the number of samples of decsriptor becomes 1, which is totally legal still with MSAA because
+		// the scene color will already be resolved to ShaderResource texture that is always 1. This is to work around
+		// hack that MSAA will have targetable texture with MSAA != shader resource, and still have descriptor indicating
+		// the number of samples of the targetable resource.
+		FullResDesc.Reset();
+
+		FullResDesc.Format = PF_FloatRGBA;
+		FullResDesc.TargetableFlags |= TexCreate_UAV;
+		FullResDesc.Flags &= ~(TexCreate_FastVRAM);
+	}
+	
+	FDOFGatherInputDescs FullResGatherInputDescs;
+	{
+		FullResGatherInputDescs.SceneColor = FullResDesc;
+		FullResGatherInputDescs.SceneColor.Format = PF_FloatRGBA;
+		
+		FullResGatherInputDescs.SeparateCoc = FullResDesc;
+		FullResGatherInputDescs.SeparateCoc.Format = PF_R16F;
+	}
+
+	FDOFGatherInputDescs HalfResGatherInputDescs;
+	{
+		HalfResGatherInputDescs.SceneColor = FullResDesc;
+		HalfResGatherInputDescs.SceneColor.Extent /= PrefilteringResolutionDivisor;
+		HalfResGatherInputDescs.SceneColor.Format = PF_FloatRGBA;
+		HalfResGatherInputDescs.SceneColor.Flags |= GFastVRamConfig.DOFSetup;
+		
+		HalfResGatherInputDescs.SeparateCoc = FullResDesc;
+		HalfResGatherInputDescs.SeparateCoc.Extent /= PrefilteringResolutionDivisor;
+		HalfResGatherInputDescs.SeparateCoc.Format = PF_R16F;
+		HalfResGatherInputDescs.SeparateCoc.Flags |= GFastVRamConfig.DOFSetup;
+	}
+
+	// Setup the shader parameter used in all shaders.
+	FDOFCommonShaderParameters CommonParameters;
+	CommonParameters.ViewUniformBuffer = View.ViewUniformBuffer;
+
+	FDOFGatherInputTextures FullResGatherInputTextures;
+	FDOFGatherInputTextures HalfResGatherInputTextures;
+
+	// Setup at lower resolution from full resolution scene color and scene depth.
+	{
+		FullResGatherInputTextures = CreateTextures(GraphBuilder, FullResGatherInputDescs, TEXT("DOFFullResSetup"));
+		HalfResGatherInputTextures = CreateTextures(GraphBuilder, HalfResGatherInputDescs, TEXT("DOFHalfResSetup"));
+
+		bool bOutputFullResolution = bRecombineDoesSlightOutOfFocus && !bProcessSceneAlpha;
+		bool bOutputHalfResolution = true; // TODO: there is a useless shader permutation.
+		
+		FDiaphragmDOFSetupCS::FPermutationDomain PermutationVector;
+	
+		FIntPoint PassViewSize = FullResViewSize;
+		FIntPoint GroupSize(kDefaultGroupSize, kDefaultGroupSize);
+		float CocRadiusBasis = 1.0f;
+		if (bOutputFullResolution && bOutputHalfResolution)
+		{
+			PermutationVector.Set<FDiaphragmDOFSetupCS::FOutputResDivisor>(0);
+			GroupSize *= 2;
+		}
+		else if (bOutputFullResolution)
+		{
+			PermutationVector.Set<FDiaphragmDOFSetupCS::FOutputResDivisor>(1);
+		}
+		else if (bOutputHalfResolution)
+		{
+			PermutationVector.Set<FDiaphragmDOFSetupCS::FOutputResDivisor>(2);
+			PassViewSize = PreprocessViewSize;
+			CocRadiusBasis = PreprocessViewSize.X;
+		}
+		else
+		{
+			check(0);
+		}
+		
+		FDiaphragmDOFSetupCS::FParameters* PassParameters = GraphBuilder.AllocParameters<FDiaphragmDOFSetupCS::FParameters>();
+		{
+			PassParameters->CommonParameters = CommonParameters;
+			SetCocModelParameters(&PassParameters->CocModel, CocModel, CocRadiusBasis);
+			PassParameters->ViewportRect = FVector4(0, 0, PassViewSize.X, PassViewSize.Y);
+			PassParameters->CocRadiusBasis = FVector2D(GatheringViewSize.X, PreprocessViewSize.X);
+			PassParameters->SceneColorTexture = InputSceneColor;
+			PassParameters->SceneDepthTexture = SceneBlackboard.SceneDepthBuffer;
+		
+			if (!bOutputFullResolution)
+			{
+				FullResGatherInputTextures.SceneColor = InputSceneColor;
+			}
+			else if (bProcessSceneAlpha)
+			{
+				// No point passing through the full res scene color, the shader just output SeparateCoc.
+				PassParameters->Output0 = CreateUAVs(GraphBuilder, FullResGatherInputTextures).SeparateCoc;
+				FullResGatherInputTextures.SceneColor = InputSceneColor;
+			}
+			else
+			{
+				PassParameters->Output0 = CreateUAVs(GraphBuilder, FullResGatherInputTextures).SceneColor;
+			}
+
+			if (bOutputHalfResolution)
+			{
+				PassParameters->Output1 = CreateUAVs(GraphBuilder, HalfResGatherInputTextures).SceneColor;
+				PassParameters->Output2 = CreateUAVs(GraphBuilder, HalfResGatherInputTextures).SeparateCoc;
+			}
+		}
+
+		FIntVector GroupCount = FComputeShaderUtils::GetGroupCount(PassViewSize, GroupSize);
+
+		TShaderMapRef<FDiaphragmDOFSetupCS> ComputeShader(View.ShaderMap, PermutationVector);
+		FComputeShaderUtils::AddPass(
+			GraphBuilder,
+			RDG_EVENT_NAME("DiaphragmDOF Setup(%s CoC=[%d;%d]) %dx%d",
+				!bOutputFullResolution ? TEXT("HalfRes") : (!bOutputHalfResolution ? TEXT("FullRes") : TEXT("Full&HalfRes")),
+				FMath::FloorToInt(CocModel.ComputeViewMinForegroundCocRadius(PassViewSize.X)),
+				FMath::CeilToInt(CocModel.ComputeViewMaxBackgroundCocRadius(PassViewSize.X)),
+				PassViewSize.X, PassViewSize.Y),
+			*ComputeShader,
+			PassParameters,
+			GroupCount);
+
+		if (!bOutputFullResolution || bProcessSceneAlpha)
+		{
+			FullResGatherInputTextures.SceneColor = InputSceneColor;
+		}
+	}
+	
+	// TAA the setup for the convolution to be temporally stable.
+	if (View.AntiAliasingMethod == AAM_TemporalAA && ViewState)
+	{
+		TAAParameters.SceneColorInput = HalfResGatherInputTextures.SceneColor;
+		TAAParameters.SceneMetadataInput = HalfResGatherInputTextures.SeparateCoc;
+
+		FTAAOutputs TAAOutputs = TAAParameters.AddTemporalAAPass(
+			GraphBuilder, SceneBlackboard, View,
+			View.PrevViewInfo.DOFPreGatherHistory,
+			&ViewState->PrevFrameViewInfo.DOFPreGatherHistory);
+		
+		HalfResGatherInputTextures.SceneColor = TAAOutputs.SceneColor;
+		HalfResGatherInputTextures.SeparateCoc = TAAOutputs.SceneMetadata;
+
+		HalfResGatherInputDescs.SceneColor = TAAOutputs.SceneColor->Desc;
+		HalfResGatherInputDescs.SceneColor.TargetableFlags |= TexCreate_UAV;
+		if (TAAOutputs.SceneMetadata)
+		{
+			HalfResGatherInputDescs.SeparateCoc = TAAOutputs.SceneMetadata->Desc;
+			HalfResGatherInputDescs.SeparateCoc.TargetableFlags |= TexCreate_UAV;
+		}
+	}
+
+	// Tile classify work than needs to be done.
+	FDOFTileClassificationTextures TileClassificationTextures;
+	{
+		// Setup the descriptors for tile classification.
+		FDOFTileClassificationDescs TileClassificationDescs;
+		{
+			FIntPoint MaxTileCount = CocTileGridSize(HalfResGatherInputTextures.SceneColor->Desc.Extent);
+			
+			TileClassificationDescs.Foreground = FPooledRenderTargetDesc::Create2DDesc(
+				MaxTileCount, PF_G16R16F, FClearValueBinding::None, TexCreate_None, TexCreate_RenderTargetable | TexCreate_UAV, false);
+
+			TileClassificationDescs.Background = FPooledRenderTargetDesc::Create2DDesc(
+				MaxTileCount, PF_FloatRGBA, FClearValueBinding::None, TexCreate_None, TexCreate_RenderTargetable | TexCreate_UAV, false);
+		}
+
+		// Adds a coc flatten pass.
+		FDOFTileClassificationTextures FlattenedTileClassificationTextures;
+		{
+			FIntPoint SrcSize = HalfResGatherInputTextures.SceneColor->Desc.Extent;
+
+			static const TCHAR* OutputDebugNames[2] = {
+				TEXT("DOFFlattenFgdCoc"),
+				TEXT("DOFFlattenBgdCoc"),
+			};
+			FlattenedTileClassificationTextures = CreateTextures(GraphBuilder, TileClassificationDescs, OutputDebugNames);
+
+			FDiaphragmDOFCocFlattenCS::FParameters* PassParameters = GraphBuilder.AllocParameters<FDiaphragmDOFCocFlattenCS::FParameters>();
+			PassParameters->CommonParameters = CommonParameters;
+			PassParameters->ViewportRect = FIntRect(0, 0, GatheringViewSize.X, GatheringViewSize.Y);
+			PassParameters->ThreadIdToBufferUV = FVector2D(
+				PreprocessViewSize.X / float(GatheringViewSize.X * SrcSize.X),
+				PreprocessViewSize.Y / float(GatheringViewSize.Y * SrcSize.Y));
+			PassParameters->MaxBufferUV = FVector2D(
+				(PreprocessViewSize.X - 1.0f) / float(SrcSize.X),
+				(PreprocessViewSize.Y - 1.0f) / float(SrcSize.Y));
+			PassParameters->GatherInput = HalfResGatherInputTextures;
+			PassParameters->TileOutput = CreateUAVs(GraphBuilder, FlattenedTileClassificationTextures);
+			
+			FDiaphragmDOFCocFlattenCS::FPermutationDomain PermutationVector;
+			PermutationVector.Set<FDiaphragmDOFCocFlattenCS::FDoCocGather4>(PreprocessViewSize != GatheringViewSize);
+
+			TShaderMapRef<FDiaphragmDOFCocFlattenCS> ComputeShader(View.ShaderMap, PermutationVector);
+			FComputeShaderUtils::AddPass(
+				GraphBuilder,
+				RDG_EVENT_NAME("DiaphragmDOF FlattenCoc(Gather4=%s) %dx%d",
+					PermutationVector.Get<FDiaphragmDOFCocFlattenCS::FDoCocGather4>() ? TEXT("Yes") : TEXT("No"),
+					GatheringViewSize.X, GatheringViewSize.Y),
+				*ComputeShader,
+				PassParameters,
+				FComputeShaderUtils::GetGroupCount(GatheringViewSize, kDefaultGroupSize));
+		}
+		
+		// Error introduced by the random offset of the gathering kernel's center.
+		const float BluringRadiusErrorMultiplier = 1.0f + 1.0f / (HalfResRingCount + 0.5f);
+
+		// Number of group to dispatch for dilate passes.
+		const FIntPoint DilatePassViewSize = FIntPoint::DivideAndRoundUp(GatheringViewSize, kCocTileSize);
+		const FIntVector DilateGroupCount = FComputeShaderUtils::GetGroupCount(DilatePassViewSize, kDefaultGroupSize);
+
+		// Add one coc dilate pass.
+		auto AddCocDilatePass = [&](
+			FRCPassDiaphragmDOFDilateCoc::EMode Mode,
+			const FDOFTileClassificationTextures& TileInput,
+			const FDOFTileClassificationTextures& DilatedTileMinMax,
+			int32 SampleRadiusCount, int32 SampleOffsetMultipler)
+		{
+			FDOFTileClassificationDescs OutputDescs = TileClassificationDescs;
+			const TCHAR* OutputDebugNames[2] = {
+				TEXT("DOFDilateFgdCoc"),
+				TEXT("DOFDilateBgdCoc"),
+			};
+			if (Mode == FRCPassDiaphragmDOFDilateCoc::EMode::MinForegroundAndMaxBackground)
+			{
+				OutputDebugNames[0] = TEXT("DOFDilateMinFgdCoc");
+				OutputDebugNames[1] = TEXT("DOFDilateMaxBgdCoc");
+				OutputDescs.Foreground.Format = PF_R16F;
+				OutputDescs.Background.Format = PF_R16F;
+			}
+
+			FDOFTileClassificationTextures TileClassificationOutputTextures = CreateTextures(GraphBuilder, OutputDescs, OutputDebugNames);
+			
+			FDiaphragmDOFCocDilateCS::FParameters* PassParameters = GraphBuilder.AllocParameters<FDiaphragmDOFCocDilateCS::FParameters>();
+			PassParameters->CommonParameters = CommonParameters;
+			PassParameters->ViewportRect = FIntRect(0, 0, DilatePassViewSize.X, DilatePassViewSize.Y);
+			PassParameters->SampleOffsetMultipler = SampleOffsetMultipler;
+			PassParameters->fSampleOffsetMultipler = SampleOffsetMultipler;
+			PassParameters->CocRadiusToBucketDistanceUpperBound = PreProcessingToProcessingCocRadiusFactor * BluringRadiusErrorMultiplier;
+			PassParameters->BucketDistanceToCocRadius = 1.0f / PassParameters->CocRadiusToBucketDistanceUpperBound;
+			PassParameters->TileInput = TileInput;
+			PassParameters->DilatedTileMinMax = DilatedTileMinMax;
+			PassParameters->TileOutput = CreateUAVs(GraphBuilder, TileClassificationOutputTextures);
+			
+			FDiaphragmDOFCocDilateCS::FPermutationDomain PermutationVector;
+			PermutationVector.Set<FDDOFDilateRadiusDim>(SampleRadiusCount);
+			PermutationVector.Set<FDDOFDilateModeDim>(Mode);
+			// TODO: permutation to do foreground and background separately, to have higher occupancy?
+
+			TShaderMapRef<FDiaphragmDOFCocDilateCS> ComputeShader(View.ShaderMap, PermutationVector);
+			FComputeShaderUtils::AddPass(
+				GraphBuilder,
+				RDG_EVENT_NAME("DiaphragmDOF DilateCoc(1/16 %s radius=%d step=%d) %dx%d",
+					GetEventName(Mode), SampleRadiusCount, SampleOffsetMultipler,
+					DilatePassViewSize.X, DilatePassViewSize.Y),
+				*ComputeShader,
+				PassParameters,
+				DilateGroupCount);
+
+			return TileClassificationOutputTextures;
+		};
+		
+		// Parameters for the dilate Coc passes.
+		int32 DilateCount = 1;
+		int32 SampleRadiusCount[3];
+		int32 SampleDistanceMultiplier[3];
+		{
+			const int32 MaxSampleRadiusCount = FRCPassDiaphragmDOFDilateCoc::MaxSampleRadiusCount;
+
+			// Compute the maximum tile dilation.
+			int32 MaximumTileDilation = FMath::CeilToInt(
+				(MaxBluringRadius * BluringRadiusErrorMultiplier) / FRCPassDiaphragmDOFFlattenCoc::CocTileResolutionDivisor);
+
+			// There is always at least one dilate pass so that even small Coc radius conservatively dilate on next neighboor.
+			int32 CurrentConvolutionRadius = FMath::Min(
+				MaximumTileDilation, MaxSampleRadiusCount);
+			
+			SampleDistanceMultiplier[0] = 1;
+			SampleRadiusCount[0] = CurrentConvolutionRadius;
+
+			// If the theoric radius is too big, setup more dilate passes.
+			for (int32 i = 1; i < ARRAY_COUNT(SampleDistanceMultiplier); i++)
+			{
+				if (MaximumTileDilation <= CurrentConvolutionRadius)
+				{
+					break;
+				}
+
+				// Highest upper bound possible for SampleDistanceMultiplier to not step over any tile.
+				int32 HighestPossibleMultiplierUpperBound = CurrentConvolutionRadius + 1;
+
+				// Find out how many step we need to do on dilate radius.
+				SampleRadiusCount[i] = FMath::Min(
+					MaximumTileDilation / HighestPossibleMultiplierUpperBound,
+					MaxSampleRadiusCount);
+
+				// Find out ideal multiplier to not dilate an area too large.
+				// TODO: Could add control over the radius of the last.
+				int32 IdealMultiplier = FMath::DivideAndRoundUp(MaximumTileDilation - CurrentConvolutionRadius, SampleRadiusCount[1]); // TODO: why 1?
+
+				SampleDistanceMultiplier[i] = FMath::Min(IdealMultiplier, HighestPossibleMultiplierUpperBound);
+
+				CurrentConvolutionRadius += SampleRadiusCount[i] * SampleDistanceMultiplier[i];
+
+				DilateCount++;
+			}
+		}
+		
+		if (DilateCount > 1)
+		{
+			// TODO.
+			FDOFTileClassificationTextures MinMaxTexture = FlattenedTileClassificationTextures;
+
+			// Dilate min foreground and max background coc radii first.
+			for (int32 i = 0; i < DilateCount; i++)
+			{
+				MinMaxTexture = AddCocDilatePass(
+					FRCPassDiaphragmDOFDilateCoc::EMode::MinForegroundAndMaxBackground,
+					MinMaxTexture, FDOFTileClassificationTextures(),
+					SampleRadiusCount[i], SampleDistanceMultiplier[i]);
+			}
+
+			TileClassificationTextures = FlattenedTileClassificationTextures;
+
+			// Dilates everything else.
+			for (int32 i = 0; i < DilateCount; i++)
+			{
+				TileClassificationTextures = AddCocDilatePass(
+					FRCPassDiaphragmDOFDilateCoc::EMode::MinimalAbsoluteRadiuses,
+					TileClassificationTextures, MinMaxTexture,
+					SampleRadiusCount[i], SampleDistanceMultiplier[i]);
+			}
+		}
+		else
+		{
+			TileClassificationTextures = AddCocDilatePass(
+				FRCPassDiaphragmDOFDilateCoc::EMode::StandAlone,
+				FlattenedTileClassificationTextures, FDOFTileClassificationTextures(),
+				SampleRadiusCount[0], SampleDistanceMultiplier[0]);
+		}
+	}
+
+	// Add the reduce pass
+	FDOFGatherInputTextures ReducedGatherInputTextures;
+	FRDGBufferRef DrawIndirectParametersBuffer = nullptr;
+	FRDGBufferRef ForegroundScatterDrawListBuffer = nullptr;
+	FRDGBufferRef BackgroundScatterDrawListBuffer = nullptr;
+	{
+		FIntPoint SrcSize = HalfResGatherInputDescs.SceneColor.Extent;
+		
+		// Compute the number of mip level required by the gathering pass.
+		const int32 MipLevelCount = FMath::Clamp(FMath::CeilToInt(FMath::Log2(MaxBluringRadius * 0.5 / HalfResRingCount)) + (bUseLowAccumulatorQuality ? 1 : 0), 2, FRCPassDiaphragmDOFReduce::kMaxMipLevelCount);
+
+		// Maximum number of scattering group per draw instance.
+		// TODO: depends.
+		const uint32 kMaxScatteringGroupPerInstance = 21;
+
+		// Maximum number of scattering group allowed per frame.
+		const uint32 MaxScatteringGroupCount = FMath::Max(MaxScatteringRatio * 0.25f * SrcSize.X * SrcSize.Y - kMaxScatteringGroupPerInstance, float(kMaxScatteringGroupPerInstance));
+
+		// Allocate the reduced gather input textures.
+		{
+			FDOFGatherInputDescs ReducedGatherInputDescs = HalfResGatherInputDescs;
+			ReducedGatherInputDescs.SceneColor.NumMips = MipLevelCount;
+			ReducedGatherInputDescs.SceneColor.Flags = (ReducedGatherInputDescs.SceneColor.Flags & ~(TexCreate_FastVRAM)) | GFastVRamConfig.DOFReduce;
+			
+			// Make sure the mip 0 is a multiple of 2^NumMips so there is no per mip level UV conversion to do in the gathering shader.
+			// Also make sure it is a multiple of group size because reduce shader unconditionally output Mip0.
+			int32 Multiple = FMath::Max(1 << (MipLevelCount - 1), kDefaultGroupSize);
+			ReducedGatherInputDescs.SceneColor.Extent.X = Multiple * FMath::DivideAndRoundUp(ReducedGatherInputDescs.SceneColor.Extent.X, Multiple);
+			ReducedGatherInputDescs.SceneColor.Extent.Y = Multiple * FMath::DivideAndRoundUp(ReducedGatherInputDescs.SceneColor.Extent.Y, Multiple);
+			
+			ReducedGatherInputDescs.SeparateCoc = ReducedGatherInputDescs.SceneColor;
+			ReducedGatherInputDescs.SeparateCoc.Format = HalfResGatherInputDescs.SeparateCoc.Format;
+			
+			if (bRGBBufferSeparateCocBuffer)
+			{
+				ReducedGatherInputDescs.SceneColor.Format = PF_FloatR11G11B10;
+				ReducedGatherInputDescs.SeparateCoc.Format = PF_R16F;
+			}
+
+			ReducedGatherInputTextures = CreateTextures(GraphBuilder, ReducedGatherInputDescs, TEXT("DOFReduce"));
+		}
+
+		// Downsample the gather color setup to have faster neighborhood comparisons.
+		FDOFGatherInputTextures QuarterResGatherInputTextures;
+		if (bForegroundHybridScattering || bBackgroundHybridScattering)
+		{
+			// Allocate quarter res textures.
+			{
+				FDOFGatherInputDescs QuarterResGatherInputDescs = HalfResGatherInputDescs;
+				QuarterResGatherInputDescs.SceneColor.Extent /= 2;
+				QuarterResGatherInputDescs.SeparateCoc.Extent /= 2;
+
+				// Lower the bit depth to speed up texture fetches in reduce pass, that is ok since this is used only for comparison purposes.
+				if (bRGBBufferSeparateCocBuffer && !bProcessSceneAlpha)
+					QuarterResGatherInputDescs.SceneColor.Format = PF_FloatR11G11B10;
+			
+				QuarterResGatherInputTextures = CreateTextures(GraphBuilder, QuarterResGatherInputDescs, TEXT("DOFDownsample"));
+			}
+
+			FIntPoint PassViewSize = FIntPoint::DivideAndRoundUp(PreprocessViewSize, 2);
+
+			FDiaphragmDOFDownsampleCS::FParameters* PassParameters = GraphBuilder.AllocParameters<FDiaphragmDOFDownsampleCS::FParameters>();
+			PassParameters->CommonParameters = CommonParameters;
+			PassParameters->ViewportRect = FIntRect(0, 0, PassViewSize.X, PassViewSize.Y);
+			PassParameters->MaxBufferUV = FVector2D(
+				(PreprocessViewSize.X - 0.5f) / float(SrcSize.X),
+				(PreprocessViewSize.Y - 0.5f) / float(SrcSize.Y));
+			PassParameters->OutputCocRadiusMultiplier = PreProcessingToProcessingCocRadiusFactor;
+
+			PassParameters->GatherInputSize = FVector4(SrcSize.X, SrcSize.Y, 1.0f / SrcSize.X, 1.0f / SrcSize.Y);
+			PassParameters->GatherInput = HalfResGatherInputTextures;
+
+			PassParameters->OutDownsampledGatherInput = CreateUAVs(GraphBuilder, QuarterResGatherInputTextures);
+			
+			TShaderMapRef<FDiaphragmDOFDownsampleCS> ComputeShader(View.ShaderMap);
+			FComputeShaderUtils::AddPass(
+				GraphBuilder,
+				RDG_EVENT_NAME("DiaphragmDOF Downsample %dx%d", PassViewSize.X, PassViewSize.Y),
+				*ComputeShader,
+				PassParameters,
+				FComputeShaderUtils::GetGroupCount(PassViewSize, kDefaultGroupSize));
+		}
+		
+		// Create and clears buffers for indirect scatter.
+		if (bForegroundHybridScattering || bBackgroundHybridScattering)
+		{
+			DrawIndirectParametersBuffer = GraphBuilder.CreateBuffer(FRDGBufferDesc::CreateIndirectDesc<FRHIDrawIndexedIndirectParameters>(2), TEXT("DOFIndirectDrawParameters"));
+
+			FRDGBufferDesc DrawListDescs = FRDGBufferDesc::CreateStructuredDesc(sizeof(float) * 4, 5 * MaxScatteringGroupCount);
+			if (bForegroundHybridScattering)
+				ForegroundScatterDrawListBuffer = GraphBuilder.CreateBuffer(DrawListDescs, TEXT("DOFForegroundDrawList"));
+			if (bForegroundHybridScattering)
+				BackgroundScatterDrawListBuffer = GraphBuilder.CreateBuffer(DrawListDescs, TEXT("DOFBackgroundDrawList"));
+		}
+		
+		// Adds the reduce pass.
+		{
+			FIntPoint PassViewSize = PreprocessViewSize;
+
+			FDiaphragmDOFReduceCS::FPermutationDomain PermutationVector;
+			PermutationVector.Set<FDiaphragmDOFReduceCS::FReduceMipCount>(MipLevelCount);
+			PermutationVector.Set<FDiaphragmDOFReduceCS::FHybridScatterForeground>(bForegroundHybridScattering);
+			PermutationVector.Set<FDiaphragmDOFReduceCS::FHybridScatterBackground>(bBackgroundHybridScattering);
+			PermutationVector.Set<FDDOFRGBColorBufferDim>(bRGBBufferSeparateCocBuffer);
+
+			FDiaphragmDOFReduceCS::FParameters* PassParameters = GraphBuilder.AllocParameters<FDiaphragmDOFReduceCS::FParameters>();
+			PassParameters->ViewportRect = FIntRect(0, 0, PassViewSize.X, PassViewSize.Y);
+			PassParameters->MaxInputBufferUV = FVector2D(
+				(PreprocessViewSize.X - 0.5f) / SrcSize.X,
+				(PreprocessViewSize.Y - 0.5f) / SrcSize.Y);
+			PassParameters->MaxScatteringGroupCount = MaxScatteringGroupCount;
+			PassParameters->PreProcessingToProcessingCocRadiusFactor = PreProcessingToProcessingCocRadiusFactor;
+			PassParameters->MinScatteringCocRadius = MinScatteringCocRadius;
+			PassParameters->NeighborCompareMaxColor = CVarScatterNeighborCompareMaxColor.GetValueOnRenderThread();
+			
+			PassParameters->EyeAdaptation = GetEyeAdaptationTexture(GraphBuilder, View);
+			PassParameters->CommonParameters = CommonParameters;
+
+			PassParameters->GatherInputSize = FVector4(SrcSize.X, SrcSize.Y, 1.0f / SrcSize.X, 1.0f / SrcSize.Y);
+			PassParameters->GatherInput = HalfResGatherInputTextures;
+			
+			PassParameters->QuarterResGatherInputSize = FVector4(SrcSize.X / 2, SrcSize.Y / 2, 2.0f / SrcSize.X, 2.0f / SrcSize.Y);
+			PassParameters->QuarterResGatherInput = QuarterResGatherInputTextures;
+
+			for (int32 MipLevel = 0; MipLevel < MipLevelCount; MipLevel++)
+			{
+				PassParameters->OutputMips[MipLevel] = CreateUAVs(GraphBuilder, ReducedGatherInputTextures, MipLevel);
+			}
+
+			if (bForegroundHybridScattering || bBackgroundHybridScattering)
+			{
+				PassParameters->OutScatterDrawIndirectParameters = GraphBuilder.CreateUAV(DrawIndirectParametersBuffer);
+				if (ForegroundScatterDrawListBuffer)
+					PassParameters->OutForegroundScatterDrawList = GraphBuilder.CreateUAV(ForegroundScatterDrawListBuffer);
+				if (BackgroundScatterDrawListBuffer)
+					PassParameters->OutBackgroundScatterDrawList = GraphBuilder.CreateUAV(BackgroundScatterDrawListBuffer);
+
+				AddPass_ClearUAV(GraphBuilder, RDG_EVENT_NAME("ClearIndirectDraw"), PassParameters->OutScatterDrawIndirectParameters, 0);
+			}
+
+			TShaderMapRef<FDiaphragmDOFReduceCS> ComputeShader(View.ShaderMap, PermutationVector);
+			FComputeShaderUtils::AddPass(
+				GraphBuilder,
+				RDG_EVENT_NAME("DiaphragmDOF Reduce(Mips=%d FgdScatter=%s BgdScatter=%s%s) %dx%d",
+					MipLevelCount,
+					bForegroundHybridScattering ? TEXT("Yes") : TEXT("No"),
+					bBackgroundHybridScattering ? TEXT("Yes") : TEXT("No"),
+					bRGBBufferSeparateCocBuffer ? TEXT(" R11G11B10") : TEXT(""),
+					PassViewSize.X, PassViewSize.Y),
+				*ComputeShader,
+				PassParameters,
+				FComputeShaderUtils::GetGroupCount(PassViewSize, kDefaultGroupSize));
+		}
+		
+		// Pack multiple scattering group on same primitive instance to increase wave occupancy in the scattering vertex shader.
+		if (bForegroundHybridScattering || bBackgroundHybridScattering)
+		{
+			FDiaphragmDOFScatterGroupPackCS::FParameters* PassParameters = GraphBuilder.AllocParameters<FDiaphragmDOFScatterGroupPackCS::FParameters>();
+			PassParameters->MaxScatteringGroupCount = MaxScatteringGroupCount;
+			PassParameters->OutScatterDrawIndirectParameters = GraphBuilder.CreateUAV(DrawIndirectParametersBuffer);
+			if (ForegroundScatterDrawListBuffer)
+				PassParameters->OutForegroundScatterDrawList = GraphBuilder.CreateUAV(ForegroundScatterDrawListBuffer);
+			if (BackgroundScatterDrawListBuffer)
+				PassParameters->OutBackgroundScatterDrawList = GraphBuilder.CreateUAV(BackgroundScatterDrawListBuffer);
+
+			TShaderMapRef<FDiaphragmDOFScatterGroupPackCS> ComputeShader(View.ShaderMap);
+			FComputeShaderUtils::AddPass(
+				GraphBuilder,
+				RDG_EVENT_NAME("DiaphragmDOF ScatterGroupPack"),
+				*ComputeShader,
+				PassParameters,
+				FIntVector(2, 1, 1));
+		}
+	}
+	
+	// Add a pass to build a bokeh LUTs.
+	auto AddBuildBokehLUTPass = [&](FRCPassDiaphragmDOFBuildBokehLUT::EFormat LUTFormat)
+	{
+		if (BokehModel.BokehShape == EBokehShape::Circle)
+		{
+			return FRDGTextureRef();
+		}
+
+		static const TCHAR* const DebugNames[] = {
+			TEXT("DOFScatterBokehLUT"),
+			TEXT("DOFRecombineBokehLUT"),
+			TEXT("DOFGatherBokehLUT"),
+		};
+
+		FRDGTextureDesc BokehLUTDesc;
+		BokehLUTDesc.NumMips = 1;
+		BokehLUTDesc.Format = LUTFormat == FRCPassDiaphragmDOFBuildBokehLUT::EFormat::GatherSamplePos ? PF_G16R16F : PF_R16F;
+		BokehLUTDesc.Extent = FIntPoint(32, 32);
+		BokehLUTDesc.TargetableFlags |= TexCreate_UAV;
+
+		FRDGTextureRef BokehLUT = GraphBuilder.CreateTexture(BokehLUTDesc, DebugNames[int32(LUTFormat)]);
+
+		FDiaphragmDOFBuildBokehLUTCS::FPermutationDomain PermutationVector;
+		PermutationVector.Set<FDiaphragmDOFBuildBokehLUTCS::FBokehSimulationDim>(BokehModel.BokehShape == DiaphragmDOF::EBokehShape::RoundedBlades);
+		PermutationVector.Set<FDiaphragmDOFBuildBokehLUTCS::FLUTFormatDim>(LUTFormat);
+	
+		FDiaphragmDOFBuildBokehLUTCS::FParameters* PassParameters = GraphBuilder.AllocParameters<FDiaphragmDOFBuildBokehLUTCS::FParameters>();
+		PassParameters->BladeCount = BokehModel.DiaphragmBladeCount;
+		PassParameters->DiaphragmRotation = BokehModel.DiaphragmRotation;
+		PassParameters->CocRadiusToCircumscribedRadius = BokehModel.CocRadiusToCircumscribedRadius;
+		PassParameters->CocRadiusToIncircleRadius = BokehModel.CocRadiusToIncircleRadius;
+		PassParameters->DiaphragmBladeRadius = BokehModel.RoundedBlades.DiaphragmBladeRadius;
+		PassParameters->DiaphragmBladeCenterOffset = BokehModel.RoundedBlades.DiaphragmBladeCenterOffset;
+		PassParameters->BokehLUTOutput = GraphBuilder.CreateUAV(BokehLUT);
+			
+		TShaderMapRef<FDiaphragmDOFBuildBokehLUTCS> ComputeShader(View.ShaderMap, PermutationVector);
+		FComputeShaderUtils::AddPass(
+			GraphBuilder,
+			RDG_EVENT_NAME("DiaphragmDOF BuildBokehLUT(Blades=%d Shape=%s, LUT=%s) %dx%d",
+				BokehModel.DiaphragmBladeCount,
+				BokehModel.BokehShape == DiaphragmDOF::EBokehShape::RoundedBlades ? TEXT("Rounded") : TEXT("Straight"),
+				GetEventName(LUTFormat),
+				BokehLUTDesc.Extent.X, BokehLUTDesc.Extent.Y),
+			*ComputeShader,
+			PassParameters,
+			FComputeShaderUtils::GetGroupCount(BokehLUTDesc.Extent, kDefaultGroupSize));
+
+		return BokehLUT;
+	};
+	
+	// Add all passes for convolutions.
+	FDOFConvolutionTextures ForegroundConvolutionTextures;
+	FDOFConvolutionTextures ForegroundHoleFillingConvolutionTextures;
+	FDOFConvolutionTextures BackgroundConvolutionTextures;
+	FDOFConvolutionTextures SlightOutOfFocusConvolutionTextures;
+	{
+		// High level configuration of a convolution
+		struct FConvolutionSettings
+		{
+			/** Which layer to gather. */
+			EDiaphragmDOFLayerProcessing LayerProcessing;
+
+			/** Configuration of the pass. */
+			FRCPassDiaphragmDOFGather::EQualityConfig QualityConfig;
+		
+			/** Postfilter method to apply on this gather pass. */
+			EDiaphragmDOFPostfilterMethod PostfilterMethod;
+
+			/** Bokeh simulation to do. */
+			EDiaphragmDOFBokehSimulation BokehSimulation;
+		
+			FConvolutionSettings()
+				: LayerProcessing(EDiaphragmDOFLayerProcessing::ForegroundAndBackground)
+				, QualityConfig(FRCPassDiaphragmDOFGather::EQualityConfig::HighQuality)
+				, PostfilterMethod(EDiaphragmDOFPostfilterMethod::None)
+				, BokehSimulation(EDiaphragmDOFBokehSimulation::Disabled)
+			{ }
+		};
+
+		// Add a gather pass
+		auto AddGatherPass = [&](
+			const FConvolutionSettings& ConvolutionSettings,
+			FRDGTextureRef BokehLUT,
+			FDOFConvolutionTextures* ConvolutionOutputTextures,
+			FRDGTextureRef* ScatterOcclusionTexture)
+		{
+			// Allocate output textures
+			{
+				FRDGTextureDesc Desc = ReducedGatherInputTextures.SceneColor->Desc;
+				Desc.Extent = RefBufferSize;
+				Desc.Format = PF_FloatRGBA;
+				Desc.TargetableFlags |= TexCreate_RenderTargetable | TexCreate_UAV;
+				Desc.NumMips = 1;
+
+				{
+					const TCHAR* DebugName = nullptr;
+					if (ConvolutionSettings.LayerProcessing == EDiaphragmDOFLayerProcessing::ForegroundOnly)
+						DebugName = TEXT("DOFGatherForeground");
+					else if (ConvolutionSettings.LayerProcessing == EDiaphragmDOFLayerProcessing::ForegroundHoleFilling)
+						DebugName = TEXT("DOFGatherForegroundFill");
+					else if (ConvolutionSettings.LayerProcessing == EDiaphragmDOFLayerProcessing::BackgroundOnly)
+						DebugName = TEXT("DOFGatherBackground");
+					else if (ConvolutionSettings.LayerProcessing == EDiaphragmDOFLayerProcessing::SlightOutOfFocus)
+						DebugName = TEXT("DOFGatherFocus");
+					else
+						check(0);
+
+					ConvolutionOutputTextures->SceneColor = GraphBuilder.CreateTexture(Desc, DebugName);
+				
+					if (bProcessSceneAlpha)
+					{
+						Desc.Format = PF_R16F;
+						ConvolutionOutputTextures->SeparateAlpha = GraphBuilder.CreateTexture(Desc, DebugName);
+					}
+				}
+
+				if (ConvolutionSettings.QualityConfig == FRCPassDiaphragmDOFGather::EQualityConfig::HighQualityWithHybridScatterOcclusion)
+				{
+					Desc.Format = PF_G16R16F;
+			
+					const TCHAR* DebugName = nullptr;
+					if (ConvolutionSettings.LayerProcessing == EDiaphragmDOFLayerProcessing::BackgroundOnly)
+						DebugName = TEXT("DOFScatterOcclusionBackground");
+					else
+						check(0);
+			
+					*ScatterOcclusionTexture = GraphBuilder.CreateTexture(Desc, DebugName);
+				}
+			}
+
+			FIntPoint ReduceOutputRectMip0(
+				kDefaultGroupSize * FMath::DivideAndRoundUp(PreprocessViewSize.X, kDefaultGroupSize),
+				kDefaultGroupSize * FMath::DivideAndRoundUp(PreprocessViewSize.Y, kDefaultGroupSize));
+		
+			FIntPoint SrcSize = ReducedGatherInputTextures.SceneColor->Desc.Extent;
+
+			FDiaphragmDOFGatherCS::FPermutationDomain PermutationVector;
+			PermutationVector.Set<FDDOFLayerProcessingDim>(ConvolutionSettings.LayerProcessing);
+			PermutationVector.Set<FDDOFGatherRingCountDim>(HalfResRingCount);
+			PermutationVector.Set<FDDOFGatherQualityDim>(ConvolutionSettings.QualityConfig);
+			PermutationVector.Set<FDDOFBokehSimulationDim>(ConvolutionSettings.BokehSimulation);
+			PermutationVector.Set<FDDOFClampInputUVDim>(ReduceOutputRectMip0 != SrcSize);
+			PermutationVector.Set<FDDOFRGBColorBufferDim>(bRGBBufferSeparateCocBuffer);
+			PermutationVector = FDiaphragmDOFGatherCS::RemapPermutation(PermutationVector);
+
+			// Affine transformtation to control whether a CocRadius is considered or not.
+			FVector2D ConsiderCocRadiusAffineTransformation0 = kContantlyPassingAffineTransformation;
+			FVector2D ConsiderCocRadiusAffineTransformation1 = kContantlyPassingAffineTransformation;
+			FVector2D ConsiderAbsCocRadiusAffineTransformation = kContantlyPassingAffineTransformation;
+			{
+				// Gathering scalability.
+				const float GatheringScalingDownFactor = float(PreprocessViewSize.X) / float(GatheringViewSize.X);
+
+				// Coc radius considered.
+				const float RecombineCocRadiusBorder = GatheringScalingDownFactor * (FRCPassDiaphragmDOFGather::kMaxSlightOutOfFocusRingCount - 1.0f);
+
+				if (ConvolutionSettings.LayerProcessing == EDiaphragmDOFLayerProcessing::ForegroundOnly)
+				{
+					ConsiderCocRadiusAffineTransformation0 = GenerateSaturatedAffineTransformation(
+						-(RecombineCocRadiusBorder - 1.0f), -RecombineCocRadiusBorder);
+
+					ConsiderAbsCocRadiusAffineTransformation = GenerateSaturatedAffineTransformation(
+						RecombineCocRadiusBorder - 1.0f, RecombineCocRadiusBorder);
+				}
+				else if (ConvolutionSettings.LayerProcessing == EDiaphragmDOFLayerProcessing::ForegroundHoleFilling)
+				{
+					ConsiderCocRadiusAffineTransformation0 = GenerateSaturatedAffineTransformation(
+						RecombineCocRadiusBorder, RecombineCocRadiusBorder + 1.0f);
+				}
+				else if (ConvolutionSettings.LayerProcessing == EDiaphragmDOFLayerProcessing::BackgroundOnly)
+				{
+					ConsiderCocRadiusAffineTransformation0 = GenerateSaturatedAffineTransformation(
+						RecombineCocRadiusBorder - 1.0f, RecombineCocRadiusBorder);
+
+					ConsiderAbsCocRadiusAffineTransformation = GenerateSaturatedAffineTransformation(
+						RecombineCocRadiusBorder - 1.0f, RecombineCocRadiusBorder);
+				}
+				else if (ConvolutionSettings.LayerProcessing == EDiaphragmDOFLayerProcessing::SlightOutOfFocus)
+				{
+					ConsiderAbsCocRadiusAffineTransformation = GenerateSaturatedAffineTransformation(
+						RecombineCocRadiusBorder + GatheringScalingDownFactor * 1.0f, RecombineCocRadiusBorder);
+				}
+				else
+				{
+					checkf(0, TEXT("What layer processing is that?"));
+				}
+			}
+
+			FDiaphragmDOFGatherCS::FParameters* PassParameters = GraphBuilder.AllocParameters<FDiaphragmDOFGatherCS::FParameters>();
+			PassParameters->ViewportSize = FVector4(GatheringViewSize.X, GatheringViewSize.Y, 1.0f / GatheringViewSize.X, 1.0f / GatheringViewSize.Y);
+			PassParameters->ViewportRect = FIntRect(0, 0, GatheringViewSize.X, GatheringViewSize.Y);
+			PassParameters->TemporalJitterPixels = View.TemporalJitterPixels;
+			PassParameters->DispatchThreadIdToInputBufferUV = FVector2D(
+				float(PreprocessViewSize.X) / float(GatheringViewSize.X * SrcSize.X),
+				float(PreprocessViewSize.Y) / float(GatheringViewSize.Y * SrcSize.Y));;
+			PassParameters->ConsiderCocRadiusAffineTransformation0 = ConsiderCocRadiusAffineTransformation0;
+			PassParameters->ConsiderCocRadiusAffineTransformation1 = ConsiderCocRadiusAffineTransformation1;
+			PassParameters->ConsiderAbsCocRadiusAffineTransformation = ConsiderAbsCocRadiusAffineTransformation;
+			PassParameters->InputBufferUVToOutputPixel = FVector2D(
+				float(SrcSize.X * GatheringViewSize.X) / float(PreprocessViewSize.X),
+				float(SrcSize.Y * GatheringViewSize.Y) / float(PreprocessViewSize.Y));
+			PassParameters->MipBias = FMath::Log2(float(PreprocessViewSize.X) / float(GatheringViewSize.X));
+			PassParameters->MaxRecombineAbsCocRadius = 3.0 * float(PreprocessViewSize.X) / float(GatheringViewSize.X);
+
+			PassParameters->CommonParameters = CommonParameters;
+		
+			PassParameters->GatherInputSize = FVector4(SrcSize.X, SrcSize.Y, 1.0f / SrcSize.X, 1.0f / SrcSize.Y);
+			PassParameters->GatherInput = ReducedGatherInputTextures;
+		
+			PassParameters->TileClassification = TileClassificationTextures;
+			PassParameters->BokehLUT = BokehLUT;
+
+			PassParameters->ConvolutionOutput = CreateUAVs(GraphBuilder, *ConvolutionOutputTextures);
+			if (ConvolutionSettings.QualityConfig == FRCPassDiaphragmDOFGather::EQualityConfig::HighQualityWithHybridScatterOcclusion)
+				PassParameters->ScatterOcclusionOutput = GraphBuilder.CreateUAV(*ScatterOcclusionTexture);
+			
+			TShaderMapRef<FDiaphragmDOFGatherCS> ComputeShader(View.ShaderMap, PermutationVector);
+			FComputeShaderUtils::AddPass(
+				GraphBuilder,
+				RDG_EVENT_NAME("DiaphragmDOF Gather(%s %s Bokeh=%s Rings=%d%s%s) %dx%d",
+					GetEventName(ConvolutionSettings.LayerProcessing),
+					GetEventName(ConvolutionSettings.QualityConfig),
+					GetEventName(ConvolutionSettings.BokehSimulation),
+					PermutationVector.Get<FDDOFGatherRingCountDim>(),
+					PermutationVector.Get<FDDOFClampInputUVDim>() ? TEXT(" ClampUV") : TEXT(""),
+					PermutationVector.Get<FDDOFRGBColorBufferDim>() ? TEXT(" R11G11B10") : TEXT(""),
+					GatheringViewSize.X, GatheringViewSize.Y),
+				*ComputeShader,
+				PassParameters,
+				FComputeShaderUtils::GetGroupCount(GatheringViewSize, kDefaultGroupSize));
+		}; // AddGatherPass()
+
+		auto AddPostFilterPass = [&](
+			const FConvolutionSettings& ConvolutionSettings,
+			FDOFConvolutionTextures* ConvolutionTextures)
+		{
+			FDOFConvolutionTextures NewConvolutionTextures;
+			NewConvolutionTextures.SceneColor = GraphBuilder.CreateTexture(ConvolutionTextures->SceneColor->Desc, ConvolutionTextures->SceneColor->Name);
+			if (ConvolutionTextures->SeparateAlpha)
+				NewConvolutionTextures.SeparateAlpha = GraphBuilder.CreateTexture(ConvolutionTextures->SeparateAlpha->Desc, ConvolutionTextures->SeparateAlpha->Name);
+			
+			FDiaphragmDOFPostfilterCS::FPermutationDomain PermutationVector;
+			PermutationVector.Set<FDDOFLayerProcessingDim>(ConvolutionSettings.LayerProcessing);
+			PermutationVector.Set<FDDOFPostfilterMethodDim>(ConvolutionSettings.PostfilterMethod);
+			PermutationVector.Set<FDiaphragmDOFPostfilterCS::FTileOptimization>(true); // TODO
+			PermutationVector = FDiaphragmDOFPostfilterCS::RemapPermutationVector(PermutationVector);
+			
+			float MaxRecombineAbsCocRadius = 3.0 * float(PreprocessViewSize.X) / float(GatheringViewSize.X);
+
+			FDiaphragmDOFPostfilterCS::FParameters* PassParameters = GraphBuilder.AllocParameters<FDiaphragmDOFPostfilterCS::FParameters>();
+			PassParameters->ViewportRect = FIntRect(0, 0, GatheringViewSize.X, GatheringViewSize.Y);
+			PassParameters->MaxInputBufferUV = FVector2D(
+				(GatheringViewSize.X - 0.5f) / float(RefBufferSize.X),
+				(GatheringViewSize.Y - 0.5f) / float(RefBufferSize.Y));
+			PassParameters->MinGatherRadius = MaxRecombineAbsCocRadius - 1;
+			PassParameters->CommonParameters = CommonParameters;
+
+			PassParameters->ConvolutionInputSize = FVector4(RefBufferSize.X, RefBufferSize.Y, 1.0f / RefBufferSize.X, 1.0f / RefBufferSize.Y);
+			PassParameters->ConvolutionInput = *ConvolutionTextures;
+
+			PassParameters->TileClassification = TileClassificationTextures;
+			PassParameters->ConvolutionOutput = CreateUAVs(GraphBuilder, NewConvolutionTextures);
+
+			TShaderMapRef<FDiaphragmDOFPostfilterCS> ComputeShader(View.ShaderMap, PermutationVector);
+			FComputeShaderUtils::AddPass(
+				GraphBuilder,
+				RDG_EVENT_NAME("DiaphragmDOF Postfilter(%s %s%s) %dx%d",
+					GetEventName(ConvolutionSettings.LayerProcessing), GetEventName(ConvolutionSettings.PostfilterMethod),
+					PermutationVector.Get<FDiaphragmDOFPostfilterCS::FTileOptimization>() ? TEXT(" TileOptimisation") : TEXT(""),
+					GatheringViewSize.X, GatheringViewSize.Y),
+				*ComputeShader,
+				PassParameters,
+				FComputeShaderUtils::GetGroupCount(GatheringViewSize, kDefaultGroupSize));
+			
+			*ConvolutionTextures = NewConvolutionTextures;
+		}; // AddPostFilterPass()
+		
+		const bool bEnableGatherBokehSettings = (
+			bSupportGatheringBokehSimulation &&
+			CVarEnableGatherBokehSettings.GetValueOnRenderThread() == 1);
+		const bool bEnableScatterBokehSettings = CVarEnableScatterBokehSettings.GetValueOnRenderThread() == 1;
+		const bool bEnableSlightOutOfFocusBokeh = bSupportGatheringBokehSimulation && bRecombineDoesSlightOutOfFocus && CVarEnableRecombineBokehSettings.GetValueOnRenderThread();
+		
+		FRDGTextureRef GatheringBokehLUT = nullptr;
+		if (bEnableGatherBokehSettings)
+			 GatheringBokehLUT = AddBuildBokehLUTPass(FRCPassDiaphragmDOFBuildBokehLUT::EFormat::GatherSamplePos);
+
+		FRDGTextureRef ScatteringBokehLUT = nullptr;
+		if (bEnableScatterBokehSettings || bEnableSlightOutOfFocusBokeh)
+			ScatteringBokehLUT = AddBuildBokehLUTPass(FRCPassDiaphragmDOFBuildBokehLUT::EFormat::CocRadiusToBokehEdgeFactor);
+
+		auto AddHybridScatterPass = [&](
+			const FConvolutionSettings& ConvolutionSettings,
+			FDOFConvolutionTextures* ConvolutionTextures,
+			FRDGTextureRef ScatterOcclusionTexture,
+			FRDGBufferRef ScatterDrawList)
+		{
+			bool bIsForeground = ConvolutionSettings.LayerProcessing == EDiaphragmDOFLayerProcessing::ForegroundOnly;
+			int32 DrawIndirectParametersOffset = bIsForeground ? 0 : 1;
+
+			FDiaphragmDOFHybridScatterPS::FPermutationDomain PermutationVector;
+			PermutationVector.Set<FDDOFLayerProcessingDim>(ConvolutionSettings.LayerProcessing);
+			PermutationVector.Set<FDiaphragmDOFHybridScatterPS::FBokehSimulationDim>(ScatteringBokehLUT ? true : false);
+			PermutationVector.Set<FDDOFScatterOcclusionDim>(ScatterOcclusionTexture ? true : false);
+			PermutationVector = FDiaphragmDOFHybridScatterPS::RemapPermutation(PermutationVector);
+
+			TShaderMapRef<FDiaphragmDOFHybridScatterVS> VertexShader(View.ShaderMap);
+			TShaderMapRef<FDiaphragmDOFHybridScatterPS> PixelShader(View.ShaderMap, PermutationVector);
+
+			FDOFHybridScatterParameters* PassParameters = GraphBuilder.AllocParameters<FDOFHybridScatterParameters>();
+			PassParameters->ViewportSize = FVector4(GatheringViewSize.X, GatheringViewSize.Y, 1.0f / GatheringViewSize.X, 1.0f / GatheringViewSize.Y);
+			PassParameters->CocRadiusToCircumscribedRadius = BokehModel.CocRadiusToCircumscribedRadius;
+			PassParameters->ScatteringScaling = float(GatheringViewSize.X) / float(PreprocessViewSize.X);
+			PassParameters->CommonParameters = CommonParameters;
+			if (bEnableScatterBokehSettings)
+				PassParameters->BokehLUT = ScatteringBokehLUT;
+			PassParameters->ScatterOcclusionSize = FVector4(RefBufferSize.X, RefBufferSize.Y, 1.0f / RefBufferSize.X, 1.0f / RefBufferSize.Y);
+			PassParameters->ScatterOcclusion = ScatterOcclusionTexture;
+			PassParameters->IndirectDrawParameter = DrawIndirectParametersBuffer;
+			PassParameters->ScatterDrawList = GraphBuilder.CreateSRV(ScatterDrawList);
+			PassParameters->RenderTargets[0] = FRenderTargetBinding(
+				ConvolutionTextures->SceneColor, ERenderTargetLoadAction::ELoad, ERenderTargetStoreAction::EStore);
+
+			ValidateShaderParameters(*VertexShader, *PassParameters);
+			ValidateShaderParameters(*PixelShader, *PassParameters);
+
+			GraphBuilder.AddPass(
+				RDG_EVENT_NAME("DiaphragmDOF IndirectScatter(%s Bokeh=%s Occlusion=%s) %dx%d",
+					GetEventName(bIsForeground ? EDiaphragmDOFLayerProcessing::ForegroundOnly : EDiaphragmDOFLayerProcessing::BackgroundOnly),
+					PermutationVector.Get<FDiaphragmDOFHybridScatterPS::FBokehSimulationDim>() ? TEXT("Generic") : TEXT("None"),
+					PermutationVector.Get<FDDOFScatterOcclusionDim>() ? TEXT("Yes") : TEXT("No"),
+					GatheringViewSize.X, GatheringViewSize.Y),
+				PassParameters,
+				ERenderGraphPassFlags::None,
+				[PassParameters, VertexShader, PixelShader, GatheringViewSize, DrawIndirectParametersOffset](FRHICommandList& RHICmdList)
+			{
+				RHICmdList.SetViewport(0, 0, 0.0f, GatheringViewSize.X, GatheringViewSize.Y, 1.0f);
+
+				FGraphicsPipelineStateInitializer GraphicsPSOInit;
+				RHICmdList.ApplyCachedRenderTargets(GraphicsPSOInit);
+				GraphicsPSOInit.DepthStencilState = TStaticDepthStencilState<false, CF_Always>::GetRHI();
+				GraphicsPSOInit.BlendState = TStaticBlendState<CW_RGBA, BO_Add, BF_One, BF_One, BO_Add, BF_One, BF_One>::GetRHI();
+				GraphicsPSOInit.RasterizerState = TStaticRasterizerState<>::GetRHI();
+				GraphicsPSOInit.PrimitiveType = GRHISupportsRectTopology ? PT_RectList : PT_TriangleList;
+				GraphicsPSOInit.BoundShaderState.VertexDeclarationRHI = GEmptyVertexDeclaration.VertexDeclarationRHI;
+				GraphicsPSOInit.BoundShaderState.VertexShaderRHI = GETSAFERHISHADER_VERTEX(*VertexShader);
+				GraphicsPSOInit.BoundShaderState.PixelShaderRHI = GETSAFERHISHADER_PIXEL(*PixelShader);
+				SetGraphicsPipelineState(RHICmdList, GraphicsPSOInit);
+
+				SetShaderParameters(RHICmdList, *VertexShader, VertexShader->GetVertexShader(), *PassParameters);
+				SetShaderParameters(RHICmdList, *PixelShader, PixelShader->GetPixelShader(), *PassParameters);
+
+				RHICmdList.SetStreamSource(0, NULL, 0);
+
+				if (GRHISupportsRectTopology)
+				{
+					RHICmdList.DrawPrimitiveIndirect(
+						PassParameters->IndirectDrawParameter->GetIndirectRHICallBuffer(),
+						sizeof(FRHIDrawIndirectParameters) * DrawIndirectParametersOffset);
+				}
+				else
+				{
+					RHICmdList.DrawIndexedPrimitiveIndirect(
+						GDOFGlobalResource.ScatterIndexBuffer.IndexBufferRHI,
+						PassParameters->IndirectDrawParameter->GetIndirectRHICallBuffer(),
+						sizeof(FRHIDrawIndexedIndirectParameters) * DrawIndirectParametersOffset);
+				}
+			});
+		}; // AddHybridScatterPass()
+
+		// Wire foreground gathering passes.
+		if (bGatherForeground)
+		{
+			FConvolutionSettings ConvolutionSettings;
+			ConvolutionSettings.LayerProcessing = EDiaphragmDOFLayerProcessing::ForegroundOnly;
+			ConvolutionSettings.PostfilterMethod = PostfilterMethod;
+
+			if (bEnableGatherBokehSettings)
+				ConvolutionSettings.BokehSimulation = BokehSimulation;
+
+			if (bUseLowAccumulatorQuality)
+			{
+				ConvolutionSettings.QualityConfig = FRCPassDiaphragmDOFGather::EQualityConfig::LowQualityAccumulator;
+			}
+
+			FRDGTextureRef ScatterOcclusionTexture = nullptr;
+			AddGatherPass(ConvolutionSettings, GatheringBokehLUT, &ForegroundConvolutionTextures, &ScatterOcclusionTexture);
+			AddPostFilterPass(ConvolutionSettings, &ForegroundConvolutionTextures);
+
+			if (bForegroundHybridScattering)
+				AddHybridScatterPass(ConvolutionSettings, &ForegroundConvolutionTextures, ScatterOcclusionTexture, ForegroundScatterDrawListBuffer);
+		}
+
+		// Wire hole filling gathering passes.
+		if (bRecombineDoesSeparateForegroundHoleFilling)
+		{
+			FConvolutionSettings ConvolutionSettings;
+			ConvolutionSettings.LayerProcessing = EDiaphragmDOFLayerProcessing::ForegroundHoleFilling;
+			ConvolutionSettings.PostfilterMethod = PostfilterMethod;
+			
+			FRDGTextureRef ScatterOcclusionTexture = nullptr;
+			AddGatherPass(ConvolutionSettings, /* BokehLUT = */ nullptr, &ForegroundHoleFillingConvolutionTextures, &ScatterOcclusionTexture);
+		}
+		
+		// Gather slight out of focus.
+		if (bRecombineDoesSlightOutOfFocus)
+		{
+			FConvolutionSettings ConvolutionSettings;
+			ConvolutionSettings.LayerProcessing = EDiaphragmDOFLayerProcessing::SlightOutOfFocus;
+			if (bEnableSlightOutOfFocusBokeh)
+				ConvolutionSettings.BokehSimulation = BokehSimulation;
+
+			FRDGTextureRef ScatterOcclusionTexture = nullptr;
+			AddGatherPass(
+				ConvolutionSettings,
+				/* BokehLUT = */ bEnableSlightOutOfFocusBokeh ? ScatteringBokehLUT : nullptr,
+				&SlightOutOfFocusConvolutionTextures, &ScatterOcclusionTexture);
+		}
+
+		// Wire background gathering passes.
+		{
+			FConvolutionSettings ConvolutionSettings;
+			ConvolutionSettings.LayerProcessing = EDiaphragmDOFLayerProcessing::BackgroundOnly;
+			ConvolutionSettings.PostfilterMethod = PostfilterMethod;
+
+			if (bEnableGatherBokehSettings)
+				ConvolutionSettings.BokehSimulation = BokehSimulation;
+
+			ConvolutionSettings.QualityConfig = FRCPassDiaphragmDOFGather::EQualityConfig::LowQualityAccumulator;
+			if (bBackgroundHybridScattering && BgdHybridScatteringMode == EHybridScatterMode::Occlusion)
+			{
+				if (bUseCinematicAccumulatorQuality)
+				{
+					ConvolutionSettings.QualityConfig = FRCPassDiaphragmDOFGather::EQualityConfig::Cinematic;
+				}
+				else
+				{
+					ConvolutionSettings.QualityConfig = FRCPassDiaphragmDOFGather::EQualityConfig::HighQualityWithHybridScatterOcclusion;
+				}
+			}
+			
+			FRDGTextureRef ScatterOcclusionTexture = nullptr;
+			AddGatherPass(ConvolutionSettings, GatheringBokehLUT, &BackgroundConvolutionTextures, &ScatterOcclusionTexture);
+			AddPostFilterPass(ConvolutionSettings, &BackgroundConvolutionTextures);
+			
+			if (bBackgroundHybridScattering)
+				AddHybridScatterPass(ConvolutionSettings, &BackgroundConvolutionTextures, ScatterOcclusionTexture, BackgroundScatterDrawListBuffer);
+		}
+	}
+	
+	// Recombine lower res out of focus with full res scene color.
+	FRDGTextureRef NewSceneColor;
+	{
+		bool bEnableSlightOutOfFocusBokeh = SlightOutOfFocusConvolutionTextures.SceneColor != nullptr;
+
+		{
+			FRDGTextureDesc Desc = InputSceneColor->Desc;
+			Desc.TargetableFlags |= TexCreate_UAV;
+			NewSceneColor = GraphBuilder.CreateTexture(Desc, TEXT("DOFRecombine"));
+		}
+		
+		FIntRect PassViewRect = View.ViewRect;
+
+		FDiaphragmDOFRecombineCS::FPermutationDomain PermutationVector;
+		PermutationVector.Set<FDDOFLayerProcessingDim>(bGatherForeground ? EDiaphragmDOFLayerProcessing::ForegroundAndBackground : EDiaphragmDOFLayerProcessing::BackgroundOnly); // TODO.
+		if (bEnableSlightOutOfFocusBokeh)
+			PermutationVector.Set<FDDOFBokehSimulationDim>(BokehSimulation);
+		PermutationVector.Set<FDiaphragmDOFRecombineCS::FQualityDim>(RecombineQuality);
+
+		FDiaphragmDOFRecombineCS::FParameters* PassParameters = GraphBuilder.AllocParameters<FDiaphragmDOFRecombineCS::FParameters>();
+		PassParameters->CommonParameters = CommonParameters;
+		SetCocModelParameters(&PassParameters->CocModel, CocModel, /* CocRadiusBasis = */ PassViewRect.Width() * 0.5f);
+
+		PassParameters->ViewportRect = PassViewRect;
+		PassParameters->ViewportSize = FVector4(PassViewRect.Width(), PassViewRect.Height(), 1.0f / PassViewRect.Width(), 1.0f / PassViewRect.Height());
+		PassParameters->TemporalJitterPixels = View.TemporalJitterPixels;
+		PassParameters->DOFBufferUVMax = FVector2D(
+			(GatheringViewSize.X - 0.5f) / float(RefBufferSize.X),
+			(GatheringViewSize.Y - 0.5f) / float(RefBufferSize.Y));
+		
+		PassParameters->SceneColorInput = FullResGatherInputTextures.SceneColor;
+		PassParameters->SceneDepthTexture = SceneBlackboard.SceneDepthBuffer;
+		PassParameters->SceneSeparateCoc = FullResGatherInputTextures.SeparateCoc; // TODO looks useless.
+		PassParameters->SceneSeparateTranslucency = SceneSeparateTranslucency;
+		
+		PassParameters->ConvolutionInputSize = FVector4(RefBufferSize.X, RefBufferSize.Y, 1.0f / RefBufferSize.X, 1.0f / RefBufferSize.Y);
+		PassParameters->ForegroundConvolution = ForegroundConvolutionTextures;
+		PassParameters->ForegroundHoleFillingConvolution = ForegroundHoleFillingConvolutionTextures;
+		PassParameters->SlightOutOfFocusConvolution = SlightOutOfFocusConvolutionTextures;
+		PassParameters->BackgroundConvolution = BackgroundConvolutionTextures;
+
+		if (bEnableSlightOutOfFocusBokeh) // && ScatteringBokehLUTOutput.IsValid() && SlightOutOfFocusConvolutionOutput.IsValid())
+		{
+			PassParameters->BokehLUT = AddBuildBokehLUTPass(FRCPassDiaphragmDOFBuildBokehLUT::EFormat::FullResOffsetToCocDistance);
+		}
+
+		PassParameters->SceneColorOutput = GraphBuilder.CreateUAV(NewSceneColor);
+
+		TShaderMapRef<FDiaphragmDOFRecombineCS> ComputeShader(View.ShaderMap, PermutationVector);
+		FComputeShaderUtils::AddPass(
+			GraphBuilder,
+			RDG_EVENT_NAME("DiaphragmDOF Recombine(%s Quality=%d Bokeh=%s alpha=no) %dx%d",
+				GetEventName(PermutationVector.Get<FDDOFLayerProcessingDim>()),
+				RecombineQuality,
+				GetEventName(PermutationVector.Get<FDDOFBokehSimulationDim>()),
+				PassViewRect.Width(), PassViewRect.Height()),
+			*ComputeShader,
+			PassParameters,
+			FComputeShaderUtils::GetGroupCount(PassViewRect.Size(), kDefaultGroupSize));
+	}
+
+	return NewSceneColor;
+}
diff --git a/Engine/Source/Runtime/Renderer/Private/PostProcess/DiaphragmDOF.h b/Engine/Source/Runtime/Renderer/Private/PostProcess/DiaphragmDOF.h
--- a/Engine/Source/Runtime/Renderer/Private/PostProcess/DiaphragmDOF.h
+++ b/Engine/Source/Runtime/Renderer/Private/PostProcess/DiaphragmDOF.h
@@ -15,6 +15,11 @@
 #define WITH_DIAPHRAGM_DOF (PLATFORM_WINDOWS || PLATFORM_XBOXONE || PLATFORM_PS4 || PLATFORM_MAC || PLATFORM_LINUX || PLATFORM_IOS || PLATFORM_SWITCH)
 
 
+class FViewInfo;
+class FSceneViewFamilyBlackboard;
+struct FTemporalAAHistory;
+
+
 namespace DiaphragmDOF
 {
 
@@ -127,7 +132,17 @@ inline bool IsSupported(EShaderPlatform ShaderPlatform)
 }
 
 
+/** Wire all DOF's passes according to view settings and cvars to convolve the scene color. */
+RENDERER_API FRDGTextureRef AddPasses(
+	FRDGBuilder& GraphBuilder,
+	const FSceneViewFamilyBlackboard& SceneBlackboard,
+	const FViewInfo& View,
+	FRDGTextureRef InputSceneColor,
+	FRDGTextureRef SeparateTranslucency);
+
+
 /** Wire all DOF's passes according to view settings and cvars to convolve the scene color (Context.FinalOutput). */
+// TODO: kill
 RENDERER_API bool WireSceneColorPasses(FPostprocessContext& Context, const FRenderingCompositeOutputRef& VelocityInput, const FRenderingCompositeOutputRef& SeparateTranslucency);
 
-}
+} // namespace DiaphragmDOF
diff --git a/Engine/Source/Runtime/Renderer/Private/PostProcess/DiaphragmDOFGraph.cpp b/Engine/Source/Runtime/Renderer/Private/PostProcess/DiaphragmDOFGraph.cpp
--- a/Engine/Source/Runtime/Renderer/Private/PostProcess/DiaphragmDOFGraph.cpp
+++ b/Engine/Source/Runtime/Renderer/Private/PostProcess/DiaphragmDOFGraph.cpp
@@ -20,104 +20,38 @@ enum class EGatheringGraphLayout
 	SeparateHalfEighth,
 };
 
+}
+
+extern TAutoConsoleVariable<int32> CVarAccumulatorQuality;
+
+extern TAutoConsoleVariable<int32> CVarEnableGatherBokehSettings;
+
+extern TAutoConsoleVariable<int32> CVarPostFilteringMethod;
+
+extern TAutoConsoleVariable<int32> CVarRingCount;
+
+
+extern TAutoConsoleVariable<int32> CVarHybridScatterForegroundMode;
+
+extern TAutoConsoleVariable<int32> CVarHybridScatterBackgroundMode;
+
+extern TAutoConsoleVariable<int32> CVarEnableScatterBokehSettings;
 
-TAutoConsoleVariable<int32> CVarAccumulatorQuality(
-	TEXT("r.DOF.Gather.AccumulatorQuality"),
-	1,
-	TEXT("Controles the quality of the gathering accumulator.\n"),
-	ECVF_Scalability | ECVF_RenderThreadSafe);
-
-TAutoConsoleVariable<int32> CVarEnableGatherBokehSettings(
-	TEXT("r.DOF.Gather.EnableBokehSettings"),
-	1,
-	TEXT("Whether to applies bokeh settings on foreground and background gathering.\n")
-	TEXT(" 0: Disable;\n 1: Enable (default)."),
-	ECVF_Scalability | ECVF_RenderThreadSafe);
-
-TAutoConsoleVariable<int32> CVarPostFilteringMethod(
-	TEXT("r.DOF.Gather.PostfilterMethod"),
-	1,
-	TEXT("Method to use to post filter a gather pass.\n")
-	TEXT(" 0: None;\n")
-	TEXT(" 1: Per RGB channel median 3x3 (default);\n")
-	TEXT(" 2: Per RGB channel max 3x3."),
-	ECVF_Scalability | ECVF_RenderThreadSafe);
-
-TAutoConsoleVariable<int32> CVarRingCount(
-	TEXT("r.DOF.Gather.RingCount"),
-	5,
-	TEXT("Number of rings for gathering kernels [[3; 5]]. Default to 5.\n"),
-	ECVF_Scalability | ECVF_RenderThreadSafe);
-
-
-TAutoConsoleVariable<int32> CVarHybridScatterForegroundMode(
-	TEXT("r.DOF.Scatter.ForegroundCompositing"),
-	1,
-	TEXT("Compositing mode of the foreground hybrid scattering.\n")
-	TEXT(" 0: Disabled;\n")
-	TEXT(" 1: Additive (default)."),
-	ECVF_Scalability | ECVF_RenderThreadSafe);
-
-TAutoConsoleVariable<int32> CVarHybridScatterBackgroundMode(
-	TEXT("r.DOF.Scatter.BackgroundCompositing"),
-	2,
-	TEXT("Compositing mode of the background hybrid scattering.\n")
-	TEXT(" 0: Disabled;\n")
-	TEXT(" 1: Additive;\n")
-	TEXT(" 2: Gather occlusion (default)."),
-	ECVF_Scalability | ECVF_RenderThreadSafe);
-
-TAutoConsoleVariable<int32> CVarEnableScatterBokehSettings(
-	TEXT("r.DOF.Scatter.EnableBokehSettings"),
-	1,
-	TEXT("Whether to enable bokeh settings on scattering.\n")
-	TEXT(" 0: Disable;\n 1: Enable (default)."),
-	ECVF_Scalability | ECVF_RenderThreadSafe);
-
-TAutoConsoleVariable<float> CVarScatterMinCocRadius(
-	TEXT("r.DOF.Scatter.MinCocRadius"),
-	3.0f,
-	TEXT("Minimal Coc radius required to be scattered (default = 3)."),
-	ECVF_Scalability | ECVF_RenderThreadSafe);
-
-TAutoConsoleVariable<float> CVarScatterMaxSpriteRatio(
-	TEXT("r.DOF.Scatter.MaxSpriteRatio"),
-	0.1f,
-	TEXT("Maximum ratio of scattered pixel quad as sprite, usefull to control DOF's scattering upperbound. ")
-	TEXT(" 1 will allow to scatter 100% pixel quads, whereas 0.2 will only allow 20% (default = 0.1)."),
-	ECVF_Scalability | ECVF_RenderThreadSafe);
-
-TAutoConsoleVariable<int32> CVarEnableRecombineBokehSettings(
-	TEXT("r.DOF.Recombine.EnableBokehSettings"),
-	1,
-	TEXT("Whether to applies bokeh settings on slight out of focus done in recombine pass.\n")
-	TEXT(" 0: Disable;\n 1: Enable (default)."),
-	ECVF_Scalability | ECVF_RenderThreadSafe);
-
-TAutoConsoleVariable<int32> CVarRecombineQuality(
-	TEXT("r.DOF.Recombine.Quality"),
-	2,
-	TEXT("Configures the quality of the recombine pass.\n")
-	TEXT(" 0: No slight out of focus;\n")
-	TEXT(" 1: Slight out of focus 24spp;\n")
-	TEXT(" 2: Slight out of focus 32spp (default)."),
-	ECVF_Scalability | ECVF_RenderThreadSafe);
-
-TAutoConsoleVariable<float> CVarMinimalFullresBlurRadius(
-	TEXT("r.DOF.Recombine.MinFullresBlurRadius"),
-	0.1f,
-	TEXT("Minimal blurring radius used in full resolution pixel width to actually do ")
-	TEXT("DOF  when slight out of focus is enabled (default = 0.1)."),
-	ECVF_RenderThreadSafe);
-
-TAutoConsoleVariable<int32> CVarDOFTemporalAAQuality(
-	TEXT("r.DOF.TemporalAAQuality"),
-	1,
-	TEXT("Quality of temporal AA pass done in DOF.\n")
-	TEXT(" 0: Faster but lower quality;")
-	TEXT(" 1: Higher quality pass (default)."),
-	ECVF_Scalability | ECVF_RenderThreadSafe);
+extern TAutoConsoleVariable<float> CVarScatterMinCocRadius;
 
+extern TAutoConsoleVariable<float> CVarScatterMaxSpriteRatio;
+
+extern TAutoConsoleVariable<int32> CVarEnableRecombineBokehSettings;
+
+extern TAutoConsoleVariable<int32> CVarRecombineQuality;
+
+extern TAutoConsoleVariable<float> CVarMinimalFullresBlurRadius;
+
+extern TAutoConsoleVariable<int32> CVarDOFTemporalAAQuality;
+
+
+namespace
+{
 
 EDiaphragmDOFPostfilterMethod GetPostfilteringMethod()
 {
diff --git a/Engine/Source/Runtime/Renderer/Private/PostProcess/DiaphragmDOFPasses.cpp b/Engine/Source/Runtime/Renderer/Private/PostProcess/DiaphragmDOFPasses.cpp
--- a/Engine/Source/Runtime/Renderer/Private/PostProcess/DiaphragmDOFPasses.cpp
+++ b/Engine/Source/Runtime/Renderer/Private/PostProcess/DiaphragmDOFPasses.cpp
@@ -1464,6 +1464,8 @@ class FPostProcessDiaphragmDOFPostfilterCS : public FPostProcessDiaphragmDOFShad
 		FPermutationDomain PermutationVector(Parameters.PermutationId);
 		if (RemapPermutationVector(PermutationVector) != PermutationVector) return false;
 		if (PermutationVector.Get<FDDOFPostfilterMethodDim>() == EDiaphragmDOFPostfilterMethod::None) return false;
+		if (PermutationVector.Get<FDDOFLayerProcessingDim>() != EDiaphragmDOFLayerProcessing::ForegroundOnly &&
+			PermutationVector.Get<FDDOFLayerProcessingDim>() != EDiaphragmDOFLayerProcessing::BackgroundOnly) return false;
 
 		return FPostProcessDiaphragmDOFShader::ShouldCompilePermutation(Parameters);
 	}
diff --git a/Engine/Source/Runtime/Renderer/Private/PostProcess/PostProcessing.cpp b/Engine/Source/Runtime/Renderer/Private/PostProcess/PostProcessing.cpp
--- a/Engine/Source/Runtime/Renderer/Private/PostProcess/PostProcessing.cpp
+++ b/Engine/Source/Runtime/Renderer/Private/PostProcess/PostProcessing.cpp
@@ -57,6 +57,7 @@
 #include "DeferredShadingRenderer.h"
 #include "MobileSeparateTranslucencyPass.h"
 #include "MobileDistortionPass.h"
+#include "SceneViewFamilyBlackboard.h"
 
 /** The global center for all post processing activities. */
 FPostProcessing GPostProcessing;
@@ -1537,11 +1538,40 @@ void FPostProcessing::Process(FRHICommandListImmediate& RHICmdList, const FViewI
 						DOFVelocityRef = FRenderingCompositeOutputRef(NoVelocity);
 					}
 					
-					if (CVarUseDiaphragmDOF.GetValueOnRenderThread() == 1 &&
+					if (CVarUseDiaphragmDOF.GetValueOnRenderThread() > 0 &&
 						DiaphragmDOF::IsSupported(View.GetShaderPlatform()))
 					{
-						bSepTransWasApplied = DiaphragmDOF::WireSceneColorPasses(
-							Context, DOFVelocityRef,SeparateTranslucency);
+						if (CVarUseDiaphragmDOF.GetValueOnRenderThread() == 2)
+						{
+							FRenderingCompositePass* DiaphragmDOFPass = Context.Graph.RegisterPass(
+								new(FMemStack::Get()) TRCPassForRDG<2, 1>([](FRenderingCompositePass* Pass, FRenderingCompositePassContext& Context)
+							{
+								FRDGBuilder GraphBuilder(Context.RHICmdList);
+
+								FSceneRenderTargets& SceneContext = FSceneRenderTargets::Get(Context.RHICmdList);
+								FSceneViewFamilyBlackboard SceneBlackboard;
+								SetupSceneViewFamilyBlackboard(GraphBuilder, &SceneBlackboard);
+	
+								FRDGTextureRef SceneColor = Pass->CreateRDGTextureForInput(GraphBuilder, ePId_Input0, TEXT("SceneColor"), eFC_0000);
+								FRDGTextureRef SeparateTranslucency = Pass->CreateRDGTextureForInput(GraphBuilder, ePId_Input1, TEXT("SeparateTranslucency"), eFC_0000);
+
+								FRDGTextureRef NewSceneColor = DiaphragmDOF::AddPasses(
+									GraphBuilder,
+									SceneBlackboard, Context.View,
+									SceneColor, SeparateTranslucency);
+		
+								Pass->ExtractRDGTextureForOutput(GraphBuilder, ePId_Output0, NewSceneColor);
+								GraphBuilder.Execute();
+							}));
+							DiaphragmDOFPass->SetInput(ePId_Input0, Context.FinalOutput);
+							DiaphragmDOFPass->SetInput(ePId_Input1, SeparateTranslucency);
+							Context.FinalOutput = FRenderingCompositeOutputRef(DiaphragmDOFPass, ePId_Output0);
+						}
+						else
+						{
+							bSepTransWasApplied = DiaphragmDOF::WireSceneColorPasses(
+								Context, DOFVelocityRef,SeparateTranslucency);
+						}
 					}
 					else
 					{
