Transforms AddPass_ClearUAV() into regular functions.

diff --git a/Engine/Source/Runtime/Renderer/Private/LightGridInjection.cpp b/Engine/Source/Runtime/Renderer/Private/LightGridInjection.cpp
--- a/Engine/Source/Runtime/Renderer/Private/LightGridInjection.cpp
+++ b/Engine/Source/Runtime/Renderer/Private/LightGridInjection.cpp
@@ -80,7 +80,7 @@ BEGIN_SHADER_PARAMETER_STRUCT(FClearUAVParameters, )
 	SHADER_PARAMETER_RDG_BUFFER_UAV(	RWBuffer<uint>,	BufferUAV)
 END_SHADER_PARAMETER_STRUCT()
 
-FORCEINLINE_DEBUGGABLE void AddPass_ClearUAV(
+void AddPass_ClearUAV(
 	FRDGBuilder& GraphBuilder,
 	FRDGEventName&& PassName,
 	FRDGBufferUAVRef BufferUAV,
