Removes FORCEINLINE on AddPass_ClearUAV declaration in RDG refactor of DOF.

diff --git a/Engine/Source/Runtime/Renderer/Private/PostProcess/DiaphragmDOF.cpp b/Engine/Source/Runtime/Renderer/Private/PostProcess/DiaphragmDOF.cpp
--- a/Engine/Source/Runtime/Renderer/Private/PostProcess/DiaphragmDOF.cpp
+++ b/Engine/Source/Runtime/Renderer/Private/PostProcess/DiaphragmDOF.cpp
@@ -22,7 +22,7 @@
 #include "SceneViewFamilyBlackboard.h"
 
 
-FORCEINLINE_DEBUGGABLE void AddPass_ClearUAV(
+void AddPass_ClearUAV(
 	FRDGBuilder& GraphBuilder,
 	FRDGEventName&& PassName,
 	FRDGBufferUAVRef BufferUAV,
