Writes dedicated pixel shader to compose separate translucency into scene color when DOF is disabled, to no longer fall back on the deprecated bokeh DOF's recombine pass.

diff --git a/Engine/Shaders/Private/ComposeSeparateTranslucency.usf b/Engine/Shaders/Private/ComposeSeparateTranslucency.usf
new file mode 100644
--- /dev/null
+++ b/Engine/Shaders/Private/ComposeSeparateTranslucency.usf
@@ -0,0 +1,26 @@
+// Copyright 1998-2019 Epic Games, Inc. All Rights Reserved.
+
+#include "Common.ush"
+
+Texture2D SceneColor;
+SamplerState SceneColorSampler;
+
+Texture2D SeparateTranslucency;
+SamplerState SeparateTranslucencySampler;
+
+void MainPS(
+	float4 SvPosition : SV_POSITION,
+	out float4 OutColor : SV_Target0)
+{
+	float2 BufferUV = SvPosition * View.BufferSizeAndInvSize.zw;
+	
+	float4 SceneColorSample = SceneColor.SampleLevel(SceneColorSampler, BufferUV, 0);
+	float4 SeparateTranslucencySample = SeparateTranslucency.SampleLevel(SeparateTranslucencySampler, BufferUV, 0);
+
+	OutColor.rgb = SceneColorSample.rgb * SeparateTranslucencySample.a + SeparateTranslucencySample.rgb;
+	OutColor.a = SceneColorSample.a * SeparateTranslucencySample.a;
+	
+#if !POST_PROCESS_ALPHA
+	OutColor.a = 0;
+#endif
+}
diff --git a/Engine/Source/Runtime/Renderer/Private/PostProcess/PostProcessing.cpp b/Engine/Source/Runtime/Renderer/Private/PostProcess/PostProcessing.cpp
--- a/Engine/Source/Runtime/Renderer/Private/PostProcess/PostProcessing.cpp
+++ b/Engine/Source/Runtime/Renderer/Private/PostProcess/PostProcessing.cpp
@@ -58,6 +58,7 @@
 #include "MobileSeparateTranslucencyPass.h"
 #include "MobileDistortionPass.h"
 #include "SceneViewFamilyBlackboard.h"
+#include "PixelShaderUtils.h"
 
 /** The global center for all post processing activities. */
 FPostProcessing GPostProcessing;
@@ -1380,6 +1381,59 @@ void FPostProcessing::RegisterHMDPostprocessPass(FPostprocessContext& Context, c
 	}
 }
 
+
+namespace
+{
+
+class FComposeSeparateTranslucencyPS : public FGlobalShader
+{
+	DECLARE_GLOBAL_SHADER(FComposeSeparateTranslucencyPS);
+	SHADER_USE_PARAMETER_STRUCT(FComposeSeparateTranslucencyPS, FGlobalShader);
+
+	BEGIN_SHADER_PARAMETER_STRUCT(FParameters, )
+		SHADER_PARAMETER_RDG_TEXTURE(Texture2D, SceneColor)
+		SHADER_PARAMETER_SAMPLER(SamplerState, SceneColorSampler)
+		SHADER_PARAMETER_RDG_TEXTURE(Texture2D, SeparateTranslucency)
+		SHADER_PARAMETER_SAMPLER(SamplerState, SeparateTranslucencySampler)
+		SHADER_PARAMETER_STRUCT_REF(FViewUniformShaderParameters, ViewUniformBuffer)
+		RENDER_TARGET_BINDING_SLOTS()
+	END_SHADER_PARAMETER_STRUCT()
+
+	static bool ShouldCompilePermutation(const FGlobalShaderPermutationParameters& Parameters)
+	{
+		return IsFeatureLevelSupported(Parameters.Platform, ERHIFeatureLevel::SM4);
+	}
+};
+
+IMPLEMENT_GLOBAL_SHADER(FComposeSeparateTranslucencyPS, "/Engine/Private/ComposeSeparateTranslucency.usf", "MainPS", SF_Pixel);
+
+FRDGTextureRef AddSeparateTranslucencyCompositionPass(FRDGBuilder& GraphBuilder, const FViewInfo& View, FRDGTextureRef SceneColor, FRDGTextureRef SeparateTranslucency)
+{
+	FRDGTextureRef NewSceneColor = GraphBuilder.CreateTexture(SceneColor->Desc, TEXT("SceneColor"));
+
+	FComposeSeparateTranslucencyPS::FParameters* PassParameters = GraphBuilder.AllocParameters<FComposeSeparateTranslucencyPS::FParameters>();
+	PassParameters->SceneColor = SceneColor;
+	PassParameters->SceneColorSampler = TStaticSamplerState<SF_Point>::GetRHI();
+	PassParameters->SeparateTranslucency = SeparateTranslucency;
+	PassParameters->SeparateTranslucencySampler = TStaticSamplerState<SF_Point>::GetRHI();
+	PassParameters->ViewUniformBuffer = View.ViewUniformBuffer;
+	PassParameters->RenderTargets[0] = FRenderTargetBinding(NewSceneColor, ERenderTargetLoadAction::ENoAction, ERenderTargetStoreAction::EStore);
+			
+	TShaderMapRef<FComposeSeparateTranslucencyPS> PixelShader(View.ShaderMap);
+	FPixelShaderUtils::AddFullscreenPass(
+		GraphBuilder,
+		View.ShaderMap,
+		RDG_EVENT_NAME("ComposeSeparateTranslucency %dx%d", View.ViewRect.Width(), View.ViewRect.Height()),
+		*PixelShader,
+		PassParameters,
+		View.ViewRect);
+
+	return NewSceneColor;
+}
+
+} // namespace
+
+
 void FPostProcessing::Process(FRHICommandListImmediate& RHICmdList, const FViewInfo& View, TRefCountPtr<IPooledRenderTarget>& VelocityRT)
 {
 	QUICK_SCOPE_CYCLE_COUNTER( STAT_PostProcessing_Process );
@@ -1560,12 +1614,21 @@ void FPostProcessing::Process(FRHICommandListImmediate& RHICmdList, const FViewI
 									SceneBlackboard, Context.View,
 									SceneColor, SeparateTranslucency);
 		
+								// DOF passes were not added, therefore need to compose Separate translucency manually. 
+								if (NewSceneColor == SceneColor)
+								{
+									NewSceneColor = AddSeparateTranslucencyCompositionPass(GraphBuilder, Context.View, SceneColor, SeparateTranslucency);
+								}
+
 								Pass->ExtractRDGTextureForOutput(GraphBuilder, ePId_Output0, NewSceneColor);
+
 								GraphBuilder.Execute();
 							}));
 							DiaphragmDOFPass->SetInput(ePId_Input0, Context.FinalOutput);
 							DiaphragmDOFPass->SetInput(ePId_Input1, SeparateTranslucency);
 							Context.FinalOutput = FRenderingCompositeOutputRef(DiaphragmDOFPass, ePId_Output0);
+
+							bSepTransWasApplied = true;
 						}
 						else
 						{
