Refactored Screen Pass utilities and Subsurface Scattering Implementation.

 - Removed samplers from input struct, since most samplers are shared between inputs.
 - Introduced the FScreenPassTextureParameters struct, describing pixel / UV rect dimension
   used to sample from a region of the texture in HLSL. Centralizes a lot of the common math
   operations currently being done ad-hoc or exposed in PostProcessCommon.ush.

 - Removed FScreenPassInput in favor of a non-shader parameter class called FScreenPassTexture. This
   data structure pairs together the RDG texture and screen pass shader parameters. This allows pass
   inputs to carry important state like which view rect is being used; this provides a more scalable alternative
   to 'SceneColorViewRect' and 'ReferenceBufferSize'.

 - Added an FScreenPassTextureViewportTransform struct that is designed to map UVs from one texture space
   to another. For example, from the upsampled temporal AA UV to velocity UV.

 - Replaced view-specific parameters with screen pass parameters in the subsurface shader.

diff --git a/Engine/Shaders/Private/MaterialTemplate.ush b/Engine/Shaders/Private/MaterialTemplate.ush
--- a/Engine/Shaders/Private/MaterialTemplate.ush
+++ b/Engine/Shaders/Private/MaterialTemplate.ush
@@ -1495,19 +1495,19 @@ float4 SceneTextureLookup(float2 UV, int SceneTextureIndex, bool bFiltered)
 		case 13: return ScreenSpaceData.GBuffer.CustomDepth;
 #if POST_PROCESS_MATERIAL
 		// PPI_PostprocessInput0
-		case 14: return Texture2DSample(PostprocessInput0, bFiltered ? BilinearTextureSampler0 : PostprocessInput0Sampler, UV);
+		case 14: return Texture2DSample(PostprocessInput0, bFiltered ? BilinearTextureSampler : PostprocessInput0Sampler, UV);
 		// PPI_PostprocessInput1
-		case 15: return Texture2DSample(PostprocessInput1, bFiltered ? BilinearTextureSampler0 : PostprocessInput1Sampler, UV);
+		case 15: return Texture2DSample(PostprocessInput1, bFiltered ? BilinearTextureSampler : PostprocessInput1Sampler, UV);
 		// PPI_PostprocessInput2
-		case 16: return Texture2DSample(PostprocessInput2, bFiltered ? BilinearTextureSampler0 : PostprocessInput2Sampler, UV);
+		case 16: return Texture2DSample(PostprocessInput2, bFiltered ? BilinearTextureSampler : PostprocessInput2Sampler, UV);
 		// PPI_PostprocessInput3
-		case 17: return Texture2DSample(PostprocessInput3, bFiltered ? BilinearTextureSampler0 : PostprocessInput3Sampler, UV);
+		case 17: return Texture2DSample(PostprocessInput3, bFiltered ? BilinearTextureSampler : PostprocessInput3Sampler, UV);
 		// PPI_PostprocessInput4
-		case 18: return Texture2DSample(PostprocessInput4, bFiltered ? BilinearTextureSampler0 : PostprocessInput4Sampler, UV);
+		case 18: return Texture2DSample(PostprocessInput4, bFiltered ? BilinearTextureSampler : PostprocessInput4Sampler, UV);
 		// PPI_PostprocessInput5
-		case 19: return Texture2DSample(PostprocessInput5, bFiltered ? BilinearTextureSampler0 : PostprocessInput5Sampler, UV);
+		case 19: return Texture2DSample(PostprocessInput5, bFiltered ? BilinearTextureSampler : PostprocessInput5Sampler, UV);
 		// PPI_PostprocessInput6
-		case 20: return Texture2DSample(PostprocessInput6, bFiltered ? BilinearTextureSampler0 : PostprocessInput6Sampler, UV);
+		case 20: return Texture2DSample(PostprocessInput6, bFiltered ? BilinearTextureSampler : PostprocessInput6Sampler, UV);
 #endif // __POST_PROCESS_COMMON__
 		// PPI_DecalMask
 		case 21: return 0;  // material compiler will return an error
diff --git a/Engine/Shaders/Private/PostProcessCommon.ush b/Engine/Shaders/Private/PostProcessCommon.ush
--- a/Engine/Shaders/Private/PostProcessCommon.ush
+++ b/Engine/Shaders/Private/PostProcessCommon.ush
@@ -71,7 +71,7 @@ float4 ViewportSize;
 // in pixels (scaled to the current rendertarget resolution), float4(minx,miny,maxx,maxy), depends on SetViewportAndCallRHI() call
 uint4 ViewportRect;
 
-SamplerState BilinearTextureSampler0;
+SamplerState BilinearTextureSampler;
 
 // transforms ScreenPos to the screen pixel position, not the viewport local position (*xy +wz)
 // (0,0) left top, (width-1,height-1) at right bottom, in pixels (without offset to reach the pixel center)
@@ -187,7 +187,7 @@ float4 GatherSceneDepth(float2 UV, float2 InvBufferSize)
 {
 #if FEATURE_LEVEL >= FEATURE_LEVEL_SM5
 	// using Gather: xyzw in counter clockwise order starting with the sample to the lower left of the queried location
-	float4 DeviceZ = SceneTexturesStruct.SceneDepthTextureNonMS.Gather(BilinearTextureSampler0, UV);
+	float4 DeviceZ = SceneTexturesStruct.SceneDepthTextureNonMS.Gather(BilinearTextureSampler, UV);
 
 	return float4(
 		ConvertFromDeviceZ(DeviceZ.x),
diff --git a/Engine/Shaders/Private/PostProcessSubsurface.usf b/Engine/Shaders/Private/PostProcessSubsurface.usf
--- a/Engine/Shaders/Private/PostProcessSubsurface.usf
+++ b/Engine/Shaders/Private/PostProcessSubsurface.usf
@@ -18,11 +18,17 @@ float4 SubsurfaceParams;
 Texture2D SubsurfaceInput0_Texture;
 Texture2D SubsurfaceInput1_Texture;
 
-SamplerState SubsurfaceInput0_Sampler;
-SamplerState SubsurfaceInput1_Sampler;
+SamplerState SubsurfaceSampler; 
 
-float4 SubsurfaceInput0_Size;
-float4 SubsurfaceInput1_Size;
+float2 SubsurfaceInput0_Extent;
+float2 SubsurfaceInput0_ExtentInverse;
+uint2 SubsurfaceInput0_ViewportMin;
+float2 SubsurfaceInput0_ViewportSize;
+float2 SubsurfaceInput0_ViewportSizeInverse;
+float2 SubsurfaceInput0_UVViewportBilinearMin;
+float2 SubsurfaceInput0_UVViewportBilinearMax;
+
+float2 SubsurfaceInput1_ExtentInverse;
 
 #define SUBSURFACE_DIRECTION_HORIZONTAL 0
 #define SUBSURFACE_DIRECTION_VERTICAL 1
@@ -112,9 +118,8 @@ float3 GetSubsurfaceProfileColor(FGBufferData GBufferData)
 
 uint GetSubsurfaceProfileId(float2 BufferUV)
 {
-#if SUBSURFACE_MANUALLY_CLAMP_UV
-	BufferUV = clamp(BufferUV, View.BufferBilinearUVMinMax.xy, View.BufferBilinearUVMinMax.zw);
-#endif
+	BufferUV = clamp(BufferUV, SubsurfaceInput0_UVViewportBilinearMin, SubsurfaceInput0_UVViewportBilinearMax);
+
 	FGBufferData LocalGBufferData = GetGBufferData(BufferUV);
 	return ExtractSubsurfaceProfileInt(LocalGBufferData);
 }
@@ -129,10 +134,9 @@ float3 GetSubsurfaceProfileBoundaryColorBleed(FGBufferData GBufferData)
 
 float4 GetSceneColor(float2 BufferUV)
 {
-#if SUBSURFACE_MANUALLY_CLAMP_UV
-	BufferUV = clamp(BufferUV, View.BufferBilinearUVMinMax.xy, View.BufferBilinearUVMinMax.zw);
-#endif
-	return Texture2DSampleLevel(SubsurfaceInput0_Texture, SubsurfaceInput0_Sampler, BufferUV, 0);
+	BufferUV = clamp(BufferUV, SubsurfaceInput0_UVViewportBilinearMin, SubsurfaceInput0_UVViewportBilinearMax);
+
+	return Texture2DSampleLevel(SubsurfaceInput0_Texture, SubsurfaceSampler, BufferUV, 0);
 }
 
 // from [REDACTED_PROFILE_URL]/separable-sss/tree/master/Demo
@@ -148,7 +152,7 @@ bool InUnitBox(float2 UV)
 }
 
 // @return 0=don't blend in, 1:fully blend in
-float ComputeFullResLerp(FScreenSpaceData ScreenSpaceData, float2 UVSceneColor, float4 FullResInputSize)
+float ComputeFullResLerp(FScreenSpaceData ScreenSpaceData, float2 UVSceneColor, float2 FullResInputSizeInverse)
 {
 	float SSSScaleX = SubsurfaceParams.x;
 
@@ -161,7 +165,7 @@ float ComputeFullResLerp(FScreenSpaceData ScreenSpaceData, float2 UVSceneColor,
 
 	finalStep *= GetSubsurfaceProfileRadiusScale(ScreenSpaceData.GBuffer);
 
-	float PixelSizeRadius = finalStep / (FullResInputSize.z * 0.5f);
+	float PixelSizeRadius = finalStep / (FullResInputSizeInverse.x * 0.5f);
 
 	// tweaked for skin, a more flat kernel might need a smaller value, around 2 seems reasonable because we do half res
 	const float PixelSize = 4.0f;
@@ -180,17 +184,17 @@ void VisualizePS(in noperspective float4 UVAndScreenPos : TEXCOORD0, float4 SvPo
 {
 	float2 UV = UVAndScreenPos.xy;
 
-	OutColor = Texture2DSample(SubsurfaceInput0_Texture, SubsurfaceInput0_Sampler, UV);
+	OutColor = Texture2DSample(SubsurfaceInput0_Texture, SubsurfaceSampler, UV);
 
 	int2 PixelPos = (int2)SvPosition.xy;
 
-	float2 ViewLocalUV = (PixelPos - View.ViewRectMin.xy) * View.ViewSizeAndInvSize.zw;
+	float2 ViewLocalUV = (PixelPos - SubsurfaceInput0_ViewportMin) * SubsurfaceInput0_ViewportSizeInverse;
 
 	float2 IDAreaLocalUV = ViewLocalUV * 2 - 1.0f;
 
 	if (InUnitBox(IDAreaLocalUV))
 	{
-		float2 UV = View.ViewRectMin.xy * View.BufferSizeAndInvSize.zw + IDAreaLocalUV * (View.ViewSizeAndInvSize.xy * View.BufferSizeAndInvSize.zw);
+		float2 UV = SubsurfaceInput0_ViewportMin * SubsurfaceInput0_ExtentInverse + IDAreaLocalUV * (SubsurfaceInput0_ViewportSize * SubsurfaceInput0_ExtentInverse);
 
 		FScreenSpaceData ScreenSpaceData = GetScreenSpaceData(UV);
 
@@ -241,7 +245,7 @@ void VisualizePS(in noperspective float4 UVAndScreenPos : TEXCOORD0, float4 SvPo
 			int2 LeftTop = (PixelPos / 8) * 8;
 			PrintCharacter(PixelPos, OutColor.rgb, float3(1, 1, 1), LeftTop, SubsurfaceProfileInt);
 
-			OutColor.rgb *= ComputeFullResLerp(ScreenSpaceData, UV, SubsurfaceInput0_Size);
+			OutColor.rgb *= ComputeFullResLerp(ScreenSpaceData, UV, SubsurfaceInput0_ExtentInverse);
 		}
 	}
 
@@ -260,10 +264,10 @@ half3 LookupSceneColor(float2 SceneUV, int2 PixelOffset)
 #if ES2_PROFILE && COMPILER_GLSL_ES2
 	// slower but always works
 	// to prevent "error: Texture offset not supported on GLSL ES"
-	return Texture2DSample(SubsurfaceInput0_Texture, SubsurfaceInput0_Sampler, SceneUV + PixelOffset * SubsurfaceInput0_Size.zw).rgb;
+	return Texture2DSample(SubsurfaceInput0_Texture, SubsurfaceSampler, SceneUV + PixelOffset * SubsurfaceInput0_ExtentInverse).rgb;
 #else
 	// faster
-	return SubsurfaceInput0_Texture.SampleLevel(SubsurfaceInput0_Sampler, SceneUV, 0, PixelOffset).rgb;
+	return SubsurfaceInput0_Texture.SampleLevel(SubsurfaceSampler, SceneUV, 0, PixelOffset).rgb;
 #endif
 }
  
@@ -280,7 +284,7 @@ SDiffuseAndSpecular ReconstructLighting(float2 UVSceneColor, uint ReconstructMet
 
 	  // todo: We could alternate the diagonal with TemporalAA or even only only 1 sample for low spec or 4 for high spec
 
-	  float3 Quant0 = SubsurfaceInput0_Texture.SampleLevel(SubsurfaceInput0_Sampler, UVSceneColor, 0).rgb;
+	  float3 Quant0 = SubsurfaceInput0_Texture.SampleLevel(SubsurfaceSampler, UVSceneColor, 0).rgb;
 
 	  // todo: expose as scalability setting (can be evaluate best without TemporalAA)
 	  // 0:fast but pattern can appear, 1:better, 2: even better, 3: best but expensive
@@ -334,7 +338,7 @@ SDiffuseAndSpecular ReconstructLighting(float2 UVSceneColor, uint ReconstructMet
 #else // SUBSURFACE_PROFILE_CHECKERBOARD
 	{
 	  // If we're not doing checkerboard encoding, we just need to read a single pixel and decode (combined diffuse/spec in RGB)
-	  float4 CenterSample = SubsurfaceInput0_Texture.SampleLevel(SubsurfaceInput0_Sampler, UVSceneColor, 0);
+	  float4 CenterSample = SubsurfaceInput0_Texture.SampleLevel(SubsurfaceSampler, UVSceneColor, 0);
 	  float3 CombinedColor = CenterSample.rgb;
 	  float DiffuseLuminance = CenterSample.a;
   
@@ -381,10 +385,10 @@ void SetupPS(in noperspective float4 UVAndScreenPos : TEXCOORD0, out float4 OutC
 	// order aligned with Gather() hardware implementation
 	// RGB: color*A, A:weight 0 if no subsurface scattering
 
-	float4 A = SetupSubsurfaceForOnePixel(min(BufferUV + float2(-0.5,  0.5f) * SubsurfaceInput0_Size.zw, View.BufferBilinearUVMinMax.zw));
-	float4 B = SetupSubsurfaceForOnePixel(min(BufferUV + float2( 0.5,  0.5f) * SubsurfaceInput0_Size.zw, View.BufferBilinearUVMinMax.zw));
-	float4 C = SetupSubsurfaceForOnePixel(min(BufferUV + float2( 0.5, -0.5f) * SubsurfaceInput0_Size.zw, View.BufferBilinearUVMinMax.zw));
-	float4 D = SetupSubsurfaceForOnePixel(min(BufferUV + float2(-0.5, -0.5f) * SubsurfaceInput0_Size.zw, View.BufferBilinearUVMinMax.zw));
+	float4 A = SetupSubsurfaceForOnePixel(min(BufferUV + float2(-0.5,  0.5f) * SubsurfaceInput0_ExtentInverse, SubsurfaceInput0_UVViewportBilinearMax));
+	float4 B = SetupSubsurfaceForOnePixel(min(BufferUV + float2( 0.5,  0.5f) * SubsurfaceInput0_ExtentInverse, SubsurfaceInput0_UVViewportBilinearMax));
+	float4 C = SetupSubsurfaceForOnePixel(min(BufferUV + float2( 0.5, -0.5f) * SubsurfaceInput0_ExtentInverse, SubsurfaceInput0_UVViewportBilinearMax));
+	float4 D = SetupSubsurfaceForOnePixel(min(BufferUV + float2(-0.5, -0.5f) * SubsurfaceInput0_ExtentInverse, SubsurfaceInput0_UVViewportBilinearMax));
 
 	float4 Sum = (A + B) + (C + D);
 
@@ -392,7 +396,7 @@ void SetupPS(in noperspective float4 UVAndScreenPos : TEXCOORD0, out float4 OutC
 
 	OutColor.rgb = Sum.rgb * Div;
 
-	float4 FourDepth = GatherSceneDepth(BufferUV, SubsurfaceInput0_Size.zw);
+	float4 FourDepth = GatherSceneDepth(BufferUV, SubsurfaceInput0_ExtentInverse);
 
 	// average all valid depth values to a single one
 	float SingleDepth = dot(FourDepth, float4(A.a, B.a, C.a, D.a)) * Div;
@@ -417,12 +421,10 @@ void MainPS(noperspective float4 UVAndScreenPos : TEXCOORD0, out float4 OutColor
 	float2 ViewportDirectionUV = float2(1, 0) * SUBSURFACE_RADIUS_SCALE;
 #else
 	// vertical
-	float2 ViewportDirectionUV = float2(0, 1) * SUBSURFACE_RADIUS_SCALE * (View.ViewSizeAndInvSize.x * View.ViewSizeAndInvSize.w);
+	float2 ViewportDirectionUV = float2(0, 1) * SUBSURFACE_RADIUS_SCALE * (SubsurfaceInput0_Extent.x * SubsurfaceInput0_ExtentInverse.y);
 #endif
 
-#if SUBSURFACE_MANUALLY_CLAMP_UV
-	ViewportDirectionUV *= (View.ViewSizeAndInvSize.x * View.BufferSizeAndInvSize.z);
-#endif
+	ViewportDirectionUV *= (SubsurfaceInput0_ViewportSize.x * SubsurfaceInput0_ExtentInverse.x);
 
 	OutColor = SSSSBlurPS(BufferUV, ViewportDirectionUV, false);
 
@@ -442,7 +444,7 @@ void SubsurfaceRecombinePS(noperspective float4 UVAndScreenPos : TEXCOORD0, floa
 	
 	if (!UseSubsurfaceProfile(ScreenSpaceData.GBuffer.ShadingModelID))
 	{
-		OutColor = Texture2DSample(SubsurfaceInput0_Texture, SubsurfaceInput0_Sampler, BufferUV);
+		OutColor = Texture2DSample(SubsurfaceInput0_Texture, SubsurfaceSampler, BufferUV);
 		return;
 	}
 
@@ -453,11 +455,11 @@ void SubsurfaceRecombinePS(noperspective float4 UVAndScreenPos : TEXCOORD0, floa
 #if SUBSURFACE_HALF_RES
 	// fade out subsurface scattering if radius is too small to be more crips (not blend with half resolution)
 	// minor quality improvement (faces are more detailed in distance)
-	LerpFactor = ComputeFullResLerp(ScreenSpaceData, BufferUV, SubsurfaceInput1_Size);
+	LerpFactor = ComputeFullResLerp(ScreenSpaceData, BufferUV, SubsurfaceInput1_ExtentInverse);
 #endif // SUBSURFACE_HALF_RES
 
 	{
-		float4 SSSColorWithAlpha = Texture2DSample(SubsurfaceInput1_Texture, SubsurfaceInput1_Sampler, BufferUV);
+		float4 SSSColorWithAlpha = Texture2DSample(SubsurfaceInput1_Texture, SubsurfaceSampler, BufferUV);
 
 		// renormalize to dilate RGB to fix half res upsampling artifacts
 		SSSColor = SSSColorWithAlpha.rgb / max(SSSColorWithAlpha.a, 0.00001f);
@@ -489,5 +491,5 @@ void SubsurfaceRecombinePS(noperspective float4 UVAndScreenPos : TEXCOORD0, floa
 
 void SubsurfaceViewportCopyPS(noperspective float4 InUV : TEXCOORD0, out float4 OutColor : SV_Target0)
 {
-	OutColor = Texture2DSample(SubsurfaceInput0_Texture, SubsurfaceInput0_Sampler, InUV.xy);
+	OutColor = Texture2DSample(SubsurfaceInput0_Texture, SubsurfaceSampler, InUV.xy);
 }
diff --git a/Engine/Source/Runtime/Renderer/Private/DeferredShadingRenderer.cpp b/Engine/Source/Runtime/Renderer/Private/DeferredShadingRenderer.cpp
--- a/Engine/Source/Runtime/Renderer/Private/DeferredShadingRenderer.cpp
+++ b/Engine/Source/Runtime/Renderer/Private/DeferredShadingRenderer.cpp
@@ -1790,7 +1790,11 @@ void FDeferredShadingSceneRenderer::Render(FRHICommandListImmediate& RHICmdList)
 				{
 					FScreenPassContext* Context = FScreenPassContext::Create(RHICmdList, View);
 
-					SceneColorTexture = ComputeSubsurface(GraphBuilder, Context, SceneColorTexture);
+					const FScreenPassTexture Input = FScreenPassTexture::Create(SceneColorTexture, View.ViewRect);
+
+					const FScreenPassTexture Output = ComputeSubsurface(GraphBuilder, Context, Input);
+
+					SceneColorTexture = Output.GetRDGTexture();
 				}
 			}
 
diff --git a/Engine/Source/Runtime/Renderer/Private/PostProcess/PostProcessSubsurface.cpp b/Engine/Source/Runtime/Renderer/Private/PostProcess/PostProcessSubsurface.cpp
--- a/Engine/Source/Runtime/Renderer/Private/PostProcess/PostProcessSubsurface.cpp
+++ b/Engine/Source/Runtime/Renderer/Private/PostProcess/PostProcessSubsurface.cpp
@@ -138,21 +138,6 @@ FTextureRHIRef GetSubsurfaceProfileTexture(FRHICommandListImmediate& RHICmdList)
 	return ProfileTextureTarget->GetRenderTargetItem().ShaderResourceTexture;
 }
 
-// Returns a half-scaled size rounded to an even multiple of two (but clamped to 1x1 minimum).
-FIntPoint GetHalfSize(FIntPoint Size)
-{
-	Size = FIntPoint::DivideAndRoundUp(Size, 2);
-	Size.X = FMath::Max(1, Size.X);
-	Size.Y = FMath::Max(1, Size.Y);
-	return Size;
-}
-
-// Returns a half-scaled rect, with the max rounded to the nearest multiple of two.
-FIntRect GetHalfRect(FIntRect Rect)
-{
-	return FIntRect(Rect.Min / 2, GetHalfSize(Rect.Max));
-}
-
 // Set of common shader parameters shared by all subsurface shaders.
 BEGIN_SHADER_PARAMETER_STRUCT(FSubsurfaceParameters, )
 	SHADER_PARAMETER(FVector4, SubsurfaceParams)
@@ -173,6 +158,20 @@ FSubsurfaceParameters GetSubsurfaceCommonParameters(FRHICommandListImmediate& RH
 	return Parameters;
 }
 
+// A shader parameter struct for a single subsurface input texture.
+BEGIN_SHADER_PARAMETER_STRUCT(FSubsurfaceInput, )
+	SHADER_PARAMETER_STRUCT_INCLUDE(FScreenPassTextureViewportParameters, Viewport)
+	SHADER_PARAMETER_RDG_TEXTURE(Texture2D, Texture)
+END_SHADER_PARAMETER_STRUCT()
+
+FSubsurfaceInput GetSubsurfaceInput(const FScreenPassTexture& Texture)
+{
+	FSubsurfaceInput Input;
+	Input.Texture = Texture.GetRDGTexture();
+	Input.Viewport = Texture.GetViewportParameters();
+	return Input;
+}
+
 // Base class for a subsurface shader.
 class FSubsurfaceShader : public FGlobalShader
 {
@@ -204,8 +203,9 @@ public:
 
 	BEGIN_SHADER_PARAMETER_STRUCT(FParameters, )
 		SHADER_PARAMETER_STRUCT_INCLUDE(FSubsurfaceParameters, Subsurface)
-		SHADER_PARAMETER_STRUCT(FScreenPassInput, SubsurfaceInput0)
+		SHADER_PARAMETER_STRUCT(FSubsurfaceInput, SubsurfaceInput0)
 		SHADER_PARAMETER_TEXTURE(Texture2D, MiniFontTexture)
+		SHADER_PARAMETER_SAMPLER(SamplerState, SubsurfaceSampler)
 		RENDER_TARGET_BINDING_SLOTS()
 	END_SHADER_PARAMETER_STRUCT()
 };
@@ -221,7 +221,8 @@ public:
 
 	BEGIN_SHADER_PARAMETER_STRUCT(FParameters, )
 		SHADER_PARAMETER_STRUCT_INCLUDE(FSubsurfaceParameters, Subsurface)
-		SHADER_PARAMETER_STRUCT(FScreenPassInput, SubsurfaceInput0)
+		SHADER_PARAMETER_STRUCT(FSubsurfaceInput, SubsurfaceInput0)
+		SHADER_PARAMETER_SAMPLER(SamplerState, SubsurfaceSampler)
 		RENDER_TARGET_BINDING_SLOTS()
 	END_SHADER_PARAMETER_STRUCT()
 
@@ -241,7 +242,8 @@ public:
 
 	BEGIN_SHADER_PARAMETER_STRUCT(FParameters, )
 		SHADER_PARAMETER_STRUCT_INCLUDE(FSubsurfaceParameters, Subsurface)
-		SHADER_PARAMETER_STRUCT(FScreenPassInput, SubsurfaceInput0)
+		SHADER_PARAMETER_STRUCT(FSubsurfaceInput, SubsurfaceInput0)
+		SHADER_PARAMETER_SAMPLER(SamplerState, SubsurfaceSampler)
 		RENDER_TARGET_BINDING_SLOTS()
 	END_SHADER_PARAMETER_STRUCT()
 
@@ -264,8 +266,7 @@ public:
 
 	class FDimensionDirection : SHADER_PERMUTATION_ENUM_CLASS("SUBSURFACE_DIRECTION", EDirection);
 	class FDimensionQuality : SHADER_PERMUTATION_ENUM_CLASS("SUBSURFACE_QUALITY", EQuality);
-	class FDimensionManuallyClampUV : SHADER_PERMUTATION_BOOL("SUBSURFACE_MANUALLY_CLAMP_UV");
-	using FPermutationDomain = TShaderPermutationDomain<FDimensionDirection, FDimensionQuality, FDimensionManuallyClampUV>;
+	using FPermutationDomain = TShaderPermutationDomain<FDimensionDirection, FDimensionQuality>;
 
 	// Returns the sampler state based on the requested SSS filter CVar setting.
 	static FSamplerStateRHIParamRef GetSamplerState()
@@ -301,8 +302,9 @@ class FSubsurfaceRecombinePS : public FSubsurfaceShader
 
 	BEGIN_SHADER_PARAMETER_STRUCT(FParameters, )
 		SHADER_PARAMETER_STRUCT_INCLUDE(FSubsurfaceParameters, Subsurface)
-		SHADER_PARAMETER_STRUCT(FScreenPassInput, SubsurfaceInput0)
-		SHADER_PARAMETER_STRUCT(FScreenPassInput, SubsurfaceInput1)
+		SHADER_PARAMETER_STRUCT(FSubsurfaceInput, SubsurfaceInput0)
+		SHADER_PARAMETER_STRUCT(FSubsurfaceInput, SubsurfaceInput1)
+		SHADER_PARAMETER_SAMPLER(SamplerState, SubsurfaceSampler)
 		RENDER_TARGET_BINDING_SLOTS()
 	END_SHADER_PARAMETER_STRUCT();
 
@@ -349,7 +351,8 @@ class FSubsurfaceViewportCopyPS : public FSubsurfaceShader
 	SHADER_USE_PARAMETER_STRUCT(FSubsurfaceViewportCopyPS, FSubsurfaceShader);
 
 	BEGIN_SHADER_PARAMETER_STRUCT(FParameters, )
-		SHADER_PARAMETER_STRUCT(FScreenPassInput, SubsurfaceInput0)
+		SHADER_PARAMETER_STRUCT(FSubsurfaceInput, SubsurfaceInput0)
+		SHADER_PARAMETER_SAMPLER(SamplerState, SubsurfaceSampler)
 		RENDER_TARGET_BINDING_SLOTS()
 	END_SHADER_PARAMETER_STRUCT();
 
@@ -401,103 +404,108 @@ bool IsSubsurfaceCheckerboardFormat(EPixelFormat SceneColorFormat)
 	return true;
 }
 
-FRDGTextureRef ComputeSubsurface(
+FScreenPassTexture ComputeSubsurface(
 	FRDGBuilder& GraphBuilder,
 	FScreenPassContextRef Context,
-	FRDGTextureRef SceneTexture)
+	const FScreenPassTexture& SceneTexture)
 {
 	check(Context);
-	check(SceneTexture);
+	check(SceneTexture.IsValid());
+
+	const FRDGTextureDesc& SceneTextureDesc = SceneTexture.GetRDGTexture()->Desc;
 
-	const FPooledRenderTargetDesc& SceneTextureDesc = SceneTexture->Desc;
 	const ESubsurfaceMode SubsurfaceMode = GetSubsurfaceModeForView(Context->View);
-	const bool bHalfRes = (SubsurfaceMode == ESubsurfaceMode::HalfRes);
 
-	// Viewport rect mapped onto the scene texture. Not necessarily a full screen mapping (e.g. VR).
-	const FIntRect ViewportRectFinal = Context->ViewportRect;
+	const bool bHalfRes = (SubsurfaceMode == ESubsurfaceMode::HalfRes);
 
-	// Viewport for intermediate passes, which may be half resolution depending on the pass settings.
-	const FIntRect ViewportRectIntermediate = bHalfRes ? GetHalfRect(ViewportRectFinal) : ViewportRectFinal;
+	const bool bCheckerboard = IsSubsurfaceCheckerboardFormat(SceneTextureDesc.Format);
 
-	// Size of the final scene texture.
-	const FIntPoint TextureSizeFinal = SceneTextureDesc.Extent;
+	/**
+	 * All intermediate passes within the screen-space subsurface effect can operate at half or full resolution,
+	 * depending on the subsurface mode. The values are precomputed and shared among all intermediate textures.
+	 */
+	FScreenPassTextureViewport IntermediateViewport = SceneTexture.GetViewport();
 
-	// Size of intermediate textures, which may be half resolution depending on pass settings.
-	const FIntPoint TextureSizeIntermediate = bHalfRes ? GetHalfSize(TextureSizeFinal) : TextureSizeFinal;
+	if (bHalfRes)
+	{
+		IntermediateViewport = FScreenPassTextureViewport::CreateDownscaled(IntermediateViewport, 2);
+	}
 
-	// Description shared by all intermediate pass textures.
-	const FRDGTextureDesc TextureDescIntermediate = FRDGTextureDesc::Create2DDesc(
-		TextureSizeIntermediate,
+	const FRDGTextureDesc IntermediateTextureDescriptor = FRDGTextureDesc::Create2DDesc(
+		IntermediateViewport.Extent,
 		PF_FloatRGBA,
 		FClearValueBinding(),
 		TexCreate_None,
 		TexCreate_RenderTargetable | TexCreate_ShaderResource,
 		false);
-
-	// Common subsurface parameters shared by all passes.
+	
 	const FSubsurfaceParameters SubsurfaceCommonParameters = GetSubsurfaceCommonParameters(GraphBuilder.RHICmdList, Context);
 
-	// Texture handles used by subsurface passes.
-	FRDGTextureRef SetupTexture = SceneTexture;
-	FRDGTextureRef SubsurfaceTextureX = nullptr;
-	FRDGTextureRef SubsurfaceTextureY = nullptr;
-	FRDGTextureRef RecombineTexture = nullptr;
+	FScreenPassTexture SetupTexture = SceneTexture;
+	FScreenPassTexture SubsurfaceTextureX;
+	FScreenPassTexture SubsurfaceTextureY;
+	FScreenPassTexture RecombineTexture;
 
-	// Sampler handles used by subsurface passes.
 	FSamplerStateRHIParamRef PointClampSampler = TStaticSamplerState<SF_Point, AM_Clamp, AM_Clamp, AM_Clamp>::GetRHI();
-	FSamplerStateRHIParamRef BilinearBorderState = TStaticSamplerState<SF_Bilinear, AM_Border, AM_Border, AM_Border>::GetRHI();
-
-	// Whether to use checkerboard rendering for subsurface passes (predicated on CVar and format).
-	const bool bCheckerboard = IsSubsurfaceCheckerboardFormat(SceneTextureDesc.Format);
+	FSamplerStateRHIParamRef BilinearBorderSampler = TStaticSamplerState<SF_Bilinear, AM_Border, AM_Border, AM_Border>::GetRHI();
 
-	// When in bypass mode the setup and convolution passes are skipped, but lighting
-	// reconstruction is still performed in the recombine pass.
+	/**
+	 * When in bypass mode, the setup and convolution passes are skipped, but lighting
+	 * reconstruction is still performed in the recombine pass.
+	 */
 	if (SubsurfaceMode != ESubsurfaceMode::Bypass)
 	{
-		SetupTexture = GraphBuilder.CreateTexture(TextureDescIntermediate, TEXT("SubsurfaceSetupTexture"));
+		SetupTexture = FScreenPassTexture::Create(
+			GraphBuilder.CreateTexture(IntermediateTextureDescriptor, TEXT("SubsurfaceSetupTexture")),
+			IntermediateViewport.Rect);
 
 		// Setup pass outputs the diffuse scene color and depth in preparation for the scatter passes.
 		{
+			const FScreenPassTexture& TextureInput = SceneTexture;
+			const FScreenPassTexture& TextureOutput = SetupTexture;
+
 			FSubsurfaceSetupPS::FParameters* PassParameters = GraphBuilder.AllocParameters<FSubsurfaceSetupPS::FParameters>();
 			PassParameters->Subsurface = SubsurfaceCommonParameters;
-			PassParameters->RenderTargets[0] = FRenderTargetBinding(SetupTexture, ERenderTargetLoadAction::ENoAction, ERenderTargetStoreAction::EStore);
-			PassParameters->SubsurfaceInput0 = GetScreenPassInputParameters(SceneTexture, PointClampSampler);
-
-			// Reading from the final target; writing to intermediate target.
-			const FIntRect ViewportRect = ViewportRectIntermediate;
-			const FIntRect TextureRect = ViewportRectFinal;
-			const FIntPoint TextureSize = TextureSizeFinal;
+			PassParameters->RenderTargets[0] = FRenderTargetBinding(TextureOutput.GetRDGTexture(), ERenderTargetLoadAction::ENoAction, ERenderTargetStoreAction::EStore);
+			PassParameters->SubsurfaceInput0 = GetSubsurfaceInput(TextureInput);
+			PassParameters->SubsurfaceSampler = PointClampSampler;
 
 			FSubsurfaceSetupPS::FPermutationDomain PixelShaderPermutationVector;
 			PixelShaderPermutationVector.Set<FSubsurfaceSetupPS::FDimensionHalfRes>(bHalfRes);
 			PixelShaderPermutationVector.Set<FSubsurfaceSetupPS::FDimensionCheckerboard>(bCheckerboard);
 			TShaderMapRef<FSubsurfaceSetupPS> PixelShader(Context->ShaderMap, PixelShaderPermutationVector);
 
-			AddDrawScreenPass(GraphBuilder, RDG_EVENT_NAME("SubsurfaceSetup"), Context, ViewportRect, TextureRect, TextureSize, *PixelShader, PassParameters);
+			AddDrawScreenPass(GraphBuilder, RDG_EVENT_NAME("SubsurfaceSetup"), Context, TextureOutput.GetViewport(), TextureInput.GetViewport(), *PixelShader, PassParameters);
 		}
 
-		SubsurfaceTextureX = GraphBuilder.CreateTexture(TextureDescIntermediate, TEXT("SubsurfaceTextureX"));
-		SubsurfaceTextureY = GraphBuilder.CreateTexture(TextureDescIntermediate, TEXT("SubsurfaceTextureY"));
+		SubsurfaceTextureX = FScreenPassTexture::Create(
+			GraphBuilder.CreateTexture(IntermediateTextureDescriptor, TEXT("SubsurfaceTextureX")),
+			IntermediateViewport.Rect);
+
+		SubsurfaceTextureY = FScreenPassTexture::Create(
+			GraphBuilder.CreateTexture(IntermediateTextureDescriptor, TEXT("SubsurfaceTextureY")),
+			IntermediateViewport.Rect);
 
 		FSamplerStateRHIParamRef SubsurfaceSamplerState = FSubsurfacePS::GetSamplerState();
 		const FSubsurfacePS::EQuality SubsurfaceQuality = FSubsurfacePS::GetQuality();
 
-		static const FRDGEventName SubsurfacePassNameByDirection[] =
-		{
-			RDG_EVENT_NAME("SubsurfaceX"),
-			RDG_EVENT_NAME("SubsurfaceY")
-		};
-
-		FRDGTextureRef SubsurfacePassOutputByDirection[] =
+		struct FSubsurfacePassInfo
 		{
-			SubsurfaceTextureX,
-			SubsurfaceTextureY
+			FSubsurfacePassInfo(const TCHAR* InName, const FScreenPassTexture& InInput, const FScreenPassTexture& InOutput)
+				: Name(InName)
+				, Input(InInput)
+				, Output(InOutput)
+			{}
+
+			const TCHAR* Name;
+			const FScreenPassTexture& Input;
+			const FScreenPassTexture& Output;
 		};
 
-		FRDGTextureRef SubsurfacePassInputByDirection[] =
+		const FSubsurfacePassInfo SubsurfacePassInfoByDirection[] =
 		{
-			SetupTexture,
-			SubsurfaceTextureX
+			{ TEXT("SubsurfaceX"), SetupTexture, SubsurfaceTextureX },
+			{ TEXT("SubsurfaceY"), SubsurfaceTextureX, SubsurfaceTextureY },
 		};
 
 		// Horizontal / Vertical scattering passes using a separable filter.
@@ -505,48 +513,50 @@ FRDGTextureRef ComputeSubsurface(
 		{
 			const auto Direction = static_cast<FSubsurfacePS::EDirection>(DirectionIndex);
 
+			const FSubsurfacePassInfo& PassInfo = SubsurfacePassInfoByDirection[DirectionIndex];
+			const FScreenPassTexture& TextureInput = PassInfo.Input;
+			const FScreenPassTexture& TextureOutput = PassInfo.Output;
+
 			FSubsurfacePS::FParameters* PassParameters = GraphBuilder.AllocParameters<FSubsurfacePS::FParameters>();
 			PassParameters->Subsurface = SubsurfaceCommonParameters;
-			PassParameters->RenderTargets[0] = FRenderTargetBinding(SubsurfacePassOutputByDirection[DirectionIndex], ERenderTargetLoadAction::ENoAction, ERenderTargetStoreAction::EStore);
-			PassParameters->SubsurfaceInput0 = GetScreenPassInputParameters(SubsurfacePassInputByDirection[DirectionIndex], SubsurfaceSamplerState);
-
-			// Read / Write between intermediate targets.
-			const FIntRect ViewportRect = ViewportRectIntermediate;
-			const FIntRect TextureRect = ViewportRectIntermediate;
-			const FIntPoint TextureSize = TextureSizeIntermediate;
-
-			// If we are sampling from a subset of the texture (e.g. stereo rendering), we have to manually clamp UVs.
-			const bool bManuallyClampUV = TextureSize != TextureRect.Size();
+			PassParameters->RenderTargets[0] = FRenderTargetBinding(TextureOutput.GetRDGTexture(), ERenderTargetLoadAction::ENoAction, ERenderTargetStoreAction::EStore);
+			PassParameters->SubsurfaceInput0 = GetSubsurfaceInput(TextureInput);
+			PassParameters->SubsurfaceSampler = SubsurfaceSamplerState;
 
 			FSubsurfacePS::FPermutationDomain PixelShaderPermutationVector;
 			PixelShaderPermutationVector.Set<FSubsurfacePS::FDimensionDirection>(Direction);
 			PixelShaderPermutationVector.Set<FSubsurfacePS::FDimensionQuality>(SubsurfaceQuality);
-			PixelShaderPermutationVector.Set<FSubsurfacePS::FDimensionManuallyClampUV>(bManuallyClampUV);
 			TShaderMapRef<FSubsurfacePS> PixelShader(Context->ShaderMap, PixelShaderPermutationVector);
 
-			AddDrawScreenPass(GraphBuilder, FRDGEventName(SubsurfacePassNameByDirection[DirectionIndex]), Context, ViewportRect, TextureRect, TextureSize, *PixelShader, PassParameters);
+			AddDrawScreenPass(GraphBuilder, FRDGEventName(PassInfo.Name), Context, TextureOutput.GetViewport(), TextureInput.GetViewport(), *PixelShader, PassParameters);
 		}
 	}
 
-	RecombineTexture = GraphBuilder.CreateTexture(SceneTextureDesc, TEXT("SubsurfaceRecombine"));
+	RecombineTexture = FScreenPassTexture::Create(
+		GraphBuilder.CreateTexture(SceneTextureDesc, TEXT("SubsurfaceRecombine")),
+		SceneTexture.GetViewport().Rect);
 
-	// If multiple views exist (e.g. VR) we need to copy all other viewports from the scene texture
+	// If multiple views exist (e.g. VR) we need to copy all other views from the scene texture
 	// to the recombine target so we don't lose them.
 	if (Context->ViewFamily.Views.Num())
 	{
-		FSubsurfaceViewportCopyPS::FParameters* PassParameters = GraphBuilder.AllocParameters<FSubsurfaceViewportCopyPS::FParameters>();
-		PassParameters->RenderTargets[0] = FRenderTargetBinding(RecombineTexture, ERenderTargetLoadAction::ENoAction, ERenderTargetStoreAction::EStore);
-		PassParameters->SubsurfaceInput0 = GetScreenPassInputParameters(SceneTexture, PointClampSampler);
+		const FScreenPassTexture& TextureInput = SceneTexture;
+		const FScreenPassTexture& TextureOutput = RecombineTexture;
 
-		const FIntPoint TextureSize = SceneTextureDesc.Extent;
+		FSubsurfaceViewportCopyPS::FParameters* PassParameters = GraphBuilder.AllocParameters<FSubsurfaceViewportCopyPS::FParameters>();
+		PassParameters->RenderTargets[0] = FRenderTargetBinding(TextureOutput.GetRDGTexture(), ERenderTargetLoadAction::ENoAction, ERenderTargetStoreAction::EStore);
+		PassParameters->SubsurfaceInput0 = GetSubsurfaceInput(TextureInput);
+		PassParameters->SubsurfaceSampler = PointClampSampler;
 
 		TShaderMapRef<FSubsurfaceViewportCopyPS> PixelShader(Context->ShaderMap);
 
+		const FIntPoint InputTextureSize = TextureInput.GetViewport().Extent;
+
 		GraphBuilder.AddPass(
 			RDG_EVENT_NAME("SubsurfaceViewportCopy"),
 			PassParameters,
 			ERenderGraphPassFlags::None,
-			[Context, PixelShader, TextureSize, PassParameters](FRHICommandListImmediate& RHICmdList)
+			[Context, PixelShader, InputTextureSize, PassParameters](FRHICommandListImmediate& RHICmdList)
 		{
 			const FSceneViewFamily& ViewFamily = Context->ViewFamily;
 
@@ -559,7 +569,7 @@ FRDGTextureRef ComputeSubsurface(
 				{
 					const FIntRect Rect = LocalView->ViewRect;
 
-					DrawScreenPass(RHICmdList, Context, Rect, Rect, TextureSize, *PixelShader, *PassParameters);
+					DrawScreenPass(RHICmdList, Context, Rect, Rect, InputTextureSize, *PixelShader, *PassParameters);
 				}
 			}
 		});
@@ -567,22 +577,21 @@ FRDGTextureRef ComputeSubsurface(
 
 	// Recombines scattering result with scene color.
 	{
+		const FScreenPassTexture& TextureInput = SceneTexture;
+		const FScreenPassTexture& TextureOutput = RecombineTexture;
+
 		FSubsurfaceRecombinePS::FParameters* PassParameters = GraphBuilder.AllocParameters<FSubsurfaceRecombinePS::FParameters>();
 		PassParameters->Subsurface = SubsurfaceCommonParameters;
-		PassParameters->RenderTargets[0] = FRenderTargetBinding(RecombineTexture, ERenderTargetLoadAction::ELoad, ERenderTargetStoreAction::EStore);
-		PassParameters->SubsurfaceInput0 = GetScreenPassInputParameters(SceneTexture, BilinearBorderState);
+		PassParameters->RenderTargets[0] = FRenderTargetBinding(TextureOutput.GetRDGTexture(), ERenderTargetLoadAction::ELoad, ERenderTargetStoreAction::EStore);
+		PassParameters->SubsurfaceInput0 = GetSubsurfaceInput(TextureInput);
+		PassParameters->SubsurfaceSampler = BilinearBorderSampler;
 
 		// Scattering output target is only used when scattering is enabled.
 		if (SubsurfaceMode != ESubsurfaceMode::Bypass)
 		{
-			PassParameters->SubsurfaceInput1 = GetScreenPassInputParameters(SubsurfaceTextureY, BilinearBorderState);
+			PassParameters->SubsurfaceInput1 = GetSubsurfaceInput(SubsurfaceTextureY);
 		}
 
-		// Read from intermediate target; write to final target.
-		const FIntRect ViewportRect = ViewportRectFinal;
-		const FIntRect TextureRect = ViewportRectIntermediate;
-		const FIntPoint TextureSize = TextureSizeIntermediate;
-
 		const FSubsurfaceRecombinePS::EQuality RecombineQuality = FSubsurfaceRecombinePS::GetQuality(Context->View);
 
 		FSubsurfaceRecombinePS::FPermutationDomain PixelShaderPermutationVector;
@@ -591,48 +600,46 @@ FRDGTextureRef ComputeSubsurface(
 		PixelShaderPermutationVector.Set<FSubsurfaceRecombinePS::FDimensionCheckerboard>(bCheckerboard);
 		TShaderMapRef<FSubsurfaceRecombinePS> PixelShader(Context->ShaderMap, PixelShaderPermutationVector);
 
-		AddDrawScreenPass(GraphBuilder, RDG_EVENT_NAME("SubsurfaceRecombine"), Context, ViewportRect, TextureRect, TextureSize, *PixelShader, PassParameters);
+		AddDrawScreenPass(GraphBuilder, RDG_EVENT_NAME("SubsurfaceRecombine"), Context, TextureOutput.GetViewport(), TextureInput.GetViewport(), *PixelShader, PassParameters);
 	}
 
 	return RecombineTexture;
 }
 
-FRDGTextureRef VisualizeSubsurface(
+FScreenPassTexture VisualizeSubsurface(
 	FRDGBuilder& GraphBuilder,
 	FScreenPassContextRef Context,
-	FRDGTextureRef SceneTexture)
+	const FScreenPassTexture& SceneTexture)
 {
 	check(Context);
-	check(SceneTexture);
-
-	FSamplerStateRHIParamRef PointClampState = TStaticSamplerState<SF_Point, AM_Clamp, AM_Clamp, AM_Clamp>::GetRHI();
+	check(SceneTexture.IsValid());
 
-	FRDGTextureRef VisualizeTexture = GraphBuilder.CreateTexture(SceneTexture->Desc, TEXT("SubsurfaceVisualize"));
+	const FScreenPassTexture& TextureInput = SceneTexture;
+	const FScreenPassTexture TextureOutput = FScreenPassTexture::Create(
+		GraphBuilder.CreateTexture(TextureInput.GetRDGTexture()->Desc, TEXT("SubsurfaceVisualize")),
+		TextureInput.GetViewport().Rect);
 
 	FSubsurfaceVisualizePS::FParameters* PassParameters = GraphBuilder.AllocParameters<FSubsurfaceVisualizePS::FParameters>();
 	PassParameters->Subsurface = GetSubsurfaceCommonParameters(GraphBuilder.RHICmdList, Context);
-	PassParameters->RenderTargets[0] = FRenderTargetBinding(VisualizeTexture, ERenderTargetLoadAction::EClear, ERenderTargetStoreAction::EStore);
-	PassParameters->SubsurfaceInput0 = GetScreenPassInputParameters(SceneTexture, PointClampState);
+	PassParameters->RenderTargets[0] = FRenderTargetBinding(TextureOutput.GetRDGTexture(), ERenderTargetLoadAction::EClear, ERenderTargetStoreAction::EStore);
+	PassParameters->SubsurfaceInput0 = GetSubsurfaceInput(TextureInput);
+	PassParameters->SubsurfaceSampler = TStaticSamplerState<SF_Point, AM_Clamp, AM_Clamp, AM_Clamp>::GetRHI();
 	PassParameters->MiniFontTexture = GetMiniFontTexture();
 
+	TShaderMapRef<FSubsurfaceVisualizePS> PixelShader(Context->ShaderMap);
+
 	GraphBuilder.AddPass(
 		RDG_EVENT_NAME("SubsurfaceVisualize"),
 		PassParameters,
 		ERenderGraphPassFlags::None,
-		[Context, SceneTexture, VisualizeTexture, PassParameters](FRHICommandListImmediate& RHICmdList)
+		[Context, TextureInput, TextureOutput, PixelShader, PassParameters](FRHICommandListImmediate& RHICmdList)
 	{
-		const FIntRect ViewportRect = Context->ViewportRect;
-		const FIntRect TextureRect = ViewportRect;
-		const FIntPoint TextureSize = SceneTexture->Desc.Extent;
-
-		TShaderMapRef<FSubsurfaceVisualizePS> PixelShader(Context->ShaderMap);
-
-		DrawScreenPass(RHICmdList, Context, ViewportRect, TextureRect, TextureSize, *PixelShader, *PassParameters);
+		DrawScreenPass(RHICmdList, Context, TextureOutput.GetViewport(), TextureInput.GetViewport(), *PixelShader, *PassParameters);
 
 		// Draw debug text
 		{
 			const FSceneViewFamily& ViewFamily = Context->ViewFamily;
-			FRenderTargetTemp TempRenderTarget(static_cast<FTexture2DRHIParamRef>(VisualizeTexture->GetRHITexture()), TextureSize);
+			FRenderTargetTemp TempRenderTarget(static_cast<FTexture2DRHIParamRef>(TextureOutput.GetRHITexture()), TextureOutput.GetViewport().Extent);
 			FCanvas Canvas(&TempRenderTarget, nullptr, ViewFamily.CurrentRealTime, ViewFamily.CurrentWorldTime, ViewFamily.DeltaWorldTime, Context->View.GetFeatureLevel());
 
 			float X = 30;
@@ -656,7 +663,7 @@ FRDGTextureRef VisualizeSubsurface(
 		}
 	});
 
-	return VisualizeTexture;
+	return TextureOutput;
 }
 
 FSubsurfaceVisualizeCompositePass::FSubsurfaceVisualizeCompositePass(FRHICommandList& RHICmdList)
@@ -669,18 +676,22 @@ void FSubsurfaceVisualizeCompositePass::Process(FRenderingCompositePassContext&
 {
 	FRHICommandListImmediate& RHICmdList = CompositePassContext.RHICmdList;
 
+	const FViewInfo& View = CompositePassContext.View;
+
 	FRDGBuilder GraphBuilder(RHICmdList);
 
 	FScreenPassContext* Context = FScreenPassContext::Create(RHICmdList, CompositePassContext.View);
 
-	FRDGTextureRef InputTexture = CreateRDGTextureForInput(GraphBuilder, ePId_Input0, TEXT("SceneColor"), eFC_0000);
+	const FScreenPassTexture TextureInput = FScreenPassTexture::Create(
+		CreateRDGTextureForInput(GraphBuilder, ePId_Input0, TEXT("SceneColor"), eFC_0000),
+		View.ViewRect);
 
-	FRDGTextureRef OutputTexture = VisualizeSubsurface(GraphBuilder, Context, InputTexture);
+	const FScreenPassTexture TextureOutput = VisualizeSubsurface(GraphBuilder, Context, TextureInput);
 
-	ExtractRDGTextureForOutput(GraphBuilder, ePId_Output0, OutputTexture);
+	ExtractRDGTextureForOutput(GraphBuilder, ePId_Output0, TextureOutput.GetRDGTexture());
 
 	TRefCountPtr<IPooledRenderTarget> OutputTarget;
-	GraphBuilder.QueueTextureExtraction(OutputTexture, &OutputTarget);
+	GraphBuilder.QueueTextureExtraction(TextureOutput.GetRDGTexture(), &OutputTarget);
 
 	GraphBuilder.Execute();
 
diff --git a/Engine/Source/Runtime/Renderer/Private/PostProcess/PostProcessSubsurface.h b/Engine/Source/Runtime/Renderer/Private/PostProcess/PostProcessSubsurface.h
--- a/Engine/Source/Runtime/Renderer/Private/PostProcess/PostProcessSubsurface.h
+++ b/Engine/Source/Runtime/Renderer/Private/PostProcess/PostProcessSubsurface.h
@@ -19,17 +19,17 @@ bool IsSubsurfaceRequiredForView(const FViewInfo& View);
 bool IsSubsurfaceCheckerboardFormat(EPixelFormat SceneColorFormat);
 
 // Computes subsurface scattering on the scene texture and produces a new scene texture as output.
-FRDGTextureRef ComputeSubsurface(
+FScreenPassTexture ComputeSubsurface(
 	FRDGBuilder& GraphBuilder,
 	FScreenPassContextRef Context,
-	FRDGTextureRef SceneTexture);
+	const FScreenPassTexture& SceneTexture);
 
 // Visualizes subsurface scattering profiles by overlaying an image on the provided scene texture.
 // Produces a new scene texture as output.
-FRDGTextureRef VisualizeSubsurface(
+FScreenPassTexture VisualizeSubsurface(
 	FRDGBuilder& GraphBuilder,
 	FScreenPassContextRef Context,
-	FRDGTextureRef SceneTexture);
+	const FScreenPassTexture& SceneTexture);
 
 // An adapter to connect the new Render Graph implementation to the legacy Composition Graph.
 class FSubsurfaceVisualizeCompositePass : public TRenderingCompositePassBase<1, 1>
diff --git a/Engine/Source/Runtime/Renderer/Private/PostProcess/RenderingCompositionGraph.cpp b/Engine/Source/Runtime/Renderer/Private/PostProcess/RenderingCompositionGraph.cpp
--- a/Engine/Source/Runtime/Renderer/Private/PostProcess/RenderingCompositionGraph.cpp
+++ b/Engine/Source/Runtime/Renderer/Private/PostProcess/RenderingCompositionGraph.cpp
@@ -836,8 +836,7 @@ FRenderingCompositeOutput *FRenderingCompositeOutputRef::GetOutput() const
 
 void FPostProcessPassParameters::Bind(const FShaderParameterMap& ParameterMap)
 {
-	BilinearTextureSampler0.Bind(ParameterMap,TEXT("BilinearTextureSampler0"));
-	BilinearTextureSampler1.Bind(ParameterMap,TEXT("BilinearTextureSampler1"));
+	BilinearTextureSampler.Bind(ParameterMap,TEXT("BilinearTextureSampler"));
 	ViewportSize.Bind(ParameterMap,TEXT("ViewportSize"));
 	ViewportRect.Bind(ParameterMap,TEXT("ViewportRect"));
 	ScreenPosToPixel.Bind(ParameterMap, TEXT("ScreenPosToPixel"));
@@ -909,20 +908,11 @@ void FPostProcessPassParameters::Set(
 	// but not both
 	check(!FilterOverrideArray || !Filter);
 
-	if(BilinearTextureSampler0.IsBound())
+	if(BilinearTextureSampler.IsBound())
 	{
 		RHICmdList.SetShaderSampler(
 			ShaderRHI, 
-			BilinearTextureSampler0.GetBaseIndex(), 
-			TStaticSamplerState<SF_Bilinear>::GetRHI()
-			);
-	}
-
-	if(BilinearTextureSampler1.IsBound())
-	{
-		RHICmdList.SetShaderSampler(
-			ShaderRHI, 
-			BilinearTextureSampler1.GetBaseIndex(), 
+			BilinearTextureSampler.GetBaseIndex(), 
 			TStaticSamplerState<SF_Bilinear>::GetRHI()
 			);
 	}
@@ -1058,7 +1048,7 @@ IMPLEMENT_POST_PROCESS_PARAM_SET( FComputeShaderRHIParamRef, FRHIAsyncComputeCom
 
 FArchive& operator<<(FArchive& Ar, FPostProcessPassParameters& P)
 {
-	Ar << P.BilinearTextureSampler0 << P.BilinearTextureSampler1 << P.ViewportSize << P.ScreenPosToPixel << P.SceneColorBufferUVViewport << P.ViewportRect;
+	Ar << P.BilinearTextureSampler << P.ViewportSize << P.ScreenPosToPixel << P.SceneColorBufferUVViewport << P.ViewportRect;
 
 	for(uint32 i = 0; i < ePId_Input_MAX; ++i)
 	{
diff --git a/Engine/Source/Runtime/Renderer/Private/ScreenPass.cpp b/Engine/Source/Runtime/Renderer/Private/ScreenPass.cpp
--- a/Engine/Source/Runtime/Renderer/Private/ScreenPass.cpp
+++ b/Engine/Source/Runtime/Renderer/Private/ScreenPass.cpp
@@ -7,6 +7,8 @@
 #include "IXRTrackingSystem.h"
 #include "IHeadMountedDisplay.h"
 
+extern bool ShouldDoComputePostProcessing(const FViewInfo& View);
+
 IMPLEMENT_GLOBAL_SHADER(FScreenPassVS, "/Engine/Private/ScreenPass.usf", "ScreenPassVS", SF_Vertex);
 
 const FTextureRHIRef& GetMiniFontTexture()
@@ -26,56 +28,139 @@ bool IsHMDHiddenAreaMaskActive()
 	// Query if we have a custom HMD post process mesh to use
 	static const auto* const HiddenAreaMaskCVar = IConsoleManager::Get().FindTConsoleVariableDataInt(TEXT("vr.HiddenAreaMask"));
 
-	return (HiddenAreaMaskCVar != nullptr &&
+	return
+		HiddenAreaMaskCVar != nullptr &&
 		HiddenAreaMaskCVar->GetValueOnRenderThread() == 1 &&
 		GEngine &&
 		GEngine->XRSystem.IsValid() && GEngine->XRSystem->GetHMDDevice() &&
-		GEngine->XRSystem->GetHMDDevice()->HasVisibleAreaMesh());
+		GEngine->XRSystem->GetHMDDevice()->HasVisibleAreaMesh();
 }
 
-FScreenPassCommonParameters GetScreenPassCommonParameters(FRHICommandListImmediate& RHICmdList, const FViewInfo& View)
+FScreenPassTextureViewport FScreenPassTextureViewport::CreateDownscaled(const FScreenPassTextureViewport& Other, uint32 ScaleFactor)
 {
-	FSceneRenderTargets& SceneContext = FSceneRenderTargets::Get(RHICmdList);
+	const auto GetDownscaledSize = [](FIntPoint Size, uint32 ScaleFactor)
+	{
+		Size = FIntPoint::DivideAndRoundUp(Size, ScaleFactor);
+		Size.X = FMath::Max(1, Size.X);
+		Size.Y = FMath::Max(1, Size.Y);
+		return Size;
+	};
+
+	FScreenPassTextureViewport Viewport;
+	Viewport.Rect.Min = Other.Rect.Min / ScaleFactor;
+	Viewport.Rect.Max = GetDownscaledSize(Other.Rect.Max, ScaleFactor);
+	Viewport.Extent = GetDownscaledSize(Other.Extent, ScaleFactor);
+	return Viewport;
+}
+
+bool FScreenPassTextureViewport::operator==(const FScreenPassTextureViewport& Other) const
+{
+	return Rect == Other.Rect && Extent == Other.Extent;
+}
 
-	const FIntRect ViewportRect = View.ViewRect;
-	const FIntPoint ViewportOffset = ViewportRect.Min;
-	const FIntPoint ViewportExtent = ViewportRect.Size();
+bool FScreenPassTextureViewport::operator!=(const FScreenPassTextureViewport& Other) const
+{
+	return Rect != Other.Rect || Extent != Other.Extent;
+}
 
-	FScreenPassCommonParameters Parameters;
-	Parameters.ViewportRect = ViewportRect;
-	Parameters.ViewportSize = FVector4(
-		ViewportExtent.X,
-		ViewportExtent.Y,
-		1.0f / ViewportExtent.X,
-		1.0f / ViewportExtent.Y);
+bool FScreenPassTextureViewport::IsEmpty() const
+{
+	return Rect.IsEmpty() || Extent == FIntPoint::ZeroValue;
+}
 
-	Parameters.ScreenPosToPixelValue = FVector4(
-		ViewportExtent.X * 0.5f,
-		-ViewportExtent.Y * 0.5f,
-		ViewportExtent.X * 0.5f - 0.5f + ViewportOffset.X,
-		ViewportExtent.Y * 0.5f - 0.5f + ViewportOffset.Y);
+FScreenPassTextureViewportParameters GetScreenPassTextureViewportParameters(const FScreenPassTextureViewport& InViewport)
+{
+	const FVector2D Extent(InViewport.Extent);
+	const FVector2D ViewportMin(InViewport.Rect.Min.X, InViewport.Rect.Min.Y);
+	const FVector2D ViewportMax(InViewport.Rect.Max.X, InViewport.Rect.Max.Y);
+	const FVector2D ViewportSize = ViewportMax - ViewportMin;
 
-	Parameters.BilinearTextureSampler0 = TStaticSamplerState<SF_Bilinear>::GetRHI();
+	FScreenPassTextureViewportParameters Parameters;
 
-	Parameters.ViewUniformBuffer = View.ViewUniformBuffer;
-	Parameters.SceneUniformBuffer = CreateSceneTextureUniformBuffer(
-		SceneContext, View.FeatureLevel, ESceneTextureSetupMode::All, EUniformBufferUsage::UniformBuffer_SingleFrame);
+	if (!InViewport.IsEmpty())
+	{
+		Parameters.Extent = Extent;
+		Parameters.ExtentInverse = FVector2D(1.0f / Extent.X, 1.0f / Extent.Y);
+
+		Parameters.ScreenPosToViewportScale = FVector2D(0.5f, -0.5f) * ViewportSize;
+		Parameters.ScreenPosToViewportBias = (0.5f * ViewportSize) + ViewportMin;
+
+		Parameters.ViewportMin = InViewport.Rect.Min;
+		Parameters.ViewportMax = InViewport.Rect.Max;
+
+		Parameters.ViewportSize = ViewportSize;
+		Parameters.ViewportSizeInverse = FVector2D(1.0f / Parameters.ViewportSize.X, 1.0f / Parameters.ViewportSize.Y);
+
+		Parameters.UVViewportMin = ViewportMin * Parameters.ExtentInverse;
+		Parameters.UVViewportMax = ViewportMax * Parameters.ExtentInverse;
+
+		Parameters.UVViewportSize = Parameters.UVViewportMax - Parameters.UVViewportMin;
+		Parameters.UVViewportSizeInverse = FVector2D(1.0f / Parameters.UVViewportSize.X, 1.0f / Parameters.UVViewportSize.Y);
+
+		Parameters.UVViewportBilinearMin = Parameters.UVViewportMin + 0.5f * Parameters.ExtentInverse;
+		Parameters.UVViewportBilinearMax = Parameters.UVViewportMax - 0.5f * Parameters.ExtentInverse;
+	}
 
 	return Parameters;
 }
 
-FScreenPassInput GetScreenPassInputParameters(FRDGTextureRef Texture, FSamplerStateRHIParamRef SamplerState)
+FScreenPassTexture FScreenPassTexture::Create(FRDGTextureRef InTexture, FIntRect InViewport)
+{
+	check(InTexture);
+
+	FScreenPassTexture Texture;
+	Texture.Texture = InTexture;
+	Texture.Viewport = FScreenPassTextureViewport(InViewport, InTexture->Desc.Extent);
+	return Texture;
+}
+
+FScreenPassTexture FScreenPassTexture::CreateFullscreen(FRDGTextureRef InTexture)
+{
+	check(InTexture);
+
+	const FIntPoint Extent = InTexture->Desc.Extent;
+	const FIntRect PixelViewRect = FIntRect(FIntPoint::ZeroValue, Extent);
+
+	FScreenPassTexture Texture;
+	Texture.Texture = InTexture;
+	Texture.Viewport = FScreenPassTextureViewport(PixelViewRect, Extent);
+	return Texture;
+}
+
+FScreenPassTextureViewportTransform GetScreenPassTextureViewportTransform(
+	FVector2D SourceOffset,
+	FVector2D SourceExtent,
+	FVector2D DestinationOffset,
+	FVector2D DestinationExtent)
 {
-	check(Texture);
-	check(SamplerState);
+	FScreenPassTextureViewportTransform Transform;
+	Transform.Scale = SourceExtent / DestinationExtent;
+	Transform.Bias = SourceOffset - Transform.Scale * DestinationOffset;
+	return Transform;
+}
 
-	const FVector2D Size(Texture->Desc.Extent.X, Texture->Desc.Extent.Y);
+FScreenPassTextureViewportTransform GetScreenPassTextureViewportTransform(
+	const FScreenPassTextureViewportParameters& Source,
+	const FScreenPassTextureViewportParameters& Destination)
+{
+	const FVector2D SourceUVOffset = Source.UVViewportMin;
+	const FVector2D SourceUVExtent = Source.UVViewportSize;
+	const FVector2D DestinationUVOffset = Destination.UVViewportMin;
+	const FVector2D DestinationUVExtent = Destination.UVViewportSize;
 
-	FScreenPassInput Input;
-	Input.Size = FVector4(Size.X, Size.Y, 1.0f / Size.X, 1.0f / Size.Y);
-	Input.Texture = Texture;
-	Input.Sampler = SamplerState;
-	return Input;
+	return GetScreenPassTextureViewportTransform(SourceUVOffset, SourceUVExtent, DestinationUVOffset, DestinationUVExtent);
+}
+
+FScreenPassCommonParameters GetScreenPassCommonParameters(FRHICommandListImmediate& RHICmdList, const FViewInfo& View)
+{
+	FSceneRenderTargets& SceneContext = FSceneRenderTargets::Get(RHICmdList);
+
+	FScreenPassCommonParameters Parameters;
+	Parameters.BilinearTextureSampler = TStaticSamplerState<SF_Bilinear>::GetRHI();
+	Parameters.ViewUniformBuffer = View.ViewUniformBuffer;
+	Parameters.SceneUniformBuffer = CreateSceneTextureUniformBuffer(
+		SceneContext, View.FeatureLevel, ESceneTextureSetupMode::All, EUniformBufferUsage::UniformBuffer_SingleFrame);
+	return Parameters;
 }
 
 FScreenPassContext* FScreenPassContext::Create(FRHICommandListImmediate& RHICmdList, const FViewInfo& InView)
@@ -86,9 +171,10 @@ FScreenPassContext* FScreenPassContext::Create(FRHICommandListImmediate& RHICmdL
 FScreenPassContext::FScreenPassContext(FRHICommandListImmediate& RHICmdList, const FViewInfo& InView)
 	: View(InView)
 	, ViewFamily(*View.Family)
-	, ViewportRect(View.ViewRect)
+	, ViewState(View.ViewState)
 	, StereoPass(View.StereoPass)
 	, bHasHMDMask(IsHMDHiddenAreaMaskActive())
+	, bUseComputePasses(ShouldDoComputePostProcessing(InView))
 	, ShaderMap(View.ShaderMap)
 	, ScreenPassVS(View.ShaderMap)
 	, ScreenPassCommonParameters(GetScreenPassCommonParameters(RHICmdList, View))
diff --git a/Engine/Source/Runtime/Renderer/Private/ScreenPass.h b/Engine/Source/Runtime/Renderer/Private/ScreenPass.h
--- a/Engine/Source/Runtime/Renderer/Private/ScreenPass.h
+++ b/Engine/Source/Runtime/Renderer/Private/ScreenPass.h
@@ -9,6 +9,12 @@
 #include "RenderGraph.h"
 #include "PostProcess/SceneFilterRendering.h"
 
+// Returns whether a HMD hidden area mask is being used for VR.
+bool IsHMDHiddenAreaMaskActive();
+
+// Returns the global engine mini font texture.
+const FTextureRHIRef& GetMiniFontTexture();
+
 // The vertex shader used by DrawScreenPass to draw a rectangle.
 class FScreenPassVS : public FGlobalShader
 {
@@ -26,33 +32,182 @@ public:
 	{}
 };
 
-// Returns whether a HMD hidden area mask is being used for VR.
-bool IsHMDHiddenAreaMaskActive();
+// Describes a viewport rect oriented within a texture.
+struct FScreenPassTextureViewport
+{
+	// Creates a viewport that is downscaled by the requested scale factor.
+	static FScreenPassTextureViewport CreateDownscaled(const FScreenPassTextureViewport& Other, uint32 ScaleFactor);
 
-// Returns the global engine mini font texture.
-const FTextureRHIRef& GetMiniFontTexture();
+	FScreenPassTextureViewport() = default;
+
+	FScreenPassTextureViewport(FIntRect InRect)
+		: Rect(InRect)
+		, Extent(InRect.Max)
+	{}
+
+	FScreenPassTextureViewport(FIntRect InRect, FIntPoint InExtent)
+		: Rect(InRect)
+		, Extent(InExtent)
+	{}
+
+	FScreenPassTextureViewport(const FScreenPassTextureViewport&) = default;
+
+	bool operator==(const FScreenPassTextureViewport& Other) const;
+	bool operator!=(const FScreenPassTextureViewport& Other) const;
+
+	// Returns whether the viewport contains an empty viewport or extent.
+	bool IsEmpty() const;
+
+	// The viewport rect, in pixels; defines a sub-set within [0, 0]x(Extent, Extent).
+	FIntRect Rect;
+
+	// The texture extent, in pixels; defines a super-set [0, 0]x(Extent, Extent).
+	FIntPoint Extent = FIntPoint::ZeroValue;
+};
+
+// Describes the set of shader parameters for a screen pass texture viewport.
+BEGIN_SHADER_PARAMETER_STRUCT(FScreenPassTextureViewportParameters, )
+	// Texture extent in pixels.
+	SHADER_PARAMETER(FVector2D, Extent)
+	SHADER_PARAMETER(FVector2D, ExtentInverse)
+
+	// Scale / Bias factor to convert from [-1, 1] to [ViewportMin, ViewportMax]
+	SHADER_PARAMETER(FVector2D, ScreenPosToViewportScale)
+	SHADER_PARAMETER(FVector2D, ScreenPosToViewportBias)
+
+	// Texture viewport min / max in pixels.
+	SHADER_PARAMETER(FIntPoint, ViewportMin)
+	SHADER_PARAMETER(FIntPoint, ViewportMax)
+
+	// Texture viewport size in pixels.
+	SHADER_PARAMETER(FVector2D, ViewportSize)
+	SHADER_PARAMETER(FVector2D, ViewportSizeInverse)
+
+	// Texture viewport min / max in normalized UV coordinates, with respect to the texture extent.
+	SHADER_PARAMETER(FVector2D, UVViewportMin)
+	SHADER_PARAMETER(FVector2D, UVViewportMax)
+
+	// Texture viewport size in normalized UV coordinates, with respect to the texture extent.
+	SHADER_PARAMETER(FVector2D, UVViewportSize)
+	SHADER_PARAMETER(FVector2D, UVViewportSizeInverse)
+
+	// Texture viewport min / max in normalized UV coordinates, with respect to the texture extent,
+	// adjusted by a half pixel offset for bilinear filtering. Useful for clamping to avoid sampling
+	// pixels on viewport edges; e.g. clamp(UV, UVViewportBilinearMin, UVViewportBilinearMax);
+	SHADER_PARAMETER(FVector2D, UVViewportBilinearMin)
+	SHADER_PARAMETER(FVector2D, UVViewportBilinearMax)
+END_SHADER_PARAMETER_STRUCT()
+
+FScreenPassTextureViewportParameters GetScreenPassTextureViewportParameters(const FScreenPassTextureViewport& InViewport);
+
+/**
+ * A class grouping a render graph texture and its shader parameters. The screen pass texture
+ * is either in an empty state or a valid state. The default constructor initializes to an empty state.
+ * Instantiating an instance from valid state requires using the provided @ref Create functions, which
+ * validate that the shader parameter invariant holds.
+ */
+class FScreenPassTexture
+{
+public:
+	/**
+	 * Creates a screen pass texture with a viewport.
+	 * @param InTexture A valid render graph texture instance. Asserts if the texture is null.
+	 * @param InViewport A viewport contained within the texture extent.
+	 */
+	static FScreenPassTexture Create(FRDGTextureRef InTexture, FIntRect InViewport);
+
+	/**
+	 * Creates a full-screen pass texture with a viewport covering the full texture.
+	 * @param InTexture A valid render graph texture instance. Asserts if the texture is null.
+	 */
+	static FScreenPassTexture CreateFullscreen(FRDGTextureRef InTexture);
+
+	// Default initializes to empty.
+	FScreenPassTexture() = default;
+
+	// Copy construction is permitted.
+	FScreenPassTexture(const FScreenPassTexture&) = default;
+
+	inline FRDGTextureRef GetRDGTexture() const
+	{
+		return Texture;
+	}
+
+	inline FScreenPassTextureViewport GetViewport() const
+	{
+		return Viewport;
+	}
+
+	inline FScreenPassTextureViewportParameters GetViewportParameters() const
+	{
+		return GetScreenPassTextureViewportParameters(Viewport);
+	}
+
+	inline const FTextureRHIParamRef GetRHITexture() const
+	{
+		return Texture ? Texture->GetRHITexture() : nullptr;
+	}
+
+	inline bool IsValid() const
+	{
+		return Texture != nullptr;
+	}
+
+	inline bool IsEmpty() const
+	{
+		return Texture == nullptr;
+	}
+
+	inline bool operator==(const FScreenPassTexture& Other) const
+	{
+		return Texture == Other.Texture && Viewport == Other.Viewport;
+	}
+
+	inline bool operator!=(const FScreenPassTexture& Other) const
+	{
+		return !(*this == Other);
+	}
+
+private:
+	FRDGTextureRef Texture = nullptr;
+	FScreenPassTextureViewport Viewport;
+};
+
+/**
+ * Contains a transform that maps UV coordinates from one screen pass texture viewport to another.
+ * Assumes normalized UV coordinates [0, 0]x[1, 1] where [0, 0] maps to the source view min
+ * coordinate and [1, 1] maps to the source view rect max coordinate.
+ *
+ * Example Usage:
+ *    float2 DestinationUV = SourceUV * UVScaleBias.xy + UVScaleBias.zw;
+ */
+BEGIN_SHADER_PARAMETER_STRUCT(FScreenPassTextureViewportTransform, )
+	// A scale / bias factor to apply to the input UV coordinate, converting it to a output UV coordinate.
+	SHADER_PARAMETER(FVector2D, Scale)
+	SHADER_PARAMETER(FVector2D, Bias)
+END_SHADER_PARAMETER_STRUCT()
+
+// Constructs a view transform from source and destination texture viewports.
+FScreenPassTextureViewportTransform GetScreenPassTextureViewportTransform(
+	const FScreenPassTextureViewportParameters& Source,
+	const FScreenPassTextureViewportParameters& Destination);
+
+// Constructs a view transform from source and destination UV offset / extent pairs.
+FScreenPassTextureViewportTransform GetScreenPassTextureViewportTransform(
+	FVector2D SourceUVOffset,
+	FVector2D SourceUVExtent,
+	FVector2D DestinationUVOffset,
+	FVector2D DestinationUVExtent);
 
 // Defines the common set of parameters for a screen space pass.
 BEGIN_SHADER_PARAMETER_STRUCT(FScreenPassCommonParameters, )
-	SHADER_PARAMETER(FIntRect, ViewportRect)
-	SHADER_PARAMETER(FVector4, ViewportSize)
-	SHADER_PARAMETER(FVector4, ScreenPosToPixelValue)
 	SHADER_PARAMETER_STRUCT_REF(FSceneTexturesUniformParameters, SceneUniformBuffer)
 	SHADER_PARAMETER_STRUCT_REF(FViewUniformShaderParameters, ViewUniformBuffer)
-	SHADER_PARAMETER_SAMPLER(SamplerState, BilinearTextureSampler0)
+	SHADER_PARAMETER_SAMPLER(SamplerState, BilinearTextureSampler)
 END_SHADER_PARAMETER_STRUCT()
 
 FScreenPassCommonParameters GetScreenPassCommonParameters(FRHICommandListImmediate& RHICmdList, const FViewInfo& View);
 
-// Defines shader parameters for a single texture input to a screen space pass.
-BEGIN_SHADER_PARAMETER_STRUCT(FScreenPassInput, )
-	SHADER_PARAMETER(FVector4, Size)
-	SHADER_PARAMETER_RDG_TEXTURE(Texture2D, Texture)
-	SHADER_PARAMETER_SAMPLER(SamplerState, Sampler)
-END_SHADER_PARAMETER_STRUCT()
-
-FScreenPassInput GetScreenPassInputParameters(FRDGTextureRef Texture, FSamplerStateRHIParamRef SamplerState);
-
 /**
  * The context used for screen pass operations. Extracts and holds common state required
  * by most screen passes from the provided View instance. Exists to reduce pointer
@@ -70,8 +225,8 @@ public:
 	// The current view family instance being processed.
 	const FSceneViewFamily& ViewFamily;
 
-	// The viewport rect for the view being processed.
-	const FIntRect ViewportRect;
+	// The current view state instance being processed.
+	const FSceneViewState* ViewState;
 
 	// VR - Which stereo pass is being rendered.
 	const EStereoscopicPass StereoPass;
@@ -79,6 +234,9 @@ public:
 	// VR - Whether a HMD hidden area mask is being used for VR.
 	const bool bHasHMDMask;
 
+	// Whether screen passes should use compute.
+	const bool bUseComputePasses;
+
 	// The global shader map for the current view being processed.
 	const TShaderMap<FGlobalShaderType>* ShaderMap;
 
@@ -108,13 +266,13 @@ template<typename TPixelShaderType>
 void DrawScreenPass(
 	FRHICommandListImmediate& RHICmdList,
 	FScreenPassContextRef Context,
-	FIntRect ViewportRect,
-	FIntRect TextureRect,
-	FIntPoint TextureSize,
+	FIntRect OutputRect,
+	FIntRect InputRect,
+	FIntPoint InputSize,
 	TPixelShaderType* PixelShader,
 	const typename TPixelShaderType::FParameters& PixelShaderParameters)
 {
-	const FIntPoint ViewportSize = ViewportRect.Size();
+	const FIntPoint OutputSize = OutputRect.Size();
 	const FPixelShaderRHIParamRef PixelShaderRHI = GETSAFERHISHADER_PIXEL(PixelShader);
 	FScreenPassVS* VertexShader = *(Context->ScreenPassVS);
 
@@ -128,25 +286,48 @@ void DrawScreenPass(
 	GraphicsPSOInit.BoundShaderState.PixelShaderRHI = PixelShaderRHI;
 	GraphicsPSOInit.PrimitiveType = PT_TriangleList;
 
-	RHICmdList.SetViewport(ViewportRect.Min.X, ViewportRect.Min.Y, 0.0f, ViewportRect.Max.X, ViewportRect.Max.Y, 1.0f);
+	RHICmdList.SetViewport(OutputRect.Min.X, OutputRect.Min.Y, 0.0f, OutputRect.Max.X, OutputRect.Max.Y, 1.0f);
 	SetGraphicsPipelineState(RHICmdList, GraphicsPSOInit);
 	SetShaderParameters(RHICmdList, PixelShader, PixelShaderRHI, PixelShaderParameters);
 
 	DrawPostProcessPass(
 		RHICmdList,
-		0, 0, ViewportSize.X, ViewportSize.Y,
-		TextureRect.Min.X, TextureRect.Min.Y, TextureRect.Width(), TextureRect.Height(),
-		ViewportSize,
-		TextureSize,
+		0, 0, OutputSize.X, OutputSize.Y,
+		InputRect.Min.X, InputRect.Min.Y, InputRect.Width(), InputRect.Height(),
+		OutputSize,
+		InputSize,
 		VertexShader,
 		Context->StereoPass,
 		Context->bHasHMDMask,
 		EDRF_UseTriangleOptimization);
 }
 
+/**
+ * Helper variant of @ref DrawScreenPass. All other parameters are forwarded.
+ * @param DestinationViewDesc The destination texture view descriptor, for viewport generation.
+ * @param SourceViewDesc The source view descriptor, for UV generation.
+ */
+template <typename TPixelShaderType>
+inline void DrawScreenPass(
+	FRHICommandListImmediate& RHICmdList,
+	FScreenPassContextRef Context,
+	const FScreenPassTextureViewport& OutputViewport,
+	const FScreenPassTextureViewport& InputViewport,
+	TPixelShaderType* PixelShader,
+	const typename TPixelShaderType::FParameters& PixelShaderParameters)
+{
+	DrawScreenPass(
+		RHICmdList,
+		Context,
+		OutputViewport.Rect,
+		InputViewport.Rect,
+		InputViewport.Extent,
+		PixelShader,
+		PixelShaderParameters);
+}
+
 /**
  * Adds a new Render Graph pass which internally calls @ref DrawScreenPass.
- *
  * Parameters are forwarded to @ref DrawScreenPass.
  */
 template <typename TPixelShaderType>
@@ -154,9 +335,9 @@ inline void AddDrawScreenPass(
 	FRDGBuilder& GraphBuilder,
 	FRDGEventName&& PassName,
 	FScreenPassContextRef Context,
-	FIntRect ViewportRect,
-	FIntRect TextureRect,
-	FIntPoint TextureSize,
+	FIntRect PixelOutputRect,
+	FIntRect PixelInputRect,
+	FIntPoint PixelInputSize,
 	TPixelShaderType* PixelShader,
 	typename TPixelShaderType::FParameters* PixelShaderParameters)
 {
@@ -170,9 +351,40 @@ inline void AddDrawScreenPass(
 		Forward<FRDGEventName>(PassName),
 		PixelShaderParameters,
 		ERenderGraphPassFlags::None,
-		[Context, ViewportRect, TextureRect, TextureSize, PixelShader, PixelShaderParameters]
+		[Context, PixelOutputRect, PixelInputRect, PixelInputSize, PixelShader, PixelShaderParameters]
 		(FRHICommandListImmediate& RHICmdList)
 	{
-		DrawScreenPass(RHICmdList, Context, ViewportRect, TextureRect, TextureSize, PixelShader, *PixelShaderParameters);
+		DrawScreenPass(
+			RHICmdList,
+			Context,
+			PixelOutputRect,
+			PixelInputRect,
+			PixelInputSize,
+			PixelShader,
+			*PixelShaderParameters);
 	});
+}
+
+/**
+ * Helper variants of @ref AddDrawScreenPass. All other parameters are forwarded.
+ */
+template <typename TPixelShaderType>
+inline void AddDrawScreenPass(
+	FRDGBuilder& GraphBuilder,
+	FRDGEventName&& PassName,
+	FScreenPassContextRef Context,
+	const FScreenPassTextureViewport& OutputViewport,
+	const FScreenPassTextureViewport& InputViewport,
+	TPixelShaderType* PixelShader,
+	typename TPixelShaderType::FParameters* PixelShaderParameters)
+{
+	AddDrawScreenPass(
+		GraphBuilder,
+		Forward<FRDGEventName>(PassName),
+		Context,
+		OutputViewport.Rect,
+		InputViewport.Rect,
+		InputViewport.Extent,
+		PixelShader,
+		PixelShaderParameters);
 }
\ No newline at end of file
diff --git a/Engine/Source/Runtime/Renderer/Public/PostProcessParameters.h b/Engine/Source/Runtime/Renderer/Public/PostProcessParameters.h
--- a/Engine/Source/Runtime/Renderer/Public/PostProcessParameters.h
+++ b/Engine/Source/Runtime/Renderer/Public/PostProcessParameters.h
@@ -97,8 +97,7 @@ private:
 	FShaderParameter PostProcessInputMinMaxParameter[ePId_Input_MAX];
 	FShaderParameter ScreenPosToPixel;
 	FShaderParameter SceneColorBufferUVViewport;
-	FShaderResourceParameter BilinearTextureSampler0;
-	FShaderResourceParameter BilinearTextureSampler1;
+	FShaderResourceParameter BilinearTextureSampler;
 
 public:
 	// @param Filter can be 0 if FilterOverrideArray is used
