Fixes regression on DOF's TAA pixel shader when alpha channel is in enabled.

diff --git a/Engine/Source/Runtime/Renderer/Private/PostProcess/PostProcessTemporalAA.cpp b/Engine/Source/Runtime/Renderer/Private/PostProcess/PostProcessTemporalAA.cpp
--- a/Engine/Source/Runtime/Renderer/Private/PostProcess/PostProcessTemporalAA.cpp
+++ b/Engine/Source/Runtime/Renderer/Private/PostProcess/PostProcessTemporalAA.cpp
@@ -590,10 +590,10 @@ FTAAOutputs FTAAPassParameters::AddTemporalAAPass(
 				SceneColorDesc.TargetableFlags |= TexCreate_UAV;
 			}
 
-			Outputs.SceneColor = GraphBuilder.CreateTexture(
-				SceneColorDesc,
-				kTAAOutputNames[static_cast<int32>(Pass)],
-				ERDGResourceFlags::MultiFrame);
+		Outputs.SceneColor = GraphBuilder.CreateTexture(
+			SceneColorDesc,
+			kTAAOutputNames[static_cast<int32>(Pass)],
+			ERDGResourceFlags::MultiFrame);
 		}
 
 		if (RenderTargetCount == 2)
@@ -722,6 +722,14 @@ FTAAOutputs FTAAPassParameters::AddTemporalAAPass(
 			ERenderTargetLoadAction::ENoAction,
 			ERenderTargetStoreAction::EStore);
 		
+		if (Outputs.SceneMetadata)
+		{
+			PassParameters->RenderTargets[1] = FRenderTargetBinding(
+				Outputs.SceneMetadata,
+				ERenderTargetLoadAction::ENoAction,
+				ERenderTargetStoreAction::EStore);
+		}
+
 		if (bUseResponsiveStencilTest)
 		{
 			PassParameters->RenderTargets.DepthStencil.Texture = SceneBlackboard.SceneDepthBuffer;
