Fixes TAAU not updating it's output extent.

diff --git a/Engine/Source/Runtime/Renderer/Private/PostProcess/PostProcessTemporalAA.cpp b/Engine/Source/Runtime/Renderer/Private/PostProcess/PostProcessTemporalAA.cpp
--- a/Engine/Source/Runtime/Renderer/Private/PostProcess/PostProcessTemporalAA.cpp
+++ b/Engine/Source/Runtime/Renderer/Private/PostProcess/PostProcessTemporalAA.cpp
@@ -574,27 +574,18 @@ FTAAOutputs FTAAPassParameters::AddTemporalAAPass(
 	// Create outputs
 	FTAAOutputs Outputs;
 	{
-		FRDGTextureDesc SceneColorDesc = SceneColorInput->Desc;
-		{
-			SceneColorDesc.Flags &= ~(TexCreate_FastVRAM | TexCreate_Transient);
-			SceneColorDesc.Reset();
-			//regardless of input type, PF_FloatRGBA is required to properly accumulate between frames for a good result.
-			SceneColorDesc.Format = PF_FloatRGBA;
-			SceneColorDesc.AutoWritable = false;
-			SceneColorDesc.TargetableFlags &= ~(TexCreate_RenderTargetable | TexCreate_UAV);
-			SceneColorDesc.TargetableFlags |= bIsComputePass ? TexCreate_UAV : TexCreate_RenderTargetable;
-
-			// Need a UAV to resource transition from gfx to compute.
-			if (IsDOFTAAConfig(Pass))
-			{
-				SceneColorDesc.TargetableFlags |= TexCreate_UAV;
-			}
+		FRDGTextureDesc SceneColorDesc = FRDGTextureDesc::Create2DDesc(
+			GetOutputExtent(),
+			PF_FloatRGBA,
+			FClearValueBinding::Black,
+			/* InFlags = */ TexCreate_None,
+			/* InTargetableFlags = */ TexCreate_ShaderResource | (bIsComputePass ? TexCreate_UAV : TexCreate_RenderTargetable),
+			/* bInForceSeparateTargetAndShaderResource = */ false);
 
 		Outputs.SceneColor = GraphBuilder.CreateTexture(
 			SceneColorDesc,
 			kTAAOutputNames[static_cast<int32>(Pass)],
 			ERDGResourceFlags::MultiFrame);
-		}
 
 		if (RenderTargetCount == 2)
 		{
@@ -608,23 +599,14 @@ FTAAOutputs FTAAPassParameters::AddTemporalAAPass(
 		{
 			check(bIsComputePass);
 			
-			FRDGTextureDesc HalfResSceneColorDesc = SceneColorInput->Desc;
-			HalfResSceneColorDesc.Flags &= ~(TexCreate_FastVRAM);
-			HalfResSceneColorDesc.Reset();
+			FRDGTextureDesc HalfResSceneColorDesc = FRDGTextureDesc::Create2DDesc(
+				SceneColorDesc.Extent / 2,
+				DownsampleOverrideFormat != PF_Unknown ? DownsampleOverrideFormat : SceneColorInput->Desc.Format,
+				FClearValueBinding::Black,
+				/* InFlags = */ TexCreate_None,
+				/* InTargetableFlags = */ TexCreate_ShaderResource | TexCreate_Transient | (bIsComputePass ? TexCreate_UAV : TexCreate_RenderTargetable),
+				/* bInForceSeparateTargetAndShaderResource = */ false);
 
-			if (DownsampleOverrideFormat != PF_Unknown)
-			{
-				HalfResSceneColorDesc.Format = DownsampleOverrideFormat;
-			}
-
-			HalfResSceneColorDesc.AutoWritable = false;
-			HalfResSceneColorDesc.TargetableFlags &= ~TexCreate_RenderTargetable;
-			HalfResSceneColorDesc.TargetableFlags |= TexCreate_UAV;
-
-			HalfResSceneColorDesc.Extent = FIntPoint::DivideAndRoundUp(HalfResSceneColorDesc.Extent, 2);
-			HalfResSceneColorDesc.Extent.X = FMath::Max(1, HalfResSceneColorDesc.Extent.X);
-			HalfResSceneColorDesc.Extent.Y = FMath::Max(1, HalfResSceneColorDesc.Extent.Y);
-			
 			Outputs.DownsampledSceneColor = GraphBuilder.CreateTexture(HalfResSceneColorDesc, TEXT("SceneColorHalfRes"));
 		}
 	}
