Fixes link error in PostProcessing.cpp

diff --git a/Engine/Source/Runtime/Renderer/Private/PostProcess/PostProcessing.cpp b/Engine/Source/Runtime/Renderer/Private/PostProcess/PostProcessing.cpp
--- a/Engine/Source/Runtime/Renderer/Private/PostProcess/PostProcessing.cpp
+++ b/Engine/Source/Runtime/Renderer/Private/PostProcess/PostProcessing.cpp
@@ -234,7 +234,7 @@ static bool HasPostProcessMaterial(FPostprocessContext& Context, EBlendableLocat
 
 // -------------------------------------------------------
 
-static bool ShouldDoComputePostProcessing(const FViewInfo& View)
+bool ShouldDoComputePostProcessing(const FViewInfo& View)
 {
 	return CVarPostProcessingPreferCompute.GetValueOnRenderThread() && View.FeatureLevel >= ERHIFeatureLevel::SM5;
 }
