Rewrote motion blur passes to work with Render Graph system.

 - Added ScreenPass.ush helper file to define shader parmeters for texture viewports.
 - Motion Blur passes use render graph.
 - Fixed dangling pointer bug in RDGEventName when copied or moved.

diff --git a/Engine/Shaders/Private/PostProcessMotionBlur.usf b/Engine/Shaders/Private/PostProcessMotionBlur.usf
--- a/Engine/Shaders/Private/PostProcessMotionBlur.usf
+++ b/Engine/Shaders/Private/PostProcessMotionBlur.usf
@@ -5,27 +5,9 @@
 =============================================================================*/
 
 #include "Common.ush"
-#include "PostProcessCommon.ush"
-#include "DeferredShadingCommon.ush"		// FGBufferData
 #include "FastMath.ush"
 #include "Random.ush"
-
-
-#if MOTION_BLUR_INNER_LOOP_CONFIG == 0
-	// Stay opt-out of everything.
-
-#elif MOTION_BLUR_INNER_LOOP_CONFIG == 1
-	#define CLAMP_UV 1
-
-#elif MOTION_BLUR_INNER_LOOP_CONFIG == 2
-	#define CLAMP_UV 1
-	#define SUPPORT_DIFFERENT_SCENE_COLOR_VIEWPORT 1
-
-#else
-	#error Sorry
-
-#endif
-
+#include "ScreenPass.ush"
 
 // from the paper: We use SOFT Z EXTENT = 1mm to 10cm for our results
 #define SOFT_Z_EXTENT 1
@@ -33,183 +15,80 @@
 // 0:off / 1:on, useful to debug the motionblur algorithm
 #define MOTIONBLUR_TESTCHART 0
 
-// 0:off / 1:on, manually clamp BufferUV to the viewport to not sample outside of teh viewport.
-#ifndef CLAMP_UV
-	#define CLAMP_UV 0
-#endif
-
-// 0:off / 1:on, required for TAA upsample that change the res of the scene color before motion blur.
-#ifndef SUPPORT_DIFFERENT_SCENE_COLOR_VIEWPORT
-	#define SUPPORT_DIFFERENT_SCENE_COLOR_VIEWPORT 0
-#endif
-
+#define MOTION_BLUR_QUALITY_LOW 0
+#define MOTION_BLUR_QUALITY_MEDIUM 1
+#define MOTION_BLUR_QUALITY_HIGH 2
+#define MOTION_BLUR_QUALITY_VERY_HIGH 2
 
-// --------------------------------------------------------------------------
+float MotionBlur_AspectRatio;
+float MotionBlur_VelocityScale;
+float MotionBlur_VelocityScaleForTiles;
+float MotionBlur_VelocityMax;
 
-// MOTION_BLUR_QUALITY == 0:visualize, 1:low, 2:medium, 3:high, 4:very high
+SCREEN_PASS_TEXTURE_VIEWPORT(Color);
+SCREEN_PASS_TEXTURE_VIEWPORT(Velocity);
+SCREEN_PASS_TEXTURE_VIEWPORT(VelocityTile);
+SCREEN_PASS_TEXTURE_VIEWPORT_TRANSFORM(ColorToVelocity);
 
-// to scale to normalized motionblur velocity
-// xy:includes y flip, zw:unused
-float4 VelocityScale;
-
-// xy:IndividualVelocityScale.xy zw:unused, from postprocess settings
-float4 MotionBlurParameters;
-
-float2 TileUVMax;
-
-#if SUPPORT_DIFFERENT_SCENE_COLOR_VIEWPORT
-	float4 SceneColorBufferUVToViewBufferUV;
-#endif
+Texture2D ColorTexture;
+Texture2D VelocityTexture;
+Texture2D VelocityTileTexture;
+Texture2D VelocityFlatTexture;
+Texture2D DepthTexture;
 
+SamplerState ColorSampler;
+SamplerState VelocitySampler;
 
 // ------------------------------------------
 
 uint GetStepCountFromQuality()
 {
-#if MOTION_BLUR_QUALITY == 1
+#if MOTION_BLUR_QUALITY == MOTION_BLUR_QUALITY_LOW
 	return 4;
-#elif MOTION_BLUR_QUALITY == 2
+#elif MOTION_BLUR_QUALITY == MOTION_BLUR_QUALITY_MEDIUM
 	return 6;
-#elif MOTION_BLUR_QUALITY == 3
+#elif MOTION_BLUR_QUALITY == MOTION_BLUR_QUALITY_HIGH
 	return 8;
-#else // MOTION_BLUR_QUALITY == 4
+#else // MOTION_BLUR_QUALITY == MOTION_BLUR_QUALITY_VERY_HIGH
 	return 16;
 #endif
 }
 
-uint	bDrawMax;
-
-void VelocityScatterVS(
-	uint VId : SV_VertexID,
-	uint IId : SV_InstanceID,
-	out nointerpolation float4 OutColor : TEXCOORD0,
-	out float4 OutPosition : SV_POSITION
-	)
-{
-	OutPosition =  float4(0, 0, 0, 1);
-
-	// needs to be the same on C++ side (faster on NVIDIA and AMD)
-	uint QuadsPerInstance = 8;
-
-	// remap the indices to get vertexid to VId and quadid into IId
-	IId = IId * QuadsPerInstance + (VId / 4);
-	VId = VId % 4;
-
-	// triangle A: 0:left top, 1:right top, 2: left bottom
-	// triangle B: 3:right bottom, 4:left bottom, 5: right top
-	float2 CornerOffset = float2(VId % 2, VId / 2) * 2 - 1;
-
-	uint2 PixelPos = uint2( IId % ViewportRect.z, IId / ViewportRect.z );
-
-	BRANCH
-	if( PixelPos.y >= ViewportRect.w )
-	{
-		OutColor = 0;
-		return;
-	}
-
-	float4 MinMaxVelocity = PostprocessInput0[ PixelPos ];
-	OutColor = MinMaxVelocity;
-
-	float4 MinMaxVelocityPixels = MinMaxVelocity * MotionBlurParameters.z;
-	float2 VelocityPixels = MinMaxVelocityPixels.zw;
-	
-#if 1
-	BRANCH
-	if( dot( VelocityPixels, VelocityPixels ) * 16*16 <= 0.25 )
-	{
-		OutPosition.xy = ( PixelPos + CornerOffset * 0.5 - ScreenPosToPixel.zw ) / ScreenPosToPixel.xy;
-		OutPosition.z = 0.0002;		// zero clips
-		return;
-	}
-#endif
-
-	float  VelocityLengthPixelsSqr = dot( VelocityPixels, VelocityPixels );
-	float  VelocityLengthPixelsInv = rsqrtFast( VelocityLengthPixelsSqr );
-	float  VelocityLengthPixels = VelocityLengthPixelsSqr * VelocityLengthPixelsInv;
-	float2 VelocityDir = VelocityPixels * VelocityLengthPixelsInv;
-
-	// Project pixel corner on to dir. This is the oriented extent of a pixel.
-	// 1/2 pixel because shape is swept tile
-	// +1/2 pixel for conservative rasterization
-	// 99% to give epsilon before neighbor is filled. Otherwise all neighbors lie on edges of quad when no velocity in their direction.
-	float Extent = abs( VelocityDir.x ) + abs( VelocityDir.y );
-	CornerOffset *= float2( VelocityLengthPixels, 0 ) + Extent.xx * 0.99;
-
-	const float TwoPixelRadius = sqrt( 2.0 );
-	//CornerOffset *= float2( VelocityLengthPixels, 0 ) + TwoPixelRadius.xx;
-
-	// Orient along velocity direction
-	float2 AxisX = VelocityDir;
-	float2 AxisY = float2( -VelocityDir.y, VelocityDir.x );
-	CornerOffset = AxisX * CornerOffset.x + AxisY * CornerOffset.y;
-		
-	OutPosition.xy = ( PixelPos + CornerOffset - ScreenPosToPixel.zw ) / ScreenPosToPixel.xy;
-
-	// Depth ordered by velocity length
-	OutPosition.z = bDrawMax ? VelocityLengthPixels : length( MinMaxVelocityPixels.xy );
-	OutPosition.z = clamp( OutPosition.z / ScreenPosToPixel.x * 0.5, 0.0002, 0.999 );
-}
-
-void VelocityScatterPS(
-	nointerpolation float4 InColor : TEXCOORD0,
-	out float4 OutColor : SV_Target0
-	)
-{
-	OutColor = InColor;
-}
-
-
-
 //-----------------------------------
 
-float2 DepthCmp( float CenterDepth, float SampleDepth, float DepthScale )
+float2 DepthCmp(float CenterDepth, float SampleDepth, float DepthScale)
 {
-	return saturate( 0.5 + float2( DepthScale, -DepthScale ) * (SampleDepth - CenterDepth) );
+	return saturate(0.5 + float2(DepthScale, -DepthScale) * (SampleDepth - CenterDepth));
 }
 
-float2 SpreadCmp( float OffsetLength, float2 SpreadLength, float PixelToSampleScale )
+float2 SpreadCmp(float OffsetLength, float2 SpreadLength, float PixelToSampleScale)
 {
-	return saturate( PixelToSampleScale * SpreadLength - max( OffsetLength - 1, 0 ) );
+	return saturate(PixelToSampleScale * SpreadLength - max(OffsetLength - 1, 0));
 }
 
-float SampleWeight( float CenterDepth, float SampleDepth, float OffsetLength, float CenterSpreadLength, float SampleSpreadLength, float PixelToSampleScale, float DepthScale )
+float SampleWeight(float CenterDepth, float SampleDepth, float OffsetLength, float CenterSpreadLength, float SampleSpreadLength, float PixelToSampleScale, float DepthScale)
 {
-	float2 DepthWeights = DepthCmp( CenterDepth, SampleDepth, DepthScale );
-	float2 SpreadWeights = SpreadCmp( OffsetLength, float2( CenterSpreadLength, SampleSpreadLength ), PixelToSampleScale );
-	return dot( DepthWeights, SpreadWeights );
+	float2 DepthWeights = DepthCmp(CenterDepth, SampleDepth, DepthScale);
+	float2 SpreadWeights = SpreadCmp(OffsetLength, float2(CenterSpreadLength, SampleSpreadLength), PixelToSampleScale);
+	return dot(DepthWeights, SpreadWeights);
 }
 
-float2 DecodeVelocity( float2 Velocity )
+float2 DecodeVelocity(float2 Velocity)
 {
-#if 1
 	// 11:11:10  (VelocityLength, VelocityAngle, Depth)
 	float VelocityLength = Velocity.x;
 	float VelocityAngle = Velocity.y * (2 * PI) - PI;
-	sincos( VelocityAngle, Velocity.x, Velocity.y );
+	sincos(VelocityAngle, Velocity.x, Velocity.y);
 	Velocity *= VelocityLength;
-#else
-	// 11:11:10 (Velocity.xy, Depth)
-	// Stored signed value in 6e5 unsigned float
-	// Extract sign from top bit of exponent
-	Velocity *= Velocity >= 2.0 ? (-1.0 / 32768.0) : 1;
-#endif
 	return Velocity;
 }
 
-float GetVelocityLengthPixels( float2 EncodedVelocity )
+float GetVelocityLengthPixels(float2 EncodedVelocity)
 {
-#if 1
 	// 11:11:10  (VelocityLength, VelocityAngle, Depth)
 	float VelocityLength = EncodedVelocity.x;
-	VelocityLength *= MotionBlurParameters.y;
-	return min( VelocityLength, MotionBlurParameters.w );
-#else
-	float2 Velocity = DecodeVelocity( EncodedVelocity ) * ScreenPosToPixel.xy;	// in pixels
-	float  VelocityLengthSqr = dot( Velocity, Velocity );
-	float  VelocityLength = sqrtFast( VelocityLengthSqr );
-	return min( VelocityLength, MotionBlurParameters.w );
-#endif
+	VelocityLength *= MotionBlur_VelocityScale;
+	return min(VelocityLength, MotionBlur_VelocityMax);
 }
 
 #if COMPUTESHADER && !COMPILER_SUPPORTS_WAVE_VOTE
@@ -218,65 +97,44 @@ groupshared uint GroupSharedFastPath;
 
 #endif
 
-float4 MainMotionBlurCommon(float2 OutputBufferUV, float2 ScreenSpacePos, float2 PixelPos)
+float4 MainMotionBlurCommon(float2 ColorUV, float2 ColorPixelPos)
 {
 	const uint StepCount = GetStepCountFromQuality() / 2;
 
-	#if !SUPPORT_DIFFERENT_SCENE_COLOR_VIEWPORT
-		float4 SceneColorBufferUVToViewBufferUV = float4(1, 1, 0, 0);
-	#endif
+	const float PixelToTileScale = (1.0 / 16.0);
 
-#if 0
-	uint2 iPixelPos = PixelPos;
-	float2 PosMod = float2( iPixelPos & 1 );
-	float Dither = ( PosMod.x * 0.5 - 0.25 ) * ( PosMod.y * 2 - 1 );
-	float Random = RandFast( iPixelPos );
-	float Random2 = RandFast( iPixelPos, 5521 );
-#else
-	float Random = InterleavedGradientNoise( PixelPos, 0 );
-	float Random2 = InterleavedGradientNoise( PixelPos, 1 );
-#endif
+	float Random  = InterleavedGradientNoise(ColorPixelPos, 0);
+	float Random2 = InterleavedGradientNoise(ColorPixelPos, 1);
+
+	// [-0.25, 0.25]
+	float2 TileJitter = (float2(Random, Random2) - 0.5) * 0.5;
 
-	float2 TileJitter = ( float2( Random, Random2 ) - 0.5 ) * 0.5;
+	// Map color UV to velocity UV space.
+	float2 VelocityUV = ColorUV * ColorToVelocity_Scale + ColorToVelocity_Bias;
 
-	float2 ViewBufferUV = OutputBufferUV * SceneColorBufferUVToViewBufferUV.xy + SceneColorBufferUVToViewBufferUV.zw;
-	float2 NearestBufferSvPosition = floor(ViewBufferUV * View.BufferSizeAndInvSize.xy) + 0.5;
-	float2 TileUV = ((NearestBufferSvPosition - View.ViewRectMin.xy) * (1.0 / 16.0) + TileJitter) * PostprocessInput3Size.zw;
+	// Map velocity UV to velocity tile UV space with jitter.
+	float2 NearestVelocitySvPosition = floor(VelocityUV * Velocity_Extent) + 0.5;
+	float2 VelocityTileUV = ((NearestVelocitySvPosition - Velocity_ViewportMin) * PixelToTileScale + TileJitter) * VelocityTile_ExtentInverse;
 
-	// Don't need to max TileUV because tiles are top left cornered in the PostprocessInput3.
-	TileUV = min(TileUV, TileUVMax.xy);
+	// Velocity tile UV originates at [0,0]; only need to clamp max.
+	VelocityTileUV = min(VelocityTileUV, VelocityTile_UVViewportBilinearMax);
 
-	float4 MinMaxVelocity = PostprocessInput3.SampleLevel(PostprocessInput3Sampler, TileUV, 0);
+	float4 MinMaxVelocity = VelocityTileTexture.SampleLevel(VelocitySampler, VelocityTileUV, 0);
 
-	float2 MinVelocityPixels = MinMaxVelocity.xy * MotionBlurParameters.y;
-	float2 MaxVelocityPixels = MinMaxVelocity.zw * MotionBlurParameters.y;
+	float2 MinVelocityPixels = MinMaxVelocity.xy * MotionBlur_VelocityScale;
+	float2 MaxVelocityPixels = MinMaxVelocity.zw * MotionBlur_VelocityScale;
 
-	float MinVelocityLengthSqrPixels = dot( MinVelocityPixels, MinVelocityPixels );
-	float MaxVelocityLengthSqrPixels = dot( MaxVelocityPixels, MaxVelocityPixels );
+	float MinVelocityLengthSqrPixels = dot(MinVelocityPixels, MinVelocityPixels);
+	float MaxVelocityLengthSqrPixels = dot(MaxVelocityPixels, MaxVelocityPixels);
 
 	// Input buffer 0 as same viewport as output buffer.
-	float4 CenterColor = PostprocessInput0.SampleLevel( PostprocessInput0Sampler, OutputBufferUV, 0 );
-
-	float2 MinViewportUV = (float2(ViewportRect.xy) + 0.5) * PostprocessInput0Size.zw;
-	float2 MaxViewportUV = (float2(ViewportRect.zw) - 0.5) * PostprocessInput0Size.zw;
-
-#if 1
-	float4 SearchVectorPixels = float4( MaxVelocityPixels, -MaxVelocityPixels );
-#else
-	// Clip MaxVelocity to screen rect
-	float2 InvVelocityScreen = ScreenPosToPixel.x * rcp( MaxVelocityPixels + 0.00001f );
-	float2 MinIntersect = -InvVelocityScreen - ScreenSpacePos * InvVelocityScreen;
-	float2 MaxIntersect =  InvVelocityScreen - ScreenSpacePos * InvVelocityScreen;
-	float4 FarIntersect = float4( max( MinIntersect, MaxIntersect ), max( -MinIntersect, -MaxIntersect ) );
-	float2 Intersect = saturate(min(FarIntersect.xz, FarIntersect.yw) );
-
-	// +/-
-	float4 SearchVectorPixels = MaxVelocityPixels.xyxy * float4( Intersect.xx, -Intersect.yy );
-#endif
-	float4 SearchVector = SearchVectorPixels * PostprocessInput0Size.zwzw;
+	float4 CenterColor = ColorTexture.SampleLevel(ColorSampler, ColorUV, 0);
+
+	float4 SearchVectorPixels = float4(MaxVelocityPixels, -MaxVelocityPixels);
+	float4 SearchVector = SearchVectorPixels * Color_ExtentInverse.xyxy;
 
 	// converts pixel length to sample steps
-	float PixelToSampleScale = StepCount * rsqrt( dot( MaxVelocityPixels, MaxVelocityPixels ) );
+	float PixelToSampleScale = StepCount * rsqrt(dot(MaxVelocityPixels, MaxVelocityPixels));
 	
 	// TODO expose cvars
 	bool bSkipPath = MaxVelocityLengthSqrPixels < 0.25;
@@ -301,7 +159,7 @@ float4 MainMotionBlurCommon(float2 OutputBufferUV, float2 ScreenSpacePos, float2
 	#endif
 	
 	BRANCH
-	if( bSkipPath )
+	if (bSkipPath)
 	{
 		#if !POST_PROCESS_ALPHA
 			CenterColor.a = 0;
@@ -312,27 +170,25 @@ float4 MainMotionBlurCommon(float2 OutputBufferUV, float2 ScreenSpacePos, float2
 	float4 OutColor = 0;
 
 	BRANCH
-	if( bFastPath )
+	if (bFastPath)
 	{
 		float4 ColorAccum = 0;
 
 		UNROLL
-		for( uint i = 0; i < StepCount; i++ )
+		for (uint i = 0; i < StepCount; i++)
 		{
-			float2 OffsetLength = (float)i + 0.5 + float2( Random - 0.5, 0.5 - Random );
+			float2 OffsetLength = (float)i + 0.5 + float2(Random - 0.5, 0.5 - Random);
 			float2 OffsetFraction = OffsetLength / StepCount;
 
 			float2 SampleUV[2];
-			SampleUV[0] = OutputBufferUV + OffsetFraction.x * SearchVector.xy;
-			SampleUV[1] = OutputBufferUV + OffsetFraction.y * SearchVector.zw;
+			SampleUV[0] = ColorUV + OffsetFraction.x * SearchVector.xy;
+			SampleUV[1] = ColorUV + OffsetFraction.y * SearchVector.zw;
 
-			#if CLAMP_UV
-				SampleUV[0] = clamp( SampleUV[0], MinViewportUV, MaxViewportUV );
-				SampleUV[1] = clamp( SampleUV[1], MinViewportUV, MaxViewportUV );
-			#endif
+			SampleUV[0] = clamp(SampleUV[0], Color_UVViewportBilinearMin, Color_UVViewportBilinearMax);
+			SampleUV[1] = clamp(SampleUV[1], Color_UVViewportBilinearMin, Color_UVViewportBilinearMax);
 
-			ColorAccum += PostprocessInput0.SampleLevel( PostprocessInput0Sampler, SampleUV[0], 0 );
-			ColorAccum += PostprocessInput0.SampleLevel( PostprocessInput0Sampler, SampleUV[1], 0 );
+			ColorAccum += ColorTexture.SampleLevel(ColorSampler, SampleUV[0], 0);
+			ColorAccum += ColorTexture.SampleLevel(ColorSampler, SampleUV[1], 0);
 		}
 	
 		ColorAccum *= 0.5 / StepCount;
@@ -341,15 +197,15 @@ float4 MainMotionBlurCommon(float2 OutputBufferUV, float2 ScreenSpacePos, float2
 	}
 	else
 	{
-		float3 CenterVelocityDepth = PostprocessInput2.SampleLevel( PostprocessInput2Sampler, OutputBufferUV * SceneColorBufferUVToViewBufferUV.xy + SceneColorBufferUVToViewBufferUV.zw, 0 ).xyz;
+		float3 CenterVelocityDepth = VelocityFlatTexture.SampleLevel(VelocitySampler, VelocityUV, 0).xyz;
 		float  CenterDepth = CenterVelocityDepth.z;
-		float  CenterVelocityLength = GetVelocityLengthPixels( CenterVelocityDepth.xy );
-		
+		float  CenterVelocityLength = GetVelocityLengthPixels(CenterVelocityDepth.xy);
+
 		float4 ColorAccum = 0;
 		float  ColorAccumWeight = 0;
 
 		UNROLL
-		for( uint i = 0; i < StepCount; i++ )
+		for (uint i = 0; i < StepCount; i++)
 		{
 			float2 SampleUV[2];
 			float4 SampleColor[2];
@@ -357,34 +213,33 @@ float4 MainMotionBlurCommon(float2 OutputBufferUV, float2 ScreenSpacePos, float2
 			float  SampleVelocityLength[2];
 			float  Weight[2];
 
-			float2 OffsetLength = (float)i + 0.5 + float2( Random - 0.5, 0.5 - Random );
+			float2 OffsetLength = (float)i + 0.5 + float2(Random - 0.5, 0.5 - Random);
 			float2 OffsetFraction = OffsetLength / StepCount;
 
-			SampleUV[0] = OutputBufferUV + OffsetFraction.x * SearchVector.xy;
-			SampleUV[1] = OutputBufferUV + OffsetFraction.y * SearchVector.zw;
+			SampleUV[0] = ColorUV + OffsetFraction.x * SearchVector.xy;
+			SampleUV[1] = ColorUV + OffsetFraction.y * SearchVector.zw;
 
-			#if CLAMP_UV
-				SampleUV[0] = clamp( SampleUV[0], MinViewportUV, MaxViewportUV );
-				SampleUV[1] = clamp( SampleUV[1], MinViewportUV, MaxViewportUV );
-			#endif
+			SampleUV[0] = clamp(SampleUV[0], Color_UVViewportBilinearMin, Color_UVViewportBilinearMax);
+			SampleUV[1] = clamp(SampleUV[1], Color_UVViewportBilinearMin, Color_UVViewportBilinearMax);
 
 			UNROLL
-			for( uint j = 0; j < 2; j++ )
+			for (uint j = 0; j < 2; j++)
 			{
-				float3 SampleVelocityDepth = PostprocessInput2.SampleLevel(
-					PostprocessInput2Sampler, SampleUV[j] * SceneColorBufferUVToViewBufferUV.xy + SceneColorBufferUVToViewBufferUV.zw, 0 ).xyz;
-				SampleColor[j] = PostprocessInput0.SampleLevel( PostprocessInput0Sampler, SampleUV[j], 0 );
+				float3 SampleVelocityDepth = VelocityFlatTexture.SampleLevel(
+					VelocitySampler, SampleUV[j] * ColorToVelocity_Scale + ColorToVelocity_Bias, 0).xyz;
+
+				SampleColor[j] = ColorTexture.SampleLevel(ColorSampler, SampleUV[j], 0);
 				SampleDepth[j] = SampleVelocityDepth.z;
 
 				// in pixels
-				SampleVelocityLength[j] = GetVelocityLengthPixels( SampleVelocityDepth.xy );
+				SampleVelocityLength[j] = GetVelocityLengthPixels(SampleVelocityDepth.xy);
 
-				Weight[j] = SampleWeight( CenterDepth, SampleDepth[j], OffsetLength.x, CenterVelocityLength, SampleVelocityLength[j], PixelToSampleScale, SOFT_Z_EXTENT );
+				Weight[j] = SampleWeight(CenterDepth, SampleDepth[j], OffsetLength.x, CenterVelocityLength, SampleVelocityLength[j], PixelToSampleScale, SOFT_Z_EXTENT);
 			}
 		
-			bool2 Mirror = bool2( SampleDepth[0] > SampleDepth[1], SampleVelocityLength[1] > SampleVelocityLength[0] );
-			Weight[0] = all( Mirror ) ? Weight[1] : Weight[0];
-			Weight[1] = any( Mirror ) ? Weight[1] : Weight[0];
+			bool2 Mirror = bool2(SampleDepth[0] > SampleDepth[1], SampleVelocityLength[1] > SampleVelocityLength[0]);
+			Weight[0] = all(Mirror) ? Weight[1] : Weight[0];
+			Weight[1] = any(Mirror) ? Weight[1] : Weight[0];
 			
 			ColorAccum += Weight[0] * SampleColor[0] + Weight[1] * SampleColor[1];
 			ColorAccumWeight += Weight[0] + Weight[1];
@@ -393,7 +248,7 @@ float4 MainMotionBlurCommon(float2 OutputBufferUV, float2 ScreenSpacePos, float2
 		ColorAccum			*= 0.5 / StepCount;
 		ColorAccumWeight	*= 0.5 / StepCount;
 
-		OutColor = ColorAccum + ( 1 - ColorAccumWeight ) * CenterColor;
+		OutColor = ColorAccum + (1 - ColorAccumWeight) * CenterColor;
 	}
 
 #if !POST_PROCESS_ALPHA
@@ -405,37 +260,28 @@ float4 MainMotionBlurCommon(float2 OutputBufferUV, float2 ScreenSpacePos, float2
 
 void MainPS(
 	in noperspective float4 UVAndScreenPos : TEXCOORD0,
-	in float4 SvPosition : SV_Position,		// after all interpolators
-	out float4 OutColor : SV_Target0
-	)
+	in float4 SvPosition : SV_Position,
+	out float4 OutColor : SV_Target0)
 {
-	float2 OutputBufferUV = UVAndScreenPos.xy;
-	float2 ScreenPos = UVAndScreenPos.zw;
-	OutColor = MainMotionBlurCommon(OutputBufferUV, ScreenPos, floor(SvPosition.xy));
+	float2 ColorUV = UVAndScreenPos.xy;
+	OutColor = MainMotionBlurCommon(ColorUV, floor(SvPosition.xy));
 }
 
 #if COMPUTESHADER
-// xy: Unused, zw: ThreadToUVScale
-float4 MotionBlurComputeParams;
-RWTexture2D<float4> RWOutComputeTex;
+RWTexture2D<float4> OutColorTexture;
 
 [numthreads(THREADGROUP_SIZEX, THREADGROUP_SIZEY, 1)]
 void MainCS(uint2 DispatchThreadId : SV_DispatchThreadID)
 {
-	uint2 iPixelPos = DispatchThreadId + ViewportRect.xy;
-	float2 PixelPos = float2(iPixelPos);
-	float2 SvPosition = PixelPos + 0.5;
+	uint2 iColorPixelPos = DispatchThreadId + Color_ViewportMin;
+	float2 ColorPixelPos = float2(iColorPixelPos);
+	float2 ColorUV = (ColorPixelPos + 0.5) * Color_ExtentInverse;
 
-	float2 OutputBufferUV = SvPosition.xy * MotionBlurComputeParams.zw;
-	
-	float2 ViewportUV = (DispatchThreadId + 0.5) * ViewportSize.zw;
-	float2 ScreenPos = ViewportUVToScreenPos(ViewportUV);
+	float4 OutColor = MainMotionBlurCommon(ColorUV, ColorPixelPos);
 
-	float4 OutColor = MainMotionBlurCommon(OutputBufferUV, ScreenPos, PixelPos);
-
-	if (all(iPixelPos < ViewportRect.zw))
+	if (all(iColorPixelPos < Color_ViewportMax))
 	{
-		RWOutComputeTex[iPixelPos] = OutColor;
+		OutColorTexture[iColorPixelPos] = OutColor;
 	}
 }
 #endif
@@ -451,7 +297,7 @@ void OverrideWithTestChart(float2 ScreenPos, inout float2 ObjectVelocity, inout
 {
 #if MOTIONBLUR_TESTCHART == 1
 	// needs to be inside the loop to prevent NVIDIA driver optimizetion (blinking)
-	float2 PixelPos = ScreenPos * ScreenPosToPixel.xy + ScreenPosToPixel.zw + 0.5f - 25;
+	float2 PixelPos = ScreenPos * Velocity_ScreenPosToViewportScale.xy + Velocity_ScreenPosToViewportBias - 25;
 	float3 BackgroundColor = lerp(0.0, 0.3f, PseudoRandom(PixelPos)); 
 	float3 ForegroundColor = lerp(float3(1, 0, 0), float3(1, 1, 0), PseudoRandom(PixelPos)); 
 
@@ -518,16 +364,16 @@ void VisualizeMotionBlurPS(
 	const float TileSize = 48;
 
 	float2 ScreenSpacePos = UVAndScreenPos.zw;
-	float2 PixelPos = ComputePixelPosCenter(ScreenSpacePos, true);
-	float2 UV = PixelPos * PostprocessInput0Size.zw;
+	float2 PixelPos = ScreenSpacePos * Velocity_ScreenPosToViewportScale + Velocity_ScreenPosToViewportBias;
+	float2 UV = PixelPos * Velocity_ExtentInverse;
 
 	float2 PixelPosAtTileCenter = PixelPos - (frac(PixelPos / TileSize) - 0.5f) * TileSize;
-	float2 UVAtTileCenter = PixelPosAtTileCenter * PostprocessInput0Size.zw;
-	float2 ScreenPosAtTileCenter = ScreenPosFromPixelPos(PixelPosAtTileCenter, true);
+	float2 UVAtTileCenter = PixelPosAtTileCenter * Velocity_ExtentInverse;
+	float2 ScreenPosAtTileCenter = (PixelPosAtTileCenter - Velocity_ScreenPosToViewportBias) / Velocity_ScreenPosToViewportScale;
 
 	// World aligned checkerboards
 	{
-		float DeviceZ = LookupDeviceZ(UV);
+		float DeviceZ = DepthTexture.SampleLevel(VelocitySampler, UV, 0).r;
 		float SceneDepth = ConvertFromDeviceZ(DeviceZ);
 
 		float3 ScreenVector = mul(float4(ScreenSpacePos, 1, 0), View.ScreenToWorld).xyz;
@@ -540,25 +386,27 @@ void VisualizeMotionBlurPS(
 		OutColor = float4(lerp(WorldCheckerboard, float3(0,0,0), 0.7f), 1);
 	}
 
-	bool bSelectorOpaqueAtTileCenter = PostprocessInput2.SampleLevel(PostprocessInput2Sampler, UVAtTileCenter, 0).x > 0;
-	bool bSelectorOpaque = PostprocessInput2.SampleLevel(PostprocessInput2Sampler, UV, 0).x > 0;
+	bool bSelectorOpaqueAtTileCenter = VelocityTexture.SampleLevel(VelocitySampler, UVAtTileCenter, 0).x > 0;
+	bool bSelectorOpaque = VelocityTexture.SampleLevel(VelocitySampler, UV, 0).x > 0;
 
 	// relative, in screen space -1...1 -1..1, can be even outside of that range, points into the movement direction
 	float2 VelocityAtTileCenter;
 	{
-		VelocityAtTileCenter = DecodeVelocityFromTexture(PostprocessInput2.SampleLevel(PostprocessInput2Sampler, UVAtTileCenter, 0).xy);
+		VelocityAtTileCenter = DecodeVelocityFromTexture(VelocityTexture.SampleLevel(VelocitySampler, UVAtTileCenter, 0).xy);
 
 		// reconstruct from camera motion if we don't have velocity data
 		if(!bSelectorOpaqueAtTileCenter)
 		{
+			float DeviceZ = DepthTexture.SampleLevel(VelocitySampler, UVAtTileCenter, 0).r;
+
 			// UV is per tile
-			float4 ThisClip = float4(ScreenPosAtTileCenter, LookupDeviceZ(UVAtTileCenter), 1);
-			float4 PrevClip = mul( ThisClip, View.ClipToPrevClip );
+			float4 ThisClip = float4(ScreenPosAtTileCenter, DeviceZ, 1);
+			float4 PrevClip = mul(ThisClip, View.ClipToPrevClip);
 			float2 PrevScreen = PrevClip.xy / PrevClip.w;
 
 			// points into the movement direction
 			VelocityAtTileCenter = ScreenPosAtTileCenter - PrevScreen;
-		} 
+		}
 	}
 
 	// tint yellow if velocity data is stored in texture
@@ -582,7 +430,7 @@ void VisualizeMotionBlurPS(
 	float3 LineColor = bSelectorOpaqueAtTileCenter ? float3(1,1,0) : float3(0.7,0.7,0.7);
 
 	// points into the movement direction
-	float2 PixelDirection = VelocityAtTileCenter * ScreenPosToPixel.xy;
+	float2 PixelDirection = VelocityAtTileCenter * Velocity_ScreenPosToViewportScale;
 
 	// arrow
 	{
@@ -593,8 +441,6 @@ void VisualizeMotionBlurPS(
 
 		bool bArrowHead = DistOnLine < 8;
 
-//		float LocalThickness = bArrowHead ? (DistOnLine * 0.5f) : 0.9f;
-//		float LocalThickness = bArrowHead ? (DistOnLine * 0.3f) : 1 + (frac(DistOnLine/8)*8)*0.2f;
 		float LocalThickness = 1 + (frac(DistOnLine/8)*8)*0.25f;
 
 		float PerpDirectionMask = saturate(LocalThickness - abs(dot(normalize(PerpPixelDirection), DirectionInTile)));
diff --git a/Engine/Shaders/Private/PostProcessSubsurface.usf b/Engine/Shaders/Private/PostProcessSubsurface.usf
--- a/Engine/Shaders/Private/PostProcessSubsurface.usf
+++ b/Engine/Shaders/Private/PostProcessSubsurface.usf
@@ -12,6 +12,8 @@
 // for VisualizeSSS
 #include "MiniFontCommon.ush"
 
+#include "ScreenPass.ush"
+
 // x:Radius*DistanceToProjectionWindow/KernelSize*0.5, y:DistanceToProjectionWindow, zw: unused
 float4 SubsurfaceParams;
 
@@ -20,15 +22,8 @@ Texture2D SubsurfaceInput1_Texture;
 
 SamplerState SubsurfaceSampler; 
 
-float2 SubsurfaceInput0_Extent;
-float2 SubsurfaceInput0_ExtentInverse;
-uint2 SubsurfaceInput0_ViewportMin;
-float2 SubsurfaceInput0_ViewportSize;
-float2 SubsurfaceInput0_ViewportSizeInverse;
-float2 SubsurfaceInput0_UVViewportBilinearMin;
-float2 SubsurfaceInput0_UVViewportBilinearMax;
-
-float2 SubsurfaceInput1_ExtentInverse;
+SCREEN_PASS_TEXTURE_VIEWPORT(SubsurfaceInput0);
+SCREEN_PASS_TEXTURE_VIEWPORT(SubsurfaceInput1);
 
 #define SUBSURFACE_DIRECTION_HORIZONTAL 0
 #define SUBSURFACE_DIRECTION_VERTICAL 1
diff --git a/Engine/Shaders/Private/PostProcessVelocityFlatten.usf b/Engine/Shaders/Private/PostProcessVelocityFlatten.usf
--- a/Engine/Shaders/Private/PostProcessVelocityFlatten.usf
+++ b/Engine/Shaders/Private/PostProcessVelocityFlatten.usf
@@ -1,46 +1,57 @@
 // Copyright 1998-2019 Epic Games, Inc. All Rights Reserved.
 
 #include "Common.ush"
-#include "PostProcessCommon.ush"
 #include "FastMath.ush"
+#include "ScreenPass.ush"
 
 #define THREADGROUP_SIZEX		16
 #define THREADGROUP_SIZEY		16
 #define THREADGROUP_TOTALSIZE	(THREADGROUP_SIZEX * THREADGROUP_SIZEY)
 
-groupshared float4 Shared[ THREADGROUP_TOTALSIZE ];
+float MotionBlur_AspectRatio;
+float MotionBlur_VelocityScale;
+float MotionBlur_VelocityScaleForTiles;
+float MotionBlur_VelocityMax;
 
-RWTexture2D<float4>	OutVelocityFlat;
-RWTexture2D<float4>	OutMaxTileVelocity;
+groupshared float4 Shared[THREADGROUP_TOTALSIZE];
 
-float4 MotionBlurParameters;
+Texture2D VelocityTexture;
+Texture2D VelocityTileTexture;
+Texture2D DepthTexture;
 
-float4 MinMaxLength( float4 v0, float4 v1 )
+RWTexture2D<float4>	OutVelocityFlatTexture;
+RWTexture2D<float4>	OutVelocityTileTexture;
+RWTexture2D<float4>	OutDilatedVelocityTileTexture;
+
+SCREEN_PASS_TEXTURE_VIEWPORT(Velocity);
+SCREEN_PASS_TEXTURE_VIEWPORT(VelocityTile);
+
+float4 MinMaxLength(float4 v0, float4 v1)
 {
-	float2 Min = dot( v0.xy, v0.xy ) < dot( v1.xy, v1.xy ) ? v0.xy : v1.xy;
-	float2 Max = dot( v0.zw, v0.zw ) > dot( v1.zw, v1.zw ) ? v0.zw : v1.zw;
-	return float4( Min, Max );
+	float2 Min = dot(v0.xy, v0.xy) < dot(v1.xy, v1.xy) ? v0.xy : v1.xy;
+	float2 Max = dot(v0.zw, v0.zw) > dot(v1.zw, v1.zw) ? v0.zw : v1.zw;
+	return float4(Min, Max);
 }
 
-float4 MinMaxLengthPolar( float4 v0, float4 v1 )
+float4 MinMaxLengthPolar(float4 v0, float4 v1)
 {
 	float2 Min = v0.x < v1.x ? v0.xy : v1.xy;
 	float2 Max = v0.z > v1.z ? v0.zw : v1.zw;
-	return float4( Min, Max );
+	return float4(Min, Max);
 }
 
-float2 CartesianToPolar( float2 Velocity )
+float2 CartesianToPolar(float2 Velocity)
 {
-	float Length = length( Velocity );
-	float Angle = atan2Fast( Velocity.y, Velocity.x );
-	return float2( Length, Angle );
+	float Length = length(Velocity);
+	float Angle = atan2Fast(Velocity.y, Velocity.x);
+	return float2(Length, Angle);
 }
 
-float2 PolarToCartesian( float2 Velocity )
+float2 PolarToCartesian(float2 Velocity)
 {
 	float Length = Velocity.x;
 	float Angle = Velocity.y;
-	sincos( Angle, Velocity.y, Velocity.x );
+	sincos(Angle, Velocity.y, Velocity.x);
 	return Velocity * Length;
 }
 
@@ -49,145 +60,128 @@ void VelocityFlattenMain(
 	uint3 GroupId : SV_GroupID,
 	uint3 DispatchThreadId : SV_DispatchThreadID,
 	uint3 GroupThreadId : SV_GroupThreadID,
-	uint GroupIndex : SV_GroupIndex ) 
-{
-	uint2 PixelPos = min(DispatchThreadId.xy + ViewportRect.xy, ViewportRect.zw - 1);
+	uint GroupIndex : SV_GroupIndex) 
+{ 
+	uint2 PixelPos = min(DispatchThreadId.xy + Velocity_ViewportMin, Velocity_ViewportMax - 1);
 
-	float2 Velocity = PostprocessInput0[ PixelPos ].xy;
-	float Depth = PostprocessInput1[ PixelPos ].x;
+	float2 Velocity = VelocityTexture[PixelPos].xy;
+	float Depth = DepthTexture[PixelPos].x;
 
-	if( Velocity.x > 0.0 )
+	if (Velocity.x > 0.0)
 	{
-		Velocity = DecodeVelocityFromTexture( Velocity );
+		Velocity = DecodeVelocityFromTexture(Velocity);
 	}
 	else
 	{
-		float2 ViewportUV = ( (float2)DispatchThreadId.xy + 0.5 ) / ( ViewportRect.zw - ViewportRect.xy );
+		// Compute velocity due to camera motion.
+		float2 ViewportUV = ((float2)DispatchThreadId.xy + 0.5) / Velocity_ViewportSize;
 		float2 ScreenPos = 2 * float2(ViewportUV.x, 1 - ViewportUV.y) - 1;
-		float4 ThisClip = float4( ScreenPos, Depth, 1 );
-		float4 PrevClip = mul( ThisClip, View.ClipToPrevClip );
+		float4 ThisClip = float4(ScreenPos, Depth, 1);
+		float4 PrevClip = mul(ThisClip, View.ClipToPrevClip);
 		float2 PrevScreen = PrevClip.xy / PrevClip.w;
 		Velocity = ScreenPos - PrevScreen;
 	}
 
-	const float AspectRatio = MotionBlurParameters.x;
-	Velocity.y *= -AspectRatio;
+	Velocity.y *= -MotionBlur_AspectRatio;
 
 	float2 VelocityPolar = CartesianToPolar(Velocity);
-	
+
 	// If the velocity vector was zero length, VelocityPolar will contain NaNs.
 	if (any(isnan(VelocityPolar)))
 	{
 		VelocityPolar = float2(0.0f, 0.0f);
 	}
-	
-	bool bInsideViewport = all( PixelPos.xy < ViewportRect.zw );
-	
-#if 1
+
+	bool bInsideViewport = all(PixelPos.xy < Velocity_ViewportMax);
+
 	// 11:11:10  (VelocityLength, VelocityAngle, Depth)
 	float2 EncodedVelocity;
 	EncodedVelocity.x = VelocityPolar.x;
 	EncodedVelocity.y = VelocityPolar.y * (0.5 / PI) + 0.5;
-#else
-	// 11:11:10 (Velocity.xy, Depth)
-	// Only handles values in range (-2,2)
-	Velocity = clamp( Velocity, -1.9999, 1.9999 );
-	// Kill denormals
-	Velocity = abs( Velocity ) < exp2(-14.0) ? 0 : Velocity;
-	// Store sign in top bit of 6e5 exponent
-	float2 EncodedVelocity = Velocity * ( Velocity < 0 ? -32768.0 : 1 );
-#endif
-	
+
 	BRANCH
-	if( bInsideViewport )
+	if (bInsideViewport)
 	{
-		OutVelocityFlat[ PixelPos ] = float3( EncodedVelocity, ConvertFromDeviceZ( Depth ) ).xyzz;
+		OutVelocityFlatTexture[PixelPos] = float3(EncodedVelocity, ConvertFromDeviceZ(Depth)).xyzz;
 	}
 
-#if 1
 	// Limit velocity
-	VelocityPolar.x = min( VelocityPolar.x, MotionBlurParameters.w / MotionBlurParameters.y );
-#endif
+	VelocityPolar.x = min(VelocityPolar.x, MotionBlur_VelocityMax / MotionBlur_VelocityScale);
 
 	float4 VelocityMinMax = VelocityPolar.xyxy;
 	VelocityMinMax.x = bInsideViewport ? VelocityMinMax.x : 2;
 	VelocityMinMax.z = bInsideViewport ? VelocityMinMax.z : 0;
 
-	Shared[ GroupIndex ] = VelocityMinMax;
+	Shared[GroupIndex] = VelocityMinMax;
 
 	GroupMemoryBarrierWithGroupSync();
 
 #if THREADGROUP_TOTALSIZE > 512
-	if( GroupIndex < 512 )	Shared[ GroupIndex ] = MinMaxLengthPolar( Shared[ GroupIndex ], Shared[ GroupIndex + 512 ] );
+	if (GroupIndex < 512) Shared[GroupIndex] = MinMaxLengthPolar(Shared[GroupIndex], Shared[GroupIndex + 512]);
 	GroupMemoryBarrierWithGroupSync();
 #endif
 #if THREADGROUP_TOTALSIZE > 256
-	if( GroupIndex < 256 )	Shared[ GroupIndex ] = MinMaxLengthPolar( Shared[ GroupIndex ], Shared[ GroupIndex + 256 ] );
+	if (GroupIndex < 256) Shared[GroupIndex] = MinMaxLengthPolar(Shared[GroupIndex], Shared[GroupIndex + 256]);
 	GroupMemoryBarrierWithGroupSync();
 #endif
 #if THREADGROUP_TOTALSIZE > 128
-	if( GroupIndex < 128 )	Shared[ GroupIndex ] = MinMaxLengthPolar( Shared[ GroupIndex ], Shared[ GroupIndex + 128 ] );
+	if (GroupIndex < 128) Shared[GroupIndex] = MinMaxLengthPolar(Shared[GroupIndex], Shared[GroupIndex + 128]);
 	GroupMemoryBarrierWithGroupSync();
 #endif
 #if THREADGROUP_TOTALSIZE > 64
-	if( GroupIndex <  64 )	Shared[ GroupIndex ] = MinMaxLengthPolar( Shared[ GroupIndex ], Shared[ GroupIndex +  64 ] );
+	if (GroupIndex <  64) Shared[GroupIndex] = MinMaxLengthPolar(Shared[GroupIndex], Shared[GroupIndex +  64]);
 	GroupMemoryBarrierWithGroupSync();
 #endif
 
 	// Safe for vector sizes 32 or larger, AMD and NV
 	// TODO Intel variable size vector
-	if( GroupIndex < 32 )	Shared[ GroupIndex ] = MinMaxLengthPolar( Shared[ GroupIndex ], Shared[ GroupIndex + 32 ] );
-	if( GroupIndex < 16 )	Shared[ GroupIndex ] = MinMaxLengthPolar( Shared[ GroupIndex ], Shared[ GroupIndex + 16 ] );
-	if( GroupIndex <  8 )	Shared[ GroupIndex ] = MinMaxLengthPolar( Shared[ GroupIndex ], Shared[ GroupIndex +  8 ] );
-	if( GroupIndex <  4 )	Shared[ GroupIndex ] = MinMaxLengthPolar( Shared[ GroupIndex ], Shared[ GroupIndex +  4 ] );
-	if( GroupIndex <  2 )	Shared[ GroupIndex ] = MinMaxLengthPolar( Shared[ GroupIndex ], Shared[ GroupIndex +  2 ] );
-	if( GroupIndex <  1 )	Shared[ GroupIndex ] = MinMaxLengthPolar( Shared[ GroupIndex ], Shared[ GroupIndex +  1 ] );
-
-	if( GroupIndex == 0 )
+	if (GroupIndex < 32) Shared[GroupIndex] = MinMaxLengthPolar(Shared[GroupIndex], Shared[GroupIndex + 32]);
+	if (GroupIndex < 16) Shared[GroupIndex] = MinMaxLengthPolar(Shared[GroupIndex], Shared[GroupIndex + 16]);
+	if (GroupIndex <  8) Shared[GroupIndex] = MinMaxLengthPolar(Shared[GroupIndex], Shared[GroupIndex +  8]);
+	if (GroupIndex <  4) Shared[GroupIndex] = MinMaxLengthPolar(Shared[GroupIndex], Shared[GroupIndex +  4]);
+	if (GroupIndex <  2) Shared[GroupIndex] = MinMaxLengthPolar(Shared[GroupIndex], Shared[GroupIndex +  2]);
+	if (GroupIndex <  1) Shared[GroupIndex] = MinMaxLengthPolar(Shared[GroupIndex], Shared[GroupIndex +  1]);
+
+	if (GroupIndex == 0)
 	{
-		OutMaxTileVelocity[ GroupId.xy ] = float4( PolarToCartesian( Shared[0].xy ), PolarToCartesian( Shared[0].zw ) );
+		OutVelocityTileTexture[GroupId.xy] = float4(PolarToCartesian(Shared[0].xy), PolarToCartesian(Shared[0].zw));
 	}
 }
 
-
-
-RWTexture2D<float4>	OutScatteredMaxVelocity;
-
 [numthreads(16, 16, 1)]
 void VelocityGatherCS(
 	uint3 GroupId : SV_GroupID,
 	uint3 DispatchThreadId : SV_DispatchThreadID,
 	uint3 GroupThreadId : SV_GroupThreadID,
-	uint GroupIndex : SV_GroupIndex ) 
+	uint GroupIndex : SV_GroupIndex)
 {
 	uint2 PixelPos = DispatchThreadId.xy;
-	float2 UV = ( (float2)PixelPos.xy + 0.5 ) * PostprocessInput0Size.zw;
+	float2 UV = ((float2)PixelPos.xy + 0.5) * VelocityTile_ExtentInverse;
 
-	float4 MinMaxVelocity = PostprocessInput0[ PixelPos ];
+	float4 MinMaxVelocity = VelocityTileTexture[PixelPos];
 
-#if 1
 	// Scatter as gather
-	for( int x = -3; x <= 3; x++ )
+	for(int x = -3; x <= 3; x++)
 	{
-		for( int y = -3; y <= 3; y++ )
+		for(int y = -3; y <= 3; y++)
 		{
-			if( x == 0 && y == 0 )
+			if (x == 0 && y == 0)
 				continue;
 
 			int2 Offset = int2(x,y);
 			int2 SampleIndex = PixelPos + Offset;
 			
-			bool2 bInsideViewport = 0 <= SampleIndex && SampleIndex < (int2)ViewportRect.zw;
-			if( !all( bInsideViewport ) )
+			bool2 bInsideViewport = 0 <= SampleIndex && SampleIndex < (int2)VelocityTile_ViewportMax;
+			if (!all(bInsideViewport))
 				continue;
 
-			float4 ScatterMinMax = PostprocessInput0[ SampleIndex ];
+			float4 ScatterMinMax = VelocityTileTexture[SampleIndex];
 			float2 MaxVelocity = ScatterMinMax.zw;
 
-			float2 VelocityPixels = MaxVelocity * MotionBlurParameters.z;
+			float2 VelocityPixels = MaxVelocity * MotionBlur_VelocityScaleForTiles;
 
-			float  VelocityLengthPixelsSqr = dot( VelocityPixels, VelocityPixels );
-			float  VelocityLengthPixelsInv = rsqrtFast( VelocityLengthPixelsSqr + 1e-8 );
+			float  VelocityLengthPixelsSqr = dot(VelocityPixels, VelocityPixels);
+			float  VelocityLengthPixelsInv = rsqrtFast(VelocityLengthPixelsSqr + 1e-8);
 			float  VelocityLengthPixels = VelocityLengthPixelsSqr * VelocityLengthPixelsInv;
 			float2 VelocityDir = VelocityPixels * VelocityLengthPixelsInv;
 
@@ -195,30 +189,111 @@ void VelocityGatherCS(
 			// 1/2 pixel because shape is swept tile
 			// +1/2 pixel for conservative rasterization
 			// 99% to give epsilon before neighbor is filled. Otherwise all neighbors lie on edges of quad when no velocity in their direction.
-			float PixelExtent = abs( VelocityDir.x ) + abs( VelocityDir.y );
-			float2 QuadExtent = float2( VelocityLengthPixels, 0 ) + PixelExtent.xx * 0.99;
+			float PixelExtent = abs(VelocityDir.x) + abs(VelocityDir.y);
+			float2 QuadExtent = float2(VelocityLengthPixels, 0) + PixelExtent.xx * 0.99;
 
 			// Orient quad along velocity direction
 			float2 AxisX = VelocityDir;
-			float2 AxisY = float2( -VelocityDir.y, VelocityDir.x );
+			float2 AxisY = float2(-VelocityDir.y, VelocityDir.x);
 
 			// Project this pixel center onto scatter quad
 			float2 PixelCenterOnQuad;
-			PixelCenterOnQuad.x = dot( AxisX, Offset );
-			PixelCenterOnQuad.y = dot( AxisY, Offset );
+			PixelCenterOnQuad.x = dot(AxisX, Offset);
+			PixelCenterOnQuad.y = dot(AxisY, Offset);
 
-			bool2 bInsideQuad = abs( PixelCenterOnQuad ) < QuadExtent;
-			if( all( bInsideQuad ) ) 
+			bool2 bInsideQuad = abs(PixelCenterOnQuad) < QuadExtent;
+			if (all(bInsideQuad)) 
 			{
-				MinMaxVelocity = MinMaxLength( MinMaxVelocity, ScatterMinMax );
+				MinMaxVelocity = MinMaxLength(MinMaxVelocity, ScatterMinMax);
 			}
 		}
 	}
-#endif
 
 	BRANCH
-	if( all( PixelPos.xy < ViewportRect.zw ) )
+	if (all(PixelPos.xy < VelocityTile_ViewportMax))
 	{
-		OutScatteredMaxVelocity[ PixelPos ] = MinMaxVelocity;
+		OutDilatedVelocityTileTexture[PixelPos] = MinMaxVelocity;
 	}
+}
+
+#define VELOCITY_SCATTER_PASS_MIN 0
+#define VELOCITY_SCATTER_PASS_MAX 1
+
+uint ScatterPass;
+
+void VelocityScatterVS(
+	uint VId : SV_VertexID,
+	uint IId : SV_InstanceID,
+	out nointerpolation float4 OutColor : TEXCOORD0,
+	out float4 OutPosition : SV_POSITION)
+{
+	OutPosition = float4(0, 0, 0, 1);
+
+	// needs to be the same on C++ side (faster on NVIDIA and AMD)
+	uint QuadsPerInstance = 8;
+
+	// remap the indices to get vertexid to VId and quadid into IId
+	IId = IId * QuadsPerInstance + (VId / 4);
+	VId = VId % 4;
+
+	// triangle A: 0:left top, 1:right top, 2: left bottom
+	// triangle B: 3:right bottom, 4:left bottom, 5: right top
+	float2 CornerOffset = float2(VId % 2, VId / 2) * 2 - 1;
+
+	uint2 PixelPos = uint2(IId % VelocityTile_ViewportMax.x, IId / VelocityTile_ViewportMax.x);
+
+	BRANCH
+	if (PixelPos.y >= VelocityTile_ViewportMax.y)
+	{
+		OutColor = 0;
+		return;
+	}
+
+	float2 SvPosition = PixelPos + 0.5;
+
+	float4 MinMaxVelocity = VelocityTileTexture[PixelPos];
+	OutColor = MinMaxVelocity;
+
+	float4 MinMaxVelocityPixels = MinMaxVelocity * MotionBlur_VelocityScaleForTiles;
+	float2 VelocityPixels = MinMaxVelocityPixels.zw;
+
+	// Is the velocity small enough not to cover adjacent tiles?
+	BRANCH
+	if (dot(VelocityPixels, VelocityPixels) * 16 * 16 <= 0.25)
+	{
+		OutPosition.xy = (SvPosition + CornerOffset * 0.5 - VelocityTile_ScreenPosToViewportBias) / VelocityTile_ScreenPosToViewportScale.xy;
+		OutPosition.z = 0.0002; // zero clips
+		return;
+	}
+
+	float  VelocityLengthPixelsSqr = dot(VelocityPixels, VelocityPixels);
+	float  VelocityLengthPixelsInv = rsqrtFast(VelocityLengthPixelsSqr);
+	float  VelocityLengthPixels = VelocityLengthPixelsSqr * VelocityLengthPixelsInv;
+	float2 VelocityDir = VelocityPixels * VelocityLengthPixelsInv;
+
+	// Project pixel corner on to dir. This is the oriented extent of a pixel.
+	// 1/2 pixel because shape is swept tile
+	// +1/2 pixel for conservative rasterization
+	// 99% to give epsilon before neighbor is filled. Otherwise all neighbors lie on edges of quad when no velocity in their direction.
+	// dot(abs(VelocityDir), float2(1, 1))
+	float Extent = abs(VelocityDir.x) + abs(VelocityDir.y);
+	CornerOffset *= float2(VelocityLengthPixels, 0) + Extent.xx * 0.99;
+
+	// Orient along velocity direction
+	float2 AxisX = VelocityDir;
+	float2 AxisY = float2(-VelocityDir.y, VelocityDir.x);
+	CornerOffset = AxisX * CornerOffset.x + AxisY * CornerOffset.y;
+
+	OutPosition.xy = (SvPosition + CornerOffset - VelocityTile_ScreenPosToViewportBias) / VelocityTile_ScreenPosToViewportScale;
+
+	// Depth ordered by velocity length
+	OutPosition.z = (ScatterPass == VELOCITY_SCATTER_PASS_MAX) ? VelocityLengthPixels : length(MinMaxVelocityPixels.xy);
+	OutPosition.z = clamp(OutPosition.z / VelocityTile_ScreenPosToViewportScale.x * 0.5, 0.0002, 0.999);
+}
+
+void VelocityScatterPS(
+	nointerpolation float4 InColor : TEXCOORD0,
+	out float4 OutColor : SV_Target0)
+{
+	OutColor = InColor;
 }
\ No newline at end of file
diff --git a/Engine/Shaders/Private/ScreenPass.ush b/Engine/Shaders/Private/ScreenPass.ush
new file mode 100644
--- /dev/null
+++ b/Engine/Shaders/Private/ScreenPass.ush
@@ -0,0 +1,25 @@
+// Copyright 1998-2019 Epic Games, Inc. All Rights Reserved.
+
+#pragma once
+
+#define SCREEN_PASS_STRUCT_MEMBER(StructName, MemberType, MemberName) MemberType StructName##_##MemberName;
+
+#define SCREEN_PASS_TEXTURE_VIEWPORT(StructName) \
+	SCREEN_PASS_STRUCT_MEMBER(StructName, float2, Extent) \
+	SCREEN_PASS_STRUCT_MEMBER(StructName, float2, ExtentInverse) \
+	SCREEN_PASS_STRUCT_MEMBER(StructName, float2, ScreenPosToViewportScale) \
+	SCREEN_PASS_STRUCT_MEMBER(StructName, float2, ScreenPosToViewportBias) \
+	SCREEN_PASS_STRUCT_MEMBER(StructName, uint2, ViewportMin) \
+	SCREEN_PASS_STRUCT_MEMBER(StructName, uint2, ViewportMax) \
+	SCREEN_PASS_STRUCT_MEMBER(StructName, float2, ViewportSize) \
+	SCREEN_PASS_STRUCT_MEMBER(StructName, float2, ViewportSizeInverse) \
+	SCREEN_PASS_STRUCT_MEMBER(StructName, float2, UVViewportMin) \
+	SCREEN_PASS_STRUCT_MEMBER(StructName, float2, UVViewportMax) \
+	SCREEN_PASS_STRUCT_MEMBER(StructName, float2, UVViewportSize) \
+	SCREEN_PASS_STRUCT_MEMBER(StructName, float2, UVViewportSizeInverse) \
+	SCREEN_PASS_STRUCT_MEMBER(StructName, float2, UVViewportBilinearMin) \
+	SCREEN_PASS_STRUCT_MEMBER(StructName, float2, UVViewportBilinearMax)
+
+#define SCREEN_PASS_TEXTURE_VIEWPORT_TRANSFORM(StructName) \
+	SCREEN_PASS_STRUCT_MEMBER(StructName, float2, Scale) \
+	SCREEN_PASS_STRUCT_MEMBER(StructName, float2, Bias)
\ No newline at end of file
diff --git a/Engine/Source/Runtime/RenderCore/Private/RenderGraphBuilder.cpp b/Engine/Source/Runtime/RenderCore/Private/RenderGraphBuilder.cpp
--- a/Engine/Source/Runtime/RenderCore/Private/RenderGraphBuilder.cpp
+++ b/Engine/Source/Runtime/RenderCore/Private/RenderGraphBuilder.cpp
@@ -382,7 +382,7 @@ void FRDGBuilder::ValidatePass(const FRenderGraphPass* Pass) const
 			checkf(RenderTargets->DepthStencil.StencilLoadAction != ERenderTargetLoadAction::ELoad,
 				TEXT("Can't load stencil from a render target that has never been produced."));
 			checkf(RenderTargets->DepthStencil.StencilStoreAction != ERenderTargetStoreAction::ENoAction,
-				TEXT("Using for the first time a render target without storing it, witch doesn't make much sens."));
+				TEXT("Using for the first time a render target without storing it, which doesn't make much sense."));
 
 			Texture->bHasEverBeenProduced = true;
 			Texture->DebugFirstProducer = Pass;
diff --git a/Engine/Source/Runtime/RenderCore/Public/RenderGraphBuilder.h b/Engine/Source/Runtime/RenderCore/Public/RenderGraphBuilder.h
--- a/Engine/Source/Runtime/RenderCore/Public/RenderGraphBuilder.h
+++ b/Engine/Source/Runtime/RenderCore/Public/RenderGraphBuilder.h
@@ -9,14 +9,17 @@
 /** Whether visualize texture tool is supported. */
 #define SUPPORTS_VISUALIZE_TEXTURE (WITH_ENGINE && !(UE_BUILD_SHIPPING || UE_BUILD_TEST))
 
-
 /** Whether render graph should support draw events or not.
  * RENDER_GRAPH_DRAW_EVENTS == 0 means there is no string processing at all.
  * RENDER_GRAPH_DRAW_EVENTS == 1 means const TCHAR* is only passdown.
  * RENDER_GRAPH_DRAW_EVENTS == 2 means complex formated FString is passdown.
  */
+#define RENDER_GRAPH_DRAW_EVENTS_NONE 0
+#define RENDER_GRAPH_DRAW_EVENTS_STRING_REF 1
+#define RENDER_GRAPH_DRAW_EVENTS_STRING_COPY 2
+
 #if WITH_PROFILEGPU
-	#define RENDER_GRAPH_DRAW_EVENTS 2
+	#define RENDER_GRAPH_DRAW_EVENTS RENDER_GRAPH_DRAW_EVENTS_STRING_COPY
 #else
 	#define RENDER_GRAPH_DRAW_EVENTS 0
 #endif
@@ -26,24 +29,57 @@
 class RENDERCORE_API FRDGEventName
 {
 public:
-	inline FRDGEventName() 
-	{ }
+	FRDGEventName() = default;
 
-	#if RENDER_GRAPH_DRAW_EVENTS == 2
+	#if RENDER_GRAPH_DRAW_EVENTS == RENDER_GRAPH_DRAW_EVENTS_STRING_COPY
 
 	explicit FRDGEventName(const TCHAR* EventFormat, ...);
 
-	#elif RENDER_GRAPH_DRAW_EVENTS == 1
+	#elif RENDER_GRAPH_DRAW_EVENTS == RENDER_GRAPH_DRAW_EVENTS_STRING_REF
 
 	explicit inline FRDGEventName(const TCHAR* EventFormat)
 		: EventName(EventFormat)
-	{ }
+	{}
 
 	#endif
-	
+
+	FRDGEventName(const FRDGEventName& Other)
+	{
+		*this = Other;
+	}
+
+	FRDGEventName(FRDGEventName&& Other)
+	{
+		*this = Forward<FRDGEventName>(Other);
+	}
+
+	FRDGEventName& operator=(const FRDGEventName& Other)
+	{
+#if RENDER_GRAPH_DRAW_EVENTS == RENDER_GRAPH_DRAW_EVENTS_STRING_REF
+		EventName = Other.EventName;
+#elif RENDER_GRAPH_DRAW_EVENTS == RENDER_GRAPH_DRAW_EVENTS_STRING_COPY
+		EventNameStorage = Other.EventNameStorage;
+		EventName = *EventNameStorage;
+#endif
+		return *this;
+	}
+
+	FRDGEventName& operator=(FRDGEventName&& Other)
+	{
+#if RENDER_GRAPH_DRAW_EVENTS == RENDER_GRAPH_DRAW_EVENTS_STRING_REF
+		EventName = Other.EventName;
+		Other.EventName = nullptr;
+#elif RENDER_GRAPH_DRAW_EVENTS == RENDER_GRAPH_DRAW_EVENTS_STRING_COPY
+		EventNameStorage = MoveTemp(Other.EventNameStorage);
+		EventName = *EventNameStorage;
+		Other.EventName = nullptr;
+#endif
+		return *this;
+	}
+
 	const TCHAR* GetTCHAR() const
 	{
-		#if RENDER_GRAPH_DRAW_EVENTS == 1 || RENDER_GRAPH_DRAW_EVENTS == 2
+		#if RENDER_GRAPH_DRAW_EVENTS == RENDER_GRAPH_DRAW_EVENTS_STRING_REF || RENDER_GRAPH_DRAW_EVENTS == RENDER_GRAPH_DRAW_EVENTS_STRING_COPY
 			return EventName;
 		#else
 			// Render graph draw events have been completely compiled for CPU performance reasons.
@@ -52,9 +88,9 @@ public:
 	}
 
 private:
-	#if RENDER_GRAPH_DRAW_EVENTS == 1 || RENDER_GRAPH_DRAW_EVENTS == 2
+	#if RENDER_GRAPH_DRAW_EVENTS == RENDER_GRAPH_DRAW_EVENTS_STRING_REF || RENDER_GRAPH_DRAW_EVENTS == RENDER_GRAPH_DRAW_EVENTS_STRING_COPY
 		const TCHAR* EventName;
-		#if RENDER_GRAPH_DRAW_EVENTS == 2
+		#if RENDER_GRAPH_DRAW_EVENTS == RENDER_GRAPH_DRAW_EVENTS_STRING_COPY
 			FString EventNameStorage;
 		#endif
 	#endif
@@ -75,8 +111,6 @@ private:
 	FRDGEventScope(const FRDGEventScope* InParentScope, FRDGEventName&& InName)
 		: ParentScope(InParentScope), Name(InName)
 	{ }
-
-
 	friend class FRDGBuilder;
 	friend class FStackRDGEventScopeRef;
 };
@@ -661,7 +695,7 @@ public:
 
 		auto NewScope = new(FMemStack::Get()) FRDGEventScope(GraphBuilder.CurrentScope, Forward<FRDGEventName>(ScopeName));
 
-		#if RENDER_GRAPH_DRAW_EVENTS == 2
+		#if RENDER_GRAPH_DRAW_EVENTS == RENDER_GRAPH_DRAW_EVENTS_STRING_COPY
 		{
 			GraphBuilder.EventScopes.Add(NewScope);
 		}
@@ -689,13 +723,13 @@ private:
  *
  *		RDG_EVENT_SCOPE(GraphBuilder, "MyProcessing %sx%s", ViewRect.Width(), ViewRect.Height());
  */
-#if RENDER_GRAPH_DRAW_EVENTS == 2
+#if RENDER_GRAPH_DRAW_EVENTS == RENDER_GRAPH_DRAW_EVENTS_STRING_COPY
 
 #define RDG_EVENT_NAME(Format, ...) FRDGEventName(TEXT(Format), ##__VA_ARGS__)
 #define RDG_EVENT_SCOPE(GraphBuilder, Format, ...) \
 	FStackRDGEventScopeRef PREPROCESSOR_JOIN(__RDG_ScopeRef_,__LINE__) ((GraphBuilder), RDG_EVENT_NAME(Format, ##__VA_ARGS__))
 
-#elif RENDER_GRAPH_DRAW_EVENTS == 1
+#elif RENDER_GRAPH_DRAW_EVENTS == RENDER_GRAPH_DRAW_EVENTS_STRING_COPY
 
 #define RDG_EVENT_NAME(Format, ...) FRDGEventName(TEXT(Format))
 #define RDG_EVENT_SCOPE(GraphBuilder, Format, ...) \
diff --git a/Engine/Source/Runtime/RenderCore/Public/ShaderParameterMacros.h b/Engine/Source/Runtime/RenderCore/Public/ShaderParameterMacros.h
--- a/Engine/Source/Runtime/RenderCore/Public/ShaderParameterMacros.h
+++ b/Engine/Source/Runtime/RenderCore/Public/ShaderParameterMacros.h
@@ -187,12 +187,36 @@ private:
 /** Render graph information about how to bind a depth-stencil render target. */
 struct alignas(SHADER_PARAMETER_STRUCT_ALIGNMENT) FDepthStencilBinding
 {
-	FDepthStencilBinding()
-		: Texture(nullptr)
-	{ }
+	FDepthStencilBinding() = default;
 	FDepthStencilBinding(const FDepthStencilBinding&) = default;
 
-	TAlignedShaderParameterPtr<FRDGTexture*> Texture;
+	FDepthStencilBinding(
+		FRDGTexture* InTexture,
+		ERenderTargetLoadAction InDepthLoadAction,
+		ERenderTargetStoreAction InDepthStoreAction,
+		ERenderTargetLoadAction InStencilLoadAction,
+		ERenderTargetStoreAction InStencilStoreAction,
+		FExclusiveDepthStencil InDepthStencilAccess)
+		: Texture(InTexture)
+		, DepthLoadAction(InDepthLoadAction)
+		, DepthStoreAction(InDepthStoreAction)
+		, StencilLoadAction(InStencilLoadAction)
+		, StencilStoreAction(InStencilStoreAction)
+		, DepthStencilAccess(InDepthStencilAccess)
+	{}
+
+	FDepthStencilBinding(
+		FRDGTexture* InTexture,
+		ERenderTargetLoadAction InDepthLoadAction,
+		ERenderTargetStoreAction InDepthStoreAction,
+		FExclusiveDepthStencil InDepthStencilAccess)
+		: Texture(InTexture)
+		, DepthLoadAction(InDepthLoadAction)
+		, DepthStoreAction(InDepthStoreAction)
+		, DepthStencilAccess(InDepthStencilAccess)
+	{}
+
+	TAlignedShaderParameterPtr<FRDGTexture*> Texture = nullptr;
 	ERenderTargetLoadAction		DepthLoadAction		= ERenderTargetLoadAction::ENoAction;
 	ERenderTargetStoreAction	DepthStoreAction	= ERenderTargetStoreAction::ENoAction;
 	ERenderTargetLoadAction		StencilLoadAction	= ERenderTargetLoadAction::ENoAction;
diff --git a/Engine/Source/Runtime/Renderer/Private/PostProcess/PostProcessMotionBlur.cpp b/Engine/Source/Runtime/Renderer/Private/PostProcess/PostProcessMotionBlur.cpp
--- a/Engine/Source/Runtime/Renderer/Private/PostProcess/PostProcessMotionBlur.cpp
+++ b/Engine/Source/Runtime/Renderer/Private/PostProcess/PostProcessMotionBlur.cpp
@@ -1,1253 +1,865 @@
 // Copyright 1998-2019 Epic Games, Inc. All Rights Reserved.
 
-/*=============================================================================
-	PostProcessMotionBlur.cpp: Post process MotionBlur implementation.
-=============================================================================*/
-
 #include "PostProcess/PostProcessMotionBlur.h"
 #include "StaticBoundShaderState.h"
 #include "CanvasTypes.h"
 #include "RenderTargetTemp.h"
 #include "SceneUtils.h"
-#include "PostProcess/SceneRenderTargets.h"
 #include "SceneRenderTargetParameters.h"
 #include "ScenePrivate.h"
 #include "PostProcess/SceneFilterRendering.h"
-#include "CompositionLighting/PostProcessAmbientOcclusion.h"
+#include "PostProcess/SceneRenderTargets.h"
 #include "PostProcess/PostProcessing.h"
+#include "CompositionLighting/PostProcessAmbientOcclusion.h"
 #include "DeferredShadingRenderer.h"
 #include "ClearQuad.h"
 #include "PipelineStateCache.h"
 #include "SpriteIndexBuffer.h"
 
+namespace
+{
 #if !(UE_BUILD_SHIPPING || UE_BUILD_TEST)
-static TAutoConsoleVariable<int32> CVarMotionBlurFiltering(
-	TEXT("r.MotionBlurFiltering"),
-	0,
-	TEXT("Useful developer variable\n")
-	TEXT("0: off (default, expected by the shader for better quality)\n")
-	TEXT("1: on"),
-	ECVF_Cheat | ECVF_RenderThreadSafe);
+	TAutoConsoleVariable<int32> CVarMotionBlurFiltering(
+		TEXT("r.MotionBlurFiltering"),
+		0,
+		TEXT("Useful developer variable\n")
+		TEXT("0: off (default, expected by the shader for better quality)\n")
+		TEXT("1: on"),
+		ECVF_Cheat | ECVF_RenderThreadSafe);
 #endif
 
-static TAutoConsoleVariable<float> CVarMotionBlur2ndScale(
-	TEXT("r.MotionBlur2ndScale"),
-	1.0f,
-	TEXT(""),
-	ECVF_Cheat | ECVF_RenderThreadSafe);
-
-const int32 kMotionBlurTileSize = 16;
-
-const int32 GMotionBlurComputeTileSizeX = 8;
-const int32 GMotionBlurComputeTileSizeY = 8;
-
-FIntPoint GetNumTiles16x16( FIntPoint PixelExtent )
-{
-	uint32 TilesX = FMath::DivideAndRoundUp(PixelExtent.X, kMotionBlurTileSize);
-	uint32 TilesY = FMath::DivideAndRoundUp(PixelExtent.Y, kMotionBlurTileSize);
-	return FIntPoint( TilesX, TilesY );
-}
-
-FVector4 GetMotionBlurParameters( const FRenderingCompositePassContext& Context, float Scale = 1.0f )
-{
-	const float TileSize = kMotionBlurTileSize;
-
-	const float SizeX = Context.SceneColorViewRect.Width();
-	const float SizeY = Context.SceneColorViewRect.Height();
-	const float AspectRatio = SizeY / SizeX;
-
-	const FSceneViewState* ViewState = (FSceneViewState*) Context.ViewState;
-	float MotionBlurTimeScale = ViewState ? ViewState->MotionBlurTimeScale : 1.0f;
-	float MotionBlurScale = 0.5f * MotionBlurTimeScale * Context.View.FinalPostProcessSettings.MotionBlurAmount;
-
-	// 0:no 1:full screen width, percent conversion
-	float MaxVelocity = Context.View.FinalPostProcessSettings.MotionBlurMax / 100.0f;
-
-	// Scale by 0.5 due to blur samples going both ways
-	float PixelScale = Scale * SizeX * 0.5f;
-
-	FVector4 MotionBlurParameters(
-		AspectRatio,
-		PixelScale * MotionBlurScale,			// Scale for pixels
-		PixelScale * MotionBlurScale / TileSize,// Scale for tiles
-		FMath::Abs( PixelScale ) * MaxVelocity	// Max velocity pixels
-	);
-
-	return MotionBlurParameters;
-}
-
-FVector4 GetSceneColorBufferUVToViewBufferUV(const FRenderingCompositePassContext& Context)
-{
-	const FSceneRenderTargets& SceneContext = FSceneRenderTargets::Get(Context.RHICmdList);
-
-	FVector2D SceneColorOffset = FVector2D(Context.SceneColorViewRect.Min) / FVector2D(Context.ReferenceBufferSize);
-	FVector2D SceneColorExtent = FVector2D(Context.SceneColorViewRect.Size()) / FVector2D(Context.ReferenceBufferSize);
-	FVector2D ViewOffset = FVector2D(Context.View.ViewRect.Min) / FVector2D(SceneContext.GetBufferSizeXY());
-	FVector2D ViewExtent = FVector2D(Context.View.ViewRect.Size()) / FVector2D(SceneContext.GetBufferSizeXY());
-
-	FVector4 SceneColorBufferUVToViewBufferUVValue;
-
-	SceneColorBufferUVToViewBufferUVValue.X = ViewExtent.X / SceneColorExtent.X;
-	SceneColorBufferUVToViewBufferUVValue.Y = ViewExtent.Y / SceneColorExtent.Y;
-	SceneColorBufferUVToViewBufferUVValue.Z = ViewOffset.X - SceneColorBufferUVToViewBufferUVValue.X * SceneColorOffset.X;
-	SceneColorBufferUVToViewBufferUVValue.W = ViewOffset.Y - SceneColorBufferUVToViewBufferUVValue.Y * SceneColorOffset.Y;
-
-	return SceneColorBufferUVToViewBufferUVValue;
-}
-
-class FPostProcessVelocityFlattenCS : public FGlobalShader
-{
-	DECLARE_SHADER_TYPE(FPostProcessVelocityFlattenCS, Global);
-
-	static bool ShouldCompilePermutation(const FGlobalShaderPermutationParameters& Parameters)
-	{
-		return IsFeatureLevelSupported(Parameters.Platform, ERHIFeatureLevel::SM5);
+	TAutoConsoleVariable<float> CVarMotionBlur2ndScale(
+		TEXT("r.MotionBlur2ndScale"),
+		1.0f,
+		TEXT(""),
+		ECVF_Cheat | ECVF_RenderThreadSafe);
+
+	TAutoConsoleVariable<int32> CVarMotionBlurScatter(
+		TEXT("r.MotionBlurScatter"),
+		0,
+		TEXT("Forces scatter based max velocity method (slower)."),
+		ECVF_RenderThreadSafe);
+
+	TAutoConsoleVariable<int32> CVarMotionBlurSeparable(
+		TEXT("r.MotionBlurSeparable"),
+		0,
+		TEXT("Adds a second motion blur pass that smooths noise for a higher quality blur."),
+		ECVF_RenderThreadSafe);
+
+	TAutoConsoleVariable<int32> CVarMotionBlurPreferCompute(
+		TEXT("r.MotionBlur.PreferCompute"),
+		0,
+		TEXT("Will use compute shaders for motion blur pass."),
+		ECVF_RenderThreadSafe);
+
+	TAutoConsoleVariable<int32> CVarAllowMotionBlurInVR(
+		TEXT("vr.AllowMotionBlurInVR"),
+		0,
+		TEXT("For projects with motion blur enabled, this allows motion blur to be enabled even while in VR."));
+
+	FMatrix GetPreviousWorldToClipMatrix(const FViewInfo& View)
+	{
+		if (View.Family->EngineShowFlags.CameraInterpolation)
+		{
+			// Instead of finding the world space position of the current pixel, calculate the world space position offset by the camera position, 
+			// then translate by the difference between last frame's camera position and this frame's camera position,
+			// then apply the rest of the transforms.  This effectively avoids precision issues near the extents of large levels whose world space position is very large.
+			FVector ViewOriginDelta = View.ViewMatrices.GetViewOrigin() - View.PrevViewInfo.ViewMatrices.GetViewOrigin();
+			return FTranslationMatrix(ViewOriginDelta) * View.PrevViewInfo.ViewMatrices.ComputeViewRotationProjectionMatrix();
+		}
+		else
+		{
+			return View.ViewMatrices.ComputeViewRotationProjectionMatrix();
+		}
 	}
 
-	static void ModifyCompilationEnvironment(const FGlobalShaderPermutationParameters& Parameters, FShaderCompilerEnvironment& OutEnvironment)
+	int32 GetMotionBlurQualityFromCVar()
 	{
-		FGlobalShader::ModifyCompilationEnvironment(Parameters, OutEnvironment);
-	}
+		int32 MotionBlurQuality;
 
-	FPostProcessVelocityFlattenCS() {}
+		static const auto ICVar = IConsoleManager::Get().FindTConsoleVariableDataInt(TEXT("r.MotionBlurQuality"));
+		MotionBlurQuality = FMath::Clamp(ICVar->GetValueOnRenderThread(), 0, 4);
 
-public:
-	FShaderParameter		OutVelocityFlat;		// UAV
-	FShaderParameter		OutMaxTileVelocity;		// UAV
-
-	FPostProcessVelocityFlattenCS(const ShaderMetaType::CompiledShaderInitializerType& Initializer)
-		: FGlobalShader(Initializer)
-	{
-		PostprocessParameter.Bind(Initializer.ParameterMap);
-		MotionBlurParameters.Bind(Initializer.ParameterMap, TEXT("MotionBlurParameters"));
-		OutVelocityFlat.Bind(Initializer.ParameterMap, TEXT("OutVelocityFlat"));
-		OutMaxTileVelocity.Bind(Initializer.ParameterMap, TEXT("OutMaxTileVelocity"));
+		return MotionBlurQuality;
 	}
+}
 
-	void SetCS( FRHICommandList& RHICmdList, const FRenderingCompositePassContext& Context, const FSceneView& View )
-	{
-		const FComputeShaderRHIParamRef ShaderRHI = GetComputeShader();
-
-		FGlobalShader::SetParameters<FViewUniformShaderParameters>(RHICmdList, ShaderRHI, Context.View.ViewUniformBuffer);
-		PostprocessParameter.SetCS(ShaderRHI, Context, Context.RHICmdList, TStaticSamplerState<SF_Point,AM_Clamp,AM_Clamp,AM_Clamp>::GetRHI());
+const int32 kMotionBlurTileSize = 16;
+const int32 kMotionBlurComputeTileSizeX = 8;
+const int32 kMotionBlurComputeTileSizeY = 8;
 
-		SetShaderValue(Context.RHICmdList, ShaderRHI, MotionBlurParameters, GetMotionBlurParameters( Context ) );
-	}
-	
-	virtual bool Serialize(FArchive& Ar) override
+bool IsMotionBlurEnabled(const FViewInfo& View)
+{
+	if (View.GetFeatureLevel() < ERHIFeatureLevel::SM5)
 	{
-		bool bShaderHasOutdatedParameters = FGlobalShader::Serialize(Ar);
-		Ar << PostprocessParameter;
-		Ar << MotionBlurParameters;
-		Ar << OutVelocityFlat;
-		Ar << OutMaxTileVelocity;
-		return bShaderHasOutdatedParameters;
+		return false;
 	}
 
-private:
-	FPostProcessPassParameters	PostprocessParameter;
-	FShaderParameter			MotionBlurParameters;
-};
-
-IMPLEMENT_SHADER_TYPE(,FPostProcessVelocityFlattenCS,TEXT("/Engine/Private/PostProcessVelocityFlatten.usf"),TEXT("VelocityFlattenMain"),SF_Compute);
+	const int32 MotionBlurQuality = GetMotionBlurQualityFromCVar();
 
+	const FSceneViewFamily& ViewFamily = *View.Family;
 
-FRCPassPostProcessVelocityFlatten::FRCPassPostProcessVelocityFlatten()
-{}
+	return ViewFamily.EngineShowFlags.PostProcessing
+		&& ViewFamily.EngineShowFlags.MotionBlur
+		&& View.FinalPostProcessSettings.MotionBlurAmount > 0.001f
+		&& View.FinalPostProcessSettings.MotionBlurMax > 0.001f
+		&& ViewFamily.bRealtimeUpdate
+		&& MotionBlurQuality > 0
+		&& !IsSimpleForwardShadingEnabled(GShaderPlatformForFeatureLevel[View.GetFeatureLevel()])
+		&& (CVarAllowMotionBlurInVR->GetInt() != 0 || !(ViewFamily.Views.Num() > 1));
+}
 
-void FRCPassPostProcessVelocityFlatten::Process(FRenderingCompositePassContext& Context)
+bool IsMotionBlurScatterRequired(FScreenPassContextRef Context, const FScreenPassTextureViewport& SceneViewport)
 {
-	const FPooledRenderTargetDesc* InputDesc = GetInputDesc(ePId_Input0);
+	const FViewInfo& View = Context->View;
+	const FSceneViewState* ViewState = Context->ViewState;
+	const float SceneExtentX = SceneViewport.Extent.X;
 
-	const FViewInfo& View = Context.View;
+	// Normalize percentage value.
+	const float VelocityMax = View.FinalPostProcessSettings.MotionBlurMax / 100.0f;
 
-	SCOPED_DRAW_EVENTF(Context.RHICmdList, MotionBlur, TEXT("VelocityFlattenCS %dx%d"),
-		View.ViewRect.Width(), View.ViewRect.Height());
+	// Scale by 0.5 due to blur samples going both ways and convert to tiles.
+	const float VelocityMaxInTiles = VelocityMax * SceneExtentX * (0.5f / 16.0f);
 
-	const FSceneRenderTargetItem& DestRenderTarget0 = PassOutputs[0].RequestSurface(Context);
-	const FSceneRenderTargetItem& DestRenderTarget1 = PassOutputs[1].RequestSurface(Context);
+	// Compute path only supports the immediate neighborhood of tiles.
+	const float TileDistanceMaxGathered = 3.0f;
 
-	TShaderMapRef< FPostProcessVelocityFlattenCS > ComputeShader( Context.GetShaderMap() );
+	// Scatter is used when maximum velocity exceeds the distance supported by the gather approach.
+	const bool bIsScatterRequiredByVelocityLength = VelocityMaxInTiles > TileDistanceMaxGathered;
 
-	// #todo-renderpasses remove once everything is renderpasses
-	UnbindRenderTargets(Context.RHICmdList);
+	// Cinematic is paused.
+	const bool bInPausedCinematic = (ViewState && ViewState->bSequencerIsPaused);
 
-	Context.SetViewportAndCallRHI( View.ViewRect );
-	Context.RHICmdList.SetComputeShader(ComputeShader->GetComputeShader());
+	// Use the scatter approach if requested by cvar or we're in a paused cinematic (higher quality).
+	const bool bIsScatterRequiredByUser = CVarMotionBlurScatter.GetValueOnRenderThread() == 1 || bInPausedCinematic;
 
-	// set destination
-	Context.RHICmdList.SetUAVParameter(ComputeShader->GetComputeShader(), ComputeShader->OutVelocityFlat.GetBaseIndex(), DestRenderTarget0.UAV);
-	Context.RHICmdList.SetUAVParameter(ComputeShader->GetComputeShader(), ComputeShader->OutMaxTileVelocity.GetBaseIndex(), DestRenderTarget1.UAV);
-
-	ComputeShader->SetCS(Context.RHICmdList, Context, View);
-
-	FIntPoint ThreadGroupCountValue = GetNumTiles16x16(View.ViewRect.Size());
-	DispatchComputeShader(Context.RHICmdList, *ComputeShader, ThreadGroupCountValue.X, ThreadGroupCountValue.Y, 1);
+	return bIsScatterRequiredByUser || bIsScatterRequiredByVelocityLength;
+}
 
-	//	void FD3D11DynamicRHI::RHIGraphicsWaitOnAsyncComputeJob( uint32 FenceIndex )
-	Context.RHICmdList.FlushComputeShaderCache();
+FIntPoint GetMotionBlurTileCount(FIntPoint PixelExtent)
+{
+	const uint32 TilesX = FMath::DivideAndRoundUp(PixelExtent.X, kMotionBlurTileSize);
+	const uint32 TilesY = FMath::DivideAndRoundUp(PixelExtent.Y, kMotionBlurTileSize);
+	return FIntPoint(TilesX, TilesY);
+}
 
-	// un-set destination
-	Context.RHICmdList.SetUAVParameter(ComputeShader->GetComputeShader(), ComputeShader->OutVelocityFlat.GetBaseIndex(), NULL);
-	Context.RHICmdList.SetUAVParameter(ComputeShader->GetComputeShader(), ComputeShader->OutMaxTileVelocity.GetBaseIndex(), NULL);
+EMotionBlurQuality GetMotionBlurQuality()
+{
+	// Quality levels begin at 1. 0 is reserved for 'off'.
+	const int32 Quality = FMath::Clamp(GetMotionBlurQualityFromCVar(), 1, static_cast<int32>(EMotionBlurQuality::MAX));
 
-	Context.RHICmdList.CopyToResolveTarget(DestRenderTarget0.TargetableTexture, DestRenderTarget0.ShaderResourceTexture, FResolveParams());
-	Context.RHICmdList.CopyToResolveTarget(DestRenderTarget1.TargetableTexture, DestRenderTarget1.ShaderResourceTexture, FResolveParams());
+	return static_cast<EMotionBlurQuality>(Quality - 1);
 }
 
-FPooledRenderTargetDesc FRCPassPostProcessVelocityFlatten::ComputeOutputDesc(EPassOutputId InPassOutputId) const
+FSamplerStateRHIParamRef GetMotionBlurColorSampler()
 {
-	if( InPassOutputId == ePId_Output0 )
+	bool bFiltered = false;
+
+#if !(UE_BUILD_SHIPPING || UE_BUILD_TEST)
+	bFiltered = CVarMotionBlurFiltering.GetValueOnRenderThread() != 0;
+#endif // !(UE_BUILD_SHIPPING || UE_BUILD_TEST)
+
+	if (bFiltered)
 	{
-		// Flattened velocity
-		FPooledRenderTargetDesc Ret = GetInput(ePId_Input0)->GetOutput()->RenderTargetDesc;
-		Ret.Reset();
-		Ret.ClearValue = FClearValueBinding::None;
-		Ret.Format = PF_FloatR11G11B10;
-		Ret.TargetableFlags |= TexCreate_UAV;
-		Ret.TargetableFlags |= TexCreate_RenderTargetable;
-		Ret.Flags |= GFastVRamConfig.VelocityFlat;
-		Ret.DebugName = TEXT("VelocityFlat");
-
-		return Ret;
+		return TStaticSamplerState<SF_Bilinear, AM_Clamp, AM_Clamp, AM_Clamp>::GetRHI();
 	}
 	else
 	{
-		// Max tile velocity
-		FPooledRenderTargetDesc UnmodifiedRet = GetInput(ePId_Input0)->GetOutput()->RenderTargetDesc;
-		UnmodifiedRet.Reset();
-
-		FIntPoint PixelExtent = UnmodifiedRet.Extent;
-		FIntPoint TileCount = GetNumTiles16x16(PixelExtent);
-
-		FPooledRenderTargetDesc Ret(FPooledRenderTargetDesc::Create2DDesc(TileCount, PF_FloatRGBA, FClearValueBinding::None, TexCreate_None, TexCreate_RenderTargetable | TexCreate_UAV, false));
-		Ret.Flags |= GFastVRamConfig.VelocityMax;
-		Ret.DebugName = TEXT("MaxVelocity");
-
-		return Ret;
+		return TStaticSamplerState<SF_Point, AM_Clamp, AM_Clamp, AM_Clamp>::GetRHI();
 	}
 }
 
-
-TGlobalResource< FSpriteIndexBuffer<8> > GScatterQuadIndexBuffer;
-
-class FPostProcessVelocityScatterVS : public FGlobalShader
+FSamplerStateRHIParamRef GetMotionBlurVelocitySampler()
 {
-	DECLARE_SHADER_TYPE(FPostProcessVelocityScatterVS,Global);
-
-	static bool ShouldCompilePermutation(const FGlobalShaderPermutationParameters& Parameters)
-	{
-		return IsFeatureLevelSupported(Parameters.Platform, ERHIFeatureLevel::SM5);
-	}
-
-	/** Default constructor. */
-	FPostProcessVelocityScatterVS() {}
-	
-public:
-	FPostProcessPassParameters PostprocessParameter;
-	FShaderParameter DrawMax;
-	FShaderParameter MotionBlurParameters;
+	return TStaticSamplerState<SF_Point, AM_Clamp, AM_Clamp, AM_Clamp>::GetRHI();
+}
 
-	/** Initialization constructor. */
-	FPostProcessVelocityScatterVS(const ShaderMetaType::CompiledShaderInitializerType& Initializer)
-		: FGlobalShader(Initializer)
-	{
-		PostprocessParameter.Bind(Initializer.ParameterMap);
-		DrawMax.Bind(Initializer.ParameterMap, TEXT("bDrawMax"));
-		MotionBlurParameters.Bind( Initializer.ParameterMap, TEXT("MotionBlurParameters") );
-	}
+// Set of common shader parameters shared by all motion blur shaders.
+BEGIN_SHADER_PARAMETER_STRUCT(FMotionBlurParameters, )
+	SHADER_PARAMETER(float, AspectRatio)
+	SHADER_PARAMETER(float, VelocityScale)
+	SHADER_PARAMETER(float, VelocityScaleForTiles)
+	SHADER_PARAMETER(float, VelocityMax)
+END_SHADER_PARAMETER_STRUCT()
 
-	// FShader interface.
-	virtual bool Serialize(FArchive& Ar) override
-	{
-		bool bShaderHasOutdatedParameters = FGlobalShader::Serialize(Ar);
-		Ar << PostprocessParameter << DrawMax << MotionBlurParameters;
-		return bShaderHasOutdatedParameters;
-	}
-
-	/** to have a similar interface as all other shaders */
-	void SetParameters(const FRenderingCompositePassContext& Context, int32 bDrawMax)
-	{
-		const FVertexShaderRHIParamRef ShaderRHI = GetVertexShader();
+FMotionBlurParameters GetMotionBlurParameters(FScreenPassContextRef Context, FIntPoint SceneViewportExtent, float BlurScale)
+{
+	const FViewInfo& View = Context->View;
+	const FSceneViewState* ViewState = Context->ViewState;
 
-		FGlobalShader::SetParameters<FViewUniformShaderParameters>(Context.RHICmdList, ShaderRHI, Context.View.ViewUniformBuffer);
-		PostprocessParameter.SetVS(ShaderRHI, Context, TStaticSamplerState<SF_Point,AM_Clamp,AM_Clamp,AM_Clamp>::GetRHI());
+	const float TileSize = kMotionBlurTileSize;
+	const float SceneViewportExtentX = SceneViewportExtent.X;
+	const float SceneViewportExtentY = SceneViewportExtent.Y;
+	const float MotionBlurTimeScale = ViewState ? ViewState->MotionBlurTimeScale : 1.0f;
 
-		SetShaderValue(Context.RHICmdList, ShaderRHI, DrawMax, bDrawMax);
-		SetShaderValue(Context.RHICmdList, ShaderRHI, MotionBlurParameters, GetMotionBlurParameters( Context ) );
-	}
+	// Scale by 0.5 due to blur samples going both ways.
+	const float VelocityScale = MotionBlurTimeScale * View.FinalPostProcessSettings.MotionBlurAmount * 0.5f;
+	const float VelocityUVToPixel = BlurScale * SceneViewportExtentX * 0.5f;
 
-	static const TCHAR* GetSourceFilename()
-	{
-		return TEXT("/Engine/Private/PostProcessMotionBlur.usf");
-	}
+	// 0:no 1:full screen width, percent conversion
+	const float UVVelocityMax = View.FinalPostProcessSettings.MotionBlurMax / 100.0f;
 
-	static const TCHAR* GetFunctionName()
-	{
-		return TEXT("VelocityScatterVS");
-	}
-};
+	FMotionBlurParameters MotionBlurParameters;
+	MotionBlurParameters.AspectRatio = SceneViewportExtentY / SceneViewportExtentX;
+	MotionBlurParameters.VelocityScale = VelocityUVToPixel * VelocityScale;
+	MotionBlurParameters.VelocityScaleForTiles = MotionBlurParameters.VelocityScale / TileSize;
+	MotionBlurParameters.VelocityMax = FMath::Abs(VelocityUVToPixel) * UVVelocityMax;
+	return MotionBlurParameters;
+}
 
-class FPostProcessVelocityScatterPS : public FGlobalShader
+// Base class for a motion blur / velocity shader.
+class FMotionBlurShader : public FGlobalShader
 {
-	DECLARE_SHADER_TYPE(FPostProcessVelocityScatterPS, Global);
-
+public:
 	static bool ShouldCompilePermutation(const FGlobalShaderPermutationParameters& Parameters)
 	{
 		return IsFeatureLevelSupported(Parameters.Platform, ERHIFeatureLevel::SM5);
 	}
 
-	/** Default constructor. */
-	FPostProcessVelocityScatterPS() {}
-
-public:
-	/** Initialization constructor. */
-	FPostProcessVelocityScatterPS(const ShaderMetaType::CompiledShaderInitializerType& Initializer)
+	FMotionBlurShader() = default;
+	FMotionBlurShader(const ShaderMetaType::CompiledShaderInitializerType& Initializer)
 		: FGlobalShader(Initializer)
 	{}
-
-	void SetParameters(const FRenderingCompositePassContext& Context)
-	{
-		const FPixelShaderRHIParamRef ShaderRHI = GetPixelShader();
-
-		FGlobalShader::SetParameters<FViewUniformShaderParameters>(Context.RHICmdList, ShaderRHI, Context.View.ViewUniformBuffer);
-	}
 };
 
-IMPLEMENT_SHADER_TYPE(,FPostProcessVelocityScatterVS,TEXT("/Engine/Private/PostProcessMotionBlur.usf"),TEXT("VelocityScatterVS"),SF_Vertex);
-IMPLEMENT_SHADER_TYPE(,FPostProcessVelocityScatterPS,TEXT("/Engine/Private/PostProcessMotionBlur.usf"),TEXT("VelocityScatterPS"),SF_Pixel);
-
-
-void FRCPassPostProcessVelocityScatter::Process(FRenderingCompositePassContext& Context)
+class FMotionBlurVelocityFlattenCS : public FMotionBlurShader
 {
-	const FPooledRenderTargetDesc* InputDesc = GetInputDesc(ePId_Input0);
-	
-	const FViewInfo& View = Context.View;
-
-	SCOPED_DRAW_EVENTF(Context.RHICmdList, MotionBlur, TEXT("VelocityScatter %dx%d"),
-		View.ViewRect.Width(), View.ViewRect.Height());
-
-	FIntPoint SrcSize = InputDesc->Extent;
-	FIntPoint DestSize = PassOutputs[0].RenderTargetDesc.Extent;
-
-	FIntPoint TileCount = GetNumTiles16x16( View.ViewRect.Size() );
-
-	const FSceneRenderTargetItem& DestRenderTarget = PassOutputs[0].RequestSurface(Context);
-
-	TRefCountPtr<IPooledRenderTarget> DepthTarget;
-	FPooledRenderTargetDesc Desc( FPooledRenderTargetDesc::Create2DDesc( DestSize, PF_ShadowDepth, FClearValueBinding::DepthOne, TexCreate_None, TexCreate_DepthStencilTargetable, false ) );
-	GRenderTargetPool.FindFreeElement(Context.RHICmdList, Desc, DepthTarget, TEXT("VelocityScatterDepth") );
-
-	// Set the view family's render target/viewport.
-	FRHIRenderPassInfo RPInfo(DestRenderTarget.TargetableTexture, ERenderTargetActions::Load_Store);
-	RPInfo.DepthStencilRenderTarget.DepthStencilTarget = DepthTarget->GetRenderTargetItem().TargetableTexture;
-	RPInfo.DepthStencilRenderTarget.Action = MakeDepthStencilTargetActions(ERenderTargetActions::Clear_Store, ERenderTargetActions::Load_Store);
-	RPInfo.DepthStencilRenderTarget.ExclusiveDepthStencil = FExclusiveDepthStencil::DepthWrite_StencilWrite;
-
-	Context.RHICmdList.BeginRenderPass(RPInfo, TEXT("VelocityScatterDepth"));
-	{
-		// clear depth
-		// Max >= Min so no need to clear on second pass
-		Context.SetViewportAndCallRHI(0, 0, 0.0f, TileCount.X, TileCount.Y, 1.0f);
-
-		FGraphicsPipelineStateInitializer GraphicsPSOInit;
-		Context.RHICmdList.ApplyCachedRenderTargets(GraphicsPSOInit);
-
-		// Min,Max
-		for (int i = 0; i < 2; i++)
-		{
-			if (i == 0)
-			{
-				// min
-				GraphicsPSOInit.BlendState = TStaticBlendStateWriteMask< CW_RGBA >::GetRHI();
-				GraphicsPSOInit.RasterizerState = TStaticRasterizerState<>::GetRHI();
-				GraphicsPSOInit.DepthStencilState = TStaticDepthStencilState< true, CF_Less >::GetRHI();
-			}
-			else
-			{
-				// max
-				GraphicsPSOInit.BlendState = TStaticBlendStateWriteMask< CW_BA >::GetRHI();
-				GraphicsPSOInit.RasterizerState = TStaticRasterizerState<>::GetRHI();
-				GraphicsPSOInit.DepthStencilState = TStaticDepthStencilState< true, CF_Greater >::GetRHI();
-			}
-
-			TShaderMapRef< FPostProcessVelocityScatterVS > VertexShader(Context.GetShaderMap());
-			TShaderMapRef< FPostProcessVelocityScatterPS > PixelShader(Context.GetShaderMap());
-
-			GraphicsPSOInit.BoundShaderState.VertexDeclarationRHI = GEmptyVertexDeclaration.VertexDeclarationRHI;
-			GraphicsPSOInit.BoundShaderState.VertexShaderRHI = GETSAFERHISHADER_VERTEX(*VertexShader);
-			GraphicsPSOInit.BoundShaderState.PixelShaderRHI = GETSAFERHISHADER_PIXEL(*PixelShader);
-			GraphicsPSOInit.PrimitiveType = PT_TriangleList;
-
-			SetGraphicsPipelineState(Context.RHICmdList, GraphicsPSOInit);
-
-			VertexShader->SetParameters(Context, i);
-			PixelShader->SetParameters(Context);
-
-			// needs to be the same on shader side (faster on NVIDIA and AMD)
-			int32 QuadsPerInstance = 8;
-
-			Context.RHICmdList.SetStreamSource(0, NULL, 0);
-			Context.RHICmdList.DrawIndexedPrimitive(GScatterQuadIndexBuffer.IndexBufferRHI, 0, 0, 32, 0, 2 * QuadsPerInstance, FMath::DivideAndRoundUp(TileCount.X * TileCount.Y, QuadsPerInstance));
-		}
-	}
-	Context.RHICmdList.EndRenderPass();
-	Context.RHICmdList.CopyToResolveTarget(DestRenderTarget.TargetableTexture, DestRenderTarget.ShaderResourceTexture, FResolveParams());
-}
-
-FPooledRenderTargetDesc FRCPassPostProcessVelocityScatter::ComputeOutputDesc(EPassOutputId InPassOutputId) const
-{
-	FPooledRenderTargetDesc Ret = GetInput(ePId_Input0)->GetOutput()->RenderTargetDesc;
-	Ret.Reset();
-
-	Ret.DebugName = TEXT("ScatteredMaxVelocity");
+public:
+	DECLARE_GLOBAL_SHADER(FMotionBlurVelocityFlattenCS);
+	SHADER_USE_PARAMETER_STRUCT(FMotionBlurVelocityFlattenCS, FMotionBlurShader);
+
+	BEGIN_SHADER_PARAMETER_STRUCT(FParameters, )
+		SHADER_PARAMETER_STRUCT_REF(FViewUniformShaderParameters, View)
+		SHADER_PARAMETER_STRUCT(FMotionBlurParameters, MotionBlur)
+		SHADER_PARAMETER_STRUCT(FScreenPassTextureViewportParameters, Velocity)
+		SHADER_PARAMETER_RDG_TEXTURE(Texture2D, VelocityTexture)
+		SHADER_PARAMETER_RDG_TEXTURE(Texture2D, DepthTexture)
+		SHADER_PARAMETER_RDG_TEXTURE_UAV(RWTexture2D, OutVelocityFlatTexture)
+		SHADER_PARAMETER_RDG_TEXTURE_UAV(RWTexture2D, OutVelocityTileTexture)
+	END_SHADER_PARAMETER_STRUCT()
+};
 
-	return Ret;
-}
+IMPLEMENT_GLOBAL_SHADER(FMotionBlurVelocityFlattenCS, "/Engine/Private/PostProcessVelocityFlatten.usf", "VelocityFlattenMain", SF_Compute);
 
+BEGIN_SHADER_PARAMETER_STRUCT(FMotionBlurVelocityDilateParameters, )
+	SHADER_PARAMETER_STRUCT(FMotionBlurParameters, MotionBlur)
+	SHADER_PARAMETER_STRUCT(FScreenPassTextureViewportParameters, VelocityTile)
+	SHADER_PARAMETER_RDG_TEXTURE(Texture2D, VelocityTileTexture)
+END_SHADER_PARAMETER_STRUCT()
 
-class FPostProcessVelocityGatherCS : public FGlobalShader
+class FMotionBlurVelocityDilateGatherCS : public FMotionBlurShader
 {
-	DECLARE_SHADER_TYPE(FPostProcessVelocityGatherCS, Global);
-
-	static bool ShouldCompilePermutation(const FGlobalShaderPermutationParameters& Parameters)
-	{
-		return IsFeatureLevelSupported(Parameters.Platform, ERHIFeatureLevel::SM5);
-	}
-
-	/** Default constructor. */
-	FPostProcessVelocityGatherCS() {}
-
 public:
-	FShaderParameter		OutScatteredMaxVelocity;
+	DECLARE_GLOBAL_SHADER(FMotionBlurVelocityDilateGatherCS);
+	SHADER_USE_PARAMETER_STRUCT(FMotionBlurVelocityDilateGatherCS, FMotionBlurShader);
 
-	/** Initialization constructor. */
-	FPostProcessVelocityGatherCS(const ShaderMetaType::CompiledShaderInitializerType& Initializer)
-		: FGlobalShader(Initializer)
-	{
-		PostprocessParameter.Bind(Initializer.ParameterMap);
-		MotionBlurParameters.Bind(Initializer.ParameterMap, TEXT("MotionBlurParameters"));
-		OutScatteredMaxVelocity.Bind(Initializer.ParameterMap, TEXT("OutScatteredMaxVelocity"));
-	}
-
-	// FShader interface.
-	virtual bool Serialize(FArchive& Ar) override
-	{
-		bool bShaderHasOutdatedParameters = FGlobalShader::Serialize(Ar);
-		Ar << PostprocessParameter;
-		Ar << MotionBlurParameters;
-		Ar << OutScatteredMaxVelocity;
-		return bShaderHasOutdatedParameters;
-	}
-
-	void SetParameters(const FRenderingCompositePassContext& Context)
-	{
-		const FComputeShaderRHIParamRef ShaderRHI = GetComputeShader();
-
-		FGlobalShader::SetParameters<FViewUniformShaderParameters>(Context.RHICmdList, ShaderRHI, Context.View.ViewUniformBuffer);
-
-		PostprocessParameter.SetCS(ShaderRHI, Context, Context.RHICmdList, TStaticSamplerState<SF_Point, AM_Clamp, AM_Clamp, AM_Clamp>::GetRHI());
-		SetShaderValue(Context.RHICmdList, ShaderRHI, MotionBlurParameters, GetMotionBlurParameters( Context ) );
-	}
-	
-private:
-	FPostProcessPassParameters	PostprocessParameter;
-	FShaderParameter			MotionBlurParameters;
+	BEGIN_SHADER_PARAMETER_STRUCT(FParameters, )
+		SHADER_PARAMETER_STRUCT_INCLUDE(FMotionBlurVelocityDilateParameters, Dilate)
+		SHADER_PARAMETER_RDG_TEXTURE_UAV(RWTexture2D<float4>, OutDilatedVelocityTileTexture)
+	END_SHADER_PARAMETER_STRUCT()
 };
 
-IMPLEMENT_SHADER_TYPE(,FPostProcessVelocityGatherCS,TEXT("/Engine/Private/PostProcessVelocityFlatten.usf"),TEXT("VelocityGatherCS"),SF_Compute);
+IMPLEMENT_GLOBAL_SHADER(FMotionBlurVelocityDilateGatherCS, "/Engine/Private/PostProcessVelocityFlatten.usf", "VelocityGatherCS", SF_Compute);
 
-void FRCPassPostProcessVelocityGather::Process(FRenderingCompositePassContext& Context)
+enum class EMotionBlurVelocityScatterPass : uint32
 {
-	const FPooledRenderTargetDesc* InputDesc = GetInputDesc(ePId_Input0);
-
-	const FViewInfo& View = Context.View;
-
-	SCOPED_DRAW_EVENTF(Context.RHICmdList, MotionBlur, TEXT("VelocityGatherCS %dx%d"),
-		View.ViewRect.Width(), View.ViewRect.Height());
-
-	FIntPoint TileCount = GetNumTiles16x16( View.ViewRect.Size() );
-	
-	const FSceneRenderTargetItem& DestRenderTarget = PassOutputs[0].RequestSurface(Context);
-
-	// #todo-renderpasses remove when everything is renderpasses.
-	UnbindRenderTargets(Context.RHICmdList);
-	Context.SetViewportAndCallRHI( 0, 0, 0.0f, TileCount.X, TileCount.Y, 1.0f );
-	
-	TShaderMapRef< FPostProcessVelocityGatherCS > ComputeShader( Context.GetShaderMap() );
-	Context.RHICmdList.SetComputeShader(ComputeShader->GetComputeShader());
-
-	// set destination
-	Context.RHICmdList.SetUAVParameter( ComputeShader->GetComputeShader(), ComputeShader->OutScatteredMaxVelocity.GetBaseIndex(), DestRenderTarget.UAV );
-
-	ComputeShader->SetParameters( Context );
-
-	FIntPoint GroupCount = GetNumTiles16x16( TileCount );
-	DispatchComputeShader(Context.RHICmdList, *ComputeShader, GroupCount.X, GroupCount.Y, 1);
-
-	// un-set destination
-	Context.RHICmdList.SetUAVParameter( ComputeShader->GetComputeShader(), ComputeShader->OutScatteredMaxVelocity.GetBaseIndex(), NULL );
+	DrawMin,
+	DrawMax,
+	MAX
+};
 
-	Context.RHICmdList.CopyToResolveTarget(DestRenderTarget.TargetableTexture, DestRenderTarget.ShaderResourceTexture, FResolveParams());
-}
+BEGIN_SHADER_PARAMETER_STRUCT(FMotionBlurVelocityDilateScatterParameters, )
+	SHADER_PARAMETER(uint32, ScatterPass)
+	SHADER_PARAMETER_STRUCT_INCLUDE(FMotionBlurVelocityDilateParameters, Dilate)
+	RENDER_TARGET_BINDING_SLOTS()
+END_SHADER_PARAMETER_STRUCT()
 
-FPooledRenderTargetDesc FRCPassPostProcessVelocityGather::ComputeOutputDesc(EPassOutputId InPassOutputId) const
+class FMotionBlurVelocityDilateScatterVS : public FMotionBlurShader
 {
-	FPooledRenderTargetDesc Ret = GetInput(ePId_Input0)->GetOutput()->RenderTargetDesc;
-	Ret.Reset();
-
-	Ret.TargetableFlags |= TexCreate_UAV;
-	Ret.TargetableFlags |= TexCreate_RenderTargetable;
-	Ret.Flags |= GFastVRamConfig.MotionBlur;
-	Ret.DebugName = TEXT("ScatteredMaxVelocity");
-
-	return Ret;
-}
+public:
+	DECLARE_GLOBAL_SHADER(FMotionBlurVelocityDilateScatterVS);
+	SHADER_USE_PARAMETER_STRUCT(FMotionBlurVelocityDilateScatterVS, FMotionBlurShader);
+	using FParameters = FMotionBlurVelocityDilateScatterParameters;
+};
 
+IMPLEMENT_GLOBAL_SHADER(FMotionBlurVelocityDilateScatterVS, "/Engine/Private/PostProcessVelocityFlatten.usf", "VelocityScatterVS", SF_Vertex);
 
-/**
- * @param Quality 0: visualize, 1:low, 2:medium, 3:high, 4:very high
- * @param InnerLoopConfig 0: fastest, 1: clamp uv, 2: clamp uv and handle different viewport res between scene color and depth.
- */
-template< uint32 Quality, uint32 InnerLoopConfig >
-class FPostProcessMotionBlurPS : public FGlobalShader
+class FMotionBlurVelocityDilateScatterPS : public FMotionBlurShader
 {
-	DECLARE_SHADER_TYPE(FPostProcessMotionBlurPS, Global);
-
-	static bool ShouldCompilePermutation(const FGlobalShaderPermutationParameters& Parameters)
-	{
-		return IsFeatureLevelSupported(Parameters.Platform, ERHIFeatureLevel::SM5);
-	}
-
-	static void ModifyCompilationEnvironment(const FGlobalShaderPermutationParameters& Parameters, FShaderCompilerEnvironment& OutEnvironment)
-	{
-		FGlobalShader::ModifyCompilationEnvironment(Parameters,OutEnvironment);
-		OutEnvironment.SetDefine(TEXT("MOTION_BLUR_QUALITY"), Quality);
-		OutEnvironment.SetDefine(TEXT("MOTION_BLUR_INNER_LOOP_CONFIG"), InnerLoopConfig);
-	}
-
-	/** Default constructor. */
-	FPostProcessMotionBlurPS() {}
-
 public:
-	FPostProcessPassParameters PostprocessParameter;
-	FSceneTextureShaderParameters SceneTextureParameters;
-	FShaderParameter MotionBlurParameters;
-	FShaderParameter SceneColorBufferUVToViewBufferUV;
-	FShaderParameter TileUVMax;
-
-	/** Initialization constructor. */
-	FPostProcessMotionBlurPS(const ShaderMetaType::CompiledShaderInitializerType& Initializer)
-		: FGlobalShader(Initializer)
-	{
-		PostprocessParameter.Bind(Initializer.ParameterMap);
-		SceneTextureParameters.Bind(Initializer);
-		MotionBlurParameters.Bind(Initializer.ParameterMap, TEXT("MotionBlurParameters"));
-		SceneColorBufferUVToViewBufferUV.Bind(Initializer.ParameterMap, TEXT("SceneColorBufferUVToViewBufferUV"));
-		TileUVMax.Bind(Initializer.ParameterMap, TEXT("TileUVMax"));
-	}
-
-	// FShader interface.
-	virtual bool Serialize(FArchive& Ar) override
-	{
-		bool bShaderHasOutdatedParameters = FGlobalShader::Serialize(Ar);
-		Ar << PostprocessParameter << SceneTextureParameters << MotionBlurParameters << SceneColorBufferUVToViewBufferUV << TileUVMax;
-		return bShaderHasOutdatedParameters;
-	}
+	DECLARE_GLOBAL_SHADER(FMotionBlurVelocityDilateScatterPS);
+	SHADER_USE_PARAMETER_STRUCT(FMotionBlurVelocityDilateScatterPS, FMotionBlurShader);
+	using FParameters = FMotionBlurVelocityDilateScatterParameters;
+};
 
-	void SetParameters(const FRenderingCompositePassContext& Context, float Scale)
-	{
-		const FPixelShaderRHIParamRef ShaderRHI = GetPixelShader();
+IMPLEMENT_GLOBAL_SHADER(FMotionBlurVelocityDilateScatterPS, "/Engine/Private/PostProcessVelocityFlatten.usf", "VelocityScatterPS", SF_Pixel);
 
-		FGlobalShader::SetParameters<FViewUniformShaderParameters>(Context.RHICmdList, ShaderRHI, Context.View.ViewUniformBuffer);
+class FMotionBlurQualityDimension : SHADER_PERMUTATION_ENUM_CLASS("MOTION_BLUR_QUALITY", EMotionBlurQuality);
 
-		SceneTextureParameters.Set(Context.RHICmdList, ShaderRHI, Context.View.FeatureLevel, ESceneTextureSetupMode::All);
+BEGIN_SHADER_PARAMETER_STRUCT(FMotionBlurFilterParameters, )
+	SHADER_PARAMETER_STRUCT(FMotionBlurParameters, MotionBlur)
 
-		{
-			bool bFiltered = false;
+	SHADER_PARAMETER_STRUCT(FScreenPassTextureViewportParameters, Color)
+	SHADER_PARAMETER_STRUCT(FScreenPassTextureViewportParameters, Velocity)
+	SHADER_PARAMETER_STRUCT(FScreenPassTextureViewportParameters, VelocityTile)
 
-#if !(UE_BUILD_SHIPPING || UE_BUILD_TEST)
-			bFiltered = CVarMotionBlurFiltering.GetValueOnRenderThread() != 0;
-#endif // !(UE_BUILD_SHIPPING || UE_BUILD_TEST)
+	SHADER_PARAMETER_STRUCT(FScreenPassTextureViewportTransform, ColorToVelocity)
+	SHADER_PARAMETER_STRUCT(FScreenPassTextureViewportTransform, ColorToVelocityTile)
 
-			if(bFiltered)
-			{
-				//PostprocessParameter.SetPS(ShaderRHI, Context, TStaticSamplerState<SF_Bilinear,AM_Border,AM_Border,AM_Clamp>::GetRHI());
+	SHADER_PARAMETER_RDG_TEXTURE(Texture2D, ColorTexture)
+	SHADER_PARAMETER_RDG_TEXTURE(Texture2D, VelocityFlatTexture)
+	SHADER_PARAMETER_RDG_TEXTURE(Texture2D, VelocityTileTexture)
 
-				FSamplerStateRHIParamRef Filters[] =
-				{
-					TStaticSamplerState<SF_Bilinear,AM_Clamp,AM_Clamp,AM_Clamp>::GetRHI(),
-					TStaticSamplerState<SF_Point,AM_Clamp,AM_Clamp,AM_Clamp>::GetRHI(),
-					TStaticSamplerState<SF_Point,AM_Clamp,AM_Clamp,AM_Clamp>::GetRHI(),
-					TStaticSamplerState<SF_Point,AM_Clamp,AM_Clamp,AM_Clamp>::GetRHI(),
-				};
+	SHADER_PARAMETER_SAMPLER(SamplerState, ColorSampler)
+	SHADER_PARAMETER_SAMPLER(SamplerState, VelocitySampler)
+END_SHADER_PARAMETER_STRUCT()
 
-				PostprocessParameter.SetPS(Context.RHICmdList, ShaderRHI, Context, 0, eFC_0000, Filters);
-			}
-			else
-			{
-				PostprocessParameter.SetPS(Context.RHICmdList, ShaderRHI, Context, TStaticSamplerState<SF_Point,AM_Clamp,AM_Clamp,AM_Clamp>::GetRHI());
-			}
-		}
-
-		SetShaderValue(Context.RHICmdList, ShaderRHI, MotionBlurParameters, GetMotionBlurParameters( Context, Scale ) );
-		SetShaderValue(Context.RHICmdList, ShaderRHI, SceneColorBufferUVToViewBufferUV, GetSceneColorBufferUVToViewBufferUV(Context));
-
-		if (TileUVMax.IsBound())
-		{
-			FIntPoint VelocityBufferSize = FSceneRenderTargets::Get(Context.RHICmdList).GetBufferSizeXY();
-
-			FIntPoint TileBufferSize = GetNumTiles16x16(VelocityBufferSize);
-			FIntPoint TileCount = GetNumTiles16x16(Context.View.ViewRect.Size());
-			FVector2D TileUVMaxValue(
-				(TileCount.X - 0.5f) / float(TileBufferSize.X),
-				(TileCount.Y - 0.5f) / float(TileBufferSize.Y));
-
-			SetShaderValue(Context.RHICmdList, ShaderRHI, TileUVMax, TileUVMaxValue);
-		}
-	}
+class FMotionBlurFilterPS : public FMotionBlurShader
+{
+public:
+	DECLARE_GLOBAL_SHADER(FMotionBlurFilterPS);
+	SHADER_USE_PARAMETER_STRUCT(FMotionBlurFilterPS, FMotionBlurShader);
 
-	static const TCHAR* GetSourceFilename()
-	{
-		return TEXT("/Engine/Private/PostProcessMotionBlur.usf");
-	}
+	BEGIN_SHADER_PARAMETER_STRUCT(FParameters, )
+		SHADER_PARAMETER_STRUCT_INCLUDE(FMotionBlurFilterParameters, Filter)
+		RENDER_TARGET_BINDING_SLOTS()
+	END_SHADER_PARAMETER_STRUCT()
 
-	static const TCHAR* GetFunctionName()
-	{
-		return TEXT("MainPS");
-	}
+	using FQualityDimension = FMotionBlurQualityDimension;
+	using FPermutationDomain = TShaderPermutationDomain<FQualityDimension>;
 };
 
-#define VARIATION2(A, B) typedef FPostProcessMotionBlurPS<A, B> FPostProcessMotionBlurPS##A##B; \
-	IMPLEMENT_SHADER_TYPE2(FPostProcessMotionBlurPS##A##B, SF_Pixel);
-
-#define VARIATION1(A) VARIATION2(A, 0) VARIATION2(A, 1) VARIATION2(A, 2)
-
-VARIATION1(1)
-VARIATION1(2)
-VARIATION1(3)
-VARIATION1(4)
-
-#undef VARIATION1
-#undef VARIATION2
+IMPLEMENT_GLOBAL_SHADER(FMotionBlurFilterPS, "/Engine/Private/PostProcessMotionBlur.usf", "MainPS", SF_Pixel);
 
-/**
- * @param Quality 0: visualize, 1:low, 2:medium, 3:high, 4:very high
- * @param InnerLoopConfig 0: fastest, 1: clamp uv, 2: clamp uv and handle different viewport res between scene color and depth.
- */
-template< uint32 Quality, uint32 InnerLoopConfig >
-class FPostProcessMotionBlurCS : public FGlobalShader
+class FMotionBlurFilterCS : public FMotionBlurShader
 {
-	DECLARE_SHADER_TYPE(FPostProcessMotionBlurCS, Global);
+public:
+	DECLARE_GLOBAL_SHADER(FMotionBlurFilterCS);
+	SHADER_USE_PARAMETER_STRUCT(FMotionBlurFilterCS, FMotionBlurShader);
 
-	static bool ShouldCompilePermutation(const FGlobalShaderPermutationParameters& Parameters)
-	{
-		return IsFeatureLevelSupported(Parameters.Platform, ERHIFeatureLevel::SM5);
-	}
+	BEGIN_SHADER_PARAMETER_STRUCT(FParameters, )
+		SHADER_PARAMETER_STRUCT_INCLUDE(FMotionBlurFilterParameters, Filter)
+		SHADER_PARAMETER_RDG_TEXTURE_UAV(RWTexture2D<float4>, OutColorTexture)
+	END_SHADER_PARAMETER_STRUCT()
 
 	static void ModifyCompilationEnvironment(const FGlobalShaderPermutationParameters& Parameters, FShaderCompilerEnvironment& OutEnvironment)
 	{
-		// CS Params
 		FGlobalShader::ModifyCompilationEnvironment(Parameters, OutEnvironment);
-		OutEnvironment.SetDefine(TEXT("THREADGROUP_SIZEX"), GMotionBlurComputeTileSizeX);
-		OutEnvironment.SetDefine(TEXT("THREADGROUP_SIZEY"), GMotionBlurComputeTileSizeY);
-
-		// PS params
-		OutEnvironment.SetDefine(TEXT("MOTION_BLUR_QUALITY"), Quality);
-		OutEnvironment.SetDefine(TEXT("MOTION_BLUR_INNER_LOOP_CONFIG"), InnerLoopConfig);
+		OutEnvironment.SetDefine(TEXT("THREADGROUP_SIZEX"), kMotionBlurComputeTileSizeX);
+		OutEnvironment.SetDefine(TEXT("THREADGROUP_SIZEY"), kMotionBlurComputeTileSizeY);
 	}
 
-	/** Default constructor. */
-	FPostProcessMotionBlurCS() {}
-
-public:
-	// CS params
-	FRWShaderParameter OutComputeTex;
-	FShaderParameter MotionBlurComputeParams;
-
-	// PS params
-	FPostProcessPassParameters PostprocessParameter;
-	FSceneTextureShaderParameters SceneTextureParameters;
-	FShaderParameter MotionBlurParameters;
-	FShaderParameter SceneColorBufferUVToViewBufferUV;
-	FShaderParameter TileUVMax;
-
-	/** Initialization constructor. */
-	FPostProcessMotionBlurCS(const ShaderMetaType::CompiledShaderInitializerType& Initializer)
-		: FGlobalShader(Initializer)
-	{
-		// CS params
-		OutComputeTex.Bind(Initializer.ParameterMap, TEXT("OutComputeTex"));
-		MotionBlurComputeParams.Bind(Initializer.ParameterMap, TEXT("MotionBlurComputeParams"));
-
-		// PS params
-		PostprocessParameter.Bind(Initializer.ParameterMap);
-		SceneTextureParameters.Bind(Initializer);
-		MotionBlurParameters.Bind(Initializer.ParameterMap, TEXT("MotionBlurParameters"));
-		SceneColorBufferUVToViewBufferUV.Bind(Initializer.ParameterMap, TEXT("SceneColorBufferUVToViewBufferUV"));
-		TileUVMax.Bind(Initializer.ParameterMap, TEXT("TileUVMax"));
-	}
-
-	template <typename TRHICmdList>
-	void SetParameters(TRHICmdList& RHICmdList, const FRenderingCompositePassContext& Context, const FIntPoint& DestSize, FUnorderedAccessViewRHIParamRef DestUAV, float Scale)
-	{
-		const FComputeShaderRHIParamRef ShaderRHI = GetComputeShader();
-
-		// CS params
-		FGlobalShader::SetParameters<FViewUniformShaderParameters>(RHICmdList, ShaderRHI, Context.View.ViewUniformBuffer);
-		OutComputeTex.SetTexture(RHICmdList, ShaderRHI, nullptr, DestUAV);
-		
-		FVector4 MotionBlurComputeValues(0, 0, 1.f / (float)DestSize.X, 1.f / (float)DestSize.Y);
-		SetShaderValue(RHICmdList, ShaderRHI, MotionBlurComputeParams, MotionBlurComputeValues);
-
-		// PS params
-		SceneTextureParameters.Set(RHICmdList, ShaderRHI, Context.View.FeatureLevel, ESceneTextureSetupMode::All);
-
-		{
-			bool bFiltered = false;
-
-#if !(UE_BUILD_SHIPPING || UE_BUILD_TEST)
-			bFiltered = CVarMotionBlurFiltering.GetValueOnRenderThread() != 0;
-#endif // !(UE_BUILD_SHIPPING || UE_BUILD_TEST)
-
-			if(bFiltered)
-			{
-				FSamplerStateRHIParamRef Filters[] =
-				{
-					TStaticSamplerState<SF_Bilinear,AM_Clamp,AM_Clamp,AM_Clamp>::GetRHI(),
-					TStaticSamplerState<SF_Point,AM_Clamp,AM_Clamp,AM_Clamp>::GetRHI(),
-					TStaticSamplerState<SF_Point,AM_Clamp,AM_Clamp,AM_Clamp>::GetRHI(),
-					TStaticSamplerState<SF_Point,AM_Clamp,AM_Clamp,AM_Clamp>::GetRHI(),
-				};
-
-				PostprocessParameter.SetCS( ShaderRHI, Context, RHICmdList, 0, eFC_0000, Filters );
-			}
-			else
-			{
-				PostprocessParameter.SetCS(ShaderRHI, Context, RHICmdList, TStaticSamplerState<SF_Point,AM_Clamp,AM_Clamp,AM_Clamp>::GetRHI());
-			}
-		}
+	using FQualityDimension = FMotionBlurQualityDimension;
+	using FPermutationDomain = TShaderPermutationDomain<FQualityDimension>;
+};
 
-		SetShaderValue(RHICmdList, ShaderRHI, MotionBlurParameters, GetMotionBlurParameters(Context, Scale));
-		SetShaderValue(Context.RHICmdList, ShaderRHI, SceneColorBufferUVToViewBufferUV, GetSceneColorBufferUVToViewBufferUV(Context));
+IMPLEMENT_GLOBAL_SHADER(FMotionBlurFilterCS, "/Engine/Private/PostProcessMotionBlur.usf", "MainCS", SF_Compute);
 
-		if (TileUVMax.IsBound())
-		{
-			FIntPoint VelocityBufferSize = FSceneRenderTargets::Get(Context.RHICmdList).GetBufferSizeXY();
+class FMotionBlurVisualizePS : public FMotionBlurShader
+{
+public:
+	DECLARE_GLOBAL_SHADER(FMotionBlurVisualizePS);
+	SHADER_USE_PARAMETER_STRUCT(FMotionBlurVisualizePS, FMotionBlurShader);
 
-			FIntPoint TileBufferSize = GetNumTiles16x16(VelocityBufferSize);
-			FIntPoint TileCount = GetNumTiles16x16(Context.View.ViewRect.Size());
-			FVector2D TileUVMaxValue(
-				(TileCount.X - 0.5f) / float(TileBufferSize.X),
-				(TileCount.Y - 0.5f) / float(TileBufferSize.Y));
+	BEGIN_SHADER_PARAMETER_STRUCT(FParameters, )
+		SHADER_PARAMETER(FMatrix, WorldToClipPrev)
+		SHADER_PARAMETER_STRUCT_REF(FViewUniformShaderParameters, View)
+		SHADER_PARAMETER_RDG_TEXTURE(Texture2D, ColorTexture)
+		SHADER_PARAMETER_RDG_TEXTURE(Texture2D, DepthTexture)
+		SHADER_PARAMETER_RDG_TEXTURE(Texture2D, VelocityTexture)
 
-			SetShaderValue(Context.RHICmdList, ShaderRHI, TileUVMax, TileUVMaxValue);
-		}
-	}
+		SHADER_PARAMETER_STRUCT(FScreenPassTextureViewportParameters, Color)
+		SHADER_PARAMETER_STRUCT(FScreenPassTextureViewportParameters, Velocity)
 
-	template <typename TRHICmdList>
-	void UnsetParameters(TRHICmdList& RHICmdList)
-	{
-		const FComputeShaderRHIParamRef ShaderRHI = GetComputeShader();
-		OutComputeTex.UnsetUAV(RHICmdList, ShaderRHI);
-	}
+		SHADER_PARAMETER_STRUCT(FScreenPassTextureViewportTransform, ColorToVelocity)
 
-	// FShader interface.
-	virtual bool Serialize(FArchive& Ar) override
-	{
-		bool bShaderHasOutdatedParameters = FGlobalShader::Serialize(Ar);
-		// CS params
-		Ar << OutComputeTex << MotionBlurComputeParams;
-		// PS params
-		Ar << PostprocessParameter << SceneTextureParameters << MotionBlurParameters << SceneColorBufferUVToViewBufferUV << TileUVMax;
-		return bShaderHasOutdatedParameters;
-	}
+		SHADER_PARAMETER_SAMPLER(SamplerState, ColorSampler)
+		SHADER_PARAMETER_SAMPLER(SamplerState, VelocitySampler)
 
-	static const TCHAR* GetSourceFilename()
-	{
-		return TEXT("/Engine/Private/PostProcessMotionBlur.usf");
-	}
+		RENDER_TARGET_BINDING_SLOTS()
+	END_SHADER_PARAMETER_STRUCT()
 
-	static const TCHAR* GetFunctionName()
+	static bool ShouldCompilePermutation(const FGlobalShaderPermutationParameters& Parameters)
 	{
-		return TEXT("MainCS");
+		return IsFeatureLevelSupported(Parameters.Platform, ERHIFeatureLevel::SM4);
 	}
 };
 
-#define VARIATION2(A, B) typedef FPostProcessMotionBlurCS<A, B> FPostProcessMotionBlurCS##A##B; \
-	IMPLEMENT_SHADER_TYPE2(FPostProcessMotionBlurCS##A##B, SF_Compute);
-
-#define VARIATION1(A) VARIATION2(A, 0) VARIATION2(A, 1) VARIATION2(A, 2)
-
-VARIATION1(1)
-VARIATION1(2)
-VARIATION1(3)
-VARIATION1(4)
+IMPLEMENT_GLOBAL_SHADER(FMotionBlurVisualizePS, "/Engine/Private/PostProcessMotionBlur.usf", "VisualizeMotionBlurPS", SF_Pixel);
 
-#undef VARIATION1
-#undef VARIATION2
+TGlobalResource<FSpriteIndexBuffer<8>> GScatterQuadIndexBuffer;
 
-
-template< uint32 Quality, uint32 InnerLoopConfig >
-static void SetMotionBlurShaderNewTempl( const FRenderingCompositePassContext& Context, float Scale )
-{
-	FGraphicsPipelineStateInitializer GraphicsPSOInit;
-	Context.RHICmdList.ApplyCachedRenderTargets(GraphicsPSOInit);
-	GraphicsPSOInit.BlendState = TStaticBlendState<>::GetRHI();
-	GraphicsPSOInit.RasterizerState = TStaticRasterizerState<>::GetRHI();
-	GraphicsPSOInit.DepthStencilState = TStaticDepthStencilState<false, CF_Always>::GetRHI();
-
-	TShaderMapRef< FPostProcessVS > VertexShader( Context.GetShaderMap() );
-	TShaderMapRef< FPostProcessMotionBlurPS< Quality, InnerLoopConfig > >	PixelShader( Context.GetShaderMap() );
-
-	GraphicsPSOInit.BoundShaderState.VertexDeclarationRHI = GFilterVertexDeclaration.VertexDeclarationRHI;
-	GraphicsPSOInit.BoundShaderState.VertexShaderRHI = GETSAFERHISHADER_VERTEX(*VertexShader);
-	GraphicsPSOInit.BoundShaderState.PixelShaderRHI = GETSAFERHISHADER_PIXEL(*PixelShader);
-	GraphicsPSOInit.PrimitiveType = PT_TriangleList;
-
-	SetGraphicsPipelineState(Context.RHICmdList, GraphicsPSOInit);
-
-	VertexShader->SetParameters(Context);
-	PixelShader->SetParameters(Context, Scale);
-}
-
-
-template< uint32 Quality >
-static void SetMotionBlurShaderNewTempl(const FRenderingCompositePassContext& Context, float Scale, uint32 InnerLoopConfig)
+enum class EMotionBlurFilterPass : uint32
 {
-	switch (InnerLoopConfig)
-	{
-	case 0: SetMotionBlurShaderNewTempl<Quality, 0>(Context, Scale); break;
-	case 1: SetMotionBlurShaderNewTempl<Quality, 1>(Context, Scale); break;
-	case 2: SetMotionBlurShaderNewTempl<Quality, 2>(Context, Scale); break;
-	default: check(0);
-	}
-}
-
+	Separable0,
+	Separable1,
+	Unified,
+	MAX
+};
 
-void FRCPassPostProcessMotionBlur::Process(FRenderingCompositePassContext& Context)
+void ComputeMotionBlurVelocity(
+	FRDGBuilder& GraphBuilder,
+	FScreenPassContextRef Context,
+	const FScreenPassTexture& ColorTexture,
+	const FScreenPassTexture& DepthTexture,
+	const FScreenPassTexture& VelocityTexture,
+	FScreenPassTexture& OutVelocityFlatTexture,
+	FScreenPassTexture& OutVelocityTileTexture)
 {
-	const FPooledRenderTargetDesc* InputDesc = GetInputDesc(ePId_Input0);
-	AsyncEndFence = FComputeFenceRHIRef();
+	// Validate inputs.
+	check(Context);
+	check(ColorTexture.IsValid());
+	check(DepthTexture.IsValid());
+	check(VelocityTexture.IsValid());
+	check(VelocityTexture.GetViewport() == DepthTexture.GetViewport());
+
+	const FScreenPassTextureViewport& ColorViewport = ColorTexture.GetViewport();
+	const FScreenPassTextureViewport& VelocityViewport = VelocityTexture.GetViewport();
+
+	OutVelocityFlatTexture =
+		FScreenPassTexture::Create(
+			GraphBuilder.CreateTexture(
+				FRDGTextureDesc::Create2DDesc(
+					VelocityViewport.Extent,
+					PF_FloatR11G11B10,
+					FClearValueBinding::None,
+					GFastVRamConfig.VelocityFlat,
+					TexCreate_ShaderResource | TexCreate_UAV,
+					false),
+				TEXT("VelocityFlat")),
+			VelocityViewport.Rect);
+
+	const FIntPoint VelocityTileCount = GetMotionBlurTileCount(VelocityViewport.Rect.Size());
+
+	const FScreenPassTexture VelocityTileTexture =
+		FScreenPassTexture::Create(
+			GraphBuilder.CreateTexture(
+				FRDGTextureDesc::Create2DDesc(
+					VelocityTileCount,
+					PF_FloatRGBA,
+					FClearValueBinding::None,
+					GFastVRamConfig.VelocityMax,
+					TexCreate_ShaderResource | TexCreate_UAV,
+					false),
+				TEXT("VelocityTile")),
+			FIntRect(FIntPoint::ZeroValue, VelocityTileCount));
+
+	const FScreenPassTextureViewport& VelocityTileViewport = VelocityTileTexture.GetViewport();
+
+	const FMotionBlurParameters MotionBlurParametersNoScale = GetMotionBlurParameters(Context, ColorViewport.Extent, 1.0f);
+
+	// Velocity flatten pass: combines depth / velocity into a single target for sampling efficiency.
+	{
+		FMotionBlurVelocityFlattenCS::FParameters* PassParameters = GraphBuilder.AllocParameters<FMotionBlurVelocityFlattenCS::FParameters>();
+		PassParameters->MotionBlur = MotionBlurParametersNoScale;
+		PassParameters->View = Context->View.ViewUniformBuffer;
+		PassParameters->Velocity = VelocityTexture.GetViewportParameters();
+		PassParameters->DepthTexture = DepthTexture.GetRDGTexture();
+		PassParameters->VelocityTexture = VelocityTexture.GetRDGTexture();
+		PassParameters->OutVelocityFlatTexture = GraphBuilder.CreateUAV(OutVelocityFlatTexture.GetRDGTexture());
+		PassParameters->OutVelocityTileTexture = GraphBuilder.CreateUAV(VelocityTileTexture.GetRDGTexture());
+
+		TShaderMapRef<FMotionBlurVelocityFlattenCS> ComputeShader(Context->ShaderMap);
+		FComputeShaderUtils::AddPass(
+			GraphBuilder,
+			RDG_EVENT_NAME("Velocity Flatten"),
+			*ComputeShader,
+			PassParameters,
+			FComputeShaderUtils::GetGroupCount(VelocityViewport.Extent, FComputeShaderUtils::kGolden2DGroupSize));
+	}
+
+	OutVelocityTileTexture =
+		FScreenPassTexture::Create(
+			GraphBuilder.CreateTexture(
+				FRDGTextureDesc::Create2DDesc(
+					VelocityTileViewport.Extent,
+					PF_FloatRGBA,
+					FClearValueBinding::None,
+					GFastVRamConfig.MotionBlur,
+					TexCreate_ShaderResource | TexCreate_UAV | TexCreate_RenderTargetable,
+					false),
+				TEXT("DilatedVelocityTile")),
+			VelocityTileViewport.Rect);
+
+	FMotionBlurVelocityDilateParameters VelocityDilateParameters;
+	VelocityDilateParameters.MotionBlur = MotionBlurParametersNoScale;
+	VelocityDilateParameters.VelocityTile = VelocityTileTexture.GetViewportParameters();
+	VelocityDilateParameters.VelocityTileTexture = VelocityTileTexture.GetRDGTexture();
+
+	if (IsMotionBlurScatterRequired(Context, ColorViewport))
+	{
+		const FScreenPassTexture ScatterVelocityTileDepthTexture =
+			FScreenPassTexture::Create(
+				GraphBuilder.CreateTexture(
+					FRDGTextureDesc::Create2DDesc(
+						VelocityTileViewport.Extent,
+						PF_ShadowDepth,
+						FClearValueBinding::DepthOne,
+						TexCreate_None,
+						TexCreate_DepthStencilTargetable,
+						false),
+					TEXT("DilatedVelocityDepth")),
+				VelocityTileViewport.Rect);
+
+		FMotionBlurVelocityDilateScatterParameters* PassParameters = GraphBuilder.AllocParameters<FMotionBlurVelocityDilateScatterParameters>();
+		PassParameters->Dilate = VelocityDilateParameters;
+
+		PassParameters->RenderTargets.DepthStencil =
+			FDepthStencilBinding(
+				ScatterVelocityTileDepthTexture.GetRDGTexture(),
+				ERenderTargetLoadAction::ENoAction,
+				ERenderTargetStoreAction::EStore,
+				ERenderTargetLoadAction::ENoAction,
+				ERenderTargetStoreAction::EStore,
+				FExclusiveDepthStencil::DepthWrite_StencilNop);
+
+		PassParameters->RenderTargets[0] =
+			FRenderTargetBinding(
+				OutVelocityTileTexture.GetRDGTexture(),
+				ERenderTargetLoadAction::ENoAction,
+				ERenderTargetStoreAction::EStore);
+
+		TShaderMapRef<FMotionBlurVelocityDilateScatterVS> VertexShader(Context->ShaderMap);
+		TShaderMapRef<FMotionBlurVelocityDilateScatterPS> PixelShader(Context->ShaderMap);
+
+		GraphBuilder.AddPass(
+			RDG_EVENT_NAME("VelocityTileScatter %dx%d", VelocityTileCount.X, VelocityTileCount.Y),
+			PassParameters,
+			ERenderGraphPassFlags::None,
+			[VertexShader, PixelShader, VelocityTileCount, PassParameters](FRHICommandListImmediate& RHICmdList)
+		{
+			FRHIVertexShader* RHIVertexShader = GETSAFERHISHADER_VERTEX(*VertexShader);
 
-	if(!InputDesc)
-	{
-		// input is not hooked up correctly
-		return;
-	}
+			FGraphicsPipelineStateInitializer GraphicsPSOInit;
+			GraphicsPSOInit.BoundShaderState.VertexDeclarationRHI = GEmptyVertexDeclaration.VertexDeclarationRHI;
+			GraphicsPSOInit.BoundShaderState.VertexShaderRHI = RHIVertexShader;
+			GraphicsPSOInit.BoundShaderState.PixelShaderRHI = GETSAFERHISHADER_PIXEL(*PixelShader);
+			GraphicsPSOInit.PrimitiveType = PT_TriangleList;
+			GraphicsPSOInit.RasterizerState = TStaticRasterizerState<>::GetRHI();
+			RHICmdList.ApplyCachedRenderTargets(GraphicsPSOInit);
 
-	const FViewInfo& View = Context.View;
+			// Max >= Min so no need to clear on second pass
+			RHICmdList.SetViewport(0, 0, 0.0f, VelocityTileCount.X, VelocityTileCount.Y, 1.0f);
 
-	// we assume the input and output is full resolution
+			// Min, Max
+			for (uint32 ScatterPassIndex = 0; ScatterPassIndex < static_cast<uint32>(EMotionBlurVelocityScatterPass::MAX); ScatterPassIndex++)
+			{
+				const EMotionBlurVelocityScatterPass ScatterPass = static_cast<EMotionBlurVelocityScatterPass>(ScatterPassIndex);
 
-	FIntPoint SrcSize = InputDesc->Extent;
-	FIntPoint DestSize = PassOutputs[0].RenderTargetDesc.Extent;
+				if (ScatterPass == EMotionBlurVelocityScatterPass::DrawMin)
+				{
+					GraphicsPSOInit.BlendState = TStaticBlendStateWriteMask<CW_RGBA>::GetRHI();
+					GraphicsPSOInit.DepthStencilState = TStaticDepthStencilState<true, CF_Less>::GetRHI();
+				}
+				else
+				{
+					GraphicsPSOInit.BlendState = TStaticBlendStateWriteMask<CW_BA>::GetRHI();
+					GraphicsPSOInit.DepthStencilState = TStaticDepthStencilState<true, CF_Greater>::GetRHI();
+				}
 
-	FIntRect SrcRect = Context.SceneColorViewRect;
-	FIntRect DestRect = SrcRect;
+				SetGraphicsPipelineState(RHICmdList, GraphicsPSOInit);
 
-	float BlurScaleLUT[] =
-	{
-		1.0f - 0.5f / 4.0f,
-		1.0f - 0.5f / 6.0f,
-		1.0f - 0.5f / 8.0f,
-		1.0f - 0.5f / 16.0f,
-		1.0f /  4.0f * CVarMotionBlur2ndScale.GetValueOnRenderThread(),
-		1.0f /  6.0f * CVarMotionBlur2ndScale.GetValueOnRenderThread(),
-		1.0f /  8.0f * CVarMotionBlur2ndScale.GetValueOnRenderThread(),
-		1.0f / 16.0f * CVarMotionBlur2ndScale.GetValueOnRenderThread(),
-	};
-	float Scale = Pass >= 0 ? BlurScaleLUT[ (Pass * 4) + (Quality - 1) ] : 1.0f;
+				PassParameters->ScatterPass = ScatterPassIndex;
 
-	SCOPED_DRAW_EVENTF(Context.RHICmdList, MotionBlur, TEXT("MotionBlur%s %dx%d"),
-		bIsComputePass?TEXT("Compute"):TEXT(""), SrcRect.Width(), SrcRect.Height());
+				SetShaderParameters(RHICmdList, *VertexShader, RHIVertexShader, *PassParameters);
 
-	const FSceneRenderTargetItem& DestRenderTarget = PassOutputs[0].RequestSurface(Context);
+				// Needs to be the same on shader side (faster on NVIDIA and AMD)
+				const int32 QuadsPerInstance = 8;
 
-	uint32 InnerLoopConfig = 1;
-	if (Context.View.PrimaryScreenPercentageMethod == EPrimaryScreenPercentageMethod::TemporalUpscale)
-	{
-		InnerLoopConfig = 2;
-	}
-	else if (DestSize == SrcRect.Size() && SrcRect.Min.X == 0 && SrcRect.Min.Y == 0)
-	{
-		InnerLoopConfig = 0;
-	}
-
-	if (bIsComputePass)
-	{
-		// Common setup
-		// #todo-renderpasses remove once everything is renderpasses.
-		UnbindRenderTargets(Context.RHICmdList);
-		Context.SetViewportAndCallRHI(DestRect, 0.0f, 1.0f);
-		
-		static FName AsyncEndFenceName(TEXT("AsyncMotionBlurEndFence"));
-		AsyncEndFence = Context.RHICmdList.CreateComputeFence(AsyncEndFenceName);
-
-		if (IsAsyncComputePass())
-		{
-			// Async path
-			FRHIAsyncComputeCommandListImmediate& RHICmdListComputeImmediate = FRHICommandListExecutor::GetImmediateAsyncComputeCommandList();
-			{
-				SCOPED_COMPUTE_EVENT(RHICmdListComputeImmediate, AsyncMotionBlur);
-				WaitForInputPassComputeFences(RHICmdListComputeImmediate);
-				RHICmdListComputeImmediate.TransitionResource(EResourceTransitionAccess::ERWBarrier, EResourceTransitionPipeline::EGfxToCompute, DestRenderTarget.UAV);
-				DispatchCS(RHICmdListComputeImmediate, Context, DestSize, DestRect, DestRenderTarget.UAV, Scale, InnerLoopConfig);
-				RHICmdListComputeImmediate.TransitionResource(EResourceTransitionAccess::EReadable, EResourceTransitionPipeline::EComputeToGfx, DestRenderTarget.UAV, AsyncEndFence);
+				RHICmdList.SetStreamSource(0, nullptr, 0);
+				RHICmdList.DrawIndexedPrimitive(GScatterQuadIndexBuffer.IndexBufferRHI, 0, 0, 32, 0, 2 * QuadsPerInstance, FMath::DivideAndRoundUp(VelocityTileCount.X * VelocityTileCount.Y, QuadsPerInstance));
 			}
-			FRHIAsyncComputeCommandListImmediate::ImmediateDispatch(RHICmdListComputeImmediate);
-		}
-		else
-		{
-			// Direct path
-			WaitForInputPassComputeFences(Context.RHICmdList);
-			Context.RHICmdList.TransitionResource(EResourceTransitionAccess::ERWBarrier, EResourceTransitionPipeline::EGfxToCompute, DestRenderTarget.UAV);
-			DispatchCS(Context.RHICmdList, Context, DestSize, DestRect, DestRenderTarget.UAV, Scale, InnerLoopConfig);			
-			Context.RHICmdList.TransitionResource(EResourceTransitionAccess::EReadable, EResourceTransitionPipeline::EComputeToGfx, DestRenderTarget.UAV, AsyncEndFence);
-		}
+		});
 	}
 	else
 	{
-		WaitForInputPassComputeFences(Context.RHICmdList);
+		FMotionBlurVelocityDilateGatherCS::FParameters* PassParameters = GraphBuilder.AllocParameters<FMotionBlurVelocityDilateGatherCS::FParameters>();
+		PassParameters->Dilate = VelocityDilateParameters;
+		PassParameters->OutDilatedVelocityTileTexture = GraphBuilder.CreateUAV(OutVelocityTileTexture.GetRDGTexture());
 
-		ERenderTargetLoadAction LoadAction = Context.GetLoadActionForRenderTarget(DestRenderTarget);
-
-		FRHIRenderPassInfo RPInfo(DestRenderTarget.TargetableTexture, MakeRenderTargetActions(LoadAction, ERenderTargetStoreAction::EStore));
-		Context.RHICmdList.BeginRenderPass(RPInfo, TEXT("MotionBlur"));
-		{
-			Context.SetViewportAndCallRHI(SrcRect);
-
-			if (Quality == 1)
-			{
-				SetMotionBlurShaderNewTempl<1>(Context, Scale, InnerLoopConfig);
-			}
-			else if (Quality == 2)
-			{
-				SetMotionBlurShaderNewTempl<2>(Context, Scale, InnerLoopConfig);
-			}
-			else if (Quality == 3 || Pass > 0)
-			{
-				SetMotionBlurShaderNewTempl<3>(Context, Scale, InnerLoopConfig);
-			}
-			else
-			{
-				check(Quality == 4);
-				SetMotionBlurShaderNewTempl<4>(Context, Scale, InnerLoopConfig);
-			}
-
-			TShaderMapRef<FPostProcessVS> VertexShader(Context.GetShaderMap());
-
-			DrawPostProcessPass(
-				Context.RHICmdList,
-				0, 0,
-				SrcRect.Width(), SrcRect.Height(),
-				SrcRect.Min.X, SrcRect.Min.Y,
-				SrcRect.Width(), SrcRect.Height(),
-				SrcRect.Size(),
-				SrcSize,
-				*VertexShader,
-				View.StereoPass,
-				Context.HasHmdMesh(),
-				EDRF_UseTriangleOptimization);
-		}
-		Context.RHICmdList.EndRenderPass();
-		Context.RHICmdList.CopyToResolveTarget(DestRenderTarget.TargetableTexture, DestRenderTarget.ShaderResourceTexture, FResolveParams());
+		TShaderMapRef<FMotionBlurVelocityDilateGatherCS> ComputeShader(Context->ShaderMap);
+		FComputeShaderUtils::AddPass(
+			GraphBuilder,
+			RDG_EVENT_NAME("VelocityTileGatherCS %dx%d", VelocityTileCount.X, VelocityTileCount.Y),
+			*ComputeShader,
+			PassParameters,
+			FComputeShaderUtils::GetGroupCount(VelocityTileCount, FComputeShaderUtils::kGolden2DGroupSize));
 	}
+
+	check(OutVelocityFlatTexture.IsValid());
+	check(OutVelocityTileTexture.IsValid());
 }
 
-template <typename TRHICmdList>
-void FRCPassPostProcessMotionBlur::DispatchCS(TRHICmdList& RHICmdList, FRenderingCompositePassContext& Context, const FIntPoint& DestSize, const FIntRect& DestRect, FUnorderedAccessViewRHIParamRef DestUAV, float Scale, uint32 InnerLoopConfig)
+FScreenPassTexture ComputeMotionBlurFilter(
+	FRDGBuilder& GraphBuilder,
+	FScreenPassContextRef Context,
+	const FScreenPassTexture& ColorTexture,
+	const FScreenPassTexture& VelocityFlatTexture,
+	const FScreenPassTexture& VelocityTileTexture,
+	EMotionBlurFilterPass MotionBlurFilterPass,
+	EMotionBlurQuality MotionBlurQuality)
 {
-	auto ShaderMap = Context.GetShaderMap();
+	// Validate inputs.
+	check(Context);
+	check(ColorTexture.IsValid());
+	check(VelocityFlatTexture.IsValid());
+	check(VelocityTileTexture.IsValid());
+	check(MotionBlurFilterPass != EMotionBlurFilterPass::MAX);
+	check(MotionBlurQuality != EMotionBlurQuality::MAX);
 
-	uint32 GroupSizeX = FMath::DivideAndRoundUp(DestSize.X, GMotionBlurComputeTileSizeX);
-	uint32 GroupSizeY = FMath::DivideAndRoundUp(DestSize.Y, GMotionBlurComputeTileSizeY);
+	const FScreenPassTextureViewport ColorViewport = ColorTexture.GetViewport();
 
-#define DISPATCH_CASE(A, B)																\
-	case A:																				\
-	{																					\
-		TShaderMapRef<FPostProcessMotionBlurCS<A, B>> ComputeShader(ShaderMap);			\
-		RHICmdList.SetComputeShader(ComputeShader->GetComputeShader());					\
-		ComputeShader->SetParameters(RHICmdList, Context, DestSize, DestUAV, Scale);	\
-		DispatchComputeShader(RHICmdList, *ComputeShader, GroupSizeX, GroupSizeY, 1);	\
-		ComputeShader->UnsetParameters(RHICmdList);										\
-	}																					\
-	break;
+	const float MotionBlur2ndScale = CVarMotionBlur2ndScale.GetValueOnRenderThread();
 
-	const uint32 CSQuality = (Pass > 0) ? 3 : Quality;
+	const bool bUseCompute = Context->bUseComputePasses;
 
-	if (InnerLoopConfig == 0)
+	const float BlurScaleLUT[static_cast<uint32>(EMotionBlurFilterPass::MAX)][static_cast<uint32>(EMotionBlurQuality::MAX)] =
 	{
-		switch (CSQuality)
+		// Separable0
 		{
-		DISPATCH_CASE(1, 0)
-		DISPATCH_CASE(2, 0)
-		DISPATCH_CASE(3, 0)
-		DISPATCH_CASE(4, 0)
-		default:
-			check(0);
-		}
-	}
-	else if (InnerLoopConfig == 1)
-	{
-		switch (CSQuality)
+			1.0f - 0.5f / 4.0f,
+			1.0f - 0.5f / 6.0f,
+			1.0f - 0.5f / 8.0f,
+			1.0f - 0.5f / 16.0f
+		},
+
+		// Separable1
 		{
-		DISPATCH_CASE(1, 1)
-		DISPATCH_CASE(2, 1)
-		DISPATCH_CASE(3, 1)
-		DISPATCH_CASE(4, 1)
-		default:
-			check(0);
-		}
-	}
-	else if (InnerLoopConfig == 2)
-	{
-		switch (CSQuality)
+			1.0f / 4.0f  * MotionBlur2ndScale,
+			1.0f / 6.0f  * MotionBlur2ndScale,
+			1.0f / 8.0f  * MotionBlur2ndScale,
+			1.0f / 16.0f * MotionBlur2ndScale
+		},
+
+		// Unified
 		{
-		DISPATCH_CASE(1, 2)
-		DISPATCH_CASE(2, 2)
-		DISPATCH_CASE(3, 2)
-		DISPATCH_CASE(4, 2)
-		default:
-			check(0);
+			1.0f,
+			1.0f,
+			1.0f,
+			1.0f
 		}
-	}
-	else
-	{
-		check(0);
-	}
-
-#undef DISPATCH_CASE
-}
-
-FPooledRenderTargetDesc FRCPassPostProcessMotionBlur::ComputeOutputDesc(EPassOutputId InPassOutputId) const
-{
-	FPooledRenderTargetDesc Ret = GetInput(ePId_Input0)->GetOutput()->RenderTargetDesc;
-	Ret.Reset();
-
-	Ret.TargetableFlags &= ~(TexCreate_RenderTargetable | TexCreate_UAV);
-	Ret.TargetableFlags |= bIsComputePass ? TexCreate_UAV : TexCreate_RenderTargetable;
-
-	if (!FPostProcessing::HasAlphaChannelSupport())
-	{
-		Ret.Format = PF_FloatRGB;
-	}
-	Ret.Flags |= GFastVRamConfig.MotionBlur;
-	Ret.DebugName = TEXT("MotionBlur");
-	Ret.AutoWritable = false;
+	};
 
-	return Ret;
-}
+	const float BlurScale = BlurScaleLUT[static_cast<uint32>(MotionBlurFilterPass)][static_cast<uint32>(MotionBlurQuality)];
+
+	const auto GetFilterParameters = [&]()
+	{
+		FMotionBlurFilterParameters Parameters;
+		Parameters.MotionBlur = GetMotionBlurParameters(Context, ColorTexture.GetViewport().Extent, BlurScale);
+		Parameters.Color = ColorTexture.GetViewportParameters();
+		Parameters.Velocity = VelocityFlatTexture.GetViewportParameters();
+		Parameters.VelocityTile = VelocityTileTexture.GetViewportParameters();
+		Parameters.ColorToVelocity = GetScreenPassTextureViewportTransform(Parameters.Color, Parameters.Velocity);
+		Parameters.ColorToVelocityTile = GetScreenPassTextureViewportTransform(Parameters.Color, Parameters.VelocityTile);
+		Parameters.ColorTexture = ColorTexture.GetRDGTexture();
+		Parameters.VelocityFlatTexture = VelocityFlatTexture.GetRDGTexture();
+		Parameters.VelocityTileTexture = VelocityTileTexture.GetRDGTexture();
+		Parameters.ColorSampler = GetMotionBlurColorSampler();
+		Parameters.VelocitySampler = GetMotionBlurVelocitySampler();
+		return Parameters;
+	};
 
+	FRDGTextureDesc OutColorDesc = ColorTexture.GetRDGTexture()->Desc;
+	OutColorDesc.Reset();
+	OutColorDesc.TargetableFlags &= ~(TexCreate_RenderTargetable | TexCreate_UAV);
+	OutColorDesc.TargetableFlags |= bUseCompute ? TexCreate_UAV : TexCreate_RenderTargetable;
+	OutColorDesc.Flags |= GFastVRamConfig.MotionBlur;
+	OutColorDesc.AutoWritable = false;
+	OutColorDesc.Format = FPostProcessing::HasAlphaChannelSupport() ? PF_FloatRGBA : PF_FloatRGB;
 
+	const FScreenPassTexture OutColorTexture =
+		FScreenPassTexture::Create(
+			GraphBuilder.CreateTexture(OutColorDesc, TEXT("MotionBlur")),
+			ColorTexture.GetViewport().Rect);
 
+	FRDGEventName PassName = RDG_EVENT_NAME("Motion Blur%dx%d", ColorViewport.Rect.Size().X, ColorViewport.Rect.Size().Y);
 
-class FPostProcessVisualizeMotionBlurPS : public FGlobalShader
-{
-	DECLARE_SHADER_TYPE(FPostProcessVisualizeMotionBlurPS, Global);
-
-	static bool ShouldCompilePermutation(const FGlobalShaderPermutationParameters& Parameters)
+	if (bUseCompute)
 	{
-		return IsFeatureLevelSupported(Parameters.Platform, ERHIFeatureLevel::SM4);
-	}
-
-	/** Default constructor. */
-	FPostProcessVisualizeMotionBlurPS() {}
-
-public:
-	FPostProcessPassParameters PostprocessParameter;
-	FSceneTextureShaderParameters SceneTextureParameters;
-	FShaderParameter PrevViewProjMatrix;
+		FMotionBlurFilterCS::FParameters* PassParameters = GraphBuilder.AllocParameters<FMotionBlurFilterCS::FParameters>();
+		PassParameters->Filter = GetFilterParameters();
+		PassParameters->OutColorTexture = GraphBuilder.CreateUAV(OutColorTexture.GetRDGTexture());
 
-	/** Initialization constructor. */
-	FPostProcessVisualizeMotionBlurPS(const ShaderMetaType::CompiledShaderInitializerType& Initializer)
-		: FGlobalShader(Initializer)
-	{
-		PostprocessParameter.Bind(Initializer.ParameterMap);
-		SceneTextureParameters.Bind(Initializer);
-		PrevViewProjMatrix.Bind(Initializer.ParameterMap, TEXT("PrevViewProjMatrix"));
-	}
+		FMotionBlurFilterCS::FPermutationDomain PermutationVector;
+		PermutationVector.Set<FMotionBlurQualityDimension>(MotionBlurQuality);
 
-	// FShader interface.
-	virtual bool Serialize(FArchive& Ar) override
-	{
-		bool bShaderHasOutdatedParameters = FGlobalShader::Serialize(Ar);
-		Ar << PostprocessParameter << SceneTextureParameters << PrevViewProjMatrix;
-		return bShaderHasOutdatedParameters;
+		TShaderMapRef<FMotionBlurFilterCS> ComputeShader(Context->ShaderMap, PermutationVector);
+		FComputeShaderUtils::AddPass(
+			GraphBuilder,
+			MoveTemp(PassName),
+			*ComputeShader,
+			PassParameters,
+			FComputeShaderUtils::GetGroupCount(ColorViewport.Extent, FComputeShaderUtils::kGolden2DGroupSize));
 	}
-
-	template <typename TRHICmdList>
-	void SetParameters(TRHICmdList& RHICmdList, const FRenderingCompositePassContext& Context)
+	else
 	{
-		const FPixelShaderRHIParamRef ShaderRHI = GetPixelShader();
-
-		FGlobalShader::SetParameters<FViewUniformShaderParameters>(RHICmdList, ShaderRHI, Context.View.ViewUniformBuffer);
+		FMotionBlurFilterPS::FParameters* PassParameters = GraphBuilder.AllocParameters<FMotionBlurFilterPS::FParameters>();
+		PassParameters->Filter = GetFilterParameters();
+		PassParameters->RenderTargets[0] = FRenderTargetBinding(OutColorTexture.GetRDGTexture(), ERenderTargetLoadAction::ENoAction, ERenderTargetStoreAction::EStore);
 
-		SceneTextureParameters.Set(RHICmdList, ShaderRHI, Context.View.FeatureLevel, ESceneTextureSetupMode::All);
-
-		{
-			bool bFiltered = false;
-
-#if !(UE_BUILD_SHIPPING || UE_BUILD_TEST)
-			bFiltered = CVarMotionBlurFiltering.GetValueOnRenderThread() != 0;
-#endif // !(UE_BUILD_SHIPPING || UE_BUILD_TEST)
-
-			if(bFiltered)
-			{
-				PostprocessParameter.SetPS(RHICmdList, ShaderRHI, Context, TStaticSamplerState<SF_Bilinear,AM_Border,AM_Border,AM_Clamp>::GetRHI());
-			}
-			else
-			{
-				PostprocessParameter.SetPS(RHICmdList, ShaderRHI, Context, TStaticSamplerState<SF_Point,AM_Border,AM_Border,AM_Clamp>::GetRHI());
-			}
-		}
-		
-		if( Context.View.Family->EngineShowFlags.CameraInterpolation )
-		{
-			// Instead of finding the world space position of the current pixel, calculate the world space position offset by the camera position, 
-			// then translate by the difference between last frame's camera position and this frame's camera position,
-			// then apply the rest of the transforms.  This effectively avoids precision issues near the extents of large levels whose world space position is very large.
-			FVector ViewOriginDelta = Context.View.ViewMatrices.GetViewOrigin() - Context.View.PrevViewInfo.ViewMatrices.GetViewOrigin();
-			SetShaderValue(RHICmdList, ShaderRHI, PrevViewProjMatrix, FTranslationMatrix(ViewOriginDelta) * Context.View.PrevViewInfo.ViewMatrices.ComputeViewRotationProjectionMatrix());
-		}
-		else
-		{
-			SetShaderValue(RHICmdList, ShaderRHI, PrevViewProjMatrix, Context.View.ViewMatrices.ComputeViewRotationProjectionMatrix());
-		}
-	}
-
-	static const TCHAR* GetSourceFilename()
-	{
-		return TEXT("/Engine/Private/PostProcessMotionBlur.usf");
-	}
+		FMotionBlurFilterPS::FPermutationDomain PermutationVector;
+		PermutationVector.Set<FMotionBlurQualityDimension>(MotionBlurQuality);
+		TShaderMapRef<FMotionBlurFilterPS> PixelShader(Context->ShaderMap, PermutationVector);
 
-	static const TCHAR* GetFunctionName()
-	{
-		return TEXT("VisualizeMotionBlurPS");
+		AddDrawScreenPass(
+			GraphBuilder,
+			MoveTemp(PassName),
+			Context,
+			ColorViewport,
+			ColorViewport,
+			*PixelShader,
+			PassParameters);
 	}
-};
 
-IMPLEMENT_SHADER_TYPE3(FPostProcessVisualizeMotionBlurPS, SF_Pixel);
+	check(OutColorTexture.IsValid());
+	return OutColorTexture;
+}
 
-void FRCPassPostProcessVisualizeMotionBlur::Process(FRenderingCompositePassContext& Context)
+FScreenPassTexture ComputeMotionBlur(
+	FRDGBuilder& GraphBuilder,
+	FScreenPassContextRef Context,
+	const FScreenPassTexture& ColorTexture,
+	const FScreenPassTexture& DepthTexture,
+	const FScreenPassTexture& VelocityTexture)
 {
-	SCOPED_DRAW_EVENT(Context.RHICmdList, VisualizeMotionBlur);
-
-	const FPooledRenderTargetDesc* InputDesc = GetInputDesc(ePId_Input0);
-
-	if(!InputDesc)
-	{
-		// input is not hooked up correctly
-		return;
+	// Validate inputs.
+	check(Context);
+	check(ColorTexture.IsValid());
+	check(DepthTexture.IsValid());
+	check(VelocityTexture.IsValid());
+	check(VelocityTexture.GetViewport() == DepthTexture.GetViewport());
+
+	FScreenPassTexture VelocityFlatTexture;
+	FScreenPassTexture VelocityTileTexture;
+	ComputeMotionBlurVelocity(GraphBuilder, Context, ColorTexture, DepthTexture, VelocityTexture, VelocityFlatTexture, VelocityTileTexture);
+
+	const EMotionBlurQuality MotionBlurQuality = GetMotionBlurQuality();
+
+	const bool bUseSeparableFilter = CVarMotionBlurSeparable.GetValueOnRenderThread() != 0;
+
+	if (bUseSeparableFilter)
+	{
+		const FScreenPassTexture MotionBlurFilterTexture = ComputeMotionBlurFilter(
+			GraphBuilder,
+			Context,
+			ColorTexture,
+			VelocityFlatTexture,
+			VelocityTileTexture,
+			EMotionBlurFilterPass::Separable0,
+			MotionBlurQuality);
+
+		return ComputeMotionBlurFilter(
+			GraphBuilder,
+			Context,
+			MotionBlurFilterTexture,
+			VelocityFlatTexture,
+			VelocityTileTexture,
+			EMotionBlurFilterPass::Separable1,
+			MotionBlurQuality);
 	}
-
-	const FViewInfo& View = Context.View;
-	const FSceneViewFamily& ViewFamily = *(View.Family);
-
-	FIntPoint TexSize = InputDesc->Extent;
-
-	// we assume the input and output is full resolution
-
-	FIntPoint SrcSize = InputDesc->Extent;
-	FIntPoint DestSize = PassOutputs[0].RenderTargetDesc.Extent;
-
-	// e.g. 4 means the input texture is 4x smaller than the buffer size
-	uint32 ScaleFactor = FSceneRenderTargets::Get(Context.RHICmdList).GetBufferSizeXY().X / SrcSize.X;
-
-	FIntRect SrcRect = FIntRect::DivideAndRoundUp(View.ViewRect, ScaleFactor);
-	FIntRect DestRect = SrcRect;
-
-	const FSceneRenderTargetItem& DestRenderTarget = PassOutputs[0].RequestSurface(Context);
-
-	// Set the view family's render target/viewport.
-	FRHIRenderPassInfo RPInfo(DestRenderTarget.TargetableTexture, ERenderTargetActions::DontLoad_Store);
-	Context.RHICmdList.BeginRenderPass(RPInfo, TEXT("VisualizeMotionBlur"));
+	else
 	{
-		Context.SetViewportAndCallRHI(SrcRect);
-
-		FGraphicsPipelineStateInitializer GraphicsPSOInit;
-		Context.RHICmdList.ApplyCachedRenderTargets(GraphicsPSOInit);
-		GraphicsPSOInit.BlendState = TStaticBlendState<>::GetRHI();
-		GraphicsPSOInit.RasterizerState = TStaticRasterizerState<>::GetRHI();
-		GraphicsPSOInit.DepthStencilState = TStaticDepthStencilState<false, CF_Always>::GetRHI();
-
-		TShaderMapRef<FPostProcessVS> VertexShader(Context.GetShaderMap());
-		TShaderMapRef<FPostProcessVisualizeMotionBlurPS> PixelShader(Context.GetShaderMap());
-
-		GraphicsPSOInit.BoundShaderState.VertexDeclarationRHI = GFilterVertexDeclaration.VertexDeclarationRHI;
-		GraphicsPSOInit.BoundShaderState.VertexShaderRHI = GETSAFERHISHADER_VERTEX(*VertexShader);
-		GraphicsPSOInit.BoundShaderState.PixelShaderRHI = GETSAFERHISHADER_PIXEL(*PixelShader);
-		GraphicsPSOInit.PrimitiveType = PT_TriangleList;
-
-		SetGraphicsPipelineState(Context.RHICmdList, GraphicsPSOInit);
-
-		VertexShader->SetParameters(Context);
-		PixelShader->SetParameters(Context.RHICmdList, Context);
-
-		// Draw a quad mapping scene color to the view's render target
-		DrawRectangle(
-			Context.RHICmdList,
-			0, 0,
-			SrcRect.Width(), SrcRect.Height(),
-			SrcRect.Min.X, SrcRect.Min.Y,
-			SrcRect.Width(), SrcRect.Height(),
-			SrcRect.Size(),
-			SrcSize,
-			*VertexShader,
-			EDRF_UseTriangleOptimization);
+		return ComputeMotionBlurFilter(
+			GraphBuilder,
+			Context,
+			ColorTexture,
+			VelocityFlatTexture,
+			VelocityTileTexture,
+			EMotionBlurFilterPass::Unified,
+			MotionBlurQuality);
 	}
-	Context.RHICmdList.EndRenderPass();
-
-	FRenderTargetTemp TempRenderTarget(View, (const FTexture2DRHIRef&)DestRenderTarget.TargetableTexture);
-	FCanvas Canvas(&TempRenderTarget, NULL, ViewFamily.CurrentRealTime, ViewFamily.CurrentWorldTime, ViewFamily.DeltaWorldTime, Context.GetFeatureLevel());
-
-	float X = 20;
-	float Y = 38;
-	const float YStep = 14;
-	const float ColumnWidth = 200;
-
-	FString Line;
+}
 
-	Line = FString::Printf(TEXT("Visualize MotionBlur"));
-	Canvas.DrawShadowedString(X, Y += YStep, *Line, GetStatsFont(), FLinearColor(1, 1, 0));
+FScreenPassTexture VisualizeMotionBlur(
+	FRDGBuilder& GraphBuilder,
+	FScreenPassContextRef Context,
+	const FScreenPassTexture& ColorTexture,
+	const FScreenPassTexture& DepthTexture,
+	const FScreenPassTexture& VelocityTexture)
+{
+	// Validate inputs.
+	check(Context);
+	check(ColorTexture.IsValid());
+	check(DepthTexture.IsValid());
+	check(VelocityTexture.IsValid());
+
+	const FViewInfo& View = Context->View;
+
+	const FScreenPassTexture& TextureInput = ColorTexture;
+	const FScreenPassTexture TextureOutput = FScreenPassTexture::Create(
+		GraphBuilder.CreateTexture(TextureInput.GetRDGTexture()->Desc, TEXT("MotionBlurVisualize")),
+		TextureInput.GetViewport().Rect);
+
+	FMotionBlurVisualizePS::FParameters* PassParameters = GraphBuilder.AllocParameters<FMotionBlurVisualizePS::FParameters>();
+	PassParameters->WorldToClipPrev = GetPreviousWorldToClipMatrix(View);
+	PassParameters->View = View.ViewUniformBuffer;
+	PassParameters->ColorTexture = ColorTexture.GetRDGTexture();
+	PassParameters->DepthTexture = DepthTexture.GetRDGTexture();
+	PassParameters->VelocityTexture = VelocityTexture.GetRDGTexture();
+	PassParameters->Color = ColorTexture.GetViewportParameters();
+	PassParameters->Velocity = VelocityTexture.GetViewportParameters();
+	PassParameters->ColorToVelocity = GetScreenPassTextureViewportTransform(PassParameters->Color, PassParameters->Velocity);
+	PassParameters->ColorSampler = GetMotionBlurColorSampler();
+	PassParameters->VelocitySampler = GetMotionBlurVelocitySampler();
+	PassParameters->RenderTargets[0] = FRenderTargetBinding(TextureOutput.GetRDGTexture(), ERenderTargetLoadAction::ENoAction, ERenderTargetStoreAction::EStore);
+
+	TShaderMapRef<FMotionBlurVisualizePS> PixelShader(Context->ShaderMap);
+
+	GraphBuilder.AddPass(
+		RDG_EVENT_NAME("MotionBlurVisualize"),
+		PassParameters,
+		ERenderGraphPassFlags::None,
+		[Context, TextureInput, TextureOutput, PixelShader, PassParameters](FRHICommandListImmediate& RHICmdList)
+	{
+		DrawScreenPass(RHICmdList, Context, TextureOutput.GetViewport(), TextureInput.GetViewport(), *PixelShader, *PassParameters);
+
+		// Draw debug text.
+		{
+			const FViewInfo& View = Context->View;
+			const FSceneViewFamily& ViewFamily = Context->ViewFamily;
+			FRenderTargetTemp TempRenderTarget(static_cast<FTexture2DRHIParamRef>(TextureOutput.GetRHITexture()), TextureOutput.GetViewport().Extent);
+			FCanvas Canvas(&TempRenderTarget, nullptr, ViewFamily.CurrentRealTime, ViewFamily.CurrentWorldTime, ViewFamily.DeltaWorldTime, View.GetFeatureLevel());
 
-	static const auto MotionBlurDebugVar = IConsoleManager::Get().FindTConsoleVariableDataInt(TEXT("r.MotionBlurDebug"));
-	const int32 MotionBlurDebug = MotionBlurDebugVar ? MotionBlurDebugVar->GetValueOnRenderThread() : 0;
+			float X = 20;
+			float Y = 38;
+			const float YStep = 14;
+			const float ColumnWidth = 200;
 
-	Line = FString::Printf(TEXT("%d, %d"), ViewFamily.FrameNumber, MotionBlurDebug);
-	Canvas.DrawShadowedString(X, Y += YStep, TEXT("FrameNo, r.MotionBlurDebug:"), GetStatsFont(), FLinearColor(1, 1, 0));
-	Canvas.DrawShadowedString(X + ColumnWidth, Y, *Line, GetStatsFont(), FLinearColor(1, 1, 0));
+			FString Line;
 
-	static const auto VelocityTestVar = IConsoleManager::Get().FindTConsoleVariableDataInt(TEXT("r.VelocityTest"));
-	const int32 VelocityTest = VelocityTestVar ? VelocityTestVar->GetValueOnRenderThread() : 0;
+			Line = FString::Printf(TEXT("Visualize MotionBlur"));
+			Canvas.DrawShadowedString(X, Y += YStep, *Line, GetStatsFont(), FLinearColor(1, 1, 0));
 
-	extern bool IsParallelVelocity();
+			static const auto MotionBlurDebugVar = IConsoleManager::Get().FindTConsoleVariableDataInt(TEXT("r.MotionBlurDebug"));
+			const int32 MotionBlurDebug = MotionBlurDebugVar ? MotionBlurDebugVar->GetValueOnRenderThread() : 0;
 
-	Line = FString::Printf(TEXT("%d, %d, %d"), ViewFamily.bWorldIsPaused, VelocityTest, IsParallelVelocity());
-	Canvas.DrawShadowedString(X, Y += YStep, TEXT("Paused, r.VelocityTest, Parallel:"), GetStatsFont(), FLinearColor(1, 1, 0));
-	Canvas.DrawShadowedString(X + ColumnWidth, Y, *Line, GetStatsFont(), FLinearColor(1, 1, 0));
+			Line = FString::Printf(TEXT("%d, %d"), ViewFamily.FrameNumber, MotionBlurDebug);
+			Canvas.DrawShadowedString(X, Y += YStep, TEXT("FrameNo, r.MotionBlurDebug:"), GetStatsFont(), FLinearColor(1, 1, 0));
+			Canvas.DrawShadowedString(X + ColumnWidth, Y, *Line, GetStatsFont(), FLinearColor(1, 1, 0));
 
-	const FScene* Scene = (const FScene*)View.Family->Scene;
-	const FSceneViewState *SceneViewState = (const FSceneViewState*)View.State;
+			static const auto VelocityTestVar = IConsoleManager::Get().FindTConsoleVariableDataInt(TEXT("r.VelocityTest"));
+			const int32 VelocityTest = VelocityTestVar ? VelocityTestVar->GetValueOnRenderThread() : 0;
 
-	Line = FString::Printf(TEXT("View=%.4x PrevView=%.4x"),
-		View.ViewMatrices.GetViewMatrix().ComputeHash() & 0xffff,
-		View.PrevViewInfo.ViewMatrices.GetViewMatrix().ComputeHash() & 0xffff);
-	Canvas.DrawShadowedString(X, Y += YStep, TEXT("ViewMatrix:"), GetStatsFont(), FLinearColor(1, 1, 0));
-	Canvas.DrawShadowedString(X + ColumnWidth, Y, *Line, GetStatsFont(), FLinearColor(1, 1, 0));
+			extern bool IsParallelVelocity();
 
-	Canvas.Flush_RenderThread(Context.RHICmdList);
+			Line = FString::Printf(TEXT("%d, %d, %d"), ViewFamily.bWorldIsPaused, VelocityTest, IsParallelVelocity());
+			Canvas.DrawShadowedString(X, Y += YStep, TEXT("Paused, r.VelocityTest, Parallel:"), GetStatsFont(), FLinearColor(1, 1, 0));
+			Canvas.DrawShadowedString(X + ColumnWidth, Y, *Line, GetStatsFont(), FLinearColor(1, 1, 0));
 
-	Context.RHICmdList.CopyToResolveTarget(DestRenderTarget.TargetableTexture, DestRenderTarget.ShaderResourceTexture, FResolveParams());
-}
+			const FScene* Scene = (const FScene*)View.Family->Scene;
+			const FSceneViewState *SceneViewState = (const FSceneViewState*)View.State;
 
-FPooledRenderTargetDesc FRCPassPostProcessVisualizeMotionBlur::ComputeOutputDesc(EPassOutputId InPassOutputId) const
-{
-	FPooledRenderTargetDesc Ret = GetInput(ePId_Input0)->GetOutput()->RenderTargetDesc;
+			Line = FString::Printf(TEXT("View=%.4x PrevView=%.4x"),
+				View.ViewMatrices.GetViewMatrix().ComputeHash() & 0xffff,
+				View.PrevViewInfo.ViewMatrices.GetViewMatrix().ComputeHash() & 0xffff);
+			Canvas.DrawShadowedString(X, Y += YStep, TEXT("ViewMatrix:"), GetStatsFont(), FLinearColor(1, 1, 0));
+			Canvas.DrawShadowedString(X + ColumnWidth, Y, *Line, GetStatsFont(), FLinearColor(1, 1, 0));
 
-	Ret.Reset();
-	Ret.DebugName = TEXT("MotionBlur");
-	Ret.AutoWritable = false;
+			const bool bForce = false;
+			const bool bInsideRenderThread = true;
+			Canvas.Flush_RenderThread(RHICmdList, bForce, bInsideRenderThread);
+		}
+	});
 
-	return Ret;
-}
+	return TextureOutput;
+}
\ No newline at end of file
diff --git a/Engine/Source/Runtime/Renderer/Private/PostProcess/PostProcessMotionBlur.h b/Engine/Source/Runtime/Renderer/Private/PostProcess/PostProcessMotionBlur.h
--- a/Engine/Source/Runtime/Renderer/Private/PostProcess/PostProcessMotionBlur.h
+++ b/Engine/Source/Runtime/Renderer/Private/PostProcess/PostProcessMotionBlur.h
@@ -6,97 +6,43 @@
 
 #pragma once
 
-#include "CoreMinimal.h"
-#include "RendererInterface.h"
-#include "PostProcess/RenderingCompositionGraph.h"
-
-// Find max velocity per 16x16 tile
-// derives from TRenderingCompositePassBase<InputCount, OutputCount>
-class FRCPassPostProcessVelocityFlatten : public TRenderingCompositePassBase<2, 2>
-{
-public:
-	FRCPassPostProcessVelocityFlatten();
-
-	// interface FRenderingCompositePass ---------
-	virtual void Process(FRenderingCompositePassContext& Context) override;
-	virtual void Release() override { delete this; }
-	virtual FPooledRenderTargetDesc ComputeOutputDesc(EPassOutputId InPassOutputId) const override;
-
-	// -------------------------------------------
-
-	static FIntPoint ComputeThreadGroupCount(FIntPoint PixelExtent);
-};
+#include "ScreenPass.h"
 
+// Returns whether motion blur is enabled for the requested view.
+bool IsMotionBlurEnabled(const FViewInfo& View);
 
-// derives from TRenderingCompositePassBase<InputCount, OutputCount>
-class FRCPassPostProcessVelocityScatter : public TRenderingCompositePassBase<1, 1>
+// The quality level of the motion blur pass.
+enum class EMotionBlurQuality : uint32
 {
-public:
-	// interface FRenderingCompositePass ---------
-	virtual void Process(FRenderingCompositePassContext& Context) override;
-	virtual void Release() override { delete this; }
-	virtual FPooledRenderTargetDesc ComputeOutputDesc(EPassOutputId InPassOutputId) const override;
+	Low,
+	Medium,
+	High,
+	VeryHigh,
+	MAX
 };
 
-// derives from TRenderingCompositePassBase<InputCount, OutputCount>
-class FRCPassPostProcessVelocityGather : public TRenderingCompositePassBase<1, 1>
-{
-public:
-	// interface FRenderingCompositePass ---------
-	virtual void Process(FRenderingCompositePassContext& Context) override;
-	virtual void Release() override { delete this; }
-	virtual FPooledRenderTargetDesc ComputeOutputDesc(EPassOutputId InPassOutputId) const override;
-};
-
-
-// ePId_Input0: Full Res Scene Color
-// ePId_Input1: Full Res Scene Depth
-// ePId_Input2: Velocity
-// ePId_Input3: Max tile velocity
-// derives from TRenderingCompositePassBase<InputCount, OutputCount> 
-class FRCPassPostProcessMotionBlur : public TRenderingCompositePassBase<4, 1>
-{
-public:
-	// @param InQuality 0xffffffff to visualize, 0:off(no shader is used), 1:low, 2:medium, 3:high, 4:very high
-	FRCPassPostProcessMotionBlur( uint32 InQuality, int32 InPass, bool InIsComputePass )
-		: Quality(InQuality)
-		, Pass(InPass)
-	{
-		bIsComputePass = InIsComputePass;
-		bPreferAsyncCompute = false;
-
-		// internal error
-		check(Quality >= 1 && Quality <= 4);
-	}
-
-	// interface FRenderingCompositePass ---------
-	virtual void Process(FRenderingCompositePassContext& Context) override;
-	virtual void Release() override { delete this; }
-	virtual FPooledRenderTargetDesc ComputeOutputDesc(EPassOutputId InPassOutputId) const override;
-
-	virtual FComputeFenceRHIParamRef GetComputePassEndFence() const override { return AsyncEndFence; }
-
-	// 1:low, 2:medium, 3:high, 4: very high
-	uint32	Quality;
-	int32	Pass;
-
-private:
-	template <typename TRHICmdList>
-	void DispatchCS(TRHICmdList& RHICmdList, FRenderingCompositePassContext& Context, const FIntPoint& DestSize, const FIntRect& DestRect, FUnorderedAccessViewRHIParamRef DestUAV, float Scale, uint32 InnerLoopConfig);
-	
-	FComputeFenceRHIRef AsyncEndFence;
-};
-
-
-// ePId_Input0: Full Res Scene Color
-// ePId_Input1: Full Res Scene Depth
-// ePId_Input2: Full Res velocity input
-// derives from TRenderingCompositePassBase<InputCount, OutputCount> 
-class FRCPassPostProcessVisualizeMotionBlur : public TRenderingCompositePassBase<3, 1>
-{
-public:
-	// interface FRenderingCompositePass ---------
-	virtual void Process(FRenderingCompositePassContext& Context) override;
-	virtual void Release() override { delete this; }
-	virtual FPooledRenderTargetDesc ComputeOutputDesc(EPassOutputId InPassOutputId) const override;
-};
+// Returns the motion blur quality level set by the 'r.MotionBlurQuality' cvar.
+EMotionBlurQuality GetMotionBlurQuality();
+
+/**
+ * Computes motion blur using the provided color, depth, and velocity textures. Returns
+ * the blurred color result. The Color viewport does not need to match Depth / Velocity
+ * viewports, but Depth / Velocity must match each other.
+ */
+FScreenPassTexture ComputeMotionBlur(
+	FRDGBuilder& GraphBuilder,
+	FScreenPassContextRef Context,
+	const FScreenPassTexture& ColorTexture,
+	const FScreenPassTexture& DepthTexture,
+	const FScreenPassTexture& VelocityTexture);
+
+/**
+ * Visualizes motion blur velocities and outputs the result. The Color viewport does not
+ * need to match Depth / Velocity viewports, but Depth / Velocity must match each other.
+ */
+FScreenPassTexture VisualizeMotionBlur(
+	FRDGBuilder& GraphBuilder,
+	FScreenPassContextRef Context,
+	const FScreenPassTexture& ColorTexture,
+	const FScreenPassTexture& DepthTexture,
+	const FScreenPassTexture& VelocityTexture);
\ No newline at end of file
diff --git a/Engine/Source/Runtime/Renderer/Private/PostProcess/PostProcessing.cpp b/Engine/Source/Runtime/Renderer/Private/PostProcess/PostProcessing.cpp
--- a/Engine/Source/Runtime/Renderer/Private/PostProcess/PostProcessing.cpp
+++ b/Engine/Source/Runtime/Renderer/Private/PostProcess/PostProcessing.cpp
@@ -111,16 +111,6 @@ static TAutoConsoleVariable<int32> CDownsampleQuality(
 	TEXT(">0: high quality (default: 3)\n"),
 	ECVF_Scalability | ECVF_RenderThreadSafe);
 
-static TAutoConsoleVariable<float> CVarMotionBlurSoftEdgeSize(
-	TEXT("r.MotionBlurSoftEdgeSize"),
-	1.0f,
-	TEXT("Defines how wide the object motion blur is blurred (percent of screen width) to allow soft edge motion blur.\n")
-	TEXT("This scales linearly with the size (up to a maximum of 32 samples, 2.5 is about 18 samples) and with screen resolution\n")
-	TEXT("Smaller values are better for performance and provide more accurate motion vectors but the blurring outside the object is reduced.\n")
-	TEXT("If needed this can be exposed like the other motionblur settings.\n")
-	TEXT(" 0:off (not free and does never completely disable), >0, 1.0 (default)"),
-	ECVF_RenderThreadSafe);
-
 static TAutoConsoleVariable<float> CVarBloomCross(
 	TEXT("r.Bloom.Cross"),
 	0.0f,
@@ -151,26 +141,6 @@ static TAutoConsoleVariable<float> CVarTonemapperMergeThreshold(
 	TEXT("Defauls to 0.49 (e.g., if r.ScreenPercentage is 70 or higher, try to merge)"),
 	ECVF_Scalability | ECVF_RenderThreadSafe);
 
-static TAutoConsoleVariable<int32> CVarMotionBlurScatter(
-	TEXT("r.MotionBlurScatter"),
-	0,
-	TEXT("Forces scatter based max velocity method (slower)."),
-	ECVF_RenderThreadSafe
-);
-
-static TAutoConsoleVariable<int32> CVarMotionBlurSeparable(
-	TEXT("r.MotionBlurSeparable"),
-	0,
-	TEXT("Adds a second motion blur pass that smooths noise for a higher quality blur."),
-	ECVF_RenderThreadSafe
-);
-
-static TAutoConsoleVariable<int32> CVarMotionBlurPreferCompute(
-	TEXT("r.MotionBlur.PreferCompute"),
-	0,
-	TEXT("Will use compute shaders for motion blur pass."),
-	ECVF_RenderThreadSafe);
-
 static TAutoConsoleVariable<int32> CVarAlphaChannel(
 	TEXT("r.PostProcessing.PropagateAlpha"),
 	0,
@@ -1781,76 +1751,59 @@ void FPostProcessing::Process(FRHICommandListImmediate& RHICmdList, const FViewI
 				SSRInputChain = AddPostProcessMaterialChain(Context, BL_SSRInput);
 			}
 
-			if(IsMotionBlurEnabled(View) && VelocityInput.IsValid() && !bVisualizeMotionBlur)
+			// Motion Blur
+			if ((IsMotionBlurEnabled(View) || bVisualizeMotionBlur) && VelocityInput.IsValid())
 			{
-				// Motion blur
-
-				FRenderingCompositeOutputRef MaxTileVelocity;
-
+				FRenderingCompositePass* MotionBlurPass = Context.Graph.RegisterPass(new(FMemStack::Get()) TRCPassForRDG<3, 1>(
+					[bVisualizeMotionBlur](FRenderingCompositePass* Pass, FRenderingCompositePassContext& Context)
 				{
-					check(!VelocityFlattenPass);
-					VelocityFlattenPass = Context.Graph.RegisterPass( new(FMemStack::Get()) FRCPassPostProcessVelocityFlatten() );
-					VelocityFlattenPass->SetInput( ePId_Input0, VelocityInput );
-					VelocityFlattenPass->SetInput( ePId_Input1, Context.SceneDepth );
+					FRDGBuilder GraphBuilder(Context.RHICmdList);
 
-					VelocityInput	= FRenderingCompositeOutputRef( VelocityFlattenPass, ePId_Output0 );
-					MaxTileVelocity	= FRenderingCompositeOutputRef( VelocityFlattenPass, ePId_Output1 );
-				}
+					const FIntRect ColorViewportRect = Context.SceneColorViewRect;
+					const FIntRect DepthViewportRect = Context.View.ViewRect;
 
-				const float SizeX = View.ViewRect.Width();
+					const FScreenPassTexture ColorTexture = FScreenPassTexture::Create(
+						Pass->CreateRDGTextureForInput(GraphBuilder, ePId_Input0, TEXT("SceneColor"), eFC_0000),
+						ColorViewportRect);
 
-				// 0:no 1:full screen width, percent conversion
-				float MaxVelocity = View.FinalPostProcessSettings.MotionBlurMax / 100.0f;
-				float MaxVelocityTiles = MaxVelocity * SizeX * (0.5f / 16.0f);
-				float MaxTileDistGathered = 3.0f;
-				if( MaxVelocityTiles > MaxTileDistGathered || CVarMotionBlurScatter.GetValueOnRenderThread() || (ViewState && ViewState->bSequencerIsPaused) )
-				{
-					FRenderingCompositePass* VelocityScatterPass = Context.Graph.RegisterPass( new(FMemStack::Get()) FRCPassPostProcessVelocityScatter() );
-					VelocityScatterPass->SetInput( ePId_Input0, MaxTileVelocity );
+					const FScreenPassTexture DepthTexture = FScreenPassTexture::Create(
+						Pass->CreateRDGTextureForInput(GraphBuilder, ePId_Input1, TEXT("SceneDepth"), eFC_0000),
+						DepthViewportRect);
 
-					MaxTileVelocity	= FRenderingCompositeOutputRef( VelocityScatterPass );
-				}
-				else
-				{
-					FRenderingCompositePass* VelocityGatherPass = Context.Graph.RegisterPass( new(FMemStack::Get()) FRCPassPostProcessVelocityGather() );
-					VelocityGatherPass->SetInput( ePId_Input0, MaxTileVelocity );
+					const FScreenPassTexture VelocityTexture = FScreenPassTexture::Create(
+						Pass->CreateRDGTextureForInput(GraphBuilder, ePId_Input2, TEXT("SceneVelocity"), eFC_0000),
+						DepthViewportRect);
 
-					MaxTileVelocity	= FRenderingCompositeOutputRef( VelocityGatherPass );
-				}
+					FScreenPassTexture OutColorTexture;
 
-				bool bTwoPass = CVarMotionBlurSeparable.GetValueOnRenderThread() != 0;
-				{
-					const bool bIsComputePass = ShouldDoComputePostProcessing(Context.View);
-					FRenderingCompositePass* MotionBlurPass = Context.Graph.RegisterPass( new(FMemStack::Get()) FRCPassPostProcessMotionBlur( GetMotionBlurQualityFromCVar(), bTwoPass ? 0 : -1, bIsComputePass ) );
-					MotionBlurPass->SetInput( ePId_Input0, Context.FinalOutput );
-					MotionBlurPass->SetInput( ePId_Input1, Context.SceneDepth );
-					MotionBlurPass->SetInput( ePId_Input2, VelocityInput );
-					MotionBlurPass->SetInput( ePId_Input3, MaxTileVelocity );
-					
-					Context.FinalOutput = FRenderingCompositeOutputRef( MotionBlurPass );
-				}
+					if (bVisualizeMotionBlur)
+					{
+						OutColorTexture = VisualizeMotionBlur(
+							GraphBuilder,
+							FScreenPassContext::Create(Context.RHICmdList, Context.View),
+							ColorTexture,
+							DepthTexture,
+							VelocityTexture);
+					}
+					else
+					{
+						OutColorTexture = ComputeMotionBlur(
+							GraphBuilder,
+							FScreenPassContext::Create(Context.RHICmdList, Context.View),
+							ColorTexture,
+							DepthTexture,
+							VelocityTexture);
+					}
 
-				if( bTwoPass )
-				{
-					const bool bIsComputePass = ShouldDoComputePostProcessing(Context.View);
-					FRenderingCompositePass* MotionBlurPass = Context.Graph.RegisterPass( new(FMemStack::Get()) FRCPassPostProcessMotionBlur( GetMotionBlurQualityFromCVar(), 1, bIsComputePass ) );
-					MotionBlurPass->SetInput( ePId_Input0, Context.FinalOutput );
-					MotionBlurPass->SetInput( ePId_Input1, Context.SceneDepth );
-					MotionBlurPass->SetInput( ePId_Input2, VelocityInput );
-					MotionBlurPass->SetInput( ePId_Input3, MaxTileVelocity );
-					
-					Context.FinalOutput = FRenderingCompositeOutputRef( MotionBlurPass );
-				}
-			}
+					Pass->ExtractRDGTextureForOutput(GraphBuilder, ePId_Output0, OutColorTexture.GetRDGTexture());
 
-			if(VelocityInput.IsValid() && bVisualizeMotionBlur)
-			{
-				auto VisualizePass = Context.Graph.RegisterPass(new(FMemStack::Get()) FRCPassPostProcessVisualizeMotionBlur());
-				VisualizePass->SetInput(ePId_Input0, Context.FinalOutput);
-				VisualizePass->SetInput(ePId_Input1, Context.SceneDepth);
-				VisualizePass->SetInput(ePId_Input2, VelocityInput);
+					GraphBuilder.Execute();
+				}));
 
-				Context.FinalOutput = FRenderingCompositeOutputRef(VisualizePass);
+				MotionBlurPass->SetInput(ePId_Input0, Context.FinalOutput);
+				MotionBlurPass->SetInput(ePId_Input1, Context.SceneDepth);
+				MotionBlurPass->SetInput(ePId_Input2, VelocityInput);
+				Context.FinalOutput = FRenderingCompositeOutputRef(MotionBlurPass, ePId_Output0);
 			}
 
 			if(bVisualizeBloom)
diff --git a/Engine/Source/Runtime/Renderer/Private/VelocityRendering.cpp b/Engine/Source/Runtime/Renderer/Private/VelocityRendering.cpp
--- a/Engine/Source/Runtime/Renderer/Private/VelocityRendering.cpp
+++ b/Engine/Source/Runtime/Renderer/Private/VelocityRendering.cpp
@@ -16,6 +16,7 @@
 #include "DeferredShadingRenderer.h"
 #include "ScenePrivate.h"
 #include "PostProcess/ScreenSpaceReflections.h"
+#include "PostProcess/PostProcessMotionBlur.h"
 #include "UnrealEngine.h"
 #if WITH_EDITOR
 #include "Misc/CoreMisc.h"
@@ -46,11 +47,6 @@ static TAutoConsoleVariable<int32> CVarRHICmdVelocityPassDeferredContexts(
 	1,
 	TEXT("True to use deferred contexts to parallelize velocity pass command list execution."));
 
-RENDERER_API TAutoConsoleVariable<int32> CVarAllowMotionBlurInVR(
-	TEXT("vr.AllowMotionBlurInVR"),
-	0,
-	TEXT("For projects with motion blur enabled, this allows motion blur to be enabled even while in VR."));
-
 DECLARE_GPU_STAT_NAMED(RenderVelocities, TEXT("Render Velocities"));
 
 bool IsParallelVelocity()
@@ -204,36 +200,6 @@ IMPLEMENT_MATERIAL_SHADER_TYPE(,FVelocityPS,TEXT("/Engine/Private/VelocityShader
 
 IMPLEMENT_SHADERPIPELINE_TYPE_VSPS(VelocityPipeline, FVelocityVS, FVelocityPS, true);
 
-
-int32 GetMotionBlurQualityFromCVar()
-{
-	int32 MotionBlurQuality;
-
-	static const auto ICVar = IConsoleManager::Get().FindTConsoleVariableDataInt(TEXT("r.MotionBlurQuality"));
-	MotionBlurQuality = FMath::Clamp(ICVar->GetValueOnRenderThread(), 0, 4);
-
-	return MotionBlurQuality;
-}
-
-bool IsMotionBlurEnabled(const FViewInfo& View)
-{
-	if (View.GetFeatureLevel() < ERHIFeatureLevel::SM5)
-	{
-		return false; 
-	}
-
-	int32 MotionBlurQuality = GetMotionBlurQualityFromCVar();
-
-	return View.Family->EngineShowFlags.PostProcessing
-		&& View.Family->EngineShowFlags.MotionBlur
-		&& View.FinalPostProcessSettings.MotionBlurAmount > 0.001f
-		&& View.FinalPostProcessSettings.MotionBlurMax > 0.001f
-		&& View.Family->bRealtimeUpdate
-		&& MotionBlurQuality > 0
-		&& !IsSimpleForwardShadingEnabled(GShaderPlatformForFeatureLevel[View.GetFeatureLevel()])
-		&& (CVarAllowMotionBlurInVR->GetInt() != 0 || !(View.Family->Views.Num() > 1));
-}
-
 static void BeginVelocityRendering(FRHICommandList& RHICmdList, TRefCountPtr<IPooledRenderTarget>& VelocityRT, bool bPerformClear)
 {
 	check(RHICmdList.IsOutsideRenderPass());
diff --git a/Engine/Source/Runtime/Renderer/Private/VelocityRendering.h b/Engine/Source/Runtime/Renderer/Private/VelocityRendering.h
--- a/Engine/Source/Runtime/Renderer/Private/VelocityRendering.h
+++ b/Engine/Source/Runtime/Renderer/Private/VelocityRendering.h
@@ -18,12 +18,6 @@ class FScene;
 class FStaticMeshBatch;
 class FViewInfo;
 
-/** Get the cvar clamped state */
-int32 GetMotionBlurQualityFromCVar();
-
-/** If this view need motion blur processing */
-bool IsMotionBlurEnabled(const FViewInfo& View);
-
 // Group Velocity Rendering accessors, types, etc.
 struct FVelocityRendering
 {
