Deletes old diaphragm DOF code using post proces graph.

diff --git a/Engine/Source/Runtime/Renderer/Private/PostProcess/DiaphragmDOF.cpp b/Engine/Source/Runtime/Renderer/Private/PostProcess/DiaphragmDOF.cpp
--- a/Engine/Source/Runtime/Renderer/Private/PostProcess/DiaphragmDOF.cpp
+++ b/Engine/Source/Runtime/Renderer/Private/PostProcess/DiaphragmDOF.cpp
@@ -5,7 +5,6 @@
 =============================================================================*/
 
 #include "PostProcess/DiaphragmDOF.h"
-#include "PostProcess/DiaphragmDOFPasses.h"
 #include "StaticBoundShaderState.h"
 #include "SceneUtils.h"
 #include "PostProcess/SceneRenderTargets.h"
@@ -31,6 +30,9 @@ void AddPass_ClearUAV(
 
 // ---------------------------------------------------- Cvars
 
+namespace
+{
+
 TAutoConsoleVariable<int32> CVarAccumulatorQuality(
 	TEXT("r.DOF.Gather.AccumulatorQuality"),
 	1,
@@ -128,9 +130,6 @@ TAutoConsoleVariable<int32> CVarDOFTemporalAAQuality(
 	TEXT(" 1: Higher quality pass (default)."),
 	ECVF_Scalability | ECVF_RenderThreadSafe);
 
-namespace // TODO encapsulate all cvars
-{
-
 TAutoConsoleVariable<float> CVarScatterNeighborCompareMaxColor(
 	TEXT("r.DOF2.Scatter.NeighborCompareMaxColor"),
 	10,
@@ -141,39 +140,144 @@ TAutoConsoleVariable<float> CVarScatterNeighborCompareMaxColor(
 
 } // namespace
 
+
+// ---------------------------------------------------- COMMON
+
 namespace
 {
+	
+/** Defines which layer to process. */
+enum class EDiaphragmDOFLayerProcessing
+{
+	// Foreground layer only.
+	ForegroundOnly,
 
-EDiaphragmDOFPostfilterMethod GetPostfilteringMethod()
+	// Foreground hole filling.
+	ForegroundHoleFilling,
+
+	// Background layer only.
+	BackgroundOnly,
+
+	// Both foreground and background layers.
+	ForegroundAndBackground,
+
+	// Slight out of focus layer.
+	SlightOutOfFocus,
+
+	MAX
+};
+
+
+
+/** Defines which layer to process. */
+enum class EDiaphragmDOFPostfilterMethod
 {
-	int32 i = CVarPostFilteringMethod.GetValueOnRenderThread();
-	if (i >= 0 && i < int32(EDiaphragmDOFPostfilterMethod::MAX))
-	{
-		return EDiaphragmDOFPostfilterMethod(i);
-	}
-	return EDiaphragmDOFPostfilterMethod::None;
-}
+	// Disable post filtering.
+	None,
 
+	// Per RGB channel median on 3x3 neighborhood.
+	RGBMedian3x3,
 
-enum class EHybridScatterMode
+	// Per RGB channel max on 3x3 neighborhood.
+	RGBMax3x3,
+
+	MAX
+};
+
+
+/** Modes to simulate a bokeh. */
+enum class EDiaphragmDOFBokehSimulation
 {
+	// No bokeh simulation.
 	Disabled,
-	Additive,
-	Occlusion,
+
+	// Symmetric bokeh (even number of blade).
+	SimmetricBokeh,
+
+	// Generic bokeh.
+	GenericBokeh,
+
+	MAX,
+};
+
+/** Dilate mode of the pass. */
+enum class EDiaphragmDOFDilateCocMode
+{
+	// One single dilate pass.
+	StandAlone,
+
+	// Dilate min foreground and max background coc radius.
+	MinForegroundAndMaxBackground,
+
+	// Dilate everything else from dilated min foreground and max background coc radius.
+	MinimalAbsoluteRadiuses,
+
+	MAX
 };
 
+/** Quality configurations for gathering passes. */
+enum class EDiaphragmDOFGatherQuality
+{
+	// Lower but faster accumulator.
+	LowQualityAccumulator,
+
+	// High quality accumulator.
+	HighQuality,
 
-}
+	// High quality accumulator with hybrid scatter occlusion buffer output.
+	// TODO: distinct shader permutation dimension for hybrid scatter occlusion?
+	HighQualityWithHybridScatterOcclusion,
 
+	// High quality accumulator, with layered full disks and hybrid scatter occlusion.
+	Cinematic,
 
-// ---------------------------------------------------- COMMON
+	MAX,
+};
 
-namespace
+/** Format of the LUT to generate. */
+enum class EDiaphragmDOFBokehLUTFormat
 {
+	// Look up table that stores a factor to transform a CocRadius to a BokehEdge distance.
+	// Used for scattering and low res focus gathering.
+	CocRadiusToBokehEdgeFactor,
+
+	// Look up table that stores Coc distance to compare against neighbor's CocRadius.
+	// Used exclusively for full res gathering in recombine pass.
+	FullResOffsetToCocDistance,
+
+
+	// Look up table to stores the gathering sample pos within the kernel.
+	// Used for low res back and foreground gathering.
+	GatherSamplePos,
+
+	MAX,
+};
+
 
 const int32 kDefaultGroupSize = 8;
+
+/** Number of half res pixel are covered by a CocTile */
 const int32 kCocTileSize = kDefaultGroupSize;
 
+/** Resolution divisor of the Coc tiles. */
+const int32 kMaxCocDilateSampleRadiusCount = 3;
+
+/** Resolution divisor of the Coc tiles. */
+const int32 kMaxMipLevelCount = 4;
+
+/** Minimum number of ring. */
+const int32 kMinGatheringRingCount = 3;
+
+/** Maximum number of ring for slight out of focus pass. Same as USH's MAX_RECOMBINE_ABS_COC_RADIUS. */
+const int32 kMaxSlightOutOfFocusRingCount = 3;
+
+/** Maximum quality level of the recombine pass. */
+const int32 kMaxRecombineQuality = 2;
+
+/** Absolute minim coc radius required for a bokeh to be scattered */
+const float kMinScatteringCocRadius = 3.0f;
+
+
 FIntPoint CocTileGridSize(FIntPoint FullResSize)
 {
 	uint32 TilesX = FMath::DivideAndRoundUp(FullResSize.X, kCocTileSize);
@@ -199,6 +303,23 @@ FIntRect GetLowerResViewport(const FIntRect& ViewRect, int32 ResDivisor)
 }
 
 
+EDiaphragmDOFPostfilterMethod GetPostfilteringMethod()
+{
+	int32 i = CVarPostFilteringMethod.GetValueOnRenderThread();
+	if (i >= 0 && i < int32(EDiaphragmDOFPostfilterMethod::MAX))
+	{
+		return EDiaphragmDOFPostfilterMethod(i);
+	}
+	return EDiaphragmDOFPostfilterMethod::None;
+}
+
+enum class EHybridScatterMode
+{
+	Disabled,
+	Additive,
+	Occlusion,
+};
+
 const TCHAR* GetEventName(EDiaphragmDOFLayerProcessing e)
 {
 	static const TCHAR* const kArray[] = {
@@ -236,7 +357,7 @@ const TCHAR* GetEventName(EDiaphragmDOFBokehSimulation e)
 	return kArray[i];
 }
 
-const TCHAR* GetEventName(FRCPassDiaphragmDOFBuildBokehLUT::EFormat e)
+const TCHAR* GetEventName(EDiaphragmDOFBokehLUTFormat e)
 {
 	static const TCHAR* const kArray[] = {
 		TEXT("Scatter"),
@@ -248,7 +369,7 @@ const TCHAR* GetEventName(FRCPassDiaphragmDOFBuildBokehLUT::EFormat e)
 	return kArray[i];
 }
 
-const TCHAR* GetEventName(FRCPassDiaphragmDOFGather::EQualityConfig e)
+const TCHAR* GetEventName(EDiaphragmDOFGatherQuality e)
 {
 	static const TCHAR* const kArray[] = {
 		TEXT("LowQ"),
@@ -261,7 +382,7 @@ const TCHAR* GetEventName(FRCPassDiaphragmDOFGather::EQualityConfig e)
 	return kArray[i];
 }
 
-const TCHAR* GetEventName(FRCPassDiaphragmDOFDilateCoc::EMode e)
+const TCHAR* GetEventName(EDiaphragmDOFDilateCocMode e)
 {
 	static const TCHAR* const kArray[] = {
 		TEXT("StandAlone"),
@@ -350,11 +471,11 @@ namespace
 {
 
 class FDDOFDilateRadiusDim     : SHADER_PERMUTATION_RANGE_INT("DIM_DILATE_RADIUS", 1, 3);
-class FDDOFDilateModeDim       : SHADER_PERMUTATION_ENUM_CLASS("DIM_DILATE_MODE", FRCPassDiaphragmDOFDilateCoc::EMode);
+class FDDOFDilateModeDim       : SHADER_PERMUTATION_ENUM_CLASS("DIM_DILATE_MODE", EDiaphragmDOFDilateCocMode);
 
 class FDDOFLayerProcessingDim  : SHADER_PERMUTATION_ENUM_CLASS("DIM_LAYER_PROCESSING", EDiaphragmDOFLayerProcessing);
-class FDDOFGatherRingCountDim  : SHADER_PERMUTATION_RANGE_INT("DIM_GATHER_RING_COUNT", FRCPassDiaphragmDOFGather::kMinRingCount, 3);
-class FDDOFGatherQualityDim    : SHADER_PERMUTATION_ENUM_CLASS("DIM_GATHER_QUALITY", FRCPassDiaphragmDOFGather::EQualityConfig);
+class FDDOFGatherRingCountDim  : SHADER_PERMUTATION_RANGE_INT("DIM_GATHER_RING_COUNT", kMinGatheringRingCount, 3);
+class FDDOFGatherQualityDim    : SHADER_PERMUTATION_ENUM_CLASS("DIM_GATHER_QUALITY", EDiaphragmDOFGatherQuality);
 class FDDOFPostfilterMethodDim : SHADER_PERMUTATION_ENUM_CLASS("DIM_POSTFILTER_METHOD", EDiaphragmDOFPostfilterMethod);
 class FDDOFClampInputUVDim     : SHADER_PERMUTATION_BOOL("DIM_CLAMP_INPUT_UV");
 class FDDOFRGBColorBufferDim   : SHADER_PERMUTATION_BOOL("DIM_RGB_COLOR_BUFFER");
@@ -503,6 +624,44 @@ void SetCocModelParameters(
 namespace
 {
 
+/** Returns whether hybrid scattering is supported. */
+static FORCEINLINE bool SupportsHybridScatter(EShaderPlatform ShaderPlatform)
+{
+	return !IsSwitchPlatform(ShaderPlatform);
+}
+
+
+/** Returns the number maximum number of ring available. */
+static FORCEINLINE int32 MaxGatheringRingCount(EShaderPlatform ShaderPlatform)
+{
+	if (IsPCPlatform(ShaderPlatform))
+	{
+		return 5;
+	}
+	return 4;
+}
+
+/** Returns whether the shader for bokeh simulation are compiled. */
+static FORCEINLINE bool SupportsBokehSimmulation(EShaderPlatform ShaderPlatform)
+{
+	// Shaders of gathering pass are big, so only compile them on desktop.
+	return IsPCPlatform(ShaderPlatform);
+}
+
+/** Returns whether separate coc buffer is supported. */
+static FORCEINLINE bool SupportsRGBColorBuffer(EShaderPlatform ShaderPlatform)
+{
+	// There is no point when alpha channel is supported because needs 4 channel anyway for fast gathering tiles.
+	if (FPostProcessing::HasAlphaChannelSupport())
+	{
+		return false;
+	}
+
+	// There is high number of UAV to write in reduce pass.
+	return ShaderPlatform == SP_PS4 || ShaderPlatform == SP_XBOXONE_D3D12 || ShaderPlatform == SP_VULKAN_SM5;
+}
+
+
 class FDiaphragmDOFSetupCS : public FDiaphragmDOFShader
 {
 	DECLARE_GLOBAL_SHADER(FDiaphragmDOFSetupCS);
@@ -581,7 +740,7 @@ class FDiaphragmDOFDownsampleCS : public FDiaphragmDOFShader
 
 	static bool ShouldCompilePermutation(const FGlobalShaderPermutationParameters& Parameters)
 	{
-		if (!FRCPassDiaphragmDOFHybridScatter::IsSupported(Parameters.Platform))
+		if (!SupportsHybridScatter(Parameters.Platform))
 		{
 			return false;
 		}
@@ -621,12 +780,12 @@ class FDiaphragmDOFReduceCS : public FDiaphragmDOFShader
 		FPermutationDomain PermutationVector(Parameters.PermutationId);
 
 		// Do not compile storing Coc independently of RGB if not supported.
-		if (PermutationVector.Get<FDDOFRGBColorBufferDim>() && !FRCPassDiaphragmDOFGather::SupportRGBColorBuffer(Parameters.Platform))
+		if (PermutationVector.Get<FDDOFRGBColorBufferDim>() && !SupportsRGBColorBuffer(Parameters.Platform))
 		{
 			return false;
 		}
 
-		if (!FRCPassDiaphragmDOFHybridScatter::IsSupported(Parameters.Platform))
+		if (!SupportsHybridScatter(Parameters.Platform))
 		{
 			if (PermutationVector.Get<FHybridScatterForeground>() || PermutationVector.Get<FHybridScatterBackground>())
 			{
@@ -654,7 +813,7 @@ class FDiaphragmDOFReduceCS : public FDiaphragmDOFShader
 		SHADER_PARAMETER(FVector4, QuarterResGatherInputSize)
 		SHADER_PARAMETER_STRUCT(FDOFGatherInputTextures, QuarterResGatherInput)
 
-		SHADER_PARAMETER_STRUCT_ARRAY(FDOFGatherInputUAVs, OutputMips, [FRCPassDiaphragmDOFReduce::kMaxMipLevelCount])
+		SHADER_PARAMETER_STRUCT_ARRAY(FDOFGatherInputUAVs, OutputMips, [kMaxMipLevelCount])
 		SHADER_PARAMETER_RDG_BUFFER_UAV(RWBuffer<uint>, OutScatterDrawIndirectParameters)
 		SHADER_PARAMETER_RDG_BUFFER_UAV(RWStructuredBuffer<float4>, OutForegroundScatterDrawList)
 		SHADER_PARAMETER_RDG_BUFFER_UAV(RWStructuredBuffer<float4>, OutBackgroundScatterDrawList)
@@ -668,7 +827,7 @@ class FDiaphragmDOFScatterGroupPackCS : public FDiaphragmDOFShader
 
 	static bool ShouldCompilePermutation(const FGlobalShaderPermutationParameters& Parameters)
 	{
-		if (!FRCPassDiaphragmDOFHybridScatter::IsSupported(Parameters.Platform))
+		if (!SupportsHybridScatter(Parameters.Platform))
 		{
 			return false;
 		}
@@ -690,7 +849,7 @@ class FDiaphragmDOFBuildBokehLUTCS : public FDiaphragmDOFShader
 	SHADER_USE_PARAMETER_STRUCT(FDiaphragmDOFBuildBokehLUTCS, FDiaphragmDOFShader);
 
 	class FBokehSimulationDim : SHADER_PERMUTATION_BOOL("DIM_ROUND_BLADES");
-	class FLUTFormatDim : SHADER_PERMUTATION_ENUM_CLASS("DIM_LUT_FORMAT", FRCPassDiaphragmDOFBuildBokehLUT::EFormat);
+	class FLUTFormatDim : SHADER_PERMUTATION_ENUM_CLASS("DIM_LUT_FORMAT", EDiaphragmDOFBokehLUTFormat);
 
 	using FPermutationDomain = TShaderPermutationDomain<FBokehSimulationDim, FLUTFormatDim>;
 	
@@ -736,7 +895,7 @@ class FDiaphragmDOFGatherCS : public FDiaphragmDOFShader
 
 		// Slight out of focus only need 3 rings.
 		if (PermutationVector.Get<FDDOFLayerProcessingDim>() == EDiaphragmDOFLayerProcessing::SlightOutOfFocus)
-			PermutationVector.Set<FDDOFGatherRingCountDim>(FRCPassDiaphragmDOFGather::kMaxSlightOutOfFocusRingCount);
+			PermutationVector.Set<FDDOFGatherRingCountDim>(kMaxSlightOutOfFocusRingCount);
 
 		return PermutationVector;
 	}
@@ -752,45 +911,45 @@ class FDiaphragmDOFGatherCS : public FDiaphragmDOFShader
 		}
 
 		// Some platforms might be to slow for even considering large number of gathering samples.
-		if (PermutationVector.Get<FDDOFGatherRingCountDim>() > FRCPassDiaphragmDOFGather::MaxRingCount(Parameters.Platform))
+		if (PermutationVector.Get<FDDOFGatherRingCountDim>() > MaxGatheringRingCount(Parameters.Platform))
 		{
 			return false;
 		}
 
 		// Do not compile storing Coc independently of RGB.
-		if (PermutationVector.Get<FDDOFRGBColorBufferDim>() && !FRCPassDiaphragmDOFGather::SupportRGBColorBuffer(Parameters.Platform))
+		if (PermutationVector.Get<FDDOFRGBColorBufferDim>() && !SupportsRGBColorBuffer(Parameters.Platform))
 		{
 			return false;
 		}
 
 		// No point compiling gather pass with hybrid scatter occlusion if the shader platform doesn't support if.
-		if (!FRCPassDiaphragmDOFHybridScatter::IsSupported(Parameters.Platform) &&
-			PermutationVector.Get<FDDOFGatherQualityDim>() == FRCPassDiaphragmDOFGather::EQualityConfig::HighQualityWithHybridScatterOcclusion) return false;
+		if (!SupportsHybridScatter(Parameters.Platform) &&
+			PermutationVector.Get<FDDOFGatherQualityDim>() == EDiaphragmDOFGatherQuality::HighQualityWithHybridScatterOcclusion) return false;
 
 		// Do not compile bokeh simulation shaders on platform that couldn't handle them anyway.
-		if (!FRCPassDiaphragmDOFGather::SupportsBokehSimmulation(Parameters.Platform) &&
+		if (!SupportsBokehSimmulation(Parameters.Platform) &&
 			PermutationVector.Get<FDDOFBokehSimulationDim>() != EDiaphragmDOFBokehSimulation::Disabled) return false;
 
 		if (PermutationVector.Get<FDDOFLayerProcessingDim>() == EDiaphragmDOFLayerProcessing::ForegroundOnly)
 		{
 			// Foreground does not support CocVariance output yet.
-			if (PermutationVector.Get<FDDOFGatherQualityDim>() == FRCPassDiaphragmDOFGather::EQualityConfig::HighQualityWithHybridScatterOcclusion) return false;
+			if (PermutationVector.Get<FDDOFGatherQualityDim>() == EDiaphragmDOFGatherQuality::HighQualityWithHybridScatterOcclusion) return false;
 
 			// Storing Coc independently of RGB is only supported for low gathering quality.
 			if (PermutationVector.Get<FDDOFRGBColorBufferDim>() &&
-				PermutationVector.Get<FDDOFGatherQualityDim>() != FRCPassDiaphragmDOFGather::EQualityConfig::LowQualityAccumulator)
+				PermutationVector.Get<FDDOFGatherQualityDim>() != EDiaphragmDOFGatherQuality::LowQualityAccumulator)
 				return false;
 		}
 		else if (PermutationVector.Get<FDDOFLayerProcessingDim>() == EDiaphragmDOFLayerProcessing::ForegroundHoleFilling)
 		{
 			// Foreground hole filling does not need to output CocVariance, since this is the job of foreground pass.
-			if (PermutationVector.Get<FDDOFGatherQualityDim>() == FRCPassDiaphragmDOFGather::EQualityConfig::HighQualityWithHybridScatterOcclusion) return false;
+			if (PermutationVector.Get<FDDOFGatherQualityDim>() == EDiaphragmDOFGatherQuality::HighQualityWithHybridScatterOcclusion) return false;
 
 			// Foreground hole filling doesn't have lower quality accumulator.
-			if (PermutationVector.Get<FDDOFGatherQualityDim>() == FRCPassDiaphragmDOFGather::EQualityConfig::LowQualityAccumulator) return false;
+			if (PermutationVector.Get<FDDOFGatherQualityDim>() == EDiaphragmDOFGatherQuality::LowQualityAccumulator) return false;
 
 			// Foreground hole filling doesn't need cinematic quality.
-			if (PermutationVector.Get<FDDOFGatherQualityDim>() == FRCPassDiaphragmDOFGather::EQualityConfig::Cinematic) return false;
+			if (PermutationVector.Get<FDDOFGatherQualityDim>() == EDiaphragmDOFGatherQuality::Cinematic) return false;
 
 			// No bokeh simulation on hole filling, always use euclidian closest distance to compute opacity alpha channel.
 			if (PermutationVector.Get<FDDOFBokehSimulationDim>() != EDiaphragmDOFBokehSimulation::Disabled) return false;
@@ -803,16 +962,16 @@ class FDiaphragmDOFGatherCS : public FDiaphragmDOFShader
 		{
 			// Slight out of focus gather pass does not need large radius since only accumulating only
 			// abs(CocRadius) < kMaxSlightOutOfFocusRingCount.
-			if (PermutationVector.Get<FDDOFGatherRingCountDim>() > FRCPassDiaphragmDOFGather::kMaxSlightOutOfFocusRingCount) return false;
+			if (PermutationVector.Get<FDDOFGatherRingCountDim>() > kMaxSlightOutOfFocusRingCount) return false;
 
 			// Slight out of focus don't need to output CocVariance since there is no hybrid scattering.
-			if (PermutationVector.Get<FDDOFGatherQualityDim>() == FRCPassDiaphragmDOFGather::EQualityConfig::HighQualityWithHybridScatterOcclusion) return false;
+			if (PermutationVector.Get<FDDOFGatherQualityDim>() == EDiaphragmDOFGatherQuality::HighQualityWithHybridScatterOcclusion) return false;
 
 			// Slight out of focus filling can't have lower quality accumulator since it needs to brute force the focus areas.
-			if (PermutationVector.Get<FDDOFGatherQualityDim>() == FRCPassDiaphragmDOFGather::EQualityConfig::LowQualityAccumulator) return false;
+			if (PermutationVector.Get<FDDOFGatherQualityDim>() == EDiaphragmDOFGatherQuality::LowQualityAccumulator) return false;
 
 			// Slight out of focus doesn't have cinematic quality, yet.
-			if (PermutationVector.Get<FDDOFGatherQualityDim>() == FRCPassDiaphragmDOFGather::EQualityConfig::Cinematic) return false;
+			if (PermutationVector.Get<FDDOFGatherQualityDim>() == EDiaphragmDOFGatherQuality::Cinematic) return false;
 
 			// Storing Coc independently of RGB is only supported for RecombineQuality == 0.
 			if (PermutationVector.Get<FDDOFRGBColorBufferDim>())
@@ -821,11 +980,11 @@ class FDiaphragmDOFGatherCS : public FDiaphragmDOFShader
 		else if (PermutationVector.Get<FDDOFLayerProcessingDim>() == EDiaphragmDOFLayerProcessing::BackgroundOnly)
 		{
 			// There is no performance point doing high quality gathering without scattering occlusion.
-			if (PermutationVector.Get<FDDOFGatherQualityDim>() == FRCPassDiaphragmDOFGather::EQualityConfig::HighQuality) return false;
+			if (PermutationVector.Get<FDDOFGatherQualityDim>() == EDiaphragmDOFGatherQuality::HighQuality) return false;
 
 			// Storing Coc independently of RGB is only supported for low gathering quality.
 			if (PermutationVector.Get<FDDOFRGBColorBufferDim>() &&
-				PermutationVector.Get<FDDOFGatherQualityDim>() != FRCPassDiaphragmDOFGather::EQualityConfig::LowQualityAccumulator)
+				PermutationVector.Get<FDDOFGatherQualityDim>() != EDiaphragmDOFGatherQuality::LowQualityAccumulator)
 				return false;
 		}
 		else if (PermutationVector.Get<FDDOFLayerProcessingDim>() == EDiaphragmDOFLayerProcessing::ForegroundAndBackground)
@@ -950,7 +1109,7 @@ class FDiaphragmDOFHybridScatterVS : public FDiaphragmDOFShader
 
 	static bool ShouldCompilePermutation(const FGlobalShaderPermutationParameters& Parameters)
 	{
-		if (!FRCPassDiaphragmDOFHybridScatter::IsSupported(Parameters.Platform))
+		if (!SupportsHybridScatter(Parameters.Platform))
 		{
 			return false;
 		}
@@ -982,7 +1141,7 @@ class FDiaphragmDOFHybridScatterPS : public FDiaphragmDOFShader
 
 	static bool ShouldCompilePermutation(const FGlobalShaderPermutationParameters& Parameters)
 	{
-		if (!FRCPassDiaphragmDOFHybridScatter::IsSupported(Parameters.Platform))
+		if (!SupportsHybridScatter(Parameters.Platform))
 		{
 			return false;
 		}
@@ -1019,7 +1178,7 @@ class FDiaphragmDOFRecombineCS : public FDiaphragmDOFShader
 			PermutationVector.Get<FDDOFLayerProcessingDim>() != EDiaphragmDOFLayerProcessing::ForegroundAndBackground) return false;
 
 		// Do not compile bokeh simulation shaders on platform that couldn't handle them anyway.
-		if (!FRCPassDiaphragmDOFGather::SupportsBokehSimmulation(Parameters.Platform) &&
+		if (!SupportsBokehSimmulation(Parameters.Platform) &&
 			PermutationVector.Get<FDDOFBokehSimulationDim>() != EDiaphragmDOFBokehSimulation::Disabled) return false;
 
 		return FDiaphragmDOFShader::ShouldCompilePermutation(Parameters);
@@ -1090,7 +1249,7 @@ FRDGTextureRef DiaphragmDOF::AddPasses(
 	const EShaderPlatform ShaderPlatform = View.GetShaderPlatform();
 
 	// Number of sampling ring in the gathering kernel.
-	const int32 HalfResRingCount = FMath::Clamp(CVarRingCount.GetValueOnRenderThread(), FRCPassDiaphragmDOFGather::kMinRingCount, FRCPassDiaphragmDOFGather::MaxRingCount(ShaderPlatform));
+	const int32 HalfResRingCount = FMath::Clamp(CVarRingCount.GetValueOnRenderThread(), kMinGatheringRingCount, MaxGatheringRingCount(ShaderPlatform));
 
 	// Post filtering method to do.
 	const EDiaphragmDOFPostfilterMethod PostfilterMethod = GetPostfilteringMethod();
@@ -1099,10 +1258,10 @@ FRDGTextureRef DiaphragmDOF::AddPasses(
 	const EHybridScatterMode FgdHybridScatteringMode = EHybridScatterMode(CVarHybridScatterForegroundMode.GetValueOnRenderThread());
 	const EHybridScatterMode BgdHybridScatteringMode = EHybridScatterMode(CVarHybridScatterBackgroundMode.GetValueOnRenderThread());
 	
-	const float MinScatteringCocRadius = FMath::Max(CVarScatterMinCocRadius.GetValueOnRenderThread(), FRCPassDiaphragmDOFHybridScatter::kMinCocRadius);
+	const float MinScatteringCocRadius = FMath::Max(CVarScatterMinCocRadius.GetValueOnRenderThread(), kMinScatteringCocRadius);
 
 	// Whether the platform support gather bokeh simmulation.
-	const bool bSupportGatheringBokehSimulation = FRCPassDiaphragmDOFGather::SupportsBokehSimmulation(ShaderPlatform);
+	const bool bSupportGatheringBokehSimulation = SupportsBokehSimmulation(ShaderPlatform);
 
 	// Whether should use shade permutation that does lower quality accumulation.
 	// TODO: this is becoming a mess.
@@ -1117,7 +1276,7 @@ FRDGTextureRef DiaphragmDOF::AddPasses(
 	const bool bSupportsSlightOutOfFocus = View.PrimaryScreenPercentageMethod != EPrimaryScreenPercentageMethod::TemporalUpscale;
 
 	// Quality setting for the recombine pass.
-	const int32 RecombineQuality = bSupportsSlightOutOfFocus ? FMath::Clamp(CVarRecombineQuality.GetValueOnRenderThread(), 0, FRCPassDiaphragmDOFRecombine::kMaxQuality) : 0;
+	const int32 RecombineQuality = bSupportsSlightOutOfFocus ? FMath::Clamp(CVarRecombineQuality.GetValueOnRenderThread(), 0, kMaxRecombineQuality) : 0;
 
 	// Resolution divisor.
 	// TODO: Exposes lower resolution divisor?
@@ -1153,7 +1312,7 @@ FRDGTextureRef DiaphragmDOF::AddPasses(
 		// This is just to get the number of shader permutation down.
 		RecombineQuality == 0 &&
 		bUseLowAccumulatorQuality &&
-		FRCPassDiaphragmDOFGather::SupportRGBColorBuffer(ShaderPlatform));
+		SupportsRGBColorBuffer(ShaderPlatform));
 
 
 	// Derives everything needed from the view.
@@ -1205,7 +1364,7 @@ FRDGTextureRef DiaphragmDOF::AddPasses(
 	bool bForegroundHybridScattering = FgdHybridScatteringMode != EHybridScatterMode::Disabled && AbsMaxForegroundCocRadius > MinScatteringCocRadius && MaxScatteringRatio > 0.0f;
 	bool bBackgroundHybridScattering = BgdHybridScatteringMode != EHybridScatterMode::Disabled && MaxBackgroundCocRadius > MinScatteringCocRadius && MaxScatteringRatio > 0.0f;
 
-	if (!FRCPassDiaphragmDOFHybridScatter::IsSupported(ShaderPlatform))
+	if (!SupportsHybridScatter(ShaderPlatform))
 	{
 		bForegroundHybridScattering = false;
 		bBackgroundHybridScattering = false;
@@ -1445,7 +1604,7 @@ FRDGTextureRef DiaphragmDOF::AddPasses(
 
 		// Add one coc dilate pass.
 		auto AddCocDilatePass = [&](
-			FRCPassDiaphragmDOFDilateCoc::EMode Mode,
+			EDiaphragmDOFDilateCocMode Mode,
 			const FDOFTileClassificationTextures& TileInput,
 			const FDOFTileClassificationTextures& DilatedTileMinMax,
 			int32 SampleRadiusCount, int32 SampleOffsetMultipler)
@@ -1455,7 +1614,7 @@ FRDGTextureRef DiaphragmDOF::AddPasses(
 				TEXT("DOFDilateFgdCoc"),
 				TEXT("DOFDilateBgdCoc"),
 			};
-			if (Mode == FRCPassDiaphragmDOFDilateCoc::EMode::MinForegroundAndMaxBackground)
+			if (Mode == EDiaphragmDOFDilateCocMode::MinForegroundAndMaxBackground)
 			{
 				OutputDebugNames[0] = TEXT("DOFDilateMinFgdCoc");
 				OutputDebugNames[1] = TEXT("DOFDilateMaxBgdCoc");
@@ -1499,11 +1658,11 @@ FRDGTextureRef DiaphragmDOF::AddPasses(
 		int32 SampleRadiusCount[3];
 		int32 SampleDistanceMultiplier[3];
 		{
-			const int32 MaxSampleRadiusCount = FRCPassDiaphragmDOFDilateCoc::MaxSampleRadiusCount;
+			const int32 MaxSampleRadiusCount = kMaxCocDilateSampleRadiusCount;
 
 			// Compute the maximum tile dilation.
 			int32 MaximumTileDilation = FMath::CeilToInt(
-				(MaxBluringRadius * BluringRadiusErrorMultiplier) / FRCPassDiaphragmDOFFlattenCoc::CocTileResolutionDivisor);
+				(MaxBluringRadius * BluringRadiusErrorMultiplier) / kCocTileSize);
 
 			// There is always at least one dilate pass so that even small Coc radius conservatively dilate on next neighboor.
 			int32 CurrentConvolutionRadius = FMath::Min(
@@ -1549,7 +1708,7 @@ FRDGTextureRef DiaphragmDOF::AddPasses(
 			for (int32 i = 0; i < DilateCount; i++)
 			{
 				MinMaxTexture = AddCocDilatePass(
-					FRCPassDiaphragmDOFDilateCoc::EMode::MinForegroundAndMaxBackground,
+					EDiaphragmDOFDilateCocMode::MinForegroundAndMaxBackground,
 					MinMaxTexture, FDOFTileClassificationTextures(),
 					SampleRadiusCount[i], SampleDistanceMultiplier[i]);
 			}
@@ -1560,7 +1719,7 @@ FRDGTextureRef DiaphragmDOF::AddPasses(
 			for (int32 i = 0; i < DilateCount; i++)
 			{
 				TileClassificationTextures = AddCocDilatePass(
-					FRCPassDiaphragmDOFDilateCoc::EMode::MinimalAbsoluteRadiuses,
+					EDiaphragmDOFDilateCocMode::MinimalAbsoluteRadiuses,
 					TileClassificationTextures, MinMaxTexture,
 					SampleRadiusCount[i], SampleDistanceMultiplier[i]);
 			}
@@ -1568,7 +1727,7 @@ FRDGTextureRef DiaphragmDOF::AddPasses(
 		else
 		{
 			TileClassificationTextures = AddCocDilatePass(
-				FRCPassDiaphragmDOFDilateCoc::EMode::StandAlone,
+				EDiaphragmDOFDilateCocMode::StandAlone,
 				FlattenedTileClassificationTextures, FDOFTileClassificationTextures(),
 				SampleRadiusCount[0], SampleDistanceMultiplier[0]);
 		}
@@ -1583,7 +1742,7 @@ FRDGTextureRef DiaphragmDOF::AddPasses(
 		FIntPoint SrcSize = HalfResGatherInputDescs.SceneColor.Extent;
 		
 		// Compute the number of mip level required by the gathering pass.
-		const int32 MipLevelCount = FMath::Clamp(FMath::CeilToInt(FMath::Log2(MaxBluringRadius * 0.5 / HalfResRingCount)) + (bUseLowAccumulatorQuality ? 1 : 0), 2, FRCPassDiaphragmDOFReduce::kMaxMipLevelCount);
+		const int32 MipLevelCount = FMath::Clamp(FMath::CeilToInt(FMath::Log2(MaxBluringRadius * 0.5 / HalfResRingCount)) + (bUseLowAccumulatorQuality ? 1 : 0), 2, kMaxMipLevelCount);
 
 		// Maximum number of scattering group per draw instance.
 		// TODO: depends.
@@ -1750,7 +1909,7 @@ FRDGTextureRef DiaphragmDOF::AddPasses(
 	}
 	
 	// Add a pass to build a bokeh LUTs.
-	auto AddBuildBokehLUTPass = [&](FRCPassDiaphragmDOFBuildBokehLUT::EFormat LUTFormat)
+	auto AddBuildBokehLUTPass = [&](EDiaphragmDOFBokehLUTFormat LUTFormat)
 	{
 		if (BokehModel.BokehShape == EBokehShape::Circle)
 		{
@@ -1765,7 +1924,7 @@ FRDGTextureRef DiaphragmDOF::AddPasses(
 
 		FRDGTextureDesc BokehLUTDesc;
 		BokehLUTDesc.NumMips = 1;
-		BokehLUTDesc.Format = LUTFormat == FRCPassDiaphragmDOFBuildBokehLUT::EFormat::GatherSamplePos ? PF_G16R16F : PF_R16F;
+		BokehLUTDesc.Format = LUTFormat == EDiaphragmDOFBokehLUTFormat::GatherSamplePos ? PF_G16R16F : PF_R16F;
 		BokehLUTDesc.Extent = FIntPoint(32, 32);
 		BokehLUTDesc.TargetableFlags |= TexCreate_UAV;
 
@@ -1812,7 +1971,7 @@ FRDGTextureRef DiaphragmDOF::AddPasses(
 			EDiaphragmDOFLayerProcessing LayerProcessing;
 
 			/** Configuration of the pass. */
-			FRCPassDiaphragmDOFGather::EQualityConfig QualityConfig;
+			EDiaphragmDOFGatherQuality QualityConfig;
 		
 			/** Postfilter method to apply on this gather pass. */
 			EDiaphragmDOFPostfilterMethod PostfilterMethod;
@@ -1822,7 +1981,7 @@ FRDGTextureRef DiaphragmDOF::AddPasses(
 		
 			FConvolutionSettings()
 				: LayerProcessing(EDiaphragmDOFLayerProcessing::ForegroundAndBackground)
-				, QualityConfig(FRCPassDiaphragmDOFGather::EQualityConfig::HighQuality)
+				, QualityConfig(EDiaphragmDOFGatherQuality::HighQuality)
 				, PostfilterMethod(EDiaphragmDOFPostfilterMethod::None)
 				, BokehSimulation(EDiaphragmDOFBokehSimulation::Disabled)
 			{ }
@@ -1865,7 +2024,7 @@ FRDGTextureRef DiaphragmDOF::AddPasses(
 					}
 				}
 
-				if (ConvolutionSettings.QualityConfig == FRCPassDiaphragmDOFGather::EQualityConfig::HighQualityWithHybridScatterOcclusion)
+				if (ConvolutionSettings.QualityConfig == EDiaphragmDOFGatherQuality::HighQualityWithHybridScatterOcclusion)
 				{
 					Desc.Format = PF_G16R16F;
 			
@@ -1903,7 +2062,7 @@ FRDGTextureRef DiaphragmDOF::AddPasses(
 				const float GatheringScalingDownFactor = float(PreprocessViewSize.X) / float(GatheringViewSize.X);
 
 				// Coc radius considered.
-				const float RecombineCocRadiusBorder = GatheringScalingDownFactor * (FRCPassDiaphragmDOFGather::kMaxSlightOutOfFocusRingCount - 1.0f);
+				const float RecombineCocRadiusBorder = GatheringScalingDownFactor * (kMaxSlightOutOfFocusRingCount - 1.0f);
 
 				if (ConvolutionSettings.LayerProcessing == EDiaphragmDOFLayerProcessing::ForegroundOnly)
 				{
@@ -1962,7 +2121,7 @@ FRDGTextureRef DiaphragmDOF::AddPasses(
 			PassParameters->BokehLUT = BokehLUT;
 
 			PassParameters->ConvolutionOutput = CreateUAVs(GraphBuilder, *ConvolutionOutputTextures);
-			if (ConvolutionSettings.QualityConfig == FRCPassDiaphragmDOFGather::EQualityConfig::HighQualityWithHybridScatterOcclusion)
+			if (ConvolutionSettings.QualityConfig == EDiaphragmDOFGatherQuality::HighQualityWithHybridScatterOcclusion)
 				PassParameters->ScatterOcclusionOutput = GraphBuilder.CreateUAV(*ScatterOcclusionTexture);
 			
 			TShaderMapRef<FDiaphragmDOFGatherCS> ComputeShader(View.ShaderMap, PermutationVector);
@@ -2034,11 +2193,11 @@ FRDGTextureRef DiaphragmDOF::AddPasses(
 		
 		FRDGTextureRef GatheringBokehLUT = nullptr;
 		if (bEnableGatherBokehSettings)
-			 GatheringBokehLUT = AddBuildBokehLUTPass(FRCPassDiaphragmDOFBuildBokehLUT::EFormat::GatherSamplePos);
+			 GatheringBokehLUT = AddBuildBokehLUTPass(EDiaphragmDOFBokehLUTFormat::GatherSamplePos);
 
 		FRDGTextureRef ScatteringBokehLUT = nullptr;
 		if (bEnableScatterBokehSettings || bEnableSlightOutOfFocusBokeh)
-			ScatteringBokehLUT = AddBuildBokehLUTPass(FRCPassDiaphragmDOFBuildBokehLUT::EFormat::CocRadiusToBokehEdgeFactor);
+			ScatteringBokehLUT = AddBuildBokehLUTPass(EDiaphragmDOFBokehLUTFormat::CocRadiusToBokehEdgeFactor);
 
 		auto AddHybridScatterPass = [&](
 			const FConvolutionSettings& ConvolutionSettings,
@@ -2131,7 +2290,7 @@ FRDGTextureRef DiaphragmDOF::AddPasses(
 
 			if (bUseLowAccumulatorQuality)
 			{
-				ConvolutionSettings.QualityConfig = FRCPassDiaphragmDOFGather::EQualityConfig::LowQualityAccumulator;
+				ConvolutionSettings.QualityConfig = EDiaphragmDOFGatherQuality::LowQualityAccumulator;
 			}
 
 			FRDGTextureRef ScatterOcclusionTexture = nullptr;
@@ -2177,16 +2336,16 @@ FRDGTextureRef DiaphragmDOF::AddPasses(
 			if (bEnableGatherBokehSettings)
 				ConvolutionSettings.BokehSimulation = BokehSimulation;
 
-			ConvolutionSettings.QualityConfig = FRCPassDiaphragmDOFGather::EQualityConfig::LowQualityAccumulator;
+			ConvolutionSettings.QualityConfig = EDiaphragmDOFGatherQuality::LowQualityAccumulator;
 			if (bBackgroundHybridScattering && BgdHybridScatteringMode == EHybridScatterMode::Occlusion)
 			{
 				if (bUseCinematicAccumulatorQuality)
 				{
-					ConvolutionSettings.QualityConfig = FRCPassDiaphragmDOFGather::EQualityConfig::Cinematic;
+					ConvolutionSettings.QualityConfig = EDiaphragmDOFGatherQuality::Cinematic;
 				}
 				else
 				{
-					ConvolutionSettings.QualityConfig = FRCPassDiaphragmDOFGather::EQualityConfig::HighQualityWithHybridScatterOcclusion;
+					ConvolutionSettings.QualityConfig = EDiaphragmDOFGatherQuality::HighQualityWithHybridScatterOcclusion;
 				}
 			}
 			
@@ -2242,7 +2401,7 @@ FRDGTextureRef DiaphragmDOF::AddPasses(
 
 		if (bEnableSlightOutOfFocusBokeh) // && ScatteringBokehLUTOutput.IsValid() && SlightOutOfFocusConvolutionOutput.IsValid())
 		{
-			PassParameters->BokehLUT = AddBuildBokehLUTPass(FRCPassDiaphragmDOFBuildBokehLUT::EFormat::FullResOffsetToCocDistance);
+			PassParameters->BokehLUT = AddBuildBokehLUTPass(EDiaphragmDOFBokehLUTFormat::FullResOffsetToCocDistance);
 		}
 
 		PassParameters->SceneColorOutput = GraphBuilder.CreateUAV(NewSceneColor);
diff --git a/Engine/Source/Runtime/Renderer/Private/PostProcess/DiaphragmDOF.h b/Engine/Source/Runtime/Renderer/Private/PostProcess/DiaphragmDOF.h
--- a/Engine/Source/Runtime/Renderer/Private/PostProcess/DiaphragmDOF.h
+++ b/Engine/Source/Runtime/Renderer/Private/PostProcess/DiaphragmDOF.h
@@ -140,9 +140,4 @@ RENDERER_API FRDGTextureRef AddPasses(
 	FRDGTextureRef InputSceneColor,
 	FRDGTextureRef SeparateTranslucency);
 
-
-/** Wire all DOF's passes according to view settings and cvars to convolve the scene color (Context.FinalOutput). */
-// TODO: kill
-RENDERER_API bool WireSceneColorPasses(FPostprocessContext& Context, const FRenderingCompositeOutputRef& VelocityInput, const FRenderingCompositeOutputRef& SeparateTranslucency);
-
 } // namespace DiaphragmDOF
diff --git a/Engine/Source/Runtime/Renderer/Private/PostProcess/DiaphragmDOFGraph.cpp b/Engine/Source/Runtime/Renderer/Private/PostProcess/DiaphragmDOFGraph.cpp
deleted file mode 100644
diff --git a/Engine/Source/Runtime/Renderer/Private/PostProcess/DiaphragmDOFPasses.cpp b/Engine/Source/Runtime/Renderer/Private/PostProcess/DiaphragmDOFPasses.cpp
deleted file mode 100644
diff --git a/Engine/Source/Runtime/Renderer/Private/PostProcess/DiaphragmDOFPasses.h b/Engine/Source/Runtime/Renderer/Private/PostProcess/DiaphragmDOFPasses.h
deleted file mode 100644
diff --git a/Engine/Source/Runtime/Renderer/Private/PostProcess/PostProcessing.cpp b/Engine/Source/Runtime/Renderer/Private/PostProcess/PostProcessing.cpp
--- a/Engine/Source/Runtime/Renderer/Private/PostProcess/PostProcessing.cpp
+++ b/Engine/Source/Runtime/Renderer/Private/PostProcess/PostProcessing.cpp
@@ -1565,46 +1565,38 @@ void FPostProcessing::Process(FRHICommandListImmediate& RHICmdList, const FViewI
 					if (CVarUseDiaphragmDOF.GetValueOnRenderThread() > 0 &&
 						DiaphragmDOF::IsSupported(View.GetShaderPlatform()))
 					{
-						if (CVarUseDiaphragmDOF.GetValueOnRenderThread() == 2)
+						FRenderingCompositePass* DiaphragmDOFPass = Context.Graph.RegisterPass(
+							new(FMemStack::Get()) TRCPassForRDG<2, 1>([](FRenderingCompositePass* Pass, FRenderingCompositePassContext& Context)
 						{
-							FRenderingCompositePass* DiaphragmDOFPass = Context.Graph.RegisterPass(
-								new(FMemStack::Get()) TRCPassForRDG<2, 1>([](FRenderingCompositePass* Pass, FRenderingCompositePassContext& Context)
-							{
-								FRDGBuilder GraphBuilder(Context.RHICmdList);
+							FRDGBuilder GraphBuilder(Context.RHICmdList);
 
-								FSceneRenderTargets& SceneContext = FSceneRenderTargets::Get(Context.RHICmdList);
-								FSceneViewFamilyBlackboard SceneBlackboard;
-								SetupSceneViewFamilyBlackboard(GraphBuilder, &SceneBlackboard);
+							FSceneRenderTargets& SceneContext = FSceneRenderTargets::Get(Context.RHICmdList);
+							FSceneViewFamilyBlackboard SceneBlackboard;
+							SetupSceneViewFamilyBlackboard(GraphBuilder, &SceneBlackboard);
 	
-								FRDGTextureRef SceneColor = Pass->CreateRDGTextureForInput(GraphBuilder, ePId_Input0, TEXT("SceneColor"), eFC_0000);
-								FRDGTextureRef SeparateTranslucency = Pass->CreateRDGTextureForInput(GraphBuilder, ePId_Input1, TEXT("SeparateTranslucency"), eFC_0000);
+							FRDGTextureRef SceneColor = Pass->CreateRDGTextureForInput(GraphBuilder, ePId_Input0, TEXT("SceneColor"), eFC_0000);
+							FRDGTextureRef SeparateTranslucency = Pass->CreateRDGTextureForInput(GraphBuilder, ePId_Input1, TEXT("SeparateTranslucency"), eFC_0000);
 
-								FRDGTextureRef NewSceneColor = DiaphragmDOF::AddPasses(
-									GraphBuilder,
-									SceneBlackboard, Context.View,
-									SceneColor, SeparateTranslucency);
+							FRDGTextureRef NewSceneColor = DiaphragmDOF::AddPasses(
+								GraphBuilder,
+								SceneBlackboard, Context.View,
+								SceneColor, SeparateTranslucency);
 		
-								// DOF passes were not added, therefore need to compose Separate translucency manually. 
-								if (NewSceneColor == SceneColor)
-								{
-									NewSceneColor = AddSeparateTranslucencyCompositionPass(GraphBuilder, Context.View, SceneColor, SeparateTranslucency);
-								}
+							// DOF passes were not added, therefore need to compose Separate translucency manually. 
+							if (NewSceneColor == SceneColor)
+							{
+								NewSceneColor = AddSeparateTranslucencyCompositionPass(GraphBuilder, Context.View, SceneColor, SeparateTranslucency);
+							}
 
-								Pass->ExtractRDGTextureForOutput(GraphBuilder, ePId_Output0, NewSceneColor);
+							Pass->ExtractRDGTextureForOutput(GraphBuilder, ePId_Output0, NewSceneColor);
 
-								GraphBuilder.Execute();
-							}));
-							DiaphragmDOFPass->SetInput(ePId_Input0, Context.FinalOutput);
-							DiaphragmDOFPass->SetInput(ePId_Input1, SeparateTranslucency);
-							Context.FinalOutput = FRenderingCompositeOutputRef(DiaphragmDOFPass, ePId_Output0);
+							GraphBuilder.Execute();
+						}));
+						DiaphragmDOFPass->SetInput(ePId_Input0, Context.FinalOutput);
+						DiaphragmDOFPass->SetInput(ePId_Input1, SeparateTranslucency);
+						Context.FinalOutput = FRenderingCompositeOutputRef(DiaphragmDOFPass, ePId_Output0);
 
-							bSepTransWasApplied = true;
-						}
-						else
-						{
-							bSepTransWasApplied = DiaphragmDOF::WireSceneColorPasses(
-								Context, DOFVelocityRef,SeparateTranslucency);
-						}
+						bSepTransWasApplied = true;
 					}
 					else
 					{
