Fixed compile errors on Linux due to missing header includes

diff --git a/Engine/Source/Runtime/Renderer/Private/ScreenPass.h b/Engine/Source/Runtime/Renderer/Private/ScreenPass.h
--- a/Engine/Source/Runtime/Renderer/Private/ScreenPass.h
+++ b/Engine/Source/Runtime/Renderer/Private/ScreenPass.h
@@ -8,6 +8,7 @@
 #include "PipelineStateCache.h"
 #include "RenderGraph.h"
 #include "PostProcess/SceneFilterRendering.h"
+#include "ScenePrivate.h"
 
 // Returns whether a HMD hidden area mask is being used for VR.
 bool IsHMDHiddenAreaMaskActive();
