Fixed SSAO is applied to ray tracing debug view modes

diff --git a/Engine/Source/Runtime/Renderer/Private/CompositionLighting/CompositionLighting.cpp b/Engine/Source/Runtime/Renderer/Private/CompositionLighting/CompositionLighting.cpp
--- a/Engine/Source/Runtime/Renderer/Private/CompositionLighting/CompositionLighting.cpp
+++ b/Engine/Source/Runtime/Renderer/Private/CompositionLighting/CompositionLighting.cpp
@@ -280,6 +280,7 @@ void FCompositionLighting::ProcessAfterBasePass(FRHICommandListImmediate& RHICmd
 	GVisualizeTexture.SetCheckPoint(RHICmdList, SceneContext.ScreenSpaceAO);
 	
 	// so that the passes can register themselves to the graph
+	if(CanOverlayRayTracingOutput(View))
 	{
 		FMemMark Mark(FMemStack::Get());
 		FRenderingCompositePassContext CompositeContext(RHICmdList, View);
diff --git a/Engine/Source/Runtime/Renderer/Private/DeferredShadingRenderer.cpp b/Engine/Source/Runtime/Renderer/Private/DeferredShadingRenderer.cpp
--- a/Engine/Source/Runtime/Renderer/Private/DeferredShadingRenderer.cpp
+++ b/Engine/Source/Runtime/Renderer/Private/DeferredShadingRenderer.cpp
@@ -443,7 +443,7 @@ bool FDeferredShadingSceneRenderer::RenderHzb(FRHICommandListImmediate& RHICmdLi
 	}
 
 	//async ssao only requires HZB and depth as inputs so get started ASAP
-	if (CanOverlayRayTracingOutput() && GCompositionLighting.CanProcessAsyncSSAO(Views))
+	if (CanOverlayRayTracingOutput(Views[0]) && GCompositionLighting.CanProcessAsyncSSAO(Views))
 	{
 		GCompositionLighting.ProcessAsyncSSAO(RHICmdList, Views);
 	}
@@ -953,7 +953,7 @@ void FDeferredShadingSceneRenderer::Render(FRHICommandListImmediate& RHICmdList)
 #endif
 		;
 
-	bool bCanOverlayRayTracingOutput = CanOverlayRayTracingOutput();
+	bool bCanOverlayRayTracingOutput = CanOverlayRayTracingOutput(Views[0]);// #dxr_todo: UE-72557 multi-view case
 	bool bComputeLightGrid = false;
 	// Simple forward shading doesn't support local lights. No need to compute light grid
 	if (!IsSimpleForwardShadingEnabled(ShaderPlatform))
@@ -1544,18 +1544,6 @@ void FDeferredShadingSceneRenderer::Render(FRHICommandListImmediate& RHICmdList)
 	const bool bRayTracingEnabled = IsRayTracingEnabled();
 	if (bRayTracingEnabled)
 	{
-		for (int32 ViewIndex = 0; ViewIndex < Views.Num(); ViewIndex++)
-		{
-			if (Views[ViewIndex].RayTracingRenderMode == ERayTracingRenderMode::PathTracing)
-			{
-				RenderPathTracing(RHICmdList, Views[ViewIndex]);
-			}
-			else if (Views[ViewIndex].RayTracingRenderMode == ERayTracingRenderMode::RayTracingDebug)
-			{
-				RenderRayTracingDebug(RHICmdList, Views[ViewIndex]);
-			}
-		}
-
 		if (bCanOverlayRayTracingOutput)
 		{
 			// TODO: convert the entire AO and skylight to rendergraph.
@@ -1566,7 +1554,6 @@ void FDeferredShadingSceneRenderer::Render(FRHICommandListImmediate& RHICmdList)
 				RenderRayTracingSkyLight(RHICmdList, SkyLightRT, HitDistanceRT);
 			}
 
-
 			if (ShouldRenderRayTracingGlobalIllumination(Views))
 			{
 				for (int32 ViewIndex = 0; ViewIndex < Views.Num(); ViewIndex++)
@@ -1995,6 +1982,23 @@ void FDeferredShadingSceneRenderer::Render(FRHICommandListImmediate& RHICmdList)
 		ServiceLocalQueue();
 	}
 
+#if RHI_RAYTRACING
+	if (bRayTracingEnabled)
+	{
+		for (int32 ViewIndex = 0; ViewIndex < Views.Num(); ViewIndex++)
+		{
+			if (Views[ViewIndex].RayTracingRenderMode == ERayTracingRenderMode::PathTracing)
+			{
+				RenderPathTracing(RHICmdList, Views[ViewIndex]);
+			}
+			else if (Views[ViewIndex].RayTracingRenderMode == ERayTracingRenderMode::RayTracingDebug)
+			{
+				RenderRayTracingDebug(RHICmdList, Views[ViewIndex]);
+			}
+		}
+	}
+#endif
+
 	for (int32 ViewIndex = 0; ViewIndex < Views.Num(); ++ViewIndex)
 	{
 		const FViewInfo& View = Views[ViewIndex];
@@ -2354,13 +2358,10 @@ void FDeferredShadingSceneRenderer::CopyStencilToLightingChannelTexture(FRHIComm
 	}
 }
 
-bool FDeferredShadingSceneRenderer::CanOverlayRayTracingOutput(void) const
-{
 #if RHI_RAYTRACING
-	// #dxr_todo:UE-72557 multi-view case
-	return  (Views[0].RayTracingRenderMode != ERayTracingRenderMode::PathTracing)
-		&& (Views[0].RayTracingRenderMode != ERayTracingRenderMode::RayTracingDebug);
-#else // RHI_RAYTRACING
-	return true;
-#endif // RHI_RAYTRACING
+bool CanOverlayRayTracingOutput(const FViewInfo& View)
+{
+	return (View.RayTracingRenderMode != ERayTracingRenderMode::PathTracing)
+		&& (View.RayTracingRenderMode != ERayTracingRenderMode::RayTracingDebug);
 }
+#endif // RHI_RAYTRACING
diff --git a/Engine/Source/Runtime/Renderer/Private/RayTracing/RaytracingOptions.h b/Engine/Source/Runtime/Renderer/Private/RayTracing/RaytracingOptions.h
--- a/Engine/Source/Runtime/Renderer/Private/RayTracing/RaytracingOptions.h
+++ b/Engine/Source/Runtime/Renderer/Private/RayTracing/RaytracingOptions.h
@@ -15,6 +15,7 @@ extern bool ShouldRenderRayTracingSkyLight(const FSkyLightSceneProxy* SkyLightSc
 extern bool ShouldRenderRayTracingAmbientOcclusion();
 extern bool ShouldRenderRayTracingGlobalIllumination(const TArray<FViewInfo>& Views);
 extern bool ShouldRenderRayTracingStochasticRectLight(const FLightSceneInfo& LightSceneInfo);
+extern bool CanOverlayRayTracingOutput(const FViewInfo& View);
 
 extern float GetRaytracingMaxNormalBias();
 
@@ -35,5 +36,9 @@ FORCEINLINE bool ShouldRenderRayTracingStochasticRectLight(const FLightSceneInfo
 	return false;
 }
 
+FORCEINLINE bool CanOverlayRayTracingOutput(const FViewInfo& View)
+{
+	return true;
+}
 
 #endif
