Fixed compilation issue because I forgot to check one file in

diff --git a/Engine/Source/Runtime/Renderer/Private/DeferredShadingRenderer.h b/Engine/Source/Runtime/Renderer/Private/DeferredShadingRenderer.h
--- a/Engine/Source/Runtime/Renderer/Private/DeferredShadingRenderer.h
+++ b/Engine/Source/Runtime/Renderer/Private/DeferredShadingRenderer.h
@@ -462,8 +462,6 @@ private:
 
 	void CopySceneCaptureComponentToTarget(FRHICommandListImmediate& RHICmdList);
 
-	bool CanOverlayRayTracingOutput(void) const;
-
 	void RenderRayTracingReflections(
 		FRDGBuilder& GraphBuilder,
 		const FViewInfo& View,
