DR - Fix some warnings
#rnx

diff --git a/Engine/Source/Runtime/Renderer/Private/PostProcess/DiaphragmDOF.cpp b/Engine/Source/Runtime/Renderer/Private/PostProcess/DiaphragmDOF.cpp
--- a/Engine/Source/Runtime/Renderer/Private/PostProcess/DiaphragmDOF.cpp
+++ b/Engine/Source/Runtime/Renderer/Private/PostProcess/DiaphragmDOF.cpp
@@ -512,7 +512,9 @@ FDOFGatherInputTextures CreateTextures(FRDGBuilder& GraphBuilder, const FDOFGath
 	FDOFGatherInputTextures Textures;
 	Textures.SceneColor = GraphBuilder.CreateTexture(Descs.SceneColor, DebugName);
 	if (Descs.SeparateCoc.Format != PF_Unknown)
+	{
 		Textures.SeparateCoc = GraphBuilder.CreateTexture(Descs.SeparateCoc, DebugName);
+	}
 	return Textures;
 }
 
@@ -521,7 +523,9 @@ FDOFGatherInputUAVs CreateUAVs(FRDGBuilder& GraphBuilder, const FDOFGatherInputT
 	FDOFGatherInputUAVs UAVs;
 	UAVs.SceneColor = GraphBuilder.CreateUAV(FRDGTextureUAVDesc(Textures.SceneColor, MipLevel));
 	if (Textures.SeparateCoc)
+	{
 		UAVs.SeparateCoc = GraphBuilder.CreateUAV(FRDGTextureUAVDesc(Textures.SeparateCoc, MipLevel));
+	}
 	return UAVs;
 }
 
@@ -547,7 +551,9 @@ FDOFConvolutionTextures CreateTextures(FRDGBuilder& GraphBuilder, const FDOFConv
 	FDOFConvolutionTextures Textures;
 	Textures.SceneColor = GraphBuilder.CreateTexture(Descs.SceneColor, DebugName);
 	if (Descs.SeparateAlpha.Format != PF_Unknown)
+	{
 		Textures.SeparateAlpha = GraphBuilder.CreateTexture(Descs.SeparateAlpha, DebugName);
+	}
 	return Textures;
 }
 
@@ -556,7 +562,9 @@ FDOFConvolutionUAVs CreateUAVs(FRDGBuilder& GraphBuilder, const FDOFConvolutionT
 	FDOFConvolutionUAVs UAVs;
 	UAVs.SceneColor = GraphBuilder.CreateUAV(Textures.SceneColor);
 	if (Textures.SeparateAlpha)
+	{
 		UAVs.SeparateAlpha = GraphBuilder.CreateUAV(Textures.SeparateAlpha);
+	}
 	return UAVs;
 }
 
@@ -2122,7 +2130,9 @@ FRDGTextureRef DiaphragmDOF::AddPasses(
 
 			PassParameters->ConvolutionOutput = CreateUAVs(GraphBuilder, *ConvolutionOutputTextures);
 			if (ConvolutionSettings.QualityConfig == EDiaphragmDOFGatherQuality::HighQualityWithHybridScatterOcclusion)
+			{
 				PassParameters->ScatterOcclusionOutput = GraphBuilder.CreateUAV(*ScatterOcclusionTexture);
+			}
 			
 			TShaderMapRef<FDiaphragmDOFGatherCS> ComputeShader(View.ShaderMap, PermutationVector);
 			FComputeShaderUtils::AddPass(
diff --git a/Engine/Source/Runtime/Renderer/Private/PostProcess/PostProcessing.cpp b/Engine/Source/Runtime/Renderer/Private/PostProcess/PostProcessing.cpp
--- a/Engine/Source/Runtime/Renderer/Private/PostProcess/PostProcessing.cpp
+++ b/Engine/Source/Runtime/Renderer/Private/PostProcess/PostProcessing.cpp
@@ -1566,11 +1566,11 @@ void FPostProcessing::Process(FRHICommandListImmediate& RHICmdList, const FViewI
 						DiaphragmDOF::IsSupported(View.GetShaderPlatform()))
 					{
 						FRenderingCompositePass* DiaphragmDOFPass = Context.Graph.RegisterPass(
-							new(FMemStack::Get()) TRCPassForRDG<2, 1>([](FRenderingCompositePass* Pass, FRenderingCompositePassContext& Context)
+							new(FMemStack::Get()) TRCPassForRDG<2, 1>([](FRenderingCompositePass* Pass, FRenderingCompositePassContext& InContext)
 						{
-							FRDGBuilder GraphBuilder(Context.RHICmdList);
+							FRDGBuilder GraphBuilder(InContext.RHICmdList);
 
-							FSceneRenderTargets& SceneContext = FSceneRenderTargets::Get(Context.RHICmdList);
+							FSceneRenderTargets& SceneContext = FSceneRenderTargets::Get(InContext.RHICmdList);
 							FSceneViewFamilyBlackboard SceneBlackboard;
 							SetupSceneViewFamilyBlackboard(GraphBuilder, &SceneBlackboard);
 	
@@ -1579,13 +1579,13 @@ void FPostProcessing::Process(FRHICommandListImmediate& RHICmdList, const FViewI
 
 							FRDGTextureRef NewSceneColor = DiaphragmDOF::AddPasses(
 								GraphBuilder,
-								SceneBlackboard, Context.View,
+								SceneBlackboard, InContext.View,
 								SceneColor, SeparateTranslucency);
 		
 							// DOF passes were not added, therefore need to compose Separate translucency manually. 
 							if (NewSceneColor == SceneColor)
 							{
-								NewSceneColor = AddSeparateTranslucencyCompositionPass(GraphBuilder, Context.View, SceneColor, SeparateTranslucency);
+								NewSceneColor = AddSeparateTranslucencyCompositionPass(GraphBuilder, InContext.View, SceneColor, SeparateTranslucency);
 							}
 
 							Pass->ExtractRDGTextureForOutput(GraphBuilder, ePId_Output0, NewSceneColor);
