DR - More warning fixes
#rnx

diff --git a/Engine/Source/Runtime/Renderer/Private/PostProcess/PostProcessMotionBlur.cpp b/Engine/Source/Runtime/Renderer/Private/PostProcess/PostProcessMotionBlur.cpp
--- a/Engine/Source/Runtime/Renderer/Private/PostProcess/PostProcessMotionBlur.cpp
+++ b/Engine/Source/Runtime/Renderer/Private/PostProcess/PostProcessMotionBlur.cpp
@@ -815,10 +815,10 @@ FScreenPassTexture VisualizeMotionBlur(
 
 		// Draw debug text.
 		{
-			const FViewInfo& View = Context->View;
+			const FViewInfo& LocalView = Context->View;
 			const FSceneViewFamily& ViewFamily = Context->ViewFamily;
 			FRenderTargetTemp TempRenderTarget(static_cast<FTexture2DRHIParamRef>(TextureOutput.GetRHITexture()), TextureOutput.GetViewport().Extent);
-			FCanvas Canvas(&TempRenderTarget, nullptr, ViewFamily.CurrentRealTime, ViewFamily.CurrentWorldTime, ViewFamily.DeltaWorldTime, View.GetFeatureLevel());
+			FCanvas Canvas(&TempRenderTarget, nullptr, ViewFamily.CurrentRealTime, ViewFamily.CurrentWorldTime, ViewFamily.DeltaWorldTime, LocalView.GetFeatureLevel());
 
 			float X = 20;
 			float Y = 38;
@@ -846,12 +846,12 @@ FScreenPassTexture VisualizeMotionBlur(
 			Canvas.DrawShadowedString(X, Y += YStep, TEXT("Paused, r.VelocityTest, Parallel:"), GetStatsFont(), FLinearColor(1, 1, 0));
 			Canvas.DrawShadowedString(X + ColumnWidth, Y, *Line, GetStatsFont(), FLinearColor(1, 1, 0));
 
-			const FScene* Scene = (const FScene*)View.Family->Scene;
-			const FSceneViewState *SceneViewState = (const FSceneViewState*)View.State;
+			const FScene* Scene = (const FScene*)LocalView.Family->Scene;
+			const FSceneViewState *SceneViewState = (const FSceneViewState*)LocalView.State;
 
 			Line = FString::Printf(TEXT("View=%.4x PrevView=%.4x"),
-				View.ViewMatrices.GetViewMatrix().ComputeHash() & 0xffff,
-				View.PrevViewInfo.ViewMatrices.GetViewMatrix().ComputeHash() & 0xffff);
+				LocalView.ViewMatrices.GetViewMatrix().ComputeHash() & 0xffff,
+				LocalView.PrevViewInfo.ViewMatrices.GetViewMatrix().ComputeHash() & 0xffff);
 			Canvas.DrawShadowedString(X, Y += YStep, TEXT("ViewMatrix:"), GetStatsFont(), FLinearColor(1, 1, 0));
 			Canvas.DrawShadowedString(X + ColumnWidth, Y, *Line, GetStatsFont(), FLinearColor(1, 1, 0));
 
@@ -862,4 +862,4 @@ FScreenPassTexture VisualizeMotionBlur(
 	});
 
 	return TextureOutput;
-}
\ No newline at end of file
+}
diff --git a/Engine/Source/Runtime/Renderer/Private/PostProcess/PostProcessing.cpp b/Engine/Source/Runtime/Renderer/Private/PostProcess/PostProcessing.cpp
--- a/Engine/Source/Runtime/Renderer/Private/PostProcess/PostProcessing.cpp
+++ b/Engine/Source/Runtime/Renderer/Private/PostProcess/PostProcessing.cpp
@@ -1570,22 +1570,21 @@ void FPostProcessing::Process(FRHICommandListImmediate& RHICmdList, const FViewI
 						{
 							FRDGBuilder GraphBuilder(InContext.RHICmdList);
 
-							FSceneRenderTargets& SceneContext = FSceneRenderTargets::Get(InContext.RHICmdList);
 							FSceneViewFamilyBlackboard SceneBlackboard;
 							SetupSceneViewFamilyBlackboard(GraphBuilder, &SceneBlackboard);
 	
 							FRDGTextureRef SceneColor = Pass->CreateRDGTextureForInput(GraphBuilder, ePId_Input0, TEXT("SceneColor"), eFC_0000);
-							FRDGTextureRef SeparateTranslucency = Pass->CreateRDGTextureForInput(GraphBuilder, ePId_Input1, TEXT("SeparateTranslucency"), eFC_0000);
+							FRDGTextureRef LocalSeparateTranslucency = Pass->CreateRDGTextureForInput(GraphBuilder, ePId_Input1, TEXT("SeparateTranslucency"), eFC_0000);
 
 							FRDGTextureRef NewSceneColor = DiaphragmDOF::AddPasses(
 								GraphBuilder,
 								SceneBlackboard, InContext.View,
-								SceneColor, SeparateTranslucency);
+								SceneColor, LocalSeparateTranslucency);
 		
 							// DOF passes were not added, therefore need to compose Separate translucency manually. 
 							if (NewSceneColor == SceneColor)
 							{
-								NewSceneColor = AddSeparateTranslucencyCompositionPass(GraphBuilder, InContext.View, SceneColor, SeparateTranslucency);
+								NewSceneColor = AddSeparateTranslucencyCompositionPass(GraphBuilder, InContext.View, SceneColor, LocalSeparateTranslucency);
 							}
 
 							Pass->ExtractRDGTextureForOutput(GraphBuilder, ePId_Output0, NewSceneColor);
@@ -1747,12 +1746,12 @@ void FPostProcessing::Process(FRHICommandListImmediate& RHICmdList, const FViewI
 			if ((IsMotionBlurEnabled(View) || bVisualizeMotionBlur) && VelocityInput.IsValid())
 			{
 				FRenderingCompositePass* MotionBlurPass = Context.Graph.RegisterPass(new(FMemStack::Get()) TRCPassForRDG<3, 1>(
-					[bVisualizeMotionBlur](FRenderingCompositePass* Pass, FRenderingCompositePassContext& Context)
+					[bVisualizeMotionBlur](FRenderingCompositePass* Pass, FRenderingCompositePassContext& InContext)
 				{
-					FRDGBuilder GraphBuilder(Context.RHICmdList);
+					FRDGBuilder GraphBuilder(InContext.RHICmdList);
 
-					const FIntRect ColorViewportRect = Context.SceneColorViewRect;
-					const FIntRect DepthViewportRect = Context.View.ViewRect;
+					const FIntRect ColorViewportRect = InContext.SceneColorViewRect;
+					const FIntRect DepthViewportRect = InContext.View.ViewRect;
 
 					const FScreenPassTexture ColorTexture = FScreenPassTexture::Create(
 						Pass->CreateRDGTextureForInput(GraphBuilder, ePId_Input0, TEXT("SceneColor"), eFC_0000),
@@ -1772,7 +1771,7 @@ void FPostProcessing::Process(FRHICommandListImmediate& RHICmdList, const FViewI
 					{
 						OutColorTexture = VisualizeMotionBlur(
 							GraphBuilder,
-							FScreenPassContext::Create(Context.RHICmdList, Context.View),
+							FScreenPassContext::Create(InContext.RHICmdList, InContext.View),
 							ColorTexture,
 							DepthTexture,
 							VelocityTexture);
@@ -1781,7 +1780,7 @@ void FPostProcessing::Process(FRHICommandListImmediate& RHICmdList, const FViewI
 					{
 						OutColorTexture = ComputeMotionBlur(
 							GraphBuilder,
-							FScreenPassContext::Create(Context.RHICmdList, Context.View),
+							FScreenPassContext::Create(InContext.RHICmdList, InContext.View),
 							ColorTexture,
 							DepthTexture,
 							VelocityTexture);
