DR - Warning fix
#rnx

diff --git a/Engine/Source/Runtime/Renderer/Private/ScreenPass.cpp b/Engine/Source/Runtime/Renderer/Private/ScreenPass.cpp
--- a/Engine/Source/Runtime/Renderer/Private/ScreenPass.cpp
+++ b/Engine/Source/Runtime/Renderer/Private/ScreenPass.cpp
@@ -38,9 +38,9 @@ bool IsHMDHiddenAreaMaskActive()
 
 FScreenPassTextureViewport FScreenPassTextureViewport::CreateDownscaled(const FScreenPassTextureViewport& Other, uint32 ScaleFactor)
 {
-	const auto GetDownscaledSize = [](FIntPoint Size, uint32 ScaleFactor)
+	const auto GetDownscaledSize = [](FIntPoint Size, uint32 InScaleFactor)
 	{
-		Size = FIntPoint::DivideAndRoundUp(Size, ScaleFactor);
+		Size = FIntPoint::DivideAndRoundUp(Size, InScaleFactor);
 		Size.X = FMath::Max(1, Size.X);
 		Size.Y = FMath::Max(1, Size.Y);
 		return Size;
