Speculative work arround PVS static analysis.

diff --git a/Engine/Source/Runtime/Renderer/Private/PostProcess/DiaphragmDOF.cpp b/Engine/Source/Runtime/Renderer/Private/PostProcess/DiaphragmDOF.cpp
--- a/Engine/Source/Runtime/Renderer/Private/PostProcess/DiaphragmDOF.cpp
+++ b/Engine/Source/Runtime/Renderer/Private/PostProcess/DiaphragmDOF.cpp
@@ -2141,7 +2141,7 @@ FRDGTextureRef DiaphragmDOF::AddPasses(
 					GetEventName(ConvolutionSettings.LayerProcessing),
 					GetEventName(ConvolutionSettings.QualityConfig),
 					GetEventName(ConvolutionSettings.BokehSimulation),
-					PermutationVector.Get<FDDOFGatherRingCountDim>(),
+					int32(PermutationVector.Get<FDDOFGatherRingCountDim>()),
 					PermutationVector.Get<FDDOFClampInputUVDim>() ? TEXT(" ClampUV") : TEXT(""),
 					PermutationVector.Get<FDDOFRGBColorBufferDim>() ? TEXT(" R11G11B10") : TEXT(""),
 					GatheringViewSize.X, GatheringViewSize.Y),
