Removed redundant search code, fog now use scene->sunlight
Fog tests are green

diff --git a/Engine/Source/Runtime/Renderer/Private/FogRendering.cpp b/Engine/Source/Runtime/Renderer/Private/FogRendering.cpp
--- a/Engine/Source/Runtime/Renderer/Private/FogRendering.cpp
+++ b/Engine/Source/Runtime/Renderer/Private/FogRendering.cpp
@@ -335,26 +335,13 @@ void FSceneRenderer::InitFogConstants()
 
 				View.DirectionalInscatteringExponent = FogInfo.DirectionalInscatteringExponent;
 				View.DirectionalInscatteringStartDistance = FogInfo.DirectionalInscatteringStartDistance;
-				View.bUseDirectionalInscattering = false;
 				View.InscatteringLightDirection = FVector(0);
-
-				for (TSparseArray<FLightSceneInfoCompact>::TConstIterator It(Scene->Lights); It; ++It)
+				if (Scene->SunLight)
 				{
-					const FLightSceneInfoCompact& LightInfo = *It;
-
-					// This will find the first directional light that is set to be used as an atmospheric sun light of sufficient brightness.
-					// If you have more than one directional light with these properties then all subsequent lights will be ignored.
-					if (LightInfo.LightSceneInfo->Proxy->GetLightType() == LightType_Directional
-						&& LightInfo.LightSceneInfo->Proxy->IsUsedAsAtmosphereSunLight()
-						&& LightInfo.LightSceneInfo->Proxy->GetColor().ComputeLuminance() > KINDA_SMALL_NUMBER
-						&& FogInfo.DirectionalInscatteringColor.ComputeLuminance() > KINDA_SMALL_NUMBER)
-					{
-						View.InscatteringLightDirection = -LightInfo.LightSceneInfo->Proxy->GetDirection();
-						View.bUseDirectionalInscattering = true;
-						View.DirectionalInscatteringColor = FogInfo.DirectionalInscatteringColor * LightInfo.LightSceneInfo->Proxy->GetColor().ComputeLuminance();
-						break;
-					}
+					View.InscatteringLightDirection = -Scene->SunLight->Proxy->GetDirection();
+					View.DirectionalInscatteringColor = FogInfo.DirectionalInscatteringColor * Scene->SunLight->Proxy->GetColor().ComputeLuminance();
 				}
+				View.bUseDirectionalInscattering = Scene->SunLight!=nullptr;
 			}
 		}
 	}
