Added mechanism to compute improved sun light and atmosphere intereactions:
 - sun light color is now affected by transmittance hwen traveling in the atmosphere
 - sun light disk is now has a radiance matching illuminance at ground level (using inverse light transport from ground to outer space and acocundting for solid angle)

diff --git a/Engine/Shaders/Private/AtmosphereCommon.ush b/Engine/Shaders/Private/AtmosphereCommon.ush
--- a/Engine/Shaders/Private/AtmosphereCommon.ush
+++ b/Engine/Shaders/Private/AtmosphereCommon.ush
@@ -544,8 +544,8 @@ float3 GetSunColor(float3 X, float T, float3 V, float3 S, float Radius, float Mu
     }
 	else 
 	{
-		float SunIntensity = step(cos(PI * View.AtmosphericFogSunDiscScale / 180.0), dot(V, S)); // Lsun
-        return TransmittanceValue * SunIntensity; // Eq (9)
+		float SunDiskMask = step(cos(PI * View.AtmosphericFogSunDiscScale / 180.0), dot(V, S));
+        return TransmittanceValue * SunDiskMask * View.AtmosphericFogSunDiscLuminance; // Eq (9)
     }
 }
 
diff --git a/Engine/Source/Runtime/Engine/Classes/Atmosphere/AtmosphericFogComponent.h b/Engine/Source/Runtime/Engine/Classes/Atmosphere/AtmosphericFogComponent.h
--- a/Engine/Source/Runtime/Engine/Classes/Atmosphere/AtmosphericFogComponent.h
+++ b/Engine/Source/Runtime/Engine/Classes/Atmosphere/AtmosphericFogComponent.h
@@ -135,7 +135,7 @@ class UAtmosphericFogComponent : public USceneComponent
 	UPROPERTY(EditAnywhere, BlueprintReadWrite, interp, Category=Atmosphere)
 	float StartDistance;
 
-	/** Distance offset, in km (to handle large distance) */
+	/** Sun half apex angle in degree, see https://en.wikipedia.org/wiki/Solid_angle */
 	UPROPERTY(EditAnywhere, BlueprintReadWrite, interp, Category = Atmosphere)
 	float SunDiscScale;
 
diff --git a/Engine/Source/Runtime/Engine/Private/Components/DirectionalLightComponent.cpp b/Engine/Source/Runtime/Engine/Private/Components/DirectionalLightComponent.cpp
--- a/Engine/Source/Runtime/Engine/Private/Components/DirectionalLightComponent.cpp
+++ b/Engine/Source/Runtime/Engine/Private/Components/DirectionalLightComponent.cpp
@@ -65,6 +65,16 @@ public:
 	 */
 	FVector LightShaftOverrideDirection;
 
+	/**
+	 * The atmosphere transmittance to apply on the illuminance
+	 */
+	FLinearColor AtmosphereTransmittanceFactor;
+
+	/**
+	 * The luminance of the sun disk in space (function of the sun illuminance and solid angle)
+	 */
+	FLinearColor SunDiscOuterSpaceLuminance;
+
 	/** 
 	 * Radius of the whole scene dynamic shadow centered on the viewer, which replaces the precomputed shadows based on distance from the camera.  
 	 * A Radius of 0 disables the dynamic shadow.
@@ -103,13 +113,15 @@ public:
 	float TraceDistance;
 
 	/** Initialization constructor. */
-	FDirectionalLightSceneProxy(const UDirectionalLightComponent* Component):
+	FDirectionalLightSceneProxy(const UDirectionalLightComponent* Component) :
 		FLightSceneProxy(Component),
 		bEnableLightShaftOcclusion(Component->bEnableLightShaftOcclusion),
 		bUseInsetShadowsForMovableObjects(Component->bUseInsetShadowsForMovableObjects),
 		OcclusionMaskDarkness(Component->OcclusionMaskDarkness),
 		OcclusionDepthRange(Component->OcclusionDepthRange),
 		LightShaftOverrideDirection(Component->LightShaftOverrideDirection),
+		AtmosphereTransmittanceFactor(FLinearColor::White),
+		SunDiscOuterSpaceLuminance(FLinearColor::White),
 		DynamicShadowCascades(Component->DynamicShadowCascades > 0 ? Component->DynamicShadowCascades : 0),
 		CascadeDistributionExponent(Component->CascadeDistributionExponent),
 		CascadeTransitionFraction(Component->CascadeTransitionFraction),
@@ -166,7 +178,7 @@ public:
 	{
 		LightParameters.Position = FVector::ZeroVector;
 		LightParameters.InvRadius = 0.0f;
-		LightParameters.Color = FVector(GetColor());
+		LightParameters.Color = FVector(GetColor() * AtmosphereTransmittanceFactor); 
 		LightParameters.FalloffExponent = 0.0f;
 
 		LightParameters.Direction = -GetDirection();
@@ -346,6 +358,17 @@ public:
 		return DoesPlatformSupportDistanceFieldShadowing(GShaderPlatformForFeatureLevel[InFeatureLevel]) && (bCreateWithCSM || bCreateWithoutCSM);
 	}
 
+	virtual void SetAtmosphereRelatedProperties(FLinearColor TransmittanceFactor, FLinearColor SunOuterSpaceLuminance) override
+	{
+		AtmosphereTransmittanceFactor = TransmittanceFactor;
+		SunDiscOuterSpaceLuminance = SunOuterSpaceLuminance;
+	}
+
+	virtual FLinearColor GetOuterSpaceLuminance() const override
+	{ 
+		return SunDiscOuterSpaceLuminance;
+	}
+
 private:
 
 	float GetEffectiveWholeSceneDynamicShadowRadius(bool bPrecomputedLightingIsValid) const
diff --git a/Engine/Source/Runtime/Engine/Public/SceneManagement.h b/Engine/Source/Runtime/Engine/Public/SceneManagement.h
--- a/Engine/Source/Runtime/Engine/Public/SceneManagement.h
+++ b/Engine/Source/Runtime/Engine/Public/SceneManagement.h
@@ -1293,6 +1293,10 @@ public:
 	virtual float GetMaxDrawDistance() const { return 0.0f; }
 	virtual float GetFadeRange() const { return 0.0f; }
 
+	// Atmosphere / Fog related functions
+	virtual void SetAtmosphereRelatedProperties(FLinearColor TransmittanceFactor, FLinearColor SunOuterSpaceLuminance) {}
+	virtual FLinearColor GetOuterSpaceLuminance() const { return FLinearColor::White; }
+
 protected:
 
 	friend class FScene;
diff --git a/Engine/Source/Runtime/Engine/Public/SceneView.h b/Engine/Source/Runtime/Engine/Public/SceneView.h
--- a/Engine/Source/Runtime/Engine/Public/SceneView.h
+++ b/Engine/Source/Runtime/Engine/Public/SceneView.h
@@ -666,6 +666,7 @@ enum ETranslucencyVolumeCascade
 	VIEW_UNIFORM_BUFFER_MEMBER_EX(float, AtmosphericFogStartDistance, EShaderPrecisionModifier::Half) \
 	VIEW_UNIFORM_BUFFER_MEMBER_EX(float, AtmosphericFogDistanceOffset, EShaderPrecisionModifier::Half) \
 	VIEW_UNIFORM_BUFFER_MEMBER_EX(float, AtmosphericFogSunDiscScale, EShaderPrecisionModifier::Half) \
+	VIEW_UNIFORM_BUFFER_MEMBER_EX(FLinearColor, AtmosphericFogSunDiscLuminance, EShaderPrecisionModifier::Half) \
 	VIEW_UNIFORM_BUFFER_MEMBER(uint32, AtmosphericFogRenderMask) \
 	VIEW_UNIFORM_BUFFER_MEMBER(uint32, AtmosphericFogInscatterAltitudeSampleNum) \
 	VIEW_UNIFORM_BUFFER_MEMBER(FLinearColor, AtmosphericFogSunColor) \
diff --git a/Engine/Source/Runtime/Renderer/Private/AtmosphereRendering.cpp b/Engine/Source/Runtime/Renderer/Private/AtmosphereRendering.cpp
--- a/Engine/Source/Runtime/Renderer/Private/AtmosphereRendering.cpp
+++ b/Engine/Source/Runtime/Renderer/Private/AtmosphereRendering.cpp
@@ -24,6 +24,7 @@
 #include "PostProcess/SceneFilterRendering.h"
 #include "PipelineStateCache.h"
 DECLARE_GPU_STAT(Atmosphere);
+DECLARE_GPU_STAT(AtmospherePreCompute);
 
 class FAtmosphereShaderPrecomputeTextureParameters
 {
@@ -1102,6 +1103,75 @@ void FAtmosphericFogSceneInfo::GetLayerValue(int Layer, float& AtmosphereR, FVec
 	DhdH = FVector4(DMin, DMax, DMinP, DMaxP);
 }
 
+FLinearColor FAtmosphericFogSceneInfo::GetTransmittance(const FVector& SunDirection) const
+{
+	// Code from HLSL here as lambda function. It will simulate atmosphere transmittance according to the current sky hardcoded parameterization.
+	// This will change in the future when the sky parameterization and workflow/ui will be updated.
+
+	const int TransmittanceIntegralSamples = 10;
+	const float RadiusLimit = RadiusAtmosphere;
+	const FVector BetaRayleighScattering(5.8e-3f, 1.35e-2f, 3.31e-2f);
+	const FVector BetaMieScattering = FVector(4e-3f, 4e-3f, 4e-3f);
+	const float HeightScaleMie = 1.2f;
+	const float BetaRatio = 0.9f;
+	const FVector BetaMieExtinction = BetaMieScattering / BetaRatio;
+	const float AtmosphericFogHeightScaleRayleigh = RHeight;
+
+	auto Limit = [&](float Radius, float Mu)
+	{
+		float Dout = -Radius * Mu + FMath::Sqrt(Radius * Radius * (Mu * Mu - 1.0) + RadiusLimit * RadiusLimit);
+		float Delta2 = Radius * Radius * (Mu * Mu - 1.0) + RadiusGround * RadiusGround;
+		if (Delta2 >= 0.0)
+		{
+			float Din = -Radius * Mu - FMath::Sqrt(Delta2);
+			if (Din >= 0.0)
+			{
+				Dout = FMath::Min(Dout, Din);
+			}
+		}
+		return Dout;
+	};
+
+	auto OpticalDepth = [&](float H, float Radius, float Mu)
+	{
+		float Result = 0.0;
+		float Dx = Limit(Radius, Mu) / float(TransmittanceIntegralSamples);
+		float Xi = 0.0;
+		float Yi = FMath::Exp(-(Radius - RadiusGround) / H);
+		for (int I = 1; I <= TransmittanceIntegralSamples; ++I)
+		{
+			float Xj = float(I) * Dx;
+			float Yj = FMath::Exp(-(FMath::Sqrt(Radius * Radius + Xj * Xj + 2.0 * Xj * Radius * Mu) - RadiusGround) / H);
+			Result += (Yi + Yj) / 2.0 * Dx;
+			Xi = Xj;
+			Yi = Yj;
+		}
+		return Mu < -FMath::Sqrt(1.0 - (RadiusGround / Radius) * (RadiusGround / Radius)) ? 1e9 : Result;
+	};
+
+	// GetTransmittanceRMuS linear version, assuming we are always close to the ground
+	const float Radius = RadiusGround;
+	float Mu = SunDirection.Z;
+	FVector OpticalDepthRGB = BetaRayleighScattering * OpticalDepth(AtmosphericFogHeightScaleRayleigh, Radius, Mu) + BetaMieExtinction * OpticalDepth(HeightScaleMie, Radius, Mu);
+	OpticalDepthRGB.ComponentMax(FVector(ForceInitToZero));
+
+	return FLinearColor(FMath::Exp(-OpticalDepthRGB.X), FMath::Exp(-OpticalDepthRGB.Y), FMath::Exp(-OpticalDepthRGB.Z));
+}
+
+void FAtmosphericFogSceneInfo::PrepareSunLightProxy(FLightSceneInfo& SunLight) const
+{
+	// See explanation in https://media.contentapi.ea.com/content/dam/eacom/frostbite/files/s2016-pbs-frostbite-sky-clouds-new.pdf page 26
+	FLinearColor TransmittanceTowardSun = GetTransmittance(-SunLight.Proxy->GetDirection());
+
+	FLinearColor SunZenithIlluminance = SunLight.Proxy->GetColor();
+	FLinearColor SunOuterSpaceIlluminance = SunZenithIlluminance / TransmittanceAtZenith;
+
+	const float SunSolidAngle = 2.0f * PI * (1.0f - FMath::Cos(0.5 * SunHalfApexAngle));	// Solid angle from aperture https://en.wikipedia.org/wiki/Solid_angle
+	FLinearColor SunDiskOuterSpaceLuminance = SunOuterSpaceIlluminance / SunSolidAngle;		// approximation  
+
+	SunLight.Proxy->SetAtmosphereRelatedProperties(TransmittanceTowardSun / TransmittanceAtZenith, SunDiskOuterSpaceLuminance);
+}
+
 void FAtmosphericFogSceneInfo::RenderAtmosphereShaders(FRHICommandList& RHICmdList, FGraphicsPipelineStateInitializer& GraphicsPSOInit, const FViewInfo& View, const FIntRect& ViewRect)
 {
 	const ERHIFeatureLevel::Type ViewFeatureLevel = View.GetFeatureLevel();
@@ -1615,6 +1685,7 @@ void FAtmosphericFogSceneInfo::Read3DPixelsPtr(FRHICommandListImmediate& RHICmdL
 
 void FAtmosphericFogSceneInfo::PrecomputeTextures(FRHICommandListImmediate& RHICmdList, const FViewInfo* View, FSceneViewFamily* ViewFamily)
 {
+	SCOPED_GPU_STAT(RHICmdList, AtmospherePreCompute);
 	check(Component != NULL);
 	if (AtmosphereTextures == NULL)
 	{
@@ -1743,7 +1814,7 @@ FAtmosphericFogSceneInfo::FAtmosphericFogSceneInfo(UAtmosphericFogComponent* InC
 	, RHeight(InComponent->PrecomputeParams.DensityHeight * InComponent->PrecomputeParams.DensityHeight * InComponent->PrecomputeParams.DensityHeight * 64.f)
 	, StartDistance(InComponent->StartDistance)
 	, DistanceOffset(InComponent->DistanceOffset)
-	, SunDiscScale(InComponent->SunDiscScale)
+	, SunHalfApexAngle(InComponent->SunDiscScale)
 	, RenderFlag(EAtmosphereRenderFlag::E_EnableAll)
 	, InscatterAltitudeSampleNum(InComponent->PrecomputeParams.InscatterAltitudeSampleNum)
 
@@ -1758,6 +1829,7 @@ FAtmosphericFogSceneInfo::FAtmosphericFogSceneInfo(UAtmosphericFogComponent* InC
 	, AtmoshpereOrder(2)
 	, AtmosphereTextures(NULL)
 #endif
+	, TransmittanceAtZenith(GetTransmittance(FVector(0.0f, 0.0f, 1.0f)))
 {
 	StartDistance *= DistanceScale * 0.00001f; // Convert to km in Atmospheric fog shader
 	// DistanceOffset is in km, no need to change...
diff --git a/Engine/Source/Runtime/Renderer/Private/AtmosphereRendering.h b/Engine/Source/Runtime/Renderer/Private/AtmosphereRendering.h
--- a/Engine/Source/Runtime/Renderer/Private/AtmosphereRendering.h
+++ b/Engine/Source/Runtime/Renderer/Private/AtmosphereRendering.h
@@ -50,7 +50,7 @@ public:
 	float RHeight;
 	float StartDistance;
 	float DistanceOffset;
-	float SunDiscScale;
+	float SunHalfApexAngle;
 	FLinearColor DefaultSunColor;
 	FVector DefaultSunDirection;
 	uint32 RenderFlag;
@@ -80,6 +80,8 @@ public:
 	explicit FAtmosphericFogSceneInfo(UAtmosphericFogComponent* InComponent, const FScene* InScene);
 	~FAtmosphericFogSceneInfo();
 
+	void PrepareSunLightProxy(FLightSceneInfo& SunLight) const;
+
 #if WITH_EDITOR
 	void PrecomputeTextures(FRHICommandListImmediate& RHICmdList, const FViewInfo* View, FSceneViewFamily* ViewFamily);
 	void StartPrecompute();
@@ -95,6 +97,11 @@ private:
 	void ReadPixelsPtr(FRHICommandListImmediate& RHICmdList, TRefCountPtr<IPooledRenderTarget> RenderTarget, FColor* OutData, FIntRect InRect);
 	void Read3DPixelsPtr(FRHICommandListImmediate& RHICmdList, TRefCountPtr<IPooledRenderTarget> RenderTarget, FFloat16Color* OutData, FIntRect InRect, FIntPoint InZMinMax);
 #endif
+
+private:
+	const FLinearColor TransmittanceAtZenith;
+
+	FLinearColor GetTransmittance(const FVector& SunDirection) const;
 };
 
 bool ShouldRenderAtmosphere(const FSceneViewFamily& Family);
diff --git a/Engine/Source/Runtime/Renderer/Private/DeferredShadingRenderer.cpp b/Engine/Source/Runtime/Renderer/Private/DeferredShadingRenderer.cpp
--- a/Engine/Source/Runtime/Renderer/Private/DeferredShadingRenderer.cpp
+++ b/Engine/Source/Runtime/Renderer/Private/DeferredShadingRenderer.cpp
@@ -808,6 +808,12 @@ void FDeferredShadingSceneRenderer::Render(FRHICommandListImmediate& RHICmdList)
 
 	PrepareViewRectsForRendering();
 
+	if (Scene->SunLight && Scene->HasAtmosphericFog())
+	{
+		// Only one atmospheric light at one time.
+		Scene->GetAtmosphericFogSceneInfo()->PrepareSunLightProxy(*Scene->SunLight);
+	}
+
 	SCOPED_NAMED_EVENT(FDeferredShadingSceneRenderer_Render, FColor::Emerald);
 
 #if RHI_RAYTRACING
diff --git a/Engine/Source/Runtime/Renderer/Private/SceneRendering.cpp b/Engine/Source/Runtime/Renderer/Private/SceneRendering.cpp
--- a/Engine/Source/Runtime/Renderer/Private/SceneRendering.cpp
+++ b/Engine/Source/Runtime/Renderer/Private/SceneRendering.cpp
@@ -1136,7 +1136,8 @@ void FViewInfo::SetupUniformBufferParameters(
 			ViewUniformShaderParameters.AtmosphericFogHeightScaleRayleigh = Scene->AtmosphericFog->RHeight;
 			ViewUniformShaderParameters.AtmosphericFogStartDistance = Scene->AtmosphericFog->StartDistance;
 			ViewUniformShaderParameters.AtmosphericFogDistanceOffset = Scene->AtmosphericFog->DistanceOffset;
-			ViewUniformShaderParameters.AtmosphericFogSunDiscScale = Scene->AtmosphericFog->SunDiscScale;
+			ViewUniformShaderParameters.AtmosphericFogSunDiscScale = Scene->AtmosphericFog->SunHalfApexAngle;
+			ViewUniformShaderParameters.AtmosphericFogSunDiscLuminance = Scene->SunLight ? Scene->SunLight->Proxy->GetOuterSpaceLuminance() : FLinearColor::White;
 			ViewUniformShaderParameters.AtmosphericFogSunColor = Scene->SunLight ? Scene->SunLight->Proxy->GetColor() : Scene->AtmosphericFog->DefaultSunColor;
 			ViewUniformShaderParameters.AtmosphericFogSunDirection = Scene->SunLight ? -Scene->SunLight->Proxy->GetDirection() : -Scene->AtmosphericFog->DefaultSunDirection;
 			ViewUniformShaderParameters.AtmosphericFogRenderMask = Scene->AtmosphericFog->RenderFlag & (EAtmosphereRenderFlag::E_DisableGroundScattering | EAtmosphereRenderFlag::E_DisableSunDisk);
@@ -1155,6 +1156,7 @@ void FViewInfo::SetupUniformBufferParameters(
 			ViewUniformShaderParameters.AtmosphericFogStartDistance = FLT_MAX;
 			ViewUniformShaderParameters.AtmosphericFogDistanceOffset = 0.f;
 			ViewUniformShaderParameters.AtmosphericFogSunDiscScale = 1.f;
+			ViewUniformShaderParameters.AtmosphericFogSunDiscLuminance = Scene->SunLight ? Scene->SunLight->Proxy->GetOuterSpaceLuminance() : FLinearColor::Black;
 			//Added check so atmospheric light color and vector can use a directional light without needing an atmospheric fog actor in the scene
 			ViewUniformShaderParameters.AtmosphericFogSunColor = Scene->SunLight ? Scene->SunLight->Proxy->GetColor() : FLinearColor::Black;
 			ViewUniformShaderParameters.AtmosphericFogSunDirection = Scene->SunLight ? -Scene->SunLight->Proxy->GetDirection() : FVector::ZeroVector;
@@ -1176,6 +1178,7 @@ void FViewInfo::SetupUniformBufferParameters(
 		ViewUniformShaderParameters.AtmosphericFogStartDistance = FLT_MAX;
 		ViewUniformShaderParameters.AtmosphericFogDistanceOffset = 0.f;
 		ViewUniformShaderParameters.AtmosphericFogSunDiscScale = 1.f;
+		ViewUniformShaderParameters.AtmosphericFogSunDiscLuminance = FLinearColor::Black;
 		ViewUniformShaderParameters.AtmosphericFogSunColor = FLinearColor::Black;
 		ViewUniformShaderParameters.AtmosphericFogSunDirection = FVector::ZeroVector;
 		ViewUniformShaderParameters.AtmosphericFogRenderMask = EAtmosphereRenderFlag::E_EnableAll;
