Attempt to fix non-unity build.

diff --git a/Engine/Source/Runtime/Renderer/Private/AtmosphereRendering.cpp b/Engine/Source/Runtime/Renderer/Private/AtmosphereRendering.cpp
--- a/Engine/Source/Runtime/Renderer/Private/AtmosphereRendering.cpp
+++ b/Engine/Source/Runtime/Renderer/Private/AtmosphereRendering.cpp
@@ -23,6 +23,7 @@
 #include "AtmosphereTextures.h"
 #include "PostProcess/SceneFilterRendering.h"
 #include "PipelineStateCache.h"
+#include "LightSceneInfo.h"
 DECLARE_GPU_STAT(Atmosphere);
 DECLARE_GPU_STAT(AtmospherePreCompute);
 
@@ -1166,7 +1167,7 @@ void FAtmosphericFogSceneInfo::PrepareSunLightProxy(FLightSceneInfo& SunLight) c
 	FLinearColor SunZenithIlluminance = SunLight.Proxy->GetColor();
 	FLinearColor SunOuterSpaceIlluminance = SunZenithIlluminance / TransmittanceAtZenith;
 
-	const float SunSolidAngle = 2.0f * PI * (1.0f - FMath::Cos(0.5 * SunHalfApexAngle));	// Solid angle from aperture https://en.wikipedia.org/wiki/Solid_angle
+	const float SunSolidAngle = 2.0f * PI * (1.0f - FMath::Cos(0.5f * SunHalfApexAngle));	// Solid angle from aperture https://en.wikipedia.org/wiki/Solid_angle
 	FLinearColor SunDiskOuterSpaceLuminance = SunOuterSpaceIlluminance / SunSolidAngle;		// approximation  
 
 	SunLight.Proxy->SetAtmosphereRelatedProperties(TransmittanceTowardSun / TransmittanceAtZenith, SunDiskOuterSpaceLuminance);
diff --git a/Engine/Source/Runtime/Renderer/Private/AtmosphereRendering.h b/Engine/Source/Runtime/Renderer/Private/AtmosphereRendering.h
--- a/Engine/Source/Runtime/Renderer/Private/AtmosphereRendering.h
+++ b/Engine/Source/Runtime/Renderer/Private/AtmosphereRendering.h
@@ -11,6 +11,7 @@
 #include "Serialization/BulkData.h"
 #include "RendererInterface.h"
 
+class FLightSceneInfo;
 class FScene;
 class FSceneViewFamily;
 class FShader;
