Fixed shared samplers in MotionBlur and Subsurface shaders for GL backend.

#rnx

diff --git a/Engine/Shaders/Private/PostProcessMotionBlur.usf b/Engine/Shaders/Private/PostProcessMotionBlur.usf
--- a/Engine/Shaders/Private/PostProcessMotionBlur.usf
+++ b/Engine/Shaders/Private/PostProcessMotionBlur.usf
@@ -38,6 +38,21 @@ Texture2D DepthTexture;
 
 SamplerState ColorSampler;
 SamplerState VelocitySampler;
+SamplerState VelocityTileSampler;
+SamplerState VelocityFlatSampler;
+SamplerState DepthSampler;
+
+#if SUPPORTS_INDEPENDENT_SAMPLERS
+	#define SharedVelocitySampler VelocitySampler
+	#define SharedVelocityTileSampler VelocitySampler
+	#define SharedVelocityFlatSampler VelocitySampler
+	#define SharedDepthSampler VelocitySampler
+#else
+	#define SharedVelocitySampler VelocitySampler
+	#define SharedVelocityTileSampler VelocityTileSampler
+	#define SharedVelocityFlatSampler VelocityFlatSampler
+	#define SharedDepthSampler DepthSampler
+#endif
 
 // ------------------------------------------
 
@@ -119,7 +134,7 @@ float4 MainMotionBlurCommon(float2 ColorUV, float2 ColorPixelPos)
 	// Velocity tile UV originates at [0,0]; only need to clamp max.
 	VelocityTileUV = min(VelocityTileUV, VelocityTile_UVViewportBilinearMax);
 
-	float4 MinMaxVelocity = VelocityTileTexture.SampleLevel(VelocitySampler, VelocityTileUV, 0);
+	float4 MinMaxVelocity = VelocityTileTexture.SampleLevel(SharedVelocityTileSampler, VelocityTileUV, 0);
 
 	float2 MinVelocityPixels = MinMaxVelocity.xy * MotionBlur_VelocityScale;
 	float2 MaxVelocityPixels = MinMaxVelocity.zw * MotionBlur_VelocityScale;
@@ -197,7 +212,7 @@ float4 MainMotionBlurCommon(float2 ColorUV, float2 ColorPixelPos)
 	}
 	else
 	{
-		float3 CenterVelocityDepth = VelocityFlatTexture.SampleLevel(VelocitySampler, VelocityUV, 0).xyz;
+		float3 CenterVelocityDepth = VelocityFlatTexture.SampleLevel(SharedVelocityFlatSampler, VelocityUV, 0).xyz;
 		float  CenterDepth = CenterVelocityDepth.z;
 		float  CenterVelocityLength = GetVelocityLengthPixels(CenterVelocityDepth.xy);
 
@@ -226,7 +241,7 @@ float4 MainMotionBlurCommon(float2 ColorUV, float2 ColorPixelPos)
 			for (uint j = 0; j < 2; j++)
 			{
 				float3 SampleVelocityDepth = VelocityFlatTexture.SampleLevel(
-					VelocitySampler, SampleUV[j] * ColorToVelocity_Scale + ColorToVelocity_Bias, 0).xyz;
+					SharedVelocityFlatSampler, SampleUV[j] * ColorToVelocity_Scale + ColorToVelocity_Bias, 0).xyz;
 
 				SampleColor[j] = ColorTexture.SampleLevel(ColorSampler, SampleUV[j], 0);
 				SampleDepth[j] = SampleVelocityDepth.z;
@@ -373,7 +388,7 @@ void VisualizeMotionBlurPS(
 
 	// World aligned checkerboards
 	{
-		float DeviceZ = DepthTexture.SampleLevel(VelocitySampler, UV, 0).r;
+		float DeviceZ = DepthTexture.SampleLevel(SharedDepthSampler, UV, 0).r;
 		float SceneDepth = ConvertFromDeviceZ(DeviceZ);
 
 		float3 ScreenVector = mul(float4(ScreenSpacePos, 1, 0), View.ScreenToWorld).xyz;
@@ -386,18 +401,18 @@ void VisualizeMotionBlurPS(
 		OutColor = float4(lerp(WorldCheckerboard, float3(0,0,0), 0.7f), 1);
 	}
 
-	bool bSelectorOpaqueAtTileCenter = VelocityTexture.SampleLevel(VelocitySampler, UVAtTileCenter, 0).x > 0;
-	bool bSelectorOpaque = VelocityTexture.SampleLevel(VelocitySampler, UV, 0).x > 0;
+	bool bSelectorOpaqueAtTileCenter = VelocityTexture.SampleLevel(SharedVelocitySampler, UVAtTileCenter, 0).x > 0;
+	bool bSelectorOpaque = VelocityTexture.SampleLevel(SharedVelocitySampler, UV, 0).x > 0;
 
 	// relative, in screen space -1...1 -1..1, can be even outside of that range, points into the movement direction
 	float2 VelocityAtTileCenter;
 	{
-		VelocityAtTileCenter = DecodeVelocityFromTexture(VelocityTexture.SampleLevel(VelocitySampler, UVAtTileCenter, 0).xy);
+		VelocityAtTileCenter = DecodeVelocityFromTexture(VelocityTexture.SampleLevel(SharedVelocitySampler, UVAtTileCenter, 0).xy);
 
 		// reconstruct from camera motion if we don't have velocity data
 		if(!bSelectorOpaqueAtTileCenter)
 		{
-			float DeviceZ = DepthTexture.SampleLevel(VelocitySampler, UVAtTileCenter, 0).r;
+			float DeviceZ = DepthTexture.SampleLevel(SharedDepthSampler, UVAtTileCenter, 0).r;
 
 			// UV is per tile
 			float4 ThisClip = float4(ScreenPosAtTileCenter, DeviceZ, 1);
diff --git a/Engine/Shaders/Private/PostProcessSubsurface.usf b/Engine/Shaders/Private/PostProcessSubsurface.usf
--- a/Engine/Shaders/Private/PostProcessSubsurface.usf
+++ b/Engine/Shaders/Private/PostProcessSubsurface.usf
@@ -20,7 +20,16 @@ float4 SubsurfaceParams;
 Texture2D SubsurfaceInput0_Texture;
 Texture2D SubsurfaceInput1_Texture;
 
-SamplerState SubsurfaceSampler; 
+SamplerState SubsurfaceSampler0;
+SamplerState SubsurfaceSampler1;
+
+#if SUPPORTS_INDEPENDENT_SAMPLERS
+	#define SharedSubsurfaceSampler0 SubsurfaceSampler0
+	#define SharedSubsurfaceSampler1 SubsurfaceSampler0
+#else
+	#define SharedSubsurfaceSampler0 SubsurfaceSampler0
+	#define SharedSubsurfaceSampler1 SubsurfaceSampler1
+#endif
 
 SCREEN_PASS_TEXTURE_VIEWPORT(SubsurfaceInput0)
 SCREEN_PASS_TEXTURE_VIEWPORT(SubsurfaceInput1)
@@ -131,7 +140,7 @@ float4 GetSceneColor(float2 BufferUV)
 {
 	BufferUV = clamp(BufferUV, SubsurfaceInput0_UVViewportBilinearMin, SubsurfaceInput0_UVViewportBilinearMax);
 
-	return Texture2DSampleLevel(SubsurfaceInput0_Texture, SubsurfaceSampler, BufferUV, 0);
+	return Texture2DSampleLevel(SubsurfaceInput0_Texture, SharedSubsurfaceSampler0, BufferUV, 0);
 }
 
 // from [REDACTED_PROFILE_URL]/separable-sss/tree/master/Demo
@@ -179,7 +188,7 @@ void VisualizePS(in noperspective float4 UVAndScreenPos : TEXCOORD0, float4 SvPo
 {
 	float2 UV = UVAndScreenPos.xy;
 
-	OutColor = Texture2DSample(SubsurfaceInput0_Texture, SubsurfaceSampler, UV);
+	OutColor = Texture2DSample(SubsurfaceInput0_Texture, SharedSubsurfaceSampler0, UV);
 
 	int2 PixelPos = (int2)SvPosition.xy;
 
@@ -259,10 +268,10 @@ half3 LookupSceneColor(float2 SceneUV, int2 PixelOffset)
 #if ES2_PROFILE && COMPILER_GLSL_ES2
 	// slower but always works
 	// to prevent "error: Texture offset not supported on GLSL ES"
-	return Texture2DSample(SubsurfaceInput0_Texture, SubsurfaceSampler, SceneUV + PixelOffset * SubsurfaceInput0_ExtentInverse).rgb;
+	return Texture2DSample(SubsurfaceInput0_Texture, SharedSubsurfaceSampler0, SceneUV + PixelOffset * SubsurfaceInput0_ExtentInverse).rgb;
 #else
 	// faster
-	return SubsurfaceInput0_Texture.SampleLevel(SubsurfaceSampler, SceneUV, 0, PixelOffset).rgb;
+	return SubsurfaceInput0_Texture.SampleLevel(SharedSubsurfaceSampler0, SceneUV, 0, PixelOffset).rgb;
 #endif
 }
  
@@ -279,7 +288,7 @@ SDiffuseAndSpecular ReconstructLighting(float2 UVSceneColor, uint ReconstructMet
 
 	  // todo: We could alternate the diagonal with TemporalAA or even only only 1 sample for low spec or 4 for high spec
 
-	  float3 Quant0 = SubsurfaceInput0_Texture.SampleLevel(SubsurfaceSampler, UVSceneColor, 0).rgb;
+	  float3 Quant0 = SubsurfaceInput0_Texture.SampleLevel(SharedSubsurfaceSampler0, UVSceneColor, 0).rgb;
 
 	  // todo: expose as scalability setting (can be evaluate best without TemporalAA)
 	  // 0:fast but pattern can appear, 1:better, 2: even better, 3: best but expensive
@@ -333,7 +342,7 @@ SDiffuseAndSpecular ReconstructLighting(float2 UVSceneColor, uint ReconstructMet
 #else // SUBSURFACE_PROFILE_CHECKERBOARD
 	{
 	  // If we're not doing checkerboard encoding, we just need to read a single pixel and decode (combined diffuse/spec in RGB)
-	  float4 CenterSample = SubsurfaceInput0_Texture.SampleLevel(SubsurfaceSampler, UVSceneColor, 0);
+	  float4 CenterSample = SubsurfaceInput0_Texture.SampleLevel(SharedSubsurfaceSampler0, UVSceneColor, 0);
 	  float3 CombinedColor = CenterSample.rgb;
 	  float DiffuseLuminance = CenterSample.a;
   
@@ -439,7 +448,7 @@ void SubsurfaceRecombinePS(noperspective float4 UVAndScreenPos : TEXCOORD0, floa
 	
 	if (!UseSubsurfaceProfile(ScreenSpaceData.GBuffer.ShadingModelID))
 	{
-		OutColor = Texture2DSample(SubsurfaceInput0_Texture, SubsurfaceSampler, BufferUV);
+		OutColor = Texture2DSample(SubsurfaceInput0_Texture, SharedSubsurfaceSampler0, BufferUV);
 		return;
 	}
 
@@ -454,7 +463,7 @@ void SubsurfaceRecombinePS(noperspective float4 UVAndScreenPos : TEXCOORD0, floa
 #endif // SUBSURFACE_HALF_RES
 
 	{
-		float4 SSSColorWithAlpha = Texture2DSample(SubsurfaceInput1_Texture, SubsurfaceSampler, BufferUV);
+		float4 SSSColorWithAlpha = Texture2DSample(SubsurfaceInput1_Texture, SharedSubsurfaceSampler1, BufferUV);
 
 		// renormalize to dilate RGB to fix half res upsampling artifacts
 		SSSColor = SSSColorWithAlpha.rgb / max(SSSColorWithAlpha.a, 0.00001f);
@@ -486,5 +495,5 @@ void SubsurfaceRecombinePS(noperspective float4 UVAndScreenPos : TEXCOORD0, floa
 
 void SubsurfaceViewportCopyPS(noperspective float4 InUV : TEXCOORD0, out float4 OutColor : SV_Target0)
 {
-	OutColor = Texture2DSample(SubsurfaceInput0_Texture, SubsurfaceSampler, InUV.xy);
+	OutColor = Texture2DSample(SubsurfaceInput0_Texture, SharedSubsurfaceSampler0, InUV.xy);
 }
diff --git a/Engine/Source/Runtime/Renderer/Private/PostProcess/PostProcessMotionBlur.cpp b/Engine/Source/Runtime/Renderer/Private/PostProcess/PostProcessMotionBlur.cpp
--- a/Engine/Source/Runtime/Renderer/Private/PostProcess/PostProcessMotionBlur.cpp
+++ b/Engine/Source/Runtime/Renderer/Private/PostProcess/PostProcessMotionBlur.cpp
@@ -311,8 +311,9 @@ BEGIN_SHADER_PARAMETER_STRUCT(FMotionBlurFilterParameters, )
 	SHADER_PARAMETER_RDG_TEXTURE(Texture2D, VelocityTileTexture)
 
 	SHADER_PARAMETER_SAMPLER(SamplerState, ColorSampler)
-	SHADER_PARAMETER_SAMPLER(SamplerState, VelocitySampler)
-END_SHADER_PARAMETER_STRUCT()
+	SHADER_PARAMETER_SAMPLER(SamplerState, VelocityTileSampler)
+	SHADER_PARAMETER_SAMPLER(SamplerState, VelocityFlatSampler)
+	END_SHADER_PARAMETER_STRUCT()
 
 class FMotionBlurFilterPS : public FMotionBlurShader
 {
@@ -375,6 +376,7 @@ public:
 
 		SHADER_PARAMETER_SAMPLER(SamplerState, ColorSampler)
 		SHADER_PARAMETER_SAMPLER(SamplerState, VelocitySampler)
+		SHADER_PARAMETER_SAMPLER(SamplerState, DepthSampler)
 
 		RENDER_TARGET_BINDING_SLOTS()
 	END_SHADER_PARAMETER_STRUCT()
@@ -655,7 +657,8 @@ FScreenPassTexture ComputeMotionBlurFilter(
 		Parameters.VelocityFlatTexture = VelocityFlatTexture.GetRDGTexture();
 		Parameters.VelocityTileTexture = VelocityTileTexture.GetRDGTexture();
 		Parameters.ColorSampler = GetMotionBlurColorSampler();
-		Parameters.VelocitySampler = GetMotionBlurVelocitySampler();
+		Parameters.VelocityTileSampler = GetMotionBlurVelocitySampler();
+		Parameters.VelocityFlatSampler = GetMotionBlurVelocitySampler();
 		return Parameters;
 	};
 
@@ -801,6 +804,7 @@ FScreenPassTexture VisualizeMotionBlur(
 	PassParameters->ColorToVelocity = GetScreenPassTextureViewportTransform(PassParameters->Color, PassParameters->Velocity);
 	PassParameters->ColorSampler = GetMotionBlurColorSampler();
 	PassParameters->VelocitySampler = GetMotionBlurVelocitySampler();
+	PassParameters->DepthSampler = GetMotionBlurVelocitySampler();
 	PassParameters->RenderTargets[0] = FRenderTargetBinding(TextureOutput.GetRDGTexture(), ERenderTargetLoadAction::ENoAction, ERenderTargetStoreAction::EStore);
 
 	TShaderMapRef<FMotionBlurVisualizePS> PixelShader(Context->ShaderMap);
diff --git a/Engine/Source/Runtime/Renderer/Private/PostProcess/PostProcessSubsurface.cpp b/Engine/Source/Runtime/Renderer/Private/PostProcess/PostProcessSubsurface.cpp
--- a/Engine/Source/Runtime/Renderer/Private/PostProcess/PostProcessSubsurface.cpp
+++ b/Engine/Source/Runtime/Renderer/Private/PostProcess/PostProcessSubsurface.cpp
@@ -205,7 +205,7 @@ public:
 		SHADER_PARAMETER_STRUCT_INCLUDE(FSubsurfaceParameters, Subsurface)
 		SHADER_PARAMETER_STRUCT(FSubsurfaceInput, SubsurfaceInput0)
 		SHADER_PARAMETER_TEXTURE(Texture2D, MiniFontTexture)
-		SHADER_PARAMETER_SAMPLER(SamplerState, SubsurfaceSampler)
+		SHADER_PARAMETER_SAMPLER(SamplerState, SubsurfaceSampler0)
 		RENDER_TARGET_BINDING_SLOTS()
 	END_SHADER_PARAMETER_STRUCT()
 };
@@ -222,7 +222,7 @@ public:
 	BEGIN_SHADER_PARAMETER_STRUCT(FParameters, )
 		SHADER_PARAMETER_STRUCT_INCLUDE(FSubsurfaceParameters, Subsurface)
 		SHADER_PARAMETER_STRUCT(FSubsurfaceInput, SubsurfaceInput0)
-		SHADER_PARAMETER_SAMPLER(SamplerState, SubsurfaceSampler)
+		SHADER_PARAMETER_SAMPLER(SamplerState, SubsurfaceSampler0)
 		RENDER_TARGET_BINDING_SLOTS()
 	END_SHADER_PARAMETER_STRUCT()
 
@@ -243,7 +243,7 @@ public:
 	BEGIN_SHADER_PARAMETER_STRUCT(FParameters, )
 		SHADER_PARAMETER_STRUCT_INCLUDE(FSubsurfaceParameters, Subsurface)
 		SHADER_PARAMETER_STRUCT(FSubsurfaceInput, SubsurfaceInput0)
-		SHADER_PARAMETER_SAMPLER(SamplerState, SubsurfaceSampler)
+		SHADER_PARAMETER_SAMPLER(SamplerState, SubsurfaceSampler0)
 		RENDER_TARGET_BINDING_SLOTS()
 	END_SHADER_PARAMETER_STRUCT()
 
@@ -304,7 +304,8 @@ class FSubsurfaceRecombinePS : public FSubsurfaceShader
 		SHADER_PARAMETER_STRUCT_INCLUDE(FSubsurfaceParameters, Subsurface)
 		SHADER_PARAMETER_STRUCT(FSubsurfaceInput, SubsurfaceInput0)
 		SHADER_PARAMETER_STRUCT(FSubsurfaceInput, SubsurfaceInput1)
-		SHADER_PARAMETER_SAMPLER(SamplerState, SubsurfaceSampler)
+		SHADER_PARAMETER_SAMPLER(SamplerState, SubsurfaceSampler0)
+		SHADER_PARAMETER_SAMPLER(SamplerState, SubsurfaceSampler1)
 		RENDER_TARGET_BINDING_SLOTS()
 	END_SHADER_PARAMETER_STRUCT();
 
@@ -352,7 +353,7 @@ class FSubsurfaceViewportCopyPS : public FSubsurfaceShader
 
 	BEGIN_SHADER_PARAMETER_STRUCT(FParameters, )
 		SHADER_PARAMETER_STRUCT(FSubsurfaceInput, SubsurfaceInput0)
-		SHADER_PARAMETER_SAMPLER(SamplerState, SubsurfaceSampler)
+		SHADER_PARAMETER_SAMPLER(SamplerState, SubsurfaceSampler0)
 		RENDER_TARGET_BINDING_SLOTS()
 	END_SHADER_PARAMETER_STRUCT();
 
@@ -468,7 +469,7 @@ FScreenPassTexture ComputeSubsurface(
 			PassParameters->Subsurface = SubsurfaceCommonParameters;
 			PassParameters->RenderTargets[0] = FRenderTargetBinding(TextureOutput.GetRDGTexture(), ERenderTargetLoadAction::ENoAction, ERenderTargetStoreAction::EStore);
 			PassParameters->SubsurfaceInput0 = GetSubsurfaceInput(TextureInput);
-			PassParameters->SubsurfaceSampler = PointClampSampler;
+			PassParameters->SubsurfaceSampler0 = PointClampSampler;
 
 			FSubsurfaceSetupPS::FPermutationDomain PixelShaderPermutationVector;
 			PixelShaderPermutationVector.Set<FSubsurfaceSetupPS::FDimensionHalfRes>(bHalfRes);
@@ -521,7 +522,7 @@ FScreenPassTexture ComputeSubsurface(
 			PassParameters->Subsurface = SubsurfaceCommonParameters;
 			PassParameters->RenderTargets[0] = FRenderTargetBinding(TextureOutput.GetRDGTexture(), ERenderTargetLoadAction::ENoAction, ERenderTargetStoreAction::EStore);
 			PassParameters->SubsurfaceInput0 = GetSubsurfaceInput(TextureInput);
-			PassParameters->SubsurfaceSampler = SubsurfaceSamplerState;
+			PassParameters->SubsurfaceSampler0 = SubsurfaceSamplerState;
 
 			FSubsurfacePS::FPermutationDomain PixelShaderPermutationVector;
 			PixelShaderPermutationVector.Set<FSubsurfacePS::FDimensionDirection>(Direction);
@@ -546,7 +547,7 @@ FScreenPassTexture ComputeSubsurface(
 		FSubsurfaceViewportCopyPS::FParameters* PassParameters = GraphBuilder.AllocParameters<FSubsurfaceViewportCopyPS::FParameters>();
 		PassParameters->RenderTargets[0] = FRenderTargetBinding(TextureOutput.GetRDGTexture(), ERenderTargetLoadAction::ENoAction, ERenderTargetStoreAction::EStore);
 		PassParameters->SubsurfaceInput0 = GetSubsurfaceInput(TextureInput);
-		PassParameters->SubsurfaceSampler = PointClampSampler;
+		PassParameters->SubsurfaceSampler0 = PointClampSampler;
 
 		TShaderMapRef<FSubsurfaceViewportCopyPS> PixelShader(Context->ShaderMap);
 
@@ -584,12 +585,13 @@ FScreenPassTexture ComputeSubsurface(
 		PassParameters->Subsurface = SubsurfaceCommonParameters;
 		PassParameters->RenderTargets[0] = FRenderTargetBinding(TextureOutput.GetRDGTexture(), ERenderTargetLoadAction::ELoad, ERenderTargetStoreAction::EStore);
 		PassParameters->SubsurfaceInput0 = GetSubsurfaceInput(TextureInput);
-		PassParameters->SubsurfaceSampler = BilinearBorderSampler;
+		PassParameters->SubsurfaceSampler0 = BilinearBorderSampler;
 
 		// Scattering output target is only used when scattering is enabled.
 		if (SubsurfaceMode != ESubsurfaceMode::Bypass)
 		{
 			PassParameters->SubsurfaceInput1 = GetSubsurfaceInput(SubsurfaceTextureY);
+			PassParameters->SubsurfaceSampler1 = BilinearBorderSampler;
 		}
 
 		const FSubsurfaceRecombinePS::EQuality RecombineQuality = FSubsurfaceRecombinePS::GetQuality(Context->View);
@@ -623,7 +625,7 @@ FScreenPassTexture VisualizeSubsurface(
 	PassParameters->Subsurface = GetSubsurfaceCommonParameters(GraphBuilder.RHICmdList, Context);
 	PassParameters->RenderTargets[0] = FRenderTargetBinding(TextureOutput.GetRDGTexture(), ERenderTargetLoadAction::EClear, ERenderTargetStoreAction::EStore);
 	PassParameters->SubsurfaceInput0 = GetSubsurfaceInput(TextureInput);
-	PassParameters->SubsurfaceSampler = TStaticSamplerState<SF_Point, AM_Clamp, AM_Clamp, AM_Clamp>::GetRHI();
+	PassParameters->SubsurfaceSampler0 = TStaticSamplerState<SF_Point, AM_Clamp, AM_Clamp, AM_Clamp>::GetRHI();
 	PassParameters->MiniFontTexture = GetMiniFontTexture();
 
 	TShaderMapRef<FSubsurfaceVisualizePS> PixelShader(Context->ShaderMap);
