Fix missing shader parameter when using shared samplers

diff --git a/Engine/Source/Runtime/Renderer/Private/PostProcess/PostProcessMotionBlur.cpp b/Engine/Source/Runtime/Renderer/Private/PostProcess/PostProcessMotionBlur.cpp
--- a/Engine/Source/Runtime/Renderer/Private/PostProcess/PostProcessMotionBlur.cpp
+++ b/Engine/Source/Runtime/Renderer/Private/PostProcess/PostProcessMotionBlur.cpp
@@ -311,6 +311,7 @@ BEGIN_SHADER_PARAMETER_STRUCT(FMotionBlurFilterParameters, )
 	SHADER_PARAMETER_RDG_TEXTURE(Texture2D, VelocityTileTexture)
 
 	SHADER_PARAMETER_SAMPLER(SamplerState, ColorSampler)
+	SHADER_PARAMETER_SAMPLER(SamplerState, VelocitySampler)
 	SHADER_PARAMETER_SAMPLER(SamplerState, VelocityTileSampler)
 	SHADER_PARAMETER_SAMPLER(SamplerState, VelocityFlatSampler)
 	END_SHADER_PARAMETER_STRUCT()
@@ -657,6 +658,7 @@ FScreenPassTexture ComputeMotionBlurFilter(
 		Parameters.VelocityFlatTexture = VelocityFlatTexture.GetRDGTexture();
 		Parameters.VelocityTileTexture = VelocityTileTexture.GetRDGTexture();
 		Parameters.ColorSampler = GetMotionBlurColorSampler();
+		Parameters.VelocitySampler = GetMotionBlurVelocitySampler();
 		Parameters.VelocityTileSampler = GetMotionBlurVelocitySampler();
 		Parameters.VelocityFlatSampler = GetMotionBlurVelocitySampler();
 		return Parameters;
