Linked atmosphere sun angular diameter to the sun/directional light angular diameter.
SunDiscScale in constant buffer is back to what it was in 5915011 (a scale), and introduced AtmosphericFogSunDiscHalfApexAngleRadian.

diff --git a/Engine/Shaders/Private/AtmosphereCommon.ush b/Engine/Shaders/Private/AtmosphereCommon.ush
--- a/Engine/Shaders/Private/AtmosphereCommon.ush
+++ b/Engine/Shaders/Private/AtmosphereCommon.ush
@@ -544,7 +544,7 @@ float3 GetSunColor(float3 X, float T, float3 V, float3 S, float Radius, float Mu
     }
 	else 
 	{
-		float SunDiskMask = step(cos(PI * View.AtmosphericFogSunDiscScale / 180.0), dot(V, S));
+		float SunDiskMask = step(cos(View.AtmosphericFogSunDiscHalfApexAngleRadian), dot(V, S));
         return TransmittanceValue * SunDiskMask * View.AtmosphericFogSunDiscLuminance; // Eq (9)
     }
 }
diff --git a/Engine/Source/Runtime/Engine/Private/Components/DirectionalLightComponent.cpp b/Engine/Source/Runtime/Engine/Private/Components/DirectionalLightComponent.cpp
--- a/Engine/Source/Runtime/Engine/Private/Components/DirectionalLightComponent.cpp
+++ b/Engine/Source/Runtime/Engine/Private/Components/DirectionalLightComponent.cpp
@@ -369,6 +369,11 @@ public:
 		return SunDiscOuterSpaceLuminance;
 	}
 
+	virtual float GetSunLightHalfApexAngleRadian() const override
+	{
+		return 0.5f * LightSourceAngle * PI / 180.0f; // LightSourceAngle is apex angle (angular diameter) in degree
+	}
+
 private:
 
 	float GetEffectiveWholeSceneDynamicShadowRadius(bool bPrecomputedLightingIsValid) const
diff --git a/Engine/Source/Runtime/Engine/Public/SceneManagement.h b/Engine/Source/Runtime/Engine/Public/SceneManagement.h
--- a/Engine/Source/Runtime/Engine/Public/SceneManagement.h
+++ b/Engine/Source/Runtime/Engine/Public/SceneManagement.h
@@ -1272,7 +1272,6 @@ public:
 		return StatId; 
 	}	
 	inline int32 GetShadowMapChannel() const { return ShadowMapChannel; }
-	inline bool IsUsedAsAtmosphereSunLight() const { return bUsedAsAtmosphereSunLight; }
 	inline int32 GetPreviewShadowMapChannel() const { return PreviewShadowMapChannel; }
 
 	inline bool HasReflectiveShadowMap() const { return bHasReflectiveShadowMap; }
@@ -1293,9 +1292,20 @@ public:
 	virtual float GetMaxDrawDistance() const { return 0.0f; }
 	virtual float GetFadeRange() const { return 0.0f; }
 
-	// Atmosphere / Fog related functions
+	// Atmosphere / Fog related functions.
+
+	inline bool IsUsedAsAtmosphereSunLight() const { return bUsedAsAtmosphereSunLight; }
 	virtual void SetAtmosphereRelatedProperties(FLinearColor TransmittanceFactor, FLinearColor SunOuterSpaceLuminance) {}
 	virtual FLinearColor GetOuterSpaceLuminance() const { return FLinearColor::White; }
+	static float GetSunOnEarthHalfApexAngleRadian() 
+	{ 
+		const float SunOnEarthApexAngleDegree = 0.545f;	// Apex angle == angular diameter
+		return 0.5f * SunOnEarthApexAngleDegree * PI / 180.0f;
+	}
+	/**
+	 * @return the light half apex angle (half angular diameter) in radian.
+	 */
+	virtual float GetSunLightHalfApexAngleRadian() const { return GetSunOnEarthHalfApexAngleRadian() ; }
 
 protected:
 
diff --git a/Engine/Source/Runtime/Engine/Public/SceneView.h b/Engine/Source/Runtime/Engine/Public/SceneView.h
--- a/Engine/Source/Runtime/Engine/Public/SceneView.h
+++ b/Engine/Source/Runtime/Engine/Public/SceneView.h
@@ -666,6 +666,7 @@ enum ETranslucencyVolumeCascade
 	VIEW_UNIFORM_BUFFER_MEMBER_EX(float, AtmosphericFogStartDistance, EShaderPrecisionModifier::Half) \
 	VIEW_UNIFORM_BUFFER_MEMBER_EX(float, AtmosphericFogDistanceOffset, EShaderPrecisionModifier::Half) \
 	VIEW_UNIFORM_BUFFER_MEMBER_EX(float, AtmosphericFogSunDiscScale, EShaderPrecisionModifier::Half) \
+	VIEW_UNIFORM_BUFFER_MEMBER_EX(float, AtmosphericFogSunDiscHalfApexAngleRadian, EShaderPrecisionModifier::Half) \
 	VIEW_UNIFORM_BUFFER_MEMBER_EX(FLinearColor, AtmosphericFogSunDiscLuminance, EShaderPrecisionModifier::Half) \
 	VIEW_UNIFORM_BUFFER_MEMBER(uint32, AtmosphericFogRenderMask) \
 	VIEW_UNIFORM_BUFFER_MEMBER(uint32, AtmosphericFogInscatterAltitudeSampleNum) \
diff --git a/Engine/Source/Runtime/Renderer/Private/AtmosphereRendering.cpp b/Engine/Source/Runtime/Renderer/Private/AtmosphereRendering.cpp
--- a/Engine/Source/Runtime/Renderer/Private/AtmosphereRendering.cpp
+++ b/Engine/Source/Runtime/Renderer/Private/AtmosphereRendering.cpp
@@ -1166,9 +1166,10 @@ void FAtmosphericFogSceneInfo::PrepareSunLightProxy(FLightSceneInfo& SunLight) c
 
 	FLinearColor SunZenithIlluminance = SunLight.Proxy->GetColor();
 	FLinearColor SunOuterSpaceIlluminance = SunZenithIlluminance / TransmittanceAtZenith;
-
-	const float SunSolidAngle = 2.0f * PI * (1.0f - FMath::Cos(0.5f * SunHalfApexAngle));	// Solid angle from aperture https://en.wikipedia.org/wiki/Solid_angle
-	FLinearColor SunDiskOuterSpaceLuminance = SunOuterSpaceIlluminance / SunSolidAngle;		// approximation  
+	
+	// SunDiscScale is only considered as a visual tweak so we do not make it influence the sun disk outerspace luminance.
+	const float SunSolidAngle = 2.0f * PI * (1.0f - FMath::Cos(SunLight.Proxy->GetSunLightHalfApexAngleRadian())); // Solid angle from aperture https://en.wikipedia.org/wiki/Solid_angle 
+	FLinearColor SunDiskOuterSpaceLuminance = SunOuterSpaceIlluminance / SunSolidAngle; // approximation  
 
 	SunLight.Proxy->SetAtmosphereRelatedProperties(TransmittanceTowardSun / TransmittanceAtZenith, SunDiskOuterSpaceLuminance);
 }
@@ -1815,7 +1816,7 @@ FAtmosphericFogSceneInfo::FAtmosphericFogSceneInfo(UAtmosphericFogComponent* InC
 	, RHeight(InComponent->PrecomputeParams.DensityHeight * InComponent->PrecomputeParams.DensityHeight * InComponent->PrecomputeParams.DensityHeight * 64.f)
 	, StartDistance(InComponent->StartDistance)
 	, DistanceOffset(InComponent->DistanceOffset)
-	, SunHalfApexAngle(InComponent->SunDiscScale)
+	, SunDiscScale(InComponent->SunDiscScale)
 	, RenderFlag(EAtmosphereRenderFlag::E_EnableAll)
 	, InscatterAltitudeSampleNum(InComponent->PrecomputeParams.InscatterAltitudeSampleNum)
 
diff --git a/Engine/Source/Runtime/Renderer/Private/AtmosphereRendering.h b/Engine/Source/Runtime/Renderer/Private/AtmosphereRendering.h
--- a/Engine/Source/Runtime/Renderer/Private/AtmosphereRendering.h
+++ b/Engine/Source/Runtime/Renderer/Private/AtmosphereRendering.h
@@ -51,7 +51,7 @@ public:
 	float RHeight;
 	float StartDistance;
 	float DistanceOffset;
-	float SunHalfApexAngle;
+	float SunDiscScale;
 	FLinearColor DefaultSunColor;
 	FVector DefaultSunDirection;
 	uint32 RenderFlag;
@@ -81,6 +81,7 @@ public:
 	explicit FAtmosphericFogSceneInfo(UAtmosphericFogComponent* InComponent, const FScene* InScene);
 	~FAtmosphericFogSceneInfo();
 
+	/** Prepare the sun light data as a function of current atmospheric fog state. */
 	void PrepareSunLightProxy(FLightSceneInfo& SunLight) const;
 
 #if WITH_EDITOR
diff --git a/Engine/Source/Runtime/Renderer/Private/SceneRendering.cpp b/Engine/Source/Runtime/Renderer/Private/SceneRendering.cpp
--- a/Engine/Source/Runtime/Renderer/Private/SceneRendering.cpp
+++ b/Engine/Source/Runtime/Renderer/Private/SceneRendering.cpp
@@ -1124,6 +1124,7 @@ void FViewInfo::SetupUniformBufferParameters(
 		}
 
 		// Atmospheric fog parameters
+		const float SunLightDiskHalfApexAngleRadian = Scene->SunLight ? Scene->SunLight->Proxy->GetSunLightHalfApexAngleRadian() : FLightSceneProxy::GetSunOnEarthHalfApexAngleRadian();
 		if (ShouldRenderAtmosphere(*Family) && Scene->AtmosphericFog)
 		{
 			ViewUniformShaderParameters.AtmosphericFogSunPower = Scene->AtmosphericFog->SunMultiplier;
@@ -1136,7 +1137,8 @@ void FViewInfo::SetupUniformBufferParameters(
 			ViewUniformShaderParameters.AtmosphericFogHeightScaleRayleigh = Scene->AtmosphericFog->RHeight;
 			ViewUniformShaderParameters.AtmosphericFogStartDistance = Scene->AtmosphericFog->StartDistance;
 			ViewUniformShaderParameters.AtmosphericFogDistanceOffset = Scene->AtmosphericFog->DistanceOffset;
-			ViewUniformShaderParameters.AtmosphericFogSunDiscScale = Scene->AtmosphericFog->SunHalfApexAngle;
+			ViewUniformShaderParameters.AtmosphericFogSunDiscScale = Scene->AtmosphericFog->SunDiscScale;
+			ViewUniformShaderParameters.AtmosphericFogSunDiscHalfApexAngleRadian = Scene->AtmosphericFog->SunDiscScale * SunLightDiskHalfApexAngleRadian;
 			ViewUniformShaderParameters.AtmosphericFogSunDiscLuminance = Scene->SunLight ? Scene->SunLight->Proxy->GetOuterSpaceLuminance() : FLinearColor::White;
 			ViewUniformShaderParameters.AtmosphericFogSunColor = Scene->SunLight ? Scene->SunLight->Proxy->GetColor() : Scene->AtmosphericFog->DefaultSunColor;
 			ViewUniformShaderParameters.AtmosphericFogSunDirection = Scene->SunLight ? -Scene->SunLight->Proxy->GetDirection() : -Scene->AtmosphericFog->DefaultSunDirection;
@@ -1156,7 +1158,7 @@ void FViewInfo::SetupUniformBufferParameters(
 			ViewUniformShaderParameters.AtmosphericFogStartDistance = FLT_MAX;
 			ViewUniformShaderParameters.AtmosphericFogDistanceOffset = 0.f;
 			ViewUniformShaderParameters.AtmosphericFogSunDiscScale = 1.f;
-			ViewUniformShaderParameters.AtmosphericFogSunDiscLuminance = Scene->SunLight ? Scene->SunLight->Proxy->GetOuterSpaceLuminance() : FLinearColor::Black;
+			ViewUniformShaderParameters.AtmosphericFogSunDiscHalfApexAngleRadian = SunLightDiskHalfApexAngleRadian;
 			//Added check so atmospheric light color and vector can use a directional light without needing an atmospheric fog actor in the scene
 			ViewUniformShaderParameters.AtmosphericFogSunColor = Scene->SunLight ? Scene->SunLight->Proxy->GetColor() : FLinearColor::Black;
 			ViewUniformShaderParameters.AtmosphericFogSunDirection = Scene->SunLight ? -Scene->SunLight->Proxy->GetDirection() : FVector::ZeroVector;
@@ -1178,6 +1180,7 @@ void FViewInfo::SetupUniformBufferParameters(
 		ViewUniformShaderParameters.AtmosphericFogStartDistance = FLT_MAX;
 		ViewUniformShaderParameters.AtmosphericFogDistanceOffset = 0.f;
 		ViewUniformShaderParameters.AtmosphericFogSunDiscScale = 1.f;
+		ViewUniformShaderParameters.AtmosphericFogSunDiscHalfApexAngleRadian = 0.f;
 		ViewUniformShaderParameters.AtmosphericFogSunDiscLuminance = FLinearColor::Black;
 		ViewUniformShaderParameters.AtmosphericFogSunColor = FLinearColor::Black;
 		ViewUniformShaderParameters.AtmosphericFogSunDirection = FVector::ZeroVector;
