Another build fix.
Development non editor Build is now green locally...

diff --git a/Engine/Source/Runtime/Renderer/Private/AtmosphereRendering.cpp b/Engine/Source/Runtime/Renderer/Private/AtmosphereRendering.cpp
--- a/Engine/Source/Runtime/Renderer/Private/AtmosphereRendering.cpp
+++ b/Engine/Source/Runtime/Renderer/Private/AtmosphereRendering.cpp
@@ -429,6 +429,12 @@ void FDeferredShadingSceneRenderer::RenderAtmosphere(FRHICommandListImmediate& R
 	}
 }
 
+namespace
+{
+	const float RadiusGround = 6360;
+	const float RadiusAtmosphere = 6420;
+}
+
 #if WITH_EDITOR
 class FAtmosphereTransmittancePS : public FGlobalShader
 {
@@ -1019,9 +1025,6 @@ IMPLEMENT_SHADER_TYPE(,FAtmospherePrecomputeInscatterVS,TEXT("/Engine/Private/At
 
 namespace
 {
-	const float RadiusGround = 6360;
-	const float RadiusAtmosphere = 6420;
-
 	enum
 	{
 		AP_Transmittance = 0,
