Disabled sun rendering in reflection captures from atmosphere as the specular component is already handled analytically when rendering directional lights.

diff --git a/Engine/Source/Runtime/Renderer/Private/AtmosphereRendering.cpp b/Engine/Source/Runtime/Renderer/Private/AtmosphereRendering.cpp
--- a/Engine/Source/Runtime/Renderer/Private/AtmosphereRendering.cpp
+++ b/Engine/Source/Runtime/Renderer/Private/AtmosphereRendering.cpp
@@ -348,6 +348,12 @@ void SetAtmosphericFogShaders(FRHICommandList& RHICmdList, FGraphicsPipelineStat
 
 	auto ShaderMap = View.ShaderMap;
 
+	if (View.bIsReflectionCapture)
+	{
+		// We do not render the sun in in reflection captures as the specular component is already handled analytically when rendering directional lights.
+		RenderFlag |= uint32(EAtmosphereRenderFlag::E_DisableSunDisk);
+	}
+
 	TShaderMapRef<FAtmosphericVS> VertexShader(ShaderMap);
 	FAtmosphericFogPS* PixelShader = NULL;
 	switch (RenderFlag)
