Sky/Atmosphere code clean up

diff --git a/Engine/Source/Runtime/Engine/Private/Atmosphere/Atmosphere.h b/Engine/Source/Runtime/Engine/Private/Atmosphere/Atmosphere.h
--- a/Engine/Source/Runtime/Engine/Private/Atmosphere/Atmosphere.h
+++ b/Engine/Source/Runtime/Engine/Private/Atmosphere/Atmosphere.h
@@ -6,7 +6,6 @@
 #include "Containers/ResourceArray.h"
 #include "RHI.h"
 #include "RenderResource.h"
-#include "Atmosphere/AtmosphericFogComponent.h"
 
 // Shared by Engine class and Renderer class, need to be core?
 
diff --git a/Engine/Source/Runtime/Engine/Private/Atmosphere/AtmosphericFog.cpp b/Engine/Source/Runtime/Engine/Private/Atmosphere/AtmosphericFog.cpp
new file mode 100644
--- /dev/null
+++ b/Engine/Source/Runtime/Engine/Private/Atmosphere/AtmosphericFog.cpp
@@ -0,0 +1,90 @@
+// Copyright 1998-2019 Epic Games, Inc. All Rights Reserved.
+
+#include "Atmosphere/AtmosphericFog.h"
+#include "Atmosphere/AtmosphericFogComponent.h"
+#include "UObject/ConstructorHelpers.h"
+#include "EngineDefines.h"
+#include "RenderingThread.h"
+#include "Components/ArrowComponent.h"
+#include "UObject/UObjectHash.h"
+#include "UObject/UObjectIterator.h"
+#include "Engine/Texture2D.h"
+#include "ComponentReregisterContext.h"
+#include "Atmosphere/AtmosphericFog.h"
+#include "Components/BillboardComponent.h"
+#include "Runtime/Renderer/Private/ScenePrivate.h"
+#include "Runtime/Renderer/Private/AtmosphereRendering.h"
+
+#if WITH_EDITOR
+#include "ObjectEditorUtils.h"
+#endif
+
+AAtmosphericFog::AAtmosphericFog(const FObjectInitializer& ObjectInitializer)
+	: Super(ObjectInitializer)
+{
+	AtmosphericFogComponent = CreateDefaultSubobject<UAtmosphericFogComponent>(TEXT("AtmosphericFogComponent0"));
+	RootComponent = AtmosphericFogComponent;
+
+#if WITH_EDITORONLY_DATA
+	ArrowComponent = CreateEditorOnlyDefaultSubobject<UArrowComponent>(TEXT("ArrowComponent0"));
+
+	if (!IsRunningCommandlet())
+	{
+		// Structure to hold one-time initialization
+		struct FConstructorStatics
+		{
+			ConstructorHelpers::FObjectFinderOptional<UTexture2D> FogTextureObject;
+			FName ID_Fog;
+			FText NAME_Fog;
+			FConstructorStatics()
+				: FogTextureObject(TEXT("/Engine/EditorResources/S_ExpoHeightFog"))
+				, ID_Fog(TEXT("Fog"))
+				, NAME_Fog(NSLOCTEXT("SpriteCategory", "Fog", "Fog"))
+			{
+			}
+		};
+		static FConstructorStatics ConstructorStatics;
+
+		if (GetSpriteComponent())
+		{
+			GetSpriteComponent()->Sprite = ConstructorStatics.FogTextureObject.Get();
+			GetSpriteComponent()->RelativeScale3D = FVector(0.5f, 0.5f, 0.5f);
+			GetSpriteComponent()->SpriteInfo.Category = ConstructorStatics.ID_Fog;
+			GetSpriteComponent()->SpriteInfo.DisplayName = ConstructorStatics.NAME_Fog;
+			GetSpriteComponent()->SetupAttachment(AtmosphericFogComponent);
+		}
+
+		if (ArrowComponent)
+		{
+			ArrowComponent->ArrowColor = FColor(150, 200, 255);
+
+			ArrowComponent->bTreatAsASprite = true;
+			ArrowComponent->SpriteInfo.Category = ConstructorStatics.ID_Fog;
+			ArrowComponent->SpriteInfo.DisplayName = ConstructorStatics.NAME_Fog;
+			ArrowComponent->SetupAttachment(AtmosphericFogComponent);
+			ArrowComponent->bLightAttachment = true;
+			ArrowComponent->bIsScreenSizeScaled = true;
+		}
+	}
+#endif // WITH_EDITORONLY_DATA
+
+
+	PrimaryActorTick.bCanEverTick = true;
+	bHidden = false;
+}
+
+#if WITH_EDITOR
+// Prepare render targets when new actor spawned
+void AAtmosphericFog::PostActorCreated()
+{
+	Super::PostActorCreated();
+	if (GIsEditor)
+	{
+		if ( !IsTemplate() && AtmosphericFogComponent )
+		{
+			AtmosphericFogComponent->InitResource();
+		}
+	}
+}
+
+#endif
diff --git a/Engine/Source/Runtime/Engine/Private/Atmosphere/AtmosphericFogComponent.cpp b/Engine/Source/Runtime/Engine/Private/Atmosphere/AtmosphericFogComponent.cpp
--- a/Engine/Source/Runtime/Engine/Private/Atmosphere/AtmosphericFogComponent.cpp
+++ b/Engine/Source/Runtime/Engine/Private/Atmosphere/AtmosphericFogComponent.cpp
@@ -19,61 +19,11 @@
 #include "ObjectEditorUtils.h"
 #endif
 
-AAtmosphericFog::AAtmosphericFog(const FObjectInitializer& ObjectInitializer)
-	: Super(ObjectInitializer)
-{
-	AtmosphericFogComponent = CreateDefaultSubobject<UAtmosphericFogComponent>(TEXT("AtmosphericFogComponent0"));
-	RootComponent = AtmosphericFogComponent;
-
-#if WITH_EDITORONLY_DATA
-	ArrowComponent = CreateEditorOnlyDefaultSubobject<UArrowComponent>(TEXT("ArrowComponent0"));
 
-	if (!IsRunningCommandlet())
-	{
-		// Structure to hold one-time initialization
-		struct FConstructorStatics
-		{
-			ConstructorHelpers::FObjectFinderOptional<UTexture2D> FogTextureObject;
-			FName ID_Fog;
-			FText NAME_Fog;
-			FConstructorStatics()
-				: FogTextureObject(TEXT("/Engine/EditorResources/S_ExpoHeightFog"))
-				, ID_Fog(TEXT("Fog"))
-				, NAME_Fog(NSLOCTEXT("SpriteCategory", "Fog", "Fog"))
-			{
-			}
-		};
-		static FConstructorStatics ConstructorStatics;
+//////////////////////////////////////////////////////////////////////////
+// FAtmospherePrecomputeParameters
 
-		if (GetSpriteComponent())
-		{
-			GetSpriteComponent()->Sprite = ConstructorStatics.FogTextureObject.Get();
-			GetSpriteComponent()->RelativeScale3D = FVector(0.5f, 0.5f, 0.5f);
-			GetSpriteComponent()->SpriteInfo.Category = ConstructorStatics.ID_Fog;
-			GetSpriteComponent()->SpriteInfo.DisplayName = ConstructorStatics.NAME_Fog;
-			GetSpriteComponent()->SetupAttachment(AtmosphericFogComponent);
-		}
-
-		if (ArrowComponent)
-		{
-			ArrowComponent->ArrowColor = FColor(150, 200, 255);
-
-			ArrowComponent->bTreatAsASprite = true;
-			ArrowComponent->SpriteInfo.Category = ConstructorStatics.ID_Fog;
-			ArrowComponent->SpriteInfo.DisplayName = ConstructorStatics.NAME_Fog;
-			ArrowComponent->SetupAttachment(AtmosphericFogComponent);
-			ArrowComponent->bLightAttachment = true;
-			ArrowComponent->bIsScreenSizeScaled = true;
-		}
-	}
-#endif // WITH_EDITORONLY_DATA
 
-
-	PrimaryActorTick.bCanEverTick = true;
-	bHidden = false;
-}
-
-// Precomputation parameters
 FAtmospherePrecomputeParameters::FAtmospherePrecomputeParameters():
 	DensityHeight(0.5f),
 	DecayHeight_DEPRECATED(0.5f),
@@ -89,6 +39,11 @@ FAtmospherePrecomputeParameters::FAtmospherePrecomputeParameters():
 {
 }
 
+
+//////////////////////////////////////////////////////////////////////////
+// UAtmosphericFogComponent
+
+
 UAtmosphericFogComponent::UAtmosphericFogComponent(const FObjectInitializer& ObjectInitializer)
 	: Super(ObjectInitializer), TransmittanceResource(NULL), IrradianceResource(NULL), InscatterResource(NULL)
 #if WITH_EDITORONLY_DATA
@@ -603,18 +558,6 @@ UAtmosphericFogComponent::~UAtmosphericFogComponent()
 }
 
 #if WITH_EDITOR
-// Prepare render targets when new actor spawned
-void AAtmosphericFog::PostActorCreated()
-{
-	Super::PostActorCreated();
-	if (GIsEditor)
-	{
-		if ( !IsTemplate() && AtmosphericFogComponent )
-		{
-			AtmosphericFogComponent->InitResource();
-		}
-	}
-}
 
 void UAtmosphericFogComponent::UpdatePrecomputedData()
 {
diff --git a/Engine/Source/Runtime/Renderer/Private/AtmosphereRendering.cpp b/Engine/Source/Runtime/Renderer/Private/AtmosphereRendering.cpp
--- a/Engine/Source/Runtime/Renderer/Private/AtmosphereRendering.cpp
+++ b/Engine/Source/Runtime/Renderer/Private/AtmosphereRendering.cpp
@@ -15,7 +15,6 @@
 #include "Atmosphere/AtmosphericFogComponent.h"
 #include "PostProcess/SceneRenderTargets.h"
 #include "GlobalShader.h"
-#include "AtmosphereTextureParameters.h"
 #include "SceneRenderTargetParameters.h"
 #include "DeferredShadingRenderer.h"
 #include "ScenePrivate.h"
@@ -27,6 +26,78 @@
 DECLARE_GPU_STAT(Atmosphere);
 DECLARE_GPU_STAT(AtmospherePreCompute);
 
+
+//////////////////////////////////////////////////////////////////////////
+// FAtmosphereShaderTextureParameters
+
+
+/** Shader parameters needed for atmosphere passes. */
+class FAtmosphereShaderTextureParameters
+{
+public:
+	void Bind(const FShaderParameterMap& ParameterMap);
+
+	template< typename ShaderRHIParamRef >
+	FORCEINLINE_DEBUGGABLE void Set(FRHICommandList& RHICmdList, const ShaderRHIParamRef ShaderRHI, const FSceneView& View) const
+	{
+		if (TransmittanceTexture.IsBound() || IrradianceTexture.IsBound() || InscatterTexture.IsBound())
+		{
+			SetTextureParameter(RHICmdList, ShaderRHI, TransmittanceTexture, TransmittanceTextureSampler,
+				TStaticSamplerState<SF_Bilinear>::GetRHI(), View.AtmosphereTransmittanceTexture);
+			SetTextureParameter(RHICmdList, ShaderRHI, IrradianceTexture, IrradianceTextureSampler,
+				TStaticSamplerState<SF_Bilinear>::GetRHI(), View.AtmosphereIrradianceTexture);
+			SetTextureParameter(RHICmdList, ShaderRHI, InscatterTexture, InscatterTextureSampler,
+				TStaticSamplerState<SF_Bilinear>::GetRHI(), View.AtmosphereInscatterTexture);
+		}
+	}
+
+	friend FArchive& operator<<(FArchive& Ar, FAtmosphereShaderTextureParameters& P);
+
+private:
+	FShaderResourceParameter TransmittanceTexture;
+	FShaderResourceParameter TransmittanceTextureSampler;
+	FShaderResourceParameter IrradianceTexture;
+	FShaderResourceParameter IrradianceTextureSampler;
+	FShaderResourceParameter InscatterTexture;
+	FShaderResourceParameter InscatterTextureSampler;
+};
+
+void FAtmosphereShaderTextureParameters::Bind(const FShaderParameterMap& ParameterMap)
+{
+	TransmittanceTexture.Bind(ParameterMap, TEXT("AtmosphereTransmittanceTexture"));
+	TransmittanceTextureSampler.Bind(ParameterMap, TEXT("AtmosphereTransmittanceTextureSampler"));
+	IrradianceTexture.Bind(ParameterMap, TEXT("AtmosphereIrradianceTexture"));
+	IrradianceTextureSampler.Bind(ParameterMap, TEXT("AtmosphereIrradianceTextureSampler"));
+	InscatterTexture.Bind(ParameterMap, TEXT("AtmosphereInscatterTexture"));
+	InscatterTextureSampler.Bind(ParameterMap, TEXT("AtmosphereInscatterTextureSampler"));
+}
+
+#define IMPLEMENT_ATMOSPHERE_TEXTURE_PARAM_SET( ShaderRHIParamRef ) \
+	template void FAtmosphereShaderTextureParameters::Set< ShaderRHIParamRef >( FRHICommandList& RHICmdList, const ShaderRHIParamRef ShaderRHI, const FSceneView& View ) const;
+
+IMPLEMENT_ATMOSPHERE_TEXTURE_PARAM_SET(FVertexShaderRHIParamRef);
+IMPLEMENT_ATMOSPHERE_TEXTURE_PARAM_SET(FHullShaderRHIParamRef);
+IMPLEMENT_ATMOSPHERE_TEXTURE_PARAM_SET(FDomainShaderRHIParamRef);
+IMPLEMENT_ATMOSPHERE_TEXTURE_PARAM_SET(FGeometryShaderRHIParamRef);
+IMPLEMENT_ATMOSPHERE_TEXTURE_PARAM_SET(FPixelShaderRHIParamRef);
+IMPLEMENT_ATMOSPHERE_TEXTURE_PARAM_SET(FComputeShaderRHIParamRef);
+
+FArchive& operator<<(FArchive& Ar, FAtmosphereShaderTextureParameters& Parameters)
+{
+	Ar << Parameters.TransmittanceTexture;
+	Ar << Parameters.TransmittanceTextureSampler;
+	Ar << Parameters.IrradianceTexture;
+	Ar << Parameters.IrradianceTextureSampler;
+	Ar << Parameters.InscatterTexture;
+	Ar << Parameters.InscatterTextureSampler;
+	return Ar;
+}
+
+
+//////////////////////////////////////////////////////////////////////////
+// FAtmosphereShaderPrecomputeTextureParameters
+
+
 class FAtmosphereShaderPrecomputeTextureParameters
 {
 public:
@@ -155,6 +226,11 @@ FArchive& operator<<(FArchive& Ar,FAtmosphereShaderPrecomputeTextureParameters&
 	return Ar;
 }
 
+
+//////////////////////////////////////////////////////////////////////////
+// Global shaders
+
+
 /** A pixel shader for rendering atmospheric fog. */
 class FAtmosphericFogPS : public FGlobalShader
 {
@@ -1029,6 +1105,11 @@ IMPLEMENT_SHADER_TYPE(,FAtmosphereCopyInscatterFBackPS,TEXT("/Engine/Private/Atm
 IMPLEMENT_SHADER_TYPE(,FAtmospherePrecomputeVS,TEXT("/Engine/Private/AtmospherePrecompute.usf"),TEXT("MainVS"),SF_Vertex);
 IMPLEMENT_SHADER_TYPE(,FAtmospherePrecomputeInscatterVS,TEXT("/Engine/Private/AtmospherePrecomputeInscatter.usf"),TEXT("MainVS"),SF_Vertex);
 
+
+//////////////////////////////////////////////////////////////////////////
+// FAtmosphericFogSceneInfo
+
+
 namespace
 {
 	enum
@@ -1885,6 +1966,11 @@ bool ShouldRenderAtmosphere(const FSceneViewFamily& Family)
 		&& SupportAtmosphericFog->GetValueOnAnyThread();
 }
 
+
+//////////////////////////////////////////////////////////////////////////
+// FScene
+
+
 void FScene::AddAtmosphericFog(UAtmosphericFogComponent* FogComponent)
 {
 	check(FogComponent);
diff --git a/Engine/Source/Runtime/Renderer/Private/AtmosphereTextures.cpp b/Engine/Source/Runtime/Renderer/Private/AtmosphereTextures.cpp
--- a/Engine/Source/Runtime/Renderer/Private/AtmosphereTextures.cpp
+++ b/Engine/Source/Runtime/Renderer/Private/AtmosphereTextures.cpp
@@ -7,7 +7,6 @@
 #include "AtmosphereTextures.h"
 #include "Atmosphere/AtmosphericFogComponent.h"
 #include "RenderTargetPool.h"
-#include "AtmosphereTextureParameters.h"
 #include "ShaderParameterUtils.h"
 
 void FAtmosphereTextures::InitDynamicRHI()
@@ -75,36 +74,3 @@ void FAtmosphereTextures::ReleaseDynamicRHI()
 
 	GRenderTargetPool.FreeUnusedResources();
 }
-
-
-void FAtmosphereShaderTextureParameters::Bind(const FShaderParameterMap& ParameterMap)
-{
-	TransmittanceTexture.Bind(ParameterMap,TEXT("AtmosphereTransmittanceTexture"));
-	TransmittanceTextureSampler.Bind(ParameterMap,TEXT("AtmosphereTransmittanceTextureSampler"));
-	IrradianceTexture.Bind(ParameterMap,TEXT("AtmosphereIrradianceTexture"));
-	IrradianceTextureSampler.Bind(ParameterMap,TEXT("AtmosphereIrradianceTextureSampler"));
-	InscatterTexture.Bind(ParameterMap,TEXT("AtmosphereInscatterTexture"));
-	InscatterTextureSampler.Bind(ParameterMap,TEXT("AtmosphereInscatterTextureSampler"));
-}
-
-
-#define IMPLEMENT_ATMOSPHERE_TEXTURE_PARAM_SET( ShaderRHIParamRef ) \
-	template void FAtmosphereShaderTextureParameters::Set< ShaderRHIParamRef >( FRHICommandList& RHICmdList, const ShaderRHIParamRef ShaderRHI, const FSceneView& View ) const;
-
-IMPLEMENT_ATMOSPHERE_TEXTURE_PARAM_SET( FVertexShaderRHIParamRef );
-IMPLEMENT_ATMOSPHERE_TEXTURE_PARAM_SET( FHullShaderRHIParamRef );
-IMPLEMENT_ATMOSPHERE_TEXTURE_PARAM_SET( FDomainShaderRHIParamRef );
-IMPLEMENT_ATMOSPHERE_TEXTURE_PARAM_SET( FGeometryShaderRHIParamRef );
-IMPLEMENT_ATMOSPHERE_TEXTURE_PARAM_SET( FPixelShaderRHIParamRef );
-IMPLEMENT_ATMOSPHERE_TEXTURE_PARAM_SET( FComputeShaderRHIParamRef );
-
-FArchive& operator<<(FArchive& Ar,FAtmosphereShaderTextureParameters& Parameters)
-{
-	Ar << Parameters.TransmittanceTexture;
-	Ar << Parameters.TransmittanceTextureSampler;
-	Ar << Parameters.IrradianceTexture;
-	Ar << Parameters.IrradianceTextureSampler;
-	Ar << Parameters.InscatterTexture;
-	Ar << Parameters.InscatterTextureSampler;
-	return Ar;
-}
diff --git a/Engine/Source/Runtime/Renderer/Private/RendererPrivate.h b/Engine/Source/Runtime/Renderer/Private/RendererPrivate.h
--- a/Engine/Source/Runtime/Renderer/Private/RendererPrivate.h
+++ b/Engine/Source/Runtime/Renderer/Private/RendererPrivate.h
@@ -7,7 +7,6 @@
 #pragma once
 
 #include "CoreMinimal.h"
-#include "AtmosphereTextureParameters.h"
 #include "PostProcessParameters.h"
 #include "ShaderComplexityRendering.h"
 #include "Async/ParallelFor.h"
diff --git a/Engine/Source/Runtime/Renderer/Public/AtmosphereTextureParameters.h b/Engine/Source/Runtime/Renderer/Public/AtmosphereTextureParameters.h
deleted file mode 100644
