Kills Gaussian DOF, old Circle DOF and Bokeh DOF on deferred shading renderer

diff --git a/Engine/Plugins/Runtime/Nvidia/Ansel/Source/Ansel/Private/Ansel.cpp b/Engine/Plugins/Runtime/Nvidia/Ansel/Source/Ansel/Private/Ansel.cpp
--- a/Engine/Plugins/Runtime/Nvidia/Ansel/Source/Ansel/Private/Ansel.cpp
+++ b/Engine/Plugins/Runtime/Nvidia/Ansel/Source/Ansel/Private/Ansel.cpp
@@ -364,31 +364,10 @@ void FNVAnselCameraPhotographyPrivate::DoCustomUIControls(FPostProcessSettings&
 		PRAGMA_DISABLE_DEPRECATION_WARNINGS
 		if (bEffectUIAllowed[DepthOfField])
 		{
-			bool bAnyDofVisible =
-				(InOutPPSettings.DepthOfFieldMethod == DOFM_CircleDOF && InOutPPSettings.DepthOfFieldDepthBlurRadius > 0.f) ||
-				(InOutPPSettings.DepthOfFieldMethod == DOFM_CircleDOF && InOutPPSettings.DepthOfFieldDepthBlurAmount > 0.f) ||
-				(InOutPPSettings.DepthOfFieldMethod == DOFM_BokehDOF && InOutPPSettings.DepthOfFieldScale > 0.f)
-				;
+			bool bAnyDofVisible = InOutPPSettings.DepthOfFieldFstop >= 0;
 
 			if (bAnyDofVisible)
 			{
-				if (InOutPPSettings.DepthOfFieldMethod == DOFM_BokehDOF)
-				{
-					DeclareSlider(
-						control_dofscale,
-						LOCTEXT("control_dofscale", "Focus Scale"),
-						0.f, 2.f,
-						InOutPPSettings.DepthOfFieldScale
-					);
-
-					DeclareSlider(
-						control_doffocalregion,
-						LOCTEXT("control_doffocalregion", "Focus Region"),
-						0.f, 10000.f, // UU
-						InOutPPSettings.DepthOfFieldFocalRegion
-					);
-				}
-
 				DeclareSlider(
 					control_dofsensorwidth,
 					LOCTEXT("control_dofsensorwidth", "Focus Sensor"), // n.b. similar effect to focus scale
@@ -403,23 +382,19 @@ void FNVAnselCameraPhotographyPrivate::DoCustomUIControls(FPostProcessSettings&
 					InOutPPSettings.DepthOfFieldFocalDistance
 				);
 
-				if (InOutPPSettings.DepthOfFieldMethod == DOFM_CircleDOF)
-				{
-					// circledof
-					DeclareSlider(
-						control_dofdepthbluramount,
-						LOCTEXT("control_dofbluramount", "Blur Distance km"),
-						0.000001f, 1.f, // km; doc'd as up to 100km but that's too coarse for a narrow UI control
-						InOutPPSettings.DepthOfFieldDepthBlurAmount
-					);
-					// circledof
-					DeclareSlider(
-						control_dofdepthblurradius,
-						LOCTEXT("control_dofblurradius", "Blur Radius"),
-						0.f, 4.f,
-						InOutPPSettings.DepthOfFieldDepthBlurRadius
-					);
-				}
+				DeclareSlider(
+					control_dofdepthbluramount,
+					LOCTEXT("control_dofbluramount", "Blur Distance km"),
+					0.000001f, 1.f, // km; doc'd as up to 100km but that's too coarse for a narrow UI control
+					InOutPPSettings.DepthOfFieldDepthBlurAmount
+				);
+
+				DeclareSlider(
+					control_dofdepthblurradius,
+					LOCTEXT("control_dofblurradius", "Blur Radius"),
+					0.f, 4.f,
+					InOutPPSettings.DepthOfFieldDepthBlurRadius
+				);
 			}
 		}
 		PRAGMA_ENABLE_DEPRECATION_WARNINGS
diff --git a/Engine/Shaders/Private/PostProcessCircleDOF.usf b/Engine/Shaders/Private/PostProcessCircleDOF.usf
deleted file mode 100644
diff --git a/Engine/Source/Runtime/CinematicCamera/Private/CineCameraComponent.cpp b/Engine/Source/Runtime/CinematicCamera/Private/CineCameraComponent.cpp
--- a/Engine/Source/Runtime/CinematicCamera/Private/CineCameraComponent.cpp
+++ b/Engine/Source/Runtime/CinematicCamera/Private/CineCameraComponent.cpp
@@ -42,11 +42,6 @@ UCineCameraComponent::UCineCameraComponent()
 
 	bConstrainAspectRatio = true;
 
-	PRAGMA_DISABLE_DEPRECATION_WARNINGS
-	// Default to CircleDOF, but allow the user to customize it
-	PostProcessSettings.DepthOfFieldMethod = DOFM_CircleDOF;
-	PRAGMA_ENABLE_DEPRECATION_WARNINGS
-
 	// Certain default values are set by a config, so Use the archetype to set them in the constructor, so they can be overridden in the editor.
 	UCineCameraComponent* Template = Cast<UCineCameraComponent>(GetArchetype());
 	
@@ -410,9 +405,6 @@ void UCineCameraComponent::UpdateCameraLens(float DeltaTime, FMinimalViewInfo& D
 {
 	if (FocusSettings.FocusMethod == ECameraFocusMethod::None)
 	{
-		PRAGMA_DISABLE_DEPRECATION_WARNINGS
-		DesiredView.PostProcessSettings.bOverride_DepthOfFieldMethod = false;
-		PRAGMA_ENABLE_DEPRECATION_WARNINGS
 		DesiredView.PostProcessSettings.bOverride_DepthOfFieldFstop = false;
 		DesiredView.PostProcessSettings.bOverride_DepthOfFieldMinFstop = false;
 		DesiredView.PostProcessSettings.bOverride_DepthOfFieldBladeCount = false;
@@ -423,10 +415,6 @@ void UCineCameraComponent::UpdateCameraLens(float DeltaTime, FMinimalViewInfo& D
 	{
 		// Update focus/DoF
 		DesiredView.PostProcessBlendWeight = 1.f;
-		PRAGMA_DISABLE_DEPRECATION_WARNINGS
-		DesiredView.PostProcessSettings.bOverride_DepthOfFieldMethod = true;
-		DesiredView.PostProcessSettings.DepthOfFieldMethod = PostProcessSettings.DepthOfFieldMethod;
-		PRAGMA_ENABLE_DEPRECATION_WARNINGS
 
 		DesiredView.PostProcessSettings.bOverride_DepthOfFieldFstop = true;
 		DesiredView.PostProcessSettings.DepthOfFieldFstop = CurrentAperture;
diff --git a/Engine/Source/Runtime/Core/Public/UObject/RenderingObjectVersion.h b/Engine/Source/Runtime/Core/Public/UObject/RenderingObjectVersion.h
--- a/Engine/Source/Runtime/Core/Public/UObject/RenderingObjectVersion.h
+++ b/Engine/Source/Runtime/Core/Public/UObject/RenderingObjectVersion.h
@@ -81,6 +81,9 @@ struct CORE_API FRenderingObjectVersion
 
 		LightmapHasShadowmapData,
 
+		// Removed old gaussian and bokeh DOF methods from deferred shading renderer.
+		DiaphragmDOFOnlyForDeferredShadingRenderer,
+
 		// -----<new versions can be added above this line>-------------------------------------------------
 		VersionPlusOne,
 		LatestVersion = VersionPlusOne - 1
diff --git a/Engine/Source/Runtime/Engine/Classes/Engine/Scene.h b/Engine/Source/Runtime/Engine/Classes/Engine/Scene.h
--- a/Engine/Source/Runtime/Engine/Classes/Engine/Scene.h
+++ b/Engine/Source/Runtime/Engine/Classes/Engine/Scene.h
@@ -995,38 +995,18 @@ struct FPostProcessSettings
 	UPROPERTY(EditAnywhere, BlueprintReadWrite, Category=Overrides, meta=(PinHiddenByDefault, InlineEditConditionToggle))
 	uint8 bOverride_DepthOfFieldScale:1;
 
-	UE_DEPRECATED(4.21, "This property is now deprecated, please use circle DOF method for deferred shading renderer.")
-	UPROPERTY(EditAnywhere, BlueprintReadWrite, Category=Overrides, meta=(PinHiddenByDefault, InlineEditConditionToggle))
-	uint8 bOverride_DepthOfFieldMaxBokehSize:1;
-
 	UPROPERTY(EditAnywhere, BlueprintReadWrite, Category=Overrides, meta=(PinHiddenByDefault, InlineEditConditionToggle))
 	uint8 bOverride_DepthOfFieldNearBlurSize:1;
 
 	UPROPERTY(EditAnywhere, BlueprintReadWrite, Category=Overrides, meta=(PinHiddenByDefault, InlineEditConditionToggle))
 	uint8 bOverride_DepthOfFieldFarBlurSize:1;
 
-	UE_DEPRECATED(4.21, "This property is now deprecated, please use circle DOF method for deferred shading renderer. Mobile renderer still support Gaussian method.")
-	UPROPERTY(EditAnywhere, BlueprintReadWrite, Category=Overrides, meta=(PinHiddenByDefault, InlineEditConditionToggle))
-	uint8 bOverride_DepthOfFieldMethod:1;
-
 	UPROPERTY(EditAnywhere, BlueprintReadWrite, Category=Overrides, meta=(PinHiddenByDefault, InlineEditConditionToggle))
 	uint8 bOverride_MobileHQGaussian:1;
 
-	UE_DEPRECATED(4.21, "This property is now deprecated, please use circle DOF method for deferred shading renderer.")
-	UPROPERTY(EditAnywhere, BlueprintReadWrite, Category=Overrides, meta=(PinHiddenByDefault, InlineEditConditionToggle))
-	uint8 bOverride_DepthOfFieldBokehShape:1;
-
 	UPROPERTY(EditAnywhere, BlueprintReadWrite, Category=Overrides, meta=(PinHiddenByDefault, InlineEditConditionToggle))
 	uint8 bOverride_DepthOfFieldOcclusion:1;
 
-	UE_DEPRECATED(4.21, "This property is now deprecated, please use circle DOF method for deferred shading renderer.")
-	UPROPERTY(EditAnywhere, BlueprintReadWrite, Category = Overrides, meta = (PinHiddenByDefault, InlineEditConditionToggle))
-	uint8 bOverride_DepthOfFieldColorThreshold:1;
-
-	UE_DEPRECATED(4.21, "This property is now deprecated, please use circle DOF method for deferred shading renderer.")
-	UPROPERTY(EditAnywhere, BlueprintReadWrite, Category = Overrides, meta = (PinHiddenByDefault, InlineEditConditionToggle))
-	uint8 bOverride_DepthOfFieldSizeThreshold:1;
-
 	UPROPERTY(EditAnywhere, BlueprintReadWrite, Category=Overrides, meta=(PinHiddenByDefault, InlineEditConditionToggle))
 	uint8 bOverride_DepthOfFieldSkyFocusDistance:1;
 
@@ -1125,11 +1105,11 @@ struct FPostProcessSettings
 	/** Luminance computation method */
 	UPROPERTY(EditAnywhere, BlueprintReadWrite, Category="Lens|Exposure", meta=(editcondition = "bOverride_AutoExposureMethod", DisplayName = "Metering Mode"))
     TEnumAsByte<enum EAutoExposureMethod> AutoExposureMethod;
-
-	/** DEPRECATED: BokehDOF, Simple gaussian, ... Mobile supports Gaussian only. Deferred shading renderer will only support CircleDOF. */
-	UE_DEPRECATED(4.21, "This property is now deprecated, please use circle DOF method for deferred shading renderer. Mobile renderer still support Gaussian method.")
-	UPROPERTY(EditAnywhere, BlueprintReadWrite, Category="Lens|Depth of Field", meta=(editcondition = "bOverride_DepthOfFieldMethod", DisplayName = "Method"))
-	TEnumAsByte<enum EDepthOfFieldMethod> DepthOfFieldMethod;
+	
+#if WITH_EDITORONLY_DATA
+	UPROPERTY()
+	TEnumAsByte<enum EDepthOfFieldMethod> DepthOfFieldMethod_DEPRECATED;
+#endif
 
 	UPROPERTY(interp, BlueprintReadWrite, Category="Color Grading|WhiteBalance", meta=(UIMin = "1500.0", UIMax = "15000.0", editcondition = "bOverride_WhiteTemp", DisplayName = "Temp"))
 	float WhiteTemp;
@@ -1613,61 +1593,41 @@ struct FPostProcessSettings
 	float DepthOfFieldDepthBlurRadius;
 
 	/** Artificial region where all content is in focus, starting after DepthOfFieldFocalDistance, in unreal units  (cm) */
-	UPROPERTY(interp, BlueprintReadWrite, Category="Lens|Depth of Field", meta=(UIMin = "0.0", UIMax = "10000.0", editcondition = "bOverride_DepthOfFieldFocalRegion", DisplayName = "Focal Region"))
+	UPROPERTY(interp, BlueprintReadWrite, Category="Lens|Mobile Depth of Field", meta=(UIMin = "0.0", UIMax = "10000.0", editcondition = "bOverride_DepthOfFieldFocalRegion", DisplayName = "Focal Region"))
 	float DepthOfFieldFocalRegion;
 
 	/** To define the width of the transition region next to the focal region on the near side (cm) */
-	UPROPERTY(interp, BlueprintReadWrite, Category="Lens|Depth of Field", meta=(UIMin = "0.0", UIMax = "10000.0", editcondition = "bOverride_DepthOfFieldNearTransitionRegion", DisplayName = "Near Transition Region"))
+	UPROPERTY(interp, BlueprintReadWrite, Category="Lens|Mobile Depth of Field", meta=(UIMin = "0.0", UIMax = "10000.0", editcondition = "bOverride_DepthOfFieldNearTransitionRegion", DisplayName = "Near Transition Region"))
 	float DepthOfFieldNearTransitionRegion;
 
 	/** To define the width of the transition region next to the focal region on the near side (cm) */
-	UPROPERTY(interp, BlueprintReadWrite, Category="Lens|Depth of Field", meta=(UIMin = "0.0", UIMax = "10000.0", editcondition = "bOverride_DepthOfFieldFarTransitionRegion", DisplayName = "Far Transition Region"))
+	UPROPERTY(interp, BlueprintReadWrite, Category="Lens|Mobile Depth of Field", meta=(UIMin = "0.0", UIMax = "10000.0", editcondition = "bOverride_DepthOfFieldFarTransitionRegion", DisplayName = "Far Transition Region"))
 	float DepthOfFieldFarTransitionRegion;
 
 	/** SM5: BokehDOF only: To amplify the depth of field effect (like aperture)  0=off 
 	    ES2: Used to blend DoF. 0=off
 	*/
-	UPROPERTY(interp, BlueprintReadWrite, Category="Lens|Depth of Field", meta=(ClampMin = "0.0", ClampMax = "2.0", editcondition = "bOverride_DepthOfFieldScale", DisplayName = "Scale"))
+	UPROPERTY(interp, BlueprintReadWrite, Category="Lens|Mobile Depth of Field", meta=(ClampMin = "0.0", ClampMax = "2.0", editcondition = "bOverride_DepthOfFieldScale", DisplayName = "Scale"))
 	float DepthOfFieldScale;
 
-	/** DEPRECATED: BokehDOF only: Maximum size of the Depth of Field blur (in percent of the view width) (note: performance cost scales with size*size) */
-	UE_DEPRECATED(4.21, "This property is now deprecated, please use circle DOF method for deferred shading renderer.")
-	UPROPERTY(interp, BlueprintReadWrite, Category="Lens|Depth of Field", meta=(UIMin = "0.0", UIMax = "32.0", editcondition = "bOverride_DepthOfFieldMaxBokehSize", DisplayName = "Max Bokeh Size"))
-	float DepthOfFieldMaxBokehSize;
-
 	/** Gaussian only: Maximum size of the Depth of Field blur (in percent of the view width) (note: performance cost scales with size) */
-	UPROPERTY(interp, BlueprintReadWrite, Category="Lens|Depth of Field", meta=(UIMin = "0.0", UIMax = "32.0", editcondition = "bOverride_DepthOfFieldNearBlurSize", DisplayName = "Near Blur Size"))
+	UPROPERTY(interp, BlueprintReadWrite, Category="Lens|Mobile Depth of Field", meta=(UIMin = "0.0", UIMax = "32.0", editcondition = "bOverride_DepthOfFieldNearBlurSize", DisplayName = "Near Blur Size"))
 	float DepthOfFieldNearBlurSize;
 
 	/** Gaussian only: Maximum size of the Depth of Field blur (in percent of the view width) (note: performance cost scales with size) */
-	UPROPERTY(interp, BlueprintReadWrite, Category="Lens|Depth of Field", meta=(UIMin = "0.0", UIMax = "32.0", editcondition = "bOverride_DepthOfFieldFarBlurSize", DisplayName = "Far Blur Size"))
+	UPROPERTY(interp, BlueprintReadWrite, Category="Lens|Mobile Depth of Field", meta=(UIMin = "0.0", UIMax = "32.0", editcondition = "bOverride_DepthOfFieldFarBlurSize", DisplayName = "Far Blur Size"))
 	float DepthOfFieldFarBlurSize;
 
 	/** Occlusion tweak factor 1 (0.18 to get natural occlusion, 0.4 to solve layer color leaking issues) */
-	UPROPERTY(interp, BlueprintReadWrite, Category="Lens|Depth of Field", AdvancedDisplay, meta=(ClampMin = "0.0", ClampMax = "1.0", editcondition = "bOverride_DepthOfFieldOcclusion", DisplayName = "Occlusion"))
+	UPROPERTY(interp, BlueprintReadWrite, Category="Lens|Mobile Depth of Field", AdvancedDisplay, meta=(ClampMin = "0.0", ClampMax = "1.0", editcondition = "bOverride_DepthOfFieldOcclusion", DisplayName = "Occlusion"))
 	float DepthOfFieldOcclusion;
 
-	/** DEPRECATED: Defines the shape of the Bokeh when object get out of focus, cannot be blended */
-	UE_DEPRECATED(4.21, "This property is now deprecated, please use circle DOF method for deferred shading renderer.")
-	UPROPERTY(EditAnywhere, BlueprintReadWrite, AdvancedDisplay, Category="Lens|Depth of Field", meta=(editcondition = "bOverride_DepthOfFieldBokehShape", DisplayName = "Shape"))
-	class UTexture* DepthOfFieldBokehShape;
-
-	/** DEPRECATED: Color threshold to do full quality DOF (BokehDOF only) */
-	UE_DEPRECATED(4.21, "This property is now deprecated, please use circle DOF method for deferred shading renderer.")
-	UPROPERTY(interp, BlueprintReadWrite, Category="Lens|Depth of Field", AdvancedDisplay, meta=(ClampMin = "0.0", ClampMax = "10.0", editcondition = "bOverride_DepthOfFieldColorThreshold", DisplayName = "Color Threshold"))
-	float DepthOfFieldColorThreshold;
-
-	/** DEPRECATED: Size threshold to do full quality DOF (BokehDOF only) */
-	UE_DEPRECATED(4.21, "This property is now deprecated, please use circle DOF method for deferred shading renderer.")
-	UPROPERTY(interp, BlueprintReadWrite, Category="Lens|Depth of Field", AdvancedDisplay, meta=(ClampMin = "0.0", ClampMax = "1.0", editcondition = "bOverride_DepthOfFieldSizeThreshold", DisplayName = "Size Threshold"))
-	float DepthOfFieldSizeThreshold;
-	
 	/** Artificial distance to allow the skybox to be in focus (e.g. 200000), <=0 to switch the feature off, only for GaussianDOF, can cost performance */
-	UPROPERTY(interp, BlueprintReadWrite, Category="Lens|Depth of Field", AdvancedDisplay, meta=(ClampMin = "0.0", ClampMax = "200000.0", editcondition = "bOverride_DepthOfFieldSkyFocusDistance", DisplayName = "Sky Distance"))
+	UPROPERTY(interp, BlueprintReadWrite, Category="Lens|Mobile Depth of Field", AdvancedDisplay, meta=(ClampMin = "0.0", ClampMax = "200000.0", editcondition = "bOverride_DepthOfFieldSkyFocusDistance", DisplayName = "Sky Distance"))
 	float DepthOfFieldSkyFocusDistance;
 
 	/** Artificial circular mask to (near) blur content outside the radius, only for GaussianDOF, diameter in percent of screen width, costs performance if the mask is used, keep Feather can Radius on default to keep it off */
-	UPROPERTY(interp, BlueprintReadWrite, Category="Lens|Depth of Field", AdvancedDisplay, meta=(UIMin = "0.0", UIMax = "100.0", editcondition = "bOverride_DepthOfFieldVignetteSize", DisplayName = "Vignette Size"))
+	UPROPERTY(interp, BlueprintReadWrite, Category="Lens|Mobile Depth of Field", AdvancedDisplay, meta=(UIMin = "0.0", UIMax = "100.0", editcondition = "bOverride_DepthOfFieldVignetteSize", DisplayName = "Vignette Size"))
 	float DepthOfFieldVignetteSize;
 
 	/** Strength of motion blur, 0:off, should be renamed to intensity */
@@ -1870,18 +1830,6 @@ struct FPostProcessSettings
 			BloomConvolutionPreFilterMax = BloomConvolutionPreFilter_DEPRECATED.Y;
 			BloomConvolutionPreFilterMult = BloomConvolutionPreFilter_DEPRECATED.Z;
 		}
-
-		/** Actually delete settings related to BokehDOF, and mark the DepthOfFieldMethod to _DEPRECATED. To keep backward compatibility,
-		 * need to have depth of field disabled by default, which means for deferred shading renderer having redicously high aperture.
-		 *
-		 * The plan here is to change the default of DepthOfFieldFstop in FPostProcessSettings's constructor to this MyMagicHightFstopValue.
-		 * And here in this OnAfterLoad() method do:
-		 *
-		 * if (DepthOfFieldFstop == MyMagicHightFstopValue && DepthOfFieldMethod_DEPRECATED == DOFM_CircleDOF)
-		 * {
-		 *		DepthOfFieldFstop = 4.0f;
-		 * }
-		 */
 	}
 #endif
 
@@ -1946,11 +1894,29 @@ struct FPostProcessSettings
 
 	// Default number of blade of the diaphragm to simulate in depth of field.
 	static constexpr int32 kDefaultDepthOfFieldBladeCount = 5;
+	
+#if WITH_EDITORONLY_DATA
+	bool Serialize(FArchive& Ar);
+	void PostSerialize(const FArchive& Ar);
+#endif
 
 	PRAGMA_DISABLE_DEPRECATION_WARNINGS
-};
+}; // struct FPostProcessSettings
 PRAGMA_ENABLE_DEPRECATION_WARNINGS
 
+#if WITH_EDITORONLY_DATA
+template<>
+struct TStructOpsTypeTraits<FPostProcessSettings>
+	: public TStructOpsTypeTraitsBase2<FPostProcessSettings>
+{
+	enum
+	{
+		WithSerializer = true,
+		WithPostSerialize = true,
+	};
+};
+#endif
+
 UCLASS()
 class UScene : public UObject
 {
diff --git a/Engine/Source/Runtime/Engine/Private/PostProcessVolume.cpp b/Engine/Source/Runtime/Engine/Private/PostProcessVolume.cpp
--- a/Engine/Source/Runtime/Engine/Private/PostProcessVolume.cpp
+++ b/Engine/Source/Runtime/Engine/Private/PostProcessVolume.cpp
@@ -67,33 +67,21 @@ bool APostProcessVolume::CanEditChange(const UProperty* InProperty) const
 		// Settings, can be shared for multiple objects types (volume, component, camera, player)
 		PRAGMA_DISABLE_DEPRECATION_WARNINGS
 		{
-			if (PropertyName == GET_MEMBER_NAME_STRING_CHECKED(FPostProcessSettings, DepthOfFieldScale))
-			{
-				return	( Settings.DepthOfFieldMethod == EDepthOfFieldMethod::DOFM_BokehDOF ||
-						  Settings.DepthOfFieldMethod == EDepthOfFieldMethod::DOFM_Gaussian );
-			}
-
-			if (PropertyName == GET_MEMBER_NAME_STRING_CHECKED(FPostProcessSettings, DepthOfFieldMaxBokehSize) ||
-				PropertyName == GET_MEMBER_NAME_STRING_CHECKED(FPostProcessSettings, DepthOfFieldColorThreshold) ||
-				PropertyName == GET_MEMBER_NAME_STRING_CHECKED(FPostProcessSettings, DepthOfFieldSizeThreshold) ||
-				PropertyName == GET_MEMBER_NAME_STRING_CHECKED(FPostProcessSettings, DepthOfFieldBokehShape))
-			{
-				return Settings.DepthOfFieldMethod == EDepthOfFieldMethod::DOFM_BokehDOF;
-			}
+			FSceneInterface* Scene = GetWorld()->Scene;
+			bool bIsMobile = Scene->GetShadingPath(Scene->GetFeatureLevel()) == EShadingPath::Mobile;
+			bool bHaveGaussianDOF = bIsMobile;
+			bool bHaveCinematicDOF = !bIsMobile;
 
-			if (PropertyName == GET_MEMBER_NAME_STRING_CHECKED(FPostProcessSettings, DepthOfFieldNearBlurSize) ||
+			if (PropertyName == GET_MEMBER_NAME_STRING_CHECKED(FPostProcessSettings, DepthOfFieldScale) ||
+				PropertyName == GET_MEMBER_NAME_STRING_CHECKED(FPostProcessSettings, DepthOfFieldNearBlurSize) ||
 				PropertyName == GET_MEMBER_NAME_STRING_CHECKED(FPostProcessSettings, DepthOfFieldFarBlurSize) ||
 				PropertyName == GET_MEMBER_NAME_STRING_CHECKED(FPostProcessSettings, DepthOfFieldSkyFocusDistance) ||
-				PropertyName == GET_MEMBER_NAME_STRING_CHECKED(FPostProcessSettings, DepthOfFieldVignetteSize))
-			{
-				return Settings.DepthOfFieldMethod == EDepthOfFieldMethod::DOFM_Gaussian;
-			}
-
-			if (PropertyName == GET_MEMBER_NAME_STRING_CHECKED(FPostProcessSettings, DepthOfFieldNearTransitionRegion) ||
+				PropertyName == GET_MEMBER_NAME_STRING_CHECKED(FPostProcessSettings, DepthOfFieldVignetteSize) ||
+				PropertyName == GET_MEMBER_NAME_STRING_CHECKED(FPostProcessSettings, DepthOfFieldNearTransitionRegion) ||
 				PropertyName == GET_MEMBER_NAME_STRING_CHECKED(FPostProcessSettings, DepthOfFieldFarTransitionRegion) ||
 				PropertyName == GET_MEMBER_NAME_STRING_CHECKED(FPostProcessSettings, DepthOfFieldFocalRegion))
 			{
-				return Settings.DepthOfFieldMethod != EDepthOfFieldMethod::DOFM_CircleDOF;
+				return bHaveGaussianDOF;
 			}
 
 			if (PropertyName == GET_MEMBER_NAME_STRING_CHECKED(FPostProcessSettings, DepthOfFieldDepthBlurAmount) ||
@@ -101,12 +89,12 @@ bool APostProcessVolume::CanEditChange(const UProperty* InProperty) const
 				PropertyName == GET_MEMBER_NAME_STRING_CHECKED(FPostProcessSettings, DepthOfFieldMinFstop) ||
 				PropertyName == GET_MEMBER_NAME_STRING_CHECKED(FPostProcessSettings, DepthOfFieldBladeCount))
 			{
-				return Settings.DepthOfFieldMethod == EDepthOfFieldMethod::DOFM_CircleDOF;
+				return bHaveCinematicDOF;
 			}
 
 			if (PropertyName == GET_MEMBER_NAME_STRING_CHECKED(FPostProcessSettings, DepthOfFieldFstop))
 			{
-				return	( Settings.DepthOfFieldMethod == EDepthOfFieldMethod::DOFM_CircleDOF || 
+				return	( bHaveCinematicDOF || 
 					      Settings.AutoExposureMethod == EAutoExposureMethod::AEM_Manual );
 			}
 
diff --git a/Engine/Source/Runtime/Engine/Private/Scene.cpp b/Engine/Source/Runtime/Engine/Private/Scene.cpp
--- a/Engine/Source/Runtime/Engine/Private/Scene.cpp
+++ b/Engine/Source/Runtime/Engine/Private/Scene.cpp
@@ -514,9 +514,8 @@ FPostProcessSettings::FPostProcessSettings()
 	RayTracingGIMaxBounces = 1;
 	RayTracingGISamplesPerPixel = 1;
 
-	PRAGMA_DISABLE_DEPRECATION_WARNINGS
 	DepthOfFieldFocalDistance = 1000.0f;
-	DepthOfFieldFstop = 4.0f;
+	DepthOfFieldFstop = 0.0f; // Intentionally invalid to disable DOF by default.
 	DepthOfFieldMinFstop = 1.2f;
 	DepthOfFieldBladeCount = FPostProcessSettings::kDefaultDepthOfFieldBladeCount;
 	DepthOfFieldSensorWidth = 24.576f;			// APS-C
@@ -526,14 +525,10 @@ FPostProcessSettings::FPostProcessSettings()
 	DepthOfFieldNearTransitionRegion = 300.0f;
 	DepthOfFieldFarTransitionRegion = 500.0f;
 	DepthOfFieldScale = 0.0f;
-	DepthOfFieldMaxBokehSize = 15.0f;
 	DepthOfFieldNearBlurSize = 15.0f;
 	DepthOfFieldFarBlurSize = 15.0f;
 	DepthOfFieldOcclusion = 0.4f;
-	DepthOfFieldColorThreshold = 1.0f;
-	DepthOfFieldSizeThreshold = 0.08f;
 	DepthOfFieldSkyFocusDistance = 0.0f;
-	PRAGMA_ENABLE_DEPRECATION_WARNINGS
 	// 200 should be enough even for extreme aspect ratios to give the default no effect
 	DepthOfFieldVignetteSize = 200.0f;
 	LensFlareTints[0] = FLinearColor(1.0f, 0.8f, 0.4f, 0.6f);
@@ -713,7 +708,6 @@ FPostProcessSettings::FPostProcessSettings(const FPostProcessSettings& Settings)
 	, bOverride_IndirectLightingIntensity(Settings.bOverride_IndirectLightingIntensity)
 	, bOverride_ColorGradingIntensity(Settings.bOverride_ColorGradingIntensity)
 	, bOverride_ColorGradingLUT(Settings.bOverride_ColorGradingLUT)
-	PRAGMA_DISABLE_DEPRECATION_WARNINGS
 	, bOverride_DepthOfFieldFocalDistance(Settings.bOverride_DepthOfFieldFocalDistance)
 	, bOverride_DepthOfFieldFstop(Settings.bOverride_DepthOfFieldFstop)
 	, bOverride_DepthOfFieldMinFstop(Settings.bOverride_DepthOfFieldMinFstop)
@@ -725,18 +719,12 @@ FPostProcessSettings::FPostProcessSettings(const FPostProcessSettings& Settings)
 	, bOverride_DepthOfFieldNearTransitionRegion(Settings.bOverride_DepthOfFieldNearTransitionRegion)
 	, bOverride_DepthOfFieldFarTransitionRegion(Settings.bOverride_DepthOfFieldFarTransitionRegion)
 	, bOverride_DepthOfFieldScale(Settings.bOverride_DepthOfFieldScale)
-	, bOverride_DepthOfFieldMaxBokehSize(Settings.bOverride_DepthOfFieldMaxBokehSize)
 	, bOverride_DepthOfFieldNearBlurSize(Settings.bOverride_DepthOfFieldNearBlurSize)
 	, bOverride_DepthOfFieldFarBlurSize(Settings.bOverride_DepthOfFieldFarBlurSize)
-	, bOverride_DepthOfFieldMethod(Settings.bOverride_DepthOfFieldMethod)
 	, bOverride_MobileHQGaussian(Settings.bOverride_MobileHQGaussian)
-	, bOverride_DepthOfFieldBokehShape(Settings.bOverride_DepthOfFieldBokehShape)
 	, bOverride_DepthOfFieldOcclusion(Settings.bOverride_DepthOfFieldOcclusion)
-	, bOverride_DepthOfFieldColorThreshold(Settings.bOverride_DepthOfFieldColorThreshold)
-	, bOverride_DepthOfFieldSizeThreshold(Settings.bOverride_DepthOfFieldSizeThreshold)
 	, bOverride_DepthOfFieldSkyFocusDistance(Settings.bOverride_DepthOfFieldSkyFocusDistance)
 	, bOverride_DepthOfFieldVignetteSize(Settings.bOverride_DepthOfFieldVignetteSize)
-	PRAGMA_ENABLE_DEPRECATION_WARNINGS
 	, bOverride_MotionBlurAmount(Settings.bOverride_MotionBlurAmount)
 	, bOverride_MotionBlurMax(Settings.bOverride_MotionBlurMax)
 	, bOverride_MotionBlurPerObjectSize(Settings.bOverride_MotionBlurPerObjectSize)
@@ -765,9 +753,7 @@ FPostProcessSettings::FPostProcessSettings(const FPostProcessSettings& Settings)
 	, bMobileHQGaussian(Settings.bMobileHQGaussian)
 	, BloomMethod(Settings.BloomMethod)
 	, AutoExposureMethod(Settings.AutoExposureMethod)
-	PRAGMA_DISABLE_DEPRECATION_WARNINGS
-	, DepthOfFieldMethod(Settings.DepthOfFieldMethod)
-	PRAGMA_ENABLE_DEPRECATION_WARNINGS
+	//, DepthOfFieldMethod_DEPRECATED(Settings.DepthOfFieldMethod_DEPRECATED)
 	, WhiteTemp(Settings.WhiteTemp)
 	, WhiteTint(Settings.WhiteTint)
 	, ColorSaturation(Settings.ColorSaturation)
@@ -888,7 +874,6 @@ FPostProcessSettings::FPostProcessSettings(const FPostProcessSettings& Settings)
 	, RayTracingGISamplesPerPixel(Settings.RayTracingGISamplesPerPixel)
 	, ColorGradingIntensity(Settings.ColorGradingIntensity)
 	, ColorGradingLUT(Settings.ColorGradingLUT)
-	PRAGMA_DISABLE_DEPRECATION_WARNINGS
 	, DepthOfFieldSensorWidth(Settings.DepthOfFieldSensorWidth)
 	, DepthOfFieldFocalDistance(Settings.DepthOfFieldFocalDistance)
 	, DepthOfFieldDepthBlurAmount(Settings.DepthOfFieldDepthBlurAmount)
@@ -897,16 +882,11 @@ FPostProcessSettings::FPostProcessSettings(const FPostProcessSettings& Settings)
 	, DepthOfFieldNearTransitionRegion(Settings.DepthOfFieldNearTransitionRegion)
 	, DepthOfFieldFarTransitionRegion(Settings.DepthOfFieldFarTransitionRegion)
 	, DepthOfFieldScale(Settings.DepthOfFieldScale)
-	, DepthOfFieldMaxBokehSize(Settings.DepthOfFieldMaxBokehSize)
 	, DepthOfFieldNearBlurSize(Settings.DepthOfFieldNearBlurSize)
 	, DepthOfFieldFarBlurSize(Settings.DepthOfFieldFarBlurSize)
 	, DepthOfFieldOcclusion(Settings.DepthOfFieldOcclusion)
-	, DepthOfFieldBokehShape(Settings.DepthOfFieldBokehShape)
-	, DepthOfFieldColorThreshold(Settings.DepthOfFieldColorThreshold)
-	, DepthOfFieldSizeThreshold(Settings.DepthOfFieldSizeThreshold)
 	, DepthOfFieldSkyFocusDistance(Settings.DepthOfFieldSkyFocusDistance)
 	, DepthOfFieldVignetteSize(Settings.DepthOfFieldVignetteSize)
-	PRAGMA_ENABLE_DEPRECATION_WARNINGS
 	, MotionBlurAmount(Settings.MotionBlurAmount)
 	, MotionBlurMax(Settings.MotionBlurMax)
 	, MotionBlurTargetFPS(Settings.MotionBlurTargetFPS)
@@ -956,6 +936,34 @@ FPostProcessSettings::FPostProcessSettings(const FPostProcessSettings& Settings)
 	for (int32 i = 0; i < ARRAY_COUNT(LensFlareTints); i++)
 		LensFlareTints[i] = Settings.LensFlareTints[i];
 }
+	
+#if WITH_EDITORONLY_DATA
+bool FPostProcessSettings::Serialize(FArchive& Ar)
+{
+	Ar.UsingCustomVersion(FRenderingObjectVersion::GUID);
+
+	// Don't actually serialize, just write the custom version for PostSerialize
+	return false;
+}
+
+void FPostProcessSettings::PostSerialize(const FArchive& Ar)
+{
+	if (Ar.IsLoading() && Ar.CustomVer(FRenderingObjectVersion::GUID) < FRenderingObjectVersion::DiaphragmDOFOnlyForDeferredShadingRenderer)
+	{
+		// Make sure the DOF of the deferred shading renderer is enabled if the circle DOF method was used before with previous default setting for DepthOfFieldFstop.
+		if (DepthOfFieldFstop == 0.0f && DepthOfFieldMethod_DEPRECATED == DOFM_CircleDOF)
+		{
+			DepthOfFieldFstop = 4.0f;
+		}
+
+		// Make sure gaussian DOF is disabled on mobile if the DOF method was set to something else.
+		if (DepthOfFieldMethod_DEPRECATED != DOFM_Gaussian)
+		{
+			DepthOfFieldScale = 0.0f;
+		}
+	}
+}
+#endif
 
 UScene::UScene(const FObjectInitializer& ObjectInitializer)
 	: Super(ObjectInitializer)
diff --git a/Engine/Source/Runtime/Engine/Private/SceneView.cpp b/Engine/Source/Runtime/Engine/Private/SceneView.cpp
--- a/Engine/Source/Runtime/Engine/Private/SceneView.cpp
+++ b/Engine/Source/Runtime/Engine/Private/SceneView.cpp
@@ -1349,16 +1349,9 @@ void FSceneView::OverridePostProcessSettings(const FPostProcessSettings& Src, fl
 		LERP_PP(DepthOfFieldNearTransitionRegion);
 		LERP_PP(DepthOfFieldFarTransitionRegion);
 		LERP_PP(DepthOfFieldScale);
-		PRAGMA_DISABLE_DEPRECATION_WARNINGS
-		LERP_PP(DepthOfFieldMaxBokehSize);
-		PRAGMA_ENABLE_DEPRECATION_WARNINGS
 		LERP_PP(DepthOfFieldNearBlurSize);
 		LERP_PP(DepthOfFieldFarBlurSize);
 		LERP_PP(DepthOfFieldOcclusion);
-		PRAGMA_DISABLE_DEPRECATION_WARNINGS
-		LERP_PP(DepthOfFieldColorThreshold);
-		LERP_PP(DepthOfFieldSizeThreshold);
-		PRAGMA_ENABLE_DEPRECATION_WARNINGS
 		LERP_PP(DepthOfFieldSkyFocusDistance);
 		LERP_PP(DepthOfFieldVignetteSize);
 		LERP_PP(MotionBlurAmount);
@@ -1505,14 +1498,6 @@ void FSceneView::OverridePostProcessSettings(const FPostProcessSettings& Src, fl
 			Dest.BloomConvolutionBufferScale = Src.BloomConvolutionBufferScale;
 		}
 
-		PRAGMA_DISABLE_DEPRECATION_WARNINGS
-		// actual texture cannot be blended but the intensity can be blended
-		IF_PP(DepthOfFieldBokehShape)
-		{
-			Dest.DepthOfFieldBokehShape = Src.DepthOfFieldBokehShape;
-		}
-		PRAGMA_ENABLE_DEPRECATION_WARNINGS
-
 		// Curve assets can not be blended.
 		IF_PP(AutoExposureBiasCurve)
 		{
@@ -1533,13 +1518,6 @@ void FSceneView::OverridePostProcessSettings(const FPostProcessSettings& Src, fl
 			}
 		}
 
-		PRAGMA_DISABLE_DEPRECATION_WARNINGS
-		if (Src.bOverride_DepthOfFieldMethod)
-		{
-			Dest.DepthOfFieldMethod = Src.DepthOfFieldMethod;
-		}
-		PRAGMA_ENABLE_DEPRECATION_WARNINGS
-
 		if (Src.bOverride_MobileHQGaussian)
 		{
 			Dest.bMobileHQGaussian = Src.bMobileHQGaussian;
@@ -1906,15 +1884,6 @@ void FSceneView::EndFinalPostprocessSettings(const FSceneViewInitOptions& ViewIn
 	}
 #endif
 
-	PRAGMA_DISABLE_DEPRECATION_WARNINGS
-	if(FinalPostProcessSettings.DepthOfFieldMethod == DOFM_CircleDOF)
-	{
-		// We intentionally don't do the DepthOfFieldFocalRegion as it breaks realism.
-		// Doing this fixes DOF material expression.
-		FinalPostProcessSettings.DepthOfFieldFocalRegion = 0;
-	}
-	PRAGMA_ENABLE_DEPRECATION_WARNINGS
-
 	{
 		const bool bStereoEnabled = StereoPass != eSSP_FULL;
 		const bool bScaledToRenderTarget = GEngine->XRSystem.IsValid() && bStereoEnabled && GEngine->XRSystem->GetHMDDevice();
diff --git a/Engine/Source/Runtime/Renderer/Private/PostProcess/DiaphragmDOF.h b/Engine/Source/Runtime/Renderer/Private/PostProcess/DiaphragmDOF.h
--- a/Engine/Source/Runtime/Renderer/Private/PostProcess/DiaphragmDOF.h
+++ b/Engine/Source/Runtime/Renderer/Private/PostProcess/DiaphragmDOF.h
@@ -22,6 +22,9 @@ struct FTemporalAAHistory;
 
 namespace DiaphragmDOF
 {
+	
+float ComputeFocalLengthFromFov(const FSceneView& View);
+FVector4 CircleDofHalfCoc(const FViewInfo& View);
 
 /** Physically based circle of confusion computation model. */
 struct FPhysicalCocModel
diff --git a/Engine/Source/Runtime/Renderer/Private/PostProcess/DiaphragmDOFUtils.cpp b/Engine/Source/Runtime/Renderer/Private/PostProcess/DiaphragmDOFUtils.cpp
--- a/Engine/Source/Runtime/Renderer/Private/PostProcess/DiaphragmDOFUtils.cpp
+++ b/Engine/Source/Runtime/Renderer/Private/PostProcess/DiaphragmDOFUtils.cpp
@@ -5,7 +5,6 @@
 =============================================================================*/
 
 #include "PostProcess/DiaphragmDOF.h"
-#include "PostProcess/PostProcessCircleDOF.h"
 
 
 namespace 
@@ -25,6 +24,79 @@ TAutoConsoleVariable<float> CVarMaxBackgroundRadius(
 
 } // namespace
 
+float DiaphragmDOF::ComputeFocalLengthFromFov(const FSceneView& View)
+{
+	// Convert FOV to focal length,
+	// 
+	// fov = 2 * atan(d/(2*f))
+	// where,
+	//   d = sensor dimension (APS-C 24.576 mm)
+	//   f = focal length
+	// 
+	// f = 0.5 * d * (1/tan(fov/2))
+
+	float const d = View.FinalPostProcessSettings.DepthOfFieldSensorWidth;
+	float const HalfFOV = FMath::Atan(1.0f / View.ViewMatrices.GetProjectionMatrix().M[0][0]);
+	float const FocalLength = 0.5f * d * (1.0f/FMath::Tan(HalfFOV));
+
+	return FocalLength;
+}
+
+// Convert f-stop and focal distance into projected size in half resolution pixels.
+// Setup depth based blur.
+FVector4 DiaphragmDOF::CircleDofHalfCoc(const FViewInfo& View)
+{
+	static const auto CVar = IConsoleManager::Get().FindTConsoleVariableDataInt(TEXT("r.DepthOfFieldQuality"));
+	bool bDepthOfField = View.Family->EngineShowFlags.DepthOfField && CVar->GetValueOnRenderThread() > 0 && View.FinalPostProcessSettings.DepthOfFieldFstop > 0;
+
+	FVector4 Ret(0, 1, 0, 0);
+
+	if(bDepthOfField)
+	{
+		float FocalLengthInMM = DiaphragmDOF::ComputeFocalLengthFromFov(View);
+	 
+		// Convert focal distance in world position to mm (from cm to mm)
+		float FocalDistanceInMM = View.FinalPostProcessSettings.DepthOfFieldFocalDistance * 10.0f;
+
+		// Convert f-stop, focal length, and focal distance to
+		// projected circle of confusion size at infinity in mm.
+		//
+		// coc = f * f / (n * (d - f))
+		// where,
+		//   f = focal length
+		//   d = focal distance
+		//   n = fstop (where n is the "n" in "f/n")
+		float Radius = FMath::Square(FocalLengthInMM) / (View.FinalPostProcessSettings.DepthOfFieldFstop * (FocalDistanceInMM - FocalLengthInMM));
+
+		// Convert mm to pixels.
+		float const Width = (float)View.ViewRect.Width();
+		float const SensorWidth = View.FinalPostProcessSettings.DepthOfFieldSensorWidth;
+		Radius = Radius * Width * (1.0f / SensorWidth);
+
+		// Convert diameter to radius at half resolution (algorithm radius is at half resolution).
+		Radius *= 0.25f;
+
+		// Comment out for now, allowing settings which the algorithm cannot cleanly do.
+		#if 0
+			// Limit to algorithm max size.
+			if(Radius > 6.0f) 
+			{
+				Radius = 6.0f; 
+			}
+		#endif
+
+		// The DepthOfFieldDepthBlurAmount = km at which depth blur is 50%.
+		// Need to convert to cm here.
+		Ret = FVector4(
+			Radius, 
+			1.0f / (View.FinalPostProcessSettings.DepthOfFieldDepthBlurAmount * 100000.0f),
+			View.FinalPostProcessSettings.DepthOfFieldDepthBlurRadius * Width / 1920.0f,
+			Width / 1920.0f);
+	}
+
+	return Ret;
+}
+
 void DiaphragmDOF::FPhysicalCocModel::Compile(const FViewInfo& View)
 {
 	// Fetches DOF settings.
@@ -46,7 +118,7 @@ void DiaphragmDOF::FPhysicalCocModel::Compile(const FViewInfo& View)
 	// Compile coc model equation.
 	{
 
-		float FocalLengthInMM = ComputeFocalLengthFromFov(View);
+		float FocalLengthInMM = DiaphragmDOF::ComputeFocalLengthFromFov(View);
 
 		// Convert focal distance in world position to mm (from cm to mm)
 		float FocalDistanceInMM = View.FinalPostProcessSettings.DepthOfFieldFocalDistance * 10.0f;
diff --git a/Engine/Source/Runtime/Renderer/Private/PostProcess/PostProcessBokehDOF.cpp b/Engine/Source/Runtime/Renderer/Private/PostProcess/PostProcessBokehDOF.cpp
--- a/Engine/Source/Runtime/Renderer/Private/PostProcess/PostProcessBokehDOF.cpp
+++ b/Engine/Source/Runtime/Renderer/Private/PostProcess/PostProcessBokehDOF.cpp
@@ -5,6 +5,7 @@
 =============================================================================*/
 
 #include "PostProcess/PostProcessBokehDOF.h"
+#include "PostProcess/PostProcessDOF.h"
 #include "EngineGlobals.h"
 #include "StaticBoundShaderState.h"
 #include "CanvasTypes.h"
@@ -15,10 +16,10 @@
 #include "PostProcess/SceneFilterRendering.h"
 #include "SceneRenderTargetParameters.h"
 #include "PostProcess/PostProcessing.h"
-#include "PostProcess/PostProcessCircleDOF.h"
 #include "ClearQuad.h"
 #include "PipelineStateCache.h"
 #include "SpriteIndexBuffer.h"
+#include "DiaphragmDOF.h"
 
 const int32 GBokehDOFSetupTileSizeX = 8;
 const int32 GBokehDOFSetupTileSizeY = 8;
@@ -89,7 +90,7 @@ public:
 			// in rendertarget pixels (half res to scene color)
 			FRenderingCompositeOutput* Output = Context.Pass->GetOutput(ePId_Output0);
 
-			FRCPassPostProcessBokehDOF::ComputeDepthOfFieldParams(Context, DepthOfFieldParamValues);
+			FRCPassPostProcessDOFSetup::ComputeDepthOfFieldParams(Context, DepthOfFieldParamValues);
 
 			SetShaderValueArray(RHICmdList, ShaderRHI, DepthOfFieldParams, DepthOfFieldParamValues, 2);
 		}
@@ -99,7 +100,7 @@ public:
 			FIntPoint CursorPosValue(-100,-100);
 
 			PRAGMA_DISABLE_DEPRECATION_WARNINGS
-			if(Context.View.FinalPostProcessSettings.DepthOfFieldMethod == DOFM_CircleDOF)
+			if(Context.View.FinalPostProcessSettings.DepthOfFieldMethod_DEPRECATED == DOFM_CircleDOF)
 			{
 				CursorPosValue = Context.View.CursorPos;
 			}
@@ -227,734 +228,80 @@ void FRCPassPostProcessVisualizeDOF::Process(FRenderingCompositePassContext& Con
 		Canvas.DrawShadowedString(20, Y += YStep, *Line, GetStatsFont(), FLinearColor(1, 1, 0));
 		Y += YStep;
 
-		PRAGMA_DISABLE_DEPRECATION_WARNINGS
-			EDepthOfFieldMethod MethodId = View.FinalPostProcessSettings.DepthOfFieldMethod;
-		PRAGMA_ENABLE_DEPRECATION_WARNINGS
-
-			if (MethodId == DOFM_BokehDOF)
-			{
-				PRAGMA_DISABLE_DEPRECATION_WARNINGS
-					Line = FString::Printf(TEXT("Method: BokehDOF (blue is far, green is near, black is in focus)"));
-				Canvas.DrawShadowedString(X, Y += YStep, *Line, GetStatsFont(), FLinearColor(1, 1, 1));
-				Y += YStep;
-				Line = FString::Printf(TEXT("FocalDistance: %.2f"), View.FinalPostProcessSettings.DepthOfFieldFocalDistance);
-				Canvas.DrawShadowedString(X, Y += YStep, *Line, GetStatsFont(), FLinearColor(1, 1, 1));
-				Line = FString::Printf(TEXT("FocalRegion (Artificial, avoid): %.2f"), View.FinalPostProcessSettings.DepthOfFieldFocalRegion);
-				Canvas.DrawShadowedString(X, Y += YStep, *Line, GetStatsFont(), FLinearColor(1, 1, 1));
-				Y += YStep;
-				Line = FString::Printf(TEXT("Scale: %.2f"), View.FinalPostProcessSettings.DepthOfFieldScale);
-				Canvas.DrawShadowedString(X, Y += YStep, *Line, GetStatsFont(), FLinearColor(1, 1, 1));
-				Line = FString::Printf(TEXT("MaxBokehSize: %.2f"), View.FinalPostProcessSettings.DepthOfFieldMaxBokehSize);
-				Canvas.DrawShadowedString(X, Y += YStep, *Line, GetStatsFont(), FLinearColor(1, 1, 1));
-				Line = FString::Printf(TEXT("NearTransitionRegion: %.2f"), View.FinalPostProcessSettings.DepthOfFieldNearTransitionRegion);
-				Canvas.DrawShadowedString(X, Y += YStep, *Line, GetStatsFont(), FLinearColor(1, 1, 1));
-				Line = FString::Printf(TEXT("FarTransitionRegion: %.2f"), View.FinalPostProcessSettings.DepthOfFieldFarTransitionRegion);
-				Canvas.DrawShadowedString(X, Y += YStep, *Line, GetStatsFont(), FLinearColor(1, 1, 1));
-				Line = FString::Printf(TEXT("ColorThreshold: %.2f"), View.FinalPostProcessSettings.DepthOfFieldColorThreshold);
-				Canvas.DrawShadowedString(X, Y += YStep, *Line, GetStatsFont(), FLinearColor(1, 1, 1));
-				Line = FString::Printf(TEXT("SizeThreshold: %.2f"), View.FinalPostProcessSettings.DepthOfFieldSizeThreshold);
-				Canvas.DrawShadowedString(X, Y += YStep, *Line, GetStatsFont(), FLinearColor(1, 1, 1));
-				Line = FString::Printf(TEXT("Occlusion: %.2f"), View.FinalPostProcessSettings.DepthOfFieldOcclusion);
-				Canvas.DrawShadowedString(X, Y += YStep, *Line, GetStatsFont(), FLinearColor(1, 1, 1));
-				PRAGMA_ENABLE_DEPRECATION_WARNINGS
-			}
-			else if (MethodId == DOFM_Gaussian)
-			{
-				Line = FString::Printf(TEXT("Method: GaussianDOF (blue is far, green is near, grey is disabled, black is in focus)"));
-				Canvas.DrawShadowedString(X, Y += YStep, *Line, GetStatsFont(), FLinearColor(1, 1, 1));
-				Y += YStep;
-				Line = FString::Printf(TEXT("FocalDistance: %.2f"), View.FinalPostProcessSettings.DepthOfFieldFocalDistance);
-				Canvas.DrawShadowedString(X, Y += YStep, *Line, GetStatsFont(), FLinearColor(1, 1, 1));
-				Line = FString::Printf(TEXT("FocalRegion (Artificial, avoid): %.2f"), View.FinalPostProcessSettings.DepthOfFieldFocalRegion);
-				Canvas.DrawShadowedString(X, Y += YStep, *Line, GetStatsFont(), FLinearColor(1, 1, 1));
-				Y += YStep;
-				Line = FString::Printf(TEXT("NearTransitionRegion: %.2f"), View.FinalPostProcessSettings.DepthOfFieldNearTransitionRegion);
-				Canvas.DrawShadowedString(X, Y += YStep, *Line, GetStatsFont(), FLinearColor(1, 1, 1));
-				Line = FString::Printf(TEXT("FarTransitionRegion: %.2f"), View.FinalPostProcessSettings.DepthOfFieldFarTransitionRegion);
-				Canvas.DrawShadowedString(X, Y += YStep, *Line, GetStatsFont(), FLinearColor(1, 1, 1));
-				Line = FString::Printf(TEXT("NearBlurSize: %.2f"), View.FinalPostProcessSettings.DepthOfFieldNearBlurSize);
-				Canvas.DrawShadowedString(X, Y += YStep, *Line, GetStatsFont(), FLinearColor(1, 1, 1));
-				Line = FString::Printf(TEXT("FarBlurSize: %.2f"), View.FinalPostProcessSettings.DepthOfFieldFarBlurSize);
-				Canvas.DrawShadowedString(X, Y += YStep, *Line, GetStatsFont(), FLinearColor(1, 1, 1));
-				Line = FString::Printf(TEXT("Occlusion: %.2f"), View.FinalPostProcessSettings.DepthOfFieldOcclusion);
-				Canvas.DrawShadowedString(X, Y += YStep, *Line, GetStatsFont(), FLinearColor(1, 1, 1));
-				Line = FString::Printf(TEXT("SkyFocusDistance: %.2f"), View.FinalPostProcessSettings.DepthOfFieldSkyFocusDistance);
-				Canvas.DrawShadowedString(X, Y += YStep, *Line, GetStatsFont(), FLinearColor(1, 1, 1));
-				Line = FString::Printf(TEXT("VignetteRadius: %.2f"), View.FinalPostProcessSettings.DepthOfFieldVignetteSize);
-				Canvas.DrawShadowedString(X, Y += YStep, *Line, GetStatsFont(), FLinearColor(1, 1, 1));
-				Y += YStep;
-				Line = FString::Printf(TEXT("Near:%d Far:%d"), DepthOfFieldStats.bNear ? 1 : 0, DepthOfFieldStats.bFar ? 1 : 0);
-				Canvas.DrawShadowedString(X, Y += YStep, *Line, GetStatsFont(), FLinearColor(1, 1, 1));
-			}
-			else if (MethodId == DOFM_CircleDOF)
-			{
-				Line = FString::Printf(TEXT("Method: CircleDOF (blue is far, green is near, black is in focus, cross hair shows Depth and CoC radius in pixel)"));
-				Canvas.DrawShadowedString(X, Y += YStep, *Line, GetStatsFont(), FLinearColor(1, 1, 1));
-				Y += YStep;
-				Line = FString::Printf(TEXT("FocalDistance: %.2f"), View.FinalPostProcessSettings.DepthOfFieldFocalDistance);
-				Canvas.DrawShadowedString(X, Y += YStep, *Line, GetStatsFont(), FLinearColor(1, 1, 1));
-				Line = FString::Printf(TEXT("Aperture F-stop: %.2f"), View.FinalPostProcessSettings.DepthOfFieldFstop);
-				Canvas.DrawShadowedString(X, Y += YStep, *Line, GetStatsFont(), FLinearColor(1, 1, 1));
-				Line = FString::Printf(TEXT("Aperture: f/%.2f"), View.FinalPostProcessSettings.DepthOfFieldFstop);
-				Canvas.DrawShadowedString(X, Y += YStep, *Line, GetStatsFont(), FLinearColor(1, 1, 1));
-				Y += YStep;
-				Line = FString::Printf(TEXT("DepthBlur (not related to Depth of Field, due to light traveling long distances in atmosphere)"));
-				Canvas.DrawShadowedString(X, Y += YStep, *Line, GetStatsFont(), FLinearColor(1, 1, 1));
-				Line = FString::Printf(TEXT("  km for 50%: %.2f"), View.FinalPostProcessSettings.DepthOfFieldDepthBlurAmount);
-				Canvas.DrawShadowedString(X, Y += YStep, *Line, GetStatsFont(), FLinearColor(1, 1, 1));
-				Line = FString::Printf(TEXT("  Radius (pixels in 1920x): %.2f"), View.FinalPostProcessSettings.DepthOfFieldDepthBlurRadius);
-				Canvas.DrawShadowedString(X, Y += YStep, *Line, GetStatsFont(), FLinearColor(1, 1, 1));
-				Y += YStep;
-
-				FVector2D Fov = View.ViewMatrices.ComputeHalfFieldOfViewPerAxis();
-
-				float FocalLength = ComputeFocalLengthFromFov(View);
-
-				Line = FString::Printf(TEXT("Field Of View in deg. (computed): %.1f x %.1f"), FMath::RadiansToDegrees(Fov.X) * 2.0f, FMath::RadiansToDegrees(Fov.Y) * 2.0f);
-				Canvas.DrawShadowedString(X, Y += YStep, *Line, GetStatsFont(), FLinearColor(0.5, 0.5, 1));
-				Line = FString::Printf(TEXT("Focal Length (computed): %.1f"), FocalLength);
-				Canvas.DrawShadowedString(X, Y += YStep, *Line, GetStatsFont(), FLinearColor(0.5, 0.5, 10));
-				Line = FString::Printf(TEXT("Sensor: APS-C 24.576 mm sensor, crop-factor 1.61x"));
-				Canvas.DrawShadowedString(X, Y += YStep, *Line, GetStatsFont(), FLinearColor(0.5, 0.5, 1));
-			}
-
-		Canvas.Flush_RenderThread(Context.RHICmdList);
-	}
-
-	Context.RHICmdList.CopyToResolveTarget(DestRenderTarget.TargetableTexture, DestRenderTarget.ShaderResourceTexture, FResolveParams());
-
-}
-
-FPooledRenderTargetDesc FRCPassPostProcessVisualizeDOF::ComputeOutputDesc(EPassOutputId InPassOutputId) const
-{
-	FPooledRenderTargetDesc Ret = GetInput(ePId_Input0)->GetOutput()->RenderTargetDesc;
-
-	Ret.Reset();
-	Ret.Format = PF_B8G8R8A8;
-	Ret.DebugName = TEXT("VisualizeDOF");
-
-	return Ret;
-}
-
-
-/** Encapsulates the post processing depth of field setup pixel shader. */
-class PostProcessBokehDOFSetupPS : public FGlobalShader
-{
-	DECLARE_SHADER_TYPE(PostProcessBokehDOFSetupPS, Global);
-
-	static bool ShouldCompilePermutation(const FGlobalShaderPermutationParameters& Parameters)
-	{
-		return IsFeatureLevelSupported(Parameters.Platform, ERHIFeatureLevel::SM4);
-	}
-
-	/** Default constructor. */
-	PostProcessBokehDOFSetupPS() {}
-
-public:
-	FPostProcessPassParameters PostprocessParameter;
-	FSceneTextureShaderParameters SceneTextureParameters;
-	FShaderParameter DepthOfFieldParams;
-
-	/** Initialization constructor. */
-	PostProcessBokehDOFSetupPS(const ShaderMetaType::CompiledShaderInitializerType& Initializer)
-		: FGlobalShader(Initializer)
-	{
-		PostprocessParameter.Bind(Initializer.ParameterMap);
-		SceneTextureParameters.Bind(Initializer);
-		DepthOfFieldParams.Bind(Initializer.ParameterMap,TEXT("DepthOfFieldParams"));
-	}
-
-	// FShader interface.
-	virtual bool Serialize(FArchive& Ar) override
-	{
-		bool bShaderHasOutdatedParameters = FGlobalShader::Serialize(Ar);
-		Ar << PostprocessParameter << SceneTextureParameters << DepthOfFieldParams;
-		return bShaderHasOutdatedParameters;
-	}
-
-	void SetParameters(const FRenderingCompositePassContext& Context)
-	{
-		const FPixelShaderRHIParamRef ShaderRHI = GetPixelShader();
-
-		FGlobalShader::SetParameters<FViewUniformShaderParameters>(Context.RHICmdList, ShaderRHI, Context.View.ViewUniformBuffer);
-		
-		SceneTextureParameters.Set(Context.RHICmdList, ShaderRHI, Context.View.FeatureLevel, ESceneTextureSetupMode::All);
-		PostprocessParameter.SetPS(Context.RHICmdList, ShaderRHI, Context, TStaticSamplerState<SF_Point,AM_Clamp,AM_Clamp,AM_Clamp>::GetRHI());
-
-		{
-			FVector4 DepthOfFieldParamValues[2];
-
-			// in rendertarget pixels (half res to scene color)
-			FRenderingCompositeOutput* Output = Context.Pass->GetOutput(ePId_Output0);
-
-			FRCPassPostProcessBokehDOF::ComputeDepthOfFieldParams(Context, DepthOfFieldParamValues);
-
-			SetShaderValueArray(Context.RHICmdList, ShaderRHI, DepthOfFieldParams, DepthOfFieldParamValues, 2);
-		}
-	}
-
-	static const TCHAR* GetSourceFilename()
-	{
-		return TEXT("/Engine/Private/PostProcessBokehDOF.usf");
-	}
-
-	static const TCHAR* GetFunctionName()
-	{
-		return TEXT("MainSetupPS");
-	}
-};
-
-IMPLEMENT_SHADER_TYPE3(PostProcessBokehDOFSetupPS, SF_Pixel);
-
-/** Encapsulates the post processing depth of field setup compute shader. */
-class FRCPassPostProcessBokehDOFSetupCS : public FGlobalShader
-{
-	DECLARE_SHADER_TYPE(FRCPassPostProcessBokehDOFSetupCS, Global);
-
-	static bool ShouldCompilePermutation(const FGlobalShaderPermutationParameters& Parameters)
-	{
-		return IsFeatureLevelSupported(Parameters.Platform, ERHIFeatureLevel::SM5);
-	}
-
-	static void ModifyCompilationEnvironment(const FGlobalShaderPermutationParameters& Parameters, FShaderCompilerEnvironment& OutEnvironment)
-	{
-		FGlobalShader::ModifyCompilationEnvironment(Parameters,OutEnvironment);
-		OutEnvironment.SetDefine(TEXT("THREADGROUP_SIZEX"), GBokehDOFSetupTileSizeX);
-		OutEnvironment.SetDefine(TEXT("THREADGROUP_SIZEY"), GBokehDOFSetupTileSizeY);
-	}
-
-	/** Default constructor. */
-	FRCPassPostProcessBokehDOFSetupCS() {}
-
-public:
-	FPostProcessPassParameters PostprocessParameter;
-	FSceneTextureShaderParameters SceneTextureParameters;
-	FShaderParameter BokehDOFSetupComputeParams;
-	FShaderParameter OutComputeTex;
-
-	/** Initialization constructor. */
-	FRCPassPostProcessBokehDOFSetupCS(const ShaderMetaType::CompiledShaderInitializerType& Initializer)
-		: FGlobalShader(Initializer)
-	{
-		PostprocessParameter.Bind(Initializer.ParameterMap);
-		SceneTextureParameters.Bind(Initializer);
-		BokehDOFSetupComputeParams.Bind(Initializer.ParameterMap,TEXT("BokehDOFSetupComputeParams"));
-		OutComputeTex.Bind(Initializer.ParameterMap, TEXT("OutComputeTex"));
-	}
-
-	// FShader interface.
-	virtual bool Serialize(FArchive& Ar) override
-	{
-		bool bShaderHasOutdatedParameters = FGlobalShader::Serialize(Ar);
-		Ar << PostprocessParameter << SceneTextureParameters << BokehDOFSetupComputeParams << OutComputeTex;
-		return bShaderHasOutdatedParameters;
-	}
-
-	template <typename TRHICmdList>
-	void SetParameters(TRHICmdList& RHICmdList, const FRenderingCompositePassContext& Context, const FIntPoint& DestSize, FUnorderedAccessViewRHIParamRef DestUAV)
-	{
-		const FComputeShaderRHIParamRef ShaderRHI = GetComputeShader();
-		FGlobalShader::SetParameters<FViewUniformShaderParameters>(RHICmdList, ShaderRHI, Context.View.ViewUniformBuffer);
-
-		SceneTextureParameters.Set(RHICmdList, ShaderRHI, Context.View.FeatureLevel, ESceneTextureSetupMode::All);
-		PostprocessParameter.SetCS(ShaderRHI, Context, RHICmdList, TStaticSamplerState<SF_Point,AM_Clamp,AM_Clamp,AM_Clamp>::GetRHI());
-		RHICmdList.SetUAVParameter(ShaderRHI, OutComputeTex.GetBaseIndex(), DestUAV);
-
-		FVector4 BokehDOFSetupComputeValues(0, 0, 1.f / (float)DestSize.X, 1.f / (float)DestSize.Y);
-		SetShaderValue(RHICmdList, ShaderRHI, BokehDOFSetupComputeParams, BokehDOFSetupComputeValues);
-	}
-
-	template <typename TRHICmdList>
-	void UnsetParameters(TRHICmdList& RHICmdList)
-	{
-		const FComputeShaderRHIParamRef ShaderRHI = GetComputeShader();
-		RHICmdList.SetUAVParameter(ShaderRHI, OutComputeTex.GetBaseIndex(), NULL);
-	}
-};
-
-IMPLEMENT_SHADER_TYPE(,FRCPassPostProcessBokehDOFSetupCS,TEXT("/Engine/Private/PostProcessBokehDOF.usf"),TEXT("MainCS"),SF_Compute);
-
-void FRCPassPostProcessBokehDOFSetup::Process(FRenderingCompositePassContext& Context)
-{
-	SCOPED_DRAW_EVENTF(Context.RHICmdList, PostProcessBokehDOFSetup, TEXT("PostProcessBokehDOFSetup%s"), bIsComputePass?TEXT("Compute"):TEXT(""));
-
-	const FPooledRenderTargetDesc* InputDesc = GetInputDesc(ePId_Input0);
-
-	if(!InputDesc)
-	{
-		// input is not hooked up correctly
-		return;
-	}
-
-	const FViewInfo& View = Context.View;
-	const FSceneViewFamily& ViewFamily = *(View.Family);
-
-	FIntPoint SrcSize = InputDesc->Extent;
-	FIntPoint DestSize = PassOutputs[0].RenderTargetDesc.Extent;
-
-	// e.g. 4 means the input texture is 4x smaller than the buffer size
-	uint32 ScaleFactor = FSceneRenderTargets::Get(Context.RHICmdList).GetBufferSizeXY().X / SrcSize.X;
-
-	FIntRect SrcRect = FIntRect::DivideAndRoundUp(View.ViewRect, ScaleFactor);
-	FIntRect DestRect = FIntRect::DivideAndRoundUp(SrcRect, 2);
-
-	const FSceneRenderTargetItem& DestRenderTarget = PassOutputs[0].RequestSurface(Context);
-
-	if (bIsComputePass)
-	{
-		DestRect = {View.ViewRect.Min, View.ViewRect.Min + DestSize};
-
-		// #todo-renderpasses Remove once everything is converted to renderpasses
-		UnbindRenderTargets(Context.RHICmdList);
-		Context.SetViewportAndCallRHI(DestRect, 0.0f, 1.0f);
-
-		WaitForInputPassComputeFences(Context.RHICmdList);
-		Context.RHICmdList.TransitionResource(EResourceTransitionAccess::ERWBarrier, EResourceTransitionPipeline::EGfxToCompute, DestRenderTarget.UAV);
-		DispatchCS(Context.RHICmdList, Context, DestRect, DestRenderTarget.UAV);
-		Context.RHICmdList.TransitionResource(EResourceTransitionAccess::EReadable, EResourceTransitionPipeline::EComputeToGfx, DestRenderTarget.UAV);
-	}
-	else
-	{
-		WaitForInputPassComputeFences(Context.RHICmdList);
-
-		// Set the view family's render target/viewport.
-		FRHIRenderPassInfo RPInfo(DestRenderTarget.TargetableTexture, ERenderTargetActions::Load_Store);
-		Context.RHICmdList.BeginRenderPass(RPInfo, TEXT("BokehDOFSetup"));
-		{
-
-			if (View.StereoPass == eSSP_FULL)
-			{
-				// can be optimized (don't clear areas we overwrite, don't clear when full screen),
-				// needed when a camera (matinee) has black borders or with multiple viewports
-				// focal distance depth is stored in the alpha channel to avoid DOF artifacts
-				DrawClearQuad(Context.RHICmdList, true, FLinearColor(0, 0, 0, View.FinalPostProcessSettings.DepthOfFieldFocalDistance), false, 0, false, 0, DestSize, DestRect);
-			}
-
-			Context.SetViewportAndCallRHI(0, 0, 0.0f, DestSize.X, DestSize.Y, 1.0f);
-
-			FGraphicsPipelineStateInitializer GraphicsPSOInit;
-			Context.RHICmdList.ApplyCachedRenderTargets(GraphicsPSOInit);
-
-			// set the state
-			GraphicsPSOInit.BlendState = TStaticBlendState<>::GetRHI();
-			GraphicsPSOInit.RasterizerState = TStaticRasterizerState<>::GetRHI();
-			GraphicsPSOInit.DepthStencilState = TStaticDepthStencilState<false, CF_Always>::GetRHI();
-
-			// setup shader
-			TShaderMapRef<FPostProcessVS> VertexShader(Context.GetShaderMap());
-			GraphicsPSOInit.BoundShaderState.VertexDeclarationRHI = GFilterVertexDeclaration.VertexDeclarationRHI;
-			GraphicsPSOInit.BoundShaderState.VertexShaderRHI = GETSAFERHISHADER_VERTEX(*VertexShader);
-			GraphicsPSOInit.PrimitiveType = PT_TriangleList;
-
-			{
-				TShaderMapRef<PostProcessBokehDOFSetupPS> PixelShader(Context.GetShaderMap());
-
-				GraphicsPSOInit.BoundShaderState.PixelShaderRHI = GETSAFERHISHADER_PIXEL(*PixelShader);
-				SetGraphicsPipelineState(Context.RHICmdList, GraphicsPSOInit);
-
-				VertexShader->SetParameters(Context);
-				PixelShader->SetParameters(Context);
-			}
-
-			DrawPostProcessPass(
-				Context.RHICmdList,
-				DestRect.Min.X, DestRect.Min.Y,
-				DestRect.Width(), DestRect.Height(),
-				SrcRect.Min.X, SrcRect.Min.Y,
-				SrcRect.Width(), SrcRect.Height(),
-				DestSize,
-				SrcSize,
-				*VertexShader,
-				View.StereoPass,
-				Context.HasHmdMesh(),
-				EDRF_UseTriangleOptimization);
-		}
-		Context.RHICmdList.EndRenderPass();
-		Context.RHICmdList.CopyToResolveTarget(DestRenderTarget.TargetableTexture, DestRenderTarget.ShaderResourceTexture, FResolveParams());
-	}
-}
-
-template <typename TRHICmdList>
-void FRCPassPostProcessBokehDOFSetup::DispatchCS(TRHICmdList& RHICmdList, FRenderingCompositePassContext& Context, const FIntRect& DestRect, FUnorderedAccessViewRHIParamRef DestUAV)
-{
-	auto ShaderMap = Context.GetShaderMap();
-	TShaderMapRef<FRCPassPostProcessBokehDOFSetupCS> ComputeShader(ShaderMap);
-	RHICmdList.SetComputeShader(ComputeShader->GetComputeShader());
-
-	FIntPoint DestSize(DestRect.Width(), DestRect.Height());
-	ComputeShader->SetParameters(RHICmdList, Context, DestSize, DestUAV);
-
-	uint32 GroupSizeX = FMath::DivideAndRoundUp(DestSize.X, GBokehDOFSetupTileSizeX);
-	uint32 GroupSizeY = FMath::DivideAndRoundUp(DestSize.Y, GBokehDOFSetupTileSizeY);
-	DispatchComputeShader(RHICmdList, *ComputeShader, GroupSizeX, GroupSizeY, 1);
-
-	ComputeShader->UnsetParameters(RHICmdList);
-}
-
-FPooledRenderTargetDesc FRCPassPostProcessBokehDOFSetup::ComputeOutputDesc(EPassOutputId InPassOutputId) const
-{
-	FPooledRenderTargetDesc Ret = GetInput(ePId_Input0)->GetOutput()->RenderTargetDesc;
-
-	Ret.Reset();
-	Ret.Extent /= 2;
-	Ret.Extent.X = FMath::Max(1, Ret.Extent.X);
-	Ret.Extent.Y = FMath::Max(1, Ret.Extent.Y);
-	Ret.Format = PF_FloatRGBA;
-	Ret.DebugName = TEXT("BokehDOFSetup");
-	Ret.TargetableFlags &= ~(TexCreate_RenderTargetable | TexCreate_UAV);
-	Ret.TargetableFlags |= bIsComputePass ? TexCreate_UAV : TexCreate_RenderTargetable;
-	Ret.Flags |= GFastVRamConfig.BokehDOF;
-
-	return Ret;
-}
-
-/** Encapsulates the post processing vertex shader. */
-template <uint32 DOFMethod>
-class FPostProcessBokehDOFVS : public FGlobalShader
-{
-	DECLARE_SHADER_TYPE(FPostProcessBokehDOFVS,Global);
-
-	static bool ShouldCompilePermutation(const FGlobalShaderPermutationParameters& Parameters)
-	{
-		return IsFeatureLevelSupported(Parameters.Platform, ERHIFeatureLevel::SM4);
-	}
-
-	static void ModifyCompilationEnvironment(const FGlobalShaderPermutationParameters& Parameters, FShaderCompilerEnvironment& OutEnvironment)
-	{
-		FGlobalShader::ModifyCompilationEnvironment(Parameters,OutEnvironment);
-		OutEnvironment.SetDefine(TEXT("DOF_METHOD"), DOFMethod);
-		OutEnvironment.SetDefine(TEXT("BOKEH_DOF_QUADS_PER_INSTANCE"), GBokehDOFQuadsPerInstance);
-	}
-
-	/** Default constructor. */
-	FPostProcessBokehDOFVS() {}
-	
-public:
-	FPostProcessPassParameters PostprocessParameter;
-	FShaderParameter TileCountAndSize;
-	FShaderParameter KernelSize;
-	FShaderParameter DepthOfFieldParams;
-	FShaderParameter DepthOfFieldThresholds;
-	FSceneTextureShaderParameters SceneTextureParameters;
-
-	/** Initialization constructor. */
-	FPostProcessBokehDOFVS(const ShaderMetaType::CompiledShaderInitializerType& Initializer)
-		: FGlobalShader(Initializer)
-	{
-		PostprocessParameter.Bind(Initializer.ParameterMap);
-		TileCountAndSize.Bind(Initializer.ParameterMap, TEXT("TileCountAndSize"));
-		KernelSize.Bind(Initializer.ParameterMap, TEXT("KernelSize"));
-		DepthOfFieldParams.Bind(Initializer.ParameterMap,TEXT("DepthOfFieldParams"));
-		DepthOfFieldThresholds.Bind(Initializer.ParameterMap,TEXT("DepthOfFieldThresholds"));		
-		SceneTextureParameters.Bind(Initializer);
-	}
-
-	// FShader interface.
-	virtual bool Serialize(FArchive& Ar) override
-	{
-		bool bShaderHasOutdatedParameters = FGlobalShader::Serialize(Ar);
-		Ar << PostprocessParameter << TileCountAndSize << KernelSize << DepthOfFieldParams << DepthOfFieldThresholds << SceneTextureParameters;
-		return bShaderHasOutdatedParameters;
-	}
-
-	/** to have a similar interface as all other shaders */
-	void SetParameters(const FRenderingCompositePassContext& Context, FIntPoint TileCountValue, uint32 TileSize, float PixelKernelSize, FIntPoint LeftTop)
-	{
-		const FVertexShaderRHIParamRef ShaderRHI = GetVertexShader();
-
-		FGlobalShader::SetParameters<FViewUniformShaderParameters>(Context.RHICmdList, ShaderRHI, Context.View.ViewUniformBuffer);
-		SceneTextureParameters.Set(Context.RHICmdList, ShaderRHI, Context.View.FeatureLevel, ESceneTextureSetupMode::All);
-		PostprocessParameter.SetVS(ShaderRHI, Context, TStaticSamplerState<SF_Bilinear,AM_Clamp,AM_Clamp,AM_Clamp>::GetRHI());
-//		PostprocessParameter.SetVS(ShaderRHI, Context, TStaticSamplerState<SF_Point,AM_Clamp,AM_Clamp,AM_Clamp>::GetRHI());
-
-		{
-			FIntRect TileCountAndSizeValue(TileCountValue, FIntPoint(TileSize, TileSize));
-
-			SetShaderValue(Context.RHICmdList, ShaderRHI, TileCountAndSize, TileCountAndSizeValue);
-		}
-
-		{
-			FVector4 KernelSizeValue(PixelKernelSize, PixelKernelSize, LeftTop.X, LeftTop.Y);
-
-			SetShaderValue(Context.RHICmdList, ShaderRHI, KernelSize, KernelSizeValue);
-		}
-		PRAGMA_DISABLE_DEPRECATION_WARNINGS
-		{
-			FVector4 Value(
-				Context.View.FinalPostProcessSettings.DepthOfFieldColorThreshold,
-				Context.View.FinalPostProcessSettings.DepthOfFieldSizeThreshold,0);
-
-			SetShaderValue(Context.RHICmdList, ShaderRHI, DepthOfFieldThresholds, Value);
-		}
-		PRAGMA_ENABLE_DEPRECATION_WARNINGS
-
+		if (ViewFamily.Scene->GetShadingPath() == EShadingPath::Mobile)
 		{
-			FVector4 DepthOfFieldParamValues[2];
-
-			FRCPassPostProcessBokehDOF::ComputeDepthOfFieldParams(Context, DepthOfFieldParamValues);
-
-			SetShaderValueArray(Context.RHICmdList, ShaderRHI, DepthOfFieldParams, DepthOfFieldParamValues, 2);
-		}
-	}
-
-	static const TCHAR* GetSourceFilename()
-	{
-		return TEXT("/Engine/Private/PostProcessBokehDOF.usf");
-	}
-
-	static const TCHAR* GetFunctionName()
-	{
-		return TEXT("MainVS");
-	}
-};
-
-/** Encapsulates a simple copy pixel shader. */
-class FPostProcessBokehDOFPS : public FGlobalShader
-{
-	DECLARE_SHADER_TYPE(FPostProcessBokehDOFPS, Global);
-
-	static bool ShouldCompilePermutation(const FGlobalShaderPermutationParameters& Parameters)
-	{
-		return IsFeatureLevelSupported(Parameters.Platform, ERHIFeatureLevel::SM4);
-	}
-
-	/** Default constructor. */
-	FPostProcessBokehDOFPS() {}
-
-public:
-	FPostProcessPassParameters PostprocessParameter;
-	FShaderResourceParameter LensTexture;
-	FShaderResourceParameter LensTextureSampler;
-
-	/** Initialization constructor. */
-	FPostProcessBokehDOFPS(const ShaderMetaType::CompiledShaderInitializerType& Initializer)
-		: FGlobalShader(Initializer)
-	{
-		PostprocessParameter.Bind(Initializer.ParameterMap);
-		LensTexture.Bind(Initializer.ParameterMap,TEXT("LensTexture"));
-		LensTextureSampler.Bind(Initializer.ParameterMap,TEXT("LensTextureSampler"));
-	}
-
-	// FShader interface.
-	virtual bool Serialize(FArchive& Ar) override
-	{
-		bool bShaderHasOutdatedParameters = FGlobalShader::Serialize(Ar);
-		Ar << PostprocessParameter << LensTexture << LensTextureSampler;
-		return bShaderHasOutdatedParameters;
-	}
-
-	static void ModifyCompilationEnvironment(const FGlobalShaderPermutationParameters& Parameters, FShaderCompilerEnvironment& OutEnvironment)
-	{
-		OutEnvironment.SetDefine(TEXT("BOKEH_DOF_QUADS_PER_INSTANCE"), GBokehDOFQuadsPerInstance);
-	}
-
-	void SetParameters(const FRenderingCompositePassContext& Context, float PixelKernelSize)
-	{
-		const FPixelShaderRHIParamRef ShaderRHI = GetPixelShader();
-
-		FGlobalShader::SetParameters<FViewUniformShaderParameters>(Context.RHICmdList, ShaderRHI, Context.View.ViewUniformBuffer);
-
-		PostprocessParameter.SetPS(Context.RHICmdList, ShaderRHI, Context, TStaticSamplerState<SF_Bilinear,AM_Clamp,AM_Clamp,AM_Clamp>::GetRHI());
-//		PostprocessParameter.SetPS(ShaderRHI, Context, TStaticSamplerState<SF_Point,AM_Clamp,AM_Clamp,AM_Clamp>::GetRHI());
-				
-		{
-			FTextureRHIParamRef TextureRHI = GWhiteTexture->TextureRHI;
-
-			if(GEngine->DefaultBokehTexture)
-			{
-				TextureRHI = GEngine->DefaultBokehTexture->Resource->TextureRHI;
-			}
-
-			PRAGMA_DISABLE_DEPRECATION_WARNINGS
-			if(Context.View.FinalPostProcessSettings.DepthOfFieldBokehShape)
-			{
-				TextureRHI = Context.View.FinalPostProcessSettings.DepthOfFieldBokehShape->Resource->TextureRHI;
-			}
-			PRAGMA_ENABLE_DEPRECATION_WARNINGS
-
-			SetTextureParameter(Context.RHICmdList, ShaderRHI, LensTexture, LensTextureSampler, TStaticSamplerState<SF_Trilinear, AM_Border, AM_Border, AM_Clamp>::GetRHI(), TextureRHI);
-		}
-	}
-};
-
-IMPLEMENT_SHADER_TYPE(,FPostProcessBokehDOFPS,TEXT("/Engine/Private/PostProcessBokehDOF.usf"),TEXT("MainPS"),SF_Pixel);
-
-// #define avoids a lot of code duplication
-#define VARIATION1(A) typedef FPostProcessBokehDOFVS<A> FPostProcessBokehDOFVS##A; \
-	IMPLEMENT_SHADER_TYPE2(FPostProcessBokehDOFVS##A, SF_Vertex);
-
-VARIATION1(0)			VARIATION1(1)			VARIATION1(2)
-#undef VARIATION1
-
-template <uint32 DOFMethod>
-void FRCPassPostProcessBokehDOF::SetShaderTempl(const FRenderingCompositePassContext& Context, FIntPoint LeftTop, FIntPoint TileCount, uint32 TileSize, float PixelKernelSize)
-{
-	FGraphicsPipelineStateInitializer GraphicsPSOInit;
-	Context.RHICmdList.ApplyCachedRenderTargets(GraphicsPSOInit);
-
-	// set the state (additive blending)
-	GraphicsPSOInit.BlendState = TStaticBlendState<CW_RGBA, BO_Add, BF_One, BF_One, BO_Add, BF_One, BF_One>::GetRHI();
-	GraphicsPSOInit.RasterizerState = TStaticRasterizerState<>::GetRHI();
-	GraphicsPSOInit.DepthStencilState = TStaticDepthStencilState<false, CF_Always>::GetRHI();
-	GraphicsPSOInit.PrimitiveType = PT_TriangleList;
-
-	TShaderMapRef<FPostProcessBokehDOFVS<DOFMethod> > VertexShader(Context.GetShaderMap());
-	TShaderMapRef<FPostProcessBokehDOFPS> PixelShader(Context.GetShaderMap());
-
-	GraphicsPSOInit.BoundShaderState.VertexDeclarationRHI = GEmptyVertexDeclaration.VertexDeclarationRHI;
-	GraphicsPSOInit.BoundShaderState.VertexShaderRHI = GETSAFERHISHADER_VERTEX(*VertexShader);
-	GraphicsPSOInit.BoundShaderState.PixelShaderRHI = GETSAFERHISHADER_PIXEL(*PixelShader);
-	SetGraphicsPipelineState(Context.RHICmdList, GraphicsPSOInit);
-
-	VertexShader->SetParameters(Context, TileCount, TileSize, PixelKernelSize, LeftTop);
-	PixelShader->SetParameters(Context, PixelKernelSize);
-}
-
-void FRCPassPostProcessBokehDOF::ComputeDepthOfFieldParams(const FRenderingCompositePassContext& Context, FVector4 Out[2])
-{
-	uint32 FullRes = FSceneRenderTargets::Get(Context.RHICmdList).GetBufferSizeXY().Y;
-	uint32 HalfRes = FMath::DivideAndRoundUp(FullRes, (uint32)2);
-	uint32 BokehLayerSizeY = HalfRes * 2 + SafetyBorder;
-
-	float SkyFocusDistance = Context.View.FinalPostProcessSettings.DepthOfFieldSkyFocusDistance;
-	
-	// *2 to go to account for Radius/Diameter, 100 for percent
-	float DepthOfFieldVignetteSize = FMath::Max(0.0f, Context.View.FinalPostProcessSettings.DepthOfFieldVignetteSize / 100.0f * 2);
-	// doesn't make much sense to expose this property as the effect is very non linear and it would cost some performance to fix that
-	float DepthOfFieldVignetteFeather = 10.0f / 100.0f;
-
-	float DepthOfFieldVignetteMul = 1.0f / DepthOfFieldVignetteFeather;
-	float DepthOfFieldVignetteAdd = (0.5f - DepthOfFieldVignetteSize) * DepthOfFieldVignetteMul;
-
-	Out[0] = FVector4(
-		(SkyFocusDistance > 0) ? SkyFocusDistance : 100000000.0f,			// very large if <0 to not mask out skybox, can be optimized to disable feature completely
-		DepthOfFieldVignetteMul,
-		DepthOfFieldVignetteAdd,
-		Context.View.FinalPostProcessSettings.DepthOfFieldOcclusion);
-
-	FIntPoint ViewSize = Context.View.ViewRect.Size();
-	
-	PRAGMA_DISABLE_DEPRECATION_WARNINGS
-	float MaxBokehSizeInPixel = FMath::Max(0.0f, Context.View.FinalPostProcessSettings.DepthOfFieldMaxBokehSize) / 100.0f * ViewSize.X;
-	PRAGMA_ENABLE_DEPRECATION_WARNINGS
-
-	// Scale and offset to put two views in one texture with safety border
-	float UsedYDivTextureY = HalfRes / (float)BokehLayerSizeY;
-	float YOffsetInPixel = HalfRes + SafetyBorder;
-	float YOffsetInUV = (HalfRes + SafetyBorder) / (float)BokehLayerSizeY;
-
-	Out[1] = FVector4(MaxBokehSizeInPixel, YOffsetInUV, UsedYDivTextureY, YOffsetInPixel);
-}
-
-
-void FRCPassPostProcessBokehDOF::Process(FRenderingCompositePassContext& Context)
-{
-	SCOPED_DRAW_EVENTF(Context.RHICmdList, PostProcessBokehDOF, TEXT("PostProcessBokehDOF%s"), bIsComputePass?TEXT("Compute"):TEXT(""));
-
-	const FPooledRenderTargetDesc* InputDesc = GetInputDesc(ePId_Input0);
-	
-	if(!InputDesc)
-	{
-		// input is not hooked up correctly
-		return;
-	}
-
-	const FViewInfo& View = Context.View;
-
-	FIntPoint TexSize = InputDesc->Extent;
-
-	// usually 1, 2, 4 or 8
-	uint32 ScaleToFullRes = FSceneRenderTargets::Get(Context.RHICmdList).GetBufferSizeXY().X / TexSize.X;
-
-	// don't use DivideAndRoundUp as this could cause cause lookups into areas we don't have setup 
-	FIntRect LocalViewRect = View.ViewRect / ScaleToFullRes;
-
-	// contract by one half res pixel to avoid using samples outside of the input (SV runs at quarter resolution with 4 quads at once)
-	// this can lead to missing content - if needed this can be made less conservative
-	LocalViewRect.InflateRect(-2);
-
-	FIntPoint LocalViewSize = LocalViewRect.Size();
-
-	const FSceneRenderTargetItem& DestRenderTarget = PassOutputs[0].RequestSurface(Context);
-
-	WaitForInputPassComputeFences(Context.RHICmdList);
-
-	// Set the view family's render target/viewport.
-	FRHIRenderPassInfo RPInfo(DestRenderTarget.TargetableTexture, ERenderTargetActions::Load_Store);
-	Context.RHICmdList.BeginRenderPass(RPInfo, TEXT("BokehDOF"));
-	{
-
-		// This clean is required to make the accumulation working
-		DrawClearQuad(Context.RHICmdList, true, FLinearColor(0, 0, 0, 0), false, 0, false, 0, GetOutput(ePId_Output0)->RenderTargetDesc.Extent, FIntRect());
-
-		// we need to output to the whole rendertarget
-		Context.SetViewportAndCallRHI(0, 0, 0.0f, PassOutputs[0].RenderTargetDesc.Extent.X, PassOutputs[0].RenderTargetDesc.Extent.Y, 1.0f);
-
-		static const auto CVar = IConsoleManager::Get().FindTConsoleVariableDataInt(TEXT("r.DepthOfFieldQuality"));
-		check(CVar);
-		int32 DOFQuality = CVar->GetValueOnRenderThread();
-		check(DOFQuality > 0);
-
-		bool bHighQuality = DOFQuality > 1;
-
-		// 1: one quad per 1 half res texel
-		// 2: one quad per 4 half res texel (faster, can alias - need to verify that with bilinear filtering)
-		uint32 TileSize = bHighQuality ? 1 : 2;
-
-		// input is half res, don't process last pixel line where we don't have input
-		LocalViewSize.X &= ~1;
-		LocalViewSize.Y &= ~1;
-
-		FIntPoint TileCount = LocalViewSize / TileSize;
-
-		PRAGMA_DISABLE_DEPRECATION_WARNINGS
-			float PixelKernelSize = Context.View.FinalPostProcessSettings.DepthOfFieldMaxBokehSize / 100.0f * LocalViewSize.X;
-		PRAGMA_ENABLE_DEPRECATION_WARNINGS
-
-			FIntPoint LeftTop = LocalViewRect.Min;
-
-		if (bHighQuality)
-		{
-			if (View.Family->EngineShowFlags.VisualizeAdaptiveDOF)
-			{
-				// high quality, visualize in red and green where we spend more performance
-				SetShaderTempl<2>(Context, LeftTop, TileCount, TileSize, PixelKernelSize);
-			}
-			else
-			{
-				// high quality
-				SetShaderTempl<1>(Context, LeftTop, TileCount, TileSize, PixelKernelSize);
-			}
+			Line = FString::Printf(TEXT("Mobile Gaussian DOF (blue is far, green is near, grey is disabled, black is in focus)"));
+			Canvas.DrawShadowedString(X, Y += YStep, *Line, GetStatsFont(), FLinearColor(1, 1, 1));
+			Y += YStep;
+			Line = FString::Printf(TEXT("FocalDistance: %.2f"), View.FinalPostProcessSettings.DepthOfFieldFocalDistance);
+			Canvas.DrawShadowedString(X, Y += YStep, *Line, GetStatsFont(), FLinearColor(1, 1, 1));
+			Line = FString::Printf(TEXT("FocalRegion (Artificial, avoid): %.2f"), View.FinalPostProcessSettings.DepthOfFieldFocalRegion);
+			Canvas.DrawShadowedString(X, Y += YStep, *Line, GetStatsFont(), FLinearColor(1, 1, 1));
+			Y += YStep;
+			Line = FString::Printf(TEXT("NearTransitionRegion: %.2f"), View.FinalPostProcessSettings.DepthOfFieldNearTransitionRegion);
+			Canvas.DrawShadowedString(X, Y += YStep, *Line, GetStatsFont(), FLinearColor(1, 1, 1));
+			Line = FString::Printf(TEXT("FarTransitionRegion: %.2f"), View.FinalPostProcessSettings.DepthOfFieldFarTransitionRegion);
+			Canvas.DrawShadowedString(X, Y += YStep, *Line, GetStatsFont(), FLinearColor(1, 1, 1));
+			Line = FString::Printf(TEXT("NearBlurSize: %.2f"), View.FinalPostProcessSettings.DepthOfFieldNearBlurSize);
+			Canvas.DrawShadowedString(X, Y += YStep, *Line, GetStatsFont(), FLinearColor(1, 1, 1));
+			Line = FString::Printf(TEXT("FarBlurSize: %.2f"), View.FinalPostProcessSettings.DepthOfFieldFarBlurSize);
+			Canvas.DrawShadowedString(X, Y += YStep, *Line, GetStatsFont(), FLinearColor(1, 1, 1));
+			Line = FString::Printf(TEXT("Occlusion: %.2f"), View.FinalPostProcessSettings.DepthOfFieldOcclusion);
+			Canvas.DrawShadowedString(X, Y += YStep, *Line, GetStatsFont(), FLinearColor(1, 1, 1));
+			Line = FString::Printf(TEXT("SkyFocusDistance: %.2f"), View.FinalPostProcessSettings.DepthOfFieldSkyFocusDistance);
+			Canvas.DrawShadowedString(X, Y += YStep, *Line, GetStatsFont(), FLinearColor(1, 1, 1));
+			Line = FString::Printf(TEXT("VignetteRadius: %.2f"), View.FinalPostProcessSettings.DepthOfFieldVignetteSize);
+			Canvas.DrawShadowedString(X, Y += YStep, *Line, GetStatsFont(), FLinearColor(1, 1, 1));
+			Y += YStep;
+			Line = FString::Printf(TEXT("Near:%d Far:%d"), DepthOfFieldStats.bNear ? 1 : 0, DepthOfFieldStats.bFar ? 1 : 0);
+			Canvas.DrawShadowedString(X, Y += YStep, *Line, GetStatsFont(), FLinearColor(1, 1, 1));
 		}
 		else
 		{
-			// low quality
-			SetShaderTempl<0>(Context, LeftTop, TileCount, TileSize, PixelKernelSize);
+			Line = FString::Printf(TEXT("Cinematic DOF (blue is far, green is near, black is in focus, cross hair shows Depth and CoC radius in pixel)"));
+			Canvas.DrawShadowedString(X, Y += YStep, *Line, GetStatsFont(), FLinearColor(1, 1, 1));
+			Y += YStep;
+			Line = FString::Printf(TEXT("FocalDistance: %.2f"), View.FinalPostProcessSettings.DepthOfFieldFocalDistance);
+			Canvas.DrawShadowedString(X, Y += YStep, *Line, GetStatsFont(), FLinearColor(1, 1, 1));
+			Line = FString::Printf(TEXT("Aperture F-stop: %.2f"), View.FinalPostProcessSettings.DepthOfFieldFstop);
+			Canvas.DrawShadowedString(X, Y += YStep, *Line, GetStatsFont(), FLinearColor(1, 1, 1));
+			Line = FString::Printf(TEXT("Aperture: f/%.2f"), View.FinalPostProcessSettings.DepthOfFieldFstop);
+			Canvas.DrawShadowedString(X, Y += YStep, *Line, GetStatsFont(), FLinearColor(1, 1, 1));
+			Y += YStep;
+			Line = FString::Printf(TEXT("DepthBlur (not related to Depth of Field, due to light traveling long distances in atmosphere)"));
+			Canvas.DrawShadowedString(X, Y += YStep, *Line, GetStatsFont(), FLinearColor(1, 1, 1));
+			Line = FString::Printf(TEXT("  km for 50%: %.2f"), View.FinalPostProcessSettings.DepthOfFieldDepthBlurAmount);
+			Canvas.DrawShadowedString(X, Y += YStep, *Line, GetStatsFont(), FLinearColor(1, 1, 1));
+			Line = FString::Printf(TEXT("  Radius (pixels in 1920x): %.2f"), View.FinalPostProcessSettings.DepthOfFieldDepthBlurRadius);
+			Canvas.DrawShadowedString(X, Y += YStep, *Line, GetStatsFont(), FLinearColor(1, 1, 1));
+			Y += YStep;
+
+			FVector2D Fov = View.ViewMatrices.ComputeHalfFieldOfViewPerAxis();
+
+			float FocalLength = DiaphragmDOF::ComputeFocalLengthFromFov(View);
+
+			Line = FString::Printf(TEXT("Field Of View in deg. (computed): %.1f x %.1f"), FMath::RadiansToDegrees(Fov.X) * 2.0f, FMath::RadiansToDegrees(Fov.Y) * 2.0f);
+			Canvas.DrawShadowedString(X, Y += YStep, *Line, GetStatsFont(), FLinearColor(0.5, 0.5, 1));
+			Line = FString::Printf(TEXT("Focal Length (computed): %.1f"), FocalLength);
+			Canvas.DrawShadowedString(X, Y += YStep, *Line, GetStatsFont(), FLinearColor(0.5, 0.5, 10));
+			Line = FString::Printf(TEXT("Sensor: APS-C 24.576 mm sensor, crop-factor 1.61x"));
+			Canvas.DrawShadowedString(X, Y += YStep, *Line, GetStatsFont(), FLinearColor(0.5, 0.5, 1));
 		}
 
-		Context.RHICmdList.SetStreamSource(0, NULL, 0);
-		Context.RHICmdList.DrawIndexedPrimitive(GBokehIndexBuffer.IndexBufferRHI, 0, 0, 4 * GBokehDOFQuadsPerInstance, 0, 2 * GBokehDOFQuadsPerInstance, FMath::DivideAndRoundUp(TileCount.X * TileCount.Y, GBokehDOFQuadsPerInstance));
+		Canvas.Flush_RenderThread(Context.RHICmdList);
 	}
 
-	Context.RHICmdList.EndRenderPass();
 	Context.RHICmdList.CopyToResolveTarget(DestRenderTarget.TargetableTexture, DestRenderTarget.ShaderResourceTexture, FResolveParams());
+
 }
 
-FPooledRenderTargetDesc FRCPassPostProcessBokehDOF::ComputeOutputDesc(EPassOutputId InPassOutputId) const
+FPooledRenderTargetDesc FRCPassPostProcessVisualizeDOF::ComputeOutputDesc(EPassOutputId InPassOutputId) const
 {
 	FPooledRenderTargetDesc Ret = GetInput(ePId_Input0)->GetOutput()->RenderTargetDesc;
 
 	Ret.Reset();
-	// more precision for additive blending
-	Ret.Format = PF_FloatRGBA;
-
-	uint32 FullRes = FSceneRenderTargets::Get_FrameConstantsOnly().GetBufferSizeXY().Y;
-	uint32 HalfRes = FMath::DivideAndRoundUp(FullRes, (uint32)2);
-
-	Ret.Flags |= GFastVRamConfig.BokehDOF;
-	// we need space for the front part and the back part
-	Ret.Extent.Y = HalfRes * 2 + SafetyBorder;
-	Ret.DebugName = TEXT("BokehDOF");
+	Ret.Format = PF_B8G8R8A8;
+	Ret.DebugName = TEXT("VisualizeDOF");
 
 	return Ret;
 }
diff --git a/Engine/Source/Runtime/Renderer/Private/PostProcess/PostProcessBokehDOF.h b/Engine/Source/Runtime/Renderer/Private/PostProcess/PostProcessBokehDOF.h
--- a/Engine/Source/Runtime/Renderer/Private/PostProcess/PostProcessBokehDOF.h
+++ b/Engine/Source/Runtime/Renderer/Private/PostProcess/PostProcessBokehDOF.h
@@ -40,48 +40,3 @@ public:
 
 	FDepthOfFieldStats DepthOfFieldStats;
 };
-
-// derives from TRenderingCompositePassBase<InputCount, OutputCount> 
-// ePId_Input0: Color input
-// ePId_Input1: Depth input
-class FRCPassPostProcessBokehDOFSetup : public TRenderingCompositePassBase<2, 1>
-{
-public:
-	FRCPassPostProcessBokehDOFSetup(bool bInIsComputePass)
-	{
-		bIsComputePass = bInIsComputePass;
-	}
-
-	// interface FRenderingCompositePass ---------
-	virtual void Process(FRenderingCompositePassContext& Context) override;
-	virtual void Release() override { delete this; }
-	virtual FPooledRenderTargetDesc ComputeOutputDesc(EPassOutputId InPassOutputId) const override;
-
-private:
-	template <typename TRHICmdList>
-	void DispatchCS(TRHICmdList& RHICmdList, FRenderingCompositePassContext& Context, const FIntRect& DestRect, FUnorderedAccessViewRHIParamRef DestUAV);
-};
-
-// derives from TRenderingCompositePassBase<InputCount, OutputCount> 
-// ePId_Input0: Half res scene with depth in alpha
-// ePId_Input1: SceneColor for high quality input (experimental)
-// ePId_Input2: SceneDepth for high quality input (experimental)
-class FRCPassPostProcessBokehDOF : public TRenderingCompositePassBase<3, 1>
-{
-public:
-	// interface FRenderingCompositePass ---------
-
-	virtual void Process(FRenderingCompositePassContext& Context) override;
-	virtual void Release() override { delete this; }
-	virtual FPooledRenderTargetDesc ComputeOutputDesc(EPassOutputId InPassOutputId) const override;
-
-	static void ComputeDepthOfFieldParams(const FRenderingCompositePassContext& Context, FVector4 Out[2]);
-
-private:
-
-	template <uint32 HighQuality>
-	void SetShaderTempl(const FRenderingCompositePassContext& Context, FIntPoint LeftTop, FIntPoint TileCount, uint32 TileSize, float PixelKernelSize);
-
-	// border between front and back layer as we don't use viewports (only possible with GS)
-	const static uint32 SafetyBorder = 40;
-};
diff --git a/Engine/Source/Runtime/Renderer/Private/PostProcess/PostProcessBokehDOFRecombine.cpp b/Engine/Source/Runtime/Renderer/Private/PostProcess/PostProcessBokehDOFRecombine.cpp
deleted file mode 100644
diff --git a/Engine/Source/Runtime/Renderer/Private/PostProcess/PostProcessBokehDOFRecombine.h b/Engine/Source/Runtime/Renderer/Private/PostProcess/PostProcessBokehDOFRecombine.h
deleted file mode 100644
diff --git a/Engine/Source/Runtime/Renderer/Private/PostProcess/PostProcessCircleDOF.cpp b/Engine/Source/Runtime/Renderer/Private/PostProcess/PostProcessCircleDOF.cpp
deleted file mode 100644
diff --git a/Engine/Source/Runtime/Renderer/Private/PostProcess/PostProcessCircleDOF.h b/Engine/Source/Runtime/Renderer/Private/PostProcess/PostProcessCircleDOF.h
deleted file mode 100644
diff --git a/Engine/Source/Runtime/Renderer/Private/PostProcess/PostProcessDOF.cpp b/Engine/Source/Runtime/Renderer/Private/PostProcess/PostProcessDOF.cpp
--- a/Engine/Source/Runtime/Renderer/Private/PostProcess/PostProcessDOF.cpp
+++ b/Engine/Source/Runtime/Renderer/Private/PostProcess/PostProcessDOF.cpp
@@ -84,7 +84,7 @@ public:
 		{
 			FVector4 DepthOfFieldParamValues[2];
 
-			FRCPassPostProcessBokehDOF::ComputeDepthOfFieldParams(Context, DepthOfFieldParamValues);
+			FRCPassPostProcessDOFSetup::ComputeDepthOfFieldParams(Context, DepthOfFieldParamValues);
 
 			SetShaderValueArray(RHICmdList, ShaderRHI, DepthOfFieldParams, DepthOfFieldParamValues, 2);
 		}
@@ -520,3 +520,34 @@ FPooledRenderTargetDesc FRCPassPostProcessDOFRecombine::ComputeOutputDesc(EPassO
 
 	return Ret;
 }
+
+// static
+void FRCPassPostProcessDOFSetup::ComputeDepthOfFieldParams(const FRenderingCompositePassContext& Context, FVector4 Out[2])
+{
+	// border between front and back layer as we don't use viewports (only possible with GS)
+	const uint32 SafetyBorder = 40;
+
+	uint32 FullRes = FSceneRenderTargets::Get(Context.RHICmdList).GetBufferSizeXY().Y;
+	uint32 HalfRes = FMath::DivideAndRoundUp(FullRes, (uint32)2);
+	uint32 BokehLayerSizeY = HalfRes * 2 + SafetyBorder;
+	float SkyFocusDistance = Context.View.FinalPostProcessSettings.DepthOfFieldSkyFocusDistance;
+	
+	// *2 to go to account for Radius/Diameter, 100 for percent
+	float DepthOfFieldVignetteSize = FMath::Max(0.0f, Context.View.FinalPostProcessSettings.DepthOfFieldVignetteSize / 100.0f * 2);
+	// doesn't make much sense to expose this property as the effect is very non linear and it would cost some performance to fix that
+	float DepthOfFieldVignetteFeather = 10.0f / 100.0f;
+	float DepthOfFieldVignetteMul = 1.0f / DepthOfFieldVignetteFeather;
+	float DepthOfFieldVignetteAdd = (0.5f - DepthOfFieldVignetteSize) * DepthOfFieldVignetteMul;
+	Out[0] = FVector4(
+		(SkyFocusDistance > 0) ? SkyFocusDistance : 100000000.0f,			// very large if <0 to not mask out skybox, can be optimized to disable feature completely
+		DepthOfFieldVignetteMul,
+		DepthOfFieldVignetteAdd,
+		Context.View.FinalPostProcessSettings.DepthOfFieldOcclusion);
+	FIntPoint ViewSize = Context.View.ViewRect.Size();
+	
+	// Scale and offset to put two views in one texture with safety border
+	float UsedYDivTextureY = HalfRes / (float)BokehLayerSizeY;
+	float YOffsetInPixel = HalfRes + SafetyBorder;
+	float YOffsetInUV = (HalfRes + SafetyBorder) / (float)BokehLayerSizeY;
+	Out[1] = FVector4(0.0f, YOffsetInUV, UsedYDivTextureY, YOffsetInPixel);
+}
diff --git a/Engine/Source/Runtime/Renderer/Private/PostProcess/PostProcessDOF.h b/Engine/Source/Runtime/Renderer/Private/PostProcess/PostProcessDOF.h
--- a/Engine/Source/Runtime/Renderer/Private/PostProcess/PostProcessDOF.h
+++ b/Engine/Source/Runtime/Renderer/Private/PostProcess/PostProcessDOF.h
@@ -21,6 +21,8 @@ public:
 		: bFarBlur(bInFarBlur)
 		, bNearBlur(bInNearBlur)
 	{}
+	
+	static void ComputeDepthOfFieldParams(const FRenderingCompositePassContext& Context, FVector4 Out[2]);
 
 	// interface FRenderingCompositePass ---------
 
diff --git a/Engine/Source/Runtime/Renderer/Private/PostProcess/PostProcessMobile.cpp b/Engine/Source/Runtime/Renderer/Private/PostProcess/PostProcessMobile.cpp
--- a/Engine/Source/Runtime/Renderer/Private/PostProcess/PostProcessMobile.cpp
+++ b/Engine/Source/Runtime/Renderer/Private/PostProcess/PostProcessMobile.cpp
@@ -24,9 +24,7 @@ static EPixelFormat GetHDRPixelFormat()
 // return Depth of Field Scale if Gaussian DoF mode is active. 0.0f otherwise.
 float GetMobileDepthOfFieldScale(const FViewInfo& View)
 {
-	PRAGMA_DISABLE_DEPRECATION_WARNINGS
-		return View.FinalPostProcessSettings.DepthOfFieldMethod == DOFM_Gaussian ? View.FinalPostProcessSettings.DepthOfFieldScale : 0.0f;
-	PRAGMA_ENABLE_DEPRECATION_WARNINGS
+	return View.FinalPostProcessSettings.DepthOfFieldScale;
 }
 
 //
diff --git a/Engine/Source/Runtime/Renderer/Private/PostProcess/PostProcessing.cpp b/Engine/Source/Runtime/Renderer/Private/PostProcess/PostProcessing.cpp
--- a/Engine/Source/Runtime/Renderer/Private/PostProcess/PostProcessing.cpp
+++ b/Engine/Source/Runtime/Renderer/Private/PostProcess/PostProcessing.cpp
@@ -31,12 +31,10 @@
 #include "PostProcess/PostProcessLensFlares.h"
 #include "PostProcess/PostProcessLensBlur.h"
 #include "PostProcess/PostProcessBokehDOF.h"
-#include "PostProcess/PostProcessBokehDOFRecombine.h"
 #include "PostProcess/PostProcessCombineLUTs.h"
 #include "PostProcess/PostProcessTemporalAA.h"
 #include "PostProcess/PostProcessMotionBlur.h"
 #include "PostProcess/PostProcessDOF.h"
-#include "PostProcess/PostProcessCircleDOF.h"
 #include "PostProcess/PostProcessUpscale.h"
 #include "PostProcess/PostProcessHMD.h"
 #include "PostProcess/PostProcessMitchellNetravali.h"
@@ -185,14 +183,6 @@ TAutoConsoleVariable<int32> CVarTemporalAAAllowDownsampling(
 	TEXT("Allows half-resolution color buffer to be produced during TAA. Only possible when motion blur is off and when using compute shaders for post processing."),
 	ECVF_RenderThreadSafe);
 
-TAutoConsoleVariable<int32> CVarUseDiaphragmDOF(
-	TEXT("r.DOF.Algorithm"),
-	1,
-	TEXT("Chooses the depth of field algorithm to use. \n")
-	TEXT(" 0: Circle DOF (old implementation);")
-	TEXT(" 1: Diaphragm DOF (default).\n"),
-	ECVF_RenderThreadSafe);
-
 static TAutoConsoleVariable<float> CVarTemporalAAHistorySP(
 	TEXT("r.TemporalAA.HistoryScreenPercentage"),
 	100.0f,
@@ -490,48 +480,6 @@ static void AddVisualizeBloomOverlay(FPostprocessContext& Context, FRenderingCom
 	Context.FinalOutput = FRenderingCompositeOutputRef(Node);
 }
 
-static void AddPostProcessDepthOfFieldBokeh(FPostprocessContext& Context, FRenderingCompositeOutputRef& SeparateTranslucency, FRenderingCompositeOutputRef& VelocityInput)
-{
-	// downsample, mask out the in focus part, depth in alpha
-	const bool bIsComputePass = ShouldDoComputePostProcessing(Context.View);
-	FRenderingCompositePass* DOFSetup = Context.Graph.RegisterPass(new(FMemStack::Get()) FRCPassPostProcessBokehDOFSetup(bIsComputePass));
-	DOFSetup->SetInput(ePId_Input0, FRenderingCompositeOutputRef(Context.FinalOutput));
-	DOFSetup->SetInput(ePId_Input1, FRenderingCompositeOutputRef(Context.SceneDepth));
-
-	FSceneViewState* ViewState = (FSceneViewState*)Context.View.State;
-	
-	FRenderingCompositePass* DOFInputPass = DOFSetup;
-	if( Context.View.AntiAliasingMethod == AAM_TemporalAA && ViewState )
-	{
-		FTAAPassParameters Parameters(Context.View);
-		Parameters.Pass = ETAAPassConfig::LegacyDepthOfField;
-		Parameters.bIsComputePass = bIsComputePass;
-		Parameters.SetupViewRect(Context.View, /* ResolutionDivisor = */ 2);
-
-		FRenderingCompositePass* NodeTemporalAA = Context.Graph.RegisterPass( new(FMemStack::Get()) FRCPassPostProcessTemporalAA(
-			Context, Parameters,
-			ViewState->DOFHistory,
-			&ViewState->DOFHistory) );
-		NodeTemporalAA->SetInput( ePId_Input0, DOFSetup );
-		NodeTemporalAA->SetInput( ePId_Input2, VelocityInput );
-
-		DOFInputPass = NodeTemporalAA;
-		ViewState->bDOFHistory = true;
-	}
-
-	FRenderingCompositePass* NodeBlurred = Context.Graph.RegisterPass(new(FMemStack::Get()) FRCPassPostProcessBokehDOF());
-	NodeBlurred->SetInput(ePId_Input0, DOFInputPass);
-	NodeBlurred->SetInput(ePId_Input1, Context.SceneColor);
-	NodeBlurred->SetInput(ePId_Input2, Context.SceneDepth);
-
-	FRenderingCompositePass* NodeRecombined = Context.Graph.RegisterPass(new(FMemStack::Get()) FRCPassPostProcessBokehDOFRecombine(bIsComputePass));
-	NodeRecombined->SetInput(ePId_Input0, Context.FinalOutput);
-	NodeRecombined->SetInput(ePId_Input1, NodeBlurred);
-	NodeRecombined->SetInput(ePId_Input2, SeparateTranslucency);
-
-	Context.FinalOutput = FRenderingCompositeOutputRef(NodeRecombined);
-}
-
 static bool AddPostProcessDepthOfFieldGaussian(FPostprocessContext& Context, FDepthOfFieldStats& Out, FRenderingCompositeOutputRef& VelocityInput, FRenderingCompositeOutputRef& SeparateTranslucencyRef)
 {
 	// GaussianDOFPass performs Gaussian setup, blur and recombine.
@@ -645,59 +593,6 @@ static bool AddPostProcessDepthOfFieldGaussian(FPostprocessContext& Context, FDe
 	}
 }
 
-static void AddPostProcessDepthOfFieldCircle(FPostprocessContext& Context, const FRenderingCompositeOutputRef& VelocityInput)
-{
-	if(Context.View.Family->EngineShowFlags.VisualizeDOF)
-	{
-		// no need for this pass
-		return;
-	}
-
-	FRenderingCompositePass* DOFSetup = Context.Graph.RegisterPass(new(FMemStack::Get()) FRCPassPostProcessCircleDOFSetup());
-	DOFSetup->SetInput(ePId_Input0, FRenderingCompositeOutputRef(Context.FinalOutput));
-	DOFSetup->SetInput(ePId_Input1, FRenderingCompositeOutputRef(Context.SceneDepth));
-
-	FSceneViewState* ViewState = (FSceneViewState*)Context.View.State;
-
-	auto ColorSetup = FRenderingCompositeOutputRef( DOFSetup, ePId_Output0 );
-	auto CocSetup = FRenderingCompositeOutputRef( DOFSetup, FPostProcessing::HasAlphaChannelSupport() ? ePId_Output1 : ePId_Output0 );
-	if( Context.View.AntiAliasingMethod == AAM_TemporalAA && ViewState )
-	{
-		FTAAPassParameters Parameters(Context.View);
-		Parameters.Pass = ETAAPassConfig::LegacyDepthOfField;
-		Parameters.bIsComputePass = ShouldDoComputePostProcessing(Context.View);
-		Parameters.SetupViewRect(Context.View, /* ResolutionDivisor = */ 2);
-
-		FRenderingCompositePass* NodeTemporalAA = Context.Graph.RegisterPass( new(FMemStack::Get()) FRCPassPostProcessTemporalAA(
-			Context, Parameters,
-			ViewState->DOFHistory,
-			&ViewState->DOFHistory) );
-		NodeTemporalAA->SetInput( ePId_Input0, DOFSetup );
-		NodeTemporalAA->SetInput( ePId_Input2, VelocityInput );
-
-		ColorSetup = FRenderingCompositeOutputRef(NodeTemporalAA);
-		ViewState->bDOFHistory = false;
-	}
-
-	FRenderingCompositePass* DOFNear = Context.Graph.RegisterPass(new(FMemStack::Get()) FRCPassPostProcessCircleDOFDilate());
-	DOFNear->SetInput(ePId_Input0, CocSetup);
-	FRenderingCompositeOutputRef Near = FRenderingCompositeOutputRef(DOFNear, ePId_Output0);
-
-	FRenderingCompositePass* DOFApply = Context.Graph.RegisterPass(new(FMemStack::Get()) FRCPassPostProcessCircleDOF());
-	DOFApply->SetInput(ePId_Input0, ColorSetup);
-	DOFApply->SetInput(ePId_Input1, Near);
-	DOFApply->SetInput(ePId_Input2, CocSetup);
-	FRenderingCompositeOutputRef Far = FRenderingCompositeOutputRef(DOFApply, ePId_Output0);
-	FRenderingCompositeOutputRef FarCoc = FPostProcessing::HasAlphaChannelSupport() ? FRenderingCompositeOutputRef(DOFApply, ePId_Output1) : Far;
-
-	FRenderingCompositePass* NodeRecombined = Context.Graph.RegisterPass(new(FMemStack::Get()) FRCPassPostProcessCircleDOFRecombine());
-	NodeRecombined->SetInput(ePId_Input0, Context.FinalOutput);
-	NodeRecombined->SetInput(ePId_Input1, Far);
-	NodeRecombined->SetInput(ePId_Input2, FarCoc);
-
-	Context.FinalOutput = FRenderingCompositeOutputRef(NodeRecombined);
-}
-
 static FRenderingCompositeOutputRef AddBloom(FBloomDownSampleArray& BloomDownSampleArray, bool bVisualizeBloom)
 {
 	
@@ -1514,147 +1409,50 @@ void FPostProcessing::Process(FRHICommandListImmediate& RHICmdList, const FViewI
 			}
 
 			Context.FinalOutput = AddPostProcessMaterialChain(Context, BL_BeforeTranslucency, SeparateTranslucency);
-
+			
 			static const auto CVar = IConsoleManager::Get().FindTConsoleVariableDataInt(TEXT("r.DepthOfFieldQuality"));
 			check(CVar)
-			bool bDepthOfField = View.Family->EngineShowFlags.DepthOfField && CVar->GetValueOnRenderThread() > 0;
-
-			FDepthOfFieldStats DepthOfFieldStat;
+			bool bDepthOfField = 
+				View.Family->EngineShowFlags.DepthOfField &&
+				CVar->GetValueOnRenderThread() > 0 &&
+				View.FinalPostProcessSettings.DepthOfFieldFstop > 0;
 
-			bool bSepTransWasApplied = false;
-
-			PRAGMA_DISABLE_DEPRECATION_WARNINGS
-			if(bDepthOfField && View.FinalPostProcessSettings.DepthOfFieldMethod != DOFM_BokehDOF)
+			// Applies DOF and separate translucency,
 			{
-				if(View.FinalPostProcessSettings.DepthOfFieldMethod == DOFM_Gaussian)
+				FRenderingCompositePass* DiaphragmDOFPass = Context.Graph.RegisterPass(
+					new(FMemStack::Get()) TRCPassForRDG<2, 1>([bDepthOfField](FRenderingCompositePass* Pass, FRenderingCompositePassContext& InContext)
 				{
-					static bool bHasWarned = false;
-					if (!bHasWarned)
-					{
-						UE_LOG(LogRenderer, Warning, TEXT("Gaussian DOF algorithm is deprecated and will be removed from deferred shading renderer. Consider using the Circle DOF method in post process settings."));
-						bHasWarned = true;
-					}
+					FRDGBuilder GraphBuilder(InContext.RHICmdList);
 
-					if (FPostProcessing::HasAlphaChannelSupport())
-					{
-						UE_LOG(LogRenderer, Log, TEXT("Gaussian depth of field does not have alpha channel support. Only Circle DOF has."));
-					}
-					if(VelocityInput.IsValid())
-					{
-						bSepTransWasApplied = AddPostProcessDepthOfFieldGaussian(Context, DepthOfFieldStat, VelocityInput, SeparateTranslucency);
-					}
-					else
-					{
-						// todo: black/white default is a compositing graph feature, no need to hook up a node
+					FSceneViewFamilyBlackboard SceneBlackboard;
+					SetupSceneViewFamilyBlackboard(GraphBuilder, &SceneBlackboard);
+	
+					FRDGTextureRef SceneColor = Pass->CreateRDGTextureForInput(GraphBuilder, ePId_Input0, TEXT("SceneColor"), eFC_0000);
+					FRDGTextureRef LocalSeparateTranslucency = Pass->CreateRDGTextureForInput(GraphBuilder, ePId_Input1, TEXT("SeparateTranslucency"), eFC_0000);
 
-						// black is how we clear the velocity buffer so this means no velocity
-						FRenderingCompositePass* NoVelocity = Context.Graph.RegisterPass(new(FMemStack::Get()) FRCPassPostProcessInput(GSystemTextures.BlackDummy));
-						FRenderingCompositeOutputRef NoVelocityRef(NoVelocity);
-						bSepTransWasApplied = AddPostProcessDepthOfFieldGaussian(Context, DepthOfFieldStat, NoVelocityRef, SeparateTranslucency);
-					}
-				}
-				else
-				{
-					FRenderingCompositeOutputRef DOFVelocityRef(VelocityInput);
+					FRDGTextureRef NewSceneColor = SceneColor;
 
-					if(!DOFVelocityRef.IsValid())
+					if (bDepthOfField && DiaphragmDOF::IsSupported(InContext.View.GetShaderPlatform()))
 					{
-						FRenderingCompositePass* NoVelocity = Context.Graph.RegisterPass(new(FMemStack::Get()) FRCPassPostProcessInput(GSystemTextures.BlackDummy));
-						DOFVelocityRef = FRenderingCompositeOutputRef(NoVelocity);
+						NewSceneColor = DiaphragmDOF::AddPasses(
+							GraphBuilder,
+							SceneBlackboard, InContext.View,
+							SceneColor, LocalSeparateTranslucency);
 					}
-					
-					if (CVarUseDiaphragmDOF.GetValueOnRenderThread() > 0 &&
-						DiaphragmDOF::IsSupported(View.GetShaderPlatform()))
-					{
-						FRenderingCompositePass* DiaphragmDOFPass = Context.Graph.RegisterPass(
-							new(FMemStack::Get()) TRCPassForRDG<2, 1>([](FRenderingCompositePass* Pass, FRenderingCompositePassContext& InContext)
-						{
-							FRDGBuilder GraphBuilder(InContext.RHICmdList);
-
-							FSceneViewFamilyBlackboard SceneBlackboard;
-							SetupSceneViewFamilyBlackboard(GraphBuilder, &SceneBlackboard);
-	
-							FRDGTextureRef SceneColor = Pass->CreateRDGTextureForInput(GraphBuilder, ePId_Input0, TEXT("SceneColor"), eFC_0000);
-							FRDGTextureRef LocalSeparateTranslucency = Pass->CreateRDGTextureForInput(GraphBuilder, ePId_Input1, TEXT("SeparateTranslucency"), eFC_0000);
 
-							FRDGTextureRef NewSceneColor = DiaphragmDOF::AddPasses(
-								GraphBuilder,
-								SceneBlackboard, InContext.View,
-								SceneColor, LocalSeparateTranslucency);
-		
-							// DOF passes were not added, therefore need to compose Separate translucency manually. 
-							if (NewSceneColor == SceneColor)
-							{
-								NewSceneColor = AddSeparateTranslucencyCompositionPass(GraphBuilder, InContext.View, SceneColor, LocalSeparateTranslucency);
-							}
-
-							Pass->ExtractRDGTextureForOutput(GraphBuilder, ePId_Output0, NewSceneColor);
-
-							GraphBuilder.Execute();
-						}));
-						DiaphragmDOFPass->SetInput(ePId_Input0, Context.FinalOutput);
-						DiaphragmDOFPass->SetInput(ePId_Input1, SeparateTranslucency);
-						Context.FinalOutput = FRenderingCompositeOutputRef(DiaphragmDOFPass, ePId_Output0);
-
-						bSepTransWasApplied = true;
-					}
-					else
+					// DOF passes were not added, therefore need to compose Separate translucency manually. 
+					if (NewSceneColor == SceneColor)
 					{
-						static bool bHasWarned = false;
-						if (!bHasWarned)
-						{
-							UE_LOG(LogRenderer, Warning, TEXT("Circle DOF algorithm is deprecated and will be removed. Consider using Diphragm DOF with r.DOF.Algorithm = 1."));
-							bHasWarned = true;
-						}
-						AddPostProcessDepthOfFieldCircle(Context, DOFVelocityRef);
+						NewSceneColor = AddSeparateTranslucencyCompositionPass(GraphBuilder, InContext.View, SceneColor, LocalSeparateTranslucency);
 					}
-				}
-			}
-
-			bool bBokehDOF = bDepthOfField
-				&& View.FinalPostProcessSettings.DepthOfFieldScale > 0
-				&& View.FinalPostProcessSettings.DepthOfFieldMethod == DOFM_BokehDOF
-				&& !Context.View.Family->EngineShowFlags.VisualizeDOF;
-			PRAGMA_ENABLE_DEPRECATION_WARNINGS
-
-			if(bBokehDOF)
-			{
-				static bool bHasWarned = false;
-				if (!bHasWarned)
-				{
-					UE_LOG(LogRenderer, Warning, TEXT("Bokeh DOF algorithm is deprecated and will be removed. Consider using the Circle DOF method in post process settings."));
-					bHasWarned = true;
-				}
-
-				if (FPostProcessing::HasAlphaChannelSupport())
-				{
-					UE_LOG(LogRenderer, Log, TEXT("Boked depth of field does not have alpha channel support. Only Circle DOF has."));
-				}
-				if(VelocityInput.IsValid())
-				{
-					AddPostProcessDepthOfFieldBokeh(Context, SeparateTranslucency, VelocityInput);
-				}
-				else
-				{
-					// todo: black/white default is a compositing graph feature, no need to hook up a node
 
-					// black is how we clear the velocity buffer so this means no velocity
-					FRenderingCompositePass* NoVelocity = Context.Graph.RegisterPass(new(FMemStack::Get()) FRCPassPostProcessInput(GSystemTextures.BlackDummy));
-					FRenderingCompositeOutputRef NoVelocityRef(NoVelocity);
-					AddPostProcessDepthOfFieldBokeh(Context, SeparateTranslucency, NoVelocityRef);
-				}
-				bSepTransWasApplied = true;
-			}
-
-			if(SeparateTranslucency.IsValid() && !bSepTransWasApplied)
-			{
-				const bool bIsComputePass = ShouldDoComputePostProcessing(Context.View);
-				// separate translucency is done here or in AddPostProcessDepthOfFieldBokeh()
-				FRenderingCompositePass* NodeRecombined = Context.Graph.RegisterPass(new(FMemStack::Get()) FRCPassPostProcessBokehDOFRecombine(bIsComputePass));
-				NodeRecombined->SetInput(ePId_Input0, Context.FinalOutput);
-				NodeRecombined->SetInput(ePId_Input2, SeparateTranslucency);
+					Pass->ExtractRDGTextureForOutput(GraphBuilder, ePId_Output0, NewSceneColor);
 
-				Context.FinalOutput = FRenderingCompositeOutputRef(NodeRecombined);
+					GraphBuilder.Execute();
+				}));
+				DiaphragmDOFPass->SetInput(ePId_Input0, Context.FinalOutput);
+				DiaphragmDOFPass->SetInput(ePId_Input1, SeparateTranslucency);
+				Context.FinalOutput = FRenderingCompositeOutputRef(DiaphragmDOFPass, ePId_Output0);
 			}
 
 			Context.FinalOutput = AddPostProcessMaterialChain(Context, BL_BeforeTonemapping, SeparateTranslucency, nullptr, nullptr, PreFlattenVelocity);
@@ -2039,6 +1837,8 @@ void FPostProcessing::Process(FRHICommandListImmediate& RHICmdList, const FViewI
 			
 			if(bDepthOfField && Context.View.Family->EngineShowFlags.VisualizeDOF)
 			{
+				FDepthOfFieldStats DepthOfFieldStat;
+
 				FRenderingCompositePass* VisualizeNode = Context.Graph.RegisterPass(new(FMemStack::Get()) FRCPassPostProcessVisualizeDOF(DepthOfFieldStat));
 				VisualizeNode->SetInput(ePId_Input0, FRenderingCompositeOutputRef(Context.FinalOutput));
 
@@ -2048,15 +1848,28 @@ void FPostProcessing::Process(FRHICommandListImmediate& RHICmdList, const FViewI
 		}
 		else
 		{
-			if (SeparateTranslucency.IsValid())
+			// Composes separate translucency,
 			{
-				const bool bIsComputePass = ShouldDoComputePostProcessing(Context.View);
-				// separate translucency is done here or in AddPostProcessDepthOfFieldBokeh()
-				FRenderingCompositePass* NodeRecombined = Context.Graph.RegisterPass(new(FMemStack::Get()) FRCPassPostProcessBokehDOFRecombine(bIsComputePass));
-				NodeRecombined->SetInput(ePId_Input0, Context.FinalOutput);
-				NodeRecombined->SetInput(ePId_Input2, SeparateTranslucency);
+				FRenderingCompositePass* ComposeSeparateTranslucencyPass = Context.Graph.RegisterPass(
+					new(FMemStack::Get()) TRCPassForRDG<2, 1>([](FRenderingCompositePass* Pass, FRenderingCompositePassContext& InContext)
+				{
+					FRDGBuilder GraphBuilder(InContext.RHICmdList);
 
-				Context.FinalOutput = FRenderingCompositeOutputRef(NodeRecombined);
+					FSceneViewFamilyBlackboard SceneBlackboard;
+					SetupSceneViewFamilyBlackboard(GraphBuilder, &SceneBlackboard);
+	
+					FRDGTextureRef SceneColor = Pass->CreateRDGTextureForInput(GraphBuilder, ePId_Input0, TEXT("SceneColor"), eFC_0000);
+					FRDGTextureRef LocalSeparateTranslucency = Pass->CreateRDGTextureForInput(GraphBuilder, ePId_Input1, TEXT("SeparateTranslucency"), eFC_0000);
+
+					FRDGTextureRef NewSceneColor = AddSeparateTranslucencyCompositionPass(GraphBuilder, InContext.View, SceneColor, LocalSeparateTranslucency);
+
+					Pass->ExtractRDGTextureForOutput(GraphBuilder, ePId_Output0, NewSceneColor);
+
+					GraphBuilder.Execute();
+				}));
+				ComposeSeparateTranslucencyPass->SetInput(ePId_Input0, Context.FinalOutput);
+				ComposeSeparateTranslucencyPass->SetInput(ePId_Input1, SeparateTranslucency);
+				Context.FinalOutput = FRenderingCompositeOutputRef(ComposeSeparateTranslucencyPass, ePId_Output0);
 			}
 
 			// Shader complexity does not actually output a color
@@ -2547,15 +2360,17 @@ void FPostProcessing::ProcessES2(FRHICommandListImmediate& RHICmdList, FScene* S
 						// black is how we clear the velocity buffer so this means no velocity
 						FRenderingCompositePass* NoVelocity = Context.Graph.RegisterPass(new(FMemStack::Get()) FRCPassPostProcessInput(GSystemTextures.BlackDummy));
 						FRenderingCompositeOutputRef NoVelocityRef(NoVelocity);
+						
+						bool bDepthOfField = 
+							View.Family->EngineShowFlags.DepthOfField &&
+							IsGaussianActive(Context);
 
-						PRAGMA_DISABLE_DEPRECATION_WARNINGS
-						if(View.FinalPostProcessSettings.DepthOfFieldMethod == DOFM_Gaussian && IsGaussianActive(Context))
+						if(bDepthOfField)
 						{
 							FDepthOfFieldStats DepthOfFieldStat;
 							FRenderingCompositeOutputRef DummySeparateTranslucency;
 							AddPostProcessDepthOfFieldGaussian(Context, DepthOfFieldStat, NoVelocityRef, DummySeparateTranslucency);
 						}
-						PRAGMA_ENABLE_DEPRECATION_WARNINGS
 					}
 				}
 
diff --git a/Engine/Source/Runtime/Renderer/Private/SceneRendering.cpp b/Engine/Source/Runtime/Renderer/Private/SceneRendering.cpp
--- a/Engine/Source/Runtime/Renderer/Private/SceneRendering.cpp
+++ b/Engine/Source/Runtime/Renderer/Private/SceneRendering.cpp
@@ -27,7 +27,6 @@
 #include "LegacyScreenPercentageDriver.h"
 #include "SceneViewExtension.h"
 #include "PostProcess/PostProcessBusyWait.h"
-#include "PostProcess/PostProcessCircleDOF.h"
 #include "AtmosphereRendering.h"
 #include "Matinee/MatineeActor.h"
 #include "ComponentRecreateRenderStateContext.h"
@@ -51,6 +50,7 @@
 #include "VisualizeTexturePresent.h"
 #include "MeshDrawCommands.h"
 #include "HAL/LowLevelMemTracker.h"
+#include "DiaphragmDOF.h"
 
 /*-----------------------------------------------------------------------------
 	Globals
@@ -1373,7 +1373,7 @@ void FViewInfo::SetupUniformBufferParameters(
 		}
 	}
 
-	ViewUniformShaderParameters.CircleDOFParams = CircleDofHalfCoc(*this);
+	ViewUniformShaderParameters.CircleDOFParams = DiaphragmDOF::CircleDofHalfCoc(*this);
 
 	ERHIFeatureLevel::Type RHIFeatureLevel = Scene == nullptr ? GMaxRHIFeatureLevel : Scene->GetFeatureLevel();
 
