Adding a tiled blue-noise image. The image is a tiled set of 7 256x256 blue noise patterns.

Image data sets donated by [REDACTED_PERSON] and converted to a single uasset texture: BlueNoise.uasset.

diff --git a/Engine/Config/BaseEngine.ini b/Engine/Config/BaseEngine.ini
--- a/Engine/Config/BaseEngine.ini
+++ b/Engine/Config/BaseEngine.ini
@@ -88,6 +88,7 @@ DefaultBloomKernelTextureName=/Engine/EngineMaterials/DefaultBloomKernel.Default
 GeomMaterialName=/Engine/EngineDebugMaterials/GeomMaterial.GeomMaterial
 DebugMeshMaterialName=/Engine/EngineDebugMaterials/DebugMeshMaterial.DebugMeshMaterial
 PreIntegratedSkinBRDFTextureName=/Engine/EngineMaterials/PreintegratedSkinBRDF.PreintegratedSkinBRDF
+BlueNoiseTextureName=/Engine/EngineMaterials/BlueNoise.BlueNoise
 MiniFontTextureName=/Engine/EngineMaterials/MiniFont.MiniFont
 WeightMapPlaceholderTextureName=/Engine/EngineMaterials/WeightMapPlaceholderTexture.WeightMapPlaceholderTexture
 LightMapDensityTextureName=/Engine/EngineMaterials/DefaultWhiteGrid.DefaultWhiteGrid
diff --git a/Engine/Shaders/Private/BlueNoise.ush b/Engine/Shaders/Private/BlueNoise.ush
new file mode 100644
--- /dev/null
+++ b/Engine/Shaders/Private/BlueNoise.ush
@@ -0,0 +1,8 @@
+// Copyright 1998-2019 Epic Games, Inc. All Rights Reserved.
+
+float3 EvalBlueNoise(uint2 TileCoord, uint Slice)
+{
+	uint3 BlueNoiseCoord = uint3(TileCoord.x, Slice * BlueNoise.Dimensions.y + TileCoord.y, 0);
+	float3 Noise = BlueNoise.Texture.Load(BlueNoiseCoord, 0);
+	return Noise;
+}
diff --git a/Engine/Source/Runtime/Engine/Classes/Engine/Engine.h b/Engine/Source/Runtime/Engine/Classes/Engine/Engine.h
--- a/Engine/Source/Runtime/Engine/Classes/Engine/Engine.h
+++ b/Engine/Source/Runtime/Engine/Classes/Engine/Engine.h
@@ -1191,6 +1191,13 @@ public:
 	UPROPERTY(globalconfig)
 	FSoftObjectPath PreIntegratedSkinBRDFTextureName;
 
+	/** Tiled blue-noise texture */
+	UPROPERTY()
+	class UTexture2D* BlueNoiseTexture;
+
+	UPROPERTY(globalconfig)
+	FSoftObjectPath BlueNoiseTextureName;
+
 	/** Texture used to do font rendering in shaders */
 	UPROPERTY()
 	class UTexture2D* MiniFontTexture;
diff --git a/Engine/Source/Runtime/Engine/Private/UnrealEngine.cpp b/Engine/Source/Runtime/Engine/Private/UnrealEngine.cpp
--- a/Engine/Source/Runtime/Engine/Private/UnrealEngine.cpp
+++ b/Engine/Source/Runtime/Engine/Private/UnrealEngine.cpp
@@ -2404,6 +2404,9 @@ void UEngine::InitializeObjectReferences()
 	LoadEngineTexture(MiniFontTexture, *MiniFontTextureName.ToString());
 	LoadEngineTexture(WeightMapPlaceholderTexture, *WeightMapPlaceholderTextureName.ToString());
 	LoadEngineTexture(LightMapDensityTexture, *LightMapDensityTextureName.ToString());
+#if RHI_RAYTRACING
+	LoadEngineTexture(BlueNoiseTexture, *BlueNoiseTextureName.ToString());
+#endif
 
 	if ( DefaultPhysMaterial == NULL )
 	{
diff --git a/Engine/Source/Runtime/Renderer/Private/BlueNoise.cpp b/Engine/Source/Runtime/Renderer/Private/BlueNoise.cpp
new file mode 100644
--- /dev/null
+++ b/Engine/Source/Runtime/Renderer/Private/BlueNoise.cpp
@@ -0,0 +1,9 @@
+// Copyright 1998-2019 Epic Games, Inc. All Rights Reserved.
+
+/*=============================================================================
+	BlueNoise.cpp: Resources for Blue-Noise vectors on the GPU.
+=============================================================================*/
+
+#include "BlueNoise.h"
+
+IMPLEMENT_GLOBAL_SHADER_PARAMETER_STRUCT(FBlueNoise, "BlueNoise");
diff --git a/Engine/Source/Runtime/Renderer/Private/BlueNoise.h b/Engine/Source/Runtime/Renderer/Private/BlueNoise.h
new file mode 100644
--- /dev/null
+++ b/Engine/Source/Runtime/Renderer/Private/BlueNoise.h
@@ -0,0 +1,24 @@
+// Copyright 1998-2019 Epic Games, Inc. All Rights Reserved.
+
+/*=============================================================================
+	BlueNoise.h: Resources for Blue-Noise vectors on the GPU.
+=============================================================================*/
+
+#pragma once
+
+#include "Engine.h"
+#include "UniformBuffer.h"
+
+// Texture data is assumed to be in tiled representation where:
+// 1) Dimensions.xy represents a single blue-noise tile
+// 2) Dimensions.z represents the number of slices available
+BEGIN_GLOBAL_SHADER_PARAMETER_STRUCT(FBlueNoise, )
+SHADER_PARAMETER(FIntVector, Dimensions)
+SHADER_PARAMETER_TEXTURE(Texture2D, Texture)
+END_GLOBAL_SHADER_PARAMETER_STRUCT()
+
+inline void InitializeBlueNoise(FBlueNoise& BlueNoise)
+{
+	BlueNoise.Dimensions = FIntVector(GEngine->BlueNoiseTexture->GetSizeX(), GEngine->BlueNoiseTexture->GetSizeX(), GEngine->BlueNoiseTexture->GetSizeY() / GEngine->BlueNoiseTexture->GetSizeX());
+	BlueNoise.Texture = GEngine->BlueNoiseTexture->Resource->TextureRHI;
+}
