Unified GLightCountMaximum and LIGHT_COUNT_MAXIMUM in ray tracing c++ and shader code. No functional changes.

diff --git a/Engine/Shaders/Private/PathTracing/Light/PathTracingLightSampling.ush b/Engine/Shaders/Private/PathTracing/Light/PathTracingLightSampling.ush
--- a/Engine/Shaders/Private/PathTracing/Light/PathTracingLightSampling.ush
+++ b/Engine/Shaders/Private/PathTracing/Light/PathTracingLightSampling.ush
@@ -13,9 +13,6 @@ LightSampling.usf: Light sampling for next-event estimation
 #include "PathTracingSkyLight.ush"
 #include "PathTracingSpotLight.ush"
 
-// #dxr_todo:UE-72532 unify with GLightCountMaximum in c++ code
-static const uint LIGHT_COUNT_MAXIMUM = 64;
-
 void GenerateLightRay(RayDesc Ray, int LightId, float3 LightUV, out RayDesc OutLightRay)
 {
 	if (IsPointLight(LightId))
@@ -88,7 +85,7 @@ bool SampleLightSelection(
 	if (SceneLightsData.Count == 0) return false;
 
 	// Build irradiance estimate prefix sum
-	float IrradianceEstimate[LIGHT_COUNT_MAXIMUM];
+	float IrradianceEstimate[RAY_TRACING_LIGHT_COUNT_MAXIMUM];
 	EstimateLight(0, Ray, Payload, IrradianceEstimate[0]);
 	{
 		for (uint LightIndex = 1; LightIndex < SceneLightsData.Count; ++LightIndex)
@@ -140,7 +137,7 @@ void PdfLightSelection(
 	if (SceneLightsData.Count == 0) return;
 
 	// Build irradiance estimate prefix sum
-	float IrradianceEstimate[LIGHT_COUNT_MAXIMUM];
+	float IrradianceEstimate[RAY_TRACING_LIGHT_COUNT_MAXIMUM];
 	EstimateLight(0, Ray, Payload, IrradianceEstimate[0]);
 	{
 		for (uint LightIndex = 1; LightIndex < SceneLightsData.Count; ++LightIndex)
diff --git a/Engine/Shaders/Shared/RayTracingDefinitions.h b/Engine/Shaders/Shared/RayTracingDefinitions.h
--- a/Engine/Shaders/Shared/RayTracingDefinitions.h
+++ b/Engine/Shaders/Shared/RayTracingDefinitions.h
@@ -24,3 +24,5 @@
 #define RAY_TRACING_SHADER_SLOT_SHADOW		1
 #define RAY_TRACING_NUM_SHADER_SLOTS		2
 
+#define RAY_TRACING_LIGHT_COUNT_MAXIMUM		64
+
diff --git a/Engine/Source/Runtime/Renderer/Private/PathTracing.cpp b/Engine/Source/Runtime/Renderer/Private/PathTracing.cpp
--- a/Engine/Source/Runtime/Renderer/Private/PathTracing.cpp
+++ b/Engine/Source/Runtime/Renderer/Private/PathTracing.cpp
@@ -171,7 +171,7 @@ public:
 
 			for (auto Light : Lights)
 			{
-				if (LightData.Count >= GLightCountMaximum) break;
+				if (LightData.Count >= RAY_TRACING_LIGHT_COUNT_MAXIMUM) break;
 
 				if (Light.LightSceneInfo->Proxy->HasStaticLighting() && Light.LightSceneInfo->IsPrecomputedLightingValid()) continue;
 
diff --git a/Engine/Source/Runtime/Renderer/Private/PathTracingUniformBuffers.h b/Engine/Source/Runtime/Renderer/Private/PathTracingUniformBuffers.h
--- a/Engine/Source/Runtime/Renderer/Private/PathTracingUniformBuffers.h
+++ b/Engine/Source/Runtime/Renderer/Private/PathTracingUniformBuffers.h
@@ -7,6 +7,7 @@
 #pragma once
 
 #include "UniformBuffer.h"
+#include "RayTracingDefinitions.h"
 
 BEGIN_GLOBAL_SHADER_PARAMETER_STRUCT(FPathTracingData, )
 	SHADER_PARAMETER(uint32, MaxBounces)
@@ -14,23 +15,22 @@ END_GLOBAL_SHADER_PARAMETER_STRUCT()
 
 
 // Lights
-const int32 GLightCountMaximum = 64;
 
 BEGIN_GLOBAL_SHADER_PARAMETER_STRUCT(FPathTracingLightData, )
 	SHADER_PARAMETER(uint32, Count)
-	SHADER_PARAMETER_ARRAY(uint32, Type, [GLightCountMaximum])
+	SHADER_PARAMETER_ARRAY(uint32, Type, [RAY_TRACING_LIGHT_COUNT_MAXIMUM])
 	// Geometry
-	SHADER_PARAMETER_ARRAY(FVector, Position, [GLightCountMaximum])
-	SHADER_PARAMETER_ARRAY(FVector, Normal, [GLightCountMaximum])
-	SHADER_PARAMETER_ARRAY(FVector, dPdu, [GLightCountMaximum])
-	SHADER_PARAMETER_ARRAY(FVector, dPdv, [GLightCountMaximum])
+	SHADER_PARAMETER_ARRAY(FVector, Position, [RAY_TRACING_LIGHT_COUNT_MAXIMUM])
+	SHADER_PARAMETER_ARRAY(FVector, Normal, [RAY_TRACING_LIGHT_COUNT_MAXIMUM])
+	SHADER_PARAMETER_ARRAY(FVector, dPdu, [RAY_TRACING_LIGHT_COUNT_MAXIMUM])
+	SHADER_PARAMETER_ARRAY(FVector, dPdv, [RAY_TRACING_LIGHT_COUNT_MAXIMUM])
 	// Color
-	SHADER_PARAMETER_ARRAY(FVector, Color, [GLightCountMaximum])
+	SHADER_PARAMETER_ARRAY(FVector, Color, [RAY_TRACING_LIGHT_COUNT_MAXIMUM])
 	// Light-specific
-	SHADER_PARAMETER_ARRAY(FVector, Dimensions, [GLightCountMaximum])
-	SHADER_PARAMETER_ARRAY(float, Attenuation, [GLightCountMaximum])
-	SHADER_PARAMETER_ARRAY(float, RectLightBarnCosAngle, [GLightCountMaximum])
-	SHADER_PARAMETER_ARRAY(float, RectLightBarnLength, [GLightCountMaximum])
+	SHADER_PARAMETER_ARRAY(FVector, Dimensions, [RAY_TRACING_LIGHT_COUNT_MAXIMUM])
+	SHADER_PARAMETER_ARRAY(float, Attenuation, [RAY_TRACING_LIGHT_COUNT_MAXIMUM])
+	SHADER_PARAMETER_ARRAY(float, RectLightBarnCosAngle, [RAY_TRACING_LIGHT_COUNT_MAXIMUM])
+	SHADER_PARAMETER_ARRAY(float, RectLightBarnLength, [RAY_TRACING_LIGHT_COUNT_MAXIMUM])
 END_GLOBAL_SHADER_PARAMETER_STRUCT()
 
 
diff --git a/Engine/Source/Runtime/Renderer/Private/RayTracing/RayTracingLighting.cpp b/Engine/Source/Runtime/Renderer/Private/RayTracing/RayTracingLighting.cpp
--- a/Engine/Source/Runtime/Renderer/Private/RayTracing/RayTracingLighting.cpp
+++ b/Engine/Source/Runtime/Renderer/Private/RayTracing/RayTracingLighting.cpp
@@ -148,7 +148,7 @@ void SetupRaytracingLightDataPacked(
 
 		LightData->Count++;
 
-		if (LightData->Count >= GRaytracingLightCountMaximum) break;
+		if (LightData->Count >= RAY_TRACING_LIGHT_COUNT_MAXIMUM) break;
 	}
 
 	// Update IES light profiles texture 
diff --git a/Engine/Source/Runtime/Renderer/Private/RayTracing/RayTracingLighting.h b/Engine/Source/Runtime/Renderer/Private/RayTracing/RayTracingLighting.h
--- a/Engine/Source/Runtime/Renderer/Private/RayTracing/RayTracingLighting.h
+++ b/Engine/Source/Runtime/Renderer/Private/RayTracing/RayTracingLighting.h
@@ -7,21 +7,20 @@
 #include "ShaderParameterMacros.h"
 #include "SceneRendering.h"
 #include "LightSceneInfo.h"
+#include "RayTracingDefinitions.h"
 
 #if RHI_RAYTRACING
 
-const static uint32 GRaytracingLightCountMaximum = 64;
-
 BEGIN_GLOBAL_SHADER_PARAMETER_STRUCT(FRaytracingLightDataPacked, )
 	SHADER_PARAMETER(uint32, Count)
 	SHADER_PARAMETER(float, IESLightProfileInvCount)
-	SHADER_PARAMETER_ARRAY(FIntVector, Type_LightProfileIndex_RectLightTextureIndex, [GRaytracingLightCountMaximum])
-	SHADER_PARAMETER_ARRAY(FVector4, LightPosition_InvRadius, [GRaytracingLightCountMaximum])
-	SHADER_PARAMETER_ARRAY(FVector4, LightColor_SpecularScale, [GRaytracingLightCountMaximum])
-	SHADER_PARAMETER_ARRAY(FVector4, Direction_FalloffExponent, [GRaytracingLightCountMaximum])
-	SHADER_PARAMETER_ARRAY(FVector4, Tangent_SourceRadius, [GRaytracingLightCountMaximum])
-	SHADER_PARAMETER_ARRAY(FVector4, SpotAngles_SourceLength_SoftSourceRadius, [GRaytracingLightCountMaximum])
-	SHADER_PARAMETER_ARRAY(FVector4, DistanceFadeMAD_RectLightBarnCosAngle_RectLightBarnLength, [GRaytracingLightCountMaximum])
+	SHADER_PARAMETER_ARRAY(FIntVector, Type_LightProfileIndex_RectLightTextureIndex, [RAY_TRACING_LIGHT_COUNT_MAXIMUM])
+	SHADER_PARAMETER_ARRAY(FVector4, LightPosition_InvRadius, [RAY_TRACING_LIGHT_COUNT_MAXIMUM])
+	SHADER_PARAMETER_ARRAY(FVector4, LightColor_SpecularScale, [RAY_TRACING_LIGHT_COUNT_MAXIMUM])
+	SHADER_PARAMETER_ARRAY(FVector4, Direction_FalloffExponent, [RAY_TRACING_LIGHT_COUNT_MAXIMUM])
+	SHADER_PARAMETER_ARRAY(FVector4, Tangent_SourceRadius, [RAY_TRACING_LIGHT_COUNT_MAXIMUM])
+	SHADER_PARAMETER_ARRAY(FVector4, SpotAngles_SourceLength_SoftSourceRadius, [RAY_TRACING_LIGHT_COUNT_MAXIMUM])
+	SHADER_PARAMETER_ARRAY(FVector4, DistanceFadeMAD_RectLightBarnCosAngle_RectLightBarnLength, [RAY_TRACING_LIGHT_COUNT_MAXIMUM])
 	SHADER_PARAMETER_TEXTURE(Texture2D, LTCMatTexture)
 	SHADER_PARAMETER_SAMPLER(SamplerState, LTCMatSampler)
 	SHADER_PARAMETER_TEXTURE(Texture2D, LTCAmpTexture)
