Deletes TAA pixel shader for DOF, since DOF is mostly compute anyway.

diff --git a/Engine/Shaders/Public/ShaderVersion.ush b/Engine/Shaders/Public/ShaderVersion.ush
--- a/Engine/Shaders/Public/ShaderVersion.ush
+++ b/Engine/Shaders/Public/ShaderVersion.ush
@@ -2,6 +2,4 @@
 // Each time the console command is executed it generates a new GUID. As this file is included
 // in Platform.ush (which should be included in any shader) it allows to invalidate the shader DDC.
 // 
-// If you are merging streams and there is a conflict with this GUID you should make a new GUID rather than taking one or the other.
-// GUID = 8CDE7B2F-4CE1-4074-AE2B-8BCE5F34437F
-
+// GUID = 09051FCE47FE175922843B8618F04059
\ No newline at end of file
diff --git a/Engine/Source/Runtime/Renderer/Private/PostProcess/PostProcessTemporalAA.cpp b/Engine/Source/Runtime/Renderer/Private/PostProcess/PostProcessTemporalAA.cpp
--- a/Engine/Source/Runtime/Renderer/Private/PostProcess/PostProcessTemporalAA.cpp
+++ b/Engine/Source/Runtime/Renderer/Private/PostProcess/PostProcessTemporalAA.cpp
@@ -161,6 +161,12 @@ class FTemporalAAPS : public FGlobalShader
 		{
 			return false;
 		}
+		
+		// DOF setup chain is full compute shader, no need for a pixel shader for TAA.
+		if (IsDOFTAAConfig(PermutationVector.Get<FTAAPassConfigDim>()))
+		{
+			return false;
+		}
 
 		// Responsive dimension is only for Main.
 		if (PermutationVector.Get<FTAAResponsiveDim>() && !SupportsResponsiveDim(PermutationVector))
