Changes diaphragm DOF's draw events to just "DOF", because there is only code path anyway for deferred shading renderer now.

diff --git a/Engine/Source/Runtime/Renderer/Private/PostProcess/DiaphragmDOF.cpp b/Engine/Source/Runtime/Renderer/Private/PostProcess/DiaphragmDOF.cpp
--- a/Engine/Source/Runtime/Renderer/Private/PostProcess/DiaphragmDOF.cpp
+++ b/Engine/Source/Runtime/Renderer/Private/PostProcess/DiaphragmDOF.cpp
@@ -1395,7 +1395,7 @@ FRDGTextureRef DiaphragmDOF::AddPasses(
 		return InputSceneColor;
 	}
 
-	RDG_EVENT_SCOPE(GraphBuilder, "DiaphragmDOF");
+	RDG_EVENT_SCOPE(GraphBuilder, "DOF(Alpha=%s)", bProcessSceneAlpha ? TEXT("Yes") : TEXT("No"));
 
 	bool bGatherForeground = AbsMaxForegroundCocRadius > kMinimalAbsGatherPassCocRadius;
 	
@@ -1513,7 +1513,7 @@ FRDGTextureRef DiaphragmDOF::AddPasses(
 		TShaderMapRef<FDiaphragmDOFSetupCS> ComputeShader(View.ShaderMap, PermutationVector);
 		FComputeShaderUtils::AddPass(
 			GraphBuilder,
-			RDG_EVENT_NAME("DiaphragmDOF Setup(%s CoC=[%d;%d]) %dx%d",
+			RDG_EVENT_NAME("DOF Setup(%s CoC=[%d;%d]) %dx%d",
 				!bOutputFullResolution ? TEXT("HalfRes") : (!bOutputHalfResolution ? TEXT("FullRes") : TEXT("Full&HalfRes")),
 				FMath::FloorToInt(CocModel.ComputeViewMinForegroundCocRadius(PassViewSize.X)),
 				FMath::CeilToInt(CocModel.ComputeViewMaxBackgroundCocRadius(PassViewSize.X)),
@@ -1595,7 +1595,7 @@ FRDGTextureRef DiaphragmDOF::AddPasses(
 			TShaderMapRef<FDiaphragmDOFCocFlattenCS> ComputeShader(View.ShaderMap, PermutationVector);
 			FComputeShaderUtils::AddPass(
 				GraphBuilder,
-				RDG_EVENT_NAME("DiaphragmDOF FlattenCoc(Gather4=%s) %dx%d",
+				RDG_EVENT_NAME("DOF FlattenCoc(Gather4=%s) %dx%d",
 					PermutationVector.Get<FDiaphragmDOFCocFlattenCS::FDoCocGather4>() ? TEXT("Yes") : TEXT("No"),
 					GatheringViewSize.X, GatheringViewSize.Y),
 				*ComputeShader,
@@ -1651,7 +1651,7 @@ FRDGTextureRef DiaphragmDOF::AddPasses(
 			TShaderMapRef<FDiaphragmDOFCocDilateCS> ComputeShader(View.ShaderMap, PermutationVector);
 			FComputeShaderUtils::AddPass(
 				GraphBuilder,
-				RDG_EVENT_NAME("DiaphragmDOF DilateCoc(1/16 %s radius=%d step=%d) %dx%d",
+				RDG_EVENT_NAME("DOF DilateCoc(1/16 %s radius=%d step=%d) %dx%d",
 					GetEventName(Mode), SampleRadiusCount, SampleOffsetMultipler,
 					DilatePassViewSize.X, DilatePassViewSize.Y),
 				*ComputeShader,
@@ -1818,7 +1818,7 @@ FRDGTextureRef DiaphragmDOF::AddPasses(
 			TShaderMapRef<FDiaphragmDOFDownsampleCS> ComputeShader(View.ShaderMap);
 			FComputeShaderUtils::AddPass(
 				GraphBuilder,
-				RDG_EVENT_NAME("DiaphragmDOF Downsample %dx%d", PassViewSize.X, PassViewSize.Y),
+				RDG_EVENT_NAME("DOF Downsample %dx%d", PassViewSize.X, PassViewSize.Y),
 				*ComputeShader,
 				PassParameters,
 				FComputeShaderUtils::GetGroupCount(PassViewSize, kDefaultGroupSize));
@@ -1884,7 +1884,7 @@ FRDGTextureRef DiaphragmDOF::AddPasses(
 			TShaderMapRef<FDiaphragmDOFReduceCS> ComputeShader(View.ShaderMap, PermutationVector);
 			FComputeShaderUtils::AddPass(
 				GraphBuilder,
-				RDG_EVENT_NAME("DiaphragmDOF Reduce(Mips=%d FgdScatter=%s BgdScatter=%s%s) %dx%d",
+				RDG_EVENT_NAME("DOF Reduce(Mips=%d FgdScatter=%s BgdScatter=%s%s) %dx%d",
 					MipLevelCount,
 					bForegroundHybridScattering ? TEXT("Yes") : TEXT("No"),
 					bBackgroundHybridScattering ? TEXT("Yes") : TEXT("No"),
@@ -1909,7 +1909,7 @@ FRDGTextureRef DiaphragmDOF::AddPasses(
 			TShaderMapRef<FDiaphragmDOFScatterGroupPackCS> ComputeShader(View.ShaderMap);
 			FComputeShaderUtils::AddPass(
 				GraphBuilder,
-				RDG_EVENT_NAME("DiaphragmDOF ScatterGroupPack"),
+				RDG_EVENT_NAME("DOF ScatterGroupPack"),
 				*ComputeShader,
 				PassParameters,
 				FIntVector(2, 1, 1));
@@ -1954,7 +1954,7 @@ FRDGTextureRef DiaphragmDOF::AddPasses(
 		TShaderMapRef<FDiaphragmDOFBuildBokehLUTCS> ComputeShader(View.ShaderMap, PermutationVector);
 		FComputeShaderUtils::AddPass(
 			GraphBuilder,
-			RDG_EVENT_NAME("DiaphragmDOF BuildBokehLUT(Blades=%d Shape=%s, LUT=%s) %dx%d",
+			RDG_EVENT_NAME("DOF BuildBokehLUT(Blades=%d Shape=%s, LUT=%s) %dx%d",
 				BokehModel.DiaphragmBladeCount,
 				BokehModel.BokehShape == DiaphragmDOF::EBokehShape::RoundedBlades ? TEXT("Rounded") : TEXT("Straight"),
 				GetEventName(LUTFormat),
@@ -2137,7 +2137,7 @@ FRDGTextureRef DiaphragmDOF::AddPasses(
 			TShaderMapRef<FDiaphragmDOFGatherCS> ComputeShader(View.ShaderMap, PermutationVector);
 			FComputeShaderUtils::AddPass(
 				GraphBuilder,
-				RDG_EVENT_NAME("DiaphragmDOF Gather(%s %s Bokeh=%s Rings=%d%s%s) %dx%d",
+				RDG_EVENT_NAME("DOF Gather(%s %s Bokeh=%s Rings=%d%s%s) %dx%d",
 					GetEventName(ConvolutionSettings.LayerProcessing),
 					GetEventName(ConvolutionSettings.QualityConfig),
 					GetEventName(ConvolutionSettings.BokehSimulation),
@@ -2184,7 +2184,7 @@ FRDGTextureRef DiaphragmDOF::AddPasses(
 			TShaderMapRef<FDiaphragmDOFPostfilterCS> ComputeShader(View.ShaderMap, PermutationVector);
 			FComputeShaderUtils::AddPass(
 				GraphBuilder,
-				RDG_EVENT_NAME("DiaphragmDOF Postfilter(%s %s%s) %dx%d",
+				RDG_EVENT_NAME("DOF Postfilter(%s %s%s) %dx%d",
 					GetEventName(ConvolutionSettings.LayerProcessing), GetEventName(ConvolutionSettings.PostfilterMethod),
 					PermutationVector.Get<FDiaphragmDOFPostfilterCS::FTileOptimization>() ? TEXT(" TileOptimisation") : TEXT(""),
 					GatheringViewSize.X, GatheringViewSize.Y),
@@ -2245,7 +2245,7 @@ FRDGTextureRef DiaphragmDOF::AddPasses(
 			ValidateShaderParameters(*PixelShader, *PassParameters);
 
 			GraphBuilder.AddPass(
-				RDG_EVENT_NAME("DiaphragmDOF IndirectScatter(%s Bokeh=%s Occlusion=%s) %dx%d",
+				RDG_EVENT_NAME("DOF IndirectScatter(%s Bokeh=%s Occlusion=%s) %dx%d",
 					GetEventName(bIsForeground ? EDiaphragmDOFLayerProcessing::ForegroundOnly : EDiaphragmDOFLayerProcessing::BackgroundOnly),
 					PermutationVector.Get<FDiaphragmDOFHybridScatterPS::FBokehSimulationDim>() ? TEXT("Generic") : TEXT("None"),
 					PermutationVector.Get<FDDOFScatterOcclusionDim>() ? TEXT("Yes") : TEXT("No"),
@@ -2419,7 +2419,7 @@ FRDGTextureRef DiaphragmDOF::AddPasses(
 		TShaderMapRef<FDiaphragmDOFRecombineCS> ComputeShader(View.ShaderMap, PermutationVector);
 		FComputeShaderUtils::AddPass(
 			GraphBuilder,
-			RDG_EVENT_NAME("DiaphragmDOF Recombine(%s Quality=%d Bokeh=%s alpha=no) %dx%d",
+			RDG_EVENT_NAME("DOF Recombine(%s Quality=%d Bokeh=%s) %dx%d",
 				GetEventName(PermutationVector.Get<FDDOFLayerProcessingDim>()),
 				RecombineQuality,
 				GetEventName(PermutationVector.Get<FDDOFBokehSimulationDim>()),
diff --git a/Engine/Source/Runtime/Renderer/Private/PostProcess/PostProcessTemporalAA.cpp b/Engine/Source/Runtime/Renderer/Private/PostProcess/PostProcessTemporalAA.cpp
--- a/Engine/Source/Runtime/Renderer/Private/PostProcess/PostProcessTemporalAA.cpp
+++ b/Engine/Source/Runtime/Renderer/Private/PostProcess/PostProcessTemporalAA.cpp
@@ -390,8 +390,8 @@ const TCHAR* const kTAAPassNames[] = {
 	TEXT("MainSuperSampling"),
 	TEXT("ScreenSpaceReflections"),
 	TEXT("LightShaft"),
-	TEXT("DiaphragmDOF"),
-	TEXT("DiaphragmDOFUpsampling"),
+	TEXT("DOF"),
+	TEXT("DOFUpsampling"),
 };
 
 
