Removes TAA shader permutation for camera cuts.

diff --git a/Engine/Shaders/Private/PostProcessTemporalAA.usf b/Engine/Shaders/Private/PostProcessTemporalAA.usf
--- a/Engine/Shaders/Private/PostProcessTemporalAA.usf
+++ b/Engine/Shaders/Private/PostProcessTemporalAA.usf
@@ -444,9 +444,11 @@ static const uint kNeighborsCount = 9;
 // (Pre-Exposure, 1 / Pre-exposure, LastFramePreExposure, 1 / LastFramePreExposure)
 float4 PreExposureSettings;
 
+float CurrentFrameWeight;
+int bCameraCut;
+
 float SampleWeights[9];
 float PlusWeights[5];
-float CurrentFrameWeight;
 
 float4 ViewportUVToInputBufferUV;
 float4 MaxViewportUVAndSvPositionToViewportUV;
@@ -2058,13 +2060,7 @@ FTAAHistoryPayload TemporalAASample(uint2 GroupId, uint2 GroupThreadId, uint Gro
 	FTAAHistoryPayload History = SampleHistory(HistoryScreenPosition);
 
 	// Whether the feedback needs to be reset.
-	bool IgnoreHistory = OffScreen;
-	
-	// When we have a camera but, we sets the feedback reset to have same scene color filtering.
-	if (TAA_CAMERA_CUT)
-	{
-		IgnoreHistory = true;
-	}
+	bool IgnoreHistory = OffScreen || bCameraCut;
 
 	// DYNAMIC ANTI GHOSTING
 	// ---------------------
@@ -2145,11 +2141,6 @@ FTAAHistoryPayload TemporalAASample(uint2 GroupId, uint2 GroupThreadId, uint Gro
 	// COMPUTE BLEND AMOUNT 
 	// --------------------
 	float BlendFinal;
-	if (TAA_CAMERA_CUT)
-	{
-		BlendFinal = 1.0;
-	}
-	else
 	{
 		float LumaFiltered = GetSceneColorLuma4(IntermediaryResult.Filtered.Color);
 
@@ -2189,6 +2180,11 @@ FTAAHistoryPayload TemporalAASample(uint2 GroupId, uint2 GroupThreadId, uint Gro
 			BlendFinal = lerp(1, BlendFinal, BilateralWeight);
 		}
 		#endif
+		
+		if (bCameraCut)
+		{
+			BlendFinal = 1.0;
+		}
 	}
 
 	// Offscreen feedback resets.
diff --git a/Engine/Source/Runtime/Renderer/Private/PostProcess/PostProcessTemporalAA.cpp b/Engine/Source/Runtime/Renderer/Private/PostProcess/PostProcessTemporalAA.cpp
--- a/Engine/Source/Runtime/Renderer/Private/PostProcess/PostProcessTemporalAA.cpp
+++ b/Engine/Source/Runtime/Renderer/Private/PostProcess/PostProcessTemporalAA.cpp
@@ -74,6 +74,7 @@ BEGIN_SHADER_PARAMETER_STRUCT(FTAAShaderParameters,)
 	SHADER_PARAMETER(FVector4, MaxViewportUVAndSvPositionToViewportUV)
 	SHADER_PARAMETER(FVector2D, ScreenPosAbsMax)
 	SHADER_PARAMETER(float, CurrentFrameWeight)
+	SHADER_PARAMETER(int32, bCameraCut)
 	
 	SHADER_PARAMETER_ARRAY(float, SampleWeights, [9])
 	SHADER_PARAMETER_ARRAY(float, PlusWeights, [5])
@@ -118,7 +119,6 @@ namespace
 class FTAAPassConfigDim : SHADER_PERMUTATION_ENUM_CLASS("TAA_PASS_CONFIG", ETAAPassConfig);
 class FTAAFastDim : SHADER_PERMUTATION_BOOL("TAA_FAST");
 class FTAAResponsiveDim : SHADER_PERMUTATION_BOOL("TAA_RESPONSIVE");
-class FTAACameraCutDim : SHADER_PERMUTATION_BOOL("TAA_CAMERA_CUT");
 class FTAAScreenPercentageDim : SHADER_PERMUTATION_INT("TAA_SCREEN_PERCENTAGE_RANGE", 4);
 class FTAAUpsampleFilteredDim : SHADER_PERMUTATION_BOOL("TAA_UPSAMPLE_FILTERED");
 class FTAADownsampleDim : SHADER_PERMUTATION_BOOL("TAA_DOWNSAMPLE");
@@ -136,8 +136,7 @@ class FTemporalAAPS : public FGlobalShader
 	using FPermutationDomain = TShaderPermutationDomain<
 		FTAAPassConfigDim,
 		FTAAFastDim,
-		FTAAResponsiveDim,
-		FTAACameraCutDim>;
+		FTAAResponsiveDim>;
 	
 	BEGIN_SHADER_PARAMETER_STRUCT(FParameters, )
 		SHADER_PARAMETER_STRUCT_INCLUDE(FTAAShaderParameters, CommonParameters)
@@ -191,7 +190,6 @@ class FTemporalAACS : public FGlobalShader
 	using FPermutationDomain = TShaderPermutationDomain<
 		FTAAPassConfigDim,
 		FTAAFastDim,
-		FTAACameraCutDim,
 		FTAAScreenPercentageDim,
 		FTAAUpsampleFilteredDim,
 		FTAADownsampleDim>;
@@ -475,6 +473,7 @@ FTAAOutputs FTAAPassParameters::AddTemporalAAPass(
 
 		CommonShaderParameters.ViewUniformBuffer = View.ViewUniformBuffer;
 		CommonShaderParameters.CurrentFrameWeight = CVarTemporalAACurrentFrameWeight.GetValueOnRenderThread();
+		CommonShaderParameters.bCameraCut = bCameraCut;
 
 		CommonShaderParameters.SceneBlackboard = SceneBlackboard;
 		CommonShaderParameters.SceneDepthBufferSampler = TStaticSamplerState<SF_Point>::GetRHI();
@@ -529,7 +528,7 @@ FTAAOutputs FTAAPassParameters::AddTemporalAAPass(
 
 			CommonShaderParameters.HistoryBufferSize = FVector4(BufferSize.X, BufferSize.Y, InvBufferSizeX, InvBufferSizeY);
 
-			if (InputHistory.IsValid())
+			if (!bCameraCut)
 			{
 				CommonShaderParameters.HistoryBuffer0 = GraphBuilder.RegisterExternalTexture(InputHistory.RT[0]);
 				if (InputHistory.RT[1].IsValid())
@@ -620,7 +619,6 @@ FTAAOutputs FTAAPassParameters::AddTemporalAAPass(
 		FTemporalAACS::FPermutationDomain PermutationVector;
 		PermutationVector.Set<FTAAPassConfigDim>(Pass);
 		PermutationVector.Set<FTAAFastDim>(bUseFast);
-		PermutationVector.Set<FTAACameraCutDim>(!View.PrevViewInfo.TemporalAAHistory.IsValid());
 		PermutationVector.Set<FTAADownsampleDim>(bDownsample);
 		PermutationVector.Set<FTAAUpsampleFilteredDim>(true);
 
@@ -729,7 +727,6 @@ FTAAOutputs FTAAPassParameters::AddTemporalAAPass(
 		FTemporalAAPS::FPermutationDomain BasePermutationVector;
 		BasePermutationVector.Set<FTAAPassConfigDim>(Pass);
 		BasePermutationVector.Set<FTAAFastDim>(bUseFast);
-		BasePermutationVector.Set<FTAACameraCutDim>(bCameraCut);
 
 		TShaderMapRef<FTemporalAAPS> PixelShader(View.ShaderMap, BasePermutationVector);
 		ClearUnusedGraphResources(*PixelShader, PassParameters);
