Fixes a crash on r.TemporalAA.HistoryScreenPercentage

diff --git a/Engine/Source/Runtime/Renderer/Private/PostProcess/PostProcessing.cpp b/Engine/Source/Runtime/Renderer/Private/PostProcess/PostProcessing.cpp
--- a/Engine/Source/Runtime/Renderer/Private/PostProcess/PostProcessing.cpp
+++ b/Engine/Source/Runtime/Renderer/Private/PostProcess/PostProcessing.cpp
@@ -1473,6 +1473,7 @@ void FPostProcessing::Process(FRHICommandListImmediate& RHICmdList, const FViewI
 					TAAParameters.Pass = ETAAPassConfig::MainSuperSampling;
 					TAAParameters.bDownsample = false;
 					TAAParameters.bUseFast = false;
+					TAAParameters.bIsComputePass = true;
 
 					TAAParameters.OutputViewRect.Min.X = 0;
 					TAAParameters.OutputViewRect.Min.Y = 0;
