Removes no long used temporal AA history entries.

diff --git a/Engine/Source/Runtime/Renderer/Private/PostProcess/DiaphragmDOF.cpp b/Engine/Source/Runtime/Renderer/Private/PostProcess/DiaphragmDOF.cpp
--- a/Engine/Source/Runtime/Renderer/Private/PostProcess/DiaphragmDOF.cpp
+++ b/Engine/Source/Runtime/Renderer/Private/PostProcess/DiaphragmDOF.cpp
@@ -1536,8 +1536,8 @@ FRDGTextureRef DiaphragmDOF::AddPasses(
 
 		FTAAOutputs TAAOutputs = TAAParameters.AddTemporalAAPass(
 			GraphBuilder, SceneBlackboard, View,
-			View.PrevViewInfo.DOFPreGatherHistory,
-			&ViewState->PrevFrameViewInfo.DOFPreGatherHistory);
+			View.PrevViewInfo.DOFSetupHistory,
+			&ViewState->PrevFrameViewInfo.DOFSetupHistory);
 		
 		HalfResGatherInputTextures.SceneColor = TAAOutputs.SceneColor;
 		HalfResGatherInputTextures.SeparateCoc = TAAOutputs.SceneMetadata;
diff --git a/Engine/Source/Runtime/Renderer/Private/ScenePrivate.h b/Engine/Source/Runtime/Renderer/Private/ScenePrivate.h
--- a/Engine/Source/Runtime/Renderer/Private/ScenePrivate.h
+++ b/Engine/Source/Runtime/Renderer/Private/ScenePrivate.h
@@ -853,10 +853,6 @@ public:
 
 	FHeightfieldLightingAtlas* HeightfieldLightingAtlas;
 
-	// TODO: move these guys in FPreviousViewInfo.
-	// Temporal AA result for DOF of last frame
-	FTemporalAAHistory DOFHistory;
-	FTemporalAAHistory DOFHistory2;
 	// Temporal AA result for SSR
 	FTemporalAAHistory SSRHistory;
 	// Temporal AA result for light shafts of last frame
@@ -1261,8 +1257,6 @@ public:
 		EyeAdaptationRTManager.SafeRelease();
 		CombinedLUTRenderTarget.SafeRelease();
 		PrevFrameViewInfo.SafeRelease();
-		DOFHistory.SafeRelease();
-		DOFHistory2.SafeRelease();
 		SSRHistory.SafeRelease();
 		LightShaftOcclusionHistory.SafeRelease();
 		LightShaftBloomHistoryRTs.Empty();
diff --git a/Engine/Source/Runtime/Renderer/Private/SceneRendering.h b/Engine/Source/Runtime/Renderer/Private/SceneRendering.h
--- a/Engine/Source/Runtime/Renderer/Private/SceneRendering.h
+++ b/Engine/Source/Runtime/Renderer/Private/SceneRendering.h
@@ -653,9 +653,7 @@ struct FPreviousViewInfo
 	FTemporalAAHistory TemporalAAHistory;
 
 	// Temporal AA history for diaphragm DOF.
-	FTemporalAAHistory DOFPreGatherHistory;
-	FTemporalAAHistory DOFPostGatherForegroundHistory;
-	FTemporalAAHistory DOFPostGatherBackgroundHistory;
+	FTemporalAAHistory DOFSetupHistory;
 
 	// Scene color input for SSR, that can be different from TemporalAAHistory.RT[0] if there is a SSR
 	// input post process material.
@@ -684,9 +682,7 @@ struct FPreviousViewInfo
 		GBufferB.SafeRelease();
 		GBufferC.SafeRelease();
 		TemporalAAHistory.SafeRelease();
-		DOFPreGatherHistory.SafeRelease();
-		DOFPostGatherForegroundHistory.SafeRelease();
-		DOFPostGatherBackgroundHistory.SafeRelease();
+		DOFSetupHistory.SafeRelease();
 		CustomSSRInput.SafeRelease();
 		ReflectionsHistory.SafeRelease();
 		AmbientOcclusionHistory.SafeRelease();
