DR - Fix compilation on non-editor
#rnx

diff --git a/Engine/Source/Runtime/Renderer/Private/PostProcess/PostProcessBokehDOF.cpp b/Engine/Source/Runtime/Renderer/Private/PostProcess/PostProcessBokehDOF.cpp
--- a/Engine/Source/Runtime/Renderer/Private/PostProcess/PostProcessBokehDOF.cpp
+++ b/Engine/Source/Runtime/Renderer/Private/PostProcess/PostProcessBokehDOF.cpp
@@ -99,12 +99,14 @@ public:
 			// a negative values disables the cross hair feature
 			FIntPoint CursorPosValue(-100,-100);
 
+#if WITH_EDITORONLY_DATA
 			PRAGMA_DISABLE_DEPRECATION_WARNINGS
 			if(Context.View.FinalPostProcessSettings.DepthOfFieldMethod_DEPRECATED == DOFM_CircleDOF)
 			{
 				CursorPosValue = Context.View.CursorPos;
 			}
 			PRAGMA_ENABLE_DEPRECATION_WARNINGS
+#endif
 
 			SetShaderValue(RHICmdList, ShaderRHI, CursorPos, CursorPosValue);
 		}
