Improvements in ray tracing debug view modes.
- Now use DebugViewModeRendering to avoid polluting the regular postprocessing pipeline.
- Opacity view mode inverted to match Raster buffer visualization mode
- Gamma applied only in the modes that need it

diff --git a/Engine/Shaders/Private/RayTracing/RayTracingDebug.usf b/Engine/Shaders/Private/RayTracing/RayTracingDebug.usf
--- a/Engine/Shaders/Private/RayTracing/RayTracingDebug.usf
+++ b/Engine/Shaders/Private/RayTracing/RayTracingDebug.usf
@@ -65,7 +65,7 @@ void RayTracingDebugMainRGS()
 			Result = float4(Payload.SpecularColor.rgb, 1.0f);
 			break;
 		case RAY_TRACING_DEBUG_VIZ_OPACITY:
-			Result = float4(Payload.Opacity, Payload.Opacity, Payload.Opacity, 1.0f);
+			Result = float4(1.0f - Payload.Opacity, 1.0f - Payload.Opacity, 1.0f - Payload.Opacity, 1.0f);
 			break;
 		case RAY_TRACING_DEBUG_VIZ_METALLIC:
 			Result = float4(Payload.Metallic, Payload.Metallic, Payload.Metallic, 1.0f);
diff --git a/Engine/Source/Editor/UnrealEd/Private/EditorBuildUtils.cpp b/Engine/Source/Editor/UnrealEd/Private/EditorBuildUtils.cpp
--- a/Engine/Source/Editor/UnrealEd/Private/EditorBuildUtils.cpp
+++ b/Engine/Source/Editor/UnrealEd/Private/EditorBuildUtils.cpp
@@ -1080,6 +1080,8 @@ EDebugViewShaderMode ViewModeIndexToDebugViewShaderMode(EViewModeIndex SelectedV
 		return DVSM_MaterialTextureScaleAccuracy;
 	case VMI_RequiredTextureResolution:
 		return DVSM_RequiredTextureResolution;
+	case VMI_RayTracingDebug:
+		return DVSM_RayTracingDebug;
 	case VMI_Unknown:
 	default :
 		return DVSM_None;
diff --git a/Engine/Source/Editor/UnrealEd/Private/EditorViewportClient.cpp b/Engine/Source/Editor/UnrealEd/Private/EditorViewportClient.cpp
--- a/Engine/Source/Editor/UnrealEd/Private/EditorViewportClient.cpp
+++ b/Engine/Source/Editor/UnrealEd/Private/EditorViewportClient.cpp
@@ -52,6 +52,7 @@
 #include "Editor/EditorPerformanceSettings.h"
 #include "ImageWriteQueue.h"
 #include "DebugViewModeHelpers.h"
+#include "RayTracingDebugVisualizationMenuCommands.h"
 #include "Misc/ScopedSlowTask.h"
 #include "UnrealEngine.h"
 
@@ -3662,9 +3663,11 @@ void FEditorViewportClient::Draw(FViewport* InViewport, FCanvas* Canvas)
 
 	EViewModeIndex CurrentViewMode = GetViewMode();
 	ViewFamily.ViewMode = CurrentViewMode;
-	bool bCanDisableTonemapper = (CurrentViewMode == VMI_VisualizeBuffer && CurrentBufferVisualizationMode != NAME_None) 
-								|| (CurrentViewMode == VMI_RayTracingDebug && CurrentRayTracingDebugVisualizationMode != NAME_None);
 
+	const bool bVisualizeBufferEnabled = CurrentViewMode == VMI_VisualizeBuffer && CurrentBufferVisualizationMode != NAME_None;
+	const bool bRayTracingDebugEnabled = CurrentViewMode == VMI_RayTracingDebug && CurrentRayTracingDebugVisualizationMode != NAME_None;
+	const bool bCanDisableTonemapper = bVisualizeBufferEnabled || (bRayTracingDebugEnabled && !FRayTracingDebugVisualizationMenuCommands::DebugModeShouldBeTonemapped(CurrentRayTracingDebugVisualizationMode));
+	
 	EngineShowFlagOverride(ESFIM_Editor, ViewFamily.ViewMode, ViewFamily.EngineShowFlags, bCanDisableTonemapper);
 	EngineShowFlagOrthographicOverride(IsPerspective(), ViewFamily.EngineShowFlags);
 
diff --git a/Engine/Source/Editor/UnrealEd/Private/RayTracingDebugVisualizationMenuCommands.cpp b/Engine/Source/Editor/UnrealEd/Private/RayTracingDebugVisualizationMenuCommands.cpp
--- a/Engine/Source/Editor/UnrealEd/Private/RayTracingDebugVisualizationMenuCommands.cpp
+++ b/Engine/Source/Editor/UnrealEd/Private/RayTracingDebugVisualizationMenuCommands.cpp
@@ -67,6 +67,21 @@ void FRayTracingDebugVisualizationMenuCommands::CreateRayTracingDebugVisualizati
 	}
 }
 
+bool FRayTracingDebugVisualizationMenuCommands::DebugModeShouldBeTonemapped(const FName& RayTracingDebugModeName)
+{
+	static TArray<FName> TonemappedRayTracingDebugVisualizationModes;
+	if (TonemappedRayTracingDebugVisualizationModes.Num() == 0)
+	{
+		TonemappedRayTracingDebugVisualizationModes.Add(*LOCTEXT("World Normal", "World Normal").ToString());
+		TonemappedRayTracingDebugVisualizationModes.Add(*LOCTEXT("BaseColor", "BaseColor").ToString());
+		TonemappedRayTracingDebugVisualizationModes.Add(*LOCTEXT("DiffuseColor", "DiffuseColor").ToString());
+		TonemappedRayTracingDebugVisualizationModes.Add(*LOCTEXT("SpecularColor", "SpecularColor").ToString());
+		TonemappedRayTracingDebugVisualizationModes.Add(*LOCTEXT("ShadingModelID", "ShadingModelID").ToString());
+	}
+
+	return TonemappedRayTracingDebugVisualizationModes.Contains(RayTracingDebugModeName);
+}
+
 void FRayTracingDebugVisualizationMenuCommands::BuildVisualisationSubMenu(FMenuBuilder& Menu)
 {
 	const FRayTracingDebugVisualizationMenuCommands& Commands = FRayTracingDebugVisualizationMenuCommands::Get();
diff --git a/Engine/Source/Editor/UnrealEd/Public/RayTracingDebugVisualizationMenuCommands.h b/Engine/Source/Editor/UnrealEd/Public/RayTracingDebugVisualizationMenuCommands.h
--- a/Engine/Source/Editor/UnrealEd/Public/RayTracingDebugVisualizationMenuCommands.h
+++ b/Engine/Source/Editor/UnrealEd/Public/RayTracingDebugVisualizationMenuCommands.h
@@ -36,6 +36,8 @@ public:
 
 	void BindCommands(FUICommandList& CommandList, const TSharedPtr<FEditorViewportClient>& Client) const;
 
+	static bool DebugModeShouldBeTonemapped(const FName& RayTracingDebugModeName);
+
 private:
 	void BuildCommandMap();
 
diff --git a/Engine/Source/Runtime/Engine/Private/DebugViewModeHelpers.cpp b/Engine/Source/Runtime/Engine/Private/DebugViewModeHelpers.cpp
--- a/Engine/Source/Runtime/Engine/Private/DebugViewModeHelpers.cpp
+++ b/Engine/Source/Runtime/Engine/Private/DebugViewModeHelpers.cpp
@@ -56,6 +56,8 @@ bool AllowDebugViewShaderMode(EDebugViewShaderMode ShaderMode, EShaderPlatform P
 	case DVSM_RequiredTextureResolution:
 	case DVSM_OutputMaterialTextureScales:
 		return FeatureLevel >= ERHIFeatureLevel::SM5 && (bForceTextureStreamingBuild || PlatformSupportsDebugViewShaders(Platform));
+	case DVSM_RayTracingDebug:
+		return FeatureLevel >= ERHIFeatureLevel::SM5 ;
 	default:
 		return false;
 	}
diff --git a/Engine/Source/Runtime/Engine/Private/SceneView.cpp b/Engine/Source/Runtime/Engine/Private/SceneView.cpp
--- a/Engine/Source/Runtime/Engine/Private/SceneView.cpp
+++ b/Engine/Source/Runtime/Engine/Private/SceneView.cpp
@@ -2555,6 +2555,11 @@ EDebugViewShaderMode FSceneViewFamily::ChooseDebugViewShaderMode() const
 	{
 		return DVSM_RequiredTextureResolution;
 	}
+	else if (EngineShowFlags.RayTracingDebug)
+	{
+		return DVSM_RayTracingDebug;
+	}
+
 	return DVSM_None;
 }
 
diff --git a/Engine/Source/Runtime/Engine/Private/ShowFlags.cpp b/Engine/Source/Runtime/Engine/Private/ShowFlags.cpp
--- a/Engine/Source/Runtime/Engine/Private/ShowFlags.cpp
+++ b/Engine/Source/Runtime/Engine/Private/ShowFlags.cpp
@@ -487,11 +487,16 @@ void EngineShowFlagOverride(EShowFlagInitMode ShowFlagInitMode, EViewModeIndex V
 			EngineShowFlags.SetVisualizeMotionBlur(false);
 			EngineShowFlags.SetDepthOfField(false);
 			EngineShowFlags.SetPostProcessMaterial(false);
+
+			if (bCanDisableTonemapper)
+			{
+				EngineShowFlags.SetTonemapper(false);
+			}
 		}
 	}
 
 	// disable AA in full screen GBuffer visualization
-	if (bCanDisableTonemapper && (EngineShowFlags.VisualizeBuffer || ViewModeIndex == VMI_RayTracingDebug ))
+	if (bCanDisableTonemapper && EngineShowFlags.VisualizeBuffer)
 	{
 		EngineShowFlags.SetTonemapper(false);
 	}
diff --git a/Engine/Source/Runtime/Engine/Public/DebugViewModeHelpers.h b/Engine/Source/Runtime/Engine/Public/DebugViewModeHelpers.h
--- a/Engine/Source/Runtime/Engine/Public/DebugViewModeHelpers.h
+++ b/Engine/Source/Runtime/Engine/Public/DebugViewModeHelpers.h
@@ -28,6 +28,7 @@ enum EDebugViewShaderMode
 	DVSM_MaterialTextureScaleAccuracy, // Visualize the accuracy of the material texture scales used for texture streaming.
 	DVSM_OutputMaterialTextureScales, // Outputs the material texture scales.
 	DVSM_RequiredTextureResolution, // Visualize the accuracy of the material texture scales used for texture streaming.
+	DVSM_RayTracingDebug,			// Visualize ray tracing debug modes.
 	DVSM_MAX
 };
 
diff --git a/Engine/Source/Runtime/Renderer/Private/DebugViewModeRendering.cpp b/Engine/Source/Runtime/Renderer/Private/DebugViewModeRendering.cpp
--- a/Engine/Source/Runtime/Renderer/Private/DebugViewModeRendering.cpp
+++ b/Engine/Source/Runtime/Renderer/Private/DebugViewModeRendering.cpp
@@ -19,6 +19,7 @@ DebugViewModeRendering.cpp: Contains definitions for rendering debug viewmodes.
 #include "CompositionLighting/PostProcessPassThrough.h"
 #include "PostProcess/PostProcessCompositeEditorPrimitives.h"
 #include "PostProcess/PostProcessUpscale.h"
+#include "PostProcess/PostProcessTemporalAA.h"
 #include "SceneRendering.h"
 #include "DeferredShadingRenderer.h"
 #include "MeshPassProcessor.inl"
@@ -68,9 +69,10 @@ void FDeferredShadingSceneRenderer::DoDebugViewModePostProcessing(FRHICommandLis
 	ensure(Context.View.PrimaryScreenPercentageMethod != EPrimaryScreenPercentageMethod::TemporalUpscale);
 
 	const bool bHDROutputEnabled = GRHISupportsHDROutput && IsHDREnabled();
-
-	// Shader complexity does not actually output a color
-	if (!View.Family->EngineShowFlags.ShaderComplexity)
+	
+	// Some view modes do not actually output a color so they should not be tonemapped	
+	const bool bAllowTonemapper = !View.Family->EngineShowFlags.ShaderComplexity && View.Family->EngineShowFlags.Tonemapper;
+	if (bAllowTonemapper)
 	{
 		GPostProcessing.AddGammaOnlyTonemapper(Context);
 	}
@@ -104,6 +106,20 @@ void FDeferredShadingSceneRenderer::DoDebugViewModePostProcessing(FRHICommandLis
 			Context.FinalOutput = FRenderingCompositeOutputRef(Node);
 			break;
 		}
+		case DVSM_RayTracingDebug:
+		{
+			FSceneViewState* ViewState = (FSceneViewState*)Context.View.State;
+			FTAAPassParameters Parameters(Context.View);	
+			FRenderingCompositePass* Node = Context.Graph.RegisterPass(new(FMemStack::Get()) FRCPassPostProcessTemporalAA(
+				Context, 
+				Parameters, 
+				Context.View.PrevViewInfo.TemporalAAHistory, 
+				&ViewState->PrevFrameViewInfo.TemporalAAHistory));
+
+			Node->SetInput(ePId_Input0, FRenderingCompositeOutputRef(Context.FinalOutput));
+			Context.FinalOutput = FRenderingCompositeOutputRef(Node);
+			break;
+		}
 		default:
 			ensure(false);
 			break;
