Fixes missing HistoryPreExposureCorrection shader parameter in SSD. Perforce edited file but did not checked it out...

diff --git a/Engine/Source/Runtime/Renderer/Private/ScreenSpaceDenoise.cpp b/Engine/Source/Runtime/Renderer/Private/ScreenSpaceDenoise.cpp
--- a/Engine/Source/Runtime/Renderer/Private/ScreenSpaceDenoise.cpp
+++ b/Engine/Source/Runtime/Renderer/Private/ScreenSpaceDenoise.cpp
@@ -653,6 +653,7 @@ class FSSDTemporalAccumulationCS : public FScreenSpaceDenoisingShader
 	BEGIN_SHADER_PARAMETER_STRUCT(FParameters, )
 		SHADER_PARAMETER_ARRAY(int32, bCameraCut, [IScreenSpaceDenoiser::kMaxBatchSize])
 		SHADER_PARAMETER(FMatrix, PrevScreenToTranslatedWorld)
+		SHADER_PARAMETER(float, HistoryPreExposureCorrection)
 
 		SHADER_PARAMETER_STRUCT_INCLUDE(FSSDCommonParameters, CommonParameters)
 		SHADER_PARAMETER_STRUCT_INCLUDE(FSSDConvolutionMetaData, ConvolutionMetaData)
@@ -1040,6 +1041,7 @@ static void DenoiseSignalAtConstantPixelDensity(
 		FSSDTemporalAccumulationCS::FParameters* PassParameters = GraphBuilder.AllocParameters<FSSDTemporalAccumulationCS::FParameters>();
 		PassParameters->CommonParameters = CommonParameters;
 		PassParameters->ConvolutionMetaData = ConvolutionMetaData;
+		PassParameters->HistoryPreExposureCorrection = View.PreExposure / View.PrevViewInfo.SceneColorPreExposure;
 
 		PassParameters->SignalInput = SignalHistory;
 		PassParameters->HistoryRejectionSignal = RejectionPreConvolutionSignal;
