Fixed wrong order of post pro passes in raytracing debug viz mode

diff --git a/Engine/Source/Runtime/Renderer/Private/DebugViewModeRendering.cpp b/Engine/Source/Runtime/Renderer/Private/DebugViewModeRendering.cpp
--- a/Engine/Source/Runtime/Renderer/Private/DebugViewModeRendering.cpp
+++ b/Engine/Source/Runtime/Renderer/Private/DebugViewModeRendering.cpp
@@ -71,7 +71,7 @@ void FDeferredShadingSceneRenderer::DoDebugViewModePostProcessing(FRHICommandLis
 	const bool bHDROutputEnabled = GRHISupportsHDROutput && IsHDREnabled();
 	
 	// Some view modes do not actually output a color so they should not be tonemapped	
-	const bool bAllowTonemapper = !View.Family->EngineShowFlags.ShaderComplexity && View.Family->EngineShowFlags.Tonemapper;
+	const bool bAllowTonemapper = !View.Family->EngineShowFlags.ShaderComplexity && !View.Family->EngineShowFlags.RayTracingDebug;
 	if (bAllowTonemapper)
 	{
 		GPostProcessing.AddGammaOnlyTonemapper(Context);
@@ -118,6 +118,11 @@ void FDeferredShadingSceneRenderer::DoDebugViewModePostProcessing(FRHICommandLis
 
 			Node->SetInput(ePId_Input0, FRenderingCompositeOutputRef(Context.FinalOutput));
 			Context.FinalOutput = FRenderingCompositeOutputRef(Node);
+
+			if (View.Family->EngineShowFlags.Tonemapper)
+			{
+				GPostProcessing.AddGammaOnlyTonemapper(Context);
+			}
 			break;
 		}
 		default:
