Refactors SSR to RDG

diff --git a/Engine/Shaders/Private/PostProcessHierarchical.usf b/Engine/Shaders/Private/PostProcessHierarchical.usf
deleted file mode 100644
diff --git a/Engine/Shaders/Private/SceneViewFamilyBlackboard.ush b/Engine/Shaders/Private/SceneViewFamilyBlackboard.ush
--- a/Engine/Shaders/Private/SceneViewFamilyBlackboard.ush
+++ b/Engine/Shaders/Private/SceneViewFamilyBlackboard.ush
@@ -21,21 +21,43 @@ Texture2D SceneGBufferD;
 Texture2D SceneGBufferE;
 Texture2D<uint> SceneLightingChannels;
 
+#if SUPPORTS_INDEPENDENT_SAMPLERS
+#define SceneDepthBufferSampler GlobalPointClampedSampler
+#define SceneVelocityBufferSampler GlobalPointClampedSampler
+#define SceneGBufferASampler GlobalPointClampedSampler
+#define SceneGBufferBSampler GlobalPointClampedSampler
+#define SceneGBufferCSampler GlobalPointClampedSampler
+#define SceneGBufferDSampler GlobalPointClampedSampler
+#define SceneGBufferESampler GlobalPointClampedSampler
+
+#else
+SamplerState SceneDepthBufferSampler;
+SamplerState SceneVelocityBufferSampler;
+SamplerState SceneGBufferASampler;
+SamplerState SceneGBufferBSampler;
+SamplerState SceneGBufferCSampler;
+SamplerState SceneGBufferDSampler;
+SamplerState SceneGBufferESampler;
+
+#endif
+
 
 // @param UV - UV space in the GBuffer textures (BufferSize resolution)
+// TOOD: need to find a way to make this more convenient.
 FGBufferData GetGBufferDataFromSceneBlackboard(float2 UV, bool bGetNormalizedNormal = true)
 {
-	float4 GBufferA = SceneGBufferA.SampleLevel(GlobalPointClampedSampler, UV, 0);
-	float4 GBufferB = SceneGBufferB.SampleLevel(GlobalPointClampedSampler, UV, 0);
-	float4 GBufferC = SceneGBufferC.SampleLevel(GlobalPointClampedSampler, UV, 0);
-	float4 GBufferD = SceneGBufferD.SampleLevel(GlobalPointClampedSampler, UV, 0);
-	float4 GBufferE = SceneGBufferE.SampleLevel(GlobalPointClampedSampler, UV, 0);
-	float4 GBufferVelocity = SceneVelocityBuffer.SampleLevel(GlobalPointClampedSampler, UV, 0);
+	float4 GBufferA = SceneGBufferA.SampleLevel(SceneGBufferASampler, UV, 0);
+	float4 GBufferB = SceneGBufferB.SampleLevel(SceneGBufferBSampler, UV, 0);
+	float4 GBufferC = SceneGBufferC.SampleLevel(SceneGBufferCSampler, UV, 0);
+	float4 GBufferD = SceneGBufferD.SampleLevel(SceneGBufferDSampler, UV, 0);
+	float4 GBufferE = SceneGBufferE.SampleLevel(SceneGBufferESampler, UV, 0);
+	float4 GBufferVelocity = SceneVelocityBuffer.SampleLevel(SceneVelocityBufferSampler, UV, 0);
 
 	uint CustomStencil = 0;
 	float CustomNativeDepth = 0;
 
-	float DeviceZ = SceneDepthBuffer.SampleLevel(GlobalPointClampedSampler, UV, 0).r;
+	float DeviceZ = SceneDepthBuffer.SampleLevel(SceneDepthBufferSampler, UV, 0).r;
+
 	float SceneDepth = ConvertFromDeviceZ(DeviceZ);
 
 	return DecodeGBufferData(GBufferA, GBufferB, GBufferC, GBufferD, GBufferE, GBufferVelocity, CustomNativeDepth, CustomStencil, SceneDepth, bGetNormalizedNormal, CheckerFromSceneColorUV(UV));
diff --git a/Engine/Shaders/Private/ScreenSpaceReflections.usf b/Engine/Shaders/Private/ScreenSpaceReflections.usf
--- a/Engine/Shaders/Private/ScreenSpaceReflections.usf
+++ b/Engine/Shaders/Private/ScreenSpaceReflections.usf
@@ -4,15 +4,12 @@
 	ScreenSpaceReflections.usf: To generate screen space reflections as a postprocess
 =============================================================================*/
 
-#include "Common.ush"	
-#include "PostProcessCommon.ush"				
-#include "DeferredShadingCommon.ush"
+#include "Common.ush"
 #include "Random.ush"
 #include "BRDF.ush"
 #include "MonteCarlo.ush"
 #include "ScreenSpaceRayCast.ush"
-
-#define SSR_CONE_QUALITY 5
+#include "SceneViewFamilyBlackboard.ush"
 
 // .r:Intensity in 0..1 range .g:RoughnessMaskMul, b:EnableDiscard for FPostProcessScreenSpaceReflectionsStencilPS, a:(bTemporalAAIsOn?TemporalAAIndex:StateFrameIndexMod8)*1551
 float4 SSRParams;
@@ -24,6 +21,12 @@ float4 HZBUvFactorAndInvFactor;
 // Equivalent of View.ScreenPositionScaleBias, but for SSR.
 float4 PrevScreenPositionScaleBias;
 
+Texture2D SceneColor;
+SamplerState SceneColorSampler;
+
+Texture2D HZB;
+SamplerState HZBSampler;
+
 float GetRoughness(const in FGBufferData GBuffer)
 {
 	float Roughness = GBuffer.Roughness;
@@ -58,8 +61,7 @@ void ScreenSpaceReflectionsStencilPS(in float4 UVAndScreenPos : TEXCOORD0, out f
 {
 	const float2 UV = UVAndScreenPos.xy;
 	
-	const FScreenSpaceData ScreenSpaceData = GetScreenSpaceData(UV);
-	const FGBufferData GBuffer = ScreenSpaceData.GBuffer;
+	FGBufferData GBuffer = GetGBufferDataFromSceneBlackboard(UV);
 	
 	const float Roughness = GetRoughness(GBuffer);
 	const float RoughnessFade = GetRoughnessFade(Roughness);
@@ -78,16 +80,17 @@ void ScreenSpaceReflectionsStencilPS(in float4 UVAndScreenPos : TEXCOORD0, out f
 }
 
 
-void ScreenSpaceReflectionsPS(in noperspective float4 UVAndScreenPos : TEXCOORD0, float4 SvPosition : SV_POSITION, out float4 OutColor : SV_Target0)
+void ScreenSpaceReflectionsPS(
+	float4 SvPosition : SV_POSITION,
+	out float4 OutColor : SV_Target0)
 {
-	float2 UV = UVAndScreenPos.xy;
-	float2 ScreenPos = UVAndScreenPos.zw;
+	float2 UV = SvPosition.xy * View.BufferSizeAndInvSize.zw;
+	float2 ScreenPos = ViewportUVToScreenPos((SvPosition.xy - View.ViewRectMin.xy) * View.ViewSizeAndInvSize.zw);
 	uint2 PixelPos = (uint2)SvPosition.xy;
 
 	OutColor = 0;
 
-	FScreenSpaceData ScreenSpaceData = GetScreenSpaceData(UV);
-	FGBufferData GBuffer = ScreenSpaceData.GBuffer;
+	FGBufferData GBuffer = GetGBufferDataFromSceneBlackboard(UV);
 	
 	float Roughness = GetRoughness(GBuffer);
 	float RoughnessFade = GetRoughnessFade(Roughness);
@@ -99,7 +102,7 @@ void ScreenSpaceReflectionsPS(in noperspective float4 UVAndScreenPos : TEXCOORD0
 	}
 
 	const float3 N = GBuffer.WorldNormal;
-	const float SceneDepth = CalcSceneDepth(UV);
+	const float SceneDepth = GBuffer.Depth;
 	const float3 PositionTranslatedWorld = mul( float4( ScreenPos * SceneDepth, SceneDepth, 1 ), View.ScreenToTranslatedWorld ).xyz;
 	const float3 V = normalize(View.TranslatedWorldCameraOrigin - PositionTranslatedWorld);
 
@@ -153,7 +156,7 @@ void ScreenSpaceReflectionsPS(in noperspective float4 UVAndScreenPos : TEXCOORD0
 
 #if 1
 			// TODO EnvBRDF isn't a perfect match anymore.
-			float3 H = mul( ImportanceSampleVisibleGGX( E, a2, TangentV ).xyz, TangentBasis );
+			float3 H = mul( ImportanceSampleVisibleGGX(UniformSampleDisk(E), a2, TangentV ).xyz, TangentBasis );
 			float3 L = 2 * dot( V, H ) * H - V;
 #elif 0
 			float3 H = mul( ImportanceSampleGGX( E, a2 ).xyz, TangentBasis );
@@ -186,7 +189,7 @@ void ScreenSpaceReflectionsPS(in noperspective float4 UVAndScreenPos : TEXCOORD0
 				float Level = 0;
 
 				bool bHit = RayCast(
-					PostprocessInput1, PostprocessInput1Sampler,
+					HZB, HZBSampler,
 					PositionTranslatedWorld, L, Roughness, SceneDepth, 
 					NumSteps, StepOffset,
 					HZBUvFactorAndInvFactor,
@@ -207,13 +210,10 @@ void ScreenSpaceReflectionsPS(in noperspective float4 UVAndScreenPos : TEXCOORD0
 				{
 					float2 SampleUV;
 					float Vignette;
-					ReprojectHit(PrevScreenPositionScaleBias, PostprocessInput3, PostprocessInput3Sampler, HitUVz, SampleUV, Vignette);
+					ReprojectHit(PrevScreenPositionScaleBias, SceneVelocityBuffer, SceneVelocityBufferSampler, HitUVz, SampleUV, Vignette);
+
+					float4 SampleColor = SampleScreenColor( SceneColor, SceneColorSampler, SampleUV ) * Vignette;
 
-#if SSR_QUALITY == SSR_CONE_QUALITY
-					float4 SampleColor = SampleHCBLevel( PostprocessInput2, PostprocessInput2Sampler, SampleUV, Level, HZBUvFactorAndInvFactor ) * Vignette;
-#else
-					float4 SampleColor = SampleScreenColor( PostprocessInput0, PostprocessInput0Sampler, SampleUV ) * Vignette;
-#endif
 					SampleColor.rgb *= rcp( 1 + Luminance(SampleColor.rgb) );
 
 					// Correct integration applies DGF below but for speed we apply EnvBRDF later when compositing
@@ -253,7 +253,7 @@ void ScreenSpaceReflectionsPS(in noperspective float4 UVAndScreenPos : TEXCOORD0
 		float Level = 0;
 
 		bool bHit = RayCast(
-			PostprocessInput1, PostprocessInput1Sampler,
+			HZB, HZBSampler,
 			PositionTranslatedWorld, R, Roughness, SceneDepth,
 			NumSteps, StepOffset,
 			HZBUvFactorAndInvFactor,
@@ -265,9 +265,9 @@ void ScreenSpaceReflectionsPS(in noperspective float4 UVAndScreenPos : TEXCOORD0
 		{
 			float2 SampleUV;
 			float Vignette;
-			ReprojectHit(PrevScreenPositionScaleBias, PostprocessInput3, PostprocessInput3Sampler, HitUVz, SampleUV, Vignette);
+			ReprojectHit(PrevScreenPositionScaleBias, SceneVelocityBuffer, SceneVelocityBufferSampler, HitUVz, SampleUV, Vignette);
 
-			OutColor = SampleScreenColor(PostprocessInput0, PostprocessInput0Sampler, SampleUV) * Vignette;
+			OutColor = SampleScreenColor(SceneColor, SceneColorSampler, SampleUV) * Vignette;
 		}
 	}
 	
diff --git a/Engine/Source/Runtime/Renderer/Private/PostProcess/PostProcessHierarchical.cpp b/Engine/Source/Runtime/Renderer/Private/PostProcess/PostProcessHierarchical.cpp
deleted file mode 100644
diff --git a/Engine/Source/Runtime/Renderer/Private/PostProcess/PostProcessHierarchical.h b/Engine/Source/Runtime/Renderer/Private/PostProcess/PostProcessHierarchical.h
deleted file mode 100644
diff --git a/Engine/Source/Runtime/Renderer/Private/PostProcess/ScreenSpaceReflections.cpp b/Engine/Source/Runtime/Renderer/Private/PostProcess/ScreenSpaceReflections.cpp
--- a/Engine/Source/Runtime/Renderer/Private/PostProcess/ScreenSpaceReflections.cpp
+++ b/Engine/Source/Runtime/Renderer/Private/PostProcess/ScreenSpaceReflections.cpp
@@ -5,19 +5,10 @@
 =============================================================================*/
 
 #include "PostProcess/ScreenSpaceReflections.h"
-#include "StaticBoundShaderState.h"
-#include "SceneUtils.h"
-#include "PostProcess/SceneRenderTargets.h"
-#include "SceneRenderTargetParameters.h"
 #include "ScenePrivate.h"
-#include "PostProcess/SceneFilterRendering.h"
-#include "PostProcess/PostProcessInput.h"
-#include "PostProcess/PostProcessOutput.h"
-#include "PostProcess/PostProcessing.h"
-#include "PostProcess/PostProcessTemporalAA.h"
-#include "PostProcess/PostProcessHierarchical.h"
-#include "ClearQuad.h"
-#include "PipelineStateCache.h"
+#include "RenderGraph.h"
+#include "PixelShaderUtils.h"
+#include "SceneViewFamilyBlackboard.h"
 
 static TAutoConsoleVariable<int32> CVarSSRQuality(
 	TEXT("r.SSR.Quality"),
@@ -46,13 +37,6 @@ static TAutoConsoleVariable<int32> CVarSSRStencil(
 	TEXT(" 0 is off (default), 1 is on"),
 	ECVF_RenderThreadSafe);
 
-static TAutoConsoleVariable<int32> CVarSSRCone(
-	TEXT("r.SSR.Cone"),
-	0,
-	TEXT("Defines if we use cone traced screen space reflection\n")
-	TEXT(" 0 is off (default), 1 is on"),
-	ECVF_RenderThreadSafe);
-
 DECLARE_GPU_STAT_NAMED(ScreenSpaceReflections, TEXT("ScreenSpace Reflections"));
 
 bool ShouldRenderScreenSpaceReflections(const FViewInfo& View)
@@ -101,10 +85,12 @@ bool IsSSRTemporalPassRequired(const FViewInfo& View, bool bCheckSSREnabled)
 	return View.AntiAliasingMethod != AAM_TemporalAA || CVarSSRTemporal.GetValueOnRenderThread() != 0;
 }
 
+namespace
+{
 
-static float ComputeRoughnessMaskScale(const FRenderingCompositePassContext& Context, uint32 SSRQuality)
+float ComputeRoughnessMaskScale(const FViewInfo& View, uint32 SSRQuality)
 {
-	float MaxRoughness = FMath::Clamp(Context.View.FinalPostProcessSettings.ScreenSpaceReflectionMaxRoughness, 0.01f, 1.0f);
+	float MaxRoughness = FMath::Clamp(View.FinalPostProcessSettings.ScreenSpaceReflectionMaxRoughness, 0.01f, 1.0f);
 
 	// f(x) = x * Scale + Bias
 	// f(MaxRoughness) = 0
@@ -114,41 +100,48 @@ static float ComputeRoughnessMaskScale(const FRenderingCompositePassContext& Con
 	return RoughnessMaskScale * (SSRQuality < 3 ? 2.0f : 1.0f);
 }
 
-FLinearColor ComputeSSRParams(const FRenderingCompositePassContext& Context, uint32 SSRQuality, bool bEnableDiscard)
+FLinearColor ComputeSSRParams(const FViewInfo& View, uint32 SSRQuality, bool bEnableDiscard)
 {
-	float RoughnessMaskScale = ComputeRoughnessMaskScale(Context, SSRQuality);
+	float RoughnessMaskScale = ComputeRoughnessMaskScale(View, SSRQuality);
 
 	float FrameRandom = 0;
 
-	if(Context.ViewState)
+	if(View.ViewState)
 	{
-		bool bTemporalAAIsOn = Context.View.AntiAliasingMethod == AAM_TemporalAA;
+		bool bTemporalAAIsOn = View.AntiAliasingMethod == AAM_TemporalAA;
 
 		if(bTemporalAAIsOn)
 		{
 			// usually this number is in the 0..7 range but it depends on the TemporalAA quality
-			FrameRandom = Context.ViewState->GetCurrentTemporalAASampleIndex() * 1551;
+			FrameRandom = View.ViewState->GetCurrentTemporalAASampleIndex() * 1551;
 		}
 		else
 		{
 			// 8 aligns with the temporal smoothing, larger number will do more flickering (power of two for best performance)
-			FrameRandom = Context.ViewState->GetFrameIndex(8) * 1551;
+			FrameRandom = View.ViewState->GetFrameIndex(8) * 1551;
 		}
 	}
 
 	return FLinearColor(
-		FMath::Clamp(Context.View.FinalPostProcessSettings.ScreenSpaceReflectionIntensity * 0.01f, 0.0f, 1.0f), 
+		FMath::Clamp(View.FinalPostProcessSettings.ScreenSpaceReflectionIntensity * 0.01f, 0.0f, 1.0f), 
 		RoughnessMaskScale,
 		(float)bEnableDiscard,	// TODO 
 		FrameRandom);
 }
 
-/**
- * Encapsulates the post processing screen space reflections pixel shader stencil pass.
- */
-class FPostProcessScreenSpaceReflectionsStencilPS : public FGlobalShader
+
+BEGIN_SHADER_PARAMETER_STRUCT(FSSRCommonParameters, )
+	SHADER_PARAMETER(FLinearColor, SSRParams)
+	SHADER_PARAMETER_STRUCT_INCLUDE(FSceneViewFamilyBlackboard, SceneTextures)
+	SHADER_PARAMETER_STRUCT_INCLUDE(FSceneTextureSamplerParameters, SceneTextureSamplers)
+	SHADER_PARAMETER_STRUCT_REF(FViewUniformShaderParameters, ViewUniformBuffer)
+END_SHADER_PARAMETER_STRUCT()
+
+
+class FScreenSpaceReflectionsStencilPS : public FGlobalShader
 {
-	DECLARE_GLOBAL_SHADER(FPostProcessScreenSpaceReflectionsStencilPS);
+	DECLARE_GLOBAL_SHADER(FScreenSpaceReflectionsStencilPS);
+	SHADER_USE_PARAMETER_STRUCT(FScreenSpaceReflectionsStencilPS, FGlobalShader);
 
 	using FPermutationDomain = FShaderPermutationNone;
 
@@ -163,69 +156,25 @@ class FPostProcessScreenSpaceReflectionsStencilPS : public FGlobalShader
 		OutEnvironment.SetDefine( TEXT("PREV_FRAME_COLOR"), uint32(0) );
 		OutEnvironment.SetDefine( TEXT("SSR_QUALITY"), uint32(0) );
 	}
-
-	/** Default constructor. */
-	FPostProcessScreenSpaceReflectionsStencilPS() {}
-
-	FPostProcessPassParameters PostprocessParameter;
-	FSceneTextureShaderParameters SceneTextureParameters;
-	FShaderParameter SSRParams;
-
-	/** Initialization constructor. */
-	FPostProcessScreenSpaceReflectionsStencilPS(const ShaderMetaType::CompiledShaderInitializerType& Initializer)
-		: FGlobalShader(Initializer)
-	{
-		PostprocessParameter.Bind(Initializer.ParameterMap);
-		SceneTextureParameters.Bind(Initializer);
-		SSRParams.Bind(Initializer.ParameterMap, TEXT("SSRParams"));
-	}
-
-	template <typename TRHICmdList>
-	void SetParameters(TRHICmdList& RHICmdList, const FRenderingCompositePassContext& Context, uint32 SSRQuality, bool EnableDiscard)
-	{
-		const FFinalPostProcessSettings& Settings = Context.View.FinalPostProcessSettings;
-		const FPixelShaderRHIParamRef ShaderRHI = GetPixelShader();
-
-		FGlobalShader::SetParameters<FViewUniformShaderParameters>(Context.RHICmdList, ShaderRHI, Context.View.ViewUniformBuffer);
-
-		PostprocessParameter.SetPS(RHICmdList, ShaderRHI, Context, TStaticSamplerState<SF_Point, AM_Clamp, AM_Clamp, AM_Clamp>::GetRHI());
-		SceneTextureParameters.Set(RHICmdList, ShaderRHI, Context.View.FeatureLevel, ESceneTextureSetupMode::All);
-
-		{
-			FLinearColor Value = ComputeSSRParams(Context, SSRQuality, EnableDiscard);
-
-			SetShaderValue(RHICmdList, ShaderRHI, SSRParams, Value);
-		}
-	}
-
-	// FShader interface.
-	virtual bool Serialize(FArchive& Ar) override
-	{
-		bool bShaderHasOutdatedParameters = FGlobalShader::Serialize(Ar);
-		Ar << PostprocessParameter << SceneTextureParameters << SSRParams;
-		return bShaderHasOutdatedParameters;
-	}
+	
+	BEGIN_SHADER_PARAMETER_STRUCT(FParameters, )
+		SHADER_PARAMETER_STRUCT_INCLUDE(FSSRCommonParameters, CommonParameters)
+		RENDER_TARGET_BINDING_SLOTS()
+	END_SHADER_PARAMETER_STRUCT()
 };
 
-IMPLEMENT_GLOBAL_SHADER(FPostProcessScreenSpaceReflectionsStencilPS, "/Engine/Private/ScreenSpaceReflections.usf", "ScreenSpaceReflectionsStencilPS", SF_Pixel);
 
-namespace
-{
-static const uint32 SSRConeQuality = 5;
-static constexpr int32 QualityCount = 6;
+static constexpr int32 QualityCount = 5;
 
 class FSSRQualityDim : SHADER_PERMUTATION_INT("SSR_QUALITY", QualityCount);
 class FSSRPrevFrameColorDim : SHADER_PERMUTATION_BOOL("PREV_FRAME_COLOR");
 class FSSRPrevFrameColorDim2 : SHADER_PERMUTATION_BOOL("PREV_FRAME_COLOR");
-}
 
 
-
-
-
-class FPostProcessScreenSpaceReflectionsPS : public FGlobalShader
+class FScreenSpaceReflectionsPS : public FGlobalShader
 {
-	DECLARE_GLOBAL_SHADER(FPostProcessScreenSpaceReflectionsPS);
+	DECLARE_GLOBAL_SHADER(FScreenSpaceReflectionsPS);
+	SHADER_USE_PARAMETER_STRUCT(FScreenSpaceReflectionsPS, FGlobalShader);
 
 	using FPermutationDomain = TShaderPermutationDomain<FSSRQualityDim, FSSRPrevFrameColorDim>;
 
@@ -238,110 +187,26 @@ class FPostProcessScreenSpaceReflectionsPS : public FGlobalShader
 		}
 		return IsFeatureLevelSupported(Parameters.Platform, ERHIFeatureLevel::SM4);
 	}
+	
+	BEGIN_SHADER_PARAMETER_STRUCT(FParameters, )
+		SHADER_PARAMETER_STRUCT_INCLUDE(FSSRCommonParameters, CommonParameters)
 
-	/** Default constructor. */
-	FPostProcessScreenSpaceReflectionsPS() {}
-
-	FPostProcessPassParameters PostprocessParameter;
-	FSceneTextureShaderParameters SceneTextureParameters;
-	FShaderParameter SSRParams;
-	FShaderParameter HZBUvFactorAndInvFactor;
-	FShaderParameter PrevScreenPositionScaleBias;
-	FShaderParameter PrevSceneColorPreExposureCorrection;
-
-	/** Initialization constructor. */
-	FPostProcessScreenSpaceReflectionsPS(const ShaderMetaType::CompiledShaderInitializerType& Initializer)
-		: FGlobalShader(Initializer)
-	{
-		PostprocessParameter.Bind(Initializer.ParameterMap);
-		SceneTextureParameters.Bind(Initializer);
-		SSRParams.Bind(Initializer.ParameterMap, TEXT("SSRParams"));
-		HZBUvFactorAndInvFactor.Bind(Initializer.ParameterMap, TEXT("HZBUvFactorAndInvFactor"));
-		PrevScreenPositionScaleBias.Bind(Initializer.ParameterMap, TEXT("PrevScreenPositionScaleBias"));
-		PrevSceneColorPreExposureCorrection.Bind(Initializer.ParameterMap, TEXT("PrevSceneColorPreExposureCorrection"));
-	}
-
-	template <typename TRHICmdList>
-	void SetParameters(TRHICmdList& RHICmdList, const FRenderingCompositePassContext& Context, int32 SSRQuality)
-	{
-		FSceneRenderTargets& SceneContext = FSceneRenderTargets::Get(Context.RHICmdList);
-
-		const FFinalPostProcessSettings& Settings = Context.View.FinalPostProcessSettings;
-		const FPixelShaderRHIParamRef ShaderRHI = GetPixelShader();
-
-		FGlobalShader::SetParameters<FViewUniformShaderParameters>(RHICmdList, ShaderRHI, Context.View.ViewUniformBuffer);
-
-		PostprocessParameter.SetPS(RHICmdList, ShaderRHI, Context, TStaticSamplerState<SF_Point, AM_Clamp, AM_Clamp, AM_Clamp>::GetRHI());
-		SceneTextureParameters.Set(RHICmdList, ShaderRHI, Context.View.FeatureLevel, ESceneTextureSetupMode::All);
-
-		{
-			FLinearColor Value = ComputeSSRParams(Context, SSRQuality, false);
-
-			SetShaderValue(RHICmdList, ShaderRHI, SSRParams, Value);
-		}
-
-		{
-			const FVector2D HZBUvFactor(
-				float(Context.View.ViewRect.Width()) / float(2 * Context.View.HZBMipmap0Size.X),
-				float(Context.View.ViewRect.Height()) / float(2 * Context.View.HZBMipmap0Size.Y)
-				);
-			const FVector4 HZBUvFactorAndInvFactorValue(
-				HZBUvFactor.X,
-				HZBUvFactor.Y,
-				1.0f / HZBUvFactor.X,
-				1.0f / HZBUvFactor.Y
-				);
-			
-			SetShaderValue(RHICmdList, ShaderRHI, HZBUvFactorAndInvFactor, HZBUvFactorAndInvFactorValue);
-		}
-
-		{
-			FIntPoint ViewportOffset = Context.View.ViewRect.Min;
-			FIntPoint ViewportExtent = Context.View.ViewRect.Size();
-			FIntPoint BufferSize = SceneContext.GetBufferSizeXY();
-
-			if (Context.View.PrevViewInfo.TemporalAAHistory.IsValid())
-			{
-				ViewportOffset = Context.View.PrevViewInfo.TemporalAAHistory.ViewportRect.Min;
-				ViewportExtent = Context.View.PrevViewInfo.TemporalAAHistory.ViewportRect.Size();
-				BufferSize = Context.View.PrevViewInfo.TemporalAAHistory.ReferenceBufferSize;
-			}
-
-			FVector2D InvBufferSize(1.0f / float(BufferSize.X), 1.0f / float(BufferSize.Y));
-
-			FVector4 ScreenPosToPixelValue(
-				ViewportExtent.X * 0.5f * InvBufferSize.X,
-				-ViewportExtent.Y * 0.5f * InvBufferSize.Y,
-				(ViewportExtent.X * 0.5f + ViewportOffset.X) * InvBufferSize.X,
-				(ViewportExtent.Y * 0.5f + ViewportOffset.Y) * InvBufferSize.Y);
-			SetShaderValue(Context.RHICmdList, ShaderRHI, PrevScreenPositionScaleBias, ScreenPosToPixelValue);
-		}
-
-		{
-			float PrevSceneColorPreExposureCorrectionValue = 1.0f;
-
-			if (Context.View.PrevViewInfo.TemporalAAHistory.IsValid())
-			{
-				PrevSceneColorPreExposureCorrectionValue = Context.View.PreExposure / Context.View.PrevViewInfo.SceneColorPreExposure;
-			}
-
-			SetShaderValue(RHICmdList, ShaderRHI, PrevSceneColorPreExposureCorrection, PrevSceneColorPreExposureCorrectionValue);
-		}
-	}
-
-	// FShader interface.
-	virtual bool Serialize(FArchive& Ar) override
-	{
-		bool bShaderHasOutdatedParameters = FGlobalShader::Serialize(Ar);
-		Ar << PostprocessParameter << SceneTextureParameters << SSRParams << HZBUvFactorAndInvFactor << PrevScreenPositionScaleBias << PrevSceneColorPreExposureCorrection;
-		return bShaderHasOutdatedParameters;
-	}
+		SHADER_PARAMETER(FVector4, HZBUvFactorAndInvFactor)
+		SHADER_PARAMETER(FVector4, PrevScreenPositionScaleBias)
+		SHADER_PARAMETER(float, PrevSceneColorPreExposureCorrection)
+		
+		SHADER_PARAMETER_RDG_TEXTURE(Texture2D, SceneColor)
+		SHADER_PARAMETER_SAMPLER(SamplerState, SceneColorSampler)
+		
+		SHADER_PARAMETER_RDG_TEXTURE(Texture2D, HZB)
+		SHADER_PARAMETER_SAMPLER(SamplerState, HZBSampler)
+		
+		RENDER_TARGET_BINDING_SLOTS()
+	END_SHADER_PARAMETER_STRUCT()
 };
 
-IMPLEMENT_GLOBAL_SHADER(FPostProcessScreenSpaceReflectionsPS, "/Engine/Private/ScreenSpaceReflections.usf", "ScreenSpaceReflectionsPS", SF_Pixel);
-
-// --------------------------------------------------------
-
+IMPLEMENT_GLOBAL_SHADER(FScreenSpaceReflectionsPS, "/Engine/Private/ScreenSpaceReflections.usf", "ScreenSpaceReflectionsPS", SF_Pixel);
+IMPLEMENT_GLOBAL_SHADER(FScreenSpaceReflectionsStencilPS, "/Engine/Private/ScreenSpaceReflections.usf", "ScreenSpaceReflectionsStencilPS", SF_Pixel);
 
 // @param Quality usually in 0..100 range, default is 50
 // @return see CVarSSRQuality, never 0
@@ -363,257 +228,170 @@ static int32 ComputeSSRQuality(float Quality)
 	return FMath::Min(Ret, SSRQualityCVar);
 }
 
-void FRCPassPostProcessScreenSpaceReflections::Process(FRenderingCompositePassContext& Context)
-{
-	auto& RHICmdList = Context.RHICmdList;
-	FSceneRenderTargets& SceneContext = FSceneRenderTargets::Get(RHICmdList);
+} // namespace
 
-	const FViewInfo& View = Context.View;
-	const auto FeatureLevel = Context.GetFeatureLevel();
-
-	int32 SSRQuality = ComputeSSRQuality(View.FinalPostProcessSettings.ScreenSpaceReflectionQuality);
-
-	SSRQuality = FMath::Clamp(SSRQuality, 1, 4);
-	
-	const bool VisualizeSSR = View.Family->EngineShowFlags.VisualizeSSR;
-	const bool SSRStencilPrePass = CVarSSRStencil.GetValueOnRenderThread() != 0 && !VisualizeSSR;
-
-	FRenderingCompositeOutputRef* Input2 = GetInput(ePId_Input2);
-
-	const bool SSRConeTracing = Input2 && Input2->GetOutput();
+FRDGTextureRef RenderScreenSpaceReflections(
+	FRDGBuilder& GraphBuilder,
+	const FSceneViewFamilyBlackboard& SceneTextures,
+	const FRDGTextureRef CurrentSceneColor,
+	const FViewInfo& View)
+{
+	RDG_EVENT_SCOPE(GraphBuilder, "ScreenSpaceReflections");
 	
-	if (VisualizeSSR)
+	FRDGTextureRef InputColor = CurrentSceneColor;
+	bool bSamplePrevFrame = false;
+	if (View.PrevViewInfo.CustomSSRInput.IsValid())
 	{
-		SSRQuality = 0;
+		InputColor = GraphBuilder.RegisterExternalTexture(View.PrevViewInfo.CustomSSRInput);
+		bSamplePrevFrame = true;
 	}
-	else if (SSRConeTracing)
+	else if (View.PrevViewInfo.TemporalAAHistory.IsValid())
 	{
-		SSRQuality = SSRConeQuality;
+		InputColor = GraphBuilder.RegisterExternalTexture(View.PrevViewInfo.TemporalAAHistory.RT[0]);
+		bSamplePrevFrame = true;
 	}
+
+	const bool VisualizeSSR = View.Family->EngineShowFlags.VisualizeSSR;
+	const bool SSRStencilPrePass = CVarSSRStencil.GetValueOnRenderThread() != 0 && !VisualizeSSR;
+	const int32 SSRQuality = VisualizeSSR ? 0 : FMath::Clamp(ComputeSSRQuality(View.FinalPostProcessSettings.ScreenSpaceReflectionQuality), 1, 4);
 	
-	const FSceneRenderTargetItem& DestRenderTarget = PassOutputs[0].RequestSurface(Context);
+	// Alloc SSR output.
+	FRDGTextureRef SSROutput;
+	{
+		FRDGTextureDesc Desc = FPooledRenderTargetDesc::Create2DDesc(
+			FSceneRenderTargets::Get_FrameConstantsOnly().GetBufferSizeXY(),
+			PF_FloatRGBA, FClearValueBinding(FLinearColor(0, 0, 0, 0)),
+			TexCreate_None, TexCreate_RenderTargetable,
+			false);
 
-	SCOPED_GPU_STAT(RHICmdList, ScreenSpaceReflections);
+		Desc.AutoWritable = false;
+		Desc.Flags |= GFastVRamConfig.SSR;
+
+		SSROutput = GraphBuilder.CreateTexture(Desc, TEXT("ScreenSpaceReflections"));
+	}
+		
+	FSSRCommonParameters CommonParameters;
+	CommonParameters.SSRParams = ComputeSSRParams(View, SSRQuality, false);
+	CommonParameters.ViewUniformBuffer = View.ViewUniformBuffer;
+	CommonParameters.SceneTextures = SceneTextures;
+	SetupSceneTextureSamplers(&CommonParameters.SceneTextureSamplers);
 	
+	FRenderTargetBindingSlots RenderTargets;
+	RenderTargets[0] = FRenderTargetBinding(SSROutput, ERenderTargetLoadAction::ENoAction, ERenderTargetStoreAction::EStore);
+
+	// Do a pre pass that output 0, or set a stencil mask to run the more expensive pixel shader.
 	if (SSRStencilPrePass)
-	{ // ScreenSpaceReflectionsStencil draw event
-		SCOPED_DRAW_EVENTF(Context.RHICmdList, ScreenSpaceReflectionsStencil, TEXT("ScreenSpaceReflectionsStencil %dx%d"), View.ViewRect.Width(), View.ViewRect.Height());
-
-		FRHIRenderPassInfo RPInfo(DestRenderTarget.TargetableTexture,
-			MakeRenderTargetActions(ERenderTargetLoadAction::ENoAction, ERenderTargetStoreAction::EStore),
-			SceneContext.GetSceneDepthSurface(),
-			MakeDepthStencilTargetActions(MakeRenderTargetActions(ERenderTargetLoadAction::ENoAction, ERenderTargetStoreAction::ENoAction),
-				MakeRenderTargetActions(ERenderTargetLoadAction::ENoAction, ERenderTargetStoreAction::EStore)),
+	{
+		// Also bind the depth buffer
+		RenderTargets.DepthStencil = FDepthStencilBinding(
+			SceneTextures.SceneDepthBuffer,
+			ERenderTargetLoadAction::ENoAction, ERenderTargetStoreAction::ENoAction,
+			ERenderTargetLoadAction::ELoad, ERenderTargetStoreAction::EStore,
 			FExclusiveDepthStencil::DepthRead_StencilWrite);
 
-		RHICmdList.BeginRenderPass(RPInfo, TEXT("SSRStencil"));
+		FScreenSpaceReflectionsStencilPS::FParameters* PassParameters = GraphBuilder.AllocParameters<FScreenSpaceReflectionsStencilPS::FParameters>();
+		PassParameters->CommonParameters = CommonParameters;
+		PassParameters->RenderTargets = RenderTargets;
 		
-		TShaderMapRef< FPostProcessVS > VertexShader(Context.GetShaderMap());
-		TShaderMapRef< FPostProcessScreenSpaceReflectionsStencilPS > PixelShader(Context.GetShaderMap());
-
-		Context.SetViewportAndCallRHI(View.ViewRect);
-
-		// Clear stencil to 0
-		DrawClearQuad(RHICmdList, false, FLinearColor(), false, 0, true, 0, PassOutputs[0].RenderTargetDesc.Extent, View.ViewRect);
-
-		FGraphicsPipelineStateInitializer GraphicsPSOInit;
-		RHICmdList.ApplyCachedRenderTargets(GraphicsPSOInit);
-
-		// Clobers the stencil to pixel that should not compute SSR
-		GraphicsPSOInit.DepthStencilState = TStaticDepthStencilState<false, CF_Always, true, CF_Always, SO_Replace, SO_Replace, SO_Replace>::GetRHI();
-
-		// Set rasterizer state to solid
-		GraphicsPSOInit.RasterizerState = TStaticRasterizerState<>::GetRHI();
-
-		// disable blend mode
-		GraphicsPSOInit.BlendState = TStaticBlendState<>::GetRHI();
-
-		// bind shader
-		GraphicsPSOInit.BoundShaderState.VertexDeclarationRHI = GFilterVertexDeclaration.VertexDeclarationRHI;
-		GraphicsPSOInit.BoundShaderState.VertexShaderRHI = GETSAFERHISHADER_VERTEX(*VertexShader);
-		GraphicsPSOInit.BoundShaderState.PixelShaderRHI = GETSAFERHISHADER_PIXEL(*PixelShader);
-		GraphicsPSOInit.PrimitiveType = PT_TriangleList;
-
-		SetGraphicsPipelineState(Context.RHICmdList, GraphicsPSOInit);
-		RHICmdList.SetStencilRef(0x80);
-
-		VertexShader->SetParameters(Context);
-		PixelShader->SetParameters(Context.RHICmdList, Context, SSRQuality, true);
+		TShaderMapRef<FScreenSpaceReflectionsStencilPS> PixelShader(View.ShaderMap);
+		ClearUnusedGraphResources(*PixelShader, PassParameters);
+		
+		GraphBuilder.AddPass(
+			RDG_EVENT_NAME("SSR StencilSetup %dx%d", View.ViewRect.Width(), View.ViewRect.Height()),
+			PassParameters,
+			ERenderGraphPassFlags::None,
+			[PassParameters, &View, PixelShader](FRHICommandList& RHICmdList)
+		{
+			RHICmdList.SetViewport(View.ViewRect.Min.X, View.ViewRect.Min.Y, 0.0f, View.ViewRect.Max.X, View.ViewRect.Max.Y, 1.0f);
+			RHICmdList.SetStencilRef(0x80);
+		
+			FGraphicsPipelineStateInitializer GraphicsPSOInit;
+			FPixelShaderUtils::InitFullscreenPipelineState(RHICmdList, View.ShaderMap, *PixelShader, /* out */ GraphicsPSOInit);
+			// Clobers the stencil to pixel that should not compute SSR
+			GraphicsPSOInit.DepthStencilState = TStaticDepthStencilState<false, CF_Always, true, CF_Always, SO_Replace, SO_Replace, SO_Replace>::GetRHI();
 
-		DrawPostProcessPass(
-			Context.RHICmdList,
-			0, 0,
-			View.ViewRect.Width(), View.ViewRect.Height(),
-			View.ViewRect.Min.X, View.ViewRect.Min.Y,
-			View.ViewRect.Width(), View.ViewRect.Height(),
-			View.ViewRect.Size(),
-			SceneContext.GetBufferSizeXY(),
-			*VertexShader,
-			View.StereoPass,
-			Context.HasHmdMesh(),
-			EDRF_UseTriangleOptimization);
+			SetGraphicsPipelineState(RHICmdList, GraphicsPSOInit);
+			SetShaderParameters(RHICmdList, *PixelShader, PixelShader->GetPixelShader(), *PassParameters);
 
-	} // ScreenSpaceReflectionsStencil draw event
+			FPixelShaderUtils::DrawFullscreenTriangle(RHICmdList);
+		});
+	}
 
-	{ // ScreenSpaceReflections draw event
-		SCOPED_DRAW_EVENTF(Context.RHICmdList, ScreenSpaceReflections, TEXT("ScreenSpaceReflections %dx%d"), View.ViewRect.Width(), View.ViewRect.Height());
+	// Adds SSR pass.
+	{
+		FScreenSpaceReflectionsPS::FPermutationDomain PermutationVector;
+		PermutationVector.Set<FSSRPrevFrameColorDim>(bSamplePrevFrame && SSRQuality != 0);
+		PermutationVector.Set<FSSRQualityDim>(SSRQuality);
 
-		FGraphicsPipelineStateInitializer GraphicsPSOInit;
-		if (SSRStencilPrePass)
+		FScreenSpaceReflectionsPS::FParameters* PassParameters = GraphBuilder.AllocParameters<FScreenSpaceReflectionsPS::FParameters>();
+		PassParameters->CommonParameters = CommonParameters;
 		{
-			check(RHICmdList.IsInsideRenderPass());
-
-			// set up the stencil test to match 0, meaning FPostProcessScreenSpaceReflectionsStencilPS has been discarded
-			GraphicsPSOInit.DepthStencilState = TStaticDepthStencilState<false, CF_Always, true, CF_Equal, SO_Keep, SO_Keep, SO_Keep>::GetRHI();
-			RHICmdList.ApplyCachedRenderTargets(GraphicsPSOInit);
+			const FVector2D HZBUvFactor(
+				float(View.ViewRect.Width()) / float(2 * View.HZBMipmap0Size.X),
+				float(View.ViewRect.Height()) / float(2 * View.HZBMipmap0Size.Y));
+			PassParameters->HZBUvFactorAndInvFactor = FVector4(
+				HZBUvFactor.X,
+				HZBUvFactor.Y,
+				1.0f / HZBUvFactor.X,
+				1.0f / HZBUvFactor.Y);
 		}
-		else
 		{
-			check(!RHICmdList.IsInsideRenderPass());
-
-			ERenderTargetLoadAction LoadAction = Context.GetLoadActionForRenderTarget(DestRenderTarget);
+			FIntPoint ViewportOffset = View.ViewRect.Min;
+			FIntPoint ViewportExtent = View.ViewRect.Size();
+			FIntPoint BufferSize = SceneTextures.SceneDepthBuffer->Desc.Extent;
 
-			// #todo-renderpasses LoadAction isn't actually used since we clear immediately
-			FRHIRenderPassInfo RPInfo(DestRenderTarget.TargetableTexture, MakeRenderTargetActions(ERenderTargetLoadAction::EClear, ERenderTargetStoreAction::EStore), DestRenderTarget.ShaderResourceTexture);
-			RHICmdList.BeginRenderPass(RPInfo, TEXT("ScreenSpaceReflections"));
+			if (View.PrevViewInfo.TemporalAAHistory.IsValid())
+			{
+				ViewportOffset = View.PrevViewInfo.TemporalAAHistory.ViewportRect.Min;
+				ViewportExtent = View.PrevViewInfo.TemporalAAHistory.ViewportRect.Size();
+				BufferSize = View.PrevViewInfo.TemporalAAHistory.ReferenceBufferSize;
+			}
 
-			Context.SetViewportAndCallRHI(View.ViewRect);
+			FVector2D InvBufferSize(1.0f / float(BufferSize.X), 1.0f / float(BufferSize.Y));
 
-			GraphicsPSOInit.DepthStencilState = TStaticDepthStencilState<false, CF_Always>::GetRHI();
-			RHICmdList.ApplyCachedRenderTargets(GraphicsPSOInit);
+			PassParameters->PrevScreenPositionScaleBias = FVector4(
+				ViewportExtent.X * 0.5f * InvBufferSize.X,
+				-ViewportExtent.Y * 0.5f * InvBufferSize.Y,
+				(ViewportExtent.X * 0.5f + ViewportOffset.X) * InvBufferSize.X,
+				(ViewportExtent.Y * 0.5f + ViewportOffset.Y) * InvBufferSize.Y);
 		}
+		PassParameters->PrevSceneColorPreExposureCorrection = bSamplePrevFrame ? View.PreExposure / View.PrevViewInfo.SceneColorPreExposure : 1.0f;
+		
+		PassParameters->SceneColor = InputColor;
+		PassParameters->SceneColorSampler = TStaticSamplerState<SF_Point>::GetRHI();
+		
+		PassParameters->HZB = GraphBuilder.RegisterExternalTexture(View.HZB);
+		PassParameters->HZBSampler = TStaticSamplerState<SF_Point>::GetRHI();
+		
+		PassParameters->RenderTargets = RenderTargets;
 
-		check(RHICmdList.IsInsideRenderPass());
-
-		// set the state
-		GraphicsPSOInit.BlendState = TStaticBlendState<>::GetRHI();
-		GraphicsPSOInit.RasterizerState = TStaticRasterizerState<>::GetRHI();
-		GraphicsPSOInit.PrimitiveType = PT_TriangleList;
-
-		TShaderMapRef< FPostProcessVS > VertexShader(Context.GetShaderMap());
-
-		FPostProcessScreenSpaceReflectionsPS::FPermutationDomain PermutationVector;
-		PermutationVector.Set<FSSRPrevFrameColorDim>(bPrevFrame && SSRQuality != 0);
-		PermutationVector.Set<FSSRQualityDim>(SSRQuality);
-
-		TShaderMapRef<FPostProcessScreenSpaceReflectionsPS> PixelShader(Context.GetShaderMap(), PermutationVector);
-		GraphicsPSOInit.BoundShaderState.VertexDeclarationRHI = GFilterVertexDeclaration.VertexDeclarationRHI;
-		GraphicsPSOInit.BoundShaderState.VertexShaderRHI = GETSAFERHISHADER_VERTEX(*VertexShader);
-		GraphicsPSOInit.BoundShaderState.PixelShaderRHI = GETSAFERHISHADER_PIXEL(*PixelShader);
-		SetGraphicsPipelineState(Context.RHICmdList, GraphicsPSOInit);
-		VertexShader->SetParameters(Context);
-		PixelShader->SetParameters(RHICmdList, Context, SSRQuality);
-
-		DrawPostProcessPass(
-			RHICmdList,
-			0, 0,
-			View.ViewRect.Width(), View.ViewRect.Height(),
-			View.ViewRect.Min.X, View.ViewRect.Min.Y,
-			View.ViewRect.Width(), View.ViewRect.Height(),
-			View.ViewRect.Size(),
-			FSceneRenderTargets::Get(Context.RHICmdList).GetBufferSizeXY(),
-			*VertexShader,
-			View.StereoPass,
-			Context.HasHmdMesh(),
-			EDRF_UseTriangleOptimization);
-	}
-	RHICmdList.EndRenderPass();
-	RHICmdList.CopyToResolveTarget(DestRenderTarget.TargetableTexture, DestRenderTarget.ShaderResourceTexture, FResolveParams());
-}
-
-FPooledRenderTargetDesc FRCPassPostProcessScreenSpaceReflections::ComputeOutputDesc(EPassOutputId InPassOutputId) const
-{
-	FPooledRenderTargetDesc Ret(FPooledRenderTargetDesc::Create2DDesc(FSceneRenderTargets::Get_FrameConstantsOnly().GetBufferSizeXY(), PF_FloatRGBA, FClearValueBinding::None, TexCreate_None, TexCreate_RenderTargetable, false));
-
-	Ret.ClearValue = FClearValueBinding(FLinearColor(0, 0, 0, 0));
-	Ret.DebugName = TEXT("ScreenSpaceReflections");
-	Ret.AutoWritable = false;
-	Ret.Flags |= GFastVRamConfig.SSR; 
-	return Ret;
-} 
-
-void RenderScreenSpaceReflections(FRHICommandListImmediate& RHICmdList, FViewInfo& View, TRefCountPtr<IPooledRenderTarget>& SSROutput, TRefCountPtr<IPooledRenderTarget>& VelocityRT)
-{
-	check(RHICmdList.IsOutsideRenderPass());
-
-	FRenderingCompositePassContext CompositeContext(RHICmdList, View);	
-	FPostprocessContext Context(RHICmdList, CompositeContext.Graph, View );
-
-	FSceneViewState* ViewState = (FSceneViewState*)Context.View.State;
-
-	FRenderingCompositePass* SceneColorInput = Context.Graph.RegisterPass( new FRCPassPostProcessInput( FSceneRenderTargets::Get(RHICmdList).GetSceneColor() ) );
-	FRenderingCompositePass* HZBInput = Context.Graph.RegisterPass( new FRCPassPostProcessInput( View.HZB ) );
-	FRenderingCompositePass* HCBInput = nullptr;
-
-	bool bSamplePrevFrame = false;
-	if (View.PrevViewInfo.CustomSSRInput.IsValid())
-	{
-		SceneColorInput = Context.Graph.RegisterPass(new FRCPassPostProcessInput(View.PrevViewInfo.CustomSSRInput));
-		bSamplePrevFrame = true;
-	}
-	else if (View.PrevViewInfo.TemporalAAHistory.IsValid())
-	{
-		SceneColorInput = Context.Graph.RegisterPass(new FRCPassPostProcessInput(View.PrevViewInfo.TemporalAAHistory.RT[0]));
-		bSamplePrevFrame = true;
-	}
-
-	FRenderingCompositeOutputRef VelocityInput;
-	if ( VelocityRT && !Context.View.bCameraCut )
-	{
-		VelocityInput = Context.Graph.RegisterPass(new FRCPassPostProcessInput(VelocityRT));
-	}
-	else
-	{
-		// No velocity, use black
-		VelocityInput = Context.Graph.RegisterPass(new FRCPassPostProcessInput(GSystemTextures.BlackDummy));
-	}
-	
-	if (CVarSSRCone.GetValueOnRenderThread())
-	{
-		HCBInput = Context.Graph.RegisterPass( new FRCPassPostProcessBuildHCB() );
-		HCBInput->SetInput( ePId_Input0, SceneColorInput );
-	}
-
-	{
-		FRenderingCompositePass* TracePass = Context.Graph.RegisterPass(
-			new FRCPassPostProcessScreenSpaceReflections( bSamplePrevFrame ) );
-		TracePass->SetInput( ePId_Input0, SceneColorInput );
-		TracePass->SetInput( ePId_Input1, HZBInput );
-		TracePass->SetInput( ePId_Input2, HCBInput );
-		TracePass->SetInput( ePId_Input3, VelocityInput );
-
-		Context.FinalOutput = FRenderingCompositeOutputRef( TracePass );
-	}
-
-	const bool bTemporalFilter = IsSSRTemporalPassRequired(View, false);
-
-	if( ViewState && bTemporalFilter )
-	{
+		TShaderMapRef<FScreenSpaceReflectionsPS> PixelShader(View.ShaderMap, PermutationVector);
+		ClearUnusedGraphResources(*PixelShader, PassParameters);
+		
+		GraphBuilder.AddPass(
+			RDG_EVENT_NAME("SSR RayMarch(Quality=%d) %dx%d",
+				SSRQuality, View.ViewRect.Width(), View.ViewRect.Height()),
+			PassParameters,
+			ERenderGraphPassFlags::None,
+			[PassParameters, &View, PixelShader, SSRStencilPrePass](FRHICommandList& RHICmdList)
 		{
-			FTAAPassParameters Parameters(Context.View);
-			Parameters.Pass = ETAAPassConfig::ScreenSpaceReflections;
-
-			FRenderingCompositePass* TemporalAAPass = Context.Graph.RegisterPass( new FRCPassPostProcessTemporalAA(
-				Context, Parameters,
-				ViewState->SSRHistory, &ViewState->SSRHistory) );
-			TemporalAAPass->SetInput( ePId_Input0, Context.FinalOutput );
-			TemporalAAPass->SetInput( ePId_Input2, VelocityInput );
-
-			Context.FinalOutput = FRenderingCompositeOutputRef( TemporalAAPass );
-		}
-
-		FRenderingCompositePass* HistoryOutput = Context.Graph.RegisterPass( new FRCPassPostProcessOutput( &ViewState->SSRHistory.RT[0] ) );
-		HistoryOutput->SetInput( ePId_Input0, Context.FinalOutput );
-
-		Context.FinalOutput = FRenderingCompositeOutputRef( HistoryOutput );
-	}
+			RHICmdList.SetViewport(View.ViewRect.Min.X, View.ViewRect.Min.Y, 0.0f, View.ViewRect.Max.X, View.ViewRect.Max.Y, 1.0f);
+			RHICmdList.SetStencilRef(0x80);
+		
+			FGraphicsPipelineStateInitializer GraphicsPSOInit;
+			FPixelShaderUtils::InitFullscreenPipelineState(RHICmdList, View.ShaderMap, *PixelShader, /* out */ GraphicsPSOInit);
+			if (SSRStencilPrePass)
+			{
+				// Clobers the stencil to pixel that should not compute SSR
+				GraphicsPSOInit.DepthStencilState = TStaticDepthStencilState<false, CF_Always, true, CF_Equal, SO_Keep, SO_Keep, SO_Keep>::GetRHI();
+			}
 
-	{
-		FRenderingCompositePass* ReflectionOutput = Context.Graph.RegisterPass( new FRCPassPostProcessOutput( &SSROutput ) );
-		ReflectionOutput->SetInput( ePId_Input0, Context.FinalOutput );
+			SetGraphicsPipelineState(RHICmdList, GraphicsPSOInit);
+			SetShaderParameters(RHICmdList, *PixelShader, PixelShader->GetPixelShader(), *PassParameters);
 
-		Context.FinalOutput = FRenderingCompositeOutputRef( ReflectionOutput );
+			FPixelShaderUtils::DrawFullscreenTriangle(RHICmdList);
+		});
 	}
 
-	CompositeContext.Process(Context.FinalOutput.GetPass(), TEXT("ReflectionEnvironments"));
+	return SSROutput;
 }
diff --git a/Engine/Source/Runtime/Renderer/Private/PostProcess/ScreenSpaceReflections.h b/Engine/Source/Runtime/Renderer/Private/PostProcess/ScreenSpaceReflections.h
--- a/Engine/Source/Runtime/Renderer/Private/PostProcess/ScreenSpaceReflections.h
+++ b/Engine/Source/Runtime/Renderer/Private/PostProcess/ScreenSpaceReflections.h
@@ -6,57 +6,17 @@
 
 #pragma once
 
-#include "CoreMinimal.h"
-#include "RendererInterface.h"
-#include "PostProcess/RenderingCompositionGraph.h"
+#include "RenderGraph.h"
 
 class FViewInfo;
+class FSceneViewFamilyBlackboard;
 
-// derives from TRenderingCompositePassBase<InputCount, OutputCount> 
-class FRCPassPostProcessDepthDownSample : public TRenderingCompositePassBase<1, 1>
-{
-public:
-	// interface FRenderingCompositePass ---------
-	virtual void Process(FRenderingCompositePassContext& Context) override;
-	virtual void Release() override { delete this; }
-	virtual FPooledRenderTargetDesc ComputeOutputDesc(EPassOutputId InPassOutputId) const override;
-};
-
-// ePId_Input0: scene color
-// ePId_Input1: scene depth
-// ePId_Input2: hierarchical scene color (optional)
-// ePId_Input3: velocity (optional)
-// derives from TRenderingCompositePassBase<InputCount, OutputCount> 
-class FRCPassPostProcessScreenSpaceReflections : public TRenderingCompositePassBase<4, 1>
-{
-public:
-	FRCPassPostProcessScreenSpaceReflections( bool InPrevFrame )
-	: bPrevFrame( InPrevFrame )
-	{}
-
-	// interface FRenderingCompositePass ---------
-	virtual void Process(FRenderingCompositePassContext& Context) override;
-	virtual void Release() override { delete this; }
-	virtual FPooledRenderTargetDesc ComputeOutputDesc(EPassOutputId InPassOutputId) const override;
-
-private:
-	bool bPrevFrame;
-};
-
-
-// ePId_Input0: half res scene color
-// derives from TRenderingCompositePassBase<InputCount, OutputCount> 
-class FRCPassPostProcessApplyScreenSpaceReflections : public TRenderingCompositePassBase<2, 1>
-{
-public:
-	// interface FRenderingCompositePass ---------
-	virtual void Process(FRenderingCompositePassContext& Context) override;
-	virtual void Release() override { delete this; }
-	virtual FPooledRenderTargetDesc ComputeOutputDesc(EPassOutputId InPassOutputId) const override;
-};
-
-void RenderScreenSpaceReflections(FRHICommandListImmediate& RHICmdList, FViewInfo& View, TRefCountPtr<IPooledRenderTarget>& SSROutput, TRefCountPtr<IPooledRenderTarget>& VelocityRT);
+FRDGTextureRef RenderScreenSpaceReflections(
+	FRDGBuilder& GraphBuilder,
+	const FSceneViewFamilyBlackboard& SceneBlackboard,
+	const FRDGTextureRef CurrentSceneColor,
+	const FViewInfo& View);
 
 bool ShouldRenderScreenSpaceReflections(const FViewInfo& View);
 
-bool IsSSRTemporalPassRequired(const FViewInfo& View, bool bCheckSSREnabled = true);
\ No newline at end of file
+bool IsSSRTemporalPassRequired(const FViewInfo& View, bool bCheckSSREnabled = true);
diff --git a/Engine/Source/Runtime/Renderer/Private/ReflectionEnvironment.cpp b/Engine/Source/Runtime/Renderer/Private/ReflectionEnvironment.cpp
--- a/Engine/Source/Runtime/Renderer/Private/ReflectionEnvironment.cpp
+++ b/Engine/Source/Runtime/Renderer/Private/ReflectionEnvironment.cpp
@@ -24,6 +24,7 @@
 #include "PostProcess/SceneFilterRendering.h"
 #include "PostProcess/PostProcessing.h"
 #include "PostProcess/PostProcessSubsurface.h"
+#include "PostProcess/PostProcessTemporalAA.h"
 #include "PostProcess/ScreenSpaceReflections.h"
 #include "LightRendering.h"
 #include "LightPropagationVolumeSettings.h"
@@ -827,7 +828,35 @@ void FDeferredShadingSceneRenderer::RenderDeferredReflectionsAndSkyLighting(FRHI
 		}
 		else if (bScreenSpaceReflections)
 		{
-			RenderScreenSpaceReflections(RHICmdList, View, ReflectionsColor, VelocityRT);
+			FRDGBuilder GraphBuilder(RHICmdList); // TODO: convert the entire reflections to render graph.
+			
+			FSceneViewFamilyBlackboard SceneBlackboard;
+			SetupSceneViewFamilyBlackboard(GraphBuilder, &SceneBlackboard);
+
+			FRDGTextureRef CurrentSceneColor = GraphBuilder.RegisterExternalTexture(SceneContext.GetSceneColor());
+
+			FRDGTextureRef ScreenSpaceReflections = RenderScreenSpaceReflections(
+				GraphBuilder, SceneBlackboard, CurrentSceneColor, View);
+			
+			// Adds a temporal AA pass to clean up some of the noise.
+			const bool bTemporalFilter = IsSSRTemporalPassRequired(View, false);
+	
+			// TODO: denoise.
+			if (View.ViewState && bTemporalFilter)
+			{
+				FTAAPassParameters TAASettings(View);
+				TAASettings.Pass = ETAAPassConfig::ScreenSpaceReflections;
+				TAASettings.SceneColorInput = ScreenSpaceReflections;
+
+				ScreenSpaceReflections = TAASettings.AddTemporalAAPass(
+					GraphBuilder,
+					SceneBlackboard, View,
+					View.PrevViewInfo.SSRHistory,
+					&View.ViewState->PrevFrameViewInfo.SSRHistory).SceneColor;
+			}
+
+			GraphBuilder.QueueTextureExtraction(ScreenSpaceReflections, &ReflectionsColor);
+			GraphBuilder.Execute();
 		}
 
 		bool bPlanarReflections = false;
diff --git a/Engine/Source/Runtime/Renderer/Private/ScenePrivate.h b/Engine/Source/Runtime/Renderer/Private/ScenePrivate.h
--- a/Engine/Source/Runtime/Renderer/Private/ScenePrivate.h
+++ b/Engine/Source/Runtime/Renderer/Private/ScenePrivate.h
@@ -853,8 +853,6 @@ public:
 
 	FHeightfieldLightingAtlas* HeightfieldLightingAtlas;
 
-	// Temporal AA result for SSR
-	FTemporalAAHistory SSRHistory;
 	// Temporal AA result for light shafts of last frame
 	FTemporalAAHistory LightShaftOcclusionHistory;
 	// Temporal AA result for light shafts of last frame
@@ -1257,7 +1255,6 @@ public:
 		EyeAdaptationRTManager.SafeRelease();
 		CombinedLUTRenderTarget.SafeRelease();
 		PrevFrameViewInfo.SafeRelease();
-		SSRHistory.SafeRelease();
 		LightShaftOcclusionHistory.SafeRelease();
 		LightShaftBloomHistoryRTs.Empty();
 		DistanceFieldAOHistoryRT.SafeRelease();
diff --git a/Engine/Source/Runtime/Renderer/Private/SceneRendering.h b/Engine/Source/Runtime/Renderer/Private/SceneRendering.h
--- a/Engine/Source/Runtime/Renderer/Private/SceneRendering.h
+++ b/Engine/Source/Runtime/Renderer/Private/SceneRendering.h
@@ -654,6 +654,9 @@ struct FPreviousViewInfo
 
 	// Temporal AA history for diaphragm DOF.
 	FTemporalAAHistory DOFSetupHistory;
+	
+	// Temporal AA history for SSR
+	FTemporalAAHistory SSRHistory;
 
 	// Scene color input for SSR, that can be different from TemporalAAHistory.RT[0] if there is a SSR
 	// input post process material.
@@ -683,6 +686,7 @@ struct FPreviousViewInfo
 		GBufferC.SafeRelease();
 		TemporalAAHistory.SafeRelease();
 		DOFSetupHistory.SafeRelease();
+		SSRHistory.SafeRelease();
 		CustomSSRInput.SafeRelease();
 		ReflectionsHistory.SafeRelease();
 		AmbientOcclusionHistory.SafeRelease();
diff --git a/Engine/Source/Runtime/Renderer/Private/SceneViewFamilyBlackboard.cpp b/Engine/Source/Runtime/Renderer/Private/SceneViewFamilyBlackboard.cpp
--- a/Engine/Source/Runtime/Renderer/Private/SceneViewFamilyBlackboard.cpp
+++ b/Engine/Source/Runtime/Renderer/Private/SceneViewFamilyBlackboard.cpp
@@ -37,6 +37,18 @@ void SetupSceneViewFamilyBlackboard(
 	}
 }
 
+void SetupSceneTextureSamplers(FSceneTextureSamplerParameters* OutSamplers)
+{
+	FSamplerStateRHIParamRef Sampler = TStaticSamplerState<SF_Point>::GetRHI();
+	OutSamplers->SceneDepthBufferSampler = Sampler;
+	OutSamplers->SceneVelocityBufferSampler = Sampler;
+	OutSamplers->SceneGBufferASampler = Sampler;
+	OutSamplers->SceneGBufferBSampler = Sampler;
+	OutSamplers->SceneGBufferCSampler = Sampler;
+	OutSamplers->SceneGBufferDSampler = Sampler;
+	OutSamplers->SceneGBufferESampler = Sampler;
+}
+
 FRDGTextureRef GetEyeAdaptationTexture(FRDGBuilder& GraphBuilder, const FViewInfo& View)
 {
 	if (View.HasValidEyeAdaptation())
diff --git a/Engine/Source/Runtime/Renderer/Private/SceneViewFamilyBlackboard.h b/Engine/Source/Runtime/Renderer/Private/SceneViewFamilyBlackboard.h
--- a/Engine/Source/Runtime/Renderer/Private/SceneViewFamilyBlackboard.h
+++ b/Engine/Source/Runtime/Renderer/Private/SceneViewFamilyBlackboard.h
@@ -13,6 +13,7 @@ class FViewInfo;
 
 
 /** Contains reference on all available buffer for a given scene. */
+// TODO: rename to FSceneTextureParemeters
 BEGIN_SHADER_PARAMETER_STRUCT(FSceneViewFamilyBlackboard, )
 	// FSceneViewFamilyBlackboard::SceneLightingChannels needs to be accessed with SceneLightingChannels.Load(), so a shader accessing
 	// needs to know when it not valid since SceneLightingChannels could end up being a dummy system texture.
@@ -29,6 +30,18 @@ BEGIN_SHADER_PARAMETER_STRUCT(FSceneViewFamilyBlackboard, )
 END_SHADER_PARAMETER_STRUCT()
 
 
+/** Contains reference on all samplers for FSceneViewFamilyBlackboard for platforms not supporting shared samplers. */
+BEGIN_SHADER_PARAMETER_STRUCT(FSceneTextureSamplerParameters, )
+	SHADER_PARAMETER_SAMPLER(SamplerState, SceneDepthBufferSampler)
+	SHADER_PARAMETER_SAMPLER(SamplerState, SceneVelocityBufferSampler)
+	SHADER_PARAMETER_SAMPLER(SamplerState, SceneGBufferASampler)
+	SHADER_PARAMETER_SAMPLER(SamplerState, SceneGBufferBSampler)
+	SHADER_PARAMETER_SAMPLER(SamplerState, SceneGBufferCSampler)
+	SHADER_PARAMETER_SAMPLER(SamplerState, SceneGBufferDSampler)
+	SHADER_PARAMETER_SAMPLER(SamplerState, SceneGBufferESampler)
+END_SHADER_PARAMETER_STRUCT()
+
+
 /** Sets up the the blackboard from available scene view family.
  *
  * Note: Once the entire renderer is built with a single render graph, would no longer need this function.
@@ -37,6 +50,9 @@ void SetupSceneViewFamilyBlackboard(
 	FRDGBuilder& GraphBuilder,
 	FSceneViewFamilyBlackboard* OutBlackboard);
 
+/** Sets up all the samplers. */
+void SetupSceneTextureSamplers(FSceneTextureSamplerParameters* OutSamplers);
+
 /** Returns a render graph texture resource reference onto the eye adaptation or fallback.
  */
 FRDGTextureRef GetEyeAdaptationTexture(FRDGBuilder& GraphBuilder, const FViewInfo& View);
