Fixing Engine.h include warning.

diff --git a/Engine/Source/Runtime/Renderer/Private/BlueNoise.h b/Engine/Source/Runtime/Renderer/Private/BlueNoise.h
--- a/Engine/Source/Runtime/Renderer/Private/BlueNoise.h
+++ b/Engine/Source/Runtime/Renderer/Private/BlueNoise.h
@@ -6,7 +6,6 @@
 
 #pragma once
 
-#include "Engine.h"
 #include "UniformBuffer.h"
 
 // Texture data is assumed to be in tiled representation where:
@@ -19,6 +18,7 @@ END_GLOBAL_SHADER_PARAMETER_STRUCT()
 
 inline void InitializeBlueNoise(FBlueNoise& BlueNoise)
 {
+	check(GEngine);
 	BlueNoise.Dimensions = FIntVector(GEngine->BlueNoiseTexture->GetSizeX(), GEngine->BlueNoiseTexture->GetSizeX(), GEngine->BlueNoiseTexture->GetSizeY() / GEngine->BlueNoiseTexture->GetSizeX());
 	BlueNoise.Texture = GEngine->BlueNoiseTexture->Resource->TextureRHI;
 }
