Fix compile error..

diff --git a/Engine/Source/Runtime/Renderer/Private/BlueNoise.h b/Engine/Source/Runtime/Renderer/Private/BlueNoise.h
--- a/Engine/Source/Runtime/Renderer/Private/BlueNoise.h
+++ b/Engine/Source/Runtime/Renderer/Private/BlueNoise.h
@@ -7,6 +7,7 @@
 #pragma once
 
 #include "UniformBuffer.h"
+#include "Engine/Engine.h"
 
 // Texture data is assumed to be in tiled representation where:
 // 1) Dimensions.xy represents a single blue-noise tile
