Fixing more compile errors..

diff --git a/Engine/Source/Runtime/Renderer/Private/BlueNoise.h b/Engine/Source/Runtime/Renderer/Private/BlueNoise.h
--- a/Engine/Source/Runtime/Renderer/Private/BlueNoise.h
+++ b/Engine/Source/Runtime/Renderer/Private/BlueNoise.h
@@ -8,6 +8,8 @@
 
 #include "UniformBuffer.h"
 #include "Engine/Engine.h"
+#include "Engine/Texture2D.h"
+#include "Math/IntVector.h"
 
 // Texture data is assumed to be in tiled representation where:
 // 1) Dimensions.xy represents a single blue-noise tile
