Fixes SSR's gpu stat.

diff --git a/Engine/Source/Runtime/Renderer/Private/PostProcess/ScreenSpaceReflections.cpp b/Engine/Source/Runtime/Renderer/Private/PostProcess/ScreenSpaceReflections.cpp
--- a/Engine/Source/Runtime/Renderer/Private/PostProcess/ScreenSpaceReflections.cpp
+++ b/Engine/Source/Runtime/Renderer/Private/PostProcess/ScreenSpaceReflections.cpp
@@ -302,6 +302,7 @@ FRDGTextureRef RenderScreenSpaceReflections(
 			ERenderGraphPassFlags::None,
 			[PassParameters, &View, PixelShader](FRHICommandList& RHICmdList)
 		{
+			SCOPED_GPU_STAT(RHICmdList, ScreenSpaceReflections);
 			RHICmdList.SetViewport(View.ViewRect.Min.X, View.ViewRect.Min.Y, 0.0f, View.ViewRect.Max.X, View.ViewRect.Max.Y, 1.0f);
 			RHICmdList.SetStencilRef(0x80);
 		
@@ -375,6 +376,7 @@ FRDGTextureRef RenderScreenSpaceReflections(
 			ERenderGraphPassFlags::None,
 			[PassParameters, &View, PixelShader, SSRStencilPrePass](FRHICommandList& RHICmdList)
 		{
+			SCOPED_GPU_STAT(RHICmdList, ScreenSpaceReflections);
 			RHICmdList.SetViewport(View.ViewRect.Min.X, View.ViewRect.Min.Y, 0.0f, View.ViewRect.Max.X, View.ViewRect.Max.Y, 1.0f);
 			RHICmdList.SetStencilRef(0x80);
 		
