Cleanup / Refactor of ScreenPass implementation, and bug fixes.

 - Fixed regression with checkerboard rendering and odd render target sizes. The passes were not guaranteed to produce
   a UV coordinate centered on an input pixel for each output pixel.

 - Fixed regression when visualizing SSS at native screen resolution (i.e. no secondary spatial upsampling). The issue was that
  the rendering composition graph injects the backbuffer as the last texture in the post processing pipeline. If visualize SSS was
  the last pass, it would overwrite the backbuffer texture, causing visual artifacts.

 - Minor performance improvement for multi-view SSS. All views in a split-screen scenario now render into the same output target (instead
   of ping ponging between multiple targets) and only views not utilizing SSS get copied from the source target. The is reduces the worst case
   of viewport copied from O(N^2) to O(N); if all viewports use SSS (common in VR), then no copies are performed.

 - Removed FScreenPassTexture in favor of more selective pairing of FRDGTextureRef and FScreenPassTextureViewport. This reduces
   some boilerplate and extraneous memory copying. Motion blur passes now compute their viewport shader parameters once.

diff --git a/Engine/Shaders/Private/PostProcessCommon.ush b/Engine/Shaders/Private/PostProcessCommon.ush
--- a/Engine/Shaders/Private/PostProcessCommon.ush
+++ b/Engine/Shaders/Private/PostProcessCommon.ush
@@ -74,7 +74,6 @@ SamplerState BilinearTextureSampler;
 
 // transforms ScreenPos to the screen pixel position, not the viewport local position (*xy +wz)
 // (0,0) left top, (width-1,height-1) at right bottom, in pixels (without offset to reach the pixel center)
-// 	use ComputePixelPosCenter(UVAndScreenPos.zw, bPixelCenter) instead of directly accessing this
 float4 ScreenPosToPixel;
 
 float4 SceneColorBufferUVViewport;
@@ -86,24 +85,6 @@ float2 ViewportUVToPostProcessingSceneColorBufferUV(float2 ViewportUV)
 	return ViewportUV * SceneColorBufferUVViewport.xy + SceneColorBufferUVViewport.zw;
 }
 
-// uses ScreenPosToPixel which is setup in FPostProcessPassParameters
-// for integer output bPixelCenter should be set to true (to avoid rounding errors), but often it's better to use "int2 PixelPos = (int2)SvPosition.xy;" instead
-// @param ScreenPos -1 .. 1, usually from UVAndScreenPos.zw
-// @param bCenter false: left top of the pixel and almost no fractional value, true: added 0.5 to be pixel center
-float2 ComputePixelPosCenter(float2 ScreenPos, bool bPixelCenter)
-{
-	float2 PixelOffset = bPixelCenter ? 0.5f : 0.0f;
-	return ScreenPos * ScreenPosToPixel.xy + ScreenPosToPixel.zw + PixelOffset;
-}
-
-// not well optimized but can be useful
-float2 ScreenPosFromPixelPos(float2 PixelPos, bool bPixelCenter)
-{
-	float2 PixelOffset = bPixelCenter ? 0.5f : 0.0f;
-
-	return (PixelPos - PixelOffset - ScreenPosToPixel.zw) / ScreenPosToPixel.xy;
-}
-
 // shape of the lens for vignette like effects and masking
 // @param ScreenPos -1 .. 1
 // @return 0..1 0:borders, 1:center
diff --git a/Engine/Shaders/Private/PostProcessVelocityFlatten.usf b/Engine/Shaders/Private/PostProcessVelocityFlatten.usf
--- a/Engine/Shaders/Private/PostProcessVelocityFlatten.usf
+++ b/Engine/Shaders/Private/PostProcessVelocityFlatten.usf
@@ -21,7 +21,6 @@ Texture2D DepthTexture;
 
 RWTexture2D<float4>	OutVelocityFlatTexture;
 RWTexture2D<float4>	OutVelocityTileTexture;
-RWTexture2D<float4>	OutDilatedVelocityTileTexture;
 
 SCREEN_PASS_TEXTURE_VIEWPORT(Velocity)
 SCREEN_PASS_TEXTURE_VIEWPORT(VelocityTile)
@@ -212,7 +211,7 @@ void VelocityGatherCS(
 	BRANCH
 	if (all(PixelPos.xy < VelocityTile_ViewportMax))
 	{
-		OutDilatedVelocityTileTexture[PixelPos] = MinMaxVelocity;
+		OutVelocityTileTexture[PixelPos] = MinMaxVelocity;
 	}
 }
 
diff --git a/Engine/Source/Runtime/RenderCore/Private/RenderGraphBuilder.cpp b/Engine/Source/Runtime/RenderCore/Private/RenderGraphBuilder.cpp
--- a/Engine/Source/Runtime/RenderCore/Private/RenderGraphBuilder.cpp
+++ b/Engine/Source/Runtime/RenderCore/Private/RenderGraphBuilder.cpp
@@ -1485,7 +1485,7 @@ FRDGBuilder::~FRDGBuilder()
 
 	#if RENDER_GRAPH_DEBUGGING
 	{
-		checkf(bHasExecuted, TEXT("Render graph execution si required to ensure consistency with immediate mode."));
+		checkf(bHasExecuted, TEXT("Render graph execution is required to ensure consistency with immediate mode."));
 	}
 	#endif
 }
diff --git a/Engine/Source/Runtime/RenderCore/Public/RenderGraphBuilder.h b/Engine/Source/Runtime/RenderCore/Public/RenderGraphBuilder.h
--- a/Engine/Source/Runtime/RenderCore/Public/RenderGraphBuilder.h
+++ b/Engine/Source/Runtime/RenderCore/Public/RenderGraphBuilder.h
@@ -165,7 +165,8 @@ struct RENDERCORE_API FRenderGraphPass
 
 	virtual void Execute(FRHICommandListImmediate& RHICmdList) const = 0;
 
-	const TCHAR* GetName() const {
+	const TCHAR* GetName() const
+	{
 		return Name.GetTCHAR();
 	}
 
@@ -174,7 +175,8 @@ struct RENDERCORE_API FRenderGraphPass
 		return PassFlags;
 	}
 
-	bool IsCompute() const {
+	bool IsCompute() const
+	{
 		return (PassFlags & ERenderGraphPassFlags::Compute) == ERenderGraphPassFlags::Compute;
 	}
 	
@@ -466,7 +468,7 @@ public:
 			 */
 			checkf(
 				AllocatedUnusedPassParameters.Contains(static_cast<void *>(ParameterStruct)),
-				TEXT("The pass parameter structure has not been alloctaed for correct life time FRDGBuilder::AllocParameters() or has already ")
+				TEXT("The pass parameter structure has not been allocated for correct life time FRDGBuilder::AllocParameters() or has already ")
 				TEXT("been used by another previous FRDGBuilder::AddPass()."));
 
 			AllocatedUnusedPassParameters.Remove(static_cast<void *>(ParameterStruct));
diff --git a/Engine/Source/Runtime/Renderer/Private/DeferredShadingRenderer.cpp b/Engine/Source/Runtime/Renderer/Private/DeferredShadingRenderer.cpp
--- a/Engine/Source/Runtime/Renderer/Private/DeferredShadingRenderer.cpp
+++ b/Engine/Source/Runtime/Renderer/Private/DeferredShadingRenderer.cpp
@@ -1757,44 +1757,7 @@ void FDeferredShadingSceneRenderer::Render(FRHICommandListImmediate& RHICmdList)
 
 		ServiceLocalQueue();
 
-		{
-			FSceneRenderTargets& SceneRenderTargets = FSceneRenderTargets::Get(RHICmdList);
-
-			FRDGBuilder GraphBuilder(RHICmdList);
-
-			FRDGTextureRef SceneColorTexture = GraphBuilder.RegisterExternalTexture(SceneRenderTargets.GetSceneColor(), TEXT("SceneColor"));
-
-			// Post-lighting composition lighting stage
-			// e.g. ScreenSpaceSubsurfaceScattering
-			for (int32 ViewIndex = 0; ViewIndex < Views.Num(); ++ViewIndex)
-			{
-				RDG_EVENT_SCOPE(GraphBuilder, "SubsurfaceScattering(ViewId=%d)", ViewIndex);
-
-				const FViewInfo& View = Views[ViewIndex];
-
-				if (IsSubsurfaceRequiredForView(View))
-				{
-					FScreenPassContext* Context = GraphBuilder.AllocObject<FScreenPassContext>(RHICmdList, View);
-
-					const FScreenPassTexture Input = FScreenPassTexture::Create(SceneColorTexture, View.ViewRect);
-
-					const FScreenPassTexture Output = ComputeSubsurface(GraphBuilder, Context, Input);
-
-					SceneColorTexture = Output.GetRDGTexture();
-				}
-			}
-
-			// Extract the result texture out and re-assign it to the scene render targets blackboard.
-			TRefCountPtr<IPooledRenderTarget> SceneColorTarget;
-			GraphBuilder.QueueTextureExtraction(SceneColorTexture, &SceneColorTarget, false);
-			GraphBuilder.Execute();
-
-			SceneRenderTargets.SetSceneColor(SceneColorTarget);
-
-			// The RT should be released as early as possible to allow sharing of that memory for other purposes.
-			// This becomes even more important with some limited VRam (XBoxOne).
-			SceneRenderTargets.SetLightAttenuation(nullptr);
-		}
+		ComputeSubsurfaceShim(RHICmdList, Views);
 
 #if RHI_RAYTRACING
 		if (SkyLightRT)
diff --git a/Engine/Source/Runtime/Renderer/Private/PostProcess/PostProcessMotionBlur.cpp b/Engine/Source/Runtime/Renderer/Private/PostProcess/PostProcessMotionBlur.cpp
--- a/Engine/Source/Runtime/Renderer/Private/PostProcess/PostProcessMotionBlur.cpp
+++ b/Engine/Source/Runtime/Renderer/Private/PostProcess/PostProcessMotionBlur.cpp
@@ -109,10 +109,9 @@ bool IsMotionBlurEnabled(const FViewInfo& View)
 		&& (CVarAllowMotionBlurInVR->GetInt() != 0 || !(ViewFamily.Views.Num() > 1));
 }
 
-bool IsMotionBlurScatterRequired(FScreenPassContextRef Context, const FScreenPassTextureViewport& SceneViewport)
+bool IsMotionBlurScatterRequired(const FViewInfo& View, const FScreenPassTextureViewport& SceneViewport)
 {
-	const FViewInfo& View = Context->View;
-	const FSceneViewState* ViewState = Context->ViewState;
+	const FSceneViewState* ViewState = View.ViewState;
 	const float SceneExtentX = SceneViewport.Extent.X;
 
 	// Normalize percentage value.
@@ -182,10 +181,9 @@ BEGIN_SHADER_PARAMETER_STRUCT(FMotionBlurParameters, )
 	SHADER_PARAMETER(float, VelocityMax)
 END_SHADER_PARAMETER_STRUCT()
 
-FMotionBlurParameters GetMotionBlurParameters(FScreenPassContextRef Context, FIntPoint SceneViewportExtent, float BlurScale)
+FMotionBlurParameters GetMotionBlurParameters(const FViewInfo& View, FIntPoint SceneViewportExtent, float BlurScale)
 {
-	const FViewInfo& View = Context->View;
-	const FSceneViewState* ViewState = Context->ViewState;
+	const FSceneViewState* ViewState = View.ViewState;
 
 	const float TileSize = kMotionBlurTileSize;
 	const float SceneViewportExtentX = SceneViewportExtent.X;
@@ -234,8 +232,8 @@ public:
 		SHADER_PARAMETER_STRUCT(FScreenPassTextureViewportParameters, Velocity)
 		SHADER_PARAMETER_RDG_TEXTURE(Texture2D, VelocityTexture)
 		SHADER_PARAMETER_RDG_TEXTURE(Texture2D, DepthTexture)
-		SHADER_PARAMETER_RDG_TEXTURE_UAV(RWTexture2D, OutVelocityFlatTexture)
-		SHADER_PARAMETER_RDG_TEXTURE_UAV(RWTexture2D, OutVelocityTileTexture)
+		SHADER_PARAMETER_RDG_TEXTURE_UAV(RWTexture2D<float4>, OutVelocityFlatTexture)
+		SHADER_PARAMETER_RDG_TEXTURE_UAV(RWTexture2D<float4>, OutVelocityTileTexture)
 	END_SHADER_PARAMETER_STRUCT()
 };
 
@@ -255,7 +253,7 @@ public:
 
 	BEGIN_SHADER_PARAMETER_STRUCT(FParameters, )
 		SHADER_PARAMETER_STRUCT_INCLUDE(FMotionBlurVelocityDilateParameters, Dilate)
-		SHADER_PARAMETER_RDG_TEXTURE_UAV(RWTexture2D<float4>, OutDilatedVelocityTileTexture)
+		SHADER_PARAMETER_RDG_TEXTURE_UAV(RWTexture2D<float4>, OutVelocityTileTexture)
 	END_SHADER_PARAMETER_STRUCT()
 };
 
@@ -296,6 +294,8 @@ IMPLEMENT_GLOBAL_SHADER(FMotionBlurVelocityDilateScatterPS, "/Engine/Private/Pos
 
 class FMotionBlurQualityDimension : SHADER_PERMUTATION_ENUM_CLASS("MOTION_BLUR_QUALITY", EMotionBlurQuality);
 
+using FMotionBlurFilterPermutationDomain = TShaderPermutationDomain<FMotionBlurQualityDimension>;
+
 BEGIN_SHADER_PARAMETER_STRUCT(FMotionBlurFilterParameters, )
 	SHADER_PARAMETER_STRUCT(FMotionBlurParameters, MotionBlur)
 
@@ -314,7 +314,7 @@ BEGIN_SHADER_PARAMETER_STRUCT(FMotionBlurFilterParameters, )
 	SHADER_PARAMETER_SAMPLER(SamplerState, VelocitySampler)
 	SHADER_PARAMETER_SAMPLER(SamplerState, VelocityTileSampler)
 	SHADER_PARAMETER_SAMPLER(SamplerState, VelocityFlatSampler)
-	END_SHADER_PARAMETER_STRUCT()
+END_SHADER_PARAMETER_STRUCT()
 
 class FMotionBlurFilterPS : public FMotionBlurShader
 {
@@ -327,8 +327,7 @@ public:
 		RENDER_TARGET_BINDING_SLOTS()
 	END_SHADER_PARAMETER_STRUCT()
 
-	using FQualityDimension = FMotionBlurQualityDimension;
-	using FPermutationDomain = TShaderPermutationDomain<FQualityDimension>;
+	using FPermutationDomain = FMotionBlurFilterPermutationDomain;
 };
 
 IMPLEMENT_GLOBAL_SHADER(FMotionBlurFilterPS, "/Engine/Private/PostProcessMotionBlur.usf", "MainPS", SF_Pixel);
@@ -351,8 +350,7 @@ public:
 		OutEnvironment.SetDefine(TEXT("THREADGROUP_SIZEY"), kMotionBlurComputeTileSizeY);
 	}
 
-	using FQualityDimension = FMotionBlurQualityDimension;
-	using FPermutationDomain = TShaderPermutationDomain<FQualityDimension>;
+	using FPermutationDomain = FMotionBlurFilterPermutationDomain;
 };
 
 IMPLEMENT_GLOBAL_SHADER(FMotionBlurFilterCS, "/Engine/Private/PostProcessMotionBlur.usf", "MainCS", SF_Compute);
@@ -400,116 +398,138 @@ enum class EMotionBlurFilterPass : uint32
 	MAX
 };
 
-void ComputeMotionBlurVelocity(
-	FRDGBuilder& GraphBuilder,
-	FScreenPassContextRef Context,
-	const FScreenPassTexture& ColorTexture,
-	const FScreenPassTexture& DepthTexture,
-	const FScreenPassTexture& VelocityTexture,
-	FScreenPassTexture& OutVelocityFlatTexture,
-	FScreenPassTexture& OutVelocityTileTexture)
+struct FMotionBlurViewports
 {
-	// Validate inputs.
-	check(Context);
-	check(ColorTexture.IsValid());
-	check(DepthTexture.IsValid());
-	check(VelocityTexture.IsValid());
-	check(VelocityTexture.GetViewport() == DepthTexture.GetViewport());
-
-	const FScreenPassTextureViewport& ColorViewport = ColorTexture.GetViewport();
-	const FScreenPassTextureViewport& VelocityViewport = VelocityTexture.GetViewport();
-
-	OutVelocityFlatTexture =
-		FScreenPassTexture::Create(
-			GraphBuilder.CreateTexture(
-				FRDGTextureDesc::Create2DDesc(
-					VelocityViewport.Extent,
-					PF_FloatR11G11B10,
-					FClearValueBinding::None,
-					GFastVRamConfig.VelocityFlat,
-					TexCreate_ShaderResource | TexCreate_UAV,
-					false),
-				TEXT("VelocityFlat")),
-			VelocityViewport.Rect);
+	FMotionBlurViewports(
+		const FScreenPassTextureViewport& InColorViewport,
+		const FScreenPassTextureViewport& InVelocityViewport)
+	{
+		Color = InColorViewport;
+		Velocity = InVelocityViewport;
+		VelocityTile = FScreenPassTextureViewport(
+			FIntRect(
+				FIntPoint::ZeroValue,
+				GetMotionBlurTileCount(Velocity.Rect.Size())));
+
+		ColorParameters = GetScreenPassTextureViewportParameters(Color);
+		VelocityParameters = GetScreenPassTextureViewportParameters(Velocity);
+		VelocityTileParameters = GetScreenPassTextureViewportParameters(VelocityTile);
+
+		ColorToVelocityTransform = GetScreenPassTextureViewportTransform(ColorParameters, VelocityParameters);
+		ColorToVelocityTileTransform = GetScreenPassTextureViewportTransform(ColorParameters, VelocityTileParameters);
+	}
 
-	const FIntPoint VelocityTileCount = GetMotionBlurTileCount(VelocityViewport.Rect.Size());
+	FScreenPassTextureViewport Color;
+	FScreenPassTextureViewport Velocity;
+	FScreenPassTextureViewport VelocityTile;
 
-	const FScreenPassTexture VelocityTileTexture =
-		FScreenPassTexture::Create(
-			GraphBuilder.CreateTexture(
-				FRDGTextureDesc::Create2DDesc(
-					VelocityTileCount,
-					PF_FloatRGBA,
-					FClearValueBinding::None,
-					GFastVRamConfig.VelocityMax,
-					TexCreate_ShaderResource | TexCreate_UAV,
-					false),
-				TEXT("VelocityTile")),
-			FIntRect(FIntPoint::ZeroValue, VelocityTileCount));
+	FScreenPassTextureViewportParameters ColorParameters;
+	FScreenPassTextureViewportParameters VelocityParameters;
+	FScreenPassTextureViewportParameters VelocityTileParameters;
 
-	const FScreenPassTextureViewport& VelocityTileViewport = VelocityTileTexture.GetViewport();
+	FScreenPassTextureViewportTransform ColorToVelocityTransform;
+	FScreenPassTextureViewportTransform ColorToVelocityTileTransform;
+};
 
-	const FMotionBlurParameters MotionBlurParametersNoScale = GetMotionBlurParameters(Context, ColorViewport.Extent, 1.0f);
+void ComputeMotionBlurVelocity(
+	FRDGBuilder& GraphBuilder,
+	const FScreenPassViewInfo& ScreenPassView,
+	const FMotionBlurViewports& Viewports,
+	FRDGTextureRef ColorTexture,
+	FRDGTextureRef DepthTexture,
+	FRDGTextureRef VelocityTexture,
+	FRDGTextureRef* VelocityFlatTextureOutput,
+	FRDGTextureRef* VelocityTileTextureOutput)
+{
+	check(ColorTexture);
+	check(DepthTexture);
+	check(VelocityTexture);
+	check(VelocityFlatTextureOutput);
+	check(VelocityTileTextureOutput);
+
+	const FViewInfo& View = ScreenPassView.View;
+
+	const FIntPoint VelocityTileCount = Viewports.VelocityTile.Extent;
+
+	FRDGTextureRef VelocityFlatTexture =
+		GraphBuilder.CreateTexture(
+			FRDGTextureDesc::Create2DDesc(
+				Viewports.Velocity.Extent,
+				PF_FloatR11G11B10,
+				FClearValueBinding::None,
+				GFastVRamConfig.VelocityFlat,
+				TexCreate_ShaderResource | TexCreate_UAV,
+				false),
+			TEXT("VelocityFlat"));
+
+	FRDGTextureRef VelocityTileTextureSetup =
+		GraphBuilder.CreateTexture(
+			FRDGTextureDesc::Create2DDesc(
+				VelocityTileCount,
+				PF_FloatRGBA,
+				FClearValueBinding::None,
+				GFastVRamConfig.VelocityMax,
+				TexCreate_ShaderResource | TexCreate_UAV,
+				false),
+			TEXT("VelocityTile"));
+
+	const FMotionBlurParameters MotionBlurParametersNoScale = GetMotionBlurParameters(View, Viewports.Color.Extent, 1.0f);
 
 	// Velocity flatten pass: combines depth / velocity into a single target for sampling efficiency.
 	{
 		FMotionBlurVelocityFlattenCS::FParameters* PassParameters = GraphBuilder.AllocParameters<FMotionBlurVelocityFlattenCS::FParameters>();
 		PassParameters->MotionBlur = MotionBlurParametersNoScale;
-		PassParameters->View = Context->View.ViewUniformBuffer;
-		PassParameters->Velocity = VelocityTexture.GetViewportParameters();
-		PassParameters->DepthTexture = DepthTexture.GetRDGTexture();
-		PassParameters->VelocityTexture = VelocityTexture.GetRDGTexture();
-		PassParameters->OutVelocityFlatTexture = GraphBuilder.CreateUAV(OutVelocityFlatTexture.GetRDGTexture());
-		PassParameters->OutVelocityTileTexture = GraphBuilder.CreateUAV(VelocityTileTexture.GetRDGTexture());
-
-		TShaderMapRef<FMotionBlurVelocityFlattenCS> ComputeShader(Context->ShaderMap);
+		PassParameters->View = View.ViewUniformBuffer;
+		PassParameters->Velocity = Viewports.VelocityParameters;
+		PassParameters->DepthTexture = DepthTexture;
+		PassParameters->VelocityTexture = VelocityTexture;
+		PassParameters->OutVelocityFlatTexture = GraphBuilder.CreateUAV(VelocityFlatTexture);
+		PassParameters->OutVelocityTileTexture = GraphBuilder.CreateUAV(VelocityTileTextureSetup);
+
+		TShaderMapRef<FMotionBlurVelocityFlattenCS> ComputeShader(View.ShaderMap);
 		FComputeShaderUtils::AddPass(
 			GraphBuilder,
 			RDG_EVENT_NAME("Velocity Flatten"),
 			*ComputeShader,
 			PassParameters,
-			FComputeShaderUtils::GetGroupCount(VelocityViewport.Extent, FComputeShaderUtils::kGolden2DGroupSize));
+			FComputeShaderUtils::GetGroupCount(Viewports.Velocity.Extent, FComputeShaderUtils::kGolden2DGroupSize));
 	}
 
-	OutVelocityTileTexture =
-		FScreenPassTexture::Create(
-			GraphBuilder.CreateTexture(
-				FRDGTextureDesc::Create2DDesc(
-					VelocityTileViewport.Extent,
-					PF_FloatRGBA,
-					FClearValueBinding::None,
-					GFastVRamConfig.MotionBlur,
-					TexCreate_ShaderResource | TexCreate_UAV | TexCreate_RenderTargetable,
-					false),
-				TEXT("DilatedVelocityTile")),
-			VelocityTileViewport.Rect);
+	FRDGTextureRef VelocityTileTexture =
+		GraphBuilder.CreateTexture(
+			FRDGTextureDesc::Create2DDesc(
+				VelocityTileCount,
+				PF_FloatRGBA,
+				FClearValueBinding::None,
+				GFastVRamConfig.MotionBlur,
+				TexCreate_ShaderResource | TexCreate_UAV | TexCreate_RenderTargetable,
+				false),
+			TEXT("DilatedVelocityTile"));
 
 	FMotionBlurVelocityDilateParameters VelocityDilateParameters;
 	VelocityDilateParameters.MotionBlur = MotionBlurParametersNoScale;
-	VelocityDilateParameters.VelocityTile = VelocityTileTexture.GetViewportParameters();
-	VelocityDilateParameters.VelocityTileTexture = VelocityTileTexture.GetRDGTexture();
+	VelocityDilateParameters.VelocityTile = Viewports.VelocityTileParameters;
+	VelocityDilateParameters.VelocityTileTexture = VelocityTileTextureSetup;
 
-	if (IsMotionBlurScatterRequired(Context, ColorViewport))
+	if (IsMotionBlurScatterRequired(View, Viewports.Color))
 	{
-		const FScreenPassTexture ScatterVelocityTileDepthTexture =
-			FScreenPassTexture::Create(
-				GraphBuilder.CreateTexture(
-					FRDGTextureDesc::Create2DDesc(
-						VelocityTileViewport.Extent,
-						PF_ShadowDepth,
-						FClearValueBinding::DepthOne,
-						TexCreate_None,
-						TexCreate_DepthStencilTargetable,
-						false),
-					TEXT("DilatedVelocityDepth")),
-				VelocityTileViewport.Rect);
+		FRDGTextureRef VelocityTileDepthTexture =
+			GraphBuilder.CreateTexture(
+				FRDGTextureDesc::Create2DDesc(
+					VelocityTileCount,
+					PF_ShadowDepth,
+					FClearValueBinding::DepthOne,
+					TexCreate_None,
+					TexCreate_DepthStencilTargetable,
+					false),
+				TEXT("DilatedVelocityDepth"));
 
 		FMotionBlurVelocityDilateScatterParameters* PassParameters = GraphBuilder.AllocParameters<FMotionBlurVelocityDilateScatterParameters>();
 		PassParameters->Dilate = VelocityDilateParameters;
 
 		PassParameters->RenderTargets.DepthStencil =
 			FDepthStencilBinding(
-				ScatterVelocityTileDepthTexture.GetRDGTexture(),
+				VelocityTileDepthTexture,
 				ERenderTargetLoadAction::ENoAction,
 				ERenderTargetStoreAction::EStore,
 				ERenderTargetLoadAction::ENoAction,
@@ -518,12 +538,12 @@ void ComputeMotionBlurVelocity(
 
 		PassParameters->RenderTargets[0] =
 			FRenderTargetBinding(
-				OutVelocityTileTexture.GetRDGTexture(),
+				VelocityTileTexture,
 				ERenderTargetLoadAction::ENoAction,
 				ERenderTargetStoreAction::EStore);
 
-		TShaderMapRef<FMotionBlurVelocityDilateScatterVS> VertexShader(Context->ShaderMap);
-		TShaderMapRef<FMotionBlurVelocityDilateScatterPS> PixelShader(Context->ShaderMap);
+		TShaderMapRef<FMotionBlurVelocityDilateScatterVS> VertexShader(View.ShaderMap);
+		TShaderMapRef<FMotionBlurVelocityDilateScatterPS> PixelShader(View.ShaderMap);
 
 		GraphBuilder.AddPass(
 			RDG_EVENT_NAME("VelocityTileScatter %dx%d", VelocityTileCount.X, VelocityTileCount.Y),
@@ -578,9 +598,9 @@ void ComputeMotionBlurVelocity(
 	{
 		FMotionBlurVelocityDilateGatherCS::FParameters* PassParameters = GraphBuilder.AllocParameters<FMotionBlurVelocityDilateGatherCS::FParameters>();
 		PassParameters->Dilate = VelocityDilateParameters;
-		PassParameters->OutDilatedVelocityTileTexture = GraphBuilder.CreateUAV(OutVelocityTileTexture.GetRDGTexture());
+		PassParameters->OutVelocityTileTexture = GraphBuilder.CreateUAV(VelocityTileTexture);
 
-		TShaderMapRef<FMotionBlurVelocityDilateGatherCS> ComputeShader(Context->ShaderMap);
+		TShaderMapRef<FMotionBlurVelocityDilateGatherCS> ComputeShader(View.ShaderMap);
 		FComputeShaderUtils::AddPass(
 			GraphBuilder,
 			RDG_EVENT_NAME("VelocityTileGatherCS %dx%d", VelocityTileCount.X, VelocityTileCount.Y),
@@ -589,32 +609,31 @@ void ComputeMotionBlurVelocity(
 			FComputeShaderUtils::GetGroupCount(VelocityTileCount, FComputeShaderUtils::kGolden2DGroupSize));
 	}
 
-	check(OutVelocityFlatTexture.IsValid());
-	check(OutVelocityTileTexture.IsValid());
+	*VelocityFlatTextureOutput = VelocityFlatTexture;
+	*VelocityTileTextureOutput = VelocityTileTexture;
 }
 
-FScreenPassTexture ComputeMotionBlurFilter(
+FRDGTextureRef ComputeMotionBlurFilter(
 	FRDGBuilder& GraphBuilder,
-	FScreenPassContextRef Context,
-	const FScreenPassTexture& ColorTexture,
-	const FScreenPassTexture& VelocityFlatTexture,
-	const FScreenPassTexture& VelocityTileTexture,
+	const FScreenPassViewInfo& ScreenPassView,
+	const FMotionBlurViewports& Viewports,
+	FRDGTextureRef ColorTexture,
+	FRDGTextureRef VelocityFlatTexture,
+	FRDGTextureRef VelocityTileTexture,
 	EMotionBlurFilterPass MotionBlurFilterPass,
 	EMotionBlurQuality MotionBlurQuality)
 {
-	// Validate inputs.
-	check(Context);
-	check(ColorTexture.IsValid());
-	check(VelocityFlatTexture.IsValid());
-	check(VelocityTileTexture.IsValid());
+	check(ColorTexture);
+	check(VelocityFlatTexture);
+	check(VelocityTileTexture);
 	check(MotionBlurFilterPass != EMotionBlurFilterPass::MAX);
 	check(MotionBlurQuality != EMotionBlurQuality::MAX);
 
-	const FScreenPassTextureViewport ColorViewport = ColorTexture.GetViewport();
+	const FViewInfo& View = ScreenPassView.View;
 
 	const float MotionBlur2ndScale = CVarMotionBlur2ndScale.GetValueOnRenderThread();
 
-	const bool bUseCompute = Context->bUseComputePasses;
+	const bool bUseCompute = ScreenPassView.bUseComputePasses;
 
 	const float BlurScaleLUT[static_cast<uint32>(EMotionBlurFilterPass::MAX)][static_cast<uint32>(EMotionBlurQuality::MAX)] =
 	{
@@ -645,26 +664,22 @@ FScreenPassTexture ComputeMotionBlurFilter(
 
 	const float BlurScale = BlurScaleLUT[static_cast<uint32>(MotionBlurFilterPass)][static_cast<uint32>(MotionBlurQuality)];
 
-	const auto GetFilterParameters = [&]()
-	{
-		FMotionBlurFilterParameters Parameters;
-		Parameters.MotionBlur = GetMotionBlurParameters(Context, ColorTexture.GetViewport().Extent, BlurScale);
-		Parameters.Color = ColorTexture.GetViewportParameters();
-		Parameters.Velocity = VelocityFlatTexture.GetViewportParameters();
-		Parameters.VelocityTile = VelocityTileTexture.GetViewportParameters();
-		Parameters.ColorToVelocity = GetScreenPassTextureViewportTransform(Parameters.Color, Parameters.Velocity);
-		Parameters.ColorToVelocityTile = GetScreenPassTextureViewportTransform(Parameters.Color, Parameters.VelocityTile);
-		Parameters.ColorTexture = ColorTexture.GetRDGTexture();
-		Parameters.VelocityFlatTexture = VelocityFlatTexture.GetRDGTexture();
-		Parameters.VelocityTileTexture = VelocityTileTexture.GetRDGTexture();
-		Parameters.ColorSampler = GetMotionBlurColorSampler();
-		Parameters.VelocitySampler = GetMotionBlurVelocitySampler();
-		Parameters.VelocityTileSampler = GetMotionBlurVelocitySampler();
-		Parameters.VelocityFlatSampler = GetMotionBlurVelocitySampler();
-		return Parameters;
-	};
-
-	FRDGTextureDesc OutColorDesc = ColorTexture.GetRDGTexture()->Desc;
+	FMotionBlurFilterParameters MotionBlurFilterParameters;
+	MotionBlurFilterParameters.MotionBlur = GetMotionBlurParameters(View, Viewports.Color.Extent, BlurScale);
+	MotionBlurFilterParameters.Color = Viewports.ColorParameters;
+	MotionBlurFilterParameters.Velocity = Viewports.VelocityParameters;
+	MotionBlurFilterParameters.VelocityTile = Viewports.VelocityTileParameters;
+	MotionBlurFilterParameters.ColorToVelocity = Viewports.ColorToVelocityTransform;
+	MotionBlurFilterParameters.ColorToVelocityTile = Viewports.ColorToVelocityTileTransform;
+	MotionBlurFilterParameters.ColorTexture = ColorTexture;
+	MotionBlurFilterParameters.VelocityFlatTexture = VelocityFlatTexture;
+	MotionBlurFilterParameters.VelocityTileTexture = VelocityTileTexture;
+	MotionBlurFilterParameters.ColorSampler = GetMotionBlurColorSampler();
+	MotionBlurFilterParameters.VelocitySampler = GetMotionBlurVelocitySampler();
+	MotionBlurFilterParameters.VelocityTileSampler = GetMotionBlurVelocitySampler();
+	MotionBlurFilterParameters.VelocityFlatSampler = GetMotionBlurVelocitySampler();
+
+	FRDGTextureDesc OutColorDesc = ColorTexture->Desc;
 	OutColorDesc.Reset();
 	OutColorDesc.TargetableFlags &= ~(TexCreate_RenderTargetable | TexCreate_UAV);
 	OutColorDesc.TargetableFlags |= bUseCompute ? TexCreate_UAV : TexCreate_RenderTargetable;
@@ -672,71 +687,75 @@ FScreenPassTexture ComputeMotionBlurFilter(
 	OutColorDesc.AutoWritable = false;
 	OutColorDesc.Format = FPostProcessing::HasAlphaChannelSupport() ? PF_FloatRGBA : PF_FloatRGB;
 
-	const FScreenPassTexture OutColorTexture =
-		FScreenPassTexture::Create(
-			GraphBuilder.CreateTexture(OutColorDesc, TEXT("MotionBlur")),
-			ColorTexture.GetViewport().Rect);
+	FRDGTextureRef ColorTextureOutput = GraphBuilder.CreateTexture(OutColorDesc, TEXT("MotionBlur"));
 
-	FRDGEventName PassName = RDG_EVENT_NAME("Motion Blur%dx%d", ColorViewport.Rect.Size().X, ColorViewport.Rect.Size().Y);
+	FMotionBlurFilterPermutationDomain PermutationVector;
+	PermutationVector.Set<FMotionBlurQualityDimension>(MotionBlurQuality);
 
 	if (bUseCompute)
 	{
 		FMotionBlurFilterCS::FParameters* PassParameters = GraphBuilder.AllocParameters<FMotionBlurFilterCS::FParameters>();
-		PassParameters->Filter = GetFilterParameters();
-		PassParameters->OutColorTexture = GraphBuilder.CreateUAV(OutColorTexture.GetRDGTexture());
+		PassParameters->Filter = MotionBlurFilterParameters;
+		PassParameters->OutColorTexture = GraphBuilder.CreateUAV(ColorTextureOutput);
 
-		FMotionBlurFilterCS::FPermutationDomain PermutationVector;
-		PermutationVector.Set<FMotionBlurQualityDimension>(MotionBlurQuality);
+		TShaderMapRef<FMotionBlurFilterCS> ComputeShader(View.ShaderMap, PermutationVector);
 
-		TShaderMapRef<FMotionBlurFilterCS> ComputeShader(Context->ShaderMap, PermutationVector);
 		FComputeShaderUtils::AddPass(
 			GraphBuilder,
-			MoveTemp(PassName),
+			RDG_EVENT_NAME("Motion Blur %dx%d (CS)", Viewports.Color.Rect.Width(), Viewports.Color.Rect.Height()),
 			*ComputeShader,
 			PassParameters,
-			FComputeShaderUtils::GetGroupCount(ColorViewport.Extent, FComputeShaderUtils::kGolden2DGroupSize));
+			FComputeShaderUtils::GetGroupCount(Viewports.Color.Extent, FComputeShaderUtils::kGolden2DGroupSize));
 	}
 	else
 	{
 		FMotionBlurFilterPS::FParameters* PassParameters = GraphBuilder.AllocParameters<FMotionBlurFilterPS::FParameters>();
-		PassParameters->Filter = GetFilterParameters();
-		PassParameters->RenderTargets[0] = FRenderTargetBinding(OutColorTexture.GetRDGTexture(), ERenderTargetLoadAction::ENoAction, ERenderTargetStoreAction::EStore);
+		PassParameters->Filter = MotionBlurFilterParameters;
+		PassParameters->RenderTargets[0] = FRenderTargetBinding(ColorTextureOutput, ERenderTargetLoadAction::ENoAction, ERenderTargetStoreAction::EStore);
 
-		FMotionBlurFilterPS::FPermutationDomain PermutationVector;
-		PermutationVector.Set<FMotionBlurQualityDimension>(MotionBlurQuality);
-		TShaderMapRef<FMotionBlurFilterPS> PixelShader(Context->ShaderMap, PermutationVector);
+		TShaderMapRef<FMotionBlurFilterPS> PixelShader(View.ShaderMap, PermutationVector);
 
 		AddDrawScreenPass(
 			GraphBuilder,
-			MoveTemp(PassName),
-			Context,
-			ColorViewport,
-			ColorViewport,
+			RDG_EVENT_NAME("Motion Blur %dx%d (PS)", Viewports.Color.Rect.Width(), Viewports.Color.Rect.Height()),
+			ScreenPassView,
+			Viewports.Color,
+			Viewports.Color,
 			*PixelShader,
 			PassParameters);
 	}
 
-	check(OutColorTexture.IsValid());
-	return OutColorTexture;
+	return ColorTextureOutput;
 }
 
-FScreenPassTexture ComputeMotionBlur(
+FRDGTextureRef ComputeMotionBlur(
 	FRDGBuilder& GraphBuilder,
-	FScreenPassContextRef Context,
-	const FScreenPassTexture& ColorTexture,
-	const FScreenPassTexture& DepthTexture,
-	const FScreenPassTexture& VelocityTexture)
+	const FScreenPassViewInfo& ScreenPassView,
+	FIntRect ColorViewportRect,
+	FIntRect VelocityViewportRect,
+	FRDGTextureRef ColorTexture,
+	FRDGTextureRef DepthTexture,
+	FRDGTextureRef VelocityTexture)
 {
-	// Validate inputs.
-	check(Context);
-	check(ColorTexture.IsValid());
-	check(DepthTexture.IsValid());
-	check(VelocityTexture.IsValid());
-	check(VelocityTexture.GetViewport() == DepthTexture.GetViewport());
-
-	FScreenPassTexture VelocityFlatTexture;
-	FScreenPassTexture VelocityTileTexture;
-	ComputeMotionBlurVelocity(GraphBuilder, Context, ColorTexture, DepthTexture, VelocityTexture, VelocityFlatTexture, VelocityTileTexture);
+	check(ColorTexture);
+	check(DepthTexture);
+	check(VelocityTexture);
+
+	const FMotionBlurViewports Viewports(
+		FScreenPassTextureViewport(ColorViewportRect, ColorTexture),
+		FScreenPassTextureViewport(VelocityViewportRect, VelocityTexture));
+
+	FRDGTextureRef VelocityFlatTexture = nullptr;
+	FRDGTextureRef VelocityTileTexture = nullptr;
+	ComputeMotionBlurVelocity(
+		GraphBuilder,
+		ScreenPassView,
+		Viewports,
+		ColorTexture,
+		DepthTexture,
+		VelocityTexture,
+		&VelocityFlatTexture,
+		&VelocityTileTexture);
 
 	const EMotionBlurQuality MotionBlurQuality = GetMotionBlurQuality();
 
@@ -744,9 +763,10 @@ FScreenPassTexture ComputeMotionBlur(
 
 	if (bUseSeparableFilter)
 	{
-		const FScreenPassTexture MotionBlurFilterTexture = ComputeMotionBlurFilter(
+		FRDGTextureRef MotionBlurFilterTexture = ComputeMotionBlurFilter(
 			GraphBuilder,
-			Context,
+			ScreenPassView,
+			Viewports,
 			ColorTexture,
 			VelocityFlatTexture,
 			VelocityTileTexture,
@@ -755,7 +775,8 @@ FScreenPassTexture ComputeMotionBlur(
 
 		return ComputeMotionBlurFilter(
 			GraphBuilder,
-			Context,
+			ScreenPassView,
+			Viewports,
 			MotionBlurFilterTexture,
 			VelocityFlatTexture,
 			VelocityTileTexture,
@@ -766,7 +787,8 @@ FScreenPassTexture ComputeMotionBlur(
 	{
 		return ComputeMotionBlurFilter(
 			GraphBuilder,
-			Context,
+			ScreenPassView,
+			Viewports,
 			ColorTexture,
 			VelocityFlatTexture,
 			VelocityTileTexture,
@@ -775,55 +797,56 @@ FScreenPassTexture ComputeMotionBlur(
 	}
 }
 
-FScreenPassTexture VisualizeMotionBlur(
+FRDGTextureRef VisualizeMotionBlur(
 	FRDGBuilder& GraphBuilder,
-	FScreenPassContextRef Context,
-	const FScreenPassTexture& ColorTexture,
-	const FScreenPassTexture& DepthTexture,
-	const FScreenPassTexture& VelocityTexture)
+	const FScreenPassViewInfo& ScreenPassView,
+	FIntRect ColorViewportRect,
+	FIntRect VelocityViewportRect,
+	FRDGTextureRef ColorTexture,
+	FRDGTextureRef DepthTexture,
+	FRDGTextureRef VelocityTexture)
 {
-	// Validate inputs.
-	check(Context);
-	check(ColorTexture.IsValid());
-	check(DepthTexture.IsValid());
-	check(VelocityTexture.IsValid());
+	check(ColorTexture);
+	check(DepthTexture);
+	check(VelocityTexture);
+
+	const FViewInfo& View = ScreenPassView.View;
 
-	const FViewInfo& View = Context->View;
+	const FScreenPassTextureViewport ColorViewport(ColorViewportRect, ColorTexture);
+	const FScreenPassTextureViewport VelocityViewport(VelocityViewportRect, VelocityTexture);
+	const FMotionBlurViewports Viewports(ColorViewport, VelocityViewport);
 
-	const FScreenPassTexture& TextureInput = ColorTexture;
-	const FScreenPassTexture TextureOutput = FScreenPassTexture::Create(
-		GraphBuilder.CreateTexture(TextureInput.GetRDGTexture()->Desc, TEXT("MotionBlurVisualize")),
-		TextureInput.GetViewport().Rect);
+	FRDGTextureRef ColorTextureOutput = GraphBuilder.CreateTexture(ColorTexture->Desc, TEXT("MotionBlurVisualize"));
 
 	FMotionBlurVisualizePS::FParameters* PassParameters = GraphBuilder.AllocParameters<FMotionBlurVisualizePS::FParameters>();
 	PassParameters->WorldToClipPrev = GetPreviousWorldToClipMatrix(View);
 	PassParameters->View = View.ViewUniformBuffer;
-	PassParameters->ColorTexture = ColorTexture.GetRDGTexture();
-	PassParameters->DepthTexture = DepthTexture.GetRDGTexture();
-	PassParameters->VelocityTexture = VelocityTexture.GetRDGTexture();
-	PassParameters->Color = ColorTexture.GetViewportParameters();
-	PassParameters->Velocity = VelocityTexture.GetViewportParameters();
-	PassParameters->ColorToVelocity = GetScreenPassTextureViewportTransform(PassParameters->Color, PassParameters->Velocity);
+	PassParameters->ColorTexture = ColorTexture;
+	PassParameters->DepthTexture = DepthTexture;
+	PassParameters->VelocityTexture = VelocityTexture;
+	PassParameters->Color = Viewports.ColorParameters;
+	PassParameters->Velocity = Viewports.VelocityParameters;
+	PassParameters->ColorToVelocity = Viewports.ColorToVelocityTransform;
 	PassParameters->ColorSampler = GetMotionBlurColorSampler();
 	PassParameters->VelocitySampler = GetMotionBlurVelocitySampler();
 	PassParameters->DepthSampler = GetMotionBlurVelocitySampler();
-	PassParameters->RenderTargets[0] = FRenderTargetBinding(TextureOutput.GetRDGTexture(), ERenderTargetLoadAction::ENoAction, ERenderTargetStoreAction::EStore);
+	PassParameters->RenderTargets[0] = FRenderTargetBinding(ColorTextureOutput, ERenderTargetLoadAction::ENoAction, ERenderTargetStoreAction::EStore);
 
-	TShaderMapRef<FMotionBlurVisualizePS> PixelShader(Context->ShaderMap);
+	TShaderMapRef<FMotionBlurVisualizePS> PixelShader(View.ShaderMap);
 
 	GraphBuilder.AddPass(
 		RDG_EVENT_NAME("MotionBlurVisualize"),
 		PassParameters,
 		ERenderGraphPassFlags::None,
-		[Context, TextureInput, TextureOutput, PixelShader, PassParameters](FRHICommandListImmediate& RHICmdList)
+		[ScreenPassView, ColorTextureOutput, ColorViewport, PixelShader, PassParameters](FRHICommandListImmediate& RHICmdList)
 	{
-		DrawScreenPass(RHICmdList, Context, TextureOutput.GetViewport(), TextureInput.GetViewport(), *PixelShader, *PassParameters);
+		DrawScreenPass(RHICmdList, ScreenPassView, ColorViewport, ColorViewport, *PixelShader, *PassParameters);
 
 		// Draw debug text.
 		{
-			const FViewInfo& LocalView = Context->View;
-			const FSceneViewFamily& ViewFamily = Context->ViewFamily;
-			FRenderTargetTemp TempRenderTarget(static_cast<FTexture2DRHIParamRef>(TextureOutput.GetRHITexture()), TextureOutput.GetViewport().Extent);
+			const FViewInfo& LocalView = ScreenPassView.View;
+			const FSceneViewFamily& ViewFamily = *LocalView.Family;
+			FRenderTargetTemp TempRenderTarget(static_cast<FTexture2DRHIParamRef>(ColorTextureOutput->GetRHITexture()), ColorViewport.Extent);
 			FCanvas Canvas(&TempRenderTarget, nullptr, ViewFamily.CurrentRealTime, ViewFamily.CurrentWorldTime, ViewFamily.DeltaWorldTime, LocalView.GetFeatureLevel());
 
 			float X = 20;
@@ -867,5 +890,67 @@ FScreenPassTexture VisualizeMotionBlur(
 		}
 	});
 
-	return TextureOutput;
+	return ColorTextureOutput;
+}
+
+//////////////////////////////////////////////////////////////////////////
+//! Shim methods to hook into the legacy pipeline until the full RDG conversion is complete.
+
+FRenderingCompositeOutputRef ComputeMotionBlurShim(
+	FRenderingCompositionGraph& Graph,
+	FRenderingCompositeOutputRef ColorInput,
+	FRenderingCompositeOutputRef DepthInput,
+	FRenderingCompositeOutputRef VelocityInput,
+	bool bVisualizeMotionBlur)
+{
+	FRenderingCompositePass* MotionBlurPass = Graph.RegisterPass(new(FMemStack::Get()) TRCPassForRDG<3, 1>(
+		[bVisualizeMotionBlur](FRenderingCompositePass* Pass, FRenderingCompositePassContext& InContext)
+	{
+		FRDGBuilder GraphBuilder(InContext.RHICmdList);
+
+		FRDGTextureRef ColorTexture = Pass->CreateRDGTextureForInput(GraphBuilder, ePId_Input0, TEXT("SceneColor"), eFC_0000);
+		FRDGTextureRef DepthTexture = Pass->CreateRDGTextureForInput(GraphBuilder, ePId_Input1, TEXT("SceneDepth"), eFC_0000);
+		FRDGTextureRef VelocityTexture = Pass->CreateRDGTextureForInput(GraphBuilder, ePId_Input2, TEXT("SceneVelocity"), eFC_0000);
+
+		const FIntRect ColorViewportRect(InContext.SceneColorViewRect);
+		const FIntRect VelocityViewportRect(InContext.View.ViewRect);
+
+		const FScreenPassViewInfo ScreenPassView(InContext.View);
+
+		FRDGTextureRef ColorTextureOutput;
+
+		if (bVisualizeMotionBlur)
+		{
+			ColorTextureOutput = VisualizeMotionBlur(
+				GraphBuilder,
+				ScreenPassView,
+				ColorViewportRect,
+				VelocityViewportRect,
+				ColorTexture,
+				DepthTexture,
+				VelocityTexture);
+		}
+		else
+		{
+			ColorTextureOutput = ComputeMotionBlur(
+				GraphBuilder,
+				ScreenPassView,
+				ColorViewportRect,
+				VelocityViewportRect,
+				ColorTexture,
+				DepthTexture,
+				VelocityTexture);
+		}
+
+		Pass->ExtractRDGTextureForOutput(GraphBuilder, ePId_Output0, ColorTextureOutput);
+
+		GraphBuilder.Execute();
+	}));
+
+	MotionBlurPass->SetInput(ePId_Input0, ColorInput);
+	MotionBlurPass->SetInput(ePId_Input1, DepthInput);
+	MotionBlurPass->SetInput(ePId_Input2, VelocityInput);
+	return FRenderingCompositeOutputRef(MotionBlurPass);
 }
+
+//////////////////////////////////////////////////////////////////////////
\ No newline at end of file
diff --git a/Engine/Source/Runtime/Renderer/Private/PostProcess/PostProcessMotionBlur.h b/Engine/Source/Runtime/Renderer/Private/PostProcess/PostProcessMotionBlur.h
--- a/Engine/Source/Runtime/Renderer/Private/PostProcess/PostProcessMotionBlur.h
+++ b/Engine/Source/Runtime/Renderer/Private/PostProcess/PostProcessMotionBlur.h
@@ -7,6 +7,7 @@
 #pragma once
 
 #include "ScreenPass.h"
+#include "RenderingCompositionGraph.h"
 
 // Returns whether motion blur is enabled for the requested view.
 bool IsMotionBlurEnabled(const FViewInfo& View);
@@ -24,25 +25,14 @@ enum class EMotionBlurQuality : uint32
 // Returns the motion blur quality level set by the 'r.MotionBlurQuality' cvar.
 EMotionBlurQuality GetMotionBlurQuality();
 
-/**
- * Computes motion blur using the provided color, depth, and velocity textures. Returns
- * the blurred color result. The Color viewport does not need to match Depth / Velocity
- * viewports, but Depth / Velocity must match each other.
- */
-FScreenPassTexture ComputeMotionBlur(
-	FRDGBuilder& GraphBuilder,
-	FScreenPassContextRef Context,
-	const FScreenPassTexture& ColorTexture,
-	const FScreenPassTexture& DepthTexture,
-	const FScreenPassTexture& VelocityTexture);
-
-/**
- * Visualizes motion blur velocities and outputs the result. The Color viewport does not
- * need to match Depth / Velocity viewports, but Depth / Velocity must match each other.
- */
-FScreenPassTexture VisualizeMotionBlur(
-	FRDGBuilder& GraphBuilder,
-	FScreenPassContextRef Context,
-	const FScreenPassTexture& ColorTexture,
-	const FScreenPassTexture& DepthTexture,
-	const FScreenPassTexture& VelocityTexture);
\ No newline at end of file
+//////////////////////////////////////////////////////////////////////////
+//! Shim methods to hook into the legacy pipeline until the full RDG conversion is complete.
+
+FRenderingCompositeOutputRef ComputeMotionBlurShim(
+	FRenderingCompositionGraph& Graph,
+	FRenderingCompositeOutputRef ColorInput,
+	FRenderingCompositeOutputRef DepthInput,
+	FRenderingCompositeOutputRef VelocityInput,
+	bool bVisualizeMotionBlur);
+
+//////////////////////////////////////////////////////////////////////////
\ No newline at end of file
diff --git a/Engine/Source/Runtime/Renderer/Private/PostProcess/PostProcessSubsurface.cpp b/Engine/Source/Runtime/Renderer/Private/PostProcess/PostProcessSubsurface.cpp
--- a/Engine/Source/Runtime/Renderer/Private/PostProcess/PostProcessSubsurface.cpp
+++ b/Engine/Source/Runtime/Renderer/Private/PostProcess/PostProcessSubsurface.cpp
@@ -124,6 +124,67 @@ ESubsurfaceMode GetSubsurfaceModeForView(const FViewInfo& View)
 	}
 }
 
+bool IsSubsurfaceEnabled()
+{
+	const bool bEnabled = CVarSubsurfaceScattering.GetValueOnAnyThread() != 0;
+	const bool bHasScale = CVarSSSScale.GetValueOnAnyThread() > 0.0f;
+	return (bEnabled && bHasScale);
+}
+
+bool IsSubsurfaceRequiredForView(const FViewInfo& View)
+{
+	const bool bSimpleDynamicLighting = IsAnyForwardShadingEnabled(View.GetShaderPlatform());
+	const bool bSubsurfaceEnabled = IsSubsurfaceEnabled();
+	const bool bViewHasSubsurfaceMaterials = ((View.ShadingModelMaskInView & GetUseSubsurfaceProfileShadingModelMask()) != 0);
+	return (bSubsurfaceEnabled && bViewHasSubsurfaceMaterials && !bSimpleDynamicLighting);
+}
+
+uint32 GetSubsurfaceRequiredViewMask(const TArray<FViewInfo>& Views)
+{
+	const uint32 ViewCount = Views.Num();
+	uint32 ViewMask = 0;
+
+	// Traverse the views to make sure we only process subsurface if requested by any view.
+	for (uint32 ViewIndex = 0; ViewIndex < ViewCount; ++ViewIndex)
+	{
+		const FViewInfo& View = Views[ViewIndex];
+
+		if (IsSubsurfaceRequiredForView(View))
+		{
+			const uint32 ViewBit = 1 << ViewIndex;
+
+			ViewMask |= ViewBit;
+		}
+	}
+
+	return ViewMask;
+}
+
+bool IsSubsurfaceCheckerboardFormat(EPixelFormat SceneColorFormat)
+{
+	int CVarValue = CVarSSSCheckerboard.GetValueOnRenderThread();
+	if (CVarValue == 0)
+	{
+		return false;
+	}
+	else if (CVarValue == 1)
+	{
+		return true;
+	}
+	else if (CVarValue == 2)
+	{
+		switch (SceneColorFormat)
+		{
+		case PF_A32B32G32R32F:
+		case PF_FloatRGBA:
+			return false;
+		default:
+			return true;
+		}
+	}
+	return true;
+}
+
 // Returns the SS profile texture with a black fallback texture if none exists yet.
 FTextureRHIRef GetSubsurfaceProfileTexture(FRHICommandListImmediate& RHICmdList)
 {
@@ -140,20 +201,27 @@ FTextureRHIRef GetSubsurfaceProfileTexture(FRHICommandListImmediate& RHICmdList)
 
 // Set of common shader parameters shared by all subsurface shaders.
 BEGIN_SHADER_PARAMETER_STRUCT(FSubsurfaceParameters, )
-	SHADER_PARAMETER(FVector4, SubsurfaceParams)
-	SHADER_PARAMETER_STRUCT_INCLUDE(FScreenPassCommonParameters, ScreenPassCommonParameters)
+SHADER_PARAMETER(FVector4, SubsurfaceParams)
+	SHADER_PARAMETER_STRUCT_REF(FSceneTexturesUniformParameters, SceneUniformBuffer)
+	SHADER_PARAMETER_STRUCT_REF(FViewUniformShaderParameters, ViewUniformBuffer)
+	SHADER_PARAMETER_SAMPLER(SamplerState, BilinearTextureSampler)
 	SHADER_PARAMETER_TEXTURE(Texture2D, SSProfilesTexture)
 END_SHADER_PARAMETER_STRUCT()
 
-FSubsurfaceParameters GetSubsurfaceCommonParameters(FRHICommandListImmediate& RHICmdList, FScreenPassContextRef Context)
+FSubsurfaceParameters GetSubsurfaceCommonParameters(FRHICommandListImmediate& RHICmdList, const FViewInfo& View)
 {
-	const float DistanceToProjectionWindow = Context->View.ViewMatrices.GetProjectionMatrix().M[0][0];
+	FSceneRenderTargets& SceneContext = FSceneRenderTargets::Get(RHICmdList);
+
+	const float DistanceToProjectionWindow = View.ViewMatrices.GetProjectionMatrix().M[0][0];
 	const float SSSScaleZ = DistanceToProjectionWindow * GetSubsurfaceRadiusScale();
 	const float SSSScaleX = SSSScaleZ / SUBSURFACE_KERNEL_SIZE * 0.5f;
 
 	FSubsurfaceParameters Parameters;
-	Parameters.ScreenPassCommonParameters = Context->ScreenPassCommonParameters;
 	Parameters.SubsurfaceParams = FVector4(SSSScaleX, SSSScaleZ, 0, 0);
+	Parameters.ViewUniformBuffer = View.ViewUniformBuffer;
+	Parameters.SceneUniformBuffer = CreateSceneTextureUniformBuffer(
+		SceneContext, View.FeatureLevel, ESceneTextureSetupMode::All, EUniformBufferUsage::UniformBuffer_SingleFrame);
+	Parameters.BilinearTextureSampler = TStaticSamplerState<SF_Bilinear>::GetRHI();
 	Parameters.SSProfilesTexture = GetSubsurfaceProfileTexture(RHICmdList);
 	return Parameters;
 }
@@ -164,11 +232,11 @@ BEGIN_SHADER_PARAMETER_STRUCT(FSubsurfaceInput, )
 	SHADER_PARAMETER_RDG_TEXTURE(Texture2D, Texture)
 END_SHADER_PARAMETER_STRUCT()
 
-FSubsurfaceInput GetSubsurfaceInput(const FScreenPassTexture& Texture)
+FSubsurfaceInput GetSubsurfaceInput(FRDGTextureRef Texture, const FScreenPassTextureViewportParameters& ViewportParameters)
 {
 	FSubsurfaceInput Input;
-	Input.Texture = Texture.GetRDGTexture();
-	Input.Viewport = Texture.GetViewportParameters();
+	Input.Texture = Texture;
+	Input.Viewport = ViewportParameters;
 	return Input;
 }
 
@@ -352,7 +420,7 @@ class FSubsurfaceViewportCopyPS : public FSubsurfaceShader
 	SHADER_USE_PARAMETER_STRUCT(FSubsurfaceViewportCopyPS, FSubsurfaceShader);
 
 	BEGIN_SHADER_PARAMETER_STRUCT(FParameters, )
-		SHADER_PARAMETER_STRUCT(FSubsurfaceInput, SubsurfaceInput0)
+		SHADER_PARAMETER_RDG_TEXTURE(Texture2D, SubsurfaceInput0_Texture)
 		SHADER_PARAMETER_SAMPLER(SamplerState, SubsurfaceSampler0)
 		RENDER_TARGET_BINDING_SLOTS()
 	END_SHADER_PARAMETER_STRUCT();
@@ -365,87 +433,48 @@ class FSubsurfaceViewportCopyPS : public FSubsurfaceShader
 
 IMPLEMENT_GLOBAL_SHADER(FSubsurfaceViewportCopyPS, "/Engine/Private/PostProcessSubsurface.usf", "SubsurfaceViewportCopyPS", SF_Pixel);
 
-bool IsSubsurfaceEnabled()
+void ComputeSubsurfaceForView(
+	FRDGBuilder& GraphBuilder,
+	const FScreenPassViewInfo& ScreenPassView,
+	const FScreenPassTextureViewport& SceneViewport,
+	FRDGTextureRef SceneTexture,
+	FRDGTextureRef SceneTextureOutput,
+	ERenderTargetLoadAction SceneTextureLoadAction)
 {
-	const bool bEnabled = CVarSubsurfaceScattering.GetValueOnAnyThread() != 0;
-	const bool bHasScale = CVarSSSScale.GetValueOnAnyThread() > 0.0f;
-	return (bEnabled && bHasScale);
-}
+	check(SceneTexture);
+	check(SceneTextureOutput);
+	check(SceneViewport.Extent == SceneTexture->Desc.Extent);
 
-bool IsSubsurfaceRequiredForView(const FViewInfo& View)
-{
-	const bool bSimpleDynamicLighting = IsAnyForwardShadingEnabled(View.GetShaderPlatform());
-	const bool bSubsurfaceEnabled = IsSubsurfaceEnabled();
-	const bool bViewHasSubsurfaceMaterials = ((View.ShadingModelMaskInView & GetUseSubsurfaceProfileShadingModelMask()) != 0);
-	return (bSubsurfaceEnabled && bViewHasSubsurfaceMaterials && !bSimpleDynamicLighting);
-}
-
-bool IsSubsurfaceCheckerboardFormat(EPixelFormat SceneColorFormat)
-{
-	int CVarValue = CVarSSSCheckerboard.GetValueOnRenderThread();
-	if (CVarValue == 0)
-	{
-		return false;
-	}
-	else if (CVarValue == 1)
-	{
-		return true;
-	}
-	else if (CVarValue == 2)
-	{
-		switch (SceneColorFormat)
-		{
-		case PF_A32B32G32R32F:
-		case PF_FloatRGBA:
-			return false;
-		default:
-			return true;
-		}
-	}
-	return true;
-}
+	const FViewInfo& View = ScreenPassView.View;
 
-FScreenPassTexture ComputeSubsurface(
-	FRDGBuilder& GraphBuilder,
-	FScreenPassContextRef Context,
-	const FScreenPassTexture& SceneTexture)
-{
-	check(Context);
-	check(SceneTexture.IsValid());
+	const FSceneViewFamily* ViewFamily = View.Family;
 
-	const FRDGTextureDesc& SceneTextureDesc = SceneTexture.GetRDGTexture()->Desc;
+	const FRDGTextureDesc& SceneTextureDesc = SceneTexture->Desc;
 
-	const ESubsurfaceMode SubsurfaceMode = GetSubsurfaceModeForView(Context->View);
+	const ESubsurfaceMode SubsurfaceMode = GetSubsurfaceModeForView(View);
 
 	const bool bHalfRes = (SubsurfaceMode == ESubsurfaceMode::HalfRes);
 
 	const bool bCheckerboard = IsSubsurfaceCheckerboardFormat(SceneTextureDesc.Format);
 
+	const uint32 ScaleFactor = bHalfRes ? 2 : 1;
+
 	/**
-	 * All intermediate passes within the screen-space subsurface effect can operate at half or full resolution,
-	 * depending on the subsurface mode. The values are precomputed and shared among all intermediate textures.
+	 * All subsurface passes within the screen-space subsurface effect can operate at half or full resolution,
+	 * depending on the subsurface mode. The values are precomputed and shared among all Subsurface textures.
 	 */
-	FScreenPassTextureViewport IntermediateViewport = SceneTexture.GetViewport();
+	const FScreenPassTextureViewport SubsurfaceViewport = FScreenPassTextureViewport::CreateDownscaled(SceneViewport, ScaleFactor);
 
-	if (bHalfRes)
-	{
-		IntermediateViewport = FScreenPassTextureViewport::CreateDownscaled(IntermediateViewport, 2);
-	}
-
-	const FRDGTextureDesc IntermediateTextureDescriptor = FRDGTextureDesc::Create2DDesc(
-		IntermediateViewport.Extent,
-		PF_FloatRGBA,
-		FClearValueBinding(),
-		TexCreate_None,
-		TexCreate_RenderTargetable | TexCreate_ShaderResource,
-		false);
+	FRDGTextureDesc SubsurfaceTextureDescriptor = SceneTexture->Desc;
+	SubsurfaceTextureDescriptor.Extent = SubsurfaceViewport.Extent;
 	
-	const FSubsurfaceParameters SubsurfaceCommonParameters = GetSubsurfaceCommonParameters(GraphBuilder.RHICmdList, Context);
+	const FSubsurfaceParameters SubsurfaceCommonParameters = GetSubsurfaceCommonParameters(GraphBuilder.RHICmdList, View);
+	const FScreenPassTextureViewportParameters SubsurfaceViewportParameters = GetScreenPassTextureViewportParameters(SubsurfaceViewport);
+	const FScreenPassTextureViewportParameters SceneViewportParameters = GetScreenPassTextureViewportParameters(SceneViewport);
 
-	FScreenPassTexture SetupTexture = SceneTexture;
-	FScreenPassTexture SubsurfaceTextureX;
-	FScreenPassTexture SubsurfaceTextureY;
-	FScreenPassTexture RecombineTexture;
+	FRDGTextureRef SetupTexture = SceneTexture;
+	FRDGTextureRef SubsurfaceTextureX = nullptr;
+	FRDGTextureRef SubsurfaceTextureY = nullptr;
 
 	FSamplerStateRHIParamRef PointClampSampler = TStaticSamplerState<SF_Point, AM_Clamp, AM_Clamp, AM_Clamp>::GetRHI();
 	FSamplerStateRHIParamRef BilinearBorderSampler = TStaticSamplerState<SF_Bilinear, AM_Border, AM_Border, AM_Border>::GetRHI();
@@ -456,51 +485,47 @@ FScreenPassTexture ComputeSubsurface(
 	 */
 	if (SubsurfaceMode != ESubsurfaceMode::Bypass)
 	{
-		SetupTexture = FScreenPassTexture::Create(
-			GraphBuilder.CreateTexture(IntermediateTextureDescriptor, TEXT("SubsurfaceSetupTexture")),
-			IntermediateViewport.Rect);
+		SetupTexture = GraphBuilder.CreateTexture(SubsurfaceTextureDescriptor, TEXT("SubsurfaceSetupTexture"));
 
 		// Setup pass outputs the diffuse scene color and depth in preparation for the scatter passes.
 		{
-			const FScreenPassTexture& TextureInput = SceneTexture;
-			const FScreenPassTexture& TextureOutput = SetupTexture;
-
 			FSubsurfaceSetupPS::FParameters* PassParameters = GraphBuilder.AllocParameters<FSubsurfaceSetupPS::FParameters>();
 			PassParameters->Subsurface = SubsurfaceCommonParameters;
-			PassParameters->RenderTargets[0] = FRenderTargetBinding(TextureOutput.GetRDGTexture(), ERenderTargetLoadAction::ENoAction, ERenderTargetStoreAction::EStore);
-			PassParameters->SubsurfaceInput0 = GetSubsurfaceInput(TextureInput);
+			PassParameters->RenderTargets[0] = FRenderTargetBinding(SetupTexture, ERenderTargetLoadAction::ENoAction, ERenderTargetStoreAction::EStore);
+			PassParameters->SubsurfaceInput0 = GetSubsurfaceInput(SceneTexture, SceneViewportParameters);
 			PassParameters->SubsurfaceSampler0 = PointClampSampler;
 
 			FSubsurfaceSetupPS::FPermutationDomain PixelShaderPermutationVector;
 			PixelShaderPermutationVector.Set<FSubsurfaceSetupPS::FDimensionHalfRes>(bHalfRes);
 			PixelShaderPermutationVector.Set<FSubsurfaceSetupPS::FDimensionCheckerboard>(bCheckerboard);
-			TShaderMapRef<FSubsurfaceSetupPS> PixelShader(Context->ShaderMap, PixelShaderPermutationVector);
-
-			AddDrawScreenPass(GraphBuilder, RDG_EVENT_NAME("SubsurfaceSetup"), Context, TextureOutput.GetViewport(), TextureInput.GetViewport(), *PixelShader, PassParameters);
+			TShaderMapRef<FSubsurfaceSetupPS> PixelShader(View.ShaderMap, PixelShaderPermutationVector);
+
+			/**
+			 * The subsurface viewport is intentionally used as both the target and texture viewport, even though the texture
+			 * is potentially double the size. This is to ensure that the source UVs map 1-to-1 with pixel centers of the target,
+			 * in order to ensure that the checkerboard pattern selects the correct pixels from the scene texture. This still works
+			 * because the texture viewport is normalized into UV space, so it doesn't matter that the dimensions are twice as large.
+			 */
+			AddDrawScreenPass(GraphBuilder, RDG_EVENT_NAME("SubsurfaceSetup"), ScreenPassView, SubsurfaceViewport, SubsurfaceViewport, *PixelShader, PassParameters);
 		}
 
-		SubsurfaceTextureX = FScreenPassTexture::Create(
-			GraphBuilder.CreateTexture(IntermediateTextureDescriptor, TEXT("SubsurfaceTextureX")),
-			IntermediateViewport.Rect);
-
-		SubsurfaceTextureY = FScreenPassTexture::Create(
-			GraphBuilder.CreateTexture(IntermediateTextureDescriptor, TEXT("SubsurfaceTextureY")),
-			IntermediateViewport.Rect);
+		SubsurfaceTextureX = GraphBuilder.CreateTexture(SubsurfaceTextureDescriptor, TEXT("SubsurfaceTextureX"));
+		SubsurfaceTextureY = GraphBuilder.CreateTexture(SubsurfaceTextureDescriptor, TEXT("SubsurfaceTextureY"));
 
 		FSamplerStateRHIParamRef SubsurfaceSamplerState = FSubsurfacePS::GetSamplerState();
 		const FSubsurfacePS::EQuality SubsurfaceQuality = FSubsurfacePS::GetQuality();
 
 		struct FSubsurfacePassInfo
 		{
-			FSubsurfacePassInfo(const TCHAR* InName, const FScreenPassTexture& InInput, const FScreenPassTexture& InOutput)
+			FSubsurfacePassInfo(const TCHAR* InName, FRDGTextureRef InInput, FRDGTextureRef InOutput)
 				: Name(InName)
 				, Input(InInput)
 				, Output(InOutput)
 			{}
 
 			const TCHAR* Name;
-			const FScreenPassTexture& Input;
-			const FScreenPassTexture& Output;
+			FRDGTextureRef Input;
+			FRDGTextureRef Output;
 		};
 
 		const FSubsurfacePassInfo SubsurfacePassInfoByDirection[] =
@@ -515,134 +540,178 @@ FScreenPassTexture ComputeSubsurface(
 			const auto Direction = static_cast<FSubsurfacePS::EDirection>(DirectionIndex);
 
 			const FSubsurfacePassInfo& PassInfo = SubsurfacePassInfoByDirection[DirectionIndex];
-			const FScreenPassTexture& TextureInput = PassInfo.Input;
-			const FScreenPassTexture& TextureOutput = PassInfo.Output;
+			FRDGTextureRef TextureInput = PassInfo.Input;
+			FRDGTextureRef TextureOutput = PassInfo.Output;
 
 			FSubsurfacePS::FParameters* PassParameters = GraphBuilder.AllocParameters<FSubsurfacePS::FParameters>();
 			PassParameters->Subsurface = SubsurfaceCommonParameters;
-			PassParameters->RenderTargets[0] = FRenderTargetBinding(TextureOutput.GetRDGTexture(), ERenderTargetLoadAction::ENoAction, ERenderTargetStoreAction::EStore);
-			PassParameters->SubsurfaceInput0 = GetSubsurfaceInput(TextureInput);
+			PassParameters->RenderTargets[0] = FRenderTargetBinding(TextureOutput, ERenderTargetLoadAction::ENoAction, ERenderTargetStoreAction::EStore);
+			PassParameters->SubsurfaceInput0 = GetSubsurfaceInput(TextureInput, SubsurfaceViewportParameters);
 			PassParameters->SubsurfaceSampler0 = SubsurfaceSamplerState;
 
 			FSubsurfacePS::FPermutationDomain PixelShaderPermutationVector;
 			PixelShaderPermutationVector.Set<FSubsurfacePS::FDimensionDirection>(Direction);
 			PixelShaderPermutationVector.Set<FSubsurfacePS::FDimensionQuality>(SubsurfaceQuality);
-			TShaderMapRef<FSubsurfacePS> PixelShader(Context->ShaderMap, PixelShaderPermutationVector);
+			TShaderMapRef<FSubsurfacePS> PixelShader(View.ShaderMap, PixelShaderPermutationVector);
 
-			AddDrawScreenPass(GraphBuilder, FRDGEventName(PassInfo.Name), Context, TextureOutput.GetViewport(), TextureInput.GetViewport(), *PixelShader, PassParameters);
+			AddDrawScreenPass(GraphBuilder, FRDGEventName(PassInfo.Name), ScreenPassView, SubsurfaceViewport, SubsurfaceViewport, *PixelShader, PassParameters);
 		}
 	}
 
-	RecombineTexture = FScreenPassTexture::Create(
-		GraphBuilder.CreateTexture(SceneTextureDesc, TEXT("SubsurfaceRecombine")),
-		SceneTexture.GetViewport().Rect);
+	// Recombines scattering result with scene color.
+	{
+		FSubsurfaceRecombinePS::FParameters* PassParameters = GraphBuilder.AllocParameters<FSubsurfaceRecombinePS::FParameters>();
+		PassParameters->Subsurface = SubsurfaceCommonParameters;
+		PassParameters->RenderTargets[0] = FRenderTargetBinding(SceneTextureOutput, SceneTextureLoadAction, ERenderTargetStoreAction::EStore);
+		PassParameters->SubsurfaceInput0 = GetSubsurfaceInput(SceneTexture, SceneViewportParameters);
+		PassParameters->SubsurfaceSampler0 = BilinearBorderSampler;
+
+		// Scattering output target is only used when scattering is enabled.
+		if (SubsurfaceMode != ESubsurfaceMode::Bypass)
+		{
+			PassParameters->SubsurfaceInput1 = GetSubsurfaceInput(SubsurfaceTextureY, SubsurfaceViewportParameters);
+			PassParameters->SubsurfaceSampler1 = BilinearBorderSampler;
+		}
+
+		const FSubsurfaceRecombinePS::EQuality RecombineQuality = FSubsurfaceRecombinePS::GetQuality(View);
+
+		FSubsurfaceRecombinePS::FPermutationDomain PixelShaderPermutationVector;
+		PixelShaderPermutationVector.Set<FSubsurfaceRecombinePS::FDimensionMode>(SubsurfaceMode);
+		PixelShaderPermutationVector.Set<FSubsurfaceRecombinePS::FDimensionQuality>(RecombineQuality);
+		PixelShaderPermutationVector.Set<FSubsurfaceRecombinePS::FDimensionCheckerboard>(bCheckerboard);
+		TShaderMapRef<FSubsurfaceRecombinePS> PixelShader(View.ShaderMap, PixelShaderPermutationVector);
+
+		/**
+		 * See the related comment above in the prepare pass. The scene viewport is used as both the target and
+		 * texture viewport in order to ensure that the correct pixel is sampled for checkerboard rendering.
+		 */
+		AddDrawScreenPass(GraphBuilder, RDG_EVENT_NAME("SubsurfaceRecombine"), ScreenPassView, SceneViewport, SceneViewport, *PixelShader, PassParameters);
+	}
+}
+
+FRDGTextureRef ComputeSubsurface(
+	FRDGBuilder& GraphBuilder,
+	FRDGTextureRef SceneTexture,
+	const TArray<FViewInfo>& Views)
+{
+	const uint32 ViewCount = Views.Num();
+	const uint32 ViewMaskAll = (1 << ViewCount) - 1;
+	const uint32 ViewMask = GetSubsurfaceRequiredViewMask(Views);
 
-	// If multiple views exist (e.g. VR) we need to copy all other views from the scene texture
-	// to the recombine target so we don't lose them.
-	if (Context->ViewFamily.Views.Num())
+	// Return the original target if no views have subsurface applied.
+	if (!ViewMask)
 	{
-		const FScreenPassTexture& TextureInput = SceneTexture;
-		const FScreenPassTexture& TextureOutput = RecombineTexture;
+		return SceneTexture;
+	}
+
+	FRDGTextureRef SceneTextureOutput = GraphBuilder.CreateTexture(SceneTexture->Desc, TEXT("SceneColorSubsurface"));
+
+	ERenderTargetLoadAction SceneTextureLoadAction = ERenderTargetLoadAction::ENoAction;
 
+	const bool bHasNonSubsurfaceView = ViewMask != ViewMaskAll;
+
+	/**
+	 * Since we are outputting to a new texture and certain views may not utilize subsurface scattering,
+	 * we need to copy all non-subsurface views onto the destination texture.
+	 */
+	if (bHasNonSubsurfaceView)
+	{
 		FSubsurfaceViewportCopyPS::FParameters* PassParameters = GraphBuilder.AllocParameters<FSubsurfaceViewportCopyPS::FParameters>();
-		PassParameters->RenderTargets[0] = FRenderTargetBinding(TextureOutput.GetRDGTexture(), ERenderTargetLoadAction::ENoAction, ERenderTargetStoreAction::EStore);
-		PassParameters->SubsurfaceInput0 = GetSubsurfaceInput(TextureInput);
-		PassParameters->SubsurfaceSampler0 = PointClampSampler;
+		PassParameters->RenderTargets[0] = FRenderTargetBinding(SceneTextureOutput, ERenderTargetLoadAction::ENoAction, ERenderTargetStoreAction::EStore);
+		PassParameters->SubsurfaceInput0_Texture = SceneTexture;
+		PassParameters->SubsurfaceSampler0 = TStaticSamplerState<SF_Point, AM_Clamp, AM_Clamp, AM_Clamp>::GetRHI();
 
-		TShaderMapRef<FSubsurfaceViewportCopyPS> PixelShader(Context->ShaderMap);
+		TShaderMapRef<FSubsurfaceViewportCopyPS> PixelShader(Views[0].ShaderMap);
 
-		const FIntPoint InputTextureSize = TextureInput.GetViewport().Extent;
+		const FIntPoint InputTextureSize = SceneTexture->Desc.Extent;
 
 		GraphBuilder.AddPass(
 			RDG_EVENT_NAME("SubsurfaceViewportCopy"),
 			PassParameters,
 			ERenderGraphPassFlags::None,
-			[Context, PixelShader, InputTextureSize, PassParameters](FRHICommandListImmediate& RHICmdList)
+			[&Views, ViewMask, ViewCount, PixelShader, InputTextureSize, PassParameters](FRHICommandListImmediate& RHICmdList)
 		{
-			const FSceneViewFamily& ViewFamily = Context->ViewFamily;
-
-			for (uint32 ViewId = 0, ViewCount = ViewFamily.Views.Num(); ViewId < ViewCount; ++ViewId)
+			for (uint32 ViewIndex = 0; ViewIndex < ViewCount; ++ViewIndex)
 			{
-				const FViewInfo* LocalView = static_cast<const FViewInfo*>(ViewFamily.Views[ViewId]);
+				const uint32 ViewBit = 1 << ViewIndex;
+
+				const bool bIsNonSubsurfaceView = (ViewMask & ViewBit) == 0;
 
-				// Skip the view we are currently processing.
-				if (LocalView != &Context->View)
+				if (bIsNonSubsurfaceView)
 				{
-					const FIntRect Rect = LocalView->ViewRect;
+					const FViewInfo& View = Views[ViewIndex];
+					const FScreenPassViewInfo ScreenPassView(View);
 
-					DrawScreenPass(RHICmdList, Context, Rect, Rect, InputTextureSize, *PixelShader, *PassParameters);
+					DrawScreenPass(RHICmdList, ScreenPassView, View.ViewRect, View.ViewRect, InputTextureSize, *PixelShader, *PassParameters);
 				}
 			}
 		});
+
+		// Subsequent render passes should load the texture contents.
+		SceneTextureLoadAction = ERenderTargetLoadAction::ELoad;
 	}
 
-	// Recombines scattering result with scene color.
+	for (uint32 ViewIndex = 0; ViewIndex < ViewCount; ++ViewIndex)
 	{
-		const FScreenPassTexture& TextureInput = SceneTexture;
-		const FScreenPassTexture& TextureOutput = RecombineTexture;
+		const uint32 ViewBit = 1 << ViewIndex;
 
-		FSubsurfaceRecombinePS::FParameters* PassParameters = GraphBuilder.AllocParameters<FSubsurfaceRecombinePS::FParameters>();
-		PassParameters->Subsurface = SubsurfaceCommonParameters;
-		PassParameters->RenderTargets[0] = FRenderTargetBinding(TextureOutput.GetRDGTexture(), ERenderTargetLoadAction::ELoad, ERenderTargetStoreAction::EStore);
-		PassParameters->SubsurfaceInput0 = GetSubsurfaceInput(TextureInput);
-		PassParameters->SubsurfaceSampler0 = BilinearBorderSampler;
+		const bool bIsSubsurfaceView = (ViewMask & ViewBit) != 0;
 
-		// Scattering output target is only used when scattering is enabled.
-		if (SubsurfaceMode != ESubsurfaceMode::Bypass)
+		if (bIsSubsurfaceView)
 		{
-			PassParameters->SubsurfaceInput1 = GetSubsurfaceInput(SubsurfaceTextureY);
-			PassParameters->SubsurfaceSampler1 = BilinearBorderSampler;
-		}
+			RDG_EVENT_SCOPE(GraphBuilder, "SubsurfaceScattering(ViewId=%d)", ViewIndex);
 
-		const FSubsurfaceRecombinePS::EQuality RecombineQuality = FSubsurfaceRecombinePS::GetQuality(Context->View);
+			const FViewInfo& View = Views[ViewIndex];
+			const FScreenPassViewInfo ScreenPassView(View);
+			const FScreenPassTextureViewport SceneViewport(View.ViewRect, SceneTexture);
 
-		FSubsurfaceRecombinePS::FPermutationDomain PixelShaderPermutationVector;
-		PixelShaderPermutationVector.Set<FSubsurfaceRecombinePS::FDimensionMode>(SubsurfaceMode);
-		PixelShaderPermutationVector.Set<FSubsurfaceRecombinePS::FDimensionQuality>(RecombineQuality);
-		PixelShaderPermutationVector.Set<FSubsurfaceRecombinePS::FDimensionCheckerboard>(bCheckerboard);
-		TShaderMapRef<FSubsurfaceRecombinePS> PixelShader(Context->ShaderMap, PixelShaderPermutationVector);
+			ComputeSubsurfaceForView(GraphBuilder, ScreenPassView, SceneViewport, SceneTexture, SceneTextureOutput, SceneTextureLoadAction);
 
-		AddDrawScreenPass(GraphBuilder, RDG_EVENT_NAME("SubsurfaceRecombine"), Context, TextureOutput.GetViewport(), TextureInput.GetViewport(), *PixelShader, PassParameters);
+			// Subsequent render passes should load the texture contents.
+			SceneTextureLoadAction = ERenderTargetLoadAction::ELoad;
+		}
 	}
 
-	return RecombineTexture;
+	return SceneTextureOutput;
 }
 
-FScreenPassTexture VisualizeSubsurface(
+void VisualizeSubsurface(
 	FRDGBuilder& GraphBuilder,
-	FScreenPassContextRef Context,
-	const FScreenPassTexture& SceneTexture)
+	const FScreenPassViewInfo& ScreenPassView,
+	const FScreenPassTextureViewport& SceneViewport,
+	FRDGTextureRef SceneTexture,
+	FRDGTextureRef SceneTextureOutput)
 {
-	check(Context);
-	check(SceneTexture.IsValid());
+	check(SceneTexture);
+	check(SceneTextureOutput);
+	check(SceneViewport.Extent == SceneTexture->Desc.Extent);
 
-	const FScreenPassTexture& TextureInput = SceneTexture;
-	const FScreenPassTexture TextureOutput = FScreenPassTexture::Create(
-		GraphBuilder.CreateTexture(TextureInput.GetRDGTexture()->Desc, TEXT("SubsurfaceVisualize")),
-		TextureInput.GetViewport().Rect);
+	const FViewInfo& View = ScreenPassView.View;
 
 	FSubsurfaceVisualizePS::FParameters* PassParameters = GraphBuilder.AllocParameters<FSubsurfaceVisualizePS::FParameters>();
-	PassParameters->Subsurface = GetSubsurfaceCommonParameters(GraphBuilder.RHICmdList, Context);
-	PassParameters->RenderTargets[0] = FRenderTargetBinding(TextureOutput.GetRDGTexture(), ERenderTargetLoadAction::EClear, ERenderTargetStoreAction::EStore);
-	PassParameters->SubsurfaceInput0 = GetSubsurfaceInput(TextureInput);
+	PassParameters->Subsurface = GetSubsurfaceCommonParameters(GraphBuilder.RHICmdList, View);
+	PassParameters->RenderTargets[0] = FRenderTargetBinding(SceneTextureOutput, ERenderTargetLoadAction::EClear, ERenderTargetStoreAction::EStore);
+	PassParameters->SubsurfaceInput0.Texture = SceneTexture;
+	PassParameters->SubsurfaceInput0.Viewport = GetScreenPassTextureViewportParameters(SceneViewport);
 	PassParameters->SubsurfaceSampler0 = TStaticSamplerState<SF_Point, AM_Clamp, AM_Clamp, AM_Clamp>::GetRHI();
 	PassParameters->MiniFontTexture = GetMiniFontTexture();
 
-	TShaderMapRef<FSubsurfaceVisualizePS> PixelShader(Context->ShaderMap);
+	TShaderMapRef<FSubsurfaceVisualizePS> PixelShader(View.ShaderMap);
 
 	GraphBuilder.AddPass(
 		RDG_EVENT_NAME("SubsurfaceVisualize"),
 		PassParameters,
 		ERenderGraphPassFlags::None,
-		[Context, TextureInput, TextureOutput, PixelShader, PassParameters](FRHICommandListImmediate& RHICmdList)
+		[ScreenPassView, SceneViewport, SceneTextureOutput, PixelShader, PassParameters](FRHICommandListImmediate& RHICmdList)
 	{
-		DrawScreenPass(RHICmdList, Context, TextureOutput.GetViewport(), TextureInput.GetViewport(), *PixelShader, *PassParameters);
+		DrawScreenPass(RHICmdList, ScreenPassView, SceneViewport, SceneViewport, *PixelShader, *PassParameters);
 
 		// Draw debug text
 		{
-			const FSceneViewFamily& ViewFamily = Context->ViewFamily;
-			FRenderTargetTemp TempRenderTarget(static_cast<FTexture2DRHIParamRef>(TextureOutput.GetRHITexture()), TextureOutput.GetViewport().Extent);
-			FCanvas Canvas(&TempRenderTarget, nullptr, ViewFamily.CurrentRealTime, ViewFamily.CurrentWorldTime, ViewFamily.DeltaWorldTime, Context->View.GetFeatureLevel());
+			const FViewInfo& View = ScreenPassView.View;
+			const FSceneViewFamily& ViewFamily = *View.Family;
+			FRenderTargetTemp TempRenderTarget(static_cast<FTexture2DRHIParamRef>(SceneTextureOutput->GetRHITexture()), SceneTextureOutput->Desc.Extent);
+			FCanvas Canvas(&TempRenderTarget, nullptr, ViewFamily.CurrentRealTime, ViewFamily.CurrentWorldTime, ViewFamily.DeltaWorldTime, View.GetFeatureLevel());
 
 			float X = 30;
 			float Y = 28;
@@ -664,55 +733,63 @@ FScreenPassTexture VisualizeSubsurface(
 			Canvas.Flush_RenderThread(RHICmdList, bFlush, bInsideRenderPass);
 		}
 	});
-
-	return TextureOutput;
 }
 
-FSubsurfaceVisualizeCompositePass::FSubsurfaceVisualizeCompositePass(FRHICommandList& RHICmdList)
-{
-	// we need the GBuffer, we release it Process()
-	FSceneRenderTargets::Get(RHICmdList).AdjustGBufferRefCount(RHICmdList, 1);
-}
+//////////////////////////////////////////////////////////////////////////
+//! Shim methods to hook into the legacy pipeline until the full RDG conversion is complete.
 
-void FSubsurfaceVisualizeCompositePass::Process(FRenderingCompositePassContext& CompositePassContext)
+void ComputeSubsurfaceShim(FRHICommandListImmediate& RHICmdList, const TArray<FViewInfo>& Views)
 {
-	FRHICommandListImmediate& RHICmdList = CompositePassContext.RHICmdList;
-
-	const FViewInfo& View = CompositePassContext.View;
+	FSceneRenderTargets& SceneRenderTargets = FSceneRenderTargets::Get(RHICmdList);
 
 	FRDGBuilder GraphBuilder(RHICmdList);
 
-	FScreenPassContext* Context = GraphBuilder.AllocObject<FScreenPassContext>(RHICmdList, CompositePassContext.View);
+	FRDGTextureRef SceneTexture = GraphBuilder.RegisterExternalTexture(SceneRenderTargets.GetSceneColor(), TEXT("SceneColor"));
+
+	FRDGTextureRef SceneTextureOutput = ComputeSubsurface(GraphBuilder, SceneTexture, Views);
+
+	// Extract the result texture out and re-assign it to the scene render targets blackboard.
+	TRefCountPtr<IPooledRenderTarget> SceneTarget;
+	GraphBuilder.QueueTextureExtraction(SceneTextureOutput, &SceneTarget, false);
+	GraphBuilder.Execute();
+
+	SceneRenderTargets.SetSceneColor(SceneTarget);
 
-	const FScreenPassTexture TextureInput = FScreenPassTexture::Create(
-		CreateRDGTextureForInput(GraphBuilder, ePId_Input0, TEXT("SceneColor"), eFC_0000),
-		View.ViewRect);
+	// The RT should be released as early as possible to allow sharing of that memory for other purposes.
+	// This becomes even more important with some limited VRam (XBoxOne).
+	SceneRenderTargets.SetLightAttenuation(nullptr);
+}
+
+FRenderingCompositeOutputRef VisualizeSubsurfaceShim(
+	FRHICommandListImmediate& RHICmdList,
+	FRenderingCompositionGraph& Graph,
+	FRenderingCompositeOutputRef Input)
+{
+	// we need the GBuffer, we release it Process()
+	FSceneRenderTargets::Get(RHICmdList).AdjustGBufferRefCount(RHICmdList, 1);
 
-	const FScreenPassTexture TextureOutput = VisualizeSubsurface(GraphBuilder, Context, TextureInput);
+	FRenderingCompositePass* SubsurfaceVisualizePass = Graph.RegisterPass(new(FMemStack::Get()) TRCPassForRDG<1, 1>(
+		[](FRenderingCompositePass* Pass, FRenderingCompositePassContext& CompositePassContext)
+	{
+		FRDGBuilder GraphBuilder(CompositePassContext.RHICmdList);
 
-	ExtractRDGTextureForOutput(GraphBuilder, ePId_Output0, TextureOutput.GetRDGTexture());
+		FRDGTextureRef SceneTexture = Pass->CreateRDGTextureForInput(GraphBuilder, ePId_Input0, TEXT("SceneColor"), eFC_0000);
+		FRDGTextureRef SceneTextureOutput = Pass->FindOrCreateRDGTextureForOutput(GraphBuilder, ePId_Output0, SceneTexture->Desc, TEXT("SubsurfaceVisualize"));
 
-	TRefCountPtr<IPooledRenderTarget> OutputTarget;
-	GraphBuilder.QueueTextureExtraction(TextureOutput.GetRDGTexture(), &OutputTarget);
+		const FScreenPassViewInfo ScreenPassView(CompositePassContext.View);
+		const FScreenPassTextureViewport SceneViewport(CompositePassContext.View.ViewRect, SceneTexture->Desc.Extent);
+		VisualizeSubsurface(GraphBuilder, ScreenPassView, SceneViewport, SceneTexture, SceneTextureOutput);
 
-	GraphBuilder.Execute();
+		Pass->ExtractRDGTextureForOutput(GraphBuilder, ePId_Output0, SceneTextureOutput);
 
-	check(OutputTarget);
+		GraphBuilder.Execute();
 
-	RHICmdList.CopyToResolveTarget(
-		OutputTarget->GetRenderTargetItem().TargetableTexture,
-		OutputTarget->GetRenderTargetItem().ShaderResourceTexture,
-		FResolveParams());
+		FRHICommandListImmediate& RHICmdList = GraphBuilder.RHICmdList;
+		FSceneRenderTargets::Get(RHICmdList).AdjustGBufferRefCount(RHICmdList, -1);
+	}));
 
-	FSceneRenderTargets::Get(RHICmdList).AdjustGBufferRefCount(RHICmdList, -1);
+	SubsurfaceVisualizePass->SetInput(ePId_Input0, Input);
+	return FRenderingCompositeOutputRef(SubsurfaceVisualizePass);
 }
 
-FPooledRenderTargetDesc FSubsurfaceVisualizeCompositePass::ComputeOutputDesc(EPassOutputId InPassOutputId) const
-{
-	FPooledRenderTargetDesc Ret = FSceneRenderTargets::Get_FrameConstantsOnly().GetSceneColor()->GetDesc();
-	Ret.Flags &= ~(TexCreate_FastVRAM | TexCreate_Transient);
-	Ret.Reset();
-	Ret.DebugName = TEXT("SubsurfaceVisualize");
-	Ret.Format = PF_FloatRGBA;
-	return Ret;
-}
\ No newline at end of file
+//////////////////////////////////////////////////////////////////////////
\ No newline at end of file
diff --git a/Engine/Source/Runtime/Renderer/Private/PostProcess/PostProcessSubsurface.h b/Engine/Source/Runtime/Renderer/Private/PostProcess/PostProcessSubsurface.h
--- a/Engine/Source/Runtime/Renderer/Private/PostProcess/PostProcessSubsurface.h
+++ b/Engine/Source/Runtime/Renderer/Private/PostProcess/PostProcessSubsurface.h
@@ -18,28 +18,14 @@ bool IsSubsurfaceRequiredForView(const FViewInfo& View);
 // Returns whether checkerboard rendering is enabled for the provided format.
 bool IsSubsurfaceCheckerboardFormat(EPixelFormat SceneColorFormat);
 
-// Computes subsurface scattering on the scene texture and produces a new scene texture as output.
-FScreenPassTexture ComputeSubsurface(
-	FRDGBuilder& GraphBuilder,
-	FScreenPassContextRef Context,
-	const FScreenPassTexture& SceneTexture);
-
-// Visualizes subsurface scattering profiles by overlaying an image on the provided scene texture.
-// Produces a new scene texture as output.
-FScreenPassTexture VisualizeSubsurface(
-	FRDGBuilder& GraphBuilder,
-	FScreenPassContextRef Context,
-	const FScreenPassTexture& SceneTexture);
-
-// An adapter to connect the new Render Graph implementation to the legacy Composition Graph.
-class FSubsurfaceVisualizeCompositePass : public TRenderingCompositePassBase<1, 1>
-{
-public:
-	FSubsurfaceVisualizeCompositePass(FRHICommandList& CmdList);
-	virtual ~FSubsurfaceVisualizeCompositePass() = default;
-
-private:
-	void Process(FRenderingCompositePassContext& Context) override;
-	FPooledRenderTargetDesc ComputeOutputDesc(EPassOutputId InPassOutputId) const override;
-	void Release() override { delete this; }
-};
\ No newline at end of file
+//////////////////////////////////////////////////////////////////////////
+//! Shim methods to hook into the legacy pipeline until the full RDG conversion is complete.
+
+void ComputeSubsurfaceShim(FRHICommandListImmediate& RHICmdList, const TArray<FViewInfo>& Views);
+
+FRenderingCompositeOutputRef VisualizeSubsurfaceShim(
+	FRHICommandListImmediate& RHICmdList,
+	FRenderingCompositionGraph& Graph,
+	FRenderingCompositeOutputRef Input);
+
+//////////////////////////////////////////////////////////////////////////
\ No newline at end of file
diff --git a/Engine/Source/Runtime/Renderer/Private/PostProcess/PostProcessing.cpp b/Engine/Source/Runtime/Renderer/Private/PostProcess/PostProcessing.cpp
--- a/Engine/Source/Runtime/Renderer/Private/PostProcess/PostProcessing.cpp
+++ b/Engine/Source/Runtime/Renderer/Private/PostProcess/PostProcessing.cpp
@@ -1519,56 +1519,7 @@ void FPostProcessing::Process(FRHICommandListImmediate& RHICmdList, const FViewI
 			// Motion Blur
 			if ((IsMotionBlurEnabled(View) || bVisualizeMotionBlur) && VelocityInput.IsValid())
 			{
-				FRenderingCompositePass* MotionBlurPass = Context.Graph.RegisterPass(new(FMemStack::Get()) TRCPassForRDG<3, 1>(
-					[bVisualizeMotionBlur](FRenderingCompositePass* Pass, FRenderingCompositePassContext& InContext)
-				{
-					FRDGBuilder GraphBuilder(InContext.RHICmdList);
-
-					const FIntRect ColorViewportRect = InContext.SceneColorViewRect;
-					const FIntRect DepthViewportRect = InContext.View.ViewRect;
-
-					const FScreenPassTexture ColorTexture = FScreenPassTexture::Create(
-						Pass->CreateRDGTextureForInput(GraphBuilder, ePId_Input0, TEXT("SceneColor"), eFC_0000),
-						ColorViewportRect);
-
-					const FScreenPassTexture DepthTexture = FScreenPassTexture::Create(
-						Pass->CreateRDGTextureForInput(GraphBuilder, ePId_Input1, TEXT("SceneDepth"), eFC_0000),
-						DepthViewportRect);
-
-					const FScreenPassTexture VelocityTexture = FScreenPassTexture::Create(
-						Pass->CreateRDGTextureForInput(GraphBuilder, ePId_Input2, TEXT("SceneVelocity"), eFC_0000),
-						DepthViewportRect);
-
-					FScreenPassTexture OutColorTexture;
-
-					if (bVisualizeMotionBlur)
-					{
-						OutColorTexture = VisualizeMotionBlur(
-							GraphBuilder,
-							GraphBuilder.AllocObject<FScreenPassContext>(InContext.RHICmdList, InContext.View),
-							ColorTexture,
-							DepthTexture,
-							VelocityTexture);
-					}
-					else
-					{
-						OutColorTexture = ComputeMotionBlur(
-							GraphBuilder,
-							GraphBuilder.AllocObject<FScreenPassContext>(InContext.RHICmdList, InContext.View),
-							ColorTexture,
-							DepthTexture,
-							VelocityTexture);
-					}
-
-					Pass->ExtractRDGTextureForOutput(GraphBuilder, ePId_Output0, OutColorTexture.GetRDGTexture());
-
-					GraphBuilder.Execute();
-				}));
-
-				MotionBlurPass->SetInput(ePId_Input0, Context.FinalOutput);
-				MotionBlurPass->SetInput(ePId_Input1, Context.SceneDepth);
-				MotionBlurPass->SetInput(ePId_Input2, VelocityInput);
-				Context.FinalOutput = FRenderingCompositeOutputRef(MotionBlurPass, ePId_Output0);
+				Context.FinalOutput = ComputeMotionBlurShim(Context.Graph, Context.FinalOutput, Context.SceneDepth, VelocityInput, bVisualizeMotionBlur);
 			}
 
 			if(bVisualizeBloom)
@@ -1957,11 +1908,7 @@ void FPostProcessing::Process(FRHICommandListImmediate& RHICmdList, const FViewI
 		if (View.Family->EngineShowFlags.VisualizeSSS)
 		{
 			ensureMsgf(!bUnscaledFinalOutput, TEXT("VisualizeSSS is incompatible with unscaled output."));
-
-			// the setup pass also does visualization, based on EngineShowFlags.VisualizeSSS
-			FRenderingCompositePass* PassVisualize = Context.Graph.RegisterPass(new(FMemStack::Get()) FSubsurfaceVisualizeCompositePass(Context.RHICmdList));
-			PassVisualize->SetInput(ePId_Input0, Context.FinalOutput);
-			Context.FinalOutput = FRenderingCompositeOutputRef(PassVisualize);
+			Context.FinalOutput = VisualizeSubsurfaceShim(RHICmdList, Context.Graph, Context.FinalOutput);
 		}
 
 		AddGBufferVisualizationOverview(Context, SeparateTranslucency, PreTonemapHDRColor, PostTonemapHDRColor, PreFlattenVelocity);
diff --git a/Engine/Source/Runtime/Renderer/Private/PostProcess/RenderingCompositionGraph.cpp b/Engine/Source/Runtime/Renderer/Private/PostProcess/RenderingCompositionGraph.cpp
--- a/Engine/Source/Runtime/Renderer/Private/PostProcess/RenderingCompositionGraph.cpp
+++ b/Engine/Source/Runtime/Renderer/Private/PostProcess/RenderingCompositionGraph.cpp
@@ -1179,10 +1179,7 @@ FRDGTextureRef FRenderingCompositePass::CreateRDGTextureForInput(
 		InputName);
 }
 
-void FRenderingCompositePass::ExtractRDGTextureForOutput(
-	FRDGBuilder& GraphBuilder,
-	EPassOutputId OutputId,
-	FRDGTextureRef Texture)
+void FRenderingCompositePass::ExtractRDGTextureForOutput(FRDGBuilder& GraphBuilder, EPassOutputId OutputId, FRDGTextureRef Texture)
 {
 	if (FRenderingCompositeOutput* Output = GetOutput(OutputId))
 	{
@@ -1190,3 +1187,26 @@ void FRenderingCompositePass::ExtractRDGTextureForOutput(
 		GraphBuilder.QueueTextureExtraction(Texture, &Output->PooledRenderTarget, false);
 	}
 }
+
+FRDGTextureRef FRenderingCompositePass::FindOrCreateRDGTextureForOutput(
+	FRDGBuilder& GraphBuilder,
+	EPassOutputId OutputId,
+	const FRDGTextureDesc& TextureDesc,
+	const TCHAR* TextureName)
+{
+	if (FRenderingCompositeOutput* Output = GetOutput(OutputId))
+	{
+		const TRefCountPtr<IPooledRenderTarget>& ExistingTarget = Output->PooledRenderTarget;
+
+		if (ExistingTarget)
+		{
+			return GraphBuilder.RegisterExternalTexture(ExistingTarget, TextureName);
+		}
+		else
+		{
+			return GraphBuilder.CreateTexture(TextureDesc, TextureName);
+		}
+	}
+
+	return nullptr;
+}
\ No newline at end of file
diff --git a/Engine/Source/Runtime/Renderer/Private/PostProcess/RenderingCompositionGraph.h b/Engine/Source/Runtime/Renderer/Private/PostProcess/RenderingCompositionGraph.h
--- a/Engine/Source/Runtime/Renderer/Private/PostProcess/RenderingCompositionGraph.h
+++ b/Engine/Source/Runtime/Renderer/Private/PostProcess/RenderingCompositionGraph.h
@@ -347,6 +347,16 @@ struct FRenderingCompositePass
 		const TCHAR* InputName,
 		EFallbackColor FallbackColor);
 
+	/**
+	 * Attempts to find an persistent render target allocated for the output, or creates
+	 * a new one with the provided descriptor.
+	 */
+	FRDGTextureRef FindOrCreateRDGTextureForOutput(
+		FRDGBuilder& GraphBuilder,
+		EPassOutputId OutputId,
+		const FRDGTextureDesc& TextureDesc,
+		const TCHAR* TextureName);
+
 	/**
 	 * Registers a RDG texture to be extracted to the assigned output during graph execution.
 	 */
diff --git a/Engine/Source/Runtime/Renderer/Private/ScreenPass.cpp b/Engine/Source/Runtime/Renderer/Private/ScreenPass.cpp
--- a/Engine/Source/Runtime/Renderer/Private/ScreenPass.cpp
+++ b/Engine/Source/Runtime/Renderer/Private/ScreenPass.cpp
@@ -104,29 +104,6 @@ FScreenPassTextureViewportParameters GetScreenPassTextureViewportParameters(cons
 	return Parameters;
 }
 
-FScreenPassTexture FScreenPassTexture::Create(FRDGTextureRef InTexture, FIntRect InViewport)
-{
-	check(InTexture);
-
-	FScreenPassTexture Texture;
-	Texture.Texture = InTexture;
-	Texture.Viewport = FScreenPassTextureViewport(InViewport, InTexture->Desc.Extent);
-	return Texture;
-}
-
-FScreenPassTexture FScreenPassTexture::CreateFullscreen(FRDGTextureRef InTexture)
-{
-	check(InTexture);
-
-	const FIntPoint Extent = InTexture->Desc.Extent;
-	const FIntRect PixelViewRect = FIntRect(FIntPoint::ZeroValue, Extent);
-
-	FScreenPassTexture Texture;
-	Texture.Texture = InTexture;
-	Texture.Viewport = FScreenPassTextureViewport(PixelViewRect, Extent);
-	return Texture;
-}
-
 FScreenPassTextureViewportTransform GetScreenPassTextureViewportTransform(
 	FVector2D SourceOffset,
 	FVector2D SourceExtent,
@@ -151,26 +128,10 @@ FScreenPassTextureViewportTransform GetScreenPassTextureViewportTransform(
 	return GetScreenPassTextureViewportTransform(SourceUVOffset, SourceUVExtent, DestinationUVOffset, DestinationUVExtent);
 }
 
-FScreenPassCommonParameters GetScreenPassCommonParameters(FRHICommandListImmediate& RHICmdList, const FViewInfo& View)
-{
-	FSceneRenderTargets& SceneContext = FSceneRenderTargets::Get(RHICmdList);
-
-	FScreenPassCommonParameters Parameters;
-	Parameters.BilinearTextureSampler = TStaticSamplerState<SF_Bilinear>::GetRHI();
-	Parameters.ViewUniformBuffer = View.ViewUniformBuffer;
-	Parameters.SceneUniformBuffer = CreateSceneTextureUniformBuffer(
-		SceneContext, View.FeatureLevel, ESceneTextureSetupMode::All, EUniformBufferUsage::UniformBuffer_SingleFrame);
-	return Parameters;
-}
-
-FScreenPassContext::FScreenPassContext(FRHICommandListImmediate& RHICmdList, const FViewInfo& InView)
+FScreenPassViewInfo::FScreenPassViewInfo(const FViewInfo& InView)
 	: View(InView)
-	, ViewFamily(*View.Family)
-	, ViewState(View.ViewState)
+	, ScreenPassVS(View.ShaderMap)
 	, StereoPass(View.StereoPass)
 	, bHasHMDMask(IsHMDHiddenAreaMaskActive())
 	, bUseComputePasses(ShouldDoComputePostProcessing(InView))
-	, ShaderMap(View.ShaderMap)
-	, ScreenPassVS(View.ShaderMap)
-	, ScreenPassCommonParameters(GetScreenPassCommonParameters(RHICmdList, View))
 {}
\ No newline at end of file
diff --git a/Engine/Source/Runtime/Renderer/Private/ScreenPass.h b/Engine/Source/Runtime/Renderer/Private/ScreenPass.h
--- a/Engine/Source/Runtime/Renderer/Private/ScreenPass.h
+++ b/Engine/Source/Runtime/Renderer/Private/ScreenPass.h
@@ -51,6 +51,11 @@ struct FScreenPassTextureViewport
 		, Extent(InExtent)
 	{}
 
+	FScreenPassTextureViewport(FIntRect InRect, FRDGTextureRef InTexture)
+		: Rect(InRect)
+		, Extent(InTexture->Desc.Extent)
+	{}
+
 	FScreenPassTextureViewport(const FScreenPassTextureViewport&) = default;
 
 	bool operator==(const FScreenPassTextureViewport& Other) const;
@@ -101,79 +106,6 @@ END_SHADER_PARAMETER_STRUCT()
 
 FScreenPassTextureViewportParameters GetScreenPassTextureViewportParameters(const FScreenPassTextureViewport& InViewport);
 
-/**
- * A class grouping a render graph texture and its shader parameters. The screen pass texture
- * is either in an empty state or a valid state. The default constructor initializes to an empty state.
- * Instantiating an instance from valid state requires using the provided @ref Create functions, which
- * validate that the shader parameter invariant holds.
- */
-class FScreenPassTexture
-{
-public:
-	/**
-	 * Creates a screen pass texture with a viewport.
-	 * @param InTexture A valid render graph texture instance. Asserts if the texture is null.
-	 * @param InViewport A viewport contained within the texture extent.
-	 */
-	static FScreenPassTexture Create(FRDGTextureRef InTexture, FIntRect InViewport);
-
-	/**
-	 * Creates a full-screen pass texture with a viewport covering the full texture.
-	 * @param InTexture A valid render graph texture instance. Asserts if the texture is null.
-	 */
-	static FScreenPassTexture CreateFullscreen(FRDGTextureRef InTexture);
-
-	// Default initializes to empty.
-	FScreenPassTexture() = default;
-
-	// Copy construction is permitted.
-	FScreenPassTexture(const FScreenPassTexture&) = default;
-
-	inline FRDGTextureRef GetRDGTexture() const
-	{
-		return Texture;
-	}
-
-	inline FScreenPassTextureViewport GetViewport() const
-	{
-		return Viewport;
-	}
-
-	inline FScreenPassTextureViewportParameters GetViewportParameters() const
-	{
-		return GetScreenPassTextureViewportParameters(Viewport);
-	}
-
-	inline const FTextureRHIParamRef GetRHITexture() const
-	{
-		return Texture ? Texture->GetRHITexture() : nullptr;
-	}
-
-	inline bool IsValid() const
-	{
-		return Texture != nullptr;
-	}
-
-	inline bool IsEmpty() const
-	{
-		return Texture == nullptr;
-	}
-
-	inline bool operator==(const FScreenPassTexture& Other) const
-	{
-		return Texture == Other.Texture && Viewport == Other.Viewport;
-	}
-
-	inline bool operator!=(const FScreenPassTexture& Other) const
-	{
-		return !(*this == Other);
-	}
-
-private:
-	FRDGTextureRef Texture = nullptr;
-	FScreenPassTextureViewport Viewport;
-};
-
 /**
  * Contains a transform that maps UV coordinates from one screen pass texture viewport to another.
  * Assumes normalized UV coordinates [0, 0]x[1, 1] where [0, 0] maps to the source view min
@@ -200,36 +132,16 @@ FScreenPassTextureViewportTransform GetScreenPassTextureViewportTransform(
 	FVector2D DestinationUVOffset,
 	FVector2D DestinationUVExtent);
 
-// Defines the common set of parameters for a screen space pass.
-BEGIN_SHADER_PARAMETER_STRUCT(FScreenPassCommonParameters, )
-	SHADER_PARAMETER_STRUCT_REF(FSceneTexturesUniformParameters, SceneUniformBuffer)
-	SHADER_PARAMETER_STRUCT_REF(FViewUniformShaderParameters, ViewUniformBuffer)
-	SHADER_PARAMETER_SAMPLER(SamplerState, BilinearTextureSampler)
-END_SHADER_PARAMETER_STRUCT()
-
-FScreenPassCommonParameters GetScreenPassCommonParameters(FRHICommandListImmediate& RHICmdList, const FViewInfo& View);
-
-/**
- * The context used for screen pass operations. Extracts and holds common state required
- * by most screen passes from the provided View instance. Exists to reduce pointer
- * chasing and reduce function parameter sizes as well as share immutable common state.
- */
-class FScreenPassContext
+class FScreenPassViewInfo
 {
 public:
+	FScreenPassViewInfo(const FViewInfo& InView);
+	FScreenPassViewInfo(const FScreenPassViewInfo&) = default;
 
-	UE_NONCOPYABLE(FScreenPassContext)
-
-	FScreenPassContext(FRHICommandListImmediate& RHICmdList, const FViewInfo& InView);
-
-	// The current view instance being processed.
 	const FViewInfo& View;
 
-	// The current view family instance being processed.
-	const FSceneViewFamily& ViewFamily;
-
-	// The current view state instance being processed.
-	const FSceneViewState* ViewState;
+	// The vertex shader used by draw screen pass. Cached here to avoid many lookups.
+	const TShaderMapRef<FScreenPassVS> ScreenPassVS;
 
 	// VR - Which stereo pass is being rendered.
 	const EStereoscopicPass StereoPass;
@@ -239,19 +151,8 @@ public:
 
 	// Whether screen passes should use compute.
 	const bool bUseComputePasses;
-
-	// The global shader map for the current view being processed.
-	const TShaderMap<FGlobalShaderType>* ShaderMap;
-
-	// The vertex shader used by draw screen pass. Cached here to avoid many lookups.
-	const TShaderMapRef<FScreenPassVS> ScreenPassVS;
-
-	// Common screen space parameters, filled at context creation time.
-	const FScreenPassCommonParameters ScreenPassCommonParameters;
 };
 
-using FScreenPassContextRef = const FScreenPassContext*;
-
 /**
  * Draws a full-viewport triangle with the provided pixel shader type. The destination full-viewport triangle
  * and interpolated source UV coordinates are derived from the viewport and texture rectangles, respectively.
@@ -264,7 +165,7 @@ using FScreenPassContextRef = const FScreenPassContext*;
 template<typename TPixelShaderType>
 void DrawScreenPass(
 	FRHICommandListImmediate& RHICmdList,
-	FScreenPassContextRef Context,
+	const FScreenPassViewInfo& ScreenPassView,
 	FIntRect OutputRect,
 	FIntRect InputRect,
 	FIntPoint InputSize,
@@ -273,7 +174,7 @@ void DrawScreenPass(
 {
 	const FIntPoint OutputSize = OutputRect.Size();
 	const FPixelShaderRHIParamRef PixelShaderRHI = GETSAFERHISHADER_PIXEL(PixelShader);
-	FScreenPassVS* VertexShader = *(Context->ScreenPassVS);
+	FScreenPassVS* VertexShader = *(ScreenPassView.ScreenPassVS);
 
 	FGraphicsPipelineStateInitializer GraphicsPSOInit;
 	RHICmdList.ApplyCachedRenderTargets(GraphicsPSOInit);
@@ -296,8 +197,8 @@ void DrawScreenPass(
 		OutputSize,
 		InputSize,
 		VertexShader,
-		Context->StereoPass,
-		Context->bHasHMDMask,
+		ScreenPassView.StereoPass,
+		ScreenPassView.bHasHMDMask,
 		EDRF_UseTriangleOptimization);
 }
 
@@ -309,7 +210,7 @@ void DrawScreenPass(
 template <typename TPixelShaderType>
 inline void DrawScreenPass(
 	FRHICommandListImmediate& RHICmdList,
-	FScreenPassContextRef Context,
+	const FScreenPassViewInfo& ScreenPassView,
 	const FScreenPassTextureViewport& OutputViewport,
 	const FScreenPassTextureViewport& InputViewport,
 	TPixelShaderType* PixelShader,
@@ -317,7 +218,7 @@ inline void DrawScreenPass(
 {
 	DrawScreenPass(
 		RHICmdList,
-		Context,
+		ScreenPassView,
 		OutputViewport.Rect,
 		InputViewport.Rect,
 		InputViewport.Extent,
@@ -333,14 +234,13 @@ template <typename TPixelShaderType>
 inline void AddDrawScreenPass(
 	FRDGBuilder& GraphBuilder,
 	FRDGEventName&& PassName,
-	FScreenPassContextRef Context,
-	FIntRect PixelOutputRect,
-	FIntRect PixelInputRect,
-	FIntPoint PixelInputSize,
+	const FScreenPassViewInfo& ScreenPassView,
+	FIntRect OutputViewportRect,
+	FIntRect InputViewportRect,
+	FIntPoint InputExtent,
 	TPixelShaderType* PixelShader,
 	typename TPixelShaderType::FParameters* PixelShaderParameters)
 {
-	check(Context);
 	check(PixelShader);
 	check(PixelShaderParameters);
 
@@ -350,15 +250,15 @@ inline void AddDrawScreenPass(
 		Forward<FRDGEventName>(PassName),
 		PixelShaderParameters,
 		ERenderGraphPassFlags::None,
-		[Context, PixelOutputRect, PixelInputRect, PixelInputSize, PixelShader, PixelShaderParameters]
+		[ScreenPassView, OutputViewportRect, InputViewportRect, InputExtent, PixelShader, PixelShaderParameters]
 		(FRHICommandListImmediate& RHICmdList)
 	{
 		DrawScreenPass(
 			RHICmdList,
-			Context,
-			PixelOutputRect,
-			PixelInputRect,
-			PixelInputSize,
+			ScreenPassView,
+			OutputViewportRect,
+			InputViewportRect,
+			InputExtent,
 			PixelShader,
 			*PixelShaderParameters);
 	});
@@ -371,7 +271,7 @@ template <typename TPixelShaderType>
 inline void AddDrawScreenPass(
 	FRDGBuilder& GraphBuilder,
 	FRDGEventName&& PassName,
-	FScreenPassContextRef Context,
+	const FScreenPassViewInfo& ScreenPassView,
 	const FScreenPassTextureViewport& OutputViewport,
 	const FScreenPassTextureViewport& InputViewport,
 	TPixelShaderType* PixelShader,
@@ -380,7 +280,7 @@ inline void AddDrawScreenPass(
 	AddDrawScreenPass(
 		GraphBuilder,
 		Forward<FRDGEventName>(PassName),
-		Context,
+		ScreenPassView,
 		OutputViewport.Rect,
 		InputViewport.Rect,
 		InputViewport.Extent,
