Fixes separate translucency fallback that was wrong, and yet unecessarily applied when separate translucency has been disabled.

This changes now explicitly write in code if a input is required using CreateRDGTextureForRequiredInput(), or optional CreateRDGTextureForOptionalInput() or should create a fallback with CreateRDGTextureForInputWithFallback() because still needed something because no permutation without it are existing.

diff --git a/Engine/Source/Runtime/Renderer/Private/PostProcess/PostProcessMotionBlur.cpp b/Engine/Source/Runtime/Renderer/Private/PostProcess/PostProcessMotionBlur.cpp
--- a/Engine/Source/Runtime/Renderer/Private/PostProcess/PostProcessMotionBlur.cpp
+++ b/Engine/Source/Runtime/Renderer/Private/PostProcess/PostProcessMotionBlur.cpp
@@ -908,9 +908,9 @@ FRenderingCompositeOutputRef ComputeMotionBlurShim(
 	{
 		FRDGBuilder GraphBuilder(InContext.RHICmdList);
 
-		FRDGTextureRef ColorTexture = Pass->CreateRDGTextureForInput(GraphBuilder, ePId_Input0, TEXT("SceneColor"), eFC_0000);
-		FRDGTextureRef DepthTexture = Pass->CreateRDGTextureForInput(GraphBuilder, ePId_Input1, TEXT("SceneDepth"), eFC_0000);
-		FRDGTextureRef VelocityTexture = Pass->CreateRDGTextureForInput(GraphBuilder, ePId_Input2, TEXT("SceneVelocity"), eFC_0000);
+		FRDGTextureRef ColorTexture = Pass->CreateRDGTextureForRequiredInput(GraphBuilder, ePId_Input0, TEXT("SceneColor"));
+		FRDGTextureRef DepthTexture = Pass->CreateRDGTextureForRequiredInput(GraphBuilder, ePId_Input1, TEXT("SceneDepth"));
+		FRDGTextureRef VelocityTexture = Pass->CreateRDGTextureForRequiredInput(GraphBuilder, ePId_Input2, TEXT("SceneVelocity"));
 
 		const FIntRect ColorViewportRect(InContext.SceneColorViewRect);
 		const FIntRect VelocityViewportRect(InContext.View.ViewRect);
diff --git a/Engine/Source/Runtime/Renderer/Private/PostProcess/PostProcessSubsurface.cpp b/Engine/Source/Runtime/Renderer/Private/PostProcess/PostProcessSubsurface.cpp
--- a/Engine/Source/Runtime/Renderer/Private/PostProcess/PostProcessSubsurface.cpp
+++ b/Engine/Source/Runtime/Renderer/Private/PostProcess/PostProcessSubsurface.cpp
@@ -773,7 +773,7 @@ FRenderingCompositeOutputRef VisualizeSubsurfaceShim(
 	{
 		FRDGBuilder GraphBuilder(CompositePassContext.RHICmdList);
 
-		FRDGTextureRef SceneTexture = Pass->CreateRDGTextureForInput(GraphBuilder, ePId_Input0, TEXT("SceneColor"), eFC_0000);
+		FRDGTextureRef SceneTexture = Pass->CreateRDGTextureForRequiredInput(GraphBuilder, ePId_Input0, TEXT("SceneColor"));
 		FRDGTextureRef SceneTextureOutput = Pass->FindOrCreateRDGTextureForOutput(GraphBuilder, ePId_Output0, SceneTexture->Desc, TEXT("SubsurfaceVisualize"));
 
 		const FScreenPassViewInfo ScreenPassView(CompositePassContext.View);
diff --git a/Engine/Source/Runtime/Renderer/Private/PostProcess/PostProcessTemporalAA.cpp b/Engine/Source/Runtime/Renderer/Private/PostProcess/PostProcessTemporalAA.cpp
--- a/Engine/Source/Runtime/Renderer/Private/PostProcess/PostProcessTemporalAA.cpp
+++ b/Engine/Source/Runtime/Renderer/Private/PostProcess/PostProcessTemporalAA.cpp
@@ -823,8 +823,8 @@ void FRCPassPostProcessTemporalAA::Process(FRenderingCompositePassContext& Conte
 	SetupSceneViewFamilyBlackboard(GraphBuilder, &SceneBlackboard);
 	
 	FTAAPassParameters Parameters = SavedParameters;
-	Parameters.SceneColorInput = CreateRDGTextureForInput(GraphBuilder, ePId_Input0, TEXT("SceneColor"), eFC_0000);
-	Parameters.SceneMetadataInput = CreateRDGTextureForInput(GraphBuilder, ePId_Input1, TEXT("SceneColor"), eFC_0000);
+	Parameters.SceneColorInput = CreateRDGTextureForRequiredInput(GraphBuilder, ePId_Input0, TEXT("SceneColor"));
+	Parameters.SceneMetadataInput = CreateRDGTextureForOptionalInput(GraphBuilder, ePId_Input1, TEXT("SceneColor"));
 
 	FTAAOutputs Outputs = Parameters.AddTemporalAAPass(
 		GraphBuilder,
diff --git a/Engine/Source/Runtime/Renderer/Private/PostProcess/PostProcessing.cpp b/Engine/Source/Runtime/Renderer/Private/PostProcess/PostProcessing.cpp
--- a/Engine/Source/Runtime/Renderer/Private/PostProcess/PostProcessing.cpp
+++ b/Engine/Source/Runtime/Renderer/Private/PostProcess/PostProcessing.cpp
@@ -1401,8 +1401,8 @@ void FPostProcessing::Process(FRHICommandListImmediate& RHICmdList, const FViewI
 					FSceneViewFamilyBlackboard SceneBlackboard;
 					SetupSceneViewFamilyBlackboard(GraphBuilder, &SceneBlackboard);
 	
-					FRDGTextureRef SceneColor = Pass->CreateRDGTextureForInput(GraphBuilder, ePId_Input0, TEXT("SceneColor"), eFC_0000);
-					FRDGTextureRef LocalSeparateTranslucency = Pass->CreateRDGTextureForInput(GraphBuilder, ePId_Input1, TEXT("SeparateTranslucency"), eFC_0000);
+					FRDGTextureRef SceneColor = Pass->CreateRDGTextureForRequiredInput(GraphBuilder, ePId_Input0, TEXT("SceneColor"));
+					FRDGTextureRef LocalSeparateTranslucency = Pass->CreateRDGTextureForOptionalInput(GraphBuilder, ePId_Input1, TEXT("SeparateTranslucency"));
 
 					FRDGTextureRef NewSceneColor = SceneColor;
 
@@ -1415,7 +1415,7 @@ void FPostProcessing::Process(FRHICommandListImmediate& RHICmdList, const FViewI
 					}
 
 					// DOF passes were not added, therefore need to compose Separate translucency manually. 
-					if (NewSceneColor == SceneColor)
+					if (NewSceneColor == SceneColor && LocalSeparateTranslucency)
 					{
 						NewSceneColor = AddSeparateTranslucencyCompositionPass(GraphBuilder, InContext.View, SceneColor, LocalSeparateTranslucency);
 					}
@@ -1784,10 +1784,14 @@ void FPostProcessing::Process(FRHICommandListImmediate& RHICmdList, const FViewI
 					FSceneViewFamilyBlackboard SceneBlackboard;
 					SetupSceneViewFamilyBlackboard(GraphBuilder, &SceneBlackboard);
 	
-					FRDGTextureRef SceneColor = Pass->CreateRDGTextureForInput(GraphBuilder, ePId_Input0, TEXT("SceneColor"), eFC_0000);
-					FRDGTextureRef LocalSeparateTranslucency = Pass->CreateRDGTextureForInput(GraphBuilder, ePId_Input1, TEXT("SeparateTranslucency"), eFC_0000);
+					FRDGTextureRef SceneColor = Pass->CreateRDGTextureForRequiredInput(GraphBuilder, ePId_Input0, TEXT("SceneColor"));
+					FRDGTextureRef LocalSeparateTranslucency = Pass->CreateRDGTextureForOptionalInput(GraphBuilder, ePId_Input1, TEXT("SeparateTranslucency"));
 
-					FRDGTextureRef NewSceneColor = AddSeparateTranslucencyCompositionPass(GraphBuilder, InContext.View, SceneColor, LocalSeparateTranslucency);
+					FRDGTextureRef NewSceneColor = SceneColor;
+					if (LocalSeparateTranslucency)
+					{
+						NewSceneColor = AddSeparateTranslucencyCompositionPass(GraphBuilder, InContext.View, SceneColor, LocalSeparateTranslucency);
+					}
 
 					Pass->ExtractRDGTextureForOutput(GraphBuilder, ePId_Output0, NewSceneColor);
 
diff --git a/Engine/Source/Runtime/Renderer/Private/PostProcess/RenderingCompositionGraph.cpp b/Engine/Source/Runtime/Renderer/Private/PostProcess/RenderingCompositionGraph.cpp
--- a/Engine/Source/Runtime/Renderer/Private/PostProcess/RenderingCompositionGraph.cpp
+++ b/Engine/Source/Runtime/Renderer/Private/PostProcess/RenderingCompositionGraph.cpp
@@ -1156,27 +1156,32 @@ FString FRenderingCompositePass::ConstructDebugName()
 	return Name;
 }
 
-FRDGTextureRef FRenderingCompositePass::CreateRDGTextureForInput(
+FRDGTextureRef FRenderingCompositePass::CreateRDGTextureForOptionalInput(
 	FRDGBuilder& GraphBuilder,
 	EPassInputId InputId,
-	const TCHAR* InputName,
-	EFallbackColor FallbackColor)
+	const TCHAR* InputName)
 {
-	TRefCountPtr<IPooledRenderTarget> RenderTarget;
-
 	if (const FRenderingCompositeOutputRef* OutputRef = GetInput(InputId))
 	{
 		if (FRenderingCompositeOutput* Input = OutputRef->GetOutput())
 		{
-			RenderTarget = Input->RequestInput();
+			return GraphBuilder.RegisterExternalTexture(Input->RequestInput(), InputName);
 		}
 	}
+	return nullptr;
+}
 
-	return RegisterExternalTextureWithFallback(
-		GraphBuilder,
-		RenderTarget,
-		GetFallbackTarget(FallbackColor),
-		InputName);
+FRDGTextureRef FRenderingCompositePass::CreateRDGTextureForInputWithFallback(
+	FRDGBuilder& GraphBuilder,
+	EPassInputId InputId,
+	const TCHAR* InputName,
+	EFallbackColor FallbackColor)
+{
+	if (FRDGTextureRef RDGTexture = CreateRDGTextureForOptionalInput(GraphBuilder, InputId, InputName))
+	{
+		return RDGTexture;
+	}
+	return GraphBuilder.RegisterExternalTexture(GetFallbackTarget(FallbackColor));
 }
 
 void FRenderingCompositePass::ExtractRDGTextureForOutput(FRDGBuilder& GraphBuilder, EPassOutputId OutputId, FRDGTextureRef Texture)
diff --git a/Engine/Source/Runtime/Renderer/Private/PostProcess/RenderingCompositionGraph.h b/Engine/Source/Runtime/Renderer/Private/PostProcess/RenderingCompositionGraph.h
--- a/Engine/Source/Runtime/Renderer/Private/PostProcess/RenderingCompositionGraph.h
+++ b/Engine/Source/Runtime/Renderer/Private/PostProcess/RenderingCompositionGraph.h
@@ -341,7 +341,22 @@ struct FRenderingCompositePass
 	 * Registers a RDG texture for the provided input. Returns a fallback color if the requested
 	 * input is null. Used for compatibility when porting to RDG.
 	 */
-	FRDGTextureRef CreateRDGTextureForInput(
+	FRDGTextureRef CreateRDGTextureForOptionalInput(
+		FRDGBuilder& GraphBuilder,
+		EPassInputId InputId,
+		const TCHAR* InputName);
+
+	FRDGTextureRef CreateRDGTextureForRequiredInput(
+		FRDGBuilder& GraphBuilder,
+		EPassInputId InputId,
+		const TCHAR* InputName)
+	{
+		FRDGTextureRef Texture = CreateRDGTextureForOptionalInput(GraphBuilder, InputId, InputName);
+		checkf(Texture, TEXT("Required texture parameter %s was not set on the post process input %d"), InputName, int32(InputId));
+		return Texture;
+	}
+
+	FRDGTextureRef CreateRDGTextureForInputWithFallback(
 		FRDGBuilder& GraphBuilder,
 		EPassInputId InputId,
 		const TCHAR* InputName,
@@ -692,8 +707,8 @@ protected:
  *	{
  *		FRDGBuilder GraphBuilder(Context.RHICmdList);
  *
- *		FRDGTextureRef SceneColor = Pass->CreateRDGTextureForInput(GraphBuilder, ePId_Input0, TEXT("SceneColor"), eFC_0000);
- *		FRDGTextureRef SeparateTranslucency = Pass->CreateRDGTextureForInput(GraphBuilder, ePId_Input1, TEXT("SeparateTranslucency"), eFC_0000);
+ *		FRDGTextureRef SceneColor = Pass->CreateRDGTextureForInput(GraphBuilder, ePId_Input0, TEXT("SceneColor"));
+ *		FRDGTextureRef SeparateTranslucency = Pass->CreateRDGTextureForInputWithFallback(GraphBuilder, ePId_Input1, TEXT("SeparateTranslucency"), eFC_0000);
  *
  *		// ...
  *
