Fixes a crashes when DOF is enabled, but separate translucency is disabled.

diff --git a/Engine/Source/Runtime/Renderer/Private/PostProcess/DiaphragmDOF.cpp b/Engine/Source/Runtime/Renderer/Private/PostProcess/DiaphragmDOF.cpp
--- a/Engine/Source/Runtime/Renderer/Private/PostProcess/DiaphragmDOF.cpp
+++ b/Engine/Source/Runtime/Renderer/Private/PostProcess/DiaphragmDOF.cpp
@@ -2401,7 +2401,7 @@ FRDGTextureRef DiaphragmDOF::AddPasses(
 		PassParameters->SceneColorInput = FullResGatherInputTextures.SceneColor;
 		PassParameters->SceneDepthTexture = SceneBlackboard.SceneDepthBuffer;
 		PassParameters->SceneSeparateCoc = FullResGatherInputTextures.SeparateCoc; // TODO looks useless.
-		PassParameters->SceneSeparateTranslucency = SceneSeparateTranslucency;
+		PassParameters->SceneSeparateTranslucency = SceneSeparateTranslucency ? SceneSeparateTranslucency : GraphBuilder.RegisterExternalTexture(GSystemTextures.BlackAlphaOneDummy);
 		
 		PassParameters->ConvolutionInputSize = FVector4(RefBufferSize.X, RefBufferSize.Y, 1.0f / RefBufferSize.X, 1.0f / RefBufferSize.Y);
 		PassParameters->ForegroundConvolution = ForegroundConvolutionTextures;
