Fixes a a bug on r.Test.SecondaryUpscaleOverride with reflection captures.

diff --git a/Engine/Source/Runtime/Renderer/Private/SceneRendering.cpp b/Engine/Source/Runtime/Renderer/Private/SceneRendering.cpp
--- a/Engine/Source/Runtime/Renderer/Private/SceneRendering.cpp
+++ b/Engine/Source/Runtime/Renderer/Private/SceneRendering.cpp
@@ -1909,7 +1909,7 @@ FSceneRenderer::FSceneRenderer(const FSceneViewFamily* InViewFamily,FHitProxyCon
 	}
 
 	// Override secondary screen percentage for testing purpose.
-	if (CVarTestSecondaryUpscaleOverride.GetValueOnGameThread() > 0)
+	if (CVarTestSecondaryUpscaleOverride.GetValueOnGameThread() > 0 && !Views[0].bIsReflectionCapture)
 	{
 		ViewFamily.SecondaryViewFraction = 1.0 / float(CVarTestSecondaryUpscaleOverride.GetValueOnGameThread());
 		ViewFamily.SecondaryScreenPercentageMethod = ESecondaryScreenPercentageMethod::NearestSpatialUpscale;
