Fixed assert when specifying stencil store but there is no stencil on a velocity render target.

diff --git a/Engine/Source/Runtime/Renderer/Private/PostProcess/PostProcessMotionBlur.cpp b/Engine/Source/Runtime/Renderer/Private/PostProcess/PostProcessMotionBlur.cpp
--- a/Engine/Source/Runtime/Renderer/Private/PostProcess/PostProcessMotionBlur.cpp
+++ b/Engine/Source/Runtime/Renderer/Private/PostProcess/PostProcessMotionBlur.cpp
@@ -533,7 +533,7 @@ void ComputeMotionBlurVelocity(
 				ERenderTargetLoadAction::ENoAction,
 				ERenderTargetStoreAction::EStore,
 				ERenderTargetLoadAction::ENoAction,
-				ERenderTargetStoreAction::EStore,
+				ERenderTargetStoreAction::ENoAction,
 				FExclusiveDepthStencil::DepthWrite_StencilNop);
 
 		PassParameters->RenderTargets[0] =
