- fixed incorrect use of ERenderTargetLoad/StoreAction in clustered deferred pass

diff --git a/Engine/Source/Runtime/Renderer/Private/ClusteredDeferredShadingPass.cpp b/Engine/Source/Runtime/Renderer/Private/ClusteredDeferredShadingPass.cpp
--- a/Engine/Source/Runtime/Renderer/Private/ClusteredDeferredShadingPass.cpp
+++ b/Engine/Source/Runtime/Renderer/Private/ClusteredDeferredShadingPass.cpp
@@ -133,7 +133,7 @@ void FDeferredShadingSceneRenderer::AddClusteredDeferredShadingPass(FRHICommandL
 			PassParameters->SSProfilesTexture = RegisterExternalTextureWithFallback(GraphBuilder, GetSubsufaceProfileTexture_RT(RHICmdList), GSystemTextures.BlackDummy);
 			PassParameters->TransmissionProfilesLinearSampler = TStaticSamplerState<SF_Bilinear, AM_Clamp, AM_Clamp, AM_Clamp>::GetRHI();
 
-			PassParameters->RenderTargets[0] = FRenderTargetBinding(GraphBuilder.RegisterExternalTexture(SceneContext.GetSceneColor()), ERenderTargetLoadAction::ENoAction, ERenderTargetStoreAction::ENoAction);
+			PassParameters->RenderTargets[0] = FRenderTargetBinding(GraphBuilder.RegisterExternalTexture(SceneContext.GetSceneColor()), ERenderTargetLoadAction::ELoad, ERenderTargetStoreAction::EStore);
 			// NOTE: if we wanted to get depth/stencil testing happening we'd need to add that here too, at the moment we don't.
 
 			GraphBuilder.AddPass(
@@ -149,6 +149,7 @@ void FDeferredShadingSceneRenderer::AddClusteredDeferredShadingPass(FRHICommandL
 				TShaderMapRef<FClusteredShadingPS> PixelShader(View.ShaderMap, PermutationVector);
 				{
 					FGraphicsPipelineStateInitializer GraphicsPSOInit;
+					InRHICmdList.ApplyCachedRenderTargets(GraphicsPSOInit);
 
 					// Additive blend to accumulate lighting contributions.
 					GraphicsPSOInit.BlendState = TStaticBlendState<CW_RGBA, BO_Add, BF_One, BF_One, BO_Add, BF_One, BF_One>::GetRHI();
@@ -161,7 +162,7 @@ void FDeferredShadingSceneRenderer::AddClusteredDeferredShadingPass(FRHICommandL
 					GraphicsPSOInit.BoundShaderState.PixelShaderRHI = GETSAFERHISHADER_PIXEL(*PixelShader);
 					GraphicsPSOInit.PrimitiveType = PT_TriangleList;
 
-					SetGraphicsPipelineState(InRHICmdList, GraphicsPSOInit, EApplyRendertargetOption::ForceApply);
+					SetGraphicsPipelineState(InRHICmdList, GraphicsPSOInit);
 				}
 				InRHICmdList.SetViewport(View.ViewRect.Min.X, View.ViewRect.Min.Y, 0.0f, View.ViewRect.Max.X, View.ViewRect.Max.Y, 1.0f);
 
