Added CVar r.RayTracing.ForceAllRayTracingEffects to force enabling/disabling all the ray tracing effects at once

diff --git a/Engine/Source/Runtime/Renderer/Private/DeferredShadingRenderer.cpp b/Engine/Source/Runtime/Renderer/Private/DeferredShadingRenderer.cpp
--- a/Engine/Source/Runtime/Renderer/Private/DeferredShadingRenderer.cpp
+++ b/Engine/Source/Runtime/Renderer/Private/DeferredShadingRenderer.cpp
@@ -105,7 +105,6 @@ static TAutoConsoleVariable<int32> CVarParallelBasePass(
 );
 
 static int32 GRayTracing = 0;
-
 static TAutoConsoleVariable<int32> CVarRayTracing(
 	TEXT("r.RayTracing"),
 	GRayTracing,
@@ -114,6 +113,16 @@ static TAutoConsoleVariable<int32> CVarRayTracing(
 	TEXT(" 1: on"),
 	ECVF_RenderThreadSafe | ECVF_ReadOnly);
 
+static int32 GForceAllRayTracingEffects = -1;
+static TAutoConsoleVariable<int32> CVarForceAllRayTracingEffects(
+	TEXT("r.RayTracing.ForceAllRayTracingEffects"),
+	GForceAllRayTracingEffects,
+	TEXT("Force all ray tracing effects ON/OFF.\n")
+	TEXT(" -1: Do not force (default) \n")
+	TEXT(" 0: All ray tracing effects disabled\n")
+	TEXT(" 1: All ray tracing effects enabled"),
+	ECVF_RenderThreadSafe);
+
 static TAutoConsoleVariable<int32> CVarUseAODenoiser(
 	TEXT("r.AmbientOcclusion.Denoiser"),
 	2,
@@ -123,14 +132,6 @@ static TAutoConsoleVariable<int32> CVarUseAODenoiser(
 	TEXT(" 2: GScreenSpaceDenoiser witch may be overriden by a third party plugin (default)."),
 	ECVF_RenderThreadSafe);
 
-static TAutoConsoleVariable<int32> CVarRayTracingTranslucency(
-	TEXT("r.RayTracing.Translucency"),
-	-1,
-	TEXT("-1: Value driven by postprocess volume (default) \n")
-	TEXT(" 0: ray tracing translucency off (use raster) \n")
-	TEXT(" 1: ray tracing translucency enabled"),
-	ECVF_RenderThreadSafe);
-
 #if !UE_BUILD_SHIPPING
 static TAutoConsoleVariable<int32> CVarForceBlackVelocityBuffer(
 	TEXT("r.Test.ForceBlackVelocityBuffer"), 0,
@@ -1885,18 +1886,7 @@ void FDeferredShadingSceneRenderer::Render(FRHICommandListImmediate& RHICmdList)
 		RHICmdList.SetCurrentStat(GET_STATID(STAT_CLM_Translucency));
 
 #if RHI_RAYTRACING
-		bool bAnyViewWithRaytracingTranslucency = false;
-		for (int32 ViewIndex = 0, Num = Views.Num(); ViewIndex < Num; ViewIndex++)
-		{
-			const FViewInfo& View = Views[ViewIndex];
-			//#dxr_todo: UE-72557 multiview case
-			bAnyViewWithRaytracingTranslucency = bAnyViewWithRaytracingTranslucency || (View.FinalPostProcessSettings.TranslucencyType == ETranslucencyType::RayTracing);
-		}
-
-		int32 RtTranslucencyCvar = CVarRayTracingTranslucency.GetValueOnRenderThread();
-		int32 bRaytracedTranslucency = RtTranslucencyCvar > -1 ? RtTranslucencyCvar : (bAnyViewWithRaytracingTranslucency? 1 : 0);
-
-		if (bRayTracingEnabled && bRaytracedTranslucency > 0)
+		if (ShouldRenderRayTracingTranslucency(Views))
 		{
 			ResolveSceneColor(RHICmdList);
 			RenderRayTracingTranslucency(RHICmdList);
@@ -2324,6 +2314,20 @@ void FDeferredShadingSceneRenderer::CopyStencilToLightingChannelTexture(FRHIComm
 }
 
 #if RHI_RAYTRACING
+
+int32 GetForceRayTracingEffectsCVarValue()
+{
+	if (IsRayTracingEnabled())
+	{
+		static auto CVar = IConsoleManager::Get().FindConsoleVariable(TEXT("r.RayTracing.ForceAllRayTracingEffects"));
+		return CVar != nullptr ? CVar->GetInt() : -1;
+	}
+	else
+	{
+		return 0;
+	}
+}
+
 bool CanOverlayRayTracingOutput(const FViewInfo& View)
 {
 	return (View.RayTracingRenderMode != ERayTracingRenderMode::PathTracing)
diff --git a/Engine/Source/Runtime/Renderer/Private/LightRendering.cpp b/Engine/Source/Runtime/Renderer/Private/LightRendering.cpp
--- a/Engine/Source/Runtime/Renderer/Private/LightRendering.cpp
+++ b/Engine/Source/Runtime/Renderer/Private/LightRendering.cpp
@@ -87,11 +87,29 @@ static FAutoConsoleVariableRef CVarDebugLightDiscardProp(
 
 
 
+#if RHI_RAYTRACING
+bool ShouldRenderRayTracingShadows(const FLightSceneProxy& LightProxy)
+{
+	const int32 ForceAllRayTracingEffects = GetForceRayTracingEffectsCVarValue();
+	const bool bRTShadowsEnabled = (ForceAllRayTracingEffects > 0 || (GRayTracingShadows > 0 && ForceAllRayTracingEffects < 0));
+
+	return IsRayTracingEnabled() && bRTShadowsEnabled && LightProxy.CastsRaytracedShadow();
+}
+
+bool ShouldRenderRayTracingShadows(const FLightSceneInfoCompact& LightInfo)
+{
+	const int32 ForceAllRayTracingEffects = GetForceRayTracingEffectsCVarValue();
+	const bool bRTShadowsEnabled = (ForceAllRayTracingEffects > 0 || (GRayTracingShadows > 0 && ForceAllRayTracingEffects < 0));
+
+	return IsRayTracingEnabled() && bRTShadowsEnabled && LightInfo.bCastRaytracedShadow;
+}
+#endif // RHI_RAYTRACING
+
+
 FLightOcclusionType GetLightOcclusionType(const FLightSceneProxy& Proxy)
 {
 #if RHI_RAYTRACING
-	const bool bCastRaytracedShadow = IsRayTracingEnabled() && GRayTracingShadows == 1 && Proxy.CastsRaytracedShadow();
-	return bCastRaytracedShadow ? FLightOcclusionType::Raytraced : FLightOcclusionType::Shadowmap;
+	return ShouldRenderRayTracingShadows(Proxy) ? FLightOcclusionType::Raytraced : FLightOcclusionType::Shadowmap;
 #else
 	return FLightOcclusionType::Shadowmap;
 #endif
@@ -100,8 +118,7 @@ FLightOcclusionType GetLightOcclusionType(const FLightSceneProxy& Proxy)
 FLightOcclusionType GetLightOcclusionType(const FLightSceneInfoCompact& LightInfo)
 {
 #if RHI_RAYTRACING
-	const bool bCastRaytracedShadow = IsRayTracingEnabled() && GRayTracingShadows == 1 && LightInfo.bCastRaytracedShadow;
-	return bCastRaytracedShadow ? FLightOcclusionType::Raytraced : FLightOcclusionType::Shadowmap;
+	return ShouldRenderRayTracingShadows(LightInfo) ? FLightOcclusionType::Raytraced : FLightOcclusionType::Shadowmap;
 #else
 	return FLightOcclusionType::Shadowmap;
 #endif
diff --git a/Engine/Source/Runtime/Renderer/Private/RayTracing/RayTracingAmbientOcclusion.cpp b/Engine/Source/Runtime/Renderer/Private/RayTracing/RayTracingAmbientOcclusion.cpp
--- a/Engine/Source/Runtime/Renderer/Private/RayTracing/RayTracingAmbientOcclusion.cpp
+++ b/Engine/Source/Runtime/Renderer/Private/RayTracing/RayTracingAmbientOcclusion.cpp
@@ -26,7 +26,10 @@ static FAutoConsoleVariableRef CVarRayTracingAmbientOcclusion(
 
 bool ShouldRenderRayTracingAmbientOcclusion()
 {
-	return IsRayTracingEnabled() && GRayTracingAmbientOcclusion != 0;
+	const int32 ForceAllRayTracingEffects = GetForceRayTracingEffectsCVarValue();
+	const bool bRTAOEnabled = (ForceAllRayTracingEffects > 0 || (GRayTracingAmbientOcclusion > 0 && ForceAllRayTracingEffects < 0));
+
+	return IsRayTracingEnabled() && bRTAOEnabled;
 }
 
 static int32 GRayTracingAmbientOcclusionSamplesPerPixel = -1;
diff --git a/Engine/Source/Runtime/Renderer/Private/RayTracing/RayTracingGlobalIllumination.cpp b/Engine/Source/Runtime/Renderer/Private/RayTracing/RayTracingGlobalIllumination.cpp
--- a/Engine/Source/Runtime/Renderer/Private/RayTracing/RayTracingGlobalIllumination.cpp
+++ b/Engine/Source/Runtime/Renderer/Private/RayTracing/RayTracingGlobalIllumination.cpp
@@ -179,6 +179,11 @@ void SetupLightParameters(
 
 bool ShouldRenderRayTracingGlobalIllumination(const TArray<FViewInfo>& Views)
 {
+	if (GetForceRayTracingEffectsCVarValue() >= 0)
+	{
+		return GetForceRayTracingEffectsCVarValue() > 0;
+	}
+
 	if (GRayTracingGlobalIllumination >= 0)
 	{
 		return (GRayTracingGlobalIllumination > 0);
@@ -325,11 +330,6 @@ void FDeferredShadingSceneRenderer::RenderRayTracingGlobalIllumination(
 	TRefCountPtr<IPooledRenderTarget>& AmbientOcclusionRT
 )
 {
-	if (GRayTracingGlobalIllumination == 0 || (GRayTracingGlobalIllumination == -1 && View.FinalPostProcessSettings.RayTracingGI == 0)) 
-	{
-		return;
-	}
-
 	SCOPED_GPU_STAT(RHICmdList, RayTracingGlobalIllumination);
 
 	FSceneRenderTargets& SceneContext = FSceneRenderTargets::Get(RHICmdList);
diff --git a/Engine/Source/Runtime/Renderer/Private/RayTracing/RayTracingRectLight.cpp b/Engine/Source/Runtime/Renderer/Private/RayTracing/RayTracingRectLight.cpp
--- a/Engine/Source/Runtime/Renderer/Private/RayTracing/RayTracingRectLight.cpp
+++ b/Engine/Source/Runtime/Renderer/Private/RayTracing/RayTracingRectLight.cpp
@@ -41,11 +41,13 @@ static FAutoConsoleVariableRef CVarRayTracingStochasticRecLightIsTextureImportan
 	TEXT("Enable importance sampling for rect light evaluation (default = 1)")
 );
 
-bool ShouldRenderRayTracingStochasticRectLight(const FLightSceneInfo& LightSceneInfo)
+bool ShouldRenderRayTracingStochasticRectLight(const FLightSceneInfo& LightInfo)
 {
-	return IsRayTracingEnabled() && GRayTracingStochasticRectLight == 1
-		&& LightSceneInfo.Proxy->CastsRaytracedShadow()
-		&& LightSceneInfo.Proxy->GetLightType() == LightType_Rect;
+	return IsRayTracingEnabled() 
+		&& GRayTracingStochasticRectLight == 1 
+		&& GetForceRayTracingEffectsCVarValue() != 0
+		&& LightInfo.Proxy->CastsRaytracedShadow()
+		&& LightInfo.Proxy->GetLightType() == LightType_Rect;
 }
 
 BEGIN_GLOBAL_SHADER_PARAMETER_STRUCT(FRectLightData, )
diff --git a/Engine/Source/Runtime/Renderer/Private/RayTracing/RayTracingTranslucency.cpp b/Engine/Source/Runtime/Renderer/Private/RayTracing/RayTracingTranslucency.cpp
--- a/Engine/Source/Runtime/Renderer/Private/RayTracing/RayTracingTranslucency.cpp
+++ b/Engine/Source/Runtime/Renderer/Private/RayTracing/RayTracingTranslucency.cpp
@@ -23,6 +23,16 @@
 #include "RayTracing/RaytracingOptions.h"
 #include "Raytracing/RaytracingLighting.h"
 
+
+static int32 GRayTracingTranslucency = -1;
+static TAutoConsoleVariable<int32> CVarRayTracingTranslucency(
+	TEXT("r.RayTracing.Translucency"),
+	GRayTracingTranslucency,
+	TEXT("-1: Value driven by postprocess volume (default) \n")
+	TEXT(" 0: ray tracing translucency off (use raster) \n")
+	TEXT(" 1: ray tracing translucency enabled"),
+	ECVF_RenderThreadSafe);
+
 static float GRayTracingTranslucencyMaxRoughness = -1;
 static FAutoConsoleVariableRef CVarRayTracingTranslucencyMaxRoughness(
 	TEXT("r.RayTracing.Translucency.MaxRoughness"),
@@ -95,6 +105,25 @@ static FAutoConsoleVariableRef CVarRayTracingTranslucencyRefraction(
 
 DECLARE_GPU_STAT_NAMED(RayTracingTranslucency, TEXT("Ray Tracing Translucency"));
 
+#if RHI_RAYTRACING
+bool ShouldRenderRayTracingTranslucency(const TArray<FViewInfo>& Views)
+{
+	bool bAnyViewWithRaytracingTranslucency = false;
+	for (int32 ViewIndex = 0, Num = Views.Num(); ViewIndex < Num; ViewIndex++)
+	{
+		const FViewInfo& View = Views[ViewIndex];
+		//#dxr_todo: UE-72557 multiview case
+		bAnyViewWithRaytracingTranslucency = bAnyViewWithRaytracingTranslucency || (View.FinalPostProcessSettings.TranslucencyType == ETranslucencyType::RayTracing);
+	}
+
+	const bool bTranslucencyCvarEnabled = GRayTracingTranslucency < 0 ? bAnyViewWithRaytracingTranslucency : GRayTracingTranslucency;
+	const int32 ForceAllRayTracingEffects = GetForceRayTracingEffectsCVarValue();
+	const bool bRayTracingTranslucencyEnabled = (ForceAllRayTracingEffects > 0 || (bTranslucencyCvarEnabled && ForceAllRayTracingEffects < 0));
+
+	return IsRayTracingEnabled() && bRayTracingTranslucencyEnabled;
+}
+#endif // RHI_RAYTRACING
+
 class FRayTracingTranslucencyRGS : public FGlobalShader
 {
 	DECLARE_GLOBAL_SHADER(FRayTracingTranslucencyRGS)
diff --git a/Engine/Source/Runtime/Renderer/Private/RayTracing/RaytracingOptions.h b/Engine/Source/Runtime/Renderer/Private/RayTracing/RaytracingOptions.h
--- a/Engine/Source/Runtime/Renderer/Private/RayTracing/RaytracingOptions.h
+++ b/Engine/Source/Runtime/Renderer/Private/RayTracing/RaytracingOptions.h
@@ -8,30 +8,70 @@
 
 #include "RHIDefinitions.h"
 
+class FSkyLightSceneProxy;
+class FViewInfo;
+class FLightSceneInfoCompact;
+class FLightSceneInfo;
 
-extern bool ShouldRenderRayTracingSkyLight(const FSkyLightSceneProxy* SkyLightSceneProxy);
 
 #if RHI_RAYTRACING
+extern int32 GetForceRayTracingEffectsCVarValue();
+
+extern bool ShouldRenderRayTracingSkyLight(const FSkyLightSceneProxy* SkyLightSceneProxy);
 extern bool ShouldRenderRayTracingAmbientOcclusion();
+extern bool ShouldRenderRayTracingReflections(const TArray<FViewInfo>& Views);
 extern bool ShouldRenderRayTracingGlobalIllumination(const TArray<FViewInfo>& Views);
-extern bool ShouldRenderRayTracingStochasticRectLight(const FLightSceneInfo& LightSceneInfo);
+extern bool ShouldRenderRayTracingTranslucency(const TArray<FViewInfo>& Views);
+extern bool ShouldRenderRayTracingShadows(const FLightSceneProxy& LightProxy);
+extern bool ShouldRenderRayTracingShadows(const FLightSceneInfoCompact& LightInfo);
+extern bool ShouldRenderRayTracingStochasticRectLight(const FLightSceneInfo& LightInfo);
 extern bool CanOverlayRayTracingOutput(const FViewInfo& View);
 
 extern float GetRaytracingMaxNormalBias();
 
 #else
 
+FORCEINLINE int32 GetForceRayTracingEffectsCVarValue()
+{
+	return 0;
+}
+
+FORCEINLINE bool ShouldRenderRayTracingSkyLight(const FSkyLightSceneProxy* SkyLightSceneProxy)
+{
+	return false;
+}
+
 FORCEINLINE bool ShouldRenderRayTracingAmbientOcclusion()
 {
 	return false;
 }
 
+FORCEINLINE bool ShouldRenderRayTracingReflections(const TArray<FViewInfo>& Views)
+{
+	return false;
+}
+
 FORCEINLINE bool ShouldRenderRayTracingGlobalIllumination(const TArray<FViewInfo>& Views)
 {
 	return false;
 }
 
-FORCEINLINE bool ShouldRenderRayTracingStochasticRectLight(const FLightSceneInfo& LightSceneInfo)
+FORCEINLINE bool ShouldRenderRayTracingTranslucency(const TArray<FViewInfo>& Views)
+{
+	return false;
+}
+
+FORCEINLINE bool ShouldRenderRayTracingShadows(const FLightSceneProxy& LightProxy)
+{
+	return false;
+}
+
+FORCEINLINE bool ShouldRenderRayTracingShadows(const FLightSceneInfoCompact& LightInfo)
+{
+	return false;
+}
+
+FORCEINLINE bool ShouldRenderRayTracingStochasticRectLight(const FLightSceneInfo& LightInfo)
 {
 	return false;
 }
diff --git a/Engine/Source/Runtime/Renderer/Private/RayTracing/RaytracingSkylight.cpp b/Engine/Source/Runtime/Renderer/Private/RayTracing/RaytracingSkylight.cpp
--- a/Engine/Source/Runtime/Renderer/Private/RayTracing/RaytracingSkylight.cpp
+++ b/Engine/Source/Runtime/Renderer/Private/RayTracing/RaytracingSkylight.cpp
@@ -5,16 +5,26 @@
 =============================================================================*/
 
 #include "DeferredShadingRenderer.h"
+#include "RayTracingOptions.h"
 
 static int32 GRayTracingSkyLight = 0;
+
+#if RHI_RAYTRACING
 bool ShouldRenderRayTracingSkyLight(const FSkyLightSceneProxy* SkyLightSceneProxy)
 {
-#if RHI_RAYTRACING
-	return IsRayTracingEnabled() && SkyLightSceneProxy && SkyLightSceneProxy->bCastRayTracedShadow;
-#else
-	return false;
-#endif
+	if (SkyLightSceneProxy != nullptr)
+	{
+		const int32 ForceAllRayTracingEffects = GetForceRayTracingEffectsCVarValue();
+		const bool bRayTracingSkyEnabled = (ForceAllRayTracingEffects > 0 || (SkyLightSceneProxy->bCastRayTracedShadow > 0 && ForceAllRayTracingEffects < 0));
+
+		return IsRayTracingEnabled() && bRayTracingSkyEnabled;
+	}
+	else
+	{
+		return false;
+	}
 }
+#endif
 
 #if RHI_RAYTRACING
 
diff --git a/Engine/Source/Runtime/Renderer/Private/ReflectionEnvironment.cpp b/Engine/Source/Runtime/Renderer/Private/ReflectionEnvironment.cpp
--- a/Engine/Source/Runtime/Renderer/Private/ReflectionEnvironment.cpp
+++ b/Engine/Source/Runtime/Renderer/Private/ReflectionEnvironment.cpp
@@ -190,6 +190,25 @@ bool IsReflectionCaptureAvailable()
 	return (!AllowStaticLightingVar || AllowStaticLightingVar->GetInt() != 0);
 }
 
+#if RHI_RAYTRACING
+bool ShouldRenderRayTracingReflections(const TArray<FViewInfo>& Views)
+{
+	bool bAnyViewWithRaytracingReflections = false;
+	for (int32 ViewIndex = 0, Num = Views.Num(); ViewIndex < Num; ViewIndex++)
+	{
+		const FViewInfo& View = Views[ViewIndex];
+		//#dxr_todo: UE-72557 multiview case
+		bAnyViewWithRaytracingReflections = bAnyViewWithRaytracingReflections || (View.FinalPostProcessSettings.ReflectionsType == EReflectionsType::RayTracing);
+	}
+
+	const bool bReflectionsCvarEnabled = GRayTracingReflections < 0 ? bAnyViewWithRaytracingReflections : GRayTracingReflections;	
+	const int32 ForceAllRayTracingEffects = GetForceRayTracingEffectsCVarValue();
+	const bool bReflectionPassEnabled = (ForceAllRayTracingEffects > 0 || (bReflectionsCvarEnabled && ForceAllRayTracingEffects < 0));
+
+	return IsRayTracingEnabled() && bReflectionPassEnabled;
+}
+#endif // RHI_RAYTRACING
+
 IMPLEMENT_GLOBAL_SHADER_PARAMETER_STRUCT(FReflectionUniformParameters, "ReflectionStruct");
 
 void SetupReflectionUniformParameters(const FViewInfo& View, FReflectionUniformParameters& OutParameters)
@@ -704,13 +723,10 @@ void FDeferredShadingSceneRenderer::RenderDeferredReflectionsAndSkyLighting(FRHI
 	// If we're currently capturing a reflection capture, output SpecularColor * IndirectIrradiance for metals so they are not black in reflections,
 	// Since we don't have multiple bounce specular reflections
 	bool bReflectionCapture = false;
-	bool bAnyViewWithRaytracingReflections = false;
 	for (int32 ViewIndex = 0, Num = Views.Num(); ViewIndex < Num; ViewIndex++)
 	{
 		const FViewInfo& View = Views[ViewIndex];
 		bReflectionCapture = bReflectionCapture || View.bIsReflectionCapture;
-		//#dxr_todo: UE-72557 multiview case
-		bAnyViewWithRaytracingReflections = bAnyViewWithRaytracingReflections || (View.FinalPostProcessSettings.ReflectionsType == EReflectionsType::RayTracing);
 	}
 
 	if (bReflectionCapture)
@@ -719,7 +735,7 @@ void FDeferredShadingSceneRenderer::RenderDeferredReflectionsAndSkyLighting(FRHI
 		return;	
 	}
 
-	const bool bRayTracedReflections = IsRayTracingEnabled() && (GRayTracingReflections < 0 ? bAnyViewWithRaytracingReflections : GRayTracingReflections);
+	const bool bRayTracedReflections = ShouldRenderRayTracingReflections(Views);
 
 	// The specular sky light contribution is also needed by RT Reflections as a fallback.
 	const bool bSkyLight = Scene->SkyLight
