Implements r.Test.FreezeTemporalSequences for the sake of visions of people working with noise all the time

diff --git a/Engine/Source/Runtime/Renderer/Private/RendererScene.cpp b/Engine/Source/Runtime/Renderer/Private/RendererScene.cpp
--- a/Engine/Source/Runtime/Renderer/Private/RendererScene.cpp
+++ b/Engine/Source/Runtime/Renderer/Private/RendererScene.cpp
@@ -155,7 +155,6 @@ FSceneViewState::FSceneViewState()
 	CachedVisibilityChunkIndex = INDEX_NONE;
 	MIDUsedCount = 0;
 	TemporalAASampleIndex = 0;
-	TemporalAASampleCount = 1;
 	FrameIndex = 0;
 	DistanceFieldTemporalSampleIndex = 0;
 	AOTileIntersectionResources = NULL;
diff --git a/Engine/Source/Runtime/Renderer/Private/ScenePrivate.h b/Engine/Source/Runtime/Renderer/Private/ScenePrivate.h
--- a/Engine/Source/Runtime/Renderer/Private/ScenePrivate.h
+++ b/Engine/Source/Runtime/Renderer/Private/ScenePrivate.h
@@ -827,14 +827,6 @@ private:
 	TArray<UMaterialInstanceDynamic*> MIDPool;
 	uint32 MIDUsedCount;
 
-	// if TemporalAA is on this cycles through 0..TemporalAASampleCount-1, ResetViewState() puts it back to 0
-	uint8 TemporalAASampleIndex;
-	// >= 1, 1 means there is no TemporalAA
-	uint8 TemporalAASampleCount;
-
-	// counts up by one each frame, warped in 0..7 range, ResetViewState() puts it back to 0
-	uint32 FrameIndex;
-
 	// counts up by one each frame, warped in 0..3 range, ResetViewState() puts it back to 0
 	int32 DistanceFieldTemporalSampleIndex;
 
@@ -848,6 +840,12 @@ private:
 	bool bRoundRobinOcclusionEnabled;
 
 public:
+	
+	// if TemporalAA is on this cycles through 0..TemporalAASampleCount-1, ResetViewState() puts it back to 0
+	int8 TemporalAASampleIndex;
+
+	// counts up by one each frame, warped in 0..7 range, ResetViewState() puts it back to 0
+	uint32 FrameIndex;
 
 	// Previous frame's view info to use.
 	FPreviousViewInfo PrevFrameViewInfo;
@@ -979,12 +977,6 @@ public:
 		return TemporalAASampleIndex;
 	}
 
-	// call after OnFrameRenderingSetup()
-	uint32 GetCurrentTemporalAASampleCount() const
-	{
-		return TemporalAASampleCount;
-	}
-
 	// Returns the index of the frame with a desired power of two modulus.
 	inline uint32 GetFrameIndex(uint32 Pow2Modulus) const
 	{
@@ -1009,29 +1001,6 @@ public:
 		ReleaseDynamicRHI();
 	}
 
-	// @param SampleCount 0 or 1 for no TemporalAA 
-	void OnFrameRenderingSetup(uint32 SampleCount, const FSceneViewFamily& Family)
-	{
-		if(!SampleCount)
-		{
-			SampleCount = 1;
-		}
-
-		TemporalAASampleCount = FMath::Min(SampleCount, (uint32)255);
-
-		if (!Family.bWorldIsPaused)
-		{
-			TemporalAASampleIndex++;
-
-			FrameIndex++;
-		}
-
-		if(TemporalAASampleIndex >= TemporalAASampleCount)
-		{
-			TemporalAASampleIndex = 0;
-		}
-	}
-
 	void SetupDistanceFieldTemporalOffset(const FSceneViewFamily& Family)
 	{
 		if (!Family.bWorldIsPaused)
diff --git a/Engine/Source/Runtime/Renderer/Private/SceneRendering.cpp b/Engine/Source/Runtime/Renderer/Private/SceneRendering.cpp
--- a/Engine/Source/Runtime/Renderer/Private/SceneRendering.cpp
+++ b/Engine/Source/Runtime/Renderer/Private/SceneRendering.cpp
@@ -1262,7 +1262,7 @@ void FViewInfo::SetupUniformBufferParameters(
 	{
 		ViewUniformShaderParameters.TemporalAAParams = FVector4(
 			ViewState->GetCurrentTemporalAASampleIndex(), 
-			ViewState->GetCurrentTemporalAASampleCount(),
+			0,
 			TemporalJitterPixels.X,
 			TemporalJitterPixels.Y);
 
diff --git a/Engine/Source/Runtime/Renderer/Private/SceneVisibility.cpp b/Engine/Source/Runtime/Renderer/Private/SceneVisibility.cpp
--- a/Engine/Source/Runtime/Renderer/Private/SceneVisibility.cpp
+++ b/Engine/Source/Runtime/Renderer/Private/SceneVisibility.cpp
@@ -225,6 +225,20 @@ static FAutoConsoleVariableRef CVarLightMaxDrawDistanceScale(
 	ECVF_Scalability | ECVF_RenderThreadSafe
 );
 
+#if !UE_BUILD_SHIPPING
+
+static TAutoConsoleVariable<float> CVarFreezeTemporalSequences(
+	TEXT("r.Test.FreezeTemporalSequences"), 0,
+	TEXT("Freezes all temporal sequences."),
+	ECVF_RenderThreadSafe);
+
+static TAutoConsoleVariable<float> CVarFreezeTemporalHistories(
+	TEXT("r.Test.FreezeTemporalHistories"), 0,
+	TEXT("Freezes all temporal histories as well as the temporal sequence."),
+	ECVF_RenderThreadSafe);
+
+#endif
+
 DECLARE_CYCLE_STAT(TEXT("Occlusion Readback"), STAT_CLMM_OcclusionReadback, STATGROUP_CommandListMarkers);
 DECLARE_CYCLE_STAT(TEXT("After Occlusion Readback"), STAT_CLMM_AfterOcclusionReadback, STATGROUP_CommandListMarkers);
 
@@ -2901,6 +2915,14 @@ void FSceneRenderer::PreVisibilityFrameSetup(FRHICommandListImmediate& RHICmdLis
 			}
 		}
 	}
+	
+#if UE_BUILD_SHIPPING
+	const bool bFreezeTemporalHistories = false;
+	const bool bFreezeTemporalSequences = false;
+#else
+	bool bFreezeTemporalHistories = CVarFreezeTemporalHistories.GetValueOnRenderThread() != 0;
+	bool bFreezeTemporalSequences = bFreezeTemporalHistories || CVarFreezeTemporalSequences.GetValueOnRenderThread() != 0;
+#endif
 
 	// Setup motion blur parameters (also check for camera movement thresholds)
 	for(int32 ViewIndex = 0;ViewIndex < Views.Num();ViewIndex++)
@@ -2943,151 +2965,166 @@ void FSceneRenderer::PreVisibilityFrameSetup(FRHICommandListImmediate& RHICmdLis
 		if (ViewState)
 		{
 			check(View.bViewStateIsReadOnly);
-			View.bViewStateIsReadOnly = ViewFamily.bWorldIsPaused || ViewFamily.EngineShowFlags.HitProxies;
+			View.bViewStateIsReadOnly = ViewFamily.bWorldIsPaused || ViewFamily.EngineShowFlags.HitProxies || bFreezeTemporalHistories;
 
 			ViewState->SetupDistanceFieldTemporalOffset(ViewFamily);
+
+			if (!View.bViewStateIsReadOnly && !bFreezeTemporalSequences)
+			{
+				ViewState->FrameIndex++;
+			}
 		}
+		
+		// Subpixel jitter for temporal AA
+		int32 CVarTemporalAASamplesValue = CVarTemporalAASamples.GetValueOnRenderThread();
 
-		if( View.AntiAliasingMethod == AAM_TemporalAA && ViewState )
-		{
-			// Subpixel jitter for temporal AA
-			int32 TemporalAASamples = CVarTemporalAASamples.GetValueOnRenderThread();
+		bool bTemporalUpsampling = View.PrimaryScreenPercentageMethod == EPrimaryScreenPercentageMethod::TemporalUpscale;
 		
-			if( TemporalAASamples > 1 && View.bAllowTemporalJitter )
-			{
-				float SampleX, SampleY;
+		// Apply a sub pixel offset to the view.
+		if (View.AntiAliasingMethod == AAM_TemporalAA && ViewState && (CVarTemporalAASamplesValue > 0 || bTemporalUpsampling) && View.bAllowTemporalJitter)
+		{
+			float EffectivePrimaryResolutionFraction = float(View.ViewRect.Width()) / float(View.GetSecondaryViewRectSize().X);
 
+			// Compute number of TAA samples.
+			int32 TemporalAASamples = CVarTemporalAASamplesValue;
+			{
 				if (Scene->GetFeatureLevel() < ERHIFeatureLevel::SM4)
 				{
 					// Only support 2 samples for mobile temporal AA.
 					TemporalAASamples = 2;
 				}
-
-				if( TemporalAASamples == 2 )
-				{
-					#if 0
-						// 2xMSAA
-						// Pattern docs: http://msdn.microsoft.com/en-us/library/windows/desktop/ff476218(v=vs.85).aspx
-						//   N.
-						//   .S
-						float SamplesX[] = { -4.0f/16.0f, 4.0/16.0f };
-						float SamplesY[] = { -4.0f/16.0f, 4.0/16.0f };
-					#else
-						// This pattern is only used for mobile.
-						// Shift to reduce blur.
-						float SamplesX[] = { -8.0f/16.0f, 0.0/16.0f };
-						float SamplesY[] = { /* - */ 0.0f/16.0f, 8.0/16.0f };
-					#endif
-					ViewState->OnFrameRenderingSetup(ARRAY_COUNT(SamplesX), ViewFamily);
-					uint32 Index = ViewState->GetCurrentTemporalAASampleIndex();
-					SampleX = SamplesX[ Index ];
-					SampleY = SamplesY[ Index ];
-				}
-				else if( TemporalAASamples == 3 )
-				{
-					// 3xMSAA
-					//   A..
-					//   ..B
-					//   .C.
-					// Rolling circle pattern (A,B,C).
-					float SamplesX[] = { -2.0f/3.0f,  2.0/3.0f,  0.0/3.0f };
-					float SamplesY[] = { -2.0f/3.0f,  0.0/3.0f,  2.0/3.0f };
-					ViewState->OnFrameRenderingSetup(ARRAY_COUNT(SamplesX), ViewFamily);
-					uint32 Index = ViewState->GetCurrentTemporalAASampleIndex();
-					SampleX = SamplesX[ Index ];
-					SampleY = SamplesY[ Index ];
-				}
-				else if( TemporalAASamples == 4 )
+				else if (bTemporalUpsampling)
 				{
-					// 4xMSAA
-					// Pattern docs: http://msdn.microsoft.com/en-us/library/windows/desktop/ff476218(v=vs.85).aspx
-					//   .N..
-					//   ...E
-					//   W...
-					//   ..S.
-					// Rolling circle pattern (N,E,S,W).
-					float SamplesX[] = { -2.0f/16.0f,  6.0/16.0f, 2.0/16.0f, -6.0/16.0f };
-					float SamplesY[] = { -6.0f/16.0f, -2.0/16.0f, 6.0/16.0f,  2.0/16.0f };
-					ViewState->OnFrameRenderingSetup(ARRAY_COUNT(SamplesX), ViewFamily);
-					uint32 Index = ViewState->GetCurrentTemporalAASampleIndex();
-					SampleX = SamplesX[ Index ];
-					SampleY = SamplesY[ Index ];
+					// When doing TAA upsample with screen percentage < 100%, we need extra temporal samples to have a
+					// constant temporal sample density for final output pixels to avoid output pixel aligned converging issues.
+					TemporalAASamples = float(TemporalAASamples) * FMath::Max(1.f, 1.f / (EffectivePrimaryResolutionFraction * EffectivePrimaryResolutionFraction));
 				}
-				else if( TemporalAASamples == 5 )
+				else if (CVarTemporalAASamplesValue == 5)
 				{
-					// Compressed 4 sample pattern on same vertical and horizontal line (less temporal flicker).
-					// Compressed 1/2 works better than correct 2/3 (reduced temporal flicker).
-					//   . N .
-					//   W . E
-					//   . S .
-					// Rolling circle pattern (N,E,S,W).
-					float SamplesX[] = {  0.0f/2.0f,  1.0/2.0f,  0.0/2.0f, -1.0/2.0f };
-					float SamplesY[] = { -1.0f/2.0f,  0.0/2.0f,  1.0/2.0f,  0.0/2.0f };
-					ViewState->OnFrameRenderingSetup(ARRAY_COUNT(SamplesX), ViewFamily);
-					uint32 Index = ViewState->GetCurrentTemporalAASampleIndex();
-					SampleX = SamplesX[ Index ];
-					SampleY = SamplesY[ Index ];
+					TemporalAASamples = 4;
 				}
-				else if (View.PrimaryScreenPercentageMethod == EPrimaryScreenPercentageMethod::TemporalUpscale)
-				{
-					// When doing TAA upsample with screen percentage < 100%, we need extra temporal samples to have a
-					// constant temporal sample density for final output pixels to avoid output pixel aligned converging issues.
-					float EffectivePrimaryResolutionFraction = float(View.ViewRect.Width()) / float(View.GetSecondaryViewRectSize().X);
-					int32 EffectiveTemporalAASamples = float(TemporalAASamples) * FMath::Max(1.f, 1.f / (EffectivePrimaryResolutionFraction * EffectivePrimaryResolutionFraction));
 
-					ViewState->OnFrameRenderingSetup(EffectiveTemporalAASamples, ViewFamily);
-					uint32 TemporalSampleIndex = ViewState->GetCurrentTemporalAASampleIndex();
+				TemporalAASamples = FMath::Clamp(TemporalAASamples, 1, 255);
+			}
 
-					// Uniformly distribute temporal jittering in [-.5; .5], because there is no longer any alignement of input and output pixels.
-					SampleX = Halton(TemporalSampleIndex + 1, 2) - 0.5f;
-					SampleY = Halton(TemporalSampleIndex + 1, 3) - 0.5f;
+			// Compute the new sample index in the temporal sequence.
+			int32 TemporalSampleIndex = ViewState->TemporalAASampleIndex + 1;
+			if(TemporalSampleIndex >= TemporalAASamples)
+			{
+				TemporalSampleIndex = 0;
+			}
 
-					View.MaterialTextureMipBias = -(FMath::Max(-FMath::Log2(EffectivePrimaryResolutionFraction), 0.0f) ) + CVarMinAutomaticViewMipBiasOffset.GetValueOnRenderThread();
-					View.MaterialTextureMipBias = FMath::Max(View.MaterialTextureMipBias, CVarMinAutomaticViewMipBias.GetValueOnRenderThread());
-				}
-				else
-				{
-					ViewState->OnFrameRenderingSetup(TemporalAASamples, ViewFamily);
-					uint32 Index = ViewState->GetCurrentTemporalAASampleIndex();
+			// Updates view state.
+			if (!View.bViewStateIsReadOnly && !bFreezeTemporalSequences)
+			{
+				ViewState->TemporalAASampleIndex = TemporalSampleIndex;
+			}
 
-					float u1 = Halton( Index + 1, 2 );
-					float u2 = Halton( Index + 1, 3 );
+			// Choose sub pixel sample coordinate in the temporal sequence.
+			float SampleX, SampleY;
+			if (Scene->GetFeatureLevel() < ERHIFeatureLevel::SM4)
+			{
+				float SamplesX[] = { -8.0f/16.0f, 0.0/16.0f };
+				float SamplesY[] = { /* - */ 0.0f/16.0f, 8.0/16.0f };
+				check(TemporalAASamples == ARRAY_COUNT(SamplesX));
+				SampleX = SamplesX[ TemporalSampleIndex ];
+				SampleY = SamplesY[ TemporalSampleIndex ];
+			}
+			else if (View.PrimaryScreenPercentageMethod == EPrimaryScreenPercentageMethod::TemporalUpscale)
+			{
+				// Uniformly distribute temporal jittering in [-.5; .5], because there is no longer any alignement of input and output pixels.
+				SampleX = Halton(TemporalSampleIndex + 1, 2) - 0.5f;
+				SampleY = Halton(TemporalSampleIndex + 1, 3) - 0.5f;
+
+				View.MaterialTextureMipBias = -(FMath::Max(-FMath::Log2(EffectivePrimaryResolutionFraction), 0.0f) ) + CVarMinAutomaticViewMipBiasOffset.GetValueOnRenderThread();
+				View.MaterialTextureMipBias = FMath::Max(View.MaterialTextureMipBias, CVarMinAutomaticViewMipBias.GetValueOnRenderThread());
+			}
+			else if( CVarTemporalAASamplesValue == 2 )
+			{
+				// 2xMSAA
+				// Pattern docs: http://msdn.microsoft.com/en-us/library/windows/desktop/ff476218(v=vs.85).aspx
+				//   N.
+				//   .S
+				float SamplesX[] = { -4.0f/16.0f, 4.0/16.0f };
+				float SamplesY[] = { -4.0f/16.0f, 4.0/16.0f };
+				check(TemporalAASamples == ARRAY_COUNT(SamplesX));
+				SampleX = SamplesX[ TemporalSampleIndex ];
+				SampleY = SamplesY[ TemporalSampleIndex ];
+			}
+			else if( CVarTemporalAASamplesValue == 3 )
+			{
+				// 3xMSAA
+				//   A..
+				//   ..B
+				//   .C.
+				// Rolling circle pattern (A,B,C).
+				float SamplesX[] = { -2.0f/3.0f,  2.0/3.0f,  0.0/3.0f };
+				float SamplesY[] = { -2.0f/3.0f,  0.0/3.0f,  2.0/3.0f };
+				check(TemporalAASamples == ARRAY_COUNT(SamplesX));
+				SampleX = SamplesX[ TemporalSampleIndex ];
+				SampleY = SamplesY[ TemporalSampleIndex ];
+			}
+			else if( CVarTemporalAASamplesValue == 4 )
+			{
+				// 4xMSAA
+				// Pattern docs: http://msdn.microsoft.com/en-us/library/windows/desktop/ff476218(v=vs.85).aspx
+				//   .N..
+				//   ...E
+				//   W...
+				//   ..S.
+				// Rolling circle pattern (N,E,S,W).
+				float SamplesX[] = { -2.0f/16.0f,  6.0/16.0f, 2.0/16.0f, -6.0/16.0f };
+				float SamplesY[] = { -6.0f/16.0f, -2.0/16.0f, 6.0/16.0f,  2.0/16.0f };
+				check(TemporalAASamples == ARRAY_COUNT(SamplesX));
+				SampleX = SamplesX[ TemporalSampleIndex ];
+				SampleY = SamplesY[ TemporalSampleIndex ];
+			}
+			else if( CVarTemporalAASamplesValue == 5 )
+			{
+				// Compressed 4 sample pattern on same vertical and horizontal line (less temporal flicker).
+				// Compressed 1/2 works better than correct 2/3 (reduced temporal flicker).
+				//   . N .
+				//   W . E
+				//   . S .
+				// Rolling circle pattern (N,E,S,W).
+				float SamplesX[] = {  0.0f/2.0f,  1.0/2.0f,  0.0/2.0f, -1.0/2.0f };
+				float SamplesY[] = { -1.0f/2.0f,  0.0/2.0f,  1.0/2.0f,  0.0/2.0f };
+				check(TemporalAASamples == ARRAY_COUNT(SamplesX));
+				SampleX = SamplesX[ TemporalSampleIndex ];
+				SampleY = SamplesY[ TemporalSampleIndex ];
+			}
+			else
+			{
+				float u1 = Halton( TemporalSampleIndex + 1, 2 );
+				float u2 = Halton( TemporalSampleIndex + 1, 3 );
 
-					// Generates samples in normal distribution
-					// exp( x^2 / Sigma^2 )
+				// Generates samples in normal distribution
+				// exp( x^2 / Sigma^2 )
 					
-					static auto CVar = IConsoleManager::Get().FindConsoleVariable(TEXT("r.TemporalAAFilterSize"));
-					float FilterSize = CVar->GetFloat();
+				static auto CVar = IConsoleManager::Get().FindConsoleVariable(TEXT("r.TemporalAAFilterSize"));
+				float FilterSize = CVar->GetFloat();
 
-					// Scale distribution to set non-unit variance
-					// Variance = Sigma^2
-					float Sigma = 0.47f * FilterSize;
+				// Scale distribution to set non-unit variance
+				// Variance = Sigma^2
+				float Sigma = 0.47f * FilterSize;
 
-					// Window to [-0.5, 0.5] output
-					// Without windowing we could generate samples far away on the infinite tails.
-					float OutWindow = 0.5f;
-					float InWindow = FMath::Exp( -0.5 * FMath::Square( OutWindow / Sigma ) );
+				// Window to [-0.5, 0.5] output
+				// Without windowing we could generate samples far away on the infinite tails.
+				float OutWindow = 0.5f;
+				float InWindow = FMath::Exp( -0.5 * FMath::Square( OutWindow / Sigma ) );
 					
-					// Box-Muller transform
-					float Theta = 2.0f * PI * u2;
-					float r = Sigma * FMath::Sqrt( -2.0f * FMath::Loge( (1.0f - u1) * InWindow + u1 ) );
+				// Box-Muller transform
+				float Theta = 2.0f * PI * u2;
+				float r = Sigma * FMath::Sqrt( -2.0f * FMath::Loge( (1.0f - u1) * InWindow + u1 ) );
 					
-					SampleX = r * FMath::Cos( Theta );
-					SampleY = r * FMath::Sin( Theta );
-				}
-
-				View.TemporalJitterPixels.X = SampleX;
-				View.TemporalJitterPixels.Y = SampleY;
-
-				View.ViewMatrices.HackAddTemporalAAProjectionJitter(FVector2D(SampleX * 2.0f / View.ViewRect.Width(), SampleY * -2.0f / View.ViewRect.Height()));
+				SampleX = r * FMath::Cos( Theta );
+				SampleY = r * FMath::Sin( Theta );
 			}
-		}
-		else if(ViewState && !View.bViewStateIsReadOnly)
-		{
-			// no TemporalAA
-			ViewState->OnFrameRenderingSetup(1, ViewFamily);
 
-			ViewState->PrevFrameViewInfo.TemporalAAHistory.SafeRelease();
+			View.TemporalJitterPixels.X = SampleX;
+			View.TemporalJitterPixels.Y = SampleY;
+
+			View.ViewMatrices.HackAddTemporalAAProjectionJitter(FVector2D(SampleX * 2.0f / View.ViewRect.Width(), SampleY * -2.0f / View.ViewRect.Height()));
 		}
 
 		// Setup a new FPreviousViewInfo from current frame infos.
