Fixed compilation error in linux

diff --git a/Engine/Source/Runtime/Renderer/Private/RayTracing/RaytracingSkylight.cpp b/Engine/Source/Runtime/Renderer/Private/RayTracing/RaytracingSkylight.cpp
--- a/Engine/Source/Runtime/Renderer/Private/RayTracing/RaytracingSkylight.cpp
+++ b/Engine/Source/Runtime/Renderer/Private/RayTracing/RaytracingSkylight.cpp
@@ -5,27 +5,9 @@
 =============================================================================*/
 
 #include "DeferredShadingRenderer.h"
-#include "RayTracingOptions.h"
 
 static int32 GRayTracingSkyLight = 0;
 
-#if RHI_RAYTRACING
-bool ShouldRenderRayTracingSkyLight(const FSkyLightSceneProxy* SkyLightSceneProxy)
-{
-	if (SkyLightSceneProxy != nullptr)
-	{
-		const int32 ForceAllRayTracingEffects = GetForceRayTracingEffectsCVarValue();
-		const bool bRayTracingSkyEnabled = (ForceAllRayTracingEffects > 0 || (SkyLightSceneProxy->bCastRayTracedShadow > 0 && ForceAllRayTracingEffects < 0));
-
-		return IsRayTracingEnabled() && bRayTracingSkyEnabled;
-	}
-	else
-	{
-		return false;
-	}
-}
-#endif
-
 #if RHI_RAYTRACING
 
 #include "RayTracingSkyLight.h"
@@ -86,6 +68,21 @@ static TAutoConsoleVariable<int32> CVarRayTracingSkyLightEnableTwoSidedGeometry(
 	ECVF_RenderThreadSafe
 );
 
+bool ShouldRenderRayTracingSkyLight(const FSkyLightSceneProxy* SkyLightSceneProxy)
+{
+	if (SkyLightSceneProxy != nullptr)
+	{
+		const int32 ForceAllRayTracingEffects = GetForceRayTracingEffectsCVarValue();
+		const bool bRayTracingSkyEnabled = (ForceAllRayTracingEffects > 0 || (SkyLightSceneProxy->bCastRayTracedShadow > 0 && ForceAllRayTracingEffects < 0));
+
+		return IsRayTracingEnabled() && bRayTracingSkyEnabled;
+	}
+	else
+	{
+		return false;
+	}
+}
+
 IMPLEMENT_GLOBAL_SHADER_PARAMETER_STRUCT(FSkyLightData, "SkyLight");
 
 void SetupSkyLightParameters(
