Fixes UV0 texture maping of the vis command.

diff --git a/Engine/Source/Runtime/Renderer/Private/VisualizeTexturePresent.cpp b/Engine/Source/Runtime/Renderer/Private/VisualizeTexturePresent.cpp
--- a/Engine/Source/Runtime/Renderer/Private/VisualizeTexturePresent.cpp
+++ b/Engine/Source/Runtime/Renderer/Private/VisualizeTexturePresent.cpp
@@ -410,7 +410,7 @@ void FVisualizeTexturePresent::PresentContent(FRHICommandListImmediate& RHICmdLi
 		{
 		case 0:
 		{
-			DestRect = SrcRect;
+			DestRect = View.UnconstrainedViewRect;
 		}
 		break;
 
