Fixed translucency console command to work with threadsafe.

diff --git a/Engine/Source/Runtime/Renderer/Private/RayTracing/RayTracingTranslucency.cpp b/Engine/Source/Runtime/Renderer/Private/RayTracing/RayTracingTranslucency.cpp
--- a/Engine/Source/Runtime/Renderer/Private/RayTracing/RayTracingTranslucency.cpp
+++ b/Engine/Source/Runtime/Renderer/Private/RayTracing/RayTracingTranslucency.cpp
@@ -24,10 +24,9 @@
 #include "Raytracing/RaytracingLighting.h"
 
 
-static int32 GRayTracingTranslucency = -1;
 static TAutoConsoleVariable<int32> CVarRayTracingTranslucency(
 	TEXT("r.RayTracing.Translucency"),
-	GRayTracingTranslucency,
+	0,
 	TEXT("-1: Value driven by postprocess volume (default) \n")
 	TEXT(" 0: ray tracing translucency off (use raster) \n")
 	TEXT(" 1: ray tracing translucency enabled"),
@@ -116,6 +115,7 @@ bool ShouldRenderRayTracingTranslucency(const TArray<FViewInfo>& Views)
 		bAnyViewWithRaytracingTranslucency = bAnyViewWithRaytracingTranslucency || (View.FinalPostProcessSettings.TranslucencyType == ETranslucencyType::RayTracing);
 	}
 
+	const int32 GRayTracingTranslucency = CVarRayTracingTranslucency.GetValueOnRenderThread();
 	const bool bTranslucencyCvarEnabled = GRayTracingTranslucency < 0 ? bAnyViewWithRaytracingTranslucency : GRayTracingTranslucency;
 	const int32 ForceAllRayTracingEffects = GetForceRayTracingEffectsCVarValue();
 	const bool bRayTracingTranslucencyEnabled = (ForceAllRayTracingEffects > 0 || (bTranslucencyCvarEnabled && ForceAllRayTracingEffects < 0));
