Restored ScenePrimitiveId to VisibleMeshDrawCommand, which enables robustly accessing ScenePrimitive properties.

diff --git a/Engine/Source/Runtime/Renderer/Private/MeshDrawCommands.cpp b/Engine/Source/Runtime/Renderer/Private/MeshDrawCommands.cpp
--- a/Engine/Source/Runtime/Renderer/Private/MeshDrawCommands.cpp
+++ b/Engine/Source/Runtime/Renderer/Private/MeshDrawCommands.cpp
@@ -148,11 +148,8 @@ void UpdateTranslucentMeshSortKeys(
 	{
 		FVisibleMeshDrawCommand& VisibleCommand = VisibleMeshCommands[CommandIndex];
 
-		FVector BoundsOrigin = FVector::ZeroVector;
-		if (VisibleCommand.DrawPrimitiveId < PrimitiveBounds.Num())
-		{
-			BoundsOrigin = PrimitiveBounds[VisibleCommand.DrawPrimitiveId].BoxSphereBounds.Origin;
-		}
+		const int32 PrimitiveIndex = VisibleCommand.ScenePrimitiveId;
+		const FVector BoundsOrigin = PrimitiveIndex >= 0 ? PrimitiveBounds[PrimitiveIndex].BoxSphereBounds.Origin : FVector::ZeroVector;
 
 		float Distance = 0.0f;
 		if (TranslucentSortPolicy == ETranslucentSortPolicy::SortByDistance)
@@ -262,9 +259,9 @@ void MergeMobileBasePassMeshDrawCommands(
 			FVisibleMeshDrawCommand& MeshCommandCSM = MeshCommandsCSM[i];
 
 			if (MobileCSMVisibilityInfo.bAlwaysUseCSM 
-				|| (MeshCommand.DrawPrimitiveId < ScenePrimitiveNum && MobileCSMVisibilityInfo.MobilePrimitiveCSMReceiverVisibilityMap[MeshCommand.DrawPrimitiveId]))
+				|| (MeshCommand.ScenePrimitiveId < ScenePrimitiveNum && MobileCSMVisibilityInfo.MobilePrimitiveCSMReceiverVisibilityMap[MeshCommand.ScenePrimitiveId]))
 			{
-				checkf(MeshCommand.DrawPrimitiveId == MeshCommandCSM.DrawPrimitiveId, TEXT("VisibleMeshDrawCommands of BasePass and MobileBasePassCSM are expected to match."));
+				checkf(MeshCommand.ScenePrimitiveId == MeshCommandCSM.ScenePrimitiveId, TEXT("VisibleMeshDrawCommands of BasePass and MobileBasePassCSM are expected to match."));
 				// Use CSM's VisibleMeshDrawCommand.
 				MeshCommand = MeshCommandCSM;
 			}
@@ -297,9 +294,9 @@ void UpdateMobileBasePassMeshSortKeys(
 			bool bMasked = Cmd.SortKey.PackedData & 0x1 ? true : false; 
 			bool bBackground = Cmd.SortKey.PackedData & 0x2 ? true : false;
 			float PrimitiveDistance = 0;
-			if (Cmd.DrawPrimitiveId < ScenePrimitiveBounds.Num())
+			if (Cmd.ScenePrimitiveId < ScenePrimitiveBounds.Num())
 			{
-				const FPrimitiveBounds& PrimitiveBounds = ScenePrimitiveBounds[Cmd.DrawPrimitiveId];
+				const FPrimitiveBounds& PrimitiveBounds = ScenePrimitiveBounds[Cmd.ScenePrimitiveId];
 				PrimitiveDistance = (PrimitiveBounds.BoxSphereBounds.Origin - ViewOrigin).Size();
 				bBackground|= (PrimitiveBounds.BoxSphereBounds.SphereRadius > HALF_WORLD_MAX / 4.0f);
 			}
@@ -319,9 +316,9 @@ void UpdateMobileBasePassMeshSortKeys(
 		{
 			FVisibleMeshDrawCommand& Cmd = VisibleMeshCommands[CmdIdx];
 			float PrimitiveDistance = 0;
-			if (Cmd.DrawPrimitiveId < ScenePrimitiveBounds.Num())
+			if (Cmd.ScenePrimitiveId < ScenePrimitiveBounds.Num())
 			{
-				const FPrimitiveBounds& PrimitiveBounds = ScenePrimitiveBounds[Cmd.DrawPrimitiveId];
+				const FPrimitiveBounds& PrimitiveBounds = ScenePrimitiveBounds[Cmd.ScenePrimitiveId];
 				PrimitiveDistance = (PrimitiveBounds.BoxSphereBounds.Origin - ViewOrigin).Size();
 			}
 
@@ -338,9 +335,9 @@ void UpdateMobileBasePassMeshSortKeys(
 			bool bMasked = Cmd.SortKey.PackedData & 0x1 ? true : false; 
 			bool bBackground = Cmd.SortKey.PackedData & 0x2 ? true : false;
 			float PrimitiveDistance = 0;
-			if (Cmd.DrawPrimitiveId < ScenePrimitiveBounds.Num())
+			if (Cmd.ScenePrimitiveId < ScenePrimitiveBounds.Num())
 			{
-				const FPrimitiveBounds& PrimitiveBounds = ScenePrimitiveBounds[Cmd.DrawPrimitiveId];
+				const FPrimitiveBounds& PrimitiveBounds = ScenePrimitiveBounds[Cmd.ScenePrimitiveId];
 				PrimitiveDistance = (PrimitiveBounds.BoxSphereBounds.Origin - ViewOrigin).Size();
 				bBackground|= (PrimitiveBounds.BoxSphereBounds.SphereRadius > HALF_WORLD_MAX / 4.0f);
 			}
@@ -428,6 +425,7 @@ void BuildMeshDrawCommandPrimitiveIdBuffer(
 					NewVisibleMeshDrawCommand.Setup(
 						&NewCommand,
 						VisibleMeshDrawCommand.DrawPrimitiveId,
+						VisibleMeshDrawCommand.ScenePrimitiveId,
 						VisibleMeshDrawCommand.StateBucketId,
 						VisibleMeshDrawCommand.MeshFillMode,
 						VisibleMeshDrawCommand.MeshCullMode,
@@ -691,6 +689,7 @@ void ApplyViewOverridesToMeshDrawCommands(
 				NewVisibleMeshDrawCommand.Setup(
 					&NewMeshCommand,
 					VisibleMeshDrawCommand.DrawPrimitiveId,
+					VisibleMeshDrawCommand.ScenePrimitiveId,
 					VisibleMeshDrawCommand.StateBucketId,
 					VisibleMeshDrawCommand.MeshFillMode,
 					VisibleMeshDrawCommand.MeshCullMode,
diff --git a/Engine/Source/Runtime/Renderer/Private/MeshPassProcessor.cpp b/Engine/Source/Runtime/Renderer/Private/MeshPassProcessor.cpp
--- a/Engine/Source/Runtime/Renderer/Private/MeshPassProcessor.cpp
+++ b/Engine/Source/Runtime/Renderer/Private/MeshPassProcessor.cpp
@@ -1083,9 +1083,13 @@ ERasterizerCullMode FMeshPassProcessor::ComputeMeshCullMode(const FMeshBatch& Me
 	return bMeshRenderTwoSided ? CM_None : (bInReverseCullModeOverride ? CM_CCW : CM_CW);
 }
 
-int32 FMeshPassProcessor::GetDrawCommandPrimitiveId(const FPrimitiveSceneInfo* RESTRICT PrimitiveSceneInfo, const FMeshBatchElement& BatchElement) const
+void FMeshPassProcessor::GetDrawCommandPrimitiveId(
+	const FPrimitiveSceneInfo* RESTRICT PrimitiveSceneInfo,
+	const FMeshBatchElement& BatchElement,
+	int32& DrawPrimitiveId,
+	int32& ScenePrimitiveId) const
 {
-	int32 DrawPrimitiveId;
+	DrawPrimitiveId = 0;
 
 	if (UseGPUScene(GMaxRHIShaderPlatform, FeatureLevel))
 	{
@@ -1101,15 +1105,10 @@ int32 FMeshPassProcessor::GetDrawCommandPrimitiveId(const FPrimitiveSceneInfo* R
 		else
 		{
 			check(BatchElement.PrimitiveIdMode == PrimID_ForceZero);
-			DrawPrimitiveId = 0;
 		}
 	}
-	else
-	{
-		DrawPrimitiveId = PrimitiveSceneInfo ? PrimitiveSceneInfo->GetIndex() : INT32_MAX;
-	}
 
-	return DrawPrimitiveId;
+	ScenePrimitiveId = PrimitiveSceneInfo ? PrimitiveSceneInfo->GetIndex() : -1;
 }
 
 FCachedPassMeshDrawListContext::FCachedPassMeshDrawListContext(FCachedMeshDrawCommandInfo& InCommandInfo, FCachedPassMeshDrawList& InDrawList, FScene& InScene) :
@@ -1141,6 +1140,7 @@ void FCachedPassMeshDrawListContext::FinalizeCommand(
 	const FMeshBatch& MeshBatch, 
 	int32 BatchElementIndex,
 	int32 DrawPrimitiveId,
+	int32 ScenePrimitiveId,
 	ERasterizerFillMode MeshFillMode,
 	ERasterizerCullMode MeshCullMode,
 	FMeshDrawCommandSortKey SortKey,
diff --git a/Engine/Source/Runtime/Renderer/Private/SceneVisibility.cpp b/Engine/Source/Runtime/Renderer/Private/SceneVisibility.cpp
--- a/Engine/Source/Runtime/Renderer/Private/SceneVisibility.cpp
+++ b/Engine/Source/Runtime/Renderer/Private/SceneVisibility.cpp
@@ -1732,6 +1732,7 @@ struct FDrawCommandRelevancePacket
 				NewVisibleMeshDrawCommand.Setup(
 					MeshDrawCommand,
 					PrimitiveIndex,
+					PrimitiveIndex,
 					CachedMeshDrawCommand.StateBucketId,
 					CachedMeshDrawCommand.MeshFillMode,
 					CachedMeshDrawCommand.MeshCullMode,
diff --git a/Engine/Source/Runtime/Renderer/Private/ShadowSetup.cpp b/Engine/Source/Runtime/Renderer/Private/ShadowSetup.cpp
--- a/Engine/Source/Runtime/Renderer/Private/ShadowSetup.cpp
+++ b/Engine/Source/Runtime/Renderer/Private/ShadowSetup.cpp
@@ -832,6 +832,7 @@ void FProjectedShadowInfo::AddCachedMeshDrawCommandsForPass(
 			NewVisibleMeshDrawCommand.Setup(
 				MeshDrawCommand,
 				PrimitiveIndex,
+				PrimitiveIndex,
 				CachedMeshDrawCommand.StateBucketId,
 				CachedMeshDrawCommand.MeshFillMode,
 				CachedMeshDrawCommand.MeshCullMode,
@@ -1363,6 +1364,7 @@ void FProjectedShadowInfo::ApplyViewOverridesToMeshDrawCommands(const FViewInfo&
 			NewVisibleMeshDrawCommand.Setup(
 				&NewMeshCommand,
 				VisibleMeshDrawCommand.DrawPrimitiveId,
+				VisibleMeshDrawCommand.ScenePrimitiveId,
 				VisibleMeshDrawCommand.StateBucketId,
 				VisibleMeshDrawCommand.MeshFillMode,
 				VisibleMeshDrawCommand.MeshCullMode,
diff --git a/Engine/Source/Runtime/Renderer/Public/MeshPassProcessor.h b/Engine/Source/Runtime/Renderer/Public/MeshPassProcessor.h
--- a/Engine/Source/Runtime/Renderer/Public/MeshPassProcessor.h
+++ b/Engine/Source/Runtime/Renderer/Public/MeshPassProcessor.h
@@ -623,6 +623,7 @@ public:
 		const FMeshBatch& MeshBatch, 
 		int32 BatchElementIndex,
 		int32 DrawPrimitiveId,
+		int32 ScenePrimitiveId,
 		ERasterizerFillMode MeshFillMode,
 		ERasterizerCullMode MeshCullMode,
 		FMeshDrawCommandSortKey SortKey,
@@ -649,10 +650,18 @@ public:
 
 	// Note: no ctor as TChunkedArray::CopyToLinearArray requires POD types
 
-	FORCEINLINE_DEBUGGABLE void Setup(const FMeshDrawCommand* InMeshDrawCommand, int32 InDrawPrimitiveIndex, int32 InStateBucketId, ERasterizerFillMode InMeshFillMode, ERasterizerCullMode InMeshCullMode, FMeshDrawCommandSortKey InSortKey)
+	FORCEINLINE_DEBUGGABLE void Setup(
+		const FMeshDrawCommand* InMeshDrawCommand,
+		int32 InDrawPrimitiveId,
+		int32 InScenePrimitiveId,
+		int32 InStateBucketId,
+		ERasterizerFillMode InMeshFillMode,
+		ERasterizerCullMode InMeshCullMode,
+		FMeshDrawCommandSortKey InSortKey)
 	{
 		MeshDrawCommand = InMeshDrawCommand;
-		DrawPrimitiveId = InDrawPrimitiveIndex;
+		DrawPrimitiveId = InDrawPrimitiveId;
+		ScenePrimitiveId = InScenePrimitiveId;
 		PrimitiveIdBufferOffset = -1;
 		StateBucketId = InStateBucketId;
 		MeshFillMode = InMeshFillMode;
@@ -670,6 +679,9 @@ public:
 	// If it's < Scene->Primitives.Num() then it's a valid Scene PrimitiveIndex and can be used to backtrack to the FPrimitiveSceneInfo.
 	int32 DrawPrimitiveId;
 
+	// Scene PrimitiveId that generated this draw command, or -1 if no FPrimitiveSceneInfo. Can be used to backtrack to the FPrimitiveSceneInfo.
+	int32 ScenePrimitiveId;
+
 	// Offset into the buffer of PrimitiveIds built for this pass, in int32's.
 	int32 PrimitiveIdBufferOffset;
 
@@ -717,6 +729,7 @@ public:
 		const FMeshBatch& MeshBatch, 
 		int32 BatchElementIndex,
 		int32 DrawPrimitiveId,
+		int32 ScenePrimitiveId,
 		ERasterizerFillMode MeshFillMode,
 		ERasterizerCullMode MeshCullMode,
 		FMeshDrawCommandSortKey SortKey,
@@ -732,7 +745,7 @@ public:
 		FVisibleMeshDrawCommand NewVisibleMeshDrawCommand;
 		//@todo MeshCommandPipeline - assign usable state ID for dynamic path draws
 		// Currently dynamic path draws will not get dynamic instancing, but they will be roughly sorted by state
-		NewVisibleMeshDrawCommand.Setup(&MeshDrawCommand, DrawPrimitiveId, -1, MeshFillMode, MeshCullMode, SortKey);
+		NewVisibleMeshDrawCommand.Setup(&MeshDrawCommand, DrawPrimitiveId, ScenePrimitiveId, -1, MeshFillMode, MeshCullMode, SortKey);
 		DrawList.Add(NewVisibleMeshDrawCommand);
 	}
 
@@ -807,6 +820,7 @@ public:
 		const FMeshBatch& MeshBatch, 
 		int32 BatchElementIndex,
 		int32 DrawPrimitiveId,
+		int32 ScenePrimitiveId,
 		ERasterizerFillMode MeshFillMode,
 		ERasterizerCullMode MeshCullMode,
 		FMeshDrawCommandSortKey SortKey,
@@ -1051,7 +1065,11 @@ public:
 		const ShaderElementDataType& ShaderElementData);
 
 private:
-	RENDERER_API int32 GetDrawCommandPrimitiveId(const FPrimitiveSceneInfo* RESTRICT PrimitiveSceneInfo, const FMeshBatchElement& BatchElement) const;
+	RENDERER_API void GetDrawCommandPrimitiveId(
+		const FPrimitiveSceneInfo* RESTRICT PrimitiveSceneInfo,
+		const FMeshBatchElement& BatchElement,
+		int32& DrawPrimitiveId,
+		int32& ScenePrimitiveId) const;
 };
 
 typedef FMeshPassProcessor* (*PassProcessorCreateFunction)(const FScene* Scene, const FSceneView* InViewIfDynamicMeshCommand, FMeshPassDrawListContext* InDrawListContext);
diff --git a/Engine/Source/Runtime/Renderer/Public/MeshPassProcessor.inl b/Engine/Source/Runtime/Renderer/Public/MeshPassProcessor.inl
--- a/Engine/Source/Runtime/Renderer/Public/MeshPassProcessor.inl
+++ b/Engine/Source/Runtime/Renderer/Public/MeshPassProcessor.inl
@@ -124,9 +124,12 @@ void FMeshPassProcessor::BuildMeshDrawCommands(
 				PassShaders.GeometryShader->GetElementShaderBindings(Scene, ViewIfDynamicMeshCommand, VertexFactory, EVertexInputStreamType::Default, FeatureLevel, PrimitiveSceneProxy, MeshBatch, BatchElement, ShaderElementData, GeometryShaderBindings, MeshDrawCommand.VertexStreams);
 			}
 
-			const int32 DrawPrimitiveId = GetDrawCommandPrimitiveId(PrimitiveSceneInfo, BatchElement);
+			int32 DrawPrimitiveId;
+			int32 ScenePrimitiveId;
+			GetDrawCommandPrimitiveId(PrimitiveSceneInfo, BatchElement, DrawPrimitiveId, ScenePrimitiveId);
+
 			FMeshProcessorShaders ShadersForDebugging = PassShaders.GetUntypedShaders();
-			DrawListContext->FinalizeCommand(MeshBatch, BatchElementIndex, DrawPrimitiveId, MeshFillMode, MeshCullMode, SortKey, PipelineState, &ShadersForDebugging, MeshDrawCommand);
+			DrawListContext->FinalizeCommand(MeshBatch, BatchElementIndex, DrawPrimitiveId, ScenePrimitiveId, MeshFillMode, MeshCullMode, SortKey, PipelineState, &ShadersForDebugging, MeshDrawCommand);
 		}
 	}
 }
