Fixes NaN ensures when EnsureOnNaNFail=1 when the temporal AA history has never been initialized.

diff --git a/Engine/Source/Runtime/Renderer/Private/PostProcess/PostProcessTemporalAA.cpp b/Engine/Source/Runtime/Renderer/Private/PostProcess/PostProcessTemporalAA.cpp
--- a/Engine/Source/Runtime/Renderer/Private/PostProcess/PostProcessTemporalAA.cpp
+++ b/Engine/Source/Runtime/Renderer/Private/PostProcess/PostProcessTemporalAA.cpp
@@ -498,47 +498,51 @@ FTAAOutputs FTAAPassParameters::AddTemporalAAPass(
 
 		// Set history shader parameters.
 		{
-			FIntPoint ReferenceViewportOffset = InputHistory.ViewportRect.Min;
-			FIntPoint ReferenceViewportExtent = InputHistory.ViewportRect.Size();
-			FIntPoint ReferenceBufferSize = InputHistory.ReferenceBufferSize;
+			if (bCameraCut)
+			{
+				CommonShaderParameters.ScreenPosToHistoryBufferUV = FVector4(1.0f, 1.0f, 1.0f, 1.0f);
+				CommonShaderParameters.ScreenPosAbsMax = FVector2D(0.0f, 0.0f);
+				CommonShaderParameters.HistoryBufferUVMinMax = FVector4(0.0f, 0.0f, 0.0f, 0.0f);
+				CommonShaderParameters.HistoryBufferSize = FVector4(1.0f, 1.0f, 1.0f, 1.0f);
+				CommonShaderParameters.HistoryBuffer0 = GraphBuilder.RegisterExternalTexture(GSystemTextures.BlackDummy);
+				CommonShaderParameters.HistoryBuffer1 = CommonShaderParameters.HistoryBuffer0;
+			}
+			else
+			{
+				FIntPoint ReferenceViewportOffset = InputHistory.ViewportRect.Min;
+				FIntPoint ReferenceViewportExtent = InputHistory.ViewportRect.Size();
+				FIntPoint ReferenceBufferSize = InputHistory.ReferenceBufferSize;
 
-			float InvReferenceBufferSizeX = 1.f / float(InputHistory.ReferenceBufferSize.X);
-			float InvReferenceBufferSizeY = 1.f / float(InputHistory.ReferenceBufferSize.Y);
+				float InvReferenceBufferSizeX = 1.f / float(InputHistory.ReferenceBufferSize.X);
+				float InvReferenceBufferSizeY = 1.f / float(InputHistory.ReferenceBufferSize.Y);
 
-			CommonShaderParameters.ScreenPosToHistoryBufferUV = FVector4(
-				ReferenceViewportExtent.X * 0.5f * InvReferenceBufferSizeX,
-				-ReferenceViewportExtent.Y * 0.5f * InvReferenceBufferSizeY,
-				(ReferenceViewportExtent.X * 0.5f + ReferenceViewportOffset.X) * InvReferenceBufferSizeX,
-				(ReferenceViewportExtent.Y * 0.5f + ReferenceViewportOffset.Y) * InvReferenceBufferSizeY);
+				CommonShaderParameters.ScreenPosToHistoryBufferUV = FVector4(
+					ReferenceViewportExtent.X * 0.5f * InvReferenceBufferSizeX,
+					-ReferenceViewportExtent.Y * 0.5f * InvReferenceBufferSizeY,
+					(ReferenceViewportExtent.X * 0.5f + ReferenceViewportOffset.X) * InvReferenceBufferSizeX,
+					(ReferenceViewportExtent.Y * 0.5f + ReferenceViewportOffset.Y) * InvReferenceBufferSizeY);
 
-			FIntPoint ViewportOffset = ReferenceViewportOffset / ResolutionDivisor;
-			FIntPoint ViewportExtent = FIntPoint::DivideAndRoundUp(ReferenceViewportExtent, ResolutionDivisor);
-			FIntPoint BufferSize = ReferenceBufferSize / ResolutionDivisor;
+				FIntPoint ViewportOffset = ReferenceViewportOffset / ResolutionDivisor;
+				FIntPoint ViewportExtent = FIntPoint::DivideAndRoundUp(ReferenceViewportExtent, ResolutionDivisor);
+				FIntPoint BufferSize = ReferenceBufferSize / ResolutionDivisor;
 
-			CommonShaderParameters.ScreenPosAbsMax = FVector2D(1.0f - 1.0f / float(ViewportExtent.X), 1.0f - 1.0f / float(ViewportExtent.Y));
+				CommonShaderParameters.ScreenPosAbsMax = FVector2D(1.0f - 1.0f / float(ViewportExtent.X), 1.0f - 1.0f / float(ViewportExtent.Y));
 
-			float InvBufferSizeX = 1.f / float(BufferSize.X);
-			float InvBufferSizeY = 1.f / float(BufferSize.Y);
+				float InvBufferSizeX = 1.f / float(BufferSize.X);
+				float InvBufferSizeY = 1.f / float(BufferSize.Y);
 
-			CommonShaderParameters.HistoryBufferUVMinMax = FVector4(
-				(ViewportOffset.X + 0.5f) * InvBufferSizeX,
-				(ViewportOffset.Y + 0.5f) * InvBufferSizeY,
-				(ViewportOffset.X + ViewportExtent.X - 0.5f) * InvBufferSizeX,
-				(ViewportOffset.Y + ViewportExtent.Y - 0.5f) * InvBufferSizeY);
+				CommonShaderParameters.HistoryBufferUVMinMax = FVector4(
+					(ViewportOffset.X + 0.5f) * InvBufferSizeX,
+					(ViewportOffset.Y + 0.5f) * InvBufferSizeY,
+					(ViewportOffset.X + ViewportExtent.X - 0.5f) * InvBufferSizeX,
+					(ViewportOffset.Y + ViewportExtent.Y - 0.5f) * InvBufferSizeY);
 
-			CommonShaderParameters.HistoryBufferSize = FVector4(BufferSize.X, BufferSize.Y, InvBufferSizeX, InvBufferSizeY);
+				CommonShaderParameters.HistoryBufferSize = FVector4(BufferSize.X, BufferSize.Y, InvBufferSizeX, InvBufferSizeY);
 
-			if (!bCameraCut)
-			{
 				CommonShaderParameters.HistoryBuffer0 = GraphBuilder.RegisterExternalTexture(InputHistory.RT[0]);
 				if (InputHistory.RT[1].IsValid())
 					CommonShaderParameters.HistoryBuffer1 = GraphBuilder.RegisterExternalTexture(InputHistory.RT[1]);
 			}
-			else
-			{
-				CommonShaderParameters.HistoryBuffer0 = GraphBuilder.RegisterExternalTexture(GSystemTextures.BlackDummy);
-				CommonShaderParameters.HistoryBuffer1 = CommonShaderParameters.HistoryBuffer0;
-			}
 
 			CommonShaderParameters.HistoryBuffer0Sampler = TStaticSamplerState<SF_Bilinear>::GetRHI();
 			CommonShaderParameters.HistoryBuffer1Sampler = TStaticSamplerState<SF_Bilinear>::GetRHI();
