DR - Fix CIS warnings
#rnx

diff --git a/Engine/Source/Runtime/Renderer/Private/PostProcess/PostProcessSubsurface.cpp b/Engine/Source/Runtime/Renderer/Private/PostProcess/PostProcessSubsurface.cpp
--- a/Engine/Source/Runtime/Renderer/Private/PostProcess/PostProcessSubsurface.cpp
+++ b/Engine/Source/Runtime/Renderer/Private/PostProcess/PostProcessSubsurface.cpp
@@ -708,10 +708,10 @@ void VisualizeSubsurface(
 
 		// Draw debug text
 		{
-			const FViewInfo& View = ScreenPassView.View;
-			const FSceneViewFamily& ViewFamily = *View.Family;
+			const FViewInfo& LocalView = ScreenPassView.View;
+			const FSceneViewFamily& ViewFamily = *LocalView.Family;
 			FRenderTargetTemp TempRenderTarget(static_cast<FTexture2DRHIParamRef>(SceneTextureOutput->GetRHITexture()), SceneTextureOutput->Desc.Extent);
-			FCanvas Canvas(&TempRenderTarget, nullptr, ViewFamily.CurrentRealTime, ViewFamily.CurrentWorldTime, ViewFamily.DeltaWorldTime, View.GetFeatureLevel());
+			FCanvas Canvas(&TempRenderTarget, nullptr, ViewFamily.CurrentRealTime, ViewFamily.CurrentWorldTime, ViewFamily.DeltaWorldTime, LocalView.GetFeatureLevel());
 
 			float X = 30;
 			float Y = 28;
@@ -761,12 +761,12 @@ void ComputeSubsurfaceShim(FRHICommandListImmediate& RHICmdList, const TArray<FV
 }
 
 FRenderingCompositeOutputRef VisualizeSubsurfaceShim(
-	FRHICommandListImmediate& RHICmdList,
+	FRHICommandListImmediate& InRHICmdList,
 	FRenderingCompositionGraph& Graph,
 	FRenderingCompositeOutputRef Input)
 {
 	// we need the GBuffer, we release it Process()
-	FSceneRenderTargets::Get(RHICmdList).AdjustGBufferRefCount(RHICmdList, 1);
+	FSceneRenderTargets::Get(InRHICmdList).AdjustGBufferRefCount(InRHICmdList, 1);
 
 	FRenderingCompositePass* SubsurfaceVisualizePass = Graph.RegisterPass(new(FMemStack::Get()) TRCPassForRDG<1, 1>(
 		[](FRenderingCompositePass* Pass, FRenderingCompositePassContext& CompositePassContext)
