Fixes indeterminism in the material TAA dither engine functions causing EngineTest to fail

diff --git a/Engine/Source/Runtime/Renderer/Private/SceneRendering.cpp b/Engine/Source/Runtime/Renderer/Private/SceneRendering.cpp
--- a/Engine/Source/Runtime/Renderer/Private/SceneRendering.cpp
+++ b/Engine/Source/Runtime/Renderer/Private/SceneRendering.cpp
@@ -863,6 +863,7 @@ void FViewInfo::Init()
 	// Disable HDR encoding for editor elements.
 	EditorSimpleElementCollector.BatchedElements.EnableMobileHDREncoding(false);
 
+	TemporalJitterIndex = 0;
 	TemporalJitterPixels = FVector2D::ZeroVector;
 
 	PreExposure = 1.0f;
@@ -1261,7 +1262,7 @@ void FViewInfo::SetupUniformBufferParameters(
 	if(State)
 	{
 		ViewUniformShaderParameters.TemporalAAParams = FVector4(
-			ViewState->GetCurrentTemporalAASampleIndex(), 
+			TemporalJitterIndex, 
 			0,
 			TemporalJitterPixels.X,
 			TemporalJitterPixels.Y);
diff --git a/Engine/Source/Runtime/Renderer/Private/SceneRendering.h b/Engine/Source/Runtime/Renderer/Private/SceneRendering.h
--- a/Engine/Source/Runtime/Renderer/Private/SceneRendering.h
+++ b/Engine/Source/Runtime/Renderer/Private/SceneRendering.h
@@ -884,6 +884,9 @@ public:
 	float TranslucencyVolumeVoxelSize[TVC_MAX];
 	FVector TranslucencyLightingVolumeSize[TVC_MAX];
 
+	/** Index of the temporal jitter in the sequence. */
+	int32 TemporalJitterIndex;
+
 	/** Temporal jitter at the pixel scale. */
 	FVector2D TemporalJitterPixels;
 
diff --git a/Engine/Source/Runtime/Renderer/Private/SceneVisibility.cpp b/Engine/Source/Runtime/Renderer/Private/SceneVisibility.cpp
--- a/Engine/Source/Runtime/Renderer/Private/SceneVisibility.cpp
+++ b/Engine/Source/Runtime/Renderer/Private/SceneVisibility.cpp
@@ -3010,7 +3010,7 @@ void FSceneRenderer::PreVisibilityFrameSetup(FRHICommandListImmediate& RHICmdLis
 
 			// Compute the new sample index in the temporal sequence.
 			int32 TemporalSampleIndex = ViewState->TemporalAASampleIndex + 1;
-			if(TemporalSampleIndex >= TemporalAASamples)
+			if(TemporalSampleIndex >= TemporalAASamples || View.bCameraCut)
 			{
 				TemporalSampleIndex = 0;
 			}
@@ -3122,6 +3122,7 @@ void FSceneRenderer::PreVisibilityFrameSetup(FRHICommandListImmediate& RHICmdLis
 				SampleY = r * FMath::Sin( Theta );
 			}
 
+			View.TemporalJitterIndex = TemporalSampleIndex;
 			View.TemporalJitterPixels.X = SampleX;
 			View.TemporalJitterPixels.Y = SampleY;
 
