Fixed scene initialization in ray tracing test code (instance mask is now correctly set to 0xFF)

diff --git a/Engine/Source/Runtime/Renderer/Private/RayTracing/RayTracingTestbed.cpp b/Engine/Source/Runtime/Renderer/Private/RayTracing/RayTracingTestbed.cpp
--- a/Engine/Source/Runtime/Renderer/Private/RayTracing/RayTracingTestbed.cpp
+++ b/Engine/Source/Runtime/Renderer/Private/RayTracing/RayTracingTestbed.cpp
@@ -97,7 +97,7 @@ void TestBasicRayTracing(bool bValidateResults)
 	FRayTracingGeometryRHIRef Geometry = RHICreateRayTracingGeometry(GeometryInitializer);
 
 	FRayTracingGeometryInstance Instances[] = {
-		FRayTracingGeometryInstance { Geometry, FMatrix::Identity }
+		FRayTracingGeometryInstance { Geometry, FMatrix::Identity, 0, 0xFF }
 	};
 
 	FRayTracingSceneInitializer Initializer;
