Fix off-by-one error when linearizing Faure permutation array.

diff --git a/Engine/Source/Runtime/Renderer/Private/HaltonUtilities.cpp b/Engine/Source/Runtime/Renderer/Private/HaltonUtilities.cpp
--- a/Engine/Source/Runtime/Renderer/Private/HaltonUtilities.cpp
+++ b/Engine/Source/Runtime/Renderer/Private/HaltonUtilities.cpp
@@ -174,7 +174,7 @@ int FHaltonSequence::Invert(int Base, int Digits, int Index, const TArray<int>&
 
 void FHaltonSequence::FlattenPermutation(const TArray<int>& Permutation)
 {
-	FaurePermutationOffsets.Add(Permutation.Num());
+	FaurePermutationOffsets.Add(FaurePermutations.Num());
 	FaurePermutations.Append(Permutation);
 }
 
