Pre-validates motion blur scattering pass parameters before creating the pass.

diff --git a/Engine/Source/Runtime/Renderer/Private/PostProcess/PostProcessMotionBlur.cpp b/Engine/Source/Runtime/Renderer/Private/PostProcess/PostProcessMotionBlur.cpp
--- a/Engine/Source/Runtime/Renderer/Private/PostProcess/PostProcessMotionBlur.cpp
+++ b/Engine/Source/Runtime/Renderer/Private/PostProcess/PostProcessMotionBlur.cpp
@@ -544,6 +544,9 @@ void ComputeMotionBlurVelocity(
 
 		TShaderMapRef<FMotionBlurVelocityDilateScatterVS> VertexShader(View.ShaderMap);
 		TShaderMapRef<FMotionBlurVelocityDilateScatterPS> PixelShader(View.ShaderMap);
+		
+		ValidateShaderParameters(*VertexShader, *PassParameters);
+		ValidateShaderParameters(*PixelShader, *PassParameters);
 
 		GraphBuilder.AddPass(
 			RDG_EVENT_NAME("VelocityTileScatter %dx%d", VelocityTileCount.X, VelocityTileCount.Y),
