Wrapping Halton primes in Buffer SRV

diff --git a/Engine/Source/Runtime/Renderer/Private/Halton.cpp b/Engine/Source/Runtime/Renderer/Private/Halton.cpp
--- a/Engine/Source/Runtime/Renderer/Private/Halton.cpp
+++ b/Engine/Source/Runtime/Renderer/Private/Halton.cpp
@@ -55,9 +55,9 @@ void FHaltonSequenceIteration::InitRHI()
 		for (uint32 IterationIndex = 0; IterationIndex < IterationCount; ++IterationIndex)
 		{
 			uint32 IterationOffset = IterationIndex * DimensionCount;
-			for (uint32 Dimension = 0; Dimension < DimensionCount; Dimension++)
+			for (uint32 DimensionIndex = 0; DimensionIndex < DimensionCount; ++DimensionIndex)
 			{
-				RandomSamples[SequenceOffset + IterationOffset + Dimension] = HaltonSequence.Sample(Dimension, SequenceValue + IterationIndex);
+				RandomSamples[SequenceOffset + IterationOffset + DimensionIndex] = HaltonSequence.Sample(DimensionIndex, SequenceValue + IterationIndex);
 			}
 		}
 	}
@@ -82,3 +82,77 @@ void FHaltonSequenceIteration::InitializeSequence()
 		Sequence[Index] = IntegerHash(Index);
 	}
 }
+
+IMPLEMENT_GLOBAL_SHADER_PARAMETER_STRUCT(FHaltonPrimes, "HaltonPrimes");
+
+FHaltonPrimesResource::FHaltonPrimesResource()
+	: DimensionCount(512)
+{
+	Primes = {
+		2,3,5,7,11,13,17,19,23,29,
+		31,37,41,43,47,53,59,61,67,71,
+		73,79,83,89,97,101,103,107,109,113,
+		127,131,137,139,149,151,157,163,167,173,
+		179,181,191,193,197,199,211,223,227,229,
+		233,239,241,251,257,263,269,271,277,281,
+		283,293,307,311,313,317,331,337,347,349,
+		353,359,367,373,379,383,389,397,401,409,
+		419,421,431,433,439,443,449,457,461,463,
+		467,479,487,491,499,503,509,521,523,541,
+		547,557,563,569,571,577,587,593,599,601,
+		607,613,617,619,631,641,643,647,653,659,
+		661,673,677,683,691,701,709,719,727,733,
+		739,743,751,757,761,769,773,787,797,809,
+		811,821,823,827,829,839,853,857,859,863,
+		877,881,883,887,907,911,919,929,937,941,
+		947,953,967,971,977,983,991,997,1009,1013,
+		1019,1021,1031,1033,1039,1049,1051,1061,1063,1069,
+		1087,1091,1093,1097,1103,1109,1117,1123,1129,1151,
+		1153,1163,1171,1181,1187,1193,1201,1213,1217,1223,
+		1229,1231,1237,1249,1259,1277,1279,1283,1289,1291,
+		1297,1301,1303,1307,1319,1321,1327,1361,1367,1373,
+		1381,1399,1409,1423,1427,1429,1433,1439,1447,1451,
+		1453,1459,1471,1481,1483,1487,1489,1493,1499,1511,
+		1523,1531,1543,1549,1553,1559,1567,1571,1579,1583,
+		1597,1601,1607,1609,1613,1619,1621,1627,1637,1657,
+		1663,1667,1669,1693,1697,1699,1709,1721,1723,1733,
+		1741,1747,1753,1759,1777,1783,1787,1789,1801,1811,
+		1823,1831,1847,1861,1867,1871,1873,1877,1879,1889,
+		1901,1907,1913,1931,1933,1949,1951,1973,1979,1987,
+		1993,1997,1999,2003,2011,2017,2027,2029,2039,2053,
+		2063,2069,2081,2083,2087,2089,2099,2111,2113,2129,
+		2131,2137,2141,2143,2153,2161,2179,2203,2207,2213,
+		2221,2237,2239,2243,2251,2267,2269,2273,2281,2287,
+		2293,2297,2309,2311,2333,2339,2341,2347,2351,2357,
+		2371,2377,2381,2383,2389,2393,2399,2411,2417,2423,
+		2437,2441,2447,2459,2467,2473,2477,2503,2521,2531,
+		2539,2543,2549,2551,2557,2579,2591,2593,2609,2617,
+		2621,2633,2647,2657,2659,2663,2671,2677,2683,2687,
+		2689,2693,2699,2707,2711,2713,2719,2729,2731,2741,
+		2749,2753,2767,2777,2789,2791,2797,2801,2803,2819,
+		2833,2837,2843,2851,2857,2861,2879,2887,2897,2903,
+		2909,2917,2927,2939,2953,2957,2963,2969,2971,2999,
+		3001,3011,3019,3023,3037,3041,3049,3061,3067,3079,
+		3083,3089,3109,3119,3121,3137,3163,3167,3169,3181,
+		3187,3191,3203,3209,3217,3221,3229,3251,3253,3257,
+		3259,3271,3299,3301,3307,3313,3319,3323,3329,3331,
+		3343,3347,3359,3361,3371,3373,3389,3391,3407,3413,
+		3433,3449,3457,3461,3463,3467,3469,3491,3499,3511,
+		3517,3527,3529,3533,3539,3541,3547,3557,3559,3571,
+		3581,3583,3593,3607,3613,3617,3623,3631,3637,3643,
+		3659,3671
+	};
+}
+
+void FHaltonPrimesResource::InitRHI()
+{
+	FRHIResourceCreateInfo CreateInfo;
+	{
+		CreateInfo.DebugName = TEXT("HaltonPrimes");
+		PrimesBuffer.VertexBufferRHI = RHICreateVertexBuffer(DimensionCount * sizeof(float), BUF_Transient | BUF_FastVRAM | BUF_ShaderResource | BUF_UnorderedAccess, CreateInfo);
+		uint32 Offset = 0;
+		void* BasePtr = RHILockVertexBuffer(PrimesBuffer.VertexBufferRHI, Offset, DimensionCount * sizeof(float), RLM_WriteOnly);
+		FPlatformMemory::Memcpy(BasePtr, Primes.GetData(), DimensionCount * sizeof(float));
+		RHIUnlockVertexBuffer(PrimesBuffer.VertexBufferRHI);
+	}
+}
\ No newline at end of file
diff --git a/Engine/Source/Runtime/Renderer/Private/Halton.h b/Engine/Source/Runtime/Renderer/Private/Halton.h
--- a/Engine/Source/Runtime/Renderer/Private/Halton.h
+++ b/Engine/Source/Runtime/Renderer/Private/Halton.h
@@ -24,28 +24,49 @@ END_GLOBAL_SHADER_PARAMETER_STRUCT()
 class FHaltonSequenceIteration : public FRenderResource
 {
 public:
-	FHaltonSequenceIteration(const FHaltonSequence& HaltonSequenceLocal, uint32 IterationCountLocal, uint32 SequenceCountLocal, uint32 DimensionCount, uint32 IterationLocal);
+	FHaltonSequenceIteration(
+		const FHaltonSequence& HaltonSequenceLocal,
+		uint32 IterationCountLocal,
+		uint32 SequenceCountLocal,
+		uint32 DimensionCount,
+		uint32 IterationLocal
+	);
+
 	virtual ~FHaltonSequenceIteration();
 
 	// FRenderResource interface.
 	virtual void InitRHI() override;
+
 	virtual void ReleaseRHI() override
 	{
 		SequenceIteration.ReleaseRHI();
 	}
 
-	virtual void InitResource() override
+	virtual FString GetFriendlyName() const override
+	{
+		return TEXT("Halton Sequence Iteration");
+	}
+
+	// Accessors
+	uint32 GetSequenceCount() const
 	{
-		FRenderResource::InitResource();
+		return SequenceCount;
 	}
-	virtual void ReleaseResource() override { FRenderResource::ReleaseResource(); }
 
-	virtual FString GetFriendlyName() const override { return TEXT("Halton Sequence Iteration"); }
+	uint32 GetIterationCount() const
+	{
+		return IterationCount;
+	}
+
+	uint32 GetDimensionCount() const
+	{
+		return DimensionCount;
+	}
 
-	uint32 GetSequenceCount() const { return SequenceCount; }
-	uint32 GetIterationCount() const { return IterationCount; }
-	uint32 GetDimensionCount() const { return DimensionCount; }
-	uint32 GetIteration() const { return Iteration; }
+	uint32 GetIteration() const
+	{
+		return Iteration;
+	}
 
 	FVertexBuffer SequenceIteration;
 
@@ -69,3 +90,49 @@ inline void InitializeHaltonSequenceIteration(const FHaltonSequenceIteration& Ha
 	HaltonIteration.IterationCount = HaltonSequenceIteration.GetIterationCount();
 	HaltonIteration.SequenceIteration = RHICreateShaderResourceView(HaltonSequenceIteration.SequenceIteration.VertexBufferRHI, sizeof(float), PF_R32_FLOAT);
 }
+
+BEGIN_GLOBAL_SHADER_PARAMETER_STRUCT(FHaltonPrimes, )
+SHADER_PARAMETER(int, Dimensions)
+SHADER_PARAMETER_SRV(Buffer<int>, Primes)
+END_GLOBAL_SHADER_PARAMETER_STRUCT()
+
+/**
+* This resource stores a number of iterations of Halton sequences, up to the specified dimensionality.
+*/
+class FHaltonPrimesResource : public FRenderResource
+{
+public:
+	FHaltonPrimesResource();
+	virtual ~FHaltonPrimesResource() {}
+
+	// FRenderResource interface.
+	virtual void InitRHI() override;
+
+	virtual void ReleaseRHI() override
+	{
+		PrimesBuffer.ReleaseRHI();
+	}
+
+	virtual FString GetFriendlyName() const override
+	{
+		return TEXT("Halton Sequence Iteration");
+	}
+
+	// Accessors
+	uint32 GetDimensionCount() const
+	{
+		return DimensionCount;
+	}
+
+	FVertexBuffer PrimesBuffer;
+
+private:
+	TArray<int> Primes;
+	uint32 DimensionCount;
+};
+
+inline void InitializeHaltonPrimes(const FHaltonPrimesResource& HaltonPrimeResource, FHaltonPrimes& HaltonPrimes)
+{
+	HaltonPrimes.Dimensions = HaltonPrimeResource.GetDimensionCount();
+	HaltonPrimes.Primes = RHICreateShaderResourceView(HaltonPrimeResource.PrimesBuffer.VertexBufferRHI, sizeof(float), PF_R32_FLOAT);
+}
\ No newline at end of file
diff --git a/Engine/Source/Runtime/Renderer/Private/RendererScene.cpp b/Engine/Source/Runtime/Renderer/Private/RendererScene.cpp
--- a/Engine/Source/Runtime/Renderer/Private/RendererScene.cpp
+++ b/Engine/Source/Runtime/Renderer/Private/RendererScene.cpp
@@ -1085,6 +1085,8 @@ FScene::FScene(UWorld* InWorld, bool bInRequiresHitProxies, bool bInIsEditorScen
 	{
 		RayTracingDynamicGeometryCollection = new FRayTracingDynamicGeometryCollection();
 	}
+
+	BeginInitResource(&HaltonPrimesResource);
 #endif
 
 	World->UpdateParameterCollectionInstances(false, false);
@@ -1134,6 +1136,8 @@ FScene::~FScene()
 		delete RayTracingDynamicGeometryCollection;
 		RayTracingDynamicGeometryCollection = nullptr;
 	}
+
+	BeginReleaseResource(&HaltonPrimesResource);
 #endif
 }
 
diff --git a/Engine/Source/Runtime/Renderer/Private/ScenePrivate.h b/Engine/Source/Runtime/Renderer/Private/ScenePrivate.h
--- a/Engine/Source/Runtime/Renderer/Private/ScenePrivate.h
+++ b/Engine/Source/Runtime/Renderer/Private/ScenePrivate.h
@@ -2630,6 +2630,7 @@ public:
 #if RHI_RAYTRACING
 	class FRayTracingDynamicGeometryCollection* RayTracingDynamicGeometryCollection;
 	FHaltonSequence HaltonSequence;
+	FHaltonPrimesResource HaltonPrimesResource;
 #endif
 
 	/** Initialization constructor. */
