Fixed SSS. Intermediate SSS render target needs to be RGBA, as alpha stores depth which is used for the subsuface blur.

diff --git a/Engine/Source/Runtime/Renderer/Private/PostProcess/PostProcessSubsurface.cpp b/Engine/Source/Runtime/Renderer/Private/PostProcess/PostProcessSubsurface.cpp
--- a/Engine/Source/Runtime/Renderer/Private/PostProcess/PostProcessSubsurface.cpp
+++ b/Engine/Source/Runtime/Renderer/Private/PostProcess/PostProcessSubsurface.cpp
@@ -465,8 +465,13 @@ void ComputeSubsurfaceForView(
 	 */
 	const FScreenPassTextureViewport SubsurfaceViewport = FScreenPassTextureViewport::CreateDownscaled(SceneViewport, ScaleFactor);
 
-	FRDGTextureDesc SubsurfaceTextureDescriptor = SceneTexture->Desc;
-	SubsurfaceTextureDescriptor.Extent = SubsurfaceViewport.Extent;
+	const FRDGTextureDesc SubsurfaceTextureDescriptor = FRDGTextureDesc::Create2DDesc(
+		SubsurfaceViewport.Extent,
+		PF_FloatRGBA,
+		FClearValueBinding(),
+		TexCreate_None,
+		TexCreate_RenderTargetable | TexCreate_ShaderResource,
+		false);
 	
 	const FSubsurfaceParameters SubsurfaceCommonParameters = GetSubsurfaceCommonParameters(GraphBuilder.RHICmdList, View);
 	const FScreenPassTextureViewportParameters SubsurfaceViewportParameters = GetScreenPassTextureViewportParameters(SubsurfaceViewport);
