Adds an ensure to validate View.TemporalAAParams to investigate dither engine test failures.

diff --git a/Engine/Source/Runtime/Renderer/Private/SceneRendering.cpp b/Engine/Source/Runtime/Renderer/Private/SceneRendering.cpp
--- a/Engine/Source/Runtime/Renderer/Private/SceneRendering.cpp
+++ b/Engine/Source/Runtime/Renderer/Private/SceneRendering.cpp
@@ -1258,21 +1258,22 @@ void FViewInfo::SetupUniformBufferParameters(
 		}
 	}
 
-	uint32 FrameIndex = 0;
-
-	if(State)
 	{
+		ensureMsgf(TemporalJitterSequenceLength == 1 || AntiAliasingMethod == AAM_TemporalAA,
+			TEXT("TemporalJitterSequenceLength = %i is invalid"), TemporalJitterSequenceLength);
+		ensureMsgf(TemporalJitterIndex >= 0 && TemporalJitterIndex < TemporalJitterSequenceLength,
+			TEXT("TemporalJitterIndex = %i is invalid (TemporalJitterSequenceLength = %i)"), TemporalJitterIndex, TemporalJitterSequenceLength);
 		ViewUniformShaderParameters.TemporalAAParams = FVector4(
 			TemporalJitterIndex, 
 			TemporalJitterSequenceLength,
 			TemporalJitterPixels.X,
 			TemporalJitterPixels.Y);
-
-		FrameIndex = ViewState->GetFrameIndex();
 	}
-	else
+		
+	uint32 FrameIndex = 0;
+	if (ViewState)
 	{
-		ViewUniformShaderParameters.TemporalAAParams = FVector4(0, 1, 0, 0);
+		FrameIndex = ViewState->GetFrameIndex();
 	}
 
 	// TODO(GA): kill StateFrameIndexMod8 because this is only a scalar bit mask with StateFrameIndex anyway.
