Added Lightmass support for sun light atmosphere transmittance.

diff --git a/Engine/Source/Editor/UnrealEd/Private/Lightmass/Lightmass.cpp b/Engine/Source/Editor/UnrealEd/Private/Lightmass/Lightmass.cpp
--- a/Engine/Source/Editor/UnrealEd/Private/Lightmass/Lightmass.cpp
+++ b/Engine/Source/Editor/UnrealEd/Private/Lightmass/Lightmass.cpp
@@ -27,6 +27,7 @@
 #include "Components/RectLightComponent.h"
 #include "Engine/GeneratedMeshAreaLight.h"
 #include "Components/DirectionalLightComponent.h"
+#include "Renderer/Private/AtmosphereRendering.h"
 #include "Components/SkyLightComponent.h"
 #include "Components/ModelComponent.h"
 #include "Materials/MaterialInstanceConstant.h"
@@ -486,6 +487,7 @@ void FLightmassProcessor::SwarmCallback( NSwarm::FMessage* CallbackMessage, void
 -----------------------------------------------------------------------------*/
 FLightmassExporter::FLightmassExporter( UWorld* InWorld )
 	: Swarm( NSwarm::FSwarmInterface::Get() ) 
+	, AtmosphericFogComponent(nullptr)
 	, ExportStage(NotRunning)
 	, CurrentAmortizationIndex(0)
 	, OpenedMaterialExportChannels()
@@ -1021,6 +1023,29 @@ void FLightmassExporter::WriteVolumetricLightmapData( int32 Channel )
 
 void FLightmassExporter::WriteLights( int32 Channel )
 {
+	// Search for the sun light component the same way as in FScene::RemoveLightSceneInfo_RenderThread.
+	// If found, we keep the pointer to apply atmosphere transmittance to it in the light loop.
+	const UDirectionalLightComponent* SunLight = nullptr;
+	FLinearColor SunLightAtmosphereTransmittance(FLinearColor::White);
+	if (AtmosphericFogComponent)
+	{
+		float SunLightEnergy = 0.0f;
+		for (int32 LightIndex = 0; LightIndex < DirectionalLights.Num(); ++LightIndex)
+		{
+			const UDirectionalLightComponent* Light = DirectionalLights[LightIndex];
+			if (Light->IsUsedAsAtmosphereSunLight() && Light->GetColoredLightBrightness().ComputeLuminance() > SunLightEnergy)
+			{
+				SunLight = Light;
+				SunLightEnergy = SunLight->GetColoredLightBrightness().ComputeLuminance();
+			}
+		}
+
+		if (SunLight)
+		{
+			SunLightAtmosphereTransmittance = AtmosphericFogComponent->GetTransmittance(-SunLight->GetDirection()); 
+		}
+	}
+
 	// Export directional lights.
 	for ( int32 LightIndex = 0; LightIndex < DirectionalLights.Num(); ++LightIndex )
 	{
@@ -1034,6 +1059,11 @@ void FLightmassExporter::WriteLights( int32 Channel )
 		LightData.LightSourceRadius = 0;
 		LightData.LightSourceLength = 0;
 
+		if (Light == SunLight)
+		{
+			LightData.Color *= SunLightAtmosphereTransmittance;
+		}
+
 		TArray< uint8 > LightProfileTextureData;
 		CopyLightProfile( Light, LightData, LightProfileTextureData );
 
diff --git a/Engine/Source/Editor/UnrealEd/Private/Lightmass/Lightmass.h b/Engine/Source/Editor/UnrealEd/Private/Lightmass/Lightmass.h
--- a/Engine/Source/Editor/UnrealEd/Private/Lightmass/Lightmass.h
+++ b/Engine/Source/Editor/UnrealEd/Private/Lightmass/Lightmass.h
@@ -14,6 +14,7 @@
 #include "StaticLightingSystem/StaticLightingPrivate.h"
 #include "Lightmass/LightmassImportanceVolume.h"
 #include "Components/LightmassPortalComponent.h"
+#include "Atmosphere/AtmosphericFogComponent.h"
 #if PLATFORM_WINDOWS
 #include "Windows/AllowWindowsPlatformTypes.h"
 #endif
@@ -114,6 +115,11 @@ public:
 		Portals.Add(InPortalComponent->GetComponentTransform().ToMatrixWithScale());
 	}
 
+	void SetAtmosphericComponent(const UAtmosphericFogComponent* AtmosphericFog)
+	{
+		AtmosphericFogComponent = AtmosphericFog;
+	}
+
 	// if provided, InStaticLightingMesh is used to UV unwrap the material into the static lighting textures
 	void AddMaterial(UMaterialInterface* InMaterialInterface, const FStaticLightingMesh* InStaticLightingMesh = nullptr);
 
@@ -199,6 +205,8 @@ private:
 	TArray<FBox>				CharacterIndirectDetailVolumes;
 	TArray<FMatrix>				Portals;
 
+	const UAtmosphericFogComponent*	AtmosphericFogComponent;
+
 	FLightmassWorldInfoSettings LevelSettings;
 	/** The number of local cores to leave unused */
 	int32 NumUnusedLocalCores;
diff --git a/Engine/Source/Editor/UnrealEd/Private/StaticLightingSystem/StaticLightingSystem.cpp b/Engine/Source/Editor/UnrealEd/Private/StaticLightingSystem/StaticLightingSystem.cpp
--- a/Engine/Source/Editor/UnrealEd/Private/StaticLightingSystem/StaticLightingSystem.cpp
+++ b/Engine/Source/Editor/UnrealEd/Private/StaticLightingSystem/StaticLightingSystem.cpp
@@ -40,6 +40,7 @@
 #include "GameFramework/WorldSettings.h"
 #include "Engine/GeneratedMeshAreaLight.h"
 #include "Components/SkyLightComponent.h"
+#include "Atmosphere/AtmosphericFogComponent.h"
 #include "Components/ModelComponent.h"
 #include "Engine/LightMapTexture2D.h"
 #include "Editor.h"
@@ -2025,6 +2026,16 @@ void FStaticLightingSystem::GatherScene()
 		}
 	}
 
+	for (TObjectIterator<UAtmosphericFogComponent> It; It; ++It)
+	{
+		UAtmosphericFogComponent* AtmosphericFog = *It;
+		if (!AtmosphericFog->IsPendingKill() && AtmosphericFog->bVisible)
+		{
+			LightmassExporter->SetAtmosphericComponent(AtmosphericFog);
+			break;	// We only register the first we find
+		}
+	}
+
 	float MinimumImportanceVolumeExtentWithoutWarning = 0.0f;
 	verify(GConfig->GetFloat(TEXT("DevOptions.StaticLightingSceneConstants"), TEXT("MinimumImportanceVolumeExtentWithoutWarning"), MinimumImportanceVolumeExtentWithoutWarning, GLightmassIni));
 
diff --git a/Engine/Source/Runtime/Engine/Classes/Atmosphere/AtmosphericFogComponent.h b/Engine/Source/Runtime/Engine/Classes/Atmosphere/AtmosphericFogComponent.h
--- a/Engine/Source/Runtime/Engine/Classes/Atmosphere/AtmosphericFogComponent.h
+++ b/Engine/Source/Runtime/Engine/Classes/Atmosphere/AtmosphericFogComponent.h
@@ -87,6 +87,11 @@ struct FAtmospherePrecomputeParameters
 			|| (InscatterMuSNum != Other.InscatterMuSNum)
 			|| (InscatterNuNum != Other.InscatterNuNum);
 	}
+
+	float GetRHeight() const
+	{
+		return DensityHeight * DensityHeight * DensityHeight * 64.f; // This formula is used for legacy conversion reason. In itself it does not make any sense.
+	}
 };
 
 /**
@@ -277,6 +282,74 @@ public:
 	void ApplyComponentInstanceData(struct FAtmospherePrecomputeInstanceData* ComponentInstanceData);
 	const FAtmospherePrecomputeParameters& GetPrecomputeParameters() const { return PrecomputeParams;  }
 
+	/** 
+	 * This needs to be in this header to be accessible to Lightmass.
+	 * @return colored transmittance given the atmosphere component current state and for SunDirection. 
+	 */
+	static FLinearColor GetTransmittance(const FVector& SunDirection, float AtmosphericFogHeightScaleRayleigh)
+	{
+		// Code from HLSL here as lambda function. It will simulate atmosphere transmittance according to the current sky hardcoded parameterization.
+		// This will change in the future when the sky parameterization and workflow/ui will be updated.
+
+		const float RadiusGround = 6360;
+		const float RadiusAtmosphere = 6420;
+		const int TransmittanceIntegralSamples = 10;
+		const float RadiusLimit = RadiusAtmosphere;
+		const FVector BetaRayleighScattering(5.8e-3f, 1.35e-2f, 3.31e-2f);
+		const FVector BetaMieScattering = FVector(4e-3f, 4e-3f, 4e-3f);
+		const float HeightScaleMie = 1.2f;
+		const float BetaRatio = 0.9f;
+		const FVector BetaMieExtinction = BetaMieScattering / BetaRatio;
+
+		auto Limit = [&](float Radius, float Mu)
+		{
+			float Dout = -Radius * Mu + FMath::Sqrt(Radius * Radius * (Mu * Mu - 1.0) + RadiusLimit * RadiusLimit);
+			float Delta2 = Radius * Radius * (Mu * Mu - 1.0) + RadiusGround * RadiusGround;
+			if (Delta2 >= 0.0)
+			{
+				float Din = -Radius * Mu - FMath::Sqrt(Delta2);
+				if (Din >= 0.0)
+				{
+					Dout = FMath::Min(Dout, Din);
+				}
+			}
+			return Dout;
+		};
+
+		auto OpticalDepth = [&](float H, float Radius, float Mu)
+		{
+			float Result = 0.0;
+			float Dx = Limit(Radius, Mu) / float(TransmittanceIntegralSamples);
+			float Xi = 0.0;
+			float Yi = FMath::Exp(-(Radius - RadiusGround) / H);
+			for (int I = 1; I <= TransmittanceIntegralSamples; ++I)
+			{
+				float Xj = float(I) * Dx;
+				float Yj = FMath::Exp(-(FMath::Sqrt(Radius * Radius + Xj * Xj + 2.0 * Xj * Radius * Mu) - RadiusGround) / H);
+				Result += (Yi + Yj) / 2.0 * Dx;
+				Xi = Xj;
+				Yi = Yj;
+			}
+			return Mu < -FMath::Sqrt(1.0 - (RadiusGround / Radius) * (RadiusGround / Radius)) ? 1e9 : Result;
+		};
+
+		// GetTransmittanceRMuS linear version, assuming we are always close to the ground
+		const float Radius = RadiusGround;
+		float Mu = SunDirection.Z;
+		FVector OpticalDepthRGB = BetaRayleighScattering * OpticalDepth(AtmosphericFogHeightScaleRayleigh, Radius, Mu) + BetaMieExtinction * OpticalDepth(HeightScaleMie, Radius, Mu);
+		OpticalDepthRGB.ComponentMax(FVector(ForceInitToZero));
+
+		return FLinearColor(FMath::Exp(-OpticalDepthRGB.X), FMath::Exp(-OpticalDepthRGB.Y), FMath::Exp(-OpticalDepthRGB.Z));
+	}
+	/**
+	 * This needs to be in this header to be accessible to Lightmass.
+	 * @return colored transmittance given the atmosphere component current state and for SunDirection.
+	 */
+	FLinearColor GetTransmittance(const FVector& SunDirection) const
+	{
+		return GetTransmittance(SunDirection, GetPrecomputeParameters().GetRHeight());
+	}
+
 private:
 #if WITH_EDITORONLY_DATA
 	class FAtmospherePrecomputeDataHandler* PrecomputeDataHandler;
diff --git a/Engine/Source/Runtime/Engine/Classes/Components/LightComponent.h b/Engine/Source/Runtime/Engine/Classes/Components/LightComponent.h
--- a/Engine/Source/Runtime/Engine/Classes/Components/LightComponent.h
+++ b/Engine/Source/Runtime/Engine/Classes/Components/LightComponent.h
@@ -423,6 +423,8 @@ public:
 
 	void InitializeStaticShadowDepthMap();
 
+	FLinearColor GetColoredLightBrightness() const;
+
 	/** 
 	 * Called when property is modified by InterpPropertyTracks
 	 *
diff --git a/Engine/Source/Runtime/Engine/Private/Components/LightComponent.cpp b/Engine/Source/Runtime/Engine/Private/Components/LightComponent.cpp
--- a/Engine/Source/Runtime/Engine/Private/Components/LightComponent.cpp
+++ b/Engine/Source/Runtime/Engine/Private/Components/LightComponent.cpp
@@ -304,19 +304,12 @@ FLightSceneProxy::FLightSceneProxy(const ULightComponent* InLightComponent)
 
 	StaticShadowDepthMap = &LightComponent->StaticShadowDepthMap;
 
-	// Brightness in Lumens
-	float LightBrightness = InLightComponent->ComputeLightBrightness();
-
 	if(LightComponent->IESTexture)
 	{
 		IESTexture = LightComponent->IESTexture;
 	}
-
-	Color = FLinearColor(InLightComponent->LightColor) * LightBrightness;
-	if( InLightComponent->bUseTemperature )
-	{
-		Color *= FLinearColor::MakeFromColorTemperature(InLightComponent->Temperature);
-	}
+	 
+	Color = LightComponent->GetColoredLightBrightness();
 
 	if(LightComponent->LightFunctionMaterial &&
 		LightComponent->LightFunctionMaterial->GetMaterial()->MaterialDomain == MD_LightFunction )
@@ -1231,6 +1224,19 @@ void ULightComponent::InitializeStaticShadowDepthMap()
 	}
 }
 
+FLinearColor ULightComponent::GetColoredLightBrightness() const
+{
+	// Brightness in Lumens
+	float LightBrightness = ComputeLightBrightness();
+	FLinearColor Energy = FLinearColor(LightColor) * LightBrightness;
+	if (bUseTemperature)
+	{
+		Energy *= FLinearColor::MakeFromColorTemperature(Temperature);
+	}
+
+	return Energy;
+}
+
 UMaterialInterface* ULightComponent::GetMaterial(int32 ElementIndex) const
 {
 	if (ElementIndex == 0)
diff --git a/Engine/Source/Runtime/Renderer/Private/AtmosphereRendering.cpp b/Engine/Source/Runtime/Renderer/Private/AtmosphereRendering.cpp
--- a/Engine/Source/Runtime/Renderer/Private/AtmosphereRendering.cpp
+++ b/Engine/Source/Runtime/Renderer/Private/AtmosphereRendering.cpp
@@ -1823,65 +1823,10 @@ void FAtmosphericFogSceneInfo::PrecomputeTextures(FRHICommandListImmediate& RHIC
 }
 #endif
 
-FLinearColor FAtmosphericFogSceneInfo::GetTransmittance(const FVector& SunDirection) const
-{
-	// Code from HLSL here as lambda function. It will simulate atmosphere transmittance according to the current sky hardcoded parameterization.
-	// This will change in the future when the sky parameterization and workflow/ui will be updated.
-
-	const int TransmittanceIntegralSamples = 10;
-	const float RadiusLimit = RadiusAtmosphere;
-	const FVector BetaRayleighScattering(5.8e-3f, 1.35e-2f, 3.31e-2f);
-	const FVector BetaMieScattering = FVector(4e-3f, 4e-3f, 4e-3f);
-	const float HeightScaleMie = 1.2f;
-	const float BetaRatio = 0.9f;
-	const FVector BetaMieExtinction = BetaMieScattering / BetaRatio;
-	const float AtmosphericFogHeightScaleRayleigh = RHeight;
-
-	auto Limit = [&](float Radius, float Mu)
-	{
-		float Dout = -Radius * Mu + FMath::Sqrt(Radius * Radius * (Mu * Mu - 1.0) + RadiusLimit * RadiusLimit);
-		float Delta2 = Radius * Radius * (Mu * Mu - 1.0) + RadiusGround * RadiusGround;
-		if (Delta2 >= 0.0)
-		{
-			float Din = -Radius * Mu - FMath::Sqrt(Delta2);
-			if (Din >= 0.0)
-			{
-				Dout = FMath::Min(Dout, Din);
-			}
-		}
-		return Dout;
-	};
-
-	auto OpticalDepth = [&](float H, float Radius, float Mu)
-	{
-		float Result = 0.0;
-		float Dx = Limit(Radius, Mu) / float(TransmittanceIntegralSamples);
-		float Xi = 0.0;
-		float Yi = FMath::Exp(-(Radius - RadiusGround) / H);
-		for (int I = 1; I <= TransmittanceIntegralSamples; ++I)
-		{
-			float Xj = float(I) * Dx;
-			float Yj = FMath::Exp(-(FMath::Sqrt(Radius * Radius + Xj * Xj + 2.0 * Xj * Radius * Mu) - RadiusGround) / H);
-			Result += (Yi + Yj) / 2.0 * Dx;
-			Xi = Xj;
-			Yi = Yj;
-		}
-		return Mu < -FMath::Sqrt(1.0 - (RadiusGround / Radius) * (RadiusGround / Radius)) ? 1e9 : Result;
-	};
-
-	// GetTransmittanceRMuS linear version, assuming we are always close to the ground
-	const float Radius = RadiusGround;
-	float Mu = SunDirection.Z;
-	FVector OpticalDepthRGB = BetaRayleighScattering * OpticalDepth(AtmosphericFogHeightScaleRayleigh, Radius, Mu) + BetaMieExtinction * OpticalDepth(HeightScaleMie, Radius, Mu);
-	OpticalDepthRGB.ComponentMax(FVector(ForceInitToZero));
-
-	return FLinearColor(FMath::Exp(-OpticalDepthRGB.X), FMath::Exp(-OpticalDepthRGB.Y), FMath::Exp(-OpticalDepthRGB.Z));
-}
-
 void FAtmosphericFogSceneInfo::PrepareSunLightProxy(FLightSceneInfo& SunLight) const
 {
 	// See explanation in https://media.contentapi.ea.com/content/dam/eacom/frostbite/files/s2016-pbs-frostbite-sky-clouds-new.pdf page 26
-	FLinearColor TransmittanceTowardSun = GetTransmittance(-SunLight.Proxy->GetDirection());
+	FLinearColor TransmittanceTowardSun = UAtmosphericFogComponent::GetTransmittance(-SunLight.Proxy->GetDirection(), RHeight);
 
 	FLinearColor SunZenithIlluminance = SunLight.Proxy->GetColor();
 	FLinearColor SunOuterSpaceIlluminance = SunZenithIlluminance / TransmittanceAtZenith;
@@ -1894,7 +1839,7 @@ void FAtmosphericFogSceneInfo::PrepareSunLightProxy(FLightSceneInfo& SunLight) c
 }
 
 /** Initialization constructor. */
-FAtmosphericFogSceneInfo::FAtmosphericFogSceneInfo(UAtmosphericFogComponent* InComponent, const FScene* InScene)
+FAtmosphericFogSceneInfo::FAtmosphericFogSceneInfo(const UAtmosphericFogComponent* InComponent)
 	: Component(InComponent)
 	, SunMultiplier(InComponent->SunMultiplier)
 	, FogMultiplier(InComponent->FogMultiplier)
@@ -1903,7 +1848,7 @@ FAtmosphericFogSceneInfo::FAtmosphericFogSceneInfo(UAtmosphericFogComponent* InC
 	, GroundOffset(InComponent->GroundOffset)
 	, DistanceScale(InComponent->DistanceScale)
 	, AltitudeScale(InComponent->AltitudeScale)
-	, RHeight(InComponent->PrecomputeParams.DensityHeight * InComponent->PrecomputeParams.DensityHeight * InComponent->PrecomputeParams.DensityHeight * 64.f)
+	, RHeight(InComponent->PrecomputeParams.GetRHeight())
 	, StartDistance(InComponent->StartDistance)
 	, DistanceOffset(InComponent->DistanceOffset)
 	, SunDiscScale(InComponent->SunDiscScale)
@@ -1921,7 +1866,7 @@ FAtmosphericFogSceneInfo::FAtmosphericFogSceneInfo(UAtmosphericFogComponent* InC
 	, AtmoshpereOrder(2)
 	, AtmosphereTextures(NULL)
 #endif
-	, TransmittanceAtZenith(GetTransmittance(FVector(0.0f, 0.0f, 1.0f)))
+	, TransmittanceAtZenith(Component->GetTransmittance(FVector(0.0f, 0.0f, 1.0f)))
 {
 	StartDistance *= DistanceScale * 0.00001f; // Convert to km in Atmospheric fog shader
 	// DistanceOffset is in km, no need to change...
@@ -1975,7 +1920,7 @@ void FScene::AddAtmosphericFog(UAtmosphericFogComponent* FogComponent)
 {
 	check(FogComponent);
 
-	FAtmosphericFogSceneInfo* FogSceneInfo = new FAtmosphericFogSceneInfo(FogComponent, this);
+	FAtmosphericFogSceneInfo* FogSceneInfo = new FAtmosphericFogSceneInfo(FogComponent);
 	FScene* Scene = this;
 	ENQUEUE_RENDER_COMMAND(FAddAtmosphericFogCommand)(
 		[Scene, FogSceneInfo](FRHICommandListImmediate& RHICmdList)
diff --git a/Engine/Source/Runtime/Renderer/Private/AtmosphereRendering.h b/Engine/Source/Runtime/Renderer/Private/AtmosphereRendering.h
--- a/Engine/Source/Runtime/Renderer/Private/AtmosphereRendering.h
+++ b/Engine/Source/Runtime/Renderer/Private/AtmosphereRendering.h
@@ -78,7 +78,7 @@ public:
 #endif
 
 	/** Initialization constructor. */
-	explicit FAtmosphericFogSceneInfo(UAtmosphericFogComponent* InComponent, const FScene* InScene);
+	explicit FAtmosphericFogSceneInfo(const UAtmosphericFogComponent* InComponent);
 	~FAtmosphericFogSceneInfo();
 
 	/** Prepare the sun light data as a function of current atmospheric fog state. */
@@ -102,8 +102,6 @@ private:
 
 private:
 	const FLinearColor TransmittanceAtZenith;
-
-	FLinearColor GetTransmittance(const FVector& SunDirection) const;
 };
 
 bool ShouldRenderAtmosphere(const FSceneViewFamily& Family);
