UE-72394: Remove usage of SetRenderTarget in VisualizeRectLightMipTree

diff --git a/Engine/Source/Runtime/Renderer/Private/DeferredShadingRenderer.h b/Engine/Source/Runtime/Renderer/Private/DeferredShadingRenderer.h
--- a/Engine/Source/Runtime/Renderer/Private/DeferredShadingRenderer.h
+++ b/Engine/Source/Runtime/Renderer/Private/DeferredShadingRenderer.h
@@ -486,6 +486,7 @@ private:
 	void CompositeRayTracingSkyLight(FRHICommandListImmediate& RHICmdList, TRefCountPtr<IPooledRenderTarget>& SkyLightRT, TRefCountPtr<IPooledRenderTarget>& HitDistanceRT);
 
 #if RHI_RAYTRACING
+	template <int TextureImportanceSampling> void RenderRayTracingRectLightInternal(FRHICommandListImmediate& RHICmdList, const TArray<FViewInfo>& Views, const FLightSceneInfo& RectLightSceneInfo, TRefCountPtr<IPooledRenderTarget>& ScreenShadowMaskTexture, TRefCountPtr<IPooledRenderTarget>& RayDistanceTexture);
 	void VisualizeRectLightMipTree(FRHICommandListImmediate& RHICmdList, const FViewInfo& View, const FRWBuffer& RectLightMipTree, const FIntVector& RectLightMipTreeDimensions);
 	
 	void RenderRayTracingAmbientOcclusion(FRHICommandListImmediate& RHICmdList, const FLightSceneInfo* LightSceneInfo, TRefCountPtr<IPooledRenderTarget>& AmbientOcclusionRT, TRefCountPtr<IPooledRenderTarget>& AmbientOcclusionHitDistanceRT);
diff --git a/Engine/Source/Runtime/Renderer/Private/RayTracing/RayTracingRectLight.cpp b/Engine/Source/Runtime/Renderer/Private/RayTracing/RayTracingRectLight.cpp
--- a/Engine/Source/Runtime/Renderer/Private/RayTracing/RayTracingRectLight.cpp
+++ b/Engine/Source/Runtime/Renderer/Private/RayTracing/RayTracingRectLight.cpp
@@ -405,13 +405,11 @@ void FDeferredShadingSceneRenderer::VisualizeRectLightMipTree(
 		SceneContext.GetSceneColor()->GetRenderTargetItem().TargetableTexture,
 		RectLightMipTreeRT->GetRenderTargetItem().TargetableTexture
 	};
-	PRAGMA_DISABLE_DEPRECATION_WARNINGS
-	SetRenderTargets(RHICmdList, 2, RenderTargets, SceneContext.GetSceneDepthSurface(), ESimpleRenderTargetMode::EExistingColorAndDepth, FExclusiveDepthStencil::DepthRead_StencilNop);
-	PRAGMA_ENABLE_DEPRECATION_WARNINGS
+	FRHIRenderPassInfo RenderPassInfo(2, RenderTargets, ERenderTargetActions::Load_Store);
+	RHICmdList.BeginRenderPass(RenderPassInfo, TEXT("RectLightMipTree Visualization"));
 
 	// PSO definition
 	FGraphicsPipelineStateInitializer GraphicsPSOInit;
-	SceneContext.BeginRenderingSceneColor(RHICmdList, ESimpleRenderTargetMode::EExistingColorAndDepth, FExclusiveDepthStencil::DepthRead_StencilWrite, true);
 	RHICmdList.ApplyCachedRenderTargets(GraphicsPSOInit);
 	GraphicsPSOInit.BlendState = TStaticBlendState<CW_RGB, BO_Add, BF_One, BF_One>::GetRHI();
 	GraphicsPSOInit.RasterizerState = TStaticRasterizerState<FM_Solid, CM_None>::GetRHI();
@@ -438,13 +436,11 @@ void FDeferredShadingSceneRenderer::VisualizeRectLightMipTree(
 		SceneContext.GetBufferSizeXY(),
 		*VertexShader);
 	ResolveSceneColor(RHICmdList);
+	RHICmdList.EndRenderPass();
 	GVisualizeTexture.SetCheckPoint(RHICmdList, RectLightMipTreeRT);
 
 	// Transition to compute
 	RHICmdList.TransitionResource(EResourceTransitionAccess::ERWBarrier, EResourceTransitionPipeline::EGfxToCompute, RectLightMipTree.UAV);
-
-	ResolveSceneColor(RHICmdList);
-	SceneContext.FinishRenderingSceneColor(RHICmdList);
 }
 
 void FDeferredShadingSceneRenderer::PrepareRayTracingRectLight(const FViewInfo& View, TArray<FRayTracingShaderRHIParamRef>& OutRayGenShaders)
@@ -459,7 +455,7 @@ void FDeferredShadingSceneRenderer::PrepareRayTracingRectLight(const FViewInfo&
 }
 
 template <int TextureImportanceSampling>
-void RenderRayTracingRectLightInternal(
+void FDeferredShadingSceneRenderer::RenderRayTracingRectLightInternal(
 	FRHICommandListImmediate& RHICmdList,
 	const TArray<FViewInfo>& Views,
 	const FLightSceneInfo& RectLightSceneInfo,
@@ -486,6 +482,16 @@ void RenderRayTracingRectLightInternal(
 		}
 	}
 
+#if 0
+	// Debug visualization
+	if (RectLightSceneProxy->SourceTexture)
+	{
+		VisualizeRectLightMipTree(RHICmdList, Views[0],
+			RectLightSceneProxy->RayTracingData->RectLightMipTree,
+			RectLightSceneProxy->RayTracingData->RectLightMipTreeDimensions);
+	}
+#endif
+
 	FLightShaderParameters LightShaderParameters;
 	RectLightSceneProxy->GetLightShaderParameters(LightShaderParameters);
 
