Converting ray tracing SkyLight pass to RDG.

diff --git a/Engine/Source/Runtime/Renderer/Private/RayTracing/RaytracingSkylight.cpp b/Engine/Source/Runtime/Renderer/Private/RayTracing/RaytracingSkylight.cpp
--- a/Engine/Source/Runtime/Renderer/Private/RayTracing/RaytracingSkylight.cpp
+++ b/Engine/Source/Runtime/Renderer/Private/RayTracing/RaytracingSkylight.cpp
@@ -146,118 +146,35 @@ void SetupSkyLightParameters(
 DECLARE_GPU_STAT_NAMED(RayTracingSkyLight, TEXT("Ray Tracing SkyLight"));
 DECLARE_GPU_STAT_NAMED(BuildSkyLightMipTree, TEXT("Build SkyLight Mip Tree"));
 
-template<uint32 EnableTwoSidedGeometry>
-class TSkyLightRGS : public FGlobalShader
+class FRayTracingSkyLightRGS : public FGlobalShader
 {
-	DECLARE_SHADER_TYPE(TSkyLightRGS, Global)
+	DECLARE_GLOBAL_SHADER(FRayTracingSkyLightRGS)
+	SHADER_USE_ROOT_PARAMETER_STRUCT(FRayTracingSkyLightRGS, FGlobalShader)
 
-	static void ModifyCompilationEnvironment(const FGlobalShaderPermutationParameters& Parameters, FShaderCompilerEnvironment& OutEnvironment)
-	{
-		FGlobalShader::ModifyCompilationEnvironment(Parameters, OutEnvironment);
-		OutEnvironment.SetDefine(TEXT("USE_TRANSMISSION"), 1);
-		OutEnvironment.SetDefine(TEXT("ENABLE_TWO_SIDED_GEOMETRY"), EnableTwoSidedGeometry);
-	}
+	class FEnableTwoSidedGeometryDim : SHADER_PERMUTATION_BOOL("ENABLE_TWO_SIDED_GEOMETRY");
+
+	using FPermutationDomain = TShaderPermutationDomain<FEnableTwoSidedGeometryDim>;
 
-public:
 	static bool ShouldCompilePermutation(const FGlobalShaderPermutationParameters& Parameters)
 	{
 		return ShouldCompileRayTracingShadersForProject(Parameters.Platform);
 	}
 
-	TSkyLightRGS() {}
-	virtual ~TSkyLightRGS() {}
-
-	TSkyLightRGS(const ShaderMetaType::CompiledShaderInitializerType& Initializer)
-		: FGlobalShader(Initializer)
-	{
-		ViewParameter.Bind(Initializer.ParameterMap, TEXT("View"));
-		TLASParameter.Bind(Initializer.ParameterMap, TEXT("TLAS"));
-		SceneTexturesParameter.Bind(Initializer.ParameterMap, TEXT("SceneTexturesStruct"));
-		SkyLightParameter.Bind(Initializer.ParameterMap, TEXT("SkyLight"));
-		TransmissionProfilesTextureParameter.Bind(Initializer.ParameterMap, TEXT("SSProfilesTexture"));
-		TransmissionProfilesLinearSamplerParameter.Bind(Initializer.ParameterMap, TEXT("TransmissionProfilesLinearSampler"));
-
-		OcclusionMaskUAVParameter.Bind(Initializer.ParameterMap, TEXT("RWOcclusionMaskUAV"));
-		RayDistanceUAVParameter.Bind(Initializer.ParameterMap, TEXT("RWRayDistanceUAV"));
-	}
-
-	bool Serialize(FArchive& Ar)
-	{
-		bool bShaderHasOutdatedParameters = FGlobalShader::Serialize(Ar);
-		Ar << ViewParameter;
-		Ar << TLASParameter;
-		Ar << SceneTexturesParameter;
-		Ar << SkyLightParameter;
-		Ar << TransmissionProfilesTextureParameter;
-		Ar << TransmissionProfilesLinearSamplerParameter;
-		Ar << OcclusionMaskUAVParameter;
-		Ar << RayDistanceUAVParameter;
-		return bShaderHasOutdatedParameters;
-	}
-
-	void Dispatch(
-		FRHICommandListImmediate& RHICmdList,
-		const FRayTracingScene& RayTracingScene,
-		FUniformBufferRHIParamRef ViewUniformBuffer,
-		FUniformBufferRHIParamRef SceneTexturesUniformBuffer,
-		FUniformBufferRHIParamRef SkyLightUniformBuffer,
-		FUnorderedAccessViewRHIParamRef OcclusionMaskUAV,
-		FUnorderedAccessViewRHIParamRef RayDistanceUAV,
-		uint32 Width, uint32 Height
-	)
-	{
-		FRayTracingPipelineStateInitializer Initializer;
-
-		FRayTracingShaderRHIParamRef RayGenShaderTable[] = { GetRayTracingShader() };
-		Initializer.SetRayGenShaderTable(RayGenShaderTable);
-
-		FRHIRayTracingPipelineState* Pipeline = PipelineStateCache::GetAndOrCreateRayTracingPipelineState(Initializer); 
-
-		FRayTracingShaderBindingsWriter GlobalResources;
-		GlobalResources.Set(TLASParameter, RayTracingScene.RayTracingSceneRHI->GetShaderResourceView());
-		GlobalResources.Set(ViewParameter, ViewUniformBuffer);
-		GlobalResources.Set(SceneTexturesParameter, SceneTexturesUniformBuffer);
-		GlobalResources.Set(SkyLightParameter, SkyLightUniformBuffer);
-		GlobalResources.Set(OcclusionMaskUAVParameter, OcclusionMaskUAV);
-		GlobalResources.Set(RayDistanceUAVParameter, RayDistanceUAV);
-
-		if (TransmissionProfilesTextureParameter.IsBound())
-		{
-			FSceneRenderTargets& SceneContext = FSceneRenderTargets::Get(RHICmdList);
-			const IPooledRenderTarget* PooledRT = GetSubsufaceProfileTexture_RT((FRHICommandListImmediate&)RHICmdList);
-
-			if (!PooledRT)
-			{
-				// no subsurface profile was used yet
-				PooledRT = GSystemTextures.BlackDummy;
-			}
-			const FSceneRenderTargetItem& Item = PooledRT->GetRenderTargetItem();
-
-			GlobalResources.SetTexture(TransmissionProfilesTextureParameter.GetBaseIndex(), Item.ShaderResourceTexture);
-			GlobalResources.SetSampler(TransmissionProfilesLinearSamplerParameter.GetBaseIndex(), TStaticSamplerState<SF_Bilinear, AM_Clamp, AM_Clamp, AM_Clamp>::GetRHI());
-		}
+	BEGIN_SHADER_PARAMETER_STRUCT(FParameters, )
+		SHADER_PARAMETER_SRV(RaytracingAccelerationStructure, TLAS)
+		SHADER_PARAMETER_RDG_TEXTURE_UAV(RWTexture2D<float4>, RWOcclusionMaskUAV)
+		SHADER_PARAMETER_RDG_TEXTURE_UAV(RWTexture2D<float2>, RWRayDistanceUAV)
 
-		RHICmdList.RayTraceDispatch(Pipeline, GetRayTracingShader(), RayTracingScene.RayTracingSceneRHI, GlobalResources, Width, Height);
-	}
+		SHADER_PARAMETER_STRUCT_REF(FViewUniformShaderParameters, ViewUniformBuffer)
+		SHADER_PARAMETER_STRUCT_REF(FSceneTexturesUniformParameters, SceneTexturesStruct)
+		SHADER_PARAMETER_STRUCT_REF(FSkyLightData, SkyLightData)
 
-private:
-	// Input
-	FShaderResourceParameter TLASParameter;
-	FShaderUniformBufferParameter ViewParameter;
-	FShaderUniformBufferParameter SceneTexturesParameter;
-	FShaderUniformBufferParameter SkyLightParameter;
-
-	// SSS Profile
-	FShaderResourceParameter TransmissionProfilesTextureParameter;
-	FShaderResourceParameter TransmissionProfilesLinearSamplerParameter;
-
-	// Output
-	FShaderResourceParameter OcclusionMaskUAVParameter;
-	FShaderResourceParameter RayDistanceUAVParameter;
+		SHADER_PARAMETER_RDG_TEXTURE(Texture2D, SSProfilesTexture)
+		SHADER_PARAMETER_SAMPLER(SamplerState, TransmissionProfilesLinearSampler)
+	END_SHADER_PARAMETER_STRUCT()
 };
 
-IMPLEMENT_SHADER_TYPE(template<>, TSkyLightRGS<0>, TEXT("/Engine/Private/Raytracing/RaytracingSkylightRGS.usf"), TEXT("SkyLightRGS"), SF_RayGen);
-IMPLEMENT_SHADER_TYPE(template<>, TSkyLightRGS<1>, TEXT("/Engine/Private/Raytracing/RaytracingSkylightRGS.usf"), TEXT("SkyLightRGS"), SF_RayGen);
+IMPLEMENT_GLOBAL_SHADER(FRayTracingSkyLightRGS, "/Engine/Private/Raytracing/RaytracingSkylightRGS.usf", "SkyLightRGS", SF_RayGen);
 
 void FDeferredShadingSceneRenderer::BuildSkyLightCdfs(FRHICommandListImmediate& RHICmdList, FSkyLightSceneProxy* SkyLight)
 {
@@ -848,117 +765,134 @@ void FDeferredShadingSceneRenderer::RenderRayTracingSkyLight(
 	{
 		BuildSkyLightCdfs(RHICmdList, Scene->SkyLight);
 	}
-	//VisualizeSkyLightMipTree(RHICmdList, Views[0], Scene->SkyLight->SkyLightMipTreePosX, Scene->SkyLight->SkyLightMipTreeNegX, Scene->SkyLight->SkyLightMipTreePosY, Scene->SkyLight->SkyLightMipTreeNegY, Scene->SkyLight->SkyLightMipTreePosZ, Scene->SkyLight->SkyLightMipTreeNegZ, Scene->SkyLight->SkyLightMipDimensions);
-	//VisualizeSkyLightMipTree(RHICmdList, Views[0], Scene->SkyLight->SkyLightMipTreePdfPosX, Scene->SkyLight->SkyLightMipTreePdfNegX, Scene->SkyLight->SkyLightMipTreePdfPosY, Scene->SkyLight->SkyLightMipTreePdfNegY, Scene->SkyLight->SkyLightMipTreePdfPosZ, Scene->SkyLight->SkyLightMipTreePdfNegZ, Scene->SkyLight->SkyLightMipDimensions);
 
+	// Declare render targets
 	FSceneRenderTargets& SceneContext = FSceneRenderTargets::Get(RHICmdList);
-	FPooledRenderTargetDesc Desc = SceneContext.GetSceneColor()->GetDesc();
-	Desc.Format = PF_FloatRGBA;
-	Desc.Flags &= ~(TexCreate_FastVRAM | TexCreate_Transient);
-	GRenderTargetPool.FindFreeElement(RHICmdList, Desc, SkyLightRT, TEXT("RayTracingSkylight"));
-	ClearUAV(RHICmdList, SkyLightRT->GetRenderTargetItem(), FLinearColor::Black);
+	{
+		FPooledRenderTargetDesc Desc = SceneContext.GetSceneColor()->GetDesc();
+		Desc.Format = PF_FloatRGBA;
+		Desc.Flags &= ~(TexCreate_FastVRAM | TexCreate_Transient);
+		GRenderTargetPool.FindFreeElement(RHICmdList, Desc, SkyLightRT, TEXT("RayTracingSkylight"));
 
-	Desc = SceneContext.GetSceneColor()->GetDesc();
-	Desc.Format = PF_G16R16;
-	Desc.Flags &= ~(TexCreate_FastVRAM | TexCreate_Transient);
-	GRenderTargetPool.FindFreeElement(RHICmdList, Desc, HitDistanceRT, TEXT("RayTracingSkyLightHitDistance"));
-	ClearUAV(RHICmdList, HitDistanceRT->GetRenderTargetItem(), FLinearColor::Black);
+		Desc.Format = PF_G16R16;
+		GRenderTargetPool.FindFreeElement(RHICmdList, Desc, HitDistanceRT, TEXT("RayTracingSkyLightHitDistance"));
+	}
+
+	FRDGBuilder GraphBuilder(RHICmdList);
+	FSceneViewFamilyBlackboard SceneBlackboard;
+	SetupSceneViewFamilyBlackboard(GraphBuilder, &SceneBlackboard);
+
+	// Define RDG targets
+	FRDGTextureRef SkyLightTexture;
+	{
+		FRDGTextureDesc Desc = SceneContext.GetSceneColor()->GetDesc();
+		Desc.Format = PF_FloatRGBA;
+		Desc.Flags &= ~(TexCreate_FastVRAM | TexCreate_Transient);
+		SkyLightTexture = GraphBuilder.CreateTexture(Desc, TEXT("RayTracingSkylight"));
+	}
 
-	// Add SkyLight parameters to uniform buffer
+	FRDGTextureRef RayDistanceTexture;
+	{
+		FRDGTextureDesc Desc = SceneContext.GetSceneColor()->GetDesc();
+		Desc.Format = PF_G16R16;
+		Desc.Flags &= ~(TexCreate_FastVRAM | TexCreate_Transient);
+		RayDistanceTexture = GraphBuilder.CreateTexture(Desc, TEXT("RayTracingSkyLightHitDistance"));
+	}
+
+	// Fill SkyLight parameters
 	FSkyLightData SkyLightData;
 	SetupSkyLightParameters(*Scene, &SkyLightData);
 	SkyLightData.SamplesPerPixel = GRayTracingSkyLightSamplesPerPixel >= 0 ? GRayTracingSkyLightSamplesPerPixel : Scene->SkyLight->SamplesPerPixel;
 	SkyLightData.MaxRayDistance = GRayTracingSkyLightMaxRayDistance;
 	SkyLightData.SamplingStopLevel = GRayTracingSkyLightSamplingStopLevel;
 
-	FUniformBufferRHIRef SkyLightUniformBuffer = RHICreateUniformBuffer(&SkyLightData, FSkyLightData::StaticStructMetadata.GetLayout(), EUniformBufferUsage::UniformBuffer_SingleDraw);
+	FSceneTexturesUniformParameters SceneTextures;
+	SetupSceneTextureUniformParameters(SceneContext, FeatureLevel, ESceneTextureSetupMode::All, SceneTextures);
+
+	TRefCountPtr<IPooledRenderTarget> SubsurfaceProfileRT((IPooledRenderTarget*)GetSubsufaceProfileTexture_RT(RHICmdList));
+	if (!SubsurfaceProfileRT)
+	{
+		SubsurfaceProfileRT = GSystemTextures.BlackDummy;
+	}
+
+	// View-invariant parameters
+	FRayTracingSkyLightRGS::FParameters *PassParameters = GraphBuilder.AllocParameters<FRayTracingSkyLightRGS::FParameters>();
+	PassParameters->RWOcclusionMaskUAV = GraphBuilder.CreateUAV(SkyLightTexture);
+	PassParameters->RWRayDistanceUAV = GraphBuilder.CreateUAV(RayDistanceTexture);
+	PassParameters->SceneTexturesStruct = CreateUniformBufferImmediate(SceneTextures, EUniformBufferUsage::UniformBuffer_SingleDraw);
+	PassParameters->SkyLightData = CreateUniformBufferImmediate(SkyLightData, EUniformBufferUsage::UniformBuffer_SingleDraw);
+	PassParameters->SSProfilesTexture = GraphBuilder.RegisterExternalTexture(SubsurfaceProfileRT);
+	PassParameters->TransmissionProfilesLinearSampler = TStaticSamplerState<SF_Bilinear, AM_Clamp, AM_Clamp, AM_Clamp>::GetRHI();
 
 	for (int32 ViewIndex = 0; ViewIndex < Views.Num(); ++ViewIndex)
 	{
 		FViewInfo& View = Views[ViewIndex];
-		FIntPoint ViewSize = View.ViewRect.Size();
 
-		FSceneTexturesUniformParameters SceneTextures;
-		SetupSceneTextureUniformParameters(SceneContext, FeatureLevel, ESceneTextureSetupMode::All, SceneTextures);
-		FUniformBufferRHIRef SceneTexturesUniformBuffer = RHICreateUniformBuffer(&SceneTextures, FSceneTexturesUniformParameters::StaticStructMetadata.GetLayout(), EUniformBufferUsage::UniformBuffer_SingleDraw);
+		//  View-dependent parameters
+		PassParameters->TLAS = View.RayTracingScene.RayTracingSceneRHI->GetShaderResourceView();
+		PassParameters->ViewUniformBuffer = View.ViewUniformBuffer;
 
-		int32 EnableTwoSidedGeometry = CVarRayTracingSkyLightEnableTwoSidedGeometry.GetValueOnRenderThread();
-		if (EnableTwoSidedGeometry)
-		{
-			TShaderMapRef<TSkyLightRGS<1>> SkyLightRayGenerationShader(GetGlobalShaderMap(FeatureLevel));
-			SkyLightRayGenerationShader->Dispatch(
-				RHICmdList,
-				View.RayTracingScene,
-				View.ViewUniformBuffer,
-				SceneTexturesUniformBuffer,
-				SkyLightUniformBuffer,
-				SkyLightRT->GetRenderTargetItem().UAV,
-				HitDistanceRT->GetRenderTargetItem().UAV,
-				ViewSize.X, ViewSize.Y
-			);
-		}
-		else
+		FRayTracingSkyLightRGS::FPermutationDomain PermutationVector;
+		PermutationVector.Set<FRayTracingSkyLightRGS::FEnableTwoSidedGeometryDim>(CVarRayTracingSkyLightEnableTwoSidedGeometry.GetValueOnRenderThread() != 0);
+		TShaderMapRef<FRayTracingSkyLightRGS> RayGenerationShader(GetGlobalShaderMap(FeatureLevel), PermutationVector);
+		ClearUnusedGraphResources(*RayGenerationShader, PassParameters);
+
+		FIntPoint RayTracingResolution = View.ViewRect.Size();
+		GraphBuilder.AddPass(
+			RDG_EVENT_NAME("SkyLightRayTracing %dx%d", RayTracingResolution.X, RayTracingResolution.Y),
+			PassParameters,
+			ERenderGraphPassFlags::Compute,
+			[PassParameters, this, &View, RayGenerationShader, RayTracingResolution](FRHICommandList& RHICmdList)
 		{
-			TShaderMapRef<TSkyLightRGS<0>> SkyLightRayGenerationShader(GetGlobalShaderMap(FeatureLevel));
-			SkyLightRayGenerationShader->Dispatch(
-				RHICmdList,
-				View.RayTracingScene,
-				View.ViewUniformBuffer,
-				SceneTexturesUniformBuffer,
-				SkyLightUniformBuffer,
-				SkyLightRT->GetRenderTargetItem().UAV,
-				HitDistanceRT->GetRenderTargetItem().UAV,
-				ViewSize.X, ViewSize.Y
-			);
-		}
-	}
+			FRayTracingShaderBindingsWriter GlobalResources;
+			SetShaderParameters(GlobalResources, *RayGenerationShader, *PassParameters);
 
-	// Transition to graphics pipeline
-	FComputeFenceRHIRef Fence = RHICmdList.CreateComputeFence(TEXT("RayTracingSkyLight"));
-	RHICmdList.TransitionResource(EResourceTransitionAccess::ERWBarrier, EResourceTransitionPipeline::EComputeToGfx, SkyLightRT->GetRenderTargetItem().UAV);
-	RHICmdList.TransitionResource(EResourceTransitionAccess::ERWBarrier, EResourceTransitionPipeline::EComputeToGfx, HitDistanceRT->GetRenderTargetItem().UAV, Fence);
-	GVisualizeTexture.SetCheckPoint(RHICmdList, SkyLightRT);
-	GVisualizeTexture.SetCheckPoint(RHICmdList, HitDistanceRT);
+			// Declare RT pipeline
+			FRayTracingPipelineStateInitializer Initializer;
+			FRayTracingShaderRHIParamRef RayGenShaderTable[] = { RayGenerationShader->GetRayTracingShader() };
+			Initializer.SetRayGenShaderTable(RayGenShaderTable);
 
-	// Denoising
-	if (GRayTracingSkyLightDenoiser != 0)
-	{
-		FRDGBuilder GraphBuilder(RHICmdList);
-		FSceneViewFamilyBlackboard SceneBlackboard;
-		SetupSceneViewFamilyBlackboard(GraphBuilder, &SceneBlackboard);
+			FRHIRayTracingPipelineState* Pipeline = PipelineStateCache::GetAndOrCreateRayTracingPipelineState(Initializer);
+			FRayTracingSceneRHIParamRef RayTracingSceneRHI = View.RayTracingScene.RayTracingSceneRHI;
+			RHICmdList.RayTraceDispatch(Pipeline, RayGenerationShader->GetRayTracingShader(), RayTracingSceneRHI, GlobalResources, RayTracingResolution.X, RayTracingResolution.Y);
+		});
 
-		const IScreenSpaceDenoiser* DefaultDenoiser = IScreenSpaceDenoiser::GetDefaultDenoiser();
-		const IScreenSpaceDenoiser* DenoiserToUse = DefaultDenoiser;// GRayTracingGlobalIlluminationDenoiser == 1 ? DefaultDenoiser : GScreenSpaceDenoiser;
+		// Denoising
+		if (GRayTracingSkyLightDenoiser != 0)
+		{
+			const IScreenSpaceDenoiser* DefaultDenoiser = IScreenSpaceDenoiser::GetDefaultDenoiser();
+			const IScreenSpaceDenoiser* DenoiserToUse = DefaultDenoiser;// GRayTracingGlobalIlluminationDenoiser == 1 ? DefaultDenoiser : GScreenSpaceDenoiser;
 
-		IScreenSpaceDenoiser::FGlobalIlluminationInputs DenoiserInputs;
-		DenoiserInputs.Color = GraphBuilder.RegisterExternalTexture(SkyLightRT);
-		DenoiserInputs.RayHitDistance = GraphBuilder.RegisterExternalTexture(HitDistanceRT);
+			IScreenSpaceDenoiser::FGlobalIlluminationInputs DenoiserInputs;
+			DenoiserInputs.Color = SkyLightTexture;
+			DenoiserInputs.RayHitDistance = RayDistanceTexture;
 
-		{
-			IScreenSpaceDenoiser::FAmbientOcclusionRayTracingConfig RayTracingConfig;
-			RayTracingConfig.ResolutionFraction = 1.0;
-			RayTracingConfig.RayCountPerPixel = GRayTracingSkyLightSamplesPerPixel >= 0 ? GRayTracingSkyLightSamplesPerPixel : Scene->SkyLight->SamplesPerPixel;
-
-			FViewInfo& View = Views[0];
-			RDG_EVENT_SCOPE(GraphBuilder, "%s%s(SkyLight) %dx%d",
-				DenoiserToUse != DefaultDenoiser ? TEXT("ThirdParty ") : TEXT(""),
-				DenoiserToUse->GetDebugName(),
-				View.ViewRect.Width(), View.ViewRect.Height());
-
-			IScreenSpaceDenoiser::FGlobalIlluminationOutputs DenoiserOutputs = DenoiserToUse->DenoiseSkyLight(
-				GraphBuilder,
-				View,
-				&View.PrevViewInfo,
-				SceneBlackboard,
-				DenoiserInputs,
-				RayTracingConfig);
-
-			GraphBuilder.QueueTextureExtraction(DenoiserOutputs.Color, &SkyLightRT);
+			{
+				IScreenSpaceDenoiser::FAmbientOcclusionRayTracingConfig RayTracingConfig;
+				RayTracingConfig.ResolutionFraction = 1.0;
+				RayTracingConfig.RayCountPerPixel = GRayTracingSkyLightSamplesPerPixel >= 0 ? GRayTracingSkyLightSamplesPerPixel : Scene->SkyLight->SamplesPerPixel;
+
+				RDG_EVENT_SCOPE(GraphBuilder, "%s%s(SkyLight) %dx%d",
+					DenoiserToUse != DefaultDenoiser ? TEXT("ThirdParty ") : TEXT(""),
+					DenoiserToUse->GetDebugName(),
+					View.ViewRect.Width(), View.ViewRect.Height());
+
+				IScreenSpaceDenoiser::FGlobalIlluminationOutputs DenoiserOutputs = DenoiserToUse->DenoiseSkyLight(
+					GraphBuilder,
+					View,
+					&View.PrevViewInfo,
+					SceneBlackboard,
+					DenoiserInputs,
+					RayTracingConfig);
+
+				SkyLightTexture = DenoiserOutputs.Color;
+			}
 		}
-
-		GraphBuilder.Execute();
-		GVisualizeTexture.SetCheckPoint(RHICmdList, SkyLightRT);
 	}
+
+	GraphBuilder.QueueTextureExtraction(SkyLightTexture, &SkyLightRT);
+	GraphBuilder.Execute();
+	GVisualizeTexture.SetCheckPoint(RHICmdList, SkyLightRT);
 }
 
 class FCompositeSkyLightPS : public FGlobalShader
