Changing r.RayTracing.SkyLight.EnableTwoSidedGeometry default to 1 to match ray traced shadows behavior.

diff --git a/Engine/Source/Runtime/Renderer/Private/RayTracing/RaytracingSkylight.cpp b/Engine/Source/Runtime/Renderer/Private/RayTracing/RaytracingSkylight.cpp
--- a/Engine/Source/Runtime/Renderer/Private/RayTracing/RaytracingSkylight.cpp
+++ b/Engine/Source/Runtime/Renderer/Private/RayTracing/RaytracingSkylight.cpp
@@ -63,8 +63,8 @@ static FAutoConsoleVariableRef CVarRayTracingSkyLightDenoiser(
 
 static TAutoConsoleVariable<int32> CVarRayTracingSkyLightEnableTwoSidedGeometry(
 	TEXT("r.RayTracing.SkyLight.EnableTwoSidedGeometry"),
-	0,
-	TEXT("Enables two-sided geometry when tracing shadow rays (default = 0)"),
+	1,
+	TEXT("Enables two-sided geometry when tracing shadow rays (default = 1)"),
 	ECVF_RenderThreadSafe
 );
 
