Fixed static analysis warning about ShadowInfo nullptr deferencing.
#rnx

diff --git a/Engine/Source/Runtime/Renderer/Private/ShadowRendering.h b/Engine/Source/Runtime/Renderer/Private/ShadowRendering.h
--- a/Engine/Source/Runtime/Renderer/Private/ShadowRendering.h
+++ b/Engine/Source/Runtime/Renderer/Private/ShadowRendering.h
@@ -1198,7 +1198,7 @@ public:
 
 		if (LightPositionOrDirection.IsBound())
 		{
-			const FVector LightPosition = ShadowInfo->GetLightSceneInfo().Proxy->GetPosition();
+			const FVector LightPosition = ShadowInfo ? ShadowInfo->GetLightSceneInfo().Proxy->GetPosition() : FVector::ZeroVector;
 			SetShaderValue(RHICmdList, ShaderRHI, LightPositionOrDirection, FVector4(LightPosition, 1));
 		}
 		
