DR - Move all Queries to pools which can be overriden by an RHI implentation. Currently only using the default pool.

diff --git a/Engine/Source/Runtime/Engine/Private/DynamicResolution.cpp b/Engine/Source/Runtime/Engine/Private/DynamicResolution.cpp
--- a/Engine/Source/Runtime/Engine/Private/DynamicResolution.cpp
+++ b/Engine/Source/Runtime/Engine/Private/DynamicResolution.cpp
@@ -556,10 +556,8 @@ public:
 			InFlightFrame.HeuristicHistoryEntry = Heuristic.CreateNewPreviousFrameTimings_RenderThread(
 				PrevGameThreadTimeMs, PrevRenderThreadTimeMs);
 
-			{
-				InFlightFrame.BeginFrameQuery = RHICreateRenderQuery(RQT_AbsoluteTime);
-				RHICmdList.EndRenderQuery(InFlightFrame.BeginFrameQuery);
-			}
+			RHICmdList.EndRenderQuery(InFlightFrame.BeginFrameQuery.GetQuery());
+			InFlightFrame.BeginFrameQuerySubmitted = true;
 		}
 		else
 		{
@@ -589,20 +587,25 @@ public:
 		{
 			InFlightFrameQueries& InFlightFrame = InFlightFrames[CurrentFrameInFlightIndex];
 
-			FRenderQueryRHIRef* QueryPtr = nullptr;
+			FRHIRenderQuery* QueryPtr = nullptr;
 			switch (Event)
 			{
 			case EDynamicResolutionStateEvent::BeginDynamicResolutionRendering:
-				QueryPtr = &InFlightFrame.BeginDynamicResolutionQuery; break;
+				QueryPtr = InFlightFrame.BeginDynamicResolutionQuery.GetQuery(); 
+				InFlightFrame.BeginDynamicResolutionQuerySubmitted = true;
+				break;
 			case EDynamicResolutionStateEvent::EndDynamicResolutionRendering:
-				QueryPtr = &InFlightFrame.EndDynamicResolutionQuery; break;
+				QueryPtr = InFlightFrame.EndDynamicResolutionQuery.GetQuery(); 
+				InFlightFrame.EndDynamicResolutionQuerySubmitted = true;
+				break;
 			case EDynamicResolutionStateEvent::EndFrame:
-				QueryPtr = &InFlightFrame.EndFrameQuery; break;
+				QueryPtr = InFlightFrame.EndFrameQuery.GetQuery(); 
+				InFlightFrame.EndFrameQuerySubmitted = true;
+				break;
 			default: check(0);
 			}
 
-			*QueryPtr = RHICreateRenderQuery(RQT_AbsoluteTime);
-			RHICmdList.EndRenderQuery(*QueryPtr);
+			RHICmdList.EndRenderQuery(QueryPtr);
 		}
 
 		// Clobber CurrentFrameInFlightIndex for internal checks.
@@ -630,16 +633,31 @@ private:
 	struct InFlightFrameQueries
 	{
 		// GPU queries.
-		FRenderQueryRHIRef BeginFrameQuery;
-		FRenderQueryRHIRef BeginDynamicResolutionQuery;
-		FRenderQueryRHIRef EndDynamicResolutionQuery;
-		FRenderQueryRHIRef EndFrameQuery;
+		FRenderQueryPoolRHIRef QueryPool;
+		FRHIPooledRenderQuery BeginFrameQuery;
+		FRHIPooledRenderQuery BeginDynamicResolutionQuery;
+		FRHIPooledRenderQuery EndDynamicResolutionQuery;
+		FRHIPooledRenderQuery EndFrameQuery;
 
 		// Heuristic's history 
 		uint64 HeuristicHistoryEntry;
 
+		uint32 BeginFrameQuerySubmitted : 1;
+		uint32 BeginDynamicResolutionQuerySubmitted : 1;
+		uint32 EndDynamicResolutionQuerySubmitted : 1;
+		uint32 EndFrameQuerySubmitted : 1;
+
 		InFlightFrameQueries()
-			: HeuristicHistoryEntry(FDynamicResolutionHeuristicProxy::kInvalidEntryId)
+			: QueryPool(RHICreateRenderQueryPool(RQT_AbsoluteTime, 4))
+			, BeginFrameQuery(QueryPool->AllocateQuery())
+			, BeginDynamicResolutionQuery(QueryPool->AllocateQuery())
+			, EndDynamicResolutionQuery(QueryPool->AllocateQuery())
+			, EndFrameQuery(QueryPool->AllocateQuery())
+			, HeuristicHistoryEntry(FDynamicResolutionHeuristicProxy::kInvalidEntryId) 
+			, BeginFrameQuerySubmitted(false)
+			, BeginDynamicResolutionQuerySubmitted(false)
+			, EndDynamicResolutionQuerySubmitted(false)
+			, EndFrameQuerySubmitted(false)
 		{ }
 	};
 
@@ -677,32 +695,36 @@ private:
 
 			int32 LandingCount = 0;
 			int32 QueryCount = 0;
-			if (InFlightFrame.BeginFrameQuery.IsValid())
+			if (InFlightFrame.BeginFrameQuery.GetQuery() != nullptr && InFlightFrame.BeginFrameQuerySubmitted == true)
 			{
 				LandingCount += RHIGetRenderQueryResult(
-					InFlightFrame.BeginFrameQuery, BeginFrameResult, bWait) ? 1 : 0;
+					InFlightFrame.BeginFrameQuery.GetQuery(), BeginFrameResult, bWait) ? 1 : 0;
 				QueryCount += 1;
+				InFlightFrame.BeginFrameQuerySubmitted = false;
 			}
 
-			if (InFlightFrame.BeginDynamicResolutionQuery.IsValid())
+			if (InFlightFrame.BeginDynamicResolutionQuery.GetQuery() != nullptr && InFlightFrame.BeginDynamicResolutionQuerySubmitted == true)
 			{
 				LandingCount += RHIGetRenderQueryResult(
-					InFlightFrame.BeginDynamicResolutionQuery, BeginDynamicResolutionResult, bWait) ? 1 : 0;
+					InFlightFrame.BeginDynamicResolutionQuery.GetQuery(), BeginDynamicResolutionResult, bWait) ? 1 : 0;
 				QueryCount += 1;
+				InFlightFrame.BeginDynamicResolutionQuerySubmitted = false;
 			}
 
-			if (InFlightFrame.EndDynamicResolutionQuery.IsValid())
+			if (InFlightFrame.EndDynamicResolutionQuery.GetQuery() != nullptr && InFlightFrame.EndDynamicResolutionQuerySubmitted == true)
 			{
 				LandingCount += RHIGetRenderQueryResult(
-					InFlightFrame.EndDynamicResolutionQuery, EndDynamicResolutionResult, bWait) ? 1 : 0;
+					InFlightFrame.EndDynamicResolutionQuery.GetQuery(), EndDynamicResolutionResult, bWait) ? 1 : 0;
 				QueryCount += 1;
+				InFlightFrame.EndDynamicResolutionQuerySubmitted = false;
 			}
 
-			if (InFlightFrame.EndFrameQuery.IsValid())
+			if (InFlightFrame.EndFrameQuery.GetQuery() != nullptr && InFlightFrame.EndFrameQuerySubmitted == true)
 			{
 				LandingCount += RHIGetRenderQueryResult(
-					InFlightFrame.EndFrameQuery, EndFrameResult, bWait) ? 1 : 0;
+					InFlightFrame.EndFrameQuery.GetQuery(), EndFrameResult, bWait) ? 1 : 0;
 				QueryCount += 1;
+				InFlightFrame.EndFrameQuerySubmitted = false;
 			}
 
 			check(QueryCount == 0 || QueryCount == 4);
@@ -738,7 +760,7 @@ private:
 		for (int32 i = 0; i < InFlightFrames.Num(); i++)
 		{
 			auto& InFlightFrame = InFlightFrames[i];
-			if (!InFlightFrame.BeginFrameQuery.IsValid())
+			if (InFlightFrame.BeginFrameQuery.GetQuery() == nullptr)
 			{
 				CurrentFrameInFlightIndex = i;
 				break;
diff --git a/Engine/Source/Runtime/Engine/Private/SceneUtils.cpp b/Engine/Source/Runtime/Engine/Private/SceneUtils.cpp
--- a/Engine/Source/Runtime/Engine/Private/SceneUtils.cpp
+++ b/Engine/Source/Runtime/Engine/Private/SceneUtils.cpp
@@ -185,7 +185,7 @@ public:
 	static const uint64 InvalidQueryResult = 0xFFFFFFFFFFFFFFFFull;
 
 public:
-	FRealtimeGPUProfilerEvent(const FName& InName, const FName& InStatName, FRenderQueryPool* RenderQueryPool)
+	FRealtimeGPUProfilerEvent(const FName& InName, const FName& InStatName, FRHIRenderQueryPool* RenderQueryPool, uint32& QueryCount)
 		: StartResultMicroseconds(InvalidQueryResult)
 		, EndResultMicroseconds(InvalidQueryResult)
 		, FrameNumber(-1)
@@ -199,19 +199,20 @@ public:
 		Name = InName;
 
 		const int MaxGPUQueries = CVarGPUStatsMaxQueriesPerFrame.GetValueOnRenderThread();
-		if ( MaxGPUQueries == -1 || RenderQueryPool->GetAllocatedQueryCount() < MaxGPUQueries )
+		if ( MaxGPUQueries == -1 || QueryCount < uint32(MaxGPUQueries) )
 		{
 			StartQuery = RenderQueryPool->AllocateQuery();
 			EndQuery = RenderQueryPool->AllocateQuery();
+			QueryCount += 2;
 		}
 	}
 
 	bool HasQueriesAllocated() const 
 	{ 
-		return IsValidRef(StartQuery); 
+		return StartQuery.GetQuery() != nullptr;
 	}
 
-	void ReleaseQueries(FRenderQueryPool* RenderQueryPool, FRHICommandListImmediate* RHICmdListPtr)
+	void ReleaseQueries(FRHIRenderQueryPool* RenderQueryPool, FRHICommandListImmediate* RHICmdListPtr, uint32& QueryCount)
 	{
 		if ( HasQueriesAllocated() )
 		{
@@ -221,16 +222,17 @@ public:
 				uint64 Temp;
 				if (bBeginQueryInFlight)
 				{
-					RHICmdListPtr->GetRenderQueryResult(StartQuery, Temp, false);
+					RHICmdListPtr->GetRenderQueryResult(StartQuery.GetQuery(), Temp, false);
 				}
 
 				if (bEndQueryInFlight)
 				{
-					RHICmdListPtr->GetRenderQueryResult(EndQuery, Temp, false);
+					RHICmdListPtr->GetRenderQueryResult(EndQuery.GetQuery(), Temp, false);
 				}
 			}
-			RenderQueryPool->ReleaseQuery(StartQuery);
-			RenderQueryPool->ReleaseQuery(EndQuery);
+			StartQuery.ReleaseQuery();
+			EndQuery.ReleaseQuery();
+			QueryCount -= 2;
 		}
 	}
 
@@ -240,9 +242,9 @@ public:
 		check(!bInsideQuery);
 		bInsideQuery = true;
 
-		if (IsValidRef(StartQuery))
+		if (StartQuery.GetQuery() != nullptr)
 		{
-			RHICmdList.EndRenderQuery(StartQuery);
+			RHICmdList.EndRenderQuery(StartQuery.GetQuery());
 			bBeginQueryInFlight = true;
 		}
 		StartResultMicroseconds = InvalidQueryResult;
@@ -258,7 +260,7 @@ public:
 
 		if ( HasQueriesAllocated() )
 		{
-			RHICmdList.EndRenderQuery(EndQuery);
+			RHICmdList.EndRenderQuery(EndQuery.GetQuery());
 			bEndQueryInFlight = true;
 		}
 	}
@@ -272,7 +274,7 @@ public:
 		{
 			if (StartResultMicroseconds == InvalidQueryResult)
 			{
-				if (!RHICmdList.GetRenderQueryResult(StartQuery, StartResultMicroseconds, false))
+				if (!RHICmdList.GetRenderQueryResult(StartQuery.GetQuery(), StartResultMicroseconds, false))
 				{
 					StartResultMicroseconds = InvalidQueryResult;
 				}
@@ -280,7 +282,7 @@ public:
 			}
 			if (EndResultMicroseconds == InvalidQueryResult)
 			{
-				if (!RHICmdList.GetRenderQueryResult(EndQuery, EndResultMicroseconds, false))
+				if (!RHICmdList.GetRenderQueryResult(EndQuery.GetQuery(), EndResultMicroseconds, false))
 				{
 					EndResultMicroseconds = InvalidQueryResult;
 				}
@@ -338,8 +340,8 @@ public:
 	}
 
 private:
-	FRenderQueryRHIRef StartQuery;
-	FRenderQueryRHIRef EndQuery;
+	FRHIPooledRenderQuery StartQuery;
+	FRHIPooledRenderQuery EndQuery;
 #if STATS
 	FName StatName;
 #endif
@@ -359,9 +361,11 @@ Container for a single frame's GPU stats
 -----------------------------------------------------------------------------*/
 class FRealtimeGPUProfilerFrame
 {
+	uint32& QueryCount;
 public:
-	FRealtimeGPUProfilerFrame(FRenderQueryPool* InRenderQueryPool)
-		: FrameNumber(-1)
+	FRealtimeGPUProfilerFrame(FRenderQueryPoolRHIRef InRenderQueryPool, uint32& InQueryCount)
+		: QueryCount(InQueryCount)
+		, FrameNumber(-1)
 		, RenderQueryPool(InRenderQueryPool)
 	{}
 
@@ -373,7 +377,7 @@ public:
 	void PushEvent(FRHICommandListImmediate& RHICmdList, const FName& Name, const FName& StatName)
 	{
 		// TODO: this should really use a pool / free list
-		FRealtimeGPUProfilerEvent* Event = new FRealtimeGPUProfilerEvent(Name, StatName, RenderQueryPool);
+		FRealtimeGPUProfilerEvent* Event = new FRealtimeGPUProfilerEvent(Name, StatName, RenderQueryPool, QueryCount);
 		const int32 EventIndex = GpuProfilerEvents.Num();
 
 		GpuProfilerEvents.Add(Event);
@@ -408,7 +412,7 @@ public:
 		{
 			if (GpuProfilerEvents[Index])
 			{
-				GpuProfilerEvents[Index]->ReleaseQueries(RenderQueryPool, RHICommandListPtr);
+				GpuProfilerEvents[Index]->ReleaseQueries(RenderQueryPool, RHICommandListPtr, QueryCount);
 				delete GpuProfilerEvents[Index];
 			}
 		}
@@ -575,7 +579,7 @@ private:
 	TArray<FGPUEventTimeAggregate> EventAggregates;
 
 	uint32 FrameNumber;
-	FRenderQueryPool* RenderQueryPool;
+	FRenderQueryPoolRHIRef RenderQueryPool;
 };
 
 /*-----------------------------------------------------------------------------
@@ -599,10 +603,11 @@ FRealtimeGPUProfiler::FRealtimeGPUProfiler()
 	, bStatGatheringPaused(false)
 	, bInBeginEndBlock(false)
 {
-	RenderQueryPool = new FRenderQueryPool(RQT_AbsoluteTime);
+	const int MaxGPUQueries = CVarGPUStatsMaxQueriesPerFrame.GetValueOnRenderThread();
+	RenderQueryPool = RHICreateRenderQueryPool(RQT_AbsoluteTime, (MaxGPUQueries > 0) ? MaxGPUQueries * 2 : 10240);
 	for (int Index = 0; Index < NumGPUProfilerBufferedFrames; Index++)
 	{
-		Frames.Add(new FRealtimeGPUProfilerFrame(RenderQueryPool));
+		Frames.Add(new FRealtimeGPUProfilerFrame(RenderQueryPool, QueryCount));
 	}
 }
 
@@ -613,8 +618,7 @@ void FRealtimeGPUProfiler::Release()
 		delete Frames[Index];
 	}
 	Frames.Empty();
-	delete RenderQueryPool;
-	RenderQueryPool = nullptr;
+	RenderQueryPool.SafeRelease();
 }
 
 void FRealtimeGPUProfiler::BeginFrame(FRHICommandListImmediate& RHICmdList)
@@ -740,48 +744,3 @@ void FScopedGPUStatEvent::End()
 	}
 }
 #endif // HAS_GPU_STATS
-
-/*-----------------------------------------------------------------------------
-FRenderQueryPool
------------------------------------------------------------------------------*/
-FRenderQueryPool::~FRenderQueryPool()
-{
-	Release();
-}
-
-void FRenderQueryPool::Release()
-{
-	Queries.Empty();
-	NumQueriesAllocated = 0;
-}
-
-FRenderQueryRHIRef FRenderQueryPool::AllocateQuery()
-{
-	// Are we out of available render queries?
-	NumQueriesAllocated++;
-	if (Queries.Num() == 0)
-	{
-		// Create a new render query.
-		return RHICreateRenderQuery(QueryType);
-	}
-
-	return Queries.Pop(/*bAllowShrinking=*/ false);
-}
-
-void FRenderQueryPool::ReleaseQuery(FRenderQueryRHIRef &Query)
-{
-	if (IsValidRef(Query))
-	{
-		NumQueriesAllocated--;
-#if RENDER_QUERY_POOLING_ENABLED
-		// Is no one else keeping a refcount to the query?
-		if (Query.GetRefCount() == 1)
-		{
-			// Return it to the pool.
-			Queries.Add(Query);
-		}
-#endif
-		// De-ref without deleting.
-		Query = NULL;
-	}
-}
diff --git a/Engine/Source/Runtime/Engine/Private/UnrealEngine.cpp b/Engine/Source/Runtime/Engine/Private/UnrealEngine.cpp
--- a/Engine/Source/Runtime/Engine/Private/UnrealEngine.cpp
+++ b/Engine/Source/Runtime/Engine/Private/UnrealEngine.cpp
@@ -1703,6 +1703,9 @@ void UEngine::PreExit()
 	SetCustomTimeStep(nullptr);
 
 	ShutdownHMD();
+
+	DynamicResolutionState.Reset();
+	NextDynamicResolutionState.Reset();
 }
 
 void UEngine::ShutdownHMD()
diff --git a/Engine/Source/Runtime/Engine/Public/SceneUtils.h b/Engine/Source/Runtime/Engine/Public/SceneUtils.h
--- a/Engine/Source/Runtime/Engine/Public/SceneUtils.h
+++ b/Engine/Source/Runtime/Engine/Public/SceneUtils.h
@@ -258,7 +258,8 @@ private:
 	int32 WriteBufferIndex;
 	int32 ReadBufferIndex;
 	uint32 WriteFrameNumber;
-	FRenderQueryPool* RenderQueryPool;
+	uint32 QueryCount = 0;
+	FRenderQueryPoolRHIRef RenderQueryPool;
 	bool bStatGatheringPaused;
 	bool bInBeginEndBlock;
 };
@@ -323,40 +324,6 @@ ENGINE_API EMobileHDRMode GetMobileHDRMode();
 
 ENGINE_API bool IsMobileColorsRGB();
 
-/**
-* A pool of render (e.g. occlusion/timer) queries which are allocated individually, and returned to the pool as a group.
-*/
-class ENGINE_API FRenderQueryPool
-{
-public:
-	FRenderQueryPool(ERenderQueryType InQueryType)
-		: QueryType(InQueryType)
-		, NumQueriesAllocated(0)
-	{ }
-
-	virtual ~FRenderQueryPool();
-
-	/** Releases all the render queries in the pool. */
-	void Release();
-
-	/** Allocates an render query from the pool. */
-	FRenderQueryRHIRef AllocateQuery();
-
-	/** De-reference an render query, returning it to the pool instead of deleting it when the refcount reaches 0. */
-	void ReleaseQuery(FRenderQueryRHIRef &Query);
-
-	/** Returns the number of currently allocated queries. This is not necessarily the same as the pool size */
-	int32 GetAllocatedQueryCount() const { return NumQueriesAllocated;  }
-
-private:
-	/** Container for available render queries. */
-	TArray<FRenderQueryRHIRef> Queries;
-
-	ERenderQueryType QueryType;
-
-	int32 NumQueriesAllocated;
-};
-
 // Callback for calling one action (typical use case: delay a clear until it's actually needed)
 class FDelayedRendererAction
 {
diff --git a/Engine/Source/Runtime/RHI/Private/DynamicRHI.cpp b/Engine/Source/Runtime/RHI/Private/DynamicRHI.cpp
--- a/Engine/Source/Runtime/RHI/Private/DynamicRHI.cpp
+++ b/Engine/Source/Runtime/RHI/Private/DynamicRHI.cpp
@@ -351,3 +351,44 @@ void FDynamicRHI::RHIUpdateShaderResourceView(FShaderResourceViewRHIParamRef SRV
 {
 	UE_LOG(LogRHI, Fatal, TEXT("RHIUpdateShaderResourceView isn't implemented for the current RHI"));
 }
+
+FDefaultRHIRenderQueryPool::FDefaultRHIRenderQueryPool(ERenderQueryType InQueryType, FDynamicRHI* InDynamicRHI, uint32 InNumQueries)
+	: DynamicRHI(InDynamicRHI)
+	, QueryType(InQueryType)
+	, NumQueries(InNumQueries)
+{
+	if (NumQueries != UINT32_MAX && (GSupportsTimestampRenderQueries || InQueryType != RQT_AbsoluteTime))
+	{
+		Queries.Reserve(NumQueries);
+		for (uint32 i = 0; i < NumQueries; i++)
+		{
+			Queries.Push(DynamicRHI->RHICreateRenderQuery(QueryType));
+			++AllocatedQueries;
+		}
+	}
+}
+
+FRHIPooledRenderQuery FDefaultRHIRenderQueryPool::AllocateQuery()
+{
+	if (Queries.Num() > 0)
+	{
+		return FRHIPooledRenderQuery(this, Queries.Pop());
+	}
+	else
+	{
+		ensure(++AllocatedQueries <= NumQueries);
+		return FRHIPooledRenderQuery(this, DynamicRHI->RHICreateRenderQuery(QueryType));
+	}
+}
+
+void FDefaultRHIRenderQueryPool::ReleaseQuery(TRefCountPtr<FRHIRenderQuery>&& Query)
+{
+	checkf(Query.IsValid() && Query.GetRefCount() <= 2, TEXT("Query has been released but reference still held: use FRHIPooledRenderQuery::GetQueryRef() with extreme caution"));
+	Queries.Push(MoveTemp(Query));
+	check(!Query.IsValid());
+}
+
+FRenderQueryPoolRHIRef RHICreateRenderQueryPool(ERenderQueryType QueryType, uint32 NumQueries)
+{
+	return GDynamicRHI->RHICreateRenderQueryPool(QueryType, NumQueries);
+}
\ No newline at end of file
diff --git a/Engine/Source/Runtime/RHI/Private/GPUProfiler.cpp b/Engine/Source/Runtime/RHI/Private/GPUProfiler.cpp
--- a/Engine/Source/Runtime/RHI/Private/GPUProfiler.cpp
+++ b/Engine/Source/Runtime/RHI/Private/GPUProfiler.cpp
@@ -617,36 +617,4 @@ FGPUTimingCalibrationTimestamp FGPUTiming::GCalibrationTimestamp = {};
 /** Whether the static variables have been initialized. */
 bool FGPUTiming::GAreGlobalsInitialized = false;
 
-
-
-float FWindowedGPUTimer::GetElapsedAverage(FRHICommandListImmediate& RHICmdList, float &OutAvgTimeInSeconds)
-{
-	if (QueriesFinished < StartQueries.Num())
-	{
-		return 0.0f;
-	}
-
-	float TotalTime = 0;
-
-	// Grab the queries in our window (specified number of frames old) and calculate total time as an average of the winow size
-	//
-	for (int32 i = StartQueries.Num() - WindowSize; i < StartQueries.Num(); i++)
-	{
-		uint64 StartTime, EndTime;
-		bool StartSucceeded = RHICmdList.GetRenderQueryResult(StartQueries[i], StartTime, false);
-		bool EndSucceeded = RHICmdList.GetRenderQueryResult(EndQueries[i], EndTime, false);
-
-		// figure out what the failure rate of the queries is; they fail because the GPU hasn't finished them when we
-		// try to get the data, so if the failure rate is too high, number of frames behind needs to be increased
-		QueriesFailed += StartSucceeded&&EndSucceeded ? -1 : 1;
-		QueriesFailed = FMath::Max<int32>(0, QueriesFailed);
-		TotalTime += (EndTime - StartTime) / 1000000.0f;
-	}
-
-	float FailRate = static_cast<float>(QueriesFailed) / WindowSize;
-	OutAvgTimeInSeconds = TotalTime / WindowSize;
-
-	return FailRate;
-}
-
 #undef LOCTEXT_NAMESPACE
diff --git a/Engine/Source/Runtime/RHI/Public/DynamicRHI.h b/Engine/Source/Runtime/RHI/Public/DynamicRHI.h
--- a/Engine/Source/Runtime/RHI/Public/DynamicRHI.h
+++ b/Engine/Source/Runtime/RHI/Public/DynamicRHI.h
@@ -150,6 +150,29 @@ struct FRayTracingSceneInitializer
 	uint32 NumCallableShaderSlots = 0;
 };
 
+class FDynamicRHI;
+
+class RHI_API FDefaultRHIRenderQueryPool final : public FRHIRenderQueryPool
+{
+public:
+	FDefaultRHIRenderQueryPool(ERenderQueryType InQueryType, FDynamicRHI* InDynamicRHI, uint32 InNumQueries);
+
+	~FDefaultRHIRenderQueryPool() override 
+	{
+		checkf(AllocatedQueries == Queries.Num(), TEXT("Querypool deleted before all Queries have been released"));
+	}
+
+private:
+	virtual FRHIPooledRenderQuery AllocateQuery() override;
+	virtual void ReleaseQuery(TRefCountPtr<FRHIRenderQuery>&& Query) override;
+
+	FDynamicRHI* DynamicRHI = nullptr;
+	ERenderQueryType QueryType;
+	uint32 NumQueries = 0;
+	uint32 AllocatedQueries = 0;
+	TArray<TRefCountPtr<FRHIRenderQuery>> Queries;
+};
+
 /** The interface which is implemented by the dynamically bound RHI. */
 class RHI_API FDynamicRHI
 {
@@ -272,6 +295,16 @@ public:
 	{
 		return nullptr;
 	}
+	/**
+	* Creates a pool for querys like timers or occlusion queries.
+	* @param QueryType The ype of the queries provided by this pool like RQT_Occlusion or RQT_AbsoluteTime.
+	* @return the Querypool.
+	*/
+	// FlushType: Must be Thread-Safe.
+	virtual FRenderQueryPoolRHIRef RHICreateRenderQueryPool(ERenderQueryType QueryType, uint32 NumQueries = UINT32_MAX)
+	{
+		return new FDefaultRHIRenderQueryPool(QueryType, this, NumQueries);
+	}
 
 	/**
 	* Creates a compute fence.  Compute fences are named GPU fences which can be written to once before resetting.
@@ -1376,6 +1409,8 @@ FORCEINLINE class IRHICommandContextContainer* RHIGetCommandContextContainer(int
 	return GDynamicRHI->RHIGetCommandContextContainer(Index, Num, GPUMask);
 }
 
+RHI_API FRenderQueryPoolRHIRef RHICreateRenderQueryPool(ERenderQueryType QueryType, uint32 NumQueries = UINT32_MAX);
+
 #if RHI_RAYTRACING
 
 FORCEINLINE FRayTracingGeometryRHIRef RHICreateRayTracingGeometry(const FRayTracingGeometryInitializer& Initializer)
diff --git a/Engine/Source/Runtime/RHI/Public/GPUProfiler.h b/Engine/Source/Runtime/RHI/Public/GPUProfiler.h
--- a/Engine/Source/Runtime/RHI/Public/GPUProfiler.h
+++ b/Engine/Source/Runtime/RHI/Public/GPUProfiler.h
@@ -276,78 +276,4 @@ struct RHI_API FGPUProfiler
 
 	virtual void PushEvent(const TCHAR* Name, FColor Color);
 	virtual void PopEvent();
-};
-
-
-
-/**
- * Simple moving window averaged GPU timer; create an instance, call Begin() and End() around the block to time, 
- * then call GetElapsedAverage to get the averaged timings; BufferSize determines the number of queries in the window,
- * FramesBehind determines how long we wait to grab query results, so we don't have to block on them, so the effective
- * window size is BufferSize - FramesBehind
- * The timer keeps track of failed queries; GetElapsedAverage returns a float [0;1] that indicates fail rate: 0 means
- * no queries have failed, 1 means all queries within the window have failed. If more than 0.1, it's a good indication
- * that FramesBehind has to be increased.
- */
-class FWindowedGPUTimer
-{
-public:
-	explicit FWindowedGPUTimer(FRHICommandListImmediate& RHICmdList)
-	{
-		PrivateInit(10,2,RHICmdList);
-	}
-
-	FWindowedGPUTimer(uint32 BufferSize, uint32 FramesBehind, FRHICommandListImmediate& RHICmdList)
-	{
-		PrivateInit(BufferSize,FramesBehind,RHICmdList);
-	}
-
-	void Begin(FRHICommandListImmediate& RHICmdList)
-	{
-		RotateQueryBuffer(StartQueries);
-		RHICmdList.EndRenderQuery(StartQueries[0]);
-	}
-
-
-	void End(FRHICommandListImmediate& RHICmdList)
-	{
-		RotateQueryBuffer(EndQueries);
-		RHICmdList.EndRenderQuery(EndQueries[0]);
-		QueriesFinished++;
-	}
-
-	void RotateQueryBuffer(TArray<FRenderQueryRHIRef> &QueryArray)
-	{
-		FRenderQueryRHIRef LastQuery = QueryArray.Last();
-
-		for (int32 i = QueryArray.Num() - 1; i > 0; --i)
-		{
-			QueryArray[i] = QueryArray[i - 1];
-		}
-		QueryArray[0] = LastQuery;
-	}
-
-
-	float GetElapsedAverage(FRHICommandListImmediate& RHICmdList, float &OutAvgTimeInSeconds);
-
-private:
-	void PrivateInit(uint32 BufferSize, uint32 FramesBehind, FRHICommandListImmediate& RHICmdList)
-	{
-		QueriesFailed = 0;
-		QueriesFinished = 0;
-		StartQueries.AddZeroed(BufferSize);
-		EndQueries.AddZeroed(BufferSize);
-		for (uint32 i = 0; i < BufferSize; i++)
-		{
-			StartQueries[i] = RHICmdList.CreateRenderQuery(RQT_AbsoluteTime);;
-			EndQueries[i] = RHICmdList.CreateRenderQuery(RQT_AbsoluteTime);;
-		}
-		WindowSize = BufferSize - FramesBehind;
-	}
-
-	int32 QueriesFailed;
-	uint32 WindowSize;
-	int32 QueriesFinished;
-	TArray<FRenderQueryRHIRef> StartQueries;
-	TArray<FRenderQueryRHIRef> EndQueries;
-};
+};
\ No newline at end of file
diff --git a/Engine/Source/Runtime/RHI/Public/RHICommandList.h b/Engine/Source/Runtime/RHI/Public/RHICommandList.h
--- a/Engine/Source/Runtime/RHI/Public/RHICommandList.h
+++ b/Engine/Source/Runtime/RHI/Public/RHICommandList.h
@@ -4058,12 +4058,14 @@ public:
 		GDynamicRHI->RHIRead3DSurfaceFloatData(Texture,Rect,ZMinMax,OutData);
 	}
 	
+	UE_DEPRECATED(4.23, "CreateRenderQuery API is deprecated; use RHICreateRenderQueryPool and suballocate queries there")
 	FORCEINLINE FRenderQueryRHIRef CreateRenderQuery(ERenderQueryType QueryType)
 	{
 		FScopedRHIThreadStaller StallRHIThread(*this);
 		return GDynamicRHI->RHICreateRenderQuery(QueryType);
 	}
 
+	UE_DEPRECATED(4.23, "CreateRenderQuery API is deprecated; use RHICreateRenderQueryPool and suballocate queries there")
 	FORCEINLINE FRenderQueryRHIRef CreateRenderQuery_RenderThread(ERenderQueryType QueryType)
 	{
 		return GDynamicRHI->RHICreateRenderQuery_RenderThread(*this, QueryType);
@@ -4811,9 +4813,12 @@ FORCEINLINE void RHIUnlockTextureCubeFace(FTextureCubeRHIParamRef Texture, uint3
 	 FRHICommandListExecutor::GetImmediateCommandList().UnlockTextureCubeFace(Texture, FaceIndex, ArrayIndex, MipIndex, bLockWithinMiptail);
 }
 
+UE_DEPRECATED(4.23, "CreateRenderQuery API is deprecated; use RHICreateRenderQueryPool and suballocate queries there")
 FORCEINLINE FRenderQueryRHIRef RHICreateRenderQuery(ERenderQueryType QueryType)
 {
+	PRAGMA_DISABLE_DEPRECATION_WARNINGS
 	return FRHICommandListExecutor::GetImmediateCommandList().CreateRenderQuery_RenderThread(QueryType);
+	PRAGMA_ENABLE_DEPRECATION_WARNINGS
 }
 
 FORCEINLINE void RHIAcquireTransientResource(FTextureRHIParamRef Resource)
diff --git a/Engine/Source/Runtime/RHI/Public/RHIResources.h b/Engine/Source/Runtime/RHI/Public/RHIResources.h
--- a/Engine/Source/Runtime/RHI/Public/RHIResources.h
+++ b/Engine/Source/Runtime/RHI/Public/RHIResources.h
@@ -943,6 +943,67 @@ private:
 
 class FRHIRenderQuery : public FRHIResource {};
 
+class FRHIRenderQueryPool;
+class RHI_API FRHIPooledRenderQuery
+{
+	TRefCountPtr<FRHIRenderQuery> Query;
+	FRHIRenderQueryPool* QueryPool = nullptr;
+
+public:
+	FRHIPooledRenderQuery() = default;
+	FRHIPooledRenderQuery(FRHIRenderQueryPool* InQueryPool, TRefCountPtr<FRHIRenderQuery>&& InQuery);
+	~FRHIPooledRenderQuery();
+
+	FRHIPooledRenderQuery(const FRHIPooledRenderQuery&) = delete;
+	FRHIPooledRenderQuery& operator=(const FRHIPooledRenderQuery&) = delete;
+	FRHIPooledRenderQuery(FRHIPooledRenderQuery&&) = default;
+	FRHIPooledRenderQuery& operator=(FRHIPooledRenderQuery&&) = default;
+
+	FRHIRenderQuery* GetQuery() const
+	{
+		return Query;
+	}
+
+	TRefCountPtr<FRHIRenderQuery> GetQueryRef() const
+	{
+		return Query;
+	}
+
+	void ReleaseQuery();
+};
+
+class RHI_API FRHIRenderQueryPool : public FRHIResource
+{
+public:
+	virtual ~FRHIRenderQueryPool() {};
+	virtual FRHIPooledRenderQuery AllocateQuery() = 0;
+
+private:
+	friend class FRHIPooledRenderQuery;
+	virtual void ReleaseQuery(TRefCountPtr<FRHIRenderQuery>&& Query) = 0;
+};
+
+inline FRHIPooledRenderQuery::FRHIPooledRenderQuery(FRHIRenderQueryPool* InQueryPool, TRefCountPtr<FRHIRenderQuery>&& InQuery) 
+	: Query(MoveTemp(InQuery))
+	, QueryPool(InQueryPool)
+{
+}
+
+inline void FRHIPooledRenderQuery::ReleaseQuery()
+{
+	if (QueryPool && Query.IsValid())
+	{
+		QueryPool->ReleaseQuery(MoveTemp(Query));
+		QueryPool = nullptr;
+	}
+	check(!Query.IsValid());
+}
+
+inline FRHIPooledRenderQuery::~FRHIPooledRenderQuery()
+{
+	ReleaseQuery();
+}
+
 class FRHIComputeFence : public FRHIResource
 {
 public:
@@ -1116,6 +1177,9 @@ typedef TRefCountPtr<FRHITextureReference> FTextureReferenceRHIRef;
 typedef FRHIRenderQuery*              FRenderQueryRHIParamRef;
 typedef TRefCountPtr<FRHIRenderQuery> FRenderQueryRHIRef;
 
+typedef FRHIRenderQueryPool*              FRenderQueryPoolRHIParamRef;
+typedef TRefCountPtr<FRHIRenderQueryPool> FRenderQueryPoolRHIRef;
+
 typedef FRHIGPUFence*				FGPUFenceRHIParamRef;
 typedef TRefCountPtr<FRHIGPUFence>	FGPUFenceRHIRef;
 
diff --git a/Engine/Source/Runtime/Renderer/Private/GPUBenchmark.cpp b/Engine/Source/Runtime/Renderer/Private/GPUBenchmark.cpp
--- a/Engine/Source/Runtime/Renderer/Private/GPUBenchmark.cpp
+++ b/Engine/Source/Runtime/Renderer/Private/GPUBenchmark.cpp
@@ -413,8 +413,6 @@ void RendererGPUBenchmark(FRHICommandListImmediate& RHICmdList, FSynthBenchmarkR
 {
 	check(IsInRenderingThread());
 
-	FRenderQueryPool TimerQueryPool(RQT_AbsoluteTime);
-
 	bool bValidGPUTimer = (FGPUTiming::GetTimingFrequency() / (1000 * 1000)) != 0;
 
 	if(!bValidGPUTimer)
@@ -498,18 +496,7 @@ void RendererGPUBenchmark(FRHICommandListImmediate& RHICmdList, FSynthBenchmarkR
 		// 0 / 1
 		uint32 DestRTIndex = 0;
 
-		const uint32 TimerSampleCount = IterationCount * MethodCount + 1;
-
-		static FRenderQueryRHIRef TimerQueries[TimerSampleCount];
-		static float LocalWorkScale[IterationCount];
-
-		for(uint32  i = 0; i < TimerSampleCount; ++i)
-		{
-			TimerQueries[i] = TimerQueryPool.AllocateQuery();
-		}
-
-		const bool bSupportsTimerQueries = (TimerQueries[0] != NULL);
-		if(!bSupportsTimerQueries)
+		if(!GSupportsTimestampRenderQueries)
 		{
 #if !PLATFORM_MAC
 			UE_LOG(LogSynthBenchmark, Warning, TEXT("GPU driver does not support timer queries."));
@@ -545,6 +532,16 @@ void RendererGPUBenchmark(FRHICommandListImmediate& RHICmdList, FSynthBenchmarkR
 			return;
 		}
 
+		const uint32 TimerSampleCount = IterationCount * MethodCount + 1;
+		FRenderQueryPoolRHIRef TimerQueryPool = RHICreateRenderQueryPool(RQT_AbsoluteTime, TimerSampleCount * 2);
+		FRHIPooledRenderQuery TimerQueries[TimerSampleCount];
+		float LocalWorkScale[IterationCount];
+
+		for (uint32 i = 0; i < TimerSampleCount; ++i)
+		{
+			TimerQueries[i] = TimerQueryPool->AllocateQuery();
+		}
+
 		IssueScalableLongGPUTask(RHICmdList, -1);
 		RHICmdList.ImmediateFlush(EImmediateFlushType::FlushRHIThread);
 
@@ -559,7 +556,7 @@ void RendererGPUBenchmark(FRHICommandListImmediate& RHICmdList, FSynthBenchmarkR
 			TimingSeries[MethodIterator].Init(IterationCount);
 		}
 
-		RHICmdList.EndRenderQuery(TimerQueries[0]);
+		RHICmdList.EndRenderQuery(TimerQueries[0].GetQuery());
 
 		// multiple iterations to see how trust able the values are
 		for(uint32 Iteration = 0; Iteration < IterationCount; ++Iteration)
@@ -610,7 +607,7 @@ void RendererGPUBenchmark(FRHICommandListImmediate& RHICmdList, FSynthBenchmarkR
 					RHIUnmapStagingSurface(RTItems[2]->GetRenderTargetItem().ShaderResourceTexture);
 				}*/
 
-				RHICmdList.EndRenderQuery(TimerQueries[QueryIndex]);
+				RHICmdList.EndRenderQuery(TimerQueries[QueryIndex].GetQuery());
 
 				// ping pong
 				DestRTIndex = 1 - DestRTIndex;
@@ -621,8 +618,7 @@ void RendererGPUBenchmark(FRHICommandListImmediate& RHICmdList, FSynthBenchmarkR
 			uint64 OldAbsTime = 0;
 			// flushes the RHI thread to make sure all RHICmdList.EndRenderQuery() commands got executed.
 			RHICmdList.ImmediateFlush(EImmediateFlushType::FlushRHIThread);
-			RHICmdList.GetRenderQueryResult(TimerQueries[0], OldAbsTime, true);
-			TimerQueryPool.ReleaseQuery(TimerQueries[0]);
+			RHICmdList.GetRenderQueryResult(TimerQueries[0].GetQuery(), OldAbsTime, true);
 
 			for(uint32 Iteration = 0; Iteration < IterationCount; ++Iteration)
 			{
@@ -633,8 +629,7 @@ void RendererGPUBenchmark(FRHICommandListImmediate& RHICmdList, FSynthBenchmarkR
 					uint32 QueryIndex = 1 + Iteration * MethodCount + MethodId;
 
 					uint64 AbsTime;
-					RHICmdList.GetRenderQueryResult(TimerQueries[QueryIndex], AbsTime, true);
-					TimerQueryPool.ReleaseQuery(TimerQueries[QueryIndex]);
+					RHICmdList.GetRenderQueryResult(TimerQueries[QueryIndex].GetQuery(), AbsTime, true);
 
 					uint64 RelTime = FMath::Max(AbsTime - OldAbsTime, 1ull);
 
@@ -667,7 +662,7 @@ void RendererGPUBenchmark(FRHICommandListImmediate& RHICmdList, FSynthBenchmarkR
 				}
 			}
 
-			if(bSupportsTimerQueries)
+			if(GSupportsTimestampRenderQueries)
 			{
 				for(uint32 MethodId = 0; MethodId < MethodCount; ++MethodId)
 				{
diff --git a/Engine/Source/Runtime/Renderer/Private/PlanarReflectionRendering.cpp b/Engine/Source/Runtime/Renderer/Private/PlanarReflectionRendering.cpp
--- a/Engine/Source/Runtime/Renderer/Private/PlanarReflectionRendering.cpp
+++ b/Engine/Source/Runtime/Renderer/Private/PlanarReflectionRendering.cpp
@@ -334,7 +334,7 @@ static void UpdatePlanarReflectionContents_RenderThread(
 				const int32 NumBufferedFrames = FOcclusionQueryHelpers::GetNumBufferedFrames(SceneRenderer->FeatureLevel) + 1;
 				// +1 to frame counter because we are operating before the main view's InitViews, which is where OcclusionFrameCounter is incremented
 				uint32 OcclusionFrameCounter = ViewState->OcclusionFrameCounter + 1;
-				FRenderQueryRHIParamRef PastQuery = OcclusionHistory.GetPastQuery(OcclusionFrameCounter, NumBufferedFrames);
+				FRHIRenderQuery* PastQuery = OcclusionHistory.GetPastQuery(OcclusionFrameCounter, NumBufferedFrames);
 
 				if (PastQuery)
 				{
diff --git a/Engine/Source/Runtime/Renderer/Private/RendererScene.cpp b/Engine/Source/Runtime/Renderer/Private/RendererScene.cpp
--- a/Engine/Source/Runtime/Renderer/Private/RendererScene.cpp
+++ b/Engine/Source/Runtime/Renderer/Private/RendererScene.cpp
@@ -125,7 +125,10 @@ struct FSpeedTreeWindComputation
 
 /** Default constructor. */
 FSceneViewState::FSceneViewState()
-	: OcclusionQueryPool(RQT_Occlusion)
+	: OcclusionQueryPool(RHICreateRenderQueryPool(RQT_Occlusion))
+	, TimerQueryPool(RHICreateRenderQueryPool(RQT_AbsoluteTime, FLatentGPUTimer::NumBufferedFrames * 2 * 2 * 2))
+	, TranslucencyTimer(TimerQueryPool)
+	, SeparateTranslucencyTimer(TimerQueryPool)
 {
 	UniqueID = FSceneViewState_UniqueID.Increment();
 	OcclusionFrameCounter = 0;
@@ -240,6 +243,7 @@ void DestroyRWBuffer(FRWBuffer* RWBuffer)
 FSceneViewState::~FSceneViewState()
 {
 	CachedVisibilityChunk = NULL;
+	ShadowOcclusionQueryMaps.Reset();
 
 	for (int32 CascadeIndex = 0; CascadeIndex < ARRAY_COUNT(TranslucencyLightingCacheAllocations); CascadeIndex++)
 	{
@@ -3605,11 +3609,12 @@ FSceneViewStateInterface* FRendererModule::AllocateViewState()
 
 //////////////////////////////////////////////////////////////////////////
 
-FLatentGPUTimer::FLatentGPUTimer(int32 InAvgSamples)
-	: AvgSamples(InAvgSamples)
-	, TotalTime(0.0f)
-	, SampleIndex(0)
-	, QueryIndex(0)
+FLatentGPUTimer::FLatentGPUTimer(FRenderQueryPoolRHIRef InTimerQueryPool, int32 InAvgSamples)
+: TimerQueryPool(InTimerQueryPool)
+, AvgSamples(InAvgSamples)
+, TotalTime(0.0f)
+, SampleIndex(0)
+, QueryIndex(0)
 {
 	TimeSamples.AddZeroed(AvgSamples);
 }
@@ -3623,7 +3628,7 @@ bool FLatentGPUTimer::Tick(FRHICommandListImmediate& RHICmdList)
 
 	QueryIndex = (QueryIndex + 1) % NumBufferedFrames;
 
-	if (StartQueries[QueryIndex] && EndQueries[QueryIndex])
+	if (StartQueries[QueryIndex].GetQuery() && EndQueries[QueryIndex].GetQuery())
 	{
 		if (IsRunningRHIInSeparateThread())
 		{
@@ -3644,8 +3649,8 @@ bool FLatentGPUTimer::Tick(FRHICommandListImmediate& RHICmdList)
 			// Block on the GPU until we have the timestamp query results, if necessary
 			// Stat disabled since we buffer 2 frames minimum, it won't actually block
 			//SCOPE_CYCLE_COUNTER(STAT_TranslucencyTimestampQuery_Wait);
-			bStartSuccess = RHICmdList.GetRenderQueryResult(StartQueries[QueryIndex], StartMicroseconds, true);
-			bEndSuccess = RHICmdList.GetRenderQueryResult(EndQueries[QueryIndex], EndMicroseconds, true);
+			bStartSuccess = RHICmdList.GetRenderQueryResult(StartQueries[QueryIndex].GetQuery(), StartMicroseconds, true);
+			bEndSuccess = RHICmdList.GetRenderQueryResult(EndQueries[QueryIndex].GetQuery(), EndMicroseconds, true);
 		}
 
 		TotalTime -= TimeSamples[SampleIndex];
@@ -3672,12 +3677,12 @@ void FLatentGPUTimer::Begin(FRHICommandListImmediate& RHICmdList)
 		return;
 	}
 	
-	if (!StartQueries[QueryIndex])
+	if (!StartQueries[QueryIndex].GetQuery())
 	{		
-		StartQueries[QueryIndex] = RHICmdList.CreateRenderQuery(RQT_AbsoluteTime);
+		StartQueries[QueryIndex] = TimerQueryPool->AllocateQuery();
 	}
 
-	RHICmdList.EndRenderQuery(StartQueries[QueryIndex]);
+	RHICmdList.EndRenderQuery(StartQueries[QueryIndex].GetQuery());
 }
 
 void FLatentGPUTimer::End(FRHICommandListImmediate& RHICmdList)
@@ -3687,12 +3692,12 @@ void FLatentGPUTimer::End(FRHICommandListImmediate& RHICmdList)
 		return;
 	}
 	
-	if (!EndQueries[QueryIndex])
+	if (!EndQueries[QueryIndex].GetQuery())
 	{
-		EndQueries[QueryIndex] = RHICmdList.CreateRenderQuery(RQT_AbsoluteTime);
+		EndQueries[QueryIndex] = TimerQueryPool->AllocateQuery();
 	}
 
-	RHICmdList.EndRenderQuery(EndQueries[QueryIndex]);
+	RHICmdList.EndRenderQuery(EndQueries[QueryIndex].GetQuery());
 	// Hint to the RHI to submit commands up to this point to the GPU if possible.  Can help avoid CPU stalls next frame waiting
 	// for these query results on some platforms.
 	RHICmdList.SubmitCommandsHint();
@@ -3714,8 +3719,8 @@ void FLatentGPUTimer::Release()
 {
 	for (int32 i = 0; i < NumBufferedFrames; ++i)
 	{
-		StartQueries[i].SafeRelease();
-		EndQueries[i].SafeRelease();
+		StartQueries[i].ReleaseQuery();
+		EndQueries[i].ReleaseQuery();
 	}
 }
 
diff --git a/Engine/Source/Runtime/Renderer/Private/SceneOcclusion.cpp b/Engine/Source/Runtime/Renderer/Private/SceneOcclusion.cpp
--- a/Engine/Source/Runtime/Renderer/Private/SceneOcclusion.cpp
+++ b/Engine/Source/Runtime/Renderer/Private/SceneOcclusion.cpp
@@ -243,14 +243,14 @@ bool FSceneViewState::IsShadowOccluded(FRHICommandListImmediate& RHICmdList, FSc
 	// Get the oldest occlusion query	
 	const uint32 QueryIndex = FOcclusionQueryHelpers::GetQueryLookupIndex(PendingPrevFrameNumber, NumBufferedFrames);
 	const FSceneViewState::ShadowKeyOcclusionQueryMap& ShadowOcclusionQueryMap = ShadowOcclusionQueryMaps[QueryIndex];	
-	const FRenderQueryRHIRef* Query = ShadowOcclusionQueryMap.Find(ShadowKey);
+	const FRHIPooledRenderQuery* Query = ShadowOcclusionQueryMap.Find(ShadowKey);
 
 	// Read the occlusion query results.
 	uint64 NumSamples = 0;
 	// Only block on the query if not running SLI
 	const bool bWaitOnQuery = GNumAlternateFrameRenderingGroups == 1;
 
-	if (Query && RHICmdList.GetRenderQueryResult(*Query, NumSamples, bWaitOnQuery))
+	if (Query && RHICmdList.GetRenderQueryResult(Query->GetQuery(), NumSamples, bWaitOnQuery))
 	{
 		// If the shadow's occlusion query didn't have any pixels visible the previous frame, it's occluded.
 		return NumSamples == 0;
@@ -280,19 +280,15 @@ void FSceneViewState::ConditionallyAllocateSceneSoftwareOcclusion(ERHIFeatureLev
 
 void FSceneViewState::Destroy()
 {
-	if ( IsInGameThread() )
+	FSceneViewState* self = this;
+	ENQUEUE_RENDER_COMMAND(FSceneViewState_Destroy)(
+	[self](FRHICommandListImmediate& RHICmdList)
 	{
 		// Release the occlusion query data.
-		BeginReleaseResource(this);
-
+		self->ReleaseResource();
 		// Defer deletion of the view state until the rendering thread is done with it.
-		BeginCleanup(this);
-	}
-	else
-	{
-		ReleaseResource();
-		delete this;
-	}
+		delete self;
+	});
 }
 
 SIZE_T FSceneViewState::GetSizeBytes() const
@@ -341,7 +337,7 @@ FOcclusionQueryBatcher::FOcclusionQueryBatcher(class FSceneViewState* ViewState,
 :	CurrentBatchOcclusionQuery(NULL)
 ,	MaxBatchedPrimitives(InMaxBatchedPrimitives)
 ,	NumBatchedPrimitives(0)
-,	OcclusionQueryPool(ViewState ? &ViewState->OcclusionQueryPool : NULL)
+,	OcclusionQueryPool(ViewState ? ViewState->OcclusionQueryPool : NULL)
 {}
 
 FOcclusionQueryBatcher::~FOcclusionQueryBatcher()
@@ -362,7 +358,7 @@ void FOcclusionQueryBatcher::Flush(FRHICommandList& RHICmdList)
 		for(int32 BatchIndex = 0, NumBatches = BatchOcclusionQueries.Num();BatchIndex < NumBatches;BatchIndex++)
 		{
 			FOcclusionBatch& Batch = BatchOcclusionQueries[BatchIndex];
-			FRenderQueryRHIParamRef BatchOcclusionQuery = Batch.Query;
+			FRHIRenderQuery* BatchOcclusionQuery = Batch.Query;
 			FVertexBufferRHIParamRef VertexBufferRHI = Batch.VertexAllocation.VertexBuffer->VertexBufferRHI;
 			uint32 VertexBufferOffset = Batch.VertexAllocation.VertexOffset;
 			const int32 NumPrimitivesThisBatch = (BatchIndex != (NumBatches-1)) ? MaxBatchedPrimitives : NumBatchedPrimitives;
@@ -388,14 +384,17 @@ void FOcclusionQueryBatcher::Flush(FRHICommandList& RHICmdList)
 	}
 }
 
-FRenderQueryRHIParamRef FOcclusionQueryBatcher::BatchPrimitive(const FVector& BoundsOrigin,const FVector& BoundsBoxExtent, FGlobalDynamicVertexBuffer& DynamicVertexBuffer)
+FRHIPooledRenderQuery FOcclusionQueryBatcher::BatchPrimitive(const FVector& BoundsOrigin,const FVector& BoundsBoxExtent, FGlobalDynamicVertexBuffer& DynamicVertexBuffer)
 {
+	FRHIPooledRenderQuery PooledQuery;
+
 	// Check if the current batch is full.
 	if(CurrentBatchOcclusionQuery == NULL || NumBatchedPrimitives >= MaxBatchedPrimitives)
 	{
 		check(OcclusionQueryPool);
 		CurrentBatchOcclusionQuery = new(BatchOcclusionQueries) FOcclusionBatch;
-		CurrentBatchOcclusionQuery->Query = OcclusionQueryPool->AllocateQuery();
+		PooledQuery = OcclusionQueryPool->AllocateQuery();
+		CurrentBatchOcclusionQuery->Query = PooledQuery.GetQueryRef();
 		CurrentBatchOcclusionQuery->VertexAllocation = DynamicVertexBuffer.Allocate(MaxBatchedPrimitives * 8 * sizeof(FVector));
 		check(CurrentBatchOcclusionQuery->VertexAllocation.IsValid());
 		NumBatchedPrimitives = 0;
@@ -419,7 +418,7 @@ FRenderQueryRHIParamRef FOcclusionQueryBatcher::BatchPrimitive(const FVector& Bo
 	CurrentBatchOcclusionQuery->VertexAllocation.Buffer = (uint8*)Vertices;
 	NumBatchedPrimitives++;
 
-	return CurrentBatchOcclusionQuery->Query;
+	return PooledQuery;
 }
 
 enum EShadowOcclusionQueryIntersectionMode
@@ -434,7 +433,7 @@ static bool AllocateProjectedShadowOcclusionQuery(
 	const FProjectedShadowInfo& ProjectedShadowInfo, 
 	int32 NumBufferedFrames, 
 	EShadowOcclusionQueryIntersectionMode IntersectionMode,
-	FRenderQueryRHIRef& ShadowOcclusionQuery)
+	FRHIRenderQuery*& ShadowOcclusionQuery)
 {
 	bool bIssueQuery = true;
 
@@ -469,21 +468,21 @@ static bool AllocateProjectedShadowOcclusionQuery(
 		FSceneViewState* ViewState = (FSceneViewState*)View.State;
 
 		// Allocate an occlusion query for the primitive from the occlusion query pool.
-		ShadowOcclusionQuery = ViewState->OcclusionQueryPool.AllocateQuery();
-
 		FSceneViewState::FProjectedShadowKey Key(ProjectedShadowInfo);
 		const uint32 QueryIndex = FOcclusionQueryHelpers::GetQueryIssueIndex(ViewState->PendingPrevFrameNumber, NumBufferedFrames);
 		FSceneViewState::ShadowKeyOcclusionQueryMap& ShadowOcclusionQueryMap = ViewState->ShadowOcclusionQueryMaps[QueryIndex];
 
 		checkSlow(ShadowOcclusionQueryMap.Find(Key) == NULL);
-		ShadowOcclusionQueryMap.Add(Key, ShadowOcclusionQuery);
+		FRHIPooledRenderQuery PooledShadowOcclusionQuery = ViewState->OcclusionQueryPool->AllocateQuery();
+		ShadowOcclusionQuery = PooledShadowOcclusionQuery.GetQuery();
+		ShadowOcclusionQueryMap.Add(Key, MoveTemp(PooledShadowOcclusionQuery));
 	}
 	
 	return bIssueQuery;
 }
 
 
-static void ExecutePointLightShadowOcclusionQuery(FRHICommandList& RHICmdList, FViewInfo& View, const FProjectedShadowInfo& ProjectedShadowInfo, FOcclusionQueryVS* VertexShader, FRenderQueryRHIRef ShadowOcclusionQuery)
+static void ExecutePointLightShadowOcclusionQuery(FRHICommandList& RHICmdList, FViewInfo& View, const FProjectedShadowInfo& ProjectedShadowInfo, FOcclusionQueryVS* VertexShader, FRHIRenderQuery* ShadowOcclusionQuery)
 {
 	FLightSceneProxy& LightProxy = *(ProjectedShadowInfo.GetLightSceneInfo().Proxy);
 	
@@ -532,7 +531,7 @@ static void PrepareDirectionalLightShadowOcclusionQuery(uint32& BaseVertexIndex,
 	BaseVertexIndex += 6;
 }
 
-static void ExecuteDirectionalLightShadowOcclusionQuery(FRHICommandList& RHICmdList, uint32& BaseVertexIndex, FRenderQueryRHIRef ShadowOcclusionQuery)
+static void ExecuteDirectionalLightShadowOcclusionQuery(FRHICommandList& RHICmdList, uint32& BaseVertexIndex, FRHIRenderQuery* ShadowOcclusionQuery)
 {
 	RHICmdList.BeginRenderQuery(ShadowOcclusionQuery);
 
@@ -572,7 +571,7 @@ static void PrepareProjectedShadowOcclusionQuery(uint32& BaseVertexIndex, FVecto
 	BaseVertexIndex += 8;
 }
 
-static void ExecuteProjectedShadowOcclusionQuery(FRHICommandList& RHICmdList, uint32& BaseVertexIndex, FRenderQueryRHIRef ShadowOcclusionQuery)
+static void ExecuteProjectedShadowOcclusionQuery(FRHICommandList& RHICmdList, uint32& BaseVertexIndex, FRHIRenderQuery* ShadowOcclusionQuery)
 {	
 	// Draw the primitive's bounding box, using the occlusion query.
 	RHICmdList.BeginRenderQuery(ShadowOcclusionQuery);
@@ -583,7 +582,7 @@ static void ExecuteProjectedShadowOcclusionQuery(FRHICommandList& RHICmdList, ui
 	RHICmdList.EndRenderQuery(ShadowOcclusionQuery);
 }
 
-static bool AllocatePlanarReflectionOcclusionQuery(const FViewInfo& View, const FPlanarReflectionSceneProxy* SceneProxy, int32 NumBufferedFrames, FRenderQueryRHIRef& OcclusionQuery)
+static bool AllocatePlanarReflectionOcclusionQuery(const FViewInfo& View, const FPlanarReflectionSceneProxy* SceneProxy, int32 NumBufferedFrames, FRHIRenderQuery*& OcclusionQuery)
 {
 	FSceneViewState* ViewState = (FSceneViewState*)View.State;
 	
@@ -618,13 +617,14 @@ static bool AllocatePlanarReflectionOcclusionQuery(const FViewInfo& View, const
 	if (bAllowBoundsTest)
 	{
 		// Allocate an occlusion query for the primitive from the occlusion query pool.
-		OcclusionQuery = ViewState->OcclusionQueryPool.AllocateQuery();
-		
-		OcclusionHistory.SetCurrentQuery(OcclusionFrameCounter, OcclusionQuery, NumBufferedFrames);
+		FRHIPooledRenderQuery PooledOcclusionQuery = ViewState->OcclusionQueryPool->AllocateQuery();
+		OcclusionQuery = PooledOcclusionQuery.GetQuery();
+
+		OcclusionHistory.SetCurrentQuery(OcclusionFrameCounter, MoveTemp(PooledOcclusionQuery), NumBufferedFrames);
 	}
 	else
 	{
-		OcclusionHistory.SetCurrentQuery(OcclusionFrameCounter, NULL, NumBufferedFrames);
+		OcclusionHistory.SetCurrentQuery(OcclusionFrameCounter, FRHIPooledRenderQuery(), NumBufferedFrames);
 	}
 	
 	return bAllowBoundsTest;
@@ -648,7 +648,7 @@ static void PreparePlanarReflectionOcclusionQuery(uint32& BaseVertexIndex, FVect
 	BaseVertexIndex += 8;
 }
 
-static void ExecutePlanarReflectionOcclusionQuery(FRHICommandList& RHICmdList, uint32& BaseVertexIndex, FRenderQueryRHIRef OcclusionQuery)
+static void ExecutePlanarReflectionOcclusionQuery(FRHICommandList& RHICmdList, uint32& BaseVertexIndex, FRHIRenderQuery* OcclusionQuery)
 {
 	// Draw the primitive's bounding box, using the occlusion query.
 	RHICmdList.BeginRenderQuery(OcclusionQuery);
@@ -1264,10 +1264,10 @@ struct FViewOcclusionQueries
 	TArray<FProjectedShadowInfo const*> ShadowQuerieInfos;
 	TArray<FPlanarReflectionSceneProxy const*> ReflectionQuerieInfos;
 
-	TArray<FRenderQueryRHIRef> PointLightQueries;
-	TArray<FRenderQueryRHIRef> CSMQueries;
-	TArray<FRenderQueryRHIRef> ShadowQueries;
-	TArray<FRenderQueryRHIRef> ReflectionQueries;
+	TArray<FRHIRenderQuery*> PointLightQueries;
+	TArray<FRHIRenderQuery*> CSMQueries;
+	TArray<FRHIRenderQuery*> ShadowQueries;
+	TArray<FRHIRenderQuery*> ReflectionQueries;
 };
 
 void FSceneRenderer::BeginOcclusionTests(FRHICommandListImmediate& RHICmdList, bool bRenderQueries)
@@ -1303,12 +1303,6 @@ void FSceneRenderer::BeginOcclusionTests(FRHICommandListImmediate& RHICmdList, b
 				ViewState->TrimOcclusionHistory(ViewFamily.CurrentRealTime, ViewFamily.CurrentRealTime - GEngine->PrimitiveProbablyVisibleTime, ViewFamily.CurrentRealTime, ViewState->OcclusionFrameCounter);
 
 				// Give back all these occlusion queries to the pool.
-				for (TMap<FSceneViewState::FProjectedShadowKey, FRenderQueryRHIRef>::TIterator QueryIt(ShadowOcclusionQueryMap); QueryIt; ++QueryIt)
-				{
-					//FRenderQueryRHIParamRef Query = QueryIt.Value();
-					//check( Query.GetRefCount() == 1 );
-					ViewState->OcclusionQueryPool.ReleaseQuery(QueryIt.Value());
-				}
 				ShadowOcclusionQueryMap.Reset();
 				
 				if (FeatureLevel > ERHIFeatureLevel::ES3_1)
@@ -1336,7 +1330,7 @@ void FSceneRenderer::BeginOcclusionTests(FRHICommandListImmediate& RHICmdList, b
 
 							if (ProjectedShadowInfo.bOnePassPointLightShadow)
 							{
-								FRenderQueryRHIRef ShadowOcclusionQuery;
+								FRHIRenderQuery* ShadowOcclusionQuery;
 								if (AllocateProjectedShadowOcclusionQuery(View, ProjectedShadowInfo, NumBufferedFrames, SOQ_LightInfluenceSphere, ShadowOcclusionQuery))
 								{
 									ViewQuery.PointLightQuerieInfos.Add(&ProjectedShadowInfo);
@@ -1350,7 +1344,7 @@ void FSceneRenderer::BeginOcclusionTests(FRHICommandListImmediate& RHICmdList, b
 								// Don't query the first cascade, it is always visible
 								if (GOcclusionCullCascadedShadowMaps && ProjectedShadowInfo.CascadeSettings.ShadowSplitIndex > 0)
 								{
-									FRenderQueryRHIRef ShadowOcclusionQuery;
+									FRHIRenderQuery* ShadowOcclusionQuery;
 									if (AllocateProjectedShadowOcclusionQuery(View, ProjectedShadowInfo, NumBufferedFrames, SOQ_None, ShadowOcclusionQuery))
 									{
 										ViewQuery.CSMQuerieInfos.Add(&ProjectedShadowInfo);
@@ -1366,7 +1360,7 @@ void FSceneRenderer::BeginOcclusionTests(FRHICommandListImmediate& RHICmdList, b
 							// Don't query if any subjects are visible because the shadow frustum will be definitely unoccluded
 								&& !ProjectedShadowInfo.SubjectsVisible(View))
 							{
-								FRenderQueryRHIRef ShadowOcclusionQuery;
+								FRHIRenderQuery* ShadowOcclusionQuery;
 								if (AllocateProjectedShadowOcclusionQuery(View, ProjectedShadowInfo, NumBufferedFrames, SOQ_NearPlaneVsShadowFrustum, ShadowOcclusionQuery))
 								{
 									ViewQuery.ShadowQuerieInfos.Add(&ProjectedShadowInfo);
@@ -1381,7 +1375,7 @@ void FSceneRenderer::BeginOcclusionTests(FRHICommandListImmediate& RHICmdList, b
 						for (int32 ShadowIndex = 0; ShadowIndex < VisibleLightInfo.OccludedPerObjectShadows.Num(); ShadowIndex++)
 						{
 							const FProjectedShadowInfo& ProjectedShadowInfo = *VisibleLightInfo.OccludedPerObjectShadows[ShadowIndex];
-							FRenderQueryRHIRef ShadowOcclusionQuery;
+							FRHIRenderQuery* ShadowOcclusionQuery;
 							if (AllocateProjectedShadowOcclusionQuery(View, ProjectedShadowInfo, NumBufferedFrames, SOQ_NearPlaneVsShadowFrustum, ShadowOcclusionQuery))
 							{
 								ViewQuery.ShadowQuerieInfos.Add(&ProjectedShadowInfo);
@@ -1404,7 +1398,7 @@ void FSceneRenderer::BeginOcclusionTests(FRHICommandListImmediate& RHICmdList, b
 					for (int32 ReflectionIndex = 0; ReflectionIndex < Scene->PlanarReflections.Num(); ReflectionIndex++)
 					{
 						FPlanarReflectionSceneProxy* SceneProxy = Scene->PlanarReflections[ReflectionIndex];
-						FRenderQueryRHIRef ShadowOcclusionQuery;
+						FRHIRenderQuery* ShadowOcclusionQuery;
 						if (AllocatePlanarReflectionOcclusionQuery(View, SceneProxy, NumReflectionBufferedFrames, ShadowOcclusionQuery))
 						{
 							ViewQuery.ReflectionQuerieInfos.Add(SceneProxy);
@@ -1562,13 +1556,13 @@ void FSceneRenderer::BeginOcclusionTests(FRHICommandListImmediate& RHICmdList, b
 						RHICmdList.SetStreamSource(0, VertexBufferRHI, 0);
 						BaseVertexOffset = 0;
 
-						for (FRenderQueryRHIRef const& Query : ViewQuery.CSMQueries)
+						for (FRHIRenderQuery* const& Query : ViewQuery.CSMQueries)
 						{
 							ExecuteDirectionalLightShadowOcclusionQuery(RHICmdList, BaseVertexOffset, Query);
 							checkSlow(BaseVertexOffset <= NumVertices);
 						}
 
-						for (FRenderQueryRHIRef const& Query : ViewQuery.ShadowQueries)
+						for (FRHIRenderQuery* const& Query : ViewQuery.ShadowQueries)
 						{
 							ExecuteProjectedShadowOcclusionQuery(RHICmdList, BaseVertexOffset, Query);
 							checkSlow(BaseVertexOffset <= NumVertices);
@@ -1578,7 +1572,7 @@ void FSceneRenderer::BeginOcclusionTests(FRHICommandListImmediate& RHICmdList, b
 					if (FeatureLevel > ERHIFeatureLevel::ES3_1)					
 					{
 						SCOPED_DRAW_EVENT(RHICmdList, PlanarReflectionQueries);
-						for (FRenderQueryRHIRef const& Query : ViewQuery.ReflectionQueries)
+						for (FRHIRenderQuery* const& Query : ViewQuery.ReflectionQueries)
 						{
 							ExecutePlanarReflectionOcclusionQuery(RHICmdList, BaseVertexOffset, Query);
 							check(BaseVertexOffset <= NumVertices);
diff --git a/Engine/Source/Runtime/Renderer/Private/ScenePrivate.h b/Engine/Source/Runtime/Renderer/Private/ScenePrivate.h
--- a/Engine/Source/Runtime/Renderer/Private/ScenePrivate.h
+++ b/Engine/Source/Runtime/Renderer/Private/ScenePrivate.h
@@ -86,7 +86,7 @@ public:
 	FPrimitiveComponentId PrimitiveId;
 
 	/** The occlusion query which contains the primitive's pending occlusion results. */
-	FRenderQueryRHIRef PendingOcclusionQuery[FOcclusionQueryHelpers::MaxBufferedOcclusionFrames];
+	FRHIPooledRenderQuery PendingOcclusionQuery[FOcclusionQueryHelpers::MaxBufferedOcclusionFrames];
 	uint32 PendingOcclusionQueryFrames[FOcclusionQueryHelpers::MaxBufferedOcclusionFrames]; 
 
 	uint32 LastTestFrameNumber;
@@ -141,7 +141,7 @@ private:
 		{
 			const uint32 ThisFrameNumber = PendingOcclusionQueryFrames[Index];
 			const int32 LaggedFrames = FrameNumber - ThisFrameNumber;
-			if (PendingOcclusionQuery[Index].IsValid() && LaggedFrames <= LagTolerance && ThisFrameNumber < OldestFrame)
+			if (PendingOcclusionQuery[Index].GetQuery() && LaggedFrames <= LagTolerance && ThisFrameNumber < OldestFrame)
 			{
 				OldestFrame = ThisFrameNumber;
 				OldestQueryIndex = Index;
@@ -202,9 +202,9 @@ public:
 			{
 				uint32 TestFrame = FrameNumber - (DeltaFrame - 1);
 				const uint32 QueryIndex = FOcclusionQueryHelpers::GetQueryLookupIndex(TestFrame, NumBufferedFrames);
-				if (PendingOcclusionQuery[QueryIndex].GetReference() && PendingOcclusionQueryFrames[QueryIndex] != TestFrame)
+				if (PendingOcclusionQuery[QueryIndex].GetQuery() && PendingOcclusionQueryFrames[QueryIndex] != TestFrame)
 				{
-					Pool.ReleaseQuery(PendingOcclusionQuery[QueryIndex]);
+					PendingOcclusionQuery[QueryIndex].ReleaseQuery();
 				}
 			}
 		}
@@ -213,42 +213,39 @@ public:
 	FORCEINLINE void ReleaseQuery(TOcclusionQueryPool& Pool, uint32 FrameNumber, int32 NumBufferedFrames)
 	{
 		const uint32 QueryIndex = FOcclusionQueryHelpers::GetQueryLookupIndex(FrameNumber, NumBufferedFrames);
-		if (PendingOcclusionQuery[QueryIndex].IsValid())
-		{
-			Pool.ReleaseQuery(PendingOcclusionQuery[QueryIndex]);
-		}
+		PendingOcclusionQuery[QueryIndex].ReleaseQuery();
 	}
 
-	FORCEINLINE FRenderQueryRHIParamRef GetQueryForEviction(uint32 FrameNumber, int32 NumBufferedFrames) const
+	FORCEINLINE FRHIRenderQuery* GetQueryForEviction(uint32 FrameNumber, int32 NumBufferedFrames) const
 	{
 		const uint32 QueryIndex = FOcclusionQueryHelpers::GetQueryLookupIndex(FrameNumber, NumBufferedFrames);
-		if (PendingOcclusionQuery[QueryIndex].IsValid())
+		if (PendingOcclusionQuery[QueryIndex].GetQuery())
 		{
-			return PendingOcclusionQuery[QueryIndex].GetReference();
+			return PendingOcclusionQuery[QueryIndex].GetQuery();
 		}
 		return nullptr;
 	}
 
 
-	FORCEINLINE FRenderQueryRHIParamRef GetQueryForReading(uint32 FrameNumber, int32 NumBufferedFrames, int32 LagTolerance, bool& bOutGrouped) const
+	FORCEINLINE FRHIRenderQuery* GetQueryForReading(uint32 FrameNumber, int32 NumBufferedFrames, int32 LagTolerance, bool& bOutGrouped) const
 	{
 		const int32 OldestQueryIndex = bNeedsScanOnRead ? ScanOldestNonStaleQueryIndex(FrameNumber, NumBufferedFrames, LagTolerance)
 														: FOcclusionQueryHelpers::GetQueryLookupIndex(FrameNumber, NumBufferedFrames);
 		const int32 LaggedFrames = FrameNumber - PendingOcclusionQueryFrames[OldestQueryIndex];
-		if (OldestQueryIndex == -1 || !PendingOcclusionQuery[OldestQueryIndex].IsValid() || LaggedFrames > LagTolerance)
+		if (OldestQueryIndex == -1 || !PendingOcclusionQuery[OldestQueryIndex].GetQuery() || LaggedFrames > LagTolerance)
 		{
 			bOutGrouped = false;
 			return nullptr;
 		}
 		bOutGrouped = bGroupedQuery[OldestQueryIndex];
-		return PendingOcclusionQuery[OldestQueryIndex];
+		return PendingOcclusionQuery[OldestQueryIndex].GetQuery();
 	}
 
-	FORCEINLINE void SetCurrentQuery(uint32 FrameNumber, FRenderQueryRHIParamRef NewQuery, int32 NumBufferedFrames, bool bGrouped, bool bNeedsScan)
+	FORCEINLINE void SetCurrentQuery(uint32 FrameNumber, FRHIPooledRenderQuery&& NewQuery, int32 NumBufferedFrames, bool bGrouped, bool bNeedsScan)
 	{
 		// Get the current occlusion query
 		const uint32 QueryIndex = FOcclusionQueryHelpers::GetQueryIssueIndex(FrameNumber, NumBufferedFrames);
-		PendingOcclusionQuery[QueryIndex] = NewQuery;
+		PendingOcclusionQuery[QueryIndex] = MoveTemp(NewQuery);
 		PendingOcclusionQueryFrames[QueryIndex] = FrameNumber;
 		bGroupedQuery[QueryIndex] = bGrouped;
 
@@ -312,7 +309,7 @@ struct FPrimitiveOcclusionHistoryKeyFuncs : BaseKeyFuncs<FPrimitiveOcclusionHist
 
 class FIndividualOcclusionHistory
 {
-	FRenderQueryRHIRef PendingOcclusionQuery[FOcclusionQueryHelpers::MaxBufferedOcclusionFrames];
+	FRHIPooledRenderQuery PendingOcclusionQuery[FOcclusionQueryHelpers::MaxBufferedOcclusionFrames];
 	uint32 PendingOcclusionQueryFrames[FOcclusionQueryHelpers::MaxBufferedOcclusionFrames]; // not intialized...this is ok
 
 public:
@@ -326,9 +323,9 @@ public:
 			{
 				uint32 TestFrame = FrameNumber - (DeltaFrame - 1);
 				const uint32 QueryIndex = FOcclusionQueryHelpers::GetQueryLookupIndex(TestFrame, NumBufferedFrames);
-				if (PendingOcclusionQuery[QueryIndex].GetReference() && PendingOcclusionQueryFrames[QueryIndex] != TestFrame)
+				if (PendingOcclusionQuery[QueryIndex].GetQuery() && PendingOcclusionQueryFrames[QueryIndex] != TestFrame)
 				{
-					Pool.ReleaseQuery(PendingOcclusionQuery[QueryIndex]);
+					PendingOcclusionQuery[QueryIndex].ReleaseQuery();
 				}
 			}
 		}
@@ -337,28 +334,28 @@ public:
 	FORCEINLINE void ReleaseQuery(TOcclusionQueryPool& Pool, uint32 FrameNumber, int32 NumBufferedFrames)
 	{
 		const uint32 QueryIndex = FOcclusionQueryHelpers::GetQueryLookupIndex(FrameNumber, NumBufferedFrames);
-		if (PendingOcclusionQuery[QueryIndex].GetReference())
+		if (PendingOcclusionQuery[QueryIndex].GetQuery())
 		{
-			Pool.ReleaseQuery(PendingOcclusionQuery[QueryIndex]);
+			PendingOcclusionQuery[QueryIndex].ReleaseQuery();
 		}
 	}
 
-	FORCEINLINE FRenderQueryRHIParamRef GetPastQuery(uint32 FrameNumber, int32 NumBufferedFrames)
+	FORCEINLINE FRHIRenderQuery* GetPastQuery(uint32 FrameNumber, int32 NumBufferedFrames)
 	{
 		// Get the oldest occlusion query
 		const uint32 QueryIndex = FOcclusionQueryHelpers::GetQueryLookupIndex(FrameNumber, NumBufferedFrames);
-		if (PendingOcclusionQuery[QueryIndex].GetReference() && PendingOcclusionQueryFrames[QueryIndex] == FrameNumber - uint32(NumBufferedFrames))
+		if (PendingOcclusionQuery[QueryIndex].GetQuery() && PendingOcclusionQueryFrames[QueryIndex] == FrameNumber - uint32(NumBufferedFrames))
 		{
-			return PendingOcclusionQuery[QueryIndex].GetReference();
+			return PendingOcclusionQuery[QueryIndex].GetQuery();
 		}
 		return nullptr;
 	}
 
-	FORCEINLINE void SetCurrentQuery(uint32 FrameNumber, FRenderQueryRHIParamRef NewQuery, int32 NumBufferedFrames)
+	FORCEINLINE void SetCurrentQuery(uint32 FrameNumber, FRHIPooledRenderQuery&& NewQuery, int32 NumBufferedFrames)
 	{
 		// Get the current occlusion query
 		const uint32 QueryIndex = FOcclusionQueryHelpers::GetQueryIssueIndex(FrameNumber, NumBufferedFrames);
-		PendingOcclusionQuery[QueryIndex] = NewQuery;
+		PendingOcclusionQuery[QueryIndex] = MoveTemp(NewQuery);
 		PendingOcclusionQueryFrames[QueryIndex] = FrameNumber;
 	}
 };
@@ -528,10 +525,15 @@ Buffers multiple frames to avoid waiting on the GPU so times are a little lagged
 */
 class FLatentGPUTimer
 {
-	static const int32 NumBufferedFrames = FOcclusionQueryHelpers::MaxBufferedOcclusionFrames + 1;
+	FRenderQueryPoolRHIRef TimerQueryPool;
 public:
+	static const int32 NumBufferedFrames = FOcclusionQueryHelpers::MaxBufferedOcclusionFrames + 1;
 
-	FLatentGPUTimer(int32 InAvgSamples = 30);
+	FLatentGPUTimer(FRenderQueryPoolRHIRef InTimerQueryPool, int32 InAvgSamples = 30);
+	~FLatentGPUTimer()
+	{
+		Release();
+	}
 
 	void Release();
 
@@ -558,8 +560,9 @@ private:
 	int32 SampleIndex;
 
 	int32 QueryIndex;
-	FRenderQueryRHIRef StartQueries[NumBufferedFrames];
-	FRenderQueryRHIRef EndQueries[NumBufferedFrames];
+	bool QueriesInFlight[NumBufferedFrames];
+	FRHIPooledRenderQuery StartQueries[NumBufferedFrames];
+	FRHIPooledRenderQuery EndQueries[NumBufferedFrames];
 	FGraphEventRef QuerySubmittedFences[NumBufferedFrames];
 };
 
@@ -635,7 +638,7 @@ struct FHLODSceneNodeVisibilityState
  * This class is associated with a particular camera across multiple frames by the game thread.
  * The game thread calls FRendererModule::AllocateViewState to create an instance of this private implementation.
  */
-class FSceneViewState : public FSceneViewStateInterface, public FDeferredCleanupInterface, public FRenderResource
+class FSceneViewState : public FSceneViewStateInterface, public FRenderResource
 {
 public:
 
@@ -677,11 +680,11 @@ public:
 	};
 
 	uint32 UniqueID;
-	typedef TMap<FSceneViewState::FProjectedShadowKey, FRenderQueryRHIRef> ShadowKeyOcclusionQueryMap;
+	typedef TMap<FSceneViewState::FProjectedShadowKey, FRHIPooledRenderQuery> ShadowKeyOcclusionQueryMap;
 	TArray<ShadowKeyOcclusionQueryMap, TInlineAllocator<FOcclusionQueryHelpers::MaxBufferedOcclusionFrames> > ShadowOcclusionQueryMaps;
 
 	/** The view's occlusion query pool. */
-	FRenderQueryPool OcclusionQueryPool;
+	FRenderQueryPoolRHIRef OcclusionQueryPool;
 
 	FHZBOcclusionTester HZBOcclusionTests;
 
@@ -950,6 +953,7 @@ public:
 	FRWBufferStructured PrimitiveShaderDataBuffer;
 
 	/** Timestamp queries around separate translucency, used for auto-downsampling. */
+	FRenderQueryPoolRHIRef TimerQueryPool;
 	FLatentGPUTimer TranslucencyTimer;
 	FLatentGPUTimer SeparateTranslucencyTimer;
 
@@ -1214,69 +1218,7 @@ public:
 
 	virtual void ReleaseDynamicRHI() override
 	{
-		for (int i = 0; i < ShadowOcclusionQueryMaps.Num(); ++i)
-		{
-			ShadowOcclusionQueryMaps[i].Reset();
-		}
-		PrimitiveOcclusionHistorySet.Empty();
-		PrimitiveFadingStates.Empty();
-		OcclusionQueryPool.Release();
 		HZBOcclusionTests.ReleaseDynamicRHI();
-		EyeAdaptationRTManager.SafeRelease();
-		CombinedLUTRenderTarget.SafeRelease();
-		PrevFrameViewInfo.SafeRelease();
-		LightShaftOcclusionHistory.SafeRelease();
-		LightShaftBloomHistoryRTs.Empty();
-		DistanceFieldAOHistoryRT.SafeRelease();
-		DistanceFieldIrradianceHistoryRT.SafeRelease();
-		MobileAaBloomSunVignette0.SafeRelease();
-		MobileAaBloomSunVignette1.SafeRelease();
-		MobileAaColor0.SafeRelease();
-		MobileAaColor1.SafeRelease();
-		BloomFFTKernel.SafeRelease();
-		SelectionOutlineCacheKey.SafeRelease();
-		SelectionOutlineCacheValue.SafeRelease();
-
-		{
-			// Sets the mipbias that is normally <= 0 to invalid large number.
-			MaterialTextureCachedMipBias = BIG_NUMBER;
-			MaterialTextureBilinearWrapedSamplerCache.SafeRelease();
-			MaterialTextureBilinearClampedSamplerCache.SafeRelease();
-		}
-
-		for (int32 CascadeIndex = 0; CascadeIndex < ARRAY_COUNT(GlobalDistanceFieldClipmapState); CascadeIndex++)
-		{
-			for (int32 CacheType = 0; CacheType < ARRAY_COUNT(GlobalDistanceFieldClipmapState[CascadeIndex].Cache); CacheType++)
-			{
-				GlobalDistanceFieldClipmapState[CascadeIndex].Cache[CacheType].VolumeTexture.SafeRelease();
-			}
-		}
-
-		IndirectShadowCapsuleShapesVertexBuffer.SafeRelease();
-		IndirectShadowCapsuleShapesSRV.SafeRelease();
-		IndirectShadowMeshDistanceFieldCasterIndicesVertexBuffer.SafeRelease();
-		IndirectShadowMeshDistanceFieldCasterIndicesSRV.SafeRelease();
-		IndirectShadowLightDirectionVertexBuffer.SafeRelease();
-		IndirectShadowLightDirectionSRV.SafeRelease();
-		CapsuleTileIntersectionCountsBuffer.Release();
-		TranslucencyTimer.Release();
-		SeparateTranslucencyTimer.Release();
-		if (!!ForwardLightingResources)
-		{
-			ForwardLightingResources->Release();
-			ForwardLightingResources = nullptr;
-		}
-		ForwardLightingCullingResources.Release();
-		LightScatteringHistory.SafeRelease();
-		PrimitiveShaderDataBuffer.Release();
-#if RHI_RAYTRACING
-		PathTracingIrradianceRT.SafeRelease();
-		PathTracingSampleCountRT.SafeRelease();
-		VarianceMipTreeDimensions = FIntVector(0);
-		TotalRayCount = 0;
-		PathTracingRect = FIntRect(0, 0, 0, 0);
-		IESLightProfileResources.Release();
-#endif 
 	}
 
 	// FSceneViewStateInterface
diff --git a/Engine/Source/Runtime/Renderer/Private/SceneRendering.h b/Engine/Source/Runtime/Renderer/Private/SceneRendering.h
--- a/Engine/Source/Runtime/Renderer/Private/SceneRendering.h
+++ b/Engine/Source/Runtime/Renderer/Private/SceneRendering.h
@@ -249,7 +249,7 @@ public:
 	 * Batches a primitive's occlusion query for rendering.
 	 * @param Bounds - The primitive's bounds.
 	 */
-	FRenderQueryRHIParamRef BatchPrimitive(const FVector& BoundsOrigin, const FVector& BoundsBoxExtent, FGlobalDynamicVertexBuffer& DynamicVertexBuffer);
+	FRHIPooledRenderQuery BatchPrimitive(const FVector& BoundsOrigin, const FVector& BoundsBoxExtent, FGlobalDynamicVertexBuffer& DynamicVertexBuffer);
 	inline int32 GetNumBatchOcclusionQueries() const
 	{
 		return BatchOcclusionQueries.Num();
@@ -259,7 +259,7 @@ private:
 
 	struct FOcclusionBatch
 	{
-		FRenderQueryRHIRef Query;
+		TRefCountPtr<FRHIRenderQuery> Query;
 		FGlobalDynamicVertexBuffer::FAllocation VertexAllocation;
 	};
 
@@ -276,7 +276,7 @@ private:
 	uint32 NumBatchedPrimitives;
 
 	/** The pool to allocate occlusion queries from. */
-	class FRenderQueryPool* OcclusionQueryPool;
+	TRefCountPtr<FRHIRenderQueryPool> OcclusionQueryPool;
 };
 
 class FHZBOcclusionTester : public FRenderResource
@@ -676,24 +676,6 @@ struct FPreviousViewInfo
 
 	// History for shadow denoising.
 	TMap<const ULightComponent*, FScreenSpaceFilteringHistory> ShadowHistories;
-
-
-	void SafeRelease()
-	{
-		DepthBuffer.SafeRelease();
-		GBufferA.SafeRelease();
-		GBufferB.SafeRelease();
-		GBufferC.SafeRelease();
-		TemporalAAHistory.SafeRelease();
-		DOFSetupHistory.SafeRelease();
-		SSRHistory.SafeRelease();
-		CustomSSRInput.SafeRelease();
-		ReflectionsHistory.SafeRelease();
-		AmbientOcclusionHistory.SafeRelease();
-		GlobalIlluminationHistory.SafeRelease();
-		SkyLightHistory.SafeRelease();
-		ShadowHistories.Reset();
-	}
 };
 
 class FViewCommands
diff --git a/Engine/Source/Runtime/Renderer/Private/SceneVisibility.cpp b/Engine/Source/Runtime/Renderer/Private/SceneVisibility.cpp
--- a/Engine/Source/Runtime/Renderer/Private/SceneVisibility.cpp
+++ b/Engine/Source/Runtime/Renderer/Private/SceneVisibility.cpp
@@ -695,7 +695,7 @@ static void FetchVisibilityForPrimitives_Range(FVisForPrimParams& Params, FGloba
 	bool bClearQueries = !View.Family->EngineShowFlags.HitProxies;
 	const float CurrentRealTime = View.Family->CurrentRealTime;
 	uint32 OcclusionFrameCounter = ViewState->OcclusionFrameCounter;
-	FRenderQueryPool& OcclusionQueryPool = ViewState->OcclusionQueryPool;
+	FRHIRenderQueryPool* OcclusionQueryPool = ViewState->OcclusionQueryPool;
 	FHZBOcclusionTester& HZBOcclusionTests = ViewState->HZBOcclusionTests;
 
 	int32 ReadBackLagTolerance = NumBufferedFrames;
@@ -858,7 +858,7 @@ static void FetchVisibilityForPrimitives_Range(FVisForPrimParams& Params, FGloba
 						// Read the occlusion query results.
 						uint64 NumSamples = 0;
 						bool bGrouped = false;
-						FRenderQueryRHIParamRef PastQuery = PrimitiveOcclusionHistory->GetQueryForReading(OcclusionFrameCounter, NumBufferedFrames, ReadBackLagTolerance, bGrouped);
+						FRHIRenderQuery* PastQuery = PrimitiveOcclusionHistory->GetQueryForReading(OcclusionFrameCounter, NumBufferedFrames, ReadBackLagTolerance, bGrouped);
 						if (PastQuery)
 						{
 							//int32 RefCount = PastQuery.GetReference()->GetRefCount();
@@ -1299,7 +1299,7 @@ static int32 FetchVisibilityForPrimitives(const FScene* Scene, FViewInfo& View,
 			StartIndex += NumToProcess;
 		}
 
-		FRenderQueryPool& OcclusionQueryPool = ViewState->OcclusionQueryPool;
+		FRHIRenderQueryPool* OcclusionQueryPool = ViewState->OcclusionQueryPool;
 		FHZBOcclusionTester& HZBOcclusionTests = ViewState->HZBOcclusionTests;		
 
 		int32 NumOccludedPrims = 0;
@@ -1368,12 +1368,10 @@ static int32 FetchVisibilityForPrimitives(const FScene* Scene, FViewInfo& View,
 
 			//now add new primitive histories to the view. may resize the view's array.
 			for (int32 i = 0; i < NumTasks; ++i)
-			{								
-				const TArray<FPrimitiveOcclusionHistory>& NewHistoryArray = OutputOcclusionHistory[i];				
-				for (int32 HistoryIndex = 0; HistoryIndex < NewHistoryArray.Num(); ++HistoryIndex)
+			{											
+				for (int32 HistoryIndex = 0; HistoryIndex < OutputOcclusionHistory[i].Num(); ++HistoryIndex)
 				{
-					const FPrimitiveOcclusionHistory& CopySourceHistory = NewHistoryArray[HistoryIndex];
-					ViewPrimitiveOcclusionHistory.Add(CopySourceHistory);
+					ViewPrimitiveOcclusionHistory.Add(MoveTemp(OutputOcclusionHistory[i][HistoryIndex]));
 				}
 
 				//accumulate occluded prims across tasks
diff --git a/Engine/Source/Runtime/UtilityShaders/Private/LongGPUTask.cpp b/Engine/Source/Runtime/UtilityShaders/Private/LongGPUTask.cpp
--- a/Engine/Source/Runtime/UtilityShaders/Private/LongGPUTask.cpp
+++ b/Engine/Source/Runtime/UtilityShaders/Private/LongGPUTask.cpp
@@ -60,10 +60,12 @@ void MeasureLongGPUTaskExecutionTime(FRHICommandListImmediate& RHICmdList)
 {
 	const int32 NumIterationsForMeasurement = 5;
 
-	FRenderQueryRHIRef TimeQueryStart = RHICmdList.CreateRenderQuery(RQT_AbsoluteTime);
-	FRenderQueryRHIRef TimeQueryEnd = RHICmdList.CreateRenderQuery(RQT_AbsoluteTime);
+	FRenderQueryPoolRHIRef TimerQueryPool = RHICreateRenderQueryPool(RQT_AbsoluteTime, 2);
 
-	if (TimeQueryStart == nullptr || TimeQueryEnd == nullptr)
+	FRHIPooledRenderQuery TimeQueryStart = TimerQueryPool->AllocateQuery();
+	FRHIPooledRenderQuery TimeQueryEnd = TimerQueryPool->AllocateQuery();
+
+	if (!GSupportsTimestampRenderQueries)
 	{
 		// Not all platforms/drivers support RQT_AbsoluteTime queries
 		// Use fixed number of iterations on those platforms
@@ -74,17 +76,17 @@ void MeasureLongGPUTaskExecutionTime(FRHICommandListImmediate& RHICmdList)
 	uint64 StartTime = 0;
 	uint64 EndTime = 0;
 
-	RHICmdList.EndRenderQuery(TimeQueryStart);
+	RHICmdList.EndRenderQuery(TimeQueryStart.GetQuery());
 
 	IssueScalableLongGPUTask(RHICmdList, NumIterationsForMeasurement);
 
-	RHICmdList.EndRenderQuery(TimeQueryEnd);
+	RHICmdList.EndRenderQuery(TimeQueryEnd.GetQuery());
 
 	// Required by DX12 to resolve the query
 	RHICmdList.SubmitCommandsHint();
 	RHICmdList.ImmediateFlush(EImmediateFlushType::FlushRHIThread);
 
-	if (RHICmdList.GetRenderQueryResult(TimeQueryStart, StartTime, true) && RHICmdList.GetRenderQueryResult(TimeQueryEnd, EndTime, true))
+	if (RHICmdList.GetRenderQueryResult(TimeQueryStart.GetQuery(), StartTime, true) && RHICmdList.GetRenderQueryResult(TimeQueryEnd.GetQuery(), EndTime, true))
 	{
 		NumMeasuredIterationsToAchieve500ms = FMath::Clamp(FMath::FloorToInt(500.0f / ((EndTime - StartTime) / 1000.0f / NumIterationsForMeasurement)), 1, 200);
 	}
