Added support for multiview in ray tracing. Ray tracing high level functions unified and cleaned.
TODO: some denoisers still need to be adapted. A refactor might be needed after more research to see in which cases it is better to have separate render targets and in which ones it is better to have one and share it.

diff --git a/Engine/Shaders/Private/PathTracing/PathTracing.usf b/Engine/Shaders/Private/PathTracing/PathTracing.usf
--- a/Engine/Shaders/Private/PathTracing/PathTracing.usf
+++ b/Engine/Shaders/Private/PathTracing/PathTracing.usf
@@ -28,7 +28,7 @@ RaytracingAccelerationStructure TLAS;
 [shader("raygeneration")]
 void PathTracingMainRG()
 {
-	uint2 LaunchIndex = DispatchRaysIndex().xy;
+	uint2 LaunchIndex = DispatchRaysIndex().xy + View.ViewRectMin;
 	
 	bool IsUnidirectionalEnabled = true;
 	bool IsNextEventEstimationEnabled = true;
diff --git a/Engine/Shaders/Private/RayTracing/RayTracingAmbientOcclusionRGS.usf b/Engine/Shaders/Private/RayTracing/RayTracingAmbientOcclusionRGS.usf
--- a/Engine/Shaders/Private/RayTracing/RayTracingAmbientOcclusionRGS.usf
+++ b/Engine/Shaders/Private/RayTracing/RayTracingAmbientOcclusionRGS.usf
@@ -39,7 +39,7 @@ void GenerateCosineNormalRay(
 [shader("raygeneration")]
 void AmbientOcclusionRGS()
 {
-	uint2 PixelCoord = DispatchRaysIndex().xy;
+	uint2 PixelCoord = DispatchRaysIndex().xy + View.ViewRectMin;
 
 	RandomSequence RandSequence;
 	uint LinearIndex = CalcLinearIndex(PixelCoord);
diff --git a/Engine/Shaders/Private/RayTracing/RayTracingBarycentrics.usf b/Engine/Shaders/Private/RayTracing/RayTracingBarycentrics.usf
--- a/Engine/Shaders/Private/RayTracing/RayTracingBarycentrics.usf
+++ b/Engine/Shaders/Private/RayTracing/RayTracingBarycentrics.usf
@@ -16,7 +16,7 @@ struct FBarycentricsPayload
 [shader("raygeneration")]
 void RayTracingBarycentricsMainRGS()
 {
-	uint2 PixelCoord = DispatchRaysIndex().xy;
+	uint2 PixelCoord = DispatchRaysIndex().xy + View.ViewRectMin;
 
 	float2 RenderTargetUV = (float2(PixelCoord) + .5f) * View.BufferSizeAndInvSize.zw;
 
diff --git a/Engine/Shaders/Private/RayTracing/RayTracingDebug.usf b/Engine/Shaders/Private/RayTracing/RayTracingDebug.usf
--- a/Engine/Shaders/Private/RayTracing/RayTracingDebug.usf
+++ b/Engine/Shaders/Private/RayTracing/RayTracingDebug.usf
@@ -15,7 +15,7 @@ uint VisualizationMode;
 [shader("raygeneration")]
 void RayTracingDebugMainRGS()
 {
-	uint2 PixelCoord = DispatchRaysIndex().xy;
+	uint2 PixelCoord = DispatchRaysIndex().xy + View.ViewRectMin;
 
 	float2 RenderTargetUV = (float2(PixelCoord) + .5f) * View.BufferSizeAndInvSize.zw;
 
diff --git a/Engine/Shaders/Private/RayTracing/RayTracingGlobalIlluminationRGS.usf b/Engine/Shaders/Private/RayTracing/RayTracingGlobalIlluminationRGS.usf
--- a/Engine/Shaders/Private/RayTracing/RayTracingGlobalIlluminationRGS.usf
+++ b/Engine/Shaders/Private/RayTracing/RayTracingGlobalIlluminationRGS.usf
@@ -55,7 +55,7 @@ void GenerateCosineNormalRay(
 [shader("raygeneration")]
 void GlobalIlluminationRGS()
 {
-	uint2 DispatchThreadId = DispatchRaysIndex().xy;
+	uint2 DispatchThreadId = DispatchRaysIndex().xy + View.ViewRectMin;
 	RWGlobalIlluminationUAV[DispatchThreadId] = 0.0;
 	RWRayDistanceUAV[DispatchThreadId] = float2(-1.0, 0.0);
 
diff --git a/Engine/Shaders/Private/RayTracing/RayTracingOcclusionRGS.usf b/Engine/Shaders/Private/RayTracing/RayTracingOcclusionRGS.usf
--- a/Engine/Shaders/Private/RayTracing/RayTracingOcclusionRGS.usf
+++ b/Engine/Shaders/Private/RayTracing/RayTracingOcclusionRGS.usf
@@ -116,8 +116,8 @@ bool GenerateOcclusionRay(
 [shader("raygeneration")]
 void OcclusionRGS()
 {
-	uint2 PixelCoord = DispatchRaysIndex().xy;
-
+	uint2 PixelCoord = DispatchRaysIndex().xy + View.ViewRectMin;
+	
 	float Visibility = 0.0;
 	float SumRayDistance = 0.0;
 	float ClosestRayDistance = -2.0;
@@ -127,130 +127,130 @@ void OcclusionRGS()
 
 	if (all(PixelCoord >= LightScissor.xy) && all(PixelCoord <= LightScissor.zw))
 	{
-	RandomSequence RandSequence;
-	uint LinearIndex = CalcLinearIndex(PixelCoord);
-	RandomSequence_Initialize(RandSequence, LinearIndex, View.StateFrameIndex);
-
-	FLightShaderParameters LightParameters = GetRootLightShaderParameters();
-
-	// Get G-Buffer surface data
-	float2 InvBufferSize = View.BufferSizeAndInvSize.zw;
-	float2 BufferUV = (float2(PixelCoord) + 0.5) * InvBufferSize;
-	FGBufferData GBufferData = GetGBufferDataFromSceneBlackboard(BufferUV);
-	float Depth = GBufferData.Depth;
-	float3 WorldPosition = ReconstructWorldPositionFromDepth(BufferUV, Depth);
-	float3 WorldNormal = GBufferData.WorldNormal;
-
-	uint LightChannel = GetSceneLightingChannel(PixelCoord);
-
-	// Mask out depth values that are infinitely far away
-	float DeviceZ = SceneDepthBuffer.Load(int3(PixelCoord, 0)).r;
-	bool IsFiniteDepth = DeviceZ > 0.0;
-	bool bTraceRay = (
-		IsFiniteDepth &&
-		GBufferData.ShadingModelID != SHADINGMODELID_UNLIT &&
-		(LightChannel & LightingChannelMask) != 0);
-
-	if (!bTraceRay)
-	{
-		LocalSamplesPerPixel = 0.0;
-	}
-
-	// Cheap solution for double sided foliage to have valid lighing through the denoiser
-	if (GBufferData.ShadingModelID == SHADINGMODELID_TWOSIDED_FOLIAGE)
-	{
-	#if LIGHT_TYPE == LIGHT_TYPE_DIRECTIONAL
-		const float3 LightDirection = LightParameters.Direction;
-	#else
-		const float3 LightDirection = normalize(LightParameters.Position - WorldPosition);
-	#endif
-		if (dot(LightDirection, WorldNormal) < 0)
-			WorldNormal = -WorldNormal;
-	}
-
-	LOOP
-	for (uint SampleIndex = 0; SampleIndex < LocalSamplesPerPixel; ++SampleIndex)
-	{
-		uint DummyVariable;
-		float2 RandSample = RandomSequence_GenerateSample2D(RandSequence, DummyVariable);
-
-		BRANCH
-		if (LocalSamplesPerPixel == 1 && 1)
+		RandomSequence RandSequence;
+		uint LinearIndex = CalcLinearIndex(PixelCoord);
+		RandomSequence_Initialize(RandSequence, LinearIndex, View.StateFrameIndex);
+
+		FLightShaderParameters LightParameters = GetRootLightShaderParameters();
+
+		// Get G-Buffer surface data
+		float2 InvBufferSize = View.BufferSizeAndInvSize.zw;
+		float2 BufferUV = (float2(PixelCoord) + 0.5) * InvBufferSize;
+		FGBufferData GBufferData = GetGBufferDataFromSceneBlackboard(BufferUV);
+		float Depth = GBufferData.Depth;
+		float3 WorldPosition = ReconstructWorldPositionFromDepth(BufferUV, Depth);
+		float3 WorldNormal = GBufferData.WorldNormal;
+
+		uint LightChannel = GetSceneLightingChannel(PixelCoord);
+
+		// Mask out depth values that are infinitely far away
+		float DeviceZ = SceneDepthBuffer.Load(int3(PixelCoord, 0)).r;
+		bool IsFiniteDepth = DeviceZ > 0.0;
+		bool bTraceRay = (
+			IsFiniteDepth &&
+			GBufferData.ShadingModelID != SHADINGMODELID_UNLIT &&
+			(LightChannel & LightingChannelMask) != 0);
+
+		if (!bTraceRay)
 		{
-			const uint TemporalBits = 10;
-			uint FrameIndexMod1024 = ReverseBitsN(GetPowerOfTwoModulatedFrameIndex(1 << TemporalBits), TemporalBits);
-
-			uint2 SobolBase = SobolPixel(PixelCoord);
-			uint2 SobolFrame = SobolIndex(SobolBase, FrameIndexMod1024, TemporalBits);
-
-			RandSample = SobolIndexToUniformUnitSquare(SobolFrame);
+			LocalSamplesPerPixel = 0.0;
 		}
 
-		RayDesc Ray;
-		bool bIsValidRay = GenerateOcclusionRay(
-			LightParameters,
-			WorldPosition, WorldNormal,
-			RandSample,
-			/* out */ Ray.Origin,
-			/* out */ Ray.Direction,
-			/* out */ Ray.TMin,
-			/* out */ Ray.TMax);
-		ApplyPositionBias(Ray, WorldNormal, NormalBias);
-
-		BRANCH
-		if (!bIsValidRay)
+		// Cheap solution for double sided foliage to have valid lighing through the denoiser
+		if (GBufferData.ShadingModelID == SHADINGMODELID_TWOSIDED_FOLIAGE)
 		{
-			continue;
-		}
-		else if (dot(WorldNormal, Ray.Direction) <= 0.0)
-		{
-			ClosestRayDistance = 0.001;
-			continue;
+		#if LIGHT_TYPE == LIGHT_TYPE_DIRECTIONAL
+			const float3 LightDirection = LightParameters.Direction;
+		#else
+			const float3 LightDirection = normalize(LightParameters.Position - WorldPosition);
+		#endif
+			if (dot(LightDirection, WorldNormal) < 0)
+				WorldNormal = -WorldNormal;
 		}
 
-		// Attenuation check
-		if (LightParameters.InvRadius > 0.0)
+		LOOP
+		for (uint SampleIndex = 0; SampleIndex < LocalSamplesPerPixel; ++SampleIndex)
 		{
-			float MaxAttenuationDistance = 1.0 / LightParameters.InvRadius;
-			if (Ray.TMax > MaxAttenuationDistance)
+			uint DummyVariable;
+			float2 RandSample = RandomSequence_GenerateSample2D(RandSequence, DummyVariable);
+
+			BRANCH
+			if (LocalSamplesPerPixel == 1 && 1)
+			{
+				const uint TemporalBits = 10;
+				uint FrameIndexMod1024 = ReverseBitsN(GetPowerOfTwoModulatedFrameIndex(1 << TemporalBits), TemporalBits);
+
+				uint2 SobolBase = SobolPixel(PixelCoord);
+				uint2 SobolFrame = SobolIndex(SobolBase, FrameIndexMod1024, TemporalBits);
+
+				RandSample = SobolIndexToUniformUnitSquare(SobolFrame);
+			}
+
+			RayDesc Ray;
+			bool bIsValidRay = GenerateOcclusionRay(
+				LightParameters,
+				WorldPosition, WorldNormal,
+				RandSample,
+				/* out */ Ray.Origin,
+				/* out */ Ray.Direction,
+				/* out */ Ray.TMin,
+				/* out */ Ray.TMax);
+			ApplyPositionBias(Ray, WorldNormal, NormalBias);
+
+			BRANCH
+			if (!bIsValidRay)
 			{
 				continue;
 			}
-		}
+			else if (dot(WorldNormal, Ray.Direction) <= 0.0)
+			{
+				ClosestRayDistance = 0.001;
+				continue;
+			}
 
-		uint RayFlags = 0;
+			// Attenuation check
+			if (LightParameters.InvRadius > 0.0)
+			{
+				float MaxAttenuationDistance = 1.0 / LightParameters.InvRadius;
+				if (Ray.TMax > MaxAttenuationDistance)
+				{
+					continue;
+				}
+			}
 
-#if !ENABLE_TWO_SIDED_GEOMETRY
-		RayFlags |= RAY_FLAG_CULL_BACK_FACING_TRIANGLES;
-#endif
+			uint RayFlags = 0;
 
-		FPackedMaterialClosestHitPayload Payload = (FPackedMaterialClosestHitPayload)0;
+	#if !ENABLE_TWO_SIDED_GEOMETRY
+			RayFlags |= RAY_FLAG_CULL_BACK_FACING_TRIANGLES;
+	#endif
+
+			FPackedMaterialClosestHitPayload Payload = (FPackedMaterialClosestHitPayload)0;
 		
-		TraceRay(
-			TLAS,      // AccelerationStructure
-			RayFlags,
-			RAY_TRACING_MASK_SHADOW,        // InstanceInclusionMask
-			RAY_TRACING_SHADER_SLOT_SHADOW, // RayContributionToHitGroupIndex
-			RAY_TRACING_NUM_SHADER_SLOTS,   // MultiplierForGeometryContributionToShaderIndex
-			0,         // MissShaderIndex
-			Ray,       // RayDesc
-			Payload    // Payload
-		);
-
-		RayCount += 1.0;
-		Visibility += Payload.HitT < 0 ? 1.0 : 0.0;
-		if (Payload.HitT >= 0.0)
-		{
-			ClosestRayDistance = ClosestRayDistance == -2 || Payload.HitT < ClosestRayDistance ? Payload.HitT : ClosestRayDistance;
-			SumRayDistance += Payload.HitT;
-			HitCount += 1.0;
-		}
-		else
-		{
-			ClosestRayDistance = ClosestRayDistance == -2 ? -1 : ClosestRayDistance;
+			TraceRay(
+				TLAS,      // AccelerationStructure
+				RayFlags,
+				RAY_TRACING_MASK_SHADOW,        // InstanceInclusionMask
+				RAY_TRACING_SHADER_SLOT_SHADOW, // RayContributionToHitGroupIndex
+				RAY_TRACING_NUM_SHADER_SLOTS,   // MultiplierForGeometryContributionToShaderIndex
+				0,         // MissShaderIndex
+				Ray,       // RayDesc
+				Payload    // Payload
+			);
+
+			RayCount += 1.0;
+			Visibility += Payload.HitT < 0 ? 1.0 : 0.0;
+			if (Payload.HitT >= 0.0)
+			{
+				ClosestRayDistance = ClosestRayDistance == -2 || Payload.HitT < ClosestRayDistance ? Payload.HitT : ClosestRayDistance;
+				SumRayDistance += Payload.HitT;
+				HitCount += 1.0;
+			}
+			else
+			{
+				ClosestRayDistance = ClosestRayDistance == -2 ? -1 : ClosestRayDistance;
+			}
 		}
 	}
-	}
 
 	float Shadow = (LocalSamplesPerPixel > 0) ? Visibility / LocalSamplesPerPixel : Visibility;
 	
@@ -297,5 +297,4 @@ void OcclusionRGS()
 
 		RWOcclusionMaskUAV[PixelCoord] = OutColor;
 	}
-
 }
\ No newline at end of file
diff --git a/Engine/Shaders/Private/RayTracing/RayTracingRectLightRGS.usf b/Engine/Shaders/Private/RayTracing/RayTracingRectLightRGS.usf
--- a/Engine/Shaders/Private/RayTracing/RayTracingRectLightRGS.usf
+++ b/Engine/Shaders/Private/RayTracing/RayTracingRectLightRGS.usf
@@ -153,7 +153,7 @@ bool GenerateRectLightImportanceSampledRay(
 [shader("raygeneration")]
 void RectLightRGS()
 {
-	uint2 PixelCoord = DispatchRaysIndex().xy;
+	uint2 PixelCoord = DispatchRaysIndex().xy + View.ViewRectMin;
 
 	RandomSequence RandSequence;
 	uint LinearIndex = CalcLinearIndex(PixelCoord);
diff --git a/Engine/Shaders/Private/RayTracing/RayTracingReflections.usf b/Engine/Shaders/Private/RayTracing/RayTracingReflections.usf
--- a/Engine/Shaders/Private/RayTracing/RayTracingReflections.usf
+++ b/Engine/Shaders/Private/RayTracing/RayTracingReflections.usf
@@ -153,6 +153,8 @@ void RayTracingReflectionsRGS()
 	uint2 DispatchThreadId = DispatchRaysIndex().xy;
 #endif
 
+	DispatchThreadId += View.ViewRectMin;
+
 	uint2 PixelCoord = GetPixelCoord(DispatchThreadId, UpscaleFactor); // GBuffer pixel coordinates (full resolution)
 
 	uint LinearIndex = PixelCoord.y * View.BufferSizeAndInvSize.x + PixelCoord.x; // TODO: PixelCoord or DispatchThreadId
diff --git a/Engine/Shaders/Private/RayTracing/RayTracingSkyLightRGS.usf b/Engine/Shaders/Private/RayTracing/RayTracingSkyLightRGS.usf
--- a/Engine/Shaders/Private/RayTracing/RayTracingSkyLightRGS.usf
+++ b/Engine/Shaders/Private/RayTracing/RayTracingSkyLightRGS.usf
@@ -225,7 +225,7 @@ bool GenerateLambertOcclusionRay(
 [shader("raygeneration")]
 void SkyLightRGS()
 {
-	uint2 PixelCoord = DispatchRaysIndex().xy;
+	uint2 PixelCoord = DispatchRaysIndex().xy + View.ViewRectMin;
 
 	RandomSequence RandSequence;
 	uint LinearIndex = CalcLinearIndex(PixelCoord);
diff --git a/Engine/Shaders/Private/RayTracing/RayTracingTranslucency.usf b/Engine/Shaders/Private/RayTracing/RayTracingTranslucency.usf
--- a/Engine/Shaders/Private/RayTracing/RayTracingTranslucency.usf
+++ b/Engine/Shaders/Private/RayTracing/RayTracingTranslucency.usf
@@ -104,7 +104,7 @@ float3 GetSkyRadiance(float3 Direction, float Roughness)
 [shader("raygeneration")]
 void RayTracingTranslucencyRGS()
 {
-	uint2 DispatchThreadId = DispatchRaysIndex().xy;
+	uint2 DispatchThreadId = DispatchRaysIndex().xy + View.ViewRectMin;
 	uint2 PixelCoord = GetPixelCoord(DispatchThreadId, UpscaleFactor);
 	uint LinearIndex = PixelCoord.y * View.BufferSizeAndInvSize.x + PixelCoord.x; // TODO: PixelCoord or DispatchThreadId
 	
diff --git a/Engine/Source/Runtime/Renderer/Private/CompositionLighting/CompositionLighting.cpp b/Engine/Source/Runtime/Renderer/Private/CompositionLighting/CompositionLighting.cpp
--- a/Engine/Source/Runtime/Renderer/Private/CompositionLighting/CompositionLighting.cpp
+++ b/Engine/Source/Runtime/Renderer/Private/CompositionLighting/CompositionLighting.cpp
@@ -102,7 +102,7 @@ bool ShouldRenderScreenSpaceAmbientOcclusion(const FViewInfo& View)
 			&& !IsSimpleForwardShadingEnabled(View.GetShaderPlatform());
 	}
 #if RHI_RAYTRACING
-	bEnabled &= !ShouldRenderRayTracingAmbientOcclusion();
+	bEnabled &= !ShouldRenderRayTracingAmbientOcclusion(View);
 #endif
 	return bEnabled;
 }
@@ -323,7 +323,7 @@ void FCompositionLighting::ProcessAfterBasePass(FRHICommandListImmediate& RHICmd
 		{
 			FRenderingCompositeOutputRef AmbientOcclusion;
 #if RHI_RAYTRACING
-			if (ShouldRenderRayTracingAmbientOcclusion() && SceneContext.bScreenSpaceAOIsValid)
+			if (ShouldRenderRayTracingAmbientOcclusion(View) && SceneContext.bScreenSpaceAOIsValid)
 			{
 				AmbientOcclusion = Context.Graph.RegisterPass(new(FMemStack::Get()) FRCPassPostProcessInput(SceneContext.ScreenSpaceAO));
 			}
diff --git a/Engine/Source/Runtime/Renderer/Private/DeferredShadingRenderer.cpp b/Engine/Source/Runtime/Renderer/Private/DeferredShadingRenderer.cpp
--- a/Engine/Source/Runtime/Renderer/Private/DeferredShadingRenderer.cpp
+++ b/Engine/Source/Runtime/Renderer/Private/DeferredShadingRenderer.cpp
@@ -123,14 +123,6 @@ static TAutoConsoleVariable<int32> CVarForceAllRayTracingEffects(
 	TEXT(" 1: All ray tracing effects enabled"),
 	ECVF_RenderThreadSafe);
 
-static TAutoConsoleVariable<int32> CVarUseAODenoiser(
-	TEXT("r.AmbientOcclusion.Denoiser"),
-	2,
-	TEXT("Choose the denoising algorithm.\n")
-	TEXT(" 0: Disabled;\n")
-	TEXT(" 1: Forces the default denoiser of the renderer;\n")
-	TEXT(" 2: GScreenSpaceDenoiser witch may be overriden by a third party plugin (default)."),
-	ECVF_RenderThreadSafe);
 
 #if !UE_BUILD_SHIPPING
 static TAutoConsoleVariable<int32> CVarForceBlackVelocityBuffer(
@@ -169,7 +161,6 @@ DECLARE_CYCLE_STAT(TEXT("DeferredShadingSceneRenderer ViewExtensionPostRenderBas
 DECLARE_GPU_STAT_NAMED(RayTracingTLAS, TEXT("Ray Tracing Top Level Acceleration Structure"));
 DECLARE_GPU_STAT(Postprocessing);
 DECLARE_GPU_STAT(HZB);
-DECLARE_GPU_STAT_NAMED(AmbientOcclusionDenoiser, TEXT("Ambient Occlusion Denoiser"));
 DECLARE_GPU_STAT_NAMED(Unaccounted, TEXT("[unaccounted]"));
 
 const TCHAR* GetDepthPassReason(bool bDitheredLODTransitionsUseStencil, EShaderPlatform ShaderPlatform)
@@ -706,14 +697,14 @@ bool FDeferredShadingSceneRenderer::GatherRayTracingWorldInstances(FRHICommandLi
 
 			if (RayTracedMeshElementsMask != 0)
 			{
-					FPrimitiveSceneProxy* SceneProxy = Scene->PrimitiveSceneProxies[PrimitiveIndex];
+				FPrimitiveSceneProxy* SceneProxy = Scene->PrimitiveSceneProxies[PrimitiveIndex];
 				TArray<FRayTracingInstance> RayTracingInstances;
 				SceneProxy->GetDynamicRayTracingInstances(MaterialGatheringContext, RayTracingInstances);
 
 				if (RayTracingInstances.Num() > 0)
-						{
+				{
 					for (FRayTracingInstance& Instance : RayTracingInstances)
-							{
+					{
 						FRayTracingGeometryInstance RayTracingInstance = { Instance.Geometry->RayTracingGeometryRHI };
 						RayTracingInstance.Transform = Instance.InstanceTransforms[0];
 						ensureMsgf(Instance.InstanceTransforms.Num() == 1, TEXT("Multi-instancing hasn't been supported"));
@@ -726,12 +717,12 @@ bool FDeferredShadingSceneRenderer::GatherRayTracingWorldInstances(FRHICommandLi
 						uint32 InstanceIndex = ReferenceView.RayTracingGeometryInstances.Add(RayTracingInstance);
 
 						for (int32 ViewIndex = 1; ViewIndex < Views.Num(); ViewIndex++)
-							{
+						{
 							Views[ViewIndex].RayTracingGeometryInstances.Add(RayTracingInstance);
-				}
+						}
 
 						for (int32 SegmentIndex = 0; SegmentIndex < Instance.Materials.Num(); SegmentIndex++)
-				{
+						{
 							FMeshBatch& MeshBatch = Instance.Materials[SegmentIndex];
 							FDynamicRayTracingMeshCommandContext CommandContext(ReferenceView.DynamicRayTracingMeshCommandStorage, ReferenceView.VisibleRayTracingMeshCommands, SegmentIndex, InstanceIndex);
 							FRayTracingMeshProcessor RayTracingMeshProcessor(&CommandContext, Scene, &ReferenceView);
@@ -758,12 +749,12 @@ bool FDeferredShadingSceneRenderer::DispatchRayTracingWorldUpdates(FRHICommandLi
 
 	Scene->GetRayTracingDynamicGeometryCollection()->DispatchUpdates(RHICmdList);
 
-		for (int32 ViewIndex = 0; ViewIndex < Views.Num(); ++ViewIndex)
-		{
-			FViewInfo& View = Views[ViewIndex];
-			SET_DWORD_STAT(STAT_RayTracingInstances, View.RayTracingGeometryInstances.Num());
-			FRayTracingSceneInitializer Initializer;
-			Initializer.Instances = View.RayTracingGeometryInstances;
+	for (int32 ViewIndex = 0; ViewIndex < Views.Num(); ++ViewIndex)
+	{
+		FViewInfo& View = Views[ViewIndex];
+		SET_DWORD_STAT(STAT_RayTracingInstances, View.RayTracingGeometryInstances.Num());
+		FRayTracingSceneInitializer Initializer;
+		Initializer.Instances = View.RayTracingGeometryInstances;
 		Initializer.ShaderSlotsPerGeometrySegment = RAY_TRACING_NUM_SHADER_SLOTS;
 		View.RayTracingScene.RayTracingSceneRHI = RHICreateRayTracingScene(Initializer);
 		RHICmdList.BuildAccelerationStructure(View.RayTracingScene.RayTracingSceneRHI);
@@ -779,7 +770,7 @@ bool FDeferredShadingSceneRenderer::DispatchRayTracingWorldUpdates(FRHICommandLi
 		PreparePathTracing(View, RayGenShaders);
 
 		if (RayGenShaders.Num())
-	{
+		{
 			auto DefaultHitShader = View.ShaderMap->GetShader<FOpaqueShadowHitGroup>()->GetRayTracingShader();
 			auto DefaultMissShader = View.ShaderMap->GetShader<FDefaultMaterialMS>()->GetRayTracingShader();
 
@@ -788,7 +779,7 @@ bool FDeferredShadingSceneRenderer::DispatchRayTracingWorldUpdates(FRHICommandLi
 				DefaultMissShader,
 				DefaultHitShader
 			);
-	}
+		}
 	}
 
 	return true;
@@ -953,17 +944,14 @@ void FDeferredShadingSceneRenderer::Render(FRHICommandListImmediate& RHICmdList)
 	bool bRequiresFarZQuadClear = false;
 
 	const bool bUseGBuffer = IsUsingGBuffers(ShaderPlatform);
+	bool bCanOverlayRayTracingOutput = CanOverlayRayTracingOutput(Views[0]);// #dxr_todo: UE-72557 multi-view case
+	
 	const bool bRenderDeferredLighting = ViewFamily.EngineShowFlags.Lighting
 		&& FeatureLevel >= ERHIFeatureLevel::SM4
 		&& ViewFamily.EngineShowFlags.DeferredLighting
 		&& bUseGBuffer
-#if RHI_RAYTRACING
-		&& (Views[0].RayTracingRenderMode != ERayTracingRenderMode::PathTracing) // #dxr_todo: UE-72557 multi-view case
-		&& (Views[0].RayTracingRenderMode != ERayTracingRenderMode::RayTracingDebug)
-#endif
-		;
+		&& bCanOverlayRayTracingOutput;
 
-	bool bCanOverlayRayTracingOutput = CanOverlayRayTracingOutput(Views[0]);// #dxr_todo: UE-72557 multi-view case
 	bool bComputeLightGrid = false;
 	// Simple forward shading doesn't support local lights. No need to compute light grid
 	if (!IsSimpleForwardShadingEnabled(ShaderPlatform))
@@ -1537,75 +1525,15 @@ void FDeferredShadingSceneRenderer::Render(FRHICommandListImmediate& RHICmdList)
 	checkSlow(RHICmdList.IsOutsideRenderPass());
 #if RHI_RAYTRACING
 	TRefCountPtr<IPooledRenderTarget> SkyLightRT;
+	TRefCountPtr<IPooledRenderTarget> SkyLightHitDistanceRT;
 	TRefCountPtr<IPooledRenderTarget> GlobalIlluminationRT;
-	TRefCountPtr<IPooledRenderTarget> HitDistanceRT;
-	const bool bRayTracingEnabled = IsRayTracingEnabled();
-	if (bRayTracingEnabled)
-	{
-		if (bCanOverlayRayTracingOutput)
-		{
-			// TODO: convert the entire AO and skylight to rendergraph.
-
-			// SkyLight takes priority over ambient occlusion
-			if (ShouldRenderRayTracingSkyLight(Scene->SkyLight))
-			{
-				RenderRayTracingSkyLight(RHICmdList, SkyLightRT, HitDistanceRT);
-			}
-
-			if (ShouldRenderRayTracingGlobalIllumination(Views))
-			{
-				for (int32 ViewIndex = 0; ViewIndex < Views.Num(); ViewIndex++)
-				{
-					RenderRayTracingGlobalIllumination(RHICmdList, Views[ViewIndex], GlobalIlluminationRT, SceneContext.ScreenSpaceAO);
-				}
-			}
-			else if (ShouldRenderRayTracingAmbientOcclusion())
-			{
-				checkSlow(RHICmdList.IsOutsideRenderPass());
-				checkSlow(RHICmdList.IsOutsideRenderPass());
-				TRefCountPtr<IPooledRenderTarget> AmbientOcclusionHitDistanceRT;
-				RenderRayTracingAmbientOcclusion(RHICmdList, nullptr, SceneContext.ScreenSpaceAO, AmbientOcclusionHitDistanceRT);
-				
-				int32 DenoiserMode = CVarUseAODenoiser.GetValueOnRenderThread();
-				if (DenoiserMode != 0)
-				{	
-					SCOPED_GPU_STAT(RHICmdList, AmbientOcclusionDenoiser);
-					FRDGBuilder GraphBuilder(RHICmdList);
-
-					FSceneViewFamilyBlackboard SceneBlackboard;
-					SetupSceneViewFamilyBlackboard(GraphBuilder, &SceneBlackboard);
-
-					const IScreenSpaceDenoiser* DefaultDenoiser = IScreenSpaceDenoiser::GetDefaultDenoiser();
-					const IScreenSpaceDenoiser* DenoiserToUse = DenoiserMode == 1 ? DefaultDenoiser : GScreenSpaceDenoiser;
-					
-					IScreenSpaceDenoiser::FAmbientOcclusionRayTracingConfig RayTracingConfig;
 
-					IScreenSpaceDenoiser::FAmbientOcclusionInputs DenoiserInputs;
-					DenoiserInputs.Mask = GraphBuilder.RegisterExternalTexture(SceneContext.ScreenSpaceAO, TEXT("AOMask"));
-					DenoiserInputs.RayHitDistance = GraphBuilder.RegisterExternalTexture(AmbientOcclusionHitDistanceRT, TEXT("AOHitDistance"));
-
-					FViewInfo& View = Views[0];
-
-					{
-						RDG_EVENT_SCOPE(GraphBuilder, "%s%s(AmbientOcclusion) %dx%d",
-							DenoiserToUse != DefaultDenoiser ? TEXT("ThirdParty ") : TEXT(""),
-							DenoiserToUse->GetDebugName(),
-							View.ViewRect.Width(), View.ViewRect.Height());
-
-						IScreenSpaceDenoiser::FAmbientOcclusionOutputs DenoiserOutputs = DenoiserToUse->DenoiseAmbientOcclusion(
-							GraphBuilder,
-							View,
-							&View.PrevViewInfo,
-							SceneBlackboard,
-							DenoiserInputs,
-							RayTracingConfig);
-						
-						GraphBuilder.QueueTextureExtraction(DenoiserOutputs.AmbientOcclusionMask, &SceneContext.ScreenSpaceAO);
-					}
-					GraphBuilder.Execute();
-				}
-			}
-		}
+	const bool bRayTracingEnabled = IsRayTracingEnabled();
+	if (bRayTracingEnabled && bCanOverlayRayTracingOutput)
+	{		
+		RenderRayTracingSkyLight(RHICmdList, SkyLightRT, SkyLightHitDistanceRT);
+		RenderRayTracingGlobalIllumination(RHICmdList, GlobalIlluminationRT); 
+		RenderRayTracingAmbientOcclusion(RHICmdList, SceneContext.ScreenSpaceAO);
 	}
 #endif // RHI_RAYTRACING
 	checkSlow(RHICmdList.IsOutsideRenderPass());
@@ -1763,8 +1691,9 @@ void FDeferredShadingSceneRenderer::Render(FRHICommandListImmediate& RHICmdList)
 #if RHI_RAYTRACING
 		if (SkyLightRT)
 		{
-			CompositeRayTracingSkyLight(RHICmdList, SkyLightRT, HitDistanceRT);
+			CompositeRayTracingSkyLight(RHICmdList, SkyLightRT, SkyLightHitDistanceRT);
 		}
+
 		if (GlobalIlluminationRT)
 		{
 			for (int32 ViewIndex = 0; ViewIndex < Views.Num(); ++ViewIndex)
@@ -1886,7 +1815,14 @@ void FDeferredShadingSceneRenderer::Render(FRHICommandListImmediate& RHICmdList)
 		RHICmdList.SetCurrentStat(GET_STATID(STAT_CLM_Translucency));
 
 #if RHI_RAYTRACING
-		if (ShouldRenderRayTracingTranslucency(Views))
+		bool bAnyViewWithRaytracingTranslucency = false;
+		for (int32 ViewIndex = 0, Num = Views.Num(); ViewIndex < Num; ViewIndex++)
+		{
+			const FViewInfo& View = Views[ViewIndex];
+			bAnyViewWithRaytracingTranslucency = bAnyViewWithRaytracingTranslucency || ShouldRenderRayTracingTranslucency(View);
+		}
+
+		if (bAnyViewWithRaytracingTranslucency)
 		{
 			ResolveSceneColor(RHICmdList);
 			RenderRayTracingTranslucency(RHICmdList);
diff --git a/Engine/Source/Runtime/Renderer/Private/DeferredShadingRenderer.h b/Engine/Source/Runtime/Renderer/Private/DeferredShadingRenderer.h
--- a/Engine/Source/Runtime/Renderer/Private/DeferredShadingRenderer.h
+++ b/Engine/Source/Runtime/Renderer/Private/DeferredShadingRenderer.h
@@ -479,8 +479,8 @@ private:
 		const FLightSceneInfo& LightSceneInfo,
 		const IScreenSpaceDenoiser::FShadowRayTracingConfig& RayTracingConfig,
 		IScreenSpaceDenoiser::EShadowRequirements DenoiserRequirements,
-		FRDGTextureRef* OutShadowMask,
-		FRDGTextureRef* OutRayHitDistance);
+		FRDGTextureUAV* OutShadowMaskUAV,
+		FRDGTextureUAV* OutRayHitDistanceUAV);
 
 	void RenderRayTracingStochasticRectLight(FRHICommandListImmediate& RHICmdList, const FLightSceneInfo& RectLightSceneInfo, TRefCountPtr<IPooledRenderTarget>& RectLightRT, TRefCountPtr<IPooledRenderTarget>& HitDistanceRT);
 	void CompositeRayTracingSkyLight(FRHICommandListImmediate& RHICmdList, TRefCountPtr<IPooledRenderTarget>& SkyLightRT, TRefCountPtr<IPooledRenderTarget>& HitDistanceRT);
@@ -489,9 +489,23 @@ private:
 	template <int TextureImportanceSampling> void RenderRayTracingRectLightInternal(FRHICommandListImmediate& RHICmdList, const TArray<FViewInfo>& Views, const FLightSceneInfo& RectLightSceneInfo, TRefCountPtr<IPooledRenderTarget>& ScreenShadowMaskTexture, TRefCountPtr<IPooledRenderTarget>& RayDistanceTexture);
 	void VisualizeRectLightMipTree(FRHICommandListImmediate& RHICmdList, const FViewInfo& View, const FRWBuffer& RectLightMipTree, const FIntVector& RectLightMipTreeDimensions);
 	
-	void RenderRayTracingAmbientOcclusion(FRHICommandListImmediate& RHICmdList, const FLightSceneInfo* LightSceneInfo, TRefCountPtr<IPooledRenderTarget>& AmbientOcclusionRT, TRefCountPtr<IPooledRenderTarget>& AmbientOcclusionHitDistanceRT);
+	void RenderRayTracingAmbientOcclusion(FRHICommandListImmediate& RHICmdList, TRefCountPtr<IPooledRenderTarget>& AmbientOcclusionRT);
+	void RenderRayTracingAmbientOcclusion(FRHICommandListImmediate& RHICmdList, FRDGBuilder& GraphBuilder, FViewInfo& View, TRefCountPtr<IPooledRenderTarget>& AmbientOcclusionRT, TRefCountPtr<IPooledRenderTarget>& HitDistanceRT);
 	void CompositeRayTracingAmbientOcclusion(FRHICommandListImmediate& RHICmdList, TRefCountPtr<IPooledRenderTarget>& AmbientOcclusionRT);
-	void RenderRayTracingGlobalIllumination(FRHICommandListImmediate& RHICmdList, FViewInfo& View, TRefCountPtr<IPooledRenderTarget>& GlobalIlluminationRT, TRefCountPtr<IPooledRenderTarget>& AmbientOcclusionRT);
+	
+	void RenderRayTracingGlobalIllumination(FRHICommandListImmediate& RHICmdList, TRefCountPtr<IPooledRenderTarget>& GlobalIlluminationRT);
+	void RenderRayTracingGlobalIllumination(
+		FRHICommandListImmediate& RHICmdList, 
+		FRDGBuilder& GraphBuilder, 
+		FViewInfo& View, 
+		IScreenSpaceDenoiser::FAmbientOcclusionRayTracingConfig& RayTracingConfig, 
+		int32 UpscaleFactor, TRefCountPtr<IPooledRenderTarget>& GlobalIlluminationRT, 
+		TRefCountPtr<IPooledRenderTarget>& AmbientOcclusionRT,
+		FRDGTextureUAV* GlobalIlluminationUAV,
+		FRDGTextureUAV* RayDistanceUAV,
+		const FRDGTextureRef GlobalIlluminationTexture,
+		const FRDGTextureRef RayDistanceTexture
+	);
 	void CompositeGlobalIllumination(FRHICommandListImmediate& RHICmdList, const FViewInfo& View, TRefCountPtr<IPooledRenderTarget>& GlobalIlluminationRT);
 
 	void BuildSkyLightCdfs(FRHICommandListImmediate& RHICmdList, FSkyLightSceneProxy* SkyLight);
diff --git a/Engine/Source/Runtime/Renderer/Private/LightRendering.cpp b/Engine/Source/Runtime/Renderer/Private/LightRendering.cpp
--- a/Engine/Source/Runtime/Renderer/Private/LightRendering.cpp
+++ b/Engine/Source/Runtime/Renderer/Private/LightRendering.cpp
@@ -936,9 +936,10 @@ void FDeferredShadingSceneRenderer::RenderLights(FRHICommandListImmediate& RHICm
 			TArray<TRefCountPtr<IPooledRenderTarget>> PreprocessedShadowMaskTextures;
 
 			const int32 MaxDenoisingBatchSize = FMath::Clamp(CVarMaxShadowDenoisingBatchSize.GetValueOnRenderThread(), 1, IScreenSpaceDenoiser::kMaxBatchSize);
-			const bool bDoShadowDenoisingBatching = DenoiserMode != 0 && MaxDenoisingBatchSize > 1;
+			//#dxr_todo: add multiview support to batching denoising 
+			const bool bDoShadowDenoisingBatching = DenoiserMode != 0 && MaxDenoisingBatchSize > 1 && Views.Num() == 1;
 
-			// Optimisations: batches all shadow ray tracing denoising. Definitely could be smarter to avoid high VGPR pressure if this entire
+			// Optimizations: batches all shadow ray tracing denoising. Definitely could be smarter to avoid high VGPR pressure if this entire
 			// function was converted to render graph, and want least intrusive change as possible. So right not it trades render target memory pressure
 			// for denoising perf.
 			if (RHI_RAYTRACING && bDoShadowDenoisingBatching)
@@ -951,155 +952,186 @@ void FDeferredShadingSceneRenderer::RenderLights(FRHICommandListImmediate& RHICm
 				FSceneViewFamilyBlackboard SceneBlackboard;
 				SetupSceneViewFamilyBlackboard(GraphBuilder, &SceneBlackboard);
 
-				FViewInfo& View = Views[0];
+				// Ray trace the shadow.
+				FRDGTextureRef ShadowMask;
 
-				// Allocate PreprocessedShadowMaskTextures once so QueueTextureExtraction can deferred write.
 				{
-					if (!View.bViewStateIsReadOnly)
-					{
-						View.ViewState->PrevFrameViewInfo.ShadowHistories.Empty();
-						View.ViewState->PrevFrameViewInfo.ShadowHistories.Reserve(SortedLights.Num());
-					}
-					PreprocessedShadowMaskTextures.Reserve(SortedLights.Num());
-					for (int32 LightIndex = AttenuationLightStart; LightIndex < SortedLights.Num(); LightIndex++)
-					{
-						const FSortedLightSceneInfo& SortedLightInfo = SortedLights[LightIndex];
-						const FLightSceneInfo& LightSceneInfo = *SortedLightInfo.LightSceneInfo;
+					FRDGTextureDesc Desc = FRDGTextureDesc::Create2DDesc(
+						SceneBlackboard.SceneDepthBuffer->Desc.Extent,
+						PF_FloatRGBA,
+						FClearValueBinding::Black,
+						TexCreate_None,
+						TexCreate_ShaderResource | TexCreate_RenderTargetable | TexCreate_UAV,
+						/* bInForceSeparateTargetAndShaderResource = */ false);
+					ShadowMask = GraphBuilder.CreateTexture(Desc, TEXT("RayTracingOcclusion"));
+				}
 
-						PreprocessedShadowMaskTextures.Add(TRefCountPtr<IPooledRenderTarget>());
+				FRDGTextureRef RayHitDistance;
+				{
+					FRDGTextureDesc Desc = FRDGTextureDesc::Create2DDesc(
+						SceneBlackboard.SceneDepthBuffer->Desc.Extent,
+						PF_R16F,
+						FClearValueBinding::Black,
+						TexCreate_None,
+						TexCreate_ShaderResource | TexCreate_RenderTargetable | TexCreate_UAV,
+						/* bInForceSeparateTargetAndShaderResource = */ false);
+					RayHitDistance = GraphBuilder.CreateTexture(Desc, TEXT("RayTracingOcclusionDistance"));
+				}
+
+				FRDGTextureUAV* ShadowMaskUAV = GraphBuilder.CreateUAV(FRDGTextureUAVDesc(ShadowMask));
+				FRDGTextureUAV* RayHitDistanceUAV = GraphBuilder.CreateUAV(FRDGTextureUAVDesc(RayHitDistance));
 
+
+				for (int32 ViewIndex = 0; ViewIndex < Views.Num(); ViewIndex++)
+				{
+					FViewInfo& View = Views[ViewIndex];
+
+					// Allocate PreprocessedShadowMaskTextures once so QueueTextureExtraction can deferred write.
+					{
 						if (!View.bViewStateIsReadOnly)
 						{
-							View.ViewState->PrevFrameViewInfo.ShadowHistories.Emplace(LightSceneInfo.Proxy->GetLightComponent());
+							View.ViewState->PrevFrameViewInfo.ShadowHistories.Empty();
+							View.ViewState->PrevFrameViewInfo.ShadowHistories.Reserve(SortedLights.Num());
 						}
-					}
-				}
+						PreprocessedShadowMaskTextures.Reserve(SortedLights.Num());
+						for (int32 LightIndex = AttenuationLightStart; LightIndex < SortedLights.Num(); LightIndex++)
+						{
+							const FSortedLightSceneInfo& SortedLightInfo = SortedLights[LightIndex];
+							const FLightSceneInfo& LightSceneInfo = *SortedLightInfo.LightSceneInfo;
 
-				TMap<const FLightSceneInfo*, FRDGTextureRef> ShadowMaskTextures;
+							PreprocessedShadowMaskTextures.Add(TRefCountPtr<IPooledRenderTarget>());
 
-				// Lambda to share the code quicking of the shadow denoiser.
-				auto QuickOffDenoisingBatch = [&](){
-					int32 InputParameterCount = 0;
-					for (int32 i = 0; i < IScreenSpaceDenoiser::kMaxBatchSize; i++)
-					{
-						InputParameterCount += DenoisingQueue[i].LightSceneInfo != nullptr ? 1 : 0;
+							if (!View.bViewStateIsReadOnly)
+							{
+								View.ViewState->PrevFrameViewInfo.ShadowHistories.Emplace(LightSceneInfo.Proxy->GetLightComponent());
+							}
+						}
 					}
 
-					check(InputParameterCount >= 1);
+					TMap<const FLightSceneInfo*, FRDGTextureRef> ShadowMaskTextures;
 
-					TStaticArray<IScreenSpaceDenoiser::FShadowPenumbraOutputs, IScreenSpaceDenoiser::kMaxBatchSize> Outputs;
+					// Lambda to share the code quicking of the shadow denoiser.
+					auto QuickOffDenoisingBatch = [&](){
+						int32 InputParameterCount = 0;
+						for (int32 i = 0; i < IScreenSpaceDenoiser::kMaxBatchSize; i++)
+						{
+							InputParameterCount += DenoisingQueue[i].LightSceneInfo != nullptr ? 1 : 0;
+						}
 
-					RDG_EVENT_SCOPE(GraphBuilder, "%s%s(Shadow BatchSize=%d) %dx%d",
-						DenoiserToUse != DefaultDenoiser ? TEXT("ThirdParty ") : TEXT(""),
-						DenoiserToUse->GetDebugName(),
-						InputParameterCount,
-						View.ViewRect.Width(), View.ViewRect.Height());
+						check(InputParameterCount >= 1);
 
-					DenoiserToUse->DenoiseShadows(
-						GraphBuilder,
-						View,
-						&View.PrevViewInfo,
-						SceneBlackboard,
-						DenoisingQueue,
-						InputParameterCount,
-						Outputs);
+						TStaticArray<IScreenSpaceDenoiser::FShadowPenumbraOutputs, IScreenSpaceDenoiser::kMaxBatchSize> Outputs;
 
-					for (int32 i = 0; i < InputParameterCount; i++)
-					{
-						const FLightSceneInfo* LightSceneInfo = DenoisingQueue[i].LightSceneInfo;
+						RDG_EVENT_SCOPE(GraphBuilder, "%s%s(Shadow BatchSize=%d) %dx%d",
+							DenoiserToUse != DefaultDenoiser ? TEXT("ThirdParty ") : TEXT(""),
+							DenoiserToUse->GetDebugName(),
+							InputParameterCount,
+							View.ViewRect.Width(), View.ViewRect.Height());
 
-						int32 LightIndex = LightIndices[i];
-						TRefCountPtr<IPooledRenderTarget>* RefDestination = &PreprocessedShadowMaskTextures[LightIndex - AttenuationLightStart];
-						check(*RefDestination == nullptr);
+						DenoiserToUse->DenoiseShadows(
+							GraphBuilder,
+							View,
+							&View.PrevViewInfo,
+							SceneBlackboard,
+							DenoisingQueue,
+							InputParameterCount,
+							Outputs);
 
-						GraphBuilder.QueueTextureExtraction(Outputs[i].DiffusePenumbra, RefDestination);
-						DenoisingQueue[i].LightSceneInfo = nullptr;
-					}
-				}; // QuickOffDenoisingBatch
+						for (int32 i = 0; i < InputParameterCount; i++)
+						{
+							const FLightSceneInfo* LightSceneInfo = DenoisingQueue[i].LightSceneInfo;
 
-				// Ray trace shadows of light that needs, and quick off denoising batch.
-				for (int32 LightIndex = AttenuationLightStart; LightIndex < SortedLights.Num(); LightIndex++)
-				{
-					const FSortedLightSceneInfo& SortedLightInfo = SortedLights[LightIndex];
-					const FLightSceneInfo& LightSceneInfo = *SortedLightInfo.LightSceneInfo;
+							int32 LightIndex = LightIndices[i];
+							TRefCountPtr<IPooledRenderTarget>* RefDestination = &PreprocessedShadowMaskTextures[LightIndex - AttenuationLightStart];
+							check(*RefDestination == nullptr);
 
-					// Denoiser do not support texture rect light important sampling.
-					const bool bDrawShadows = SortedLightInfo.SortKey.Fields.bShadowed && !ShouldRenderRayTracingStochasticRectLight(LightSceneInfo);
+							GraphBuilder.QueueTextureExtraction(Outputs[i].DiffusePenumbra, RefDestination);
+							DenoisingQueue[i].LightSceneInfo = nullptr;
+						}
+					}; // QuickOffDenoisingBatch
 
-					if (!bDrawShadows)
-						continue;
+					// Ray trace shadows of light that needs, and quick off denoising batch.
+					for (int32 LightIndex = AttenuationLightStart; LightIndex < SortedLights.Num(); LightIndex++)
+					{
+						const FSortedLightSceneInfo& SortedLightInfo = SortedLights[LightIndex];
+						const FLightSceneInfo& LightSceneInfo = *SortedLightInfo.LightSceneInfo;
 
-					INC_DWORD_STAT(STAT_NumShadowedLights);
+						// Denoiser do not support texture rect light important sampling.
+						const bool bDrawShadows = SortedLightInfo.SortKey.Fields.bShadowed && !ShouldRenderRayTracingStochasticRectLight(LightSceneInfo);
 
-					const FLightOcclusionType OcclusionType = GetLightOcclusionType(*LightSceneInfo.Proxy);
-					if (OcclusionType != FLightOcclusionType::Raytraced)
-						continue;
+						if (!bDrawShadows)
+							continue;
 
-					if (!LightRequiresDenosier(LightSceneInfo))
-						continue;
+						INC_DWORD_STAT(STAT_NumShadowedLights);
 
-					IScreenSpaceDenoiser::FShadowRayTracingConfig RayTracingConfig;
-					RayTracingConfig.RayCountPerPixel = LightSceneInfo.Proxy->GetSamplesPerPixel();
+						const FLightOcclusionType OcclusionType = GetLightOcclusionType(*LightSceneInfo.Proxy);
+						if (OcclusionType != FLightOcclusionType::Raytraced)
+							continue;
 
-					IScreenSpaceDenoiser::EShadowRequirements DenoiserRequirements = DenoiserToUse->GetShadowRequirements(
-						View, LightSceneInfo, RayTracingConfig);
+						if (!LightRequiresDenosier(LightSceneInfo))
+							continue;
 
-					// Not worth batching and increase memory pressure if the denoiser do not support this ray tracing config.
-					// TODO: add suport for batch with multiple SPP.
-					if (DenoiserRequirements != IScreenSpaceDenoiser::EShadowRequirements::PenumbraAndClosestOccluder)
-					{
-						continue;
-					}
+						IScreenSpaceDenoiser::FShadowRayTracingConfig RayTracingConfig;
+						RayTracingConfig.RayCountPerPixel = LightSceneInfo.Proxy->GetSamplesPerPixel();
 
-					// Ray trace the shadow.
-					FRDGTextureRef ShadowMask;
-					FRDGTextureRef RayHitDistance;
-					{
-						FString LightNameWithLevel;
-						GetLightNameForDrawEvent(LightSceneInfo.Proxy, LightNameWithLevel);
-						RDG_EVENT_SCOPE(GraphBuilder, "%s", *LightNameWithLevel);
+						IScreenSpaceDenoiser::EShadowRequirements DenoiserRequirements = DenoiserToUse->GetShadowRequirements(
+							View, LightSceneInfo, RayTracingConfig);
+
+						// Not worth batching and increase memory pressure if the denoiser do not support this ray tracing config.
+						// TODO: add suport for batch with multiple SPP.
+						if (DenoiserRequirements != IScreenSpaceDenoiser::EShadowRequirements::PenumbraAndClosestOccluder)
+						{
+							continue;
+						}
 
-						RenderRayTracingShadows(
-							GraphBuilder,
-							SceneBlackboard,
-							View,
-							LightSceneInfo,
-							RayTracingConfig,
-							DenoiserRequirements,
-							&ShadowMask,
-							&RayHitDistance);
-					}
 
-					// Queue the ray tracing output for shadow denoising.
-					for (int32 i = 0; i < IScreenSpaceDenoiser::kMaxBatchSize; i++)
-					{
-						if (DenoisingQueue[i].LightSceneInfo == nullptr)
 						{
-							DenoisingQueue[i].LightSceneInfo = &LightSceneInfo;
-							DenoisingQueue[i].RayTracingConfig = RayTracingConfig;
-							DenoisingQueue[i].InputTextures.Penumbra = ShadowMask;
-							DenoisingQueue[i].InputTextures.ClosestOccluder = RayHitDistance;
-							LightIndices[i] = LightIndex;
-
-							// If queue for this light type is full, quick of the batch.
-							if ((i + 1) == MaxDenoisingBatchSize)
-							{
-								QuickOffDenoisingBatch();
-							}
-							break;
+							FString LightNameWithLevel;
+							GetLightNameForDrawEvent(LightSceneInfo.Proxy, LightNameWithLevel);
+							RDG_EVENT_SCOPE(GraphBuilder, "%s", *LightNameWithLevel);
+
+							RenderRayTracingShadows(
+								GraphBuilder,
+								SceneBlackboard,
+								View,
+								LightSceneInfo,
+								RayTracingConfig,
+								DenoiserRequirements,
+								ShadowMaskUAV,
+								RayHitDistanceUAV);
 						}
-						else
+
+						// Queue the ray tracing output for shadow denoising.
+						for (int32 i = 0; i < IScreenSpaceDenoiser::kMaxBatchSize; i++)
 						{
-							check((i - 1) < IScreenSpaceDenoiser::kMaxBatchSize);
+							if (DenoisingQueue[i].LightSceneInfo == nullptr)
+							{
+								DenoisingQueue[i].LightSceneInfo = &LightSceneInfo;
+								DenoisingQueue[i].RayTracingConfig = RayTracingConfig;
+								DenoisingQueue[i].InputTextures.Penumbra = ShadowMask;
+								DenoisingQueue[i].InputTextures.ClosestOccluder = RayHitDistance;
+								LightIndices[i] = LightIndex;
+
+								// If queue for this light type is full, quick of the batch.
+								if ((i + 1) == MaxDenoisingBatchSize)
+								{
+									QuickOffDenoisingBatch();
+								}
+								break;
+							}
+							else
+							{
+								check((i - 1) < IScreenSpaceDenoiser::kMaxBatchSize);
+							}
 						}
-					}
 
-				} // for (int32 LightIndex = AttenuationLightStart; LightIndex < SortedLights.Num(); LightIndex++)
+					} // for (int32 LightIndex = AttenuationLightStart; LightIndex < SortedLights.Num(); LightIndex++)
 
-				// Ensures all denoising queues are processed.
-				if (DenoisingQueue[0].LightSceneInfo)
-				{
-					QuickOffDenoisingBatch();
+					// Ensures all denoising queues are processed.
+					if (DenoisingQueue[0].LightSceneInfo)
+					{
+						QuickOffDenoisingBatch();
+					}
 				}
 
 				GraphBuilder.Execute();
@@ -1146,66 +1178,93 @@ void FDeferredShadingSceneRenderer::RenderLights(FRHICommandListImmediate& RHICm
 					}
 					else if (OcclusionType == FLightOcclusionType::Raytraced)
 					{
-						FViewInfo& View = Views[0];
-
-						IScreenSpaceDenoiser::FShadowRayTracingConfig RayTracingConfig;
-						RayTracingConfig.RayCountPerPixel = LightSceneInfo.Proxy->GetSamplesPerPixel();
-
-						IScreenSpaceDenoiser::EShadowRequirements DenoiserRequirements = IScreenSpaceDenoiser::EShadowRequirements::Bailout;
-						if (DenoiserMode != 0 && LightRequiresDenosier(LightSceneInfo))
-						{
-							DenoiserRequirements = DenoiserToUse->GetShadowRequirements(View, LightSceneInfo, RayTracingConfig);
-						}
-
 						FRDGBuilder GraphBuilder(RHICmdList);
-
+						
 						FSceneViewFamilyBlackboard SceneBlackboard;
 						SetupSceneViewFamilyBlackboard(GraphBuilder, &SceneBlackboard);
 
 						FRDGTextureRef ShadowMask;
-						FRDGTextureRef RayHitDistance;
-						RenderRayTracingShadows(
-							GraphBuilder,
-							SceneBlackboard,
-							View,
-							LightSceneInfo,
-							RayTracingConfig,
-							DenoiserRequirements,
-							&ShadowMask,
-							&RayHitDistance);
+						{
+							FRDGTextureDesc Desc = FRDGTextureDesc::Create2DDesc(
+								SceneBlackboard.SceneDepthBuffer->Desc.Extent,
+								PF_FloatRGBA,
+								FClearValueBinding::Black,
+								TexCreate_None,
+								TexCreate_ShaderResource | TexCreate_RenderTargetable | TexCreate_UAV,
+								/* bInForceSeparateTargetAndShaderResource = */ false);
+							ShadowMask = GraphBuilder.CreateTexture(Desc, TEXT("RayTracingOcclusion"));
+						}
 
-						if (DenoiserRequirements != IScreenSpaceDenoiser::EShadowRequirements::Bailout)
+						FRDGTextureRef RayHitDistance;
 						{
-							TStaticArray<IScreenSpaceDenoiser::FShadowParameters, IScreenSpaceDenoiser::kMaxBatchSize> InputParameters;
-							TStaticArray<IScreenSpaceDenoiser::FShadowPenumbraOutputs, IScreenSpaceDenoiser::kMaxBatchSize> Outputs;
+							FRDGTextureDesc Desc = FRDGTextureDesc::Create2DDesc(
+								SceneBlackboard.SceneDepthBuffer->Desc.Extent,
+								PF_R16F,
+								FClearValueBinding::Black,
+								TexCreate_None,
+								TexCreate_ShaderResource | TexCreate_RenderTargetable | TexCreate_UAV,
+								/* bInForceSeparateTargetAndShaderResource = */ false);
+							RayHitDistance = GraphBuilder.CreateTexture(Desc, TEXT("RayTracingOcclusionDistance"));
+						}
 
-							InputParameters[0].InputTextures.Penumbra = ShadowMask;
-							InputParameters[0].InputTextures.ClosestOccluder = RayHitDistance;
-							InputParameters[0].LightSceneInfo = &LightSceneInfo;
-							InputParameters[0].RayTracingConfig = RayTracingConfig;
+						FRDGTextureUAV* ShadowMaskUAV = GraphBuilder.CreateUAV(FRDGTextureUAVDesc(ShadowMask));
+						FRDGTextureUAV* RayHitDistanceUAV = GraphBuilder.CreateUAV(FRDGTextureUAVDesc(RayHitDistance));
 
-							int32 InputParameterCount = 1;
+						for (int32 ViewIndex = 0; ViewIndex < Views.Num(); ViewIndex++)
+						{
+							FViewInfo& View = Views[ViewIndex];
 
-							RDG_EVENT_SCOPE(GraphBuilder, "%s%s(Shadow BatchSize=%d) %dx%d",
-								DenoiserToUse != DefaultDenoiser ? TEXT("ThirdParty ") : TEXT(""),
-								DenoiserToUse->GetDebugName(),
-								InputParameterCount,
-								View.ViewRect.Width(), View.ViewRect.Height());
+							IScreenSpaceDenoiser::FShadowRayTracingConfig RayTracingConfig;
+							RayTracingConfig.RayCountPerPixel = LightSceneInfo.Proxy->GetSamplesPerPixel();
+							IScreenSpaceDenoiser::EShadowRequirements DenoiserRequirements = IScreenSpaceDenoiser::EShadowRequirements::Bailout;
+							if (DenoiserMode != 0 && LightRequiresDenosier(LightSceneInfo))
+							{
+								DenoiserRequirements = DenoiserToUse->GetShadowRequirements(View, LightSceneInfo, RayTracingConfig);
+							}
 
-							DenoiserToUse->DenoiseShadows(
+							RenderRayTracingShadows(
 								GraphBuilder,
-								View,
-								&View.PrevViewInfo,
 								SceneBlackboard,
-								InputParameters,
-								InputParameterCount,
-								Outputs);
+								View,
+								LightSceneInfo,
+								RayTracingConfig,
+								DenoiserRequirements,
+								ShadowMaskUAV,
+								RayHitDistanceUAV);
 
-							GraphBuilder.QueueTextureExtraction(Outputs[0].DiffusePenumbra, &ScreenShadowMaskTexture);
-						}
-						else
-						{
-							GraphBuilder.QueueTextureExtraction(ShadowMask, &ScreenShadowMaskTexture);
+							if (DenoiserRequirements != IScreenSpaceDenoiser::EShadowRequirements::Bailout)
+							{
+								TStaticArray<IScreenSpaceDenoiser::FShadowParameters, IScreenSpaceDenoiser::kMaxBatchSize> InputParameters;
+								TStaticArray<IScreenSpaceDenoiser::FShadowPenumbraOutputs, IScreenSpaceDenoiser::kMaxBatchSize> Outputs;
+
+								InputParameters[0].InputTextures.Penumbra = ShadowMask;
+								InputParameters[0].InputTextures.ClosestOccluder = RayHitDistance;
+								InputParameters[0].LightSceneInfo = &LightSceneInfo;
+								InputParameters[0].RayTracingConfig = RayTracingConfig;
+
+								int32 InputParameterCount = 1;
+
+								RDG_EVENT_SCOPE(GraphBuilder, "%s%s(Shadow BatchSize=%d) %dx%d",
+									DenoiserToUse != DefaultDenoiser ? TEXT("ThirdParty ") : TEXT(""),
+									DenoiserToUse->GetDebugName(),
+									InputParameterCount,
+									View.ViewRect.Width(), View.ViewRect.Height());
+
+								DenoiserToUse->DenoiseShadows(
+									GraphBuilder,
+									View,
+									&View.PrevViewInfo,
+									SceneBlackboard,
+									InputParameters,
+									InputParameterCount,
+									Outputs);
+
+								GraphBuilder.QueueTextureExtraction(Outputs[0].DiffusePenumbra, &ScreenShadowMaskTexture);
+							}
+							else
+							{
+								GraphBuilder.QueueTextureExtraction(ShadowMask, &ScreenShadowMaskTexture);
+							}
 						}
 
 						GraphBuilder.Execute();
diff --git a/Engine/Source/Runtime/Renderer/Private/RayTracing/RayTracingAmbientOcclusion.cpp b/Engine/Source/Runtime/Renderer/Private/RayTracing/RayTracingAmbientOcclusion.cpp
--- a/Engine/Source/Runtime/Renderer/Private/RayTracing/RayTracingAmbientOcclusion.cpp
+++ b/Engine/Source/Runtime/Renderer/Private/RayTracing/RayTracingAmbientOcclusion.cpp
@@ -13,6 +13,7 @@
 #include "UniformBuffer.h"
 #include "RHI/Public/PipelineStateCache.h"
 #include "Raytracing/RaytracingOptions.h"
+#include "SceneViewFamilyBlackboard.h"
 
 #include "PostProcess/PostProcessing.h"
 #include "PostProcess/SceneFilterRendering.h"
@@ -24,13 +25,14 @@ static FAutoConsoleVariableRef CVarRayTracingAmbientOcclusion(
 	TEXT("Enables ray tracing ambient occlusion (default = 1)")
 );
 
-bool ShouldRenderRayTracingAmbientOcclusion()
-{
-	const int32 ForceAllRayTracingEffects = GetForceRayTracingEffectsCVarValue();
-	const bool bRTAOEnabled = (ForceAllRayTracingEffects > 0 || (GRayTracingAmbientOcclusion > 0 && ForceAllRayTracingEffects < 0));
-
-	return IsRayTracingEnabled() && bRTAOEnabled;
-}
+static TAutoConsoleVariable<int32> CVarUseAODenoiser(
+	TEXT("r.AmbientOcclusion.Denoiser"),
+	2,
+	TEXT("Choose the denoising algorithm.\n")
+	TEXT(" 0: Disabled;\n")
+	TEXT(" 1: Forces the default denoiser of the renderer;\n")
+	TEXT(" 2: GScreenSpaceDenoiser witch may be overriden by a third party plugin (default)."),
+	ECVF_RenderThreadSafe);
 
 static int32 GRayTracingAmbientOcclusionSamplesPerPixel = -1;
 static FAutoConsoleVariableRef CVarRayTracingAmbientOcclusionSamplesPerPixel(
@@ -46,6 +48,17 @@ static TAutoConsoleVariable<int32> CVarRayTracingAmbientOcclusionEnableTwoSidedG
 	ECVF_RenderThreadSafe
 );
 
+
+bool ShouldRenderRayTracingAmbientOcclusion(const FViewInfo& View)
+{
+	const int32 ForceAllRayTracingEffects = GetForceRayTracingEffectsCVarValue();
+	const bool bRTAOEnabled = (ForceAllRayTracingEffects > 0 || (GRayTracingAmbientOcclusion > 0 && ForceAllRayTracingEffects < 0));
+
+	//#dxr_todo: add option to enable RTAO in View.FinalPostProcessSettings
+	return IsRayTracingEnabled() && !ShouldRenderRayTracingGlobalIllumination(View) && bRTAOEnabled;
+}
+
+
 BEGIN_GLOBAL_SHADER_PARAMETER_STRUCT(FAmbientOcclusionData, )
 SHADER_PARAMETER(int, SamplesPerPixel)
 SHADER_PARAMETER(float, MaxRayDistance)
@@ -56,6 +69,7 @@ END_GLOBAL_SHADER_PARAMETER_STRUCT()
 IMPLEMENT_GLOBAL_SHADER_PARAMETER_STRUCT(FAmbientOcclusionData, "AmbientOcclusion");
 
 DECLARE_GPU_STAT_NAMED(RayTracingAmbientOcclusion, TEXT("Ray Tracing Ambient Occlusion"));
+DECLARE_GPU_STAT_NAMED(AmbientOcclusionDenoiser, TEXT("Ambient Occlusion Denoiser"));
 
 template<uint32 EnableTwoSidedGeometry>
 class TAmbientOcclusionRGS : public FGlobalShader
@@ -145,83 +159,147 @@ private:
 IMPLEMENT_SHADER_TYPE(template<>, TAmbientOcclusionRGS<0>, TEXT("/Engine/Private/RayTracing/RayTracingAmbientOcclusionRGS.usf"), TEXT("AmbientOcclusionRGS"), SF_RayGen)
 IMPLEMENT_SHADER_TYPE(template<>, TAmbientOcclusionRGS<1>, TEXT("/Engine/Private/RayTracing/RayTracingAmbientOcclusionRGS.usf"), TEXT("AmbientOcclusionRGS"), SF_RayGen)
 
-void FDeferredShadingSceneRenderer::RenderRayTracingAmbientOcclusion(
-	FRHICommandListImmediate& RHICmdList,
-	const FLightSceneInfo* SkyLightSceneInfo,
-	TRefCountPtr<IPooledRenderTarget>& AmbientOcclusionMask,
-	TRefCountPtr<IPooledRenderTarget>& HitDistance
-)
+void FDeferredShadingSceneRenderer::RenderRayTracingAmbientOcclusion(FRHICommandListImmediate& RHICmdList, TRefCountPtr<IPooledRenderTarget>& AmbientOcclusionMaskRT)
 {
 	SCOPED_DRAW_EVENT(RHICmdList, RayTracingAmbientOcclusion);
 	SCOPED_GPU_STAT(RHICmdList, RayTracingAmbientOcclusion);
 
+	bool bAnyViewWithRTAO = false;
+	for (int32 ViewIndex = 0; ViewIndex < Views.Num(); ++ViewIndex)
+	{
+		FViewInfo& View = Views[ViewIndex];
+		bAnyViewWithRTAO = bAnyViewWithRTAO | ShouldRenderRayTracingAmbientOcclusion(View);
+	}
+
+	if (!bAnyViewWithRTAO)
+	{
+		return;
+	}
+
+	FRDGBuilder GraphBuilder(RHICmdList);
+
 	FSceneRenderTargets& SceneContext = FSceneRenderTargets::Get(RHICmdList);
 	FPooledRenderTargetDesc Desc = SceneContext.GetSceneColor()->GetDesc();
 	Desc.Format = PF_R16F;
 	Desc.Flags &= ~(TexCreate_FastVRAM | TexCreate_Transient);
-	GRenderTargetPool.FindFreeElement(RHICmdList, Desc, AmbientOcclusionMask, TEXT("RayTracingAmbientOcclusion"));
-	GRenderTargetPool.FindFreeElement(RHICmdList, Desc, HitDistance, TEXT("RayTracingAmbientOcclusionHitDistance"));
+
+	TRefCountPtr<IPooledRenderTarget> HitDistanceRT;
+
+	GRenderTargetPool.FindFreeElement(RHICmdList, Desc, AmbientOcclusionMaskRT, TEXT("RayTracingAmbientOcclusion"));
+	GRenderTargetPool.FindFreeElement(RHICmdList, Desc, HitDistanceRT, TEXT("RayTracingAmbientOcclusionHitDistance"));
+
+	for (int32 ViewIndex = 0; ViewIndex < Views.Num(); ++ViewIndex)
+	{	
+		FViewInfo& View = Views[ViewIndex];
+		if (ShouldRenderRayTracingAmbientOcclusion(View))
+		{
+			RenderRayTracingAmbientOcclusion(RHICmdList, GraphBuilder, View, AmbientOcclusionMaskRT, HitDistanceRT);
+		}
+	}
+
+	GraphBuilder.Execute();
+
+	SceneContext.bScreenSpaceAOIsValid = true;
+}
+
+void FDeferredShadingSceneRenderer::RenderRayTracingAmbientOcclusion(
+	FRHICommandListImmediate& RHICmdList, 
+	FRDGBuilder& GraphBuilder, 
+	FViewInfo& View,
+	TRefCountPtr<IPooledRenderTarget>& AmbientOcclusionMaskRT,
+	TRefCountPtr<IPooledRenderTarget>& HitDistanceRT
+)
+{
+	//#dxr_todo refactor to RDG
+	
+	FIntPoint ViewSize = View.ViewRect.Size();
+	FSceneRenderTargets& SceneContext = FSceneRenderTargets::Get(RHICmdList);
 
 	// Add ambient occlusion parameters to uniform buffer
 	FAmbientOcclusionData AmbientOcclusionData;
 
-	for (int32 ViewIndex = 0; ViewIndex < Views.Num(); ++ViewIndex)
+	FSceneTexturesUniformParameters SceneTextures;
+	SetupSceneTextureUniformParameters(SceneContext, FeatureLevel, ESceneTextureSetupMode::All, SceneTextures);
+	FUniformBufferRHIRef SceneTexturesUniformBuffer = RHICreateUniformBuffer(&SceneTextures, FSceneTexturesUniformParameters::StaticStructMetadata.GetLayout(), EUniformBufferUsage::UniformBuffer_SingleDraw);
+
+	AmbientOcclusionData.SamplesPerPixel = GRayTracingAmbientOcclusionSamplesPerPixel >= 0 ? GRayTracingAmbientOcclusionSamplesPerPixel : View.FinalPostProcessSettings.RayTracingAOSamplesPerPixel;
+	AmbientOcclusionData.MaxRayDistance = View.FinalPostProcessSettings.AmbientOcclusionRadius;
+	AmbientOcclusionData.Intensity = View.FinalPostProcessSettings.AmbientOcclusionIntensity;
+	AmbientOcclusionData.MaxNormalBias = GetRaytracingMaxNormalBias();
+	FUniformBufferRHIRef AmbientOcclusionUniformBuffer = RHICreateUniformBuffer(&AmbientOcclusionData, FAmbientOcclusionData::StaticStructMetadata.GetLayout(), EUniformBufferUsage::UniformBuffer_SingleDraw);
+
+	int32 EnableTwoSidedGeometry = CVarRayTracingAmbientOcclusionEnableTwoSidedGeometry.GetValueOnRenderThread();
+	if (EnableTwoSidedGeometry)
 	{
-		FViewInfo& View = Views[ViewIndex];
-		FIntPoint ViewSize = View.ViewRect.Size();
+		TShaderMapRef<TAmbientOcclusionRGS<1>> AmbientOcclusionRayGenerationShader(GetGlobalShaderMap(FeatureLevel));
+		AmbientOcclusionRayGenerationShader->Dispatch(
+			RHICmdList,
+			View.RayTracingScene,
+			View.ViewUniformBuffer,
+			SceneTexturesUniformBuffer,
+			AmbientOcclusionUniformBuffer,
+			AmbientOcclusionMaskRT->GetRenderTargetItem().UAV,
+			HitDistanceRT->GetRenderTargetItem().UAV,
+			ViewSize.X, ViewSize.Y
+		);
+	}
+	else
+	{
+		TShaderMapRef<TAmbientOcclusionRGS<0>> AmbientOcclusionRayGenerationShader(GetGlobalShaderMap(FeatureLevel));
+		AmbientOcclusionRayGenerationShader->Dispatch(
+			RHICmdList,
+			View.RayTracingScene,
+			View.ViewUniformBuffer,
+			SceneTexturesUniformBuffer,
+			AmbientOcclusionUniformBuffer,
+			AmbientOcclusionMaskRT->GetRenderTargetItem().UAV,
+			HitDistanceRT->GetRenderTargetItem().UAV,
+			ViewSize.X, ViewSize.Y
+		);
+	}
+
+	int32 DenoiserMode = CVarUseAODenoiser.GetValueOnRenderThread();
+	if (DenoiserMode != 0)
+	{
+		SCOPED_GPU_STAT(RHICmdList, AmbientOcclusionDenoiser);
 
-		FSceneTexturesUniformParameters SceneTextures;
-		SetupSceneTextureUniformParameters(SceneContext, FeatureLevel, ESceneTextureSetupMode::All, SceneTextures);
-		FUniformBufferRHIRef SceneTexturesUniformBuffer = RHICreateUniformBuffer(&SceneTextures, FSceneTexturesUniformParameters::StaticStructMetadata.GetLayout(), EUniformBufferUsage::UniformBuffer_SingleDraw);
+		FSceneViewFamilyBlackboard SceneBlackboard;
+		SetupSceneViewFamilyBlackboard(GraphBuilder, &SceneBlackboard);
 
-		AmbientOcclusionData.SamplesPerPixel = GRayTracingAmbientOcclusionSamplesPerPixel >= 0 ? GRayTracingAmbientOcclusionSamplesPerPixel : View.FinalPostProcessSettings.RayTracingAOSamplesPerPixel;
-		AmbientOcclusionData.MaxRayDistance = View.FinalPostProcessSettings.AmbientOcclusionRadius;
-		AmbientOcclusionData.Intensity = View.FinalPostProcessSettings.AmbientOcclusionIntensity;
-		AmbientOcclusionData.MaxNormalBias = GetRaytracingMaxNormalBias();
-		FUniformBufferRHIRef AmbientOcclusionUniformBuffer = RHICreateUniformBuffer(&AmbientOcclusionData, FAmbientOcclusionData::StaticStructMetadata.GetLayout(), EUniformBufferUsage::UniformBuffer_SingleDraw);
+		const IScreenSpaceDenoiser* DefaultDenoiser = IScreenSpaceDenoiser::GetDefaultDenoiser();
+		const IScreenSpaceDenoiser* DenoiserToUse = DenoiserMode == 1 ? DefaultDenoiser : GScreenSpaceDenoiser;
+
+		IScreenSpaceDenoiser::FAmbientOcclusionRayTracingConfig RayTracingConfig;
+
+		IScreenSpaceDenoiser::FAmbientOcclusionInputs DenoiserInputs;
+		DenoiserInputs.Mask = GraphBuilder.RegisterExternalTexture(SceneContext.ScreenSpaceAO, TEXT("AOMask"));
+		DenoiserInputs.RayHitDistance = GraphBuilder.RegisterExternalTexture(HitDistanceRT, TEXT("AOHitDistance"));
 
-		int32 EnableTwoSidedGeometry = CVarRayTracingAmbientOcclusionEnableTwoSidedGeometry.GetValueOnRenderThread();
-		if (EnableTwoSidedGeometry)
-		{
-			TShaderMapRef<TAmbientOcclusionRGS<1>> AmbientOcclusionRayGenerationShader(GetGlobalShaderMap(FeatureLevel));
-			AmbientOcclusionRayGenerationShader->Dispatch(
-				RHICmdList,
-				View.RayTracingScene,
-				View.ViewUniformBuffer,
-				SceneTexturesUniformBuffer,
-				AmbientOcclusionUniformBuffer,
-				AmbientOcclusionMask->GetRenderTargetItem().UAV,
-				HitDistance->GetRenderTargetItem().UAV,
-				ViewSize.X, ViewSize.Y
-			);
-		}
-		else
 		{
-			TShaderMapRef<TAmbientOcclusionRGS<0>> AmbientOcclusionRayGenerationShader(GetGlobalShaderMap(FeatureLevel));
-			AmbientOcclusionRayGenerationShader->Dispatch(
-				RHICmdList,
-				View.RayTracingScene,
-				View.ViewUniformBuffer,
-				SceneTexturesUniformBuffer,
-				AmbientOcclusionUniformBuffer,
-				AmbientOcclusionMask->GetRenderTargetItem().UAV,
-				HitDistance->GetRenderTargetItem().UAV,
-				ViewSize.X, ViewSize.Y
-			);
+			RDG_EVENT_SCOPE(GraphBuilder, "%s%s(AmbientOcclusion) %dx%d",
+				DenoiserToUse != DefaultDenoiser ? TEXT("ThirdParty ") : TEXT(""),
+				DenoiserToUse->GetDebugName(),
+				View.ViewRect.Width(), View.ViewRect.Height());
+
+			IScreenSpaceDenoiser::FAmbientOcclusionOutputs DenoiserOutputs = DenoiserToUse->DenoiseAmbientOcclusion(
+				GraphBuilder,
+				View,
+				&View.PrevViewInfo,
+				SceneBlackboard,
+				DenoiserInputs,
+				RayTracingConfig);
+
+			GraphBuilder.QueueTextureExtraction(DenoiserOutputs.AmbientOcclusionMask, &SceneContext.ScreenSpaceAO);
 		}
 	}
 
 	FUnorderedAccessViewRHIParamRef UAVs[]
 	{
-		AmbientOcclusionMask->GetRenderTargetItem().UAV,
-		HitDistance->GetRenderTargetItem().UAV
+		AmbientOcclusionMaskRT->GetRenderTargetItem().UAV,
+		HitDistanceRT->GetRenderTargetItem().UAV
 	};
 	FComputeFenceRHIRef Fence = RHICmdList.CreateComputeFence(TEXT("RayTracingAmbientOcclusion"));
 	RHICmdList.TransitionResources(EResourceTransitionAccess::ERWBarrier, EResourceTransitionPipeline::EComputeToGfx, UAVs, ARRAY_COUNT(UAVs), Fence);
-
-	GVisualizeTexture.SetCheckPoint(RHICmdList, AmbientOcclusionMask);
-	GVisualizeTexture.SetCheckPoint(RHICmdList, HitDistance);
-	SceneContext.bScreenSpaceAOIsValid = true;
 }
 
 #endif // RHI_RAYTRACING
diff --git a/Engine/Source/Runtime/Renderer/Private/RayTracing/RayTracingGlobalIllumination.cpp b/Engine/Source/Runtime/Renderer/Private/RayTracing/RayTracingGlobalIllumination.cpp
--- a/Engine/Source/Runtime/Renderer/Private/RayTracing/RayTracingGlobalIllumination.cpp
+++ b/Engine/Source/Runtime/Renderer/Private/RayTracing/RayTracingGlobalIllumination.cpp
@@ -178,30 +178,19 @@ void SetupLightParameters(
 	}
 }
 
-bool ShouldRenderRayTracingGlobalIllumination(const TArray<FViewInfo>& Views)
+bool ShouldRenderRayTracingGlobalIllumination(const FViewInfo& View)
 {
 	if (GetForceRayTracingEffectsCVarValue() >= 0)
 	{
 		return GetForceRayTracingEffectsCVarValue() > 0;
 	}
-
-	if (GRayTracingGlobalIllumination >= 0)
+	else if (GRayTracingGlobalIllumination >= 0)
 	{
 		return (GRayTracingGlobalIllumination > 0);
 	}
-	else
+	else 
 	{
-		for (int32 ViewIndex = 0, Num = Views.Num(); ViewIndex < Num; ViewIndex++)
-		{
-			const FViewInfo& View = Views[ViewIndex];
-			//#dxr_todo: UE-72557 multiview case
-			if (View.FinalPostProcessSettings.RayTracingGI > 0)
-			{
-				return true;
-			}
-		}
-
-		return false;
+		return View.FinalPostProcessSettings.RayTracingGI > 0;
 	}
 }
 
@@ -328,23 +317,23 @@ void FDeferredShadingSceneRenderer::PrepareRayTracingGlobalIllumination(const FV
 }
 
 void FDeferredShadingSceneRenderer::RenderRayTracingGlobalIllumination(
-	FRHICommandListImmediate& RHICmdList,
-	FViewInfo& View,
-	TRefCountPtr<IPooledRenderTarget>& GlobalIlluminationRT,
-	TRefCountPtr<IPooledRenderTarget>& AmbientOcclusionRT
+	FRHICommandListImmediate& RHICmdList, 
+	TRefCountPtr<IPooledRenderTarget>& GlobalIlluminationRT
 )
 {
 	SCOPED_GPU_STAT(RHICmdList, RayTracingGlobalIllumination);
 
-	FSceneRenderTargets& SceneContext = FSceneRenderTargets::Get(RHICmdList);
+	bool bAnyViewWithRTGI = false;
+	for (int32 ViewIndex = 0; ViewIndex < Views.Num(); ++ViewIndex)
 	{
-		FPooledRenderTargetDesc Desc = SceneContext.GetSceneColor()->GetDesc();
-		Desc.Format = PF_FloatRGBA;
-		Desc.Flags &= ~(TexCreate_FastVRAM | TexCreate_Transient);
-		GRenderTargetPool.FindFreeElement(RHICmdList, Desc, GlobalIlluminationRT, TEXT("RayTracingGlobalIllumination"));
+		FViewInfo& View = Views[ViewIndex];
+		bAnyViewWithRTGI = bAnyViewWithRTGI | ShouldRenderRayTracingGlobalIllumination(View);
 	}
 
-	FRDGBuilder GraphBuilder(RHICmdList);
+	if (!bAnyViewWithRTGI)
+	{
+		return;
+	}
 
 	IScreenSpaceDenoiser::FAmbientOcclusionRayTracingConfig RayTracingConfig;
 	RayTracingConfig.ResolutionFraction = 1.0;
@@ -353,11 +342,18 @@ void FDeferredShadingSceneRenderer::RenderRayTracingGlobalIllumination(
 		RayTracingConfig.ResolutionFraction = FMath::Clamp(GRayTracingGlobalIlluminationScreenPercentage / 100.0, 0.25, 1.0);
 	}
 
-	int32 RayTracingGISamplesPerPixel = GRayTracingGlobalIlluminationSamplesPerPixel > -1 ? GRayTracingGlobalIlluminationSamplesPerPixel : View.FinalPostProcessSettings.RayTracingGISamplesPerPixel;
-	RayTracingConfig.RayCountPerPixel = RayTracingGISamplesPerPixel;
 	int32 UpscaleFactor = int32(1.0 / RayTracingConfig.ResolutionFraction);
 
-	// Render targets
+	FSceneRenderTargets& SceneContext = FSceneRenderTargets::Get(RHICmdList);
+	{
+		FPooledRenderTargetDesc Desc = SceneContext.GetSceneColor()->GetDesc();
+		Desc.Format = PF_FloatRGBA;
+		Desc.Flags &= ~(TexCreate_FastVRAM | TexCreate_Transient);
+		GRenderTargetPool.FindFreeElement(RHICmdList, Desc, GlobalIlluminationRT, TEXT("RayTracingGlobalIllumination"));
+	}
+
+	FRDGBuilder GraphBuilder(RHICmdList);
+
 	FRDGTextureRef GlobalIlluminationTexture;
 	{
 		FRDGTextureDesc Desc = SceneContext.GetSceneColor()->GetDesc();
@@ -366,6 +362,7 @@ void FDeferredShadingSceneRenderer::RenderRayTracingGlobalIllumination(
 		Desc.Flags &= ~(TexCreate_FastVRAM | TexCreate_Transient);
 		GlobalIlluminationTexture = GraphBuilder.CreateTexture(Desc, TEXT("RayTracingGlobalIllumination"));
 	}
+	FRDGTextureUAV* GlobalIlluminationUAV = GraphBuilder.CreateUAV(FRDGTextureUAVDesc(GlobalIlluminationTexture));
 
 	FRDGTextureRef RayDistanceTexture;
 	{
@@ -375,8 +372,55 @@ void FDeferredShadingSceneRenderer::RenderRayTracingGlobalIllumination(
 		Desc.Flags &= ~(TexCreate_FastVRAM | TexCreate_Transient);
 		RayDistanceTexture = GraphBuilder.CreateTexture(Desc, TEXT("RayTracingGlobalIlluminationRayDistance"));
 	}
-	FRDGTextureRef ResultTexture;
-	
+	FRDGTextureUAV* RayDistanceUAV = GraphBuilder.CreateUAV(FRDGTextureUAVDesc(RayDistanceTexture));
+
+	TRefCountPtr<IPooledRenderTarget>& AmbientOcclusionRT = SceneContext.ScreenSpaceAO;
+
+	for (int32 ViewIndex = 0; ViewIndex < Views.Num(); ++ViewIndex)
+	{
+		FViewInfo& View = Views[ViewIndex];
+		if (ShouldRenderRayTracingGlobalIllumination(View))
+		{
+			RenderRayTracingGlobalIllumination(
+				RHICmdList, 
+				GraphBuilder, 
+				View, 
+				RayTracingConfig, 
+				UpscaleFactor, 
+				GlobalIlluminationRT, 
+				AmbientOcclusionRT, 
+				GlobalIlluminationUAV, 
+				RayDistanceUAV, 
+				GlobalIlluminationTexture, 
+				RayDistanceTexture
+			);
+		}
+	}
+
+	GraphBuilder.Execute();
+	SceneContext.bScreenSpaceAOIsValid = true;
+	GVisualizeTexture.SetCheckPoint(RHICmdList, GlobalIlluminationRT);
+}
+
+void FDeferredShadingSceneRenderer::RenderRayTracingGlobalIllumination(
+	FRHICommandListImmediate& RHICmdList,
+	FRDGBuilder& GraphBuilder,
+	FViewInfo& View,
+	IScreenSpaceDenoiser::FAmbientOcclusionRayTracingConfig& RayTracingConfig,
+	int32 UpscaleFactor,
+	TRefCountPtr<IPooledRenderTarget>& GlobalIlluminationRT,
+	TRefCountPtr<IPooledRenderTarget>& AmbientOcclusionRT,
+	FRDGTextureUAV* GlobalIlluminationUAV,
+	FRDGTextureUAV* RayDistanceUAV,
+	const FRDGTextureRef GlobalIlluminationTexture,
+	const FRDGTextureRef RayDistanceTexture
+)
+{	
+	FSceneRenderTargets& SceneContext = FSceneRenderTargets::Get(RHICmdList);
+
+	int32 RayTracingGISamplesPerPixel = GRayTracingGlobalIlluminationSamplesPerPixel > -1 ? GRayTracingGlobalIlluminationSamplesPerPixel : View.FinalPostProcessSettings.RayTracingGISamplesPerPixel;
+	RayTracingConfig.RayCountPerPixel = RayTracingGISamplesPerPixel;
+
 	FSceneTexturesUniformParameters SceneTextures;
 	SetupSceneTextureUniformParameters(SceneContext, FeatureLevel, ESceneTextureSetupMode::All, SceneTextures);
 	// Ray generation
@@ -438,8 +482,8 @@ void FDeferredShadingSceneRenderer::RenderRayTracingGlobalIllumination(
 		}
 		PassParameters->SSProfilesTexture = GraphBuilder.RegisterExternalTexture(SubsurfaceProfileRT);
 		PassParameters->TransmissionProfilesLinearSampler = TStaticSamplerState<SF_Bilinear, AM_Clamp, AM_Clamp, AM_Clamp>::GetRHI();
-		PassParameters->RWGlobalIlluminationUAV = GraphBuilder.CreateUAV(GlobalIlluminationTexture);
-		PassParameters->RWRayDistanceUAV = GraphBuilder.CreateUAV(RayDistanceTexture);
+		PassParameters->RWGlobalIlluminationUAV = GlobalIlluminationUAV;
+		PassParameters->RWRayDistanceUAV = RayDistanceUAV;
 
 		FGlobalIlluminationRGS::FPermutationDomain PermutationVector;
 		PermutationVector.Set<FGlobalIlluminationRGS::FUseAttenuationTermDim>(true);
@@ -462,6 +506,9 @@ void FDeferredShadingSceneRenderer::RenderRayTracingGlobalIllumination(
 	}
 
 	// Denoising
+
+	FRDGTextureRef ResultTexture; //#dxr_todo review
+
 	if (GRayTracingGlobalIlluminationDenoiser != 0)
 	{
 		FSceneViewFamilyBlackboard SceneBlackboard;
@@ -546,10 +593,6 @@ void FDeferredShadingSceneRenderer::RenderRayTracingGlobalIllumination(
 			}
 		);
 	}
-
-	GraphBuilder.Execute();
-	SceneContext.bScreenSpaceAOIsValid = true;
-	GVisualizeTexture.SetCheckPoint(RHICmdList, GlobalIlluminationRT);
 }
 
 void FDeferredShadingSceneRenderer::CompositeGlobalIllumination(
@@ -558,7 +601,6 @@ void FDeferredShadingSceneRenderer::CompositeGlobalIllumination(
 	TRefCountPtr<IPooledRenderTarget>& GlobalIlluminationRT
 )
 {
-
 	FSceneRenderTargets& SceneContext = FSceneRenderTargets::Get(RHICmdList);
 
 	FSceneTexturesUniformParameters SceneTextures;
diff --git a/Engine/Source/Runtime/Renderer/Private/RayTracing/RayTracingShadows.cpp b/Engine/Source/Runtime/Renderer/Private/RayTracing/RayTracingShadows.cpp
--- a/Engine/Source/Runtime/Renderer/Private/RayTracing/RayTracingShadows.cpp
+++ b/Engine/Source/Runtime/Renderer/Private/RayTracing/RayTracingShadows.cpp
@@ -116,39 +116,14 @@ void FDeferredShadingSceneRenderer::RenderRayTracingShadows(
 	const FLightSceneInfo& LightSceneInfo,
 	const IScreenSpaceDenoiser::FShadowRayTracingConfig& RayTracingConfig,
 	IScreenSpaceDenoiser::EShadowRequirements DenoiserRequirements,
-	FRDGTextureRef* OutShadowMask,
-	FRDGTextureRef* OutRayHitDistance)
+	FRDGTextureUAV* OutShadowMaskUAV,
+	FRDGTextureUAV* OutRayHitDistanceUAV)
 #if RHI_RAYTRACING
 {
 	FLightSceneProxy* LightSceneProxy = LightSceneInfo.Proxy;
 	check(LightSceneProxy);
 
-	// Render targets
-	FRDGTextureRef ScreenShadowMaskTexture;
-	{
-		FRDGTextureDesc Desc = FRDGTextureDesc::Create2DDesc(
-			SceneBlackboard.SceneDepthBuffer->Desc.Extent,
-			PF_FloatRGBA,
-			FClearValueBinding::Black,
-			TexCreate_None,
-			TexCreate_ShaderResource | TexCreate_RenderTargetable | TexCreate_UAV,
-			/* bInForceSeparateTargetAndShaderResource = */ false);
-		ScreenShadowMaskTexture = GraphBuilder.CreateTexture(Desc, TEXT("RayTracingOcclusion"));
-	}
-
-	FRDGTextureRef RayDistanceTexture;
-	{
-		FRDGTextureDesc Desc = FRDGTextureDesc::Create2DDesc(
-			SceneBlackboard.SceneDepthBuffer->Desc.Extent,
-			PF_R16F,
-			FClearValueBinding::Black,
-			TexCreate_None,
-			TexCreate_ShaderResource | TexCreate_RenderTargetable | TexCreate_UAV,
-			/* bInForceSeparateTargetAndShaderResource = */ false);
-		RayDistanceTexture = GraphBuilder.CreateTexture(Desc, TEXT("RayTracingOcclusionDistance"));
-	}
-
-	FIntRect ScissorRect = { {0,0}, View.ViewRect.Size() };
+	FIntRect ScissorRect = { View.ViewRect.Min, View.ViewRect.Min + View.ViewRect.Size() };
 
 	if (LightSceneProxy->GetScissorRect(ScissorRect, View, View.ViewRect))
 	{
@@ -156,16 +131,12 @@ void FDeferredShadingSceneRenderer::RenderRayTracingShadows(
 		ScissorRect.Min = ScissorRect.Min - View.ViewRect.Min;
 		ScissorRect.Max = ScissorRect.Max - View.ViewRect.Min;
 	}
-	else
-	{
-		ScissorRect = { {0,0}, View.ViewRect.Size() };
-	}
 
 	// Ray generation pass for shadow occlusion.
 	{
 		FOcclusionRGS::FParameters* PassParameters = GraphBuilder.AllocParameters<FOcclusionRGS::FParameters>();
-		PassParameters->RWOcclusionMaskUAV = GraphBuilder.CreateUAV(FRDGTextureUAVDesc(ScreenShadowMaskTexture));
-		PassParameters->RWRayDistanceUAV = GraphBuilder.CreateUAV(FRDGTextureUAVDesc(RayDistanceTexture));
+		PassParameters->RWOcclusionMaskUAV = OutShadowMaskUAV;
+		PassParameters->RWRayDistanceUAV = OutRayHitDistanceUAV;
 		PassParameters->SamplesPerPixel = RayTracingConfig.RayCountPerPixel;
 		PassParameters->NormalBias = GetRaytracingMaxNormalBias();
 		PassParameters->LightingChannelMask = LightSceneProxy->GetLightingChannelMask();
@@ -239,9 +210,6 @@ void FDeferredShadingSceneRenderer::RenderRayTracingShadows(
 			}
 		});
 	}
-
-	*OutShadowMask = ScreenShadowMaskTexture;
-	*OutRayHitDistance = RayDistanceTexture;
 }
 #else // !RHI_RAYTRACING
 {
diff --git a/Engine/Source/Runtime/Renderer/Private/RayTracing/RayTracingTranslucency.cpp b/Engine/Source/Runtime/Renderer/Private/RayTracing/RayTracingTranslucency.cpp
--- a/Engine/Source/Runtime/Renderer/Private/RayTracing/RayTracingTranslucency.cpp
+++ b/Engine/Source/Runtime/Renderer/Private/RayTracing/RayTracingTranslucency.cpp
@@ -105,18 +105,12 @@ static FAutoConsoleVariableRef CVarRayTracingTranslucencyRefraction(
 DECLARE_GPU_STAT_NAMED(RayTracingTranslucency, TEXT("Ray Tracing Translucency"));
 
 #if RHI_RAYTRACING
-bool ShouldRenderRayTracingTranslucency(const TArray<FViewInfo>& Views)
+bool ShouldRenderRayTracingTranslucency(const FViewInfo& View)
 {
-	bool bAnyViewWithRaytracingTranslucency = false;
-	for (int32 ViewIndex = 0, Num = Views.Num(); ViewIndex < Num; ViewIndex++)
-	{
-		const FViewInfo& View = Views[ViewIndex];
-		//#dxr_todo: UE-72557 multiview case
-		bAnyViewWithRaytracingTranslucency = bAnyViewWithRaytracingTranslucency || (View.FinalPostProcessSettings.TranslucencyType == ETranslucencyType::RayTracing);
-	}
+	bool bViewWithRaytracingTranslucency = View.FinalPostProcessSettings.TranslucencyType == ETranslucencyType::RayTracing;
 
 	const int32 GRayTracingTranslucency = CVarRayTracingTranslucency.GetValueOnRenderThread();
-	const bool bTranslucencyCvarEnabled = GRayTracingTranslucency < 0 ? bAnyViewWithRaytracingTranslucency : GRayTracingTranslucency;
+	const bool bTranslucencyCvarEnabled = GRayTracingTranslucency < 0 ? bViewWithRaytracingTranslucency : GRayTracingTranslucency;
 	const int32 ForceAllRayTracingEffects = GetForceRayTracingEffectsCVarValue();
 	const bool bRayTracingTranslucencyEnabled = (ForceAllRayTracingEffects > 0 || (bTranslucencyCvarEnabled && ForceAllRayTracingEffects < 0));
 
diff --git a/Engine/Source/Runtime/Renderer/Private/RayTracing/RaytracingOptions.h b/Engine/Source/Runtime/Renderer/Private/RayTracing/RaytracingOptions.h
--- a/Engine/Source/Runtime/Renderer/Private/RayTracing/RaytracingOptions.h
+++ b/Engine/Source/Runtime/Renderer/Private/RayTracing/RaytracingOptions.h
@@ -18,10 +18,10 @@ class FLightSceneInfo;
 extern int32 GetForceRayTracingEffectsCVarValue();
 
 extern bool ShouldRenderRayTracingSkyLight(const FSkyLightSceneProxy* SkyLightSceneProxy);
-extern bool ShouldRenderRayTracingAmbientOcclusion();
-extern bool ShouldRenderRayTracingReflections(const TArray<FViewInfo>& Views);
-extern bool ShouldRenderRayTracingGlobalIllumination(const TArray<FViewInfo>& Views);
-extern bool ShouldRenderRayTracingTranslucency(const TArray<FViewInfo>& Views);
+extern bool ShouldRenderRayTracingAmbientOcclusion(const FViewInfo& View);
+extern bool ShouldRenderRayTracingReflections(const FViewInfo& View);
+extern bool ShouldRenderRayTracingGlobalIllumination(const FViewInfo& View);
+extern bool ShouldRenderRayTracingTranslucency(const FViewInfo& View);
 extern bool ShouldRenderRayTracingShadows(const FLightSceneProxy& LightProxy);
 extern bool ShouldRenderRayTracingShadows(const FLightSceneInfoCompact& LightInfo);
 extern bool ShouldRenderRayTracingStochasticRectLight(const FLightSceneInfo& LightInfo);
@@ -41,22 +41,22 @@ FORCEINLINE bool ShouldRenderRayTracingSkyLight(const FSkyLightSceneProxy* SkyLi
 	return false;
 }
 
-FORCEINLINE bool ShouldRenderRayTracingAmbientOcclusion()
+FORCEINLINE bool ShouldRenderRayTracingAmbientOcclusion(const FViewInfo& View)
 {
 	return false;
 }
 
-FORCEINLINE bool ShouldRenderRayTracingReflections(const TArray<FViewInfo>& Views)
+FORCEINLINE bool ShouldRenderRayTracingReflections(const FViewInfo& View)
 {
 	return false;
 }
 
-FORCEINLINE bool ShouldRenderRayTracingGlobalIllumination(const TArray<FViewInfo>& Views)
+FORCEINLINE bool ShouldRenderRayTracingGlobalIllumination(const FViewInfo& View)
 {
 	return false;
 }
 
-FORCEINLINE bool ShouldRenderRayTracingTranslucency(const TArray<FViewInfo>& Views)
+FORCEINLINE bool ShouldRenderRayTracingTranslucency(const FViewInfo& View)
 {
 	return false;
 }
diff --git a/Engine/Source/Runtime/Renderer/Private/RayTracing/RaytracingSkylight.cpp b/Engine/Source/Runtime/Renderer/Private/RayTracing/RaytracingSkylight.cpp
--- a/Engine/Source/Runtime/Renderer/Private/RayTracing/RaytracingSkylight.cpp
+++ b/Engine/Source/Runtime/Renderer/Private/RayTracing/RaytracingSkylight.cpp
@@ -755,12 +755,16 @@ void FDeferredShadingSceneRenderer::RenderRayTracingSkyLight(
 	TRefCountPtr<IPooledRenderTarget>& HitDistanceRT
 )
 {
-	check(Scene->SkyLight);
-	check(Scene->SkyLight->ProcessedTexture);
-
 	SCOPED_DRAW_EVENT(RHICmdList, RayTracingSkyLight);
 	SCOPED_GPU_STAT(RHICmdList, RayTracingSkyLight);
 
+	if (!ShouldRenderRayTracingSkyLight(Scene->SkyLight))
+	{
+		return;
+	}
+
+	check(Scene->SkyLight->ProcessedTexture);
+
 	if (Scene->SkyLight->ShouldRebuildCdf())
 	{
 		BuildSkyLightCdfs(RHICmdList, Scene->SkyLight);
diff --git a/Engine/Source/Runtime/Renderer/Private/ReflectionEnvironment.cpp b/Engine/Source/Runtime/Renderer/Private/ReflectionEnvironment.cpp
--- a/Engine/Source/Runtime/Renderer/Private/ReflectionEnvironment.cpp
+++ b/Engine/Source/Runtime/Renderer/Private/ReflectionEnvironment.cpp
@@ -191,17 +191,11 @@ bool IsReflectionCaptureAvailable()
 }
 
 #if RHI_RAYTRACING
-bool ShouldRenderRayTracingReflections(const TArray<FViewInfo>& Views)
+bool ShouldRenderRayTracingReflections(const FViewInfo& View)
 {
-	bool bAnyViewWithRaytracingReflections = false;
-	for (int32 ViewIndex = 0, Num = Views.Num(); ViewIndex < Num; ViewIndex++)
-	{
-		const FViewInfo& View = Views[ViewIndex];
-		//#dxr_todo: UE-72557 multiview case
-		bAnyViewWithRaytracingReflections = bAnyViewWithRaytracingReflections || (View.FinalPostProcessSettings.ReflectionsType == EReflectionsType::RayTracing);
-	}
+	bool bThisViewHasRaytracingReflections = View.FinalPostProcessSettings.ReflectionsType == EReflectionsType::RayTracing;
 
-	const bool bReflectionsCvarEnabled = GRayTracingReflections < 0 ? bAnyViewWithRaytracingReflections : GRayTracingReflections;	
+	const bool bReflectionsCvarEnabled = GRayTracingReflections < 0 ? bThisViewHasRaytracingReflections : GRayTracingReflections;
 	const int32 ForceAllRayTracingEffects = GetForceRayTracingEffectsCVarValue();
 	const bool bReflectionPassEnabled = (ForceAllRayTracingEffects > 0 || (bReflectionsCvarEnabled && ForceAllRayTracingEffects < 0));
 
@@ -735,8 +729,6 @@ void FDeferredShadingSceneRenderer::RenderDeferredReflectionsAndSkyLighting(FRHI
 		return;	
 	}
 
-	const bool bRayTracedReflections = ShouldRenderRayTracingReflections(Views);
-
 	// The specular sky light contribution is also needed by RT Reflections as a fallback.
 	const bool bSkyLight = Scene->SkyLight
 		&& Scene->SkyLight->ProcessedTexture
@@ -771,6 +763,8 @@ void FDeferredShadingSceneRenderer::RenderDeferredReflectionsAndSkyLighting(FRHI
 	{
 		FViewInfo& View = Views[ViewIndex];
 
+		const bool bRayTracedReflections = ShouldRenderRayTracingReflections(View);
+
 		const bool bScreenSpaceReflections = !bRayTracedReflections && ShouldRenderScreenSpaceReflections(Views[ViewIndex]);
 
 		TRefCountPtr<IPooledRenderTarget> ReflectionsColor = GSystemTextures.BlackDummy;
