Adding optional any-hit material support to SkyLight pass.

r.RayTracing.SkyLight.EnableMaterials (default = 0) allows support for visibility tracing against masked materials.

diff --git a/Engine/Shaders/Private/RayTracing/RayTracingSkyLightRGS.usf b/Engine/Shaders/Private/RayTracing/RayTracingSkyLightRGS.usf
--- a/Engine/Shaders/Private/RayTracing/RayTracingSkyLightRGS.usf
+++ b/Engine/Shaders/Private/RayTracing/RayTracingSkyLightRGS.usf
@@ -302,12 +302,14 @@ void SkyLightRGS()
 			continue;
 		}
 
-		uint RayFlags = 0
-			| RAY_FLAG_FORCE_OPAQUE; // #dxr_todo: UE-72563 generate a special shader permutation for occlusion ray tracing that includes any-hit shaders for masked geometry
+		uint RayFlags = 0;
+#if !ENABLE_MATERIALS
+		RayFlags |= RAY_FLAG_FORCE_OPAQUE;
+#endif
 #if !ENABLE_TWO_SIDED_GEOMETRY
 		RayFlags |= RAY_FLAG_CULL_BACK_FACING_TRIANGLES;
 #endif
-		FDefaultPayload Payload = (FDefaultPayload)0;
+		FPackedMaterialClosestHitPayload Payload = (FPackedMaterialClosestHitPayload)0;
 
 		TraceRay(
 			TLAS,      // AccelerationStructure
diff --git a/Engine/Source/Runtime/Renderer/Private/DeferredShadingRenderer.cpp b/Engine/Source/Runtime/Renderer/Private/DeferredShadingRenderer.cpp
--- a/Engine/Source/Runtime/Renderer/Private/DeferredShadingRenderer.cpp
+++ b/Engine/Source/Runtime/Renderer/Private/DeferredShadingRenderer.cpp
@@ -763,6 +763,7 @@ bool FDeferredShadingSceneRenderer::DispatchRayTracingWorldUpdates(FRHICommandLi
 		TArray<FRayTracingShaderRHIParamRef> RayGenShaders;
 		PrepareRayTracingReflections(View, RayGenShaders);
 		PrepareRayTracingShadows(View, RayGenShaders);
+		PrepareRayTracingSkyLight(View, RayGenShaders);
 		PrepareRayTracingRectLight(View, RayGenShaders);
 		PrepareRayTracingGlobalIllumination(View, RayGenShaders);
 		PrepareRayTracingTranslucency(View, RayGenShaders);
diff --git a/Engine/Source/Runtime/Renderer/Private/DeferredShadingRenderer.h b/Engine/Source/Runtime/Renderer/Private/DeferredShadingRenderer.h
--- a/Engine/Source/Runtime/Renderer/Private/DeferredShadingRenderer.h
+++ b/Engine/Source/Runtime/Renderer/Private/DeferredShadingRenderer.h
@@ -569,6 +569,7 @@ private:
 	// #dxr_todo: UE-72565: refactor ray tracing effects to not be member functions of DeferredShadingRenderer. Register each effect at startup and just loop over them automatically
 	static void PrepareRayTracingReflections(const FViewInfo& View, TArray<FRayTracingShaderRHIParamRef>& OutRayGenShaders);
 	static void PrepareRayTracingShadows(const FViewInfo& View, TArray<FRayTracingShaderRHIParamRef>& OutRayGenShaders);
+	static void PrepareRayTracingSkyLight(const FViewInfo& View, TArray<FRayTracingShaderRHIParamRef>& OutRayGenShaders);
 	static void PrepareRayTracingRectLight(const FViewInfo& View, TArray<FRayTracingShaderRHIParamRef>& OutRayGenShaders);
 	static void PrepareRayTracingGlobalIllumination(const FViewInfo& View, TArray<FRayTracingShaderRHIParamRef>& OutRayGenShaders);
 	static void PrepareRayTracingTranslucency(const FViewInfo& View, TArray<FRayTracingShaderRHIParamRef>& OutRayGenShaders);
diff --git a/Engine/Source/Runtime/Renderer/Private/RayTracing/RaytracingSkylight.cpp b/Engine/Source/Runtime/Renderer/Private/RayTracing/RaytracingSkylight.cpp
--- a/Engine/Source/Runtime/Renderer/Private/RayTracing/RaytracingSkylight.cpp
+++ b/Engine/Source/Runtime/Renderer/Private/RayTracing/RaytracingSkylight.cpp
@@ -11,6 +11,7 @@ static int32 GRayTracingSkyLight = 0;
 #if RHI_RAYTRACING
 
 #include "RayTracingSkyLight.h"
+#include "RayTracingMaterialHitShaders.h"
 #include "ClearQuad.h"
 #include "DistanceFieldAmbientOcclusion.h"
 #include "SceneRendering.h"
@@ -68,6 +69,13 @@ static TAutoConsoleVariable<int32> CVarRayTracingSkyLightEnableTwoSidedGeometry(
 	ECVF_RenderThreadSafe
 );
 
+static TAutoConsoleVariable<int32> CVarRayTracingSkyLightEnableMaterials(
+	TEXT("r.RayTracing.SkyLight.EnableMaterials"),
+	0,
+	TEXT("Enables material shader binding for shadow rays. If this is disabled, then a default trivial shader is used. (default = 0)"),
+	ECVF_RenderThreadSafe
+);
+
 bool ShouldRenderRayTracingSkyLight(const FSkyLightSceneProxy* SkyLightSceneProxy)
 {
 	if (SkyLightSceneProxy != nullptr)
@@ -152,8 +160,9 @@ class FRayTracingSkyLightRGS : public FGlobalShader
 	SHADER_USE_ROOT_PARAMETER_STRUCT(FRayTracingSkyLightRGS, FGlobalShader)
 
 	class FEnableTwoSidedGeometryDim : SHADER_PERMUTATION_BOOL("ENABLE_TWO_SIDED_GEOMETRY");
+	class FEnableMaterialsDim : SHADER_PERMUTATION_BOOL("ENABLE_MATERIALS");
 
-	using FPermutationDomain = TShaderPermutationDomain<FEnableTwoSidedGeometryDim>;
+	using FPermutationDomain = TShaderPermutationDomain<FEnableTwoSidedGeometryDim, FEnableMaterialsDim>;
 
 	static bool ShouldCompilePermutation(const FGlobalShaderPermutationParameters& Parameters)
 	{
@@ -749,6 +758,22 @@ void FDeferredShadingSceneRenderer::VisualizeSkyLightMipTree(
 	RHICmdList.TransitionResource(EResourceTransitionAccess::ERWBarrier, EResourceTransitionPipeline::EGfxToCompute, SkyLightMipTreeNegZ.UAV);
 }
 
+void FDeferredShadingSceneRenderer::PrepareRayTracingSkyLight(const FViewInfo& View, TArray<FRayTracingShaderRHIParamRef>& OutRayGenShaders)
+{
+	// Declare all RayGen shaders that require material closest hit shaders to be bound
+	FRayTracingSkyLightRGS::FPermutationDomain PermutationVector;
+	for (uint32 TwoSidedGeometryIndex = 0; TwoSidedGeometryIndex < 2; ++TwoSidedGeometryIndex)
+	{
+		for (uint32 EnableMaterialsIndex = 0; EnableMaterialsIndex < 2; ++EnableMaterialsIndex)
+		{
+			PermutationVector.Set<FRayTracingSkyLightRGS::FEnableTwoSidedGeometryDim>(TwoSidedGeometryIndex != 0);
+			PermutationVector.Set<FRayTracingSkyLightRGS::FEnableMaterialsDim>(EnableMaterialsIndex != 0);
+			TShaderMapRef<FRayTracingSkyLightRGS> RayGenerationShader(View.ShaderMap, PermutationVector);
+			OutRayGenShaders.Add(RayGenerationShader->GetRayTracingShader());
+		}
+	}
+}
+
 void FDeferredShadingSceneRenderer::RenderRayTracingSkyLight(
 	FRHICommandListImmediate& RHICmdList,
 	TRefCountPtr<IPooledRenderTarget>& SkyLightRT,
@@ -838,6 +863,7 @@ void FDeferredShadingSceneRenderer::RenderRayTracingSkyLight(
 
 		FRayTracingSkyLightRGS::FPermutationDomain PermutationVector;
 		PermutationVector.Set<FRayTracingSkyLightRGS::FEnableTwoSidedGeometryDim>(CVarRayTracingSkyLightEnableTwoSidedGeometry.GetValueOnRenderThread() != 0);
+		PermutationVector.Set<FRayTracingSkyLightRGS::FEnableMaterialsDim>(CVarRayTracingSkyLightEnableMaterials.GetValueOnRenderThread() != 0);
 		TShaderMapRef<FRayTracingSkyLightRGS> RayGenerationShader(GetGlobalShaderMap(FeatureLevel), PermutationVector);
 		ClearUnusedGraphResources(*RayGenerationShader, PassParameters);
 
@@ -851,12 +877,25 @@ void FDeferredShadingSceneRenderer::RenderRayTracingSkyLight(
 			FRayTracingShaderBindingsWriter GlobalResources;
 			SetShaderParameters(GlobalResources, *RayGenerationShader, *PassParameters);
 
-			// Declare RT pipeline
-			FRayTracingPipelineStateInitializer Initializer;
-			FRayTracingShaderRHIParamRef RayGenShaderTable[] = { RayGenerationShader->GetRayTracingShader() };
-			Initializer.SetRayGenShaderTable(RayGenShaderTable);
+			FRHIRayTracingPipelineState* Pipeline = View.RayTracingMaterialPipeline;
+			if (CVarRayTracingSkyLightEnableMaterials.GetValueOnRenderThread() == 0)
+			{
+				// Declare default pipeline
+				FRayTracingPipelineStateInitializer Initializer;
+				Initializer.MaxPayloadSizeInBytes = 52; // sizeof(FPackedMaterialClosestHitPayload)
+				FRayTracingShaderRHIParamRef RayGenShaderTable[] = { RayGenerationShader->GetRayTracingShader() };
+				Initializer.SetRayGenShaderTable(RayGenShaderTable);
+
+				FRayTracingShaderRHIParamRef MissShaderTable[] = { View.ShaderMap->GetShader<FDefaultMaterialMS>()->GetRayTracingShader() };
+				Initializer.SetMissShaderTable(MissShaderTable);
+
+				FRayTracingShaderRHIParamRef HitGroupTable[] = { View.ShaderMap->GetShader<FOpaqueShadowHitGroup>()->GetRayTracingShader() };
+				Initializer.SetHitGroupTable(HitGroupTable);
+				Initializer.bAllowHitGroupIndexing = false; // Use the same hit shader for all geometry in the scene by disabling SBT indexing.
+
+				Pipeline = PipelineStateCache::GetAndOrCreateRayTracingPipelineState(Initializer);
+			}
 
-			FRHIRayTracingPipelineState* Pipeline = PipelineStateCache::GetAndOrCreateRayTracingPipelineState(Initializer);
 			FRayTracingSceneRHIParamRef RayTracingSceneRHI = View.RayTracingScene.RayTracingSceneRHI;
 			RHICmdList.RayTraceDispatch(Pipeline, RayGenerationShader->GetRayTracingShader(), RayTracingSceneRHI, GlobalResources, RayTracingResolution.X, RayTracingResolution.Y);
 		});
