Convert ray tracing ambient occlusion pass to RDG.

diff --git a/Engine/Shaders/Private/RayTracing/RayTracingAmbientOcclusionRGS.usf b/Engine/Shaders/Private/RayTracing/RayTracingAmbientOcclusionRGS.usf
--- a/Engine/Shaders/Private/RayTracing/RayTracingAmbientOcclusionRGS.usf
+++ b/Engine/Shaders/Private/RayTracing/RayTracingAmbientOcclusionRGS.usf
@@ -86,7 +86,7 @@ void AmbientOcclusionRGS()
 		RayFlags |= RAY_FLAG_CULL_BACK_FACING_TRIANGLES;
 #endif
 
-		FDefaultPayload Payload = (FDefaultPayload)0;
+		FPackedMaterialClosestHitPayload Payload = (FPackedMaterialClosestHitPayload)0;
 
 		TraceRay(
 			TLAS,
diff --git a/Engine/Source/Runtime/Renderer/Private/DeferredShadingRenderer.h b/Engine/Source/Runtime/Renderer/Private/DeferredShadingRenderer.h
--- a/Engine/Source/Runtime/Renderer/Private/DeferredShadingRenderer.h
+++ b/Engine/Source/Runtime/Renderer/Private/DeferredShadingRenderer.h
@@ -490,7 +490,7 @@ private:
 	void VisualizeRectLightMipTree(FRHICommandListImmediate& RHICmdList, const FViewInfo& View, const FRWBuffer& RectLightMipTree, const FIntVector& RectLightMipTreeDimensions);
 	
 	void RenderRayTracingAmbientOcclusion(FRHICommandListImmediate& RHICmdList, TRefCountPtr<IPooledRenderTarget>& AmbientOcclusionRT);
-	void RenderRayTracingAmbientOcclusion(FRHICommandListImmediate& RHICmdList, FRDGBuilder& GraphBuilder, FViewInfo& View, TRefCountPtr<IPooledRenderTarget>& AmbientOcclusionRT, TRefCountPtr<IPooledRenderTarget>& HitDistanceRT);
+	void RenderRayTracingAmbientOcclusion(FRHICommandListImmediate& RHICmdList, FRDGBuilder& GraphBuilder, FViewInfo& View, FRDGTextureRef AmbientOcclusionTexture, FRDGTextureRef RayDistanceTexture, TRefCountPtr<IPooledRenderTarget>& AmbientOcclusionMaskRT);
 	void CompositeRayTracingAmbientOcclusion(FRHICommandListImmediate& RHICmdList, TRefCountPtr<IPooledRenderTarget>& AmbientOcclusionRT);
 	
 	void RenderRayTracingGlobalIllumination(FRHICommandListImmediate& RHICmdList, TRefCountPtr<IPooledRenderTarget>& GlobalIlluminationRT);
diff --git a/Engine/Source/Runtime/Renderer/Private/RayTracing/RayTracingAmbientOcclusion.cpp b/Engine/Source/Runtime/Renderer/Private/RayTracing/RayTracingAmbientOcclusion.cpp
--- a/Engine/Source/Runtime/Renderer/Private/RayTracing/RayTracingAmbientOcclusion.cpp
+++ b/Engine/Source/Runtime/Renderer/Private/RayTracing/RayTracingAmbientOcclusion.cpp
@@ -13,6 +13,7 @@
 #include "UniformBuffer.h"
 #include "RHI/Public/PipelineStateCache.h"
 #include "Raytracing/RaytracingOptions.h"
+#include "RayTracingMaterialHitShaders.h"
 #include "SceneViewFamilyBlackboard.h"
 
 #include "PostProcess/PostProcessing.h"
@@ -48,6 +49,12 @@ static TAutoConsoleVariable<int32> CVarRayTracingAmbientOcclusionEnableTwoSidedG
 	ECVF_RenderThreadSafe
 );
 
+static TAutoConsoleVariable<int32> CVarRayTracingAmbientOcclusionEnableMaterials(
+	TEXT("r.RayTracing.AmbientOcclusion.EnableMaterials"),
+	0,
+	TEXT("Enables "),
+	ECVF_RenderThreadSafe
+);
 
 bool ShouldRenderRayTracingAmbientOcclusion(const FViewInfo& View)
 {
@@ -71,95 +78,36 @@ IMPLEMENT_GLOBAL_SHADER_PARAMETER_STRUCT(FAmbientOcclusionData, "AmbientOcclusio
 DECLARE_GPU_STAT_NAMED(RayTracingAmbientOcclusion, TEXT("Ray Tracing Ambient Occlusion"));
 DECLARE_GPU_STAT_NAMED(AmbientOcclusionDenoiser, TEXT("Ambient Occlusion Denoiser"));
 
-template<uint32 EnableTwoSidedGeometry>
-class TAmbientOcclusionRGS : public FGlobalShader
+class FRayTracingAmbientOcclusionRGS : public FGlobalShader
 {
-	DECLARE_SHADER_TYPE(TAmbientOcclusionRGS, Global)
+	DECLARE_GLOBAL_SHADER(FRayTracingAmbientOcclusionRGS)
+	SHADER_USE_ROOT_PARAMETER_STRUCT(FRayTracingAmbientOcclusionRGS, FGlobalShader)
 
-	static void ModifyCompilationEnvironment(const FGlobalShaderPermutationParameters& Parameters, FShaderCompilerEnvironment& OutEnvironment)
-	{
-		FGlobalShader::ModifyCompilationEnvironment(Parameters, OutEnvironment);
-		OutEnvironment.SetDefine(TEXT("ENABLE_TWO_SIDED_GEOMETRY"), EnableTwoSidedGeometry);
-	}
+	class FEnableTwoSidedGeometryDim : SHADER_PERMUTATION_BOOL("ENABLE_TWO_SIDED_GEOMETRY");
+	class FEnableMaterialsDim : SHADER_PERMUTATION_BOOL("ENABLE_MATERIALS");
+
+	using FPermutationDomain = TShaderPermutationDomain<FEnableTwoSidedGeometryDim, FEnableMaterialsDim>;
 
-public:
 	static bool ShouldCompilePermutation(const FGlobalShaderPermutationParameters& Parameters)
 	{
 		return ShouldCompileRayTracingShadersForProject(Parameters.Platform);
 	}
 
-	TAmbientOcclusionRGS() {}
-	virtual ~TAmbientOcclusionRGS() {}
-
-	TAmbientOcclusionRGS(const ShaderMetaType::CompiledShaderInitializerType& Initializer)
-		: FGlobalShader(Initializer)
-	{
-		ViewParameter.Bind(Initializer.ParameterMap, TEXT("View"));
-		TLASParameter.Bind(Initializer.ParameterMap, TEXT("TLAS"));
-		SceneTexturesParameter.Bind(Initializer.ParameterMap, TEXT("SceneTexturesStruct"));
-		AmbientOcclusionParameter.Bind(Initializer.ParameterMap, TEXT("AmbientOcclusion"));
-
-		OcclusionMaskUAVParameter.Bind(Initializer.ParameterMap, TEXT("RWOcclusionMaskUAV"));
-		RayDistanceUAVParameter.Bind(Initializer.ParameterMap, TEXT("RWHitDistanceUAV"));
-	}
-
-	bool Serialize(FArchive& Ar)
-	{
-		bool bShaderHasOutdatedParameters = FGlobalShader::Serialize(Ar);
-		Ar << ViewParameter;
-		Ar << TLASParameter;
-		Ar << SceneTexturesParameter;
-		Ar << AmbientOcclusionParameter;
-		Ar << OcclusionMaskUAVParameter;
-		Ar << RayDistanceUAVParameter;
-		return bShaderHasOutdatedParameters;
-	}
+	BEGIN_SHADER_PARAMETER_STRUCT(FParameters, )
+		SHADER_PARAMETER_SRV(RaytracingAccelerationStructure, TLAS)
+		SHADER_PARAMETER_RDG_TEXTURE_UAV(RWTexture2D<float>, RWOcclusionMaskUAV)
+		SHADER_PARAMETER_RDG_TEXTURE_UAV(RWTexture2D<float>, RWHitDistanceUAV)
 
-	void Dispatch(
-		FRHICommandListImmediate& RHICmdList,
-		const FRayTracingScene& RayTracingScene,
-		FUniformBufferRHIParamRef ViewUniformBuffer,
-		FUniformBufferRHIParamRef SceneTexturesUniformBuffer,
-		FUniformBufferRHIParamRef AmbientOcclusionUniformBuffer,
-		FUnorderedAccessViewRHIParamRef OcclusionMaskUAV,
-		FUnorderedAccessViewRHIParamRef HitDistanceUAV,
-		uint32 Width, uint32 Height
-	)
-	{
-		FRayTracingPipelineStateInitializer Initializer;
-
-		FRayTracingShaderRHIParamRef RayGenShaderTable[] = { GetRayTracingShader() };
-		Initializer.SetRayGenShaderTable(RayGenShaderTable);
-
-		FRHIRayTracingPipelineState* Pipeline = PipelineStateCache::GetAndOrCreateRayTracingPipelineState(Initializer); 
-
-		FRayTracingShaderBindingsWriter GlobalResources;
-		GlobalResources.Set(TLASParameter, RayTracingScene.RayTracingSceneRHI->GetShaderResourceView());
-		GlobalResources.Set(ViewParameter, ViewUniformBuffer);
-		GlobalResources.Set(SceneTexturesParameter, SceneTexturesUniformBuffer);
-		GlobalResources.Set(AmbientOcclusionParameter, AmbientOcclusionUniformBuffer);
-		GlobalResources.Set(OcclusionMaskUAVParameter, OcclusionMaskUAV);
-		GlobalResources.Set(RayDistanceUAVParameter, HitDistanceUAV);
-
-		RHICmdList.RayTraceDispatch(Pipeline, GetRayTracingShader(), RayTracingScene.RayTracingSceneRHI, GlobalResources, Width, Height);
-	}
-
-private:
-	// Input
-	FShaderResourceParameter TLASParameter;
-	FShaderUniformBufferParameter ViewParameter;
-	FShaderUniformBufferParameter SceneTexturesParameter;
-	FShaderUniformBufferParameter AmbientOcclusionParameter;
-
-	// Output
-	FShaderResourceParameter OcclusionMaskUAVParameter;
-	FShaderResourceParameter RayDistanceUAVParameter;
+		SHADER_PARAMETER_STRUCT_REF(FViewUniformShaderParameters, ViewUniformBuffer)
+		SHADER_PARAMETER_STRUCT_REF(FSceneTexturesUniformParameters, SceneTexturesStruct)
+		SHADER_PARAMETER_STRUCT_REF(FAmbientOcclusionData, AmbientOcclusionData)
+	END_SHADER_PARAMETER_STRUCT()
 };
 
-IMPLEMENT_SHADER_TYPE(template<>, TAmbientOcclusionRGS<0>, TEXT("/Engine/Private/RayTracing/RayTracingAmbientOcclusionRGS.usf"), TEXT("AmbientOcclusionRGS"), SF_RayGen)
-IMPLEMENT_SHADER_TYPE(template<>, TAmbientOcclusionRGS<1>, TEXT("/Engine/Private/RayTracing/RayTracingAmbientOcclusionRGS.usf"), TEXT("AmbientOcclusionRGS"), SF_RayGen)
+IMPLEMENT_GLOBAL_SHADER(FRayTracingAmbientOcclusionRGS, "/Engine/Private/RayTracing/RayTracingAmbientOcclusionRGS.usf", "AmbientOcclusionRGS", SF_RayGen);
+
 
-void FDeferredShadingSceneRenderer::RenderRayTracingAmbientOcclusion(FRHICommandListImmediate& RHICmdList, TRefCountPtr<IPooledRenderTarget>& AmbientOcclusionMaskRT)
+void FDeferredShadingSceneRenderer::RenderRayTracingAmbientOcclusion(FRHICommandListImmediate& RHICmdList, TRefCountPtr<IPooledRenderTarget>& AmbientOcclusionRT)
 {
 	SCOPED_DRAW_EVENT(RHICmdList, RayTracingAmbientOcclusion);
 	SCOPED_GPU_STAT(RHICmdList, RayTracingAmbientOcclusion);
@@ -177,86 +125,109 @@ void FDeferredShadingSceneRenderer::RenderRayTracingAmbientOcclusion(FRHICommand
 	}
 
 	FRDGBuilder GraphBuilder(RHICmdList);
-
 	FSceneRenderTargets& SceneContext = FSceneRenderTargets::Get(RHICmdList);
-	FPooledRenderTargetDesc Desc = SceneContext.GetSceneColor()->GetDesc();
-	Desc.Format = PF_R16F;
-	Desc.Flags &= ~(TexCreate_FastVRAM | TexCreate_Transient);
 
-	TRefCountPtr<IPooledRenderTarget> HitDistanceRT;
+	// Define RDG targets
+	FRDGTextureRef AmbientOcclusionTexture;
+	{
+		FRDGTextureDesc Desc = SceneContext.GetSceneColor()->GetDesc();
+		Desc.Format = PF_R16F;
+		Desc.Flags &= ~(TexCreate_FastVRAM | TexCreate_Transient);
+		AmbientOcclusionTexture = GraphBuilder.CreateTexture(Desc, TEXT("RayTracingAmbientOcclusion"));
+	}
 
-	GRenderTargetPool.FindFreeElement(RHICmdList, Desc, AmbientOcclusionMaskRT, TEXT("RayTracingAmbientOcclusion"));
-	GRenderTargetPool.FindFreeElement(RHICmdList, Desc, HitDistanceRT, TEXT("RayTracingAmbientOcclusionHitDistance"));
+	FRDGTextureRef RayDistanceTexture;
+	{
+		FRDGTextureDesc Desc = SceneContext.GetSceneColor()->GetDesc();
+		Desc.Format = PF_R16F;
+		Desc.Flags &= ~(TexCreate_FastVRAM | TexCreate_Transient);
+		RayDistanceTexture = GraphBuilder.CreateTexture(Desc, TEXT("RayTracingAmbientOcclusionHitDistance"));
+	}
 
 	for (int32 ViewIndex = 0; ViewIndex < Views.Num(); ++ViewIndex)
 	{	
 		FViewInfo& View = Views[ViewIndex];
 		if (ShouldRenderRayTracingAmbientOcclusion(View))
 		{
-			RenderRayTracingAmbientOcclusion(RHICmdList, GraphBuilder, View, AmbientOcclusionMaskRT, HitDistanceRT);
+			RenderRayTracingAmbientOcclusion(RHICmdList, GraphBuilder, View, AmbientOcclusionTexture, RayDistanceTexture, AmbientOcclusionRT);
 		}
 	}
 
 	GraphBuilder.Execute();
-
 	SceneContext.bScreenSpaceAOIsValid = true;
+	AmbientOcclusionRT->SetDebugName(TEXT("RayTracingAmbientOcclusion"));
+	GVisualizeTexture.SetCheckPoint(RHICmdList, AmbientOcclusionRT);
 }
 
 void FDeferredShadingSceneRenderer::RenderRayTracingAmbientOcclusion(
-	FRHICommandListImmediate& RHICmdList, 
-	FRDGBuilder& GraphBuilder, 
+	FRHICommandListImmediate& RHICmdList,
+	FRDGBuilder& GraphBuilder,
 	FViewInfo& View,
-	TRefCountPtr<IPooledRenderTarget>& AmbientOcclusionMaskRT,
-	TRefCountPtr<IPooledRenderTarget>& HitDistanceRT
+	FRDGTextureRef AmbientOcclusionTexture,
+	FRDGTextureRef RayDistanceTexture,
+	TRefCountPtr<IPooledRenderTarget>& AmbientOcclusionMaskRT
 )
 {
-	//#dxr_todo refactor to RDG
-	
-	FIntPoint ViewSize = View.ViewRect.Size();
 	FSceneRenderTargets& SceneContext = FSceneRenderTargets::Get(RHICmdList);
 
 	// Add ambient occlusion parameters to uniform buffer
 	FAmbientOcclusionData AmbientOcclusionData;
-
-	FSceneTexturesUniformParameters SceneTextures;
-	SetupSceneTextureUniformParameters(SceneContext, FeatureLevel, ESceneTextureSetupMode::All, SceneTextures);
-	FUniformBufferRHIRef SceneTexturesUniformBuffer = RHICreateUniformBuffer(&SceneTextures, FSceneTexturesUniformParameters::StaticStructMetadata.GetLayout(), EUniformBufferUsage::UniformBuffer_SingleDraw);
-
 	AmbientOcclusionData.SamplesPerPixel = GRayTracingAmbientOcclusionSamplesPerPixel >= 0 ? GRayTracingAmbientOcclusionSamplesPerPixel : View.FinalPostProcessSettings.RayTracingAOSamplesPerPixel;
 	AmbientOcclusionData.MaxRayDistance = View.FinalPostProcessSettings.AmbientOcclusionRadius;
 	AmbientOcclusionData.Intensity = View.FinalPostProcessSettings.AmbientOcclusionIntensity;
 	AmbientOcclusionData.MaxNormalBias = GetRaytracingMaxNormalBias();
-	FUniformBufferRHIRef AmbientOcclusionUniformBuffer = RHICreateUniformBuffer(&AmbientOcclusionData, FAmbientOcclusionData::StaticStructMetadata.GetLayout(), EUniformBufferUsage::UniformBuffer_SingleDraw);
 
-	int32 EnableTwoSidedGeometry = CVarRayTracingAmbientOcclusionEnableTwoSidedGeometry.GetValueOnRenderThread();
-	if (EnableTwoSidedGeometry)
-	{
-		TShaderMapRef<TAmbientOcclusionRGS<1>> AmbientOcclusionRayGenerationShader(GetGlobalShaderMap(FeatureLevel));
-		AmbientOcclusionRayGenerationShader->Dispatch(
-			RHICmdList,
-			View.RayTracingScene,
-			View.ViewUniformBuffer,
-			SceneTexturesUniformBuffer,
-			AmbientOcclusionUniformBuffer,
-			AmbientOcclusionMaskRT->GetRenderTargetItem().UAV,
-			HitDistanceRT->GetRenderTargetItem().UAV,
-			ViewSize.X, ViewSize.Y
-		);
-	}
-	else
+	FSceneTexturesUniformParameters SceneTextures;
+	SetupSceneTextureUniformParameters(SceneContext, FeatureLevel, ESceneTextureSetupMode::All, SceneTextures);
+
+	// Build RTAO parameters
+	FRayTracingAmbientOcclusionRGS::FParameters *PassParameters = GraphBuilder.AllocParameters<FRayTracingAmbientOcclusionRGS::FParameters>();
+	PassParameters->TLAS = View.RayTracingScene.RayTracingSceneRHI->GetShaderResourceView();
+	PassParameters->RWOcclusionMaskUAV = GraphBuilder.CreateUAV(AmbientOcclusionTexture);
+	PassParameters->RWHitDistanceUAV = GraphBuilder.CreateUAV(RayDistanceTexture);
+	PassParameters->ViewUniformBuffer = View.ViewUniformBuffer;
+	PassParameters->SceneTexturesStruct = CreateUniformBufferImmediate(SceneTextures, EUniformBufferUsage::UniformBuffer_SingleDraw);
+	PassParameters->AmbientOcclusionData = CreateUniformBufferImmediate(AmbientOcclusionData, EUniformBufferUsage::UniformBuffer_SingleDraw);
+
+	FRayTracingAmbientOcclusionRGS::FPermutationDomain PermutationVector;
+	PermutationVector.Set<FRayTracingAmbientOcclusionRGS::FEnableTwoSidedGeometryDim>(CVarRayTracingAmbientOcclusionEnableTwoSidedGeometry.GetValueOnRenderThread() != 0);
+	PermutationVector.Set<FRayTracingAmbientOcclusionRGS::FEnableMaterialsDim>(CVarRayTracingAmbientOcclusionEnableMaterials.GetValueOnRenderThread() != 0);
+	TShaderMapRef<FRayTracingAmbientOcclusionRGS> RayGenerationShader(GetGlobalShaderMap(FeatureLevel), PermutationVector);
+	ClearUnusedGraphResources(*RayGenerationShader, PassParameters);
+
+	FIntPoint RayTracingResolution = View.ViewRect.Size();
+	GraphBuilder.AddPass(
+		RDG_EVENT_NAME("AmbientOcclusionRayTracing %dx%d", RayTracingResolution.X, RayTracingResolution.Y),
+		PassParameters,
+		ERenderGraphPassFlags::Compute,
+		[PassParameters, this, &View, RayGenerationShader, RayTracingResolution](FRHICommandList& RHICmdList)
 	{
-		TShaderMapRef<TAmbientOcclusionRGS<0>> AmbientOcclusionRayGenerationShader(GetGlobalShaderMap(FeatureLevel));
-		AmbientOcclusionRayGenerationShader->Dispatch(
-			RHICmdList,
-			View.RayTracingScene,
-			View.ViewUniformBuffer,
-			SceneTexturesUniformBuffer,
-			AmbientOcclusionUniformBuffer,
-			AmbientOcclusionMaskRT->GetRenderTargetItem().UAV,
-			HitDistanceRT->GetRenderTargetItem().UAV,
-			ViewSize.X, ViewSize.Y
-		);
-	}
+		FRayTracingShaderBindingsWriter GlobalResources;
+		SetShaderParameters(GlobalResources, *RayGenerationShader, *PassParameters);
+
+		// TODO: Provide material support for opacity mask
+		FRHIRayTracingPipelineState* Pipeline = View.RayTracingMaterialPipeline;
+		//if (CVarRayTracingAmbientOcclusionEnableMaterials.GetValueOnRenderThread() == 0)
+		{
+			// Declare default pipeline
+			FRayTracingPipelineStateInitializer Initializer;
+			Initializer.MaxPayloadSizeInBytes = 52; // sizeof(FPackedMaterialClosestHitPayload)
+			FRayTracingShaderRHIParamRef RayGenShaderTable[] = { RayGenerationShader->GetRayTracingShader() };
+			Initializer.SetRayGenShaderTable(RayGenShaderTable);
+
+			FRayTracingShaderRHIParamRef MissShaderTable[] = { View.ShaderMap->GetShader<FDefaultMaterialMS>()->GetRayTracingShader() };
+			Initializer.SetMissShaderTable(MissShaderTable);
+
+			FRayTracingShaderRHIParamRef HitGroupTable[] = { View.ShaderMap->GetShader<FOpaqueShadowHitGroup>()->GetRayTracingShader() };
+			Initializer.SetHitGroupTable(HitGroupTable);
+			Initializer.bAllowHitGroupIndexing = false; // Use the same hit shader for all geometry in the scene by disabling SBT indexing.
+
+			Pipeline = PipelineStateCache::GetAndOrCreateRayTracingPipelineState(Initializer);
+		}
+
+		FRayTracingSceneRHIParamRef RayTracingSceneRHI = View.RayTracingScene.RayTracingSceneRHI;
+		RHICmdList.RayTraceDispatch(Pipeline, RayGenerationShader->GetRayTracingShader(), RayTracingSceneRHI, GlobalResources, RayTracingResolution.X, RayTracingResolution.Y);
+	});
 
 	int32 DenoiserMode = CVarUseAODenoiser.GetValueOnRenderThread();
 	if (DenoiserMode != 0)
@@ -272,8 +243,8 @@ void FDeferredShadingSceneRenderer::RenderRayTracingAmbientOcclusion(
 		IScreenSpaceDenoiser::FAmbientOcclusionRayTracingConfig RayTracingConfig;
 
 		IScreenSpaceDenoiser::FAmbientOcclusionInputs DenoiserInputs;
-		DenoiserInputs.Mask = GraphBuilder.RegisterExternalTexture(SceneContext.ScreenSpaceAO, TEXT("AOMask"));
-		DenoiserInputs.RayHitDistance = GraphBuilder.RegisterExternalTexture(HitDistanceRT, TEXT("AOHitDistance"));
+		DenoiserInputs.Mask = AmbientOcclusionTexture;
+		DenoiserInputs.RayHitDistance = RayDistanceTexture;
 
 		{
 			RDG_EVENT_SCOPE(GraphBuilder, "%s%s(AmbientOcclusion) %dx%d",
@@ -289,17 +260,13 @@ void FDeferredShadingSceneRenderer::RenderRayTracingAmbientOcclusion(
 				DenoiserInputs,
 				RayTracingConfig);
 
-			GraphBuilder.QueueTextureExtraction(DenoiserOutputs.AmbientOcclusionMask, &SceneContext.ScreenSpaceAO);
+			GraphBuilder.QueueTextureExtraction(DenoiserOutputs.AmbientOcclusionMask, &AmbientOcclusionMaskRT);
 		}
 	}
-
-	FUnorderedAccessViewRHIParamRef UAVs[]
+	else
 	{
-		AmbientOcclusionMaskRT->GetRenderTargetItem().UAV,
-		HitDistanceRT->GetRenderTargetItem().UAV
-	};
-	FComputeFenceRHIRef Fence = RHICmdList.CreateComputeFence(TEXT("RayTracingAmbientOcclusion"));
-	RHICmdList.TransitionResources(EResourceTransitionAccess::ERWBarrier, EResourceTransitionPipeline::EComputeToGfx, UAVs, ARRAY_COUNT(UAVs), Fence);
+		GraphBuilder.QueueTextureExtraction(AmbientOcclusionTexture, &AmbientOcclusionMaskRT);
+	}
 }
 
 #endif // RHI_RAYTRACING
