Removing FAmbientOcclusionData shader struct in favor of inlined parameters.

diff --git a/Engine/Shaders/Private/RayTracing/RayTracingAmbientOcclusionRGS.usf b/Engine/Shaders/Private/RayTracing/RayTracingAmbientOcclusionRGS.usf
--- a/Engine/Shaders/Private/RayTracing/RayTracingAmbientOcclusionRGS.usf
+++ b/Engine/Shaders/Private/RayTracing/RayTracingAmbientOcclusionRGS.usf
@@ -6,6 +6,10 @@
 #include "../PathTracing/Utilities/PathTracingRandomSequence.ush"
 #include "RayTracingCommon.ush"
 
+uint SamplesPerPixel;
+float MaxRayDistance;
+float Intensity;
+float MaxNormalBias;
 RaytracingAccelerationStructure TLAS;
 RWTexture2D<float> RWOcclusionMaskUAV;
 RWTexture2D<float> RWHitDistanceUAV;
@@ -33,7 +37,7 @@ void GenerateCosineNormalRay(
 	RayOrigin = WorldPosition;
 	RayDirection = Direction_World;
 	RayTMin = 0.0;
-	RayTMax = AmbientOcclusion.MaxRayDistance;
+	RayTMax = MaxRayDistance;
 }
 
 [shader("raygeneration")]
@@ -55,9 +59,8 @@ void AmbientOcclusionRGS()
 
 	float RayCount = 0.0;
 	float Visibility = 0.0;
-	float Intensity = AmbientOcclusion.Intensity;
 	float ClosestRayHitDistance = 1.#INF;
-	uint SamplesPerPixel = AmbientOcclusion.SamplesPerPixel;
+	uint SamplesPerPixelLocal = SamplesPerPixel;
 	
 	// Mask out depth values beyond far plane
 	float DeviceZ = SceneTexturesStruct.SceneDepthTexture.Load(int3(PixelCoord, 0)).r;
@@ -69,14 +72,14 @@ void AmbientOcclusionRGS()
 	{
 		Visibility = 1.0;
 		RayCount = SamplesPerPixel;
-		SamplesPerPixel = 0.0;
+		SamplesPerPixelLocal = 0.0;
 	}
 
-	for (uint SampleIndex = 0; SampleIndex < SamplesPerPixel; ++SampleIndex)
+	for (uint SampleIndex = 0; SampleIndex < SamplesPerPixelLocal; ++SampleIndex)
 	{
 		RayDesc Ray;
 		GenerateCosineNormalRay(WorldPosition, WorldNormal, RandSequence, Ray.Origin, Ray.Direction, Ray.TMin, Ray.TMax);
-		ApplyPositionBias(Ray, WorldNormal, AmbientOcclusion.MaxNormalBias);
+		ApplyPositionBias(Ray, WorldNormal, MaxNormalBias);
 
 		if (dot(WorldNormal, Ray.Direction) <= 0.0) // TODO: does it needs to be handled by the denoiser?
 			continue;
@@ -111,7 +114,7 @@ void AmbientOcclusionRGS()
 	{
 		float2 RawOutput = 1;
 
-		if (SamplesPerPixel == 0)
+		if (SamplesPerPixelLocal == 0)
 		{
 			RawOutput.x = 1;
 			RawOutput.y = -2;
diff --git a/Engine/Source/Runtime/Renderer/Private/RayTracing/RayTracingAmbientOcclusion.cpp b/Engine/Source/Runtime/Renderer/Private/RayTracing/RayTracingAmbientOcclusion.cpp
--- a/Engine/Source/Runtime/Renderer/Private/RayTracing/RayTracingAmbientOcclusion.cpp
+++ b/Engine/Source/Runtime/Renderer/Private/RayTracing/RayTracingAmbientOcclusion.cpp
@@ -65,16 +65,6 @@ bool ShouldRenderRayTracingAmbientOcclusion(const FViewInfo& View)
 	return IsRayTracingEnabled() && !ShouldRenderRayTracingGlobalIllumination(View) && bRTAOEnabled;
 }
 
-
-BEGIN_GLOBAL_SHADER_PARAMETER_STRUCT(FAmbientOcclusionData, )
-SHADER_PARAMETER(int, SamplesPerPixel)
-SHADER_PARAMETER(float, MaxRayDistance)
-SHADER_PARAMETER(float, Intensity)
-SHADER_PARAMETER(float, MaxNormalBias)
-END_GLOBAL_SHADER_PARAMETER_STRUCT()
-
-IMPLEMENT_GLOBAL_SHADER_PARAMETER_STRUCT(FAmbientOcclusionData, "AmbientOcclusion");
-
 DECLARE_GPU_STAT_NAMED(RayTracingAmbientOcclusion, TEXT("Ray Tracing Ambient Occlusion"));
 DECLARE_GPU_STAT_NAMED(AmbientOcclusionDenoiser, TEXT("Ambient Occlusion Denoiser"));
 
@@ -94,13 +84,16 @@ class FRayTracingAmbientOcclusionRGS : public FGlobalShader
 	}
 
 	BEGIN_SHADER_PARAMETER_STRUCT(FParameters, )
+		SHADER_PARAMETER(int, SamplesPerPixel)
+		SHADER_PARAMETER(float, MaxRayDistance)
+		SHADER_PARAMETER(float, Intensity)
+		SHADER_PARAMETER(float, MaxNormalBias)
 		SHADER_PARAMETER_SRV(RaytracingAccelerationStructure, TLAS)
 		SHADER_PARAMETER_RDG_TEXTURE_UAV(RWTexture2D<float>, RWOcclusionMaskUAV)
 		SHADER_PARAMETER_RDG_TEXTURE_UAV(RWTexture2D<float>, RWHitDistanceUAV)
 
 		SHADER_PARAMETER_STRUCT_REF(FViewUniformShaderParameters, ViewUniformBuffer)
 		SHADER_PARAMETER_STRUCT_REF(FSceneTexturesUniformParameters, SceneTexturesStruct)
-		SHADER_PARAMETER_STRUCT_REF(FAmbientOcclusionData, AmbientOcclusionData)
 	END_SHADER_PARAMETER_STRUCT()
 };
 
@@ -169,25 +162,20 @@ void FDeferredShadingSceneRenderer::RenderRayTracingAmbientOcclusion(
 )
 {
 	FSceneRenderTargets& SceneContext = FSceneRenderTargets::Get(RHICmdList);
-
-	// Add ambient occlusion parameters to uniform buffer
-	FAmbientOcclusionData AmbientOcclusionData;
-	AmbientOcclusionData.SamplesPerPixel = GRayTracingAmbientOcclusionSamplesPerPixel >= 0 ? GRayTracingAmbientOcclusionSamplesPerPixel : View.FinalPostProcessSettings.RayTracingAOSamplesPerPixel;
-	AmbientOcclusionData.MaxRayDistance = View.FinalPostProcessSettings.AmbientOcclusionRadius;
-	AmbientOcclusionData.Intensity = View.FinalPostProcessSettings.AmbientOcclusionIntensity;
-	AmbientOcclusionData.MaxNormalBias = GetRaytracingMaxNormalBias();
-
 	FSceneTexturesUniformParameters SceneTextures;
 	SetupSceneTextureUniformParameters(SceneContext, FeatureLevel, ESceneTextureSetupMode::All, SceneTextures);
 
 	// Build RTAO parameters
 	FRayTracingAmbientOcclusionRGS::FParameters *PassParameters = GraphBuilder.AllocParameters<FRayTracingAmbientOcclusionRGS::FParameters>();
+	PassParameters->SamplesPerPixel = GRayTracingAmbientOcclusionSamplesPerPixel >= 0 ? GRayTracingAmbientOcclusionSamplesPerPixel : View.FinalPostProcessSettings.RayTracingAOSamplesPerPixel;
+	PassParameters->MaxRayDistance = View.FinalPostProcessSettings.AmbientOcclusionRadius;
+	PassParameters->Intensity = View.FinalPostProcessSettings.AmbientOcclusionIntensity;
+	PassParameters->MaxNormalBias = GetRaytracingMaxNormalBias();
 	PassParameters->TLAS = View.RayTracingScene.RayTracingSceneRHI->GetShaderResourceView();
 	PassParameters->RWOcclusionMaskUAV = GraphBuilder.CreateUAV(AmbientOcclusionTexture);
 	PassParameters->RWHitDistanceUAV = GraphBuilder.CreateUAV(RayDistanceTexture);
 	PassParameters->ViewUniformBuffer = View.ViewUniformBuffer;
 	PassParameters->SceneTexturesStruct = CreateUniformBufferImmediate(SceneTextures, EUniformBufferUsage::UniformBuffer_SingleDraw);
-	PassParameters->AmbientOcclusionData = CreateUniformBufferImmediate(AmbientOcclusionData, EUniformBufferUsage::UniformBuffer_SingleDraw);
 
 	FRayTracingAmbientOcclusionRGS::FPermutationDomain PermutationVector;
 	PermutationVector.Set<FRayTracingAmbientOcclusionRGS::FEnableTwoSidedGeometryDim>(CVarRayTracingAmbientOcclusionEnableTwoSidedGeometry.GetValueOnRenderThread() != 0);
