Make FScreenPassTextureViewport a class for consistency.

diff --git a/Engine/Source/Runtime/Renderer/Private/ScreenPass.h b/Engine/Source/Runtime/Renderer/Private/ScreenPass.h
--- a/Engine/Source/Runtime/Renderer/Private/ScreenPass.h
+++ b/Engine/Source/Runtime/Renderer/Private/ScreenPass.h
@@ -34,8 +34,9 @@ public:
 };
 
 // Describes a viewport rect oriented within a texture.
-struct FScreenPassTextureViewport
+class FScreenPassTextureViewport
 {
+public:
 	// Creates a viewport that is downscaled by the requested scale factor.
 	static FScreenPassTextureViewport CreateDownscaled(const FScreenPassTextureViewport& Other, uint32 ScaleFactor);
 
