Activate new shadow bias.

* Add default value for constant / slope bias.
* Enable receiver bias for directional, spot light, & rect light
* Add variable to tweak receiver bias per light types
* Add receiver bias for point light

This changelist will require test screenshots update.

diff --git a/Engine/Shaders/Private/ShadowDepthVertexShader.usf b/Engine/Shaders/Private/ShadowDepthVertexShader.usf
--- a/Engine/Shaders/Private/ShadowDepthVertexShader.usf
+++ b/Engine/Shaders/Private/ShadowDepthVertexShader.usf
@@ -56,7 +56,7 @@ void SetShadowDepthOutputs(
 	#endif
 
 	const float MaxSlopeDepthBias = PassStruct.ShadowParams.z;
-	const float Slope = min(abs(NoL) > 0 ? sqrt(saturate(1 - NoL*NoL)) / NoL : MaxSlopeDepthBias, MaxSlopeDepthBias);
+	const float Slope = clamp(abs(NoL) > 0 ? sqrt(saturate(1 - NoL*NoL)) / NoL : MaxSlopeDepthBias, 0, MaxSlopeDepthBias);
 	
 	const float SlopeDepthBias = PassStruct.ShadowParams.y;
 	const float SlopeBias = SlopeDepthBias * Slope;
diff --git a/Engine/Shaders/Private/ShadowProjectionCommon.ush b/Engine/Shaders/Private/ShadowProjectionCommon.ush
--- a/Engine/Shaders/Private/ShadowProjectionCommon.ush
+++ b/Engine/Shaders/Private/ShadowProjectionCommon.ush
@@ -190,7 +190,7 @@ float CubemapHardwarePCF(float3 WorldPosition, float3 LightPosition, float Light
 
 		// Calculate the Z buffer value that would have been stored for this position in the shadow map
 		float CompareDistance = ShadowPosition.z / ShadowPosition.w;
-		float ShadowDepthBias = - DepthBias / ShadowPosition.w;
+		float ShadowDepthBias = - (DepthBias + SlopeDepthBias) / ShadowPosition.w;
 
 		Shadow = 0;
 
diff --git a/Engine/Shaders/Private/ShadowProjectionPixelShader.usf b/Engine/Shaders/Private/ShadowProjectionPixelShader.usf
--- a/Engine/Shaders/Private/ShadowProjectionPixelShader.usf
+++ b/Engine/Shaders/Private/ShadowProjectionPixelShader.usf
@@ -55,11 +55,12 @@ float4 LightPositionAndInvRadius;
 	#include "TransmissionThickness.ush"
 #else
 	float4x4 ScreenToShadowMatrix;
-	// .x:DepthBias, .y:SlopeDepthBias, .z:MaxSlopeDepthBias, .w: MaxSubjectZ - MinSubjectZ
+	// .x:DepthBias, .y:SlopeDepthBias, .z:ReceiverBias, .w: MaxSubjectZ - MinSubjectZ
 	float4 ProjectionDepthBiasParameters;
 #endif
 
 	float4 LightPositionOrDirection;
+	float ShadowReceiverBias;
 
 #if USE_FADE_PLANE
 	float FadePlaneOffset;
@@ -181,7 +182,7 @@ void Main(
 	}
 	#else // !USE_PCSS
 	{
-		// Attenuate Soft transition based on the angle with the light and the shading normal
+		// Attenuate soft transition based on the angle with the light and the shading normal (acts as a receiver bias)
 		FGBufferData GBufferData = GetGBufferData(ScreenUV);
 		const bool bIsDirectional = LightPositionOrDirection.w == 0;
 		const float3 LightDirection = bIsDirectional ? -LightPositionOrDirection.xyz : normalize(LightPositionOrDirection.xyz - WorldPosition);
@@ -192,7 +193,7 @@ void Main(
 		Settings.ShadowDepthTexture = ShadowDepthTexture;
 		Settings.ShadowDepthTextureSampler = ShadowDepthTextureSampler;
 		Settings.ShadowBufferSize = ShadowBufferSize;
-		Settings.TransitionScale = SoftTransitionScale.z;// * lerp(0.1, 1.0, NoL); Temporally disabled until slope bias is activated at the same time, to reduce the number of screenshots to update.
+		Settings.TransitionScale = SoftTransitionScale.z * lerp(ProjectionDepthBiasParameters.z, 1.0, NoL);
 		Settings.SceneDepth = LightSpacePixelDepthForOpaque;
 		Settings.bSubsurface = false;
 		Settings.bTreatMaxDepthUnshadowed = false;
@@ -332,13 +333,25 @@ void MainOnePassPointLightPS(
 	float3 LightVector = LightPositionAndInvRadius.xyz - WorldPosition.xyz;
 #if OUTPUT_CUBE_SHADOW_DEPTH_NO_FILTERING
 	// Note: point light shadow depth is Z / W, not linear
-	#if OPENGL_PROFILE || COMPILER_VULKAN || COMPILER_SWITCH
+#if OPENGL_PROFILE || COMPILER_VULKAN || COMPILER_SWITCH
 	float Shadow = TextureCubeSampleLevel(ShadowDepthCubeTexture2, ShadowDepthTextureSampler, LightVector, 0);
-	#else
+#else
 	float Shadow = TextureCubeSampleLevel(ShadowDepthCubeTexture, ShadowDepthTextureSampler, LightVector, 0);
-	#endif
+#endif
 #else
- 	float Shadow = CubemapHardwarePCF(WorldPosition, LightPositionAndInvRadius.xyz, LightPositionAndInvRadius.w, PointLightDepthBias.x, PointLightDepthBias.y, PointLightDepthBias.z);
+	// Normal receiver bias: Increase depth based on the angle with the light and the shading normal. 
+	// Unlike spot/rect/direction light, point lights performs depth bias only at projection time. 
+	// This is why the slope bias drives the normal receiver bias here.
+	float SlopeBias = 0;
+	{
+		FGBufferData GBufferData = GetGBufferData(ScreenUV);
+		const bool bIsDirectional = LightPositionOrDirection.w == 0;
+		const float3 LightDirection = bIsDirectional ? -LightPositionOrDirection.xyz : normalize(LightPositionOrDirection.xyz - WorldPosition);
+		const float NoL = saturate(dot(GBufferData.WorldNormal, LightDirection));
+		SlopeBias = (1-NoL) * PointLightDepthBias.y;
+	}
+
+	float Shadow = CubemapHardwarePCF(WorldPosition, LightPositionAndInvRadius.xyz, LightPositionAndInvRadius.w, PointLightDepthBias.x, SlopeBias, PointLightDepthBias.z);
 
 	Shadow = saturate( (Shadow - 0.5) * ShadowSharpen + 0.5 );
 #endif
diff --git a/Engine/Source/Runtime/Engine/Private/Components/LightComponent.cpp b/Engine/Source/Runtime/Engine/Private/Components/LightComponent.cpp
--- a/Engine/Source/Runtime/Engine/Private/Components/LightComponent.cpp
+++ b/Engine/Source/Runtime/Engine/Private/Components/LightComponent.cpp
@@ -389,7 +389,7 @@ ULightComponent::ULightComponent(const FObjectInitializer& ObjectInitializer)
 	IndirectLightingIntensity = 1.0f;
 	ShadowResolutionScale = 1.0f;
 	ShadowBias = 0.5f;
-	ShadowSlopeBias = 0.0f;
+	ShadowSlopeBias = 0.5f;
 	ShadowSharpen = 0.0f;
 	ContactShadowLength = 0.0f;
 	ContactShadowLengthInWS = false;
diff --git a/Engine/Source/Runtime/Renderer/Private/ShadowRendering.cpp b/Engine/Source/Runtime/Renderer/Private/ShadowRendering.cpp
--- a/Engine/Source/Runtime/Renderer/Private/ShadowRendering.cpp
+++ b/Engine/Source/Runtime/Renderer/Private/ShadowRendering.cpp
@@ -14,28 +14,30 @@
 #include "PipelineStateCache.h"
 #include "ClearQuad.h"
 
+///////////////////////////////////////////////////////////////////////////////////////////////////
+// Directional light
 static TAutoConsoleVariable<float> CVarCSMShadowDepthBias(
 	TEXT("r.Shadow.CSMDepthBias"),
-	20.0f,
+	10.0f,
 	TEXT("Constant depth bias used by CSM"),
 	ECVF_RenderThreadSafe);
 
 static TAutoConsoleVariable<float> CVarCSMShadowSlopeScaleDepthBias(
 	TEXT("r.Shadow.CSMSlopeScaleDepthBias"),
-	5.0f,
+	3.0f,
 	TEXT("Slope scale depth bias used by CSM"),
 	ECVF_RenderThreadSafe);
 
 static TAutoConsoleVariable<float> CVarPerObjectDirectionalShadowDepthBias(
 	TEXT("r.Shadow.PerObjectDirectionalDepthBias"),
-	20.0f,
+	10.0f,
 	TEXT("Constant depth bias used by per-object shadows from directional lights\n")
 	TEXT("Lower values give better shadow contact, but increase self-shadowing artifacts"),
 	ECVF_RenderThreadSafe);
 
 static TAutoConsoleVariable<float> CVarPerObjectDirectionalShadowSlopeScaleDepthBias(
 	TEXT("r.Shadow.PerObjectDirectionalSlopeDepthBias"),
-	5.0f,
+	3.0f,
 	TEXT("Slope scale depth bias used by per-object shadows from directional lights\n")
 	TEXT("Lower values give better shadow contact, but increase self-shadowing artifacts"),
 	ECVF_RenderThreadSafe);
@@ -52,42 +54,83 @@ static TAutoConsoleVariable<int32> CVarCSMDepthBoundsTest(
 	TEXT("Whether to use depth bounds tests rather than stencil tests for the CSM bounds"),
 	ECVF_RenderThreadSafe);
 
-static TAutoConsoleVariable<float> CVarSpotLightShadowTransitionScale(
-	TEXT("r.Shadow.SpotLightTransitionScale"),
-	60.0f,
-	TEXT("Transition scale for spotlights"),
-	ECVF_RenderThreadSafe);
-
 static TAutoConsoleVariable<float> CVarShadowTransitionScale(
 	TEXT("r.Shadow.TransitionScale"),
 	60.0f,
 	TEXT("This controls the 'fade in' region between a caster and where his shadow shows up.  Larger values make a smaller region which will have more self shadowing artifacts"),
 	ECVF_RenderThreadSafe);
 
+static TAutoConsoleVariable<float> CVarCSMShadowReceiverBias(
+	TEXT("r.Shadow.CSMReceiverBias"),
+	0.9f,
+	TEXT("Receiver bias used by CSM. Value between 0 and 1."),
+	ECVF_RenderThreadSafe);
+
+
+///////////////////////////////////////////////////////////////////////////////////////////////////
+// Point light
 static TAutoConsoleVariable<float> CVarPointLightShadowDepthBias(
 	TEXT("r.Shadow.PointLightDepthBias"),
-	0.05f,
+	0.02f,
 	TEXT("Depth bias that is applied in the depth pass for shadows from point lights. (0.03 avoids peter paning but has some shadow acne)"),
 	ECVF_RenderThreadSafe);
 
 static TAutoConsoleVariable<float> CVarPointLightShadowSlopeScaleDepthBias(
 	TEXT("r.Shadow.PointLightSlopeScaleDepthBias"),
-	5.0f,
+	3.0f,
 	TEXT("Slope scale depth bias that is applied in the depth pass for shadows from point lights"),
 	ECVF_RenderThreadSafe);
 
+
+///////////////////////////////////////////////////////////////////////////////////////////////////
+// Rect light
+static TAutoConsoleVariable<float> CVarRectLightShadowDepthBias(
+	TEXT("r.Shadow.RectLightDepthBias"),
+	0.025f,
+	TEXT("Depth bias that is applied in the depth pass for shadows from rect lights. (0.03 avoids peter paning but has some shadow acne)"),
+	ECVF_RenderThreadSafe);
+
+static TAutoConsoleVariable<float> CVarRectLightShadowSlopeScaleDepthBias(
+	TEXT("r.Shadow.RectLightSlopeScaleDepthBias"),
+	2.5f,
+	TEXT("Slope scale depth bias that is applied in the depth pass for shadows from rect lights"),
+	ECVF_RenderThreadSafe);
+
+static TAutoConsoleVariable<float> CVarRectLightShadowReceiverBias(
+	TEXT("r.Shadow.RectLightReceiverBias"),
+	0.3f,
+	TEXT("Receiver bias used by rect light. Value between 0 and 1."),
+	ECVF_RenderThreadSafe);
+
+///////////////////////////////////////////////////////////////////////////////////////////////////
+// Spot light
 static TAutoConsoleVariable<float> CVarSpotLightShadowDepthBias(
 	TEXT("r.Shadow.SpotLightDepthBias"),
-	5.0f,
+	3.0f,
 	TEXT("Depth bias that is applied in the depth pass for per object projected shadows from spot lights"),
 	ECVF_RenderThreadSafe);
 
 static TAutoConsoleVariable<float> CVarSpotLightShadowSlopeScaleDepthBias(
 	TEXT("r.Shadow.SpotLightSlopeDepthBias"),
-	5.0f,
+	3.0f,
 	TEXT("Slope scale depth bias that is applied in the depth pass for per object projected shadows from spot lights"),
 	ECVF_RenderThreadSafe);
 
+static TAutoConsoleVariable<float> CVarSpotLightShadowTransitionScale(
+	TEXT("r.Shadow.SpotLightTransitionScale"),
+	60.0f,
+	TEXT("Transition scale for spotlights"),
+	ECVF_RenderThreadSafe);
+
+static TAutoConsoleVariable<float> CVarSpotLightShadowReceiverBias(
+	TEXT("r.Shadow.SpotLightReceiverBias"),
+	0.5f,
+	TEXT("Receiver bias used by spotlights. Value between 0 and 1."),
+	ECVF_RenderThreadSafe);
+
+
+///////////////////////////////////////////////////////////////////////////////////////////////////
+// General
 static TAutoConsoleVariable<int32> CVarEnableModulatedSelfShadow(
 	TEXT("r.Shadow.EnableModulatedSelfShadow"),
 	0,
@@ -118,7 +161,7 @@ static TAutoConsoleVariable<int32> CVarMaxSoftKernelSize(
 
 static TAutoConsoleVariable<float> CVarShadowMaxSlopeScaleDepthBias(
 	TEXT("r.Shadow.ShadowMaxSlopeScaleDepthBias"),
-	2.0f,
+	1.0f,
 	TEXT("Max Slope depth bias used for shadows for all lights\n")
 	TEXT("Higher values give better self-shadowing, but increase self-shadowing artifacts"),
 	ECVF_RenderThreadSafe);
@@ -1166,11 +1209,25 @@ void FProjectedShadowInfo::UpdateShaderDepthBias()
 
 	if (IsWholeScenePointLightShadow())
 	{
-		DepthBias = CVarPointLightShadowDepthBias.GetValueOnRenderThread() * 512.0f / FMath::Max(ResolutionX, ResolutionY);
+		const bool bIsRectLight = LightSceneInfo->Proxy->GetLightType() == LightType_Rect;
+		float DeptBiasConstant = 0;
+		float SlopeDepthBiasConstant = 0;
+		if (bIsRectLight)
+		{
+			DeptBiasConstant = CVarRectLightShadowDepthBias.GetValueOnRenderThread();
+			SlopeDepthBiasConstant = CVarRectLightShadowSlopeScaleDepthBias.GetValueOnRenderThread();
+		}
+		else
+		{
+			DeptBiasConstant = CVarPointLightShadowDepthBias.GetValueOnRenderThread();
+			SlopeDepthBiasConstant = CVarPointLightShadowSlopeScaleDepthBias.GetValueOnRenderThread();
+		}
+
+		DepthBias = DeptBiasConstant * 512.0f / FMath::Max(ResolutionX, ResolutionY);
 		// * 2.0f to be compatible with the system we had before ShadowBias
 		DepthBias *= 2.0f * LightSceneInfo->Proxy->GetUserShadowBias();
 
-		SlopeScaleDepthBias = CVarPointLightShadowSlopeScaleDepthBias.GetValueOnRenderThread();
+		SlopeScaleDepthBias = SlopeDepthBiasConstant;
 		SlopeScaleDepthBias *= LightSceneInfo->Proxy->GetUserShadowSlopeBias();
 	}
 	else if (IsWholeSceneDirectionalShadow())
@@ -1276,6 +1333,24 @@ float FProjectedShadowInfo::ComputeTransitionSize() const
 	return FMath::Max(TransitionSize, MinTransitionSize);
 }
 
+float FProjectedShadowInfo::GetShaderReceiverDepthBias() const
+{
+	float ShadowReceiverBias = 1;
+	{
+		switch (GetLightSceneInfo().Proxy->GetLightType())
+		{
+		case LightType_Directional	: ShadowReceiverBias = CVarCSMShadowReceiverBias.GetValueOnRenderThread(); break;
+		case LightType_Rect			: ShadowReceiverBias = CVarRectLightShadowReceiverBias.GetValueOnRenderThread(); break;
+		case LightType_Spot			: ShadowReceiverBias = CVarSpotLightShadowReceiverBias.GetValueOnRenderThread(); break;
+		case LightType_Point		: ShadowReceiverBias = GetShaderSlopeDepthBias(); break;
+		}
+	}
+
+	// Return the min lerp value for depth biasing
+	// 0 : max bias when NoL == 0
+	// 1 : no bias
+	return 1.0f - FMath::Clamp(ShadowReceiverBias, 0.0f, 1.0f);
+}
 /*-----------------------------------------------------------------------------
 FDeferredShadingSceneRenderer
 -----------------------------------------------------------------------------*/
diff --git a/Engine/Source/Runtime/Renderer/Private/ShadowRendering.h b/Engine/Source/Runtime/Renderer/Private/ShadowRendering.h
--- a/Engine/Source/Runtime/Renderer/Private/ShadowRendering.h
+++ b/Engine/Source/Runtime/Renderer/Private/ShadowRendering.h
@@ -371,7 +371,7 @@ public:
 	float GetShaderDepthBias() const { return ShaderDepthBias; }
 	float GetShaderSlopeDepthBias() const { return ShaderSlopeDepthBias; }
 	float GetShaderMaxSlopeDepthBias() const { return ShaderMaxSlopeDepthBias; }
-	FVector GetShaderShadowReceiverBias() const { return ShaderShadowReceiverDepthBias; }
+	float GetShaderReceiverDepthBias() const;
 
 	/**
 	 * Renders the shadow subject depth.
@@ -555,7 +555,6 @@ private:
 	float ShaderDepthBias;
 	float ShaderSlopeDepthBias;
 	float ShaderMaxSlopeDepthBias;
-	FVector ShaderShadowReceiverDepthBias;
 
 	void CopyCachedShadowMap(FRHICommandList& RHICmdList, const FMeshPassProcessorRenderState& DrawRenderState, FSceneRenderer* SceneRenderer, const FViewInfo& View);
 
@@ -769,7 +768,6 @@ public:
 		FadePlaneOffset.Bind(ParameterMap,TEXT("FadePlaneOffset"));
 		InvFadePlaneLength.Bind(ParameterMap,TEXT("InvFadePlaneLength"));
 		ShadowTileOffsetAndSizeParam.Bind(ParameterMap, TEXT("ShadowTileOffsetAndSize"));
-		ShadowReceiverParams.Bind(ParameterMap, TEXT("ShadowReceiverParams"));
 		LightPositionOrDirection.Bind(ParameterMap, TEXT("LightPositionOrDirection"));
 	}
 
@@ -845,7 +843,7 @@ public:
 				);
 		}
 
-		SetShaderValue(RHICmdList, ShaderRHI, ProjectionDepthBias, FVector4(ShadowInfo->GetShaderDepthBias(), ShadowInfo->GetShaderSlopeDepthBias(), ShadowInfo->GetShaderMaxSlopeDepthBias(), ShadowInfo->MaxSubjectZ - ShadowInfo->MinSubjectZ));
+		SetShaderValue(RHICmdList, ShaderRHI, ProjectionDepthBias, FVector4(ShadowInfo->GetShaderDepthBias(), ShadowInfo->GetShaderSlopeDepthBias(), ShadowInfo->GetShaderReceiverDepthBias(), ShadowInfo->MaxSubjectZ - ShadowInfo->MinSubjectZ));
 		SetShaderValue(RHICmdList, ShaderRHI, FadePlaneOffset, ShadowInfo->CascadeSettings.FadePlaneOffset);
 
 		if(InvFadePlaneLength.IsBound())
@@ -854,10 +852,6 @@ public:
 			SetShaderValue(RHICmdList, ShaderRHI, InvFadePlaneLength, 1.0f / ShadowInfo->CascadeSettings.FadePlaneLength);
 		}
 
-		if (ShadowReceiverParams.IsBound())
-		{
-			SetShaderValue(RHICmdList, ShaderRHI, ShadowReceiverParams, ShadowInfo->GetShaderShadowReceiverBias());
-		}
 
 		if (LightPositionOrDirection.IsBound())
 		{
@@ -881,7 +875,6 @@ public:
 		Ar << P.FadePlaneOffset;
 		Ar << P.InvFadePlaneLength;
 		Ar << P.ShadowTileOffsetAndSizeParam;
-		Ar << P.ShadowReceiverParams;
 		Ar << P.LightPositionOrDirection;
 		return Ar;
 	}
@@ -898,7 +891,6 @@ private:
 	FShaderParameter FadePlaneOffset;
 	FShaderParameter InvFadePlaneLength;
 	FShaderParameter ShadowTileOffsetAndSizeParam;
-	FShaderParameter ShadowReceiverParams;
 	FShaderParameter LightPositionOrDirection;
 };
 
