Uses RDG_GPU_STAT_SCOPE() in TAA.

diff --git a/Engine/Source/Runtime/Renderer/Private/PostProcess/PostProcessTemporalAA.cpp b/Engine/Source/Runtime/Renderer/Private/PostProcess/PostProcessTemporalAA.cpp
--- a/Engine/Source/Runtime/Renderer/Private/PostProcess/PostProcessTemporalAA.cpp
+++ b/Engine/Source/Runtime/Renderer/Private/PostProcess/PostProcessTemporalAA.cpp
@@ -609,6 +609,8 @@ FTAAOutputs FTAAPassParameters::AddTemporalAAPass(
 		}
 	}
 
+	RDG_GPU_STAT_SCOPE(GraphBuilder, TAA);
+
 	if (bIsComputePass)
 	{
 		FTemporalAACS::FPermutationDomain PermutationVector;
@@ -671,21 +673,15 @@ FTAAOutputs FTAAPassParameters::AddTemporalAAPass(
 		}
 
 		TShaderMapRef<FTemporalAACS> ComputeShader(View.ShaderMap, PermutationVector);
-		ClearUnusedGraphResources(*ComputeShader, PassParameters);
-		GraphBuilder.AddPass(
+		FComputeShaderUtils::AddPass(
+			GraphBuilder,
 			RDG_EVENT_NAME("TAA %s CS%s %dx%d -> %dx%d",
 				PassName, bUseFast ? TEXT(" Fast") : TEXT(""),
 				PracticableSrcRect.Width(), PracticableSrcRect.Height(),
 				PracticableDestRect.Width(), PracticableDestRect.Height()),
+			*ComputeShader,
 			PassParameters,
-			ERenderGraphPassFlags::Compute,
-			[PassParameters, ComputeShader, PracticableDestRect](FRHICommandList& RHICmdList)
-		{
-			SCOPED_GPU_STAT(RHICmdList, TAA); // TODO(RDG) find a way to integrate GPU stat groups into RDG so can use FComputeShaderUtils::AddPass()
-			FComputeShaderUtils::Dispatch(
-				RHICmdList, *ComputeShader, *PassParameters,
-				FComputeShaderUtils::GetGroupCount(PracticableDestRect.Size(), GTemporalAATileSizeX));
-		});
+			FComputeShaderUtils::GetGroupCount(PracticableDestRect.Size(), GTemporalAATileSizeX));
 	}
 	else
 	{
@@ -733,7 +729,6 @@ FTAAOutputs FTAAPassParameters::AddTemporalAAPass(
 			ERenderGraphPassFlags::None,
 			[PassParameters, &View, PracticableDestRect, bCameraCut, BasePermutationVector, bUseResponsiveStencilTest](FRHICommandList& RHICmdList)
 		{
-			SCOPED_GPU_STAT(RHICmdList, TAA);
 			RHICmdList.SetViewport(PracticableDestRect.Min.X, PracticableDestRect.Min.Y, 0.0f, PracticableDestRect.Max.X, PracticableDestRect.Max.Y, 1.0f);
 
 			FTemporalAAPS::FPermutationDomain PermutationVector = BasePermutationVector;
