Adding optional any-hit material support to AmbientOcclusion pass.

r.RayTracing.AmbientOcclusion.EnableMaterials (default = 0) allows support for visibility tracing against masked materials.

diff --git a/Engine/Shaders/Private/RayTracing/RayTracingAmbientOcclusionRGS.usf b/Engine/Shaders/Private/RayTracing/RayTracingAmbientOcclusionRGS.usf
--- a/Engine/Shaders/Private/RayTracing/RayTracingAmbientOcclusionRGS.usf
+++ b/Engine/Shaders/Private/RayTracing/RayTracingAmbientOcclusionRGS.usf
@@ -85,6 +85,9 @@ void AmbientOcclusionRGS()
 			continue;
 		
 		uint RayFlags = 0;
+#if !ENABLE_MATERIALS
+		RayFlags |= RAY_FLAG_FORCE_OPAQUE;
+#endif
 #if !ENABLE_TWO_SIDED_GEOMETRY
 		RayFlags |= RAY_FLAG_CULL_BACK_FACING_TRIANGLES;
 #endif
diff --git a/Engine/Source/Runtime/Renderer/Private/DeferredShadingRenderer.cpp b/Engine/Source/Runtime/Renderer/Private/DeferredShadingRenderer.cpp
--- a/Engine/Source/Runtime/Renderer/Private/DeferredShadingRenderer.cpp
+++ b/Engine/Source/Runtime/Renderer/Private/DeferredShadingRenderer.cpp
@@ -763,6 +763,7 @@ bool FDeferredShadingSceneRenderer::DispatchRayTracingWorldUpdates(FRHICommandLi
 		TArray<FRayTracingShaderRHIParamRef> RayGenShaders;
 		PrepareRayTracingReflections(View, RayGenShaders);
 		PrepareRayTracingShadows(View, RayGenShaders);
+		PrepareRayTracingAmbientOcclusion(View, RayGenShaders);
 		PrepareRayTracingSkyLight(View, RayGenShaders);
 		PrepareRayTracingRectLight(View, RayGenShaders);
 		PrepareRayTracingGlobalIllumination(View, RayGenShaders);
diff --git a/Engine/Source/Runtime/Renderer/Private/DeferredShadingRenderer.h b/Engine/Source/Runtime/Renderer/Private/DeferredShadingRenderer.h
--- a/Engine/Source/Runtime/Renderer/Private/DeferredShadingRenderer.h
+++ b/Engine/Source/Runtime/Renderer/Private/DeferredShadingRenderer.h
@@ -569,6 +569,7 @@ private:
 	// #dxr_todo: UE-72565: refactor ray tracing effects to not be member functions of DeferredShadingRenderer. Register each effect at startup and just loop over them automatically
 	static void PrepareRayTracingReflections(const FViewInfo& View, TArray<FRayTracingShaderRHIParamRef>& OutRayGenShaders);
 	static void PrepareRayTracingShadows(const FViewInfo& View, TArray<FRayTracingShaderRHIParamRef>& OutRayGenShaders);
+	static void PrepareRayTracingAmbientOcclusion(const FViewInfo& View, TArray<FRayTracingShaderRHIParamRef>& OutRayGenShaders);
 	static void PrepareRayTracingSkyLight(const FViewInfo& View, TArray<FRayTracingShaderRHIParamRef>& OutRayGenShaders);
 	static void PrepareRayTracingRectLight(const FViewInfo& View, TArray<FRayTracingShaderRHIParamRef>& OutRayGenShaders);
 	static void PrepareRayTracingGlobalIllumination(const FViewInfo& View, TArray<FRayTracingShaderRHIParamRef>& OutRayGenShaders);
diff --git a/Engine/Source/Runtime/Renderer/Private/RayTracing/RayTracingAmbientOcclusion.cpp b/Engine/Source/Runtime/Renderer/Private/RayTracing/RayTracingAmbientOcclusion.cpp
--- a/Engine/Source/Runtime/Renderer/Private/RayTracing/RayTracingAmbientOcclusion.cpp
+++ b/Engine/Source/Runtime/Renderer/Private/RayTracing/RayTracingAmbientOcclusion.cpp
@@ -99,6 +99,21 @@ class FRayTracingAmbientOcclusionRGS : public FGlobalShader
 
 IMPLEMENT_GLOBAL_SHADER(FRayTracingAmbientOcclusionRGS, "/Engine/Private/RayTracing/RayTracingAmbientOcclusionRGS.usf", "AmbientOcclusionRGS", SF_RayGen);
 
+void FDeferredShadingSceneRenderer::PrepareRayTracingAmbientOcclusion(const FViewInfo& View, TArray<FRayTracingShaderRHIParamRef>& OutRayGenShaders)
+{
+	// Declare all RayGen shaders that require material closest hit shaders to be bound
+	FRayTracingAmbientOcclusionRGS::FPermutationDomain PermutationVector;
+	for (uint32 TwoSidedGeometryIndex = 0; TwoSidedGeometryIndex < 2; ++TwoSidedGeometryIndex)
+	{
+		for (uint32 EnableMaterialsIndex = 0; EnableMaterialsIndex < 2; ++EnableMaterialsIndex)
+		{
+			PermutationVector.Set<FRayTracingAmbientOcclusionRGS::FEnableTwoSidedGeometryDim>(TwoSidedGeometryIndex != 0);
+			PermutationVector.Set<FRayTracingAmbientOcclusionRGS::FEnableMaterialsDim>(EnableMaterialsIndex != 0);
+			TShaderMapRef<FRayTracingAmbientOcclusionRGS> RayGenerationShader(View.ShaderMap, PermutationVector);
+			OutRayGenShaders.Add(RayGenerationShader->GetRayTracingShader());
+		}
+	}
+}
 
 void FDeferredShadingSceneRenderer::RenderRayTracingAmbientOcclusion(FRHICommandListImmediate& RHICmdList, TRefCountPtr<IPooledRenderTarget>& AmbientOcclusionRT)
 {
@@ -195,7 +210,7 @@ void FDeferredShadingSceneRenderer::RenderRayTracingAmbientOcclusion(
 
 		// TODO: Provide material support for opacity mask
 		FRHIRayTracingPipelineState* Pipeline = View.RayTracingMaterialPipeline;
-		//if (CVarRayTracingAmbientOcclusionEnableMaterials.GetValueOnRenderThread() == 0)
+		if (CVarRayTracingAmbientOcclusionEnableMaterials.GetValueOnRenderThread() == 0)
 		{
 			// Declare default pipeline
 			FRayTracingPipelineStateInitializer Initializer;
