Plugs SSR into reflection denoiser

diff --git a/Engine/Shaders/Private/RayTracing/RayTracingOcclusionRGS.usf b/Engine/Shaders/Private/RayTracing/RayTracingOcclusionRGS.usf
--- a/Engine/Shaders/Private/RayTracing/RayTracingOcclusionRGS.usf
+++ b/Engine/Shaders/Private/RayTracing/RayTracingOcclusionRGS.usf
@@ -5,6 +5,7 @@
 #include "../DeferredShadingCommon.ush"
 #include "../LightShaderParameters.ush"
 #include "../SceneViewFamilyBlackboard.ush"
+#include "../ScreenSpaceDenoise/SSDPublic.ush"
 #include "RayTracingCommon.ush"
 #include "RayTracingDirectionalLight.ush"
 #include "RayTracingRectLight.ush"
@@ -120,7 +121,7 @@ void OcclusionRGS()
 	
 	float Visibility = 0.0;
 	float SumRayDistance = 0.0;
-	float ClosestRayDistance = -2.0;
+	float ClosestRayDistance = DENOISER_INVALID_HIT_DISTANCE;
 	float HitCount = 0.0;
 	float RayCount = 0.0;
 	uint LocalSamplesPerPixel = SamplesPerPixel;
@@ -241,13 +242,13 @@ void OcclusionRGS()
 			Visibility += Payload.HitT < 0 ? 1.0 : 0.0;
 			if (Payload.HitT >= 0.0)
 			{
-				ClosestRayDistance = ClosestRayDistance == -2 || Payload.HitT < ClosestRayDistance ? Payload.HitT : ClosestRayDistance;
+				ClosestRayDistance = ClosestRayDistance == DENOISER_INVALID_HIT_DISTANCE || Payload.HitT < ClosestRayDistance ? Payload.HitT : ClosestRayDistance;
 				SumRayDistance += Payload.HitT;
 				HitCount += 1.0;
 			}
 			else
 			{
-				ClosestRayDistance = ClosestRayDistance == -2 ? -1 : ClosestRayDistance;
+				ClosestRayDistance = ClosestRayDistance == DENOISER_INVALID_HIT_DISTANCE ? DENOISER_MISS_HIT_DISTANCE : ClosestRayDistance;
 			}
 		}
 	}
diff --git a/Engine/Shaders/Private/RayTracing/RayTracingReflections.usf b/Engine/Shaders/Private/RayTracing/RayTracingReflections.usf
--- a/Engine/Shaders/Private/RayTracing/RayTracingReflections.usf
+++ b/Engine/Shaders/Private/RayTracing/RayTracingReflections.usf
@@ -38,6 +38,7 @@ float2 SobolIndexToUniformUnitSquare(uint2 SobolRand)
 #include "../PathTracing/Utilities/PathTracingRandomSequence.ush" 
 #include "../HeightFogCommon.ush"
 #include "../LightShaderParameters.ush"
+#include "../ScreenSpaceDenoise/SSDPublic.ush"
 #include "RayTracingHitGroupCommon.ush"
 #include "RayTracingDirectionalLight.ush"
 #include "RayTracingRectLight.ush"
@@ -185,7 +186,7 @@ void RayTracingReflectionsRGS()
 
 	float4 ReflectedColor = 0.0;
 	float ReflectedWeight = 0.0;
-	float HitDistance = -2.0;
+	float HitDistance = DENOISER_INVALID_HIT_DISTANCE;
 
 	float RoughnessFade = GetRoughnessFade(ScreenSpaceData.GBuffer.Roughness, ReflectionMaxRoughness);
 	float ClearCoatRoughnessFade = ScreenSpaceData.GBuffer.ShadingModelID == SHADINGMODELID_CLEAR_COAT ? GetRoughnessFade(ScreenSpaceData.GBuffer.CustomData.y, ReflectionMaxRoughness) : 0.0f;
@@ -350,7 +351,7 @@ void RayTracingReflectionsRGS()
 					if (BounceIndex == 0)
 					{
 						// Tell to the denoiser this sample reached outer bounds of the scene without any intersection.
-						HitDistance = -1;
+						HitDistance = DENOISER_MISS_HIT_DISTANCE;
 					}
 					TopLayerRadiance = 0;
 				}
diff --git a/Engine/Shaders/Private/ScreenSpaceDenoise/SSDCommon.ush b/Engine/Shaders/Private/ScreenSpaceDenoise/SSDCommon.ush
--- a/Engine/Shaders/Private/ScreenSpaceDenoise/SSDCommon.ush
+++ b/Engine/Shaders/Private/ScreenSpaceDenoise/SSDCommon.ush
@@ -10,6 +10,7 @@
 #include "../SceneViewFamilyBlackboard.ush"
 #include "../ColorSpace.ush"
 #include "../SphericalGaussian.ush"
+#include "SSDPublic.ush"
 #include "SSDDefinitions.ush"
 
 
diff --git a/Engine/Shaders/Private/ScreenSpaceDenoise/SSDPublic.ush b/Engine/Shaders/Private/ScreenSpaceDenoise/SSDPublic.ush
new file mode 100644
--- /dev/null
+++ b/Engine/Shaders/Private/ScreenSpaceDenoise/SSDPublic.ush
@@ -0,0 +1,13 @@
+// Copyright 1998-2019 Epic Games, Inc. All Rights Reserved.
+
+#pragma once
+
+
+/** Bias used on GGX important sample when denoising, to remove part of the tale that create a lot more noise. */
+#define GGX_IMPORTANT_SAMPLE_BIAS 0.1
+
+/** Special value to be encoded in denoiser input when no ray has been shot. */
+#define DENOISER_INVALID_HIT_DISTANCE -2.0
+
+/** Special value to be encoded in denoiser input when no ray intersection has been found. */
+#define DENOISER_MISS_HIT_DISTANCE -1.0
diff --git a/Engine/Shaders/Private/ScreenSpaceDenoise/SSDSignalCore.ush b/Engine/Shaders/Private/ScreenSpaceDenoise/SSDSignalCore.ush
--- a/Engine/Shaders/Private/ScreenSpaceDenoise/SSDSignalCore.ush
+++ b/Engine/Shaders/Private/ScreenSpaceDenoise/SSDSignalCore.ush
@@ -215,9 +215,9 @@ void DecodeMultiplexedSignals(
 		for (uint BatchSignalId = 0; BatchSignalId < MAX_SIGNAL_BATCH_SIZE; BatchSignalId++)
 		{
 			uint MultiplexId = BatchSignalId;
-			OutSamples.Array[MultiplexId].AccumulatedSampleCount = (RawSample[MultiplexId].r == -2.0 ? 0.0 : 1.0);
+			OutSamples.Array[MultiplexId].AccumulatedSampleCount = (RawSample[MultiplexId].r == DENOISER_INVALID_HIT_DISTANCE ? 0.0 : 1.0);
 			OutSamples.Array[MultiplexId].LeavingRayCount = OutSamples.Array[MultiplexId].AccumulatedSampleCount * (RawSample[MultiplexId].r > 0.0 ? 0.0 : 1.0);
-			OutSamples.Array[MultiplexId].PenumbraClosestHit = OutSamples.Array[MultiplexId].AccumulatedSampleCount * (RawSample[MultiplexId].r == -1 ? 1000000 : RawSample[MultiplexId].r);
+			OutSamples.Array[MultiplexId].PenumbraClosestHit = OutSamples.Array[MultiplexId].AccumulatedSampleCount * (RawSample[MultiplexId].r == DENOISER_MISS_HIT_DISTANCE ? 1000000 : RawSample[MultiplexId].r);
 		}
 	}
 	else if (SignalBufferLayout == SIGNAL_BUFFER_LAYOUT_PENUMBRA_INPUT_NSPP)
@@ -227,9 +227,9 @@ void DecodeMultiplexedSignals(
 		{
 			uint MultiplexId = BatchSignalId;
 			// TODO: feed the actual number of sample.
-			OutSamples.Array[MultiplexId].AccumulatedSampleCount = (RawSample[MultiplexId].g == -2.0 ? 0.0 : 1.0);
+			OutSamples.Array[MultiplexId].AccumulatedSampleCount = (RawSample[MultiplexId].g == DENOISER_INVALID_HIT_DISTANCE ? 0.0 : 1.0);
 			OutSamples.Array[MultiplexId].LeavingRayCount = OutSamples.Array[MultiplexId].AccumulatedSampleCount * RawSample[MultiplexId].r;
-			OutSamples.Array[MultiplexId].PenumbraClosestHit = OutSamples.Array[MultiplexId].AccumulatedSampleCount * (RawSample[MultiplexId].g == -1 ? 1000000 : RawSample[MultiplexId].g);
+			OutSamples.Array[MultiplexId].PenumbraClosestHit = OutSamples.Array[MultiplexId].AccumulatedSampleCount * (RawSample[MultiplexId].g == DENOISER_MISS_HIT_DISTANCE ? 1000000 : RawSample[MultiplexId].g);
 		}
 	}
 	else if (SignalBufferLayout == SIGNAL_BUFFER_LAYOUT_PENUMBRA_INJESTION_1SPP)
@@ -288,7 +288,7 @@ void DecodeMultiplexedSignals(
 
 			OutSamples.Array[MultiplexId].LeavingRayCount = Channels.r * SampleCount;
 			OutSamples.Array[MultiplexId].AccumulatedSampleCount = SampleCount;
-			OutSamples.Array[MultiplexId].WorldBluringRadius = (Channels.b == -1 ? WORLD_RADIUS_MISS : Channels.b) * SampleCount;
+			OutSamples.Array[MultiplexId].WorldBluringRadius = (Channels.b == DENOISER_MISS_HIT_DISTANCE ? WORLD_RADIUS_MISS : Channels.b) * SampleCount;
 		}
 	}
 	else if (SignalBufferLayout == SIGNAL_BUFFER_LAYOUT_PENUMBRA_REJECTION)
@@ -310,13 +310,13 @@ void DecodeMultiplexedSignals(
 	}
 	else if (SignalBufferLayout == SIGNAL_BUFFER_LAYOUT_REFLECTIONS_INPUT)
 	{
-		OutSamples.Array[0].AccumulatedSampleCount = RawSample[1].r == -2 ? 0 : 1;
+		OutSamples.Array[0].AccumulatedSampleCount = RawSample[1].r == DENOISER_INVALID_HIT_DISTANCE ? 0 : 1;
 		OutSamples.Array[0].LeavingRayCount = OutSamples.Array[0].AccumulatedSampleCount * (1 - RawSample[0].a);
 		OutSamples.Array[0].SceneColor = RawSample[0];
 
 		// max() because ray hit distance is < 0 when miss / not ray has been shot.
 		//OutSamples.Array[0].PenumbraClosestHit = RawSample[1].r * OutSamples.Array[0].AccumulatedSampleCount;
-		OutSamples.Array[0].PenumbraClosestHit = ((RawSample[1].r == -1.0) ? 1.0e20 : RawSample[1].r) * OutSamples.Array[0].AccumulatedSampleCount;
+		OutSamples.Array[0].PenumbraClosestHit = ((RawSample[1].r == DENOISER_MISS_HIT_DISTANCE) ? 1.0e20 : RawSample[1].r) * OutSamples.Array[0].AccumulatedSampleCount;
 	}
 	else if (SignalBufferLayout == SIGNAL_BUFFER_LAYOUT_REFLECTIONS_REJECTION)
 	{
diff --git a/Engine/Shaders/Private/ScreenSpaceDenoise/SSDSignalFramework.ush b/Engine/Shaders/Private/ScreenSpaceDenoise/SSDSignalFramework.ush
--- a/Engine/Shaders/Private/ScreenSpaceDenoise/SSDSignalFramework.ush
+++ b/Engine/Shaders/Private/ScreenSpaceDenoise/SSDSignalFramework.ush
@@ -360,7 +360,7 @@ float ComputeInfinityBluringRadius(FSSFSampleSceneInfos SceneMetadata)
 	#if 1
 	{
 		// Measure the angle that only truncate 10% of D_GGX's tale.
-		const float E = (1.0 - 0.1);
+		const float E = (1.0 - GGX_IMPORTANT_SAMPLE_BIAS);
 	
 		float a2 = Pow4(SceneMetadata.Roughness);
 	
diff --git a/Engine/Shaders/Private/ScreenSpaceDenoise/SSDSpatialAccumulation.usf b/Engine/Shaders/Private/ScreenSpaceDenoise/SSDSpatialAccumulation.usf
--- a/Engine/Shaders/Private/ScreenSpaceDenoise/SSDSpatialAccumulation.usf
+++ b/Engine/Shaders/Private/ScreenSpaceDenoise/SSDSpatialAccumulation.usf
@@ -104,11 +104,17 @@
 		#define CONFIG_SAMPLE_SET           SAMPLE_SET_STACKOWIAK_4_SETS
 		#define CONFIG_SAMPLE_COUNT_POLICY  SAMPLE_COUNT_POLICY_HIT_DISTANCE_BASED
 
-		// Reconstruction has 1 ray per input pixel, therefore can use specular ratio estimator instead of Tokoyashi axes bilateral kernel.
-		#define CONFIG_SPECULAR_RATIO_ESTIMATOR 1
+		// Configure the bilateral therm.
+		#if DIM_MULTI_SPP
+			#define CONFIG_TOKOYASHI_LOBE_BILATERAL 1
+			#define CONFIG_TOKOYASHI_AXES_BILATERAL 1
+		#else
+			// Reconstruction has 1 ray per input pixel, therefore can use specular ratio estimator instead of Tokoyashi axes bilateral kernel.
+			#define CONFIG_SPECULAR_RATIO_ESTIMATOR 1
 
-		// Uses tokoyashi lobe bilateral rejection to avoid mixing up significantly different lobes that can cause leaking/strong biasing.
-		#define CONFIG_TOKOYASHI_LOBE_BILATERAL 1
+			// Uses tokoyashi lobe bilateral rejection to avoid mixing up significantly different lobes that can cause leaking/strong biasing.
+			#define CONFIG_TOKOYASHI_LOBE_BILATERAL 1
+		#endif
 
 		//#define CONFIG_SAMPLE_COUNT_POLICY  SAMPLE_COUNT_POLICY_DISABLED
 
@@ -357,6 +363,12 @@ void MainCS(
 			/* MultiplexedSampleId = */ 0,
 			/* bNormalizeSample = */ CONFIG_NORMALIZE_INPUT != 0,
 			SceneBufferUV);
+		
+		#if CONFIG_NORMALIZE_INPUT
+			FSSDSignalArray NormalizedRefSamples = RefSamples;
+		#else
+			FSSDSignalArray NormalizedRefSamples = NormalizeToOneSample(RefSamples);
+		#endif
 	#endif
 
 	/** factor by witch should be spread out. */
@@ -383,7 +395,7 @@ void MainCS(
 		#if CONFIG_SIGNAL_BATCH_SIZE != 1
 			#error Unable to support more than one signal.
 		#endif
-		float RefBluringRadius = ComputeBluringRadiusOfNormalizedSample(RefSamples.Array[0], RefSceneMetadata) * View.ViewSizeAndInvSize.x * rcp(KernelSpreadFactor);
+		float RefBluringRadius = ComputeBluringRadiusOfNormalizedSample(NormalizedRefSamples.Array[0], RefSceneMetadata) * View.ViewSizeAndInvSize.x * rcp(KernelSpreadFactor);
 		RequestedSampleCount = PI * RefBluringRadius * RefBluringRadius;
 	}
 	#elif CONFIG_SAMPLE_COUNT_POLICY == SAMPLE_COUNT_POLICY_ACCUMULATION_BASED
@@ -566,7 +578,7 @@ void MainCS(
 			SignalInput_Textures_2,
 			SignalInput_Textures_3,
 			/* inout */ SignalAccumulators);
-		
+	
 		//if (DIM_STAGE == STAGE_REJECTION_PRE_CONVOLUTION)
 		//	DebugOutput[DispatchThreadId] = float4(
 		//		KernelConfig.bForceAllAccumulation ? 1 : 0,
diff --git a/Engine/Shaders/Private/ScreenSpaceRayCast.ush b/Engine/Shaders/Private/ScreenSpaceRayCast.ush
--- a/Engine/Shaders/Private/ScreenSpaceRayCast.ush
+++ b/Engine/Shaders/Private/ScreenSpaceRayCast.ush
@@ -225,6 +225,17 @@ void ReprojectHit(float4 PrevScreenPositionScaleBias, Texture2D Texture, Sampler
 	OutPrevUV = PrevUV;
 }
 
+float ComputeRayHitSqrDistance(float3 OriginTranslatedWorld, float3 HitUVz)
+{
+	// ALU get factored out with ReprojectHit.
+	float2 HitScreenPos = (HitUVz.xy - View.ScreenPositionScaleBias.wz) / View.ScreenPositionScaleBias.xy;
+	float HitSceneDepth = ConvertFromDeviceZ(HitUVz.z);
+
+	float3 HitTranslatedWorld = mul(float4(HitScreenPos * HitSceneDepth, HitSceneDepth, 1), View.ScreenToTranslatedWorld).xyz;
+
+	return length2(OriginTranslatedWorld - HitTranslatedWorld);
+}
+
 float4 SampleScreenColor(Texture2D Texture, SamplerState Sampler, float2 UV)
 {
 	float4 OutColor;
diff --git a/Engine/Shaders/Private/ScreenSpaceReflections.usf b/Engine/Shaders/Private/ScreenSpaceReflections.usf
--- a/Engine/Shaders/Private/ScreenSpaceReflections.usf
+++ b/Engine/Shaders/Private/ScreenSpaceReflections.usf
@@ -10,6 +10,11 @@
 #include "MonteCarlo.ush"
 #include "ScreenSpaceRayCast.ush"
 #include "SceneViewFamilyBlackboard.ush"
+#include "ScreenSpaceDenoise/SSDPublic.ush"
+
+// Generate vector truncation warnings to errors.
+#pragma warning(error: 3206)
+
 
 // .r:Intensity in 0..1 range .g:RoughnessMaskMul, b:EnableDiscard for FPostProcessScreenSpaceReflectionsStencilPS, a:(bTemporalAAIsOn?TemporalAAIndex:StateFrameIndexMod8)*1551
 float4 SSRParams;
@@ -57,9 +62,15 @@ float GetRoughnessFade(in float Roughness)
 // This optimization use the stencil to test to prunes out the pixel
 // that dones't need to compute SSR, before pixel shader invocation are
 // packed into front lines.
-void ScreenSpaceReflectionsStencilPS(in float4 UVAndScreenPos : TEXCOORD0, out float4 OutColor : SV_Target0)
+void ScreenSpaceReflectionsStencilPS(
+	float4 SvPosition : SV_POSITION
+	, out float4 OutColor : SV_Target0
+#if SSR_OUTPUT_FOR_DENOISER
+	, out float4 OutClosestHitDistance : SV_Target1
+#endif
+)
 {
-	const float2 UV = UVAndScreenPos.xy;
+	float2 UV = SvPosition.xy * View.BufferSizeAndInvSize.zw;
 	
 	FGBufferData GBuffer = GetGBufferDataFromSceneBlackboard(UV);
 	
@@ -77,18 +88,28 @@ void ScreenSpaceReflectionsStencilPS(in float4 UVAndScreenPos : TEXCOORD0, out f
 	// we are not going to compute SSR for this pixel, so we clear the color
 	// since ScreenSpaceReflectionsPS() won't be executed in this pixel.
 	OutColor = 0;
+	#if SSR_OUTPUT_FOR_DENOISER
+		OutClosestHitDistance = DENOISER_INVALID_HIT_DISTANCE;
+	#endif
 }
 
-
 void ScreenSpaceReflectionsPS(
-	float4 SvPosition : SV_POSITION,
-	out float4 OutColor : SV_Target0)
+	float4 SvPosition : SV_POSITION
+	, out float4 OutColor : SV_Target0
+#if SSR_OUTPUT_FOR_DENOISER
+	, out float4 OutClosestHitDistance : SV_Target1
+#endif
+)
 {
 	float2 UV = SvPosition.xy * View.BufferSizeAndInvSize.zw;
 	float2 ScreenPos = ViewportUVToScreenPos((SvPosition.xy - View.ViewRectMin.xy) * View.ViewSizeAndInvSize.zw);
 	uint2 PixelPos = (uint2)SvPosition.xy;
 
 	OutColor = 0;
+	
+	#if SSR_OUTPUT_FOR_DENOISER
+		OutClosestHitDistance = -2.0;
+	#endif
 
 	FGBufferData GBuffer = GetGBufferDataFromSceneBlackboard(UV);
 	
@@ -112,18 +133,28 @@ void ScreenSpaceReflectionsPS(
 	float NoV = saturate( dot( N, V ) );
 	float G_SmithV = 2 * NoV / (NoV + sqrt(NoV * (NoV - NoV * a2) + a2));
 
+	float ClosestHitDistanceSqr = 1.#INF;
+
 #if SSR_QUALITY == 1
 	uint NumSteps = 8;
 	uint NumRays = 1;
+	bool bGlossy = false;
 #elif SSR_QUALITY == 2
 	uint NumSteps = 16;
 	uint NumRays = 1;
+	#if SSR_OUTPUT_FOR_DENOISER
+		bool bGlossy = true;
+	#else
+		bool bGlossy = false;
+	#endif
 #elif SSR_QUALITY == 3
 	uint NumSteps = 8;
 	uint NumRays = 4;
+	bool bGlossy = true;
 #else // SSR_QUALITY == 4
 	uint NumSteps = 12;
 	uint NumRays = 12;
+	bool bGlossy = true;
 #endif
 	
 	if( NumRays > 1 )
@@ -155,7 +186,6 @@ void ScreenSpaceReflectionsPS(
 			float2 E = Hammersley16( i, NumRays, Random );
 
 #if 1
-			// TODO EnvBRDF isn't a perfect match anymore.
 			float3 H = mul( ImportanceSampleVisibleGGX(E, a2, TangentV ).xyz, TangentBasis );
 			float3 L = 2 * dot( V, H ) * H - V;
 #elif 0
@@ -178,83 +208,106 @@ void ScreenSpaceReflectionsPS(
 			float NoH = saturate( dot(N, H) );
 			float VoH = saturate( dot(V, H) );
 		
-			if( NoL > 0 && VoH > 0 )
-			{
-				float D = D_GGX( a2, NoH );
-				float G_SmithL = 2 * NoL / (NoL + sqrt(NoL * (NoL - NoL * a2) + a2));
-				float Vis = Vis_Smith( a2, NoV,  NoL );
-				float3 F = F_Schlick( GBuffer.SpecularColor, VoH );
+			float D = D_GGX( a2, NoH );
+			float G_SmithL = 2 * NoL / (NoL + sqrt(NoL * (NoL - NoL * a2) + a2));
+			float Vis = Vis_Smith( a2, NoV,  NoL );
+			float3 F = F_Schlick( GBuffer.SpecularColor, VoH );
 				
-				float3 HitUVz;
-				float Level = 0;
-
-				bool bHit = RayCast(
-					HZB, HZBSampler,
-					PositionTranslatedWorld, L, Roughness, SceneDepth, 
-					NumSteps, StepOffset,
-					HZBUvFactorAndInvFactor,
-					HitUVz,
-					Level
-				);
+			float3 HitUVz;
+			float Level = 0;
+
+			bool bHit = RayCast(
+				HZB, HZBSampler,
+				PositionTranslatedWorld, L, Roughness, SceneDepth, 
+				NumSteps, StepOffset,
+				HZBUvFactorAndInvFactor,
+				HitUVz,
+				Level
+			);
 
 #if 0	// Backface check
-				if( bHit )
-				{
-					float3 SampleNormal = GetGBufferData( HitUVz.xy ).WorldNormal;
-					bHit = dot( SampleNormal, L ) < 0;
-				}
+			if( bHit )
+			{
+				float3 SampleNormal = GetGBufferData( HitUVz.xy ).WorldNormal;
+				bHit = dot( SampleNormal, L ) < 0;
+			}
 #endif
 
-				// if there was a hit
-				BRANCH if( bHit )
-				{
-					float2 SampleUV;
-					float Vignette;
-					ReprojectHit(PrevScreenPositionScaleBias, SceneVelocityBuffer, SceneVelocityBufferSampler, HitUVz, SampleUV, Vignette);
+			// if there was a hit
+			BRANCH if( bHit )
+			{
+				ClosestHitDistanceSqr = min(ClosestHitDistanceSqr, ComputeRayHitSqrDistance(PositionTranslatedWorld, HitUVz));
 
-					float4 SampleColor = SampleScreenColor( SceneColor, SceneColorSampler, SampleUV ) * Vignette;
+				float2 SampleUV;
+				float Vignette;
+				ReprojectHit(PrevScreenPositionScaleBias, SceneVelocityBuffer, SceneVelocityBufferSampler, HitUVz, SampleUV, Vignette);
 
-					SampleColor.rgb *= rcp( 1 + Luminance(SampleColor.rgb) );
+				float4 SampleColor = SampleScreenColor( SceneColor, SceneColorSampler, SampleUV ) * Vignette;
 
-					// Correct integration applies DGF below but for speed we apply EnvBRDF later when compositing
+				SampleColor.rgb *= rcp( 1 + Luminance(SampleColor.rgb) );
+
+				// Correct integration applies DGF below but for speed we apply EnvBRDF later when compositing
 #if 0
-					// CosineSampleHemisphere
-					// PDF = NoL / PI,
-					SampleColor.rgb *= F * (D * Vis * PI);
+				// CosineSampleHemisphere
+				// PDF = NoL / PI,
+				SampleColor.rgb *= F * (D * Vis * PI);
 #elif 0
-					// ImportanceSampleGGX
-					// PDF = D * NoH / (4 * VoH),
-					SampleColor.rgb *= F * ( NoL * Vis * (4 * VoH / NoH) );
+				// ImportanceSampleGGX
+				// PDF = D * NoH / (4 * VoH),
+				SampleColor.rgb *= F * ( NoL * Vis * (4 * VoH / NoH) );
 #elif 0
-					// ImportanceSampleVisibleGGX
-					// PDF = G_SmithV * VoH * D / NoV / (4 * VoH)
-					// PDF = G_SmithV * D / (4 * NoV);
-					SampleColor.rgb *= F * G_SmithL;
+				// ImportanceSampleVisibleGGX
+				// PDF = G_SmithV * VoH * D / NoV / (4 * VoH)
+				// PDF = G_SmithV * D / (4 * NoV);
+				SampleColor.rgb *= F * G_SmithL;
 #endif
 
-					OutColor += SampleColor;
-				}
-
-				Count++;
+				OutColor += SampleColor;
 			}
 		}
 
 		//OutColor /= NumRays;
-		OutColor /= max( Count, 0.0001 );
+		OutColor /= max( NumRays, 0.0001 );
 		OutColor.rgb *= rcp( 1 - Luminance(OutColor.rgb) );
+
+		if (ClosestHitDistanceSqr == 1.#INF)
+		{
+			ClosestHitDistanceSqr = DENOISER_MISS_HIT_DISTANCE;
+		}
 	}
 	else
 	{
 		float StepOffset = InterleavedGradientNoise(SvPosition.xy, View.StateFrameIndexMod8);
 		StepOffset -= 0.9;
 		
-		float3 R = reflect( -V, N );
+		float3 L;
+		if (bGlossy)
+		{
+			uint2 Random = Rand3DPCG16( int3( PixelPos, View.StateFrameIndexMod8 ) ).xy;
+			float2 E = float2(Random) * rcp(65536.0);
+			
+			float3x3 TangentBasis = GetTangentBasis( N );
+			float3 TangentV = mul( TangentBasis, V );
+			
+			#if SSR_OUTPUT_FOR_DENOISER
+			{
+				E.x *= 1 - GGX_IMPORTANT_SAMPLE_BIAS;
+			}
+			#endif
+
+			float3 H = mul( ImportanceSampleVisibleGGX(E, a2, TangentV ).xyz, TangentBasis );
+			L = 2 * dot( V, H ) * H - V;
+		}
+		else
+		{
+			L = reflect( -V, N );
+		}
+		
 		float3 HitUVz;
 		float Level = 0;
-
 		bool bHit = RayCast(
 			HZB, HZBSampler,
-			PositionTranslatedWorld, R, Roughness, SceneDepth,
+			PositionTranslatedWorld, L, Roughness, SceneDepth,
 			NumSteps, StepOffset,
 			HZBUvFactorAndInvFactor,
 			HitUVz,
@@ -263,12 +316,18 @@ void ScreenSpaceReflectionsPS(
 
 		BRANCH if( bHit )
 		{
+			ClosestHitDistanceSqr = ComputeRayHitSqrDistance(PositionTranslatedWorld, HitUVz);
+
 			float2 SampleUV;
 			float Vignette;
 			ReprojectHit(PrevScreenPositionScaleBias, SceneVelocityBuffer, SceneVelocityBufferSampler, HitUVz, SampleUV, Vignette);
 
 			OutColor = SampleScreenColor(SceneColor, SceneColorSampler, SampleUV) * Vignette;
 		}
+		else
+		{
+			ClosestHitDistanceSqr = DENOISER_MISS_HIT_DISTANCE;
+		}
 	}
 	
 	OutColor *= RoughnessFade;
@@ -277,4 +336,11 @@ void ScreenSpaceReflectionsPS(
 #if USE_PREEXPOSURE
 	OutColor.rgb *= PrevSceneColorPreExposureCorrection;
 #endif
+	
+	// Output closest hit distance for denoiser.
+	#if SSR_OUTPUT_FOR_DENOISER
+	{
+		OutClosestHitDistance = ClosestHitDistanceSqr > 0 ? ClampToHalfFloatRange(sqrt(ClosestHitDistanceSqr)) : ClosestHitDistanceSqr;
+	}
+	#endif
 }
\ No newline at end of file
diff --git a/Engine/Source/Runtime/Renderer/Private/PostProcess/ScreenSpaceReflections.cpp b/Engine/Source/Runtime/Renderer/Private/PostProcess/ScreenSpaceReflections.cpp
--- a/Engine/Source/Runtime/Renderer/Private/PostProcess/ScreenSpaceReflections.cpp
+++ b/Engine/Source/Runtime/Renderer/Private/PostProcess/ScreenSpaceReflections.cpp
@@ -72,12 +72,10 @@ bool ShouldRenderScreenSpaceReflections(const FViewInfo& View)
 	return true;
 }
 
-bool IsSSRTemporalPassRequired(const FViewInfo& View, bool bCheckSSREnabled)
+bool IsSSRTemporalPassRequired(const FViewInfo& View)
 {
-	if (bCheckSSREnabled && !ShouldRenderScreenSpaceReflections(View))
-	{
-		return false;
-	}
+	check(ShouldRenderScreenSpaceReflections(View));
+
 	if (!View.State)
 	{
 		return false;
@@ -88,7 +86,7 @@ bool IsSSRTemporalPassRequired(const FViewInfo& View, bool bCheckSSREnabled)
 namespace
 {
 
-float ComputeRoughnessMaskScale(const FViewInfo& View, uint32 SSRQuality)
+float ComputeRoughnessMaskScale(const FViewInfo& View, ESSRQuality SSRQuality)
 {
 	float MaxRoughness = FMath::Clamp(View.FinalPostProcessSettings.ScreenSpaceReflectionMaxRoughness, 0.01f, 1.0f);
 
@@ -97,10 +95,10 @@ float ComputeRoughnessMaskScale(const FViewInfo& View, uint32 SSRQuality)
 	// f(MaxRoughness/2) = 1
 
 	float RoughnessMaskScale = -2.0f / MaxRoughness;
-	return RoughnessMaskScale * (SSRQuality < 3 ? 2.0f : 1.0f);
+	return RoughnessMaskScale * (int32(SSRQuality) < 3 ? 2.0f : 1.0f);
 }
 
-FLinearColor ComputeSSRParams(const FViewInfo& View, uint32 SSRQuality, bool bEnableDiscard)
+FLinearColor ComputeSSRParams(const FViewInfo& View, ESSRQuality SSRQuality, bool bEnableDiscard)
 {
 	float RoughnessMaskScale = ComputeRoughnessMaskScale(View, SSRQuality);
 
@@ -138,12 +136,18 @@ BEGIN_SHADER_PARAMETER_STRUCT(FSSRCommonParameters, )
 END_SHADER_PARAMETER_STRUCT()
 
 
+
+class FSSRQualityDim : SHADER_PERMUTATION_ENUM_CLASS("SSR_QUALITY", ESSRQuality);
+class FSSRPrevFrameColorDim : SHADER_PERMUTATION_BOOL("PREV_FRAME_COLOR");
+class FSSROutputForDenoiser : SHADER_PERMUTATION_BOOL("SSR_OUTPUT_FOR_DENOISER");
+
+
 class FScreenSpaceReflectionsStencilPS : public FGlobalShader
 {
 	DECLARE_GLOBAL_SHADER(FScreenSpaceReflectionsStencilPS);
 	SHADER_USE_PARAMETER_STRUCT(FScreenSpaceReflectionsStencilPS, FGlobalShader);
 
-	using FPermutationDomain = FShaderPermutationNone;
+	using FPermutationDomain = TShaderPermutationDomain<FSSROutputForDenoiser>;
 
 	static bool ShouldCompilePermutation(const FGlobalShaderPermutationParameters& Parameters)
 	{
@@ -163,25 +167,17 @@ class FScreenSpaceReflectionsStencilPS : public FGlobalShader
 	END_SHADER_PARAMETER_STRUCT()
 };
 
-
-static constexpr int32 QualityCount = 5;
-
-class FSSRQualityDim : SHADER_PERMUTATION_INT("SSR_QUALITY", QualityCount);
-class FSSRPrevFrameColorDim : SHADER_PERMUTATION_BOOL("PREV_FRAME_COLOR");
-class FSSRPrevFrameColorDim2 : SHADER_PERMUTATION_BOOL("PREV_FRAME_COLOR");
-
-
 class FScreenSpaceReflectionsPS : public FGlobalShader
 {
 	DECLARE_GLOBAL_SHADER(FScreenSpaceReflectionsPS);
 	SHADER_USE_PARAMETER_STRUCT(FScreenSpaceReflectionsPS, FGlobalShader);
 
-	using FPermutationDomain = TShaderPermutationDomain<FSSRQualityDim, FSSRPrevFrameColorDim>;
+	using FPermutationDomain = TShaderPermutationDomain<FSSRQualityDim, FSSRPrevFrameColorDim, FSSROutputForDenoiser>;
 
 	static bool ShouldCompilePermutation(const FGlobalShaderPermutationParameters& Parameters)
 	{
 		FPermutationDomain PermutationVector(Parameters.PermutationId);
-		if ((PermutationVector.Get<FSSRQualityDim>() == 0) && PermutationVector.Get<FSSRPrevFrameColorDim>())
+		if ((PermutationVector.Get<FSSRQualityDim>() == ESSRQuality::VisualizeSSR) && PermutationVector.Get<FSSRPrevFrameColorDim>())
 		{
 			return false;
 		}
@@ -208,55 +204,92 @@ class FScreenSpaceReflectionsPS : public FGlobalShader
 IMPLEMENT_GLOBAL_SHADER(FScreenSpaceReflectionsPS, "/Engine/Private/ScreenSpaceReflections.usf", "ScreenSpaceReflectionsPS", SF_Pixel);
 IMPLEMENT_GLOBAL_SHADER(FScreenSpaceReflectionsStencilPS, "/Engine/Private/ScreenSpaceReflections.usf", "ScreenSpaceReflectionsStencilPS", SF_Pixel);
 
-// @param Quality usually in 0..100 range, default is 50
-// @return see CVarSSRQuality, never 0
-static int32 ComputeSSRQuality(float Quality)
+void GetSSRShaderOptionsForQuality(ESSRQuality Quality, IScreenSpaceDenoiser::FReflectionsRayTracingConfig* OutRayTracingConfigs)
 {
-	int32 Ret;
-
-	if(Quality >= 60.0f)
+	if (Quality == ESSRQuality::VisualizeSSR)
 	{
-		Ret = (Quality >= 80.0f) ? 4 : 3;
+		OutRayTracingConfigs->RayCountPerPixel = 12;
+	}
+	else if (Quality == ESSRQuality::Epic)
+	{
+		OutRayTracingConfigs->RayCountPerPixel = 12;
+	}
+	else if (Quality == ESSRQuality::High)
+	{
+		OutRayTracingConfigs->RayCountPerPixel = 4;
+	}
+	else if (Quality == ESSRQuality::Medium)
+	{
+		OutRayTracingConfigs->RayCountPerPixel = 1;
+	}
+	else if (Quality == ESSRQuality::Low)
+	{
+		OutRayTracingConfigs->RayCountPerPixel = 1;
 	}
 	else
 	{
-		Ret = (Quality >= 40.0f) ? 2 : 1;
+		check(0);
 	}
-
-	int SSRQualityCVar = FMath::Clamp(CVarSSRQuality.GetValueOnRenderThread(), 0, QualityCount - 1);
-
-	return FMath::Min(Ret, SSRQualityCVar);
 }
 
 } // namespace
 
-FRDGTextureRef RenderScreenSpaceReflections(
+void GetSSRQualityForView(const FViewInfo& View, ESSRQuality* OutQuality, IScreenSpaceDenoiser::FReflectionsRayTracingConfig* OutRayTracingConfigs)
+{
+	check(ShouldRenderScreenSpaceReflections(View));
+	
+	int32 SSRQualityCVar = FMath::Clamp(CVarSSRQuality.GetValueOnRenderThread(), 0, int32(ESSRQuality::MAX) - 1);
+	
+	if (View.Family->EngineShowFlags.VisualizeSSR)
+	{
+		*OutQuality = ESSRQuality::VisualizeSSR;
+		return;
+	}
+	else if (View.FinalPostProcessSettings.ScreenSpaceReflectionQuality >= 80.0f && SSRQualityCVar >= 4)
+	{
+		*OutQuality = ESSRQuality::Epic;
+	}
+	else if (View.FinalPostProcessSettings.ScreenSpaceReflectionQuality >= 60.0f && SSRQualityCVar >= 3)
+	{
+		*OutQuality = ESSRQuality::High;
+	}
+	else if (View.FinalPostProcessSettings.ScreenSpaceReflectionQuality >= 40.0f && SSRQualityCVar >= 2)
+	{
+		*OutQuality = ESSRQuality::Medium;
+	}
+	else
+	{
+		*OutQuality = ESSRQuality::Low;
+	}
+
+	GetSSRShaderOptionsForQuality(*OutQuality, OutRayTracingConfigs);
+}
+
+void RenderScreenSpaceReflections(
 	FRDGBuilder& GraphBuilder,
 	const FSceneViewFamilyBlackboard& SceneTextures,
 	const FRDGTextureRef CurrentSceneColor,
-	const FViewInfo& View)
+	const FViewInfo& View,
+	ESSRQuality SSRQuality,
+	bool bDenoiser,
+	IScreenSpaceDenoiser::FReflectionsInputs* DenoiserInputs)
 {
-	RDG_EVENT_SCOPE(GraphBuilder, "ScreenSpaceReflections");
-	
 	FRDGTextureRef InputColor = CurrentSceneColor;
-	bool bSamplePrevFrame = false;
-	if (View.PrevViewInfo.CustomSSRInput.IsValid())
-	{
-		InputColor = GraphBuilder.RegisterExternalTexture(View.PrevViewInfo.CustomSSRInput);
-		bSamplePrevFrame = true;
-	}
-	else if (View.PrevViewInfo.TemporalAAHistory.IsValid())
+	if (SSRQuality != ESSRQuality::VisualizeSSR)
 	{
-		InputColor = GraphBuilder.RegisterExternalTexture(View.PrevViewInfo.TemporalAAHistory.RT[0]);
-		bSamplePrevFrame = true;
+		if (View.PrevViewInfo.CustomSSRInput.IsValid())
+		{
+			InputColor = GraphBuilder.RegisterExternalTexture(View.PrevViewInfo.CustomSSRInput);
+		}
+		else if (View.PrevViewInfo.TemporalAAHistory.IsValid())
+		{
+			InputColor = GraphBuilder.RegisterExternalTexture(View.PrevViewInfo.TemporalAAHistory.RT[0]);
+		}
 	}
 
-	const bool VisualizeSSR = View.Family->EngineShowFlags.VisualizeSSR;
-	const bool SSRStencilPrePass = CVarSSRStencil.GetValueOnRenderThread() != 0 && !VisualizeSSR;
-	const int32 SSRQuality = VisualizeSSR ? 0 : FMath::Clamp(ComputeSSRQuality(View.FinalPostProcessSettings.ScreenSpaceReflectionQuality), 1, 4);
+	const bool SSRStencilPrePass = CVarSSRStencil.GetValueOnRenderThread() != 0 && SSRQuality != ESSRQuality::VisualizeSSR;
 	
-	// Alloc SSR output.
-	FRDGTextureRef SSROutput;
+	// Alloc inputs for denoising.
 	{
 		FRDGTextureDesc Desc = FPooledRenderTargetDesc::Create2DDesc(
 			FSceneRenderTargets::Get_FrameConstantsOnly().GetBufferSizeXY(),
@@ -267,8 +300,17 @@ FRDGTextureRef RenderScreenSpaceReflections(
 		Desc.AutoWritable = false;
 		Desc.Flags |= GFastVRamConfig.SSR;
 
-		SSROutput = GraphBuilder.CreateTexture(Desc, TEXT("ScreenSpaceReflections"));
+		DenoiserInputs->Color = GraphBuilder.CreateTexture(Desc, TEXT("ScreenSpaceReflections"));
+
+		if (bDenoiser)
+		{
+			Desc.Format = PF_R16F;
+			DenoiserInputs->RayHitDistance = GraphBuilder.CreateTexture(Desc, TEXT("ScreenSpaceReflectionsHitDistance"));
+		}
 	}
+
+	IScreenSpaceDenoiser::FReflectionsRayTracingConfig RayTracingConfigs;
+	GetSSRShaderOptionsForQuality(SSRQuality, &RayTracingConfigs);
 		
 	FSSRCommonParameters CommonParameters;
 	CommonParameters.SSRParams = ComputeSSRParams(View, SSRQuality, false);
@@ -277,7 +319,12 @@ FRDGTextureRef RenderScreenSpaceReflections(
 	SetupSceneTextureSamplers(&CommonParameters.SceneTextureSamplers);
 	
 	FRenderTargetBindingSlots RenderTargets;
-	RenderTargets[0] = FRenderTargetBinding(SSROutput, ERenderTargetLoadAction::ENoAction, ERenderTargetStoreAction::EStore);
+	RenderTargets[0] = FRenderTargetBinding(DenoiserInputs->Color, ERenderTargetLoadAction::ENoAction, ERenderTargetStoreAction::EStore);
+
+	if (bDenoiser)
+	{
+		RenderTargets[1] = FRenderTargetBinding(DenoiserInputs->RayHitDistance, ERenderTargetLoadAction::ENoAction, ERenderTargetStoreAction::EStore);
+	}
 
 	// Do a pre pass that output 0, or set a stencil mask to run the more expensive pixel shader.
 	if (SSRStencilPrePass)
@@ -289,11 +336,14 @@ FRDGTextureRef RenderScreenSpaceReflections(
 			ERenderTargetLoadAction::ELoad, ERenderTargetStoreAction::EStore,
 			FExclusiveDepthStencil::DepthRead_StencilWrite);
 
+		FScreenSpaceReflectionsStencilPS::FPermutationDomain PermutationVector;
+		PermutationVector.Set<FSSROutputForDenoiser>(bDenoiser);
+
 		FScreenSpaceReflectionsStencilPS::FParameters* PassParameters = GraphBuilder.AllocParameters<FScreenSpaceReflectionsStencilPS::FParameters>();
 		PassParameters->CommonParameters = CommonParameters;
 		PassParameters->RenderTargets = RenderTargets;
 		
-		TShaderMapRef<FScreenSpaceReflectionsStencilPS> PixelShader(View.ShaderMap);
+		TShaderMapRef<FScreenSpaceReflectionsStencilPS> PixelShader(View.ShaderMap, PermutationVector);
 		ClearUnusedGraphResources(*PixelShader, PassParameters);
 		
 		GraphBuilder.AddPass(
@@ -321,8 +371,9 @@ FRDGTextureRef RenderScreenSpaceReflections(
 	// Adds SSR pass.
 	{
 		FScreenSpaceReflectionsPS::FPermutationDomain PermutationVector;
-		PermutationVector.Set<FSSRPrevFrameColorDim>(bSamplePrevFrame && SSRQuality != 0);
+		PermutationVector.Set<FSSRPrevFrameColorDim>(InputColor != CurrentSceneColor);
 		PermutationVector.Set<FSSRQualityDim>(SSRQuality);
+		PermutationVector.Set<FSSROutputForDenoiser>(bDenoiser);
 
 		FScreenSpaceReflectionsPS::FParameters* PassParameters = GraphBuilder.AllocParameters<FScreenSpaceReflectionsPS::FParameters>();
 		PassParameters->CommonParameters = CommonParameters;
@@ -356,7 +407,7 @@ FRDGTextureRef RenderScreenSpaceReflections(
 				(ViewportExtent.X * 0.5f + ViewportOffset.X) * InvBufferSize.X,
 				(ViewportExtent.Y * 0.5f + ViewportOffset.Y) * InvBufferSize.Y);
 		}
-		PassParameters->PrevSceneColorPreExposureCorrection = bSamplePrevFrame ? View.PreExposure / View.PrevViewInfo.SceneColorPreExposure : 1.0f;
+		PassParameters->PrevSceneColorPreExposureCorrection = InputColor != CurrentSceneColor ? View.PreExposure / View.PrevViewInfo.SceneColorPreExposure : 1.0f;
 		
 		PassParameters->SceneColor = InputColor;
 		PassParameters->SceneColorSampler = TStaticSamplerState<SF_Point>::GetRHI();
@@ -370,8 +421,9 @@ FRDGTextureRef RenderScreenSpaceReflections(
 		ClearUnusedGraphResources(*PixelShader, PassParameters);
 		
 		GraphBuilder.AddPass(
-			RDG_EVENT_NAME("SSR RayMarch(Quality=%d) %dx%d",
-				SSRQuality, View.ViewRect.Width(), View.ViewRect.Height()),
+			RDG_EVENT_NAME("SSR RayMarch(Quality=%d RayPerPixel=%d%s) %dx%d",
+				SSRQuality, RayTracingConfigs.RayCountPerPixel, bDenoiser ? TEXT(" DenoiserOutput") : TEXT(""),
+				View.ViewRect.Width(), View.ViewRect.Height()),
 			PassParameters,
 			ERenderGraphPassFlags::None,
 			[PassParameters, &View, PixelShader, SSRStencilPrePass](FRHICommandList& RHICmdList)
@@ -394,6 +446,4 @@ FRDGTextureRef RenderScreenSpaceReflections(
 			FPixelShaderUtils::DrawFullscreenTriangle(RHICmdList);
 		});
 	}
-
-	return SSROutput;
 }
diff --git a/Engine/Source/Runtime/Renderer/Private/PostProcess/ScreenSpaceReflections.h b/Engine/Source/Runtime/Renderer/Private/PostProcess/ScreenSpaceReflections.h
--- a/Engine/Source/Runtime/Renderer/Private/PostProcess/ScreenSpaceReflections.h
+++ b/Engine/Source/Runtime/Renderer/Private/PostProcess/ScreenSpaceReflections.h
@@ -7,16 +7,34 @@
 #pragma once
 
 #include "RenderGraph.h"
+#include "ScreenSpaceDenoise.h"
 
 class FViewInfo;
 class FSceneViewFamilyBlackboard;
 
-FRDGTextureRef RenderScreenSpaceReflections(
-	FRDGBuilder& GraphBuilder,
-	const FSceneViewFamilyBlackboard& SceneBlackboard,
-	const FRDGTextureRef CurrentSceneColor,
-	const FViewInfo& View);
+enum class ESSRQuality
+{
+	VisualizeSSR,
+
+	Low,
+	Medium,
+	High,
+	Epic,
+
+	MAX
+};
 
 bool ShouldRenderScreenSpaceReflections(const FViewInfo& View);
 
-bool IsSSRTemporalPassRequired(const FViewInfo& View, bool bCheckSSREnabled = true);
+void GetSSRQualityForView(const FViewInfo& View, ESSRQuality* OutQuality, IScreenSpaceDenoiser::FReflectionsRayTracingConfig* OutRayTracingConfigs);
+
+bool IsSSRTemporalPassRequired(const FViewInfo& View);
+
+void RenderScreenSpaceReflections(
+	FRDGBuilder& GraphBuilder,
+	const FSceneViewFamilyBlackboard& SceneBlackboard,
+	const FRDGTextureRef CurrentSceneColor,
+	const FViewInfo& View,
+	ESSRQuality SSRQuality,
+	bool bDenoiser,
+	IScreenSpaceDenoiser::FReflectionsInputs* DenoiserInputs);
diff --git a/Engine/Source/Runtime/Renderer/Private/ReflectionEnvironment.cpp b/Engine/Source/Runtime/Renderer/Private/ReflectionEnvironment.cpp
--- a/Engine/Source/Runtime/Renderer/Private/ReflectionEnvironment.cpp
+++ b/Engine/Source/Runtime/Renderer/Private/ReflectionEnvironment.cpp
@@ -139,6 +139,11 @@ static TAutoConsoleVariable<int32> CVarUseReflectionDenoiser(
 	TEXT(" 2: GScreenSpaceDenoiser which may be overriden by a third party plugin (default)."),
 	ECVF_RenderThreadSafe);
 
+static TAutoConsoleVariable<int32> CVarDenoiseSSR(
+	TEXT("r.SSR.ExperimentalDenoiser"), 0,
+	TEXT("Replace SSR's TAA pass with denoiser."),
+	ECVF_RenderThreadSafe);
+
 
 // to avoid having direct access from many places
 static int GetReflectionEnvironmentCVar()
@@ -768,39 +773,61 @@ void FDeferredShadingSceneRenderer::RenderDeferredReflectionsAndSkyLighting(FRHI
 		const bool bScreenSpaceReflections = !bRayTracedReflections && ShouldRenderScreenSpaceReflections(Views[ViewIndex]);
 
 		TRefCountPtr<IPooledRenderTarget> ReflectionsColor = GSystemTextures.BlackDummy;
-		if (bRayTracedReflections)
+		if (bRayTracedReflections || bScreenSpaceReflections)
 		{
-#if RHI_RAYTRACING
-			SCOPED_DRAW_EVENT(RHICmdList, RayTracingReflections);
-			SCOPED_GPU_STAT(RHICmdList, RayTracingReflections);
-
 			FRDGBuilder GraphBuilder(RHICmdList); // TODO: convert the entire reflections to render graph.
 
 			FSceneViewFamilyBlackboard SceneBlackboard;
 			SetupSceneViewFamilyBlackboard(GraphBuilder, &SceneBlackboard);
 
+			int32 DenoiserMode = CVarUseReflectionDenoiser.GetValueOnRenderThread();
+			
+			bool bDenoise = false;
+			bool bTemporalFilter = false;
 
+			// Traces the reflections, either using screen space reflection, or ray tracing.
+			IScreenSpaceDenoiser::FReflectionsInputs DenoiserInputs;
 			IScreenSpaceDenoiser::FReflectionsRayTracingConfig RayTracingConfig;
-			RayTracingConfig.ResolutionFraction = FMath::Clamp(CVarReflectionScreenPercentage.GetValueOnRenderThread() / 100.0f, 0.25f, 1.0f);
+			if (bRayTracedReflections)
+			{
+				RDG_EVENT_SCOPE(GraphBuilder, "RayTracingReflections");
+				RDG_GPU_STAT_SCOPE(GraphBuilder, RayTracingReflections);
 
-			int32 RayTracingReflectionsSPP = GRayTracingReflectionsSamplesPerPixel > -1 ? GRayTracingReflectionsSamplesPerPixel : View.FinalPostProcessSettings.RayTracingReflectionsSamplesPerPixel;
-			int32 DenoiserMode = CVarUseReflectionDenoiser.GetValueOnRenderThread();
-			const bool bDenoise = DenoiserMode != 0 && RayTracingReflectionsSPP == 1;
+				RayTracingConfig.ResolutionFraction = FMath::Clamp(CVarReflectionScreenPercentage.GetValueOnRenderThread() / 100.0f, 0.25f, 1.0f);
+				RayTracingConfig.RayCountPerPixel = GRayTracingReflectionsSamplesPerPixel > -1 ? GRayTracingReflectionsSamplesPerPixel : View.FinalPostProcessSettings.RayTracingReflectionsSamplesPerPixel;
 
-			if (!bDenoise)
-			{
-				RayTracingConfig.ResolutionFraction = 1.0f;
+				bDenoise = DenoiserMode != 0 && RayTracingConfig.RayCountPerPixel == 1;
+				
+				if (!bDenoise)
+				{
+					RayTracingConfig.ResolutionFraction = 1.0f;
+				}
+
+				RenderRayTracingReflections(
+					GraphBuilder,
+					View, &DenoiserInputs.Color, &DenoiserInputs.RayHitDistance, &DenoiserInputs.RayImaginaryDepth,
+					RayTracingConfig.RayCountPerPixel, GRayTracingReflectionsHeightFog, RayTracingConfig.ResolutionFraction);
 			}
+			else if (bScreenSpaceReflections)
+			{
+				bDenoise = DenoiserMode != 0 && CVarDenoiseSSR.GetValueOnRenderThread();
+				bTemporalFilter = !bDenoise && View.ViewState && IsSSRTemporalPassRequired(View);
 
-			// Ray trace the reflection.
-			IScreenSpaceDenoiser::FReflectionsInputs DenoiserInputs;
-			RenderRayTracingReflections(
-				GraphBuilder,
-				View, &DenoiserInputs.Color, &DenoiserInputs.RayHitDistance, &DenoiserInputs.RayImaginaryDepth,
-				RayTracingReflectionsSPP, GRayTracingReflectionsHeightFog, RayTracingConfig.ResolutionFraction);
+				FRDGTextureRef CurrentSceneColor = GraphBuilder.RegisterExternalTexture(SceneContext.GetSceneColor());
 
+				ESSRQuality SSRQuality;
+				GetSSRQualityForView(View, &SSRQuality, &RayTracingConfig);
+				
+				RDG_EVENT_SCOPE(GraphBuilder, "ScreenSpaceReflections(Quality=%d)", int32(SSRQuality));
+			
+				RenderScreenSpaceReflections(
+					GraphBuilder, SceneBlackboard, CurrentSceneColor, View, SSRQuality, bDenoise, &DenoiserInputs);
+			}
+			else
+			{
+				check(0);
+			}
 
-			// Denoise the reflections.
 			if (bDenoise)
 			{
 				const IScreenSpaceDenoiser* DefaultDenoiser = IScreenSpaceDenoiser::GetDefaultDenoiser();
@@ -822,52 +849,37 @@ void FDeferredShadingSceneRenderer::RenderDeferredReflectionsAndSkyLighting(FRHI
 
 				GraphBuilder.QueueTextureExtraction(DenoiserOutputs.Color, &ReflectionsColor);
 			}
-			else
-			{
-				// The performance of ray tracing does not allow to run without a denoiser in real time.
-				// Multiple rays per pixel is unsupported by the denoiser that will most likely more bound by to
-				// many rays than exporting the hit distance buffer. Therefore no permutation of the ray generation
-				// shader has been judged required to be supported.
-				GraphBuilder.RemoveUnusedTextureWarning(DenoiserInputs.RayHitDistance);
-
-				GraphBuilder.QueueTextureExtraction(DenoiserInputs.Color, &ReflectionsColor);
-			}
-
-			GraphBuilder.Execute();
-#endif
-		}
-		else if (bScreenSpaceReflections)
-		{
-			FRDGBuilder GraphBuilder(RHICmdList); // TODO: convert the entire reflections to render graph.
-			
-			FSceneViewFamilyBlackboard SceneBlackboard;
-			SetupSceneViewFamilyBlackboard(GraphBuilder, &SceneBlackboard);
-
-			FRDGTextureRef CurrentSceneColor = GraphBuilder.RegisterExternalTexture(SceneContext.GetSceneColor());
-
-			FRDGTextureRef ScreenSpaceReflections = RenderScreenSpaceReflections(
-				GraphBuilder, SceneBlackboard, CurrentSceneColor, View);
-			
-			// Adds a temporal AA pass to clean up some of the noise.
-			const bool bTemporalFilter = IsSSRTemporalPassRequired(View, false);
-	
-			// TODO: denoise.
-			if (View.ViewState && bTemporalFilter)
+			else if (bTemporalFilter)
 			{
+				check(View.ViewState);
 				FTAAPassParameters TAASettings(View);
 				TAASettings.Pass = ETAAPassConfig::ScreenSpaceReflections;
-				TAASettings.SceneColorInput = ScreenSpaceReflections;
-
-				ScreenSpaceReflections = TAASettings.AddTemporalAAPass(
+				TAASettings.SceneColorInput = DenoiserInputs.Color;
+				
+				FTAAOutputs TAAOutputs = TAASettings.AddTemporalAAPass(
 					GraphBuilder,
 					SceneBlackboard, View,
 					View.PrevViewInfo.SSRHistory,
-					&View.ViewState->PrevFrameViewInfo.SSRHistory).SceneColor;
+					&View.ViewState->PrevFrameViewInfo.SSRHistory);
+				
+				GraphBuilder.QueueTextureExtraction(TAAOutputs.SceneColor, &ReflectionsColor);
+			}
+			else
+			{
+				if (bRayTracedReflections && DenoiserInputs.RayHitDistance)
+				{
+					// The performance of ray tracing does not allow to run without a denoiser in real time.
+					// Multiple rays per pixel is unsupported by the denoiser that will most likely more bound by to
+					// many rays than exporting the hit distance buffer. Therefore no permutation of the ray generation
+					// shader has been judged required to be supported.
+					GraphBuilder.RemoveUnusedTextureWarning(DenoiserInputs.RayHitDistance);
+				}
+
+				GraphBuilder.QueueTextureExtraction(DenoiserInputs.Color, &ReflectionsColor);
 			}
 
-			GraphBuilder.QueueTextureExtraction(ScreenSpaceReflections, &ReflectionsColor);
 			GraphBuilder.Execute();
-		}
+		} // if (bRayTracedReflections || bScreenSpaceReflections)
 
 		bool bPlanarReflections = false;
 		if (!bRayTracedReflections)
diff --git a/Engine/Source/Runtime/Renderer/Private/ScreenSpaceDenoise.cpp b/Engine/Source/Runtime/Renderer/Private/ScreenSpaceDenoise.cpp
--- a/Engine/Source/Runtime/Renderer/Private/ScreenSpaceDenoise.cpp
+++ b/Engine/Source/Runtime/Renderer/Private/ScreenSpaceDenoise.cpp
@@ -196,7 +196,9 @@ static int32 SignalMaxBatchSize(ESignalProcessing SignalProcessing)
 /** Returns whether a signal can denoise multi sample per pixel. */
 static bool SignalSupportMultiSPP(ESignalProcessing SignalProcessing)
 {
-	return SignalProcessing == ESignalProcessing::MonochromaticPenumbra;
+	return (
+		SignalProcessing == ESignalProcessing::MonochromaticPenumbra ||
+		SignalProcessing == ESignalProcessing::Reflections);
 }
 
 
@@ -401,20 +403,25 @@ static_assert(ARRAY_COUNT(kHistoryConvolutionResourceNames) == int32(ESignalProc
 static_assert(ARRAY_COUNT(kDenoiserOutputResourceNames) == int32(ESignalProcessing::MAX) * kMaxBufferProcessingCount, "You forgot me!");
 
 
-/** Base class for a screen space denoising shader. */
-class FScreenSpaceDenoisingShader : public FGlobalShader
+/** Returns whether should compole pipeline for a given shader platform.*/
+bool ShouldCompileSignalPipeline(ESignalProcessing SignalProcessing, EShaderPlatform Platform)
 {
-public:
-	static bool ShouldCompilePermutation(const FGlobalShaderPermutationParameters& Parameters)
+	if (SignalProcessing == ESignalProcessing::Reflections)
 	{
-		return Parameters.Platform == SP_PCD3D_SM5;
+		// Ray traced reflection and SSR.
+		return Platform == SP_PCD3D_SM5 || Platform == SP_XBOXONE_D3D12;
 	}
-
-	FScreenSpaceDenoisingShader() {}
-	FScreenSpaceDenoisingShader(const ShaderMetaType::CompiledShaderInitializerType& Initializer)
-		: FGlobalShader(Initializer)
-	{ }
-};
+	else if (
+		SignalProcessing == ESignalProcessing::MonochromaticPenumbra ||
+		SignalProcessing == ESignalProcessing::AmbientOcclusion ||
+		SignalProcessing == ESignalProcessing::GlobalIllumination)
+	{
+		// Only for ray tracing denoising.
+		return Platform == SP_PCD3D_SM5;
+	}
+	check(0);
+	return false;
+}
 
 
 /** Shader parameter structure used for all shaders. */
@@ -472,10 +479,10 @@ FSSDSignalUAVs CreateMultiplexedUAVs(FRDGBuilder& GraphBuilder, const FSSDSignal
 }
 
 
-class FSSDInjestCS : public FScreenSpaceDenoisingShader
+class FSSDInjestCS : public FGlobalShader
 {
 	DECLARE_GLOBAL_SHADER(FSSDInjestCS);
-	SHADER_USE_PARAMETER_STRUCT(FSSDInjestCS, FScreenSpaceDenoisingShader);
+	SHADER_USE_PARAMETER_STRUCT(FSSDInjestCS, FGlobalShader);
 
 	using FPermutationDomain = TShaderPermutationDomain<FSignalProcessingDim, FSignalBatchSizeDim, FMultiSPPDim>;
 
@@ -502,7 +509,7 @@ class FSSDInjestCS : public FScreenSpaceDenoisingShader
 			return false;
 		}
 
-		return FScreenSpaceDenoisingShader::ShouldCompilePermutation(Parameters);
+		return ShouldCompileSignalPipeline(SignalProcessing, Parameters.Platform);
 	}
 
 	BEGIN_SHADER_PARAMETER_STRUCT(FParameters, )
@@ -514,10 +521,10 @@ class FSSDInjestCS : public FScreenSpaceDenoisingShader
 	END_SHADER_PARAMETER_STRUCT()
 };
 
-class FSSDSpatialAccumulationCS : public FScreenSpaceDenoisingShader
+class FSSDSpatialAccumulationCS : public FGlobalShader
 {
 	DECLARE_GLOBAL_SHADER(FSSDSpatialAccumulationCS);
-	SHADER_USE_PARAMETER_STRUCT(FSSDSpatialAccumulationCS, FScreenSpaceDenoisingShader);
+	SHADER_USE_PARAMETER_STRUCT(FSSDSpatialAccumulationCS, FGlobalShader);
 
 	static const uint32 kGroupSize = 8;
 	
@@ -605,7 +612,7 @@ class FSSDSpatialAccumulationCS : public FScreenSpaceDenoisingShader
 			return false;
 		}
 
-		return FScreenSpaceDenoisingShader::ShouldCompilePermutation(Parameters);
+		return ShouldCompileSignalPipeline(SignalProcessing, Parameters.Platform);
 	}
 
 	BEGIN_SHADER_PARAMETER_STRUCT(FParameters, )
@@ -623,10 +630,10 @@ class FSSDSpatialAccumulationCS : public FScreenSpaceDenoisingShader
 	END_SHADER_PARAMETER_STRUCT()
 };
 
-class FSSDTemporalAccumulationCS : public FScreenSpaceDenoisingShader
+class FSSDTemporalAccumulationCS : public FGlobalShader
 {
 	DECLARE_GLOBAL_SHADER(FSSDTemporalAccumulationCS);
-	SHADER_USE_PARAMETER_STRUCT(FSSDTemporalAccumulationCS, FScreenSpaceDenoisingShader);
+	SHADER_USE_PARAMETER_STRUCT(FSSDTemporalAccumulationCS, FGlobalShader);
 
 	using FPermutationDomain = TShaderPermutationDomain<FSignalProcessingDim, FSignalBatchSizeDim>;
 
@@ -647,7 +654,7 @@ class FSSDTemporalAccumulationCS : public FScreenSpaceDenoisingShader
 			return false;
 		}
 
-		return FScreenSpaceDenoisingShader::ShouldCompilePermutation(Parameters);
+		return ShouldCompileSignalPipeline(SignalProcessing, Parameters.Platform);
 	}
 
 	BEGIN_SHADER_PARAMETER_STRUCT(FParameters, )
@@ -1334,6 +1341,7 @@ public:
 		Settings.ReconstructionSamples = CVarReflectionReconstructionSampleCount.GetValueOnRenderThread();
 		Settings.bUseTemporalAccumulation = CVarReflectionTemporalAccumulation.GetValueOnRenderThread() != 0;
 		Settings.HistoryConvolutionSampleCount = CVarReflectionHistoryConvolutionSampleCount.GetValueOnRenderThread();
+		Settings.MaxInputSPP = RayTracingConfig.RayCountPerPixel;
 
 		TStaticArray<FScreenSpaceFilteringHistory*, IScreenSpaceDenoiser::kMaxBatchSize> PrevHistories;
 		TStaticArray<FScreenSpaceFilteringHistory*, IScreenSpaceDenoiser::kMaxBatchSize> NewHistories;
diff --git a/Engine/Source/Runtime/Renderer/Private/ScreenSpaceDenoise.h b/Engine/Source/Runtime/Renderer/Private/ScreenSpaceDenoise.h
--- a/Engine/Source/Runtime/Renderer/Private/ScreenSpaceDenoise.h
+++ b/Engine/Source/Runtime/Renderer/Private/ScreenSpaceDenoise.h
@@ -66,7 +66,10 @@ public:
 	struct FReflectionsRayTracingConfig
 	{
 		// Resolution fraction the ray tracing is being traced at.
-		float ResolutionFraction;
+		float ResolutionFraction = 1.0f;
+		
+		// Number of rays per pixels.
+		int32 RayCountPerPixel = 1;
 	};
 	
 	/** The configuration of the reflection ray tracing. */
diff --git a/Engine/Source/Runtime/Renderer/Private/VelocityRendering.cpp b/Engine/Source/Runtime/Renderer/Private/VelocityRendering.cpp
--- a/Engine/Source/Runtime/Renderer/Private/VelocityRendering.cpp
+++ b/Engine/Source/Runtime/Renderer/Private/VelocityRendering.cpp
@@ -388,7 +388,7 @@ bool FDeferredShadingSceneRenderer::ShouldRenderVelocities() const
 		bool bMotionBlur = IsMotionBlurEnabled(View);
 		bool bDistanceFieldAO = ShouldPrepareForDistanceFieldAO();
 
-		bool bSSRTemporal = IsSSRTemporalPassRequired(View);
+		bool bSSRTemporal = ShouldRenderScreenSpaceReflections(View) && IsSSRTemporalPassRequired(View);
 
 		bNeedsVelocity |= bMotionBlur || bTemporalAA || bDistanceFieldAO || bSSRTemporal;
 	}
