Offers DOF it's own gpu stat event

diff --git a/Engine/Source/Runtime/Renderer/Private/PostProcess/DiaphragmDOF.cpp b/Engine/Source/Runtime/Renderer/Private/PostProcess/DiaphragmDOF.cpp
--- a/Engine/Source/Runtime/Renderer/Private/PostProcess/DiaphragmDOF.cpp
+++ b/Engine/Source/Runtime/Renderer/Private/PostProcess/DiaphragmDOF.cpp
@@ -33,6 +33,8 @@ void AddPass_ClearUAV(
 namespace
 {
 
+DECLARE_GPU_STAT(DepthOfField)
+
 TAutoConsoleVariable<int32> CVarAccumulatorQuality(
 	TEXT("r.DOF.Gather.AccumulatorQuality"),
 	1,
@@ -1395,6 +1397,7 @@ FRDGTextureRef DiaphragmDOF::AddPasses(
 		return InputSceneColor;
 	}
 
+	RDG_GPU_STAT_SCOPE(GraphBuilder, DepthOfField);
 	RDG_EVENT_SCOPE(GraphBuilder, "DOF(Alpha=%s)", bProcessSceneAlpha ? TEXT("Yes") : TEXT("No"));
 
 	bool bGatherForeground = AbsMaxForegroundCocRadius > kMinimalAbsGatherPassCocRadius;
