Offers gpu stat events to the denoisers.

diff --git a/Engine/Source/Runtime/Renderer/Private/RayTracing/RayTracingAmbientOcclusion.cpp b/Engine/Source/Runtime/Renderer/Private/RayTracing/RayTracingAmbientOcclusion.cpp
--- a/Engine/Source/Runtime/Renderer/Private/RayTracing/RayTracingAmbientOcclusion.cpp
+++ b/Engine/Source/Runtime/Renderer/Private/RayTracing/RayTracingAmbientOcclusion.cpp
@@ -66,7 +66,6 @@ bool ShouldRenderRayTracingAmbientOcclusion(const FViewInfo& View)
 }
 
 DECLARE_GPU_STAT_NAMED(RayTracingAmbientOcclusion, TEXT("Ray Tracing Ambient Occlusion"));
-DECLARE_GPU_STAT_NAMED(AmbientOcclusionDenoiser, TEXT("Ambient Occlusion Denoiser"));
 
 class FRayTracingAmbientOcclusionRGS : public FGlobalShader
 {
@@ -235,8 +234,6 @@ void FDeferredShadingSceneRenderer::RenderRayTracingAmbientOcclusion(
 	int32 DenoiserMode = CVarUseAODenoiser.GetValueOnRenderThread();
 	if (DenoiserMode != 0)
 	{
-		SCOPED_GPU_STAT(RHICmdList, AmbientOcclusionDenoiser);
-
 		FSceneViewFamilyBlackboard SceneBlackboard;
 		SetupSceneViewFamilyBlackboard(GraphBuilder, &SceneBlackboard);
 
diff --git a/Engine/Source/Runtime/Renderer/Private/ScreenSpaceDenoise.cpp b/Engine/Source/Runtime/Renderer/Private/ScreenSpaceDenoise.cpp
--- a/Engine/Source/Runtime/Renderer/Private/ScreenSpaceDenoise.cpp
+++ b/Engine/Source/Runtime/Renderer/Private/ScreenSpaceDenoise.cpp
@@ -115,6 +115,11 @@ static_assert(IScreenSpaceDenoiser::kMaxBatchSize <= kMaxBufferProcessingCount,
 
 const IScreenSpaceDenoiser* GScreenSpaceDenoiser = nullptr;
 
+DECLARE_GPU_STAT(ReflectionsDenoiser)
+DECLARE_GPU_STAT(ShadowsDenoiser)
+DECLARE_GPU_STAT(AmbientOcclusionDenoiser)
+DECLARE_GPU_STAT(DiffuseIndirectDenoiser)
+
 namespace
 {
 
@@ -1231,7 +1236,6 @@ static void DenoiseSignalAtConstantPixelDensity(
 	}
 } // DenoiseSignalAtConstantPixelDensity()
 
-
 /** The implementation of the default denoiser of the renderer. */
 class FDefaultScreenSpaceDenoiser : public IScreenSpaceDenoiser
 {
@@ -1263,6 +1267,8 @@ public:
 		const int32 InputParameterCount,
 		TStaticArray<FShadowPenumbraOutputs, IScreenSpaceDenoiser::kMaxBatchSize>& Outputs) const override
 	{
+		RDG_GPU_STAT_SCOPE(GraphBuilder, ShadowsDenoiser);
+
 		FSSDSignalTextures InputSignal;
 
 		FSSDConstantPixelDensitySettings Settings;
@@ -1331,6 +1337,8 @@ public:
 		const FReflectionsInputs& ReflectionInputs,
 		const FReflectionsRayTracingConfig RayTracingConfig) const override
 	{
+		RDG_GPU_STAT_SCOPE(GraphBuilder, ReflectionsDenoiser);
+
 		FSSDSignalTextures InputSignal;
 		InputSignal.Textures[0] = ReflectionInputs.Color;
 		InputSignal.Textures[1] = ReflectionInputs.RayHitDistance;
@@ -1369,6 +1377,8 @@ public:
 		const FAmbientOcclusionInputs& ReflectionInputs,
 		const FAmbientOcclusionRayTracingConfig RayTracingConfig) const override
 	{
+		RDG_GPU_STAT_SCOPE(GraphBuilder, AmbientOcclusionDenoiser);
+
 		FSSDSignalTextures InputSignal;
 		InputSignal.Textures[0] = ReflectionInputs.Mask;
 		InputSignal.Textures[1] = ReflectionInputs.RayHitDistance;
@@ -1407,6 +1417,8 @@ public:
 		const FGlobalIlluminationInputs& Inputs,
 		const FAmbientOcclusionRayTracingConfig Config) const override
 	{
+		RDG_GPU_STAT_SCOPE(GraphBuilder, DiffuseIndirectDenoiser);
+
 		FSSDSignalTextures InputSignal;
 		InputSignal.Textures[0] = Inputs.Color;
 		InputSignal.Textures[1] = Inputs.RayHitDistance;
@@ -1445,6 +1457,8 @@ public:
 		const FGlobalIlluminationInputs& Inputs,
 		const FAmbientOcclusionRayTracingConfig Config) const override
 	{
+		RDG_GPU_STAT_SCOPE(GraphBuilder, DiffuseIndirectDenoiser);
+
 		FSSDSignalTextures InputSignal;
 		InputSignal.Textures[0] = Inputs.Color;
 		InputSignal.Textures[1] = Inputs.RayHitDistance;
