Reduce tonemap shader permutations for LinearNoToneCurve

diff --git a/Engine/Shaders/Private/PostProcessTonemap.usf b/Engine/Shaders/Private/PostProcessTonemap.usf
--- a/Engine/Shaders/Private/PostProcessTonemap.usf
+++ b/Engine/Shaders/Private/PostProcessTonemap.usf
@@ -510,8 +510,11 @@ float4 TonemapCommonPS(
 	
 		half LuminanceForPostProcessAA  = dot(OutDeviceColor, half3(0.299f, 0.587f, 0.114f));
 
-		#if USE_GRAIN_QUANTIZATION && DIM_OUTPUT_DEVICE != TONEMAPPER_OUTOUT_NoToneCurve
-			// Needs to go after tonemapping.
+		#if USE_GRAIN_QUANTIZATION
+			#if DIM_OUTPUT_DEVICE == TONEMAPPER_OUTOUT_NoToneCurve
+				#error LinearNoToneCurve device requires no Grain Quantization
+			#endif
+			// Needs to go after tonemapping. 
 			#if USE_LINEAR_FLOAT_RT
 				half GrainQuantization = 1.0/1024.0;
 			#else
diff --git a/Engine/Source/Runtime/Renderer/Private/PostProcess/PostProcessTonemap.cpp b/Engine/Source/Runtime/Renderer/Private/PostProcess/PostProcessTonemap.cpp
--- a/Engine/Source/Runtime/Renderer/Private/PostProcess/PostProcessTonemap.cpp
+++ b/Engine/Source/Runtime/Renderer/Private/PostProcess/PostProcessTonemap.cpp
@@ -216,8 +216,11 @@ FDesktopDomain RemapPermutation(FDesktopDomain PermutationVector)
 	// You most likely need Bloom anyway.
 	CommonPermutationVector.Set<FTonemapperBloomDim>(true);
 
-	// Grain quantization is pretty important anyway.
-	PermutationVector.Set<FTonemapperGrainQuantizationDim>(true);
+	// Disable grain quantization only for LinearNoToneCurve output device
+	if (PermutationVector.Get<FTonemapperOutputDeviceDim>() == FTonemapperOutputDevice::LinearNoToneCurve)
+		PermutationVector.Set<FTonemapperGrainQuantizationDim>(false);
+	else
+		PermutationVector.Set<FTonemapperGrainQuantizationDim>(true);
 
 	PermutationVector.Set<FCommonDomain>(CommonPermutationVector);
 	return PermutationVector;
