Removes unused permutation of SSR.

diff --git a/Engine/Shaders/Public/ShaderVersion.ush b/Engine/Shaders/Public/ShaderVersion.ush
--- a/Engine/Shaders/Public/ShaderVersion.ush
+++ b/Engine/Shaders/Public/ShaderVersion.ush
@@ -2,4 +2,4 @@
 // Each time the console command is executed it generates a new GUID. As this file is included
 // in Platform.ush (which should be included in any shader) it allows to invalidate the shader DDC.
 // 
-// GUID = {53723FA1-FF63-44CA-A6CB-9BE1854E250D}
\ No newline at end of file
+// GUID = A3150FB44B9DB7ED80213F8BA3B16F32
\ No newline at end of file
diff --git a/Engine/Source/Runtime/Renderer/Private/PostProcess/ScreenSpaceReflections.cpp b/Engine/Source/Runtime/Renderer/Private/PostProcess/ScreenSpaceReflections.cpp
--- a/Engine/Source/Runtime/Renderer/Private/PostProcess/ScreenSpaceReflections.cpp
+++ b/Engine/Source/Runtime/Renderer/Private/PostProcess/ScreenSpaceReflections.cpp
@@ -138,7 +138,6 @@ END_SHADER_PARAMETER_STRUCT()
 
 
 class FSSRQualityDim : SHADER_PERMUTATION_ENUM_CLASS("SSR_QUALITY", ESSRQuality);
-class FSSRPrevFrameColorDim : SHADER_PERMUTATION_BOOL("PREV_FRAME_COLOR");
 class FSSROutputForDenoiser : SHADER_PERMUTATION_BOOL("SSR_OUTPUT_FOR_DENOISER");
 
 
@@ -157,7 +156,6 @@ class FScreenSpaceReflectionsStencilPS : public FGlobalShader
 	static void ModifyCompilationEnvironment(const FGlobalShaderPermutationParameters& Parameters, FShaderCompilerEnvironment& OutEnvironment)
 	{
 		FGlobalShader::ModifyCompilationEnvironment(Parameters, OutEnvironment);
-		OutEnvironment.SetDefine( TEXT("PREV_FRAME_COLOR"), uint32(0) );
 		OutEnvironment.SetDefine( TEXT("SSR_QUALITY"), uint32(0) );
 	}
 	
@@ -172,15 +170,11 @@ class FScreenSpaceReflectionsPS : public FGlobalShader
 	DECLARE_GLOBAL_SHADER(FScreenSpaceReflectionsPS);
 	SHADER_USE_PARAMETER_STRUCT(FScreenSpaceReflectionsPS, FGlobalShader);
 
-	using FPermutationDomain = TShaderPermutationDomain<FSSRQualityDim, FSSRPrevFrameColorDim, FSSROutputForDenoiser>;
+	using FPermutationDomain = TShaderPermutationDomain<FSSRQualityDim, FSSROutputForDenoiser>;
 
 	static bool ShouldCompilePermutation(const FGlobalShaderPermutationParameters& Parameters)
 	{
 		FPermutationDomain PermutationVector(Parameters.PermutationId);
-		if ((PermutationVector.Get<FSSRQualityDim>() == ESSRQuality::VisualizeSSR) && PermutationVector.Get<FSSRPrevFrameColorDim>())
-		{
-			return false;
-		}
 		return IsFeatureLevelSupported(Parameters.Platform, ERHIFeatureLevel::SM4);
 	}
 	
@@ -371,7 +365,6 @@ void RenderScreenSpaceReflections(
 	// Adds SSR pass.
 	{
 		FScreenSpaceReflectionsPS::FPermutationDomain PermutationVector;
-		PermutationVector.Set<FSSRPrevFrameColorDim>(InputColor != CurrentSceneColor);
 		PermutationVector.Set<FSSRQualityDim>(SSRQuality);
 		PermutationVector.Set<FSSROutputForDenoiser>(bDenoiser);
 
