Fixes bokeh simulation on DOF slight out of focus that wasn't getting disabled properly.

diff --git a/Engine/Source/Runtime/Renderer/Private/PostProcess/DiaphragmDOF.cpp b/Engine/Source/Runtime/Renderer/Private/PostProcess/DiaphragmDOF.cpp
--- a/Engine/Source/Runtime/Renderer/Private/PostProcess/DiaphragmDOF.cpp
+++ b/Engine/Source/Runtime/Renderer/Private/PostProcess/DiaphragmDOF.cpp
@@ -1402,6 +1402,12 @@ FRDGTextureRef DiaphragmDOF::AddPasses(
 
 	bool bGatherForeground = AbsMaxForegroundCocRadius > kMinimalAbsGatherPassCocRadius;
 	
+	const bool bEnableGatherBokehSettings = (
+		bSupportGatheringBokehSimulation &&
+		CVarEnableGatherBokehSettings.GetValueOnRenderThread() == 1);
+	const bool bEnableScatterBokehSettings = CVarEnableScatterBokehSettings.GetValueOnRenderThread() == 1;
+	const bool bEnableSlightOutOfFocusBokeh = bSupportGatheringBokehSimulation && bRecombineDoesSlightOutOfFocus && CVarEnableRecombineBokehSettings.GetValueOnRenderThread();
+		
 	// Setup all the desciptors.
 	FRDGTextureDesc FullResDesc;
 	{
@@ -2198,12 +2204,6 @@ FRDGTextureRef DiaphragmDOF::AddPasses(
 			*ConvolutionTextures = NewConvolutionTextures;
 		}; // AddPostFilterPass()
 		
-		const bool bEnableGatherBokehSettings = (
-			bSupportGatheringBokehSimulation &&
-			CVarEnableGatherBokehSettings.GetValueOnRenderThread() == 1);
-		const bool bEnableScatterBokehSettings = CVarEnableScatterBokehSettings.GetValueOnRenderThread() == 1;
-		const bool bEnableSlightOutOfFocusBokeh = bSupportGatheringBokehSimulation && bRecombineDoesSlightOutOfFocus && CVarEnableRecombineBokehSettings.GetValueOnRenderThread();
-		
 		FRDGTextureRef GatheringBokehLUT = nullptr;
 		if (bEnableGatherBokehSettings)
 			 GatheringBokehLUT = AddBuildBokehLUTPass(EDiaphragmDOFBokehLUTFormat::GatherSamplePos);
@@ -2374,8 +2374,6 @@ FRDGTextureRef DiaphragmDOF::AddPasses(
 	// Recombine lower res out of focus with full res scene color.
 	FRDGTextureRef NewSceneColor;
 	{
-		bool bEnableSlightOutOfFocusBokeh = SlightOutOfFocusConvolutionTextures.SceneColor != nullptr;
-
 		{
 			FRDGTextureDesc Desc = InputSceneColor->Desc;
 			Desc.TargetableFlags |= TexCreate_UAV;
