Temporary disable multiview support in ray traced shadows to avoid breaking the shadows batch denoiser. More refactoring is required to make things working. Porting RenderLights() to RGD would also be needed.

diff --git a/Engine/Source/Runtime/Renderer/Private/DeferredShadingRenderer.h b/Engine/Source/Runtime/Renderer/Private/DeferredShadingRenderer.h
--- a/Engine/Source/Runtime/Renderer/Private/DeferredShadingRenderer.h
+++ b/Engine/Source/Runtime/Renderer/Private/DeferredShadingRenderer.h
@@ -479,8 +479,8 @@ private:
 		const FLightSceneInfo& LightSceneInfo,
 		const IScreenSpaceDenoiser::FShadowRayTracingConfig& RayTracingConfig,
 		IScreenSpaceDenoiser::EShadowRequirements DenoiserRequirements,
-		FRDGTextureUAV* OutShadowMaskUAV,
-		FRDGTextureUAV* OutRayHitDistanceUAV);
+		FRDGTextureRef* OutShadowMask,
+		FRDGTextureRef* OutRayHitDistance);
 
 	void RenderRayTracingStochasticRectLight(FRHICommandListImmediate& RHICmdList, const FLightSceneInfo& RectLightSceneInfo, TRefCountPtr<IPooledRenderTarget>& RectLightRT, TRefCountPtr<IPooledRenderTarget>& HitDistanceRT);
 	void CompositeRayTracingSkyLight(FRHICommandListImmediate& RHICmdList, TRefCountPtr<IPooledRenderTarget>& SkyLightRT, TRefCountPtr<IPooledRenderTarget>& HitDistanceRT);
diff --git a/Engine/Source/Runtime/Renderer/Private/LightRendering.cpp b/Engine/Source/Runtime/Renderer/Private/LightRendering.cpp
--- a/Engine/Source/Runtime/Renderer/Private/LightRendering.cpp
+++ b/Engine/Source/Runtime/Renderer/Private/LightRendering.cpp
@@ -936,10 +936,9 @@ void FDeferredShadingSceneRenderer::RenderLights(FRHICommandListImmediate& RHICm
 			TArray<TRefCountPtr<IPooledRenderTarget>> PreprocessedShadowMaskTextures;
 
 			const int32 MaxDenoisingBatchSize = FMath::Clamp(CVarMaxShadowDenoisingBatchSize.GetValueOnRenderThread(), 1, IScreenSpaceDenoiser::kMaxBatchSize);
-			//#dxr_todo: add multiview support to batching denoising 
-			const bool bDoShadowDenoisingBatching = DenoiserMode != 0 && MaxDenoisingBatchSize > 1 && Views.Num() == 1;
+			const bool bDoShadowDenoisingBatching = DenoiserMode != 0 && MaxDenoisingBatchSize > 1;
 
-			// Optimizations: batches all shadow ray tracing denoising. Definitely could be smarter to avoid high VGPR pressure if this entire
+			// Optimisations: batches all shadow ray tracing denoising. Definitely could be smarter to avoid high VGPR pressure if this entire
 			// function was converted to render graph, and want least intrusive change as possible. So right not it trades render target memory pressure
 			// for denoising perf.
 			if (RHI_RAYTRACING && bDoShadowDenoisingBatching)
@@ -952,39 +951,7 @@ void FDeferredShadingSceneRenderer::RenderLights(FRHICommandListImmediate& RHICm
 				FSceneViewFamilyBlackboard SceneBlackboard;
 				SetupSceneViewFamilyBlackboard(GraphBuilder, &SceneBlackboard);
 
-				// Ray trace the shadow.
-				FRDGTextureRef ShadowMask;
-
-				{
-					FRDGTextureDesc Desc = FRDGTextureDesc::Create2DDesc(
-						SceneBlackboard.SceneDepthBuffer->Desc.Extent,
-						PF_FloatRGBA,
-						FClearValueBinding::Black,
-						TexCreate_None,
-						TexCreate_ShaderResource | TexCreate_RenderTargetable | TexCreate_UAV,
-						/* bInForceSeparateTargetAndShaderResource = */ false);
-					ShadowMask = GraphBuilder.CreateTexture(Desc, TEXT("RayTracingOcclusion"));
-				}
-
-				FRDGTextureRef RayHitDistance;
-				{
-					FRDGTextureDesc Desc = FRDGTextureDesc::Create2DDesc(
-						SceneBlackboard.SceneDepthBuffer->Desc.Extent,
-						PF_R16F,
-						FClearValueBinding::Black,
-						TexCreate_None,
-						TexCreate_ShaderResource | TexCreate_RenderTargetable | TexCreate_UAV,
-						/* bInForceSeparateTargetAndShaderResource = */ false);
-					RayHitDistance = GraphBuilder.CreateTexture(Desc, TEXT("RayTracingOcclusionDistance"));
-				}
-
-				FRDGTextureUAV* ShadowMaskUAV = GraphBuilder.CreateUAV(FRDGTextureUAVDesc(ShadowMask));
-				FRDGTextureUAV* RayHitDistanceUAV = GraphBuilder.CreateUAV(FRDGTextureUAVDesc(RayHitDistance));
-
-
-				for (int32 ViewIndex = 0; ViewIndex < Views.Num(); ViewIndex++)
-				{
-					FViewInfo& View = Views[ViewIndex];
+				FViewInfo& View = Views[0];
 
 					// Allocate PreprocessedShadowMaskTextures once so QueueTextureExtraction can deferred write.
 					{
@@ -1084,7 +1051,9 @@ void FDeferredShadingSceneRenderer::RenderLights(FRHICommandListImmediate& RHICm
 							continue;
 						}
 
-
+					// Ray trace the shadow.
+					FRDGTextureRef ShadowMask;
+					FRDGTextureRef RayHitDistance;
 						{
 							FString LightNameWithLevel;
 							GetLightNameForDrawEvent(LightSceneInfo.Proxy, LightNameWithLevel);
@@ -1097,8 +1066,8 @@ void FDeferredShadingSceneRenderer::RenderLights(FRHICommandListImmediate& RHICm
 								LightSceneInfo,
 								RayTracingConfig,
 								DenoiserRequirements,
-								ShadowMaskUAV,
-								RayHitDistanceUAV);
+							&ShadowMask,
+							&RayHitDistance);
 						}
 
 						// Queue the ray tracing output for shadow denoising.
@@ -1132,7 +1101,6 @@ void FDeferredShadingSceneRenderer::RenderLights(FRHICommandListImmediate& RHICm
 					{
 						QuickOffDenoisingBatch();
 					}
-				}
 
 				GraphBuilder.Execute();
 			} // if (RHI_RAYTRACING)
@@ -1178,50 +1146,24 @@ void FDeferredShadingSceneRenderer::RenderLights(FRHICommandListImmediate& RHICm
 					}
 					else if (OcclusionType == FLightOcclusionType::Raytraced)
 					{
-						FRDGBuilder GraphBuilder(RHICmdList);
-						
-						FSceneViewFamilyBlackboard SceneBlackboard;
-						SetupSceneViewFamilyBlackboard(GraphBuilder, &SceneBlackboard);
-
-						FRDGTextureRef ShadowMask;
-						{
-							FRDGTextureDesc Desc = FRDGTextureDesc::Create2DDesc(
-								SceneBlackboard.SceneDepthBuffer->Desc.Extent,
-								PF_FloatRGBA,
-								FClearValueBinding::Black,
-								TexCreate_None,
-								TexCreate_ShaderResource | TexCreate_RenderTargetable | TexCreate_UAV,
-								/* bInForceSeparateTargetAndShaderResource = */ false);
-							ShadowMask = GraphBuilder.CreateTexture(Desc, TEXT("RayTracingOcclusion"));
-						}
-
-						FRDGTextureRef RayHitDistance;
-						{
-							FRDGTextureDesc Desc = FRDGTextureDesc::Create2DDesc(
-								SceneBlackboard.SceneDepthBuffer->Desc.Extent,
-								PF_R16F,
-								FClearValueBinding::Black,
-								TexCreate_None,
-								TexCreate_ShaderResource | TexCreate_RenderTargetable | TexCreate_UAV,
-								/* bInForceSeparateTargetAndShaderResource = */ false);
-							RayHitDistance = GraphBuilder.CreateTexture(Desc, TEXT("RayTracingOcclusionDistance"));
-						}
-
-						FRDGTextureUAV* ShadowMaskUAV = GraphBuilder.CreateUAV(FRDGTextureUAVDesc(ShadowMask));
-						FRDGTextureUAV* RayHitDistanceUAV = GraphBuilder.CreateUAV(FRDGTextureUAVDesc(RayHitDistance));
-
-						for (int32 ViewIndex = 0; ViewIndex < Views.Num(); ViewIndex++)
-						{
-							FViewInfo& View = Views[ViewIndex];
+						FViewInfo& View = Views[0];
 
 							IScreenSpaceDenoiser::FShadowRayTracingConfig RayTracingConfig;
 							RayTracingConfig.RayCountPerPixel = LightSceneInfo.Proxy->GetSamplesPerPixel();
+
 							IScreenSpaceDenoiser::EShadowRequirements DenoiserRequirements = IScreenSpaceDenoiser::EShadowRequirements::Bailout;
 							if (DenoiserMode != 0 && LightRequiresDenosier(LightSceneInfo))
 							{
 								DenoiserRequirements = DenoiserToUse->GetShadowRequirements(View, LightSceneInfo, RayTracingConfig);
 							}
 
+						FRDGBuilder GraphBuilder(RHICmdList);
+
+						FSceneViewFamilyBlackboard SceneBlackboard;
+						SetupSceneViewFamilyBlackboard(GraphBuilder, &SceneBlackboard);
+
+						FRDGTextureRef ShadowMask;
+						FRDGTextureRef RayHitDistance;
 							RenderRayTracingShadows(
 								GraphBuilder,
 								SceneBlackboard,
@@ -1229,8 +1171,8 @@ void FDeferredShadingSceneRenderer::RenderLights(FRHICommandListImmediate& RHICm
 								LightSceneInfo,
 								RayTracingConfig,
 								DenoiserRequirements,
-								ShadowMaskUAV,
-								RayHitDistanceUAV);
+							&ShadowMask,
+							&RayHitDistance);
 
 							if (DenoiserRequirements != IScreenSpaceDenoiser::EShadowRequirements::Bailout)
 							{
@@ -1265,7 +1207,6 @@ void FDeferredShadingSceneRenderer::RenderLights(FRHICommandListImmediate& RHICm
 							{
 								GraphBuilder.QueueTextureExtraction(ShadowMask, &ScreenShadowMaskTexture);
 							}
-						}
 
 						GraphBuilder.Execute();
 					}
diff --git a/Engine/Source/Runtime/Renderer/Private/RayTracing/RayTracingShadows.cpp b/Engine/Source/Runtime/Renderer/Private/RayTracing/RayTracingShadows.cpp
--- a/Engine/Source/Runtime/Renderer/Private/RayTracing/RayTracingShadows.cpp
+++ b/Engine/Source/Runtime/Renderer/Private/RayTracing/RayTracingShadows.cpp
@@ -116,14 +116,39 @@ void FDeferredShadingSceneRenderer::RenderRayTracingShadows(
 	const FLightSceneInfo& LightSceneInfo,
 	const IScreenSpaceDenoiser::FShadowRayTracingConfig& RayTracingConfig,
 	IScreenSpaceDenoiser::EShadowRequirements DenoiserRequirements,
-	FRDGTextureUAV* OutShadowMaskUAV,
-	FRDGTextureUAV* OutRayHitDistanceUAV)
+	FRDGTextureRef* OutShadowMask,
+	FRDGTextureRef* OutRayHitDistance)
 #if RHI_RAYTRACING
 {
 	FLightSceneProxy* LightSceneProxy = LightSceneInfo.Proxy;
 	check(LightSceneProxy);
 
-	FIntRect ScissorRect = { View.ViewRect.Min, View.ViewRect.Min + View.ViewRect.Size() };
+	// Render targets
+	FRDGTextureRef ScreenShadowMaskTexture;
+	{
+		FRDGTextureDesc Desc = FRDGTextureDesc::Create2DDesc(
+			SceneBlackboard.SceneDepthBuffer->Desc.Extent,
+			PF_FloatRGBA,
+			FClearValueBinding::Black,
+			TexCreate_None,
+			TexCreate_ShaderResource | TexCreate_RenderTargetable | TexCreate_UAV,
+			/* bInForceSeparateTargetAndShaderResource = */ false);
+		ScreenShadowMaskTexture = GraphBuilder.CreateTexture(Desc, TEXT("RayTracingOcclusion"));
+	}
+
+	FRDGTextureRef RayDistanceTexture;
+	{
+		FRDGTextureDesc Desc = FRDGTextureDesc::Create2DDesc(
+			SceneBlackboard.SceneDepthBuffer->Desc.Extent,
+			PF_R16F,
+			FClearValueBinding::Black,
+			TexCreate_None,
+			TexCreate_ShaderResource | TexCreate_RenderTargetable | TexCreate_UAV,
+			/* bInForceSeparateTargetAndShaderResource = */ false);
+		RayDistanceTexture = GraphBuilder.CreateTexture(Desc, TEXT("RayTracingOcclusionDistance"));
+	}
+
+	FIntRect ScissorRect = { {0,0}, View.ViewRect.Size() };
 
 	if (LightSceneProxy->GetScissorRect(ScissorRect, View, View.ViewRect))
 	{
@@ -131,12 +156,16 @@ void FDeferredShadingSceneRenderer::RenderRayTracingShadows(
 		ScissorRect.Min = ScissorRect.Min - View.ViewRect.Min;
 		ScissorRect.Max = ScissorRect.Max - View.ViewRect.Min;
 	}
+	else
+	{
+		ScissorRect = { {0,0}, View.ViewRect.Size() };
+	}
 
 	// Ray generation pass for shadow occlusion.
 	{
 		FOcclusionRGS::FParameters* PassParameters = GraphBuilder.AllocParameters<FOcclusionRGS::FParameters>();
-		PassParameters->RWOcclusionMaskUAV = OutShadowMaskUAV;
-		PassParameters->RWRayDistanceUAV = OutRayHitDistanceUAV;
+		PassParameters->RWOcclusionMaskUAV = GraphBuilder.CreateUAV(FRDGTextureUAVDesc(ScreenShadowMaskTexture));
+		PassParameters->RWRayDistanceUAV = GraphBuilder.CreateUAV(FRDGTextureUAVDesc(RayDistanceTexture));
 		PassParameters->SamplesPerPixel = RayTracingConfig.RayCountPerPixel;
 		PassParameters->NormalBias = GetRaytracingMaxNormalBias();
 		PassParameters->LightingChannelMask = LightSceneProxy->GetLightingChannelMask();
@@ -210,6 +239,9 @@ void FDeferredShadingSceneRenderer::RenderRayTracingShadows(
 			}
 		});
 	}
+
+	*OutShadowMask = ScreenShadowMaskTexture;
+	*OutRayHitDistance = RayDistanceTexture;
 }
 #else // !RHI_RAYTRACING
 {
