Add RaytracingTestbed to the C++ automated tests. Tests are hosted in \Runtime\Renderer\Private\Tests to follow the UE4 convention.

diff --git a/Engine/Source/Runtime/Renderer/Private/RayTracing/RayTracingTestbed.cpp b/Engine/Source/Runtime/Renderer/Private/RayTracing/RayTracingTestbed.cpp
deleted file mode 100644
diff --git a/Engine/Source/Runtime/Renderer/Private/RayTracing/RayTracingTestbed.h b/Engine/Source/Runtime/Renderer/Private/RayTracing/RayTracingTestbed.h
deleted file mode 100644
diff --git a/Engine/Source/Runtime/Renderer/Private/Tests/RayTracingTestbed.cpp b/Engine/Source/Runtime/Renderer/Private/Tests/RayTracingTestbed.cpp
new file mode 100644
--- /dev/null
+++ b/Engine/Source/Runtime/Renderer/Private/Tests/RayTracingTestbed.cpp
@@ -0,0 +1,250 @@
+// Copyright 1998-2019 Epic Games, Inc. All Rights Reserved.
+
+#include "CoreMinimal.h"
+#include "RHI.h"
+#include "Misc/AutomationTest.h"
+
+#if WITH_DEV_AUTOMATION_TESTS
+
+IMPLEMENT_SIMPLE_AUTOMATION_TEST(FRayTracingTestbed, "System.Renderer.RayTracing.BasicRayTracing", EAutomationTestFlags::ApplicationContextMask | EAutomationTestFlags::SmokeFilter)
+//IMPLEMENT_SIMPLE_AUTOMATION_TEST(FRayTracingTestbed, "System.Renderer.RayTracing.BasicRayTracing", EAutomationTestFlags::EditorContext | EAutomationTestFlags::HighPriority | EAutomationTestFlags::EngineFilter)
+
+
+
+#if RHI_RAYTRACING
+
+#include "Containers/DynamicRHIResourceArray.h"
+#include "GlobalShader.h"
+#include "RayTracingDefinitions.h"
+
+
+bool RunRayTracingTestbed_RenderThread(const FString& Parameters)
+{
+	check(IsInRenderingThread());
+
+	if (!GRHISupportsRayTracing)
+	{
+		return false;
+	}
+
+	FVertexBufferRHIRef VertexBuffer;
+
+	{
+		TResourceArray<FVector> PositionData;
+		PositionData.SetNumUninitialized(3);
+		PositionData[0] = FVector( 1, -1, 0);
+		PositionData[1] = FVector( 1,  1, 0);
+		PositionData[2] = FVector(-1, -1, 0);
+
+		FRHIResourceCreateInfo CreateInfo;
+		CreateInfo.ResourceArray = &PositionData;
+
+		VertexBuffer = RHICreateVertexBuffer(PositionData.GetResourceDataSize(), BUF_Static, CreateInfo);
+	}
+
+	FIndexBufferRHIRef IndexBuffer;
+
+	{
+		TResourceArray<uint16> IndexData;
+		IndexData.SetNumUninitialized(3);
+		IndexData[0] = 0;
+		IndexData[1] = 1;
+		IndexData[2] = 2;
+
+		FRHIResourceCreateInfo CreateInfo;
+		CreateInfo.ResourceArray = &IndexData;
+
+		IndexBuffer = RHICreateIndexBuffer(2, IndexData.GetResourceDataSize(), BUF_Static, CreateInfo);
+	}
+
+	static constexpr uint32 NumRays = 4;
+
+	FStructuredBufferRHIRef RayBuffer;
+	FShaderResourceViewRHIRef RayBufferView;
+
+	{
+		TResourceArray<FBasicRayData> RayData;
+		RayData.SetNumUninitialized(NumRays);
+		RayData[0] = FBasicRayData{ { 0.75f, 0.0f, -1.0f}, 0xFFFFFFFF, {0.0f, 0.0f,  1.0f}, 100000.0f }; // expected to hit
+		RayData[1] = FBasicRayData{ { 0.75f, 0.0f, -1.0f}, 0xFFFFFFFF, {0.0f, 0.0f,  1.0f},      0.5f }; // expected to miss (short ray)
+		RayData[2] = FBasicRayData{ { 0.75f, 0.0f,  1.0f}, 0xFFFFFFFF, {0.0f, 0.0f, -1.0f}, 100000.0f }; // expected to miss (back face culled)
+		RayData[3] = FBasicRayData{ {-0.75f, 0.0f, -1.0f}, 0xFFFFFFFF, {0.0f, 0.0f,  1.0f}, 100000.0f }; // expected to miss (doesn't intersect)
+
+		FRHIResourceCreateInfo CreateInfo;
+		CreateInfo.ResourceArray = &RayData;
+
+		RayBuffer = RHICreateStructuredBuffer(sizeof(FBasicRayData), RayData.GetResourceDataSize(), BUF_Static | BUF_ShaderResource, CreateInfo);
+		RayBufferView = RHICreateShaderResourceView(RayBuffer);
+	}
+
+	FVertexBufferRHIRef OcclusionResultBuffer;
+	FUnorderedAccessViewRHIRef OcclusionResultBufferView;
+
+	{
+		FRHIResourceCreateInfo CreateInfo;
+		OcclusionResultBuffer = RHICreateVertexBuffer(sizeof(uint32)*NumRays, BUF_Static | BUF_UnorderedAccess, CreateInfo);
+		OcclusionResultBufferView = RHICreateUnorderedAccessView(OcclusionResultBuffer, PF_R32_UINT);
+	}
+
+	FVertexBufferRHIRef IntersectionResultBuffer;
+	FUnorderedAccessViewRHIRef IntersectionResultBufferView;
+
+	{
+		FRHIResourceCreateInfo CreateInfo;
+		IntersectionResultBuffer = RHICreateVertexBuffer(sizeof(FBasicRayIntersectionData)*NumRays, BUF_Static | BUF_UnorderedAccess, CreateInfo);
+		IntersectionResultBufferView = RHICreateUnorderedAccessView(IntersectionResultBuffer, PF_R32_UINT);
+	}
+
+	FRayTracingGeometryInitializer GeometryInitializer;
+	GeometryInitializer.IndexBuffer = IndexBuffer;
+	GeometryInitializer.PositionVertexBuffer = VertexBuffer;
+	GeometryInitializer.VertexBufferByteOffset = 0;
+	GeometryInitializer.VertexBufferStride = sizeof(FVector);
+	GeometryInitializer.VertexBufferElementType = VET_Float3;
+	GeometryInitializer.BaseVertexIndex = 0;
+	GeometryInitializer.GeometryType = RTGT_Triangles;
+	GeometryInitializer.TotalPrimitiveCount = 1;
+	GeometryInitializer.bFastBuild = false;
+	FRayTracingGeometryRHIRef Geometry = RHICreateRayTracingGeometry(GeometryInitializer);
+
+	FRayTracingGeometryInstance Instances[] = {
+		FRayTracingGeometryInstance { Geometry, FMatrix::Identity, 0, 0xFF }
+	};
+
+	FRayTracingSceneInitializer Initializer;
+	Initializer.Instances = Instances;
+	Initializer.ShaderSlotsPerGeometrySegment = RAY_TRACING_NUM_SHADER_SLOTS;
+	FRayTracingSceneRHIParamRef Scene = RHICreateRayTracingScene(Initializer);
+
+	FRHICommandListImmediate& RHICmdList = FRHICommandListExecutor::GetImmediateCommandList();
+
+	RHICmdList.BuildAccelerationStructure(Geometry);
+	RHICmdList.BuildAccelerationStructure(Scene);
+
+	RHICmdList.RayTraceOcclusion(Scene, RayBufferView, OcclusionResultBufferView, NumRays);
+	RHICmdList.RayTraceIntersection(Scene, RayBufferView, IntersectionResultBufferView, NumRays);
+
+	const bool bValidateResults = true;
+	bool bOcclusionTestOK = false;
+	bool bIntersectionTestOK = false;
+
+	if (bValidateResults)
+	{
+		GDynamicRHI->RHISubmitCommandsAndFlushGPU();
+		GDynamicRHI->RHIBlockUntilGPUIdle();
+
+		// Read back and validate occlusion trace results
+
+		{
+			auto MappedResults = (const uint32*)RHILockVertexBuffer(OcclusionResultBuffer, 0, sizeof(uint32)*NumRays, RLM_ReadOnly);
+
+			check(MappedResults);
+
+			check(MappedResults[0] != 0); // expect hit
+			check(MappedResults[1] == 0); // expect miss
+			check(MappedResults[2] == 0); // expect miss
+			check(MappedResults[3] == 0); // expect miss
+
+			RHIUnlockVertexBuffer(OcclusionResultBuffer);
+
+			bOcclusionTestOK = (MappedResults[0] != 0) && (MappedResults[1] == 0) && (MappedResults[2] == 0) && (MappedResults[3] == 0);
+		}
+
+		// Read back and validate intersection trace results
+
+		{
+			auto MappedResults = (const FBasicRayIntersectionData*)RHILockVertexBuffer(IntersectionResultBuffer, 0, sizeof(FBasicRayIntersectionData)*NumRays, RLM_ReadOnly);
+
+			check(MappedResults);
+
+			// expect hit primitive 0, instance 0, barycentrics {0.5, 0.125}
+			check(MappedResults[0].PrimitiveIndex == 0);
+			check(MappedResults[0].InstanceIndex == 0);
+			check(FMath::IsNearlyEqual(MappedResults[0].Barycentrics[0], 0.5f));
+			check(FMath::IsNearlyEqual(MappedResults[0].Barycentrics[1], 0.125f));
+
+			check(MappedResults[1].PrimitiveIndex == ~0u); // expect miss
+			check(MappedResults[2].PrimitiveIndex == ~0u); // expect miss
+			check(MappedResults[3].PrimitiveIndex == ~0u); // expect miss
+
+			RHIUnlockVertexBuffer(IntersectionResultBuffer);
+
+			bIntersectionTestOK = (MappedResults[1].PrimitiveIndex == ~0u) && (MappedResults[2].PrimitiveIndex == ~0u) && (MappedResults[3].PrimitiveIndex == ~0u);
+		}
+	}
+
+	return (bOcclusionTestOK && bIntersectionTestOK);
+}
+ 
+// Dummy shader to test shader compilation and reflection.
+class FTestRaygenShader : public FGlobalShader
+{
+	DECLARE_SHADER_TYPE(FTestRaygenShader, Global);
+
+	static void ModifyCompilationEnvironment(const FGlobalShaderPermutationParameters& Parameters, FShaderCompilerEnvironment& OutEnvironment)
+	{
+	}
+
+public:
+
+	static bool ShouldCompilePermutation(const FGlobalShaderPermutationParameters& Parameters)
+	{
+		return ShouldCompileRayTracingShadersForProject(Parameters.Platform);
+	}
+
+	FTestRaygenShader() {}
+	virtual ~FTestRaygenShader() {}
+
+	/** Initialization constructor. */
+	FTestRaygenShader(const ShaderMetaType::CompiledShaderInitializerType& Initializer)
+		: FGlobalShader(Initializer)
+	{
+		TLAS.Bind(Initializer.ParameterMap, TEXT("TLAS"));
+		Rays.Bind(Initializer.ParameterMap, TEXT("Rays"));
+		Output.Bind(Initializer.ParameterMap, TEXT("Output"));
+	}
+
+	bool Serialize(FArchive& Ar)
+	{
+		bool bShaderHasOutdatedParameters = FGlobalShader::Serialize(Ar);
+		Ar << TLAS;
+		Ar << Rays;
+		Ar << Output;
+		return bShaderHasOutdatedParameters;
+	}
+
+	FShaderResourceParameter	TLAS;   // SRV RaytracingAccelerationStructure
+	FShaderResourceParameter	Rays;   // SRV StructuredBuffer<FBasicRayData>
+	FShaderResourceParameter	Output; // UAV RWStructuredBuffer<uint>
+};
+
+IMPLEMENT_SHADER_TYPE(, FTestRaygenShader, TEXT("/Engine/Private/RayTracing/RayTracingTest.usf"), TEXT("TestMainRGS"), SF_RayGen);
+
+
+bool FRayTracingTestbed::RunTest(const FString& Parameters)
+{
+	bool bTestPassed = false;
+
+	ENQUEUE_RENDER_COMMAND(FRayTracingTestbed)(
+		[&](FRHICommandListImmediate& RHICmdList)
+	{
+		bTestPassed = RunRayTracingTestbed_RenderThread(Parameters);
+	}
+	);  
+
+	FlushRenderingCommands();
+
+	return bTestPassed;
+}
+
+#else // RHI_RAYTRACING
+
+bool FRayTracingTestbed::RunTest(const FString& Parameters)
+{
+	// Nothing to do when ray tracing is disabled
+	return true;
+}
+
+#endif // RHI_RAYTRACING
+
+#endif //WITH_DEV_AUTOMATION_TESTS
