Fixed local volumetric fog shadowed light scattering.
#rnx

diff --git a/Engine/Source/Runtime/Renderer/Private/VolumetricFog.cpp b/Engine/Source/Runtime/Renderer/Private/VolumetricFog.cpp
--- a/Engine/Source/Runtime/Renderer/Private/VolumetricFog.cpp
+++ b/Engine/Source/Runtime/Renderer/Private/VolumetricFog.cpp
@@ -501,7 +501,7 @@ void FDeferredShadingSceneRenderer::RenderLocalLightsForVolumetricFog(
 		OutLocalShadowedLightScattering = GraphBuilder.CreateTexture(VolumeDesc, TEXT("LocalShadowedLightScattering"));
 
 		FRenderTargetParameters* PassParameters = GraphBuilder.AllocParameters<FRenderTargetParameters>();
-		PassParameters->RenderTargets[0] = FRenderTargetBinding(OutLocalShadowedLightScattering, ERenderTargetLoadAction::EClear, ERenderTargetStoreAction::ENoAction);
+		PassParameters->RenderTargets[0] = FRenderTargetBinding(OutLocalShadowedLightScattering, ERenderTargetLoadAction::EClear, ERenderTargetStoreAction::EStore);
 
 		GraphBuilder.AddPass(
 			RDG_EVENT_NAME("ShadowedLights"),
