Fixes SSR reprojecting current scene color when TAA is off.

Regression caused by 3364300 that no longer used PREV_FRAME_COLOR.

diff --git a/Engine/Source/Runtime/Renderer/Private/PostProcess/ScreenSpaceReflections.cpp b/Engine/Source/Runtime/Renderer/Private/PostProcess/ScreenSpaceReflections.cpp
--- a/Engine/Source/Runtime/Renderer/Private/PostProcess/ScreenSpaceReflections.cpp
+++ b/Engine/Source/Runtime/Renderer/Private/PostProcess/ScreenSpaceReflections.cpp
@@ -402,6 +402,13 @@ void RenderScreenSpaceReflections(
 		}
 		PassParameters->PrevSceneColorPreExposureCorrection = InputColor != CurrentSceneColor ? View.PreExposure / View.PrevViewInfo.SceneColorPreExposure : 1.0f;
 		
+		// Pipe down a mid grey texture when not using TAA's history to avoid wrongly reprojecting current scene color as if previous frame's TAA history.
+		if (InputColor == CurrentSceneColor)
+		{
+			// Technically should be 32767.0f / 65535.0f to perfectly null out DecodeVelocityFromTexture(), but 0.5f is good enough.
+			PassParameters->CommonParameters.SceneTextures.SceneVelocityBuffer = GraphBuilder.RegisterExternalTexture(GSystemTextures.MidGreyDummy);
+		}
+
 		PassParameters->SceneColor = InputColor;
 		PassParameters->SceneColorSampler = TStaticSamplerState<SF_Point>::GetRHI();
 		
diff --git a/Engine/Source/Runtime/Renderer/Private/SystemTextures.cpp b/Engine/Source/Runtime/Renderer/Private/SystemTextures.cpp
--- a/Engine/Source/Runtime/Renderer/Private/SystemTextures.cpp
+++ b/Engine/Source/Runtime/Renderer/Private/SystemTextures.cpp
@@ -140,6 +140,19 @@ void FSystemTextures::InitializeCommonTextures(FRHICommandListImmediate& RHICmdL
 		RHICmdList.EndRenderPass();
 		RHICmdList.CopyToResolveTarget(DepthDummy->GetRenderTargetItem().TargetableTexture, DepthDummy->GetRenderTargetItem().ShaderResourceTexture, FResolveParams());
 	}
+	
+	if (!GSupportsShaderFramebufferFetch && GPixelFormats[PF_FloatRGBA].Supported)
+	{
+		// PF_FloatRGBA to encode exactly the 0.5.
+		FPooledRenderTargetDesc Desc(FPooledRenderTargetDesc::Create2DDesc(FIntPoint(1, 1), PF_FloatRGBA, FClearValueBinding(FLinearColor(0.5f, 0.5f, 0.5f, 0.5f)), TexCreate_HideInVisualizeTexture, TexCreate_RenderTargetable | TexCreate_NoFastClear, false));
+		Desc.AutoWritable = false;
+		GRenderTargetPool.FindFreeElement(RHICmdList, Desc, MidGreyDummy, TEXT("MidGreyDummy"), true, ERenderTargetTransience::NonTransient);
+
+		FRHIRenderPassInfo RPInfo(MidGreyDummy->GetRenderTargetItem().TargetableTexture, ERenderTargetActions::Clear_Store);
+		RHICmdList.BeginRenderPass(RPInfo, TEXT("MidGreyDummy"));
+		RHICmdList.EndRenderPass();
+		RHICmdList.CopyToResolveTarget(MidGreyDummy->GetRenderTargetItem().TargetableTexture, MidGreyDummy->GetRenderTargetItem().ShaderResourceTexture, FResolveParams());
+	}
 }
 
 void FSystemTextures::InitializeFeatureLevelDependentTextures(FRHICommandListImmediate& RHICmdList, const ERHIFeatureLevel::Type InFeatureLevel)
diff --git a/Engine/Source/Runtime/Renderer/Private/SystemTextures.h b/Engine/Source/Runtime/Renderer/Private/SystemTextures.h
--- a/Engine/Source/Runtime/Renderer/Private/SystemTextures.h
+++ b/Engine/Source/Runtime/Renderer/Private/SystemTextures.h
@@ -80,6 +80,8 @@ public:
 	TRefCountPtr<IPooledRenderTarget> GreenDummy;
 	// float4(0.5,0.5,0.5,1)
 	TRefCountPtr<IPooledRenderTarget> DefaultNormal8Bit;
+	// float4(0.5,0.5,0.5,0.5)
+	TRefCountPtr<IPooledRenderTarget> MidGreyDummy;
 
 	/** float4(0,0,0,0) volumetric texture. */
 	TRefCountPtr<IPooledRenderTarget> VolumetricBlackDummy;
