Fixes compilation failure in reflection code

diff --git a/Engine/Source/Runtime/Renderer/Private/DeferredShadingRenderer.h b/Engine/Source/Runtime/Renderer/Private/DeferredShadingRenderer.h
--- a/Engine/Source/Runtime/Renderer/Private/DeferredShadingRenderer.h
+++ b/Engine/Source/Runtime/Renderer/Private/DeferredShadingRenderer.h
@@ -465,12 +465,10 @@ private:
 	void RenderRayTracingReflections(
 		FRDGBuilder& GraphBuilder,
 		const FViewInfo& View,
-		FRDGTextureRef* OutColorTexture,
-		FRDGTextureRef* OutRayHitDistanceTexture,
-		FRDGTextureRef* OutRayImaginaryDepthTexture,
 		int32 SamplePerPixel,
 		int32 HeightFog,
-		float ResolutionFraction);
+		float ResolutionFraction,
+		IScreenSpaceDenoiser::FReflectionsInputs* OutDenoiserInputs);
 
 	void RenderRayTracingShadows(
 		FRDGBuilder& GraphBuilder,
diff --git a/Engine/Source/Runtime/Renderer/Private/RayTracing/RayTracingReflections.cpp b/Engine/Source/Runtime/Renderer/Private/RayTracing/RayTracingReflections.cpp
--- a/Engine/Source/Runtime/Renderer/Private/RayTracing/RayTracingReflections.cpp
+++ b/Engine/Source/Runtime/Renderer/Private/RayTracing/RayTracingReflections.cpp
@@ -213,12 +213,10 @@ void FDeferredShadingSceneRenderer::PrepareRayTracingReflections(const FViewInfo
 void FDeferredShadingSceneRenderer::RenderRayTracingReflections(
 	FRDGBuilder& GraphBuilder,
 	const FViewInfo& View,
-	FRDGTextureRef* OutColorTexture,
-	FRDGTextureRef* OutRayHitDistanceTexture,
-	FRDGTextureRef* OutRayImaginaryDepthTexture,
 	int32 SamplePerPixel,
 	int32 HeightFog,
-	float ResolutionFraction)
+	float ResolutionFraction,
+	IScreenSpaceDenoiser::FReflectionsInputs* OutDenoiserInputs)
 #if RHI_RAYTRACING
 {
 	const uint32 SortTileSize = CVarRayTracingReflectionsSortTileSize.GetValueOnRenderThread();
@@ -238,11 +236,11 @@ void FDeferredShadingSceneRenderer::RenderRayTracingReflections(
 		Desc.Extent /= UpscaleFactor;
 		Desc.TargetableFlags |= TexCreate_UAV;
 
-		*OutColorTexture = GraphBuilder.CreateTexture(Desc, TEXT("RayTracingReflections"));
+		OutDenoiserInputs->Color = GraphBuilder.CreateTexture(Desc, TEXT("RayTracingReflections"));
 		
 		Desc.Format = PF_R16F;
-		*OutRayHitDistanceTexture = GraphBuilder.CreateTexture(Desc, TEXT("RayTracingReflectionsHitDistance"));
-		*OutRayImaginaryDepthTexture = GraphBuilder.CreateTexture(Desc, TEXT("RayTracingReflectionsImaginaryDepth"));
+		OutDenoiserInputs->RayHitDistance = GraphBuilder.CreateTexture(Desc, TEXT("RayTracingReflectionsHitDistance"));
+		OutDenoiserInputs->RayImaginaryDepth = GraphBuilder.CreateTexture(Desc, TEXT("RayTracingReflectionsImaginaryDepth"));
 	}
 
 	// When deferred materials are used, we need to dispatch the reflection shader twice:
@@ -289,9 +287,9 @@ void FDeferredShadingSceneRenderer::RenderRayTracingReflections(
 	CommonParameters.SceneTexturesStruct = CreateSceneTextureUniformBuffer( SceneContext, FeatureLevel, ESceneTextureSetupMode::All, EUniformBufferUsage::UniformBuffer_SingleFrame);
 	CommonParameters.ReflectionStruct = CreateReflectionUniformBuffer(View, EUniformBufferUsage::UniformBuffer_SingleFrame);
 	CommonParameters.FogUniformParameters = CreateFogUniformBuffer(View, EUniformBufferUsage::UniformBuffer_SingleFrame);
-	CommonParameters.ColorOutput = GraphBuilder.CreateUAV(*OutColorTexture);
-	CommonParameters.RayHitDistanceOutput = GraphBuilder.CreateUAV(*OutRayHitDistanceTexture);
-	CommonParameters.RayImaginaryDepthOutput = GraphBuilder.CreateUAV(*OutRayImaginaryDepthTexture);
+	CommonParameters.ColorOutput = GraphBuilder.CreateUAV(OutDenoiserInputs->Color);
+	CommonParameters.RayHitDistanceOutput = GraphBuilder.CreateUAV(OutDenoiserInputs->RayHitDistance);
+	CommonParameters.RayImaginaryDepthOutput = GraphBuilder.CreateUAV(OutDenoiserInputs->RayImaginaryDepth);
 	CommonParameters.SortTileSize = SortTileSize;
 	CommonParameters.ReflectionCapture = View.ReflectionCaptureUniformBuffer;
 	CommonParameters.Forward = View.ForwardLightingResources->ForwardLightDataUniformBuffer;
diff --git a/Engine/Source/Runtime/Renderer/Private/ReflectionEnvironment.cpp b/Engine/Source/Runtime/Renderer/Private/ReflectionEnvironment.cpp
--- a/Engine/Source/Runtime/Renderer/Private/ReflectionEnvironment.cpp
+++ b/Engine/Source/Runtime/Renderer/Private/ReflectionEnvironment.cpp
@@ -805,8 +805,9 @@ void FDeferredShadingSceneRenderer::RenderDeferredReflectionsAndSkyLighting(FRHI
 
 				RenderRayTracingReflections(
 					GraphBuilder,
-					View, &DenoiserInputs.Color, &DenoiserInputs.RayHitDistance, &DenoiserInputs.RayImaginaryDepth,
-					RayTracingConfig.RayCountPerPixel, GRayTracingReflectionsHeightFog, RayTracingConfig.ResolutionFraction);
+					View,
+					RayTracingConfig.RayCountPerPixel, GRayTracingReflectionsHeightFog, RayTracingConfig.ResolutionFraction,
+					&DenoiserInputs);
 			}
 			else if (bScreenSpaceReflections)
 			{
