Fixed some post merge compilation errors and warnings.

diff --git a/Engine/Plugins/FX/Niagara/Source/Niagara/Private/NiagaraRendererRibbons.cpp b/Engine/Plugins/FX/Niagara/Source/Niagara/Private/NiagaraRendererRibbons.cpp
--- a/Engine/Plugins/FX/Niagara/Source/Niagara/Private/NiagaraRendererRibbons.cpp
+++ b/Engine/Plugins/FX/Niagara/Source/Niagara/Private/NiagaraRendererRibbons.cpp
@@ -254,7 +254,7 @@ void NiagaraRendererRibbons::GetDynamicMeshElements(const TArray<const FSceneVie
 			false,
 			SceneProxy->UseSingleSampleShadowFromStationaryLights(),
 			SceneProxy->GetScene().HasPrecomputedVolumetricLightmap_RenderThread(),
-			SceneProxy->UseEditorDepthTest(),
+			SceneProxy->DrawsVelocity(),
 			SceneProxy->GetLightingChannelMask(),
 			0,
 			INDEX_NONE,
diff --git a/Engine/Source/Runtime/Engine/Private/Materials/HLSLMaterialTranslator.h b/Engine/Source/Runtime/Engine/Private/Materials/HLSLMaterialTranslator.h
--- a/Engine/Source/Runtime/Engine/Private/Materials/HLSLMaterialTranslator.h
+++ b/Engine/Source/Runtime/Engine/Private/Materials/HLSLMaterialTranslator.h
@@ -814,7 +814,7 @@ ResourcesString = TEXT("");
 				}
 			}
 
-			if (BlendMode == BLEND_AlphaHoldout && MaterialShadingModel != MSM_Unlit)
+			if (BlendMode == BLEND_AlphaHoldout && BlendMode != MSM_Unlit)
 			{
 				Errorf(TEXT("Alpha Holdout blend mode must use unlit shading model."));
 			}
diff --git a/Engine/Source/Runtime/Renderer/Private/MeshPassProcessor.cpp b/Engine/Source/Runtime/Renderer/Private/MeshPassProcessor.cpp
--- a/Engine/Source/Runtime/Renderer/Private/MeshPassProcessor.cpp
+++ b/Engine/Source/Runtime/Renderer/Private/MeshPassProcessor.cpp
@@ -1096,6 +1096,7 @@ void FMeshPassProcessor::GetDrawCommandPrimitiveId(
 		if (BatchElement.PrimitiveIdMode == PrimID_FromPrimitiveSceneInfo)
 		{
 			ensureMsgf(BatchElement.PrimitiveUniformBufferResource == nullptr, TEXT("PrimitiveUniformBufferResource should not be setup when PrimitiveIdMode == PrimID_FromPrimitiveSceneInfo"));
+			check(PrimitiveSceneInfo);
 			DrawPrimitiveId = PrimitiveSceneInfo->GetIndex();
 		}
 		else if (BatchElement.PrimitiveIdMode == PrimID_DynamicPrimitiveShaderData)
diff --git a/Engine/Source/Runtime/Renderer/Private/RayTracing/RayTracingAmbientOcclusion.cpp b/Engine/Source/Runtime/Renderer/Private/RayTracing/RayTracingAmbientOcclusion.cpp
--- a/Engine/Source/Runtime/Renderer/Private/RayTracing/RayTracingAmbientOcclusion.cpp
+++ b/Engine/Source/Runtime/Renderer/Private/RayTracing/RayTracingAmbientOcclusion.cpp
@@ -123,7 +123,7 @@ void FDeferredShadingSceneRenderer::RenderRayTracingAmbientOcclusion(FRHICommand
 	for (int32 ViewIndex = 0; ViewIndex < Views.Num(); ++ViewIndex)
 	{
 		FViewInfo& View = Views[ViewIndex];
-		bAnyViewWithRTAO = bAnyViewWithRTAO | ShouldRenderRayTracingAmbientOcclusion(View);
+		bAnyViewWithRTAO = bAnyViewWithRTAO || ShouldRenderRayTracingAmbientOcclusion(View);
 	}
 
 	if (!bAnyViewWithRTAO)
diff --git a/Engine/Source/Runtime/Renderer/Private/RayTracing/RayTracingGlobalIllumination.cpp b/Engine/Source/Runtime/Renderer/Private/RayTracing/RayTracingGlobalIllumination.cpp
--- a/Engine/Source/Runtime/Renderer/Private/RayTracing/RayTracingGlobalIllumination.cpp
+++ b/Engine/Source/Runtime/Renderer/Private/RayTracing/RayTracingGlobalIllumination.cpp
@@ -327,7 +327,7 @@ void FDeferredShadingSceneRenderer::RenderRayTracingGlobalIllumination(
 	for (int32 ViewIndex = 0; ViewIndex < Views.Num(); ++ViewIndex)
 	{
 		FViewInfo& View = Views[ViewIndex];
-		bAnyViewWithRTGI = bAnyViewWithRTGI | ShouldRenderRayTracingGlobalIllumination(View);
+		bAnyViewWithRTGI = bAnyViewWithRTGI || ShouldRenderRayTracingGlobalIllumination(View);
 	}
 
 	if (!bAnyViewWithRTGI)
