Adding load/store parameters to ray tracing SkyLight and Global Illumination FRenderTargetBindings.

diff --git a/Engine/Source/Runtime/Renderer/Private/RayTracing/RayTracingGlobalIllumination.cpp b/Engine/Source/Runtime/Renderer/Private/RayTracing/RayTracingGlobalIllumination.cpp
--- a/Engine/Source/Runtime/Renderer/Private/RayTracing/RayTracingGlobalIllumination.cpp
+++ b/Engine/Source/Runtime/Renderer/Private/RayTracing/RayTracingGlobalIllumination.cpp
@@ -550,8 +550,8 @@ void FDeferredShadingSceneRenderer::RenderRayTracingGlobalIllumination(
 		PassParameters->GlobalIlluminationSampler = TStaticSamplerState<SF_Point, AM_Clamp, AM_Clamp, AM_Clamp>::GetRHI();
 		PassParameters->ViewUniformBuffer = View.ViewUniformBuffer;
 		PassParameters->SceneTexturesStruct = CreateUniformBufferImmediate(SceneTextures, EUniformBufferUsage::UniformBuffer_SingleDraw);
-		PassParameters->RenderTargets[0] = FRenderTargetBinding(GraphBuilder.RegisterExternalTexture(GlobalIlluminationRT), ERenderTargetLoadAction::ENoAction, ERenderTargetStoreAction::ENoAction);
-		PassParameters->RenderTargets[1] = FRenderTargetBinding(GraphBuilder.RegisterExternalTexture(AmbientOcclusionRT), ERenderTargetLoadAction::ENoAction, ERenderTargetStoreAction::ENoAction);
+		PassParameters->RenderTargets[0] = FRenderTargetBinding(GraphBuilder.RegisterExternalTexture(GlobalIlluminationRT), ERenderTargetLoadAction::ELoad, ERenderTargetStoreAction::EStore);
+		PassParameters->RenderTargets[1] = FRenderTargetBinding(GraphBuilder.RegisterExternalTexture(AmbientOcclusionRT), ERenderTargetLoadAction::ELoad, ERenderTargetStoreAction::EStore);
 
 		GraphBuilder.AddPass(
 			RDG_EVENT_NAME("GlobalIlluminationComposite"),
@@ -611,7 +611,7 @@ void FDeferredShadingSceneRenderer::CompositeGlobalIllumination(
 	FRayTracingGlobalIlluminationSceneColorCompositePS::FParameters *PassParameters = GraphBuilder.AllocParameters<FRayTracingGlobalIlluminationSceneColorCompositePS::FParameters>();
 	PassParameters->GlobalIlluminationTexture = GraphBuilder.RegisterExternalTexture(GlobalIlluminationRT);
 	PassParameters->GlobalIlluminationSampler = TStaticSamplerState<SF_Point, AM_Clamp, AM_Clamp, AM_Clamp>::GetRHI();
-	PassParameters->RenderTargets[0] = FRenderTargetBinding(GraphBuilder.RegisterExternalTexture(SceneContext.GetSceneColor()), ERenderTargetLoadAction::ENoAction, ERenderTargetStoreAction::ENoAction);
+	PassParameters->RenderTargets[0] = FRenderTargetBinding(GraphBuilder.RegisterExternalTexture(SceneContext.GetSceneColor()), ERenderTargetLoadAction::ELoad, ERenderTargetStoreAction::EStore);
 	PassParameters->SceneTexturesStruct = CreateUniformBufferImmediate(SceneTextures, EUniformBufferUsage::UniformBuffer_SingleDraw);
 
 	GraphBuilder.AddPass(
diff --git a/Engine/Source/Runtime/Renderer/Private/RayTracing/RaytracingSkylight.cpp b/Engine/Source/Runtime/Renderer/Private/RayTracing/RaytracingSkylight.cpp
--- a/Engine/Source/Runtime/Renderer/Private/RayTracing/RaytracingSkylight.cpp
+++ b/Engine/Source/Runtime/Renderer/Private/RayTracing/RaytracingSkylight.cpp
@@ -983,7 +983,7 @@ void FDeferredShadingSceneRenderer::CompositeRayTracingSkyLight(
 		PassParameters->SkyLightTextureSampler = TStaticSamplerState<SF_Point, AM_Clamp, AM_Clamp, AM_Clamp>::GetRHI();
 		PassParameters->ViewUniformBuffer = Views[ViewIndex].ViewUniformBuffer;
 		PassParameters->SceneTexturesStruct = CreateUniformBufferImmediate(SceneTextures, EUniformBufferUsage::UniformBuffer_SingleDraw);
-		PassParameters->RenderTargets[0] = FRenderTargetBinding(GraphBuilder.RegisterExternalTexture(SceneContext.GetSceneColor()), ERenderTargetLoadAction::ENoAction, ERenderTargetStoreAction::ENoAction);
+		PassParameters->RenderTargets[0] = FRenderTargetBinding(GraphBuilder.RegisterExternalTexture(SceneContext.GetSceneColor()), ERenderTargetLoadAction::ELoad, ERenderTargetStoreAction::EStore);
 
 		// dxr_todo: Unify with RTGI compositing workflow
 		GraphBuilder.AddPass(
