Unified RHI validation for CopyTexture/CopySubTextureRegion and implementation of RHICopyTexture function in D3D12 and OpenGL RHI.

#rnx

diff --git a/Engine/Source/Runtime/D3D12RHI/Private/D3D12CommandContext.h b/Engine/Source/Runtime/D3D12RHI/Private/D3D12CommandContext.h
--- a/Engine/Source/Runtime/D3D12RHI/Private/D3D12CommandContext.h
+++ b/Engine/Source/Runtime/D3D12RHI/Private/D3D12CommandContext.h
@@ -250,6 +250,7 @@ public:
 	virtual void RHISetMultipleViewports(uint32 Count, const FViewportBounds* Data) final override;
 	virtual void RHIClearTinyUAV(FUnorderedAccessViewRHIParamRef UnorderedAccessViewRHI, const uint32* Values) final override;
 	virtual void RHICopyToResolveTarget(FTextureRHIParamRef SourceTexture, FTextureRHIParamRef DestTexture, const FResolveParams& ResolveParams) override;
+	virtual void RHICopyTexture(FTextureRHIParamRef SourceTexture, FTextureRHIParamRef DestTexture, const FRHICopyTextureInfo& CopyInfo) final override;
 	virtual void RHITransitionResources(EResourceTransitionAccess TransitionType, FTextureRHIParamRef* InTextures, int32 NumTextures) final override;
 	virtual void RHICopyToStagingBuffer(FVertexBufferRHIParamRef SourceBuffer, FStagingBufferRHIParamRef DestinationStagingBuffer, uint32 Offset, uint32 NumBytes) final override;
 	virtual void RHIWriteGPUFence(FGPUFenceRHIParamRef Fence) final override;
diff --git a/Engine/Source/Runtime/D3D12RHI/Private/D3D12Texture.cpp b/Engine/Source/Runtime/D3D12RHI/Private/D3D12Texture.cpp
--- a/Engine/Source/Runtime/D3D12RHI/Private/D3D12Texture.cpp
+++ b/Engine/Source/Runtime/D3D12RHI/Private/D3D12Texture.cpp
@@ -2835,3 +2835,44 @@ void FD3D12DynamicRHI::RHICopySubTextureRegion(FTexture2DRHIParamRef SourceTextu
 	}
 }
 
+void FD3D12CommandContext::RHICopyTexture(FTextureRHIParamRef SourceTextureRHI, FTextureRHIParamRef DestTextureRHI, const FRHICopyTextureInfo& CopyInfo)
+{
+	FD3D12TextureBase* SourceTexture = GetD3D12TextureFromRHITexture(SourceTextureRHI);
+	FD3D12TextureBase* DestTexture = GetD3D12TextureFromRHITexture(DestTextureRHI);
+
+	const CD3DX12_BOX SourceBoxD3D(
+		CopyInfo.SourcePosition.X,
+		CopyInfo.SourcePosition.Y,
+		CopyInfo.SourcePosition.Z,
+		CopyInfo.SourcePosition.X + CopyInfo.Size.X,
+		CopyInfo.SourcePosition.Y + CopyInfo.Size.Y,
+		CopyInfo.SourcePosition.Z + CopyInfo.Size.Z
+	);
+
+	CD3DX12_TEXTURE_COPY_LOCATION SourceCopyLocation(
+		SourceTexture->GetResource()->GetResource(),
+		CalcSubresource(CopyInfo.SourceMipIndex, CopyInfo.SourceSliceIndex, SourceTexture->GetResource()->GetMipLevels())
+	);
+	CD3DX12_TEXTURE_COPY_LOCATION DestCopyLocation(
+		DestTexture->GetResource()->GetResource(),
+		CalcSubresource(CopyInfo.DestMipIndex, CopyInfo.DestSliceIndex, DestTexture->GetResource()->GetMipLevels())
+	);
+
+	FConditionalScopeResourceBarrier ConditionalScopeResourceBarrierSource(CommandListHandle, SourceTexture->GetResource(), D3D12_RESOURCE_STATE_COPY_SOURCE, SourceCopyLocation.SubresourceIndex);
+	FConditionalScopeResourceBarrier ConditionalScopeResourceBarrierDest(CommandListHandle, DestTexture->GetResource(), D3D12_RESOURCE_STATE_COPY_DEST, DestCopyLocation.SubresourceIndex);
+
+	numCopies++;
+	CommandListHandle.FlushResourceBarriers();
+	CommandListHandle->CopyTextureRegion(
+		&DestCopyLocation,
+		CopyInfo.DestPosition.X,
+		CopyInfo.DestPosition.Y,
+		CopyInfo.DestPosition.Z,
+		&SourceCopyLocation,
+		&SourceBoxD3D
+	);
+
+	CommandListHandle.UpdateResidency(SourceTexture->GetResource());
+	CommandListHandle.UpdateResidency(DestTexture->GetResource());
+}
+
diff --git a/Engine/Source/Runtime/OpenGLDrv/Private/OpenGLTexture.cpp b/Engine/Source/Runtime/OpenGLDrv/Private/OpenGLTexture.cpp
--- a/Engine/Source/Runtime/OpenGLDrv/Private/OpenGLTexture.cpp
+++ b/Engine/Source/Runtime/OpenGLDrv/Private/OpenGLTexture.cpp
@@ -2589,6 +2589,71 @@ void FOpenGLDynamicRHI::RHICopySubTextureRegion(FTexture2DRHIParamRef SourceText
 	ContextState.Framebuffer = (GLuint)-1;
 }
 
+void FOpenGLDynamicRHI::RHICopyTexture(FTextureRHIParamRef SourceTextureRHI, FTextureRHIParamRef DestTextureRHI, const FRHICopyTextureInfo& CopyInfo)
+{
+	VERIFY_GL_SCOPE();
+	FOpenGLTextureBase* SourceTexture = GetOpenGLTextureFromRHITexture(SourceTextureRHI);
+	FOpenGLTextureBase* DestTexture = GetOpenGLTextureFromRHITexture(DestTextureRHI);
+
+	check(SourceTexture->Target == DestTexture->Target);
+
+	// Use a texture stage that's not likely to be used for draws, to avoid waiting
+	FOpenGLContextState& ContextState = GetContextStateForCurrentContext();
+	CachedSetupTextureStage(ContextState, FOpenGL::GetMaxCombinedTextureImageUnits() - 1, DestTexture->Target, DestTexture->Resource, 0, DestTextureRHI->GetNumMips());
+	CachedBindPixelUnpackBuffer(ContextState, 0);
+
+	// Convert sub texture regions to GL types
+	GLint XOffset = CopyInfo.DestPosition.X;
+	GLint YOffset = CopyInfo.DestPosition.Y;
+	GLint ZOffset = CopyInfo.DestPosition.Z;
+	GLint X = CopyInfo.SourcePosition.X;
+	GLint Y = CopyInfo.SourcePosition.Y;
+	GLint Z = CopyInfo.SourcePosition.Z;
+	GLsizei Width = CopyInfo.Size.X;
+	GLsizei Height = CopyInfo.Size.Y;
+	GLsizei Depth = CopyInfo.Size.Z;
+
+	// Bind source texture to an FBO to read from
+	for (GLsizei Layer = 0; Layer < Depth; ++Layer)
+	{
+		FOpenGLTextureBase* RenderTargets[1] = { SourceTexture };
+		uint32 MipLevels[1] = { CopyInfo.SourceMipIndex };
+		uint32 ArrayIndices[1] = { CopyInfo.SourceSliceIndex + static_cast<uint32>(Layer) };
+
+		GLuint SourceFBO = GetOpenGLFramebuffer(1, RenderTargets, ArrayIndices, MipLevels, nullptr);
+		check(SourceFBO != 0);
+
+		glBindFramebuffer(GL_FRAMEBUFFER, SourceFBO);
+
+		FOpenGL::ReadBuffer(GL_COLOR_ATTACHMENT0);
+
+		switch (DestTexture->Target)
+		{
+		case GL_TEXTURE_1D:
+			FOpenGL::CopyTexSubImage1D(DestTexture->Target, CopyInfo.DestMipIndex, XOffset, X, 0, Width);
+			break;
+		case GL_TEXTURE_1D_ARRAY:
+		case GL_TEXTURE_2D:
+		case GL_TEXTURE_CUBE_MAP_POSITIVE_X:
+		case GL_TEXTURE_CUBE_MAP_NEGATIVE_X:
+		case GL_TEXTURE_CUBE_MAP_POSITIVE_Y:
+		case GL_TEXTURE_CUBE_MAP_NEGATIVE_Y:
+		case GL_TEXTURE_CUBE_MAP_POSITIVE_Z:
+		case GL_TEXTURE_CUBE_MAP_NEGATIVE_Z:
+		case GL_TEXTURE_RECTANGLE:
+			FOpenGL::CopyTexSubImage2D(DestTexture->Target, CopyInfo.DestMipIndex, XOffset, YOffset, X, Y, Width, Height);
+			break;
+		case GL_TEXTURE_3D:
+		case GL_TEXTURE_2D_ARRAY:
+		case GL_TEXTURE_CUBE_MAP_ARRAY:
+			FOpenGL::CopyTexSubImage3D(DestTexture->Target, CopyInfo.DestMipIndex, XOffset, YOffset, ZOffset + Layer, X, Y, Width, Depth);
+			break;
+		}
+	}
+
+	ContextState.Framebuffer = (GLuint)-1;
+}
+
 FTexture2DRHIRef FOpenGLDynamicRHI::RHICreateTexture2DFromResource(EPixelFormat Format, uint32 SizeX, uint32 SizeY, uint32 NumMips, uint32 NumSamples, uint32 NumSamplesTileMem, const FClearValueBinding& ClearValueBinding, GLuint Resource, uint32 TexCreateFlags)
 {
 	FOpenGLTexture2D* Texture2D = new FOpenGLTexture2D(
diff --git a/Engine/Source/Runtime/OpenGLDrv/Public/OpenGL3.h b/Engine/Source/Runtime/OpenGLDrv/Public/OpenGL3.h
--- a/Engine/Source/Runtime/OpenGLDrv/Public/OpenGL3.h
+++ b/Engine/Source/Runtime/OpenGLDrv/Public/OpenGL3.h
@@ -411,6 +411,11 @@ struct FOpenGL3 : public FOpenGLBase
 		glTexSubImage3D(Target, Level, XOffset, YOffset, ZOffset, Width, Height, Depth, Format, Type, PixelData);
 	}
 
+	static FORCEINLINE void	CopyTexSubImage1D(GLenum Target, GLint Level, GLint XOffset, GLint X, GLint Y, GLsizei Width)
+	{
+		glCopyTexSubImage1D(Target, Level, XOffset, X, Y, Width);
+	}
+
 	static FORCEINLINE void	CopyTexSubImage2D(GLenum Target, GLint Level, GLint XOffset, GLint YOffset, GLint X, GLint Y, GLsizei Width, GLsizei Height)
 	{
 		glCopyTexSubImage2D(Target, Level, XOffset, YOffset, X, Y, Width, Height);
diff --git a/Engine/Source/Runtime/OpenGLDrv/Public/OpenGLDrv.h b/Engine/Source/Runtime/OpenGLDrv/Public/OpenGLDrv.h
--- a/Engine/Source/Runtime/OpenGLDrv/Public/OpenGLDrv.h
+++ b/Engine/Source/Runtime/OpenGLDrv/Public/OpenGLDrv.h
@@ -525,6 +525,7 @@ public:
 	virtual void RHIDiscardRenderTargets(bool Depth,bool Stencil,uint32 ColorBitMask) final override;
 
 	virtual void RHICopySubTextureRegion(FTexture2DRHIParamRef SourceTexture, FTexture2DRHIParamRef DestinationTexture, FBox2D SourceBox, FBox2D DestinationBox) final override;
+	virtual void RHICopyTexture(FTextureRHIParamRef SourceTexture, FTextureRHIParamRef DestTexture, const FRHICopyTextureInfo& CopyInfo) final override;
 
 	// FOpenGLDynamicRHI interface.
 	virtual FTexture2DRHIRef RHICreateTexture2DFromResource(EPixelFormat Format, uint32 SizeX, uint32 SizeY, uint32 NumMips, uint32 NumSamples, uint32 NumSamplesTileMem, const FClearValueBinding& ClearValueBinding, GLuint Resource, uint32 Flags);
diff --git a/Engine/Source/Runtime/RHI/Private/RHIValidation.cpp b/Engine/Source/Runtime/RHI/Private/RHIValidation.cpp
--- a/Engine/Source/Runtime/RHI/Private/RHIValidation.cpp
+++ b/Engine/Source/Runtime/RHI/Private/RHIValidation.cpp
@@ -89,18 +89,6 @@ void FValidationRHI::ValidatePipeline(const FGraphicsPipelineStateInitializer& P
 	}
 }
 
-void FValidationRHI::ValidateCopySubTextureRegion(FRHITexture2D* SourceTexture, FRHITexture2D* DestinationTexture, FBox2D SourceBox, FBox2D DestinationBox)
-{
-	const int32 SrcWidth = (int32)(SourceBox.Max.X - SourceBox.Min.X);
-	const int32 SrcHeight = (int32)(SourceBox.Max.Y - SourceBox.Min.Y);
-	const int32 DestWidth = (int32)(DestinationBox.Max.X - DestinationBox.Min.X);
-	const int32 DestHeight = (int32)(DestinationBox.Max.Y - DestinationBox.Min.Y);
-
-	check(SourceTexture != nullptr && DestinationTexture != nullptr);
-	//	check(SrcWidth > 0 && SrcHeight > 0);
-	check(DestWidth == SrcWidth && DestHeight == SrcHeight);
-}
-
 
 FValidationComputeContext::FValidationComputeContext(FValidationRHI* InRHI)
 	: RHIContext(nullptr)
diff --git a/Engine/Source/Runtime/RHI/Public/RHIValidation.h b/Engine/Source/Runtime/RHI/Public/RHIValidation.h
--- a/Engine/Source/Runtime/RHI/Public/RHIValidation.h
+++ b/Engine/Source/Runtime/RHI/Public/RHIValidation.h
@@ -1613,12 +1613,25 @@ public:
 	/* Copy the source box pixels in the destination box texture, return true if implemented for the current platform*/
 	virtual void RHICopySubTextureRegion_RenderThread(class FRHICommandListImmediate& RHICmdList, FRHITexture2D* SourceTexture, FRHITexture2D* DestinationTexture, FBox2D SourceBox, FBox2D DestinationBox) override final
 	{
+		int32 Width = static_cast<int32>(SourceBox.GetSize().X);
+		int32 Height = static_cast<int32>(SourceBox.GetSize().Y);
+		int32 SrcX = static_cast<int32>(SourceBox.Min.X);
+		int32 SrcY = static_cast<int32>(SourceBox.Min.Y);
+		int32 DestX = static_cast<int32>(DestinationBox.Min.X);
+		int32 DestY = static_cast<int32>(DestinationBox.Min.Y);
+		FValidationRHIUtils::ValidateCopyTexture(SourceTexture, DestinationTexture, FIntVector(Width, Height, 1), FIntVector(SrcX, SrcY, 0), FIntVector(DestX, DestY, 0));
 		return RHI->RHICopySubTextureRegion_RenderThread(RHICmdList, SourceTexture, DestinationTexture, SourceBox, DestinationBox);
 	}
 
 	virtual void RHICopySubTextureRegion(FRHITexture2D* SourceTexture, FRHITexture2D* DestinationTexture, FBox2D SourceBox, FBox2D DestinationBox) override final
 	{
-		ValidateCopySubTextureRegion(SourceTexture, DestinationTexture, SourceBox, DestinationBox);
+		int32 Width = static_cast<int32>(SourceBox.GetSize().X);
+		int32 Height = static_cast<int32>(SourceBox.GetSize().Y);
+		int32 SrcX = static_cast<int32>(SourceBox.Min.X);
+		int32 SrcY = static_cast<int32>(SourceBox.Min.Y);
+		int32 DestX = static_cast<int32>(DestinationBox.Min.X);
+		int32 DestY = static_cast<int32>(DestinationBox.Min.Y);
+		FValidationRHIUtils::ValidateCopyTexture(SourceTexture, DestinationTexture, FIntVector(Width, Height, 1), FIntVector(SrcX, SrcY, 0), FIntVector(DestX, DestY, 0));
 		RHI->RHICopySubTextureRegion(SourceTexture, DestinationTexture, SourceBox, DestinationBox);
 	}
 
@@ -1687,7 +1700,6 @@ private:
 	FString						RHIName;
 
 	void ValidatePipeline(const FGraphicsPipelineStateInitializer& Initializer);
-	void ValidateCopySubTextureRegion(FRHITexture2D* SourceTexture, FRHITexture2D* DestinationTexture, FBox2D SourceBox, FBox2D DestinationBox);
 };
 
 extern RHI_API FValidationRHI* GValidationRHI;
diff --git a/Engine/Source/Runtime/RHI/Public/RHIValidationCommon.h b/Engine/Source/Runtime/RHI/Public/RHIValidationCommon.h
--- a/Engine/Source/Runtime/RHI/Public/RHIValidationCommon.h
+++ b/Engine/Source/Runtime/RHI/Public/RHIValidationCommon.h
@@ -14,3 +14,42 @@
 #include "RHIResources.h"
 #include "RHIContext.h"
 #include "DynamicRHI.h"
+
+class FValidationRHIUtils
+{
+public:
+	static void ValidateCopyTexture(
+		FTextureRHIParamRef	SourceTexture,
+		FTextureRHIParamRef	DestTexture,
+		FIntVector			CopySize		= FIntVector::ZeroValue,
+		const FIntVector&	SourcePosition	= FIntVector::ZeroValue,
+		const FIntVector&	DestPosition	= FIntVector::ZeroValue)
+	{
+		check(SourceTexture);
+		check(DestTexture);
+		checkf(SourceTexture->GetFormat() == DestTexture->GetFormat(), TEXT("Some RHIs do not allow format conversion by the GPU for transfer operations!"));
+
+		FIntVector SrcSize = SourceTexture->GetSizeXYZ();
+		FIntVector DestSize = DestTexture->GetSizeXYZ();
+		if (CopySize == FIntVector::ZeroValue)
+		{
+			CopySize = SrcSize;
+		}
+
+		checkf(CopySize.X <= DestSize.X && CopySize.Y <= DestSize.Y, TEXT("Some RHIs can't perform scaling operations [%dx%d to %dx%d] during copies!"), SrcSize.X, SrcSize.Y, DestSize.X, DestSize.Y);
+
+		check(SourcePosition.X >= 0 && SourcePosition.Y >= 0 && SourcePosition.Z >= 0);
+		check(SourcePosition.X + CopySize.X <= SrcSize.X && SourcePosition.Y + CopySize.Y <= SrcSize.Y);
+
+		check(DestPosition.X >= 0 && DestPosition.Y >= 0 && DestPosition.Z >= 0);
+		check(DestPosition.X + CopySize.X <= DestSize.X && DestPosition.Y + CopySize.Y <= DestSize.Y);
+
+		if (SourceTexture->GetTexture3D() && DestTexture->GetTexture3D())
+		{
+			check(SourcePosition.Z + CopySize.Z <= SrcSize.Z);
+			check(DestPosition.Z + CopySize.Z <= DestSize.Z);
+		}
+	}
+
+};
+
diff --git a/Engine/Source/Runtime/RHI/Public/RHIValidationContext.h b/Engine/Source/Runtime/RHI/Public/RHIValidationContext.h
--- a/Engine/Source/Runtime/RHI/Public/RHIValidationContext.h
+++ b/Engine/Source/Runtime/RHI/Public/RHIValidationContext.h
@@ -800,28 +800,7 @@ public:
 	virtual void RHICopyTexture(FTextureRHIParamRef SourceTexture, FTextureRHIParamRef DestTexture, const FRHICopyTextureInfo& CopyInfo) override final
 	{
 		ensureMsgf(!State.bInsideBeginRenderPass, TEXT("Copying inside a RenderPass is not efficient!"));
-		check(SourceTexture);
-		check(DestTexture);
-		checkf(SourceTexture->GetFormat() == DestTexture->GetFormat(), TEXT("Some RHIs do not allow format conversion by the GPU for transfer operations!"));
-
-		FIntVector CopySize = CopyInfo.Size;
-		FIntVector SrcSize = SourceTexture->GetSizeXYZ();
-		FIntVector DestSize = DestTexture->GetSizeXYZ();
-		if (CopySize == FIntVector::ZeroValue)
-		{
-			CopySize = SrcSize;
-		}
-
-		checkf(SrcSize.X <= DestSize.X && SrcSize.Y <= DestSize.Y, TEXT("Some RHIs can't perform scaling operations [%dx%d to %dx%d] during copies!"), SrcSize.X, SrcSize.Y, DestSize.X, DestSize.Y);
-		check(CopyInfo.SourcePosition.X >= 0 && CopyInfo.SourcePosition.Y >= 0 && CopyInfo.SourcePosition.Z >= 0);
-		check(CopyInfo.DestPosition.X >= 0 && CopyInfo.DestPosition.Y >= 0 && CopyInfo.DestPosition.Z >= 0);
-		check(CopyInfo.SourcePosition.X + CopySize.X <= CopyInfo.DestPosition.X + DestSize.X);
-		check(CopyInfo.SourcePosition.Y + CopySize.Y <= CopyInfo.DestPosition.Y + DestSize.Y);
-		if (SourceTexture->GetTexture3D() && DestTexture->GetTexture3D())
-		{
-			check(CopyInfo.SourcePosition.Z + CopySize.Z <= CopyInfo.DestPosition.Z + DestSize.Z);
-		}
-
+		FValidationRHIUtils::ValidateCopyTexture(SourceTexture, DestTexture, CopyInfo.SourcePosition, CopyInfo.DestPosition, CopyInfo.Size);
 		RHIContext->RHICopyTexture(SourceTexture, DestTexture, CopyInfo);
 	}
 
diff --git a/Engine/Source/Runtime/Renderer/Private/PostProcess/PostProcessBufferInspector.cpp b/Engine/Source/Runtime/Renderer/Private/PostProcess/PostProcessBufferInspector.cpp
--- a/Engine/Source/Runtime/Renderer/Private/PostProcess/PostProcessBufferInspector.cpp
+++ b/Engine/Source/Runtime/Renderer/Private/PostProcess/PostProcessBufferInspector.cpp
@@ -94,9 +94,11 @@ void FRCPassPostProcessBufferInspector::Process(FRenderingCompositePassContext&
 			// Pixel Depth
 			if (Scene->PixelInspectorData.RenderTargetBufferDepth[PixelInspectorRequest->BufferIndex] != nullptr)
 			{
-				FVector2D SourcePoint(FMath::FloorToInt(SourceViewportUV.X * View.ViewRect.Width()),
-									  FMath::FloorToInt(SourceViewportUV.Y * View.ViewRect.Height()));
-				FBox2D SourceBox(SourcePoint, SourcePoint + ExtendSize);
+				const FIntVector SourcePoint(
+					FMath::FloorToInt(SourceViewportUV.X * View.ViewRect.Width()),
+					FMath::FloorToInt(SourceViewportUV.Y * View.ViewRect.Height()),
+					0
+				);
 				
 				const FTexture2DRHIRef &DestinationBufferDepth = Scene->PixelInspectorData.RenderTargetBufferDepth[PixelInspectorRequest->BufferIndex]->GetRenderTargetTexture();
 				if (DestinationBufferDepth.IsValid())
@@ -104,8 +106,10 @@ void FRCPassPostProcessBufferInspector::Process(FRenderingCompositePassContext&
 					FTexture2DRHIRef SourceBufferSceneDepth = SceneContext.GetSceneDepthTexture();
 					if (DestinationBufferDepth->GetFormat() == SourceBufferSceneDepth->GetFormat())
 					{
-						FBox2D DestinationBox(FVector2D(0.0f, 0.0f), FVector2D(1.0f, 1.0f));
-						RHICmdList.CopySubTextureRegion(SourceBufferSceneDepth, DestinationBufferDepth, SourceBox, DestinationBox);
+						FRHICopyTextureInfo CopyInfo;
+						CopyInfo.SourcePosition = SourcePoint;
+						CopyInfo.Size = FIntVector(1);
+						RHICmdList.CopyTexture(SourceBufferSceneDepth, DestinationBufferDepth, CopyInfo);
 					}
 				}
 			}
@@ -115,9 +119,11 @@ void FRCPassPostProcessBufferInspector::Process(FRenderingCompositePassContext&
 			const FTexture2DRHIRef &DestinationBufferFinalColor = Scene->PixelInspectorData.RenderTargetBufferFinalColor[PixelInspectorRequest->BufferIndex]->GetRenderTargetTexture();
 			if (DestinationBufferFinalColor.IsValid())
 			{
-				FVector2D SourcePoint(
+				const FIntVector SourcePoint(
 					FMath::FloorToInt(SourceViewportUV.X * Context.SceneColorViewRect.Width()),
-					FMath::FloorToInt(SourceViewportUV.Y * Context.SceneColorViewRect.Height()));
+					FMath::FloorToInt(SourceViewportUV.Y * Context.SceneColorViewRect.Height()),
+					0
+				);
 
 				const FRenderingCompositeOutputRef* OutputRef0 = GetInput(ePId_Input0);
 				if (OutputRef0)
@@ -128,11 +134,30 @@ void FRCPassPostProcessBufferInspector::Process(FRenderingCompositePassContext&
 						FTexture2DRHIRef SourceBufferFinalColor = (FRHITexture2D*)(Input->PooledRenderTarget->GetRenderTargetItem().ShaderResourceTexture.GetReference());
 						if (DestinationBufferFinalColor->GetFormat() == SourceBufferFinalColor->GetFormat())
 						{
-							FVector2D ExtendSizeContextLeft(FMath::FloorToFloat((float)DestinationBufferFinalColor->GetSizeX()/2.0f), FMath::FloorToFloat((float)DestinationBufferFinalColor->GetSizeY()/2.0f));
-							FVector2D ExtendSizeContextRight(FMath::CeilToFloat((float)DestinationBufferFinalColor->GetSizeX() / 2.0f), FMath::CeilToFloat((float)DestinationBufferFinalColor->GetSizeY() / 2.0f));
-							FBox2D SourceBoxContext(SourcePoint- ExtendSizeContextLeft, SourcePoint+ ExtendSizeContextRight);
-							FBox2D DestinationBox(FVector2D(0.0f, 0.0f), FVector2D((float)DestinationBufferFinalColor->GetSizeX(), (float)DestinationBufferFinalColor->GetSizeY()));
-							RHICmdList.CopySubTextureRegion(SourceBufferFinalColor, DestinationBufferFinalColor, SourceBoxContext, DestinationBox);
+							FRHICopyTextureInfo CopyInfo;
+							CopyInfo.Size = DestinationBufferFinalColor->GetSizeXYZ();
+							CopyInfo.SourcePosition = SourcePoint - CopyInfo.Size/2;
+
+							const FIntVector OutlineCornerMin(
+								FMath::Min(CopyInfo.SourcePosition.X - Context.SceneColorViewRect.Min.X, 0),
+								FMath::Min(CopyInfo.SourcePosition.Y - Context.SceneColorViewRect.Min.Y, 0),
+								0
+							);
+							CopyInfo.SourcePosition -= OutlineCornerMin;
+							CopyInfo.DestPosition -= OutlineCornerMin;
+							CopyInfo.Size += OutlineCornerMin;
+
+							const FIntVector OutlineCornerMax(
+								FMath::Max(0, CopyInfo.SourcePosition.X + CopyInfo.Size.X - Context.SceneColorViewRect.Max.X),
+								FMath::Max(0, CopyInfo.SourcePosition.Y + CopyInfo.Size.Y - Context.SceneColorViewRect.Max.Y),
+								0
+							);
+							CopyInfo.Size -= OutlineCornerMax;
+
+							if (CopyInfo.Size.X > 0 && CopyInfo.Size.Y > 0)
+							{
+								RHICmdList.CopyTexture(SourceBufferFinalColor, DestinationBufferFinalColor, CopyInfo);
+							}
 						}
 					}
 				}
@@ -143,10 +168,11 @@ void FRCPassPostProcessBufferInspector::Process(FRenderingCompositePassContext&
 			const FTexture2DRHIRef &DestinationBufferSceneColor = Scene->PixelInspectorData.RenderTargetBufferSceneColor[PixelInspectorRequest->BufferIndex]->GetRenderTargetTexture();
 			if (DestinationBufferSceneColor.IsValid())
 			{
-				FVector2D SourcePoint(
+				const FIntVector SourcePoint(
 					FMath::FloorToInt(SourceViewportUV.X * View.ViewRect.Width()),
-					FMath::FloorToInt(SourceViewportUV.Y * View.ViewRect.Height()));
-				FBox2D SourceBox(SourcePoint, SourcePoint + ExtendSize);
+					FMath::FloorToInt(SourceViewportUV.Y * View.ViewRect.Height()),
+					0
+				);
 
 				const FRenderingCompositeOutputRef* OutputRef0 = GetInput(ePId_Input2);
 				if (OutputRef0)
@@ -157,8 +183,10 @@ void FRCPassPostProcessBufferInspector::Process(FRenderingCompositePassContext&
 						FTexture2DRHIRef SourceBufferSceneColor = (FRHITexture2D*)(Input->PooledRenderTarget->GetRenderTargetItem().ShaderResourceTexture.GetReference());
 						if (DestinationBufferSceneColor->GetFormat() == SourceBufferSceneColor->GetFormat())
 						{
-							FBox2D DestinationBox(FVector2D(0.0f, 0.0f), FVector2D(1.0f, 1.0f));
-							RHICmdList.CopySubTextureRegion(SourceBufferSceneColor, DestinationBufferSceneColor, SourceBox, DestinationBox);
+							FRHICopyTextureInfo CopyInfo;
+							CopyInfo.SourcePosition = SourcePoint;
+							CopyInfo.Size = FIntVector(1, 1, 1);
+							RHICmdList.CopyTexture(SourceBufferSceneColor, DestinationBufferSceneColor, CopyInfo);
 						}
 					}
 				}
@@ -169,9 +197,11 @@ void FRCPassPostProcessBufferInspector::Process(FRenderingCompositePassContext&
 			const FTexture2DRHIRef &DestinationBufferHDR = Scene->PixelInspectorData.RenderTargetBufferHDR[PixelInspectorRequest->BufferIndex]->GetRenderTargetTexture();
 			if (InputDescHDR != nullptr && DestinationBufferHDR.IsValid())
 			{
-				FVector2D SourcePoint(FMath::FloorToInt(SourceViewportUV.X * Context.SceneColorViewRect.Width()),
-									  FMath::FloorToInt(SourceViewportUV.Y * Context.SceneColorViewRect.Height()));
-				FBox2D SourceBox(SourcePoint, SourcePoint + ExtendSize);
+				const FIntVector SourcePoint(
+					FMath::FloorToInt(SourceViewportUV.X * Context.SceneColorViewRect.Width()),
+					FMath::FloorToInt(SourceViewportUV.Y * Context.SceneColorViewRect.Height()),
+					0
+				);
 
 				const FRenderingCompositeOutputRef* OutputRef1 = GetInput(ePId_Input1);
 				if (OutputRef1)
@@ -182,8 +212,10 @@ void FRCPassPostProcessBufferInspector::Process(FRenderingCompositePassContext&
 						FTexture2DRHIRef SourceBufferHDR = (FRHITexture2D*)(Input->PooledRenderTarget->GetRenderTargetItem().ShaderResourceTexture.GetReference());
 						if (SourceBufferHDR.IsValid() && DestinationBufferHDR->GetFormat() == SourceBufferHDR->GetFormat())
 						{
-							FBox2D DestinationBox(FVector2D(0.0f, 0.0f), FVector2D(1.0f, 1.0f));
-							RHICmdList.CopySubTextureRegion(SourceBufferHDR, DestinationBufferHDR, SourceBox, DestinationBox);
+							FRHICopyTextureInfo CopyInfo;
+							CopyInfo.SourcePosition = SourcePoint;
+							CopyInfo.Size = FIntVector(1, 1, 1);
+							RHICmdList.CopyTexture(SourceBufferHDR, DestinationBufferHDR, CopyInfo);
 						}
 					}
 				}
@@ -193,9 +225,11 @@ void FRCPassPostProcessBufferInspector::Process(FRenderingCompositePassContext&
 			// GBuffer A
 			if (Scene->PixelInspectorData.RenderTargetBufferA[PixelInspectorRequest->BufferIndex] != nullptr)
 			{
-				FVector2D SourcePoint(FMath::FloorToInt(SourceViewportUV.X * View.ViewRect.Width()),
-									  FMath::FloorToInt(SourceViewportUV.Y * View.ViewRect.Height()));
-				FBox2D SourceBox(SourcePoint, SourcePoint + ExtendSize);
+				const FIntVector SourcePoint(
+					FMath::FloorToInt(SourceViewportUV.X * View.ViewRect.Width()),
+					FMath::FloorToInt(SourceViewportUV.Y * View.ViewRect.Height()),
+					0
+				);
 
 				const FTexture2DRHIRef &DestinationBufferA = Scene->PixelInspectorData.RenderTargetBufferA[PixelInspectorRequest->BufferIndex]->GetRenderTargetTexture();
 				if (DestinationBufferA.IsValid() && SceneContext.GBufferA.IsValid() && SceneContext.GBufferA->GetRenderTargetItem().ShaderResourceTexture.IsValid())
@@ -203,8 +237,10 @@ void FRCPassPostProcessBufferInspector::Process(FRenderingCompositePassContext&
 					FTexture2DRHIRef SourceBufferA = (FRHITexture2D*)(SceneContext.GBufferA->GetRenderTargetItem().ShaderResourceTexture.GetReference());
 					if (SourceBufferA.IsValid() && DestinationBufferA->GetFormat() == SourceBufferA->GetFormat())
 					{
-						FBox2D DestinationBox(FVector2D(0.0f, 0.0f), FVector2D(1.0f, 1.0f));
-						RHICmdList.CopySubTextureRegion(SourceBufferA, DestinationBufferA, SourceBox, DestinationBox);
+						FRHICopyTextureInfo CopyInfo;
+						CopyInfo.SourcePosition = SourcePoint;
+						CopyInfo.Size = FIntVector(1, 1, 1);
+						RHICmdList.CopyTexture(SourceBufferA, DestinationBufferA, CopyInfo);
 					}
 				}
 			}
@@ -214,17 +250,21 @@ void FRCPassPostProcessBufferInspector::Process(FRenderingCompositePassContext&
 			const FTexture2DRHIRef &DestinationBufferBCDE = Scene->PixelInspectorData.RenderTargetBufferBCDE[PixelInspectorRequest->BufferIndex]->GetRenderTargetTexture();
 			if (DestinationBufferBCDE.IsValid())
 			{
-				FVector2D SourcePoint(FMath::FloorToInt(SourceViewportUV.X * View.ViewRect.Width()),
-									  FMath::FloorToInt(SourceViewportUV.Y * View.ViewRect.Height()));
-				FBox2D SourceBox(SourcePoint, SourcePoint + ExtendSize);
+				const FIntVector SourcePoint(
+					FMath::FloorToInt(SourceViewportUV.X * View.ViewRect.Width()),
+					FMath::FloorToInt(SourceViewportUV.Y * View.ViewRect.Height()),
+					0
+				);
 
 				if (SceneContext.GBufferB.IsValid() && SceneContext.GBufferB->GetRenderTargetItem().ShaderResourceTexture.IsValid())
 				{
 					FTexture2DRHIRef SourceBufferB = (FRHITexture2D*)(SceneContext.GBufferB->GetRenderTargetItem().ShaderResourceTexture.GetReference());
 					if (SourceBufferB.IsValid() && DestinationBufferBCDE->GetFormat() == SourceBufferB->GetFormat())
 					{
-						FBox2D DestinationBox(FVector2D(0.0f, 0.0f), FVector2D(1.0f, 1.0f));
-						RHICmdList.CopySubTextureRegion(SourceBufferB, DestinationBufferBCDE, SourceBox, DestinationBox);
+						FRHICopyTextureInfo CopyInfo;
+						CopyInfo.SourcePosition = SourcePoint;
+						CopyInfo.Size = FIntVector(1, 1, 1);
+						RHICmdList.CopyTexture(SourceBufferB, DestinationBufferBCDE, CopyInfo);
 					}
 				}
 
@@ -233,8 +273,11 @@ void FRCPassPostProcessBufferInspector::Process(FRenderingCompositePassContext&
 					FTexture2DRHIRef SourceBufferC = (FRHITexture2D*)(SceneContext.GBufferC->GetRenderTargetItem().ShaderResourceTexture.GetReference());
 					if (SourceBufferC.IsValid() && DestinationBufferBCDE->GetFormat() == SourceBufferC->GetFormat())
 					{
-						FBox2D DestinationBox(FVector2D(1.0f, 0.0f), FVector2D(2.0f, 1.0f));
-						RHICmdList.CopySubTextureRegion(SourceBufferC, DestinationBufferBCDE, SourceBox, DestinationBox);
+						FRHICopyTextureInfo CopyInfo;
+						CopyInfo.SourcePosition = SourcePoint;
+						CopyInfo.DestPosition = FIntVector(1, 0, 0);
+						CopyInfo.Size = FIntVector(1, 1, 1);
+						RHICmdList.CopyTexture(SourceBufferC, DestinationBufferBCDE, CopyInfo);
 					}
 				}
 
@@ -243,8 +286,11 @@ void FRCPassPostProcessBufferInspector::Process(FRenderingCompositePassContext&
 					FTexture2DRHIRef SourceBufferD = (FRHITexture2D*)(SceneContext.GBufferD->GetRenderTargetItem().ShaderResourceTexture.GetReference());
 					if (SourceBufferD.IsValid() && DestinationBufferBCDE->GetFormat() == SourceBufferD->GetFormat())
 					{
-						FBox2D DestinationBox(FVector2D(2.0f, 0.0f), FVector2D(3.0f, 1.0f));
-						RHICmdList.CopySubTextureRegion(SourceBufferD, DestinationBufferBCDE, SourceBox, DestinationBox);
+						FRHICopyTextureInfo CopyInfo;
+						CopyInfo.SourcePosition = SourcePoint;
+						CopyInfo.DestPosition = FIntVector(2, 0, 0);
+						CopyInfo.Size = FIntVector(1, 1, 1);
+						RHICmdList.CopyTexture(SourceBufferD, DestinationBufferBCDE, CopyInfo);
 					}
 				}
 
@@ -253,8 +299,11 @@ void FRCPassPostProcessBufferInspector::Process(FRenderingCompositePassContext&
 					FTexture2DRHIRef SourceBufferE = (FRHITexture2D*)(SceneContext.GBufferE->GetRenderTargetItem().ShaderResourceTexture.GetReference());
 					if (SourceBufferE.IsValid() && DestinationBufferBCDE->GetFormat() == SourceBufferE->GetFormat())
 					{
-						FBox2D DestinationBox(FVector2D(3.0f, 0.0f), FVector2D(4.0f, 1.0f));
-						RHICmdList.CopySubTextureRegion(SourceBufferE, DestinationBufferBCDE, SourceBox, DestinationBox);
+						FRHICopyTextureInfo CopyInfo;
+						CopyInfo.SourcePosition = SourcePoint;
+						CopyInfo.DestPosition = FIntVector(3, 0, 0);
+						CopyInfo.Size = FIntVector(1, 1, 1);
+						RHICmdList.CopyTexture(SourceBufferE, DestinationBufferBCDE, CopyInfo);
 					}
 				}
 			}
