Release new MidGreyDummy texture to fix error on shutdown

diff --git a/Engine/Source/Runtime/Renderer/Private/SystemTextures.cpp b/Engine/Source/Runtime/Renderer/Private/SystemTextures.cpp
--- a/Engine/Source/Runtime/Renderer/Private/SystemTextures.cpp
+++ b/Engine/Source/Runtime/Renderer/Private/SystemTextures.cpp
@@ -589,6 +589,7 @@ void FSystemTextures::ReleaseDynamicRHI()
 	GreenDummy.SafeRelease();
 	DefaultNormal8Bit.SafeRelease();
 	VolumetricBlackDummy.SafeRelease();
+	MidGreyDummy.SafeRelease();
 
 	GRenderTargetPool.FreeUnusedResources();
 
