Fix issues with translucent objects not correctly applying skylight changes. FScene::SetSkyLight only checks for ShouldRenderSkylightInBasePass(BLEND_Opaque) which isn't enough to recache translucent meshes. Adding extra condition for translucents would mean that any movable skylight runtime property change with MarkRenderStateDirty(), so the ones which don't use ENQUEUE_RENDER_COMMAND for a fast path update, will cause a full mesh draw command recreate. Decided just to disable caching on all translucent meshes instead (just like it was done before HLR).

diff --git a/Engine/Source/Runtime/Renderer/Private/BasePassRendering.cpp b/Engine/Source/Runtime/Renderer/Private/BasePassRendering.cpp
--- a/Engine/Source/Runtime/Renderer/Private/BasePassRendering.cpp
+++ b/Engine/Source/Runtime/Renderer/Private/BasePassRendering.cpp
@@ -1565,6 +1565,6 @@ FMeshPassProcessor* CreateTranslucencyAllPassProcessor(const FScene* Scene, cons
 }
 
 FRegisterPassProcessorCreateFunction RegisterBasePass(&CreateBasePassProcessor, EShadingPath::Deferred, EMeshPass::BasePass, EMeshPassFlags::CachedMeshCommands | EMeshPassFlags::MainView);
-FRegisterPassProcessorCreateFunction RegisterTranslucencyStandardPass(&CreateTranslucencyStandardPassProcessor, EShadingPath::Deferred, EMeshPass::TranslucencyStandard, EMeshPassFlags::CachedMeshCommands | EMeshPassFlags::MainView);
-FRegisterPassProcessorCreateFunction RegisterTranslucencyAfterDOFPass(&CreateTranslucencyAfterDOFProcessor, EShadingPath::Deferred, EMeshPass::TranslucencyAfterDOF, EMeshPassFlags::CachedMeshCommands | EMeshPassFlags::MainView);
-FRegisterPassProcessorCreateFunction RegisterTranslucencyAllPass(&CreateTranslucencyAllPassProcessor, EShadingPath::Deferred, EMeshPass::TranslucencyAll, EMeshPassFlags::CachedMeshCommands | EMeshPassFlags::MainView);
\ No newline at end of file
+FRegisterPassProcessorCreateFunction RegisterTranslucencyStandardPass(&CreateTranslucencyStandardPassProcessor, EShadingPath::Deferred, EMeshPass::TranslucencyStandard, EMeshPassFlags::MainView);
+FRegisterPassProcessorCreateFunction RegisterTranslucencyAfterDOFPass(&CreateTranslucencyAfterDOFProcessor, EShadingPath::Deferred, EMeshPass::TranslucencyAfterDOF, EMeshPassFlags::MainView);
+FRegisterPassProcessorCreateFunction RegisterTranslucencyAllPass(&CreateTranslucencyAllPassProcessor, EShadingPath::Deferred, EMeshPass::TranslucencyAll, EMeshPassFlags::MainView);
\ No newline at end of file
