Raytracing testbed is now launched from the automated tests dialog instead of being a smoke test.
Fixed potential collision if the test is launched while the rendering thread is working

diff --git a/Engine/Source/Runtime/Renderer/Private/Tests/RayTracingTestbed.cpp b/Engine/Source/Runtime/Renderer/Private/Tests/RayTracingTestbed.cpp
--- a/Engine/Source/Runtime/Renderer/Private/Tests/RayTracingTestbed.cpp
+++ b/Engine/Source/Runtime/Renderer/Private/Tests/RayTracingTestbed.cpp
@@ -6,10 +6,7 @@
 
 #if WITH_DEV_AUTOMATION_TESTS
 
-IMPLEMENT_SIMPLE_AUTOMATION_TEST(FRayTracingTestbed, "System.Renderer.RayTracing.BasicRayTracing", EAutomationTestFlags::ApplicationContextMask | EAutomationTestFlags::SmokeFilter)
-//IMPLEMENT_SIMPLE_AUTOMATION_TEST(FRayTracingTestbed, "System.Renderer.RayTracing.BasicRayTracing", EAutomationTestFlags::EditorContext | EAutomationTestFlags::HighPriority | EAutomationTestFlags::EngineFilter)
-
-
+IMPLEMENT_SIMPLE_AUTOMATION_TEST(FRayTracingTestbed, "System.Renderer.RayTracing.BasicRayTracing", EAutomationTestFlags::EditorContext | EAutomationTestFlags::HighPriority | EAutomationTestFlags::EngineFilter)
 
 #if RHI_RAYTRACING
 
@@ -224,6 +221,7 @@ IMPLEMENT_SHADER_TYPE(, FTestRaygenShader, TEXT("/Engine/Private/RayTracing/RayT
 bool FRayTracingTestbed::RunTest(const FString& Parameters)
 {
 	bool bTestPassed = false;
+	FlushRenderingCommands();
 
 	ENQUEUE_RENDER_COMMAND(FRayTracingTestbed)(
 		[&](FRHICommandListImmediate& RHICmdList)
