Fixed an error when attempting to run screen space denoising passes with r.Shaders.Optimize=0

The error is reported when shader compiler doesn't remove resources from the shader (so they remain in the reflection data), but high level code does not set them.
Solution in this case is to always fill all slots of FSSDSignalTextures with one of the valid textures that are already present there.

diff --git a/Engine/Source/Runtime/Renderer/Private/ScreenSpaceDenoise.cpp b/Engine/Source/Runtime/Renderer/Private/ScreenSpaceDenoise.cpp
--- a/Engine/Source/Runtime/Renderer/Private/ScreenSpaceDenoise.cpp
+++ b/Engine/Source/Runtime/Renderer/Private/ScreenSpaceDenoise.cpp
@@ -705,6 +705,34 @@ struct FSSDConstantPixelDensitySettings
 	TStaticArray<const FLightSceneInfo*, IScreenSpaceDenoiser::kMaxBatchSize> LightSceneInfo;
 };
 
+/**
+ * Returns a copy of SignalInput and fills all texture slots with the last valid texture, such that all elements of the array become valid.
+ * This is required because shader compiler may not always be able to remove unused input textures. In this case a mismatch
+ * between shader reflection data and parameters provided by high-level code will cause a validation error.
+ * For example, this happens when using non-optimized shaders (r.Shaders.Optimize=0).
+ */
+static FSSDSignalTextures CopyAndBackfillSignalInput(const FSSDSignalTextures& SignalInput, int32 NumExpectedValidTextures)
+{
+	checkf(NumExpectedValidTextures > 0, TEXT("FSSDSignalTextures must contain at least one valid texture"));
+	check(NumExpectedValidTextures <= kMaxBufferProcessingCount);
+
+	FSSDSignalTextures Result = SignalInput;
+
+	for (int32 i = 0; i < NumExpectedValidTextures; ++i)
+	{
+		checkf(Result.Textures[i], TEXT("Valid texture is required in slot %d of FSSDSignalTextures."), i);
+	}
+
+	for (int32 i = NumExpectedValidTextures; i < kMaxBufferProcessingCount; ++i)
+	{
+		if (!Result.Textures[i])
+		{
+			Result.Textures[i] = Result.Textures[i - 1];
+		}
+	}
+
+	return Result;
+}
 
 /** Denoises a signal at constant pixel density across the viewport. */
 static void DenoiseSignalAtConstantPixelDensity(
@@ -878,7 +906,7 @@ static void DenoiseSignalAtConstantPixelDensity(
 		FSSDInjestCS::FParameters* PassParameters = GraphBuilder.AllocParameters<FSSDInjestCS::FParameters>();
 		PassParameters->CommonParameters = CommonParameters;
 		PassParameters->ConvolutionMetaData = ConvolutionMetaData;
-		PassParameters->SignalInput = SignalHistory;
+		PassParameters->SignalInput = CopyAndBackfillSignalInput(SignalHistory, Settings.SignalBatchSize);
 		PassParameters->SignalOutput = CreateMultiplexedUAVs(GraphBuilder, NewSignalOutput);
 
 		FSSDInjestCS::FPermutationDomain PermutationVector;
@@ -910,7 +938,7 @@ static void DenoiseSignalAtConstantPixelDensity(
 		PassParameters->UpscaleFactor = int32(1.0f / Settings.InputResolutionFraction);
 		PassParameters->CommonParameters = CommonParameters;
 		PassParameters->ConvolutionMetaData = ConvolutionMetaData;
-		PassParameters->SignalInput = SignalHistory;
+		PassParameters->SignalInput = CopyAndBackfillSignalInput(SignalHistory, Settings.SignalBatchSize);
 		PassParameters->SignalOutput = CreateMultiplexedUAVs(GraphBuilder, NewSignalOutput);
 		
 		PassParameters->DebugOutput = GraphBuilder.CreateUAV(GraphBuilder.CreateTexture(DebugDesc, TEXT("SSDDebugReflectionReconstruction")));
@@ -949,7 +977,7 @@ static void DenoiseSignalAtConstantPixelDensity(
 		FSSDSpatialAccumulationCS::FParameters* PassParameters = GraphBuilder.AllocParameters<FSSDSpatialAccumulationCS::FParameters>();
 		PassParameters->CommonParameters = CommonParameters;
 		PassParameters->ConvolutionMetaData = ConvolutionMetaData;
-		PassParameters->SignalInput = SignalHistory;
+		PassParameters->SignalInput = CopyAndBackfillSignalInput(SignalHistory, Settings.SignalBatchSize);
 		PassParameters->SignalOutput = CreateMultiplexedUAVs(GraphBuilder, NewSignalOutput);
 
 		PassParameters->DebugOutput = GraphBuilder.CreateUAV(GraphBuilder.CreateTexture(DebugDesc, TEXT("DebugDenoiserPreConvolution")));
@@ -1019,7 +1047,7 @@ static void DenoiseSignalAtConstantPixelDensity(
 			FSSDSpatialAccumulationCS::FParameters* PassParameters = GraphBuilder.AllocParameters<FSSDSpatialAccumulationCS::FParameters>();
 			PassParameters->CommonParameters = CommonParameters;
 			PassParameters->ConvolutionMetaData = ConvolutionMetaData;
-			PassParameters->SignalInput = SignalHistory;
+			PassParameters->SignalInput = CopyAndBackfillSignalInput(SignalHistory, Settings.SignalBatchSize);
 			PassParameters->SignalOutput = CreateMultiplexedUAVs(GraphBuilder, RejectionPreConvolutionSignal);
 
 			FSSDSpatialAccumulationCS::FPermutationDomain PermutationVector;
@@ -1055,7 +1083,7 @@ static void DenoiseSignalAtConstantPixelDensity(
 		PassParameters->ConvolutionMetaData = ConvolutionMetaData;
 		PassParameters->HistoryPreExposureCorrection = View.PreExposure / View.PrevViewInfo.SceneColorPreExposure;
 
-		PassParameters->SignalInput = SignalHistory;
+		PassParameters->SignalInput = CopyAndBackfillSignalInput(SignalHistory, Settings.SignalBatchSize);
 		PassParameters->HistoryRejectionSignal = RejectionPreConvolutionSignal;
 		PassParameters->SignalHistoryOutput = CreateMultiplexedUAVs(GraphBuilder, SignalOutput);
 		
@@ -1090,6 +1118,8 @@ static void DenoiseSignalAtConstantPixelDensity(
 			PrevFrameHistory->SafeRelease();
 		} // for (uint32 BatchedSignalId = 0; BatchedSignalId < Settings.SignalBatchSize; BatchedSignalId++)
 
+		PassParameters->PrevHistory = CopyAndBackfillSignalInput(PassParameters->PrevHistory, Settings.SignalBatchSize);
+
 		PassParameters->DebugOutput = GraphBuilder.CreateUAV(GraphBuilder.CreateTexture(DebugDesc, TEXT("SSDDebugReflectionTemporalAccumulation")));
 
 		FComputeShaderUtils::AddPass(
@@ -1116,7 +1146,7 @@ static void DenoiseSignalAtConstantPixelDensity(
 		PassParameters->KernelSpreadFactor = Settings.HistoryConvolutionKernelSpreadFactor;
 		PassParameters->CommonParameters = CommonParameters;
 		PassParameters->ConvolutionMetaData = ConvolutionMetaData;
-		PassParameters->SignalInput = SignalHistory;
+		PassParameters->SignalInput = CopyAndBackfillSignalInput(SignalHistory, Settings.SignalBatchSize);
 		PassParameters->SignalOutput = CreateMultiplexedUAVs(GraphBuilder, SignalOutput);
 
 		FSSDSpatialAccumulationCS::FPermutationDomain PermutationVector;
@@ -1213,7 +1243,7 @@ static void DenoiseSignalAtConstantPixelDensity(
 
 		FSSDSpatialAccumulationCS::FParameters* PassParameters = GraphBuilder.AllocParameters<FSSDSpatialAccumulationCS::FParameters>();
 		PassParameters->CommonParameters = CommonParameters;
-		PassParameters->SignalInput = SignalHistory;
+		PassParameters->SignalInput = CopyAndBackfillSignalInput(SignalHistory, Settings.SignalBatchSize);
 		PassParameters->SignalOutput = CreateMultiplexedUAVs(GraphBuilder, *OutputSignal);
 
 		FSSDSpatialAccumulationCS::FPermutationDomain PermutationVector;
