Temporary workaround for validation error reported by CopyAndBackfillSignalInput() during shadow denoising. Assumption that number of textures = Settings.SignalBatchSize is incorrect.

diff --git a/Engine/Source/Runtime/Renderer/Private/ScreenSpaceDenoise.cpp b/Engine/Source/Runtime/Renderer/Private/ScreenSpaceDenoise.cpp
--- a/Engine/Source/Runtime/Renderer/Private/ScreenSpaceDenoise.cpp
+++ b/Engine/Source/Runtime/Renderer/Private/ScreenSpaceDenoise.cpp
@@ -711,7 +711,7 @@ struct FSSDConstantPixelDensitySettings
  * between shader reflection data and parameters provided by high-level code will cause a validation error.
  * For example, this happens when using non-optimized shaders (r.Shaders.Optimize=0).
  */
-static FSSDSignalTextures CopyAndBackfillSignalInput(const FSSDSignalTextures& SignalInput, int32 NumExpectedValidTextures)
+static FSSDSignalTextures CopyAndBackfillSignalInput(const FSSDSignalTextures& SignalInput, int32 NumExpectedValidTextures = 1)
 {
 	checkf(NumExpectedValidTextures > 0, TEXT("FSSDSignalTextures must contain at least one valid texture"));
 	check(NumExpectedValidTextures <= kMaxBufferProcessingCount);
@@ -906,7 +906,7 @@ static void DenoiseSignalAtConstantPixelDensity(
 		FSSDInjestCS::FParameters* PassParameters = GraphBuilder.AllocParameters<FSSDInjestCS::FParameters>();
 		PassParameters->CommonParameters = CommonParameters;
 		PassParameters->ConvolutionMetaData = ConvolutionMetaData;
-		PassParameters->SignalInput = CopyAndBackfillSignalInput(SignalHistory, Settings.SignalBatchSize);
+		PassParameters->SignalInput = CopyAndBackfillSignalInput(SignalHistory);
 		PassParameters->SignalOutput = CreateMultiplexedUAVs(GraphBuilder, NewSignalOutput);
 
 		FSSDInjestCS::FPermutationDomain PermutationVector;
@@ -938,7 +938,7 @@ static void DenoiseSignalAtConstantPixelDensity(
 		PassParameters->UpscaleFactor = int32(1.0f / Settings.InputResolutionFraction);
 		PassParameters->CommonParameters = CommonParameters;
 		PassParameters->ConvolutionMetaData = ConvolutionMetaData;
-		PassParameters->SignalInput = CopyAndBackfillSignalInput(SignalHistory, Settings.SignalBatchSize);
+		PassParameters->SignalInput = CopyAndBackfillSignalInput(SignalHistory);
 		PassParameters->SignalOutput = CreateMultiplexedUAVs(GraphBuilder, NewSignalOutput);
 		
 		PassParameters->DebugOutput = GraphBuilder.CreateUAV(GraphBuilder.CreateTexture(DebugDesc, TEXT("SSDDebugReflectionReconstruction")));
@@ -977,7 +977,7 @@ static void DenoiseSignalAtConstantPixelDensity(
 		FSSDSpatialAccumulationCS::FParameters* PassParameters = GraphBuilder.AllocParameters<FSSDSpatialAccumulationCS::FParameters>();
 		PassParameters->CommonParameters = CommonParameters;
 		PassParameters->ConvolutionMetaData = ConvolutionMetaData;
-		PassParameters->SignalInput = CopyAndBackfillSignalInput(SignalHistory, Settings.SignalBatchSize);
+		PassParameters->SignalInput = CopyAndBackfillSignalInput(SignalHistory);
 		PassParameters->SignalOutput = CreateMultiplexedUAVs(GraphBuilder, NewSignalOutput);
 
 		PassParameters->DebugOutput = GraphBuilder.CreateUAV(GraphBuilder.CreateTexture(DebugDesc, TEXT("DebugDenoiserPreConvolution")));
@@ -1047,7 +1047,7 @@ static void DenoiseSignalAtConstantPixelDensity(
 			FSSDSpatialAccumulationCS::FParameters* PassParameters = GraphBuilder.AllocParameters<FSSDSpatialAccumulationCS::FParameters>();
 			PassParameters->CommonParameters = CommonParameters;
 			PassParameters->ConvolutionMetaData = ConvolutionMetaData;
-			PassParameters->SignalInput = CopyAndBackfillSignalInput(SignalHistory, Settings.SignalBatchSize);
+			PassParameters->SignalInput = CopyAndBackfillSignalInput(SignalHistory);
 			PassParameters->SignalOutput = CreateMultiplexedUAVs(GraphBuilder, RejectionPreConvolutionSignal);
 
 			FSSDSpatialAccumulationCS::FPermutationDomain PermutationVector;
@@ -1083,7 +1083,7 @@ static void DenoiseSignalAtConstantPixelDensity(
 		PassParameters->ConvolutionMetaData = ConvolutionMetaData;
 		PassParameters->HistoryPreExposureCorrection = View.PreExposure / View.PrevViewInfo.SceneColorPreExposure;
 
-		PassParameters->SignalInput = CopyAndBackfillSignalInput(SignalHistory, Settings.SignalBatchSize);
+		PassParameters->SignalInput = CopyAndBackfillSignalInput(SignalHistory);
 		PassParameters->HistoryRejectionSignal = RejectionPreConvolutionSignal;
 		PassParameters->SignalHistoryOutput = CreateMultiplexedUAVs(GraphBuilder, SignalOutput);
 		
@@ -1118,7 +1118,7 @@ static void DenoiseSignalAtConstantPixelDensity(
 			PrevFrameHistory->SafeRelease();
 		} // for (uint32 BatchedSignalId = 0; BatchedSignalId < Settings.SignalBatchSize; BatchedSignalId++)
 
-		PassParameters->PrevHistory = CopyAndBackfillSignalInput(PassParameters->PrevHistory, Settings.SignalBatchSize);
+		PassParameters->PrevHistory = CopyAndBackfillSignalInput(PassParameters->PrevHistory);
 
 		PassParameters->DebugOutput = GraphBuilder.CreateUAV(GraphBuilder.CreateTexture(DebugDesc, TEXT("SSDDebugReflectionTemporalAccumulation")));
 
@@ -1146,7 +1146,7 @@ static void DenoiseSignalAtConstantPixelDensity(
 		PassParameters->KernelSpreadFactor = Settings.HistoryConvolutionKernelSpreadFactor;
 		PassParameters->CommonParameters = CommonParameters;
 		PassParameters->ConvolutionMetaData = ConvolutionMetaData;
-		PassParameters->SignalInput = CopyAndBackfillSignalInput(SignalHistory, Settings.SignalBatchSize);
+		PassParameters->SignalInput = CopyAndBackfillSignalInput(SignalHistory);
 		PassParameters->SignalOutput = CreateMultiplexedUAVs(GraphBuilder, SignalOutput);
 
 		FSSDSpatialAccumulationCS::FPermutationDomain PermutationVector;
@@ -1243,7 +1243,7 @@ static void DenoiseSignalAtConstantPixelDensity(
 
 		FSSDSpatialAccumulationCS::FParameters* PassParameters = GraphBuilder.AllocParameters<FSSDSpatialAccumulationCS::FParameters>();
 		PassParameters->CommonParameters = CommonParameters;
-		PassParameters->SignalInput = CopyAndBackfillSignalInput(SignalHistory, Settings.SignalBatchSize);
+		PassParameters->SignalInput = CopyAndBackfillSignalInput(SignalHistory);
 		PassParameters->SignalOutput = CreateMultiplexedUAVs(GraphBuilder, *OutputSignal);
 
 		FSSDSpatialAccumulationCS::FPermutationDomain PermutationVector;
