Removes validation on CopyAndBackfillSignalInput()

diff --git a/Engine/Source/Runtime/Renderer/Private/ScreenSpaceDenoise.cpp b/Engine/Source/Runtime/Renderer/Private/ScreenSpaceDenoise.cpp
--- a/Engine/Source/Runtime/Renderer/Private/ScreenSpaceDenoise.cpp
+++ b/Engine/Source/Runtime/Renderer/Private/ScreenSpaceDenoise.cpp
@@ -711,18 +711,10 @@ struct FSSDConstantPixelDensitySettings
  * between shader reflection data and parameters provided by high-level code will cause a validation error.
  * For example, this happens when using non-optimized shaders (r.Shaders.Optimize=0).
  */
-static FSSDSignalTextures CopyAndBackfillSignalInput(const FSSDSignalTextures& SignalInput, int32 NumExpectedValidTextures = 1)
+static FSSDSignalTextures CopyAndBackfillSignalInput(const FSSDSignalTextures& SignalInput)
 {
-	checkf(NumExpectedValidTextures > 0, TEXT("FSSDSignalTextures must contain at least one valid texture"));
-	check(NumExpectedValidTextures <= kMaxBufferProcessingCount);
-
 	FSSDSignalTextures Result = SignalInput;
 
-	for (int32 i = 0; i < NumExpectedValidTextures; ++i)
-	{
-		checkf(Result.Textures[i], TEXT("Valid texture is required in slot %d of FSSDSignalTextures."), i);
-	}
-
 	for (int32 i = NumExpectedValidTextures; i < kMaxBufferProcessingCount; ++i)
 	{
 		if (!Result.Textures[i])
