Adding EnableTwoSidedGeometry CVar to ray tracing global illumination pass.

r.RayTracing.GlobalIllumination.EnableTwoSidedGeometry (default = 1)

diff --git a/Engine/Shaders/Private/RayTracing/RayTracingGlobalIlluminationRGS.usf b/Engine/Shaders/Private/RayTracing/RayTracingGlobalIlluminationRGS.usf
--- a/Engine/Shaders/Private/RayTracing/RayTracingGlobalIlluminationRGS.usf
+++ b/Engine/Shaders/Private/RayTracing/RayTracingGlobalIlluminationRGS.usf
@@ -148,6 +148,9 @@ void GlobalIlluminationRGS()
 		{
 			// Cast ray
 			uint RayFlags = 0;
+#if !ENABLE_TWO_SIDED_GEOMETRY
+			RayFlags |= RAY_FLAG_CULL_BACK_FACING_TRIANGLES;
+#endif
 			FRayCone RayCone = (FRayCone)0;
 			FMaterialClosestHitPayload Payload = TraceRayInternal(
 				TLAS,   // AccelerationStructure
@@ -226,6 +229,9 @@ void GlobalIlluminationRGS()
 
 					// Trace visibility ray
 					uint RayFlags = RAY_FLAG_ACCEPT_FIRST_HIT_AND_END_SEARCH | RAY_FLAG_SKIP_CLOSEST_HIT_SHADER;
+#if !ENABLE_TWO_SIDED_GEOMETRY
+					RayFlags |= RAY_FLAG_CULL_BACK_FACING_TRIANGLES;
+#endif
 					FRayCone LightRayCone = (FRayCone)0;
 					FMaterialClosestHitPayload NeePayload = TraceRayInternal(
 						TLAS,   // AccelerationStructure
diff --git a/Engine/Source/Runtime/Renderer/Private/RayTracing/RayTracingGlobalIllumination.cpp b/Engine/Source/Runtime/Renderer/Private/RayTracing/RayTracingGlobalIllumination.cpp
--- a/Engine/Source/Runtime/Renderer/Private/RayTracing/RayTracingGlobalIllumination.cpp
+++ b/Engine/Source/Runtime/Renderer/Private/RayTracing/RayTracingGlobalIllumination.cpp
@@ -97,6 +97,13 @@ static FAutoConsoleVariableRef CVarRayTracingGlobalIlluminationScreenPercentage(
 	TEXT("Screen percentage for ray tracing global illumination (default = 100)")
 );
 
+static TAutoConsoleVariable<int32> CVarRayTracingGlobalIlluminationEnableTwoSidedGeometry(
+	TEXT("r.RayTracing.GlobalIllumination.EnableTwoSidedGeometry"),
+	1,
+	TEXT("Enables two-sided geometry when tracing GI rays (default = 1)"),
+	ECVF_RenderThreadSafe
+);
+
 static const int32 GLightCountMax = 64;
 
 DECLARE_GPU_STAT_NAMED(RayTracingGlobalIllumination, TEXT("Ray Tracing Global Illumination"));
@@ -200,8 +207,9 @@ class FGlobalIlluminationRGS : public FGlobalShader
 	SHADER_USE_ROOT_PARAMETER_STRUCT(FGlobalIlluminationRGS, FGlobalShader)
 
 	class FUseAttenuationTermDim : SHADER_PERMUTATION_BOOL("USE_ATTENUATION_TERM");
+	class FEnableTwoSidedGeometryDim : SHADER_PERMUTATION_BOOL("ENABLE_TWO_SIDED_GEOMETRY");
 
-	using FPermutationDomain = TShaderPermutationDomain<FUseAttenuationTermDim>;
+	using FPermutationDomain = TShaderPermutationDomain<FUseAttenuationTermDim, FEnableTwoSidedGeometryDim>;
 
 	static bool ShouldCompilePermutation(const FGlobalShaderPermutationParameters& Parameters)
 	{
@@ -308,11 +316,17 @@ IMPLEMENT_GLOBAL_SHADER(FRayTracingGlobalIlluminationSceneColorCompositePS, "/En
 void FDeferredShadingSceneRenderer::PrepareRayTracingGlobalIllumination(const FViewInfo& View, TArray<FRayTracingShaderRHIParamRef>& OutRayGenShaders)
 {
 	// Declare all RayGen shaders that require material closest hit shaders to be bound
-
-	FGlobalIlluminationRGS::FPermutationDomain PermutationVector;
-	PermutationVector.Set<FGlobalIlluminationRGS::FUseAttenuationTermDim>(true);
-	auto RayGenerationShader = View.ShaderMap->GetShader<FGlobalIlluminationRGS>();
-	OutRayGenShaders.Add(RayGenerationShader->GetRayTracingShader());
+	for (int UseAttenuationTerm = 0; UseAttenuationTerm < 2; ++UseAttenuationTerm)
+	{
+		for (int EnableTwoSidedGeometry = 0; EnableTwoSidedGeometry < 2; ++EnableTwoSidedGeometry)
+		{
+			FGlobalIlluminationRGS::FPermutationDomain PermutationVector;
+			PermutationVector.Set<FGlobalIlluminationRGS::FUseAttenuationTermDim>(UseAttenuationTerm == 1);
+			PermutationVector.Set<FGlobalIlluminationRGS::FEnableTwoSidedGeometryDim>(EnableTwoSidedGeometry == 1);
+			TShaderMapRef<FGlobalIlluminationRGS> RayGenerationShader(View.ShaderMap, PermutationVector);
+			OutRayGenShaders.Add(RayGenerationShader->GetRayTracingShader());
+		}
+	}
 
 }
 
@@ -487,8 +501,9 @@ void FDeferredShadingSceneRenderer::RenderRayTracingGlobalIllumination(
 
 		FGlobalIlluminationRGS::FPermutationDomain PermutationVector;
 		PermutationVector.Set<FGlobalIlluminationRGS::FUseAttenuationTermDim>(true);
-		auto RayGenerationShader = View.ShaderMap->GetShader<FGlobalIlluminationRGS>();
-		ClearUnusedGraphResources(RayGenerationShader, PassParameters);
+		PermutationVector.Set<FGlobalIlluminationRGS::FEnableTwoSidedGeometryDim>(CVarRayTracingGlobalIlluminationEnableTwoSidedGeometry.GetValueOnRenderThread() != 0);
+		TShaderMapRef<FGlobalIlluminationRGS> RayGenerationShader(GetGlobalShaderMap(FeatureLevel), PermutationVector);
+		ClearUnusedGraphResources(*RayGenerationShader, PassParameters);
 
 		FIntPoint RayTracingResolution = FIntPoint::DivideAndRoundUp(View.ViewRect.Size(), UpscaleFactor);
 		GraphBuilder.AddPass(
@@ -498,7 +513,7 @@ void FDeferredShadingSceneRenderer::RenderRayTracingGlobalIllumination(
 			[PassParameters, this, &View, RayGenerationShader, RayTracingResolution](FRHICommandList& RHICmdList)
 		{
 			FRayTracingShaderBindingsWriter GlobalResources;
-			SetShaderParameters(GlobalResources, RayGenerationShader, *PassParameters);
+			SetShaderParameters(GlobalResources, *RayGenerationShader, *PassParameters);
 
 			FRayTracingSceneRHIParamRef RayTracingSceneRHI = View.RayTracingScene.RayTracingSceneRHI;
 			RHICmdList.RayTraceDispatch(View.RayTracingMaterialPipeline, RayGenerationShader->GetRayTracingShader(), RayTracingSceneRHI, GlobalResources, RayTracingResolution.X, RayTracingResolution.Y);
