DR - Partially revert OcclusionQuery Pools until I can figure out what is going on.
Undo //UE4/Dev-Rendering/Engine/Source/Runtime/... changelist 6259847

diff --git a/Engine/Source/Runtime/Engine/Private/SceneUtils.cpp b/Engine/Source/Runtime/Engine/Private/SceneUtils.cpp
--- a/Engine/Source/Runtime/Engine/Private/SceneUtils.cpp
+++ b/Engine/Source/Runtime/Engine/Private/SceneUtils.cpp
@@ -70,4 +70,49 @@ ENGINE_API bool IsMobileColorsRGB()
 	const bool bMobileUseHWsRGBEncoding = (MobileUseHWsRGBEncodingCVAR && MobileUseHWsRGBEncodingCVAR->GetValueOnAnyThread() == 1);
 
 	return !IsMobileHDR() && bMobileUseHWsRGBEncoding;
+}
+
+/*-----------------------------------------------------------------------------
+FRenderQueryPool
+-----------------------------------------------------------------------------*/
+FRenderQueryPool::~FRenderQueryPool()
+{
+	Release();
+}
+
+void FRenderQueryPool::Release()
+{
+	Queries.Empty();
+	NumQueriesAllocated = 0;
+}
+
+FRenderQueryRHIRef FRenderQueryPool::AllocateQuery()
+{
+	// Are we out of available render queries?
+	NumQueriesAllocated++;
+	if (Queries.Num() == 0)
+	{
+		// Create a new render query.
+		return RHICreateRenderQuery(QueryType);
+	}
+
+	return Queries.Pop(/*bAllowShrinking=*/ false);
+}
+
+void FRenderQueryPool::ReleaseQuery(FRenderQueryRHIRef &Query)
+{
+	if (IsValidRef(Query))
+	{
+		NumQueriesAllocated--;
+
+		// Is no one else keeping a refcount to the query?
+		if (Query.GetRefCount() == 1)
+		{
+			// Return it to the pool.
+			Queries.Add(Query);
+		}
+
+		// De-ref without deleting.
+		Query = NULL;
+	}
 }
\ No newline at end of file
diff --git a/Engine/Source/Runtime/Engine/Public/SceneUtils.h b/Engine/Source/Runtime/Engine/Public/SceneUtils.h
--- a/Engine/Source/Runtime/Engine/Public/SceneUtils.h
+++ b/Engine/Source/Runtime/Engine/Public/SceneUtils.h
@@ -42,6 +42,40 @@ ENGINE_API EMobileHDRMode GetMobileHDRMode();
 
 ENGINE_API bool IsMobileColorsRGB();
 
+/**
+* A pool of render (e.g. occlusion/timer) queries which are allocated individually, and returned to the pool as a group.
+*/
+class ENGINE_API FRenderQueryPool
+{
+public:
+	FRenderQueryPool(ERenderQueryType InQueryType)
+		: QueryType(InQueryType)
+		, NumQueriesAllocated(0)
+	{ }
+
+	virtual ~FRenderQueryPool();
+
+	/** Releases all the render queries in the pool. */
+	void Release();
+
+	/** Allocates an render query from the pool. */
+	FRenderQueryRHIRef AllocateQuery();
+
+	/** De-reference an render query, returning it to the pool instead of deleting it when the refcount reaches 0. */
+	void ReleaseQuery(FRenderQueryRHIRef &Query);
+
+	/** Returns the number of currently allocated queries. This is not necessarily the same as the pool size */
+	int32 GetAllocatedQueryCount() const { return NumQueriesAllocated;  }
+
+private:
+	/** Container for available render queries. */
+	TArray<FRenderQueryRHIRef> Queries;
+
+	ERenderQueryType QueryType;
+
+	int32 NumQueriesAllocated;
+};
+
 // Callback for calling one action (typical use case: delay a clear until it's actually needed)
 class FDelayedRendererAction
 {
diff --git a/Engine/Source/Runtime/RHI/Public/RHICommandList.h b/Engine/Source/Runtime/RHI/Public/RHICommandList.h
--- a/Engine/Source/Runtime/RHI/Public/RHICommandList.h
+++ b/Engine/Source/Runtime/RHI/Public/RHICommandList.h
@@ -4062,14 +4062,12 @@ public:
 		GDynamicRHI->RHIRead3DSurfaceFloatData(Texture,Rect,ZMinMax,OutData);
 	}
 	
-	UE_DEPRECATED(4.23, "CreateRenderQuery API is deprecated; use RHICreateRenderQueryPool and suballocate queries there")
 	FORCEINLINE FRenderQueryRHIRef CreateRenderQuery(ERenderQueryType QueryType)
 	{
 		FScopedRHIThreadStaller StallRHIThread(*this);
 		return GDynamicRHI->RHICreateRenderQuery(QueryType);
 	}
 
-	UE_DEPRECATED(4.23, "CreateRenderQuery API is deprecated; use RHICreateRenderQueryPool and suballocate queries there")
 	FORCEINLINE FRenderQueryRHIRef CreateRenderQuery_RenderThread(ERenderQueryType QueryType)
 	{
 		return GDynamicRHI->RHICreateRenderQuery_RenderThread(*this, QueryType);
@@ -4807,12 +4805,9 @@ FORCEINLINE void RHIUnlockTextureCubeFace(FTextureCubeRHIParamRef Texture, uint3
 	 FRHICommandListExecutor::GetImmediateCommandList().UnlockTextureCubeFace(Texture, FaceIndex, ArrayIndex, MipIndex, bLockWithinMiptail);
 }
 
-UE_DEPRECATED(4.23, "CreateRenderQuery API is deprecated; use RHICreateRenderQueryPool and suballocate queries there")
 FORCEINLINE FRenderQueryRHIRef RHICreateRenderQuery(ERenderQueryType QueryType)
 {
-	PRAGMA_DISABLE_DEPRECATION_WARNINGS
 	return FRHICommandListExecutor::GetImmediateCommandList().CreateRenderQuery_RenderThread(QueryType);
-	PRAGMA_ENABLE_DEPRECATION_WARNINGS
 }
 
 FORCEINLINE void RHIAcquireTransientResource(FTextureRHIParamRef Resource)
diff --git a/Engine/Source/Runtime/Renderer/Private/RendererScene.cpp b/Engine/Source/Runtime/Renderer/Private/RendererScene.cpp
--- a/Engine/Source/Runtime/Renderer/Private/RendererScene.cpp
+++ b/Engine/Source/Runtime/Renderer/Private/RendererScene.cpp
@@ -126,7 +126,7 @@ struct FSpeedTreeWindComputation
 
 /** Default constructor. */
 FSceneViewState::FSceneViewState()
-	: OcclusionQueryPool(RHICreateRenderQueryPool(RQT_Occlusion))
+	: OcclusionQueryPool(RQT_Occlusion)
 	, TimerQueryPool(RHICreateRenderQueryPool(RQT_AbsoluteTime, FLatentGPUTimer::NumBufferedFrames * 2 * 2 * 2))
 	, TranslucencyTimer(TimerQueryPool)
 	, SeparateTranslucencyTimer(TimerQueryPool)
diff --git a/Engine/Source/Runtime/Renderer/Private/SceneOcclusion.cpp b/Engine/Source/Runtime/Renderer/Private/SceneOcclusion.cpp
--- a/Engine/Source/Runtime/Renderer/Private/SceneOcclusion.cpp
+++ b/Engine/Source/Runtime/Renderer/Private/SceneOcclusion.cpp
@@ -243,14 +243,14 @@ bool FSceneViewState::IsShadowOccluded(FRHICommandListImmediate& RHICmdList, FSc
 	// Get the oldest occlusion query	
 	const uint32 QueryIndex = FOcclusionQueryHelpers::GetQueryLookupIndex(PendingPrevFrameNumber, NumBufferedFrames);
 	const FSceneViewState::ShadowKeyOcclusionQueryMap& ShadowOcclusionQueryMap = ShadowOcclusionQueryMaps[QueryIndex];	
-	const FRHIPooledRenderQuery* Query = ShadowOcclusionQueryMap.Find(ShadowKey);
+	const FRenderQueryRHIRef* Query = ShadowOcclusionQueryMap.Find(ShadowKey);
 
 	// Read the occlusion query results.
 	uint64 NumSamples = 0;
 	// Only block on the query if not running SLI
 	const bool bWaitOnQuery = GNumAlternateFrameRenderingGroups == 1;
 
-	if (Query && RHICmdList.GetRenderQueryResult(Query->GetQuery(), NumSamples, bWaitOnQuery))
+	if (Query && RHICmdList.GetRenderQueryResult(*Query, NumSamples, bWaitOnQuery))
 	{
 		// If the shadow's occlusion query didn't have any pixels visible the previous frame, it's occluded.
 		return NumSamples == 0;
@@ -337,7 +337,7 @@ FOcclusionQueryBatcher::FOcclusionQueryBatcher(class FSceneViewState* ViewState,
 :	CurrentBatchOcclusionQuery(NULL)
 ,	MaxBatchedPrimitives(InMaxBatchedPrimitives)
 ,	NumBatchedPrimitives(0)
-,	OcclusionQueryPool(ViewState ? ViewState->OcclusionQueryPool : NULL)
+,	OcclusionQueryPool(ViewState ? &ViewState->OcclusionQueryPool : NULL)
 {}
 
 FOcclusionQueryBatcher::~FOcclusionQueryBatcher()
@@ -358,7 +358,7 @@ void FOcclusionQueryBatcher::Flush(FRHICommandList& RHICmdList)
 		for(int32 BatchIndex = 0, NumBatches = BatchOcclusionQueries.Num();BatchIndex < NumBatches;BatchIndex++)
 		{
 			FOcclusionBatch& Batch = BatchOcclusionQueries[BatchIndex];
-			FRHIRenderQuery* BatchOcclusionQuery = Batch.Query;
+			FRenderQueryRHIParamRef BatchOcclusionQuery = Batch.Query;
 			FVertexBufferRHIParamRef VertexBufferRHI = Batch.VertexAllocation.VertexBuffer->VertexBufferRHI;
 			uint32 VertexBufferOffset = Batch.VertexAllocation.VertexOffset;
 			const int32 NumPrimitivesThisBatch = (BatchIndex != (NumBatches-1)) ? MaxBatchedPrimitives : NumBatchedPrimitives;
@@ -384,17 +384,14 @@ void FOcclusionQueryBatcher::Flush(FRHICommandList& RHICmdList)
 	}
 }
 
-FRHIPooledRenderQuery FOcclusionQueryBatcher::BatchPrimitive(const FVector& BoundsOrigin,const FVector& BoundsBoxExtent, FGlobalDynamicVertexBuffer& DynamicVertexBuffer)
+FRenderQueryRHIParamRef FOcclusionQueryBatcher::BatchPrimitive(const FVector& BoundsOrigin,const FVector& BoundsBoxExtent, FGlobalDynamicVertexBuffer& DynamicVertexBuffer)
 {
-	FRHIPooledRenderQuery PooledQuery;
-
 	// Check if the current batch is full.
 	if(CurrentBatchOcclusionQuery == NULL || NumBatchedPrimitives >= MaxBatchedPrimitives)
 	{
 		check(OcclusionQueryPool);
 		CurrentBatchOcclusionQuery = new(BatchOcclusionQueries) FOcclusionBatch;
-		PooledQuery = OcclusionQueryPool->AllocateQuery();
-		CurrentBatchOcclusionQuery->Query = PooledQuery.GetQueryRef();
+		CurrentBatchOcclusionQuery->Query = OcclusionQueryPool->AllocateQuery();
 		CurrentBatchOcclusionQuery->VertexAllocation = DynamicVertexBuffer.Allocate(MaxBatchedPrimitives * 8 * sizeof(FVector));
 		check(CurrentBatchOcclusionQuery->VertexAllocation.IsValid());
 		NumBatchedPrimitives = 0;
@@ -418,7 +415,7 @@ FRHIPooledRenderQuery FOcclusionQueryBatcher::BatchPrimitive(const FVector& Boun
 	CurrentBatchOcclusionQuery->VertexAllocation.Buffer = (uint8*)Vertices;
 	NumBatchedPrimitives++;
 
-	return PooledQuery;
+	return CurrentBatchOcclusionQuery->Query;
 }
 
 enum EShadowOcclusionQueryIntersectionMode
@@ -433,7 +430,7 @@ static bool AllocateProjectedShadowOcclusionQuery(
 	const FProjectedShadowInfo& ProjectedShadowInfo, 
 	int32 NumBufferedFrames, 
 	EShadowOcclusionQueryIntersectionMode IntersectionMode,
-	FRHIRenderQuery*& ShadowOcclusionQuery)
+	FRenderQueryRHIRef& ShadowOcclusionQuery)
 {
 	bool bIssueQuery = true;
 
@@ -468,21 +465,21 @@ static bool AllocateProjectedShadowOcclusionQuery(
 		FSceneViewState* ViewState = (FSceneViewState*)View.State;
 
 		// Allocate an occlusion query for the primitive from the occlusion query pool.
+		ShadowOcclusionQuery = ViewState->OcclusionQueryPool.AllocateQuery();
+
 		FSceneViewState::FProjectedShadowKey Key(ProjectedShadowInfo);
 		const uint32 QueryIndex = FOcclusionQueryHelpers::GetQueryIssueIndex(ViewState->PendingPrevFrameNumber, NumBufferedFrames);
 		FSceneViewState::ShadowKeyOcclusionQueryMap& ShadowOcclusionQueryMap = ViewState->ShadowOcclusionQueryMaps[QueryIndex];
 
 		checkSlow(ShadowOcclusionQueryMap.Find(Key) == NULL);
-		FRHIPooledRenderQuery PooledShadowOcclusionQuery = ViewState->OcclusionQueryPool->AllocateQuery();
-		ShadowOcclusionQuery = PooledShadowOcclusionQuery.GetQuery();
-		ShadowOcclusionQueryMap.Add(Key, MoveTemp(PooledShadowOcclusionQuery));
+		ShadowOcclusionQueryMap.Add(Key, ShadowOcclusionQuery);
 	}
 	
 	return bIssueQuery;
 }
 
 
-static void ExecutePointLightShadowOcclusionQuery(FRHICommandList& RHICmdList, FViewInfo& View, const FProjectedShadowInfo& ProjectedShadowInfo, FOcclusionQueryVS* VertexShader, FRHIRenderQuery* ShadowOcclusionQuery)
+static void ExecutePointLightShadowOcclusionQuery(FRHICommandList& RHICmdList, FViewInfo& View, const FProjectedShadowInfo& ProjectedShadowInfo, FOcclusionQueryVS* VertexShader, FRenderQueryRHIRef ShadowOcclusionQuery)
 {
 	FLightSceneProxy& LightProxy = *(ProjectedShadowInfo.GetLightSceneInfo().Proxy);
 	
@@ -531,7 +528,7 @@ static void PrepareDirectionalLightShadowOcclusionQuery(uint32& BaseVertexIndex,
 	BaseVertexIndex += 6;
 }
 
-static void ExecuteDirectionalLightShadowOcclusionQuery(FRHICommandList& RHICmdList, uint32& BaseVertexIndex, FRHIRenderQuery* ShadowOcclusionQuery)
+static void ExecuteDirectionalLightShadowOcclusionQuery(FRHICommandList& RHICmdList, uint32& BaseVertexIndex, FRenderQueryRHIRef ShadowOcclusionQuery)
 {
 	RHICmdList.BeginRenderQuery(ShadowOcclusionQuery);
 
@@ -571,7 +568,7 @@ static void PrepareProjectedShadowOcclusionQuery(uint32& BaseVertexIndex, FVecto
 	BaseVertexIndex += 8;
 }
 
-static void ExecuteProjectedShadowOcclusionQuery(FRHICommandList& RHICmdList, uint32& BaseVertexIndex, FRHIRenderQuery* ShadowOcclusionQuery)
+static void ExecuteProjectedShadowOcclusionQuery(FRHICommandList& RHICmdList, uint32& BaseVertexIndex, FRenderQueryRHIRef ShadowOcclusionQuery)
 {	
 	// Draw the primitive's bounding box, using the occlusion query.
 	RHICmdList.BeginRenderQuery(ShadowOcclusionQuery);
@@ -582,7 +579,7 @@ static void ExecuteProjectedShadowOcclusionQuery(FRHICommandList& RHICmdList, ui
 	RHICmdList.EndRenderQuery(ShadowOcclusionQuery);
 }
 
-static bool AllocatePlanarReflectionOcclusionQuery(const FViewInfo& View, const FPlanarReflectionSceneProxy* SceneProxy, int32 NumBufferedFrames, FRHIRenderQuery*& OcclusionQuery)
+static bool AllocatePlanarReflectionOcclusionQuery(const FViewInfo& View, const FPlanarReflectionSceneProxy* SceneProxy, int32 NumBufferedFrames, FRenderQueryRHIRef& OcclusionQuery)
 {
 	FSceneViewState* ViewState = (FSceneViewState*)View.State;
 	
@@ -617,14 +614,13 @@ static bool AllocatePlanarReflectionOcclusionQuery(const FViewInfo& View, const
 	if (bAllowBoundsTest)
 	{
 		// Allocate an occlusion query for the primitive from the occlusion query pool.
-		FRHIPooledRenderQuery PooledOcclusionQuery = ViewState->OcclusionQueryPool->AllocateQuery();
-		OcclusionQuery = PooledOcclusionQuery.GetQuery();
-
-		OcclusionHistory.SetCurrentQuery(OcclusionFrameCounter, MoveTemp(PooledOcclusionQuery), NumBufferedFrames);
+		OcclusionQuery = ViewState->OcclusionQueryPool.AllocateQuery();
+		
+		OcclusionHistory.SetCurrentQuery(OcclusionFrameCounter, OcclusionQuery, NumBufferedFrames);
 	}
 	else
 	{
-		OcclusionHistory.SetCurrentQuery(OcclusionFrameCounter, FRHIPooledRenderQuery(), NumBufferedFrames);
+		OcclusionHistory.SetCurrentQuery(OcclusionFrameCounter, NULL, NumBufferedFrames);
 	}
 	
 	return bAllowBoundsTest;
@@ -648,7 +644,7 @@ static void PreparePlanarReflectionOcclusionQuery(uint32& BaseVertexIndex, FVect
 	BaseVertexIndex += 8;
 }
 
-static void ExecutePlanarReflectionOcclusionQuery(FRHICommandList& RHICmdList, uint32& BaseVertexIndex, FRHIRenderQuery* OcclusionQuery)
+static void ExecutePlanarReflectionOcclusionQuery(FRHICommandList& RHICmdList, uint32& BaseVertexIndex, FRenderQueryRHIRef OcclusionQuery)
 {
 	// Draw the primitive's bounding box, using the occlusion query.
 	RHICmdList.BeginRenderQuery(OcclusionQuery);
@@ -1264,10 +1260,10 @@ struct FViewOcclusionQueries
 	TArray<FProjectedShadowInfo const*> ShadowQuerieInfos;
 	TArray<FPlanarReflectionSceneProxy const*> ReflectionQuerieInfos;
 
-	TArray<FRHIRenderQuery*> PointLightQueries;
-	TArray<FRHIRenderQuery*> CSMQueries;
-	TArray<FRHIRenderQuery*> ShadowQueries;
-	TArray<FRHIRenderQuery*> ReflectionQueries;
+	TArray<FRenderQueryRHIRef> PointLightQueries;
+	TArray<FRenderQueryRHIRef> CSMQueries;
+	TArray<FRenderQueryRHIRef> ShadowQueries;
+	TArray<FRenderQueryRHIRef> ReflectionQueries;
 };
 
 void FSceneRenderer::BeginOcclusionTests(FRHICommandListImmediate& RHICmdList, bool bRenderQueries)
@@ -1303,6 +1299,12 @@ void FSceneRenderer::BeginOcclusionTests(FRHICommandListImmediate& RHICmdList, b
 				ViewState->TrimOcclusionHistory(ViewFamily.CurrentRealTime, ViewFamily.CurrentRealTime - GEngine->PrimitiveProbablyVisibleTime, ViewFamily.CurrentRealTime, ViewState->OcclusionFrameCounter);
 
 				// Give back all these occlusion queries to the pool.
+				for (TMap<FSceneViewState::FProjectedShadowKey, FRenderQueryRHIRef>::TIterator QueryIt(ShadowOcclusionQueryMap); QueryIt; ++QueryIt)
+				{
+					//FRenderQueryRHIParamRef Query = QueryIt.Value();
+					//check( Query.GetRefCount() == 1 );
+					ViewState->OcclusionQueryPool.ReleaseQuery(QueryIt.Value());
+				}
 				ShadowOcclusionQueryMap.Reset();
 				
 				if (FeatureLevel > ERHIFeatureLevel::ES3_1)
@@ -1330,7 +1332,7 @@ void FSceneRenderer::BeginOcclusionTests(FRHICommandListImmediate& RHICmdList, b
 
 							if (ProjectedShadowInfo.bOnePassPointLightShadow)
 							{
-								FRHIRenderQuery* ShadowOcclusionQuery;
+								FRenderQueryRHIRef ShadowOcclusionQuery;
 								if (AllocateProjectedShadowOcclusionQuery(View, ProjectedShadowInfo, NumBufferedFrames, SOQ_LightInfluenceSphere, ShadowOcclusionQuery))
 								{
 									ViewQuery.PointLightQuerieInfos.Add(&ProjectedShadowInfo);
@@ -1344,7 +1346,7 @@ void FSceneRenderer::BeginOcclusionTests(FRHICommandListImmediate& RHICmdList, b
 								// Don't query the first cascade, it is always visible
 								if (GOcclusionCullCascadedShadowMaps && ProjectedShadowInfo.CascadeSettings.ShadowSplitIndex > 0)
 								{
-									FRHIRenderQuery* ShadowOcclusionQuery;
+									FRenderQueryRHIRef ShadowOcclusionQuery;
 									if (AllocateProjectedShadowOcclusionQuery(View, ProjectedShadowInfo, NumBufferedFrames, SOQ_None, ShadowOcclusionQuery))
 									{
 										ViewQuery.CSMQuerieInfos.Add(&ProjectedShadowInfo);
@@ -1360,7 +1362,7 @@ void FSceneRenderer::BeginOcclusionTests(FRHICommandListImmediate& RHICmdList, b
 							// Don't query if any subjects are visible because the shadow frustum will be definitely unoccluded
 								&& !ProjectedShadowInfo.SubjectsVisible(View))
 							{
-								FRHIRenderQuery* ShadowOcclusionQuery;
+								FRenderQueryRHIRef ShadowOcclusionQuery;
 								if (AllocateProjectedShadowOcclusionQuery(View, ProjectedShadowInfo, NumBufferedFrames, SOQ_NearPlaneVsShadowFrustum, ShadowOcclusionQuery))
 								{
 									ViewQuery.ShadowQuerieInfos.Add(&ProjectedShadowInfo);
@@ -1375,7 +1377,7 @@ void FSceneRenderer::BeginOcclusionTests(FRHICommandListImmediate& RHICmdList, b
 						for (int32 ShadowIndex = 0; ShadowIndex < VisibleLightInfo.OccludedPerObjectShadows.Num(); ShadowIndex++)
 						{
 							const FProjectedShadowInfo& ProjectedShadowInfo = *VisibleLightInfo.OccludedPerObjectShadows[ShadowIndex];
-							FRHIRenderQuery* ShadowOcclusionQuery;
+							FRenderQueryRHIRef ShadowOcclusionQuery;
 							if (AllocateProjectedShadowOcclusionQuery(View, ProjectedShadowInfo, NumBufferedFrames, SOQ_NearPlaneVsShadowFrustum, ShadowOcclusionQuery))
 							{
 								ViewQuery.ShadowQuerieInfos.Add(&ProjectedShadowInfo);
@@ -1398,7 +1400,7 @@ void FSceneRenderer::BeginOcclusionTests(FRHICommandListImmediate& RHICmdList, b
 					for (int32 ReflectionIndex = 0; ReflectionIndex < Scene->PlanarReflections.Num(); ReflectionIndex++)
 					{
 						FPlanarReflectionSceneProxy* SceneProxy = Scene->PlanarReflections[ReflectionIndex];
-						FRHIRenderQuery* ShadowOcclusionQuery;
+						FRenderQueryRHIRef ShadowOcclusionQuery;
 						if (AllocatePlanarReflectionOcclusionQuery(View, SceneProxy, NumReflectionBufferedFrames, ShadowOcclusionQuery))
 						{
 							ViewQuery.ReflectionQuerieInfos.Add(SceneProxy);
@@ -1556,13 +1558,13 @@ void FSceneRenderer::BeginOcclusionTests(FRHICommandListImmediate& RHICmdList, b
 						RHICmdList.SetStreamSource(0, VertexBufferRHI, 0);
 						BaseVertexOffset = 0;
 
-						for (FRHIRenderQuery* const& Query : ViewQuery.CSMQueries)
+						for (FRenderQueryRHIRef const& Query : ViewQuery.CSMQueries)
 						{
 							ExecuteDirectionalLightShadowOcclusionQuery(RHICmdList, BaseVertexOffset, Query);
 							checkSlow(BaseVertexOffset <= NumVertices);
 						}
 
-						for (FRHIRenderQuery* const& Query : ViewQuery.ShadowQueries)
+						for (FRenderQueryRHIRef const& Query : ViewQuery.ShadowQueries)
 						{
 							ExecuteProjectedShadowOcclusionQuery(RHICmdList, BaseVertexOffset, Query);
 							checkSlow(BaseVertexOffset <= NumVertices);
@@ -1572,7 +1574,7 @@ void FSceneRenderer::BeginOcclusionTests(FRHICommandListImmediate& RHICmdList, b
 					if (FeatureLevel > ERHIFeatureLevel::ES3_1)					
 					{
 						SCOPED_DRAW_EVENT(RHICmdList, PlanarReflectionQueries);
-						for (FRHIRenderQuery* const& Query : ViewQuery.ReflectionQueries)
+						for (FRenderQueryRHIRef const& Query : ViewQuery.ReflectionQueries)
 						{
 							ExecutePlanarReflectionOcclusionQuery(RHICmdList, BaseVertexOffset, Query);
 							check(BaseVertexOffset <= NumVertices);
diff --git a/Engine/Source/Runtime/Renderer/Private/ScenePrivate.h b/Engine/Source/Runtime/Renderer/Private/ScenePrivate.h
--- a/Engine/Source/Runtime/Renderer/Private/ScenePrivate.h
+++ b/Engine/Source/Runtime/Renderer/Private/ScenePrivate.h
@@ -86,7 +86,7 @@ public:
 	FPrimitiveComponentId PrimitiveId;
 
 	/** The occlusion query which contains the primitive's pending occlusion results. */
-	FRHIPooledRenderQuery PendingOcclusionQuery[FOcclusionQueryHelpers::MaxBufferedOcclusionFrames];
+	FRenderQueryRHIRef PendingOcclusionQuery[FOcclusionQueryHelpers::MaxBufferedOcclusionFrames];
 	uint32 PendingOcclusionQueryFrames[FOcclusionQueryHelpers::MaxBufferedOcclusionFrames]; 
 
 	uint32 LastTestFrameNumber;
@@ -141,7 +141,7 @@ private:
 		{
 			const uint32 ThisFrameNumber = PendingOcclusionQueryFrames[Index];
 			const int32 LaggedFrames = FrameNumber - ThisFrameNumber;
-			if (PendingOcclusionQuery[Index].GetQuery() && LaggedFrames <= LagTolerance && ThisFrameNumber < OldestFrame)
+			if (PendingOcclusionQuery[Index].IsValid() && LaggedFrames <= LagTolerance && ThisFrameNumber < OldestFrame)
 			{
 				OldestFrame = ThisFrameNumber;
 				OldestQueryIndex = Index;
@@ -202,9 +202,9 @@ public:
 			{
 				uint32 TestFrame = FrameNumber - (DeltaFrame - 1);
 				const uint32 QueryIndex = FOcclusionQueryHelpers::GetQueryLookupIndex(TestFrame, NumBufferedFrames);
-				if (PendingOcclusionQuery[QueryIndex].GetQuery() && PendingOcclusionQueryFrames[QueryIndex] != TestFrame)
+				if (PendingOcclusionQuery[QueryIndex].GetReference() && PendingOcclusionQueryFrames[QueryIndex] != TestFrame)
 				{
-					PendingOcclusionQuery[QueryIndex].ReleaseQuery();
+					Pool.ReleaseQuery(PendingOcclusionQuery[QueryIndex]);
 				}
 			}
 		}
@@ -213,39 +213,42 @@ public:
 	FORCEINLINE void ReleaseQuery(TOcclusionQueryPool& Pool, uint32 FrameNumber, int32 NumBufferedFrames)
 	{
 		const uint32 QueryIndex = FOcclusionQueryHelpers::GetQueryLookupIndex(FrameNumber, NumBufferedFrames);
-		PendingOcclusionQuery[QueryIndex].ReleaseQuery();
+		if (PendingOcclusionQuery[QueryIndex].IsValid())
+		{
+			Pool.ReleaseQuery(PendingOcclusionQuery[QueryIndex]);
+		}
 	}
 
-	FORCEINLINE FRHIRenderQuery* GetQueryForEviction(uint32 FrameNumber, int32 NumBufferedFrames) const
+	FORCEINLINE FRenderQueryRHIParamRef GetQueryForEviction(uint32 FrameNumber, int32 NumBufferedFrames) const
 	{
 		const uint32 QueryIndex = FOcclusionQueryHelpers::GetQueryLookupIndex(FrameNumber, NumBufferedFrames);
-		if (PendingOcclusionQuery[QueryIndex].GetQuery())
+		if (PendingOcclusionQuery[QueryIndex].IsValid())
 		{
-			return PendingOcclusionQuery[QueryIndex].GetQuery();
+			return PendingOcclusionQuery[QueryIndex].GetReference();
 		}
 		return nullptr;
 	}
 
 
-	FORCEINLINE FRHIRenderQuery* GetQueryForReading(uint32 FrameNumber, int32 NumBufferedFrames, int32 LagTolerance, bool& bOutGrouped) const
+	FORCEINLINE FRenderQueryRHIParamRef GetQueryForReading(uint32 FrameNumber, int32 NumBufferedFrames, int32 LagTolerance, bool& bOutGrouped) const
 	{
 		const int32 OldestQueryIndex = bNeedsScanOnRead ? ScanOldestNonStaleQueryIndex(FrameNumber, NumBufferedFrames, LagTolerance)
 														: FOcclusionQueryHelpers::GetQueryLookupIndex(FrameNumber, NumBufferedFrames);
 		const int32 LaggedFrames = FrameNumber - PendingOcclusionQueryFrames[OldestQueryIndex];
-		if (OldestQueryIndex == -1 || !PendingOcclusionQuery[OldestQueryIndex].GetQuery() || LaggedFrames > LagTolerance)
+		if (OldestQueryIndex == -1 || !PendingOcclusionQuery[OldestQueryIndex].IsValid() || LaggedFrames > LagTolerance)
 		{
 			bOutGrouped = false;
 			return nullptr;
 		}
 		bOutGrouped = bGroupedQuery[OldestQueryIndex];
-		return PendingOcclusionQuery[OldestQueryIndex].GetQuery();
+		return PendingOcclusionQuery[OldestQueryIndex];
 	}
 
-	FORCEINLINE void SetCurrentQuery(uint32 FrameNumber, FRHIPooledRenderQuery&& NewQuery, int32 NumBufferedFrames, bool bGrouped, bool bNeedsScan)
+	FORCEINLINE void SetCurrentQuery(uint32 FrameNumber, FRenderQueryRHIParamRef NewQuery, int32 NumBufferedFrames, bool bGrouped, bool bNeedsScan)
 	{
 		// Get the current occlusion query
 		const uint32 QueryIndex = FOcclusionQueryHelpers::GetQueryIssueIndex(FrameNumber, NumBufferedFrames);
-		PendingOcclusionQuery[QueryIndex] = MoveTemp(NewQuery);
+		PendingOcclusionQuery[QueryIndex] = NewQuery;
 		PendingOcclusionQueryFrames[QueryIndex] = FrameNumber;
 		bGroupedQuery[QueryIndex] = bGrouped;
 
@@ -309,7 +312,7 @@ struct FPrimitiveOcclusionHistoryKeyFuncs : BaseKeyFuncs<FPrimitiveOcclusionHist
 
 class FIndividualOcclusionHistory
 {
-	FRHIPooledRenderQuery PendingOcclusionQuery[FOcclusionQueryHelpers::MaxBufferedOcclusionFrames];
+	FRenderQueryRHIRef PendingOcclusionQuery[FOcclusionQueryHelpers::MaxBufferedOcclusionFrames];
 	uint32 PendingOcclusionQueryFrames[FOcclusionQueryHelpers::MaxBufferedOcclusionFrames]; // not intialized...this is ok
 
 public:
@@ -323,9 +326,9 @@ public:
 			{
 				uint32 TestFrame = FrameNumber - (DeltaFrame - 1);
 				const uint32 QueryIndex = FOcclusionQueryHelpers::GetQueryLookupIndex(TestFrame, NumBufferedFrames);
-				if (PendingOcclusionQuery[QueryIndex].GetQuery() && PendingOcclusionQueryFrames[QueryIndex] != TestFrame)
+				if (PendingOcclusionQuery[QueryIndex].GetReference() && PendingOcclusionQueryFrames[QueryIndex] != TestFrame)
 				{
-					PendingOcclusionQuery[QueryIndex].ReleaseQuery();
+					Pool.ReleaseQuery(PendingOcclusionQuery[QueryIndex]);
 				}
 			}
 		}
@@ -334,28 +337,28 @@ public:
 	FORCEINLINE void ReleaseQuery(TOcclusionQueryPool& Pool, uint32 FrameNumber, int32 NumBufferedFrames)
 	{
 		const uint32 QueryIndex = FOcclusionQueryHelpers::GetQueryLookupIndex(FrameNumber, NumBufferedFrames);
-		if (PendingOcclusionQuery[QueryIndex].GetQuery())
+		if (PendingOcclusionQuery[QueryIndex].GetReference())
 		{
-			PendingOcclusionQuery[QueryIndex].ReleaseQuery();
+			Pool.ReleaseQuery(PendingOcclusionQuery[QueryIndex]);
 		}
 	}
 
-	FORCEINLINE FRHIRenderQuery* GetPastQuery(uint32 FrameNumber, int32 NumBufferedFrames)
+	FORCEINLINE FRenderQueryRHIParamRef GetPastQuery(uint32 FrameNumber, int32 NumBufferedFrames)
 	{
 		// Get the oldest occlusion query
 		const uint32 QueryIndex = FOcclusionQueryHelpers::GetQueryLookupIndex(FrameNumber, NumBufferedFrames);
-		if (PendingOcclusionQuery[QueryIndex].GetQuery() && PendingOcclusionQueryFrames[QueryIndex] == FrameNumber - uint32(NumBufferedFrames))
+		if (PendingOcclusionQuery[QueryIndex].GetReference() && PendingOcclusionQueryFrames[QueryIndex] == FrameNumber - uint32(NumBufferedFrames))
 		{
-			return PendingOcclusionQuery[QueryIndex].GetQuery();
+			return PendingOcclusionQuery[QueryIndex].GetReference();
 		}
 		return nullptr;
 	}
 
-	FORCEINLINE void SetCurrentQuery(uint32 FrameNumber, FRHIPooledRenderQuery&& NewQuery, int32 NumBufferedFrames)
+	FORCEINLINE void SetCurrentQuery(uint32 FrameNumber, FRenderQueryRHIParamRef NewQuery, int32 NumBufferedFrames)
 	{
 		// Get the current occlusion query
 		const uint32 QueryIndex = FOcclusionQueryHelpers::GetQueryIssueIndex(FrameNumber, NumBufferedFrames);
-		PendingOcclusionQuery[QueryIndex] = MoveTemp(NewQuery);
+		PendingOcclusionQuery[QueryIndex] = NewQuery;
 		PendingOcclusionQueryFrames[QueryIndex] = FrameNumber;
 	}
 };
@@ -680,11 +683,11 @@ public:
 	};
 
 	uint32 UniqueID;
-	typedef TMap<FSceneViewState::FProjectedShadowKey, FRHIPooledRenderQuery> ShadowKeyOcclusionQueryMap;
+	typedef TMap<FSceneViewState::FProjectedShadowKey, FRenderQueryRHIRef> ShadowKeyOcclusionQueryMap;
 	TArray<ShadowKeyOcclusionQueryMap, TInlineAllocator<FOcclusionQueryHelpers::MaxBufferedOcclusionFrames> > ShadowOcclusionQueryMaps;
 
 	/** The view's occlusion query pool. */
-	FRenderQueryPoolRHIRef OcclusionQueryPool;
+	FRenderQueryPool OcclusionQueryPool;
 
 	FHZBOcclusionTester HZBOcclusionTests;
 
diff --git a/Engine/Source/Runtime/Renderer/Private/SceneRendering.h b/Engine/Source/Runtime/Renderer/Private/SceneRendering.h
--- a/Engine/Source/Runtime/Renderer/Private/SceneRendering.h
+++ b/Engine/Source/Runtime/Renderer/Private/SceneRendering.h
@@ -249,7 +249,7 @@ public:
 	 * Batches a primitive's occlusion query for rendering.
 	 * @param Bounds - The primitive's bounds.
 	 */
-	FRHIPooledRenderQuery BatchPrimitive(const FVector& BoundsOrigin, const FVector& BoundsBoxExtent, FGlobalDynamicVertexBuffer& DynamicVertexBuffer);
+	FRenderQueryRHIParamRef BatchPrimitive(const FVector& BoundsOrigin, const FVector& BoundsBoxExtent, FGlobalDynamicVertexBuffer& DynamicVertexBuffer);
 	inline int32 GetNumBatchOcclusionQueries() const
 	{
 		return BatchOcclusionQueries.Num();
@@ -259,7 +259,7 @@ private:
 
 	struct FOcclusionBatch
 	{
-		TRefCountPtr<FRHIRenderQuery> Query;
+		FRenderQueryRHIRef Query;
 		FGlobalDynamicVertexBuffer::FAllocation VertexAllocation;
 	};
 
@@ -276,7 +276,7 @@ private:
 	uint32 NumBatchedPrimitives;
 
 	/** The pool to allocate occlusion queries from. */
-	TRefCountPtr<FRHIRenderQueryPool> OcclusionQueryPool;
+	class FRenderQueryPool* OcclusionQueryPool;
 };
 
 class FHZBOcclusionTester : public FRenderResource
diff --git a/Engine/Source/Runtime/Renderer/Private/SceneVisibility.cpp b/Engine/Source/Runtime/Renderer/Private/SceneVisibility.cpp
--- a/Engine/Source/Runtime/Renderer/Private/SceneVisibility.cpp
+++ b/Engine/Source/Runtime/Renderer/Private/SceneVisibility.cpp
@@ -695,7 +695,7 @@ static void FetchVisibilityForPrimitives_Range(FVisForPrimParams& Params, FGloba
 	bool bClearQueries = !View.Family->EngineShowFlags.HitProxies;
 	const float CurrentRealTime = View.Family->CurrentRealTime;
 	uint32 OcclusionFrameCounter = ViewState->OcclusionFrameCounter;
-	FRHIRenderQueryPool* OcclusionQueryPool = ViewState->OcclusionQueryPool;
+	FRenderQueryPool& OcclusionQueryPool = ViewState->OcclusionQueryPool;
 	FHZBOcclusionTester& HZBOcclusionTests = ViewState->HZBOcclusionTests;
 
 	int32 ReadBackLagTolerance = NumBufferedFrames;
@@ -858,7 +858,7 @@ static void FetchVisibilityForPrimitives_Range(FVisForPrimParams& Params, FGloba
 						// Read the occlusion query results.
 						uint64 NumSamples = 0;
 						bool bGrouped = false;
-						FRHIRenderQuery* PastQuery = PrimitiveOcclusionHistory->GetQueryForReading(OcclusionFrameCounter, NumBufferedFrames, ReadBackLagTolerance, bGrouped);
+						FRenderQueryRHIParamRef PastQuery = PrimitiveOcclusionHistory->GetQueryForReading(OcclusionFrameCounter, NumBufferedFrames, ReadBackLagTolerance, bGrouped);
 						if (PastQuery)
 						{
 							//int32 RefCount = PastQuery.GetReference()->GetRefCount();
@@ -1299,7 +1299,7 @@ static int32 FetchVisibilityForPrimitives(const FScene* Scene, FViewInfo& View,
 			StartIndex += NumToProcess;
 		}
 
-		FRHIRenderQueryPool* OcclusionQueryPool = ViewState->OcclusionQueryPool;
+		FRenderQueryPool& OcclusionQueryPool = ViewState->OcclusionQueryPool;
 		FHZBOcclusionTester& HZBOcclusionTests = ViewState->HZBOcclusionTests;		
 
 		int32 NumOccludedPrims = 0;
@@ -1368,10 +1368,12 @@ static int32 FetchVisibilityForPrimitives(const FScene* Scene, FViewInfo& View,
 
 			//now add new primitive histories to the view. may resize the view's array.
 			for (int32 i = 0; i < NumTasks; ++i)
-			{											
-				for (int32 HistoryIndex = 0; HistoryIndex < OutputOcclusionHistory[i].Num(); ++HistoryIndex)
+			{								
+				const TArray<FPrimitiveOcclusionHistory>& NewHistoryArray = OutputOcclusionHistory[i];				
+				for (int32 HistoryIndex = 0; HistoryIndex < NewHistoryArray.Num(); ++HistoryIndex)
 				{
-					ViewPrimitiveOcclusionHistory.Add(MoveTemp(OutputOcclusionHistory[i][HistoryIndex]));
+					const FPrimitiveOcclusionHistory& CopySourceHistory = NewHistoryArray[HistoryIndex];
+					ViewPrimitiveOcclusionHistory.Add(CopySourceHistory);
 				}
 
 				//accumulate occluded prims across tasks
