Add multiple extensions to virtual texturing system
* bCreateRenderTarget to specify render target & UAV flags on physical spaces
* bContinuousUpdate to specify if a resident tile should be re-requested every frame for continuous refinement
* bPersistentHighestMip to specify if the texture's highest mip wants to be locked so that no garbage data will be displayed when the detailed mips are still in streaming. This defaults to true for disk-streamed VTs, however since each VT has at least one tile in the highest mip, it can saturate the system when the number of allocated VTs is large

diff --git a/Engine/Source/Runtime/RenderCore/Public/VirtualTexturing.h b/Engine/Source/Runtime/RenderCore/Public/VirtualTexturing.h
--- a/Engine/Source/Runtime/RenderCore/Public/VirtualTexturing.h
+++ b/Engine/Source/Runtime/RenderCore/Public/VirtualTexturing.h
@@ -93,6 +93,9 @@ inline bool operator!=(const FAllocatedVTDescription& Lhs, const FAllocatedVTDes
 struct FVTProducerDescription
 {
 	FName Name; /** Will be name of UTexture for streaming VTs, mostly here for debugging */
+	bool bPersistentHighestMip = true;
+	bool bContinuousUpdate = false;
+	bool bCreateRenderTarget = false;
 	uint32 TileSize = 0u;
 	uint32 TileBorderSize = 0u;
 	uint32 WidthInTiles = 0u;
@@ -165,6 +168,8 @@ struct FVTProduceTargetLayer
 	/** The texture to write to */
 	FRHITexture* TextureRHI = nullptr;
 
+	TRefCountPtr<struct IPooledRenderTarget> PooledRenderTarget = nullptr;
+
 	/** Location within the texture to write */
 	FIntVector pPageLocation;
 };
diff --git a/Engine/Source/Runtime/Renderer/Private/VT/TexturePagePool.cpp b/Engine/Source/Runtime/Renderer/Private/VT/TexturePagePool.cpp
--- a/Engine/Source/Runtime/Renderer/Private/VT/TexturePagePool.cpp
+++ b/Engine/Source/Runtime/Renderer/Private/VT/TexturePagePool.cpp
@@ -5,7 +5,6 @@
 #include "VirtualTextureSpace.h"
 #include "VirtualTextureSystem.h"
 
-
 FTexturePagePool::FTexturePagePool()
 	: PageHash(16u * 1024)
 	, NumPages(0u)
@@ -16,8 +15,10 @@ FTexturePagePool::FTexturePagePool()
 FTexturePagePool::~FTexturePagePool()
 {}
 
+
 void FTexturePagePool::Initialize(uint32 InNumPages)
 {
+
 	NumPages = InNumPages;
 	Pages.AddZeroed(InNumPages);
 	PageHash.Resize(InNumPages);
@@ -42,6 +43,7 @@ void FTexturePagePool::Initialize(uint32 InNumPages)
 
 void FTexturePagePool::EvictAllPages(FVirtualTextureSystem* System)
 {
+
 	TArray<uint16> PagesToEvict;
 	while (FreeHeap.Num() > 0u)
 	{
@@ -59,6 +61,7 @@ void FTexturePagePool::EvictAllPages(FVirtualTextureSystem* System)
 
 void FTexturePagePool::UnmapAllPagesForSpace(FVirtualTextureSystem* System, uint8 SpaceID)
 {
+
 	// walk through all of our current mapping entries, and unmap any that belong to the current space
 	for (int32 MappingIndex = NumPages + 1; MappingIndex < PageMapping.Num(); ++MappingIndex)
 	{
@@ -73,6 +76,7 @@ void FTexturePagePool::UnmapAllPagesForSpace(FVirtualTextureSystem* System, uint
 
 void FTexturePagePool::EvictPages(FVirtualTextureSystem* System, const FVirtualTextureProducerHandle& ProducerHandle)
 {
+
 	for (uint32 pAddress = 0u; pAddress < NumPages; ++pAddress)
 	{
 		const FPageEntry& PageEntry = Pages[pAddress];
@@ -86,6 +90,7 @@ void FTexturePagePool::EvictPages(FVirtualTextureSystem* System, const FVirtualT
 
 void FTexturePagePool::GetAllLockedPages(FVirtualTextureSystem* System, TSet<FVirtualTextureLocalTile>& OutPages)
 {
+
 	OutPages.Reserve(OutPages.Num() + GetNumLockedPages());
 
 	for (uint32 i = 0; i < NumPages; ++i)
@@ -97,6 +102,11 @@ void FTexturePagePool::GetAllLockedPages(FVirtualTextureSystem* System, TSet<FVi
 	}
 }
 
+FVirtualTextureLocalTile FTexturePagePool::GetLocalTileFromPhysicalAddress(uint16 pAddress)
+{
+	return FVirtualTextureLocalTile(FVirtualTextureProducerHandle(Pages[pAddress].PackedProducerHandle), Pages[pAddress].Local_vAddress, Pages[pAddress].Local_vLevel);
+}
+
 bool FTexturePagePool::AnyFreeAvailable( uint32 Frame ) const
 {
 	if( FreeHeap.Num() > 0 )
diff --git a/Engine/Source/Runtime/Renderer/Private/VT/TexturePagePool.h b/Engine/Source/Runtime/Renderer/Private/VT/TexturePagePool.h
--- a/Engine/Source/Runtime/Renderer/Private/VT/TexturePagePool.h
+++ b/Engine/Source/Runtime/Renderer/Private/VT/TexturePagePool.h
@@ -51,6 +51,8 @@ public:
 	 * Get descriptions of the locked pages in this pool
 	 */
 	void GetAllLockedPages(FVirtualTextureSystem* System, TSet<union FVirtualTextureLocalTile>& OutPages);
+
+	FVirtualTextureLocalTile GetLocalTileFromPhysicalAddress(uint16 pAddress);
 	
 	/** Get the local vLevel of the page allocated at the given physical address */
 	inline uint8 GetLocalLevelForAddress(uint16 pAddress) const { check(Pages[pAddress].PackedValue != 0u); return Pages[pAddress].Local_vLevel; }
diff --git a/Engine/Source/Runtime/Renderer/Private/VT/UniqueRequestList.h b/Engine/Source/Runtime/Renderer/Private/VT/UniqueRequestList.h
--- a/Engine/Source/Runtime/Renderer/Private/VT/UniqueRequestList.h
+++ b/Engine/Source/Runtime/Renderer/Private/VT/UniqueRequestList.h
@@ -148,7 +148,7 @@ inline uint16 FUniqueRequestList::AddLoadRequest(const FVirtualTextureLocalTile&
 		}
 	}
 	
-	if (ensure(NumLoadRequests < LoadRequestCapacity))
+	if (NumLoadRequests < LoadRequestCapacity)
 	{
 		const uint32 Index = NumLoadRequests++;
 		LoadRequestHash.Add(Hash, Index);
@@ -179,7 +179,7 @@ inline uint16 FUniqueRequestList::LockLoadRequest(const FVirtualTextureLocalTile
 		}
 	}
 
-	if (ensure(NumLoadRequests < LoadRequestCapacity))
+	if (NumLoadRequests < LoadRequestCapacity)
 	{
 		const uint32 Index = NumLoadRequests++;
 		LoadRequestHash.Add(Hash, Index);
diff --git a/Engine/Source/Runtime/Renderer/Private/VT/VirtualTexturePhysicalSpace.cpp b/Engine/Source/Runtime/Renderer/Private/VT/VirtualTexturePhysicalSpace.cpp
--- a/Engine/Source/Runtime/Renderer/Private/VT/VirtualTexturePhysicalSpace.cpp
+++ b/Engine/Source/Runtime/Renderer/Private/VT/VirtualTexturePhysicalSpace.cpp
@@ -62,7 +62,7 @@ void FVirtualTexturePhysicalSpace::InitRHI()
 		Description.Format,
 		FClearValueBinding::None,
 		TexCreate_None,
-		TexCreate_ShaderResource,
+		TexCreate_ShaderResource | (Description.bCreateRenderTarget ? (TexCreate_RenderTargetable | TexCreate_UAV) : 0),
 		false);
 
 	GRenderTargetPool.FindFreeElement(RHICmdList, Desc, PooledRenderTarget, TEXT("PhysicalTexture"));
diff --git a/Engine/Source/Runtime/Renderer/Private/VT/VirtualTexturePhysicalSpace.h b/Engine/Source/Runtime/Renderer/Private/VT/VirtualTexturePhysicalSpace.h
--- a/Engine/Source/Runtime/Renderer/Private/VT/VirtualTexturePhysicalSpace.h
+++ b/Engine/Source/Runtime/Renderer/Private/VT/VirtualTexturePhysicalSpace.h
@@ -13,11 +13,13 @@ struct FVTPhysicalSpaceDescription
 	uint32 TileSize;
 	TEnumAsByte<EPixelFormat> Format;
 	uint8 Dimensions;
+	bool bContinuousUpdate;
+	bool bCreateRenderTarget;
 };
 
 inline bool operator==(const FVTPhysicalSpaceDescription& Lhs, const FVTPhysicalSpaceDescription& Rhs)
 {
-	return Lhs.TileSize == Rhs.TileSize && Lhs.Format == Rhs.Format && Lhs.Dimensions == Rhs.Dimensions;
+	return Lhs.TileSize == Rhs.TileSize && Lhs.Format == Rhs.Format && Lhs.Dimensions == Rhs.Dimensions && Lhs.bContinuousUpdate == Rhs.bContinuousUpdate && Lhs.bCreateRenderTarget == Rhs.bCreateRenderTarget;
 }
 inline bool operator!=(const FVTPhysicalSpaceDescription& Lhs, const FVTPhysicalSpaceDescription& Rhs)
 {
@@ -60,6 +62,13 @@ public:
 		return PooledRenderTarget->GetRenderTargetItem().ShaderResourceTexture;
 	}
 
+	TRefCountPtr<IPooledRenderTarget> GetPhysicalTexturePooledRenderTarget() const
+	{
+		check(PooledRenderTarget.IsValid());
+		check(Description.bCreateRenderTarget);
+		return PooledRenderTarget;
+	}
+
 	FRHIShaderResourceView* GetPhysicalTextureView(bool bSRGB) const
 	{
 		return bSRGB ? TextureSRGBView : TextureView;
diff --git a/Engine/Source/Runtime/Renderer/Private/VT/VirtualTextureProducer.cpp b/Engine/Source/Runtime/Renderer/Private/VT/VirtualTextureProducer.cpp
--- a/Engine/Source/Runtime/Renderer/Private/VT/VirtualTextureProducer.cpp
+++ b/Engine/Source/Runtime/Renderer/Private/VT/VirtualTextureProducer.cpp
@@ -6,15 +6,18 @@
 
 void FVirtualTextureProducer::Release(FVirtualTextureSystem* System, const FVirtualTextureProducerHandle& HandleToSelf)
 {
-	const uint32 RootWidthInTiles = FMath::Max(Description.WidthInTiles >> Description.MaxLevel, 1u);
-	const uint32 RootHeightInTiles = FMath::Max(Description.HeightInTiles >> Description.MaxLevel, 1u);
-	for (uint32 TileY = 0u; TileY < RootHeightInTiles; ++TileY)
+	if (Description.bPersistentHighestMip)
 	{
-		for (uint32 TileX = 0u; TileX < RootWidthInTiles; ++TileX)
+		const uint32 RootWidthInTiles = FMath::Max(Description.WidthInTiles >> Description.MaxLevel, 1u);
+		const uint32 RootHeightInTiles = FMath::Max(Description.HeightInTiles >> Description.MaxLevel, 1u);
+		for (uint32 TileY = 0u; TileY < RootHeightInTiles; ++TileY)
 		{
-			const uint32 Local_vAddress = FMath::MortonCode2(TileX) | (FMath::MortonCode2(TileY) << 1);
-			const FVirtualTextureLocalTile TileToUnlock(HandleToSelf, Local_vAddress, Description.MaxLevel);
-			System->UnlockTile(TileToUnlock);
+			for (uint32 TileX = 0u; TileX < RootWidthInTiles; ++TileX)
+			{
+				const uint32 Local_vAddress = FMath::MortonCode2(TileX) | (FMath::MortonCode2(TileY) << 1);
+				const FVirtualTextureLocalTile TileToUnlock(HandleToSelf, Local_vAddress, Description.MaxLevel);
+				System->UnlockTile(TileToUnlock);
+			}
 		}
 	}
 
@@ -54,20 +57,26 @@ FVirtualTextureProducerHandle FVirtualTextureProducerCollection::RegisterProduce
 		PhysicalSpaceDesc.Dimensions = InDesc.Dimensions;
 		PhysicalSpaceDesc.TileSize = InDesc.TileSize + InDesc.TileBorderSize * 2u;
 		PhysicalSpaceDesc.Format = InDesc.LayerFormat[LayerIndex];
+		PhysicalSpaceDesc.bContinuousUpdate = InDesc.bContinuousUpdate;
+		PhysicalSpaceDesc.bCreateRenderTarget = InDesc.bCreateRenderTarget;
 		Entry.Producer.PhysicalSpace[LayerIndex] = System->AcquirePhysicalSpace(PhysicalSpaceDesc);
 	}
 
 	const FVirtualTextureProducerHandle Handle(Index, Entry.Magic);
-	const uint32 MaxLevel = InDesc.MaxLevel;
-	const uint32 RootWidthInTiles = FMath::Max(InDesc.WidthInTiles >> MaxLevel, 1u);
-	const uint32 RootHeightInTiles = FMath::Max(InDesc.HeightInTiles >> MaxLevel, 1u);
-	for (uint32 TileY = 0u; TileY < RootHeightInTiles; ++TileY)
+
+	if (InDesc.bPersistentHighestMip)
 	{
-		for (uint32 TileX = 0u; TileX < RootWidthInTiles; ++TileX)
+		const uint32 MaxLevel = InDesc.MaxLevel;
+		const uint32 RootWidthInTiles = FMath::Max(InDesc.WidthInTiles >> MaxLevel, 1u);
+		const uint32 RootHeightInTiles = FMath::Max(InDesc.HeightInTiles >> MaxLevel, 1u);
+		for (uint32 TileY = 0u; TileY < RootHeightInTiles; ++TileY)
 		{
-			const uint32 Local_vAddress = FMath::MortonCode2(TileX) | (FMath::MortonCode2(TileY) << 1);
-			const FVirtualTextureLocalTile TileToLock(Handle, Local_vAddress, MaxLevel);
-			System->LockTile(TileToLock);
+			for (uint32 TileX = 0u; TileX < RootWidthInTiles; ++TileX)
+			{
+				const uint32 Local_vAddress = FMath::MortonCode2(TileX) | (FMath::MortonCode2(TileY) << 1);
+				const FVirtualTextureLocalTile TileToLock(Handle, Local_vAddress, MaxLevel);
+				System->LockTile(TileToLock);
+			}
 		}
 	}
 
diff --git a/Engine/Source/Runtime/Renderer/Private/VT/VirtualTextureSystem.cpp b/Engine/Source/Runtime/Renderer/Private/VT/VirtualTextureSystem.cpp
--- a/Engine/Source/Runtime/Renderer/Private/VT/VirtualTextureSystem.cpp
+++ b/Engine/Source/Runtime/Renderer/Private/VT/VirtualTextureSystem.cpp
@@ -15,6 +15,7 @@
 #include "HAL/IConsoleManager.h"
 #include "PostProcess/SceneRenderTargets.h"
 
+
 DECLARE_CYCLE_STAT(TEXT("Feedback Analysis"), STAT_FeedbackAnalysis, STATGROUP_VirtualTexturing);
 DECLARE_CYCLE_STAT(TEXT("VirtualTextureSystem Update"), STAT_VirtualTextureSystem_Update, STATGROUP_VirtualTexturing);
 
@@ -34,6 +35,7 @@ DECLARE_DWORD_COUNTER_STAT(TEXT("Num page visible resident"), STAT_NumPageVisibl
 DECLARE_DWORD_COUNTER_STAT(TEXT("Num page visible not resident"), STAT_NumPageVisibleNotResident, STATGROUP_VirtualTexturing);
 DECLARE_DWORD_COUNTER_STAT(TEXT("Num page prefetch"), STAT_NumPagePrefetch, STATGROUP_VirtualTexturing);
 DECLARE_DWORD_COUNTER_STAT(TEXT("Num page update"), STAT_NumPageUpdate, STATGROUP_VirtualTexturing);
+DECLARE_DWORD_COUNTER_STAT(TEXT("Num continuous page update"), STAT_NumContinuousPageUpdate, STATGROUP_VirtualTexturing);
 
 DECLARE_DWORD_COUNTER_STAT(TEXT("Num stacks requested"), STAT_NumStacksRequested, STATGROUP_VirtualTexturing);
 DECLARE_DWORD_COUNTER_STAT(TEXT("Num stacks produced"), STAT_NumStacksProduced, STATGROUP_VirtualTexturing);
@@ -329,6 +331,7 @@ IAllocatedVirtualTexture* FVirtualTextureSystem::AllocateVirtualTexture(const FA
 	uint32 DepthInTiles = 0u;
 	bool bSupport16BitPageTable = true;
 	FVirtualTextureProducer* ProducerForLayer[VIRTUALTEXTURE_SPACE_MAXLAYERS] = { nullptr };
+	bool bAnyLayerProducerWantsPersistentHighestMip = false;
 	for (uint32 LayerIndex = 0u; LayerIndex < Desc.NumLayers; ++LayerIndex)
 	{
 		FVirtualTextureProducer* Producer = Producers.FindProducer(Desc.ProducerHandle[LayerIndex]);
@@ -343,6 +346,7 @@ IAllocatedVirtualTexture* FVirtualTextureSystem::AllocateVirtualTexture(const FA
 			{
 				bSupport16BitPageTable = false;
 			}
+			bAnyLayerProducerWantsPersistentHighestMip |= Producer->GetDescription().bPersistentHighestMip;
 		}
 	}
 
@@ -360,7 +364,10 @@ IAllocatedVirtualTexture* FVirtualTextureSystem::AllocateVirtualTexture(const FA
 	FVirtualTextureSpace* Space = AcquireSpace(SpaceDesc, FMath::Max(WidthInTiles, HeightInTiles));
 
 	AllocatedVT = new FAllocatedVirtualTexture(Frame, Desc, Space, ProducerForLayer, WidthInTiles, HeightInTiles, DepthInTiles);
-	AllocatedVTsToMap.Add(AllocatedVT);
+	if (bAnyLayerProducerWantsPersistentHighestMip)
+	{
+		AllocatedVTsToMap.Add(AllocatedVT);
+	}
 	return AllocatedVT;
 }
 
@@ -1009,7 +1016,11 @@ void FVirtualTextureSystem::GatherRequestsTask(const FGatherRequestsParameters&
 	uint32 NumResidentPages = 0u;
 	uint32 NumNonResidentPages = 0u;
 	uint32 NumPrefetchPages = 0u;
-	
+
+#if WITH_EDITOR
+	TSet<FVirtualTextureLocalTile> ContinuousUpdateTilesToProduceThreadLocal;
+#endif
+
 	for (uint32 i = Parameters.PageStartIndex; i < PageEndIndex; ++i)
 	{
 		const uint32 PageEncoded = UniquePageList->GetPage(i);
@@ -1054,6 +1065,17 @@ void FVirtualTextureSystem::GatherRequestsTask(const FGatherRequestsParameters&
 					// Page is already resident, just need to update LRU free list
 					AddPageUpdate(PageUpdateBuffers, PageUpdateFlushCount, PhysicalSpaceIDAndAddress.PhysicalSpaceID, PhysicalSpaceIDAndAddress.pAddress);
 
+				#if WITH_EDITOR
+					{
+						if (GetPhysicalSpace(PhysicalSpaceIDAndAddress.PhysicalSpaceID)->GetDescription().bContinuousUpdate)
+						{
+							FTexturePagePool& RESTRICT PagePool = GetPhysicalSpace(PhysicalSpaceIDAndAddress.PhysicalSpaceID)->GetPagePool();
+
+							ContinuousUpdateTilesToProduceThreadLocal.Add(PagePool.GetLocalTileFromPhysicalAddress(PhysicalSpaceIDAndAddress.pAddress));
+						}
+					}
+				#endif
+
 					++PageUpdateBuffers[PhysicalSpaceIDAndAddress.PhysicalSpaceID].WorkingSetSize;
 					++NumResidentPages;
 				}
@@ -1183,7 +1205,7 @@ void FVirtualTextureSystem::GatherRequestsTask(const FGatherRequestsParameters&
 					const uint32 Allocated_vAddress = vAddress & (0xffffffff << (Allocated_vLevel * vDimensions));
 
 					AddPageUpdate(PageUpdateBuffers, PageUpdateFlushCount, PhysicalSpace->GetID(), pAddress);
-	
+
 					uint32 NumMappedPages = 0u;
 					for (uint32 LoadLayerIndex = 0u; LoadLayerIndex < NumLayersToLoad; ++LoadLayerIndex)
 					{
@@ -1342,6 +1364,14 @@ void FVirtualTextureSystem::GatherRequestsTask(const FGatherRequestsParameters&
 			{
 				PagePool.UpdateUsage(Frame, Buffer.PhysicalAddresses[i]);
 			}
+
+		#if WITH_EDITOR
+			if (PhysicalSpace->GetDescription().bContinuousUpdate)
+			{
+				FScopeLock Lock(&ContinuousUpdateTilesToProduceCS);
+				ContinuousUpdateTilesToProduce.Append(ContinuousUpdateTilesToProduceThreadLocal);
+			}
+		#endif
 		}
 		
 		INC_DWORD_STAT_BY(STAT_NumPageUpdate, Buffer.NumPageUpdates);
@@ -1353,9 +1383,9 @@ void FVirtualTextureSystem::GatherRequestsTask(const FGatherRequestsParameters&
 	INC_DWORD_STAT_BY(STAT_NumPagePrefetch, NumPrefetchPages);
 }
 
-void FVirtualTextureSystem::SubmitPreMappedRequests(FRHICommandListImmediate& RHICmdList, ERHIFeatureLevel::Type FeatureLevel)
+void FVirtualTextureSystem::SubmitRequestsFromLocalTileList(const TSet<FVirtualTextureLocalTile>& LocalTileList, EVTProducePageFlags Flags, FRHICommandListImmediate& RHICmdList, ERHIFeatureLevel::Type FeatureLevel)
 {
-	for (FVirtualTextureLocalTile const& Tile : MappedTilesToProduce)
+	for (const FVirtualTextureLocalTile& Tile : LocalTileList)
 	{
 		const FVirtualTextureProducerHandle ProducerHandle = Tile.GetProducerHandle();
 		const FVirtualTextureProducer& Producer = Producers.GetProducer(ProducerHandle);
@@ -1371,6 +1401,10 @@ void FVirtualTextureSystem::SubmitPreMappedRequests(FRHICommandListImmediate& RH
 			if (pAddress != ~0u)
 			{
 				ProduceTarget[LocalLayerIndex].TextureRHI = PhysicalSpace->GetPhysicalTexture();
+				if (PhysicalSpace->GetDescription().bCreateRenderTarget)
+				{
+					ProduceTarget[LocalLayerIndex].PooledRenderTarget = PhysicalSpace->GetPhysicalTexturePooledRenderTarget();
+				}
 				ProduceTarget[LocalLayerIndex].pPageLocation = PhysicalSpace->GetPhysicalLocation(pAddress);
 				LayerMask |= 1 << LocalLayerIndex;
 			}
@@ -1393,7 +1427,7 @@ void FVirtualTextureSystem::SubmitPreMappedRequests(FRHICommandListImmediate& RH
 
 		IVirtualTextureFinalizer* VTFinalizer = Producer.GetVirtualTexture()->ProducePageData(
 			RHICmdList, FeatureLevel,
-			EVTProducePageFlags::None,
+			Flags,
 			ProducerHandle, LayerMask, Tile.Local_vLevel, Tile.Local_vAddress,
 			RequestPageResult.Handle,
 			ProduceTarget);
@@ -1404,8 +1438,20 @@ void FVirtualTextureSystem::SubmitPreMappedRequests(FRHICommandListImmediate& RH
 			Finalizers.AddUnique(VTFinalizer);
 		}
 	}
+}
+
+void FVirtualTextureSystem::SubmitPreMappedRequests(FRHICommandListImmediate& RHICmdList, ERHIFeatureLevel::Type FeatureLevel)
+{
+	{
+		SubmitRequestsFromLocalTileList(MappedTilesToProduce, EVTProducePageFlags::None, RHICmdList, FeatureLevel);
+		MappedTilesToProduce.Reset();
+	}
 
-	MappedTilesToProduce.Reset();
+	{
+		INC_DWORD_STAT_BY(STAT_NumContinuousPageUpdate, ContinuousUpdateTilesToProduce.Num());
+		SubmitRequestsFromLocalTileList(ContinuousUpdateTilesToProduce, EVTProducePageFlags::None, RHICmdList, FeatureLevel);
+		ContinuousUpdateTilesToProduce.Reset();
+	}
 }
 
 void FVirtualTextureSystem::SubmitRequests(FRHICommandListImmediate& RHICmdList, ERHIFeatureLevel::Type FeatureLevel, FMemStack& MemStack, FUniqueRequestList* RequestList, bool bAsync)
@@ -1465,6 +1511,10 @@ void FVirtualTextureSystem::SubmitRequests(FRHICommandListImmediate& RHICmdList,
 							check(pAddress != ~0u);
 
 							ProduceTarget[LocalLayerIndex].TextureRHI = PhysicalSpace->GetPhysicalTexture();
+							if (PhysicalSpace->GetDescription().bCreateRenderTarget)
+							{
+								ProduceTarget[LocalLayerIndex].PooledRenderTarget = PhysicalSpace->GetPhysicalTexturePooledRenderTarget();
+							}
 							ProduceTarget[LocalLayerIndex].pPageLocation = PhysicalSpace->GetPhysicalLocation(pAddress);
 							Allocate_pAddress[LocalLayerIndex] = pAddress;
 						}
diff --git a/Engine/Source/Runtime/Renderer/Private/VT/VirtualTextureSystem.h b/Engine/Source/Runtime/Renderer/Private/VT/VirtualTextureSystem.h
--- a/Engine/Source/Runtime/Renderer/Private/VT/VirtualTextureSystem.h
+++ b/Engine/Source/Runtime/Renderer/Private/VT/VirtualTextureSystem.h
@@ -68,6 +68,8 @@ private:
 	void DestroyPendingVirtualTextures();
 
 	void RequestTilesForRegionInternal(const IAllocatedVirtualTexture* AllocatedVT, const FIntRect& InTextureRegion, uint32 vLevel);
+	
+	void SubmitRequestsFromLocalTileList(const TSet<FVirtualTextureLocalTile>& LocalTileList, EVTProducePageFlags Flags, FRHICommandListImmediate& RHICmdList, ERHIFeatureLevel::Type FeatureLevel);
 
 	void SubmitPreMappedRequests(FRHICommandListImmediate& RHICmdList, ERHIFeatureLevel::Type FeatureLevel);
 
@@ -106,6 +108,8 @@ private:
 	TArray<uint32> RequestedPackedTiles;
 
 	TArray<FVirtualTextureLocalTile> TilesToLock;
+	FCriticalSection ContinuousUpdateTilesToProduceCS;
+	TSet<FVirtualTextureLocalTile> ContinuousUpdateTilesToProduce;
 	TSet<FVirtualTextureLocalTile> MappedTilesToProduce;
 	TArray<FAllocatedVirtualTexture*> AllocatedVTsToMap;
 	TArray<IVirtualTextureFinalizer*> Finalizers;
