Fix variable shadowing

diff --git a/Engine/Source/Runtime/Renderer/Private/VT/VirtualTextureSystem.cpp b/Engine/Source/Runtime/Renderer/Private/VT/VirtualTextureSystem.cpp
--- a/Engine/Source/Runtime/Renderer/Private/VT/VirtualTextureSystem.cpp
+++ b/Engine/Source/Runtime/Renderer/Private/VT/VirtualTextureSystem.cpp
@@ -1368,7 +1368,7 @@ void FVirtualTextureSystem::GatherRequestsTask(const FGatherRequestsParameters&
 		#if WITH_EDITOR
 			if (PhysicalSpace->GetDescription().bContinuousUpdate)
 			{
-				FScopeLock Lock(&ContinuousUpdateTilesToProduceCS);
+				FScopeLock ScopeLock(&ContinuousUpdateTilesToProduceCS);
 				ContinuousUpdateTilesToProduce.Append(ContinuousUpdateTilesToProduceThreadLocal);
 			}
 		#endif
