Fix truncation compile error when building for Android (clang?)

diff --git a/Engine/Source/Runtime/Renderer/Private/VT/TexturePageMap.cpp b/Engine/Source/Runtime/Renderer/Private/VT/TexturePageMap.cpp
--- a/Engine/Source/Runtime/Renderer/Private/VT/TexturePageMap.cpp
+++ b/Engine/Source/Runtime/Renderer/Private/VT/TexturePageMap.cpp
@@ -291,8 +291,7 @@ void FTexturePageMap::ReleaseUnmappedPages()
 		FPageEntry& Entry = Pages[PageIndex];
 		const uint32 NextPageIndex = Entry.NextIndex;
 		Entry.Page.Packed = ~0u;
-		Entry.pAddress = ~0u;
-		Entry.PhysicalSpaceID = ~0u;
+		Entry.Packed = ~0u;
 		RemovePageFromList(PageIndex);
 		AddPageToList(PageListHead_Free, PageIndex);
 		PageIndex = NextPageIndex;
diff --git a/Engine/Source/Runtime/Renderer/Private/VT/TexturePageMap.h b/Engine/Source/Runtime/Renderer/Private/VT/TexturePageMap.h
--- a/Engine/Source/Runtime/Renderer/Private/VT/TexturePageMap.h
+++ b/Engine/Source/Runtime/Renderer/Private/VT/TexturePageMap.h
@@ -124,9 +124,16 @@ private:
 		FTexturePage Page;
 		uint32 NextIndex;
 		uint32 PrevIndex;
-		uint32 pAddress : 16;
-		uint32 PhysicalSpaceID : 12;
-		uint32 vLevel : 4;
+		union
+		{
+			uint32 Packed;
+			struct 
+			{
+				uint32 pAddress : 16;
+				uint32 PhysicalSpaceID : 12;
+				uint32 vLevel : 4;
+			};
+		};
 	};
 
 	void RemovePageFromList(uint32 Index)
