Add bounds checks for new'd array, to satisfy static analysis

diff --git a/Engine/Source/Runtime/Renderer/Private/VT/UniqueRequestList.h b/Engine/Source/Runtime/Renderer/Private/VT/UniqueRequestList.h
--- a/Engine/Source/Runtime/Renderer/Private/VT/UniqueRequestList.h
+++ b/Engine/Source/Runtime/Renderer/Private/VT/UniqueRequestList.h
@@ -325,6 +325,7 @@ inline void FUniqueRequestList::SortRequests(FVirtualTextureProducerCollection&
 		const uint32 SortedIndex = SortedKeys[i].Index;
 		SortedLoadRequests[i] = LoadRequests[SortedIndex];
 		SortedLayerMask[i] = LoadRequestLayerMask[SortedIndex];
+		checkSlow(SortedIndex < NumLoadRequests);
 		LoadIndexToSortedLoadIndex[SortedIndex] = i;
 	}
 	FMemory::Memcpy(LoadRequests, SortedLoadRequests, sizeof(FVirtualTextureLocalTile) * NewNumLoadRequests);
@@ -337,6 +338,7 @@ inline void FUniqueRequestList::SortRequests(FVirtualTextureProducerCollection&
 	for (uint32 i = 0u; i < NumMappingRequests; ++i)
 	{
 		FMappingRequest Request = GetMappingRequest(i);
+		checkSlow(Request.LoadRequestIndex < NumLoadRequests);
 		const uint16 SortedLoadIndex = LoadIndexToSortedLoadIndex[Request.LoadRequestIndex];
 		if (SortedLoadIndex != 0xffff)
 		{
