Uses view state's frame index for RTGI so the denoise stands a chance to reconstruct ray direction

diff --git a/Engine/Source/Runtime/Renderer/Private/RayTracing/RayTracingGlobalIllumination.cpp b/Engine/Source/Runtime/Renderer/Private/RayTracing/RayTracingGlobalIllumination.cpp
--- a/Engine/Source/Runtime/Renderer/Private/RayTracing/RayTracingGlobalIllumination.cpp
+++ b/Engine/Source/Runtime/Renderer/Private/RayTracing/RayTracingGlobalIllumination.cpp
@@ -432,6 +432,8 @@ void FDeferredShadingSceneRenderer::RenderRayTracingGlobalIllumination(
 {	
 	FSceneRenderTargets& SceneContext = FSceneRenderTargets::Get(RHICmdList);
 
+	RDG_EVENT_SCOPE(GraphBuilder, "RTGI");
+
 	int32 RayTracingGISamplesPerPixel = GRayTracingGlobalIlluminationSamplesPerPixel > -1 ? GRayTracingGlobalIlluminationSamplesPerPixel : View.FinalPostProcessSettings.RayTracingGISamplesPerPixel;
 	RayTracingConfig.RayCountPerPixel = RayTracingGISamplesPerPixel;
 
@@ -439,15 +441,13 @@ void FDeferredShadingSceneRenderer::RenderRayTracingGlobalIllumination(
 	SetupSceneTextureUniformParameters(SceneContext, FeatureLevel, ESceneTextureSetupMode::All, SceneTextures);
 	// Ray generation
 	{
-		static uint32 Iteration = 0;
 		uint32 IterationCount = RayTracingGISamplesPerPixel;
 		uint32 SequenceCount = 1;
 		uint32 DimensionCount = 24;
-		FHaltonSequenceIteration HaltonSequenceIteration(Scene->HaltonSequence, IterationCount, SequenceCount, DimensionCount, Iteration);
+		FHaltonSequenceIteration HaltonSequenceIteration(Scene->HaltonSequence, IterationCount, SequenceCount, DimensionCount, View.ViewState ? (View.ViewState->FrameIndex % 1024) : 0);
 
 		FHaltonIteration HaltonIteration;
 		InitializeHaltonSequenceIteration(HaltonSequenceIteration, HaltonIteration);
-		Iteration = (Iteration + 1) % 1024;
 
 		FHaltonPrimes HaltonPrimes;
 		InitializeHaltonPrimes(Scene->HaltonPrimesResource, HaltonPrimes);
