* Add support for static textured area light in reflection capture
* Change are light to be additive instead of lerp

Add a new materials default textured emissive material. Its additive behavior makes it works properly with reflection capture fading, without creating artifact.

diff --git a/Engine/Config/BaseEngine.ini b/Engine/Config/BaseEngine.ini
--- a/Engine/Config/BaseEngine.ini
+++ b/Engine/Config/BaseEngine.ini
@@ -88,6 +88,7 @@ DefaultBokehTextureName=/Engine/EngineMaterials/DefaultBokeh.DefaultBokeh
 DefaultBloomKernelTextureName=/Engine/EngineMaterials/DefaultBloomKernel.DefaultBloomKernel
 GeomMaterialName=/Engine/EngineDebugMaterials/GeomMaterial.GeomMaterial
 DebugMeshMaterialName=/Engine/EngineDebugMaterials/DebugMeshMaterial.DebugMeshMaterial
+EmissiveMeshMaterialName=/Engine/EngineMaterials/EmissiveMeshMaterial.EmissiveMeshMaterial
 PreIntegratedSkinBRDFTextureName=/Engine/EngineMaterials/PreintegratedSkinBRDF.PreintegratedSkinBRDF
 BlueNoiseTextureName=/Engine/EngineMaterials/BlueNoise.BlueNoise
 MiniFontTextureName=/Engine/EngineMaterials/MiniFont.MiniFont
@@ -713,6 +714,7 @@ bFullyCompressStartupPackages=false
 +Package=/Engine/EngineMaterials/DefaultPhysicalMaterial
 +Package=/Engine/EngineMaterials/DefaultWhiteGrid
 +Package=/Engine/EngineMaterials/EditorBrushMaterial
++Package=/Engine/EngineMaterials/EmissiveMeshMaterial
 +Package=/Engine/EngineMaterials/Good64x64TilingNoiseHighFreq
 +Package=/Engine/EngineMaterials/Grid
 +Package=/Engine/EngineMaterials/Grid_N
diff --git a/Engine/Source/Runtime/Engine/Classes/Engine/Engine.h b/Engine/Source/Runtime/Engine/Classes/Engine/Engine.h
--- a/Engine/Source/Runtime/Engine/Classes/Engine/Engine.h
+++ b/Engine/Source/Runtime/Engine/Classes/Engine/Engine.h
@@ -884,10 +884,18 @@ public:
 	UPROPERTY()
 	class UMaterial* DebugMeshMaterial;
 
-	/** @todo document */
+	/** Path of the default material for debug mesh */
 	UPROPERTY(globalconfig)
 	FSoftObjectPath DebugMeshMaterialName;
 
+	/** A material used to render emissive meshes (e.g. light source surface). */
+	UPROPERTY()
+	class UMaterial* EmissiveMeshMaterial;
+
+	/** Path of the default material for emissive mesh */
+	UPROPERTY(globalconfig)
+	FSoftObjectPath EmissiveMeshMaterialName;
+
 	/** Material used for visualizing level membership in lit view port modes. */
 	UPROPERTY()
 	class UMaterial* LevelColorationLitMaterial;
diff --git a/Engine/Source/Runtime/Engine/Private/Materials/MaterialShared.cpp b/Engine/Source/Runtime/Engine/Private/Materials/MaterialShared.cpp
--- a/Engine/Source/Runtime/Engine/Private/Materials/MaterialShared.cpp
+++ b/Engine/Source/Runtime/Engine/Private/Materials/MaterialShared.cpp
@@ -2584,6 +2584,23 @@ bool FColoredMaterialRenderProxy::GetTextureValue(const FMaterialParameterInfo&
 	return Parent->GetTextureValue(ParameterInfo,OutValue,Context);
 }
 
+/*-----------------------------------------------------------------------------
+	FColoredTexturedMaterialRenderProxy
+-----------------------------------------------------------------------------*/
+
+bool FColoredTexturedMaterialRenderProxy::GetTextureValue(const FMaterialParameterInfo& ParameterInfo, const UTexture** OutValue, const FMaterialRenderContext& Context) const
+{
+	if (ParameterInfo.Name == TextureParamName)
+	{
+		*OutValue = Texture;
+		return true;
+	}
+	else
+	{
+		return Parent->GetTextureValue(ParameterInfo, OutValue, Context);
+	}
+}
+
 /*-----------------------------------------------------------------------------
 	FOverrideSelectionColorMaterialRenderProxy
 -----------------------------------------------------------------------------*/
diff --git a/Engine/Source/Runtime/Engine/Private/UnrealEngine.cpp b/Engine/Source/Runtime/Engine/Private/UnrealEngine.cpp
--- a/Engine/Source/Runtime/Engine/Private/UnrealEngine.cpp
+++ b/Engine/Source/Runtime/Engine/Private/UnrealEngine.cpp
@@ -2344,6 +2344,7 @@ void UEngine::InitializeObjectReferences()
 
 	// these one's are needed both editor and standalone 
 	LoadSpecialMaterial(TEXT("DebugMeshMaterialName"), DebugMeshMaterialName.ToString(), DebugMeshMaterial, false);
+	LoadSpecialMaterial(TEXT("EmissiveMeshMaterialName"), EmissiveMeshMaterialName.ToString(), EmissiveMeshMaterial, false);
 	LoadSpecialMaterial(TEXT("InvalidLightmapSettingsMaterialName"), InvalidLightmapSettingsMaterialName.ToString(), InvalidLightmapSettingsMaterial, false);
 	LoadSpecialMaterial(TEXT("ArrowMaterialName"), ArrowMaterialName.ToString(), ArrowMaterial, false);
 
diff --git a/Engine/Source/Runtime/Engine/Public/MaterialShared.h b/Engine/Source/Runtime/Engine/Public/MaterialShared.h
--- a/Engine/Source/Runtime/Engine/Public/MaterialShared.h
+++ b/Engine/Source/Runtime/Engine/Public/MaterialShared.h
@@ -1984,6 +1984,28 @@ public:
 	ENGINE_API virtual bool GetTextureValue(const FMaterialParameterInfo& ParameterInfo,const UTexture** OutValue, const FMaterialRenderContext& Context) const;
 };
 
+
+/**
+ * An material render proxy which overrides the material's Color vector and Texture parameter (mixed together).
+ */
+class FColoredTexturedMaterialRenderProxy : public FColoredMaterialRenderProxy
+{
+public:
+
+	const UTexture* Texture;
+	FName TextureParamName;
+
+	/** Initialization constructor. */
+	FColoredTexturedMaterialRenderProxy(const FMaterialRenderProxy* InParent, const FLinearColor& InColor, FName InColorParamName, const UTexture* InTexture, FName InTextureParamName) :
+		FColoredMaterialRenderProxy(InParent, InColor, InColorParamName),
+		Texture(InTexture),
+		TextureParamName(InTextureParamName)
+	{}
+
+	ENGINE_API virtual bool GetTextureValue(const FMaterialParameterInfo& ParameterInfo, const UTexture** OutValue, const FMaterialRenderContext& Context) const;
+};
+
+
 /**
  * A material render proxy which overrides the selection color
  */
diff --git a/Engine/Source/Runtime/Renderer/Private/SceneVisibility.cpp b/Engine/Source/Runtime/Renderer/Private/SceneVisibility.cpp
--- a/Engine/Source/Runtime/Renderer/Private/SceneVisibility.cpp
+++ b/Engine/Source/Runtime/Renderer/Private/SceneVisibility.cpp
@@ -36,6 +36,7 @@
 #include "GPUScene.h"
 #include "TranslucentRendering.h"
 #include "Async/ParallelFor.h"
+#include "RectLightSceneProxy.h"
 
 /*------------------------------------------------------------------------------
 	Globals
@@ -3964,7 +3965,8 @@ void FSceneRenderer::PostVisibilityFrameSetup(FILCUpdatePrimTaskData& OutILCTask
 					Origin = ToLight + View.ViewMatrices.GetViewOrigin();
 				
 					FLinearColor Color( LightParameters.Color.X, LightParameters.Color.Y, LightParameters.Color.Z, LightParameters.FalloffExponent );
-					if( !Proxy->IsRectLight() )
+					const bool bIsRectLight = Proxy->IsRectLight();
+					if( !bIsRectLight )
 					{
 						const float SphereArea = (4.0f * PI) * FMath::Square( LightParameters.SourceRadius );
 						const float CylinderArea = (2.0f * PI) * LightParameters.SourceRadius * LightParameters.SourceLength;
@@ -3993,17 +3995,28 @@ void FSceneRenderer::PostVisibilityFrameSetup(FILCUpdatePrimTaskData& OutILCTask
 					Color.A *= LightParameters.SpecularScale;
 
 					// Rect is one sided
-					if( Proxy->IsRectLight() && (L | LightParameters.Direction) < 0.0f )
+					if( bIsRectLight && (L | LightParameters.Direction) < 0.0f )
 						continue;
-				
-					FMaterialRenderProxy* const ColoredMeshInstance = new(FMemStack::Get()) FColoredMaterialRenderProxy( GEngine->DebugMeshMaterial->GetRenderProxy(), Color );
+
+					UTexture* SurfaceTexture = nullptr;
+					if (bIsRectLight)
+					{
+						const FRectLightSceneProxy* RectLightProxy = (const FRectLightSceneProxy*)Proxy;
+						SurfaceTexture = RectLightProxy->SourceTexture;
+					}
+					
+					FMaterialRenderProxy* ColoredMeshInstance = nullptr;
+					if (SurfaceTexture)
+						ColoredMeshInstance = new(FMemStack::Get()) FColoredTexturedMaterialRenderProxy(GEngine->EmissiveMeshMaterial->GetRenderProxy(), Color, NAME_Color, SurfaceTexture, NAME_LinearColor);
+					else
+						ColoredMeshInstance = new(FMemStack::Get()) FColoredMaterialRenderProxy(GEngine->EmissiveMeshMaterial->GetRenderProxy(), Color, NAME_Color);
 
 					FMatrix LightToWorld = Proxy->GetLightToWorld();
 					LightToWorld.RemoveScaling();
 
 					FViewElementPDI LightPDI( &View, NULL, &View.DynamicPrimitiveShaderData );
 
-					if( Proxy->IsRectLight() )
+					if( bIsRectLight )
 					{
 						DrawBox( &LightPDI, LightToWorld, FVector( 0.0f, LightParameters.SourceRadius, LightParameters.SourceLength ), ColoredMeshInstance, SDPG_World );
 					}
