Fixed missing header include

diff --git a/Engine/Source/Runtime/Renderer/Private/ScreenSpaceDenoise.cpp b/Engine/Source/Runtime/Renderer/Private/ScreenSpaceDenoise.cpp
--- a/Engine/Source/Runtime/Renderer/Private/ScreenSpaceDenoise.cpp
+++ b/Engine/Source/Runtime/Renderer/Private/ScreenSpaceDenoise.cpp
@@ -14,6 +14,7 @@
 #include "PipelineStateCache.h"
 #include "SceneViewFamilyBlackboard.h"
 #include "BlueNoise.h"
+#include "Halton.h"
 
 
 // ---------------------------------------------------- Cvars
