Removed usage of SetRenderTarget in VisualizeTexturePresent.cpp

diff --git a/Engine/Source/Runtime/Renderer/Private/VisualizeTexturePresent.cpp b/Engine/Source/Runtime/Renderer/Private/VisualizeTexturePresent.cpp
--- a/Engine/Source/Runtime/Renderer/Private/VisualizeTexturePresent.cpp
+++ b/Engine/Source/Runtime/Renderer/Private/VisualizeTexturePresent.cpp
@@ -453,9 +453,8 @@ void FVisualizeTexturePresent::PresentContent(FRHICommandListImmediate& RHICmdLi
 	}
 
 	auto& RenderTarget = View.Family->RenderTarget->GetRenderTargetTexture();
-	PRAGMA_DISABLE_DEPRECATION_WARNINGS
-	SetRenderTarget(RHICmdList, RenderTarget, FTextureRHIRef(), true);
-	PRAGMA_ENABLE_DEPRECATION_WARNINGS
+	FRHIRenderPassInfo RPInfo(RenderTarget, ERenderTargetActions::Load_Store);
+	RHICmdList.BeginRenderPass(RPInfo, TEXT("VisualizeTexture"));
 	RHICmdList.SetViewport(DestRect.Min.X, DestRect.Min.Y, 0.0f, DestRect.Max.X, DestRect.Max.Y, 1.0f);
 
 	FGraphicsPipelineStateInitializer GraphicsPSOInit;
@@ -495,6 +494,7 @@ void FVisualizeTexturePresent::PresentContent(FRHICommandListImmediate& RHICmdLi
 			*VertexShader,
 			EDRF_Default);
 	}
+	RHICmdList.EndRenderPass();
 
 	FRenderTargetTemp TempRenderTarget(View, View.UnconstrainedViewRect.Size());
 	FCanvas Canvas(&TempRenderTarget, NULL, View.Family->CurrentRealTime, View.Family->CurrentWorldTime, View.Family->DeltaWorldTime, View.GetFeatureLevel());
