DR - Redo Occlusion Query pools. Issue was that the original pool only released the query when it's refcount was 1. The RHI pool cannot have this therefore doing a (non atomic) in the Occlusion Query pool code to compensate for this.

diff --git a/Engine/Source/Runtime/Engine/Private/SceneUtils.cpp b/Engine/Source/Runtime/Engine/Private/SceneUtils.cpp
--- a/Engine/Source/Runtime/Engine/Private/SceneUtils.cpp
+++ b/Engine/Source/Runtime/Engine/Private/SceneUtils.cpp
@@ -70,49 +70,4 @@ ENGINE_API bool IsMobileColorsRGB()
 	const bool bMobileUseHWsRGBEncoding = (MobileUseHWsRGBEncodingCVAR && MobileUseHWsRGBEncodingCVAR->GetValueOnAnyThread() == 1);
 
 	return !IsMobileHDR() && bMobileUseHWsRGBEncoding;
-}
-
-/*-----------------------------------------------------------------------------
-FRenderQueryPool
------------------------------------------------------------------------------*/
-FRenderQueryPool::~FRenderQueryPool()
-{
-	Release();
-}
-
-void FRenderQueryPool::Release()
-{
-	Queries.Empty();
-	NumQueriesAllocated = 0;
-}
-
-FRenderQueryRHIRef FRenderQueryPool::AllocateQuery()
-{
-	// Are we out of available render queries?
-	NumQueriesAllocated++;
-	if (Queries.Num() == 0)
-	{
-		// Create a new render query.
-		return RHICreateRenderQuery(QueryType);
-	}
-
-	return Queries.Pop(/*bAllowShrinking=*/ false);
-}
-
-void FRenderQueryPool::ReleaseQuery(FRenderQueryRHIRef &Query)
-{
-	if (IsValidRef(Query))
-	{
-		NumQueriesAllocated--;
-
-		// Is no one else keeping a refcount to the query?
-		if (Query.GetRefCount() == 1)
-		{
-			// Return it to the pool.
-			Queries.Add(Query);
-		}
-
-		// De-ref without deleting.
-		Query = NULL;
-	}
 }
\ No newline at end of file
diff --git a/Engine/Source/Runtime/Engine/Public/SceneUtils.h b/Engine/Source/Runtime/Engine/Public/SceneUtils.h
--- a/Engine/Source/Runtime/Engine/Public/SceneUtils.h
+++ b/Engine/Source/Runtime/Engine/Public/SceneUtils.h
@@ -42,40 +42,6 @@ ENGINE_API EMobileHDRMode GetMobileHDRMode();
 
 ENGINE_API bool IsMobileColorsRGB();
 
-/**
-* A pool of render (e.g. occlusion/timer) queries which are allocated individually, and returned to the pool as a group.
-*/
-class ENGINE_API FRenderQueryPool
-{
-public:
-	FRenderQueryPool(ERenderQueryType InQueryType)
-		: QueryType(InQueryType)
-		, NumQueriesAllocated(0)
-	{ }
-
-	virtual ~FRenderQueryPool();
-
-	/** Releases all the render queries in the pool. */
-	void Release();
-
-	/** Allocates an render query from the pool. */
-	FRenderQueryRHIRef AllocateQuery();
-
-	/** De-reference an render query, returning it to the pool instead of deleting it when the refcount reaches 0. */
-	void ReleaseQuery(FRenderQueryRHIRef &Query);
-
-	/** Returns the number of currently allocated queries. This is not necessarily the same as the pool size */
-	int32 GetAllocatedQueryCount() const { return NumQueriesAllocated;  }
-
-private:
-	/** Container for available render queries. */
-	TArray<FRenderQueryRHIRef> Queries;
-
-	ERenderQueryType QueryType;
-
-	int32 NumQueriesAllocated;
-};
-
 // Callback for calling one action (typical use case: delay a clear until it's actually needed)
 class FDelayedRendererAction
 {
diff --git a/Engine/Source/Runtime/RHI/Private/DynamicRHI.cpp b/Engine/Source/Runtime/RHI/Private/DynamicRHI.cpp
--- a/Engine/Source/Runtime/RHI/Private/DynamicRHI.cpp
+++ b/Engine/Source/Runtime/RHI/Private/DynamicRHI.cpp
@@ -384,6 +384,11 @@ FRHIPooledRenderQuery FDefaultRHIRenderQueryPool::AllocateQuery()
 
 void FDefaultRHIRenderQueryPool::ReleaseQuery(TRefCountPtr<FRHIRenderQuery>&& Query)
 {
+	if (QueryType == ERenderQueryType::RQT_Occlusion)
+	{
+		static int dbg = 0;
+		dbg++;
+	}
 	check(IsInRenderingThread());
 	//Hard to validate because of Resource resurrection, better to remove GetQueryRef entirely
 	//checkf(Query.IsValid() && Query.GetRefCount() <= 2, TEXT("Query has been released but reference still held: use FRHIPooledRenderQuery::GetQueryRef() with extreme caution"));
diff --git a/Engine/Source/Runtime/RHI/Public/RHICommandList.h b/Engine/Source/Runtime/RHI/Public/RHICommandList.h
--- a/Engine/Source/Runtime/RHI/Public/RHICommandList.h
+++ b/Engine/Source/Runtime/RHI/Public/RHICommandList.h
@@ -4062,12 +4062,14 @@ public:
 		GDynamicRHI->RHIRead3DSurfaceFloatData(Texture,Rect,ZMinMax,OutData);
 	}
 	
+	UE_DEPRECATED(4.23, "CreateRenderQuery API is deprecated; use RHICreateRenderQueryPool and suballocate queries there")
 	FORCEINLINE FRenderQueryRHIRef CreateRenderQuery(ERenderQueryType QueryType)
 	{
 		FScopedRHIThreadStaller StallRHIThread(*this);
 		return GDynamicRHI->RHICreateRenderQuery(QueryType);
 	}
 
+	UE_DEPRECATED(4.23, "CreateRenderQuery API is deprecated; use RHICreateRenderQueryPool and suballocate queries there")
 	FORCEINLINE FRenderQueryRHIRef CreateRenderQuery_RenderThread(ERenderQueryType QueryType)
 	{
 		return GDynamicRHI->RHICreateRenderQuery_RenderThread(*this, QueryType);
@@ -4805,9 +4807,12 @@ FORCEINLINE void RHIUnlockTextureCubeFace(FTextureCubeRHIParamRef Texture, uint3
 	 FRHICommandListExecutor::GetImmediateCommandList().UnlockTextureCubeFace(Texture, FaceIndex, ArrayIndex, MipIndex, bLockWithinMiptail);
 }
 
+UE_DEPRECATED(4.23, "CreateRenderQuery API is deprecated; use RHICreateRenderQueryPool and suballocate queries there")
 FORCEINLINE FRenderQueryRHIRef RHICreateRenderQuery(ERenderQueryType QueryType)
 {
+	PRAGMA_DISABLE_DEPRECATION_WARNINGS
 	return FRHICommandListExecutor::GetImmediateCommandList().CreateRenderQuery_RenderThread(QueryType);
+	PRAGMA_ENABLE_DEPRECATION_WARNINGS
 }
 
 FORCEINLINE void RHIAcquireTransientResource(FTextureRHIParamRef Resource)
diff --git a/Engine/Source/Runtime/RHI/Public/RHIResources.h b/Engine/Source/Runtime/RHI/Public/RHIResources.h
--- a/Engine/Source/Runtime/RHI/Public/RHIResources.h
+++ b/Engine/Source/Runtime/RHI/Public/RHIResources.h
@@ -952,12 +952,12 @@ public:
 	FRHIPooledRenderQuery(FRHIPooledRenderQuery&&) = default;
 	FRHIPooledRenderQuery& operator=(FRHIPooledRenderQuery&&) = default;
 
-	FRHIRenderQuery* GetQuery() const
+	bool IsValid() const
 	{
-		return Query;
+		return Query.IsValid();
 	}
 
-	TRefCountPtr<FRHIRenderQuery> GetQueryRef() const
+	FRHIRenderQuery* GetQuery() const
 	{
 		return Query;
 	}
diff --git a/Engine/Source/Runtime/Renderer/Private/RendererScene.cpp b/Engine/Source/Runtime/Renderer/Private/RendererScene.cpp
--- a/Engine/Source/Runtime/Renderer/Private/RendererScene.cpp
+++ b/Engine/Source/Runtime/Renderer/Private/RendererScene.cpp
@@ -126,7 +126,7 @@ struct FSpeedTreeWindComputation
 
 /** Default constructor. */
 FSceneViewState::FSceneViewState()
-	: OcclusionQueryPool(RQT_Occlusion)
+	: OcclusionQueryPool(RHICreateRenderQueryPool(RQT_Occlusion))
 	, TimerQueryPool(RHICreateRenderQueryPool(RQT_AbsoluteTime, FLatentGPUTimer::NumBufferedFrames * 2 * 2 * 2))
 	, TranslucencyTimer(TimerQueryPool)
 	, SeparateTranslucencyTimer(TimerQueryPool)
diff --git a/Engine/Source/Runtime/Renderer/Private/SceneOcclusion.cpp b/Engine/Source/Runtime/Renderer/Private/SceneOcclusion.cpp
--- a/Engine/Source/Runtime/Renderer/Private/SceneOcclusion.cpp
+++ b/Engine/Source/Runtime/Renderer/Private/SceneOcclusion.cpp
@@ -225,7 +225,7 @@ void FSceneViewState::TrimOcclusionHistory(float CurrentTime, float MinHistoryTi
 			// If the primitive has an old pending occlusion query, release it.
 			if(PrimitiveIt->LastConsideredTime < MinQueryTime)
 			{
-				PrimitiveIt->ReleaseStaleQueries(OcclusionQueryPool, FrameNumber, NumBufferedFrames);
+				PrimitiveIt->ReleaseStaleQueries(FrameNumber, NumBufferedFrames);
 			}
 
 			// If the primitive hasn't been considered for visibility recently, remove its history from the set.
@@ -243,14 +243,14 @@ bool FSceneViewState::IsShadowOccluded(FRHICommandListImmediate& RHICmdList, FSc
 	// Get the oldest occlusion query	
 	const uint32 QueryIndex = FOcclusionQueryHelpers::GetQueryLookupIndex(PendingPrevFrameNumber, NumBufferedFrames);
 	const FSceneViewState::ShadowKeyOcclusionQueryMap& ShadowOcclusionQueryMap = ShadowOcclusionQueryMaps[QueryIndex];	
-	const FRenderQueryRHIRef* Query = ShadowOcclusionQueryMap.Find(ShadowKey);
+	const FRHIPooledRenderQuery* Query = ShadowOcclusionQueryMap.Find(ShadowKey);
 
 	// Read the occlusion query results.
 	uint64 NumSamples = 0;
 	// Only block on the query if not running SLI
 	const bool bWaitOnQuery = GNumAlternateFrameRenderingGroups == 1;
 
-	if (Query && RHICmdList.GetRenderQueryResult(*Query, NumSamples, bWaitOnQuery))
+	if (Query && RHICmdList.GetRenderQueryResult(Query->GetQuery(), NumSamples, bWaitOnQuery))
 	{
 		// If the shadow's occlusion query didn't have any pixels visible the previous frame, it's occluded.
 		return NumSamples == 0;
@@ -337,7 +337,7 @@ FOcclusionQueryBatcher::FOcclusionQueryBatcher(class FSceneViewState* ViewState,
 :	CurrentBatchOcclusionQuery(NULL)
 ,	MaxBatchedPrimitives(InMaxBatchedPrimitives)
 ,	NumBatchedPrimitives(0)
-,	OcclusionQueryPool(ViewState ? &ViewState->OcclusionQueryPool : NULL)
+,	OcclusionQueryPool(ViewState ? ViewState->OcclusionQueryPool : NULL)
 {}
 
 FOcclusionQueryBatcher::~FOcclusionQueryBatcher()
@@ -358,7 +358,7 @@ void FOcclusionQueryBatcher::Flush(FRHICommandList& RHICmdList)
 		for(int32 BatchIndex = 0, NumBatches = BatchOcclusionQueries.Num();BatchIndex < NumBatches;BatchIndex++)
 		{
 			FOcclusionBatch& Batch = BatchOcclusionQueries[BatchIndex];
-			FRenderQueryRHIParamRef BatchOcclusionQuery = Batch.Query;
+			FRHIRenderQuery* BatchOcclusionQuery = Batch.Query.GetQuery();
 			FVertexBufferRHIParamRef VertexBufferRHI = Batch.VertexAllocation.VertexBuffer->VertexBufferRHI;
 			uint32 VertexBufferOffset = Batch.VertexAllocation.VertexOffset;
 			const int32 NumPrimitivesThisBatch = (BatchIndex != (NumBatches-1)) ? MaxBatchedPrimitives : NumBatchedPrimitives;
@@ -384,14 +384,14 @@ void FOcclusionQueryBatcher::Flush(FRHICommandList& RHICmdList)
 	}
 }
 
-FRenderQueryRHIParamRef FOcclusionQueryBatcher::BatchPrimitive(const FVector& BoundsOrigin,const FVector& BoundsBoxExtent, FGlobalDynamicVertexBuffer& DynamicVertexBuffer)
+FRefCountedRHIPooledRenderQuery FOcclusionQueryBatcher::BatchPrimitive(const FVector& BoundsOrigin,const FVector& BoundsBoxExtent, FGlobalDynamicVertexBuffer& DynamicVertexBuffer)
 {
 	// Check if the current batch is full.
 	if(CurrentBatchOcclusionQuery == NULL || NumBatchedPrimitives >= MaxBatchedPrimitives)
 	{
 		check(OcclusionQueryPool);
 		CurrentBatchOcclusionQuery = new(BatchOcclusionQueries) FOcclusionBatch;
-		CurrentBatchOcclusionQuery->Query = OcclusionQueryPool->AllocateQuery();
+		CurrentBatchOcclusionQuery->Query = FRefCountedRHIPooledRenderQuery(OcclusionQueryPool->AllocateQuery());
 		CurrentBatchOcclusionQuery->VertexAllocation = DynamicVertexBuffer.Allocate(MaxBatchedPrimitives * 8 * sizeof(FVector));
 		check(CurrentBatchOcclusionQuery->VertexAllocation.IsValid());
 		NumBatchedPrimitives = 0;
@@ -430,7 +430,7 @@ static bool AllocateProjectedShadowOcclusionQuery(
 	const FProjectedShadowInfo& ProjectedShadowInfo, 
 	int32 NumBufferedFrames, 
 	EShadowOcclusionQueryIntersectionMode IntersectionMode,
-	FRenderQueryRHIRef& ShadowOcclusionQuery)
+	FRHIRenderQuery*& ShadowOcclusionQuery)
 {
 	bool bIssueQuery = true;
 
@@ -465,21 +465,21 @@ static bool AllocateProjectedShadowOcclusionQuery(
 		FSceneViewState* ViewState = (FSceneViewState*)View.State;
 
 		// Allocate an occlusion query for the primitive from the occlusion query pool.
-		ShadowOcclusionQuery = ViewState->OcclusionQueryPool.AllocateQuery();
-
 		FSceneViewState::FProjectedShadowKey Key(ProjectedShadowInfo);
 		const uint32 QueryIndex = FOcclusionQueryHelpers::GetQueryIssueIndex(ViewState->PendingPrevFrameNumber, NumBufferedFrames);
 		FSceneViewState::ShadowKeyOcclusionQueryMap& ShadowOcclusionQueryMap = ViewState->ShadowOcclusionQueryMaps[QueryIndex];
 
 		checkSlow(ShadowOcclusionQueryMap.Find(Key) == NULL);
-		ShadowOcclusionQueryMap.Add(Key, ShadowOcclusionQuery);
+		FRHIPooledRenderQuery PooledShadowOcclusionQuery = ViewState->OcclusionQueryPool->AllocateQuery();
+		ShadowOcclusionQuery = PooledShadowOcclusionQuery.GetQuery();
+		ShadowOcclusionQueryMap.Add(Key, MoveTemp(PooledShadowOcclusionQuery));
 	}
 	
 	return bIssueQuery;
 }
 
 
-static void ExecutePointLightShadowOcclusionQuery(FRHICommandList& RHICmdList, FViewInfo& View, const FProjectedShadowInfo& ProjectedShadowInfo, FOcclusionQueryVS* VertexShader, FRenderQueryRHIRef ShadowOcclusionQuery)
+static void ExecutePointLightShadowOcclusionQuery(FRHICommandList& RHICmdList, FViewInfo& View, const FProjectedShadowInfo& ProjectedShadowInfo, FOcclusionQueryVS* VertexShader, FRHIRenderQuery* ShadowOcclusionQuery)
 {
 	FLightSceneProxy& LightProxy = *(ProjectedShadowInfo.GetLightSceneInfo().Proxy);
 	
@@ -528,7 +528,7 @@ static void PrepareDirectionalLightShadowOcclusionQuery(uint32& BaseVertexIndex,
 	BaseVertexIndex += 6;
 }
 
-static void ExecuteDirectionalLightShadowOcclusionQuery(FRHICommandList& RHICmdList, uint32& BaseVertexIndex, FRenderQueryRHIRef ShadowOcclusionQuery)
+static void ExecuteDirectionalLightShadowOcclusionQuery(FRHICommandList& RHICmdList, uint32& BaseVertexIndex, FRHIRenderQuery* ShadowOcclusionQuery)
 {
 	RHICmdList.BeginRenderQuery(ShadowOcclusionQuery);
 
@@ -568,7 +568,7 @@ static void PrepareProjectedShadowOcclusionQuery(uint32& BaseVertexIndex, FVecto
 	BaseVertexIndex += 8;
 }
 
-static void ExecuteProjectedShadowOcclusionQuery(FRHICommandList& RHICmdList, uint32& BaseVertexIndex, FRenderQueryRHIRef ShadowOcclusionQuery)
+static void ExecuteProjectedShadowOcclusionQuery(FRHICommandList& RHICmdList, uint32& BaseVertexIndex, FRHIRenderQuery* ShadowOcclusionQuery)
 {	
 	// Draw the primitive's bounding box, using the occlusion query.
 	RHICmdList.BeginRenderQuery(ShadowOcclusionQuery);
@@ -579,7 +579,7 @@ static void ExecuteProjectedShadowOcclusionQuery(FRHICommandList& RHICmdList, ui
 	RHICmdList.EndRenderQuery(ShadowOcclusionQuery);
 }
 
-static bool AllocatePlanarReflectionOcclusionQuery(const FViewInfo& View, const FPlanarReflectionSceneProxy* SceneProxy, int32 NumBufferedFrames, FRenderQueryRHIRef& OcclusionQuery)
+static bool AllocatePlanarReflectionOcclusionQuery(const FViewInfo& View, const FPlanarReflectionSceneProxy* SceneProxy, int32 NumBufferedFrames, FRHIRenderQuery*& OcclusionQuery)
 {
 	FSceneViewState* ViewState = (FSceneViewState*)View.State;
 	
@@ -609,18 +609,19 @@ static bool AllocatePlanarReflectionOcclusionQuery(const FViewInfo& View, const
 	
 	uint32 OcclusionFrameCounter = ViewState->OcclusionFrameCounter;
 	FIndividualOcclusionHistory& OcclusionHistory = ViewState->PlanarReflectionOcclusionHistories.FindOrAdd(SceneProxy->PlanarReflectionId);
-	OcclusionHistory.ReleaseQuery(ViewState->OcclusionQueryPool, OcclusionFrameCounter, NumBufferedFrames);
+	OcclusionHistory.ReleaseQuery(OcclusionFrameCounter, NumBufferedFrames);
 	
 	if (bAllowBoundsTest)
 	{
 		// Allocate an occlusion query for the primitive from the occlusion query pool.
-		OcclusionQuery = ViewState->OcclusionQueryPool.AllocateQuery();
-		
-		OcclusionHistory.SetCurrentQuery(OcclusionFrameCounter, OcclusionQuery, NumBufferedFrames);
+		FRHIPooledRenderQuery PooledOcclusionQuery = ViewState->OcclusionQueryPool->AllocateQuery();
+		OcclusionQuery = PooledOcclusionQuery.GetQuery();
+
+		OcclusionHistory.SetCurrentQuery(OcclusionFrameCounter, MoveTemp(PooledOcclusionQuery), NumBufferedFrames);
 	}
 	else
 	{
-		OcclusionHistory.SetCurrentQuery(OcclusionFrameCounter, NULL, NumBufferedFrames);
+		OcclusionHistory.SetCurrentQuery(OcclusionFrameCounter, FRHIPooledRenderQuery(), NumBufferedFrames);
 	}
 	
 	return bAllowBoundsTest;
@@ -644,7 +645,7 @@ static void PreparePlanarReflectionOcclusionQuery(uint32& BaseVertexIndex, FVect
 	BaseVertexIndex += 8;
 }
 
-static void ExecutePlanarReflectionOcclusionQuery(FRHICommandList& RHICmdList, uint32& BaseVertexIndex, FRenderQueryRHIRef OcclusionQuery)
+static void ExecutePlanarReflectionOcclusionQuery(FRHICommandList& RHICmdList, uint32& BaseVertexIndex, FRHIRenderQuery* OcclusionQuery)
 {
 	// Draw the primitive's bounding box, using the occlusion query.
 	RHICmdList.BeginRenderQuery(OcclusionQuery);
@@ -1260,10 +1261,10 @@ struct FViewOcclusionQueries
 	TArray<FProjectedShadowInfo const*> ShadowQuerieInfos;
 	TArray<FPlanarReflectionSceneProxy const*> ReflectionQuerieInfos;
 
-	TArray<FRenderQueryRHIRef> PointLightQueries;
-	TArray<FRenderQueryRHIRef> CSMQueries;
-	TArray<FRenderQueryRHIRef> ShadowQueries;
-	TArray<FRenderQueryRHIRef> ReflectionQueries;
+	TArray<FRHIRenderQuery*> PointLightQueries;
+	TArray<FRHIRenderQuery*> CSMQueries;
+	TArray<FRHIRenderQuery*> ShadowQueries;
+	TArray<FRHIRenderQuery*> ReflectionQueries;
 };
 
 void FSceneRenderer::BeginOcclusionTests(FRHICommandListImmediate& RHICmdList, bool bRenderQueries)
@@ -1299,12 +1300,6 @@ void FSceneRenderer::BeginOcclusionTests(FRHICommandListImmediate& RHICmdList, b
 				ViewState->TrimOcclusionHistory(ViewFamily.CurrentRealTime, ViewFamily.CurrentRealTime - GEngine->PrimitiveProbablyVisibleTime, ViewFamily.CurrentRealTime, ViewState->OcclusionFrameCounter);
 
 				// Give back all these occlusion queries to the pool.
-				for (TMap<FSceneViewState::FProjectedShadowKey, FRenderQueryRHIRef>::TIterator QueryIt(ShadowOcclusionQueryMap); QueryIt; ++QueryIt)
-				{
-					//FRenderQueryRHIParamRef Query = QueryIt.Value();
-					//check( Query.GetRefCount() == 1 );
-					ViewState->OcclusionQueryPool.ReleaseQuery(QueryIt.Value());
-				}
 				ShadowOcclusionQueryMap.Reset();
 				
 				if (FeatureLevel > ERHIFeatureLevel::ES3_1)
@@ -1332,7 +1327,7 @@ void FSceneRenderer::BeginOcclusionTests(FRHICommandListImmediate& RHICmdList, b
 
 							if (ProjectedShadowInfo.bOnePassPointLightShadow)
 							{
-								FRenderQueryRHIRef ShadowOcclusionQuery;
+								FRHIRenderQuery* ShadowOcclusionQuery;
 								if (AllocateProjectedShadowOcclusionQuery(View, ProjectedShadowInfo, NumBufferedFrames, SOQ_LightInfluenceSphere, ShadowOcclusionQuery))
 								{
 									ViewQuery.PointLightQuerieInfos.Add(&ProjectedShadowInfo);
@@ -1346,7 +1341,7 @@ void FSceneRenderer::BeginOcclusionTests(FRHICommandListImmediate& RHICmdList, b
 								// Don't query the first cascade, it is always visible
 								if (GOcclusionCullCascadedShadowMaps && ProjectedShadowInfo.CascadeSettings.ShadowSplitIndex > 0)
 								{
-									FRenderQueryRHIRef ShadowOcclusionQuery;
+									FRHIRenderQuery* ShadowOcclusionQuery;
 									if (AllocateProjectedShadowOcclusionQuery(View, ProjectedShadowInfo, NumBufferedFrames, SOQ_None, ShadowOcclusionQuery))
 									{
 										ViewQuery.CSMQuerieInfos.Add(&ProjectedShadowInfo);
@@ -1362,7 +1357,7 @@ void FSceneRenderer::BeginOcclusionTests(FRHICommandListImmediate& RHICmdList, b
 							// Don't query if any subjects are visible because the shadow frustum will be definitely unoccluded
 								&& !ProjectedShadowInfo.SubjectsVisible(View))
 							{
-								FRenderQueryRHIRef ShadowOcclusionQuery;
+								FRHIRenderQuery* ShadowOcclusionQuery;
 								if (AllocateProjectedShadowOcclusionQuery(View, ProjectedShadowInfo, NumBufferedFrames, SOQ_NearPlaneVsShadowFrustum, ShadowOcclusionQuery))
 								{
 									ViewQuery.ShadowQuerieInfos.Add(&ProjectedShadowInfo);
@@ -1377,7 +1372,7 @@ void FSceneRenderer::BeginOcclusionTests(FRHICommandListImmediate& RHICmdList, b
 						for (int32 ShadowIndex = 0; ShadowIndex < VisibleLightInfo.OccludedPerObjectShadows.Num(); ShadowIndex++)
 						{
 							const FProjectedShadowInfo& ProjectedShadowInfo = *VisibleLightInfo.OccludedPerObjectShadows[ShadowIndex];
-							FRenderQueryRHIRef ShadowOcclusionQuery;
+							FRHIRenderQuery* ShadowOcclusionQuery;
 							if (AllocateProjectedShadowOcclusionQuery(View, ProjectedShadowInfo, NumBufferedFrames, SOQ_NearPlaneVsShadowFrustum, ShadowOcclusionQuery))
 							{
 								ViewQuery.ShadowQuerieInfos.Add(&ProjectedShadowInfo);
@@ -1400,7 +1395,7 @@ void FSceneRenderer::BeginOcclusionTests(FRHICommandListImmediate& RHICmdList, b
 					for (int32 ReflectionIndex = 0; ReflectionIndex < Scene->PlanarReflections.Num(); ReflectionIndex++)
 					{
 						FPlanarReflectionSceneProxy* SceneProxy = Scene->PlanarReflections[ReflectionIndex];
-						FRenderQueryRHIRef ShadowOcclusionQuery;
+						FRHIRenderQuery* ShadowOcclusionQuery;
 						if (AllocatePlanarReflectionOcclusionQuery(View, SceneProxy, NumReflectionBufferedFrames, ShadowOcclusionQuery))
 						{
 							ViewQuery.ReflectionQuerieInfos.Add(SceneProxy);
@@ -1558,13 +1553,13 @@ void FSceneRenderer::BeginOcclusionTests(FRHICommandListImmediate& RHICmdList, b
 						RHICmdList.SetStreamSource(0, VertexBufferRHI, 0);
 						BaseVertexOffset = 0;
 
-						for (FRenderQueryRHIRef const& Query : ViewQuery.CSMQueries)
+						for (FRHIRenderQuery* const& Query : ViewQuery.CSMQueries)
 						{
 							ExecuteDirectionalLightShadowOcclusionQuery(RHICmdList, BaseVertexOffset, Query);
 							checkSlow(BaseVertexOffset <= NumVertices);
 						}
 
-						for (FRenderQueryRHIRef const& Query : ViewQuery.ShadowQueries)
+						for (FRHIRenderQuery* const& Query : ViewQuery.ShadowQueries)
 						{
 							ExecuteProjectedShadowOcclusionQuery(RHICmdList, BaseVertexOffset, Query);
 							checkSlow(BaseVertexOffset <= NumVertices);
@@ -1574,7 +1569,7 @@ void FSceneRenderer::BeginOcclusionTests(FRHICommandListImmediate& RHICmdList, b
 					if (FeatureLevel > ERHIFeatureLevel::ES3_1)					
 					{
 						SCOPED_DRAW_EVENT(RHICmdList, PlanarReflectionQueries);
-						for (FRenderQueryRHIRef const& Query : ViewQuery.ReflectionQueries)
+						for (FRHIRenderQuery* const& Query : ViewQuery.ReflectionQueries)
 						{
 							ExecutePlanarReflectionOcclusionQuery(RHICmdList, BaseVertexOffset, Query);
 							check(BaseVertexOffset <= NumVertices);
diff --git a/Engine/Source/Runtime/Renderer/Private/ScenePrivate.h b/Engine/Source/Runtime/Renderer/Private/ScenePrivate.h
--- a/Engine/Source/Runtime/Renderer/Private/ScenePrivate.h
+++ b/Engine/Source/Runtime/Renderer/Private/ScenePrivate.h
@@ -86,7 +86,7 @@ public:
 	FPrimitiveComponentId PrimitiveId;
 
 	/** The occlusion query which contains the primitive's pending occlusion results. */
-	FRenderQueryRHIRef PendingOcclusionQuery[FOcclusionQueryHelpers::MaxBufferedOcclusionFrames];
+	FRefCountedRHIPooledRenderQuery PendingOcclusionQuery[FOcclusionQueryHelpers::MaxBufferedOcclusionFrames];
 	uint32 PendingOcclusionQueryFrames[FOcclusionQueryHelpers::MaxBufferedOcclusionFrames]; 
 
 	uint32 LastTestFrameNumber;
@@ -133,7 +133,7 @@ private:
 	 *	Conditions where this is needed to get a query for read-back are described for bNeedsScanOnRead.
 	 *	Returns -1 if no such query exists in the occlusion history.
 	 */
-	FORCEINLINE int32 ScanOldestNonStaleQueryIndex(uint32 FrameNumber, int32 NumBufferedFrames, int32 LagTolerance) const
+	inline int32 ScanOldestNonStaleQueryIndex(uint32 FrameNumber, int32 NumBufferedFrames, int32 LagTolerance) const
 	{
 		uint32 OldestFrame = UINT32_MAX;
 		int32 OldestQueryIndex = -1;
@@ -152,7 +152,7 @@ private:
 
 public:
 	/** Initialization constructor. */
-	FORCEINLINE FPrimitiveOcclusionHistory(FPrimitiveComponentId InPrimitiveId, int32 SubQuery)
+	inline FPrimitiveOcclusionHistory(FPrimitiveComponentId InPrimitiveId, int32 SubQuery)
 		: PrimitiveId(InPrimitiveId)
 		, LastTestFrameNumber(~0u)
 		, LastConsideredFrameNumber(~0u)
@@ -173,7 +173,7 @@ public:
 		}
 	}
 
-	FORCEINLINE FPrimitiveOcclusionHistory()
+	inline FPrimitiveOcclusionHistory()
 		: LastTestFrameNumber(~0u)
 		, LastConsideredFrameNumber(~0u)
 		, HZBTestIndex(0)
@@ -193,8 +193,7 @@ public:
 		}
 	}
 
-	template<class TOcclusionQueryPool> // here we use a template just to allow this to be inlined without sorting out the header order
-	FORCEINLINE void ReleaseStaleQueries(TOcclusionQueryPool& Pool, uint32 FrameNumber, int32 NumBufferedFrames)
+	inline void ReleaseStaleQueries(uint32 FrameNumber, int32 NumBufferedFrames)
 	{
 		for (uint32 DeltaFrame = NumBufferedFrames; DeltaFrame > 0; DeltaFrame--)
 		{
@@ -202,35 +201,32 @@ public:
 			{
 				uint32 TestFrame = FrameNumber - (DeltaFrame - 1);
 				const uint32 QueryIndex = FOcclusionQueryHelpers::GetQueryLookupIndex(TestFrame, NumBufferedFrames);
-				if (PendingOcclusionQuery[QueryIndex].GetReference() && PendingOcclusionQueryFrames[QueryIndex] != TestFrame)
+				if (PendingOcclusionQueryFrames[QueryIndex] != TestFrame)
 				{
-					Pool.ReleaseQuery(PendingOcclusionQuery[QueryIndex]);
+					PendingOcclusionQuery[QueryIndex].ReleaseQuery();
 				}
 			}
 		}
 	}
-	template<class TOcclusionQueryPool> // here we use a template just to allow this to be inlined without sorting out the header order
-	FORCEINLINE void ReleaseQuery(TOcclusionQueryPool& Pool, uint32 FrameNumber, int32 NumBufferedFrames)
+
+	inline void ReleaseQuery(uint32 FrameNumber, int32 NumBufferedFrames)
 	{
 		const uint32 QueryIndex = FOcclusionQueryHelpers::GetQueryLookupIndex(FrameNumber, NumBufferedFrames);
-		if (PendingOcclusionQuery[QueryIndex].IsValid())
-		{
-			Pool.ReleaseQuery(PendingOcclusionQuery[QueryIndex]);
-		}
+		PendingOcclusionQuery[QueryIndex].ReleaseQuery();
 	}
 
-	FORCEINLINE FRenderQueryRHIParamRef GetQueryForEviction(uint32 FrameNumber, int32 NumBufferedFrames) const
+	inline FRHIRenderQuery* GetQueryForEviction(uint32 FrameNumber, int32 NumBufferedFrames) const
 	{
 		const uint32 QueryIndex = FOcclusionQueryHelpers::GetQueryLookupIndex(FrameNumber, NumBufferedFrames);
 		if (PendingOcclusionQuery[QueryIndex].IsValid())
 		{
-			return PendingOcclusionQuery[QueryIndex].GetReference();
+			return PendingOcclusionQuery[QueryIndex].GetQuery();
 		}
 		return nullptr;
 	}
 
 
-	FORCEINLINE FRenderQueryRHIParamRef GetQueryForReading(uint32 FrameNumber, int32 NumBufferedFrames, int32 LagTolerance, bool& bOutGrouped) const
+	inline FRHIRenderQuery* GetQueryForReading(uint32 FrameNumber, int32 NumBufferedFrames, int32 LagTolerance, bool& bOutGrouped) const
 	{
 		const int32 OldestQueryIndex = bNeedsScanOnRead ? ScanOldestNonStaleQueryIndex(FrameNumber, NumBufferedFrames, LagTolerance)
 														: FOcclusionQueryHelpers::GetQueryLookupIndex(FrameNumber, NumBufferedFrames);
@@ -241,21 +237,21 @@ public:
 			return nullptr;
 		}
 		bOutGrouped = bGroupedQuery[OldestQueryIndex];
-		return PendingOcclusionQuery[OldestQueryIndex];
+		return PendingOcclusionQuery[OldestQueryIndex].GetQuery();
 	}
 
-	FORCEINLINE void SetCurrentQuery(uint32 FrameNumber, FRenderQueryRHIParamRef NewQuery, int32 NumBufferedFrames, bool bGrouped, bool bNeedsScan)
+	inline void SetCurrentQuery(uint32 FrameNumber, FRefCountedRHIPooledRenderQuery&& NewQuery, int32 NumBufferedFrames, bool bGrouped, bool bNeedsScan)
 	{
 		// Get the current occlusion query
 		const uint32 QueryIndex = FOcclusionQueryHelpers::GetQueryIssueIndex(FrameNumber, NumBufferedFrames);
-		PendingOcclusionQuery[QueryIndex] = NewQuery;
+		PendingOcclusionQuery[QueryIndex] = MoveTemp(NewQuery);
 		PendingOcclusionQueryFrames[QueryIndex] = FrameNumber;
 		bGroupedQuery[QueryIndex] = bGrouped;
 
 		bNeedsScanOnRead = bNeedsScan;
 	}
 
-	FORCEINLINE uint32 LastQuerySubmitFrame() const
+	inline uint32 LastQuerySubmitFrame() const
 	{
 		uint32 Result = 0;
 
@@ -312,13 +308,12 @@ struct FPrimitiveOcclusionHistoryKeyFuncs : BaseKeyFuncs<FPrimitiveOcclusionHist
 
 class FIndividualOcclusionHistory
 {
-	FRenderQueryRHIRef PendingOcclusionQuery[FOcclusionQueryHelpers::MaxBufferedOcclusionFrames];
+	FRHIPooledRenderQuery PendingOcclusionQuery[FOcclusionQueryHelpers::MaxBufferedOcclusionFrames];
 	uint32 PendingOcclusionQueryFrames[FOcclusionQueryHelpers::MaxBufferedOcclusionFrames]; // not intialized...this is ok
 
 public:
 
-	template<class TOcclusionQueryPool> // here we use a template just to allow this to be inlined without sorting out the header order
-	FORCEINLINE void ReleaseStaleQueries(TOcclusionQueryPool& Pool, uint32 FrameNumber, int32 NumBufferedFrames)
+	inline void ReleaseStaleQueries(uint32 FrameNumber, int32 NumBufferedFrames)
 	{
 		for (uint32 DeltaFrame = NumBufferedFrames; DeltaFrame > 0; DeltaFrame--)
 		{
@@ -326,39 +321,35 @@ public:
 			{
 				uint32 TestFrame = FrameNumber - (DeltaFrame - 1);
 				const uint32 QueryIndex = FOcclusionQueryHelpers::GetQueryLookupIndex(TestFrame, NumBufferedFrames);
-				if (PendingOcclusionQuery[QueryIndex].GetReference() && PendingOcclusionQueryFrames[QueryIndex] != TestFrame)
+				if (PendingOcclusionQueryFrames[QueryIndex] != TestFrame)
 				{
-					Pool.ReleaseQuery(PendingOcclusionQuery[QueryIndex]);
+					PendingOcclusionQuery[QueryIndex].ReleaseQuery();
 				}
 			}
 		}
 	}
-	template<class TOcclusionQueryPool> // here we use a template just to allow this to be inlined without sorting out the header order
-	FORCEINLINE void ReleaseQuery(TOcclusionQueryPool& Pool, uint32 FrameNumber, int32 NumBufferedFrames)
+	inline void ReleaseQuery(uint32 FrameNumber, int32 NumBufferedFrames)
 	{
 		const uint32 QueryIndex = FOcclusionQueryHelpers::GetQueryLookupIndex(FrameNumber, NumBufferedFrames);
-		if (PendingOcclusionQuery[QueryIndex].GetReference())
-		{
-			Pool.ReleaseQuery(PendingOcclusionQuery[QueryIndex]);
-		}
+		PendingOcclusionQuery[QueryIndex].ReleaseQuery();
 	}
 
-	FORCEINLINE FRenderQueryRHIParamRef GetPastQuery(uint32 FrameNumber, int32 NumBufferedFrames)
+	inline FRHIRenderQuery* GetPastQuery(uint32 FrameNumber, int32 NumBufferedFrames)
 	{
 		// Get the oldest occlusion query
 		const uint32 QueryIndex = FOcclusionQueryHelpers::GetQueryLookupIndex(FrameNumber, NumBufferedFrames);
-		if (PendingOcclusionQuery[QueryIndex].GetReference() && PendingOcclusionQueryFrames[QueryIndex] == FrameNumber - uint32(NumBufferedFrames))
+		if (PendingOcclusionQuery[QueryIndex].GetQuery() && PendingOcclusionQueryFrames[QueryIndex] == FrameNumber - uint32(NumBufferedFrames))
 		{
-			return PendingOcclusionQuery[QueryIndex].GetReference();
+			return PendingOcclusionQuery[QueryIndex].GetQuery();
 		}
 		return nullptr;
 	}
 
-	FORCEINLINE void SetCurrentQuery(uint32 FrameNumber, FRenderQueryRHIParamRef NewQuery, int32 NumBufferedFrames)
+	inline void SetCurrentQuery(uint32 FrameNumber, FRHIPooledRenderQuery&& NewQuery, int32 NumBufferedFrames)
 	{
 		// Get the current occlusion query
 		const uint32 QueryIndex = FOcclusionQueryHelpers::GetQueryIssueIndex(FrameNumber, NumBufferedFrames);
-		PendingOcclusionQuery[QueryIndex] = NewQuery;
+		PendingOcclusionQuery[QueryIndex] = MoveTemp(NewQuery);
 		PendingOcclusionQueryFrames[QueryIndex] = FrameNumber;
 	}
 };
@@ -502,7 +493,7 @@ public:
 	}
 
 	/** @return A random number between 0 and 1. */
-	FORCEINLINE float GetFraction()
+	inline float GetFraction()
 	{
 		if (CurrentSample >= NumSamples)
 		{
@@ -649,7 +640,7 @@ public:
 	{
 	public:
 
-		FORCEINLINE bool operator == (const FProjectedShadowKey &Other) const
+		inline bool operator == (const FProjectedShadowKey &Other) const
 		{
 			return (PrimitiveId == Other.PrimitiveId && Light == Other.Light && ShadowSplitIndex == Other.ShadowSplitIndex && bTranslucentShadow == Other.bTranslucentShadow);
 		}
@@ -670,7 +661,7 @@ public:
 		{
 		}
 
-		friend FORCEINLINE uint32 GetTypeHash(const FSceneViewState::FProjectedShadowKey& Key)
+		friend inline uint32 GetTypeHash(const FSceneViewState::FProjectedShadowKey& Key)
 		{
 			return PointerHash(Key.Light,GetTypeHash(Key.PrimitiveId));
 		}
@@ -683,11 +674,11 @@ public:
 	};
 
 	uint32 UniqueID;
-	typedef TMap<FSceneViewState::FProjectedShadowKey, FRenderQueryRHIRef> ShadowKeyOcclusionQueryMap;
+	typedef TMap<FSceneViewState::FProjectedShadowKey, FRHIPooledRenderQuery> ShadowKeyOcclusionQueryMap;
 	TArray<ShadowKeyOcclusionQueryMap, TInlineAllocator<FOcclusionQueryHelpers::MaxBufferedOcclusionFrames> > ShadowOcclusionQueryMaps;
 
 	/** The view's occlusion query pool. */
-	FRenderQueryPool OcclusionQueryPool;
+	FRenderQueryPoolRHIRef OcclusionQueryPool;
 
 	FHZBOcclusionTester HZBOcclusionTests;
 
@@ -2332,7 +2323,7 @@ struct MeshDrawCommandKeyFuncs : DefaultKeyFuncs<FMeshDrawCommandStateBucket,fal
 	/**
 	 * @return True if the keys match.
 	 */
-	static FORCEINLINE bool Matches(KeyInitType A,KeyInitType B)
+	static inline bool Matches(KeyInitType A,KeyInitType B)
 	{
 		return A.MatchesForDynamicInstancing(B);
 	}
@@ -2340,13 +2331,13 @@ struct MeshDrawCommandKeyFuncs : DefaultKeyFuncs<FMeshDrawCommandStateBucket,fal
 	/**
 	 * @return The key used to index the given element.
 	 */
-	static FORCEINLINE KeyInitType GetSetKey(ElementInitType Element)
+	static inline KeyInitType GetSetKey(ElementInitType Element)
 	{
 		return Element.MeshDrawCommand;
 	}
 
 	/** Calculates a hash index for a key. */
-	static FORCEINLINE uint32 GetKeyHash(KeyInitType Key)
+	static inline uint32 GetKeyHash(KeyInitType Key)
 	{
 		return PointerHash(Key.IndexBuffer, GetTypeHash(Key.CachedPipelineId.GetId()));
 	}
diff --git a/Engine/Source/Runtime/Renderer/Private/SceneRendering.h b/Engine/Source/Runtime/Renderer/Private/SceneRendering.h
--- a/Engine/Source/Runtime/Renderer/Private/SceneRendering.h
+++ b/Engine/Source/Runtime/Renderer/Private/SceneRendering.h
@@ -223,6 +223,101 @@ struct FOcclusionPrimitive
 	FVector Extent;
 };
 
+class FRefCountedRHIPooledRenderQuery
+{
+public:
+	FRefCountedRHIPooledRenderQuery() : FRefCountedRHIPooledRenderQuery(FRHIPooledRenderQuery())
+	{
+	}
+
+	explicit FRefCountedRHIPooledRenderQuery(FRHIPooledRenderQuery&& InQuery)
+	{
+		char* data = new char[sizeof(FRHIPooledRenderQuery) + sizeof(int)];
+		Query = new (data) FRHIPooledRenderQuery();
+		RefCount = new (data + sizeof(FRHIPooledRenderQuery)) int(1);
+		*Query = MoveTemp(InQuery);
+	}
+
+	~FRefCountedRHIPooledRenderQuery()
+	{
+		Deref();
+	}
+
+	bool IsValid() const
+	{
+		return Query && Query->IsValid();
+	}
+
+	FRHIRenderQuery* GetQuery() const
+	{
+		return Query ? Query->GetQuery() : nullptr;
+	}
+
+	void ReleaseQuery()
+	{
+		Deref();
+	}
+
+	FRefCountedRHIPooledRenderQuery(const FRefCountedRHIPooledRenderQuery& Other)
+	{
+		Other.Addref();
+		RefCount = Other.RefCount;
+		Query = Other.Query;
+	}
+
+	FRefCountedRHIPooledRenderQuery& operator= (const FRefCountedRHIPooledRenderQuery& Other)
+	{
+		Other.Addref();
+		Deref();
+		RefCount = Other.RefCount;
+		Query = Other.Query;
+
+		return *this;
+	}
+
+	FRefCountedRHIPooledRenderQuery(FRefCountedRHIPooledRenderQuery&& Other)
+	{
+		RefCount = Other.RefCount;
+		Query = Other.Query;
+		Other.RefCount = nullptr;
+		Other.Query = nullptr;
+	}
+
+	FRefCountedRHIPooledRenderQuery& operator=(FRefCountedRHIPooledRenderQuery&& Other)
+	{
+		Deref();
+		RefCount = Other.RefCount;
+		Query = Other.Query;
+		Other.RefCount = nullptr;
+		Other.Query = nullptr;
+
+		return *this;
+	}
+
+private:
+	void Addref() const
+	{
+		check(RefCount != nullptr);
+		(*RefCount)++;
+	}
+
+	void Deref()
+	{
+		if (RefCount && --(*RefCount) == 0)
+		{
+			Query->~FRHIPooledRenderQuery();
+			char* data = reinterpret_cast<char*>(Query);
+			delete[] data;
+		}
+
+		Query = nullptr;
+		RefCount = nullptr;
+	}
+
+	mutable int* RefCount = nullptr;
+	FRHIPooledRenderQuery* Query = nullptr;
+};
+
 /**
  * Combines consecutive primitives which use the same occlusion query into a single DrawIndexedPrimitive call.
  */
@@ -249,7 +344,7 @@ public:
 	 * Batches a primitive's occlusion query for rendering.
 	 * @param Bounds - The primitive's bounds.
 	 */
-	FRenderQueryRHIParamRef BatchPrimitive(const FVector& BoundsOrigin, const FVector& BoundsBoxExtent, FGlobalDynamicVertexBuffer& DynamicVertexBuffer);
+	FRefCountedRHIPooledRenderQuery BatchPrimitive(const FVector& BoundsOrigin, const FVector& BoundsBoxExtent, FGlobalDynamicVertexBuffer& DynamicVertexBuffer);
 	inline int32 GetNumBatchOcclusionQueries() const
 	{
 		return BatchOcclusionQueries.Num();
@@ -259,7 +354,7 @@ private:
 
 	struct FOcclusionBatch
 	{
-		FRenderQueryRHIRef Query;
+		FRefCountedRHIPooledRenderQuery Query;
 		FGlobalDynamicVertexBuffer::FAllocation VertexAllocation;
 	};
 
@@ -276,7 +371,7 @@ private:
 	uint32 NumBatchedPrimitives;
 
 	/** The pool to allocate occlusion queries from. */
-	class FRenderQueryPool* OcclusionQueryPool;
+	TRefCountPtr<FRHIRenderQueryPool> OcclusionQueryPool;
 };
 
 class FHZBOcclusionTester : public FRenderResource
diff --git a/Engine/Source/Runtime/Renderer/Private/SceneVisibility.cpp b/Engine/Source/Runtime/Renderer/Private/SceneVisibility.cpp
--- a/Engine/Source/Runtime/Renderer/Private/SceneVisibility.cpp
+++ b/Engine/Source/Runtime/Renderer/Private/SceneVisibility.cpp
@@ -696,7 +696,7 @@ static void FetchVisibilityForPrimitives_Range(FVisForPrimParams& Params, FGloba
 	bool bClearQueries = !View.Family->EngineShowFlags.HitProxies;
 	const float CurrentRealTime = View.Family->CurrentRealTime;
 	uint32 OcclusionFrameCounter = ViewState->OcclusionFrameCounter;
-	FRenderQueryPool& OcclusionQueryPool = ViewState->OcclusionQueryPool;
+	FRHIRenderQueryPool* OcclusionQueryPool = ViewState->OcclusionQueryPool;
 	FHZBOcclusionTester& HZBOcclusionTests = ViewState->HZBOcclusionTests;
 
 	int32 ReadBackLagTolerance = NumBufferedFrames;
@@ -859,7 +859,7 @@ static void FetchVisibilityForPrimitives_Range(FVisForPrimParams& Params, FGloba
 						// Read the occlusion query results.
 						uint64 NumSamples = 0;
 						bool bGrouped = false;
-						FRenderQueryRHIParamRef PastQuery = PrimitiveOcclusionHistory->GetQueryForReading(OcclusionFrameCounter, NumBufferedFrames, ReadBackLagTolerance, bGrouped);
+						FRHIRenderQuery* PastQuery = PrimitiveOcclusionHistory->GetQueryForReading(OcclusionFrameCounter, NumBufferedFrames, ReadBackLagTolerance, bGrouped);
 						if (PastQuery)
 						{
 							//int32 RefCount = PastQuery.GetReference()->GetRefCount();
@@ -937,7 +937,7 @@ static void FetchVisibilityForPrimitives_Range(FVisForPrimParams& Params, FGloba
 				{					
 					if (bSingleThreaded)
 					{						
-						PrimitiveOcclusionHistory->ReleaseQuery(OcclusionQueryPool, OcclusionFrameCounter, NumBufferedFrames);
+						PrimitiveOcclusionHistory->ReleaseQuery(OcclusionFrameCounter, NumBufferedFrames);
 					}
 					else
 					{
@@ -1300,7 +1300,7 @@ static int32 FetchVisibilityForPrimitives(const FScene* Scene, FViewInfo& View,
 			StartIndex += NumToProcess;
 		}
 
-		FRenderQueryPool& OcclusionQueryPool = ViewState->OcclusionQueryPool;
+		FRHIRenderQueryPool* OcclusionQueryPool = ViewState->OcclusionQueryPool;
 		FHZBOcclusionTester& HZBOcclusionTests = ViewState->HZBOcclusionTests;		
 
 		int32 NumOccludedPrims = 0;
@@ -1350,7 +1350,7 @@ static int32 FetchVisibilityForPrimitives(const FScene* Scene, FViewInfo& View,
 				for (auto ReleaseQueryIter = OutQueriesToRelease[i].CreateIterator(); ReleaseQueryIter; ++ReleaseQueryIter)
 				{
 					FPrimitiveOcclusionHistory* History = *ReleaseQueryIter;
-					History->ReleaseQuery(OcclusionQueryPool, OcclusionFrameCounter, NumBufferedFrames);
+					History->ReleaseQuery(OcclusionFrameCounter, NumBufferedFrames);
 				}
 				
 				//New query batching
@@ -1369,12 +1369,10 @@ static int32 FetchVisibilityForPrimitives(const FScene* Scene, FViewInfo& View,
 
 			//now add new primitive histories to the view. may resize the view's array.
 			for (int32 i = 0; i < NumTasks; ++i)
-			{								
-				const TArray<FPrimitiveOcclusionHistory>& NewHistoryArray = OutputOcclusionHistory[i];				
-				for (int32 HistoryIndex = 0; HistoryIndex < NewHistoryArray.Num(); ++HistoryIndex)
+			{											
+				for (int32 HistoryIndex = 0; HistoryIndex < OutputOcclusionHistory[i].Num(); ++HistoryIndex)
 				{
-					const FPrimitiveOcclusionHistory& CopySourceHistory = NewHistoryArray[HistoryIndex];
-					ViewPrimitiveOcclusionHistory.Add(CopySourceHistory);
+					ViewPrimitiveOcclusionHistory.Add(MoveTemp(OutputOcclusionHistory[i][HistoryIndex]));
 				}
 
 				//accumulate occluded prims across tasks
