Fixes multi spp GI denoising.

diff --git a/Engine/Source/Runtime/Renderer/Private/ScreenSpaceDenoise.cpp b/Engine/Source/Runtime/Renderer/Private/ScreenSpaceDenoise.cpp
--- a/Engine/Source/Runtime/Renderer/Private/ScreenSpaceDenoise.cpp
+++ b/Engine/Source/Runtime/Renderer/Private/ScreenSpaceDenoise.cpp
@@ -208,12 +208,22 @@ static int32 SignalMaxBatchSize(ESignalProcessing SignalProcessing)
 	return 1;
 }
 
+/** Returns whether a signal have a code path for 1 sample per pixel. */
+static bool SignalSupport1SPP(ESignalProcessing SignalProcessing)
+{
+	return (
+		SignalProcessing == ESignalProcessing::MonochromaticPenumbra ||
+		SignalProcessing == ESignalProcessing::Reflections ||
+		SignalProcessing == ESignalProcessing::AmbientOcclusion);
+}
+
 /** Returns whether a signal can denoise multi sample per pixel. */
 static bool SignalSupportMultiSPP(ESignalProcessing SignalProcessing)
 {
 	return (
 		SignalProcessing == ESignalProcessing::MonochromaticPenumbra ||
-		SignalProcessing == ESignalProcessing::Reflections);
+		SignalProcessing == ESignalProcessing::Reflections ||
+		SignalProcessing == ESignalProcessing::GlobalIllumination);
 }
 
 
@@ -624,9 +634,8 @@ class FSSDSpatialAccumulationCS : public FGlobalShader
 			return false;
 		}
 
-		// Only the reconstruction pass can support 1spp.
-		if (PermutationVector.Get<FStageDim>() != EStage::ReConstruction &&
-			!PermutationVector.Get<FMultiSPPDim>())
+		// Compile out the shader if this permutation gets remapped.
+		if (RemapPermutationVector(PermutationVector) != PermutationVector)
 		{
 			return false;
 		}
@@ -634,6 +643,26 @@ class FSSDSpatialAccumulationCS : public FGlobalShader
 		return ShouldCompileSignalPipeline(SignalProcessing, Parameters.Platform);
 	}
 
+	static FPermutationDomain RemapPermutationVector(FPermutationDomain PermutationVector)
+	{
+		ESignalProcessing SignalProcessing = PermutationVector.Get<FSignalProcessingDim>();
+
+		if (PermutationVector.Get<FStageDim>() == EStage::ReConstruction)
+		{
+			// force use the multi sample per pixel code path.
+			if (!SignalSupport1SPP(SignalProcessing))
+			{
+				PermutationVector.Set<FMultiSPPDim>(true);
+			}
+		}
+		else
+		{
+			PermutationVector.Set<FMultiSPPDim>(true);
+		}
+
+		return PermutationVector;
+	}
+
 	BEGIN_SHADER_PARAMETER_STRUCT(FParameters, )
 		SHADER_PARAMETER(uint32, MaxSampleCount)
 		SHADER_PARAMETER(int32, UpscaleFactor)
@@ -988,6 +1017,7 @@ static void DenoiseSignalAtConstantPixelDensity(
 		PermutationVector.Set<FSSDSpatialAccumulationCS::FStageDim>(FSSDSpatialAccumulationCS::EStage::ReConstruction);
 		PermutationVector.Set<FSSDSpatialAccumulationCS::FUpscaleDim>(PassParameters->UpscaleFactor != 1);
 		PermutationVector.Set<FMultiSPPDim>(bUseMultiInputSPPShaderPath);
+		PermutationVector = FSSDSpatialAccumulationCS::RemapPermutationVector(PermutationVector);
 
 		TShaderMapRef<FSSDSpatialAccumulationCS> ComputeShader(View.ShaderMap, PermutationVector);
 		FComputeShaderUtils::AddPass(
