Renames GI denoiser code path to diffuse indirect.

diff --git a/Engine/Shaders/Private/ScreenSpaceDenoise/SSDCommon.ush b/Engine/Shaders/Private/ScreenSpaceDenoise/SSDCommon.ush
--- a/Engine/Shaders/Private/ScreenSpaceDenoise/SSDCommon.ush
+++ b/Engine/Shaders/Private/ScreenSpaceDenoise/SSDCommon.ush
@@ -32,7 +32,7 @@
 	// A gray scale valued encoded into a 16bit float only have 10bits mantissa.
 	#define TARGETED_SAMPLE_COUNT 1024
 
-#elif CONFIG_SIGNAL_PROCESSING == SIGNAL_PROCESSING_GI_AND_AO
+#elif CONFIG_SIGNAL_PROCESSING == SIGNAL_PROCESSING_DIFFUSE_INDIRECT_AND_AO
 	// Indirect diffuse requires really large spatial kernel to clean the noise.
 	#define TARGETED_SAMPLE_COUNT 4096
 
diff --git a/Engine/Shaders/Private/ScreenSpaceDenoise/SSDDefinitions.ush b/Engine/Shaders/Private/ScreenSpaceDenoise/SSDDefinitions.ush
--- a/Engine/Shaders/Private/ScreenSpaceDenoise/SSDDefinitions.ush
+++ b/Engine/Shaders/Private/ScreenSpaceDenoise/SSDDefinitions.ush
@@ -114,14 +114,14 @@
 	#define SIGNAL_BUFFER_LAYOUT_AO_HISTORY 0xA002
 
 
-	/** Buffer layout given by the ray generation shader for global illumination. */
-	#define SIGNAL_BUFFER_LAYOUT_GI_AND_AO_INPUT 7
+	/** Buffer layout given by the ray generation shader for diffuse indirect  illumination. */
+	#define SIGNAL_BUFFER_LAYOUT_DIFFUSE_INDIRECT_AND_AO_INPUT_NSPP 0xD100
 
-	/** Internal buffer layout for global illumination. */
-	#define SIGNAL_BUFFER_LAYOUT_GI_AND_AO_RECONSTRUCTION 13
+	/** Internal buffer layout for diffuse indirect illumination. */
+	#define SIGNAL_BUFFER_LAYOUT_DIFFUSE_INDIRECT_AND_AO_RECONSTRUCTION 0xD101
 
-	/** Internal buffer layout for global illumination. */
-	#define SIGNAL_BUFFER_LAYOUT_GI_AND_AO_HISTORY 8
+	/** Internal buffer layout for diffuse indirect illumination. */
+	#define SIGNAL_BUFFER_LAYOUT_DIFFUSE_INDIRECT_AND_AO_HISTORY 0xD102
 
 
 /** Defines how the signal is being processed. Matches C++'s ESignalProcessing. */
@@ -135,7 +135,7 @@
 	#define SIGNAL_PROCESSING_AO 2
 
 	/** Global illumination denoising. */
-	#define SIGNAL_PROCESSING_GI_AND_AO 3
+	#define SIGNAL_PROCESSING_DIFFUSE_INDIRECT_AND_AO 3
 
 
 
diff --git a/Engine/Shaders/Private/ScreenSpaceDenoise/SSDSignalCore.ush b/Engine/Shaders/Private/ScreenSpaceDenoise/SSDSignalCore.ush
--- a/Engine/Shaders/Private/ScreenSpaceDenoise/SSDSignalCore.ush
+++ b/Engine/Shaders/Private/ScreenSpaceDenoise/SSDSignalCore.ush
@@ -553,7 +553,7 @@ void DecodeMultiplexedSignals(
 		}
 	}
 	#if CONFIG_SIGNAL_COLOR
-	else if (SignalBufferLayout == SIGNAL_BUFFER_LAYOUT_GI_AND_AO_INPUT)
+	else if (SignalBufferLayout == SIGNAL_BUFFER_LAYOUT_DIFFUSE_INDIRECT_AND_AO_INPUT_NSPP)
 	{
 		OutSamples.Array[0].SampleCount = RawSample[1].g;
 		OutSamples.Array[0].MissCount = OutSamples.Array[0].SampleCount * (1 - RawSample[0].a);
@@ -563,7 +563,7 @@ void DecodeMultiplexedSignals(
 		OutSamples.Array[0].ClosestHitDistance = RawSample[1].r * OutSamples.Array[0].SampleCount;
 		OutSamples.Array[0].ClosestHitDistance = ((RawSample[1].r == -1.0) ? 1.0e20 : RawSample[1].r) * OutSamples.Array[0].SampleCount;
 	}
-	else if (SignalBufferLayout == SIGNAL_BUFFER_LAYOUT_GI_AND_AO_RECONSTRUCTION)
+	else if (SignalBufferLayout == SIGNAL_BUFFER_LAYOUT_DIFFUSE_INDIRECT_AND_AO_RECONSTRUCTION)
 	{
 		if (bNormalizeSample)
 		{
@@ -577,7 +577,7 @@ void DecodeMultiplexedSignals(
 		OutSamples.Array[0].MissCount = OutSamples.Array[0].SampleCount * (1 - RawSample[0].a);
 		OutSamples.Array[0].SceneColor = OutSamples.Array[0].SampleCount * RawSample[0];
 	}
-	else if (SignalBufferLayout == SIGNAL_BUFFER_LAYOUT_GI_AND_AO_HISTORY)
+	else if (SignalBufferLayout == SIGNAL_BUFFER_LAYOUT_DIFFUSE_INDIRECT_AND_AO_HISTORY)
 	{
 		if (bNormalizeSample)
 		{
@@ -766,7 +766,7 @@ void EncodeMultiplexedSignals(
 		OutRawSample[0].g = Sample[0].SampleCount;
 	}
 	#if CONFIG_SIGNAL_COLOR
-	else if (SignalBufferLayout == SIGNAL_BUFFER_LAYOUT_GI_AND_AO_RECONSTRUCTION)
+	else if (SignalBufferLayout == SIGNAL_BUFFER_LAYOUT_DIFFUSE_INDIRECT_AND_AO_RECONSTRUCTION)
 	{
 		float NormalizationFactor = SafeRcp(Sample[0].SampleCount);
 
@@ -774,7 +774,7 @@ void EncodeMultiplexedSignals(
 		OutRawSample[1].r = Sample[0].SampleCount;
 		OutBufferCount = 2;
 	}
-	else if (SignalBufferLayout == SIGNAL_BUFFER_LAYOUT_GI_AND_AO_HISTORY)
+	else if (SignalBufferLayout == SIGNAL_BUFFER_LAYOUT_DIFFUSE_INDIRECT_AND_AO_HISTORY)
 	{
 		float NormalizationFactor = SafeRcp(Sample[0].SampleCount);
 
diff --git a/Engine/Shaders/Private/ScreenSpaceDenoise/SSDSignalFramework.ush b/Engine/Shaders/Private/ScreenSpaceDenoise/SSDSignalFramework.ush
--- a/Engine/Shaders/Private/ScreenSpaceDenoise/SSDSignalFramework.ush
+++ b/Engine/Shaders/Private/ScreenSpaceDenoise/SSDSignalFramework.ush
@@ -60,7 +60,7 @@
 	#include "../DeferredShadingCommon.ush"
 	#include "../CommonViewUniformBuffer.ush"
 
-#elif CONFIG_SIGNAL_PROCESSING == SIGNAL_PROCESSING_AO || CONFIG_SIGNAL_PROCESSING == SIGNAL_PROCESSING_GI_AND_AO
+#elif CONFIG_SIGNAL_PROCESSING == SIGNAL_PROCESSING_AO || CONFIG_SIGNAL_PROCESSING == SIGNAL_PROCESSING_DIFFUSE_INDIRECT_AND_AO
 
 #else
 	#error Unknown CONFIG_SIGNAL_PROCESSING
@@ -669,7 +669,7 @@ float GetSignalWorldBluringRadius(FSSDSignalSample Sample, FSSDSampleSceneInfos
 {
 	return Sample.ClosestHitDistance;
 }
-#elif CONFIG_SIGNAL_PROCESSING == SIGNAL_PROCESSING_GI_AND_AO
+#elif CONFIG_SIGNAL_PROCESSING == SIGNAL_PROCESSING_DIFFUSE_INDIRECT_AND_AO
 {
 	return Sample.ClosestHitDistance;
 }
diff --git a/Engine/Shaders/Private/ScreenSpaceDenoise/SSDSpatialAccumulation.usf b/Engine/Shaders/Private/ScreenSpaceDenoise/SSDSpatialAccumulation.usf
--- a/Engine/Shaders/Private/ScreenSpaceDenoise/SSDSpatialAccumulation.usf
+++ b/Engine/Shaders/Private/ScreenSpaceDenoise/SSDSpatialAccumulation.usf
@@ -223,7 +223,7 @@
 		#error Unexpected stage.
 	#endif
 
-#elif CONFIG_SIGNAL_PROCESSING == SIGNAL_PROCESSING_GI_AND_AO
+#elif CONFIG_SIGNAL_PROCESSING == SIGNAL_PROCESSING_DIFFUSE_INDIRECT_AND_AO
 	#define CONFIG_SIGNAL_COLOR 1
 	#define CONFIG_SIGNAL_COLOR_SH 1
 
@@ -235,8 +235,8 @@
 
 	#if DIM_STAGE == STAGE_RECONSTRUCTION
 		// Input and output layout.
-		#define CONFIG_SIGNAL_INPUT_LAYOUT  SIGNAL_BUFFER_LAYOUT_GI_AND_AO_INPUT
-		#define CONFIG_SIGNAL_OUTPUT_LAYOUT SIGNAL_BUFFER_LAYOUT_GI_AND_AO_RECONSTRUCTION
+		#define CONFIG_SIGNAL_INPUT_LAYOUT  SIGNAL_BUFFER_LAYOUT_DIFFUSE_INDIRECT_AND_AO_INPUT_NSPP
+		#define CONFIG_SIGNAL_OUTPUT_LAYOUT SIGNAL_BUFFER_LAYOUT_DIFFUSE_INDIRECT_AND_AO_RECONSTRUCTION
 
 		#define CONFIG_SAMPLE_SET           SAMPLE_SET_STACKOWIAK_4_SETS
 
@@ -244,8 +244,8 @@
 
 	#elif DIM_STAGE == STAGE_PRE_CONVOLUTION
 		// Input and output layout.
-		#define CONFIG_SIGNAL_INPUT_LAYOUT  SIGNAL_BUFFER_LAYOUT_GI_AND_AO_RECONSTRUCTION
-		#define CONFIG_SIGNAL_OUTPUT_LAYOUT SIGNAL_BUFFER_LAYOUT_GI_AND_AO_RECONSTRUCTION
+		#define CONFIG_SIGNAL_INPUT_LAYOUT  SIGNAL_BUFFER_LAYOUT_DIFFUSE_INDIRECT_AND_AO_RECONSTRUCTION
+		#define CONFIG_SIGNAL_OUTPUT_LAYOUT SIGNAL_BUFFER_LAYOUT_DIFFUSE_INDIRECT_AND_AO_RECONSTRUCTION
 
 		#define CONFIG_SAMPLE_SET           SAMPLE_SET_HEXAWEB
 		#define CONFIG_CUSTOM_SPREAD_FACTOR 1
@@ -253,8 +253,8 @@
 
 	#elif DIM_STAGE == STAGE_POST_FILTERING
 		// Input and output layout.
-		#define CONFIG_SIGNAL_INPUT_LAYOUT  SIGNAL_BUFFER_LAYOUT_GI_AND_AO_HISTORY
-		#define CONFIG_SIGNAL_OUTPUT_LAYOUT SIGNAL_BUFFER_LAYOUT_GI_AND_AO_HISTORY
+		#define CONFIG_SIGNAL_INPUT_LAYOUT  SIGNAL_BUFFER_LAYOUT_DIFFUSE_INDIRECT_AND_AO_HISTORY
+		#define CONFIG_SIGNAL_OUTPUT_LAYOUT SIGNAL_BUFFER_LAYOUT_DIFFUSE_INDIRECT_AND_AO_HISTORY
 
 		#define CONFIG_SAMPLE_SET           SAMPLE_SET_STACKOWIAK_4_SETS
 		#define CONFIG_SAMPLE_SUBSET        1
diff --git a/Engine/Shaders/Private/ScreenSpaceDenoise/SSDTemporalAccumulation.usf b/Engine/Shaders/Private/ScreenSpaceDenoise/SSDTemporalAccumulation.usf
--- a/Engine/Shaders/Private/ScreenSpaceDenoise/SSDTemporalAccumulation.usf
+++ b/Engine/Shaders/Private/ScreenSpaceDenoise/SSDTemporalAccumulation.usf
@@ -118,12 +118,12 @@
 		#define CONFIG_REJECTION_SAMPLE_SET               REJECTION_SAMPLE_SET_SOBEK2018
 	#endif
 
-#elif CONFIG_SIGNAL_PROCESSING == SIGNAL_PROCESSING_GI_AND_AO
+#elif CONFIG_SIGNAL_PROCESSING == SIGNAL_PROCESSING_DIFFUSE_INDIRECT_AND_AO
 	#define CONFIG_SIGNAL_COLOR 1
 
 	// Input and output layout.
-	#define CONFIG_SIGNAL_INPUT_LAYOUT   SIGNAL_BUFFER_LAYOUT_GI_AND_AO_RECONSTRUCTION
-	#define CONFIG_SIGNAL_HISTORY_LAYOUT SIGNAL_BUFFER_LAYOUT_GI_AND_AO_HISTORY
+	#define CONFIG_SIGNAL_INPUT_LAYOUT   SIGNAL_BUFFER_LAYOUT_DIFFUSE_INDIRECT_AND_AO_RECONSTRUCTION
+	#define CONFIG_SIGNAL_HISTORY_LAYOUT SIGNAL_BUFFER_LAYOUT_DIFFUSE_INDIRECT_AND_AO_HISTORY
 
 	#define CONFIG_INPUT_TEXTURE_COUNT 2
 	#define CONFIG_HISTORY_TEXTURE_COUNT 2
@@ -473,7 +473,7 @@ void TemporallyAccumulate(
 		} // if (!bIsPreviousFrameOffscreen)
 	}
 	
-	#if CONFIG_SIGNAL_PROCESSING == SIGNAL_PROCESSING_GI_AND_AO && 0
+	#if CONFIG_SIGNAL_PROCESSING == SIGNAL_PROCESSING_DIFFUSE_INDIRECT_AND_AO && 0
 		DebugOutput[DispatchThreadId] = float4(
 			HistorySamples.Array[0].SampleCount / 4096,
 			0,
@@ -701,7 +701,7 @@ void TemporallyAccumulate(
 						RejectionFactor = max(RejectionFactor, 1 * max(RejectedDiff.SceneColor.y, RejectedDiff.SceneColor.z));
 						RejectionFactor = max(RejectionFactor, 1 * RejectedDiff.SceneColor.a);
 					}
-					#elif CONFIG_SIGNAL_PROCESSING == SIGNAL_PROCESSING_MONOCHROMATIC_PENUMBRA || CONFIG_SIGNAL_PROCESSING == SIGNAL_PROCESSING_AO || CONFIG_SIGNAL_PROCESSING == SIGNAL_PROCESSING_GI_AND_AO
+					#elif CONFIG_SIGNAL_PROCESSING == SIGNAL_PROCESSING_MONOCHROMATIC_PENUMBRA || CONFIG_SIGNAL_PROCESSING == SIGNAL_PROCESSING_AO || CONFIG_SIGNAL_PROCESSING == SIGNAL_PROCESSING_DIFFUSE_INDIRECT_AND_AO
 					{
 						RejectionFactor = abs(ClampedTransformedHistorySample.MissCount - TransformedHistorySample.MissCount);
 					}
diff --git a/Engine/Source/Runtime/Renderer/Private/RayTracing/RayTracingGlobalIllumination.cpp b/Engine/Source/Runtime/Renderer/Private/RayTracing/RayTracingGlobalIllumination.cpp
--- a/Engine/Source/Runtime/Renderer/Private/RayTracing/RayTracingGlobalIllumination.cpp
+++ b/Engine/Source/Runtime/Renderer/Private/RayTracing/RayTracingGlobalIllumination.cpp
@@ -532,17 +532,17 @@ void FDeferredShadingSceneRenderer::RenderRayTracingGlobalIllumination(
 		const IScreenSpaceDenoiser* DefaultDenoiser = IScreenSpaceDenoiser::GetDefaultDenoiser();
 		const IScreenSpaceDenoiser* DenoiserToUse = GRayTracingGlobalIlluminationDenoiser == 1 ? DefaultDenoiser : GScreenSpaceDenoiser;
 
-		IScreenSpaceDenoiser::FGlobalIlluminationInputs DenoiserInputs;
+		IScreenSpaceDenoiser::FDiffuseIndirectInputs DenoiserInputs;
 		DenoiserInputs.Color = GlobalIlluminationTexture;
 		DenoiserInputs.RayHitDistance = RayDistanceTexture;
 
 		{
-			RDG_EVENT_SCOPE(GraphBuilder, "%s%s(GlobalIllumination) %dx%d",
+			RDG_EVENT_SCOPE(GraphBuilder, "%s%s(DiffuseIndirect) %dx%d",
 				DenoiserToUse != DefaultDenoiser ? TEXT("ThirdParty ") : TEXT(""),
 				DenoiserToUse->GetDebugName(),
 				View.ViewRect.Width(), View.ViewRect.Height());
 
-			IScreenSpaceDenoiser::FGlobalIlluminationOutputs DenoiserOutputs = DenoiserToUse->DenoiseGlobalIllumination(
+			IScreenSpaceDenoiser::FDiffuseIndirectOutputs DenoiserOutputs = DenoiserToUse->DenoiseDiffuseIndirect(
 				GraphBuilder,
 				View,
 				&View.PrevViewInfo,
diff --git a/Engine/Source/Runtime/Renderer/Private/RayTracing/RaytracingSkylight.cpp b/Engine/Source/Runtime/Renderer/Private/RayTracing/RaytracingSkylight.cpp
--- a/Engine/Source/Runtime/Renderer/Private/RayTracing/RaytracingSkylight.cpp
+++ b/Engine/Source/Runtime/Renderer/Private/RayTracing/RaytracingSkylight.cpp
@@ -906,7 +906,7 @@ void FDeferredShadingSceneRenderer::RenderRayTracingSkyLight(
 			const IScreenSpaceDenoiser* DefaultDenoiser = IScreenSpaceDenoiser::GetDefaultDenoiser();
 			const IScreenSpaceDenoiser* DenoiserToUse = DefaultDenoiser;// GRayTracingGlobalIlluminationDenoiser == 1 ? DefaultDenoiser : GScreenSpaceDenoiser;
 
-			IScreenSpaceDenoiser::FGlobalIlluminationInputs DenoiserInputs;
+			IScreenSpaceDenoiser::FDiffuseIndirectInputs DenoiserInputs;
 			DenoiserInputs.Color = SkyLightTexture;
 			DenoiserInputs.RayHitDistance = RayDistanceTexture;
 
@@ -920,7 +920,7 @@ void FDeferredShadingSceneRenderer::RenderRayTracingSkyLight(
 					DenoiserToUse->GetDebugName(),
 					View.ViewRect.Width(), View.ViewRect.Height());
 
-				IScreenSpaceDenoiser::FGlobalIlluminationOutputs DenoiserOutputs = DenoiserToUse->DenoiseSkyLight(
+				IScreenSpaceDenoiser::FDiffuseIndirectOutputs DenoiserOutputs = DenoiserToUse->DenoiseSkyLight(
 					GraphBuilder,
 					View,
 					&View.PrevViewInfo,
diff --git a/Engine/Source/Runtime/Renderer/Private/SceneRendering.h b/Engine/Source/Runtime/Renderer/Private/SceneRendering.h
--- a/Engine/Source/Runtime/Renderer/Private/SceneRendering.h
+++ b/Engine/Source/Runtime/Renderer/Private/SceneRendering.h
@@ -764,7 +764,7 @@ struct FPreviousViewInfo
 	FScreenSpaceFilteringHistory AmbientOcclusionHistory;
 
 	// History for global illumination
-	FScreenSpaceFilteringHistory GlobalIlluminationHistory;
+	FScreenSpaceFilteringHistory DiffuseIndirectHistory;
 
 	// History for sky light
 	FScreenSpaceFilteringHistory SkyLightHistory;
diff --git a/Engine/Source/Runtime/Renderer/Private/ScreenSpaceDenoise.cpp b/Engine/Source/Runtime/Renderer/Private/ScreenSpaceDenoise.cpp
--- a/Engine/Source/Runtime/Renderer/Private/ScreenSpaceDenoise.cpp
+++ b/Engine/Source/Runtime/Renderer/Private/ScreenSpaceDenoise.cpp
@@ -138,7 +138,7 @@ enum class ESignalProcessing
 	MonochromaticPenumbra,
 	Reflections,
 	AmbientOcclusion,
-	GlobalIllumination,
+	DiffuseIndirect,
 
 	MAX,
 };
@@ -158,7 +158,7 @@ static bool UsesConstantPixelDensityPassLayout(ESignalProcessing SignalProcessin
 		SignalProcessing == ESignalProcessing::MonochromaticPenumbra ||
 		SignalProcessing == ESignalProcessing::Reflections ||
 		SignalProcessing == ESignalProcessing::AmbientOcclusion ||
-		SignalProcessing == ESignalProcessing::GlobalIllumination);
+		SignalProcessing == ESignalProcessing::DiffuseIndirect);
 }
 
 /** Returns whether a signal processing uses an injestion pass. */
@@ -172,7 +172,7 @@ static bool SignalUsesPreConvolution(ESignalProcessing SignalProcessing)
 {
 	return
 		SignalProcessing == ESignalProcessing::MonochromaticPenumbra ||
-		SignalProcessing == ESignalProcessing::GlobalIllumination;
+		SignalProcessing == ESignalProcessing::DiffuseIndirect;
 }
 
 /** Returns whether a signal processing uses a history rejection pre convolution pass. */
@@ -200,7 +200,7 @@ static int32 SignalMaxBatchSize(ESignalProcessing SignalProcessing)
 	else if (
 		SignalProcessing == ESignalProcessing::Reflections ||
 		SignalProcessing == ESignalProcessing::AmbientOcclusion ||
-		SignalProcessing == ESignalProcessing::GlobalIllumination)
+		SignalProcessing == ESignalProcessing::DiffuseIndirect)
 	{
 		return 1;
 	}
@@ -214,7 +214,7 @@ static bool SignalSupport1SPP(ESignalProcessing SignalProcessing)
 	return (
 		SignalProcessing == ESignalProcessing::MonochromaticPenumbra ||
 		SignalProcessing == ESignalProcessing::Reflections ||
-		SignalProcessing == ESignalProcessing::AmbientOcclusion);
+		SignalProcessing == ESignalProcessing::DiffuseIndirect);
 }
 
 /** Returns whether a signal can denoise multi sample per pixel. */
@@ -223,7 +223,7 @@ static bool SignalSupportMultiSPP(ESignalProcessing SignalProcessing)
 	return (
 		SignalProcessing == ESignalProcessing::MonochromaticPenumbra ||
 		SignalProcessing == ESignalProcessing::Reflections ||
-		SignalProcessing == ESignalProcessing::GlobalIllumination);
+		SignalProcessing == ESignalProcessing::DiffuseIndirect);
 }
 
 
@@ -258,7 +258,7 @@ const TCHAR* const kInjestResourceNames[] = {
 	nullptr,
 	nullptr,
 
-	// GlobalIllumination
+	// DiffuseIndirect
 	nullptr,
 	nullptr,
 	nullptr,
@@ -284,9 +284,9 @@ const TCHAR* const kReconstructionResourceNames[] = {
 	nullptr,
 	nullptr,
 
-	// GlobalIllumination
-	TEXT("GIReconstruction0"),
-	TEXT("GIReconstruction1"),
+	// DiffuseIndirect
+	TEXT("DiffuseIndirectReconstruction0"),
+	TEXT("DiffuseIndirectReconstruction1"),
 	nullptr,
 	nullptr,
 };
@@ -310,9 +310,9 @@ const TCHAR* const kPreConvolutionResourceNames[] = {
 	nullptr,
 	nullptr,
 
-	// GlobalIllumination
-	TEXT("GIPreConvolution0"),
-	TEXT("GIPreConvolution1"),
+	// DiffuseIndirect
+	TEXT("DiffuseIndirectPreConvolution0"),
+	TEXT("DiffuseIndirectPreConvolution1"),
 	nullptr,
 	nullptr,
 };
@@ -336,7 +336,7 @@ const TCHAR* const kRejectionPreConvolutionResourceNames[] = {
 	nullptr,
 	nullptr,
 
-	// GlobalIllumination
+	// DiffuseIndirect
 	nullptr,
 	nullptr,
 	nullptr,
@@ -362,9 +362,9 @@ const TCHAR* const kTemporalAccumulationResourceNames[] = {
 	nullptr,
 	nullptr,
 
-	// GlobalIllumination
-	TEXT("GITemporalAccumulation0"),
-	TEXT("GITemporalAccumulation1"),
+	// DiffuseIndirect
+	TEXT("DiffuseIndirectTemporalAccumulation0"),
+	TEXT("DiffuseIndirectTemporalAccumulation1"),
 	nullptr,
 	nullptr,
 };
@@ -388,9 +388,9 @@ const TCHAR* const kHistoryConvolutionResourceNames[] = {
 	nullptr,
 	nullptr,
 
-	// GlobalIllumination
-	TEXT("GIHistoryConvolution0"),
-	TEXT("GIHistoryConvolution1"),
+	// DiffuseIndirect
+	TEXT("DiffuseIndirectHistoryConvolution0"),
+	TEXT("DiffuseIndirectHistoryConvolution1"),
 	nullptr,
 	nullptr,
 };
@@ -414,7 +414,7 @@ const TCHAR* const kDenoiserOutputResourceNames[] = {
 	nullptr,
 	nullptr,
 
-	// GlobalIllumination
+	// DiffuseIndirect
 	nullptr,
 	nullptr,
 	nullptr,
@@ -439,7 +439,7 @@ bool ShouldCompileSignalPipeline(ESignalProcessing SignalProcessing, EShaderPlat
 	else if (
 		SignalProcessing == ESignalProcessing::MonochromaticPenumbra ||
 		SignalProcessing == ESignalProcessing::AmbientOcclusion ||
-		SignalProcessing == ESignalProcessing::GlobalIllumination)
+		SignalProcessing == ESignalProcessing::DiffuseIndirect)
 	{
 		// Only for ray tracing denoising.
 		return Platform == SP_PCD3D_SM5;
@@ -877,7 +877,7 @@ static void DenoiseSignalAtConstantPixelDensity(
 			ReconstructionTextureCount = HistoryTextureCountPerSignal = 1;
 			bHasReconstructionLayoutDifferentFromHistory = false;
 		}
-		else if (Settings.SignalProcessing == ESignalProcessing::GlobalIllumination)
+		else if (Settings.SignalProcessing == ESignalProcessing::DiffuseIndirect)
 		{
 			ReconstructionDescs[0].Format = PF_FloatRGBA;
 			ReconstructionDescs[1].Format = PF_R16F;
@@ -916,7 +916,7 @@ static void DenoiseSignalAtConstantPixelDensity(
 	}
 
 	#if RHI_RAYTRACING
-	if (Settings.SignalProcessing == ESignalProcessing::GlobalIllumination)
+	if (Settings.SignalProcessing == ESignalProcessing::DiffuseIndirect)
 	{
 		uint32 IterationCount = Settings.MaxInputSPP;
 		uint32 SequenceCount = 1;
@@ -1253,7 +1253,7 @@ static void DenoiseSignalAtConstantPixelDensity(
 			// Requires the normal that are in GBuffer A.
 			if (Settings.SignalProcessing == ESignalProcessing::Reflections ||
 				Settings.SignalProcessing == ESignalProcessing::AmbientOcclusion ||
-				Settings.SignalProcessing == ESignalProcessing::GlobalIllumination)
+				Settings.SignalProcessing == ESignalProcessing::DiffuseIndirect)
 			{
 				GraphBuilder.QueueTextureExtraction(SceneBlackboard.SceneGBufferA, &View.ViewState->PrevFrameViewInfo.GBufferA);
 			}
@@ -1514,12 +1514,12 @@ public:
 		return AmbientOcclusionOutput;
 	}
 
-	FGlobalIlluminationOutputs DenoiseGlobalIllumination(
+	FDiffuseIndirectOutputs DenoiseDiffuseIndirect(
 		FRDGBuilder& GraphBuilder,
 		const FViewInfo& View,
 		FPreviousViewInfo* PreviousViewInfos,
 		const FSceneViewFamilyBlackboard& SceneBlackboard,
-		const FGlobalIlluminationInputs& Inputs,
+		const FDiffuseIndirectInputs& Inputs,
 		const FAmbientOcclusionRayTracingConfig Config) const override
 	{
 		RDG_GPU_STAT_SCOPE(GraphBuilder, DiffuseIndirectDenoiser);
@@ -1529,7 +1529,7 @@ public:
 		InputSignal.Textures[1] = Inputs.RayHitDistance;
 
 		FSSDConstantPixelDensitySettings Settings;
-		Settings.SignalProcessing = ESignalProcessing::GlobalIllumination;
+		Settings.SignalProcessing = ESignalProcessing::DiffuseIndirect;
 		Settings.InputResolutionFraction = Config.ResolutionFraction;
 		Settings.ReconstructionSamples = CVarGIReconstructionSampleCount.GetValueOnRenderThread();
 		Settings.PreConvolutionCount = CVarGIPreConvolutionCount.GetValueOnRenderThread();
@@ -1540,8 +1540,8 @@ public:
 
 		TStaticArray<FScreenSpaceFilteringHistory*, IScreenSpaceDenoiser::kMaxBatchSize> PrevHistories;
 		TStaticArray<FScreenSpaceFilteringHistory*, IScreenSpaceDenoiser::kMaxBatchSize> NewHistories;
-		PrevHistories[0] = &PreviousViewInfos->GlobalIlluminationHistory;
-		NewHistories[0] = View.ViewState ? &View.ViewState->PrevFrameViewInfo.GlobalIlluminationHistory : nullptr;
+		PrevHistories[0] = &PreviousViewInfos->DiffuseIndirectHistory;
+		NewHistories[0] = View.ViewState ? &View.ViewState->PrevFrameViewInfo.DiffuseIndirectHistory : nullptr;
 
 		FSSDSignalTextures SignalOutput;
 		DenoiseSignalAtConstantPixelDensity(
@@ -1551,17 +1551,17 @@ public:
 			NewHistories,
 			&SignalOutput);
 
-		FGlobalIlluminationOutputs GlobalIlluminationOutputs;
+		FDiffuseIndirectOutputs GlobalIlluminationOutputs;
 		GlobalIlluminationOutputs.Color = SignalOutput.Textures[0];
 		return GlobalIlluminationOutputs;
 	}
 
-	FGlobalIlluminationOutputs DenoiseSkyLight(
+	FDiffuseIndirectOutputs DenoiseSkyLight(
 		FRDGBuilder& GraphBuilder,
 		const FViewInfo& View,
 		FPreviousViewInfo* PreviousViewInfos,
 		const FSceneViewFamilyBlackboard& SceneBlackboard,
-		const FGlobalIlluminationInputs& Inputs,
+		const FDiffuseIndirectInputs& Inputs,
 		const FAmbientOcclusionRayTracingConfig Config) const override
 	{
 		RDG_GPU_STAT_SCOPE(GraphBuilder, DiffuseIndirectDenoiser);
@@ -1571,7 +1571,7 @@ public:
 		InputSignal.Textures[1] = Inputs.RayHitDistance;
 
 		FSSDConstantPixelDensitySettings Settings;
-		Settings.SignalProcessing = ESignalProcessing::GlobalIllumination;
+		Settings.SignalProcessing = ESignalProcessing::DiffuseIndirect;
 		Settings.InputResolutionFraction = Config.ResolutionFraction;
 		Settings.ReconstructionSamples = CVarGIReconstructionSampleCount.GetValueOnRenderThread();
 		Settings.PreConvolutionCount = CVarGIPreConvolutionCount.GetValueOnRenderThread();
@@ -1593,7 +1593,7 @@ public:
 			NewHistories,
 			&SignalOutput);
 
-		FGlobalIlluminationOutputs GlobalIlluminationOutputs;
+		FDiffuseIndirectOutputs GlobalIlluminationOutputs;
 		GlobalIlluminationOutputs.Color = SignalOutput.Textures[0];
 		return GlobalIlluminationOutputs;
 	}
diff --git a/Engine/Source/Runtime/Renderer/Private/ScreenSpaceDenoise.h b/Engine/Source/Runtime/Renderer/Private/ScreenSpaceDenoise.h
--- a/Engine/Source/Runtime/Renderer/Private/ScreenSpaceDenoise.h
+++ b/Engine/Source/Runtime/Renderer/Private/ScreenSpaceDenoise.h
@@ -157,34 +157,34 @@ public:
 		const FAmbientOcclusionRayTracingConfig RayTracingConfig) const = 0;
 
 	/** All the inputs of the GI denoisers. */
-	BEGIN_SHADER_PARAMETER_STRUCT(FGlobalIlluminationInputs, )
+	BEGIN_SHADER_PARAMETER_STRUCT(FDiffuseIndirectInputs, )
 		// Irradiance in RGB, AO mask in alpha.
 		SHADER_PARAMETER_RDG_TEXTURE(Texture2D, Color)
 		SHADER_PARAMETER_RDG_TEXTURE(Texture2D, RayHitDistance)
 	END_SHADER_PARAMETER_STRUCT()
 
 	/** All the outputs of the GI denoiser may generate. */
-	BEGIN_SHADER_PARAMETER_STRUCT(FGlobalIlluminationOutputs, )
+	BEGIN_SHADER_PARAMETER_STRUCT(FDiffuseIndirectOutputs, )
 		// Irradiance in RGB, AO mask in alpha.
 		SHADER_PARAMETER_RDG_TEXTURE(Texture2D, Color)
 	END_SHADER_PARAMETER_STRUCT()
 
 	/** Entry point to denoise reflections. */
-	virtual FGlobalIlluminationOutputs DenoiseGlobalIllumination(
+	virtual FDiffuseIndirectOutputs DenoiseDiffuseIndirect(
 		FRDGBuilder& GraphBuilder,
 		const FViewInfo& View,
 		FPreviousViewInfo* PreviousViewInfos,
 		const FSceneViewFamilyBlackboard& SceneBlackboard,
-		const FGlobalIlluminationInputs& Inputs,
+		const FDiffuseIndirectInputs& Inputs,
 		const FAmbientOcclusionRayTracingConfig Config) const = 0;
 
 	/** Entry point to denoise SkyLight. */
-	virtual FGlobalIlluminationOutputs DenoiseSkyLight(
+	virtual FDiffuseIndirectOutputs DenoiseSkyLight(
 		FRDGBuilder& GraphBuilder,
 		const FViewInfo& View,
 		FPreviousViewInfo* PreviousViewInfos,
 		const FSceneViewFamilyBlackboard& SceneBlackboard,
-		const FGlobalIlluminationInputs& Inputs,
+		const FDiffuseIndirectInputs& Inputs,
 		const FAmbientOcclusionRayTracingConfig Config) const = 0;
 
 	/** Returns the interface of the default denoiser of the renderer. */
