Fixed decal rendering in environment captures.

diff --git a/Engine/Source/Runtime/Renderer/Private/DecalRenderingShared.cpp b/Engine/Source/Runtime/Renderer/Private/DecalRenderingShared.cpp
--- a/Engine/Source/Runtime/Renderer/Private/DecalRenderingShared.cpp
+++ b/Engine/Source/Runtime/Renderer/Private/DecalRenderingShared.cpp
@@ -182,7 +182,14 @@ public:
 			SetShaderValue(RHICmdList, ShaderRHI, DecalOrientation, ComponentTrans.GetUnitAxis(EAxis::X));
 		}
 		
-		float LifetimeAlpha = FMath::Clamp(FMath::Min(View.Family->CurrentWorldTime * -DecalProxy.InvFadeDuration + DecalProxy.FadeStartDelayNormalized, View.Family->CurrentWorldTime * DecalProxy.InvFadeInDuration + DecalProxy.FadeInStartDelayNormalized), 0.0f, 1.0f);
+		float LifetimeAlpha = 1.0f;
+
+		// Certain engine captures (e.g. environment reflection) don't have a tick. Default to fully opaque.
+		if (View.Family->CurrentWorldTime)
+		{
+			LifetimeAlpha = FMath::Clamp(FMath::Min(View.Family->CurrentWorldTime * -DecalProxy.InvFadeDuration + DecalProxy.FadeStartDelayNormalized, View.Family->CurrentWorldTime * DecalProxy.InvFadeInDuration + DecalProxy.FadeInStartDelayNormalized), 0.0f, 1.0f);
+		}
+ 
 		SetShaderValue(RHICmdList, ShaderRHI, DecalParams, FVector2D(FadeAlphaValue, LifetimeAlpha));
 	}
 
