Ensure EyeAdaptationHistogram is applied only with valid Histogram

diff --git a/Engine/Source/Runtime/Renderer/Private/PostProcess/PostProcessing.cpp b/Engine/Source/Runtime/Renderer/Private/PostProcess/PostProcessing.cpp
--- a/Engine/Source/Runtime/Renderer/Private/PostProcess/PostProcessing.cpp
+++ b/Engine/Source/Runtime/Renderer/Private/PostProcess/PostProcessing.cpp
@@ -1617,8 +1617,9 @@ void FPostProcessing::Process(FRHICommandListImmediate& RHICmdList, const FViewI
 				}
 				else  // Use histogram version version
 				{
-					// we always add eye adaptation, if the engine show flag is disabled we set the ExposureScale in the texture to a fixed value
-					AutoExposure.EyeAdaptation = AddPostProcessHistogramEyeAdaptation(Context, Histogram);
+					// If we have a histogram, we always add eye adaptation, if the engine show flag is disabled we set the ExposureScale in the texture to a fixed value
+					if (Histogram.IsValid())
+						AutoExposure.EyeAdaptation = AddPostProcessHistogramEyeAdaptation(Context, Histogram);
 				}
 			}
 
