Made AtmosphericSunDirection valid by default.
Make material using those work in preview while we wait before we get proper sky/atmosphere in preview.

Related to https://udn.unrealengine.com/questions/497545/atmospheric-light-vector-in-preview-windows.html

diff --git a/Engine/Source/Runtime/Renderer/Private/SceneRendering.cpp b/Engine/Source/Runtime/Renderer/Private/SceneRendering.cpp
--- a/Engine/Source/Runtime/Renderer/Private/SceneRendering.cpp
+++ b/Engine/Source/Runtime/Renderer/Private/SceneRendering.cpp
@@ -1113,6 +1113,7 @@ void FViewInfo::SetupUniformBufferParameters(
 		Scene = Family->Scene->GetRenderScene();
 	}
 
+	const FVector DefaultSunDirection(0.0f, 0.0f, 1.0f); // Up vector so that the AtmosphericLightVector node always output a valid direction.
 	if (Scene)
 	{
 		if (Scene->SimpleDirectionalLight)
@@ -1164,7 +1165,7 @@ void FViewInfo::SetupUniformBufferParameters(
 			ViewUniformShaderParameters.AtmosphericFogSunDiscHalfApexAngleRadian = SunLightDiskHalfApexAngleRadian;
 			//Added check so atmospheric light color and vector can use a directional light without needing an atmospheric fog actor in the scene
 			ViewUniformShaderParameters.AtmosphericFogSunColor = Scene->SunLight ? Scene->SunLight->Proxy->GetColor() : FLinearColor::Black;
-			ViewUniformShaderParameters.AtmosphericFogSunDirection = Scene->SunLight ? -Scene->SunLight->Proxy->GetDirection() : FVector::ZeroVector;
+			ViewUniformShaderParameters.AtmosphericFogSunDirection = Scene->SunLight ? -Scene->SunLight->Proxy->GetDirection() : DefaultSunDirection;
 			ViewUniformShaderParameters.AtmosphericFogRenderMask = EAtmosphereRenderFlag::E_EnableAll;
 			ViewUniformShaderParameters.AtmosphericFogInscatterAltitudeSampleNum = 0;
 		}
@@ -1186,7 +1187,7 @@ void FViewInfo::SetupUniformBufferParameters(
 		ViewUniformShaderParameters.AtmosphericFogSunDiscHalfApexAngleRadian = 0.f;
 		ViewUniformShaderParameters.AtmosphericFogSunDiscLuminance = FLinearColor::Black;
 		ViewUniformShaderParameters.AtmosphericFogSunColor = FLinearColor::Black;
-		ViewUniformShaderParameters.AtmosphericFogSunDirection = FVector::ZeroVector;
+		ViewUniformShaderParameters.AtmosphericFogSunDirection = DefaultSunDirection;
 		ViewUniformShaderParameters.AtmosphericFogRenderMask = EAtmosphereRenderFlag::E_EnableAll;
 		ViewUniformShaderParameters.AtmosphericFogInscatterAltitudeSampleNum = 0;
 	}
