Fix missing forward declare compile errors

diff --git a/Engine/Source/Runtime/Renderer/Private/VT/RuntimeVirtualTextureSceneProxy.h b/Engine/Source/Runtime/Renderer/Private/VT/RuntimeVirtualTextureSceneProxy.h
--- a/Engine/Source/Runtime/Renderer/Private/VT/RuntimeVirtualTextureSceneProxy.h
+++ b/Engine/Source/Runtime/Renderer/Private/VT/RuntimeVirtualTextureSceneProxy.h
@@ -4,6 +4,9 @@
 
 #include "CoreMinimal.h"
 
+class URuntimeVirtualTexture;
+class URuntimeVirtualTextureComponent;
+
 /** Scene proxy for the URuntimeVirtualTextureComponent. Manages a runtime virtual texture in the renderer scene. */
 class FRuntimeVirtualTextureSceneProxy
 {
@@ -12,5 +15,5 @@ public:
 	~FRuntimeVirtualTextureSceneProxy();
 
 	/** Pointer to linked URuntimeVirtualTexture. Not for dereferencing, just for debug! */
-	class URuntimeVirtualTexture* VirtualTexture;
+	URuntimeVirtualTexture* VirtualTexture;
 };
