Raytracingtestbed returns true if the hardware does not support ray tracing or if UE4 is running in DX11 mode so tests are not broken. Ideally the testing framework should have the ability to skip tests depending on whether some preconditions are met, probably through a virtual function in FAutomationTestBase. Another alternative is that FAutomationTestBase::RunTest() returns an error code instead of a boolean so more sophisticated error handling can be done.

diff --git a/Engine/Source/Runtime/Renderer/Private/Tests/RayTracingTestbed.cpp b/Engine/Source/Runtime/Renderer/Private/Tests/RayTracingTestbed.cpp
--- a/Engine/Source/Runtime/Renderer/Private/Tests/RayTracingTestbed.cpp
+++ b/Engine/Source/Runtime/Renderer/Private/Tests/RayTracingTestbed.cpp
@@ -21,7 +21,8 @@ bool RunRayTracingTestbed_RenderThread(const FString& Parameters)
 
 	if (!GRHISupportsRayTracing)
 	{
-		return false;
+		//Return true so the test passes in DX11, until the testing framework allows to skip tests depending on defined preconditions
+		return true;
 	}
 
 	FVertexBufferRHIRef VertexBuffer;
