Fixes a missing resource in DOF when only the background convolution should be scattering.

diff --git a/Engine/Source/Runtime/Renderer/Private/PostProcess/DiaphragmDOF.cpp b/Engine/Source/Runtime/Renderer/Private/PostProcess/DiaphragmDOF.cpp
--- a/Engine/Source/Runtime/Renderer/Private/PostProcess/DiaphragmDOF.cpp
+++ b/Engine/Source/Runtime/Renderer/Private/PostProcess/DiaphragmDOF.cpp
@@ -1841,7 +1841,7 @@ FRDGTextureRef DiaphragmDOF::AddPasses(
 			FRDGBufferDesc DrawListDescs = FRDGBufferDesc::CreateStructuredDesc(sizeof(float) * 4, 5 * MaxScatteringGroupCount);
 			if (bForegroundHybridScattering)
 				ForegroundScatterDrawListBuffer = GraphBuilder.CreateBuffer(DrawListDescs, TEXT("DOFForegroundDrawList"));
-			if (bForegroundHybridScattering)
+			if (bBackgroundHybridScattering)
 				BackgroundScatterDrawListBuffer = GraphBuilder.CreateBuffer(DrawListDescs, TEXT("DOFBackgroundDrawList"));
 		}
 		
