Fixes long standing bug in DOF where UAV setting could be forgotten on the scatter compilation pass.

diff --git a/Engine/Shaders/Private/DiaphragmDOF/DOFHybridScatterCompilation.usf b/Engine/Shaders/Private/DiaphragmDOF/DOFHybridScatterCompilation.usf
--- a/Engine/Shaders/Private/DiaphragmDOF/DOFHybridScatterCompilation.usf
+++ b/Engine/Shaders/Private/DiaphragmDOF/DOFHybridScatterCompilation.usf
@@ -80,17 +80,21 @@ void ScatterGroupPackMainCS(
 		BRANCH
 		if (GroupId == 0)
 		{
+			#if DIM_HYBRID_SCATTER_FGD
 			for (uint i = 0; i < 5; i++)
 			{
 				OutForegroundScatterDrawList[ScatterGroupCount * 5 + GroupThreadIndex + i * SCATTER_GROUP_PER_INSTANCE] = float4(0, 0, 0, 0);
 			}
+			#endif
 		}
 		else
-		{		
+		{	
+			#if DIM_HYBRID_SCATTER_BGD
 			for (uint i = 0; i < 5; i++)
 			{
 				OutBackgroundScatterDrawList[ScatterGroupCount * 5 + GroupThreadIndex + i * SCATTER_GROUP_PER_INSTANCE] = float4(0, 0, 0, 0);
 			}
+			#endif
 		}
 	}
 	#endif
diff --git a/Engine/Source/Runtime/Renderer/Private/PostProcess/DiaphragmDOF.cpp b/Engine/Source/Runtime/Renderer/Private/PostProcess/DiaphragmDOF.cpp
--- a/Engine/Source/Runtime/Renderer/Private/PostProcess/DiaphragmDOF.cpp
+++ b/Engine/Source/Runtime/Renderer/Private/PostProcess/DiaphragmDOF.cpp
@@ -834,6 +834,10 @@ class FDiaphragmDOFScatterGroupPackCS : public FDiaphragmDOFShader
 {
 	DECLARE_GLOBAL_SHADER(FDiaphragmDOFScatterGroupPackCS);
 	SHADER_USE_PARAMETER_STRUCT(FDiaphragmDOFScatterGroupPackCS, FDiaphragmDOFShader);
+	
+	using FPermutationDomain = TShaderPermutationDomain<
+		FDiaphragmDOFReduceCS::FHybridScatterForeground,
+		FDiaphragmDOFReduceCS::FHybridScatterBackground>;
 
 	static bool ShouldCompilePermutation(const FGlobalShaderPermutationParameters& Parameters)
 	{
@@ -841,6 +845,16 @@ class FDiaphragmDOFScatterGroupPackCS : public FDiaphragmDOFShader
 		{
 			return false;
 		}
+		
+		FPermutationDomain PermutationVector(Parameters.PermutationId);
+
+		// This shader is used there is at least foreground and/or background to scatter.
+		if (!PermutationVector.Get<FDiaphragmDOFReduceCS::FHybridScatterForeground>() &&
+			!PermutationVector.Get<FDiaphragmDOFReduceCS::FHybridScatterBackground>())
+		{
+			return false;
+		}
+
 		return FDiaphragmDOFShader::ShouldCompilePermutation(Parameters);
 	}
 	
@@ -1907,6 +1921,8 @@ FRDGTextureRef DiaphragmDOF::AddPasses(
 		// Pack multiple scattering group on same primitive instance to increase wave occupancy in the scattering vertex shader.
 		if (bForegroundHybridScattering || bBackgroundHybridScattering)
 		{
+			// TODO: could avoid multiple shader permutation by doing multiple passes with a no barrier UAV that isn't implemented yet.
+
 			FDiaphragmDOFScatterGroupPackCS::FParameters* PassParameters = GraphBuilder.AllocParameters<FDiaphragmDOFScatterGroupPackCS::FParameters>();
 			PassParameters->MaxScatteringGroupCount = MaxScatteringGroupCount;
 			PassParameters->OutScatterDrawIndirectParameters = GraphBuilder.CreateUAV(DrawIndirectParametersBuffer);
@@ -1915,7 +1931,11 @@ FRDGTextureRef DiaphragmDOF::AddPasses(
 			if (BackgroundScatterDrawListBuffer)
 				PassParameters->OutBackgroundScatterDrawList = GraphBuilder.CreateUAV(BackgroundScatterDrawListBuffer);
 
-			TShaderMapRef<FDiaphragmDOFScatterGroupPackCS> ComputeShader(View.ShaderMap);
+			FDiaphragmDOFScatterGroupPackCS::FPermutationDomain PermutationVector;
+			PermutationVector.Set<FDiaphragmDOFReduceCS::FHybridScatterForeground>(bForegroundHybridScattering);
+			PermutationVector.Set<FDiaphragmDOFReduceCS::FHybridScatterBackground>(bBackgroundHybridScattering);
+
+			TShaderMapRef<FDiaphragmDOFScatterGroupPackCS> ComputeShader(View.ShaderMap, PermutationVector);
 			FComputeShaderUtils::AddPass(
 				GraphBuilder,
 				RDG_EVENT_NAME("DOF ScatterGroupPack"),
