Use least-common-multiple for block size of allocated VT, rather than simple max of all layers
- This is needed to support stacks of different block sizes that don't evenly divide each other (such as 2 and 3, LCM is 6, using max value of 3 will not properly wrap for both layers)

diff --git a/Engine/Source/Runtime/Renderer/Private/VT/VirtualTextureSystem.cpp b/Engine/Source/Runtime/Renderer/Private/VT/VirtualTextureSystem.cpp
--- a/Engine/Source/Runtime/Renderer/Private/VT/VirtualTextureSystem.cpp
+++ b/Engine/Source/Runtime/Renderer/Private/VT/VirtualTextureSystem.cpp
@@ -328,8 +328,8 @@ IAllocatedVirtualTexture* FVirtualTextureSystem::AllocateVirtualTexture(const FA
 
 	uint32 BlockWidthInTiles = 0u;
 	uint32 BlockHeightInTiles = 0u;
-	uint32 WidthInBlocks = 0u;
-	uint32 HeightInBlocks = 0u;
+	uint32 MinWidthInBlocks = ~0u;
+	uint32 MinHeightInBlocks = ~0u;
 	uint32 DepthInTiles = 0u;
 	bool bSupport16BitPageTable = true;
 	FVirtualTextureProducer* ProducerForLayer[VIRTUALTEXTURE_SPACE_MAXLAYERS] = { nullptr };
@@ -343,8 +343,8 @@ IAllocatedVirtualTexture* FVirtualTextureSystem::AllocateVirtualTexture(const FA
 			const FVTProducerDescription& ProducerDesc = Producer->GetDescription();
 			BlockWidthInTiles = FMath::Max(BlockWidthInTiles, ProducerDesc.BlockWidthInTiles);
 			BlockHeightInTiles = FMath::Max(BlockHeightInTiles, ProducerDesc.BlockHeightInTiles);
-			WidthInBlocks = FMath::Max<uint32>(WidthInBlocks, ProducerDesc.WidthInBlocks);
-			HeightInBlocks = FMath::Max<uint32>(HeightInBlocks, ProducerDesc.HeightInBlocks);
+			MinWidthInBlocks = FMath::Min<uint32>(MinWidthInBlocks, ProducerDesc.WidthInBlocks);
+			MinHeightInBlocks = FMath::Min<uint32>(MinHeightInBlocks, ProducerDesc.HeightInBlocks);
 			DepthInTiles = FMath::Max(DepthInTiles, ProducerDesc.DepthInTiles);
 
 			FVirtualTexturePhysicalSpace* PhysicalSpace = Producer->GetPhysicalSpace(Desc.LocalLayerToProduce[LayerIndex]);
@@ -358,10 +358,51 @@ IAllocatedVirtualTexture* FVirtualTextureSystem::AllocateVirtualTexture(const FA
 
 	check(BlockWidthInTiles > 0u);
 	check(BlockHeightInTiles > 0u);
-	check(WidthInBlocks > 0u);
-	check(HeightInBlocks > 0u);
 	check(DepthInTiles > 0u);
 
+	// Find a block width that is evenly divided by all layers (least common multiple)
+	// Start with min size, then increment by min size until a valid size is found
+	uint32 WidthInBlocks = MinWidthInBlocks;
+	{
+		bool bFoundValidWidthInBlocks = false;
+		while (!bFoundValidWidthInBlocks)
+		{
+			bFoundValidWidthInBlocks = true;
+			for (uint32 LayerIndex = 0u; LayerIndex < Desc.NumLayers; ++LayerIndex)
+			{
+				const FVirtualTextureProducer* Producer = ProducerForLayer[LayerIndex];
+				if ((WidthInBlocks % Producer->GetDescription().WidthInBlocks) != 0u)
+				{
+					WidthInBlocks += MinWidthInBlocks;
+					check(WidthInBlocks > MinWidthInBlocks); // check for overflow
+					bFoundValidWidthInBlocks = false;
+					break;
+				}
+			}
+		}
+	}
+
+	// Same thing for height
+	uint32 HeightInBlocks = MinHeightInBlocks;
+	{
+		bool bFoundValidHeightInBlocks = false;
+		while (!bFoundValidHeightInBlocks)
+		{
+			bFoundValidHeightInBlocks = true;
+			for (uint32 LayerIndex = 0u; LayerIndex < Desc.NumLayers; ++LayerIndex)
+			{
+				const FVirtualTextureProducer* Producer = ProducerForLayer[LayerIndex];
+				if ((HeightInBlocks % Producer->GetDescription().HeightInBlocks) != 0u)
+				{
+					HeightInBlocks += MinHeightInBlocks;
+					check(HeightInBlocks > MinHeightInBlocks); // check for overflow
+					bFoundValidHeightInBlocks = false;
+					break;
+				}
+			}
+		}
+	}
+
 	FVTSpaceDescription SpaceDesc;
 	SpaceDesc.Dimensions = Desc.Dimensions;
 	SpaceDesc.NumLayers = Desc.NumLayers;
