Fixes long standing issue in DOF's recombine pass trying to create UAV on MSAA render target.

diff --git a/Engine/Source/Runtime/Renderer/Private/PostProcess/DiaphragmDOF.cpp b/Engine/Source/Runtime/Renderer/Private/PostProcess/DiaphragmDOF.cpp
--- a/Engine/Source/Runtime/Renderer/Private/PostProcess/DiaphragmDOF.cpp
+++ b/Engine/Source/Runtime/Renderer/Private/PostProcess/DiaphragmDOF.cpp
@@ -2396,6 +2396,7 @@ FRDGTextureRef DiaphragmDOF::AddPasses(
 	{
 		{
 			FRDGTextureDesc Desc = InputSceneColor->Desc;
+			Desc.NumSamples = 1;
 			Desc.TargetableFlags |= TexCreate_UAV;
 			NewSceneColor = GraphBuilder.CreateTexture(Desc, TEXT("DOFRecombine"));
 		}
