Removed previously added shader premutation for post process.

platofrms<SM4, Metal and OpenGL will always use flattened 2d LUTs to avoid all those extra permutations.

diff --git a/Engine/Source/Runtime/Engine/Private/ShaderCompiler/ShaderCompiler.cpp b/Engine/Source/Runtime/Engine/Private/ShaderCompiler/ShaderCompiler.cpp
--- a/Engine/Source/Runtime/Engine/Private/ShaderCompiler/ShaderCompiler.cpp
+++ b/Engine/Source/Runtime/Engine/Private/ShaderCompiler/ShaderCompiler.cpp
@@ -3880,7 +3880,7 @@ void VerifyGlobalShaders(EShaderPlatform Platform, bool bLoadedFromCacheFile)
 		}
 
 		ensureMsgf(
-			PermutationCountToCompile < 397,	// ToneMapper today (2019-04-17) can go up to 396 permutations
+			PermutationCountToCompile < 200,	// Please try to not bump this
 			TEXT("Global shader %s has %i permutation: probably more that it needs."),
 			GlobalShaderType->GetName(), PermutationCountToCompile);
 
diff --git a/Engine/Source/Runtime/RHI/Public/RHIDefinitions.h b/Engine/Source/Runtime/RHI/Public/RHIDefinitions.h
--- a/Engine/Source/Runtime/RHI/Public/RHIDefinitions.h
+++ b/Engine/Source/Runtime/RHI/Public/RHIDefinitions.h
@@ -1093,6 +1093,16 @@ inline bool RHISupportsGeometryShaders(const EShaderPlatform Platform)
 	return IsFeatureLevelSupported(Platform, ERHIFeatureLevel::SM4) && !IsMetalPlatform(Platform) && !IsVulkanMobilePlatform(Platform);
 }
 
+/** Return true if and only if the GPU support rendering to volume textures (2D Array, 3D) is guaranteed supported for a target platform.
+	if PipelineVolumeTextureLUTSupportGuaranteedAtRuntime is true then it is guaranteed that GSupportsVolumeTextureRendering is true at runtime.
+*/
+inline bool RHIVolumeTextureRenderingSupportGuaranteed(const EShaderPlatform Platform)
+{
+	return IsFeatureLevelSupported(Platform, ERHIFeatureLevel::SM4)
+		&& !IsMetalPlatform(Platform)		// EMetalFeaturesLayeredRendering supports needs to be checked at runtime
+		&& !IsOpenGLPlatform(Platform);		// Apparently, some OpenGL 3.3 cards support SM4 but can't render to volume textures 
+}
+
 inline bool RHIHasTiledGPU(const EShaderPlatform Platform)
 {
 	// @todo MetalMRT Technically we should include (Platform == SP_METAL_MRT) but this would disable depth-pre-pass which is currently required.
diff --git a/Engine/Source/Runtime/Renderer/Private/PostProcess/PostProcessCombineLUTs.cpp b/Engine/Source/Runtime/Renderer/Private/PostProcess/PostProcessCombineLUTs.cpp
--- a/Engine/Source/Runtime/Renderer/Private/PostProcess/PostProcessCombineLUTs.cpp
+++ b/Engine/Source/Runtime/Renderer/Private/PostProcess/PostProcessCombineLUTs.cpp
@@ -58,15 +58,12 @@ static TAutoConsoleVariable<int32> CVarMobileTonemapperFilm(
 	);
 
 // false:use 256x16 texture / true:use volume texture (faster, requires geometry shader)
-// USE_VOLUME_LUT: needs to be the same for C++ and HLSL
-bool PipelineVolumeTextureLUTMayBeSupportedAtRuntime(EShaderPlatform Platform)
+// USE_VOLUME_LUT: needs to be the same for C++ and HLSL.
+// Safe to use at pipeline and run time.
+bool PipelineVolumeTextureLUTSupportGuaranteedAtRuntime(EShaderPlatform Platform)
 {
 	// This is used to know if the target shader platform does not support required volume texture features we need for sure (read, render to).
-	return (IsFeatureLevelSupported(Platform, ERHIFeatureLevel::SM4) && (Platform != SP_METAL_MRT && Platform != SP_METAL_MRT_MAC && Platform != SP_VULKAN_SM5_LUMIN) && (RHISupportsGeometryShaders(Platform) || RHISupportsVertexShaderLayer(Platform)));
-}
-bool RuntimeVolumeTextureLUTSupported(EShaderPlatform Platform)	// Only valid to be called during runtime
-{
-	return PipelineVolumeTextureLUTMayBeSupportedAtRuntime(Platform) && GSupportsVolumeTextureRendering;
+	return RHIVolumeTextureRenderingSupportGuaranteed(Platform) && (RHISupportsGeometryShaders(Platform) || RHISupportsVertexShaderLayer(Platform));
 }
 
 // including the neutral one at index 0
@@ -568,17 +565,11 @@ class FLUTBlenderPS : public FGlobalShader
 
 public:
 
-	class FBlenderUseVolumeLut : SHADER_PERMUTATION_BOOL("USE_VOLUME_LUT");
 	class FBlendCount : SHADER_PERMUTATION_RANGE_INT("BLENDCOUNT", 1, 5);
-	using FPermutationDomain = TShaderPermutationDomain<FBlenderUseVolumeLut, FBlendCount>;
+	using FPermutationDomain = TShaderPermutationDomain<FBlendCount>;
 
 	static bool ShouldCompilePermutation(const FGlobalShaderPermutationParameters& Parameters)
 	{
-		FPermutationDomain PermutationVector(Parameters.PermutationId);
-		if (!PipelineVolumeTextureLUTMayBeSupportedAtRuntime(Parameters.Platform) && PermutationVector.Get<FBlenderUseVolumeLut>())
-		{
-			return false; // We know that this platform does not support volume textures feature needed for our post process, so do not compile the permutation.
-		}
 		return true;
 	}
 
@@ -598,6 +589,9 @@ public:
 	static void ModifyCompilationEnvironment(const FGlobalShaderPermutationParameters& Parameters, FShaderCompilerEnvironment& OutEnvironment)
 	{
 		FGlobalShader::ModifyCompilationEnvironment(Parameters, OutEnvironment);
+
+		const int UseVolumeLut = PipelineVolumeTextureLUTSupportGuaranteedAtRuntime(Parameters.Platform) ? 1 : 0;
+		OutEnvironment.SetDefine(TEXT("USE_VOLUME_LUT"), UseVolumeLut);
 	}
 
 	virtual bool Serialize(FArchive& Ar) override
@@ -624,17 +618,11 @@ class FLUTBlenderCS : public FGlobalShader
 
 public:
 
-	class FBlenderUseVolumeLut : SHADER_PERMUTATION_BOOL("USE_VOLUME_LUT");
 	class FBlendCount : SHADER_PERMUTATION_RANGE_INT("BLENDCOUNT", 1, 5);
-	using FPermutationDomain = TShaderPermutationDomain<FBlenderUseVolumeLut, FBlendCount>;
+	using FPermutationDomain = TShaderPermutationDomain<FBlendCount>;
 
 	static bool ShouldCompilePermutation(const FGlobalShaderPermutationParameters& Parameters)
 	{
-		FPermutationDomain PermutationVector(Parameters.PermutationId);
-		if (!PipelineVolumeTextureLUTMayBeSupportedAtRuntime(Parameters.Platform) && PermutationVector.template Get<FBlenderUseVolumeLut>())
-		{
-			return false; // We know that this platform does not support volume textures feature needed for our post process, so do not compile the permutation.
-		}
 		return IsFeatureLevelSupported(Parameters.Platform, ERHIFeatureLevel::SM5);
 	}
 
@@ -643,6 +631,9 @@ public:
 		// CS params
 		FGlobalShader::ModifyCompilationEnvironment(Parameters, OutEnvironment);
 		OutEnvironment.SetDefine(TEXT("THREADGROUP_SIZE"), GCombineLUTsComputeTileSize);
+
+		const int UseVolumeLut = PipelineVolumeTextureLUTSupportGuaranteedAtRuntime(Parameters.Platform) ? 1 : 0;
+		OutEnvironment.SetDefine(TEXT("USE_VOLUME_LUT"), UseVolumeLut);
 	}
 
 	FLUTBlenderCS(const ShaderMetaType::CompiledShaderInitializerType& Initializer)
@@ -718,11 +709,10 @@ static void SetLUTBlenderShader(FRenderingCompositePassContext& Context, TRHICom
 	auto ShaderMap = Context.GetShaderMap();
 
 	FLUTBlenderPS::FPermutationDomain PermutationVector;
-	PermutationVector.Set<FLUTBlenderPS::FBlenderUseVolumeLut>(RuntimeVolumeTextureLUTSupported(Context.GetShaderPlatform()));
 	PermutationVector.Set<FLUTBlenderPS::FBlendCount>(BlendCount);
 	TShaderMapRef<FLUTBlenderPS> PixelShader(ShaderMap, PermutationVector);
 
-	if(RuntimeVolumeTextureLUTSupported(Context.View.GetShaderPlatform()))
+	if(PipelineVolumeTextureLUTSupportGuaranteedAtRuntime(Context.View.GetShaderPlatform()))
 	{
 		TShaderMapRef<FWriteToSliceVS> VertexShader(ShaderMap);
 		TOptionalShaderMapRef<FWriteToSliceGS> GeometryShader(ShaderMap);
@@ -888,7 +878,7 @@ void FRCPassPostProcessCombineLUTs::Process(FRenderingCompositePassContext& Cont
 	SCOPED_DRAW_EVENTF(Context.RHICmdList, PostProcessCombineLUTs, TEXT("PostProcessCombineLUTs%s [%d] %dx%dx%d"),
 		bIsComputePass?TEXT("Compute"):TEXT(""), LocalCount, GLUTSize, GLUTSize, GLUTSize);
 
-	const bool bUseVolumeTextureLUT = RuntimeVolumeTextureLUTSupported(ShaderPlatform);
+	const bool bUseVolumeTextureLUT = PipelineVolumeTextureLUTSupportGuaranteedAtRuntime(ShaderPlatform);
 	// for a 3D texture, the viewport is 16x16 (per slice), for a 2D texture, it's unwrapped to 256x16
 	FIntPoint DestSize(bUseVolumeTextureLUT ? GLUTSize : GLUTSize * GLUTSize, GLUTSize);
 
@@ -985,14 +975,13 @@ void FRCPassPostProcessCombineLUTs::DispatchCS(TRHICmdList& RHICmdList, FRenderi
 {
 	auto ShaderMap = Context.GetShaderMap();
 
-	const bool bRuntimeVolumeTextureLUTSupported = RuntimeVolumeTextureLUTSupported(ShaderPlatform);
+	const bool bRuntimeVolumeTextureLUTSupported = PipelineVolumeTextureLUTSupportGuaranteedAtRuntime(ShaderPlatform);
 
 	FIntPoint DestSize(DestRect.Width(), DestRect.Height());
 	uint32 GroupSizeXY = FMath::DivideAndRoundUp(DestSize.X, GCombineLUTsComputeTileSize);
 	uint32 GroupSizeZ = bRuntimeVolumeTextureLUTSupported ? GroupSizeXY : 1;
 
 	FLUTBlenderCS::FPermutationDomain PermutationVector;
-	PermutationVector.Set<FLUTBlenderCS::FBlenderUseVolumeLut>(bRuntimeVolumeTextureLUTSupported);
 	PermutationVector.Set<FLUTBlenderCS::FBlendCount>(BlendCount);
 	TShaderMapRef<FLUTBlenderCS> ComputeShader(ShaderMap, PermutationVector);
 	RHICmdList.SetComputeShader(ComputeShader->GetComputeShader());
@@ -1015,7 +1004,7 @@ FPooledRenderTargetDesc FRCPassPostProcessCombineLUTs::ComputeOutputDesc(EPassOu
 	}
 	else
 	{
-		Ret =FSceneViewState::CreateLUTRenderTarget(GLUTSize, RuntimeVolumeTextureLUTSupported(ShaderPlatform), false, bNeedFloatOutput);
+		Ret =FSceneViewState::CreateLUTRenderTarget(GLUTSize, PipelineVolumeTextureLUTSupportGuaranteedAtRuntime(ShaderPlatform), false, bNeedFloatOutput);
 	}
 	Ret.ClearValue = FClearValueBinding::Transparent;
 
diff --git a/Engine/Source/Runtime/Renderer/Private/PostProcess/PostProcessCombineLUTs.h b/Engine/Source/Runtime/Renderer/Private/PostProcess/PostProcessCombineLUTs.h
--- a/Engine/Source/Runtime/Renderer/Private/PostProcess/PostProcessCombineLUTs.h
+++ b/Engine/Source/Runtime/Renderer/Private/PostProcess/PostProcessCombineLUTs.h
@@ -14,8 +14,7 @@
 class UTexture;
 class FFinalPostProcessSettings;
 
-bool PipelineVolumeTextureLUTMayBeSupportedAtRuntime(EShaderPlatform Platform);
-bool RuntimeVolumeTextureLUTSupported(EShaderPlatform Platform);
+bool PipelineVolumeTextureLUTSupportGuaranteedAtRuntime(EShaderPlatform Platform);
 
 // derives from TRenderingCompositePassBase<InputCount, OutputCount>
 class FRCPassPostProcessCombineLUTs : public TRenderingCompositePassBase<0, 1>
diff --git a/Engine/Source/Runtime/Renderer/Private/PostProcess/PostProcessTonemap.cpp b/Engine/Source/Runtime/Renderer/Private/PostProcess/PostProcessTonemap.cpp
--- a/Engine/Source/Runtime/Renderer/Private/PostProcess/PostProcessTonemap.cpp
+++ b/Engine/Source/Runtime/Renderer/Private/PostProcess/PostProcessTonemap.cpp
@@ -125,7 +125,6 @@ class FTonemapperVignetteDim       : SHADER_PERMUTATION_BOOL("USE_VIGNETTE");
 class FTonemapperSharpenDim        : SHADER_PERMUTATION_BOOL("USE_SHARPEN");
 class FTonemapperGrainJitterDim    : SHADER_PERMUTATION_BOOL("USE_GRAIN_JITTER");
 class FTonemapperSwitchAxis        : SHADER_PERMUTATION_BOOL("NEEDTOSWITCHVERTICLEAXIS");
-class FTonemapperUseVolumeLut      : SHADER_PERMUTATION_BOOL("USE_VOLUME_LUT");
 
 using FCommonDomain = TShaderPermutationDomain<
 	FTonemapperBloomDim,
@@ -134,8 +133,7 @@ using FCommonDomain = TShaderPermutationDomain<
 	FTonemapperVignetteDim,
 	FTonemapperSharpenDim,
 	FTonemapperGrainJitterDim,
-	FTonemapperSwitchAxis,
-	FTonemapperUseVolumeLut>;
+	FTonemapperSwitchAxis>;
 
 FORCEINLINE_DEBUGGABLE bool ShouldCompileCommonPermutation(const FGlobalShaderPermutationParameters& Parameters, const FCommonDomain& PermutationVector)
 {
@@ -181,8 +179,6 @@ FCommonDomain BuildCommonPermutationDomain(const FViewInfo& View, bool bGammaOnl
 	PermutationVector.Set<FTonemapperSharpenDim>(CVarTonemapperSharpen.GetValueOnRenderThread() > 0.0f);	
 	PermutationVector.Set<FTonemapperSwitchAxis>(bSwitchVerticalAxis);
 
-	PermutationVector.Set<FTonemapperUseVolumeLut>(RuntimeVolumeTextureLUTSupported(View.GetShaderPlatform()));
-
 	return PermutationVector;
 }
 
@@ -254,10 +250,6 @@ bool ShouldCompileDesktopPermutation(const FGlobalShaderPermutationParameters& P
 		return !PermutationVector.Get<FTonemapperColorFringeDim>() &&
 			!PermutationVector.Get<FTonemapperGrainQuantizationDim>();
 	}
-
-	if (!PipelineVolumeTextureLUTMayBeSupportedAtRuntime(Parameters.Platform) && CommonPermutationVector.Get<FTonemapperUseVolumeLut>())
-		return false; // We know that this platform does not support volume textures feature needed for our post process, so do not compile the permutation.
-
 	return true;
 }
 
@@ -719,6 +711,8 @@ class FPostProcessTonemapPS : public FGlobalShader
 
 	static void ModifyCompilationEnvironment(const FGlobalShaderPermutationParameters& Parameters, FShaderCompilerEnvironment& OutEnvironment)
 	{
+		const int UseVolumeLut = PipelineVolumeTextureLUTSupportGuaranteedAtRuntime(Parameters.Platform) ? 1 : 0;
+		OutEnvironment.SetDefine(TEXT("USE_VOLUME_LUT"), UseVolumeLut);
 	}
 
 	/** Default constructor. */
@@ -795,6 +789,9 @@ class FPostProcessTonemapCS : public FGlobalShader
 	{
 		OutEnvironment.SetDefine(TEXT("THREADGROUP_SIZEX"), GTonemapComputeTileSizeX);
 		OutEnvironment.SetDefine(TEXT("THREADGROUP_SIZEY"), GTonemapComputeTileSizeY);
+
+		const int UseVolumeLut = PipelineVolumeTextureLUTSupportGuaranteedAtRuntime(Parameters.Platform) ? 1 : 0;
+		OutEnvironment.SetDefine(TEXT("USE_VOLUME_LUT"), UseVolumeLut);
 	}
 
 	/** Default constructor. */
