Fixes r.DOF.Scatter.NeighborCompareMaxColor

diff --git a/Engine/Source/Runtime/Renderer/Private/PostProcess/DiaphragmDOF.cpp b/Engine/Source/Runtime/Renderer/Private/PostProcess/DiaphragmDOF.cpp
--- a/Engine/Source/Runtime/Renderer/Private/PostProcess/DiaphragmDOF.cpp
+++ b/Engine/Source/Runtime/Renderer/Private/PostProcess/DiaphragmDOF.cpp
@@ -133,7 +133,7 @@ TAutoConsoleVariable<int32> CVarDOFTemporalAAQuality(
 	ECVF_Scalability | ECVF_RenderThreadSafe);
 
 TAutoConsoleVariable<float> CVarScatterNeighborCompareMaxColor(
-	TEXT("r.DOF2.Scatter.NeighborCompareMaxColor"),
+	TEXT("r.DOF.Scatter.NeighborCompareMaxColor"),
 	10,
 	TEXT("Controles the linear color clamping upperbound applied before color of pixel and neighbors are compared.")
 	TEXT(" To low, and you may not scatter enough; to high you may scatter unnecessarily too much in highlights")
