Fix merge

diff --git a/Engine/Source/Runtime/Renderer/Private/VT/RuntimeVirtualTextureRender.cpp b/Engine/Source/Runtime/Renderer/Private/VT/RuntimeVirtualTextureRender.cpp
--- a/Engine/Source/Runtime/Renderer/Private/VT/RuntimeVirtualTextureRender.cpp
+++ b/Engine/Source/Runtime/Renderer/Private/VT/RuntimeVirtualTextureRender.cpp
@@ -281,7 +281,8 @@ namespace RuntimeVirtualTexture
 		// Uncached mesh processor
 		FDynamicMeshDrawCommandStorage MeshDrawCommandStorage;
 		FMeshCommandOneFrameArray AllocatedCommands;
-		FDynamicPassMeshDrawListContext DynamicMeshPassContext(MeshDrawCommandStorage, AllocatedCommands);
+		FGraphicsMinimalPipelineStateSet GraphicsMinimalPipelineStateSet;
+		FDynamicPassMeshDrawListContext DynamicMeshPassContext(MeshDrawCommandStorage, AllocatedCommands, GraphicsMinimalPipelineStateSet);
 		FRuntimeVirtualTextureMeshProcessor MeshProcessor(Scene, View, &DynamicMeshPassContext);
 
 		// Iterate over scene and collect visible virtual texture draw commands for this view
@@ -357,7 +358,7 @@ namespace RuntimeVirtualTexture
 			const uint32 InstanceFactor = 1;
 
 			SortAndMergeDynamicPassMeshDrawCommands(View->FeatureLevel, AllocatedCommands, MeshDrawCommandStorage, PrimitiveIdsBuffer, InstanceFactor);
-			SubmitMeshDrawCommands(AllocatedCommands, PrimitiveIdsBuffer, 0, bDynamicInstancing, InstanceFactor, RHICmdList);
+			SubmitMeshDrawCommands(AllocatedCommands, GraphicsMinimalPipelineStateSet, PrimitiveIdsBuffer, 0, bDynamicInstancing, InstanceFactor, RHICmdList);
 		}
 	}
 
