Fixes velocity buffer being released in post processing before TAA and DOF.

diff --git a/Engine/Source/Runtime/RenderCore/Public/RenderGraphUtils.h b/Engine/Source/Runtime/RenderCore/Public/RenderGraphUtils.h
--- a/Engine/Source/Runtime/RenderCore/Public/RenderGraphUtils.h
+++ b/Engine/Source/Runtime/RenderCore/Public/RenderGraphUtils.h
@@ -55,7 +55,12 @@ FORCEINLINE_DEBUGGABLE void ClearUnusedGraphResources(
 }
 
 
-/** Register external texture with fallback if the resource is invalid. */
+/**
+ * Register external texture with fallback if the resource is invalid.
+ *
+ * CAUTION: use this function very wisely. It may actually remove shader parameter validation
+ * failure when a pass is actually trying to access a resource not yet or no longer available.
+ */
 RENDERCORE_API FRDGTextureRef RegisterExternalTextureWithFallback(
 	FRDGBuilder& GraphBuilder,
 	const TRefCountPtr<IPooledRenderTarget>& ExternalPooledTexture,
diff --git a/Engine/Source/Runtime/Renderer/Private/PostProcess/PostProcessTemporalAA.cpp b/Engine/Source/Runtime/Renderer/Private/PostProcess/PostProcessTemporalAA.cpp
--- a/Engine/Source/Runtime/Renderer/Private/PostProcess/PostProcessTemporalAA.cpp
+++ b/Engine/Source/Runtime/Renderer/Private/PostProcess/PostProcessTemporalAA.cpp
@@ -820,7 +820,13 @@ void FRCPassPostProcessTemporalAA::Process(FRenderingCompositePassContext& Conte
 	FSceneRenderTargets& SceneContext = FSceneRenderTargets::Get(Context.RHICmdList);
 	FSceneViewFamilyBlackboard SceneBlackboard;
 	SetupSceneViewFamilyBlackboard(GraphBuilder, &SceneBlackboard);
-	
+
+	// FPostProcessing::Process() does a AdjustGBufferRefCount(RHICmdList, -1), therefore need to pass down reference on velocity buffer manually.
+	if (FRDGTextureRef SceneVelocityBuffer = CreateRDGTextureForOptionalInput(GraphBuilder, ePId_Input2, TEXT("SceneVelocity")))
+	{
+		SceneBlackboard.SceneVelocityBuffer = SceneVelocityBuffer;
+	}
+
 	FTAAPassParameters Parameters = SavedParameters;
 	Parameters.SceneColorInput = CreateRDGTextureForRequiredInput(GraphBuilder, ePId_Input0, TEXT("SceneColor"));
 	Parameters.SceneMetadataInput = CreateRDGTextureForOptionalInput(GraphBuilder, ePId_Input1, TEXT("SceneColor"));
diff --git a/Engine/Source/Runtime/Renderer/Private/PostProcess/PostProcessing.cpp b/Engine/Source/Runtime/Renderer/Private/PostProcess/PostProcessing.cpp
--- a/Engine/Source/Runtime/Renderer/Private/PostProcess/PostProcessing.cpp
+++ b/Engine/Source/Runtime/Renderer/Private/PostProcess/PostProcessing.cpp
@@ -1398,7 +1398,7 @@ void FPostProcessing::Process(FRHICommandListImmediate& RHICmdList, const FViewI
 			// Applies DOF and separate translucency,
 			{
 				FRenderingCompositePass* DiaphragmDOFPass = Context.Graph.RegisterPass(
-					new(FMemStack::Get()) TRCPassForRDG<2, 1>([bDepthOfField](FRenderingCompositePass* Pass, FRenderingCompositePassContext& InContext)
+					new(FMemStack::Get()) TRCPassForRDG<3, 1>([bDepthOfField](FRenderingCompositePass* Pass, FRenderingCompositePassContext& InContext)
 				{
 					FRDGBuilder GraphBuilder(InContext.RHICmdList);
 
@@ -1407,6 +1407,12 @@ void FPostProcessing::Process(FRHICommandListImmediate& RHICmdList, const FViewI
 	
 					FRDGTextureRef SceneColor = Pass->CreateRDGTextureForRequiredInput(GraphBuilder, ePId_Input0, TEXT("SceneColor"));
 					FRDGTextureRef LocalSeparateTranslucency = Pass->CreateRDGTextureForOptionalInput(GraphBuilder, ePId_Input1, TEXT("SeparateTranslucency"));
+					
+					// FPostProcessing::Process() does a AdjustGBufferRefCount(RHICmdList, -1), therefore need to pass down reference on velocity buffer manually.
+					if (FRDGTextureRef SceneVelocityBuffer = Pass->CreateRDGTextureForOptionalInput(GraphBuilder, ePId_Input2, TEXT("SceneVelocity")))
+					{
+						SceneBlackboard.SceneVelocityBuffer = SceneVelocityBuffer;
+					}
 
 					FRDGTextureRef NewSceneColor = SceneColor;
 
@@ -1430,6 +1436,7 @@ void FPostProcessing::Process(FRHICommandListImmediate& RHICmdList, const FViewI
 				}));
 				DiaphragmDOFPass->SetInput(ePId_Input0, Context.FinalOutput);
 				DiaphragmDOFPass->SetInput(ePId_Input1, SeparateTranslucency);
+				DiaphragmDOFPass->SetInput(ePId_Input2, VelocityInput);
 				Context.FinalOutput = FRenderingCompositeOutputRef(DiaphragmDOFPass, ePId_Output0);
 			}
 
diff --git a/Engine/Source/Runtime/Renderer/Private/SceneViewFamilyBlackboard.cpp b/Engine/Source/Runtime/Renderer/Private/SceneViewFamilyBlackboard.cpp
--- a/Engine/Source/Runtime/Renderer/Private/SceneViewFamilyBlackboard.cpp
+++ b/Engine/Source/Runtime/Renderer/Private/SceneViewFamilyBlackboard.cpp
@@ -17,14 +17,20 @@ void SetupSceneViewFamilyBlackboard(
 	const FSceneRenderTargets& SceneContext = FSceneRenderTargets::Get(GraphBuilder.RHICmdList);
 
 	*OutBlackboard = FSceneViewFamilyBlackboard();
+
+	// Should always have a depth buffer around allocated, since early z-pass is first.
 	OutBlackboard->SceneDepthBuffer = GraphBuilder.RegisterExternalTexture(SceneContext.SceneDepthZ);
-	OutBlackboard->SceneVelocityBuffer = RegisterExternalTextureWithFallback(GraphBuilder, SceneContext.SceneVelocity, GSystemTextures.BlackDummy);
-	OutBlackboard->SceneGBufferA = RegisterExternalTextureWithFallback(GraphBuilder, SceneContext.GBufferA, GSystemTextures.DefaultNormal8Bit, TEXT("GBufferA"));
-	OutBlackboard->SceneGBufferB = RegisterExternalTextureWithFallback(GraphBuilder, SceneContext.GBufferB, GSystemTextures.BlackDummy, TEXT("GBufferB"));
-	OutBlackboard->SceneGBufferC = RegisterExternalTextureWithFallback(GraphBuilder, SceneContext.GBufferC, GSystemTextures.BlackDummy, TEXT("GBufferC"));
-	OutBlackboard->SceneGBufferD = RegisterExternalTextureWithFallback(GraphBuilder, SceneContext.GBufferD, GSystemTextures.BlackDummy, TEXT("GBufferD"));
-	OutBlackboard->SceneGBufferE = RegisterExternalTextureWithFallback(GraphBuilder, SceneContext.GBufferE, GSystemTextures.BlackDummy, TEXT("GBufferE"));
 
+	// Registers all the scene texture from the scene context. No fallback is provided to catch mistake at shader parameter validation time
+	// when a pass is trying to access a resource before any other pass actually created it.
+	OutBlackboard->SceneVelocityBuffer = SceneContext.SceneVelocity ? GraphBuilder.RegisterExternalTexture(SceneContext.SceneVelocity) : nullptr;
+	OutBlackboard->SceneGBufferA = SceneContext.GBufferA ? GraphBuilder.RegisterExternalTexture(SceneContext.GBufferA) : nullptr;
+	OutBlackboard->SceneGBufferB = SceneContext.GBufferB ? GraphBuilder.RegisterExternalTexture(SceneContext.GBufferB) : nullptr;
+	OutBlackboard->SceneGBufferC = SceneContext.GBufferC ? GraphBuilder.RegisterExternalTexture(SceneContext.GBufferC) : nullptr;
+	OutBlackboard->SceneGBufferD = SceneContext.GBufferD ? GraphBuilder.RegisterExternalTexture(SceneContext.GBufferD) : nullptr;
+	OutBlackboard->SceneGBufferE = SceneContext.GBufferE ? GraphBuilder.RegisterExternalTexture(SceneContext.GBufferE) : nullptr;
+
+	// Ligthing channels might be disabled when all lights are on the same channel.
 	if (SceneContext.LightingChannels)
 	{
 		OutBlackboard->SceneLightingChannels = GraphBuilder.RegisterExternalTexture(SceneContext.LightingChannels, TEXT("LightingChannels"));
diff --git a/Engine/Source/Runtime/Renderer/Private/VelocityRendering.cpp b/Engine/Source/Runtime/Renderer/Private/VelocityRendering.cpp
--- a/Engine/Source/Runtime/Renderer/Private/VelocityRendering.cpp
+++ b/Engine/Source/Runtime/Renderer/Private/VelocityRendering.cpp
@@ -390,7 +390,10 @@ bool FDeferredShadingSceneRenderer::ShouldRenderVelocities() const
 
 		bool bSSRTemporal = ShouldRenderScreenSpaceReflections(View) && IsSSRTemporalPassRequired(View);
 
-		bNeedsVelocity |= bMotionBlur || bTemporalAA || bDistanceFieldAO || bSSRTemporal;
+		bool bRayTracing = IsRayTracingEnabled();
+		bool bDenoise = bRayTracing;
+
+		bNeedsVelocity |= bMotionBlur || bTemporalAA || bDistanceFieldAO || bSSRTemporal || bDenoise;
 	}
 
 	return bNeedsVelocity;
