Add dummy reflection capture uniform buffer to FMeshPassProcessorRenderState, used for drawing tile meshes where no scene exists,
but forward-rendering material shaders require something bound for this buffer

diff --git a/Engine/Source/Runtime/Renderer/Private/BasePassRendering.inl b/Engine/Source/Runtime/Renderer/Private/BasePassRendering.inl
--- a/Engine/Source/Runtime/Renderer/Private/BasePassRendering.inl
+++ b/Engine/Source/Runtime/Renderer/Private/BasePassRendering.inl
@@ -39,7 +39,7 @@ void TBasePassVertexShaderPolicyParamType<LightMapPolicyType>::GetShaderBindings
 	}
 	else
 	{
-		ensure(!ReflectionCaptureBuffer.IsBound());
+		ShaderBindings.Add(ReflectionCaptureBuffer, DrawRenderState.GetReflectionCaptureUniformBuffer());
 	}
 
 	LightMapPolicyType::GetVertexShaderBindings(
@@ -86,7 +86,7 @@ void TBasePassPixelShaderPolicyParamType<LightMapPolicyType>::GetShaderBindings(
 	}
 	else
 	{
-		ensure(!ReflectionCaptureBuffer.IsBound());
+		ShaderBindings.Add(ReflectionCaptureBuffer, DrawRenderState.GetReflectionCaptureUniformBuffer());
 	}
 
 	LightMapPolicyType::GetPixelShaderBindings(
diff --git a/Engine/Source/Runtime/Renderer/Private/Renderer.cpp b/Engine/Source/Runtime/Renderer/Private/Renderer.cpp
--- a/Engine/Source/Runtime/Renderer/Private/Renderer.cpp
+++ b/Engine/Source/Runtime/Renderer/Private/Renderer.cpp
@@ -139,6 +139,12 @@ void FRendererModule::DrawTileMesh(FRHICommandListImmediate& RHICmdList, FMeshPa
 		View.InitRHIResources();
 		DrawRenderState.SetViewUniformBuffer(View.ViewUniformBuffer);
 
+		if (!DrawRenderState.GetReflectionCaptureUniformBuffer())
+		{
+			FReflectionCaptureShaderData EmptyData;
+			DrawRenderState.SetReflectionCaptureUniformBuffer(TUniformBufferRef<FReflectionCaptureShaderData>::CreateUniformBufferImmediate(EmptyData, UniformBuffer_SingleFrame));
+		}
+
 		if (ShadingPath == EShadingPath::Mobile)
 		{
 			View.MobileDirectionalLightUniformBuffers[0] = TUniformBufferRef<FMobileDirectionalLightShaderParameters>::CreateUniformBufferImmediate(FMobileDirectionalLightShaderParameters(), UniformBuffer_SingleFrame);
diff --git a/Engine/Source/Runtime/Renderer/Public/MeshPassProcessor.h b/Engine/Source/Runtime/Renderer/Public/MeshPassProcessor.h
--- a/Engine/Source/Runtime/Renderer/Public/MeshPassProcessor.h
+++ b/Engine/Source/Runtime/Renderer/Public/MeshPassProcessor.h
@@ -915,6 +915,7 @@ struct FMeshPassProcessorRenderState
 		, DepthStencilAccess(FExclusiveDepthStencil::DepthRead_StencilRead)
 		, ViewUniformBuffer(SceneView.ViewUniformBuffer)
 		, InstancedViewUniformBuffer()
+		, ReflectionCaptureUniformBuffer()
 		, PassUniformBuffer(InPassUniformBuffer)
 		, StencilRef(0)
 	{
@@ -926,6 +927,7 @@ struct FMeshPassProcessorRenderState
 		, DepthStencilAccess(FExclusiveDepthStencil::DepthRead_StencilRead)
 		, ViewUniformBuffer(InViewUniformBuffer)
 		, InstancedViewUniformBuffer()
+		, ReflectionCaptureUniformBuffer()
 		, PassUniformBuffer(InPassUniformBuffer)
 		, StencilRef(0)
 	{
@@ -936,6 +938,7 @@ struct FMeshPassProcessorRenderState
 		, DepthStencilState(nullptr)
 		, ViewUniformBuffer()
 		, InstancedViewUniformBuffer()
+		, ReflectionCaptureUniformBuffer()
 		, PassUniformBuffer(nullptr)
 		, StencilRef(0)
 	{
@@ -947,6 +950,7 @@ struct FMeshPassProcessorRenderState
 		, DepthStencilAccess(DrawRenderState.DepthStencilAccess)
 		, ViewUniformBuffer(DrawRenderState.ViewUniformBuffer)
 		, InstancedViewUniformBuffer(DrawRenderState.InstancedViewUniformBuffer)
+		, ReflectionCaptureUniformBuffer(DrawRenderState.ReflectionCaptureUniformBuffer)
 		, PassUniformBuffer(DrawRenderState.PassUniformBuffer)
 		, StencilRef(DrawRenderState.StencilRef)
 	{
@@ -1013,7 +1017,17 @@ public:
 		return InstancedViewUniformBuffer.IsValid() ? InstancedViewUniformBuffer : reinterpret_cast<const TUniformBufferRef<FInstancedViewUniformShaderParameters>&>(ViewUniformBuffer);
 	}
 
-	FORCEINLINE_DEBUGGABLE void SetPassUniformBuffer(FUniformBufferRHIParamRef InPassUniformBuffer)
+	FORCEINLINE_DEBUGGABLE void SetReflectionCaptureUniformBuffer(FUniformBufferRHIParamRef InUniformBuffer)
+	{
+		ReflectionCaptureUniformBuffer = InUniformBuffer;
+	}
+
+	FORCEINLINE_DEBUGGABLE const FUniformBufferRHIRef& GetReflectionCaptureUniformBuffer() const
+	{
+		return ReflectionCaptureUniformBuffer;
+	}
+
+	FORCEINLINE_DEBUGGABLE void SetPassUniformBuffer(const FUniformBufferRHIRef& InPassUniformBuffer)
 	{
 		PassUniformBuffer = InPassUniformBuffer;
 	}
@@ -1041,6 +1055,10 @@ private:
 
 	TUniformBufferRef<FViewUniformShaderParameters>	ViewUniformBuffer;
 	TUniformBufferRef<FInstancedViewUniformShaderParameters> InstancedViewUniformBuffer;
+
+	/** Will be bound as reflection capture uniform buffer in case where scene is not available, typically set to dummy/empty buffer to avoid null binding */
+	FUniformBufferRHIRef			ReflectionCaptureUniformBuffer;
+
 	FUniformBufferRHIParamRef		PassUniformBuffer;
 	uint32							StencilRef;
 };
