Fixes an efficiency problem detected by RDG in SSGI, and add the SSGI quality on the draw event

diff --git a/Engine/Source/Runtime/Renderer/Private/ScreenSpaceDiffuseIndirect.cpp b/Engine/Source/Runtime/Renderer/Private/ScreenSpaceDiffuseIndirect.cpp
--- a/Engine/Source/Runtime/Renderer/Private/ScreenSpaceDiffuseIndirect.cpp
+++ b/Engine/Source/Runtime/Renderer/Private/ScreenSpaceDiffuseIndirect.cpp
@@ -90,7 +90,7 @@ void RenderScreenSpaceDiffuseIndirect( FRHICommandListImmediate& RHICmdList, FVi
 	FSceneRenderTargets& SceneContext = FSceneRenderTargets::Get( RHICmdList );
 
 	FRDGTexture* SceneColor	= GraphBuilder.RegisterExternalTexture( SceneContext.GetSceneColor() );
-	FRDGTexture* ScreenSpaceAO= GraphBuilder.RegisterExternalTexture( SceneContext.ScreenSpaceAO );
+	FRDGTexture* ScreenSpaceAO = GraphBuilder.CreateTexture(SceneContext.ScreenSpaceAO->GetDesc(), TEXT("SSRTAO"));
 	FRDGTexture* HZBTexture	= GraphBuilder.RegisterExternalTexture( View.HZB );
 	FRDGTexture* ColorTexture	= GraphBuilder.RegisterExternalTexture( TemporalAAHistory.RT[0] );
 
@@ -146,7 +146,7 @@ void RenderScreenSpaceDiffuseIndirect( FRHICommandListImmediate& RHICmdList, FVi
 	PassParameters->SceneTextures = CreateSceneTextureUniformBufferSingleDraw( RHICmdList, ESceneTextureSetupMode::SceneDepth | ESceneTextureSetupMode::GBuffers, View.FeatureLevel );
 
 	GraphBuilder.AddPass(
-		RDG_EVENT_NAME( "ScreenSpaceDiffuseIndirect %dx%d", View.ViewRect.Width(), View.ViewRect.Height() ),
+		RDG_EVENT_NAME( "ScreenSpaceDiffuseIndirect(Quality=%d) %dx%d", Quality, View.ViewRect.Width(), View.ViewRect.Height() ),
 		PassParameters,
 		ERenderGraphPassFlags::None,
 		[ PassParameters, &View, Quality ]( FRHICommandListImmediate& InRHICmdList )
@@ -193,5 +193,7 @@ void RenderScreenSpaceDiffuseIndirect( FRHICommandListImmediate& RHICmdList, FVi
 				EDRF_UseTriangleOptimization);
 		});
 
+	GraphBuilder.QueueTextureExtraction(ScreenSpaceAO, &SceneContext.ScreenSpaceAO);
+
 	GraphBuilder.Execute();
 }
