Revert CL 6525726 ( EyeAdaptation Histogram)

diff --git a/Engine/Source/Runtime/Renderer/Private/PostProcess/PostProcessing.cpp b/Engine/Source/Runtime/Renderer/Private/PostProcess/PostProcessing.cpp
--- a/Engine/Source/Runtime/Renderer/Private/PostProcess/PostProcessing.cpp
+++ b/Engine/Source/Runtime/Renderer/Private/PostProcess/PostProcessing.cpp
@@ -1625,9 +1625,8 @@ void FPostProcessing::Process(FRHICommandListImmediate& RHICmdList, const FViewI
 				}
 				else  // Use histogram version version
 				{
-					// If we have a histogram, we always add eye adaptation, if the engine show flag is disabled we set the ExposureScale in the texture to a fixed value
-					if (Histogram.IsValid())
-						AutoExposure.EyeAdaptation = AddPostProcessHistogramEyeAdaptation(Context, Histogram);
+					// We always add eye adaptation, if the engine show flag is disabled we set the ExposureScale in the texture to a fixed value
+					AutoExposure.EyeAdaptation = AddPostProcessHistogramEyeAdaptation(Context, Histogram);
 				}
 			}
 
