Ensure alpha is stored for postprocessing when requested

diff --git a/Engine/Source/Runtime/Renderer/Private/PostProcess/PostProcessTonemap.cpp b/Engine/Source/Runtime/Renderer/Private/PostProcess/PostProcessTonemap.cpp
--- a/Engine/Source/Runtime/Renderer/Private/PostProcess/PostProcessTonemap.cpp
+++ b/Engine/Source/Runtime/Renderer/Private/PostProcess/PostProcessTonemap.cpp
@@ -8,6 +8,7 @@
 #include "EngineGlobals.h"
 #include "ScenePrivate.h"
 #include "RendererModule.h"
+#include "PostProcess/PostProcessing.h"
 #include "PostProcess/SceneFilterRendering.h"
 #include "PostProcess/PostProcessCombineLUTs.h"
 #include "PostProcess/PostProcessMobile.h"
@@ -1122,7 +1123,7 @@ void FRCPassPostProcessTonemap::Process(FRenderingCompositePassContext& Context)
 			{
 				FGraphicsPipelineStateInitializer GraphicsPSOInit;
 				Context.RHICmdList.ApplyCachedRenderTargets(GraphicsPSOInit);
-				if (Context.View.AntiAliasingMethod == AAM_FXAA)
+				if (Context.View.AntiAliasingMethod == AAM_FXAA || FPostProcessing::HasAlphaChannelSupport())
 				{
 					// Alpha is used to store luminance for post process AA
 					GraphicsPSOInit.BlendState = TStaticBlendState<>::GetRHI();
