Reorganises all screen space ray tracing technics into a /Engine/Private/SSRT/ shader directory and unique C++ file.

diff --git a/Engine/Shaders/Private/BasePassPixelShader.usf b/Engine/Shaders/Private/BasePassPixelShader.usf
--- a/Engine/Shaders/Private/BasePassPixelShader.usf
+++ b/Engine/Shaders/Private/BasePassPixelShader.usf
@@ -29,7 +29,7 @@
 #include "DBufferDecalShared.ush"
 
 #define PREV_FRAME_COLOR	1
-#include "ScreenSpaceRayCast.ush"
+#include "SSRT/SSRTRayCast.ush"
 
 #if NEEDS_BASEPASS_PIXEL_FOGGING || NEEDS_BASEPASS_PIXEL_VOLUMETRIC_FOGGING
 	#include "HeightFogCommon.ush"
diff --git a/Engine/Shaders/Private/SSRT/SSRTDiffuseIndirect.usf b/Engine/Shaders/Private/SSRT/SSRTDiffuseIndirect.usf
new file mode 100644
--- /dev/null
+++ b/Engine/Shaders/Private/SSRT/SSRTDiffuseIndirect.usf
@@ -0,0 +1,153 @@
+// Copyright 1998-2019 Epic Games, Inc. All Rights Reserved.
+
+#define IS_SSGI_SHADER 1
+
+#include "SSRTRayCast.ush"	
+#include "../PostProcessCommon.ush"				
+#include "../DeferredShadingCommon.ush"
+#include "../Random.ush"
+#include "../BRDF.ush"
+#include "../MonteCarlo.ush"
+
+Texture2D		HZBTexture;
+SamplerState	HZBSampler;
+
+Texture2D		VelocityTexture;
+SamplerState	VelocitySampler;
+
+Texture2D		ColorTexture;
+SamplerState	ColorSampler;
+
+float4 HZBUvFactorAndInvFactor;
+
+float4 PrevScreenPositionScaleBias;
+float  PrevSceneColorPreExposureCorrection;
+
+void ScreenSpaceDiffuseIndirectPS(
+	in noperspective float4 UVAndScreenPos : TEXCOORD0,
+	float4 SvPosition : SV_POSITION,
+	out float4 OutColor : SV_Target0,
+	out float OutAO : SV_Target1 )
+{
+	float2 UV = UVAndScreenPos.xy;
+	float2 ScreenPos = UVAndScreenPos.zw;
+	uint2 PixelPos = (uint2)SvPosition.xy;
+
+	OutColor = 0;
+
+	FScreenSpaceData ScreenSpaceData = GetScreenSpaceData(UV);
+	FGBufferData GBuffer = ScreenSpaceData.GBuffer;
+
+	// No SSS because diffuse luminance is in alpha channel and we don't have dual source blending.
+	BRANCH if( GBuffer.ShadingModelID == 0 || GBuffer.ShadingModelID == SHADINGMODELID_SUBSURFACE_PROFILE )
+	{
+		discard;
+	}
+
+	const float3 N = GBuffer.WorldNormal;
+	const float SceneDepth = CalcSceneDepth(UV);
+	const float3 PositionTranslatedWorld = mul( float4( ScreenPos * SceneDepth, SceneDepth, 1 ), View.ScreenToTranslatedWorld ).xyz;
+
+#if QUALITY == 1
+	uint NumSteps = 8;
+	uint NumRays = 4;
+#elif QUALITY == 2
+	uint NumSteps = 8;
+	uint NumRays = 8;
+#elif QUALITY == 3
+	uint NumSteps = 8;
+	uint NumRays = 12;
+#else // QUALITY == 4
+	uint NumSteps = 12;
+	uint NumRays = 16;
+#endif
+
+	float2 Noise;
+	Noise.x = InterleavedGradientNoise( SvPosition.xy, View.StateFrameIndexMod8 );
+	Noise.y = InterleavedGradientNoise( SvPosition.xy, View.StateFrameIndexMod8 * 117 );
+		
+	//uint2 Random = 0x10000 * Noise;
+	uint2 Random = Rand3DPCG16( int3( PixelPos, View.StateFrameIndexMod8 ) ).xy;
+			
+	// TODO: lot of VGPR, worth recomputing every ray instead?
+	float3x3 TangentBasis = GetTangentBasis( N );
+
+	// Shoot multiple rays
+	LOOP for( int i = 0; i < NumRays; i++ )
+	{
+		// TOOD: VGPR, might be worth recomputing instead?
+		// SvPosition: is SvPosition really needed? Can recompute from PositionTranslatedWorld?
+		float StepOffset = Noise.x;
+		StepOffset -= 0.9;
+			
+		float2 E = Hammersley16( i, NumRays, Random );
+
+#if 0
+		float3 L = CosineSampleHemisphere( E ).xyz;
+		L = mul( L, TangentBasis );
+#else
+		float3 L;
+		L.xy = UniformSampleDiskConcentric( E );
+		L.z = sqrt( 1 - dot( L.xy, L.xy ) );
+		L = mul( L, TangentBasis );
+#endif
+				
+		float3 HitUVz;
+		float Level = 0;
+
+		bool bHit = RayCast(
+			HZBTexture, HZBSampler,
+			PositionTranslatedWorld, L, 1, SceneDepth, 
+			NumSteps, StepOffset,
+			HZBUvFactorAndInvFactor,
+			HitUVz,
+			Level
+		);
+
+#if 0	// Backface check
+		if( bHit )
+		{
+			float3 SampleNormal = GetGBufferData( HitUVz.xy ).WorldNormal;
+			bHit = dot( SampleNormal, L ) < 0;
+		}
+#endif
+
+		// if there was a hit
+		BRANCH if( bHit )
+		{
+			float2 SampleUV;
+			float Vignette;
+			ReprojectHit(PrevScreenPositionScaleBias, VelocityTexture, VelocitySampler, HitUVz, SampleUV, Vignette);
+
+#if 0
+			float4 SampleColor = SampleHCBLevel( ColorTexture, ColorSampler, SampleUV, Level, HZBUvFactorAndInvFactor );
+#else
+			float4 SampleColor = SampleScreenColor( ColorTexture, ColorSampler, SampleUV );
+#endif
+
+			// Backface modulation of hit surface
+			float3 SampleNormal = GetGBufferData( HitUVz.xy ).WorldNormal;
+			SampleColor.rgb *= saturate( 1 - dot( SampleNormal, L ) );
+
+			SampleColor.rgb *= rcp( 1 + Luminance(SampleColor.rgb) );
+
+			OutColor += SampleColor;
+		}
+	}
+
+	OutColor /= NumRays;
+	OutColor.rgb *= rcp( 1 - Luminance(OutColor.rgb) );
+	OutColor.rgb *= View.IndirectLightingColorScale;
+
+	OutColor.rgb *= GBuffer.DiffuseColor;
+	OutColor.a = 1 - OutColor.a;
+
+	//OutColor.rgb = OutColor.a;
+	//OutColor.a = 0;
+
+#if USE_PREEXPOSURE
+	OutColor.rgb *= PrevSceneColorPreExposureCorrection;
+#endif
+
+	OutAO = OutColor.a;
+}
\ No newline at end of file
diff --git a/Engine/Shaders/Private/SSRT/SSRTRayCast.ush b/Engine/Shaders/Private/SSRT/SSRTRayCast.ush
new file mode 100644
--- /dev/null
+++ b/Engine/Shaders/Private/SSRT/SSRTRayCast.ush
@@ -0,0 +1,268 @@
+// Copyright 1998-2019 Epic Games, Inc. All Rights Reserved.
+
+#pragma once
+
+#ifndef IS_SSGI_SHADER
+#define IS_SSGI_SHADER 0
+#endif
+
+#include "../Common.ush"	
+
+bool RayCast(
+	Texture2D Texture, SamplerState Sampler,
+	float3 RayOriginTranslatedWorld, float3 RayDirection,
+	float Roughness, float SceneDepth,
+	uint NumSteps, float StepOffset,
+	float4 HZBUvFactorAndInvFactor, 
+	out float3 OutHitUVz,
+	out float Level
+)
+{
+	// TODO provide RayStartUVz
+	const float4 RayStartClip	= mul( float4( RayOriginTranslatedWorld, 1 ), View.TranslatedWorldToClip );
+	const float4 RayEndClip		= mul( float4( RayOriginTranslatedWorld + RayDirection * SceneDepth, 1 ), View.TranslatedWorldToClip );
+
+	const float3 RayStartScreen = RayStartClip.xyz / RayStartClip.w;
+	const float3 RayEndScreen = RayEndClip.xyz / RayEndClip.w;
+
+	const float4 RayDepthClip	= RayStartClip + mul( float4( 0, 0, SceneDepth, 0 ), View.ViewToClip );
+	const float3 RayDepthScreen = RayDepthClip.xyz / RayDepthClip.w;
+
+	float3 RayStepScreen = RayEndScreen - RayStartScreen;
+
+	{
+		// Computes the scale down factor for RayStepScreen required to fit on the X and Y axis in order to clip it in the viewport
+		const float RayStepScreenInvFactor = 0.5 * length( RayStepScreen.xy );
+		const float2 S = 1 - max(abs(RayStepScreen.xy + RayStartScreen.xy * RayStepScreenInvFactor) - RayStepScreenInvFactor, 0.0f) / abs(RayStepScreen.xy);
+
+		// Rescales RayStepScreen accordingly
+		const float RayStepFactor = min(S.x, S.y) / RayStepScreenInvFactor;
+		RayStepScreen *= RayStepFactor;
+	}
+
+	float3 RayStartUVz = float3( (RayStartScreen.xy * float2( 0.5, -0.5 ) + 0.5) * HZBUvFactorAndInvFactor.xy, RayStartScreen.z );
+	float3 RayStepUVz  = float3(  RayStepScreen.xy  * float2( 0.5, -0.5 )		 * HZBUvFactorAndInvFactor.xy, RayStepScreen.z );
+	
+	const float Step = 1.0 / NumSteps;
+	#if IS_SSGI_SHADER
+		float CompareTolerance = max( abs( RayStepUVz.z ), ( RayStartScreen.z - RayDepthScreen.z ) * 2 ) * Step;
+	#else
+		float CompareTolerance = abs( RayDepthScreen.z - RayStartScreen.z ) * Step * 4;
+	#endif
+
+	// avoid bugs with early returns inside of loops on certain platform compilers.
+	bool bHit = false;
+	OutHitUVz = float3( 0, 0, 0 );
+
+	float LastDiff = 0;
+	Level = 1;
+
+	RayStepUVz *= Step;
+	float3 RayUVz = RayStartUVz + RayStepUVz * StepOffset;
+
+	LOOP
+	for( uint i = 0; i < NumSteps; i += 4 )
+	{
+		// Vectorized to group fetches
+		float4 SampleUV0 = RayUVz.xyxy + RayStepUVz.xyxy * float4( 1, 1, 2, 2 );
+		float4 SampleUV1 = RayUVz.xyxy + RayStepUVz.xyxy * float4( 3, 3, 4, 4 );
+		float4 SampleZ   = RayUVz.zzzz + RayStepUVz.zzzz * float4( 1, 2, 3, 4 );
+		
+		// Use lower res for farther samples
+		float4 SampleDepth;
+		SampleDepth.x = Texture.SampleLevel( Sampler, SampleUV0.xy, Level ).r;
+		SampleDepth.y = Texture.SampleLevel( Sampler, SampleUV0.zw, Level ).r;
+		Level += (8.0 / NumSteps) * Roughness;
+		
+		SampleDepth.z = Texture.SampleLevel( Sampler, SampleUV1.xy, Level ).r;
+		SampleDepth.w = Texture.SampleLevel( Sampler, SampleUV1.zw, Level ).r;
+		Level += (8.0 / NumSteps) * Roughness;
+
+		float4 DepthDiff = SampleZ - SampleDepth;
+		bool4 Hit = abs( DepthDiff + CompareTolerance ) < CompareTolerance;
+
+		BRANCH if( any( Hit ) )
+		{
+			if(IS_SSGI_SHADER)
+			{
+				// If hit set to intersect time. If missed set to beyond end of ray
+				float4 HitTime = Hit ? float4( 1, 2, 3, 4 ) : 5;
+
+				// Take closest hit
+				float Time1 = min( min3( HitTime.x, HitTime.y, HitTime.z ), HitTime.w );
+				float Time0 = Time1 - 1;
+
+				const uint NumBinarySteps = Roughness < 0.2 ? 4 : 0;
+
+				// Binary search
+				for( uint j = 0; j < NumBinarySteps; j++ )
+				{
+					CompareTolerance *= 0.5;
+
+					float  MidTime = 0.5 * ( Time0 + Time1 );
+					float3 MidUVz = RayUVz + RayStepUVz * MidTime;
+					float  MidDepth = Texture.SampleLevel( Sampler, MidUVz.xy, Level ).r;
+					float  MidDepthDiff = MidUVz.z - MidDepth;
+
+					if( abs( MidDepthDiff + CompareTolerance ) < CompareTolerance )
+					{
+						Time1 = MidTime;
+					}
+					else
+					{
+						Time0 = MidTime;
+					}
+				}
+			
+				OutHitUVz = RayUVz + RayStepUVz * Time1;
+			}
+			else
+			{
+				float DepthDiff0 = DepthDiff[2];
+				float DepthDiff1 = DepthDiff[3];
+				float Time0 = 3;
+
+				FLATTEN if( Hit[2] )
+				{
+					DepthDiff0 = DepthDiff[1];
+					DepthDiff1 = DepthDiff[2];
+					Time0 = 2;
+				}
+				FLATTEN if( Hit[1] )
+				{
+					DepthDiff0 = DepthDiff[0];
+					DepthDiff1 = DepthDiff[1];
+					Time0 = 1;
+				}
+				FLATTEN if( Hit[0] )
+				{
+					DepthDiff0 = LastDiff;
+					DepthDiff1 = DepthDiff[0];
+					Time0 = 0;
+				}
+
+				float Time1 = Time0 + 1;
+	#if 0
+				// Binary search
+				for( uint j = 0; j < 4; j++ )
+				{
+					CompareTolerance *= 0.5;
+
+					float  MidTime = 0.5 * ( Time0 + Time1 );
+					float3 MidUVz = RayUVz + RayStepUVz * MidTime;
+					float  MidDepth = Texture.SampleLevel( Sampler, MidUVz.xy, Level ).r;
+					float  MidDepthDiff = MidUVz.z - MidDepth;
+
+					if( abs( MidDepthDiff + CompareTolerance ) < CompareTolerance )
+					{
+						DepthDiff1	= MidDepthDiff;
+						Time1		= MidTime;
+					}
+					else
+					{
+						DepthDiff0	= MidDepthDiff;
+						Time0		= MidTime;
+					}
+				}
+	#endif
+
+				// Find more accurate hit using line segment intersection
+				float TimeLerp = saturate( DepthDiff0 / (DepthDiff0 - DepthDiff1) );
+				float IntersectTime = Time0 + TimeLerp;
+				//float IntersectTime = lerp( Time0, Time1, TimeLerp );
+				
+				OutHitUVz = RayUVz + RayStepUVz * IntersectTime;
+			}
+
+			bHit = true;
+			break;
+		}
+
+		LastDiff = DepthDiff.w;
+
+		RayUVz += 4 * RayStepUVz;
+	}
+	
+	OutHitUVz.xy *= HZBUvFactorAndInvFactor.zw;
+	OutHitUVz.xy = OutHitUVz.xy * float2( 2, -2 ) + float2( -1, 1 );
+	OutHitUVz.xy = OutHitUVz.xy * View.ScreenPositionScaleBias.xy + View.ScreenPositionScaleBias.wz;
+
+	return bHit;
+}
+
+float ComputeHitVignetteFromScreenPos(float2 ScreenPos)
+{
+	float2 Vignette = saturate(abs(ScreenPos) * 5 - 4);
+	
+	//PrevScreen sometimes has NaNs or Infs.  DX11 is protected because saturate turns NaNs -> 0.
+	//Do a SafeSaturate so other platforms get the same protection.
+	return SafeSaturate(1.0 - dot(Vignette, Vignette));
+}
+
+void ReprojectHit(float4 PrevScreenPositionScaleBias, float3 HitUVz, out float2 OutPrevUV, out float OutVignette)
+{
+	// Camera motion for pixel (in ScreenPos space).
+	float2 ThisScreen = (HitUVz.xy - View.ScreenPositionScaleBias.wz) / View.ScreenPositionScaleBias.xy;
+	float4 ThisClip = float4( ThisScreen, HitUVz.z, 1 );
+	float4 PrevClip = mul( ThisClip, View.ClipToPrevClip );
+	float2 PrevScreen = PrevClip.xy / PrevClip.w;
+	float2 PrevUV = PrevScreen.xy * PrevScreenPositionScaleBias.xy + PrevScreenPositionScaleBias.zw;
+
+	OutVignette = min(ComputeHitVignetteFromScreenPos(ThisScreen), ComputeHitVignetteFromScreenPos(PrevScreen));
+	OutPrevUV = PrevUV;
+}
+
+void ReprojectHit(float4 PrevScreenPositionScaleBias, Texture2D Texture, SamplerState Sampler, float3 HitUVz, out float2 OutPrevUV, out float OutVignette)
+{
+	// Camera motion for pixel (in ScreenPos space).
+	float2 ThisScreen = (HitUVz.xy - View.ScreenPositionScaleBias.wz) / View.ScreenPositionScaleBias.xy;
+	float4 ThisClip = float4( ThisScreen, HitUVz.z, 1 );
+	float4 PrevClip = mul( ThisClip, View.ClipToPrevClip );
+	float2 PrevScreen = PrevClip.xy / PrevClip.w;
+
+	float2 Velocity = Texture.SampleLevel( Sampler, HitUVz.xy, 0 ).xy;
+	if( Velocity.x > 0.0 )
+	{
+		PrevScreen = ThisClip.xy - DecodeVelocityFromTexture( Velocity );
+	}
+
+	float2 PrevUV = PrevScreen.xy * PrevScreenPositionScaleBias.xy + PrevScreenPositionScaleBias.zw;
+	
+	OutVignette = min(ComputeHitVignetteFromScreenPos(ThisScreen), ComputeHitVignetteFromScreenPos(PrevScreen));
+	OutPrevUV = PrevUV;
+}
+
+float ComputeRayHitSqrDistance(float3 OriginTranslatedWorld, float3 HitUVz)
+{
+	// ALU get factored out with ReprojectHit.
+	float2 HitScreenPos = (HitUVz.xy - View.ScreenPositionScaleBias.wz) / View.ScreenPositionScaleBias.xy;
+	float HitSceneDepth = ConvertFromDeviceZ(HitUVz.z);
+
+	float3 HitTranslatedWorld = mul(float4(HitScreenPos * HitSceneDepth, HitSceneDepth, 1), View.ScreenToTranslatedWorld).xyz;
+
+	return length2(OriginTranslatedWorld - HitTranslatedWorld);
+}
+
+float4 SampleScreenColor(Texture2D Texture, SamplerState Sampler, float2 UV)
+{
+	float4 OutColor;
+
+	OutColor.rgb = Texture.SampleLevel( Sampler, UV, 0 ).rgb;
+	// Transform NaNs to black, transform negative colors to black.
+	OutColor.rgb = -min(-OutColor.rgb, 0.0);
+	OutColor.a = 1;
+	
+	return OutColor;
+}
+
+float4 SampleHCBLevel( Texture2D Texture, SamplerState Sampler, float2 UV, float Level, float4 HZBUvFactorAndInvFactor )
+{
+	float4 OutColor;
+
+	OutColor.rgb = Texture.SampleLevel( Sampler, UV * HZBUvFactorAndInvFactor.xy, Level ).rgb;
+	// Transform NaNs to black, transform negative colors to black.
+	OutColor.rgb = -min(-OutColor.rgb, 0.0);
+	OutColor.a = 1;
+
+	return OutColor;
+}
\ No newline at end of file
diff --git a/Engine/Shaders/Private/SSRT/SSRTReflections.usf b/Engine/Shaders/Private/SSRT/SSRTReflections.usf
new file mode 100644
--- /dev/null
+++ b/Engine/Shaders/Private/SSRT/SSRTReflections.usf
@@ -0,0 +1,355 @@
+// Copyright 1998-2019 Epic Games, Inc. All Rights Reserved.
+
+#include "SSRTRayCast.ush"
+#include "../Random.ush"
+#include "../BRDF.ush"
+#include "../MonteCarlo.ush"
+#include "../SceneViewFamilyBlackboard.ush"
+#include "../ScreenSpaceDenoise/SSDPublic.ush"
+
+// Generate vector truncation warnings to errors.
+#pragma warning(error: 3206)
+
+
+// .r:Intensity in 0..1 range .g:RoughnessMaskMul, b:EnableDiscard for FPostProcessScreenSpaceReflectionsStencilPS, a:(bTemporalAAIsOn?TemporalAAIndex:StateFrameIndexMod8)*1551
+float4 SSRParams;
+
+float PrevSceneColorPreExposureCorrection;
+
+float4 HZBUvFactorAndInvFactor;
+
+// Equivalent of View.ScreenPositionScaleBias, but for SSR.
+float4 PrevScreenPositionScaleBias;
+
+Texture2D SceneColor;
+SamplerState SceneColorSampler;
+
+Texture2D HZB;
+SamplerState HZBSampler;
+
+#if COMPILER_HLSLCC || PS4_PROFILE
+#define INFINITE_WORKAROUND		3.402823e+38
+#else
+#define INFINITE_WORKAROUND		1.#INF
+#endif
+
+
+float GetRoughness(const in FGBufferData GBuffer)
+{
+	float Roughness = GBuffer.Roughness;
+
+	BRANCH
+	if( GBuffer.ShadingModelID == SHADINGMODELID_CLEAR_COAT )
+	{
+		const float ClearCoat			= GBuffer.CustomData.x;
+		const float ClearCoatRoughness	= GBuffer.CustomData.y;
+
+		Roughness = lerp( Roughness, ClearCoatRoughness, ClearCoat );
+	}
+
+	return Roughness;
+}
+
+float GetRoughnessFade(in float Roughness)
+{
+	// mask SSR to reduce noise and for better performance, roughness of 0 should have SSR, at MaxRoughness we fade to 0
+	return min(Roughness * SSRParams.y + 2, 1.0);
+}
+
+// Used only if r.SSR.Stencil = 1
+//
+// This shader generate the stencil buffer according to if a pixel needs to
+// compute SSR. If it doesn't, it initializes the pixel's color.
+//
+// This optimization use the stencil to test to prunes out the pixel
+// that dones't need to compute SSR, before pixel shader invocation are
+// packed into front lines.
+void ScreenSpaceReflectionsStencilPS(
+	float4 SvPosition : SV_POSITION
+	, out float4 OutColor : SV_Target0
+#if SSR_OUTPUT_FOR_DENOISER
+	, out float4 OutClosestHitDistance : SV_Target1
+#endif
+)
+{
+	float2 UV = SvPosition.xy * View.BufferSizeAndInvSize.zw;
+	
+	FGBufferData GBuffer = GetGBufferDataFromSceneBlackboard(UV);
+	
+	const float Roughness = GetRoughness(GBuffer);
+	const float RoughnessFade = GetRoughnessFade(Roughness);
+	
+	if (RoughnessFade > 0.0 && GBuffer.ShadingModelID != 0)
+	{
+		// we are going to compute SSR for this pixel, so we discard this
+		// pixel shader invocation to not overwrite the stencil buffer and
+		// tehrefore execute ScreenSpaceReflectionsPS() for this pixel. 
+		discard;
+	}
+	
+	// we are not going to compute SSR for this pixel, so we clear the color
+	// since ScreenSpaceReflectionsPS() won't be executed in this pixel.
+	OutColor = 0;
+	#if SSR_OUTPUT_FOR_DENOISER
+		OutClosestHitDistance = DENOISER_INVALID_HIT_DISTANCE;
+	#endif
+}
+
+void ScreenSpaceReflectionsPS(
+	float4 SvPosition : SV_POSITION
+	, out float4 OutColor : SV_Target0
+#if SSR_OUTPUT_FOR_DENOISER
+	, out float4 OutClosestHitDistance : SV_Target1
+#endif
+)
+{
+	float2 UV = SvPosition.xy * View.BufferSizeAndInvSize.zw;
+	float2 ScreenPos = ViewportUVToScreenPos((SvPosition.xy - View.ViewRectMin.xy) * View.ViewSizeAndInvSize.zw);
+	uint2 PixelPos = (uint2)SvPosition.xy;
+
+	OutColor = 0;
+	
+	#if SSR_OUTPUT_FOR_DENOISER
+		OutClosestHitDistance = -2.0;
+	#endif
+
+	FGBufferData GBuffer = GetGBufferDataFromSceneBlackboard(UV);
+	
+	float Roughness = GetRoughness(GBuffer);
+	float RoughnessFade = GetRoughnessFade(Roughness);
+
+	// Early out. Useless if using the stencil prepass.
+	BRANCH if( RoughnessFade <= 0.0 || GBuffer.ShadingModelID == 0 )
+	{
+		return;
+	}
+
+	const float3 N = GBuffer.WorldNormal;
+	const float SceneDepth = GBuffer.Depth;
+	const float3 PositionTranslatedWorld = mul( float4( ScreenPos * SceneDepth, SceneDepth, 1 ), View.ScreenToTranslatedWorld ).xyz;
+	const float3 V = normalize(View.TranslatedWorldCameraOrigin - PositionTranslatedWorld);
+
+	float a = Roughness * Roughness;
+	float a2 = a * a;
+	
+	float NoV = saturate( dot( N, V ) );
+	float G_SmithV = 2 * NoV / (NoV + sqrt(NoV * (NoV - NoV * a2) + a2));
+
+	float ClosestHitDistanceSqr = INFINITE_WORKAROUND;
+
+#if SSR_QUALITY == 1
+	uint NumSteps = 8;
+	uint NumRays = 1;
+	bool bGlossy = false;
+#elif SSR_QUALITY == 2
+	uint NumSteps = 16;
+	uint NumRays = 1;
+	#if SSR_OUTPUT_FOR_DENOISER
+		bool bGlossy = true;
+	#else
+		bool bGlossy = false;
+	#endif
+#elif SSR_QUALITY == 3
+	uint NumSteps = 8;
+	uint NumRays = 4;
+	bool bGlossy = true;
+#else // SSR_QUALITY == 4
+	uint NumSteps = 12;
+	uint NumRays = 12;
+	bool bGlossy = true;
+#endif
+	
+	if( NumRays > 1 )
+	{
+		float2 Noise;
+		Noise.x = InterleavedGradientNoise( SvPosition.xy, View.StateFrameIndexMod8 );
+		Noise.y = InterleavedGradientNoise( SvPosition.xy, View.StateFrameIndexMod8 * 117 );
+		
+		//uint2 Random = 0x10000 * Noise;
+		uint2 Random = Rand3DPCG16( int3( PixelPos, View.StateFrameIndexMod8 ) ).xy;
+			
+		float3x3 TangentBasis = GetTangentBasis( N );
+		float3 TangentV = mul( TangentBasis, V );
+
+		float Count = 0;
+
+		if( Roughness < 0.1 )
+		{
+			NumSteps = min( NumSteps * NumRays, 24u );
+			NumRays = 1;
+		}
+
+		// Shoot multiple rays
+		LOOP for( uint i = 0; i < NumRays; i++ )
+		{
+			float StepOffset = Noise.x;
+			#if 0 // TODO
+				StepOffset -= 0.9;
+			#else
+				StepOffset -= 0.5;
+			#endif
+			
+			float2 E = Hammersley16( i, NumRays, Random );
+
+#if 1
+			float3 H = mul( ImportanceSampleVisibleGGX(E, a2, TangentV ).xyz, TangentBasis );
+			float3 L = 2 * dot( V, H ) * H - V;
+#elif 0
+			float3 H = mul( ImportanceSampleGGX( E, a2 ).xyz, TangentBasis );
+			float3 L = 2 * dot( V, H ) * H - V;
+#elif 0
+			float3 L = CosineSampleHemisphere( E, N ).xyz;
+#elif 0
+			float3 L = CosineSampleHemisphere( E ).xyz;
+			L = mul( L, TangentBasis );
+#else
+			float3 L;
+			L.xy = UniformSampleDiskConcentric( E );
+			L.z = sqrt( 1 - dot( L.xy, L.xy ) );
+			L = mul( L, TangentBasis );
+			float3 H = normalize(V + L);
+#endif
+
+			float NoL = saturate( dot(N, L) );
+			float NoH = saturate( dot(N, H) );
+			float VoH = saturate( dot(V, H) );
+		
+			float D = D_GGX( a2, NoH );
+			float G_SmithL = 2 * NoL / (NoL + sqrt(NoL * (NoL - NoL * a2) + a2));
+			float Vis = Vis_Smith( a2, NoV,  NoL );
+			float3 F = F_Schlick( GBuffer.SpecularColor, VoH );
+				
+			float3 HitUVz;
+			float Level = 0;
+
+			bool bHit = RayCast(
+				HZB, HZBSampler,
+				PositionTranslatedWorld, L, Roughness, SceneDepth, 
+				NumSteps, StepOffset,
+				HZBUvFactorAndInvFactor,
+				HitUVz,
+				Level
+			);
+
+#if 0	// Backface check
+			if( bHit )
+			{
+				float3 SampleNormal = GetGBufferData( HitUVz.xy ).WorldNormal;
+				bHit = dot( SampleNormal, L ) < 0;
+			}
+#endif
+
+			// if there was a hit
+			BRANCH if( bHit )
+			{
+				ClosestHitDistanceSqr = min(ClosestHitDistanceSqr, ComputeRayHitSqrDistance(PositionTranslatedWorld, HitUVz));
+
+				float2 SampleUV;
+				float Vignette;
+				ReprojectHit(PrevScreenPositionScaleBias, SceneVelocityBuffer, SceneVelocityBufferSampler, HitUVz, SampleUV, Vignette);
+
+				float4 SampleColor = SampleScreenColor( SceneColor, SceneColorSampler, SampleUV ) * Vignette;
+
+				SampleColor.rgb *= rcp( 1 + Luminance(SampleColor.rgb) );
+
+				// Correct integration applies DGF below but for speed we apply EnvBRDF later when compositing
+#if 0
+				// CosineSampleHemisphere
+				// PDF = NoL / PI,
+				SampleColor.rgb *= F * (D * Vis * PI);
+#elif 0
+				// ImportanceSampleGGX
+				// PDF = D * NoH / (4 * VoH),
+				SampleColor.rgb *= F * ( NoL * Vis * (4 * VoH / NoH) );
+#elif 0
+				// ImportanceSampleVisibleGGX
+				// PDF = G_SmithV * VoH * D / NoV / (4 * VoH)
+				// PDF = G_SmithV * D / (4 * NoV);
+				SampleColor.rgb *= F * G_SmithL;
+#endif
+
+				OutColor += SampleColor;
+			}
+		}
+
+		//OutColor /= NumRays;
+		OutColor /= max( NumRays, 0.0001 );
+		OutColor.rgb *= rcp( 1 - Luminance(OutColor.rgb) );
+
+		if (ClosestHitDistanceSqr == INFINITE_WORKAROUND)
+		{
+			ClosestHitDistanceSqr = DENOISER_MISS_HIT_DISTANCE;
+		}
+	}
+	else
+	{
+		float StepOffset = InterleavedGradientNoise(SvPosition.xy, View.StateFrameIndexMod8);
+		#if 0 // TODO
+			StepOffset -= 0.9;
+		#else
+			StepOffset -= 0.5;
+		#endif
+		
+		float3 L;
+		if (bGlossy)
+		{
+			float2 E = Rand1SPPDenoiserInput(PixelPos);
+			
+			#if SSR_OUTPUT_FOR_DENOISER
+			{
+				E.x *= 1 - GGX_IMPORTANT_SAMPLE_BIAS;
+			}
+			#endif
+			
+			float3x3 TangentBasis = GetTangentBasis( N );
+			float3 TangentV = mul( TangentBasis, V );
+
+			float3 H = mul( ImportanceSampleVisibleGGX(E, a2, TangentV ).xyz, TangentBasis );
+			L = 2 * dot( V, H ) * H - V;
+		}
+		else
+		{
+			L = reflect( -V, N );
+		}
+		
+		float3 HitUVz;
+		float Level = 0;
+		bool bHit = RayCast(
+			HZB, HZBSampler,
+			PositionTranslatedWorld, L, Roughness, SceneDepth,
+			NumSteps, StepOffset,
+			HZBUvFactorAndInvFactor,
+			HitUVz,
+			Level
+		);
+
+		BRANCH if( bHit )
+		{
+			ClosestHitDistanceSqr = ComputeRayHitSqrDistance(PositionTranslatedWorld, HitUVz);
+
+			float2 SampleUV;
+			float Vignette;
+			ReprojectHit(PrevScreenPositionScaleBias, SceneVelocityBuffer, SceneVelocityBufferSampler, HitUVz, SampleUV, Vignette);
+
+			OutColor = SampleScreenColor(SceneColor, SceneColorSampler, SampleUV) * Vignette;
+		}
+		else
+		{
+			ClosestHitDistanceSqr = DENOISER_MISS_HIT_DISTANCE;
+		}
+	}
+	
+	OutColor *= RoughnessFade;
+	OutColor *= SSRParams.r;
+
+#if USE_PREEXPOSURE
+	OutColor.rgb *= PrevSceneColorPreExposureCorrection;
+#endif
+	
+	// Output closest hit distance for denoiser.
+	#if SSR_OUTPUT_FOR_DENOISER
+	{
+		OutClosestHitDistance = ClosestHitDistanceSqr > 0 ? ClampToHalfFloatRange(sqrt(ClosestHitDistanceSqr)) : ClosestHitDistanceSqr;
+	}
+	#endif
+}
\ No newline at end of file
diff --git a/Engine/Shaders/Private/ScreenSpaceDiffuseIndirect.usf b/Engine/Shaders/Private/ScreenSpaceDiffuseIndirect.usf
deleted file mode 100644
diff --git a/Engine/Shaders/Private/ScreenSpaceRayCast.ush b/Engine/Shaders/Private/ScreenSpaceRayCast.ush
deleted file mode 100644
diff --git a/Engine/Shaders/Private/ScreenSpaceReflections.usf b/Engine/Shaders/Private/ScreenSpaceReflections.usf
deleted file mode 100644
diff --git a/Engine/Source/Runtime/Renderer/Private/DeferredShadingRenderer.cpp b/Engine/Source/Runtime/Renderer/Private/DeferredShadingRenderer.cpp
--- a/Engine/Source/Runtime/Renderer/Private/DeferredShadingRenderer.cpp
+++ b/Engine/Source/Runtime/Renderer/Private/DeferredShadingRenderer.cpp
@@ -10,7 +10,6 @@
 #include "ScenePrivate.h"
 #include "ScreenRendering.h"
 #include "PostProcess/SceneFilterRendering.h"
-#include "PostProcess/ScreenSpaceReflections.h"
 #include "PostProcess/PostProcessSubsurface.h"
 #include "CompositionLighting/CompositionLighting.h"
 #include "FXSystem.h"
@@ -33,6 +32,7 @@
 #include "RayTracingDynamicGeometryCollection.h"
 #include "SceneViewFamilyBlackboard.h"
 #include "ScreenSpaceDenoise.h"
+#include "ScreenSpaceRayTracing.h"
 #include "RayTracing/RaytracingOptions.h"
 #include "RayTracingDefinitions.h"
 #include "RayTracingInstance.h"
@@ -795,7 +795,6 @@ bool FDeferredShadingSceneRenderer::DispatchRayTracingWorldUpdates(FRHICommandLi
 #endif // RHI_RAYTRACING
 
 extern bool IsLpvIndirectPassRequired(const FViewInfo& View);
-void RenderScreenSpaceDiffuseIndirect( FRHICommandListImmediate& RHICmdList, FViewInfo& View, TRefCountPtr<IPooledRenderTarget>& VelocityRT );
 
 static TAutoConsoleVariable<float> CVarStallInitViews(
 	TEXT("CriticalPathStall.AfterInitViews"),
diff --git a/Engine/Source/Runtime/Renderer/Private/PostProcess/ScreenSpaceReflections.cpp b/Engine/Source/Runtime/Renderer/Private/PostProcess/ScreenSpaceReflections.cpp
deleted file mode 100644
diff --git a/Engine/Source/Runtime/Renderer/Private/PostProcess/ScreenSpaceReflections.h b/Engine/Source/Runtime/Renderer/Private/PostProcess/ScreenSpaceReflections.h
deleted file mode 100644
diff --git a/Engine/Source/Runtime/Renderer/Private/ReflectionEnvironment.cpp b/Engine/Source/Runtime/Renderer/Private/ReflectionEnvironment.cpp
--- a/Engine/Source/Runtime/Renderer/Private/ReflectionEnvironment.cpp
+++ b/Engine/Source/Runtime/Renderer/Private/ReflectionEnvironment.cpp
@@ -25,13 +25,13 @@
 #include "PostProcess/PostProcessing.h"
 #include "PostProcess/PostProcessSubsurface.h"
 #include "PostProcess/PostProcessTemporalAA.h"
-#include "PostProcess/ScreenSpaceReflections.h"
 #include "LightRendering.h"
 #include "LightPropagationVolumeSettings.h"
 #include "PipelineStateCache.h"
 #include "DistanceFieldAmbientOcclusion.h"
 #include "SceneViewFamilyBlackboard.h"
 #include "ScreenSpaceDenoise.h"
+#include "ScreenSpaceRayTracing.h"
 #include "RayTracing/RaytracingOptions.h"
 
 DECLARE_GPU_STAT_NAMED(ReflectionEnvironment, TEXT("Reflection Environment"));
diff --git a/Engine/Source/Runtime/Renderer/Private/ScreenSpaceDiffuseIndirect.cpp b/Engine/Source/Runtime/Renderer/Private/ScreenSpaceDiffuseIndirect.cpp
deleted file mode 100644
diff --git a/Engine/Source/Runtime/Renderer/Private/ScreenSpaceRayTracing.cpp b/Engine/Source/Runtime/Renderer/Private/ScreenSpaceRayTracing.cpp
new file mode 100644
--- /dev/null
+++ b/Engine/Source/Runtime/Renderer/Private/ScreenSpaceRayTracing.cpp
@@ -0,0 +1,634 @@
+// Copyright 1998-2019 Epic Games, Inc. All Rights Reserved.
+
+#include "ScreenSpaceRayTracing.h"
+#include "RenderGraph.h"
+#include "PixelShaderUtils.h"
+#include "ScreenPass.h"
+#include "ScenePrivate.h"
+#include "SceneViewFamilyBlackboard.h"
+
+
+static TAutoConsoleVariable<int32> CVarSSRQuality(
+	TEXT("r.SSR.Quality"),
+	3,
+	TEXT("Whether to use screen space reflections and at what quality setting.\n")
+	TEXT("(limits the setting in the post process settings which has a different scale)\n")
+	TEXT("(costs performance, adds more visual realism but the technique has limits)\n")
+	TEXT(" 0: off (default)\n")
+	TEXT(" 1: low (no glossy)\n")
+	TEXT(" 2: medium (no glossy)\n")
+	TEXT(" 3: high (glossy/using roughness, few samples)\n")
+	TEXT(" 4: very high (likely too slow for real-time)"),
+	ECVF_Scalability | ECVF_RenderThreadSafe);
+
+static TAutoConsoleVariable<int32> CVarSSRTemporal(
+	TEXT("r.SSR.Temporal"),
+	0,
+	TEXT("Defines if we use the temporal smoothing for the screen space reflection\n")
+	TEXT(" 0 is off (for debugging), 1 is on (default)"),
+	ECVF_RenderThreadSafe);
+
+static TAutoConsoleVariable<int32> CVarSSRStencil(
+	TEXT("r.SSR.Stencil"),
+	0,
+	TEXT("Defines if we use the stencil prepass for the screen space reflection\n")
+	TEXT(" 0 is off (default), 1 is on"),
+	ECVF_RenderThreadSafe);
+
+static TAutoConsoleVariable<int32> CVarSSGIQuality(
+	TEXT("r.SSGI.Quality"),
+	0,
+	TEXT("Whether to use screen space diffuse indirect and at what quality setting.\n"),
+	ECVF_Scalability | ECVF_RenderThreadSafe);
+
+
+DECLARE_GPU_STAT_NAMED(ScreenSpaceReflections, TEXT("ScreenSpace Reflections"));
+DECLARE_GPU_STAT_NAMED(ScreenSpaceDiffuseIndirect, TEXT("Screen Space Diffuse Indirect"));
+
+
+bool ShouldRenderScreenSpaceReflections(const FViewInfo& View)
+{
+	if(!View.Family->EngineShowFlags.ScreenSpaceReflections)
+	{
+		return false;
+	}
+
+	if(!View.State)
+	{
+		// not view state (e.g. thumbnail rendering?), no HZB (no screen space reflections or occlusion culling)
+		return false;
+	}
+
+	int SSRQuality = CVarSSRQuality.GetValueOnRenderThread();
+
+	if(SSRQuality <= 0)
+	{
+		return false;
+	}
+
+	if(View.FinalPostProcessSettings.ScreenSpaceReflectionIntensity < 1.0f)
+	{
+		return false;
+	}
+
+	if (IsAnyForwardShadingEnabled(View.GetShaderPlatform()))
+	{
+		return false;
+	}
+
+	return true;
+}
+
+bool ShouldRenderScreenSpaceDiffuseIndirect( const FViewInfo& View )
+{
+	int Quality = CVarSSGIQuality.GetValueOnRenderThread();
+
+	if( Quality <= 0 )
+	{
+		return false;
+	}
+
+	if (IsAnyForwardShadingEnabled(View.GetShaderPlatform()))
+	{
+		return false;
+	}
+
+	return true;
+}
+
+bool IsSSRTemporalPassRequired(const FViewInfo& View)
+{
+	check(ShouldRenderScreenSpaceReflections(View));
+
+	if (!View.State)
+	{
+		return false;
+	}
+	return View.AntiAliasingMethod != AAM_TemporalAA || CVarSSRTemporal.GetValueOnRenderThread() != 0;
+}
+
+namespace
+{
+
+float ComputeRoughnessMaskScale(const FViewInfo& View, ESSRQuality SSRQuality)
+{
+	float MaxRoughness = FMath::Clamp(View.FinalPostProcessSettings.ScreenSpaceReflectionMaxRoughness, 0.01f, 1.0f);
+
+	// f(x) = x * Scale + Bias
+	// f(MaxRoughness) = 0
+	// f(MaxRoughness/2) = 1
+
+	float RoughnessMaskScale = -2.0f / MaxRoughness;
+	return RoughnessMaskScale * (int32(SSRQuality) < 3 ? 2.0f : 1.0f);
+}
+
+FLinearColor ComputeSSRParams(const FViewInfo& View, ESSRQuality SSRQuality, bool bEnableDiscard)
+{
+	float RoughnessMaskScale = ComputeRoughnessMaskScale(View, SSRQuality);
+
+	float FrameRandom = 0;
+
+	if(View.ViewState)
+	{
+		bool bTemporalAAIsOn = View.AntiAliasingMethod == AAM_TemporalAA;
+
+		if(bTemporalAAIsOn)
+		{
+			// usually this number is in the 0..7 range but it depends on the TemporalAA quality
+			FrameRandom = View.ViewState->GetCurrentTemporalAASampleIndex() * 1551;
+		}
+		else
+		{
+			// 8 aligns with the temporal smoothing, larger number will do more flickering (power of two for best performance)
+			FrameRandom = View.ViewState->GetFrameIndex(8) * 1551;
+		}
+	}
+
+	return FLinearColor(
+		FMath::Clamp(View.FinalPostProcessSettings.ScreenSpaceReflectionIntensity * 0.01f, 0.0f, 1.0f), 
+		RoughnessMaskScale,
+		(float)bEnableDiscard,	// TODO 
+		FrameRandom);
+}
+
+
+BEGIN_SHADER_PARAMETER_STRUCT(FSSRCommonParameters, )
+	SHADER_PARAMETER(FLinearColor, SSRParams)
+	SHADER_PARAMETER_STRUCT_INCLUDE(FSceneViewFamilyBlackboard, SceneTextures)
+	SHADER_PARAMETER_STRUCT_INCLUDE(FSceneTextureSamplerParameters, SceneTextureSamplers)
+	SHADER_PARAMETER_STRUCT_REF(FViewUniformShaderParameters, ViewUniformBuffer)
+END_SHADER_PARAMETER_STRUCT()
+
+
+
+class FSSRQualityDim : SHADER_PERMUTATION_ENUM_CLASS("SSR_QUALITY", ESSRQuality);
+class FSSROutputForDenoiser : SHADER_PERMUTATION_BOOL("SSR_OUTPUT_FOR_DENOISER");
+
+
+class FScreenSpaceReflectionsStencilPS : public FGlobalShader
+{
+	DECLARE_GLOBAL_SHADER(FScreenSpaceReflectionsStencilPS);
+	SHADER_USE_PARAMETER_STRUCT(FScreenSpaceReflectionsStencilPS, FGlobalShader);
+
+	using FPermutationDomain = TShaderPermutationDomain<FSSROutputForDenoiser>;
+
+	static bool ShouldCompilePermutation(const FGlobalShaderPermutationParameters& Parameters)
+	{
+		return IsFeatureLevelSupported(Parameters.Platform, ERHIFeatureLevel::SM4);
+	}
+
+	static void ModifyCompilationEnvironment(const FGlobalShaderPermutationParameters& Parameters, FShaderCompilerEnvironment& OutEnvironment)
+	{
+		FGlobalShader::ModifyCompilationEnvironment(Parameters, OutEnvironment);
+		OutEnvironment.SetDefine( TEXT("SSR_QUALITY"), uint32(0) );
+	}
+	
+	BEGIN_SHADER_PARAMETER_STRUCT(FParameters, )
+		SHADER_PARAMETER_STRUCT_INCLUDE(FSSRCommonParameters, CommonParameters)
+		RENDER_TARGET_BINDING_SLOTS()
+	END_SHADER_PARAMETER_STRUCT()
+};
+
+class FScreenSpaceReflectionsPS : public FGlobalShader
+{
+	DECLARE_GLOBAL_SHADER(FScreenSpaceReflectionsPS);
+	SHADER_USE_PARAMETER_STRUCT(FScreenSpaceReflectionsPS, FGlobalShader);
+
+	using FPermutationDomain = TShaderPermutationDomain<FSSRQualityDim, FSSROutputForDenoiser>;
+
+	static bool ShouldCompilePermutation(const FGlobalShaderPermutationParameters& Parameters)
+	{
+		FPermutationDomain PermutationVector(Parameters.PermutationId);
+		return IsFeatureLevelSupported(Parameters.Platform, ERHIFeatureLevel::SM4);
+	}
+	
+	BEGIN_SHADER_PARAMETER_STRUCT(FParameters, )
+		SHADER_PARAMETER_STRUCT_INCLUDE(FSSRCommonParameters, CommonParameters)
+
+		SHADER_PARAMETER(FVector4, HZBUvFactorAndInvFactor)
+		SHADER_PARAMETER(FVector4, PrevScreenPositionScaleBias)
+		SHADER_PARAMETER(float, PrevSceneColorPreExposureCorrection)
+		
+		SHADER_PARAMETER_RDG_TEXTURE(Texture2D, SceneColor)
+		SHADER_PARAMETER_SAMPLER(SamplerState, SceneColorSampler)
+		
+		SHADER_PARAMETER_RDG_TEXTURE(Texture2D, HZB)
+		SHADER_PARAMETER_SAMPLER(SamplerState, HZBSampler)
+		
+		RENDER_TARGET_BINDING_SLOTS()
+	END_SHADER_PARAMETER_STRUCT()
+};
+
+class FScreenSpaceDiffuseIndirectPS : public FGlobalShader
+{
+	DECLARE_GLOBAL_SHADER(FScreenSpaceDiffuseIndirectPS);
+	SHADER_USE_PARAMETER_STRUCT(FScreenSpaceDiffuseIndirectPS, FGlobalShader)
+
+	class FQualityDim : SHADER_PERMUTATION_INT( "QUALITY", 5 );
+	using FPermutationDomain = TShaderPermutationDomain< FQualityDim >;
+	
+	BEGIN_SHADER_PARAMETER_STRUCT( FParameters, )
+		SHADER_PARAMETER( FVector4,		HZBUvFactorAndInvFactor )
+		SHADER_PARAMETER( FVector4,		PrevScreenPositionScaleBias )
+		SHADER_PARAMETER( float,		PrevSceneColorPreExposureCorrection )
+		
+		RENDER_TARGET_BINDING_SLOTS()
+
+		SHADER_PARAMETER_RDG_TEXTURE( Texture2D,	HZBTexture )
+		SHADER_PARAMETER_SAMPLER( SamplerState,		HZBSampler )
+
+		SHADER_PARAMETER_RDG_TEXTURE( Texture2D,	VelocityTexture )
+		SHADER_PARAMETER_SAMPLER( SamplerState,		VelocitySampler )
+
+		SHADER_PARAMETER_RDG_TEXTURE( Texture2D,	ColorTexture )
+		SHADER_PARAMETER_SAMPLER( SamplerState,		ColorSampler )
+
+		SHADER_PARAMETER_STRUCT_REF(	FViewUniformShaderParameters,	View )
+		SHADER_PARAMETER_STRUCT_REF(	FSceneTexturesUniformParameters,SceneTextures )
+	END_SHADER_PARAMETER_STRUCT()
+
+	static bool ShouldCompilePermutation(const FGlobalShaderPermutationParameters& Parameters)
+	{
+		return IsFeatureLevelSupported(Parameters.Platform, ERHIFeatureLevel::SM4);
+	}
+};
+
+
+IMPLEMENT_GLOBAL_SHADER(FScreenSpaceReflectionsPS,        "/Engine/Private/SSRT/SSRTReflections.usf", "ScreenSpaceReflectionsPS", SF_Pixel);
+IMPLEMENT_GLOBAL_SHADER(FScreenSpaceReflectionsStencilPS, "/Engine/Private/SSRT/SSRTReflections.usf", "ScreenSpaceReflectionsStencilPS", SF_Pixel);
+IMPLEMENT_GLOBAL_SHADER(FScreenSpaceDiffuseIndirectPS,    "/Engine/Private/SSRT/SSRTDiffuseIndirect.usf", "ScreenSpaceDiffuseIndirectPS", SF_Pixel);
+
+
+void GetSSRShaderOptionsForQuality(ESSRQuality Quality, IScreenSpaceDenoiser::FReflectionsRayTracingConfig* OutRayTracingConfigs)
+{
+	if (Quality == ESSRQuality::VisualizeSSR)
+	{
+		OutRayTracingConfigs->RayCountPerPixel = 12;
+	}
+	else if (Quality == ESSRQuality::Epic)
+	{
+		OutRayTracingConfigs->RayCountPerPixel = 12;
+	}
+	else if (Quality == ESSRQuality::High)
+	{
+		OutRayTracingConfigs->RayCountPerPixel = 4;
+	}
+	else if (Quality == ESSRQuality::Medium)
+	{
+		OutRayTracingConfigs->RayCountPerPixel = 1;
+	}
+	else if (Quality == ESSRQuality::Low)
+	{
+		OutRayTracingConfigs->RayCountPerPixel = 1;
+	}
+	else
+	{
+		check(0);
+	}
+}
+
+} // namespace
+
+void GetSSRQualityForView(const FViewInfo& View, ESSRQuality* OutQuality, IScreenSpaceDenoiser::FReflectionsRayTracingConfig* OutRayTracingConfigs)
+{
+	check(ShouldRenderScreenSpaceReflections(View));
+	
+	int32 SSRQualityCVar = FMath::Clamp(CVarSSRQuality.GetValueOnRenderThread(), 0, int32(ESSRQuality::MAX) - 1);
+	
+	if (View.Family->EngineShowFlags.VisualizeSSR)
+	{
+		*OutQuality = ESSRQuality::VisualizeSSR;
+		return;
+	}
+	else if (View.FinalPostProcessSettings.ScreenSpaceReflectionQuality >= 80.0f && SSRQualityCVar >= 4)
+	{
+		*OutQuality = ESSRQuality::Epic;
+	}
+	else if (View.FinalPostProcessSettings.ScreenSpaceReflectionQuality >= 60.0f && SSRQualityCVar >= 3)
+	{
+		*OutQuality = ESSRQuality::High;
+	}
+	else if (View.FinalPostProcessSettings.ScreenSpaceReflectionQuality >= 40.0f && SSRQualityCVar >= 2)
+	{
+		*OutQuality = ESSRQuality::Medium;
+	}
+	else
+	{
+		*OutQuality = ESSRQuality::Low;
+	}
+
+	GetSSRShaderOptionsForQuality(*OutQuality, OutRayTracingConfigs);
+}
+
+void RenderScreenSpaceReflections(
+	FRDGBuilder& GraphBuilder,
+	const FSceneViewFamilyBlackboard& SceneTextures,
+	const FRDGTextureRef CurrentSceneColor,
+	const FViewInfo& View,
+	ESSRQuality SSRQuality,
+	bool bDenoiser,
+	IScreenSpaceDenoiser::FReflectionsInputs* DenoiserInputs)
+{
+	FRDGTextureRef InputColor = CurrentSceneColor;
+	if (SSRQuality != ESSRQuality::VisualizeSSR)
+	{
+		if (View.PrevViewInfo.CustomSSRInput.IsValid())
+		{
+			InputColor = GraphBuilder.RegisterExternalTexture(View.PrevViewInfo.CustomSSRInput);
+		}
+		else if (View.PrevViewInfo.TemporalAAHistory.IsValid())
+		{
+			InputColor = GraphBuilder.RegisterExternalTexture(View.PrevViewInfo.TemporalAAHistory.RT[0]);
+		}
+	}
+
+	const bool SSRStencilPrePass = CVarSSRStencil.GetValueOnRenderThread() != 0 && SSRQuality != ESSRQuality::VisualizeSSR;
+	
+	// Alloc inputs for denoising.
+	{
+		FRDGTextureDesc Desc = FPooledRenderTargetDesc::Create2DDesc(
+			FSceneRenderTargets::Get_FrameConstantsOnly().GetBufferSizeXY(),
+			PF_FloatRGBA, FClearValueBinding(FLinearColor(0, 0, 0, 0)),
+			TexCreate_None, TexCreate_RenderTargetable,
+			false);
+
+		Desc.AutoWritable = false;
+		Desc.Flags |= GFastVRamConfig.SSR;
+
+		DenoiserInputs->Color = GraphBuilder.CreateTexture(Desc, TEXT("ScreenSpaceReflections"));
+
+		if (bDenoiser)
+		{
+			Desc.Format = PF_R16F;
+			DenoiserInputs->RayHitDistance = GraphBuilder.CreateTexture(Desc, TEXT("ScreenSpaceReflectionsHitDistance"));
+		}
+	}
+
+	IScreenSpaceDenoiser::FReflectionsRayTracingConfig RayTracingConfigs;
+	GetSSRShaderOptionsForQuality(SSRQuality, &RayTracingConfigs);
+		
+	FSSRCommonParameters CommonParameters;
+	CommonParameters.SSRParams = ComputeSSRParams(View, SSRQuality, false);
+	CommonParameters.ViewUniformBuffer = View.ViewUniformBuffer;
+	CommonParameters.SceneTextures = SceneTextures;
+	SetupSceneTextureSamplers(&CommonParameters.SceneTextureSamplers);
+	
+	FRenderTargetBindingSlots RenderTargets;
+	RenderTargets[0] = FRenderTargetBinding(DenoiserInputs->Color, ERenderTargetLoadAction::ENoAction, ERenderTargetStoreAction::EStore);
+
+	if (bDenoiser)
+	{
+		RenderTargets[1] = FRenderTargetBinding(DenoiserInputs->RayHitDistance, ERenderTargetLoadAction::ENoAction, ERenderTargetStoreAction::EStore);
+	}
+
+	// Do a pre pass that output 0, or set a stencil mask to run the more expensive pixel shader.
+	if (SSRStencilPrePass)
+	{
+		// Also bind the depth buffer
+		RenderTargets.DepthStencil = FDepthStencilBinding(
+			SceneTextures.SceneDepthBuffer,
+			ERenderTargetLoadAction::ENoAction, ERenderTargetStoreAction::ENoAction,
+			ERenderTargetLoadAction::ELoad, ERenderTargetStoreAction::EStore,
+			FExclusiveDepthStencil::DepthRead_StencilWrite);
+
+		FScreenSpaceReflectionsStencilPS::FPermutationDomain PermutationVector;
+		PermutationVector.Set<FSSROutputForDenoiser>(bDenoiser);
+
+		FScreenSpaceReflectionsStencilPS::FParameters* PassParameters = GraphBuilder.AllocParameters<FScreenSpaceReflectionsStencilPS::FParameters>();
+		PassParameters->CommonParameters = CommonParameters;
+		PassParameters->RenderTargets = RenderTargets;
+		
+		TShaderMapRef<FScreenSpaceReflectionsStencilPS> PixelShader(View.ShaderMap, PermutationVector);
+		ClearUnusedGraphResources(*PixelShader, PassParameters);
+		
+		GraphBuilder.AddPass(
+			RDG_EVENT_NAME("SSR StencilSetup %dx%d", View.ViewRect.Width(), View.ViewRect.Height()),
+			PassParameters,
+			ERenderGraphPassFlags::None,
+			[PassParameters, &View, PixelShader](FRHICommandList& RHICmdList)
+		{
+			SCOPED_GPU_STAT(RHICmdList, ScreenSpaceReflections);
+			RHICmdList.SetViewport(View.ViewRect.Min.X, View.ViewRect.Min.Y, 0.0f, View.ViewRect.Max.X, View.ViewRect.Max.Y, 1.0f);
+			RHICmdList.SetStencilRef(0x80);
+		
+			FGraphicsPipelineStateInitializer GraphicsPSOInit;
+			FPixelShaderUtils::InitFullscreenPipelineState(RHICmdList, View.ShaderMap, *PixelShader, /* out */ GraphicsPSOInit);
+			// Clobers the stencil to pixel that should not compute SSR
+			GraphicsPSOInit.DepthStencilState = TStaticDepthStencilState<false, CF_Always, true, CF_Always, SO_Replace, SO_Replace, SO_Replace>::GetRHI();
+
+			SetGraphicsPipelineState(RHICmdList, GraphicsPSOInit);
+			SetShaderParameters(RHICmdList, *PixelShader, PixelShader->GetPixelShader(), *PassParameters);
+
+			FPixelShaderUtils::DrawFullscreenTriangle(RHICmdList);
+		});
+	}
+
+	// Adds SSR pass.
+	{
+		FScreenSpaceReflectionsPS::FPermutationDomain PermutationVector;
+		PermutationVector.Set<FSSRQualityDim>(SSRQuality);
+		PermutationVector.Set<FSSROutputForDenoiser>(bDenoiser);
+
+		FScreenSpaceReflectionsPS::FParameters* PassParameters = GraphBuilder.AllocParameters<FScreenSpaceReflectionsPS::FParameters>();
+		PassParameters->CommonParameters = CommonParameters;
+		{
+			const FVector2D HZBUvFactor(
+				float(View.ViewRect.Width()) / float(2 * View.HZBMipmap0Size.X),
+				float(View.ViewRect.Height()) / float(2 * View.HZBMipmap0Size.Y));
+			PassParameters->HZBUvFactorAndInvFactor = FVector4(
+				HZBUvFactor.X,
+				HZBUvFactor.Y,
+				1.0f / HZBUvFactor.X,
+				1.0f / HZBUvFactor.Y);
+		}
+		{
+			FIntPoint ViewportOffset = View.ViewRect.Min;
+			FIntPoint ViewportExtent = View.ViewRect.Size();
+			FIntPoint BufferSize = SceneTextures.SceneDepthBuffer->Desc.Extent;
+
+			if (View.PrevViewInfo.TemporalAAHistory.IsValid())
+			{
+				ViewportOffset = View.PrevViewInfo.TemporalAAHistory.ViewportRect.Min;
+				ViewportExtent = View.PrevViewInfo.TemporalAAHistory.ViewportRect.Size();
+				BufferSize = View.PrevViewInfo.TemporalAAHistory.ReferenceBufferSize;
+			}
+
+			FVector2D InvBufferSize(1.0f / float(BufferSize.X), 1.0f / float(BufferSize.Y));
+
+			PassParameters->PrevScreenPositionScaleBias = FVector4(
+				ViewportExtent.X * 0.5f * InvBufferSize.X,
+				-ViewportExtent.Y * 0.5f * InvBufferSize.Y,
+				(ViewportExtent.X * 0.5f + ViewportOffset.X) * InvBufferSize.X,
+				(ViewportExtent.Y * 0.5f + ViewportOffset.Y) * InvBufferSize.Y);
+		}
+		PassParameters->PrevSceneColorPreExposureCorrection = InputColor != CurrentSceneColor ? View.PreExposure / View.PrevViewInfo.SceneColorPreExposure : 1.0f;
+		
+		// Pipe down a mid grey texture when not using TAA's history to avoid wrongly reprojecting current scene color as if previous frame's TAA history.
+		if (InputColor == CurrentSceneColor)
+		{
+			// Technically should be 32767.0f / 65535.0f to perfectly null out DecodeVelocityFromTexture(), but 0.5f is good enough.
+			PassParameters->CommonParameters.SceneTextures.SceneVelocityBuffer = GraphBuilder.RegisterExternalTexture(GSystemTextures.MidGreyDummy);
+		}
+
+		PassParameters->SceneColor = InputColor;
+		PassParameters->SceneColorSampler = TStaticSamplerState<SF_Point>::GetRHI();
+		
+		PassParameters->HZB = GraphBuilder.RegisterExternalTexture(View.HZB);
+		PassParameters->HZBSampler = TStaticSamplerState<SF_Point>::GetRHI();
+		
+		PassParameters->RenderTargets = RenderTargets;
+
+		TShaderMapRef<FScreenSpaceReflectionsPS> PixelShader(View.ShaderMap, PermutationVector);
+		ClearUnusedGraphResources(*PixelShader, PassParameters);
+		
+		GraphBuilder.AddPass(
+			RDG_EVENT_NAME("SSR RayMarch(Quality=%d RayPerPixel=%d%s) %dx%d",
+				SSRQuality, RayTracingConfigs.RayCountPerPixel, bDenoiser ? TEXT(" DenoiserOutput") : TEXT(""),
+				View.ViewRect.Width(), View.ViewRect.Height()),
+			PassParameters,
+			ERenderGraphPassFlags::None,
+			[PassParameters, &View, PixelShader, SSRStencilPrePass](FRHICommandList& RHICmdList)
+		{
+			SCOPED_GPU_STAT(RHICmdList, ScreenSpaceReflections);
+			RHICmdList.SetViewport(View.ViewRect.Min.X, View.ViewRect.Min.Y, 0.0f, View.ViewRect.Max.X, View.ViewRect.Max.Y, 1.0f);
+			RHICmdList.SetStencilRef(0x80);
+		
+			FGraphicsPipelineStateInitializer GraphicsPSOInit;
+			FPixelShaderUtils::InitFullscreenPipelineState(RHICmdList, View.ShaderMap, *PixelShader, /* out */ GraphicsPSOInit);
+			if (SSRStencilPrePass)
+			{
+				// Clobers the stencil to pixel that should not compute SSR
+				GraphicsPSOInit.DepthStencilState = TStaticDepthStencilState<false, CF_Always, true, CF_Equal, SO_Keep, SO_Keep, SO_Keep>::GetRHI();
+			}
+
+			SetGraphicsPipelineState(RHICmdList, GraphicsPSOInit);
+			SetShaderParameters(RHICmdList, *PixelShader, PixelShader->GetPixelShader(), *PassParameters);
+
+			FPixelShaderUtils::DrawFullscreenTriangle(RHICmdList);
+		});
+	}
+} // RenderScreenSpaceReflections()
+
+void RenderScreenSpaceDiffuseIndirect( FRHICommandListImmediate& RHICmdList, FViewInfo& View, TRefCountPtr<IPooledRenderTarget>& VelocityRT )
+{
+	FTemporalAAHistory& TemporalAAHistory = View.PrevViewInfo.TemporalAAHistory;
+
+	if( !ShouldRenderScreenSpaceDiffuseIndirect( View ) || !TemporalAAHistory.IsValid() )
+	{
+		return;
+	}
+	
+	const int32 Quality = FMath::Clamp( CVarSSGIQuality.GetValueOnRenderThread(), 1, 4 );
+
+	FRDGBuilder GraphBuilder( RHICmdList );
+
+	FSceneRenderTargets& SceneContext = FSceneRenderTargets::Get( RHICmdList );
+
+	FRDGTexture* SceneColor	= GraphBuilder.RegisterExternalTexture( SceneContext.GetSceneColor() );
+	FRDGTexture* ScreenSpaceAO = GraphBuilder.CreateTexture(SceneContext.ScreenSpaceAO->GetDesc(), TEXT("SSRTAO"));
+	FRDGTexture* HZBTexture	= GraphBuilder.RegisterExternalTexture( View.HZB );
+	FRDGTexture* ColorTexture	= GraphBuilder.RegisterExternalTexture( TemporalAAHistory.RT[0] );
+
+	FRDGTexture* VelocityTexture;
+	if( VelocityRT && !View.bCameraCut )
+	{
+		VelocityTexture = GraphBuilder.RegisterExternalTexture( VelocityRT );
+	}
+	else
+	{
+		// No velocity, use black
+		VelocityTexture = GraphBuilder.RegisterExternalTexture( GSystemTextures.BlackDummy );
+	}
+
+	FScreenSpaceDiffuseIndirectPS::FParameters* PassParameters = GraphBuilder.AllocParameters< FScreenSpaceDiffuseIndirectPS::FParameters >();
+
+	PassParameters->RenderTargets[0] = FRenderTargetBinding( SceneColor, ERenderTargetLoadAction::ELoad, ERenderTargetStoreAction::EStore );
+	PassParameters->RenderTargets[1] = FRenderTargetBinding( ScreenSpaceAO, ERenderTargetLoadAction::ENoAction, ERenderTargetStoreAction::EStore );
+
+	PassParameters->HZBTexture = HZBTexture;
+	PassParameters->HZBSampler = TStaticSamplerState< SF_Point, AM_Clamp, AM_Clamp, AM_Clamp >::GetRHI();
+
+	PassParameters->VelocityTexture = VelocityTexture;
+	PassParameters->VelocitySampler = TStaticSamplerState< SF_Point, AM_Clamp, AM_Clamp, AM_Clamp >::GetRHI();
+
+	PassParameters->ColorTexture = ColorTexture;
+	PassParameters->ColorSampler = TStaticSamplerState< SF_Bilinear, AM_Clamp, AM_Clamp, AM_Clamp >::GetRHI();
+
+	const FVector2D HZBUvFactor(
+		float( View.ViewRect.Width() )  / float( 2 * View.HZBMipmap0Size.X ),
+		float( View.ViewRect.Height() ) / float( 2 * View.HZBMipmap0Size.Y )
+		);
+			
+	PassParameters->HZBUvFactorAndInvFactor = FVector4(
+		HZBUvFactor.X,
+		HZBUvFactor.Y,
+		1.0f / HZBUvFactor.X,
+		1.0f / HZBUvFactor.Y );
+
+	FIntPoint ViewportOffset	= TemporalAAHistory.ViewportRect.Min;
+	FIntPoint ViewportExtent	= TemporalAAHistory.ViewportRect.Size();
+	FIntPoint BufferSize		= TemporalAAHistory.ReferenceBufferSize;
+
+	PassParameters->PrevScreenPositionScaleBias = FVector4(
+		 ViewportExtent.X * 0.5f / BufferSize.X,
+		-ViewportExtent.Y * 0.5f / BufferSize.Y,
+		(ViewportExtent.X * 0.5f + ViewportOffset.X) / BufferSize.X,
+		(ViewportExtent.Y * 0.5f + ViewportOffset.Y) / BufferSize.Y );
+
+	PassParameters->PrevSceneColorPreExposureCorrection = View.PreExposure / View.PrevViewInfo.SceneColorPreExposure;
+
+	PassParameters->View = View.ViewUniformBuffer;
+	PassParameters->SceneTextures = CreateSceneTextureUniformBufferSingleDraw( RHICmdList, ESceneTextureSetupMode::SceneDepth | ESceneTextureSetupMode::GBuffers, View.FeatureLevel );
+
+	GraphBuilder.AddPass(
+		RDG_EVENT_NAME( "ScreenSpaceDiffuseIndirect(Quality=%d) %dx%d", Quality, View.ViewRect.Width(), View.ViewRect.Height() ),
+		PassParameters,
+		ERenderGraphPassFlags::None,
+		[ PassParameters, &View, Quality ]( FRHICommandListImmediate& InRHICmdList )
+	{
+		SCOPED_GPU_STAT(InRHICmdList, ScreenSpaceDiffuseIndirect);
+
+		InRHICmdList.SetViewport( View.ViewRect.Min.X, View.ViewRect.Min.Y, 0.0f, View.ViewRect.Max.X, View.ViewRect.Max.Y, 1.0f );
+
+		FScreenSpaceDiffuseIndirectPS::FPermutationDomain PermutationVector;
+		PermutationVector.Set< FScreenSpaceDiffuseIndirectPS::FQualityDim >( Quality );
+
+		auto VertexShader	= View.ShaderMap->GetShader< FScreenPassVS >();
+		auto PixelShader	= View.ShaderMap->GetShader< FScreenSpaceDiffuseIndirectPS >( PermutationVector );
+		CA_ASSUME(PixelShader);
+
+		{
+			FGraphicsPipelineStateInitializer GraphicsPSOInit;
+			InRHICmdList.ApplyCachedRenderTargets( GraphicsPSOInit );
+
+			GraphicsPSOInit.BlendState = TStaticBlendState< CW_RGB, BO_Add, BF_One, BF_SourceAlpha >::GetRHI();
+			GraphicsPSOInit.DepthStencilState = TStaticDepthStencilState<false, CF_Always>::GetRHI();
+			GraphicsPSOInit.RasterizerState = TStaticRasterizerState<>::GetRHI();
+			GraphicsPSOInit.PrimitiveType = PT_TriangleList;
+			GraphicsPSOInit.BoundShaderState.VertexDeclarationRHI = GFilterVertexDeclaration.VertexDeclarationRHI;
+			GraphicsPSOInit.BoundShaderState.VertexShaderRHI = GETSAFERHISHADER_VERTEX( VertexShader );
+			GraphicsPSOInit.BoundShaderState.PixelShaderRHI = GETSAFERHISHADER_PIXEL( PixelShader );
+
+			SetGraphicsPipelineState( InRHICmdList, GraphicsPSOInit );
+		}
+
+		SetShaderParameters( InRHICmdList, PixelShader, PixelShader->GetPixelShader(), *PassParameters );
+
+		DrawPostProcessPass(
+			InRHICmdList,
+			0, 0,
+			View.ViewRect.Width(), View.ViewRect.Height(),
+			View.ViewRect.Min.X, View.ViewRect.Min.Y,
+			View.ViewRect.Width(), View.ViewRect.Height(),
+			View.ViewRect.Size(),
+			FSceneRenderTargets::Get( InRHICmdList ).GetBufferSizeXY(),
+			VertexShader,
+			View.StereoPass,
+			false,
+			EDRF_UseTriangleOptimization);
+	});
+
+	GraphBuilder.QueueTextureExtraction(ScreenSpaceAO, &SceneContext.ScreenSpaceAO);
+
+	GraphBuilder.Execute();
+} // RenderScreenSpaceDiffuseIndirect()
diff --git a/Engine/Source/Runtime/Renderer/Private/ScreenSpaceRayTracing.h b/Engine/Source/Runtime/Renderer/Private/ScreenSpaceRayTracing.h
new file mode 100644
--- /dev/null
+++ b/Engine/Source/Runtime/Renderer/Private/ScreenSpaceRayTracing.h
@@ -0,0 +1,38 @@
+// Copyright 1998-2019 Epic Games, Inc. All Rights Reserved.
+
+#pragma once
+
+#include "RenderGraph.h"
+#include "ScreenSpaceDenoise.h"
+
+class FViewInfo;
+class FSceneViewFamilyBlackboard;
+
+enum class ESSRQuality
+{
+	VisualizeSSR,
+
+	Low,
+	Medium,
+	High,
+	Epic,
+
+	MAX
+};
+
+bool ShouldRenderScreenSpaceReflections(const FViewInfo& View);
+
+void GetSSRQualityForView(const FViewInfo& View, ESSRQuality* OutQuality, IScreenSpaceDenoiser::FReflectionsRayTracingConfig* OutRayTracingConfigs);
+
+bool IsSSRTemporalPassRequired(const FViewInfo& View);
+
+void RenderScreenSpaceReflections(
+	FRDGBuilder& GraphBuilder,
+	const FSceneViewFamilyBlackboard& SceneBlackboard,
+	const FRDGTextureRef CurrentSceneColor,
+	const FViewInfo& View,
+	ESSRQuality SSRQuality,
+	bool bDenoiser,
+	IScreenSpaceDenoiser::FReflectionsInputs* DenoiserInputs);
+
+void RenderScreenSpaceDiffuseIndirect( FRHICommandListImmediate& RHICmdList, FViewInfo& View, TRefCountPtr<IPooledRenderTarget>& VelocityRT );
diff --git a/Engine/Source/Runtime/Renderer/Private/VelocityRendering.cpp b/Engine/Source/Runtime/Renderer/Private/VelocityRendering.cpp
--- a/Engine/Source/Runtime/Renderer/Private/VelocityRendering.cpp
+++ b/Engine/Source/Runtime/Renderer/Private/VelocityRendering.cpp
@@ -15,7 +15,7 @@
 #include "SceneRendering.h"
 #include "DeferredShadingRenderer.h"
 #include "ScenePrivate.h"
-#include "PostProcess/ScreenSpaceReflections.h"
+#include "ScreenSpaceRayTracing.h"
 #include "PostProcess/PostProcessMotionBlur.h"
 #include "UnrealEngine.h"
 #if WITH_EDITOR
