Fixed crash when using ray tracing debug view modes due to velocity buffer not being provided to TAA pass.

diff --git a/Engine/Source/Runtime/Renderer/Private/DebugViewModeRendering.cpp b/Engine/Source/Runtime/Renderer/Private/DebugViewModeRendering.cpp
--- a/Engine/Source/Runtime/Renderer/Private/DebugViewModeRendering.cpp
+++ b/Engine/Source/Runtime/Renderer/Private/DebugViewModeRendering.cpp
@@ -20,6 +20,7 @@ DebugViewModeRendering.cpp: Contains definitions for rendering debug viewmodes.
 #include "PostProcess/PostProcessCompositeEditorPrimitives.h"
 #include "PostProcess/PostProcessUpscale.h"
 #include "PostProcess/PostProcessTemporalAA.h"
+#include "PostProcess/PostProcessInput.h"
 #include "SceneRendering.h"
 #include "DeferredShadingRenderer.h"
 #include "MeshPassProcessor.inl"
@@ -109,20 +110,29 @@ void FDeferredShadingSceneRenderer::DoDebugViewModePostProcessing(FRHICommandLis
 		case DVSM_RayTracingDebug:
 		{
 			FSceneViewState* ViewState = (FSceneViewState*)Context.View.State;
-			FTAAPassParameters Parameters(Context.View);	
+			FTAAPassParameters Parameters(Context.View);
 			FRenderingCompositePass* Node = Context.Graph.RegisterPass(new(FMemStack::Get()) FRCPassPostProcessTemporalAA(
 				Context, 
 				Parameters, 
 				Context.View.PrevViewInfo.TemporalAAHistory, 
 				&ViewState->PrevFrameViewInfo.TemporalAAHistory));
 
+			FRenderingCompositePass* VelocityNode = VelocityRT
+				? Context.Graph.RegisterPass(new(FMemStack::Get()) FRCPassPostProcessInput(VelocityRT))
+				: Context.Graph.RegisterPass(new(FMemStack::Get()) FRCPassPostProcessInput(GSystemTextures.BlackDummy));
+
+			FRenderingCompositeOutputRef VelocityRef(VelocityNode);
+
 			Node->SetInput(ePId_Input0, FRenderingCompositeOutputRef(Context.FinalOutput));
+			Node->SetInput(ePId_Input2, VelocityRef);
+
 			Context.FinalOutput = FRenderingCompositeOutputRef(Node);
 
 			if (View.Family->EngineShowFlags.Tonemapper)
 			{
 				GPostProcessing.AddGammaOnlyTonemapper(Context);
 			}
+
 			break;
 		}
 		default:
