Fix/suppress some static analysis warnings

diff --git a/Engine/Source/Runtime/Engine/Private/LightmapUniformShaderParameters.cpp b/Engine/Source/Runtime/Engine/Private/LightmapUniformShaderParameters.cpp
--- a/Engine/Source/Runtime/Engine/Private/LightmapUniformShaderParameters.cpp
+++ b/Engine/Source/Runtime/Engine/Private/LightmapUniformShaderParameters.cpp
@@ -97,6 +97,7 @@ void GetPrecomputedLightingParameters(
 		uint32 NumLightmapVTLayers = 0u;
 		if (bUseVirtualTextures)
 		{
+			check(LCI); // If LCI was nullptr, LightMapInteraction.GetType() would be LMIT_None, not LMIT_Texture
 			const FLightmapResourceCluster* ResourceCluster = LCI->GetResourceCluster();
 			check(ResourceCluster);
 			IAllocatedVirtualTexture* AllocatedVT = ResourceCluster->AcquireAllocatedVT();
diff --git a/Engine/Source/Runtime/Renderer/Private/VT/UniqueRequestList.h b/Engine/Source/Runtime/Renderer/Private/VT/UniqueRequestList.h
--- a/Engine/Source/Runtime/Renderer/Private/VT/UniqueRequestList.h
+++ b/Engine/Source/Runtime/Renderer/Private/VT/UniqueRequestList.h
@@ -319,7 +319,8 @@ inline void FUniqueRequestList::SortRequests(FVirtualTextureProducerCollection&
 	// Re-index load request list, using sorted indices
 	FVirtualTextureLocalTile* SortedLoadRequests = new(MemStack) FVirtualTextureLocalTile[NewNumLoadRequests];
 	uint8* SortedLayerMask = new(MemStack) uint8[NewNumLoadRequests];
-	uint16* LoadIndexToSortedLoadIndex = new(MemStack, MEM_Oned) uint16[NumLoadRequests];
+	uint16* LoadIndexToSortedLoadIndex = new(MemStack) uint16[NumLoadRequests];
+	FMemory::Memset(LoadIndexToSortedLoadIndex, 0xff, NumLoadRequests * sizeof(uint16));
 	for (uint32 i = 0u; i < NewNumLoadRequests; ++i)
 	{
 		const uint32 SortedIndex = SortedKeys[i].Index;
diff --git a/Engine/Source/Runtime/Renderer/Private/VT/VirtualTextureSystem.cpp b/Engine/Source/Runtime/Renderer/Private/VT/VirtualTextureSystem.cpp
--- a/Engine/Source/Runtime/Renderer/Private/VT/VirtualTextureSystem.cpp
+++ b/Engine/Source/Runtime/Renderer/Private/VT/VirtualTextureSystem.cpp
@@ -371,12 +371,15 @@ IAllocatedVirtualTexture* FVirtualTextureSystem::AllocateVirtualTexture(const FA
 			for (uint32 LayerIndex = 0u; LayerIndex < Desc.NumLayers; ++LayerIndex)
 			{
 				const FVirtualTextureProducer* Producer = ProducerForLayer[LayerIndex];
-				if ((WidthInBlocks % Producer->GetDescription().WidthInBlocks) != 0u)
+				if (Producer)
 				{
-					WidthInBlocks += MinWidthInBlocks;
-					check(WidthInBlocks > MinWidthInBlocks); // check for overflow
-					bFoundValidWidthInBlocks = false;
-					break;
+					if ((WidthInBlocks % Producer->GetDescription().WidthInBlocks) != 0u)
+					{
+						WidthInBlocks += MinWidthInBlocks;
+						check(WidthInBlocks > MinWidthInBlocks); // check for overflow
+						bFoundValidWidthInBlocks = false;
+						break;
+					}
 				}
 			}
 		}
@@ -392,12 +395,15 @@ IAllocatedVirtualTexture* FVirtualTextureSystem::AllocateVirtualTexture(const FA
 			for (uint32 LayerIndex = 0u; LayerIndex < Desc.NumLayers; ++LayerIndex)
 			{
 				const FVirtualTextureProducer* Producer = ProducerForLayer[LayerIndex];
-				if ((HeightInBlocks % Producer->GetDescription().HeightInBlocks) != 0u)
+				if (Producer)
 				{
-					HeightInBlocks += MinHeightInBlocks;
-					check(HeightInBlocks > MinHeightInBlocks); // check for overflow
-					bFoundValidHeightInBlocks = false;
-					break;
+					if ((HeightInBlocks % Producer->GetDescription().HeightInBlocks) != 0u)
+					{
+						HeightInBlocks += MinHeightInBlocks;
+						check(HeightInBlocks > MinHeightInBlocks); // check for overflow
+						bFoundValidHeightInBlocks = false;
+						break;
+					}
 				}
 			}
 		}
