Remove unnecessary BindForLegacyShaderParameters

diff --git a/Engine/Source/Runtime/Renderer/Private/VT/RuntimeVirtualTextureRender.cpp b/Engine/Source/Runtime/Renderer/Private/VT/RuntimeVirtualTextureRender.cpp
--- a/Engine/Source/Runtime/Renderer/Private/VT/RuntimeVirtualTextureRender.cpp
+++ b/Engine/Source/Runtime/Renderer/Private/VT/RuntimeVirtualTextureRender.cpp
@@ -7,6 +7,7 @@
 #include "MaterialShader.h"
 #include "MeshPassProcessor.h"
 #include "RenderGraphBuilder.h"
+#include "RenderGraphUtils.h"
 #include "RenderUtils.h"
 #include "ScenePrivate.h"
 #include "SceneRenderTargets.h"
@@ -21,9 +22,7 @@ namespace RuntimeVirtualTexture
 	class FShader_VirtualTextureMaterialDraw : public FMeshMaterialShader
 	{
 	public:
-		BEGIN_SHADER_PARAMETER_STRUCT(FParameters, )
-			RENDER_TARGET_BINDING_SLOTS()
-		END_SHADER_PARAMETER_STRUCT()
+		typedef FRenderTargetParameters FParameters;
 
 		static bool ShouldCompilePermutation(EShaderPlatform Platform, const FMaterial* Material, const FVertexFactoryType* VertexFactoryType)
 		{
@@ -45,7 +44,6 @@ namespace RuntimeVirtualTexture
 		FShader_VirtualTextureMaterialDraw(const FMeshMaterialShaderType::CompiledShaderInitializerType& Initializer)
 			: FMeshMaterialShader(Initializer)
 		{
-			Bindings.BindForLegacyShaderParameters(this, Initializer.ParameterMap, *FParameters::FTypeInfo::GetStructMetadata());
 			// Ensure FMeshMaterialShader::PassUniformBuffer is bound (although currently unused)
 			PassUniformBuffer.Bind(Initializer.ParameterMap, FSceneTexturesUniformParameters::StaticStructMetadata.GetShaderVariableName());
 		}
