DR - Merge fix
#rnx

diff --git a/Engine/Source/Runtime/Renderer/Private/PostProcess/PostProcessing.cpp b/Engine/Source/Runtime/Renderer/Private/PostProcess/PostProcessing.cpp
--- a/Engine/Source/Runtime/Renderer/Private/PostProcess/PostProcessing.cpp
+++ b/Engine/Source/Runtime/Renderer/Private/PostProcess/PostProcessing.cpp
@@ -1368,7 +1368,7 @@ void FPostProcessing::Process(FRHICommandListImmediate& RHICmdList, const FViewI
 		const bool bHDROutputEnabled = GRHISupportsHDROutput && IsHDREnabled();
 
 		static const auto CVarDumpFramesAsHDR = IConsoleManager::Get().FindTConsoleVariableDataInt(TEXT("r.BufferVisualizationDumpFramesAsHDR"));
-    	const bool bHDRTonemapperOutput = bAllowTonemapper && (View.Family->SceneCaptureSource == SCS_FinalColorHDR || GetHighResScreenshotConfig().bCaptureHDR || CVarDumpFramesAsHDR->GetValueOnRenderThread() || bHDROutputEnabled || View.Family->bIsHDR);
+    	const bool bHDRTonemapperOutput = bAllowTonemapper && (View.Family->SceneCaptureSource == SCS_FinalColorHDR || GetHighResScreenshotConfig().bCaptureHDR || CVarDumpFramesAsHDR->GetValueOnRenderThread() || bHDROutputEnabled);
 
 
 		FRCPassPostProcessTonemap* Tonemapper = 0;
