Fixed issue where view uniform buffer would not update back to full screen translucency parameters.

diff --git a/Engine/Source/Runtime/Renderer/Private/PostProcess/PostProcessCompositeEditorPrimitives.cpp b/Engine/Source/Runtime/Renderer/Private/PostProcess/PostProcessCompositeEditorPrimitives.cpp
--- a/Engine/Source/Runtime/Renderer/Private/PostProcess/PostProcessCompositeEditorPrimitives.cpp
+++ b/Engine/Source/Runtime/Renderer/Private/PostProcess/PostProcessCompositeEditorPrimitives.cpp
@@ -492,7 +492,7 @@ void FRCPassPostProcessCompositeEditorPrimitives::Process(FRenderingCompositePas
 		EditorView.CachedViewUniformShaderParameters->NumSceneColorMSAASamples = MSAASampleCount;
 
 		FScene* Scene = Context.View.Family->Scene->GetRenderScene();
-		Scene->UniformBuffers.ViewUniformBuffer.UpdateUniformBufferImmediate(*EditorView.CachedViewUniformShaderParameters);
+		Scene->UniformBuffers.UpdateViewUniformBufferImmediate(*EditorView.CachedViewUniformShaderParameters);
 		EditorView.ViewUniformBuffer = Scene->UniformBuffers.ViewUniformBuffer;
 	}
 
diff --git a/Engine/Source/Runtime/Renderer/Private/PostProcess/PostProcessSelectionOutline.cpp b/Engine/Source/Runtime/Renderer/Private/PostProcess/PostProcessSelectionOutline.cpp
--- a/Engine/Source/Runtime/Renderer/Private/PostProcess/PostProcessSelectionOutline.cpp
+++ b/Engine/Source/Runtime/Renderer/Private/PostProcess/PostProcessSelectionOutline.cpp
@@ -106,7 +106,7 @@ void FRCPassPostProcessSelectionOutlineColor::Process(FRenderingCompositePassCon
 			FBox VolumeBounds[TVC_MAX];
 			EditorView.SetupUniformBufferParameters(SceneContext, VolumeBounds, TVC_MAX, *EditorView.CachedViewUniformShaderParameters);
 
-			Scene->UniformBuffers.ViewUniformBuffer.UpdateUniformBufferImmediate(*EditorView.CachedViewUniformShaderParameters);
+			Scene->UniformBuffers.UpdateViewUniformBufferImmediate(*EditorView.CachedViewUniformShaderParameters);
 			EditorView.ViewUniformBuffer = Scene->UniformBuffers.ViewUniformBuffer;
 
 			// Run selection pass on static elements
diff --git a/Engine/Source/Runtime/Renderer/Private/RendererScene.cpp b/Engine/Source/Runtime/Renderer/Private/RendererScene.cpp
--- a/Engine/Source/Runtime/Renderer/Private/RendererScene.cpp
+++ b/Engine/Source/Runtime/Renderer/Private/RendererScene.cpp
@@ -1025,6 +1025,12 @@ bool FPersistentUniformBuffers::UpdateViewUniformBuffer(const FViewInfo& View)
 	return false;
 }
 
+void FPersistentUniformBuffers::UpdateViewUniformBufferImmediate(const FViewUniformShaderParameters& Parameters)
+{
+	ViewUniformBuffer.UpdateUniformBufferImmediate(Parameters);
+	CachedView = nullptr;
+}
+
 FScene::FScene(UWorld* InWorld, bool bInRequiresHitProxies, bool bInIsEditorScene, bool bCreateFXSystem, ERHIFeatureLevel::Type InFeatureLevel)
 :	FSceneInterface(InFeatureLevel)
 ,	World(InWorld)
diff --git a/Engine/Source/Runtime/Renderer/Private/ScenePrivate.h b/Engine/Source/Runtime/Renderer/Private/ScenePrivate.h
--- a/Engine/Source/Runtime/Renderer/Private/ScenePrivate.h
+++ b/Engine/Source/Runtime/Renderer/Private/ScenePrivate.h
@@ -2263,6 +2263,7 @@ public:
 	void Initialize();
 	// @return Whether uniform buffer was updated
 	bool UpdateViewUniformBuffer(const FViewInfo& View);
+	void UpdateViewUniformBufferImmediate(const FViewUniformShaderParameters& Parameters);
 
 	TUniformBufferRef<FViewUniformShaderParameters> ViewUniformBuffer;
 	TUniformBufferRef<FInstancedViewUniformShaderParameters> InstancedViewUniformBuffer;
diff --git a/Engine/Source/Runtime/Renderer/Private/TranslucentRendering.cpp b/Engine/Source/Runtime/Renderer/Private/TranslucentRendering.cpp
--- a/Engine/Source/Runtime/Renderer/Private/TranslucentRendering.cpp
+++ b/Engine/Source/Runtime/Renderer/Private/TranslucentRendering.cpp
@@ -910,7 +910,7 @@ void FDeferredShadingSceneRenderer::RenderTranslucency(FRHICommandListImmediate&
 			{
 				FViewUniformShaderParameters DownsampledTranslucencyViewParameters;
 				SetupDownsampledTranslucencyViewParameters(RHICmdList, View, DownsampledTranslucencyViewParameters);
-				Scene->UniformBuffers.ViewUniformBuffer.UpdateUniformBufferImmediate(DownsampledTranslucencyViewParameters);
+				Scene->UniformBuffers.UpdateViewUniformBufferImmediate(DownsampledTranslucencyViewParameters);
 				DrawRenderState.SetViewUniformBuffer(Scene->UniformBuffers.ViewUniformBuffer);
 
 				if ((View.IsInstancedStereoPass() || View.bIsMobileMultiViewEnabled) && View.Family->Views.Num() > 0)
