Updated some view uniform buffer update comments.

diff --git a/Engine/Source/Runtime/Renderer/Private/ScenePrivate.h b/Engine/Source/Runtime/Renderer/Private/ScenePrivate.h
--- a/Engine/Source/Runtime/Renderer/Private/ScenePrivate.h
+++ b/Engine/Source/Runtime/Renderer/Private/ScenePrivate.h
@@ -2261,8 +2261,13 @@ public:
 	}
 
 	void Initialize();
-	// @return Whether uniform buffer was updated
+
+	/** Compares the provided view against the cached view and updates the view uniform buffer
+	 *  if the views differ. Returns whether uniform buffer was updated.
+	 */
 	bool UpdateViewUniformBuffer(const FViewInfo& View);
+
+	/** Updates view uniform buffer and invalidates the internally cached view instance. */
 	void UpdateViewUniformBufferImmediate(const FViewUniformShaderParameters& Parameters);
 
 	TUniformBufferRef<FViewUniformShaderParameters> ViewUniformBuffer;
