Deal with more runtime virtual texture update conditions

diff --git a/Engine/Source/Runtime/Engine/Private/Components/PrimitiveComponent.cpp b/Engine/Source/Runtime/Engine/Private/Components/PrimitiveComponent.cpp
--- a/Engine/Source/Runtime/Engine/Private/Components/PrimitiveComponent.cpp
+++ b/Engine/Source/Runtime/Engine/Private/Components/PrimitiveComponent.cpp
@@ -3072,14 +3072,6 @@ void UPrimitiveComponent::UpdateBounds()
 				Collection->NotifyPrimitiveUpdated_Concurrent(this);
 			}
 		}
-
-		if (RuntimeVirtualTextures.Num() > 0)
-		{
-			// If we render to a runtime virtual texture then we need to flush here
-			//todo[vt]: Only flush Bounds 
-			//todo[vt]: Maybe only flush specific virtual textures (which would mean that we need to specify which virtual textures to write to in the primitive UI)
-			GetRendererModule().FlushVirtualTextureCache();
-		}
 	}
 }
 #endif
diff --git a/Engine/Source/Runtime/Engine/Private/PrimitiveSceneProxy.cpp b/Engine/Source/Runtime/Engine/Private/PrimitiveSceneProxy.cpp
--- a/Engine/Source/Runtime/Engine/Private/PrimitiveSceneProxy.cpp
+++ b/Engine/Source/Runtime/Engine/Private/PrimitiveSceneProxy.cpp
@@ -7,6 +7,7 @@
 #include "PrimitiveSceneProxy.h"
 #include "Engine/Brush.h"
 #include "UObject/Package.h"
+#include "EngineModule.h"
 #include "EngineUtils.h"
 #include "Components/BrushComponent.h"
 #include "SceneManagement.h"
@@ -213,6 +214,7 @@ FPrimitiveSceneProxy::FPrimitiveSceneProxy(const UPrimitiveComponent* InComponen
 	InComponent->GetUsedMaterials(UsedMaterialsForVerification, bGetDebugMaterials);
 #endif
 
+	// Store the target runtime virtual texture information
 	if (UseVirtualTexturing(GetScene().GetFeatureLevel()))
 	{
 		for (URuntimeVirtualTexture* VirtualTexture : InComponent->RuntimeVirtualTextures)
@@ -221,6 +223,10 @@ FPrimitiveSceneProxy::FPrimitiveSceneProxy(const UPrimitiveComponent* InComponen
 			{
 				RuntimeVirtualTextures.Add(VirtualTexture);
 				RuntimeVirtualTextureMaterialTypes.Add(VirtualTexture->GetMaterialType());
+			
+				//todo[vt]: Only flush this specific virtual textures
+				//todo[vt]: Only flush primitive bounds 
+				GetRendererModule().FlushVirtualTextureCache();
 			}
 		}
 	}
@@ -238,6 +244,13 @@ void FPrimitiveSceneProxy::SetUsedMaterialForVerification(const TArray<UMaterial
 FPrimitiveSceneProxy::~FPrimitiveSceneProxy()
 {
 	check(IsInRenderingThread());
+
+	for (URuntimeVirtualTexture* VirtualTexture : RuntimeVirtualTextures)
+	{
+		//todo[vt]: Only flush Bounds 
+		//todo[vt]: Only flush specific virtual textures
+		GetRendererModule().FlushVirtualTextureCache();
+	}
 }
 
 HHitProxy* FPrimitiveSceneProxy::CreateHitProxies(UPrimitiveComponent* Component,TArray<TRefCountPtr<HHitProxy> >& OutHitProxies)
diff --git a/Engine/Source/Runtime/Landscape/Private/LandscapeRender.cpp b/Engine/Source/Runtime/Landscape/Private/LandscapeRender.cpp
--- a/Engine/Source/Runtime/Landscape/Private/LandscapeRender.cpp
+++ b/Engine/Source/Runtime/Landscape/Private/LandscapeRender.cpp
@@ -24,6 +24,7 @@ LandscapeRender.cpp: New terrain rendering
 #include "LandscapeMaterialInstanceConstant.h"
 #include "Engine/ShadowMapTexture2D.h"
 #include "EngineGlobals.h"
+#include "EngineModule.h"
 #include "UnrealEngine.h"
 #include "LandscapeLight.h"
 #include "Algo/Find.h"
@@ -861,6 +862,7 @@ FLandscapeComponentSceneProxy::FLandscapeComponentSceneProxy(ULandscapeComponent
 
 	bSupportsHeightfieldRepresentation = true;
 
+	// Store the target runtime virtual texture information
 	if (UseVirtualTexturing(FeatureLevel))
 	{
 		//todo[vt]: Should the virtual texture settings be in the component overrides?
@@ -870,6 +872,10 @@ FLandscapeComponentSceneProxy::FLandscapeComponentSceneProxy(ULandscapeComponent
 			{
 				RuntimeVirtualTextures.Add(VirtualTexture);
 				RuntimeVirtualTextureMaterialTypes.Add(VirtualTexture->GetMaterialType());
+
+				//todo[vt]: Only flush specific virtual textures
+				//todo[vt]: Only flush primitive bounds 
+				GetRendererModule().FlushVirtualTextureCache();
 			}
 		}
 	}
@@ -5035,6 +5041,13 @@ FLandscapeMeshProxySceneProxy::~FLandscapeMeshProxySceneProxy()
 	{
 		Info.UnregisterNeighbors();
 	}
+
+	for (URuntimeVirtualTexture* VirtualTexture : RuntimeVirtualTextures)
+	{
+		//todo[vt]: Only flush specific virtual textures
+		//todo[vt]: Only flush primitive bounds 
+		GetRendererModule().FlushVirtualTextureCache();
+	}
 }
 
 FPrimitiveSceneProxy* ULandscapeMeshProxyComponent::CreateSceneProxy()
diff --git a/Engine/Source/Runtime/Renderer/Private/VT/RuntimeVirtualTextureRender.cpp b/Engine/Source/Runtime/Renderer/Private/VT/RuntimeVirtualTextureRender.cpp
--- a/Engine/Source/Runtime/Renderer/Private/VT/RuntimeVirtualTextureRender.cpp
+++ b/Engine/Source/Runtime/Renderer/Private/VT/RuntimeVirtualTextureRender.cpp
@@ -309,23 +309,26 @@ namespace RuntimeVirtualTexture
 							{
 								// Use cached draw command
 								const int32 StaticMeshCommandInfoIndex = StaticMeshRelevance.GetStaticMeshCommandInfoIndex(EMeshPass::VirtualTexture);
-								FCachedMeshDrawCommandInfo& CachedMeshDrawCommand = PrimitiveSceneInfo->StaticMeshCommandInfos[StaticMeshCommandInfoIndex];
-
-								const FMeshDrawCommand* MeshDrawCommand = CachedMeshDrawCommand.StateBucketId >= 0
-									? &Scene->CachedMeshDrawCommandStateBuckets[FSetElementId::FromInteger(CachedMeshDrawCommand.StateBucketId)].MeshDrawCommand
-									: &SceneDrawList.MeshDrawCommands[CachedMeshDrawCommand.CommandIndex];
-
-								FVisibleMeshDrawCommand NewVisibleMeshDrawCommand;
-								NewVisibleMeshDrawCommand.Setup(
-									MeshDrawCommand,
-									PrimitiveIndex,
-									PrimitiveIndex,
-									CachedMeshDrawCommand.StateBucketId,
-									CachedMeshDrawCommand.MeshFillMode,
-									CachedMeshDrawCommand.MeshCullMode,
-									CachedMeshDrawCommand.SortKey);
-
-								CachedDrawCommands.Add(NewVisibleMeshDrawCommand);
+								if (StaticMeshCommandInfoIndex >= 0)
+								{
+									FCachedMeshDrawCommandInfo& CachedMeshDrawCommand = PrimitiveSceneInfo->StaticMeshCommandInfos[StaticMeshCommandInfoIndex];
+
+									const FMeshDrawCommand* MeshDrawCommand = CachedMeshDrawCommand.StateBucketId >= 0
+										? &Scene->CachedMeshDrawCommandStateBuckets[FSetElementId::FromInteger(CachedMeshDrawCommand.StateBucketId)].MeshDrawCommand
+										: &SceneDrawList.MeshDrawCommands[CachedMeshDrawCommand.CommandIndex];
+
+									FVisibleMeshDrawCommand NewVisibleMeshDrawCommand;
+									NewVisibleMeshDrawCommand.Setup(
+										MeshDrawCommand,
+										PrimitiveIndex,
+										PrimitiveIndex,
+										CachedMeshDrawCommand.StateBucketId,
+										CachedMeshDrawCommand.MeshFillMode,
+										CachedMeshDrawCommand.MeshCullMode,
+										CachedMeshDrawCommand.SortKey);
+
+									CachedDrawCommands.Add(NewVisibleMeshDrawCommand);
+								}
 							}
 							else
 							{
