Fix missing shaders with FLandscapeXYOffsetVertexFactory

diff --git a/Engine/Source/Runtime/Renderer/Private/RayTracing/RayTracingDynamicGeometry.cpp b/Engine/Source/Runtime/Renderer/Private/RayTracing/RayTracingDynamicGeometry.cpp
--- a/Engine/Source/Runtime/Renderer/Private/RayTracing/RayTracingDynamicGeometry.cpp
+++ b/Engine/Source/Runtime/Renderer/Private/RayTracing/RayTracingDynamicGeometry.cpp
@@ -10,7 +10,9 @@ DECLARE_GPU_STAT_NAMED(RayTracingDynamicGeom, TEXT("Ray Tracing Dynamic Geometry
 
 static bool IsSupportedDynamicVertexFactoryType(const FVertexFactoryType* VertexFactoryType)
 {
-	return VertexFactoryType == FindVertexFactoryType(FName(TEXT("FNiagaraSpriteVertexFactory"), FNAME_Find)) || VertexFactoryType == FindVertexFactoryType(FName(TEXT("FLandscapeVertexFactory"), FNAME_Find));
+	return VertexFactoryType == FindVertexFactoryType(FName(TEXT("FNiagaraSpriteVertexFactory"), FNAME_Find))
+		|| VertexFactoryType == FindVertexFactoryType(FName(TEXT("FLandscapeVertexFactory"), FNAME_Find))
+		|| VertexFactoryType == FindVertexFactoryType(FName(TEXT("FLandscapeXYOffsetVertexFactory"), FNAME_Find));
 }
 
 class FRayTracingDynamicGeometryConverterCS : public FMeshMaterialShader
diff --git a/Engine/Source/Runtime/Renderer/Private/RayTracing/RayTracingMaterialHitShaders.cpp b/Engine/Source/Runtime/Renderer/Private/RayTracing/RayTracingMaterialHitShaders.cpp
--- a/Engine/Source/Runtime/Renderer/Private/RayTracing/RayTracingMaterialHitShaders.cpp
+++ b/Engine/Source/Runtime/Renderer/Private/RayTracing/RayTracingMaterialHitShaders.cpp
@@ -33,12 +33,14 @@ static bool IsSupportedVertexFactoryType(const FVertexFactoryType* VertexFactory
 	static FName NiagaraSpriteVfFname = FName(TEXT("FNiagaraSpriteVertexFactory"), FNAME_Find);
 	static FName GeometryCacheVfFname = FName(TEXT("FGeometryCacheVertexVertexFactory"), FNAME_Find);
 	static FName LandscapeVfFname = FName(TEXT("FLandscapeVertexFactory"), FNAME_Find);
+	static FName LandscapeXYOffsetVfFname = FName(TEXT("FLandscapeXYOffsetVertexFactory"), FNAME_Find);
 
 	return VertexFactoryType == FindVertexFactoryType(LocalVfFname)
 		|| VertexFactoryType == FindVertexFactoryType(LSkinnedVfFname)
 		|| VertexFactoryType == FindVertexFactoryType(NiagaraSpriteVfFname)
 		|| VertexFactoryType == FindVertexFactoryType(GeometryCacheVfFname)
-		|| VertexFactoryType == FindVertexFactoryType(LandscapeVfFname);
+		|| VertexFactoryType == FindVertexFactoryType(LandscapeVfFname)
+		|| VertexFactoryType == FindVertexFactoryType(LandscapeXYOffsetVfFname);
 }
 
 class FMaterialCHS : public FMeshMaterialShader, public FUniformLightMapPolicyShaderParametersType
