Enable GPUScene for vulkan.
Serial shadow are now in a seperate renderpass so the update of the PrimitiveShaderBuffer does not happen inside a renderpass.

diff --git a/Engine/Source/Runtime/RenderCore/Public/RenderUtils.h b/Engine/Source/Runtime/RenderCore/Public/RenderUtils.h
--- a/Engine/Source/Runtime/RenderCore/Public/RenderUtils.h
+++ b/Engine/Source/Runtime/RenderCore/Public/RenderUtils.h
@@ -500,7 +500,6 @@ inline bool UseGPUScene(EShaderPlatform Platform, ERHIFeatureLevel::Type Feature
 	return FeatureLevel >= ERHIFeatureLevel::SM5 
 		//@todo - support GPU Scene management compute shaders on these platforms to get dynamic instancing speedups on the Rendering Thread and RHI Thread
 		&& !IsOpenGLPlatform(Platform)
-		&& !IsVulkanPlatform(Platform)
 		&& !IsSwitchPlatform(Platform);
 }
 
diff --git a/Engine/Source/Runtime/Renderer/Private/ShadowDepthRendering.cpp b/Engine/Source/Runtime/Renderer/Private/ShadowDepthRendering.cpp
--- a/Engine/Source/Runtime/Renderer/Private/ShadowDepthRendering.cpp
+++ b/Engine/Source/Runtime/Renderer/Private/ShadowDepthRendering.cpp
@@ -1065,11 +1065,10 @@ void FProjectedShadowInfo::CopyCachedShadowMap(FRHICommandList& RHICmdList, cons
 	}
 }
 
-void FProjectedShadowInfo::RenderDepthInner(FRHICommandListImmediate& RHICmdList, FSceneRenderer* SceneRenderer, FBeginShadowRenderPassFunction BeginShadowRenderPass, bool bDoParallelDispatch)
+
+void FProjectedShadowInfo::SetupShadowUniformBuffers(FRHICommandListImmediate& RHICmdList, FScene* Scene)
 {
 	const ERHIFeatureLevel::Type FeatureLevel = ShadowDepthView->FeatureLevel;
-	FUniformBufferRHIParamRef PassUniformBuffer = nullptr;
-
 	if (FSceneInterface::GetShadingPath(FeatureLevel) == EShadingPath::Deferred)
 	{
 		FShadowDepthPassUniformParameters ShadowDepthPassParameters;
@@ -1078,14 +1077,20 @@ void FProjectedShadowInfo::RenderDepthInner(FRHICommandListImmediate& RHICmdList
 		if (IsWholeSceneDirectionalShadow() && !bReflectiveShadowmap)
 		{
 			check(GetShadowDepthType() == CSMShadowDepthType);
-			SceneRenderer->Scene->UniformBuffers.CSMShadowDepthPassUniformBuffer.UpdateUniformBufferImmediate(ShadowDepthPassParameters);
+			Scene->UniformBuffers.CSMShadowDepthPassUniformBuffer.UpdateUniformBufferImmediate(ShadowDepthPassParameters);
 		}
 
 		ShadowDepthPassUniformBuffer.UpdateUniformBufferImmediate(ShadowDepthPassParameters);
-		PassUniformBuffer = ShadowDepthPassUniformBuffer;
 
-		UploadDynamicPrimitiveShaderDataForView(RHICmdList, *SceneRenderer->Scene, *ShadowDepthView);
+		UploadDynamicPrimitiveShaderDataForView(RHICmdList, *Scene, *ShadowDepthView);
 	}
+}
+
+
+void FProjectedShadowInfo::RenderDepthInner(FRHICommandListImmediate& RHICmdList, FSceneRenderer* SceneRenderer, FBeginShadowRenderPassFunction BeginShadowRenderPass, bool bDoParallelDispatch)
+{
+	const ERHIFeatureLevel::Type FeatureLevel = ShadowDepthView->FeatureLevel;
+	FUniformBufferRHIParamRef PassUniformBuffer = ShadowDepthPassUniformBuffer;
 
 	const bool bIsWholeSceneDirectionalShadow = IsWholeSceneDirectionalShadow();
 
@@ -1388,7 +1393,7 @@ void FSceneRenderer::RenderShadowDepthMapAtlases(FRHICommandListImmediate& RHICm
 						LightEvent,
 						*LightNameWithLevel);
 				}
-
+				ProjectedShadowInfo->SetupShadowUniformBuffers(RHICmdList, Scene);
 				ProjectedShadowInfo->RenderDepth(RHICmdList, this, BeginShadowRenderPass, true);
 			}
 		}
@@ -1402,11 +1407,6 @@ void FSceneRenderer::RenderShadowDepthMapAtlases(FRHICommandListImmediate& RHICm
 
 		if (SerialShadowPasses.Num() > 0)
 		{
-			{
-				SCOPED_DRAW_EVENT(RHICmdList, SetShadowRTsAndClear);
-				BeginShadowRenderPass(RHICmdList, true);
-			}
-
 			for (int32 ShadowIndex = 0; ShadowIndex < SerialShadowPasses.Num(); ShadowIndex++)
 			{
 				FProjectedShadowInfo* ProjectedShadowInfo = SerialShadowPasses[ShadowIndex];
@@ -1429,9 +1429,11 @@ void FSceneRenderer::RenderShadowDepthMapAtlases(FRHICommandListImmediate& RHICm
 						*LightNameWithLevel);
 				}
 
+				ProjectedShadowInfo->SetupShadowUniformBuffers(RHICmdList, Scene);
+				BeginShadowRenderPass(RHICmdList, ShadowIndex == 0);
 				ProjectedShadowInfo->RenderDepth(RHICmdList, this, BeginShadowRenderPass, false);
+				RHICmdList.EndRenderPass();
 			}
-			RHICmdList.EndRenderPass();
 		}
 
 		if (CurrentLightForDrawEvent)
diff --git a/Engine/Source/Runtime/Renderer/Private/ShadowRendering.h b/Engine/Source/Runtime/Renderer/Private/ShadowRendering.h
--- a/Engine/Source/Runtime/Renderer/Private/ShadowRendering.h
+++ b/Engine/Source/Runtime/Renderer/Private/ShadowRendering.h
@@ -510,6 +510,12 @@ public:
 		return FShadowDepthType(bDirectionalLight, bOnePassPointLightShadow, bReflectiveShadowmap);
 	}
 
+	/**
+	* Setup uniformbuffers and update Primitive Shader Data
+	*/
+	void SetupShadowUniformBuffers(FRHICommandListImmediate& RHICmdList, FScene* Scene);
+
+
 private:
 	// 0 if Setup...() wasn't called yet
 	const FLightSceneInfo* LightSceneInfo;
