Fix check(RHICmdList.IsOutsideRenderPass()); when in debug view modes

diff --git a/Engine/Source/Runtime/Renderer/Private/CompositionLighting/PostProcessDeferredDecals.cpp b/Engine/Source/Runtime/Renderer/Private/CompositionLighting/PostProcessDeferredDecals.cpp
--- a/Engine/Source/Runtime/Renderer/Private/CompositionLighting/PostProcessDeferredDecals.cpp
+++ b/Engine/Source/Runtime/Renderer/Private/CompositionLighting/PostProcessDeferredDecals.cpp
@@ -534,6 +534,8 @@ void FRCPassPostProcessDeferredDecals::Process(FRenderingCompositePassContext& C
 	const FSceneViewFamily& ViewFamily = *(View.Family);
 	bool bNeedsDBufferTargets = false;
 	bool bDidClearDBuffer = false;
+	// Debug view framework does not yet support decals.
+	bool bRenderDecals = ViewFamily.EngineShowFlags.Decals && !ViewFamily.UseDebugViewPS();
 
 	if (CurrentStage == DRS_BeforeBasePass)
 	{
@@ -543,7 +545,7 @@ void FRCPassPostProcessDeferredDecals::Process(FRenderingCompositePassContext& C
 		// If we're rendering dbuffer decals but there are no decals in the scene, we avoid the 
 		// clears/decompresses and set the targets to NULL		
 		// The DBufferA-C will be replaced with dummy textures in FSceneTextureShaderParameters
-		if (ViewFamily.EngineShowFlags.Decals)
+		if (bRenderDecals)
 		{
 			FScene& Scene = *(FScene*)ViewFamily.Scene;
 			if (Scene.Decals.Num() > 0 || Context.View.MeshDecalBatches.Num() > 0)
@@ -635,8 +637,7 @@ void FRCPassPostProcessDeferredDecals::Process(FRenderingCompositePassContext& C
 
 	bool bHasValidDBufferMask = false;
 
-	// Debug view framework does not yet support decals.
-	if (ViewFamily.EngineShowFlags.Decals && !ViewFamily.UseDebugViewPS())
+	if (bRenderDecals)
 	{
 		bool bShouldResolveTargets = false;
 
