Allow cubemaps to capture other capture source types. Github: #5227

diff --git a/Engine/Source/Runtime/Engine/Classes/Components/SceneCaptureComponent.h b/Engine/Source/Runtime/Engine/Classes/Components/SceneCaptureComponent.h
--- a/Engine/Source/Runtime/Engine/Classes/Components/SceneCaptureComponent.h
+++ b/Engine/Source/Runtime/Engine/Classes/Components/SceneCaptureComponent.h
@@ -65,6 +65,9 @@ class ENGINE_API USceneCaptureComponent : public USceneComponent
 	UPROPERTY(EditAnywhere, BlueprintReadWrite, Category=SceneCapture)
 	ESceneCapturePrimitiveRenderMode PrimitiveRenderMode;
 
+	UPROPERTY(interp, Category = SceneCapture, meta = (DisplayName = "Capture Source"))
+	TEnumAsByte<enum ESceneCaptureSource> CaptureSource;
+
 	/** The components won't rendered by current component.*/
  	UPROPERTY()
  	TArray<TWeakObjectPtr<UPrimitiveComponent> > HiddenComponents;
diff --git a/Engine/Source/Runtime/Engine/Classes/Components/SceneCaptureComponent2D.h b/Engine/Source/Runtime/Engine/Classes/Components/SceneCaptureComponent2D.h
--- a/Engine/Source/Runtime/Engine/Classes/Components/SceneCaptureComponent2D.h
+++ b/Engine/Source/Runtime/Engine/Classes/Components/SceneCaptureComponent2D.h
@@ -38,9 +38,6 @@ public:
 	UPROPERTY(EditAnywhere, BlueprintReadWrite, Category=SceneCapture)
 	class UTextureRenderTarget2D* TextureTarget;
 
-	UPROPERTY(interp, Category=SceneCapture, meta=(DisplayName = "Capture Source"))
-	TEnumAsByte<enum ESceneCaptureSource> CaptureSource;
-
 	/** When enabled, the scene capture will composite into the render target instead of overwriting its contents. */
 	UPROPERTY(EditAnywhere, BlueprintReadWrite, Category=SceneCapture)
 	TEnumAsByte<enum ESceneCaptureCompositeMode> CompositeMode;
diff --git a/Engine/Source/Runtime/Engine/Private/Components/SceneCaptureComponent.cpp b/Engine/Source/Runtime/Engine/Private/Components/SceneCaptureComponent.cpp
--- a/Engine/Source/Runtime/Engine/Private/Components/SceneCaptureComponent.cpp
+++ b/Engine/Source/Runtime/Engine/Private/Components/SceneCaptureComponent.cpp
@@ -127,6 +127,7 @@ void ASceneCaptureCube::OnInterpToggle(bool bEnable)
 USceneCaptureComponent::USceneCaptureComponent(const FObjectInitializer& ObjectInitializer)
 	: Super(ObjectInitializer), ShowFlags(FEngineShowFlags(ESFIM_Game))
 {
+	CaptureSource = SCS_SceneColorHDR;
 	bCaptureEveryFrame = true;
 	bCaptureOnMovement = true;
 	bAlwaysPersistRenderingState = false;
@@ -437,10 +438,10 @@ USceneCaptureComponent2D::USceneCaptureComponent2D(const FObjectInitializer& Obj
 	PrimaryComponentTick.bCanEverTick = true;
 	PrimaryComponentTick.TickGroup = TG_DuringPhysics;
 	PrimaryComponentTick.bAllowTickOnDedicatedServer = false;
+
 	// Tick in the editor so that bCaptureEveryFrame preview works
 	bTickInEditor = true;
-	// previous behavior was to capture from raw scene color 
-	CaptureSource = SCS_SceneColorHDR;
+
 	// default to full blend weight..
 	PostProcessBlendWeight = 1.0f;
 	CaptureStereoPass = EStereoscopicPass::eSSP_FULL;
diff --git a/Engine/Source/Runtime/Renderer/Private/SceneCaptureRendering.cpp b/Engine/Source/Runtime/Renderer/Private/SceneCaptureRendering.cpp
--- a/Engine/Source/Runtime/Renderer/Private/SceneCaptureRendering.cpp
+++ b/Engine/Source/Runtime/Renderer/Private/SceneCaptureRendering.cpp
@@ -882,8 +882,8 @@ void FScene::UpdateSceneCaptureContents(USceneCaptureComponentCube* CaptureCompo
 					StereoIPD = (CaptureIter == 1) ? CaptureComponent->IPD * -0.5f : CaptureComponent->IPD * 0.5f;
 				}
 
-			FSceneRenderer* SceneRenderer = CreateSceneRendererForSceneCapture(this, CaptureComponent, TextureTarget->GameThread_GetRenderTargetResource(), CaptureSize, ViewRotationMatrix, Location, ProjectionMatrix, CaptureComponent->MaxViewDistanceOverride, true, &PostProcessSettings, 0, CaptureComponent->GetViewOwner(), StereoIPD);
-			SceneRenderer->ViewFamily.SceneCaptureSource = SCS_SceneColorHDR;
+				FSceneRenderer* SceneRenderer = CreateSceneRendererForSceneCapture(this, CaptureComponent, TextureTarget->GameThread_GetRenderTargetResource(), CaptureSize, ViewRotationMatrix, Location, ProjectionMatrix, CaptureComponent->MaxViewDistanceOverride, true, &PostProcessSettings, 0, CaptureComponent->GetViewOwner(), StereoIPD);
+				SceneRenderer->ViewFamily.SceneCaptureSource = CaptureComponent->CaptureSource;
 
 				FTextureRenderTargetCubeResource* TextureRenderTarget = static_cast<FTextureRenderTargetCubeResource*>(TextureTarget->GameThread_GetRenderTargetResource());
 				FString EventName;
