Allow reflection captures supersampling to be set via CVar.

Github: #5807

diff --git a/Engine/Source/Runtime/Renderer/Private/ReflectionEnvironmentCapture.cpp b/Engine/Source/Runtime/Renderer/Private/ReflectionEnvironmentCapture.cpp
--- a/Engine/Source/Runtime/Renderer/Private/ReflectionEnvironmentCapture.cpp
+++ b/Engine/Source/Runtime/Renderer/Private/ReflectionEnvironmentCapture.cpp
@@ -45,6 +45,14 @@
 float GReflectionCaptureNearPlane = 5;
 
 int32 GSupersampleCaptureFactor = 1;
+static FAutoConsoleVariableRef CVarGSupersampleCaptureFactor(
+	TEXT("r.ReflectionCaptureSupersampleFactor"),
+	GSupersampleCaptureFactor,
+	TEXT("Super sample factor when rendering reflection captures.\n")
+	TEXT("Default = 1, no super sampling\n")
+	TEXT("Maximum clamped to 8."),
+	ECVF_RenderThreadSafe
+	);
 
 /** 
  * Mip map used by a Roughness of 0, counting down from the lowest resolution mip (MipCount - 1).  
@@ -791,12 +799,14 @@ void CaptureSceneToScratchCubemap(FRHICommandListImmediate& RHICmdList, FSceneRe
 			PixelShader->SetParameters(RHICmdList, SceneRenderer->Views[0], bCapturingForSkyLight, bLowerHemisphereIsBlack, LowerHemisphereColor);
 			VertexShader->SetParameters(RHICmdList, SceneRenderer->Views[0]);
 
+			int32 SupersampleCaptureFactor = FMath::Clamp(GSupersampleCaptureFactor, 1, 8);
+
 			DrawRectangle( 
 				RHICmdList,
 				ViewRect.Min.X, ViewRect.Min.Y, 
 				ViewRect.Width(), ViewRect.Height(),
 				ViewRect.Min.X, ViewRect.Min.Y, 
-				ViewRect.Width() * GSupersampleCaptureFactor, ViewRect.Height() * GSupersampleCaptureFactor,
+				ViewRect.Width() * SupersampleCaptureFactor, ViewRect.Height() * SupersampleCaptureFactor,
 				FIntPoint(ViewRect.Width(), ViewRect.Height()),
 				SceneContext.GetBufferSizeXY(),
 				*VertexShader);
@@ -1277,6 +1287,8 @@ void CaptureSceneIntoScratchCubemap(
 	const FLinearColor& LowerHemisphereColor
 	)
 {
+	int32 SupersampleCaptureFactor = FMath::Clamp(GSupersampleCaptureFactor, 1, 8);
+
 	for (int32 CubeFace = 0; CubeFace < CubeFace_MAX; CubeFace++)
 	{
 		if( !bCapturingForSkyLight )
@@ -1329,7 +1341,7 @@ void CaptureSceneIntoScratchCubemap(
 		ViewInitOptions.ViewFamily = &ViewFamily;
 		ViewInitOptions.BackgroundColor = FLinearColor::Black;
 		ViewInitOptions.OverlayColor = FLinearColor::Black;
-		ViewInitOptions.SetViewRectangle(FIntRect(0, 0, CubemapSize * GSupersampleCaptureFactor, CubemapSize * GSupersampleCaptureFactor));
+		ViewInitOptions.SetViewRectangle(FIntRect(0, 0, CubemapSize * SupersampleCaptureFactor, CubemapSize * SupersampleCaptureFactor));
 
 		const float NearPlane = bCapturingForSkyLight ? SkyLightNearPlane : GReflectionCaptureNearPlane;
 
@@ -1339,8 +1351,8 @@ void CaptureSceneIntoScratchCubemap(
 		{
 			ViewInitOptions.ProjectionMatrix = FReversedZPerspectiveMatrix(
 				90.0f * (float)PI / 360.0f,
-				(float)CubemapSize * GSupersampleCaptureFactor,
-				(float)CubemapSize * GSupersampleCaptureFactor,
+				(float)CubemapSize * SupersampleCaptureFactor,
+				(float)CubemapSize * SupersampleCaptureFactor,
 				NearPlane
 				);
 		}
@@ -1348,8 +1360,8 @@ void CaptureSceneIntoScratchCubemap(
 		{
 			ViewInitOptions.ProjectionMatrix = FPerspectiveMatrix(
 				90.0f * (float)PI / 360.0f,
-				(float)CubemapSize * GSupersampleCaptureFactor,
-				(float)CubemapSize * GSupersampleCaptureFactor,
+				(float)CubemapSize * SupersampleCaptureFactor,
+				(float)CubemapSize * SupersampleCaptureFactor,
 				NearPlane
 				);
 		}
