Replace hard coded values with constexpr value.

Github: #5807

diff --git a/Engine/Source/Runtime/Renderer/Private/ReflectionEnvironmentCapture.cpp b/Engine/Source/Runtime/Renderer/Private/ReflectionEnvironmentCapture.cpp
--- a/Engine/Source/Runtime/Renderer/Private/ReflectionEnvironmentCapture.cpp
+++ b/Engine/Source/Runtime/Renderer/Private/ReflectionEnvironmentCapture.cpp
@@ -44,6 +44,9 @@
 /** Near plane to use when capturing the scene. */
 float GReflectionCaptureNearPlane = 5;
 
+constexpr int32 MinSupersampleCaptureFactor = 1;
+constexpr int32 MaxSupersampleCaptureFactor = 8;
+
 int32 GSupersampleCaptureFactor = 1;
 static FAutoConsoleVariableRef CVarGSupersampleCaptureFactor(
 	TEXT("r.ReflectionCaptureSupersampleFactor"),
@@ -799,7 +802,7 @@ void CaptureSceneToScratchCubemap(FRHICommandListImmediate& RHICmdList, FSceneRe
 			PixelShader->SetParameters(RHICmdList, SceneRenderer->Views[0], bCapturingForSkyLight, bLowerHemisphereIsBlack, LowerHemisphereColor);
 			VertexShader->SetParameters(RHICmdList, SceneRenderer->Views[0]);
 
-			int32 SupersampleCaptureFactor = FMath::Clamp(GSupersampleCaptureFactor, 1, 8);
+			int32 SupersampleCaptureFactor = FMath::Clamp(GSupersampleCaptureFactor, MinSupersampleCaptureFactor, MaxSupersampleCaptureFactor);
 
 			DrawRectangle( 
 				RHICmdList,
@@ -1287,7 +1290,7 @@ void CaptureSceneIntoScratchCubemap(
 	const FLinearColor& LowerHemisphereColor
 	)
 {
-	int32 SupersampleCaptureFactor = FMath::Clamp(GSupersampleCaptureFactor, 1, 8);
+	int32 SupersampleCaptureFactor = FMath::Clamp(GSupersampleCaptureFactor, MinSupersampleCaptureFactor, MaxSupersampleCaptureFactor);
 
 	for (int32 CubeFace = 0; CubeFace < CubeFace_MAX; CubeFace++)
 	{
