Remove destroy callbacks even after producer has been released
Without this change we don't remove pending callbacks from the list (and the baton ptr might be stale by the time the callback is called)

diff --git a/Engine/Source/Runtime/Renderer/Private/VT/VirtualTextureProducer.cpp b/Engine/Source/Runtime/Renderer/Private/VT/VirtualTextureProducer.cpp
--- a/Engine/Source/Runtime/Renderer/Private/VT/VirtualTextureProducer.cpp
+++ b/Engine/Source/Runtime/Renderer/Private/VT/VirtualTextureProducer.cpp
@@ -148,14 +148,10 @@ void FVirtualTextureProducerCollection::RemoveAllCallbacks(const void* Baton)
 		if (Callback.Baton == Baton)
 		{
 			check(Callback.DestroyedFunction);
-			FProducerEntry* Entry = GetEntry(Callback.OwnerHandle);
-			if (Entry)
-			{
-				Callback.DestroyedFunction = nullptr;
-				Callback.Baton = nullptr;
-				Callback.OwnerHandle = FVirtualTextureProducerHandle();
-				ReleaseCallback(CallbackIndex);
-			}
+			Callback.DestroyedFunction = nullptr;
+			Callback.Baton = nullptr;
+			Callback.OwnerHandle = FVirtualTextureProducerHandle();
+			ReleaseCallback(CallbackIndex);
 		}
 	}
 }
