Refactors RenderDeferredReflectionsAndSkyLighting()'s entire view loop to RDG

diff --git a/Engine/Shaders/Private/DistanceFieldAOShared.ush b/Engine/Shaders/Private/DistanceFieldAOShared.ush
--- a/Engine/Shaders/Private/DistanceFieldAOShared.ush
+++ b/Engine/Shaders/Private/DistanceFieldAOShared.ush
@@ -168,10 +168,10 @@ float3 ComputeBentNormal(float3 RecordWorldNormal, uint RelativeRecordIndex)
 
 float2 AOBufferBilinearUVMax;
 
-float3 UpsampleDFAO(float4 UVAndScreenPos)
+float3 UpsampleDFAO(float2 BufferUV, float SceneDepth)
 {
 	// Distance field AO was computed at 0,0 regardless of viewrect min
-	float2 DistanceFieldUVs = UVAndScreenPos.xy - View.ViewRectMin.xy * View.BufferSizeAndInvSize.zw;
+	float2 DistanceFieldUVs = BufferUV - View.ViewRectMin.xy * View.BufferSizeAndInvSize.zw;
 	DistanceFieldUVs = min(DistanceFieldUVs, AOBufferBilinearUVMax);
 
 #define BILATERAL_UPSAMPLE 1
@@ -195,7 +195,6 @@ float3 UpsampleDFAO(float4 UVAndScreenPos)
 	float Epsilon = .0001f;
 
 	float4 CornerDepths = float4(TextureValues00.w, TextureValues10.w, TextureValues01.w, TextureValues11.w);
-	float SceneDepth = CalcSceneDepth(UVAndScreenPos.xy);
 	float4 DepthWeights = 1.0f / (abs(CornerDepths - SceneDepth.xxxx) + Epsilon);
 	float4 FinalWeights = CornerWeights * DepthWeights;
 
diff --git a/Engine/Shaders/Private/DistanceFieldLightingPost.usf b/Engine/Shaders/Private/DistanceFieldLightingPost.usf
--- a/Engine/Shaders/Private/DistanceFieldLightingPost.usf
+++ b/Engine/Shaders/Private/DistanceFieldLightingPost.usf
@@ -372,7 +372,7 @@ void AOUpsamplePS(
 	,out float4 OutSceneColor : SV_Target0
 	)
 {
-	float3 BentNormal = UpsampleDFAO(UVAndScreenPos);
+	float3 BentNormal = UpsampleDFAO(UVAndScreenPos.xy, CalcSceneDepth(UVAndScreenPos.xy));
 
 #if MODULATE_SCENE_COLOR
 	float Visibility = lerp(length(BentNormal), 1.0f, MinIndirectDiffuseOcclusion);
diff --git a/Engine/Shaders/Private/PlanarReflectionShaders.usf b/Engine/Shaders/Private/PlanarReflectionShaders.usf
--- a/Engine/Shaders/Private/PlanarReflectionShaders.usf
+++ b/Engine/Shaders/Private/PlanarReflectionShaders.usf
@@ -7,6 +7,7 @@
 #include "Common.ush"
 #include "DeferredShadingCommon.ush"
 #include "PlanarReflectionShared.ush"
+#include "SceneViewFamilyBlackboard.ush"
 
 Texture2D SceneColorInputTexture;
 SamplerState SceneColorInputSampler;
@@ -123,7 +124,7 @@ void PrefilterPlanarReflectionPS(
 }
 
 void PlanarReflectionPS(
-	float2 InUV : TEXCOORD0,
+	float2 BufferUV : TEXCOORD0,
 	float3 ScreenVector : TEXCOORD1,
 	float4 SVPos : SV_POSITION,
 	out float4 OutColor : SV_Target0
@@ -132,16 +133,15 @@ void PlanarReflectionPS(
 	ResolvedView = ResolveView();
 	OutColor = 0;
 
-	float SceneDepth = CalcSceneDepth(InUV);
-	float3 WorldPosition = ScreenVector * SceneDepth + View.WorldCameraOrigin;
-		
-	FScreenSpaceData ScreenSpaceData = GetScreenSpaceData(InUV);
+	FGBufferData GBuffer = GetGBufferDataFromSceneBlackboard(BufferUV);
+
+	float3 WorldPosition = ScreenVector * GBuffer.Depth + View.WorldCameraOrigin;
 
 	// Skip unlit pixels (skybox)
 	BRANCH
-	if (ScreenSpaceData.GBuffer.ShadingModelID > 0)
+	if (GBuffer.ShadingModelID != SHADINGMODELID_UNLIT)
 	{
-		OutColor = ComputePlanarReflections(WorldPosition, ScreenSpaceData.GBuffer.WorldNormal, ScreenSpaceData.GBuffer.Roughness, PlanarReflectionStruct.PlanarReflectionSampler);
+		OutColor = ComputePlanarReflections(WorldPosition, GBuffer.WorldNormal, GBuffer.Roughness, PlanarReflectionStruct.PlanarReflectionSampler);
 
 	#if USE_PREEXPOSURE
 		OutColor.xyz *= View.PreExposure;
diff --git a/Engine/Shaders/Private/ReflectionEnvironmentPixelShader.usf b/Engine/Shaders/Private/ReflectionEnvironmentPixelShader.usf
--- a/Engine/Shaders/Private/ReflectionEnvironmentPixelShader.usf
+++ b/Engine/Shaders/Private/ReflectionEnvironmentPixelShader.usf
@@ -12,6 +12,7 @@
 #include "DistanceFieldAOShared.ush"
 #include "ShadingModels.ush"
 #include "LightGridCommon.ush"
+#include "SceneViewFamilyBlackboard.ush"
 
 #define REFLECTION_COMPOSITE_USE_BLENDED_REFLECTION_CAPTURES 1
 #define REFLECTION_COMPOSITE_SUPPORT_SKYLIGHT_BLEND 1
@@ -47,17 +48,20 @@ float3 GatherRadiance(float CompositeAlpha, float3 WorldPosition, float3 RayDire
 Texture2D ScreenSpaceReflectionsTexture;
 SamplerState ScreenSpaceReflectionsSampler;
 
+
+Texture2D AmbientOcclusionTexture;
+SamplerState AmbientOcclusionSampler;
+
+
 float3 ContrastAndNormalizeMulAdd;
 float OcclusionExponent;
 float OcclusionCombineMode;
 
-float3 SkyLightDiffuse(FScreenSpaceData ScreenSpaceData, float4 UVAndScreenPos, float3 BentNormal)
+float3 SkyLightDiffuse(FGBufferData GBuffer, float AmbientOcclusion, float2 BufferUV, float2 ScreenPosition, float3 BentNormal)
 {
-	float2 UV = UVAndScreenPos.xy;
+	float2 UV = BufferUV;
 	float3 Lighting = 0;
 
-	FGBufferData GBuffer = ScreenSpaceData.GBuffer;
-
 	float SkyVisibility = 1;
 	float DotProductFactor = 1;
 	float3 SkyLightingNormal = GBuffer.WorldNormal;
@@ -69,7 +73,7 @@ float3 SkyLightDiffuse(FScreenSpaceData ScreenSpaceData, float4 UVAndScreenPos,
 	
 	float3 DiffuseIrradiance = 0;
 
-#if APPLY_SKY_SHADOWING
+#if APPLY_SKY_SHADOWING && 0
 #define USE_DIRECTIONAL_OCCLUSION_ON_SKY_DIFFUSE 1
 
 #if USE_DIRECTIONAL_OCCLUSION_ON_SKY_DIFFUSE
@@ -105,12 +109,12 @@ float3 SkyLightDiffuse(FScreenSpaceData ScreenSpaceData, float4 UVAndScreenPos,
 	if (OcclusionCombineMode == 0)
 	{
 		// Combine with min which nicely avoids over-occlusion in cases where strong DFAO is present along with strong SSAO (distant trees)
-		SkyVisibility = min(SkyVisibility, min(GBuffer.GBufferAO, ScreenSpaceData.AmbientOcclusion));
+		SkyVisibility = min(SkyVisibility, min(GBuffer.GBufferAO, AmbientOcclusion));
 	}
 	else
 	{
 		// Combine with mul, which continues to add SSAO depth even indoors.  SSAO will need to be tweaked to be less strong.
-		SkyVisibility = SkyVisibility * min(GBuffer.GBufferAO, ScreenSpaceData.AmbientOcclusion);
+		SkyVisibility = SkyVisibility * min(GBuffer.GBufferAO, AmbientOcclusion);
 	}
 
 	float ScalarFactors = SkyVisibility;
@@ -136,7 +140,7 @@ float3 SkyLightDiffuse(FScreenSpaceData ScreenSpaceData, float4 UVAndScreenPos,
 	if (GBuffer.ShadingModelID == SHADINGMODELID_HAIR)
 	{
 		float3 N = GBuffer.WorldNormal;
-		float3 V = -normalize(mul(float4(UVAndScreenPos.zw, 1, 0), View.ScreenToWorld).xyz);
+		float3 V = -normalize(mul(float4(ScreenPosition, 1, 0), View.ScreenToWorld).xyz);
 		float3 L = normalize(V - N * dot(V, N));
 		SkyLightingNormal = L;
 		DiffuseColor = PI * HairShading(GBuffer, L, V, N, 1, 0, 0.2, uint2(0, 0));
@@ -154,7 +158,7 @@ float3 SkyLightDiffuse(FScreenSpaceData ScreenSpaceData, float4 UVAndScreenPos,
 	// apply the Diffuse color to the lighting (including OcclusionTintAndMinOcclusion as it's considered another light, that fixes SubsurfaceProfile being too dark)
 	Lighting += ((ScalarFactors * DotProductFactor) * DiffuseLookup + (1 - SkyVisibility) * OcclusionTintAndMinOcclusion.xyz) * DiffuseColor;
 
-	Lighting += DiffuseIrradiance * GBuffer.DiffuseColor * (GBuffer.GBufferAO * ScreenSpaceData.AmbientOcclusion);
+	Lighting += DiffuseIrradiance * GBuffer.DiffuseColor * (GBuffer.GBufferAO * AmbientOcclusion);
 
 #if USE_PREEXPOSURE
 	Lighting *= View.PreExposure;
@@ -170,7 +174,7 @@ float3 SkyLightDiffuse(FScreenSpaceData ScreenSpaceData, float4 UVAndScreenPos,
 	return Lighting;
 }
 
-float3 ReflectionEnvironment(FScreenSpaceData ScreenSpaceData, float4 UVAndScreenPos, float4 SvPosition, float3 BentNormal)
+float3 ReflectionEnvironment(FGBufferData GBuffer, float AmbientOcclusion, float2 BufferUV, float2 ScreenPosition, float4 SvPosition, float3 BentNormal)
 {
 #if USE_PREEXPOSURE
 	const float PreExposure = View.PreExposure;
@@ -180,12 +184,7 @@ float3 ReflectionEnvironment(FScreenSpaceData ScreenSpaceData, float4 UVAndScree
 	
 	float4 Color = float4(0, 0, 0, 1);
 
-	float2 ScreenUV = UVAndScreenPos.xy;
-	float2 ScreenPosition = UVAndScreenPos.zw;
-	FGBufferData GBuffer = ScreenSpaceData.GBuffer;
-	float SceneDepth = CalcSceneDepth(ScreenUV);
-	
-	float3 WorldPosition = mul(float4(ScreenPosition * SceneDepth, SceneDepth, 1), View.ScreenToWorld).xyz;
+	float3 WorldPosition = mul(float4(ScreenPosition * GBuffer.Depth, GBuffer.Depth, 1), View.ScreenToWorld).xyz;
 	float3 CameraToPixel = normalize(WorldPosition - View.WorldCameraOrigin);
 	float3 ReflectionVector = reflect(CameraToPixel, GBuffer.WorldNormal);
 	float IndirectIrradiance = GBuffer.IndirectIrradiance;
@@ -195,7 +194,6 @@ float3 ReflectionEnvironment(FScreenSpaceData ScreenSpaceData, float4 UVAndScree
 	// Add in diffuse contribution from dynamic skylights so reflection captures will have something to mix with
 	if (ReflectionStruct.SkyLightParameters.y > 0 && ReflectionStruct.SkyLightParameters.z > 0)
 	{
-		float2 ScreenUV = ScreenPosition * View.ScreenPositionScaleBias.xy + View.ScreenPositionScaleBias.wz;
 		IndirectIrradiance += GetDynamicSkyIndirectIrradiance(BentNormal, GBuffer.WorldNormal);
 	}
 #endif
@@ -211,7 +209,7 @@ float3 ReflectionEnvironment(FScreenSpaceData ScreenSpaceData, float4 UVAndScree
 
 #if 1
 	// Note: this texture may also contain planar reflections
-	float4 SSR = Texture2DSample( ScreenSpaceReflectionsTexture, ScreenSpaceReflectionsSampler, ScreenUV );
+	float4 SSR = Texture2DSample(ScreenSpaceReflectionsTexture, ScreenSpaceReflectionsSampler, BufferUV);
 	Color.rgb = SSR.rgb;
 	Color.a = 1 - SSR.a;
 #endif
@@ -236,7 +234,7 @@ float3 ReflectionEnvironment(FScreenSpaceData ScreenSpaceData, float4 UVAndScree
 #endif
 	}
 
-	float AO = GBuffer.GBufferAO * ScreenSpaceData.AmbientOcclusion;
+	float AO = GBuffer.GBufferAO * AmbientOcclusion;
 	float RoughnessSq = GBuffer.Roughness * GBuffer.Roughness;
 	float SpecularOcclusion = GetSpecularOcclusion(NoV, RoughnessSq, AO);
 	Color.a *= SpecularOcclusion;
@@ -244,7 +242,7 @@ float3 ReflectionEnvironment(FScreenSpaceData ScreenSpaceData, float4 UVAndScree
 #if FEATURE_LEVEL >= FEATURE_LEVEL_SM5
 	float2 LocalPosition = SvPosition.xy - View.ViewRectMin.xy;
 
-	uint GridIndex = ComputeLightGridCellIndex(uint2(LocalPosition.x, LocalPosition.y), SceneDepth);
+	uint GridIndex = ComputeLightGridCellIndex(uint2(LocalPosition.x, LocalPosition.y), GBuffer.Depth);
 	uint NumCulledEntryIndex = (ForwardLightData.NumGridCells + GridIndex) * NUM_CULLED_LIGHTS_GRID_STRIDE;
 	uint NumCulledReflectionCaptures = min(ForwardLightData.NumCulledLightsGrid[NumCulledEntryIndex + 0], ForwardLightData.NumReflectionCaptures);
 	uint DataStartIndex = ForwardLightData.NumCulledLightsGrid[NumCulledEntryIndex + 1];
@@ -298,27 +296,34 @@ float3 ReflectionEnvironment(FScreenSpaceData ScreenSpaceData, float4 UVAndScree
 }
 
 void ReflectionEnvironmentSkyLighting(
-	in noperspective float4 UVAndScreenPos : TEXCOORD0,
 	in float4 SvPosition : SV_Position,
 	out float4 OutColor : SV_Target0)
 {
-	float2 ScreenUV = UVAndScreenPos.xy;
-	FScreenSpaceData ScreenSpaceData = GetScreenSpaceData(ScreenUV);
-	uint ShadingModelID = ScreenSpaceData.GBuffer.ShadingModelID;
+	float2 BufferUV = SvPositionToBufferUV(SvPosition);
+	float2 ScreenPosition = SvPositionToScreenPosition(SvPosition).xy;
 
-	OutColor = 0.0f;
+	// Sample scene textures.
+	FGBufferData GBuffer = GetGBufferDataFromSceneBlackboard(BufferUV);
 
-	float3 BentNormal = ScreenSpaceData.GBuffer.WorldNormal;
-#if APPLY_SKY_SHADOWING
-	BentNormal = UpsampleDFAO(UVAndScreenPos);
-#endif
+	// Sample the ambient occlusion that is dynamically generated every frame.
+	float AmbientOcclusion = AmbientOcclusionTexture.SampleLevel(AmbientOcclusionSampler, BufferUV, 0).r;
+	
+	uint ShadingModelID = GBuffer.ShadingModelID;
+	float3 BentNormal = GBuffer.WorldNormal;
+	#if APPLY_SKY_SHADOWING
+	{
+		BentNormal = UpsampleDFAO(BufferUV, GBuffer.Depth);
+	}
+	#endif
+	
+	OutColor = 0.0f;
 
+	#if ENABLE_DYNAMIC_SKY_LIGHT
 	BRANCH
 	// Only light pixels marked as lit
-	if (ShadingModelID > 0)
+	if (ShadingModelID != SHADINGMODELID_UNLIT)
 	{
-#if ENABLE_DYNAMIC_SKY_LIGHT
-		float3 SkyLighting = SkyLightDiffuse(ScreenSpaceData, UVAndScreenPos, BentNormal);
+		float3 SkyLighting = SkyLightDiffuse(GBuffer, AmbientOcclusion, BufferUV, ScreenPosition, BentNormal);
 
 		// Scene color alpha is used for ScreenSpaceSubsurfaceScattering (if that is not needed it can be disabled with SUBSURFACE_CHANNEL_MODE)
 		FLightAccumulator LightAccumulator = (FLightAccumulator)0;
@@ -327,13 +332,13 @@ void ReflectionEnvironmentSkyLighting(
 
 		LightAccumulator_Add(LightAccumulator, SkyLighting, SkyLighting, 1.0f, bNeedsSeparateSubsurfaceLightAccumulation);
 		OutColor = LightAccumulator_GetResult(LightAccumulator);
-#endif
 	}
+	#endif
 
 	BRANCH
 	if (ShadingModelID != SHADINGMODELID_UNLIT && ShadingModelID != SHADINGMODELID_HAIR)
 	{
-		OutColor.xyz += ReflectionEnvironment(ScreenSpaceData, UVAndScreenPos, SvPosition, BentNormal);
+		OutColor.xyz += ReflectionEnvironment(GBuffer, AmbientOcclusion, BufferUV, ScreenPosition, SvPosition, BentNormal);
 	}
 }
 
diff --git a/Engine/Source/Runtime/Renderer/Private/DeferredShadingRenderer.h b/Engine/Source/Runtime/Renderer/Private/DeferredShadingRenderer.h
--- a/Engine/Source/Runtime/Renderer/Private/DeferredShadingRenderer.h
+++ b/Engine/Source/Runtime/Renderer/Private/DeferredShadingRenderer.h
@@ -440,7 +440,7 @@ private:
 	/** Render image based reflections (SSR, Env, SkyLight) without compute shaders */
 	void RenderStandardDeferredImageBasedReflections(FRHICommandListImmediate& RHICmdList, FGraphicsPipelineStateInitializer& GraphicsPSOInit, bool bReflectionEnv, const TRefCountPtr<IPooledRenderTarget>& DynamicBentNormalAO, TRefCountPtr<IPooledRenderTarget>& VelocityRT);
 
-	bool RenderDeferredPlanarReflections(FRHICommandListImmediate& RHICmdList, const FViewInfo& View, bool bLightAccumulationIsInUse, TRefCountPtr<IPooledRenderTarget>& Output);
+	void RenderDeferredPlanarReflections(FRDGBuilder& GraphBuilder, const FSceneViewFamilyBlackboard& SceneTextures, const FViewInfo& View, FRDGTextureRef& ReflectionsOutput);
 
 	bool ShouldDoReflectionEnvironment() const;
 	
diff --git a/Engine/Source/Runtime/Renderer/Private/PlanarReflectionRendering.cpp b/Engine/Source/Runtime/Renderer/Private/PlanarReflectionRendering.cpp
--- a/Engine/Source/Runtime/Renderer/Private/PlanarReflectionRendering.cpp
+++ b/Engine/Source/Runtime/Renderer/Private/PlanarReflectionRendering.cpp
@@ -31,6 +31,7 @@
 #include "Containers/ArrayView.h"
 #include "PipelineStateCache.h"
 #include "ClearQuad.h"
+#include "SceneViewFamilyBlackboard.h"
 
 void SetupPlanarReflectionUniformParameters(const class FSceneView& View, const FPlanarReflectionSceneProxy* ReflectionSceneProxy, FPlanarReflectionUniformParameters& OutParameters)
 {
@@ -642,53 +643,29 @@ void FScene::UpdatePlanarReflectionTransform(UPlanarReflectionComponent* Compone
 
 class FPlanarReflectionPS : public FGlobalShader
 {
-	DECLARE_SHADER_TYPE(FPlanarReflectionPS, Global);
-public:
+	DECLARE_GLOBAL_SHADER(FPlanarReflectionPS);
+	SHADER_USE_PARAMETER_STRUCT(FPlanarReflectionPS, FGlobalShader)
 
 	static bool ShouldCompilePermutation(const FGlobalShaderPermutationParameters& Parameters)
 	{
 		return IsFeatureLevelSupported(Parameters.Platform, ERHIFeatureLevel::SM4);
 	}
 
-	/** Default constructor. */
-	FPlanarReflectionPS() {}
+	BEGIN_SHADER_PARAMETER_STRUCT(FParameters, )
+		SHADER_PARAMETER_STRUCT_INCLUDE(FSceneViewFamilyBlackboard, SceneTextures)
+		SHADER_PARAMETER_STRUCT_INCLUDE(FSceneTextureSamplerParameters, SceneTextureSamplers)
 
-	/** Initialization constructor. */
-	FPlanarReflectionPS(const ShaderMetaType::CompiledShaderInitializerType& Initializer)
-		: FGlobalShader(Initializer)
-	{
-		SceneTextureParameters.Bind(Initializer);
-	}
+		SHADER_PARAMETER_STRUCT_REF(FViewUniformShaderParameters, ViewUniformBuffer)
+		SHADER_PARAMETER_STRUCT_REF(FPlanarReflectionUniformParameters, PlanarReflectionParameters)
 
-	void SetParameters(FRHICommandList& RHICmdList, const FSceneView& View, FPlanarReflectionSceneProxy* ReflectionSceneProxy)
-	{
-		const FPixelShaderRHIParamRef ShaderRHI = GetPixelShader();
-		FGlobalShader::SetParameters<FViewUniformShaderParameters>(RHICmdList, ShaderRHI, View.ViewUniformBuffer);
-		SceneTextureParameters.Set(RHICmdList, ShaderRHI, View.FeatureLevel, ESceneTextureSetupMode::All);
-
-		FPlanarReflectionUniformParameters PlanarReflectionUniformParameters;
-		SetupPlanarReflectionUniformParameters(View, ReflectionSceneProxy, PlanarReflectionUniformParameters);
-		SetUniformBufferParameterImmediate(RHICmdList, ShaderRHI, GetUniformBufferParameter<FPlanarReflectionUniformParameters>(), PlanarReflectionUniformParameters);
-	}
-
-	// FShader interface.
-	virtual bool Serialize(FArchive& Ar) override
-	{
-		bool bShaderHasOutdatedParameters = FGlobalShader::Serialize(Ar);
-		Ar << SceneTextureParameters;
-		return bShaderHasOutdatedParameters;
-	}
-
-private:
-
-	FSceneTextureShaderParameters SceneTextureParameters;
+		RENDER_TARGET_BINDING_SLOTS()
+	END_SHADER_PARAMETER_STRUCT()
 };
 
-IMPLEMENT_SHADER_TYPE(,FPlanarReflectionPS,TEXT("/Engine/Private/PlanarReflectionShaders.usf"),TEXT("PlanarReflectionPS"),SF_Pixel);
+IMPLEMENT_GLOBAL_SHADER(FPlanarReflectionPS, "/Engine/Private/PlanarReflectionShaders.usf", "PlanarReflectionPS", SF_Pixel);
 
-bool FDeferredShadingSceneRenderer::RenderDeferredPlanarReflections(FRHICommandListImmediate& RHICmdList, const FViewInfo& View, bool bLightAccumulationIsInUse, TRefCountPtr<IPooledRenderTarget>& Output)
+void FDeferredShadingSceneRenderer::RenderDeferredPlanarReflections(FRDGBuilder& GraphBuilder, const FSceneViewFamilyBlackboard& SceneTextures, const FViewInfo& View, FRDGTextureRef& ReflectionsOutputTexture)
 {
-	check(RHICmdList.IsOutsideRenderPass());
 	// Prevent rendering unsupported views when ViewIndex >= GMaxPlanarReflectionViews
 	// Planar reflections in those views will fallback to other reflection methods
 	{
@@ -698,7 +675,7 @@ bool FDeferredShadingSceneRenderer::RenderDeferredPlanarReflections(FRHICommandL
 
 		if (ViewIndex >= GMaxPlanarReflectionViews)
 		{
-			return false;
+			return;
 		}
 	}
 
@@ -716,90 +693,95 @@ bool FDeferredShadingSceneRenderer::RenderDeferredPlanarReflections(FRHICommandL
 
 	bool bViewIsReflectionCapture = View.bIsPlanarReflection || View.bIsReflectionCapture;
 
+	bool bComposePlanarReflections = Scene->PlanarReflections.Num() > 0 && !bViewIsReflectionCapture && bAnyVisiblePlanarReflections;
+
 	// Prevent reflection recursion, or view-dependent planar reflections being seen in reflection captures
-	if (Scene->PlanarReflections.Num() > 0 && !bViewIsReflectionCapture && bAnyVisiblePlanarReflections)
+	if (!bComposePlanarReflections)
 	{
-		SCOPED_DRAW_EVENT(RHICmdList, CompositePlanarReflections);
+		return;
+	}
 
-		bool bSSRAsInput = true;
+	// Allocate planar reflection texture
+	bool bClearReflectionsOutputTexture = false;
+	if (!ReflectionsOutputTexture)
+	{
+		FRDGTextureDesc Desc = FPooledRenderTargetDesc::Create2DDesc(
+			SceneTextures.SceneDepthBuffer->Desc.Extent,
+			PF_FloatRGBA, FClearValueBinding(FLinearColor(0, 0, 0, 0)),
+			TexCreate_None, TexCreate_RenderTargetable,
+			false);
 
-		if (Output == GSystemTextures.BlackDummy)
-		{
-			bSSRAsInput = false;
-			FSceneRenderTargets& SceneContext = FSceneRenderTargets::Get(RHICmdList);
+		Desc.AutoWritable = false;
 
-			if (bLightAccumulationIsInUse)
-			{
-				FPooledRenderTargetDesc Desc(FPooledRenderTargetDesc::Create2DDesc(SceneContext.GetBufferSizeXY(), PF_FloatRGBA, FClearValueBinding::Black, TexCreate_None, TexCreate_RenderTargetable, false));
-				GRenderTargetPool.FindFreeElement(RHICmdList, Desc, Output, TEXT("PlanarReflectionComposite"));
-			}
-			else
-			{
-				Output = SceneContext.LightAccumulation;
-			}
-		}
+		ReflectionsOutputTexture = GraphBuilder.CreateTexture(Desc, TEXT("PlanarReflections"));
+		bClearReflectionsOutputTexture = true;
+	}
+
+	FPlanarReflectionPS::FParameters* PassParameters = GraphBuilder.AllocParameters<FPlanarReflectionPS::FParameters>();
+	PassParameters->SceneTextures.SceneDepthBuffer = SceneTextures.SceneDepthBuffer;
+	PassParameters->SceneTextures.SceneGBufferA = SceneTextures.SceneGBufferA;
+	PassParameters->SceneTextures.SceneGBufferB = SceneTextures.SceneGBufferB;
+	SetupSceneTextureSamplers(&PassParameters->SceneTextureSamplers);
+	PassParameters->ViewUniformBuffer = View.ViewUniformBuffer;
+	PassParameters->RenderTargets[0] = FRenderTargetBinding(
+		ReflectionsOutputTexture, bClearReflectionsOutputTexture ? ERenderTargetLoadAction::EClear : ERenderTargetLoadAction::ELoad, ERenderTargetStoreAction::EStore);
+
+	GraphBuilder.AddPass(
+		RDG_EVENT_NAME("CompositePlanarReflections"),
+		PassParameters,
+		ERDGPassFlags::None,
+		[PassParameters, &View, this](FRHICommandList& RHICmdList)
+	{
+		RHICmdList.SetViewport(View.ViewRect.Min.X, View.ViewRect.Min.Y, 0.0f, View.ViewRect.Max.X, View.ViewRect.Max.Y, 1.0f);
+
+		FGraphicsPipelineStateInitializer GraphicsPSOInit;
+		RHICmdList.ApplyCachedRenderTargets(GraphicsPSOInit);
+
+		// Blend over previous reflections in the output target (SSR or planar reflections that have already been rendered)
+		// Planar reflections win over SSR and reflection environment
+		//@todo - this is order dependent blending, but ordering is coming from registration order
+		GraphicsPSOInit.BlendState = TStaticBlendState<CW_RGBA, BO_Add, BF_One, BF_InverseSourceAlpha, BO_Max, BF_One, BF_One>::GetRHI();
+		GraphicsPSOInit.RasterizerState = TStaticRasterizerState<FM_Solid, CM_None>::GetRHI();
+		GraphicsPSOInit.DepthStencilState = TStaticDepthStencilState<false, CF_Always>::GetRHI();
 
-		FRHIRenderPassInfo RPInfo(Output->GetRenderTargetItem().TargetableTexture, ERenderTargetActions::Load_Store);
-		RHICmdList.BeginRenderPass(RPInfo, TEXT("DeferredPlanarReflections"));
+		for (FPlanarReflectionSceneProxy* ReflectionSceneProxy : Scene->PlanarReflections)
 		{
-			if (!bSSRAsInput)
-			{
-				DrawClearQuad(RHICmdList, FLinearColor(0, 0, 0, 0));
-			}
+			if (!View.ViewFrustum.IntersectBox(ReflectionSceneProxy->WorldBounds.GetCenter(), ReflectionSceneProxy->WorldBounds.GetExtent()))
+				continue;
 
-			{
-				RHICmdList.SetViewport(View.ViewRect.Min.X, View.ViewRect.Min.Y, 0.0f, View.ViewRect.Max.X, View.ViewRect.Max.Y, 1.0f);
+			SCOPED_DRAW_EVENTF(RHICmdList, PlanarReflection, *ReflectionSceneProxy->OwnerName.ToString());
 
-				FGraphicsPipelineStateInitializer GraphicsPSOInit;
-				RHICmdList.ApplyCachedRenderTargets(GraphicsPSOInit);
+			TShaderMapRef<TDeferredLightVS<false> > VertexShader(View.ShaderMap);
+			TShaderMapRef<FPlanarReflectionPS> PixelShader(View.ShaderMap);
 
-				// Blend over previous reflections in the output target (SSR or planar reflections that have already been rendered)
-				// Planar reflections win over SSR and reflection environment
-				//@todo - this is order dependent blending, but ordering is coming from registration order
-				GraphicsPSOInit.BlendState = TStaticBlendState<CW_RGBA, BO_Add, BF_One, BF_InverseSourceAlpha, BO_Max, BF_One, BF_One>::GetRHI();
-				GraphicsPSOInit.RasterizerState = TStaticRasterizerState<FM_Solid, CM_None>::GetRHI();
-				GraphicsPSOInit.DepthStencilState = TStaticDepthStencilState<false, CF_Always>::GetRHI();
+			GraphicsPSOInit.BoundShaderState.VertexDeclarationRHI = GFilterVertexDeclaration.VertexDeclarationRHI;
+			GraphicsPSOInit.BoundShaderState.VertexShaderRHI = GETSAFERHISHADER_VERTEX(*VertexShader);
+			GraphicsPSOInit.BoundShaderState.PixelShaderRHI = GETSAFERHISHADER_PIXEL(*PixelShader);
+			GraphicsPSOInit.PrimitiveType = PT_TriangleList;
 
-				for (int32 PlanarReflectionIndex = 0; PlanarReflectionIndex < Scene->PlanarReflections.Num(); PlanarReflectionIndex++)
-				{
-					FPlanarReflectionSceneProxy* ReflectionSceneProxy = Scene->PlanarReflections[PlanarReflectionIndex];
+			SetGraphicsPipelineState(RHICmdList, GraphicsPSOInit);
 
-					if (View.ViewFrustum.IntersectBox(ReflectionSceneProxy->WorldBounds.GetCenter(), ReflectionSceneProxy->WorldBounds.GetExtent()))
-					{
-						SCOPED_DRAW_EVENTF(RHICmdList, PlanarReflection, *ReflectionSceneProxy->OwnerName.ToString());
-
-						TShaderMapRef<TDeferredLightVS<false> > VertexShader(View.ShaderMap);
-						TShaderMapRef<FPlanarReflectionPS> PixelShader(View.ShaderMap);
-
-						GraphicsPSOInit.BoundShaderState.VertexDeclarationRHI = GFilterVertexDeclaration.VertexDeclarationRHI;
-						GraphicsPSOInit.BoundShaderState.VertexShaderRHI = GETSAFERHISHADER_VERTEX(*VertexShader);
-						GraphicsPSOInit.BoundShaderState.PixelShaderRHI = GETSAFERHISHADER_PIXEL(*PixelShader);
-						GraphicsPSOInit.PrimitiveType = PT_TriangleList;
-
-						SetGraphicsPipelineState(RHICmdList, GraphicsPSOInit);
-
-						PixelShader->SetParameters(RHICmdList, View, ReflectionSceneProxy);
-						VertexShader->SetSimpleLightParameters(RHICmdList, View, FSphere(0));
-
-						DrawRectangle(
-							RHICmdList,
-							0, 0,
-							View.ViewRect.Width(), View.ViewRect.Height(),
-							View.ViewRect.Min.X, View.ViewRect.Min.Y,
-							View.ViewRect.Width(), View.ViewRect.Height(),
-							View.ViewRect.Size(),
-							FSceneRenderTargets::Get(RHICmdList).GetBufferSizeXY(),
-							*VertexShader,
-							EDRF_UseTriangleOptimization);
-					}
-				}
-			}
-		}
-		RHICmdList.EndRenderPass();
-		RHICmdList.CopyToResolveTarget(Output->GetRenderTargetItem().TargetableTexture, Output->GetRenderTargetItem().ShaderResourceTexture, FResolveParams());
+			VertexShader->SetSimpleLightParameters(RHICmdList, View, FSphere(0));
 
-		return true;
-	}
+			{
+				FPlanarReflectionUniformParameters PlanarReflectionUniformParameters;
+				SetupPlanarReflectionUniformParameters(View, ReflectionSceneProxy, PlanarReflectionUniformParameters);
+		
+				FPlanarReflectionPS::FParameters ShaderParameters = *PassParameters;
+				ShaderParameters.PlanarReflectionParameters = CreateUniformBufferImmediate(PlanarReflectionUniformParameters, UniformBuffer_SingleDraw);
+				SetShaderParameters(RHICmdList, *PixelShader, PixelShader->GetPixelShader(), ShaderParameters);
+			}
 
-	return false;
+			DrawRectangle(
+				RHICmdList,
+				0, 0,
+				View.ViewRect.Width(), View.ViewRect.Height(),
+				View.ViewRect.Min.X, View.ViewRect.Min.Y,
+				View.ViewRect.Width(), View.ViewRect.Height(),
+				View.ViewRect.Size(),
+				FSceneRenderTargets::Get(RHICmdList).GetBufferSizeXY(),
+				*VertexShader,
+				EDRF_UseTriangleOptimization);
+		}
+	});
 }
diff --git a/Engine/Source/Runtime/Renderer/Private/ReflectionEnvironment.cpp b/Engine/Source/Runtime/Renderer/Private/ReflectionEnvironment.cpp
--- a/Engine/Source/Runtime/Renderer/Private/ReflectionEnvironment.cpp
+++ b/Engine/Source/Runtime/Renderer/Private/ReflectionEnvironment.cpp
@@ -33,6 +33,8 @@
 #include "ScreenSpaceDenoise.h"
 #include "ScreenSpaceRayTracing.h"
 #include "RayTracing/RaytracingOptions.h"
+#include "RenderGraph.h"
+#include "PixelShaderUtils.h"
 
 DECLARE_GPU_STAT_NAMED(ReflectionEnvironment, TEXT("Reflection Environment"));
 DECLARE_GPU_STAT_NAMED(RayTracingReflections, TEXT("Ray Tracing Reflections"));
@@ -458,78 +460,13 @@ void FReflectionEnvironmentCubemapArray::UpdateMaxCubemaps(uint32 InMaxCubemaps,
 	}
 }
 
-class FSkyLightParameters
-{
-public:
-	void Bind(const FShaderParameterMap& ParameterMap)
-	{
-		ApplyBentNormalAO.Bind(ParameterMap, TEXT("ApplyBentNormalAO"));
-		InvSkySpecularOcclusionStrength.Bind(ParameterMap, TEXT("InvSkySpecularOcclusionStrength"));
-		OcclusionTintAndMinOcclusion.Bind(ParameterMap, TEXT("OcclusionTintAndMinOcclusion"));
-		ContrastAndNormalizeMulAdd.Bind(ParameterMap, TEXT("ContrastAndNormalizeMulAdd"));
-		OcclusionExponent.Bind(ParameterMap, TEXT("OcclusionExponent"));
-		OcclusionCombineMode.Bind(ParameterMap, TEXT("OcclusionCombineMode"));
-	}
-
-	template<typename ShaderRHIParamRef, typename TRHICmdList>
-	void SetParameters(TRHICmdList& RHICmdList, const ShaderRHIParamRef& ShaderRHI, bool bApplyBentNormalAO, float SkySpecularOcclusionStrength, const FSkyLightSceneProxy* SkyLight)
-	{
-		float SkyLightContrast = 0.01f;
-		float SkyLightOcclusionExponent = 1.0f;
-		FVector4 SkyLightOcclusionTintAndMinOcclusion(0.0f, 0.0f, 0.0f, 0.0f);
-		EOcclusionCombineMode SkyLightOcclusionCombineMode = EOcclusionCombineMode::OCM_MAX;
-		if (SkyLight)
-		{
-			FDistanceFieldAOParameters Parameters(SkyLight->OcclusionMaxDistance, SkyLight->Contrast);
-			SkyLightContrast = Parameters.Contrast;
-			SkyLightOcclusionExponent = SkyLight->OcclusionExponent;
-			SkyLightOcclusionTintAndMinOcclusion = FVector4(SkyLight->OcclusionTint);
-			SkyLightOcclusionTintAndMinOcclusion.W = SkyLight->MinOcclusion;
-			SkyLightOcclusionCombineMode = SkyLight->OcclusionCombineMode;
-		}
-
-		// Scale and bias to remap the contrast curve to [0,1]
-		const float Min = 1 / (1 + FMath::Exp(-SkyLightContrast * (0 * 10 - 5)));
-		const float Max = 1 / (1 + FMath::Exp(-SkyLightContrast * (1 * 10 - 5)));
-		const float Mul = 1.0f / (Max - Min);
-		const float Add = -Min / (Max - Min);
-
-		SetShaderValue(RHICmdList, ShaderRHI, ContrastAndNormalizeMulAdd, FVector(SkyLightContrast, Mul, Add));
-
-		SetShaderValue(RHICmdList, ShaderRHI, OcclusionExponent, SkyLightOcclusionExponent);
-		SetShaderValue(RHICmdList, ShaderRHI, OcclusionTintAndMinOcclusion, SkyLightOcclusionTintAndMinOcclusion);
-		SetShaderValue(RHICmdList, ShaderRHI, OcclusionCombineMode, SkyLightOcclusionCombineMode == OCM_Minimum ? 0.0f : 1.0f);
-		SetShaderValue(RHICmdList, ShaderRHI, ApplyBentNormalAO, bApplyBentNormalAO ? 1.0f : 0.0f);
-		SetShaderValue(RHICmdList, ShaderRHI, InvSkySpecularOcclusionStrength, 1.0f / FMath::Max(SkySpecularOcclusionStrength, .1f));
-	}
-
-	friend FArchive& operator<<(FArchive& Ar, FSkyLightParameters& P)
-	{
-		Ar << P.ApplyBentNormalAO;
-		Ar << P.InvSkySpecularOcclusionStrength;
-		Ar << P.OcclusionTintAndMinOcclusion;
-		Ar << P.ContrastAndNormalizeMulAdd;
-		Ar << P.OcclusionExponent;
-		Ar << P.OcclusionCombineMode;
-
-		return Ar;
-	}
-
-private:
-	FShaderParameter ApplyBentNormalAO;
-	FShaderParameter InvSkySpecularOcclusionStrength;
-	FShaderParameter OcclusionTintAndMinOcclusion;
-	FShaderParameter ContrastAndNormalizeMulAdd;
-	FShaderParameter OcclusionExponent;
-	FShaderParameter OcclusionCombineMode;
-};
-
 IMPLEMENT_GLOBAL_SHADER_PARAMETER_STRUCT(FReflectionCaptureShaderData, "ReflectionCapture");
 
 /** Pixel shader that does tiled deferred culling of reflection captures, then sorts and composites them. */
 class FReflectionEnvironmentSkyLightingPS : public FGlobalShader
 {
 	DECLARE_GLOBAL_SHADER(FReflectionEnvironmentSkyLightingPS);
+	SHADER_USE_PARAMETER_STRUCT(FReflectionEnvironmentSkyLightingPS, FGlobalShader)
 
 	class FHasBoxCaptures			: SHADER_PERMUTATION_BOOL("REFLECTION_COMPOSITE_HAS_BOX_CAPTURES");
 	class FHasSphereCaptures		: SHADER_PERMUTATION_BOOL("REFLECTION_COMPOSITE_HAS_SPHERE_CAPTURES");
@@ -599,105 +536,42 @@ class FReflectionEnvironmentSkyLightingPS : public FGlobalShader
 		FForwardLightingParameters::ModifyCompilationEnvironment(Parameters.Platform, OutEnvironment);
 	}
 
+	BEGIN_SHADER_PARAMETER_STRUCT(FParameters, )
 
-public:
-
-	FReflectionEnvironmentSkyLightingPS(const ShaderMetaType::CompiledShaderInitializerType& Initializer)
-		: FGlobalShader(Initializer)
-	{
-		SceneTextureParameters.Bind(Initializer);
-		ReflectionCubemap.Bind(Initializer.ParameterMap,TEXT("ReflectionCubemap"));
-		ReflectionCubemapSampler.Bind(Initializer.ParameterMap,TEXT("ReflectionCubemapSampler"));
-		ScreenSpaceReflectionsTexture.Bind(Initializer.ParameterMap, TEXT("ScreenSpaceReflectionsTexture"));
-		ScreenSpaceReflectionsSampler.Bind(Initializer.ParameterMap, TEXT("ScreenSpaceReflectionsSampler"));
-		PreIntegratedGF.Bind(Initializer.ParameterMap, TEXT("PreIntegratedGF"));
-		PreIntegratedGFSampler.Bind(Initializer.ParameterMap, TEXT("PreIntegratedGFSampler"));
-		DFAOUpsampleParameters.Bind(Initializer.ParameterMap);
-		SkyLightParameters.Bind(Initializer.ParameterMap);
-	}
-
-	FReflectionEnvironmentSkyLightingPS()
-	{
-	}
-
-	template<typename TRHICommandList>
-	void SetParameters(
-		TRHICommandList& RHICmdList,
-		const FViewInfo& View,
-		FTextureRHIParamRef SSRTexture,
-		const TRefCountPtr<IPooledRenderTarget>& DynamicBentNormalAO
-	)
-	{
-		const FPixelShaderRHIParamRef ShaderRHI = GetPixelShader();
-
-		FGlobalShader::SetParameters<FViewUniformShaderParameters>(RHICmdList, ShaderRHI, View.ViewUniformBuffer);
-		SceneTextureParameters.Set(RHICmdList, ShaderRHI, View.FeatureLevel, ESceneTextureSetupMode::All);
-
-		FScene* Scene = (FScene*)View.Family->Scene;
-
-		FTextureRHIParamRef CubemapArray;
-		if (Scene->ReflectionSceneData.CubemapArray.IsValid() &&
-			Scene->ReflectionSceneData.CubemapArray.GetRenderTarget().IsValid())
-		{
-			CubemapArray = Scene->ReflectionSceneData.CubemapArray.GetRenderTarget().ShaderResourceTexture;
-		}
-		else
-		{
-			CubemapArray = GBlackCubeArrayTexture->TextureRHI;
-		}
-
-		SetTextureParameter(
-			RHICmdList,
-			ShaderRHI,
-			ReflectionCubemap,
-			ReflectionCubemapSampler,
-			TStaticSamplerState<SF_Trilinear, AM_Clamp, AM_Clamp, AM_Clamp>::GetRHI(),
-			CubemapArray);
-
-		SetTextureParameter(RHICmdList, ShaderRHI, ScreenSpaceReflectionsTexture, ScreenSpaceReflectionsSampler, TStaticSamplerState<SF_Point>::GetRHI(), SSRTexture);
-
-		SetUniformBufferParameter(RHICmdList, ShaderRHI, GetUniformBufferParameter<FReflectionCaptureShaderData>(), View.ReflectionCaptureUniformBuffer);
-
-		SetTextureParameter(RHICmdList, ShaderRHI, PreIntegratedGF, PreIntegratedGFSampler, TStaticSamplerState<SF_Bilinear, AM_Clamp, AM_Clamp, AM_Clamp>::GetRHI(), GSystemTextures.PreintegratedGF->GetRenderTargetItem().ShaderResourceTexture);
-	
-		FReflectionUniformParameters ReflectionUniformParameters;
-		SetupReflectionUniformParameters(View, ReflectionUniformParameters);
-		SetUniformBufferParameterImmediate(RHICmdList, ShaderRHI, GetUniformBufferParameter<FReflectionUniformParameters>(), ReflectionUniformParameters);
-
-		bool const bApplyBentNormalAO = DynamicBentNormalAO ? true : false;
-		SkyLightParameters.SetParameters(RHICmdList, ShaderRHI, bApplyBentNormalAO, CVarSkySpecularOcclusionStrength.GetValueOnRenderThread(), Scene->SkyLight);
-
-		SetUniformBufferParameter(RHICmdList, ShaderRHI, GetUniformBufferParameter<FForwardLightData>(), View.ForwardLightingResources->ForwardLightDataUniformBuffer);
-		DFAOUpsampleParameters.Set(RHICmdList, ShaderRHI, View, DynamicBentNormalAO);
-	}
-
-	virtual bool Serialize(FArchive& Ar) override
-	{
-		bool bShaderHasOutdatedParameters = FGlobalShader::Serialize(Ar);
-		Ar << SceneTextureParameters;
-		Ar << ReflectionCubemap;
-		Ar << ReflectionCubemapSampler;
-		Ar << ScreenSpaceReflectionsTexture;
-		Ar << ScreenSpaceReflectionsSampler;
-		Ar << PreIntegratedGF;
-		Ar << PreIntegratedGFSampler;
-		Ar << DFAOUpsampleParameters;
-		Ar << SkyLightParameters;
-		return bShaderHasOutdatedParameters;
-	}
+		// Sky light parameters.
+		SHADER_PARAMETER(FVector4, OcclusionTintAndMinOcclusion)
+		SHADER_PARAMETER(FVector, ContrastAndNormalizeMulAdd)
+		SHADER_PARAMETER(float, ApplyBentNormalAO)
+		SHADER_PARAMETER(float, InvSkySpecularOcclusionStrength)
+		SHADER_PARAMETER(float, OcclusionExponent)
+		SHADER_PARAMETER(float, OcclusionCombineMode)
+		
+		// Distance field AO parameters.
+		// TODO. FDFAOUpsampleParameters
+		SHADER_PARAMETER(FVector2D, AOBufferBilinearUVMax)
+		SHADER_PARAMETER_RDG_TEXTURE(Texture2D, BentNormalAOTexture)
+		SHADER_PARAMETER_SAMPLER(SamplerState,  BentNormalAOSampler)
+		
+		SHADER_PARAMETER_RDG_TEXTURE(Texture2D, AmbientOcclusionTexture)
+		SHADER_PARAMETER_SAMPLER(SamplerState,  AmbientOcclusionSampler)
+		
+		SHADER_PARAMETER_RDG_TEXTURE(Texture2D, ScreenSpaceReflectionsTexture)
+		SHADER_PARAMETER_SAMPLER(SamplerState,  ScreenSpaceReflectionsSampler)
+		
+		SHADER_PARAMETER_TEXTURE(Texture2D,    PreIntegratedGF)
+		SHADER_PARAMETER_SAMPLER(SamplerState, PreIntegratedGFSampler)
+		
+		SHADER_PARAMETER_STRUCT_INCLUDE(FSceneViewFamilyBlackboard, SceneTextures)
+		SHADER_PARAMETER_STRUCT_INCLUDE(FSceneTextureSamplerParameters, SceneTextureSamplers)
 
-private:
+		SHADER_PARAMETER_STRUCT_REF(FViewUniformShaderParameters, ViewUniformBuffer)
+		SHADER_PARAMETER_STRUCT_REF(FReflectionUniformParameters, ReflectionsParameters)
+		SHADER_PARAMETER_STRUCT_REF(FReflectionCaptureShaderData, ReflectionCaptureData)
+		SHADER_PARAMETER_STRUCT_REF(FForwardLightData, ForwardLightData)
 
-	FSceneTextureShaderParameters SceneTextureParameters;
-	FShaderResourceParameter ReflectionCubemap;
-	FShaderResourceParameter ReflectionCubemapSampler;
-	FShaderResourceParameter ScreenSpaceReflectionsTexture;
-	FShaderResourceParameter ScreenSpaceReflectionsSampler;
-	FShaderResourceParameter PreIntegratedGF;
-	FShaderResourceParameter PreIntegratedGFSampler;
-	FDFAOUpsampleParameters DFAOUpsampleParameters;
-	FSkyLightParameters SkyLightParameters;
-};
+		RENDER_TARGET_BINDING_SLOTS()
+	END_SHADER_PARAMETER_STRUCT()
+}; // FReflectionEnvironmentSkyLightingPS
 
 IMPLEMENT_GLOBAL_SHADER(FReflectionEnvironmentSkyLightingPS, "/Engine/Private/ReflectionEnvironmentPixelShader.usf", "ReflectionEnvironmentSkyLighting", SF_Pixel);
 
@@ -754,6 +628,7 @@ void FDeferredShadingSceneRenderer::RenderDeferredReflectionsAndSkyLighting(FRHI
 			&& ShouldRenderDistanceFieldAO()
 			&& ViewFamily.EngineShowFlags.AmbientOcclusion)
 		{
+			// TODO: convert to RDG.
 			bApplySkyShadowing = RenderDistanceFieldLighting(RHICmdList, Parameters, VelocityRT, DynamicBentNormalAO, false, false);
 		}
 	}
@@ -764,22 +639,24 @@ void FDeferredShadingSceneRenderer::RenderDeferredReflectionsAndSkyLighting(FRHI
 
 	const bool bReflectionEnv = ShouldDoReflectionEnvironment();
 
-	for (int32 ViewIndex = 0, Num = Views.Num(); ViewIndex < Num; ViewIndex++)
-	{
-		FViewInfo& View = Views[ViewIndex];
+	FRDGBuilder GraphBuilder(RHICmdList);
 
+	FRDGTextureRef SceneColorTexture = GraphBuilder.RegisterExternalTexture(SceneContext.GetSceneColor());
+	FRDGTextureRef AmbientOcclusionTexture = GraphBuilder.RegisterExternalTexture(SceneContext.ScreenSpaceAO);
+	FRDGTextureRef DynamicBentNormalAOTexture = GraphBuilder.RegisterExternalTexture(DynamicBentNormalAO ? DynamicBentNormalAO : GSystemTextures.WhiteDummy);
+
+	FSceneViewFamilyBlackboard SceneTextures;
+	SetupSceneViewFamilyBlackboard(GraphBuilder, &SceneTextures);
+
+	for (FViewInfo& View : Views)
+	{
 		const bool bRayTracedReflections = ShouldRenderRayTracingReflections(View);
 
-		const bool bScreenSpaceReflections = !bRayTracedReflections && ShouldRenderScreenSpaceReflections(Views[ViewIndex]);
+		const bool bScreenSpaceReflections = !bRayTracedReflections && ShouldRenderScreenSpaceReflections(View);
 
-		TRefCountPtr<IPooledRenderTarget> ReflectionsColor = GSystemTextures.BlackDummy;
+		FRDGTextureRef ReflectionsColor = nullptr;
 		if (bRayTracedReflections || bScreenSpaceReflections)
 		{
-			FRDGBuilder GraphBuilder(RHICmdList); // TODO: convert the entire reflections to render graph.
-
-			FSceneViewFamilyBlackboard SceneBlackboard;
-			SetupSceneViewFamilyBlackboard(GraphBuilder, &SceneBlackboard);
-
 			int32 DenoiserMode = CVarUseReflectionDenoiser.GetValueOnRenderThread();
 			
 			bool bDenoise = false;
@@ -822,7 +699,7 @@ void FDeferredShadingSceneRenderer::RenderDeferredReflectionsAndSkyLighting(FRHI
 				RDG_EVENT_SCOPE(GraphBuilder, "ScreenSpaceReflections(Quality=%d)", int32(SSRQuality));
 			
 				RenderScreenSpaceReflections(
-					GraphBuilder, SceneBlackboard, CurrentSceneColor, View, SSRQuality, bDenoise, &DenoiserInputs);
+					GraphBuilder, SceneTextures, CurrentSceneColor, View, SSRQuality, bDenoise, &DenoiserInputs);
 			}
 			else
 			{
@@ -844,11 +721,11 @@ void FDeferredShadingSceneRenderer::RenderDeferredReflectionsAndSkyLighting(FRHI
 					GraphBuilder,
 					View,
 					&View.PrevViewInfo,
-					SceneBlackboard,
+					SceneTextures,
 					DenoiserInputs,
 					RayTracingConfig);
 
-				GraphBuilder.QueueTextureExtraction(DenoiserOutputs.Color, &ReflectionsColor);
+				ReflectionsColor = DenoiserOutputs.Color;
 			}
 			else if (bTemporalFilter)
 			{
@@ -859,11 +736,11 @@ void FDeferredShadingSceneRenderer::RenderDeferredReflectionsAndSkyLighting(FRHI
 				
 				FTAAOutputs TAAOutputs = TAASettings.AddTemporalAAPass(
 					GraphBuilder,
-					SceneBlackboard, View,
+					SceneTextures, View,
 					View.PrevViewInfo.SSRHistory,
 					&View.ViewState->PrevFrameViewInfo.SSRHistory);
 				
-				GraphBuilder.QueueTextureExtraction(TAAOutputs.SceneColor, &ReflectionsColor);
+				ReflectionsColor = TAAOutputs.SceneColor;
 			}
 			else
 			{
@@ -875,90 +752,144 @@ void FDeferredShadingSceneRenderer::RenderDeferredReflectionsAndSkyLighting(FRHI
 					// shader has been judged required to be supported.
 					GraphBuilder.RemoveUnusedTextureWarning(DenoiserInputs.RayHitDistance);
 				}
-
-				GraphBuilder.QueueTextureExtraction(DenoiserInputs.Color, &ReflectionsColor);
+				
+				ReflectionsColor = DenoiserInputs.Color;
 			}
-
-			GraphBuilder.Execute();
 		} // if (bRayTracedReflections || bScreenSpaceReflections)
 
-		bool bPlanarReflections = false;
 		if (!bRayTracedReflections)
 		{
-			bPlanarReflections = RenderDeferredPlanarReflections(RHICmdList, View, false, ReflectionsColor);
+			RenderDeferredPlanarReflections(GraphBuilder, SceneTextures, View, /* inout */ ReflectionsColor);
 		}
 
-		bool bRequiresApply = bSkyLight || bDynamicSkyLight || bReflectionEnv || bScreenSpaceReflections || bPlanarReflections || bRayTracedReflections;
+		bool bRequiresApply = ReflectionsColor != nullptr || bSkyLight || bDynamicSkyLight || bReflectionEnv;
 
 		if (bRequiresApply)
 		{
-			SCOPED_GPU_STAT(RHICmdList, ReflectionEnvironment);
-			SCOPED_DRAW_EVENTF(RHICmdList, ReflectionEnvironment, TEXT("ReflectionEnvironmentAndSky"));
+			RDG_GPU_STAT_SCOPE(GraphBuilder, ReflectionEnvironment);
 
 			// Render the reflection environment with tiled deferred culling
 			bool bHasBoxCaptures = (View.NumBoxReflectionCaptures > 0);
 			bool bHasSphereCaptures = (View.NumSphereReflectionCaptures > 0);
+			
+			FReflectionEnvironmentSkyLightingPS::FParameters* PassParameters = GraphBuilder.AllocParameters<FReflectionEnvironmentSkyLightingPS::FParameters>();
 
-			TShaderMapRef<FPostProcessVS> VertexShader(View.ShaderMap);
-
-			auto PermutationVector = FReflectionEnvironmentSkyLightingPS::BuildPermutationVector(View, bHasBoxCaptures, bHasSphereCaptures, DynamicBentNormalAO != NULL, bSkyLight, bDynamicSkyLight, bApplySkyShadowing, bRayTracedReflections);
-
-			TShaderMapRef<FReflectionEnvironmentSkyLightingPS> PixelShader(View.ShaderMap, PermutationVector);
+			// Setup the parameters of the shader.
+			{
+				// Setups all shader parameters related to skylight.
+				{
+					FSkyLightSceneProxy* SkyLight = Scene->SkyLight;
 
-			FGraphicsPipelineStateInitializer GraphicsPSOInit;
+					float SkyLightContrast = 0.01f;
+					float SkyLightOcclusionExponent = 1.0f;
+					FVector4 SkyLightOcclusionTintAndMinOcclusion(0.0f, 0.0f, 0.0f, 0.0f);
+					EOcclusionCombineMode SkyLightOcclusionCombineMode = EOcclusionCombineMode::OCM_MAX;
+					if (SkyLight)
+					{
+						FDistanceFieldAOParameters Parameters(SkyLight->OcclusionMaxDistance, SkyLight->Contrast);
+						SkyLightContrast = Parameters.Contrast;
+						SkyLightOcclusionExponent = SkyLight->OcclusionExponent;
+						SkyLightOcclusionTintAndMinOcclusion = FVector4(SkyLight->OcclusionTint);
+						SkyLightOcclusionTintAndMinOcclusion.W = SkyLight->MinOcclusion;
+						SkyLightOcclusionCombineMode = SkyLight->OcclusionCombineMode;
+					}
 
-			SceneContext.BeginRenderingSceneColor(
-				RHICmdList,
-				ESimpleRenderTargetMode::EExistingColorAndDepth,
-				FExclusiveDepthStencil::DepthRead_StencilWrite);
-			RHICmdList.ApplyCachedRenderTargets(GraphicsPSOInit);
+					// Scale and bias to remap the contrast curve to [0,1]
+					const float Min = 1 / (1 + FMath::Exp(-SkyLightContrast * (0 * 10 - 5)));
+					const float Max = 1 / (1 + FMath::Exp(-SkyLightContrast * (1 * 10 - 5)));
+					const float Mul = 1.0f / (Max - Min);
+					const float Add = -Min / (Max - Min);
+					
+					PassParameters->OcclusionTintAndMinOcclusion = SkyLightOcclusionTintAndMinOcclusion;
+					PassParameters->ContrastAndNormalizeMulAdd = FVector(SkyLightContrast, Mul, Add);
+					PassParameters->OcclusionExponent = SkyLightOcclusionExponent;
+					PassParameters->OcclusionCombineMode = SkyLightOcclusionCombineMode == OCM_Minimum ? 0.0f : 1.0f;
+					PassParameters->ApplyBentNormalAO = DynamicBentNormalAO ? 1.0f : 0.0f;
+					PassParameters->InvSkySpecularOcclusionStrength = 1.0f / FMath::Max(CVarSkySpecularOcclusionStrength.GetValueOnRenderThread(), 0.1f);
+				}
 
-			extern int32 GAOOverwriteSceneColor;
-			if (GetReflectionEnvironmentCVar() == 2 || GAOOverwriteSceneColor)
-			{
-				// override scene color for debugging
-				GraphicsPSOInit.BlendState = TStaticBlendState<>::GetRHI();
-			}
-			else
-			{
-				const bool bCheckerboardSubsurfaceRendering = IsSubsurfaceCheckerboardFormat(SceneContext.GetSceneColorFormat());
-				if (bCheckerboardSubsurfaceRendering)
+				// Setups all shader parameters related to distance field AO
 				{
-					GraphicsPSOInit.BlendState = TStaticBlendState<CW_RGB, BO_Add, BF_One, BF_One>::GetRHI();
+					FIntPoint AOBufferSize = GetBufferSizeForAO();
+					PassParameters->AOBufferBilinearUVMax = FVector2D(
+						(View.ViewRect.Width() / GAODownsampleFactor - 0.51f) / AOBufferSize.X, // 0.51 - so bilateral gather4 won't sample invalid texels
+						(View.ViewRect.Height() / GAODownsampleFactor - 0.51f) / AOBufferSize.Y);
+
+					PassParameters->BentNormalAOTexture = DynamicBentNormalAOTexture;
+					PassParameters->BentNormalAOSampler = TStaticSamplerState<SF_Bilinear>::GetRHI();
 				}
-				else
+
+				PassParameters->AmbientOcclusionTexture = AmbientOcclusionTexture;
+				PassParameters->AmbientOcclusionSampler = TStaticSamplerState<SF_Point>::GetRHI();
+
+				PassParameters->ScreenSpaceReflectionsTexture = ReflectionsColor ? ReflectionsColor : GraphBuilder.RegisterExternalTexture(GSystemTextures.BlackDummy);
+				PassParameters->ScreenSpaceReflectionsSampler = TStaticSamplerState<SF_Point>::GetRHI();
+
+				PassParameters->PreIntegratedGF = GSystemTextures.PreintegratedGF->GetRenderTargetItem().ShaderResourceTexture;
+				PassParameters->PreIntegratedGFSampler = TStaticSamplerState<SF_Bilinear, AM_Clamp, AM_Clamp, AM_Clamp>::GetRHI();
+				
+				PassParameters->SceneTextures = SceneTextures;
+				SetupSceneTextureSamplers(&PassParameters->SceneTextureSamplers);
+
+				PassParameters->ViewUniformBuffer = View.ViewUniformBuffer;
+				PassParameters->ReflectionCaptureData = View.ReflectionCaptureUniformBuffer;
 				{
-					GraphicsPSOInit.BlendState = TStaticBlendState<CW_RGBA, BO_Add, BF_One, BF_One, BO_Add, BF_One, BF_One>::GetRHI();
+					FReflectionUniformParameters ReflectionUniformParameters;
+					SetupReflectionUniformParameters(View, ReflectionUniformParameters);
+					PassParameters->ReflectionsParameters = CreateUniformBufferImmediate(ReflectionUniformParameters, UniformBuffer_SingleDraw);
 				}
+				PassParameters->ForwardLightData = View.ForwardLightingResources->ForwardLightDataUniformBuffer;
 			}
 
-			GraphicsPSOInit.RasterizerState = TStaticRasterizerState<FM_Solid, CM_None>::GetRHI();
-			GraphicsPSOInit.DepthStencilState = TStaticDepthStencilState<false, CF_Always>::GetRHI();
+			PassParameters->RenderTargets[0] = FRenderTargetBinding(SceneColorTexture, ERenderTargetLoadAction::ELoad, ERenderTargetStoreAction::EStore);
 
-			GraphicsPSOInit.BoundShaderState.VertexDeclarationRHI = GFilterVertexDeclaration.VertexDeclarationRHI;
-			GraphicsPSOInit.BoundShaderState.VertexShaderRHI = GETSAFERHISHADER_VERTEX(*VertexShader);
-			GraphicsPSOInit.BoundShaderState.PixelShaderRHI = GETSAFERHISHADER_PIXEL(*PixelShader);
-			GraphicsPSOInit.PrimitiveType = PT_TriangleList;
+			auto PermutationVector = FReflectionEnvironmentSkyLightingPS::BuildPermutationVector(
+				View, bHasBoxCaptures, bHasSphereCaptures, DynamicBentNormalAO != NULL, bSkyLight, bDynamicSkyLight, bApplySkyShadowing, bRayTracedReflections);
 
-			SetGraphicsPipelineState(RHICmdList, GraphicsPSOInit);
+			TShaderMapRef<FReflectionEnvironmentSkyLightingPS> PixelShader(View.ShaderMap, PermutationVector);
+			ClearUnusedGraphResources(*PixelShader, PassParameters);
 
-			RHICmdList.SetViewport(View.ViewRect.Min.X, View.ViewRect.Min.Y, 0.0f, View.ViewRect.Max.X, View.ViewRect.Max.Y, 1.0f);
+			GraphBuilder.AddPass(
+				RDG_EVENT_NAME("ReflectionEnvironmentAndSky %dx%d", View.ViewRect.Width(), View.ViewRect.Height()),
+				PassParameters,
+				ERDGPassFlags::None,
+				[PassParameters, &View, PixelShader](FRHICommandList& RHICmdList)
+			{
+				RHICmdList.SetViewport(View.ViewRect.Min.X, View.ViewRect.Min.Y, 0.0f, View.ViewRect.Max.X, View.ViewRect.Max.Y, 1.0f);
 
-			PixelShader->SetParameters(RHICmdList, View, ReflectionsColor->GetRenderTargetItem().ShaderResourceTexture, DynamicBentNormalAO);
+				FGraphicsPipelineStateInitializer GraphicsPSOInit;
+				FPixelShaderUtils::InitFullscreenPipelineState(RHICmdList, View.ShaderMap, *PixelShader, GraphicsPSOInit);
 
-			DrawRectangle(
-				RHICmdList,
-				0, 0,
-				View.ViewRect.Width(), View.ViewRect.Height(),
-				View.ViewRect.Min.X, View.ViewRect.Min.Y,
-				View.ViewRect.Width(), View.ViewRect.Height(),
-				FIntPoint(View.ViewRect.Width(), View.ViewRect.Height()),
-				SceneContext.GetBufferSizeXY(),
-				*VertexShader);
+				extern int32 GAOOverwriteSceneColor;
+				if (GetReflectionEnvironmentCVar() == 2 || GAOOverwriteSceneColor)
+				{
+					// override scene color for debugging
+					GraphicsPSOInit.BlendState = TStaticBlendState<>::GetRHI();
+				}
+				else
+				{
+					const bool bCheckerboardSubsurfaceRendering = IsSubsurfaceCheckerboardFormat(PassParameters->RenderTargets[0].GetTexture()->Desc.Format);
+					if (bCheckerboardSubsurfaceRendering)
+					{
+						GraphicsPSOInit.BlendState = TStaticBlendState<CW_RGB, BO_Add, BF_One, BF_One>::GetRHI();
+					}
+					else
+					{
+						GraphicsPSOInit.BlendState = TStaticBlendState<CW_RGBA, BO_Add, BF_One, BF_One, BO_Add, BF_One, BF_One>::GetRHI();
+					}
+				}
 
-			SceneContext.FinishRenderingSceneColor(RHICmdList);
+				SetGraphicsPipelineState(RHICmdList, GraphicsPSOInit);
+				SetShaderParameters(RHICmdList, *PixelShader, PixelShader->GetPixelShader(), *PassParameters);
+				FPixelShaderUtils::DrawFullscreenTriangle(RHICmdList);
+			});
+		} // if (bRequiresApply)
+	} // for (FViewInfo& View : Views)
+	
+	TRefCountPtr<IPooledRenderTarget> OutSceneColor;
+	GraphBuilder.QueueTextureExtraction(SceneColorTexture, &OutSceneColor);
 
-			ResolveSceneColor(RHICmdList);
-		}
-	}
+	GraphBuilder.Execute();
+	
+	ResolveSceneColor(RHICmdList);
 }
