Fixed LPV not working.

diff --git a/Engine/Source/Runtime/Renderer/Private/ShadowDepthRendering.cpp b/Engine/Source/Runtime/Renderer/Private/ShadowDepthRendering.cpp
--- a/Engine/Source/Runtime/Renderer/Private/ShadowDepthRendering.cpp
+++ b/Engine/Source/Runtime/Renderer/Private/ShadowDepthRendering.cpp
@@ -1480,6 +1480,8 @@ void FSceneRenderer::RenderShadowDepthMaps(FRHICommandListImmediate& RHICmdList)
 		GetLightNameForDrawEvent(ProjectedShadowInfo->GetLightSceneInfo().Proxy, LightNameWithLevel);
 		SCOPED_DRAW_EVENTF(RHICmdList, EventShadowDepths, TEXT("Cubemap %s %u^2"), *LightNameWithLevel, TargetSize.X, TargetSize.Y);
 
+		ProjectedShadowInfo->SetupShadowUniformBuffers(RHICmdList, Scene);
+
 		auto BeginShadowRenderPass = [this, &RenderTarget, &SceneContext](FRHICommandList& InRHICmdList, bool bPerformClear)
 		{
 			FTextureRHIParamRef DepthTarget = RenderTarget.TargetableTexture;
@@ -1551,6 +1553,8 @@ void FSceneRenderer::RenderShadowDepthMaps(FRHICommandListImmediate& RHICmdList)
 					GRHICommandList.UseParallelAlgorithms() && CVarParallelShadows.GetValueOnRenderThread() &&
 					(ProjectedShadowInfo->IsWholeSceneDirectionalShadow() || CVarParallelShadowsNonWholeScene.GetValueOnRenderThread());
 
+				ProjectedShadowInfo->SetupShadowUniformBuffers(RHICmdList, Scene);
+
 				auto BeginShadowRenderPass = [this, ProjectedShadowInfo](FRHICommandList& InRHICmdList, bool bPerformClear)
 				{
 					FTextureRHIParamRef PreShadowCacheDepthZ = Scene->PreShadowCacheDepthZ->GetRenderTargetItem().TargetableTexture.GetReference();
@@ -1659,6 +1663,8 @@ void FSceneRenderer::RenderShadowDepthMaps(FRHICommandListImmediate& RHICmdList)
 			FSceneViewState* ViewState = (FSceneViewState*)ProjectedShadowInfo->DependentView->State;
 			FLightPropagationVolume* LightPropagationVolume = ViewState->GetLightPropagationVolume(FeatureLevel);
 
+			ProjectedShadowInfo->SetupShadowUniformBuffers(RHICmdList, Scene);
+
 			auto BeginShadowRenderPass = [this, LightPropagationVolume, ProjectedShadowInfo, &ColorTarget0, &ColorTarget1, &DepthTarget](FRHICommandList& InRHICmdList, bool bPerformClear)
 			{
 				FTextureRHIParamRef RenderTargets[2];
