Fixed missing header include

diff --git a/Engine/Source/Runtime/Renderer/Private/RayTracing/RayTracingLighting.h b/Engine/Source/Runtime/Renderer/Private/RayTracing/RayTracingLighting.h
--- a/Engine/Source/Runtime/Renderer/Private/RayTracing/RayTracingLighting.h
+++ b/Engine/Source/Runtime/Renderer/Private/RayTracing/RayTracingLighting.h
@@ -8,6 +8,7 @@
 #include "SceneRendering.h"
 #include "LightSceneInfo.h"
 #include "RayTracingDefinitions.h"
+#include "Containers/DynamicRHIResourceArray.h"
 
 #if RHI_RAYTRACING
 
