Fix cases where primitives didn't correctly select which runtime virtual texture to render to when there are multiple VTs in the scene

diff --git a/Engine/Source/Runtime/Engine/Private/VT/RuntimeVirtualTexture.cpp b/Engine/Source/Runtime/Engine/Private/VT/RuntimeVirtualTexture.cpp
--- a/Engine/Source/Runtime/Engine/Private/VT/RuntimeVirtualTexture.cpp
+++ b/Engine/Source/Runtime/Engine/Private/VT/RuntimeVirtualTexture.cpp
@@ -74,14 +74,14 @@ public:
 	/** Getter for the virtual texture producer. */
 	FVirtualTextureProducerHandle GetProducerHandle() const
 	{
-		check(IsInRenderingThread());
+		checkSlow(IsInRenderingThread());
 		return ProducerHandle;
 	}
 
 	/** Getter for the virtual texture allocation. */
 	IAllocatedVirtualTexture* GetAllocatedVirtualTexture() const 
 	{
-		check(IsInRenderingThread());
+		checkSlow(IsInRenderingThread());
 		return AllocatedVirtualTexture;
 	}
 
diff --git a/Engine/Source/Runtime/Renderer/Private/PrimitiveSceneInfo.cpp b/Engine/Source/Runtime/Renderer/Private/PrimitiveSceneInfo.cpp
--- a/Engine/Source/Runtime/Renderer/Private/PrimitiveSceneInfo.cpp
+++ b/Engine/Source/Runtime/Renderer/Private/PrimitiveSceneInfo.cpp
@@ -94,6 +94,7 @@ FPrimitiveFlagsCompact::FPrimitiveFlagsCompact(const FPrimitiveSceneProxy* Proxy
 	, bStaticLighting(Proxy->HasStaticLighting())
 	, bCastStaticShadow(Proxy->CastsStaticShadow())
 	, bRenderToVirtualTexture(Proxy->WritesVirtualTexture())
+	, RuntimeVirtualTextureMask(0)
 {}
 
 FPrimitiveSceneInfoCompact::FPrimitiveSceneInfoCompact(FPrimitiveSceneInfo* InPrimitiveSceneInfo) :
@@ -574,6 +575,10 @@ void FPrimitiveSceneInfo::AddToScene(FRHICommandListImmediate& RHICmdList, bool
 	PrimitiveBounds.MaxCullDistance = PrimitiveBounds.MaxDrawDistance;
 
 	Scene->PrimitiveFlagsCompact[PackedIndex] = FPrimitiveFlagsCompact(Proxy);
+	if (Scene->PrimitiveFlagsCompact[PackedIndex].bRenderToVirtualTexture)
+	{
+		Scene->PrimitiveFlagsCompact[PackedIndex].RuntimeVirtualTextureMask = Scene->GetRuntimeVirtualTextureMask(Proxy);
+	}
 
 	// Store precomputed visibility ID.
 	int32 VisibilityBitIndex = Proxy->GetVisibilityId();
diff --git a/Engine/Source/Runtime/Renderer/Private/RendererScene.cpp b/Engine/Source/Runtime/Renderer/Private/RendererScene.cpp
--- a/Engine/Source/Runtime/Renderer/Private/RendererScene.cpp
+++ b/Engine/Source/Runtime/Renderer/Private/RendererScene.cpp
@@ -2405,14 +2405,134 @@ void FScene::RemovePrecomputedVolumetricLightmap(const FPrecomputedVolumetricLig
 
 void FScene::AddRuntimeVirtualTexture(class URuntimeVirtualTextureComponent* Component)
 {
-	delete Component->SceneProxy;
-	Component->SceneProxy = new FRuntimeVirtualTextureSceneProxy(Component);
+	if (Component->SceneProxy == nullptr)
+	{
+		Component->SceneProxy = new FRuntimeVirtualTextureSceneProxy(Component);
+
+		FScene* Scene = this;
+		FRuntimeVirtualTextureSceneProxy* SceneProxy = Component->SceneProxy;
+
+		ENQUEUE_RENDER_COMMAND(AddRuntimeVirtualTextureCommand)(
+			[Scene, SceneProxy](FRHICommandListImmediate& RHICmdList)
+		{
+			Scene->AddRuntimeVirtualTexture_RenderThread(SceneProxy);
+			Scene->UpdateRuntimeVirtualTextureForAllPrimitives_RenderThread();
+		});
+	}
+	else
+	{
+		// This is a component update.
+		// Store the new FRuntimeVirtualTextureSceneProxy at the same index as the old (to avoid needing to update any associated primitives).
+		// Defer old proxy deletion to the render thread.
+		FRuntimeVirtualTextureSceneProxy* SceneProxyToReplace = Component->SceneProxy;
+		Component->SceneProxy = new FRuntimeVirtualTextureSceneProxy(Component);
+
+		FScene* Scene = this;
+		FRuntimeVirtualTextureSceneProxy* SceneProxy = Component->SceneProxy;
+
+		ENQUEUE_RENDER_COMMAND(AddRuntimeVirtualTextureCommand)(
+			[Scene, SceneProxy, SceneProxyToReplace](FRHICommandListImmediate& RHICmdList)
+		{
+			const bool bUpdatePrimitives = SceneProxy->VirtualTexture != SceneProxyToReplace->VirtualTexture;
+			Scene->UpdateRuntimeVirtualTexture_RenderThread(SceneProxy, SceneProxyToReplace);
+			if (bUpdatePrimitives)
+			{
+				Scene->UpdateRuntimeVirtualTextureForAllPrimitives_RenderThread();
+			}
+		});
+	}
 }
 
 void FScene::RemoveRuntimeVirtualTexture(class URuntimeVirtualTextureComponent* Component)
 {
-	delete Component->SceneProxy;
-	Component->SceneProxy = nullptr;
+	FRuntimeVirtualTextureSceneProxy* SceneProxy = Component->SceneProxy;
+	if (SceneProxy != nullptr)
+	{
+		// Release now but defer any deletion to the render thread
+		Component->SceneProxy->Release();
+		Component->SceneProxy = nullptr;
+
+		FScene* Scene = this;
+		ENQUEUE_RENDER_COMMAND(RemoveRuntimeVirtualTextureCommand)(
+			[Scene, SceneProxy](FRHICommandListImmediate& RHICmdList)
+		{
+			Scene->RemoveRuntimeVirtualTexture_RenderThread(SceneProxy);
+			Scene->UpdateRuntimeVirtualTextureForAllPrimitives_RenderThread();
+		});
+	}
+}
+
+void FScene::AddRuntimeVirtualTexture_RenderThread(FRuntimeVirtualTextureSceneProxy* SceneProxy)
+{
+	SceneProxy->SceneIndex = RuntimeVirtualTextures.Add(SceneProxy);
+}
+
+void FScene::UpdateRuntimeVirtualTexture_RenderThread(FRuntimeVirtualTextureSceneProxy* SceneProxy, FRuntimeVirtualTextureSceneProxy* SceneProxyToReplace)
+{
+	for (int32 SceneIndex = 0; SceneIndex < RuntimeVirtualTextures.Num(); ++SceneIndex)
+	{
+		if (RuntimeVirtualTextures.IsValidIndex(SceneIndex) && RuntimeVirtualTextures[SceneIndex] == SceneProxyToReplace)
+		{
+			SceneProxy->SceneIndex = SceneIndex;
+			RuntimeVirtualTextures[SceneIndex] = SceneProxy;
+			delete SceneProxyToReplace;
+			return;
+		}
+	}
+	// If we get here then we didn't find the object to replace!
+	check(false);
+}
+
+void FScene::RemoveRuntimeVirtualTexture_RenderThread(FRuntimeVirtualTextureSceneProxy* SceneProxy)
+{
+	RuntimeVirtualTextures.RemoveAt(SceneProxy->SceneIndex);
+	delete SceneProxy;
+}
+
+void FScene::UpdateRuntimeVirtualTextureForAllPrimitives_RenderThread()
+{
+	for (int32 Index = 0; Index < Primitives.Num(); ++Index)
+	{
+		if (PrimitiveFlagsCompact[Index].bRenderToVirtualTexture)
+		{
+			PrimitiveFlagsCompact[Index].RuntimeVirtualTextureMask = GetRuntimeVirtualTextureMask(PrimitiveSceneProxies[Index]);
+		}
+	}
+}
+
+uint32 FScene::GetRuntimeVirtualTextureSceneIndex(uint32 ProducerId)
+{
+	checkSlow(IsInRenderingThread());
+	for (FRuntimeVirtualTextureSceneProxy const* Proxy : RuntimeVirtualTextures)
+	{
+		if (Proxy->ProducerId == ProducerId)
+		{
+			return Proxy->SceneIndex;
+		}
+	}
+	// Should not get here
+	check(false);
+	return 0;
+}
+
+uint32 FScene::GetRuntimeVirtualTextureMask(FPrimitiveSceneProxy const* Proxy)
+{
+	uint32 Mask = 0;
+	for (int32 i = 0; i < FPrimitiveFlagsCompact::RuntimeVirtualTexture_BitCount && i < RuntimeVirtualTextures.Num(); ++i)
+	{
+		if (RuntimeVirtualTextures.IsValidIndex(i))
+		{
+			URuntimeVirtualTexture* SceneVirtualTexture = RuntimeVirtualTextures[i]->VirtualTexture;
+			for (URuntimeVirtualTexture* PrimitiveVirtualTexture : Proxy->RuntimeVirtualTextures)
+			{
+				if (SceneVirtualTexture == PrimitiveVirtualTexture)
+				{
+					Mask |= 1 << i;
+				}
+			}
+		}
+	}
+	return Mask;
 }
 
 bool FScene::GetPreviousLocalToWorld(const FPrimitiveSceneInfo* PrimitiveSceneInfo, FMatrix& OutPreviousLocalToWorld) const
diff --git a/Engine/Source/Runtime/Renderer/Private/SceneCore.h b/Engine/Source/Runtime/Renderer/Private/SceneCore.h
--- a/Engine/Source/Runtime/Renderer/Private/SceneCore.h
+++ b/Engine/Source/Runtime/Renderer/Private/SceneCore.h
@@ -182,6 +182,7 @@ public:
 		, bUseForDepthPass(StaticMesh.bUseForDepthPass)
 		, bUseAsOccluder(StaticMesh.bUseAsOccluder)
 		, bRenderToVirtualTexture(StaticMesh.bRenderToVirtualTexture)
+		, RuntimeVirtualTextureMaterialType(StaticMesh.RuntimeVirtualTextureMaterialType)
 		, bSupportsCachingMeshDrawCommands(InbSupportsCachingMeshDrawCommands)
 	{
 	}
@@ -220,6 +221,8 @@ public:
 
 	/** Whether the mesh batch can be used for rendering to a virtual texture. */
 	uint8 bRenderToVirtualTexture : 1;
+	/** What virtual texture material type this mesh batch should be rendered with. */
+	uint32 RuntimeVirtualTextureMaterialType : ERuntimeVirtualTextureMaterialType_NumBits;
 
 	/** Cached from vertex factory to avoid dereferencing VF in InitViews. */
 	uint8 bSupportsCachingMeshDrawCommands : 1;
diff --git a/Engine/Source/Runtime/Renderer/Private/ScenePrivate.h b/Engine/Source/Runtime/Renderer/Private/ScenePrivate.h
--- a/Engine/Source/Runtime/Renderer/Private/ScenePrivate.h
+++ b/Engine/Source/Runtime/Renderer/Private/ScenePrivate.h
@@ -77,6 +77,7 @@ class UTextureCube;
 class UWindDirectionalSourceComponent;
 class FRHIGPUBufferReadback;
 class FRHIGPUTextureReadback;
+class FRuntimeVirtualTextureSceneProxy;
 
 /** Holds information about a single primitive's occlusion. */
 class FPrimitiveOcclusionHistory
@@ -2572,6 +2573,9 @@ public:
 	/** LOD Tree Holder for massive LOD system */
 	FLODSceneTree SceneLODHierarchy;
 
+	/** The runtime virtual textures in the scene. */
+	TSparseArray<FRuntimeVirtualTextureSceneProxy*> RuntimeVirtualTextures;
+
 	float DefaultMaxDistanceFieldOcclusionDistance;
 
 	float GlobalDistanceFieldViewDistance;
@@ -2798,6 +2802,11 @@ public:
 		return PrimitiveComponentIds;
 	}
 
+	/** Get the scene index of the FRuntimeVirtualTextureSceneProxy associated with the producer. */
+	uint32 GetRuntimeVirtualTextureSceneIndex(uint32 ProducerId);
+	/** Build the mask used for filtering this Proxy against all of the runtime virtual textures in the scene. */
+	uint32 GetRuntimeVirtualTextureMask(FPrimitiveSceneProxy const* Proxy);
+
 #if WITH_EDITOR
 	virtual bool InitializePixelInspector(FRenderTarget* BufferFinalColor, FRenderTarget* BufferSceneColor, FRenderTarget* BufferDepth, FRenderTarget* BufferHDR, FRenderTarget* BufferA, FRenderTarget* BufferBCDE, int32 BufferIndex) override;
 
@@ -2892,6 +2901,15 @@ private:
 	/** Updates all static draw lists. */
 	void UpdateStaticDrawLists_RenderThread(FRHICommandListImmediate& RHICmdList);
 
+	/** Add a runtime virtual texture proxy to the scene. Called in the rendering thread by AddRuntimeVirtualTexture. */
+	void AddRuntimeVirtualTexture_RenderThread(FRuntimeVirtualTextureSceneProxy* SceneProxy);
+	/** Update a runtime virtual texture proxy to the scene. Called in the rendering thread by AddRuntimeVirtualTexture. */
+	void UpdateRuntimeVirtualTexture_RenderThread(FRuntimeVirtualTextureSceneProxy* SceneProxy, FRuntimeVirtualTextureSceneProxy* SceneProxyToReplace);
+	/** Remove a runtime virtual texture proxy from the scene. Called in the rendering thread by RemoveRuntimeVirtualTexture. */
+	void RemoveRuntimeVirtualTexture_RenderThread(FRuntimeVirtualTextureSceneProxy* SceneProxy);
+	/** Update the scene primitive data after completing operations that add or remove runtime virtual textures. This can be slow and should be called rarely. */
+	void UpdateRuntimeVirtualTextureForAllPrimitives_RenderThread();
+
 	/**
 	 * Shifts scene data by provided delta
 	 * Called on world origin changes
diff --git a/Engine/Source/Runtime/Renderer/Private/VT/RuntimeVirtualTextureProducer.cpp b/Engine/Source/Runtime/Renderer/Private/VT/RuntimeVirtualTextureProducer.cpp
--- a/Engine/Source/Runtime/Renderer/Private/VT/RuntimeVirtualTextureProducer.cpp
+++ b/Engine/Source/Runtime/Renderer/Private/VT/RuntimeVirtualTextureProducer.cpp
@@ -6,8 +6,10 @@
 #include "RuntimeVirtualTextureRender.h"
 #include "ScenePrivate.h"
 
-FRuntimeVirtualTextureFinalizer::FRuntimeVirtualTextureFinalizer(FVTProducerDescription const& InDesc, ERuntimeVirtualTextureMaterialType InMaterialType, FSceneInterface* InScene, FTransform const& InUVToWorld)
+FRuntimeVirtualTextureFinalizer::FRuntimeVirtualTextureFinalizer(FVTProducerDescription const& InDesc, uint32 InProducerId, ERuntimeVirtualTextureMaterialType InMaterialType, FSceneInterface* InScene, FTransform const& InUVToWorld)
 	: Desc(InDesc)
+	, ProducerId(InProducerId)
+	, RuntimeVirtualTextureMask(0)
 	, MaterialType(InMaterialType)
 	, Scene(InScene)
 	, UVToWorld(InUVToWorld)
@@ -28,6 +30,20 @@ void FRuntimeVirtualTextureFinalizer::AddTile(FTileEntry& Tile)
 
 void FRuntimeVirtualTextureFinalizer::Finalize(FRHICommandListImmediate& RHICmdList)
 {
+	if (RuntimeVirtualTextureMask == 0)
+	{
+		// Initialize the RuntimeVirtualTextureMask by matching this producer with those registered in the scene's runtime virtual textures.
+		// We only need to do this once. If the associated scene proxy is removed this finalizer will also be destroyed.
+		const uint32 VirtualTextureSceneIndex = Scene->GetRenderScene()->GetRuntimeVirtualTextureSceneIndex(ProducerId);
+		RuntimeVirtualTextureMask = 1 << VirtualTextureSceneIndex;
+	
+		//todo[vt]: 
+		// Add a slow render path inside RenderPage() when this check fails. 
+		// It will need to iterate the virtual textures on each primitive instead of using the RuntimeVirtualTextureMask.
+		// Currently nothing will render for this finalizer when the check fails.
+		checkSlow(VirtualTextureSceneIndex < FPrimitiveFlagsCompact::RuntimeVirtualTexture_BitCount);
+	}
+
 	for (auto Entry : Tiles)
 	{
 		const int32 TileSize = Desc.TileSize + 2 * Desc.TileBorderSize;
@@ -49,8 +65,9 @@ void FRuntimeVirtualTextureFinalizer::Finalize(FRHICommandListImmediate& RHICmdL
 		const FBox2D UVRange(UV - UVBorder, UV + UVSize + UVBorder);
 
 		RuntimeVirtualTexture::RenderPage(
-			RHICmdList, 
+			RHICmdList,
 			Scene->GetRenderScene(),
+			RuntimeVirtualTextureMask,
 			MaterialType,
 			Entry.Texture0, 
 			DestinationBox0, 
@@ -63,8 +80,8 @@ void FRuntimeVirtualTextureFinalizer::Finalize(FRHICommandListImmediate& RHICmdL
 	Tiles.SetNumUnsafeInternal(0);
 }
 
-FRuntimeVirtualTextureProducer::FRuntimeVirtualTextureProducer(FVTProducerDescription const& InDesc, ERuntimeVirtualTextureMaterialType InMaterialType, FSceneInterface* InScene, FTransform const& InUVToWorld)
-	: Finalizer(InDesc, InMaterialType, InScene, InUVToWorld)
+FRuntimeVirtualTextureProducer::FRuntimeVirtualTextureProducer(FVTProducerDescription const& InDesc, uint32 InProducerId, ERuntimeVirtualTextureMaterialType InMaterialType, FSceneInterface* InScene, FTransform const& InUVToWorld)
+	: Finalizer(InDesc, InProducerId, InMaterialType, InScene, InUVToWorld)
 {
 }
 
diff --git a/Engine/Source/Runtime/Renderer/Private/VT/RuntimeVirtualTextureProducer.h b/Engine/Source/Runtime/Renderer/Private/VT/RuntimeVirtualTextureProducer.h
--- a/Engine/Source/Runtime/Renderer/Private/VT/RuntimeVirtualTextureProducer.h
+++ b/Engine/Source/Runtime/Renderer/Private/VT/RuntimeVirtualTextureProducer.h
@@ -14,7 +14,7 @@ class FSceneInterface;
 class FRuntimeVirtualTextureFinalizer : public IVirtualTextureFinalizer
 {
 public:
-	FRuntimeVirtualTextureFinalizer(FVTProducerDescription const& InDesc, ERuntimeVirtualTextureMaterialType InMaterialType, FSceneInterface* InScene, FTransform const& InUVToWorld);
+	FRuntimeVirtualTextureFinalizer(FVTProducerDescription const& InDesc, uint32 InProducerId, ERuntimeVirtualTextureMaterialType InMaterialType, FSceneInterface* InScene, FTransform const& InUVToWorld);
 	virtual ~FRuntimeVirtualTextureFinalizer() {}
 
 	/** A description for a single tile to render. */
@@ -43,6 +43,10 @@ public:
 private:
 	/** Description of our virtual texture. */
 	const FVTProducerDescription Desc;
+	/** Producer index used once to build the RuntimeVirtualTextureMask for any rendering. */
+	uint32 ProducerId;
+	/** The mask of runtime virtual textures in the scene that we should render to. */
+	uint32 RuntimeVirtualTextureMask;
 	/** Contents of virtual texture layer stack. */
 	ERuntimeVirtualTextureMaterialType MaterialType;
 	/** Scene that the virtual texture is placed within. */
@@ -57,7 +61,7 @@ private:
 class FRuntimeVirtualTextureProducer : public IVirtualTexture
 {
 public:
-	RENDERER_API FRuntimeVirtualTextureProducer(FVTProducerDescription const& InDesc, ERuntimeVirtualTextureMaterialType InMaterialType, FSceneInterface* InScene, FTransform const& InUVToWorld);
+	RENDERER_API FRuntimeVirtualTextureProducer(FVTProducerDescription const& InDesc, uint32 InProducerId, ERuntimeVirtualTextureMaterialType InMaterialType, FSceneInterface* InScene, FTransform const& InUVToWorld);
 	RENDERER_API virtual ~FRuntimeVirtualTextureProducer() {}
 
 	//~ Begin IVirtualTexture Interface.
diff --git a/Engine/Source/Runtime/Renderer/Private/VT/RuntimeVirtualTextureRender.cpp b/Engine/Source/Runtime/Renderer/Private/VT/RuntimeVirtualTextureRender.cpp
--- a/Engine/Source/Runtime/Renderer/Private/VT/RuntimeVirtualTextureRender.cpp
+++ b/Engine/Source/Runtime/Renderer/Private/VT/RuntimeVirtualTextureRender.cpp
@@ -270,7 +270,7 @@ namespace RuntimeVirtualTexture
 
 
 	/** Collect meshes and draw. */
-	void DrawMeshes(FRHICommandListImmediate& RHICmdList, FScene const* Scene, FViewInfo const* View, ERuntimeVirtualTextureMaterialType MaterialType)
+	void DrawMeshes(FRHICommandListImmediate& RHICmdList, FScene const* Scene, FViewInfo const* View, ERuntimeVirtualTextureMaterialType MaterialType, uint32 RuntimeVirtualTextureMask)
 	{
 		// Cached draw command collectors
 		const FCachedPassMeshDrawList& SceneDrawList = Scene->CachedDrawLists[EMeshPass::VirtualTexture];
@@ -287,24 +287,19 @@ namespace RuntimeVirtualTexture
 		//todo: Consider a broad phase (quad tree etc?) here. (But only if running over PrimitiveFlagsCompact shows up as a bottleneck.)
 		for (int32 PrimitiveIndex = 0; PrimitiveIndex < Scene->Primitives.Num(); ++PrimitiveIndex)
 		{
-			if (Scene->PrimitiveFlagsCompact[PrimitiveIndex].bRenderToVirtualTexture)
+			if (Scene->PrimitiveFlagsCompact[PrimitiveIndex].RuntimeVirtualTextureMask & RuntimeVirtualTextureMask)
 			{
 				//todo[vt]: Use quicker/more accurate 2d test here since we can pre-calculate 2d bounds in VT space.
 				FBoxSphereBounds const& Bounds = Scene->PrimitiveBounds[PrimitiveIndex].BoxSphereBounds;
 				if (View->ViewFrustum.IntersectSphere(Bounds.GetSphere().Center, Bounds.GetSphere().W))
 				{
 					FPrimitiveSceneInfo* PrimitiveSceneInfo = Scene->Primitives[PrimitiveIndex];
-
 					for (int32 MeshIndex = 0; MeshIndex < PrimitiveSceneInfo->StaticMeshes.Num(); ++MeshIndex)
 					{
 						FStaticMeshBatchRelevance const& StaticMeshRelevance = PrimitiveSceneInfo->StaticMeshRelevances[MeshIndex];
-						FStaticMeshBatch const& MeshBatch = PrimitiveSceneInfo->StaticMeshes[MeshIndex];
-
-						//todo[vt]: 
-						// Filter for currently rendered VT, not MaterialType (currently we would end up with multiple or unwanted draws)
-						// Also better if we can do that without having to read from MeshBatch data (to save mem cache)
-						if (StaticMeshRelevance.bRenderToVirtualTexture && MeshBatch.RuntimeVirtualTextureMaterialType == (uint32)MaterialType)
+						if (StaticMeshRelevance.bRenderToVirtualTexture && StaticMeshRelevance.RuntimeVirtualTextureMaterialType == (uint32)MaterialType)
 						{
+							bool bCachedDraw = false;
 							if (StaticMeshRelevance.bSupportsCachingMeshDrawCommands)
 							{
 								// Use cached draw command
@@ -328,13 +323,15 @@ namespace RuntimeVirtualTexture
 										CachedMeshDrawCommand.SortKey);
 
 									CachedDrawCommands.Add(NewVisibleMeshDrawCommand);
+									bCachedDraw = true;
 								}
 							}
-							else
+
+							if (!bCachedDraw)
 							{
-								// No cached draw command available. Render static mesh.
+								// No cached draw command was available. Process the mesh batch.
 								uint64 BatchElementMask = ~0ull;
-								MeshProcessor.AddMeshBatch(MeshBatch, BatchElementMask, Scene->PrimitiveSceneProxies[PrimitiveIndex]);
+								MeshProcessor.AddMeshBatch(PrimitiveSceneInfo->StaticMeshes[MeshIndex], BatchElementMask, Scene->PrimitiveSceneProxies[PrimitiveIndex]);
 							}
 						}
 					}
@@ -544,6 +541,7 @@ namespace RuntimeVirtualTexture
 	void RenderPage(
 		FRHICommandListImmediate& RHICmdList,
 		FScene* Scene,
+		uint32 RuntimeVirtualTextureMask,
 		ERuntimeVirtualTextureMaterialType MaterialType,
 		FRHITexture2D* OutputTexture0,
 		FBox2D const& DestBox0,
@@ -619,9 +617,9 @@ namespace RuntimeVirtualTexture
 				RDG_EVENT_NAME("VirtualTextureDraw"),
 				PassParameters,
 				ERenderGraphPassFlags::None,
-				[Scene, View, MaterialType](FRHICommandListImmediate& RHICmdListImmediate)
+				[Scene, View, MaterialType, RuntimeVirtualTextureMask](FRHICommandListImmediate& RHICmdListImmediate)
 			{
-				DrawMeshes(RHICmdListImmediate, Scene, View, MaterialType);
+				DrawMeshes(RHICmdListImmediate, Scene, View, MaterialType, RuntimeVirtualTextureMask);
 			});
 		}
 
@@ -675,7 +673,7 @@ namespace RuntimeVirtualTexture
 		GraphBuilder.Execute();
 
 		// Copy to final destination
-		if (GraphSetup.OutputAlias0 != nullptr)
+		if (GraphSetup.OutputAlias0 != nullptr && OutputTexture0 != nullptr)
 		{
 			FRHICopyTextureInfo Info;
 			Info.Size = GraphOutputSize0;
@@ -684,7 +682,7 @@ namespace RuntimeVirtualTexture
 			RHICmdList.CopyTexture(GraphOutputTexture0->GetRenderTargetItem().ShaderResourceTexture->GetTexture2D(), OutputTexture0->GetTexture2D(), Info);
 		}
 
-		if (GraphSetup.OutputAlias1 != nullptr)
+		if (GraphSetup.OutputAlias1 != nullptr && OutputTexture1 != nullptr)
 		{
 			FRHICopyTextureInfo Info;
 			Info.Size = GraphOutputSize1;
diff --git a/Engine/Source/Runtime/Renderer/Private/VT/RuntimeVirtualTextureRender.h b/Engine/Source/Runtime/Renderer/Private/VT/RuntimeVirtualTextureRender.h
--- a/Engine/Source/Runtime/Renderer/Private/VT/RuntimeVirtualTextureRender.h
+++ b/Engine/Source/Runtime/Renderer/Private/VT/RuntimeVirtualTextureRender.h
@@ -16,6 +16,7 @@ namespace RuntimeVirtualTexture
 	void RenderPage(
 		FRHICommandListImmediate& RHICmdList, 
 		FScene* Scene,
+		uint32 RuntimeVirtualTextureMask,
 		ERuntimeVirtualTextureMaterialType MaterialType,
 		FRHITexture2D* Texture0, FBox2D const& DestBox0,
 		FRHITexture2D* Texture1, FBox2D const& DestBox1, 
diff --git a/Engine/Source/Runtime/Renderer/Private/VT/RuntimeVirtualTextureSceneProxy.cpp b/Engine/Source/Runtime/Renderer/Private/VT/RuntimeVirtualTextureSceneProxy.cpp
--- a/Engine/Source/Runtime/Renderer/Private/VT/RuntimeVirtualTextureSceneProxy.cpp
+++ b/Engine/Source/Runtime/Renderer/Private/VT/RuntimeVirtualTextureSceneProxy.cpp
@@ -6,30 +6,43 @@
 #include "VT/RuntimeVirtualTexture.h"
 #include "VT/RuntimeVirtualTextureProducer.h"
 
+int32 FRuntimeVirtualTextureSceneProxy::ProducerIdGenerator = 1;
 
 FRuntimeVirtualTextureSceneProxy::FRuntimeVirtualTextureSceneProxy(URuntimeVirtualTextureComponent* InComponent)
-	: VirtualTexture(InComponent->GetVirtualTexture())
+	: ProducerId(0)
+	, VirtualTexture(InComponent->GetVirtualTexture())
 {
 	if (VirtualTexture != nullptr)
 	{
-		// The Producer object created here will be passed into the Virtual Texture system which will take ownership
 		FVTProducerDescription Desc;
 		VirtualTexture->GetProducerDescription(Desc);
 
+		// We store a ProducerId here so that we will be able to find our SceneIndex from the Producer during rendering.
+		// We will need the SceneIndex to determine which primitives should render to this Producer.
+		ProducerId = ProducerIdGenerator++;
+
 		const ERuntimeVirtualTextureMaterialType MaterialType = VirtualTexture->GetMaterialType();
 
 		// Transform is based on bottom left of the URuntimeVirtualTextureComponent unit box (which is centered on the origin)
 		FTransform Transform = FTransform(FVector(-0.5f, -0.5f, 0.f)) * InComponent->GetComponentTransform();
 
-		FRuntimeVirtualTextureProducer* Producer = new FRuntimeVirtualTextureProducer(Desc, MaterialType, InComponent->GetScene(), Transform);
+		// The Producer object created here will be passed into the Virtual Texture system which will take ownership.
+		// The Initialize() call will allocate the VT by spawning work on the render thread.
+		FRuntimeVirtualTextureProducer* Producer = new FRuntimeVirtualTextureProducer(Desc, ProducerId, MaterialType, InComponent->GetScene(), Transform);
 		VirtualTexture->Initialize(Producer, Transform);
 	}
 }
 
 FRuntimeVirtualTextureSceneProxy::~FRuntimeVirtualTextureSceneProxy()
+{
+	checkSlow(IsInRenderingThread());
+}
+
+void FRuntimeVirtualTextureSceneProxy::Release()
 {
 	if (VirtualTexture != nullptr)
 	{
 		VirtualTexture->Release();
+		VirtualTexture = nullptr;
 	}
 }
diff --git a/Engine/Source/Runtime/Renderer/Private/VT/RuntimeVirtualTextureSceneProxy.h b/Engine/Source/Runtime/Renderer/Private/VT/RuntimeVirtualTextureSceneProxy.h
--- a/Engine/Source/Runtime/Renderer/Private/VT/RuntimeVirtualTextureSceneProxy.h
+++ b/Engine/Source/Runtime/Renderer/Private/VT/RuntimeVirtualTextureSceneProxy.h
@@ -11,9 +11,24 @@ class URuntimeVirtualTextureComponent;
 class FRuntimeVirtualTextureSceneProxy
 {
 public:
+	/** Constructor initializes resources for the URuntimeVirtualTexture associated with the provided component. */
 	FRuntimeVirtualTextureSceneProxy(URuntimeVirtualTextureComponent* InComponent);
 	~FRuntimeVirtualTextureSceneProxy();
 
-	/** Pointer to linked URuntimeVirtualTexture. Not for dereferencing, just for debug! */
+	/**
+	 * Release the object and it's associated runtime virtual texture resources. 
+	 * Call this on the main thread before deferring deletion to happen on the render thread.
+	 */
+	void Release();
+
+	/** Index in FScene::RuntimeVirtualTextures. */
+	int32 SceneIndex;
+	/** Unique ID for producer that this Proxy created. Used for finding this object and it's SceneIndex from the producer.  */
+	int32 ProducerId;
+
+	/** Pointer to linked URuntimeVirtualTexture. Not for dereferencing, just for pointer comparison. */
 	URuntimeVirtualTexture* VirtualTexture;
+
+private:
+	static int32 ProducerIdGenerator;
 };
diff --git a/Engine/Source/Runtime/Renderer/Public/PrimitiveSceneInfo.h b/Engine/Source/Runtime/Renderer/Public/PrimitiveSceneInfo.h
--- a/Engine/Source/Runtime/Renderer/Public/PrimitiveSceneInfo.h
+++ b/Engine/Source/Runtime/Renderer/Public/PrimitiveSceneInfo.h
@@ -140,17 +140,22 @@ public:
 	}
 };
 
-/** Flags needed for shadow culling.  These are pulled out of the FPrimitiveSceneProxy so we can do rough culling before dereferencing the proxy. */
+/** Flags needed for early culling of some passes. These are pulled out of the FPrimitiveSceneProxy so we can do rough culling before dereferencing the proxy. */
 struct FPrimitiveFlagsCompact
 {
 	/** True if the primitive casts dynamic shadows. */
-	uint32 bCastDynamicShadow : 1;
+	uint16 bCastDynamicShadow : 1;
 	/** True if the primitive will cache static lighting. */
-	uint32 bStaticLighting : 1;
+	uint16 bStaticLighting : 1;
 	/** True if the primitive casts static shadows. */
-	uint32 bCastStaticShadow : 1;
+	uint16 bCastStaticShadow : 1;
 	/** True if renders to virtual texture */
-	uint32 bRenderToVirtualTexture : 1;
+	uint16 bRenderToVirtualTexture : 1;
+	
+	/** Number of bits to reserve for the RuntimeVirtualTextureMask. If we use more than this number of runtime virtual textures in a scene we will trigger a slower path in rendering the VT pages. */
+	enum { RuntimeVirtualTexture_BitCount = 8 };
+	/** Mask of the allocated runtime virtual textures in the scene to render to. */
+	uint16 RuntimeVirtualTextureMask : RuntimeVirtualTexture_BitCount;
 
 	FPrimitiveFlagsCompact(const FPrimitiveSceneProxy* Proxy);
 };
