Breakout the storage of the per-primitive runtime virtual texture flags

diff --git a/Engine/Source/Runtime/Renderer/Private/PrimitiveSceneInfo.cpp b/Engine/Source/Runtime/Renderer/Private/PrimitiveSceneInfo.cpp
--- a/Engine/Source/Runtime/Renderer/Private/PrimitiveSceneInfo.cpp
+++ b/Engine/Source/Runtime/Renderer/Private/PrimitiveSceneInfo.cpp
@@ -93,8 +93,6 @@ FPrimitiveFlagsCompact::FPrimitiveFlagsCompact(const FPrimitiveSceneProxy* Proxy
 	: bCastDynamicShadow(Proxy->CastsDynamicShadow())
 	, bStaticLighting(Proxy->HasStaticLighting())
 	, bCastStaticShadow(Proxy->CastsStaticShadow())
-	, bRenderToVirtualTexture(Proxy->WritesVirtualTexture())
-	, RuntimeVirtualTextureMask(0)
 {}
 
 FPrimitiveSceneInfoCompact::FPrimitiveSceneInfoCompact(FPrimitiveSceneInfo* InPrimitiveSceneInfo) :
@@ -575,11 +573,7 @@ void FPrimitiveSceneInfo::AddToScene(FRHICommandListImmediate& RHICmdList, bool
 	PrimitiveBounds.MaxCullDistance = PrimitiveBounds.MaxDrawDistance;
 
 	Scene->PrimitiveFlagsCompact[PackedIndex] = FPrimitiveFlagsCompact(Proxy);
-	if (Scene->PrimitiveFlagsCompact[PackedIndex].bRenderToVirtualTexture)
-	{
-		Scene->PrimitiveFlagsCompact[PackedIndex].RuntimeVirtualTextureMask = Scene->GetRuntimeVirtualTextureMask(Proxy);
-	}
-
+	
 	// Store precomputed visibility ID.
 	int32 VisibilityBitIndex = Proxy->GetVisibilityId();
 	FPrimitiveVisibilityId& VisibilityId = Scene->PrimitiveVisibilityIds[PackedIndex];
@@ -623,6 +617,14 @@ void FPrimitiveSceneInfo::AddToScene(FRHICommandListImmediate& RHICmdList, bool
 	// Store the component.
 	Scene->PrimitiveComponentIds[PackedIndex] = PrimitiveComponentId;
 
+	// Store the runtime virtual texture flags.
+	const bool bRenderToVirtualTexture = Proxy->WritesVirtualTexture();
+	Scene->PrimitiveVirtualTextureFlags[PackedIndex].bRenderToVirtualTexture = bRenderToVirtualTexture;
+	if (bRenderToVirtualTexture)
+	{
+		Scene->PrimitiveVirtualTextureFlags[PackedIndex].RuntimeVirtualTextureMask = Scene->GetRuntimeVirtualTextureMask(Proxy);
+	}
+
 	{
 		FMemMark MemStackMark(FMemStack::Get());
 
diff --git a/Engine/Source/Runtime/Renderer/Private/RendererScene.cpp b/Engine/Source/Runtime/Renderer/Private/RendererScene.cpp
--- a/Engine/Source/Runtime/Renderer/Private/RendererScene.cpp
+++ b/Engine/Source/Runtime/Renderer/Private/RendererScene.cpp
@@ -552,6 +552,7 @@ void FScene::CheckPrimitiveArrays()
 	check(Primitives.Num() == PrimitiveVisibilityIds.Num());
 	check(Primitives.Num() == PrimitiveOcclusionFlags.Num());
 	check(Primitives.Num() == PrimitiveComponentIds.Num());
+	check(Primitives.Num() == PrimitiveVirtualTextureFlags.Num());
 	check(Primitives.Num() == PrimitiveOcclusionBounds.Num());
 	check(Primitives.Num() == PrimitivesNeedingStaticMeshUpdate.Num());
 }
@@ -771,6 +772,7 @@ void FScene::AddPrimitiveSceneInfo_RenderThread(FRHICommandListImmediate& RHICmd
 	PrimitiveVisibilityIds.AddUninitialized();
 	PrimitiveOcclusionFlags.AddUninitialized();
 	PrimitiveComponentIds.AddUninitialized();
+	PrimitiveVirtualTextureFlags.AddUninitialized();
 	PrimitiveOcclusionBounds.AddUninitialized();
 	PrimitivesNeedingStaticMeshUpdate.Add(false);
 	
@@ -839,6 +841,7 @@ void FScene::AddPrimitiveSceneInfo_RenderThread(FRHICommandListImmediate& RHICmd
 				TArraySwapElements(PrimitiveVisibilityIds, DestIndex, SourceIndex);
 				TArraySwapElements(PrimitiveOcclusionFlags, DestIndex, SourceIndex);
 				TArraySwapElements(PrimitiveComponentIds, DestIndex, SourceIndex);
+				TArraySwapElements(PrimitiveVirtualTextureFlags, DestIndex, SourceIndex);
 				TArraySwapElements(PrimitiveOcclusionBounds, DestIndex, SourceIndex);
 				TBitArraySwapElements(PrimitivesNeedingStaticMeshUpdate, DestIndex, SourceIndex);
 
@@ -1557,6 +1560,7 @@ void FScene::RemovePrimitiveSceneInfo_RenderThread(FPrimitiveSceneInfo* Primitiv
 				TArraySwapElements(PrimitiveVisibilityIds, DestIndex, SourceIndex);
 				TArraySwapElements(PrimitiveOcclusionFlags, DestIndex, SourceIndex);
 				TArraySwapElements(PrimitiveComponentIds, DestIndex, SourceIndex);
+				TArraySwapElements(PrimitiveVirtualTextureFlags, DestIndex, SourceIndex);
 				TArraySwapElements(PrimitiveOcclusionBounds, DestIndex, SourceIndex);
 				TBitArraySwapElements(PrimitivesNeedingStaticMeshUpdate, DestIndex, SourceIndex);
 				SourceIndex = DestIndex;
@@ -1588,6 +1592,7 @@ void FScene::RemovePrimitiveSceneInfo_RenderThread(FPrimitiveSceneInfo* Primitiv
 		PrimitiveVisibilityIds.Pop();
 		PrimitiveOcclusionFlags.Pop();
 		PrimitiveComponentIds.Pop();
+		PrimitiveVirtualTextureFlags.Pop();
 		PrimitiveOcclusionBounds.Pop();
 		PrimitivesNeedingStaticMeshUpdate.RemoveAt(PrimitivesNeedingStaticMeshUpdate.Num()-1);
 	}
@@ -2493,9 +2498,9 @@ void FScene::UpdateRuntimeVirtualTextureForAllPrimitives_RenderThread()
 {
 	for (int32 Index = 0; Index < Primitives.Num(); ++Index)
 	{
-		if (PrimitiveFlagsCompact[Index].bRenderToVirtualTexture)
+		if (PrimitiveVirtualTextureFlags[Index].bRenderToVirtualTexture)
 		{
-			PrimitiveFlagsCompact[Index].RuntimeVirtualTextureMask = GetRuntimeVirtualTextureMask(PrimitiveSceneProxies[Index]);
+			PrimitiveVirtualTextureFlags[Index].RuntimeVirtualTextureMask = GetRuntimeVirtualTextureMask(PrimitiveSceneProxies[Index]);
 		}
 	}
 }
@@ -2518,7 +2523,7 @@ uint32 FScene::GetRuntimeVirtualTextureSceneIndex(uint32 ProducerId)
 uint32 FScene::GetRuntimeVirtualTextureMask(FPrimitiveSceneProxy const* Proxy)
 {
 	uint32 Mask = 0;
-	for (int32 i = 0; i < FPrimitiveFlagsCompact::RuntimeVirtualTexture_BitCount && i < RuntimeVirtualTextures.Num(); ++i)
+	for (int32 i = 0; i < FPrimitiveVirtualTextureFlags::RuntimeVirtualTexture_BitCount && i < RuntimeVirtualTextures.Num(); ++i)
 	{
 		if (RuntimeVirtualTextures.IsValidIndex(i))
 		{
diff --git a/Engine/Source/Runtime/Renderer/Private/ScenePrivate.h b/Engine/Source/Runtime/Renderer/Private/ScenePrivate.h
--- a/Engine/Source/Runtime/Renderer/Private/ScenePrivate.h
+++ b/Engine/Source/Runtime/Renderer/Private/ScenePrivate.h
@@ -2416,6 +2416,8 @@ public:
 	TArray<FBoxSphereBounds> PrimitiveOcclusionBounds;
 	/** Packed array of primitive components associated with the primitive. */
 	TArray<FPrimitiveComponentId> PrimitiveComponentIds;
+	/** Packed array of runtime virtual texture flags. */
+	TArray<FPrimitiveVirtualTextureFlags> PrimitiveVirtualTextureFlags;
 
 	TBitArray<> PrimitivesNeedingStaticMeshUpdate;
 	TSet<FPrimitiveSceneInfo*> PrimitivesNeedingStaticMeshUpdateWithoutVisibilityCheck;
diff --git a/Engine/Source/Runtime/Renderer/Private/VT/RuntimeVirtualTextureProducer.cpp b/Engine/Source/Runtime/Renderer/Private/VT/RuntimeVirtualTextureProducer.cpp
--- a/Engine/Source/Runtime/Renderer/Private/VT/RuntimeVirtualTextureProducer.cpp
+++ b/Engine/Source/Runtime/Renderer/Private/VT/RuntimeVirtualTextureProducer.cpp
@@ -41,7 +41,7 @@ void FRuntimeVirtualTextureFinalizer::Finalize(FRHICommandListImmediate& RHICmdL
 		// Add a slow render path inside RenderPage() when this check fails. 
 		// It will need to iterate the virtual textures on each primitive instead of using the RuntimeVirtualTextureMask.
 		// Currently nothing will render for this finalizer when the check fails.
-		checkSlow(VirtualTextureSceneIndex < FPrimitiveFlagsCompact::RuntimeVirtualTexture_BitCount);
+		checkSlow(VirtualTextureSceneIndex < FPrimitiveVirtualTextureFlags::RuntimeVirtualTexture_BitCount);
 	}
 
 	for (auto Entry : Tiles)
diff --git a/Engine/Source/Runtime/Renderer/Private/VT/RuntimeVirtualTextureRender.cpp b/Engine/Source/Runtime/Renderer/Private/VT/RuntimeVirtualTextureRender.cpp
--- a/Engine/Source/Runtime/Renderer/Private/VT/RuntimeVirtualTextureRender.cpp
+++ b/Engine/Source/Runtime/Renderer/Private/VT/RuntimeVirtualTextureRender.cpp
@@ -284,10 +284,10 @@ namespace RuntimeVirtualTexture
 		FRuntimeVirtualTextureMeshProcessor MeshProcessor(Scene, View, &DynamicMeshPassContext);
 
 		// Iterate over scene and collect visible virtual texture draw commands for this view
-		//todo: Consider a broad phase (quad tree etc?) here. (But only if running over PrimitiveFlagsCompact shows up as a bottleneck.)
+		//todo: Consider a broad phase (quad tree etc?) here. (But only if running over PrimitiveVirtualTextureFlags shows up as a bottleneck.)
 		for (int32 PrimitiveIndex = 0; PrimitiveIndex < Scene->Primitives.Num(); ++PrimitiveIndex)
 		{
-			if (Scene->PrimitiveFlagsCompact[PrimitiveIndex].RuntimeVirtualTextureMask & RuntimeVirtualTextureMask)
+			if (Scene->PrimitiveVirtualTextureFlags[PrimitiveIndex].RuntimeVirtualTextureMask & RuntimeVirtualTextureMask)
 			{
 				//todo[vt]: Use quicker/more accurate 2d test here since we can pre-calculate 2d bounds in VT space.
 				FBoxSphereBounds const& Bounds = Scene->PrimitiveBounds[PrimitiveIndex].BoxSphereBounds;
diff --git a/Engine/Source/Runtime/Renderer/Public/PrimitiveSceneInfo.h b/Engine/Source/Runtime/Renderer/Public/PrimitiveSceneInfo.h
--- a/Engine/Source/Runtime/Renderer/Public/PrimitiveSceneInfo.h
+++ b/Engine/Source/Runtime/Renderer/Public/PrimitiveSceneInfo.h
@@ -140,22 +140,15 @@ public:
 	}
 };
 
-/** Flags needed for early culling of some passes. These are pulled out of the FPrimitiveSceneProxy so we can do rough culling before dereferencing the proxy. */
+/** Flags needed for shadow culling.  These are pulled out of the FPrimitiveSceneProxy so we can do rough culling before dereferencing the proxy. */
 struct FPrimitiveFlagsCompact
 {
 	/** True if the primitive casts dynamic shadows. */
-	uint16 bCastDynamicShadow : 1;
+	uint8 bCastDynamicShadow : 1;
 	/** True if the primitive will cache static lighting. */
-	uint16 bStaticLighting : 1;
+	uint8 bStaticLighting : 1;
 	/** True if the primitive casts static shadows. */
-	uint16 bCastStaticShadow : 1;
-	/** True if renders to virtual texture */
-	uint16 bRenderToVirtualTexture : 1;
-	
-	/** Number of bits to reserve for the RuntimeVirtualTextureMask. If we use more than this number of runtime virtual textures in a scene we will trigger a slower path in rendering the VT pages. */
-	enum { RuntimeVirtualTexture_BitCount = 8 };
-	/** Mask of the allocated runtime virtual textures in the scene to render to. */
-	uint16 RuntimeVirtualTextureMask : RuntimeVirtualTexture_BitCount;
+	uint8 bCastStaticShadow : 1;
 
 	FPrimitiveFlagsCompact(const FPrimitiveSceneProxy* Proxy);
 };
@@ -177,6 +170,18 @@ public:
 	FPrimitiveSceneInfoCompact(FPrimitiveSceneInfo* InPrimitiveSceneInfo);
 };
 
+/** Flags needed for broad phase culling of runtime virtual texture page rendering. */
+struct FPrimitiveVirtualTextureFlags
+{
+	/** True if the primitive can render to virtual texture */
+	uint8 bRenderToVirtualTexture : 1;
+
+	/** Number of bits to reserve for the RuntimeVirtualTextureMask. If we use more than this number of runtime virtual textures in a scene we will trigger a slower path in rendering the VT pages. */
+	enum { RuntimeVirtualTexture_BitCount = 7 };
+	/** Mask of the allocated runtime virtual textures in the scene to render to. */
+	uint8 RuntimeVirtualTextureMask : RuntimeVirtualTexture_BitCount;
+};
+
 /** The type of the octree used by FScene to find primitives. */
 typedef TOctree<FPrimitiveSceneInfoCompact,struct FPrimitiveOctreeSemantics> FScenePrimitiveOctree;
 
