DR - Fix clang warning
#rnx

diff --git a/Engine/Source/Runtime/Renderer/Private/ReflectionEnvironment.cpp b/Engine/Source/Runtime/Renderer/Private/ReflectionEnvironment.cpp
--- a/Engine/Source/Runtime/Renderer/Private/ReflectionEnvironment.cpp
+++ b/Engine/Source/Runtime/Renderer/Private/ReflectionEnvironment.cpp
@@ -853,12 +853,12 @@ void FDeferredShadingSceneRenderer::RenderDeferredReflectionsAndSkyLighting(FRHI
 				RDG_EVENT_NAME("ReflectionEnvironmentAndSky %dx%d", View.ViewRect.Width(), View.ViewRect.Height()),
 				PassParameters,
 				ERDGPassFlags::None,
-				[PassParameters, &View, PixelShader](FRHICommandList& RHICmdList)
+				[PassParameters, &View, PixelShader](FRHICommandList& InRHICmdList)
 			{
-				RHICmdList.SetViewport(View.ViewRect.Min.X, View.ViewRect.Min.Y, 0.0f, View.ViewRect.Max.X, View.ViewRect.Max.Y, 1.0f);
+				InRHICmdList.SetViewport(View.ViewRect.Min.X, View.ViewRect.Min.Y, 0.0f, View.ViewRect.Max.X, View.ViewRect.Max.Y, 1.0f);
 
 				FGraphicsPipelineStateInitializer GraphicsPSOInit;
-				FPixelShaderUtils::InitFullscreenPipelineState(RHICmdList, View.ShaderMap, *PixelShader, GraphicsPSOInit);
+				FPixelShaderUtils::InitFullscreenPipelineState(InRHICmdList, View.ShaderMap, *PixelShader, GraphicsPSOInit);
 
 				extern int32 GAOOverwriteSceneColor;
 				if (GetReflectionEnvironmentCVar() == 2 || GAOOverwriteSceneColor)
@@ -879,9 +879,9 @@ void FDeferredShadingSceneRenderer::RenderDeferredReflectionsAndSkyLighting(FRHI
 					}
 				}
 
-				SetGraphicsPipelineState(RHICmdList, GraphicsPSOInit);
-				SetShaderParameters(RHICmdList, *PixelShader, PixelShader->GetPixelShader(), *PassParameters);
-				FPixelShaderUtils::DrawFullscreenTriangle(RHICmdList);
+				SetGraphicsPipelineState(InRHICmdList, GraphicsPSOInit);
+				SetShaderParameters(InRHICmdList, *PixelShader, PixelShader->GetPixelShader(), *PassParameters);
+				FPixelShaderUtils::DrawFullscreenTriangle(InRHICmdList);
 			});
 		} // if (bRequiresApply)
 	} // for (FViewInfo& View : Views)
