DR - Fix missing parameter bound
#rnx

diff --git a/Engine/Source/Runtime/Renderer/Private/PlanarReflectionRendering.cpp b/Engine/Source/Runtime/Renderer/Private/PlanarReflectionRendering.cpp
--- a/Engine/Source/Runtime/Renderer/Private/PlanarReflectionRendering.cpp
+++ b/Engine/Source/Runtime/Renderer/Private/PlanarReflectionRendering.cpp
@@ -721,6 +721,7 @@ void FDeferredShadingSceneRenderer::RenderDeferredPlanarReflections(FRDGBuilder&
 	PassParameters->SceneTextures.SceneDepthBuffer = SceneTextures.SceneDepthBuffer;
 	PassParameters->SceneTextures.SceneGBufferA = SceneTextures.SceneGBufferA;
 	PassParameters->SceneTextures.SceneGBufferB = SceneTextures.SceneGBufferB;
+	PassParameters->SceneTextures.SceneGBufferC = SceneTextures.SceneGBufferC;
 	SetupSceneTextureSamplers(&PassParameters->SceneTextureSamplers);
 	PassParameters->ViewUniformBuffer = View.ViewUniformBuffer;
 	PassParameters->RenderTargets[0] = FRenderTargetBinding(
@@ -747,7 +748,9 @@ void FDeferredShadingSceneRenderer::RenderDeferredPlanarReflections(FRDGBuilder&
 		for (FPlanarReflectionSceneProxy* ReflectionSceneProxy : Scene->PlanarReflections)
 		{
 			if (!View.ViewFrustum.IntersectBox(ReflectionSceneProxy->WorldBounds.GetCenter(), ReflectionSceneProxy->WorldBounds.GetExtent()))
+			{
 				continue;
+			}
 
 			SCOPED_DRAW_EVENTF(RHICmdList, PlanarReflection, *ReflectionSceneProxy->OwnerName.ToString());
 
