Fix problem with VT producer callback list management

diff --git a/Engine/Source/Runtime/Renderer/Private/VT/VirtualTextureProducer.cpp b/Engine/Source/Runtime/Renderer/Private/VT/VirtualTextureProducer.cpp
--- a/Engine/Source/Runtime/Renderer/Private/VT/VirtualTextureProducer.cpp
+++ b/Engine/Source/Runtime/Renderer/Private/VT/VirtualTextureProducer.cpp
@@ -24,7 +24,7 @@ void FVirtualTextureProducer::Release(FVirtualTextureSystem* System, const FVirt
 	Description = FVTProducerDescription();
 }
 
-FVirtualTextureProducerCollection::FVirtualTextureProducerCollection()
+FVirtualTextureProducerCollection::FVirtualTextureProducerCollection() : NumPendingCallbacks(0u)
 {
 	Producers.AddDefaulted(1);
 	Producers[0].Magic = 1u; // make sure FVirtualTextureProducerHandle(0) will not resolve to the dummy producer entry
@@ -82,9 +82,13 @@ void FVirtualTextureProducerCollection::ReleaseProducer(FVirtualTextureSystem* S
 			const uint32 NextIndex = Callback.NextIndex;
 			check(Callback.OwnerHandle == Handle);
 
+			check(!Callback.bPending);
+			Callback.bPending = true;
+
 			// Move callback to pending list
 			RemoveCallbackFromList(CallbackIndex);
 			AddCallbackToList(CallbackList_Pending, CallbackIndex);
+			++NumPendingCallbacks;
 			
 			CallbackIndex = NextIndex;
 		}
@@ -101,23 +105,36 @@ void FVirtualTextureProducerCollection::ReleaseProducer(FVirtualTextureSystem* S
 void FVirtualTextureProducerCollection::CallPendingCallbacks()
 {
 	uint32 CallbackIndex = Callbacks[CallbackList_Pending].NextIndex;
+	uint32 NumCallbacksChecked = 0u;
 	while (CallbackIndex != CallbackList_Pending)
 	{
 		FCallbackEntry& Callback = Callbacks[CallbackIndex];
-		check(Callback.DestroyedFunction);
-		check(Callback.OwnerHandle.PackedValue != 0u);
-
+		check(Callback.bPending);
+	
 		// Make a copy, then release the callback entry before calling the callback function
 		// (The destroyed callback may try to remove this or other callbacks, so need to make sure state is valid before calling)
 		const FCallbackEntry CallbackCopy(Callback);
 		Callback.DestroyedFunction = nullptr;
 		Callback.Baton = nullptr;
 		Callback.OwnerHandle = FVirtualTextureProducerHandle();
+		Callback.PackedFlags = 0u;
 		ReleaseCallback(CallbackIndex);
 
-		CallbackCopy.DestroyedFunction(CallbackCopy.OwnerHandle, CallbackCopy.Baton);
+		// Possible that this callback may have been removed from the list by a previous pending callback
+		// In this case, the function pointer will be set to nullptr
+		if (CallbackCopy.DestroyedFunction)
+		{
+			check(CallbackCopy.OwnerHandle.PackedValue != 0u);
+			CallbackCopy.DestroyedFunction(CallbackCopy.OwnerHandle, CallbackCopy.Baton);
+		}
 		CallbackIndex = CallbackCopy.NextIndex;
+		++NumCallbacksChecked;
 	}
+
+	// Extra check to detect list corruption
+	check(NumCallbacksChecked == NumPendingCallbacks);
+	NumPendingCallbacks = 0u;
+
 }
 
 void FVirtualTextureProducerCollection::AddDestroyedCallback(const FVirtualTextureProducerHandle& Handle, FVTProducerDestroyedFunction* Function, void* Baton)
@@ -134,6 +151,7 @@ void FVirtualTextureProducerCollection::AddDestroyedCallback(const FVirtualTextu
 		Callback.DestroyedFunction = Function;
 		Callback.Baton = Baton;
 		Callback.OwnerHandle = Handle;
+		Callback.PackedFlags = 0u;
 	}
 }
 
@@ -152,7 +170,14 @@ uint32 FVirtualTextureProducerCollection::RemoveAllCallbacks(const void* Baton)
 			Callback.DestroyedFunction = nullptr;
 			Callback.Baton = nullptr;
 			Callback.OwnerHandle = FVirtualTextureProducerHandle();
-			ReleaseCallback(CallbackIndex);
+
+			// If callback is already pending, we can't move it back to free list, or we risk corrupting the pending list while it's being iterated
+			// Setting DestroyedFunction tol nullptr here will ensure callback is no longer invoked, and it will be moved to free list later when it's removed from pending list
+			if (!Callback.bPending)
+			{
+				Callback.PackedFlags = 0u;
+				ReleaseCallback(CallbackIndex);
+			}
 			++NumRemoved;
 		}
 	}
diff --git a/Engine/Source/Runtime/Renderer/Private/VT/VirtualTextureProducer.h b/Engine/Source/Runtime/Renderer/Private/VT/VirtualTextureProducer.h
--- a/Engine/Source/Runtime/Renderer/Private/VT/VirtualTextureProducer.h
+++ b/Engine/Source/Runtime/Renderer/Private/VT/VirtualTextureProducer.h
@@ -73,6 +73,15 @@ private:
 		FVirtualTextureProducerHandle OwnerHandle;
 		uint32 NextIndex = 0u;
 		uint32 PrevIndex = 0u;
+		union
+		{
+			uint32 PackedFlags = 0u;
+			struct
+			{
+				uint32 bPending : 1;
+				uint32 Pad : 31;
+			};
+		};
 	};
 
 	enum ECallbackListType
@@ -178,4 +187,5 @@ private:
 
 	TArray<FProducerEntry> Producers;
 	TArray<FCallbackEntry> Callbacks;
+	uint32 NumPendingCallbacks;
 };
