Ported all material and mesh material shaders to use new shader permutation API.

diff --git a/Engine/Plugins/Experimental/RemoteSession/Source/RemoteSession/Private/Channels/RemoteSessionARCameraChannel.cpp b/Engine/Plugins/Experimental/RemoteSession/Source/RemoteSession/Private/Channels/RemoteSessionARCameraChannel.cpp
--- a/Engine/Plugins/Experimental/RemoteSession/Source/RemoteSession/Private/Channels/RemoteSessionARCameraChannel.cpp
+++ b/Engine/Plugins/Experimental/RemoteSession/Source/RemoteSession/Private/Channels/RemoteSessionARCameraChannel.cpp
@@ -68,16 +68,16 @@ class FRemoteSessionARCameraVS :
 
 public:
 
-	static bool ShouldCompilePermutation(EShaderPlatform Platform, const FMaterial* Material)
+	static bool ShouldCompilePermutation(const FMaterialShaderPermutationParameters& Parameters)
 	{
-		return Material->GetMaterialDomain() == MD_PostProcess && !IsMobilePlatform(Platform);
+		return Parameters.Material->GetMaterialDomain() == MD_PostProcess && !IsMobilePlatform(Parameters.Platform);
 	}
 
-	static void ModifyCompilationEnvironment(EShaderPlatform Platform, const class FMaterial* Material, FShaderCompilerEnvironment& OutEnvironment)
+	static void ModifyCompilationEnvironment(const FMaterialShaderPermutationParameters& Parameters, FShaderCompilerEnvironment& OutEnvironment)
 	{
-		FMaterialShader::ModifyCompilationEnvironment(Platform, OutEnvironment);
+		FMaterialShader::ModifyCompilationEnvironment(Parameters.Platform, OutEnvironment);
 		OutEnvironment.SetDefine(TEXT("POST_PROCESS_MATERIAL"), 1);
-		OutEnvironment.SetDefine(TEXT("POST_PROCESS_MATERIAL_BEFORE_TONEMAP"), (Material->GetBlendableLocation() != BL_AfterTonemapping) ? 1 : 0);
+		OutEnvironment.SetDefine(TEXT("POST_PROCESS_MATERIAL_BEFORE_TONEMAP"), (Parameters.Material->GetBlendableLocation() != BL_AfterTonemapping) ? 1 : 0);
 		OutEnvironment.SetDefine(TEXT("POST_PROCESS_AR_PASSTHROUGH"), 1);
 	}
 
@@ -109,17 +109,17 @@ class FRemoteSessionARCameraPS :
 
 public:
 
-	static bool ShouldCompilePermutation(EShaderPlatform Platform, const FMaterial* Material)
+	static bool ShouldCompilePermutation(const FMaterialShaderPermutationParameters& Parameters)
 	{
-		return Material->GetMaterialDomain() == MD_PostProcess && !IsMobilePlatform(Platform);
+		return Parameters.Material->GetMaterialDomain() == MD_PostProcess && !IsMobilePlatform(Parameters.Platform);
 	}
 
-	static void ModifyCompilationEnvironment(EShaderPlatform Platform, const class FMaterial* Material, FShaderCompilerEnvironment& OutEnvironment)
+	static void ModifyCompilationEnvironment(const FMaterialShaderPermutationParameters& Parameters, FShaderCompilerEnvironment& OutEnvironment)
 	{
-		FMaterialShader::ModifyCompilationEnvironment(Platform, OutEnvironment);
+		FMaterialShader::ModifyCompilationEnvironment(Parameters.Platform, OutEnvironment);
 		OutEnvironment.SetDefine(TEXT("POST_PROCESS_MATERIAL"), 1);
 		OutEnvironment.SetDefine(TEXT("OUTPUT_MOBILE_HDR"), IsMobileHDR() ? 1 : 0);
-		OutEnvironment.SetDefine(TEXT("POST_PROCESS_MATERIAL_BEFORE_TONEMAP"), (Material->GetBlendableLocation() != BL_AfterTonemapping) ? 1 : 0);
+		OutEnvironment.SetDefine(TEXT("POST_PROCESS_MATERIAL_BEFORE_TONEMAP"), (Parameters.Material->GetBlendableLocation() != BL_AfterTonemapping) ? 1 : 0);
 	}
 
 	FRemoteSessionARCameraPS() {}
diff --git a/Engine/Plugins/Runtime/AR/Apple/AppleARKit/Source/AppleARKit/Private/AppleARKitVideoOverlay.cpp b/Engine/Plugins/Runtime/AR/Apple/AppleARKit/Source/AppleARKit/Private/AppleARKitVideoOverlay.cpp
--- a/Engine/Plugins/Runtime/AR/Apple/AppleARKit/Source/AppleARKit/Private/AppleARKitVideoOverlay.cpp
+++ b/Engine/Plugins/Runtime/AR/Apple/AppleARKit/Source/AppleARKit/Private/AppleARKitVideoOverlay.cpp
@@ -262,23 +262,23 @@ class FARKitCameraOverlayVS : public FMaterialShader
 	DECLARE_SHADER_TYPE(FARKitCameraOverlayVS, Material);
 public:
 
-	static bool ShouldCompilePermutation(EShaderPlatform Platform, const FMaterial* Material)
+	static bool ShouldCompilePermutation(const FMaterialShaderPermutationParameters& Parameters)
 	{
 		if (bIsMobileRenderer)
 		{
-			return Material->GetMaterialDomain() == MD_PostProcess && IsMobilePlatform(Platform);
+			return Parameters.Material->GetMaterialDomain() == MD_PostProcess && IsMobilePlatform(Parameters.Platform);
 		}
 		else
 		{
-			return Material->GetMaterialDomain() == MD_PostProcess && !IsMobilePlatform(Platform);
+			return Parameters.Material->GetMaterialDomain() == MD_PostProcess && !IsMobilePlatform(Parameters.Platform);
 		}
 	}
 
-	static void ModifyCompilationEnvironment(EShaderPlatform Platform, const class FMaterial* Material, FShaderCompilerEnvironment& OutEnvironment)
+	static void ModifyCompilationEnvironment(const FMaterialShaderPermutationParameters& Parameters, FShaderCompilerEnvironment& OutEnvironment)
 	{
-		FMaterialShader::ModifyCompilationEnvironment(Platform, OutEnvironment);
+		FMaterialShader::ModifyCompilationEnvironment(Parameters.Platform, OutEnvironment);
 		OutEnvironment.SetDefine(TEXT("POST_PROCESS_MATERIAL"), 1);
-		OutEnvironment.SetDefine(TEXT("POST_PROCESS_MATERIAL_BEFORE_TONEMAP"), (Material->GetBlendableLocation() != BL_AfterTonemapping) ? 1 : 0);
+		OutEnvironment.SetDefine(TEXT("POST_PROCESS_MATERIAL_BEFORE_TONEMAP"), (Parameters.Material->GetBlendableLocation() != BL_AfterTonemapping) ? 1 : 0);
 		OutEnvironment.SetDefine(TEXT("POST_PROCESS_AR_PASSTHROUGH"), 1);
 	}
 
@@ -310,24 +310,24 @@ class FARKitCameraOverlayPS : public FMaterialShader
 	DECLARE_SHADER_TYPE(FARKitCameraOverlayPS, Material);
 public:
 
-	static bool ShouldCompilePermutation(EShaderPlatform Platform, const FMaterial* Material)
+	static bool ShouldCompilePermutation(const FMaterialShaderPermutationParameters& Parameters)
 	{
 		if (bIsMobileRenderer)
 		{
-			return Material->GetMaterialDomain() == MD_PostProcess && IsMobilePlatform(Platform);
+			return Parameters.Material->GetMaterialDomain() == MD_PostProcess && IsMobilePlatform(Parameters.Platform);
 		}
 		else
 		{
-			return Material->GetMaterialDomain() == MD_PostProcess && !IsMobilePlatform(Platform);
+			return Parameters.Material->GetMaterialDomain() == MD_PostProcess && !IsMobilePlatform(Parameters.Platform);
 		}
 	}
 
-	static void ModifyCompilationEnvironment(EShaderPlatform Platform, const class FMaterial* Material, FShaderCompilerEnvironment& OutEnvironment)
+	static void ModifyCompilationEnvironment(const FMeshMaterialShaderPermutationParameters& Parameters, FShaderCompilerEnvironment& OutEnvironment)
 	{
 		FMaterialShader::ModifyCompilationEnvironment(Platform, OutEnvironment);
 		OutEnvironment.SetDefine(TEXT("POST_PROCESS_MATERIAL"), 1);
 		OutEnvironment.SetDefine(TEXT("OUTPUT_MOBILE_HDR"), IsMobileHDR() ? 1 : 0);
-		OutEnvironment.SetDefine(TEXT("POST_PROCESS_MATERIAL_BEFORE_TONEMAP"), (Material->GetBlendableLocation() != BL_AfterTonemapping) ? 1 : 0);
+		OutEnvironment.SetDefine(TEXT("POST_PROCESS_MATERIAL_BEFORE_TONEMAP"), (Parameters.Material->GetBlendableLocation() != BL_AfterTonemapping) ? 1 : 0);
 	}
 
 	FARKitCameraOverlayPS() {}
diff --git a/Engine/Plugins/Runtime/AR/Google/GoogleARCore/Source/GoogleARCoreRendering/Private/GoogleARCorePassthroughCameraRenderer.cpp b/Engine/Plugins/Runtime/AR/Google/GoogleARCore/Source/GoogleARCoreRendering/Private/GoogleARCorePassthroughCameraRenderer.cpp
--- a/Engine/Plugins/Runtime/AR/Google/GoogleARCore/Source/GoogleARCoreRendering/Private/GoogleARCorePassthroughCameraRenderer.cpp
+++ b/Engine/Plugins/Runtime/AR/Google/GoogleARCore/Source/GoogleARCoreRendering/Private/GoogleARCorePassthroughCameraRenderer.cpp
@@ -164,17 +164,17 @@ class FGoogleARCoreCameraOverlayVS : public FMaterialShader
 	DECLARE_SHADER_TYPE(FGoogleARCoreCameraOverlayVS, Material);
 public:
 
-	static bool ShouldCompilePermutation(EShaderPlatform Platform, const FMaterial* Material)
+	static bool ShouldCompilePermutation(const FMaterialShaderPermutationParameters& Parameters)
 	{
-		return Material->GetMaterialDomain() == MD_PostProcess && IsMobilePlatform(Platform);
+		return Parameters.Material->GetMaterialDomain() == MD_PostProcess && IsMobilePlatform(Parameters.Platform);
 	}
 
-	static void ModifyCompilationEnvironment(EShaderPlatform Platform, const class FMaterial* Material, FShaderCompilerEnvironment& OutEnvironment)
+	static void ModifyCompilationEnvironment(const FMaterialShaderPermutationParameters& Parameters, FShaderCompilerEnvironment& OutEnvironment)
 	{
-		FMaterialShader::ModifyCompilationEnvironment(Platform, OutEnvironment);
+		FMaterialShader::ModifyCompilationEnvironment(Parameters.Platform, OutEnvironment);
 
 		OutEnvironment.SetDefine(TEXT("POST_PROCESS_MATERIAL"), 1);
-		OutEnvironment.SetDefine(TEXT("POST_PROCESS_MATERIAL_BEFORE_TONEMAP"), (Material->GetBlendableLocation() != BL_AfterTonemapping) ? 1 : 0);
+		OutEnvironment.SetDefine(TEXT("POST_PROCESS_MATERIAL_BEFORE_TONEMAP"), (Parameters.Material->GetBlendableLocation() != BL_AfterTonemapping) ? 1 : 0);
 		OutEnvironment.SetDefine(TEXT("POST_PROCESS_AR_PASSTHROUGH"), 1);
 	}
 
@@ -204,18 +204,18 @@ class FGoogleARCoreCameraOverlayPS : public FMaterialShader
 	DECLARE_SHADER_TYPE(FGoogleARCoreCameraOverlayPS, Material);
 public:
 
-	static bool ShouldCompilePermutation(EShaderPlatform Platform, const FMaterial* Material)
+	static bool ShouldCompilePermutation(const FMaterialShaderPermutationParameters& Parameters)
 	{
-		return Material->GetMaterialDomain() == MD_PostProcess && IsMobilePlatform(Platform);
+		return Parameters.Material->GetMaterialDomain() == MD_PostProcess && IsMobilePlatform(Parameters.Platform);
 	}
 
-	static void ModifyCompilationEnvironment(EShaderPlatform Platform, const class FMaterial* Material, FShaderCompilerEnvironment& OutEnvironment)
+	static void ModifyCompilationEnvironment(const FMaterialShaderPermutationParameters& Parameters, FShaderCompilerEnvironment& OutEnvironment)
 	{
-		FMaterialShader::ModifyCompilationEnvironment(Platform, OutEnvironment);
+		FMaterialShader::ModifyCompilationEnvironment(Parameters.Platform, OutEnvironment);
 
 		OutEnvironment.SetDefine(TEXT("POST_PROCESS_MATERIAL"), 1);
 		OutEnvironment.SetDefine(TEXT("OUTPUT_MOBILE_HDR"), IsMobileHDR() ? 1 : 0);
-		OutEnvironment.SetDefine(TEXT("POST_PROCESS_MATERIAL_BEFORE_TONEMAP"), (Material->GetBlendableLocation() != BL_AfterTonemapping) ? 1 : 0);
+		OutEnvironment.SetDefine(TEXT("POST_PROCESS_MATERIAL_BEFORE_TONEMAP"), (Parameters.Material->GetBlendableLocation() != BL_AfterTonemapping) ? 1 : 0);
 	}
 
 	FGoogleARCoreCameraOverlayPS() {}
diff --git a/Engine/Source/Runtime/Engine/Private/Materials/MaterialShared.cpp b/Engine/Source/Runtime/Engine/Private/Materials/MaterialShared.cpp
--- a/Engine/Source/Runtime/Engine/Private/Materials/MaterialShared.cpp
+++ b/Engine/Source/Runtime/Engine/Private/Materials/MaterialShared.cpp
@@ -2038,7 +2038,7 @@ const FMaterial& FColoredMaterialRenderProxy::GetMaterialWithFallback(ERHIFeatur
 /**
  * Finds the shader matching the template type and the passed in vertex factory, asserts if not found.
  */
-FShader* FMaterial::GetShader(FMeshMaterialShaderType* ShaderType, FVertexFactoryType* VertexFactoryType, bool bFatalIfMissing) const
+FShader* FMaterial::GetShader(FMeshMaterialShaderType* ShaderType, FVertexFactoryType* VertexFactoryType, int32 PermutationId, bool bFatalIfMissing) const
 {
 #if WITH_EDITOR && DO_CHECK
 	// Attempt to get some more info for a rare crash (UE-35937)
@@ -2046,19 +2046,19 @@ FShader* FMaterial::GetShader(FMeshMaterialShaderType* ShaderType, FVertexFactor
 	checkf( RenderingThreadShaderMap, TEXT("RenderingThreadShaderMap was NULL (GameThreadShaderMap is %p). This may relate to bug UE-35937"), GameThreadShaderMapPtr );
 #endif
 	const FMeshMaterialShaderMap* MeshShaderMap = RenderingThreadShaderMap->GetMeshShaderMap(VertexFactoryType);
-	FShader* Shader = MeshShaderMap ? MeshShaderMap->GetShader(ShaderType) : nullptr;
+	FShader* Shader = MeshShaderMap ? MeshShaderMap->GetShader(ShaderType, PermutationId) : nullptr;
 	if (!Shader)
 	{
 		// we don't care about thread safety because we are about to crash 
 		const auto CachedGameThreadShaderMap = GameThreadShaderMap;
 		const auto CachedGameMeshShaderMap = CachedGameThreadShaderMap ? CachedGameThreadShaderMap->GetMeshShaderMap(VertexFactoryType) : nullptr;
-		bool bShaderWasFoundInGameShaderMap = CachedGameMeshShaderMap && CachedGameMeshShaderMap->GetShader(ShaderType) != nullptr;
+		bool bShaderWasFoundInGameShaderMap = CachedGameMeshShaderMap && CachedGameMeshShaderMap->GetShader(ShaderType, PermutationId) != nullptr;
 
 		// Get the ShouldCache results that determine whether the shader should be compiled
 		auto ShaderPlatform = GShaderPlatformForFeatureLevel[GetFeatureLevel()];
 		bool bMaterialShouldCache = ShouldCache(ShaderPlatform, ShaderType, VertexFactoryType);
 		bool bVFShouldCache = VertexFactoryType->ShouldCache(ShaderPlatform, this, ShaderType);
-		bool bShaderShouldCache = ShaderType->ShouldCache(ShaderPlatform, this, VertexFactoryType);
+		bool bShaderShouldCache = ShaderType->ShouldCompilePermutation(ShaderPlatform, this, VertexFactoryType, PermutationId);
 		FString MaterialUsage = GetMaterialUsageDescription();
 
 		int BreakPoint = 0;
@@ -2066,13 +2066,13 @@ FShader* FMaterial::GetShader(FMeshMaterialShaderType* ShaderType, FVertexFactor
 		// Assert with detailed information if the shader wasn't found for rendering.  
 		// This is usually the result of an incorrect ShouldCache function.
 		UE_LOG(LogMaterial, Error,
-			TEXT("Couldn't find Shader %s for Material Resource %s!\n")
+			TEXT("Couldn't find Shader (%s, %d) for Material Resource %s!\n")
 			TEXT("		RenderMeshShaderMap %d, RenderThreadShaderMap %d\n")
 			TEXT("		GameMeshShaderMap %d, GameThreadShaderMap %d, bShaderWasFoundInGameShaderMap %d\n")
 			TEXT("		With VF=%s, Platform=%s\n")
 			TEXT("		ShouldCache: Mat=%u, VF=%u, Shader=%u \n")
 			TEXT("		MaterialUsageDesc: %s"),
-			ShaderType->GetName(),  *GetFriendlyName(),
+			ShaderType->GetName(), PermutationId, *GetFriendlyName(),
 			MeshShaderMap != nullptr, RenderingThreadShaderMap != nullptr,
 			CachedGameMeshShaderMap != nullptr, CachedGameThreadShaderMap != nullptr, bShaderWasFoundInGameShaderMap,
 			VertexFactoryType->GetName(), *LegacyShaderPlatformToShaderFormat(ShaderPlatform).ToString(),
@@ -2122,7 +2122,7 @@ FShaderPipeline* FMaterial::GetShaderPipeline(class FShaderPipelineType* ShaderP
 			{
 				bool bMaterialShouldCache = ShouldCache(ShaderPlatform, ShaderType->GetMeshMaterialShaderType(), VertexFactoryType);
 				bool bVFShouldCache = VertexFactoryType->ShouldCache(ShaderPlatform, this, ShaderType->GetMeshMaterialShaderType());
-				bool bShaderShouldCache = ShaderType->GetMeshMaterialShaderType()->ShouldCache(ShaderPlatform, this, VertexFactoryType);
+				bool bShaderShouldCache = ShaderType->GetMeshMaterialShaderType()->ShouldCompilePermutation(ShaderPlatform, this, VertexFactoryType, kUniqueShaderPermutationId);
 
 				UE_LOG(LogMaterial, Error, TEXT("%s %s ShouldCache: Mat=%u, VF=%u, Shader=%u"),
 					GetShaderFrequencyString(ShaderType->GetFrequency(), false), ShaderType->GetName(), bMaterialShouldCache, bVFShouldCache, bShaderShouldCache);
@@ -2130,7 +2130,7 @@ FShaderPipeline* FMaterial::GetShaderPipeline(class FShaderPipelineType* ShaderP
 			else if (ShaderType->GetMaterialShaderType())
 			{
 				bool bMaterialShouldCache = ShouldCache(ShaderPlatform, ShaderType->GetMaterialShaderType(), VertexFactoryType);
-				bool bShaderShouldCache = ShaderType->GetMaterialShaderType()->ShouldCache(ShaderPlatform, this);
+				bool bShaderShouldCache = ShaderType->GetMaterialShaderType()->ShouldCompilePermutation(ShaderPlatform, this, kUniqueShaderPermutationId);
 
 				UE_LOG(LogMaterial, Error, TEXT("%s %s ShouldCache: Mat=%u, NO VF, Shader=%u"),
 					GetShaderFrequencyString(ShaderType->GetFrequency(), false), ShaderType->GetName(), bMaterialShouldCache, bShaderShouldCache);
@@ -2816,15 +2816,16 @@ void FMaterial::GetDependentShaderAndVFTypes(EShaderPlatform Platform, TArray<FS
 			for(TLinkedList<FShaderType*>::TIterator ShaderTypeIt(FShaderType::GetTypeList());ShaderTypeIt;ShaderTypeIt.Next())
 			{
 				FMeshMaterialShaderType* ShaderType = ShaderTypeIt->GetMeshMaterialShaderType();
-
-				if (ShaderType && 
-					ShaderType->ShouldCache(Platform, this, VertexFactoryType) && 
-					ShouldCache(Platform, ShaderType, VertexFactoryType) &&
-					VertexFactoryType->ShouldCache(Platform, this, ShaderType)
-					)
+				const int32 PermutationCount = ShaderType ? ShaderType->GetPermutationCount() : 0;
+				for (int32 PermutationId = 0; PermutationId < PermutationCount; ++PermutationId)
 				{
-					bAddedTypeFromThisVF = true;
-					OutShaderTypes.AddUnique(ShaderType);
+					if (ShaderType->ShouldCompilePermutation(Platform, this, VertexFactoryType, PermutationId) &&
+						ShouldCache(Platform, ShaderType, VertexFactoryType) &&
+						VertexFactoryType->ShouldCache(Platform, this, ShaderType))
+					{
+						bAddedTypeFromThisVF = true;
+						OutShaderTypes.AddUnique(ShaderType);
+					}
 				}
 			}
 
@@ -2838,7 +2839,7 @@ void FMaterial::GetDependentShaderAndVFTypes(EShaderPlatform Platform, TArray<FS
 					for (const FShaderType* Type : ShaderStages)
 					{
 						const FMeshMaterialShaderType* ShaderType = Type->GetMeshMaterialShaderType();
-						if (ShaderType->ShouldCache(Platform, this, VertexFactoryType) &&
+						if (ShaderType->ShouldCompilePermutation(Platform, this, VertexFactoryType, kUniqueShaderPermutationId) &&
 							ShouldCache(Platform, ShaderType, VertexFactoryType) &&
 							VertexFactoryType->ShouldCache(Platform, this, ShaderType)
 							)
@@ -2870,13 +2871,14 @@ void FMaterial::GetDependentShaderAndVFTypes(EShaderPlatform Platform, TArray<FS
 	for (TLinkedList<FShaderType*>::TIterator ShaderTypeIt(FShaderType::GetTypeList()); ShaderTypeIt; ShaderTypeIt.Next())
 	{
 		FMaterialShaderType* ShaderType = ShaderTypeIt->GetMaterialShaderType();
-
-		if (ShaderType && 
-			ShaderType->ShouldCache(Platform, this) && 
-			ShouldCache(Platform, ShaderType, nullptr)
-			)
+		const int32 PermutationCount = ShaderType ? ShaderType->GetPermutationCount() : 0;
+		for (int32 PermutationId = 0; PermutationId < PermutationCount; ++PermutationId)
 		{
-			OutShaderTypes.Add(ShaderType);
+			if (ShaderType->ShouldCompilePermutation(Platform, this, PermutationId) &&
+				ShouldCache(Platform, ShaderType, nullptr))
+			{
+				OutShaderTypes.Add(ShaderType);
+			}
 		}
 	}
 
@@ -2891,7 +2893,7 @@ void FMaterial::GetDependentShaderAndVFTypes(EShaderPlatform Platform, TArray<FS
 			{
 				const FMaterialShaderType* ShaderType = Type->GetMaterialShaderType();
 				if (ShaderType &&
-					ShaderType->ShouldCache(Platform, this) &&
+					ShaderType->ShouldCompilePermutation(Platform, this, kUniqueShaderPermutationId) &&
 					ShouldCache(Platform, ShaderType, nullptr)
 					)
 				{
@@ -3467,11 +3469,12 @@ bool FMaterialInstanceBasePropertyOverrides::operator!=(const FMaterialInstanceB
 
 //////////////////////////////////////////////////////////////////////////
 #if WITH_EDITOR
-bool FMaterialShaderMapId::ContainsShaderType(const FShaderType* ShaderType) const
+bool FMaterialShaderMapId::ContainsShaderType(const FShaderType* ShaderType, int32 PermutationId) const
 {
 	for (int32 TypeIndex = 0; TypeIndex < ShaderTypeDependencies.Num(); TypeIndex++)
 	{
-		if (ShaderTypeDependencies[TypeIndex].ShaderType == ShaderType)
+		if (ShaderTypeDependencies[TypeIndex].ShaderType == ShaderType &&
+			ShaderTypeDependencies[TypeIndex].PermutationId == PermutationId)
 		{
 			return true;
 		}
diff --git a/Engine/Source/Runtime/Landscape/Private/LandscapeGrass.cpp b/Engine/Source/Runtime/Landscape/Private/LandscapeGrass.cpp
--- a/Engine/Source/Runtime/Landscape/Private/LandscapeGrass.cpp
+++ b/Engine/Source/Runtime/Landscape/Private/LandscapeGrass.cpp
@@ -237,9 +237,9 @@ protected:
 
 public:
 
-	static bool ShouldCompilePermutation(EShaderPlatform Platform, const FMaterial* Material, const FVertexFactoryType* VertexFactoryType)
+	static bool ShouldCompilePermutation(const FMeshMaterialShaderPermutationParameters& Parameters)
 	{
-		return ShouldCacheLandscapeGrassShaders(Platform, Material, VertexFactoryType);
+		return ShouldCacheLandscapeGrassShaders(Parameters.Platform, Parameters.Material, Parameters.VertexFactoryType);
 	}
 
 	void GetShaderBindings(
@@ -273,9 +273,9 @@ class FLandscapeGrassWeightPS : public FMeshMaterialShader
 	FShaderParameter OutputPassParameter;
 public:
 
-	static bool ShouldCompilePermutation(EShaderPlatform Platform, const FMaterial* Material, const FVertexFactoryType* VertexFactoryType)
+	static bool ShouldCompilePermutation(const FMeshMaterialShaderPermutationParameters& Parameters)
 	{
-		return ShouldCacheLandscapeGrassShaders(Platform, Material, VertexFactoryType);
+		return ShouldCacheLandscapeGrassShaders(Parameters.Platform, Parameters.Material, Parameters.VertexFactoryType);
 	}
 
 	FLandscapeGrassWeightPS(const ShaderMetaType::CompiledShaderInitializerType& Initializer)
diff --git a/Engine/Source/Runtime/Renderer/Private/BasePassRendering.h b/Engine/Source/Runtime/Renderer/Private/BasePassRendering.h
--- a/Engine/Source/Runtime/Renderer/Private/BasePassRendering.h
+++ b/Engine/Source/Runtime/Renderer/Private/BasePassRendering.h
@@ -188,10 +188,10 @@ protected:
 
 public:
 
-	static void ModifyCompilationEnvironment(EShaderPlatform Platform, const FMaterial* Material, FShaderCompilerEnvironment& OutEnvironment)
+	static void ModifyCompilationEnvironment(const FMaterialShaderPermutationParameters& Parameters, FShaderCompilerEnvironment& OutEnvironment)
 	{
-		FMeshMaterialShader::ModifyCompilationEnvironment(Platform, Material, OutEnvironment);
-		FForwardLightingParameters::ModifyCompilationEnvironment(Platform, OutEnvironment);
+		FMeshMaterialShader::ModifyCompilationEnvironment(Parameters, OutEnvironment);
+		FForwardLightingParameters::ModifyCompilationEnvironment(Parameters.Platform, OutEnvironment);
 	}
 
 	virtual bool Serialize(FArchive& Ar)
@@ -249,15 +249,15 @@ protected:
 
 public:
 
-	static bool ShouldCompilePermutation(EShaderPlatform Platform,const FMaterial* Material,const FVertexFactoryType* VertexFactoryType)
+	static bool ShouldCompilePermutation(const FMeshMaterialShaderPermutationParameters& Parameters)
 	{
-		return LightMapPolicyType::ShouldCompilePermutation(Platform, Material, VertexFactoryType);
+		return LightMapPolicyType::ShouldCompilePermutation(Parameters);
 	}
 
-	static void ModifyCompilationEnvironment(EShaderPlatform Platform, const FMaterial* Material, FShaderCompilerEnvironment& OutEnvironment)
+	static void ModifyCompilationEnvironment(const FMaterialShaderPermutationParameters& Parameters, FShaderCompilerEnvironment& OutEnvironment)
 	{
-		LightMapPolicyType::ModifyCompilationEnvironment(Platform, Material, OutEnvironment);
-		Super::ModifyCompilationEnvironment(Platform, Material, OutEnvironment);
+		LightMapPolicyType::ModifyCompilationEnvironment(Parameters, OutEnvironment);
+		Super::ModifyCompilationEnvironment(Parameters, OutEnvironment);
 	}
 };
 
@@ -276,7 +276,7 @@ protected:
 	}
 
 public:
-	static bool ShouldCompilePermutation(EShaderPlatform Platform,const FMaterial* Material,const FVertexFactoryType* VertexFactoryType)
+	static bool ShouldCompilePermutation(const FMeshMaterialShaderPermutationParameters& Parameters)
 	{
 		static const auto SupportAtmosphericFog = IConsoleManager::Get().FindTConsoleVariableDataInt(TEXT("r.SupportAtmosphericFog"));
 		static const auto SupportAllShaderPermutations = IConsoleManager::Get().FindTConsoleVariableDataInt(TEXT("r.SupportAllShaderPermutations"));
@@ -284,18 +284,18 @@ public:
 
 		const bool bProjectAllowsAtmosphericFog = !SupportAtmosphericFog || SupportAtmosphericFog->GetValueOnAnyThread() != 0 || bForceAllPermutations;
 
-		bool bShouldCache = Super::ShouldCompilePermutation(Platform, Material, VertexFactoryType);
-		bShouldCache &= (bEnableAtmosphericFog && bProjectAllowsAtmosphericFog && IsTranslucentBlendMode(Material->GetBlendMode())) || !bEnableAtmosphericFog;
+		bool bShouldCache = Super::ShouldCompilePermutation(Parameters);
+		bShouldCache &= (bEnableAtmosphericFog && bProjectAllowsAtmosphericFog && IsTranslucentBlendMode(Parameters.Material->GetBlendMode())) || !bEnableAtmosphericFog;
 
 		return bShouldCache
-			&& (IsFeatureLevelSupported(Platform, ERHIFeatureLevel::SM4));
+			&& (IsFeatureLevelSupported(Parameters.Platform, ERHIFeatureLevel::SM4));
 	}
 
-	static void ModifyCompilationEnvironment(EShaderPlatform Platform, const FMaterial* Material, FShaderCompilerEnvironment& OutEnvironment)
+	static void ModifyCompilationEnvironment(const FMaterialShaderPermutationParameters& Parameters, FShaderCompilerEnvironment& OutEnvironment)
 	{
-		Super::ModifyCompilationEnvironment(Platform, Material, OutEnvironment);
+		Super::ModifyCompilationEnvironment(Parameters, OutEnvironment);
 		// @todo MetalMRT: Remove this hack and implement proper atmospheric-fog solution for Metal MRT...
-		OutEnvironment.SetDefine(TEXT("BASEPASS_ATMOSPHERIC_FOG"), !IsMetalMRTPlatform(Platform) ? bEnableAtmosphericFog : 0);
+		OutEnvironment.SetDefine(TEXT("BASEPASS_ATMOSPHERIC_FOG"), !IsMetalMRTPlatform(Parameters.Platform) ? bEnableAtmosphericFog : 0);
 	}
 };
 
@@ -317,19 +317,19 @@ protected:
 		BindBasePassUniformBuffer(Initializer.ParameterMap, PassUniformBuffer);
 	}
 
-	static bool ShouldCompilePermutation(EShaderPlatform Platform,const FMaterial* Material,const FVertexFactoryType* VertexFactoryType)
+	static bool ShouldCompilePermutation(const FMeshMaterialShaderPermutationParameters& Parameters)
 	{
 		// Re-use vertex shader gating
 		// Metal requires matching permutations, but no other platform should worry about this complication.
-		return (bEnableAtmosphericFog == false || IsMetalPlatform(Platform))
-			&& FBaseHS::ShouldCompilePermutation(Platform, Material, VertexFactoryType)
-			&& TBasePassVS<LightMapPolicyType,bEnableAtmosphericFog>::ShouldCompilePermutation(Platform,Material,VertexFactoryType);
+		return (bEnableAtmosphericFog == false || IsMetalPlatform(Parameters.Platform))
+			&& FBaseHS::ShouldCompilePermutation(Parameters)
+			&& TBasePassVS<LightMapPolicyType,bEnableAtmosphericFog>::ShouldCompilePermutation(Parameters);
 	}
 
-	static void ModifyCompilationEnvironment(EShaderPlatform Platform, const FMaterial* Material, FShaderCompilerEnvironment& OutEnvironment)
+	static void ModifyCompilationEnvironment(const FMaterialShaderPermutationParameters& Parameters, FShaderCompilerEnvironment& OutEnvironment)
 	{
 		// Re-use vertex shader compilation environment
-		TBasePassVS<LightMapPolicyType,bEnableAtmosphericFog>::ModifyCompilationEnvironment(Platform, Material, OutEnvironment);
+		TBasePassVS<LightMapPolicyType,bEnableAtmosphericFog>::ModifyCompilationEnvironment(Parameters, OutEnvironment);
 	}
 };
 
@@ -351,17 +351,17 @@ protected:
 		BindBasePassUniformBuffer(Initializer.ParameterMap, PassUniformBuffer);
 	}
 
-	static bool ShouldCompilePermutation(EShaderPlatform Platform,const FMaterial* Material,const FVertexFactoryType* VertexFactoryType)
+	static bool ShouldCompilePermutation(const FMeshMaterialShaderPermutationParameters& Parameters)
 	{
 		// Re-use vertex shader gating
-		return FBaseDS::ShouldCompilePermutation(Platform, Material, VertexFactoryType)
-			&& TBasePassVS<LightMapPolicyType,false>::ShouldCompilePermutation(Platform,Material,VertexFactoryType);
+		return FBaseDS::ShouldCompilePermutation(Parameters)
+			&& TBasePassVS<LightMapPolicyType,false>::ShouldCompilePermutation(Parameters);
 	}
 
-	static void ModifyCompilationEnvironment(EShaderPlatform Platform, const FMaterial* Material, FShaderCompilerEnvironment& OutEnvironment)
+	static void ModifyCompilationEnvironment(const FMaterialShaderPermutationParameters& Parameters, FShaderCompilerEnvironment& OutEnvironment)
 	{
 		// Re-use vertex shader compilation environment
-		TBasePassVS<LightMapPolicyType,false>::ModifyCompilationEnvironment(Platform, Material, OutEnvironment);
+		TBasePassVS<LightMapPolicyType,false>::ModifyCompilationEnvironment(Parameters, OutEnvironment);
 	}
 
 public:
@@ -381,20 +381,20 @@ class TBasePassPixelShaderPolicyParamType : public FMeshMaterialShader, public L
 {
 public:
 
-	// static bool ShouldCompilePermutation(EShaderPlatform Platform,const FMaterial* Material,const FVertexFactoryType* VertexFactoryType)
+	// static bool ShouldCompilePermutation(const FMeshMaterialShaderPermutationParameters& Parameters)
 
-	static void ModifyCompilationEnvironment(EShaderPlatform Platform, const FMaterial* Material, FShaderCompilerEnvironment& OutEnvironment)
+	static void ModifyCompilationEnvironment(const FMaterialShaderPermutationParameters& Parameters, FShaderCompilerEnvironment& OutEnvironment)
 	{
-		FMeshMaterialShader::ModifyCompilationEnvironment(Platform, Material, OutEnvironment);
+		FMeshMaterialShader::ModifyCompilationEnvironment(Parameters, OutEnvironment);
 
-		const bool bOutputVelocity = FVelocityRendering::BasePassCanOutputVelocity(Platform);
+		const bool bOutputVelocity = FVelocityRendering::BasePassCanOutputVelocity(Parameters.Platform);
 		if (bOutputVelocity)
 		{
-			const int32 VelocityIndex = IsForwardShadingEnabled(Platform) ? 1 : 4; // As defined in BasePassPixelShader.usf
+			const int32 VelocityIndex = IsForwardShadingEnabled(Parameters.Platform) ? 1 : 4; // As defined in BasePassPixelShader.usf
 			OutEnvironment.SetRenderTargetOutputFormat(VelocityIndex, PF_G16R16);
 		}
 
-		FForwardLightingParameters::ModifyCompilationEnvironment(Platform, OutEnvironment);
+		FForwardLightingParameters::ModifyCompilationEnvironment(Parameters.Platform, OutEnvironment);
 	}
 
 	static bool ValidateCompiledResult(EShaderPlatform Platform, const TArray<FMaterial*>& Materials, const FVertexFactoryType* VertexFactoryType, const FShaderParameterMap& ParameterMap, TArray<FString>& OutError)
@@ -456,15 +456,15 @@ class TBasePassPixelShaderBaseType : public TBasePassPixelShaderPolicyParamType<
 
 public:
 
-	static bool ShouldCompilePermutation(EShaderPlatform Platform,const FMaterial* Material,const FVertexFactoryType* VertexFactoryType)
+	static bool ShouldCompilePermutation(const FMeshMaterialShaderPermutationParameters& Parameters)
 	{
-		return LightMapPolicyType::ShouldCompilePermutation(Platform,Material,VertexFactoryType);
+		return LightMapPolicyType::ShouldCompilePermutation(Parameters);
 	}
 
-	static void ModifyCompilationEnvironment(EShaderPlatform Platform, const FMaterial* Material, FShaderCompilerEnvironment& OutEnvironment)
+	static void ModifyCompilationEnvironment(const FMaterialShaderPermutationParameters& Parameters, FShaderCompilerEnvironment& OutEnvironment)
 	{
-		LightMapPolicyType::ModifyCompilationEnvironment(Platform, Material, OutEnvironment);
-		Super::ModifyCompilationEnvironment(Platform, Material, OutEnvironment);
+		LightMapPolicyType::ModifyCompilationEnvironment(Parameters, OutEnvironment);
+		Super::ModifyCompilationEnvironment(Parameters, OutEnvironment);
 	}
 
 	/** Initialization constructor. */
@@ -480,13 +480,13 @@ class TBasePassPS : public TBasePassPixelShaderBaseType<LightMapPolicyType>
 	DECLARE_SHADER_TYPE(TBasePassPS,MeshMaterial);
 public:
 
-	static bool ShouldCompilePermutation(EShaderPlatform Platform,const FMaterial* Material,const FVertexFactoryType* VertexFactoryType)
+	static bool ShouldCompilePermutation(const FMeshMaterialShaderPermutationParameters& Parameters)
 	{
 		// Only compile skylight version for lit materials, and if the project allows them.
 		static const auto SupportStationarySkylight = IConsoleManager::Get().FindTConsoleVariableDataInt(TEXT("r.SupportStationarySkylight"));
 		static const auto SupportAllShaderPermutations = IConsoleManager::Get().FindTConsoleVariableDataInt(TEXT("r.SupportAllShaderPermutations"));
 
-		const bool bTranslucent = IsTranslucentBlendMode(Material->GetBlendMode());
+		const bool bTranslucent = IsTranslucentBlendMode(Parameters.Material->GetBlendMode());
 		const bool bForceAllPermutations = SupportAllShaderPermutations && SupportAllShaderPermutations->GetValueOnAnyThread() != 0;
 		const bool bProjectSupportsStationarySkylight = !SupportStationarySkylight || SupportStationarySkylight->GetValueOnAnyThread() != 0 || bForceAllPermutations;
 
@@ -495,19 +495,19 @@ public:
 			|| bTranslucent
 			// Some lightmap policies (eg Simple Forward) always require skylight support
 			|| LightMapPolicyType::RequiresSkylight()
-			|| ((bProjectSupportsStationarySkylight || IsForwardShadingEnabled(Platform)) && Material->GetShadingModels().IsLit());
+			|| ((bProjectSupportsStationarySkylight || IsForwardShadingEnabled(Parameters.Platform)) && Parameters.Material->GetShadingModels().IsLit());
 		return bCacheShaders
-			&& (IsFeatureLevelSupported(Platform, ERHIFeatureLevel::SM4))
-			&& TBasePassPixelShaderBaseType<LightMapPolicyType>::ShouldCompilePermutation(Platform, Material, VertexFactoryType);
+			&& (IsFeatureLevelSupported(Parameters.Platform, ERHIFeatureLevel::SM4))
+			&& TBasePassPixelShaderBaseType<LightMapPolicyType>::ShouldCompilePermutation(Parameters);
 	}
 
-	static void ModifyCompilationEnvironment(EShaderPlatform Platform, const FMaterial* Material, FShaderCompilerEnvironment& OutEnvironment)
+	static void ModifyCompilationEnvironment(const FMaterialShaderPermutationParameters& Parameters, FShaderCompilerEnvironment& OutEnvironment)
 	{
 		// For deferred decals, the shader class used is FDeferredDecalPS. the TBasePassPS is only used in the material editor and will read wrong values.
-		OutEnvironment.SetDefine(TEXT("SCENE_TEXTURES_DISABLED"), Material->GetMaterialDomain() != MD_Surface); 
+		OutEnvironment.SetDefine(TEXT("SCENE_TEXTURES_DISABLED"), Parameters.Material->GetMaterialDomain() != MD_Surface); 
 
 		OutEnvironment.SetDefine(TEXT("ENABLE_SKY_LIGHT"), bEnableSkyLight);
-		TBasePassPixelShaderBaseType<LightMapPolicyType>::ModifyCompilationEnvironment(Platform, Material, OutEnvironment);
+		TBasePassPixelShaderBaseType<LightMapPolicyType>::ModifyCompilationEnvironment(Parameters, OutEnvironment);
 	}
 	
 	/** Initialization constructor. */
diff --git a/Engine/Source/Runtime/Renderer/Private/CompositionLighting/PostProcessMeshDecals.cpp b/Engine/Source/Runtime/Renderer/Private/CompositionLighting/PostProcessMeshDecals.cpp
--- a/Engine/Source/Runtime/Renderer/Private/CompositionLighting/PostProcessMeshDecals.cpp
+++ b/Engine/Source/Runtime/Renderer/Private/CompositionLighting/PostProcessMeshDecals.cpp
@@ -31,9 +31,9 @@
 class FMeshDecalAccumulatePolicy
 {	
 public:
-	static bool ShouldCompilePermutation(EShaderPlatform Platform,const FMaterial* Material,const FVertexFactoryType* VertexFactoryType)
+	static bool ShouldCompilePermutation(const FMeshMaterialShaderPermutationParameters& Parameters)
 	{
-		return Material && Material->IsDeferredDecal() && IsFeatureLevelSupported(Platform, ERHIFeatureLevel::SM4);
+		return Parameters.Material && Parameters.Material->IsDeferredDecal() && IsFeatureLevelSupported(Parameters.Platform, ERHIFeatureLevel::SM4);
 	}
 };
 
@@ -56,9 +56,9 @@ protected:
 	{
 	}
 
-	static bool ShouldCompilePermutation(EShaderPlatform Platform,const FMaterial* Material,const FVertexFactoryType* VertexFactoryType)
+	static bool ShouldCompilePermutation(const FMeshMaterialShaderPermutationParameters& Parameters)
 	{
-		return FMeshDecalAccumulatePolicy::ShouldCompilePermutation(Platform,Material,VertexFactoryType);
+		return FMeshDecalAccumulatePolicy::ShouldCompilePermutation(Parameters);
 	}
 };
 
@@ -80,10 +80,10 @@ protected:
 
 	FMeshDecalsHS() {}
 
-	static bool ShouldCompilePermutation(EShaderPlatform Platform,const FMaterial* Material,const FVertexFactoryType* VertexFactoryType)
+	static bool ShouldCompilePermutation(const FMeshMaterialShaderPermutationParameters& Parameters)
 	{
-		return FBaseHS::ShouldCompilePermutation(Platform, Material, VertexFactoryType)
-			&& FMeshDecalAccumulatePolicy::ShouldCompilePermutation(Platform, Material, VertexFactoryType);
+		return FBaseHS::ShouldCompilePermutation(Parameters)
+			&& FMeshDecalAccumulatePolicy::ShouldCompilePermutation(Parameters);
 	}
 };
 
@@ -102,10 +102,10 @@ protected:
 
 	FMeshDecalsDS() {}
 
-	static bool ShouldCompilePermutation(EShaderPlatform Platform,const FMaterial* Material,const FVertexFactoryType* VertexFactoryType)
+	static bool ShouldCompilePermutation(const FMeshMaterialShaderPermutationParameters& Parameters)
 	{
-		return FBaseDS::ShouldCompilePermutation(Platform, Material, VertexFactoryType)
-			&& FMeshDecalAccumulatePolicy::ShouldCompilePermutation(Platform, Material, VertexFactoryType);
+		return FBaseDS::ShouldCompilePermutation(Parameters)
+			&& FMeshDecalAccumulatePolicy::ShouldCompilePermutation(Parameters);
 	}
 };
 
@@ -122,14 +122,14 @@ class FMeshDecalsPS : public FMeshMaterialShader
 	DECLARE_SHADER_TYPE(FMeshDecalsPS, MeshMaterial);
 
 public:
-	static bool ShouldCompilePermutation(EShaderPlatform Platform, const FMaterial* Material, const FVertexFactoryType* VertexFactoryType)
+	static bool ShouldCompilePermutation(const FMeshMaterialShaderPermutationParameters& Parameters)
 	{
-		return FMeshDecalAccumulatePolicy::ShouldCompilePermutation(Platform,Material,VertexFactoryType);
+		return FMeshDecalAccumulatePolicy::ShouldCompilePermutation(Parameters);
 	}
-	static void ModifyCompilationEnvironment(EShaderPlatform Platform, const FMaterial* Material, FShaderCompilerEnvironment& OutEnvironment)
+	static void ModifyCompilationEnvironment(const FMaterialShaderPermutationParameters& Parameters, FShaderCompilerEnvironment& OutEnvironment)
 	{
-		FMeshMaterialShader::ModifyCompilationEnvironment(Platform, Material, OutEnvironment);
-		FDecalRendering::SetDecalCompilationEnvironment(Platform, Material, OutEnvironment);
+		FMeshMaterialShader::ModifyCompilationEnvironment(Parameters, OutEnvironment);
+		FDecalRendering::SetDecalCompilationEnvironment(Parameters.Platform, Parameters.Material, OutEnvironment);
 	}
 
 	FMeshDecalsPS(const ShaderMetaType::CompiledShaderInitializerType& Initializer)
@@ -156,16 +156,16 @@ class FMeshDecalsEmissivePS : public FMeshDecalsPS
 	DECLARE_SHADER_TYPE(FMeshDecalsEmissivePS, MeshMaterial);
 
 public:
-	static bool ShouldCompilePermutation(EShaderPlatform Platform, const FMaterial* Material, const FVertexFactoryType* VertexFactoryType)
+	static bool ShouldCompilePermutation(const FMeshMaterialShaderPermutationParameters& Parameters)
 	{
-		return FMeshDecalsPS::ShouldCompilePermutation(Platform, Material, VertexFactoryType)
-			&& Material->HasEmissiveColorConnected()
-			&& IsDBufferDecalBlendMode(FDecalRenderingCommon::ComputeFinalDecalBlendMode(Platform, Material));
+		return FMeshDecalsPS::ShouldCompilePermutation(Parameters)
+			&& Parameters.Material->HasEmissiveColorConnected()
+			&& IsDBufferDecalBlendMode(FDecalRenderingCommon::ComputeFinalDecalBlendMode(Parameters.Platform, Parameters.Material));
 	}
-	static void ModifyCompilationEnvironment(EShaderPlatform Platform, const FMaterial* Material, FShaderCompilerEnvironment& OutEnvironment)
+	static void ModifyCompilationEnvironment(const FMaterialShaderPermutationParameters& Parameters, FShaderCompilerEnvironment& OutEnvironment)
 	{
-		FMeshDecalsPS::ModifyCompilationEnvironment(Platform, Material, OutEnvironment);
-		FDecalRendering::SetEmissiveDBufferDecalCompilationEnvironment(Platform, Material, OutEnvironment);
+		FMeshDecalsPS::ModifyCompilationEnvironment(Parameters, OutEnvironment);
+		FDecalRendering::SetEmissiveDBufferDecalCompilationEnvironment(Parameters.Platform, Parameters.Material, OutEnvironment);
 	}
 
 	FMeshDecalsEmissivePS() {}
diff --git a/Engine/Source/Runtime/Renderer/Private/ConvertToUniformMesh.cpp b/Engine/Source/Runtime/Renderer/Private/ConvertToUniformMesh.cpp
--- a/Engine/Source/Runtime/Renderer/Private/ConvertToUniformMesh.cpp
+++ b/Engine/Source/Runtime/Renderer/Private/ConvertToUniformMesh.cpp
@@ -32,12 +32,12 @@ protected:
 	{
 	}
 
-	static bool ShouldCompilePermutation(EShaderPlatform Platform,const FMaterial* Material,const FVertexFactoryType* VertexFactoryType)
+	static bool ShouldCompilePermutation(const FMeshMaterialShaderPermutationParameters& Parameters)
 	{
-		return IsFeatureLevelSupported(Platform, ERHIFeatureLevel::SM5) 
-			&& DoesPlatformSupportDistanceFieldGI(Platform)
-			&& (FCString::Strstr(VertexFactoryType->GetName(), TEXT("LocalVertexFactory")) != NULL
-				|| FCString::Strstr(VertexFactoryType->GetName(), TEXT("InstancedStaticMeshVertexFactory")) != NULL);
+		return IsFeatureLevelSupported(Parameters.Platform, ERHIFeatureLevel::SM5) 
+			&& DoesPlatformSupportDistanceFieldGI(Parameters.Platform)
+			&& (FCString::Strstr(Parameters.VertexFactoryType->GetName(), TEXT("LocalVertexFactory")) != NULL
+				|| FCString::Strstr(Parameters.VertexFactoryType->GetName(), TEXT("InstancedStaticMeshVertexFactory")) != NULL);
 	}
 };
 
@@ -104,12 +104,12 @@ protected:
 	{
 	}
 
-	static bool ShouldCompilePermutation(EShaderPlatform Platform,const FMaterial* Material,const FVertexFactoryType* VertexFactoryType)
+	static bool ShouldCompilePermutation(const FMeshMaterialShaderPermutationParameters& Parameters)
 	{
-		return IsFeatureLevelSupported(Platform, ERHIFeatureLevel::SM5) 
-			&& DoesPlatformSupportDistanceFieldGI(Platform)
-			&& (FCString::Strstr(VertexFactoryType->GetName(), TEXT("LocalVertexFactory")) != NULL
-				|| FCString::Strstr(VertexFactoryType->GetName(), TEXT("InstancedStaticMeshVertexFactory")) != NULL);
+		return IsFeatureLevelSupported(Parameters.Platform, ERHIFeatureLevel::SM5) 
+			&& DoesPlatformSupportDistanceFieldGI(Parameters.Platform)
+			&& (FCString::Strstr(Parameters.VertexFactoryType->GetName(), TEXT("LocalVertexFactory")) != NULL
+				|| FCString::Strstr(Parameters.VertexFactoryType->GetName(), TEXT("InstancedStaticMeshVertexFactory")) != NULL);
 	}
 
 	static void GetStreamOutElements(FStreamOutElementList& ElementList, TArray<uint32>& StreamStrides, int32& RasterizedStream)
@@ -304,15 +304,15 @@ class FEvaluateSurfelMaterialCS : public FMaterialShader
 	DECLARE_SHADER_TYPE(FEvaluateSurfelMaterialCS,Material)
 public:
 
-	static bool ShouldCompilePermutation(EShaderPlatform Platform, const FMaterial* Material)
+	static bool ShouldCompilePermutation(const FMaterialShaderPermutationParameters& Parameters)
 	{
 		//@todo - lit materials only 
-		return IsFeatureLevelSupported(Platform, ERHIFeatureLevel::SM5) && DoesPlatformSupportDistanceFieldGI(Platform);
+		return IsFeatureLevelSupported(Parameters.Platform, ERHIFeatureLevel::SM5) && DoesPlatformSupportDistanceFieldGI(Parameters.Platform);
 	}
 
-	static void ModifyCompilationEnvironment(EShaderPlatform Platform, const FMaterial* Material, FShaderCompilerEnvironment& OutEnvironment)
+	static void ModifyCompilationEnvironment(const FMaterialShaderPermutationParameters& Parameters, FShaderCompilerEnvironment& OutEnvironment)
 	{
-		FMaterialShader::ModifyCompilationEnvironment(Platform,OutEnvironment);
+		FMaterialShader::ModifyCompilationEnvironment(Parameters.Platform,OutEnvironment);
 		OutEnvironment.SetDefine(TEXT("EVALUATE_SURFEL_MATERIAL_GROUP_SIZE"), GEvaluateSurfelMaterialGroupSize);
 		OutEnvironment.SetDefine(TEXT("HAS_PRIMITIVE_UNIFORM_BUFFER"), 1);
 	}
diff --git a/Engine/Source/Runtime/Renderer/Private/DebugViewModeRendering.h b/Engine/Source/Runtime/Renderer/Private/DebugViewModeRendering.h
--- a/Engine/Source/Runtime/Renderer/Private/DebugViewModeRendering.h
+++ b/Engine/Source/Runtime/Renderer/Private/DebugViewModeRendering.h
@@ -86,9 +86,9 @@ protected:
 
 public:
 
-	static bool ShouldCompilePermutation(EShaderPlatform Platform,const FMaterial* Material,const FVertexFactoryType* VertexFactoryType)
+	static bool ShouldCompilePermutation(const FMeshMaterialShaderPermutationParameters& Parameters)
 	{
-		return AllowDebugViewVSDSHS(Platform) && Material->GetFriendlyName().Contains(TEXT("DebugViewMode"));
+		return AllowDebugViewVSDSHS(Parameters.Platform) && Parameters.Material->GetFriendlyName().Contains(TEXT("DebugViewMode"));
 	}
 
 	void GetShaderBindings(
@@ -104,10 +104,10 @@ public:
 		FMeshMaterialShader::GetShaderBindings(Scene, FeatureLevel, PrimitiveSceneProxy, MaterialRenderProxy, Material, DrawRenderState, ShaderElementData, ShaderBindings);
 	}
 
-	static void SetCommonDefinitions(EShaderPlatform Platform, const FMaterial* Material, FShaderCompilerEnvironment& OutEnvironment)
+	static void SetCommonDefinitions(const FMaterialShaderPermutationParameters& Parameters, FShaderCompilerEnvironment& OutEnvironment)
 	{
 		// SM4 has less input interpolants. Also instanced meshes use more interpolants.
-		if (Material->IsDefaultMaterial() || (IsFeatureLevelSupported(Platform, ERHIFeatureLevel::SM5) && !Material->IsUsedWithInstancedStaticMeshes()))
+		if (Parameters.Material->IsDefaultMaterial() || (IsFeatureLevelSupported(Parameters.Platform, ERHIFeatureLevel::SM5) && !Parameters.Material->IsUsedWithInstancedStaticMeshes()))
 		{	// Force the default material to pass enough texcoords to the pixel shaders (even though not using them).
 			// This is required to allow material shaders to have access to the sampled coords.
 			OutEnvironment.SetDefine(TEXT("MIN_MATERIAL_TEXCOORDS"), (uint32)4);
@@ -119,10 +119,10 @@ public:
 
 	}
 
-	static void ModifyCompilationEnvironment(EShaderPlatform Platform, const FMaterial* Material, FShaderCompilerEnvironment& OutEnvironment)
+	static void ModifyCompilationEnvironment(const FMaterialShaderPermutationParameters& Parameters, FShaderCompilerEnvironment& OutEnvironment)
 	{
-		SetCommonDefinitions(Platform, Material, OutEnvironment);
-		FMeshMaterialShader::ModifyCompilationEnvironment(Platform, OutEnvironment);
+		SetCommonDefinitions(Parameters, OutEnvironment);
+		FMeshMaterialShader::ModifyCompilationEnvironment(Parameters, OutEnvironment);
 	}
 };
 
@@ -134,15 +134,15 @@ class FDebugViewModeHS : public FBaseHS
 	DECLARE_SHADER_TYPE(FDebugViewModeHS,MeshMaterial);
 public:
 
-	static bool ShouldCompilePermutation(EShaderPlatform Platform,const FMaterial* Material,const FVertexFactoryType* VertexFactoryType)
+	static bool ShouldCompilePermutation(const FMeshMaterialShaderPermutationParameters& Parameters)
 	{
-		return FBaseHS::ShouldCompilePermutation(Platform, Material, VertexFactoryType) && FDebugViewModeVS::ShouldCompilePermutation(Platform,Material,VertexFactoryType);
+		return FBaseHS::ShouldCompilePermutation(Parameters) && FDebugViewModeVS::ShouldCompilePermutation(Parameters);
 	}
 
-	static void ModifyCompilationEnvironment(EShaderPlatform Platform, const FMaterial* Material, FShaderCompilerEnvironment& OutEnvironment)
+	static void ModifyCompilationEnvironment(const FMaterialShaderPermutationParameters& Parameters, FShaderCompilerEnvironment& OutEnvironment)
 	{
-		FDebugViewModeVS::SetCommonDefinitions(Platform, Material, OutEnvironment);
-		FBaseHS::ModifyCompilationEnvironment(Platform, OutEnvironment);
+		FDebugViewModeVS::SetCommonDefinitions(Parameters, OutEnvironment);
+		FBaseHS::ModifyCompilationEnvironment(Parameters, OutEnvironment);
 	}
 
 
@@ -158,15 +158,15 @@ class FDebugViewModeDS : public FBaseDS
 	DECLARE_SHADER_TYPE(FDebugViewModeDS,MeshMaterial);
 public:
 
-	static bool ShouldCompilePermutation(EShaderPlatform Platform,const FMaterial* Material,const FVertexFactoryType* VertexFactoryType)
+	static bool ShouldCompilePermutation(const FMeshMaterialShaderPermutationParameters& Parameters)
 	{
-		return FBaseDS::ShouldCompilePermutation(Platform, Material, VertexFactoryType) && FDebugViewModeVS::ShouldCompilePermutation(Platform,Material,VertexFactoryType);		
+		return FBaseDS::ShouldCompilePermutation(Parameters) && FDebugViewModeVS::ShouldCompilePermutation(Parameters);		
 	}
 
-	static void ModifyCompilationEnvironment(EShaderPlatform Platform, const FMaterial* Material, FShaderCompilerEnvironment& OutEnvironment)
+	static void ModifyCompilationEnvironment(const FMaterialShaderPermutationParameters& Parameters, FShaderCompilerEnvironment& OutEnvironment)
 	{
-		FDebugViewModeVS::SetCommonDefinitions(Platform, Material, OutEnvironment);
-		FBaseDS::ModifyCompilationEnvironment(Platform, OutEnvironment);
+		FDebugViewModeVS::SetCommonDefinitions(Parameters, OutEnvironment);
+		FBaseDS::ModifyCompilationEnvironment(Parameters, OutEnvironment);
 	}
 
 	FDebugViewModeDS(const ShaderMetaType::CompiledShaderInitializerType& Initializer): FBaseDS(Initializer) {}
diff --git a/Engine/Source/Runtime/Renderer/Private/DecalRenderingShared.cpp b/Engine/Source/Runtime/Renderer/Private/DecalRenderingShared.cpp
--- a/Engine/Source/Runtime/Renderer/Private/DecalRenderingShared.cpp
+++ b/Engine/Source/Runtime/Renderer/Private/DecalRenderingShared.cpp
@@ -106,15 +106,15 @@ public:
 	  * as 'UsedAsDeferredDecal' in the Material Editor gets compiled into
 	  * the shader cache.
 	  */
-	static bool ShouldCompilePermutation(EShaderPlatform Platform, const FMaterial* Material)
+	static bool ShouldCompilePermutation(const FMaterialShaderPermutationParameters& Parameters)
 	{
-		return Material->IsDeferredDecal();
+		return Parameters.Material->IsDeferredDecal();
 	}
 
-	static void ModifyCompilationEnvironment( EShaderPlatform Platform, const FMaterial* Material, FShaderCompilerEnvironment& OutEnvironment )
+	static void ModifyCompilationEnvironment(const FMaterialShaderPermutationParameters& Parameters, FShaderCompilerEnvironment& OutEnvironment )
 	{
-		FMaterialShader::ModifyCompilationEnvironment(Platform, Material, OutEnvironment);
-		FDecalRendering::SetDecalCompilationEnvironment(Platform, Material, OutEnvironment);
+		FMaterialShader::ModifyCompilationEnvironment(Parameters.Platform, OutEnvironment);
+		FDecalRendering::SetDecalCompilationEnvironment(Parameters.Platform, Parameters.Material, OutEnvironment);
 	}
 
 	FDeferredDecalPS() {}
@@ -214,16 +214,16 @@ class FDeferredDecalEmissivePS : public FDeferredDecalPS
 {
 	DECLARE_SHADER_TYPE(FDeferredDecalEmissivePS, Material);
 public:
-	static bool ShouldCompilePermutation(EShaderPlatform Platform, const FMaterial* Material)
+	static bool ShouldCompilePermutation(const FMaterialShaderPermutationParameters& Parameters)
 	{
-		return FDeferredDecalPS::ShouldCompilePermutation(Platform, Material)
-			&& Material->HasEmissiveColorConnected()
-			&& IsDBufferDecalBlendMode(FDecalRenderingCommon::ComputeFinalDecalBlendMode(Platform, Material));
+		return FDeferredDecalPS::ShouldCompilePermutation(Parameters)
+			&& Parameters.Material->HasEmissiveColorConnected()
+			&& IsDBufferDecalBlendMode(FDecalRenderingCommon::ComputeFinalDecalBlendMode(Parameters.Platform, Parameters.Material));
 	}
-	static void ModifyCompilationEnvironment(EShaderPlatform Platform, const FMaterial* Material, FShaderCompilerEnvironment& OutEnvironment)
+	static void ModifyCompilationEnvironment(const FMaterialShaderPermutationParameters& Parameters, FShaderCompilerEnvironment& OutEnvironment)
 	{
-		FDeferredDecalPS::ModifyCompilationEnvironment(Platform, Material, OutEnvironment);
-		FDecalRendering::SetEmissiveDBufferDecalCompilationEnvironment(Platform, Material, OutEnvironment);
+		FDeferredDecalPS::ModifyCompilationEnvironment(Parameters, OutEnvironment);
+		FDecalRendering::SetEmissiveDBufferDecalCompilationEnvironment(Parameters.Platform, Parameters.Material, OutEnvironment);
 	}
 
 	FDeferredDecalEmissivePS() {}
diff --git a/Engine/Source/Runtime/Renderer/Private/DepthRendering.h b/Engine/Source/Runtime/Renderer/Private/DepthRendering.h
--- a/Engine/Source/Runtime/Renderer/Private/DepthRendering.h
+++ b/Engine/Source/Runtime/Renderer/Private/DepthRendering.h
@@ -61,16 +61,20 @@ protected:
 
 public:
 
-	static bool ShouldCompilePermutation(EShaderPlatform Platform,const FMaterial* Material,const FVertexFactoryType* VertexFactoryType)
+	static bool ShouldCompilePermutation(const FMeshMaterialShaderPermutationParameters& Parameters)
 	{
 		// Only the local vertex factory supports the position-only stream
 		if (bUsePositionOnlyStream)
 		{
-			return VertexFactoryType->SupportsPositionOnly() && Material->IsSpecialEngineMaterial();
+			return Parameters.VertexFactoryType->SupportsPositionOnly() && Parameters.Material->IsSpecialEngineMaterial();
 		}
 
 		// Only compile for the default material and masked materials
-		return (Material->IsSpecialEngineMaterial() || !Material->WritesEveryPixel() || Material->MaterialMayModifyMeshPosition() || Material->IsTranslucencyWritingCustomDepth());
+		return (
+			Parameters.Material->IsSpecialEngineMaterial() ||
+			!Parameters.Material->WritesEveryPixel() ||
+			Parameters.Material->MaterialMayModifyMeshPosition() ||
+			Parameters.Material->IsTranslucencyWritingCustomDepth());
 	}
 
 	void GetShaderBindings(
@@ -95,10 +99,10 @@ class FDepthOnlyHS : public FBaseHS
 	DECLARE_SHADER_TYPE(FDepthOnlyHS,MeshMaterial);
 public:
 
-	static bool ShouldCompilePermutation(EShaderPlatform Platform,const FMaterial* Material,const FVertexFactoryType* VertexFactoryType)
+	static bool ShouldCompilePermutation(const FMeshMaterialShaderPermutationParameters& Parameters)
 	{
-		return FBaseHS::ShouldCompilePermutation(Platform, Material, VertexFactoryType)
-			&& TDepthOnlyVS<false>::ShouldCompilePermutation(Platform,Material,VertexFactoryType);
+		return FBaseHS::ShouldCompilePermutation(Parameters)
+			&& TDepthOnlyVS<false>::ShouldCompilePermutation(Parameters);
 	}
 
 	FDepthOnlyHS(const ShaderMetaType::CompiledShaderInitializerType& Initializer):
@@ -116,10 +120,10 @@ class FDepthOnlyDS : public FBaseDS
 	DECLARE_SHADER_TYPE(FDepthOnlyDS,MeshMaterial);
 public:
 
-	static bool ShouldCompilePermutation(EShaderPlatform Platform,const FMaterial* Material,const FVertexFactoryType* VertexFactoryType)
+	static bool ShouldCompilePermutation(const FMeshMaterialShaderPermutationParameters& Parameters)
 	{
-		return FBaseDS::ShouldCompilePermutation(Platform, Material, VertexFactoryType)
-			&& TDepthOnlyVS<false>::ShouldCompilePermutation(Platform,Material,VertexFactoryType);		
+		return FBaseDS::ShouldCompilePermutation(Parameters)
+			&& TDepthOnlyVS<false>::ShouldCompilePermutation(Parameters);		
 	}
 
 	FDepthOnlyDS(const ShaderMetaType::CompiledShaderInitializerType& Initializer):
@@ -137,13 +141,13 @@ class FDepthOnlyPS : public FMeshMaterialShader
 	DECLARE_SHADER_TYPE(FDepthOnlyPS,MeshMaterial);
 public:
 
-	static bool ShouldCompilePermutation(EShaderPlatform Platform,const FMaterial* Material,const FVertexFactoryType* VertexFactoryType)
+	static bool ShouldCompilePermutation(const FMeshMaterialShaderPermutationParameters& Parameters)
 	{
 		return 
 			// Compile for materials that are masked.
-			(!Material->WritesEveryPixel() || Material->HasPixelDepthOffsetConnected() || Material->IsTranslucencyWritingCustomDepth()) 
+			(!Parameters.Material->WritesEveryPixel() || Parameters.Material->HasPixelDepthOffsetConnected() || Parameters.Material->IsTranslucencyWritingCustomDepth()) 
 			// Mobile uses material pixel shader to write custom stencil to color target
-			|| (IsMobilePlatform(Platform) && (Material->IsDefaultMaterial() || Material->MaterialMayModifyMeshPosition()));
+			|| (IsMobilePlatform(Parameters.Platform) && (Parameters.Material->IsDefaultMaterial() || Parameters.Material->MaterialMayModifyMeshPosition()));
 	}
 
 	FDepthOnlyPS(const ShaderMetaType::CompiledShaderInitializerType& Initializer):
diff --git a/Engine/Source/Runtime/Renderer/Private/DistortionRendering.cpp b/Engine/Source/Runtime/Renderer/Private/DistortionRendering.cpp
--- a/Engine/Source/Runtime/Renderer/Private/DistortionRendering.cpp
+++ b/Engine/Source/Runtime/Renderer/Private/DistortionRendering.cpp
@@ -294,9 +294,9 @@ protected:
 	{
 	}
 
-	static bool ShouldCompilePermutation(EShaderPlatform Platform,const FMaterial* Material,const FVertexFactoryType* VertexFactoryType)
+	static bool ShouldCompilePermutation(const FMeshMaterialShaderPermutationParameters& Parameters)
 	{
-		return Material && IsTranslucentBlendMode(Material->GetBlendMode()) && Material->IsDistorted();
+		return Parameters.Material && IsTranslucentBlendMode(Parameters.Material->GetBlendMode()) && Parameters.Material->IsDistorted();
 	}
 };
 
@@ -318,10 +318,10 @@ protected:
 
 	FDistortionMeshHS() {}
 
-	static bool ShouldCompilePermutation(EShaderPlatform Platform,const FMaterial* Material,const FVertexFactoryType* VertexFactoryType)
+	static bool ShouldCompilePermutation(const FMeshMaterialShaderPermutationParameters& Parameters)
 	{
-		return FBaseHS::ShouldCompilePermutation(Platform, Material, VertexFactoryType)
-			&& Material && IsTranslucentBlendMode(Material->GetBlendMode()) && Material->IsDistorted();
+		return FBaseHS::ShouldCompilePermutation(Parameters)
+			&& Parameters.Material && IsTranslucentBlendMode(Parameters.Material->GetBlendMode()) && Parameters.Material->IsDistorted();
 	}
 };
 
@@ -342,10 +342,10 @@ protected:
 
 	FDistortionMeshDS() {}
 
-	static bool ShouldCompilePermutation(EShaderPlatform Platform,const FMaterial* Material,const FVertexFactoryType* VertexFactoryType)
+	static bool ShouldCompilePermutation(const FMeshMaterialShaderPermutationParameters& Parameters)
 	{
-		return FBaseDS::ShouldCompilePermutation(Platform, Material, VertexFactoryType)
-			&& Material && IsTranslucentBlendMode(Material->GetBlendMode()) && Material->IsDistorted();
+		return FBaseDS::ShouldCompilePermutation(Parameters)
+			&& Parameters.Material && IsTranslucentBlendMode(Parameters.Material->GetBlendMode()) && Parameters.Material->IsDistorted();
 	}
 };
 
@@ -363,9 +363,9 @@ class FDistortionMeshPS : public FMeshMaterialShader
 	DECLARE_SHADER_TYPE(FDistortionMeshPS,MeshMaterial);
 
 public:
-	static bool ShouldCompilePermutation(EShaderPlatform Platform,const FMaterial* Material,const FVertexFactoryType* VertexFactoryType)
+	static bool ShouldCompilePermutation(const FMeshMaterialShaderPermutationParameters& Parameters)
 	{
-		return Material && IsTranslucentBlendMode(Material->GetBlendMode()) && Material->IsDistorted();
+		return Parameters.Material && IsTranslucentBlendMode(Parameters.Material->GetBlendMode()) && Parameters.Material->IsDistorted();
 	}
 
 	FDistortionMeshPS(const ShaderMetaType::CompiledShaderInitializerType& Initializer)
@@ -382,11 +382,11 @@ public:
 		}
 	}
 
-	static void ModifyCompilationEnvironment(EShaderPlatform Platform, const FMaterial* Material, FShaderCompilerEnvironment& OutEnvironment)
+	static void ModifyCompilationEnvironment(const FMaterialShaderPermutationParameters& Parameters, FShaderCompilerEnvironment& OutEnvironment)
 	{
-		FMeshMaterialShader::ModifyCompilationEnvironment(Platform, Material, OutEnvironment);
+		FMeshMaterialShader::ModifyCompilationEnvironment(Parameters, OutEnvironment);
 
-		if (GetMaxSupportedFeatureLevel(Platform) <= ERHIFeatureLevel::ES3_1)
+		if (GetMaxSupportedFeatureLevel(Parameters.Platform) <= ERHIFeatureLevel::ES3_1)
 		{
 			// do not allow to sample scene textures on mobile during accumulation path
 			OutEnvironment.SetDefine(TEXT("SCENE_TEXTURES_DISABLED"), 1);
diff --git a/Engine/Source/Runtime/Renderer/Private/HeightfieldLighting.cpp b/Engine/Source/Runtime/Renderer/Private/HeightfieldLighting.cpp
--- a/Engine/Source/Runtime/Renderer/Private/HeightfieldLighting.cpp
+++ b/Engine/Source/Runtime/Renderer/Private/HeightfieldLighting.cpp
@@ -1071,12 +1071,12 @@ class FLightHeightfieldsPS : public FMaterialShader
 	DECLARE_SHADER_TYPE(FLightHeightfieldsPS, Material);
 public:
 
-	static bool ShouldCompilePermutation(EShaderPlatform Platform, const FMaterial* Material)
+	static bool ShouldCompilePermutation(const FMaterialShaderPermutationParameters& Parameters)
 	{
-		return Material->IsLightFunction() && IsFeatureLevelSupported(Platform, ERHIFeatureLevel::SM5) && DoesPlatformSupportDistanceFieldGI(Platform);
+		return Parameters.Material->IsLightFunction() && IsFeatureLevelSupported(Parameters.Platform, ERHIFeatureLevel::SM5) && DoesPlatformSupportDistanceFieldGI(Parameters.Platform);
 	}
 
-	static void ModifyCompilationEnvironment(EShaderPlatform Platform, const FMaterial* Material, FShaderCompilerEnvironment& OutEnvironment)
+	static void ModifyCompilationEnvironment(const FMaterialShaderPermutationParameters& Parameters, FShaderCompilerEnvironment& OutEnvironment)
 	{
 		OutEnvironment.SetDefine(TEXT("APPLY_LIGHT_FUNCTION"), 1);
 	}
diff --git a/Engine/Source/Runtime/Renderer/Private/LightFunctionRendering.cpp b/Engine/Source/Runtime/Renderer/Private/LightFunctionRendering.cpp
--- a/Engine/Source/Runtime/Renderer/Private/LightFunctionRendering.cpp
+++ b/Engine/Source/Runtime/Renderer/Private/LightFunctionRendering.cpp
@@ -36,9 +36,9 @@ public:
 	  * as 'UsedAsLightFunction' in the Material Editor gets compiled into
 	  * the shader cache.
 	  */
-	static bool ShouldCompilePermutation(EShaderPlatform Platform, const FMaterial* Material)
+	static bool ShouldCompilePermutation(const FMaterialShaderPermutationParameters& Parameters)
 	{
-		return Material->IsLightFunction() && IsFeatureLevelSupported(Platform, ERHIFeatureLevel::SM4);
+		return Parameters.Material->IsLightFunction() && IsFeatureLevelSupported(Parameters.Platform, ERHIFeatureLevel::SM4);
 	}
 
 	FLightFunctionVS( )	{ }
@@ -93,9 +93,9 @@ public:
 	  * as 'UsedAsLightFunction' in the Material Editor gets compiled into
 	  * the shader cache.
 	  */
-	static bool ShouldCompilePermutation(EShaderPlatform Platform, const FMaterial* Material)
+	static bool ShouldCompilePermutation(const FMaterialShaderPermutationParameters& Parameters)
 	{
-		return Material->IsLightFunction() && IsFeatureLevelSupported(Platform, ERHIFeatureLevel::SM4);
+		return Parameters.Material->IsLightFunction() && IsFeatureLevelSupported(Parameters.Platform, ERHIFeatureLevel::SM4);
 	}
 
 	FLightFunctionPS() {}
diff --git a/Engine/Source/Runtime/Renderer/Private/LightMapDensityRendering.h b/Engine/Source/Runtime/Renderer/Private/LightMapDensityRendering.h
--- a/Engine/Source/Runtime/Renderer/Private/LightMapDensityRendering.h
+++ b/Engine/Source/Runtime/Renderer/Private/LightMapDensityRendering.h
@@ -56,18 +56,18 @@ class TLightMapDensityVS : public FMeshMaterialShader, public LightMapPolicyType
 
 public:
 
-	static bool ShouldCompilePermutation(EShaderPlatform Platform,const FMaterial* Material,const FVertexFactoryType* VertexFactoryType)
+	static bool ShouldCompilePermutation(const FMeshMaterialShaderPermutationParameters& Parameters)
 	{
-		return  AllowDebugViewmodes(Platform) 
-				&& (Material->IsSpecialEngineMaterial() || Material->IsMasked() || Material->MaterialMayModifyMeshPosition())
-				&& LightMapPolicyType::ShouldCompilePermutation(Platform,Material,VertexFactoryType)
-				&& IsFeatureLevelSupported(Platform, ERHIFeatureLevel::SM4);
+		return  AllowDebugViewmodes(Parameters.Platform) 
+				&& (Parameters.Material->IsSpecialEngineMaterial() || Parameters.Material->IsMasked() || Parameters.Material->MaterialMayModifyMeshPosition())
+				&& LightMapPolicyType::ShouldCompilePermutation(Parameters)
+				&& IsFeatureLevelSupported(Parameters.Platform, ERHIFeatureLevel::SM4);
 	}
 
-	static void ModifyCompilationEnvironment(EShaderPlatform Platform, const FMaterial* Material, FShaderCompilerEnvironment& OutEnvironment)
+	static void ModifyCompilationEnvironment(const FMaterialShaderPermutationParameters& Parameters, FShaderCompilerEnvironment& OutEnvironment)
 	{
-		FMeshMaterialShader::ModifyCompilationEnvironment(Platform, Material, OutEnvironment);
-		LightMapPolicyType::ModifyCompilationEnvironment(Platform, Material, OutEnvironment);
+		FMeshMaterialShader::ModifyCompilationEnvironment(Parameters, OutEnvironment);
+		LightMapPolicyType::ModifyCompilationEnvironment(Parameters, OutEnvironment);
 	}
 
 	/** Initialization constructor. */
@@ -116,17 +116,17 @@ class TLightMapDensityHS : public FBaseHS
 
 public:
 
-	static bool ShouldCompilePermutation(EShaderPlatform Platform,const FMaterial* Material,const FVertexFactoryType* VertexFactoryType)
+	static bool ShouldCompilePermutation(const FMeshMaterialShaderPermutationParameters& Parameters)
 	{
-		return AllowDebugViewmodes(Platform) 
-			&& FBaseHS::ShouldCompilePermutation(Platform, Material, VertexFactoryType)
-			&& TLightMapDensityVS<LightMapPolicyType>::ShouldCompilePermutation(Platform, Material, VertexFactoryType);
+		return AllowDebugViewmodes(Parameters.Platform) 
+			&& FBaseHS::ShouldCompilePermutation(Parameters)
+			&& TLightMapDensityVS<LightMapPolicyType>::ShouldCompilePermutation(Parameters);
 	}
 
-	static void ModifyCompilationEnvironment(EShaderPlatform Platform, const FMaterial* Material, FShaderCompilerEnvironment& OutEnvironment)
+	static void ModifyCompilationEnvironment(const FMaterialShaderPermutationParameters& Parameters, FShaderCompilerEnvironment& OutEnvironment)
 	{
-		FMeshMaterialShader::ModifyCompilationEnvironment(Platform, Material, OutEnvironment);
-		LightMapPolicyType::ModifyCompilationEnvironment(Platform, Material, OutEnvironment);
+		FMeshMaterialShader::ModifyCompilationEnvironment(Parameters, OutEnvironment);
+		LightMapPolicyType::ModifyCompilationEnvironment(Parameters, OutEnvironment);
 	}
 
 	/** Initialization constructor. */
@@ -147,17 +147,17 @@ class TLightMapDensityDS : public FBaseDS
 
 public:
 
-	static bool ShouldCompilePermutation(EShaderPlatform Platform,const FMaterial* Material,const FVertexFactoryType* VertexFactoryType)
+	static bool ShouldCompilePermutation(const FMeshMaterialShaderPermutationParameters& Parameters)
 	{
-		return AllowDebugViewmodes(Platform) 
-			&& FBaseDS::ShouldCompilePermutation(Platform, Material, VertexFactoryType)
-			&& TLightMapDensityVS<LightMapPolicyType>::ShouldCompilePermutation(Platform, Material, VertexFactoryType);		
+		return AllowDebugViewmodes(Parameters.Platform) 
+			&& FBaseDS::ShouldCompilePermutation(Parameters)
+			&& TLightMapDensityVS<LightMapPolicyType>::ShouldCompilePermutation(Parameters);		
 	}
 
-	static void ModifyCompilationEnvironment(EShaderPlatform Platform, const FMaterial* Material, FShaderCompilerEnvironment& OutEnvironment)
+	static void ModifyCompilationEnvironment(const FMaterialShaderPermutationParameters& Parameters, FShaderCompilerEnvironment& OutEnvironment)
 	{
-		FMeshMaterialShader::ModifyCompilationEnvironment(Platform, Material, OutEnvironment);
-		LightMapPolicyType::ModifyCompilationEnvironment(Platform, Material, OutEnvironment);
+		FMeshMaterialShader::ModifyCompilationEnvironment(Parameters, OutEnvironment);
+		LightMapPolicyType::ModifyCompilationEnvironment(Parameters, OutEnvironment);
 	}
 
 	/** Initialization constructor. */
@@ -179,18 +179,18 @@ class TLightMapDensityPS : public FMeshMaterialShader, public LightMapPolicyType
 
 public:
 
-	static bool ShouldCompilePermutation(EShaderPlatform Platform,const FMaterial* Material,const FVertexFactoryType* VertexFactoryType)
+	static bool ShouldCompilePermutation(const FMeshMaterialShaderPermutationParameters& Parameters)
 	{
-		return	AllowDebugViewmodes(Platform) 
-				&& (Material->IsSpecialEngineMaterial() || Material->IsMasked() || Material->MaterialMayModifyMeshPosition())
-				&& LightMapPolicyType::ShouldCompilePermutation(Platform,Material,VertexFactoryType)
-				&& IsFeatureLevelSupported(Platform, ERHIFeatureLevel::SM4);
+		return	AllowDebugViewmodes(Parameters.Platform) 
+				&& (Parameters.Material->IsSpecialEngineMaterial() || Parameters.Material->IsMasked() || Parameters.Material->MaterialMayModifyMeshPosition())
+				&& LightMapPolicyType::ShouldCompilePermutation(Parameters)
+				&& IsFeatureLevelSupported(Parameters.Platform, ERHIFeatureLevel::SM4);
 	}
 
-	static void ModifyCompilationEnvironment(EShaderPlatform Platform, const FMaterial* Material, FShaderCompilerEnvironment& OutEnvironment)
+	static void ModifyCompilationEnvironment(const FMaterialShaderPermutationParameters& Parameters, FShaderCompilerEnvironment& OutEnvironment)
 	{
-		FMeshMaterialShader::ModifyCompilationEnvironment(Platform, Material, OutEnvironment);
-		LightMapPolicyType::ModifyCompilationEnvironment(Platform, Material, OutEnvironment);
+		FMeshMaterialShader::ModifyCompilationEnvironment(Parameters, OutEnvironment);
+		LightMapPolicyType::ModifyCompilationEnvironment(Parameters, OutEnvironment);
 	}
 
 	/** Initialization constructor. */
diff --git a/Engine/Source/Runtime/Renderer/Private/LightMapRendering.cpp b/Engine/Source/Runtime/Renderer/Private/LightMapRendering.cpp
--- a/Engine/Source/Runtime/Renderer/Private/LightMapRendering.cpp
+++ b/Engine/Source/Runtime/Renderer/Private/LightMapRendering.cpp
@@ -26,15 +26,15 @@ int32 GNumLightmapCoefficients[2] =
 	NUM_HQ_LIGHTMAP_COEF
 };
 
-void FCachedPointIndirectLightingPolicy::ModifyCompilationEnvironment(EShaderPlatform Platform, const FMaterial* Material, FShaderCompilerEnvironment& OutEnvironment)
+void FCachedPointIndirectLightingPolicy::ModifyCompilationEnvironment(const FMaterialShaderPermutationParameters& Parameters, FShaderCompilerEnvironment& OutEnvironment)
 {
 	OutEnvironment.SetDefine(TEXT("CACHED_POINT_INDIRECT_LIGHTING"),TEXT("1"));	
 }
 
-void FSelfShadowedCachedPointIndirectLightingPolicy::ModifyCompilationEnvironment(EShaderPlatform Platform, const FMaterial* Material, FShaderCompilerEnvironment& OutEnvironment)
+void FSelfShadowedCachedPointIndirectLightingPolicy::ModifyCompilationEnvironment(const FMaterialShaderPermutationParameters& Parameters, FShaderCompilerEnvironment& OutEnvironment)
 {
 	OutEnvironment.SetDefine(TEXT("CACHED_POINT_INDIRECT_LIGHTING"),TEXT("1"));	
-	FSelfShadowedTranslucencyPolicy::ModifyCompilationEnvironment(Platform, Material, OutEnvironment);
+	FSelfShadowedTranslucencyPolicy::ModifyCompilationEnvironment(Parameters, OutEnvironment);
 }
 
 void SetupLCIUniformBuffers(const FPrimitiveSceneProxy* PrimitiveSceneProxy, const FLightCacheInterface* LCI, FUniformBufferRHIParamRef& PrecomputedLightingBuffer, FUniformBufferRHIParamRef& LightmapResourceClusterBuffer, FUniformBufferRHIParamRef& IndirectLightingCacheBuffer)
diff --git a/Engine/Source/Runtime/Renderer/Private/LightMapRendering.h b/Engine/Source/Runtime/Renderer/Private/LightMapRendering.h
--- a/Engine/Source/Runtime/Renderer/Private/LightMapRendering.h
+++ b/Engine/Source/Runtime/Renderer/Private/LightMapRendering.h
@@ -78,12 +78,12 @@ extern TGlobalResource< FEmptyIndirectLightingCacheUniformBuffer > GEmptyIndirec
  */
 struct FNoLightMapPolicy
 {
-	static bool ShouldCompilePermutation(EShaderPlatform Platform,const FMaterial* Material,const FVertexFactoryType* VertexFactoryType)
+	static bool ShouldCompilePermutation(const FMeshMaterialShaderPermutationParameters& Parameters)
 	{
 		return true;
 	}
 
-	static void ModifyCompilationEnvironment(EShaderPlatform Platform, const FMaterial* Material, FShaderCompilerEnvironment& OutEnvironment)
+	static void ModifyCompilationEnvironment(const FMaterialShaderPermutationParameters& Parameters, FShaderCompilerEnvironment& OutEnvironment)
 	{}
 
 	static bool RequiresSkylight()
@@ -108,7 +108,7 @@ extern int32 GNumLightmapCoefficients[2];
 template< ELightmapQuality LightmapQuality >
 struct TLightMapPolicy
 {
-	static void ModifyCompilationEnvironment(EShaderPlatform Platform, const FMaterial* Material, FShaderCompilerEnvironment& OutEnvironment)
+	static void ModifyCompilationEnvironment(const FMaterialShaderPermutationParameters& Parameters, FShaderCompilerEnvironment& OutEnvironment)
 	{
 		OutEnvironment.SetDefine(GLightmapDefineName[LightmapQuality],TEXT("1"));
 		OutEnvironment.SetDefine(TEXT("NUM_LIGHTMAP_COEFFICIENTS"), GNumLightmapCoefficients[LightmapQuality]);
@@ -118,7 +118,7 @@ struct TLightMapPolicy
 		OutEnvironment.SetDefine(TEXT("LIGHTMAP_VT_ENABLED"), VirtualTextureLightmaps);
 	}
 
-	static bool ShouldCompilePermutation(EShaderPlatform Platform,const FMaterial* Material,const FVertexFactoryType* VertexFactoryType)
+	static bool ShouldCompilePermutation(const FMeshMaterialShaderPermutationParameters& Parameters)
 	{
 		static const auto AllowStaticLightingVar = IConsoleManager::Get().FindTConsoleVariableDataInt(TEXT("r.AllowStaticLighting"));
 		static const auto CVarProjectCanHaveLowQualityLightmaps = IConsoleManager::Get().FindTConsoleVariableDataInt(TEXT("r.SupportLowQualityLightmaps"));
@@ -131,11 +131,11 @@ struct TLightMapPolicy
 		const bool bShouldCacheQuality = (LightmapQuality != ELightmapQuality::LQ_LIGHTMAP) || bProjectCanHaveLowQualityLightmaps;
 
 		// GetValueOnAnyThread() as it's possible that ShouldCache is called from rendering thread. That is to output some error message.
-		return (Material->GetShadingModels().IsLit())
+		return (Parameters.Material->GetShadingModels().IsLit())
 			&& bShouldCacheQuality
-			&& VertexFactoryType->SupportsStaticLighting() 
+			&& Parameters.VertexFactoryType->SupportsStaticLighting()
 			&& (!AllowStaticLightingVar || AllowStaticLightingVar->GetValueOnAnyThread() != 0)
-			&& (Material->IsUsedWithStaticLighting() || Material->IsSpecialEngineMaterial());		
+			&& (Parameters.Material->IsUsedWithStaticLighting() || Parameters.Material->IsSpecialEngineMaterial());
 	}
 
 	static bool RequiresSkylight()
@@ -150,11 +150,11 @@ struct TDistanceFieldShadowsAndLightMapPolicy : public TLightMapPolicy< Lightmap
 {
 	typedef TLightMapPolicy< LightmapQuality >	Super;
 
-	static void ModifyCompilationEnvironment(EShaderPlatform Platform, const FMaterial* Material, FShaderCompilerEnvironment& OutEnvironment)
+	static void ModifyCompilationEnvironment(const FMaterialShaderPermutationParameters& Parameters, FShaderCompilerEnvironment& OutEnvironment)
 	{
 		OutEnvironment.SetDefine(TEXT("STATICLIGHTING_TEXTUREMASK"), 1);
 		OutEnvironment.SetDefine(TEXT("STATICLIGHTING_SIGNEDDISTANCEFIELD"), 1);
-		Super::ModifyCompilationEnvironment(Platform, Material, OutEnvironment);
+		Super::ModifyCompilationEnvironment(Parameters, OutEnvironment);
 	}
 };
 
@@ -165,9 +165,9 @@ struct FDummyLightMapPolicy : public TLightMapPolicy< HQ_LIGHTMAP >
 {
 public:
 
-	static bool ShouldCompilePermutation(EShaderPlatform Platform,const FMaterial* Material,const FVertexFactoryType* VertexFactoryType)
+	static bool ShouldCompilePermutation(const FMeshMaterialShaderPermutationParameters& Parameters)
 	{
-		return Material->GetShadingModels().IsLit() && VertexFactoryType->SupportsStaticLighting();
+		return Parameters.Material->GetShadingModels().IsLit() && Parameters.VertexFactoryType->SupportsStaticLighting();
 	}
 };
 
@@ -203,12 +203,14 @@ public:
 		FShaderUniformBufferParameter TranslucentSelfShadowBufferParameter;
 	};
 
-	static bool ShouldCompilePermutation(EShaderPlatform Platform,const FMaterial* Material,const FVertexFactoryType* VertexFactoryType)
+	static bool ShouldCompilePermutation(const FMeshMaterialShaderPermutationParameters& Parameters)
 	{
-		return Material->GetShadingModels().IsLit() && IsTranslucentBlendMode(Material->GetBlendMode()) && IsFeatureLevelSupported(Platform, ERHIFeatureLevel::SM4);
+		return Parameters.Material->GetShadingModels().IsLit() &&
+			IsTranslucentBlendMode(Parameters.Material->GetBlendMode()) &&
+			IsFeatureLevelSupported(Parameters.Platform, ERHIFeatureLevel::SM4);
 	}
 
-	static void ModifyCompilationEnvironment(EShaderPlatform Platform, const FMaterial* Material, FShaderCompilerEnvironment& OutEnvironment)
+	static void ModifyCompilationEnvironment(const FMaterialShaderPermutationParameters& Parameters, FShaderCompilerEnvironment& OutEnvironment)
 	{
 		OutEnvironment.SetDefine(TEXT("TRANSLUCENT_SELF_SHADOWING"),TEXT("1"));
 	}
@@ -247,15 +249,15 @@ public:
  */
 struct FPrecomputedVolumetricLightmapLightingPolicy
 {
-	static bool ShouldCompilePermutation(EShaderPlatform Platform,const FMaterial* Material,const FVertexFactoryType* VertexFactoryType)
+	static bool ShouldCompilePermutation(const FMeshMaterialShaderPermutationParameters& Parameters)
 	{
 		static const auto AllowStaticLightingVar = IConsoleManager::Get().FindTConsoleVariableDataInt(TEXT("r.AllowStaticLighting"));
 	
-		return Material->GetShadingModels().IsLit()
-			&& (!AllowStaticLightingVar || AllowStaticLightingVar->GetValueOnAnyThread() != 0);
+		return Parameters.Material->GetShadingModels().IsLit() &&
+			(!AllowStaticLightingVar || AllowStaticLightingVar->GetValueOnAnyThread() != 0);
 	}
 
-	static void ModifyCompilationEnvironment(EShaderPlatform Platform, const FMaterial* Material, FShaderCompilerEnvironment& OutEnvironment)
+	static void ModifyCompilationEnvironment(const FMaterialShaderPermutationParameters& Parameters, FShaderCompilerEnvironment& OutEnvironment)
 	{
 		OutEnvironment.SetDefine(TEXT("PRECOMPUTED_IRRADIANCE_VOLUME_LIGHTING"),TEXT("1"));
 	}
@@ -266,17 +268,17 @@ struct FPrecomputedVolumetricLightmapLightingPolicy
  */
 struct FCachedVolumeIndirectLightingPolicy
 {
-	static bool ShouldCompilePermutation(EShaderPlatform Platform,const FMaterial* Material,const FVertexFactoryType* VertexFactoryType)
+	static bool ShouldCompilePermutation(const FMeshMaterialShaderPermutationParameters& Parameters)
 	{
 		static const auto AllowStaticLightingVar = IConsoleManager::Get().FindTConsoleVariableDataInt(TEXT("r.AllowStaticLighting"));
 
-		return Material->GetShadingModels().IsLit()
-			&& !IsTranslucentBlendMode(Material->GetBlendMode()) 
+		return Parameters.Material->GetShadingModels().IsLit()
+			&& !IsTranslucentBlendMode(Parameters.Material->GetBlendMode())
 			&& (!AllowStaticLightingVar || AllowStaticLightingVar->GetValueOnAnyThread() != 0)
-			&& IsFeatureLevelSupported(Platform, ERHIFeatureLevel::SM4);
+			&& IsFeatureLevelSupported(Parameters.Platform, ERHIFeatureLevel::SM4);
 	}
 
-	static void ModifyCompilationEnvironment(EShaderPlatform Platform, const FMaterial* Material, FShaderCompilerEnvironment& OutEnvironment)
+	static void ModifyCompilationEnvironment(const FMaterialShaderPermutationParameters& Parameters, FShaderCompilerEnvironment& OutEnvironment)
 	{
 		OutEnvironment.SetDefine(TEXT("CACHED_VOLUME_INDIRECT_LIGHTING"),TEXT("1"));
 	}
@@ -293,15 +295,15 @@ struct FCachedVolumeIndirectLightingPolicy
  */
 struct FCachedPointIndirectLightingPolicy
 {
-	static bool ShouldCompilePermutation(EShaderPlatform Platform,const FMaterial* Material,const FVertexFactoryType* VertexFactoryType)
+	static bool ShouldCompilePermutation(const FMeshMaterialShaderPermutationParameters& Parameters)
 	{
 		static const auto AllowStaticLightingVar = IConsoleManager::Get().FindTConsoleVariableDataInt(TEXT("r.AllowStaticLighting"));
 	
-		return Material->GetShadingModels().IsLit()
+		return Parameters.Material->GetShadingModels().IsLit()
 			&& (!AllowStaticLightingVar || AllowStaticLightingVar->GetValueOnAnyThread() != 0);
 	}
 
-	static void ModifyCompilationEnvironment(EShaderPlatform Platform, const FMaterial* Material, FShaderCompilerEnvironment& OutEnvironment);
+	static void ModifyCompilationEnvironment(const FMaterialShaderPermutationParameters& Parameters, FShaderCompilerEnvironment& OutEnvironment);
 
 	static bool RequiresSkylight()
 	{
@@ -314,16 +316,16 @@ struct FCachedPointIndirectLightingPolicy
  */
 struct FSimpleNoLightmapLightingPolicy : public FNoLightMapPolicy
 {
-	static bool ShouldCompilePermutation(EShaderPlatform Platform,const FMaterial* Material,const FVertexFactoryType* VertexFactoryType)
+	static bool ShouldCompilePermutation(const FMeshMaterialShaderPermutationParameters& Parameters)
 	{
-		return PlatformSupportsSimpleForwardShading(Platform)
-			&& FNoLightMapPolicy::ShouldCompilePermutation(Platform, Material, VertexFactoryType);
+		return PlatformSupportsSimpleForwardShading(Parameters.Platform)
+			&& FNoLightMapPolicy::ShouldCompilePermutation(Parameters);
 	}
 
-	static void ModifyCompilationEnvironment(EShaderPlatform Platform, const FMaterial* Material, FShaderCompilerEnvironment& OutEnvironment)
+	static void ModifyCompilationEnvironment(const FMaterialShaderPermutationParameters& Parameters, FShaderCompilerEnvironment& OutEnvironment)
 	{
 		OutEnvironment.SetDefine(TEXT("SIMPLE_FORWARD_SHADING"),TEXT("1"));
-		FNoLightMapPolicy::ModifyCompilationEnvironment(Platform, Material, OutEnvironment);
+		FNoLightMapPolicy::ModifyCompilationEnvironment(Parameters, OutEnvironment);
 	}
 
 	static bool RequiresSkylight()
@@ -337,19 +339,19 @@ struct FSimpleNoLightmapLightingPolicy : public FNoLightMapPolicy
  */
 struct FSimpleLightmapOnlyLightingPolicy : public TLightMapPolicy<HQ_LIGHTMAP>
 {
-	static bool ShouldCompilePermutation(EShaderPlatform Platform,const FMaterial* Material,const FVertexFactoryType* VertexFactoryType)
+	static bool ShouldCompilePermutation(const FMeshMaterialShaderPermutationParameters& Parameters)
 	{
 		static const auto AllowStaticLightingVar = IConsoleManager::Get().FindTConsoleVariableDataInt(TEXT("r.AllowStaticLighting"));
 
 		return AllowStaticLightingVar->GetValueOnAnyThread() != 0
-			&& PlatformSupportsSimpleForwardShading(Platform)
-			&& TLightMapPolicy<HQ_LIGHTMAP>::ShouldCompilePermutation(Platform, Material, VertexFactoryType);
+			&& PlatformSupportsSimpleForwardShading(Parameters.Platform)
+			&& TLightMapPolicy<HQ_LIGHTMAP>::ShouldCompilePermutation(Parameters);
 	}
 
-	static void ModifyCompilationEnvironment(EShaderPlatform Platform, const FMaterial* Material, FShaderCompilerEnvironment& OutEnvironment)
+	static void ModifyCompilationEnvironment(const FMaterialShaderPermutationParameters& Parameters, FShaderCompilerEnvironment& OutEnvironment)
 	{
 		OutEnvironment.SetDefine(TEXT("SIMPLE_FORWARD_SHADING"),TEXT("1"));
-		TLightMapPolicy<HQ_LIGHTMAP>::ModifyCompilationEnvironment(Platform, Material, OutEnvironment);
+		TLightMapPolicy<HQ_LIGHTMAP>::ModifyCompilationEnvironment(Parameters, OutEnvironment);
 	}
 
 	static bool RequiresSkylight()
@@ -363,13 +365,13 @@ struct FSimpleLightmapOnlyLightingPolicy : public TLightMapPolicy<HQ_LIGHTMAP>
  */
 struct FSimpleDirectionalLightLightingPolicy
 {
-	static bool ShouldCompilePermutation(EShaderPlatform Platform,const FMaterial* Material,const FVertexFactoryType* VertexFactoryType)
+	static bool ShouldCompilePermutation(const FMeshMaterialShaderPermutationParameters& Parameters)
 	{
-		return PlatformSupportsSimpleForwardShading(Platform)
-			&& Material->GetShadingModels().IsLit();
+		return PlatformSupportsSimpleForwardShading(Parameters.Platform)
+			&& Parameters.Material->GetShadingModels().IsLit();
 	}
 
-	static void ModifyCompilationEnvironment(EShaderPlatform Platform, const FMaterial* Material, FShaderCompilerEnvironment& OutEnvironment)
+	static void ModifyCompilationEnvironment(const FMaterialShaderPermutationParameters& Parameters, FShaderCompilerEnvironment& OutEnvironment)
 	{
 		OutEnvironment.SetDefine(TEXT("SIMPLE_FORWARD_SHADING"),TEXT("1"));
 		OutEnvironment.SetDefine(TEXT("SIMPLE_FORWARD_DIRECTIONAL_LIGHT"),TEXT("1"));
@@ -386,20 +388,20 @@ struct FSimpleDirectionalLightLightingPolicy
  */
 struct FSimpleStationaryLightPrecomputedShadowsLightingPolicy : public TDistanceFieldShadowsAndLightMapPolicy<HQ_LIGHTMAP>
 {
-	static bool ShouldCompilePermutation(EShaderPlatform Platform,const FMaterial* Material,const FVertexFactoryType* VertexFactoryType)
+	static bool ShouldCompilePermutation(const FMeshMaterialShaderPermutationParameters& Parameters)
 	{
 		static const auto AllowStaticLightingVar = IConsoleManager::Get().FindTConsoleVariableDataInt(TEXT("r.AllowStaticLighting"));
 
 		return AllowStaticLightingVar->GetValueOnAnyThread() != 0
-			&& PlatformSupportsSimpleForwardShading(Platform)
-			&& TDistanceFieldShadowsAndLightMapPolicy<HQ_LIGHTMAP>::ShouldCompilePermutation(Platform, Material, VertexFactoryType);
+			&& PlatformSupportsSimpleForwardShading(Parameters.Platform)
+			&& TDistanceFieldShadowsAndLightMapPolicy<HQ_LIGHTMAP>::ShouldCompilePermutation(Parameters);
 	}
 
-	static void ModifyCompilationEnvironment(EShaderPlatform Platform, const FMaterial* Material, FShaderCompilerEnvironment& OutEnvironment)
+	static void ModifyCompilationEnvironment(const FMaterialShaderPermutationParameters& Parameters, FShaderCompilerEnvironment& OutEnvironment)
 	{
 		OutEnvironment.SetDefine(TEXT("SIMPLE_FORWARD_SHADING"),TEXT("1"));
 		OutEnvironment.SetDefine(TEXT("SIMPLE_FORWARD_DIRECTIONAL_LIGHT"),TEXT("1"));
-		TDistanceFieldShadowsAndLightMapPolicy<HQ_LIGHTMAP>::ModifyCompilationEnvironment(Platform, Material, OutEnvironment);
+		TDistanceFieldShadowsAndLightMapPolicy<HQ_LIGHTMAP>::ModifyCompilationEnvironment(Parameters, OutEnvironment);
 	}
 
 	static bool RequiresSkylight()
@@ -413,20 +415,20 @@ struct FSimpleStationaryLightPrecomputedShadowsLightingPolicy : public TDistance
  */
 struct FSimpleStationaryLightSingleSampleShadowsLightingPolicy : public FCachedPointIndirectLightingPolicy
 {
-	static bool ShouldCompilePermutation(EShaderPlatform Platform,const FMaterial* Material,const FVertexFactoryType* VertexFactoryType)
+	static bool ShouldCompilePermutation(const FMeshMaterialShaderPermutationParameters& Parameters)
 	{
 		static const auto AllowStaticLightingVar = IConsoleManager::Get().FindTConsoleVariableDataInt(TEXT("r.AllowStaticLighting"));
 
 		return AllowStaticLightingVar->GetValueOnAnyThread() != 0
-			&& PlatformSupportsSimpleForwardShading(Platform)
-			&& FCachedPointIndirectLightingPolicy::ShouldCompilePermutation(Platform, Material, VertexFactoryType);
+			&& PlatformSupportsSimpleForwardShading(Parameters.Platform)
+			&& FCachedPointIndirectLightingPolicy::ShouldCompilePermutation(Parameters);
 	}
 
-	static void ModifyCompilationEnvironment(EShaderPlatform Platform, const FMaterial* Material, FShaderCompilerEnvironment& OutEnvironment)
+	static void ModifyCompilationEnvironment(const FMaterialShaderPermutationParameters& Parameters, FShaderCompilerEnvironment& OutEnvironment)
 	{
 		OutEnvironment.SetDefine(TEXT("SIMPLE_FORWARD_SHADING"),TEXT("1"));
 		OutEnvironment.SetDefine(TEXT("SIMPLE_FORWARD_DIRECTIONAL_LIGHT"),TEXT("1"));
-		FCachedPointIndirectLightingPolicy::ModifyCompilationEnvironment(Platform, Material, OutEnvironment);
+		FCachedPointIndirectLightingPolicy::ModifyCompilationEnvironment(Parameters, OutEnvironment);
 	}
 
 	static bool RequiresSkylight()
@@ -437,20 +439,20 @@ struct FSimpleStationaryLightSingleSampleShadowsLightingPolicy : public FCachedP
 
 struct FSimpleStationaryLightVolumetricLightmapShadowsLightingPolicy : public FPrecomputedVolumetricLightmapLightingPolicy
 {
-	static bool ShouldCompilePermutation(EShaderPlatform Platform,const FMaterial* Material,const FVertexFactoryType* VertexFactoryType)
+	static bool ShouldCompilePermutation(const FMeshMaterialShaderPermutationParameters& Parameters)
 	{
 		static const auto AllowStaticLightingVar = IConsoleManager::Get().FindTConsoleVariableDataInt(TEXT("r.AllowStaticLighting"));
 
 		return AllowStaticLightingVar->GetValueOnAnyThread() != 0
-			&& PlatformSupportsSimpleForwardShading(Platform)
-			&& FPrecomputedVolumetricLightmapLightingPolicy::ShouldCompilePermutation(Platform, Material, VertexFactoryType);
+			&& PlatformSupportsSimpleForwardShading(Parameters.Platform)
+			&& FPrecomputedVolumetricLightmapLightingPolicy::ShouldCompilePermutation(Parameters);
 	}
 
-	static void ModifyCompilationEnvironment(EShaderPlatform Platform, const FMaterial* Material, FShaderCompilerEnvironment& OutEnvironment)
+	static void ModifyCompilationEnvironment(const FMaterialShaderPermutationParameters& Parameters, FShaderCompilerEnvironment& OutEnvironment)
 	{
 		OutEnvironment.SetDefine(TEXT("SIMPLE_FORWARD_SHADING"),TEXT("1"));
 		OutEnvironment.SetDefine(TEXT("SIMPLE_FORWARD_DIRECTIONAL_LIGHT"),TEXT("1"));
-		FPrecomputedVolumetricLightmapLightingPolicy::ModifyCompilationEnvironment(Platform, Material, OutEnvironment);
+		FPrecomputedVolumetricLightmapLightingPolicy::ModifyCompilationEnvironment(Parameters, OutEnvironment);
 	}
 };
 
@@ -459,11 +461,11 @@ class FMobileDistanceFieldShadowsAndLQLightMapPolicy : public TDistanceFieldShad
 {
 	typedef TDistanceFieldShadowsAndLightMapPolicy<LQ_LIGHTMAP>	Super;
 public:
-	static bool ShouldCompilePermutation(EShaderPlatform Platform, const FMaterial* Material, const FVertexFactoryType* VertexFactoryType)
+	static bool ShouldCompilePermutation(const FMeshMaterialShaderPermutationParameters& Parameters)
 	{
 		static auto* CVarMobileAllowDistanceFieldShadows = IConsoleManager::Get().FindTConsoleVariableDataInt(TEXT("r.Mobile.AllowDistanceFieldShadows"));
 		const bool bMobileAllowDistanceFieldShadows = CVarMobileAllowDistanceFieldShadows->GetValueOnAnyThread() == 1;
-		return bMobileAllowDistanceFieldShadows && Super::ShouldCompilePermutation(Platform, Material, VertexFactoryType);
+		return bMobileAllowDistanceFieldShadows && Super::ShouldCompilePermutation(Parameters);
 	}
 
 	static bool RequiresSkylight()
@@ -484,19 +486,19 @@ class FMobileDistanceFieldShadowsLightMapAndCSMLightingPolicy : public FMobileDi
 	typedef FMobileDistanceFieldShadowsAndLQLightMapPolicy Super;
 
 public:
-	static bool ShouldCompilePermutation(EShaderPlatform Platform, const FMaterial* Material, const FVertexFactoryType* VertexFactoryType)
+	static bool ShouldCompilePermutation(const FMeshMaterialShaderPermutationParameters& Parameters)
 	{
 		static auto* CVarMobileEnableStaticAndCSMShadowReceivers = IConsoleManager::Get().FindTConsoleVariableDataInt(TEXT("r.Mobile.EnableStaticAndCSMShadowReceivers"));
 		const bool bMobileEnableStaticAndCSMShadowReceivers = CVarMobileEnableStaticAndCSMShadowReceivers->GetValueOnAnyThread() == 1;
-		return bMobileEnableStaticAndCSMShadowReceivers && (Material->GetShadingModels().IsLit()) && Super::ShouldCompilePermutation(Platform, Material, VertexFactoryType);
+		return bMobileEnableStaticAndCSMShadowReceivers && (Parameters.Material->GetShadingModels().IsLit()) && Super::ShouldCompilePermutation(Parameters);
 	}
 
-	static void ModifyCompilationEnvironment(EShaderPlatform Platform, const FMaterial* Material, FShaderCompilerEnvironment& OutEnvironment)
+	static void ModifyCompilationEnvironment(const FMaterialShaderPermutationParameters& Parameters, FShaderCompilerEnvironment& OutEnvironment)
 	{
 		OutEnvironment.SetDefine(TEXT("DIRECTIONAL_LIGHT_CSM"), TEXT("1"));
 		OutEnvironment.SetDefine(TEXT(PREPROCESSOR_TO_STRING(MAX_MOBILE_SHADOWCASCADES)), GetMobileMaxShadowCascades());
 
-		Super::ModifyCompilationEnvironment(Platform, Material, OutEnvironment);
+		Super::ModifyCompilationEnvironment(Parameters, OutEnvironment);
 	}
 
 	static bool RequiresSkylight()
@@ -508,17 +510,17 @@ public:
 /** Mobile Specific: Combines an unshadowed directional light with indirect lighting from a single SH sample. */
 struct FMobileDirectionalLightAndSHIndirectPolicy
 {
-	static bool ShouldCompilePermutation(EShaderPlatform Platform,const FMaterial* Material,const FVertexFactoryType* VertexFactoryType)
+	static bool ShouldCompilePermutation(const FMeshMaterialShaderPermutationParameters& Parameters)
 	{
 		static auto* CVarAllowStaticLighting = IConsoleManager::Get().FindTConsoleVariableDataInt(TEXT("r.AllowStaticLighting"));
 		const bool bAllowStaticLighting = CVarAllowStaticLighting->GetValueOnAnyThread() != 0;
 
-		return bAllowStaticLighting && Material->GetShadingModels().IsLit() && FCachedPointIndirectLightingPolicy::ShouldCompilePermutation(Platform, Material, VertexFactoryType);
+		return bAllowStaticLighting && Parameters.Material->GetShadingModels().IsLit() && FCachedPointIndirectLightingPolicy::ShouldCompilePermutation(Parameters);
 	}
 
-	static void ModifyCompilationEnvironment(EShaderPlatform Platform, const FMaterial* Material, FShaderCompilerEnvironment& OutEnvironment)
+	static void ModifyCompilationEnvironment(const FMaterialShaderPermutationParameters& Parameters, FShaderCompilerEnvironment& OutEnvironment)
 	{
-		FCachedPointIndirectLightingPolicy::ModifyCompilationEnvironment(Platform, Material, OutEnvironment);
+		FCachedPointIndirectLightingPolicy::ModifyCompilationEnvironment(Parameters, OutEnvironment);
 	}
 
 	static bool RequiresSkylight()
@@ -530,19 +532,19 @@ struct FMobileDirectionalLightAndSHIndirectPolicy
 /** Mobile Specific: Combines a movable directional light with indirect lighting from a single SH sample. */
 struct FMobileMovableDirectionalLightAndSHIndirectPolicy : public FMobileDirectionalLightAndSHIndirectPolicy
 {
-	static bool ShouldCompilePermutation(EShaderPlatform Platform, const FMaterial* Material, const FVertexFactoryType* VertexFactoryType)
+	static bool ShouldCompilePermutation(const FMeshMaterialShaderPermutationParameters& Parameters)
 	{
 		static auto* CVarMobileAllowMovableDirectionalLights = IConsoleManager::Get().FindTConsoleVariableDataInt(TEXT("r.Mobile.AllowMovableDirectionalLights"));
 		const bool bMobileAllowMovableDirectionalLights = CVarMobileAllowMovableDirectionalLights->GetValueOnAnyThread() != 0;
 
-		return bMobileAllowMovableDirectionalLights && FMobileDirectionalLightAndSHIndirectPolicy::ShouldCompilePermutation(Platform, Material, VertexFactoryType);
+		return bMobileAllowMovableDirectionalLights && FMobileDirectionalLightAndSHIndirectPolicy::ShouldCompilePermutation(Parameters);
 	}
 	
-	static void ModifyCompilationEnvironment(EShaderPlatform Platform, const FMaterial* Material, FShaderCompilerEnvironment& OutEnvironment)
+	static void ModifyCompilationEnvironment(const FMaterialShaderPermutationParameters& Parameters, FShaderCompilerEnvironment& OutEnvironment)
 	{
 		OutEnvironment.SetDefine(TEXT("MOVABLE_DIRECTIONAL_LIGHT"), TEXT("1"));
 		OutEnvironment.SetDefine(TEXT(PREPROCESSOR_TO_STRING(MAX_MOBILE_SHADOWCASCADES)), GetMobileMaxShadowCascades());
-		FMobileDirectionalLightAndSHIndirectPolicy::ModifyCompilationEnvironment(Platform, Material, OutEnvironment);
+		FMobileDirectionalLightAndSHIndirectPolicy::ModifyCompilationEnvironment(Parameters, OutEnvironment);
 	}
 
 	static bool RequiresSkylight()
@@ -554,11 +556,11 @@ struct FMobileMovableDirectionalLightAndSHIndirectPolicy : public FMobileDirecti
 /** Mobile Specific: Combines a movable directional light with CSM with indirect lighting from a single SH sample. */
 struct FMobileMovableDirectionalLightCSMAndSHIndirectPolicy : public FMobileMovableDirectionalLightAndSHIndirectPolicy
 {
-	static void ModifyCompilationEnvironment(EShaderPlatform Platform, const FMaterial* Material, FShaderCompilerEnvironment& OutEnvironment)
+	static void ModifyCompilationEnvironment(const FMaterialShaderPermutationParameters& Parameters, FShaderCompilerEnvironment& OutEnvironment)
 	{
 		OutEnvironment.SetDefine(TEXT("DIRECTIONAL_LIGHT_CSM"), TEXT("1"));
 		OutEnvironment.SetDefine(TEXT(PREPROCESSOR_TO_STRING(MAX_MOBILE_SHADOWCASCADES)), GetMobileMaxShadowCascades());
-		FMobileMovableDirectionalLightAndSHIndirectPolicy::ModifyCompilationEnvironment(Platform, Material, OutEnvironment);
+		FMobileMovableDirectionalLightAndSHIndirectPolicy::ModifyCompilationEnvironment(Parameters, OutEnvironment);
 	}
 
 	static bool RequiresSkylight()
@@ -571,11 +573,11 @@ struct FMobileMovableDirectionalLightCSMAndSHIndirectPolicy : public FMobileMova
 class FMobileDirectionalLightCSMAndSHIndirectPolicy : public FMobileDirectionalLightAndSHIndirectPolicy
 {
 public:
-	static void ModifyCompilationEnvironment(EShaderPlatform Platform, const FMaterial* Material, FShaderCompilerEnvironment& OutEnvironment)
+	static void ModifyCompilationEnvironment(const FMaterialShaderPermutationParameters& Parameters, FShaderCompilerEnvironment& OutEnvironment)
 	{
 		OutEnvironment.SetDefine(TEXT("DIRECTIONAL_LIGHT_CSM"), TEXT("1"));
 		OutEnvironment.SetDefine(TEXT(PREPROCESSOR_TO_STRING(MAX_MOBILE_SHADOWCASCADES)), GetMobileMaxShadowCascades());
-		FMobileDirectionalLightAndSHIndirectPolicy::ModifyCompilationEnvironment(Platform, Material, OutEnvironment);
+		FMobileDirectionalLightAndSHIndirectPolicy::ModifyCompilationEnvironment(Parameters, OutEnvironment);
 	}
 
 	static bool RequiresSkylight()
@@ -587,15 +589,15 @@ public:
 /** Mobile Specific */
 struct FMobileMovableDirectionalLightLightingPolicy
 {
-	static bool ShouldCompilePermutation(EShaderPlatform Platform,const FMaterial* Material,const FVertexFactoryType* VertexFactoryType)
+	static bool ShouldCompilePermutation(const FMeshMaterialShaderPermutationParameters& Parameters)
 	{
 		static auto* CVarMobileAllowMovableDirectionalLights = IConsoleManager::Get().FindTConsoleVariableDataInt(TEXT("r.Mobile.AllowMovableDirectionalLights"));
 		const bool bMobileAllowMovableDirectionalLights = CVarMobileAllowMovableDirectionalLights->GetValueOnAnyThread() != 0;
 
-		return bMobileAllowMovableDirectionalLights && Material->GetShadingModels().IsLit();
+		return bMobileAllowMovableDirectionalLights && Parameters.Material->GetShadingModels().IsLit();
 	}
 
-	static void ModifyCompilationEnvironment(EShaderPlatform Platform, const FMaterial* Material, FShaderCompilerEnvironment& OutEnvironment)
+	static void ModifyCompilationEnvironment(const FMaterialShaderPermutationParameters& Parameters, FShaderCompilerEnvironment& OutEnvironment)
 	{
 		OutEnvironment.SetDefine(TEXT("MOVABLE_DIRECTIONAL_LIGHT"),TEXT("1"));
 	}
@@ -609,21 +611,21 @@ struct FMobileMovableDirectionalLightLightingPolicy
 /** Mobile Specific */
 struct FMobileMovableDirectionalLightCSMLightingPolicy
 {
-	static bool ShouldCompilePermutation(EShaderPlatform Platform, const FMaterial* Material, const FVertexFactoryType* VertexFactoryType)
+	static bool ShouldCompilePermutation(const FMeshMaterialShaderPermutationParameters& Parameters)
 	{
 		static auto* CVarMobileAllowMovableDirectionalLights = IConsoleManager::Get().FindTConsoleVariableDataInt(TEXT("r.Mobile.AllowMovableDirectionalLights"));
 		const bool bMobileAllowMovableDirectionalLights = CVarMobileAllowMovableDirectionalLights->GetValueOnAnyThread() != 0;
 
-		return bMobileAllowMovableDirectionalLights && Material->GetShadingModels().IsLit();
+		return bMobileAllowMovableDirectionalLights && Parameters.Material->GetShadingModels().IsLit();
 	}	
 
-	static void ModifyCompilationEnvironment(EShaderPlatform Platform, const FMaterial* Material, FShaderCompilerEnvironment& OutEnvironment)
+	static void ModifyCompilationEnvironment(const FMaterialShaderPermutationParameters& Parameters, FShaderCompilerEnvironment& OutEnvironment)
 	{
 		OutEnvironment.SetDefine(TEXT("MOVABLE_DIRECTIONAL_LIGHT"), TEXT("1"));
 		OutEnvironment.SetDefine(TEXT("DIRECTIONAL_LIGHT_CSM"), TEXT("1"));
 		OutEnvironment.SetDefine(TEXT(PREPROCESSOR_TO_STRING(MAX_MOBILE_SHADOWCASCADES)), GetMobileMaxShadowCascades());
 
-		FNoLightMapPolicy::ModifyCompilationEnvironment(Platform, Material, OutEnvironment);
+		FNoLightMapPolicy::ModifyCompilationEnvironment(Parameters, OutEnvironment);
 	}
 
 	static bool RequiresSkylight()
@@ -637,7 +639,7 @@ struct FMobileMovableDirectionalLightWithLightmapPolicy : public TLightMapPolicy
 {
 	typedef TLightMapPolicy<LQ_LIGHTMAP> Super;
 
-	static bool ShouldCompilePermutation(EShaderPlatform Platform, const FMaterial* Material, const FVertexFactoryType* VertexFactoryType)
+	static bool ShouldCompilePermutation(const FMeshMaterialShaderPermutationParameters& Parameters)
 	{
 		static auto* CVarMobileAllowMovableDirectionalLights = IConsoleManager::Get().FindTConsoleVariableDataInt(TEXT("r.Mobile.AllowMovableDirectionalLights"));
 		const bool bMobileAllowMovableDirectionalLights = CVarMobileAllowMovableDirectionalLights->GetValueOnAnyThread() != 0;
@@ -645,15 +647,15 @@ struct FMobileMovableDirectionalLightWithLightmapPolicy : public TLightMapPolicy
 		static auto* CVarAllowStaticLighting = IConsoleManager::Get().FindTConsoleVariableDataInt(TEXT("r.AllowStaticLighting"));
 		const bool bAllowStaticLighting = (!CVarAllowStaticLighting || CVarAllowStaticLighting->GetValueOnAnyThread() != 0);
 
-		return bAllowStaticLighting && bMobileAllowMovableDirectionalLights && (Material->GetShadingModels().IsLit()) && Super::ShouldCompilePermutation(Platform, Material, VertexFactoryType);
+		return bAllowStaticLighting && bMobileAllowMovableDirectionalLights && (Parameters.Material->GetShadingModels().IsLit()) && Super::ShouldCompilePermutation(Parameters);
 	}
 
-	static void ModifyCompilationEnvironment(EShaderPlatform Platform, const FMaterial* Material, FShaderCompilerEnvironment& OutEnvironment)
+	static void ModifyCompilationEnvironment(const FMaterialShaderPermutationParameters& Parameters, FShaderCompilerEnvironment& OutEnvironment)
 	{
 		OutEnvironment.SetDefine(TEXT("MOVABLE_DIRECTIONAL_LIGHT"), TEXT("1"));
 		OutEnvironment.SetDefine(TEXT(PREPROCESSOR_TO_STRING(MAX_MOBILE_SHADOWCASCADES)), GetMobileMaxShadowCascades());
 
-		Super::ModifyCompilationEnvironment(Platform, Material, OutEnvironment);
+		Super::ModifyCompilationEnvironment(Parameters, OutEnvironment);
 	}
 
 	static bool RequiresSkylight()
@@ -665,11 +667,11 @@ struct FMobileMovableDirectionalLightWithLightmapPolicy : public TLightMapPolicy
 /** Mobile Specific */
 struct FMobileMovableDirectionalLightCSMWithLightmapPolicy : public FMobileMovableDirectionalLightWithLightmapPolicy
 {
-	static void ModifyCompilationEnvironment(EShaderPlatform Platform, const FMaterial* Material, FShaderCompilerEnvironment& OutEnvironment)
+	static void ModifyCompilationEnvironment(const FMaterialShaderPermutationParameters& Parameters, FShaderCompilerEnvironment& OutEnvironment)
 	{
 		OutEnvironment.SetDefine(TEXT("DIRECTIONAL_LIGHT_CSM"), TEXT("1"));
 
-		FMobileMovableDirectionalLightWithLightmapPolicy::ModifyCompilationEnvironment(Platform, Material, OutEnvironment);
+		FMobileMovableDirectionalLightWithLightmapPolicy::ModifyCompilationEnvironment(Parameters, OutEnvironment);
 	}
 
 	static bool RequiresSkylight()
@@ -742,12 +744,12 @@ public:
 	typedef FUniformLightMapPolicyShaderParametersType RayHitGroupParametersType;
 #endif
 
-	static bool ShouldCompilePermutation(EShaderPlatform Platform,const FMaterial* Material,const FVertexFactoryType* VertexFactoryType)
+	static bool ShouldCompilePermutation(const FMeshMaterialShaderPermutationParameters& Parameters)
 	{
 		return false; // This one does not compile shaders since we can't tell which policy to use.
 	}
 
-	static void ModifyCompilationEnvironment(EShaderPlatform Platform, const FMaterial* Material, FShaderCompilerEnvironment& OutEnvironment) 
+	static void ModifyCompilationEnvironment(const FMaterialShaderPermutationParameters& Parameters, FShaderCompilerEnvironment& OutEnvironment)
 	{}
 
 	FUniformLightMapPolicy(ELightMapPolicyType InIndirectPolicy) : IndirectPolicy(InIndirectPolicy) {}
@@ -792,65 +794,65 @@ public:
 
 	TUniformLightMapPolicy() : FUniformLightMapPolicy(Policy) {}
 
-	static bool ShouldCompilePermutation(EShaderPlatform Platform,const FMaterial* Material,const FVertexFactoryType* VertexFactoryType)
+	static bool ShouldCompilePermutation(const FMeshMaterialShaderPermutationParameters& Parameters)
 	{
 		CA_SUPPRESS(6326);
 		switch (Policy)
 		{
 		case LMP_NO_LIGHTMAP:
-			return FNoLightMapPolicy::ShouldCompilePermutation(Platform, Material, VertexFactoryType);
+			return FNoLightMapPolicy::ShouldCompilePermutation(Parameters);
 		case LMP_PRECOMPUTED_IRRADIANCE_VOLUME_INDIRECT_LIGHTING:
-			return FPrecomputedVolumetricLightmapLightingPolicy::ShouldCompilePermutation(Platform, Material, VertexFactoryType);
+			return FPrecomputedVolumetricLightmapLightingPolicy::ShouldCompilePermutation(Parameters);
 		case LMP_CACHED_VOLUME_INDIRECT_LIGHTING:
-			return FCachedVolumeIndirectLightingPolicy::ShouldCompilePermutation(Platform, Material, VertexFactoryType);
+			return FCachedVolumeIndirectLightingPolicy::ShouldCompilePermutation(Parameters);
 		case LMP_CACHED_POINT_INDIRECT_LIGHTING:
-			return FCachedPointIndirectLightingPolicy::ShouldCompilePermutation(Platform, Material, VertexFactoryType);
+			return FCachedPointIndirectLightingPolicy::ShouldCompilePermutation(Parameters);
 		case LMP_SIMPLE_NO_LIGHTMAP:
-			return FSimpleNoLightmapLightingPolicy::ShouldCompilePermutation(Platform, Material, VertexFactoryType);
+			return FSimpleNoLightmapLightingPolicy::ShouldCompilePermutation(Parameters);
 		case LMP_SIMPLE_LIGHTMAP_ONLY_LIGHTING:
-			return FSimpleLightmapOnlyLightingPolicy::ShouldCompilePermutation(Platform, Material, VertexFactoryType);
+			return FSimpleLightmapOnlyLightingPolicy::ShouldCompilePermutation(Parameters);
 		case LMP_SIMPLE_DIRECTIONAL_LIGHT_LIGHTING:
-			return FSimpleDirectionalLightLightingPolicy::ShouldCompilePermutation(Platform, Material, VertexFactoryType);
+			return FSimpleDirectionalLightLightingPolicy::ShouldCompilePermutation(Parameters);
 		case LMP_SIMPLE_STATIONARY_PRECOMPUTED_SHADOW_LIGHTING:
-			return FSimpleStationaryLightPrecomputedShadowsLightingPolicy::ShouldCompilePermutation(Platform, Material, VertexFactoryType);
+			return FSimpleStationaryLightPrecomputedShadowsLightingPolicy::ShouldCompilePermutation(Parameters);
 		case LMP_SIMPLE_STATIONARY_SINGLESAMPLE_SHADOW_LIGHTING:
-			return FSimpleStationaryLightSingleSampleShadowsLightingPolicy::ShouldCompilePermutation(Platform, Material, VertexFactoryType);
+			return FSimpleStationaryLightSingleSampleShadowsLightingPolicy::ShouldCompilePermutation(Parameters);
 		case LMP_SIMPLE_STATIONARY_VOLUMETRICLIGHTMAP_SHADOW_LIGHTING:
-			return FSimpleStationaryLightVolumetricLightmapShadowsLightingPolicy::ShouldCompilePermutation(Platform, Material, VertexFactoryType);
+			return FSimpleStationaryLightVolumetricLightmapShadowsLightingPolicy::ShouldCompilePermutation(Parameters);
 		case LMP_LQ_LIGHTMAP:
-			return TLightMapPolicy<LQ_LIGHTMAP>::ShouldCompilePermutation(Platform, Material, VertexFactoryType);
+			return TLightMapPolicy<LQ_LIGHTMAP>::ShouldCompilePermutation(Parameters);
 		case LMP_HQ_LIGHTMAP:
-			return TLightMapPolicy<HQ_LIGHTMAP>::ShouldCompilePermutation(Platform, Material, VertexFactoryType);
+			return TLightMapPolicy<HQ_LIGHTMAP>::ShouldCompilePermutation(Parameters);
 		case LMP_DISTANCE_FIELD_SHADOWS_AND_HQ_LIGHTMAP:
-			return TDistanceFieldShadowsAndLightMapPolicy<HQ_LIGHTMAP>::ShouldCompilePermutation(Platform, Material, VertexFactoryType);
+			return TDistanceFieldShadowsAndLightMapPolicy<HQ_LIGHTMAP>::ShouldCompilePermutation(Parameters);
 
 		// Mobile specific
 		
 		case LMP_MOBILE_DISTANCE_FIELD_SHADOWS_AND_LQ_LIGHTMAP:
-			return FMobileDistanceFieldShadowsAndLQLightMapPolicy::ShouldCompilePermutation(Platform, Material, VertexFactoryType);
+			return FMobileDistanceFieldShadowsAndLQLightMapPolicy::ShouldCompilePermutation(Parameters);
 		case LMP_MOBILE_DISTANCE_FIELD_SHADOWS_LIGHTMAP_AND_CSM:
-			return FMobileDistanceFieldShadowsLightMapAndCSMLightingPolicy::ShouldCompilePermutation(Platform, Material, VertexFactoryType);
+			return FMobileDistanceFieldShadowsLightMapAndCSMLightingPolicy::ShouldCompilePermutation(Parameters);
 		case LMP_MOBILE_DIRECTIONAL_LIGHT_AND_SH_INDIRECT:
-			return FMobileDirectionalLightAndSHIndirectPolicy::ShouldCompilePermutation(Platform, Material, VertexFactoryType);
+			return FMobileDirectionalLightAndSHIndirectPolicy::ShouldCompilePermutation(Parameters);
 		case LMP_MOBILE_MOVABLE_DIRECTIONAL_LIGHT_AND_SH_INDIRECT:
-			return FMobileMovableDirectionalLightAndSHIndirectPolicy::ShouldCompilePermutation(Platform, Material, VertexFactoryType);
+			return FMobileMovableDirectionalLightAndSHIndirectPolicy::ShouldCompilePermutation(Parameters);
 		case LMP_MOBILE_DIRECTIONAL_LIGHT_CSM_AND_SH_INDIRECT:
-			return FMobileDirectionalLightCSMAndSHIndirectPolicy::ShouldCompilePermutation(Platform, Material, VertexFactoryType);
+			return FMobileDirectionalLightCSMAndSHIndirectPolicy::ShouldCompilePermutation(Parameters);
 		case LMP_MOBILE_MOVABLE_DIRECTIONAL_LIGHT_CSM_AND_SH_INDIRECT:
-			return FMobileMovableDirectionalLightCSMAndSHIndirectPolicy::ShouldCompilePermutation(Platform, Material, VertexFactoryType);
+			return FMobileMovableDirectionalLightCSMAndSHIndirectPolicy::ShouldCompilePermutation(Parameters);
 		case LMP_MOBILE_MOVABLE_DIRECTIONAL_LIGHT:
-			return FMobileMovableDirectionalLightLightingPolicy::ShouldCompilePermutation(Platform, Material, VertexFactoryType);
+			return FMobileMovableDirectionalLightLightingPolicy::ShouldCompilePermutation(Parameters);
 		case LMP_MOBILE_MOVABLE_DIRECTIONAL_LIGHT_CSM:
-			return FMobileMovableDirectionalLightCSMLightingPolicy::ShouldCompilePermutation(Platform, Material, VertexFactoryType);
+			return FMobileMovableDirectionalLightCSMLightingPolicy::ShouldCompilePermutation(Parameters);
 		case LMP_MOBILE_MOVABLE_DIRECTIONAL_LIGHT_WITH_LIGHTMAP:
-			return FMobileMovableDirectionalLightWithLightmapPolicy::ShouldCompilePermutation(Platform, Material, VertexFactoryType);
+			return FMobileMovableDirectionalLightWithLightmapPolicy::ShouldCompilePermutation(Parameters);
 		case LMP_MOBILE_MOVABLE_DIRECTIONAL_LIGHT_CSM_WITH_LIGHTMAP:
-			return FMobileMovableDirectionalLightCSMWithLightmapPolicy::ShouldCompilePermutation(Platform, Material, VertexFactoryType);
+			return FMobileMovableDirectionalLightCSMWithLightmapPolicy::ShouldCompilePermutation(Parameters);
 
 		// LightMapDensity
 	
 		case LMP_DUMMY:
-			return FDummyLightMapPolicy::ShouldCompilePermutation(Platform, Material, VertexFactoryType);
+			return FDummyLightMapPolicy::ShouldCompilePermutation(Parameters);
 
 		default:
 			check(false);
@@ -858,7 +860,7 @@ public:
 		};
 	}
 
-	static void ModifyCompilationEnvironment(EShaderPlatform Platform, const FMaterial* Material, FShaderCompilerEnvironment& OutEnvironment) 
+	static void ModifyCompilationEnvironment(const FMaterialShaderPermutationParameters& Parameters, FShaderCompilerEnvironment& OutEnvironment)
 	{
 		OutEnvironment.SetDefine(TEXT("MAX_NUM_LIGHTMAP_COEF"), MAX_NUM_LIGHTMAP_COEF);
 
@@ -866,81 +868,81 @@ public:
 		switch (Policy)
 		{
 		case LMP_NO_LIGHTMAP:							
-			FNoLightMapPolicy::ModifyCompilationEnvironment(Platform, Material, OutEnvironment);
+			FNoLightMapPolicy::ModifyCompilationEnvironment(Parameters, OutEnvironment);
 			break;
 		case LMP_PRECOMPUTED_IRRADIANCE_VOLUME_INDIRECT_LIGHTING:
-			FPrecomputedVolumetricLightmapLightingPolicy::ModifyCompilationEnvironment(Platform, Material, OutEnvironment);
+			FPrecomputedVolumetricLightmapLightingPolicy::ModifyCompilationEnvironment(Parameters, OutEnvironment);
 			break;
 		case LMP_CACHED_VOLUME_INDIRECT_LIGHTING:
-			FCachedVolumeIndirectLightingPolicy::ModifyCompilationEnvironment(Platform, Material, OutEnvironment);
+			FCachedVolumeIndirectLightingPolicy::ModifyCompilationEnvironment(Parameters, OutEnvironment);
 			break;
 		case LMP_CACHED_POINT_INDIRECT_LIGHTING:
-			FCachedPointIndirectLightingPolicy::ModifyCompilationEnvironment(Platform, Material, OutEnvironment);
+			FCachedPointIndirectLightingPolicy::ModifyCompilationEnvironment(Parameters, OutEnvironment);
 			break;
 		case LMP_SIMPLE_NO_LIGHTMAP:
-			return FSimpleNoLightmapLightingPolicy::ModifyCompilationEnvironment(Platform, Material, OutEnvironment);
+			return FSimpleNoLightmapLightingPolicy::ModifyCompilationEnvironment(Parameters, OutEnvironment);
 			break;
 		case LMP_SIMPLE_LIGHTMAP_ONLY_LIGHTING:
-			return FSimpleLightmapOnlyLightingPolicy::ModifyCompilationEnvironment(Platform, Material, OutEnvironment);
+			return FSimpleLightmapOnlyLightingPolicy::ModifyCompilationEnvironment(Parameters, OutEnvironment);
 			break;
 		case LMP_SIMPLE_DIRECTIONAL_LIGHT_LIGHTING:
-			FSimpleDirectionalLightLightingPolicy::ModifyCompilationEnvironment(Platform, Material, OutEnvironment);
+			FSimpleDirectionalLightLightingPolicy::ModifyCompilationEnvironment(Parameters, OutEnvironment);
 			break;
 		case LMP_SIMPLE_STATIONARY_PRECOMPUTED_SHADOW_LIGHTING:
-			return FSimpleStationaryLightPrecomputedShadowsLightingPolicy::ModifyCompilationEnvironment(Platform, Material, OutEnvironment);
+			return FSimpleStationaryLightPrecomputedShadowsLightingPolicy::ModifyCompilationEnvironment(Parameters, OutEnvironment);
 			break;
 		case LMP_SIMPLE_STATIONARY_SINGLESAMPLE_SHADOW_LIGHTING:
-			return FSimpleStationaryLightSingleSampleShadowsLightingPolicy::ModifyCompilationEnvironment(Platform, Material, OutEnvironment);
+			return FSimpleStationaryLightSingleSampleShadowsLightingPolicy::ModifyCompilationEnvironment(Parameters, OutEnvironment);
 			break;
 		case LMP_SIMPLE_STATIONARY_VOLUMETRICLIGHTMAP_SHADOW_LIGHTING:
-			return FSimpleStationaryLightVolumetricLightmapShadowsLightingPolicy::ModifyCompilationEnvironment(Platform, Material, OutEnvironment);
+			return FSimpleStationaryLightVolumetricLightmapShadowsLightingPolicy::ModifyCompilationEnvironment(Parameters, OutEnvironment);
 			break;
 		case LMP_LQ_LIGHTMAP:
-			TLightMapPolicy<LQ_LIGHTMAP>::ModifyCompilationEnvironment(Platform, Material, OutEnvironment);
+			TLightMapPolicy<LQ_LIGHTMAP>::ModifyCompilationEnvironment(Parameters, OutEnvironment);
 			break;
 		case LMP_HQ_LIGHTMAP:
-			TLightMapPolicy<HQ_LIGHTMAP>::ModifyCompilationEnvironment(Platform, Material, OutEnvironment);
+			TLightMapPolicy<HQ_LIGHTMAP>::ModifyCompilationEnvironment(Parameters, OutEnvironment);
 			break;
 		case LMP_DISTANCE_FIELD_SHADOWS_AND_HQ_LIGHTMAP:
-			TDistanceFieldShadowsAndLightMapPolicy<HQ_LIGHTMAP>::ModifyCompilationEnvironment(Platform, Material, OutEnvironment);
+			TDistanceFieldShadowsAndLightMapPolicy<HQ_LIGHTMAP>::ModifyCompilationEnvironment(Parameters, OutEnvironment);
 			break;
 
 		// Mobile specific
 		case LMP_MOBILE_DISTANCE_FIELD_SHADOWS_AND_LQ_LIGHTMAP:
-			FMobileDistanceFieldShadowsAndLQLightMapPolicy::ModifyCompilationEnvironment(Platform, Material, OutEnvironment);
+			FMobileDistanceFieldShadowsAndLQLightMapPolicy::ModifyCompilationEnvironment(Parameters, OutEnvironment);
 			break;
 		case LMP_MOBILE_DISTANCE_FIELD_SHADOWS_LIGHTMAP_AND_CSM:
-			FMobileDistanceFieldShadowsLightMapAndCSMLightingPolicy::ModifyCompilationEnvironment(Platform, Material, OutEnvironment);
+			FMobileDistanceFieldShadowsLightMapAndCSMLightingPolicy::ModifyCompilationEnvironment(Parameters, OutEnvironment);
 			break;
 		case LMP_MOBILE_DIRECTIONAL_LIGHT_AND_SH_INDIRECT:
-			FMobileDirectionalLightAndSHIndirectPolicy::ModifyCompilationEnvironment(Platform, Material, OutEnvironment);
+			FMobileDirectionalLightAndSHIndirectPolicy::ModifyCompilationEnvironment(Parameters, OutEnvironment);
 			break;
 		case LMP_MOBILE_MOVABLE_DIRECTIONAL_LIGHT_AND_SH_INDIRECT:
-			FMobileMovableDirectionalLightAndSHIndirectPolicy::ModifyCompilationEnvironment(Platform, Material, OutEnvironment);
+			FMobileMovableDirectionalLightAndSHIndirectPolicy::ModifyCompilationEnvironment(Parameters, OutEnvironment);
 			break;
 		case LMP_MOBILE_DIRECTIONAL_LIGHT_CSM_AND_SH_INDIRECT:
-			FMobileDirectionalLightCSMAndSHIndirectPolicy::ModifyCompilationEnvironment(Platform, Material, OutEnvironment);
+			FMobileDirectionalLightCSMAndSHIndirectPolicy::ModifyCompilationEnvironment(Parameters, OutEnvironment);
 			break;
 		case LMP_MOBILE_MOVABLE_DIRECTIONAL_LIGHT_CSM_AND_SH_INDIRECT:
-			FMobileMovableDirectionalLightCSMAndSHIndirectPolicy::ModifyCompilationEnvironment(Platform, Material, OutEnvironment);
+			FMobileMovableDirectionalLightCSMAndSHIndirectPolicy::ModifyCompilationEnvironment(Parameters, OutEnvironment);
 			break;
 		case LMP_MOBILE_MOVABLE_DIRECTIONAL_LIGHT:
-			FMobileMovableDirectionalLightLightingPolicy::ModifyCompilationEnvironment(Platform, Material, OutEnvironment);
+			FMobileMovableDirectionalLightLightingPolicy::ModifyCompilationEnvironment(Parameters, OutEnvironment);
 			break;
 		case LMP_MOBILE_MOVABLE_DIRECTIONAL_LIGHT_CSM:
-			FMobileMovableDirectionalLightCSMLightingPolicy::ModifyCompilationEnvironment(Platform, Material, OutEnvironment);
+			FMobileMovableDirectionalLightCSMLightingPolicy::ModifyCompilationEnvironment(Parameters, OutEnvironment);
 			break;
 		case LMP_MOBILE_MOVABLE_DIRECTIONAL_LIGHT_WITH_LIGHTMAP:
-			FMobileMovableDirectionalLightWithLightmapPolicy::ModifyCompilationEnvironment(Platform, Material, OutEnvironment);
+			FMobileMovableDirectionalLightWithLightmapPolicy::ModifyCompilationEnvironment(Parameters, OutEnvironment);
 			break;
 		case LMP_MOBILE_MOVABLE_DIRECTIONAL_LIGHT_CSM_WITH_LIGHTMAP:
-			FMobileMovableDirectionalLightCSMWithLightmapPolicy::ModifyCompilationEnvironment(Platform, Material, OutEnvironment);
+			FMobileMovableDirectionalLightCSMWithLightmapPolicy::ModifyCompilationEnvironment(Parameters, OutEnvironment);
 			break;
 
 		// LightMapDensity
 	
 		case LMP_DUMMY:
-			FDummyLightMapPolicy::ModifyCompilationEnvironment(Platform, Material, OutEnvironment);
+			FDummyLightMapPolicy::ModifyCompilationEnvironment(Parameters, OutEnvironment);
 			break;
 
 		default:
@@ -997,17 +999,17 @@ public:
 		}
 	};
 
-	static bool ShouldCompilePermutation(EShaderPlatform Platform,const FMaterial* Material,const FVertexFactoryType* VertexFactoryType)
+	static bool ShouldCompilePermutation(const FMeshMaterialShaderPermutationParameters& Parameters)
 	{
 		static IConsoleVariable* AllowStaticLightingVar = IConsoleManager::Get().FindConsoleVariable(TEXT("r.AllowStaticLighting"));
 
-		return Material->GetShadingModels().IsLit()
-			&& IsTranslucentBlendMode(Material->GetBlendMode()) 
+		return Parameters.Material->GetShadingModels().IsLit()
+			&& IsTranslucentBlendMode(Parameters.Material->GetBlendMode())
 			&& (!AllowStaticLightingVar || AllowStaticLightingVar->GetInt() != 0)
-			&& FSelfShadowedTranslucencyPolicy::ShouldCompilePermutation(Platform, Material, VertexFactoryType);
+			&& FSelfShadowedTranslucencyPolicy::ShouldCompilePermutation(Parameters);
 	}
 
-	static void ModifyCompilationEnvironment(EShaderPlatform Platform, const FMaterial* Material, FShaderCompilerEnvironment& OutEnvironment);
+	static void ModifyCompilationEnvironment(const FMaterialShaderPermutationParameters& Parameters, FShaderCompilerEnvironment& OutEnvironment);
 
 	/** Initialization constructor. */
 	FSelfShadowedCachedPointIndirectLightingPolicy() {}
@@ -1046,20 +1048,20 @@ public:
 		}
 	};
 
-	static bool ShouldCompilePermutation(EShaderPlatform Platform,const FMaterial* Material,const FVertexFactoryType* VertexFactoryType)
+	static bool ShouldCompilePermutation(const FMeshMaterialShaderPermutationParameters& Parameters)
 	{
 		static IConsoleVariable* AllowStaticLightingVar = IConsoleManager::Get().FindConsoleVariable(TEXT("r.AllowStaticLighting"));
 
-		return Material->GetShadingModels().IsLit()
-			&& IsTranslucentBlendMode(Material->GetBlendMode()) 
+		return Parameters.Material->GetShadingModels().IsLit()
+			&& IsTranslucentBlendMode(Parameters.Material->GetBlendMode())
 			&& (!AllowStaticLightingVar || AllowStaticLightingVar->GetInt() != 0)
-			&& FSelfShadowedTranslucencyPolicy::ShouldCompilePermutation(Platform, Material, VertexFactoryType);
+			&& FSelfShadowedTranslucencyPolicy::ShouldCompilePermutation(Parameters);
 	}
 
-	static void ModifyCompilationEnvironment(EShaderPlatform Platform, const FMaterial* Material, FShaderCompilerEnvironment& OutEnvironment)
+	static void ModifyCompilationEnvironment(const FMaterialShaderPermutationParameters& Parameters, FShaderCompilerEnvironment& OutEnvironment)
 	{
 		OutEnvironment.SetDefine(TEXT("PRECOMPUTED_IRRADIANCE_VOLUME_LIGHTING"),TEXT("1"));	
-		FSelfShadowedTranslucencyPolicy::ModifyCompilationEnvironment(Platform, Material, OutEnvironment);
+		FSelfShadowedTranslucencyPolicy::ModifyCompilationEnvironment(Parameters, OutEnvironment);
 	}
 
 	/** Initialization constructor. */
diff --git a/Engine/Source/Runtime/Renderer/Private/MaterialTexCoordScalesRendering.h b/Engine/Source/Runtime/Renderer/Private/MaterialTexCoordScalesRendering.h
--- a/Engine/Source/Runtime/Renderer/Private/MaterialTexCoordScalesRendering.h
+++ b/Engine/Source/Runtime/Renderer/Private/MaterialTexCoordScalesRendering.h
@@ -31,10 +31,10 @@ class FMaterialTexCoordScalePS : public FDebugViewModePS
 
 public:
 
-	static bool ShouldCompilePermutation(EShaderPlatform Platform, const FMaterial* Material, const FVertexFactoryType* VertexFactoryType)
+	static bool ShouldCompilePermutation(const FMeshMaterialShaderPermutationParameters& Parameters)
 	{
 		// See FDebugViewModeMaterialProxy::GetFriendlyName()
-		return AllowDebugViewShaderMode(DVSM_OutputMaterialTextureScales, Platform, GetMaxSupportedFeatureLevel(Platform)) && Material->GetFriendlyName().Contains(TEXT("MaterialTexCoordScale"));
+		return AllowDebugViewShaderMode(DVSM_OutputMaterialTextureScales, Parameters.Platform, GetMaxSupportedFeatureLevel(Parameters.Platform)) && Parameters.Material->GetFriendlyName().Contains(TEXT("MaterialTexCoordScale"));
 	}
 
 	FMaterialTexCoordScalePS(const ShaderMetaType::CompiledShaderInitializerType& Initializer):
@@ -58,14 +58,14 @@ public:
 		return bShaderHasOutdatedParameters;
 	}
 
-	static void ModifyCompilationEnvironment(EShaderPlatform Platform, const FMaterial* Material, FShaderCompilerEnvironment& OutEnvironment)
+	static void ModifyCompilationEnvironment(const FMaterialShaderPermutationParameters& Parameters, FShaderCompilerEnvironment& OutEnvironment)
 	{
 		OutEnvironment.SetDefine(TEXT("UNDEFINED_ACCURACY"), UndefinedStreamingAccuracyIntensity);
 		OutEnvironment.SetDefine(TEXT("MAX_NUM_TEX_COORD"), (uint32)TEXSTREAM_MAX_NUM_UVCHANNELS);
 		OutEnvironment.SetDefine(TEXT("INITIAL_GPU_SCALE"), (uint32)TEXSTREAM_INITIAL_GPU_SCALE);
 		OutEnvironment.SetDefine(TEXT("TILE_RESOLUTION"), (uint32)TEXSTREAM_TILE_RESOLUTION);
 		OutEnvironment.SetDefine(TEXT("MAX_NUM_TEXTURE_REGISTER"), (uint32)TEXSTREAM_MAX_NUM_TEXTURES_PER_MATERIAL);
-		FMeshMaterialShader::ModifyCompilationEnvironment(Platform, OutEnvironment);
+		FMeshMaterialShader::ModifyCompilationEnvironment(Parameters.Platform, OutEnvironment);
 	}
 
 	virtual void GetDebugViewModeShaderBindings(
diff --git a/Engine/Source/Runtime/Renderer/Private/MeshTexCoordSizeAccuracyRendering.h b/Engine/Source/Runtime/Renderer/Private/MeshTexCoordSizeAccuracyRendering.h
--- a/Engine/Source/Runtime/Renderer/Private/MeshTexCoordSizeAccuracyRendering.h
+++ b/Engine/Source/Runtime/Renderer/Private/MeshTexCoordSizeAccuracyRendering.h
@@ -28,10 +28,10 @@ class FMeshTexCoordSizeAccuracyPS : public FDebugViewModePS
 
 public:
 
-	static bool ShouldCompilePermutation(EShaderPlatform Platform, const FMaterial* Material, const FVertexFactoryType* VertexFactoryType)
+	static bool ShouldCompilePermutation(const FMeshMaterialShaderPermutationParameters& Parameters)
 	{
 		// See FDebugViewModeMaterialProxy::GetFriendlyName()
-		return AllowDebugViewShaderMode(DVSM_MeshUVDensityAccuracy, Platform, GetMaxSupportedFeatureLevel(Platform)) && Material->GetFriendlyName().Contains(TEXT("MeshTexCoordSizeAccuracy"));
+		return AllowDebugViewShaderMode(DVSM_MeshUVDensityAccuracy, Parameters.Platform, GetMaxSupportedFeatureLevel(Parameters.Platform)) && Parameters.Material->GetFriendlyName().Contains(TEXT("MeshTexCoordSizeAccuracy"));
 	}
 
 	FMeshTexCoordSizeAccuracyPS(const ShaderMetaType::CompiledShaderInitializerType& Initializer):
@@ -53,7 +53,7 @@ public:
 		return bShaderHasOutdatedParameters;
 	}
 
-	static void ModifyCompilationEnvironment(EShaderPlatform Platform, const FMaterial* Material, FShaderCompilerEnvironment& OutEnvironment)
+	static void ModifyCompilationEnvironment(const FMaterialShaderPermutationParameters& Parameters, FShaderCompilerEnvironment& OutEnvironment)
 	{
 		OutEnvironment.SetDefine(TEXT("UNDEFINED_ACCURACY"), UndefinedStreamingAccuracyIntensity);
 	}
diff --git a/Engine/Source/Runtime/Renderer/Private/MobileBasePassRendering.h b/Engine/Source/Runtime/Renderer/Private/MobileBasePassRendering.h
--- a/Engine/Source/Runtime/Renderer/Private/MobileBasePassRendering.h
+++ b/Engine/Source/Runtime/Renderer/Private/MobileBasePassRendering.h
@@ -121,9 +121,9 @@ public:
 
 	// static bool ShouldCompilePermutation(EShaderPlatform Platform,const FMaterial* Material,const FVertexFactoryType* VertexFactoryType)
 
-	static void ModifyCompilationEnvironment(EShaderPlatform Platform, const FMaterial* Material, FShaderCompilerEnvironment& OutEnvironment)
+	static void ModifyCompilationEnvironment(const FMaterialShaderPermutationParameters& Parameters, FShaderCompilerEnvironment& OutEnvironment)
 	{
-		FMeshMaterialShader::ModifyCompilationEnvironment(Platform, Material, OutEnvironment);
+		FMeshMaterialShader::ModifyCompilationEnvironment(Parameters, OutEnvironment);
 	}
 
 	virtual bool Serialize(FArchive& Ar) override
@@ -165,15 +165,15 @@ protected:
 
 public:
 
-	static bool ShouldCompilePermutation(EShaderPlatform Platform,const FMaterial* Material,const FVertexFactoryType* VertexFactoryType)
+	static bool ShouldCompilePermutation(const FMeshMaterialShaderPermutationParameters& Parameters)
 	{
-		return IsMobilePlatform(Platform) && LightMapPolicyType::ShouldCompilePermutation(Platform,Material,VertexFactoryType);
+		return IsMobilePlatform(Parameters.Platform) && LightMapPolicyType::ShouldCompilePermutation(Parameters);
 	}
 
-	static void ModifyCompilationEnvironment(EShaderPlatform Platform, const FMaterial* Material, FShaderCompilerEnvironment& OutEnvironment)
+	static void ModifyCompilationEnvironment(const FMaterialShaderPermutationParameters& Parameters, FShaderCompilerEnvironment& OutEnvironment)
 	{
-		LightMapPolicyType::ModifyCompilationEnvironment(Platform, Material, OutEnvironment);
-		Super::ModifyCompilationEnvironment(Platform, Material, OutEnvironment);
+		LightMapPolicyType::ModifyCompilationEnvironment(Parameters, OutEnvironment);
+		Super::ModifyCompilationEnvironment(Parameters, OutEnvironment);
 	}
 };
 
@@ -183,17 +183,17 @@ class TMobileBasePassVS : public TMobileBasePassVSBaseType<LightMapPolicyType>
 	DECLARE_SHADER_TYPE(TMobileBasePassVS,MeshMaterial);
 public:
 	
-	static bool ShouldCompilePermutation(EShaderPlatform Platform, const FMaterial* Material, const FVertexFactoryType* VertexFactoryType)
+	static bool ShouldCompilePermutation(const FMeshMaterialShaderPermutationParameters& Parameters)
 	{		
-		return TMobileBasePassVSBaseType<LightMapPolicyType>::ShouldCompilePermutation(Platform, Material, VertexFactoryType) && ShouldCacheShaderByPlatformAndOutputFormat(Platform,OutputFormat);
+		return TMobileBasePassVSBaseType<LightMapPolicyType>::ShouldCompilePermutation(Parameters) && ShouldCacheShaderByPlatformAndOutputFormat(Parameters.Platform,OutputFormat);
 	}
 
-	static void ModifyCompilationEnvironment(EShaderPlatform Platform, const FMaterial* Material, FShaderCompilerEnvironment& OutEnvironment)
+	static void ModifyCompilationEnvironment(const FMaterialShaderPermutationParameters& Parameters, FShaderCompilerEnvironment& OutEnvironment)
 	{
 		static auto* MobileUseHWsRGBEncodingCVAR = IConsoleManager::Get().FindTConsoleVariableDataInt(TEXT("r.Mobile.UseHWsRGBEncoding"));
 		const bool bMobileUseHWsRGBEncoding = (MobileUseHWsRGBEncodingCVAR && MobileUseHWsRGBEncodingCVAR->GetValueOnAnyThread() == 1);
 
-		TMobileBasePassVSBaseType<LightMapPolicyType>::ModifyCompilationEnvironment(Platform, Material, OutEnvironment);
+		TMobileBasePassVSBaseType<LightMapPolicyType>::ModifyCompilationEnvironment(Parameters, OutEnvironment);
 		OutEnvironment.SetDefine( TEXT("OUTPUT_GAMMA_SPACE"), OutputFormat == LDR_GAMMA_32 && !bMobileUseHWsRGBEncoding);
 		OutEnvironment.SetDefine( TEXT("OUTPUT_MOBILE_HDR"), OutputFormat == HDR_LINEAR_64 ? 1u : 0u);
 	}
@@ -216,16 +216,16 @@ class TMobileBasePassPSPolicyParamType : public FMeshMaterialShader, public Ligh
 {
 public:
 
-	static bool ShouldCompilePermutation(EShaderPlatform Platform,const FMaterial* Material,const FVertexFactoryType* VertexFactoryType)
+	static bool ShouldCompilePermutation(const FMeshMaterialShaderPermutationParameters& Parameters)
 	{
-		return IsMobilePlatform(Platform);
+		return IsMobilePlatform(Parameters.Platform);
 	}
 
-	static void ModifyCompilationEnvironment(EShaderPlatform Platform, const FMaterial* Material, FShaderCompilerEnvironment& OutEnvironment)
+	static void ModifyCompilationEnvironment(const FMaterialShaderPermutationParameters& Parameters, FShaderCompilerEnvironment& OutEnvironment)
 	{
-		FMeshMaterialShader::ModifyCompilationEnvironment(Platform, Material, OutEnvironment);
+		FMeshMaterialShader::ModifyCompilationEnvironment(Parameters, OutEnvironment);
 		// Modify compilation environment depending upon material shader quality level settings.
-		ModifyCompilationEnvironmentForQualityLevel(Platform, Material->GetQualityLevel(), OutEnvironment);
+		ModifyCompilationEnvironmentForQualityLevel(Parameters.Platform, Parameters.Material->GetQualityLevel(), OutEnvironment);
 	}
 
 	/** Initialization constructor. */
@@ -326,16 +326,16 @@ class TMobileBasePassPSBaseType : public TMobileBasePassPSPolicyParamType<LightM
 
 public:
 
-	static bool ShouldCompilePermutation(EShaderPlatform Platform,const FMaterial* Material,const FVertexFactoryType* VertexFactoryType)
+	static bool ShouldCompilePermutation(const FMeshMaterialShaderPermutationParameters& Parameters)
 	{
-		return LightMapPolicyType::ShouldCompilePermutation(Platform, Material, VertexFactoryType) 
-			&& Super::ShouldCompilePermutation(Platform, Material, VertexFactoryType);
+		return LightMapPolicyType::ShouldCompilePermutation(Parameters)
+			&& Super::ShouldCompilePermutation(Parameters);
 	}
 
-	static void ModifyCompilationEnvironment(EShaderPlatform Platform, const FMaterial* Material, FShaderCompilerEnvironment& OutEnvironment)
+	static void ModifyCompilationEnvironment(const FMaterialShaderPermutationParameters& Parameters, FShaderCompilerEnvironment& OutEnvironment)
 	{
-		LightMapPolicyType::ModifyCompilationEnvironment(Platform, Material, OutEnvironment);
-		Super::ModifyCompilationEnvironment(Platform, Material, OutEnvironment);
+		LightMapPolicyType::ModifyCompilationEnvironment(Parameters, OutEnvironment);
+		Super::ModifyCompilationEnvironment(Parameters, OutEnvironment);
 	}
 
 	/** Initialization constructor. */
@@ -392,7 +392,7 @@ class TMobileBasePassPS : public TMobileBasePassPSBaseType<LightMapPolicyType>
 	DECLARE_SHADER_TYPE(TMobileBasePassPS,MeshMaterial);
 public:
 
-	static bool ShouldCompilePermutation(EShaderPlatform Platform, const FMaterial* Material, const FVertexFactoryType* VertexFactoryType)
+	static bool ShouldCompilePermutation(const FMeshMaterialShaderPermutationParameters& Parameters)
 	{		
 		// We compile the point light shader combinations based on the project settings
 		static auto* MobileDynamicPointLightsUseStaticBranchCVar = IConsoleManager::Get().FindTConsoleVariableDataInt(TEXT("r.MobileDynamicPointLightsUseStaticBranch"));
@@ -401,7 +401,7 @@ public:
 		const bool bMobileDynamicPointLightsUseStaticBranch = (MobileDynamicPointLightsUseStaticBranchCVar->GetValueOnAnyThread() == 1);
 		const int32 MobileNumDynamicPointLights = MobileNumDynamicPointLightsCVar->GetValueOnAnyThread();
 		const int32 MobileSkyLightPermutationOptions = MobileSkyLightPermutationCVar->GetValueOnAnyThread();
-		const bool bIsUnlit = Material->GetShadingModels().IsUnlit();
+		const bool bIsUnlit = Parameters.Material->GetShadingModels().IsUnlit();
 
 		// Only compile skylight version for lit materials on ES2 (Metal) or higher
 		const bool bShouldCacheBySkylight = !bEnableSkyLight || !bIsUnlit;
@@ -417,15 +417,15 @@ public:
 			(!bIsUnlit && NumMovablePointLights == INT32_MAX && bMobileDynamicPointLightsUseStaticBranch && MobileNumDynamicPointLights > 0) ||	// single shader for variable number of point lights
 				(!bIsUnlit && NumMovablePointLights <= MobileNumDynamicPointLights && !bMobileDynamicPointLightsUseStaticBranch));				// unique 1...N point light shaders
 
-		return TMobileBasePassPSBaseType<LightMapPolicyType>::ShouldCompilePermutation(Platform, Material, VertexFactoryType) && ShouldCacheShaderByPlatformAndOutputFormat(Platform, OutputFormat) && bShouldCacheBySkylight && bShouldCacheByNumDynamicPointLights;
+		return TMobileBasePassPSBaseType<LightMapPolicyType>::ShouldCompilePermutation(Parameters) && ShouldCacheShaderByPlatformAndOutputFormat(Parameters.Platform, OutputFormat) && bShouldCacheBySkylight && bShouldCacheByNumDynamicPointLights;
 	}
 	
-	static void ModifyCompilationEnvironment(EShaderPlatform Platform, const FMaterial* Material, FShaderCompilerEnvironment& OutEnvironment)
+	static void ModifyCompilationEnvironment(const FMaterialShaderPermutationParameters& Parameters, FShaderCompilerEnvironment& OutEnvironment)
 	{		
 		static auto* MobileUseHWsRGBEncodingCVAR = IConsoleManager::Get().FindTConsoleVariableDataInt(TEXT("r.Mobile.UseHWsRGBEncoding"));
 		const bool bMobileUseHWsRGBEncoding = (MobileUseHWsRGBEncodingCVAR && MobileUseHWsRGBEncodingCVAR->GetValueOnAnyThread() == 1);
 
-		TMobileBasePassPSBaseType<LightMapPolicyType>::ModifyCompilationEnvironment(Platform, Material, OutEnvironment);
+		TMobileBasePassPSBaseType<LightMapPolicyType>::ModifyCompilationEnvironment(Parameters, OutEnvironment);
 		OutEnvironment.SetDefine(TEXT("ENABLE_SKY_LIGHT"), bEnableSkyLight);
 		OutEnvironment.SetDefine(TEXT("OUTPUT_GAMMA_SPACE"), OutputFormat == LDR_GAMMA_32 && !bMobileUseHWsRGBEncoding);
 		OutEnvironment.SetDefine(TEXT("OUTPUT_MOBILE_HDR"), OutputFormat == HDR_LINEAR_64 ? 1u : 0u);
diff --git a/Engine/Source/Runtime/Renderer/Private/MobileTranslucentRendering.cpp b/Engine/Source/Runtime/Renderer/Private/MobileTranslucentRendering.cpp
--- a/Engine/Source/Runtime/Renderer/Private/MobileTranslucentRendering.cpp
+++ b/Engine/Source/Runtime/Renderer/Private/MobileTranslucentRendering.cpp
@@ -170,17 +170,17 @@ protected:
 
 public:
 
-	static bool ShouldCompilePermutation(EShaderPlatform Platform, const FMaterial* Material, const FVertexFactoryType* VertexFactoryType)
+	static bool ShouldCompilePermutation(const FMeshMaterialShaderPermutationParameters& Parameters)
 	{
-		return IsTranslucentBlendMode(Material->GetBlendMode()) && IsMobilePlatform(Platform);
+		return IsTranslucentBlendMode(Parameters.Material->GetBlendMode()) && IsMobilePlatform(Parameters.Platform);
 	}
 
-	static void ModifyCompilationEnvironment(EShaderPlatform Platform, const FMaterial* Material, FShaderCompilerEnvironment& OutEnvironment)
+	static void ModifyCompilationEnvironment(const FMaterialShaderPermutationParameters& Parameters, FShaderCompilerEnvironment& OutEnvironment)
 	{
 		static auto* MobileUseHWsRGBEncodingCVAR = IConsoleManager::Get().FindTConsoleVariableDataInt(TEXT("r.Mobile.UseHWsRGBEncoding"));
 		const bool bMobileUseHWsRGBEncoding = (MobileUseHWsRGBEncodingCVAR && MobileUseHWsRGBEncodingCVAR->GetValueOnAnyThread() == 1);
 
-		FMeshMaterialShader::ModifyCompilationEnvironment(Platform, Material, OutEnvironment);
+		FMeshMaterialShader::ModifyCompilationEnvironment(Parameters, OutEnvironment);
 		OutEnvironment.SetDefine(TEXT("OUTPUT_GAMMA_SPACE"), IsMobileHDR() == false && !bMobileUseHWsRGBEncoding);
 		OutEnvironment.SetDefine(TEXT("OUTPUT_MOBILE_HDR"), IsMobileHDR() == true);
 	}
@@ -202,14 +202,14 @@ class FOpacityOnlyPS : public FMeshMaterialShader
 	DECLARE_SHADER_TYPE(FOpacityOnlyPS, MeshMaterial);
 public:
 
-	static bool ShouldCompilePermutation(EShaderPlatform Platform, const FMaterial* Material, const FVertexFactoryType* VertexFactoryType)
+	static bool ShouldCompilePermutation(const FMeshMaterialShaderPermutationParameters& Parameters)
 	{
-		return IsTranslucentBlendMode(Material->GetBlendMode()) && IsMobilePlatform(Platform);
+		return IsTranslucentBlendMode(Parameters.Material->GetBlendMode()) && IsMobilePlatform(Parameters.Platform);
 	}
 
-	static void ModifyCompilationEnvironment(EShaderPlatform Platform, const FMaterial* Material, FShaderCompilerEnvironment& OutEnvironment)
+	static void ModifyCompilationEnvironment(const FMaterialShaderPermutationParameters& Parameters, FShaderCompilerEnvironment& OutEnvironment)
 	{
-		FMeshMaterialShader::ModifyCompilationEnvironment(Platform, Material, OutEnvironment);
+		FMeshMaterialShader::ModifyCompilationEnvironment(Parameters, OutEnvironment);
 		OutEnvironment.SetDefine(TEXT("MOBILE_FORCE_DEPTH_TEXTURE_READS"), 1u);
 	}
 
diff --git a/Engine/Source/Runtime/Renderer/Private/PrimitiveDistanceAccuracyRendering.h b/Engine/Source/Runtime/Renderer/Private/PrimitiveDistanceAccuracyRendering.h
--- a/Engine/Source/Runtime/Renderer/Private/PrimitiveDistanceAccuracyRendering.h
+++ b/Engine/Source/Runtime/Renderer/Private/PrimitiveDistanceAccuracyRendering.h
@@ -28,10 +28,10 @@ class FPrimitiveDistanceAccuracyPS : public FDebugViewModePS
 
 public:
 
-	static bool ShouldCompilePermutation(EShaderPlatform Platform, const FMaterial* Material, const FVertexFactoryType* VertexFactoryType)
+	static bool ShouldCompilePermutation(const FMeshMaterialShaderPermutationParameters& Parameters)
 	{
 		// See FDebugViewModeMaterialProxy::GetFriendlyName()
-		return AllowDebugViewShaderMode(DVSM_PrimitiveDistanceAccuracy, Platform, GetMaxSupportedFeatureLevel(Platform)) && Material->GetFriendlyName().Contains(TEXT("PrimitiveDistanceAccuracy"));
+		return AllowDebugViewShaderMode(DVSM_PrimitiveDistanceAccuracy, Parameters.Platform, GetMaxSupportedFeatureLevel(Parameters.Platform)) && Parameters.Material->GetFriendlyName().Contains(TEXT("PrimitiveDistanceAccuracy"));
 	}
 
 	FPrimitiveDistanceAccuracyPS(const ShaderMetaType::CompiledShaderInitializerType& Initializer):
@@ -51,7 +51,7 @@ public:
 		return bShaderHasOutdatedParameters;
 	}
 
-	static void ModifyCompilationEnvironment(EShaderPlatform Platform, const FMaterial* Material, FShaderCompilerEnvironment& OutEnvironment)
+	static void ModifyCompilationEnvironment(const FMaterialShaderPermutationParameters& Parameters, FShaderCompilerEnvironment& OutEnvironment)
 	{
 		OutEnvironment.SetDefine(TEXT("UNDEFINED_ACCURACY"), UndefinedStreamingAccuracyIntensity);
 	}
diff --git a/Engine/Source/Runtime/Renderer/Private/RayTracing/RayTracingDynamicGeometry.cpp b/Engine/Source/Runtime/Renderer/Private/RayTracing/RayTracingDynamicGeometry.cpp
--- a/Engine/Source/Runtime/Renderer/Private/RayTracing/RayTracingDynamicGeometry.cpp
+++ b/Engine/Source/Runtime/Renderer/Private/RayTracing/RayTracingDynamicGeometry.cpp
@@ -31,9 +31,9 @@ public:
 
 	FRayTracingDynamicGeometryConverterCS() = default;
 
-	static bool ShouldCompilePermutation(EShaderPlatform Platform, const FMaterial* Material, const FVertexFactoryType* VertexFactoryType)
+	static bool ShouldCompilePermutation(const FMeshMaterialShaderPermutationParameters& Parameters)
 	{
-		return IsSupportedDynamicVertexFactoryType(VertexFactoryType) && ShouldCompileRayTracingShadersForProject(Platform);
+		return IsSupportedDynamicVertexFactoryType(Parameters.VertexFactoryType) && ShouldCompileRayTracingShadersForProject(Parameters.Platform);
 	}
 
 	virtual bool Serialize(FArchive& Ar) override
diff --git a/Engine/Source/Runtime/Renderer/Private/RayTracing/RayTracingMaterialHitShaders.cpp b/Engine/Source/Runtime/Renderer/Private/RayTracing/RayTracingMaterialHitShaders.cpp
--- a/Engine/Source/Runtime/Renderer/Private/RayTracing/RayTracingMaterialHitShaders.cpp
+++ b/Engine/Source/Runtime/Renderer/Private/RayTracing/RayTracingMaterialHitShaders.cpp
@@ -110,20 +110,20 @@ public:
 
 	TMaterialCHS() {}
 
-	static bool ShouldCompilePermutation(EShaderPlatform Platform, const FMaterial* Material, const FVertexFactoryType* VertexFactoryType)
+	static bool ShouldCompilePermutation(const FMeshMaterialShaderPermutationParameters& Parameters)
 	{
-		return IsSupportedVertexFactoryType(VertexFactoryType)
-			&& (Material->IsMasked() == UseAnyHitShader)
-			&& LightMapPolicyType::ShouldCompilePermutation(Platform, Material, VertexFactoryType)
-			&& ShouldCompileRayTracingShadersForProject(Platform);
+		return IsSupportedVertexFactoryType(Parameters.VertexFactoryType)
+			&& (Parameters.Material->IsMasked() == UseAnyHitShader)
+			&& LightMapPolicyType::ShouldCompilePermutation(Parameters)
+			&& ShouldCompileRayTracingShadersForProject(Parameters.Platform);
 	}
 
-	static void ModifyCompilationEnvironment(EShaderPlatform Platform, const FMaterial* Material, FShaderCompilerEnvironment& OutEnvironment)
+	static void ModifyCompilationEnvironment(const FMaterialShaderPermutationParameters& Parameters, FShaderCompilerEnvironment& OutEnvironment)
 	{
 		OutEnvironment.SetDefine(TEXT("USE_RAYTRACED_TEXTURE_RAYCONE_LOD"), UseRayConeTextureLod ? 1 : 0);
 		OutEnvironment.SetDefine(TEXT("SCENE_TEXTURES_DISABLED"), 1);
-		LightMapPolicyType::ModifyCompilationEnvironment(Platform, Material, OutEnvironment);
-		FMeshMaterialShader::ModifyCompilationEnvironment(Platform, Material, OutEnvironment);
+		LightMapPolicyType::ModifyCompilationEnvironment(Parameters, OutEnvironment);
+		FMeshMaterialShader::ModifyCompilationEnvironment(Parameters, OutEnvironment);
 	}
 
 	static bool ValidateCompiledResult(EShaderPlatform Platform, const TArray<FMaterial*>& Materials, const FVertexFactoryType* VertexFactoryType, const FShaderParameterMap& ParameterMap, TArray<FString>& OutError)
diff --git a/Engine/Source/Runtime/Renderer/Private/RequiredTextureResolutionRendering.h b/Engine/Source/Runtime/Renderer/Private/RequiredTextureResolutionRendering.h
--- a/Engine/Source/Runtime/Renderer/Private/RequiredTextureResolutionRendering.h
+++ b/Engine/Source/Runtime/Renderer/Private/RequiredTextureResolutionRendering.h
@@ -24,10 +24,10 @@ class FRequiredTextureResolutionPS : public FDebugViewModePS
 
 public:
 
-	static bool ShouldCompilePermutation(EShaderPlatform Platform, const FMaterial* Material, const FVertexFactoryType* VertexFactoryType)
+	static bool ShouldCompilePermutation(const FMeshMaterialShaderPermutationParameters& Parameters)
 	{
 		// See FDebugViewModeMaterialProxy::GetFriendlyName()
-		return AllowDebugViewShaderMode(DVSM_RequiredTextureResolution, Platform, GetMaxSupportedFeatureLevel(Platform)) && Material->GetFriendlyName().Contains(TEXT("RequiredTextureResolution"));
+		return AllowDebugViewShaderMode(DVSM_RequiredTextureResolution, Parameters.Platform, GetMaxSupportedFeatureLevel(Parameters.Platform)) && Parameters.Material->GetFriendlyName().Contains(TEXT("RequiredTextureResolution"));
 	}
 
 	FRequiredTextureResolutionPS(const ShaderMetaType::CompiledShaderInitializerType& Initializer):
@@ -47,12 +47,12 @@ public:
 		return bShaderHasOutdatedParameters;
 	}
 
-	static void ModifyCompilationEnvironment(EShaderPlatform Platform, const FMaterial* Material, FShaderCompilerEnvironment& OutEnvironment)
+	static void ModifyCompilationEnvironment(const FMaterialShaderPermutationParameters& Parameters, FShaderCompilerEnvironment& OutEnvironment)
 	{
 		OutEnvironment.SetDefine(TEXT("UNDEFINED_ACCURACY"), UndefinedStreamingAccuracyIntensity);
 		OutEnvironment.SetDefine(TEXT("MAX_NUM_TEX_COORD"), (uint32)TEXSTREAM_MAX_NUM_UVCHANNELS);
 		OutEnvironment.SetDefine(TEXT("MAX_NUM_TEXTURE_REGISTER"), (uint32)TEXSTREAM_MAX_NUM_TEXTURES_PER_MATERIAL);
-		FMeshMaterialShader::ModifyCompilationEnvironment(Platform, OutEnvironment);
+		FMeshMaterialShader::ModifyCompilationEnvironment(Parameters.Platform, OutEnvironment);
 	}
 
 	virtual void GetDebugViewModeShaderBindings(
diff --git a/Engine/Source/Runtime/Renderer/Private/SceneHitProxyRendering.cpp b/Engine/Source/Runtime/Renderer/Private/SceneHitProxyRendering.cpp
--- a/Engine/Source/Runtime/Renderer/Private/SceneHitProxyRendering.cpp
+++ b/Engine/Source/Runtime/Renderer/Private/SceneHitProxyRendering.cpp
@@ -48,12 +48,12 @@ public:
 		return bShaderHasOutdatedParameters;
 	}
 
-	static bool ShouldCompilePermutation(EShaderPlatform Platform,const FMaterial* Material,const FVertexFactoryType* VertexFactoryType)
+	static bool ShouldCompilePermutation(const FMeshMaterialShaderPermutationParameters& Parameters)
 	{
 		// Only compile the hit proxy vertex shader on PC
-		return IsPCPlatform(Platform)
+		return IsPCPlatform(Parameters.Platform)
 			// and only compile for the default material or materials that are masked.
-			&& (Material->IsSpecialEngineMaterial() || !Material->WritesEveryPixel() || Material->MaterialMayModifyMeshPosition() || Material->IsTwoSided());
+			&& (Parameters.Material->IsSpecialEngineMaterial() || !Parameters.Material->WritesEveryPixel() || Parameters.Material->MaterialMayModifyMeshPosition() || Parameters.Material->IsTwoSided());
 	}
 
 protected:
@@ -82,10 +82,10 @@ protected:
 		FBaseHS(Initializer)
 	{}
 
-	static bool ShouldCompilePermutation(EShaderPlatform Platform,const FMaterial* Material,const FVertexFactoryType* VertexFactoryType)
+	static bool ShouldCompilePermutation(const FMeshMaterialShaderPermutationParameters& Parameters)
 	{
-		return FBaseHS::ShouldCompilePermutation(Platform, Material, VertexFactoryType)
-			&& FHitProxyVS::ShouldCompilePermutation(Platform,Material,VertexFactoryType);
+		return FBaseHS::ShouldCompilePermutation(Parameters)
+			&& FHitProxyVS::ShouldCompilePermutation(Parameters);
 	}
 };
 
@@ -104,10 +104,10 @@ protected:
 		FBaseDS(Initializer)
 	{}
 
-	static bool ShouldCompilePermutation(EShaderPlatform Platform,const FMaterial* Material,const FVertexFactoryType* VertexFactoryType)
+	static bool ShouldCompilePermutation(const FMeshMaterialShaderPermutationParameters& Parameters)
 	{
-		return FBaseDS::ShouldCompilePermutation(Platform, Material, VertexFactoryType)
-			&& FHitProxyVS::ShouldCompilePermutation(Platform,Material,VertexFactoryType);
+		return FBaseDS::ShouldCompilePermutation(Parameters)
+			&& FHitProxyVS::ShouldCompilePermutation(Parameters);
 	}
 };
 
@@ -122,12 +122,12 @@ class FHitProxyPS : public FMeshMaterialShader
 	DECLARE_SHADER_TYPE(FHitProxyPS,MeshMaterial);
 public:
 
-	static bool ShouldCompilePermutation(EShaderPlatform Platform,const FMaterial* Material,const FVertexFactoryType* VertexFactoryType)
+	static bool ShouldCompilePermutation(const FMeshMaterialShaderPermutationParameters& Parameters)
 	{
 		// Only compile the hit proxy vertex shader on PC
-		return IsPCPlatform(Platform) 
+		return IsPCPlatform(Parameters.Platform) 
 			// and only compile for default materials or materials that are masked.
-			&& (Material->IsSpecialEngineMaterial() || !Material->WritesEveryPixel() || Material->MaterialMayModifyMeshPosition() || Material->IsTwoSided());
+			&& (Parameters.Material->IsSpecialEngineMaterial() || !Parameters.Material->WritesEveryPixel() || Parameters.Material->MaterialMayModifyMeshPosition() || Parameters.Material->IsTwoSided());
 	}
 
 	FHitProxyPS(const ShaderMetaType::CompiledShaderInitializerType& Initializer):
diff --git a/Engine/Source/Runtime/Renderer/Private/ShaderBaseClasses.h b/Engine/Source/Runtime/Renderer/Private/ShaderBaseClasses.h
--- a/Engine/Source/Runtime/Renderer/Private/ShaderBaseClasses.h
+++ b/Engine/Source/Runtime/Renderer/Private/ShaderBaseClasses.h
@@ -39,20 +39,20 @@ class FBaseHS : public FMeshMaterialShader
 	DECLARE_SHADER_TYPE(FBaseHS,MeshMaterial);
 public:
 
-	static bool ShouldCompilePermutation(EShaderPlatform Platform,const FMaterial* Material,const FVertexFactoryType* VertexFactoryType)
+	static bool ShouldCompilePermutation(const FMeshMaterialShaderPermutationParameters& Parameters)
 	{
-		if (!RHISupportsTessellation(Platform))
+		if (!RHISupportsTessellation(Parameters.Platform))
 		{
 			return false;
 		}
 
-		if (VertexFactoryType && !VertexFactoryType->SupportsTessellationShaders())
+		if (Parameters.VertexFactoryType && !Parameters.VertexFactoryType->SupportsTessellationShaders())
 		{
 			// VF can opt out of tessellation
 			return false;	
 		}
 
-		if (!Material || Material->GetTessellationMode() == MTM_NoTessellation) 
+		if (!Parameters.Material || Parameters.Material->GetTessellationMode() == MTM_NoTessellation)
 		{
 			// Material controls use of tessellation
 			return false;	
@@ -79,20 +79,20 @@ class FBaseDS : public FMeshMaterialShader
 	DECLARE_SHADER_TYPE(FBaseDS,MeshMaterial);
 public:
 
-	static bool ShouldCompilePermutation(EShaderPlatform Platform,const FMaterial* Material,const FVertexFactoryType* VertexFactoryType)
+	static bool ShouldCompilePermutation(const FMeshMaterialShaderPermutationParameters& Parameters)
 	{
-		if (!RHISupportsTessellation(Platform))
+		if (!RHISupportsTessellation(Parameters.Platform))
 		{
 			return false;
 		}
 
-		if (VertexFactoryType && !VertexFactoryType->SupportsTessellationShaders())
+		if (Parameters.VertexFactoryType && !Parameters.VertexFactoryType->SupportsTessellationShaders())
 		{
 			// VF can opt out of tessellation
 			return false;	
 		}
 
-		if (!Material || Material->GetTessellationMode() == MTM_NoTessellation) 
+		if (!Parameters.Material || Parameters.Material->GetTessellationMode() == MTM_NoTessellation)
 		{
 			// Material controls use of tessellation
 			return false;	
diff --git a/Engine/Source/Runtime/Renderer/Private/ShaderComplexityRendering.h b/Engine/Source/Runtime/Renderer/Private/ShaderComplexityRendering.h
--- a/Engine/Source/Runtime/Renderer/Private/ShaderComplexityRendering.h
+++ b/Engine/Source/Runtime/Renderer/Private/ShaderComplexityRendering.h
@@ -33,10 +33,10 @@ class TComplexityAccumulatePS : public FDebugViewModePS
 	DECLARE_SHADER_TYPE(TComplexityAccumulatePS,MeshMaterial);
 public:
 
-	static bool ShouldCompilePermutation(EShaderPlatform Platform, const FMaterial* Material, const FVertexFactoryType* VertexFactoryType)
+	static bool ShouldCompilePermutation(const FMeshMaterialShaderPermutationParameters& Parameters)
 	{
 		// See FDebugViewModeMaterialProxy::GetFriendlyName()
-		return AllowDebugViewShaderMode(bQuadComplexity ? DVSM_QuadComplexity : DVSM_ShaderComplexity, Platform, GetMaxSupportedFeatureLevel(Platform)) && Material->GetFriendlyName().Contains(TEXT("ComplexityAccumulate"));
+		return AllowDebugViewShaderMode(bQuadComplexity ? DVSM_QuadComplexity : DVSM_ShaderComplexity, Parameters.Platform, GetMaxSupportedFeatureLevel(Parameters.Platform)) && Parameters.Material->GetFriendlyName().Contains(TEXT("ComplexityAccumulate"));
 	}
 
 	TComplexityAccumulatePS(const ShaderMetaType::CompiledShaderInitializerType& Initializer):
@@ -58,11 +58,11 @@ public:
 		return bShaderHasOutdatedParameters;
 	}
 
-	static void ModifyCompilationEnvironment(EShaderPlatform Platform, const FMaterial* Material, FShaderCompilerEnvironment& OutEnvironment)
+	static void ModifyCompilationEnvironment(const FMaterialShaderPermutationParameters& Parameters, FShaderCompilerEnvironment& OutEnvironment)
 	{
-		OutEnvironment.SetDefine(TEXT("OUTPUT_QUAD_OVERDRAW"), AllowDebugViewShaderMode(DVSM_QuadComplexity, Platform, GetMaxSupportedFeatureLevel(Platform)));
+		OutEnvironment.SetDefine(TEXT("OUTPUT_QUAD_OVERDRAW"), AllowDebugViewShaderMode(DVSM_QuadComplexity, Parameters.Platform, GetMaxSupportedFeatureLevel(Parameters.Platform)));
 		TCHAR BufferRegister[] = { 'u', '0', 0 };
-		BufferRegister[1] += FSceneRenderTargets::GetQuadOverdrawUAVIndex(Platform, Material->GetFeatureLevel());
+		BufferRegister[1] += FSceneRenderTargets::GetQuadOverdrawUAVIndex(Parameters.Platform, Parameters.Material->GetFeatureLevel());
 		OutEnvironment.SetDefine(TEXT("QUAD_BUFFER_REGISTER"), BufferRegister);
 	}
 
diff --git a/Engine/Source/Runtime/Renderer/Private/ShadowDepthRendering.cpp b/Engine/Source/Runtime/Renderer/Private/ShadowDepthRendering.cpp
--- a/Engine/Source/Runtime/Renderer/Private/ShadowDepthRendering.cpp
+++ b/Engine/Source/Runtime/Renderer/Private/ShadowDepthRendering.cpp
@@ -124,7 +124,7 @@ class FShadowDepthVS : public FMeshMaterialShader
 	DECLARE_SHADER_TYPE(FShadowDepthVS, MeshMaterial);
 public:
 
-	static bool ShouldCompilePermutation(EShaderPlatform Platform, const FMaterial* Material, const FVertexFactoryType* VertexFactoryType)
+	static bool ShouldCompilePermutation(const FMeshMaterialShaderPermutationParameters& Parameters)
 	{
 		return false;
 	}
@@ -204,8 +204,11 @@ public:
 
 	TShadowDepthVS() {}
 
-	static bool ShouldCompilePermutation(EShaderPlatform Platform, const FMaterial* Material, const FVertexFactoryType* VertexFactoryType)
+	static bool ShouldCompilePermutation(const FMeshMaterialShaderPermutationParameters& Parameters)
 	{
+		const EShaderPlatform Platform = Parameters.Platform;
+		const FMaterial* Material = Parameters.Material;
+
 		static const auto SupportAllShaderPermutationsVar = IConsoleManager::Get().FindTConsoleVariableDataInt(TEXT("r.SupportAllShaderPermutations"));
 		const bool bForceAllPermutations = SupportAllShaderPermutationsVar && SupportAllShaderPermutationsVar->GetValueOnAnyThread() != 0;
 		const bool bSupportPointLightWholeSceneShadows = CVarSupportPointLightWholeSceneShadows.GetValueOnAnyThread() != 0 || bForceAllPermutations;
@@ -235,7 +238,7 @@ public:
 				// Only compile one pass point light shaders for feature levels >= SM4
 				&& (ShaderMode != VertexShadowDepth_OnePassPointLight || IsFeatureLevelSupported(Platform, ERHIFeatureLevel::SM4))
 				// Only compile position-only shaders for vertex factories that support it. (Note: this assumes that a vertex factor which supports PositionOnly, supports also PositionAndNormalOnly)
-				&& (!bUsePositionOnlyStream || VertexFactoryType->SupportsPositionOnly())
+				&& (!bUsePositionOnlyStream || Parameters.VertexFactoryType->SupportsPositionOnly())
 				// Don't render ShadowDepth for translucent unlit materials
 				&& Material->ShouldCastDynamicShadows()
 				// Only compile perspective correct light shaders for feature levels >= SM4
@@ -243,9 +246,9 @@ public:
 		}
 	}
 
-	static void ModifyCompilationEnvironment(EShaderPlatform Platform, const FMaterial* Material, FShaderCompilerEnvironment& OutEnvironment)
+	static void ModifyCompilationEnvironment(const FMaterialShaderPermutationParameters& Parameters, FShaderCompilerEnvironment& OutEnvironment)
 	{
-		FShadowDepthVS::ModifyCompilationEnvironment(Platform, Material, OutEnvironment);
+		FShadowDepthVS::ModifyCompilationEnvironment(Parameters, OutEnvironment);
 		OutEnvironment.SetDefine(TEXT("PERSPECTIVE_CORRECT_DEPTH"), (uint32)(ShaderMode == VertexShadowDepth_PerspectiveCorrect));
 		OutEnvironment.SetDefine(TEXT("ONEPASS_POINTLIGHT_SHADOW"), (uint32)(ShaderMode == VertexShadowDepth_OnePassPointLight));
 		OutEnvironment.SetDefine(TEXT("REFLECTIVE_SHADOW_MAP"), (uint32)bRenderReflectiveShadowMap);
@@ -275,18 +278,18 @@ public:
 
 	TShadowDepthHS() {}
 
-	static bool ShouldCompilePermutation(EShaderPlatform Platform, const FMaterial* Material, const FVertexFactoryType* VertexFactoryType)
+	static bool ShouldCompilePermutation(const FMeshMaterialShaderPermutationParameters& Parameters)
 	{
 		// Re-use ShouldCache from vertex shader
-		return FBaseHS::ShouldCompilePermutation(Platform, Material, VertexFactoryType)
-			&& TShadowDepthVS<ShaderMode, bRenderReflectiveShadowMap, false>::ShouldCompilePermutation(Platform, Material, VertexFactoryType);
+		return FBaseHS::ShouldCompilePermutation(Parameters)
+			&& TShadowDepthVS<ShaderMode, bRenderReflectiveShadowMap, false>::ShouldCompilePermutation(Parameters);
 	}
 
-	static void ModifyCompilationEnvironment(EShaderPlatform Platform, const FMaterial* Material, FShaderCompilerEnvironment& OutEnvironment)
+	static void ModifyCompilationEnvironment(const FMaterialShaderPermutationParameters& Parameters, FShaderCompilerEnvironment& OutEnvironment)
 	{
 		// Re-use compilation env from vertex shader
 
-		TShadowDepthVS<ShaderMode, bRenderReflectiveShadowMap, false>::ModifyCompilationEnvironment(Platform, Material, OutEnvironment);
+		TShadowDepthVS<ShaderMode, bRenderReflectiveShadowMap, false>::ModifyCompilationEnvironment(Parameters, OutEnvironment);
 	}
 };
 
@@ -317,17 +320,17 @@ public:
 
 	TShadowDepthDS() {}
 
-	static bool ShouldCompilePermutation(EShaderPlatform Platform, const FMaterial* Material, const FVertexFactoryType* VertexFactoryType)
+	static bool ShouldCompilePermutation(const FMeshMaterialShaderPermutationParameters& Parameters)
 	{
 		// Re-use ShouldCache from vertex shader
-		return FBaseDS::ShouldCompilePermutation(Platform, Material, VertexFactoryType)
-			&& TShadowDepthVS<ShaderMode, bRenderReflectiveShadowMap, false>::ShouldCompilePermutation(Platform, Material, VertexFactoryType);
+		return FBaseDS::ShouldCompilePermutation(Parameters)
+			&& TShadowDepthVS<ShaderMode, bRenderReflectiveShadowMap, false>::ShouldCompilePermutation(Parameters);
 	}
 
-	static void ModifyCompilationEnvironment(EShaderPlatform Platform, const FMaterial* Material, FShaderCompilerEnvironment& OutEnvironment)
+	static void ModifyCompilationEnvironment(const FMaterialShaderPermutationParameters& Parameters, FShaderCompilerEnvironment& OutEnvironment)
 	{
 		// Re-use compilation env from vertex shader
-		TShadowDepthVS<ShaderMode, bRenderReflectiveShadowMap, false>::ModifyCompilationEnvironment(Platform, Material, OutEnvironment);
+		TShadowDepthVS<ShaderMode, bRenderReflectiveShadowMap, false>::ModifyCompilationEnvironment(Parameters, OutEnvironment);
 	}
 };
 
@@ -337,16 +340,16 @@ class FOnePassPointShadowDepthGS : public FMeshMaterialShader
 	DECLARE_SHADER_TYPE(FOnePassPointShadowDepthGS, MeshMaterial);
 public:
 
-	static bool ShouldCompilePermutation(EShaderPlatform Platform, const FMaterial* Material, const FVertexFactoryType* VertexFactoryType)
+	static bool ShouldCompilePermutation(const FMeshMaterialShaderPermutationParameters& Parameters)
 	{
-		return RHISupportsGeometryShaders(Platform) && TShadowDepthVS<VertexShadowDepth_OnePassPointLight, false, false, true>::ShouldCompilePermutation(Platform, Material, VertexFactoryType);
+		return RHISupportsGeometryShaders(Parameters.Platform) && TShadowDepthVS<VertexShadowDepth_OnePassPointLight, false, false, true>::ShouldCompilePermutation(Parameters);
 	}
 
-	static void ModifyCompilationEnvironment(EShaderPlatform Platform, const FMaterial* Material, FShaderCompilerEnvironment& OutEnvironment)
+	static void ModifyCompilationEnvironment(const FMaterialShaderPermutationParameters& Parameters, FShaderCompilerEnvironment& OutEnvironment)
 	{
-		FMeshMaterialShader::ModifyCompilationEnvironment(Platform, Material, OutEnvironment);
+		FMeshMaterialShader::ModifyCompilationEnvironment(Parameters, OutEnvironment);
 		OutEnvironment.SetDefine(TEXT("ONEPASS_POINTLIGHT_SHADOW"), 1);
-		TShadowDepthVS<VertexShadowDepth_OnePassPointLight, false, false, true>::ModifyCompilationEnvironment(Platform, Material, OutEnvironment);
+		TShadowDepthVS<VertexShadowDepth_OnePassPointLight, false, false, true>::ModifyCompilationEnvironment(Parameters, OutEnvironment);
 	}
 
 	FOnePassPointShadowDepthGS(const ShaderMetaType::CompiledShaderInitializerType& Initializer)
@@ -461,8 +464,11 @@ class TShadowDepthPS : public TShadowDepthBasePS<bRenderReflectiveShadowMap>
 	DECLARE_SHADER_TYPE(TShadowDepthPS, MeshMaterial);
 public:
 
-	static bool ShouldCompilePermutation(EShaderPlatform Platform, const FMaterial* Material, const FVertexFactoryType* VertexFactoryType)
+	static bool ShouldCompilePermutation(const FMeshMaterialShaderPermutationParameters& Parameters)
 	{
+		const EShaderPlatform Platform = Parameters.Platform;
+		const FMaterial* Material = Parameters.Material;
+
 		if (!IsFeatureLevelSupported(Platform, ERHIFeatureLevel::SM4))
 		{
 			return (Material->IsSpecialEngineMaterial()
@@ -503,9 +509,9 @@ public:
 		}
 	}
 
-	static void ModifyCompilationEnvironment(EShaderPlatform Platform, const FMaterial* Material, FShaderCompilerEnvironment& OutEnvironment)
+	static void ModifyCompilationEnvironment(const FMaterialShaderPermutationParameters& Parameters, FShaderCompilerEnvironment& OutEnvironment)
 	{
-		TShadowDepthBasePS<bRenderReflectiveShadowMap>::ModifyCompilationEnvironment(Platform, Material, OutEnvironment);
+		TShadowDepthBasePS<bRenderReflectiveShadowMap>::ModifyCompilationEnvironment(Parameters, OutEnvironment);
 		OutEnvironment.SetDefine(TEXT("PERSPECTIVE_CORRECT_DEPTH"), (uint32)(ShaderMode == PixelShadowDepth_PerspectiveCorrect));
 		OutEnvironment.SetDefine(TEXT("ONEPASS_POINTLIGHT_SHADOW"), (uint32)(ShaderMode == PixelShadowDepth_OnePassPointLight));
 		OutEnvironment.SetDefine(TEXT("REFLECTIVE_SHADOW_MAP"), (uint32)bRenderReflectiveShadowMap);
diff --git a/Engine/Source/Runtime/Renderer/Private/TranslucentLighting.cpp b/Engine/Source/Runtime/Renderer/Private/TranslucentLighting.cpp
--- a/Engine/Source/Runtime/Renderer/Private/TranslucentLighting.cpp
+++ b/Engine/Source/Runtime/Renderer/Private/TranslucentLighting.cpp
@@ -287,9 +287,9 @@ class FTranslucencyShadowDepthVS : public FMeshMaterialShader
 	DECLARE_SHADER_TYPE(FTranslucencyShadowDepthVS, MeshMaterial);
 public:
 
-	static bool ShouldCompilePermutation(EShaderPlatform Platform, const FMaterial* Material, const FVertexFactoryType* VertexFactoryType)
+	static bool ShouldCompilePermutation(const FMeshMaterialShaderPermutationParameters& Parameters)
 	{
-		return IsTranslucentBlendMode(Material->GetBlendMode()) && IsFeatureLevelSupported(Platform, ERHIFeatureLevel::SM4);
+		return IsTranslucentBlendMode(Parameters.Material->GetBlendMode()) && IsFeatureLevelSupported(Parameters.Platform, ERHIFeatureLevel::SM4);
 	}
 
 	FTranslucencyShadowDepthVS() {}
@@ -319,9 +319,9 @@ public:
 
 	TTranslucencyShadowDepthVS() {}
 
-	static void ModifyCompilationEnvironment(EShaderPlatform Platform, const FMaterial* Material, FShaderCompilerEnvironment& OutEnvironment)
+	static void ModifyCompilationEnvironment(const FMaterialShaderPermutationParameters& Parameters, FShaderCompilerEnvironment& OutEnvironment)
 	{
-		FTranslucencyShadowDepthVS::ModifyCompilationEnvironment(Platform, Material, OutEnvironment);
+		FTranslucencyShadowDepthVS::ModifyCompilationEnvironment(Parameters, OutEnvironment);
 		OutEnvironment.SetDefine(TEXT("PERSPECTIVE_CORRECT_DEPTH"), (uint32)(ShaderMode == TranslucencyShadowDepth_PerspectiveCorrect ? 1 : 0));
 	}
 };
@@ -337,9 +337,9 @@ class FTranslucencyShadowDepthPS : public FMeshMaterialShader
 	DECLARE_SHADER_TYPE(FTranslucencyShadowDepthPS,MeshMaterial);
 public:
 
-	static bool ShouldCompilePermutation(EShaderPlatform Platform,const FMaterial* Material,const FVertexFactoryType* VertexFactoryType)
+	static bool ShouldCompilePermutation(const FMeshMaterialShaderPermutationParameters& Parameters)
 	{
-		return IsTranslucentBlendMode(Material->GetBlendMode()) && IsFeatureLevelSupported(Platform, ERHIFeatureLevel::SM4);
+		return IsTranslucentBlendMode(Parameters.Material->GetBlendMode()) && IsFeatureLevelSupported(Parameters.Platform, ERHIFeatureLevel::SM4);
 	}
 
 	FTranslucencyShadowDepthPS(const ShaderMetaType::CompiledShaderInitializerType& Initializer):
@@ -391,9 +391,9 @@ public:
 
 	TTranslucencyShadowDepthPS() {}
 
-	static void ModifyCompilationEnvironment( EShaderPlatform Platform, const FMaterial* Material, FShaderCompilerEnvironment& OutEnvironment )
+	static void ModifyCompilationEnvironment(const FMaterialShaderPermutationParameters& Parameters, FShaderCompilerEnvironment& OutEnvironment )
 	{
-		FTranslucencyShadowDepthPS::ModifyCompilationEnvironment(Platform, Material, OutEnvironment);
+		FTranslucencyShadowDepthPS::ModifyCompilationEnvironment(Parameters, OutEnvironment);
 		OutEnvironment.SetDefine(TEXT("PERSPECTIVE_CORRECT_DEPTH"), (uint32)(ShaderMode == TranslucencyShadowDepth_PerspectiveCorrect ? 1 : 0));
 	}
 };
@@ -804,9 +804,9 @@ class TTranslucentLightingInjectPS : public FMaterialShader
 	DECLARE_SHADER_TYPE(TTranslucentLightingInjectPS,Material);
 public:
 
-	static void ModifyCompilationEnvironment( EShaderPlatform Platform, const FMaterial* Material, FShaderCompilerEnvironment& OutEnvironment )
+	static void ModifyCompilationEnvironment(const FMaterialShaderPermutationParameters& Parameters, FShaderCompilerEnvironment& OutEnvironment )
 	{
-		FMaterialShader::ModifyCompilationEnvironment(Platform, Material, OutEnvironment);
+		FMaterialShader::ModifyCompilationEnvironment(Parameters, OutEnvironment);
 		OutEnvironment.SetDefine(TEXT("RADIAL_ATTENUATION"), (uint32)(InjectionType != LightType_Directional));
 		OutEnvironment.SetDefine(TEXT("INJECTION_PIXEL_SHADER"), 1);
 		OutEnvironment.SetDefine(TEXT("DYNAMICALLY_SHADOWED"), (uint32)bDynamicallyShadowed);
@@ -819,9 +819,11 @@ public:
 	  * as 'UsedAsLightFunction' in the Material Editor gets compiled into
 	  * the shader cache.
 	  */
-	static bool ShouldCompilePermutation(EShaderPlatform Platform, const FMaterial* Material)
+	static bool ShouldCompilePermutation(const FMaterialShaderPermutationParameters& Parameters)
 	{
-		return (Material->IsLightFunction() || Material->IsSpecialEngineMaterial()) && (IsFeatureLevelSupported(Platform, ERHIFeatureLevel::SM4) && (RHISupportsGeometryShaders(Platform) || RHISupportsVertexShaderLayer(Platform)));
+		return (Parameters.Material->IsLightFunction() || Parameters.Material->IsSpecialEngineMaterial()) &&
+			(IsFeatureLevelSupported(Parameters.Platform, ERHIFeatureLevel::SM4) &&
+			(RHISupportsGeometryShaders(Parameters.Platform) || RHISupportsVertexShaderLayer(Parameters.Platform)));
 	}
 
 	TTranslucentLightingInjectPS(const ShaderMetaType::CompiledShaderInitializerType& Initializer):
diff --git a/Engine/Source/Runtime/Renderer/Private/VT/RuntimeVirtualTextureRender.cpp b/Engine/Source/Runtime/Renderer/Private/VT/RuntimeVirtualTextureRender.cpp
--- a/Engine/Source/Runtime/Renderer/Private/VT/RuntimeVirtualTextureRender.cpp
+++ b/Engine/Source/Runtime/Renderer/Private/VT/RuntimeVirtualTextureRender.cpp
@@ -24,15 +24,15 @@ namespace RuntimeVirtualTexture
 	public:
 		typedef FRenderTargetParameters FParameters;
 
-		static bool ShouldCompilePermutation(EShaderPlatform Platform, const FMaterial* Material, const FVertexFactoryType* VertexFactoryType)
+		static bool ShouldCompilePermutation(const FMeshMaterialShaderPermutationParameters& Parameters)
 		{
-			return IsFeatureLevelSupported(Platform, ERHIFeatureLevel::SM4) && (Material->GetMaterialDomain() == MD_RuntimeVirtualTexture || Material->HasRuntimeVirtualTextureOutput());
+			return IsFeatureLevelSupported(Parameters.Platform, ERHIFeatureLevel::SM4) && (Parameters.Material->GetMaterialDomain() == MD_RuntimeVirtualTexture || Parameters.Material->HasRuntimeVirtualTextureOutput());
 		}
 
-		static void ModifyCompilationEnvironment(EShaderPlatform Platform, const FMaterial* Material, FShaderCompilerEnvironment& OutEnvironment)
+		static void ModifyCompilationEnvironment(const FMaterialShaderPermutationParameters& Parameters, FShaderCompilerEnvironment& OutEnvironment)
 		{
-			FMeshMaterialShader::ModifyCompilationEnvironment(Platform, Material, OutEnvironment);
-			if (Material->HasRuntimeVirtualTextureOutput())
+			FMeshMaterialShader::ModifyCompilationEnvironment(Parameters, OutEnvironment);
+			if (Parameters.Material->HasRuntimeVirtualTextureOutput())
 			{
 				OutEnvironment.SetDefine(TEXT("VIRTUAL_TEXTURE_OUTPUT"), 1);
 			}
@@ -128,9 +128,9 @@ namespace RuntimeVirtualTexture
 			: FShader_VirtualTextureMaterialDraw(Initializer)
 		{}
 
-		static void ModifyCompilationEnvironment(EShaderPlatform Platform, const FMaterial* Material, FShaderCompilerEnvironment& OutEnvironment)
+		static void ModifyCompilationEnvironment(const FMaterialShaderPermutationParameters& Parameters, FShaderCompilerEnvironment& OutEnvironment)
 		{
-			FShader_VirtualTextureMaterialDraw::ModifyCompilationEnvironment(Platform, Material, OutEnvironment);
+			FShader_VirtualTextureMaterialDraw::ModifyCompilationEnvironment(Parameters, OutEnvironment);
 			MaterialPolicy::ModifyCompilationEnvironment(OutEnvironment);
 		}
 	};
@@ -149,9 +149,9 @@ namespace RuntimeVirtualTexture
 			: FShader_VirtualTextureMaterialDraw(Initializer)
 		{}
 
-		static void ModifyCompilationEnvironment(EShaderPlatform Platform, const FMaterial* Material, FShaderCompilerEnvironment& OutEnvironment)
+		static void ModifyCompilationEnvironment(const FMaterialShaderPermutationParameters& Parameters, FShaderCompilerEnvironment& OutEnvironment)
 		{
-			FShader_VirtualTextureMaterialDraw::ModifyCompilationEnvironment(Platform, Material, OutEnvironment);
+			FShader_VirtualTextureMaterialDraw::ModifyCompilationEnvironment(Parameters, OutEnvironment);
 			MaterialPolicy::ModifyCompilationEnvironment(OutEnvironment);
 		}
 	};
diff --git a/Engine/Source/Runtime/Renderer/Private/VelocityRendering.cpp b/Engine/Source/Runtime/Renderer/Private/VelocityRendering.cpp
--- a/Engine/Source/Runtime/Renderer/Private/VelocityRendering.cpp
+++ b/Engine/Source/Runtime/Renderer/Private/VelocityRendering.cpp
@@ -71,8 +71,11 @@ public:
 			true;
 	}
 
-	static bool ShouldCompilePermutation(EShaderPlatform Platform,const FMaterial* Material,const FVertexFactoryType* VertexFactoryType)
+	static bool ShouldCompilePermutation(const FMeshMaterialShaderPermutationParameters& Parameters)
 	{
+		const EShaderPlatform Platform = Parameters.Platform;
+		const FMaterial* Material = Parameters.Material;
+
 		//Only compile the velocity shaders for the default material or if it's masked,
 		return ((Material->IsSpecialEngineMaterial() || !Material->WritesEveryPixel() 
 			//or if the material is opaque and two-sided,
@@ -80,7 +83,7 @@ public:
 			// or if the material modifies meshes
 			|| Material->MaterialMayModifyMeshPosition()))
 			&& IsFeatureLevelSupported(Platform, ERHIFeatureLevel::SM4) 
-			&& !FVelocityRendering::VertexFactoryOnlyOutputsVelocityInBasePass(Platform, VertexFactoryType->SupportsStaticLighting());
+			&& !FVelocityRendering::VertexFactoryOnlyOutputsVelocityInBasePass(Platform, Parameters.VertexFactoryType->SupportsStaticLighting());
 	}
 
 protected:
@@ -132,10 +135,10 @@ protected:
 
 	FVelocityHS() {}
 
-	static bool ShouldCompilePermutation(EShaderPlatform Platform,const FMaterial* Material,const FVertexFactoryType* VertexFactoryType)
+	static bool ShouldCompilePermutation(const FMeshMaterialShaderPermutationParameters& Parameters)
 	{
-		return FBaseHS::ShouldCompilePermutation(Platform, Material, VertexFactoryType) &&
-				FVelocityVS::ShouldCompilePermutation(Platform, Material, VertexFactoryType); // same rules as VS
+		return FBaseHS::ShouldCompilePermutation(Parameters) &&
+				FVelocityVS::ShouldCompilePermutation(Parameters); // same rules as VS
 	}
 };
 
@@ -152,10 +155,10 @@ protected:
 
 	FVelocityDS() {}
 
-	static bool ShouldCompilePermutation(EShaderPlatform Platform, const FMaterial* Material, const FVertexFactoryType* VertexFactoryType)
+	static bool ShouldCompilePermutation(const FMeshMaterialShaderPermutationParameters& Parameters)
 	{
-		return FBaseDS::ShouldCompilePermutation(Platform, Material, VertexFactoryType) &&
-				FVelocityVS::ShouldCompilePermutation(Platform, Material, VertexFactoryType); // same rules as VS
+		return FBaseDS::ShouldCompilePermutation(Parameters) &&
+				FVelocityVS::ShouldCompilePermutation(Parameters); // same rules as VS
 	}
 };
 
@@ -169,8 +172,11 @@ class FVelocityPS : public FMeshMaterialShader
 {
 	DECLARE_SHADER_TYPE(FVelocityPS,MeshMaterial);
 public:
-	static bool ShouldCompilePermutation(EShaderPlatform Platform,const FMaterial* Material,const FVertexFactoryType* VertexFactoryType)
+	static bool ShouldCompilePermutation(const FMeshMaterialShaderPermutationParameters& Parameters)
 	{
+		const EShaderPlatform Platform = Parameters.Platform;
+		const FMaterial* Material = Parameters.Material;
+
 		//Only compile the velocity shaders for the default material or if it's masked,
 		return ((Material->IsSpecialEngineMaterial() || !Material->WritesEveryPixel() 
 			//or if the material is opaque and two-sided,
@@ -178,12 +184,12 @@ public:
 			// or if the material modifies meshes
 			|| Material->MaterialMayModifyMeshPosition()))
 			&& IsFeatureLevelSupported(Platform, ERHIFeatureLevel::SM4) && 
-			!FVelocityRendering::VertexFactoryOnlyOutputsVelocityInBasePass(Platform, VertexFactoryType->SupportsStaticLighting());
+			!FVelocityRendering::VertexFactoryOnlyOutputsVelocityInBasePass(Platform, Parameters.VertexFactoryType->SupportsStaticLighting());
 	}
 
-	static void ModifyCompilationEnvironment( EShaderPlatform Platform, const FMaterial* Material, FShaderCompilerEnvironment& OutEnvironment )
+	static void ModifyCompilationEnvironment(const FMaterialShaderPermutationParameters& Parameters, FShaderCompilerEnvironment& OutEnvironment )
 	{
-		FMaterialShader::ModifyCompilationEnvironment(Platform, Material, OutEnvironment);
+		FMaterialShader::ModifyCompilationEnvironment(Parameters, OutEnvironment);
 		OutEnvironment.SetRenderTargetOutputFormat(0, PF_G16R16);
 	}
 
diff --git a/Engine/Source/Runtime/Renderer/Private/VolumetricFogLightFunction.cpp b/Engine/Source/Runtime/Renderer/Private/VolumetricFogLightFunction.cpp
--- a/Engine/Source/Runtime/Renderer/Private/VolumetricFogLightFunction.cpp
+++ b/Engine/Source/Runtime/Renderer/Private/VolumetricFogLightFunction.cpp
@@ -26,9 +26,9 @@ class FVolumetricFogLightFunctionPS : public FMaterialShader
 	DECLARE_SHADER_TYPE(FVolumetricFogLightFunctionPS,Material);
 public:
 
-	static bool ShouldCompilePermutation(EShaderPlatform Platform, const FMaterial* Material)
+	static bool ShouldCompilePermutation(const FMaterialShaderPermutationParameters& Parameters)
 	{
-		return Material->IsLightFunction() && DoesPlatformSupportVolumetricFog(Platform);
+		return Parameters.Material->IsLightFunction() && DoesPlatformSupportVolumetricFog(Parameters.Platform);
 	}
 
 	FVolumetricFogLightFunctionPS() {}
diff --git a/Engine/Source/Runtime/Renderer/Private/VolumetricFogVoxelization.cpp b/Engine/Source/Runtime/Renderer/Private/VolumetricFogVoxelization.cpp
--- a/Engine/Source/Runtime/Renderer/Private/VolumetricFogVoxelization.cpp
+++ b/Engine/Source/Runtime/Renderer/Private/VolumetricFogVoxelization.cpp
@@ -161,17 +161,17 @@ class FVoxelizeVolumeVS : public FMeshMaterialShader
 	{
 	}
 
-	static bool ShouldCompilePermutation(EShaderPlatform Platform,const FMaterial* Material,const FVertexFactoryType* VertexFactoryType)
+	static bool ShouldCompilePermutation(const FMeshMaterialShaderPermutationParameters& Parameters)
 	{
-		return IsFeatureLevelSupported(Platform, ERHIFeatureLevel::SM5) 
-			&& DoesPlatformSupportVolumetricFogVoxelization(Platform)
-			&& Material->GetMaterialDomain() == MD_Volume;
+		return IsFeatureLevelSupported(Parameters.Platform, ERHIFeatureLevel::SM5) 
+			&& DoesPlatformSupportVolumetricFogVoxelization(Parameters.Platform)
+			&& Parameters.Material->GetMaterialDomain() == MD_Volume;
 	}
 
-	static void ModifyCompilationEnvironment(EShaderPlatform Platform, const FMaterial* Material, FShaderCompilerEnvironment& OutEnvironment)
+	static void ModifyCompilationEnvironment(const FMaterialShaderPermutationParameters& Parameters, FShaderCompilerEnvironment& OutEnvironment)
 	{
-		FMeshMaterialShader::ModifyCompilationEnvironment(Platform, Material, OutEnvironment);
-		if (RHISupportsGeometryShaders(Platform))
+		FMeshMaterialShader::ModifyCompilationEnvironment(Parameters, OutEnvironment);
+		if (RHISupportsGeometryShaders(Parameters.Platform))
 		{
 			OutEnvironment.CompilerFlags.Add( CFLAG_VertexToGeometryShader );
 		}
@@ -230,9 +230,9 @@ protected:
 
 public:
 
-	static void ModifyCompilationEnvironment(EShaderPlatform Platform, const FMaterial* Material, FShaderCompilerEnvironment& OutEnvironment)
+	static void ModifyCompilationEnvironment(const FMaterialShaderPermutationParameters& Parameters, FShaderCompilerEnvironment& OutEnvironment)
 	{
-		Super::ModifyCompilationEnvironment(Platform, Material, OutEnvironment);
+		Super::ModifyCompilationEnvironment(Parameters, OutEnvironment);
 		
 		if (Mode == VMode_Primitive_Sphere)
 		{
@@ -265,18 +265,18 @@ protected:
 	{
 	}
 
-	static bool ShouldCompilePermutation(EShaderPlatform Platform,const FMaterial* Material,const FVertexFactoryType* VertexFactoryType)
+	static bool ShouldCompilePermutation(const FMeshMaterialShaderPermutationParameters& Parameters)
 	{
-		return IsFeatureLevelSupported(Platform, ERHIFeatureLevel::SM5) 
-			&& RHISupportsGeometryShaders(Platform)
-			&& DoesPlatformSupportVolumetricFogVoxelization(Platform)
-			&& Material->GetMaterialDomain() == MD_Volume;
+		return IsFeatureLevelSupported(Parameters.Platform, ERHIFeatureLevel::SM5) 
+			&& RHISupportsGeometryShaders(Parameters.Platform)
+			&& DoesPlatformSupportVolumetricFogVoxelization(Parameters.Platform)
+			&& Parameters.Material->GetMaterialDomain() == MD_Volume;
 	}
 
-	static void ModifyCompilationEnvironment(EShaderPlatform Platform, const FMaterial* Material, FShaderCompilerEnvironment& OutEnvironment)
+	static void ModifyCompilationEnvironment(const FMaterialShaderPermutationParameters& Parameters, FShaderCompilerEnvironment& OutEnvironment)
 	{
-		FMeshMaterialShader::ModifyCompilationEnvironment(Platform, Material, OutEnvironment);
-		OutEnvironment.SetDefine(TEXT("MAX_SLICES_PER_VOXELIZATION_PASS"), GetVoxelizationSlicesPerPass(Platform));
+		FMeshMaterialShader::ModifyCompilationEnvironment(Parameters, OutEnvironment);
+		OutEnvironment.SetDefine(TEXT("MAX_SLICES_PER_VOXELIZATION_PASS"), GetVoxelizationSlicesPerPass(Parameters.Platform));
 	}
 
 public:
@@ -324,9 +324,9 @@ protected:
 
 public:
 
-	static void ModifyCompilationEnvironment(EShaderPlatform Platform, const FMaterial* Material, FShaderCompilerEnvironment& OutEnvironment)
+	static void ModifyCompilationEnvironment(const FMaterialShaderPermutationParameters& Parameters, FShaderCompilerEnvironment& OutEnvironment)
 	{
-		Super::ModifyCompilationEnvironment(Platform, Material, OutEnvironment);
+		Super::ModifyCompilationEnvironment(Parameters, OutEnvironment);
 		
 		if (Mode == VMode_Primitive_Sphere)
 		{
@@ -358,11 +358,11 @@ protected:
 	{
 	}
 
-	static bool ShouldCompilePermutation(EShaderPlatform Platform,const FMaterial* Material,const FVertexFactoryType* VertexFactoryType)
+	static bool ShouldCompilePermutation(const FMeshMaterialShaderPermutationParameters& Parameters)
 	{
-		return IsFeatureLevelSupported(Platform, ERHIFeatureLevel::SM5) 
-			&& DoesPlatformSupportVolumetricFogVoxelization(Platform)
-			&& Material->GetMaterialDomain() == MD_Volume;
+		return IsFeatureLevelSupported(Parameters.Platform, ERHIFeatureLevel::SM5) 
+			&& DoesPlatformSupportVolumetricFogVoxelization(Parameters.Platform)
+			&& Parameters.Material->GetMaterialDomain() == MD_Volume;
 	}
 };
 
@@ -382,9 +382,9 @@ protected:
 
 public:
 
-	static void ModifyCompilationEnvironment(EShaderPlatform Platform, const FMaterial* Material, FShaderCompilerEnvironment& OutEnvironment)
+	static void ModifyCompilationEnvironment(const FMaterialShaderPermutationParameters& Parameters, FShaderCompilerEnvironment& OutEnvironment)
 	{
-		Super::ModifyCompilationEnvironment(Platform, Material, OutEnvironment);
+		Super::ModifyCompilationEnvironment(Parameters, OutEnvironment);
 		
 		if (Mode == VMode_Primitive_Sphere)
 		{
diff --git a/Engine/Source/Runtime/SlateRHIRenderer/Private/SlateMaterialShader.cpp b/Engine/Source/Runtime/SlateRHIRenderer/Private/SlateMaterialShader.cpp
--- a/Engine/Source/Runtime/SlateRHIRenderer/Private/SlateMaterialShader.cpp
+++ b/Engine/Source/Runtime/SlateRHIRenderer/Private/SlateMaterialShader.cpp
@@ -13,19 +13,19 @@ FSlateMaterialShaderVS::FSlateMaterialShaderVS(const FMaterialShaderType::Compil
 }
 
 
-void FSlateMaterialShaderVS::ModifyCompilationEnvironment(EShaderPlatform Platform, const FMaterial* Material, FShaderCompilerEnvironment& OutEnvironment)
+void FSlateMaterialShaderVS::ModifyCompilationEnvironment(const FMaterialShaderPermutationParameters& Parameters, FShaderCompilerEnvironment& OutEnvironment)
 {
 	// Set defines based on what this shader will be used for
 	OutEnvironment.SetDefine( TEXT("USE_MATERIALS"), 1 );
-	OutEnvironment.SetDefine( TEXT("NUM_CUSTOMIZED_UVS"), Material->GetNumCustomizedUVs() );
-	OutEnvironment.SetDefine( TEXT("HAS_SCREEN_POSITION"), Material->HasVertexPositionOffsetConnected() );
+	OutEnvironment.SetDefine( TEXT("NUM_CUSTOMIZED_UVS"), Parameters.Material->GetNumCustomizedUVs() );
+	OutEnvironment.SetDefine( TEXT("HAS_SCREEN_POSITION"), Parameters.Material->HasVertexPositionOffsetConnected() );
 
-	FMaterialShader::ModifyCompilationEnvironment( Platform, Material, OutEnvironment );
+	FMaterialShader::ModifyCompilationEnvironment( Parameters, OutEnvironment );
 }
 
-bool FSlateMaterialShaderVS::ShouldCompilePermutation(EShaderPlatform Platform, const FMaterial* Material)
+bool FSlateMaterialShaderVS::ShouldCompilePermutation(const FMaterialShaderPermutationParameters& Parameters)
 {
-	return Material->GetMaterialDomain() == MD_UI;
+	return Parameters.Material->GetMaterialDomain() == MD_UI;
 }
 
 void FSlateMaterialShaderVS::SetViewProjection(FRHICommandList& RHICmdList, const FMatrix& InViewProjection )
@@ -57,19 +57,19 @@ bool FSlateMaterialShaderVS::Serialize(FArchive& Ar)
 }
 
 
-bool FSlateMaterialShaderPS::ShouldCompilePermutation(EShaderPlatform Platform, const FMaterial* Material)
+bool FSlateMaterialShaderPS::ShouldCompilePermutation(const FMaterialShaderPermutationParameters& Parameters)
 {
-	return Material->GetMaterialDomain() == MD_UI;
+	return Parameters.Material->GetMaterialDomain() == MD_UI;
 }
 
 
-void FSlateMaterialShaderPS::ModifyCompilationEnvironment(EShaderPlatform Platform, const FMaterial* Material, FShaderCompilerEnvironment& OutEnvironment)
+void FSlateMaterialShaderPS::ModifyCompilationEnvironment(const FMaterialShaderPermutationParameters& Parameters, FShaderCompilerEnvironment& OutEnvironment)
 {
 	// Set defines based on what this shader will be used for
 	OutEnvironment.SetDefine( TEXT("USE_MATERIALS"), 1 );
-	OutEnvironment.SetDefine( TEXT("NUM_CUSTOMIZED_UVS"), Material->GetNumCustomizedUVs() );
+	OutEnvironment.SetDefine( TEXT("NUM_CUSTOMIZED_UVS"), Parameters.Material->GetNumCustomizedUVs() );
 
-	FMaterialShader::ModifyCompilationEnvironment( Platform, Material, OutEnvironment );
+	FMaterialShader::ModifyCompilationEnvironment( Parameters, OutEnvironment );
 }
 
 FSlateMaterialShaderPS::FSlateMaterialShaderPS(const FMaterialShaderType::CompiledShaderInitializerType& Initializer)
diff --git a/Engine/Source/Runtime/SlateRHIRenderer/Private/SlateMaterialShader.h b/Engine/Source/Runtime/SlateRHIRenderer/Private/SlateMaterialShader.h
--- a/Engine/Source/Runtime/SlateRHIRenderer/Private/SlateMaterialShader.h
+++ b/Engine/Source/Runtime/SlateRHIRenderer/Private/SlateMaterialShader.h
@@ -15,9 +15,9 @@ public:
 	FSlateMaterialShaderVS() {}
 	FSlateMaterialShaderVS(const FMaterialShaderType::CompiledShaderInitializerType& Initializer);
 
-	static void ModifyCompilationEnvironment(EShaderPlatform Platform, const FMaterial* Material, FShaderCompilerEnvironment& OutEnvironment);
+	static void ModifyCompilationEnvironment(const FMaterialShaderPermutationParameters& Parameters, FShaderCompilerEnvironment& OutEnvironment);
 
-	static bool ShouldCompilePermutation(EShaderPlatform Platform, const FMaterial* Material);
+	static bool ShouldCompilePermutation(const FMaterialShaderPermutationParameters& Parameters);
 
 	/** 
 	 * Sets the view projection parameter
@@ -47,10 +47,10 @@ class FSlateMaterialShaderPS : public FMaterialShader
 public:
 
 	/** Only compile shaders used with UI. */
-	static bool ShouldCompilePermutation(EShaderPlatform Platform, const FMaterial* Material);
+	static bool ShouldCompilePermutation(const FMaterialShaderPermutationParameters& Parameters);
 
 	/** Modifies the compilation of this shader. */
-	static void ModifyCompilationEnvironment(EShaderPlatform Platform, const FMaterial* Material, FShaderCompilerEnvironment& OutEnvironment);
+	static void ModifyCompilationEnvironment(const FMaterialShaderPermutationParameters& Parameters, FShaderCompilerEnvironment& OutEnvironment);
 
 	FSlateMaterialShaderPS() {}
 	FSlateMaterialShaderPS(const FMaterialShaderType::CompiledShaderInitializerType& Initializer);
@@ -86,15 +86,15 @@ public:
 	{ }
 	
 	/** Only compile shaders used with UI. */
-	static bool ShouldCompilePermutation(EShaderPlatform Platform, const FMaterial* Material)
+	static bool ShouldCompilePermutation(const FMaterialShaderPermutationParameters& Parameters)
 	{
-		return FSlateMaterialShaderVS::ShouldCompilePermutation(Platform, Material);
+		return FSlateMaterialShaderVS::ShouldCompilePermutation(Parameters);
 	}
 
 	/** Modifies the compilation of this shader. */
-	static void ModifyCompilationEnvironment(EShaderPlatform Platform, const FMaterial* Material, FShaderCompilerEnvironment& OutEnvironment)
+	static void ModifyCompilationEnvironment(const FMaterialShaderPermutationParameters& Parameters, FShaderCompilerEnvironment& OutEnvironment)
 	{
-		FSlateMaterialShaderVS::ModifyCompilationEnvironment(Platform, Material, OutEnvironment);
+		FSlateMaterialShaderVS::ModifyCompilationEnvironment(Parameters, OutEnvironment);
 
 		OutEnvironment.SetDefine(TEXT("USE_SLATE_INSTANCING"), (uint32)( bUseInstancing ? 1 : 0 ));
 	}
@@ -118,15 +118,15 @@ public:
 	{ }
 	
 	/** Only compile shaders used with UI. */
-	static bool ShouldCompilePermutation(EShaderPlatform Platform, const FMaterial* Material)
+	static bool ShouldCompilePermutation(const FMaterialShaderPermutationParameters& Parameters)
 	{
-		return FSlateMaterialShaderPS::ShouldCompilePermutation(Platform,Material);
+		return FSlateMaterialShaderPS::ShouldCompilePermutation(Parameters);
 	}
 
 	/** Modifies the compilation of this shader. */
-	static void ModifyCompilationEnvironment(EShaderPlatform Platform, const FMaterial* Material, FShaderCompilerEnvironment& OutEnvironment)
+	static void ModifyCompilationEnvironment(const FMaterialShaderPermutationParameters& Parameters, FShaderCompilerEnvironment& OutEnvironment)
 	{
-		FSlateMaterialShaderPS::ModifyCompilationEnvironment(Platform, Material,OutEnvironment);
+		FSlateMaterialShaderPS::ModifyCompilationEnvironment(Parameters,OutEnvironment);
 
 		OutEnvironment.SetDefine(TEXT("SHADER_TYPE"), (uint32)ShaderType);
 		OutEnvironment.SetDefine(TEXT("DRAW_DISABLED_EFFECT"), (uint32)(bDrawDisabledEffect ? 1 : 0));
