DR - Revert change
#rnx

diff --git a/Engine/Source/Runtime/Renderer/Private/PlanarReflectionRendering.cpp b/Engine/Source/Runtime/Renderer/Private/PlanarReflectionRendering.cpp
--- a/Engine/Source/Runtime/Renderer/Private/PlanarReflectionRendering.cpp
+++ b/Engine/Source/Runtime/Renderer/Private/PlanarReflectionRendering.cpp
@@ -721,7 +721,6 @@ void FDeferredShadingSceneRenderer::RenderDeferredPlanarReflections(FRDGBuilder&
 	PassParameters->SceneTextures.SceneDepthBuffer = SceneTextures.SceneDepthBuffer;
 	PassParameters->SceneTextures.SceneGBufferA = SceneTextures.SceneGBufferA;
 	PassParameters->SceneTextures.SceneGBufferB = SceneTextures.SceneGBufferB;
-	PassParameters->SceneTextures.SceneGBufferC = SceneTextures.SceneGBufferC;
 	SetupSceneTextureSamplers(&PassParameters->SceneTextureSamplers);
 	PassParameters->ViewUniformBuffer = View.ViewUniformBuffer;
 	PassParameters->RenderTargets[0] = FRenderTargetBinding(
