Adjust VT feedback GPU->CPU transfer to avoid crash if lightshafts are enabled

diff --git a/Engine/Source/Runtime/Renderer/Private/PostProcess/SceneRenderTargets.cpp b/Engine/Source/Runtime/Renderer/Private/PostProcess/SceneRenderTargets.cpp
--- a/Engine/Source/Runtime/Renderer/Private/PostProcess/SceneRenderTargets.cpp
+++ b/Engine/Source/Runtime/Renderer/Private/PostProcess/SceneRenderTargets.cpp
@@ -846,6 +846,7 @@ void FSceneRenderTargets::BindVirtualTextureFeedbackUAV(FRHIRenderPassInfo& RPIn
 	// must match VT_FEEDBACK_REGISTER in VirtualTextureCommon.ush
 	static const int32 VT_FEEDBACK_REGISTER = 7;
 
+	checkf(VirtualTextureFeedback.FeedbackTextureGPU, TEXT("Invalid attempt to render VT feedback after it has already been transferred to CPU, need to adjust location of TransferGPUToCPU"));
 	check(RPInfo.GetNumColorRenderTargets() <= VT_FEEDBACK_REGISTER);
 	RPInfo.UAVIndex = VT_FEEDBACK_REGISTER;
 	RPInfo.UAVs[RPInfo.NumUAVs++] = VirtualTextureFeedback.FeedbackTextureGPU->GetRenderTargetItem().UAV;
