DR - Workaround for hlslcc platforms not removing unused parameters
#rnx

diff --git a/Engine/Source/Runtime/Renderer/Private/PlanarReflectionRendering.cpp b/Engine/Source/Runtime/Renderer/Private/PlanarReflectionRendering.cpp
--- a/Engine/Source/Runtime/Renderer/Private/PlanarReflectionRendering.cpp
+++ b/Engine/Source/Runtime/Renderer/Private/PlanarReflectionRendering.cpp
@@ -721,6 +721,11 @@ void FDeferredShadingSceneRenderer::RenderDeferredPlanarReflections(FRDGBuilder&
 	PassParameters->SceneTextures.SceneDepthBuffer = SceneTextures.SceneDepthBuffer;
 	PassParameters->SceneTextures.SceneGBufferA = SceneTextures.SceneGBufferA;
 	PassParameters->SceneTextures.SceneGBufferB = SceneTextures.SceneGBufferB;
+	if (IsHlslccShaderPlatform(GMaxRHIShaderPlatform))
+	{
+		// hlslcc doesn't remove all unused parameters
+		PassParameters->SceneTextures.SceneGBufferC = GraphBuilder.RegisterExternalTexture(GSystemTextures.BlackDummy);
+	}
 	SetupSceneTextureSamplers(&PassParameters->SceneTextureSamplers);
 	PassParameters->ViewUniformBuffer = View.ViewUniformBuffer;
 	PassParameters->RenderTargets[0] = FRenderTargetBinding(
