Renamed FSceneViewFamilyBlackboard to FSceneTextureParameters.

diff --git a/Engine/Shaders/Private/PlanarReflectionShaders.usf b/Engine/Shaders/Private/PlanarReflectionShaders.usf
--- a/Engine/Shaders/Private/PlanarReflectionShaders.usf
+++ b/Engine/Shaders/Private/PlanarReflectionShaders.usf
@@ -7,7 +7,7 @@
 #include "Common.ush"
 #include "DeferredShadingCommon.ush"
 #include "PlanarReflectionShared.ush"
-#include "SceneViewFamilyBlackboard.ush"
+#include "SceneTextureParameters.ush"
 
 Texture2D SceneColorInputTexture;
 SamplerState SceneColorInputSampler;
@@ -133,7 +133,7 @@ void PlanarReflectionPS(
 	ResolvedView = ResolveView();
 	OutColor = 0;
 
-	FGBufferData GBuffer = GetGBufferDataFromSceneBlackboard(BufferUV);
+	FGBufferData GBuffer = GetGBufferDataFromSceneTextures(BufferUV);
 
 	float3 WorldPosition = ScreenVector * GBuffer.Depth + View.WorldCameraOrigin;
 
diff --git a/Engine/Shaders/Private/RayTracing/RayTracingOcclusionRGS.usf b/Engine/Shaders/Private/RayTracing/RayTracingOcclusionRGS.usf
--- a/Engine/Shaders/Private/RayTracing/RayTracingOcclusionRGS.usf
+++ b/Engine/Shaders/Private/RayTracing/RayTracingOcclusionRGS.usf
@@ -4,7 +4,7 @@
 #include "../MonteCarlo.ush"
 #include "../DeferredShadingCommon.ush"
 #include "../LightShaderParameters.ush"
-#include "../SceneViewFamilyBlackboard.ush"
+#include "../SceneTextureParameters.ush"
 #include "../ScreenSpaceDenoise/SSDPublic.ush"
 #include "../SobolRandom.ush"
 #include "RayTracingCommon.ush"
@@ -133,7 +133,7 @@ void OcclusionRGS()
 		// Get G-Buffer surface data
 		float2 InvBufferSize = View.BufferSizeAndInvSize.zw;
 		float2 BufferUV = (float2(PixelCoord) + 0.5) * InvBufferSize;
-		FGBufferData GBufferData = GetGBufferDataFromSceneBlackboard(BufferUV);
+		FGBufferData GBufferData = GetGBufferDataFromSceneTextures(BufferUV);
 		float Depth = GBufferData.Depth;
 		float3 WorldPosition = ReconstructWorldPositionFromDepth(BufferUV, Depth);
 		float3 WorldNormal = GBufferData.WorldNormal;
diff --git a/Engine/Shaders/Private/ReflectionEnvironmentPixelShader.usf b/Engine/Shaders/Private/ReflectionEnvironmentPixelShader.usf
--- a/Engine/Shaders/Private/ReflectionEnvironmentPixelShader.usf
+++ b/Engine/Shaders/Private/ReflectionEnvironmentPixelShader.usf
@@ -12,7 +12,7 @@
 #include "DistanceFieldAOShared.ush"
 #include "ShadingModels.ush"
 #include "LightGridCommon.ush"
-#include "SceneViewFamilyBlackboard.ush"
+#include "SceneTextureParameters.ush"
 
 #define REFLECTION_COMPOSITE_USE_BLENDED_REFLECTION_CAPTURES 1
 #define REFLECTION_COMPOSITE_SUPPORT_SKYLIGHT_BLEND 1
@@ -303,7 +303,7 @@ void ReflectionEnvironmentSkyLighting(
 	float2 ScreenPosition = SvPositionToScreenPosition(SvPosition).xy;
 
 	// Sample scene textures.
-	FGBufferData GBuffer = GetGBufferDataFromSceneBlackboard(BufferUV);
+	FGBufferData GBuffer = GetGBufferDataFromSceneTextures(BufferUV);
 
 	// Sample the ambient occlusion that is dynamically generated every frame.
 	float AmbientOcclusion = AmbientOcclusionTexture.SampleLevel(AmbientOcclusionSampler, BufferUV, 0).r;
diff --git a/Engine/Shaders/Private/SSRT/SSRTReflections.usf b/Engine/Shaders/Private/SSRT/SSRTReflections.usf
--- a/Engine/Shaders/Private/SSRT/SSRTReflections.usf
+++ b/Engine/Shaders/Private/SSRT/SSRTReflections.usf
@@ -4,7 +4,7 @@
 #include "../Random.ush"
 #include "../BRDF.ush"
 #include "../MonteCarlo.ush"
-#include "../SceneViewFamilyBlackboard.ush"
+#include "../SceneTextureParameters.ush"
 #include "../ScreenSpaceDenoise/SSDPublic.ush"
 
 // Generate vector truncation warnings to errors.
@@ -74,7 +74,7 @@ void ScreenSpaceReflectionsStencilPS(
 {
 	float2 UV = SvPosition.xy * View.BufferSizeAndInvSize.zw;
 	
-	FGBufferData GBuffer = GetGBufferDataFromSceneBlackboard(UV);
+	FGBufferData GBuffer = GetGBufferDataFromSceneTextures(UV);
 	
 	const float Roughness = GetRoughness(GBuffer);
 	const float RoughnessFade = GetRoughnessFade(Roughness);
@@ -113,7 +113,7 @@ void ScreenSpaceReflectionsPS(
 		OutClosestHitDistance = -2.0;
 	#endif
 
-	FGBufferData GBuffer = GetGBufferDataFromSceneBlackboard(UV);
+	FGBufferData GBuffer = GetGBufferDataFromSceneTextures(UV);
 	
 	float Roughness = GetRoughness(GBuffer);
 	float RoughnessFade = GetRoughnessFade(Roughness);
diff --git a/Engine/Shaders/Private/SceneTextureParameters.ush b/Engine/Shaders/Private/SceneTextureParameters.ush
new file mode 100644
--- /dev/null
+++ b/Engine/Shaders/Private/SceneTextureParameters.ush
@@ -0,0 +1,71 @@
+// Copyright 1998-2019 Epic Games, Inc. All Rights Reserved.
+
+#pragma once
+
+#include "SceneTexturesCommon.ush"
+#include "DeferredShadingCommon.ush"
+
+uint bIsSceneLightingChannelsValid;
+
+// Matches FSceneTextureParameters
+Texture2D SceneDepthBuffer;
+Texture2D SceneVelocityBuffer;
+Texture2D SceneGBufferA;
+Texture2D SceneGBufferB;
+Texture2D SceneGBufferC;
+Texture2D SceneGBufferD;
+Texture2D SceneGBufferE;
+Texture2D<uint> SceneLightingChannels;
+
+#if SUPPORTS_INDEPENDENT_SAMPLERS
+#define SceneDepthBufferSampler GlobalPointClampedSampler
+#define SceneVelocityBufferSampler GlobalPointClampedSampler
+#define SceneGBufferASampler GlobalPointClampedSampler
+#define SceneGBufferBSampler GlobalPointClampedSampler
+#define SceneGBufferCSampler GlobalPointClampedSampler
+#define SceneGBufferDSampler GlobalPointClampedSampler
+#define SceneGBufferESampler GlobalPointClampedSampler
+
+#else
+SamplerState SceneDepthBufferSampler;
+SamplerState SceneVelocityBufferSampler;
+SamplerState SceneGBufferASampler;
+SamplerState SceneGBufferBSampler;
+SamplerState SceneGBufferCSampler;
+SamplerState SceneGBufferDSampler;
+SamplerState SceneGBufferESampler;
+
+#endif
+
+
+// @param UV - UV space in the GBuffer textures (BufferSize resolution)
+// TOOD: need to find a way to make this more convenient.
+FGBufferData GetGBufferDataFromSceneTextures(float2 UV, bool bGetNormalizedNormal = true)
+{
+	float4 GBufferA = SceneGBufferA.SampleLevel(SceneGBufferASampler, UV, 0);
+	float4 GBufferB = SceneGBufferB.SampleLevel(SceneGBufferBSampler, UV, 0);
+	float4 GBufferC = SceneGBufferC.SampleLevel(SceneGBufferCSampler, UV, 0);
+	float4 GBufferD = SceneGBufferD.SampleLevel(SceneGBufferDSampler, UV, 0);
+	float4 GBufferE = SceneGBufferE.SampleLevel(SceneGBufferESampler, UV, 0);
+	float4 GBufferVelocity = SceneVelocityBuffer.SampleLevel(SceneVelocityBufferSampler, UV, 0);
+
+	uint CustomStencil = 0;
+	float CustomNativeDepth = 0;
+
+	float DeviceZ = SceneDepthBuffer.SampleLevel(SceneDepthBufferSampler, UV, 0).r;
+
+	float SceneDepth = ConvertFromDeviceZ(DeviceZ);
+
+	return DecodeGBufferData(GBufferA, GBufferB, GBufferC, GBufferD, GBufferE, GBufferVelocity, CustomNativeDepth, CustomStencil, SceneDepth, bGetNormalizedNormal, CheckerFromSceneColorUV(UV));
+}
+
+/** Returns the light channel mask that should be executed for this pixel. */
+uint GetSceneLightingChannel(uint2 PixelCoord)
+{
+	BRANCH
+	if (bIsSceneLightingChannelsValid)
+	{
+		return SceneLightingChannels.Load(uint3(PixelCoord, 0)).x;
+	}
+	return ~0;
+}
diff --git a/Engine/Shaders/Private/SceneViewFamilyBlackboard.ush b/Engine/Shaders/Private/SceneViewFamilyBlackboard.ush
deleted file mode 100644
diff --git a/Engine/Shaders/Private/ScreenSpaceDenoise/SSDCommon.ush b/Engine/Shaders/Private/ScreenSpaceDenoise/SSDCommon.ush
--- a/Engine/Shaders/Private/ScreenSpaceDenoise/SSDCommon.ush
+++ b/Engine/Shaders/Private/ScreenSpaceDenoise/SSDCommon.ush
@@ -11,7 +11,7 @@
 #endif
 
 #include "../Common.ush"
-#include "../SceneViewFamilyBlackboard.ush"
+#include "../SceneTextureParameters.ush"
 #include "../ColorSpace.ush"
 #include "../SphericalGaussian.ush"
 
@@ -122,7 +122,7 @@ FSSDSampleSceneInfos InitSampleSceneInfos(
 
 FSSDSampleSceneInfos FetchCurrentSceneInfos(float2 ScreenPosition, float2 BufferUV)
 {
-	FGBufferData GBufferData = GetGBufferDataFromSceneBlackboard(BufferUV);
+	FGBufferData GBufferData = GetGBufferDataFromSceneTextures(BufferUV);
 	return InitSampleSceneInfos(
 		View.ScreenToTranslatedWorld,
 		ScreenPosition, GBufferData.Depth, GBufferData.WorldNormal, GBufferData.Roughness);
diff --git a/Engine/Source/Runtime/Renderer/Private/DeferredShadingRenderer.cpp b/Engine/Source/Runtime/Renderer/Private/DeferredShadingRenderer.cpp
--- a/Engine/Source/Runtime/Renderer/Private/DeferredShadingRenderer.cpp
+++ b/Engine/Source/Runtime/Renderer/Private/DeferredShadingRenderer.cpp
@@ -30,7 +30,7 @@
 #include "GPUScene.h"
 #include "RayTracing/RayTracingMaterialHitShaders.h"
 #include "RayTracingDynamicGeometryCollection.h"
-#include "SceneViewFamilyBlackboard.h"
+#include "SceneTextureParameters.h"
 #include "ScreenSpaceDenoise.h"
 #include "ScreenSpaceRayTracing.h"
 #include "RayTracing/RaytracingOptions.h"
diff --git a/Engine/Source/Runtime/Renderer/Private/DeferredShadingRenderer.h b/Engine/Source/Runtime/Renderer/Private/DeferredShadingRenderer.h
--- a/Engine/Source/Runtime/Renderer/Private/DeferredShadingRenderer.h
+++ b/Engine/Source/Runtime/Renderer/Private/DeferredShadingRenderer.h
@@ -16,7 +16,7 @@
 #include "DepthRendering.h"
 #include "ScreenSpaceDenoise.h"
 
-class FSceneViewFamilyBlackboard;
+class FSceneTextureParameters;
 
 class FDistanceFieldAOParameters;
 class UStaticMeshComponent;
@@ -440,7 +440,7 @@ private:
 	/** Render image based reflections (SSR, Env, SkyLight) without compute shaders */
 	void RenderStandardDeferredImageBasedReflections(FRHICommandListImmediate& RHICmdList, FGraphicsPipelineStateInitializer& GraphicsPSOInit, bool bReflectionEnv, const TRefCountPtr<IPooledRenderTarget>& DynamicBentNormalAO, TRefCountPtr<IPooledRenderTarget>& VelocityRT);
 
-	void RenderDeferredPlanarReflections(FRDGBuilder& GraphBuilder, const FSceneViewFamilyBlackboard& SceneTextures, const FViewInfo& View, FRDGTextureRef& ReflectionsOutput);
+	void RenderDeferredPlanarReflections(FRDGBuilder& GraphBuilder, const FSceneTextureParameters& SceneTextures, const FViewInfo& View, FRDGTextureRef& ReflectionsOutput);
 
 	bool ShouldDoReflectionEnvironment() const;
 	
@@ -472,7 +472,7 @@ private:
 
 	void RenderRayTracingShadows(
 		FRDGBuilder& GraphBuilder,
-		const FSceneViewFamilyBlackboard& SceneBlackboard,
+		const FSceneTextureParameters& SceneTextures,
 		const FViewInfo& View,
 		const FLightSceneInfo& LightSceneInfo,
 		const IScreenSpaceDenoiser::FShadowRayTracingConfig& RayTracingConfig,
diff --git a/Engine/Source/Runtime/Renderer/Private/LightRendering.cpp b/Engine/Source/Runtime/Renderer/Private/LightRendering.cpp
--- a/Engine/Source/Runtime/Renderer/Private/LightRendering.cpp
+++ b/Engine/Source/Runtime/Renderer/Private/LightRendering.cpp
@@ -16,7 +16,7 @@
 #include "ShowFlags.h"
 #include "VisualizeTexture.h"
 #include "RayTracing/RaytracingOptions.h"
-#include "SceneViewFamilyBlackboard.h"
+#include "SceneTextureParameters.h"
 
 // ENABLE_DEBUG_DISCARD_PROP is used to test the lighting code by allowing to discard lights to see how performance scales
 // It ought never to be enabled in a shipping build, and is probably only really useful when woring on the shading code.
@@ -948,8 +948,8 @@ void FDeferredShadingSceneRenderer::RenderLights(FRHICommandListImmediate& RHICm
 
 				FRDGBuilder GraphBuilder(RHICmdList);
 
-				FSceneViewFamilyBlackboard SceneBlackboard;
-				SetupSceneViewFamilyBlackboard(GraphBuilder, &SceneBlackboard);
+				FSceneTextureParameters SceneTextures;
+				SetupSceneTextureParameters(GraphBuilder, &SceneTextures);
 
 				FViewInfo& View = Views[0];
 
@@ -999,7 +999,7 @@ void FDeferredShadingSceneRenderer::RenderLights(FRHICommandListImmediate& RHICm
 							GraphBuilder,
 							View,
 							&View.PrevViewInfo,
-							SceneBlackboard,
+							SceneTextures,
 							DenoisingQueue,
 							InputParameterCount,
 							Outputs);
@@ -1061,7 +1061,7 @@ void FDeferredShadingSceneRenderer::RenderLights(FRHICommandListImmediate& RHICm
 
 							RenderRayTracingShadows(
 								GraphBuilder,
-								SceneBlackboard,
+								SceneTextures,
 								View,
 								LightSceneInfo,
 								RayTracingConfig,
@@ -1159,14 +1159,14 @@ void FDeferredShadingSceneRenderer::RenderLights(FRHICommandListImmediate& RHICm
 
 						FRDGBuilder GraphBuilder(RHICmdList);
 
-						FSceneViewFamilyBlackboard SceneBlackboard;
-						SetupSceneViewFamilyBlackboard(GraphBuilder, &SceneBlackboard);
+						FSceneTextureParameters SceneTextures;
+						SetupSceneTextureParameters(GraphBuilder, &SceneTextures);
 
 						FRDGTextureRef ShadowMask;
 						FRDGTextureRef RayHitDistance;
 							RenderRayTracingShadows(
 								GraphBuilder,
-								SceneBlackboard,
+								SceneTextures,
 								View,
 								LightSceneInfo,
 								RayTracingConfig,
@@ -1196,7 +1196,7 @@ void FDeferredShadingSceneRenderer::RenderLights(FRHICommandListImmediate& RHICm
 									GraphBuilder,
 									View,
 									&View.PrevViewInfo,
-									SceneBlackboard,
+									SceneTextures,
 									InputParameters,
 									InputParameterCount,
 									Outputs);
diff --git a/Engine/Source/Runtime/Renderer/Private/PlanarReflectionRendering.cpp b/Engine/Source/Runtime/Renderer/Private/PlanarReflectionRendering.cpp
--- a/Engine/Source/Runtime/Renderer/Private/PlanarReflectionRendering.cpp
+++ b/Engine/Source/Runtime/Renderer/Private/PlanarReflectionRendering.cpp
@@ -31,7 +31,7 @@
 #include "Containers/ArrayView.h"
 #include "PipelineStateCache.h"
 #include "ClearQuad.h"
-#include "SceneViewFamilyBlackboard.h"
+#include "SceneTextureParameters.h"
 
 void SetupPlanarReflectionUniformParameters(const class FSceneView& View, const FPlanarReflectionSceneProxy* ReflectionSceneProxy, FPlanarReflectionUniformParameters& OutParameters)
 {
@@ -652,7 +652,7 @@ class FPlanarReflectionPS : public FGlobalShader
 	}
 
 	BEGIN_SHADER_PARAMETER_STRUCT(FParameters, )
-		SHADER_PARAMETER_STRUCT_INCLUDE(FSceneViewFamilyBlackboard, SceneTextures)
+		SHADER_PARAMETER_STRUCT_INCLUDE(FSceneTextureParameters, SceneTextures)
 		SHADER_PARAMETER_STRUCT_INCLUDE(FSceneTextureSamplerParameters, SceneTextureSamplers)
 
 		SHADER_PARAMETER_STRUCT_REF(FViewUniformShaderParameters, ViewUniformBuffer)
@@ -664,7 +664,7 @@ class FPlanarReflectionPS : public FGlobalShader
 
 IMPLEMENT_GLOBAL_SHADER(FPlanarReflectionPS, "/Engine/Private/PlanarReflectionShaders.usf", "PlanarReflectionPS", SF_Pixel);
 
-void FDeferredShadingSceneRenderer::RenderDeferredPlanarReflections(FRDGBuilder& GraphBuilder, const FSceneViewFamilyBlackboard& SceneTextures, const FViewInfo& View, FRDGTextureRef& ReflectionsOutputTexture)
+void FDeferredShadingSceneRenderer::RenderDeferredPlanarReflections(FRDGBuilder& GraphBuilder, const FSceneTextureParameters& SceneTextures, const FViewInfo& View, FRDGTextureRef& ReflectionsOutputTexture)
 {
 	// Prevent rendering unsupported views when ViewIndex >= GMaxPlanarReflectionViews
 	// Planar reflections in those views will fallback to other reflection methods
diff --git a/Engine/Source/Runtime/Renderer/Private/PostProcess/DiaphragmDOF.cpp b/Engine/Source/Runtime/Renderer/Private/PostProcess/DiaphragmDOF.cpp
--- a/Engine/Source/Runtime/Renderer/Private/PostProcess/DiaphragmDOF.cpp
+++ b/Engine/Source/Runtime/Renderer/Private/PostProcess/DiaphragmDOF.cpp
@@ -18,7 +18,7 @@
 #include "ClearQuad.h"
 #include "SpriteIndexBuffer.h"
 #include "PostProcessTemporalAA.h"
-#include "SceneViewFamilyBlackboard.h"
+#include "SceneTextureParameters.h"
 
 
 void AddPass_ClearUAV(
@@ -1252,7 +1252,7 @@ IMPLEMENT_GLOBAL_SHADER(FDiaphragmDOFRecombineCS,        "/Engine/Private/Diaphr
 
 FRDGTextureRef DiaphragmDOF::AddPasses(
 	FRDGBuilder& GraphBuilder,
-	const FSceneViewFamilyBlackboard& SceneBlackboard,
+	const FSceneTextureParameters& SceneTextures,
 	const FViewInfo& View,
 	FRDGTextureRef InputSceneColor,
 	FRDGTextureRef SceneSeparateTranslucency)
@@ -1506,7 +1506,7 @@ FRDGTextureRef DiaphragmDOF::AddPasses(
 			PassParameters->ViewportRect = FVector4(0, 0, PassViewSize.X, PassViewSize.Y);
 			PassParameters->CocRadiusBasis = FVector2D(GatheringViewSize.X, PreprocessViewSize.X);
 			PassParameters->SceneColorTexture = InputSceneColor;
-			PassParameters->SceneDepthTexture = SceneBlackboard.SceneDepthBuffer;
+			PassParameters->SceneDepthTexture = SceneTextures.SceneDepthBuffer;
 		
 			if (!bOutputFullResolution)
 			{
@@ -1557,7 +1557,7 @@ FRDGTextureRef DiaphragmDOF::AddPasses(
 		TAAParameters.SceneMetadataInput = HalfResGatherInputTextures.SeparateCoc;
 
 		FTAAOutputs TAAOutputs = TAAParameters.AddTemporalAAPass(
-			GraphBuilder, SceneBlackboard, View,
+			GraphBuilder, SceneTextures, View,
 			View.PrevViewInfo.DOFSetupHistory,
 			&ViewState->PrevFrameViewInfo.DOFSetupHistory);
 		
@@ -2420,7 +2420,7 @@ FRDGTextureRef DiaphragmDOF::AddPasses(
 			(GatheringViewSize.Y - 0.5f) / float(RefBufferSize.Y));
 		
 		PassParameters->SceneColorInput = FullResGatherInputTextures.SceneColor;
-		PassParameters->SceneDepthTexture = SceneBlackboard.SceneDepthBuffer;
+		PassParameters->SceneDepthTexture = SceneTextures.SceneDepthBuffer;
 		PassParameters->SceneSeparateCoc = FullResGatherInputTextures.SeparateCoc; // TODO looks useless.
 		PassParameters->SceneSeparateTranslucency = SceneSeparateTranslucency ? SceneSeparateTranslucency : GraphBuilder.RegisterExternalTexture(GSystemTextures.BlackAlphaOneDummy);
 		
diff --git a/Engine/Source/Runtime/Renderer/Private/PostProcess/DiaphragmDOF.h b/Engine/Source/Runtime/Renderer/Private/PostProcess/DiaphragmDOF.h
--- a/Engine/Source/Runtime/Renderer/Private/PostProcess/DiaphragmDOF.h
+++ b/Engine/Source/Runtime/Renderer/Private/PostProcess/DiaphragmDOF.h
@@ -16,7 +16,7 @@
 
 
 class FViewInfo;
-class FSceneViewFamilyBlackboard;
+class FSceneTextureParameters;
 struct FTemporalAAHistory;
 
 
@@ -139,7 +139,7 @@ inline bool IsSupported(EShaderPlatform ShaderPlatform)
 /** Wire all DOF's passes according to view settings and cvars to convolve the scene color. */
 RENDERER_API FRDGTextureRef AddPasses(
 	FRDGBuilder& GraphBuilder,
-	const FSceneViewFamilyBlackboard& SceneBlackboard,
+	const FSceneTextureParameters& SceneTextures,
 	const FViewInfo& View,
 	FRDGTextureRef InputSceneColor,
 	FRDGTextureRef SeparateTranslucency);
diff --git a/Engine/Source/Runtime/Renderer/Private/PostProcess/PostProcessTemporalAA.cpp b/Engine/Source/Runtime/Renderer/Private/PostProcess/PostProcessTemporalAA.cpp
--- a/Engine/Source/Runtime/Renderer/Private/PostProcess/PostProcessTemporalAA.cpp
+++ b/Engine/Source/Runtime/Renderer/Private/PostProcess/PostProcessTemporalAA.cpp
@@ -20,7 +20,7 @@
 
 
 #include "RenderGraph.h"
-#include "SceneViewFamilyBlackboard.h"
+#include "SceneTextureParameters.h"
 #include "PixelShaderUtils.h"
 
 
@@ -88,7 +88,7 @@ BEGIN_SHADER_PARAMETER_STRUCT(FTAAShaderParameters,)
 	SHADER_PARAMETER(FVector4, HistoryBufferUVMinMax)
 	SHADER_PARAMETER(FVector4, ScreenPosToHistoryBufferUV)
 	
-	SHADER_PARAMETER_STRUCT_INCLUDE(FSceneViewFamilyBlackboard, SceneBlackboard)
+	SHADER_PARAMETER_STRUCT_INCLUDE(FSceneTextureParameters, SceneTextures)
 	
 	SHADER_PARAMETER_RDG_TEXTURE(Texture2D, EyeAdaptation)
 
@@ -433,7 +433,7 @@ bool FTAAPassParameters::Validate() const
 
 FTAAOutputs FTAAPassParameters::AddTemporalAAPass(
 	FRDGBuilder& GraphBuilder,
-	const FSceneViewFamilyBlackboard& SceneBlackboard,
+	const FSceneTextureParameters& SceneTextures,
 	const FViewInfo& View,
 	const FTemporalAAHistory& InputHistory,
 	FTemporalAAHistory* OutputHistory) const
@@ -475,7 +475,7 @@ FTAAOutputs FTAAPassParameters::AddTemporalAAPass(
 		CommonShaderParameters.CurrentFrameWeight = CVarTemporalAACurrentFrameWeight.GetValueOnRenderThread();
 		CommonShaderParameters.bCameraCut = bCameraCut;
 
-		CommonShaderParameters.SceneBlackboard = SceneBlackboard;
+		CommonShaderParameters.SceneTextures = SceneTextures;
 		CommonShaderParameters.SceneDepthBufferSampler = TStaticSamplerState<SF_Point>::GetRHI();
 		CommonShaderParameters.SceneVelocityBufferSampler = TStaticSamplerState<SF_Point>::GetRHI();
 
@@ -708,7 +708,7 @@ FTAAOutputs FTAAPassParameters::AddTemporalAAPass(
 		if (bUseResponsiveStencilTest)
 		{
 			PassParameters->RenderTargets.DepthStencil = FDepthStencilBinding(
-				SceneBlackboard.SceneDepthBuffer,
+				SceneTextures.SceneDepthBuffer,
 				ERenderTargetLoadAction::ENoAction, ERenderTargetStoreAction::ENoAction,
 				ERenderTargetLoadAction::ELoad, ERenderTargetStoreAction::EStore,
 				FExclusiveDepthStencil::DepthRead_StencilWrite);
@@ -818,13 +818,13 @@ void FRCPassPostProcessTemporalAA::Process(FRenderingCompositePassContext& Conte
 	FRDGBuilder GraphBuilder(Context.RHICmdList);
 
 	FSceneRenderTargets& SceneContext = FSceneRenderTargets::Get(Context.RHICmdList);
-	FSceneViewFamilyBlackboard SceneBlackboard;
-	SetupSceneViewFamilyBlackboard(GraphBuilder, &SceneBlackboard);
+	FSceneTextureParameters SceneTextures;
+	SetupSceneTextureParameters(GraphBuilder, &SceneTextures);
 
 	// FPostProcessing::Process() does a AdjustGBufferRefCount(RHICmdList, -1), therefore need to pass down reference on velocity buffer manually.
 	if (FRDGTextureRef SceneVelocityBuffer = CreateRDGTextureForOptionalInput(GraphBuilder, ePId_Input2, TEXT("SceneVelocity")))
 	{
-		SceneBlackboard.SceneVelocityBuffer = SceneVelocityBuffer;
+		SceneTextures.SceneVelocityBuffer = SceneVelocityBuffer;
 	}
 
 	FTAAPassParameters Parameters = SavedParameters;
@@ -833,7 +833,7 @@ void FRCPassPostProcessTemporalAA::Process(FRenderingCompositePassContext& Conte
 
 	FTAAOutputs Outputs = Parameters.AddTemporalAAPass(
 		GraphBuilder,
-		SceneBlackboard, Context.View,
+		SceneTextures, Context.View,
 		InputHistory, /* out */ OutputHistory);
 		
 	ExtractRDGTextureForOutput(GraphBuilder, ePId_Output0, Outputs.SceneColor);
diff --git a/Engine/Source/Runtime/Renderer/Private/PostProcess/PostProcessTemporalAA.h b/Engine/Source/Runtime/Renderer/Private/PostProcess/PostProcessTemporalAA.h
--- a/Engine/Source/Runtime/Renderer/Private/PostProcess/PostProcessTemporalAA.h
+++ b/Engine/Source/Runtime/Renderer/Private/PostProcess/PostProcessTemporalAA.h
@@ -15,7 +15,7 @@
 
 
 class FViewInfo;
-class FSceneViewFamilyBlackboard;
+class FSceneTextureParameters;
 struct FTemporalAAHistory;
 
 
@@ -148,7 +148,7 @@ struct FTAAPassParameters
 	/** Apply a temporal AA pass. */
 	FTAAOutputs AddTemporalAAPass(
 		FRDGBuilder& GraphBuilder,
-		const FSceneViewFamilyBlackboard& SceneBlackboard,
+		const FSceneTextureParameters& SceneTextures,
 		const FViewInfo& View,
 		const FTemporalAAHistory& InputHistory,
 		FTemporalAAHistory* OutputHistory) const;
diff --git a/Engine/Source/Runtime/Renderer/Private/PostProcess/PostProcessing.cpp b/Engine/Source/Runtime/Renderer/Private/PostProcess/PostProcessing.cpp
--- a/Engine/Source/Runtime/Renderer/Private/PostProcess/PostProcessing.cpp
+++ b/Engine/Source/Runtime/Renderer/Private/PostProcess/PostProcessing.cpp
@@ -55,7 +55,7 @@
 #include "DeferredShadingRenderer.h"
 #include "MobileSeparateTranslucencyPass.h"
 #include "MobileDistortionPass.h"
-#include "SceneViewFamilyBlackboard.h"
+#include "SceneTextureParameters.h"
 #include "PixelShaderUtils.h"
 
 /** The global center for all post processing activities. */
@@ -1410,8 +1410,8 @@ void FPostProcessing::Process(FRHICommandListImmediate& RHICmdList, const FViewI
 				{
 					FRDGBuilder GraphBuilder(InContext.RHICmdList);
 
-					FSceneViewFamilyBlackboard SceneBlackboard;
-					SetupSceneViewFamilyBlackboard(GraphBuilder, &SceneBlackboard);
+					FSceneTextureParameters SceneTextures;
+					SetupSceneTextureParameters(GraphBuilder, &SceneTextures);
 	
 					FRDGTextureRef SceneColor = Pass->CreateRDGTextureForRequiredInput(GraphBuilder, ePId_Input0, TEXT("SceneColor"));
 					FRDGTextureRef LocalSeparateTranslucency = Pass->CreateRDGTextureForOptionalInput(GraphBuilder, ePId_Input1, TEXT("SeparateTranslucency"));
@@ -1419,7 +1419,7 @@ void FPostProcessing::Process(FRHICommandListImmediate& RHICmdList, const FViewI
 					// FPostProcessing::Process() does a AdjustGBufferRefCount(RHICmdList, -1), therefore need to pass down reference on velocity buffer manually.
 					if (FRDGTextureRef SceneVelocityBuffer = Pass->CreateRDGTextureForOptionalInput(GraphBuilder, ePId_Input2, TEXT("SceneVelocity")))
 					{
-						SceneBlackboard.SceneVelocityBuffer = SceneVelocityBuffer;
+						SceneTextures.SceneVelocityBuffer = SceneVelocityBuffer;
 					}
 
 					FRDGTextureRef NewSceneColor = SceneColor;
@@ -1428,7 +1428,7 @@ void FPostProcessing::Process(FRHICommandListImmediate& RHICmdList, const FViewI
 					{
 						NewSceneColor = DiaphragmDOF::AddPasses(
 							GraphBuilder,
-							SceneBlackboard, InContext.View,
+							SceneTextures, InContext.View,
 							SceneColor, LocalSeparateTranslucency);
 					}
 
@@ -1800,8 +1800,8 @@ void FPostProcessing::Process(FRHICommandListImmediate& RHICmdList, const FViewI
 				{
 					FRDGBuilder GraphBuilder(InContext.RHICmdList);
 
-					FSceneViewFamilyBlackboard SceneBlackboard;
-					SetupSceneViewFamilyBlackboard(GraphBuilder, &SceneBlackboard);
+					FSceneTextureParameters SceneTextures;
+					SetupSceneTextureParameters(GraphBuilder, &SceneTextures);
 	
 					FRDGTextureRef SceneColor = Pass->CreateRDGTextureForRequiredInput(GraphBuilder, ePId_Input0, TEXT("SceneColor"));
 					FRDGTextureRef LocalSeparateTranslucency = Pass->CreateRDGTextureForOptionalInput(GraphBuilder, ePId_Input1, TEXT("SeparateTranslucency"));
diff --git a/Engine/Source/Runtime/Renderer/Private/RayTracing/RayTracingAmbientOcclusion.cpp b/Engine/Source/Runtime/Renderer/Private/RayTracing/RayTracingAmbientOcclusion.cpp
--- a/Engine/Source/Runtime/Renderer/Private/RayTracing/RayTracingAmbientOcclusion.cpp
+++ b/Engine/Source/Runtime/Renderer/Private/RayTracing/RayTracingAmbientOcclusion.cpp
@@ -14,7 +14,7 @@
 #include "RHI/Public/PipelineStateCache.h"
 #include "Raytracing/RaytracingOptions.h"
 #include "RayTracingMaterialHitShaders.h"
-#include "SceneViewFamilyBlackboard.h"
+#include "SceneTextureParameters.h"
 
 #include "PostProcess/PostProcessing.h"
 #include "PostProcess/SceneFilterRendering.h"
@@ -231,8 +231,8 @@ void FDeferredShadingSceneRenderer::RenderRayTracingAmbientOcclusion(
 	int32 DenoiserMode = CVarUseAODenoiser.GetValueOnRenderThread();
 	if (DenoiserMode != 0)
 	{
-		FSceneViewFamilyBlackboard SceneBlackboard;
-		SetupSceneViewFamilyBlackboard(GraphBuilder, &SceneBlackboard);
+		FSceneTextureParameters SceneTextures;
+		SetupSceneTextureParameters(GraphBuilder, &SceneTextures);
 
 		const IScreenSpaceDenoiser* DefaultDenoiser = IScreenSpaceDenoiser::GetDefaultDenoiser();
 		const IScreenSpaceDenoiser* DenoiserToUse = DenoiserMode == 1 ? DefaultDenoiser : GScreenSpaceDenoiser;
@@ -253,7 +253,7 @@ void FDeferredShadingSceneRenderer::RenderRayTracingAmbientOcclusion(
 				GraphBuilder,
 				View,
 				&View.PrevViewInfo,
-				SceneBlackboard,
+				SceneTextures,
 				DenoiserInputs,
 				RayTracingConfig);
 
diff --git a/Engine/Source/Runtime/Renderer/Private/RayTracing/RayTracingGlobalIllumination.cpp b/Engine/Source/Runtime/Renderer/Private/RayTracing/RayTracingGlobalIllumination.cpp
--- a/Engine/Source/Runtime/Renderer/Private/RayTracing/RayTracingGlobalIllumination.cpp
+++ b/Engine/Source/Runtime/Renderer/Private/RayTracing/RayTracingGlobalIllumination.cpp
@@ -13,7 +13,7 @@
 #include "UniformBuffer.h"
 #include "RayGenShaderUtils.h"
 #include "PathTracingUniformBuffers.h"
-#include "SceneViewFamilyBlackboard.h"
+#include "SceneTextureParameters.h"
 #include "ScreenSpaceDenoise.h"
 #include "ClearQuad.h"
 #include "PostProcess/PostProcessing.h"
@@ -512,8 +512,8 @@ void FDeferredShadingSceneRenderer::RenderRayTracingGlobalIllumination(
 
 	if (GRayTracingGlobalIlluminationDenoiser != 0)
 	{
-		FSceneViewFamilyBlackboard SceneBlackboard;
-		SetupSceneViewFamilyBlackboard(GraphBuilder, &SceneBlackboard);
+		FSceneTextureParameters SceneTextures;
+		SetupSceneTextureParameters(GraphBuilder, &SceneTextures);
 
 		const IScreenSpaceDenoiser* DefaultDenoiser = IScreenSpaceDenoiser::GetDefaultDenoiser();
 		const IScreenSpaceDenoiser* DenoiserToUse = GRayTracingGlobalIlluminationDenoiser == 1 ? DefaultDenoiser : GScreenSpaceDenoiser;
@@ -532,7 +532,7 @@ void FDeferredShadingSceneRenderer::RenderRayTracingGlobalIllumination(
 				GraphBuilder,
 				View,
 				&View.PrevViewInfo,
-				SceneBlackboard,
+				SceneTextures,
 				DenoiserInputs,
 				RayTracingConfig);
 
diff --git a/Engine/Source/Runtime/Renderer/Private/RayTracing/RayTracingShadows.cpp b/Engine/Source/Runtime/Renderer/Private/RayTracing/RayTracingShadows.cpp
--- a/Engine/Source/Runtime/Renderer/Private/RayTracing/RayTracingShadows.cpp
+++ b/Engine/Source/Runtime/Renderer/Private/RayTracing/RayTracingShadows.cpp
@@ -18,7 +18,7 @@
 #include "BuiltInRayTracingShaders.h"
 #include "RayTracing/RayTracingMaterialHitShaders.h"
 #include "Containers/DynamicRHIResourceArray.h"
-#include "SceneViewFamilyBlackboard.h"
+#include "SceneTextureParameters.h"
 
 static float GRayTracingMaxNormalBias = 0.1f;
 static FAutoConsoleVariableRef CVarRayTracingNormalBias(
@@ -64,7 +64,7 @@ class FOcclusionRGS : public FGlobalShader
 		SHADER_PARAMETER(FIntRect, LightScissor)
 
 		SHADER_PARAMETER_STRUCT(FLightShaderParameters, Light)
-		SHADER_PARAMETER_STRUCT_INCLUDE(FSceneViewFamilyBlackboard, SceneBlackboard)
+		SHADER_PARAMETER_STRUCT_INCLUDE(FSceneTextureParameters, SceneTextures)
 
 		SHADER_PARAMETER_SRV(RaytracingAccelerationStructure, TLAS)
 		SHADER_PARAMETER_RDG_TEXTURE_UAV(RWTexture2D<float4>, RWOcclusionMaskUAV)
@@ -111,7 +111,7 @@ void FDeferredShadingSceneRenderer::PrepareRayTracingShadows(const FViewInfo& Vi
 
 void FDeferredShadingSceneRenderer::RenderRayTracingShadows(
 	FRDGBuilder& GraphBuilder,
-	const FSceneViewFamilyBlackboard& SceneBlackboard,
+	const FSceneTextureParameters& SceneTextures,
 	const FViewInfo& View,
 	const FLightSceneInfo& LightSceneInfo,
 	const IScreenSpaceDenoiser::FShadowRayTracingConfig& RayTracingConfig,
@@ -127,7 +127,7 @@ void FDeferredShadingSceneRenderer::RenderRayTracingShadows(
 	FRDGTextureRef ScreenShadowMaskTexture;
 	{
 		FRDGTextureDesc Desc = FRDGTextureDesc::Create2DDesc(
-			SceneBlackboard.SceneDepthBuffer->Desc.Extent,
+			SceneTextures.SceneDepthBuffer->Desc.Extent,
 			PF_FloatRGBA,
 			FClearValueBinding::Black,
 			TexCreate_None,
@@ -139,7 +139,7 @@ void FDeferredShadingSceneRenderer::RenderRayTracingShadows(
 	FRDGTextureRef RayDistanceTexture;
 	{
 		FRDGTextureDesc Desc = FRDGTextureDesc::Create2DDesc(
-			SceneBlackboard.SceneDepthBuffer->Desc.Extent,
+			SceneTextures.SceneDepthBuffer->Desc.Extent,
 			PF_R16F,
 			FClearValueBinding::Black,
 			TexCreate_None,
@@ -172,7 +172,7 @@ void FDeferredShadingSceneRenderer::RenderRayTracingShadows(
 		LightSceneProxy->GetLightShaderParameters(PassParameters->Light);
 		PassParameters->TLAS = View.RayTracingScene.RayTracingSceneRHI->GetShaderResourceView();
 		PassParameters->ViewUniformBuffer = View.ViewUniformBuffer;
-		PassParameters->SceneBlackboard = SceneBlackboard;
+		PassParameters->SceneTextures = SceneTextures;
 		PassParameters->LightScissor = ScissorRect;
 		
 		FOcclusionRGS::FPermutationDomain PermutationVector;
diff --git a/Engine/Source/Runtime/Renderer/Private/RayTracing/RayTracingTranslucency.cpp b/Engine/Source/Runtime/Renderer/Private/RayTracing/RayTracingTranslucency.cpp
--- a/Engine/Source/Runtime/Renderer/Private/RayTracing/RayTracingTranslucency.cpp
+++ b/Engine/Source/Runtime/Renderer/Private/RayTracing/RayTracingTranslucency.cpp
@@ -15,7 +15,7 @@
 #include "VisualizeTexture.h"
 #include "LightRendering.h"
 #include "SystemTextures.h"
-#include "SceneViewFamilyBlackboard.h"
+#include "SceneTextureParameters.h"
 #include "ScreenSpaceDenoise.h"
 #include "PostProcess/PostProcessing.h"
 #include "PostProcess/SceneFilterRendering.h"
@@ -268,8 +268,8 @@ void FDeferredShadingSceneRenderer::RenderRayTracingTranslucency(FRHICommandList
 		FRDGBuilder GraphBuilder(RHICmdList); 
 		FSceneRenderTargets& SceneContext = FSceneRenderTargets::Get(GraphBuilder.RHICmdList);
 
-		FSceneViewFamilyBlackboard SceneBlackboard;
-		SetupSceneViewFamilyBlackboard(GraphBuilder, &SceneBlackboard);
+		FSceneTextureParameters SceneTextures;
+		SetupSceneTextureParameters(GraphBuilder, &SceneTextures);
 
 		//#dxr_todo: UE-72581 do not use reflections denoiser structs but separated ones
 		IScreenSpaceDenoiser::FReflectionsInputs DenoiserInputs;
diff --git a/Engine/Source/Runtime/Renderer/Private/RayTracing/RaytracingSkylight.cpp b/Engine/Source/Runtime/Renderer/Private/RayTracing/RaytracingSkylight.cpp
--- a/Engine/Source/Runtime/Renderer/Private/RayTracing/RaytracingSkylight.cpp
+++ b/Engine/Source/Runtime/Renderer/Private/RayTracing/RaytracingSkylight.cpp
@@ -21,7 +21,7 @@ static int32 GRayTracingSkyLight = 0;
 #include "RenderTargetPool.h"
 #include "VisualizeTexture.h"
 #include "RayGenShaderUtils.h"
-#include "SceneViewFamilyBlackboard.h"
+#include "SceneTextureParameters.h"
 #include "ScreenSpaceDenoise.h"
 
 #include "Raytracing/RaytracingOptions.h"
@@ -808,8 +808,8 @@ void FDeferredShadingSceneRenderer::RenderRayTracingSkyLight(
 	}
 
 	FRDGBuilder GraphBuilder(RHICmdList);
-	FSceneViewFamilyBlackboard SceneBlackboard;
-	SetupSceneViewFamilyBlackboard(GraphBuilder, &SceneBlackboard);
+	FSceneTextureParameters SceneTextures;
+	SetupSceneTextureParameters(GraphBuilder, &SceneTextures);
 
 	// Define RDG targets
 	FRDGTextureRef SkyLightTexture;
@@ -835,8 +835,8 @@ void FDeferredShadingSceneRenderer::RenderRayTracingSkyLight(
 	SkyLightData.MaxRayDistance = GRayTracingSkyLightMaxRayDistance;
 	SkyLightData.SamplingStopLevel = GRayTracingSkyLightSamplingStopLevel;
 
-	FSceneTexturesUniformParameters SceneTextures;
-	SetupSceneTextureUniformParameters(SceneContext, FeatureLevel, ESceneTextureSetupMode::All, SceneTextures);
+	FSceneTexturesUniformParameters SceneUniformTextures;
+	SetupSceneTextureUniformParameters(SceneContext, FeatureLevel, ESceneTextureSetupMode::All, SceneUniformTextures);
 
 	TRefCountPtr<IPooledRenderTarget> SubsurfaceProfileRT((IPooledRenderTarget*)GetSubsufaceProfileTexture_RT(RHICmdList));
 	if (!SubsurfaceProfileRT)
@@ -848,7 +848,7 @@ void FDeferredShadingSceneRenderer::RenderRayTracingSkyLight(
 	FRayTracingSkyLightRGS::FParameters *PassParameters = GraphBuilder.AllocParameters<FRayTracingSkyLightRGS::FParameters>();
 	PassParameters->RWOcclusionMaskUAV = GraphBuilder.CreateUAV(SkyLightTexture);
 	PassParameters->RWRayDistanceUAV = GraphBuilder.CreateUAV(RayDistanceTexture);
-	PassParameters->SceneTexturesStruct = CreateUniformBufferImmediate(SceneTextures, EUniformBufferUsage::UniformBuffer_SingleDraw);
+	PassParameters->SceneTexturesStruct = CreateUniformBufferImmediate(SceneUniformTextures, EUniformBufferUsage::UniformBuffer_SingleDraw);
 	PassParameters->SkyLightData = CreateUniformBufferImmediate(SkyLightData, EUniformBufferUsage::UniformBuffer_SingleDraw);
 	PassParameters->SSProfilesTexture = GraphBuilder.RegisterExternalTexture(SubsurfaceProfileRT);
 	PassParameters->TransmissionProfilesLinearSampler = TStaticSamplerState<SF_Bilinear, AM_Clamp, AM_Clamp, AM_Clamp>::GetRHI();
@@ -921,7 +921,7 @@ void FDeferredShadingSceneRenderer::RenderRayTracingSkyLight(
 					GraphBuilder,
 					View,
 					&View.PrevViewInfo,
-					SceneBlackboard,
+					SceneTextures,
 					DenoiserInputs,
 					RayTracingConfig);
 
diff --git a/Engine/Source/Runtime/Renderer/Private/ReflectionEnvironment.cpp b/Engine/Source/Runtime/Renderer/Private/ReflectionEnvironment.cpp
--- a/Engine/Source/Runtime/Renderer/Private/ReflectionEnvironment.cpp
+++ b/Engine/Source/Runtime/Renderer/Private/ReflectionEnvironment.cpp
@@ -29,7 +29,7 @@
 #include "LightPropagationVolumeSettings.h"
 #include "PipelineStateCache.h"
 #include "DistanceFieldAmbientOcclusion.h"
-#include "SceneViewFamilyBlackboard.h"
+#include "SceneTextureParameters.h"
 #include "ScreenSpaceDenoise.h"
 #include "ScreenSpaceRayTracing.h"
 #include "RayTracing/RaytracingOptions.h"
@@ -561,7 +561,7 @@ class FReflectionEnvironmentSkyLightingPS : public FGlobalShader
 		SHADER_PARAMETER_TEXTURE(Texture2D,    PreIntegratedGF)
 		SHADER_PARAMETER_SAMPLER(SamplerState, PreIntegratedGFSampler)
 		
-		SHADER_PARAMETER_STRUCT_INCLUDE(FSceneViewFamilyBlackboard, SceneTextures)
+		SHADER_PARAMETER_STRUCT_INCLUDE(FSceneTextureParameters, SceneTextures)
 		SHADER_PARAMETER_STRUCT_INCLUDE(FSceneTextureSamplerParameters, SceneTextureSamplers)
 
 		SHADER_PARAMETER_STRUCT_REF(FViewUniformShaderParameters, ViewUniformBuffer)
@@ -645,8 +645,8 @@ void FDeferredShadingSceneRenderer::RenderDeferredReflectionsAndSkyLighting(FRHI
 	FRDGTextureRef AmbientOcclusionTexture = GraphBuilder.RegisterExternalTexture(SceneContext.ScreenSpaceAO);
 	FRDGTextureRef DynamicBentNormalAOTexture = GraphBuilder.RegisterExternalTexture(DynamicBentNormalAO ? DynamicBentNormalAO : GSystemTextures.WhiteDummy);
 
-	FSceneViewFamilyBlackboard SceneTextures;
-	SetupSceneViewFamilyBlackboard(GraphBuilder, &SceneTextures);
+	FSceneTextureParameters SceneTextures;
+	SetupSceneTextureParameters(GraphBuilder, &SceneTextures);
 
 	for (FViewInfo& View : Views)
 	{
diff --git a/Engine/Source/Runtime/Renderer/Private/SceneTextureParameters.cpp b/Engine/Source/Runtime/Renderer/Private/SceneTextureParameters.cpp
new file mode 100644
--- /dev/null
+++ b/Engine/Source/Runtime/Renderer/Private/SceneTextureParameters.cpp
@@ -0,0 +1,68 @@
+// Copyright 1998-2019 Epic Games, Inc. All Rights Reserved.
+
+/*=============================================================================
+	DeferredShadingRenderer.h: Scene rendering definitions.
+=============================================================================*/
+
+#include "SceneTextureParameters.h"
+#include "SceneRendering.h"
+#include "SceneRenderTargets.h"
+#include "SystemTextures.h"
+
+
+void SetupSceneTextureParameters(
+	FRDGBuilder& GraphBuilder,
+	FSceneTextureParameters* OutTextures)
+{
+	const FSceneRenderTargets& SceneContext = FSceneRenderTargets::Get(GraphBuilder.RHICmdList);
+
+	*OutTextures = FSceneTextureParameters();
+
+	// Should always have a depth buffer around allocated, since early z-pass is first.
+	OutTextures->SceneDepthBuffer = GraphBuilder.RegisterExternalTexture(SceneContext.SceneDepthZ);
+
+	// Registers all the scene texture from the scene context. No fallback is provided to catch mistake at shader parameter validation time
+	// when a pass is trying to access a resource before any other pass actually created it.
+	OutTextures->SceneVelocityBuffer = SceneContext.SceneVelocity ? GraphBuilder.RegisterExternalTexture(SceneContext.SceneVelocity) : nullptr;
+	OutTextures->SceneGBufferA = SceneContext.GBufferA ? GraphBuilder.RegisterExternalTexture(SceneContext.GBufferA) : nullptr;
+	OutTextures->SceneGBufferB = SceneContext.GBufferB ? GraphBuilder.RegisterExternalTexture(SceneContext.GBufferB) : nullptr;
+	OutTextures->SceneGBufferC = SceneContext.GBufferC ? GraphBuilder.RegisterExternalTexture(SceneContext.GBufferC) : nullptr;
+	OutTextures->SceneGBufferD = SceneContext.GBufferD ? GraphBuilder.RegisterExternalTexture(SceneContext.GBufferD) : nullptr;
+	OutTextures->SceneGBufferE = SceneContext.GBufferE ? GraphBuilder.RegisterExternalTexture(SceneContext.GBufferE) : nullptr;
+
+	// Ligthing channels might be disabled when all lights are on the same channel.
+	if (SceneContext.LightingChannels)
+	{
+		OutTextures->SceneLightingChannels = GraphBuilder.RegisterExternalTexture(SceneContext.LightingChannels, TEXT("LightingChannels"));
+		OutTextures->bIsSceneLightingChannelsValid = true;
+	}
+	else
+	{
+		OutTextures->SceneLightingChannels = GraphBuilder.RegisterExternalTexture(GSystemTextures.WhiteDummy, TEXT("LightingChannels"));
+		OutTextures->bIsSceneLightingChannelsValid = false;
+	}
+}
+
+void SetupSceneTextureSamplers(FSceneTextureSamplerParameters* OutSamplers)
+{
+	FRHISamplerState* Sampler = TStaticSamplerState<SF_Point>::GetRHI();
+	OutSamplers->SceneDepthBufferSampler = Sampler;
+	OutSamplers->SceneVelocityBufferSampler = Sampler;
+	OutSamplers->SceneGBufferASampler = Sampler;
+	OutSamplers->SceneGBufferBSampler = Sampler;
+	OutSamplers->SceneGBufferCSampler = Sampler;
+	OutSamplers->SceneGBufferDSampler = Sampler;
+	OutSamplers->SceneGBufferESampler = Sampler;
+}
+
+FRDGTextureRef GetEyeAdaptationTexture(FRDGBuilder& GraphBuilder, const FViewInfo& View)
+{
+	if (View.HasValidEyeAdaptation())
+	{
+		return GraphBuilder.RegisterExternalTexture(View.GetEyeAdaptation(GraphBuilder.RHICmdList), TEXT("ViewEyeAdaptation"));
+	}
+	else
+	{
+		return GraphBuilder.RegisterExternalTexture(GSystemTextures.WhiteDummy, TEXT("DefaultViewEyeAdaptation"));
+	}
+}
diff --git a/Engine/Source/Runtime/Renderer/Private/SceneTextureParameters.h b/Engine/Source/Runtime/Renderer/Private/SceneTextureParameters.h
new file mode 100644
--- /dev/null
+++ b/Engine/Source/Runtime/Renderer/Private/SceneTextureParameters.h
@@ -0,0 +1,58 @@
+// Copyright 1998-2019 Epic Games, Inc. All Rights Reserved.
+
+/*=============================================================================
+	DeferredShadingRenderer.h: Scene rendering definitions.
+=============================================================================*/
+
+#pragma once
+
+#include "RenderGraph.h"
+
+
+class FViewInfo;
+
+
+/** Contains reference on all available buffer for a given scene. */
+// TODO: rename to FSceneTextureParemeters
+BEGIN_SHADER_PARAMETER_STRUCT(FSceneTextureParameters, )
+	// FSceneTextureParameters::SceneLightingChannels needs to be accessed with SceneLightingChannels.Load(), so a shader accessing
+	// needs to know when it not valid since SceneLightingChannels could end up being a dummy system texture.
+	SHADER_PARAMETER(uint32, bIsSceneLightingChannelsValid)
+
+	SHADER_PARAMETER_RDG_TEXTURE(Texture2D, SceneDepthBuffer)
+	SHADER_PARAMETER_RDG_TEXTURE(Texture2D, SceneVelocityBuffer)
+	SHADER_PARAMETER_RDG_TEXTURE(Texture2D, SceneGBufferA)
+	SHADER_PARAMETER_RDG_TEXTURE(Texture2D, SceneGBufferB)
+	SHADER_PARAMETER_RDG_TEXTURE(Texture2D, SceneGBufferC)
+	SHADER_PARAMETER_RDG_TEXTURE(Texture2D, SceneGBufferD)
+	SHADER_PARAMETER_RDG_TEXTURE(Texture2D, SceneGBufferE)
+	SHADER_PARAMETER_RDG_TEXTURE(Texture2D<uint>, SceneLightingChannels)
+END_SHADER_PARAMETER_STRUCT()
+
+
+/** Contains reference on all samplers for FSceneTextureParameters for platforms not supporting shared samplers. */
+BEGIN_SHADER_PARAMETER_STRUCT(FSceneTextureSamplerParameters, )
+	SHADER_PARAMETER_SAMPLER(SamplerState, SceneDepthBufferSampler)
+	SHADER_PARAMETER_SAMPLER(SamplerState, SceneVelocityBufferSampler)
+	SHADER_PARAMETER_SAMPLER(SamplerState, SceneGBufferASampler)
+	SHADER_PARAMETER_SAMPLER(SamplerState, SceneGBufferBSampler)
+	SHADER_PARAMETER_SAMPLER(SamplerState, SceneGBufferCSampler)
+	SHADER_PARAMETER_SAMPLER(SamplerState, SceneGBufferDSampler)
+	SHADER_PARAMETER_SAMPLER(SamplerState, SceneGBufferESampler)
+END_SHADER_PARAMETER_STRUCT()
+
+
+/** Sets up the the blackboard from available scene view family.
+ *
+ * Note: Once the entire renderer is built with a single render graph, would no longer need this function.
+ */
+void SetupSceneTextureParameters(
+	FRDGBuilder& GraphBuilder,
+	FSceneTextureParameters* OutTextures);
+
+/** Sets up all the samplers. */
+void SetupSceneTextureSamplers(FSceneTextureSamplerParameters* OutSamplers);
+
+/** Returns a render graph texture resource reference onto the eye adaptation or fallback.
+ */
+FRDGTextureRef GetEyeAdaptationTexture(FRDGBuilder& GraphBuilder, const FViewInfo& View);
diff --git a/Engine/Source/Runtime/Renderer/Private/SceneViewFamilyBlackboard.cpp b/Engine/Source/Runtime/Renderer/Private/SceneViewFamilyBlackboard.cpp
deleted file mode 100644
diff --git a/Engine/Source/Runtime/Renderer/Private/SceneViewFamilyBlackboard.h b/Engine/Source/Runtime/Renderer/Private/SceneViewFamilyBlackboard.h
deleted file mode 100644
diff --git a/Engine/Source/Runtime/Renderer/Private/ScreenSpaceDenoise.cpp b/Engine/Source/Runtime/Renderer/Private/ScreenSpaceDenoise.cpp
--- a/Engine/Source/Runtime/Renderer/Private/ScreenSpaceDenoise.cpp
+++ b/Engine/Source/Runtime/Renderer/Private/ScreenSpaceDenoise.cpp
@@ -12,7 +12,7 @@
 #include "ScenePrivate.h"
 #include "ClearQuad.h"
 #include "PipelineStateCache.h"
-#include "SceneViewFamilyBlackboard.h"
+#include "SceneTextureParameters.h"
 #include "BlueNoise.h"
 #include "Halton.h"
 
@@ -452,7 +452,7 @@ bool ShouldCompileSignalPipeline(ESignalProcessing SignalProcessing, EShaderPlat
 
 /** Shader parameter structure used for all shaders. */
 BEGIN_SHADER_PARAMETER_STRUCT(FSSDCommonParameters, )
-	SHADER_PARAMETER_STRUCT_INCLUDE(FSceneViewFamilyBlackboard, SceneBlackboard)
+	SHADER_PARAMETER_STRUCT_INCLUDE(FSceneTextureParameters, SceneTextures)
 	SHADER_PARAMETER_RDG_TEXTURE(Texture2D, EyeAdaptation)
 	SHADER_PARAMETER_RDG_TEXTURE(Texture2D<uint>, TileClassificationTexture)
 	SHADER_PARAMETER_STRUCT_REF(FViewUniformShaderParameters, ViewUniformBuffer)
@@ -776,7 +776,7 @@ static FSSDSignalTextures CopyAndBackfillSignalInput(const FSSDSignalTextures& S
 static void DenoiseSignalAtConstantPixelDensity(
 	FRDGBuilder& GraphBuilder,
 	const FViewInfo& View,
-	const FSceneViewFamilyBlackboard& SceneBlackboard,
+	const FSceneTextureParameters& SceneTextures,
 	const FSSDSignalTextures& InputSignal,
 	FSSDConstantPixelDensitySettings Settings,
 	TStaticArray<FScreenSpaceFilteringHistory*, IScreenSpaceDenoiser::kMaxBatchSize> PrevFilteringHistory,
@@ -824,7 +824,7 @@ static void DenoiseSignalAtConstantPixelDensity(
 		};
 
 		FRDGTextureDesc RefDesc = FRDGTextureDesc::Create2DDesc(
-			SceneBlackboard.SceneDepthBuffer->Desc.Extent,
+			SceneTextures.SceneDepthBuffer->Desc.Extent,
 			PF_Unknown,
 			FClearValueBinding::Black,
 			/* InFlags = */ TexCreate_None,
@@ -911,7 +911,7 @@ static void DenoiseSignalAtConstantPixelDensity(
 	// Setup common shader parameters.
 	FSSDCommonParameters CommonParameters;
 	{
-		CommonParameters.SceneBlackboard = SceneBlackboard;
+		CommonParameters.SceneTextures = SceneTextures;
 		CommonParameters.ViewUniformBuffer = View.ViewUniformBuffer;
 		CommonParameters.EyeAdaptation = GetEyeAdaptationTexture(GraphBuilder, View);
 	}
@@ -1249,20 +1249,20 @@ static void DenoiseSignalAtConstantPixelDensity(
 
 		// Keep depth buffer and GBuffer around for next frame.
 		{
-			GraphBuilder.QueueTextureExtraction(SceneBlackboard.SceneDepthBuffer, &View.ViewState->PrevFrameViewInfo.DepthBuffer);
+			GraphBuilder.QueueTextureExtraction(SceneTextures.SceneDepthBuffer, &View.ViewState->PrevFrameViewInfo.DepthBuffer);
 
 			// Requires the normal that are in GBuffer A.
 			if (Settings.SignalProcessing == ESignalProcessing::Reflections ||
 				Settings.SignalProcessing == ESignalProcessing::AmbientOcclusion ||
 				Settings.SignalProcessing == ESignalProcessing::DiffuseIndirect)
 			{
-				GraphBuilder.QueueTextureExtraction(SceneBlackboard.SceneGBufferA, &View.ViewState->PrevFrameViewInfo.GBufferA);
+				GraphBuilder.QueueTextureExtraction(SceneTextures.SceneGBufferA, &View.ViewState->PrevFrameViewInfo.GBufferA);
 			}
 
 			// Reflections requires the roughness that is in GBuffer B.
 			if (Settings.SignalProcessing == ESignalProcessing::Reflections)
 			{
-				GraphBuilder.QueueTextureExtraction(SceneBlackboard.SceneGBufferB, &View.ViewState->PrevFrameViewInfo.GBufferB);
+				GraphBuilder.QueueTextureExtraction(SceneTextures.SceneGBufferB, &View.ViewState->PrevFrameViewInfo.GBufferB);
 			}
 		}
 
@@ -1367,7 +1367,7 @@ public:
 		FRDGBuilder& GraphBuilder,
 		const FViewInfo& View,
 		FPreviousViewInfo* PreviousViewInfos,
-		const FSceneViewFamilyBlackboard& SceneBlackboard,
+		const FSceneTextureParameters& SceneTextures,
 		const TStaticArray<FShadowParameters, IScreenSpaceDenoiser::kMaxBatchSize>& InputParameters,
 		const int32 InputParameterCount,
 		TStaticArray<FShadowPenumbraOutputs, IScreenSpaceDenoiser::kMaxBatchSize>& Outputs) const override
@@ -1421,7 +1421,7 @@ public:
 
 		FSSDSignalTextures SignalOutput;
 		DenoiseSignalAtConstantPixelDensity(
-			GraphBuilder, View, SceneBlackboard,
+			GraphBuilder, View, SceneTextures,
 			InputSignal, Settings,
 			PrevHistories,
 			NewHistories,
@@ -1438,7 +1438,7 @@ public:
 		FRDGBuilder& GraphBuilder,
 		const FViewInfo& View,
 		FPreviousViewInfo* PreviousViewInfos,
-		const FSceneViewFamilyBlackboard& SceneBlackboard,
+		const FSceneTextureParameters& SceneTextures,
 		const FReflectionsInputs& ReflectionInputs,
 		const FReflectionsRayTracingConfig RayTracingConfig) const override
 	{
@@ -1463,7 +1463,7 @@ public:
 
 		FSSDSignalTextures SignalOutput;
 		DenoiseSignalAtConstantPixelDensity(
-			GraphBuilder, View, SceneBlackboard,
+			GraphBuilder, View, SceneTextures,
 			InputSignal, Settings,
 			PrevHistories,
 			NewHistories,
@@ -1478,7 +1478,7 @@ public:
 		FRDGBuilder& GraphBuilder,
 		const FViewInfo& View,
 		FPreviousViewInfo* PreviousViewInfos,
-		const FSceneViewFamilyBlackboard& SceneBlackboard,
+		const FSceneTextureParameters& SceneTextures,
 		const FAmbientOcclusionInputs& ReflectionInputs,
 		const FAmbientOcclusionRayTracingConfig RayTracingConfig) const override
 	{
@@ -1504,7 +1504,7 @@ public:
 
 		FSSDSignalTextures SignalOutput;
 		DenoiseSignalAtConstantPixelDensity(
-			GraphBuilder, View, SceneBlackboard,
+			GraphBuilder, View, SceneTextures,
 			InputSignal, Settings,
 			PrevHistories,
 			NewHistories,
@@ -1519,7 +1519,7 @@ public:
 		FRDGBuilder& GraphBuilder,
 		const FViewInfo& View,
 		FPreviousViewInfo* PreviousViewInfos,
-		const FSceneViewFamilyBlackboard& SceneBlackboard,
+		const FSceneTextureParameters& SceneTextures,
 		const FDiffuseIndirectInputs& Inputs,
 		const FAmbientOcclusionRayTracingConfig Config) const override
 	{
@@ -1546,7 +1546,7 @@ public:
 
 		FSSDSignalTextures SignalOutput;
 		DenoiseSignalAtConstantPixelDensity(
-			GraphBuilder, View, SceneBlackboard,
+			GraphBuilder, View, SceneTextures,
 			InputSignal, Settings,
 			PrevHistories,
 			NewHistories,
@@ -1561,7 +1561,7 @@ public:
 		FRDGBuilder& GraphBuilder,
 		const FViewInfo& View,
 		FPreviousViewInfo* PreviousViewInfos,
-		const FSceneViewFamilyBlackboard& SceneBlackboard,
+		const FSceneTextureParameters& SceneTextures,
 		const FDiffuseIndirectInputs& Inputs,
 		const FAmbientOcclusionRayTracingConfig Config) const override
 	{
@@ -1588,7 +1588,7 @@ public:
 
 		FSSDSignalTextures SignalOutput;
 		DenoiseSignalAtConstantPixelDensity(
-			GraphBuilder, View, SceneBlackboard,
+			GraphBuilder, View, SceneTextures,
 			InputSignal, Settings,
 			PrevHistories,
 			NewHistories,
diff --git a/Engine/Source/Runtime/Renderer/Private/ScreenSpaceDenoise.h b/Engine/Source/Runtime/Renderer/Private/ScreenSpaceDenoise.h
--- a/Engine/Source/Runtime/Renderer/Private/ScreenSpaceDenoise.h
+++ b/Engine/Source/Runtime/Renderer/Private/ScreenSpaceDenoise.h
@@ -8,7 +8,7 @@
 class FViewInfo;
 struct FPreviousViewInfo;
 class FLightSceneInfo;
-class FSceneViewFamilyBlackboard;
+class FSceneTextureParameters;
 
 
 /** Interface for denoiser to have all hook in the renderer. */
@@ -108,7 +108,7 @@ public:
 		FRDGBuilder& GraphBuilder,
 		const FViewInfo& View,
 		FPreviousViewInfo* PreviousViewInfos,
-		const FSceneViewFamilyBlackboard& SceneBlackboard,
+		const FSceneTextureParameters& SceneTextures,
 		const TStaticArray<FShadowParameters, IScreenSpaceDenoiser::kMaxBatchSize>& InputParameters,
 		const int32 InputParameterCount,
 		TStaticArray<FShadowPenumbraOutputs, IScreenSpaceDenoiser::kMaxBatchSize>& Outputs) const = 0;
@@ -130,7 +130,7 @@ public:
 		FRDGBuilder& GraphBuilder,
 		const FViewInfo& View,
 		FPreviousViewInfo* PreviousViewInfos,
-		const FSceneViewFamilyBlackboard& SceneBlackboard,
+		const FSceneTextureParameters& SceneTextures,
 		const FReflectionsInputs& ReflectionInputs,
 		const FReflectionsRayTracingConfig RayTracingConfig) const = 0;
 	
@@ -152,7 +152,7 @@ public:
 		FRDGBuilder& GraphBuilder,
 		const FViewInfo& View,
 		FPreviousViewInfo* PreviousViewInfos,
-		const FSceneViewFamilyBlackboard& SceneBlackboard,
+		const FSceneTextureParameters& SceneTextures,
 		const FAmbientOcclusionInputs& ReflectionInputs,
 		const FAmbientOcclusionRayTracingConfig RayTracingConfig) const = 0;
 
@@ -174,7 +174,7 @@ public:
 		FRDGBuilder& GraphBuilder,
 		const FViewInfo& View,
 		FPreviousViewInfo* PreviousViewInfos,
-		const FSceneViewFamilyBlackboard& SceneBlackboard,
+		const FSceneTextureParameters& SceneTextures,
 		const FDiffuseIndirectInputs& Inputs,
 		const FAmbientOcclusionRayTracingConfig Config) const = 0;
 
@@ -183,7 +183,7 @@ public:
 		FRDGBuilder& GraphBuilder,
 		const FViewInfo& View,
 		FPreviousViewInfo* PreviousViewInfos,
-		const FSceneViewFamilyBlackboard& SceneBlackboard,
+		const FSceneTextureParameters& SceneTextures,
 		const FDiffuseIndirectInputs& Inputs,
 		const FAmbientOcclusionRayTracingConfig Config) const = 0;
 
diff --git a/Engine/Source/Runtime/Renderer/Private/ScreenSpaceRayTracing.cpp b/Engine/Source/Runtime/Renderer/Private/ScreenSpaceRayTracing.cpp
--- a/Engine/Source/Runtime/Renderer/Private/ScreenSpaceRayTracing.cpp
+++ b/Engine/Source/Runtime/Renderer/Private/ScreenSpaceRayTracing.cpp
@@ -5,7 +5,7 @@
 #include "PixelShaderUtils.h"
 #include "ScreenPass.h"
 #include "ScenePrivate.h"
-#include "SceneViewFamilyBlackboard.h"
+#include "SceneTextureParameters.h"
 
 
 static TAutoConsoleVariable<int32> CVarSSRQuality(
@@ -154,7 +154,7 @@ FLinearColor ComputeSSRParams(const FViewInfo& View, ESSRQuality SSRQuality, boo
 
 BEGIN_SHADER_PARAMETER_STRUCT(FSSRCommonParameters, )
 	SHADER_PARAMETER(FLinearColor, SSRParams)
-	SHADER_PARAMETER_STRUCT_INCLUDE(FSceneViewFamilyBlackboard, SceneTextures)
+	SHADER_PARAMETER_STRUCT_INCLUDE(FSceneTextureParameters, SceneTextures)
 	SHADER_PARAMETER_STRUCT_INCLUDE(FSceneTextureSamplerParameters, SceneTextureSamplers)
 	SHADER_PARAMETER_STRUCT_REF(FViewUniformShaderParameters, ViewUniformBuffer)
 END_SHADER_PARAMETER_STRUCT()
@@ -322,7 +322,7 @@ void GetSSRQualityForView(const FViewInfo& View, ESSRQuality* OutQuality, IScree
 
 void RenderScreenSpaceReflections(
 	FRDGBuilder& GraphBuilder,
-	const FSceneViewFamilyBlackboard& SceneTextures,
+	const FSceneTextureParameters& SceneTextures,
 	const FRDGTextureRef CurrentSceneColor,
 	const FViewInfo& View,
 	ESSRQuality SSRQuality,
diff --git a/Engine/Source/Runtime/Renderer/Private/ScreenSpaceRayTracing.h b/Engine/Source/Runtime/Renderer/Private/ScreenSpaceRayTracing.h
--- a/Engine/Source/Runtime/Renderer/Private/ScreenSpaceRayTracing.h
+++ b/Engine/Source/Runtime/Renderer/Private/ScreenSpaceRayTracing.h
@@ -6,7 +6,7 @@
 #include "ScreenSpaceDenoise.h"
 
 class FViewInfo;
-class FSceneViewFamilyBlackboard;
+class FSceneTextureParameters;
 
 enum class ESSRQuality
 {
@@ -28,7 +28,7 @@ bool IsSSRTemporalPassRequired(const FViewInfo& View);
 
 void RenderScreenSpaceReflections(
 	FRDGBuilder& GraphBuilder,
-	const FSceneViewFamilyBlackboard& SceneBlackboard,
+	const FSceneTextureParameters& SceneTextures,
 	const FRDGTextureRef CurrentSceneColor,
 	const FViewInfo& View,
 	ESSRQuality SSRQuality,
