DR - Deprecate FVertexShaderRHIParamRef, FHullShaderRHIParamRef, FDomainShaderRHIParamRef, FPixelShaderRHIParamRef, FGeometryShaderRHIParamRef and FComputeShaderRHIParamRef

diff --git a/Engine/Source/Runtime/RenderCore/Public/ShaderParameterUtils.h b/Engine/Source/Runtime/RenderCore/Public/ShaderParameterUtils.h
--- a/Engine/Source/Runtime/RenderCore/Public/ShaderParameterUtils.h
+++ b/Engine/Source/Runtime/RenderCore/Public/ShaderParameterUtils.h
@@ -87,7 +87,7 @@ void SetShaderValue(
 	FRHICommandList& RHICmdList, 
 	ShaderRHIParamRef Shader,
 	const FShaderParameter& Parameter,
-	const bool& Value,
+	bool Value,
 	uint32 ElementIndex = 0
 	)
 {
@@ -101,7 +101,7 @@ void SetShaderValue(
 	FRHIAsyncComputeCommandList& RHICmdList,
 	ShaderRHIParamRef Shader,
 	const FShaderParameter& Parameter,
-	const bool& Value,
+	bool Value,
 	uint32 ElementIndex = 0
 	)
 {
@@ -162,7 +162,7 @@ void SetShaderValueArray(
  */
 inline void SetPixelShaderBool(
 	FRHICommandList& RHICmdList, 
-	FPixelShaderRHIParamRef PixelShader,
+	FRHIPixelShader* PixelShader,
 	const FShaderParameter& Parameter,
 	bool Value
 	)
@@ -187,10 +187,10 @@ inline void SetPixelShaderBool(
 /**
  * Sets the value of a shader texture parameter.  Template'd on shader type
  */
-template<typename ShaderTypeRHIParamRef, typename TRHICmdList>
+template<typename TRHIShader, typename TRHICmdList>
 FORCEINLINE void SetTextureParameter(
 	TRHICmdList& RHICmdList,
-	ShaderTypeRHIParamRef Shader,
+	TRHIShader* Shader,
 	const FShaderResourceParameter& TextureParameter,
 	const FShaderResourceParameter& SamplerParameter,
 	const FTexture* Texture,
@@ -228,10 +228,10 @@ FORCEINLINE void SetTextureParameter(
 /**
  * Sets the value of a shader texture parameter. Template'd on shader type.
  */
-template<typename ShaderTypeRHIParamRef, typename TRHICmdList>
+template<typename TRHIShader, typename TRHICmdList>
 FORCEINLINE void SetTextureParameter(
 	TRHICmdList& RHICmdList,
-	ShaderTypeRHIParamRef Shader,
+	TRHIShader* Shader,
 	const FShaderResourceParameter& TextureParameter,
 	const FShaderResourceParameter& SamplerParameter,
 	FRHISamplerState* SamplerStateRHI,
@@ -268,10 +268,10 @@ FORCEINLINE void SetTextureParameter(
  * Sets the value of a shader surface parameter (e.g. to access MSAA samples).
  * Template'd on shader type (e.g. pixel shader or compute shader).
  */
-template<typename ShaderTypeRHIParamRef, typename TRHICmdList>
+template<typename TRHIShader, typename TRHICmdList>
 FORCEINLINE void SetTextureParameter(
 	TRHICmdList& RHICmdList,
-	ShaderTypeRHIParamRef Shader,
+	TRHIShader* Shader,
 	const FShaderResourceParameter& Parameter,
 	FTextureRHIParamRef NewTextureRHI
 	)
@@ -289,10 +289,10 @@ FORCEINLINE void SetTextureParameter(
 /**
  * Sets the value of a shader sampler parameter. Template'd on shader type.
  */
-template<typename ShaderTypeRHIParamRef, typename TRHICmdList>
+template<typename TRHIShader, typename TRHICmdList>
 FORCEINLINE void SetSamplerParameter(
 	TRHICmdList& RHICmdList,
-	ShaderTypeRHIParamRef Shader,
+	TRHIShader* Shader,
 	const FShaderResourceParameter& Parameter,
 	FRHISamplerState* SamplerStateRHI
 	)
@@ -311,10 +311,10 @@ FORCEINLINE void SetSamplerParameter(
  * Sets the value of a shader resource view parameter
  * Template'd on shader type (e.g. pixel shader or compute shader).
  */
-template<typename ShaderTypeRHIParamRef, typename TRHICmdList>
+template<typename TRHIShader, typename TRHICmdList>
 FORCEINLINE void SetSRVParameter(
 	TRHICmdList& RHICmdList,
-	ShaderTypeRHIParamRef Shader,
+	TRHIShader* Shader,
 	const FShaderResourceParameter& Parameter,
 	FRHIShaderResourceView* NewShaderResourceViewRHI
 	)
@@ -329,13 +329,31 @@ FORCEINLINE void SetSRVParameter(
 	}
 }
 
+
+template<typename TRHIShader, typename TRHICmdList>
+FORCEINLINE void SetSRVParameter(
+	TRHICmdList& RHICmdList,
+	TRefCountPtr<TRHIShader> Shader,
+	const FShaderResourceParameter& Parameter,
+	FRHIShaderResourceView* NewShaderResourceViewRHI
+)
+{
+	if (Parameter.IsBound())
+	{
+		RHICmdList.SetShaderResourceViewParameter(
+			Shader.GetReference(),
+			Parameter.GetBaseIndex(),
+			NewShaderResourceViewRHI
+		);
+	}
+}
 /**
  * Sets the value of a unordered access view parameter
  */
 template<typename TRHICmdList>
 FORCEINLINE void SetUAVParameter(
-TRHICmdList& RHICmdList,
-	FComputeShaderRHIParamRef ComputeShader,
+	TRHICmdList& RHICmdList,
+	FRHIComputeShader* ComputeShader,
 	const FShaderResourceParameter& Parameter,
 	FRHIUnorderedAccessView* NewUnorderedAccessViewRHI
 	)
@@ -352,37 +370,37 @@ TRHICmdList& RHICmdList,
 
 
 template<typename TRHICmdList>
-inline bool SetUAVParameterIfCS(TRHICmdList& RHICmdList, const FVertexShaderRHIParamRef Shader, const FShaderResourceParameter& UAVParameter, FRHIUnorderedAccessView* UAV)
+inline bool SetUAVParameterIfCS(TRHICmdList& RHICmdList, FRHIVertexShader* Shader, const FShaderResourceParameter& UAVParameter, FRHIUnorderedAccessView* UAV)
 {
 	return false;
 }
 
 template<typename TRHICmdList>
-inline bool SetUAVParameterIfCS(TRHICmdList& RHICmdList, const FPixelShaderRHIParamRef Shader, const FShaderResourceParameter& UAVParameter, FRHIUnorderedAccessView* UAV)
+inline bool SetUAVParameterIfCS(TRHICmdList& RHICmdList, FRHIPixelShader* Shader, const FShaderResourceParameter& UAVParameter, FRHIUnorderedAccessView* UAV)
 {
 	return false;
 }
 
 template<typename TRHICmdList>
-inline bool SetUAVParameterIfCS(TRHICmdList& RHICmdList, const FHullShaderRHIParamRef Shader, const FShaderResourceParameter& UAVParameter, FRHIUnorderedAccessView* UAV)
+inline bool SetUAVParameterIfCS(TRHICmdList& RHICmdList, FRHIHullShader* Shader, const FShaderResourceParameter& UAVParameter, FRHIUnorderedAccessView* UAV)
 {
 	return false;
 }
 
 template<typename TRHICmdList>
-inline bool SetUAVParameterIfCS(TRHICmdList& RHICmdList, const FDomainShaderRHIParamRef Shader, const FShaderResourceParameter& UAVParameter, FRHIUnorderedAccessView* UAV)
+inline bool SetUAVParameterIfCS(TRHICmdList& RHICmdList, FRHIDomainShader* Shader, const FShaderResourceParameter& UAVParameter, FRHIUnorderedAccessView* UAV)
 {
 	return false;
 }
 
 template<typename TRHICmdList>
-inline bool SetUAVParameterIfCS(TRHICmdList& RHICmdList, const FGeometryShaderRHIParamRef Shader, const FShaderResourceParameter& UAVParameter, FRHIUnorderedAccessView* UAV)
+inline bool SetUAVParameterIfCS(TRHICmdList& RHICmdList, FRHIGeometryShader* Shader, const FShaderResourceParameter& UAVParameter, FRHIUnorderedAccessView* UAV)
 {
 	return false;
 }
 
 template<typename TRHICmdList>
-inline bool SetUAVParameterIfCS(TRHICmdList& RHICmdList, const FComputeShaderRHIParamRef Shader, const FShaderResourceParameter& UAVParameter, FRHIUnorderedAccessView* UAV)
+inline bool SetUAVParameterIfCS(TRHICmdList& RHICmdList, FRHIComputeShader* Shader, const FShaderResourceParameter& UAVParameter, FRHIUnorderedAccessView* UAV)
 {
 	SetUAVParameter(RHICmdList, Shader, UAVParameter, UAV);
 	return UAVParameter.IsBound();
@@ -416,7 +434,7 @@ inline void FRWShaderParameter::SetTexture(TRHICmdList& RHICmdList, TShaderRHIRe
 }
 
 template<typename TRHICmdList>
-inline void FRWShaderParameter::UnsetUAV(TRHICmdList& RHICmdList, FComputeShaderRHIParamRef ComputeShader) const
+inline void FRWShaderParameter::UnsetUAV(TRHICmdList& RHICmdList, FRHIComputeShader* ComputeShader) const
 {
 	SetUAVParameter(RHICmdList, ComputeShader,UAVParameter,FUnorderedAccessViewRHIRef());
 }
diff --git a/Engine/Source/Runtime/RenderCore/Public/ShaderParameters.h b/Engine/Source/Runtime/RenderCore/Public/ShaderParameters.h
--- a/Engine/Source/Runtime/RenderCore/Public/ShaderParameters.h
+++ b/Engine/Source/Runtime/RenderCore/Public/ShaderParameters.h
@@ -148,7 +148,7 @@ public:
 	inline void SetTexture(TRHICmdList& RHICmdList, TShaderRHIRef Shader, const FTextureRHIParamRef Texture, FRHIUnorderedAccessView* UAV) const;
 
 	template<typename TRHICmdList>
-	inline void UnsetUAV(TRHICmdList& RHICmdList, FComputeShaderRHIParamRef ComputeShader) const;
+	inline void UnsetUAV(TRHICmdList& RHICmdList, FRHIComputeShader* ComputeShader) const;
 
 private:
 
diff --git a/Engine/Source/Runtime/RenderCore/Public/StaticBoundShaderState.h b/Engine/Source/Runtime/RenderCore/Public/StaticBoundShaderState.h
--- a/Engine/Source/Runtime/RenderCore/Public/StaticBoundShaderState.h
+++ b/Engine/Source/Runtime/RenderCore/Public/StaticBoundShaderState.h
@@ -36,9 +36,9 @@ public:
 	 */
 	FRHIBoundShaderState* GetInitializedRHI(
 		FRHIVertexDeclaration* VertexDeclaration,
-		FVertexShaderRHIParamRef VertexShader, 
-		FPixelShaderRHIParamRef PixelShader,
-		FGeometryShaderRHIParamRef GeometryShader
+		FRHIVertexShader* VertexShader,
+		FRHIPixelShader* PixelShader,
+		FRHIGeometryShader* GeometryShader
 		);
 
 	/**
@@ -60,9 +60,9 @@ private:
 
 #if DO_CHECK
 	FRHIVertexDeclaration* BoundVertexDeclaration;
-	FVertexShaderRHIParamRef BoundVertexShader;
-	FPixelShaderRHIParamRef BoundPixelShader;
-	FGeometryShaderRHIParamRef BoundGeometryShader;
+	FRHIVertexShader* BoundVertexShader;
+	FRHIPixelShader* BoundPixelShader;
+	FRHIGeometryShader* BoundGeometryShader;
 #endif 
 };
 
diff --git a/Engine/Source/Runtime/Renderer/Private/AmbientCubemapParameters.h b/Engine/Source/Runtime/Renderer/Private/AmbientCubemapParameters.h
--- a/Engine/Source/Runtime/Renderer/Private/AmbientCubemapParameters.h
+++ b/Engine/Source/Runtime/Renderer/Private/AmbientCubemapParameters.h
@@ -19,8 +19,8 @@ public:
 
 	void Bind(const FShaderParameterMap& ParameterMap);
 
-	void SetParameters(FRHICommandList& RHICmdList, const FPixelShaderRHIParamRef ShaderRHI, const FFinalPostProcessSettings::FCubemapEntry& Entry) const;
-	void SetParameters(FRHICommandList& RHICmdList, const FComputeShaderRHIParamRef ShaderRHI, const FFinalPostProcessSettings::FCubemapEntry& Entry) const;
+	void SetParameters(FRHICommandList& RHICmdList, FRHIPixelShader* ShaderRHI, const FFinalPostProcessSettings::FCubemapEntry& Entry) const;
+	void SetParameters(FRHICommandList& RHICmdList, FRHIComputeShader* ShaderRHI, const FFinalPostProcessSettings::FCubemapEntry& Entry) const;
 
 	friend FArchive& operator<<(FArchive& Ar, FCubemapShaderParameters& P);
 
@@ -30,7 +30,7 @@ private:
 	FShaderResourceParameter AmbientCubemap;
 	FShaderResourceParameter AmbientCubemapSampler;
 
-	template<typename TShaderRHIRef>
-	void SetParametersTemplate(FRHICommandList& RHICmdList, const TShaderRHIRef& ShaderRHI, const FFinalPostProcessSettings::FCubemapEntry& Entry) const;
+	template<typename TRHIShader>
+	void SetParametersTemplate(FRHICommandList& RHICmdList, TRHIShader* ShaderRHI, const FFinalPostProcessSettings::FCubemapEntry& Entry) const;
 };
 
diff --git a/Engine/Source/Runtime/Renderer/Private/AtmosphereRendering.cpp b/Engine/Source/Runtime/Renderer/Private/AtmosphereRendering.cpp
--- a/Engine/Source/Runtime/Renderer/Private/AtmosphereRendering.cpp
+++ b/Engine/Source/Runtime/Renderer/Private/AtmosphereRendering.cpp
@@ -37,8 +37,8 @@ class FAtmosphereShaderTextureParameters
 public:
 	void Bind(const FShaderParameterMap& ParameterMap);
 
-	template< typename ShaderRHIParamRef >
-	FORCEINLINE_DEBUGGABLE void Set(FRHICommandList& RHICmdList, const ShaderRHIParamRef ShaderRHI, const FSceneView& View) const
+	template< typename TRHIShader >
+	FORCEINLINE_DEBUGGABLE void Set(FRHICommandList& RHICmdList, TRHIShader* ShaderRHI, const FSceneView& View) const
 	{
 		if (TransmittanceTexture.IsBound() || IrradianceTexture.IsBound() || InscatterTexture.IsBound())
 		{
@@ -72,15 +72,15 @@ void FAtmosphereShaderTextureParameters::Bind(const FShaderParameterMap& Paramet
 	InscatterTextureSampler.Bind(ParameterMap, TEXT("AtmosphereInscatterTextureSampler"));
 }
 
-#define IMPLEMENT_ATMOSPHERE_TEXTURE_PARAM_SET( ShaderRHIParamRef ) \
-	template void FAtmosphereShaderTextureParameters::Set< ShaderRHIParamRef >( FRHICommandList& RHICmdList, const ShaderRHIParamRef ShaderRHI, const FSceneView& View ) const;
+#define IMPLEMENT_ATMOSPHERE_TEXTURE_PARAM_SET( TRHIShader ) \
+	template void FAtmosphereShaderTextureParameters::Set< TRHIShader >( FRHICommandList& RHICmdList, TRHIShader* ShaderRHI, const FSceneView& View ) const;
 
-IMPLEMENT_ATMOSPHERE_TEXTURE_PARAM_SET(FVertexShaderRHIParamRef);
-IMPLEMENT_ATMOSPHERE_TEXTURE_PARAM_SET(FHullShaderRHIParamRef);
-IMPLEMENT_ATMOSPHERE_TEXTURE_PARAM_SET(FDomainShaderRHIParamRef);
-IMPLEMENT_ATMOSPHERE_TEXTURE_PARAM_SET(FGeometryShaderRHIParamRef);
-IMPLEMENT_ATMOSPHERE_TEXTURE_PARAM_SET(FPixelShaderRHIParamRef);
-IMPLEMENT_ATMOSPHERE_TEXTURE_PARAM_SET(FComputeShaderRHIParamRef);
+IMPLEMENT_ATMOSPHERE_TEXTURE_PARAM_SET(FRHIVertexShader);
+IMPLEMENT_ATMOSPHERE_TEXTURE_PARAM_SET(FRHIHullShader);
+IMPLEMENT_ATMOSPHERE_TEXTURE_PARAM_SET(FRHIDomainShader);
+IMPLEMENT_ATMOSPHERE_TEXTURE_PARAM_SET(FRHIGeometryShader);
+IMPLEMENT_ATMOSPHERE_TEXTURE_PARAM_SET(FRHIPixelShader);
+IMPLEMENT_ATMOSPHERE_TEXTURE_PARAM_SET(FRHIComputeShader);
 
 FArchive& operator<<(FArchive& Ar, FAtmosphereShaderTextureParameters& Parameters)
 {
diff --git a/Engine/Source/Runtime/Renderer/Private/ByteBuffer.cpp b/Engine/Source/Runtime/Renderer/Private/ByteBuffer.cpp
--- a/Engine/Source/Runtime/Renderer/Private/ByteBuffer.cpp
+++ b/Engine/Source/Runtime/Renderer/Private/ByteBuffer.cpp
@@ -65,7 +65,7 @@ void MemsetBuffer(FRHICommandList& RHICmdList, const FRWBufferStructured& DstBuf
 
 	TShaderMapRef< FMemsetBufferCS > ComputeShader( ShaderMap );
 
-	const FComputeShaderRHIParamRef ShaderRHI = ComputeShader->GetComputeShader();
+	FRHIComputeShader* ShaderRHI = ComputeShader->GetComputeShader();
 	RHICmdList.SetComputeShader( ShaderRHI );
 
 	SetShaderValue( RHICmdList, ShaderRHI, ComputeShader->Value, Value );
@@ -136,7 +136,7 @@ void MemcpyBuffer(FRHICommandList& RHICmdList, const FRWBufferStructured& SrcBuf
 
 	TShaderMapRef< FMemcpyBufferCS > ComputeShader( ShaderMap );
 		
-	const FComputeShaderRHIParamRef ShaderRHI = ComputeShader->GetComputeShader();
+	FRHIComputeShader* ShaderRHI = ComputeShader->GetComputeShader();
 	RHICmdList.SetComputeShader( ShaderRHI );
 
 	RHICmdList.TransitionResource(EResourceTransitionAccess::EWritable, EResourceTransitionPipeline::EGfxToCompute, DstBuffer.UAV);
@@ -280,7 +280,7 @@ void FScatterUploadBuilder::UploadTo(FRHICommandList& RHICmdList, FRWBufferStruc
 
 	TShaderMapRef<FScatterCopyCS> ComputeShader(ShaderMap);
 
-	const FComputeShaderRHIParamRef ShaderRHI = ComputeShader->GetComputeShader();
+	FRHIComputeShader* ShaderRHI = ComputeShader->GetComputeShader();
 	RHICmdList.SetComputeShader(ShaderRHI);
 
 	SetShaderValue(RHICmdList, ShaderRHI, ComputeShader->NumScatters, NumScatters);
@@ -305,7 +305,7 @@ void FScatterUploadBuilder::UploadTo_Flush(FRHICommandList& RHICmdList, FRWBuffe
 
 	TShaderMapRef<FScatterCopyCS> ComputeShader(ShaderMap);
 		
-	const FComputeShaderRHIParamRef ShaderRHI = ComputeShader->GetComputeShader();
+	FRHIComputeShader* ShaderRHI = ComputeShader->GetComputeShader();
 	RHICmdList.SetComputeShader(ShaderRHI);
 
 	SetShaderValue(RHICmdList, ShaderRHI, ComputeShader->NumScatters, NumScatters);
diff --git a/Engine/Source/Runtime/Renderer/Private/CapsuleShadowRendering.cpp b/Engine/Source/Runtime/Renderer/Private/CapsuleShadowRendering.cpp
--- a/Engine/Source/Runtime/Renderer/Private/CapsuleShadowRendering.cpp
+++ b/Engine/Source/Runtime/Renderer/Private/CapsuleShadowRendering.cpp
@@ -152,7 +152,7 @@ public:
 		FRHIShaderResourceView* LightDirectionDataSRV,
 		const FRWBuffer& ComputedLightDirectionDataBuffer)
 	{
-		const FComputeShaderRHIParamRef ShaderRHI = GetComputeShader();
+		FRHIComputeShader* ShaderRHI = GetComputeShader();
 
 		FGlobalShader::SetParameters<FViewUniformShaderParameters>(RHICmdList, ShaderRHI, View.ViewUniformBuffer);
 
@@ -329,7 +329,7 @@ public:
 		FRHIShaderResourceView* LightDirectionDataSRV,
 		FTextureRHIParamRef ReceiverBentNormalTextureValue)
 	{
-		const FComputeShaderRHIParamRef ShaderRHI = GetComputeShader();
+		FRHIComputeShader* ShaderRHI = GetComputeShader();
 
 		FGlobalShader::SetParameters<FViewUniformShaderParameters>(RHICmdList, ShaderRHI, View.ViewUniformBuffer);
 
@@ -592,7 +592,7 @@ public:
 
 	void SetParameters(FRHICommandList& RHICmdList, const FSceneView& View, FIntPoint TileDimensionsValue, const FIntRect& ScissorRect, const FRWBuffer& TileIntersectionCountsBuffer)
 	{
-		const FVertexShaderRHIParamRef ShaderRHI = GetVertexShader();
+		FRHIVertexShader* ShaderRHI = GetVertexShader();
 		FGlobalShader::SetParameters<FViewUniformShaderParameters>(RHICmdList, ShaderRHI, View.ViewUniformBuffer);
 
 		SetShaderValue(RHICmdList, ShaderRHI, TileDimensions, TileDimensionsValue);
@@ -659,7 +659,7 @@ public:
 
 	void SetParameters(FRHICommandList& RHICmdList, const FSceneView& View, const FIntRect& ScissorRect, TRefCountPtr<IPooledRenderTarget>& ShadowFactorsTextureValue, bool bOutputtingToLightAttenuation)
 	{
-		const FPixelShaderRHIParamRef ShaderRHI = GetPixelShader();
+		FRHIPixelShader* ShaderRHI = GetPixelShader();
 
 		FGlobalShader::SetParameters<FViewUniformShaderParameters>(RHICmdList, ShaderRHI, View.ViewUniformBuffer);
 		SceneTextureParameters.Set(RHICmdList, ShaderRHI, View.FeatureLevel, ESceneTextureSetupMode::All);
diff --git a/Engine/Source/Runtime/Renderer/Private/CompositionLighting/PostProcessAmbient.cpp b/Engine/Source/Runtime/Renderer/Private/CompositionLighting/PostProcessAmbient.cpp
--- a/Engine/Source/Runtime/Renderer/Private/CompositionLighting/PostProcessAmbient.cpp
+++ b/Engine/Source/Runtime/Renderer/Private/CompositionLighting/PostProcessAmbient.cpp
@@ -49,7 +49,7 @@ public:
 	template <typename TRHICmdList>
 	void SetParameters(TRHICmdList& RHICmdList, const FRenderingCompositePassContext& Context, const FFinalPostProcessSettings::FCubemapEntry& Entry)
 	{
-		const FPixelShaderRHIParamRef ShaderRHI = GetPixelShader();
+		FRHIPixelShader* ShaderRHI = GetPixelShader();
 
 		FGlobalShader::SetParameters<FViewUniformShaderParameters>(RHICmdList, ShaderRHI, Context.View.ViewUniformBuffer);
 		PostprocessParameter.SetPS(RHICmdList, ShaderRHI, Context, TStaticSamplerState<SF_Bilinear, AM_Clamp, AM_Clamp, AM_Clamp>::GetRHI());
@@ -167,18 +167,18 @@ void FCubemapShaderParameters::Bind(const FShaderParameterMap& ParameterMap)
 	AmbientCubemapSampler.Bind(ParameterMap, TEXT("AmbientCubemapSampler"));
 }
 
-void FCubemapShaderParameters::SetParameters(FRHICommandList& RHICmdList, const FPixelShaderRHIParamRef ShaderRHI, const FFinalPostProcessSettings::FCubemapEntry& Entry) const
+void FCubemapShaderParameters::SetParameters(FRHICommandList& RHICmdList, FRHIPixelShader* ShaderRHI, const FFinalPostProcessSettings::FCubemapEntry& Entry) const
 {
 	SetParametersTemplate(RHICmdList, ShaderRHI, Entry);
 }
 
-void FCubemapShaderParameters::SetParameters(FRHICommandList& RHICmdList, const FComputeShaderRHIParamRef ShaderRHI, const FFinalPostProcessSettings::FCubemapEntry& Entry) const
+void FCubemapShaderParameters::SetParameters(FRHICommandList& RHICmdList, FRHIComputeShader* ShaderRHI, const FFinalPostProcessSettings::FCubemapEntry& Entry) const
 {
 	SetParametersTemplate(RHICmdList, ShaderRHI, Entry);
 }
 
-template<typename TShaderRHIRef>
-void FCubemapShaderParameters::SetParametersTemplate(FRHICommandList& RHICmdList, const TShaderRHIRef& ShaderRHI, const FFinalPostProcessSettings::FCubemapEntry& Entry) const
+template<typename TRHIShader>
+void FCubemapShaderParameters::SetParametersTemplate(FRHICommandList& RHICmdList, TRHIShader* ShaderRHI, const FFinalPostProcessSettings::FCubemapEntry& Entry) const
 {
 	// floats to render the cubemap
 	{
diff --git a/Engine/Source/Runtime/Renderer/Private/CompositionLighting/PostProcessAmbientOcclusion.cpp b/Engine/Source/Runtime/Renderer/Private/CompositionLighting/PostProcessAmbientOcclusion.cpp
--- a/Engine/Source/Runtime/Renderer/Private/CompositionLighting/PostProcessAmbientOcclusion.cpp
+++ b/Engine/Source/Runtime/Renderer/Private/CompositionLighting/PostProcessAmbientOcclusion.cpp
@@ -319,7 +319,7 @@ public:
 	void SetParameters(const FRenderingCompositePassContext& Context)
 	{
 		const FFinalPostProcessSettings& Settings = Context.View.FinalPostProcessSettings;
-		const FPixelShaderRHIParamRef ShaderRHI = GetPixelShader();
+		FRHIPixelShader* ShaderRHI = GetPixelShader();
 
 		FGlobalShader::SetParameters<FViewUniformShaderParameters>(Context.RHICmdList, ShaderRHI, Context.View.ViewUniformBuffer);
 
@@ -524,7 +524,7 @@ public:
 		const FIntRect& OutputRect,
 		FRHIUnorderedAccessView* OutUAV)
 	{
-		const FComputeShaderRHIParamRef ShaderRHI = GetComputeShader();
+		FRHIComputeShader* ShaderRHI = GetComputeShader();
 
 		FGlobalShader::SetParameters<FViewUniformShaderParameters>(RHICmdList, ShaderRHI, Context.View.ViewUniformBuffer);
 
@@ -539,7 +539,7 @@ public:
 	template <typename TRHICmdList>
 	void UnsetParameters(TRHICmdList& RHICmdList)
 	{
-		const FComputeShaderRHIParamRef ShaderRHI = GetComputeShader();
+		FRHIComputeShader* ShaderRHI = GetComputeShader();
 		RHICmdList.SetUAVParameter(ShaderRHI, SSAOSmoothResult.GetBaseIndex(), nullptr);
 	}
 
@@ -739,7 +739,7 @@ public:
 		const FVector4 HZBRemappingValue = GetHZBValue(View);				
 		const FSceneRenderTargetItem& SSAORandomization = GSystemTextures.SSAORandomization->GetRenderTargetItem();
 
-		const FComputeShaderRHIParamRef ShaderRHI = GetComputeShader();
+		FRHIComputeShader* ShaderRHI = GetComputeShader();
 		FGlobalShader::SetParameters<FViewUniformShaderParameters>(RHICmdList, ShaderRHI, View.ViewUniformBuffer);			
 		RHICmdList.SetUAVParameter(ShaderRHI, OutTexture.GetBaseIndex(), OutUAV);
 
@@ -759,7 +759,7 @@ public:
 		const FVector4 HZBRemappingValue = GetHZBValue(View);
 		const FSceneRenderTargetItem& SSAORandomization = GSystemTextures.SSAORandomization->GetRenderTargetItem();
 
-		const FPixelShaderRHIParamRef ShaderRHI = GetPixelShader();
+		FRHIPixelShader* ShaderRHI = GetPixelShader();
 		FGlobalShader::SetParameters<FViewUniformShaderParameters>(RHICmdList, ShaderRHI, View.ViewUniformBuffer);
 
 		// SF_Point is better than bilinear to avoid halos around objects
@@ -774,7 +774,7 @@ public:
 	template <typename TRHICmdList>
 	void UnsetParameters(TRHICmdList& RHICmdList)
 	{
-		const FComputeShaderRHIParamRef ShaderRHI = GetComputeShader();
+		FRHIComputeShader* ShaderRHI = GetComputeShader();
 		RHICmdList.SetUAVParameter(ShaderRHI, OutTexture.GetBaseIndex(), NULL);
 	}
 	
@@ -1184,7 +1184,7 @@ public:
 	void SetParameters(TRHICmdList& RHICmdList, const FRenderingCompositePassContext& Context, FIntPoint InputTextureSize)
 	{
 		const FFinalPostProcessSettings& Settings = Context.View.FinalPostProcessSettings;
-		const FPixelShaderRHIParamRef ShaderRHI = GetPixelShader();
+		FRHIPixelShader* ShaderRHI = GetPixelShader();
 
 		FGlobalShader::SetParameters<FViewUniformShaderParameters>(RHICmdList, ShaderRHI, Context.View.ViewUniformBuffer);
 		PostprocessParameter.SetPS(RHICmdList, ShaderRHI, Context, TStaticSamplerState<SF_Point,AM_Clamp,AM_Clamp,AM_Clamp>::GetRHI());
diff --git a/Engine/Source/Runtime/Renderer/Private/CompositionLighting/PostProcessLpvIndirect.cpp b/Engine/Source/Runtime/Renderer/Private/CompositionLighting/PostProcessLpvIndirect.cpp
--- a/Engine/Source/Runtime/Renderer/Private/CompositionLighting/PostProcessLpvIndirect.cpp
+++ b/Engine/Source/Runtime/Renderer/Private/CompositionLighting/PostProcessLpvIndirect.cpp
@@ -83,7 +83,7 @@ public:
 		FLpvReadUniformBufferRef LpvUniformBuffer, 
 		const FRenderingCompositePassContext& Context )
 	{
-		const FPixelShaderRHIParamRef ShaderRHI = GetPixelShader();
+		FRHIPixelShader* ShaderRHI = GetPixelShader();
 
 		SetUniformBufferParameter(RHICmdList, ShaderRHI, GetUniformBufferParameter<FLpvReadUniformBufferParameters>(), LpvUniformBuffer);
 
@@ -181,7 +181,7 @@ public:
 		FLpvReadUniformBufferRef LpvUniformBuffer, 
 		const FRenderingCompositePassContext& Context )
 	{
-		const FPixelShaderRHIParamRef ShaderRHI = GetPixelShader();
+		FRHIPixelShader* ShaderRHI = GetPixelShader();
 
 		SetUniformBufferParameter(Context.RHICmdList, ShaderRHI, GetUniformBufferParameter<FLpvReadUniformBufferParameters>(), LpvUniformBuffer);
 
diff --git a/Engine/Source/Runtime/Renderer/Private/CompositionLighting/PostProcessPassThrough.cpp b/Engine/Source/Runtime/Renderer/Private/CompositionLighting/PostProcessPassThrough.cpp
--- a/Engine/Source/Runtime/Renderer/Private/CompositionLighting/PostProcessPassThrough.cpp
+++ b/Engine/Source/Runtime/Renderer/Private/CompositionLighting/PostProcessPassThrough.cpp
@@ -28,7 +28,7 @@ bool FPostProcessPassThroughPS::Serialize(FArchive& Ar)
 template <typename TRHICmdList>
 void FPostProcessPassThroughPS::SetParameters(TRHICmdList& RHICmdList, const FRenderingCompositePassContext& Context)
 {
-	const FPixelShaderRHIParamRef ShaderRHI = GetPixelShader();
+	FRHIPixelShader* ShaderRHI = GetPixelShader();
 
 	FGlobalShader::SetParameters<FViewUniformShaderParameters>(RHICmdList, ShaderRHI, Context.View.ViewUniformBuffer);
 
diff --git a/Engine/Source/Runtime/Renderer/Private/ConvertToUniformMesh.cpp b/Engine/Source/Runtime/Renderer/Private/ConvertToUniformMesh.cpp
--- a/Engine/Source/Runtime/Renderer/Private/ConvertToUniformMesh.cpp
+++ b/Engine/Source/Runtime/Renderer/Private/ConvertToUniformMesh.cpp
@@ -340,7 +340,7 @@ public:
 		const FMatrix& Instance0Transform
 		)
 	{
-		FComputeShaderRHIParamRef ShaderRHI = GetComputeShader();
+		FRHIComputeShader* ShaderRHI = GetComputeShader();
 		FMaterialShader::SetParameters(RHICmdList, ShaderRHI, MaterialProxy, *MaterialProxy->GetMaterial(View.GetFeatureLevel()), View, View.ViewUniformBuffer, ESceneTextureSetupMode::None);
 
 		SetUniformBufferParameter(RHICmdList, ShaderRHI,GetUniformBufferParameter<FPrimitiveUniformShaderParameters>(),PrimitiveUniformBuffer);
@@ -361,7 +361,7 @@ public:
 
 	void UnsetParameters(FRHICommandList& RHICmdList, FViewInfo& View)
 	{
-		FComputeShaderRHIParamRef ShaderRHI = GetComputeShader();
+		FRHIComputeShader* ShaderRHI = GetComputeShader();
 		SurfelBufferParameters.UnsetParameters(RHICmdList, ShaderRHI);
 
 		const FScene* Scene = (const FScene*)View.Family->Scene;
diff --git a/Engine/Source/Runtime/Renderer/Private/DistanceFieldAmbientOcclusion.cpp b/Engine/Source/Runtime/Renderer/Private/DistanceFieldAmbientOcclusion.cpp
--- a/Engine/Source/Runtime/Renderer/Private/DistanceFieldAmbientOcclusion.cpp
+++ b/Engine/Source/Runtime/Renderer/Private/DistanceFieldAmbientOcclusion.cpp
@@ -302,7 +302,7 @@ public:
 
 	void SetParameters(FRHICommandList& RHICmdList, const FSceneView& View, const FDistanceFieldAOParameters& Parameters)
 	{
-		const FPixelShaderRHIParamRef ShaderRHI = GetPixelShader();
+		FRHIPixelShader* ShaderRHI = GetPixelShader();
 
 		FGlobalShader::SetParameters<FViewUniformShaderParameters>(RHICmdList, ShaderRHI, View.ViewUniformBuffer);
 		AOParameters.Set(RHICmdList, ShaderRHI, Parameters);
@@ -357,7 +357,7 @@ public:
 
 	void SetParameters(FRHICommandList& RHICmdList, const FSceneView& View, FSceneRenderTargetItem& DistanceFieldNormalValue, const FDistanceFieldAOParameters& Parameters)
 	{
-		const FComputeShaderRHIParamRef ShaderRHI = GetComputeShader();
+		FRHIComputeShader* ShaderRHI = GetComputeShader();
 
 		FGlobalShader::SetParameters<FViewUniformShaderParameters>(RHICmdList, ShaderRHI, View.ViewUniformBuffer);
 
diff --git a/Engine/Source/Runtime/Renderer/Private/DistanceFieldAmbientOcclusion.h b/Engine/Source/Runtime/Renderer/Private/DistanceFieldAmbientOcclusion.h
--- a/Engine/Source/Runtime/Renderer/Private/DistanceFieldAmbientOcclusion.h
+++ b/Engine/Source/Runtime/Renderer/Private/DistanceFieldAmbientOcclusion.h
@@ -348,7 +348,7 @@ public:
 		AOBufferBilinearUVMax.Bind(ParameterMap, TEXT("AOBufferBilinearUVMax"));
 	}
 
-	void Set(FRHICommandList& RHICmdList, const FPixelShaderRHIParamRef& ShaderRHI, const FViewInfo& View, const TRefCountPtr<IPooledRenderTarget>& DistanceFieldAOBentNormal)
+	void Set(FRHICommandList& RHICmdList, FRHIPixelShader* ShaderRHI, const FViewInfo& View, const TRefCountPtr<IPooledRenderTarget>& DistanceFieldAOBentNormal)
 	{
 		FTextureRHIParamRef BentNormalAO = DistanceFieldAOBentNormal ? DistanceFieldAOBentNormal->GetRenderTargetItem().ShaderResourceTexture : GWhiteTexture->TextureRHI;
 		SetTextureParameter(RHICmdList, ShaderRHI, BentNormalAOTexture, BentNormalAOSampler, TStaticSamplerState<SF_Bilinear>::GetRHI(), BentNormalAO);
diff --git a/Engine/Source/Runtime/Renderer/Private/DistanceFieldGlobalIllumination.cpp b/Engine/Source/Runtime/Renderer/Private/DistanceFieldGlobalIllumination.cpp
--- a/Engine/Source/Runtime/Renderer/Private/DistanceFieldGlobalIllumination.cpp
+++ b/Engine/Source/Runtime/Renderer/Private/DistanceFieldGlobalIllumination.cpp
@@ -139,7 +139,7 @@ public:
 		const FMatrix& ShadowToWorldValue,
 		FLightTileIntersectionResources* TileIntersectionResources)
 	{
-		const FComputeShaderRHIParamRef ShaderRHI = GetComputeShader();
+		FRHIComputeShader* ShaderRHI = GetComputeShader();
 		FGlobalShader::SetParameters<FViewUniformShaderParameters>(RHICmdList, ShaderRHI, View.ViewUniformBuffer);
 		ObjectParameters.Set(RHICmdList, ShaderRHI, GShadowCulledObjectBuffers.Buffers);
 
@@ -243,7 +243,7 @@ public:
 
 	void SetParameters(FRHICommandList& RHICmdList, const FSceneView& View)
 	{
-		FComputeShaderRHIParamRef ShaderRHI = GetComputeShader();
+		FRHIComputeShader* ShaderRHI = GetComputeShader();
 		FGlobalShader::SetParameters<FViewUniformShaderParameters>(RHICmdList, ShaderRHI, View.ViewUniformBuffer);
 
 		RHICmdList.TransitionResource(EResourceTransitionAccess::ERWBarrier, EResourceTransitionPipeline::EComputeToCompute, GVPLResources.VPLDispatchIndirectBuffer.UAV);
@@ -308,7 +308,7 @@ public:
 
 	void SetParameters(FRHICommandList& RHICmdList, const FScene* Scene, const FSceneView& View, const FDistanceFieldAOParameters& Parameters)
 	{
-		FComputeShaderRHIParamRef ShaderRHI = GetComputeShader();
+		FRHIComputeShader* ShaderRHI = GetComputeShader();
 		FGlobalShader::SetParameters<FViewUniformShaderParameters>(RHICmdList, ShaderRHI, View.ViewUniformBuffer);
 
 		FRHIUnorderedAccessView* OutUAVs[2];
@@ -575,7 +575,7 @@ public:
 
 	void SetParameters(FRHICommandList& RHICmdList, const FSceneView& View)
 	{
-		FComputeShaderRHIParamRef ShaderRHI = GetComputeShader();
+		FRHIComputeShader* ShaderRHI = GetComputeShader();
 		FGlobalShader::SetParameters<FViewUniformShaderParameters>(RHICmdList, ShaderRHI, View.ViewUniformBuffer);
 
 		RHICmdList.TransitionResource(EResourceTransitionAccess::ERWBarrier, EResourceTransitionPipeline::EComputeToCompute, GAOCulledObjectBuffers.Buffers.ObjectIndirectDispatch.UAV);
@@ -587,7 +587,7 @@ public:
 
 	void UnsetParameters(FRHICommandList& RHICmdList)
 	{
-		FComputeShaderRHIParamRef ShaderRHI = GetComputeShader();
+		FRHIComputeShader* ShaderRHI = GetComputeShader();
 		DispatchParameters.UnsetUAV(RHICmdList, ShaderRHI);
 
 		ObjectParameters.UnsetParameters(RHICmdList, ShaderRHI);
@@ -660,7 +660,7 @@ public:
 		const FDistanceFieldAOParameters& Parameters,
 		FLightTileIntersectionResources* TileIntersectionResources)
 	{
-		const FComputeShaderRHIParamRef ShaderRHI = GetComputeShader();
+		FRHIComputeShader* ShaderRHI = GetComputeShader();
 
 		FGlobalShader::SetParameters<FViewUniformShaderParameters>(RHICmdList, ShaderRHI, View.ViewUniformBuffer);
 		AOParameters.Set(RHICmdList, ShaderRHI, Parameters);
@@ -909,7 +909,7 @@ public:
 		FSceneRenderTargetItem& DistanceFieldNormal, 
 		const FAOScreenGridResources& ScreenGridResources)
 	{
-		FComputeShaderRHIParamRef ShaderRHI = GetComputeShader();
+		FRHIComputeShader* ShaderRHI = GetComputeShader();
 		FGlobalShader::SetParameters<FViewUniformShaderParameters>(RHICmdList, ShaderRHI, View.ViewUniformBuffer);
 		ScreenGridParameters.Set(RHICmdList, ShaderRHI, View, DistanceFieldNormal);
 
@@ -1015,7 +1015,7 @@ public:
 		FSceneRenderTargetItem& DistanceFieldNormal, 
 		const FDistanceFieldAOParameters& Parameters)
 	{
-		FComputeShaderRHIParamRef ShaderRHI = GetComputeShader();
+		FRHIComputeShader* ShaderRHI = GetComputeShader();
 		FGlobalShader::SetParameters<FViewUniformShaderParameters>(RHICmdList, ShaderRHI, View.ViewUniformBuffer);
 		SceneTextureParameters.Set(RHICmdList, ShaderRHI, View.FeatureLevel, ESceneTextureSetupMode::All);
 
@@ -1118,7 +1118,7 @@ public:
 		const FAOScreenGridResources& ScreenGridResources,
 		FSceneRenderTargetItem& IrradianceTextureValue)
 	{
-		FComputeShaderRHIParamRef ShaderRHI = GetComputeShader();
+		FRHIComputeShader* ShaderRHI = GetComputeShader();
 		FGlobalShader::SetParameters<FViewUniformShaderParameters>(RHICmdList, ShaderRHI, View.ViewUniformBuffer);
 
 		SetSRVParameter(RHICmdList, ShaderRHI, SurfelIrradiance, ScreenGridResources.SurfelIrradiance.SRV);
diff --git a/Engine/Source/Runtime/Renderer/Private/DistanceFieldLightingPost.cpp b/Engine/Source/Runtime/Renderer/Private/DistanceFieldLightingPost.cpp
--- a/Engine/Source/Runtime/Renderer/Private/DistanceFieldLightingPost.cpp
+++ b/Engine/Source/Runtime/Renderer/Private/DistanceFieldLightingPost.cpp
@@ -83,7 +83,7 @@ public:
 		JitterOffset.Bind(ParameterMap, TEXT("JitterOffset"));
 	}
 
-	void Set(FRHICommandList& RHICmdList, const FPixelShaderRHIParamRef& ShaderRHI, const FViewInfo& View, FSceneRenderTargetItem& DistanceFieldNormal, FSceneRenderTargetItem& DistanceFieldAOBentNormal)
+	void Set(FRHICommandList& RHICmdList, FRHIPixelShader* ShaderRHI, const FViewInfo& View, FSceneRenderTargetItem& DistanceFieldNormal, FSceneRenderTargetItem& DistanceFieldAOBentNormal)
 	{
 		SetTextureParameter(
 			RHICmdList,
@@ -204,7 +204,7 @@ public:
 		IPooledRenderTarget* VelocityTextureValue,
 		const FDistanceFieldAOParameters& Parameters)
 	{
-		const FPixelShaderRHIParamRef ShaderRHI = GetPixelShader();
+		FRHIPixelShader* ShaderRHI = GetPixelShader();
 
 		FGlobalShader::SetParameters<FViewUniformShaderParameters>(RHICmdList, ShaderRHI, View.ViewUniformBuffer);
 		SceneTextureParameters.Set(RHICmdList, ShaderRHI, View.FeatureLevel, ESceneTextureSetupMode::All);
@@ -333,7 +333,7 @@ public:
 		FSceneRenderTargetItem& DistanceFieldNormal, 
 		FSceneRenderTargetItem& BentNormalHistoryTextureValue)
 	{
-		const FPixelShaderRHIParamRef ShaderRHI = GetPixelShader();
+		FRHIPixelShader* ShaderRHI = GetPixelShader();
 
 		FGlobalShader::SetParameters<FViewUniformShaderParameters>(RHICmdList, ShaderRHI, View.ViewUniformBuffer);
 
@@ -439,7 +439,7 @@ public:
 		FSceneRenderTargetItem& DistanceFieldAOBentNormal,
 		const FDistanceFieldAOParameters& Parameters)
 	{
-		const FPixelShaderRHIParamRef ShaderRHI = GetPixelShader();
+		FRHIPixelShader* ShaderRHI = GetPixelShader();
 
 		FGlobalShader::SetParameters<FViewUniformShaderParameters>(RHICmdList, ShaderRHI, View.ViewUniformBuffer);
 
@@ -725,7 +725,7 @@ public:
 
 	void SetParameters(FRHICommandList& RHICmdList, const FViewInfo& View, TRefCountPtr<IPooledRenderTarget>& DistanceFieldAOBentNormal)
 	{
-		const FPixelShaderRHIParamRef ShaderRHI = GetPixelShader();
+		FRHIPixelShader* ShaderRHI = GetPixelShader();
 
 		FGlobalShader::SetParameters<FViewUniformShaderParameters>(RHICmdList, ShaderRHI, View.ViewUniformBuffer);
 		SceneTextureParameters.Set(RHICmdList, ShaderRHI, View.FeatureLevel, ESceneTextureSetupMode::All);
diff --git a/Engine/Source/Runtime/Renderer/Private/DistanceFieldLightingShared.h b/Engine/Source/Runtime/Renderer/Private/DistanceFieldLightingShared.h
--- a/Engine/Source/Runtime/Renderer/Private/DistanceFieldLightingShared.h
+++ b/Engine/Source/Runtime/Renderer/Private/DistanceFieldLightingShared.h
@@ -364,8 +364,8 @@ public:
 		DistanceFieldAtlasTexelSize.Bind(ParameterMap, TEXT("DistanceFieldAtlasTexelSize"));
 	}
 
-	template<typename TParamRef, typename TRHICommandList>
-	void Set(TRHICommandList& RHICmdList, const TParamRef& ShaderRHI, const FDistanceFieldCulledObjectBuffers& ObjectBuffers)
+	template<typename TShaderRHI, typename TRHICommandList>
+	void Set(TRHICommandList& RHICmdList, TShaderRHI* ShaderRHI, const FDistanceFieldCulledObjectBuffers& ObjectBuffers)
 	{
 		ObjectIndirectArguments.SetBuffer(RHICmdList, ShaderRHI, ObjectBuffers.ObjectIndirectArguments);
 		CulledObjectBounds.SetBuffer(RHICmdList, ShaderRHI, ObjectBuffers.Bounds);
@@ -393,8 +393,8 @@ public:
 		SetShaderValue(RHICmdList, ShaderRHI, DistanceFieldAtlasTexelSize, InvTextureDim);
 	}
 
-	template<typename TParamRef, typename TRHICommandList>
-	void UnsetParameters(TRHICommandList& RHICmdList, const TParamRef& ShaderRHI)
+	template<typename TRHIShader, typename TRHICommandList>
+	void UnsetParameters(TRHICommandList& RHICmdList, TRHIShader* ShaderRHI)
 	{
 		ObjectIndirectArguments.UnsetUAV(RHICmdList, ShaderRHI);
 		CulledObjectBounds.UnsetUAV(RHICmdList, ShaderRHI);
diff --git a/Engine/Source/Runtime/Renderer/Private/DistanceFieldObjectCulling.cpp b/Engine/Source/Runtime/Renderer/Private/DistanceFieldObjectCulling.cpp
--- a/Engine/Source/Runtime/Renderer/Private/DistanceFieldObjectCulling.cpp
+++ b/Engine/Source/Runtime/Renderer/Private/DistanceFieldObjectCulling.cpp
@@ -145,7 +145,7 @@ public:
 		OutUAVs[5] = Scene->DistanceFieldSceneData.ObjectBuffers->Bounds.UAV;
 		RHICmdList.TransitionResources(EResourceTransitionAccess::ERWBarrier, EResourceTransitionPipeline::EComputeToCompute, OutUAVs, ARRAY_COUNT(OutUAVs));
 
-		FComputeShaderRHIParamRef ShaderRHI = GetComputeShader();
+		FRHIComputeShader* ShaderRHI = GetComputeShader();
 		FGlobalShader::SetParameters<FViewUniformShaderParameters>(RHICmdList, ShaderRHI, View.ViewUniformBuffer);
 		ObjectBufferParameters.Set(RHICmdList, ShaderRHI, *(Scene->DistanceFieldSceneData.ObjectBuffers), Scene->DistanceFieldSceneData.NumObjectsInBuffer);
 		CulledObjectParameters.Set(RHICmdList, ShaderRHI, GAOCulledObjectBuffers.Buffers);
@@ -263,7 +263,7 @@ public:
 	}
 	void SetParameters(FRHICommandList& RHICmdList, const FViewInfo& View, FSceneRenderTargetItem& DistanceFieldNormal, FScene* Scene, FVector2D NumGroupsValue, const FDistanceFieldAOParameters& Parameters)
 	{
-		FComputeShaderRHIParamRef ShaderRHI = GetComputeShader();
+		FRHIComputeShader* ShaderRHI = GetComputeShader();
 
 		FGlobalShader::SetParameters<FViewUniformShaderParameters>(RHICmdList, ShaderRHI, View.ViewUniformBuffer);
 		SceneTextureParameters.Set(RHICmdList, ShaderRHI, View.FeatureLevel, ESceneTextureSetupMode::All);
@@ -357,7 +357,7 @@ public:
 
 	void SetParameters(FRHICommandList& RHICmdList, const FSceneView& View, const FDistanceFieldAOParameters& Parameters)
 	{
-		const FVertexShaderRHIParamRef ShaderRHI = GetVertexShader();
+		FRHIVertexShader* ShaderRHI = GetVertexShader();
 		FGlobalShader::SetParameters<FViewUniformShaderParameters>(RHICmdList, ShaderRHI, View.ViewUniformBuffer);
 		
 		ObjectParameters.Set(RHICmdList, ShaderRHI, GAOCulledObjectBuffers.Buffers);
@@ -426,7 +426,7 @@ public:
 
 	void SetParameters(FRHICommandList& RHICmdList, const FSceneView& View, FVector2D NumGroupsValue, const FDistanceFieldAOParameters& Parameters)
 	{
-		const FPixelShaderRHIParamRef ShaderRHI = GetPixelShader();
+		FRHIPixelShader* ShaderRHI = GetPixelShader();
 		FGlobalShader::SetParameters<FViewUniformShaderParameters>(RHICmdList, ShaderRHI, View.ViewUniformBuffer);
 		ObjectParameters.Set(RHICmdList, ShaderRHI, GAOCulledObjectBuffers.Buffers);
 		AOParameters.Set(RHICmdList, ShaderRHI, Parameters);
@@ -505,7 +505,7 @@ public:
 	}
 	void SetParameters(FRHICommandList& RHICmdList, const FSceneView& View)
 	{
-		FComputeShaderRHIParamRef ShaderRHI = GetComputeShader();
+		FRHIComputeShader* ShaderRHI = GetComputeShader();
 
 		FGlobalShader::SetParameters<FViewUniformShaderParameters>(RHICmdList, ShaderRHI, View.ViewUniformBuffer);
 
diff --git a/Engine/Source/Runtime/Renderer/Private/DistanceFieldObjectManagement.cpp b/Engine/Source/Runtime/Renderer/Private/DistanceFieldObjectManagement.cpp
--- a/Engine/Source/Runtime/Renderer/Private/DistanceFieldObjectManagement.cpp
+++ b/Engine/Source/Runtime/Renderer/Private/DistanceFieldObjectManagement.cpp
@@ -178,7 +178,7 @@ public:
 
 	void SetParameters(FRHICommandList& RHICmdList, const FScene* Scene, uint32 NumUploadOperationsValue, FRHIShaderResourceView* InUploadOperationIndices, FRHIShaderResourceView* InUploadOperationData)
 	{
-		FComputeShaderRHIParamRef ShaderRHI = GetComputeShader();
+		FRHIComputeShader* ShaderRHI = GetComputeShader();
 
 		SetShaderValue(RHICmdList, ShaderRHI, NumUploadOperations, NumUploadOperationsValue);
 
@@ -250,7 +250,7 @@ public:
 
 	void SetParameters(FRHICommandList& RHICmdList, FDistanceFieldObjectBuffers& ObjectBuffersSource, FDistanceFieldObjectBuffers& ObjectBuffersDest, int32 NumObjectsValue)
 	{
-		FComputeShaderRHIParamRef ShaderRHI = GetComputeShader();
+		FRHIComputeShader* ShaderRHI = GetComputeShader();
 
 		FRHIUnorderedAccessView* OutUAVs[2];
 		OutUAVs[0] = ObjectBuffersDest.Bounds.UAV;
@@ -324,7 +324,7 @@ public:
 
 	void SetParameters(FRHICommandList& RHICmdList, const FSurfelBuffers& SurfelBuffersSource, const FInstancedSurfelBuffers& InstancedSurfelBuffersSource, FSurfelBuffers& SurfelBuffersDest, int32 NumSurfelsValue)
 	{
-		FComputeShaderRHIParamRef ShaderRHI = GetComputeShader();
+		FRHIComputeShader* ShaderRHI = GetComputeShader();
 
 		FRHIUnorderedAccessView* OutUAVs[2];
 		OutUAVs[0] = SurfelBuffersDest.InterpolatedVertexData.UAV;
@@ -400,7 +400,7 @@ public:
 
 	void SetParameters(FRHICommandList& RHICmdList, const FSurfelBuffers& SurfelBuffersSource, const FInstancedSurfelBuffers& InstancedSurfelBuffersSource, FInstancedSurfelBuffers& InstancedSurfelBuffersDest, int32 NumSurfelsValue)
 	{
-		FComputeShaderRHIParamRef ShaderRHI = GetComputeShader();
+		FRHIComputeShader* ShaderRHI = GetComputeShader();
 
 		RHICmdList.TransitionResource(EResourceTransitionAccess::ERWBarrier, EResourceTransitionPipeline::EComputeToCompute, InstancedSurfelBuffersDest.VPLFlux.UAV);
 		CopyVPLFlux.SetBuffer(RHICmdList, ShaderRHI, InstancedSurfelBuffersDest.VPLFlux);
@@ -473,7 +473,7 @@ public:
 		FRHIShaderResourceView* InObjectBounds2,
 		FRHIShaderResourceView* InObjectData2)
 	{
-		FComputeShaderRHIParamRef ShaderRHI = GetComputeShader();
+		FRHIComputeShader* ShaderRHI = GetComputeShader();
 
 		SetShaderValue(RHICmdList, ShaderRHI, NumRemoveOperations, NumRemoveOperationsValue);
 		SetSRVParameter(RHICmdList, ShaderRHI, RemoveOperationIndices, InRemoveOperationIndices);
diff --git a/Engine/Source/Runtime/Renderer/Private/DistanceFieldScreenGridLighting.cpp b/Engine/Source/Runtime/Renderer/Private/DistanceFieldScreenGridLighting.cpp
--- a/Engine/Source/Runtime/Renderer/Private/DistanceFieldScreenGridLighting.cpp
+++ b/Engine/Source/Runtime/Renderer/Private/DistanceFieldScreenGridLighting.cpp
@@ -128,7 +128,7 @@ public:
 		const FDistanceFieldAOParameters& Parameters,
 		const FGlobalDistanceFieldInfo& GlobalDistanceFieldInfo)
 	{
-		FComputeShaderRHIParamRef ShaderRHI = GetComputeShader();
+		FRHIComputeShader* ShaderRHI = GetComputeShader();
 		FGlobalShader::SetParameters<FViewUniformShaderParameters>(RHICmdList, ShaderRHI, View.ViewUniformBuffer);
 		SceneTextureParameters.Set(RHICmdList, ShaderRHI, View.FeatureLevel, ESceneTextureSetupMode::All);
 		ObjectParameters.Set(RHICmdList, ShaderRHI, GAOCulledObjectBuffers.Buffers);
@@ -279,7 +279,7 @@ public:
 		const FDistanceFieldAOParameters& Parameters,
 		const FGlobalDistanceFieldInfo& GlobalDistanceFieldInfo)
 	{
-		FComputeShaderRHIParamRef ShaderRHI = GetComputeShader();
+		FRHIComputeShader* ShaderRHI = GetComputeShader();
 		FGlobalShader::SetParameters<FViewUniformShaderParameters>(RHICmdList, ShaderRHI, View.ViewUniformBuffer);
 		SceneTextureParameters.Set(RHICmdList, ShaderRHI, View.FeatureLevel, ESceneTextureSetupMode::All);
 		ObjectParameters.Set(RHICmdList, ShaderRHI, GAOCulledObjectBuffers.Buffers);
@@ -419,7 +419,7 @@ public:
 		FSceneRenderTargetItem& DistanceFieldNormal, 
 		FSceneRenderTargetItem& DownsampledBentNormal)
 	{
-		FComputeShaderRHIParamRef ShaderRHI = GetComputeShader();
+		FRHIComputeShader* ShaderRHI = GetComputeShader();
 		FGlobalShader::SetParameters<FViewUniformShaderParameters>(RHICmdList, ShaderRHI, View.ViewUniformBuffer);
 		ScreenGridParameters.Set(RHICmdList, ShaderRHI, View, DistanceFieldNormal);
 
diff --git a/Engine/Source/Runtime/Renderer/Private/DistanceFieldShadowing.cpp b/Engine/Source/Runtime/Renderer/Private/DistanceFieldShadowing.cpp
--- a/Engine/Source/Runtime/Renderer/Private/DistanceFieldShadowing.cpp
+++ b/Engine/Source/Runtime/Renderer/Private/DistanceFieldShadowing.cpp
@@ -134,7 +134,7 @@ public:
 
 	void SetParameters(FRHICommandList& RHICmdList, const FScene* Scene, const FSceneView& View, const FMatrix& WorldToShadowValue, int32 NumPlanes, const FPlane* PlaneData, const FVector4& ShadowBoundingSphereValue)
 	{
-		FComputeShaderRHIParamRef ShaderRHI = GetComputeShader();
+		FRHIComputeShader* ShaderRHI = GetComputeShader();
 		FGlobalShader::SetParameters<FViewUniformShaderParameters>(RHICmdList, ShaderRHI, View.ViewUniformBuffer);
 		ObjectBufferParameters.Set(RHICmdList, ShaderRHI, *(Scene->DistanceFieldSceneData.ObjectBuffers), Scene->DistanceFieldSceneData.NumObjectsInBuffer);
 
@@ -224,7 +224,7 @@ public:
 
 	void SetParameters(FRHICommandList& RHICmdList, const FSceneView& View, FVector2D NumGroupsValue, const FMatrix& WorldToShadowMatrixValue, float ShadowRadius)
 	{
-		const FVertexShaderRHIParamRef ShaderRHI = GetVertexShader();
+		FRHIVertexShader* ShaderRHI = GetVertexShader();
 		FGlobalShader::SetParameters<FViewUniformShaderParameters>(RHICmdList, ShaderRHI, View.ViewUniformBuffer);
 		
 		ObjectParameters.Set(RHICmdList, ShaderRHI, GShadowCulledObjectBuffers.Buffers);
@@ -286,7 +286,7 @@ public:
 		const FSceneView& View, 
 		FLightTileIntersectionResources* TileIntersectionResources)
 	{
-		const FPixelShaderRHIParamRef ShaderRHI = GetPixelShader();
+		FRHIPixelShader* ShaderRHI = GetPixelShader();
 		FGlobalShader::SetParameters<FViewUniformShaderParameters>(RHICmdList, ShaderRHI, View.ViewUniformBuffer);
 
 		ObjectParameters.Set(RHICmdList, ShaderRHI, GShadowCulledObjectBuffers.Buffers);
@@ -378,7 +378,7 @@ public:
 		const FIntRect& ScissorRect,
 		FLightTileIntersectionResources* TileIntersectionResources)
 	{
-		const FComputeShaderRHIParamRef ShaderRHI = GetComputeShader();
+		FRHIComputeShader* ShaderRHI = GetComputeShader();
 
 		FGlobalShader::SetParameters<FViewUniformShaderParameters>(RHICmdList, ShaderRHI, View.ViewUniformBuffer);
 
@@ -545,7 +545,7 @@ public:
 
 	void SetParameters(FRHICommandList& RHICmdList, const FSceneView& View, const FProjectedShadowInfo* ShadowInfo, const FIntRect& ScissorRect, TRefCountPtr<IPooledRenderTarget>& ShadowFactorsTextureValue)
 	{
-		const FPixelShaderRHIParamRef ShaderRHI = GetPixelShader();
+		FRHIPixelShader* ShaderRHI = GetPixelShader();
 
 		FGlobalShader::SetParameters<FViewUniformShaderParameters>(RHICmdList, ShaderRHI, View.ViewUniformBuffer);
 		SceneTextureParameters.Set(RHICmdList, ShaderRHI, View.FeatureLevel, ESceneTextureSetupMode::All);
@@ -637,7 +637,7 @@ public:
 	}
 	void SetParameters(FRHICommandList& RHICmdList, const FSceneView& View, FLightTileIntersectionResources* TileIntersectionResources)
 	{
-		FComputeShaderRHIParamRef ShaderRHI = GetComputeShader();
+		FRHIComputeShader* ShaderRHI = GetComputeShader();
 
 		FGlobalShader::SetParameters<FViewUniformShaderParameters>(RHICmdList, ShaderRHI, View.ViewUniformBuffer);
 
diff --git a/Engine/Source/Runtime/Renderer/Private/DistanceFieldVisualization.cpp b/Engine/Source/Runtime/Renderer/Private/DistanceFieldVisualization.cpp
--- a/Engine/Source/Runtime/Renderer/Private/DistanceFieldVisualization.cpp
+++ b/Engine/Source/Runtime/Renderer/Private/DistanceFieldVisualization.cpp
@@ -59,7 +59,7 @@ public:
 		const FDistanceFieldAOParameters& Parameters,
 		const FGlobalDistanceFieldInfo& GlobalDistanceFieldInfo)
 	{
-		const FComputeShaderRHIParamRef ShaderRHI = GetComputeShader();
+		FRHIComputeShader* ShaderRHI = GetComputeShader();
 
 		FGlobalShader::SetParameters<FViewUniformShaderParameters>(RHICmdList, ShaderRHI, View.ViewUniformBuffer);
 
@@ -139,7 +139,7 @@ public:
 
 	void SetParameters(FRHICommandList& RHICmdList, const FSceneView& View, TRefCountPtr<IPooledRenderTarget>& VisualizeDistanceField)
 	{
-		const FPixelShaderRHIParamRef ShaderRHI = GetPixelShader();
+		FRHIPixelShader* ShaderRHI = GetPixelShader();
 
 		FGlobalShader::SetParameters<FViewUniformShaderParameters>(RHICmdList, ShaderRHI, View.ViewUniformBuffer);
 		SceneTextureParameters.Set(RHICmdList, ShaderRHI, View.FeatureLevel, ESceneTextureSetupMode::All);
diff --git a/Engine/Source/Runtime/Renderer/Private/GPUBenchmark.cpp b/Engine/Source/Runtime/Renderer/Private/GPUBenchmark.cpp
--- a/Engine/Source/Runtime/Renderer/Private/GPUBenchmark.cpp
+++ b/Engine/Source/Runtime/Renderer/Private/GPUBenchmark.cpp
@@ -70,7 +70,7 @@ public:
 
 	void SetParameters(FRHICommandList& RHICmdList, const FSceneView& View, TRefCountPtr<IPooledRenderTarget>& Src)
 	{
-		const FPixelShaderRHIParamRef ShaderRHI = GetPixelShader();
+		FRHIPixelShader* ShaderRHI = GetPixelShader();
 
 		FGlobalShader::SetParameters<FViewUniformShaderParameters>(RHICmdList, ShaderRHI, View.ViewUniformBuffer);
 
@@ -133,7 +133,7 @@ public:
 
 	void SetParameters(FRHICommandList& RHICmdList, const FSceneView& View)
 	{
-		const FVertexShaderRHIParamRef ShaderRHI = GetVertexShader();
+		FRHIVertexShader* ShaderRHI = GetVertexShader();
 		
 		FGlobalShader::SetParameters<FViewUniformShaderParameters>(RHICmdList, ShaderRHI, View.ViewUniformBuffer);
 	}
diff --git a/Engine/Source/Runtime/Renderer/Private/GPUFastFourierTransform.cpp b/Engine/Source/Runtime/Renderer/Private/GPUFastFourierTransform.cpp
--- a/Engine/Source/Runtime/Renderer/Private/GPUFastFourierTransform.cpp
+++ b/Engine/Source/Runtime/Renderer/Private/GPUFastFourierTransform.cpp
@@ -215,7 +215,7 @@ namespace GPUFFT
 		{
 			using GPUFFTComputeShaderUtils::FComputeParamterValueSetter;
 
-			const FComputeShaderRHIParamRef ShaderRHI = GetComputeShader();
+			FRHIComputeShader* ShaderRHI = GetComputeShader();
 
 			// Set up the input.  We have to do this explicitly because the FFT dispatches multiple compute shaders and manages their input/output.
 			FComputeParamterValueSetter ParamSetter(RHICmdList, ShaderRHI);
@@ -240,7 +240,7 @@ namespace GPUFFT
 		}
 
 		// Method for use with the FScopedUAVBind
-		FShaderResourceParameter& DestinationResourceParamter() { return DstRWTexture; }
+		FShaderResourceParameter& DestinationResourceParameter() { return DstRWTexture; }
 
 
 		// FShader interface.
@@ -330,7 +330,7 @@ namespace GPUFFT
 
 			using GPUFFTComputeShaderUtils::FComputeParamterValueSetter;
 
-			const FComputeShaderRHIParamRef ShaderRHI = GetComputeShader();
+			FRHIComputeShader* ShaderRHI = GetComputeShader();
 
 			// Set up the input.  We have to do this explicitly because the FFT dispatches multiple compute shaders and manages their input/output.
 			FComputeParamterValueSetter ParamSetter(RHICmdList, ShaderRHI);
@@ -347,7 +347,7 @@ namespace GPUFFT
 		}
 
 		// Method for use with the FScopedUAVBind
-		FShaderResourceParameter& DestinationResourceParamter() { return DstRWTexture; }
+		FShaderResourceParameter& DestinationResourceParameter() { return DstRWTexture; }
 
 
 		// FShader interface.
@@ -426,7 +426,7 @@ namespace GPUFFT
 
 			using GPUFFTComputeShaderUtils::FComputeParamterValueSetter;
 
-			const FComputeShaderRHIParamRef ShaderRHI = GetComputeShader();
+			FRHIComputeShader* ShaderRHI = GetComputeShader();
 
 			// Set up the input.  We have to do this explicitly because the FFT dispatches multiple compute shaders and manages their input/output.
 			FComputeParamterValueSetter ParamSetter(RHICmdList, ShaderRHI);
@@ -448,7 +448,7 @@ namespace GPUFFT
 		}
 
 		// Method for use with the FScopedUAVBind
-		FShaderResourceParameter& DestinationResourceParamter() { return DstRWTexture; }
+		FShaderResourceParameter& DestinationResourceParameter() { return DstRWTexture; }
 
 
 		// FShader interface.
@@ -524,7 +524,7 @@ namespace GPUFFT
 		{
 			using GPUFFTComputeShaderUtils::FComputeParamterValueSetter;
 
-			const FComputeShaderRHIParamRef ShaderRHI = GetComputeShader();
+			FRHIComputeShader* ShaderRHI = GetComputeShader();
 
 			// Set up the input.  We have to do this explicitly because the FFT dispatches multiple compute shaders and manages their input/output.
 			FComputeParamterValueSetter ParamSetter(RHICmdList, ShaderRHI);
@@ -539,7 +539,7 @@ namespace GPUFFT
 		}
 
 		// Method for use with the FScopedUAVBind
-		FShaderResourceParameter& DestinationResourceParamter() { return DstRWTexture; }
+		FShaderResourceParameter& DestinationResourceParameter() { return DstRWTexture; }
 
 
 		// FShader interface.
@@ -613,7 +613,7 @@ namespace GPUFFT
 		{
 			using GPUFFTComputeShaderUtils::FComputeParamterValueSetter;
 
-			const FComputeShaderRHIParamRef ShaderRHI = GetComputeShader();
+			FRHIComputeShader* ShaderRHI = GetComputeShader();
 
 			// Set up the input.  We have to do this explicitly because the FFT dispatches multiple compute shaders and manages their input/output.
 			FComputeParamterValueSetter ParamSetter(RHICmdList, ShaderRHI);
@@ -625,7 +625,7 @@ namespace GPUFFT
 		}
 
 		// Method for use with the FScopedUAVBind
-		FShaderResourceParameter& DestinationResourceParamter() { return DstRWTexture; }
+		FShaderResourceParameter& DestinationResourceParameter() { return DstRWTexture; }
 
 
 		// FShader interface.
@@ -697,7 +697,7 @@ namespace GPUFFT
 		{
 			using GPUFFTComputeShaderUtils::FComputeParamterValueSetter;
 
-			const FComputeShaderRHIParamRef ShaderRHI = GetComputeShader();
+			FRHIComputeShader* ShaderRHI = GetComputeShader();
 
 			const uint32 DataLayoutValue = (bHorizontalScanlines) ? 1 : 0;
 			// Set up the input.  We have to do this explicitly because the FFT dispatches multiple compute shaders and manages their input/output.
@@ -711,7 +711,7 @@ namespace GPUFFT
 		}
 
 		// Method for use with the FScopedUAVBind
-		FShaderResourceParameter& DestinationResourceParamter() { return DstRWTexture; }
+		FShaderResourceParameter& DestinationResourceParameter() { return DstRWTexture; }
 
 
 		// FShader interface.
@@ -769,7 +769,7 @@ namespace GPUFFT
 		{
 			using GPUFFTComputeShaderUtils::FComputeParamterValueSetter;
 
-			const FComputeShaderRHIParamRef ShaderRHI = GetComputeShader();
+			FRHIComputeShader* ShaderRHI = GetComputeShader();
 
 			// Set up the input.  We have to do this explicitly because the FFT dispatches multiple compute shaders and manages their input/output.
 			FComputeParamterValueSetter ParamSetter(RHICmdList, ShaderRHI);
@@ -795,7 +795,7 @@ namespace GPUFFT
 		}
 
 		// Method for use with the FScopedUAVBind
-		FShaderResourceParameter& DestinationResourceParamter() { return DstRWTexture; }
+		FShaderResourceParameter& DestinationResourceParameter() { return DstRWTexture; }
 
 
 		// FShader interface.
@@ -940,7 +940,7 @@ namespace GPUFFT
 				TransformTypeValue |= 8;
 			}
 
-			const FComputeShaderRHIParamRef ShaderRHI = GetComputeShader();
+			FRHIComputeShader* ShaderRHI = GetComputeShader();
 
 			// Set up the input.  We have to do this explicitly because the FFT dispatches multiple compute shaders and manages their input/output.
 
@@ -955,7 +955,7 @@ namespace GPUFFT
 		}
 
 		// Method for use with the FScopedUAVBind
-		FShaderResourceParameter& DestinationResourceParamter() { return DstRWTexture; }
+		FShaderResourceParameter& DestinationResourceParameter() { return DstRWTexture; }
 
 
 		// FShader interface.
@@ -1016,7 +1016,7 @@ namespace GPUFFT
 			FGSConvolutionBaseCS::SetCSParamters(RHICmdList, XFormType, SrcTexture, SrcRect, DstExtentValue);
 
 			// additional source input for sampling the spectral texture
-			const FComputeShaderRHIParamRef ShaderRHI = GetComputeShader();
+			FRHIComputeShader* ShaderRHI = GetComputeShader();
 			GPUFFTComputeShaderUtils::FComputeParamterValueSetter ParamSetter(RHICmdList, ShaderRHI);
 			// set the texture
 			ParamSetter(FilterSrcROTexture, FilterSrcTexture);
@@ -1309,7 +1309,7 @@ namespace
 		RHICmdList.SetComputeShader(ComputeShader->GetComputeShader());
 
 		// Bind output
-		FScopedUAVBind ScopedBind = FScopedUAVBind::BindOutput(RHICmdList, *ComputeShader, DstUAV);
+		FScopedUAVBind ScopedBind = FScopedUAVBind::BindOutput(RHICmdList, ComputeShader, DstUAV);
 
 		ComputeShader->SetCSParamters(RHICmdList, SrcWindow, SrcTexture,  DstWindow, PreFilter);
 
@@ -1370,7 +1370,7 @@ namespace
 		RHICmdList.SetComputeShader(ComputeShader->GetComputeShader());
 
 		// Bind output
-		FScopedUAVBind ScopedBind = FScopedUAVBind::BindOutput(RHICmdList, *ComputeShader, DstUAV);
+		FScopedUAVBind ScopedBind = FScopedUAVBind::BindOutput(RHICmdList, ComputeShader, DstUAV);
 
 		ComputeShader->SetCSParamters(RHICmdList, bHorizontalScanlines, SrcWindow, SrcTexture, KnlTexure);
 		
@@ -1439,7 +1439,7 @@ namespace
 		RHICmdList.SetComputeShader(ComputeShader->GetComputeShader());
 
 		// Bind output
-		FScopedUAVBind ScopedBind = FScopedUAVBind::BindOutput(RHICmdList, *ComputeShader, DstUAV);
+		FScopedUAVBind ScopedBind = FScopedUAVBind::BindOutput(RHICmdList, ComputeShader, DstUAV);
 
 		ComputeShader->SetCSParamters(RHICmdList, FFTDesc.XFormType, SrcTexture, FIntRect(FIntPoint(0, 0), DstExtent));
 
@@ -1490,7 +1490,7 @@ namespace
 		RHICmdList.SetComputeShader(ComputeShader->GetComputeShader());
 
 		// Bind output
-		FScopedUAVBind ScopedBind = FScopedUAVBind::BindOutput(RHICmdList, *ComputeShader, DstUAV);
+		FScopedUAVBind ScopedBind = FScopedUAVBind::BindOutput(RHICmdList, ComputeShader, DstUAV);
 
 		ComputeShader->SetCSParamters(RHICmdList, FFTDesc.XFormType, SrcTexture, SrcWindow, FIntRect(FIntPoint(0, 0), DstExtent), TransformLength, PassLength, bScrubNaNs);
 
@@ -1543,7 +1543,7 @@ namespace
 		RHICmdList.SetComputeShader(ComputeShader->GetComputeShader());
 
 		// Bind output
-		FScopedUAVBind ScopedBind = FScopedUAVBind::BindOutput(RHICmdList, *ComputeShader, DstUAV);
+		FScopedUAVBind ScopedBind = FScopedUAVBind::BindOutput(RHICmdList, ComputeShader, DstUAV);
 
 		ComputeShader->SetCSParamters(RHICmdList, FFTDesc.XFormType, SrcTexture, SrcWindow, DstWindow, TransformLength, SubLength, bScrubNaNs);
 			
@@ -1606,7 +1606,7 @@ namespace
 		RHICmdList.SetComputeShader(ComputeShader->GetComputeShader());
 
 		// Bind output
-		FScopedUAVBind ScopedBind = FScopedUAVBind::BindOutput(RHICmdList, *ComputeShader, DstUAV);
+		FScopedUAVBind ScopedBind = FScopedUAVBind::BindOutput(RHICmdList, ComputeShader, DstUAV);
 
 		ComputeShader->SetCSParamters(RHICmdList, FFTDesc.XFormType, TransformLength, SubPassWindow, SrcTexture, NumSubRegions);
 		// The number of signals to transform simultaneously (i.e. number of scan lines)
@@ -1654,7 +1654,7 @@ namespace
 		RHICmdList.SetComputeShader(ComputeShader->GetComputeShader());
 
 		// Bind output
-		FScopedUAVBind ScopedBind = FScopedUAVBind::BindOutput(RHICmdList, *ComputeShader, DstUAV);
+		FScopedUAVBind ScopedBind = FScopedUAVBind::BindOutput(RHICmdList, ComputeShader, DstUAV);
 
 		ComputeShader->SetCSParamters(RHICmdList, FFTDesc.XFormType, SrcTexture, SrcRect, FIntRect(FIntPoint(0, 0), DstExtent));
 
@@ -1706,7 +1706,7 @@ namespace
 		RHICmdList.SetComputeShader(ComputeShader->GetComputeShader());
 
 
-		FScopedUAVBind ScopedBind = FScopedUAVBind::BindOutput(RHICmdList, *ComputeShader, DstUAV);
+		FScopedUAVBind ScopedBind = FScopedUAVBind::BindOutput(RHICmdList, ComputeShader, DstUAV);
 		ComputeShader->SetCSParamters(RHICmdList, FFTDesc.XFormType, SrcTexture, SrcRect, DstRect, PreFilter);
 
 
@@ -1763,7 +1763,7 @@ namespace
 		UnbindRenderTargets(RHICmdList);
 		RHICmdList.SetComputeShader(ComputeShader->GetComputeShader());
 		
-		FScopedUAVBind ScopedBind = FScopedUAVBind::BindOutput(RHICmdList, *ComputeShader, DstUAV);
+		FScopedUAVBind ScopedBind = FScopedUAVBind::BindOutput(RHICmdList, ComputeShader, DstUAV);
 
 		ComputeShader->SetCSParamters(RHICmdList, FFTDesc.XFormType, PreTransformedKernel, SrcTexture, SrcRect, SrcRect.Size());
 
diff --git a/Engine/Source/Runtime/Renderer/Private/GPUFastFourierTransform.h b/Engine/Source/Runtime/Renderer/Private/GPUFastFourierTransform.h
--- a/Engine/Source/Runtime/Renderer/Private/GPUFastFourierTransform.h
+++ b/Engine/Source/Runtime/Renderer/Private/GPUFastFourierTransform.h
@@ -502,21 +502,21 @@ namespace GPUFFTComputeShaderUtils
 	class FComputeParamterValueSetter
 	{
 	public:
-		FComputeParamterValueSetter(FRHICommandList& CmdList, const FComputeShaderRHIParamRef& ShaderRHIParamRef) :
-			RHICmdList(&CmdList), ShaderRHI(&ShaderRHIParamRef) {}
+		FComputeParamterValueSetter(FRHICommandList& CmdList, FRHIComputeShader* InShaderRHI) :
+			RHICmdList(&CmdList), ShaderRHI(InShaderRHI) {}
 
 
 		template < typename TValue >
 		inline FComputeParamterValueSetter& operator()(FShaderParameter& Parameter, const TValue& Value)
 		{
-			SetShaderValue(*RHICmdList, *ShaderRHI, Parameter, Value);
+			SetShaderValue(*RHICmdList, ShaderRHI, Parameter, Value);
 			return *this;
 		}
 
 		template <typename TValue>
 		inline FComputeParamterValueSetter& operator()(FShaderParameter& Parameter, const TValue* Value, const uint32 NumElements)
 		{
-			SetShaderValueArray(*RHICmdList, *ShaderRHI, Parameter, Value, NumElements);
+			SetShaderValueArray(*RHICmdList, ShaderRHI, Parameter, Value, NumElements);
 			return *this;
 		}
 
@@ -524,7 +524,7 @@ namespace GPUFFTComputeShaderUtils
 		{
 			if (TextureParameter.IsBound())
 			{
-				SetTextureParameter(*RHICmdList, *ShaderRHI, TextureParameter, TextureRHI);
+				SetTextureParameter(*RHICmdList, ShaderRHI, TextureParameter, TextureRHI);
 			}
 			return *this;
 		}
@@ -533,14 +533,14 @@ namespace GPUFFTComputeShaderUtils
 		template <ESamplerFilter Filter, ESamplerAddressMode AddressMode>
 		inline FComputeParamterValueSetter& Set(FShaderResourceParameter& TextureParameter, FShaderResourceParameter& SamplerParameter, const FTextureRHIParamRef& TextureRHI)
 		{
-			SetTextureParameter(*RHICmdList, *ShaderRHI, TextureParameter, SamplerParameter,
+			SetTextureParameter(*RHICmdList, ShaderRHI, TextureParameter, SamplerParameter,
 				TStaticSamplerState<Filter, AddressMode, AddressMode, AddressMode>::GetRHI(), TextureRHI);
 			return *this;
 		}
 
 	private:
 		FRHICommandList* RHICmdList;
-		const FComputeShaderRHIParamRef* ShaderRHI;
+		FRHIComputeShader* ShaderRHI;
 	};
 
 	class FComputeParameterBinder
@@ -571,10 +571,10 @@ namespace GPUFFTComputeShaderUtils
 
 		// Factory method.
 		template <typename ComputeShaderT>
-		inline static FScopedUAVBind BindOutput(FRHICommandListImmediate& CmdList, ComputeShaderT& Shader, const FUnorderedAccessViewRHIRef& BufferUAV)
+		inline static FScopedUAVBind BindOutput(FRHICommandListImmediate& CmdList, ComputeShaderT* Shader, const FUnorderedAccessViewRHIRef& BufferUAV)
 		{
-			const FComputeShaderRHIParamRef ShaderRHI = Shader.GetComputeShader();
-			const uint32 BaseIndex = Shader.DestinationResourceParamter().GetBaseIndex();
+			FRHIComputeShader* ShaderRHI = Shader->GetComputeShader();
+			const uint32 BaseIndex = Shader->DestinationResourceParameter().GetBaseIndex();
 			return FScopedUAVBind(CmdList, ShaderRHI, BaseIndex, BufferUAV);
 		}
 
@@ -589,14 +589,14 @@ namespace GPUFFTComputeShaderUtils
 	private:
 		FScopedUAVBind();
 
-		FScopedUAVBind(FRHICommandListImmediate& CmdList, const FComputeShaderRHIParamRef ShaderRHI, const uint32 BaseIndex, const FUnorderedAccessViewRHIRef& BufferUAV)
+		FScopedUAVBind(FRHICommandListImmediate& CmdList, FRHIComputeShader* ShaderRHI, const uint32 BaseIndex, const FUnorderedAccessViewRHIRef& BufferUAV)
 			:ComputeShader(ShaderRHI), RHICmdList(&CmdList), Index(BaseIndex)
 		{
 			RHICmdList->SetUAVParameter(ComputeShader, Index, BufferUAV);
 		}
 
 
-		const FComputeShaderRHIParamRef ComputeShader;
+		FRHIComputeShader* ComputeShader;
 		FRHICommandListImmediate* RHICmdList;
 		uint32 Index;
 	};
diff --git a/Engine/Source/Runtime/Renderer/Private/GammaCorrection.cpp b/Engine/Source/Runtime/Renderer/Private/GammaCorrection.cpp
--- a/Engine/Source/Runtime/Renderer/Private/GammaCorrection.cpp
+++ b/Engine/Source/Runtime/Renderer/Private/GammaCorrection.cpp
@@ -140,7 +140,7 @@ void FSceneRenderer::GammaCorrectToViewportRenderTarget(FRHICommandList& RHICmdL
 		InvDisplayGamma = 1 / OverrideGamma;
 	}
 
-	const FPixelShaderRHIParamRef ShaderRHI = PixelShader->GetPixelShader();
+	FRHIPixelShader* ShaderRHI = PixelShader->GetPixelShader();
 
 	SetShaderValue(
 		RHICmdList, 
diff --git a/Engine/Source/Runtime/Renderer/Private/GlobalDistanceField.cpp b/Engine/Source/Runtime/Renderer/Private/GlobalDistanceField.cpp
--- a/Engine/Source/Runtime/Renderer/Private/GlobalDistanceField.cpp
+++ b/Engine/Source/Runtime/Renderer/Private/GlobalDistanceField.cpp
@@ -182,7 +182,7 @@ public:
 
 	void SetParameters(FRHICommandList& RHICmdList, const FScene* Scene, const FSceneView& View, float MaxOcclusionDistance, const FVector4& VolumeBoundsValue, FGlobalDFCacheType CacheType)
 	{
-		FComputeShaderRHIParamRef ShaderRHI = GetComputeShader();
+		FRHIComputeShader* ShaderRHI = GetComputeShader();
 		FGlobalShader::SetParameters<FViewUniformShaderParameters>(RHICmdList, ShaderRHI, View.ViewUniformBuffer);
 		ObjectBufferParameters.Set(RHICmdList, ShaderRHI, *(Scene->DistanceFieldSceneData.ObjectBuffers), Scene->DistanceFieldSceneData.NumObjectsInBuffer);
 
@@ -358,7 +358,7 @@ public:
 		int32 ClipmapIndexValue,
 		const FVolumeUpdateRegion& UpdateRegion)
 	{
-		FComputeShaderRHIParamRef ShaderRHI = GetComputeShader();
+		FRHIComputeShader* ShaderRHI = GetComputeShader();
 		FGlobalShader::SetParameters<FViewUniformShaderParameters>(RHICmdList, ShaderRHI, View.ViewUniformBuffer);
 		CulledObjectBufferParameters.Set(RHICmdList, ShaderRHI, GGlobalDistanceFieldCulledObjectBuffers.Buffers);
 		GlobalDistanceFieldParameters.Set(RHICmdList, ShaderRHI, GlobalDistanceFieldInfo.ParameterData);
@@ -493,7 +493,7 @@ public:
 		int32 ClipmapIndexValue,
 		const FVolumeUpdateRegion& UpdateRegion)
 	{
-		FComputeShaderRHIParamRef ShaderRHI = GetComputeShader();
+		FRHIComputeShader* ShaderRHI = GetComputeShader();
 		FGlobalShader::SetParameters<FViewUniformShaderParameters>(RHICmdList, ShaderRHI, View.ViewUniformBuffer);
 		CulledObjectBufferParameters.Set(RHICmdList, ShaderRHI, GGlobalDistanceFieldCulledObjectBuffers.Buffers);
 		GlobalDistanceFieldParameters.Set(RHICmdList, ShaderRHI, ParameterData);
@@ -637,7 +637,7 @@ public:
 		UTexture2D* VisibilityTextureValue,
 		int32 NumHeightfieldsValue)
 	{
-		FComputeShaderRHIParamRef ShaderRHI = GetComputeShader();
+		FRHIComputeShader* ShaderRHI = GetComputeShader();
 		FGlobalShader::SetParameters<FViewUniformShaderParameters>(RHICmdList, ShaderRHI, View.ViewUniformBuffer);
 		GlobalDistanceFieldParameters.Set(RHICmdList, ShaderRHI, GlobalDistanceFieldInfo.ParameterData);
 
diff --git a/Engine/Source/Runtime/Renderer/Private/HeightfieldLighting.cpp b/Engine/Source/Runtime/Renderer/Private/HeightfieldLighting.cpp
--- a/Engine/Source/Runtime/Renderer/Private/HeightfieldLighting.cpp
+++ b/Engine/Source/Runtime/Renderer/Private/HeightfieldLighting.cpp
@@ -267,7 +267,7 @@ public:
 
 	void SetParameters(FRHICommandList& RHICmdList, const FSceneView& View)
 	{
-		const FVertexShaderRHIParamRef ShaderRHI = GetVertexShader();
+		FRHIVertexShader* ShaderRHI = GetVertexShader();
 
 		FGlobalShader::SetParameters<FViewUniformShaderParameters>(RHICmdList, ShaderRHI, View.ViewUniformBuffer);
 		SubsectionHeightfieldParameters.Set(RHICmdList, ShaderRHI);
@@ -310,7 +310,7 @@ public:
 
 	void SetParameters(FRHICommandList& RHICmdList, const FSceneView& View, UTexture2D* HeightfieldTextureValue, UTexture2D* DiffuseColorTextureValue)
 	{
-		const FPixelShaderRHIParamRef ShaderRHI = GetPixelShader();
+		FRHIPixelShader* ShaderRHI = GetPixelShader();
 
 		FGlobalShader::SetParameters<FViewUniformShaderParameters>(RHICmdList, ShaderRHI, View.ViewUniformBuffer);
 		HeightfieldTextureParameters.Set(RHICmdList, ShaderRHI, HeightfieldTextureValue, DiffuseColorTextureValue, nullptr);
@@ -719,7 +719,7 @@ public:
 
 	void SetParameters(FRHICommandList& RHICmdList, const FSceneView& View, int32 NumHeightfieldsValue)
 	{
-		const FVertexShaderRHIParamRef ShaderRHI = GetVertexShader();
+		FRHIVertexShader* ShaderRHI = GetVertexShader();
 
 		FGlobalShader::SetParameters<FViewUniformShaderParameters>(RHICmdList, ShaderRHI, View.ViewUniformBuffer);
 		HeightfieldDescriptionParameters.Set(RHICmdList, ShaderRHI, GetHeightfieldDescriptionsSRV(), NumHeightfieldsValue);
@@ -769,7 +769,7 @@ public:
 
 	void SetParameters(FRHICommandList& RHICmdList, const FSceneView& View, const FProjectedShadowInfo* ProjectedShadowInfo, int32 NumHeightfieldsValue, const FHeightfieldLightingAtlas& Atlas)
 	{
-		const FPixelShaderRHIParamRef ShaderRHI = GetPixelShader();
+		FRHIPixelShader* ShaderRHI = GetPixelShader();
 
 		FGlobalShader::SetParameters<FViewUniformShaderParameters>(RHICmdList, ShaderRHI, View.ViewUniformBuffer);
 		HeightfieldDescriptionParameters.Set(RHICmdList, ShaderRHI, GetHeightfieldDescriptionsSRV(), NumHeightfieldsValue);
@@ -940,7 +940,7 @@ public:
 		FLightTileIntersectionResources* TileIntersectionResources,
 		FDistanceFieldObjectBufferResource& CulledObjectBuffers)
 	{
-		const FPixelShaderRHIParamRef ShaderRHI = GetPixelShader();
+		FRHIPixelShader* ShaderRHI = GetPixelShader();
 
 		FGlobalShader::SetParameters<FViewUniformShaderParameters>(RHICmdList, ShaderRHI, View.ViewUniformBuffer);
 		HeightfieldDescriptionParameters.Set(RHICmdList, ShaderRHI, GetHeightfieldDescriptionsSRV(), NumHeightfieldsValue);
@@ -1108,7 +1108,7 @@ public:
 		const FHeightfieldLightingAtlas& Atlas,
 		float SkyLightIndirectScaleValue)
 	{
-		const FPixelShaderRHIParamRef ShaderRHI = GetPixelShader();
+		FRHIPixelShader* ShaderRHI = GetPixelShader();
 
 		FMaterialShader::SetParameters(RHICmdList, ShaderRHI, MaterialProxy, *MaterialProxy->GetMaterial(View.GetFeatureLevel()), View, View.ViewUniformBuffer, ESceneTextureSetupMode::None);
 
@@ -1295,7 +1295,7 @@ public:
 		const FAOScreenGridResources& ScreenGridResources,
 		const FDistanceFieldAOParameters& Parameters)
 	{
-		FComputeShaderRHIParamRef ShaderRHI = GetComputeShader();
+		FRHIComputeShader* ShaderRHI = GetComputeShader();
 		FGlobalShader::SetParameters<FViewUniformShaderParameters>(RHICmdList, ShaderRHI, View.ViewUniformBuffer);
 
 		AOParameters.Set(RHICmdList, ShaderRHI, Parameters);
@@ -1446,7 +1446,7 @@ public:
 		const FDistanceFieldAOParameters& Parameters,
 		const FHeightfieldLightingAtlas& Atlas)
 	{
-		FComputeShaderRHIParamRef ShaderRHI = GetComputeShader();
+		FRHIComputeShader* ShaderRHI = GetComputeShader();
 		FGlobalShader::SetParameters<FViewUniformShaderParameters>(RHICmdList, ShaderRHI, View.ViewUniformBuffer);
 
 		AOParameters.Set(RHICmdList, ShaderRHI, Parameters);
diff --git a/Engine/Source/Runtime/Renderer/Private/LightFunctionRendering.cpp b/Engine/Source/Runtime/Renderer/Private/LightFunctionRendering.cpp
--- a/Engine/Source/Runtime/Renderer/Private/LightFunctionRendering.cpp
+++ b/Engine/Source/Runtime/Renderer/Private/LightFunctionRendering.cpp
@@ -110,7 +110,7 @@ public:
 
 	void SetParameters(FRHICommandList& RHICmdList, const FViewInfo& View, const FLightSceneInfo* LightSceneInfo, const FMaterialRenderProxy* MaterialProxy, bool bRenderingPreviewShadowIndicator, float ShadowFadeFraction )
 	{
-		const FPixelShaderRHIParamRef ShaderRHI = GetPixelShader();
+		FRHIPixelShader* ShaderRHI = GetPixelShader();
 
 		FMaterialShader::SetParameters(RHICmdList, ShaderRHI, MaterialProxy, *MaterialProxy->GetMaterial(View.GetFeatureLevel()), View, View.ViewUniformBuffer, ESceneTextureSetupMode::All);
 
diff --git a/Engine/Source/Runtime/Renderer/Private/LightPropagationVolume.cpp b/Engine/Source/Runtime/Renderer/Private/LightPropagationVolume.cpp
--- a/Engine/Source/Runtime/Renderer/Private/LightPropagationVolume.cpp
+++ b/Engine/Source/Runtime/Renderer/Private/LightPropagationVolume.cpp
@@ -201,7 +201,7 @@ public:
 
 	void SetParameters(FRHICommandList& RHICmdList, const FLpvBaseWriteShaderParams& Params )
 	{
-		FComputeShaderRHIParamRef ShaderRHI = GetComputeShader();
+		FRHIComputeShader* ShaderRHI = GetComputeShader();
 		SetUniformBufferParameter(RHICmdList, ShaderRHI, GetUniformBufferParameter<FLpvWriteUniformBufferParameters>(), Params.UniformBuffer );
 
 		TArray<int32> ResourceIndices;
@@ -297,7 +297,7 @@ public:
 		TArray<int32> ResourceIndices;
 		TArray<FRHIUnorderedAccessView*> UAVs;
 
-		FComputeShaderRHIParamRef ShaderRHI = GetComputeShader();
+		FRHIComputeShader* ShaderRHI = GetComputeShader();
 		for ( int i = 0; i < 7; i++ )
 		{
 			if ( LpvBufferSRVParameters[i].IsBound() )
@@ -483,7 +483,7 @@ public:
 		FTextureRHIParamRef RsmNormalTextureRHI,
 		FTextureRHIParamRef RsmDepthTextureRHI )
 	{
-		FComputeShaderRHIParamRef ShaderRHI = GetComputeShader();
+		FRHIComputeShader* ShaderRHI = GetComputeShader();
 		FLpvWriteShaderCSBase::SetParameters(RHICmdList, BaseParams );
 
 		FRHISamplerState* SamplerStateLinear  = TStaticSamplerState<SF_Bilinear,AM_Clamp,AM_Clamp,AM_Clamp>::GetRHI();
@@ -569,7 +569,6 @@ public:
 		FLpvBaseWriteShaderParams& BaseParams )
 	{
 		FLpvWriteShaderCSBase::SetParameters( RHICmdList, BaseParams );
-		const FComputeShaderRHIParamRef ShaderRHI = GetComputeShader();
 	}
 };
 IMPLEMENT_SHADER_TYPE(,FLpvDirectionalOcclusionCS,TEXT("/Engine/Private/LPVDirectionalOcclusion.usf"),TEXT("CSDirectionalOcclusion"),SF_Compute);
@@ -603,7 +602,6 @@ public:
 		FLpvBaseWriteShaderParams& BaseParams )
 	{	
 		FLpvWriteShaderCSBase::SetParameters( RHICmdList, BaseParams );
-		const FComputeShaderRHIParamRef ShaderRHI = GetComputeShader();
 	}
 };
 IMPLEMENT_SHADER_TYPE(,FLpvCopyAOVolumeCS,TEXT("/Engine/Private/LPVDirectionalOcclusion.usf"),TEXT("CSCopyAOVolume"),SF_Compute);
@@ -716,7 +714,7 @@ public:
 		FDirectLightInjectBufferRef& InjectUniformBuffer )
 	{
 		FLpvWriteShaderCSBase::SetParameters(RHICmdList, BaseParams );
-		FComputeShaderRHIParamRef ComputeShaderRHI = GetComputeShader();
+		FRHIComputeShader* ComputeShaderRHI = GetComputeShader();
 		SetUniformBufferParameter(RHICmdList, ComputeShaderRHI, GetUniformBufferParameter<FLpvDirectLightInjectParameters>(), InjectUniformBuffer );
 	}
 
diff --git a/Engine/Source/Runtime/Renderer/Private/LightPropagationVolumeVisualisation.cpp b/Engine/Source/Runtime/Renderer/Private/LightPropagationVolumeVisualisation.cpp
--- a/Engine/Source/Runtime/Renderer/Private/LightPropagationVolumeVisualisation.cpp
+++ b/Engine/Source/Runtime/Renderer/Private/LightPropagationVolumeVisualisation.cpp
@@ -64,7 +64,7 @@ public:
 		FRHICommandList& RHICmdList, 
 		const FSceneView& View )
 	{
-		FGeometryShaderRHIParamRef ShaderRHI = GetGeometryShader();
+		FRHIGeometryShader* ShaderRHI = GetGeometryShader();
 		FGlobalShader::SetParameters<FViewUniformShaderParameters>(RHICmdList, ShaderRHI, View.ViewUniformBuffer);
 	}
 };
@@ -84,7 +84,7 @@ public:
 		FRHICommandList& RHICmdList, 
 		const FSceneView& View )
 	{
-		FVertexShaderRHIParamRef ShaderRHI = GetVertexShader();
+		FRHIVertexShader* ShaderRHI = GetVertexShader();
 		FGlobalShader::SetParameters<FViewUniformShaderParameters>(RHICmdList, ShaderRHI, View.ViewUniformBuffer);
 	}
 };
@@ -116,7 +116,7 @@ public:
 		const FLightPropagationVolume* LPV,
 		const FSceneView& View )
 	{
-		FPixelShaderRHIParamRef ShaderRHI = GetPixelShader();
+		FRHIPixelShader* ShaderRHI = GetPixelShader();
 		FGlobalShader::SetParameters<FViewUniformShaderParameters>(RHICmdList, ShaderRHI, View.ViewUniformBuffer);
 		
 		for ( int i = 0; i < 7; i++ )
@@ -169,7 +169,7 @@ public:
 	void UnbindBuffers(FRHICommandList& RHICmdList)
 	{
 		// TODO: Is this necessary here?
-		FPixelShaderRHIParamRef ShaderRHI = GetPixelShader();
+		FRHIPixelShader* ShaderRHI = GetPixelShader();
 		for ( int i = 0; i < 7; i++ )
 		{
 			if ( LpvBufferSRVParameters[i].IsBound() )
diff --git a/Engine/Source/Runtime/Renderer/Private/LightRendering.cpp b/Engine/Source/Runtime/Renderer/Private/LightRendering.cpp
--- a/Engine/Source/Runtime/Renderer/Private/LightRendering.cpp
+++ b/Engine/Source/Runtime/Renderer/Private/LightRendering.cpp
@@ -281,14 +281,14 @@ class FDeferredLightPS : public FGlobalShader
 public:
 	void SetParameters(FRHICommandList& RHICmdList, const FSceneView& View, const FLightSceneInfo* LightSceneInfo, IPooledRenderTarget* ScreenShadowMaskTexture)
 	{
-		const FPixelShaderRHIParamRef ShaderRHI = GetPixelShader();
+		FRHIPixelShader* ShaderRHI = GetPixelShader();
 		SetParametersBase(RHICmdList, ShaderRHI, View, ScreenShadowMaskTexture, LightSceneInfo->Proxy->GetIESTextureResource());
 		SetDeferredLightParameters(RHICmdList, ShaderRHI, GetUniformBufferParameter<FDeferredLightUniformStruct>(), LightSceneInfo, View);
 	}
 
 	void SetParametersSimpleLight(FRHICommandList& RHICmdList, const FSceneView& View, const FSimpleLightEntry& SimpleLight, const FSimpleLightPerViewEntry& SimpleLightPerViewData)
 	{
-		const FPixelShaderRHIParamRef ShaderRHI = GetPixelShader();
+		FRHIPixelShader* ShaderRHI = GetPixelShader();
 		SetParametersBase(RHICmdList, ShaderRHI, View, NULL, NULL);
 		SetSimpleDeferredLightParameters(RHICmdList, ShaderRHI, GetUniformBufferParameter<FDeferredLightUniformStruct>(), SimpleLight, SimpleLightPerViewData, View);
 	}
@@ -314,7 +314,7 @@ public:
 
 private:
 
-	void SetParametersBase(FRHICommandList& RHICmdList, const FPixelShaderRHIParamRef ShaderRHI, const FSceneView& View, IPooledRenderTarget* ScreenShadowMaskTexture, FTexture* IESTextureResource)
+	void SetParametersBase(FRHICommandList& RHICmdList, FRHIPixelShader* ShaderRHI, const FSceneView& View, IPooledRenderTarget* ScreenShadowMaskTexture, FTexture* IESTextureResource)
 	{
 		FGlobalShader::SetParameters<FViewUniformShaderParameters>(RHICmdList, ShaderRHI,View.ViewUniformBuffer);
 		SceneTextureParameters.Set(RHICmdList, ShaderRHI, View.FeatureLevel, ESceneTextureSetupMode::All);
@@ -449,7 +449,7 @@ public:
 
 	void SetParameters(FRHICommandList& RHICmdList, const FSceneView& View, const FLightSceneInfo* LightSceneInfo)
 	{
-		const FPixelShaderRHIParamRef ShaderRHI = GetPixelShader();
+		FRHIPixelShader* ShaderRHI = GetPixelShader();
 		FGlobalShader::SetParameters<FViewUniformShaderParameters>(RHICmdList, ShaderRHI,View.ViewUniformBuffer);
 		const float HasValidChannelValue = LightSceneInfo->Proxy->GetPreviewShadowMapChannel() == INDEX_NONE ? 0.0f : 1.0f;
 		SetShaderValue(RHICmdList, ShaderRHI, HasValidChannel, HasValidChannelValue);
diff --git a/Engine/Source/Runtime/Renderer/Private/MeshPassProcessor.cpp b/Engine/Source/Runtime/Renderer/Private/MeshPassProcessor.cpp
--- a/Engine/Source/Runtime/Renderer/Private/MeshPassProcessor.cpp
+++ b/Engine/Source/Runtime/Renderer/Private/MeshPassProcessor.cpp
@@ -802,7 +802,7 @@ void FMeshDrawShaderBindings::SetOnCommandList(FRHICommandList& RHICmdList, FBou
 	}
 }
 
-void FMeshDrawShaderBindings::SetOnCommandListForCompute(FRHICommandList& RHICmdList, FComputeShaderRHIParamRef Shader) const
+void FMeshDrawShaderBindings::SetOnCommandListForCompute(FRHICommandList& RHICmdList, FRHIComputeShader* Shader) const
 {
 	check(ShaderLayouts.Num() == 1);
 	FReadOnlyMeshDrawSingleShaderBindings SingleShaderBindings(ShaderLayouts[0], GetData());
diff --git a/Engine/Source/Runtime/Renderer/Private/MobileReflectionEnvironmentCapture.cpp b/Engine/Source/Runtime/Renderer/Private/MobileReflectionEnvironmentCapture.cpp
--- a/Engine/Source/Runtime/Renderer/Private/MobileReflectionEnvironmentCapture.cpp
+++ b/Engine/Source/Runtime/Renderer/Private/MobileReflectionEnvironmentCapture.cpp
@@ -402,7 +402,7 @@ namespace MobileReflectionEnvironmentCapture
 						FCubeFilterPS* HQFilterPixelShader = *TShaderMapRef< TCubeFilterPS<0> >(ShaderMap);
 						check(HQFilterPixelShader);
 						TShaderMapRef<FMobileDownsamplePS> BilinFilterPixelShader(ShaderMap);
-						FPixelShaderRHIParamRef PixelShaderRHI = bUseHQFiltering ? HQFilterPixelShader->GetPixelShader() : BilinFilterPixelShader->GetPixelShader();
+						FRHIPixelShader* PixelShaderRHI = bUseHQFiltering ? HQFilterPixelShader->GetPixelShader() : BilinFilterPixelShader->GetPixelShader();
 
 						GraphicsPSOInit.BoundShaderState.VertexDeclarationRHI = GFilterVertexDeclaration.VertexDeclarationRHI;
 						GraphicsPSOInit.BoundShaderState.VertexShaderRHI = GETSAFERHISHADER_VERTEX(*VertexShader);
diff --git a/Engine/Source/Runtime/Renderer/Private/PathTracing/PathCompactionCompute.cpp b/Engine/Source/Runtime/Renderer/Private/PathTracing/PathCompactionCompute.cpp
--- a/Engine/Source/Runtime/Renderer/Private/PathTracing/PathCompactionCompute.cpp
+++ b/Engine/Source/Runtime/Renderer/Private/PathTracing/PathCompactionCompute.cpp
@@ -64,7 +64,7 @@ public:
 		FRHIUnorderedAccessView* RadianceSortedAlphaUAV,
 		FRHIUnorderedAccessView* SampleCountSortedUAV)
 	{
-		FComputeShaderRHIParamRef ShaderRHI = GetComputeShader();
+		FRHIComputeShader* ShaderRHI = GetComputeShader();
 		FGlobalShader::SetParameters<FViewUniformShaderParameters>(RHICmdList, ShaderRHI, View.ViewUniformBuffer);
 
 		// Input textures
@@ -91,7 +91,7 @@ public:
 		FRHIUnorderedAccessView* SampleCountSortedUAV,
 		FRHIComputeFence* Fence)
 	{
-		FComputeShaderRHIParamRef ShaderRHI = GetComputeShader();
+		FRHIComputeShader* ShaderRHI = GetComputeShader();
 
 		SetUAVParameter(RHICmdList, ShaderRHI, RadianceSortedRedUAVParameter, FUnorderedAccessViewRHIRef());
 		SetUAVParameter(RHICmdList, ShaderRHI, RadianceSortedGreenUAVParameter, FUnorderedAccessViewRHIRef());
diff --git a/Engine/Source/Runtime/Renderer/Private/PathTracing/PathTracingTestbed.cpp b/Engine/Source/Runtime/Renderer/Private/PathTracing/PathTracingTestbed.cpp
--- a/Engine/Source/Runtime/Renderer/Private/PathTracing/PathTracingTestbed.cpp
+++ b/Engine/Source/Runtime/Renderer/Private/PathTracing/PathTracingTestbed.cpp
@@ -54,7 +54,7 @@ public:
 		FRHIUnorderedAccessView* ResultsBuffer,
 		FRHIUnorderedAccessView* FloatResultsBuffer)
 	{
-		FComputeShaderRHIParamRef ShaderRHI = GetComputeShader();
+		FRHIComputeShader* ShaderRHI = GetComputeShader();
 		SetShaderValue(RHICmdList, ShaderRHI, SamplesCountParameter, SamplesCount);
 		SetShaderValue(RHICmdList, ShaderRHI, BrdfTypeParameter, BrdfType);
 		SetUAVParameter(RHICmdList, ShaderRHI, ResultsBufferParameter, ResultsBuffer);
@@ -69,7 +69,7 @@ public:
 		FRWBuffer& FloatResultsBuffer,
 		FRHIComputeFence* Fence)
 	{
-		FComputeShaderRHIParamRef ShaderRHI = GetComputeShader();
+		FRHIComputeShader* ShaderRHI = GetComputeShader();
 
 		RHICmdList.TransitionResource(TransitionAccess, TransitionPipeline, ResultsBuffer.UAV, Fence);
 		RHICmdList.TransitionResource(TransitionAccess, TransitionPipeline, FloatResultsBuffer.UAV, Fence);
@@ -139,7 +139,7 @@ public:
 		uint32 NumPhiSteps,
 		FRHIUnorderedAccessView* PdfsResultsBuffer)
 	{
-		FComputeShaderRHIParamRef ShaderRHI = GetComputeShader();
+		FRHIComputeShader* ShaderRHI = GetComputeShader();
 		SetShaderValue(RHICmdList, ShaderRHI, BrdfTypeParameter, BrdfType);
 		SetShaderValue(RHICmdList, ShaderRHI, WoParameter, Wo);
 		SetShaderValue(RHICmdList, ShaderRHI, NumThetaStepsParameter, NumThetaSteps);
@@ -154,7 +154,7 @@ public:
 		FRWBuffer& PdfsResultsBuffer,
 		FRHIComputeFence* Fence)
 	{
-		FComputeShaderRHIParamRef ShaderRHI = GetComputeShader();
+		FRHIComputeShader* ShaderRHI = GetComputeShader();
 
 		RHICmdList.TransitionResource(TransitionAccess, TransitionPipeline, PdfsResultsBuffer.UAV, Fence);
 	}
@@ -227,7 +227,7 @@ public:
 		uint32 NumPhiSteps,
 		FRHIUnorderedAccessView* WisBinDistributionResultsBuffer)
 	{
-		FComputeShaderRHIParamRef ShaderRHI = GetComputeShader();
+		FRHIComputeShader* ShaderRHI = GetComputeShader();
 		SetShaderValue(RHICmdList, ShaderRHI, BrdfTypeParameter, BrdfType);
 		SetShaderValue(RHICmdList, ShaderRHI, NumSamplesParameter, NumSamples);
 		SetShaderValue(RHICmdList, ShaderRHI, WoParameter, Wo);
@@ -243,7 +243,7 @@ public:
 		FRWBuffer& PdfsResultsBuffer,
 		FRHIComputeFence* Fence)
 	{
-		FComputeShaderRHIParamRef ShaderRHI = GetComputeShader();
+		FRHIComputeShader* ShaderRHI = GetComputeShader();
 
 		RHICmdList.TransitionResource(TransitionAccess, TransitionPipeline, PdfsResultsBuffer.UAV, Fence);
 	}
@@ -316,7 +316,7 @@ public:
 		uint32 NumPhiSteps,
 		FRHIUnorderedAccessView* PdfsResultsBuffer)
 	{
-		FComputeShaderRHIParamRef ShaderRHI = GetComputeShader();
+		FRHIComputeShader* ShaderRHI = GetComputeShader();
 		SetShaderValue(RHICmdList, ShaderRHI, BrdfTypeParameter, BrdfType);
 		SetShaderValue(RHICmdList, ShaderRHI, WoParameter, Wo);
 		SetShaderValue(RHICmdList, ShaderRHI, NumThetaStepsParameter, NumThetaSteps);
@@ -331,7 +331,7 @@ public:
 		FRWBuffer& PdfsResultsBuffer,
 		FRHIComputeFence* Fence)
 	{
-		FComputeShaderRHIParamRef ShaderRHI = GetComputeShader();
+		FRHIComputeShader* ShaderRHI = GetComputeShader();
 
 		RHICmdList.TransitionResource(TransitionAccess, TransitionPipeline, PdfsResultsBuffer.UAV, Fence);
 	}
diff --git a/Engine/Source/Runtime/Renderer/Private/PathTracing/SkyLightCdfBuild.cpp b/Engine/Source/Runtime/Renderer/Private/PathTracing/SkyLightCdfBuild.cpp
--- a/Engine/Source/Runtime/Renderer/Private/PathTracing/SkyLightCdfBuild.cpp
+++ b/Engine/Source/Runtime/Renderer/Private/PathTracing/SkyLightCdfBuild.cpp
@@ -52,7 +52,7 @@ public:
 
 	void SetParameters(FRHICommandList& RHICmdList, uint32 Mode, const FTexture& TextureCube, uint32 CubeFace, uint32 Level, FIntVector RowCdfDimensions, const FRWBuffer& RowCdf)
 	{
-		FComputeShaderRHIParamRef ShaderRHI = GetComputeShader();
+		FRHIComputeShader* ShaderRHI = GetComputeShader();
 
 		SetShaderValue(RHICmdList, ShaderRHI, ModeParameter, Mode);
 		SetTextureParameter(RHICmdList, ShaderRHI, TextureCubeParameter, TextureCubeSamplerParameter, TStaticSamplerState<SF_Bilinear>::GetRHI(), TextureCube.TextureRHI);
@@ -72,7 +72,7 @@ public:
 	void UnsetParameters(FRHICommandList& RHICmdList, EResourceTransitionAccess TransitionAccess, EResourceTransitionPipeline TransitionPipeline,
 		const FRWBuffer& Buffer, FRHIComputeFence* Fence)
 	{
-		FComputeShaderRHIParamRef ShaderRHI = GetComputeShader();
+		FRHIComputeShader* ShaderRHI = GetComputeShader();
 		RowCdfParameter.UnsetUAV(RHICmdList, ShaderRHI);
 
 		FRHIUnorderedAccessView* UAVs[] = {
@@ -142,7 +142,7 @@ public:
 
 	void SetParameters(FRHICommandList& RHICmdList, uint32 Mode, FIntVector RowCdfDimensions, const FRWBuffer& RowCdf, uint32 Level, const FRWBuffer& ColumnCdf)
 	{
-		FComputeShaderRHIParamRef ShaderRHI = GetComputeShader();
+		FRHIComputeShader* ShaderRHI = GetComputeShader();
 
 		SetShaderValue(RHICmdList, ShaderRHI, ModeParameter, Mode);
 		SetShaderValue(RHICmdList, ShaderRHI, RowCdfDimensionsParameter, RowCdfDimensions);
@@ -159,7 +159,7 @@ public:
 	void UnsetParameters(FRHICommandList& RHICmdList, EResourceTransitionAccess TransitionAccess, EResourceTransitionPipeline TransitionPipeline,
 		const FRWBuffer& ColumnCdf, FRHIComputeFence* Fence)
 	{
-		FComputeShaderRHIParamRef ShaderRHI = GetComputeShader();
+		FRHIComputeShader* ShaderRHI = GetComputeShader();
 		ColumnCdfParameter.UnsetUAV(RHICmdList, ShaderRHI);
 
 		FRHIUnorderedAccessView* UAVs[] = {
@@ -223,7 +223,7 @@ public:
 
 	void SetParameters(FRHICommandList& RHICmdList, FIntVector ColumnCdfDimensions, const FRWBuffer& ColumnCdf, const FRWBuffer& CubeFaceCdf)
 	{
-		FComputeShaderRHIParamRef ShaderRHI = GetComputeShader();
+		FRHIComputeShader* ShaderRHI = GetComputeShader();
 
 		SetShaderValue(RHICmdList, ShaderRHI, ColumnCdfDimensionsParameter, ColumnCdfDimensions);
 		ColumnCdfParameter.SetBuffer(RHICmdList, ShaderRHI, ColumnCdf);
@@ -238,7 +238,7 @@ public:
 	void UnsetParameters(FRHICommandList& RHICmdList, EResourceTransitionAccess TransitionAccess, EResourceTransitionPipeline TransitionPipeline,
 		const FRWBuffer& CubeFaceCdf, FRHIComputeFence* Fence)
 	{
-		FComputeShaderRHIParamRef ShaderRHI = GetComputeShader();
+		FRHIComputeShader* ShaderRHI = GetComputeShader();
 		CubeFaceCdfParameter.UnsetUAV(RHICmdList, ShaderRHI);
 
 		FRHIUnorderedAccessView* UAVs[] = {
@@ -300,7 +300,7 @@ class FVisualizeCdfPS : public FGlobalShader
 		const FRWBuffer& ColumnCdf,
 		const FRWBuffer& CubeFaceCdf)
 	{
-		const FPixelShaderRHIParamRef ShaderRHI = GetPixelShader();
+		FRHIPixelShader* ShaderRHI = GetPixelShader();
 		FGlobalShader::SetParameters<FViewUniformShaderParameters>(RHICmdList, ShaderRHI, View.ViewUniformBuffer);
 
 		SetShaderValue(RHICmdList, ShaderRHI, DimensionsParameter, Dimensions);
diff --git a/Engine/Source/Runtime/Renderer/Private/PathTracing/VarianceMipTreeBuild.cpp b/Engine/Source/Runtime/Renderer/Private/PathTracing/VarianceMipTreeBuild.cpp
--- a/Engine/Source/Runtime/Renderer/Private/PathTracing/VarianceMipTreeBuild.cpp
+++ b/Engine/Source/Runtime/Renderer/Private/PathTracing/VarianceMipTreeBuild.cpp
@@ -55,7 +55,7 @@ public:
 		uint32 MipLevel,
 		FRWBuffer& VarianceMipTree)
 	{
-		FComputeShaderRHIParamRef ShaderRHI = GetComputeShader();
+		FRHIComputeShader* ShaderRHI = GetComputeShader();
 
 		SetShaderValue(RHICmdList, ShaderRHI, ViewSizeParameter, ViewSize);
 		SetShaderValue(RHICmdList, ShaderRHI, VarianceMapDimensionsParameter, VarianceMapDimensions);
@@ -73,7 +73,7 @@ public:
 		FRWBuffer& VarianceMap,
 		FRHIComputeFence* Fence)
 	{
-		FComputeShaderRHIParamRef ShaderRHI = GetComputeShader();
+		FRHIComputeShader* ShaderRHI = GetComputeShader();
 
 		VarianceMipTreeParameter.UnsetUAV(RHICmdList, ShaderRHI);
 		RHICmdList.TransitionResource(TransitionAccess, TransitionPipeline, VarianceMap.UAV, Fence);
@@ -173,7 +173,7 @@ public:
 		const FIntVector Dimensions,
 		const FRWBuffer& MipTree)
 	{
-		const FPixelShaderRHIParamRef ShaderRHI = GetPixelShader();
+		FRHIPixelShader* ShaderRHI = GetPixelShader();
 		FGlobalShader::SetParameters<FViewUniformShaderParameters>(RHICmdList, ShaderRHI, View.ViewUniformBuffer);
 
 		SetShaderValue(RHICmdList, ShaderRHI, DimensionsParameter, Dimensions);
diff --git a/Engine/Source/Runtime/Renderer/Private/PostProcess/PostProcessAA.cpp b/Engine/Source/Runtime/Renderer/Private/PostProcess/PostProcessAA.cpp
--- a/Engine/Source/Runtime/Renderer/Private/PostProcess/PostProcessAA.cpp
+++ b/Engine/Source/Runtime/Renderer/Private/PostProcess/PostProcessAA.cpp
@@ -74,7 +74,7 @@ public:
 	template <typename TRHICmdList>
 	void SetParameters(TRHICmdList& RHICmdList, const FRenderingCompositePassContext& Context)
 	{
-		const FPixelShaderRHIParamRef ShaderRHI = GetPixelShader();
+		FRHIPixelShader* ShaderRHI = GetPixelShader();
 
 		const FViewInfo& View = Context.View;
 
@@ -190,7 +190,7 @@ public:
 
 	void SetParameters(const FRenderingCompositePassContext& Context)
 	{
-		const FVertexShaderRHIParamRef ShaderRHI = GetVertexShader();
+		FRHIVertexShader* ShaderRHI = GetVertexShader();
 
 		FGlobalShader::SetParameters<FViewUniformShaderParameters>(Context.RHICmdList, ShaderRHI, Context.View.ViewUniformBuffer);
 
diff --git a/Engine/Source/Runtime/Renderer/Private/PostProcess/PostProcessBloomSetup.cpp b/Engine/Source/Runtime/Renderer/Private/PostProcess/PostProcessBloomSetup.cpp
--- a/Engine/Source/Runtime/Renderer/Private/PostProcess/PostProcessBloomSetup.cpp
+++ b/Engine/Source/Runtime/Renderer/Private/PostProcess/PostProcessBloomSetup.cpp
@@ -56,7 +56,7 @@ public:
 	template <typename TRHICmdList>
 	void SetPS(TRHICmdList& RHICmdList, const FRenderingCompositePassContext& Context)
 	{
-		const FPixelShaderRHIParamRef ShaderRHI = GetPixelShader();
+		FRHIPixelShader* ShaderRHI = GetPixelShader();
 
 		const FPostProcessSettings& Settings = Context.View.FinalPostProcessSettings;
 
@@ -110,7 +110,7 @@ public:
 	
 	void SetVS(const FRenderingCompositePassContext& Context)
 	{
-		const FVertexShaderRHIParamRef ShaderRHI = GetVertexShader();
+		FRHIVertexShader* ShaderRHI = GetVertexShader();
 
 		FGlobalShader::SetParameters<FViewUniformShaderParameters>(Context.RHICmdList, ShaderRHI, Context.View.ViewUniformBuffer);
 
@@ -180,7 +180,7 @@ public:
 	template <typename TRHICmdList>
 	void SetParameters(TRHICmdList& RHICmdList, const FRenderingCompositePassContext& Context, const FIntPoint& DestSize, FRHIUnorderedAccessView* DestUAV, FTextureRHIParamRef EyeAdaptationTex)
 	{
-		const FComputeShaderRHIParamRef ShaderRHI = GetComputeShader();
+		FRHIComputeShader* ShaderRHI = GetComputeShader();
 		const FPostProcessSettings& Settings = Context.View.FinalPostProcessSettings;
 
 		FGlobalShader::SetParameters<FViewUniformShaderParameters>(RHICmdList, ShaderRHI, Context.View.ViewUniformBuffer);
@@ -197,7 +197,7 @@ public:
 	template <typename TRHICmdList>
 	void UnsetParameters(TRHICmdList& RHICmdList)
 	{
-		const FComputeShaderRHIParamRef ShaderRHI = GetComputeShader();
+		FRHIComputeShader* ShaderRHI = GetComputeShader();
 		RHICmdList.SetUAVParameter(ShaderRHI, OutComputeTex.GetBaseIndex(), NULL);
 	}
 	
@@ -397,7 +397,7 @@ public:
 	template <typename TRHICmdList>
 	void SetParameters(TRHICmdList& RHICmdList, const FRenderingCompositePassContext& Context)
 	{
-		const FPixelShaderRHIParamRef ShaderRHI = GetPixelShader();
+		FRHIPixelShader* ShaderRHI = GetPixelShader();
 
 		FGlobalShader::SetParameters<FViewUniformShaderParameters>(RHICmdList, ShaderRHI, Context.View.ViewUniformBuffer);
 		PostprocessParameter.SetPS(RHICmdList, ShaderRHI, Context, TStaticSamplerState<SF_Bilinear, AM_Border, AM_Border, AM_Border>::GetRHI());
@@ -531,7 +531,7 @@ public:
 
 	void SetParameters(const FRenderingCompositePassContext& Context)
 	{
-		const FPixelShaderRHIParamRef ShaderRHI = GetPixelShader();
+		FRHIPixelShader* ShaderRHI = GetPixelShader();
 
 		const FPostProcessSettings& Settings = Context.View.FinalPostProcessSettings;
 
diff --git a/Engine/Source/Runtime/Renderer/Private/PostProcess/PostProcessCombineLUTs.cpp b/Engine/Source/Runtime/Renderer/Private/PostProcess/PostProcessCombineLUTs.cpp
--- a/Engine/Source/Runtime/Renderer/Private/PostProcess/PostProcessCombineLUTs.cpp
+++ b/Engine/Source/Runtime/Renderer/Private/PostProcess/PostProcessCombineLUTs.cpp
@@ -511,7 +511,7 @@ FColorRemapShaderParameters::FColorRemapShaderParameters(const FShaderParameterM
 	MappingPolynomial.Bind(ParameterMap, TEXT("MappingPolynomial"));
 }
 
-void FColorRemapShaderParameters::Set(FRHICommandList& RHICmdList, const FPixelShaderRHIParamRef ShaderRHI)
+void FColorRemapShaderParameters::Set(FRHICommandList& RHICmdList, FRHIPixelShader* ShaderRHI)
 {
 	FColorTransform ColorTransform;
 	ColorTransform.MinValue = FMath::Clamp(CVarColorMin.GetValueOnRenderThread(), -10.0f, 10.0f);
@@ -531,7 +531,7 @@ void FColorRemapShaderParameters::Set(FRHICommandList& RHICmdList, const FPixelS
 }
 
 template <typename TRHICmdList>
-void FColorRemapShaderParameters::Set(TRHICmdList& RHICmdList, const FComputeShaderRHIParamRef ShaderRHI)
+void FColorRemapShaderParameters::Set(TRHICmdList& RHICmdList, FRHIComputeShader* ShaderRHI)
 {
 	FColorTransform ColorTransform;
 	ColorTransform.MinValue = FMath::Clamp(CVarColorMin.GetValueOnRenderThread(), -10.0f, 10.0f);
@@ -582,7 +582,7 @@ public:
 	template <typename TRHICommandList>
 	void SetParameters(TRHICommandList& RHICmdList, const FSceneView& View, FTexture** Textures, float* Weights)
 	{
-		const FPixelShaderRHIParamRef ShaderRHI = GetPixelShader();
+		FRHIPixelShader* ShaderRHI = GetPixelShader();
 		CombineLUTsShaderParameters.Set(RHICmdList, ShaderRHI, View, Textures, Weights);
 	}
 
@@ -648,7 +648,7 @@ public:
 	template <typename TRHICmdList>
 	void SetParameters(TRHICmdList& RHICmdList, const FRenderingCompositePassContext& Context, const FIntPoint& DestSize, FRHIUnorderedAccessView* DestUAV, FTexture** Textures, float* Weights)
 	{
-		const FComputeShaderRHIParamRef ShaderRHI = GetComputeShader();
+		FRHIComputeShader* ShaderRHI = GetComputeShader();
 
 		// CS params
 		FGlobalShader::SetParameters<FViewUniformShaderParameters>(RHICmdList, ShaderRHI, Context.View.ViewUniformBuffer);
@@ -664,7 +664,7 @@ public:
 	template <typename TRHICmdList>
 	void UnsetParameters(TRHICmdList& RHICmdList)
 	{
-		const FComputeShaderRHIParamRef ShaderRHI = GetComputeShader();
+		FRHIComputeShader* ShaderRHI = GetComputeShader();
 		OutComputeTex.UnsetUAV(RHICmdList, ShaderRHI);
 	}
 
diff --git a/Engine/Source/Runtime/Renderer/Private/PostProcess/PostProcessCombineLUTs.h b/Engine/Source/Runtime/Renderer/Private/PostProcess/PostProcessCombineLUTs.h
--- a/Engine/Source/Runtime/Renderer/Private/PostProcess/PostProcessCombineLUTs.h
+++ b/Engine/Source/Runtime/Renderer/Private/PostProcess/PostProcessCombineLUTs.h
@@ -67,10 +67,10 @@ public:
 	FColorRemapShaderParameters(const FShaderParameterMap& ParameterMap);
 
 	// Explicit declarations here because templates unresolved when used in other files
-	void Set(FRHICommandList& RHICmdList, const FPixelShaderRHIParamRef ShaderRHI);
+	void Set(FRHICommandList& RHICmdList, FRHIPixelShader* ShaderRHI);
 
 	template <typename TRHICmdList>
-	void Set(TRHICmdList& RHICmdList, const FComputeShaderRHIParamRef ShaderRHI);
+	void Set(TRHICmdList& RHICmdList, FRHIComputeShader* ShaderRHI);
 
 	friend FArchive& operator<<(FArchive& Ar,FColorRemapShaderParameters& P);
 
diff --git a/Engine/Source/Runtime/Renderer/Private/PostProcess/PostProcessDownsample.cpp b/Engine/Source/Runtime/Renderer/Private/PostProcess/PostProcessDownsample.cpp
--- a/Engine/Source/Runtime/Renderer/Private/PostProcess/PostProcessDownsample.cpp
+++ b/Engine/Source/Runtime/Renderer/Private/PostProcess/PostProcessDownsample.cpp
@@ -66,7 +66,7 @@ public:
 	template <typename TRHICmdList>
 	void SetParameters(TRHICmdList& RHICmdList, const FRenderingCompositePassContext& Context, const FPooledRenderTargetDesc* InputDesc, const FIntPoint& SrcSize, const FIntRect& SrcRect)
 	{
-		const FPixelShaderRHIParamRef ShaderRHI = GetPixelShader();
+		FRHIPixelShader* ShaderRHI = GetPixelShader();
 
 		FGlobalShader::SetParameters<FViewUniformShaderParameters>(RHICmdList, ShaderRHI, Context.View.ViewUniformBuffer);
 		SceneTextureParameters.Set(RHICmdList, ShaderRHI, Context.View.FeatureLevel, ESceneTextureSetupMode::All);
@@ -144,7 +144,7 @@ public:
 
 	void SetParameters(const FRenderingCompositePassContext& Context)
 	{
-		const FVertexShaderRHIParamRef ShaderRHI = GetVertexShader();
+		FRHIVertexShader* ShaderRHI = GetVertexShader();
 
 		FGlobalShader::SetParameters<FViewUniformShaderParameters>(Context.RHICmdList, ShaderRHI, Context.View.ViewUniformBuffer);
 
@@ -209,7 +209,7 @@ public:
 	template <typename TRHICmdList>
 	void SetParameters(TRHICmdList& RHICmdList, const FRenderingCompositePassContext& Context, const FIntPoint& SrcSize, FRHIUnorderedAccessView* DestUAV)
 	{
-		const FComputeShaderRHIParamRef ShaderRHI = GetComputeShader();
+		FRHIComputeShader* ShaderRHI = GetComputeShader();
 		FGlobalShader::SetParameters<FViewUniformShaderParameters>(RHICmdList, ShaderRHI, Context.View.ViewUniformBuffer);
 
 		// filter only if needed for better performance
@@ -229,7 +229,7 @@ public:
 	template <typename TRHICmdList>
 	void UnsetParameters(TRHICmdList& RHICmdList)
 	{
-		const FComputeShaderRHIParamRef ShaderRHI = GetComputeShader();
+		FRHIComputeShader* ShaderRHI = GetComputeShader();
 		RHICmdList.SetUAVParameter(ShaderRHI, OutComputeTex.GetBaseIndex(), NULL);
 	}
 };
