DR - Deprecate FVertexShaderRHIParamRef, FHullShaderRHIParamRef, FDomainShaderRHIParamRef, FPixelShaderRHIParamRef, FGeometryShaderRHIParamRef and FComputeShaderRHIParamRef

diff --git a/Engine/Source/Runtime/Renderer/Private/PostProcess/PostProcessEyeAdaptation.cpp b/Engine/Source/Runtime/Renderer/Private/PostProcess/PostProcessEyeAdaptation.cpp
--- a/Engine/Source/Runtime/Renderer/Private/PostProcess/PostProcessEyeAdaptation.cpp
+++ b/Engine/Source/Runtime/Renderer/Private/PostProcess/PostProcessEyeAdaptation.cpp
@@ -203,7 +203,7 @@ public:
 	template <typename TRHICmdList>
 	void SetPS(const FRenderingCompositePassContext& Context, TRHICmdList& RHICmdList, IPooledRenderTarget* LastEyeAdaptation)
 	{
-		const FPixelShaderRHIParamRef ShaderRHI = GetPixelShader();
+		FRHIPixelShader* ShaderRHI = GetPixelShader();
 
 		FGlobalShader::SetParameters<FViewUniformShaderParameters>(RHICmdList, ShaderRHI, Context.View.ViewUniformBuffer);
 
@@ -273,7 +273,7 @@ public:
 	template <typename TRHICmdList>
 	void SetParameters(TRHICmdList& RHICmdList, const FRenderingCompositePassContext& Context, FRHIUnorderedAccessView* DestUAV, IPooledRenderTarget* LastEyeAdaptation)
 	{
-		const FComputeShaderRHIParamRef ShaderRHI = GetComputeShader();
+		FRHIComputeShader* ShaderRHI = GetComputeShader();
 
 		// CS params
 		FGlobalShader::SetParameters<FViewUniformShaderParameters>(RHICmdList, ShaderRHI, Context.View.ViewUniformBuffer);
@@ -289,7 +289,7 @@ public:
 	template <typename TRHICmdList>
 	void UnsetParameters(TRHICmdList& RHICmdList)
 	{
-		const FComputeShaderRHIParamRef ShaderRHI = GetComputeShader();
+		FRHIComputeShader* ShaderRHI = GetComputeShader();
 		OutComputeTex.UnsetUAV(RHICmdList, ShaderRHI);
 	}
 
@@ -556,7 +556,7 @@ public:
 
 	void SetPS(const FRenderingCompositePassContext& Context)
 	{
-		const FPixelShaderRHIParamRef ShaderRHI = GetPixelShader();
+		FRHIPixelShader* ShaderRHI = GetPixelShader();
 
 		FGlobalShader::SetParameters<FViewUniformShaderParameters>(Context.RHICmdList, ShaderRHI, Context.View.ViewUniformBuffer);
 		PostprocessParameter.SetPS(Context.RHICmdList, ShaderRHI, Context, TStaticSamplerState<SF_Point, AM_Clamp, AM_Clamp, AM_Clamp>::GetRHI());
@@ -729,7 +729,7 @@ public:
 
 	void SetPS(const FRenderingCompositePassContext& Context, const FIntRect& SrcRect, IPooledRenderTarget* LastEyeAdaptation)
 	{
-		const FPixelShaderRHIParamRef ShaderRHI = GetPixelShader();
+		FRHIPixelShader* ShaderRHI = GetPixelShader();
 
 		FGlobalShader::SetParameters<FViewUniformShaderParameters>(Context.RHICmdList, ShaderRHI, Context.View.ViewUniformBuffer);
 
@@ -803,13 +803,13 @@ public:
 
 	void UnsetParameters(const FRenderingCompositePassContext& Context)
 	{
-		const FComputeShaderRHIParamRef ShaderRHI = GetComputeShader();
+		FRHIComputeShader* ShaderRHI = GetComputeShader();
 		OutComputeTex.UnsetUAV(Context.RHICmdList, ShaderRHI);
 	}
 
 	void SetParameters(const FRenderingCompositePassContext& Context, const FIntRect& SrcRect, IPooledRenderTarget* LastEyeAdaptation, FRHIUnorderedAccessView* DestUAV)
 	{
-		const FComputeShaderRHIParamRef ShaderRHI = GetComputeShader();
+		FRHIComputeShader* ShaderRHI = GetComputeShader();
 
 		FGlobalShader::SetParameters<FViewUniformShaderParameters>(Context.RHICmdList, ShaderRHI, Context.View.ViewUniformBuffer);
 		PostprocessParameter.SetCS(ShaderRHI, Context, Context.RHICmdList, TStaticSamplerState<SF_Bilinear, AM_Clamp, AM_Clamp, AM_Clamp>::GetRHI());
diff --git a/Engine/Source/Runtime/Renderer/Private/PostProcess/PostProcessFFTBloom.cpp b/Engine/Source/Runtime/Renderer/Private/PostProcess/PostProcessFFTBloom.cpp
--- a/Engine/Source/Runtime/Renderer/Private/PostProcess/PostProcessFFTBloom.cpp
+++ b/Engine/Source/Runtime/Renderer/Private/PostProcess/PostProcessFFTBloom.cpp
@@ -58,7 +58,7 @@ public:
 	// Determine the number of threads used per scanline when writing the physical space kernel
 	static int32 NumThreadsPerGroup() { return 32; }
 	// Method for use with the FScopedUAVBind
-	FShaderResourceParameter& DestinationResourceParamter() { return DstRWTexture; }
+	FShaderResourceParameter& DestinationResourceParameter() { return DstRWTexture; }
 
 	static const TCHAR* GetSourceFilename() { return TEXT("/Engine/Private/PostProcessFFTBloom.usf"); }
 	static const TCHAR* GetFunctionName()   { return TEXT("ResizeAndCenterTextureCS"); }
@@ -71,7 +71,7 @@ public:
 		float CenterScale = bForceCenterZero ? 0.f : 1.f;
 		const FLinearColor KernelCenterAndScaleValue(KernelUVCenter.X, KernelUVCenter.Y, ResizeScaleValue, CenterScale);
 
-		const FComputeShaderRHIParamRef ShaderRHI = GetComputeShader();
+		FRHIComputeShader* ShaderRHI = GetComputeShader();
 
 		// Set up the input.  We have to do this explicitly because the FFT dispatches multiple compute shaders and manages their input/output.
 
@@ -149,7 +149,7 @@ public:
 	}
 
 	// Method for use with the FScopedUAVBind
-	FShaderResourceParameter& DestinationResourceParamter() { return DstRWTexture; }
+	FShaderResourceParameter& DestinationResourceParameter() { return DstRWTexture; }
 
 	static const TCHAR* GetSourceFilename() { return TEXT("/Engine/Private/PostProcessFFTBloom.usf"); }
 	static const TCHAR* GetFunctionName()   { return TEXT("CaptureKernelWeightsCS"); }
@@ -161,7 +161,7 @@ public:
 		const FTextureRHIRef& PhysicalKernelRef,
 		const FVector2D& UVCenterValue)
 	{
-		const FComputeShaderRHIParamRef ShaderRHI = GetComputeShader();
+		FRHIComputeShader* ShaderRHI = GetComputeShader();
 
 
 
@@ -239,7 +239,7 @@ public:
 	}
 
 	// Method for use with the FScopedUAVBind
-	FShaderResourceParameter& DestinationResourceParamter() { return DstRWTexture; }
+	FShaderResourceParameter& DestinationResourceParameter() { return DstRWTexture; }
 
 	// Determine the number of threads used per scanline when writing the physical space kernel
 	static int32 NumThreadsPerGroup() { return 32; }
@@ -255,7 +255,7 @@ public:
 
 		using GPUFFTComputeShaderUtils::FComputeParamterValueSetter;
 
-		const FComputeShaderRHIParamRef ShaderRHI = GetComputeShader();
+		FRHIComputeShader* ShaderRHI = GetComputeShader();
 
 
 		// Set up the input.  We have to do this explicitly because the FFT dispatches multiple compute shaders and manages their input/output.
@@ -336,7 +336,7 @@ public:
 	}
 
 	// Method for use with the FScopedUAVBind
-	FShaderResourceParameter& DestinationResourceParamter() { return DstRWTexture; }
+	FShaderResourceParameter& DestinationResourceParameter() { return DstRWTexture; }
 
 	static int32 NumThreadsPerGroup() { return 32; }
 
@@ -347,7 +347,7 @@ public:
 		const FTextureRHIRef& SrcTexture, const FIntRect& SrcRectValue, const FIntRect& DstRectValue)
 	{
 		using GPUFFTComputeShaderUtils::FComputeParamterValueSetter;
-		const FComputeShaderRHIParamRef ShaderRHI = GetComputeShader();
+		FRHIComputeShader* ShaderRHI = GetComputeShader();
 
 
 
@@ -430,7 +430,7 @@ void ResizeAndCenterTexture(FRenderingCompositePassContext& Context,
 	RHICmdList.TransitionResource(EResourceTransitionAccess::ERWBarrier, EResourceTransitionPipeline::EGfxToCompute, DstUAV);
 
 	{
-		GPUFFTComputeShaderUtils::FScopedUAVBind ScopedBindOutput = GPUFFTComputeShaderUtils::FScopedUAVBind::BindOutput(RHICmdList, *ComputeShader, DstUAV);
+		GPUFFTComputeShaderUtils::FScopedUAVBind ScopedBindOutput = GPUFFTComputeShaderUtils::FScopedUAVBind::BindOutput(RHICmdList, ComputeShader, DstUAV);
 
 
 		ComputeShader->SetCSParamters(RHICmdList, Context, TargetSize, SrcImageSize, ResizeScale, ClampedImageCenterUV, SrcTexture, DstBufferSize, bForceCenterZero);
@@ -485,7 +485,7 @@ void CaptureKernelWeight(FRenderingCompositePassContext& Context,
 	check(DstTargetItem.UAV);
 	RHICmdList.TransitionResource(EResourceTransitionAccess::ERWBarrier, EResourceTransitionPipeline::EGfxToCompute, DstTargetItem.UAV);
 	{
-		GPUFFTComputeShaderUtils::FScopedUAVBind ScopedBindOutput = GPUFFTComputeShaderUtils::FScopedUAVBind::BindOutput(Context.RHICmdList, *ComputeShader, DstTargetItem.UAV);
+		GPUFFTComputeShaderUtils::FScopedUAVBind ScopedBindOutput = GPUFFTComputeShaderUtils::FScopedUAVBind::BindOutput(Context.RHICmdList, ComputeShader, DstTargetItem.UAV);
 		RHICmdList.SetUAVParameter(ComputeShader->GetComputeShader(), ComputeShader->DstRWTexture.GetBaseIndex(), DstTargetItem.UAV);
 
 		ComputeShader->SetCSParamters(RHICmdList, Context, HalfResKernel, HalfResSumLocation, PhysicalKernel, CenterUV);
@@ -545,7 +545,7 @@ void BlendLowRes(FRenderingCompositePassContext& Context,
 	RHICmdList.TransitionResource(EResourceTransitionAccess::ERWBarrier, EResourceTransitionPipeline::EComputeToCompute, DstUAV);
 	
 	{
-		GPUFFTComputeShaderUtils::FScopedUAVBind ScopedBindOutput = GPUFFTComputeShaderUtils::FScopedUAVBind::BindOutput(Context.RHICmdList, *ComputeShader, DstUAV);
+		GPUFFTComputeShaderUtils::FScopedUAVBind ScopedBindOutput = GPUFFTComputeShaderUtils::FScopedUAVBind::BindOutput(Context.RHICmdList, ComputeShader, DstUAV);
 		ComputeShader->SetCSParamters(RHICmdList, Context, FullResImageRect, HalfResRect, HalfBufferSize, CenterWeightTexutre, FullResImage, HalfResConvolvedImage);
 
 		FIntPoint TargetExtent = FullResImageRect.Size();
@@ -592,7 +592,7 @@ void CopyImageRect(FRenderingCompositePassContext& Context, const FSceneRenderTa
 	check(DstUAV);
 	RHICmdList.TransitionResource(EResourceTransitionAccess::ERWBarrier, EResourceTransitionPipeline::EGfxToCompute, DstUAV);
 	{
-		GPUFFTComputeShaderUtils::FScopedUAVBind ScopedBindOutput = GPUFFTComputeShaderUtils::FScopedUAVBind::BindOutput(RHICmdList, *ComputeShader, DstUAV);
+		GPUFFTComputeShaderUtils::FScopedUAVBind ScopedBindOutput = GPUFFTComputeShaderUtils::FScopedUAVBind::BindOutput(RHICmdList, ComputeShader, DstUAV);
 
 
 		ComputeShader->SetCSParamters(RHICmdList, Context, SrcTargetItem.ShaderResourceTexture, SrcRect, DstRect);
diff --git a/Engine/Source/Runtime/Renderer/Private/PostProcess/PostProcessHistogram.cpp b/Engine/Source/Runtime/Renderer/Private/PostProcess/PostProcessHistogram.cpp
--- a/Engine/Source/Runtime/Renderer/Private/PostProcess/PostProcessHistogram.cpp
+++ b/Engine/Source/Runtime/Renderer/Private/PostProcess/PostProcessHistogram.cpp
@@ -55,7 +55,7 @@ public:
 
 	void SetCS(FRHICommandList& RHICmdList, const FRenderingCompositePassContext& Context, FIntPoint ThreadGroupCountValue, FIntPoint LeftTopOffsetValue, FIntPoint GatherExtent)
 	{
-		const FComputeShaderRHIParamRef ShaderRHI = GetComputeShader();
+		FRHIComputeShader* ShaderRHI = GetComputeShader();
 
 		FGlobalShader::SetParameters<FViewUniformShaderParameters>(RHICmdList, ShaderRHI, Context.View.ViewUniformBuffer);
 
diff --git a/Engine/Source/Runtime/Renderer/Private/PostProcess/PostProcessMitchellNetravali.cpp b/Engine/Source/Runtime/Renderer/Private/PostProcess/PostProcessMitchellNetravali.cpp
--- a/Engine/Source/Runtime/Renderer/Private/PostProcess/PostProcessMitchellNetravali.cpp
+++ b/Engine/Source/Runtime/Renderer/Private/PostProcess/PostProcessMitchellNetravali.cpp
@@ -89,9 +89,9 @@ void FRCPassMitchellNetravaliDownsample::Process(FRenderingCompositePassContext&
 	{
 		Context.RHICmdList.TransitionResource(EResourceTransitionAccess::ERWBarrier, EResourceTransitionPipeline::EGfxToCompute, DestRenderTarget.UAV);
 
-		auto ShaderMap = Context.GetShaderMap();
+		auto* ShaderMap = Context.GetShaderMap();
 		TShaderMapRef<FMitchellNetravaliDownsamplCS> Shader(ShaderMap);
-		const FComputeShaderRHIParamRef ShaderRHI = Shader->GetComputeShader();
+		FRHIComputeShader* ShaderRHI = Shader->GetComputeShader();
 
 		Context.RHICmdList.SetComputeShader(ShaderRHI);
 
diff --git a/Engine/Source/Runtime/Renderer/Private/PostProcess/PostProcessMobile.cpp b/Engine/Source/Runtime/Renderer/Private/PostProcess/PostProcessMobile.cpp
--- a/Engine/Source/Runtime/Renderer/Private/PostProcess/PostProcessMobile.cpp
+++ b/Engine/Source/Runtime/Renderer/Private/PostProcess/PostProcessMobile.cpp
@@ -66,7 +66,7 @@ public:
 
 	void SetVS(const FRenderingCompositePassContext& Context)
 	{
-		const FVertexShaderRHIParamRef ShaderRHI = GetVertexShader();
+		FRHIVertexShader* ShaderRHI = GetVertexShader();
 
 		FGlobalShader::SetParameters<FViewUniformShaderParameters>(Context.RHICmdList, ShaderRHI, Context.View.ViewUniformBuffer);
 
@@ -123,7 +123,7 @@ public:
 
 	void SetPS(const FRenderingCompositePassContext& Context)
 	{
-		const FPixelShaderRHIParamRef ShaderRHI = GetPixelShader();
+		FRHIPixelShader* ShaderRHI = GetPixelShader();
 
 		const FPostProcessSettings& Settings = Context.View.FinalPostProcessSettings;
 
@@ -326,7 +326,7 @@ public:
 
 	void SetVS(const FRenderingCompositePassContext& Context)
 	{
-		const FVertexShaderRHIParamRef ShaderRHI = GetVertexShader();
+		FRHIVertexShader* ShaderRHI = GetVertexShader();
 
 		FGlobalShader::SetParameters<FViewUniformShaderParameters>(Context.RHICmdList, ShaderRHI, Context.View.ViewUniformBuffer);
 
@@ -368,7 +368,7 @@ public:
 	template <typename TRHICmdList>
 	void SetPS(TRHICmdList& RHICmdList, const FRenderingCompositePassContext& Context)
 	{
-		const FPixelShaderRHIParamRef ShaderRHI = GetPixelShader();
+		FRHIPixelShader* ShaderRHI = GetPixelShader();
 
 		const FPostProcessSettings& Settings = Context.View.FinalPostProcessSettings;
 
@@ -541,7 +541,7 @@ public:
 	template <typename TRHICmdList>
 	void SetPS(TRHICmdList& RHICmdList, const FRenderingCompositePassContext& Context)
 	{
-		const FPixelShaderRHIParamRef ShaderRHI = GetPixelShader();
+		FRHIPixelShader* ShaderRHI = GetPixelShader();
 		const FPostProcessSettings& Settings = Context.View.FinalPostProcessSettings;
 		FGlobalShader::SetParameters<FViewUniformShaderParameters>(RHICmdList, ShaderRHI, Context.View.ViewUniformBuffer);
 		PostprocessParameter.SetPS(RHICmdList, ShaderRHI, Context, TStaticSamplerState<SF_Bilinear,AM_Clamp,AM_Clamp,AM_Clamp>::GetRHI());
@@ -582,7 +582,7 @@ public:
 
 	void SetVS(const FRenderingCompositePassContext& Context, float InScale)
 	{
-		const FVertexShaderRHIParamRef ShaderRHI = GetVertexShader();
+		FRHIVertexShader* ShaderRHI = GetVertexShader();
 		FGlobalShader::SetParameters<FViewUniformShaderParameters>(Context.RHICmdList, ShaderRHI, Context.View.ViewUniformBuffer);
 		PostprocessParameter.SetVS(ShaderRHI, Context, TStaticSamplerState<SF_Bilinear,AM_Clamp,AM_Clamp,AM_Clamp>::GetRHI());
 		SetShaderValue(Context.RHICmdList, ShaderRHI, BloomDownScale, InScale);
@@ -723,7 +723,7 @@ public:
 	template <typename TRHICmdList>
 	void SetPS(TRHICmdList& RHICmdList, const FRenderingCompositePassContext& Context, FVector4& InTintA, FVector4& InTintB)
 	{
-		const FPixelShaderRHIParamRef ShaderRHI = GetPixelShader();
+		FRHIPixelShader* ShaderRHI = GetPixelShader();
 		const FPostProcessSettings& Settings = Context.View.FinalPostProcessSettings;
 		FGlobalShader::SetParameters<FViewUniformShaderParameters>(RHICmdList, ShaderRHI, Context.View.ViewUniformBuffer);
 		PostprocessParameter.SetPS(RHICmdList, ShaderRHI, Context, TStaticSamplerState<SF_Bilinear,AM_Clamp,AM_Clamp,AM_Clamp>::GetRHI());
@@ -766,7 +766,7 @@ public:
 
 	void SetVS(const FRenderingCompositePassContext& Context, FVector2D InScale)
 	{
-		const FVertexShaderRHIParamRef ShaderRHI = GetVertexShader();
+		FRHIVertexShader* ShaderRHI = GetVertexShader();
 		FGlobalShader::SetParameters<FViewUniformShaderParameters>(Context.RHICmdList, ShaderRHI, Context.View.ViewUniformBuffer);
 		PostprocessParameter.SetVS(ShaderRHI, Context, TStaticSamplerState<SF_Bilinear,AM_Clamp,AM_Clamp,AM_Clamp>::GetRHI());
 		SetShaderValue(Context.RHICmdList, ShaderRHI, BloomUpScales, InScale);
@@ -921,7 +921,7 @@ public:
 
 	void SetPS(const FRenderingCompositePassContext& Context)
 	{
-		const FPixelShaderRHIParamRef ShaderRHI = GetPixelShader();
+		FRHIPixelShader* ShaderRHI = GetPixelShader();
 		const FPostProcessSettings& Settings = Context.View.FinalPostProcessSettings;
 		FGlobalShader::SetParameters<FViewUniformShaderParameters>(Context.RHICmdList, ShaderRHI, Context.View.ViewUniformBuffer);
 		PostprocessParameter.SetPS(Context.RHICmdList, ShaderRHI, Context, TStaticSamplerState<SF_Bilinear,AM_Clamp,AM_Clamp,AM_Clamp>::GetRHI());
@@ -976,7 +976,7 @@ public:
 
 	void SetVS(const FRenderingCompositePassContext& Context)
 	{
-		const FVertexShaderRHIParamRef ShaderRHI = GetVertexShader();
+		FRHIVertexShader* ShaderRHI = GetVertexShader();
 		FGlobalShader::SetParameters<FViewUniformShaderParameters>(Context.RHICmdList, ShaderRHI, Context.View.ViewUniformBuffer);
 		PostprocessParameter.SetVS(ShaderRHI, Context, TStaticSamplerState<SF_Bilinear,AM_Clamp,AM_Clamp,AM_Clamp>::GetRHI());
 	}
@@ -1160,7 +1160,7 @@ public:
 
 	void SetPS(const FRenderingCompositePassContext& Context)
 	{
-		const FPixelShaderRHIParamRef ShaderRHI = GetPixelShader();
+		FRHIPixelShader* ShaderRHI = GetPixelShader();
 		const FPostProcessSettings& Settings = Context.View.FinalPostProcessSettings;
 		FGlobalShader::SetParameters<FViewUniformShaderParameters>(Context.RHICmdList, ShaderRHI, Context.View.ViewUniformBuffer);
 		PostprocessParameter.SetPS(Context.RHICmdList, ShaderRHI, Context, TStaticSamplerState<SF_Bilinear, AM_Clamp, AM_Clamp, AM_Clamp>::GetRHI());
@@ -1203,7 +1203,7 @@ public:
 
 	void SetVS(const FRenderingCompositePassContext& Context)
 	{
-		const FVertexShaderRHIParamRef ShaderRHI = GetVertexShader();
+		FRHIVertexShader* ShaderRHI = GetVertexShader();
 		FGlobalShader::SetParameters<FViewUniformShaderParameters>(Context.RHICmdList, ShaderRHI, Context.View.ViewUniformBuffer);
 		PostprocessParameter.SetVS(ShaderRHI, Context, TStaticSamplerState<SF_Bilinear,AM_Clamp,AM_Clamp,AM_Clamp>::GetRHI());
 
@@ -1359,7 +1359,7 @@ public:
 
 	void SetPS(const FRenderingCompositePassContext& Context)
 	{
-		const FPixelShaderRHIParamRef ShaderRHI = GetPixelShader();
+		FRHIPixelShader* ShaderRHI = GetPixelShader();
 		const FPostProcessSettings& Settings = Context.View.FinalPostProcessSettings;
 		FGlobalShader::SetParameters<FViewUniformShaderParameters>(Context.RHICmdList, ShaderRHI, Context.View.ViewUniformBuffer);
 		PostprocessParameter.SetPS(Context.RHICmdList, ShaderRHI, Context, TStaticSamplerState<SF_Bilinear, AM_Clamp, AM_Clamp, AM_Clamp>::GetRHI());
@@ -1400,7 +1400,7 @@ public:
 
 	void SetVS(const FRenderingCompositePassContext& Context)
 	{
-		const FVertexShaderRHIParamRef ShaderRHI = GetVertexShader();
+		FRHIVertexShader* ShaderRHI = GetVertexShader();
 		FGlobalShader::SetParameters<FViewUniformShaderParameters>(Context.RHICmdList, ShaderRHI, Context.View.ViewUniformBuffer);
 		PostprocessParameter.SetVS(ShaderRHI, Context, TStaticSamplerState<SF_Bilinear,AM_Clamp,AM_Clamp,AM_Clamp>::GetRHI());
 
@@ -1551,7 +1551,7 @@ public:
 
 	void SetPS(const FRenderingCompositePassContext& Context)
 	{
-		const FPixelShaderRHIParamRef ShaderRHI = GetPixelShader();
+		FRHIPixelShader* ShaderRHI = GetPixelShader();
 		const FPostProcessSettings& Settings = Context.View.FinalPostProcessSettings;
 		FGlobalShader::SetParameters<FViewUniformShaderParameters>(Context.RHICmdList, ShaderRHI, Context.View.ViewUniformBuffer);
 		PostprocessParameter.SetPS(Context.RHICmdList, ShaderRHI, Context, TStaticSamplerState<SF_Bilinear, AM_Clamp, AM_Clamp, AM_Clamp>::GetRHI());
@@ -1609,7 +1609,7 @@ public:
 
 	void SetVS(const FRenderingCompositePassContext& Context)
 	{
-		const FVertexShaderRHIParamRef ShaderRHI = GetVertexShader();
+		FRHIVertexShader* ShaderRHI = GetVertexShader();
 		FGlobalShader::SetParameters<FViewUniformShaderParameters>(Context.RHICmdList, ShaderRHI, Context.View.ViewUniformBuffer);
 		PostprocessParameter.SetVS(ShaderRHI, Context, TStaticSamplerState<SF_Bilinear,AM_Clamp,AM_Clamp,AM_Clamp>::GetRHI());
 
@@ -1790,7 +1790,7 @@ public:
 
 	void SetPS(const FRenderingCompositePassContext& Context)
 	{
-		const FPixelShaderRHIParamRef ShaderRHI = GetPixelShader();
+		FRHIPixelShader* ShaderRHI = GetPixelShader();
 		const FPostProcessSettings& Settings = Context.View.FinalPostProcessSettings;
 		FGlobalShader::SetParameters<FViewUniformShaderParameters>(Context.RHICmdList, ShaderRHI, Context.View.ViewUniformBuffer);
 		PostprocessParameter.SetPS(Context.RHICmdList, ShaderRHI, Context, TStaticSamplerState<SF_Bilinear,AM_Clamp,AM_Clamp,AM_Clamp>::GetRHI());
@@ -1839,7 +1839,7 @@ public:
 
 	void SetVS(const FRenderingCompositePassContext& Context)
 	{
-		const FVertexShaderRHIParamRef ShaderRHI = GetVertexShader();
+		FRHIVertexShader* ShaderRHI = GetVertexShader();
 		FGlobalShader::SetParameters<FViewUniformShaderParameters>(Context.RHICmdList, ShaderRHI, Context.View.ViewUniformBuffer);
 		PostprocessParameter.SetVS(ShaderRHI, Context, TStaticSamplerState<SF_Bilinear,AM_Clamp,AM_Clamp,AM_Clamp>::GetRHI());
 	}
@@ -1994,7 +1994,7 @@ public:
 
 	void SetVS(const FRenderingCompositePassContext& Context)
 	{
-		const FVertexShaderRHIParamRef ShaderRHI = GetVertexShader();
+		FRHIVertexShader* ShaderRHI = GetVertexShader();
 
 		FGlobalShader::SetParameters<FViewUniformShaderParameters>(Context.RHICmdList, ShaderRHI, Context.View.ViewUniformBuffer);
 
@@ -2039,7 +2039,7 @@ public:
 
 	void SetPS(const FRenderingCompositePassContext& Context)
 	{
-		const FPixelShaderRHIParamRef ShaderRHI = GetPixelShader();
+		FRHIPixelShader* ShaderRHI = GetPixelShader();
 		const FPostProcessSettings& Settings = Context.View.FinalPostProcessSettings;
 		FGlobalShader::SetParameters<FViewUniformShaderParameters>(Context.RHICmdList, ShaderRHI, Context.View.ViewUniformBuffer);
 		PostprocessParameter.SetPS(Context.RHICmdList, ShaderRHI, Context, TStaticSamplerState<SF_Point, AM_Clamp, AM_Clamp, AM_Clamp>::GetRHI());
@@ -2220,7 +2220,7 @@ public:
 
 	void SetVS(const FRenderingCompositePassContext& Context)
 	{
-		const FVertexShaderRHIParamRef ShaderRHI = GetVertexShader();
+		FRHIVertexShader* ShaderRHI = GetVertexShader();
 
 		FGlobalShader::SetParameters<FViewUniformShaderParameters>(Context.RHICmdList, ShaderRHI, Context.View.ViewUniformBuffer);
 
@@ -2265,7 +2265,7 @@ public:
 
 	void SetPS(const FRenderingCompositePassContext& Context)
 	{
-		const FPixelShaderRHIParamRef ShaderRHI = GetPixelShader();
+		FRHIPixelShader* ShaderRHI = GetPixelShader();
 		const FPostProcessSettings& Settings = Context.View.FinalPostProcessSettings;
 		FGlobalShader::SetParameters<FViewUniformShaderParameters>(Context.RHICmdList, ShaderRHI, Context.View.ViewUniformBuffer);
 		PostprocessParameter.SetPS(Context.RHICmdList, ShaderRHI, Context, TStaticSamplerState<SF_Point, AM_Clamp, AM_Clamp, AM_Clamp>::GetRHI());
@@ -2424,7 +2424,7 @@ public:
 
 	void SetPS(const FRenderingCompositePassContext& Context)
 	{
-		const FPixelShaderRHIParamRef ShaderRHI = GetPixelShader();
+		FRHIPixelShader* ShaderRHI = GetPixelShader();
 		const FPostProcessSettings& Settings = Context.View.FinalPostProcessSettings;
 		FGlobalShader::SetParameters<FViewUniformShaderParameters>(Context.RHICmdList, ShaderRHI, Context.View.ViewUniformBuffer);
 		PostprocessParameter.SetPS(Context.RHICmdList, ShaderRHI, Context, TStaticSamplerState<SF_Bilinear,AM_Clamp,AM_Clamp,AM_Clamp>::GetRHI());
@@ -2463,7 +2463,7 @@ public:
 
 	void SetVS(const FRenderingCompositePassContext& Context)
 	{
-		const FVertexShaderRHIParamRef ShaderRHI = GetVertexShader();
+		FRHIVertexShader* ShaderRHI = GetVertexShader();
 		FGlobalShader::SetParameters<FViewUniformShaderParameters>(Context.RHICmdList, ShaderRHI, Context.View.ViewUniformBuffer);
 		PostprocessParameter.SetVS(ShaderRHI, Context, TStaticSamplerState<SF_Point,AM_Clamp,AM_Clamp,AM_Clamp>::GetRHI());
 	}
@@ -2596,7 +2596,7 @@ public:
 
 	void SetPS(const FRenderingCompositePassContext& Context)
 	{
-		const FPixelShaderRHIParamRef ShaderRHI = GetPixelShader();
+		FRHIPixelShader* ShaderRHI = GetPixelShader();
 		const FPostProcessSettings& Settings = Context.View.FinalPostProcessSettings;
 		FGlobalShader::SetParameters<FViewUniformShaderParameters>(Context.RHICmdList, ShaderRHI, Context.View.ViewUniformBuffer);
 		PostprocessParameter.SetPS(Context.RHICmdList, ShaderRHI, Context, TStaticSamplerState<SF_Bilinear,AM_Clamp,AM_Clamp,AM_Clamp>::GetRHI());
@@ -2636,7 +2636,7 @@ public:
 
 	void SetVS(const FRenderingCompositePassContext& Context)
 	{
-		const FVertexShaderRHIParamRef ShaderRHI = GetVertexShader();
+		FRHIVertexShader* ShaderRHI = GetVertexShader();
 		FGlobalShader::SetParameters<FViewUniformShaderParameters>(Context.RHICmdList, ShaderRHI, Context.View.ViewUniformBuffer);
 		PostprocessParameter.SetVS(ShaderRHI, Context, TStaticSamplerState<SF_Bilinear,AM_Clamp,AM_Clamp,AM_Clamp>::GetRHI());
 	}
@@ -2782,7 +2782,7 @@ public:
 
 	void SetPS(const FRenderingCompositePassContext& Context)
 	{
-		const FPixelShaderRHIParamRef ShaderRHI = GetPixelShader();
+		FRHIPixelShader* ShaderRHI = GetPixelShader();
 		const FPostProcessSettings& Settings = Context.View.FinalPostProcessSettings;
 		FGlobalShader::SetParameters<FViewUniformShaderParameters>(Context.RHICmdList, ShaderRHI, Context.View.ViewUniformBuffer);
 		PostprocessParameter.SetPS(Context.RHICmdList, ShaderRHI, Context, TStaticSamplerState<SF_Bilinear,AM_Clamp,AM_Clamp,AM_Clamp>::GetRHI());
@@ -2887,7 +2887,7 @@ public:
 
 	void SetVS(const FRenderingCompositePassContext& Context)
 	{
-		const FVertexShaderRHIParamRef ShaderRHI = GetVertexShader();
+		FRHIVertexShader* ShaderRHI = GetVertexShader();
 		FGlobalShader::SetParameters<FViewUniformShaderParameters>(Context.RHICmdList, ShaderRHI, Context.View.ViewUniformBuffer);
 		PostprocessParameter.SetVS(ShaderRHI, Context, TStaticSamplerState<SF_Bilinear,AM_Clamp,AM_Clamp,AM_Clamp>::GetRHI());
 	}
diff --git a/Engine/Source/Runtime/Renderer/Private/PostProcess/PostProcessTonemap.cpp b/Engine/Source/Runtime/Renderer/Private/PostProcess/PostProcessTonemap.cpp
--- a/Engine/Source/Runtime/Renderer/Private/PostProcess/PostProcessTonemap.cpp
+++ b/Engine/Source/Runtime/Renderer/Private/PostProcess/PostProcessTonemap.cpp
@@ -743,7 +743,7 @@ public:
 
 	void SetPS(const FRenderingCompositePassContext& Context)
 	{
-		const FPixelShaderRHIParamRef ShaderRHI = GetPixelShader();
+		FRHIPixelShader* ShaderRHI = GetPixelShader();
 		
 		FGlobalShader::SetParameters<FViewUniformShaderParameters>(Context.RHICmdList, ShaderRHI, Context.View.ViewUniformBuffer);
 
@@ -832,7 +832,7 @@ public:
 	template <typename TRHICmdList>
 	void SetParameters(TRHICmdList& RHICmdList, const FRenderingCompositePassContext& Context, const FIntPoint& DestSize, FRHIUnorderedAccessView* DestUAV, FTextureRHIParamRef EyeAdaptationTex)
 	{
-		const FComputeShaderRHIParamRef ShaderRHI = GetComputeShader();
+		FRHIComputeShader* ShaderRHI = GetComputeShader();
 		const FPostProcessSettings& Settings = Context.View.FinalPostProcessSettings;
 		const FSceneViewFamily& ViewFamily = *(Context.View.Family);
 
@@ -884,7 +884,7 @@ public:
 	template <typename TRHICmdList>
 	void UnsetParameters(TRHICmdList& RHICmdList)
 	{
-		const FComputeShaderRHIParamRef ShaderRHI = GetComputeShader();
+		FRHIComputeShader* ShaderRHI = GetComputeShader();
 		OutComputeTex.UnsetUAV(RHICmdList, ShaderRHI);
 	}
 	
@@ -1473,7 +1473,7 @@ public:
 		const FPostProcessSettings& Settings = Context.View.FinalPostProcessSettings;
 		const FSceneViewFamily& ViewFamily = *(Context.View.Family);
 
-		const FPixelShaderRHIParamRef ShaderRHI = GetPixelShader();
+		FRHIPixelShader* ShaderRHI = GetPixelShader();
 		
 		FGlobalShader::SetParameters<FViewUniformShaderParameters>(RHICmdList, ShaderRHI, Context.View.ViewUniformBuffer);
 
@@ -1612,7 +1612,7 @@ public:
 
 	void SetVS(const FRenderingCompositePassContext& Context)
 	{
-		const FVertexShaderRHIParamRef ShaderRHI = GetVertexShader();
+		FRHIVertexShader* ShaderRHI = GetVertexShader();
 		FGlobalShader::SetParameters<FViewUniformShaderParameters>(Context.RHICmdList, ShaderRHI, Context.View.ViewUniformBuffer);
 
 		PostprocessParameter.SetVS(ShaderRHI, Context, TStaticSamplerState<SF_Bilinear,AM_Clamp,AM_Clamp,AM_Clamp>::GetRHI());
diff --git a/Engine/Source/Runtime/Renderer/Private/PostProcess/PostProcessTonemap.h b/Engine/Source/Runtime/Renderer/Private/PostProcess/PostProcessTonemap.h
--- a/Engine/Source/Runtime/Renderer/Private/PostProcess/PostProcessTonemap.h
+++ b/Engine/Source/Runtime/Renderer/Private/PostProcess/PostProcessTonemap.h
@@ -170,7 +170,7 @@ public:
 
 	void SetVS(const FRenderingCompositePassContext& Context, const FPermutationDomain& PermutationVector )
 	{
-		const FVertexShaderRHIParamRef ShaderRHI = GetVertexShader();
+		FRHIVertexShader* ShaderRHI = GetVertexShader();
 
 		FGlobalShader::SetParameters<FViewUniformShaderParameters>(Context.RHICmdList, ShaderRHI, Context.View.ViewUniformBuffer);
 
diff --git a/Engine/Source/Runtime/Renderer/Private/PostProcess/PostProcessWeightedSampleSum.cpp b/Engine/Source/Runtime/Renderer/Private/PostProcess/PostProcessWeightedSampleSum.cpp
--- a/Engine/Source/Runtime/Renderer/Private/PostProcess/PostProcessWeightedSampleSum.cpp
+++ b/Engine/Source/Runtime/Renderer/Private/PostProcess/PostProcessWeightedSampleSum.cpp
@@ -141,7 +141,7 @@ public:
 		const FLinearColor* SampleWeightValues, const FVector2D* SampleOffsetValues, uint32 NumSamples, const FVector4& FilteredBufferUVMinMaxValue, const FVector4& AdditiveBufferUVMinMaxValue )
 	{
 		check(CompileTimeNumSamples == 0 && NumSamples > 0 && NumSamples <= MAX_FILTER_SAMPLES || CompileTimeNumSamples == NumSamples);
-		const FPixelShaderRHIParamRef ShaderRHI = GetPixelShader();
+		FRHIPixelShader* ShaderRHI = GetPixelShader();
 
 		SetTextureParameter(RHICmdList, ShaderRHI, FilterTexture, FilterTextureSampler, SamplerStateRHI, FilterTextureRHI);
 		SetTextureParameter(RHICmdList, ShaderRHI, AdditiveTexture, AdditiveTextureSampler, SamplerStateRHI, AdditiveTextureRHI);
@@ -402,7 +402,7 @@ public:
 	void SetParameters(TRHICmdList& RHICmdList, const FRenderingCompositePassContext& Context, const FIntPoint& DestSize, FRHIUnorderedAccessView* DestUAV,
 		FTextureRHIParamRef FilterTextureRHI, FTextureRHIParamRef AdditiveTextureRHI, FLinearColor* SampleWeightValues, FVector2D* SampleOffsetValues, uint32 NumSamples)
 	{
-		const FComputeShaderRHIParamRef ShaderRHI = GetComputeShader();
+		FRHIComputeShader* ShaderRHI = GetComputeShader();
 		const FPostProcessSettings& Settings = Context.View.FinalPostProcessSettings;
 
 		// CS params
@@ -446,7 +446,7 @@ public:
 	template <typename TRHICmdList>
 	void UnsetParameters(TRHICmdList& RHICmdList)
 	{
-		const FComputeShaderRHIParamRef ShaderRHI = GetComputeShader();
+		FRHIComputeShader* ShaderRHI = GetComputeShader();
 		RHICmdList.SetUAVParameter(ShaderRHI, OutComputeTex.GetBaseIndex(), NULL);
 	}
 
diff --git a/Engine/Source/Runtime/Renderer/Private/PostProcess/RenderingCompositionGraph.cpp b/Engine/Source/Runtime/Renderer/Private/PostProcess/RenderingCompositionGraph.cpp
--- a/Engine/Source/Runtime/Renderer/Private/PostProcess/RenderingCompositionGraph.cpp
+++ b/Engine/Source/Runtime/Renderer/Private/PostProcess/RenderingCompositionGraph.cpp
@@ -852,21 +852,21 @@ void FPostProcessPassParameters::Bind(const FShaderParameterMap& ParameterMap)
 }
 
 template <typename TRHICmdList>
-void FPostProcessPassParameters::SetPS(TRHICmdList& RHICmdList, const FPixelShaderRHIParamRef& ShaderRHI, const FRenderingCompositePassContext& Context, FRHISamplerState* Filter, EFallbackColor FallbackColor, FRHISamplerState** FilterOverrideArray)
+void FPostProcessPassParameters::SetPS(TRHICmdList& RHICmdList, FRHIPixelShader* ShaderRHI, const FRenderingCompositePassContext& Context, FRHISamplerState* Filter, EFallbackColor FallbackColor, FRHISamplerState** FilterOverrideArray)
 {
 	Set(RHICmdList, ShaderRHI, Context, Filter, FallbackColor, FilterOverrideArray);
 }
 
-template void FPostProcessPassParameters::SetPS(FRHICommandList& RHICmdList, const FPixelShaderRHIParamRef& ShaderRHI, const FRenderingCompositePassContext& Context, FRHISamplerState* Filter, EFallbackColor FallbackColor, FRHISamplerState** FilterOverrideArray);
-template void FPostProcessPassParameters::SetPS(FRHICommandListImmediate& RHICmdList, const FPixelShaderRHIParamRef& ShaderRHI, const FRenderingCompositePassContext& Context, FRHISamplerState* Filter, EFallbackColor FallbackColor, FRHISamplerState** FilterOverrideArray);
+template void FPostProcessPassParameters::SetPS(FRHICommandList& RHICmdList, FRHIPixelShader* ShaderRHI, const FRenderingCompositePassContext& Context, FRHISamplerState* Filter, EFallbackColor FallbackColor, FRHISamplerState** FilterOverrideArray);
+template void FPostProcessPassParameters::SetPS(FRHICommandListImmediate& RHICmdList, FRHIPixelShader* ShaderRHI, const FRenderingCompositePassContext& Context, FRHISamplerState* Filter, EFallbackColor FallbackColor, FRHISamplerState** FilterOverrideArray);
 
 template< typename TRHICmdList >
-void FPostProcessPassParameters::SetCS(const FComputeShaderRHIParamRef& ShaderRHI, const FRenderingCompositePassContext& Context, TRHICmdList& RHICmdList, FRHISamplerState* Filter, EFallbackColor FallbackColor, FRHISamplerState** FilterOverrideArray)
+void FPostProcessPassParameters::SetCS(FRHIComputeShader* ShaderRHI, const FRenderingCompositePassContext& Context, TRHICmdList& RHICmdList, FRHISamplerState* Filter, EFallbackColor FallbackColor, FRHISamplerState** FilterOverrideArray)
 {
 	Set(RHICmdList, ShaderRHI, Context, Filter, FallbackColor, FilterOverrideArray);
 }
 template void FPostProcessPassParameters::SetCS< FRHICommandListImmediate >(
-	const FComputeShaderRHIParamRef& ShaderRHI,
+	FRHIComputeShader* ShaderRHI,
 	const FRenderingCompositePassContext& Context,
 	FRHICommandListImmediate& RHICmdList,
 	FRHISamplerState* Filter,
@@ -875,7 +875,7 @@ template void FPostProcessPassParameters::SetCS< FRHICommandListImmediate >(
 	);
 
 template void FPostProcessPassParameters::SetCS< FRHIAsyncComputeCommandListImmediate >(
-	const FComputeShaderRHIParamRef& ShaderRHI,
+	FRHIComputeShader* ShaderRHI,
 	const FRenderingCompositePassContext& Context,
 	FRHIAsyncComputeCommandListImmediate& RHICmdList,
 	FRHISamplerState* Filter,
@@ -883,15 +883,15 @@ template void FPostProcessPassParameters::SetCS< FRHIAsyncComputeCommandListImme
 	FRHISamplerState** FilterOverrideArray
 	);
 
-void FPostProcessPassParameters::SetVS(const FVertexShaderRHIParamRef& ShaderRHI, const FRenderingCompositePassContext& Context, FRHISamplerState* Filter, EFallbackColor FallbackColor, FRHISamplerState** FilterOverrideArray)
+void FPostProcessPassParameters::SetVS(FRHIVertexShader* ShaderRHI, const FRenderingCompositePassContext& Context, FRHISamplerState* Filter, EFallbackColor FallbackColor, FRHISamplerState** FilterOverrideArray)
 {
 	Set(Context.RHICmdList, ShaderRHI, Context, Filter, FallbackColor, FilterOverrideArray);
 }
 
-template< typename TShaderRHIParamRef, typename TRHICmdList >
+template< typename TRHIShader, typename TRHICmdList >
 void FPostProcessPassParameters::Set(
 	TRHICmdList& RHICmdList,
-	const TShaderRHIParamRef& ShaderRHI,
+	TRHIShader* ShaderRHI,
 	const FRenderingCompositePassContext& Context,
 	FRHISamplerState* Filter,
 	EFallbackColor FallbackColor,
@@ -1028,23 +1028,23 @@ void FPostProcessPassParameters::Set(
 	// todo warning if Input[] or InputSize[] is bound but not available, maybe set a specific input texture (blinking?)
 }
 
-#define IMPLEMENT_POST_PROCESS_PARAM_SET( TShaderRHIParamRef, TRHICmdList ) \
-	template void FPostProcessPassParameters::Set< TShaderRHIParamRef >( \
+#define IMPLEMENT_POST_PROCESS_PARAM_SET( TRHIShader, TRHICmdList ) \
+	template void FPostProcessPassParameters::Set< TRHIShader >( \
 		TRHICmdList& RHICmdList,						\
-		const TShaderRHIParamRef& ShaderRHI,			\
+		TRHIShader* ShaderRHI,			\
 		const FRenderingCompositePassContext& Context,	\
 		FRHISamplerState* Filter,				\
 		EFallbackColor FallbackColor,					\
 		FRHISamplerState** FilterOverrideArray	\
 	);
 
-IMPLEMENT_POST_PROCESS_PARAM_SET( FVertexShaderRHIParamRef, FRHICommandListImmediate );
-IMPLEMENT_POST_PROCESS_PARAM_SET( FHullShaderRHIParamRef, FRHICommandListImmediate );
-IMPLEMENT_POST_PROCESS_PARAM_SET( FDomainShaderRHIParamRef, FRHICommandListImmediate );
-IMPLEMENT_POST_PROCESS_PARAM_SET( FGeometryShaderRHIParamRef, FRHICommandListImmediate );
-IMPLEMENT_POST_PROCESS_PARAM_SET( FPixelShaderRHIParamRef, FRHICommandListImmediate );
-IMPLEMENT_POST_PROCESS_PARAM_SET( FComputeShaderRHIParamRef, FRHICommandListImmediate );
-IMPLEMENT_POST_PROCESS_PARAM_SET( FComputeShaderRHIParamRef, FRHIAsyncComputeCommandListImmediate );
+IMPLEMENT_POST_PROCESS_PARAM_SET(FRHIVertexShader, FRHICommandListImmediate );
+IMPLEMENT_POST_PROCESS_PARAM_SET(FRHIHullShader, FRHICommandListImmediate );
+IMPLEMENT_POST_PROCESS_PARAM_SET(FRHIDomainShader, FRHICommandListImmediate );
+IMPLEMENT_POST_PROCESS_PARAM_SET(FRHIGeometryShader, FRHICommandListImmediate );
+IMPLEMENT_POST_PROCESS_PARAM_SET(FRHIPixelShader, FRHICommandListImmediate );
+IMPLEMENT_POST_PROCESS_PARAM_SET(FRHIComputeShader, FRHICommandListImmediate );
+IMPLEMENT_POST_PROCESS_PARAM_SET(FRHIComputeShader, FRHIAsyncComputeCommandListImmediate );
 
 FArchive& operator<<(FArchive& Ar, FPostProcessPassParameters& P)
 {
diff --git a/Engine/Source/Runtime/Renderer/Private/PostProcess/SceneRenderTargets.cpp b/Engine/Source/Runtime/Renderer/Private/PostProcess/SceneRenderTargets.cpp
--- a/Engine/Source/Runtime/Renderer/Private/PostProcess/SceneRenderTargets.cpp
+++ b/Engine/Source/Runtime/Renderer/Private/PostProcess/SceneRenderTargets.cpp
@@ -1806,7 +1806,7 @@ void FSceneRenderTargets::ResolveDepthTexture(FRHICommandList& RHICmdList, const
 	TShaderMapRef<FResolveDepth8XPS> ResolvePixelShader8X(ShaderMap);
 
 	int32 TextureIndex = -1;
-	FPixelShaderRHIParamRef ResolvePixelShader;
+	FRHIPixelShader* ResolvePixelShader;
 	switch (SourceTexture->GetNumSamples())
 	{
 	case 2:
diff --git a/Engine/Source/Runtime/Renderer/Private/RayTracing/RayTracingRectLight.cpp b/Engine/Source/Runtime/Renderer/Private/RayTracing/RayTracingRectLight.cpp
--- a/Engine/Source/Runtime/Renderer/Private/RayTracing/RayTracingRectLight.cpp
+++ b/Engine/Source/Runtime/Renderer/Private/RayTracing/RayTracingRectLight.cpp
@@ -119,7 +119,7 @@ public:
 		FRWBuffer& MipTree
 	)
 	{
-		FComputeShaderRHIParamRef ShaderRHI = GetComputeShader();
+		FRHIComputeShader* ShaderRHI = GetComputeShader();
 
 		SetShaderValue(RHICmdList, ShaderRHI, DimensionsParameter, Dimensions);
 		SetShaderValue(RHICmdList, ShaderRHI, MipLevelParameter, MipLevel);
@@ -136,7 +136,7 @@ public:
 		FRWBuffer& MipTree,
 		FRHIComputeFence* Fence)
 	{
-		FComputeShaderRHIParamRef ShaderRHI = GetComputeShader();
+		FRHIComputeShader* ShaderRHI = GetComputeShader();
 
 		MipTreeParameter.UnsetUAV(RHICmdList, ShaderRHI);
 		RHICmdList.TransitionResource(TransitionAccess, TransitionPipeline, MipTree.UAV, Fence);
@@ -360,7 +360,7 @@ public:
 		const FRWBuffer& MipTree,
 		const FIntVector Dimensions)
 	{
-		const FPixelShaderRHIParamRef ShaderRHI = GetPixelShader();
+		FRHIPixelShader* ShaderRHI = GetPixelShader();
 		FGlobalShader::SetParameters<FViewUniformShaderParameters>(RHICmdList, ShaderRHI, View.ViewUniformBuffer);
 
 		SetShaderValue(RHICmdList, ShaderRHI, DimensionsParameter, Dimensions);
diff --git a/Engine/Source/Runtime/Renderer/Private/RayTracing/RaytracingSkylight.cpp b/Engine/Source/Runtime/Renderer/Private/RayTracing/RaytracingSkylight.cpp
--- a/Engine/Source/Runtime/Renderer/Private/RayTracing/RaytracingSkylight.cpp
+++ b/Engine/Source/Runtime/Renderer/Private/RayTracing/RaytracingSkylight.cpp
@@ -239,7 +239,7 @@ public:
 		uint32 MipLevel,
 		FRWBuffer& MipTree)
 	{
-		FComputeShaderRHIParamRef ShaderRHI = GetComputeShader();
+		FRHIComputeShader* ShaderRHI = GetComputeShader();
 
 		SetShaderValue(RHICmdList, ShaderRHI, DimensionsParameter, Dimensions);
 		SetShaderValue(RHICmdList, ShaderRHI, FaceIndexParameter, FaceIndex);
@@ -256,7 +256,7 @@ public:
 		EResourceTransitionPipeline TransitionPipeline,
 		FRWBuffer& MipTree)
 	{
-		FComputeShaderRHIParamRef ShaderRHI = GetComputeShader();
+		FRHIComputeShader* ShaderRHI = GetComputeShader();
 
 		MipTreeParameter.UnsetUAV(RHICmdList, ShaderRHI);
 		RHICmdList.TransitionResource(TransitionAccess, TransitionPipeline, MipTree.UAV);
@@ -387,7 +387,7 @@ public:
 		FRWBuffer& SolidAnglePdf
 	)
 	{
-		FComputeShaderRHIParamRef ShaderRHI = GetComputeShader();
+		FRHIComputeShader* ShaderRHI = GetComputeShader();
 
 		SetShaderValue(RHICmdList, ShaderRHI, MipLevelParameter, MipLevel);
 		SetShaderValue(RHICmdList, ShaderRHI, DimensionsParameter, Dimensions);
@@ -403,7 +403,7 @@ public:
 		FRHIComputeFence* Fence
 	)
 	{
-		FComputeShaderRHIParamRef ShaderRHI = GetComputeShader();
+		FRHIComputeShader* ShaderRHI = GetComputeShader();
 
 		SolidAnglePdfParameter.UnsetUAV(RHICmdList, ShaderRHI);
 		RHICmdList.TransitionResource(TransitionAccess, TransitionPipeline, MipTreePdf.UAV, Fence);
@@ -493,7 +493,7 @@ public:
 		uint32 MipLevel,
 		FRWBuffer& MipTreePdf)
 	{
-		FComputeShaderRHIParamRef ShaderRHI = GetComputeShader();
+		FRHIComputeShader* ShaderRHI = GetComputeShader();
 
 		SetSRVParameter(RHICmdList, ShaderRHI, MipTreeParameter, MipTree.SRV);
 		SetShaderValue(RHICmdList, ShaderRHI, DimensionsParameter, Dimensions);
@@ -509,7 +509,7 @@ public:
 		EResourceTransitionPipeline TransitionPipeline,
 		FRWBuffer& MipTreePdf)
 	{
-		FComputeShaderRHIParamRef ShaderRHI = GetComputeShader();
+		FRHIComputeShader* ShaderRHI = GetComputeShader();
 
 		MipTreePdfParameter.UnsetUAV(RHICmdList, ShaderRHI);
 		RHICmdList.TransitionResource(TransitionAccess, TransitionPipeline, MipTreePdf.UAV);
@@ -646,7 +646,7 @@ public:
 		const FRWBuffer& MipTreePosZ,
 		const FRWBuffer& MipTreeNegZ)
 	{
-		const FPixelShaderRHIParamRef ShaderRHI = GetPixelShader();
+		FRHIPixelShader* ShaderRHI = GetPixelShader();
 		FGlobalShader::SetParameters<FViewUniformShaderParameters>(RHICmdList, ShaderRHI, View.ViewUniformBuffer);
 
 		SetShaderValue(RHICmdList, ShaderRHI, DimensionsParameter, Dimensions);
diff --git a/Engine/Source/Runtime/Renderer/Private/ReflectionEnvironmentCapture.cpp b/Engine/Source/Runtime/Renderer/Private/ReflectionEnvironmentCapture.cpp
--- a/Engine/Source/Runtime/Renderer/Private/ReflectionEnvironmentCapture.cpp
+++ b/Engine/Source/Runtime/Renderer/Private/ReflectionEnvironmentCapture.cpp
@@ -217,7 +217,7 @@ void CreateCubeMips( FRHICommandListImmediate& RHICmdList, ERHIFeatureLevel::Typ
 			SetGraphicsPipelineState(RHICmdList, GraphicsPSOInit);
 
 			{
-				const FPixelShaderRHIParamRef ShaderRHI = PixelShader->GetPixelShader();
+				FRHIPixelShader* ShaderRHI = PixelShader->GetPixelShader();
 
 				SetShaderValue(RHICmdList, ShaderRHI, PixelShader->CubeFace, CubeFace);
 				SetShaderValue(RHICmdList, ShaderRHI, PixelShader->MipIndex, MipIndex);
@@ -433,7 +433,7 @@ void FilterReflectionEnvironment(FRHICommandListImmediate& RHICmdList, ERHIFeatu
 				SetGraphicsPipelineState(RHICmdList, GraphicsPSOInit);
 
 				{
-					const FPixelShaderRHIParamRef ShaderRHI = PixelShader->GetPixelShader();
+					FRHIPixelShader* ShaderRHI = PixelShader->GetPixelShader();
 
 					SetShaderValue( RHICmdList, ShaderRHI, PixelShader->CubeFace, CubeFace );
 					SetShaderValue( RHICmdList, ShaderRHI, PixelShader->MipIndex, MipIndex );
@@ -533,7 +533,7 @@ public:
 
 	void SetParameters(FRHICommandList& RHICmdList, const FViewInfo& View, bool bCapturingForSkyLight, bool bLowerHemisphereIsBlack, const FLinearColor& LowerHemisphereColorValue)
 	{
-		const FPixelShaderRHIParamRef ShaderRHI = GetPixelShader();
+		FRHIPixelShader* ShaderRHI = GetPixelShader();
 
 		FGlobalShader::SetParameters<FViewUniformShaderParameters>(RHICmdList, ShaderRHI, View.ViewUniformBuffer);
 		SceneTextureParameters.Set(RHICmdList, ShaderRHI, View.FeatureLevel, ESceneTextureSetupMode::All);
@@ -616,7 +616,7 @@ public:
 
 	void SetParameters(FRHICommandList& RHICmdList, const FTexture* SourceCubemap, uint32 CubeFaceValue, bool bIsSkyLight, bool bLowerHemisphereIsBlack, float SourceCubemapRotation, const FLinearColor& LowerHemisphereColorValue)
 	{
-		const FPixelShaderRHIParamRef ShaderRHI = GetPixelShader();
+		FRHIPixelShader* ShaderRHI = GetPixelShader();
 
 		SetShaderValue(RHICmdList, ShaderRHI, CubeFace, CubeFaceValue);
 
diff --git a/Engine/Source/Runtime/Renderer/Private/RendererUtils.cpp b/Engine/Source/Runtime/Renderer/Private/RendererUtils.cpp
--- a/Engine/Source/Runtime/Renderer/Private/RendererUtils.cpp
+++ b/Engine/Source/Runtime/Renderer/Private/RendererUtils.cpp
@@ -139,7 +139,7 @@ void FRenderTargetWriteMask::Decode(FRHICommandListImmediate& RHICmdList, TShade
 	RHICmdList.SetUAVParameter(ComputeShader->GetComputeShader(), ComputeShader->OutCombinedRTWriteMask.GetBaseIndex(), OutRTWriteMask->GetRenderTargetItem().UAV);
 	ComputeShader->SetParameters(RHICmdList, RTWriteMaskDims, InRenderTargets);
 
-	const FComputeShaderRHIParamRef ShaderRHI = ComputeShader->GetComputeShader();
+	FRHIComputeShader* ShaderRHI = ComputeShader->GetComputeShader();
 
 	RHICmdList.TransitionResource(EResourceTransitionAccess::EWritable, EResourceTransitionPipeline::EGfxToCompute, OutRTWriteMask->GetRenderTargetItem().UAV);
 	{
diff --git a/Engine/Source/Runtime/Renderer/Private/ScreenPass.h b/Engine/Source/Runtime/Renderer/Private/ScreenPass.h
--- a/Engine/Source/Runtime/Renderer/Private/ScreenPass.h
+++ b/Engine/Source/Runtime/Renderer/Private/ScreenPass.h
@@ -174,7 +174,7 @@ void DrawScreenPass(
 	const typename TPixelShaderType::FParameters& PixelShaderParameters)
 {
 	const FIntPoint OutputSize = OutputRect.Size();
-	const FPixelShaderRHIParamRef PixelShaderRHI = GETSAFERHISHADER_PIXEL(PixelShader);
+	FRHIPixelShader* PixelShaderRHI = GETSAFERHISHADER_PIXEL(PixelShader);
 	FScreenPassVS* VertexShader = *(ScreenPassView.ScreenPassVS);
 
 	FGraphicsPipelineStateInitializer GraphicsPSOInit;
diff --git a/Engine/Source/Runtime/Renderer/Private/ShaderBaseClasses.cpp b/Engine/Source/Runtime/Renderer/Private/ShaderBaseClasses.cpp
--- a/Engine/Source/Runtime/Renderer/Private/ShaderBaseClasses.cpp
+++ b/Engine/Source/Runtime/Renderer/Private/ShaderBaseClasses.cpp
@@ -196,10 +196,10 @@ void FMaterialShader::VerifyExpressionAndShaderMaps(const FMaterialRenderProxy*
 }
 #endif
 
-template<typename ShaderRHIParamRef>
+template<typename TRHIShader>
 void FMaterialShader::SetParametersInner(
 	FRHICommandList& RHICmdList,
-	const ShaderRHIParamRef ShaderRHI,
+	TRHIShader* ShaderRHI,
 	const FMaterialRenderProxy* MaterialRenderProxy,
 	const FMaterial& Material,
 	const FSceneView& View)
@@ -310,10 +310,10 @@ void FMaterialShader::SetParametersInner(
 	}
 }
 
-template<typename ShaderRHIParamRef>
+template<typename TRHIShader>
 void FMaterialShader::SetParameters(
 	FRHICommandList& RHICmdList,
-	const ShaderRHIParamRef ShaderRHI,
+	TRHIShader* ShaderRHI,
 	const FMaterialRenderProxy* MaterialRenderProxy,
 	const FMaterial& Material,
 	const FSceneView& View,
@@ -329,31 +329,31 @@ void FMaterialShader::SetParameters(
 // Doxygen struggles to parse these explicit specializations. Just ignore them for now.
 #if !UE_BUILD_DOCS
 
-#define IMPLEMENT_MATERIAL_SHADER_SetParametersInner( ShaderRHIParamRef ) \
-	template RENDERER_API void FMaterialShader::SetParametersInner< ShaderRHIParamRef >( \
+#define IMPLEMENT_MATERIAL_SHADER_SetParametersInner( TRHIShader ) \
+	template RENDERER_API void FMaterialShader::SetParametersInner< TRHIShader >( \
 		FRHICommandList& RHICmdList,					\
-		const ShaderRHIParamRef ShaderRHI,				\
+		TRHIShader* ShaderRHI,							\
 		const FMaterialRenderProxy* MaterialRenderProxy,\
 		const FMaterial& Material,						\
 		const FSceneView& View							\
 	);
 
-IMPLEMENT_MATERIAL_SHADER_SetParametersInner( FVertexShaderRHIParamRef );
-IMPLEMENT_MATERIAL_SHADER_SetParametersInner( FHullShaderRHIParamRef );
-IMPLEMENT_MATERIAL_SHADER_SetParametersInner( FDomainShaderRHIParamRef );
-IMPLEMENT_MATERIAL_SHADER_SetParametersInner( FGeometryShaderRHIParamRef );
-IMPLEMENT_MATERIAL_SHADER_SetParametersInner( FPixelShaderRHIParamRef );
-IMPLEMENT_MATERIAL_SHADER_SetParametersInner( FComputeShaderRHIParamRef );
+IMPLEMENT_MATERIAL_SHADER_SetParametersInner(FRHIVertexShader);
+IMPLEMENT_MATERIAL_SHADER_SetParametersInner(FRHIHullShader);
+IMPLEMENT_MATERIAL_SHADER_SetParametersInner(FRHIDomainShader);
+IMPLEMENT_MATERIAL_SHADER_SetParametersInner(FRHIGeometryShader);
+IMPLEMENT_MATERIAL_SHADER_SetParametersInner(FRHIPixelShader);
+IMPLEMENT_MATERIAL_SHADER_SetParametersInner(FRHIComputeShader);
 
 #endif
 
 // Doxygen struggles to parse these explicit specializations. Just ignore them for now.
 #if !UE_BUILD_DOCS
 
-#define IMPLEMENT_MATERIAL_SHADER_SetParameters( ShaderRHIParamRef ) \
-	template RENDERER_API void FMaterialShader::SetParameters< ShaderRHIParamRef >( \
+#define IMPLEMENT_MATERIAL_SHADER_SetParameters( TRHIShader ) \
+	template RENDERER_API void FMaterialShader::SetParameters< TRHIShader >( \
 		FRHICommandList& RHICmdList,					\
-		const ShaderRHIParamRef ShaderRHI,				\
+		TRHIShader* ShaderRHI,							\
 		const FMaterialRenderProxy* MaterialRenderProxy,\
 		const FMaterial& Material,						\
 		const FSceneView& View,							\
@@ -361,12 +361,12 @@ IMPLEMENT_MATERIAL_SHADER_SetParametersInner( FComputeShaderRHIParamRef );
 		ESceneTextureSetupMode SceneTextureSetupMode		\
 	);
 
-IMPLEMENT_MATERIAL_SHADER_SetParameters( FVertexShaderRHIParamRef );
-IMPLEMENT_MATERIAL_SHADER_SetParameters( FHullShaderRHIParamRef );
-IMPLEMENT_MATERIAL_SHADER_SetParameters( FDomainShaderRHIParamRef );
-IMPLEMENT_MATERIAL_SHADER_SetParameters( FGeometryShaderRHIParamRef );
-IMPLEMENT_MATERIAL_SHADER_SetParameters( FPixelShaderRHIParamRef );
-IMPLEMENT_MATERIAL_SHADER_SetParameters( FComputeShaderRHIParamRef );
+IMPLEMENT_MATERIAL_SHADER_SetParameters(FRHIVertexShader);
+IMPLEMENT_MATERIAL_SHADER_SetParameters(FRHIHullShader);
+IMPLEMENT_MATERIAL_SHADER_SetParameters(FRHIDomainShader);
+IMPLEMENT_MATERIAL_SHADER_SetParameters(FRHIGeometryShader);
+IMPLEMENT_MATERIAL_SHADER_SetParameters(FRHIPixelShader);
+IMPLEMENT_MATERIAL_SHADER_SetParameters(FRHIComputeShader);
 
 #endif
 
diff --git a/Engine/Source/Runtime/Renderer/Private/ShaderPrint.cpp b/Engine/Source/Runtime/Renderer/Private/ShaderPrint.cpp
--- a/Engine/Source/Runtime/Renderer/Private/ShaderPrint.cpp
+++ b/Engine/Source/Runtime/Renderer/Private/ShaderPrint.cpp
@@ -111,28 +111,28 @@ namespace ShaderPrint
 	}
 
 	template<typename TShaderRHIParamRef>
-	void SetShaderParameters(FShaderParametersLegacy const* P, FRHICommandListImmediate& RHICmdList, TShaderRHIParamRef const& ShaderRHI, FViewInfo const& View)
+	void SetShaderParameters(FShaderParametersLegacy const* P, FRHICommandListImmediate& RHICmdList, TShaderRHIParamRef ShaderRHI, FViewInfo const& View)
 	{
 		SetUniformBufferParameter(RHICmdList, ShaderRHI, P->UniformBufferParameter, CreateUniformBuffer(View));
 		P->ValuesBufferParameter.SetBuffer(RHICmdList, ShaderRHI, View.ShaderPrintValueBuffer);
 	}
 
-	void FShaderParametersLegacy::SetParameters(FRHICommandListImmediate& RHICmdList, FVertexShaderRHIParamRef const& ShaderRHI, FViewInfo const& View)
+	void FShaderParametersLegacy::SetParameters(FRHICommandListImmediate& RHICmdList, FRHIVertexShader* ShaderRHI, FViewInfo const& View)
 	{
 		SetShaderParameters(this, RHICmdList, ShaderRHI, View);
 	}
 
-	void FShaderParametersLegacy::SetParameters(FRHICommandListImmediate& RHICmdList, FPixelShaderRHIParamRef const& ShaderRHI, FViewInfo const& View)
+	void FShaderParametersLegacy::SetParameters(FRHICommandListImmediate& RHICmdList, FRHIPixelShader* ShaderRHI, FViewInfo const& View)
 	{
 		SetShaderParameters(this, RHICmdList, ShaderRHI, View);
 	}
 
-	void FShaderParametersLegacy::SetParameters(FRHICommandListImmediate& RHICmdList, FComputeShaderRHIParamRef const& ShaderRHI, FViewInfo const& View)
+	void FShaderParametersLegacy::SetParameters(FRHICommandListImmediate& RHICmdList, FRHIComputeShader* ShaderRHI, FViewInfo const& View)
 	{
 		SetShaderParameters(this, RHICmdList, ShaderRHI, View);
 	}
 
-	void FShaderParametersLegacy::UnsetUAV(FRHICommandListImmediate& RHICmdList, FComputeShaderRHIParamRef const& ShaderRHI)
+	void FShaderParametersLegacy::UnsetUAV(FRHICommandListImmediate& RHICmdList, FRHIComputeShader* ShaderRHI)
 	{
 		ValuesBufferParameter.UnsetUAV(RHICmdList, ShaderRHI);
 	}
diff --git a/Engine/Source/Runtime/Renderer/Private/ShaderPrintParameters.h b/Engine/Source/Runtime/Renderer/Private/ShaderPrintParameters.h
--- a/Engine/Source/Runtime/Renderer/Private/ShaderPrintParameters.h
+++ b/Engine/Source/Runtime/Renderer/Private/ShaderPrintParameters.h
@@ -47,12 +47,12 @@ namespace ShaderPrint
 		void Bind(FShaderParameterMap const& ParameterMap);
 
 		// Call to set parameters
-		void SetParameters(FRHICommandListImmediate& RHICmdList, FVertexShaderRHIParamRef const& ShaderRHI, FViewInfo const& View);
-		void SetParameters(FRHICommandListImmediate& RHICmdList, FPixelShaderRHIParamRef const& ShaderRHI, FViewInfo const& View);
-		void SetParameters(FRHICommandListImmediate& RHICmdList, FComputeShaderRHIParamRef const& ShaderRHI, FViewInfo const& View);
+		void SetParameters(FRHICommandListImmediate& RHICmdList, FRHIVertexShader* ShaderRHI, FViewInfo const& View);
+		void SetParameters(FRHICommandListImmediate& RHICmdList, FRHIPixelShader* ShaderRHI, FViewInfo const& View);
+		void SetParameters(FRHICommandListImmediate& RHICmdList, FRHIComputeShader* ShaderRHI, FViewInfo const& View);
 
 		// Call to unbind the UAV
-		void UnsetUAV(FRHICommandListImmediate& RHICmdList, FComputeShaderRHIParamRef const& ShaderRHI);
+		void UnsetUAV(FRHICommandListImmediate& RHICmdList, FRHIComputeShader* ShaderRHI);
 
 		// Serializer
 		friend FArchive& operator<<(FArchive& Ar, FShaderParametersLegacy& P);
diff --git a/Engine/Source/Runtime/Renderer/Private/ShadowRendering.h b/Engine/Source/Runtime/Renderer/Private/ShadowRendering.h
--- a/Engine/Source/Runtime/Renderer/Private/ShadowRendering.h
+++ b/Engine/Source/Runtime/Renderer/Private/ShadowRendering.h
@@ -780,7 +780,7 @@ public:
 
 	void Set(FRHICommandList& RHICmdList, FShader* Shader, const FSceneView& View, const FProjectedShadowInfo* ShadowInfo)
 	{
-		const FPixelShaderRHIParamRef ShaderRHI = Shader->GetPixelShader();
+		FRHIPixelShader* ShaderRHI = Shader->GetPixelShader();
 
 		SceneTextureParameters.Set(RHICmdList, ShaderRHI, View.FeatureLevel, ESceneTextureSetupMode::All);
 
@@ -960,7 +960,7 @@ public:
 		const FProjectedShadowInfo* ShadowInfo
 		) override
 	{
-		const FPixelShaderRHIParamRef ShaderRHI = GetPixelShader();
+		FRHIPixelShader* ShaderRHI = GetPixelShader();
 
 		FShadowProjectionPixelShaderInterface::SetParameters(RHICmdList, ViewIndex,View,ShadowInfo);
 
@@ -1059,7 +1059,7 @@ public:
 		const FProjectedShadowInfo* ShadowInfo) override
 	{
 		TShadowProjectionPS<Quality, false, true>::SetParameters(RHICmdList, ViewIndex, View, ShadowInfo);
-		const FPixelShaderRHIParamRef ShaderRHI = this->GetPixelShader();
+		FRHIPixelShader* ShaderRHI = this->GetPixelShader();
 		SetShaderValue(RHICmdList, ShaderRHI, ModulatedShadowColorParameter, ShadowInfo->GetLightSceneInfo().Proxy->GetModulatedShadowColor());
 	}
 
@@ -1306,7 +1306,7 @@ public:
 		const FProjectedShadowInfo* ShadowInfo
 		)
 	{
-		const FPixelShaderRHIParamRef ShaderRHI = GetPixelShader();
+		FRHIPixelShader* ShaderRHI = GetPixelShader();
 
 		FGlobalShader::SetParameters<FViewUniformShaderParameters>(RHICmdList, ShaderRHI,View.ViewUniformBuffer);
 
@@ -1436,7 +1436,7 @@ public:
 	{
 		TShadowProjectionPS<Quality, bUseFadePlane>::SetParameters(RHICmdList, ViewIndex, View, ShadowInfo);
 
-		const FPixelShaderRHIParamRef ShaderRHI = this->GetPixelShader();
+		FRHIPixelShader* ShaderRHI = this->GetPixelShader();
 
 		// GetLightSourceAngle returns the full angle.
 		float TanLightSourceAngle = FMath::Tan(0.5 * FMath::DegreesToRadians(ShadowInfo->GetLightSceneInfo().Proxy->GetLightSourceAngle()));
@@ -1505,7 +1505,7 @@ public:
 
 		TShadowProjectionPS<Quality, bUseFadePlane>::SetParameters(RHICmdList, ViewIndex, View, ShadowInfo);
 
-		const FPixelShaderRHIParamRef ShaderRHI = this->GetPixelShader();
+		FRHIPixelShader* ShaderRHI = this->GetPixelShader();
 
 		static IConsoleVariable* CVarMaxSoftShadowKernelSize = IConsoleManager::Get().FindConsoleVariable(TEXT("r.Shadow.MaxSoftKernelSize"));
 		check(CVarMaxSoftShadowKernelSize);
diff --git a/Engine/Source/Runtime/Renderer/Private/SurfelTree.cpp b/Engine/Source/Runtime/Renderer/Private/SurfelTree.cpp
--- a/Engine/Source/Runtime/Renderer/Private/SurfelTree.cpp
+++ b/Engine/Source/Runtime/Renderer/Private/SurfelTree.cpp
@@ -73,7 +73,7 @@ public:
 		const FUniformMeshBuffers& UniformMeshBuffers
 		)
 	{
-		FComputeShaderRHIParamRef ShaderRHI = GetComputeShader();
+		FRHIComputeShader* ShaderRHI = GetComputeShader();
 		FGlobalShader::SetParameters<FViewUniformShaderParameters>(RHICmdList, ShaderRHI, View.ViewUniformBuffer);
 
 		SetShaderValue(RHICmdList, ShaderRHI, NumTriangles, NumTrianglesValue);
@@ -84,7 +84,7 @@ public:
 
 	void UnsetParameters(FRHICommandList& RHICmdList)
 	{
-		FComputeShaderRHIParamRef ShaderRHI = GetComputeShader();
+		FRHIComputeShader* ShaderRHI = GetComputeShader();
 		TriangleAreas.UnsetUAV(RHICmdList, ShaderRHI);
 		// RHISetStreamOutTargets doesn't unbind existing uses like render targets do 
 		SetSRVParameter(RHICmdList, ShaderRHI, TriangleVertexData, NULL);
@@ -145,7 +145,7 @@ public:
 		const FUniformMeshBuffers& UniformMeshBuffers
 		)
 	{
-		FComputeShaderRHIParamRef ShaderRHI = GetComputeShader();
+		FRHIComputeShader* ShaderRHI = GetComputeShader();
 		FGlobalShader::SetParameters<FViewUniformShaderParameters>(RHICmdList, ShaderRHI, View.ViewUniformBuffer);
 
 		SetShaderValue(RHICmdList, ShaderRHI, NumTriangles, NumTrianglesValue);
@@ -156,7 +156,7 @@ public:
 
 	void UnsetParameters(FRHICommandList& RHICmdList)
 	{
-		FComputeShaderRHIParamRef ShaderRHI = GetComputeShader();
+		FRHIComputeShader* ShaderRHI = GetComputeShader();
 		TriangleCDFs.UnsetUAV(RHICmdList, ShaderRHI);
 	}
 
@@ -220,7 +220,7 @@ public:
 		const FUniformMeshBuffers& UniformMeshBuffers
 		)
 	{
-		FComputeShaderRHIParamRef ShaderRHI = GetComputeShader();
+		FRHIComputeShader* ShaderRHI = GetComputeShader();
 		FGlobalShader::SetParameters<FViewUniformShaderParameters>(RHICmdList, ShaderRHI, View.ViewUniformBuffer);
 
 		const FScene* Scene = (const FScene*)View.Family->Scene;
@@ -241,7 +241,7 @@ public:
 
 	void UnsetParameters(FRHICommandList& RHICmdList, const FSceneView& View)
 	{
-		FComputeShaderRHIParamRef ShaderRHI = GetComputeShader();
+		FRHIComputeShader* ShaderRHI = GetComputeShader();
 
 		const FScene* Scene = (const FScene*)View.Family->Scene;
 		SurfelBufferParameters.UnsetParameters(RHICmdList, ShaderRHI);
diff --git a/Engine/Source/Runtime/Renderer/Private/TiledDeferredLightRendering.cpp b/Engine/Source/Runtime/Renderer/Private/TiledDeferredLightRendering.cpp
--- a/Engine/Source/Runtime/Renderer/Private/TiledDeferredLightRendering.cpp
+++ b/Engine/Source/Runtime/Renderer/Private/TiledDeferredLightRendering.cpp
@@ -133,7 +133,7 @@ public:
 		IPooledRenderTarget& InTextureValue,
 		IPooledRenderTarget& OutTextureValue)
 	{
-		FComputeShaderRHIParamRef ShaderRHI = GetComputeShader();
+		FRHIComputeShader* ShaderRHI = GetComputeShader();
 
 		FGlobalShader::SetParameters<FViewUniformShaderParameters>(RHICmdList, ShaderRHI, View.ViewUniformBuffer);
 		SceneTextureParameters.Set(RHICmdList, ShaderRHI, View.FeatureLevel, ESceneTextureSetupMode::All);
diff --git a/Engine/Source/Runtime/Renderer/Private/TranslucentLighting.cpp b/Engine/Source/Runtime/Renderer/Private/TranslucentLighting.cpp
--- a/Engine/Source/Runtime/Renderer/Private/TranslucentLighting.cpp
+++ b/Engine/Source/Runtime/Renderer/Private/TranslucentLighting.cpp
@@ -652,7 +652,7 @@ public:
 	void SetParameters(FRHICommandList& RHICmdList, const FViewInfo& View, int32 VolumeCascadeIndex, const int32 ViewIndex)
 	{
 		FSceneRenderTargets& SceneContext = FSceneRenderTargets::Get(RHICmdList);
-		const FPixelShaderRHIParamRef ShaderRHI = GetPixelShader();
+		FRHIPixelShader* ShaderRHI = GetPixelShader();
 
 		FGlobalShader::SetParameters<FViewUniformShaderParameters>(RHICmdList, ShaderRHI, View.ViewUniformBuffer);
 
@@ -850,7 +850,7 @@ public:
 	{
 		check(ShadowMap || !bDynamicallyShadowed);
 		
-		const FPixelShaderRHIParamRef ShaderRHI = GetPixelShader();
+		FRHIPixelShader* ShaderRHI = GetPixelShader();
 
 		FMaterialShader::SetParameters(RHICmdList, ShaderRHI, MaterialProxy, *MaterialProxy->GetMaterial(View.GetFeatureLevel()), View, View.ViewUniformBuffer, ESceneTextureSetupMode::All);
 		
@@ -964,7 +964,7 @@ public:
 	)
 	{
 		check(NumUAVs == 4);
-		const FComputeShaderRHIParamRef ShaderRHI = GetComputeShader();
+		FRHIComputeShader* ShaderRHI = GetComputeShader();
 		Ambient0.SetTexture(RHICmdList, ShaderRHI, NULL, VolumeUAVs[0]);
 		Directional0.SetTexture(RHICmdList, ShaderRHI, NULL, VolumeUAVs[1]);
 		Ambient1.SetTexture(RHICmdList, ShaderRHI, NULL, VolumeUAVs[2]);
@@ -973,7 +973,7 @@ public:
 
 	void UnsetParameters(FRHIAsyncComputeCommandListImmediate& RHICmdList)
 	{
-		const FComputeShaderRHIParamRef ShaderRHI = GetComputeShader();
+		FRHIComputeShader* ShaderRHI = GetComputeShader();
 		Ambient0.UnsetUAV(RHICmdList, ShaderRHI);
 		Directional0.UnsetUAV(RHICmdList, ShaderRHI);
 		Ambient1.UnsetUAV(RHICmdList, ShaderRHI);
@@ -1088,7 +1088,7 @@ public:
 
 	void SetParameters(FRHICommandList& RHICmdList, const FViewInfo& View, const FFinalPostProcessSettings::FCubemapEntry& CubemapEntry)
 	{
-		const FPixelShaderRHIParamRef ShaderRHI = GetPixelShader();
+		FRHIPixelShader* ShaderRHI = GetPixelShader();
 
 		FGlobalShader::SetParameters<FViewUniformShaderParameters>(RHICmdList, ShaderRHI, View.ViewUniformBuffer);
 
diff --git a/Engine/Source/Runtime/Renderer/Private/VT/VirtualTextureSpace.cpp b/Engine/Source/Runtime/Renderer/Private/VT/VirtualTextureSpace.cpp
--- a/Engine/Source/Runtime/Renderer/Private/VT/VirtualTextureSpace.cpp
+++ b/Engine/Source/Runtime/Renderer/Private/VT/VirtualTextureSpace.cpp
@@ -468,7 +468,7 @@ void FVirtualTextureSpace::ApplyUpdates(FVirtualTextureSystem* System, FRHIComma
 				SetGraphicsPipelineState(RHICmdList, GraphicsPSOInit);
 
 				{
-					const FVertexShaderRHIParamRef ShaderRHI = VertexShader->GetVertexShader();
+					FRHIVertexShader* ShaderRHI = VertexShader->GetVertexShader();
 					SetShaderValue(RHICmdList, ShaderRHI, VertexShader->PageTableSize, PageTableSize);
 					SetShaderValue(RHICmdList, ShaderRHI, VertexShader->FirstUpdate, FirstUpdate);
 					SetShaderValue(RHICmdList, ShaderRHI, VertexShader->NumUpdates, NumUpdates);
diff --git a/Engine/Source/Runtime/Renderer/Private/VolumetricFog.cpp b/Engine/Source/Runtime/Renderer/Private/VolumetricFog.cpp
--- a/Engine/Source/Runtime/Renderer/Private/VolumetricFog.cpp
+++ b/Engine/Source/Runtime/Renderer/Private/VolumetricFog.cpp
@@ -184,7 +184,7 @@ public:
 		const FViewInfo& View,
 		const FVolumetricFogIntegrationParameterData& IntegrationData)
 	{
-		FComputeShaderRHIParamRef ShaderRHI = GetComputeShader();
+		FRHIComputeShader* ShaderRHI = GetComputeShader();
 		VolumetricFogParameters.Set(RHICmdList, ShaderRHI, View, IntegrationData);
 	}
 
@@ -311,7 +311,7 @@ public:
 		const FProjectedShadowInfo* ShadowMap,
 		bool bDynamicallyShadowed)
 	{
-		const FPixelShaderRHIParamRef ShaderRHI = GetPixelShader();
+		FRHIPixelShader* ShaderRHI = GetPixelShader();
 
 		FGlobalShader::SetParameters<FViewUniformShaderParameters>(RHICmdList, ShaderRHI, View.ViewUniformBuffer);
 
@@ -635,7 +635,7 @@ public:
 		bool bUseDirectionalLightShadowing,
 		const FMatrix& DirectionalLightFunctionWorldToShadowValue)
 	{
-		FComputeShaderRHIParamRef ShaderRHI = GetComputeShader();
+		FRHIComputeShader* ShaderRHI = GetComputeShader();
 		FGlobalShader::SetParameters<FViewUniformShaderParameters>(RHICmdList, ShaderRHI, View.ViewUniformBuffer);
 
 		if (!LightScatteringHistoryTexture)
@@ -793,7 +793,7 @@ class FVolumetricFogFinalIntegrationCS : public FGlobalShader
 public:
 	void SetParameters(FRHICommandList& RHICmdList, const FViewInfo& View, const FVolumetricFogIntegrationParameterData& IntegrationData)
 	{
-		FComputeShaderRHIParamRef ShaderRHI = GetComputeShader();
+		FRHIComputeShader* ShaderRHI = GetComputeShader();
 		FGlobalShader::SetParameters<FViewUniformShaderParameters>(RHICmdList, ShaderRHI, View.ViewUniformBuffer);
 		VolumetricFogParameters.Set(RHICmdList, ShaderRHI, View, IntegrationData);
 	}
diff --git a/Engine/Source/Runtime/Renderer/Public/HdrCustomResolveShaders.h b/Engine/Source/Runtime/Renderer/Public/HdrCustomResolveShaders.h
--- a/Engine/Source/Runtime/Renderer/Public/HdrCustomResolveShaders.h
+++ b/Engine/Source/Runtime/Renderer/Public/HdrCustomResolveShaders.h
@@ -56,7 +56,7 @@ public:
 
 	void SetParameters(FRHICommandList& RHICmdList, FTextureRHIParamRef Texture2DMS)
 	{
-		FPixelShaderRHIParamRef PixelShaderRHI = GetPixelShader();
+		FRHIPixelShader* PixelShaderRHI = GetPixelShader();
 		SetTextureParameter(RHICmdList, PixelShaderRHI, Tex, Texture2DMS);
 	}
 
@@ -95,7 +95,7 @@ public:
 
 	void SetParameters(FRHICommandList& RHICmdList, FTextureRHIParamRef Texture2DMS)
 	{
-		FPixelShaderRHIParamRef PixelShaderRHI = GetPixelShader();
+		FRHIPixelShader* PixelShaderRHI = GetPixelShader();
 		SetTextureParameter(RHICmdList, PixelShaderRHI, Tex, Texture2DMS);
 	}
 
@@ -135,7 +135,7 @@ public:
 
 	void SetParameters(FRHICommandList& RHICmdList, FTextureRHIParamRef Texture2DMS)
 	{
-		FPixelShaderRHIParamRef PixelShaderRHI = GetPixelShader();
+		FRHIPixelShader* PixelShaderRHI = GetPixelShader();
 		SetTextureParameter(RHICmdList, PixelShaderRHI, Tex, Texture2DMS);
 	}
 
@@ -176,7 +176,7 @@ public:
 
 	void SetParameters(FRHICommandList& RHICmdList, FTextureRHIParamRef Texture2DMS, FTextureRHIParamRef FMaskTexure2D)
 	{
-		FPixelShaderRHIParamRef PixelShaderRHI = GetPixelShader();
+		FRHIPixelShader* PixelShaderRHI = GetPixelShader();
 		SetTextureParameter(RHICmdList, PixelShaderRHI, Tex, Texture2DMS);
 		SetTextureParameter(RHICmdList, PixelShaderRHI, FMaskTex, FMaskTexure2D);
 	}
@@ -220,7 +220,7 @@ public:
 
 	void SetParameters(FRHICommandList& RHICmdList, FTextureRHIParamRef Texture2DMS, FTextureRHIParamRef FMaskTexure2D)
 	{
-		FPixelShaderRHIParamRef PixelShaderRHI = GetPixelShader();
+		FRHIPixelShader* PixelShaderRHI = GetPixelShader();
 		SetTextureParameter(RHICmdList, PixelShaderRHI, Tex, Texture2DMS);
 		SetTextureParameter(RHICmdList, PixelShaderRHI, FMaskTex, FMaskTexure2D);
 	}
@@ -265,7 +265,7 @@ public:
 
 	void SetParameters(FRHICommandList& RHICmdList, FTextureRHIParamRef Texture2DMS, FTextureRHIParamRef FMaskTexure2D)
 	{
-		FPixelShaderRHIParamRef PixelShaderRHI = GetPixelShader();
+		FRHIPixelShader* PixelShaderRHI = GetPixelShader();
 		SetTextureParameter(RHICmdList, PixelShaderRHI, Tex, Texture2DMS);
 		SetTextureParameter(RHICmdList, PixelShaderRHI, FMaskTex, FMaskTexure2D);
 	}
diff --git a/Engine/Source/Runtime/Renderer/Public/MaterialShader.h b/Engine/Source/Runtime/Renderer/Public/MaterialShader.h
--- a/Engine/Source/Runtime/Renderer/Public/MaterialShader.h
+++ b/Engine/Source/Runtime/Renderer/Public/MaterialShader.h
@@ -146,19 +146,19 @@ public:
 	{ }
 
 	/** Sets pixel parameters that are material specific but not FMeshBatch specific. */
-	template< typename ShaderRHIParamRef >
+	template< typename TRHIShader >
 	void SetParametersInner(
 		FRHICommandList& RHICmdList,
-		const ShaderRHIParamRef ShaderRHI, 
+		TRHIShader* ShaderRHI,
 		const FMaterialRenderProxy* MaterialRenderProxy, 
 		const FMaterial& Material,
 		const FSceneView& View);
 
 	/** Sets pixel parameters that are material specific but not FMeshBatch specific. */
-	template< typename ShaderRHIParamRef >
+	template< typename TRHIShader >
 	void SetParameters(
 		FRHICommandList& RHICmdList,
-		const ShaderRHIParamRef ShaderRHI, 
+		TRHIShader* ShaderRHI,
 		const FMaterialRenderProxy* MaterialRenderProxy, 
 		const FMaterial& Material,
 		const FSceneView& View, 
diff --git a/Engine/Source/Runtime/Renderer/Public/MeshPassProcessor.h b/Engine/Source/Runtime/Renderer/Public/MeshPassProcessor.h
--- a/Engine/Source/Runtime/Renderer/Public/MeshPassProcessor.h
+++ b/Engine/Source/Runtime/Renderer/Public/MeshPassProcessor.h
@@ -216,14 +216,14 @@ private:
 
 struct FMeshProcessorShaders
 {
-	FMeshMaterialShader* VertexShader;
-	FMeshMaterialShader* HullShader;
-	FMeshMaterialShader* DomainShader;
-	FMeshMaterialShader* PixelShader;
-	FMeshMaterialShader* GeometryShader;
-	FMeshMaterialShader* ComputeShader;
+	mutable FMeshMaterialShader* VertexShader;
+	mutable FMeshMaterialShader* HullShader;
+	mutable FMeshMaterialShader* DomainShader;
+	mutable FMeshMaterialShader* PixelShader;
+	mutable FMeshMaterialShader* GeometryShader;
+	mutable FMeshMaterialShader* ComputeShader;
 #if RHI_RAYTRACING
-	FMeshMaterialShader* RayHitGroupShader;
+	mutable FMeshMaterialShader* RayHitGroupShader;
 #endif
 
 	FMeshMaterialShader* GetShader(EShaderFrequency Frequency) const
@@ -334,7 +334,7 @@ public:
 	/** Set shader bindings on the commandlist, filtered by state cache. */
 	void SetOnCommandList(FRHICommandList& RHICmdList, FBoundShaderStateInput Shaders, class FShaderBindingState* StateCacheShaderBindings) const;
 
-	void SetOnCommandListForCompute(FRHICommandList& RHICmdList, FComputeShaderRHIParamRef Shader) const;
+	void SetOnCommandListForCompute(FRHICommandList& RHICmdList, FRHIComputeShader* Shader) const;
 
 #if RHI_RAYTRACING
 	void SetRayTracingShaderBindingsForHitGroup(FRHICommandList& RHICmdList, FRHIRayTracingScene* Scene, uint32 InstanceIndex, uint32 SegmentIndex, FRayTracingPipelineState* Pipeline, uint32 HitGroupIndex, uint32 ShaderSlot) const;
diff --git a/Engine/Source/Runtime/Renderer/Public/PostProcessParameters.h b/Engine/Source/Runtime/Renderer/Public/PostProcessParameters.h
--- a/Engine/Source/Runtime/Renderer/Public/PostProcessParameters.h
+++ b/Engine/Source/Runtime/Renderer/Public/PostProcessParameters.h
@@ -75,14 +75,14 @@ struct FPostProcessPassParameters
 
 	/** Set the pixel shader parameter values. */
 	template <typename TRHICmdList>
-	void SetPS(TRHICmdList& RHICmdList, const FPixelShaderRHIParamRef& ShaderRHI, const FRenderingCompositePassContext& Context, FRHISamplerState* Filter = TStaticSamplerState<>::GetRHI(), EFallbackColor FallbackColor = eFC_0000, FRHISamplerState** FilterOverrideArray = 0);
+	void SetPS(TRHICmdList& RHICmdList, FRHIPixelShader* ShaderRHI, const FRenderingCompositePassContext& Context, FRHISamplerState* Filter = TStaticSamplerState<>::GetRHI(), EFallbackColor FallbackColor = eFC_0000, FRHISamplerState** FilterOverrideArray = 0);
 
 	/** Set the compute shader parameter values. */
 	template< typename TRHICmdList >
-	void SetCS(const FComputeShaderRHIParamRef& ShaderRHI, const FRenderingCompositePassContext& Context, TRHICmdList& RHICmdList, FRHISamplerState* Filter = TStaticSamplerState<>::GetRHI(), EFallbackColor FallbackColor = eFC_0000, FRHISamplerState** FilterOverrideArray = 0);
+	void SetCS(FRHIComputeShader* ShaderRHI, const FRenderingCompositePassContext& Context, TRHICmdList& RHICmdList, FRHISamplerState* Filter = TStaticSamplerState<>::GetRHI(), EFallbackColor FallbackColor = eFC_0000, FRHISamplerState** FilterOverrideArray = 0);
 
 	/** Set the vertex shader parameter values. */
-	void SetVS(const FVertexShaderRHIParamRef& ShaderRHI, const FRenderingCompositePassContext& Context, FRHISamplerState* Filter = TStaticSamplerState<>::GetRHI(), EFallbackColor FallbackColor = eFC_0000, FRHISamplerState** FilterOverrideArray = 0);
+	void SetVS(FRHIVertexShader* ShaderRHI, const FRenderingCompositePassContext& Context, FRHISamplerState* Filter = TStaticSamplerState<>::GetRHI(), EFallbackColor FallbackColor = eFC_0000, FRHISamplerState** FilterOverrideArray = 0);
 
 	/** Serializer. */
 	friend FArchive& operator<<(FArchive& Ar,FPostProcessPassParameters& P);
@@ -102,10 +102,10 @@ private:
 public:
 	// @param Filter can be 0 if FilterOverrideArray is used
 	// @param FilterOverrideArray can be 0 if Filter is used
-	template< typename TShaderRHIParamRef, typename TRHICmdList >
+	template< typename TRHIShader, typename TRHICmdList >
 	void Set(
 		TRHICmdList& RHICmdList,
-		const TShaderRHIParamRef& ShaderRHI,
+		TRHIShader* ShaderRHI,
 		const FRenderingCompositePassContext& Context,
 		FRHISamplerState* Filter,
 		EFallbackColor FallbackColor,
diff --git a/Engine/Source/Runtime/SlateRHIRenderer/Private/SlateMaterialShader.cpp b/Engine/Source/Runtime/SlateRHIRenderer/Private/SlateMaterialShader.cpp
--- a/Engine/Source/Runtime/SlateRHIRenderer/Private/SlateMaterialShader.cpp
+++ b/Engine/Source/Runtime/SlateRHIRenderer/Private/SlateMaterialShader.cpp
@@ -35,9 +35,9 @@ void FSlateMaterialShaderVS::SetViewProjection(FRHICommandList& RHICmdList, cons
 
 void FSlateMaterialShaderVS::SetMaterialShaderParameters(FRHICommandList& RHICmdList, const FSceneView& View, const FMaterialRenderProxy* MaterialRenderProxy, const FMaterial* Material)
 {
-	const FVertexShaderRHIParamRef ShaderRHI = GetVertexShader();
+	FRHIVertexShader* ShaderRHI = GetVertexShader();
 
-	FMaterialShader::SetParameters<FVertexShaderRHIParamRef>(RHICmdList, ShaderRHI, MaterialRenderProxy, *Material, View, View.ViewUniformBuffer, ESceneTextureSetupMode::None);
+	FMaterialShader::SetParameters<FRHIVertexShader>(RHICmdList, ShaderRHI, MaterialRenderProxy, *Material, View, View.ViewUniformBuffer, ESceneTextureSetupMode::None);
 }
 
 void FSlateMaterialShaderVS::SetVerticalAxisMultiplier(FRHICommandList& RHICmdList, float InMultiplier )
@@ -119,12 +119,12 @@ void FSlateMaterialShaderPS::SetBlendState(FGraphicsPipelineStateInitializer& Gr
 
 void FSlateMaterialShaderPS::SetParameters(FRHICommandList& RHICmdList, const FSceneView& View, const FMaterialRenderProxy* MaterialRenderProxy, const FMaterial* Material, const FVector4& InShaderParams)
 {
-	const FPixelShaderRHIParamRef ShaderRHI = GetPixelShader();
+	FRHIPixelShader* ShaderRHI = GetPixelShader();
 
 	SetShaderValue( RHICmdList, ShaderRHI, ShaderParams, InShaderParams );
 
 	const ESceneTextureSetupMode SceneTextures = ESceneTextureSetupMode::SceneDepth | ESceneTextureSetupMode::SSAO | ESceneTextureSetupMode::CustomDepth; // TODO - SSAO valid here?
-	FMaterialShader::SetParameters<FPixelShaderRHIParamRef>(RHICmdList, ShaderRHI, MaterialRenderProxy, *Material, View, View.ViewUniformBuffer, SceneTextures);
+	FMaterialShader::SetParameters<FRHIPixelShader>(RHICmdList, ShaderRHI, MaterialRenderProxy, *Material, View, View.ViewUniformBuffer, SceneTextures);
 }
 
 void FSlateMaterialShaderPS::SetAdditionalTexture( FRHICommandList& RHICmdList, const FTextureRHIParamRef InTexture, const FSamplerStateRHIRef SamplerState )
diff --git a/Engine/Source/Runtime/SlateRHIRenderer/Private/SlateShaders.h b/Engine/Source/Runtime/SlateRHIRenderer/Private/SlateShaders.h
--- a/Engine/Source/Runtime/SlateRHIRenderer/Private/SlateShaders.h
+++ b/Engine/Source/Runtime/SlateRHIRenderer/Private/SlateShaders.h
@@ -338,7 +338,7 @@ public:
 	void SetWeightsAndOffsets(FRHICommandList& RHICmdList, const TArray<FVector4>& InWeightsAndOffsets, int32 NumSamples )
 	{
 		check(InWeightsAndOffsets.Num() <= MAX_BLUR_SAMPLES);
-		SetShaderValueArray<FPixelShaderRHIParamRef, FVector4>(RHICmdList, GetPixelShader(), WeightAndOffsets, InWeightsAndOffsets.GetData(), InWeightsAndOffsets.Num() );
+		SetShaderValueArray<FRHIPixelShader*, FVector4>(RHICmdList, GetPixelShader(), WeightAndOffsets, InWeightsAndOffsets.GetData(), InWeightsAndOffsets.Num() );
 		SetShaderValue(RHICmdList, GetPixelShader(), SampleCount, NumSamples);
 	}
 
diff --git a/Engine/Source/Runtime/UtilityShaders/Private/ClearQuad.cpp b/Engine/Source/Runtime/UtilityShaders/Private/ClearQuad.cpp
--- a/Engine/Source/Runtime/UtilityShaders/Private/ClearQuad.cpp
+++ b/Engine/Source/Runtime/UtilityShaders/Private/ClearQuad.cpp
@@ -131,7 +131,7 @@ static void ClearUAVShader(FRHICommandList& RHICmdList, FRHIUnorderedAccessView*
 		TEXT("Buffer size is not a multiple of DWORDs. Up to 3 bytes after buffer end will also be cleared"));
 
 	TShaderMapRef<FClearBufferReplacementCS> ComputeShader(GetGlobalShaderMap(GMaxRHIFeatureLevel));
-	FComputeShaderRHIParamRef ShaderRHI = ComputeShader->GetComputeShader();
+	FRHIComputeShader* ShaderRHI = ComputeShader->GetComputeShader();
 	
 	uint32 NumDWordsToClear = (SizeInBytes + 3) / 4;
 	uint32 NumThreadGroupsX = (NumDWordsToClear + 63) / 64;
@@ -190,7 +190,7 @@ inline void ClearUAV_T(FRHICommandList& RHICmdList, FRHITexture* Texture, FRHIUn
 	if (auto Texture2d = Texture->GetTexture2D())
 	{
 		TShaderMapRef< FClearTexture2DReplacementCS<T> > ComputeShader(GetGlobalShaderMap(GMaxRHIFeatureLevel));
-		FComputeShaderRHIParamRef ShaderRHI = ComputeShader->GetComputeShader();
+		FRHIComputeShader* ShaderRHI = ComputeShader->GetComputeShader();
 		RHICmdList.SetComputeShader(ShaderRHI);
 		ComputeShader->SetParameters(RHICmdList, TextureUAV, ClearValues);
 		uint32 x = (Texture2d->GetSizeX() + 7) / 8;
@@ -201,7 +201,7 @@ inline void ClearUAV_T(FRHICommandList& RHICmdList, FRHITexture* Texture, FRHIUn
 	else if (auto Texture2dArray = Texture->GetTexture2DArray())
 	{
 		TShaderMapRef< FClearTexture2DArrayReplacementCS<T> > ComputeShader(GetGlobalShaderMap(GMaxRHIFeatureLevel));
-		FComputeShaderRHIParamRef ShaderRHI = ComputeShader->GetComputeShader();
+		FRHIComputeShader* ShaderRHI = ComputeShader->GetComputeShader();
 		RHICmdList.SetComputeShader(ShaderRHI);
 		ComputeShader->SetParameters(RHICmdList, TextureUAV, ClearValues);
 		uint32 x = (Texture2dArray->GetSizeX() + 7) / 8;
@@ -213,7 +213,7 @@ inline void ClearUAV_T(FRHICommandList& RHICmdList, FRHITexture* Texture, FRHIUn
 	else if (auto TextureCube = Texture->GetTextureCube())
 	{
 		TShaderMapRef< FClearTexture2DArrayReplacementCS<T> > ComputeShader(GetGlobalShaderMap(GMaxRHIFeatureLevel));
-		FComputeShaderRHIParamRef ShaderRHI = ComputeShader->GetComputeShader();
+		FRHIComputeShader* ShaderRHI = ComputeShader->GetComputeShader();
 		RHICmdList.SetComputeShader(ShaderRHI);
 		ComputeShader->SetParameters(RHICmdList, TextureUAV, ClearValues);
 		uint32 x = (TextureCube->GetSize() + 7) / 8;
@@ -224,7 +224,7 @@ inline void ClearUAV_T(FRHICommandList& RHICmdList, FRHITexture* Texture, FRHIUn
 	else if (auto Texture3d = Texture->GetTexture3D())
 	{
 		TShaderMapRef< FClearVolumeReplacementCS<T> > ComputeShader(GetGlobalShaderMap(GMaxRHIFeatureLevel));
-		FComputeShaderRHIParamRef ShaderRHI = ComputeShader->GetComputeShader();
+		FRHIComputeShader* ShaderRHI = ComputeShader->GetComputeShader();
 		RHICmdList.SetComputeShader(ShaderRHI);
 		ComputeShader->SetParameters(RHICmdList, TextureUAV, ClearValues);
 		uint32 x = (Texture3d->GetSizeX() + 3) / 4;
@@ -242,7 +242,7 @@ inline void ClearUAV_T(FRHICommandList& RHICmdList, FRHITexture* Texture, FRHIUn
 void ClearTexture2DUAV(FRHICommandList& RHICmdList, FRHIUnorderedAccessView* UAV, int32 Width, int32 Height, const FLinearColor& ClearColor)
 {
 	TShaderMapRef< FClearTexture2DReplacementCS<float> > ComputeShader(GetGlobalShaderMap(GMaxRHIFeatureLevel));
-	FComputeShaderRHIParamRef ShaderRHI = ComputeShader->GetComputeShader();
+	FRHIComputeShader* ShaderRHI = ComputeShader->GetComputeShader();
 	RHICmdList.SetComputeShader(ShaderRHI);
 	ComputeShader->SetParameters(RHICmdList, UAV, reinterpret_cast<const float(&)[4]>(ClearColor));
 	uint32 x = (Width + 7) / 8;
diff --git a/Engine/Source/Runtime/UtilityShaders/Private/ClearReplacementShaders.cpp b/Engine/Source/Runtime/UtilityShaders/Private/ClearReplacementShaders.cpp
--- a/Engine/Source/Runtime/UtilityShaders/Private/ClearReplacementShaders.cpp
+++ b/Engine/Source/Runtime/UtilityShaders/Private/ClearReplacementShaders.cpp
@@ -8,7 +8,7 @@
 template< typename T >
 void FClearTexture2DReplacementCS<T>::SetParameters( FRHICommandList& RHICmdList, FRHIUnorderedAccessView* TextureRW, const T(&Values)[4] )
 {
-	FComputeShaderRHIParamRef ComputeShaderRHI = GetComputeShader();
+	FRHIComputeShader* ComputeShaderRHI = GetComputeShader();
 	SetShaderValue(RHICmdList, ComputeShaderRHI, ClearColor, Values);
 
 	RHICmdList.TransitionResource(EResourceTransitionAccess::ERWBarrier, EResourceTransitionPipeline::EGfxToCompute, TextureRW);
@@ -19,14 +19,14 @@ template< typename T >
 void FClearTexture2DReplacementCS<T>::FinalizeParameters(FRHICommandList& RHICmdList, FRHIUnorderedAccessView* TextureRW)
 {
 	RHICmdList.TransitionResource(EResourceTransitionAccess::ERWBarrier, EResourceTransitionPipeline::EComputeToCompute, TextureRW);
-	FComputeShaderRHIParamRef ComputeShaderRHI = GetComputeShader();
+	FRHIComputeShader* ComputeShaderRHI = GetComputeShader();
 	SetUAVParameter(RHICmdList, ComputeShaderRHI, ClearTextureRW, FUnorderedAccessViewRHIRef());
 }
 
 template< typename T >
 void FClearTexture2DArrayReplacementCS<T>::SetParameters(FRHICommandList& RHICmdList, FRHIUnorderedAccessView* TextureRW, const T(&Values)[4])
 {
-	FComputeShaderRHIParamRef ComputeShaderRHI = GetComputeShader();
+	FRHIComputeShader* ComputeShaderRHI = GetComputeShader();
 	SetShaderValue(RHICmdList, ComputeShaderRHI, ClearColor, Values);
 
 	RHICmdList.TransitionResource(EResourceTransitionAccess::ERWBarrier, EResourceTransitionPipeline::EGfxToCompute, TextureRW);
@@ -37,14 +37,14 @@ template< typename T >
 void FClearTexture2DArrayReplacementCS<T>::FinalizeParameters(FRHICommandList& RHICmdList, FRHIUnorderedAccessView* TextureRW)
 {
 	RHICmdList.TransitionResource(EResourceTransitionAccess::ERWBarrier, EResourceTransitionPipeline::EComputeToCompute, TextureRW);
-	FComputeShaderRHIParamRef ComputeShaderRHI = GetComputeShader();
+	FRHIComputeShader* ComputeShaderRHI = GetComputeShader();
 	SetUAVParameter(RHICmdList, ComputeShaderRHI, ClearTextureArrayRW, FUnorderedAccessViewRHIRef());
 }
 
 template< typename T >
 void FClearVolumeReplacementCS<T>::SetParameters(FRHICommandList& RHICmdList, FRHIUnorderedAccessView* TextureRW, const T(&Values)[4])
 {
-	FComputeShaderRHIParamRef ComputeShaderRHI = GetComputeShader();
+	FRHIComputeShader* ComputeShaderRHI = GetComputeShader();
 	SetShaderValue(RHICmdList, ComputeShaderRHI, ClearColor, Values);
 
 	RHICmdList.TransitionResource(EResourceTransitionAccess::ERWBarrier, EResourceTransitionPipeline::EGfxToCompute, TextureRW);
@@ -55,13 +55,13 @@ template< typename T >
 void FClearVolumeReplacementCS<T>::FinalizeParameters(FRHICommandList& RHICmdList, FRHIUnorderedAccessView* TextureRW)
 {
 	RHICmdList.TransitionResource(EResourceTransitionAccess::ERWBarrier, EResourceTransitionPipeline::EComputeToCompute, TextureRW);
-	FComputeShaderRHIParamRef ComputeShaderRHI = GetComputeShader();
+	FRHIComputeShader* ComputeShaderRHI = GetComputeShader();
 	SetUAVParameter(RHICmdList, ComputeShaderRHI, ClearVolumeRW, FUnorderedAccessViewRHIRef());
 }
 
 void FClearTexture2DReplacementScissorCS::SetParameters(FRHICommandList& RHICmdList, FRHIUnorderedAccessView* TextureRW, FLinearColor InClearColor, const FVector4& InTargetBounds)
 {
-	FComputeShaderRHIParamRef ComputeShaderRHI = GetComputeShader();
+	FRHIComputeShader* ComputeShaderRHI = GetComputeShader();
 	SetShaderValue(RHICmdList, ComputeShaderRHI, ClearColor, InClearColor);
 	SetShaderValue(RHICmdList, ComputeShaderRHI, TargetBounds, InTargetBounds);
 	RHICmdList.TransitionResource(EResourceTransitionAccess::ERWBarrier, EResourceTransitionPipeline::EGfxToCompute, TextureRW);
@@ -71,13 +71,13 @@ void FClearTexture2DReplacementScissorCS::SetParameters(FRHICommandList& RHICmdL
 void FClearTexture2DReplacementScissorCS::FinalizeParameters(FRHICommandList& RHICmdList, FRHIUnorderedAccessView* TextureRW)
 {
 	RHICmdList.TransitionResource(EResourceTransitionAccess::ERWBarrier, EResourceTransitionPipeline::EComputeToCompute, TextureRW);
-	FComputeShaderRHIParamRef ComputeShaderRHI = GetComputeShader();
+	FRHIComputeShader* ComputeShaderRHI = GetComputeShader();
 	SetUAVParameter(RHICmdList, ComputeShaderRHI, ClearTextureRW, FUnorderedAccessViewRHIRef());
 }
 
 void FClearBufferReplacementCS::SetParameters(FRHICommandList& RHICmdList, FRHIUnorderedAccessView* BufferRW, uint32 NumDWordsToClear,uint32 ClearValue)
 {
-	FComputeShaderRHIParamRef ComputeShaderRHI = GetComputeShader();
+	FRHIComputeShader* ComputeShaderRHI = GetComputeShader();
 	SetShaderValue(RHICmdList, ComputeShaderRHI, ClearBufferCSParams, FUintVector4(ClearValue, NumDWordsToClear, 0, 0));
 	RHICmdList.TransitionResource(EResourceTransitionAccess::ERWBarrier, EResourceTransitionPipeline::EGfxToCompute, BufferRW);
 	RHICmdList.SetUAVParameter(ComputeShaderRHI, ClearBufferRW.GetBaseIndex(), BufferRW);
@@ -86,7 +86,7 @@ void FClearBufferReplacementCS::SetParameters(FRHICommandList& RHICmdList, FRHIU
 void FClearBufferReplacementCS::FinalizeParameters(FRHICommandList& RHICmdList, FRHIUnorderedAccessView* BufferRW)
 {
 	RHICmdList.TransitionResource(EResourceTransitionAccess::ERWBarrier, EResourceTransitionPipeline::EComputeToCompute, BufferRW);
-	FComputeShaderRHIParamRef ComputeShaderRHI = GetComputeShader();
+	FRHIComputeShader* ComputeShaderRHI = GetComputeShader();
 	SetUAVParameter(RHICmdList, ComputeShaderRHI, ClearBufferRW, FUnorderedAccessViewRHIRef());
 }
 
diff --git a/Engine/Source/Runtime/UtilityShaders/Private/RGBAToYUV420Shader.cpp b/Engine/Source/Runtime/UtilityShaders/Private/RGBAToYUV420Shader.cpp
--- a/Engine/Source/Runtime/UtilityShaders/Private/RGBAToYUV420Shader.cpp
+++ b/Engine/Source/Runtime/UtilityShaders/Private/RGBAToYUV420Shader.cpp
@@ -13,7 +13,7 @@ IMPLEMENT_SHADER_TYPE(, FRGBAToYUV420CS, TEXT("/Engine/Private/RGBAToYUV420.usf"
 
 void FRGBAToYUV420CS::SetParameters(FRHICommandList& RHICmdList, TRefCountPtr<FRHITexture2D> InSrcTexture, FRHIUnorderedAccessView* InOutUAV, float InTargetHeight, float InScaleFactorX, float InScaleFactorY, float InTextureYOffset)
 {
-	FComputeShaderRHIParamRef ComputeShaderRHI = GetComputeShader();
+	FRHIComputeShader* ComputeShaderRHI = GetComputeShader();
 	SetShaderValue(RHICmdList, ComputeShaderRHI, TargetHeight, InTargetHeight);
 	SetShaderValue(RHICmdList, ComputeShaderRHI, ScaleFactorX, InScaleFactorX);
 	SetShaderValue(RHICmdList, ComputeShaderRHI, ScaleFactorY, InScaleFactorY);
@@ -27,7 +27,7 @@ void FRGBAToYUV420CS::SetParameters(FRHICommandList& RHICmdList, TRefCountPtr<FR
 */
 void FRGBAToYUV420CS::UnbindBuffers(FRHICommandList& RHICmdList)
 {
-	FComputeShaderRHIParamRef ComputeShaderRHI = GetComputeShader();
+	FRHIComputeShader* ComputeShaderRHI = GetComputeShader();
 	RHICmdList.SetUAVParameter(ComputeShaderRHI, OutTextureRW.GetBaseIndex(), nullptr);
 }
 
