DR - Fix warnings
#rnx

diff --git a/Engine/Source/Runtime/Renderer/Private/RayTracing/RayTracingAmbientOcclusion.cpp b/Engine/Source/Runtime/Renderer/Private/RayTracing/RayTracingAmbientOcclusion.cpp
--- a/Engine/Source/Runtime/Renderer/Private/RayTracing/RayTracingAmbientOcclusion.cpp
+++ b/Engine/Source/Runtime/Renderer/Private/RayTracing/RayTracingAmbientOcclusion.cpp
@@ -231,8 +231,8 @@ void FDeferredShadingSceneRenderer::RenderRayTracingAmbientOcclusion(
 	int32 DenoiserMode = CVarUseAODenoiser.GetValueOnRenderThread();
 	if (DenoiserMode != 0)
 	{
-		FSceneTextureParameters SceneTextures;
-		SetupSceneTextureParameters(GraphBuilder, &SceneTextures);
+		FSceneTextureParameters SceneTextureParams;
+		SetupSceneTextureParameters(GraphBuilder, &SceneTextureParams);
 
 		const IScreenSpaceDenoiser* DefaultDenoiser = IScreenSpaceDenoiser::GetDefaultDenoiser();
 		const IScreenSpaceDenoiser* DenoiserToUse = DenoiserMode == 1 ? DefaultDenoiser : GScreenSpaceDenoiser;
@@ -253,7 +253,7 @@ void FDeferredShadingSceneRenderer::RenderRayTracingAmbientOcclusion(
 				GraphBuilder,
 				View,
 				&View.PrevViewInfo,
-				SceneTextures,
+				SceneTextureParams,
 				DenoiserInputs,
 				RayTracingConfig);
 
diff --git a/Engine/Source/Runtime/Renderer/Private/RayTracing/RayTracingGlobalIllumination.cpp b/Engine/Source/Runtime/Renderer/Private/RayTracing/RayTracingGlobalIllumination.cpp
--- a/Engine/Source/Runtime/Renderer/Private/RayTracing/RayTracingGlobalIllumination.cpp
+++ b/Engine/Source/Runtime/Renderer/Private/RayTracing/RayTracingGlobalIllumination.cpp
@@ -512,8 +512,8 @@ void FDeferredShadingSceneRenderer::RenderRayTracingGlobalIllumination(
 
 	if (GRayTracingGlobalIlluminationDenoiser != 0)
 	{
-		FSceneTextureParameters SceneTextures;
-		SetupSceneTextureParameters(GraphBuilder, &SceneTextures);
+		FSceneTextureParameters SceneTextureParams;
+		SetupSceneTextureParameters(GraphBuilder, &SceneTextureParams);
 
 		const IScreenSpaceDenoiser* DefaultDenoiser = IScreenSpaceDenoiser::GetDefaultDenoiser();
 		const IScreenSpaceDenoiser* DenoiserToUse = GRayTracingGlobalIlluminationDenoiser == 1 ? DefaultDenoiser : GScreenSpaceDenoiser;
@@ -532,7 +532,7 @@ void FDeferredShadingSceneRenderer::RenderRayTracingGlobalIllumination(
 				GraphBuilder,
 				View,
 				&View.PrevViewInfo,
-				SceneTextures,
+				SceneTextureParams,
 				DenoiserInputs,
 				RayTracingConfig);
 
