Enable virtual texture support on Switch

diff --git a/Engine/Shaders/Private/BasePassPixelShader.usf b/Engine/Shaders/Private/BasePassPixelShader.usf
--- a/Engine/Shaders/Private/BasePassPixelShader.usf
+++ b/Engine/Shaders/Private/BasePassPixelShader.usf
@@ -38,8 +38,12 @@
 #include "ReflectionEnvironmentShared.ush"
 
 #if LIGHTMAP_VT_ENABLED || NUM_VIRTUALTEXTURE_SAMPLES
+#if COMPILER_HLSLCC
+	RWTexture2D<uint> VTFeedbackBuffer;
+#else
 	RWTexture2D<uint> VTFeedbackBuffer : register(VT_FEEDBACK_REGISTER);
 #endif
+#endif
 
 float NormalCurvatureToRoughness(float3 WorldNormal)
 {
diff --git a/Engine/Source/Runtime/Engine/Private/VT/VirtualTextureDataBuilder.cpp b/Engine/Source/Runtime/Engine/Private/VT/VirtualTextureDataBuilder.cpp
--- a/Engine/Source/Runtime/Engine/Private/VT/VirtualTextureDataBuilder.cpp
+++ b/Engine/Source/Runtime/Engine/Private/VT/VirtualTextureDataBuilder.cpp
@@ -506,6 +506,10 @@ void FVirtualTextureDataBuilder::BuildTiles(const TArray<FVTSourceTileEntry>& Ti
 	{
 		TextureFormatName = *BaseTextureFormatName.Replace(TEXT("XBOXONE_"), TEXT(""));
 	}
+	else if (BaseTextureFormatName.StartsWith("SWITCH_"))
+	{
+		TextureFormatName = *BaseTextureFormatName.Replace(TEXT("SWITCH_"), TEXT(""));
+	}
 
 	// We handle AutoDXT specially here since otherwise the texture format compressor would choose a DXT format for every tile
 	// individually. Causing tiles in the same VT to use different formats which we don't allow.
diff --git a/Engine/Source/Runtime/Renderer/Private/VT/VirtualTexturePhysicalSpace.cpp b/Engine/Source/Runtime/Renderer/Private/VT/VirtualTexturePhysicalSpace.cpp
--- a/Engine/Source/Runtime/Renderer/Private/VT/VirtualTexturePhysicalSpace.cpp
+++ b/Engine/Source/Runtime/Renderer/Private/VT/VirtualTexturePhysicalSpace.cpp
@@ -79,6 +79,8 @@ void FVirtualTexturePhysicalSpace::InitRHI()
 void FVirtualTexturePhysicalSpace::ReleaseRHI()
 {
 	GRenderTargetPool.FreeUnusedResource(PooledRenderTarget);
+	TextureView->Release();
+	TextureView->Release();
 }
 
 uint32 FVirtualTexturePhysicalSpace::GetSizeInBytes() const
diff --git a/Engine/Source/Runtime/Renderer/Private/VT/VirtualTextureSpace.cpp b/Engine/Source/Runtime/Renderer/Private/VT/VirtualTextureSpace.cpp
--- a/Engine/Source/Runtime/Renderer/Private/VT/VirtualTextureSpace.cpp
+++ b/Engine/Source/Runtime/Renderer/Private/VT/VirtualTextureSpace.cpp
@@ -184,7 +184,7 @@ public:
 
 	static bool ShouldCompilePermutation(const FGlobalShaderPermutationParameters& Parameters)
 	{
-		return IsFeatureLevelSupported(Parameters.Platform, ERHIFeatureLevel::SM5) && !IsHlslccShaderPlatform(Parameters.Platform);
+		return IsFeatureLevelSupported(Parameters.Platform, ERHIFeatureLevel::SM5);
 	}
 
 	virtual bool Serialize( FArchive& Ar ) override
@@ -211,7 +211,7 @@ public:
 	
 	static bool ShouldCompilePermutation(const FGlobalShaderPermutationParameters& Parameters)
 	{
-		return IsFeatureLevelSupported( Parameters.Platform, ERHIFeatureLevel::SM5 ) && !IsHlslccShaderPlatform(Parameters.Platform);
+		return IsFeatureLevelSupported(Parameters.Platform, ERHIFeatureLevel::SM5);
 	}
 };
 
