Fix double release in last submit!

diff --git a/Engine/Source/Runtime/Renderer/Private/VT/VirtualTexturePhysicalSpace.cpp b/Engine/Source/Runtime/Renderer/Private/VT/VirtualTexturePhysicalSpace.cpp
--- a/Engine/Source/Runtime/Renderer/Private/VT/VirtualTexturePhysicalSpace.cpp
+++ b/Engine/Source/Runtime/Renderer/Private/VT/VirtualTexturePhysicalSpace.cpp
@@ -79,8 +79,8 @@ void FVirtualTexturePhysicalSpace::InitRHI()
 void FVirtualTexturePhysicalSpace::ReleaseRHI()
 {
 	GRenderTargetPool.FreeUnusedResource(PooledRenderTarget);
-	TextureView->Release();
-	TextureView->Release();
+	TextureView.SafeRelease();
+	TextureSRGBView.SafeRelease();
 }
 
 uint32 FVirtualTexturePhysicalSpace::GetSizeInBytes() const
