Don't render runtime virtual texture until the scene is ready

diff --git a/Engine/Source/Runtime/Renderer/Private/VT/RuntimeVirtualTextureProducer.cpp b/Engine/Source/Runtime/Renderer/Private/VT/RuntimeVirtualTextureProducer.cpp
--- a/Engine/Source/Runtime/Renderer/Private/VT/RuntimeVirtualTextureProducer.cpp
+++ b/Engine/Source/Runtime/Renderer/Private/VT/RuntimeVirtualTextureProducer.cpp
@@ -20,7 +20,7 @@ bool FRuntimeVirtualTextureFinalizer::IsReady()
 {
 	//todo[vt]: 
 	// Test if we have everything we need to render (mips loaded, shaders loaded etc).
-	return Scene != nullptr && Scene->GetRenderScene() != nullptr;
+	return Scene != nullptr && Scene->GetRenderScene() != nullptr && Scene->GetRenderScene()->GetFrameNumber() > 0;
 }
 
 void FRuntimeVirtualTextureFinalizer::AddTile(FTileEntry& Tile)
