6673299 re-enabled. Pixel shaders used for composition purposes in ray tracing passes updated

diff --git a/Engine/Source/Runtime/RenderCore/Private/ShaderParameterStruct.cpp b/Engine/Source/Runtime/RenderCore/Private/ShaderParameterStruct.cpp
--- a/Engine/Source/Runtime/RenderCore/Private/ShaderParameterStruct.cpp
+++ b/Engine/Source/Runtime/RenderCore/Private/ShaderParameterStruct.cpp
@@ -236,7 +236,7 @@ void FShaderParameterBindings::BindForLegacyShaderParameters(const FShader* Shad
 	case SF_Compute:
 		break;
 	default:
-		checkf(0, TEXT("Invalid shader frequency for this shader binding technic."));
+		checkf(0, TEXT("Invalid shader frequency for this shader binding technique."));
 		break;
 	}
 
@@ -287,7 +287,7 @@ void FShaderParameterBindings::BindForRootShaderParameters(const FShader* Shader
 	case SF_RayCallable:
 		break;
 	default:
-		//checkf(0, TEXT("Invalid shader frequency for this shader binding technic."));
+		checkf(0, TEXT("Invalid shader frequency for this shader binding technic."));
 		break;
 	}
 
diff --git a/Engine/Source/Runtime/Renderer/Private/RayTracing/RayTracingGlobalIllumination.cpp b/Engine/Source/Runtime/Renderer/Private/RayTracing/RayTracingGlobalIllumination.cpp
--- a/Engine/Source/Runtime/Renderer/Private/RayTracing/RayTracingGlobalIllumination.cpp
+++ b/Engine/Source/Runtime/Renderer/Private/RayTracing/RayTracingGlobalIllumination.cpp
@@ -247,7 +247,7 @@ class FGlobalIlluminationRGS : public FGlobalShader
 class FRayTracingGlobalIlluminationCompositePS : public FGlobalShader
 {
 	DECLARE_GLOBAL_SHADER(FRayTracingGlobalIlluminationCompositePS)
-	SHADER_USE_ROOT_PARAMETER_STRUCT(FRayTracingGlobalIlluminationCompositePS, FGlobalShader)
+	SHADER_USE_PARAMETER_STRUCT(FRayTracingGlobalIlluminationCompositePS, FGlobalShader)
 
 	static bool ShouldCompilePermutation(const FGlobalShaderPermutationParameters& Parameters)
 	{
@@ -266,7 +266,7 @@ class FRayTracingGlobalIlluminationCompositePS : public FGlobalShader
 class FRayTracingGlobalIlluminationSceneColorCompositePS : public FGlobalShader
 {
 	DECLARE_GLOBAL_SHADER(FRayTracingGlobalIlluminationSceneColorCompositePS)
-	SHADER_USE_ROOT_PARAMETER_STRUCT(FRayTracingGlobalIlluminationSceneColorCompositePS, FGlobalShader)
+	SHADER_USE_PARAMETER_STRUCT(FRayTracingGlobalIlluminationSceneColorCompositePS, FGlobalShader)
 
 	static bool ShouldCompilePermutation(const FGlobalShaderPermutationParameters& Parameters)
 	{
diff --git a/Engine/Source/Runtime/Renderer/Private/RayTracing/RaytracingSkylight.cpp b/Engine/Source/Runtime/Renderer/Private/RayTracing/RaytracingSkylight.cpp
--- a/Engine/Source/Runtime/Renderer/Private/RayTracing/RaytracingSkylight.cpp
+++ b/Engine/Source/Runtime/Renderer/Private/RayTracing/RaytracingSkylight.cpp
@@ -939,9 +939,9 @@ void FDeferredShadingSceneRenderer::RenderRayTracingSkyLight(
 class FCompositeSkyLightPS : public FGlobalShader
 {
 	DECLARE_GLOBAL_SHADER(FCompositeSkyLightPS)
-	SHADER_USE_ROOT_PARAMETER_STRUCT(FCompositeSkyLightPS, FGlobalShader)
+	SHADER_USE_PARAMETER_STRUCT(FCompositeSkyLightPS, FGlobalShader)
 
-		static bool ShouldCompilePermutation(const FGlobalShaderPermutationParameters& Parameters)
+	static bool ShouldCompilePermutation(const FGlobalShaderPermutationParameters& Parameters)
 	{
 		return ShouldCompileRayTracingShadersForProject(Parameters.Platform);
 	}
