Test scene is ready before updating any runtime VT
(Also a couple of unrelated data improvements for test map)

diff --git a/Engine/Source/Runtime/Renderer/Private/VT/RuntimeVirtualTextureProducer.cpp b/Engine/Source/Runtime/Renderer/Private/VT/RuntimeVirtualTextureProducer.cpp
--- a/Engine/Source/Runtime/Renderer/Private/VT/RuntimeVirtualTextureProducer.cpp
+++ b/Engine/Source/Runtime/Renderer/Private/VT/RuntimeVirtualTextureProducer.cpp
@@ -19,8 +19,10 @@ FRuntimeVirtualTextureFinalizer::FRuntimeVirtualTextureFinalizer(FVTProducerDesc
 bool FRuntimeVirtualTextureFinalizer::IsReady()
 {
 	//todo[vt]: 
-	// Test if we have everything we need to render (mips loaded, shaders loaded etc).
-	return Scene != nullptr && Scene->GetRenderScene() != nullptr && Scene->GetRenderScene()->GetFrameNumber() > 0;
+	// Test if we have everything we need to render (mips loaded etc).
+
+	// Test scene is loaded and has been updated once by main rendering passes
+	return Scene != nullptr && Scene->GetRenderScene() != nullptr && Scene->GetRenderScene()->GetFrameNumber() > 1;
 }
 
 void FRuntimeVirtualTextureFinalizer::AddTile(FTileEntry& Tile)
