Transform for runtime virtual texture tangent space

diff --git a/Engine/Shaders/Private/VirtualTextureMaterial.usf b/Engine/Shaders/Private/VirtualTextureMaterial.usf
--- a/Engine/Shaders/Private/VirtualTextureMaterial.usf
+++ b/Engine/Shaders/Private/VirtualTextureMaterial.usf
@@ -82,7 +82,7 @@ void FPixelShaderInOut_MainPS(
 	half3 BaseColor = GetMaterialBaseColor(PixelMaterialInputs);
 	half Specular = GetMaterialSpecular(PixelMaterialInputs);
 	half Roughness = GetMaterialRoughness(PixelMaterialInputs);
-	half3 Normal = MaterialParameters.WorldNormal;
+	half3 Normal = mul(MaterialParameters.WorldNormal, (float3x3)ResolvedView.WorldToVirtualTexture);
 	float Opacity = GetMaterialOpacity(PixelMaterialInputs);
 #endif
 
diff --git a/Engine/Source/Runtime/Engine/Public/SceneView.h b/Engine/Source/Runtime/Engine/Public/SceneView.h
--- a/Engine/Source/Runtime/Engine/Public/SceneView.h
+++ b/Engine/Source/Runtime/Engine/Public/SceneView.h
@@ -704,7 +704,8 @@ enum ETranslucencyVolumeCascade
 	VIEW_UNIFORM_BUFFER_MEMBER(FVector, VolumetricLightmapBrickTexelSize) \
 	VIEW_UNIFORM_BUFFER_MEMBER(float, StereoIPD) \
 	VIEW_UNIFORM_BUFFER_MEMBER(float, IndirectLightingCacheShowFlag) \
-	VIEW_UNIFORM_BUFFER_MEMBER(float, EyeToPixelSpreadAngle)
+	VIEW_UNIFORM_BUFFER_MEMBER(float, EyeToPixelSpreadAngle) \
+	VIEW_UNIFORM_BUFFER_MEMBER(FMatrix, WorldToVirtualTexture) \
 
 #define VIEW_UNIFORM_BUFFER_MEMBER(type, identifier) \
 	SHADER_PARAMETER(type, identifier)
diff --git a/Engine/Source/Runtime/Renderer/Private/VT/RuntimeVirtualTextureRender.cpp b/Engine/Source/Runtime/Renderer/Private/VT/RuntimeVirtualTextureRender.cpp
--- a/Engine/Source/Runtime/Renderer/Private/VT/RuntimeVirtualTextureRender.cpp
+++ b/Engine/Source/Runtime/Renderer/Private/VT/RuntimeVirtualTextureRender.cpp
@@ -597,6 +597,7 @@ namespace RuntimeVirtualTexture
 		View->ViewRect = View->UnconstrainedViewRect;
 		View->CachedViewUniformShaderParameters = MakeUnique<FViewUniformShaderParameters>();
 		View->SetupUniformBufferParameters(SceneContext, nullptr, 0, *View->CachedViewUniformShaderParameters);
+		View->CachedViewUniformShaderParameters->WorldToVirtualTexture = WorldToUVRotate.ToMatrixNoScale();
 		View->ViewUniformBuffer = TUniformBufferRef<FViewUniformShaderParameters>::CreateUniformBufferImmediate(*View->CachedViewUniformShaderParameters, UniformBuffer_SingleFrame);
 		UploadDynamicPrimitiveShaderDataForView(RHICmdList, *(const_cast<FScene*>(Scene)), *View);
 		Scene->UniformBuffers.VirtualTextureViewUniformBuffer.UpdateUniformBufferImmediate(*View->CachedViewUniformShaderParameters);
