Fix OpenGL4 compilation error
#UE-75421

diff --git a/Engine/Shaders/Private/VirtualTextureCompress.usf b/Engine/Shaders/Private/VirtualTextureCompress.usf
--- a/Engine/Shaders/Private/VirtualTextureCompress.usf
+++ b/Engine/Shaders/Private/VirtualTextureCompress.usf
@@ -10,16 +10,18 @@
 #include "BCCompressionCommon.ush"
 
 uint4 DestRect;
-SamplerState TextureSampler;
 Texture2D<float4> RenderTexture0;
+SamplerState TextureSampler0;
 Texture2D<float4> RenderTexture1;
+SamplerState TextureSampler1;
 Texture2D<float4> RenderTexture2;
+SamplerState TextureSampler2;
 RWTexture2D<uint2> OutCompressTexture0;
 RWTexture2D<uint4> OutCompressTexture1;
 RWTexture2D<float4> OutCopyTexture0;
 
 // Read a 4x4 color block ready for BC1 compression
-void ReadBlockRGB( Texture2D<float4> SourceTexture, float2 UV, out float3 Block[16] )
+void ReadBlockRGB( Texture2D<float4> SourceTexture, SamplerState TextureSampler, float2 UV, out float3 Block[16] )
 {
 	{
 		float4 Red = SourceTexture.GatherRed(TextureSampler, UV, int2(0,0));
@@ -60,7 +62,7 @@ void ReadBlockRGB( Texture2D<float4> SourceTexture, float2 UV, out float3 Block[
 }
 
 // Read a 4x4 block of XY channels from a normal texture ready for BC5 compression
-void ReadBlockRG( Texture2D<float4> SourceTexture, float2 UV, out float BlockR[16], out float BlockG[16] )
+void ReadBlockRG( Texture2D<float4> SourceTexture, SamplerState TextureSampler, float2 UV, out float BlockR[16], out float BlockG[16] )
 {
 	{
 		float4 Red = SourceTexture.GatherRed(TextureSampler, UV, int2(0,0));
@@ -118,7 +120,7 @@ void CompressBaseColorCS(
 	float2 SampleUV = float2(SamplePos) / float2(DestRect.zw);
 
 	float3 BlockBaseColor[16];
-	ReadBlockRGB(RenderTexture0, SampleUV, BlockBaseColor);
+	ReadBlockRGB(RenderTexture0, TextureSampler0, SampleUV, BlockBaseColor);
 
 	OutCompressTexture0[ThreadId.xy + DestRect.xy] = CompressBC1Block_SRGB(BlockBaseColor);
 }
@@ -136,11 +138,11 @@ void CompressBaseColorNormalCS(
 	float2 SampleUV = float2(SamplePos) / float2(DestRect.zw);
 
 	float3 BlockBaseColor[16];
-	ReadBlockRGB(RenderTexture0, SampleUV, BlockBaseColor);
+	ReadBlockRGB(RenderTexture0, TextureSampler0, SampleUV, BlockBaseColor);
 
 	float BlockNormalX[16];
 	float BlockNormalY[16];
-	ReadBlockRG(RenderTexture1, SampleUV, BlockNormalX, BlockNormalY);
+	ReadBlockRG(RenderTexture1, TextureSampler1, SampleUV, BlockNormalX, BlockNormalY);
 
 	OutCompressTexture0[ThreadId.xy + DestRect.xy] = CompressBC1Block_SRGB(BlockBaseColor);
 	OutCompressTexture1[ThreadId.xy + DestRect.xy] = CompressBC5Block(BlockNormalX, BlockNormalY);
@@ -159,14 +161,14 @@ void CompressBaseColorNormalSpecularCS(
 	float2 SampleUV = float2(SamplePos) / float2(DestRect.zw);
 
 	float3 BlockBaseColor[16];
-	ReadBlockRGB(RenderTexture0, SampleUV, BlockBaseColor);
+	ReadBlockRGB(RenderTexture0, TextureSampler0, SampleUV, BlockBaseColor);
 
 	float BlockNormalX[16];
 	float BlockNormalY[16];
-	ReadBlockRG(RenderTexture1, SampleUV, BlockNormalX, BlockNormalY);
+	ReadBlockRG(RenderTexture1, TextureSampler1, SampleUV, BlockNormalX, BlockNormalY);
 
 	float3 BlockSpecular[16];
-	ReadBlockRGB(RenderTexture2, SampleUV, BlockSpecular);
+	ReadBlockRGB(RenderTexture2, TextureSampler2, SampleUV, BlockSpecular);
 
 	CombineNormalSpecularBlocksForBC3(BlockNormalY, BlockSpecular);
 
diff --git a/Engine/Source/Runtime/Renderer/Private/VT/RuntimeVirtualTextureRender.cpp b/Engine/Source/Runtime/Renderer/Private/VT/RuntimeVirtualTextureRender.cpp
--- a/Engine/Source/Runtime/Renderer/Private/VT/RuntimeVirtualTextureRender.cpp
+++ b/Engine/Source/Runtime/Renderer/Private/VT/RuntimeVirtualTextureRender.cpp
@@ -367,10 +367,12 @@ namespace RuntimeVirtualTexture
 	public:
 		BEGIN_SHADER_PARAMETER_STRUCT(FParameters, )
 			SHADER_PARAMETER(FIntVector4, DestRect)
-			SHADER_PARAMETER_SAMPLER(SamplerState, TextureSampler)
 			SHADER_PARAMETER_RDG_TEXTURE(Texture2D<float4>, RenderTexture0)
+			SHADER_PARAMETER_SAMPLER(SamplerState, TextureSampler0)
 			SHADER_PARAMETER_RDG_TEXTURE(Texture2D<float4>, RenderTexture1)
+			SHADER_PARAMETER_SAMPLER(SamplerState, TextureSampler1)
 			SHADER_PARAMETER_RDG_TEXTURE(Texture2D<float4>, RenderTexture2)
+			SHADER_PARAMETER_SAMPLER(SamplerState, TextureSampler2)
 			SHADER_PARAMETER_RDG_TEXTURE_UAV(RWTexture<uint2>, OutCompressTexture0)
 			SHADER_PARAMETER_RDG_TEXTURE_UAV(RWTexture<uint4>, OutCompressTexture1)
 			SHADER_PARAMETER_RDG_TEXTURE_UAV(RWTexture<float4>, OutCopyTexture0)
@@ -629,10 +631,12 @@ namespace RuntimeVirtualTexture
 		{
 			FShader_VirtualTextureCompress::FParameters* PassParameters = GraphBuilder.AllocParameters<FShader_VirtualTextureCompress::FParameters>();
 			PassParameters->DestRect = FIntVector4(0, 0, TextureSize.X, TextureSize.Y);
-			PassParameters->TextureSampler = TStaticSamplerState<SF_Point, AM_Clamp, AM_Clamp, AM_Clamp>::GetRHI();
 			PassParameters->RenderTexture0 = GraphSetup.RenderTexture0;
+			PassParameters->TextureSampler0 = TStaticSamplerState<SF_Point, AM_Clamp, AM_Clamp, AM_Clamp>::GetRHI();
 			PassParameters->RenderTexture1 = GraphSetup.RenderTexture1;
+			PassParameters->TextureSampler1 = TStaticSamplerState<SF_Point, AM_Clamp, AM_Clamp, AM_Clamp>::GetRHI();
 			PassParameters->RenderTexture2 = GraphSetup.RenderTexture2;
+			PassParameters->TextureSampler2 = TStaticSamplerState<SF_Point, AM_Clamp, AM_Clamp, AM_Clamp>::GetRHI();
 			PassParameters->OutCompressTexture0 = GraphSetup.CompressTexture0 ? GraphBuilder.CreateUAV(FRDGTextureUAVDesc(GraphSetup.CompressTexture0)) : nullptr;
 			PassParameters->OutCompressTexture1 = GraphSetup.CompressTexture1 ? GraphBuilder.CreateUAV(FRDGTextureUAVDesc(GraphSetup.CompressTexture1)) : nullptr;
 
@@ -644,10 +648,12 @@ namespace RuntimeVirtualTexture
 		{
 			FShader_VirtualTextureCompress::FParameters* PassParameters = GraphBuilder.AllocParameters<FShader_VirtualTextureCompress::FParameters>();
 			PassParameters->DestRect = FIntVector4(0, 0, TextureSize.X, TextureSize.Y);
-			PassParameters->TextureSampler = TStaticSamplerState<SF_Point, AM_Clamp, AM_Clamp, AM_Clamp>::GetRHI();
 			PassParameters->RenderTexture0 = GraphSetup.RenderTexture0;
+			PassParameters->TextureSampler0 = TStaticSamplerState<SF_Point, AM_Clamp, AM_Clamp, AM_Clamp>::GetRHI();
 			PassParameters->RenderTexture1 = GraphSetup.RenderTexture1;
+			PassParameters->TextureSampler1 = TStaticSamplerState<SF_Point, AM_Clamp, AM_Clamp, AM_Clamp>::GetRHI();
 			PassParameters->RenderTexture2 = GraphSetup.RenderTexture2;
+			PassParameters->TextureSampler2 = TStaticSamplerState<SF_Point, AM_Clamp, AM_Clamp, AM_Clamp>::GetRHI();
 			PassParameters->OutCopyTexture0 = GraphSetup.CopyTexture0 ? GraphBuilder.CreateUAV(FRDGTextureUAVDesc(GraphSetup.CopyTexture0)) : nullptr;
 
 			AddCopyPass(GraphBuilder, View->GetFeatureLevel(), PassParameters, TextureSize, MaterialType);
