Fixing artifact with TAAU and motion blur in editor.

diff --git a/Engine/Source/Runtime/Renderer/Private/ScreenPass.cpp b/Engine/Source/Runtime/Renderer/Private/ScreenPass.cpp
--- a/Engine/Source/Runtime/Renderer/Private/ScreenPass.cpp
+++ b/Engine/Source/Runtime/Renderer/Private/ScreenPass.cpp
@@ -111,8 +111,8 @@ FScreenPassTextureViewportTransform GetScreenPassTextureViewportTransform(
 	FVector2D DestinationExtent)
 {
 	FScreenPassTextureViewportTransform Transform;
-	Transform.Scale = SourceExtent / DestinationExtent;
-	Transform.Bias = SourceOffset - Transform.Scale * DestinationOffset;
+	Transform.Scale = DestinationExtent / SourceExtent;
+	Transform.Bias = DestinationOffset - Transform.Scale * SourceOffset;
 	return Transform;
 }
 
