Add some material exposed view parameters that assist in runtime virtual texture rendering.
Mip level and world derivatives can potentially be used for adapting material tiling according to render distance.
Might be useful for some landscape cases.
#rnx

diff --git a/Engine/Source/Runtime/Engine/Classes/Materials/MaterialExpressionViewProperty.h b/Engine/Source/Runtime/Engine/Classes/Materials/MaterialExpressionViewProperty.h
--- a/Engine/Source/Runtime/Engine/Classes/Materials/MaterialExpressionViewProperty.h
+++ b/Engine/Source/Runtime/Engine/Classes/Materials/MaterialExpressionViewProperty.h
@@ -31,6 +31,10 @@ enum EMaterialExposedViewProperty
 	MEVP_TemporalSampleIndex UMETA(DisplayName = "Temporal Sample Index"),
 	/** Offset of the temporal sample for this frame in pixel size. */
 	MEVP_TemporalSampleOffset UMETA(DisplayName = "Temporal Sample Offset"),
+	/** Mip Level that Runtime Virtual Texture Output is rendering to. */
+	MEVP_RuntimeVirtualTextureOutputLevel UMETA(DisplayName = "Virtual Texture Output Level"),
+	/** World space derivatives for Runtime Virtual Texture Output. */
+	MEVP_RuntimeVirtualTextureOutputDerivative UMETA(DisplayName = "Virtual Texture Output Derivative"),
 
 	MEVP_MAX,
 };
diff --git a/Engine/Source/Runtime/Engine/Private/SceneView.cpp b/Engine/Source/Runtime/Engine/Private/SceneView.cpp
--- a/Engine/Source/Runtime/Engine/Private/SceneView.cpp
+++ b/Engine/Source/Runtime/Engine/Private/SceneView.cpp
@@ -2344,6 +2344,8 @@ void FSceneView::SetupCommonViewUniformBufferParameters(
 
 	ViewUniformShaderParameters.CameraCut = bCameraCut ? 1 : 0;
 
+	ViewUniformShaderParameters.VirtualTextureParams = FVector4(ForceInitToZero);
+
 	//to tail call keep the order and number of parameters of the caller function
 	SetupViewRectUniformBufferParameters(ViewUniformShaderParameters, BufferSize, EffectiveViewRect, InViewMatrices, InPrevViewMatrices);
 }
diff --git a/Engine/Source/Runtime/Engine/Public/SceneView.h b/Engine/Source/Runtime/Engine/Public/SceneView.h
--- a/Engine/Source/Runtime/Engine/Public/SceneView.h
+++ b/Engine/Source/Runtime/Engine/Public/SceneView.h
@@ -706,6 +706,7 @@ enum ETranslucencyVolumeCascade
 	VIEW_UNIFORM_BUFFER_MEMBER(float, IndirectLightingCacheShowFlag) \
 	VIEW_UNIFORM_BUFFER_MEMBER(float, EyeToPixelSpreadAngle) \
 	VIEW_UNIFORM_BUFFER_MEMBER(FMatrix, WorldToVirtualTexture) \
+	VIEW_UNIFORM_BUFFER_MEMBER(FVector4, VirtualTextureParams) \
 
 #define VIEW_UNIFORM_BUFFER_MEMBER(type, identifier) \
 	SHADER_PARAMETER(type, identifier)
diff --git a/Engine/Source/Runtime/Renderer/Private/VT/RuntimeVirtualTextureProducer.cpp b/Engine/Source/Runtime/Renderer/Private/VT/RuntimeVirtualTextureProducer.cpp
--- a/Engine/Source/Runtime/Renderer/Private/VT/RuntimeVirtualTextureProducer.cpp
+++ b/Engine/Source/Runtime/Renderer/Private/VT/RuntimeVirtualTextureProducer.cpp
@@ -76,7 +76,8 @@ void FRuntimeVirtualTextureFinalizer::Finalize(FRHICommandListImmediate& RHICmdL
 			Entry.Texture1,
 			DestinationBox1,
 			UVToWorld,
-			UVRange);
+			UVRange, 
+			Entry.vLevel);
 	}
 
 	Tiles.SetNumUnsafeInternal(0);
diff --git a/Engine/Source/Runtime/Renderer/Private/VT/RuntimeVirtualTextureRender.cpp b/Engine/Source/Runtime/Renderer/Private/VT/RuntimeVirtualTextureRender.cpp
--- a/Engine/Source/Runtime/Renderer/Private/VT/RuntimeVirtualTextureRender.cpp
+++ b/Engine/Source/Runtime/Renderer/Private/VT/RuntimeVirtualTextureRender.cpp
@@ -549,7 +549,8 @@ namespace RuntimeVirtualTexture
 		FRHITexture2D* OutputTexture1,
 		FBox2D const& DestBox1,
 		FTransform const& UVToWorld,
-		FBox2D const& UVRange)
+		FBox2D const& UVRange,
+		uint8 vLevel)
 	{
 		SCOPED_DRAW_EVENT(RHICmdList, VirtualTextureDynamicCache);
 
@@ -599,6 +600,7 @@ namespace RuntimeVirtualTexture
 		View->CachedViewUniformShaderParameters = MakeUnique<FViewUniformShaderParameters>();
 		View->SetupUniformBufferParameters(SceneContext, nullptr, 0, *View->CachedViewUniformShaderParameters);
 		View->CachedViewUniformShaderParameters->WorldToVirtualTexture = WorldToUVRotate.ToMatrixNoScale();
+		View->CachedViewUniformShaderParameters->VirtualTextureParams = FVector4((float)vLevel, 0.f, OrthoWidth/(float)TextureSize.X, OrthoHeight / (float)TextureSize.Y);
 		View->ViewUniformBuffer = TUniformBufferRef<FViewUniformShaderParameters>::CreateUniformBufferImmediate(*View->CachedViewUniformShaderParameters, UniformBuffer_SingleFrame);
 		UploadDynamicPrimitiveShaderDataForView(RHICmdList, *(const_cast<FScene*>(Scene)), *View);
 		Scene->UniformBuffers.VirtualTextureViewUniformBuffer.UpdateUniformBufferImmediate(*View->CachedViewUniformShaderParameters);
diff --git a/Engine/Source/Runtime/Renderer/Private/VT/RuntimeVirtualTextureRender.h b/Engine/Source/Runtime/Renderer/Private/VT/RuntimeVirtualTextureRender.h
--- a/Engine/Source/Runtime/Renderer/Private/VT/RuntimeVirtualTextureRender.h
+++ b/Engine/Source/Runtime/Renderer/Private/VT/RuntimeVirtualTextureRender.h
@@ -21,5 +21,6 @@ namespace RuntimeVirtualTexture
 		FRHITexture2D* Texture0, FBox2D const& DestBox0,
 		FRHITexture2D* Texture1, FBox2D const& DestBox1, 
 		FTransform const& UVToWorld,
-		FBox2D const& UVRange);
+		FBox2D const& UVRange,
+		uint8 vLevel);
 }
