Replacing QuasiRandom dependencies with StructuredBuffer instead of Buffer. Preferring vector loads to scalar loads.

diff --git a/Engine/Shaders/Private/Halton.ush b/Engine/Shaders/Private/Halton.ush
--- a/Engine/Shaders/Private/Halton.ush
+++ b/Engine/Shaders/Private/Halton.ush
@@ -6,9 +6,14 @@ Halton.ush: Halton sampling interface and utilities
 
 #pragma once
 
+uint3 HaltonPrime3(uint DimensionDividedByThree)
+{
+	return HaltonPrimes.Primes[DimensionDividedByThree % HaltonPrimes.Dimensions];
+}
+
 uint HaltonPrime(uint Dimension)
 {
-	return HaltonPrimes.Primes[Dimension % HaltonPrimes.Dimensions];
+	return HaltonPrimes.Primes[Dimension % HaltonPrimes.Dimensions].x;
 }
 
 // https://en.wikipedia.org/wiki/Halton_sequence
@@ -33,6 +38,15 @@ float HaltonSequence(uint SequenceIndex, int IterationIndex, int DimensionIndex)
 {
 	int SequenceOffset = SequenceIndex * HaltonIteration.IterationCount * HaltonIteration.Dimensions;
 	int IterationOffset = IterationIndex * HaltonIteration.Dimensions;
-	float RandSample = HaltonIteration.SequenceIteration[SequenceOffset + IterationOffset + DimensionIndex];
+	float RandSample = HaltonIteration.SequenceIteration[SequenceOffset + IterationOffset + DimensionIndex].x;
 	return RandSample;
 }
+
+float3 HaltonSequence3(uint SequenceIndex, int IterationIndex, int DimensionIndex)
+{
+	int SequenceOffset = SequenceIndex * HaltonIteration.IterationCount * HaltonIteration.Dimensions;
+	int IterationOffset = IterationIndex * HaltonIteration.Dimensions;
+	int SampleOffset = SequenceOffset + IterationOffset + DimensionIndex;
+	float3 RandSample = HaltonIteration.SequenceIteration[SampleOffset];
+	return RandSample;
+}
\ No newline at end of file
diff --git a/Engine/Shaders/Private/QuasiRandom.ush b/Engine/Shaders/Private/QuasiRandom.ush
--- a/Engine/Shaders/Private/QuasiRandom.ush
+++ b/Engine/Shaders/Private/QuasiRandom.ush
@@ -94,19 +94,17 @@ float3 QuasiRandom_Sample3(inout FQuasiRandom QuasiRandom)
 		float3 BlueNoiseShift = EvalBlueNoise(QuasiRandom.BlueNoiseCoord, QuasiRandom.Dimension / 3);
 
 		// Evaluate Halton
-		float3 HaltonSample;
-		HaltonSample.x = HaltonSequence(QuasiRandom.HaltonSequence, QuasiRandom.HaltonIteration, QuasiRandom.Dimension);
-		HaltonSample.y = HaltonSequence(QuasiRandom.HaltonSequence, QuasiRandom.HaltonIteration, QuasiRandom.Dimension + 1);
-		HaltonSample.z = HaltonSequence(QuasiRandom.HaltonSequence, QuasiRandom.HaltonIteration, QuasiRandom.Dimension + 2);
+		float3 HaltonSample = HaltonSequence3(QuasiRandom.HaltonSequence, QuasiRandom.HaltonIteration, QuasiRandom.Dimension / 3);
 
 		// Wrap sum
-		RandomSample = fmod(HaltonSample + BlueNoiseShift, 1.0);
+		RandomSample = fmod(BlueNoiseShift + HaltonSample, 1.0);
 	}
 	else
 	{
-		RandomSample.x = HaltonProcedural(QuasiRandom.PositionSeed + QuasiRandom.TimeSeed, HaltonPrime(QuasiRandom.Dimension));
-		RandomSample.y = HaltonProcedural(QuasiRandom.PositionSeed + QuasiRandom.TimeSeed, HaltonPrime(QuasiRandom.Dimension + 1));
-		RandomSample.z = HaltonProcedural(QuasiRandom.PositionSeed + QuasiRandom.TimeSeed, HaltonPrime(QuasiRandom.Dimension + 2));
+		uint3 Primes = HaltonPrime3(QuasiRandom.Dimension / 3);
+		RandomSample.x = HaltonProcedural(QuasiRandom.PositionSeed + QuasiRandom.TimeSeed, Primes[0]);
+		RandomSample.y = HaltonProcedural(QuasiRandom.PositionSeed + QuasiRandom.TimeSeed, Primes[1]);
+		RandomSample.z = HaltonProcedural(QuasiRandom.PositionSeed + QuasiRandom.TimeSeed, Primes[2]);
 	}
 
 	// Increment dimension
diff --git a/Engine/Source/Runtime/Renderer/Private/Halton.cpp b/Engine/Source/Runtime/Renderer/Private/Halton.cpp
--- a/Engine/Source/Runtime/Renderer/Private/Halton.cpp
+++ b/Engine/Source/Runtime/Renderer/Private/Halton.cpp
@@ -43,21 +43,30 @@ DECLARE_GPU_STAT_NAMED(HaltonSequence, TEXT("Halton Sequence"));
 
 void FHaltonSequenceIteration::InitRHI()
 {
+	struct FSequenceIterationData
+	{
+		float Sequence[3];
+	};
+
 	SCOPED_GPU_STAT(FRHICommandListExecutor::GetImmediateCommandList(), HaltonSequence);
 	InitializeSequence();
 
-	TArray<float> RandomSamples;
-	RandomSamples.SetNum(SequenceCount * IterationCount * DimensionCount);
+	TArray<FSequenceIterationData> RandomSamples;
+	uint32 ElementCount = FMath::DivideAndRoundUp(DimensionCount, 3u);
+	RandomSamples.SetNum(SequenceCount * IterationCount * ElementCount);
 	for (uint32 SequenceIndex = 0; SequenceIndex < SequenceCount; ++SequenceIndex)
 	{
 		uint32 SequenceValue = Sequence[SequenceIndex] + Iteration * IterationCount;
-		uint32 SequenceOffset = SequenceIndex * IterationCount * DimensionCount;
+		uint32 SequenceOffset = SequenceIndex * IterationCount * ElementCount;
 		for (uint32 IterationIndex = 0; IterationIndex < IterationCount; ++IterationIndex)
 		{
-			uint32 IterationOffset = IterationIndex * DimensionCount;
-			for (uint32 DimensionIndex = 0; DimensionIndex < DimensionCount; ++DimensionIndex)
+			uint32 IterationOffset = IterationIndex * ElementCount;
+			for (uint32 ElementIndex = 0; ElementIndex < ElementCount; ++ElementIndex)
 			{
-				RandomSamples[SequenceOffset + IterationOffset + DimensionIndex] = HaltonSequence.Sample(DimensionIndex, SequenceValue + IterationIndex);
+				uint32 DimensionOffset = ElementIndex * 3;
+				RandomSamples[SequenceOffset + IterationOffset + ElementIndex].Sequence[0] = HaltonSequence.Sample(DimensionOffset, SequenceValue + IterationIndex);
+				RandomSamples[SequenceOffset + IterationOffset + ElementIndex].Sequence[1] = HaltonSequence.Sample(DimensionOffset + 1, SequenceValue + IterationIndex);
+				RandomSamples[SequenceOffset + IterationOffset + ElementIndex].Sequence[2] = HaltonSequence.Sample(DimensionOffset + 2, SequenceValue + IterationIndex);
 			}
 		}
 	}
@@ -65,11 +74,11 @@ void FHaltonSequenceIteration::InitRHI()
 	FRHIResourceCreateInfo CreateInfo;
 	{
 		CreateInfo.DebugName = TEXT("HaltonSequenceIteration");
-		SequenceIteration.VertexBufferRHI = RHICreateVertexBuffer(RandomSamples.Num() * sizeof(float), BUF_Transient | BUF_FastVRAM | BUF_ShaderResource | BUF_UnorderedAccess, CreateInfo);
+		SequenceIteration = RHICreateStructuredBuffer(sizeof(FSequenceIterationData), RandomSamples.Num() * sizeof(FSequenceIterationData), BUF_Transient | BUF_FastVRAM | BUF_ShaderResource | BUF_UnorderedAccess, CreateInfo);
 		uint32 Offset = 0;
-		void* BasePtr = RHILockVertexBuffer(SequenceIteration.VertexBufferRHI, 0, RandomSamples.Num() * sizeof(float), RLM_WriteOnly);
-		FPlatformMemory::Memcpy(BasePtr, RandomSamples.GetData(), RandomSamples.Num() * sizeof(float));
-		RHIUnlockVertexBuffer(SequenceIteration.VertexBufferRHI);
+		void* BasePtr = RHILockStructuredBuffer(SequenceIteration, Offset, RandomSamples.Num() * sizeof(FSequenceIterationData), RLM_WriteOnly);
+		FPlatformMemory::Memcpy(BasePtr, RandomSamples.GetData(), RandomSamples.Num() * sizeof(FSequenceIterationData));
+		RHIUnlockStructuredBuffer(SequenceIteration);
 	}
 }
 
@@ -148,11 +157,17 @@ void FHaltonPrimesResource::InitRHI()
 {
 	FRHIResourceCreateInfo CreateInfo;
 	{
+		struct FPrimeData
+		{
+			int Primes[3];
+		};
+
 		CreateInfo.DebugName = TEXT("HaltonPrimes");
-		PrimesBuffer.VertexBufferRHI = RHICreateVertexBuffer(DimensionCount * sizeof(float), BUF_Transient | BUF_FastVRAM | BUF_ShaderResource | BUF_UnorderedAccess, CreateInfo);
+		uint32 ElementCount = FMath::DivideAndRoundUp(DimensionCount, 3u);
+		PrimesBuffer = RHICreateStructuredBuffer(sizeof(FPrimeData), ElementCount * sizeof(FPrimeData), BUF_Transient | BUF_FastVRAM | BUF_ShaderResource | BUF_UnorderedAccess, CreateInfo);
 		uint32 Offset = 0;
-		void* BasePtr = RHILockVertexBuffer(PrimesBuffer.VertexBufferRHI, Offset, DimensionCount * sizeof(float), RLM_WriteOnly);
-		FPlatformMemory::Memcpy(BasePtr, Primes.GetData(), DimensionCount * sizeof(float));
-		RHIUnlockVertexBuffer(PrimesBuffer.VertexBufferRHI);
+		void* BasePtr = RHILockStructuredBuffer(PrimesBuffer, Offset, ElementCount * sizeof(FPrimeData), RLM_WriteOnly);
+		FPlatformMemory::Memcpy(BasePtr, Primes.GetData(), DimensionCount * sizeof(int));
+		RHIUnlockStructuredBuffer(PrimesBuffer);
 	}
 }
\ No newline at end of file
diff --git a/Engine/Source/Runtime/Renderer/Private/Halton.h b/Engine/Source/Runtime/Renderer/Private/Halton.h
--- a/Engine/Source/Runtime/Renderer/Private/Halton.h
+++ b/Engine/Source/Runtime/Renderer/Private/Halton.h
@@ -15,7 +15,7 @@ SHADER_PARAMETER(int, Dimensions)
 SHADER_PARAMETER(int, SequenceRowCount)
 SHADER_PARAMETER(int, SequenceColumnCount)
 SHADER_PARAMETER(int, IterationCount)
-SHADER_PARAMETER_SRV(Buffer<float>, SequenceIteration)
+SHADER_PARAMETER_SRV(StructuredBuffer<float3>, SequenceIteration)
 END_GLOBAL_SHADER_PARAMETER_STRUCT()
 
 /**
@@ -39,7 +39,6 @@ public:
 
 	virtual void ReleaseRHI() override
 	{
-		SequenceIteration.ReleaseRHI();
 	}
 
 	virtual FString GetFriendlyName() const override
@@ -68,7 +67,7 @@ public:
 		return Iteration;
 	}
 
-	FVertexBuffer SequenceIteration;
+	FStructuredBufferRHIRef SequenceIteration;
 
 private:
 	void InitializeSequence();
@@ -84,16 +83,16 @@ private:
 
 inline void InitializeHaltonSequenceIteration(const FHaltonSequenceIteration& HaltonSequenceIteration, FHaltonIteration& HaltonIteration)
 {
-	HaltonIteration.Dimensions = HaltonSequenceIteration.GetDimensionCount();
+	HaltonIteration.Dimensions = HaltonSequenceIteration.GetDimensionCount() / 3u;
 	HaltonIteration.SequenceRowCount = FMath::Sqrt(HaltonSequenceIteration.GetSequenceCount());
 	HaltonIteration.SequenceColumnCount = HaltonSequenceIteration.GetSequenceCount() / HaltonIteration.SequenceRowCount;
 	HaltonIteration.IterationCount = HaltonSequenceIteration.GetIterationCount();
-	HaltonIteration.SequenceIteration = RHICreateShaderResourceView(HaltonSequenceIteration.SequenceIteration.VertexBufferRHI, sizeof(float), PF_R32_FLOAT);
+	HaltonIteration.SequenceIteration = RHICreateShaderResourceView(HaltonSequenceIteration.SequenceIteration);
 }
 
 BEGIN_GLOBAL_SHADER_PARAMETER_STRUCT(FHaltonPrimes, )
 SHADER_PARAMETER(int, Dimensions)
-SHADER_PARAMETER_SRV(Buffer<int>, Primes)
+SHADER_PARAMETER_SRV(StructuredBuffer<uint3>, Primes)
 END_GLOBAL_SHADER_PARAMETER_STRUCT()
 
 /**
@@ -110,7 +109,6 @@ public:
 
 	virtual void ReleaseRHI() override
 	{
-		PrimesBuffer.ReleaseRHI();
 	}
 
 	virtual FString GetFriendlyName() const override
@@ -124,7 +122,7 @@ public:
 		return DimensionCount;
 	}
 
-	FVertexBuffer PrimesBuffer;
+	FStructuredBufferRHIRef PrimesBuffer;
 
 private:
 	TArray<int> Primes;
@@ -134,5 +132,5 @@ private:
 inline void InitializeHaltonPrimes(const FHaltonPrimesResource& HaltonPrimeResource, FHaltonPrimes& HaltonPrimes)
 {
 	HaltonPrimes.Dimensions = HaltonPrimeResource.GetDimensionCount();
-	HaltonPrimes.Primes = RHICreateShaderResourceView(HaltonPrimeResource.PrimesBuffer.VertexBufferRHI, sizeof(float), PF_R32_FLOAT);
+	HaltonPrimes.Primes = RHICreateShaderResourceView(HaltonPrimeResource.PrimesBuffer);
 }
\ No newline at end of file
