Porting of the old scene textures uniform buffers to the new FSceneTextureParameters object for all the ray tracing shaders

diff --git a/Engine/Shaders/Private/RayTracing/CompositeSkyLightPS.usf b/Engine/Shaders/Private/RayTracing/CompositeSkyLightPS.usf
--- a/Engine/Shaders/Private/RayTracing/CompositeSkyLightPS.usf
+++ b/Engine/Shaders/Private/RayTracing/CompositeSkyLightPS.usf
@@ -2,6 +2,7 @@
 
 #include "../Common.ush"
 #include "../DeferredShadingCommon.ush"
+#include "../SceneTextureParameters.ush"
 
 Texture2D<float4> SkyLightTexture;
 SamplerState SkyLightTextureSampler;
@@ -11,8 +12,8 @@ void CompositeSkyLightPS(
 	out float4 OutColor : SV_Target0
 )
 {
-	FScreenSpaceData ScreenSpaceData = GetScreenSpaceData(UV);
-	float3 Albedo = ScreenSpaceData.GBuffer.StoredBaseColor - ScreenSpaceData.GBuffer.StoredBaseColor * ScreenSpaceData.GBuffer.Metallic;
+	FGBufferData GBufferData = GetGBufferDataFromSceneTextures(UV);
+	float3 Albedo = GBufferData.StoredBaseColor - GBufferData.StoredBaseColor * GBufferData.Metallic;
 	float4 SkyLight = SkyLightTexture.Sample(SkyLightTextureSampler, UV);
 	// Apply albedo after denoising
 	SkyLight.rgb *= Albedo;
diff --git a/Engine/Shaders/Private/RayTracing/RayTracingAmbientOcclusionRGS.usf b/Engine/Shaders/Private/RayTracing/RayTracingAmbientOcclusionRGS.usf
--- a/Engine/Shaders/Private/RayTracing/RayTracingAmbientOcclusionRGS.usf
+++ b/Engine/Shaders/Private/RayTracing/RayTracingAmbientOcclusionRGS.usf
@@ -3,6 +3,7 @@
 #include "../Common.ush"
 #include "../Montecarlo.ush"
 #include "../DeferredShadingCommon.ush"
+#include "../SceneTextureParameters.ush"
 #include "../PathTracing/Utilities/PathTracingRandomSequence.ush"
 #include "RayTracingCommon.ush"
 
@@ -52,10 +53,10 @@ void AmbientOcclusionRGS()
 	// Get G-Buffer surface data
 	float2 InvBufferSize = View.BufferSizeAndInvSize.zw;
 	float2 UV = (float2(PixelCoord) + 0.5) * InvBufferSize;
-	FScreenSpaceData ScreenSpaceData = GetScreenSpaceData(UV);
-	float Depth = ScreenSpaceData.GBuffer.Depth;
+	FGBufferData GBufferData = GetGBufferDataFromSceneTextures(UV);	
+	float Depth = GBufferData.Depth;
 	float3 WorldPosition = ReconstructWorldPositionFromDepth(UV, Depth);
-	float3 WorldNormal = ScreenSpaceData.GBuffer.WorldNormal;
+	float3 WorldNormal = GBufferData.WorldNormal;
 
 	float RayCount = 0.0;
 	float Visibility = 0.0;
@@ -65,11 +66,9 @@ void AmbientOcclusionRGS()
 	uint SamplesPerPixelLocal = SamplesPerPixel;
 	
 	// Mask out depth values beyond far plane
-	float DeviceZ = SceneTexturesStruct.SceneDepthTexture.Load(int3(PixelCoord, 0)).r;
+	float DeviceZ = SceneDepthBuffer.Load(int3(PixelCoord, 0)).r;
 	bool IsFiniteDepth = DeviceZ > 0.0;
-	bool bTraceRay = (
-		IsFiniteDepth &&
-		ScreenSpaceData.GBuffer.ShadingModelID != SHADINGMODELID_UNLIT);
+	bool bTraceRay = (IsFiniteDepth && GBufferData.ShadingModelID != SHADINGMODELID_UNLIT);
 	if (!bTraceRay)
 	{
 		Visibility = 1.0;
diff --git a/Engine/Shaders/Private/RayTracing/RayTracingGlobalIlluminationCompositePS.usf b/Engine/Shaders/Private/RayTracing/RayTracingGlobalIlluminationCompositePS.usf
--- a/Engine/Shaders/Private/RayTracing/RayTracingGlobalIlluminationCompositePS.usf
+++ b/Engine/Shaders/Private/RayTracing/RayTracingGlobalIlluminationCompositePS.usf
@@ -2,6 +2,7 @@
 
 #include "../Common.ush"
 #include "../DeferredShadingCommon.ush"
+#include "../SceneTextureParameters.ush"
 
 Texture2D<float4> GlobalIlluminationTexture;
 SamplerState GlobalIlluminationSampler;
@@ -12,11 +13,11 @@ void GlobalIlluminationCompositePS(
 	out float OutAmbientOcclusion : SV_Target1
 )
 {
-	FScreenSpaceData ScreenSpaceData = GetScreenSpaceData(UV);
-	float3 DiffuseColor = ScreenSpaceData.GBuffer.DiffuseColor;
-	if (UseSubsurfaceProfile(ScreenSpaceData.GBuffer.ShadingModelID))
+	FGBufferData GBufferData = GetGBufferDataFromSceneTextures(UV);
+	float3 DiffuseColor = GBufferData.DiffuseColor;
+	if (UseSubsurfaceProfile(GBufferData.ShadingModelID))
 	{
-		DiffuseColor = ScreenSpaceData.GBuffer.StoredBaseColor;
+		DiffuseColor = GBufferData.StoredBaseColor;
 	}
 	float4 GlobalIllumination = GlobalIlluminationTexture.Sample(GlobalIlluminationSampler, UV);
 	// Apply albedo after denoising
diff --git a/Engine/Shaders/Private/RayTracing/RayTracingGlobalIlluminationRGS.usf b/Engine/Shaders/Private/RayTracing/RayTracingGlobalIlluminationRGS.usf
--- a/Engine/Shaders/Private/RayTracing/RayTracingGlobalIlluminationRGS.usf
+++ b/Engine/Shaders/Private/RayTracing/RayTracingGlobalIlluminationRGS.usf
@@ -7,6 +7,7 @@
 #include "../ShadingModels.ush"
 #include "RayTracingCommon.ush"
 #include "RayTracingHitGroupCommon.ush"
+#include "../SceneTextureParameters.ush"
 
 #include "../PathTracing/Light/PathTracingLightSampling.ush"
 #include "../PathTracing/Material/PathTracingMaterialSampling.ush"
@@ -67,18 +68,18 @@ void GlobalIlluminationRGS()
 	// Get G-Buffer surface data
 	float2 InvBufferSize = View.BufferSizeAndInvSize.zw;
 	float2 UV = (float2(PixelCoord) + 0.5) * InvBufferSize;
-	FScreenSpaceData ScreenSpaceData = GetScreenSpaceData(UV);
+	FGBufferData GBufferData = GetGBufferDataFromSceneTextures(UV);
 	// Remap DiffuseColor when using SubsurfaceProfile (GBuffer decoding replaces with 100% albedo)
-	if (UseSubsurfaceProfile(ScreenSpaceData.GBuffer.ShadingModelID))
+	if (UseSubsurfaceProfile(GBufferData.ShadingModelID))
 	{
-		ScreenSpaceData.GBuffer.DiffuseColor = ScreenSpaceData.GBuffer.StoredBaseColor;
+		GBufferData.DiffuseColor = GBufferData.StoredBaseColor;
 	}
-	float Depth = ScreenSpaceData.GBuffer.Depth;
+	float Depth = GBufferData.Depth;
 	float3 WorldPosition = ReconstructWorldPositionFromDepth(UV, Depth);
 	float3 CameraOrigin = View.WorldCameraOrigin;
 	float3 CameraDirection = normalize(WorldPosition - CameraOrigin);
-	float3 WorldNormal = ScreenSpaceData.GBuffer.WorldNormal;
-	uint ShadingModelID = ScreenSpaceData.GBuffer.ShadingModelID;
+	float3 WorldNormal = GBufferData.WorldNormal;
+	uint ShadingModelID = GBufferData.ShadingModelID;
 	if (ShadingModelID == SHADINGMODELID_UNLIT
 		|| ShadingModelID == SHADINGMODELID_TWOSIDED_FOLIAGE
 		)
@@ -87,7 +88,7 @@ void GlobalIlluminationRGS()
 	}
 
 	// Diffuse color rejection threshold
-	float3 DiffuseColor = ScreenSpaceData.GBuffer.DiffuseColor;
+	float3 DiffuseColor = GBufferData.DiffuseColor;
 	if (Luminance(DiffuseColor) < DiffuseThreshold)
 	{
 		return;
@@ -127,7 +128,7 @@ void GlobalIlluminationRGS()
 		half3 L = Ray.Direction;
 		float NoL = saturate(dot(N, L));
 		FShadowTerms ShadowTerms = { 0.0, 0.0, 0.0 };
-		FDirectLighting LightingSample = EvaluateBxDF(ScreenSpaceData.GBuffer, N, V, L, NoL, ShadowTerms);
+		FDirectLighting LightingSample = EvaluateBxDF(GBufferData, N, V, L, NoL, ShadowTerms);
 		RayThroughput *= LightingSample.Diffuse / DiffuseColor;
 #else
 		uint DummyVariable;
diff --git a/Engine/Shaders/Private/RayTracing/RayTracingReflections.usf b/Engine/Shaders/Private/RayTracing/RayTracingReflections.usf
--- a/Engine/Shaders/Private/RayTracing/RayTracingReflections.usf
+++ b/Engine/Shaders/Private/RayTracing/RayTracingReflections.usf
@@ -35,6 +35,7 @@
 #include "../LightShaderParameters.ush"
 #include "../ScreenSpaceDenoise/SSDPublic.ush"
 #include "../SobolRandom.ush"
+#include "../SceneTextureParameters.ush"
 #include "RayTracingHitGroupCommon.ush"
 #include "RayTracingDirectionalLight.ush"
 #include "RayTracingRectLight.ush"
@@ -164,13 +165,13 @@ void RayTracingReflectionsRGS()
 
 	float2 InvBufferSize = View.BufferSizeAndInvSize.zw;
 	float2 UV = (float2(PixelCoord) + 0.5) * InvBufferSize;
-	FScreenSpaceData ScreenSpaceData = GetScreenSpaceData(UV);
-	float Depth = ScreenSpaceData.GBuffer.Depth;
+	FGBufferData GBufferData = GetGBufferDataFromSceneTextures(UV);
+	float Depth = GBufferData.Depth;
 
-	uint LocalSamplesPerPixel = SamplesPerPixel;
+	uint LocalSamplesPerPixel = SamplesPerPixel; 
 
 	// Mask out depth values that are infinitely far away
-	float DeviceZ = SceneTexturesStruct.SceneDepthTexture.Load(int3(PixelCoord, 0)).r;
+	float DeviceZ = SceneDepthBuffer.Load(int3(PixelCoord, 0)).r;
 	bool IsFiniteDepth = DeviceZ > 0.0;
 	if (!IsFiniteDepth)
 	{
@@ -187,20 +188,20 @@ void RayTracingReflectionsRGS()
 	float ReflectedWeight = 0.0;
 	float HitDistance = DENOISER_INVALID_HIT_DISTANCE;
 
-	float RoughnessFade = GetRoughnessFade(ScreenSpaceData.GBuffer.Roughness, ReflectionMaxRoughness);
-	float ClearCoatRoughnessFade = ScreenSpaceData.GBuffer.ShadingModelID == SHADINGMODELID_CLEAR_COAT ? GetRoughnessFade(ScreenSpaceData.GBuffer.CustomData.y, ReflectionMaxRoughness) : 0.0f;
+	float RoughnessFade = GetRoughnessFade(GBufferData.Roughness, ReflectionMaxRoughness);
+	float ClearCoatRoughnessFade = GBufferData.ShadingModelID == SHADINGMODELID_CLEAR_COAT ? GetRoughnessFade(GBufferData.CustomData.y, ReflectionMaxRoughness) : 0.0f;
 
 	// Sky is only sampled when infinite reflection rays are used.
 	const bool bAllowSkySampling = ReflectionMaxRayDistance < 0;
 
 	// Shorten the rays on rougher surfaces between user-provided min and max ray lengths.
 	// When a shortened ray misses the geometry, we fall back to local reflection capture sampling (similar to SSR).
-	const float LocalMaxRayDistance = bAllowSkySampling ? 1e27f : lerp(ReflectionMaxRayDistance, ReflectionMinRayDistance, ScreenSpaceData.GBuffer.Roughness);
+	const float LocalMaxRayDistance = bAllowSkySampling ? 1e27f : lerp(ReflectionMaxRayDistance, ReflectionMinRayDistance, GBufferData.Roughness);
 
 	// Cutoff
 	// * Roughness value
 	// * Hair material is not for RT reflection at the moment. Their indirect contributuion is based on diffuse SH evaluation in ReflectionEnvironemenPixelShader.usf
-    const bool bIsValidPixel = (RoughnessFade > 0 || ClearCoatRoughnessFade > 0) && ScreenSpaceData.GBuffer.ShadingModelID != SHADINGMODELID_HAIR;
+    const bool bIsValidPixel = (RoughnessFade > 0 || ClearCoatRoughnessFade > 0) && GBufferData.ShadingModelID != SHADINGMODELID_HAIR;
 	   
 	if (bIsValidPixel)
 	{
@@ -218,10 +219,10 @@ void RayTracingReflectionsRGS()
 			float3 CameraDirection = normalize(WorldPosition - CameraOrigin);
 			float3 IncidentDirection = CameraDirection;
 
-			uint ShadingModelID = ScreenSpaceData.GBuffer.ShadingModelID;
-			float3 TopLayerWorldNormal = ScreenSpaceData.GBuffer.WorldNormal;
-			float  TopLayerRoughness = ScreenSpaceData.GBuffer.Roughness;
-			float3 TopLayerSpecularColor = ScreenSpaceData.GBuffer.SpecularColor;
+			uint ShadingModelID = GBufferData.ShadingModelID;
+			float3 TopLayerWorldNormal = GBufferData.WorldNormal;
+			float  TopLayerRoughness = GBufferData.Roughness;
+			float3 TopLayerSpecularColor = GBufferData.SpecularColor;
 
 			// Track Roughness to terminate bounces early
 			float AccumRoughness = TopLayerRoughness;
@@ -233,12 +234,12 @@ void RayTracingReflectionsRGS()
 			if (ShadingModelID == SHADINGMODELID_CLEAR_COAT)
 			{
 				#if CLEAR_COAT_BOTTOM_NORMAL
-				const float2 oct1 = ((float2(ScreenSpaceData.GBuffer.CustomData.a, ScreenSpaceData.GBuffer.CustomData.z) * 2) - (256.0/255.0)) + UnitVectorToOctahedron(TopLayerWorldNormal);
+				const float2 oct1 = ((float2(GBufferData.CustomData.a, GBufferData.CustomData.z) * 2) - (256.0/255.0)) + UnitVectorToOctahedron(TopLayerWorldNormal);
 				BottomLayerWorldNormal = OctahedronToUnitVector(oct1);
 				#endif
 
-				ClearCoat = ScreenSpaceData.GBuffer.CustomData.x;			// Clear coat weight
-				TopLayerRoughness = ScreenSpaceData.GBuffer.CustomData.y;	// Clear coat roughness
+				ClearCoat = GBufferData.CustomData.x;			// Clear coat weight
+				TopLayerRoughness = GBufferData.CustomData.y;	// Clear coat roughness
 				TopLayerSpecularColor = ClearCoatSpecularColor;				// Hard coded top layer specular color
 			}
 
@@ -287,7 +288,7 @@ void RayTracingReflectionsRGS()
 				{
 					// This fixes black reflexions due to R pointing into the surface.
 					// It assumes that the clear coat relies on a smooth normal following the mesh surface.
-					FixSampleDirectionIfNeeded(ScreenSpaceData.GBuffer.WorldNormal, TopLayerRay.Direction);
+					FixSampleDirectionIfNeeded(GBufferData.WorldNormal, TopLayerRay.Direction);
 				}
 
 				if (BounceIndex == 0 && SampleIndex == 0)
@@ -391,9 +392,9 @@ void RayTracingReflectionsRGS()
 					BottomLayerRay.TMax = LocalMaxRayDistance;
 					BottomLayerRay.Direction = GenerateReflectionDirection(RandSequence, DispatchThreadId, BounceIndex, LocalSamplesPerPixel, IncidentDirection, BottomLayerWorldNormal, BottomLayerRoughness);
 					ApplyPositionBias(BottomLayerRay, BottomLayerWorldNormal, ReflectionMaxNormalBias);
-					if (BounceIndex == 0 && dot(ScreenSpaceData.GBuffer.WorldNormal, BottomLayerRay.Direction) < 0.0)
+					if (BounceIndex == 0 && dot(GBufferData.WorldNormal, BottomLayerRay.Direction) < 0.0)
 					{
-						FixSampleDirectionIfNeeded(ScreenSpaceData.GBuffer.WorldNormal, BottomLayerRay.Direction);
+						FixSampleDirectionIfNeeded(GBufferData.WorldNormal, BottomLayerRay.Direction);
 					}
 
 					bool TraceBottomRay = (BounceIndex < MaxUnderCoatBounces) && (GetRoughnessFade(BottomLayerRoughness, ReflectionMaxRoughness) > 0.0);
diff --git a/Engine/Shaders/Private/RayTracing/RayTracingSkyLightRGS.usf b/Engine/Shaders/Private/RayTracing/RayTracingSkyLightRGS.usf
--- a/Engine/Shaders/Private/RayTracing/RayTracingSkyLightRGS.usf
+++ b/Engine/Shaders/Private/RayTracing/RayTracingSkyLightRGS.usf
@@ -7,6 +7,7 @@
 #include "../QuasiRandom.ush"
 #include "SkyLightMipTreeCommon.ush"
 #include "RayTracingCommon.ush"
+#include "../SceneTextureParameters.ush"
 
 RaytracingAccelerationStructure TLAS;
 RWTexture2D<float4> RWOcclusionMaskUAV;
@@ -229,33 +230,33 @@ void SkyLightRGS()
 	float2 InvBufferSize = View.BufferSizeAndInvSize.zw;
 	float2 UV = (float2(PixelCoord) + 0.5) * InvBufferSize;
 	FScreenSpaceData ScreenSpaceData = GetScreenSpaceData(UV);
-	FGBufferData GBuffer = ScreenSpaceData.GBuffer;
-	float Depth = ScreenSpaceData.GBuffer.Depth;
+	FGBufferData GBufferData = GetGBufferDataFromSceneTextures(UV);
+	float Depth = GBufferData.Depth;
 	float3 WorldPosition = ReconstructWorldPositionFromDepth(UV, Depth);
 	float3 CameraOrigin = View.WorldCameraOrigin;
 	float3 CameraDirection = normalize(WorldPosition - CameraOrigin);
-	float3 WorldNormal = ScreenSpaceData.GBuffer.WorldNormal;
-	float3 Albedo = ScreenSpaceData.GBuffer.DiffuseColor;
+	float3 WorldNormal = GBufferData.WorldNormal;
+	float3 Albedo = GBufferData.DiffuseColor;
 
 	// Recalculate DiffuseColor if subsurface reverted the contribution within the G-Buffer
-	if (UseSubsurfaceProfile(GBuffer.ShadingModelID))
+	if (UseSubsurfaceProfile(GBufferData.ShadingModelID))
 	{
-		Albedo = ScreenSpaceData.GBuffer.StoredBaseColor - ScreenSpaceData.GBuffer.StoredBaseColor * ScreenSpaceData.GBuffer.Metallic;
-		ScreenSpaceData.GBuffer.DiffuseColor = Albedo;
+		Albedo = GBufferData.StoredBaseColor - GBufferData.StoredBaseColor * GBufferData.Metallic;
+		GBufferData.DiffuseColor = Albedo;
 	}
-	float3 SpecularColor = ScreenSpaceData.GBuffer.SpecularColor;
-	float Roughness = ScreenSpaceData.GBuffer.Roughness;
+	float3 SpecularColor = GBufferData.SpecularColor;
+	float Roughness = GBufferData.Roughness;
 
 	float RayDistance = 0.0;
 	float HitCount = 0.0;
 	uint SamplesPerPixel = SkyLight.SamplesPerPixel;
 
 	// Mask out depth values that are infinitely far away
-	float DeviceZ = SceneTexturesStruct.SceneDepthTexture.Load(int3(PixelCoord, 0)).r;
+	float DeviceZ = SceneDepthBuffer.Load(int3(PixelCoord, 0)).r;
 	bool IsFiniteDepth = DeviceZ > 0.0;
 	bool bTraceRay = (
 		IsFiniteDepth &&
-		GBuffer.ShadingModelID != SHADINGMODELID_UNLIT);
+		GBufferData.ShadingModelID != SHADINGMODELID_UNLIT);
 	if (!bTraceRay)
 	{
 		SamplesPerPixel = 0.0;
@@ -335,7 +336,7 @@ void SkyLightRGS()
 			half3 L = Ray.Direction;
 			float NoL = dot(WorldNormal, Ray.Direction);
 			FShadowTerms ShadowTerms = { 0.0, 0.0, 0.0 };
-			FDirectLighting LightingSample = EvaluateBxDF(ScreenSpaceData.GBuffer, N, V, L, NoL, ShadowTerms);
+			FDirectLighting LightingSample = EvaluateBxDF(GBufferData, N, V, L, NoL, ShadowTerms);
 			float3 Brdf = LightingSample.Diffuse + LightingSample.Transmission + LightingSample.Specular;
 
 			float3 IncomingRadiance = EvalSkyLight(Ray.Direction);
diff --git a/Engine/Source/Runtime/Renderer/Private/DeferredShadingRenderer.h b/Engine/Source/Runtime/Renderer/Private/DeferredShadingRenderer.h
--- a/Engine/Source/Runtime/Renderer/Private/DeferredShadingRenderer.h
+++ b/Engine/Source/Runtime/Renderer/Private/DeferredShadingRenderer.h
@@ -464,6 +464,7 @@ private:
 
 	void RenderRayTracingReflections(
 		FRDGBuilder& GraphBuilder,
+		const FSceneTextureParameters& SceneTextures,
 		const FViewInfo& View,
 		int32 SamplePerPixel,
 		int32 HeightFog,
diff --git a/Engine/Source/Runtime/Renderer/Private/RayTracing/RayTracingAmbientOcclusion.cpp b/Engine/Source/Runtime/Renderer/Private/RayTracing/RayTracingAmbientOcclusion.cpp
--- a/Engine/Source/Runtime/Renderer/Private/RayTracing/RayTracingAmbientOcclusion.cpp
+++ b/Engine/Source/Runtime/Renderer/Private/RayTracing/RayTracingAmbientOcclusion.cpp
@@ -92,7 +92,7 @@ class FRayTracingAmbientOcclusionRGS : public FGlobalShader
 		SHADER_PARAMETER_RDG_TEXTURE_UAV(RWTexture2D<float>, RWHitDistanceUAV)
 
 		SHADER_PARAMETER_STRUCT_REF(FViewUniformShaderParameters, ViewUniformBuffer)
-		SHADER_PARAMETER_STRUCT_REF(FSceneTexturesUniformParameters, SceneTexturesStruct)
+		SHADER_PARAMETER_STRUCT_INCLUDE(FSceneTextureParameters, SceneTextures)
 	END_SHADER_PARAMETER_STRUCT()
 };
 
@@ -176,8 +176,6 @@ void FDeferredShadingSceneRenderer::RenderRayTracingAmbientOcclusion(
 )
 {
 	FSceneRenderTargets& SceneContext = FSceneRenderTargets::Get(RHICmdList);
-	FSceneTexturesUniformParameters SceneTextures;
-	SetupSceneTextureUniformParameters(SceneContext, FeatureLevel, ESceneTextureSetupMode::All, SceneTextures);
 
 	// Build RTAO parameters
 	FRayTracingAmbientOcclusionRGS::FParameters *PassParameters = GraphBuilder.AllocParameters<FRayTracingAmbientOcclusionRGS::FParameters>();
@@ -189,7 +187,10 @@ void FDeferredShadingSceneRenderer::RenderRayTracingAmbientOcclusion(
 	PassParameters->RWOcclusionMaskUAV = GraphBuilder.CreateUAV(AmbientOcclusionTexture);
 	PassParameters->RWHitDistanceUAV = GraphBuilder.CreateUAV(RayDistanceTexture);
 	PassParameters->ViewUniformBuffer = View.ViewUniformBuffer;
-	PassParameters->SceneTexturesStruct = CreateUniformBufferImmediate(SceneTextures, EUniformBufferUsage::UniformBuffer_SingleDraw);
+
+	FSceneTextureParameters SceneTextures;
+	SetupSceneTextureParameters(GraphBuilder, &SceneTextures);
+	PassParameters->SceneTextures = SceneTextures;
 
 	FRayTracingAmbientOcclusionRGS::FPermutationDomain PermutationVector;
 	PermutationVector.Set<FRayTracingAmbientOcclusionRGS::FEnableTwoSidedGeometryDim>(CVarRayTracingAmbientOcclusionEnableTwoSidedGeometry.GetValueOnRenderThread() != 0);
diff --git a/Engine/Source/Runtime/Renderer/Private/RayTracing/RayTracingGlobalIllumination.cpp b/Engine/Source/Runtime/Renderer/Private/RayTracing/RayTracingGlobalIllumination.cpp
--- a/Engine/Source/Runtime/Renderer/Private/RayTracing/RayTracingGlobalIllumination.cpp
+++ b/Engine/Source/Runtime/Renderer/Private/RayTracing/RayTracingGlobalIllumination.cpp
@@ -20,6 +20,7 @@
 #include "PostProcess/SceneFilterRendering.h"
 #include "Raytracing/RaytracingOptions.h"
 #include "BlueNoise.h"
+#include "SceneTextureParameters.h"
 
 static int32 GRayTracingGlobalIllumination = -1;
 static FAutoConsoleVariableRef CVarRayTracingGlobalIllumination(
@@ -232,13 +233,12 @@ class FGlobalIlluminationRGS : public FGlobalShader
 		SHADER_PARAMETER_RDG_TEXTURE_UAV(RWTexture2D<float4>, RWGlobalIlluminationUAV)
 		SHADER_PARAMETER_RDG_TEXTURE_UAV(RWTexture2D<float>, RWRayDistanceUAV)
 		SHADER_PARAMETER_STRUCT_REF(FViewUniformShaderParameters, ViewUniformBuffer)
-		SHADER_PARAMETER_STRUCT_REF(FSceneTexturesUniformParameters, SceneTexturesStruct)
 		SHADER_PARAMETER_STRUCT_REF(FHaltonIteration, HaltonIteration)
 		SHADER_PARAMETER_STRUCT_REF(FHaltonPrimes, HaltonPrimes)
 		SHADER_PARAMETER_STRUCT_REF(FBlueNoise, BlueNoise)
 		SHADER_PARAMETER_STRUCT_REF(FPathTracingLightData, LightParameters)
 		SHADER_PARAMETER_STRUCT_REF(FSkyLightData, SkyLight)
-
+		SHADER_PARAMETER_STRUCT_INCLUDE(FSceneTextureParameters, SceneTextures)
 		SHADER_PARAMETER_RDG_TEXTURE(Texture2D, SSProfilesTexture)
 		SHADER_PARAMETER_SAMPLER(SamplerState, TransmissionProfilesLinearSampler)
 	END_SHADER_PARAMETER_STRUCT()
@@ -259,7 +259,7 @@ class FRayTracingGlobalIlluminationCompositePS : public FGlobalShader
 		SHADER_PARAMETER_RDG_TEXTURE(Texture2D, GlobalIlluminationTexture)
 		SHADER_PARAMETER_SAMPLER(SamplerState, GlobalIlluminationSampler)
 		SHADER_PARAMETER_STRUCT_REF(FViewUniformShaderParameters, ViewUniformBuffer)
-		SHADER_PARAMETER_STRUCT_REF(FSceneTexturesUniformParameters, SceneTexturesStruct)
+		SHADER_PARAMETER_STRUCT_INCLUDE(FSceneTextureParameters, SceneTextures)
 	END_SHADER_PARAMETER_STRUCT()
 };
 
@@ -277,8 +277,8 @@ class FRayTracingGlobalIlluminationSceneColorCompositePS : public FGlobalShader
 		RENDER_TARGET_BINDING_SLOTS()
 		SHADER_PARAMETER_RDG_TEXTURE(Texture2D, GlobalIlluminationTexture)
 		SHADER_PARAMETER_SAMPLER(SamplerState, GlobalIlluminationSampler)
-		SHADER_PARAMETER_STRUCT_REF(FSceneTexturesUniformParameters, SceneTexturesStruct)
-		END_SHADER_PARAMETER_STRUCT()
+		SHADER_PARAMETER_STRUCT_INCLUDE(FSceneTextureParameters, SceneTextures)
+	END_SHADER_PARAMETER_STRUCT()
 };
 
 class FRayTracingGlobalIlluminationCHS : public FGlobalShader
@@ -423,8 +423,9 @@ void FDeferredShadingSceneRenderer::RenderRayTracingGlobalIllumination(
 	int32 RayTracingGISamplesPerPixel = GRayTracingGlobalIlluminationSamplesPerPixel > -1 ? GRayTracingGlobalIlluminationSamplesPerPixel : View.FinalPostProcessSettings.RayTracingGISamplesPerPixel;
 	RayTracingConfig.RayCountPerPixel = RayTracingGISamplesPerPixel;
 
-	FSceneTexturesUniformParameters SceneTextures;
-	SetupSceneTextureUniformParameters(SceneContext, FeatureLevel, ESceneTextureSetupMode::All, SceneTextures);
+	FSceneTextureParameters SceneTextures;
+	SetupSceneTextureParameters(GraphBuilder, &SceneTextures);
+
 	// Ray generation
 	{
 		uint32 IterationCount = RayTracingGISamplesPerPixel;
@@ -469,11 +470,11 @@ void FDeferredShadingSceneRenderer::RenderRayTracingGlobalIllumination(
 		PassParameters->NextEventEstimationSamples = GRayTracingGlobalIlluminationNextEventEstimationSamples;
 		PassParameters->TLAS = View.RayTracingScene.RayTracingSceneRHI->GetShaderResourceView();
 		PassParameters->ViewUniformBuffer = View.ViewUniformBuffer;
-		PassParameters->SceneTexturesStruct = CreateUniformBufferImmediate(SceneTextures, EUniformBufferUsage::UniformBuffer_SingleDraw);
 		PassParameters->HaltonIteration = CreateUniformBufferImmediate(HaltonIteration, EUniformBufferUsage::UniformBuffer_SingleDraw);
 		PassParameters->HaltonPrimes = CreateUniformBufferImmediate(HaltonPrimes, EUniformBufferUsage::UniformBuffer_SingleDraw);
 		PassParameters->BlueNoise = CreateUniformBufferImmediate(BlueNoise, EUniformBufferUsage::UniformBuffer_SingleDraw);
 		PassParameters->LightParameters = CreateUniformBufferImmediate(LightParameters, EUniformBufferUsage::UniformBuffer_SingleDraw);
+		PassParameters->SceneTextures = SceneTextures;
 		PassParameters->SkyLight = CreateUniformBufferImmediate(SkyLightParameters, EUniformBufferUsage::UniformBuffer_SingleDraw);
 		TRefCountPtr<IPooledRenderTarget> SubsurfaceProfileRT((IPooledRenderTarget*) GetSubsufaceProfileTexture_RT(RHICmdList));
 		if (!SubsurfaceProfileRT)
@@ -485,6 +486,7 @@ void FDeferredShadingSceneRenderer::RenderRayTracingGlobalIllumination(
 		PassParameters->RWGlobalIlluminationUAV = GlobalIlluminationUAV;
 		PassParameters->RWRayDistanceUAV = RayDistanceUAV;
 
+
 		FGlobalIlluminationRGS::FPermutationDomain PermutationVector;
 		PermutationVector.Set<FGlobalIlluminationRGS::FUseAttenuationTermDim>(true);
 		PermutationVector.Set<FGlobalIlluminationRGS::FEnableTwoSidedGeometryDim>(CVarRayTracingGlobalIlluminationEnableTwoSidedGeometry.GetValueOnRenderThread() != 0);
@@ -512,9 +514,6 @@ void FDeferredShadingSceneRenderer::RenderRayTracingGlobalIllumination(
 
 	if (GRayTracingGlobalIlluminationDenoiser != 0)
 	{
-		FSceneTextureParameters SceneTextureParams;
-		SetupSceneTextureParameters(GraphBuilder, &SceneTextureParams);
-
 		const IScreenSpaceDenoiser* DefaultDenoiser = IScreenSpaceDenoiser::GetDefaultDenoiser();
 		const IScreenSpaceDenoiser* DenoiserToUse = GRayTracingGlobalIlluminationDenoiser == 1 ? DefaultDenoiser : GScreenSpaceDenoiser;
 
@@ -532,7 +531,7 @@ void FDeferredShadingSceneRenderer::RenderRayTracingGlobalIllumination(
 				GraphBuilder,
 				View,
 				&View.PrevViewInfo,
-				SceneTextureParams,
+				SceneTextures,
 				DenoiserInputs,
 				RayTracingConfig);
 
@@ -550,9 +549,9 @@ void FDeferredShadingSceneRenderer::RenderRayTracingGlobalIllumination(
 		PassParameters->GlobalIlluminationTexture = ResultTexture;
 		PassParameters->GlobalIlluminationSampler = TStaticSamplerState<SF_Point, AM_Clamp, AM_Clamp, AM_Clamp>::GetRHI();
 		PassParameters->ViewUniformBuffer = View.ViewUniformBuffer;
-		PassParameters->SceneTexturesStruct = CreateUniformBufferImmediate(SceneTextures, EUniformBufferUsage::UniformBuffer_SingleDraw);
 		PassParameters->RenderTargets[0] = FRenderTargetBinding(GraphBuilder.RegisterExternalTexture(GlobalIlluminationRT), ERenderTargetLoadAction::ELoad, ERenderTargetStoreAction::EStore);
 		PassParameters->RenderTargets[1] = FRenderTargetBinding(GraphBuilder.RegisterExternalTexture(AmbientOcclusionRT), ERenderTargetLoadAction::ELoad, ERenderTargetStoreAction::EStore);
+		PassParameters->SceneTextures = SceneTextures;
 
 		GraphBuilder.AddPass(
 			RDG_EVENT_NAME("GlobalIlluminationComposite"),
@@ -604,16 +603,15 @@ void FDeferredShadingSceneRenderer::CompositeGlobalIllumination(
 {
 	FSceneRenderTargets& SceneContext = FSceneRenderTargets::Get(RHICmdList);
 
-	FSceneTexturesUniformParameters SceneTextures;
-	SetupSceneTextureUniformParameters(SceneContext, FeatureLevel, ESceneTextureSetupMode::All, SceneTextures);
-
 	FRDGBuilder GraphBuilder(RHICmdList);
+	FSceneTextureParameters SceneTextures;
+	SetupSceneTextureParameters(GraphBuilder, &SceneTextures);
 
 	FRayTracingGlobalIlluminationSceneColorCompositePS::FParameters *PassParameters = GraphBuilder.AllocParameters<FRayTracingGlobalIlluminationSceneColorCompositePS::FParameters>();
 	PassParameters->GlobalIlluminationTexture = GraphBuilder.RegisterExternalTexture(GlobalIlluminationRT);
 	PassParameters->GlobalIlluminationSampler = TStaticSamplerState<SF_Point, AM_Clamp, AM_Clamp, AM_Clamp>::GetRHI();
 	PassParameters->RenderTargets[0] = FRenderTargetBinding(GraphBuilder.RegisterExternalTexture(SceneContext.GetSceneColor()), ERenderTargetLoadAction::ELoad, ERenderTargetStoreAction::EStore);
-	PassParameters->SceneTexturesStruct = CreateUniformBufferImmediate(SceneTextures, EUniformBufferUsage::UniformBuffer_SingleDraw);
+	PassParameters->SceneTextures = SceneTextures;
 
 	GraphBuilder.AddPass(
 		RDG_EVENT_NAME("GlobalIlluminationComposite"),
diff --git a/Engine/Source/Runtime/Renderer/Private/RayTracing/RayTracingReflections.cpp b/Engine/Source/Runtime/Renderer/Private/RayTracing/RayTracingReflections.cpp
--- a/Engine/Source/Runtime/Renderer/Private/RayTracing/RayTracingReflections.cpp
+++ b/Engine/Source/Runtime/Renderer/Private/RayTracing/RayTracingReflections.cpp
@@ -20,6 +20,7 @@
 #include "RayTracing/RayTracingLighting.h"
 #include "RenderGraph.h"
 #include "RayTracing/RayTracingLighting.h"
+#include "SceneTextureParameters.h"
 
 static float GRayTracingReflectionsMaxRoughness = -1;
 static FAutoConsoleVariableRef CVarRayTracingReflectionsMaxRoughness(
@@ -165,8 +166,10 @@ class FRayTracingReflectionsRGS : public FGlobalShader
 		SHADER_PARAMETER_SRV(RaytracingAccelerationStructure, TLAS)
 		SHADER_PARAMETER_SRV(StructuredBuffer<FRTLightingData>, LightDataBuffer)
 
+		SHADER_PARAMETER_STRUCT_INCLUDE(FSceneTextureParameters, SceneTextures)
+		SHADER_PARAMETER_STRUCT_INCLUDE(FSceneTextureSamplerParameters, SceneTextureSamplers)
+
 		SHADER_PARAMETER_STRUCT_REF(FViewUniformShaderParameters, ViewUniformBuffer)
-		SHADER_PARAMETER_STRUCT_REF(FSceneTexturesUniformParameters, SceneTexturesStruct)
 		SHADER_PARAMETER_STRUCT_REF(FRaytracingLightDataPacked, LightDataPacked)
 		SHADER_PARAMETER_STRUCT_REF(FReflectionUniformParameters, ReflectionStruct)
 		SHADER_PARAMETER_STRUCT_REF(FFogUniformParameters, FogUniformParameters)
@@ -209,6 +212,7 @@ void FDeferredShadingSceneRenderer::PrepareRayTracingReflections(const FViewInfo
 
 void FDeferredShadingSceneRenderer::RenderRayTracingReflections(
 	FRDGBuilder& GraphBuilder,
+	const FSceneTextureParameters& SceneTextures,
 	const FViewInfo& View,
 	int32 SamplePerPixel,
 	int32 HeightFog,
@@ -288,7 +292,9 @@ void FDeferredShadingSceneRenderer::RenderRayTracingReflections(
 	CommonParameters.LightDataPacked = CreateLightDataPackedUniformBuffer(Scene->Lights, View, EUniformBufferUsage::UniformBuffer_SingleFrame, LightingDataBuffer);
 	CommonParameters.LightDataBuffer = RHICreateShaderResourceView(LightingDataBuffer);
 
-	CommonParameters.SceneTexturesStruct = CreateSceneTextureUniformBuffer( SceneContext, FeatureLevel, ESceneTextureSetupMode::All, EUniformBufferUsage::UniformBuffer_SingleFrame);
+	CommonParameters.SceneTextures = SceneTextures;
+	SetupSceneTextureSamplers(&CommonParameters.SceneTextureSamplers);
+
 	CommonParameters.ReflectionStruct = CreateReflectionUniformBuffer(View, EUniformBufferUsage::UniformBuffer_SingleFrame);
 	CommonParameters.FogUniformParameters = CreateFogUniformBuffer(View, EUniformBufferUsage::UniformBuffer_SingleFrame);
 	CommonParameters.ColorOutput = GraphBuilder.CreateUAV(OutDenoiserInputs->Color);
diff --git a/Engine/Source/Runtime/Renderer/Private/RayTracing/RayTracingTranslucency.cpp b/Engine/Source/Runtime/Renderer/Private/RayTracing/RayTracingTranslucency.cpp
--- a/Engine/Source/Runtime/Renderer/Private/RayTracing/RayTracingTranslucency.cpp
+++ b/Engine/Source/Runtime/Renderer/Private/RayTracing/RayTracingTranslucency.cpp
@@ -3,6 +3,7 @@
 #include "RendererPrivate.h"
 #include "GlobalShader.h"
 #include "DeferredShadingRenderer.h"
+#include "SceneTextureParameters.h"
 
 #if RHI_RAYTRACING
 
@@ -144,11 +145,13 @@ class FRayTracingTranslucencyRGS : public FGlobalShader
 		SHADER_PARAMETER_SRV(StructuredBuffer<FRTLightingData>, LightDataBuffer)
 
 		SHADER_PARAMETER_STRUCT_REF(FViewUniformShaderParameters, ViewUniformBuffer)
-		SHADER_PARAMETER_STRUCT_REF(FSceneTexturesUniformParameters, SceneTexturesStruct)
 		SHADER_PARAMETER_STRUCT_REF(FRaytracingLightDataPacked, LightDataPacked)
 		SHADER_PARAMETER_STRUCT_REF(FReflectionUniformParameters, ReflectionStruct)
 		SHADER_PARAMETER_STRUCT_REF(FFogUniformParameters, FogUniformParameters)
 
+		SHADER_PARAMETER_STRUCT_INCLUDE(FSceneTextureParameters, SceneTextures)
+		SHADER_PARAMETER_STRUCT_INCLUDE(FSceneTextureSamplerParameters, SceneTextureSamplers)
+
 		SHADER_PARAMETER_RDG_TEXTURE_UAV(RWTexture2D<float4>, ColorOutput)
 		SHADER_PARAMETER_RDG_TEXTURE_UAV(RWTexture2D<float>, RayHitDistanceOutput)
 	END_SHADER_PARAMETER_STRUCT()
@@ -341,6 +344,11 @@ void FDeferredShadingSceneRenderer::RenderRayTracingTranslucencyView(
 {
 	FSceneRenderTargets& SceneContext = FSceneRenderTargets::Get(GraphBuilder.RHICmdList);
 
+	FSceneTextureParameters SceneTextures;
+	SetupSceneTextureParameters(GraphBuilder, &SceneTextures);
+	FSceneTextureSamplerParameters SceneTextureSamplers;
+	SetupSceneTextureSamplers(&SceneTextureSamplers);
+
 	int32 UpscaleFactor = int32(1.0f / ResolutionFraction);
 	ensure(ResolutionFraction == 1.0 / UpscaleFactor);
 	ensureMsgf(FComputeShaderUtils::kGolden2DGroupSize % UpscaleFactor == 0, TEXT("Translucency ray tracing will have uv misalignement."));
@@ -381,7 +389,9 @@ void FDeferredShadingSceneRenderer::RenderRayTracingTranslucencyView(
 	PassParameters->LightDataPacked = CreateLightDataPackedUniformBuffer(Scene->Lights, View, EUniformBufferUsage::UniformBuffer_SingleFrame, LightingDataBuffer);
 	PassParameters->LightDataBuffer = RHICreateShaderResourceView(LightingDataBuffer);
 
-	PassParameters->SceneTexturesStruct = CreateSceneTextureUniformBuffer(SceneContext, FeatureLevel, ESceneTextureSetupMode::All, EUniformBufferUsage::UniformBuffer_SingleFrame);	
+	PassParameters->SceneTextures = SceneTextures;
+	PassParameters->SceneTextureSamplers = SceneTextureSamplers;
+
 	PassParameters->ReflectionStruct = CreateReflectionUniformBuffer(View, EUniformBufferUsage::UniformBuffer_SingleFrame);
 	PassParameters->FogUniformParameters = CreateFogUniformBuffer(View, EUniformBufferUsage::UniformBuffer_SingleFrame);
 
diff --git a/Engine/Source/Runtime/Renderer/Private/RayTracing/RaytracingSkylight.cpp b/Engine/Source/Runtime/Renderer/Private/RayTracing/RaytracingSkylight.cpp
--- a/Engine/Source/Runtime/Renderer/Private/RayTracing/RaytracingSkylight.cpp
+++ b/Engine/Source/Runtime/Renderer/Private/RayTracing/RaytracingSkylight.cpp
@@ -5,6 +5,7 @@
 =============================================================================*/
 
 #include "DeferredShadingRenderer.h"
+#include "SceneTextureParameters.h"
 
 static int32 GRayTracingSkyLight = 0;
 
@@ -177,12 +178,12 @@ class FRayTracingSkyLightRGS : public FGlobalShader
 		SHADER_PARAMETER_RDG_TEXTURE_UAV(RWTexture2D<float2>, RWRayDistanceUAV)
 
 		SHADER_PARAMETER_STRUCT_REF(FViewUniformShaderParameters, ViewUniformBuffer)
-		SHADER_PARAMETER_STRUCT_REF(FSceneTexturesUniformParameters, SceneTexturesStruct)
 		SHADER_PARAMETER_STRUCT_REF(FSkyLightData, SkyLightData)
 		SHADER_PARAMETER_STRUCT_REF(FHaltonIteration, HaltonIteration)
 		SHADER_PARAMETER_STRUCT_REF(FHaltonPrimes, HaltonPrimes)
 		SHADER_PARAMETER_STRUCT_REF(FBlueNoise, BlueNoise)
 
+		SHADER_PARAMETER_STRUCT_INCLUDE(FSceneTextureParameters, SceneTextures)
 		SHADER_PARAMETER_RDG_TEXTURE(Texture2D, SSProfilesTexture)
 		SHADER_PARAMETER_SAMPLER(SamplerState, TransmissionProfilesLinearSampler)
 	END_SHADER_PARAMETER_STRUCT()
@@ -813,8 +814,6 @@ void FDeferredShadingSceneRenderer::RenderRayTracingSkyLight(
 	}
 
 	FRDGBuilder GraphBuilder(RHICmdList);
-	FSceneTextureParameters SceneTextures;
-	SetupSceneTextureParameters(GraphBuilder, &SceneTextures);
 
 	// Define RDG targets
 	FRDGTextureRef SkyLightTexture;
@@ -847,8 +846,8 @@ void FDeferredShadingSceneRenderer::RenderRayTracingSkyLight(
 	SkyLightData.MaxRayDistance = GRayTracingSkyLightMaxRayDistance;
 	SkyLightData.SamplingStopLevel = GRayTracingSkyLightSamplingStopLevel;
 
-	FSceneTexturesUniformParameters SceneUniformTextures;
-	SetupSceneTextureUniformParameters(SceneContext, FeatureLevel, ESceneTextureSetupMode::All, SceneUniformTextures);
+	FSceneTextureParameters SceneTextures;
+	SetupSceneTextureParameters(GraphBuilder, &SceneTextures);
 
 	TRefCountPtr<IPooledRenderTarget> SubsurfaceProfileRT((IPooledRenderTarget*)GetSubsufaceProfileTexture_RT(RHICmdList));
 	if (!SubsurfaceProfileRT)
@@ -860,12 +859,12 @@ void FDeferredShadingSceneRenderer::RenderRayTracingSkyLight(
 	FRayTracingSkyLightRGS::FParameters *PassParameters = GraphBuilder.AllocParameters<FRayTracingSkyLightRGS::FParameters>();
 	PassParameters->RWOcclusionMaskUAV = GraphBuilder.CreateUAV(SkyLightTexture);
 	PassParameters->RWRayDistanceUAV = GraphBuilder.CreateUAV(RayDistanceTexture);
-	PassParameters->SceneTexturesStruct = CreateUniformBufferImmediate(SceneUniformTextures, EUniformBufferUsage::UniformBuffer_SingleDraw);
 	PassParameters->SkyLightData = CreateUniformBufferImmediate(SkyLightData, EUniformBufferUsage::UniformBuffer_SingleDraw);
 	PassParameters->HaltonPrimes = CreateUniformBufferImmediate(HaltonPrimes, EUniformBufferUsage::UniformBuffer_SingleDraw);
 	PassParameters->BlueNoise = CreateUniformBufferImmediate(BlueNoise, EUniformBufferUsage::UniformBuffer_SingleDraw);
 	PassParameters->SSProfilesTexture = GraphBuilder.RegisterExternalTexture(SubsurfaceProfileRT);
 	PassParameters->TransmissionProfilesLinearSampler = TStaticSamplerState<SF_Bilinear, AM_Clamp, AM_Clamp, AM_Clamp>::GetRHI();
+	PassParameters->SceneTextures = SceneTextures;
 
 	for (int32 ViewIndex = 0; ViewIndex < Views.Num(); ++ViewIndex)
 	{
@@ -975,8 +974,9 @@ class FCompositeSkyLightPS : public FGlobalShader
 		SHADER_PARAMETER_RDG_TEXTURE(Texture2D, SkyLightTexture)
 		SHADER_PARAMETER_SAMPLER(SamplerState, SkyLightTextureSampler)
 		SHADER_PARAMETER_STRUCT_REF(FViewUniformShaderParameters, ViewUniformBuffer)
-		SHADER_PARAMETER_STRUCT_REF(FSceneTexturesUniformParameters, SceneTexturesStruct)
-		END_SHADER_PARAMETER_STRUCT()
+
+		SHADER_PARAMETER_STRUCT_INCLUDE(FSceneTextureParameters, SceneTextures)
+	END_SHADER_PARAMETER_STRUCT()
 };
 
 IMPLEMENT_GLOBAL_SHADER(FCompositeSkyLightPS, "/Engine/Private/RayTracing/CompositeSkyLightPS.usf", "CompositeSkyLightPS", SF_Pixel);
@@ -993,19 +993,21 @@ void FDeferredShadingSceneRenderer::CompositeRayTracingSkyLight(
 {
 	check(SkyLightRT);
 	FSceneRenderTargets& SceneContext = FSceneRenderTargets::Get(RHICmdList);
-	FSceneTexturesUniformParameters SceneTextures;
-	SetupSceneTextureUniformParameters(SceneContext, FeatureLevel, ESceneTextureSetupMode::All, SceneTextures);
 
 	for (int32 ViewIndex = 0; ViewIndex < Views.Num(); ViewIndex++)
 	{
 		const FViewInfo& View = Views[ViewIndex];
 		FRDGBuilder GraphBuilder(RHICmdList);
+
+		FSceneTextureParameters SceneTextures;
+		SetupSceneTextureParameters(GraphBuilder, &SceneTextures);
+
 		FCompositeSkyLightPS::FParameters *PassParameters = GraphBuilder.AllocParameters<FCompositeSkyLightPS::FParameters>();
 		PassParameters->SkyLightTexture = GraphBuilder.RegisterExternalTexture(SkyLightRT);
 		PassParameters->SkyLightTextureSampler = TStaticSamplerState<SF_Point, AM_Clamp, AM_Clamp, AM_Clamp>::GetRHI();
 		PassParameters->ViewUniformBuffer = Views[ViewIndex].ViewUniformBuffer;
-		PassParameters->SceneTexturesStruct = CreateUniformBufferImmediate(SceneTextures, EUniformBufferUsage::UniformBuffer_SingleDraw);
 		PassParameters->RenderTargets[0] = FRenderTargetBinding(GraphBuilder.RegisterExternalTexture(SceneContext.GetSceneColor()), ERenderTargetLoadAction::ELoad, ERenderTargetStoreAction::EStore);
+		PassParameters->SceneTextures = SceneTextures;
 
 		// dxr_todo: Unify with RTGI compositing workflow
 		GraphBuilder.AddPass(
diff --git a/Engine/Source/Runtime/Renderer/Private/ReflectionEnvironment.cpp b/Engine/Source/Runtime/Renderer/Private/ReflectionEnvironment.cpp
--- a/Engine/Source/Runtime/Renderer/Private/ReflectionEnvironment.cpp
+++ b/Engine/Source/Runtime/Renderer/Private/ReflectionEnvironment.cpp
@@ -682,6 +682,7 @@ void FDeferredShadingSceneRenderer::RenderDeferredReflectionsAndSkyLighting(FRHI
 
 				RenderRayTracingReflections(
 					GraphBuilder,
+					SceneTextures,
 					View,
 					RayTracingConfig.RayCountPerPixel, GRayTracingReflectionsHeightFog, RayTracingConfig.ResolutionFraction,
 					&DenoiserInputs);
