Fix incorrect TSparseArray iteration of RuntimeVirtualTextures causing a crash
#rnx

diff --git a/Engine/Source/Runtime/Renderer/Private/RendererScene.cpp b/Engine/Source/Runtime/Renderer/Private/RendererScene.cpp
--- a/Engine/Source/Runtime/Renderer/Private/RendererScene.cpp
+++ b/Engine/Source/Runtime/Renderer/Private/RendererScene.cpp
@@ -2474,12 +2474,12 @@ void FScene::AddRuntimeVirtualTexture_RenderThread(FRuntimeVirtualTextureScenePr
 
 void FScene::UpdateRuntimeVirtualTexture_RenderThread(FRuntimeVirtualTextureSceneProxy* SceneProxy, FRuntimeVirtualTextureSceneProxy* SceneProxyToReplace)
 {
-	for (int32 SceneIndex = 0; SceneIndex < RuntimeVirtualTextures.Num(); ++SceneIndex)
+	for (TSparseArray<FRuntimeVirtualTextureSceneProxy*>::TIterator It(RuntimeVirtualTextures); It; ++It)
 	{
-		if (RuntimeVirtualTextures.IsValidIndex(SceneIndex) && RuntimeVirtualTextures[SceneIndex] == SceneProxyToReplace)
+		if (*It == SceneProxyToReplace)
 		{
-			SceneProxy->SceneIndex = SceneIndex;
-			RuntimeVirtualTextures[SceneIndex] = SceneProxy;
+			SceneProxy->SceneIndex = It.GetIndex();
+			*It = SceneProxy;
 			delete SceneProxyToReplace;
 			return;
 		}
@@ -2523,16 +2523,17 @@ uint32 FScene::GetRuntimeVirtualTextureSceneIndex(uint32 ProducerId)
 uint32 FScene::GetRuntimeVirtualTextureMask(FPrimitiveSceneProxy const* Proxy)
 {
 	uint32 Mask = 0;
-	for (int32 i = 0; i < FPrimitiveVirtualTextureFlags::RuntimeVirtualTexture_BitCount && i < RuntimeVirtualTextures.Num(); ++i)
+	for (TSparseArray<FRuntimeVirtualTextureSceneProxy*>::TConstIterator It(RuntimeVirtualTextures); It; ++It)
 	{
-		if (RuntimeVirtualTextures.IsValidIndex(i))
+		int32 SceneIndex = It.GetIndex();
+		if (SceneIndex < FPrimitiveVirtualTextureFlags::RuntimeVirtualTexture_BitCount)
 		{
-			URuntimeVirtualTexture* SceneVirtualTexture = RuntimeVirtualTextures[i]->VirtualTexture;
+			URuntimeVirtualTexture* SceneVirtualTexture = (*It)->VirtualTexture;
 			for (URuntimeVirtualTexture* PrimitiveVirtualTexture : Proxy->RuntimeVirtualTextures)
 			{
 				if (SceneVirtualTexture == PrimitiveVirtualTexture)
 				{
-					Mask |= 1 << i;
+					Mask |= 1 << SceneIndex;
 				}
 			}
 		}
