Fix to support LinearNoToneCurve shader permutation
(merge error from Main)

diff --git a/Engine/Source/Runtime/Renderer/Private/PostProcess/PostProcessTonemap.cpp b/Engine/Source/Runtime/Renderer/Private/PostProcess/PostProcessTonemap.cpp
--- a/Engine/Source/Runtime/Renderer/Private/PostProcess/PostProcessTonemap.cpp
+++ b/Engine/Source/Runtime/Renderer/Private/PostProcess/PostProcessTonemap.cpp
@@ -1005,7 +1005,14 @@ void FRCPassPostProcessTonemap::Process(FRenderingCompositePassContext& Context)
 			DesktopPermutationVector.Set<TonemapperPermutation::FTonemapperColorFringeDim>(View.FinalPostProcessSettings.SceneFringeIntensity > 0.01f);
 		}
 
-		DesktopPermutationVector.Set<TonemapperPermutation::FTonemapperOutputDeviceDim>(GetOutputDeviceValue());
+		if (ViewFamily.SceneCaptureSource == SCS_FinalColorHDR)
+		{
+			DesktopPermutationVector.Set<TonemapperPermutation::FTonemapperOutputDeviceDim>(FTonemapperOutputDevice::LinearNoToneCurve);
+		}
+		else
+		{
+			DesktopPermutationVector.Set<TonemapperPermutation::FTonemapperOutputDeviceDim>(GetOutputDeviceValue());
+		}
 		DesktopPermutationVector = TonemapperPermutation::RemapPermutation(DesktopPermutationVector);
 	}
 
