Fixes a bug where SSAO was being applied in reflection environement pixel shader even if texture was invalid.

diff --git a/Engine/Source/Runtime/Renderer/Private/ReflectionEnvironment.cpp b/Engine/Source/Runtime/Renderer/Private/ReflectionEnvironment.cpp
--- a/Engine/Source/Runtime/Renderer/Private/ReflectionEnvironment.cpp
+++ b/Engine/Source/Runtime/Renderer/Private/ReflectionEnvironment.cpp
@@ -642,7 +642,7 @@ void FDeferredShadingSceneRenderer::RenderDeferredReflectionsAndSkyLighting(FRHI
 	FRDGBuilder GraphBuilder(RHICmdList);
 
 	FRDGTextureRef SceneColorTexture = GraphBuilder.RegisterExternalTexture(SceneContext.GetSceneColor());
-	FRDGTextureRef AmbientOcclusionTexture = GraphBuilder.RegisterExternalTexture(SceneContext.ScreenSpaceAO);
+	FRDGTextureRef AmbientOcclusionTexture = GraphBuilder.RegisterExternalTexture(SceneContext.bScreenSpaceAOIsValid ? SceneContext.ScreenSpaceAO : GSystemTextures.WhiteDummy);
 	FRDGTextureRef DynamicBentNormalAOTexture = GraphBuilder.RegisterExternalTexture(DynamicBentNormalAO ? DynamicBentNormalAO : GSystemTextures.WhiteDummy);
 
 	FSceneTextureParameters SceneTextures;
