Fixes RDG warning message in AddPass_ClearUAV()

diff --git a/Engine/Source/Runtime/Renderer/Private/LightGridInjection.cpp b/Engine/Source/Runtime/Renderer/Private/LightGridInjection.cpp
--- a/Engine/Source/Runtime/Renderer/Private/LightGridInjection.cpp
+++ b/Engine/Source/Runtime/Renderer/Private/LightGridInjection.cpp
@@ -95,6 +95,7 @@ void AddPass_ClearUAV(
 		ERenderGraphPassFlags::Compute,
 		[&Parameters, BufferUAV, Value](FRHICommandList& RHICmdList)
 		{
+			BufferUAV->MarkResourceAsUsed();
 			ClearUAV( RHICmdList, BufferUAV->GetRHI(), BufferUAV->Desc.Buffer->Desc.GetTotalNumBytes(), Value );
 		} );
 }
