Silents a RDG warning when using reflection denoiser because not using the imaginary depth buffer

diff --git a/Engine/Source/Runtime/Renderer/Private/ScreenSpaceDenoise.cpp b/Engine/Source/Runtime/Renderer/Private/ScreenSpaceDenoise.cpp
--- a/Engine/Source/Runtime/Renderer/Private/ScreenSpaceDenoise.cpp
+++ b/Engine/Source/Runtime/Renderer/Private/ScreenSpaceDenoise.cpp
@@ -135,10 +135,17 @@ namespace
 /** Different signals to denoise. */
 enum class ESignalProcessing
 {
+	// Denoise a shadow mask.
 	MonochromaticPenumbra,
+
+	// Denoise first bounce specular.
 	Reflections,
+
+	// Denoise ambient occlusion.
 	AmbientOcclusion,
-	DiffuseIndirect,
+
+	// Denoise first bounce diffuse and ambient occlusion.
+	DiffuseAndAmbientOcclusion,
 
 	MAX,
 };
@@ -158,7 +165,7 @@ static bool UsesConstantPixelDensityPassLayout(ESignalProcessing SignalProcessin
 		SignalProcessing == ESignalProcessing::MonochromaticPenumbra ||
 		SignalProcessing == ESignalProcessing::Reflections ||
 		SignalProcessing == ESignalProcessing::AmbientOcclusion ||
-		SignalProcessing == ESignalProcessing::DiffuseIndirect);
+		SignalProcessing == ESignalProcessing::DiffuseAndAmbientOcclusion);
 }
 
 /** Returns whether a signal processing uses an injestion pass. */
@@ -172,7 +179,7 @@ static bool SignalUsesPreConvolution(ESignalProcessing SignalProcessing)
 {
 	return
 		SignalProcessing == ESignalProcessing::MonochromaticPenumbra ||
-		SignalProcessing == ESignalProcessing::DiffuseIndirect;
+		SignalProcessing == ESignalProcessing::DiffuseAndAmbientOcclusion;
 }
 
 /** Returns whether a signal processing uses a history rejection pre convolution pass. */
@@ -200,7 +207,7 @@ static int32 SignalMaxBatchSize(ESignalProcessing SignalProcessing)
 	else if (
 		SignalProcessing == ESignalProcessing::Reflections ||
 		SignalProcessing == ESignalProcessing::AmbientOcclusion ||
-		SignalProcessing == ESignalProcessing::DiffuseIndirect)
+		SignalProcessing == ESignalProcessing::DiffuseAndAmbientOcclusion)
 	{
 		return 1;
 	}
@@ -214,7 +221,7 @@ static bool SignalSupport1SPP(ESignalProcessing SignalProcessing)
 	return (
 		SignalProcessing == ESignalProcessing::MonochromaticPenumbra ||
 		SignalProcessing == ESignalProcessing::Reflections ||
-		SignalProcessing == ESignalProcessing::DiffuseIndirect);
+		SignalProcessing == ESignalProcessing::DiffuseAndAmbientOcclusion);
 }
 
 /** Returns whether a signal can denoise multi sample per pixel. */
@@ -224,7 +231,7 @@ static bool SignalSupportMultiSPP(ESignalProcessing SignalProcessing)
 		SignalProcessing == ESignalProcessing::MonochromaticPenumbra ||
 		SignalProcessing == ESignalProcessing::Reflections ||
 		SignalProcessing == ESignalProcessing::AmbientOcclusion ||
-		SignalProcessing == ESignalProcessing::DiffuseIndirect);
+		SignalProcessing == ESignalProcessing::DiffuseAndAmbientOcclusion);
 }
 
 
@@ -440,7 +447,7 @@ bool ShouldCompileSignalPipeline(ESignalProcessing SignalProcessing, EShaderPlat
 	else if (
 		SignalProcessing == ESignalProcessing::MonochromaticPenumbra ||
 		SignalProcessing == ESignalProcessing::AmbientOcclusion ||
-		SignalProcessing == ESignalProcessing::DiffuseIndirect)
+		SignalProcessing == ESignalProcessing::DiffuseAndAmbientOcclusion)
 	{
 		// Only for ray tracing denoising.
 		return Platform == SP_PCD3D_SM5;
@@ -878,7 +885,7 @@ static void DenoiseSignalAtConstantPixelDensity(
 			ReconstructionTextureCount = HistoryTextureCountPerSignal = 1;
 			bHasReconstructionLayoutDifferentFromHistory = false;
 		}
-		else if (Settings.SignalProcessing == ESignalProcessing::DiffuseIndirect)
+		else if (Settings.SignalProcessing == ESignalProcessing::DiffuseAndAmbientOcclusion)
 		{
 			ReconstructionDescs[0].Format = PF_FloatRGBA;
 			ReconstructionDescs[1].Format = PF_R16F;
@@ -917,7 +924,7 @@ static void DenoiseSignalAtConstantPixelDensity(
 	}
 
 	#if RHI_RAYTRACING
-	if (Settings.SignalProcessing == ESignalProcessing::DiffuseIndirect)
+	if (Settings.SignalProcessing == ESignalProcessing::DiffuseAndAmbientOcclusion)
 	{
 		uint32 IterationCount = Settings.MaxInputSPP;
 		uint32 SequenceCount = 1;
@@ -1254,7 +1261,7 @@ static void DenoiseSignalAtConstantPixelDensity(
 			// Requires the normal that are in GBuffer A.
 			if (Settings.SignalProcessing == ESignalProcessing::Reflections ||
 				Settings.SignalProcessing == ESignalProcessing::AmbientOcclusion ||
-				Settings.SignalProcessing == ESignalProcessing::DiffuseIndirect)
+				Settings.SignalProcessing == ESignalProcessing::DiffuseAndAmbientOcclusion)
 			{
 				GraphBuilder.QueueTextureExtraction(SceneTextures.SceneGBufferA, &View.ViewState->PrevFrameViewInfo.GBufferA);
 			}
@@ -1444,6 +1451,10 @@ public:
 	{
 		RDG_GPU_STAT_SCOPE(GraphBuilder, ReflectionsDenoiser);
 
+		// Imaginary depth is only used for Nvidia denoiser.
+		// TODO: permutation to not generate it?
+		GraphBuilder.RemoveUnusedTextureWarning(ReflectionInputs.RayImaginaryDepth);
+
 		FSSDSignalTextures InputSignal;
 		InputSignal.Textures[0] = ReflectionInputs.Color;
 		InputSignal.Textures[1] = ReflectionInputs.RayHitDistance;
@@ -1530,7 +1541,7 @@ public:
 		InputSignal.Textures[1] = Inputs.RayHitDistance;
 
 		FSSDConstantPixelDensitySettings Settings;
-		Settings.SignalProcessing = ESignalProcessing::DiffuseIndirect;
+		Settings.SignalProcessing = ESignalProcessing::DiffuseAndAmbientOcclusion;
 		Settings.InputResolutionFraction = Config.ResolutionFraction;
 		Settings.ReconstructionSamples = CVarGIReconstructionSampleCount.GetValueOnRenderThread();
 		Settings.PreConvolutionCount = CVarGIPreConvolutionCount.GetValueOnRenderThread();
@@ -1572,7 +1583,7 @@ public:
 		InputSignal.Textures[1] = Inputs.RayHitDistance;
 
 		FSSDConstantPixelDensitySettings Settings;
-		Settings.SignalProcessing = ESignalProcessing::DiffuseIndirect;
+		Settings.SignalProcessing = ESignalProcessing::DiffuseAndAmbientOcclusion;
 		Settings.InputResolutionFraction = Config.ResolutionFraction;
 		Settings.ReconstructionSamples = CVarGIReconstructionSampleCount.GetValueOnRenderThread();
 		Settings.PreConvolutionCount = CVarGIPreConvolutionCount.GetValueOnRenderThread();
