Fix scissor rect ensures

diff --git a/Engine/Source/Runtime/D3D12RHI/Private/D3D12StateCachePrivate.h b/Engine/Source/Runtime/D3D12RHI/Private/D3D12StateCachePrivate.h
--- a/Engine/Source/Runtime/D3D12RHI/Private/D3D12StateCachePrivate.h
+++ b/Engine/Source/Runtime/D3D12RHI/Private/D3D12StateCachePrivate.h
@@ -829,7 +829,9 @@ public:
 	{
 		if (LastComputePipelineType != PipelineType)
 		{
-			ClearState();
+			PipelineState.Common.bNeedSetPSO = true;
+			PipelineState.Compute.bNeedSetRootSignature = true;
+
 			LastComputePipelineType = PipelineType;
 		}
 	}
diff --git a/Engine/Source/Runtime/Renderer/Private/PostProcess/SceneRenderTargets.cpp b/Engine/Source/Runtime/Renderer/Private/PostProcess/SceneRenderTargets.cpp
--- a/Engine/Source/Runtime/Renderer/Private/PostProcess/SceneRenderTargets.cpp
+++ b/Engine/Source/Runtime/Renderer/Private/PostProcess/SceneRenderTargets.cpp
@@ -1881,9 +1881,11 @@ void FSceneRenderTargets::ResolveSceneDepthTexture(FRHICommandList& RHICmdList,
 
 	SCOPED_DRAW_EVENT(RHICmdList, ResolveSceneDepthTexture);
 
+	FResolveParams ResolveParams;
 	if (ResolveRect.IsValid())
 	{
-		RHICmdList.SetScissorRect(true, ResolveRect.X1, ResolveRect.Y1, ResolveRect.X2, ResolveRect.Y2);
+		ResolveParams.DestRect = ResolveRect;
+		ResolveParams.Rect = ResolveRect;
 	}
 
 	FSceneRenderTargets& SceneContext = FSceneRenderTargets::Get(RHICmdList);
@@ -1892,16 +1894,11 @@ void FSceneRenderTargets::ResolveSceneDepthTexture(FRHICommandList& RHICmdList,
 	const EShaderPlatform CurrentShaderPlatform = GShaderPlatformForFeatureLevel[SceneContext.GetCurrentFeatureLevel()];
 	if ((CurrentNumSamples <= 1 || !RHISupportsSeparateMSAAAndResolveTextures(CurrentShaderPlatform)) || !GAllowCustomMSAAResolves)
 	{
-		RHICmdList.CopyToResolveTarget(GetSceneDepthSurface(), GetSceneDepthTexture(), FResolveParams());
+		RHICmdList.CopyToResolveTarget(GetSceneDepthSurface(), GetSceneDepthTexture(), ResolveParams);
 	}
 	else
 	{
-		ResolveDepthTexture(RHICmdList, GetSceneDepthSurface(), GetSceneDepthTexture(), FResolveParams());
-	}
-
-	if (ResolveRect.IsValid())
-	{
-		RHICmdList.SetScissorRect(false, 0, 0, 0, 0);
+		ResolveDepthTexture(RHICmdList, GetSceneDepthSurface(), GetSceneDepthTexture(), ResolveParams);
 	}
 }
 
