Fixes shader parameter failure in TAA on camera cut when nothing else needs motion vector

diff --git a/Engine/Source/Runtime/Renderer/Private/PostProcess/PostProcessing.cpp b/Engine/Source/Runtime/Renderer/Private/PostProcess/PostProcessing.cpp
--- a/Engine/Source/Runtime/Renderer/Private/PostProcess/PostProcessing.cpp
+++ b/Engine/Source/Runtime/Renderer/Private/PostProcess/PostProcessing.cpp
@@ -1393,6 +1393,11 @@ void FPostProcessing::Process(FRHICommandListImmediate& RHICmdList, const FViewI
 				VelocityInput = Context.Graph.RegisterPass(new(FMemStack::Get()) FRCPassPostProcessInput(VelocityRT));
 				PreFlattenVelocity = VelocityInput;
 			}
+			else if (View.AntiAliasingMethod == AAM_TemporalAA && View.bCameraCut)
+			{
+				// Velocity pass is not being rendered on camera cut.
+				VelocityInput = Context.Graph.RegisterPass(new(FMemStack::Get()) FRCPassPostProcessInput(GSystemTextures.BlackDummy));
+			}
 
 			Context.FinalOutput = AddPostProcessMaterialChain(Context, BL_BeforeTranslucency, SeparateTranslucency);
 			
@@ -1421,6 +1426,10 @@ void FPostProcessing::Process(FRHICommandListImmediate& RHICmdList, const FViewI
 					{
 						SceneTextures.SceneVelocityBuffer = SceneVelocityBuffer;
 					}
+					else
+					{
+						UE_DEBUG_BREAK();
+					}
 
 					FRDGTextureRef NewSceneColor = SceneColor;
 
