Removes the dependency of the velocity buffer in temporal reprojections of TAA and denoiser on camera cut.

diff --git a/Engine/Source/Runtime/Renderer/Private/PostProcess/PostProcessTemporalAA.cpp b/Engine/Source/Runtime/Renderer/Private/PostProcess/PostProcessTemporalAA.cpp
--- a/Engine/Source/Runtime/Renderer/Private/PostProcess/PostProcessTemporalAA.cpp
+++ b/Engine/Source/Runtime/Renderer/Private/PostProcess/PostProcessTemporalAA.cpp
@@ -500,12 +500,17 @@ FTAAOutputs FTAAPassParameters::AddTemporalAAPass(
 		{
 			if (bCameraCut)
 			{
+				FRDGTextureRef BlackDummy = GraphBuilder.RegisterExternalTexture(GSystemTextures.BlackDummy);
+
 				CommonShaderParameters.ScreenPosToHistoryBufferUV = FVector4(1.0f, 1.0f, 1.0f, 1.0f);
 				CommonShaderParameters.ScreenPosAbsMax = FVector2D(0.0f, 0.0f);
 				CommonShaderParameters.HistoryBufferUVMinMax = FVector4(0.0f, 0.0f, 0.0f, 0.0f);
 				CommonShaderParameters.HistoryBufferSize = FVector4(1.0f, 1.0f, 1.0f, 1.0f);
-				CommonShaderParameters.HistoryBuffer0 = GraphBuilder.RegisterExternalTexture(GSystemTextures.BlackDummy);
-				CommonShaderParameters.HistoryBuffer1 = CommonShaderParameters.HistoryBuffer0;
+				CommonShaderParameters.HistoryBuffer0 = BlackDummy;
+				CommonShaderParameters.HistoryBuffer1 = BlackDummy;
+				
+				// Remove dependency of the velocity buffer on camera cut, given it's going to be ignored by the shader.
+				CommonShaderParameters.SceneTextures.SceneVelocityBuffer = BlackDummy;
 			}
 			else
 			{
diff --git a/Engine/Source/Runtime/Renderer/Private/PostProcess/PostProcessing.cpp b/Engine/Source/Runtime/Renderer/Private/PostProcess/PostProcessing.cpp
--- a/Engine/Source/Runtime/Renderer/Private/PostProcess/PostProcessing.cpp
+++ b/Engine/Source/Runtime/Renderer/Private/PostProcess/PostProcessing.cpp
@@ -1393,11 +1393,6 @@ void FPostProcessing::Process(FRHICommandListImmediate& RHICmdList, const FViewI
 				VelocityInput = Context.Graph.RegisterPass(new(FMemStack::Get()) FRCPassPostProcessInput(VelocityRT));
 				PreFlattenVelocity = VelocityInput;
 			}
-			else if (View.AntiAliasingMethod == AAM_TemporalAA && View.bCameraCut)
-			{
-				// Velocity pass is not being rendered on camera cut.
-				VelocityInput = Context.Graph.RegisterPass(new(FMemStack::Get()) FRCPassPostProcessInput(GSystemTextures.BlackDummy));
-			}
 
 			Context.FinalOutput = AddPostProcessMaterialChain(Context, BL_BeforeTranslucency, SeparateTranslucency);
 			
@@ -1422,14 +1417,7 @@ void FPostProcessing::Process(FRHICommandListImmediate& RHICmdList, const FViewI
 					FRDGTextureRef LocalSeparateTranslucency = Pass->CreateRDGTextureForOptionalInput(GraphBuilder, ePId_Input1, TEXT("SeparateTranslucency"));
 					
 					// FPostProcessing::Process() does a AdjustGBufferRefCount(RHICmdList, -1), therefore need to pass down reference on velocity buffer manually.
-					if (FRDGTextureRef SceneVelocityBuffer = Pass->CreateRDGTextureForOptionalInput(GraphBuilder, ePId_Input2, TEXT("SceneVelocity")))
-					{
-						SceneTextures.SceneVelocityBuffer = SceneVelocityBuffer;
-					}
-					else
-					{
-						UE_DEBUG_BREAK();
-					}
+					SceneTextures.SceneVelocityBuffer = Pass->CreateRDGTextureForOptionalInput(GraphBuilder, ePId_Input2, TEXT("SceneVelocity"));
 
 					FRDGTextureRef NewSceneColor = SceneColor;
 
diff --git a/Engine/Source/Runtime/Renderer/Private/ScreenSpaceDenoise.cpp b/Engine/Source/Runtime/Renderer/Private/ScreenSpaceDenoise.cpp
--- a/Engine/Source/Runtime/Renderer/Private/ScreenSpaceDenoise.cpp
+++ b/Engine/Source/Runtime/Renderer/Private/ScreenSpaceDenoise.cpp
@@ -921,6 +921,12 @@ static void DenoiseSignalAtConstantPixelDensity(
 		CommonParameters.SceneTextures = SceneTextures;
 		CommonParameters.ViewUniformBuffer = View.ViewUniformBuffer;
 		CommonParameters.EyeAdaptation = GetEyeAdaptationTexture(GraphBuilder, View);
+
+		// Remove dependency of the velocity buffer on camera cut, given it's going to be ignored by the shaders.
+		if (View.bCameraCut)
+		{
+			CommonParameters.SceneTextures.SceneVelocityBuffer = GraphBuilder.RegisterExternalTexture(GSystemTextures.BlackDummy);
+		}
 	}
 
 	#if RHI_RAYTRACING
