Fix some VS2019 compile errors
#rnx

diff --git a/Engine/Source/Runtime/Engine/Private/Materials/MaterialShader.cpp b/Engine/Source/Runtime/Engine/Private/Materials/MaterialShader.cpp
--- a/Engine/Source/Runtime/Engine/Private/Materials/MaterialShader.cpp
+++ b/Engine/Source/Runtime/Engine/Private/Materials/MaterialShader.cpp
@@ -2077,7 +2077,7 @@ bool FMaterialShaderMap::IsComplete(const FMaterial* Material, bool bSilent)
 				for (int32 Index = 0; Index < StageTypes.Num(); ++Index)
 				{
 					auto* ShaderType = StageTypes[Index]->GetMaterialShaderType();
-					if (!IsMaterialShaderComplete(Material, ShaderType, Pipeline, bSilent, kUniqueShaderPermutationId))
+					if (!IsMaterialShaderComplete(Material, ShaderType, Pipeline, kUniqueShaderPermutationId, bSilent))
 					{
 						return false;
 					}
diff --git a/Engine/Source/Runtime/Landscape/Private/LandscapeRender.cpp b/Engine/Source/Runtime/Landscape/Private/LandscapeRender.cpp
--- a/Engine/Source/Runtime/Landscape/Private/LandscapeRender.cpp
+++ b/Engine/Source/Runtime/Landscape/Private/LandscapeRender.cpp
@@ -4278,7 +4278,7 @@ public:
 
 				// For now only compile FLandscapeFixedGridVertexFactory for runtime virtual texture page rendering (can change if we need for other cases)
 				// Todo: only compile LandscapeXYOffsetVertexFactory if we are using it
-				bool bIsRuntimeVirtualTextureShaderType = Algo::Find(GetRuntimeVirtualTextureShaderTypes(), ShaderType->GetFName());
+				bool bIsRuntimeVirtualTextureShaderType = Algo::Find(GetRuntimeVirtualTextureShaderTypes(), ShaderType->GetFName()) != nullptr;
 
 				static const FName LandscapeVertexFactory = FName(TEXT("FLandscapeVertexFactory"));
 				static const FName LandscapeXYOffsetVertexFactory = FName(TEXT("FLandscapeXYOffsetVertexFactory"));
diff --git a/Engine/Source/Runtime/Renderer/Private/LightGridInjection.cpp b/Engine/Source/Runtime/Renderer/Private/LightGridInjection.cpp
--- a/Engine/Source/Runtime/Renderer/Private/LightGridInjection.cpp
+++ b/Engine/Source/Runtime/Renderer/Private/LightGridInjection.cpp
@@ -746,10 +746,10 @@ void FDeferredShadingSceneRenderer::ComputeLightGrid(FRHICommandListImmediate& R
 #endif // ENABLE_LIGHT_CULLING_VIEW_SPACE_BUILD_DATA
 
 					FLightGridInjectionCS::FPermutationDomain PermutationVector;
-					PermutationVector.Set<FLightGridInjectionCS::FUseLinkedListDim>(GLightLinkedListCulling);
+					PermutationVector.Set<FLightGridInjectionCS::FUseLinkedListDim>(GLightLinkedListCulling != 0);
 					TShaderMapRef<FLightGridInjectionCS> ComputeShader(View.ShaderMap, PermutationVector);
 
-					if (GLightLinkedListCulling)
+					if (GLightLinkedListCulling != 0)
 					{
 						AddPass_ClearUAV(GraphBuilder, RDG_EVENT_NAME("Clear:StartOffsetGrid"), PassParameters->RWStartOffsetGrid, 0xFFFFFFFF);
 						AddPass_ClearUAV(GraphBuilder, RDG_EVENT_NAME("Clear:NextCulledLightLink"), PassParameters->RWNextCulledLightLink, 0);
diff --git a/Engine/Source/Runtime/Renderer/Private/RayTracing/RayTracingTranslucency.cpp b/Engine/Source/Runtime/Renderer/Private/RayTracing/RayTracingTranslucency.cpp
--- a/Engine/Source/Runtime/Renderer/Private/RayTracing/RayTracingTranslucency.cpp
+++ b/Engine/Source/Runtime/Renderer/Private/RayTracing/RayTracingTranslucency.cpp
@@ -111,7 +111,7 @@ bool ShouldRenderRayTracingTranslucency(const FViewInfo& View)
 	bool bViewWithRaytracingTranslucency = View.FinalPostProcessSettings.TranslucencyType == ETranslucencyType::RayTracing;
 
 	const int32 GRayTracingTranslucency = CVarRayTracingTranslucency.GetValueOnRenderThread();
-	const bool bTranslucencyCvarEnabled = GRayTracingTranslucency < 0 ? bViewWithRaytracingTranslucency : GRayTracingTranslucency;
+	const bool bTranslucencyCvarEnabled = GRayTracingTranslucency < 0 ? bViewWithRaytracingTranslucency : (GRayTracingTranslucency != 0);
 	const int32 ForceAllRayTracingEffects = GetForceRayTracingEffectsCVarValue();
 	const bool bRayTracingTranslucencyEnabled = (ForceAllRayTracingEffects > 0 || (bTranslucencyCvarEnabled && ForceAllRayTracingEffects < 0));
 
diff --git a/Engine/Source/Runtime/Renderer/Private/ReflectionEnvironment.cpp b/Engine/Source/Runtime/Renderer/Private/ReflectionEnvironment.cpp
--- a/Engine/Source/Runtime/Renderer/Private/ReflectionEnvironment.cpp
+++ b/Engine/Source/Runtime/Renderer/Private/ReflectionEnvironment.cpp
@@ -202,7 +202,7 @@ bool ShouldRenderRayTracingReflections(const FViewInfo& View)
 {
 	bool bThisViewHasRaytracingReflections = View.FinalPostProcessSettings.ReflectionsType == EReflectionsType::RayTracing;
 
-	const bool bReflectionsCvarEnabled = GRayTracingReflections < 0 ? bThisViewHasRaytracingReflections : GRayTracingReflections;
+	const bool bReflectionsCvarEnabled = GRayTracingReflections < 0 ? bThisViewHasRaytracingReflections : (GRayTracingReflections != 0);
 	const int32 ForceAllRayTracingEffects = GetForceRayTracingEffectsCVarValue();
 	const bool bReflectionPassEnabled = (ForceAllRayTracingEffects > 0 || (bReflectionsCvarEnabled && ForceAllRayTracingEffects < 0));
 
diff --git a/Engine/Source/Runtime/Renderer/Private/ShaderPrint.cpp b/Engine/Source/Runtime/Renderer/Private/ShaderPrint.cpp
--- a/Engine/Source/Runtime/Renderer/Private/ShaderPrint.cpp
+++ b/Engine/Source/Runtime/Renderer/Private/ShaderPrint.cpp
@@ -156,7 +156,7 @@ namespace ShaderPrint
 
 	bool IsEnabled()
 	{
-		return CVarEnable.GetValueOnAnyThread();
+		return CVarEnable.GetValueOnAnyThread() != 0;
 	}
 
 	// Shader to initialize the output value buffer
