Addresses static analysis warning

diff --git a/Engine/Source/Runtime/Renderer/Private/IndirectLightRendering.cpp b/Engine/Source/Runtime/Renderer/Private/IndirectLightRendering.cpp
--- a/Engine/Source/Runtime/Renderer/Private/IndirectLightRendering.cpp
+++ b/Engine/Source/Runtime/Renderer/Private/IndirectLightRendering.cpp
@@ -145,6 +145,7 @@ void FDeferredShadingSceneRenderer::RenderDiffuseIndirectAndAmbientOcclusion(FRH
 		}
 		
 		IScreenSpaceDenoiser::FDiffuseIndirectOutputs DenoiserOutputs;
+#if 0 // TODO
 		if (DenoiseMode != 0)
 		{
 			const IScreenSpaceDenoiser* DefaultDenoiser = IScreenSpaceDenoiser::GetDefaultDenoiser();
@@ -164,6 +165,7 @@ void FDeferredShadingSceneRenderer::RenderDiffuseIndirectAndAmbientOcclusion(FRH
 				RayTracingConfig);
 		}
 		else
+#endif
 		{
 			DenoiserOutputs.Color = DenoiserInputs.Color;
 			DenoiserOutputs.AmbientOcclusionMask = DenoiserInputs.AmbientOcclusionMask;
