Fix runtime virtual texture desyncing when layer mask doesn't specify the full set of VT layers
#rnx

diff --git a/Engine/Source/Runtime/Renderer/Private/VT/RuntimeVirtualTextureProducer.cpp b/Engine/Source/Runtime/Renderer/Private/VT/RuntimeVirtualTextureProducer.cpp
--- a/Engine/Source/Runtime/Renderer/Private/VT/RuntimeVirtualTextureProducer.cpp
+++ b/Engine/Source/Runtime/Renderer/Private/VT/RuntimeVirtualTextureProducer.cpp
@@ -127,14 +127,14 @@ IVirtualTextureFinalizer* FRuntimeVirtualTextureProducer::ProducePageData(
 	// If we can't avoid partial layer masks then we could look at ways to handle it more efficiently (right now we render all layers even for these partial requests).
 
 	//todo[vt]: Add support for more than two layers
-	if (LayerMask & 1)
+	if (TargetLayers[0].TextureRHI->GetTexture2D() != nullptr)
 	{
 		Tile.Texture0 = TargetLayers[0].TextureRHI->GetTexture2D();
 		Tile.DestX0 = TargetLayers[0].pPageLocation.X;
 		Tile.DestY0 = TargetLayers[0].pPageLocation.Y;
 	}
 
-	if (LayerMask & 2)
+	if (TargetLayers[1].TextureRHI->GetTexture2D() != nullptr)
 	{
 		Tile.Texture1 = TargetLayers[1].TextureRHI->GetTexture2D();
 		Tile.DestX1 = TargetLayers[1].pPageLocation.X;
