Fix RVT renderer potential divide by 0 when VT is non-square
#rnx

diff --git a/Engine/Source/Runtime/Renderer/Private/VT/RuntimeVirtualTextureProducer.cpp b/Engine/Source/Runtime/Renderer/Private/VT/RuntimeVirtualTextureProducer.cpp
--- a/Engine/Source/Runtime/Renderer/Private/VT/RuntimeVirtualTextureProducer.cpp
+++ b/Engine/Source/Runtime/Renderer/Private/VT/RuntimeVirtualTextureProducer.cpp
@@ -56,13 +56,13 @@ void FRuntimeVirtualTextureFinalizer::Finalize(FRHICommandListImmediate& RHICmdL
 		const FVector2D DestinationBoxStart1(Entry.DestX1 * TileSize, Entry.DestY1 * TileSize);
 		const FBox2D DestinationBox1(DestinationBoxStart1, DestinationBoxStart1 + FVector2D(TileSize, TileSize));
 
-		const uint32 X = FMath::ReverseMortonCode2(Entry.vAddress);
-		const uint32 Y = FMath::ReverseMortonCode2(Entry.vAddress >> 1);
-		const uint32 DivisorX = Desc.BlockWidthInTiles >> Entry.vLevel;
-		const uint32 DivisorY = Desc.BlockHeightInTiles >> Entry.vLevel;
+		const float X = (float)FMath::ReverseMortonCode2(Entry.vAddress);
+		const float Y = (float)FMath::ReverseMortonCode2(Entry.vAddress >> 1);
+		const float DivisorX = (float)Desc.BlockWidthInTiles / (float)(1 << Entry.vLevel);
+		const float DivisorY = (float)Desc.BlockHeightInTiles / (float)(1 << Entry.vLevel);
 
-		const FVector2D UV((float)X / (float)DivisorX, (float)Y / (float)DivisorY);
-		const FVector2D UVSize(1.f / (float)DivisorX, 1.f / (float)DivisorY);
+		const FVector2D UV(X / DivisorX, Y / DivisorY);
+		const FVector2D UVSize(1.f / DivisorX, 1.f / DivisorY);
 		const FVector2D UVBorder = UVSize * ((float)Desc.TileBorderSize / (float)Desc.TileSize);
 		const FBox2D UVRange(UV - UVBorder, UV + UVSize + UVBorder);
 
@@ -127,14 +127,14 @@ IVirtualTextureFinalizer* FRuntimeVirtualTextureProducer::ProducePageData(
 	// If we can't avoid partial layer masks then we could look at ways to handle it more efficiently (right now we render all layers even for these partial requests).
 
 	//todo[vt]: Add support for more than two layers
-	if (TargetLayers[0].TextureRHI->GetTexture2D() != nullptr)
+	if (TargetLayers[0].TextureRHI != nullptr)
 	{
 		Tile.Texture0 = TargetLayers[0].TextureRHI->GetTexture2D();
 		Tile.DestX0 = TargetLayers[0].pPageLocation.X;
 		Tile.DestY0 = TargetLayers[0].pPageLocation.Y;
 	}
 
-	if (TargetLayers[1].TextureRHI->GetTexture2D() != nullptr)
+	if (TargetLayers[1].TextureRHI != nullptr)
 	{
 		Tile.Texture1 = TargetLayers[1].TextureRHI->GetTexture2D();
 		Tile.DestX1 = TargetLayers[1].pPageLocation.X;
