Fix non-square VT not working
#rnx

diff --git a/Engine/Source/Runtime/Renderer/Private/VT/VirtualTextureSystem.cpp b/Engine/Source/Runtime/Renderer/Private/VT/VirtualTextureSystem.cpp
--- a/Engine/Source/Runtime/Renderer/Private/VT/VirtualTextureSystem.cpp
+++ b/Engine/Source/Runtime/Renderer/Private/VT/VirtualTextureSystem.cpp
@@ -1265,7 +1265,7 @@ void FVirtualTextureSystem::GatherRequestsTask(const FGatherRequestsParameters&
 			// Possible there is a more clever approach to take here
 			{
 				const uint32 ProducerMipWidthInTiles = FMath::Max(Producer->GetWidthInTiles() >> Local_vLevel, 1u);
-				const uint32 ProducerMipHeightInTiles = FMath::Max(Producer->GetWidthInTiles() >> Local_vLevel, 1u);
+				const uint32 ProducerMipHeightInTiles = FMath::Max(Producer->GetHeightInTiles() >> Local_vLevel, 1u);
 				uint32 Local_vTileX = FMath::ReverseMortonCode2(Local_vAddress);
 				uint32 Local_vTileY = FMath::ReverseMortonCode2(Local_vAddress >> 1);
 				Local_vTileX &= (ProducerMipWidthInTiles - 1u);
