Integrates RTAO and RTGI withing same code path as SSGI

diff --git a/Engine/Shaders/Private/RayTracing/RayTracingGlobalIlluminationCompositePS.usf b/Engine/Shaders/Private/RayTracing/RayTracingGlobalIlluminationCompositePS.usf
deleted file mode 100644
diff --git a/Engine/Source/Runtime/Renderer/Private/DeferredShadingRenderer.cpp b/Engine/Source/Runtime/Renderer/Private/DeferredShadingRenderer.cpp
--- a/Engine/Source/Runtime/Renderer/Private/DeferredShadingRenderer.cpp
+++ b/Engine/Source/Runtime/Renderer/Private/DeferredShadingRenderer.cpp
@@ -1534,14 +1534,11 @@ void FDeferredShadingSceneRenderer::Render(FRHICommandListImmediate& RHICmdList)
 #if RHI_RAYTRACING
 	TRefCountPtr<IPooledRenderTarget> SkyLightRT;
 	TRefCountPtr<IPooledRenderTarget> SkyLightHitDistanceRT;
-	TRefCountPtr<IPooledRenderTarget> GlobalIlluminationRT;
 
 	const bool bRayTracingEnabled = IsRayTracingEnabled();
 	if (bRayTracingEnabled && bCanOverlayRayTracingOutput)
 	{		
 		RenderRayTracingSkyLight(RHICmdList, SkyLightRT, SkyLightHitDistanceRT);
-		RenderRayTracingGlobalIllumination(RHICmdList, GlobalIlluminationRT); 
-		RenderRayTracingAmbientOcclusion(RHICmdList, SceneContext.ScreenSpaceAO);
 	}
 #endif // RHI_RAYTRACING
 	checkSlow(RHICmdList.IsOutsideRenderPass());
@@ -1701,14 +1698,6 @@ void FDeferredShadingSceneRenderer::Render(FRHICommandListImmediate& RHICmdList)
 		{
 			CompositeRayTracingSkyLight(RHICmdList, SkyLightRT, SkyLightHitDistanceRT);
 		}
-
-		if (GlobalIlluminationRT)
-		{
-			for (int32 ViewIndex = 0; ViewIndex < Views.Num(); ++ViewIndex)
-			{
-				CompositeGlobalIllumination(RHICmdList, Views[ViewIndex], GlobalIlluminationRT);
-			}
-		}
 #endif // RHI_RAYTRACING
 		ServiceLocalQueue();
 	}
diff --git a/Engine/Source/Runtime/Renderer/Private/DeferredShadingRenderer.h b/Engine/Source/Runtime/Renderer/Private/DeferredShadingRenderer.h
--- a/Engine/Source/Runtime/Renderer/Private/DeferredShadingRenderer.h
+++ b/Engine/Source/Runtime/Renderer/Private/DeferredShadingRenderer.h
@@ -485,30 +485,25 @@ private:
 
 	void RenderRayTracingStochasticRectLight(FRHICommandListImmediate& RHICmdList, const FLightSceneInfo& RectLightSceneInfo, TRefCountPtr<IPooledRenderTarget>& RectLightRT, TRefCountPtr<IPooledRenderTarget>& HitDistanceRT);
 	void CompositeRayTracingSkyLight(FRHICommandListImmediate& RHICmdList, TRefCountPtr<IPooledRenderTarget>& SkyLightRT, TRefCountPtr<IPooledRenderTarget>& HitDistanceRT);
+	
+	void RenderRayTracingGlobalIllumination(
+		FRDGBuilder& GraphBuilder, 
+		FSceneTextureParameters& SceneTextures,
+		FViewInfo& View,
+		IScreenSpaceDenoiser::FAmbientOcclusionRayTracingConfig* RayTracingConfig,
+		IScreenSpaceDenoiser::FDiffuseIndirectInputs* OutDenoiserInputs);
+	
+	void RenderRayTracingAmbientOcclusion(
+		FRDGBuilder& GraphBuilder,
+		FViewInfo& View,
+		const FSceneTextureParameters& SceneTextures,
+		FRDGTextureRef* OutAmbientOcclusionTexture);
+	
 
 #if RHI_RAYTRACING
 	template <int TextureImportanceSampling> void RenderRayTracingRectLightInternal(FRHICommandListImmediate& RHICmdList, const TArray<FViewInfo>& Views, const FLightSceneInfo& RectLightSceneInfo, TRefCountPtr<IPooledRenderTarget>& ScreenShadowMaskTexture, TRefCountPtr<IPooledRenderTarget>& RayDistanceTexture);
 	void VisualizeRectLightMipTree(FRHICommandListImmediate& RHICmdList, const FViewInfo& View, const FRWBuffer& RectLightMipTree, const FIntVector& RectLightMipTreeDimensions);
 	
-	void RenderRayTracingAmbientOcclusion(FRHICommandListImmediate& RHICmdList, TRefCountPtr<IPooledRenderTarget>& AmbientOcclusionRT);
-	void RenderRayTracingAmbientOcclusion(FRHICommandListImmediate& RHICmdList, FRDGBuilder& GraphBuilder, FViewInfo& View, FRDGTextureRef AmbientOcclusionTexture, FRDGTextureRef RayDistanceTexture, TRefCountPtr<IPooledRenderTarget>& AmbientOcclusionMaskRT);
-	void CompositeRayTracingAmbientOcclusion(FRHICommandListImmediate& RHICmdList, TRefCountPtr<IPooledRenderTarget>& AmbientOcclusionRT);
-	
-	void RenderRayTracingGlobalIllumination(FRHICommandListImmediate& RHICmdList, TRefCountPtr<IPooledRenderTarget>& GlobalIlluminationRT);
-	void RenderRayTracingGlobalIllumination(
-		FRHICommandListImmediate& RHICmdList, 
-		FRDGBuilder& GraphBuilder, 
-		FViewInfo& View, 
-		IScreenSpaceDenoiser::FAmbientOcclusionRayTracingConfig& RayTracingConfig, 
-		int32 UpscaleFactor, TRefCountPtr<IPooledRenderTarget>& GlobalIlluminationRT, 
-		TRefCountPtr<IPooledRenderTarget>& AmbientOcclusionRT,
-		FRDGTextureUAV* GlobalIlluminationUAV,
-		FRDGTextureUAV* RayDistanceUAV,
-		const FRDGTextureRef GlobalIlluminationTexture,
-		const FRDGTextureRef RayDistanceTexture
-	);
-	void CompositeGlobalIllumination(FRHICommandListImmediate& RHICmdList, const FViewInfo& View, TRefCountPtr<IPooledRenderTarget>& GlobalIlluminationRT);
-
 	void BuildSkyLightCdfs(FRHICommandListImmediate& RHICmdList, FSkyLightSceneProxy* SkyLight);
 	void BuildSkyLightMipTree(FRHICommandListImmediate& RHICmdList, FTextureRHIRef SkyLightTexture, FRWBuffer& SkyLightMipTreePosX, FRWBuffer& SkyLightMipTreePosY, FRWBuffer& SkyLightMipTreePosZ, FRWBuffer& SkyLightMipTreeNegX, FRWBuffer& SkyLightMipTreeNegY, FRWBuffer& SkyLightMipTreeNegZ, FIntVector& SkyLightMipTreeDimensions);
 	void BuildSkyLightMipTreePdf(
diff --git a/Engine/Source/Runtime/Renderer/Private/IndirectLightRendering.cpp b/Engine/Source/Runtime/Renderer/Private/IndirectLightRendering.cpp
--- a/Engine/Source/Runtime/Renderer/Private/IndirectLightRendering.cpp
+++ b/Engine/Source/Runtime/Renderer/Private/IndirectLightRendering.cpp
@@ -9,6 +9,7 @@
 #include "ScreenSpaceRayTracing.h"
 #include "DeferredShadingRenderer.h"
 #include "PostProcessing.h" // for FPostProcessVS
+#include "RayTracing/RaytracingOptions.h"
 
 
 static TAutoConsoleVariable<int32> CVarDiffuseIndirectDenoiser(
@@ -120,18 +121,23 @@ void FDeferredShadingSceneRenderer::RenderDiffuseIndirectAndAmbientOcclusion(FRH
 
 	for (FViewInfo& View : Views)
 	{
+		// TODO: enum cvar. 
+		const bool bApplyRTGI = ShouldRenderRayTracingGlobalIllumination(View);
 		const bool bApplySSGI = ShouldRenderScreenSpaceDiffuseIndirect(View);
 		const bool bApplySSAO = SceneContext.bScreenSpaceAOIsValid;
+		const bool bApplyRTAO = ShouldRenderRayTracingAmbientOcclusion(View);
 
 		int32 DenoiseMode = CVarDiffuseIndirectDenoiser.GetValueOnRenderThread();
 
-		
 		IScreenSpaceDenoiser::FAmbientOcclusionRayTracingConfig RayTracingConfig;
 
-		// TODO: integrate RTGI.
 		// TODO: hybrid SSGI / RTGI
 		IScreenSpaceDenoiser::FDiffuseIndirectInputs DenoiserInputs;
-		if (bApplySSGI)
+		if (bApplyRTGI)
+		{
+			RenderRayTracingGlobalIllumination(GraphBuilder, SceneTextures, View, /* out */ &RayTracingConfig, /* out */ &DenoiserInputs);
+		}
+		else if (bApplySSGI)
 		{
 			RenderScreenSpaceDiffuseIndirect(GraphBuilder, SceneTextures, SceneColor, View, /* out */ &DenoiserInputs);
 			
@@ -145,7 +151,6 @@ void FDeferredShadingSceneRenderer::RenderDiffuseIndirectAndAmbientOcclusion(FRH
 		}
 		
 		IScreenSpaceDenoiser::FDiffuseIndirectOutputs DenoiserOutputs;
-#if 0 // TODO
 		if (DenoiseMode != 0)
 		{
 			const IScreenSpaceDenoiser* DefaultDenoiser = IScreenSpaceDenoiser::GetDefaultDenoiser();
@@ -165,27 +170,37 @@ void FDeferredShadingSceneRenderer::RenderDiffuseIndirectAndAmbientOcclusion(FRH
 				RayTracingConfig);
 		}
 		else
-#endif
 		{
 			DenoiserOutputs.Color = DenoiserInputs.Color;
 			DenoiserOutputs.AmbientOcclusionMask = DenoiserInputs.AmbientOcclusionMask;
 		}
 
-		// Extract the dynamic AO for further down lide
+		// Render RTAO that override any technic.
+		if (bApplyRTAO)
+		{
+			RenderRayTracingAmbientOcclusion(
+				GraphBuilder,
+				View,
+				SceneTextures,
+				&DenoiserOutputs.AmbientOcclusionMask);
+		}
+
+		// Extract the dynamic AO for application of AO beyond RenderDiffuseIndirectAndAmbientOcclusion()
 		if (DenoiserOutputs.AmbientOcclusionMask)
 		{
 			ensureMsgf(!bApplySSAO, TEXT("Looks like SSAO has been computed for this view but is being overridden."));
-			ensureMsgf(Views.Num() == 1, TEXT("SSGI can only output AO for 1 view")); // TODO.
+			ensureMsgf(Views.Num() == 1, TEXT("Need to add support for one AO texture per view in FSceneRenderTargets")); // TODO.
 			GraphBuilder.QueueTextureExtraction(DenoiserOutputs.AmbientOcclusionMask, &SceneContext.ScreenSpaceAO);
 			SceneContext.bScreenSpaceAOIsValid = true;
 		}
 		else if (bApplySSAO)
 		{
+			// Fetch result of SSAO that was done earlier.
 			DenoiserOutputs.AmbientOcclusionMask = GraphBuilder.RegisterExternalTexture(SceneContext.ScreenSpaceAO);
 		}
 
 		// Applies diffuse indirect and ambient occlusion to the scene color.
-		if (DenoiserOutputs.Color || bApplySSAO)
+		if (DenoiserOutputs.Color || DenoiserOutputs.AmbientOcclusionMask)
 		{
 			FDiffuseIndirectCompositePS::FParameters* PassParameters = GraphBuilder.AllocParameters<FDiffuseIndirectCompositePS::FParameters>();
 			
diff --git a/Engine/Source/Runtime/Renderer/Private/RayTracing/RayTracingAmbientOcclusion.cpp b/Engine/Source/Runtime/Renderer/Private/RayTracing/RayTracingAmbientOcclusion.cpp
--- a/Engine/Source/Runtime/Renderer/Private/RayTracing/RayTracingAmbientOcclusion.cpp
+++ b/Engine/Source/Runtime/Renderer/Private/RayTracing/RayTracingAmbientOcclusion.cpp
@@ -114,82 +114,42 @@ void FDeferredShadingSceneRenderer::PrepareRayTracingAmbientOcclusion(const FVie
 	}
 }
 
-void FDeferredShadingSceneRenderer::RenderRayTracingAmbientOcclusion(FRHICommandListImmediate& RHICmdList, TRefCountPtr<IPooledRenderTarget>& AmbientOcclusionRT)
-{
-	SCOPED_DRAW_EVENT(RHICmdList, RayTracingAmbientOcclusion);
-	SCOPED_GPU_STAT(RHICmdList, RayTracingAmbientOcclusion);
-
-	bool bAnyViewWithRTAO = false;
-	for (int32 ViewIndex = 0; ViewIndex < Views.Num(); ++ViewIndex)
-	{
-		FViewInfo& View = Views[ViewIndex];
-		bAnyViewWithRTAO = bAnyViewWithRTAO || ShouldRenderRayTracingAmbientOcclusion(View);
-	}
-
-	if (!bAnyViewWithRTAO)
-	{
-		return;
-	}
-
-	FRDGBuilder GraphBuilder(RHICmdList);
-	FSceneRenderTargets& SceneContext = FSceneRenderTargets::Get(RHICmdList);
-
-	// Define RDG targets
-	FRDGTextureRef AmbientOcclusionTexture;
-	{
-		FRDGTextureDesc Desc = SceneContext.GetSceneColor()->GetDesc();
-		Desc.Format = PF_R16F;
-		Desc.Flags &= ~(TexCreate_FastVRAM | TexCreate_Transient);
-		AmbientOcclusionTexture = GraphBuilder.CreateTexture(Desc, TEXT("RayTracingAmbientOcclusion"));
-	}
-
-	FRDGTextureRef RayDistanceTexture;
-	{
-		FRDGTextureDesc Desc = SceneContext.GetSceneColor()->GetDesc();
-		Desc.Format = PF_R16F;
-		Desc.Flags &= ~(TexCreate_FastVRAM | TexCreate_Transient);
-		RayDistanceTexture = GraphBuilder.CreateTexture(Desc, TEXT("RayTracingAmbientOcclusionHitDistance"));
-	}
-
-	for (int32 ViewIndex = 0; ViewIndex < Views.Num(); ++ViewIndex)
-	{	
-		FViewInfo& View = Views[ViewIndex];
-		if (ShouldRenderRayTracingAmbientOcclusion(View))
-		{
-			RenderRayTracingAmbientOcclusion(RHICmdList, GraphBuilder, View, AmbientOcclusionTexture, RayDistanceTexture, AmbientOcclusionRT);
-		}
-	}
-
-	GraphBuilder.Execute();
-	SceneContext.bScreenSpaceAOIsValid = true;
-	AmbientOcclusionRT->SetDebugName(TEXT("RayTracingAmbientOcclusion"));
-	GVisualizeTexture.SetCheckPoint(RHICmdList, AmbientOcclusionRT);
-}
+#endif // RHI_RAYTRACING
 
 void FDeferredShadingSceneRenderer::RenderRayTracingAmbientOcclusion(
-	FRHICommandListImmediate& RHICmdList,
 	FRDGBuilder& GraphBuilder,
 	FViewInfo& View,
-	FRDGTextureRef AmbientOcclusionTexture,
-	FRDGTextureRef RayDistanceTexture,
-	TRefCountPtr<IPooledRenderTarget>& AmbientOcclusionMaskRT
-)
+	const FSceneTextureParameters& SceneTextures,
+	FRDGTextureRef* OutAmbientOcclusionTexture)
+#if RHI_RAYTRACING
 {
-	FSceneRenderTargets& SceneContext = FSceneRenderTargets::Get(RHICmdList);
+	// Allocates denoiser inputs.
+	IScreenSpaceDenoiser::FAmbientOcclusionInputs DenoiserInputs;
+	{
+		FRDGTextureDesc Desc = FRDGTextureDesc::Create2DDesc(
+			SceneTextures.SceneDepthBuffer->Desc.Extent,
+			PF_R16F,
+			FClearValueBinding::None,
+			/* InFlags = */ TexCreate_None,
+			/* InTargetableFlags = */ TexCreate_ShaderResource | TexCreate_RenderTargetable | TexCreate_UAV,
+			/* bInForceSeparateTargetAndShaderResource = */ false);
+		DenoiserInputs.Mask = GraphBuilder.CreateTexture(Desc, TEXT("RayTracingAmbientOcclusion"));
+		DenoiserInputs.RayHitDistance = GraphBuilder.CreateTexture(Desc, TEXT("RayTracingAmbientOcclusionHitDistance"));
+	}
+	
+	IScreenSpaceDenoiser::FAmbientOcclusionRayTracingConfig RayTracingConfig;
+	RayTracingConfig.RayCountPerPixel =  GRayTracingAmbientOcclusionSamplesPerPixel >= 0 ? GRayTracingAmbientOcclusionSamplesPerPixel : View.FinalPostProcessSettings.RayTracingAOSamplesPerPixel;
 
 	// Build RTAO parameters
 	FRayTracingAmbientOcclusionRGS::FParameters *PassParameters = GraphBuilder.AllocParameters<FRayTracingAmbientOcclusionRGS::FParameters>();
-	PassParameters->SamplesPerPixel = GRayTracingAmbientOcclusionSamplesPerPixel >= 0 ? GRayTracingAmbientOcclusionSamplesPerPixel : View.FinalPostProcessSettings.RayTracingAOSamplesPerPixel;
+	PassParameters->SamplesPerPixel = RayTracingConfig.RayCountPerPixel;
 	PassParameters->MaxRayDistance = View.FinalPostProcessSettings.AmbientOcclusionRadius;
 	PassParameters->Intensity = View.FinalPostProcessSettings.AmbientOcclusionIntensity;
 	PassParameters->MaxNormalBias = GetRaytracingMaxNormalBias();
 	PassParameters->TLAS = View.RayTracingScene.RayTracingSceneRHI->GetShaderResourceView();
-	PassParameters->RWOcclusionMaskUAV = GraphBuilder.CreateUAV(AmbientOcclusionTexture);
-	PassParameters->RWHitDistanceUAV = GraphBuilder.CreateUAV(RayDistanceTexture);
+	PassParameters->RWOcclusionMaskUAV = GraphBuilder.CreateUAV(DenoiserInputs.Mask);
+	PassParameters->RWHitDistanceUAV = GraphBuilder.CreateUAV(DenoiserInputs.RayHitDistance);
 	PassParameters->ViewUniformBuffer = View.ViewUniformBuffer;
-
-	FSceneTextureParameters SceneTextures;
-	SetupSceneTextureParameters(GraphBuilder, &SceneTextures);
 	PassParameters->SceneTextures = SceneTextures;
 
 	FRayTracingAmbientOcclusionRGS::FPermutationDomain PermutationVector;
@@ -238,33 +198,28 @@ void FDeferredShadingSceneRenderer::RenderRayTracingAmbientOcclusion(
 		const IScreenSpaceDenoiser* DefaultDenoiser = IScreenSpaceDenoiser::GetDefaultDenoiser();
 		const IScreenSpaceDenoiser* DenoiserToUse = DenoiserMode == 1 ? DefaultDenoiser : GScreenSpaceDenoiser;
 
-		IScreenSpaceDenoiser::FAmbientOcclusionRayTracingConfig RayTracingConfig;
+		RDG_EVENT_SCOPE(GraphBuilder, "%s%s(AmbientOcclusion) %dx%d",
+			DenoiserToUse != DefaultDenoiser ? TEXT("ThirdParty ") : TEXT(""),
+			DenoiserToUse->GetDebugName(),
+			View.ViewRect.Width(), View.ViewRect.Height());
 
-		IScreenSpaceDenoiser::FAmbientOcclusionInputs DenoiserInputs;
-		DenoiserInputs.Mask = AmbientOcclusionTexture;
-		DenoiserInputs.RayHitDistance = RayDistanceTexture;
+		IScreenSpaceDenoiser::FAmbientOcclusionOutputs DenoiserOutputs = DenoiserToUse->DenoiseAmbientOcclusion(
+			GraphBuilder,
+			View,
+			&View.PrevViewInfo,
+			SceneTextureParams,
+			DenoiserInputs,
+			RayTracingConfig);
 
-		{
-			RDG_EVENT_SCOPE(GraphBuilder, "%s%s(AmbientOcclusion) %dx%d",
-				DenoiserToUse != DefaultDenoiser ? TEXT("ThirdParty ") : TEXT(""),
-				DenoiserToUse->GetDebugName(),
-				View.ViewRect.Width(), View.ViewRect.Height());
-
-			IScreenSpaceDenoiser::FAmbientOcclusionOutputs DenoiserOutputs = DenoiserToUse->DenoiseAmbientOcclusion(
-				GraphBuilder,
-				View,
-				&View.PrevViewInfo,
-				SceneTextureParams,
-				DenoiserInputs,
-				RayTracingConfig);
-
-			GraphBuilder.QueueTextureExtraction(DenoiserOutputs.AmbientOcclusionMask, &AmbientOcclusionMaskRT);
-		}
+		*OutAmbientOcclusionTexture = DenoiserOutputs.AmbientOcclusionMask;
 	}
 	else
 	{
-		GraphBuilder.QueueTextureExtraction(AmbientOcclusionTexture, &AmbientOcclusionMaskRT);
+		*OutAmbientOcclusionTexture = DenoiserInputs.Mask;
 	}
 }
-
-#endif // RHI_RAYTRACING
+#else
+{
+	unimplemented();
+}
+#endif
diff --git a/Engine/Source/Runtime/Renderer/Private/RayTracing/RayTracingGlobalIllumination.cpp b/Engine/Source/Runtime/Renderer/Private/RayTracing/RayTracingGlobalIllumination.cpp
--- a/Engine/Source/Runtime/Renderer/Private/RayTracing/RayTracingGlobalIllumination.cpp
+++ b/Engine/Source/Runtime/Renderer/Private/RayTracing/RayTracingGlobalIllumination.cpp
@@ -244,43 +244,6 @@ class FGlobalIlluminationRGS : public FGlobalShader
 	END_SHADER_PARAMETER_STRUCT()
 };
 
-class FRayTracingGlobalIlluminationCompositePS : public FGlobalShader
-{
-	DECLARE_GLOBAL_SHADER(FRayTracingGlobalIlluminationCompositePS)
-	SHADER_USE_PARAMETER_STRUCT(FRayTracingGlobalIlluminationCompositePS, FGlobalShader)
-
-	static bool ShouldCompilePermutation(const FGlobalShaderPermutationParameters& Parameters)
-	{
-		return ShouldCompileRayTracingShadersForProject(Parameters.Platform);
-	}
-
-	BEGIN_SHADER_PARAMETER_STRUCT(FParameters, )
-		RENDER_TARGET_BINDING_SLOTS()
-		SHADER_PARAMETER_RDG_TEXTURE(Texture2D, GlobalIlluminationTexture)
-		SHADER_PARAMETER_SAMPLER(SamplerState, GlobalIlluminationSampler)
-		SHADER_PARAMETER_STRUCT_REF(FViewUniformShaderParameters, ViewUniformBuffer)
-		SHADER_PARAMETER_STRUCT_INCLUDE(FSceneTextureParameters, SceneTextures)
-	END_SHADER_PARAMETER_STRUCT()
-};
-
-class FRayTracingGlobalIlluminationSceneColorCompositePS : public FGlobalShader
-{
-	DECLARE_GLOBAL_SHADER(FRayTracingGlobalIlluminationSceneColorCompositePS)
-	SHADER_USE_PARAMETER_STRUCT(FRayTracingGlobalIlluminationSceneColorCompositePS, FGlobalShader)
-
-	static bool ShouldCompilePermutation(const FGlobalShaderPermutationParameters& Parameters)
-	{
-		return ShouldCompileRayTracingShadersForProject(Parameters.Platform);
-	}
-
-	BEGIN_SHADER_PARAMETER_STRUCT(FParameters, )
-		RENDER_TARGET_BINDING_SLOTS()
-		SHADER_PARAMETER_RDG_TEXTURE(Texture2D, GlobalIlluminationTexture)
-		SHADER_PARAMETER_SAMPLER(SamplerState, GlobalIlluminationSampler)
-		SHADER_PARAMETER_STRUCT_INCLUDE(FSceneTextureParameters, SceneTextures)
-	END_SHADER_PARAMETER_STRUCT()
-};
-
 class FRayTracingGlobalIlluminationCHS : public FGlobalShader
 {
 	DECLARE_GLOBAL_SHADER(FRayTracingGlobalIlluminationCHS)
@@ -296,8 +259,6 @@ class FRayTracingGlobalIlluminationCHS : public FGlobalShader
 
 IMPLEMENT_GLOBAL_SHADER(FGlobalIlluminationRGS, "/Engine/Private/RayTracing/RayTracingGlobalIlluminationRGS.usf", "GlobalIlluminationRGS", SF_RayGen);
 IMPLEMENT_GLOBAL_SHADER(FRayTracingGlobalIlluminationCHS, "/Engine/Private/RayTracing/RayTracingGlobalIlluminationRGS.usf", "RayTracingGlobalIlluminationCHS", SF_RayHitGroup);
-IMPLEMENT_GLOBAL_SHADER(FRayTracingGlobalIlluminationCompositePS, "/Engine/Private/RayTracing/RayTracingGlobalIlluminationCompositePS.usf", "GlobalIlluminationCompositePS", SF_Pixel);
-IMPLEMENT_GLOBAL_SHADER(FRayTracingGlobalIlluminationSceneColorCompositePS, "/Engine/Private/RayTracing/RayTracingGlobalIlluminationCompositePS.usf", "GlobalIlluminationSceneColorCompositePS", SF_Pixel);
 
 void FDeferredShadingSceneRenderer::PrepareRayTracingGlobalIllumination(const FViewInfo& View, TArray<FRHIRayTracingShader*>& OutRayGenShaders)
 {
@@ -316,117 +277,46 @@ void FDeferredShadingSceneRenderer::PrepareRayTracingGlobalIllumination(const FV
 
 }
 
-void FDeferredShadingSceneRenderer::RenderRayTracingGlobalIllumination(
-	FRHICommandListImmediate& RHICmdList, 
-	TRefCountPtr<IPooledRenderTarget>& GlobalIlluminationRT
-)
-{
-	SCOPED_GPU_STAT(RHICmdList, RayTracingGlobalIllumination);
-
-	bool bAnyViewWithRTGI = false;
-	for (int32 ViewIndex = 0; ViewIndex < Views.Num(); ++ViewIndex)
-	{
-		FViewInfo& View = Views[ViewIndex];
-		bAnyViewWithRTGI = bAnyViewWithRTGI || ShouldRenderRayTracingGlobalIllumination(View);
-	}
-
-	if (!bAnyViewWithRTGI)
-	{
-		return;
-	}
+#endif // RHI_RAYTRACING
 
-	IScreenSpaceDenoiser::FAmbientOcclusionRayTracingConfig RayTracingConfig;
-	RayTracingConfig.ResolutionFraction = 1.0;
+void FDeferredShadingSceneRenderer::RenderRayTracingGlobalIllumination(
+	FRDGBuilder& GraphBuilder, 
+	FSceneTextureParameters& SceneTextures,
+	FViewInfo& View,
+	IScreenSpaceDenoiser::FAmbientOcclusionRayTracingConfig* OutRayTracingConfig,
+	IScreenSpaceDenoiser::FDiffuseIndirectInputs* OutDenoiserInputs)
+#if RHI_RAYTRACING
+{	
+	OutRayTracingConfig->ResolutionFraction = 1.0;
 	if (GRayTracingGlobalIlluminationDenoiser != 0)
 	{
-		RayTracingConfig.ResolutionFraction = FMath::Clamp(GRayTracingGlobalIlluminationScreenPercentage / 100.0, 0.25, 1.0);
+		OutRayTracingConfig->ResolutionFraction = FMath::Clamp(GRayTracingGlobalIlluminationScreenPercentage / 100.0, 0.25, 1.0);
 	}
 
-	int32 UpscaleFactor = int32(1.0 / RayTracingConfig.ResolutionFraction);
-
-	FSceneRenderTargets& SceneContext = FSceneRenderTargets::Get(RHICmdList);
-	{
-		FPooledRenderTargetDesc Desc = SceneContext.GetSceneColor()->GetDesc();
-		Desc.Format = PF_FloatRGBA;
-		Desc.Flags &= ~(TexCreate_FastVRAM | TexCreate_Transient);
-		GRenderTargetPool.FindFreeElement(RHICmdList, Desc, GlobalIlluminationRT, TEXT("RayTracingGlobalIllumination"));
-	}
-
-	FRDGBuilder GraphBuilder(RHICmdList);
+	RDG_EVENT_SCOPE(GraphBuilder, "RTGI");
 
-	FRDGTextureRef GlobalIlluminationTexture;
-	{
-		FRDGTextureDesc Desc = SceneContext.GetSceneColor()->GetDesc();
-		Desc.Extent /= UpscaleFactor;
-		Desc.Format = PF_FloatRGBA;
-		Desc.Flags &= ~(TexCreate_FastVRAM | TexCreate_Transient);
-		GlobalIlluminationTexture = GraphBuilder.CreateTexture(Desc, TEXT("RayTracingGlobalIllumination"));
-	}
-	FRDGTextureUAV* GlobalIlluminationUAV = GraphBuilder.CreateUAV(FRDGTextureUAVDesc(GlobalIlluminationTexture));
+	int32 RayTracingGISamplesPerPixel = GRayTracingGlobalIlluminationSamplesPerPixel > -1 ? GRayTracingGlobalIlluminationSamplesPerPixel : View.FinalPostProcessSettings.RayTracingGISamplesPerPixel;
+	OutRayTracingConfig->RayCountPerPixel = RayTracingGISamplesPerPixel;
+	
+	int32 UpscaleFactor = int32(1.0 / OutRayTracingConfig->ResolutionFraction);
 
-	FRDGTextureRef RayDistanceTexture;
+	// Allocate input for the denoiser.
 	{
-		FRDGTextureDesc Desc = SceneContext.GetSceneColor()->GetDesc();
-		Desc.Extent /= UpscaleFactor;
-		Desc.Format = PF_G16R16;
-		Desc.Flags &= ~(TexCreate_FastVRAM | TexCreate_Transient);
-		RayDistanceTexture = GraphBuilder.CreateTexture(Desc, TEXT("RayTracingGlobalIlluminationRayDistance"));
-	}
-	FRDGTextureUAV* RayDistanceUAV = GraphBuilder.CreateUAV(FRDGTextureUAVDesc(RayDistanceTexture));
+		FRDGTextureDesc Desc = FRDGTextureDesc::Create2DDesc(
+			SceneTextures.SceneDepthBuffer->Desc.Extent / UpscaleFactor,
+			PF_FloatRGBA,
+			FClearValueBinding::None,
+			/* InFlags = */ TexCreate_None,
+			/* InTargetableFlags = */ TexCreate_ShaderResource | TexCreate_RenderTargetable | TexCreate_UAV,
+			/* bInForceSeparateTargetAndShaderResource = */ false);
 
-	TRefCountPtr<IPooledRenderTarget>& AmbientOcclusionRT = SceneContext.ScreenSpaceAO;
+		OutDenoiserInputs->Color = GraphBuilder.CreateTexture(Desc, TEXT("RayTracingDiffuseIndirect"));
 
-	for (int32 ViewIndex = 0; ViewIndex < Views.Num(); ++ViewIndex)
-	{
-		FViewInfo& View = Views[ViewIndex];
-		if (ShouldRenderRayTracingGlobalIllumination(View))
-		{
-			RenderRayTracingGlobalIllumination(
-				RHICmdList, 
-				GraphBuilder, 
-				View, 
-				RayTracingConfig, 
-				UpscaleFactor, 
-				GlobalIlluminationRT, 
-				AmbientOcclusionRT, 
-				GlobalIlluminationUAV, 
-				RayDistanceUAV, 
-				GlobalIlluminationTexture, 
-				RayDistanceTexture
-			);
-		}
+		Desc.Format = PF_G16R16;
+		OutDenoiserInputs->RayHitDistance = GraphBuilder.CreateTexture(Desc, TEXT("RayTracingDiffuseIndirectHitDistance"));
 	}
 
-	GraphBuilder.Execute();
-	SceneContext.bScreenSpaceAOIsValid = true;
-	GVisualizeTexture.SetCheckPoint(RHICmdList, GlobalIlluminationRT);
-}
-
-void FDeferredShadingSceneRenderer::RenderRayTracingGlobalIllumination(
-	FRHICommandListImmediate& RHICmdList,
-	FRDGBuilder& GraphBuilder,
-	FViewInfo& View,
-	IScreenSpaceDenoiser::FAmbientOcclusionRayTracingConfig& RayTracingConfig,
-	int32 UpscaleFactor,
-	TRefCountPtr<IPooledRenderTarget>& GlobalIlluminationRT,
-	TRefCountPtr<IPooledRenderTarget>& AmbientOcclusionRT,
-	FRDGTextureUAV* GlobalIlluminationUAV,
-	FRDGTextureUAV* RayDistanceUAV,
-	const FRDGTextureRef GlobalIlluminationTexture,
-	const FRDGTextureRef RayDistanceTexture
-)
-{	
-	FSceneRenderTargets& SceneContext = FSceneRenderTargets::Get(RHICmdList);
-
-	RDG_EVENT_SCOPE(GraphBuilder, "RTGI");
-
-	int32 RayTracingGISamplesPerPixel = GRayTracingGlobalIlluminationSamplesPerPixel > -1 ? GRayTracingGlobalIlluminationSamplesPerPixel : View.FinalPostProcessSettings.RayTracingGISamplesPerPixel;
-	RayTracingConfig.RayCountPerPixel = RayTracingGISamplesPerPixel;
-
-	FSceneTextureParameters SceneTextures;
-	SetupSceneTextureParameters(GraphBuilder, &SceneTextures);
-
-	// Ray generation
+	// Ray generation pass
 	{
 		uint32 IterationCount = RayTracingGISamplesPerPixel;
 		uint32 SequenceCount = 1;
@@ -447,7 +337,8 @@ void FDeferredShadingSceneRenderer::RenderRayTracingGlobalIllumination(
 
 		if (Scene->SkyLight && Scene->SkyLight->ShouldRebuildCdf())
 		{
-			BuildSkyLightCdfs(RHICmdList, Scene->SkyLight);
+			// TODO: should be converted to RDG.
+			BuildSkyLightCdfs(GraphBuilder.RHICmdList, Scene->SkyLight);
 		}
 		FSkyLightData SkyLightParameters;
 		SetupSkyLightParameters(*Scene, &SkyLightParameters);
@@ -476,15 +367,17 @@ void FDeferredShadingSceneRenderer::RenderRayTracingGlobalIllumination(
 		PassParameters->LightParameters = CreateUniformBufferImmediate(LightParameters, EUniformBufferUsage::UniformBuffer_SingleDraw);
 		PassParameters->SceneTextures = SceneTextures;
 		PassParameters->SkyLight = CreateUniformBufferImmediate(SkyLightParameters, EUniformBufferUsage::UniformBuffer_SingleDraw);
-		TRefCountPtr<IPooledRenderTarget> SubsurfaceProfileRT((IPooledRenderTarget*) GetSubsufaceProfileTexture_RT(RHICmdList));
+
+		// TODO: should be converted to RDG
+		TRefCountPtr<IPooledRenderTarget> SubsurfaceProfileRT((IPooledRenderTarget*) GetSubsufaceProfileTexture_RT(GraphBuilder.RHICmdList));
 		if (!SubsurfaceProfileRT)
 		{
 			SubsurfaceProfileRT = GSystemTextures.BlackDummy;
 		}
 		PassParameters->SSProfilesTexture = GraphBuilder.RegisterExternalTexture(SubsurfaceProfileRT);
 		PassParameters->TransmissionProfilesLinearSampler = TStaticSamplerState<SF_Bilinear, AM_Clamp, AM_Clamp, AM_Clamp>::GetRHI();
-		PassParameters->RWGlobalIlluminationUAV = GlobalIlluminationUAV;
-		PassParameters->RWRayDistanceUAV = RayDistanceUAV;
+		PassParameters->RWGlobalIlluminationUAV = GraphBuilder.CreateUAV(OutDenoiserInputs->Color);
+		PassParameters->RWRayDistanceUAV = GraphBuilder.CreateUAV(OutDenoiserInputs->RayHitDistance);
 
 
 		FGlobalIlluminationRGS::FPermutationDomain PermutationVector;
@@ -507,152 +400,9 @@ void FDeferredShadingSceneRenderer::RenderRayTracingGlobalIllumination(
 			RHICmdList.RayTraceDispatch(View.RayTracingMaterialPipeline, RayGenerationShader->GetRayTracingShader(), RayTracingSceneRHI, GlobalResources, RayTracingResolution.X, RayTracingResolution.Y);
 		});
 	}
-
-	// Denoising
-
-	FRDGTextureRef ResultTexture; //#dxr_todo review
-
-	if (GRayTracingGlobalIlluminationDenoiser != 0)
-	{
-		const IScreenSpaceDenoiser* DefaultDenoiser = IScreenSpaceDenoiser::GetDefaultDenoiser();
-		const IScreenSpaceDenoiser* DenoiserToUse = GRayTracingGlobalIlluminationDenoiser == 1 ? DefaultDenoiser : GScreenSpaceDenoiser;
-
-		IScreenSpaceDenoiser::FDiffuseIndirectInputs DenoiserInputs;
-		DenoiserInputs.Color = GlobalIlluminationTexture;
-		DenoiserInputs.RayHitDistance = RayDistanceTexture;
-
-		{
-			RDG_EVENT_SCOPE(GraphBuilder, "%s%s(DiffuseIndirect) %dx%d",
-				DenoiserToUse != DefaultDenoiser ? TEXT("ThirdParty ") : TEXT(""),
-				DenoiserToUse->GetDebugName(),
-				View.ViewRect.Width(), View.ViewRect.Height());
-
-			IScreenSpaceDenoiser::FDiffuseIndirectOutputs DenoiserOutputs = DenoiserToUse->DenoiseDiffuseIndirect(
-				GraphBuilder,
-				View,
-				&View.PrevViewInfo,
-				SceneTextures,
-				DenoiserInputs,
-				RayTracingConfig);
-
-			ResultTexture = DenoiserOutputs.Color;
-		}
-	}
-	else
-	{
-		ResultTexture = GlobalIlluminationTexture;
-	}
-
-	// Compositing
-	{
-		FRayTracingGlobalIlluminationCompositePS::FParameters *PassParameters = GraphBuilder.AllocParameters<FRayTracingGlobalIlluminationCompositePS::FParameters>();
-		PassParameters->GlobalIlluminationTexture = ResultTexture;
-		PassParameters->GlobalIlluminationSampler = TStaticSamplerState<SF_Point, AM_Clamp, AM_Clamp, AM_Clamp>::GetRHI();
-		PassParameters->ViewUniformBuffer = View.ViewUniformBuffer;
-		PassParameters->RenderTargets[0] = FRenderTargetBinding(GraphBuilder.RegisterExternalTexture(GlobalIlluminationRT), ERenderTargetLoadAction::ELoad, ERenderTargetStoreAction::EStore);
-		PassParameters->RenderTargets[1] = FRenderTargetBinding(GraphBuilder.RegisterExternalTexture(AmbientOcclusionRT), ERenderTargetLoadAction::ELoad, ERenderTargetStoreAction::EStore);
-		PassParameters->SceneTextures = SceneTextures;
-
-		GraphBuilder.AddPass(
-			RDG_EVENT_NAME("GlobalIlluminationComposite"),
-			PassParameters,
-			ERenderGraphPassFlags::None,
-			[this, &SceneContext, &View, &SceneTextures, PassParameters](FRHICommandListImmediate& RHICmdList)
-			{
-				TShaderMapRef<FPostProcessVS> VertexShader(View.ShaderMap);
-				TShaderMapRef<FRayTracingGlobalIlluminationCompositePS> PixelShader(View.ShaderMap);
-				FGraphicsPipelineStateInitializer GraphicsPSOInit;
-				RHICmdList.ApplyCachedRenderTargets(GraphicsPSOInit);
-
-				// Additive blending
-				GraphicsPSOInit.BlendState = TStaticBlendState<CW_RGBA, BO_Add, BF_One, BF_Zero, BO_Add, BF_One, BF_Zero>::GetRHI();
-				//GraphicsPSOInit.BlendState = TStaticBlendState<CW_RGBA, BO_Add, BF_One, BF_One, BO_Add, BF_One, BF_One>::GetRHI();
-				GraphicsPSOInit.RasterizerState = TStaticRasterizerState<FM_Solid, CM_None>::GetRHI();
-				GraphicsPSOInit.DepthStencilState = TStaticDepthStencilState<false, CF_Always>::GetRHI();
-
-				GraphicsPSOInit.BoundShaderState.VertexDeclarationRHI = GFilterVertexDeclaration.VertexDeclarationRHI;
-				GraphicsPSOInit.BoundShaderState.VertexShaderRHI = GETSAFERHISHADER_VERTEX(*VertexShader);
-				GraphicsPSOInit.BoundShaderState.PixelShaderRHI = GETSAFERHISHADER_PIXEL(*PixelShader);
-				GraphicsPSOInit.PrimitiveType = PT_TriangleList;
-				SetGraphicsPipelineState(RHICmdList, GraphicsPSOInit);
-
-				SetShaderParameters(RHICmdList, *PixelShader, PixelShader->GetPixelShader(), *PassParameters);
-
-				RHICmdList.SetViewport(View.ViewRect.Min.X, View.ViewRect.Min.Y, 0.0f, View.ViewRect.Max.X, View.ViewRect.Max.Y, 1.0f);
-
-				DrawRectangle(
-					RHICmdList,
-					0, 0,
-					View.ViewRect.Width(), View.ViewRect.Height(),
-					View.ViewRect.Min.X, View.ViewRect.Min.Y,
-					View.ViewRect.Width(), View.ViewRect.Height(),
-					FIntPoint(View.ViewRect.Width(), View.ViewRect.Height()),
-					SceneContext.GetBufferSizeXY(),
-					*VertexShader
-				);
-			}
-		);
-	}
 }
-
-void FDeferredShadingSceneRenderer::CompositeGlobalIllumination(
-	FRHICommandListImmediate& RHICmdList,
-	const FViewInfo& View,
-	TRefCountPtr<IPooledRenderTarget>& GlobalIlluminationRT
-)
+#else
 {
-	FSceneRenderTargets& SceneContext = FSceneRenderTargets::Get(RHICmdList);
-
-	FRDGBuilder GraphBuilder(RHICmdList);
-	FSceneTextureParameters SceneTextures;
-	SetupSceneTextureParameters(GraphBuilder, &SceneTextures);
-
-	FRayTracingGlobalIlluminationSceneColorCompositePS::FParameters *PassParameters = GraphBuilder.AllocParameters<FRayTracingGlobalIlluminationSceneColorCompositePS::FParameters>();
-	PassParameters->GlobalIlluminationTexture = GraphBuilder.RegisterExternalTexture(GlobalIlluminationRT);
-	PassParameters->GlobalIlluminationSampler = TStaticSamplerState<SF_Point, AM_Clamp, AM_Clamp, AM_Clamp>::GetRHI();
-	PassParameters->RenderTargets[0] = FRenderTargetBinding(GraphBuilder.RegisterExternalTexture(SceneContext.GetSceneColor()), ERenderTargetLoadAction::ELoad, ERenderTargetStoreAction::EStore);
-	PassParameters->SceneTextures = SceneTextures;
-
-	GraphBuilder.AddPass(
-		RDG_EVENT_NAME("GlobalIlluminationComposite"),
-		PassParameters,
-		ERenderGraphPassFlags::None,
-		[this, &SceneContext, &View, PassParameters](FRHICommandListImmediate& RHICmdList)
-		{
-			TShaderMapRef<FPostProcessVS> VertexShader(View.ShaderMap);
-			TShaderMapRef<FRayTracingGlobalIlluminationSceneColorCompositePS> PixelShader(View.ShaderMap);
-			FGraphicsPipelineStateInitializer GraphicsPSOInit;
-			RHICmdList.ApplyCachedRenderTargets(GraphicsPSOInit);
-
-			// Additive blending
-			GraphicsPSOInit.BlendState = TStaticBlendState<CW_RGBA, BO_Add, BF_One, BF_One, BO_Add, BF_One, BF_One>::GetRHI();
-			//GraphicsPSOInit.BlendState = TStaticBlendState<CW_RGBA, BO_Add, BF_One, BF_One, BO_Add, BF_One, BF_One>::GetRHI();
-			GraphicsPSOInit.RasterizerState = TStaticRasterizerState<FM_Solid, CM_None>::GetRHI();
-			GraphicsPSOInit.DepthStencilState = TStaticDepthStencilState<false, CF_Always>::GetRHI();
-
-			GraphicsPSOInit.BoundShaderState.VertexDeclarationRHI = GFilterVertexDeclaration.VertexDeclarationRHI;
-			GraphicsPSOInit.BoundShaderState.VertexShaderRHI = GETSAFERHISHADER_VERTEX(*VertexShader);
-			GraphicsPSOInit.BoundShaderState.PixelShaderRHI = GETSAFERHISHADER_PIXEL(*PixelShader);
-			GraphicsPSOInit.PrimitiveType = PT_TriangleList;
-			SetGraphicsPipelineState(RHICmdList, GraphicsPSOInit);
-
-			SetShaderParameters(RHICmdList, *PixelShader, PixelShader->GetPixelShader(), *PassParameters);
-
-			RHICmdList.SetViewport(View.ViewRect.Min.X, View.ViewRect.Min.Y, 0.0f, View.ViewRect.Max.X, View.ViewRect.Max.Y, 1.0f);
-
-			DrawRectangle(
-				RHICmdList,
-				0, 0,
-				View.ViewRect.Width(), View.ViewRect.Height(),
-				View.ViewRect.Min.X, View.ViewRect.Min.Y,
-				View.ViewRect.Width(), View.ViewRect.Height(),
-				FIntPoint(View.ViewRect.Width(), View.ViewRect.Height()),
-				SceneContext.GetBufferSizeXY(),
-				*VertexShader
-			);
-		}
-	);
-	GraphBuilder.Execute();
+	unimplemented();
 }
-
 #endif // RHI_RAYTRACING
