Deletes no longer used CopyAndBackfillSignalInput() in denoisers.

diff --git a/Engine/Source/Runtime/Renderer/Private/ScreenSpaceDenoise.cpp b/Engine/Source/Runtime/Renderer/Private/ScreenSpaceDenoise.cpp
--- a/Engine/Source/Runtime/Renderer/Private/ScreenSpaceDenoise.cpp
+++ b/Engine/Source/Runtime/Renderer/Private/ScreenSpaceDenoise.cpp
@@ -756,29 +756,6 @@ struct FSSDConstantPixelDensitySettings
 	TStaticArray<const FLightSceneInfo*, IScreenSpaceDenoiser::kMaxBatchSize> LightSceneInfo;
 };
 
-/**
- * Returns a copy of SignalInput and fills all texture slots with the last valid texture, such that all elements of the array become valid.
- * This is required because shader compiler may not always be able to remove unused input textures. In this case a mismatch
- * between shader reflection data and parameters provided by high-level code will cause a validation error.
- * For example, this happens when using non-optimized shaders (r.Shaders.Optimize=0).
- */
-static FSSDSignalTextures CopyAndBackfillSignalInput(const FSSDSignalTextures& SignalInput)
-{
-	FSSDSignalTextures Result = SignalInput;
-
-	#if 0
-	for (int32 i = 0; i < kMaxBufferProcessingCount; ++i)
-	{
-		if (!Result.Textures[i])
-		{
-			Result.Textures[i] = Result.Textures[i - 1];
-		}
-	}
-	#endif
-
-	return Result;
-}
-
 /** Denoises a signal at constant pixel density across the viewport. */
 static void DenoiseSignalAtConstantPixelDensity(
 	FRDGBuilder& GraphBuilder,
@@ -988,7 +965,7 @@ static void DenoiseSignalAtConstantPixelDensity(
 		FSSDInjestCS::FParameters* PassParameters = GraphBuilder.AllocParameters<FSSDInjestCS::FParameters>();
 		PassParameters->CommonParameters = CommonParameters;
 		PassParameters->ConvolutionMetaData = ConvolutionMetaData;
-		PassParameters->SignalInput = CopyAndBackfillSignalInput(SignalHistory);
+		PassParameters->SignalInput = SignalHistory;
 		PassParameters->SignalOutput = CreateMultiplexedUAVs(GraphBuilder, NewSignalOutput);
 
 		FSSDInjestCS::FPermutationDomain PermutationVector;
@@ -1020,7 +997,7 @@ static void DenoiseSignalAtConstantPixelDensity(
 		PassParameters->UpscaleFactor = int32(1.0f / Settings.InputResolutionFraction);
 		PassParameters->CommonParameters = CommonParameters;
 		PassParameters->ConvolutionMetaData = ConvolutionMetaData;
-		PassParameters->SignalInput = CopyAndBackfillSignalInput(SignalHistory);
+		PassParameters->SignalInput = SignalHistory;
 		PassParameters->SignalOutput = CreateMultiplexedUAVs(GraphBuilder, NewSignalOutput);
 		
 		PassParameters->DebugOutput = CreateDebugUAV(TEXT("DebugDenoiserReconstruction"));
@@ -1061,7 +1038,7 @@ static void DenoiseSignalAtConstantPixelDensity(
 		PassParameters->CommonParameters = CommonParameters;
 		PassParameters->ConvolutionMetaData = ConvolutionMetaData;
 		PassParameters->KernelSpreadFactor = 8 * (1 << PreConvolutionId);
-		PassParameters->SignalInput = CopyAndBackfillSignalInput(SignalHistory);
+		PassParameters->SignalInput = SignalHistory;
 		PassParameters->SignalOutput = CreateMultiplexedUAVs(GraphBuilder, NewSignalOutput);
 		
 		PassParameters->DebugOutput = CreateDebugUAV(TEXT("DebugDenoiserPreConvolution"));
@@ -1135,7 +1112,7 @@ static void DenoiseSignalAtConstantPixelDensity(
 			FSSDSpatialAccumulationCS::FParameters* PassParameters = GraphBuilder.AllocParameters<FSSDSpatialAccumulationCS::FParameters>();
 			PassParameters->CommonParameters = CommonParameters;
 			PassParameters->ConvolutionMetaData = ConvolutionMetaData;
-			PassParameters->SignalInput = CopyAndBackfillSignalInput(SignalHistory);
+			PassParameters->SignalInput = SignalHistory;
 			PassParameters->SignalOutput = CreateMultiplexedUAVs(GraphBuilder, RejectionPreConvolutionSignal);
 
 			FSSDSpatialAccumulationCS::FPermutationDomain PermutationVector;
@@ -1171,7 +1148,7 @@ static void DenoiseSignalAtConstantPixelDensity(
 		PassParameters->ConvolutionMetaData = ConvolutionMetaData;
 		PassParameters->HistoryPreExposureCorrection = View.PreExposure / View.PrevViewInfo.SceneColorPreExposure;
 
-		PassParameters->SignalInput = CopyAndBackfillSignalInput(SignalHistory);
+		PassParameters->SignalInput = SignalHistory;
 		PassParameters->HistoryRejectionSignal = RejectionPreConvolutionSignal;
 		PassParameters->SignalHistoryOutput = CreateMultiplexedUAVs(GraphBuilder, SignalOutput);
 		
@@ -1206,7 +1183,7 @@ static void DenoiseSignalAtConstantPixelDensity(
 			PrevFrameHistory->SafeRelease();
 		} // for (uint32 BatchedSignalId = 0; BatchedSignalId < Settings.SignalBatchSize; BatchedSignalId++)
 
-		PassParameters->PrevHistory = CopyAndBackfillSignalInput(PassParameters->PrevHistory);
+		PassParameters->PrevHistory = PassParameters->PrevHistory;
 		
 		PassParameters->DebugOutput = CreateDebugUAV(TEXT("DebugDenoiserTemporalAccumulation"));
 
@@ -1234,7 +1211,7 @@ static void DenoiseSignalAtConstantPixelDensity(
 		PassParameters->KernelSpreadFactor = Settings.HistoryConvolutionKernelSpreadFactor;
 		PassParameters->CommonParameters = CommonParameters;
 		PassParameters->ConvolutionMetaData = ConvolutionMetaData;
-		PassParameters->SignalInput = CopyAndBackfillSignalInput(SignalHistory);
+		PassParameters->SignalInput = SignalHistory;
 		PassParameters->SignalOutput = CreateMultiplexedUAVs(GraphBuilder, SignalOutput);
 
 		FSSDSpatialAccumulationCS::FPermutationDomain PermutationVector;
@@ -1331,7 +1308,7 @@ static void DenoiseSignalAtConstantPixelDensity(
 
 		FSSDSpatialAccumulationCS::FParameters* PassParameters = GraphBuilder.AllocParameters<FSSDSpatialAccumulationCS::FParameters>();
 		PassParameters->CommonParameters = CommonParameters;
-		PassParameters->SignalInput = CopyAndBackfillSignalInput(SignalHistory);
+		PassParameters->SignalInput = SignalHistory;
 		PassParameters->SignalOutput = CreateMultiplexedUAVs(GraphBuilder, *OutputSignal);
 
 		FSSDSpatialAccumulationCS::FPermutationDomain PermutationVector;
