Fixing build break.

diff --git a/Engine/Source/Runtime/Renderer/Private/IndirectLightRendering.cpp b/Engine/Source/Runtime/Renderer/Private/IndirectLightRendering.cpp
--- a/Engine/Source/Runtime/Renderer/Private/IndirectLightRendering.cpp
+++ b/Engine/Source/Runtime/Renderer/Private/IndirectLightRendering.cpp
@@ -178,11 +178,15 @@ void FDeferredShadingSceneRenderer::RenderDiffuseIndirectAndAmbientOcclusion(FRH
 		// Render RTAO that override any technic.
 		if (bApplyRTAO)
 		{
+			FRDGTextureRef AmbientOcclusionMask = nullptr;
+
 			RenderRayTracingAmbientOcclusion(
 				GraphBuilder,
 				View,
 				SceneTextures,
-				&DenoiserOutputs.AmbientOcclusionMask);
+				&AmbientOcclusionMask);
+
+			DenoiserOutputs.AmbientOcclusionMask = AmbientOcclusionMask;
 		}
 
 		// Extract the dynamic AO for application of AO beyond RenderDiffuseIndirectAndAmbientOcclusion()
diff --git a/Engine/Source/Runtime/Renderer/Private/ScreenSpaceDenoise.cpp b/Engine/Source/Runtime/Renderer/Private/ScreenSpaceDenoise.cpp
--- a/Engine/Source/Runtime/Renderer/Private/ScreenSpaceDenoise.cpp
+++ b/Engine/Source/Runtime/Renderer/Private/ScreenSpaceDenoise.cpp
@@ -1183,8 +1183,6 @@ static void DenoiseSignalAtConstantPixelDensity(
 			PrevFrameHistory->SafeRelease();
 		} // for (uint32 BatchedSignalId = 0; BatchedSignalId < Settings.SignalBatchSize; BatchedSignalId++)
 
-		PassParameters->PrevHistory = PassParameters->PrevHistory;
-		
 		PassParameters->DebugOutput = CreateDebugUAV(TEXT("DebugDenoiserTemporalAccumulation"));
 
 		FComputeShaderUtils::AddPass(
