Support eye adaptation material expression in deferred decal and post process materials

diff --git a/Engine/Source/Runtime/Engine/Private/Materials/HLSLMaterialTranslator.h b/Engine/Source/Runtime/Engine/Private/Materials/HLSLMaterialTranslator.h
--- a/Engine/Source/Runtime/Engine/Private/Materials/HLSLMaterialTranslator.h
+++ b/Engine/Source/Runtime/Engine/Private/Materials/HLSLMaterialTranslator.h
@@ -6848,7 +6848,7 @@ protected:
 
 		if( ShaderFrequency != SF_Pixel )
 		{
-			NonPixelShaderExpressionError();
+			return NonPixelShaderExpressionError();
 		}
 
 		MaterialCompilationOutput.bUsesEyeAdaptation = true;
diff --git a/Engine/Source/Runtime/Renderer/Private/PostProcess/SceneRenderTargets.cpp b/Engine/Source/Runtime/Renderer/Private/PostProcess/SceneRenderTargets.cpp
--- a/Engine/Source/Runtime/Renderer/Private/PostProcess/SceneRenderTargets.cpp
+++ b/Engine/Source/Runtime/Renderer/Private/PostProcess/SceneRenderTargets.cpp
@@ -3104,11 +3104,26 @@ TUniformBufferRef<FSceneTexturesUniformParameters> CreateSceneTextureUniformBuff
 	return TUniformBufferRef<FSceneTexturesUniformParameters>::CreateUniformBufferImmediate(SceneTextureParameters, UniformBuffer_SingleDraw);
 }
 
+template< typename TRHICmdList >
+TUniformBufferRef<FSceneTexturesUniformParameters> CreateSceneTextureUniformBufferSingleDrawWithView(TRHICmdList& RHICmdList, ESceneTextureSetupMode SceneTextureSetupMode, const FViewInfo& View)
+{
+	FSceneRenderTargets& SceneContext = FSceneRenderTargets::Get(RHICmdList);
+	FSceneTexturesUniformParameters SceneTextureParameters;
+	SetupSceneTextureUniformParameters(SceneContext, View.FeatureLevel, SceneTextureSetupMode, SceneTextureParameters);
+	SceneTextureParameters.EyeAdaptation = GetEyeAdaptation(View);
+	return TUniformBufferRef<FSceneTexturesUniformParameters>::CreateUniformBufferImmediate(SceneTextureParameters, UniformBuffer_SingleDraw);
+}
+
 #define IMPLEMENT_CreateSceneTextureUniformBuffer( TRHICmdList ) \
 	template TUniformBufferRef<FSceneTexturesUniformParameters> CreateSceneTextureUniformBufferSingleDraw< TRHICmdList >(\
 		TRHICmdList& RHICmdList,						\
 		ESceneTextureSetupMode SceneTextureSetupMode,	\
 		ERHIFeatureLevel::Type FeatureLevel				\
+	); \
+	template TUniformBufferRef<FSceneTexturesUniformParameters> CreateSceneTextureUniformBufferSingleDrawWithView< TRHICmdList >(\
+		TRHICmdList& RHICmdList,						\
+		ESceneTextureSetupMode SceneTextureSetupMode,	\
+		const FViewInfo& View							\
 	);
 
 IMPLEMENT_CreateSceneTextureUniformBuffer( FRHICommandList );
diff --git a/Engine/Source/Runtime/Renderer/Private/ShaderBaseClasses.cpp b/Engine/Source/Runtime/Renderer/Private/ShaderBaseClasses.cpp
--- a/Engine/Source/Runtime/Renderer/Private/ShaderBaseClasses.cpp
+++ b/Engine/Source/Runtime/Renderer/Private/ShaderBaseClasses.cpp
@@ -326,6 +326,22 @@ void FMaterialShader::SetParameters(
 	SceneTextureParameters.Set(RHICmdList, ShaderRHI, View.FeatureLevel, SceneTextureSetupMode);
 }
 
+template<typename TRHIShader>
+void FMaterialShader::SetParameters(
+	FRHICommandList& RHICmdList,
+	TRHIShader* ShaderRHI,
+	const FMaterialRenderProxy* MaterialRenderProxy,
+	const FMaterial& Material,
+	const FViewInfo& View,
+	const TUniformBufferRef<FViewUniformShaderParameters>& ViewUniformBuffer,
+	ESceneTextureSetupMode SceneTextureSetupMode)
+{
+	SetViewParameters(RHICmdList, ShaderRHI, View, ViewUniformBuffer);
+	FMaterialShader::SetParametersInner(RHICmdList, ShaderRHI, MaterialRenderProxy, Material, View);
+
+	SceneTextureParameters.SetWithView(RHICmdList, ShaderRHI, SceneTextureSetupMode, View);
+}
+
 // Doxygen struggles to parse these explicit specializations. Just ignore them for now.
 #if !UE_BUILD_DOCS
 
@@ -359,6 +375,15 @@ IMPLEMENT_MATERIAL_SHADER_SetParametersInner(FRHIComputeShader);
 		const FSceneView& View,							\
 		const TUniformBufferRef<FViewUniformShaderParameters>& ViewUniformBuffer, \
 		ESceneTextureSetupMode SceneTextureSetupMode		\
+	); \
+	template RENDERER_API void FMaterialShader::SetParameters< TRHIShader >( \
+		FRHICommandList& RHICmdList,					\
+		TRHIShader* ShaderRHI,							\
+		const FMaterialRenderProxy* MaterialRenderProxy,\
+		const FMaterial& Material,						\
+		const FViewInfo& View,							\
+		const TUniformBufferRef<FViewUniformShaderParameters>& ViewUniformBuffer, \
+		ESceneTextureSetupMode SceneTextureSetupMode		\
 	);
 
 IMPLEMENT_MATERIAL_SHADER_SetParameters(FRHIVertexShader);
diff --git a/Engine/Source/Runtime/Renderer/Public/MaterialShader.h b/Engine/Source/Runtime/Renderer/Public/MaterialShader.h
--- a/Engine/Source/Runtime/Renderer/Public/MaterialShader.h
+++ b/Engine/Source/Runtime/Renderer/Public/MaterialShader.h
@@ -166,6 +166,17 @@ public:
 		const TUniformBufferRef<FViewUniformShaderParameters>& ViewUniformBuffer,
 		ESceneTextureSetupMode SceneTextureSetupMode);
 
+	/** Like SetParameters above, but takes a FViewInfo rather than FSceneView, which allows additional per-view parameters to be set */
+	template< typename TRHIShader >
+	void SetParameters(
+		FRHICommandList& RHICmdList,
+		TRHIShader* ShaderRHI,
+		const FMaterialRenderProxy* MaterialRenderProxy,
+		const FMaterial& Material,
+		const FViewInfo& View,
+		const TUniformBufferRef<FViewUniformShaderParameters>& ViewUniformBuffer,
+		ESceneTextureSetupMode SceneTextureSetupMode);
+
 	void GetShaderBindings(
 		const FScene* Scene,
 		ERHIFeatureLevel::Type FeatureLevel,
diff --git a/Engine/Source/Runtime/Renderer/Public/SceneRenderTargetParameters.h b/Engine/Source/Runtime/Renderer/Public/SceneRenderTargetParameters.h
--- a/Engine/Source/Runtime/Renderer/Public/SceneRenderTargetParameters.h
+++ b/Engine/Source/Runtime/Renderer/Public/SceneRenderTargetParameters.h
@@ -101,6 +101,9 @@ extern RENDERER_API TUniformBufferRef<FSceneTexturesUniformParameters> CreateSce
 template< typename TRHICmdList >
 RENDERER_API TUniformBufferRef<FSceneTexturesUniformParameters> CreateSceneTextureUniformBufferSingleDraw(TRHICmdList& RHICmdList, ESceneTextureSetupMode SceneTextureSetupMode, ERHIFeatureLevel::Type FeatureLevel);
 
+template< typename TRHICmdList >
+RENDERER_API TUniformBufferRef<FSceneTexturesUniformParameters> CreateSceneTextureUniformBufferSingleDrawWithView(TRHICmdList& RHICmdList, ESceneTextureSetupMode SceneTextureSetupMode, const FViewInfo& View);
+
 BEGIN_GLOBAL_SHADER_PARAMETER_STRUCT(FMobileSceneTextureUniformParameters, RENDERER_API)
 	SHADER_PARAMETER_TEXTURE(Texture2D, SceneColorTexture)
 	SHADER_PARAMETER_SAMPLER(SamplerState, SceneColorTextureSampler)
@@ -154,6 +157,22 @@ public:
 		}
 	}
 
+	template< typename ShaderRHIParamRef, typename TRHICmdList >
+	void SetWithView(TRHICmdList& RHICmdList, const ShaderRHIParamRef& ShaderRHI, ESceneTextureSetupMode SetupMode, const FViewInfo& View) const
+	{
+		if (FSceneInterface::GetShadingPath(View.FeatureLevel) == EShadingPath::Deferred && SceneTexturesUniformBuffer.IsBound())
+		{
+			TUniformBufferRef<FSceneTexturesUniformParameters> UniformBuffer = CreateSceneTextureUniformBufferSingleDrawWithView(RHICmdList, SetupMode, View);
+			SetUniformBufferParameter(RHICmdList, ShaderRHI, SceneTexturesUniformBuffer, UniformBuffer);
+		}
+
+		if (FSceneInterface::GetShadingPath(View.FeatureLevel) == EShadingPath::Mobile && MobileSceneTexturesUniformBuffer.IsBound())
+		{
+			TUniformBufferRef<FMobileSceneTextureUniformParameters> UniformBuffer = CreateMobileSceneTextureUniformBufferSingleDraw(RHICmdList, View.FeatureLevel);
+			SetUniformBufferParameter(RHICmdList, ShaderRHI, MobileSceneTexturesUniformBuffer, UniformBuffer);
+		}
+	}
+
 	/** Serializer. */
 	friend FArchive& operator<<(FArchive& Ar,FSceneTextureShaderParameters& P)
 	{
