[RENDERING] [!] Fix incomplete type compile error. Moved the implementation of FSceneTextureShaderParameters::SetWithView into the .cpp where it is used (there's only one place). There doesn't seem to be a trivial solution to solve this (i.e. adding a header include). This probably needs refactoring.

diff --git a/Engine/Source/Runtime/Renderer/Private/ShaderBaseClasses.cpp b/Engine/Source/Runtime/Renderer/Private/ShaderBaseClasses.cpp
--- a/Engine/Source/Runtime/Renderer/Private/ShaderBaseClasses.cpp
+++ b/Engine/Source/Runtime/Renderer/Private/ShaderBaseClasses.cpp
@@ -326,6 +326,22 @@ void FMaterialShader::SetParameters(
 	SceneTextureParameters.Set(RHICmdList, ShaderRHI, View.FeatureLevel, SceneTextureSetupMode);
 }
 
+template <typename ShaderRHIParamRef, typename TRHICmdList>
+void FSceneTextureShaderParameters::SetWithView(TRHICmdList& RHICmdList, const ShaderRHIParamRef& ShaderRHI, ESceneTextureSetupMode SetupMode, const FViewInfo& View) const
+{
+	if (FSceneInterface::GetShadingPath(View.FeatureLevel) == EShadingPath::Deferred && SceneTexturesUniformBuffer.IsBound())
+	{
+		TUniformBufferRef<FSceneTexturesUniformParameters> UniformBuffer = CreateSceneTextureUniformBufferSingleDrawWithView(RHICmdList, SetupMode, View);
+		SetUniformBufferParameter(RHICmdList, ShaderRHI, SceneTexturesUniformBuffer, UniformBuffer);
+	}
+
+	if (FSceneInterface::GetShadingPath(View.FeatureLevel) == EShadingPath::Mobile && MobileSceneTexturesUniformBuffer.IsBound())
+	{
+		TUniformBufferRef<FMobileSceneTextureUniformParameters> UniformBuffer = CreateMobileSceneTextureUniformBufferSingleDraw(RHICmdList, View.FeatureLevel);
+		SetUniformBufferParameter(RHICmdList, ShaderRHI, MobileSceneTexturesUniformBuffer, UniformBuffer);
+	}
+}
+
 template<typename TRHIShader>
 void FMaterialShader::SetParameters(
 	FRHICommandList& RHICmdList,
diff --git a/Engine/Source/Runtime/Renderer/Public/SceneRenderTargetParameters.h b/Engine/Source/Runtime/Renderer/Public/SceneRenderTargetParameters.h
--- a/Engine/Source/Runtime/Renderer/Public/SceneRenderTargetParameters.h
+++ b/Engine/Source/Runtime/Renderer/Public/SceneRenderTargetParameters.h
@@ -158,20 +158,7 @@ public:
 	}
 
 	template< typename ShaderRHIParamRef, typename TRHICmdList >
-	void SetWithView(TRHICmdList& RHICmdList, const ShaderRHIParamRef& ShaderRHI, ESceneTextureSetupMode SetupMode, const FViewInfo& View) const
-	{
-		if (FSceneInterface::GetShadingPath(View.FeatureLevel) == EShadingPath::Deferred && SceneTexturesUniformBuffer.IsBound())
-		{
-			TUniformBufferRef<FSceneTexturesUniformParameters> UniformBuffer = CreateSceneTextureUniformBufferSingleDrawWithView(RHICmdList, SetupMode, View);
-			SetUniformBufferParameter(RHICmdList, ShaderRHI, SceneTexturesUniformBuffer, UniformBuffer);
-		}
-
-		if (FSceneInterface::GetShadingPath(View.FeatureLevel) == EShadingPath::Mobile && MobileSceneTexturesUniformBuffer.IsBound())
-		{
-			TUniformBufferRef<FMobileSceneTextureUniformParameters> UniformBuffer = CreateMobileSceneTextureUniformBufferSingleDraw(RHICmdList, View.FeatureLevel);
-			SetUniformBufferParameter(RHICmdList, ShaderRHI, MobileSceneTexturesUniformBuffer, UniformBuffer);
-		}
-	}
+	void SetWithView(TRHICmdList& RHICmdList, const ShaderRHIParamRef& ShaderRHI, ESceneTextureSetupMode SetupMode, const FViewInfo& View) const;
 
 	/** Serializer. */
 	friend FArchive& operator<<(FArchive& Ar,FSceneTextureShaderParameters& P)
