Somewhat cleanup usage of private FViewInfo type in public FSceneTextureShaderParameters interface

diff --git a/Engine/Source/Runtime/Renderer/Private/PostProcess/SceneRenderTargets.cpp b/Engine/Source/Runtime/Renderer/Private/PostProcess/SceneRenderTargets.cpp
--- a/Engine/Source/Runtime/Renderer/Private/PostProcess/SceneRenderTargets.cpp
+++ b/Engine/Source/Runtime/Renderer/Private/PostProcess/SceneRenderTargets.cpp
@@ -3078,27 +3078,12 @@ TUniformBufferRef<FSceneTexturesUniformParameters> CreateSceneTextureUniformBuff
 	return TUniformBufferRef<FSceneTexturesUniformParameters>::CreateUniformBufferImmediate(SceneTextureParameters, UniformBuffer_SingleDraw);
 }
 
-template< typename TRHICmdList >
-TUniformBufferRef<FSceneTexturesUniformParameters> CreateSceneTextureUniformBufferSingleDrawWithView(TRHICmdList& RHICmdList, ESceneTextureSetupMode SceneTextureSetupMode, const FViewInfo& View)
-{
-	FSceneRenderTargets& SceneContext = FSceneRenderTargets::Get(RHICmdList);
-	FSceneTexturesUniformParameters SceneTextureParameters;
-	SetupSceneTextureUniformParameters(SceneContext, View.FeatureLevel, SceneTextureSetupMode, SceneTextureParameters);
-	SceneTextureParameters.EyeAdaptation = GetEyeAdaptation(View);
-	return TUniformBufferRef<FSceneTexturesUniformParameters>::CreateUniformBufferImmediate(SceneTextureParameters, UniformBuffer_SingleDraw);
-}
-
 #define IMPLEMENT_CreateSceneTextureUniformBuffer( TRHICmdList ) \
 	template TUniformBufferRef<FSceneTexturesUniformParameters> CreateSceneTextureUniformBufferSingleDraw< TRHICmdList >(\
 		TRHICmdList& RHICmdList,						\
 		ESceneTextureSetupMode SceneTextureSetupMode,	\
 		ERHIFeatureLevel::Type FeatureLevel				\
-	); \
-	template TUniformBufferRef<FSceneTexturesUniformParameters> CreateSceneTextureUniformBufferSingleDrawWithView< TRHICmdList >(\
-		TRHICmdList& RHICmdList,						\
-		ESceneTextureSetupMode SceneTextureSetupMode,	\
-		const FViewInfo& View							\
-	);
+	)
 
 IMPLEMENT_CreateSceneTextureUniformBuffer( FRHICommandList );
 IMPLEMENT_CreateSceneTextureUniformBuffer( FRHICommandListImmediate );
diff --git a/Engine/Source/Runtime/Renderer/Private/ShaderBaseClasses.cpp b/Engine/Source/Runtime/Renderer/Private/ShaderBaseClasses.cpp
--- a/Engine/Source/Runtime/Renderer/Private/ShaderBaseClasses.cpp
+++ b/Engine/Source/Runtime/Renderer/Private/ShaderBaseClasses.cpp
@@ -326,22 +326,6 @@ void FMaterialShader::SetParameters(
 	SceneTextureParameters.Set(RHICmdList, ShaderRHI, View.FeatureLevel, SceneTextureSetupMode);
 }
 
-template <typename ShaderRHIParamRef, typename TRHICmdList>
-void FSceneTextureShaderParameters::SetWithView(TRHICmdList& RHICmdList, const ShaderRHIParamRef& ShaderRHI, ESceneTextureSetupMode SetupMode, const FViewInfo& View) const
-{
-	if (FSceneInterface::GetShadingPath(View.FeatureLevel) == EShadingPath::Deferred && SceneTexturesUniformBuffer.IsBound())
-	{
-		TUniformBufferRef<FSceneTexturesUniformParameters> UniformBuffer = CreateSceneTextureUniformBufferSingleDrawWithView(RHICmdList, SetupMode, View);
-		SetUniformBufferParameter(RHICmdList, ShaderRHI, SceneTexturesUniformBuffer, UniformBuffer);
-	}
-
-	if (FSceneInterface::GetShadingPath(View.FeatureLevel) == EShadingPath::Mobile && MobileSceneTexturesUniformBuffer.IsBound())
-	{
-		TUniformBufferRef<FMobileSceneTextureUniformParameters> UniformBuffer = CreateMobileSceneTextureUniformBufferSingleDraw(RHICmdList, View.FeatureLevel);
-		SetUniformBufferParameter(RHICmdList, ShaderRHI, MobileSceneTexturesUniformBuffer, UniformBuffer);
-	}
-}
-
 template<typename TRHIShader>
 void FMaterialShader::SetParameters(
 	FRHICommandList& RHICmdList,
@@ -355,7 +339,21 @@ void FMaterialShader::SetParameters(
 	SetViewParameters(RHICmdList, ShaderRHI, View, ViewUniformBuffer);
 	FMaterialShader::SetParametersInner(RHICmdList, ShaderRHI, MaterialRenderProxy, Material, View);
 
-	SceneTextureParameters.SetWithView(RHICmdList, ShaderRHI, SceneTextureSetupMode, View);
+	if (FSceneInterface::GetShadingPath(View.FeatureLevel) == EShadingPath::Deferred && SceneTextureParameters.GetUniformBufferParameter().IsBound())
+	{
+		FSceneRenderTargets& SceneContext = FSceneRenderTargets::Get(RHICmdList);
+		FSceneTexturesUniformParameters UniformParameters;
+		SetupSceneTextureUniformParameters(SceneContext, View.FeatureLevel, SceneTextureSetupMode, UniformParameters);
+		UniformParameters.EyeAdaptation = GetEyeAdaptation(View);
+		TUniformBufferRef<FSceneTexturesUniformParameters> UniformBuffer = TUniformBufferRef<FSceneTexturesUniformParameters>::CreateUniformBufferImmediate(UniformParameters, UniformBuffer_SingleDraw);
+		SetUniformBufferParameter(RHICmdList, ShaderRHI, SceneTextureParameters.GetUniformBufferParameter(), UniformBuffer);
+	}
+
+	if (FSceneInterface::GetShadingPath(View.FeatureLevel) == EShadingPath::Mobile && SceneTextureParameters.GetMobileUniformBufferParameter().IsBound())
+	{
+		TUniformBufferRef<FMobileSceneTextureUniformParameters> UniformBuffer = CreateMobileSceneTextureUniformBufferSingleDraw(RHICmdList, View.FeatureLevel);
+		SetUniformBufferParameter(RHICmdList, ShaderRHI, SceneTextureParameters.GetMobileUniformBufferParameter(), UniformBuffer);
+	}
 }
 
 // Doxygen struggles to parse these explicit specializations. Just ignore them for now.
diff --git a/Engine/Source/Runtime/Renderer/Public/SceneRenderTargetParameters.h b/Engine/Source/Runtime/Renderer/Public/SceneRenderTargetParameters.h
--- a/Engine/Source/Runtime/Renderer/Public/SceneRenderTargetParameters.h
+++ b/Engine/Source/Runtime/Renderer/Public/SceneRenderTargetParameters.h
@@ -101,9 +101,6 @@ extern RENDERER_API TUniformBufferRef<FSceneTexturesUniformParameters> CreateSce
 template< typename TRHICmdList >
 RENDERER_API TUniformBufferRef<FSceneTexturesUniformParameters> CreateSceneTextureUniformBufferSingleDraw(TRHICmdList& RHICmdList, ESceneTextureSetupMode SceneTextureSetupMode, ERHIFeatureLevel::Type FeatureLevel);
 
-template< typename TRHICmdList >
-RENDERER_API TUniformBufferRef<FSceneTexturesUniformParameters> CreateSceneTextureUniformBufferSingleDrawWithView(TRHICmdList& RHICmdList, ESceneTextureSetupMode SceneTextureSetupMode, const FViewInfo& View);
-
 BEGIN_GLOBAL_SHADER_PARAMETER_STRUCT(FMobileSceneTextureUniformParameters, RENDERER_API)
 	SHADER_PARAMETER_TEXTURE(Texture2D, SceneColorTexture)
 	SHADER_PARAMETER_SAMPLER(SamplerState, SceneColorTextureSampler)
@@ -141,6 +138,9 @@ public:
 		BindSceneTextureUniformBufferDependentOnShadingPath(Initializer, SceneTexturesUniformBuffer, MobileSceneTexturesUniformBuffer);
 	}
 
+	const TShaderUniformBufferParameter<FSceneTexturesUniformParameters>& GetUniformBufferParameter() const { return SceneTexturesUniformBuffer; }
+	const TShaderUniformBufferParameter<FMobileSceneTextureUniformParameters>& GetMobileUniformBufferParameter() const { return MobileSceneTexturesUniformBuffer; }
+
 	template< typename ShaderRHIParamRef, typename TRHICmdList >
 	void Set(TRHICmdList& RHICmdList, const ShaderRHIParamRef& ShaderRHI, ERHIFeatureLevel::Type FeatureLevel, ESceneTextureSetupMode SetupMode) const
 	{
@@ -157,9 +157,6 @@ public:
 		}
 	}
 
-	template< typename ShaderRHIParamRef, typename TRHICmdList >
-	void SetWithView(TRHICmdList& RHICmdList, const ShaderRHIParamRef& ShaderRHI, ESceneTextureSetupMode SetupMode, const FViewInfo& View) const;
-
 	/** Serializer. */
 	friend FArchive& operator<<(FArchive& Ar,FSceneTextureShaderParameters& P)
 	{
