Only store a single parameter reference inside FSceneTextureShaderParameters, since should never need both deferred and mobile reference in the same shader

diff --git a/Engine/Shaders/Public/ShaderVersion.ush b/Engine/Shaders/Public/ShaderVersion.ush
--- a/Engine/Shaders/Public/ShaderVersion.ush
+++ b/Engine/Shaders/Public/ShaderVersion.ush
@@ -2,4 +2,4 @@
 // Each time the console command is executed it generates a new GUID. As this file is included
 // in Platform.ush (which should be included in any shader) it allows to invalidate the shader DDC.
 // 
-// GUID = 3A7029DE8C044F0A97C52C4CD3F6182E
\ No newline at end of file
+// GUID = 75367DA1-93B1-4CF4-9732-5E12FD8972B6
\ No newline at end of file
diff --git a/Engine/Source/Runtime/Renderer/Private/DepthRendering.h b/Engine/Source/Runtime/Renderer/Private/DepthRendering.h
--- a/Engine/Source/Runtime/Renderer/Private/DepthRendering.h
+++ b/Engine/Source/Runtime/Renderer/Private/DepthRendering.h
@@ -56,7 +56,7 @@ protected:
 	TDepthOnlyVS(const FMeshMaterialShaderType::CompiledShaderInitializerType& Initializer) :
 		FMeshMaterialShader(Initializer)
 	{
-		BindSceneTextureUniformBufferDependentOnShadingPath(Initializer, PassUniformBuffer, PassUniformBuffer);
+		BindSceneTextureUniformBufferDependentOnShadingPath(Initializer, PassUniformBuffer);
 	}
 
 public:
@@ -154,7 +154,7 @@ public:
 		FMeshMaterialShader(Initializer)
 	{
 		MobileColorValue.Bind(Initializer.ParameterMap, TEXT("MobileColorValue"));
-		BindSceneTextureUniformBufferDependentOnShadingPath(Initializer, PassUniformBuffer, PassUniformBuffer);
+		BindSceneTextureUniformBufferDependentOnShadingPath(Initializer, PassUniformBuffer);
 	}
 
 	FDepthOnlyPS() {}
diff --git a/Engine/Source/Runtime/Renderer/Private/PostProcess/SceneRenderTargets.cpp b/Engine/Source/Runtime/Renderer/Private/PostProcess/SceneRenderTargets.cpp
--- a/Engine/Source/Runtime/Renderer/Private/PostProcess/SceneRenderTargets.cpp
+++ b/Engine/Source/Runtime/Renderer/Private/PostProcess/SceneRenderTargets.cpp
@@ -3154,8 +3154,7 @@ IMPLEMENT_CreateMobileSceneTextureUniformBuffer( FRHIAsyncComputeCommandListImme
 
 void BindSceneTextureUniformBufferDependentOnShadingPath(
 	const FShader::CompiledShaderInitializerType& Initializer, 
-	FShaderUniformBufferParameter& SceneTexturesUniformBuffer, 
-	FShaderUniformBufferParameter& MobileSceneTexturesUniformBuffer)
+	FShaderUniformBufferParameter& SceneTexturesUniformBuffer)
 {
 	const ERHIFeatureLevel::Type FeatureLevel = GetMaxSupportedFeatureLevel((EShaderPlatform)Initializer.Target.Platform);
 		
@@ -3164,10 +3163,9 @@ void BindSceneTextureUniformBufferDependentOnShadingPath(
 		SceneTexturesUniformBuffer.Bind(Initializer.ParameterMap, FSceneTexturesUniformParameters::StaticStructMetadata.GetShaderVariableName());
 		checkfSlow(!Initializer.ParameterMap.ContainsParameterAllocation(FMobileSceneTextureUniformParameters::StaticStructMetadata.GetShaderVariableName()), TEXT("Shader for Deferred shading path tried to bind FMobileSceneTextureUniformParameters which is only available in the mobile shading path: %s"), Initializer.Type->GetName());
 	}
-		
-	if (FSceneInterface::GetShadingPath(FeatureLevel) == EShadingPath::Mobile)
+	else if (FSceneInterface::GetShadingPath(FeatureLevel) == EShadingPath::Mobile)
 	{
-		MobileSceneTexturesUniformBuffer.Bind(Initializer.ParameterMap, FMobileSceneTextureUniformParameters::StaticStructMetadata.GetShaderVariableName());
+		SceneTexturesUniformBuffer.Bind(Initializer.ParameterMap, FMobileSceneTextureUniformParameters::StaticStructMetadata.GetShaderVariableName());
 		checkfSlow(!Initializer.ParameterMap.ContainsParameterAllocation(FSceneTexturesUniformParameters::StaticStructMetadata.GetShaderVariableName()), TEXT("Shader for Mobile shading path tried to bind FSceneTexturesUniformParameters which is only available in the deferred shading path: %s"), Initializer.Type->GetName());
 	}
 }
diff --git a/Engine/Source/Runtime/Renderer/Private/ShaderBaseClasses.cpp b/Engine/Source/Runtime/Renderer/Private/ShaderBaseClasses.cpp
--- a/Engine/Source/Runtime/Renderer/Private/ShaderBaseClasses.cpp
+++ b/Engine/Source/Runtime/Renderer/Private/ShaderBaseClasses.cpp
@@ -339,20 +339,22 @@ void FMaterialShader::SetParameters(
 	SetViewParameters(RHICmdList, ShaderRHI, View, ViewUniformBuffer);
 	FMaterialShader::SetParametersInner(RHICmdList, ShaderRHI, MaterialRenderProxy, Material, View);
 
-	if (FSceneInterface::GetShadingPath(View.FeatureLevel) == EShadingPath::Deferred && SceneTextureParameters.GetUniformBufferParameter().IsBound())
+	if (SceneTextureParameters.IsBound())
 	{
-		FSceneRenderTargets& SceneContext = FSceneRenderTargets::Get(RHICmdList);
-		FSceneTexturesUniformParameters UniformParameters;
-		SetupSceneTextureUniformParameters(SceneContext, View.FeatureLevel, SceneTextureSetupMode, UniformParameters);
-		UniformParameters.EyeAdaptation = GetEyeAdaptation(View);
-		TUniformBufferRef<FSceneTexturesUniformParameters> UniformBuffer = TUniformBufferRef<FSceneTexturesUniformParameters>::CreateUniformBufferImmediate(UniformParameters, UniformBuffer_SingleDraw);
-		SetUniformBufferParameter(RHICmdList, ShaderRHI, SceneTextureParameters.GetUniformBufferParameter(), UniformBuffer);
-	}
-
-	if (FSceneInterface::GetShadingPath(View.FeatureLevel) == EShadingPath::Mobile && SceneTextureParameters.GetMobileUniformBufferParameter().IsBound())
-	{
-		TUniformBufferRef<FMobileSceneTextureUniformParameters> UniformBuffer = CreateMobileSceneTextureUniformBufferSingleDraw(RHICmdList, View.FeatureLevel);
-		SetUniformBufferParameter(RHICmdList, ShaderRHI, SceneTextureParameters.GetMobileUniformBufferParameter(), UniformBuffer);
+		if (FSceneInterface::GetShadingPath(View.FeatureLevel) == EShadingPath::Deferred)
+		{
+			FSceneRenderTargets& SceneContext = FSceneRenderTargets::Get(RHICmdList);
+			FSceneTexturesUniformParameters UniformParameters;
+			SetupSceneTextureUniformParameters(SceneContext, View.FeatureLevel, SceneTextureSetupMode, UniformParameters);
+			UniformParameters.EyeAdaptation = GetEyeAdaptation(View);
+			TUniformBufferRef<FSceneTexturesUniformParameters> UniformBuffer = TUniformBufferRef<FSceneTexturesUniformParameters>::CreateUniformBufferImmediate(UniformParameters, UniformBuffer_SingleDraw);
+			SetUniformBufferParameter(RHICmdList, ShaderRHI, SceneTextureParameters.GetUniformBufferParameter(), UniformBuffer);
+		}
+		else if (FSceneInterface::GetShadingPath(View.FeatureLevel) == EShadingPath::Mobile)
+		{
+			TUniformBufferRef<FMobileSceneTextureUniformParameters> UniformBuffer = CreateMobileSceneTextureUniformBufferSingleDraw(RHICmdList, View.FeatureLevel);
+			SetUniformBufferParameter(RHICmdList, ShaderRHI, SceneTextureParameters.GetUniformBufferParameter(), UniformBuffer);
+		}
 	}
 }
 
diff --git a/Engine/Source/Runtime/Renderer/Public/SceneRenderTargetParameters.h b/Engine/Source/Runtime/Renderer/Public/SceneRenderTargetParameters.h
--- a/Engine/Source/Runtime/Renderer/Public/SceneRenderTargetParameters.h
+++ b/Engine/Source/Runtime/Renderer/Public/SceneRenderTargetParameters.h
@@ -125,8 +125,7 @@ RENDERER_API TUniformBufferRef<FMobileSceneTextureUniformParameters> CreateMobil
 
 extern RENDERER_API void BindSceneTextureUniformBufferDependentOnShadingPath(
 	const FShader::CompiledShaderInitializerType& Initializer,
-	FShaderUniformBufferParameter& SceneTexturesUniformBuffer,
-	FShaderUniformBufferParameter& MobileSceneTexturesUniformBuffer);
+	FShaderUniformBufferParameter& SceneTexturesUniformBuffer);
 
 /** Encapsulates scene texture shader parameter bindings. */
 class RENDERER_API FSceneTextureShaderParameters
@@ -135,25 +134,26 @@ public:
 	/** Binds the parameters using a compiled shader's parameter map. */
 	void Bind(const FShader::CompiledShaderInitializerType& Initializer)
 	{
-		BindSceneTextureUniformBufferDependentOnShadingPath(Initializer, SceneTexturesUniformBuffer, MobileSceneTexturesUniformBuffer);
+		BindSceneTextureUniformBufferDependentOnShadingPath(Initializer, SceneTexturesUniformBuffer);
 	}
 
-	const TShaderUniformBufferParameter<FSceneTexturesUniformParameters>& GetUniformBufferParameter() const { return SceneTexturesUniformBuffer; }
-	const TShaderUniformBufferParameter<FMobileSceneTextureUniformParameters>& GetMobileUniformBufferParameter() const { return MobileSceneTexturesUniformBuffer; }
+	const FShaderUniformBufferParameter& GetUniformBufferParameter() const { return SceneTexturesUniformBuffer; }
 
 	template< typename ShaderRHIParamRef, typename TRHICmdList >
 	void Set(TRHICmdList& RHICmdList, const ShaderRHIParamRef& ShaderRHI, ERHIFeatureLevel::Type FeatureLevel, ESceneTextureSetupMode SetupMode) const
 	{
-		if (FSceneInterface::GetShadingPath(FeatureLevel) == EShadingPath::Deferred && SceneTexturesUniformBuffer.IsBound())
+		if (SceneTexturesUniformBuffer.IsBound())
 		{
-			TUniformBufferRef<FSceneTexturesUniformParameters> UniformBuffer = CreateSceneTextureUniformBufferSingleDraw(RHICmdList, SetupMode, FeatureLevel);
-			SetUniformBufferParameter(RHICmdList, ShaderRHI, SceneTexturesUniformBuffer, UniformBuffer);
-		}
-
-		if (FSceneInterface::GetShadingPath(FeatureLevel) == EShadingPath::Mobile && MobileSceneTexturesUniformBuffer.IsBound())
-		{
-			TUniformBufferRef<FMobileSceneTextureUniformParameters> UniformBuffer = CreateMobileSceneTextureUniformBufferSingleDraw(RHICmdList, FeatureLevel);
-			SetUniformBufferParameter(RHICmdList, ShaderRHI, MobileSceneTexturesUniformBuffer, UniformBuffer);
+			if (FSceneInterface::GetShadingPath(FeatureLevel) == EShadingPath::Deferred)
+			{
+				TUniformBufferRef<FSceneTexturesUniformParameters> UniformBuffer = CreateSceneTextureUniformBufferSingleDraw(RHICmdList, SetupMode, FeatureLevel);
+				SetUniformBufferParameter(RHICmdList, ShaderRHI, SceneTexturesUniformBuffer, UniformBuffer);
+			}
+			else if (FSceneInterface::GetShadingPath(FeatureLevel) == EShadingPath::Mobile)
+			{
+				TUniformBufferRef<FMobileSceneTextureUniformParameters> UniformBuffer = CreateMobileSceneTextureUniformBufferSingleDraw(RHICmdList, FeatureLevel);
+				SetUniformBufferParameter(RHICmdList, ShaderRHI, SceneTexturesUniformBuffer, UniformBuffer);
+			}
 		}
 	}
 
@@ -161,34 +161,24 @@ public:
 	friend FArchive& operator<<(FArchive& Ar,FSceneTextureShaderParameters& P)
 	{
 		Ar << P.SceneTexturesUniformBuffer;
-		Ar << P.MobileSceneTexturesUniformBuffer;
 		return Ar;
 	}
 
 	inline bool IsBound() const 
 	{ 
-		return SceneTexturesUniformBuffer.IsBound() || MobileSceneTexturesUniformBuffer.IsBound(); 
+		return SceneTexturesUniformBuffer.IsBound(); 
 	}
 
 	bool IsSameUniformParameter(const FShaderUniformBufferParameter& Parameter)
 	{
-		if (Parameter.IsBound())
+		if (Parameter.IsBound() && SceneTexturesUniformBuffer.IsBound() && Parameter.GetBaseIndex() == SceneTexturesUniformBuffer.GetBaseIndex())
 		{
-			if (SceneTexturesUniformBuffer.IsBound() && SceneTexturesUniformBuffer.GetBaseIndex() == Parameter.GetBaseIndex())
-			{
-				return true;
-			}
-
-			if (MobileSceneTexturesUniformBuffer.IsBound() && MobileSceneTexturesUniformBuffer.GetBaseIndex() == Parameter.GetBaseIndex())
-			{
-				return true;
-			}
+			return true;
 		}
 
 		return false;
 	}
 
 private:
-	TShaderUniformBufferParameter<FSceneTexturesUniformParameters> SceneTexturesUniformBuffer;
-	TShaderUniformBufferParameter<FMobileSceneTextureUniformParameters> MobileSceneTexturesUniformBuffer;
+	FShaderUniformBufferParameter SceneTexturesUniformBuffer;
 };
