Post Process Materials ported to RDG.

diff --git a/Engine/Plugins/Experimental/RemoteSession/Source/RemoteSession/Private/Channels/RemoteSessionARCameraChannel.cpp b/Engine/Plugins/Experimental/RemoteSession/Source/RemoteSession/Private/Channels/RemoteSessionARCameraChannel.cpp
--- a/Engine/Plugins/Experimental/RemoteSession/Source/RemoteSession/Private/Channels/RemoteSessionARCameraChannel.cpp
+++ b/Engine/Plugins/Experimental/RemoteSession/Source/RemoteSession/Private/Channels/RemoteSessionARCameraChannel.cpp
@@ -25,6 +25,7 @@
 #include "ScreenRendering.h"
 #include "Containers/DynamicRHIResourceArray.h"
 #include "PostProcess/SceneFilterRendering.h"
+#include "PostProcess/PostProcessMaterial.h"
 #include "PostProcessParameters.h"
 #include "EngineModule.h"
 
@@ -60,13 +61,11 @@ TAutoConsoleVariable<int32> CVarJPEGGpu(
 	TEXT("1 (default) compresses on the GPU, 0 on the CPU"),
 	ECVF_Default);
 
-/** Shaders to render our post process material */
-class FRemoteSessionARCameraVS :
-	public FMaterialShader
+class FPostProcessMaterialShader : public FMaterialShader
 {
-	DECLARE_SHADER_TYPE(FRemoteSessionARCameraVS, Material);
-
 public:
+	using FParameters = FPostProcessMaterialParameters;
+	SHADER_USE_PARAMETER_STRUCT_WITH_LEGACY_BASE(FPostProcessMaterialShader, FMaterialShader);
 
 	static bool ShouldCompilePermutation(const FMaterialShaderPermutationParameters& Parameters)
 	{
@@ -77,94 +76,63 @@ public:
 	{
 		FMaterialShader::ModifyCompilationEnvironment(Parameters.Platform, OutEnvironment);
 		OutEnvironment.SetDefine(TEXT("POST_PROCESS_MATERIAL"), 1);
+		OutEnvironment.SetDefine(TEXT("POST_PROCESS_MATERIAL_MOBILE"), 0);
 		OutEnvironment.SetDefine(TEXT("POST_PROCESS_MATERIAL_BEFORE_TONEMAP"), (Parameters.Material->GetBlendableLocation() != BL_AfterTonemapping) ? 1 : 0);
-		OutEnvironment.SetDefine(TEXT("POST_PROCESS_AR_PASSTHROUGH"), 1);
 	}
+};
 
-	FRemoteSessionARCameraVS() { }
-	FRemoteSessionARCameraVS(const ShaderMetaType::CompiledShaderInitializerType& Initializer)
-		: FMaterialShader(Initializer)
+/** Shaders to render our post process material */
+class FRemoteSessionARCameraVS :
+	public FPostProcessMaterialShader
+{
+public:
+	DECLARE_MATERIAL_SHADER(FRemoteSessionARCameraVS);
+
+	static void ModifyCompilationEnvironment(const FMaterialShaderPermutationParameters& Parameters, FShaderCompilerEnvironment& OutEnvironment)
 	{
+		FPostProcessMaterialShader::ModifyCompilationEnvironment(Parameters, OutEnvironment);
+		OutEnvironment.SetDefine(TEXT("POST_PROCESS_AR_PASSTHROUGH"), 1);
 	}
 
+	FRemoteSessionARCameraVS() = default;
+	FRemoteSessionARCameraVS(const ShaderMetaType::CompiledShaderInitializerType& Initializer)
+		: FPostProcessMaterialShader(Initializer)
+	{}
+
 	void SetParameters(FRHICommandList& RHICmdList, const FSceneView& View)
 	{
 		FRHIVertexShader* ShaderRHI = GetVertexShader();
 		FMaterialShader::SetViewParameters(RHICmdList, ShaderRHI, View, View.ViewUniformBuffer);
 	}
-
-	virtual bool Serialize(FArchive& Ar) override
-	{
-		const bool bShaderHasOutdatedParameters = FMaterialShader::Serialize(Ar);
-		return bShaderHasOutdatedParameters;
-	}
 };
 
-IMPLEMENT_MATERIAL_SHADER_TYPE(, FRemoteSessionARCameraVS, TEXT("/Engine/Private/PostProcessMaterialShaders.usf"), TEXT("MainVS_VideoOverlay"), SF_Vertex);
+IMPLEMENT_MATERIAL_SHADER(FRemoteSessionARCameraVS, "/Engine/Private/PostProcessMaterialShaders.usf", "MainVS_VideoOverlay", SF_Vertex);
 
 class FRemoteSessionARCameraPS :
-	public FMaterialShader
+	public FPostProcessMaterialShader
 {
-	DECLARE_SHADER_TYPE(FRemoteSessionARCameraPS, Material);
-
 public:
-
-	static bool ShouldCompilePermutation(const FMaterialShaderPermutationParameters& Parameters)
-	{
-		return Parameters.Material->GetMaterialDomain() == MD_PostProcess && !IsMobilePlatform(Parameters.Platform);
-	}
+	DECLARE_MATERIAL_SHADER(FRemoteSessionARCameraPS);
 
 	static void ModifyCompilationEnvironment(const FMaterialShaderPermutationParameters& Parameters, FShaderCompilerEnvironment& OutEnvironment)
 	{
-		FMaterialShader::ModifyCompilationEnvironment(Parameters.Platform, OutEnvironment);
-		OutEnvironment.SetDefine(TEXT("POST_PROCESS_MATERIAL"), 1);
+		FPostProcessMaterialShader::ModifyCompilationEnvironment(Parameters, OutEnvironment);
 		OutEnvironment.SetDefine(TEXT("OUTPUT_MOBILE_HDR"), IsMobileHDR() ? 1 : 0);
-		OutEnvironment.SetDefine(TEXT("POST_PROCESS_MATERIAL_BEFORE_TONEMAP"), (Parameters.Material->GetBlendableLocation() != BL_AfterTonemapping) ? 1 : 0);
 	}
 
-	FRemoteSessionARCameraPS() {}
-	FRemoteSessionARCameraPS(const ShaderMetaType::CompiledShaderInitializerType& Initializer) :
-		FMaterialShader(Initializer)
-	{
-		for (uint32 InputIter = 0; InputIter < ePId_Input_MAX; ++InputIter)
-		{
-			PostprocessInputParameter[InputIter].Bind(Initializer.ParameterMap, *FString::Printf(TEXT("PostprocessInput%d"), InputIter));
-			PostprocessInputParameterSampler[InputIter].Bind(Initializer.ParameterMap, *FString::Printf(TEXT("PostprocessInput%dSampler"), InputIter));
-		}
-	}
+	FRemoteSessionARCameraPS() = default;
+	FRemoteSessionARCameraPS(const ShaderMetaType::CompiledShaderInitializerType& Initializer)
+		: FPostProcessMaterialShader(Initializer)
+	{}
 
 	void SetParameters(FRHICommandList& RHICmdList, const FSceneView& View, const FMaterialRenderProxy* Material)
 	{
 		FRHIPixelShader* ShaderRHI = GetPixelShader();
 		FMaterialShader::SetParameters(RHICmdList, ShaderRHI, Material, *Material->GetMaterial(View.GetFeatureLevel()), View, View.ViewUniformBuffer, ESceneTextureSetupMode::None);
-
-		for (uint32 InputIter = 0; InputIter < ePId_Input_MAX; ++InputIter)
-		{
-			if (PostprocessInputParameter[InputIter].IsBound())
-			{
-				SetTextureParameter(
-						RHICmdList,
-						ShaderRHI,
-						PostprocessInputParameter[InputIter],
-						PostprocessInputParameterSampler[InputIter],
-						TStaticSamplerState<>::GetRHI(),
-						GBlackTexture->TextureRHI);
-			}
-		}
 	}
-
-	virtual bool Serialize(FArchive& Ar) override
-	{
-		const bool bShaderHasOutdatedParameters = FMaterialShader::Serialize(Ar);
-		return bShaderHasOutdatedParameters;
-	}
-
-private:
-	FShaderResourceParameter PostprocessInputParameter[ePId_Input_MAX];
-	FShaderResourceParameter PostprocessInputParameterSampler[ePId_Input_MAX];
 };
 
-IMPLEMENT_MATERIAL_SHADER_TYPE(, FRemoteSessionARCameraPS, TEXT("/Engine/Private/PostProcessMaterialShaders.usf"), TEXT("MainPS_VideoOverlay"), SF_Pixel);
+IMPLEMENT_MATERIAL_SHADER(FRemoteSessionARCameraPS, "/Engine/Private/PostProcessMaterialShaders.usf", "MainPS_VideoOverlay", SF_Pixel);
 
 class FARCameraSceneViewExtension :
 	public FSceneViewExtensionBase
diff --git a/Engine/Plugins/Runtime/AR/Apple/AppleARKit/Source/AppleARKit/Private/AppleARKitVideoOverlay.cpp b/Engine/Plugins/Runtime/AR/Apple/AppleARKit/Source/AppleARKit/Private/AppleARKitVideoOverlay.cpp
--- a/Engine/Plugins/Runtime/AR/Apple/AppleARKit/Source/AppleARKit/Private/AppleARKitVideoOverlay.cpp
+++ b/Engine/Plugins/Runtime/AR/Apple/AppleARKit/Source/AppleARKit/Private/AppleARKitVideoOverlay.cpp
@@ -20,6 +20,7 @@
 #include "ScreenRendering.h"
 #include "Containers/DynamicRHIResourceArray.h"
 #include "PostProcess/SceneFilterRendering.h"
+#include "PostProcess/PostProcessMaterial.h"
 #include "PostProcessParameters.h"
 #include "CommonRenderResources.h"
 #if SUPPORTS_ARKIT_1_0
@@ -38,12 +39,12 @@ FAppleARKitVideoOverlay::FAppleARKitVideoOverlay()
 	check(RenderingOverlayMaterial != nullptr);
 }
 
-// We use something similar to the PostProcessMaterial to render the color camera overlay.
 template <bool bIsMobileRenderer>
-class FARKitCameraOverlayVS : public FMaterialShader
+class TPostProcessMaterialShader : public FMaterialShader
 {
-	DECLARE_SHADER_TYPE(FARKitCameraOverlayVS, Material);
 public:
+	using FParameters = FPostProcessMaterialParameters;
+	SHADER_USE_PARAMETER_STRUCT_WITH_LEGACY_BASE(TPostProcessMaterialShader, FMaterialShader);
 
 	static bool ShouldCompilePermutation(const FMaterialShaderPermutationParameters& Parameters)
 	{
@@ -61,102 +62,71 @@ public:
 	{
 		FMaterialShader::ModifyCompilationEnvironment(Parameters.Platform, OutEnvironment);
 		OutEnvironment.SetDefine(TEXT("POST_PROCESS_MATERIAL"), 1);
+		OutEnvironment.SetDefine(TEXT("POST_PROCESS_MATERIAL_MOBILE"), bIsMobileRenderer);
 		OutEnvironment.SetDefine(TEXT("POST_PROCESS_MATERIAL_BEFORE_TONEMAP"), (Parameters.Material->GetBlendableLocation() != BL_AfterTonemapping) ? 1 : 0);
-		OutEnvironment.SetDefine(TEXT("POST_PROCESS_AR_PASSTHROUGH"), 1);
 	}
+};
 
-	FARKitCameraOverlayVS() { }
-	FARKitCameraOverlayVS(const ShaderMetaType::CompiledShaderInitializerType& Initializer)
-		: FMaterialShader(Initializer)
+// We use something similar to the PostProcessMaterial to render the color camera overlay.
+template <bool bIsMobileRenderer>
+class TARKitCameraOverlayVS : public TPostProcessMaterialShader<bIsMobileRenderer>
+{
+public:
+	DECLARE_MATERIAL_SHADER(TARKitCameraOverlayVS);
+
+	static void ModifyCompilationEnvironment(const FMaterialShaderPermutationParameters& Parameters, FShaderCompilerEnvironment& OutEnvironment)
 	{
+		TPostProcessMaterialShader::ModifyCompilationEnvironment(Parameters, OutEnvironment);
+		OutEnvironment.SetDefine(TEXT("POST_PROCESS_AR_PASSTHROUGH"), 1);
 	}
 
+	TARKitCameraOverlayVS() = default;
+	TARKitCameraOverlayVS(const ShaderMetaType::CompiledShaderInitializerType& Initializer)
+		: TPostProcessMaterialShader(Initializer)
+	{}
+
 	void SetParameters(FRHICommandList& RHICmdList, const FSceneView View)
 	{
 		FRHIVertexShader* ShaderRHI = GetVertexShader();
-		FMaterialShader::SetViewParameters(RHICmdList, ShaderRHI, View, View.ViewUniformBuffer);
-	}
-
-	virtual bool Serialize(FArchive& Ar) override
-	{
-		const bool bShaderHasOutdatedParameters = FMaterialShader::Serialize(Ar);
-		return bShaderHasOutdatedParameters;
+		TPostProcessMaterialShader::SetViewParameters(RHICmdList, ShaderRHI, View, View.ViewUniformBuffer);
 	}
 };
 
-IMPLEMENT_MATERIAL_SHADER_TYPE(template<>, FARKitCameraOverlayVS<true>, TEXT("/Engine/Private/PostProcessMaterialShaders.usf"), TEXT("MainVS_ES2"), SF_Vertex);
-IMPLEMENT_MATERIAL_SHADER_TYPE(template<>, FARKitCameraOverlayVS<false>, TEXT("/Engine/Private/PostProcessMaterialShaders.usf"), TEXT("MainVS_VideoOverlay"), SF_Vertex);
+using FARKitCameraOverlayVS = TARKitCameraOverlayVS<false>;
+using FARKitCameraOverlayMobileVS = TARKitCameraOverlayVS<true>;
+
+IMPLEMENT_MATERIAL_SHADER(FARKitCameraOverlayVS, "/Engine/Private/PostProcessMaterialShaders.usf", "MainVS_VideoOverlay", SF_Vertex);
+IMPLEMENT_MATERIAL_SHADER(FARKitCameraOverlayMobileVS, "/Engine/Private/PostProcessMaterialShaders.usf", "MainVS", SF_Vertex);
 
 template <bool bIsMobileRenderer>
-class FARKitCameraOverlayPS : public FMaterialShader
+class TARKitCameraOverlayPS : public TPostProcessMaterialShader<bIsMobileRenderer>
 {
-	DECLARE_SHADER_TYPE(FARKitCameraOverlayPS, Material);
 public:
-
-	static bool ShouldCompilePermutation(const FMaterialShaderPermutationParameters& Parameters)
-	{
-		if (bIsMobileRenderer)
-		{
-			return Parameters.Material->GetMaterialDomain() == MD_PostProcess && IsMobilePlatform(Parameters.Platform);
-		}
-		else
-		{
-			return Parameters.Material->GetMaterialDomain() == MD_PostProcess && !IsMobilePlatform(Parameters.Platform);
-		}
-	}
+	DECLARE_MATERIAL_SHADER(TARKitCameraOverlayPS);
 
 	static void ModifyCompilationEnvironment(const FMaterialShaderPermutationParameters& Parameters, FShaderCompilerEnvironment& OutEnvironment)
 	{
-		FMaterialShader::ModifyCompilationEnvironment(Parameters, OutEnvironment);
-		OutEnvironment.SetDefine(TEXT("POST_PROCESS_MATERIAL"), 1);
+		TPostProcessMaterialShader::ModifyCompilationEnvironment(Parameters, OutEnvironment);
 		OutEnvironment.SetDefine(TEXT("OUTPUT_MOBILE_HDR"), IsMobileHDR() ? 1 : 0);
-		OutEnvironment.SetDefine(TEXT("POST_PROCESS_MATERIAL_BEFORE_TONEMAP"), (Parameters.Material->GetBlendableLocation() != BL_AfterTonemapping) ? 1 : 0);
 	}
 
-	FARKitCameraOverlayPS() {}
-	FARKitCameraOverlayPS(const ShaderMetaType::CompiledShaderInitializerType& Initializer) :
-		FMaterialShader(Initializer)
-	{
-		for (uint32 InputIter = 0; InputIter < ePId_Input_MAX; ++InputIter)
-		{
-			PostprocessInputParameter[InputIter].Bind(Initializer.ParameterMap, *FString::Printf(TEXT("PostprocessInput%d"), InputIter));
-			PostprocessInputParameterSampler[InputIter].Bind(Initializer.ParameterMap, *FString::Printf(TEXT("PostprocessInput%dSampler"), InputIter));
-		}
-	}
+	TARKitCameraOverlayPS() = default;
+	TARKitCameraOverlayPS(const ShaderMetaType::CompiledShaderInitializerType& Initializer)
+		: TPostProcessMaterialShader(Initializer)
+	{}
 
 	void SetParameters(FRHICommandList& RHICmdList, const FSceneView View, const FMaterialRenderProxy* Material)
 	{
 		FRHIPixelShader* ShaderRHI = GetPixelShader();
-		FMaterialShader::SetParameters(RHICmdList, ShaderRHI, Material, *Material->GetMaterial(View.GetFeatureLevel()), View, View.ViewUniformBuffer, ESceneTextureSetupMode::None);
-
-		for (uint32 InputIter = 0; InputIter < ePId_Input_MAX; ++InputIter)
-		{
-			if (PostprocessInputParameter[InputIter].IsBound())
-			{
-				SetTextureParameter(
-					RHICmdList, 
-					ShaderRHI, 
-					PostprocessInputParameter[InputIter],
-					PostprocessInputParameterSampler[InputIter],
-					TStaticSamplerState<>::GetRHI(),
-					GBlackTexture->TextureRHI);
-			}
-		}
+		TPostProcessMaterialShader::SetParameters(RHICmdList, ShaderRHI, Material, *Material->GetMaterial(View.GetFeatureLevel()), View, View.ViewUniformBuffer, ESceneTextureSetupMode::None);
 	}
-
-	virtual bool Serialize(FArchive& Ar) override
-	{
-		const bool bShaderHasOutdatedParameters = FMaterialShader::Serialize(Ar);
-		return bShaderHasOutdatedParameters;
-	}
-
-private:
-	FShaderResourceParameter PostprocessInputParameter[ePId_Input_MAX];
-	FShaderResourceParameter PostprocessInputParameterSampler[ePId_Input_MAX];
 };
 
-IMPLEMENT_MATERIAL_SHADER_TYPE(template<>, FARKitCameraOverlayPS<true>, TEXT("/Engine/Private/PostProcessMaterialShaders.usf"), TEXT("MainPS_ES2"), SF_Pixel);
-IMPLEMENT_MATERIAL_SHADER_TYPE(template<>, FARKitCameraOverlayPS<false>, TEXT("/Engine/Private/PostProcessMaterialShaders.usf"), TEXT("MainPS_VideoOverlay"), SF_Pixel);
+using FARKitCameraOverlayPS = TARKitCameraOverlayPS<false>;
+using FARKitCameraOverlayMobilePS = TARKitCameraOverlayPS<true>;
+
+IMPLEMENT_MATERIAL_SHADER(FARKitCameraOverlayPS, "/Engine/Private/PostProcessMaterialShaders.usf", "MainPS_VideoOverlay", SF_Pixel);
+IMPLEMENT_MATERIAL_SHADER(FARKitCameraOverlayMobilePS, "/Engine/Private/PostProcessMaterialShaders.usf", "MainPS_ES2", SF_Pixel);
 
 void FAppleARKitVideoOverlay::RenderVideoOverlay_RenderThread(FRHICommandListImmediate& RHICmdList, const FSceneView& InView)
 {
@@ -221,18 +191,18 @@ void FAppleARKitVideoOverlay::RenderVideoOverlay_RenderThread(FRHICommandListImm
 	FMaterialShader* PixelShader = nullptr;
 	if (bIsMobileRenderer)
 	{
-		VertexShader = MaterialShaderMap->GetShader<FARKitCameraOverlayVS<true>>();
+		VertexShader = MaterialShaderMap->GetShader<FARKitCameraOverlayMobileVS>();
 		check(VertexShader != nullptr);
-		PixelShader = MaterialShaderMap->GetShader<FARKitCameraOverlayPS<true>>();
+		PixelShader = MaterialShaderMap->GetShader<FARKitCameraOverlayMobilePS>();
 		check(PixelShader != nullptr);
 		GraphicsPSOInit.BoundShaderState.VertexShaderRHI = GETSAFERHISHADER_VERTEX(VertexShader);
 		GraphicsPSOInit.BoundShaderState.PixelShaderRHI = GETSAFERHISHADER_PIXEL(PixelShader);
 	}
 	else
 	{
-		VertexShader = MaterialShaderMap->GetShader<FARKitCameraOverlayVS<false>>();
+		VertexShader = MaterialShaderMap->GetShader<FARKitCameraOverlayVS>();
 		check(VertexShader != nullptr);
-		PixelShader = MaterialShaderMap->GetShader<FARKitCameraOverlayPS<false>>();
+		PixelShader = MaterialShaderMap->GetShader<FARKitCameraOverlayPS>();
 		check(PixelShader != nullptr);
 		GraphicsPSOInit.BoundShaderState.VertexShaderRHI = GETSAFERHISHADER_VERTEX(VertexShader);
 		GraphicsPSOInit.BoundShaderState.PixelShaderRHI = GETSAFERHISHADER_PIXEL(PixelShader);
@@ -250,21 +220,21 @@ void FAppleARKitVideoOverlay::RenderVideoOverlay_RenderThread(FRHICommandListImm
 
 	if (bIsMobileRenderer)
 	{
-		FARKitCameraOverlayVS<true>* const VertexShaderPtr = static_cast<FARKitCameraOverlayVS<true>*>(VertexShader);
+		FARKitCameraOverlayVS<true>* const VertexShaderPtr = static_cast<FARKitCameraOverlayMobileVS*>(VertexShader);
 		check(VertexShaderPtr != nullptr);
 		SetUniformBufferParameterImmediate(RHICmdList, VertexShaderPtr->GetVertexShader(), VertexShaderPtr->GetUniformBufferParameter<FDrawRectangleParameters>(), Parameters);
 		VertexShaderPtr->SetParameters(RHICmdList, InView);
-		FARKitCameraOverlayPS<true>* PixelShaderPtr = static_cast<FARKitCameraOverlayPS<true>*>(PixelShader);
+		FARKitCameraOverlayPS<true>* PixelShaderPtr = static_cast<FARKitCameraOverlayMobilePS*>(PixelShader);
 		check(PixelShaderPtr != nullptr);
 		PixelShaderPtr->SetParameters(RHICmdList, InView, RenderingOverlayMaterial->GetRenderProxy());
 	}
 	else
 	{
-		FARKitCameraOverlayVS<false>* const VertexShaderPtr = static_cast<FARKitCameraOverlayVS<false>*>(VertexShader);
+		FARKitCameraOverlayVS<false>* const VertexShaderPtr = static_cast<FARKitCameraOverlayVS*>(VertexShader);
 		check(VertexShaderPtr != nullptr);
 		SetUniformBufferParameterImmediate(RHICmdList, VertexShaderPtr->GetVertexShader(), VertexShaderPtr->GetUniformBufferParameter<FDrawRectangleParameters>(), Parameters);
 		VertexShaderPtr->SetParameters(RHICmdList, InView);
-		FARKitCameraOverlayPS<false>* PixelShaderPtr = static_cast<FARKitCameraOverlayPS<false>*>(PixelShader);
+		FARKitCameraOverlayPS<false>* PixelShaderPtr = static_cast<FARKitCameraOverlayPS*>(PixelShader);
 		check(PixelShaderPtr != nullptr);
 		PixelShaderPtr->SetParameters(RHICmdList, InView, RenderingOverlayMaterial->GetRenderProxy());
 	}
diff --git a/Engine/Plugins/Runtime/AR/Google/GoogleARCore/Source/GoogleARCoreRendering/Private/GoogleARCorePassthroughCameraRenderer.cpp b/Engine/Plugins/Runtime/AR/Google/GoogleARCore/Source/GoogleARCoreRendering/Private/GoogleARCorePassthroughCameraRenderer.cpp
--- a/Engine/Plugins/Runtime/AR/Google/GoogleARCore/Source/GoogleARCoreRendering/Private/GoogleARCorePassthroughCameraRenderer.cpp
+++ b/Engine/Plugins/Runtime/AR/Google/GoogleARCore/Source/GoogleARCoreRendering/Private/GoogleARCorePassthroughCameraRenderer.cpp
@@ -12,6 +12,7 @@
 #include "SceneUtils.h"
 #include "Engine/TextureRenderTarget2D.h"
 #include "PostProcess/SceneFilterRendering.h"
+#include "PostProcess/PostProcessMaterial.h"
 #include "PostProcessParameters.h"
 #include "MaterialShader.h"
 #include "MaterialShaderType.h"
@@ -142,11 +143,11 @@ void FGoogleARCorePassthroughCameraRenderer::UpdateOverlayUVCoordinate_RenderThr
 	OverlayVertexBufferRHI = RHICreateVertexBuffer(Vertices.GetResourceDataSize(), BUF_Static, CreateInfo);
 }
 
-// We use something similar to the PostProcessMaterial to render the color camera overlay.
-class FGoogleARCoreCameraOverlayVS : public FMaterialShader
+class FPostProcessMaterialShader : public FMaterialShader
 {
-	DECLARE_SHADER_TYPE(FGoogleARCoreCameraOverlayVS, Material);
 public:
+	using FParameters = FPostProcessMaterialParameters;
+	SHADER_USE_PARAMETER_STRUCT_WITH_LEGACY_BASE(FPostProcessMaterialShader, FMaterialShader);
 
 	static bool ShouldCompilePermutation(const FMaterialShaderPermutationParameters& Parameters)
 	{
@@ -156,96 +157,62 @@ public:
 	static void ModifyCompilationEnvironment(const FMaterialShaderPermutationParameters& Parameters, FShaderCompilerEnvironment& OutEnvironment)
 	{
 		FMaterialShader::ModifyCompilationEnvironment(Parameters.Platform, OutEnvironment);
-
 		OutEnvironment.SetDefine(TEXT("POST_PROCESS_MATERIAL"), 1);
+		OutEnvironment.SetDefine(TEXT("POST_PROCESS_MATERIAL_MOBILE"), 1);
 		OutEnvironment.SetDefine(TEXT("POST_PROCESS_MATERIAL_BEFORE_TONEMAP"), (Parameters.Material->GetBlendableLocation() != BL_AfterTonemapping) ? 1 : 0);
-		OutEnvironment.SetDefine(TEXT("POST_PROCESS_AR_PASSTHROUGH"), 1);
 	}
+};
 
-	FGoogleARCoreCameraOverlayVS( )	{ }
-	FGoogleARCoreCameraOverlayVS(const ShaderMetaType::CompiledShaderInitializerType& Initializer)
-		: FMaterialShader(Initializer)
+// We use something similar to the PostProcessMaterial to render the color camera overlay.
+class FGoogleARCoreCameraOverlayVS : public FPostProcessMaterialShader
+{
+public:
+	DECLARE_MATERIAL_SHADER(FGoogleARCoreCameraOverlayVS);
+
+	static void ModifyCompilationEnvironment(const FMaterialShaderPermutationParameters& Parameters, FShaderCompilerEnvironment& OutEnvironment)
 	{
+		FPostProcessMaterialShader::ModifyCompilationEnvironment(Parameters, OutEnvironment);
+		OutEnvironment.SetDefine(TEXT("POST_PROCESS_AR_PASSTHROUGH"), 1);
 	}
 
+	FGoogleARCoreCameraOverlayVS() = default;
+	FGoogleARCoreCameraOverlayVS(const ShaderMetaType::CompiledShaderInitializerType& Initializer)
+		: FPostProcessMaterialShader(Initializer)
+	{}
+
 	void SetParameters(FRHICommandList& RHICmdList, const FSceneView View)
 	{
 		FRHIVertexShader* ShaderRHI = GetVertexShader();
 		FMaterialShader::SetViewParameters(RHICmdList, ShaderRHI, View, View.ViewUniformBuffer);
 	}
-
-	virtual bool Serialize(FArchive& Ar) override
-	{
-		bool bShaderHasOutdatedParameters = FMaterialShader::Serialize(Ar);
-		return bShaderHasOutdatedParameters;
-	}
 };
 
-IMPLEMENT_MATERIAL_SHADER_TYPE(,FGoogleARCoreCameraOverlayVS,TEXT("/Engine/Private/PostProcessMaterialShaders.usf"),TEXT("MainVS_ES2"),SF_Vertex);
+IMPLEMENT_MATERIAL_SHADER(FGoogleARCoreCameraOverlayVS, "/Engine/Private/PostProcessMaterialShaders.usf", "MainVS", SF_Vertex);
 
-class FGoogleARCoreCameraOverlayPS : public FMaterialShader
+class FGoogleARCoreCameraOverlayPS : public FPostProcessMaterialShader
 {
-	DECLARE_SHADER_TYPE(FGoogleARCoreCameraOverlayPS, Material);
 public:
-
-	static bool ShouldCompilePermutation(const FMaterialShaderPermutationParameters& Parameters)
-	{
-		return Parameters.Material->GetMaterialDomain() == MD_PostProcess && IsMobilePlatform(Parameters.Platform);
-	}
+	DECLARE_MATERIAL_SHADER(FGoogleARCoreCameraOverlayPS);
 
 	static void ModifyCompilationEnvironment(const FMaterialShaderPermutationParameters& Parameters, FShaderCompilerEnvironment& OutEnvironment)
 	{
-		FMaterialShader::ModifyCompilationEnvironment(Parameters.Platform, OutEnvironment);
-
-		OutEnvironment.SetDefine(TEXT("POST_PROCESS_MATERIAL"), 1);
+		FPostProcessMaterialShader::ModifyCompilationEnvironment(Parameters, OutEnvironment);
 		OutEnvironment.SetDefine(TEXT("OUTPUT_GAMMA_SPACE"), IsMobileHDR() ? 0 : 1);
-		OutEnvironment.SetDefine(TEXT("POST_PROCESS_MATERIAL_BEFORE_TONEMAP"), (Parameters.Material->GetBlendableLocation() != BL_AfterTonemapping) ? 1 : 0);
 	}
 
-	FGoogleARCoreCameraOverlayPS() {}
-	FGoogleARCoreCameraOverlayPS(const ShaderMetaType::CompiledShaderInitializerType& Initializer) :
-		FMaterialShader(Initializer)
-	{
-		for (uint32 InputIter = 0; InputIter < ePId_Input_MAX; ++InputIter)
-		{
-			PostprocessInputParameter[InputIter].Bind(Initializer.ParameterMap, *FString::Printf(TEXT("PostprocessInput%d"), InputIter));
-			PostprocessInputParameterSampler[InputIter].Bind(Initializer.ParameterMap, *FString::Printf(TEXT("PostprocessInput%dSampler"), InputIter));
-		}
-	}
+	FGoogleARCoreCameraOverlayPS() = default;
+	FGoogleARCoreCameraOverlayPS(const ShaderMetaType::CompiledShaderInitializerType& Initializer)
+		: FPostProcessMaterialShader(Initializer)
+	{}
 
 	void SetParameters(FRHICommandList& RHICmdList, const FSceneView View, const FMaterialRenderProxy* Material)
 	{
 		FRHIPixelShader* ShaderRHI = GetPixelShader();
-
 		FMaterialShader::SetParameters(RHICmdList, ShaderRHI, Material, *Material->GetMaterial(View.GetFeatureLevel()), View, View.ViewUniformBuffer, ESceneTextureSetupMode::None);
-
-		for (uint32 InputIter = 0; InputIter < ePId_Input_MAX; ++InputIter)
-		{
-			if (PostprocessInputParameter[InputIter].IsBound())
-			{
-				SetTextureParameter(
-					RHICmdList,
-					ShaderRHI,
-					PostprocessInputParameter[InputIter],
-					PostprocessInputParameterSampler[InputIter],
-					TStaticSamplerState<>::GetRHI(),
-					GBlackTexture->TextureRHI);
-			}
-		}
 	}
-
-	virtual bool Serialize(FArchive& Ar) override
-	{
-		bool bShaderHasOutdatedParameters = FMaterialShader::Serialize(Ar);
-		return bShaderHasOutdatedParameters;
-	}
-
-private:
-	FShaderResourceParameter PostprocessInputParameter[ePId_Input_MAX];
-	FShaderResourceParameter PostprocessInputParameterSampler[ePId_Input_MAX];
 };
 
-IMPLEMENT_MATERIAL_SHADER_TYPE(,FGoogleARCoreCameraOverlayPS,TEXT("/Engine/Private/PostProcessMaterialShaders.usf"),TEXT("MainPS_ES2"),SF_Pixel);
+IMPLEMENT_MATERIAL_SHADER(FGoogleARCoreCameraOverlayPS, "/Engine/Private/PostProcessMaterialShaders.usf", "MainPS", SF_Pixel);
 
 void FGoogleARCorePassthroughCameraRenderer::RenderVideoOverlay_RenderThread(FRHICommandListImmediate& RHICmdList, FSceneView& InView)
 {
diff --git a/Engine/Shaders/Private/MaterialTemplate.ush b/Engine/Shaders/Private/MaterialTemplate.ush
--- a/Engine/Shaders/Private/MaterialTemplate.ush
+++ b/Engine/Shaders/Private/MaterialTemplate.ush
@@ -18,6 +18,41 @@
 	#include "/Engine/Private/SpeedTreeCommon.ush"
 #endif
 
+//////////////////////////////////////////////////////////////////////////
+//! Must match ESceneTextureId
+
+#define PPI_SceneColor 0
+#define PPI_SceneDepth 1
+#define PPI_DiffuseColor 2
+#define PPI_SpecularColor 3
+#define PPI_SubsurfaceColor 4
+#define PPI_BaseColor 5
+#define PPI_Specular 6
+#define PPI_Metallic 7
+#define PPI_WorldNormal 8
+#define PPI_SeparateTranslucency 9
+#define PPI_Opacity 10
+#define PPI_Roughness 11
+#define PPI_MaterialAO 12
+#define PPI_CustomDepth 13
+#define PPI_PostProcessInput0 14
+#define PPI_PostProcessInput1 15
+#define PPI_PostProcessInput2 16
+#define PPI_PostProcessInput3 17
+#define PPI_PostProcessInput4 18
+#define PPI_PostProcessInput5 19 // (UNUSED)
+#define PPI_PostProcessInput6 20 // (UNUSED)
+#define PPI_DecalMask 21
+#define PPI_ShadingModelColor 22
+#define PPI_ShadingModelID 23
+#define PPI_AmbientOcclusion 24
+#define PPI_CustomStencil 25
+#define PPI_StoredBaseColor 26
+#define PPI_StoredSpecular 27
+#define PPI_Velocity 28
+
+//////////////////////////////////////////////////////////////////////////
+
 #define NUM_MATERIAL_TEXCOORDS_VERTEX %s
 #define NUM_MATERIAL_TEXCOORDS %s
 #define NUM_CUSTOM_VERTEX_INTERPOLATORS %s
@@ -37,11 +72,6 @@
 
 #include "/Engine/Private/PaniniProjection.ush"
 
-#if POST_PROCESS_MATERIAL
-	// only Postprocess material bind PostprocessInput0, PostprocessInput0Size, PostprocessInput0MinMax ...
-	#include "/Engine/Private/PostProcessCommon.ush"			// PostprocessInput0
-#endif
-
 #ifndef USE_DITHERED_LOD_TRANSITION
     #if USE_INSTANCING
 	    #ifndef USE_DITHERED_LOD_TRANSITION_FOR_INSTANCED
@@ -458,16 +488,16 @@ float2 GetPixelPosition(FMaterialVertexParameters Parameters)
 }
 
 
-#if defined(POST_PROCESS_MATERIAL) && POST_PROCESS_MATERIAL && POST_PROCESS_MATERIAL_AFTER_TAA_UPSAMPLE
+#if POST_PROCESS_MATERIAL
 
 float2 GetPixelPosition(FMaterialPixelParameters Parameters)
 {
-	return Parameters.SvPosition.xy - float2(ViewportRect.xy);
+	return Parameters.SvPosition.xy - float2(PostProcessOutput_ViewportMin);
 }
 
 float2 GetViewportUV(FMaterialPixelParameters Parameters)
 {
-	return GetPixelPosition(Parameters) * ViewportSize.zw;
+	return GetPixelPosition(Parameters) * PostProcessOutput_ViewportSizeInverse;
 }
 
 #else
@@ -484,117 +514,113 @@ float2 GetViewportUV(FMaterialPixelParameters Parameters)
 
 #endif
 
-
 // Returns whether a scene texture id is a for a post process input or not.
 bool IsPostProcessInputSceneTexture(const uint SceneTextureId)
 {
-	return (SceneTextureId >= 14 && SceneTextureId <= 20);
+	return (SceneTextureId >= PPI_PostProcessInput0 && SceneTextureId <= PPI_PostProcessInput6);
 }
 
-
 // Returns the view size and texel size in a given scene texture.
 float4 GetSceneTextureViewSize(const uint SceneTextureId)
 {
-	#if POST_PROCESS_MATERIAL && POST_PROCESS_MATERIAL_AFTER_TAA_UPSAMPLE
+	#if POST_PROCESS_MATERIAL
 	if (IsPostProcessInputSceneTexture(SceneTextureId))
 	{
-		return ViewportSize;
+		switch (SceneTextureId)
+		{
+		case PPI_PostProcessInput0:
+			return float4(PostProcessInput_0_ViewportSize, PostProcessInput_0_ViewportSizeInverse);
+		case PPI_PostProcessInput1:
+			return float4(PostProcessInput_1_ViewportSize, PostProcessInput_1_ViewportSizeInverse);
+		case PPI_PostProcessInput2:
+			return float4(PostProcessInput_2_ViewportSize, PostProcessInput_2_ViewportSizeInverse);
+		case PPI_PostProcessInput3:
+			return float4(PostProcessInput_3_ViewportSize, PostProcessInput_3_ViewportSizeInverse);
+		case PPI_PostProcessInput4:
+			return float4(PostProcessInput_4_ViewportSize, PostProcessInput_4_ViewportSizeInverse);
+		default:
+			return float4(0, 0, 0, 0);
+		}
 	}
 	#endif
 	return ResolvedView.ViewSizeAndInvSize;
 }
 
-
 // Return the buffer UV min and max for a given scene texture id.
 float4 GetSceneTextureUVMinMax(const uint SceneTextureId)
 {
 	#if POST_PROCESS_MATERIAL
+	if (IsPostProcessInputSceneTexture(SceneTextureId))
 	{
-		const uint PostProcessInputIndex = SceneTextureId - 14;
-
-			 if (PostProcessInputIndex == 0)	return PostprocessInput0MinMax;
-		else if (PostProcessInputIndex == 1)	return PostprocessInput1MinMax;
-		else if (PostProcessInputIndex == 2)	return PostprocessInput2MinMax;
-		else if (PostProcessInputIndex == 3)	return PostprocessInput3MinMax;
-		else if (PostProcessInputIndex == 4)	return PostprocessInput4MinMax;
-		else if (PostProcessInputIndex == 5)	return PostprocessInput5MinMax;
-		else if (PostProcessInputIndex == 6)	return PostprocessInput6MinMax;
+		switch (SceneTextureId)
+		{
+		case PPI_PostProcessInput0:
+			return float4(PostProcessInput_0_UVViewportBilinearMin, PostProcessInput_0_UVViewportBilinearMax);
+		case PPI_PostProcessInput1:
+			return float4(PostProcessInput_1_UVViewportBilinearMin, PostProcessInput_1_UVViewportBilinearMax);
+		case PPI_PostProcessInput2:
+			return float4(PostProcessInput_2_UVViewportBilinearMin, PostProcessInput_2_UVViewportBilinearMax);
+		case PPI_PostProcessInput3:
+			return float4(PostProcessInput_3_UVViewportBilinearMin, PostProcessInput_3_UVViewportBilinearMax);
+		case PPI_PostProcessInput4:
+			return float4(PostProcessInput_4_UVViewportBilinearMin, PostProcessInput_4_UVViewportBilinearMax);
+		default:
+			return float4(0, 0, 1, 1);
+		}
 	}
 	#endif
 
 	return View.BufferBilinearUVMinMax;
 }
 
-
 // Transforms viewport UV to scene texture's UV.
 MaterialFloat2 ViewportUVToSceneTextureUV(MaterialFloat2 ViewportUV, const uint SceneTextureId)
 {
 	#if POST_PROCESS_MATERIAL
+	if (IsPostProcessInputSceneTexture(SceneTextureId))
 	{
-		// If compiling for material permutation that assume BufferSize == ViewSize.
-		#if POST_PROCESS_MATERIAL_UV_POLICY == 0
+		switch (SceneTextureId)
 		{
+		case PPI_PostProcessInput0:
+			return ViewportUV * PostProcessInput_0_UVViewportSize + PostProcessInput_0_UVViewportMin;
+		case PPI_PostProcessInput1:
+			return ViewportUV * PostProcessInput_1_UVViewportSize + PostProcessInput_1_UVViewportMin;
+		case PPI_PostProcessInput2:
+			return ViewportUV * PostProcessInput_2_UVViewportSize + PostProcessInput_2_UVViewportMin;
+		case PPI_PostProcessInput3:
+			return ViewportUV * PostProcessInput_3_UVViewportSize + PostProcessInput_3_UVViewportMin;
+		case PPI_PostProcessInput4:
+			return ViewportUV * PostProcessInput_4_UVViewportSize + PostProcessInput_4_UVViewportMin;
+		default:
 			return ViewportUV;
 		}
-		#endif
-
-		#if POST_PROCESS_MATERIAL_AFTER_TAA_UPSAMPLE
-		if (IsPostProcessInputSceneTexture(SceneTextureId))
-		{
-			return ViewportUVToPostProcessingSceneColorBufferUV(ViewportUV);
-		}
-		#endif
 	}
 	#endif
 
 	return ViewportUVToBufferUV(ViewportUV);
 }
 
-
 // Manually clamp scene texture UV as if using a clamp sampler.
 MaterialFloat2 ClampSceneTextureUV(MaterialFloat2 BufferUV, const uint SceneTextureId)
 {
-	#if POST_PROCESS_MATERIAL
-	{
-		// If compiling for material permutation that assume BufferSize == ViewSize, then there is no need to clamp.
-		#if POST_PROCESS_MATERIAL_UV_POLICY == 0
-		{
-			return BufferUV;
-		}
-		#endif
-	}
-	#endif
-
 	float4 MinMax = GetSceneTextureUVMinMax(SceneTextureId);
 
 	return clamp(BufferUV, MinMax.xy, MinMax.zw);
 }
 
-
 // Get default scene texture's UV.
 MaterialFloat2 GetDefaultSceneTextureUV(FMaterialVertexParameters Parameters, const uint SceneTextureId)
 {
 	return GetSceneTextureUV(Parameters);
 }
 
-
 // Get default scene texture's UV.
 MaterialFloat2 GetDefaultSceneTextureUV(FMaterialPixelParameters Parameters, const uint SceneTextureId)
 {
 	#if POST_PROCESS_MATERIAL
+	if (IsPostProcessInputSceneTexture(SceneTextureId))
 	{
-		// If compiling for material permutation that assume BufferSize == ViewSize, then return scene texture UV.
-		#if POST_PROCESS_MATERIAL_UV_POLICY == 0
-		{
-			return GetSceneTextureUV(Parameters);
-		}
-		#endif
-
-		#if POST_PROCESS_MATERIAL_AFTER_TAA_UPSAMPLE
-		{
-			return ViewportUVToSceneTextureUV(GetViewportUV(Parameters), SceneTextureId);
-		}
-		#endif
+		return ViewportUVToSceneTextureUV(GetViewportUV(Parameters), SceneTextureId);
 	}
 	#endif
 
@@ -1378,53 +1404,35 @@ float EyeAdaptationLookup()
 
 #if NEEDS_SCENE_TEXTURES
 
-float4 GetPostProcessInputMinMax(uint Index)
-{
-#if POST_PROCESS_MATERIAL
-		 if (Index == 0)	return PostprocessInput0MinMax;
-	else if (Index == 1)	return PostprocessInput1MinMax;
-	else if (Index == 2)	return PostprocessInput2MinMax;
-	else if (Index == 3)	return PostprocessInput3MinMax;
-	else if (Index == 4)	return PostprocessInput4MinMax;
-	else if (Index == 5)	return PostprocessInput5MinMax;
-	else if (Index == 6)	return PostprocessInput6MinMax;
-#endif
-	return 1.0f;
-}
-
 #if SHADING_PATH_MOBILE
 
 MaterialFloat4 MobileSceneTextureLookup(inout FMaterialPixelParameters Parameters, int SceneTextureId, float2 UV)
 {
 #if (FEATURE_LEVEL <= FEATURE_LEVEL_ES3_1)
-	
-	//PPI_SceneDepth
-	if (SceneTextureId == 1)
+
+	if (SceneTextureId == PPI_SceneDepth)
 	{
 		MaterialFloat Depth = Texture2DSample(MobileSceneTextures.SceneColorTexture, MobileSceneTextures.SceneColorTextureSampler, UV).a;
 		return MaterialFloat4(Depth.rrr, 0.f);
 	}
-	//PPI_CustomDepth
-	else if (SceneTextureId == 13)
+	else if (SceneTextureId == PPI_CustomDepth)
 	{
 		MaterialFloat Depth = ConvertFromDeviceZ(Texture2DSample(MobileSceneTextures.CustomDepthTexture, MobileSceneTextures.CustomDepthTextureSampler, UV).r);
 		return MaterialFloat4(Depth.rrr, 0.f);
 	}
-	// PPI_PostprocessInput
-	else if (SceneTextureId == 14)
+	else if (SceneTextureId == PPI_PostProcessInput0)
 	{
-#if POST_PROCESS_MATERIAL	
-		MaterialFloat4 Input0 = Texture2DSample(PostprocessInput0, PostprocessInput0Sampler, UV);
+#if POST_PROCESS_MATERIAL
+		MaterialFloat4 Input0 = Texture2DSample(PostProcessInput_0_Texture, PostProcessInput_0_SharedSampler, UV);
 		#if POST_PROCESS_MATERIAL_BEFORE_TONEMAP
 		Input0 = Decode32BPPHDR(Input0);
 		#endif
-		// We need to preserve original SceneColor Alpha as it's used by tonemaper on mobile
+		// We need to preserve original SceneColor Alpha as it's used by tonemapper on mobile
 		Parameters.BackupSceneColorAlpha = Input0.a;
 		return Input0;
-#endif// POST_PROCESS_MATERIAL	
+#endif// POST_PROCESS_MATERIAL
 	}
-	// PPI_CustomStencil
-	else if (SceneTextureId == 25)
+	else if (SceneTextureId == PPI_CustomStencil)
 	{
 		MaterialFloat Stencil = Texture2DSample(MobileSceneTextures.MobileCustomStencilTexture, MobileSceneTextures.MobileCustomStencilTextureSampler, UV).r*255.0;
 		Stencil = floor(Stencil + 0.5);
@@ -1449,7 +1457,7 @@ float2 PostProcessVelocityLookup(float Depth, float2 UV)
 #if GBUFFER_HAS_VELOCITY
 	float2 Velocity = Texture2DSampleLevel(SceneTexturesStruct.GBufferVelocityTexture, SceneTexturesStruct.GBufferVelocityTextureSampler, UV, 0).xy;
 #else
-	float2 Velocity = Texture2DSample(PostprocessInput4, PostprocessInput4Sampler, UV).xy;
+	float2 Velocity = Texture2DSample(PostProcessInput_4_Texture, PostProcessInput_4_SharedSampler, UV).xy;
 #endif
 
 	if( Velocity.x > 0.0 )
@@ -1480,67 +1488,63 @@ float4 SceneTextureLookup(float2 UV, int SceneTextureIndex, bool bFiltered)
 	{
 		// order needs to match to ESceneTextureId
 
-		// PPI_SceneColor
-		case 0: return float4(CalcSceneColor(UV), 0);
-		// PPI_SceneDepth
-		case 1: return ScreenSpaceData.GBuffer.Depth;
-		// PPI_DiffuseColor
-		case 2: return float4(ScreenSpaceData.GBuffer.DiffuseColor, 0);
-		// PPI_SpecularColor
-		case 3: return float4(ScreenSpaceData.GBuffer.SpecularColor, 0);
-		// PPI_SubsurfaceColor
-		case 4: return IsSubsurfaceModel(ScreenSpaceData.GBuffer.ShadingModelID) ? float4( ExtractSubsurfaceColor(ScreenSpaceData.GBuffer), ScreenSpaceData.GBuffer.CustomData.a ) : ScreenSpaceData.GBuffer.CustomData;
-		// PPI_BaseColor
-		case 5: return float4(ScreenSpaceData.GBuffer.BaseColor, 0);
-		// PPI_Specular
-		case 6: return ScreenSpaceData.GBuffer.Specular;
-		// PPI_Metallic
-		case 7: return ScreenSpaceData.GBuffer.Metallic;
-		// PPI_WorldNormal
-		case 8: return float4(ScreenSpaceData.GBuffer.WorldNormal, 0);
-		// PPI_SeparateTranslucency
-		case 9: return float4(1, 1, 1, 1);	// todo
-		// PPI_Opacity
-		case 10: return ScreenSpaceData.GBuffer.CustomData.a;
-		// PPI_Roughness
-		case 11: return ScreenSpaceData.GBuffer.Roughness;
-		// PPI_MaterialAO
-		case 12: return ScreenSpaceData.GBuffer.GBufferAO;
-		// PPI_CustomDepth
-		case 13: return ScreenSpaceData.GBuffer.CustomDepth;
+		case PPI_SceneColor:
+			return float4(CalcSceneColor(UV), 0);
+		case PPI_SceneDepth:
+			return ScreenSpaceData.GBuffer.Depth;
+		case PPI_DiffuseColor:
+			return float4(ScreenSpaceData.GBuffer.DiffuseColor, 0);
+		case PPI_SpecularColor:
+			return float4(ScreenSpaceData.GBuffer.SpecularColor, 0);
+		case PPI_SubsurfaceColor:
+			return IsSubsurfaceModel(ScreenSpaceData.GBuffer.ShadingModelID) ? float4( ExtractSubsurfaceColor(ScreenSpaceData.GBuffer), ScreenSpaceData.GBuffer.CustomData.a ) : ScreenSpaceData.GBuffer.CustomData;
+		case PPI_BaseColor:
+			return float4(ScreenSpaceData.GBuffer.BaseColor, 0);
+		case PPI_Specular:
+			return ScreenSpaceData.GBuffer.Specular;
+		case PPI_Metallic:
+			return ScreenSpaceData.GBuffer.Metallic;
+		case PPI_WorldNormal:
+			return float4(ScreenSpaceData.GBuffer.WorldNormal, 0);
+		case PPI_SeparateTranslucency:
+			return float4(1, 1, 1, 1);	// todo
+		case PPI_Opacity:
+			return ScreenSpaceData.GBuffer.CustomData.a;
+		case PPI_Roughness:
+			return ScreenSpaceData.GBuffer.Roughness;
+		case PPI_MaterialAO:
+			return ScreenSpaceData.GBuffer.GBufferAO;
+		case PPI_CustomDepth:
+			return ScreenSpaceData.GBuffer.CustomDepth;
 #if POST_PROCESS_MATERIAL
-		// PPI_PostprocessInput0
-		case 14: return Texture2DSample(PostprocessInput0, bFiltered ? BilinearTextureSampler : PostprocessInput0Sampler, UV);
-		// PPI_PostprocessInput1
-		case 15: return Texture2DSample(PostprocessInput1, bFiltered ? BilinearTextureSampler : PostprocessInput1Sampler, UV);
-		// PPI_PostprocessInput2
-		case 16: return Texture2DSample(PostprocessInput2, bFiltered ? BilinearTextureSampler : PostprocessInput2Sampler, UV);
-		// PPI_PostprocessInput3
-		case 17: return Texture2DSample(PostprocessInput3, bFiltered ? BilinearTextureSampler : PostprocessInput3Sampler, UV);
-		// PPI_PostprocessInput4
-		case 18: return Texture2DSample(PostprocessInput4, bFiltered ? BilinearTextureSampler : PostprocessInput4Sampler, UV);
-		// PPI_PostprocessInput5
-		case 19: return Texture2DSample(PostprocessInput5, bFiltered ? BilinearTextureSampler : PostprocessInput5Sampler, UV);
-		// PPI_PostprocessInput6
-		case 20: return Texture2DSample(PostprocessInput6, bFiltered ? BilinearTextureSampler : PostprocessInput6Sampler, UV);
+		case PPI_PostProcessInput0:
+			return Texture2DSample(PostProcessInput_0_Texture, bFiltered ? PostProcessInput_BilinearSampler : PostProcessInput_0_SharedSampler, UV);
+		case PPI_PostProcessInput1:
+			return Texture2DSample(PostProcessInput_1_Texture, bFiltered ? PostProcessInput_BilinearSampler : PostProcessInput_1_SharedSampler, UV);
+		case PPI_PostProcessInput2:
+			return Texture2DSample(PostProcessInput_2_Texture, bFiltered ? PostProcessInput_BilinearSampler : PostProcessInput_2_SharedSampler, UV);
+		case PPI_PostProcessInput3:
+			return Texture2DSample(PostProcessInput_3_Texture, bFiltered ? PostProcessInput_BilinearSampler : PostProcessInput_3_SharedSampler, UV);
+		case PPI_PostProcessInput4:
+			return Texture2DSample(PostProcessInput_4_Texture, bFiltered ? PostProcessInput_BilinearSampler : PostProcessInput_4_SharedSampler, UV);
 #endif // __POST_PROCESS_COMMON__
-		// PPI_DecalMask
-		case 21: return 0;  // material compiler will return an error
-		// PPI_ShadingModelColor
-		case 22: return float4(GetShadingModelColor(ScreenSpaceData.GBuffer.ShadingModelID), 1);
-		// PPI_ShadingModelID
-		case 23: return float4(ScreenSpaceData.GBuffer.ShadingModelID, 0, 0, 0);
-		// PPI_AmbientOcclusion
-		case 24: return ScreenSpaceData.AmbientOcclusion;
-		// PPI_CustomStencil
-		case 25: return ScreenSpaceData.GBuffer.CustomStencil;
-		// PPI_StoredBaseColor
-		case 26: return float4(ScreenSpaceData.GBuffer.StoredBaseColor, 0);
-		// PPI_StoredSpecular
-		case 27: return float4(ScreenSpaceData.GBuffer.StoredSpecular.rrr, 0);
+		case PPI_DecalMask:
+			return 0;  // material compiler will return an error
+		case PPI_ShadingModelColor:
+			return float4(GetShadingModelColor(ScreenSpaceData.GBuffer.ShadingModelID), 1);
+		case PPI_ShadingModelID:
+			return float4(ScreenSpaceData.GBuffer.ShadingModelID, 0, 0, 0);
+		case PPI_AmbientOcclusion:
+			return ScreenSpaceData.AmbientOcclusion;
+		case PPI_CustomStencil:
+			return ScreenSpaceData.GBuffer.CustomStencil;
+		case PPI_StoredBaseColor:
+			return float4(ScreenSpaceData.GBuffer.StoredBaseColor, 0);
+		case PPI_StoredSpecular:
+			return float4(ScreenSpaceData.GBuffer.StoredSpecular.rrr, 0);
 #if POST_PROCESS_MATERIAL
-		// PPI_Velocity
-		case 28: return float4(PostProcessVelocityLookup(ScreenSpaceData.GBuffer.Depth, UV), 0, 0);
+		case PPI_Velocity:
+			return float4(PostProcessVelocityLookup(ScreenSpaceData.GBuffer.Depth, UV), 0, 0);
 #endif
 		default:
 			return float4(0, 0, 0, 0);
diff --git a/Engine/Shaders/Private/PostProcessMaterialShaders.usf b/Engine/Shaders/Private/PostProcessMaterialShaders.usf
--- a/Engine/Shaders/Private/PostProcessMaterialShaders.usf
+++ b/Engine/Shaders/Private/PostProcessMaterialShaders.usf
@@ -5,6 +5,65 @@
 =============================================================================*/
 
 #include "Common.ush"
+#include "ScreenPass.ush"
+
+#ifndef POST_PROCESS_MATERIAL
+#define POST_PROCESS_MATERIAL 0
+#endif
+
+#ifndef POST_PROCESS_MATERIAL_MOBILE
+#define POST_PROCESS_MATERIAL_MOBILE 0
+#endif
+
+#if (POST_PROCESS_MATERIAL == 0)
+#error POST_PROCESS_MATERIAL must be defined to non-zero in the shader compilation environment.
+#endif
+
+#if (FEATURE_LEVEL <= FEATURE_LEVEL_ES3_1) != POST_PROCESS_MATERIAL_MOBILE
+#error POST_PROCESS_MATERIAL_MOBILE define does not match feature level.
+#endif
+
+// Must match ESceneTextureId
+#define PPI_PostProcessInput0 14
+#define PPI_PostProcessInput1 15
+#define PPI_PostProcessInput2 16
+#define PPI_PostProcessInput3 17
+#define PPI_PostProcessInput4 18
+
+SCREEN_PASS_TEXTURE_VIEWPORT(PostProcessInput_0)
+SCREEN_PASS_TEXTURE_VIEWPORT(PostProcessInput_1)
+SCREEN_PASS_TEXTURE_VIEWPORT(PostProcessInput_2)
+SCREEN_PASS_TEXTURE_VIEWPORT(PostProcessInput_3)
+SCREEN_PASS_TEXTURE_VIEWPORT(PostProcessInput_4)
+SCREEN_PASS_TEXTURE_VIEWPORT(PostProcessOutput)
+
+Texture2D PostProcessInput_0_Texture;
+Texture2D PostProcessInput_1_Texture;
+Texture2D PostProcessInput_2_Texture;
+Texture2D PostProcessInput_3_Texture;
+Texture2D PostProcessInput_4_Texture;
+
+SamplerState PostProcessInput_0_Sampler;
+SamplerState PostProcessInput_1_Sampler;
+SamplerState PostProcessInput_2_Sampler;
+SamplerState PostProcessInput_3_Sampler;
+SamplerState PostProcessInput_4_Sampler;
+SamplerState PostProcessInput_BilinearSampler;
+
+#if SUPPORTS_INDEPENDENT_SAMPLERS
+	#define PostProcessInput_0_SharedSampler PostProcessInput_0_Sampler
+	#define PostProcessInput_1_SharedSampler PostProcessInput_0_Sampler
+	#define PostProcessInput_2_SharedSampler PostProcessInput_0_Sampler
+	#define PostProcessInput_3_SharedSampler PostProcessInput_0_Sampler
+	#define PostProcessInput_4_SharedSampler PostProcessInput_0_Sampler
+#else
+	#define PostProcessInput_0_SharedSampler PostProcessInput_0_Sampler
+	#define PostProcessInput_1_SharedSampler PostProcessInput_1_Sampler
+	#define PostProcessInput_2_SharedSampler PostProcessInput_2_Sampler
+	#define PostProcessInput_3_SharedSampler PostProcessInput_3_Sampler
+	#define PostProcessInput_4_SharedSampler PostProcessInput_4_Sampler
+#endif
+
 #include "/Engine/Generated/Material.ush"
 
 struct FPostProcessMaterialVSToPS
@@ -110,7 +169,7 @@ void MainPS(
 
 	// can be optimized
 	float4 SvPosition = Input.Position;
-	float2 ViewportUV = (SvPosition - ViewportRect.xy) * ViewportSize.zw;
+	float2 ViewportUV = (SvPosition - PostProcessOutput_ViewportMin) * PostProcessOutput_ViewportSizeInverse;
 
 #if NUM_TEX_COORD_INTERPOLATORS
 	UNROLL
@@ -185,14 +244,12 @@ void MainPS_VideoOverlay(
 
 #else // FEATURE_LEVEL_ES3_1
 
-#ifndef NEEDTOSWITCHVERTICLEAXIS
-#define NEEDTOSWITCHVERTICLEAXIS 0
-#endif
+uint bFlipYAxis;
 
 //
 // Mobile version
 //
-void MainVS_ES2(
+void MainVS(
 	in float4 InPosition : ATTRIBUTE0,
 	in float2 InTexCoord : ATTRIBUTE1,
 	out FPostProcessMaterialVSToPS Output,
@@ -229,7 +286,7 @@ void MainVS_ES2(
 #endif
 }
 
-void MainPS_ES2(
+void MainPS(
 	in FPostProcessMaterialVSToPS Input,
 	in float2 InViewportUV : TEXCOORD8,
 	out half4 OutColor : SV_Target0
@@ -250,9 +307,10 @@ void MainPS_ES2(
 	Parameters.VertexColor = 1;
 	
 	float4 SvPosition = Input.Position;
-#if NEEDTOSWITCHVERTICLEAXIS
-		SvPosition.y = View.ViewSizeAndInvSize.y - SvPosition.y; 
-#endif
+	if (bFlipYAxis)
+	{
+		SvPosition.y = PostProcessOutput_ViewportSize.y - SvPosition.y;
+	}
 	SvPosition.z = LookupDeviceZ(InViewportUV);
 	SvPosition.z = max(SvPosition.z, 1e-18);
 	// fill out other related material parameters
@@ -277,4 +335,4 @@ void MainPS_ES2(
 #endif
 	OutColor = FullSceneColor;
 }
-#endif //(FEATURE_LEVEL > FEATURE_LEVEL_ES3_1)
\ No newline at end of file
+#endif //(FEATURE_LEVEL > FEATURE_LEVEL_ES3_1)
diff --git a/Engine/Shaders/Private/ScreenPass.usf b/Engine/Shaders/Private/ScreenPass.usf
--- a/Engine/Shaders/Private/ScreenPass.usf
+++ b/Engine/Shaders/Private/ScreenPass.usf
@@ -2,8 +2,8 @@
 
 #include "Common.ush"
 
-Texture2D SrcTexture;
-SamplerState SrcSampler;
+Texture2D InputTexture;
+SamplerState InputSampler;
 
 void ScreenPassVS(
 	in float4 InPosition : ATTRIBUTE0,
@@ -16,8 +16,8 @@ void ScreenPassVS(
 
 float4 CopyRectPS(
 	noperspective float4 UVAndScreenPos : TEXCOORD0
-    ) : SV_Target0
+	) : SV_Target0
 {
-    float2 UV = UVAndScreenPos.xy;
-    return Texture2DSample(SrcTexture, SrcSampler, UV);
-}
+	float2 UV = UVAndScreenPos.xy;
+	return Texture2DSample(InputTexture, InputSampler, UV);
+}
\ No newline at end of file
diff --git a/Engine/Source/Runtime/RenderCore/Public/CommonRenderResources.h b/Engine/Source/Runtime/RenderCore/Public/CommonRenderResources.h
--- a/Engine/Source/Runtime/RenderCore/Public/CommonRenderResources.h
+++ b/Engine/Source/Runtime/RenderCore/Public/CommonRenderResources.h
@@ -118,7 +118,8 @@ class RENDERCORE_API FCopyRectPS : public FGlobalShader
 	}
 
 	BEGIN_SHADER_PARAMETER_STRUCT(FParameters, )
-		SHADER_PARAMETER_TEXTURE(Texture2D, SrcTexture)
-		SHADER_PARAMETER_SAMPLER(SamplerState, SrcSampler)
+		SHADER_PARAMETER_RDG_TEXTURE(Texture2D, InputTexture)
+		SHADER_PARAMETER_SAMPLER(SamplerState, InputSampler)
+		RENDER_TARGET_BINDING_SLOTS()
 	END_SHADER_PARAMETER_STRUCT()
 };
diff --git a/Engine/Source/Runtime/Renderer/Private/PostProcess/PostProcessMaterial.cpp b/Engine/Source/Runtime/Renderer/Private/PostProcess/PostProcessMaterial.cpp
--- a/Engine/Source/Runtime/Renderer/Private/PostProcess/PostProcessMaterial.cpp
+++ b/Engine/Source/Runtime/Renderer/Private/PostProcess/PostProcessMaterial.cpp
@@ -18,8 +18,12 @@
 #include "RHI/Public/PipelineStateCache.h"
 #include "PostProcessing.h"
 #include "PostProcessMobile.h"
+#include "BufferVisualizationData.h"
 
-static TAutoConsoleVariable<int32> CVarPostProcessAllowStencilTest(
+namespace
+{
+
+TAutoConsoleVariable<int32> CVarPostProcessAllowStencilTest(
 	TEXT("r.PostProcessAllowStencilTest"),
 	1,
 	TEXT("Enables stencil testing in post process materials.\n")
@@ -28,7 +32,7 @@ static TAutoConsoleVariable<int32> CVarPostProcessAllowStencilTest(
 	TEXT("2: allow stencil testing and, if necessary, making a copy of custom depth/stencil buffer\n")
 	);
 
-static TAutoConsoleVariable<int32> CVarPostProcessAllowBlendModes(
+TAutoConsoleVariable<int32> CVarPostProcessAllowBlendModes(
 	TEXT("r.PostProcessAllowBlendModes"),
 	1,
 	TEXT("Enables blend modes in post process materials.\n")
@@ -36,610 +40,748 @@ static TAutoConsoleVariable<int32> CVarPostProcessAllowBlendModes(
 	TEXT("1: allow blend modes\n")
 	);
 
-enum class EPostProcessMaterialTarget
+TAutoConsoleVariable<int32> CVarPostProcessingDisableMaterials(
+	TEXT("r.PostProcessing.DisableMaterials"),
+	0,
+	TEXT(" Allows to disable post process materials. \n"),
+	ECVF_Scalability | ECVF_RenderThreadSafe);
+
+bool IsPostProcessStencilTestAllowed()
+{
+	return CVarPostProcessAllowStencilTest.GetValueOnRenderThread() != 0;
+}
+
+bool IsPostProcessStencilTestWithCopyAllowed()
+{
+	return CVarPostProcessAllowStencilTest.GetValueOnRenderThread() == 2;
+}
+
+bool IsCustomDepthEnabled()
+{
+	static const IConsoleVariable* CVarCustomDepth = IConsoleManager::Get().FindConsoleVariable(TEXT("r.CustomDepth"));
+
+	check(CVarCustomDepth);
+
+	return CVarCustomDepth->GetInt() == 3;
+}
+
+enum class ECustomDepthPolicy : uint32
 {
-	HighEnd,
-	Mobile
+	// Custom depth is disabled.
+	Disabled,
+
+	// Custom Depth-Stencil is enabled; potentially simultaneous SRV / DSV usage.
+	Enabled,
+
+	// Custom Depth-Stencil is enabled; but makes a copy of the target to avoid simultaneous SRV / DSV usage.
+	EnabledWithCopy
 };
 
-static bool ShouldCachePostProcessMaterial(EPostProcessMaterialTarget MaterialTarget, EShaderPlatform Platform, const FMaterial* Material)
+ECustomDepthPolicy GetMaterialCustomDepthPolicy(const FMaterial* Material, ERHIFeatureLevel::Type FeatureLevel)
 {
-	if (Material->GetMaterialDomain() == MD_PostProcess)
+	check(Material);
+
+	// Material requesting stencil test and post processing CVar allows it.
+	if (Material->IsStencilTestEnabled() && IsPostProcessStencilTestAllowed())
 	{
-		switch (MaterialTarget)
+		// Custom stencil texture allocated and available.
+		if (IsCustomDepthEnabled())
+		{
+			if (FeatureLevel == ERHIFeatureLevel::SM4)
+			{
+				const FMaterialShaderMap* MaterialShaderMap = Material->GetRenderingThreadShaderMap();
+
+				const bool bIsDepthStencilReadByShader = MaterialShaderMap->UsesSceneTexture(PPI_CustomDepth) || MaterialShaderMap->UsesSceneTexture(PPI_CustomStencil);
+
+				if (!bIsDepthStencilReadByShader)
+				{
+					return ECustomDepthPolicy::Enabled;
+				}
+				else if (IsPostProcessStencilTestWithCopyAllowed())
+				{
+					return ECustomDepthPolicy::EnabledWithCopy;
+				}
+			}
+			else
+			{
+				return ECustomDepthPolicy::Enabled;
+			}
+		}
+		else
 		{
-		case EPostProcessMaterialTarget::HighEnd:
-			return IsFeatureLevelSupported(Platform, ERHIFeatureLevel::SM4);
-		case EPostProcessMaterialTarget::Mobile:
-			return IsMobilePlatform(Platform) && IsMobileHDR();
+			UE_LOG(LogRenderer, Warning, TEXT("PostProcessMaterial uses stencil test, but stencil not allocated. Set r.CustomDepth to 3 to allocate custom stencil."));
 		}
 	}
 
-	return false;
+	return ECustomDepthPolicy::Disabled;
 }
 
-template<EPostProcessMaterialTarget MaterialTarget, bool bSwitchVerticalAxis>
-class FPostProcessMaterialVS : public FMaterialShader
+void GetMaterialInfo(
+	const UMaterialInterface* InMaterialInterface,
+	ERHIFeatureLevel::Type InFeatureLevel,
+	EPixelFormat InOutputFormat,
+	const FMaterial*& OutMaterial,
+	const FMaterialRenderProxy*& OutMaterialProxy,
+	const FMaterialShaderMap*& OutMaterialShaderMap)
 {
-	DECLARE_SHADER_TYPE(FPostProcessMaterialVS, Material);
-public:
+	const FMaterialRenderProxy* MaterialProxy = InMaterialInterface->GetRenderProxy();
+	check(MaterialProxy);
 
-	/**
-	  * Only compile these shaders for post processing domain materials
-	  */
-	static bool ShouldCompilePermutation(const FMaterialShaderPermutationParameters& Parameters)
+	const FMaterial* Material = MaterialProxy->GetMaterialNoFallback(InFeatureLevel);
+
+	if (!Material || Material->GetMaterialDomain() != MD_PostProcess || !Material->GetRenderingThreadShaderMap())
 	{
-		return ShouldCachePostProcessMaterial(MaterialTarget, Parameters.Platform, Parameters.Material)
-			// compile mobile axis switched versions only for after tonemapper passes
-			&& (bSwitchVerticalAxis == false || (RHINeedsToSwitchVerticalAxis(Parameters.Platform) && Parameters.Material->GetBlendableLocation() == BL_AfterTonemapping));
+		// Fallback to the default post process material.
+		const UMaterialInterface* DefaultMaterial = UMaterial::GetDefaultMaterial(MD_PostProcess);
+		check(DefaultMaterial);
+		check(DefaultMaterial != InMaterialInterface);
+
+		return GetMaterialInfo(
+			DefaultMaterial,
+			InFeatureLevel,
+			InOutputFormat,
+			OutMaterial,
+			OutMaterialProxy,
+			OutMaterialShaderMap);
 	}
 
-	static void ModifyCompilationEnvironment(const FMaterialShaderPermutationParameters& Parameters, FShaderCompilerEnvironment& OutEnvironment)
+	if (Material->IsStencilTestEnabled() || Material->GetBlendableOutputAlpha())
 	{
-		FMaterialShader::ModifyCompilationEnvironment(Parameters, OutEnvironment);
+		// Only allowed to have blend/stencil test if output format is compatible with ePId_Input0. 
+		// PF_Unknown implies output format is that of EPId_Input0
+		ensure(InOutputFormat == PF_Unknown);
+	}
 
-		OutEnvironment.SetDefine(TEXT("POST_PROCESS_MATERIAL"), 1);
-		OutEnvironment.SetDefine(TEXT("NEEDTOSWITCHVERTICLEAXIS"), bSwitchVerticalAxis);
+	const FMaterialShaderMap* MaterialShaderMap = Material->GetRenderingThreadShaderMap();;
+	check(MaterialShaderMap);
 
-		if (MaterialTarget == EPostProcessMaterialTarget::Mobile)
-		{
-			OutEnvironment.SetDefine(TEXT("POST_PROCESS_MATERIAL_BEFORE_TONEMAP"), (Parameters.Material->GetBlendableLocation() != BL_AfterTonemapping) ? 1 : 0);
-		}
-	}
-	
-	FPostProcessMaterialVS( )	{ }
-	FPostProcessMaterialVS(const ShaderMetaType::CompiledShaderInitializerType& Initializer)
-		: FMaterialShader(Initializer)
-	{
-		PostprocessParameter.Bind(Initializer.ParameterMap);
-	}
+	OutMaterial = Material;
+	OutMaterialProxy = MaterialProxy;
+	OutMaterialShaderMap = MaterialShaderMap;
+}
 
-	void SetParameters(FRHICommandList& RHICmdList, const FRenderingCompositePassContext& Context, const FMaterialRenderProxy* Proxy)
+FRHIDepthStencilState* GetMaterialStencilState(const FMaterial* Material)
+{
+	static FRHIDepthStencilState* StencilStates[] =
 	{
-		FRHIVertexShader* ShaderRHI = GetVertexShader();
-		FMaterialShader::SetParameters(RHICmdList, ShaderRHI, Proxy, *Proxy->GetMaterial(Context.View.GetFeatureLevel()), Context.View, Context.View.ViewUniformBuffer, ESceneTextureSetupMode::All);
-		PostprocessParameter.SetVS(ShaderRHI, Context, TStaticSamplerState<SF_Point, AM_Clamp, AM_Clamp, AM_Clamp>::GetRHI());
-	}
+		TStaticDepthStencilState<false, CF_Always, true, CF_Less>::GetRHI(),
+		TStaticDepthStencilState<false, CF_Always, true, CF_LessEqual>::GetRHI(),
+		TStaticDepthStencilState<false, CF_Always, true, CF_Greater>::GetRHI(),
+		TStaticDepthStencilState<false, CF_Always, true, CF_GreaterEqual>::GetRHI(),
+		TStaticDepthStencilState<false, CF_Always, true, CF_Equal>::GetRHI(),
+		TStaticDepthStencilState<false, CF_Always, true, CF_NotEqual>::GetRHI(),
+		TStaticDepthStencilState<false, CF_Always, true, CF_Never>::GetRHI(),
+		TStaticDepthStencilState<false, CF_Always, true, CF_Always>::GetRHI(),
+	};
+	static_assert(EMaterialStencilCompare::MSC_Count == ARRAY_COUNT(StencilStates), "Ensure that all EMaterialStencilCompare values are accounted for.");
+
+	check(Material);
+
+	return StencilStates[Material->GetStencilCompare()];
+}
 
-	// Begin FShader interface
-	virtual bool Serialize(FArchive& Ar) override
+bool IsMaterialBlendEnabled(const FMaterial* Material)
+{
+	check(Material);
+
+	return Material->GetBlendableOutputAlpha() && CVarPostProcessAllowBlendModes.GetValueOnRenderThread() != 0;
+}
+
+FRHIBlendState* GetMaterialBlendState(const FMaterial* Material)
+{
+	static FRHIBlendState* BlendStates[] =
 	{
-		bool bShaderHasOutdatedParameters = FMaterialShader::Serialize(Ar);
-		Ar << PostprocessParameter;
-		return bShaderHasOutdatedParameters;
-	}
-	//  End FShader interface 
-private:
-	FPostProcessPassParameters PostprocessParameter;
-};
+		TStaticBlendState<>::GetRHI(),
+		TStaticBlendState<>::GetRHI(),
+		TStaticBlendState<CW_RGB, BO_Add, BF_SourceAlpha, BF_InverseSourceAlpha, BO_Add, BF_Zero, BF_One>::GetRHI(),
+		TStaticBlendState<CW_RGB, BO_Add, BF_One, BF_One>::GetRHI(),
+		TStaticBlendState<CW_RGB, BO_Add, BF_DestColor, BF_Zero>::GetRHI(),
+		TStaticBlendState<CW_RGBA, BO_Add, BF_One, BF_InverseSourceAlpha, BO_Add, BF_One, BF_InverseSourceAlpha>::GetRHI(),
+		TStaticBlendState<CW_RGBA, BO_Add, BF_Zero, BF_InverseSourceAlpha, BO_Add, BF_Zero, BF_InverseSourceAlpha>::GetRHI(),
+	};
+	static_assert(EBlendMode::BLEND_MAX == ARRAY_COUNT(BlendStates), "Ensure that all EBlendMode values are accounted for.");
 
-typedef FPostProcessMaterialVS<EPostProcessMaterialTarget::HighEnd, false> FPostProcessMaterialVS_HighEnd;
-typedef FPostProcessMaterialVS<EPostProcessMaterialTarget::Mobile, false> FPostProcessMaterialVS_Mobile;
-typedef FPostProcessMaterialVS<EPostProcessMaterialTarget::Mobile, true> FPostProcessMaterialVS_Mobile_AxisSwitch;
+	check(Material);
 
-IMPLEMENT_MATERIAL_SHADER_TYPE(template<>,FPostProcessMaterialVS_HighEnd,TEXT("/Engine/Private/PostProcessMaterialShaders.usf"),TEXT("MainVS"),SF_Vertex);
-IMPLEMENT_MATERIAL_SHADER_TYPE(template<>,FPostProcessMaterialVS_Mobile,TEXT("/Engine/Private/PostProcessMaterialShaders.usf"),TEXT("MainVS_ES2"),SF_Vertex);
-IMPLEMENT_MATERIAL_SHADER_TYPE(template<>,FPostProcessMaterialVS_Mobile_AxisSwitch,TEXT("/Engine/Private/PostProcessMaterialShaders.usf"),TEXT("MainVS_ES2"),SF_Vertex);
+	return BlendStates[Material->GetBlendMode()];
+}
 
-/**
- * A pixel shader for rendering a post process material
- */
-template<EPostProcessMaterialTarget MaterialTarget, uint32 UVPolicy, bool bSwitchVerticalAxis>
-class FPostProcessMaterialPS : public FMaterialShader
+class FPostProcessMaterialShader : public FMaterialShader
 {
-	DECLARE_SHADER_TYPE(FPostProcessMaterialPS,Material);
 public:
+	using FParameters = FPostProcessMaterialParameters;
+	SHADER_USE_PARAMETER_STRUCT_WITH_LEGACY_BASE(FPostProcessMaterialShader, FMaterialShader);
+
+	class FMobileDimension : SHADER_PERMUTATION_BOOL("POST_PROCESS_MATERIAL_MOBILE");
+	using FPermutationDomain = TShaderPermutationDomain<FMobileDimension>;
 
-	/**
-	  * Only compile these shaders for post processing domain materials
-	  */
 	static bool ShouldCompilePermutation(const FMaterialShaderPermutationParameters& Parameters)
 	{
-		return ShouldCachePostProcessMaterial(MaterialTarget, Parameters.Platform, Parameters.Material)
-			// compile mobile axis switched versions only for after tonemapper passes
-			&& (bSwitchVerticalAxis == false || (RHINeedsToSwitchVerticalAxis(Parameters.Platform) && Parameters.Material->GetBlendableLocation() == BL_AfterTonemapping));
+		if (Parameters.Material->GetMaterialDomain() == MD_PostProcess)
+		{
+			const FPermutationDomain PermutationVector(Parameters.PermutationId);
+
+			if (PermutationVector.Get<FMobileDimension>())
+			{
+				return IsMobilePlatform(Parameters.Platform) && IsMobileHDR();
+			}
+			else
+			{
+				return IsFeatureLevelSupported(Parameters.Platform, ERHIFeatureLevel::SM4);
+			}
+		}
+		return false;
 	}
 
 	static void ModifyCompilationEnvironment(const FMaterialShaderPermutationParameters& Parameters, FShaderCompilerEnvironment& OutEnvironment)
 	{
 		FMaterialShader::ModifyCompilationEnvironment(Parameters, OutEnvironment);
-
 		OutEnvironment.SetDefine(TEXT("POST_PROCESS_MATERIAL"), 1);
-		OutEnvironment.SetDefine(TEXT("POST_PROCESS_MATERIAL_UV_POLICY"), UVPolicy);
 
-		EBlendableLocation Location = EBlendableLocation(Parameters.Material->GetBlendableLocation());
-		OutEnvironment.SetDefine(TEXT("POST_PROCESS_MATERIAL_AFTER_TAA_UPSAMPLE"), (Location == BL_AfterTonemapping || Location == BL_ReplacingTonemapper) ? 1 : 0);
-
-		OutEnvironment.SetDefine(TEXT("NEEDTOSWITCHVERTICLEAXIS"), bSwitchVerticalAxis);
-
-		if (MaterialTarget == EPostProcessMaterialTarget::Mobile)
+		FPermutationDomain PermutationVector(Parameters.PermutationId);
+		if (PermutationVector.Get<FMobileDimension>())
 		{
 			OutEnvironment.SetDefine(TEXT("MOBILE_FORCE_DEPTH_TEXTURE_READS"), 1); // Ensure post process materials will not attempt depth buffer fetch operations.
-			OutEnvironment.SetDefine(TEXT("POST_PROCESS_MATERIAL_BEFORE_TONEMAP"), (Location != BL_AfterTonemapping) ? 1 : 0);
+			OutEnvironment.SetDefine(TEXT("POST_PROCESS_MATERIAL_BEFORE_TONEMAP"), (Parameters.Material->GetBlendableLocation() != BL_AfterTonemapping) ? 1 : 0);
 		}
 	}
 
-	FPostProcessMaterialPS() {}
-	FPostProcessMaterialPS(const ShaderMetaType::CompiledShaderInitializerType& Initializer):
-		FMaterialShader(Initializer)
+protected:
+	template <typename TRHIShader>
+	void SetParameters(FRHICommandList& RHICmdList, TRHIShader* ShaderRHI, const FViewInfo& View, const FMaterialRenderProxy* Proxy, const FParameters& Parameters)
 	{
-		PostprocessParameter.Bind(Initializer.ParameterMap);
+		FMaterialShader::SetParameters(RHICmdList, ShaderRHI, Proxy, *Proxy->GetMaterial(View.GetFeatureLevel()), View, View.ViewUniformBuffer, ESceneTextureSetupMode::All);
+		SetShaderParameters(RHICmdList, this, ShaderRHI, Parameters);
 	}
+};
 
-	template <typename TRHICmdList>
-	void SetParameters(TRHICmdList& RHICmdList, const FRenderingCompositePassContext& Context, const FMaterialRenderProxy* Proxy)
-	{
-		FRHIPixelShader* ShaderRHI = GetPixelShader();
-
-		FMaterialShader::SetParameters(RHICmdList, ShaderRHI, Proxy, *Proxy->GetMaterial(Context.View.GetFeatureLevel()), Context.View, Context.View.ViewUniformBuffer, ESceneTextureSetupMode::All);
-		PostprocessParameter.SetPS(RHICmdList, ShaderRHI, Context, TStaticSamplerState<SF_Point,AM_Clamp,AM_Clamp,AM_Clamp>::GetRHI());
-	}
+class FPostProcessMaterialVS : public FPostProcessMaterialShader
+{
+public:
+	DECLARE_MATERIAL_SHADER(FPostProcessMaterialVS);
 
-	virtual bool Serialize(FArchive& Ar) override
+	void SetParameters(FRHICommandList& RHICmdList, const FViewInfo& View, const FMaterialRenderProxy* Proxy, const FParameters& Parameters)
 	{
-		bool bShaderHasOutdatedParameters = FMaterialShader::Serialize(Ar);
-		Ar << PostprocessParameter;
-		return bShaderHasOutdatedParameters;
+		FPostProcessMaterialShader::SetParameters(RHICmdList, GetVertexShader(), View, Proxy, Parameters);
 	}
 
-private:
-	FPostProcessPassParameters PostprocessParameter;
+	FPostProcessMaterialVS() = default;
+	FPostProcessMaterialVS(const ShaderMetaType::CompiledShaderInitializerType& Initializer)
+		: FPostProcessMaterialShader(Initializer)
+	{}
 };
 
-typedef FPostProcessMaterialPS<EPostProcessMaterialTarget::HighEnd, 0, false> FFPostProcessMaterialPS_HighEnd0;
-typedef FPostProcessMaterialPS<EPostProcessMaterialTarget::HighEnd, 1, false> FFPostProcessMaterialPS_HighEnd1;
-typedef FPostProcessMaterialPS<EPostProcessMaterialTarget::Mobile, 0, false> FPostProcessMaterialPS_Mobile0;
-typedef FPostProcessMaterialPS<EPostProcessMaterialTarget::Mobile, 1, false> FPostProcessMaterialPS_Mobile1;
-typedef FPostProcessMaterialPS<EPostProcessMaterialTarget::Mobile, 0, true> FPostProcessMaterialPS_Mobile0_AxisSwitch;
-typedef FPostProcessMaterialPS<EPostProcessMaterialTarget::Mobile, 1,true> FPostProcessMaterialPS_Mobile1_AxisSwitch;
-
-IMPLEMENT_MATERIAL_SHADER_TYPE(template<>, FFPostProcessMaterialPS_HighEnd0, TEXT("/Engine/Private/PostProcessMaterialShaders.usf"), TEXT("MainPS"), SF_Pixel);
-IMPLEMENT_MATERIAL_SHADER_TYPE(template<>, FFPostProcessMaterialPS_HighEnd1, TEXT("/Engine/Private/PostProcessMaterialShaders.usf"), TEXT("MainPS"), SF_Pixel);
-IMPLEMENT_MATERIAL_SHADER_TYPE(template<>, FPostProcessMaterialPS_Mobile0, TEXT("/Engine/Private/PostProcessMaterialShaders.usf"), TEXT("MainPS_ES2"), SF_Pixel);
-IMPLEMENT_MATERIAL_SHADER_TYPE(template<>, FPostProcessMaterialPS_Mobile1, TEXT("/Engine/Private/PostProcessMaterialShaders.usf"), TEXT("MainPS_ES2"), SF_Pixel);
-IMPLEMENT_MATERIAL_SHADER_TYPE(template<>, FPostProcessMaterialPS_Mobile0_AxisSwitch, TEXT("/Engine/Private/PostProcessMaterialShaders.usf"), TEXT("MainPS_ES2"), SF_Pixel);
-IMPLEMENT_MATERIAL_SHADER_TYPE(template<>, FPostProcessMaterialPS_Mobile1_AxisSwitch, TEXT("/Engine/Private/PostProcessMaterialShaders.usf"), TEXT("MainPS_ES2"), SF_Pixel);
-
-FRCPassPostProcessMaterial::FRCPassPostProcessMaterial(UMaterialInterface* InMaterialInterface, ERHIFeatureLevel::Type InFeatureLevel, EPixelFormat OutputFormatIN)
-: MaterialInterface(InMaterialInterface), OutputFormat(OutputFormatIN)
+class FPostProcessMaterialPS : public FPostProcessMaterialShader
 {
-	FMaterialRenderProxy* Proxy = MaterialInterface->GetRenderProxy();
-	check(Proxy);
+public:
+	DECLARE_MATERIAL_SHADER(FPostProcessMaterialPS);
 
-	const FMaterial* Material = Proxy->GetMaterialNoFallback(InFeatureLevel);
-	
-	if (!Material || Material->GetMaterialDomain() != MD_PostProcess)
+	void SetParameters(FRHICommandList& RHICmdList, const FViewInfo& View, const FMaterialRenderProxy* Proxy, const FParameters& Parameters)
 	{
-		MaterialInterface = UMaterial::GetDefaultMaterial(MD_PostProcess);
+		FPostProcessMaterialShader::SetParameters(RHICmdList, GetPixelShader(), View, Proxy, Parameters);
 	}
 
-	if (Material && (Material->IsStencilTestEnabled() || Material->GetBlendableOutputAlpha()))
-	{
-		// Only allowed to have blend/stencil test if output format is compatible with ePId_Input0. 
-		// PF_Unknown implies output format is that of EPId_Input0
-		ensure(OutputFormat == PF_Unknown);
-	}
-}
-		
-/** The filter vertex declaration resource type. */
+	FPostProcessMaterialPS() = default;
+	FPostProcessMaterialPS(const ShaderMetaType::CompiledShaderInitializerType& Initializer)
+		: FPostProcessMaterialShader(Initializer)
+	{}
+};
+
+IMPLEMENT_MATERIAL_SHADER(FPostProcessMaterialVS, "/Engine/Private/PostProcessMaterialShaders.usf", "MainVS", SF_Vertex);
+IMPLEMENT_MATERIAL_SHADER(FPostProcessMaterialPS, "/Engine/Private/PostProcessMaterialShaders.usf", "MainPS", SF_Pixel);
+
 class FPostProcessMaterialVertexDeclaration : public FRenderResource
 {
 public:
 	FVertexDeclarationRHIRef VertexDeclarationRHI;
-	
-	/** Destructor. */
-	virtual ~FPostProcessMaterialVertexDeclaration() {}
-	
-	virtual void InitRHI()
+
+	void InitRHI() override
 	{
 		FVertexDeclarationElementList Elements;
 		uint32 Stride = sizeof(FFilterVertex);
-		Elements.Add(FVertexElement(0,STRUCT_OFFSET(FFilterVertex,Position),VET_Float4,0,Stride));
+		Elements.Add(FVertexElement(0, STRUCT_OFFSET(FFilterVertex, Position), VET_Float4, 0, Stride));
 		VertexDeclarationRHI = PipelineStateCache::GetOrCreateVertexDeclaration(Elements);
 	}
-	
-	virtual void ReleaseRHI()
+
+	void ReleaseRHI() override
 	{
 		VertexDeclarationRHI.SafeRelease();
 	}
 };
+
 TGlobalResource<FPostProcessMaterialVertexDeclaration> GPostProcessMaterialVertexDeclaration;
 
-template<typename TVertexShader, typename TPixelShader>
-FShader* SetMobileShaders(const FMaterialShaderMap* MaterialShaderMap, FGraphicsPipelineStateInitializer &GraphicsPSOInit, FRenderingCompositePassContext &Context, FMaterialRenderProxy* Proxy)
+} //! namespace
+
+FPostProcessMaterialInput GetPostProcessMaterialInput(FIntRect ViewportRect, FRDGTextureRef Texture, FRHISamplerState* Sampler)
 {
-	TPixelShader* PixelShader_Mobile = MaterialShaderMap->GetShader<TPixelShader>();
-	TVertexShader* VertexShader_Mobile = MaterialShaderMap->GetShader<TVertexShader>();
+	FPostProcessMaterialInput Input;
+	Input.Viewport = GetScreenPassTextureViewportParameters(FScreenPassTextureViewport(ViewportRect, Texture));
+	Input.Texture = Texture;
+	Input.Sampler = Sampler;
+	return Input;
+}
 
-	GraphicsPSOInit.BoundShaderState.VertexDeclarationRHI = GFilterVertexDeclaration.VertexDeclarationRHI;
-	GraphicsPSOInit.BoundShaderState.VertexShaderRHI = GETSAFERHISHADER_VERTEX(VertexShader_Mobile);
-	GraphicsPSOInit.BoundShaderState.PixelShaderRHI = GETSAFERHISHADER_PIXEL(PixelShader_Mobile);
+FRDGTextureRef ComputePostProcessMaterial(
+	FRDGBuilder& GraphBuilder,
+	const FScreenPassViewInfo& ScreenPassView,
+	const UMaterialInterface* MaterialInterface,
+	const FPostProcessMaterialInputs& Inputs)
+{
+	Inputs.Validate();
 
-	SetGraphicsPipelineState(Context.RHICmdList, GraphicsPSOInit);
+	FRDGTextureRef SceneColorTexture = nullptr;
+	FIntRect SceneColorViewportRect;
+	Inputs.GetInput(EPostProcessMaterialInput::SceneColor, SceneColorTexture, SceneColorViewportRect);
 
-	VertexShader_Mobile->SetParameters(Context.RHICmdList, Context, Proxy);
-	PixelShader_Mobile->SetParameters(Context.RHICmdList, Context, Proxy);
-	return VertexShader_Mobile;
-}
+	const FScreenPassTextureViewport SceneColorViewport(SceneColorViewportRect, SceneColorTexture);
 
-void FRCPassPostProcessMaterial::Process(FRenderingCompositePassContext& Context)
-{
-	FMaterialRenderProxy* Proxy = MaterialInterface->GetRenderProxy();
-	check(Proxy);
+	const ERHIFeatureLevel::Type FeatureLevel = ScreenPassView.View.GetFeatureLevel();
 
-	ERHIFeatureLevel::Type FeatureLevel = Context.View.GetFeatureLevel();
+	const FMaterial* Material = nullptr;
+	const FMaterialRenderProxy* MaterialRenderProxy = nullptr;
+	const FMaterialShaderMap* MaterialShaderMap = nullptr;
+	GetMaterialInfo(MaterialInterface, FeatureLevel, Inputs.OverrideOutputFormat, Material, MaterialRenderProxy, MaterialShaderMap);
 
-	FSceneRenderTargets& SceneContext = FSceneRenderTargets::Get(Context.RHICmdList);
-	
-	const FMaterial* Material = Proxy->GetMaterial(FeatureLevel);
-	check(Material);
+	RDG_EVENT_SCOPE(GraphBuilder, "PostProcessMaterial %dx%d Material=%s", SceneColorViewport.Rect.Width(), SceneColorViewport.Rect.Height(), *Material->GetFriendlyName());
+
+	FRHIDepthStencilState* DefaultDepthStencilState = FScreenPassDrawInfo::FDefaultDepthStencilState::GetRHI();
+	FRHIDepthStencilState* DepthStencilState = DefaultDepthStencilState;
 
-	const FMaterialShaderMap* MaterialShaderMap = Material->GetRenderingThreadShaderMap();
+	FRDGTextureRef DepthStencilTexture = nullptr;
+	FRDGTextureRef DepthStencilTextureForSRV = nullptr;
 
-	int32 AllowStencilTest = CVarPostProcessAllowStencilTest.GetValueOnRenderThread();
-	bool bReadsCustomDepthStencil = MaterialShaderMap->UsesSceneTexture(PPI_CustomDepth) || MaterialShaderMap->UsesSceneTexture(PPI_CustomStencil);
-	bool bDoStencilTest = false;
-	if (Material->IsStencilTestEnabled() && AllowStencilTest > 0)
+	// Allocate custom depth stencil texture(s) and depth stencil state.
+	const ECustomDepthPolicy CustomStencilPolicy = GetMaterialCustomDepthPolicy(Material, FeatureLevel);
+
+	if (CustomStencilPolicy != ECustomDepthPolicy::Disabled)
 	{
-		static const IConsoleVariable* CVarCustomDepth = IConsoleManager::Get().FindConsoleVariable(TEXT("r.CustomDepth"));
+		check(Inputs.CustomDepthTexture);
+
+		DepthStencilTexture = Inputs.CustomDepthTexture;
 
-		int32 CustomDepthSetting = CVarCustomDepth->GetInt();
-		if (CustomDepthSetting == 3)
+		if (CustomStencilPolicy == ECustomDepthPolicy::EnabledWithCopy)
 		{
-			bool AllowStencilTestWithCopy = AllowStencilTest == 2;
-
-			// do the stencil test if 
-			// not SM4 
-			//   OR 
-			// reads DS but allowed make DS copy 
-			//   OR 
-			// DS not read at all.
-			bDoStencilTest = (FeatureLevel != ERHIFeatureLevel::SM4) || 
-				((bReadsCustomDepthStencil == AllowStencilTestWithCopy) || AllowStencilTestWithCopy);
+			// NOTE: SM4 does not allow depth stencil to be a destination for a copy. Therefore,
+			// we create a shader resource texture and copy into that.
+			FRDGTextureDesc CopyDesc = Inputs.CustomDepthTexture->Desc;
+			CopyDesc.Flags = TexCreate_None;
+			CopyDesc.TargetableFlags = TexCreate_ShaderResource;
+
+			DepthStencilTextureForSRV = GraphBuilder.CreateTexture(CopyDesc, TEXT("CustomDepthCopy"));
+
+			AddCopyTexturePass(GraphBuilder, DepthStencilTexture, DepthStencilTextureForSRV);
 		}
 		else
 		{
-			UE_LOG(LogRenderer, Warning, TEXT("PostProcessMaterial uses stencil test, but custom stencil not allocated. Set r.CustomDepth to 3 to allocate custom stencil."));
+			DepthStencilTextureForSRV = DepthStencilTexture;
 		}
+
+		DepthStencilState = GetMaterialStencilState(Material);
 	}
 
-	FIntPoint RectSize = Context.SceneColorViewRect.Size();
-	SCOPED_DRAW_EVENTF(Context.RHICmdList, PostProcessMaterial, TEXT("PostProcessMaterial %dx%d Material=%s"), RectSize.X, RectSize.Y, *Material->GetFriendlyName());
+	FRHIBlendState* DefaultBlendState = FScreenPassDrawInfo::FDefaultBlendState::GetRHI();
+	FRHIBlendState* BlendState = GetMaterialBlendState(Material);
 
-	// Copy of custom depth/stencil buffer if HW does not support simultaneously a texture bound as DepthRead_StencilRead and SRV
-	TRefCountPtr<IPooledRenderTarget> CustomDepthStencilCopy;
+	const FViewInfo& View = ScreenPassView.View;
 
-	const FSceneRenderTargetItem* CustomDepthStencilTarget = nullptr;
+	// Blend / Depth Stencil usage requires that the render target have primed color data.
+	const bool bIsCompositeWithInput = DepthStencilState != DefaultDepthStencilState || BlendState != DefaultBlendState;
 
-	FRHIDepthStencilState* DepthStencilState;
-	uint32 StencilRefValue = 0;
+	// Multiple views are composited onto the same target.
+	const bool bIsNotPrimaryView = (&View != View.Family->Views[0]);
 
-	if (bDoStencilTest)
-	{
-		CustomDepthStencilTarget = &SceneContext.CustomDepth->GetRenderTargetItem();
+	// We only prime color on the output texture if we are using fixed function Blend / Depth-Stencil,
+	// or we need to retain previously rendered views.
+	const bool bPrimeOutputColor = bIsCompositeWithInput || bIsNotPrimaryView;
 
-		// SM4 HW lacks support for texture bound as DepthRead_StencilRead and SRV simultaneously thus make a copy of DS buffer
-		if (FeatureLevel == ERHIFeatureLevel::SM4 && bReadsCustomDepthStencil)
+	FRDGTextureRef OutputTexture = Inputs.OverrideOutputTexture;
+
+	// We can re-use the scene color texture as the render target if we're not simultaneously reading from it.
+	if (!OutputTexture && !MaterialShaderMap->UsesSceneTexture(PPI_PostProcessInput0))
+	{
+		OutputTexture = SceneColorTexture;
+	}
+	else
+	{
+		// Allocate new transient output texture if none exists.
+		if (!OutputTexture)
 		{
-			// Dest param of CopyResource() call can only be an SRV (No render target flags) on DX10.0 (SM4)
-			FPooledRenderTargetDesc DSCopyDesc = SceneContext.CustomDepth->GetDesc();
-			DSCopyDesc.Flags = TexCreate_None;
-			DSCopyDesc.TargetableFlags = TexCreate_ShaderResource;
-
-			GRenderTargetPool.FindFreeElement(Context.RHICmdList, DSCopyDesc, CustomDepthStencilCopy, TEXT("CustomDepthStencilCopy"));
-
-			Context.RHICmdList.CopyTexture(
-				SceneContext.CustomDepth->GetRenderTargetItem().ShaderResourceTexture,
-				CustomDepthStencilCopy->GetRenderTargetItem().ShaderResourceTexture,
-				FRHICopyTextureInfo()
-				);
-
-			// The copy DS buffer was created only as a SRV (no render target flags), thus swap with real CustomDepth buffer
-			// with copy to allow DS buffer with render target flag to be set as DepthRead_StencilRead and DS buffer copy to
-			// be set as SRV for post process material.
-			Swap(SceneContext.CustomDepth, CustomDepthStencilCopy);
+			const EPixelFormat OverrideOutputFormat = Inputs.OverrideOutputFormat;
+
+			FRDGTextureDesc OutputTextureDesc = SceneColorTexture->Desc;
+			if (OverrideOutputFormat != PF_Unknown)
+			{
+				OutputTextureDesc.Format = OverrideOutputFormat;
+			}
+			OutputTextureDesc.Reset();
+			OutputTextureDesc.ClearValue = FClearValueBinding(FLinearColor::Black);
+			OutputTextureDesc.Flags |= GFastVRamConfig.PostProcessMaterial;
+
+			OutputTexture = GraphBuilder.CreateTexture(OutputTextureDesc, TEXT("PostProcessMaterialOutput"));
 		}
 
-		static FRHIDepthStencilState* StencilStates[] =
+		if (bPrimeOutputColor)
 		{
-			TStaticDepthStencilState<false, CF_Always, true, CF_Less>::GetRHI(),
-			TStaticDepthStencilState<false, CF_Always, true, CF_LessEqual>::GetRHI(),
-			TStaticDepthStencilState<false, CF_Always, true, CF_Greater>::GetRHI(),
-			TStaticDepthStencilState<false, CF_Always, true, CF_GreaterEqual>::GetRHI(),
-			TStaticDepthStencilState<false, CF_Always, true, CF_Equal>::GetRHI(),
-			TStaticDepthStencilState<false, CF_Always, true, CF_NotEqual>::GetRHI(),
-			TStaticDepthStencilState<false, CF_Always, true, CF_Never>::GetRHI(),
-			TStaticDepthStencilState<false, CF_Always, true, CF_Always>::GetRHI(),
-		};
-		static_assert(EMaterialStencilCompare::MSC_Count == ARRAY_COUNT(StencilStates), "Ensure that all EMaterialStencilCompare values are accounted for.");
-
-		DepthStencilState = StencilStates[Material->GetStencilCompare()];
-		StencilRefValue = Material->GetStencilRefValue();
+			// Copy existing contents to new output and use load-action to preserve untouched pixels.
+			AddDrawTexturePass(GraphBuilder, ScreenPassView, SceneColorTexture, OutputTexture);
+		}
 	}
-	else
+
+	ERenderTargetLoadAction OutputLoadAction = ERenderTargetLoadAction::ENoAction;
+
+	// We have color data that needs to be loaded.
+	if (bPrimeOutputColor)
+	{
+		OutputLoadAction = ERenderTargetLoadAction::ELoad;
+	}
+	// HMD masks leave unrendered pixels; perform a clear instead.
+	else if (ScreenPassView.bHasHMDMask)
 	{
-		DepthStencilState = TStaticDepthStencilState<false, CF_Always>::GetRHI();
+		OutputLoadAction = ERenderTargetLoadAction::EClear;
 	}
 
-	FRHIBlendState* BlendState = TStaticBlendState<>::GetRHI();
-	bool bDoOutputBlend = Material->GetBlendableOutputAlpha() && CVarPostProcessAllowBlendModes.GetValueOnRenderThread() != 0;
-	if (bDoOutputBlend)
+	FPostProcessMaterialParameters* PostProcessMaterialParameters = GraphBuilder.AllocParameters<FPostProcessMaterialParameters>();
+
+	PostProcessMaterialParameters->PostProcessOutput = GetScreenPassTextureViewportParameters(SceneColorViewport);
+	PostProcessMaterialParameters->CustomDepth = DepthStencilTextureForSRV;
+
+	PostProcessMaterialParameters->RenderTargets[0] = FRenderTargetBinding(
+		OutputTexture,
+		OutputLoadAction,
+		ERenderTargetStoreAction::EStore);
+
+	if (DepthStencilTexture)
 	{
-		static FRHIBlendState* BlendStates[] =
-		{
-			TStaticBlendState<>::GetRHI(),
-			TStaticBlendState<>::GetRHI(),
-			TStaticBlendState<CW_RGB, BO_Add, BF_SourceAlpha, BF_InverseSourceAlpha, BO_Add, BF_Zero, BF_One>::GetRHI(),
-			TStaticBlendState<CW_RGB, BO_Add, BF_One, BF_One>::GetRHI(),
-			TStaticBlendState<CW_RGB, BO_Add, BF_DestColor, BF_Zero>::GetRHI(),
-			TStaticBlendState<CW_RGBA, BO_Add, BF_One, BF_InverseSourceAlpha, BO_Add, BF_One, BF_InverseSourceAlpha>::GetRHI(),
-			TStaticBlendState<CW_RGBA, BO_Add, BF_Zero, BF_InverseSourceAlpha, BO_Add, BF_Zero, BF_InverseSourceAlpha>::GetRHI(),
-		};
-		static_assert(EBlendMode::BLEND_MAX == ARRAY_COUNT(BlendStates), "Ensure that all EBlendMode values are accounted for.");
-
-		BlendState = BlendStates[Material->GetBlendMode()];
+		PostProcessMaterialParameters->RenderTargets.DepthStencil = FDepthStencilBinding(
+			DepthStencilTexture,
+			ERenderTargetLoadAction::ENoAction,
+			ERenderTargetStoreAction::ENoAction,
+			ERenderTargetLoadAction::ELoad,
+			ERenderTargetStoreAction::ENoAction,
+			FExclusiveDepthStencil::DepthRead_StencilRead);
 	}
 
-	// The PP target - either from the render target pool or the ePId_Input0
-	const FSceneRenderTargetItem* DestRenderTarget = nullptr;
-	ERenderTargetLoadAction DestRenderTargetLoadAction = ERenderTargetLoadAction::Num;
-		
-	if (bDoStencilTest || bDoOutputBlend)
+	PostProcessMaterialParameters->PostProcessInput_BilinearSampler = TStaticSamplerState<SF_Bilinear, AM_Clamp, AM_Clamp, AM_Clamp>::GetRHI();;
+
+	FRDGTextureRef BlackTexture = GraphBuilder.RegisterExternalTexture(GSystemTextures.BlackDummy);
+
+	FRHISamplerState* PointClampSampler = TStaticSamplerState<SF_Point, AM_Clamp, AM_Clamp, AM_Clamp>::GetRHI();
+
+	for (uint32 InputIndex = 0; InputIndex < kPostProcessMaterialInputCountMax; ++InputIndex)
 	{
-		if (!MaterialShaderMap->UsesSceneTexture(PPI_PostProcessInput0) && !PassOutputs[0].PooledRenderTarget)
+		FRDGTextureRef InputTexture = nullptr;
+		FIntRect InputViewportRect;
+		Inputs.GetInput((EPostProcessMaterialInput)InputIndex, InputTexture, InputViewportRect);
+
+		// Need to provide valid textures for when shader compilation doesn't cull unused parameters.
+		if (!InputTexture || !MaterialShaderMap->UsesSceneTexture(PPI_PostProcessInput0 + InputIndex))
 		{
-			PassOutputs[0].PooledRenderTarget = GetInput(ePId_Input0)->GetOutput()->RequestInput();
-			DestRenderTarget = &PassOutputs[0].RequestSurface(Context);
+			InputTexture = BlackTexture;
 		}
-		else
-		{
-			DestRenderTarget = &PassOutputs[0].RequestSurface(Context);
 
-			FRHITexture* DstTexture = DestRenderTarget->TargetableTexture;
-			FRHITexture* SrcTexture = GetInput(ePId_Input0)->GetOutput()->RequestSurface(Context).ShaderResourceTexture;
+		PostProcessMaterialParameters->PostProcessInput[InputIndex] = GetPostProcessMaterialInput(InputViewportRect, InputTexture, PointClampSampler);
+	}
 
-			// CopyResource() can only be called when format and size match. Otherwise must use shader to do stretch & format conversion.
-			if (DstTexture->GetFormat() == SrcTexture->GetFormat() && DstTexture->GetSizeXYZ() == SrcTexture->GetSizeXYZ())
-			{
-			Context.RHICmdList.CopyTexture(
-					SrcTexture,
-					DstTexture,
-				FRHICopyTextureInfo()
-				);
-		}
-			else
-			{
-				TShaderMap<FGlobalShaderType>* GlobalShaderMap = GetGlobalShaderMap(Context.View.FeatureLevel);
+	const bool bIsMobile = FeatureLevel <= ERHIFeatureLevel::ES3_1;
 
-				FIntRect DstRect = Context.GetSceneColorDestRect(*DestRenderTarget);
-				FIntRect SrcRect = Context.SceneColorViewRect;
+	const bool bFlipYAxis = Inputs.bFlipYAxis && RHINeedsToSwitchVerticalAxis(GShaderPlatformForFeatureLevel[FeatureLevel]);
 
-				FRHIRenderPassInfo RPInfo = FRHIRenderPassInfo(
-					DstTexture,
-					MakeRenderTargetActions(Context.GetLoadActionForRenderTarget(*DestRenderTarget), ERenderTargetStoreAction::EStore)
-					);
+	PostProcessMaterialParameters->bFlipYAxis = bFlipYAxis;
 
-				Context.RHICmdList.BeginRenderPass(RPInfo, TEXT("Copy Rect"));
-				{
-					Context.SetViewportAndCallRHI(DstRect);
-
-					FGraphicsPipelineStateInitializer GraphicsPSOInit;
-					Context.RHICmdList.ApplyCachedRenderTargets(GraphicsPSOInit);
-
-					GraphicsPSOInit.BlendState = TStaticBlendState<>::GetRHI();
-					GraphicsPSOInit.RasterizerState = TStaticRasterizerState<>::GetRHI();
-					GraphicsPSOInit.DepthStencilState = TStaticDepthStencilState<false, CF_Always>::GetRHI();
-
-					auto VS = GlobalShaderMap->GetShader<FPostProcessVS>();
-					auto PS = GlobalShaderMap->GetShader<FCopyRectPS>();
-
-					GraphicsPSOInit.BoundShaderState.VertexDeclarationRHI = GFilterVertexDeclaration.VertexDeclarationRHI;
-					GraphicsPSOInit.BoundShaderState.VertexShaderRHI = GETSAFERHISHADER_VERTEX(VS);
-					GraphicsPSOInit.BoundShaderState.PixelShaderRHI = GETSAFERHISHADER_PIXEL(PS);
-					GraphicsPSOInit.PrimitiveType = PT_TriangleList;
-
-					SetGraphicsPipelineState(Context.RHICmdList, GraphicsPSOInit);
-
-					FCopyRectPS::FParameters PixelParams;
-					PixelParams.SrcTexture = SrcTexture;
-					PixelParams.SrcSampler = TStaticSamplerState<>::GetRHI();
-
-					SetShaderParameters(Context.RHICmdList, PS, PS->GetPixelShader(), PixelParams);
-
-					DrawRectangle(
-						Context.RHICmdList,
-						DstRect.Min.X,
-						DstRect.Min.Y,
-						DstRect.Width(),
-						DstRect.Height(),
-						SrcRect.Min.X,
-						SrcRect.Min.Y,
-						SrcRect.Width(),
-						SrcRect.Height(),
-						DstTexture->GetTexture2D()->GetSizeXY(),
-						SrcTexture->GetTexture2D()->GetSizeXY(),
-						VS,
-						EDRF_UseTriangleOptimization
-						);
-				}
-				Context.RHICmdList.EndRenderPass();
-				Context.RHICmdList.CopyToResolveTarget(DestRenderTarget->TargetableTexture, DestRenderTarget->ShaderResourceTexture, FResolveParams());
-			}
+	FPostProcessMaterialShader::FPermutationDomain PermutationVector;
+	PermutationVector.Set<FPostProcessMaterialShader::FMobileDimension>(bIsMobile);
+
+	FPostProcessMaterialVS* VertexShader = MaterialShaderMap->GetShader<FPostProcessMaterialVS>(PermutationVector);
+	FPostProcessMaterialPS* PixelShader = MaterialShaderMap->GetShader<FPostProcessMaterialPS>(PermutationVector);
+
+	const uint32 MaterialStencilRef = Material->GetStencilRefValue();
+
+	const auto SetupFunction = [VertexShader, PixelShader, &ScreenPassView, MaterialRenderProxy, PostProcessMaterialParameters, MaterialStencilRef]
+		(FRHICommandListImmediate& RHICmdList)
+	{
+		VertexShader->SetParameters(RHICmdList, ScreenPassView.View, MaterialRenderProxy, *PostProcessMaterialParameters);
+		PixelShader->SetParameters(RHICmdList, ScreenPassView.View, MaterialRenderProxy, *PostProcessMaterialParameters);
+		RHICmdList.SetStencilRef(MaterialStencilRef);
+	};
+
+	const FScreenPassDrawInfo ScreenPassDraw(
+		VertexShader,
+		PixelShader,
+		BlendState,
+		DepthStencilState,
+		bFlipYAxis ? FScreenPassDrawInfo::EFlags::FlipYAxis : FScreenPassDrawInfo::EFlags::None);
+
+	ScreenPassDraw.Validate();
+
+	GraphBuilder.AddPass(
+		RDG_EVENT_NAME("PostProcessMaterial"),
+		PostProcessMaterialParameters,
+		ERenderGraphPassFlags::None,
+		[ScreenPassView, SceneColorViewport, ScreenPassDraw, PostProcessMaterialParameters, DepthStencilTexture, DepthStencilTextureForSRV, SetupFunction, bFlipYAxis]
+		(FRHICommandListImmediate& RHICmdList)
+	{
+		const bool bUsesDepthStencilCopy = DepthStencilTexture != DepthStencilTextureForSRV;
+
+		if (bUsesDepthStencilCopy)
+		{
+			// Swap in the new copied SRV just prior to rendering.
+			FSceneRenderTargets::Get(RHICmdList).CustomDepth = DepthStencilTextureForSRV->GetPooledRenderTarget();
 		}
 
-		DestRenderTargetLoadAction = ERenderTargetLoadAction::ELoad;
-	}
-	else
+		DrawScreenPass<decltype(SetupFunction)>(
+			RHICmdList,
+			ScreenPassView,
+			SceneColorViewport,
+			SceneColorViewport,
+			ScreenPassDraw,
+			SetupFunction);
+
+		if (bUsesDepthStencilCopy)
+		{
+			// Swap back to the up-to-date target.
+			FSceneRenderTargets::Get(RHICmdList).CustomDepth = DepthStencilTexture->GetPooledRenderTarget();
+		}
+	});
+
+	return OutputTexture;
+}
+
+FRenderingCompositePass* AddPostProcessMaterialPass(
+	const FPostprocessContext& PostProcessContext,
+	const UMaterialInterface* MaterialInterface,
+	EPixelFormat OverrideOutputFormat)
+{
+	const FMaterial* Material = nullptr;
+
 	{
-		DestRenderTarget = &PassOutputs[0].RequestSurface(Context);
-		DestRenderTargetLoadAction = Context.GetLoadActionForRenderTarget(*DestRenderTarget);
+		const FMaterialRenderProxy* MaterialRenderProxy = nullptr;
+		const FMaterialShaderMap* MaterialShaderMap = nullptr;
+		GetMaterialInfo(MaterialInterface, PostProcessContext.View.GetFeatureLevel(), OverrideOutputFormat, Material, MaterialRenderProxy, MaterialShaderMap);
+		check(Material);
 	}
 
-	FIntRect SrcRect = Context.SceneColorViewRect;
-	FIntRect DestRect = Context.GetSceneColorDestRect(*DestRenderTarget);
-	FIntPoint DestPos(0, 0);
-	FIntPoint DestSize = DestRect.Size();
-	EDrawRectangleFlags DrawRectangleFlags = EDRF_UseTriangleOptimization;
-
-	checkf(DestRect.Size() == SrcRect.Size(), TEXT("Post process material should not be used as upscaling pass."));
-	
-	FRHIRenderPassInfo RPInfo;
-	if (CustomDepthStencilTarget)
-	{
-		RPInfo = FRHIRenderPassInfo(
-			DestRenderTarget->TargetableTexture,
-			MakeRenderTargetActions(DestRenderTargetLoadAction, ERenderTargetStoreAction::EStore),
-			CustomDepthStencilTarget->TargetableTexture,
-			MakeDepthStencilTargetActions(
-				MakeRenderTargetActions(ERenderTargetLoadAction::ENoAction, ERenderTargetStoreAction::ENoAction),
-				MakeRenderTargetActions(ERenderTargetLoadAction::ELoad, ERenderTargetStoreAction::ENoAction)
-				),
-			FExclusiveDepthStencil::DepthRead_StencilRead
-			);
-	}
-	else
+	if (Material->NeedsGBuffer())
 	{
-		RPInfo = FRHIRenderPassInfo(
-			DestRenderTarget->TargetableTexture,
-			MakeRenderTargetActions(Context.GetLoadActionForRenderTarget(*DestRenderTarget), ERenderTargetStoreAction::EStore)
-			);
+		FSceneRenderTargets::Get(PostProcessContext.RHICmdList).AdjustGBufferRefCount(PostProcessContext.RHICmdList, 1);
 	}
 
-	Context.RHICmdList.BeginRenderPass(RPInfo, TEXT("PostProcessMaterial"));
+	return PostProcessContext.Graph.RegisterPass(new(FMemStack::Get()) TRCPassForRDG<kPostProcessMaterialInputCountMax, 1>(
+		[MaterialInterface, Material, OverrideOutputFormat](FRenderingCompositePass* Pass, FRenderingCompositePassContext& InContext)
 	{
-		Context.SetViewportAndCallRHI(DestRect);
-
-		FGraphicsPipelineStateInitializer GraphicsPSOInit;
-		Context.RHICmdList.ApplyCachedRenderTargets(GraphicsPSOInit);
-		GraphicsPSOInit.BlendState = BlendState;
-		GraphicsPSOInit.RasterizerState = TStaticRasterizerState<>::GetRHI();
-		GraphicsPSOInit.DepthStencilState = DepthStencilState;
-		GraphicsPSOInit.PrimitiveType = PT_TriangleList;
+		FRDGBuilder GraphBuilder(InContext.RHICmdList);
 
-		const FViewInfo& View = Context.View;
-		const FSceneViewFamily& ViewFamily = *(View.Family);
+		FPostProcessMaterialInputs Inputs;
+		Inputs.OverrideOutputFormat = OverrideOutputFormat;
 
-		FIntPoint SrcSize = GetInputDesc(ePId_Input0)->Extent;
+		// Either finds the overridden frame buffer target or returns null.
+		Inputs.OverrideOutputTexture = Pass->FindRDGTextureForOutput(GraphBuilder, ePId_Output0, TEXT("FrameBufferOverride"));
 
-		FShader* VertexShader = nullptr;
-
-		const bool bViewSizeMatchesBufferSize = (View.ViewRect == Context.SceneColorViewRect && View.ViewRect.Size() == SrcSize && View.ViewRect.Min == FIntPoint::ZeroValue);
-		if (FeatureLevel <= ERHIFeatureLevel::ES3_1)
+		for (uint32 InputIndex = 0; InputIndex < kPostProcessMaterialInputCountMax; ++InputIndex)
 		{
-			bool bNeedsVerticalAxisFlip = ShouldMobilePassFlipVerticalAxis(this) && RHINeedsToSwitchVerticalAxis(GShaderPlatformForFeatureLevel[Context.GetFeatureLevel()]);
-			if (bNeedsVerticalAxisFlip)
-			{
-				if (EBlendableLocation(Material->GetBlendableLocation()) == EBlendableLocation::BL_AfterTonemapping)
-				{
-					// flip dest rect y axis.
-					GraphicsPSOInit.RasterizerState = TStaticRasterizerState<FM_Solid, CM_None>::GetRHI();
-					DestPos.Y = DestRect.Max.Y;
-					DestSize.Y = -DestRect.Max.Y;
-					// triangle optimization currently doesn't work when flipped.
-					DrawRectangleFlags = EDRF_Default;
-				}
-			}
+			FRDGTextureRef InputTexture = Pass->CreateRDGTextureForOptionalInput(GraphBuilder, EPassInputId(InputIndex), TEXT("PostProcessInput"));
 
-			// use mobile's post process material.
-			if (bViewSizeMatchesBufferSize)
-			{
-				if(bNeedsVerticalAxisFlip)
-				{
-					VertexShader = SetMobileShaders<FPostProcessMaterialVS_Mobile_AxisSwitch, FPostProcessMaterialPS_Mobile0_AxisSwitch>(MaterialShaderMap, GraphicsPSOInit, Context, Proxy);
-				}
-				else
-				{
-					VertexShader = SetMobileShaders<FPostProcessMaterialVS_Mobile, FPostProcessMaterialPS_Mobile0>(MaterialShaderMap, GraphicsPSOInit, Context, Proxy);
-				}
-			}
-			else
-			{
-				if (bNeedsVerticalAxisFlip)
-				{
-					VertexShader = SetMobileShaders<FPostProcessMaterialVS_Mobile_AxisSwitch, FPostProcessMaterialPS_Mobile1_AxisSwitch>(MaterialShaderMap, GraphicsPSOInit, Context, Proxy);
-				}
-				else
-				{
-					VertexShader = SetMobileShaders<FPostProcessMaterialVS_Mobile, FPostProcessMaterialPS_Mobile1>(MaterialShaderMap, GraphicsPSOInit, Context, Proxy);
-				}
-			}
-			Context.RHICmdList.SetStencilRef(StencilRefValue);
+			/**
+			 * TODO: Propagate each texture viewport through the graph setup instead of guessing. This is wrong for
+			 * any scaled target (e.g. half resolution bloom). We deal with the upsample from TAAU explicitly here,
+			 * but it's a band-aid at best. The problem is that we rely too heavily on the ViewRect--in pixels--which
+			 * only applies to the primary screen resolution viewport.
+			 */
+			const FIntRect InputViewportRect = (InputIndex == 0) ? InContext.SceneColorViewRect : InContext.View.ViewRect;
+
+			Inputs.SetInput((EPostProcessMaterialInput)InputIndex, InputTexture, InputViewportRect);
 		}
-		// Uses highend post process material that assumed ViewSize == BufferSize.
-		else if (bViewSizeMatchesBufferSize)
+
+		Inputs.bFlipYAxis = ShouldMobilePassFlipVerticalAxis(Pass);
+
+		if (TRefCountPtr<IPooledRenderTarget> CustomDepthTarget = FSceneRenderTargets::Get(InContext.RHICmdList).CustomDepth)
 		{
-			FFPostProcessMaterialPS_HighEnd0* PixelShader_HighEnd = MaterialShaderMap->GetShader<FFPostProcessMaterialPS_HighEnd0>();
-			FPostProcessMaterialVS_HighEnd* VertexShader_HighEnd = MaterialShaderMap->GetShader<FPostProcessMaterialVS_HighEnd>();
+			Inputs.CustomDepthTexture = GraphBuilder.RegisterExternalTexture(CustomDepthTarget, TEXT("CustomDepth"));
+		}
 
-			GraphicsPSOInit.BoundShaderState.VertexDeclarationRHI = GPostProcessMaterialVertexDeclaration.VertexDeclarationRHI;
-			GraphicsPSOInit.BoundShaderState.VertexShaderRHI = GETSAFERHISHADER_VERTEX(VertexShader_HighEnd);
-			GraphicsPSOInit.BoundShaderState.PixelShaderRHI = GETSAFERHISHADER_PIXEL(PixelShader_HighEnd);
+		FRDGTextureRef OutputTexture = ComputePostProcessMaterial(GraphBuilder, FScreenPassViewInfo(InContext.View), MaterialInterface, Inputs);
 
-			SetGraphicsPipelineState(Context.RHICmdList, GraphicsPSOInit);
-			Context.RHICmdList.SetStencilRef(StencilRefValue);
+		Pass->ExtractRDGTextureForOutput(GraphBuilder, ePId_Output0, OutputTexture);
 
-			VertexShader_HighEnd->SetParameters(Context.RHICmdList, Context, Proxy);
-			PixelShader_HighEnd->SetParameters(Context.RHICmdList, Context, Proxy);
-			VertexShader = VertexShader_HighEnd;
-		}
-		// Uses highend post process material that handle ViewSize != BufferSize.
-		else
-		{
-			FFPostProcessMaterialPS_HighEnd1* PixelShader_HighEnd = MaterialShaderMap->GetShader<FFPostProcessMaterialPS_HighEnd1>();
-			FPostProcessMaterialVS_HighEnd* VertexShader_HighEnd = MaterialShaderMap->GetShader<FPostProcessMaterialVS_HighEnd>();
+		GraphBuilder.Execute();
 
-			GraphicsPSOInit.BoundShaderState.VertexDeclarationRHI = GPostProcessMaterialVertexDeclaration.VertexDeclarationRHI;
-			GraphicsPSOInit.BoundShaderState.VertexShaderRHI = GETSAFERHISHADER_VERTEX(VertexShader_HighEnd);
-			GraphicsPSOInit.BoundShaderState.PixelShaderRHI = GETSAFERHISHADER_PIXEL(PixelShader_HighEnd);
+		if (Material->NeedsGBuffer())
+		{
+			FSceneRenderTargets::Get(InContext.RHICmdList).AdjustGBufferRefCount(InContext.RHICmdList, -1);
+		}
+	}));
+}
 
-			SetGraphicsPipelineState(Context.RHICmdList, GraphicsPSOInit);
-			Context.RHICmdList.SetStencilRef(StencilRefValue);
+static FPostProcessMaterialNode* IteratePostProcessMaterialNodes(const FFinalPostProcessSettings& Dest, EBlendableLocation Location, FBlendableEntry*& Iterator)
+{
+	for (;;)
+	{
+		FPostProcessMaterialNode* DataPtr = Dest.BlendableManager.IterateBlendables<FPostProcessMaterialNode>(Iterator);
 
-			VertexShader_HighEnd->SetParameters(Context.RHICmdList, Context, Proxy);
-			PixelShader_HighEnd->SetParameters(Context.RHICmdList, Context, Proxy);
-			VertexShader = VertexShader_HighEnd;
+		if (!DataPtr || DataPtr->GetLocation() == Location)
+		{
+			return DataPtr;
 		}
+	}
+}
 
-		DrawPostProcessPass(
-			Context.RHICmdList,
-			DestPos.X, DestPos.Y,
-			DestSize.X, DestSize.Y,
-			SrcRect.Min.X, SrcRect.Min.Y,
-			SrcRect.Width(), SrcRect.Height(),
-			DestRect.Size(),
-			SrcSize,
-			VertexShader,
-			View.StereoPass,
-			Context.HasHmdMesh(),
-			DrawRectangleFlags);
+FRenderingCompositeOutputRef AddPostProcessMaterialReplaceTonemapPass(
+	FPostprocessContext& Context,
+	FRenderingCompositeOutputRef SeparateTranslucency,
+	FRenderingCompositeOutputRef CombinedBloom)
+{
+	if (!Context.View.Family->EngineShowFlags.PostProcessing || !Context.View.Family->EngineShowFlags.PostProcessMaterial)
+	{
+		return Context.FinalOutput;
 	}
-	Context.RHICmdList.EndRenderPass();
-	Context.RHICmdList.CopyToResolveTarget(DestRenderTarget->TargetableTexture, DestRenderTarget->ShaderResourceTexture, FResolveParams());
 
-	if (CustomDepthStencilCopy.IsValid() && CustomDepthStencilCopy->GetRenderTargetItem().IsValid())
+	FBlendableEntry* Iterator = nullptr;
+	FPostProcessMaterialNode PPNode;
+	while (FPostProcessMaterialNode* Data = IteratePostProcessMaterialNodes(Context.View.FinalPostProcessSettings, BL_ReplacingTonemapper, Iterator))
 	{
-		Swap(SceneContext.CustomDepth, CustomDepthStencilCopy);
+		check(Data->GetMaterialInterface());
+
+		if (PPNode.IsValid())
+		{
+			FPostProcessMaterialNode::FCompare Dummy;
+
+			// take the one with the highest priority
+			if (!Dummy.operator()(PPNode, *Data))
+			{
+				continue;
+			}
+		}
+
+		PPNode = *Data;
 	}
 
-	if(Material->NeedsGBuffer())
+	if (UMaterialInterface* MaterialInterface = PPNode.GetMaterialInterface())
 	{
-		FSceneRenderTargets::Get(Context.RHICmdList).AdjustGBufferRefCount(Context.RHICmdList,-1);
+		FRenderingCompositePass* Pass = AddPostProcessMaterialPass(Context, MaterialInterface);
+		Pass->SetInput(EPassInputId(EPostProcessMaterialInput::SceneColor), Context.FinalOutput);
+		Pass->SetInput(EPassInputId(EPostProcessMaterialInput::SeparateTranslucency), SeparateTranslucency);
+		Pass->SetInput(EPassInputId(EPostProcessMaterialInput::CombinedBloom), CombinedBloom);
+
+		return FRenderingCompositeOutputRef(Pass);
 	}
+
+	return Context.FinalOutput;
 }
 
-FPooledRenderTargetDesc FRCPassPostProcessMaterial::ComputeOutputDesc(EPassOutputId InPassOutputId) const
+FRenderingCompositeOutputRef AddPostProcessMaterialChain(
+	FPostprocessContext& Context,
+	EBlendableLocation Location,
+	FRenderingCompositeOutputRef SeparateTranslucency,
+	FRenderingCompositeOutputRef PreTonemapHDRColor,
+	FRenderingCompositeOutputRef PostTonemapHDRColor,
+	FRenderingCompositeOutputRef PreFlattenVelocity)
 {
-	FPooledRenderTargetDesc Ret = GetInput(ePId_Input0)->GetOutput()->RenderTargetDesc;
+	if (!Context.View.Family->EngineShowFlags.PostProcessing ||
+		!Context.View.Family->EngineShowFlags.PostProcessMaterial ||
+		Context.View.Family->EngineShowFlags.VisualizeShadingModels ||
+		CVarPostProcessingDisableMaterials.GetValueOnRenderThread() != 0)
+	{
+		return Context.FinalOutput;
+	}
+
+	// hard coded - this should be a reasonable limit
+	const uint32 MAX_PPMATERIALNODES = 10;
+	FBlendableEntry* Iterator = 0;
+	FPostProcessMaterialNode PPNodes[MAX_PPMATERIALNODES];
+	uint32 PPNodeCount = 0;
+	bool bVisualizingBuffer = false;
+
+	if (Context.View.Family->EngineShowFlags.VisualizeBuffer)
+	{
+		// Apply requested material to the full screen
+		UMaterial* Material = GetBufferVisualizationData().GetMaterial(Context.View.CurrentBufferVisualizationMode);
 
-	if (OutputFormat != PF_Unknown)
+		if (Material && Material->BlendableLocation == Location)
+		{
+			PPNodes[0] = FPostProcessMaterialNode(Material, Location, Material->BlendablePriority);
+			++PPNodeCount;
+			bVisualizingBuffer = true;
+		}
+	}
+	for (; PPNodeCount < MAX_PPMATERIALNODES; ++PPNodeCount)
 	{
-		Ret.Format = OutputFormat;
+		FPostProcessMaterialNode* Data = IteratePostProcessMaterialNodes(Context.View.FinalPostProcessSettings, Location, Iterator);
+
+		if (!Data)
+		{
+			break;
+		}
+
+		check(Data->GetMaterialInterface());
+
+		PPNodes[PPNodeCount] = *Data;
 	}
-	Ret.Reset();
-	Ret.AutoWritable = false;
-	Ret.DebugName = TEXT("PostProcessMaterial");
-	Ret.ClearValue = FClearValueBinding(FLinearColor::Black);
-	Ret.Flags |= GFastVRamConfig.PostProcessMaterial;
 
-	return Ret;
+	::Sort(PPNodes, PPNodeCount, FPostProcessMaterialNode::FCompare());
+
+	ERHIFeatureLevel::Type FeatureLevel = Context.View.GetFeatureLevel();
+
+	FRenderingCompositeOutputRef LastOutput = Context.FinalOutput;
+
+	for (uint32 i = 0; i < PPNodeCount; ++i)
+	{
+		const UMaterialInterface* MaterialInterface = PPNodes[i].GetMaterialInterface();
+
+		FRenderingCompositePass* Pass = AddPostProcessMaterialPass(Context, MaterialInterface);
+
+		Pass->SetInput(EPassInputId(EPostProcessMaterialInput::SceneColor), LastOutput);
+		Pass->SetInput(EPassInputId(EPostProcessMaterialInput::SeparateTranslucency), SeparateTranslucency);
+
+		// This input is only needed for visualization and frame dumping
+		if (bVisualizingBuffer)
+		{
+			Pass->SetInput(EPassInputId(EPostProcessMaterialInput::PreTonemapHDRColor), PreTonemapHDRColor);
+			Pass->SetInput(EPassInputId(EPostProcessMaterialInput::PostTonemapHDRColor), PostTonemapHDRColor);
+		}
+
+		if (!FVelocityRendering::BasePassCanOutputVelocity(FeatureLevel))
+		{
+			Pass->SetInput(EPassInputId(EPostProcessMaterialInput::Velocity), PreFlattenVelocity);
+		}
+
+		LastOutput = FRenderingCompositeOutputRef(Pass);
+	}
+
+	return LastOutput;
 }
+
+void AddHighResScreenshotMask(FPostprocessContext& Context)
+{
+	if (Context.View.Family->EngineShowFlags.HighResScreenshotMask != 0)
+	{
+		check(Context.View.FinalPostProcessSettings.HighResScreenshotMaterial);
+
+		FRenderingCompositeOutputRef Input = Context.FinalOutput;
+
+		FRenderingCompositePass* CompositePass = AddPostProcessMaterialPass(Context, Context.View.FinalPostProcessSettings.HighResScreenshotMaterial);
+		CompositePass->SetInput(ePId_Input0, FRenderingCompositeOutputRef(Input));
+		Context.FinalOutput = FRenderingCompositeOutputRef(CompositePass);
+
+		if (GIsHighResScreenshot)
+		{
+			check(Context.View.FinalPostProcessSettings.HighResScreenshotMaskMaterial);
+
+			FRenderingCompositePass* MaskPass = AddPostProcessMaterialPass(Context, Context.View.FinalPostProcessSettings.HighResScreenshotMaskMaterial);
+			MaskPass->SetInput(ePId_Input0, FRenderingCompositeOutputRef(Input));
+			CompositePass->AddDependency(MaskPass);
+
+			FString BaseFilename = FString(Context.View.FinalPostProcessSettings.BufferVisualizationDumpBaseFilename);
+			MaskPass->SetOutputColorArray(ePId_Output0, FScreenshotRequest::GetHighresScreenshotMaskColorArray());
+		}
+	}
+
+	// Draw the capture region if a material was supplied
+	if (Context.View.FinalPostProcessSettings.HighResScreenshotCaptureRegionMaterial)
+	{
+		auto Material = Context.View.FinalPostProcessSettings.HighResScreenshotCaptureRegionMaterial;
+
+		FRenderingCompositePass* CaptureRegionVisualizationPass = AddPostProcessMaterialPass(Context, Material);
+		CaptureRegionVisualizationPass->SetInput(ePId_Input0, FRenderingCompositeOutputRef(Context.FinalOutput));
+		Context.FinalOutput = FRenderingCompositeOutputRef(CaptureRegionVisualizationPass);
+	}
+}
\ No newline at end of file
diff --git a/Engine/Source/Runtime/Renderer/Private/PostProcess/PostProcessMaterial.h b/Engine/Source/Runtime/Renderer/Private/PostProcess/PostProcessMaterial.h
--- a/Engine/Source/Runtime/Renderer/Private/PostProcess/PostProcessMaterial.h
+++ b/Engine/Source/Runtime/Renderer/Private/PostProcess/PostProcessMaterial.h
@@ -9,32 +9,136 @@
 #include "CoreMinimal.h"
 #include "RendererInterface.h"
 #include "PostProcess/RenderingCompositionGraph.h"
+#include "ScreenPass.h"
 
 class UMaterialInterface;
 
-// derives from TRenderingCompositePassBase<InputCount, OutputCount>
-// ePId_Input0: former pass
-// ePId_Input1: optional, depends on EBlendableLocation
-// ePId_Input2: optional, depends on EBlendableLocation
-// ePId_Input3: optional, depends on EBlendableLocation
-// ePId_Input4: optional, always velocity (if required by the material)
-class FRCPassPostProcessMaterial : public TRenderingCompositePassBase<5,1>
+const uint32 kPostProcessMaterialInputCountMax = 5;
+
+/** Named post process material slots. Inputs are aliased and have different semantics
+ *  based on the post process material blend point, which is documented with the input.
+ */
+enum class EPostProcessMaterialInput : uint32
 {
-public:
-	// constructor
-	FRCPassPostProcessMaterial(UMaterialInterface* InMaterialInterface, ERHIFeatureLevel::Type InFeatureLevel, EPixelFormat OutputFormatIN = PF_Unknown);
+	// Always Active. Color from the previous stage of the post process chain.
+	SceneColor = 0,
+
+	// Always Active.
+	SeparateTranslucency = 1,
+
+	// Replace Tonemap Only. Half resolution combined bloom input.
+	CombinedBloom = 2,
+
+	// Buffer Visualization Only.
+	PreTonemapHDRColor = 2,
+	PostTonemapHDRColor = 3,
+
+	// Active if separate velocity pass is used--i.e. not part of base pass; Not active during Replace Tonemap.
+	Velocity = 4
+};
+
+struct FPostProcessMaterialInputs
+{
+	inline void SetInput(EPostProcessMaterialInput Input, FRDGTextureRef InTexture, FIntRect InViewport)
+	{
+		Textures[(uint32)Input] = InTexture;
+		Viewports[(uint32)Input] = InViewport;
+	}
+
+	inline void GetInput(EPostProcessMaterialInput Input, FRDGTextureRef& OutTexture, FIntRect& OutViewport) const
+	{
+		OutTexture = Textures[(uint32)Input];
+		OutViewport = Viewports[(uint32)Input];
+	}
+
+	inline void Validate() const
+	{
+		ValidateInputExists(EPostProcessMaterialInput::SceneColor);
+		ValidateInputExists(EPostProcessMaterialInput::SeparateTranslucency);
+
+		// Either override output format is valid or the override output texture is; not both.
+		if (OverrideOutputFormat != PF_Unknown)
+		{
+			check(OverrideOutputTexture == nullptr);
+		}
+		if (OverrideOutputTexture)
+		{
+			check(OverrideOutputFormat == PF_Unknown);
+		}
+	}
 
-	// interface FRenderingCompositePass ---------
+	inline void ValidateInputExists(EPostProcessMaterialInput Input) const
+	{
+		FRDGTextureRef Texture = nullptr;
+		FIntRect Viewport;
+		GetInput(EPostProcessMaterialInput::SceneColor, Texture, Viewport);
+		check(Texture);
+		check(!Viewport.IsEmpty());
+	}
 
-	virtual void Process(FRenderingCompositePassContext& Context) override;
-	virtual void Release() override { delete this; }
-	virtual FPooledRenderTargetDesc ComputeOutputDesc(EPassOutputId InPassOutputId) const override;
+	/** Array of input textures bound to the material. The first element represents the output from
+	 *  the previous post process and is required. All other inputs are optional.
+	 */
+	TStaticArray<FRDGTextureRef, kPostProcessMaterialInputCountMax> Textures;
 
-private:
-	void SetShader(const FRenderingCompositePassContext& Context);
+	/** Array of post process input viewports corresponding to @ref Textures */
+	TStaticArray<FIntRect, kPostProcessMaterialInputCountMax> Viewports;
 
-	UMaterialInterface* MaterialInterface;
+	/** The override output texture to use. If this is null, a new texture is created. */
+	FRDGTextureRef OverrideOutputTexture = nullptr;
 
-	// PF_Unknown for default behavior.
-	EPixelFormat OutputFormat;
+	/** The override output texture format to use if a new texture is created. */
+	EPixelFormat OverrideOutputFormat = PF_Unknown;
+
+	/** Custom stencil texture used for stencil operations. */
+	FRDGTextureRef CustomDepthTexture = nullptr;
+
+	/** Performs a vertical axis flip if the RHI allows it. */
+	bool bFlipYAxis = false;
 };
+
+FRDGTextureRef ComputePostProcessMaterial(
+	FRDGBuilder& GraphBuilder,
+	const FScreenPassViewInfo& ScreenPassView,
+	const UMaterialInterface* MaterialInterface,
+	const FPostProcessMaterialInputs& Inputs);
+
+BEGIN_SHADER_PARAMETER_STRUCT(FPostProcessMaterialInput, )
+	SHADER_PARAMETER_STRUCT_INCLUDE(FScreenPassTextureViewportParameters, Viewport)
+	SHADER_PARAMETER_RDG_TEXTURE(Texture2D, Texture)
+	SHADER_PARAMETER_SAMPLER(SamplerState, Sampler)
+END_SHADER_PARAMETER_STRUCT()
+
+FPostProcessMaterialInput GetPostProcessMaterialInput(FIntRect ViewportRect, FRDGTextureRef Texture, FRHISamplerState* Sampler);
+
+BEGIN_SHADER_PARAMETER_STRUCT(FPostProcessMaterialParameters, )
+	SHADER_PARAMETER_STRUCT(FScreenPassTextureViewportParameters, PostProcessOutput)
+	SHADER_PARAMETER_STRUCT_ARRAY(FPostProcessMaterialInput, PostProcessInput, [kPostProcessMaterialInputCountMax])
+	SHADER_PARAMETER_SAMPLER(SamplerState, PostProcessInput_BilinearSampler)
+	SHADER_PARAMETER_RDG_TEXTURE(Texture2D, CustomDepth)
+	SHADER_PARAMETER(uint32, bFlipYAxis)
+	RENDER_TARGET_BINDING_SLOTS()
+END_SHADER_PARAMETER_STRUCT()
+
+//////////////////////////////////////////////////////////////////////////
+//! Legacy Composition Graph Methods
+
+FRenderingCompositePass* AddPostProcessMaterialPass(
+	const FPostprocessContext& Context,
+	const UMaterialInterface* MaterialInterface,
+	EPixelFormat OutputFormat = PF_Unknown);
+
+FRenderingCompositeOutputRef AddPostProcessMaterialReplaceTonemapPass(
+	FPostprocessContext& Context,
+	FRenderingCompositeOutputRef SeparateTranslucency,
+	FRenderingCompositeOutputRef CombinedBloom);
+
+FRenderingCompositeOutputRef AddPostProcessMaterialChain(
+	FPostprocessContext& Context,
+	EBlendableLocation InLocation,
+	FRenderingCompositeOutputRef SeparateTranslucency = FRenderingCompositeOutputRef(),
+	FRenderingCompositeOutputRef PreTonemapHDRColor = FRenderingCompositeOutputRef(),
+	FRenderingCompositeOutputRef PostTonemapHDRColor = FRenderingCompositeOutputRef(),
+	FRenderingCompositeOutputRef PreFlattenVelocity = FRenderingCompositeOutputRef());
+
+void AddHighResScreenshotMask(FPostprocessContext& Context);
\ No newline at end of file
diff --git a/Engine/Source/Runtime/Renderer/Private/PostProcess/PostProcessSubsurface.cpp b/Engine/Source/Runtime/Renderer/Private/PostProcess/PostProcessSubsurface.cpp
--- a/Engine/Source/Runtime/Renderer/Private/PostProcess/PostProcessSubsurface.cpp
+++ b/Engine/Source/Runtime/Renderer/Private/PostProcess/PostProcessSubsurface.cpp
@@ -14,68 +14,68 @@ ENGINE_API IPooledRenderTarget* GetSubsufaceProfileTexture_RT(FRHICommandListImm
 
 namespace
 {
-	TAutoConsoleVariable<int32> CVarSubsurfaceScattering(
-		TEXT("r.SubsurfaceScattering"),
-		1,
-		TEXT(" 0: disabled\n")
-		TEXT(" 1: enabled (default)"),
-		ECVF_RenderThreadSafe | ECVF_Scalability);
-
-	TAutoConsoleVariable<float> CVarSSSScale(
-		TEXT("r.SSS.Scale"),
-		1.0f,
-		TEXT("Affects the Screen space subsurface scattering pass")
-		TEXT("(use shadingmodel SubsurfaceProfile, get near to the object as the default)\n")
-		TEXT("is human skin which only scatters about 1.2cm)\n")
-		TEXT(" 0: off (if there is no object on the screen using this pass it should automatically disable the post process pass)\n")
-		TEXT("<1: scale scatter radius down (for testing)\n")
-		TEXT(" 1: use given radius form the Subsurface scattering asset (default)\n")
-		TEXT(">1: scale scatter radius up (for testing)"),
-		ECVF_Scalability | ECVF_RenderThreadSafe);
-
-	TAutoConsoleVariable<int32> CVarSSSHalfRes(
-		TEXT("r.SSS.HalfRes"),
-		1,
-		TEXT(" 0: full quality (not optimized, as reference)\n")
-		TEXT(" 1: parts of the algorithm runs in half resolution which is lower quality but faster (default)"),
-		ECVF_RenderThreadSafe | ECVF_Scalability);
-
-	TAutoConsoleVariable<int32> CVarSSSQuality(
-		TEXT("r.SSS.Quality"),
-		0,
-		TEXT("Defines the quality of the recombine pass when using the SubsurfaceScatteringProfile shading model\n")
-		TEXT(" 0: low (faster, default)\n")
-		TEXT(" 1: high (sharper details but slower)\n")
-		TEXT("-1: auto, 1 if TemporalAA is disabled (without TemporalAA the quality is more noticable)"),
-		ECVF_RenderThreadSafe | ECVF_Scalability);
-
-	TAutoConsoleVariable<int32> CVarSSSFilter(
-		TEXT("r.SSS.Filter"),
-		1,
-		TEXT("Defines the filter method for Screenspace Subsurface Scattering feature.\n")
-		TEXT(" 0: point filter (useful for testing, could be cleaner)\n")
-		TEXT(" 1: bilinear filter"),
-		ECVF_RenderThreadSafe | ECVF_Scalability);
-
-	TAutoConsoleVariable<int32> CVarSSSSampleSet(
-		TEXT("r.SSS.SampleSet"),
-		2,
-		TEXT("Defines how many samples we use for Screenspace Subsurface Scattering feature.\n")
-		TEXT(" 0: lowest quality (6*2+1)\n")
-		TEXT(" 1: medium quality (9*2+1)\n")
-		TEXT(" 2: high quality (13*2+1) (default)"),
-		ECVF_RenderThreadSafe | ECVF_Scalability);
-
-	TAutoConsoleVariable<int32> CVarSSSCheckerboard(
-		TEXT("r.SSS.Checkerboard"),
-		2,
-		TEXT("Enables or disables checkerboard rendering for subsurface profile rendering.\n")
-		TEXT("This is necessary if SceneColor does not include a floating point alpha channel (e.g 32-bit formats)\n")
-		TEXT(" 0: Disabled (high quality) \n")
-		TEXT(" 1: Enabled (low quality). Surface lighting will be at reduced resolution.\n")
-		TEXT(" 2: Automatic. Non-checkerboard lighting will be applied if we have a suitable rendertarget format\n"),
-		ECVF_RenderThreadSafe);
-}
+TAutoConsoleVariable<int32> CVarSubsurfaceScattering(
+	TEXT("r.SubsurfaceScattering"),
+	1,
+	TEXT(" 0: disabled\n")
+	TEXT(" 1: enabled (default)"),
+	ECVF_RenderThreadSafe | ECVF_Scalability);
+
+TAutoConsoleVariable<float> CVarSSSScale(
+	TEXT("r.SSS.Scale"),
+	1.0f,
+	TEXT("Affects the Screen space subsurface scattering pass")
+	TEXT("(use shadingmodel SubsurfaceProfile, get near to the object as the default)\n")
+	TEXT("is human skin which only scatters about 1.2cm)\n")
+	TEXT(" 0: off (if there is no object on the screen using this pass it should automatically disable the post process pass)\n")
+	TEXT("<1: scale scatter radius down (for testing)\n")
+	TEXT(" 1: use given radius form the Subsurface scattering asset (default)\n")
+	TEXT(">1: scale scatter radius up (for testing)"),
+	ECVF_Scalability | ECVF_RenderThreadSafe);
+
+TAutoConsoleVariable<int32> CVarSSSHalfRes(
+	TEXT("r.SSS.HalfRes"),
+	1,
+	TEXT(" 0: full quality (not optimized, as reference)\n")
+	TEXT(" 1: parts of the algorithm runs in half resolution which is lower quality but faster (default)"),
+	ECVF_RenderThreadSafe | ECVF_Scalability);
+
+TAutoConsoleVariable<int32> CVarSSSQuality(
+	TEXT("r.SSS.Quality"),
+	0,
+	TEXT("Defines the quality of the recombine pass when using the SubsurfaceScatteringProfile shading model\n")
+	TEXT(" 0: low (faster, default)\n")
+	TEXT(" 1: high (sharper details but slower)\n")
+	TEXT("-1: auto, 1 if TemporalAA is disabled (without TemporalAA the quality is more noticable)"),
+	ECVF_RenderThreadSafe | ECVF_Scalability);
+
+TAutoConsoleVariable<int32> CVarSSSFilter(
+	TEXT("r.SSS.Filter"),
+	1,
+	TEXT("Defines the filter method for Screenspace Subsurface Scattering feature.\n")
+	TEXT(" 0: point filter (useful for testing, could be cleaner)\n")
+	TEXT(" 1: bilinear filter"),
+	ECVF_RenderThreadSafe | ECVF_Scalability);
+
+TAutoConsoleVariable<int32> CVarSSSSampleSet(
+	TEXT("r.SSS.SampleSet"),
+	2,
+	TEXT("Defines how many samples we use for Screenspace Subsurface Scattering feature.\n")
+	TEXT(" 0: lowest quality (6*2+1)\n")
+	TEXT(" 1: medium quality (9*2+1)\n")
+	TEXT(" 2: high quality (13*2+1) (default)"),
+	ECVF_RenderThreadSafe | ECVF_Scalability);
+
+TAutoConsoleVariable<int32> CVarSSSCheckerboard(
+	TEXT("r.SSS.Checkerboard"),
+	2,
+	TEXT("Enables or disables checkerboard rendering for subsurface profile rendering.\n")
+	TEXT("This is necessary if SceneColor does not include a floating point alpha channel (e.g 32-bit formats)\n")
+	TEXT(" 0: Disabled (high quality) \n")
+	TEXT(" 1: Enabled (low quality). Surface lighting will be at reduced resolution.\n")
+	TEXT(" 2: Automatic. Non-checkerboard lighting will be applied if we have a suitable rendertarget format\n"),
+	ECVF_RenderThreadSafe);
+} //! namespace
 
 enum class ESubsurfaceMode : uint32
 {
@@ -186,7 +186,7 @@ bool IsSubsurfaceCheckerboardFormat(EPixelFormat SceneColorFormat)
 }
 
 // Returns the SS profile texture with a black fallback texture if none exists yet.
-FTextureRHIRef GetSubsurfaceProfileTexture(FRHICommandListImmediate& RHICmdList)
+FRHITexture* GetSubsurfaceProfileTexture(FRHICommandListImmediate& RHICmdList)
 {
 	const IPooledRenderTarget* ProfileTextureTarget = GetSubsufaceProfileTexture_RT(RHICmdList);
 
@@ -646,8 +646,9 @@ FRDGTextureRef ComputeSubsurface(
 				{
 					const FViewInfo& View = Views[ViewIndex];
 					const FScreenPassViewInfo ScreenPassView(View);
+					const FScreenPassTextureViewport TextureViewport(View.ViewRect, InputTextureSize);
 
-					DrawScreenPass(RHICmdList, ScreenPassView, View.ViewRect, View.ViewRect, InputTextureSize, *PixelShader, *PassParameters);
+					DrawScreenPass(RHICmdList, ScreenPassView, TextureViewport, TextureViewport, *PixelShader, *PassParameters);
 				}
 			}
 		});
diff --git a/Engine/Source/Runtime/Renderer/Private/PostProcess/PostProcessing.cpp b/Engine/Source/Runtime/Renderer/Private/PostProcess/PostProcessing.cpp
--- a/Engine/Source/Runtime/Renderer/Private/PostProcess/PostProcessing.cpp
+++ b/Engine/Source/Runtime/Renderer/Private/PostProcess/PostProcessing.cpp
@@ -171,12 +171,6 @@ TAutoConsoleVariable<int32> CVarHalfResFFTBloom(
 	TEXT(" 1: Half-resolution convoltuion that excludes the center of the kernel.\n"),
 	ECVF_Scalability | ECVF_RenderThreadSafe);
 
-TAutoConsoleVariable<int32> CVarPostProcessingDisableMaterials(
-	TEXT("r.PostProcessing.DisableMaterials"),
-	0,
-	TEXT(" Allows to disable post process materials. \n"),
-	ECVF_Scalability | ECVF_RenderThreadSafe);
-
 TAutoConsoleVariable<int32> CVarTemporalAAAllowDownsampling(
 	TEXT("r.TemporalAA.AllowDownsampling"),
 	1,
@@ -190,8 +184,6 @@ static TAutoConsoleVariable<float> CVarTemporalAAHistorySP(
 	ECVF_RenderThreadSafe
 );
 
-static bool HasPostProcessMaterial(FPostprocessContext& Context, EBlendableLocation InLocation);
-
 // -------------------------------------------------------
 
 bool ShouldDoComputePostProcessing(const FViewInfo& View)
@@ -799,246 +791,12 @@ static void AddTemporalAA( FPostprocessContext& Context, FRenderingCompositeOutp
 	}
 }
 
-FPostProcessMaterialNode* IteratePostProcessMaterialNodes(const FFinalPostProcessSettings& Dest, EBlendableLocation InLocation, FBlendableEntry*& Iterator)
-{
-	for(;;)
-	{
-		FPostProcessMaterialNode* DataPtr = Dest.BlendableManager.IterateBlendables<FPostProcessMaterialNode>(Iterator);
-
-		if(!DataPtr || DataPtr->GetLocation() == InLocation)
-		{
-			return DataPtr;
-		}
-	}
-}
-
-
-static FRenderingCompositePass* AddSinglePostProcessMaterial(FPostprocessContext& Context, EBlendableLocation InLocation)
-{
-	if(!Context.View.Family->EngineShowFlags.PostProcessing || !Context.View.Family->EngineShowFlags.PostProcessMaterial)
-	{
-		return 0;
-	}
-
-	FBlendableEntry* Iterator = 0;
-	FPostProcessMaterialNode PPNode;
-
-	while(FPostProcessMaterialNode* Data = IteratePostProcessMaterialNodes(Context.View.FinalPostProcessSettings, InLocation, Iterator))
-	{
-		check(Data->GetMaterialInterface());
-
-		if(PPNode.IsValid())
-		{
-			FPostProcessMaterialNode::FCompare Dummy;
-
-			// take the one with the highest priority
-			if(!Dummy.operator()(PPNode, *Data))
-			{
-				continue;
-			}
-		}
-
-		PPNode = *Data;
-	}
-
-	if(UMaterialInterface* MaterialInterface = PPNode.GetMaterialInterface())
-	{
-		FMaterialRenderProxy* Proxy = MaterialInterface->GetRenderProxy();
-
-		check(Proxy);
-
-		const FMaterial* Material = Proxy->GetMaterial(Context.View.GetFeatureLevel());
-
-		check(Material);
-
-		if(Material->NeedsGBuffer())
-		{
-			// AdjustGBufferRefCount(-1) call is done when the pass gets executed
-			FSceneRenderTargets::Get(Context.RHICmdList).AdjustGBufferRefCount(Context.RHICmdList, 1);
-		}
-
-		FRenderingCompositePass* Node = Context.Graph.RegisterPass(new(FMemStack::Get()) FRCPassPostProcessMaterial(MaterialInterface, Context.View.GetFeatureLevel()));
-
-		return Node;
-	}
-
-	return 0;
-}
-
-// simplied version of AddPostProcessMaterialChain(), side effect free
-static bool HasPostProcessMaterial(FPostprocessContext& Context, EBlendableLocation InLocation)
-{
-	if(!Context.View.Family->EngineShowFlags.PostProcessing || !Context.View.Family->EngineShowFlags.PostProcessMaterial)
-	{
-		return false;
-	}
-
-	if(Context.View.Family->EngineShowFlags.VisualizeBuffer)
-	{	
-		// Apply requested material to the full screen
-		UMaterial* Material = GetBufferVisualizationData().GetMaterial(Context.View.CurrentBufferVisualizationMode);
-		
-		if(Material && Material->BlendableLocation == InLocation)
-		{
-			return true;
-		}
-	}
-
-	FBlendableEntry* Iterator = 0;
-	FPostProcessMaterialNode* Data = IteratePostProcessMaterialNodes(Context.View.FinalPostProcessSettings, InLocation, Iterator);
-
-	if(Data)
-	{
-		return true;
-	}
-	
-	return false;
-}
-
-static FRenderingCompositeOutputRef AddPostProcessMaterialChain(
+static void AddGBufferVisualizationOverview(
 	FPostprocessContext& Context,
-	EBlendableLocation InLocation,
-	FRenderingCompositeOutputRef SeparateTranslucency = FRenderingCompositeOutputRef(),
-	FRenderingCompositeOutputRef PreTonemapHDRColor = FRenderingCompositeOutputRef(),
-	FRenderingCompositeOutputRef PostTonemapHDRColor = FRenderingCompositeOutputRef(),
-	FRenderingCompositeOutputRef PreFlattenVelocity = FRenderingCompositeOutputRef())
-{
-	if( !Context.View.Family->EngineShowFlags.PostProcessing ||
-		!Context.View.Family->EngineShowFlags.PostProcessMaterial ||
-		Context.View.Family->EngineShowFlags.VisualizeShadingModels || 
-		CVarPostProcessingDisableMaterials.GetValueOnRenderThread() != 0)		// we should add more
-	{
-		return Context.FinalOutput;
-	}
-
-	// hard coded - this should be a reasonable limit
-	const uint32 MAX_PPMATERIALNODES = 10;
-	FBlendableEntry* Iterator = 0;
-	FPostProcessMaterialNode PPNodes[MAX_PPMATERIALNODES];
-	uint32 PPNodeCount = 0;
-	bool bVisualizingBuffer = false;
-
-	if(Context.View.Family->EngineShowFlags.VisualizeBuffer)
-	{	
-		// Apply requested material to the full screen
-		UMaterial* Material = GetBufferVisualizationData().GetMaterial(Context.View.CurrentBufferVisualizationMode);
-		
-		if(Material && Material->BlendableLocation == InLocation)
-		{
-			PPNodes[0] = FPostProcessMaterialNode(Material, InLocation, Material->BlendablePriority);
-			++PPNodeCount;
-			bVisualizingBuffer = true;
-		}
-	}
-	for(;PPNodeCount < MAX_PPMATERIALNODES; ++PPNodeCount)
-	{
-		FPostProcessMaterialNode* Data = IteratePostProcessMaterialNodes(Context.View.FinalPostProcessSettings, InLocation, Iterator);
-
-		if(!Data)
-		{
-			break;
-		}
-
-		check(Data->GetMaterialInterface());
-
-		PPNodes[PPNodeCount] = *Data;
-	}
-	
-	::Sort(PPNodes, PPNodeCount, FPostProcessMaterialNode::FCompare());
-
-	ERHIFeatureLevel::Type FeatureLevel = Context.View.GetFeatureLevel();
-
-	FRenderingCompositeOutputRef LastestOutput = Context.FinalOutput;
-
-	for(uint32 i = 0; i < PPNodeCount; ++i)
-	{
-		UMaterialInterface* MaterialInterface = PPNodes[i].GetMaterialInterface();
-
-		FMaterialRenderProxy* Proxy = MaterialInterface->GetRenderProxy();
-
-		check(Proxy);
-
-		const FMaterial* Material = Proxy->GetMaterial(Context.View.GetFeatureLevel());
-
-		check(Material);
-
-		if(Material->NeedsGBuffer())
-		{
-			// AdjustGBufferRefCount(-1) call is done when the pass gets executed
-			FSceneRenderTargets::Get(Context.RHICmdList).AdjustGBufferRefCount(Context.RHICmdList, 1);
-		}
-
-		FRenderingCompositePass* Node = Context.Graph.RegisterPass(new(FMemStack::Get()) FRCPassPostProcessMaterial(MaterialInterface,FeatureLevel));
-		Node->SetInput(ePId_Input0, FRenderingCompositeOutputRef(LastestOutput));
-
-		// We are binding separate translucency here because the post process SceneTexture node can reference 
-		// the separate translucency buffers through ePId_Input1. 
-		// TODO: Check if material actually uses this texture and only bind if needed.
-		Node->SetInput(ePId_Input1, SeparateTranslucency);
-
-		// This input is only needed for visualization and frame dumping
-		if (bVisualizingBuffer)
-		{
-			Node->SetInput(ePId_Input2, PreTonemapHDRColor);
-			Node->SetInput(ePId_Input3, PostTonemapHDRColor);
-		}
-
-		if (Material->GetRenderingThreadShaderMap()->UsesVelocitySceneTexture() && !FVelocityRendering::BasePassCanOutputVelocity(FeatureLevel))
-		{
-			Node->SetInput(ePId_Input4, PreFlattenVelocity);
-		}
-
-		LastestOutput = FRenderingCompositeOutputRef(Node);
-	}
-
-	return LastestOutput;
-}
-
-static void AddHighResScreenshotMask(FPostprocessContext& Context, FRenderingCompositeOutputRef& SeparateTranslucencyInput)
-{
-	if (Context.View.Family->EngineShowFlags.HighResScreenshotMask != 0)
-	{
-		check(Context.View.FinalPostProcessSettings.HighResScreenshotMaterial);
-
-		FRenderingCompositeOutputRef Input = Context.FinalOutput;
-
-		FRenderingCompositePass* CompositePass = Context.Graph.RegisterPass(new(FMemStack::Get()) FRCPassPostProcessMaterial(Context.View.FinalPostProcessSettings.HighResScreenshotMaterial, Context.View.GetFeatureLevel()));
-		CompositePass->SetInput(ePId_Input0, FRenderingCompositeOutputRef(Input));
-		Context.FinalOutput = FRenderingCompositeOutputRef(CompositePass);
-
-		if (GIsHighResScreenshot)
-		{
-			check(Context.View.FinalPostProcessSettings.HighResScreenshotMaskMaterial); 
-
-			FRenderingCompositePass* MaskPass = Context.Graph.RegisterPass(new(FMemStack::Get()) FRCPassPostProcessMaterial(Context.View.FinalPostProcessSettings.HighResScreenshotMaskMaterial, Context.View.GetFeatureLevel()));
-			MaskPass->SetInput(ePId_Input0, FRenderingCompositeOutputRef(Input));
-			CompositePass->AddDependency(MaskPass);
-
-			FString BaseFilename = FString(Context.View.FinalPostProcessSettings.BufferVisualizationDumpBaseFilename);
-			MaskPass->SetOutputColorArray(ePId_Output0, FScreenshotRequest::GetHighresScreenshotMaskColorArray());
-		}
-	}
-
-	// Draw the capture region if a material was supplied
-	if (Context.View.FinalPostProcessSettings.HighResScreenshotCaptureRegionMaterial)
-	{
-		auto Material = Context.View.FinalPostProcessSettings.HighResScreenshotCaptureRegionMaterial;
-
-		FRenderingCompositePass* CaptureRegionVisualizationPass = Context.Graph.RegisterPass(new(FMemStack::Get()) FRCPassPostProcessMaterial(Material, Context.View.GetFeatureLevel()));
-		CaptureRegionVisualizationPass->SetInput(ePId_Input0, FRenderingCompositeOutputRef(Context.FinalOutput));
-		Context.FinalOutput = FRenderingCompositeOutputRef(CaptureRegionVisualizationPass);
-
-		auto Proxy = Material->GetRenderProxy();
-		const FMaterial* RendererMaterial = Proxy->GetMaterial(Context.View.GetFeatureLevel());
-		if (RendererMaterial->NeedsGBuffer())
-		{
-			// AdjustGBufferRefCount(-1) call is done when the pass gets executed
-			FSceneRenderTargets::Get(Context.RHICmdList).AdjustGBufferRefCount(Context.RHICmdList, 1);
-		}
-	}
-}
-
-static void AddGBufferVisualizationOverview(FPostprocessContext& Context, FRenderingCompositeOutputRef& SeparateTranslucencyInput, FRenderingCompositeOutputRef& PreTonemapHDRColorInput, FRenderingCompositeOutputRef& PostTonemapHDRColorInput, FRenderingCompositeOutputRef& PreFlattenVelocity)
+	FRenderingCompositeOutputRef SeparateTranslucencyInput,
+	FRenderingCompositeOutputRef PreTonemapHDRColorInput,
+	FRenderingCompositeOutputRef PostTonemapHDRColorInput,
+	FRenderingCompositeOutputRef PreFlattenVelocity)
 {
 	static const auto CVarDumpFrames = IConsoleManager::Get().FindTConsoleVariableDataInt(TEXT("r.BufferVisualizationDumpFrames"));
 	static const auto CVarDumpFramesAsHDR = IConsoleManager::Get().FindTConsoleVariableDataInt(TEXT("r.BufferVisualizationDumpFramesAsHDR"));
@@ -1062,7 +820,7 @@ static void AddGBufferVisualizationOverview(FPostprocessContext& Context, FRende
 	}
 	
 	if (bDumpFrames || bVisualizationEnabled)
-	{	
+	{
 		FRenderingCompositeOutputRef IncomingStage = Context.FinalOutput;
 
 		if (bDumpFrames || bOverviewModeEnabled)
@@ -1075,28 +833,20 @@ static void AddGBufferVisualizationOverview(FPostprocessContext& Context, FRende
 			// Loop over materials, creating stages for generation and downsampling of the tiles.			
 			for (TArray<UMaterialInterface*>::TConstIterator It = Context.View.FinalPostProcessSettings.BufferVisualizationOverviewMaterials.CreateConstIterator(); It; ++It)
 			{
-				auto MaterialInterface = *It;
+				const UMaterialInterface* MaterialInterface = *It;
 				if (MaterialInterface)
 				{
 					// Apply requested material
-					FRenderingCompositePass* MaterialPass = Context.Graph.RegisterPass(new(FMemStack::Get()) FRCPassPostProcessMaterial(*It, Context.View.GetFeatureLevel(), OutputFormat));
-					MaterialPass->SetInput(ePId_Input0, FRenderingCompositeOutputRef(IncomingStage));
-					MaterialPass->SetInput(ePId_Input1, FRenderingCompositeOutputRef(SeparateTranslucencyInput));
-					MaterialPass->SetInput(ePId_Input2, FRenderingCompositeOutputRef(PreTonemapHDRColorInput));
-					MaterialPass->SetInput(ePId_Input3, FRenderingCompositeOutputRef(PostTonemapHDRColorInput));
-					MaterialPass->SetInput(ePId_Input4, FRenderingCompositeOutputRef(PreFlattenVelocity));
-
-					auto Proxy = MaterialInterface->GetRenderProxy();
-					const FMaterial* Material = Proxy->GetMaterial(Context.View.GetFeatureLevel());
-					if (Material->NeedsGBuffer())
-					{
-						// AdjustGBufferRefCount(-1) call is done when the pass gets executed
-						FSceneRenderTargets::Get(Context.RHICmdList).AdjustGBufferRefCount(Context.RHICmdList, 1);
-					}
+					FRenderingCompositePass* MaterialPass = AddPostProcessMaterialPass(Context, MaterialInterface, OutputFormat);
+					MaterialPass->SetInput(EPassInputId(EPostProcessMaterialInput::SceneColor), IncomingStage);
+					MaterialPass->SetInput(EPassInputId(EPostProcessMaterialInput::SeparateTranslucency), SeparateTranslucencyInput);
+					MaterialPass->SetInput(EPassInputId(EPostProcessMaterialInput::PreTonemapHDRColor), PreTonemapHDRColorInput);
+					MaterialPass->SetInput(EPassInputId(EPostProcessMaterialInput::PostTonemapHDRColor), PostTonemapHDRColorInput);
+					MaterialPass->SetInput(EPassInputId(EPostProcessMaterialInput::Velocity), PreFlattenVelocity);
 
-					FString VisualizationName = (*It)->GetName();
+					FString VisualizationName = MaterialInterface->GetName();
 
-					const TSharedPtr<FImagePixelPipe, ESPMode::ThreadSafe>* OutputPipe = Context.View.FinalPostProcessSettings.BufferVisualizationPipes.Find((*It)->GetFName());
+					const TSharedPtr<FImagePixelPipe, ESPMode::ThreadSafe>* OutputPipe = Context.View.FinalPostProcessSettings.BufferVisualizationPipes.Find(MaterialInterface->GetFName());
 					if (OutputPipe && OutputPipe->IsValid())
 					{
 						MaterialPass->SetOutputDumpPipe(ePId_Output0, *OutputPipe);
@@ -1742,23 +1492,16 @@ void FPostProcessing::Process(FRHICommandListImmediate& RHICmdList, const FViewI
 
 			if(bAllowTonemapper)
 			{
-				auto Node = AddSinglePostProcessMaterial(Context, BL_ReplacingTonemapper);
-
-				if(Node)
 				{
-					// a custom tonemapper is provided
-					Node->SetInput(ePId_Input0, Context.FinalOutput);
-
-					// We are binding separate translucency here because the post process SceneTexture node can reference 
-					// the separate translucency buffers through ePId_Input1. 
-					// TODO: Check if material actually uses this texture and only bind if needed.
-					Node->SetInput(ePId_Input1, SeparateTranslucency);
-					Node->SetInput(ePId_Input2, BloomOutputCombined);
-					Context.FinalOutput = Node;
-				}
-				else
-				{
-					Tonemapper = AddTonemapper(Context, BloomOutputCombined, AutoExposure.EyeAdaptation, AutoExposure.MethodId, false, bHDRTonemapperOutput);
+					FRenderingCompositeOutputRef FinalOutputPrev = Context.FinalOutput;
+
+					Context.FinalOutput = AddPostProcessMaterialReplaceTonemapPass(Context, SeparateTranslucency, BloomOutputCombined);
+
+					// No-op from post process material pass; run built-in tonemapper instead.
+					if (Context.FinalOutput == FinalOutputPrev)
+					{
+						Tonemapper = AddTonemapper(Context, BloomOutputCombined, AutoExposure.EyeAdaptation, AutoExposure.MethodId, false, bHDRTonemapperOutput);
+					}
 				}
 
 				PostTonemapHDRColor = Context.FinalOutput;
@@ -1961,7 +1704,7 @@ void FPostProcessing::Process(FRHICommandListImmediate& RHICmdList, const FViewI
 			Context.FinalOutput = FRenderingCompositeOutputRef(Node);
 		}
 
-		AddHighResScreenshotMask(Context, SeparateTranslucency);
+		AddHighResScreenshotMask(Context);
 
 		FIntPoint PrimaryUpscaleViewSize = Context.View.GetSecondaryViewRectSize();
 
@@ -2572,13 +2315,8 @@ void FPostProcessing::ProcessES2(FRHICommandListImmediate& RHICmdList, FScene* S
 			}
 		}
 
-		// Screenshot mask
-		{
-			FRenderingCompositeOutputRef EmptySeparateTranslucency;
-			AddHighResScreenshotMask(Context, EmptySeparateTranslucency);
-		}
-		
-	
+		AddHighResScreenshotMask(Context);
+
 #if WITH_EDITOR
 		// Show the selection outline if it is in the editor and we aren't in wireframe 
 		// If the engine is in demo mode and game view is on we also do not show the selection outline
diff --git a/Engine/Source/Runtime/Renderer/Private/PostProcess/RenderingCompositionGraph.cpp b/Engine/Source/Runtime/Renderer/Private/PostProcess/RenderingCompositionGraph.cpp
--- a/Engine/Source/Runtime/Renderer/Private/PostProcess/RenderingCompositionGraph.cpp
+++ b/Engine/Source/Runtime/Renderer/Private/PostProcess/RenderingCompositionGraph.cpp
@@ -1199,19 +1199,24 @@ FRDGTextureRef FRenderingCompositePass::FindOrCreateRDGTextureForOutput(
 	const FRDGTextureDesc& TextureDesc,
 	const TCHAR* TextureName)
 {
-	if (FRenderingCompositeOutput* Output = GetOutput(OutputId))
+	if (FRDGTextureRef OutputTexture = FindRDGTextureForOutput(GraphBuilder, OutputId, TextureName))
 	{
-		const TRefCountPtr<IPooledRenderTarget>& ExistingTarget = Output->PooledRenderTarget;
+		return OutputTexture;
+	}
+	return GraphBuilder.CreateTexture(TextureDesc, TextureName);
+}
 
-		if (ExistingTarget)
+FRDGTextureRef FRenderingCompositePass::FindRDGTextureForOutput(
+	FRDGBuilder& GraphBuilder,
+	EPassOutputId OutputId,
+	const TCHAR* TextureName)
+{
+	if (FRenderingCompositeOutput* Output = GetOutput(OutputId))
+	{
+		if (const TRefCountPtr<IPooledRenderTarget>& ExistingTarget = Output->PooledRenderTarget)
 		{
 			return GraphBuilder.RegisterExternalTexture(ExistingTarget, TextureName);
 		}
-		else
-		{
-			return GraphBuilder.CreateTexture(TextureDesc, TextureName);
-		}
 	}
-
 	return nullptr;
 }
\ No newline at end of file
diff --git a/Engine/Source/Runtime/Renderer/Private/PostProcess/RenderingCompositionGraph.h b/Engine/Source/Runtime/Renderer/Private/PostProcess/RenderingCompositionGraph.h
--- a/Engine/Source/Runtime/Renderer/Private/PostProcess/RenderingCompositionGraph.h
+++ b/Engine/Source/Runtime/Renderer/Private/PostProcess/RenderingCompositionGraph.h
@@ -372,6 +372,12 @@ struct FRenderingCompositePass
 		const FRDGTextureDesc& TextureDesc,
 		const TCHAR* TextureName);
 
+	/** Finds an allocated output texture if one exists. Otherwise, returns null */
+	FRDGTextureRef FindRDGTextureForOutput(
+		FRDGBuilder& GraphBuilder,
+		EPassOutputId OutputId,
+		const TCHAR* TextureName);
+
 	/**
 	 * Registers a RDG texture to be extracted to the assigned output during graph execution.
 	 */
@@ -453,6 +459,16 @@ struct FRenderingCompositeOutputRef
 		return IsValid() ? Source->GetComputePassEndFence() : nullptr;
 	}
 
+	bool operator==(const FRenderingCompositeOutputRef& Other) const
+	{
+		return Source == Other.Source && PassOutputId == Other.PassOutputId;
+	}
+
+	bool operator!=(const FRenderingCompositeOutputRef& Other) const
+	{
+		return !(*this == Other);
+	}
+
 private:
 	/** can be 0 */
 	FRenderingCompositePass* Source;
diff --git a/Engine/Source/Runtime/Renderer/Private/ScreenPass.cpp b/Engine/Source/Runtime/Renderer/Private/ScreenPass.cpp
--- a/Engine/Source/Runtime/Renderer/Private/ScreenPass.cpp
+++ b/Engine/Source/Runtime/Renderer/Private/ScreenPass.cpp
@@ -6,6 +6,7 @@
 #include "RendererModule.h"
 #include "IXRTrackingSystem.h"
 #include "IHeadMountedDisplay.h"
+#include "RenderGraphUtils.h"
 
 extern bool ShouldDoComputePostProcessing(const FViewInfo& View);
 
@@ -134,4 +135,71 @@ FScreenPassViewInfo::FScreenPassViewInfo(const FViewInfo& InView)
 	, StereoPass(View.StereoPass)
 	, bHasHMDMask(IsHMDHiddenAreaMaskActive())
 	, bUseComputePasses(ShouldDoComputePostProcessing(InView))
-{}
\ No newline at end of file
+{}
+
+void SetScreenPassPipelineState(FRHICommandListImmediate& RHICmdList, const FScreenPassDrawInfo& ScreenPassDraw)
+{
+	FRHIPixelShader* PixelShaderRHI = GETSAFERHISHADER_PIXEL(ScreenPassDraw.PixelShader);
+	FRHIVertexShader* VertexShaderRHI = GETSAFERHISHADER_VERTEX(ScreenPassDraw.VertexShader);
+
+	FGraphicsPipelineStateInitializer GraphicsPSOInit;
+	RHICmdList.ApplyCachedRenderTargets(GraphicsPSOInit);
+	GraphicsPSOInit.BlendState = ScreenPassDraw.BlendState;
+	GraphicsPSOInit.RasterizerState = TStaticRasterizerState<>::GetRHI();
+	GraphicsPSOInit.DepthStencilState = ScreenPassDraw.DepthStencilState;
+	GraphicsPSOInit.BoundShaderState.VertexDeclarationRHI = GFilterVertexDeclaration.VertexDeclarationRHI;
+	GraphicsPSOInit.BoundShaderState.VertexShaderRHI = VertexShaderRHI;
+	GraphicsPSOInit.BoundShaderState.PixelShaderRHI = PixelShaderRHI;
+	GraphicsPSOInit.PrimitiveType = PT_TriangleList;
+
+	SetGraphicsPipelineState(RHICmdList, GraphicsPSOInit);
+}
+
+void AddDrawTexturePass(
+	FRDGBuilder& GraphBuilder,
+	const FScreenPassViewInfo& ScreenPassView,
+	FRDGTextureRef InputTexture,
+	FRDGTextureRef OutputTexture,
+	FIntPoint InputPosition,
+	FIntPoint OutputPosition,
+	FIntPoint Size)
+{
+	const FRDGTextureDesc& InputDesc = InputTexture->Desc;
+	const FRDGTextureDesc& OutputDesc = OutputTexture->Desc;
+
+	// Use a hardware copy if formats match.
+	if (InputDesc.Format == OutputDesc.Format)
+	{
+		return AddCopyTexturePass(GraphBuilder, InputTexture, OutputTexture, InputPosition, OutputPosition, Size);
+	}
+
+	if (Size == FIntPoint::ZeroValue)
+	{
+		// Copy entire input texture to output texture.
+		Size = InputTexture->Desc.Extent;
+	}
+
+	// Don't prime color data if the whole texture is being overwritten.
+	const ERenderTargetLoadAction LoadAction = (OutputPosition == FIntPoint::ZeroValue && Size == OutputDesc.Extent)
+		? ERenderTargetLoadAction::ENoAction
+		: ERenderTargetLoadAction::ELoad;
+
+	const FScreenPassTextureViewport InputViewport(FIntRect(InputPosition, InputPosition + Size), InputTexture);
+	const FScreenPassTextureViewport OutputViewport(FIntRect(OutputPosition, OutputPosition + Size), OutputTexture);
+
+	TShaderMapRef<FCopyRectPS> PixelShader(ScreenPassView.View.ShaderMap);
+
+	FCopyRectPS::FParameters* Parameters = GraphBuilder.AllocParameters<FCopyRectPS::FParameters>();
+	Parameters->InputTexture = InputTexture;
+	Parameters->InputSampler = TStaticSamplerState<>::GetRHI();
+	Parameters->RenderTargets[0] = FRenderTargetBinding(OutputTexture, LoadAction, ERenderTargetStoreAction::EStore);
+
+	AddDrawScreenPass(
+		GraphBuilder,
+		RDG_EVENT_NAME("DrawTexturePass"),
+		ScreenPassView,
+		OutputViewport,
+		InputViewport,
+		*PixelShader,
+		Parameters);
+}
\ No newline at end of file
diff --git a/Engine/Source/Runtime/Renderer/Private/ScreenPass.h b/Engine/Source/Runtime/Renderer/Private/ScreenPass.h
--- a/Engine/Source/Runtime/Renderer/Private/ScreenPass.h
+++ b/Engine/Source/Runtime/Renderer/Private/ScreenPass.h
@@ -107,13 +107,12 @@ END_SHADER_PARAMETER_STRUCT()
 
 FScreenPassTextureViewportParameters GetScreenPassTextureViewportParameters(const FScreenPassTextureViewport& InViewport);
 
-/**
- * Contains a transform that maps UV coordinates from one screen pass texture viewport to another.
- * Assumes normalized UV coordinates [0, 0]x[1, 1] where [0, 0] maps to the source view min
- * coordinate and [1, 1] maps to the source view rect max coordinate.
+/** Contains a transform that maps UV coordinates from one screen pass texture viewport to another.
+ *  Assumes normalized UV coordinates [0, 0]x[1, 1] where [0, 0] maps to the source view min
+ *  coordinate and [1, 1] maps to the source view rect max coordinate.
  *
- * Example Usage:
- *    float2 DestinationUV = SourceUV * UVScaleBias.xy + UVScaleBias.zw;
+ *  Example Usage:
+ *     float2 DestinationUV = SourceUV * UVScaleBias.xy + UVScaleBias.zw;
  */
 BEGIN_SHADER_PARAMETER_STRUCT(FScreenPassTextureViewportTransform, )
 	// A scale / bias factor to apply to the input UV coordinate, converting it to a output UV coordinate.
@@ -133,6 +132,7 @@ FScreenPassTextureViewportTransform GetScreenPassTextureViewportTransform(
 	FVector2D DestinationUVOffset,
 	FVector2D DestinationUVExtent);
 
+// View information cached off for use by screen passes.
 class FScreenPassViewInfo
 {
 public:
@@ -154,42 +154,83 @@ public:
 	const bool bUseComputePasses;
 };
 
-/**
- * Draws a full-viewport triangle with the provided pixel shader type. The destination full-viewport triangle
- * and interpolated source UV coordinates are derived from the viewport and texture rectangles, respectively.
- * @param ViewportRect The rectangle, in pixels, of the viewport drawn onto the destination render target.
- * @param TextureRect The rectangle, in pixels, of the source texture mapped to the [0, 1]x[0, 1] UV space.
- * @param TextureSize The total size, in pixels, of the source texture.
- * @param PixelShader The pixel shader instance assigned to the pipeline state.
- * @param PixelShaderParameters The parameter block assigned to the pixel shader prior to draw.
+/** Draw information for the more advanced DrawScreenPass variant. Allows customizing the blend / depth stencil state,
+ *  providing a custom vertex shader, and more fine-grained control of the underlying draw call.
  */
-template<typename TPixelShaderType>
+struct FScreenPassDrawInfo
+{
+	enum class EFlags : uint8
+	{
+		None,
+
+		// Flips the Y axis of the rendered quad. Used by mobile rendering.
+		FlipYAxis = 0x1
+	};
+
+	using FDefaultBlendState = TStaticBlendState<>;
+	using FDefaultDepthStencilState = TStaticDepthStencilState<false, CF_Always>;
+
+	FScreenPassDrawInfo() = default;
+
+	FScreenPassDrawInfo(
+		FShader* InVertexShader,
+		FShader* InPixelShader,
+		FRHIBlendState* InBlendState = FDefaultBlendState::GetRHI(),
+		FRHIDepthStencilState* InDepthStencilState = FDefaultDepthStencilState::GetRHI(),
+		EFlags InFlags = EFlags::None)
+		: VertexShader(InVertexShader)
+		, PixelShader(InPixelShader)
+		, BlendState(InBlendState)
+		, DepthStencilState(InDepthStencilState)
+		, Flags(InFlags)
+	{}
+
+	void Validate() const
+	{
+		check(VertexShader);
+		check(PixelShader);
+		check(BlendState);
+		check(DepthStencilState);
+	}
+
+	FShader* VertexShader = nullptr;
+	FShader* PixelShader = nullptr;
+	FRHIBlendState* BlendState = nullptr;
+	FRHIDepthStencilState* DepthStencilState = nullptr;
+	EFlags Flags = EFlags::None;
+};
+
+ENUM_CLASS_FLAGS(FScreenPassDrawInfo::EFlags);
+
+// Helper function which sets the pipeline state object on the command list prior to invoking a screen pass.
+void SetScreenPassPipelineState(FRHICommandListImmediate& RHICmdList, const FScreenPassDrawInfo& ScreenPassDraw);
+
+/** Draws a full-viewport triangle with the provided pixel shader type. The destination full-viewport triangle
+ *  and interpolated source UV coordinates are derived from the viewport and texture rectangles, respectively.
+ *  @param OutputViewport The output viewport defining the render target extent and viewport rect.
+ *  @param InputViewport The input viewport defining the rect region to map UV coordinates into.
+ *  @param PixelShader The pixel shader instance assigned to the pipeline state.
+ *  @param PixelShaderParameters The parameter block assigned to the pixel shader prior to draw.
+ */
+template <typename TPixelShaderType>
 void DrawScreenPass(
 	FRHICommandListImmediate& RHICmdList,
 	const FScreenPassViewInfo& ScreenPassView,
-	FIntRect OutputRect,
-	FIntRect InputRect,
-	FIntPoint InputSize,
+	const FScreenPassTextureViewport& OutputViewport,
+	const FScreenPassTextureViewport& InputViewport,
 	TPixelShaderType* PixelShader,
 	const typename TPixelShaderType::FParameters& PixelShaderParameters)
 {
+	const FIntRect InputRect = InputViewport.Rect;
+	const FIntPoint InputSize = InputViewport.Extent;
+	const FIntRect OutputRect = OutputViewport.Rect;
 	const FIntPoint OutputSize = OutputRect.Size();
-	FRHIPixelShader* PixelShaderRHI = GETSAFERHISHADER_PIXEL(PixelShader);
-	FScreenPassVS* VertexShader = *(ScreenPassView.ScreenPassVS);
-
-	FGraphicsPipelineStateInitializer GraphicsPSOInit;
-	RHICmdList.ApplyCachedRenderTargets(GraphicsPSOInit);
-	GraphicsPSOInit.BlendState = TStaticBlendState<>::GetRHI();
-	GraphicsPSOInit.RasterizerState = TStaticRasterizerState<>::GetRHI();
-	GraphicsPSOInit.DepthStencilState = TStaticDepthStencilState<false, CF_Always>::GetRHI();
-	GraphicsPSOInit.BoundShaderState.VertexDeclarationRHI = GFilterVertexDeclaration.VertexDeclarationRHI;
-	GraphicsPSOInit.BoundShaderState.VertexShaderRHI = GETSAFERHISHADER_VERTEX(VertexShader);
-	GraphicsPSOInit.BoundShaderState.PixelShaderRHI = PixelShaderRHI;
-	GraphicsPSOInit.PrimitiveType = PT_TriangleList;
 
 	RHICmdList.SetViewport(OutputRect.Min.X, OutputRect.Min.Y, 0.0f, OutputRect.Max.X, OutputRect.Max.Y, 1.0f);
-	SetGraphicsPipelineState(RHICmdList, GraphicsPSOInit);
-	SetShaderParameters(RHICmdList, PixelShader, PixelShaderRHI, PixelShaderParameters);
+
+	SetScreenPassPipelineState(RHICmdList, FScreenPassDrawInfo(*ScreenPassView.ScreenPassVS, PixelShader));
+
+	SetShaderParameters(RHICmdList, PixelShader, PixelShader->GetPixelShader(), PixelShaderParameters);
 
 	DrawPostProcessPass(
 		RHICmdList,
@@ -197,48 +238,76 @@ void DrawScreenPass(
 		InputRect.Min.X, InputRect.Min.Y, InputRect.Width(), InputRect.Height(),
 		OutputSize,
 		InputSize,
-		VertexShader,
+		*ScreenPassView.ScreenPassVS,
 		ScreenPassView.StereoPass,
 		ScreenPassView.bHasHMDMask,
 		EDRF_UseTriangleOptimization);
 }
 
-/**
- * Helper variant of @ref DrawScreenPass. All other parameters are forwarded.
- * @param DestinationViewDesc The destination texture view descriptor, for viewport generation.
- * @param SourceViewDesc The source view descriptor, for UV generation.
+/** More advanced variant of screen pass drawing. Supports overriding blend / depth stencil
+ *  pipeline state, and providing a custom vertex shader. Shader parameters are not bound by
+ *  this method, instead the user provides a setup function that is called prior to draw, but
+ *  after setting the PSO. This setup function should assign shader parameters.
  */
-template <typename TPixelShaderType>
-inline void DrawScreenPass(
+template<typename TSetupFunction>
+void DrawScreenPass(
 	FRHICommandListImmediate& RHICmdList,
 	const FScreenPassViewInfo& ScreenPassView,
 	const FScreenPassTextureViewport& OutputViewport,
 	const FScreenPassTextureViewport& InputViewport,
-	TPixelShaderType* PixelShader,
-	const typename TPixelShaderType::FParameters& PixelShaderParameters)
+	const FScreenPassDrawInfo& ScreenPassDraw,
+	TSetupFunction SetupFunction)
 {
-	DrawScreenPass(
+	ScreenPassDraw.Validate();
+
+	const FIntRect InputRect = InputViewport.Rect;
+	const FIntPoint InputSize = InputViewport.Extent;
+	const FIntRect OutputRect = OutputViewport.Rect;
+	const FIntPoint OutputSize = OutputRect.Size();
+
+	RHICmdList.SetViewport(OutputRect.Min.X, OutputRect.Min.Y, 0.0f, OutputRect.Max.X, OutputRect.Max.Y, 1.0f);
+
+	EDrawRectangleFlags DrawRectangleFlags = EDRF_UseTriangleOptimization;
+
+	FIntPoint LocalOutputPos(FIntPoint::ZeroValue);
+	FIntPoint LocalOutputSize(OutputSize);
+
+	if ((ScreenPassDraw.Flags & FScreenPassDrawInfo::EFlags::FlipYAxis) == FScreenPassDrawInfo::EFlags::FlipYAxis)
+	{
+		// Draw the quad flipped. Requires that the cull mode be disabled.
+		LocalOutputPos.Y = OutputSize.Y;
+		LocalOutputSize.Y = -OutputSize.Y;
+
+		// Triangle optimization currently doesn't work when flipped.
+		DrawRectangleFlags = EDRF_Default;
+	}
+
+	SetScreenPassPipelineState(RHICmdList, ScreenPassDraw);
+
+	SetupFunction(RHICmdList);
+
+	DrawPostProcessPass(
 		RHICmdList,
-		ScreenPassView,
-		OutputViewport.Rect,
-		InputViewport.Rect,
-		InputViewport.Extent,
-		PixelShader,
-		PixelShaderParameters);
+		LocalOutputPos.X, LocalOutputPos.Y, LocalOutputSize.X, LocalOutputSize.Y,
+		InputRect.Min.X, InputRect.Min.Y, InputRect.Width(), InputRect.Height(),
+		OutputSize,
+		InputSize,
+		ScreenPassDraw.VertexShader,
+		ScreenPassView.StereoPass,
+		ScreenPassView.bHasHMDMask,
+		DrawRectangleFlags);
 }
 
-/**
- * Adds a new Render Graph pass which internally calls @ref DrawScreenPass.
- * Parameters are forwarded to @ref DrawScreenPass.
+/** Render graph variant of simpler DrawScreenPass function. Clears graph resources unused by the
+ *  pixel shader prior to adding the pass.
  */
 template <typename TPixelShaderType>
-inline void AddDrawScreenPass(
+void AddDrawScreenPass(
 	FRDGBuilder& GraphBuilder,
 	FRDGEventName&& PassName,
 	const FScreenPassViewInfo& ScreenPassView,
-	FIntRect OutputViewportRect,
-	FIntRect InputViewportRect,
-	FIntPoint InputExtent,
+	const FScreenPassTextureViewport& OutputViewport,
+	const FScreenPassTextureViewport& InputViewport,
 	TPixelShaderType* PixelShader,
 	typename TPixelShaderType::FParameters* PixelShaderParameters)
 {
@@ -251,40 +320,81 @@ inline void AddDrawScreenPass(
 		Forward<FRDGEventName>(PassName),
 		PixelShaderParameters,
 		ERenderGraphPassFlags::None,
-		[ScreenPassView, OutputViewportRect, InputViewportRect, InputExtent, PixelShader, PixelShaderParameters]
-		(FRHICommandListImmediate& RHICmdList)
+		[ScreenPassView, OutputViewport, InputViewport, PixelShader, PixelShaderParameters]
+	(FRHICommandListImmediate& RHICmdList)
 	{
 		DrawScreenPass(
 			RHICmdList,
 			ScreenPassView,
-			OutputViewportRect,
-			InputViewportRect,
-			InputExtent,
+			OutputViewport,
+			InputViewport,
 			PixelShader,
 			*PixelShaderParameters);
 	});
 }
 
-/**
- * Helper variants of @ref AddDrawScreenPass. All other parameters are forwarded.
+/** Render graph variant of more advanced DrawScreenPass function. Does *not* clear unused graph
+ *  resources, since the parameters might be shared between the vertex and pixel shaders.
  */
-template <typename TPixelShaderType>
-inline void AddDrawScreenPass(
+template <typename TSetupFunction, typename TPassParameterStruct>
+void AddDrawScreenPass(
 	FRDGBuilder& GraphBuilder,
 	FRDGEventName&& PassName,
 	const FScreenPassViewInfo& ScreenPassView,
 	const FScreenPassTextureViewport& OutputViewport,
 	const FScreenPassTextureViewport& InputViewport,
-	TPixelShaderType* PixelShader,
-	typename TPixelShaderType::FParameters* PixelShaderParameters)
+	const FScreenPassDrawInfo& ScreenPassDraw,
+	TPassParameterStruct* PassParameterStruct,
+	TSetupFunction SetupFunction)
 {
-	AddDrawScreenPass(
-		GraphBuilder,
+	ScreenPassDraw.Validate();
+	check(PassParameterStruct);
+
+	GraphBuilder.AddPass(
 		Forward<FRDGEventName>(PassName),
+		PassParameterStruct,
+		ERenderGraphPassFlags::None,
+		[ScreenPassView, OutputViewport, InputViewport, ScreenPassDraw, SetupFunction]
+		(FRHICommandListImmediate& RHICmdList)
+	{
+		DrawScreenPass(
+			RHICmdList,
+			ScreenPassView,
+			OutputViewport,
+			InputViewport,
+			ScreenPassDraw,
+			SetupFunction);
+	});
+}
+
+/** Helper function which copies a region of an input texture to a region of the output texture,
+ *  with support for format conversion. If formats match, the method falls back to a simple DMA
+ *  (CopyTexture); otherwise, it rasterizes using a pixel shader. Use this method if the two
+ *  textures may have different formats.
+ */
+void AddDrawTexturePass(
+	FRDGBuilder& GraphBuilder,
+	const FScreenPassViewInfo& ScreenPassView,
+	FRDGTextureRef InputTexture,
+	FRDGTextureRef OutputTexture,
+	FIntPoint InputPosition = FIntPoint::ZeroValue,
+	FIntPoint OutputPosition = FIntPoint::ZeroValue,
+	FIntPoint Size = FIntPoint::ZeroValue);
+
+/** Helper variant which takes a shared viewport instead of unique input / output positions. */
+inline void AddDrawTexturePass(
+	FRDGBuilder& GraphBuilder,
+	const FScreenPassViewInfo& ScreenPassView,
+	FRDGTextureRef InputTexture,
+	FRDGTextureRef OutputTexture,
+	FIntRect ViewportRect)
+{
+	AddDrawTexturePass(
+		GraphBuilder,
 		ScreenPassView,
-		OutputViewport.Rect,
-		InputViewport.Rect,
-		InputViewport.Extent,
-		PixelShader,
-		PixelShaderParameters);
+		InputTexture,
+		OutputTexture,
+		ViewportRect.Min,
+		ViewportRect.Min,
+		ViewportRect.Size());
 }
\ No newline at end of file
