Don't crash on VT feedback size of 0

diff --git a/Engine/Source/Runtime/Engine/Private/ShaderCompiler/ShaderCompiler.cpp b/Engine/Source/Runtime/Engine/Private/ShaderCompiler/ShaderCompiler.cpp
--- a/Engine/Source/Runtime/Engine/Private/ShaderCompiler/ShaderCompiler.cpp
+++ b/Engine/Source/Runtime/Engine/Private/ShaderCompiler/ShaderCompiler.cpp
@@ -3383,7 +3383,7 @@ void GlobalBeginCompileShader(
 	
 	{
 		static const auto CVar = IConsoleManager::Get().FindConsoleVariable(TEXT("r.vt.FeedbackFactor"));
-		Input.Environment.SetDefine(TEXT("VIRTUAL_TEXTURE_FEEDBACK_FACTOR"), CVar ? CVar->GetInt() : 0);
+		Input.Environment.SetDefine(TEXT("VIRTUAL_TEXTURE_FEEDBACK_FACTOR"), CVar ? FMath::Max(CVar->GetInt(), 1) : 1);
 	}
 
 
diff --git a/Engine/Source/Runtime/Renderer/Private/PostProcess/SceneRenderTargets.cpp b/Engine/Source/Runtime/Renderer/Private/PostProcess/SceneRenderTargets.cpp
--- a/Engine/Source/Runtime/Renderer/Private/PostProcess/SceneRenderTargets.cpp
+++ b/Engine/Source/Runtime/Renderer/Private/PostProcess/SceneRenderTargets.cpp
@@ -2413,7 +2413,7 @@ void FSceneRenderTargets::AllocateDeferredShadingPathRenderTargets(FRHICommandLi
 
 	if (UseVirtualTexturing(CurrentFeatureLevel))
 	{
-		FIntPoint FeedbackSize = FIntPoint::DivideAndRoundUp(BufferSize, GVirtualTextureFeedbackFactor);
+		FIntPoint FeedbackSize = FIntPoint::DivideAndRoundUp(BufferSize, FMath::Max(GVirtualTextureFeedbackFactor, 1));
 		VirtualTextureFeedback.CreateResourceGPU(RHICmdList, FeedbackSize);
 	}
 }
