Adding collateral profile scopes back into ray tracing ambient occlusion and global illumination passes.

diff --git a/Engine/Source/Runtime/Renderer/Private/RayTracing/RayTracingAmbientOcclusion.cpp b/Engine/Source/Runtime/Renderer/Private/RayTracing/RayTracingAmbientOcclusion.cpp
--- a/Engine/Source/Runtime/Renderer/Private/RayTracing/RayTracingAmbientOcclusion.cpp
+++ b/Engine/Source/Runtime/Renderer/Private/RayTracing/RayTracingAmbientOcclusion.cpp
@@ -123,6 +123,9 @@ void FDeferredShadingSceneRenderer::RenderRayTracingAmbientOcclusion(
 	FRDGTextureRef* OutAmbientOcclusionTexture)
 #if RHI_RAYTRACING
 {
+	RDG_GPU_STAT_SCOPE(GraphBuilder, RayTracingAmbientOcclusion);
+	RDG_EVENT_SCOPE(GraphBuilder, "Ray Tracing Ambient Occlusion");
+
 	// Allocates denoiser inputs.
 	IScreenSpaceDenoiser::FAmbientOcclusionInputs DenoiserInputs;
 	{
diff --git a/Engine/Source/Runtime/Renderer/Private/RayTracing/RayTracingGlobalIllumination.cpp b/Engine/Source/Runtime/Renderer/Private/RayTracing/RayTracingGlobalIllumination.cpp
--- a/Engine/Source/Runtime/Renderer/Private/RayTracing/RayTracingGlobalIllumination.cpp
+++ b/Engine/Source/Runtime/Renderer/Private/RayTracing/RayTracingGlobalIllumination.cpp
@@ -293,7 +293,8 @@ void FDeferredShadingSceneRenderer::RenderRayTracingGlobalIllumination(
 		OutRayTracingConfig->ResolutionFraction = FMath::Clamp(GRayTracingGlobalIlluminationScreenPercentage / 100.0, 0.25, 1.0);
 	}
 
-	RDG_EVENT_SCOPE(GraphBuilder, "RTGI");
+	RDG_GPU_STAT_SCOPE(GraphBuilder, RayTracingGlobalIllumination);
+	RDG_EVENT_SCOPE(GraphBuilder, "Ray Tracing Global Illumination");
 
 	int32 RayTracingGISamplesPerPixel = GRayTracingGlobalIlluminationSamplesPerPixel > -1 ? GRayTracingGlobalIlluminationSamplesPerPixel : View.FinalPostProcessSettings.RayTracingGISamplesPerPixel;
 	OutRayTracingConfig->RayCountPerPixel = RayTracingGISamplesPerPixel;
