Reprojects and reduces previous frame for SSRTGI.

diff --git a/Engine/Shaders/Private/SSRT/SSRTDiffuseIndirect.usf b/Engine/Shaders/Private/SSRT/SSRTDiffuseIndirect.usf
--- a/Engine/Shaders/Private/SSRT/SSRTDiffuseIndirect.usf
+++ b/Engine/Shaders/Private/SSRT/SSRTDiffuseIndirect.usf
@@ -29,6 +29,7 @@
 #elif QUALITY == 3
 	#define CONFIG_RAY_STEPS 8
 	#define CONFIG_RAY_COUNT 16
+	#define CONFIG_BACKFACE_MODULATION 1
 
 	#define TILE_PIXEL_SIZE_X 4
 	#define TILE_PIXEL_SIZE_X_LOG 2
@@ -39,6 +40,7 @@
 #elif QUALITY == 4
 	#define CONFIG_RAY_STEPS 12
 	#define CONFIG_RAY_COUNT 32
+	#define CONFIG_BACKFACE_MODULATION 1
 
 	#define TILE_PIXEL_SIZE_X 4
 	#define TILE_PIXEL_SIZE_X_LOG 2
@@ -51,6 +53,11 @@
 
 #endif
 
+
+#ifndef CONFIG_BACKFACE_MODULATION
+	#define CONFIG_BACKFACE_MODULATION 0
+#endif
+
 #define TILE_PIXEL_COUNT (TILE_PIXEL_SIZE_X * TILE_PIXEL_SIZE_Y)
 #define LANE_PER_GROUPS (TILE_PIXEL_COUNT * CONFIG_RAY_COUNT)
 		
@@ -69,7 +76,6 @@
 #endif
 
 Texture2D HZBTexture;
-Texture2D VelocityTexture;
 Texture2D ColorTexture;
 
 float4 HZBUvFactorAndInvFactor;
@@ -371,7 +377,7 @@ void MainCS(
 	
             bool bDebugPrint = all(PixelPosition == uint2(View.ViewSizeAndInvSize.xy) / 2);
 
-			float Level = 0; // TODO: remove
+			float Level;
 			float3 HitUVz;
 			bool bHit = RayCast(
 				HZBTexture, GlobalPointClampedSampler,
@@ -394,15 +400,12 @@ void MainCS(
 			//BRANCH // TODO
 			if (bHit)
 			{
-				float2 SampleUV;
-				float Vignette;
-				ReprojectHit(PrevScreenPositionScaleBias, VelocityTexture, GlobalPointClampedSampler, HitUVz, SampleUV, Vignette);
-
-				float4 SampleColor = SampleScreenColor(ColorTexture, GlobalBilinearClampedSampler, SampleUV);
-
+				float4 SampleColor = ColorTexture.SampleLevel(GlobalBilinearClampedSampler, HitUVz, Level);
+				
 				float SampleColorWeight = 1.0;
 
 				// Backface modulation of hit surface
+				#if CONFIG_BACKFACE_MODULATION
 				{
 					float3 SampleNormal = GetGBufferDataFromSceneTextures(HitUVz.xy).WorldNormal;
 					
@@ -417,13 +420,14 @@ void MainCS(
 		
 					SampleColorWeight *= saturate( 1 - dot( SampleNormal, L ) );
 				}
+				#endif
 
 				#if CONFIG_RAY_COUNT > 1
 					SampleColorWeight *= rcp( 1 + Luminance(SampleColor.rgb) );
 				#endif
-			
+				
 				float3 DiffuseColor = SampleColor.rgb * SampleColorWeight;
-				float AmbientOcclusion = SampleColor.a;
+				float AmbientOcclusion = 1.0;
 				
 				CompressedColor.x = asuint(f32tof16(DiffuseColor.r) << 16 | f32tof16(DiffuseColor.g));
 				CompressedColor.y = asuint(f32tof16(DiffuseColor.b) << 16 | f32tof16(AmbientOcclusion));
diff --git a/Engine/Shaders/Private/SSRT/SSRTPrevFrameReduction.usf b/Engine/Shaders/Private/SSRT/SSRTPrevFrameReduction.usf
new file mode 100644
--- /dev/null
+++ b/Engine/Shaders/Private/SSRT/SSRTPrevFrameReduction.usf
@@ -0,0 +1,142 @@
+// Copyright 1998-2019 Epic Games, Inc. All Rights Reserved.
+
+#include "SSRTRayCast.ush"
+#include "../SceneTextureParameters.ush"
+
+
+#define GROUP_TILE_SIZE 8
+#define GROUP_PIXEL_COUNT GROUP_TILE_SIZE * GROUP_TILE_SIZE
+
+
+uint bIsTemporalAAHistory;
+Texture2D PrevTemporalAAHistory;
+
+float4 PrevScreenPositionScaleBias;
+float  PrevSceneColorPreExposureCorrection;
+
+RWTexture2D<float4>	ReducedSceneColorOutput_0;
+RWTexture2D<float4>	ReducedSceneColorOutput_1;
+RWTexture2D<float4>	ReducedSceneColorOutput_2;
+
+groupshared uint SharedMemory[GROUP_PIXEL_COUNT * 2];
+
+
+[numthreads(GROUP_TILE_SIZE, GROUP_TILE_SIZE, 1)]
+void MainCS(
+	uint2 DispatchThreadId : SV_DispatchThreadID,
+	uint2 GroupId : SV_GroupID,
+	uint2 GroupThreadId : SV_GroupThreadID,
+	uint GroupThreadIndex : SV_GroupIndex)
+{
+	float2 BufferUV = (float2(DispatchThreadId) + 0.5) * View.BufferSizeAndInvSize.zw;
+	float2 ViewportUV = (float2(DispatchThreadId) + 0.5f) * View.ViewSizeAndInvSize.zw;
+    float2 ScreenPosition = ViewportUVToScreenPos(ViewportUV);
+	
+	float2 PrevFrameUV;
+	float2 PrevScreen;
+
+	BRANCH
+	if (bIsTemporalAAHistory)
+    {
+		float DeviceZ = SceneDepthBuffer.SampleLevel(SceneDepthBufferSampler, BufferUV, 0).r;
+
+		// Camera motion for pixel (in ScreenPos space).
+		float4 ThisClip = float4(ScreenPosition, DeviceZ, 1);
+		float4 PrevClip = mul(ThisClip, View.ClipToPrevClip);
+		PrevScreen = PrevClip.xy / PrevClip.w;
+
+		float2 EncodedVelocity = SceneVelocityBuffer.SampleLevel(SceneVelocityBufferSampler, BufferUV, 0).xy;
+		if (EncodedVelocity.x > 0.0)
+		{
+			PrevScreen = ThisClip.xy - DecodeVelocityFromTexture(EncodedVelocity);
+		}
+
+		PrevFrameUV = PrevScreen.xy * PrevScreenPositionScaleBias.xy + PrevScreenPositionScaleBias.zw;
+    }
+	else
+    {
+		PrevFrameUV = BufferUV;
+		PrevScreen = ScreenPosition;
+    }
+
+	// Fetch color.
+	float4 PrevSceneColor = float4(
+		PrevTemporalAAHistory.SampleLevel(GlobalBilinearClampedSampler, PrevFrameUV, 0).rgb * PrevSceneColorPreExposureCorrection, 1.0);
+
+	// Apply vignette to the color.
+	{
+		float Vignette = min(ComputeHitVignetteFromScreenPos(ScreenPosition), ComputeHitVignetteFromScreenPos(PrevScreen));
+		PrevSceneColor *= Vignette;
+	}
+
+	// Output mip 0
+	ReducedSceneColorOutput_0[DispatchThreadId] = PrevSceneColor;
+	
+	// Reduce lower mip levels.
+	{
+		// Store to LDS
+		{
+			SharedMemory[GROUP_PIXEL_COUNT * 0 | GroupThreadIndex] = (f32tof16(PrevSceneColor.r) << 0) | (f32tof16(PrevSceneColor.g) << 16);
+			SharedMemory[GROUP_PIXEL_COUNT * 1 | GroupThreadIndex] = (f32tof16(PrevSceneColor.b) << 0) | (f32tof16(PrevSceneColor.a) << 16);
+		}
+	
+		GroupMemoryBarrierWithGroupSync();
+
+		// Reduce lower mip levels.
+		UNROLL
+		for (uint MipLevel = 1; MipLevel < 3; MipLevel++)
+		{
+			float4 PrevSceneColor = 0;
+			
+			const uint ReductionAmount = 1 << MipLevel;
+			const uint NumberPixelInMip = GROUP_PIXEL_COUNT / (ReductionAmount * ReductionAmount);
+			
+			if (GroupThreadIndex < NumberPixelInMip)
+            {
+                uint2 OutputCoord = uint2(
+					GroupThreadIndex % (GROUP_TILE_SIZE / ReductionAmount),
+					GroupThreadIndex / (GROUP_TILE_SIZE / ReductionAmount));
+
+				float4 ReducedColor = 0;
+
+				UNROLL
+				for (uint x = 0; x < 2; x++)
+                {
+					UNROLL
+					for (uint y = 0; y < 2; y++)
+					{
+						uint2 Coord = OutputCoord * 2 + uint2(x, y);
+						uint LDSIndex = Coord.x + Coord.y * ((2 * GROUP_TILE_SIZE) / ReductionAmount);
+
+						uint Raw0 = SharedMemory[GROUP_PIXEL_COUNT * 0 | LDSIndex];
+						uint Raw1 = SharedMemory[GROUP_PIXEL_COUNT * 1 | LDSIndex];
+
+						float4 Color;
+						Color.r = f16tof32(Raw0 >> 0);
+						Color.g = f16tof32(Raw0 >> 16);
+						Color.b = f16tof32(Raw1 >> 0);
+						Color.a = f16tof32(Raw1 >> 16);
+
+						ReducedColor += Color;
+					}
+                }
+
+				ReducedColor *= rcp(4.0);
+
+				uint2 OutputPosition = GroupId * (GROUP_TILE_SIZE / ReductionAmount) + OutputCoord;
+
+				if (MipLevel == 1)
+                {
+					ReducedSceneColorOutput_1[OutputPosition] = ReducedColor;
+                }
+				else if (MipLevel == 2)
+                {
+					ReducedSceneColorOutput_2[OutputPosition] = ReducedColor;
+                }
+				
+				SharedMemory[GROUP_PIXEL_COUNT * 0 | GroupThreadIndex] = (f32tof16(ReducedColor.r) << 0) | (f32tof16(ReducedColor.g) << 16);
+				SharedMemory[GROUP_PIXEL_COUNT * 1 | GroupThreadIndex] = (f32tof16(ReducedColor.b) << 0) | (f32tof16(ReducedColor.a) << 16);
+            } // if (GroupThreadIndex < NumberPixelInMip)
+        } // for (uint MipLevel = 1; MipLevel < 3; MipLevel++)
+    }
+} // MainCS()
diff --git a/Engine/Source/Runtime/Renderer/Private/ScreenSpaceRayTracing.cpp b/Engine/Source/Runtime/Renderer/Private/ScreenSpaceRayTracing.cpp
--- a/Engine/Source/Runtime/Renderer/Private/ScreenSpaceRayTracing.cpp
+++ b/Engine/Source/Runtime/Renderer/Private/ScreenSpaceRayTracing.cpp
@@ -165,6 +165,30 @@ class FSSRQualityDim : SHADER_PERMUTATION_ENUM_CLASS("SSR_QUALITY", ESSRQuality)
 class FSSROutputForDenoiser : SHADER_PERMUTATION_BOOL("SSR_OUTPUT_FOR_DENOISER");
 
 
+class FSSRTPrevFrameReductionCS : public FGlobalShader
+{
+	DECLARE_GLOBAL_SHADER(FSSRTPrevFrameReductionCS);
+	SHADER_USE_PARAMETER_STRUCT(FSSRTPrevFrameReductionCS, FGlobalShader);
+
+	static bool ShouldCompilePermutation(const FGlobalShaderPermutationParameters& Parameters)
+	{
+		return IsFeatureLevelSupported(Parameters.Platform, ERHIFeatureLevel::SM5);
+	}
+
+	BEGIN_SHADER_PARAMETER_STRUCT(FParameters, )
+		SHADER_PARAMETER(FVector4, PrevScreenPositionScaleBias)
+		SHADER_PARAMETER(float, PrevSceneColorPreExposureCorrection)
+		SHADER_PARAMETER(uint32, bIsTemporalAAHistory)
+
+		SHADER_PARAMETER_RDG_TEXTURE(Texture2D, PrevTemporalAAHistory)
+
+		SHADER_PARAMETER_STRUCT_INCLUDE(FSceneTextureParameters, SceneTextures)
+		SHADER_PARAMETER_STRUCT_REF(FViewUniformShaderParameters, View)
+
+		SHADER_PARAMETER_RDG_TEXTURE_UAV_ARRAY(RWTexture<float4>, ReducedSceneColorOutput, [4])
+	END_SHADER_PARAMETER_STRUCT()
+};
+
 class FScreenSpaceReflectionsStencilPS : public FGlobalShader
 {
 	DECLARE_GLOBAL_SHADER(FScreenSpaceReflectionsStencilPS);
@@ -234,17 +258,10 @@ class FScreenSpaceDiffuseIndirectCS : public FGlobalShader
 		SHADER_PARAMETER( FVector4,		PrevScreenPositionScaleBias )
 		SHADER_PARAMETER( float,		PrevSceneColorPreExposureCorrection )
 		
-		SHADER_PARAMETER_RDG_TEXTURE( Texture2D,	HZBTexture )
-		SHADER_PARAMETER_SAMPLER( SamplerState,		HZBSampler )
-
-		SHADER_PARAMETER_RDG_TEXTURE( Texture2D,	VelocityTexture )
-		SHADER_PARAMETER_SAMPLER( SamplerState,		VelocitySampler )
-
-		SHADER_PARAMETER_RDG_TEXTURE( Texture2D,	ColorTexture )
-		SHADER_PARAMETER_SAMPLER( SamplerState,		ColorSampler )
+		SHADER_PARAMETER_RDG_TEXTURE(Texture2D, HZBTexture)
+		SHADER_PARAMETER_RDG_TEXTURE(Texture2D, ColorTexture)
 		
 		SHADER_PARAMETER_STRUCT_INCLUDE(FSceneTextureParameters, SceneTextures)
-		SHADER_PARAMETER_STRUCT_INCLUDE(FSceneTextureSamplerParameters, SceneTextureSamplers)
 		SHADER_PARAMETER_STRUCT_REF(FViewUniformShaderParameters, View)
 
 		SHADER_PARAMETER_RDG_TEXTURE_UAV(RWTexture<float4>, IndirectDiffuseOutput)
@@ -260,6 +277,7 @@ class FScreenSpaceDiffuseIndirectCS : public FGlobalShader
 };
 
 
+IMPLEMENT_GLOBAL_SHADER(FSSRTPrevFrameReductionCS,        "/Engine/Private/SSRT/SSRTPrevFrameReduction.usf", "MainCS", SF_Compute);
 IMPLEMENT_GLOBAL_SHADER(FScreenSpaceReflectionsPS,        "/Engine/Private/SSRT/SSRTReflections.usf", "ScreenSpaceReflectionsPS", SF_Pixel);
 IMPLEMENT_GLOBAL_SHADER(FScreenSpaceReflectionsStencilPS, "/Engine/Private/SSRT/SSRTReflections.usf", "ScreenSpaceReflectionsStencilPS", SF_Pixel);
 IMPLEMENT_GLOBAL_SHADER(FScreenSpaceDiffuseIndirectCS,    "/Engine/Private/SSRT/SSRTDiffuseIndirect.usf", "MainCS", SF_Compute);
@@ -593,78 +611,123 @@ void RenderScreenSpaceDiffuseIndirect(
 	int32 RayCountPerPixel;
 	GetSSRTGIShaderOptionsForQuality(Quality, &GroupSize, &RayCountPerPixel);
 
-	FSceneRenderTargets& SceneContext = FSceneRenderTargets::Get(GraphBuilder.RHICmdList);
+	FRDGTexture* HZBTexture = GraphBuilder.RegisterExternalTexture(View.HZB);
+	FRDGTexture* ColorTexture = GraphBuilder.RegisterExternalTexture(TemporalAAHistory.RT[0]);
 
-	// Allocate outputs.
+	// Reproject and reduce previous frame color.
+	FRDGTexture* ReducedSceneColor;
 	{
-		FRDGTextureDesc Desc = FRDGTextureDesc::Create2DDesc(
-			SceneTextures.SceneDepthBuffer->Desc.Extent,
-			PF_FloatRGBA,
-			FClearValueBinding::Transparent,
-			/* InFlags = */ TexCreate_None,
-			/* InTargetableFlags = */ TexCreate_ShaderResource | TexCreate_RenderTargetable | TexCreate_UAV,
-			/* bInForceSeparateTargetAndShaderResource = */ false);
+		// Allocate ReducedSceneColor.
+		{
+			FRDGTextureDesc Desc = FRDGTextureDesc::Create2DDesc(
+				SceneTextures.SceneDepthBuffer->Desc.Extent,
+				PF_FloatRGBA, // TODO: might be worth using FloatRGB + R8
+				FClearValueBinding::None,
+				/* InFlags = */ TexCreate_None,
+				/* InTargetableFlags = */ TexCreate_ShaderResource | TexCreate_RenderTargetable | TexCreate_UAV,
+				/* bInForceSeparateTargetAndShaderResource = */ false);
+			Desc.NumMips = 3;
+
+			ReducedSceneColor = GraphBuilder.CreateTexture(Desc, TEXT("SSRTReducedSceneColor"));
+		}
+
+		FSSRTPrevFrameReductionCS::FParameters* PassParameters = GraphBuilder.AllocParameters<FSSRTPrevFrameReductionCS::FParameters>();
+
+		FIntPoint ViewportOffset = TemporalAAHistory.ViewportRect.Min;
+		FIntPoint ViewportExtent = TemporalAAHistory.ViewportRect.Size();
+		FIntPoint BufferSize = TemporalAAHistory.ReferenceBufferSize;
+
+		PassParameters->PrevScreenPositionScaleBias = FVector4(
+			ViewportExtent.X * 0.5f / BufferSize.X,
+			-ViewportExtent.Y * 0.5f / BufferSize.Y,
+			(ViewportExtent.X * 0.5f + ViewportOffset.X) / BufferSize.X,
+			(ViewportExtent.Y * 0.5f + ViewportOffset.Y) / BufferSize.Y);
 
-		OutDenoiserInputs->Color = GraphBuilder.CreateTexture(Desc, TEXT("SSRTDiffuseIndirect"));
+		PassParameters->PrevSceneColorPreExposureCorrection = View.PreExposure / View.PrevViewInfo.SceneColorPreExposure;
+
+		PassParameters->bIsTemporalAAHistory = true;
+		PassParameters->PrevTemporalAAHistory = ColorTexture;
+
+		PassParameters->SceneTextures = SceneTextures;
+		PassParameters->View = View.ViewUniformBuffer;
+
+		for (int32 MipLevel = 0; MipLevel < ReducedSceneColor->Desc.NumMips; MipLevel++)
+		{
+			PassParameters->ReducedSceneColorOutput[MipLevel] = GraphBuilder.CreateUAV(FRDGTextureUAVDesc(ReducedSceneColor, MipLevel));
+		}
 
-		Desc.Format = PF_R16F;
-		OutDenoiserInputs->AmbientOcclusionMask = GraphBuilder.CreateTexture(Desc, TEXT("SSRTAmbientOcclusion"));
+		TShaderMapRef<FSSRTPrevFrameReductionCS> ComputeShader(View.ShaderMap);
+		FComputeShaderUtils::AddPass(
+			GraphBuilder,
+			RDG_EVENT_NAME("PrevFrameReduction %dx%d", View.ViewRect.Width(), View.ViewRect.Height()),
+			*ComputeShader,
+			PassParameters,
+			FComputeShaderUtils::GetGroupCount(View.ViewRect.Size(), GroupSize));
 	}
 
-	FRDGTexture* HZBTexture	= GraphBuilder.RegisterExternalTexture( View.HZB );
-	FRDGTexture* ColorTexture	= GraphBuilder.RegisterExternalTexture( TemporalAAHistory.RT[0] );
 
-	FScreenSpaceDiffuseIndirectCS::FParameters* PassParameters = GraphBuilder.AllocParameters<FScreenSpaceDiffuseIndirectCS::FParameters>();
+	{
+		// Allocate outputs.
+		{
+			FRDGTextureDesc Desc = FRDGTextureDesc::Create2DDesc(
+				SceneTextures.SceneDepthBuffer->Desc.Extent,
+				PF_FloatRGBA,
+				FClearValueBinding::Transparent,
+				/* InFlags = */ TexCreate_None,
+				/* InTargetableFlags = */ TexCreate_ShaderResource | TexCreate_RenderTargetable | TexCreate_UAV,
+				/* bInForceSeparateTargetAndShaderResource = */ false);
 
-	PassParameters->HZBTexture = HZBTexture;
-	PassParameters->HZBSampler = TStaticSamplerState< SF_Point, AM_Clamp, AM_Clamp, AM_Clamp >::GetRHI();
+			OutDenoiserInputs->Color = GraphBuilder.CreateTexture(Desc, TEXT("SSRTDiffuseIndirect"));
 
-	PassParameters->VelocityTexture = SceneTextures.SceneVelocityBuffer;
-	PassParameters->VelocitySampler = TStaticSamplerState< SF_Point, AM_Clamp, AM_Clamp, AM_Clamp >::GetRHI();
+			Desc.Format = PF_R16F;
+			OutDenoiserInputs->AmbientOcclusionMask = GraphBuilder.CreateTexture(Desc, TEXT("SSRTAmbientOcclusion"));
+		}
 
-	PassParameters->ColorTexture = ColorTexture;
-	PassParameters->ColorSampler = TStaticSamplerState< SF_Bilinear, AM_Clamp, AM_Clamp, AM_Clamp >::GetRHI();
+		FScreenSpaceDiffuseIndirectCS::FParameters* PassParameters = GraphBuilder.AllocParameters<FScreenSpaceDiffuseIndirectCS::FParameters>();
 
-	const FVector2D HZBUvFactor(
-		float( View.ViewRect.Width() )  / float( 2 * View.HZBMipmap0Size.X ),
-		float( View.ViewRect.Height() ) / float( 2 * View.HZBMipmap0Size.Y )
-		);
+		PassParameters->HZBTexture = HZBTexture;
+		PassParameters->ColorTexture = ReducedSceneColor;
+
+		const FVector2D HZBUvFactor(
+			float( View.ViewRect.Width() )  / float( 2 * View.HZBMipmap0Size.X ),
+			float( View.ViewRect.Height() ) / float( 2 * View.HZBMipmap0Size.Y )
+			);
 			
-	PassParameters->HZBUvFactorAndInvFactor = FVector4(
-		HZBUvFactor.X,
-		HZBUvFactor.Y,
-		1.0f / HZBUvFactor.X,
-		1.0f / HZBUvFactor.Y );
-
-	FIntPoint ViewportOffset	= TemporalAAHistory.ViewportRect.Min;
-	FIntPoint ViewportExtent	= TemporalAAHistory.ViewportRect.Size();
-	FIntPoint BufferSize		= TemporalAAHistory.ReferenceBufferSize;
-
-	PassParameters->PrevScreenPositionScaleBias = FVector4(
-		 ViewportExtent.X * 0.5f / BufferSize.X,
-		-ViewportExtent.Y * 0.5f / BufferSize.Y,
-		(ViewportExtent.X * 0.5f + ViewportOffset.X) / BufferSize.X,
-		(ViewportExtent.Y * 0.5f + ViewportOffset.Y) / BufferSize.Y );
-
-	PassParameters->PrevSceneColorPreExposureCorrection = View.PreExposure / View.PrevViewInfo.SceneColorPreExposure;
-
-	PassParameters->SceneTextures = SceneTextures;
-	SetupSceneTextureSamplers(&PassParameters->SceneTextureSamplers);
-	PassParameters->View = View.ViewUniformBuffer;
+		PassParameters->HZBUvFactorAndInvFactor = FVector4(
+			HZBUvFactor.X,
+			HZBUvFactor.Y,
+			1.0f / HZBUvFactor.X,
+			1.0f / HZBUvFactor.Y );
+
+		FIntPoint ViewportOffset	= TemporalAAHistory.ViewportRect.Min;
+		FIntPoint ViewportExtent	= TemporalAAHistory.ViewportRect.Size();
+		FIntPoint BufferSize		= TemporalAAHistory.ReferenceBufferSize;
+
+		PassParameters->PrevScreenPositionScaleBias = FVector4(
+			 ViewportExtent.X * 0.5f / BufferSize.X,
+			-ViewportExtent.Y * 0.5f / BufferSize.Y,
+			(ViewportExtent.X * 0.5f + ViewportOffset.X) / BufferSize.X,
+			(ViewportExtent.Y * 0.5f + ViewportOffset.Y) / BufferSize.Y );
+
+		PassParameters->PrevSceneColorPreExposureCorrection = View.PreExposure / View.PrevViewInfo.SceneColorPreExposure;
+
+		PassParameters->SceneTextures = SceneTextures;
+		PassParameters->View = View.ViewUniformBuffer;
 	
-	PassParameters->IndirectDiffuseOutput = GraphBuilder.CreateUAV(OutDenoiserInputs->Color);
-	PassParameters->AmbientOcclusionOutput = GraphBuilder.CreateUAV(OutDenoiserInputs->AmbientOcclusionMask);
-	PassParameters->ScreenSpaceRayTracingDebugOutput = CreateScreenSpaceRayTracingDebugUAV(GraphBuilder, OutDenoiserInputs->Color->Desc);
-
-	FScreenSpaceDiffuseIndirectCS::FPermutationDomain PermutationVector;
-	PermutationVector.Set<FScreenSpaceDiffuseIndirectCS::FQualityDim>(Quality);
-
-	TShaderMapRef<FScreenSpaceDiffuseIndirectCS> ComputeShader(View.ShaderMap, PermutationVector);
-	FComputeShaderUtils::AddPass(
-		GraphBuilder,
-		RDG_EVENT_NAME("ScreenSpaceDiffuseIndirect(Quality=%d RayPerPixel=%d) %dx%d",
-			Quality, RayCountPerPixel, View.ViewRect.Width(), View.ViewRect.Height()),
-		*ComputeShader,
-		PassParameters,
-		FComputeShaderUtils::GetGroupCount(View.ViewRect.Size(), GroupSize));
+		PassParameters->IndirectDiffuseOutput = GraphBuilder.CreateUAV(OutDenoiserInputs->Color);
+		PassParameters->AmbientOcclusionOutput = GraphBuilder.CreateUAV(OutDenoiserInputs->AmbientOcclusionMask);
+		PassParameters->ScreenSpaceRayTracingDebugOutput = CreateScreenSpaceRayTracingDebugUAV(GraphBuilder, OutDenoiserInputs->Color->Desc);
+
+		FScreenSpaceDiffuseIndirectCS::FPermutationDomain PermutationVector;
+		PermutationVector.Set<FScreenSpaceDiffuseIndirectCS::FQualityDim>(Quality);
+
+		TShaderMapRef<FScreenSpaceDiffuseIndirectCS> ComputeShader(View.ShaderMap, PermutationVector);
+		FComputeShaderUtils::AddPass(
+			GraphBuilder,
+			RDG_EVENT_NAME("ScreenSpaceDiffuseIndirect(Quality=%d RayPerPixel=%d) %dx%d",
+				Quality, RayCountPerPixel, View.ViewRect.Width(), View.ViewRect.Height()),
+			*ComputeShader,
+			PassParameters,
+			FComputeShaderUtils::GetGroupCount(View.ViewRect.Size(), GroupSize));
+	}
 } // RenderScreenSpaceDiffuseIndirect()
