Deletes TAA's pixel shader given SM4 is no longer supported.

diff --git a/Engine/Shaders/Private/PostProcessTemporalAA.usf b/Engine/Shaders/Private/PostProcessTemporalAA.usf
--- a/Engine/Shaders/Private/PostProcessTemporalAA.usf
+++ b/Engine/Shaders/Private/PostProcessTemporalAA.usf
@@ -2236,9 +2236,7 @@ FTAAHistoryPayload TemporalAASample(uint2 GroupId, uint2 GroupThreadId, uint Gro
 }
 
 
-//------------------------------------------------------- ENTRY POINTS
-
-#if COMPUTESHADER
+//------------------------------------------------------- ENTRY POINT
 
 [numthreads(THREADGROUP_SIZEX, THREADGROUP_SIZEY, 1)]
 void MainCS(
@@ -2288,7 +2286,7 @@ void MainCS(
 		#endif
 	}
 
-#if TAA_DOWNSAMPLE
+	#if TAA_DOWNSAMPLE
 	{
 		// This shader permutation outputs half resolution image in addition to main full-res one.
 		// It is more efficient than performing a separate downsampling pass afterwards.
@@ -2309,62 +2307,5 @@ void MainCS(
 				(OutColor0 + GroupSharedDownsampleArray[P1] + GroupSharedDownsampleArray[P2] + GroupSharedDownsampleArray[P3]) * 0.25;
 		}
 	}
-#endif //TAA_DOWNSAMPLE
+	#endif //TAA_DOWNSAMPLE
 }
-
-
-#elif PIXELSHADER // DEPRECATED.
-
-#if AA_UPSAMPLE
-	#error Up sample only available to CS.
-#endif
-
-
-void MainPS(
-	float4 SvPosition : SV_POSITION,
-	#if HISTORY_RENDER_TARGETS == 2
-		out float4 OutColor0 : SV_Target0,
-		out float4 OutColor1 : SV_Target1
-	#else
-		out float4 OutColor0 : SV_Target0
-	#endif
-	)
-{
-	float2 ViewportUV = (SvPosition.xy - OutputViewportRect.xy) * OutputViewportSize.zw;
-
-	#if AA_LOWER_RESOLUTION
-	{
-		ViewportUV = (SvPosition.xy - OutputViewportRect.xy) * MaxViewportUVAndSvPositionToViewportUV.zw;
-		ViewportUV = min(ViewportUV, MaxViewportUVAndSvPositionToViewportUV.xy);
-	}
-	#endif
-
-	uint2 GroupId = int2(0, 0);
-	uint2 GroupThreadId = int2(0, 0);
-	uint GroupThreadIndex = 0;
-	float FrameExposureScale = EyeAdaptationLookup();
-	
-	FTAAHistoryPayload OutputPayload = TemporalAASample(GroupId, GroupThreadId, GroupThreadIndex, ViewportUV, FrameExposureScale);
-
-	#if AA_HISTORY_PAYLOAD == HISTORY_PAYLOAD_RGB_COC
-	{
-		OutColor0.rgb = OutputPayload.Color.rgb;
-		OutColor0.a = OutputPayload.CocRadius;
-	}
-	#elif AA_HISTORY_PAYLOAD == HISTORY_PAYLOAD_RGB_OPACITY_COC
-	{
-		OutColor0 = OutputPayload.Color;
-		OutColor1.r = OutputPayload.CocRadius;
-	}
-	#else
-	{
-		OutColor0 = OutputPayload.Color;
-	}
-	#endif
-}
-
-
-#else
-	#error Only PS or CS.
-
-#endif
diff --git a/Engine/Source/Runtime/Renderer/Private/PostProcess/DiaphragmDOF.cpp b/Engine/Source/Runtime/Renderer/Private/PostProcess/DiaphragmDOF.cpp
--- a/Engine/Source/Runtime/Renderer/Private/PostProcess/DiaphragmDOF.cpp
+++ b/Engine/Source/Runtime/Renderer/Private/PostProcess/DiaphragmDOF.cpp
@@ -1351,7 +1351,6 @@ FRDGTextureRef DiaphragmDOF::AddPasses(
 	FTAAPassParameters TAAParameters(View);
 	{
 		TAAParameters.Pass = ETAAPassConfig::DiaphragmDOF;
-		TAAParameters.bIsComputePass = true;
 
 		// When using dynamic resolution, the blur introduce by TAA's history resolution changes is quite noticeable on DOF.
 		// Therefore we switch for a temporal upsampling technic to maintain the same history resolution.
diff --git a/Engine/Source/Runtime/Renderer/Private/PostProcess/PostProcessTemporalAA.cpp b/Engine/Source/Runtime/Renderer/Private/PostProcess/PostProcessTemporalAA.cpp
--- a/Engine/Source/Runtime/Renderer/Private/PostProcess/PostProcessTemporalAA.cpp
+++ b/Engine/Source/Runtime/Renderer/Private/PostProcess/PostProcessTemporalAA.cpp
@@ -68,49 +68,6 @@ static float CatmullRom( float x )
 }
 
 
-BEGIN_SHADER_PARAMETER_STRUCT(FTAAShaderParameters,)
-	SHADER_PARAMETER(FVector4, ViewportUVToInputBufferUV)
-	SHADER_PARAMETER(FVector4, MaxViewportUVAndSvPositionToViewportUV)
-	SHADER_PARAMETER(FVector2D, ScreenPosAbsMax)
-	SHADER_PARAMETER(float, HistoryPreExposureCorrection)
-	SHADER_PARAMETER(float, CurrentFrameWeight)
-	SHADER_PARAMETER(int32, bCameraCut)
-	
-	SHADER_PARAMETER_ARRAY(float, SampleWeights, [9])
-	SHADER_PARAMETER_ARRAY(float, PlusWeights, [5])
-
-	SHADER_PARAMETER(FVector4, InputSceneColorSize)
-	SHADER_PARAMETER(FVector4, OutputViewportSize)
-	SHADER_PARAMETER(FVector4, OutputViewportRect)
-
-	// History parameters
-	SHADER_PARAMETER(FVector4, HistoryBufferSize)
-	SHADER_PARAMETER(FVector4, HistoryBufferUVMinMax)
-	SHADER_PARAMETER(FVector4, ScreenPosToHistoryBufferUV)
-	
-	SHADER_PARAMETER_STRUCT_INCLUDE(FSceneTextureParameters, SceneTextures)
-	
-	SHADER_PARAMETER_RDG_TEXTURE(Texture2D, EyeAdaptation)
-
-	// Inputs
-	SHADER_PARAMETER_RDG_TEXTURE(Texture2D, InputSceneColor)
-	SHADER_PARAMETER_SAMPLER(SamplerState, InputSceneColorSampler)
-	SHADER_PARAMETER_RDG_TEXTURE(Texture2D, InputSceneMetadata)
-	SHADER_PARAMETER_SAMPLER(SamplerState, InputSceneMetadataSampler)
-
-	// History resourrces
-	SHADER_PARAMETER_RDG_TEXTURE(Texture2D, HistoryBuffer0)
-	SHADER_PARAMETER_SAMPLER(SamplerState, HistoryBuffer0Sampler)
-	SHADER_PARAMETER_RDG_TEXTURE(Texture2D, HistoryBuffer1)
-	SHADER_PARAMETER_SAMPLER(SamplerState, HistoryBuffer1Sampler)
-	
-	SHADER_PARAMETER_SAMPLER(SamplerState, SceneDepthBufferSampler)
-	SHADER_PARAMETER_SAMPLER(SamplerState, SceneVelocityBufferSampler)
-
-	SHADER_PARAMETER_STRUCT_REF(FViewUniformShaderParameters, ViewUniformBuffer)
-END_SHADER_PARAMETER_STRUCT()
-
-
 // ---------------------------------------------------- Shader permutation dimensions
 
 namespace
@@ -128,74 +85,58 @@ class FTAADownsampleDim : SHADER_PERMUTATION_BOOL("TAA_DOWNSAMPLE");
 
 // ---------------------------------------------------- Shaders
 
-class FTemporalAAPS : public FGlobalShader
+class FTemporalAACS : public FGlobalShader
 {
-	DECLARE_GLOBAL_SHADER(FTemporalAAPS);
-	SHADER_USE_PARAMETER_STRUCT(FTemporalAAPS, FGlobalShader);
+	DECLARE_GLOBAL_SHADER(FTemporalAACS);
+	SHADER_USE_PARAMETER_STRUCT(FTemporalAACS, FGlobalShader);
 
 	using FPermutationDomain = TShaderPermutationDomain<
 		FTAAPassConfigDim,
 		FTAAFastDim,
-		FTAAResponsiveDim>;
-	
+		FTAAScreenPercentageDim,
+		FTAAUpsampleFilteredDim,
+		FTAADownsampleDim>;
+
 	BEGIN_SHADER_PARAMETER_STRUCT(FParameters, )
-		SHADER_PARAMETER_STRUCT_INCLUDE(FTAAShaderParameters, CommonParameters)
-		RENDER_TARGET_BINDING_SLOTS()
-	END_SHADER_PARAMETER_STRUCT()
+		SHADER_PARAMETER(FVector4, ViewportUVToInputBufferUV)
+		SHADER_PARAMETER(FVector4, MaxViewportUVAndSvPositionToViewportUV)
+		SHADER_PARAMETER(FVector2D, ScreenPosAbsMax)
+		SHADER_PARAMETER(float, HistoryPreExposureCorrection)
+		SHADER_PARAMETER(float, CurrentFrameWeight)
+		SHADER_PARAMETER(int32, bCameraCut)
 
-	static bool ShouldCompilePermutation(const FGlobalShaderPermutationParameters& Parameters)
-	{
-		FPermutationDomain PermutationVector(Parameters.PermutationId);
+		SHADER_PARAMETER_ARRAY(float, SampleWeights, [9])
+		SHADER_PARAMETER_ARRAY(float, PlusWeights, [5])
 
-		// TAAU is compute shader only.
-		if (IsTAAUpsamplingConfig(PermutationVector.Get<FTAAPassConfigDim>()))
-		{
-			return false;
-		}
+		SHADER_PARAMETER(FVector4, InputSceneColorSize)
+		SHADER_PARAMETER(FVector4, OutputViewportSize)
+		SHADER_PARAMETER(FVector4, OutputViewportRect)
 
-		// Fast dimensions is only for Main and Diaphragm DOF.
-		if (PermutationVector.Get<FTAAFastDim>() &&
-			!IsMainTAAConfig(PermutationVector.Get<FTAAPassConfigDim>()) &&
-			!IsDOFTAAConfig(PermutationVector.Get<FTAAPassConfigDim>()))
-		{
-			return false;
-		}
-		
-		// DOF setup chain is full compute shader, no need for a pixel shader for TAA.
-		if (IsDOFTAAConfig(PermutationVector.Get<FTAAPassConfigDim>()))
-		{
-			return false;
-		}
+		// History parameters
+		SHADER_PARAMETER(FVector4, HistoryBufferSize)
+		SHADER_PARAMETER(FVector4, HistoryBufferUVMinMax)
+		SHADER_PARAMETER(FVector4, ScreenPosToHistoryBufferUV)
 
-		// Responsive dimension is only for Main.
-		if (PermutationVector.Get<FTAAResponsiveDim>() && !SupportsResponsiveDim(PermutationVector))
-		{
-			return false;
-		}
+		SHADER_PARAMETER_STRUCT_INCLUDE(FSceneTextureParameters, SceneTextures)
 
-		return IsFeatureLevelSupported(Parameters.Platform, ERHIFeatureLevel::SM4);
-	}
+		SHADER_PARAMETER_RDG_TEXTURE(Texture2D, EyeAdaptation)
 
-	static bool SupportsResponsiveDim(const FPermutationDomain& PermutationVector)
-	{
-		return PermutationVector.Get<FTAAPassConfigDim>() == ETAAPassConfig::Main;
-	}
-}; // class FTemporalAAPS
+		// Inputs
+		SHADER_PARAMETER_RDG_TEXTURE(Texture2D, InputSceneColor)
+		SHADER_PARAMETER_SAMPLER(SamplerState, InputSceneColorSampler)
+		SHADER_PARAMETER_RDG_TEXTURE(Texture2D, InputSceneMetadata)
+		SHADER_PARAMETER_SAMPLER(SamplerState, InputSceneMetadataSampler)
 
-class FTemporalAACS : public FGlobalShader
-{
-	DECLARE_GLOBAL_SHADER(FTemporalAACS);
-	SHADER_USE_PARAMETER_STRUCT(FTemporalAACS, FGlobalShader);
+		// History resources
+		SHADER_PARAMETER_RDG_TEXTURE(Texture2D, HistoryBuffer0)
+		SHADER_PARAMETER_SAMPLER(SamplerState, HistoryBuffer0Sampler)
+		SHADER_PARAMETER_RDG_TEXTURE(Texture2D, HistoryBuffer1)
+		SHADER_PARAMETER_SAMPLER(SamplerState, HistoryBuffer1Sampler)
 
-	using FPermutationDomain = TShaderPermutationDomain<
-		FTAAPassConfigDim,
-		FTAAFastDim,
-		FTAAScreenPercentageDim,
-		FTAAUpsampleFilteredDim,
-		FTAADownsampleDim>;
+		SHADER_PARAMETER_SAMPLER(SamplerState, SceneDepthBufferSampler)
+		SHADER_PARAMETER_SAMPLER(SamplerState, SceneVelocityBufferSampler)
 
-	BEGIN_SHADER_PARAMETER_STRUCT(FParameters, )
-		SHADER_PARAMETER_STRUCT_INCLUDE(FTAAShaderParameters, CommonParameters)
+		SHADER_PARAMETER_STRUCT_REF(FViewUniformShaderParameters, ViewUniformBuffer)
 		
 		// Temporal upsample specific parameters.
 		SHADER_PARAMETER(FVector4, InputViewSize)
@@ -295,11 +236,10 @@ class FTemporalAACS : public FGlobalShader
 }; // class FTemporalAACS
 
 
-IMPLEMENT_GLOBAL_SHADER(FTemporalAAPS, "/Engine/Private/PostProcessTemporalAA.usf", "MainPS", SF_Pixel);
 IMPLEMENT_GLOBAL_SHADER(FTemporalAACS, "/Engine/Private/PostProcessTemporalAA.usf", "MainCS", SF_Compute);
 
 
-static void SetupSampleWeightParameters(FTAAShaderParameters* OutTAAParameters, const FTAAPassParameters& PassParameters, FVector2D TemporalJitterPixels)
+static void SetupSampleWeightParameters(FTemporalAACS::FParameters* OutTAAParameters, const FTAAPassParameters& PassParameters, FVector2D TemporalJitterPixels)
 {
 	float JitterX = TemporalJitterPixels.X;
 	float JitterY = TemporalJitterPixels.Y;
@@ -460,41 +400,114 @@ FTAAOutputs FTAAPassParameters::AddTemporalAAPass(
 	// Name of the pass.
 	const TCHAR* PassName = kTAAPassNames[static_cast<int32>(Pass)];
 
-	// Setups common shader parameters
-	FTAAShaderParameters CommonShaderParameters;
+	// Create outputs
+	FTAAOutputs Outputs;
+	{
+		FRDGTextureDesc SceneColorDesc = FRDGTextureDesc::Create2DDesc(
+			GetOutputExtent(),
+			PF_FloatRGBA,
+			FClearValueBinding::Black,
+			/* InFlags = */ TexCreate_None,
+			/* InTargetableFlags = */ TexCreate_ShaderResource | TexCreate_UAV,
+			/* bInForceSeparateTargetAndShaderResource = */ false);
+
+		Outputs.SceneColor = GraphBuilder.CreateTexture(
+			SceneColorDesc,
+			kTAAOutputNames[static_cast<int32>(Pass)],
+			ERDGResourceFlags::MultiFrame);
+
+		if (RenderTargetCount == 2)
+		{
+			Outputs.SceneMetadata = GraphBuilder.CreateTexture(
+				SceneColorDesc,
+				kTAAOutputNames[static_cast<int32>(Pass)],
+				ERDGResourceFlags::MultiFrame);
+		}
+
+		if (bDownsample)
+		{
+			FRDGTextureDesc HalfResSceneColorDesc = FRDGTextureDesc::Create2DDesc(
+				SceneColorDesc.Extent / 2,
+				DownsampleOverrideFormat != PF_Unknown ? DownsampleOverrideFormat : SceneColorInput->Desc.Format,
+				FClearValueBinding::Black,
+				/* InFlags = */ TexCreate_None,
+				/* InTargetableFlags = */ TexCreate_ShaderResource | TexCreate_Transient | TexCreate_UAV,
+				/* bInForceSeparateTargetAndShaderResource = */ false);
+
+			Outputs.DownsampledSceneColor = GraphBuilder.CreateTexture(HalfResSceneColorDesc, TEXT("SceneColorHalfRes"));
+		}
+	}
+
+	RDG_GPU_STAT_SCOPE(GraphBuilder, TAA);
+
 	{
+		FTemporalAACS::FPermutationDomain PermutationVector;
+		PermutationVector.Set<FTAAPassConfigDim>(Pass);
+		PermutationVector.Set<FTAAFastDim>(bUseFast);
+		PermutationVector.Set<FTAADownsampleDim>(bDownsample);
+		PermutationVector.Set<FTAAUpsampleFilteredDim>(true);
+
+		if (IsTAAUpsamplingConfig(Pass))
+		{
+			const bool bUpsampleFiltered = CVarTemporalAAUpsampleFiltered.GetValueOnRenderThread() != 0 || Pass != ETAAPassConfig::MainUpsampling;
+			PermutationVector.Set<FTAAUpsampleFilteredDim>(bUpsampleFiltered);
+
+			// If screen percentage > 100% on X or Y axes, then use screen percentage range = 2 shader permutation to disable LDS caching.
+			if (SrcRect.Width() > DestRect.Width() ||
+				SrcRect.Height() > DestRect.Height())
+			{
+				PermutationVector.Set<FTAAScreenPercentageDim>(2);
+			}
+			// If screen percentage < 50% on X and Y axes, then use screen percentage range = 3 shader permutation.
+			else if (SrcRect.Width() * 100 < 50 * DestRect.Width() &&
+				SrcRect.Height() * 100 < 50 * DestRect.Height() &&
+				Pass == ETAAPassConfig::MainSuperSampling)
+			{
+				PermutationVector.Set<FTAAScreenPercentageDim>(3);
+			}
+			// If screen percentage < 71% on X and Y axes, then use screen percentage range = 1 shader permutation to have smaller LDS caching.
+			else if (SrcRect.Width() * 100 < 71 * DestRect.Width() &&
+				SrcRect.Height() * 100 < 71 * DestRect.Height())
+			{
+				PermutationVector.Set<FTAAScreenPercentageDim>(1);
+			}
+		}
+
+		FTemporalAACS::FParameters* PassParameters = GraphBuilder.AllocParameters<FTemporalAACS::FParameters>();
+
+		// Setups common shader parameters
 		FIntPoint SrcSize = SceneColorInput->Desc.Extent;
 
 		if (!IsTAAUpsamplingConfig(Pass))
-			SetupSampleWeightParameters(&CommonShaderParameters, *this, View.TemporalJitterPixels);
-	
+			SetupSampleWeightParameters(PassParameters, *this, View.TemporalJitterPixels);
+
 		float ResDivisor = ResolutionDivisor;
 		float ResDivisorInv = 1.0f / ResDivisor;
 
-		CommonShaderParameters.ViewUniformBuffer = View.ViewUniformBuffer;
-		CommonShaderParameters.CurrentFrameWeight = CVarTemporalAACurrentFrameWeight.GetValueOnRenderThread();
-		CommonShaderParameters.bCameraCut = bCameraCut;
+		PassParameters->ViewUniformBuffer = View.ViewUniformBuffer;
+		PassParameters->CurrentFrameWeight = CVarTemporalAACurrentFrameWeight.GetValueOnRenderThread();
+		PassParameters->bCameraCut = bCameraCut;
 
-		CommonShaderParameters.SceneTextures = SceneTextures;
-		CommonShaderParameters.SceneDepthBufferSampler = TStaticSamplerState<SF_Point>::GetRHI();
-		CommonShaderParameters.SceneVelocityBufferSampler = TStaticSamplerState<SF_Point>::GetRHI();
+		PassParameters->SceneTextures = SceneTextures;
+		PassParameters->SceneDepthBufferSampler = TStaticSamplerState<SF_Point>::GetRHI();
+		PassParameters->SceneVelocityBufferSampler = TStaticSamplerState<SF_Point>::GetRHI();
 
 		// Input buffer shader parameters
 		{
-			CommonShaderParameters.InputSceneColorSize = FVector4(
+			PassParameters->InputSceneColorSize = FVector4(
 				SceneColorInput->Desc.Extent.X,
 				SceneColorInput->Desc.Extent.Y,
 				1.0f / float(SceneColorInput->Desc.Extent.X),
 				1.0f / float(SceneColorInput->Desc.Extent.Y));
-			CommonShaderParameters.InputSceneColor = SceneColorInput;
-			CommonShaderParameters.InputSceneColorSampler = TStaticSamplerState<SF_Point>::GetRHI();
-			CommonShaderParameters.InputSceneMetadata = SceneMetadataInput;
-			CommonShaderParameters.InputSceneMetadataSampler = TStaticSamplerState<SF_Point>::GetRHI();
+			PassParameters->InputSceneColor = SceneColorInput;
+			PassParameters->InputSceneColorSampler = TStaticSamplerState<SF_Point>::GetRHI();
+			PassParameters->InputSceneMetadata = SceneMetadataInput;
+			PassParameters->InputSceneMetadataSampler = TStaticSamplerState<SF_Point>::GetRHI();
 		}
 
-		CommonShaderParameters.OutputViewportSize = FVector4(
+		PassParameters->OutputViewportSize = FVector4(
 			PracticableDestRect.Width(), PracticableDestRect.Height(), 1.0f / float(PracticableDestRect.Width()), 1.0f / float(PracticableDestRect.Height()));
-		CommonShaderParameters.OutputViewportRect = FVector4(PracticableDestRect.Min.X, PracticableDestRect.Min.Y, PracticableDestRect.Max.X, PracticableDestRect.Max.Y);
+		PassParameters->OutputViewportRect = FVector4(PracticableDestRect.Min.X, PracticableDestRect.Min.Y, PracticableDestRect.Max.X, PracticableDestRect.Max.Y);
 
 		// Set history shader parameters.
 		{
@@ -502,15 +515,15 @@ FTAAOutputs FTAAPassParameters::AddTemporalAAPass(
 			{
 				FRDGTextureRef BlackDummy = GraphBuilder.RegisterExternalTexture(GSystemTextures.BlackDummy);
 
-				CommonShaderParameters.ScreenPosToHistoryBufferUV = FVector4(1.0f, 1.0f, 1.0f, 1.0f);
-				CommonShaderParameters.ScreenPosAbsMax = FVector2D(0.0f, 0.0f);
-				CommonShaderParameters.HistoryBufferUVMinMax = FVector4(0.0f, 0.0f, 0.0f, 0.0f);
-				CommonShaderParameters.HistoryBufferSize = FVector4(1.0f, 1.0f, 1.0f, 1.0f);
-				CommonShaderParameters.HistoryBuffer0 = BlackDummy;
-				CommonShaderParameters.HistoryBuffer1 = BlackDummy;
-				
+				PassParameters->ScreenPosToHistoryBufferUV = FVector4(1.0f, 1.0f, 1.0f, 1.0f);
+				PassParameters->ScreenPosAbsMax = FVector2D(0.0f, 0.0f);
+				PassParameters->HistoryBufferUVMinMax = FVector4(0.0f, 0.0f, 0.0f, 0.0f);
+				PassParameters->HistoryBufferSize = FVector4(1.0f, 1.0f, 1.0f, 1.0f);
+				PassParameters->HistoryBuffer0 = BlackDummy;
+				PassParameters->HistoryBuffer1 = BlackDummy;
+
 				// Remove dependency of the velocity buffer on camera cut, given it's going to be ignored by the shader.
-				CommonShaderParameters.SceneTextures.SceneVelocityBuffer = BlackDummy;
+				PassParameters->SceneTextures.SceneVelocityBuffer = BlackDummy;
 			}
 			else
 			{
@@ -521,7 +534,7 @@ FTAAOutputs FTAAPassParameters::AddTemporalAAPass(
 				float InvReferenceBufferSizeX = 1.f / float(InputHistory.ReferenceBufferSize.X);
 				float InvReferenceBufferSizeY = 1.f / float(InputHistory.ReferenceBufferSize.Y);
 
-				CommonShaderParameters.ScreenPosToHistoryBufferUV = FVector4(
+				PassParameters->ScreenPosToHistoryBufferUV = FVector4(
 					ReferenceViewportExtent.X * 0.5f * InvReferenceBufferSizeX,
 					-ReferenceViewportExtent.Y * 0.5f * InvReferenceBufferSizeY,
 					(ReferenceViewportExtent.X * 0.5f + ReferenceViewportOffset.X) * InvReferenceBufferSizeX,
@@ -531,127 +544,47 @@ FTAAOutputs FTAAPassParameters::AddTemporalAAPass(
 				FIntPoint ViewportExtent = FIntPoint::DivideAndRoundUp(ReferenceViewportExtent, ResolutionDivisor);
 				FIntPoint BufferSize = ReferenceBufferSize / ResolutionDivisor;
 
-				CommonShaderParameters.ScreenPosAbsMax = FVector2D(1.0f - 1.0f / float(ViewportExtent.X), 1.0f - 1.0f / float(ViewportExtent.Y));
+				PassParameters->ScreenPosAbsMax = FVector2D(1.0f - 1.0f / float(ViewportExtent.X), 1.0f - 1.0f / float(ViewportExtent.Y));
 
 				float InvBufferSizeX = 1.f / float(BufferSize.X);
 				float InvBufferSizeY = 1.f / float(BufferSize.Y);
 
-				CommonShaderParameters.HistoryBufferUVMinMax = FVector4(
+				PassParameters->HistoryBufferUVMinMax = FVector4(
 					(ViewportOffset.X + 0.5f) * InvBufferSizeX,
 					(ViewportOffset.Y + 0.5f) * InvBufferSizeY,
 					(ViewportOffset.X + ViewportExtent.X - 0.5f) * InvBufferSizeX,
 					(ViewportOffset.Y + ViewportExtent.Y - 0.5f) * InvBufferSizeY);
 
-				CommonShaderParameters.HistoryBufferSize = FVector4(BufferSize.X, BufferSize.Y, InvBufferSizeX, InvBufferSizeY);
+				PassParameters->HistoryBufferSize = FVector4(BufferSize.X, BufferSize.Y, InvBufferSizeX, InvBufferSizeY);
 
-				CommonShaderParameters.HistoryBuffer0 = GraphBuilder.RegisterExternalTexture(InputHistory.RT[0]);
+				PassParameters->HistoryBuffer0 = GraphBuilder.RegisterExternalTexture(InputHistory.RT[0]);
 				if (InputHistory.RT[1].IsValid())
-					CommonShaderParameters.HistoryBuffer1 = GraphBuilder.RegisterExternalTexture(InputHistory.RT[1]);
+					PassParameters->HistoryBuffer1 = GraphBuilder.RegisterExternalTexture(InputHistory.RT[1]);
 			}
 
-			CommonShaderParameters.HistoryBuffer0Sampler = TStaticSamplerState<SF_Bilinear>::GetRHI();
-			CommonShaderParameters.HistoryBuffer1Sampler = TStaticSamplerState<SF_Bilinear>::GetRHI();
+			PassParameters->HistoryBuffer0Sampler = TStaticSamplerState<SF_Bilinear>::GetRHI();
+			PassParameters->HistoryBuffer1Sampler = TStaticSamplerState<SF_Bilinear>::GetRHI();
 		}
 
-		CommonShaderParameters.MaxViewportUVAndSvPositionToViewportUV = FVector4(
+		PassParameters->MaxViewportUVAndSvPositionToViewportUV = FVector4(
 			(PracticableDestRect.Width() - 0.5f * ResDivisor) / float(PracticableDestRect.Width()),
 			(PracticableDestRect.Height() - 0.5f * ResDivisor) / float(PracticableDestRect.Height()),
 			ResDivisor / float(DestRect.Width()),
 			ResDivisor / float(DestRect.Height()));
 
-		CommonShaderParameters.HistoryPreExposureCorrection = View.PreExposure / View.PrevViewInfo.SceneColorPreExposure;
+		PassParameters->HistoryPreExposureCorrection = View.PreExposure / View.PrevViewInfo.SceneColorPreExposure;
 
 		{
 			float InvSizeX = 1.0f / float(SrcSize.X);
 			float InvSizeY = 1.0f / float(SrcSize.Y);
-			CommonShaderParameters.ViewportUVToInputBufferUV = FVector4(
+			PassParameters->ViewportUVToInputBufferUV = FVector4(
 				ResDivisorInv * InputViewRect.Width() * InvSizeX,
 				ResDivisorInv * InputViewRect.Height() * InvSizeY,
 				ResDivisorInv * InputViewRect.Min.X * InvSizeX,
 				ResDivisorInv * InputViewRect.Min.Y * InvSizeY);
 		}
-		
-		CommonShaderParameters.EyeAdaptation = GetEyeAdaptationTexture(GraphBuilder, View);
-	}
-
-	// Create outputs
-	FTAAOutputs Outputs;
-	{
-		FRDGTextureDesc SceneColorDesc = FRDGTextureDesc::Create2DDesc(
-			GetOutputExtent(),
-			PF_FloatRGBA,
-			FClearValueBinding::Black,
-			/* InFlags = */ TexCreate_None,
-			/* InTargetableFlags = */ TexCreate_ShaderResource | (bIsComputePass ? TexCreate_UAV : TexCreate_RenderTargetable),
-			/* bInForceSeparateTargetAndShaderResource = */ false);
-
-		Outputs.SceneColor = GraphBuilder.CreateTexture(
-			SceneColorDesc,
-			kTAAOutputNames[static_cast<int32>(Pass)],
-			ERDGResourceFlags::MultiFrame);
-
-		if (RenderTargetCount == 2)
-		{
-			Outputs.SceneMetadata = GraphBuilder.CreateTexture(
-				SceneColorDesc,
-				kTAAOutputNames[static_cast<int32>(Pass)],
-				ERDGResourceFlags::MultiFrame);
-		}
 
-		if (bDownsample)
-		{
-			check(bIsComputePass);
-			
-			FRDGTextureDesc HalfResSceneColorDesc = FRDGTextureDesc::Create2DDesc(
-				SceneColorDesc.Extent / 2,
-				DownsampleOverrideFormat != PF_Unknown ? DownsampleOverrideFormat : SceneColorInput->Desc.Format,
-				FClearValueBinding::Black,
-				/* InFlags = */ TexCreate_None,
-				/* InTargetableFlags = */ TexCreate_ShaderResource | TexCreate_Transient | (bIsComputePass ? TexCreate_UAV : TexCreate_RenderTargetable),
-				/* bInForceSeparateTargetAndShaderResource = */ false);
-
-			Outputs.DownsampledSceneColor = GraphBuilder.CreateTexture(HalfResSceneColorDesc, TEXT("SceneColorHalfRes"));
-		}
-	}
-
-	RDG_GPU_STAT_SCOPE(GraphBuilder, TAA);
-
-	if (bIsComputePass)
-	{
-		FTemporalAACS::FPermutationDomain PermutationVector;
-		PermutationVector.Set<FTAAPassConfigDim>(Pass);
-		PermutationVector.Set<FTAAFastDim>(bUseFast);
-		PermutationVector.Set<FTAADownsampleDim>(bDownsample);
-		PermutationVector.Set<FTAAUpsampleFilteredDim>(true);
-
-		if (IsTAAUpsamplingConfig(Pass))
-		{
-			const bool bUpsampleFiltered = CVarTemporalAAUpsampleFiltered.GetValueOnRenderThread() != 0 || Pass != ETAAPassConfig::MainUpsampling;
-			PermutationVector.Set<FTAAUpsampleFilteredDim>(bUpsampleFiltered);
-
-			// If screen percentage > 100% on X or Y axes, then use screen percentage range = 2 shader permutation to disable LDS caching.
-			if (SrcRect.Width() > DestRect.Width() ||
-				SrcRect.Height() > DestRect.Height())
-			{
-				PermutationVector.Set<FTAAScreenPercentageDim>(2);
-			}
-			// If screen percentage < 50% on X and Y axes, then use screen percentage range = 3 shader permutation.
-			else if (SrcRect.Width() * 100 < 50 * DestRect.Width() &&
-				SrcRect.Height() * 100 < 50 * DestRect.Height() &&
-				Pass == ETAAPassConfig::MainSuperSampling)
-			{
-				PermutationVector.Set<FTAAScreenPercentageDim>(3);
-			}
-			// If screen percentage < 71% on X and Y axes, then use screen percentage range = 1 shader permutation to have smaller LDS caching.
-			else if (SrcRect.Width() * 100 < 71 * DestRect.Width() &&
-				SrcRect.Height() * 100 < 71 * DestRect.Height())
-			{
-				PermutationVector.Set<FTAAScreenPercentageDim>(1);
-			}
-		}
-
-		FTemporalAACS::FParameters* PassParameters = GraphBuilder.AllocParameters<FTemporalAACS::FParameters>();
-		PassParameters->CommonParameters = CommonShaderParameters;
+		PassParameters->EyeAdaptation = GetEyeAdaptationTexture(GraphBuilder, View);
 
 		// Temporal upsample specific shader parameters.
 		{
@@ -680,7 +613,7 @@ FTAAOutputs FTAAPassParameters::AddTemporalAAPass(
 		TShaderMapRef<FTemporalAACS> ComputeShader(View.ShaderMap, PermutationVector);
 		FComputeShaderUtils::AddPass(
 			GraphBuilder,
-			RDG_EVENT_NAME("TAA %s CS%s %dx%d -> %dx%d",
+			RDG_EVENT_NAME("TAA %s%s %dx%d -> %dx%d",
 				PassName, bUseFast ? TEXT(" Fast") : TEXT(""),
 				PracticableSrcRect.Width(), PracticableSrcRect.Height(),
 				PracticableDestRect.Width(), PracticableDestRect.Height()),
@@ -688,100 +621,6 @@ FTAAOutputs FTAAPassParameters::AddTemporalAAPass(
 			PassParameters,
 			FComputeShaderUtils::GetGroupCount(PracticableDestRect.Size(), GTemporalAATileSizeX));
 	}
-	else
-	{
-		check(!IsTAAUpsamplingConfig(Pass));
-
-		// Whether to use responsive stencil test.
-		bool bUseResponsiveStencilTest = Pass == ETAAPassConfig::Main && !bIsComputePass && !bCameraCut;
-	
-		FTemporalAAPS::FParameters* PassParameters = GraphBuilder.AllocParameters<FTemporalAAPS::FParameters>();
-		PassParameters->CommonParameters = CommonShaderParameters;
-		PassParameters->RenderTargets[0] = FRenderTargetBinding(
-			Outputs.SceneColor,
-			ERenderTargetLoadAction::ENoAction,
-			ERenderTargetStoreAction::EStore);
-		
-		if (Outputs.SceneMetadata)
-		{
-			PassParameters->RenderTargets[1] = FRenderTargetBinding(
-				Outputs.SceneMetadata,
-				ERenderTargetLoadAction::ENoAction,
-				ERenderTargetStoreAction::EStore);
-		}
-
-		if (bUseResponsiveStencilTest)
-		{
-			PassParameters->RenderTargets.DepthStencil = FDepthStencilBinding(
-				SceneTextures.SceneDepthBuffer,
-				ERenderTargetLoadAction::ENoAction, ERenderTargetStoreAction::ENoAction,
-				ERenderTargetLoadAction::ELoad, ERenderTargetStoreAction::EStore,
-				FExclusiveDepthStencil::DepthRead_StencilRead);
-		}
-		
-		FTemporalAAPS::FPermutationDomain BasePermutationVector;
-		BasePermutationVector.Set<FTAAPassConfigDim>(Pass);
-		BasePermutationVector.Set<FTAAFastDim>(bUseFast);
-
-		TShaderMapRef<FTemporalAAPS> PixelShader(View.ShaderMap, BasePermutationVector);
-		ClearUnusedGraphResources(*PixelShader, PassParameters);
-
-		GraphBuilder.AddPass(
-			RDG_EVENT_NAME("TAA %s PS%s %dx%d",
-				PassName, bUseFast ? TEXT(" Fast") : TEXT(""),
-				PracticableDestRect.Width(), PracticableDestRect.Height()),
-			PassParameters,
-			ERenderGraphPassFlags::None,
-			[PassParameters, &View, PracticableDestRect, bCameraCut, BasePermutationVector, bUseResponsiveStencilTest](FRHICommandList& RHICmdList)
-		{
-			RHICmdList.SetViewport(PracticableDestRect.Min.X, PracticableDestRect.Min.Y, 0.0f, PracticableDestRect.Max.X, PracticableDestRect.Max.Y, 1.0f);
-
-			FTemporalAAPS::FPermutationDomain PermutationVector = BasePermutationVector;
-
-			// Lambda to draw pixel shader.
-			auto DrawTAAPixelShader = [&](FRHIDepthStencilState* DepthStencilState)
-			{
-				TShaderMapRef<FTemporalAAPS> TAAPixelShader(View.ShaderMap, PermutationVector);
-
-				FGraphicsPipelineStateInitializer GraphicsPSOInit;
-				FPixelShaderUtils::InitFullscreenPipelineState(RHICmdList, View.ShaderMap, *TAAPixelShader, GraphicsPSOInit);
-				GraphicsPSOInit.DepthStencilState = DepthStencilState;
-				SetGraphicsPipelineState(RHICmdList, GraphicsPSOInit);
-
-				SetShaderParameters(RHICmdList, *TAAPixelShader, TAAPixelShader->GetPixelShader(), *PassParameters);
-				
-				FPixelShaderUtils::DrawFullscreenTriangle(RHICmdList);
-			};
-	
-			if (bUseResponsiveStencilTest)
-			{
-				// Normal temporal feedback
-				// Draw to pixels where stencil == 0
-				FRHIDepthStencilState* DepthStencilState = TStaticDepthStencilState<
-					false, CF_Always,
-					true, CF_Equal, SO_Keep, SO_Keep, SO_Keep,
-					false, CF_Always, SO_Keep, SO_Keep, SO_Keep,
-					STENCIL_TEMPORAL_RESPONSIVE_AA_MASK, STENCIL_TEMPORAL_RESPONSIVE_AA_MASK>::GetRHI();
-
-				DrawTAAPixelShader(DepthStencilState);
-
-				// Responsive feedback for tagged pixels
-				// Draw to pixels where stencil != 0
-				DepthStencilState = TStaticDepthStencilState<
-					false, CF_Always,
-					true, CF_NotEqual, SO_Keep, SO_Keep, SO_Keep,
-					false, CF_Always, SO_Keep, SO_Keep, SO_Keep,
-					STENCIL_TEMPORAL_RESPONSIVE_AA_MASK, STENCIL_TEMPORAL_RESPONSIVE_AA_MASK>::GetRHI();
-
-				PermutationVector.Set<FTAAResponsiveDim>(true);
-				DrawTAAPixelShader(DepthStencilState);
-			}
-			else
-			{
-				DrawTAAPixelShader(TStaticDepthStencilState<false, CF_Always>::GetRHI());
-			}
-		});
-	}
 	
 	if (!View.bViewStateIsReadOnly)
 	{
@@ -810,7 +649,7 @@ FRCPassPostProcessTemporalAA::FRCPassPostProcessTemporalAA(
 {
 	check(SavedParameters.Validate());
 
-	bIsComputePass = SavedParameters.bIsComputePass;
+	bIsComputePass = true;
 	bPreferAsyncCompute = false;
 }
 
diff --git a/Engine/Source/Runtime/Renderer/Private/PostProcess/PostProcessTemporalAA.h b/Engine/Source/Runtime/Renderer/Private/PostProcess/PostProcessTemporalAA.h
--- a/Engine/Source/Runtime/Renderer/Private/PostProcess/PostProcessTemporalAA.h
+++ b/Engine/Source/Runtime/Renderer/Private/PostProcess/PostProcessTemporalAA.h
@@ -81,11 +81,7 @@ struct FTAAPassParameters
 	// Whether to use the faster shader permutation.
 	bool bUseFast = false;
 
-	// Whether to do compute or not.
-	bool bIsComputePass = false;
-
 	// Whether downsampled (box filtered, half resolution) frame should be written out.
-	// Only used when bIsComputePass is true.
 	bool bDownsample = false;
 	EPixelFormat DownsampleOverrideFormat = PF_Unknown;
 
diff --git a/Engine/Source/Runtime/Renderer/Private/PostProcess/PostProcessing.cpp b/Engine/Source/Runtime/Renderer/Private/PostProcess/PostProcessing.cpp
--- a/Engine/Source/Runtime/Renderer/Private/PostProcess/PostProcessing.cpp
+++ b/Engine/Source/Runtime/Renderer/Private/PostProcess/PostProcessing.cpp
@@ -752,11 +752,7 @@ static FTAAPassParameters MakeTAAPassParametersForView( const FViewInfo& View )
 	Parameters.Pass = View.PrimaryScreenPercentageMethod == EPrimaryScreenPercentageMethod::TemporalUpscale
 		? ETAAPassConfig::MainUpsampling
 		: ETAAPassConfig::Main;
-
-	Parameters.bIsComputePass = Parameters.Pass == ETAAPassConfig::MainUpsampling
-		? true // TAAU is always a compute shader
-		: ShouldDoComputePostProcessing(View);
-
+	
 	Parameters.SetupViewRect(View);
 
 	{
@@ -1219,7 +1215,6 @@ void FPostProcessing::Process(FRHICommandListImmediate& RHICmdList, const FViewI
 				TAAParameters.bDownsample = (CVarTemporalAAAllowDownsampling.GetValueOnRenderThread() != 0)
 					&& !IsMotionBlurEnabled(View)
 					&& !bVisualizeMotionBlur
-					&& TAAParameters.bIsComputePass
 					&& (DownsampleQuality == 0)
 					&& TAAParameters.bUseFast;
 
@@ -1239,7 +1234,6 @@ void FPostProcessing::Process(FRHICommandListImmediate& RHICmdList, const FViewI
 					TAAParameters.Pass = ETAAPassConfig::MainSuperSampling;
 					TAAParameters.bDownsample = false;
 					TAAParameters.bUseFast = false;
-					TAAParameters.bIsComputePass = true;
 
 					TAAParameters.OutputViewRect.Min.X = 0;
 					TAAParameters.OutputViewRect.Min.Y = 0;
