Implements initial spherical harmonic denoiser for indirect diffuse.

diff --git a/Engine/Shaders/Private/ScreenSpaceDenoise/SSDCommon.ush b/Engine/Shaders/Private/ScreenSpaceDenoise/SSDCommon.ush
--- a/Engine/Shaders/Private/ScreenSpaceDenoise/SSDCommon.ush
+++ b/Engine/Shaders/Private/ScreenSpaceDenoise/SSDCommon.ush
@@ -36,6 +36,9 @@
 	// Indirect diffuse requires really large spatial kernel to clean the noise.
 	#define TARGETED_SAMPLE_COUNT 4096
 
+#elif CONFIG_SIGNAL_PROCESSING == SIGNAL_PROCESSING_DIFFUSE_SPHERICAL_HARMONIC
+	#define TARGETED_SAMPLE_COUNT 128
+
 #else
 	#error Unconfigured TARGETED_SAMPLE_COUNT
 
diff --git a/Engine/Shaders/Private/ScreenSpaceDenoise/SSDDefinitions.ush b/Engine/Shaders/Private/ScreenSpaceDenoise/SSDDefinitions.ush
--- a/Engine/Shaders/Private/ScreenSpaceDenoise/SSDDefinitions.ush
+++ b/Engine/Shaders/Private/ScreenSpaceDenoise/SSDDefinitions.ush
@@ -47,7 +47,7 @@
 
 // World bluring radius for a miss.
 #define WORLD_RADIUS_INVALID -1
-#define WORLD_RADIUS_MISS 1.#INF
+#define WORLD_RADIUS_MISS (INFINITE_FLOAT)
 
 // Mip level for sample that are invalid or missed.
 #define MIP_LEVEL_INVALID -1
@@ -124,6 +124,10 @@
 	#define SIGNAL_BUFFER_LAYOUT_DIFFUSE_INDIRECT_AND_AO_HISTORY 0xD102
 
 
+	/** Buffer layout given by the ray generation shader for diffuse indirect  illumination. */
+	#define SIGNAL_BUFFER_LAYOUT_DIFFUSE_INDIRECT_HARMONIC 0xD200
+
+
 /** Defines how the signal is being processed. Matches C++'s ESignalProcessing. */
 	/** Shadow penumbra denoising. */
 	#define SIGNAL_PROCESSING_MONOCHROMATIC_PENUMBRA 0
@@ -134,9 +138,12 @@
 	/** Ambient occlusion denoising. */
 	#define SIGNAL_PROCESSING_AO 2
 
-	/** Global illumination denoising. */
+	/** Diffuse indirect and ambient occlusion. */
 	#define SIGNAL_PROCESSING_DIFFUSE_INDIRECT_AND_AO 3
 
+	/** Diffuse indirect using spherical harmonic. */
+	#define SIGNAL_PROCESSING_DIFFUSE_SPHERICAL_HARMONIC 4
+
 
 
 /** Defines the color space bitfield. */
diff --git a/Engine/Shaders/Private/ScreenSpaceDenoise/SSDSignalCore.ush b/Engine/Shaders/Private/ScreenSpaceDenoise/SSDSignalCore.ush
--- a/Engine/Shaders/Private/ScreenSpaceDenoise/SSDSignalCore.ush
+++ b/Engine/Shaders/Private/ScreenSpaceDenoise/SSDSignalCore.ush
@@ -363,7 +363,7 @@ void DecodeMultiplexedSignals(
 	if (SignalBufferLayout == SIGNAL_BUFFER_LAYOUT_UNINITIALIZED)
 	{
 		// trash out the output sample to identify quickly a miss configured buffer layout.
-		OutSamples = CreateSignalArrayFromScalarValue(-1.#INF);
+		OutSamples = CreateSignalArrayFromScalarValue(-INFINITE_FLOAT);
 	}
 	else if (SignalBufferLayout == SIGNAL_BUFFER_LAYOUT_PENUMBRA_INPUT_1SPP)
 	{
@@ -598,6 +598,27 @@ void DecodeMultiplexedSignals(
 		OutSamples.Array[2].SceneColor.x = OutSamples.Array[2].SampleCount * RawSample[1].b;
 	}
 	#endif
+	#if CONFIG_SIGNAL_COLOR_SH
+	else if (SignalBufferLayout == SIGNAL_BUFFER_LAYOUT_DIFFUSE_INDIRECT_HARMONIC)
+	{
+		float SampleCount;
+		if (bNormalizeSample)
+		{
+			SampleCount = 1.0;
+			OutSamples.Array[0].SampleCount = RawSample[3].x > 0.0 ? 1.0 : 0.0;
+		}
+		else
+		{
+			SampleCount = RawSample[3].x;
+			OutSamples.Array[0].SampleCount = SampleCount;
+		}
+
+		OutSamples.Array[0].MissCount = RawSample[3].y * SampleCount;
+		OutSamples.Array[0].ColorSH.R.V = RawSample[0] * SampleCount;
+		OutSamples.Array[0].ColorSH.G.V = RawSample[1] * SampleCount;
+		OutSamples.Array[0].ColorSH.B.V = RawSample[2] * SampleCount;
+	}
+	#endif
 } // DecodeMultiplexedSignals()
 
 /** Encode output signal sample. */
@@ -789,6 +810,19 @@ void EncodeMultiplexedSignals(
 		}
 	}
 	#endif
+	#if CONFIG_SIGNAL_COLOR_SH
+	else if (SignalBufferLayout == SIGNAL_BUFFER_LAYOUT_DIFFUSE_INDIRECT_HARMONIC)
+	{
+		float NormalizationFactor = SafeRcp(Sample[0].SampleCount);
+
+		OutRawSample[3].x = Sample[0].SampleCount;
+		OutRawSample[3].y = Sample[0].MissCount * NormalizationFactor;
+		OutRawSample[0] = Sample[0].ColorSH.R.V * NormalizationFactor;
+		OutRawSample[1] = Sample[0].ColorSH.G.V * NormalizationFactor;
+		OutRawSample[2] = Sample[0].ColorSH.B.V * NormalizationFactor;
+		OutBufferCount = 4;
+	}
+	#endif
 } // EncodeMultiplexedSignals()
 
 /** Sample multiple input signals that have been multiplexed. */
diff --git a/Engine/Shaders/Private/ScreenSpaceDenoise/SSDSignalFramework.ush b/Engine/Shaders/Private/ScreenSpaceDenoise/SSDSignalFramework.ush
--- a/Engine/Shaders/Private/ScreenSpaceDenoise/SSDSignalFramework.ush
+++ b/Engine/Shaders/Private/ScreenSpaceDenoise/SSDSignalFramework.ush
@@ -60,7 +60,9 @@
 	#include "../DeferredShadingCommon.ush"
 	#include "../CommonViewUniformBuffer.ush"
 
-#elif CONFIG_SIGNAL_PROCESSING == SIGNAL_PROCESSING_AO || CONFIG_SIGNAL_PROCESSING == SIGNAL_PROCESSING_DIFFUSE_INDIRECT_AND_AO
+#elif CONFIG_SIGNAL_PROCESSING == SIGNAL_PROCESSING_AO
+#elif CONFIG_SIGNAL_PROCESSING == SIGNAL_PROCESSING_DIFFUSE_INDIRECT_AND_AO
+#elif CONFIG_SIGNAL_PROCESSING == SIGNAL_PROCESSING_DIFFUSE_SPHERICAL_HARMONIC
 
 #else
 	#error Unknown CONFIG_SIGNAL_PROCESSING
@@ -673,6 +675,10 @@ float GetSignalWorldBluringRadius(FSSDSignalSample Sample, FSSDSampleSceneInfos
 {
 	return Sample.ClosestHitDistance;
 }
+#elif CONFIG_SIGNAL_PROCESSING == SIGNAL_PROCESSING_DIFFUSE_SPHERICAL_HARMONIC
+{
+	return 1; // TODO
+}
 #else
 	#error Unimplemented GetSignalWorldBluringRadius()
 #endif
diff --git a/Engine/Shaders/Private/ScreenSpaceDenoise/SSDSpatialAccumulation.usf b/Engine/Shaders/Private/ScreenSpaceDenoise/SSDSpatialAccumulation.usf
--- a/Engine/Shaders/Private/ScreenSpaceDenoise/SSDSpatialAccumulation.usf
+++ b/Engine/Shaders/Private/ScreenSpaceDenoise/SSDSpatialAccumulation.usf
@@ -231,7 +231,6 @@
 
 #elif CONFIG_SIGNAL_PROCESSING == SIGNAL_PROCESSING_DIFFUSE_INDIRECT_AND_AO
 	#define CONFIG_SIGNAL_COLOR 1
-	#define CONFIG_SIGNAL_COLOR_SH 1
 
 	#define CONFIG_VIEW_POSITION_BASED_BILATERAL 2
 	#define CONFIG_NORMAL_BASED_BILATERAL 1
@@ -271,6 +270,31 @@
 		#error Unexpected stage.
 	#endif
 
+#elif CONFIG_SIGNAL_PROCESSING == SIGNAL_PROCESSING_DIFFUSE_SPHERICAL_HARMONIC
+	// #define CONFIG_SIGNAL_COLOR 0 // TODO
+	#define CONFIG_SIGNAL_COLOR_SH 1
+
+	// Given it's a spherical harmonic that store directionality, only need position based rejection.
+	#define CONFIG_VIEW_POSITION_BASED_BILATERAL 2
+
+	#define CONFIG_INPUT_TEXTURE_COUNT 4
+	#define CONFIG_OUTPUT_TEXTURE_COUNT 4
+
+	#define CONFIG_BILATERAL_DISTANCE_COMPUTATION SIGNAL_WORLD_FREQUENCY_MIN_METADATA
+
+	// Input and output layout.
+	#define CONFIG_SIGNAL_INPUT_LAYOUT  SIGNAL_BUFFER_LAYOUT_DIFFUSE_INDIRECT_HARMONIC
+	#define CONFIG_SIGNAL_OUTPUT_LAYOUT SIGNAL_BUFFER_LAYOUT_DIFFUSE_INDIRECT_HARMONIC
+
+	#if DIM_STAGE == STAGE_RECONSTRUCTION
+		#define CONFIG_SAMPLE_SET           SAMPLE_SET_STACKOWIAK_4_SETS
+
+	#elif DIM_STAGE == STAGE_REJECTION_PRE_CONVOLUTION
+
+	#else
+		#error Unexpected stage.
+	#endif
+
 #else
 	#error Unknown signal processing.
 #endif
@@ -545,7 +569,7 @@ void MainCS(
 		KernelConfig.bForceKernelCenterAccumulation = true;
 
 		// Reconstruct the spherical harmonic when reconstructing from 1spp.
-		KernelConfig.bComputeSampleColorSH = DIM_STAGE == STAGE_RECONSTRUCTION;
+		KernelConfig.bComputeSampleColorSH = DIM_STAGE == STAGE_RECONSTRUCTION && DIM_MULTI_SPP == 0;
 		
 		{
 			UNROLL_N(MAX_SIGNAL_MULTIPLEXING)
diff --git a/Engine/Shaders/Private/ScreenSpaceDenoise/SSDTemporalAccumulation.usf b/Engine/Shaders/Private/ScreenSpaceDenoise/SSDTemporalAccumulation.usf
--- a/Engine/Shaders/Private/ScreenSpaceDenoise/SSDTemporalAccumulation.usf
+++ b/Engine/Shaders/Private/ScreenSpaceDenoise/SSDTemporalAccumulation.usf
@@ -135,6 +135,35 @@
 	#define CONFIG_VIEW_POSITION_BASED_BILATERAL 2
 	#define CONFIG_NORMAL_BASED_BILATERAL 1
 
+#elif CONFIG_SIGNAL_PROCESSING == SIGNAL_PROCESSING_DIFFUSE_SPHERICAL_HARMONIC
+	// #define CONFIG_SIGNAL_COLOR 0 // TODO
+	#define CONFIG_SIGNAL_COLOR_SH 1
+
+	// Input and output layout.
+	#define CONFIG_SIGNAL_INPUT_LAYOUT   SIGNAL_BUFFER_LAYOUT_DIFFUSE_INDIRECT_HARMONIC
+	#define CONFIG_SIGNAL_HISTORY_LAYOUT SIGNAL_BUFFER_LAYOUT_DIFFUSE_INDIRECT_HARMONIC
+
+	#define CONFIG_INPUT_TEXTURE_COUNT 4
+	#define CONFIG_HISTORY_TEXTURE_COUNT 4
+
+	#define CONFIG_HISTORY_REJECTION HISTORY_REJECTION_MINMAX_BOUNDARIES
+	//#define CONFIG_HISTORY_REJECTION HISTORY_REJECTION_ADAPTIVE_INTEGRATION
+	#define CONFIG_MULTIPLEXED_SIGNALS_PER_SIGNAL_DOMAIN 4
+
+	// Given it's a spherical harmonic that store directionality, only need position based rejection.
+	#define CONFIG_VIEW_POSITION_BASED_BILATERAL 2
+
+	// Use pre transformed rejection buffer that contains pre transformed momment 1 & 2.
+	#if 0
+		#define CONFIG_USE_REJECTION_BUFFER               1
+		#define CONFIG_SIGNAL_HISTORY_REJECTION_LAYOUT    SIGNAL_BUFFER_LAYOUT_DIFFUSE_INDIRECT_HARMONIC
+		
+		#define CONFIG_REJECTION_TEXTURE_COUNT 4
+
+		#define CONFIG_REJECTION_INPUT_MODE               REJECTION_INPUT_MODE_2PRETRANSFORMED_MOMMENTS
+		#define CONFIG_REJECTION_SAMPLE_SET               REJECTION_SAMPLE_SET_SOBEK2018
+	#endif
+
 #else
 	#error Unimplemented signal processing.
 #endif
diff --git a/Engine/Source/Runtime/Renderer/Private/ScreenSpaceDenoise.cpp b/Engine/Source/Runtime/Renderer/Private/ScreenSpaceDenoise.cpp
--- a/Engine/Source/Runtime/Renderer/Private/ScreenSpaceDenoise.cpp
+++ b/Engine/Source/Runtime/Renderer/Private/ScreenSpaceDenoise.cpp
@@ -147,6 +147,9 @@ enum class ESignalProcessing
 	// Denoise first bounce diffuse and ambient occlusion.
 	DiffuseAndAmbientOcclusion,
 
+	// Denoise first bounce diffuse as sperical harmonic
+	DiffuseSphericalHarmonic,
+
 	MAX,
 };
 
@@ -165,7 +168,8 @@ static bool UsesConstantPixelDensityPassLayout(ESignalProcessing SignalProcessin
 		SignalProcessing == ESignalProcessing::MonochromaticPenumbra ||
 		SignalProcessing == ESignalProcessing::Reflections ||
 		SignalProcessing == ESignalProcessing::AmbientOcclusion ||
-		SignalProcessing == ESignalProcessing::DiffuseAndAmbientOcclusion);
+		SignalProcessing == ESignalProcessing::DiffuseAndAmbientOcclusion ||
+		SignalProcessing == ESignalProcessing::DiffuseSphericalHarmonic);
 }
 
 /** Returns whether a signal processing uses an injestion pass. */
@@ -191,6 +195,16 @@ static bool SignalUsesRejectionPreConvolution(ESignalProcessing SignalProcessing
 		SignalProcessing == ESignalProcessing::AmbientOcclusion);
 }
 
+/** Returns whether a signal processing uses a convolution pass after temporal accumulation pass. */
+static bool SignalUsesPostConvolution(ESignalProcessing SignalProcessing)
+{
+	return (
+		SignalProcessing == ESignalProcessing::MonochromaticPenumbra ||
+		SignalProcessing == ESignalProcessing::Reflections ||
+		SignalProcessing == ESignalProcessing::AmbientOcclusion ||
+		SignalProcessing == ESignalProcessing::DiffuseAndAmbientOcclusion);
+}
+
 /** Returns whether a signal processing uses a history rejection pre convolution pass. */
 static bool SignalUsesFinalConvolution(ESignalProcessing SignalProcessing)
 {
@@ -207,7 +221,8 @@ static int32 SignalMaxBatchSize(ESignalProcessing SignalProcessing)
 	else if (
 		SignalProcessing == ESignalProcessing::Reflections ||
 		SignalProcessing == ESignalProcessing::AmbientOcclusion ||
-		SignalProcessing == ESignalProcessing::DiffuseAndAmbientOcclusion)
+		SignalProcessing == ESignalProcessing::DiffuseAndAmbientOcclusion ||
+		SignalProcessing == ESignalProcessing::DiffuseSphericalHarmonic)
 	{
 		return 1;
 	}
@@ -231,7 +246,8 @@ static bool SignalSupportMultiSPP(ESignalProcessing SignalProcessing)
 		SignalProcessing == ESignalProcessing::MonochromaticPenumbra ||
 		SignalProcessing == ESignalProcessing::Reflections ||
 		SignalProcessing == ESignalProcessing::AmbientOcclusion ||
-		SignalProcessing == ESignalProcessing::DiffuseAndAmbientOcclusion);
+		SignalProcessing == ESignalProcessing::DiffuseAndAmbientOcclusion ||
+		SignalProcessing == ESignalProcessing::DiffuseSphericalHarmonic);
 }
 
 
@@ -271,6 +287,12 @@ const TCHAR* const kInjestResourceNames[] = {
 	nullptr,
 	nullptr,
 	nullptr,
+
+	// DiffuseSphericalHarmonic
+	nullptr,
+	nullptr,
+	nullptr,
+	nullptr,
 };
 
 const TCHAR* const kReconstructionResourceNames[] = {
@@ -297,6 +319,12 @@ const TCHAR* const kReconstructionResourceNames[] = {
 	TEXT("DiffuseIndirectReconstruction1"),
 	nullptr,
 	nullptr,
+
+	// DiffuseSphericalHarmonic
+	TEXT("DiffuseHarmonicReconstruction0"),
+	TEXT("DiffuseHarmonicReconstruction1"),
+	TEXT("DiffuseHarmonicReconstruction2"),
+	TEXT("DiffuseHarmonicReconstruction3"),
 };
 
 const TCHAR* const kPreConvolutionResourceNames[] = {
@@ -323,6 +351,12 @@ const TCHAR* const kPreConvolutionResourceNames[] = {
 	TEXT("DiffuseIndirectPreConvolution1"),
 	nullptr,
 	nullptr,
+
+	// DiffuseSphericalHarmonic
+	nullptr,
+	nullptr,
+	nullptr,
+	nullptr,
 };
 
 const TCHAR* const kRejectionPreConvolutionResourceNames[] = {
@@ -349,6 +383,12 @@ const TCHAR* const kRejectionPreConvolutionResourceNames[] = {
 	nullptr,
 	nullptr,
 	nullptr,
+
+	// DiffuseSphericalHarmonic
+	nullptr,
+	nullptr,
+	nullptr,
+	nullptr,
 };
 
 const TCHAR* const kTemporalAccumulationResourceNames[] = {
@@ -375,6 +415,12 @@ const TCHAR* const kTemporalAccumulationResourceNames[] = {
 	TEXT("DiffuseIndirectTemporalAccumulation1"),
 	nullptr,
 	nullptr,
+
+	// DiffuseSphericalHarmonic
+	TEXT("DiffuseHarmonicTemporalAccumulation0"),
+	TEXT("DiffuseHarmonicTemporalAccumulation1"),
+	TEXT("DiffuseHarmonicTemporalAccumulation2"),
+	TEXT("DiffuseHarmonicTemporalAccumulation3"),
 };
 
 const TCHAR* const kHistoryConvolutionResourceNames[] = {
@@ -401,6 +447,12 @@ const TCHAR* const kHistoryConvolutionResourceNames[] = {
 	TEXT("DiffuseIndirectHistoryConvolution1"),
 	nullptr,
 	nullptr,
+
+	// DiffuseSphericalHarmonic
+	nullptr,
+	nullptr,
+	nullptr,
+	nullptr,
 };
 
 const TCHAR* const kDenoiserOutputResourceNames[] = {
@@ -427,6 +479,12 @@ const TCHAR* const kDenoiserOutputResourceNames[] = {
 	nullptr,
 	nullptr,
 	nullptr,
+
+	// DiffuseSphericalHarmonic
+	nullptr,
+	nullptr,
+	nullptr,
+	nullptr,
 };
 
 static_assert(ARRAY_COUNT(kReconstructionResourceNames) == int32(ESignalProcessing::MAX) * kMaxBufferProcessingCount, "You forgot me!");
@@ -439,10 +497,11 @@ static_assert(ARRAY_COUNT(kDenoiserOutputResourceNames) == int32(ESignalProcessi
 /** Returns whether should compole pipeline for a given shader platform.*/
 bool ShouldCompileSignalPipeline(ESignalProcessing SignalProcessing, EShaderPlatform Platform)
 {
-	if (SignalProcessing == ESignalProcessing::Reflections)
+	if (SignalProcessing == ESignalProcessing::Reflections ||
+		SignalProcessing == ESignalProcessing::DiffuseSphericalHarmonic)
 	{
 		// Ray traced reflection and SSR.
-		return Platform == SP_PCD3D_SM5 || Platform == SP_XBOXONE_D3D12;
+		return Platform == SP_PCD3D_SM5 || Platform == SP_XBOXONE_D3D12 || Platform == SP_PS4;
 	}
 	else if (
 		SignalProcessing == ESignalProcessing::MonochromaticPenumbra ||
@@ -628,6 +687,13 @@ class FSSDSpatialAccumulationCS : public FGlobalShader
 			return false;
 		}
 
+		// Only compile post convolution for signal that uses it.
+		if (!SignalUsesPostConvolution(SignalProcessing) &&
+			PermutationVector.Get<FStageDim>() == EStage::PostFiltering)
+		{
+			return false;
+		}
+
 		// Only compile final convolution for signal that uses it.
 		if (!SignalUsesFinalConvolution(SignalProcessing) &&
 			PermutationVector.Get<FStageDim>() == EStage::FinalOutput)
@@ -873,6 +939,21 @@ static void DenoiseSignalAtConstantPixelDensity(
 			HistoryTextureCountPerSignal = 2;
 			bHasReconstructionLayoutDifferentFromHistory = false;
 		}
+		else if (Settings.SignalProcessing == ESignalProcessing::DiffuseSphericalHarmonic)
+		{
+			for (int32 i = 0; i < 3; i++)
+			{
+				ReconstructionDescs[i].Format = PF_FloatRGBA;
+				HistoryDescs[i].Format = PF_FloatRGBA;
+			}
+
+			ReconstructionDescs[3].Format = PF_G16R16F;
+			HistoryDescs[3].Format = PF_G16R16F;
+
+			ReconstructionTextureCount = IScreenSpaceDenoiser::kSphericalHarmonicTextureCount;
+			HistoryTextureCountPerSignal = IScreenSpaceDenoiser::kSphericalHarmonicTextureCount; // TODO: only 3 textures for history
+			bHasReconstructionLayoutDifferentFromHistory = false;
+		}
 		else
 		{
 			check(0);
@@ -1591,6 +1672,45 @@ public:
 		return GlobalIlluminationOutputs;
 	}
 
+	FDiffuseIndirectHarmonic DenoiseDiffuseIndirectHarmonic(
+		FRDGBuilder& GraphBuilder,
+		const FViewInfo& View,
+		FPreviousViewInfo* PreviousViewInfos,
+		const FSceneTextureParameters& SceneTextures,
+		const FDiffuseIndirectHarmonic& Inputs,
+		const FAmbientOcclusionRayTracingConfig Config) const override
+	{
+		RDG_GPU_STAT_SCOPE(GraphBuilder, DiffuseIndirectDenoiser);
+
+		FSSDSignalTextures InputSignal;
+		for (int32 i = 0; i < IScreenSpaceDenoiser::kSphericalHarmonicTextureCount; i++)
+			InputSignal.Textures[i] = Inputs.SphericalHarmonic[i];
+
+		FSSDConstantPixelDensitySettings Settings;
+		Settings.SignalProcessing = ESignalProcessing::DiffuseSphericalHarmonic;
+		Settings.InputResolutionFraction = Config.ResolutionFraction;
+		Settings.ReconstructionSamples = CVarGIReconstructionSampleCount.GetValueOnRenderThread();
+		Settings.bUseTemporalAccumulation = CVarGITemporalAccumulation.GetValueOnRenderThread() != 0;
+		Settings.MaxInputSPP = Config.RayCountPerPixel;
+
+		TStaticArray<FScreenSpaceFilteringHistory*, IScreenSpaceDenoiser::kMaxBatchSize> PrevHistories;
+		TStaticArray<FScreenSpaceFilteringHistory*, IScreenSpaceDenoiser::kMaxBatchSize> NewHistories;
+		PrevHistories[0] = &PreviousViewInfos->DiffuseIndirectHistory;
+		NewHistories[0] = View.ViewState ? &View.ViewState->PrevFrameViewInfo.DiffuseIndirectHistory : nullptr;
+
+		FSSDSignalTextures SignalOutput;
+		DenoiseSignalAtConstantPixelDensity(
+			GraphBuilder, View, SceneTextures,
+			InputSignal, Settings,
+			PrevHistories,
+			NewHistories,
+			&SignalOutput);
+
+		FDiffuseIndirectHarmonic GlobalIlluminationOutputs;
+		for (int32 i = 0; i < IScreenSpaceDenoiser::kSphericalHarmonicTextureCount; i++)
+			GlobalIlluminationOutputs.SphericalHarmonic[i] = SignalOutput.Textures[i];
+		return GlobalIlluminationOutputs;
+	}
 }; // class FDefaultScreenSpaceDenoiser
 
 
diff --git a/Engine/Source/Runtime/Renderer/Private/ScreenSpaceDenoise.h b/Engine/Source/Runtime/Renderer/Private/ScreenSpaceDenoise.h
--- a/Engine/Source/Runtime/Renderer/Private/ScreenSpaceDenoise.h
+++ b/Engine/Source/Runtime/Renderer/Private/ScreenSpaceDenoise.h
@@ -177,7 +177,7 @@ public:
 		SHADER_PARAMETER_RDG_TEXTURE(Texture2D, AmbientOcclusionMask)
 	END_SHADER_PARAMETER_STRUCT()
 
-	/** Entry point to denoise reflections. */
+	/** Entry point to denoise diffuse indirect and AO. */
 	virtual FDiffuseIndirectOutputs DenoiseDiffuseIndirect(
 		FRDGBuilder& GraphBuilder,
 		const FViewInfo& View,
@@ -186,7 +186,7 @@ public:
 		const FDiffuseIndirectInputs& Inputs,
 		const FAmbientOcclusionRayTracingConfig Config) const = 0;
 
-	/** Entry point to denoise SkyLight. */
+	/** Entry point to denoise SkyLight diffuse indirect. */
 	virtual FDiffuseIndirectOutputs DenoiseSkyLight(
 		FRDGBuilder& GraphBuilder,
 		const FViewInfo& View,
@@ -195,9 +195,28 @@ public:
 		const FDiffuseIndirectInputs& Inputs,
 		const FAmbientOcclusionRayTracingConfig Config) const = 0;
 
+	// Number of texture to store spherical harmonics.
+	static constexpr int32 kSphericalHarmonicTextureCount = 4;
+
+	/** All the inputs and outputs for spherical harmonic denoising. */
+	BEGIN_SHADER_PARAMETER_STRUCT(FDiffuseIndirectHarmonic, )
+		// FloatR11G11B10
+		SHADER_PARAMETER_RDG_TEXTURE_ARRAY(Texture2D, SphericalHarmonic, [kSphericalHarmonicTextureCount])
+	END_SHADER_PARAMETER_STRUCT()
+
+	/** Entry point to denoise spherical harmonic for diffuse indirect. */
+	virtual FDiffuseIndirectHarmonic DenoiseDiffuseIndirectHarmonic(
+		FRDGBuilder& GraphBuilder,
+		const FViewInfo& View,
+		FPreviousViewInfo* PreviousViewInfos,
+		const FSceneTextureParameters& SceneTextures,
+		const FDiffuseIndirectHarmonic& Inputs,
+		const FAmbientOcclusionRayTracingConfig Config) const = 0;
+
+
 	/** Returns the interface of the default denoiser of the renderer. */
 	static const IScreenSpaceDenoiser* GetDefaultDenoiser();
-}; // class FScreenSpaceDenoisingInterface
+}; // class IScreenSpaceDenoiser
 
 
 // The interface for the renderer to denoise what it needs, Plugins can come and point this to custom interface.
