Reconstructs with mip chain to reduce number of samples.

diff --git a/Engine/Shaders/Private/ScreenSpaceDenoise/SSDDefinitions.ush b/Engine/Shaders/Private/ScreenSpaceDenoise/SSDDefinitions.ush
--- a/Engine/Shaders/Private/ScreenSpaceDenoise/SSDDefinitions.ush
+++ b/Engine/Shaders/Private/ScreenSpaceDenoise/SSDDefinitions.ush
@@ -190,6 +190,8 @@
 
 	#define SAMPLE_SET_HEXAWEB 11
 
+	#define SAMPLE_SET_STOCASTIC_HIERARCHY 12
+
 
 /** Compression of signal for VGPR occupency. */
 	/** The signal is not being compressed. */
diff --git a/Engine/Shaders/Private/ScreenSpaceDenoise/SSDReduce.usf b/Engine/Shaders/Private/ScreenSpaceDenoise/SSDReduce.usf
new file mode 100644
--- /dev/null
+++ b/Engine/Shaders/Private/ScreenSpaceDenoise/SSDReduce.usf
@@ -0,0 +1,237 @@
+// Copyright 1998-2019 Epic Games, Inc. All Rights Reserved.
+
+#include "SSDDefinitions.ush"
+
+
+//------------------------------------------------------- CONFIGS
+
+#define TILE_PIXEL_SIZE 8
+
+#define CONFIG_SIGNAL_PROCESSING DIM_SIGNAL_PROCESSING
+#define CONFIG_SIGNAL_BATCH_SIZE DIM_SIGNAL_BATCH_SIZE
+
+#define CONFIG_MIP_LEVEL_COUNT 2
+
+
+// Configures all the pass for each individual signals.
+#if CONFIG_SIGNAL_PROCESSING == SIGNAL_PROCESSING_DIFFUSE_SPHERICAL_HARMONIC
+	// #define CONFIG_SIGNAL_COLOR 0 // TODO
+	#define CONFIG_SIGNAL_COLOR_SH 1
+
+	// Given it's a spherical harmonic that store directionality, only need position based rejection.
+	#define CONFIG_VIEW_POSITION_BASED_BILATERAL 2
+
+	#define CONFIG_INPUT_TEXTURE_COUNT 4
+	#define CONFIG_OUTPUT_TEXTURE_COUNT 4
+
+	#define CONFIG_BILATERAL_DISTANCE_COMPUTATION SIGNAL_WORLD_FREQUENCY_MIN_METADATA
+
+	// Input and output layout.
+	#define CONFIG_SIGNAL_INPUT_LAYOUT  SIGNAL_BUFFER_LAYOUT_DIFFUSE_INDIRECT_HARMONIC
+	#define CONFIG_SIGNAL_OUTPUT_LAYOUT SIGNAL_BUFFER_LAYOUT_DIFFUSE_INDIRECT_HARMONIC
+	#define CONFIG_SIGNAL_TEXTURE_TYPE   SIGNAL_TEXTURE_TYPE_UINT2
+
+	// Spherical harmonics are a lot of data, need to shrink VGPR pressure to improve lattency hidding when fetching the buffer.
+	// TODO: some shader compiler completly falls apparts with the current implementation of
+	//       CONFIG_SIGNAL_VGPR_COMPRESSION and actually drops in occupency.
+	//#define CONFIG_SIGNAL_VGPR_COMPRESSION SIGNAL_COMPRESSION_DIFFUSE_INDIRECT_HARMONIC
+	//#define CONFIG_NEIGHBOR_TO_REF_COMPUTATION NEIGHBOR_TO_REF_LOWEST_VGPR_PRESSURE
+
+#else
+	#error Unknown signal processing.
+#endif
+
+
+//------------------------------------------------------- CONFIG DISABLED DEFAULTS
+
+/** Selects how the world distance should be computed for bilateral rejection. */
+#ifndef CONFIG_BILATERAL_DISTANCE_COMPUTATION
+	#define CONFIG_BILATERAL_DISTANCE_COMPUTATION SIGNAL_WORLD_FREQUENCY_MIN_METADATA
+#endif
+
+
+//------------------------------------------------------- INCLUDES
+
+#include "SSDSignalFramework.ush"
+#include "SSDSignalArray.ush"
+#include "SSDSpatialKernel.ush"
+
+
+//------------------------------------------------------- PARAMETERS
+
+uint MaxSampleCount;
+uint UpscaleFactor;
+
+#if !CONFIG_UPSCALE && CONFIG_CUSTOM_SPREAD_FACTOR
+	float KernelSpreadFactor;
+#endif
+
+
+#if !defined(CONFIG_INPUT_TEXTURE_COUNT)
+	#error Missing CONFIG_INPUT_TEXTURE_COUNT
+#endif
+
+FSSDTexture2D SignalInput_Textures_0;
+
+#if CONFIG_INPUT_TEXTURE_COUNT > 1
+FSSDTexture2D SignalInput_Textures_1;
+#else
+#define SignalInput_Textures_1 SignalInput_Textures_0
+#endif
+
+#if CONFIG_INPUT_TEXTURE_COUNT > 2
+FSSDTexture2D SignalInput_Textures_2;
+#else
+#define SignalInput_Textures_2 SignalInput_Textures_0
+#endif
+
+#if CONFIG_INPUT_TEXTURE_COUNT > 3
+FSSDTexture2D SignalInput_Textures_3;
+#else
+#define SignalInput_Textures_3 SignalInput_Textures_0
+#endif
+
+
+#if !defined(CONFIG_OUTPUT_TEXTURE_COUNT)
+	#error Missing CONFIG_OUTPUT_TEXTURE_COUNT
+#endif
+
+FSSDRWTexture2D	SignalOutputMips_0_UAVs_0;
+FSSDRWTexture2D	SignalOutputMips_1_UAVs_0;
+FSSDRWTexture2D	SignalOutputMips_2_UAVs_0;
+
+#if CONFIG_OUTPUT_TEXTURE_COUNT > 1
+FSSDRWTexture2D	SignalOutputMips_0_UAVs_1;
+FSSDRWTexture2D	SignalOutputMips_1_UAVs_1;
+FSSDRWTexture2D	SignalOutputMips_2_UAVs_1;
+#else
+#define SignalOutputMips_0_UAVs_1 SignalOutputMips_0_UAVs_0
+#define SignalOutputMips_1_UAVs_1 SignalOutputMips_1_UAVs_0
+#define SignalOutputMips_2_UAVs_1 SignalOutputMips_2_UAVs_0
+#endif
+
+#if CONFIG_OUTPUT_TEXTURE_COUNT > 2
+FSSDRWTexture2D	SignalOutputMips_0_UAVs_2;
+FSSDRWTexture2D	SignalOutputMips_1_UAVs_2;
+FSSDRWTexture2D	SignalOutputMips_2_UAVs_2;
+#else
+#define SignalOutputMips_0_UAVs_2 SignalOutputMips_0_UAVs_0
+#define SignalOutputMips_1_UAVs_2 SignalOutputMips_1_UAVs_0
+#define SignalOutputMips_2_UAVs_2 SignalOutputMips_2_UAVs_0
+#endif
+
+#if CONFIG_OUTPUT_TEXTURE_COUNT > 3
+FSSDRWTexture2D	SignalOutputMips_0_UAVs_3;
+FSSDRWTexture2D	SignalOutputMips_1_UAVs_3;
+FSSDRWTexture2D	SignalOutputMips_2_UAVs_3;
+#else
+#define SignalOutputMips_0_UAVs_3 SignalOutputMips_0_UAVs_0
+#define SignalOutputMips_1_UAVs_3 SignalOutputMips_1_UAVs_0
+#define SignalOutputMips_2_UAVs_3 SignalOutputMips_2_UAVs_0
+#endif
+
+RWTexture2D<float4>	DebugOutput;
+
+
+//------------------------------------------------------- FUNCTIONS
+
+
+
+
+//------------------------------------------------------- ENTRY POINTS
+
+[numthreads(TILE_PIXEL_SIZE, TILE_PIXEL_SIZE, 1)]
+void MainCS(
+	uint2 DispatchThreadId : SV_DispatchThreadID,
+	uint2 GroupId : SV_GroupID,
+	uint2 GroupThreadId : SV_GroupThreadID,
+	uint GroupThreadIndex : SV_GroupIndex)
+{
+	// Find out scene buffer UV.
+	float2 SceneBufferUV = (DispatchThreadId + 0.5 + View.ViewRectMin.xy) * View.BufferSizeAndInvSize.zw;
+	if (true)
+	{
+		SceneBufferUV = clamp(SceneBufferUV, View.BufferBilinearUVMinMax.xy, View.BufferBilinearUVMinMax.zw);
+	}
+	
+	// Read reference meta data.
+	float2 ScreenPosition = ViewportUVToScreenPos(BufferUVToViewportUV(SceneBufferUV));
+	FSSDSampleSceneInfos RefSceneMetadata = FetchCurrentSceneInfos(ScreenPosition, SceneBufferUV);
+
+	// Sample the reference sample.
+	FSSDSignalArray RefSamples = SampleMultiplexedSignals(
+		SignalInput_Textures_0,
+		SignalInput_Textures_1,
+		SignalInput_Textures_2,
+		SignalInput_Textures_3,
+		GlobalPointClampedSampler,
+		CONFIG_SIGNAL_INPUT_LAYOUT,
+		/* MultiplexedSampleId = */ 0,
+		/* bNormalizeSample = */ false,
+		SceneBufferUV);
+	
+	const uint MultiplexCount = CONFIG_SIGNAL_BATCH_SIZE;
+
+	// Output the multiplexed signal in mip 0.
+	{
+		// Whitelists computation that needs to be done.
+		FSSDSignalArray OutputSamples = CreateSignalArrayFromScalarValue(0.0);
+		UNROLL
+		for (uint MultiplexId = 0; MultiplexId < MultiplexCount; MultiplexId++)
+		{
+			OutputSamples.Array[MultiplexId] = RefSamples.Array[MultiplexId];
+		}
+
+		// TODO: TODO: reduce VGPR pressure here
+		uint2 OutputPixelPostion = DispatchThreadId;
+
+		OutputMultiplexedSignal(
+			SignalOutputMips_0_UAVs_0,
+			SignalOutputMips_0_UAVs_1,
+			SignalOutputMips_0_UAVs_2,
+			SignalOutputMips_0_UAVs_3,
+			CONFIG_SIGNAL_OUTPUT_LAYOUT, MultiplexCount,
+			OutputPixelPostion, OutputSamples.Array);
+	}
+
+	UNROLL_N(2)
+	for (uint MipLevel = 1; MipLevel < CONFIG_MIP_LEVEL_COUNT; MipLevel++)
+	{
+		// Whitelists computation that needs to be done.
+		FSSDSignalArray OutputSamples = CreateSignalArrayFromScalarValue(0.0);
+		UNROLL
+		for (uint MultiplexId = 0; MultiplexId < MultiplexCount; MultiplexId++)
+		{
+			OutputSamples.Array[MultiplexId] = RefSamples.Array[MultiplexId];
+		}
+
+		uint2 OutputPixelPostion = DispatchThreadId >> MipLevel;
+
+		BRANCH
+		if (all((OutputPixelPostion << MipLevel) == DispatchThreadId))
+		{
+			if (MipLevel == 1)
+			{
+				OutputMultiplexedSignal(
+					SignalOutputMips_1_UAVs_0,
+					SignalOutputMips_1_UAVs_1,
+					SignalOutputMips_1_UAVs_2,
+					SignalOutputMips_1_UAVs_3,
+					CONFIG_SIGNAL_OUTPUT_LAYOUT, MultiplexCount,
+					OutputPixelPostion, OutputSamples.Array);
+			}
+			#if CONFIG_MIP_LEVEL_COUNT >= 3
+			else if (MipLevel == 2)
+			{
+				OutputMultiplexedSignal(
+					SignalOutputMips_2_UAVs_0,
+					SignalOutputMips_2_UAVs_1,
+					SignalOutputMips_2_UAVs_2,
+					SignalOutputMips_2_UAVs_3,
+					CONFIG_SIGNAL_OUTPUT_LAYOUT, MultiplexCount,
+					OutputPixelPostion, OutputSamples.Array);
+			}
+			#endif
+		} // if ((OutputPixelPostion << MipLevel) ==  DispatchThreadId)
+	} // for (uint MipLevel = 1; MipLevel < 3; MipLevel++)
+} // MainCS
diff --git a/Engine/Shaders/Private/ScreenSpaceDenoise/SSDSpatialAccumulation.usf b/Engine/Shaders/Private/ScreenSpaceDenoise/SSDSpatialAccumulation.usf
--- a/Engine/Shaders/Private/ScreenSpaceDenoise/SSDSpatialAccumulation.usf
+++ b/Engine/Shaders/Private/ScreenSpaceDenoise/SSDSpatialAccumulation.usf
@@ -294,7 +294,8 @@
 	#define CONFIG_NEIGHBOR_TO_REF_COMPUTATION NEIGHBOR_TO_REF_LOWEST_VGPR_PRESSURE
 
 	#if DIM_STAGE == STAGE_RECONSTRUCTION
-		#define CONFIG_SAMPLE_SET           SAMPLE_SET_STACKOWIAK_4_SETS
+		#define CONFIG_SAMPLE_SET           SAMPLE_SET_STOCASTIC_HIERARCHY
+		//#define CONFIG_SAMPLE_SET           SAMPLE_SET_STACKOWIAK_4_SETS
 
 	#elif DIM_STAGE == STAGE_REJECTION_PRE_CONVOLUTION
 
@@ -569,7 +570,7 @@ void MainCS(
 	#endif
 
 	// Register renaming of members of FSSDKernelConfig to survive until the output to UAV
-	#if CONFIG_SAMPLE_SET == SAMPLE_SET_STACKOWIAK_4_SETS && CONFIG_VGPR_OPTIMIZATION
+	#if (CONFIG_SAMPLE_SET == SAMPLE_SET_STACKOWIAK_4_SETS || CONFIG_SAMPLE_SET == SAMPLE_SET_STOCASTIC_HIERARCHY) && CONFIG_VGPR_OPTIMIZATION
 		float2 KernelBufferUV;
 		uint SampleTrackId;
 	#endif
@@ -733,7 +734,27 @@ void MainCS(
 				KernelBufferUV = KernelConfig.BufferUV;
 			#endif
 		}
-		#endif // CONFIG_SAMPLE_SET == SAMPLE_SET_STACKOWIAK_4_SETS
+		#elif CONFIG_SAMPLE_SET == SAMPLE_SET_STOCASTIC_HIERARCHY
+		{
+			#if CONFIG_UPSCALE
+				#error Unsupported.
+			#endif
+
+			// Put the kernel center at the center of the quad. Half pixel shift is done in the sample offsets.
+			KernelConfig.BufferUV = (DispatchThreadId | 1) * View.BufferSizeAndInvSize.zw;
+
+			// Id of the pixel in the quad. This is to match hard coded first samples of the sample set.
+			KernelConfig.SampleTrackId = ((DispatchThreadId.x & 1) | ((DispatchThreadId.y & 1) << 1));
+
+			#if CONFIG_VGPR_OPTIMIZATION
+				// Keek sample SampleTrackId & SceneBufferUV arround for computation of pixel output coordinate.
+				// Should be VGPR free given it's curernt is being used in accumulation has well that is highest VGPR pressure of the shader.
+				// TODO: could save 1 VGPR by using 2 SGPR instead of SampleTrackId.
+				SampleTrackId = KernelConfig.SampleTrackId;
+				KernelBufferUV = KernelConfig.BufferUV;
+			#endif
+		}
+		#endif // CONFIG_SAMPLE_SET == SAMPLE_SET_STOCASTIC_HIERARCHY
 
 		AccumulateKernel(
 			KernelConfig,
@@ -906,7 +927,7 @@ void MainCS(
 	#if CONFIG_VGPR_OPTIMIZATION
 	{
 		// No need to keep DispatchThreadId, can recompute the output pixel position based on information stored in VGPRs for spatial kernel.
-		#if CONFIG_SAMPLE_SET == SAMPLE_SET_STACKOWIAK_4_SETS
+		#if CONFIG_SAMPLE_SET == SAMPLE_SET_STACKOWIAK_4_SETS || CONFIG_SAMPLE_SET == SAMPLE_SET_STOCASTIC_HIERARCHY
 			#if CONFIG_VGPR_FREE_SAMPLE_TRACK_ID
 				SampleTrackId = GetSampleTrackIdFromLaneIndex();
 			#endif
diff --git a/Engine/Shaders/Private/ScreenSpaceDenoise/SSDSpatialKernel.ush b/Engine/Shaders/Private/ScreenSpaceDenoise/SSDSpatialKernel.ush
--- a/Engine/Shaders/Private/ScreenSpaceDenoise/SSDSpatialKernel.ush
+++ b/Engine/Shaders/Private/ScreenSpaceDenoise/SSDSpatialKernel.ush
@@ -342,6 +342,7 @@ void SampleMultiplexedSignals(
 	FSSDTexture2D SignalBuffer2,
 	FSSDTexture2D SignalBuffer3,
 	float2 SampleBufferUV,
+	float MipLevel,
 	out FSSDSampleSceneInfos OutSampleSceneMetadata,
 	out FSSDSignalArray OutMultiplexedSamples)
 {
@@ -370,7 +371,7 @@ void SampleMultiplexedSignals(
 		/* MultiplexedSampleId = */ 0,
 		KernelConfig.bNormalizeSample,
 		SampleBufferUV,
-		KernelConfig.BufferMipLevel);
+		KernelConfig.BufferMipLevel + MipLevel);
 } // SampleMultiplexedSignals()
 
 /** Accumulate multiplexed samples and their metadata to an accumulator. */
@@ -494,6 +495,7 @@ void SampleAndAccumulateMultiplexedSignals(
 	FSSDTexture2D SignalBuffer3,
 	inout FSSDSignalAccumulatorArray Accumulators,
 	float2 SampleBufferUV,
+	float MipLevel,
 	float KernelSampleWeight,
 	const bool bForceSample)
 {
@@ -510,6 +512,7 @@ void SampleAndAccumulateMultiplexedSignals(
 		SignalBuffer2,
 		SignalBuffer3,
 		SampleBufferUV,
+		MipLevel,
 		/* out */ SampleSceneMetadata,
 		/* out */ MultiplexedSamples);
 
@@ -551,6 +554,7 @@ void SampleAndAccumulateMultiplexedSignalsPair(
 			SignalBuffer2,
 			SignalBuffer3,
 			SampleBufferUV[PairFetchId],
+			/* MipLevel = */ 0.0,
 			/* out */ SampleSceneMetadata[PairFetchId],
 			/* out */ MultiplexedSamples[PairFetchId]);
 	}
@@ -638,6 +642,7 @@ void SampleAndAccumulateCenterSampleAsItsOwnCluster(
 		SignalBuffer3,
 		/* inout */ Accumulators,
 		KernelConfig.BufferUV,
+		/* MipLevel = */ 0.0,
 		/* KernelSampleWeight = */ 1.0,
 		/* bForceSample = */ KernelConfig.bForceKernelCenterAccumulation);
 
@@ -682,6 +687,7 @@ void AccumulateBilinear(
 			SignalBuffer3,
 			/* inout */ Accumulators,
 			SampleBufferUV,
+			/* MipLevel = */ 0.0,
 			BilinearWeight,
 			/* bForceSample = */ false);
 	}
@@ -742,6 +748,7 @@ void AccumulateSquareKernel(
 				SignalBuffer3,
 				/* inout */ Accumulators,
 				SampleBufferUV,
+				/* MipLevel = */ 0.0,
 				KernelWeight,
 				/* bForceSample = */ bIsKernelCenterSample && KernelConfig.bForceKernelCenterAccumulation);
 		}
@@ -784,6 +791,7 @@ void AccumulateSquareKernel(
 				SignalBuffer3,
 				/* inout */ Accumulators,
 				SampleBufferUV,
+				/* MipLevel = */ 0.0,
 				KernelWeight,
 				/* bForceSample = */ bIsKernelCenterSample && KernelConfig.bForceKernelCenterAccumulation);
 		}
@@ -993,6 +1001,7 @@ void ConvolveStackowiakKernel(
 					SignalBuffer3,
 					/* inout */ Accumulators,
 					SampleBufferUV,
+					/* MipLevel = */ 0.0,
 					KernelWeight,
 					/* bForceSample = */ bIsKernelCenterSample && KernelConfig.bForceKernelCenterAccumulation);
 
@@ -1061,6 +1070,7 @@ void ConvolveStackowiakKernel(
 				SignalBuffer3,
 				/* inout */ Accumulators,
 				SampleBufferUV,
+				/* MipLevel = */ 0.0,
 				KernelWeight,
 				/* bForceSample = */ bIsKernelCenterSample && KernelConfig.bForceKernelCenterAccumulation);
 		}
@@ -1295,6 +1305,83 @@ void ConvolveDiskKernel(
 }
 
 
+//------------------------------------------------------- STOCASTIC HIERARCHY
+
+void ConvolveStocasticHierarchy(
+	FSSDKernelConfig KernelConfig,
+	FSSDTexture2D SignalBuffer0,
+	FSSDTexture2D SignalBuffer1,
+	FSSDTexture2D SignalBuffer2,
+	FSSDTexture2D SignalBuffer3,
+	inout FSSDSignalAccumulatorArray Accumulators)
+{
+	// Accumulate the center of the kernel.
+	if (KernelConfig.bSampleKernelCenter && !KernelConfig.bDescOrder)
+	{
+		SampleAndAccumulateCenterSampleAsItsOwnCluster(
+			KernelConfig,
+			SignalBuffer0,
+			SignalBuffer1,
+			SignalBuffer2,
+			SignalBuffer3,
+			/* inout */ Accumulators);
+	}
+
+	UNROLL_N(1)
+	for (uint MipLevel = 1; MipLevel < 2; MipLevel++)
+	{
+		uint SampleTrackId = KernelConfig.SampleTrackId;
+		#if CONFIG_VGPR_FREE_SAMPLE_TRACK_ID
+			SampleTrackId = GetSampleTrackIdFromLaneIndex();
+		#endif
+
+		SampleTrackId ^= View.StateFrameIndex;
+
+		float2 BasisX;
+		BasisX.x = float((~SampleTrackId) & 0x1) * (SampleTrackId & 0x2 ? -1.0 : 1.0);
+		BasisX.y = float(SampleTrackId & 0x1) * (SampleTrackId & 0x2 ? -1.0 : 1.0);
+
+		float2 BasisY = float2(-BasisX.y, BasisX.x);
+
+		float2 SampleTexelOffset[2];
+		SampleTexelOffset[0] = BasisX;
+		SampleTexelOffset[1] = -BasisX + BasisY;
+		
+		float MipLevelSize = float(1 << MipLevel);
+
+		const uint SampleCount = 2;
+		UNROLL_N(2)
+		for (uint SampleId = 0; SampleId < SampleCount; SampleId++)
+		{
+			float2 SampleBufferUV = KernelConfig.BufferUV + (SampleTexelOffset[SampleId] * MipLevelSize * KernelConfig.KernelSpreadFactor) * KernelConfig.BufferSizeAndInvSize.zw;
+
+			SampleAndAccumulateMultiplexedSignals(
+				KernelConfig,
+				SignalBuffer0,
+				SignalBuffer1,
+				SignalBuffer2,
+				SignalBuffer3,
+				/* inout */ Accumulators,
+				SampleBufferUV,
+				/* MipLevel = */ float(MipLevel),
+				/* KernelSampleWeight = */ 1.0,
+				/* bForceSample = */ false);
+		}
+	}
+	
+	// Accumulate the center of the kernel.
+	if (KernelConfig.bSampleKernelCenter && KernelConfig.bDescOrder)
+	{
+		SampleAndAccumulateCenterSampleAsItsOwnCluster(
+			KernelConfig,
+			SignalBuffer0,
+			SignalBuffer1,
+			SignalBuffer2,
+			SignalBuffer3,
+			/* inout */ Accumulators);
+	}
+}
+
 //------------------------------------------------------- MAIN ENTRY POINTS
 
 void AccumulateKernel(
@@ -1364,4 +1451,14 @@ void AccumulateKernel(
 			SignalBuffer3,
 			/* inout */ Accumulators);
 	}
+	else if (KernelConfig.SampleSet == SAMPLE_SET_STOCASTIC_HIERARCHY)
+	{
+		ConvolveStocasticHierarchy(
+			KernelConfig,
+			SignalBuffer0,
+			SignalBuffer1,
+			SignalBuffer2,
+			SignalBuffer3,
+			/* inout */ Accumulators);
+	}
 } // AccumulateKernel()
diff --git a/Engine/Shaders/Private/ScreenSpaceDenoise/SSDTemporalAccumulation.usf b/Engine/Shaders/Private/ScreenSpaceDenoise/SSDTemporalAccumulation.usf
--- a/Engine/Shaders/Private/ScreenSpaceDenoise/SSDTemporalAccumulation.usf
+++ b/Engine/Shaders/Private/ScreenSpaceDenoise/SSDTemporalAccumulation.usf
@@ -154,17 +154,6 @@
 	// Given it's a spherical harmonic that store directionality, only need position based rejection.
 	#define CONFIG_VIEW_POSITION_BASED_BILATERAL 2
 
-	// Use pre transformed rejection buffer that contains pre transformed momment 1 & 2.
-	#if 0
-		#define CONFIG_USE_REJECTION_BUFFER               1
-		#define CONFIG_SIGNAL_HISTORY_REJECTION_LAYOUT    SIGNAL_BUFFER_LAYOUT_DIFFUSE_INDIRECT_HARMONIC
-		
-		#define CONFIG_REJECTION_TEXTURE_COUNT 4
-
-		#define CONFIG_REJECTION_INPUT_MODE               REJECTION_INPUT_MODE_2PRETRANSFORMED_MOMMENTS
-		#define CONFIG_REJECTION_SAMPLE_SET               REJECTION_SAMPLE_SET_SOBEK2018
-	#endif
-
 #else
 	#error Unimplemented signal processing.
 #endif
diff --git a/Engine/Source/Runtime/Renderer/Private/ScreenSpaceDenoise.cpp b/Engine/Source/Runtime/Renderer/Private/ScreenSpaceDenoise.cpp
--- a/Engine/Source/Runtime/Renderer/Private/ScreenSpaceDenoise.cpp
+++ b/Engine/Source/Runtime/Renderer/Private/ScreenSpaceDenoise.cpp
@@ -107,7 +107,8 @@ static TAutoConsoleVariable<float> CVarGIHistoryConvolutionKernelSpreadFactor(
 
 
 /** The maximum number of mip level supported in the denoiser. */
-static const int32 kMaxMipLevel = 4;
+// TODO: jump to 3 because bufefr size already have a size multiple of 4.
+static const int32 kMaxMipLevel = 2;
 
 /** Maximum number of sample per pixel supported in the stackowiak sample set. */
 static const int32 kStackowiakMaxSampleCountPerSet = 56;
@@ -172,12 +173,27 @@ static bool UsesConstantPixelDensityPassLayout(ESignalProcessing SignalProcessin
 		SignalProcessing == ESignalProcessing::DiffuseSphericalHarmonic);
 }
 
+/** Returns whether a signal processing support upscaling. */
+static bool SignalSupportsUpscaling(ESignalProcessing SignalProcessing)
+{
+	return (
+		SignalProcessing == ESignalProcessing::Reflections ||
+		SignalProcessing == ESignalProcessing::AmbientOcclusion ||
+		SignalProcessing == ESignalProcessing::DiffuseAndAmbientOcclusion);
+}
+
 /** Returns whether a signal processing uses an injestion pass. */
 static bool SignalUsesInjestion(ESignalProcessing SignalProcessing)
 {
 	return SignalProcessing == ESignalProcessing::MonochromaticPenumbra;
 }
 
+/** Returns whether a signal processing uses a reduction pass before the reconstruction. */
+static bool SignalUsesReduction(ESignalProcessing SignalProcessing)
+{
+	return SignalProcessing == ESignalProcessing::DiffuseSphericalHarmonic;
+}
+
 /** Returns whether a signal processing uses an additional pre convolution pass. */
 static bool SignalUsesPreConvolution(ESignalProcessing SignalProcessing)
 {
@@ -295,6 +311,38 @@ const TCHAR* const kInjestResourceNames[] = {
 	nullptr,
 };
 
+const TCHAR* const kReduceResourceNames[] = {
+	// Penumbra
+	nullptr,
+	nullptr,
+	nullptr,
+	nullptr,
+
+	// Reflections
+	nullptr,
+	nullptr,
+	nullptr,
+	nullptr,
+
+	// AmbientOcclusion
+	nullptr,
+	nullptr,
+	nullptr,
+	nullptr,
+
+	// DiffuseIndirect
+	nullptr,
+	nullptr,
+	nullptr,
+	nullptr,
+
+	// DiffuseSphericalHarmonic
+	TEXT("DiffuseHarmonicReduce0"),
+	TEXT("DiffuseHarmonicReduce1"),
+	TEXT("DiffuseHarmonicReduce2"),
+	TEXT("DiffuseHarmonicReduce3"),
+};
+
 const TCHAR* const kReconstructionResourceNames[] = {
 	// Penumbra
 	TEXT("ShadowReconstruction0"),
@@ -563,18 +611,19 @@ FSSDSignalTextures CreateMultiplexedTextures(
 	return SignalTextures;
 }
 
-FSSDSignalUAVs CreateMultiplexedUAVs(FRDGBuilder& GraphBuilder, const FSSDSignalTextures& SignalTextures)
+FSSDSignalUAVs CreateMultiplexedUAVs(FRDGBuilder& GraphBuilder, const FSSDSignalTextures& SignalTextures, int32 MipLevel = 0)
 {
 	FSSDSignalUAVs UAVs;
 	for (int32 i = 0; i < kMaxBufferProcessingCount; i++)
 	{
 		if (SignalTextures.Textures[i])
-			UAVs.UAVs[i] = GraphBuilder.CreateUAV(SignalTextures.Textures[i]);
+			UAVs.UAVs[i] = GraphBuilder.CreateUAV(FRDGTextureUAVDesc(SignalTextures.Textures[i], MipLevel));
 	}
 	return UAVs;
 }
 
 
+// TOOD: collapse within reduce pass.
 class FSSDInjestCS : public FGlobalShader
 {
 	DECLARE_GLOBAL_SHADER(FSSDInjestCS);
@@ -617,6 +666,67 @@ class FSSDInjestCS : public FGlobalShader
 	END_SHADER_PARAMETER_STRUCT()
 };
 
+class FSSDReduceCS : public FGlobalShader
+{
+	DECLARE_GLOBAL_SHADER(FSSDReduceCS);
+	SHADER_USE_PARAMETER_STRUCT(FSSDReduceCS, FGlobalShader);
+
+	using FPermutationDomain = TShaderPermutationDomain<FSignalProcessingDim, FSignalBatchSizeDim, FMultiSPPDim>;
+
+	static bool ShouldCompilePermutation(const FGlobalShaderPermutationParameters& Parameters)
+	{
+		FPermutationDomain PermutationVector(Parameters.PermutationId);
+		ESignalProcessing SignalProcessing = PermutationVector.Get<FSignalProcessingDim>();
+
+		// Only compile this shader for signal processing that uses it.
+		if (!SignalUsesReduction(SignalProcessing))
+		{
+			return false;
+		}
+
+		// Not all signal processing allow to batch multiple signals at the same time.
+		if (PermutationVector.Get<FSignalBatchSizeDim>() > SignalMaxBatchSize(SignalProcessing))
+		{
+			return false;
+		}
+
+		// Only compiler multi SPP permutation for signal that supports it.
+		if (PermutationVector.Get<FMultiSPPDim>() && !SignalSupportMultiSPP(SignalProcessing))
+		{
+			return false;
+		}
+
+		// Compile out the shader if this permutation gets remapped.
+		if (RemapPermutationVector(PermutationVector) != PermutationVector)
+		{
+			return false;
+		}
+
+		return ShouldCompileSignalPipeline(SignalProcessing, Parameters.Platform);
+	}
+
+	static FPermutationDomain RemapPermutationVector(FPermutationDomain PermutationVector)
+	{
+		ESignalProcessing SignalProcessing = PermutationVector.Get<FSignalProcessingDim>();
+
+		// force use the multi sample per pixel code path.
+		if (!SignalSupport1SPP(SignalProcessing))
+		{
+			PermutationVector.Set<FMultiSPPDim>(true);
+		}
+
+		return PermutationVector;
+	}
+
+	BEGIN_SHADER_PARAMETER_STRUCT(FParameters, )
+		SHADER_PARAMETER_STRUCT_INCLUDE(FSSDCommonParameters, CommonParameters)
+		SHADER_PARAMETER_STRUCT_INCLUDE(FSSDConvolutionMetaData, ConvolutionMetaData)
+
+		SHADER_PARAMETER_STRUCT(FSSDSignalTextures, SignalInput)
+		SHADER_PARAMETER_STRUCT_ARRAY(FSSDSignalUAVs, SignalOutputMips, [kMaxMipLevel])
+	END_SHADER_PARAMETER_STRUCT()
+};
+
 class FSSDSpatialAccumulationCS : public FGlobalShader
 {
 	DECLARE_GLOBAL_SHADER(FSSDSpatialAccumulationCS);
@@ -673,6 +783,13 @@ class FSSDSpatialAccumulationCS : public FGlobalShader
 			return false;
 		}
 
+		// Only upscale is only for signal that needs it.
+		if (PermutationVector.Get<FUpscaleDim>() &&
+			!SignalSupportsUpscaling(SignalProcessing))
+		{
+			return false;
+		}
+
 		// Only compile pre convolution for signal that uses it.
 		if (!SignalUsesPreConvolution(SignalProcessing) &&
 			PermutationVector.Get<FStageDim>() == EStage::PreConvolution)
@@ -801,6 +918,7 @@ class FSSDTemporalAccumulationCS : public FGlobalShader
 };
 
 IMPLEMENT_GLOBAL_SHADER(FSSDInjestCS, "/Engine/Private/ScreenSpaceDenoise/SSDInjest.usf", "MainCS", SF_Compute);
+IMPLEMENT_GLOBAL_SHADER(FSSDReduceCS, "/Engine/Private/ScreenSpaceDenoise/SSDReduce.usf", "MainCS", SF_Compute);
 IMPLEMENT_GLOBAL_SHADER(FSSDSpatialAccumulationCS, "/Engine/Private/ScreenSpaceDenoise/SSDSpatialAccumulation.usf", "MainCS", SF_Compute);
 IMPLEMENT_GLOBAL_SHADER(FSSDTemporalAccumulationCS, "/Engine/Private/ScreenSpaceDenoise/SSDTemporalAccumulation.usf", "MainCS", SF_Compute);
 
@@ -861,6 +979,7 @@ static void DenoiseSignalAtConstantPixelDensity(
 	// Descriptor to allocate internal denoising buffer.
 	bool bHasReconstructionLayoutDifferentFromHistory = false;
 	TStaticArray<FRDGTextureDesc, kMaxBufferProcessingCount> InjestDescs;
+	TStaticArray<FRDGTextureDesc, kMaxBufferProcessingCount> ReduceDescs;
 	TStaticArray<FRDGTextureDesc, kMaxBufferProcessingCount> ReconstructionDescs;
 	TStaticArray<FRDGTextureDesc, kMaxBufferProcessingCount> HistoryDescs;
 	FRDGTextureDesc DebugDesc;
@@ -890,6 +1009,8 @@ static void DenoiseSignalAtConstantPixelDensity(
 		for (int32 i = 0; i < kMaxBufferProcessingCount; i++)
 		{
 			InjestDescs[i] = RefDesc;
+			ReduceDescs[i] = RefDesc;
+			ReduceDescs[i].NumMips = kMaxMipLevel;
 			ReconstructionDescs[i] = RefDesc;
 			HistoryDescs[i] = RefDesc;
 		}
@@ -946,16 +1067,20 @@ static void DenoiseSignalAtConstantPixelDensity(
 		{
 			for (int32 i = 0; i < 3; i++)
 			{
-				ReconstructionDescs[i].Format = bManualTexelFormatting ? PF_G32R32F : PF_FloatRGBA;
-				HistoryDescs[i].Format = bManualTexelFormatting ? PF_G32R32F : PF_FloatRGBA;
+				ReduceDescs[i].Format = PF_G32R32F;
+				ReconstructionDescs[i].Format = PF_G32R32F;
+				HistoryDescs[i].Format = PF_G32R32F;
 			}
 
-			ReconstructionDescs[3].Format = bManualTexelFormatting ? PF_R32_FLOAT : PF_G16R16F;
-			HistoryDescs[3].Format = bManualTexelFormatting ? PF_R32_FLOAT : PF_G16R16F;
+			ReduceDescs[3].Format = PF_R32_FLOAT;
+			ReconstructionDescs[3].Format = PF_R32_FLOAT;
+			HistoryDescs[3].Format = PF_R32_FLOAT;
 
 			ReconstructionTextureCount = IScreenSpaceDenoiser::kSphericalHarmonicTextureCount;
 			HistoryTextureCountPerSignal = IScreenSpaceDenoiser::kSphericalHarmonicTextureCount; // TODO: only 3 textures for history
 			bHasReconstructionLayoutDifferentFromHistory = false;
+
+			InjestTextureCount = 4;
 		}
 		else
 		{
@@ -1069,6 +1194,40 @@ static void DenoiseSignalAtConstantPixelDensity(
 		SignalHistory = NewSignalOutput;
 	}
 
+	// Reduce the input to speed up the reconstruction pass.
+	if (SignalUsesReduction(Settings.SignalProcessing))
+	{
+		FSSDSignalTextures NewSignalOutput = CreateMultiplexedTextures(
+			GraphBuilder,
+			InjestTextureCount, ReduceDescs,
+			GetResourceNames(kReduceResourceNames));
+
+		FSSDReduceCS::FParameters* PassParameters = GraphBuilder.AllocParameters<FSSDReduceCS::FParameters>();
+		PassParameters->CommonParameters = CommonParameters;
+		PassParameters->ConvolutionMetaData = ConvolutionMetaData;
+		PassParameters->SignalInput = SignalHistory;
+		for (int32 MipLevel = 0; MipLevel < kMaxMipLevel; MipLevel++)
+			PassParameters->SignalOutputMips[MipLevel] = CreateMultiplexedUAVs(GraphBuilder, NewSignalOutput, MipLevel);
+
+		FSSDReduceCS::FPermutationDomain PermutationVector;
+		PermutationVector.Set<FSignalProcessingDim>(Settings.SignalProcessing);
+		PermutationVector.Set<FSignalBatchSizeDim>(Settings.SignalBatchSize);
+		PermutationVector.Set<FMultiSPPDim>(bUseMultiInputSPPShaderPath);
+		PermutationVector = FSSDReduceCS::RemapPermutationVector(PermutationVector);
+
+		TShaderMapRef<FSSDReduceCS> ComputeShader(View.ShaderMap, PermutationVector);
+		FComputeShaderUtils::AddPass(
+			GraphBuilder,
+			RDG_EVENT_NAME("SSD Reduce(Mips=%i MultiSPP=%i)",
+				kMaxMipLevel,
+				int32(PermutationVector.Get<FMultiSPPDim>())),
+			*ComputeShader,
+			PassParameters,
+			FComputeShaderUtils::GetGroupCount(DenoiseResolution, FComputeShaderUtils::kGolden2DGroupSize));
+
+		SignalHistory = NewSignalOutput;
+	}
+
 	// Spatial reconstruction with multiple important sampling to be more precise in the history rejection.
 	{
 		FSSDSignalTextures NewSignalOutput = CreateMultiplexedTextures(
