Added support for EyeAdapation node in the material graph for Volumetric material.
Tested locally that it can then make emissive indpeendent to exposure.
UDN https://udn.unrealengine.com/questions/510720/volumetric-fog-crash-in-422.html

diff --git a/Engine/Shaders/Private/VolumetricFogVoxelization.usf b/Engine/Shaders/Private/VolumetricFogVoxelization.usf
--- a/Engine/Shaders/Private/VolumetricFogVoxelization.usf
+++ b/Engine/Shaders/Private/VolumetricFogVoxelization.usf
@@ -6,6 +6,9 @@
 #define PASS_NEEDS_PRIMITIVE_VOLUME_BOUNDS 1
 
 #include "Common.ush"
+
+#define SceneTexturesStruct VoxelizeVolumePass.SceneTextures 
+
 #include "/Engine/Generated/Material.ush"
 #include "/Engine/Generated/VertexFactory.ush"
 
diff --git a/Engine/Source/Runtime/Renderer/Private/VolumetricFogVoxelization.cpp b/Engine/Source/Runtime/Renderer/Private/VolumetricFogVoxelization.cpp
--- a/Engine/Source/Runtime/Renderer/Private/VolumetricFogVoxelization.cpp
+++ b/Engine/Source/Runtime/Renderer/Private/VolumetricFogVoxelization.cpp
@@ -46,6 +46,7 @@ void SetupVoxelizeVolumePassUniformBuffer(FRHICommandListImmediate& RHICmdList,
 {
 	FSceneRenderTargets& SceneRenderTargets = FSceneRenderTargets::Get(RHICmdList);
 	SetupSceneTextureUniformParameters(SceneRenderTargets, View.FeatureLevel, ESceneTextureSetupMode::None, Parameters.SceneTextures);
+	Parameters.SceneTextures.EyeAdaptation = GetEyeAdaptation(View);
 
 	Parameters.ViewToVolumeClip = View.ViewMatrices.ComputeProjectionNoAAMatrix();
 	Parameters.ViewToVolumeClip.M[2][0] += Jitter.X;
