Fixed compile error in SystemTextures.cpp (missing include header).

diff --git a/Engine/Source/Runtime/Renderer/Private/SystemTextures.cpp b/Engine/Source/Runtime/Renderer/Private/SystemTextures.cpp
--- a/Engine/Source/Runtime/Renderer/Private/SystemTextures.cpp
+++ b/Engine/Source/Runtime/Renderer/Private/SystemTextures.cpp
@@ -6,6 +6,7 @@
 
 #include "SystemTextures.h"
 #include "Math/RandomStream.h"
+#include "Math/Sobol.h"
 #include "RenderTargetPool.h"
 #include "ClearQuad.h"
 #include "LTC.h"
@@ -174,7 +175,7 @@ void FSystemTextures::InitializeFeatureLevelDependentTextures(FRHICommandListImm
 			uint32 DestStride;
 			uint8* DestBuffer = (uint8*)RHICmdList.LockTexture2D((FTexture2DRHIRef&)SobolSampling->GetRenderTargetItem().ShaderResourceTexture, 0, RLM_WriteOnly, DestStride, false);
 
-			uint16 Result, *Dest;
+			uint16 *Dest;
 			for (int y = 0; y < 16; ++y)
 			{
 				Dest = (uint16*)(DestBuffer + y * DestStride);
