Hybrid ray traced / screen space reflections (experimental)

If reflected ray hit point is visible on screen, then SceneColor texture is sampled instead of tracing the material evaluation ray and running lighting. This is a significant performance improvement, but has several current known issues: 1) SSAO is applied to SceneColor before reflections, causing a slight mismatch 2) Direct specular lighting in SceneColor uses the main view eye vector, not a reflection vector, causing a direct lighting mismatch 3) Main depth buffer is used to determine visibility of the reflection hit point, which has issues with surfaces that are steep in clip space (large surface areas map to small number of depth buffer pixels)

This mode can be enabled via r.RayTracing.Reflections.Hybrid=1 (off by default).
Only single-bounce reflections are supported. Material sorting is required and is enabled automatically when hybrid reflections are used.

diff --git a/Engine/Shaders/Private/RayTracing/RayTracingReflections.usf b/Engine/Shaders/Private/RayTracing/RayTracingReflections.usf
--- a/Engine/Shaders/Private/RayTracing/RayTracingReflections.usf
+++ b/Engine/Shaders/Private/RayTracing/RayTracingReflections.usf
@@ -13,6 +13,13 @@
 	#define DIM_DEFERRED_MATERIAL_MODE DEFERRED_MATERIAL_MODE_NONE
 #endif
 
+#ifndef DIM_HYBRID
+	#define DIM_HYBRID 0
+#endif
+
+// Reserved material sorting key for rays that can fetch screen-space data instead of tracing rays in the shading pass.
+#define RAY_TRACING_DEFERRED_MATERIAL_KEY_SCREEN_SPACE 0x80000000
+
 #include "../Common.ush"
 
 #define SUPPORT_CONTACT_SHADOWS		0
@@ -89,6 +96,9 @@ float ReflectionMaxNormalBias;
 
 RaytracingAccelerationStructure TLAS;
 
+// Scene color texture is used for hybrid RT/SS reflection path (DIM_HYBRID=1). Set to dummy black texture otherwise.
+Texture2D SceneColor;
+
 RWTexture2D<float4> ColorOutput;
 RWTexture2D<float> RayHitDistanceOutput;
 RWTexture2D<float> RayImaginaryDepthOutput;
@@ -202,7 +212,9 @@ void RayTracingReflectionsRGS()
 	// * Roughness value
 	// * Hair material is not for RT reflection at the moment. Their indirect contributuion is based on diffuse SH evaluation in ReflectionEnvironemenPixelShader.usf
     const bool bIsValidPixel = (RoughnessFade > 0 || ClearCoatRoughnessFade > 0) && GBufferData.ShadingModelID != SHADINGMODELID_HAIR;
-	   
+
+    const uint LocalMaxBounces = DIM_HYBRID ? 1 : MaxBounces;
+
 	if (bIsValidPixel)
 	{
 		for (uint SampleIndex = 0; SampleIndex < LocalSamplesPerPixel; ++SampleIndex)
@@ -257,7 +269,7 @@ void RayTracingReflectionsRGS()
 			uint BounceIndex = 0;
 			RayDesc TopLayerRay;
 			bool bSkyWasHit = false;
-			for (; BounceIndex < MaxBounces; ++BounceIndex)
+			for (; BounceIndex < LocalMaxBounces; ++BounceIndex)
 			{
 				// Composite pass applies BRDF, except
 				bool bApplyBRDF = BounceIndex != 0;
@@ -308,6 +320,28 @@ void RayTracingReflectionsRGS()
 							DeferredMaterialPayload
 						);
 
+						if (DIM_HYBRID && DeferredMaterialPayload.SortKey < RAY_TRACING_DEFERRED_MATERIAL_KEY_RAY_MISS)
+						{
+							float3 RayHitWorldPos = TopLayerRay.Origin + TopLayerRay.Direction * DeferredMaterialPayload.HitT;
+							float4 RayHitClipPos = mul(float4(RayHitWorldPos, 1), View.WorldToClip);
+
+							if (RayHitClipPos.w > 0)
+							{
+								float2 RayHitScreenPos = RayHitClipPos.xy / RayHitClipPos.w;
+								if (all(abs(RayHitScreenPos) < float2(1, 1)))
+								{
+									float2 RayHitUV = RayHitScreenPos * View.ScreenPositionScaleBias.xy + View.ScreenPositionScaleBias.wz;
+									float GBufferHitDepth  = ConvertFromDeviceZ(SceneDepthBuffer.SampleLevel(GlobalPointClampedSampler, RayHitUV, 0).r);
+									float RayHitDepth = RayHitClipPos.w;
+									float ScreenSpaceReflectionDepthDeltaThreshold = 5.0; // Arbitrary threshold
+									if (RayHitDepth - GBufferHitDepth < ScreenSpaceReflectionDepthDeltaThreshold)
+									{
+										DeferredMaterialPayload.SortKey = RAY_TRACING_DEFERRED_MATERIAL_KEY_SCREEN_SPACE;
+									}
+								}
+							}
+						}
+
 						break; // nothing else to do during gather pass, so exit the loop
 					}
 					else if (DIM_DEFERRED_MATERIAL_MODE == DEFERRED_MATERIAL_MODE_SHADE)
@@ -330,9 +364,42 @@ void RayTracingReflectionsRGS()
 
 				// Trace the top layer reflection
 				float3 TopLayerRadiance = float3(0, 0, 0);
-				FMaterialClosestHitPayload Payload;
-				const uint RayFlags = 0;					
-				Payload = TraceRayAndAccumulateResults(TopLayerRay, TLAS, RandSequence, RayFlags, ReflectionMaxNormalBias, RAY_TRACING_MASK_OPAQUE, ReflectedShadowsType, ShouldDoDirectLighting, ShouldDoEmissiveAndIndirectLighting, RayCone, TopLayerRadiance);
+				FMaterialClosestHitPayload Payload = (FMaterialClosestHitPayload)0;
+				const uint RayFlags = 0;
+
+				if (DIM_HYBRID && DIM_DEFERRED_MATERIAL_MODE == DEFERRED_MATERIAL_MODE_SHADE 
+					&& DeferredMaterialPayload.SortKey == RAY_TRACING_DEFERRED_MATERIAL_KEY_SCREEN_SPACE)
+				{
+					float3 RayHitWorldPos = TopLayerRay.Origin + TopLayerRay.Direction * DeferredMaterialPayload.HitT;
+					float4 RayHitClipPos = mul(float4(RayHitWorldPos, 1), View.WorldToClip);
+					float2 RayHitScreenPos = RayHitClipPos.xy / RayHitClipPos.w;
+					float2 RayHitUV = RayHitScreenPos * View.ScreenPositionScaleBias.xy + View.ScreenPositionScaleBias.wz;
+
+					// Supporting sorted multi-bounce reflections is possible in the future, but requires full GBuffer lookup and decode into Payload format:
+					//   FScreenSpaceData HitScreenSpaceData = GetScreenSpaceData(RayHitUV);
+					//   Payload = MaterialPayloadFromScreenSpaceData(HitScreenSpaceData);
+
+					Payload.WorldPos = RayHitWorldPos;
+					Payload.HitT = DeferredMaterialPayload.HitT;
+
+					TopLayerRadiance += SceneColor.SampleLevel(GlobalPointClampedSampler, RayHitUV, 0) / View.PreExposure;
+				}
+				else
+				{
+					Payload = TraceRayAndAccumulateResults(
+						TopLayerRay,
+						TLAS,
+						RandSequence,
+						RayFlags,
+						ReflectionMaxNormalBias,
+						RAY_TRACING_MASK_OPAQUE,
+						ReflectedShadowsType,
+						ShouldDoDirectLighting,
+						ShouldDoEmissiveAndIndirectLighting,
+						RayCone,
+						TopLayerRadiance);
+				}
+
 				bool isTopLayerRayValid = Payload.IsHit() || (bAllowSkySampling && ReflectionStruct.SkyLightParameters.y > 0);
 
 
@@ -548,7 +615,7 @@ void RayTracingReflectionsRGS()
 			
 			if (bIsValidSample)
 			{
-				bNeedsCapture |= BounceIndex == MaxBounces && !(bSkyWasHit || HitDistance < 0.0f);
+				bNeedsCapture |= BounceIndex == LocalMaxBounces && !(bSkyWasHit || HitDistance < 0.0f);
 				if (UseReflectionCaptures && bNeedsCapture)
 				{
 					// We have reached the end of the bounce count on a surface. 
diff --git a/Engine/Source/Runtime/Renderer/Private/RayTracing/RayTracingReflections.cpp b/Engine/Source/Runtime/Renderer/Private/RayTracing/RayTracingReflections.cpp
--- a/Engine/Source/Runtime/Renderer/Private/RayTracing/RayTracingReflections.cpp
+++ b/Engine/Source/Runtime/Renderer/Private/RayTracing/RayTracingReflections.cpp
@@ -90,6 +90,15 @@ static TAutoConsoleVariable<int32> CVarRayTracingReflectionsSortMaterials(
 	TEXT("1: Enabled, using Trace->Sort->Trace\n"),
 	ECVF_RenderThreadSafe);
 
+static TAutoConsoleVariable<int32> CVarRayTracingReflectionsHybrid(
+	TEXT("r.RayTracing.Reflections.Hybrid"),
+	0,
+	TEXT("Sets whether screen space reflections should be used when possible (experimental).\n")
+	TEXT("Forces material sorting and single ray bounce.\n")
+	TEXT("0: Disabled (Default)\n ")
+	TEXT("1: Enabled\n"),
+	ECVF_RenderThreadSafe);
+
 static TAutoConsoleVariable<int32> CVarRayTracingReflectionsSortTileSize(
 	TEXT("r.RayTracing.Reflections.SortTileSize"),
 	64,
@@ -140,8 +149,9 @@ class FRayTracingReflectionsRGS : public FGlobalShader
 	class FDenoiserOutput : SHADER_PERMUTATION_BOOL("DIM_DENOISER_OUTPUT");
 
 	class FDeferredMaterialMode : SHADER_PERMUTATION_ENUM_CLASS("DIM_DEFERRED_MATERIAL_MODE", EDeferredMaterialMode);
+	class FHybrid : SHADER_PERMUTATION_BOOL("DIM_HYBRID");
 
-	using FPermutationDomain = TShaderPermutationDomain<FDenoiserOutput, FDeferredMaterialMode>;
+	using FPermutationDomain = TShaderPermutationDomain<FDenoiserOutput, FDeferredMaterialMode, FHybrid>;
 
 	BEGIN_SHADER_PARAMETER_STRUCT(FParameters, )
 		SHADER_PARAMETER(int32, SamplesPerPixel)
@@ -164,6 +174,7 @@ class FRayTracingReflectionsRGS : public FGlobalShader
 		SHADER_PARAMETER(int32, MaxUnderCoatBounces)
 
 		SHADER_PARAMETER_SRV(RaytracingAccelerationStructure, TLAS)
+		SHADER_PARAMETER_RDG_TEXTURE(Texture2D, SceneColor)
 		SHADER_PARAMETER_SRV(StructuredBuffer<FRTLightingData>, LightDataBuffer)
 
 		SHADER_PARAMETER_STRUCT_INCLUDE(FSceneTextureParameters, SceneTextures)
@@ -196,16 +207,34 @@ void FDeferredShadingSceneRenderer::PrepareRayTracingReflections(const FViewInfo
 {
 	// Declare all RayGen shaders that require material closest hit shaders to be bound
 
-	const bool bSortMaterials = CVarRayTracingReflectionsSortMaterials.GetValueOnRenderThread() != 0;
-
-	const EDeferredMaterialMode DeferredMaterialMode = bSortMaterials ? EDeferredMaterialMode::Shade : EDeferredMaterialMode::None;
+	const bool bHybridReflections = CVarRayTracingReflectionsHybrid.GetValueOnRenderThread() != 0;
+	const bool bSortMaterials = bHybridReflections || CVarRayTracingReflectionsSortMaterials.GetValueOnRenderThread() != 0;
 
-	FRayTracingReflectionsRGS::FPermutationDomain PermutationVector;
-	PermutationVector.Set<FRayTracingReflectionsRGS::FDeferredMaterialMode>(DeferredMaterialMode);
-
-	auto RayGenShader = View.ShaderMap->GetShader<FRayTracingReflectionsRGS>(PermutationVector);
-
-	OutRayGenShaders.Add(RayGenShader->GetRayTracingShader());
+	if (bSortMaterials)
+	{
+		{
+			FRayTracingReflectionsRGS::FPermutationDomain PermutationVector;
+			PermutationVector.Set<FRayTracingReflectionsRGS::FDeferredMaterialMode>(EDeferredMaterialMode::Gather);
+			PermutationVector.Set<FRayTracingReflectionsRGS::FHybrid>(bHybridReflections);
+			auto RayGenShader = View.ShaderMap->GetShader<FRayTracingReflectionsRGS>(PermutationVector);
+			OutRayGenShaders.Add(RayGenShader->GetRayTracingShader());
+		}
+		
+		{
+			FRayTracingReflectionsRGS::FPermutationDomain PermutationVector;
+			PermutationVector.Set<FRayTracingReflectionsRGS::FDeferredMaterialMode>(EDeferredMaterialMode::Shade);
+			PermutationVector.Set<FRayTracingReflectionsRGS::FHybrid>(bHybridReflections);
+			auto RayGenShader = View.ShaderMap->GetShader<FRayTracingReflectionsRGS>(PermutationVector);
+			OutRayGenShaders.Add(RayGenShader->GetRayTracingShader());
+		}
+	}
+	else
+	{
+		FRayTracingReflectionsRGS::FPermutationDomain PermutationVector;
+		PermutationVector.Set<FRayTracingReflectionsRGS::FDeferredMaterialMode>(EDeferredMaterialMode::None);
+		auto RayGenShader = View.ShaderMap->GetShader<FRayTracingReflectionsRGS>(PermutationVector);
+		OutRayGenShaders.Add(RayGenShader->GetRayTracingShader());
+	}
 }
 
 #endif // RHI_RAYTRACING
@@ -221,7 +250,9 @@ void FDeferredShadingSceneRenderer::RenderRayTracingReflections(
 #if RHI_RAYTRACING
 {
 	const uint32 SortTileSize = CVarRayTracingReflectionsSortTileSize.GetValueOnRenderThread();
-	const bool bSortMaterials = CVarRayTracingReflectionsSortMaterials.GetValueOnRenderThread() != 0;
+
+	const bool bHybridReflections = CVarRayTracingReflectionsHybrid.GetValueOnRenderThread() != 0;
+	const bool bSortMaterials = bHybridReflections || CVarRayTracingReflectionsSortMaterials.GetValueOnRenderThread() != 0;
 
 	FSceneRenderTargets& SceneContext = FSceneRenderTargets::Get(GraphBuilder.RHICmdList);
 
@@ -295,6 +326,9 @@ void FDeferredShadingSceneRenderer::RenderRayTracingReflections(
 	CommonParameters.SceneTextures = SceneTextures;
 	SetupSceneTextureSamplers(&CommonParameters.SceneTextureSamplers);
 
+	// Hybrid reflection path samples lit scene color texture instead of performing a ray trace.
+	CommonParameters.SceneColor = GraphBuilder.RegisterExternalTexture(bHybridReflections ? SceneContext.GetSceneColor() : GSystemTextures.BlackDummy);
+
 	CommonParameters.ReflectionStruct = CreateReflectionUniformBuffer(View, EUniformBufferUsage::UniformBuffer_SingleFrame);
 	CommonParameters.FogUniformParameters = CreateFogUniformBuffer(View, EUniformBufferUsage::UniformBuffer_SingleFrame);
 	CommonParameters.ColorOutput = GraphBuilder.CreateUAV(OutDenoiserInputs->Color);
@@ -324,6 +358,7 @@ void FDeferredShadingSceneRenderer::RenderRayTracingReflections(
 
 		FRayTracingReflectionsRGS::FPermutationDomain PermutationVector;
 		PermutationVector.Set<FRayTracingReflectionsRGS::FDeferredMaterialMode>(DeferredMaterialMode);
+		PermutationVector.Set<FRayTracingReflectionsRGS::FHybrid>(bHybridReflections);
 
 		auto RayGenShader = View.ShaderMap->GetShader<FRayTracingReflectionsRGS>(PermutationVector);
 		ClearUnusedGraphResources(RayGenShader, PassParameters);
