Re-adds support for non shared sampler in SSGI to unblock GL4.

diff --git a/Engine/Shaders/Private/SSRT/SSRTDiffuseIndirect.usf b/Engine/Shaders/Private/SSRT/SSRTDiffuseIndirect.usf
--- a/Engine/Shaders/Private/SSRT/SSRTDiffuseIndirect.usf
+++ b/Engine/Shaders/Private/SSRT/SSRTDiffuseIndirect.usf
@@ -78,6 +78,18 @@
 Texture2D HZBTexture;
 Texture2D ColorTexture;
 
+#if SUPPORTS_INDEPENDENT_SAMPLERS
+
+#define HZBTextureSampler GlobalPointClampedSampler
+#define ColorTextureSampler GlobalBilinearClampedSampler
+
+#else
+
+SamplerState HZBTextureSampler;
+SamplerState ColorTextureSampler;
+
+#endif
+
 float4 HZBUvFactorAndInvFactor;
 
 float4 PrevScreenPositionScaleBias;
@@ -380,7 +392,7 @@ void MainCS(
 			float Level;
 			float3 HitUVz;
 			bool bHit = RayCast(
-				HZBTexture, GlobalPointClampedSampler,
+				HZBTexture, HZBTextureSampler,
 				TranslatedWorldPosition, L, 1, SceneDepth, 
 				CONFIG_RAY_STEPS, StepOffset,
 				HZBUvFactorAndInvFactor,
@@ -400,7 +412,7 @@ void MainCS(
 			//BRANCH // TODO
 			if (bHit)
 			{
-				float4 SampleColor = ColorTexture.SampleLevel(GlobalBilinearClampedSampler, HitUVz, Level);
+				float4 SampleColor = ColorTexture.SampleLevel(ColorTextureSampler, HitUVz, Level);
 				
 				float SampleColorWeight = 1.0;
 
diff --git a/Engine/Shaders/Private/SSRT/SSRTPrevFrameReduction.usf b/Engine/Shaders/Private/SSRT/SSRTPrevFrameReduction.usf
--- a/Engine/Shaders/Private/SSRT/SSRTPrevFrameReduction.usf
+++ b/Engine/Shaders/Private/SSRT/SSRTPrevFrameReduction.usf
@@ -11,6 +11,16 @@
 uint bIsTemporalAAHistory;
 Texture2D PrevTemporalAAHistory;
 
+#if SUPPORTS_INDEPENDENT_SAMPLERS
+
+#define PrevTemporalAAHistorySampler GlobalBilinearClampedSampler
+
+#else
+
+SamplerState PrevTemporalAAHistorySampler;
+
+#endif
+
 float4 PrevScreenPositionScaleBias;
 float  PrevSceneColorPreExposureCorrection;
 
@@ -61,7 +71,7 @@ void MainCS(
 
 	// Fetch color.
 	float4 PrevSceneColor = float4(
-		PrevTemporalAAHistory.SampleLevel(GlobalBilinearClampedSampler, PrevFrameUV, 0).rgb * PrevSceneColorPreExposureCorrection, 1.0);
+		PrevTemporalAAHistory.SampleLevel(PrevTemporalAAHistorySampler, PrevFrameUV, 0).rgb * PrevSceneColorPreExposureCorrection, 1.0);
 
 	// Apply vignette to the color.
 	{
diff --git a/Engine/Source/Runtime/Renderer/Private/ScreenSpaceRayTracing.cpp b/Engine/Source/Runtime/Renderer/Private/ScreenSpaceRayTracing.cpp
--- a/Engine/Source/Runtime/Renderer/Private/ScreenSpaceRayTracing.cpp
+++ b/Engine/Source/Runtime/Renderer/Private/ScreenSpaceRayTracing.cpp
@@ -181,8 +181,10 @@ class FSSRTPrevFrameReductionCS : public FGlobalShader
 		SHADER_PARAMETER(uint32, bIsTemporalAAHistory)
 
 		SHADER_PARAMETER_RDG_TEXTURE(Texture2D, PrevTemporalAAHistory)
+		SHADER_PARAMETER_SAMPLER(SamplerState, PrevTemporalAAHistorySampler)
 
 		SHADER_PARAMETER_STRUCT_INCLUDE(FSceneTextureParameters, SceneTextures)
+		SHADER_PARAMETER_STRUCT_INCLUDE(FSceneTextureSamplerParameters, SceneTextureSamplers)
 		SHADER_PARAMETER_STRUCT_REF(FViewUniformShaderParameters, View)
 
 		SHADER_PARAMETER_RDG_TEXTURE_UAV_ARRAY(RWTexture<float4>, ReducedSceneColorOutput, [4])
@@ -259,9 +261,13 @@ class FScreenSpaceDiffuseIndirectCS : public FGlobalShader
 		SHADER_PARAMETER( float,		PrevSceneColorPreExposureCorrection )
 		
 		SHADER_PARAMETER_RDG_TEXTURE(Texture2D, HZBTexture)
+		SHADER_PARAMETER_SAMPLER(SamplerState, HZBTextureSampler)
+
 		SHADER_PARAMETER_RDG_TEXTURE(Texture2D, ColorTexture)
+		SHADER_PARAMETER_SAMPLER(SamplerState, ColorTextureSampler)
 		
 		SHADER_PARAMETER_STRUCT_INCLUDE(FSceneTextureParameters, SceneTextures)
+		SHADER_PARAMETER_STRUCT_INCLUDE(FSceneTextureSamplerParameters, SceneTextureSamplers)
 		SHADER_PARAMETER_STRUCT_REF(FViewUniformShaderParameters, View)
 
 		SHADER_PARAMETER_RDG_TEXTURE_UAV(RWTexture<float4>, IndirectDiffuseOutput)
@@ -647,8 +653,10 @@ void RenderScreenSpaceDiffuseIndirect(
 
 		PassParameters->bIsTemporalAAHistory = true;
 		PassParameters->PrevTemporalAAHistory = ColorTexture;
+		PassParameters->PrevTemporalAAHistorySampler = TStaticSamplerState<SF_Bilinear>::GetRHI();
 
 		PassParameters->SceneTextures = SceneTextures;
+		SetupSceneTextureSamplers(&PassParameters->SceneTextureSamplers);
 		PassParameters->View = View.ViewUniformBuffer;
 
 		for (int32 MipLevel = 0; MipLevel < ReducedSceneColor->Desc.NumMips; MipLevel++)
@@ -686,7 +694,9 @@ void RenderScreenSpaceDiffuseIndirect(
 		FScreenSpaceDiffuseIndirectCS::FParameters* PassParameters = GraphBuilder.AllocParameters<FScreenSpaceDiffuseIndirectCS::FParameters>();
 
 		PassParameters->HZBTexture = HZBTexture;
+		PassParameters->HZBTextureSampler = TStaticSamplerState<SF_Point>::GetRHI();
 		PassParameters->ColorTexture = ReducedSceneColor;
+		PassParameters->ColorTextureSampler = TStaticSamplerState<SF_Bilinear>::GetRHI();
 
 		const FVector2D HZBUvFactor(
 			float( View.ViewRect.Width() )  / float( 2 * View.HZBMipmap0Size.X ),
@@ -712,6 +722,7 @@ void RenderScreenSpaceDiffuseIndirect(
 		PassParameters->PrevSceneColorPreExposureCorrection = View.PreExposure / View.PrevViewInfo.SceneColorPreExposure;
 
 		PassParameters->SceneTextures = SceneTextures;
+		SetupSceneTextureSamplers(&PassParameters->SceneTextureSamplers);
 		PassParameters->View = View.ViewUniformBuffer;
 	
 		PassParameters->IndirectDiffuseOutput = GraphBuilder.CreateUAV(OutDenoiserInputs->Color);
