Remove ERDGPassFlags::None and ERenderGraphPassFlags.

diff --git a/Engine/Source/Runtime/RenderCore/Private/VisualizeTexture.cpp b/Engine/Source/Runtime/RenderCore/Private/VisualizeTexture.cpp
--- a/Engine/Source/Runtime/RenderCore/Private/VisualizeTexture.cpp
+++ b/Engine/Source/Runtime/RenderCore/Private/VisualizeTexture.cpp
@@ -289,7 +289,7 @@ void FVisualizeTexture::CreateContentCapturePass(FRDGBuilder& GraphBuilder, cons
 			SrcDesc.Extent.X, SrcDesc.Extent.Y,
 			*ExtendedDrawEvent),
 		PassParameters,
-		ERenderGraphPassFlags::None,
+		ERDGPassFlags::Raster,
 		[this, PassParameters, ShaderMap, PixelShader, RTExtent](FRHICommandList& RHICmdList)
 	{
 		FVisualizeTexturePS::FParameters ShaderParameter = *PassParameters;
diff --git a/Engine/Source/Runtime/RenderCore/Public/PixelShaderUtils.h b/Engine/Source/Runtime/RenderCore/Public/PixelShaderUtils.h
--- a/Engine/Source/Runtime/RenderCore/Public/PixelShaderUtils.h
+++ b/Engine/Source/Runtime/RenderCore/Public/PixelShaderUtils.h
@@ -67,7 +67,7 @@ struct RENDERCORE_API FPixelShaderUtils
 		GraphBuilder.AddPass(
 			Forward<FRDGEventName>(PassName),
 			Parameters,
-			ERenderGraphPassFlags::None,
+			ERDGPassFlags::Raster,
 			[Parameters, GlobalShaderMap, PixelShader, Viewport](FRHICommandList& RHICmdList)
 		{
 			FPixelShaderUtils::DrawFullscreenPixelShader(RHICmdList, GlobalShaderMap, PixelShader, *Parameters, Viewport);
diff --git a/Engine/Source/Runtime/RenderCore/Public/RayGenShaderUtils.h b/Engine/Source/Runtime/RenderCore/Public/RayGenShaderUtils.h
--- a/Engine/Source/Runtime/RenderCore/Public/RayGenShaderUtils.h
+++ b/Engine/Source/Runtime/RenderCore/Public/RayGenShaderUtils.h
@@ -28,7 +28,7 @@ struct RENDERCORE_API FRayGenShaderUtils
 		GraphBuilder.AddPass(
 			Forward<FRDGEventName>(PassName),
 			Parameters,
-			ERenderGraphPassFlags::Compute,
+			ERDGPassFlags::Compute,
 			[RayGenerationShader, Parameters, Resolution](FRHICommandList& RHICmdList)
 		{
 			FRayTracingShaderBindingsWriter GlobalResources;
diff --git a/Engine/Source/Runtime/RenderCore/Public/RenderGraphPass.h b/Engine/Source/Runtime/RenderCore/Public/RenderGraphPass.h
--- a/Engine/Source/Runtime/RenderCore/Public/RenderGraphPass.h
+++ b/Engine/Source/Runtime/RenderCore/Public/RenderGraphPass.h
@@ -164,11 +164,8 @@ private:
 /** Flags to annotate passes. */
 enum class ERDGPassFlags : uint8
 {
-	/** Equivalent to 'Raster'. */
-	None = 0x1,
-
 	/** Pass uses raster pipeline. */
-	Raster = None,
+	Raster = 0x1,
 
 	/** Pass uses compute only */
 	Compute = 0x2,
@@ -182,9 +179,6 @@ enum class ERDGPassFlags : uint8
 };
 ENUM_CLASS_FLAGS(ERDGPassFlags);
 
-// TODO(RDG): Bulk rename across codebase.
-using ERenderGraphPassFlags = ERDGPassFlags;
-
 /** Base class of a render graph pass. */
 class RENDERCORE_API FRDGPass
 {
diff --git a/Engine/Source/Runtime/RenderCore/Public/RenderGraphUtils.h b/Engine/Source/Runtime/RenderCore/Public/RenderGraphUtils.h
--- a/Engine/Source/Runtime/RenderCore/Public/RenderGraphUtils.h
+++ b/Engine/Source/Runtime/RenderCore/Public/RenderGraphUtils.h
@@ -147,7 +147,7 @@ struct RENDERCORE_API FComputeShaderUtils
 		GraphBuilder.AddPass(
 			Forward<FRDGEventName>(PassName),
 			Parameters,
-			ERenderGraphPassFlags::Compute,
+			ERDGPassFlags::Compute,
 			[Parameters, ComputeShader, GroupCount](FRHICommandList& RHICmdList)
 		{
 			FComputeShaderUtils::Dispatch(RHICmdList, ComputeShader, *Parameters, GroupCount);
@@ -171,7 +171,7 @@ struct RENDERCORE_API FComputeShaderUtils
 		GraphBuilder.AddPass(
 			Forward<FRDGEventName>(PassName),
 			Parameters,
-			ERenderGraphPassFlags::Compute,
+			ERDGPassFlags::Compute,
 			[Parameters, ComputeShader, IndirectArgsBuffer, IndirectArgOffset](FRHICommandList& RHICmdList)
 		{			
 			// Marks the indirect draw parameter as used by the pass manually, given it can't be bound directly by any of the shader,
diff --git a/Engine/Source/Runtime/Renderer/Private/ClusteredDeferredShadingPass.cpp b/Engine/Source/Runtime/Renderer/Private/ClusteredDeferredShadingPass.cpp
--- a/Engine/Source/Runtime/Renderer/Private/ClusteredDeferredShadingPass.cpp
+++ b/Engine/Source/Runtime/Renderer/Private/ClusteredDeferredShadingPass.cpp
@@ -139,7 +139,7 @@ void FDeferredShadingSceneRenderer::AddClusteredDeferredShadingPass(FRHICommandL
 			GraphBuilder.AddPass(
 				RDG_EVENT_NAME("ClusteredDeferredShading, #Lights: %d", NumLightsToRender),
 				PassParameters,
-				ERenderGraphPassFlags::None,
+				ERDGPassFlags::Raster,
 				[PassParameters, &View, &SceneContext](FRHICommandListImmediate& InRHICmdList)
 			{
 				TShaderMapRef<FPostProcessVS> VertexShader(View.ShaderMap);
diff --git a/Engine/Source/Runtime/Renderer/Private/IndirectLightRendering.cpp b/Engine/Source/Runtime/Renderer/Private/IndirectLightRendering.cpp
--- a/Engine/Source/Runtime/Renderer/Private/IndirectLightRendering.cpp
+++ b/Engine/Source/Runtime/Renderer/Private/IndirectLightRendering.cpp
@@ -237,7 +237,7 @@ void FDeferredShadingSceneRenderer::RenderDiffuseIndirectAndAmbientOcclusion(FRH
 					PermutationVector.Get<FDiffuseIndirectCompositePS::FApplyDiffuseIndirectDim>() ? TEXT("Yes") : TEXT("No"),
 					View.ViewRect.Width(), View.ViewRect.Height()),
 				PassParameters,
-				ERenderGraphPassFlags::None,
+				ERDGPassFlags::Raster,
 				[PassParameters, &View, PixelShader, PermutationVector](FRHICommandList& RHICmdList)
 			{
 				RHICmdList.SetViewport(View.ViewRect.Min.X, View.ViewRect.Min.Y, 0.0f, View.ViewRect.Max.X, View.ViewRect.Max.Y, 0.0);
@@ -288,7 +288,7 @@ void FDeferredShadingSceneRenderer::RenderDiffuseIndirectAndAmbientOcclusion(FRH
 			GraphBuilder.AddPass(
 				RDG_EVENT_NAME("AmbientCubemapComposite %dx%d", View.ViewRect.Width(), View.ViewRect.Height()),
 				PassParameters,
-				ERDGPassFlags::None,
+				ERDGPassFlags::Raster,
 				[PassParameters, &View, PixelShader](FRHICommandList& RHICmdList)
 			{
 				TShaderMapRef<FPostProcessVS> VertexShader(View.ShaderMap);
diff --git a/Engine/Source/Runtime/Renderer/Private/LightGridInjection.cpp b/Engine/Source/Runtime/Renderer/Private/LightGridInjection.cpp
--- a/Engine/Source/Runtime/Renderer/Private/LightGridInjection.cpp
+++ b/Engine/Source/Runtime/Renderer/Private/LightGridInjection.cpp
@@ -92,7 +92,7 @@ void AddPass_ClearUAV(
 	GraphBuilder.AddPass(
 		Forward<FRDGEventName>(PassName),
 		Parameters,
-		ERenderGraphPassFlags::Compute,
+		ERDGPassFlags::Compute,
 		[&Parameters, BufferUAV, Value](FRHICommandList& RHICmdList)
 		{
 			BufferUAV->MarkResourceAsUsed();
diff --git a/Engine/Source/Runtime/Renderer/Private/PlanarReflectionRendering.cpp b/Engine/Source/Runtime/Renderer/Private/PlanarReflectionRendering.cpp
--- a/Engine/Source/Runtime/Renderer/Private/PlanarReflectionRendering.cpp
+++ b/Engine/Source/Runtime/Renderer/Private/PlanarReflectionRendering.cpp
@@ -734,7 +734,7 @@ void FDeferredShadingSceneRenderer::RenderDeferredPlanarReflections(FRDGBuilder&
 	GraphBuilder.AddPass(
 		RDG_EVENT_NAME("CompositePlanarReflections"),
 		PassParameters,
-		ERDGPassFlags::None,
+		ERDGPassFlags::Raster,
 		[PassParameters, &View, this](FRHICommandList& RHICmdList)
 	{
 		RHICmdList.SetViewport(View.ViewRect.Min.X, View.ViewRect.Min.Y, 0.0f, View.ViewRect.Max.X, View.ViewRect.Max.Y, 1.0f);
diff --git a/Engine/Source/Runtime/Renderer/Private/PostProcess/DiaphragmDOF.cpp b/Engine/Source/Runtime/Renderer/Private/PostProcess/DiaphragmDOF.cpp
--- a/Engine/Source/Runtime/Renderer/Private/PostProcess/DiaphragmDOF.cpp
+++ b/Engine/Source/Runtime/Renderer/Private/PostProcess/DiaphragmDOF.cpp
@@ -2272,7 +2272,7 @@ FRDGTextureRef DiaphragmDOF::AddPasses(
 					PermutationVector.Get<FDDOFScatterOcclusionDim>() ? TEXT("Yes") : TEXT("No"),
 					GatheringViewSize.X, GatheringViewSize.Y),
 				PassParameters,
-				ERenderGraphPassFlags::None,
+				ERDGPassFlags::Raster,
 				[PassParameters, VertexShader, PixelShader, GatheringViewSize, DrawIndirectParametersOffset](FRHICommandList& RHICmdList)
 			{
 				RHICmdList.SetViewport(0, 0, 0.0f, GatheringViewSize.X, GatheringViewSize.Y, 1.0f);
diff --git a/Engine/Source/Runtime/Renderer/Private/PostProcess/PostProcessMaterial.cpp b/Engine/Source/Runtime/Renderer/Private/PostProcess/PostProcessMaterial.cpp
--- a/Engine/Source/Runtime/Renderer/Private/PostProcess/PostProcessMaterial.cpp
+++ b/Engine/Source/Runtime/Renderer/Private/PostProcess/PostProcessMaterial.cpp
@@ -517,7 +517,7 @@ FRDGTextureRef ComputePostProcessMaterial(
 	GraphBuilder.AddPass(
 		RDG_EVENT_NAME("PostProcessMaterial"),
 		PostProcessMaterialParameters,
-		ERenderGraphPassFlags::None,
+		ERDGPassFlags::Raster,
 		[ScreenPassView, SceneColorViewport, ScreenPassDraw, PostProcessMaterialParameters, DepthStencilTexture, DepthStencilTextureForSRV, SetupFunction, bFlipYAxis]
 		(FRHICommandListImmediate& RHICmdList)
 	{
diff --git a/Engine/Source/Runtime/Renderer/Private/PostProcess/PostProcessMotionBlur.cpp b/Engine/Source/Runtime/Renderer/Private/PostProcess/PostProcessMotionBlur.cpp
--- a/Engine/Source/Runtime/Renderer/Private/PostProcess/PostProcessMotionBlur.cpp
+++ b/Engine/Source/Runtime/Renderer/Private/PostProcess/PostProcessMotionBlur.cpp
@@ -551,7 +551,7 @@ void ComputeMotionBlurVelocity(
 		GraphBuilder.AddPass(
 			RDG_EVENT_NAME("VelocityTileScatter %dx%d", VelocityTileCount.X, VelocityTileCount.Y),
 			PassParameters,
-			ERenderGraphPassFlags::None,
+			ERDGPassFlags::Raster,
 			[VertexShader, PixelShader, VelocityTileCount, PassParameters](FRHICommandListImmediate& RHICmdList)
 		{
 			FRHIVertexShader* RHIVertexShader = GETSAFERHISHADER_VERTEX(*VertexShader);
@@ -840,7 +840,7 @@ FRDGTextureRef VisualizeMotionBlur(
 	GraphBuilder.AddPass(
 		RDG_EVENT_NAME("MotionBlurVisualize"),
 		PassParameters,
-		ERenderGraphPassFlags::None,
+		ERDGPassFlags::Raster,
 		[ScreenPassView, ColorTextureOutput, ColorViewport, PixelShader, PassParameters](FRHICommandListImmediate& RHICmdList)
 	{
 		DrawScreenPass(RHICmdList, ScreenPassView, ColorViewport, ColorViewport, *PixelShader, *PassParameters);
diff --git a/Engine/Source/Runtime/Renderer/Private/PostProcess/PostProcessSubsurface.cpp b/Engine/Source/Runtime/Renderer/Private/PostProcess/PostProcessSubsurface.cpp
--- a/Engine/Source/Runtime/Renderer/Private/PostProcess/PostProcessSubsurface.cpp
+++ b/Engine/Source/Runtime/Renderer/Private/PostProcess/PostProcessSubsurface.cpp
@@ -174,7 +174,7 @@ FRDGTextureRef ComputeSubsurface(
 		GraphBuilder.AddPass(
 			RDG_EVENT_NAME("SubsurfaceViewportCopy"),
 			PassParameters,
-			ERenderGraphPassFlags::None,
+			ERDGPassFlags::Raster,
 			[&Views, ViewMask, ViewCount, PixelShader, InputTextureSize, PassParameters](FRHICommandListImmediate& RHICmdList)
 		{
 			for (uint32 ViewIndex = 0; ViewIndex < ViewCount; ++ViewIndex)
@@ -248,7 +248,7 @@ void VisualizeSubsurface(
 	GraphBuilder.AddPass(
 		RDG_EVENT_NAME("SubsurfaceVisualize"),
 		PassParameters,
-		ERenderGraphPassFlags::None,
+		ERDGPassFlags::Raster,
 		[ScreenPassView, SceneViewport, SceneTextureOutput, PixelShader, PassParameters](FRHICommandListImmediate& RHICmdList)
 	{
 		DrawScreenPass(RHICmdList, ScreenPassView, SceneViewport, SceneViewport, *PixelShader, *PassParameters);
diff --git a/Engine/Source/Runtime/Renderer/Private/RayTracing/RayTracingAmbientOcclusion.cpp b/Engine/Source/Runtime/Renderer/Private/RayTracing/RayTracingAmbientOcclusion.cpp
--- a/Engine/Source/Runtime/Renderer/Private/RayTracing/RayTracingAmbientOcclusion.cpp
+++ b/Engine/Source/Runtime/Renderer/Private/RayTracing/RayTracingAmbientOcclusion.cpp
@@ -165,7 +165,7 @@ void FDeferredShadingSceneRenderer::RenderRayTracingAmbientOcclusion(
 	GraphBuilder.AddPass(
 		RDG_EVENT_NAME("AmbientOcclusionRayTracing %dx%d", RayTracingResolution.X, RayTracingResolution.Y),
 		PassParameters,
-		ERenderGraphPassFlags::Compute,
+		ERDGPassFlags::Compute,
 		[PassParameters, this, &View, RayGenerationShader, RayTracingResolution](FRHICommandList& RHICmdList)
 	{
 		FRayTracingShaderBindingsWriter GlobalResources;
diff --git a/Engine/Source/Runtime/Renderer/Private/RayTracing/RayTracingBarycentrics.cpp b/Engine/Source/Runtime/Renderer/Private/RayTracing/RayTracingBarycentrics.cpp
--- a/Engine/Source/Runtime/Renderer/Private/RayTracing/RayTracingBarycentrics.cpp
+++ b/Engine/Source/Runtime/Renderer/Private/RayTracing/RayTracingBarycentrics.cpp
@@ -87,7 +87,7 @@ void FDeferredShadingSceneRenderer::RenderRayTracingBarycentrics(FRHICommandList
 	GraphBuilder.AddPass(
 		RDG_EVENT_NAME("Barycentrics"),
 		RayGenParameters,
-		ERenderGraphPassFlags::Compute,
+		ERDGPassFlags::Compute,
 		[this, RayGenParameters, RayGenShader, &SceneContext, RayTracingSceneRHI, Pipeline, ViewRect](FRHICommandList& RHICmdList)
 	{
 		FRayTracingShaderBindingsWriter GlobalResources;
diff --git a/Engine/Source/Runtime/Renderer/Private/RayTracing/RayTracingDebug.cpp b/Engine/Source/Runtime/Renderer/Private/RayTracing/RayTracingDebug.cpp
--- a/Engine/Source/Runtime/Renderer/Private/RayTracing/RayTracingDebug.cpp
+++ b/Engine/Source/Runtime/Renderer/Private/RayTracing/RayTracingDebug.cpp
@@ -122,7 +122,7 @@ void FDeferredShadingSceneRenderer::RenderRayTracingDebug(FRHICommandListImmedia
 	GraphBuilder.AddPass(
 		RDG_EVENT_NAME("RayTracingDebug"),
 		RayGenParameters,
-		ERenderGraphPassFlags::Compute,
+		ERDGPassFlags::Compute,
 		[this, RayGenParameters, RayGenShader, &SceneContext, RayTracingSceneRHI, Pipeline, ViewRect](FRHICommandList& RHICmdList)
 	{
 		SCOPED_GPU_STAT(RHICmdList, RayTracingDebug);
diff --git a/Engine/Source/Runtime/Renderer/Private/RayTracing/RayTracingGlobalIllumination.cpp b/Engine/Source/Runtime/Renderer/Private/RayTracing/RayTracingGlobalIllumination.cpp
--- a/Engine/Source/Runtime/Renderer/Private/RayTracing/RayTracingGlobalIllumination.cpp
+++ b/Engine/Source/Runtime/Renderer/Private/RayTracing/RayTracingGlobalIllumination.cpp
@@ -391,7 +391,7 @@ void FDeferredShadingSceneRenderer::RenderRayTracingGlobalIllumination(
 		GraphBuilder.AddPass(
 			RDG_EVENT_NAME("GlobalIlluminationRayTracing %dx%d", RayTracingResolution.X, RayTracingResolution.Y),
 			PassParameters,
-			ERenderGraphPassFlags::Compute,
+			ERDGPassFlags::Compute,
 			[PassParameters, this, &View, RayGenerationShader, RayTracingResolution](FRHICommandList& RHICmdList)
 		{
 			FRayTracingShaderBindingsWriter GlobalResources;
diff --git a/Engine/Source/Runtime/Renderer/Private/RayTracing/RayTracingReflections.cpp b/Engine/Source/Runtime/Renderer/Private/RayTracing/RayTracingReflections.cpp
--- a/Engine/Source/Runtime/Renderer/Private/RayTracing/RayTracingReflections.cpp
+++ b/Engine/Source/Runtime/Renderer/Private/RayTracing/RayTracingReflections.cpp
@@ -368,7 +368,7 @@ void FDeferredShadingSceneRenderer::RenderRayTracingReflections(
 			GraphBuilder.AddPass(
 				RDG_EVENT_NAME("ReflectionRayTracingGatherMaterials %dx%d", TileAlignedResolution.X, TileAlignedResolution.Y),
 				PassParameters,
-				ERenderGraphPassFlags::Compute,
+				ERDGPassFlags::Compute,
 				[PassParameters, this, &View, RayGenShader, TileAlignedResolution](FRHICommandList& RHICmdList)
 			{
 				FRayTracingPipelineState* Pipeline = BindRayTracingDeferredMaterialGatherPipeline(RHICmdList, View, RayGenShader->GetRayTracingShader());
@@ -392,7 +392,7 @@ void FDeferredShadingSceneRenderer::RenderRayTracingReflections(
 			GraphBuilder.AddPass(
 				RDG_EVENT_NAME("ReflectionRayTracing %dx%d", RayTracingResolution.X, RayTracingResolution.Y),
 				PassParameters,
-				ERenderGraphPassFlags::Compute,
+				ERDGPassFlags::Compute,
 				[PassParameters, this, &View, RayGenShader, RayTracingResolution, DeferredMaterialBufferNumElements, DeferredMaterialMode](FRHICommandList& RHICmdList)
 			{
 				FRayTracingShaderBindingsWriter GlobalResources;
diff --git a/Engine/Source/Runtime/Renderer/Private/RayTracing/RayTracingShadows.cpp b/Engine/Source/Runtime/Renderer/Private/RayTracing/RayTracingShadows.cpp
--- a/Engine/Source/Runtime/Renderer/Private/RayTracing/RayTracingShadows.cpp
+++ b/Engine/Source/Runtime/Renderer/Private/RayTracing/RayTracingShadows.cpp
@@ -205,7 +205,7 @@ void FDeferredShadingSceneRenderer::RenderRayTracingShadows(
 		GraphBuilder.AddPass(
 			RDG_EVENT_NAME("RayTracedShadow (spp=%d) %dx%d", RayTracingConfig.RayCountPerPixel, View.ViewRect.Width(), View.ViewRect.Height()),
 			PassParameters,
-			ERenderGraphPassFlags::Compute,
+			ERDGPassFlags::Compute,
 			[this, &View, RayGenerationShader, PassParameters, Resolution](FRHICommandList& RHICmdList)
 		{
 			FRayTracingShaderBindingsWriter GlobalResources;
diff --git a/Engine/Source/Runtime/Renderer/Private/RayTracing/RayTracingTranslucency.cpp b/Engine/Source/Runtime/Renderer/Private/RayTracing/RayTracingTranslucency.cpp
--- a/Engine/Source/Runtime/Renderer/Private/RayTracing/RayTracingTranslucency.cpp
+++ b/Engine/Source/Runtime/Renderer/Private/RayTracing/RayTracingTranslucency.cpp
@@ -409,7 +409,7 @@ void FDeferredShadingSceneRenderer::RenderRayTracingTranslucencyView(
 	GraphBuilder.AddPass(
 		RDG_EVENT_NAME("TranslucencyRayTracing %dx%d", RayTracingResolution.X, RayTracingResolution.Y),
 		PassParameters,
-		ERenderGraphPassFlags::Compute,
+		ERDGPassFlags::Compute,
 		[PassParameters, this, &View, RayGenShader, RayTracingResolution](FRHICommandList& RHICmdList)
 	{
 		FRayTracingPipelineState* Pipeline = View.RayTracingMaterialPipeline;
diff --git a/Engine/Source/Runtime/Renderer/Private/RayTracing/RaytracingSkylight.cpp b/Engine/Source/Runtime/Renderer/Private/RayTracing/RaytracingSkylight.cpp
--- a/Engine/Source/Runtime/Renderer/Private/RayTracing/RaytracingSkylight.cpp
+++ b/Engine/Source/Runtime/Renderer/Private/RayTracing/RaytracingSkylight.cpp
@@ -1036,7 +1036,7 @@ void FDeferredShadingSceneRenderer::RenderRayTracingSkyLight(
 		GraphBuilder.AddPass(
 			RDG_EVENT_NAME("SkyLightRayTracing %dx%d", RayTracingResolution.X, RayTracingResolution.Y),
 			PassParameters,
-			ERenderGraphPassFlags::Compute,
+			ERDGPassFlags::Compute,
 			[PassParameters, this, &View, RayGenerationShader, RayTracingResolution](FRHICommandList& RHICmdList)
 		{
 			FRayTracingShaderBindingsWriter GlobalResources;
@@ -1154,7 +1154,7 @@ void FDeferredShadingSceneRenderer::CompositeRayTracingSkyLight(
 		GraphBuilder.AddPass(
 			RDG_EVENT_NAME("GlobalIlluminationComposite"),
 			PassParameters,
-			ERenderGraphPassFlags::None,
+			ERDGPassFlags::Raster,
 			[this, &SceneContext, &View, PassParameters](FRHICommandListImmediate& RHICmdList)
 		{
 			TShaderMapRef<FPostProcessVS> VertexShader(View.ShaderMap);
diff --git a/Engine/Source/Runtime/Renderer/Private/ReflectionEnvironment.cpp b/Engine/Source/Runtime/Renderer/Private/ReflectionEnvironment.cpp
--- a/Engine/Source/Runtime/Renderer/Private/ReflectionEnvironment.cpp
+++ b/Engine/Source/Runtime/Renderer/Private/ReflectionEnvironment.cpp
@@ -853,7 +853,7 @@ void FDeferredShadingSceneRenderer::RenderDeferredReflectionsAndSkyLighting(FRHI
 			GraphBuilder.AddPass(
 				RDG_EVENT_NAME("ReflectionEnvironmentAndSky %dx%d", View.ViewRect.Width(), View.ViewRect.Height()),
 				PassParameters,
-				ERDGPassFlags::None,
+				ERDGPassFlags::Raster,
 				[PassParameters, &View, PixelShader](FRHICommandList& InRHICmdList)
 			{
 				InRHICmdList.SetViewport(View.ViewRect.Min.X, View.ViewRect.Min.Y, 0.0f, View.ViewRect.Max.X, View.ViewRect.Max.Y, 1.0f);
diff --git a/Engine/Source/Runtime/Renderer/Private/SceneOcclusion.cpp b/Engine/Source/Runtime/Renderer/Private/SceneOcclusion.cpp
--- a/Engine/Source/Runtime/Renderer/Private/SceneOcclusion.cpp
+++ b/Engine/Source/Runtime/Renderer/Private/SceneOcclusion.cpp
@@ -1120,7 +1120,7 @@ void BuildHZB(FRDGBuilder& GraphBuilder, FViewInfo& View)
 		GraphBuilder.AddPass(
 			RDG_EVENT_NAME("HZB(mip=0) %dx%d", HZBSize.X, HZBSize.Y),
 			PassParameters,
-			ERenderGraphPassFlags::None,
+			ERDGPassFlags::Raster,
 			[PassParameters, &View, HZBSize](FRHICommandListImmediate& RHICmdList)
 			{
 				FGraphicsPipelineStateInitializer GraphicsPSOInit;
@@ -1203,7 +1203,7 @@ void BuildHZB(FRDGBuilder& GraphBuilder, FViewInfo& View)
 		GraphBuilder.AddPass(
 			RDG_EVENT_NAME("HZB(mip=%d) %dx%d", MipIndex, DstSize.X, DstSize.Y),
 			PassParameters,
-			ERenderGraphPassFlags::Raster | ERenderGraphPassFlags::GenerateMips,
+			ERDGPassFlags::Raster | ERDGPassFlags::GenerateMips,
 			[PassParameters, SrcSize, DstSize, &View](FRHICommandListImmediate& RHICmdList)
 			{
 				FHZBBuildPS::FPermutationDomain PermutationVector;
diff --git a/Engine/Source/Runtime/Renderer/Private/ScreenPass.h b/Engine/Source/Runtime/Renderer/Private/ScreenPass.h
--- a/Engine/Source/Runtime/Renderer/Private/ScreenPass.h
+++ b/Engine/Source/Runtime/Renderer/Private/ScreenPass.h
@@ -319,7 +319,7 @@ void AddDrawScreenPass(
 	GraphBuilder.AddPass(
 		Forward<FRDGEventName>(PassName),
 		PixelShaderParameters,
-		ERenderGraphPassFlags::None,
+		ERDGPassFlags::Raster,
 		[ScreenPassView, OutputViewport, InputViewport, PixelShader, PixelShaderParameters]
 	(FRHICommandListImmediate& RHICmdList)
 	{
@@ -353,7 +353,7 @@ void AddDrawScreenPass(
 	GraphBuilder.AddPass(
 		Forward<FRDGEventName>(PassName),
 		PassParameterStruct,
-		ERenderGraphPassFlags::None,
+		ERDGPassFlags::Raster,
 		[ScreenPassView, OutputViewport, InputViewport, ScreenPassDraw, SetupFunction]
 		(FRHICommandListImmediate& RHICmdList)
 	{
diff --git a/Engine/Source/Runtime/Renderer/Private/ScreenSpaceRayTracing.cpp b/Engine/Source/Runtime/Renderer/Private/ScreenSpaceRayTracing.cpp
--- a/Engine/Source/Runtime/Renderer/Private/ScreenSpaceRayTracing.cpp
+++ b/Engine/Source/Runtime/Renderer/Private/ScreenSpaceRayTracing.cpp
@@ -362,7 +362,7 @@ FRDGTextureUAV* CreateScreenSpaceRayTracingDebugUAV(FRDGBuilder& GraphBuilder, c
 	GraphBuilder.AddPass(
 		RDG_EVENT_NAME("Clear SSRTDebug"),
 		PassParameters,
-		ERenderGraphPassFlags::None,
+		ERDGPassFlags::Raster,
 		[](FRHICommandList& RHICmdList)
 	{
 		// NOP
@@ -493,7 +493,7 @@ void RenderScreenSpaceReflections(
 		GraphBuilder.AddPass(
 			RDG_EVENT_NAME("SSR StencilSetup %dx%d", View.ViewRect.Width(), View.ViewRect.Height()),
 			PassParameters,
-			ERenderGraphPassFlags::None,
+			ERDGPassFlags::Raster,
 			[PassParameters, &View, PixelShader](FRHICommandList& RHICmdList)
 		{
 			SCOPED_GPU_STAT(RHICmdList, ScreenSpaceReflections);
@@ -576,7 +576,7 @@ void RenderScreenSpaceReflections(
 				SSRQuality, RayTracingConfigs.RayCountPerPixel, bDenoiser ? TEXT(" DenoiserOutput") : TEXT(""),
 				View.ViewRect.Width(), View.ViewRect.Height()),
 			PassParameters,
-			ERenderGraphPassFlags::None,
+			ERDGPassFlags::Raster,
 			[PassParameters, &View, PixelShader, SSRStencilPrePass](FRHICommandList& RHICmdList)
 		{
 			SCOPED_GPU_STAT(RHICmdList, ScreenSpaceReflections);
diff --git a/Engine/Source/Runtime/Renderer/Private/ShaderPrint.cpp b/Engine/Source/Runtime/Renderer/Private/ShaderPrint.cpp
--- a/Engine/Source/Runtime/Renderer/Private/ShaderPrint.cpp
+++ b/Engine/Source/Runtime/Renderer/Private/ShaderPrint.cpp
@@ -424,7 +424,7 @@ namespace ShaderPrint
 			GraphBuilder.AddPass(
 				RDG_EVENT_NAME("DrawSymbols"),
 				PassParameters,
-				ERenderGraphPassFlags::None,
+				ERDGPassFlags::Raster,
 				[VertexShader, PixelShader, PassParameters](FRHICommandListImmediate& RHICmdListImmediate)
 			{
 				FGraphicsPipelineStateInitializer GraphicsPSOInit;
diff --git a/Engine/Source/Runtime/Renderer/Private/VT/RuntimeVirtualTextureRender.cpp b/Engine/Source/Runtime/Renderer/Private/VT/RuntimeVirtualTextureRender.cpp
--- a/Engine/Source/Runtime/Renderer/Private/VT/RuntimeVirtualTextureRender.cpp
+++ b/Engine/Source/Runtime/Renderer/Private/VT/RuntimeVirtualTextureRender.cpp
@@ -658,7 +658,7 @@ namespace RuntimeVirtualTexture
 			GraphBuilder.AddPass(
 				RDG_EVENT_NAME("VirtualTextureDraw"),
 				PassParameters,
-				ERenderGraphPassFlags::None,
+				ERDGPassFlags::Raster,
 				[Scene, View, MaterialType, RuntimeVirtualTextureMask, vLevel, MaxLevel](FRHICommandListImmediate& RHICmdListImmediate)
 			{
 				DrawMeshes(RHICmdListImmediate, Scene, View, MaterialType, RuntimeVirtualTextureMask, vLevel, MaxLevel);
diff --git a/Engine/Source/Runtime/Renderer/Private/VolumetricFog.cpp b/Engine/Source/Runtime/Renderer/Private/VolumetricFog.cpp
--- a/Engine/Source/Runtime/Renderer/Private/VolumetricFog.cpp
+++ b/Engine/Source/Runtime/Renderer/Private/VolumetricFog.cpp
@@ -493,7 +493,7 @@ void FDeferredShadingSceneRenderer::RenderLocalLightsForVolumetricFog(
 		GraphBuilder.AddPass(
 			RDG_EVENT_NAME("ShadowedLights"),
 			PassParameters,
-			ERenderGraphPassFlags::None,
+			ERDGPassFlags::Raster,
 			[PassParameters, &View, this, LightsToInject, VolumetricFogGridSize, GridZParams, bUseTemporalReprojection, IntegrationData, FogInfo](FRHICommandListImmediate& RHICmdList)
 		{
 			for (int32 LightIndex = 0; LightIndex < LightsToInject.Num(); LightIndex++)
@@ -1049,7 +1049,7 @@ void FDeferredShadingSceneRenderer::ComputeVolumetricFog(FRHICommandListImmediat
 				GraphBuilder.AddPass(
 					RDG_EVENT_NAME("InitializeVolumeAttributes"),
 					PassParameters,
-					ERenderGraphPassFlags::Compute,
+					ERDGPassFlags::Compute,
 					[PassParameters, &View, VolumetricFogGridSize, IntegrationData, ComputeShader](FRHICommandListImmediate& RHICmdList)
 				{
 					const FIntVector NumGroups = FIntVector::DivideAndRoundUp(VolumetricFogGridSize, VolumetricFogGridInjectionGroupSize);
@@ -1115,7 +1115,7 @@ void FDeferredShadingSceneRenderer::ComputeVolumetricFog(FRHICommandListImmediat
 						bUseDistanceFieldSkyOcclusion ? TEXT("DFAO") : TEXT(""),
 						PassParameters->LightFunctionTexture ? TEXT("LF") : TEXT("")),
 					PassParameters,
-					ERenderGraphPassFlags::Compute,
+					ERDGPassFlags::Compute,
 					[PassParameters, ComputeShader, &View, this, FogInfo, bUseTemporalReprojection, VolumetricFogGridSize, IntegrationData, bUseDirectionalLightShadowing, bUseDistanceFieldSkyOcclusion, LightFunctionWorldToShadow](FRHICommandListImmediate& RHICmdList)
 				{
 					UnbindRenderTargets(RHICmdList);
@@ -1146,7 +1146,7 @@ void FDeferredShadingSceneRenderer::ComputeVolumetricFog(FRHICommandListImmediat
 				GraphBuilder.AddPass(
 					RDG_EVENT_NAME("FinalIntegration"),
 					PassParameters,
-					ERenderGraphPassFlags::Compute,
+					ERDGPassFlags::Compute,
 					[PassParameters, &View, VolumetricFogGridSize, IntegrationData, this](FRHICommandListImmediate& RHICmdList)
 				{
 					const FIntVector NumGroups = FIntVector::DivideAndRoundUp(VolumetricFogGridSize, VolumetricFogIntegrationGroupSize);
diff --git a/Engine/Source/Runtime/Renderer/Private/VolumetricFogLightFunction.cpp b/Engine/Source/Runtime/Renderer/Private/VolumetricFogLightFunction.cpp
--- a/Engine/Source/Runtime/Renderer/Private/VolumetricFogLightFunction.cpp
+++ b/Engine/Source/Runtime/Renderer/Private/VolumetricFogLightFunction.cpp
@@ -205,7 +205,7 @@ void FDeferredShadingSceneRenderer::RenderLightFunctionForVolumetricFog(
 			GraphBuilder.AddPass(
 				RDG_EVENT_NAME("LightFunction"),
 				PassParameters,
-				ERenderGraphPassFlags::None,
+				ERDGPassFlags::Raster,
 				[PassParameters, &View, MaterialProxy, LightFunctionResolution, DirectionalLightSceneInfo, WorldToShadowValue, this](FRHICommandListImmediate& RHICmdList)
 			{
 				const FMaterial* Material = MaterialProxy->GetMaterial(Scene->GetFeatureLevel());
diff --git a/Engine/Source/Runtime/Renderer/Private/VolumetricFogVoxelization.cpp b/Engine/Source/Runtime/Renderer/Private/VolumetricFogVoxelization.cpp
--- a/Engine/Source/Runtime/Renderer/Private/VolumetricFogVoxelization.cpp
+++ b/Engine/Source/Runtime/Renderer/Private/VolumetricFogVoxelization.cpp
@@ -596,7 +596,7 @@ void FDeferredShadingSceneRenderer::VoxelizeFogVolumePrimitives(
 		GraphBuilder.AddPass(
 			RDG_EVENT_NAME("VoxelizeVolumePrimitives"),
 			PassParameters,
-			ERenderGraphPassFlags::None,
+			ERDGPassFlags::Raster,
 			[PassParameters, Scene = Scene, &View, VolumetricFogGridSize, IntegrationData, VolumetricFogDistance, GridZParams](FRHICommandListImmediate& RHICmdList)
 		{
 			FViewUniformShaderParameters ViewVoxelizeParameters = *View.CachedViewUniformShaderParameters;
