Moves BuildHZB into its own file for upcoming more modular scene texture reduction passes.

diff --git a/Engine/Shaders/Private/HZB.usf b/Engine/Shaders/Private/HZB.usf
new file mode 100644
--- /dev/null
+++ b/Engine/Shaders/Private/HZB.usf
@@ -0,0 +1,60 @@
+// Copyright 1998-2019 Epic Games, Inc. All Rights Reserved.
+
+#include "Common.ush"
+
+float2 InvSize;
+float4 InputUvFactorAndOffset;
+float2 InputViewportMaxBound;
+
+Texture2D		Texture;
+SamplerState	TextureSampler;
+
+void HZBBuildPS(noperspective float4 RenderTargetUVAndScreenPos : TEXCOORD0, out float4 OutColor : SV_Target0)
+{
+	float4 Depth;
+
+#if STAGE == 0
+	const float2 ViewUV = RenderTargetUVAndScreenPos.zw * float2(0.5f, -0.5f) + 0.5f;
+	const float2 InUV = ViewUV * InputUvFactorAndOffset.xy + InputUvFactorAndOffset.zw;
+	
+	// min(..., InputViewportMaxBound) because we don't want to sample outside of the viewport
+	// when the view size has odd dimensions on X/Y axis.
+#if COMPILER_GLSL || COMPILER_GLSL_ES2 || COMPILER_GLSL_ES3_1 || FEATURE_LEVEL < FEATURE_LEVEL_SM5
+	float2 UV[4];
+	UV[0] = min(InUV + float2(-0.25f, -0.25f) * InvSize, InputViewportMaxBound);
+	UV[1] = min(InUV + float2( 0.25f, -0.25f) * InvSize, InputViewportMaxBound);
+	UV[2] = min(InUV + float2(-0.25f,  0.25f) * InvSize, InputViewportMaxBound);
+	UV[3] = min(InUV + float2( 0.25f,  0.25f) * InvSize, InputViewportMaxBound);
+
+	Depth.x = SceneTexturesStruct.SceneDepthTexture.SampleLevel( SceneTexturesStruct.SceneDepthTextureSampler, UV[0], 0 ).r;
+	Depth.y = SceneTexturesStruct.SceneDepthTexture.SampleLevel( SceneTexturesStruct.SceneDepthTextureSampler, UV[1], 0 ).r;
+	Depth.z = SceneTexturesStruct.SceneDepthTexture.SampleLevel( SceneTexturesStruct.SceneDepthTextureSampler, UV[2], 0 ).r;
+	Depth.w = SceneTexturesStruct.SceneDepthTexture.SampleLevel( SceneTexturesStruct.SceneDepthTextureSampler, UV[3], 0 ).r;
+#else
+	float2 UV = min(InUV + float2(-0.25f, -0.25f) * InvSize, InputViewportMaxBound);
+	Depth = SceneTexturesStruct.SceneDepthTexture.GatherRed( SceneTexturesStruct.SceneDepthTextureSampler, UV, 0 );
+#endif
+
+#else
+	const float2 InUV = RenderTargetUVAndScreenPos.xy;
+	
+#if COMPILER_GLSL || COMPILER_GLSL_ES2 || COMPILER_GLSL_ES3_1 || FEATURE_LEVEL < FEATURE_LEVEL_SM5
+	float2 UV[4];
+	UV[0] = InUV + float2(-0.25f, -0.25f) * InvSize;
+	UV[1] = InUV + float2( 0.25f, -0.25f) * InvSize;
+	UV[2] = InUV + float2(-0.25f,  0.25f) * InvSize;
+	UV[3] = InUV + float2( 0.25f,  0.25f) * InvSize;
+
+	Depth.x = Texture.SampleLevel( TextureSampler, UV[0], 0 ).r;
+	Depth.y = Texture.SampleLevel( TextureSampler, UV[1], 0 ).r;
+	Depth.z = Texture.SampleLevel( TextureSampler, UV[2], 0 ).r;
+	Depth.w = Texture.SampleLevel( TextureSampler, UV[3], 0 ).r;
+#else
+	float2 UV = InUV + float2(-0.25f, -0.25f) * InvSize;
+	Depth = Texture.GatherRed( TextureSampler, UV, 0 );
+#endif
+
+#endif
+
+	OutColor = min( min(Depth.x, Depth.y), min(Depth.z, Depth.w) );
+}
diff --git a/Engine/Shaders/Private/HZBOcclusion.usf b/Engine/Shaders/Private/HZBOcclusion.usf
--- a/Engine/Shaders/Private/HZBOcclusion.usf
+++ b/Engine/Shaders/Private/HZBOcclusion.usf
@@ -2,63 +2,6 @@
 
 #include "Common.ush"
 
-float2 InvSize;
-float4 InputUvFactorAndOffset;
-float2 InputViewportMaxBound;
-
-Texture2D		Texture;
-SamplerState	TextureSampler;
-
-void HZBBuildPS(noperspective float4 RenderTargetUVAndScreenPos : TEXCOORD0, out float4 OutColor : SV_Target0)
-{
-	float4 Depth;
-
-#if STAGE == 0
-	const float2 ViewUV = RenderTargetUVAndScreenPos.zw * float2(0.5f, -0.5f) + 0.5f;
-	const float2 InUV = ViewUV * InputUvFactorAndOffset.xy + InputUvFactorAndOffset.zw;
-	
-	// min(..., InputViewportMaxBound) because we don't want to sample outside of the viewport
-	// when the view size has odd dimensions on X/Y axis.
-#if COMPILER_GLSL || COMPILER_GLSL_ES2 || COMPILER_GLSL_ES3_1 || FEATURE_LEVEL < FEATURE_LEVEL_SM5
-	float2 UV[4];
-	UV[0] = min(InUV + float2(-0.25f, -0.25f) * InvSize, InputViewportMaxBound);
-	UV[1] = min(InUV + float2( 0.25f, -0.25f) * InvSize, InputViewportMaxBound);
-	UV[2] = min(InUV + float2(-0.25f,  0.25f) * InvSize, InputViewportMaxBound);
-	UV[3] = min(InUV + float2( 0.25f,  0.25f) * InvSize, InputViewportMaxBound);
-
-	Depth.x = SceneTexturesStruct.SceneDepthTexture.SampleLevel( SceneTexturesStruct.SceneDepthTextureSampler, UV[0], 0 ).r;
-	Depth.y = SceneTexturesStruct.SceneDepthTexture.SampleLevel( SceneTexturesStruct.SceneDepthTextureSampler, UV[1], 0 ).r;
-	Depth.z = SceneTexturesStruct.SceneDepthTexture.SampleLevel( SceneTexturesStruct.SceneDepthTextureSampler, UV[2], 0 ).r;
-	Depth.w = SceneTexturesStruct.SceneDepthTexture.SampleLevel( SceneTexturesStruct.SceneDepthTextureSampler, UV[3], 0 ).r;
-#else
-	float2 UV = min(InUV + float2(-0.25f, -0.25f) * InvSize, InputViewportMaxBound);
-	Depth = SceneTexturesStruct.SceneDepthTexture.GatherRed( SceneTexturesStruct.SceneDepthTextureSampler, UV, 0 );
-#endif
-
-#else
-	const float2 InUV = RenderTargetUVAndScreenPos.xy;
-	
-#if COMPILER_GLSL || COMPILER_GLSL_ES2 || COMPILER_GLSL_ES3_1 || FEATURE_LEVEL < FEATURE_LEVEL_SM5
-	float2 UV[4];
-	UV[0] = InUV + float2(-0.25f, -0.25f) * InvSize;
-	UV[1] = InUV + float2( 0.25f, -0.25f) * InvSize;
-	UV[2] = InUV + float2(-0.25f,  0.25f) * InvSize;
-	UV[3] = InUV + float2( 0.25f,  0.25f) * InvSize;
-
-	Depth.x = Texture.SampleLevel( TextureSampler, UV[0], 0 ).r;
-	Depth.y = Texture.SampleLevel( TextureSampler, UV[1], 0 ).r;
-	Depth.z = Texture.SampleLevel( TextureSampler, UV[2], 0 ).r;
-	Depth.w = Texture.SampleLevel( TextureSampler, UV[3], 0 ).r;
-#else
-	float2 UV = InUV + float2(-0.25f, -0.25f) * InvSize;
-	Depth = Texture.GatherRed( TextureSampler, UV, 0 );
-#endif
-
-#endif
-
-	OutColor = min( min(Depth.x, Depth.y), min(Depth.z, Depth.w) );
-}
-
 Texture2D		HZBTexture;
 SamplerState	HZBSampler;
 
diff --git a/Engine/Source/Runtime/Renderer/Private/SceneOcclusion.cpp b/Engine/Source/Runtime/Renderer/Private/SceneOcclusion.cpp
--- a/Engine/Source/Runtime/Renderer/Private/SceneOcclusion.cpp
+++ b/Engine/Source/Runtime/Renderer/Private/SceneOcclusion.cpp
@@ -1052,208 +1052,6 @@ void FHZBOcclusionTester::Submit(FRHICommandListImmediate& RHICmdList, const FVi
 	RHICmdList.CopyToResolveTarget(ResultsTextureGPU->GetRenderTargetItem().TargetableTexture, ResultsTextureCPU->GetRenderTargetItem().ShaderResourceTexture, FResolveParams());
 }
 
-BEGIN_SHADER_PARAMETER_STRUCT(FHZBBuildPassParameters, )
-	RENDER_TARGET_BINDING_SLOTS()
-	SHADER_PARAMETER_RDG_TEXTURE_SRV(Texture2D, Texture)
-	SHADER_PARAMETER_SAMPLER(SamplerState, TextureSampler)
-END_SHADER_PARAMETER_STRUCT()
-
-class FHZBBuildPS : public FGlobalShader
-{
-	DECLARE_GLOBAL_SHADER(FHZBBuildPS);
-	SHADER_USE_PARAMETER_STRUCT(FHZBBuildPS, FGlobalShader)
-
-	class FStageDim : SHADER_PERMUTATION_BOOL("STAGE");
-	using FPermutationDomain = TShaderPermutationDomain<FStageDim>;
-	
-	BEGIN_SHADER_PARAMETER_STRUCT( FParameters, )
-		SHADER_PARAMETER( FVector2D,	InvSize )
-		SHADER_PARAMETER( FVector4,		InputUvFactorAndOffset )
-		SHADER_PARAMETER( FVector2D,	InputViewportMaxBound )
-
-		SHADER_PARAMETER_STRUCT_INCLUDE(FHZBBuildPassParameters, Pass)
-		SHADER_PARAMETER_STRUCT_REF(	FViewUniformShaderParameters,	View )
-		SHADER_PARAMETER_STRUCT_REF(	FSceneTexturesUniformParameters,SceneTextures )
-	END_SHADER_PARAMETER_STRUCT()
-
-	static bool ShouldCompilePermutation(const FGlobalShaderPermutationParameters& Parameters)
-	{
-		return IsFeatureLevelSupported(Parameters.Platform, ERHIFeatureLevel::SM4);
-	}
-
-	static void ModifyCompilationEnvironment(const FGlobalShaderPermutationParameters& Parameters, FShaderCompilerEnvironment& OutEnvironment)
-	{
-		OutEnvironment.SetRenderTargetOutputFormat(0, PF_R32_FLOAT);
-	}
-};
-
-IMPLEMENT_GLOBAL_SHADER(FHZBBuildPS, "/Engine/Private/HZBOcclusion.usf", "HZBBuildPS", SF_Pixel);
-
-
-void BuildHZB(FRDGBuilder& GraphBuilder, FViewInfo& View)
-{
-	QUICK_SCOPE_CYCLE_COUNTER(STAT_BuildHZB);
-	
-	// View.ViewRect.{Width,Height}() are most likely to be < 2^24, so the float
-	// conversion won't loss any precision (assuming float have 23bits for mantissa)
-	const int32 NumMipsX = FMath::Max(FPlatformMath::CeilToInt(FMath::Log2(float(View.ViewRect.Width()))) - 1, 1);
-	const int32 NumMipsY = FMath::Max(FPlatformMath::CeilToInt(FMath::Log2(float(View.ViewRect.Height()))) - 1, 1);
-	const uint32 NumMips = FMath::Max(NumMipsX, NumMipsY);
-
-	// Must be power of 2
-	const FIntPoint HZBSize( 1 << NumMipsX, 1 << NumMipsY );
-	View.HZBMipmap0Size = HZBSize;
-
-	//@DW: Configure texture creation
-	FRDGTextureDesc HZBDesc = FRDGTextureDesc::Create2DDesc(HZBSize, PF_R16F, FClearValueBinding::None, TexCreate_None, TexCreate_RenderTargetable | TexCreate_ShaderResource | TexCreate_NoFastClear, false, NumMips);
-	HZBDesc.Flags |= GFastVRamConfig.HZB;
-
-	//@DW: Explicit creation of graph resource handles - full support for everything the RHI supports
-	//@DW: Now that we've created a resource handle, it will have to be passed around to other passes for manual wiring or put into a Blackboard structure for automatic wiring
-	FRDGTextureRef HZBTexture = GraphBuilder.CreateTexture(HZBDesc, TEXT("HZB"));
-
-	{
-		FHZBBuildPassParameters* PassParameters = GraphBuilder.AllocParameters<FHZBBuildPassParameters>();
-		PassParameters->RenderTargets[0] = FRenderTargetBinding( HZBTexture, ERenderTargetLoadAction::ENoAction, ERenderTargetStoreAction::EStore);
-
-		//@DW - this pass only reads external textures, we don't have any graph inputs
-		GraphBuilder.AddPass(
-			RDG_EVENT_NAME("HZB(mip=0) %dx%d", HZBSize.X, HZBSize.Y),
-			PassParameters,
-			ERDGPassFlags::Raster,
-			[PassParameters, &View, HZBSize](FRHICommandListImmediate& RHICmdList)
-			{
-				FGraphicsPipelineStateInitializer GraphicsPSOInit;
-				GraphicsPSOInit.BlendState = TStaticBlendState<>::GetRHI();
-				GraphicsPSOInit.RasterizerState = TStaticRasterizerState<>::GetRHI();
-				GraphicsPSOInit.DepthStencilState = TStaticDepthStencilState<false, CF_Always>::GetRHI();
-
-				RHICmdList.ApplyCachedRenderTargets(GraphicsPSOInit);
-				GraphicsPSOInit.PrimitiveType = PT_TriangleList;
-
-				FHZBBuildPS::FPermutationDomain PermutationVector;
-				PermutationVector.Set<FHZBBuildPS::FStageDim>(false);
-
-				TShaderMapRef<FPostProcessVS> VertexShader(View.ShaderMap);
-				TShaderMapRef<FHZBBuildPS> PixelShader(View.ShaderMap, PermutationVector);
-
-				GraphicsPSOInit.BoundShaderState.VertexDeclarationRHI = GFilterVertexDeclaration.VertexDeclarationRHI;
-				GraphicsPSOInit.BoundShaderState.VertexShaderRHI = GETSAFERHISHADER_VERTEX(*VertexShader);
-				GraphicsPSOInit.BoundShaderState.PixelShaderRHI = GETSAFERHISHADER_PIXEL(*PixelShader);
-				GraphicsPSOInit.PrimitiveType = PT_TriangleList;
-
-				SetGraphicsPipelineState(RHICmdList, GraphicsPSOInit);
-
-				// Imperfect sampling, doesn't matter too much
-				FSceneRenderTargets& SceneContext = FSceneRenderTargets::Get(RHICmdList);
-				FIntPoint Size = SceneContext.GetBufferSizeXY();
-
-				FHZBBuildPS::FParameters Parameters;
-				Parameters.InvSize = FVector2D(1.0f / Size.X, 1.0f / Size.Y);
-				Parameters.InputUvFactorAndOffset = FVector4(
-					float(2 * HZBSize.X) / float(Size.X),
-					float(2 * HZBSize.Y) / float(Size.Y),
-					float(View.ViewRect.Min.X) / float(Size.X),
-					float(View.ViewRect.Min.Y) / float(Size.Y));
-				Parameters.InputViewportMaxBound = FVector2D(
-					float(View.ViewRect.Max.X) / float(Size.X) - 0.5f * Parameters.InvSize.X,
-					float(View.ViewRect.Max.Y) / float(Size.Y) - 0.5f * Parameters.InvSize.Y);
-		
-				Parameters.Pass = *PassParameters;
-				Parameters.View = View.ViewUniformBuffer;
-				Parameters.SceneTextures = CreateSceneTextureUniformBufferSingleDraw(RHICmdList, ESceneTextureSetupMode::SceneDepth, View.FeatureLevel);
-
-				SetShaderParameters(RHICmdList, *PixelShader, PixelShader->GetPixelShader(), Parameters);
-
-				RHICmdList.SetViewport(0, 0, 0.0f, HZBSize.X, HZBSize.Y, 1.0f);
-
-				DrawRectangle(
-					RHICmdList,
-					0, 0,
-					HZBSize.X, HZBSize.Y,
-					View.ViewRect.Min.X, View.ViewRect.Min.Y,
-					View.ViewRect.Width(), View.ViewRect.Height(),
-					HZBSize,
-					FSceneRenderTargets::Get(RHICmdList).GetBufferSizeXY(),
-					*VertexShader,
-					EDRF_UseTriangleOptimization);
-			});
-	}
-
-	FIntPoint SrcSize = HZBSize;
-	FIntPoint DstSize = SrcSize / 2;
-
-	// Downsampling...
-	for (uint8 MipIndex = 1; MipIndex < NumMips; MipIndex++)
-	{
-		SrcSize.X = FMath::Max(SrcSize.X, 1);
-		SrcSize.Y = FMath::Max(SrcSize.Y, 1);
-		DstSize.X = FMath::Max(DstSize.X, 1);
-		DstSize.Y = FMath::Max(DstSize.Y, 1);
-
-		//@DW: Explicit creation of SRV, full configuration of SRV supported
-		FRDGTextureSRVDesc Desc(HZBTexture, MipIndex - 1);
-		FRDGTextureSRVRef ParentMipSRV = GraphBuilder.CreateSRV(Desc);
-
-		FHZBBuildPassParameters* PassParameters = GraphBuilder.AllocParameters<FHZBBuildPassParameters>();
-		PassParameters->RenderTargets[0] = FRenderTargetBinding( HZBTexture, ERenderTargetLoadAction::ENoAction, ERenderTargetStoreAction::EStore, MipIndex );
-		PassParameters->Texture = ParentMipSRV;
-		PassParameters->TextureSampler = TStaticSamplerState<SF_Point, AM_Clamp, AM_Clamp, AM_Clamp>::GetRHI();
-
-		GraphBuilder.AddPass(
-			RDG_EVENT_NAME("HZB(mip=%d) %dx%d", MipIndex, DstSize.X, DstSize.Y),
-			PassParameters,
-			ERDGPassFlags::Raster | ERDGPassFlags::GenerateMips,
-			[PassParameters, SrcSize, DstSize, &View](FRHICommandListImmediate& RHICmdList)
-			{
-				FHZBBuildPS::FPermutationDomain PermutationVector;
-				PermutationVector.Set<FHZBBuildPS::FStageDim>(true);
-
-				TShaderMapRef<FPostProcessVS> VertexShader(View.ShaderMap);
-				TShaderMapRef<FHZBBuildPS> PixelShader(View.ShaderMap, PermutationVector);
-
-				FGraphicsPipelineStateInitializer GraphicsPSOInit;
-				GraphicsPSOInit.BlendState = TStaticBlendState<>::GetRHI();
-				GraphicsPSOInit.RasterizerState = TStaticRasterizerState<>::GetRHI();
-				GraphicsPSOInit.DepthStencilState = TStaticDepthStencilState<false, CF_Always>::GetRHI();
-
-				RHICmdList.ApplyCachedRenderTargets(GraphicsPSOInit);
-
-				GraphicsPSOInit.BoundShaderState.VertexDeclarationRHI = GFilterVertexDeclaration.VertexDeclarationRHI;
-				GraphicsPSOInit.BoundShaderState.VertexShaderRHI = GETSAFERHISHADER_VERTEX(*VertexShader);
-				GraphicsPSOInit.BoundShaderState.PixelShaderRHI = GETSAFERHISHADER_PIXEL(*PixelShader);
-				GraphicsPSOInit.PrimitiveType = PT_TriangleList;
-
-				SetGraphicsPipelineState(RHICmdList, GraphicsPSOInit);
-
-				FHZBBuildPS::FParameters Parameters;
-				Parameters.InvSize = FVector2D(1.0f / SrcSize.X, 1.0f / SrcSize.Y);
-				Parameters.Pass = *PassParameters;
-				Parameters.View = View.ViewUniformBuffer;
-
-				SetShaderParameters(RHICmdList, *PixelShader, PixelShader->GetPixelShader(), Parameters);
-
-				RHICmdList.SetViewport(0, 0, 0.0f, DstSize.X, DstSize.Y, 1.0f);
-
-				DrawRectangle(
-					RHICmdList,
-					0, 0,
-					DstSize.X, DstSize.Y,
-					0, 0,
-					SrcSize.X, SrcSize.Y,
-					DstSize,
-					SrcSize,
-					*VertexShader,
-					EDRF_UseTriangleOptimization);
-			});
-
-		SrcSize /= 2;
-		DstSize /= 2;
-	}
-
-	GraphBuilder.QueueTextureExtraction(HZBTexture, &View.HZB);
-}
-
 struct FViewOcclusionQueries
 {
 	TArray<FProjectedShadowInfo const*> PointLightQuerieInfos;
diff --git a/Engine/Source/Runtime/Renderer/Private/SceneTextureReductions.cpp b/Engine/Source/Runtime/Renderer/Private/SceneTextureReductions.cpp
new file mode 100644
--- /dev/null
+++ b/Engine/Source/Runtime/Renderer/Private/SceneTextureReductions.cpp
@@ -0,0 +1,214 @@
+// Copyright 1998-2019 Epic Games, Inc. All Rights Reserved.
+
+/*=============================================================================
+	SceneRendering.cpp: Scene rendering.
+=============================================================================*/
+
+#include "ScenePrivate.h"
+#include "RenderGraph.h"
+#include "PixelShaderUtils.h"
+#include "SceneTextureParameters.h"
+#include "PostProcess/PostProcessing.h" // for FPostProcessVS
+
+
+BEGIN_SHADER_PARAMETER_STRUCT(FHZBBuildPassParameters, )
+	RENDER_TARGET_BINDING_SLOTS()
+	SHADER_PARAMETER_RDG_TEXTURE_SRV(Texture2D, Texture)
+	SHADER_PARAMETER_SAMPLER(SamplerState, TextureSampler)
+END_SHADER_PARAMETER_STRUCT()
+
+class FHZBBuildPS : public FGlobalShader
+{
+	DECLARE_GLOBAL_SHADER(FHZBBuildPS);
+	SHADER_USE_PARAMETER_STRUCT(FHZBBuildPS, FGlobalShader)
+
+	class FStageDim : SHADER_PERMUTATION_BOOL("STAGE");
+	using FPermutationDomain = TShaderPermutationDomain<FStageDim>;
+	
+	BEGIN_SHADER_PARAMETER_STRUCT( FParameters, )
+		SHADER_PARAMETER( FVector2D,	InvSize )
+		SHADER_PARAMETER( FVector4,		InputUvFactorAndOffset )
+		SHADER_PARAMETER( FVector2D,	InputViewportMaxBound )
+
+		SHADER_PARAMETER_STRUCT_INCLUDE(FHZBBuildPassParameters, Pass)
+		SHADER_PARAMETER_STRUCT_REF(	FViewUniformShaderParameters,	View )
+		SHADER_PARAMETER_STRUCT_REF(	FSceneTexturesUniformParameters,SceneTextures )
+	END_SHADER_PARAMETER_STRUCT()
+
+	static bool ShouldCompilePermutation(const FGlobalShaderPermutationParameters& Parameters)
+	{
+		return IsFeatureLevelSupported(Parameters.Platform, ERHIFeatureLevel::SM4);
+	}
+
+	static void ModifyCompilationEnvironment(const FGlobalShaderPermutationParameters& Parameters, FShaderCompilerEnvironment& OutEnvironment)
+	{
+		OutEnvironment.SetRenderTargetOutputFormat(0, PF_R32_FLOAT);
+	}
+};
+
+IMPLEMENT_GLOBAL_SHADER(FHZBBuildPS, "/Engine/Private/HZB.usf", "HZBBuildPS", SF_Pixel);
+
+
+void BuildHZB(FRDGBuilder& GraphBuilder, FViewInfo& View)
+{
+	QUICK_SCOPE_CYCLE_COUNTER(STAT_BuildHZB);
+	
+	// View.ViewRect.{Width,Height}() are most likely to be < 2^24, so the float
+	// conversion won't loss any precision (assuming float have 23bits for mantissa)
+	const int32 NumMipsX = FMath::Max(FPlatformMath::CeilToInt(FMath::Log2(float(View.ViewRect.Width()))) - 1, 1);
+	const int32 NumMipsY = FMath::Max(FPlatformMath::CeilToInt(FMath::Log2(float(View.ViewRect.Height()))) - 1, 1);
+	const uint32 NumMips = FMath::Max(NumMipsX, NumMipsY);
+
+	// Must be power of 2
+	const FIntPoint HZBSize( 1 << NumMipsX, 1 << NumMipsY );
+	View.HZBMipmap0Size = HZBSize;
+
+	//@DW: Configure texture creation
+	FRDGTextureDesc HZBDesc = FRDGTextureDesc::Create2DDesc(HZBSize, PF_R16F, FClearValueBinding::None, TexCreate_None, TexCreate_RenderTargetable | TexCreate_ShaderResource | TexCreate_NoFastClear, false, NumMips);
+	HZBDesc.Flags |= GFastVRamConfig.HZB;
+
+	//@DW: Explicit creation of graph resource handles - full support for everything the RHI supports
+	//@DW: Now that we've created a resource handle, it will have to be passed around to other passes for manual wiring or put into a Blackboard structure for automatic wiring
+	FRDGTextureRef HZBTexture = GraphBuilder.CreateTexture(HZBDesc, TEXT("HZB"));
+
+	{
+		FHZBBuildPassParameters* PassParameters = GraphBuilder.AllocParameters<FHZBBuildPassParameters>();
+		PassParameters->RenderTargets[0] = FRenderTargetBinding( HZBTexture, ERenderTargetLoadAction::ENoAction, ERenderTargetStoreAction::EStore);
+
+		//@DW - this pass only reads external textures, we don't have any graph inputs
+		GraphBuilder.AddPass(
+			RDG_EVENT_NAME("HZB(mip=0) %dx%d", HZBSize.X, HZBSize.Y),
+			PassParameters,
+			ERDGPassFlags::Raster,
+			[PassParameters, &View, HZBSize](FRHICommandListImmediate& RHICmdList)
+			{
+				FGraphicsPipelineStateInitializer GraphicsPSOInit;
+				GraphicsPSOInit.BlendState = TStaticBlendState<>::GetRHI();
+				GraphicsPSOInit.RasterizerState = TStaticRasterizerState<>::GetRHI();
+				GraphicsPSOInit.DepthStencilState = TStaticDepthStencilState<false, CF_Always>::GetRHI();
+
+				RHICmdList.ApplyCachedRenderTargets(GraphicsPSOInit);
+				GraphicsPSOInit.PrimitiveType = PT_TriangleList;
+
+				FHZBBuildPS::FPermutationDomain PermutationVector;
+				PermutationVector.Set<FHZBBuildPS::FStageDim>(false);
+
+				TShaderMapRef<FPostProcessVS> VertexShader(View.ShaderMap);
+				TShaderMapRef<FHZBBuildPS> PixelShader(View.ShaderMap, PermutationVector);
+
+				GraphicsPSOInit.BoundShaderState.VertexDeclarationRHI = GFilterVertexDeclaration.VertexDeclarationRHI;
+				GraphicsPSOInit.BoundShaderState.VertexShaderRHI = GETSAFERHISHADER_VERTEX(*VertexShader);
+				GraphicsPSOInit.BoundShaderState.PixelShaderRHI = GETSAFERHISHADER_PIXEL(*PixelShader);
+				GraphicsPSOInit.PrimitiveType = PT_TriangleList;
+
+				SetGraphicsPipelineState(RHICmdList, GraphicsPSOInit);
+
+				// Imperfect sampling, doesn't matter too much
+				FSceneRenderTargets& SceneContext = FSceneRenderTargets::Get(RHICmdList);
+				FIntPoint Size = SceneContext.GetBufferSizeXY();
+
+				FHZBBuildPS::FParameters Parameters;
+				Parameters.InvSize = FVector2D(1.0f / Size.X, 1.0f / Size.Y);
+				Parameters.InputUvFactorAndOffset = FVector4(
+					float(2 * HZBSize.X) / float(Size.X),
+					float(2 * HZBSize.Y) / float(Size.Y),
+					float(View.ViewRect.Min.X) / float(Size.X),
+					float(View.ViewRect.Min.Y) / float(Size.Y));
+				Parameters.InputViewportMaxBound = FVector2D(
+					float(View.ViewRect.Max.X) / float(Size.X) - 0.5f * Parameters.InvSize.X,
+					float(View.ViewRect.Max.Y) / float(Size.Y) - 0.5f * Parameters.InvSize.Y);
+		
+				Parameters.Pass = *PassParameters;
+				Parameters.View = View.ViewUniformBuffer;
+				Parameters.SceneTextures = CreateSceneTextureUniformBufferSingleDraw(RHICmdList, ESceneTextureSetupMode::SceneDepth, View.FeatureLevel);
+
+				SetShaderParameters(RHICmdList, *PixelShader, PixelShader->GetPixelShader(), Parameters);
+
+				RHICmdList.SetViewport(0, 0, 0.0f, HZBSize.X, HZBSize.Y, 1.0f);
+
+				DrawRectangle(
+					RHICmdList,
+					0, 0,
+					HZBSize.X, HZBSize.Y,
+					View.ViewRect.Min.X, View.ViewRect.Min.Y,
+					View.ViewRect.Width(), View.ViewRect.Height(),
+					HZBSize,
+					FSceneRenderTargets::Get(RHICmdList).GetBufferSizeXY(),
+					*VertexShader,
+					EDRF_UseTriangleOptimization);
+			});
+	}
+
+	FIntPoint SrcSize = HZBSize;
+	FIntPoint DstSize = SrcSize / 2;
+
+	// Downsampling...
+	for (uint8 MipIndex = 1; MipIndex < NumMips; MipIndex++)
+	{
+		SrcSize.X = FMath::Max(SrcSize.X, 1);
+		SrcSize.Y = FMath::Max(SrcSize.Y, 1);
+		DstSize.X = FMath::Max(DstSize.X, 1);
+		DstSize.Y = FMath::Max(DstSize.Y, 1);
+
+		//@DW: Explicit creation of SRV, full configuration of SRV supported
+		FRDGTextureSRVDesc Desc(HZBTexture, MipIndex - 1);
+		FRDGTextureSRVRef ParentMipSRV = GraphBuilder.CreateSRV(Desc);
+
+		FHZBBuildPassParameters* PassParameters = GraphBuilder.AllocParameters<FHZBBuildPassParameters>();
+		PassParameters->RenderTargets[0] = FRenderTargetBinding( HZBTexture, ERenderTargetLoadAction::ENoAction, ERenderTargetStoreAction::EStore, MipIndex );
+		PassParameters->Texture = ParentMipSRV;
+		PassParameters->TextureSampler = TStaticSamplerState<SF_Point, AM_Clamp, AM_Clamp, AM_Clamp>::GetRHI();
+
+		GraphBuilder.AddPass(
+			RDG_EVENT_NAME("HZB(mip=%d) %dx%d", MipIndex, DstSize.X, DstSize.Y),
+			PassParameters,
+			ERDGPassFlags::Raster | ERDGPassFlags::GenerateMips,
+			[PassParameters, SrcSize, DstSize, &View](FRHICommandListImmediate& RHICmdList)
+			{
+				FHZBBuildPS::FPermutationDomain PermutationVector;
+				PermutationVector.Set<FHZBBuildPS::FStageDim>(true);
+
+				TShaderMapRef<FPostProcessVS> VertexShader(View.ShaderMap);
+				TShaderMapRef<FHZBBuildPS> PixelShader(View.ShaderMap, PermutationVector);
+
+				FGraphicsPipelineStateInitializer GraphicsPSOInit;
+				GraphicsPSOInit.BlendState = TStaticBlendState<>::GetRHI();
+				GraphicsPSOInit.RasterizerState = TStaticRasterizerState<>::GetRHI();
+				GraphicsPSOInit.DepthStencilState = TStaticDepthStencilState<false, CF_Always>::GetRHI();
+
+				RHICmdList.ApplyCachedRenderTargets(GraphicsPSOInit);
+
+				GraphicsPSOInit.BoundShaderState.VertexDeclarationRHI = GFilterVertexDeclaration.VertexDeclarationRHI;
+				GraphicsPSOInit.BoundShaderState.VertexShaderRHI = GETSAFERHISHADER_VERTEX(*VertexShader);
+				GraphicsPSOInit.BoundShaderState.PixelShaderRHI = GETSAFERHISHADER_PIXEL(*PixelShader);
+				GraphicsPSOInit.PrimitiveType = PT_TriangleList;
+
+				SetGraphicsPipelineState(RHICmdList, GraphicsPSOInit);
+
+				FHZBBuildPS::FParameters Parameters;
+				Parameters.InvSize = FVector2D(1.0f / SrcSize.X, 1.0f / SrcSize.Y);
+				Parameters.Pass = *PassParameters;
+				Parameters.View = View.ViewUniformBuffer;
+
+				SetShaderParameters(RHICmdList, *PixelShader, PixelShader->GetPixelShader(), Parameters);
+
+				RHICmdList.SetViewport(0, 0, 0.0f, DstSize.X, DstSize.Y, 1.0f);
+
+				DrawRectangle(
+					RHICmdList,
+					0, 0,
+					DstSize.X, DstSize.Y,
+					0, 0,
+					SrcSize.X, SrcSize.Y,
+					DstSize,
+					SrcSize,
+					*VertexShader,
+					EDRF_UseTriangleOptimization);
+			});
+
+		SrcSize /= 2;
+		DstSize /= 2;
+	}
+
+	GraphBuilder.QueueTextureExtraction(HZBTexture, &View.HZB);
+}
