Fixes unecessary TexCreate_RenderTargetable and TexCreate_UAV flags in the renderer that are being caught by RDG

diff --git a/Engine/Source/Runtime/Renderer/Private/PostProcess/DiaphragmDOF.cpp b/Engine/Source/Runtime/Renderer/Private/PostProcess/DiaphragmDOF.cpp
--- a/Engine/Source/Runtime/Renderer/Private/PostProcess/DiaphragmDOF.cpp
+++ b/Engine/Source/Runtime/Renderer/Private/PostProcess/DiaphragmDOF.cpp
@@ -1433,7 +1433,8 @@ FRDGTextureRef DiaphragmDOF::AddPasses(
 
 		FullResDesc.Format = PF_FloatRGBA;
 		FullResDesc.TargetableFlags |= TexCreate_UAV;
-		FullResDesc.Flags &= ~(TexCreate_FastVRAM);
+		FullResDesc.TargetableFlags &= ~TexCreate_RenderTargetable;
+		FullResDesc.Flags &= ~TexCreate_FastVRAM;
 	}
 	
 	FDOFGatherInputDescs FullResGatherInputDescs;
@@ -1564,11 +1565,9 @@ FRDGTextureRef DiaphragmDOF::AddPasses(
 		HalfResGatherInputTextures.SeparateCoc = TAAOutputs.SceneMetadata;
 
 		HalfResGatherInputDescs.SceneColor = TAAOutputs.SceneColor->Desc;
-		HalfResGatherInputDescs.SceneColor.TargetableFlags |= TexCreate_UAV;
 		if (TAAOutputs.SceneMetadata)
 		{
 			HalfResGatherInputDescs.SeparateCoc = TAAOutputs.SceneMetadata->Desc;
-			HalfResGatherInputDescs.SeparateCoc.TargetableFlags |= TexCreate_UAV;
 		}
 	}
 
@@ -1581,10 +1580,10 @@ FRDGTextureRef DiaphragmDOF::AddPasses(
 			FIntPoint MaxTileCount = CocTileGridSize(HalfResGatherInputTextures.SceneColor->Desc.Extent);
 			
 			TileClassificationDescs.Foreground = FPooledRenderTargetDesc::Create2DDesc(
-				MaxTileCount, PF_G16R16F, FClearValueBinding::None, TexCreate_None, TexCreate_RenderTargetable | TexCreate_UAV, false);
+				MaxTileCount, PF_G16R16F, FClearValueBinding::None, TexCreate_None, TexCreate_ShaderResource | TexCreate_UAV, false);
 
 			TileClassificationDescs.Background = FPooledRenderTargetDesc::Create2DDesc(
-				MaxTileCount, PF_FloatRGBA, FClearValueBinding::None, TexCreate_None, TexCreate_RenderTargetable | TexCreate_UAV, false);
+				MaxTileCount, PF_FloatRGBA, FClearValueBinding::None, TexCreate_None, TexCreate_ShaderResource | TexCreate_UAV, false);
 		}
 
 		// Adds a coc flatten pass.
@@ -2014,6 +2013,9 @@ FRDGTextureRef DiaphragmDOF::AddPasses(
 			/** Bokeh simulation to do. */
 			EDiaphragmDOFBokehSimulation BokehSimulation;
 		
+			/** Whether there is a scattering pass. */
+			bool bHasScatterPass = false;
+
 			FConvolutionSettings()
 				: LayerProcessing(EDiaphragmDOFLayerProcessing::ForegroundAndBackground)
 				, QualityConfig(EDiaphragmDOFGatherQuality::HighQuality)
@@ -2034,7 +2036,7 @@ FRDGTextureRef DiaphragmDOF::AddPasses(
 				FRDGTextureDesc Desc = ReducedGatherInputTextures.SceneColor->Desc;
 				Desc.Extent = RefBufferSize;
 				Desc.Format = PF_FloatRGBA;
-				Desc.TargetableFlags |= TexCreate_RenderTargetable | TexCreate_UAV;
+				Desc.TargetableFlags |= TexCreate_UAV;
 				Desc.NumMips = 1;
 
 				{
@@ -2050,7 +2052,14 @@ FRDGTextureRef DiaphragmDOF::AddPasses(
 					else
 						check(0);
 
-					ConvolutionOutputTextures->SceneColor = GraphBuilder.CreateTexture(Desc, DebugName);
+					FRDGTextureDesc SceneColorDesc = Desc;
+					// Scattering pass will be drawing directly into the color of the gathered texture, so need to be render targetable.
+					if (ConvolutionSettings.bHasScatterPass && ConvolutionSettings.PostfilterMethod == EDiaphragmDOFPostfilterMethod::None)
+					{
+						SceneColorDesc.TargetableFlags |= TexCreate_RenderTargetable;
+					}
+
+					ConvolutionOutputTextures->SceneColor = GraphBuilder.CreateTexture(SceneColorDesc, DebugName);
 				
 					if (bProcessSceneAlpha)
 					{
@@ -2181,11 +2190,29 @@ FRDGTextureRef DiaphragmDOF::AddPasses(
 			const FConvolutionSettings& ConvolutionSettings,
 			FDOFConvolutionTextures* ConvolutionTextures)
 		{
+			if (ConvolutionSettings.PostfilterMethod == EDiaphragmDOFPostfilterMethod::None)
+			{
+				return;
+			}
+
 			FDOFConvolutionTextures NewConvolutionTextures;
-			NewConvolutionTextures.SceneColor = GraphBuilder.CreateTexture(ConvolutionTextures->SceneColor->Desc, ConvolutionTextures->SceneColor->Name);
+			{
+				FRDGTextureDesc Desc = ConvolutionTextures->SceneColor->Desc;
+
+				// Scattering pass will be drawing directly into the post filtered color texture, so need to be render targetable.
+				if (ConvolutionSettings.bHasScatterPass)
+				{
+					Desc.TargetableFlags |= TexCreate_RenderTargetable;
+				}
+
+				NewConvolutionTextures.SceneColor = GraphBuilder.CreateTexture(Desc, ConvolutionTextures->SceneColor->Name);
+			}
+
 			if (ConvolutionTextures->SeparateAlpha)
+			{
 				NewConvolutionTextures.SeparateAlpha = GraphBuilder.CreateTexture(ConvolutionTextures->SeparateAlpha->Desc, ConvolutionTextures->SeparateAlpha->Name);
-			
+			}
+
 			FDiaphragmDOFPostfilterCS::FPermutationDomain PermutationVector;
 			PermutationVector.Set<FDDOFLayerProcessingDim>(ConvolutionSettings.LayerProcessing);
 			PermutationVector.Set<FDDOFPostfilterMethodDim>(ConvolutionSettings.PostfilterMethod);
@@ -2318,6 +2345,7 @@ FRDGTextureRef DiaphragmDOF::AddPasses(
 			FConvolutionSettings ConvolutionSettings;
 			ConvolutionSettings.LayerProcessing = EDiaphragmDOFLayerProcessing::ForegroundOnly;
 			ConvolutionSettings.PostfilterMethod = PostfilterMethod;
+			ConvolutionSettings.bHasScatterPass = bForegroundHybridScattering;
 
 			if (bEnableGatherBokehSettings)
 				ConvolutionSettings.BokehSimulation = BokehSimulation;
@@ -2366,6 +2394,7 @@ FRDGTextureRef DiaphragmDOF::AddPasses(
 			FConvolutionSettings ConvolutionSettings;
 			ConvolutionSettings.LayerProcessing = EDiaphragmDOFLayerProcessing::BackgroundOnly;
 			ConvolutionSettings.PostfilterMethod = PostfilterMethod;
+			ConvolutionSettings.bHasScatterPass = bBackgroundHybridScattering;
 
 			if (bEnableGatherBokehSettings)
 				ConvolutionSettings.BokehSimulation = BokehSimulation;
@@ -2399,6 +2428,7 @@ FRDGTextureRef DiaphragmDOF::AddPasses(
 			FRDGTextureDesc Desc = InputSceneColor->Desc;
 			Desc.NumSamples = 1;
 			Desc.TargetableFlags |= TexCreate_UAV;
+			Desc.TargetableFlags &= ~TexCreate_RenderTargetable;
 			NewSceneColor = GraphBuilder.CreateTexture(Desc, TEXT("DOFRecombine"));
 		}
 		
diff --git a/Engine/Source/Runtime/Renderer/Private/PostProcess/PostProcessMotionBlur.cpp b/Engine/Source/Runtime/Renderer/Private/PostProcess/PostProcessMotionBlur.cpp
--- a/Engine/Source/Runtime/Renderer/Private/PostProcess/PostProcessMotionBlur.cpp
+++ b/Engine/Source/Runtime/Renderer/Private/PostProcess/PostProcessMotionBlur.cpp
@@ -502,7 +502,7 @@ void ComputeMotionBlurVelocity(
 				PF_FloatRGBA,
 				FClearValueBinding::None,
 				GFastVRamConfig.MotionBlur,
-				TexCreate_ShaderResource | TexCreate_UAV | TexCreate_RenderTargetable,
+				TexCreate_ShaderResource | TexCreate_UAV,
 				false),
 			TEXT("DilatedVelocityTile"));
 
diff --git a/Engine/Source/Runtime/Renderer/Private/PostProcess/PostProcessSubsurface.cpp b/Engine/Source/Runtime/Renderer/Private/PostProcess/PostProcessSubsurface.cpp
--- a/Engine/Source/Runtime/Renderer/Private/PostProcess/PostProcessSubsurface.cpp
+++ b/Engine/Source/Runtime/Renderer/Private/PostProcess/PostProcessSubsurface.cpp
@@ -150,7 +150,10 @@ FRDGTextureRef ComputeSubsurface(
 		return SceneTexture;
 	}
 
-	FRDGTextureRef SceneTextureOutput = GraphBuilder.CreateTexture(SceneTexture->Desc, TEXT("SceneColorSubsurface"));
+	FRDGTextureDesc SceneColorDesc = SceneTexture->Desc;
+	SceneColorDesc.TargetableFlags &= ~TexCreate_UAV;
+	SceneColorDesc.TargetableFlags |= TexCreate_RenderTargetable;
+	FRDGTextureRef SceneTextureOutput = GraphBuilder.CreateTexture(SceneColorDesc, TEXT("SceneColorSubsurface"));
 
 	ERenderTargetLoadAction SceneTextureLoadAction = ERenderTargetLoadAction::ENoAction;
 
diff --git a/Engine/Source/Runtime/Renderer/Private/PostProcess/PostProcessing.cpp b/Engine/Source/Runtime/Renderer/Private/PostProcess/PostProcessing.cpp
--- a/Engine/Source/Runtime/Renderer/Private/PostProcess/PostProcessing.cpp
+++ b/Engine/Source/Runtime/Renderer/Private/PostProcess/PostProcessing.cpp
@@ -1005,7 +1005,11 @@ IMPLEMENT_GLOBAL_SHADER(FComposeSeparateTranslucencyPS, "/Engine/Private/Compose
 
 FRDGTextureRef AddSeparateTranslucencyCompositionPass(FRDGBuilder& GraphBuilder, const FViewInfo& View, FRDGTextureRef SceneColor, FRDGTextureRef SeparateTranslucency)
 {
-	FRDGTextureRef NewSceneColor = GraphBuilder.CreateTexture(SceneColor->Desc, TEXT("SceneColor"));
+	FRDGTextureDesc SceneColorDesc = SceneColor->Desc;
+	SceneColorDesc.TargetableFlags &= ~TexCreate_UAV;
+	SceneColorDesc.TargetableFlags |= TexCreate_RenderTargetable;
+
+	FRDGTextureRef NewSceneColor = GraphBuilder.CreateTexture(SceneColorDesc, TEXT("SceneColor"));
 
 	FComposeSeparateTranslucencyPS::FParameters* PassParameters = GraphBuilder.AllocParameters<FComposeSeparateTranslucencyPS::FParameters>();
 	PassParameters->SceneColor = SceneColor;
diff --git a/Engine/Source/Runtime/Renderer/Private/ScreenSpaceRayTracing.cpp b/Engine/Source/Runtime/Renderer/Private/ScreenSpaceRayTracing.cpp
--- a/Engine/Source/Runtime/Renderer/Private/ScreenSpaceRayTracing.cpp
+++ b/Engine/Source/Runtime/Renderer/Private/ScreenSpaceRayTracing.cpp
@@ -438,7 +438,7 @@ void RenderScreenSpaceReflections(
 		FRDGTextureDesc Desc = FPooledRenderTargetDesc::Create2DDesc(
 			FSceneRenderTargets::Get_FrameConstantsOnly().GetBufferSizeXY(),
 			PF_FloatRGBA, FClearValueBinding(FLinearColor(0, 0, 0, 0)),
-			TexCreate_None, TexCreate_RenderTargetable | TexCreate_ShaderResource | TexCreate_UAV,
+			TexCreate_None, TexCreate_RenderTargetable | TexCreate_ShaderResource,
 			false);
 
 		Desc.AutoWritable = false;
@@ -630,7 +630,7 @@ void RenderScreenSpaceDiffuseIndirect(
 				PF_FloatRGBA, // TODO: might be worth using FloatRGB + R8
 				FClearValueBinding::None,
 				/* InFlags = */ TexCreate_None,
-				/* InTargetableFlags = */ TexCreate_ShaderResource | TexCreate_RenderTargetable | TexCreate_UAV,
+				/* InTargetableFlags = */ TexCreate_ShaderResource | TexCreate_UAV,
 				/* bInForceSeparateTargetAndShaderResource = */ false);
 			Desc.NumMips = 3;
 
@@ -682,7 +682,7 @@ void RenderScreenSpaceDiffuseIndirect(
 				PF_FloatRGBA,
 				FClearValueBinding::Transparent,
 				/* InFlags = */ TexCreate_None,
-				/* InTargetableFlags = */ TexCreate_ShaderResource | TexCreate_RenderTargetable | TexCreate_UAV,
+				/* InTargetableFlags = */ TexCreate_ShaderResource | TexCreate_UAV,
 				/* bInForceSeparateTargetAndShaderResource = */ false);
 
 			OutDenoiserInputs->Color = GraphBuilder.CreateTexture(Desc, TEXT("SSRTDiffuseIndirect"));
