Fixes validation failure when trying to comp planar reflection in TAA's history.

diff --git a/Engine/Source/Runtime/Renderer/Private/DeferredShadingRenderer.h b/Engine/Source/Runtime/Renderer/Private/DeferredShadingRenderer.h
--- a/Engine/Source/Runtime/Renderer/Private/DeferredShadingRenderer.h
+++ b/Engine/Source/Runtime/Renderer/Private/DeferredShadingRenderer.h
@@ -442,6 +442,7 @@ private:
 	/** Render image based reflections (SSR, Env, SkyLight) without compute shaders */
 	void RenderStandardDeferredImageBasedReflections(FRHICommandListImmediate& RHICmdList, FGraphicsPipelineStateInitializer& GraphicsPSOInit, bool bReflectionEnv, const TRefCountPtr<IPooledRenderTarget>& DynamicBentNormalAO, TRefCountPtr<IPooledRenderTarget>& VelocityRT);
 
+	bool HasDeferredPlanarReflections(const FViewInfo& View) const;
 	void RenderDeferredPlanarReflections(FRDGBuilder& GraphBuilder, const FSceneTextureParameters& SceneTextures, const FViewInfo& View, FRDGTextureRef& ReflectionsOutput);
 
 	bool ShouldDoReflectionEnvironment() const;
diff --git a/Engine/Source/Runtime/Renderer/Private/PlanarReflectionRendering.cpp b/Engine/Source/Runtime/Renderer/Private/PlanarReflectionRendering.cpp
--- a/Engine/Source/Runtime/Renderer/Private/PlanarReflectionRendering.cpp
+++ b/Engine/Source/Runtime/Renderer/Private/PlanarReflectionRendering.cpp
@@ -664,8 +664,13 @@ class FPlanarReflectionPS : public FGlobalShader
 
 IMPLEMENT_GLOBAL_SHADER(FPlanarReflectionPS, "/Engine/Private/PlanarReflectionShaders.usf", "PlanarReflectionPS", SF_Pixel);
 
-void FDeferredShadingSceneRenderer::RenderDeferredPlanarReflections(FRDGBuilder& GraphBuilder, const FSceneTextureParameters& SceneTextures, const FViewInfo& View, FRDGTextureRef& ReflectionsOutputTexture)
+bool FDeferredShadingSceneRenderer::HasDeferredPlanarReflections(const FViewInfo& View) const
 {
+	if (View.bIsPlanarReflection || View.bIsReflectionCapture)
+	{
+		return false;
+	}
+
 	// Prevent rendering unsupported views when ViewIndex >= GMaxPlanarReflectionViews
 	// Planar reflections in those views will fallback to other reflection methods
 	{
@@ -675,7 +680,7 @@ void FDeferredShadingSceneRenderer::RenderDeferredPlanarReflections(FRDGBuilder&
 
 		if (ViewIndex >= GMaxPlanarReflectionViews)
 		{
-			return;
+			return false;
 		}
 	}
 
@@ -688,18 +693,18 @@ void FDeferredShadingSceneRenderer::RenderDeferredPlanarReflections(FRDGBuilder&
 		if (View.ViewFrustum.IntersectBox(ReflectionSceneProxy->WorldBounds.GetCenter(), ReflectionSceneProxy->WorldBounds.GetExtent()))
 		{
 			bAnyVisiblePlanarReflections = true;
+			break;
 		}
 	}
 
-	bool bViewIsReflectionCapture = View.bIsPlanarReflection || View.bIsReflectionCapture;
+	bool bComposePlanarReflections = Scene->PlanarReflections.Num() > 0 && bAnyVisiblePlanarReflections;
 
-	bool bComposePlanarReflections = Scene->PlanarReflections.Num() > 0 && !bViewIsReflectionCapture && bAnyVisiblePlanarReflections;
+	return bComposePlanarReflections;
+}
 
-	// Prevent reflection recursion, or view-dependent planar reflections being seen in reflection captures
-	if (!bComposePlanarReflections)
-	{
-		return;
-	}
+void FDeferredShadingSceneRenderer::RenderDeferredPlanarReflections(FRDGBuilder& GraphBuilder, const FSceneTextureParameters& SceneTextures, const FViewInfo& View, FRDGTextureRef& ReflectionsOutputTexture)
+{
+	check(HasDeferredPlanarReflections(View));
 
 	// Allocate planar reflection texture
 	bool bClearReflectionsOutputTexture = false;
@@ -708,7 +713,7 @@ void FDeferredShadingSceneRenderer::RenderDeferredPlanarReflections(FRDGBuilder&
 		FRDGTextureDesc Desc = FPooledRenderTargetDesc::Create2DDesc(
 			SceneTextures.SceneDepthBuffer->Desc.Extent,
 			PF_FloatRGBA, FClearValueBinding(FLinearColor(0, 0, 0, 0)),
-			TexCreate_None, TexCreate_RenderTargetable,
+			TexCreate_None, TexCreate_ShaderResource | TexCreate_RenderTargetable,
 			false);
 
 		Desc.AutoWritable = false;
diff --git a/Engine/Source/Runtime/Renderer/Private/PostProcess/PostProcessTemporalAA.cpp b/Engine/Source/Runtime/Renderer/Private/PostProcess/PostProcessTemporalAA.cpp
--- a/Engine/Source/Runtime/Renderer/Private/PostProcess/PostProcessTemporalAA.cpp
+++ b/Engine/Source/Runtime/Renderer/Private/PostProcess/PostProcessTemporalAA.cpp
@@ -411,6 +411,11 @@ FTAAOutputs FTAAPassParameters::AddTemporalAAPass(
 			/* InTargetableFlags = */ TexCreate_ShaderResource | TexCreate_UAV,
 			/* bInForceSeparateTargetAndShaderResource = */ false);
 
+		if (bOutputRenderTargetable)
+		{
+			SceneColorDesc.TargetableFlags |= TexCreate_RenderTargetable;
+		}
+
 		Outputs.SceneColor = GraphBuilder.CreateTexture(
 			SceneColorDesc,
 			kTAAOutputNames[static_cast<int32>(Pass)],
diff --git a/Engine/Source/Runtime/Renderer/Private/PostProcess/PostProcessTemporalAA.h b/Engine/Source/Runtime/Renderer/Private/PostProcess/PostProcessTemporalAA.h
--- a/Engine/Source/Runtime/Renderer/Private/PostProcess/PostProcessTemporalAA.h
+++ b/Engine/Source/Runtime/Renderer/Private/PostProcess/PostProcessTemporalAA.h
@@ -81,6 +81,9 @@ struct FTAAPassParameters
 	// Whether to use the faster shader permutation.
 	bool bUseFast = false;
 
+	// Whether output texture should be render targetable.
+	bool bOutputRenderTargetable = false;
+
 	// Whether downsampled (box filtered, half resolution) frame should be written out.
 	bool bDownsample = false;
 	EPixelFormat DownsampleOverrideFormat = PF_Unknown;
diff --git a/Engine/Source/Runtime/Renderer/Private/ReflectionEnvironment.cpp b/Engine/Source/Runtime/Renderer/Private/ReflectionEnvironment.cpp
--- a/Engine/Source/Runtime/Renderer/Private/ReflectionEnvironment.cpp
+++ b/Engine/Source/Runtime/Renderer/Private/ReflectionEnvironment.cpp
@@ -655,6 +655,8 @@ void FDeferredShadingSceneRenderer::RenderDeferredReflectionsAndSkyLighting(FRHI
 
 		const bool bScreenSpaceReflections = !bRayTracedReflections && ShouldRenderScreenSpaceReflections(View);
 
+		const bool bComposePlanarReflections = !bRayTracedReflections && HasDeferredPlanarReflections(View);
+
 		FRDGTextureRef ReflectionsColor = nullptr;
 		if (bRayTracedReflections || bScreenSpaceReflections)
 		{
@@ -735,6 +737,7 @@ void FDeferredShadingSceneRenderer::RenderDeferredReflectionsAndSkyLighting(FRHI
 				FTAAPassParameters TAASettings(View);
 				TAASettings.Pass = ETAAPassConfig::ScreenSpaceReflections;
 				TAASettings.SceneColorInput = DenoiserInputs.Color;
+				TAASettings.bOutputRenderTargetable = bComposePlanarReflections;
 				
 				FTAAOutputs TAAOutputs = TAASettings.AddTemporalAAPass(
 					GraphBuilder,
@@ -759,8 +762,9 @@ void FDeferredShadingSceneRenderer::RenderDeferredReflectionsAndSkyLighting(FRHI
 			}
 		} // if (bRayTracedReflections || bScreenSpaceReflections)
 
-		if (!bRayTracedReflections)
+		if (bComposePlanarReflections)
 		{
+			check(!bRayTracedReflections);
 			RenderDeferredPlanarReflections(GraphBuilder, SceneTextures, View, /* inout */ ReflectionsColor);
 		}
 
