Copying //UE4/Dev-Volumetrics to Dev-Rendering (//UE4/Dev-Rendering)

diff --git a/Engine/Source/Runtime/Renderer/Private/SceneRendering.cpp b/Engine/Source/Runtime/Renderer/Private/SceneRendering.cpp
--- a/Engine/Source/Runtime/Renderer/Private/SceneRendering.cpp
+++ b/Engine/Source/Runtime/Renderer/Private/SceneRendering.cpp
@@ -835,6 +835,9 @@ void FViewInfo::Init()
 	FogInscatteringColorCubemap = NULL;
 	FogInscatteringTextureParameters = FVector::ZeroVector;
 
+	SkyAtmosphereCameraAerialPerspectiveVolume = nullptr;
+	SkyAtmosphereUniformShaderParameters = nullptr;
+
 	bUseDirectionalInscattering = false;
 	DirectionalInscatteringExponent = 0;
 	DirectionalInscatteringStartDistance = 0;
@@ -1134,7 +1137,8 @@ void FViewInfo::SetupUniformBufferParameters(
 		}
 
 		// Atmospheric fog parameters
-		const float SunLightDiskHalfApexAngleRadian = Scene->SunLight ? Scene->SunLight->Proxy->GetSunLightHalfApexAngleRadian() : FLightSceneProxy::GetSunOnEarthHalfApexAngleRadian();
+		FLightSceneInfo* SunLight = Scene->AtmosphereLights[0];	// Atmospheric fog only takes into account the a single sun light with index 0.
+		const float SunLightDiskHalfApexAngleRadian = SunLight ? SunLight->Proxy->GetSunLightHalfApexAngleRadian() : FLightSceneProxy::GetSunOnEarthHalfApexAngleRadian();
 		if (ShouldRenderAtmosphere(*Family) && Scene->AtmosphericFog)
 		{
 			ViewUniformShaderParameters.AtmosphericFogSunPower = Scene->AtmosphericFog->SunMultiplier;
@@ -1148,12 +1152,12 @@ void FViewInfo::SetupUniformBufferParameters(
 			ViewUniformShaderParameters.AtmosphericFogStartDistance = Scene->AtmosphericFog->StartDistance;
 			ViewUniformShaderParameters.AtmosphericFogDistanceOffset = Scene->AtmosphericFog->DistanceOffset;
 			ViewUniformShaderParameters.AtmosphericFogSunDiscScale = Scene->AtmosphericFog->SunDiscScale;
-			ViewUniformShaderParameters.AtmosphericFogSunDiscHalfApexAngleRadian = Scene->AtmosphericFog->SunDiscScale * SunLightDiskHalfApexAngleRadian;
-			ViewUniformShaderParameters.AtmosphericFogSunDiscLuminance = Scene->SunLight ? Scene->SunLight->Proxy->GetOuterSpaceLuminance() : FLinearColor::White;
-			ViewUniformShaderParameters.AtmosphericFogSunColor = Scene->SunLight ? Scene->SunLight->Proxy->GetColor() : Scene->AtmosphericFog->DefaultSunColor; // Sun light color unaffected by atmosphere transmittance
-			ViewUniformShaderParameters.AtmosphericFogSunDirection = Scene->SunLight ? -Scene->SunLight->Proxy->GetDirection() : -Scene->AtmosphericFog->DefaultSunDirection;
 			ViewUniformShaderParameters.AtmosphericFogRenderMask = Scene->AtmosphericFog->RenderFlag & (EAtmosphereRenderFlag::E_DisableGroundScattering | EAtmosphereRenderFlag::E_DisableSunDisk);
 			ViewUniformShaderParameters.AtmosphericFogInscatterAltitudeSampleNum = Scene->AtmosphericFog->InscatterAltitudeSampleNum;
+			ViewUniformShaderParameters.AtmosphereLightDiscCosHalfApexAngle[0] = FMath::Cos(Scene->AtmosphericFog->SunDiscScale * SunLightDiskHalfApexAngleRadian);
+			ViewUniformShaderParameters.AtmosphereLightDiscLuminance[0] = SunLight ? SunLight->Proxy->GetOuterSpaceLuminance() : FLinearColor::White;
+			ViewUniformShaderParameters.AtmosphereLightColor[0] = SunLight ? SunLight->Proxy->GetColor() : Scene->AtmosphericFog->DefaultSunColor; // Sun light color unaffected by atmosphere transmittance
+			ViewUniformShaderParameters.AtmosphereLightDirection[0] = SunLight ? -SunLight->Proxy->GetDirection() : -Scene->AtmosphericFog->DefaultSunDirection;
 		}
 		else
 		{
@@ -1168,12 +1172,13 @@ void FViewInfo::SetupUniformBufferParameters(
 			ViewUniformShaderParameters.AtmosphericFogStartDistance = FLT_MAX;
 			ViewUniformShaderParameters.AtmosphericFogDistanceOffset = 0.f;
 			ViewUniformShaderParameters.AtmosphericFogSunDiscScale = 1.f;
-			ViewUniformShaderParameters.AtmosphericFogSunDiscHalfApexAngleRadian = SunLightDiskHalfApexAngleRadian;
-			//Added check so atmospheric light color and vector can use a directional light without needing an atmospheric fog actor in the scene
-			ViewUniformShaderParameters.AtmosphericFogSunColor = Scene->SunLight ? Scene->SunLight->Proxy->GetColor() : FLinearColor::Black;
-			ViewUniformShaderParameters.AtmosphericFogSunDirection = Scene->SunLight ? -Scene->SunLight->Proxy->GetDirection() : DefaultSunDirection;
 			ViewUniformShaderParameters.AtmosphericFogRenderMask = EAtmosphereRenderFlag::E_EnableAll;
 			ViewUniformShaderParameters.AtmosphericFogInscatterAltitudeSampleNum = 0;
+			ViewUniformShaderParameters.AtmosphereLightDiscCosHalfApexAngle[0] = FMath::Cos(SunLightDiskHalfApexAngleRadian);
+			//Added check so atmospheric light color and vector can use a directional light without needing an atmospheric fog actor in the scene
+			ViewUniformShaderParameters.AtmosphereLightDiscLuminance[0] = SunLight ? SunLight->Proxy->GetOuterSpaceLuminance() : FLinearColor::Black;
+			ViewUniformShaderParameters.AtmosphereLightColor[0] = SunLight ? SunLight->Proxy->GetColor() : FLinearColor::Black;
+			ViewUniformShaderParameters.AtmosphereLightDirection[0] = SunLight ? -SunLight->Proxy->GetDirection() : DefaultSunDirection;
 		}
 	}
 	else
@@ -1190,13 +1195,88 @@ void FViewInfo::SetupUniformBufferParameters(
 		ViewUniformShaderParameters.AtmosphericFogStartDistance = FLT_MAX;
 		ViewUniformShaderParameters.AtmosphericFogDistanceOffset = 0.f;
 		ViewUniformShaderParameters.AtmosphericFogSunDiscScale = 1.f;
-		ViewUniformShaderParameters.AtmosphericFogSunDiscHalfApexAngleRadian = 0.f;
-		ViewUniformShaderParameters.AtmosphericFogSunDiscLuminance = FLinearColor::Black;
-		ViewUniformShaderParameters.AtmosphericFogSunColor = FLinearColor::Black;
-		ViewUniformShaderParameters.AtmosphericFogSunDirection = DefaultSunDirection;
 		ViewUniformShaderParameters.AtmosphericFogRenderMask = EAtmosphereRenderFlag::E_EnableAll;
 		ViewUniformShaderParameters.AtmosphericFogInscatterAltitudeSampleNum = 0;
+		ViewUniformShaderParameters.AtmosphereLightDiscCosHalfApexAngle[0] = 1.0f;
+		ViewUniformShaderParameters.AtmosphereLightDiscLuminance[0] = FLinearColor::Black;
+		ViewUniformShaderParameters.AtmosphereLightColor[0] = FLinearColor::Black;
+		ViewUniformShaderParameters.AtmosphereLightDirection[0] = DefaultSunDirection;
+	}
+
+	FRHITexture* TransmittanceLutTextureFound = nullptr;
+	FRHITexture* SkyViewLutTextureFound = nullptr;
+	FRHITexture* CameraAerialPerspectiveVolumeFound = nullptr;
+	FRHITexture* DistantSkyLightLutTextureFound = nullptr;
+	if (Scene && ShouldRenderSkyAtmosphere(Scene->SkyAtmosphere, Scene->GetShaderPlatform()))
+	{
+		check(Scene->SkyAtmosphere->GetTransmittanceLutTexture().IsValid());
+		FSkyAtmosphereRenderSceneInfo* SkyAtmosphere = Scene->SkyAtmosphere;
+		const TRefCountPtr<IPooledRenderTarget>& PooledTransmittanceLutTexture = SkyAtmosphere->GetTransmittanceLutTexture();
+		TransmittanceLutTextureFound = PooledTransmittanceLutTexture->GetRenderTargetItem().ShaderResourceTexture;
+
+		SkyViewLutTextureFound = this->SkyAtmosphereViewLutTexture->GetRenderTargetItem().ShaderResourceTexture;
+		uint32 SkyViewLutWidth = float(this->SkyAtmosphereViewLutTexture->GetDesc().GetSize().X);
+		uint32 SkyViewLutHeight = float(this->SkyAtmosphereViewLutTexture->GetDesc().GetSize().Y);
+		ViewUniformShaderParameters.SkyViewLutSizeAndInvSize = FVector4(SkyViewLutWidth, SkyViewLutHeight, 1.0f / SkyViewLutWidth, 1.0f / SkyViewLutHeight);
+
+		CameraAerialPerspectiveVolumeFound = this->SkyAtmosphereCameraAerialPerspectiveVolume->GetRenderTargetItem().ShaderResourceTexture;
+
+		DistantSkyLightLutTextureFound = SkyAtmosphere->GetDistantSkyLightLutTextureRHI();
+
+		const FAtmosphereSetup& AtmosphereSetup =  SkyAtmosphere->GetAtmosphereSetup();
+		ViewUniformShaderParameters.SkyAtmosphereBottomRadius = AtmosphereSetup.BottomRadius;
+		ViewUniformShaderParameters.SkyAtmosphereTopRadius = AtmosphereSetup.TopRadius;
+
+		FSkyAtmosphereBasePassSharedUniformShaderParameters OutParameters;
+		SetupSkyAtmosphereBasePassSharedUniformShaderParameters(*this, OutParameters);
+		ViewUniformShaderParameters.SkyAtmosphereCameraAerialPerspectiveVolumeStartDepth = OutParameters.CameraAerialPerspectiveVolumeStartDepth;
+		ViewUniformShaderParameters.SkyAtmosphereCameraAerialPerspectiveVolumeDepthResolution = OutParameters.CameraAerialPerspectiveVolumeDepthResolution;
+		ViewUniformShaderParameters.SkyAtmosphereCameraAerialPerspectiveVolumeDepthResolutionInv = OutParameters.CameraAerialPerspectiveVolumeDepthResolutionInv;
+		ViewUniformShaderParameters.SkyAtmosphereCameraAerialPerspectiveVolumeDepthSliceLength = OutParameters.CameraAerialPerspectiveVolumeDepthSliceLength;
+		ViewUniformShaderParameters.SkyAtmosphereCameraAerialPerspectiveVolumeDepthSliceLengthInv = OutParameters.CameraAerialPerspectiveVolumeDepthSliceLengthInv;
+		ViewUniformShaderParameters.SkyAtmosphereApplyCameraAerialPerspectiveVolume = OutParameters.ApplyCameraAerialPerspectiveVolume;
+
+		// Fill atmosphere lights shader parameters
+		for (uint8 Index = 0; Index < NUM_ATMOSPHERE_LIGHTS; ++Index)
+		{
+			FLightSceneInfo* Light = Scene->AtmosphereLights[Index];
+			if (Light)
+			{
+				ViewUniformShaderParameters.AtmosphereLightDiscCosHalfApexAngle[Index] = FMath::Cos(Light->Proxy->GetSunLightHalfApexAngleRadian());
+				ViewUniformShaderParameters.AtmosphereLightDiscLuminance[Index] = Light->Proxy->GetOuterSpaceLuminance();
+				ViewUniformShaderParameters.AtmosphereLightColor[Index] = Light->Proxy->GetColor();
+				ViewUniformShaderParameters.AtmosphereLightDirection[Index] = -Light->Proxy->GetDirection();
+			}
+			else
+			{
+				ViewUniformShaderParameters.AtmosphereLightDiscCosHalfApexAngle[Index] = 1.0f;
+				ViewUniformShaderParameters.AtmosphereLightDiscLuminance[Index] = FLinearColor::Black;
+				ViewUniformShaderParameters.AtmosphereLightColor[Index] = FLinearColor::Black;
+				ViewUniformShaderParameters.AtmosphereLightDirection[Index] = DefaultSunDirection;
+			}
+		}
 	}
+	else
+	{
+		ViewUniformShaderParameters.SkyViewLutSizeAndInvSize = FVector4(1.0f, 1.0f, 1.0f, 1.0f);
+		ViewUniformShaderParameters.SkyAtmosphereBottomRadius = 1.0f;
+		ViewUniformShaderParameters.SkyAtmosphereTopRadius = 1.0f;
+		ViewUniformShaderParameters.SkyAtmosphereCameraAerialPerspectiveVolumeStartDepth = 1.0f;
+		ViewUniformShaderParameters.SkyAtmosphereCameraAerialPerspectiveVolumeDepthResolution = 1.0f;
+		ViewUniformShaderParameters.SkyAtmosphereCameraAerialPerspectiveVolumeDepthResolutionInv = 1.0f;
+		ViewUniformShaderParameters.SkyAtmosphereCameraAerialPerspectiveVolumeDepthSliceLength = 1.0f;
+		ViewUniformShaderParameters.SkyAtmosphereCameraAerialPerspectiveVolumeDepthSliceLengthInv = 1.0f;
+		ViewUniformShaderParameters.SkyAtmosphereApplyCameraAerialPerspectiveVolume = 1.0f;
+	}
+
+	ViewUniformShaderParameters.TransmittanceLutTexture = OrWhite2DIfNull(TransmittanceLutTextureFound);
+	ViewUniformShaderParameters.TransmittanceLutTextureSampler = TStaticSamplerState<SF_Bilinear>::GetRHI();
+	ViewUniformShaderParameters.DistantSkyLightLutTexture = OrBlack2DIfNull(DistantSkyLightLutTextureFound);
+	ViewUniformShaderParameters.DistantSkyLightLutTextureSampler = TStaticSamplerState<SF_Bilinear>::GetRHI();
+	ViewUniformShaderParameters.SkyViewLutTexture = OrBlack2DIfNull(SkyViewLutTextureFound);
+	ViewUniformShaderParameters.SkyViewLutTextureSampler = TStaticSamplerState<SF_Bilinear>::GetRHI();
+	ViewUniformShaderParameters.CameraAerialPerspectiveVolume = OrBlack3DAlpha1IfNull(CameraAerialPerspectiveVolumeFound);
+	ViewUniformShaderParameters.CameraAerialPerspectiveVolumeSampler = TStaticSamplerState<SF_Bilinear>::GetRHI();
 
 	ViewUniformShaderParameters.AtmosphereTransmittanceTexture = OrBlack2DIfNull(AtmosphereTransmittanceTexture);
 	ViewUniformShaderParameters.AtmosphereIrradianceTexture = OrBlack2DIfNull(AtmosphereIrradianceTexture);
diff --git a/Engine/Source/Runtime/Renderer/Private/SceneRendering.h b/Engine/Source/Runtime/Renderer/Private/SceneRendering.h
--- a/Engine/Source/Runtime/Renderer/Private/SceneRendering.h
+++ b/Engine/Source/Runtime/Renderer/Private/SceneRendering.h
@@ -27,6 +27,7 @@
 #include "DistortionRendering.h"
 #include "HeightfieldLighting.h"
 #include "GlobalDistanceFieldParameters.h"
+#include "SkyAtmosphereRendering.h"
 #include "Templates/UniquePtr.h"
 #include "RenderGraph.h"
 #include "MeshDrawCommands.h"
@@ -1045,6 +1046,11 @@ public:
 	float FurthestReflectionCaptureDistance;
 	TUniformBufferRef<FReflectionCaptureShaderData> ReflectionCaptureUniformBuffer;
 
+	// Sky / Atmosphere textures (transient owned by this view info) and pointer to constants owned by SkyAtmosphere proxy.
+	TRefCountPtr<IPooledRenderTarget> SkyAtmosphereCameraAerialPerspectiveVolume;
+	TRefCountPtr<IPooledRenderTarget> SkyAtmosphereViewLutTexture;
+	const FAtmosphereUniformShaderParameters* SkyAtmosphereUniformShaderParameters;
+
 	/** Used when there is no view state, buffers reallocate every frame. */
 	TUniquePtr<FForwardLightingViewResources> ForwardLightingResourcesStorage;
 
@@ -1669,6 +1675,13 @@ protected:
 
 	void RenderPlanarReflection(class FPlanarReflectionSceneProxy* ReflectionSceneProxy);
 
+	/** Initialise sky atmosphere resources.*/
+	void InitSkyAtmosphereForViews(FRHICommandListImmediate& RHICmdList);
+	/** Render the sky atmosphere look up table needed for this frame.*/
+	void RenderSkyAtmosphereLookUpTables(FRHICommandListImmediate& RHICmdList);
+	/** Render the sky atmosphere over the scene.*/
+	void RenderSkyAtmosphere(FRHICommandListImmediate& RHICmdList);
+
 	void ResolveSceneColor(FRHICommandList& RHICmdList);
 
 private:
@@ -1765,6 +1778,13 @@ private:
 // The noise textures need to be set in Slate too.
 RENDERER_API void UpdateNoiseTextureParameters(FViewUniformShaderParameters& ViewUniformShaderParameters);
 
+inline FRHITexture* OrWhite2DIfNull(FRHITexture* Tex)
+{
+	FRHITexture* Result = Tex ? Tex : GWhiteTexture->TextureRHI.GetReference();
+	check(Result);
+	return Result;
+}
+
 inline FRHITexture* OrBlack2DIfNull(FRHITexture* Tex)
 {
 	FRHITexture* Result = Tex ? Tex : GBlackTexture->TextureRHI.GetReference();
@@ -1778,6 +1798,12 @@ inline FRHITexture* OrBlack3DIfNull(FRHITexture* Tex)
 	return OrBlack2DIfNull(Tex ? Tex : GBlackVolumeTexture->TextureRHI.GetReference());
 }
 
+inline FRHITexture* OrBlack3DAlpha1IfNull(FRHITexture* Tex)
+{
+	// we fall back to 2D which are unbound es2 parameters
+	return OrBlack2DIfNull(Tex ? Tex : GBlackAlpha1VolumeTexture->TextureRHI.GetReference());
+}
+
 inline FRHITexture* OrBlack3DUintIfNull(FRHITexture* Tex)
 {
 	// we fall back to 2D which are unbound es2 parameters
@@ -1803,6 +1829,16 @@ inline void SetBlack3DIfNull(FRHITexture*& Tex)
 	}
 }
 
+inline void SetBlackAlpha13DIfNull(FRHITexture*& Tex)
+{
+	if (!Tex)
+	{
+		Tex = GBlackAlpha1VolumeTexture->TextureRHI.GetReference();
+		// we fall back to 2D which are unbound es2 parameters
+		SetBlack2DIfNull(Tex); // This is actually a rgb=0, a=1 texture
+	}
+}
+
 extern TAutoConsoleVariable<int32> CVarTransientResourceAliasing_Buffers;
 
 FORCEINLINE bool IsTransientResourceBufferAliasingEnabled()
diff --git a/Engine/Source/Runtime/Renderer/Private/SceneVisibility.cpp b/Engine/Source/Runtime/Renderer/Private/SceneVisibility.cpp
--- a/Engine/Source/Runtime/Renderer/Private/SceneVisibility.cpp
+++ b/Engine/Source/Runtime/Renderer/Private/SceneVisibility.cpp
@@ -4142,6 +4142,12 @@ bool FDeferredShadingSceneRenderer::InitViews(FRHICommandListImmediate& RHICmdLi
 	CreateIndirectCapsuleShadows();
 	RHICmdList.ImmediateFlush(EImmediateFlushType::DispatchToRHIThread);
 
+	// Initialise Sky/View resources before the view global uniform buffer is built.
+	if (Scene && ShouldRenderSkyAtmosphere(Scene->GetSkyAtmosphereSceneInfo(), Scene->GetShaderPlatform()))
+	{
+		InitSkyAtmosphereForViews(RHICmdList);
+	}
+
 	PostVisibilityFrameSetup(ILCTaskData);
 	RHICmdList.ImmediateFlush(EImmediateFlushType::DispatchToRHIThread);
 
diff --git a/Engine/Source/Runtime/Renderer/Private/SkyAtmosphereRendering.cpp b/Engine/Source/Runtime/Renderer/Private/SkyAtmosphereRendering.cpp
new file mode 100644
--- /dev/null
+++ b/Engine/Source/Runtime/Renderer/Private/SkyAtmosphereRendering.cpp
@@ -0,0 +1,1381 @@
+// Copyright 1998-2019 Epic Games, Inc. All Rights Reserved.
+
+/*=============================================================================
+	SkyAtmosphereRendering.cpp
+=============================================================================*/
+
+#include "SkyAtmosphereRendering.h"
+#include "Components/SkyAtmosphereComponent.h"
+#include "DeferredShadingRenderer.h"
+#include "Engine/MapBuildDataRegistry.h"
+#include "LightSceneInfo.h"
+#include "PixelShaderUtils.h"
+#include "RenderTargetTemp.h"
+#include "ScenePrivate.h"
+#include "SceneRenderTargetParameters.h"
+
+//#pragma optimize( "", off )
+
+
+
+static TAutoConsoleVariable<int32> CVarSkyAtmosphere(
+	TEXT("r.SkyAtmosphere"), 1,
+	TEXT("SkyAtmosphere components are rendered when this is not 0, otherwise ignored.\n"),
+	ECVF_RenderThreadSafe);
+
+////////////////////////////////////////////////////////////////////////// Regular sky 
+
+static TAutoConsoleVariable<float> CVarSkyAtmosphereSampleCountMin(
+	TEXT("r.SkyAtmosphere.SampleCountMin"), 2.0f,
+	TEXT("The minimum sample count used to compute sky/atmosphere scattering and transmittance.\n")
+	TEXT("The minimal value will be clamped to 1.\n"),
+	ECVF_RenderThreadSafe | ECVF_Scalability);
+
+static TAutoConsoleVariable<float> CVarSkyAtmosphereSampleCountMax(
+	TEXT("r.SkyAtmosphere.SampleCountMax"), 16.0f,
+	TEXT("The maximum sample count used to compute sky/atmosphere scattering and transmittance.\n")
+	TEXT("The minimal value will be clamped to r.SkyAtmosphere.SampleCountMin + 1.\n"),
+	ECVF_RenderThreadSafe | ECVF_Scalability);
+
+static TAutoConsoleVariable<float> CVarSkyAtmosphereDistanceToSampleCountMax(
+	TEXT("r.SkyAtmosphere.DistanceToSampleCountMax"), 150.0f,
+	TEXT("The distance in kilometer after which at which SampleCountMax samples will be used to ray march the atmosphere."),
+	ECVF_RenderThreadSafe | ECVF_Scalability);
+
+////////////////////////////////////////////////////////////////////////// Fast sky
+
+static TAutoConsoleVariable<int32> CVarSkyAtmosphereFastSkyLUT(
+	TEXT("r.SkyAtmosphere.FastSkyLUT"), 0,
+	TEXT("When enabled, a look up texture is used to render the sky.\n")
+	TEXT("It is faster but can result in visual artefacts if there are some high frequency details\n")
+	TEXT("in the sky such as earth shadow or scattering lob."),
+	ECVF_RenderThreadSafe | ECVF_Scalability);
+
+static TAutoConsoleVariable<float> CVarSkyAtmosphereFastSkyLUTSampleCountMin(
+	TEXT("r.SkyAtmosphere.FastSkyLUT.SampleCountMin"), 4.0f,
+	TEXT("Fast sky minimum sample count used to compute sky/atmosphere scattering and transmittance.\n")
+	TEXT("The minimal value will be clamped to 1.\n"),
+	ECVF_RenderThreadSafe | ECVF_Scalability);
+
+static TAutoConsoleVariable<float> CVarSkyAtmosphereFastSkyLUTSampleCountMax(
+	TEXT("r.SkyAtmosphere.FastSkyLUT.SampleCountMax"), 32.0f,
+	TEXT("Fast sky maximum sample count used to compute sky/atmosphere scattering and transmittance.\n")
+	TEXT("The minimal value will be clamped to r.SkyAtmosphere.FastSkyLUT.SampleCountMin + 1.\n"),
+	ECVF_RenderThreadSafe | ECVF_Scalability);
+
+static TAutoConsoleVariable<float> CVarSkyAtmosphereFastSkyLUTDistanceToSampleCountMax(
+	TEXT("r.SkyAtmosphere.FastSkyLUT.DistanceToSampleCountMax"), 150.0f,
+	TEXT("Fast sky distance in kilometer after which at which SampleCountMax samples will be used to ray march the atmosphere."),
+	ECVF_RenderThreadSafe | ECVF_Scalability);
+
+static TAutoConsoleVariable<float> CVarSkyAtmosphereFastSkyLUTWidth(
+	TEXT("r.SkyAtmosphere.FastSkyLUT.Width"), 192,
+	TEXT(""),
+	ECVF_RenderThreadSafe | ECVF_Scalability);
+
+static TAutoConsoleVariable<float> CVarSkyAtmosphereFastSkyLUTHeight(
+	TEXT("r.SkyAtmosphere.FastSkyLUT.Height"), 108,
+	TEXT(""),
+	ECVF_RenderThreadSafe | ECVF_Scalability);
+
+////////////////////////////////////////////////////////////////////////// Aerial perspective
+
+static TAutoConsoleVariable<float> CVarSkyAtmosphereAerialPerspectiveStartDepth(
+	TEXT("r.SkyAtmosphere.AerialPerspective.StartDepth"), 0.1f,
+	TEXT("The distance at which we start evaluate the aerial pespective. Default: 100 meters."),
+	ECVF_RenderThreadSafe | ECVF_Scalability);
+
+static TAutoConsoleVariable<int32> CVarSkyAtmosphereAerialPerspectiveDepthTest(
+	TEXT("r.SkyAtmosphere.AerialPerspective.DepthTest"), 1,
+	TEXT("When enabled, a depth test will be used to not write pixel closer to the camera than StartDepth, effectively improving performance."),
+	ECVF_RenderThreadSafe | ECVF_Scalability);
+
+////////////////////////////////////////////////////////////////////////// Aerial perspective LUT
+
+static TAutoConsoleVariable<float> CVarSkyAtmosphereAerialPerspectiveLUTDepthResolution(
+	TEXT("r.SkyAtmosphere.AerialPerspectiveLUT.DepthResolution"), 32.0f,
+	TEXT("The number of depth slice to use for the aerial perspective volume texture."),
+	ECVF_RenderThreadSafe | ECVF_Scalability);
+
+static TAutoConsoleVariable<float> CVarSkyAtmosphereAerialPerspectiveLUTDepth(
+	TEXT("r.SkyAtmosphere.AerialPerspectiveLUT.Depth"), 8.0f,
+	TEXT("The length of the LUT in kilometers (default = 8km). After, the scattering and transmittance will be constantly the last slice."),
+	ECVF_RenderThreadSafe | ECVF_Scalability);
+
+static TAutoConsoleVariable<float> CVarSkyAtmosphereAerialPerspectiveLUTSampleCountPerSlice(
+	TEXT("r.SkyAtmosphere.AerialPerspectiveLUT.SampleCountPerSlice"), 2.0f,
+	TEXT("The sample count used per slice to evaluate aerial perspective\n")
+	TEXT("scattering and transmittance in camera frustum space froxel.\n"),
+	ECVF_RenderThreadSafe | ECVF_Scalability);
+
+static TAutoConsoleVariable<float> CVarSkyAtmosphereAerialPerspectiveLUTWidth(
+	TEXT("r.SkyAtmosphere.AerialPerspectiveLUT.Width"), 32,
+	TEXT(""),
+	ECVF_RenderThreadSafe | ECVF_Scalability);
+
+static TAutoConsoleVariable<int32> CVarSkyAtmosphereAerialPerspectiveApplyOnOpaque(
+	TEXT("r.SkyAtmosphere.AerialPerspectiveLUT.FastApplyOnOpaque"), 0,
+	TEXT("When enabled, the low resolution camera frustum/froxel volume containing atmospheric fog\n")
+	TEXT(", usually used for fog on translucent surface, is used to render fog on opaque.\n")
+	TEXT("It is faster but can result in visual artefacts if there are some high frequency details\n")
+	TEXT("such as earth shadow or scattering lob."),
+	ECVF_RenderThreadSafe | ECVF_Scalability);
+
+////////////////////////////////////////////////////////////////////////// Transmittance LUT
+
+static TAutoConsoleVariable<float> CVarSkyAtmosphereTranstmittanceLUTSampleCount(
+	TEXT("r.SkyAtmosphere.TransmittanceLUT.SampleCount"), 10.0f,
+	TEXT("The sample count used to evaluate transmittance."),
+	ECVF_RenderThreadSafe | ECVF_Scalability);
+
+static TAutoConsoleVariable<int32> CVarSkyAtmosphereTransmittanceLUTUseSmallFormat(
+	TEXT("r.SkyAtmosphere.TransmittanceLUT.UseSmallFormat"), 0,
+	TEXT("If true, the transmittance LUT will use a small R8BG8B8A8 format to store data at lower quality."),
+	ECVF_RenderThreadSafe | ECVF_Scalability);
+
+static TAutoConsoleVariable<float> CVarSkyAtmosphereTransmittanceLUTWidth(
+	TEXT("r.SkyAtmosphere.TransmittanceLUT.Width"), 256,
+	TEXT(""),
+	ECVF_RenderThreadSafe | ECVF_Scalability);
+
+static TAutoConsoleVariable<float> CVarSkyAtmosphereTransmittanceLUTHeight(
+	TEXT("r.SkyAtmosphere.TransmittanceLUT.Height"), 64,
+	TEXT(""),
+	ECVF_RenderThreadSafe | ECVF_Scalability);
+
+////////////////////////////////////////////////////////////////////////// Multi-scattering LUT
+
+static TAutoConsoleVariable<float> CVarSkyAtmosphereMultiScatteringLUTSampleCount(
+	TEXT("r.SkyAtmosphere.MultiScatteringLUT.SampleCount"), 15.0f,
+	TEXT("The sample count used to evaluate multi-scattering.\n"),
+	ECVF_RenderThreadSafe | ECVF_Scalability);
+
+static TAutoConsoleVariable<int32> CVarSkyAtmosphereMultiScatteringLUTUseSmallFormat(
+	TEXT("r.SkyAtmosphere.MultiScatteringLUT.UseSmallFormat"), 0,
+	TEXT("If true, the multi-scattering LUT will use a small R8BG8B8A8 format to store data at lower quality."),
+	ECVF_RenderThreadSafe | ECVF_Scalability);
+
+static TAutoConsoleVariable<float> CVarSkyAtmosphereMultiScatteringLUTWidth(
+	TEXT("r.SkyAtmosphere.MultiScatteringLUT.Width"), 32,
+	TEXT(""),
+	ECVF_RenderThreadSafe | ECVF_Scalability);
+
+static TAutoConsoleVariable<float> CVarSkyAtmosphereMultiScatteringLUTHeight(
+	TEXT("r.SkyAtmosphere.MultiScatteringLUT.Height"), 32,
+	TEXT(""),
+	ECVF_RenderThreadSafe | ECVF_Scalability);
+
+////////////////////////////////////////////////////////////////////////// Distant Sky Light LUT
+
+static TAutoConsoleVariable<int32> CVarSkyAtmosphereDistantSkyLightLUT(
+	TEXT("r.SkyAtmosphere.DistantSkyLightLUT"), 1,
+	TEXT("Enable the generation the sky ambient lighting value.\n"),
+	ECVF_RenderThreadSafe | ECVF_Scalability);
+
+static TAutoConsoleVariable<float> CVarSkyAtmosphereDistantSkyLightLUTAltitude(
+	TEXT("r.SkyAtmosphere.DistantSkyLightLUT.Altitude"), 6.0f,
+	TEXT("The altitude at which the sky samples are taken to integrate the sky lighting. Default to 6km, typicaly cirrus clouds altitude.\n"),
+	ECVF_RenderThreadSafe | ECVF_Scalability);
+
+////////////////////////////////////////////////////////////////////////// Debug / Visualisation
+
+TAutoConsoleVariable<int32> CVarSkyAtmosphereVisualize(
+	TEXT("r.SkyAtmosphere.Visualize"),
+	0,
+	TEXT("DRAW DEBUG"),
+	ECVF_RenderThreadSafe);
+
+
+// Extra internal constants shared between all passes.
+BEGIN_GLOBAL_SHADER_PARAMETER_STRUCT(FSkyAtmosphereInternalCommonParameters, )
+	SHADER_PARAMETER(float, SampleCountMin)
+	SHADER_PARAMETER(float, SampleCountMax)
+	SHADER_PARAMETER(float, DistanceToSampleCountMaxInv)
+
+	SHADER_PARAMETER(float, FastSkySampleCountMin)
+	SHADER_PARAMETER(float, FastSkySampleCountMax)
+	SHADER_PARAMETER(float, FastSkyDistanceToSampleCountMaxInv)
+
+	SHADER_PARAMETER(FVector4, CameraAerialPerspectiveVolumeSizeAndInvSize)
+	SHADER_PARAMETER(float, CameraAerialPerspectiveVolumeStartDepth)
+	SHADER_PARAMETER(float, CameraAerialPerspectiveVolumeDepthResolution)
+	SHADER_PARAMETER(float, CameraAerialPerspectiveVolumeDepthResolutionInv)
+	SHADER_PARAMETER(float, CameraAerialPerspectiveVolumeDepthSliceLength)
+	SHADER_PARAMETER(float, CameraAerialPerspectiveVolumeDepthSliceLengthInv)
+	SHADER_PARAMETER(float, CameraAerialPerspectiveSampleCountPerSlice)
+
+	SHADER_PARAMETER(FVector4, TransmittanceLutSizeAndInvSize)
+	SHADER_PARAMETER(FVector4, MultiScatteredLuminanceLutSizeAndInvSize)
+	SHADER_PARAMETER(FVector4, SkyViewLutSizeAndInvSize)
+
+	SHADER_PARAMETER(float, TransmittanceSampleCount)
+	SHADER_PARAMETER(float, MultiScatteringSampleCount)
+	SHADER_PARAMETER(float, AerialPespectiveViewDistanceScale)
+
+	SHADER_PARAMETER(FVector, SkyLuminanceFactor)
+
+END_GLOBAL_SHADER_PARAMETER_STRUCT()
+
+IMPLEMENT_GLOBAL_SHADER_PARAMETER_STRUCT(FAtmosphereUniformShaderParameters, "Atmosphere");
+IMPLEMENT_GLOBAL_SHADER_PARAMETER_STRUCT(FSkyAtmosphereBasePassSharedUniformShaderParameters, "SkyAtmosphere");
+IMPLEMENT_GLOBAL_SHADER_PARAMETER_STRUCT(FSkyAtmosphereInternalCommonParameters, "SkyAtmosphere");
+
+
+#define GET_VALID_DATA_FROM_CVAR \
+	auto ValidateLUTResolution = [](int32 Value) \
+	{ \
+		return Value < 4 ? 4 : Value; \
+	}; \
+	int32 TransmittanceLutWidth = ValidateLUTResolution(CVarSkyAtmosphereTransmittanceLUTWidth.GetValueOnRenderThread()); \
+	int32 TransmittanceLutHeight = ValidateLUTResolution(CVarSkyAtmosphereTransmittanceLUTHeight.GetValueOnRenderThread()); \
+	int32 MultiScatteredLuminanceLutWidth = ValidateLUTResolution(CVarSkyAtmosphereMultiScatteringLUTWidth.GetValueOnRenderThread()); \
+	int32 MultiScatteredLuminanceLutHeight = ValidateLUTResolution(CVarSkyAtmosphereMultiScatteringLUTHeight.GetValueOnRenderThread()); \
+	int32 SkyViewLutWidth = ValidateLUTResolution(CVarSkyAtmosphereFastSkyLUTWidth.GetValueOnRenderThread()); \
+	int32 SkyViewLutHeight = ValidateLUTResolution(CVarSkyAtmosphereFastSkyLUTHeight.GetValueOnRenderThread()); \
+	int32 CameraAerialPerspectiveVolumeScreenResolution = ValidateLUTResolution(CVarSkyAtmosphereAerialPerspectiveLUTWidth.GetValueOnRenderThread()); \
+	int32 CameraAerialPerspectiveVolumeDepthResolution = ValidateLUTResolution(CVarSkyAtmosphereAerialPerspectiveLUTDepthResolution.GetValueOnRenderThread()); \
+	float CameraAerialPerspectiveVolumeStartDepth = CVarSkyAtmosphereAerialPerspectiveStartDepth.GetValueOnRenderThread(); \
+	CameraAerialPerspectiveVolumeStartDepth = CameraAerialPerspectiveVolumeStartDepth < 0.0f ? 0.0f : CameraAerialPerspectiveVolumeStartDepth; \
+	float CameraAerialPerspectiveVolumeDepth = CVarSkyAtmosphereAerialPerspectiveLUTDepth.GetValueOnRenderThread(); \
+	CameraAerialPerspectiveVolumeDepth = CameraAerialPerspectiveVolumeDepth < 1.0f ? 1.0f : CameraAerialPerspectiveVolumeDepth;	/* 1 kilometer minimum */ \
+	float CameraAerialPerspectiveVolumeDepthSliceLength = CameraAerialPerspectiveVolumeDepth / CameraAerialPerspectiveVolumeDepthResolution;
+
+
+static bool ShouldPipelineCompileSkyAtmosphereShader(EShaderPlatform ShaderPlatform)
+{
+	// Requires SM5 or ES3_1 (GL/Vulkan) for compute shaders and volume textures support.
+	return RHISupportsComputeShaders(ShaderPlatform); // RHISupportsVolumeTextures(GetMaxSupportedFeatureLevel(ShaderPlatform));
+}
+
+bool ShouldRenderSkyAtmosphere(const FSkyAtmosphereRenderSceneInfo* SkyAtmosphere, EShaderPlatform ShaderPlatform)
+{
+	const bool ShadersCompiled = ShouldPipelineCompileSkyAtmosphereShader(ShaderPlatform);
+	return SkyAtmosphere && ShadersCompiled && CVarSkyAtmosphere.GetValueOnRenderThread() > 0;
+	// TODO: Add new or reuse EngineShowFlags.AtmosphericFog? ALso take into account EngineShowFlags.Fog as previously?
+}
+
+void SetupSkyAtmosphereBasePassSharedUniformShaderParameters(const class FViewInfo& View, FSkyAtmosphereBasePassSharedUniformShaderParameters& OutParameters)
+{
+	GET_VALID_DATA_FROM_CVAR;
+
+	FRHITexture* SkyAtmosphereCameraAerialPerspectiveVolume = nullptr;
+	if (View.SkyAtmosphereCameraAerialPerspectiveVolume)
+	{
+		SkyAtmosphereCameraAerialPerspectiveVolume = View.SkyAtmosphereCameraAerialPerspectiveVolume->GetRenderTargetItem().ShaderResourceTexture;
+	}
+	SetBlackAlpha13DIfNull(SkyAtmosphereCameraAerialPerspectiveVolume);
+
+	OutParameters.ApplyCameraAerialPerspectiveVolume = View.SkyAtmosphereCameraAerialPerspectiveVolume == nullptr ? 0.0f : 1.0f;
+	OutParameters.CameraAerialPerspectiveVolumeStartDepth = CameraAerialPerspectiveVolumeStartDepth;
+	OutParameters.CameraAerialPerspectiveVolumeDepthResolution = float(CameraAerialPerspectiveVolumeDepthResolution);
+	OutParameters.CameraAerialPerspectiveVolumeDepthResolutionInv = 1.0f / OutParameters.CameraAerialPerspectiveVolumeDepthResolution;
+	OutParameters.CameraAerialPerspectiveVolumeDepthSliceLength = CameraAerialPerspectiveVolumeDepthSliceLength;
+	OutParameters.CameraAerialPerspectiveVolumeDepthSliceLengthInv = 1.0f / OutParameters.CameraAerialPerspectiveVolumeDepthSliceLength;
+	OutParameters.CameraAerialPerspectiveVolume = SkyAtmosphereCameraAerialPerspectiveVolume;
+	OutParameters.CameraAerialPerspectiveVolumeSampler = TStaticSamplerState<SF_Trilinear>::GetRHI();
+}
+
+static void CopyAtmosphereSetupToUniformShaderParameters(FAtmosphereUniformShaderParameters& out, const FAtmosphereSetup& Atmosphere)
+{
+#define COPYMACRO(MemberName) out.MemberName = Atmosphere.MemberName 
+	COPYMACRO(MultiScatteringFactor);
+	COPYMACRO(BottomRadius);
+	COPYMACRO(TopRadius);
+	COPYMACRO(RayleighDensityExpScale);
+	COPYMACRO(RayleighScattering);
+	COPYMACRO(MieScattering);
+	COPYMACRO(MieDensityExpScale);
+	COPYMACRO(MieExtinction);
+	COPYMACRO(MiePhaseG);
+	COPYMACRO(MieAbsorption);
+	COPYMACRO(AbsorptionDensity0LayerWidth);
+	COPYMACRO(AbsorptionDensity0ConstantTerm);
+	COPYMACRO(AbsorptionDensity0LinearTerm);
+	COPYMACRO(AbsorptionDensity1ConstantTerm);
+	COPYMACRO(AbsorptionDensity1LinearTerm);
+	COPYMACRO(AbsorptionExtinction);
+	COPYMACRO(GroundAlbedo);
+#undef COPYMACRO
+}
+
+
+
+/*=============================================================================
+	FSkyAtmosphereRenderSceneInfo implementation.
+=============================================================================*/
+
+
+
+FSkyAtmosphereRenderSceneInfo::FSkyAtmosphereRenderSceneInfo(const USkyAtmosphereComponent* InComponent)
+	: bStaticLightingBuilt(false)
+	, AtmosphereSetup(*InComponent)
+	, TransmittanceAtZenith(AtmosphereSetup.GetTransmittanceAtGroundLevel(FVector(0.0f, 0.0f, 1.0f)))
+{
+	SkyLuminanceFactor = InComponent->SkyLuminanceFactor;
+	AerialPespectiveViewDistanceScale = InComponent->AerialPespectiveViewDistanceScale;
+
+	// Create a multiframe uniform buffer. A render command is used because FSkyAtmosphereRenderSceneInfo ctor is called on the Game thread.
+	TUniformBufferRef<FAtmosphereUniformShaderParameters>* AtmosphereUniformBufferPtr = &AtmosphereUniformBuffer;
+	FAtmosphereUniformShaderParameters* AtmosphereUniformShaderParametersPtr = &AtmosphereUniformShaderParameters;
+	CopyAtmosphereSetupToUniformShaderParameters(AtmosphereUniformShaderParameters, AtmosphereSetup);
+	ENQUEUE_RENDER_COMMAND(FCreateUniformBuffer)(
+		[AtmosphereUniformBufferPtr, AtmosphereUniformShaderParametersPtr](FRHICommandListImmediate& RHICmdList)
+	{
+		*AtmosphereUniformBufferPtr = TUniformBufferRef<FAtmosphereUniformShaderParameters>::CreateUniformBufferImmediate(*AtmosphereUniformShaderParametersPtr, UniformBuffer_MultiFrame);
+	});
+}
+
+FSkyAtmosphereRenderSceneInfo::~FSkyAtmosphereRenderSceneInfo()
+{
+	// Destructor is called on the render thread so safe to release here
+	AtmosphereUniformBuffer->Release();
+	//TransmittanceLutTexture->Release(); This is automatically tracked by FRenderTargetPool and RefCount
+}
+
+void FSkyAtmosphereRenderSceneInfo::PrepareSunLightProxy(FLightSceneInfo& SunLight) const
+{
+	// See explanation in https://media.contentapi.ea.com/content/dam/eacom/frostbite/files/s2016-pbs-frostbite-sky-clouds-new.pdf page 26
+	const bool bAtmosphereAffectsSunIlluminance = true;
+	FLinearColor TransmittanceTowardSun = bAtmosphereAffectsSunIlluminance ? AtmosphereSetup.GetTransmittanceAtGroundLevel(-SunLight.Proxy->GetDirection()) : FLinearColor(FLinearColor::White);
+	FLinearColor TransmittanceAtZenithFinal = bAtmosphereAffectsSunIlluminance ? TransmittanceAtZenith : FLinearColor(FLinearColor::White);
+
+	FLinearColor SunZenithIlluminance = SunLight.Proxy->GetColor();
+	FLinearColor SunOuterSpaceIlluminance = SunZenithIlluminance / TransmittanceAtZenithFinal;
+
+	// SunDiscScale is only considered as a visual tweak so we do not make it influence the sun disk outerspace luminance.
+	const float SunSolidAngle = 2.0f * PI * (1.0f - FMath::Cos(SunLight.Proxy->GetSunLightHalfApexAngleRadian())); // Solid angle from aperture https://en.wikipedia.org/wiki/Solid_angle 
+	FLinearColor SunDiskOuterSpaceLuminance = SunOuterSpaceIlluminance / SunSolidAngle; // approximation
+
+	SunLight.Proxy->SetAtmosphereRelatedProperties(TransmittanceTowardSun / TransmittanceAtZenithFinal, SunDiskOuterSpaceLuminance);
+}
+
+FTextureRHIRef FSkyAtmosphereRenderSceneInfo::GetDistantSkyLightLutTextureRHI()
+{
+	if (CVarSkyAtmosphereDistantSkyLightLUT.GetValueOnRenderThread() > 0)
+	{
+		return DistantSkyLightLutTexture->GetRenderTargetItem().ShaderResourceTexture;
+	}
+	return GBlackTexture->TextureRHI;
+}
+
+
+
+/*=============================================================================
+	FScene functions
+=============================================================================*/
+
+void FScene::AddSkyAtmosphere(const USkyAtmosphereComponent* SkyAtmosphereComponent, bool bStaticLightingBuilt)
+{
+	check(SkyAtmosphereComponent);
+
+	FSkyAtmosphereRenderSceneInfo* SceneInfo = new FSkyAtmosphereRenderSceneInfo(SkyAtmosphereComponent);
+	FScene* Scene = this;
+	ENQUEUE_RENDER_COMMAND(FAddSkyAtmosphereCommand)(
+		[Scene, SceneInfo, bStaticLightingBuilt](FRHICommandListImmediate& RHICmdList)
+		{
+			if (Scene->SkyAtmosphere)
+			{
+				delete Scene->SkyAtmosphere;
+			}
+			Scene->SkyAtmosphere = SceneInfo;
+			Scene->SkyAtmosphere->bStaticLightingBuilt = bStaticLightingBuilt;
+			if (!Scene->SkyAtmosphere->bStaticLightingBuilt)
+			{
+				FPlatformAtomics::InterlockedIncrement(&Scene->NumUncachedStaticLightingInteractions);
+			}
+		} );
+}
+
+void FScene::RemoveSkyAtmosphere(const USkyAtmosphereComponent* SkyAtmosphereComponent)
+{
+	FScene* Scene = this;
+	check(SkyAtmosphereComponent);
+
+	ENQUEUE_RENDER_COMMAND(FRemoveSkyAtmosphereCommand)(
+		[Scene](FRHICommandListImmediate& RHICmdList)
+		{
+			if (Scene->SkyAtmosphere)
+			{
+				if (Scene->SkyAtmosphere && !Scene->SkyAtmosphere->bStaticLightingBuilt)
+				{
+					FPlatformAtomics::InterlockedDecrement(&Scene->NumUncachedStaticLightingInteractions);
+				}
+				delete Scene->SkyAtmosphere;
+				Scene->SkyAtmosphere = NULL;
+			}
+		} );
+}
+
+
+
+/*=============================================================================
+	Sky/Atmosphere rendering functions
+=============================================================================*/
+
+
+
+namespace
+{
+
+class FSkyPermutationMultiScatteringApprox : SHADER_PERMUTATION_BOOL("MULTISCATTERING_APPROX_ENABLED");
+class FFastSky : SHADER_PERMUTATION_BOOL("FASTSKY_ENABLED");
+class FFastAerialPespective : SHADER_PERMUTATION_BOOL("FASTAERIALPERSPECTIVE_ENABLED");
+class FSourceDiskEnabled : SHADER_PERMUTATION_BOOL("SOURCE_DISK_ENABLED");
+class FOnlySkyPixel : SHADER_PERMUTATION_BOOL("ONLY_SKY_PIXEL_ENABLED");
+class FSecondAtmosphereLight : SHADER_PERMUTATION_BOOL("SECOND_ATMOSPHERE_LIGHT_ENABLED");
+
+}
+
+class FRenderSkyAtmosphereVS : public FGlobalShader
+{
+	DECLARE_GLOBAL_SHADER(FRenderSkyAtmosphereVS);
+	SHADER_USE_PARAMETER_STRUCT(FRenderSkyAtmosphereVS, FGlobalShader);
+
+	using FPermutationDomain = TShaderPermutationDomain<FOnlySkyPixel>;
+
+	BEGIN_SHADER_PARAMETER_STRUCT(FParameters, )
+		SHADER_PARAMETER_STRUCT_REF(FViewUniformShaderParameters, ViewUniformBuffer)
+		SHADER_PARAMETER(float, StartDepthZ)
+	END_SHADER_PARAMETER_STRUCT()
+
+	static FPermutationDomain RemapPermutation(FPermutationDomain PermutationVector)
+	{
+		return PermutationVector;
+	}
+
+	static bool ShouldCompilePermutation(const FGlobalShaderPermutationParameters& Parameters)
+	{
+		return ShouldPipelineCompileSkyAtmosphereShader(Parameters.Platform);
+	}
+
+	static void ModifyCompilationEnvironment(const FGlobalShaderPermutationParameters& Parameters, FShaderCompilerEnvironment& OutEnvironment)
+	{
+		FGlobalShader::ModifyCompilationEnvironment(Parameters, OutEnvironment);
+	}
+};
+IMPLEMENT_GLOBAL_SHADER(FRenderSkyAtmosphereVS, "/Engine/Private/SkyAtmosphere.usf", "SkyAtmosphereVS", SF_Vertex);
+
+//////////////////////////////////////////////////////////////////////////
+
+class FRenderSkyAtmospherePS : public FGlobalShader
+{
+	DECLARE_GLOBAL_SHADER(FRenderSkyAtmospherePS);
+	SHADER_USE_PARAMETER_STRUCT(FRenderSkyAtmospherePS, FGlobalShader);
+
+	using FPermutationDomain = TShaderPermutationDomain<FSkyPermutationMultiScatteringApprox, FFastSky, FFastAerialPespective, FSourceDiskEnabled, FOnlySkyPixel, FSecondAtmosphereLight>;
+
+	BEGIN_SHADER_PARAMETER_STRUCT(FParameters, )
+		SHADER_PARAMETER_STRUCT_REF(FAtmosphereUniformShaderParameters, Atmosphere)
+		SHADER_PARAMETER_STRUCT_REF(FSkyAtmosphereInternalCommonParameters, SkyAtmosphere)
+		SHADER_PARAMETER_STRUCT_REF(FViewUniformShaderParameters, ViewUniformBuffer)
+//		RENDER_TARGET_BINDING_SLOTS()
+		SHADER_PARAMETER_STRUCT_REF(FSceneTexturesUniformParameters, SceneTextures)
+		SHADER_PARAMETER_TEXTURE(Texture2D<float3>, TransmittanceLutTexture)
+		SHADER_PARAMETER_TEXTURE(Texture2D<float3>, MultiScatteredLuminanceLutTexture)
+		SHADER_PARAMETER_TEXTURE(Texture2D<float3>, SkyViewLutTexture)
+		SHADER_PARAMETER_TEXTURE(Texture3D<float4>, CameraAerialPerspectiveVolumeTexture)
+		SHADER_PARAMETER_SAMPLER(SamplerState, TransmittanceLutTextureSampler)
+		SHADER_PARAMETER_SAMPLER(SamplerState, MultiScatteredLuminanceLutTextureSampler)
+		SHADER_PARAMETER_SAMPLER(SamplerState, SkyViewLutTextureSampler)
+		SHADER_PARAMETER_SAMPLER(SamplerState, CameraAerialPerspectiveVolumeTextureSampler)
+	END_SHADER_PARAMETER_STRUCT()
+
+	static FPermutationDomain RemapPermutation(FPermutationDomain PermutationVector)
+	{
+		return PermutationVector;
+	}
+
+	static bool ShouldCompilePermutation(const FGlobalShaderPermutationParameters& Parameters)
+	{
+		return ShouldPipelineCompileSkyAtmosphereShader(Parameters.Platform);
+	}
+
+	static void ModifyCompilationEnvironment(const FGlobalShaderPermutationParameters& Parameters, FShaderCompilerEnvironment& OutEnvironment)
+	{
+		FGlobalShader::ModifyCompilationEnvironment(Parameters, OutEnvironment);
+		OutEnvironment.SetDefine(TEXT("PER_PIXEL_NOISE"), 1);
+	}
+};
+IMPLEMENT_GLOBAL_SHADER(FRenderSkyAtmospherePS, "/Engine/Private/SkyAtmosphere.usf", "RenderSkyAtmosphereRayMarchingPS", SF_Pixel);
+
+//////////////////////////////////////////////////////////////////////////
+
+class FRenderTransmittanceLutCS : public FGlobalShader
+{
+	DECLARE_GLOBAL_SHADER(FRenderTransmittanceLutCS);
+	SHADER_USE_PARAMETER_STRUCT(FRenderTransmittanceLutCS, FGlobalShader);
+
+	using FPermutationDomain = TShaderPermutationDomain<>;
+
+public:
+	const static uint32 GroupSize = 8;
+
+	BEGIN_SHADER_PARAMETER_STRUCT(FParameters, )
+		SHADER_PARAMETER_STRUCT_REF(FAtmosphereUniformShaderParameters, Atmosphere)
+		SHADER_PARAMETER_STRUCT_REF(FSkyAtmosphereInternalCommonParameters, SkyAtmosphere)
+		SHADER_PARAMETER_UAV(RWTexture2D<float3>, TransmittanceLutUAV)
+	END_SHADER_PARAMETER_STRUCT()
+
+	static bool ShouldCompilePermutation(const FGlobalShaderPermutationParameters& Parameters)
+	{
+		return ShouldPipelineCompileSkyAtmosphereShader(Parameters.Platform);
+	}
+
+	static void ModifyCompilationEnvironment(const FGlobalShaderPermutationParameters& Parameters, FShaderCompilerEnvironment& OutEnvironment)
+	{
+		FGlobalShader::ModifyCompilationEnvironment(Parameters, OutEnvironment);
+		OutEnvironment.SetDefine(TEXT("THREADGROUP_SIZE"), GroupSize);
+		OutEnvironment.SetDefine(TEXT("WHITE_TRANSMITTANCE"),1); // Workaround for some compiler not culling enough unused code (e.g. when computing TransmittanceLUT, Transmittance texture is still requested but we are computing it) 
+		OutEnvironment.SetDefine(TEXT("TRANSMITTANCE_PASS"), 1);
+	}
+};
+IMPLEMENT_GLOBAL_SHADER(FRenderTransmittanceLutCS, "/Engine/Private/SkyAtmosphere.usf", "RenderTransmittanceLutCS", SF_Compute);
+
+//////////////////////////////////////////////////////////////////////////
+
+class FRenderMultiScatteredLuminanceLutCS : public FGlobalShader
+{
+	DECLARE_GLOBAL_SHADER(FRenderMultiScatteredLuminanceLutCS);
+	SHADER_USE_PARAMETER_STRUCT(FRenderMultiScatteredLuminanceLutCS, FGlobalShader);
+
+	using FPermutationDomain = TShaderPermutationDomain<>;
+
+public:
+	const static uint32 GroupSize = 8;
+
+	BEGIN_SHADER_PARAMETER_STRUCT(FParameters, )
+		SHADER_PARAMETER_STRUCT_REF(FAtmosphereUniformShaderParameters, Atmosphere)
+		SHADER_PARAMETER_STRUCT_REF(FSkyAtmosphereInternalCommonParameters, SkyAtmosphere)
+		SHADER_PARAMETER_UAV(RWTexture2D<float3>, MultiScatteredLuminanceLutUAV)
+		SHADER_PARAMETER_TEXTURE(Texture2D<float3>, TransmittanceLutTexture)
+		SHADER_PARAMETER_SAMPLER(SamplerState, TransmittanceLutTextureSampler)
+	END_SHADER_PARAMETER_STRUCT()
+
+	static bool ShouldCompilePermutation(const FGlobalShaderPermutationParameters& Parameters)
+	{
+		return ShouldPipelineCompileSkyAtmosphereShader(Parameters.Platform);
+	}
+
+	static void ModifyCompilationEnvironment(const FGlobalShaderPermutationParameters& Parameters, FShaderCompilerEnvironment& OutEnvironment)
+	{
+		FGlobalShader::ModifyCompilationEnvironment(Parameters, OutEnvironment);
+		OutEnvironment.SetDefine(TEXT("THREADGROUP_SIZE"), GroupSize);
+		OutEnvironment.SetDefine(TEXT("MULTISCATT_PASS"), 1);
+	}
+};
+IMPLEMENT_GLOBAL_SHADER(FRenderMultiScatteredLuminanceLutCS, "/Engine/Private/SkyAtmosphere.usf", "RenderMultiScatteredLuminanceLutCS", SF_Compute);
+
+//////////////////////////////////////////////////////////////////////////
+
+class FRenderDistantSkyLightLutCS : public FGlobalShader
+{
+	DECLARE_GLOBAL_SHADER(FRenderDistantSkyLightLutCS);
+	SHADER_USE_PARAMETER_STRUCT(FRenderDistantSkyLightLutCS, FGlobalShader);
+
+	using FPermutationDomain = TShaderPermutationDomain<FSkyPermutationMultiScatteringApprox, FSecondAtmosphereLight>;
+
+public:
+	const static uint32 GroupSize = 8;
+
+	BEGIN_SHADER_PARAMETER_STRUCT(FParameters, )
+		SHADER_PARAMETER_STRUCT_REF(FAtmosphereUniformShaderParameters, Atmosphere)
+		SHADER_PARAMETER_STRUCT_REF(FSkyAtmosphereInternalCommonParameters, SkyAtmosphere)
+		SHADER_PARAMETER_UAV(RWTexture2D<float3>, DistantSkyLightLutUAV)
+		SHADER_PARAMETER_TEXTURE(Texture2D<float3>, TransmittanceLutTexture)
+		SHADER_PARAMETER_TEXTURE(Texture2D<float3>, MultiScatteredLuminanceLutTexture)
+		SHADER_PARAMETER_SAMPLER(SamplerState, TransmittanceLutTextureSampler)
+		SHADER_PARAMETER_SAMPLER(SamplerState, MultiScatteredLuminanceLutTextureSampler)
+		SHADER_PARAMETER_SRV(Buffer<float4>, UniformSphereSamplesBuffer)
+		SHADER_PARAMETER(FVector4, AtmosphereLightDirection0)
+		SHADER_PARAMETER(FVector4, AtmosphereLightDirection1)
+		SHADER_PARAMETER(FLinearColor, AtmosphereLightColor0)
+		SHADER_PARAMETER(FLinearColor, AtmosphereLightColor1)
+		SHADER_PARAMETER(float, DistantSkyLightSampleAltitude)
+	END_SHADER_PARAMETER_STRUCT()
+
+	static bool ShouldCompilePermutation(const FGlobalShaderPermutationParameters& Parameters)
+	{
+		return ShouldPipelineCompileSkyAtmosphereShader(Parameters.Platform);
+	}
+
+	static void ModifyCompilationEnvironment(const FGlobalShaderPermutationParameters& Parameters, FShaderCompilerEnvironment& OutEnvironment)
+	{
+		FGlobalShader::ModifyCompilationEnvironment(Parameters, OutEnvironment);
+		OutEnvironment.SetDefine(TEXT("THREADGROUP_SIZE"), GroupSize);
+		OutEnvironment.SetDefine(TEXT("SKYLIGHT_PASS"), 1);
+	}
+};
+IMPLEMENT_GLOBAL_SHADER(FRenderDistantSkyLightLutCS, "/Engine/Private/SkyAtmosphere.usf", "RenderDistantSkyLightLutCS", SF_Compute);
+
+//////////////////////////////////////////////////////////////////////////
+
+class FRenderSkyViewLutCS : public FGlobalShader
+{
+	DECLARE_GLOBAL_SHADER(FRenderSkyViewLutCS);
+	SHADER_USE_PARAMETER_STRUCT(FRenderSkyViewLutCS, FGlobalShader);
+
+	using FPermutationDomain = TShaderPermutationDomain<FSkyPermutationMultiScatteringApprox, FSourceDiskEnabled, FSecondAtmosphereLight>;
+
+public:
+	const static uint32 GroupSize = 8;
+
+	BEGIN_SHADER_PARAMETER_STRUCT(FParameters, )
+		SHADER_PARAMETER_STRUCT_REF(FAtmosphereUniformShaderParameters, Atmosphere)
+		SHADER_PARAMETER_STRUCT_REF(FSkyAtmosphereInternalCommonParameters, SkyAtmosphere)
+		SHADER_PARAMETER_STRUCT_REF(FViewUniformShaderParameters, ViewUniformBuffer)
+		SHADER_PARAMETER_UAV(RWTexture2D<float3>, SkyViewLutUAV)
+		SHADER_PARAMETER_TEXTURE(Texture2D<float3>, TransmittanceLutTexture)
+		SHADER_PARAMETER_TEXTURE(Texture2D<float3>, MultiScatteredLuminanceLutTexture)
+		SHADER_PARAMETER_SAMPLER(SamplerState, TransmittanceLutTextureSampler)
+		SHADER_PARAMETER_SAMPLER(SamplerState, MultiScatteredLuminanceLutTextureSampler)
+	END_SHADER_PARAMETER_STRUCT()
+
+	static bool ShouldCompilePermutation(const FGlobalShaderPermutationParameters& Parameters)
+	{
+		return ShouldPipelineCompileSkyAtmosphereShader(Parameters.Platform);
+	}
+
+	static void ModifyCompilationEnvironment(const FGlobalShaderPermutationParameters& Parameters, FShaderCompilerEnvironment& OutEnvironment)
+	{
+		FGlobalShader::ModifyCompilationEnvironment(Parameters, OutEnvironment);
+		OutEnvironment.SetDefine(TEXT("THREADGROUP_SIZE"), GroupSize);
+	}
+};
+IMPLEMENT_GLOBAL_SHADER(FRenderSkyViewLutCS, "/Engine/Private/SkyAtmosphere.usf", "RenderSkyViewLutCS", SF_Compute);
+
+//////////////////////////////////////////////////////////////////////////
+
+class FRenderCameraAerialPerspectiveVolumeCS : public FGlobalShader
+{
+	DECLARE_GLOBAL_SHADER(FRenderCameraAerialPerspectiveVolumeCS);
+	SHADER_USE_PARAMETER_STRUCT(FRenderCameraAerialPerspectiveVolumeCS, FGlobalShader);
+
+	using FPermutationDomain = TShaderPermutationDomain<FSkyPermutationMultiScatteringApprox, FSecondAtmosphereLight>;
+
+public:
+	const static uint32 GroupSize = 4;
+
+	BEGIN_SHADER_PARAMETER_STRUCT(FParameters, )
+		SHADER_PARAMETER_STRUCT_REF(FAtmosphereUniformShaderParameters, Atmosphere)
+		SHADER_PARAMETER_STRUCT_REF(FSkyAtmosphereInternalCommonParameters, SkyAtmosphere)
+		SHADER_PARAMETER_STRUCT_REF(FViewUniformShaderParameters, ViewUniformBuffer)
+		SHADER_PARAMETER_UAV(RWTexture3D<float4>, CameraAerialPerspectiveVolumeUAV)
+		SHADER_PARAMETER_TEXTURE(Texture2D<float3>, TransmittanceLutTexture)
+		SHADER_PARAMETER_TEXTURE(Texture2D<float3>, MultiScatteredLuminanceLutTexture)
+		SHADER_PARAMETER_SAMPLER(SamplerState, TransmittanceLutTextureSampler)
+		SHADER_PARAMETER_SAMPLER(SamplerState, MultiScatteredLuminanceLutTextureSampler)
+	END_SHADER_PARAMETER_STRUCT()
+
+	static bool ShouldCompilePermutation(const FGlobalShaderPermutationParameters& Parameters)
+	{
+		return ShouldPipelineCompileSkyAtmosphereShader(Parameters.Platform);
+	}
+
+	static void ModifyCompilationEnvironment(const FGlobalShaderPermutationParameters& Parameters, FShaderCompilerEnvironment& OutEnvironment)
+	{
+		FGlobalShader::ModifyCompilationEnvironment(Parameters, OutEnvironment);
+		OutEnvironment.SetDefine(TEXT("THREADGROUP_SIZE"), GroupSize);
+	}
+};
+IMPLEMENT_GLOBAL_SHADER(FRenderCameraAerialPerspectiveVolumeCS, "/Engine/Private/SkyAtmosphere.usf", "RenderCameraAerialPerspectiveVolumeCS", SF_Compute);
+
+//////////////////////////////////////////////////////////////////////////
+
+class FRenderDebugSkyAtmospherePS : public FGlobalShader
+{
+	DECLARE_GLOBAL_SHADER(FRenderDebugSkyAtmospherePS);
+	SHADER_USE_PARAMETER_STRUCT(FRenderDebugSkyAtmospherePS, FGlobalShader);
+
+	using FPermutationDomain = TShaderPermutationDomain<FSkyPermutationMultiScatteringApprox>;
+
+	BEGIN_SHADER_PARAMETER_STRUCT(FParameters, )
+		SHADER_PARAMETER_STRUCT_REF(FAtmosphereUniformShaderParameters, Atmosphere)
+		SHADER_PARAMETER_STRUCT_REF(FSkyAtmosphereInternalCommonParameters, SkyAtmosphere)
+		SHADER_PARAMETER_STRUCT_REF(FViewUniformShaderParameters, ViewUniformBuffer)
+		RENDER_TARGET_BINDING_SLOTS()
+		SHADER_PARAMETER_TEXTURE(Texture2D<float3>, TransmittanceLutTexture)
+		SHADER_PARAMETER_TEXTURE(Texture2D<float3>, MultiScatteredLuminanceLutTexture)
+		SHADER_PARAMETER_SAMPLER(SamplerState, TransmittanceLutTextureSampler)
+		SHADER_PARAMETER_SAMPLER(SamplerState, MultiScatteredLuminanceLutTextureSampler)
+		SHADER_PARAMETER(float, ViewPortWidth)
+		SHADER_PARAMETER(float, ViewPortHeight)
+	END_SHADER_PARAMETER_STRUCT()
+
+	static FPermutationDomain RemapPermutation(FPermutationDomain PermutationVector)
+	{
+		return PermutationVector;
+	}
+
+	static bool ShouldCompilePermutation(const FGlobalShaderPermutationParameters& Parameters)
+	{
+		// TODO: Exclude when shipping.
+		return ShouldPipelineCompileSkyAtmosphereShader(Parameters.Platform);
+	}
+
+	static void ModifyCompilationEnvironment(const FGlobalShaderPermutationParameters& Parameters, FShaderCompilerEnvironment& OutEnvironment)
+	{
+		FGlobalShader::ModifyCompilationEnvironment(Parameters, OutEnvironment);
+	}
+};
+IMPLEMENT_GLOBAL_SHADER(FRenderDebugSkyAtmospherePS, "/Engine/Private/SkyAtmosphere.usf", "RenderSkyAtmosphereDebugPS", SF_Pixel);
+
+
+
+/*=============================================================================
+	FUniformSphereSamplesBuffer
+=============================================================================*/
+
+class FUniformSphereSamplesBuffer : public FRenderResource
+{
+public:
+	FReadBuffer UniformSphereSamplesBuffer;
+
+	virtual void InitRHI() override
+	{
+		const float GroupSize = float(FRenderDistantSkyLightLutCS::GroupSize);
+		const float GroupSizeInv = 1.0f / GroupSize;
+
+		UniformSphereSamplesBuffer.Initialize(sizeof(FVector4), FRenderDistantSkyLightLutCS::GroupSize * FRenderDistantSkyLightLutCS::GroupSize, EPixelFormat::PF_A32B32G32R32F, BUF_Static);
+		FVector4* Dest = (FVector4*)RHILockVertexBuffer(UniformSphereSamplesBuffer.Buffer, 0, sizeof(FVector4)*FRenderDistantSkyLightLutCS::GroupSize*FRenderDistantSkyLightLutCS::GroupSize, RLM_WriteOnly);
+
+		FMath::SRandInit(0xDE4DC0DE);
+		for (int32 i = 0; i < FRenderDistantSkyLightLutCS::GroupSize; ++i)
+		{
+			for (int32 j = 0; j < FRenderDistantSkyLightLutCS::GroupSize; ++j)
+			{
+				const float u0 = (float(i) + FMath::SRand()) * GroupSizeInv;
+				const float u1 = (float(j) + FMath::SRand()) * GroupSizeInv;
+
+				const float a = 1.0f - 2.0f * u0;
+				const float b = FMath::Sqrt(1.0f - a*a);
+				const float phi = 2 * PI * u1;
+
+				int32 idx = j * FRenderDistantSkyLightLutCS::GroupSize + i;
+				Dest[idx].X = b * FMath::Cos(phi);
+				Dest[idx].Y = b * FMath::Sin(phi);
+				Dest[idx].Z = a;
+				Dest[idx].W = 0.0f;
+			}
+		}
+
+		RHIUnlockVertexBuffer(UniformSphereSamplesBuffer.Buffer);
+	}
+
+	virtual void ReleaseRHI()
+	{
+		UniformSphereSamplesBuffer.Release();
+	}
+};
+TGlobalResource<FUniformSphereSamplesBuffer> GUniformSphereSamplesBuffer;
+
+
+
+/*=============================================================================
+	FSceneRenderer functions
+=============================================================================*/
+
+void FSceneRenderer::InitSkyAtmosphereForViews(FRHICommandListImmediate& RHICmdList)
+{
+	if (Scene)
+	{
+		GET_VALID_DATA_FROM_CVAR;
+
+		check(ShouldRenderSkyAtmosphere(Scene->GetSkyAtmosphereSceneInfo(), Scene->GetShaderPlatform())); // This should not be called if we should not render SkyAtmosphere
+		FPooledRenderTargetDesc Desc;
+		check(Scene->GetSkyAtmosphereSceneInfo());
+		FSkyAtmosphereRenderSceneInfo& SkyInfo = *Scene->GetSkyAtmosphereSceneInfo();
+
+		EPixelFormat TextureLUTFormat = PF_FloatRGB;
+		EPixelFormat TextureLUTSmallFormat = PF_R8G8B8A8;
+		if (Scene->GetFeatureLevel() <= ERHIFeatureLevel::ES3_1)
+		{
+			// OpenGL ES3.1 does not support storing into 3-component images
+			// TODO: check if need this for Metal, Vulkan
+			TextureLUTFormat = PF_FloatRGBA;
+		}  
+
+		//
+		// Initialise per scene/atmosphere resources
+		//
+		const bool TranstmittanceLUTUseSmallFormat = CVarSkyAtmosphereTransmittanceLUTUseSmallFormat.GetValueOnRenderThread() > 0;
+		const bool MultiScatteringLUTUseSmallFormat = CVarSkyAtmosphereMultiScatteringLUTUseSmallFormat.GetValueOnRenderThread() > 0;
+
+		TRefCountPtr<IPooledRenderTarget>& TransmittanceLutTexture = SkyInfo.GetTransmittanceLutTexture();
+		Desc = FPooledRenderTargetDesc::Create2DDesc(
+			FIntPoint(TransmittanceLutWidth, TransmittanceLutHeight), 
+			TranstmittanceLUTUseSmallFormat ? TextureLUTSmallFormat : TextureLUTFormat, FClearValueBinding::None, TexCreate_HideInVisualizeTexture, TexCreate_ShaderResource | TexCreate_UAV, false);
+		GRenderTargetPool.FindFreeElement(RHICmdList, Desc, TransmittanceLutTexture, TEXT("TransmittanceLutTexture"), true, ERenderTargetTransience::NonTransient);
+
+		TRefCountPtr<IPooledRenderTarget>& MultiScatteredLuminanceLutTexture = SkyInfo.GetMultiScatteredLuminanceLutTexture();
+		Desc = FPooledRenderTargetDesc::Create2DDesc(
+			FIntPoint(MultiScatteredLuminanceLutWidth, MultiScatteredLuminanceLutHeight), 
+			MultiScatteringLUTUseSmallFormat ? TextureLUTSmallFormat : TextureLUTFormat, FClearValueBinding::None, TexCreate_HideInVisualizeTexture, TexCreate_ShaderResource | TexCreate_UAV, false);
+		GRenderTargetPool.FindFreeElement(RHICmdList, Desc, MultiScatteredLuminanceLutTexture, TEXT("MultiScatteredLuminanceLutTexture"), true, ERenderTargetTransience::NonTransient);
+
+		if (CVarSkyAtmosphereDistantSkyLightLUT.GetValueOnRenderThread() > 0)
+		{
+			TRefCountPtr<IPooledRenderTarget>& DistantSkyLightLutTexture = SkyInfo.GetDistantSkyLightLutTexture();
+			Desc = FPooledRenderTargetDesc::Create2DDesc(
+				FIntPoint(1, 1),
+				TextureLUTFormat, FClearValueBinding::None, TexCreate_HideInVisualizeTexture, TexCreate_ShaderResource | TexCreate_UAV, false);
+			GRenderTargetPool.FindFreeElement(RHICmdList, Desc, DistantSkyLightLutTexture, TEXT("DistantSkyLightLutTexture"), true, ERenderTargetTransience::NonTransient);
+		}
+
+		//
+		// Initialise transient per view resources.
+		//
+
+		FPooledRenderTargetDesc SkyAtmosphereViewLutTextureDesc = FPooledRenderTargetDesc::Create2DDesc(
+			FIntPoint(SkyViewLutWidth, SkyViewLutHeight), 
+			TextureLUTFormat, FClearValueBinding::None, TexCreate_HideInVisualizeTexture, TexCreate_ShaderResource | TexCreate_UAV, false);
+
+		FPooledRenderTargetDesc SkyAtmosphereCameraAerialPerspectiveVolumeDesc = FPooledRenderTargetDesc::CreateVolumeDesc(
+			CameraAerialPerspectiveVolumeScreenResolution, CameraAerialPerspectiveVolumeScreenResolution, CameraAerialPerspectiveVolumeDepthResolution, 
+			PF_FloatRGBA, FClearValueBinding::None, TexCreate_HideInVisualizeTexture, TexCreate_ShaderResource | TexCreate_UAV, false);
+
+		for (int32 ViewIndex = 0; ViewIndex < Views.Num(); ViewIndex++)
+		{
+			FViewInfo& View = Views[ViewIndex];
+			// Set textures and data that will be needed later on the view.
+			View.SkyAtmosphereUniformShaderParameters = SkyInfo.GetAtmosphereShaderParameters();
+			GRenderTargetPool.FindFreeElement(RHICmdList, SkyAtmosphereViewLutTextureDesc, View.SkyAtmosphereViewLutTexture, TEXT("View.SkyAtmosphereViewLutTexture"), true, ERenderTargetTransience::Transient);
+			GRenderTargetPool.FindFreeElement(RHICmdList, SkyAtmosphereCameraAerialPerspectiveVolumeDesc, View.SkyAtmosphereCameraAerialPerspectiveVolume, TEXT("View.SkyAtmosphereCameraAerialPerspectiveVolume"), true, ERenderTargetTransience::Transient);
+		}
+	}
+}
+
+static void SetupSkyAtmosphereInternalCommonParameters(
+	FSkyAtmosphereInternalCommonParameters& InternalCommonParameters, 
+	const FScene& Scene,
+	const FSkyAtmosphereRenderSceneInfo& SkyInfo)
+{
+	GET_VALID_DATA_FROM_CVAR;
+
+	auto GetSizeAndInvSize = [](int32 Width, int32 Height)
+	{
+		float FWidth = float(Width);
+		float FHeight = float(Height);
+		return FVector4(FWidth, FHeight, 1.0f / FWidth, 1.0f / FHeight);
+	};
+	InternalCommonParameters.TransmittanceLutSizeAndInvSize = GetSizeAndInvSize(TransmittanceLutWidth, TransmittanceLutHeight);
+	InternalCommonParameters.MultiScatteredLuminanceLutSizeAndInvSize = GetSizeAndInvSize(MultiScatteredLuminanceLutWidth, MultiScatteredLuminanceLutHeight);
+	InternalCommonParameters.SkyViewLutSizeAndInvSize = GetSizeAndInvSize(SkyViewLutWidth, SkyViewLutHeight);
+
+	InternalCommonParameters.SampleCountMin = CVarSkyAtmosphereSampleCountMin.GetValueOnRenderThread();
+	InternalCommonParameters.SampleCountMin = CVarSkyAtmosphereSampleCountMin.GetValueOnRenderThread();
+	InternalCommonParameters.SampleCountMax = CVarSkyAtmosphereSampleCountMax.GetValueOnRenderThread();
+	float DistanceToSampleCountMaxInv = CVarSkyAtmosphereDistanceToSampleCountMax.GetValueOnRenderThread();
+
+	InternalCommonParameters.FastSkySampleCountMin = CVarSkyAtmosphereFastSkyLUTSampleCountMin.GetValueOnRenderThread();
+	InternalCommonParameters.FastSkySampleCountMax = CVarSkyAtmosphereFastSkyLUTSampleCountMax.GetValueOnRenderThread();
+	float FastSkyDistanceToSampleCountMaxInv = CVarSkyAtmosphereFastSkyLUTDistanceToSampleCountMax.GetValueOnRenderThread();
+
+	InternalCommonParameters.CameraAerialPerspectiveVolumeStartDepth = CameraAerialPerspectiveVolumeStartDepth;
+	InternalCommonParameters.CameraAerialPerspectiveVolumeDepthResolution = float(CameraAerialPerspectiveVolumeDepthResolution);
+	InternalCommonParameters.CameraAerialPerspectiveVolumeDepthResolutionInv = 1.0f / InternalCommonParameters.CameraAerialPerspectiveVolumeDepthResolution;
+	InternalCommonParameters.CameraAerialPerspectiveVolumeDepthSliceLength = CameraAerialPerspectiveVolumeDepthSliceLength;
+	InternalCommonParameters.CameraAerialPerspectiveVolumeDepthSliceLengthInv = 1.0f / CameraAerialPerspectiveVolumeDepthSliceLength;
+	InternalCommonParameters.CameraAerialPerspectiveSampleCountPerSlice = CVarSkyAtmosphereAerialPerspectiveLUTSampleCountPerSlice.GetValueOnRenderThread();
+
+	InternalCommonParameters.TransmittanceSampleCount = CVarSkyAtmosphereTranstmittanceLUTSampleCount.GetValueOnRenderThread();
+	InternalCommonParameters.MultiScatteringSampleCount = CVarSkyAtmosphereMultiScatteringLUTSampleCount.GetValueOnRenderThread();
+
+	InternalCommonParameters.SkyLuminanceFactor = FVector(SkyInfo.GetSkyLuminanceFactor());
+	InternalCommonParameters.AerialPespectiveViewDistanceScale = SkyInfo.GetAerialPespectiveViewDistanceScale();
+
+	auto ValidateDistanceValue = [](float& Value)
+	{
+		Value = Value < KINDA_SMALL_NUMBER ? KINDA_SMALL_NUMBER : Value;
+	};
+	auto ValidateSampleCountValue = [](float& Value)
+	{
+		Value = Value < 1.0f ? 1.0f : Value;
+	};
+	auto ValidateMaxSampleCountValue = [](float& Value, float& MinValue)
+	{
+		Value = Value < MinValue ? MinValue : Value;
+	};
+	ValidateSampleCountValue(InternalCommonParameters.SampleCountMin);
+	ValidateMaxSampleCountValue(InternalCommonParameters.SampleCountMax, InternalCommonParameters.SampleCountMin);
+	ValidateSampleCountValue(InternalCommonParameters.FastSkySampleCountMin);
+	ValidateMaxSampleCountValue(InternalCommonParameters.FastSkySampleCountMax, InternalCommonParameters.FastSkySampleCountMin);
+	ValidateSampleCountValue(InternalCommonParameters.CameraAerialPerspectiveSampleCountPerSlice);
+	ValidateSampleCountValue(InternalCommonParameters.TransmittanceSampleCount);
+	ValidateSampleCountValue(InternalCommonParameters.MultiScatteringSampleCount);
+	ValidateDistanceValue(DistanceToSampleCountMaxInv);
+	ValidateDistanceValue(FastSkyDistanceToSampleCountMaxInv);
+
+	// Derived values post validation
+	InternalCommonParameters.DistanceToSampleCountMaxInv = 1.0f / DistanceToSampleCountMaxInv;
+	InternalCommonParameters.FastSkyDistanceToSampleCountMaxInv = 1.0f / FastSkyDistanceToSampleCountMaxInv;
+	InternalCommonParameters.CameraAerialPerspectiveVolumeSizeAndInvSize = GetSizeAndInvSize(CameraAerialPerspectiveVolumeScreenResolution, CameraAerialPerspectiveVolumeScreenResolution);
+}
+
+static bool IsSecondAtmosphereLightEnabled(FScene* Scene)
+{
+	// If the second light is not null then we enable the second light.
+	// We do not do any light1 to light0 remapping is light0 is null.
+	return Scene->AtmosphereLights[1] != nullptr;
+}
+
+void FSceneRenderer::RenderSkyAtmosphereLookUpTables(FRHICommandListImmediate& RHICmdList)
+{
+	check(ShouldRenderSkyAtmosphere(Scene->GetSkyAtmosphereSceneInfo(), Scene->GetShaderPlatform())); // This should not be called if we should not render SkyAtmosphere
+
+	DECLARE_GPU_STAT(SkyAtmosphereLUTs);
+	SCOPED_DRAW_EVENT(RHICmdList, SkyAtmosphereLUTs);
+	SCOPED_GPU_STAT(RHICmdList, SkyAtmosphereLUTs);
+
+	FSkyAtmosphereRenderSceneInfo& SkyInfo = *Scene->GetSkyAtmosphereSceneInfo();
+	const bool bMultiScattering = SkyInfo.IsMultiScatteringEnabled();
+	const bool bFastSky = CVarSkyAtmosphereFastSkyLUT.GetValueOnRenderThread() > 0;
+	const bool bFastAerialPespective = CVarSkyAtmosphereAerialPerspectiveApplyOnOpaque.GetValueOnRenderThread() > 0;
+	const bool bSecondAtmosphereLightEnabled = IsSecondAtmosphereLightEnabled(Scene);
+
+	FRHISamplerState* SamplerLinearClamp = TStaticSamplerState<SF_Trilinear>::GetRHI();
+
+	// Initialise common internal parameters
+	FSkyAtmosphereInternalCommonParameters InternalCommonParameters;
+	SetupSkyAtmosphereInternalCommonParameters(InternalCommonParameters, *Scene, SkyInfo);
+	TUniformBufferRef<FSkyAtmosphereInternalCommonParameters> InternalCommonParametersRef = TUniformBufferRef<FSkyAtmosphereInternalCommonParameters>::CreateUniformBufferImmediate(InternalCommonParameters, UniformBuffer_SingleFrame);
+
+	UnbindRenderTargets(RHICmdList);
+
+	// Transmittance LUT
+	TShaderMap<FGlobalShaderType>* GlobalShaderMap = GetGlobalShaderMap(FeatureLevel);
+	{
+		TShaderMapRef<FRenderTransmittanceLutCS> ComputeShader(GlobalShaderMap);
+		RHICmdList.SetComputeShader(ComputeShader->GetComputeShader());
+
+		FRenderTransmittanceLutCS::FParameters PassParameters;
+		PassParameters.Atmosphere = Scene->GetSkyAtmosphereSceneInfo()->GetAtmosphereUniformBuffer();
+		PassParameters.SkyAtmosphere = InternalCommonParametersRef;
+		PassParameters.TransmittanceLutUAV = SkyInfo.GetTransmittanceLutTexture()->GetRenderTargetItem().UAV;
+
+		SetShaderParameters(RHICmdList, *ComputeShader, ComputeShader->GetComputeShader(), PassParameters);
+
+		FIntVector TextureSize = SkyInfo.GetTransmittanceLutTexture()->GetDesc().GetSize();
+		TextureSize.Z = 1;
+		const FIntVector NumGroups = FIntVector::DivideAndRoundUp(TextureSize, FRenderTransmittanceLutCS::GroupSize);
+		RHICmdList.DispatchComputeShader(NumGroups.X, NumGroups.Y, NumGroups.Z);
+		UnsetShaderUAVs(RHICmdList, *ComputeShader, ComputeShader->GetComputeShader());
+	}
+
+	// Mean Illuminance LUT
+	{
+		TShaderMapRef<FRenderMultiScatteredLuminanceLutCS> ComputeShader(GlobalShaderMap);
+		RHICmdList.SetComputeShader(ComputeShader->GetComputeShader());
+
+		FRenderMultiScatteredLuminanceLutCS::FParameters PassParameters;
+		PassParameters.Atmosphere = Scene->GetSkyAtmosphereSceneInfo()->GetAtmosphereUniformBuffer();
+		PassParameters.SkyAtmosphere = InternalCommonParametersRef;
+		PassParameters.TransmittanceLutTextureSampler = SamplerLinearClamp;
+		PassParameters.TransmittanceLutTexture = SkyInfo.GetTransmittanceLutTexture()->GetRenderTargetItem().ShaderResourceTexture;
+		PassParameters.MultiScatteredLuminanceLutUAV = SkyInfo.GetMultiScatteredLuminanceLutTexture()->GetRenderTargetItem().UAV;
+
+		SetShaderParameters(RHICmdList, *ComputeShader, ComputeShader->GetComputeShader(), PassParameters);
+
+		FIntVector TextureSize = SkyInfo.GetMultiScatteredLuminanceLutTexture()->GetDesc().GetSize();
+		TextureSize.Z = 1;
+		const FIntVector NumGroups = FIntVector::DivideAndRoundUp(TextureSize, FRenderMultiScatteredLuminanceLutCS::GroupSize);
+		RHICmdList.DispatchComputeShader(NumGroups.X, NumGroups.Y, NumGroups.Z);
+		UnsetShaderUAVs(RHICmdList, *ComputeShader, ComputeShader->GetComputeShader());
+	}
+
+	// Distant Sky Light LUT
+	if(CVarSkyAtmosphereDistantSkyLightLUT.GetValueOnRenderThread() > 0)
+	{
+		FRenderDistantSkyLightLutCS::FPermutationDomain PermutationVector;
+		PermutationVector.Set<FSkyPermutationMultiScatteringApprox>(bMultiScattering);
+		PermutationVector.Set<FSecondAtmosphereLight>(bSecondAtmosphereLightEnabled);
+		TShaderMapRef<FRenderDistantSkyLightLutCS> ComputeShader(GlobalShaderMap, PermutationVector);
+
+		using FPermutationDomain = TShaderPermutationDomain<FSkyPermutationMultiScatteringApprox, FSecondAtmosphereLight>;
+
+		RHICmdList.SetComputeShader(ComputeShader->GetComputeShader());
+
+		FRenderDistantSkyLightLutCS::FParameters PassParameters;
+		PassParameters.Atmosphere = Scene->GetSkyAtmosphereSceneInfo()->GetAtmosphereUniformBuffer();
+		PassParameters.SkyAtmosphere = InternalCommonParametersRef;
+		PassParameters.TransmittanceLutTextureSampler = SamplerLinearClamp;
+		PassParameters.MultiScatteredLuminanceLutTextureSampler = SamplerLinearClamp;
+		PassParameters.TransmittanceLutTexture = SkyInfo.GetTransmittanceLutTexture()->GetRenderTargetItem().ShaderResourceTexture;
+		PassParameters.MultiScatteredLuminanceLutTexture = SkyInfo.GetMultiScatteredLuminanceLutTexture()->GetRenderTargetItem().ShaderResourceTexture;
+		PassParameters.UniformSphereSamplesBuffer = GUniformSphereSamplesBuffer.UniformSphereSamplesBuffer.SRV;
+		PassParameters.DistantSkyLightLutUAV = SkyInfo.GetDistantSkyLightLutTexture()->GetRenderTargetItem().UAV;
+
+		FLightSceneInfo* Light0 = Scene->AtmosphereLights[0];
+		FLightSceneInfo* Light1 = Scene->AtmosphereLights[1];
+		if (Light0)
+		{
+			PassParameters.AtmosphereLightDirection0 = -Light0->Proxy->GetDirection();
+			PassParameters.AtmosphereLightColor0 = Light0->Proxy->GetColor();
+		}
+		if (Light1)
+		{
+			PassParameters.AtmosphereLightDirection1 = -Light1->Proxy->GetDirection();
+			PassParameters.AtmosphereLightColor1 = Light1->Proxy->GetColor();
+		}
+		PassParameters.DistantSkyLightSampleAltitude = CVarSkyAtmosphereDistantSkyLightLUTAltitude.GetValueOnAnyThread();
+
+		SetShaderParameters(RHICmdList, *ComputeShader, ComputeShader->GetComputeShader(), PassParameters);
+
+		FIntVector TextureSize = FIntVector(1, 1, 1);
+		const FIntVector NumGroups = FIntVector::DivideAndRoundUp(TextureSize, FRenderDistantSkyLightLutCS::GroupSize);
+		RHICmdList.DispatchComputeShader(NumGroups.X, NumGroups.Y, NumGroups.Z);
+		UnsetShaderUAVs(RHICmdList, *ComputeShader, ComputeShader->GetComputeShader());
+	}
+
+	for (int32 ViewIndex = 0; ViewIndex < Views.Num(); ViewIndex++)
+	{
+		FViewInfo& View = Views[ViewIndex];
+
+		const bool bLightDiskEnabled = !View.bIsReflectionCapture;
+
+		UnbindRenderTargets(RHICmdList);
+
+		// Sky View LUT
+		{
+			FRenderSkyViewLutCS::FPermutationDomain PermutationVector;
+			PermutationVector.Set<FSkyPermutationMultiScatteringApprox>(bMultiScattering);
+			PermutationVector.Set<FSourceDiskEnabled>(bLightDiskEnabled);
+			PermutationVector.Set<FSecondAtmosphereLight>(bSecondAtmosphereLightEnabled);
+			TShaderMapRef<FRenderSkyViewLutCS> ComputeShader(View.ShaderMap, PermutationVector);
+			RHICmdList.SetComputeShader(ComputeShader->GetComputeShader());
+
+			FRenderSkyViewLutCS::FParameters PassParameters;
+			PassParameters.Atmosphere = Scene->GetSkyAtmosphereSceneInfo()->GetAtmosphereUniformBuffer();
+			PassParameters.SkyAtmosphere = InternalCommonParametersRef;
+			PassParameters.ViewUniformBuffer = View.ViewUniformBuffer;
+			PassParameters.TransmittanceLutTextureSampler = SamplerLinearClamp;
+			PassParameters.MultiScatteredLuminanceLutTextureSampler = SamplerLinearClamp;
+			PassParameters.TransmittanceLutTexture = SkyInfo.GetTransmittanceLutTexture()->GetRenderTargetItem().ShaderResourceTexture;
+			PassParameters.MultiScatteredLuminanceLutTexture = SkyInfo.GetMultiScatteredLuminanceLutTexture()->GetRenderTargetItem().ShaderResourceTexture;
+			PassParameters.SkyViewLutUAV = View.SkyAtmosphereViewLutTexture->GetRenderTargetItem().UAV;
+
+			SetShaderParameters(RHICmdList, *ComputeShader, ComputeShader->GetComputeShader(), PassParameters);
+
+			FIntVector TextureSize = View.SkyAtmosphereViewLutTexture->GetDesc().GetSize();
+			TextureSize.Z = 1;
+			const FIntVector NumGroups = FIntVector::DivideAndRoundUp(TextureSize, FRenderSkyViewLutCS::GroupSize);
+			RHICmdList.DispatchComputeShader(NumGroups.X, NumGroups.Y, NumGroups.Z);
+			UnsetShaderUAVs(RHICmdList, *ComputeShader, ComputeShader->GetComputeShader());
+		}
+
+		// Camera Atmosphere Volume
+		{
+			FRenderCameraAerialPerspectiveVolumeCS::FPermutationDomain PermutationVector;
+			PermutationVector.Set<FSkyPermutationMultiScatteringApprox>(bMultiScattering);
+			PermutationVector.Set<FSecondAtmosphereLight>(bSecondAtmosphereLightEnabled);
+			TShaderMapRef<FRenderCameraAerialPerspectiveVolumeCS> ComputeShader(View.ShaderMap, PermutationVector);
+			RHICmdList.SetComputeShader(ComputeShader->GetComputeShader());
+
+			FRenderCameraAerialPerspectiveVolumeCS::FParameters PassParameters;
+			PassParameters.Atmosphere = Scene->GetSkyAtmosphereSceneInfo()->GetAtmosphereUniformBuffer();
+			PassParameters.SkyAtmosphere = InternalCommonParametersRef;
+			PassParameters.ViewUniformBuffer = View.ViewUniformBuffer;
+			PassParameters.TransmittanceLutTextureSampler = SamplerLinearClamp;
+			PassParameters.MultiScatteredLuminanceLutTextureSampler = SamplerLinearClamp;
+			PassParameters.TransmittanceLutTexture = SkyInfo.GetTransmittanceLutTexture()->GetRenderTargetItem().ShaderResourceTexture;
+			PassParameters.MultiScatteredLuminanceLutTexture = SkyInfo.GetMultiScatteredLuminanceLutTexture()->GetRenderTargetItem().ShaderResourceTexture;
+			PassParameters.CameraAerialPerspectiveVolumeUAV = View.SkyAtmosphereCameraAerialPerspectiveVolume->GetRenderTargetItem().UAV;
+
+			SetShaderParameters(RHICmdList, *ComputeShader, ComputeShader->GetComputeShader(), PassParameters);
+
+			FIntVector TextureSize = View.SkyAtmosphereCameraAerialPerspectiveVolume->GetDesc().GetSize();
+			const FIntVector NumGroups = FIntVector::DivideAndRoundUp(TextureSize, FRenderCameraAerialPerspectiveVolumeCS::GroupSize);
+			RHICmdList.DispatchComputeShader(NumGroups.X, NumGroups.Y, NumGroups.Z);
+			UnsetShaderUAVs(RHICmdList, *ComputeShader, ComputeShader->GetComputeShader());
+		}
+	}
+
+	// TODO
+	// RHICmdList.AutomaticCacheFlushAfterComputeShader(true);
+	// RHICmdList.FlushComputeShaderCache();
+}
+
+void FSceneRenderer::RenderSkyAtmosphere(FRHICommandListImmediate& RHICmdList)
+{
+	check(ShouldRenderSkyAtmosphere(Scene->GetSkyAtmosphereSceneInfo(), Scene->GetShaderPlatform())); // This should not be called if we should not render SkyAtmosphere
+
+	DECLARE_GPU_STAT(SkyAtmosphere);
+	SCOPED_DRAW_EVENT(RHICmdList, SkyAtmosphere);
+	SCOPED_GPU_STAT(RHICmdList, SkyAtmosphere);
+
+	FSkyAtmosphereRenderSceneInfo& SkyInfo = *Scene->GetSkyAtmosphereSceneInfo();
+	const bool bMultiScattering = SkyInfo.IsMultiScatteringEnabled();
+	const bool bFastSky = CVarSkyAtmosphereFastSkyLUT.GetValueOnRenderThread() > 0;
+	const bool bFastAerialPespective = CVarSkyAtmosphereAerialPerspectiveApplyOnOpaque.GetValueOnRenderThread() > 0;
+	const bool bFastAerialPespectiveDepthTest = CVarSkyAtmosphereAerialPerspectiveDepthTest.GetValueOnRenderThread() > 0;
+	const bool bSecondAtmosphereLightEnabled = IsSecondAtmosphereLightEnabled(Scene);
+
+	FSceneRenderTargets& SceneContext = FSceneRenderTargets::Get(RHICmdList);
+
+	FRHIBlendState* PreMultipliedColorTransmittanceBlend = TStaticBlendState<CW_RGB, BO_Add, BF_One, BF_SourceAlpha, BO_Add, BF_Zero, BF_One>::GetRHI();
+	FRHISamplerState* SamplerLinearClamp = TStaticSamplerState<SF_Trilinear>::GetRHI();
+
+	// Initialise common internal parameters
+	FSkyAtmosphereInternalCommonParameters InternalCommonParameters;
+	SetupSkyAtmosphereInternalCommonParameters(InternalCommonParameters, *Scene, SkyInfo);
+	TUniformBufferRef<FSkyAtmosphereInternalCommonParameters> InternalCommonParametersRef = TUniformBufferRef<FSkyAtmosphereInternalCommonParameters>::CreateUniformBufferImmediate(InternalCommonParameters, UniformBuffer_SingleFrame);
+
+	const bool bIsMobilePlatformAndOnlySkyPixel = IsMobilePlatform(Scene->GetShaderPlatform());
+	const bool bUseOpaqueDepthTestAtStartDepth = bFastAerialPespectiveDepthTest && !bIsMobilePlatformAndOnlySkyPixel;
+
+	for (int32 ViewIndex = 0; ViewIndex < Views.Num(); ViewIndex++)
+	{
+		FViewInfo& View = Views[ViewIndex];
+
+		const bool bLightDiskEnabled = !View.bIsReflectionCapture;
+
+		// Render the sky, and optionally the atmosphere aerial perspective, on the scene luminance buffer
+		if (!bIsMobilePlatformAndOnlySkyPixel)
+		{
+			// This is skipped on mobile as this pass is part of the BasePass
+			FRHITexture* RTs[1] = { SceneContext.GetSceneColor()->GetRenderTargetItem().TargetableTexture };
+			FRHIRenderPassInfo RPInfo(1, RTs, ERenderTargetActions::DontLoad_Store, SceneContext.GetSceneDepthTexture(), EDepthStencilTargetActions::LoadDepthNotStencil_DontStore, FExclusiveDepthStencil::DepthRead_StencilWrite);
+			RHICmdList.BeginRenderPass(RPInfo, TEXT("DrawUVDisplacement"));
+		}
+		{
+			FRenderSkyAtmospherePS::FPermutationDomain PermutationVector;
+			PermutationVector.Set<FSkyPermutationMultiScatteringApprox>(bMultiScattering);
+			PermutationVector.Set<FFastSky>(bFastSky);
+			PermutationVector.Set<FFastAerialPespective>(bFastAerialPespective);
+			PermutationVector.Set<FSourceDiskEnabled>(bLightDiskEnabled);
+			PermutationVector.Set<FOnlySkyPixel>(bIsMobilePlatformAndOnlySkyPixel);
+			PermutationVector.Set<FSecondAtmosphereLight>(bSecondAtmosphereLightEnabled);
+			TShaderMapRef<FRenderSkyAtmospherePS> PixelShader(View.ShaderMap, PermutationVector);
+
+			FRenderSkyAtmosphereVS::FPermutationDomain VsPermutationVector;
+			VsPermutationVector.Set<FOnlySkyPixel>(bIsMobilePlatformAndOnlySkyPixel);
+			TShaderMapRef<FRenderSkyAtmosphereVS> VertexShader(View.ShaderMap, VsPermutationVector);
+
+			FRenderSkyAtmospherePS::FParameters PassParameters;
+			PassParameters.Atmosphere = Scene->GetSkyAtmosphereSceneInfo()->GetAtmosphereUniformBuffer();
+			PassParameters.SkyAtmosphere = InternalCommonParametersRef;
+			PassParameters.ViewUniformBuffer = View.ViewUniformBuffer;
+			//PassParameters->RenderTargets[0] = FRenderTargetBinding(SceneColor, ERenderTargetLoadAction::ELoad, ERenderTargetStoreAction::EStore);
+			PassParameters.SceneTextures = CreateSceneTextureUniformBufferSingleDraw(RHICmdList, ESceneTextureSetupMode::All, View.FeatureLevel);
+			PassParameters.TransmittanceLutTextureSampler = SamplerLinearClamp;
+			PassParameters.MultiScatteredLuminanceLutTextureSampler = SamplerLinearClamp;
+			PassParameters.SkyViewLutTextureSampler = SamplerLinearClamp;
+			PassParameters.CameraAerialPerspectiveVolumeTextureSampler = SamplerLinearClamp;
+
+			PassParameters.TransmittanceLutTexture = SkyInfo.GetTransmittanceLutTexture()->GetRenderTargetItem().ShaderResourceTexture;
+			PassParameters.MultiScatteredLuminanceLutTexture = SkyInfo.GetMultiScatteredLuminanceLutTexture()->GetRenderTargetItem().ShaderResourceTexture;
+			PassParameters.SkyViewLutTexture = View.SkyAtmosphereViewLutTexture->GetRenderTargetItem().ShaderResourceTexture;
+			PassParameters.CameraAerialPerspectiveVolumeTexture = View.SkyAtmosphereCameraAerialPerspectiveVolume->GetRenderTargetItem().ShaderResourceTexture;
+
+			FRenderSkyAtmosphereVS::FParameters VsPassParameters;
+			VsPassParameters.ViewUniformBuffer = View.ViewUniformBuffer;
+			VsPassParameters.StartDepthZ = 0.1f;
+			if (bUseOpaqueDepthTestAtStartDepth)
+			{
+				const FMatrix ViewProjectionMatrix = View.ViewMatrices.GetProjectionMatrix();
+				float HalfHorizontalFOV = FMath::Atan(1.0f / ViewProjectionMatrix.M[0][0]);
+				float HalfVerticalFOV = FMath::Atan(1.0f / ViewProjectionMatrix.M[1][1]);
+				const float HalfFOV = View.FOV * 0.5f * PI / 180.0f;
+				const float KmToCm = 100000.0f;
+				const float StartDepthView = FMath::Cos(FMath::Max(HalfHorizontalFOV, HalfVerticalFOV)) * InternalCommonParameters.CameraAerialPerspectiveVolumeStartDepth * KmToCm;
+				const FVector4 Projected = ViewProjectionMatrix.TransformFVector4(FVector4(0.0f, 0.0f, StartDepthView, 1.0f));
+				VsPassParameters.StartDepthZ = Projected.Z / Projected.W;
+			}
+
+			FGraphicsPipelineStateInitializer GraphicsPSOInit;
+			
+			RHICmdList.ApplyCachedRenderTargets(GraphicsPSOInit);
+
+			GraphicsPSOInit.RasterizerState = TStaticRasterizerState<FM_Solid, CM_None>::GetRHI(); ;
+			GraphicsPSOInit.BlendState = PreMultipliedColorTransmittanceBlend;
+			if (bIsMobilePlatformAndOnlySkyPixel || bUseOpaqueDepthTestAtStartDepth)
+			{
+				GraphicsPSOInit.DepthStencilState = TStaticDepthStencilState<false, CF_DepthNearOrEqual>::GetRHI();
+			}
+			else
+			{
+				GraphicsPSOInit.DepthStencilState = TStaticDepthStencilState<false, CF_Always>::GetRHI();
+			}
+
+			GraphicsPSOInit.BoundShaderState.VertexDeclarationRHI = GEmptyVertexDeclaration.VertexDeclarationRHI;
+			GraphicsPSOInit.BoundShaderState.VertexShaderRHI = GETSAFERHISHADER_VERTEX(*VertexShader);
+			GraphicsPSOInit.BoundShaderState.PixelShaderRHI = GETSAFERHISHADER_PIXEL(*PixelShader);
+			GraphicsPSOInit.PrimitiveType = PT_TriangleList;
+
+			RHICmdList.SetViewport(0, 0, 0.0f, View.ViewRect.Width(), View.ViewRect.Height(), 1.0f);
+
+			SetGraphicsPipelineState(RHICmdList, GraphicsPSOInit);
+
+			SetShaderParameters(RHICmdList, *PixelShader, PixelShader->GetPixelShader(), PassParameters);
+			SetShaderParameters(RHICmdList, *VertexShader, VertexShader->GetVertexShader(), VsPassParameters);
+
+			RHICmdList.DrawPrimitive(0, 1, 1);
+		}
+		if (!bIsMobilePlatformAndOnlySkyPixel)
+		{
+			RHICmdList.EndRenderPass();
+		}
+	}
+
+#if WITH_EDITOR
+	if (CVarSkyAtmosphereFastSkyLUT.GetValueOnAnyThread() == 0 && CVarSkyAtmosphereAerialPerspectiveApplyOnOpaque.GetValueOnAnyThread() > 0)
+	{
+		for (int32 ViewIndex = 0; ViewIndex < Views.Num(); ViewIndex++)
+		{
+			FViewInfo& View = Views[ViewIndex];
+			const float ViewPortWidth = float(View.ViewRect.Width());
+			const float ViewPortHeight = float(View.ViewRect.Height());
+
+			FRenderTargetTemp TempRenderTarget(View, (const FTexture2DRHIRef&)SceneContext.GetSceneColor()->GetRenderTargetItem().TargetableTexture);
+			FCanvas Canvas(&TempRenderTarget, NULL, View.Family->CurrentRealTime, ViewFamily.CurrentWorldTime, ViewFamily.DeltaWorldTime, View.GetFeatureLevel());
+
+			FLinearColor TextColor(1.0f, 0.5f, 0.0f);
+			FString Text = TEXT("You are using a FastAerialPespective without FastSky, visuals might look wrong.");
+			Canvas.DrawShadowedString(ViewPortWidth*0.5 - Text.Len()*7, ViewPortHeight*0.5, *Text, GetStatsFont(), TextColor);
+
+			Canvas.Flush_RenderThread(RHICmdList);
+		}
+	}
+#endif
+}
+
+
+
+/*=============================================================================
+	FDeferredShadingSceneRenderer functions
+=============================================================================*/
+
+
+
+void FDeferredShadingSceneRenderer::RenderDebugSkyAtmosphere(FRHICommandListImmediate& RHICmdList)
+{
+#if WITH_EDITOR
+	check(ShouldRenderSkyAtmosphere(Scene->GetSkyAtmosphereSceneInfo(), Scene->GetShaderPlatform())); // This should not be called if we should not render SkyAtmosphere
+
+	//if (!RHISupportsComputeShaders()) return;	// TODO cannot render, add a ShouldRender function. Also should PipelineShouldCook ?
+
+	DECLARE_GPU_STAT(SkyAtmosphereDebugViz);
+	SCOPED_DRAW_EVENT(RHICmdList, SkyAtmosphereDebugViz);
+	SCOPED_GPU_STAT(RHICmdList, SkyAtmosphereDebugViz);
+
+	const bool SkyAtmosphereVisualize = CVarSkyAtmosphereVisualize.GetValueOnRenderThread();
+	FSkyAtmosphereRenderSceneInfo& SkyInfo = *Scene->GetSkyAtmosphereSceneInfo();
+	const FAtmosphereSetup& Atmosphere = SkyInfo.GetAtmosphereSetup();
+	FSceneRenderTargets& SceneContext = FSceneRenderTargets::Get(RHICmdList);
+
+	if (SkyAtmosphereVisualize)
+	{
+		const bool bMultiScattering = SkyInfo.IsMultiScatteringEnabled();
+		const bool bFastSky = CVarSkyAtmosphereFastSkyLUT.GetValueOnRenderThread() > 0;
+		const bool bFastAerialPespective = CVarSkyAtmosphereAerialPerspectiveApplyOnOpaque.GetValueOnRenderThread() > 0;
+
+		FRDGBuilder GraphBuilder(RHICmdList);
+		FRDGTextureRef SceneColor = GraphBuilder.RegisterExternalTexture(SceneContext.GetSceneColor(), TEXT("SceneColor"));
+		FRDGTextureRef SceneDepth = GraphBuilder.RegisterExternalTexture(SceneContext.SceneDepthZ, TEXT("SceneColor"));
+
+		FRHIBlendState* PreMultipliedColorTransmittanceBlend = TStaticBlendState<CW_RGB, BO_Add, BF_One, BF_SourceAlpha, BO_Add, BF_Zero, BF_One>::GetRHI();
+		FRHIDepthStencilState* DepthStencilStateWrite = TStaticDepthStencilState<true, CF_Always>::GetRHI();
+		FRHISamplerState* SamplerLinearClamp = TStaticSamplerState<SF_Trilinear>::GetRHI();
+
+		// Initialise common internal parameters
+		FSkyAtmosphereInternalCommonParameters InternalCommonParameters;
+		SetupSkyAtmosphereInternalCommonParameters(InternalCommonParameters, *Scene, SkyInfo);
+		TUniformBufferRef<FSkyAtmosphereInternalCommonParameters> InternalCommonParametersRef = TUniformBufferRef<FSkyAtmosphereInternalCommonParameters>::CreateUniformBufferImmediate(InternalCommonParameters, UniformBuffer_SingleFrame);
+
+		for (int32 ViewIndex = 0; ViewIndex < Views.Num(); ViewIndex++)
+		{
+			FViewInfo& View = Views[ViewIndex];
+
+			// Render the sky and atmosphere on the scene luminance buffer
+			{
+				FRenderDebugSkyAtmospherePS::FPermutationDomain PermutationVector;
+				PermutationVector.Set<FSkyPermutationMultiScatteringApprox>(bMultiScattering);
+				TShaderMapRef<FRenderDebugSkyAtmospherePS> PixelShader(View.ShaderMap, PermutationVector);
+
+				FRenderDebugSkyAtmospherePS::FParameters* PassParameters = GraphBuilder.AllocParameters<FRenderDebugSkyAtmospherePS::FParameters>();
+				PassParameters->Atmosphere = Scene->GetSkyAtmosphereSceneInfo()->GetAtmosphereUniformBuffer();
+				PassParameters->SkyAtmosphere = InternalCommonParametersRef;
+				PassParameters->ViewUniformBuffer = View.ViewUniformBuffer;
+				PassParameters->RenderTargets[0] = FRenderTargetBinding(SceneColor, ERenderTargetLoadAction::ELoad, ERenderTargetStoreAction::EStore);
+				PassParameters->RenderTargets.DepthStencil = FDepthStencilBinding(SceneDepth, ERenderTargetLoadAction::ELoad, ERenderTargetStoreAction::EStore,
+					ERenderTargetLoadAction::ENoAction, ERenderTargetStoreAction::ENoAction, FExclusiveDepthStencil::DepthWrite_StencilNop);
+				PassParameters->TransmittanceLutTextureSampler = SamplerLinearClamp;
+				PassParameters->MultiScatteredLuminanceLutTextureSampler = SamplerLinearClamp;
+				PassParameters->TransmittanceLutTexture = SkyInfo.GetTransmittanceLutTexture()->GetRenderTargetItem().ShaderResourceTexture;
+				PassParameters->MultiScatteredLuminanceLutTexture = SkyInfo.GetMultiScatteredLuminanceLutTexture()->GetRenderTargetItem().ShaderResourceTexture;
+				PassParameters->ViewPortWidth = float(View.ViewRect.Width());
+				PassParameters->ViewPortHeight = float(View.ViewRect.Height());
+
+				FPixelShaderUtils::AddFullscreenPass<FRenderDebugSkyAtmospherePS>(GraphBuilder, View.ShaderMap, RDG_EVENT_NAME("SkyAtmosphere"), *PixelShader, PassParameters,
+					View.ViewRect, PreMultipliedColorTransmittanceBlend, nullptr, DepthStencilStateWrite);
+			}
+		}
+		GraphBuilder.Execute();
+	}
+
+	// Now debug print
+	for (int32 ViewIndex = 0; ViewIndex < Views.Num(); ViewIndex++)
+	{
+		FViewInfo& View = Views[ViewIndex];
+		const float ViewPortWidth = float(View.ViewRect.Width());
+		const float ViewPortHeight = float(View.ViewRect.Height());
+
+		FRenderTargetTemp TempRenderTarget(View, (const FTexture2DRHIRef&)SceneContext.GetSceneColor()->GetRenderTargetItem().TargetableTexture);
+		FCanvas Canvas(&TempRenderTarget, NULL, View.Family->CurrentRealTime, ViewFamily.CurrentWorldTime, ViewFamily.DeltaWorldTime, View.GetFeatureLevel());
+
+		FLinearColor TextColor(FLinearColor::White);
+		FLinearColor GrayTextColor(FLinearColor::Gray);
+		FLinearColor WarningColor(1.0f, 0.5f, 0.0f);
+		FString Text;
+
+		const float ViewPlanetAltitude = (View.ViewLocation*0.00001f - FVector(0.0f, 0.0f, -Atmosphere.BottomRadius)).Size() - Atmosphere.BottomRadius;
+		const bool ViewUnderGroundLevel = ViewPlanetAltitude < 0.0f;
+		if (ViewUnderGroundLevel)
+		{
+			Text = FString::Printf(TEXT("SkyAtmosphere: View is %.3f km under the planet ground level!"), -ViewPlanetAltitude);
+			Canvas.DrawShadowedString(ViewPortWidth*0.5 - 250.0f, ViewPortHeight*0.5f, *Text, GetStatsFont(), WarningColor);
+		}
+
+		if (SkyAtmosphereVisualize)
+		{
+			// This needs to stay in sync with RenderSkyAtmosphereDebugPS.
+			const float DensityViewTop = ViewPortHeight * 0.1f;
+			const float DensityViewBottom = ViewPortHeight * 0.8f;
+			const float DensityViewLeft = ViewPortWidth * 0.8f;
+			const float Margin = 2.0f;
+			const float TimeOfDayViewHeight = 64.0f;
+			const float TimeOfDayViewTop = ViewPortHeight - (TimeOfDayViewHeight + Margin * 2.0);
+			const float HemiViewHeight = ViewPortWidth * 0.25f;
+			const float HemiViewTop = ViewPortHeight - HemiViewHeight - TimeOfDayViewHeight - Margin * 2.0;
+
+			Text = FString::Printf(TEXT("Atmosphere top = %.1f km"), Atmosphere.TopRadius - Atmosphere.BottomRadius);
+			Canvas.DrawShadowedString(DensityViewLeft, DensityViewTop, *Text, GetStatsFont(), TextColor);
+			Text = FString::Printf(TEXT("Rayleigh extinction"));
+			Canvas.DrawShadowedString(DensityViewLeft + 60.0f, DensityViewTop + 30.0f, *Text, GetStatsFont(), FLinearColor(FLinearColor::Red));
+			Text = FString::Printf(TEXT("Mie extinction"));
+			Canvas.DrawShadowedString(DensityViewLeft + 60.0f, DensityViewTop + 45.0f, *Text, GetStatsFont(), FLinearColor(FLinearColor::Green));
+			Text = FString::Printf(TEXT("Absorption"));
+			Canvas.DrawShadowedString(DensityViewLeft + 60.0f, DensityViewTop + 60.0f, *Text, GetStatsFont(), FLinearColor(FLinearColor::Blue));
+			Text = FString::Printf(TEXT("<=== Low visual contribution"));
+			Canvas.DrawShadowedString(DensityViewLeft + 2.0, DensityViewTop + 150.0f, *Text, GetStatsFont(), GrayTextColor);
+			Text = FString::Printf(TEXT("High visual contribution ===>"));
+			Canvas.DrawShadowedString(ViewPortWidth - 170.0f, DensityViewTop + 166.0f, *Text, GetStatsFont(), GrayTextColor);
+			Text = FString::Printf(TEXT("Ground level"));
+			Canvas.DrawShadowedString(DensityViewLeft, DensityViewBottom, *Text, GetStatsFont(), TextColor);
+
+			Text = FString::Printf(TEXT("Time-of-day preview"));
+			Canvas.DrawShadowedString(ViewPortWidth*0.5f - 80.0f, TimeOfDayViewTop, *Text, GetStatsFont(), TextColor);
+
+			Text = FString::Printf(TEXT("Hemisphere view"));
+			Canvas.DrawShadowedString(Margin, HemiViewTop, *Text, GetStatsFont(), TextColor);
+		}
+
+		if (ViewUnderGroundLevel || SkyAtmosphereVisualize)
+		{
+			Canvas.Flush_RenderThread(RHICmdList);
+		}
+	}
+
+#endif
+}
+
+
diff --git a/Engine/Source/Runtime/Renderer/Private/SkyAtmosphereRendering.h b/Engine/Source/Runtime/Renderer/Private/SkyAtmosphereRendering.h
new file mode 100644
--- /dev/null
+++ b/Engine/Source/Runtime/Renderer/Private/SkyAtmosphereRendering.h
@@ -0,0 +1,100 @@
+// Copyright 1998-2019 Epic Games, Inc. All Rights Reserved.
+
+/*=============================================================================
+	SkyAtmosphereRendering.h
+=============================================================================*/
+
+#pragma once
+
+#include "CoreMinimal.h"
+#include "RendererInterface.h"
+#include "RenderResource.h"
+#include "Rendering/SkyAtmosphereCommonData.h"
+
+
+class FLightSceneInfo;
+class USkyAtmosphereComponent;
+
+
+
+// Use as a global shader parameter struct and also the CPU structure representing the atmosphere it self.
+BEGIN_GLOBAL_SHADER_PARAMETER_STRUCT(FAtmosphereUniformShaderParameters, )
+	SHADER_PARAMETER(float, MultiScatteringFactor)
+	SHADER_PARAMETER(float, BottomRadius)
+	SHADER_PARAMETER(float, TopRadius)
+	SHADER_PARAMETER(float, RayleighDensityExpScale)
+	SHADER_PARAMETER(FLinearColor, RayleighScattering)
+	SHADER_PARAMETER(FLinearColor, MieScattering)
+	SHADER_PARAMETER(float, MieDensityExpScale)
+	SHADER_PARAMETER(FLinearColor, MieExtinction)
+	SHADER_PARAMETER(float, MiePhaseG)
+	SHADER_PARAMETER(FLinearColor, MieAbsorption)
+	SHADER_PARAMETER(float, AbsorptionDensity0LayerWidth)
+	SHADER_PARAMETER(float, AbsorptionDensity0ConstantTerm)
+	SHADER_PARAMETER(float, AbsorptionDensity0LinearTerm)
+	SHADER_PARAMETER(float, AbsorptionDensity1ConstantTerm)
+	SHADER_PARAMETER(float, AbsorptionDensity1LinearTerm)
+	SHADER_PARAMETER(FLinearColor, AbsorptionExtinction)
+	SHADER_PARAMETER(FLinearColor, GroundAlbedo)
+END_GLOBAL_SHADER_PARAMETER_STRUCT()
+
+// Extra data shared by the Sky/Atmosphere system with Base passes.
+BEGIN_GLOBAL_SHADER_PARAMETER_STRUCT(FSkyAtmosphereBasePassSharedUniformShaderParameters,)
+	SHADER_PARAMETER_TEXTURE(Texture3D, CameraAerialPerspectiveVolume)
+	SHADER_PARAMETER_SAMPLER(SamplerState, CameraAerialPerspectiveVolumeSampler)
+	SHADER_PARAMETER(float, CameraAerialPerspectiveVolumeStartDepth)
+	SHADER_PARAMETER(float, CameraAerialPerspectiveVolumeDepthResolution)
+	SHADER_PARAMETER(float, CameraAerialPerspectiveVolumeDepthResolutionInv)
+	SHADER_PARAMETER(float, CameraAerialPerspectiveVolumeDepthSliceLength)
+	SHADER_PARAMETER(float, CameraAerialPerspectiveVolumeDepthSliceLengthInv)
+	SHADER_PARAMETER(float, ApplyCameraAerialPerspectiveVolume)
+END_GLOBAL_SHADER_PARAMETER_STRUCT()
+
+
+
+class FSkyAtmosphereRenderSceneInfo
+{
+public:
+
+	/** Initialization constructor. */
+	explicit FSkyAtmosphereRenderSceneInfo(const USkyAtmosphereComponent* InComponent);
+	~FSkyAtmosphereRenderSceneInfo();
+
+	/** Prepare the sun light data as a function of current atmosphere state. */
+	void PrepareSunLightProxy(FLightSceneInfo& SunLight) const;
+
+	bool IsMultiScatteringEnabled() const { return AtmosphereSetup.MultiScatteringFactor > 0.0f; }
+	FLinearColor GetSkyLuminanceFactor() const { return SkyLuminanceFactor; }
+	float GetAerialPespectiveViewDistanceScale() const { return AerialPespectiveViewDistanceScale; }
+
+	const TUniformBufferRef<FAtmosphereUniformShaderParameters>& GetAtmosphereUniformBuffer() { return AtmosphereUniformBuffer; }
+	TRefCountPtr<IPooledRenderTarget>& GetTransmittanceLutTexture() { return TransmittanceLutTexture; }
+	TRefCountPtr<IPooledRenderTarget>& GetMultiScatteredLuminanceLutTexture() { return MultiScatteredLuminanceLutTexture; }
+	TRefCountPtr<IPooledRenderTarget>& GetDistantSkyLightLutTexture() { return DistantSkyLightLutTexture; }
+	FTextureRHIRef GetDistantSkyLightLutTextureRHI();
+
+	const FAtmosphereSetup GetAtmosphereSetup() const { return AtmosphereSetup; }
+	const FAtmosphereUniformShaderParameters* GetAtmosphereShaderParameters() const { return &AtmosphereUniformShaderParameters; }
+
+	bool bStaticLightingBuilt;
+
+private:
+
+	FAtmosphereSetup AtmosphereSetup;
+	FAtmosphereUniformShaderParameters AtmosphereUniformShaderParameters;
+
+	TUniformBufferRef<FAtmosphereUniformShaderParameters> AtmosphereUniformBuffer;
+	TRefCountPtr<IPooledRenderTarget> TransmittanceLutTexture;
+	TRefCountPtr<IPooledRenderTarget> MultiScatteredLuminanceLutTexture;
+	TRefCountPtr<IPooledRenderTarget> DistantSkyLightLutTexture;
+
+	FLinearColor TransmittanceAtZenith;
+	FLinearColor SkyLuminanceFactor;
+	float AerialPespectiveViewDistanceScale;
+};
+
+bool ShouldRenderSkyAtmosphere(const FSkyAtmosphereRenderSceneInfo* SkyAtmosphere, EShaderPlatform ShaderPlatform);
+
+extern void SetupSkyAtmosphereBasePassSharedUniformShaderParameters(const class FViewInfo& View, FSkyAtmosphereBasePassSharedUniformShaderParameters& OutParameters);
+
+
