Build fix attempt for 7555515

diff --git a/Engine/Source/Runtime/Engine/Private/Components/SkyAtmosphereComponent.cpp b/Engine/Source/Runtime/Engine/Private/Components/SkyAtmosphereComponent.cpp
--- a/Engine/Source/Runtime/Engine/Private/Components/SkyAtmosphereComponent.cpp
+++ b/Engine/Source/Runtime/Engine/Private/Components/SkyAtmosphereComponent.cpp
@@ -12,6 +12,7 @@
 #include "Misc/MapErrors.h"
 #include "Misc/UObjectToken.h"
 #include "UObject/UObjectIterator.h"
+#include "UObject/ConstructorHelpers.h"
 
 #if WITH_EDITOR
 #include "ObjectEditorUtils.h"
diff --git a/Engine/Source/Runtime/Renderer/Private/SkyAtmosphereRendering.cpp b/Engine/Source/Runtime/Renderer/Private/SkyAtmosphereRendering.cpp
--- a/Engine/Source/Runtime/Renderer/Private/SkyAtmosphereRendering.cpp
+++ b/Engine/Source/Runtime/Renderer/Private/SkyAtmosphereRendering.cpp
@@ -5,6 +5,7 @@
 =============================================================================*/
 
 #include "SkyAtmosphereRendering.h"
+#include "CanvasTypes.h"
 #include "Components/SkyAtmosphereComponent.h"
 #include "DeferredShadingRenderer.h"
 #include "Engine/MapBuildDataRegistry.h"
