Fixed build issue resulting from 7562289

diff --git a/Engine/Source/Runtime/Renderer/Private/SkyAtmosphereRendering.cpp b/Engine/Source/Runtime/Renderer/Private/SkyAtmosphereRendering.cpp
--- a/Engine/Source/Runtime/Renderer/Private/SkyAtmosphereRendering.cpp
+++ b/Engine/Source/Runtime/Renderer/Private/SkyAtmosphereRendering.cpp
@@ -1296,9 +1296,8 @@ void FDeferredShadingSceneRenderer::RenderDebugSkyAtmosphere(FRHICommandListImme
 				PassParameters->Atmosphere = Scene->GetSkyAtmosphereSceneInfo()->GetAtmosphereUniformBuffer();
 				PassParameters->SkyAtmosphere = InternalCommonParametersRef;
 				PassParameters->ViewUniformBuffer = View.ViewUniformBuffer;
-				PassParameters->RenderTargets[0] = FRenderTargetBinding(SceneColor, ERenderTargetLoadAction::ELoad, ERenderTargetStoreAction::EStore);
-				PassParameters->RenderTargets.DepthStencil = FDepthStencilBinding(SceneDepth, ERenderTargetLoadAction::ELoad, ERenderTargetStoreAction::EStore,
-					ERenderTargetLoadAction::ENoAction, ERenderTargetStoreAction::ENoAction, FExclusiveDepthStencil::DepthWrite_StencilNop);
+				PassParameters->RenderTargets[0] = FRenderTargetBinding(SceneColor, ERenderTargetLoadAction::ELoad);
+				PassParameters->RenderTargets.DepthStencil = FDepthStencilBinding(SceneDepth, ERenderTargetLoadAction::ELoad, ERenderTargetLoadAction::ENoAction, FExclusiveDepthStencil::DepthWrite_StencilNop);
 				PassParameters->TransmittanceLutTextureSampler = SamplerLinearClamp;
 				PassParameters->MultiScatteredLuminanceLutTextureSampler = SamplerLinearClamp;
 				PassParameters->TransmittanceLutTexture = SkyInfo.GetTransmittanceLutTexture()->GetRenderTargetItem().ShaderResourceTexture;
