Registers scene textures with debug names.

diff --git a/Engine/Source/Runtime/Renderer/Private/SceneTextureParameters.cpp b/Engine/Source/Runtime/Renderer/Private/SceneTextureParameters.cpp
--- a/Engine/Source/Runtime/Renderer/Private/SceneTextureParameters.cpp
+++ b/Engine/Source/Runtime/Renderer/Private/SceneTextureParameters.cpp
@@ -19,16 +19,16 @@ void SetupSceneTextureParameters(
 	*OutTextures = FSceneTextureParameters();
 
 	// Should always have a depth buffer around allocated, since early z-pass is first.
-	OutTextures->SceneDepthBuffer = GraphBuilder.RegisterExternalTexture(SceneContext.SceneDepthZ);
+	OutTextures->SceneDepthBuffer = GraphBuilder.RegisterExternalTexture(SceneContext.SceneDepthZ, TEXT("SceneDepthZ"));
 
 	// Registers all the scene texture from the scene context. No fallback is provided to catch mistake at shader parameter validation time
 	// when a pass is trying to access a resource before any other pass actually created it.
-	OutTextures->SceneVelocityBuffer = SceneContext.SceneVelocity ? GraphBuilder.RegisterExternalTexture(SceneContext.SceneVelocity) : nullptr;
-	OutTextures->SceneGBufferA = SceneContext.GBufferA ? GraphBuilder.RegisterExternalTexture(SceneContext.GBufferA) : nullptr;
-	OutTextures->SceneGBufferB = SceneContext.GBufferB ? GraphBuilder.RegisterExternalTexture(SceneContext.GBufferB) : nullptr;
-	OutTextures->SceneGBufferC = SceneContext.GBufferC ? GraphBuilder.RegisterExternalTexture(SceneContext.GBufferC) : nullptr;
-	OutTextures->SceneGBufferD = SceneContext.GBufferD ? GraphBuilder.RegisterExternalTexture(SceneContext.GBufferD) : nullptr;
-	OutTextures->SceneGBufferE = SceneContext.GBufferE ? GraphBuilder.RegisterExternalTexture(SceneContext.GBufferE) : nullptr;
+	OutTextures->SceneVelocityBuffer = SceneContext.SceneVelocity ? GraphBuilder.RegisterExternalTexture(SceneContext.SceneVelocity, TEXT("VelocityBuffer")) : nullptr;
+	OutTextures->SceneGBufferA = SceneContext.GBufferA ? GraphBuilder.RegisterExternalTexture(SceneContext.GBufferA, TEXT("GBufferA")) : nullptr;
+	OutTextures->SceneGBufferB = SceneContext.GBufferB ? GraphBuilder.RegisterExternalTexture(SceneContext.GBufferB, TEXT("GBufferB")) : nullptr;
+	OutTextures->SceneGBufferC = SceneContext.GBufferC ? GraphBuilder.RegisterExternalTexture(SceneContext.GBufferC, TEXT("GBufferC")) : nullptr;
+	OutTextures->SceneGBufferD = SceneContext.GBufferD ? GraphBuilder.RegisterExternalTexture(SceneContext.GBufferD, TEXT("GBufferD")) : nullptr;
+	OutTextures->SceneGBufferE = SceneContext.GBufferE ? GraphBuilder.RegisterExternalTexture(SceneContext.GBufferE, TEXT("GBufferE")) : nullptr;
 
 	// Ligthing channels might be disabled when all lights are on the same channel.
 	if (SceneContext.LightingChannels)
