Ported Mitchell-Netravali downsample pass to RDG.

diff --git a/Engine/Shaders/Private/PostProcessMitchellNetravali.usf b/Engine/Shaders/Private/PostProcessMitchellNetravali.usf
--- a/Engine/Shaders/Private/PostProcessMitchellNetravali.usf
+++ b/Engine/Shaders/Private/PostProcessMitchellNetravali.usf
@@ -5,10 +5,9 @@
 =============================================================================*/
 
 #include "Common.ush"
-#include "PostProcessCommon.ush"
+#include "ScreenPass.ush"
 #include "EyeAdaptationCommon.ush"
 
-
 //------------------------------------------------------- COMPILE TIME CONSTANTS
 
 #define THREADGROUP_SIZEX 8
@@ -80,20 +79,15 @@ float MitchellNetravali(float d, const float B, const float C)
 }
 
 
+Texture2D InputTexture;
+SamplerState InputSampler;
 
-//------------------------------------------------------- PARAMETERS
-
-float4 DispatchThreadToInputBufferUV;
-
-float2 DownscaleFactor;
-
-
-//------------------------------------------------------- OUTPUTS
-
-RWTexture2D<float4>	Output0;
+SCREEN_PASS_TEXTURE_VIEWPORT(Input)
 
+float2 DispatchThreadToInputUVScale;
+float2 DispatchThreadToInputUVBias;
 
-//------------------------------------------------------- ENTRY POINT
+RWTexture2D<float4>	OutputTexture;
 
 [numthreads(THREADGROUP_SIZEX, THREADGROUP_SIZEY, 1)]
 void DownsampleMainCS(uint2 DispatchThreadId : SV_DispatchThreadID)
@@ -107,35 +101,32 @@ void DownsampleMainCS(uint2 DispatchThreadId : SV_DispatchThreadID)
 	const float B = rcp(3.0);
 	const float C = rcp(3.0);
 
-	float2 InputBufferUV = DispatchThreadId * DispatchThreadToInputBufferUV.xy + DispatchThreadToInputBufferUV.zw;
+	float2 InputBufferUV = DispatchThreadId * DispatchThreadToInputUVScale + DispatchThreadToInputUVBias;
 
 	float4 OutColor = 0;
 	float OutWeight = 0;
 
-	#if 1
 	{
 		const uint KernelSize = 4;
 
 		// Position of the output pixel in the input buffer.
-		float2 OutputPixelPos = InputBufferUV * PostprocessInput0Size.xy;
+		float2 OutputPixelPos = InputBufferUV * Input_Extent;
 
 		float2 TopLeftKernel = floor(OutputPixelPos + (0.5 - 0.5 * KernelSize)) + 0.5;
 
 		UNROLL
 		for (uint x = 0; x < KernelSize; x++)
-		{	
+		{
 			UNROLL
 			for (uint y = 0; y < KernelSize; y++)
 			{
 				float2 SamplePixelPos = TopLeftKernel + float2(x, y);
 
-				float2 SampleUV = SamplePixelPos * PostprocessInput0Size.zw;
-				if (true)
-				{
-					SampleUV = clamp(SampleUV, PostprocessInput0MinMax.xy, PostprocessInput0MinMax.zw);
-				}
-				
-				float4 SampleColor = PostprocessInput0.SampleLevel(PostprocessInput0Sampler, SampleUV, 0);
+				float2 SampleUV = SamplePixelPos * Input_ExtentInverse;
+
+				SampleUV = clamp(SampleUV, Input_UVViewportBilinearMin, Input_UVViewportBilinearMax);
+
+				float4 SampleColor = InputTexture.SampleLevel(InputSampler, SampleUV, 0);
 
 				float2 PixelOffset = abs(SamplePixelPos - OutputPixelPos);
 
@@ -148,27 +139,6 @@ void DownsampleMainCS(uint2 DispatchThreadId : SV_DispatchThreadID)
 			}
 		}
 	}
-	#else
-	{
-		// TODO: This is assuming input view size == 2 * output view size. Actually need to implement Mitchell-Netrali filter.
-		UNROLL
-		for (uint i = 0; i < 4; i++)
-		{
-			float2 SampleUV = InputBufferUV + (kOffsetsCross3x3[i] * 0.5) * PostprocessInput0Size.zw;
-			if (true)
-			{
-				SampleUV = clamp(SampleUV, PostprocessInput0MinMax.xy, PostprocessInput0MinMax.zw);
-			}
-
-			float4 SampleColor = PostprocessInput0.SampleLevel(PostprocessInput0Sampler, SampleUV, 0);
-
-			float HdrWeight = HdrWeight4(SampleColor.rgb, FrameExposureScale);
-
-			OutColor += HdrWeight * SampleColor;
-			OutWeight += HdrWeight;
-		}
-	}
-	#endif
 	
 	if (OutWeight > 0)
 	{
@@ -180,9 +150,9 @@ void DownsampleMainCS(uint2 DispatchThreadId : SV_DispatchThreadID)
 		OutColor.a = 0;
 	#endif
 
-	if (all(DispatchThreadId < ViewportRect.zw))
+	if (all(DispatchThreadId < Input_ViewportMax))
 	{
-		Output0[DispatchThreadId] = OutColor;
+		OutputTexture[DispatchThreadId] = OutColor;
 	}
 }
  
\ No newline at end of file
diff --git a/Engine/Source/Runtime/Renderer/Private/PostProcess/PostProcessMitchellNetravali.cpp b/Engine/Source/Runtime/Renderer/Private/PostProcess/PostProcessMitchellNetravali.cpp
--- a/Engine/Source/Runtime/Renderer/Private/PostProcess/PostProcessMitchellNetravali.cpp
+++ b/Engine/Source/Runtime/Renderer/Private/PostProcess/PostProcessMitchellNetravali.cpp
@@ -18,134 +18,109 @@
 #include "PipelineStateCache.h"
 #include "PostProcessing.h"
 
-
-
-
-class FMitchellNetravaliDownsamplCS : public FGlobalShader
+class FMitchellNetravaliDownsampleCS : public FGlobalShader
 {
-	DECLARE_GLOBAL_SHADER(FMitchellNetravaliDownsamplCS);
+public:
+	DECLARE_GLOBAL_SHADER(FMitchellNetravaliDownsampleCS);
+	SHADER_USE_PARAMETER_STRUCT(FMitchellNetravaliDownsampleCS, FGlobalShader);
 
 	static bool ShouldCompilePermutation(const FGlobalShaderPermutationParameters& Parameters)
 	{
 		return IsFeatureLevelSupported(Parameters.Platform, ERHIFeatureLevel::SM5) && !IsOpenGLPlatform(Parameters.Platform);
 	}
 
-	FMitchellNetravaliDownsamplCS() {}
-	FMitchellNetravaliDownsamplCS(const ShaderMetaType::CompiledShaderInitializerType& Initializer)
-		: FGlobalShader(Initializer)
-	{
-		PostprocessParameter.Bind(Initializer.ParameterMap);
-		Output.Bind(Initializer.ParameterMap, TEXT("Output0"));
-		EyeAdaptation.Bind(Initializer.ParameterMap, TEXT("EyeAdaptation"));
-		DispatchThreadToInputBufferUV.Bind(Initializer.ParameterMap, TEXT("DispatchThreadToInputBufferUV"));
-		DownscaleFactor.Bind(Initializer.ParameterMap, TEXT("DownscaleFactor"));
-	}
-
-	virtual bool Serialize(FArchive& Ar) override
-	{
-		bool bShaderHasOutdatedParameters = FGlobalShader::Serialize(Ar);
-		Ar << PostprocessParameter << Output << EyeAdaptation << DispatchThreadToInputBufferUV << DownscaleFactor;
-		return bShaderHasOutdatedParameters;
-	}
-
-	FPostProcessPassParameters PostprocessParameter;
-	FShaderResourceParameter Output;
-	FShaderResourceParameter EyeAdaptation;
-	
-	FShaderParameter DispatchThreadToInputBufferUV;
-	FShaderParameter DownscaleFactor;
+	BEGIN_SHADER_PARAMETER_STRUCT(FParameters, )
+		SHADER_PARAMETER_STRUCT_REF(FViewUniformShaderParameters, ViewUniformBuffer)
+		SHADER_PARAMETER_STRUCT(FScreenPassTextureViewportParameters, Input)
+		SHADER_PARAMETER_RDG_TEXTURE_UAV(RWTexture2D<float4>, OutputTexture)
+		SHADER_PARAMETER_RDG_TEXTURE(Texture2D, InputTexture)
+		SHADER_PARAMETER_RDG_TEXTURE(Texture2D, EyeAdaptation)
+		SHADER_PARAMETER_SAMPLER(SamplerState, InputSampler)
+		SHADER_PARAMETER(FVector2D, DispatchThreadToInputUVScale)
+		SHADER_PARAMETER(FVector2D, DispatchThreadToInputUVBias)
+	END_SHADER_PARAMETER_STRUCT()
 };
 
-IMPLEMENT_GLOBAL_SHADER(FMitchellNetravaliDownsamplCS, "/Engine/Private/PostProcessMitchellNetravali.usf", "DownsampleMainCS", SF_Compute);
+IMPLEMENT_GLOBAL_SHADER(FMitchellNetravaliDownsampleCS, "/Engine/Private/PostProcessMitchellNetravali.usf", "DownsampleMainCS", SF_Compute);
 
-
-void FRCPassMitchellNetravaliDownsample::Process(FRenderingCompositePassContext& Context)
+FRDGTextureRef ComputeMitchellNetravaliDownsample(
+	FRDGBuilder& GraphBuilder,
+	const FScreenPassViewInfo& ScreenPassView,
+	FRDGTextureRef InputTexture,
+	const FIntRect InputViewport,
+	const FScreenPassTextureViewport OutputViewport)
 {
-	FSceneRenderTargets& SceneContext = FSceneRenderTargets::Get(Context.RHICmdList);
-
-	const FPooledRenderTargetDesc* InputDesc = GetInputDesc(ePId_Input0);
-	FIntPoint SrcSize = InputDesc->Extent;
-
-	const FSceneRenderTargetItem& DestRenderTarget = PassOutputs[0].RequestSurface(Context);
-
-	// Src rectangle.
-	FIntRect SrcRect = Params.InputViewRect;
-	FIntRect DestRect = Params.OutputViewRect;
-
-	SCOPED_DRAW_EVENTF(Context.RHICmdList, MitchellNetravaliDownsample, TEXT("MitchellNetravaliDownsample %dx%d -> %dx%d"),
-		SrcRect.Width(), SrcRect.Height(),
-		DestRect.Width(), DestRect.Height());
+	FRDGTextureRef EyeAdaptationTexture = nullptr;
 
-	// Common setup
-	UnbindRenderTargets(Context.RHICmdList);
-	Context.SetViewportAndCallRHI(DestRect, 0.0f, 1.0f);
-
-	FTextureRHIRef EyeAdaptationTex = GWhiteTexture->TextureRHI;
-	if (Context.View.HasValidEyeAdaptation())
+	if (ScreenPassView.View.HasValidEyeAdaptation())
 	{
-		EyeAdaptationTex = Context.View.GetEyeAdaptation(Context.RHICmdList)->GetRenderTargetItem().TargetableTexture;
+		EyeAdaptationTexture = GraphBuilder.RegisterExternalTexture(ScreenPassView.View.GetEyeAdaptation(), TEXT("EyeAdaptation"));
 	}
-
+	else
 	{
-		Context.RHICmdList.TransitionResource(EResourceTransitionAccess::ERWBarrier, EResourceTransitionPipeline::EGfxToCompute, DestRenderTarget.UAV);
-
-		auto* ShaderMap = Context.GetShaderMap();
-		TShaderMapRef<FMitchellNetravaliDownsamplCS> Shader(ShaderMap);
-		FRHIComputeShader* ShaderRHI = Shader->GetComputeShader();
-
-		Context.RHICmdList.SetComputeShader(ShaderRHI);
-
-		// Parameters plumbing.
-		{
-			Shader->PostprocessParameter.SetCS(ShaderRHI, Context, Context.RHICmdList);
-
-			Context.RHICmdList.SetUAVParameter(ShaderRHI, Shader->Output.GetBaseIndex(), DestRenderTarget.MipUAVs[0]);
+		EyeAdaptationTexture = GraphBuilder.RegisterExternalTexture(GSystemTextures.WhiteDummy, TEXT("EyeAdaptation"));
+	}
 
-			// For global samplers.
-			Shader->SetParameters<FViewUniformShaderParameters>(
-				Context.RHICmdList, ShaderRHI, Context.View.ViewUniformBuffer);
+	const FRDGTextureDesc OutputTextureDesc = FRDGTextureDesc::Create2DDesc(
+		OutputViewport.Extent,
+		PF_FloatRGBA,
+		FClearValueBinding::Black,
+		TexCreate_None,
+		TexCreate_UAV,
+		false);
+
+	FRDGTextureRef OutputTexture = GraphBuilder.CreateTexture(OutputTextureDesc, TEXT("MitchelNetravaliDownsampleOutput"));
+
+	FMitchellNetravaliDownsampleCS::FParameters* PassParameters = GraphBuilder.AllocParameters<FMitchellNetravaliDownsampleCS::FParameters>();
+	PassParameters->ViewUniformBuffer = ScreenPassView.View.ViewUniformBuffer;
+	PassParameters->Input = GetScreenPassTextureViewportParameters(FScreenPassTextureViewport(InputViewport, InputTexture));
+	PassParameters->InputTexture = InputTexture;
+	PassParameters->InputSampler = TStaticSamplerState<SF_Bilinear, AM_Clamp, AM_Clamp, AM_Clamp>::GetRHI();
+	PassParameters->OutputTexture = GraphBuilder.CreateUAV(OutputTexture);
+	PassParameters->EyeAdaptation = EyeAdaptationTexture;
+
+	// Scale / Bias factor to map the dispatch thread id to the input texture UV.
+	PassParameters->DispatchThreadToInputUVScale.X = InputViewport.Width()  / float(OutputViewport.Rect.Width()  * InputTexture->Desc.Extent.X);
+	PassParameters->DispatchThreadToInputUVScale.Y = InputViewport.Height() / float(OutputViewport.Rect.Height() * InputTexture->Desc.Extent.Y);
+	PassParameters->DispatchThreadToInputUVBias.X = PassParameters->DispatchThreadToInputUVScale.X * (0.5f * InputViewport.Min.X);
+	PassParameters->DispatchThreadToInputUVBias.Y = PassParameters->DispatchThreadToInputUVScale.Y * (0.5f * InputViewport.Min.Y);
+
+	TShaderMapRef<FMitchellNetravaliDownsampleCS> ComputeShader(ScreenPassView.View.ShaderMap);
+
+	FComputeShaderUtils::AddPass(
+		GraphBuilder,
+		RDG_EVENT_NAME("MitchellNetravaliDownsample %dx%d -> %dx%d", InputViewport.Width(), InputViewport.Height(), OutputViewport.Rect.Width(), OutputViewport.Rect.Height()),
+		*ComputeShader,
+		PassParameters,
+		FComputeShaderUtils::GetGroupCount(OutputViewport.Rect.Size(), FComputeShaderUtils::kGolden2DGroupSize));
+
+	return OutputTexture;
+}
 
-			FVector4 DispatchThreadToInputBufferUV;
-			DispatchThreadToInputBufferUV.X = SrcRect.Width() / float(DestRect.Width() * SrcSize.X);
-			DispatchThreadToInputBufferUV.Y = SrcRect.Height() / float(DestRect.Height() * SrcSize.Y);
-			DispatchThreadToInputBufferUV.Z = DispatchThreadToInputBufferUV.X * (0.5f + SrcRect.Min.X);
-			DispatchThreadToInputBufferUV.W = DispatchThreadToInputBufferUV.Y * (0.5f + SrcRect.Min.Y);
+FRenderingCompositeOutputRef ComputeMitchellNetravaliDownsample(
+	FRenderingCompositionGraph& Graph,
+	FRenderingCompositeOutputRef Input,
+	FIntRect InputViewport,
+	FScreenPassTextureViewport OutputViewport)
+{
+	FRenderingCompositePass* MitchelNetravaliDownsamplePass = Graph.RegisterPass(new(FMemStack::Get()) TRCPassForRDG<1, 1>(
+		[InputViewport, OutputViewport](FRenderingCompositePass* Pass, FRenderingCompositePassContext& InContext)
+	{
+		FRDGBuilder GraphBuilder(InContext.RHICmdList);
 
-			SetShaderValue(Context.RHICmdList, ShaderRHI, Shader->DispatchThreadToInputBufferUV, DispatchThreadToInputBufferUV);
+		FRDGTextureRef InputTexture = Pass->CreateRDGTextureForRequiredInput(GraphBuilder, ePId_Input0, TEXT("Input"));
 
-			FVector2D DownscaleFactor(DestRect.Width() / float(SrcRect.Width()), SrcRect.Width() / float(DestRect.Width()));
-			SetShaderValue(Context.RHICmdList, ShaderRHI, Shader->DownscaleFactor, DownscaleFactor);
-		}
+		FRDGTextureRef OutputTexture = ComputeMitchellNetravaliDownsample(GraphBuilder, FScreenPassViewInfo(InContext.View), InputTexture, InputViewport, OutputViewport);
 
-		FIntPoint DestSize(DestRect.Width(), DestRect.Height());
-		uint32 GroupSizeX = FMath::DivideAndRoundUp(DestSize.X, 8);
-		uint32 GroupSizeY = FMath::DivideAndRoundUp(DestSize.Y, 8);
-		DispatchComputeShader(Context.RHICmdList, *Shader, GroupSizeX, GroupSizeY, 1);
+		Pass->ExtractRDGTextureForOutput(GraphBuilder, ePId_Output0, OutputTexture);
 
-		// Unset.
-		{
-			Context.RHICmdList.SetUAVParameter(ShaderRHI, Shader->Output.GetBaseIndex(), nullptr);
-		}
+		InContext.SceneColorViewRect = OutputViewport.Rect;
+		InContext.ReferenceBufferSize = OutputViewport.Extent;
 
-		Context.RHICmdList.TransitionResource(EResourceTransitionAccess::EReadable, EResourceTransitionPipeline::EComputeToGfx, DestRenderTarget.UAV);
-	}
+		GraphBuilder.Execute();
+	}));
 
-	Context.SceneColorViewRect = DestRect;
-	Context.ReferenceBufferSize = Params.OutputExtent;
-}
+	MitchelNetravaliDownsamplePass->SetInput(ePId_Input0, Input);
 
-FPooledRenderTargetDesc FRCPassMitchellNetravaliDownsample::ComputeOutputDesc(EPassOutputId InPassOutputId) const
-{
-	FPooledRenderTargetDesc Ret = GetInput(ePId_Input0)->GetOutput()->RenderTargetDesc;
-	Ret.Flags &= ~(TexCreate_FastVRAM | TexCreate_Transient);
-	Ret.Reset();
-	Ret.Format = PF_FloatRGBA;
-	Ret.DebugName = TEXT("MitchellNetravaliDownsample");
-	Ret.AutoWritable = false;
-	Ret.TargetableFlags &= ~(TexCreate_RenderTargetable | TexCreate_UAV);
-	Ret.TargetableFlags |= bIsComputePass ? TexCreate_UAV : TexCreate_RenderTargetable;
-
-	Ret.Extent = Params.OutputExtent;
-
-	return Ret;
-}
+	return FRenderingCompositeOutputRef(MitchelNetravaliDownsamplePass);
+}
\ No newline at end of file
diff --git a/Engine/Source/Runtime/Renderer/Private/PostProcess/PostProcessMitchellNetravali.h b/Engine/Source/Runtime/Renderer/Private/PostProcess/PostProcessMitchellNetravali.h
--- a/Engine/Source/Runtime/Renderer/Private/PostProcess/PostProcessMitchellNetravali.h
+++ b/Engine/Source/Runtime/Renderer/Private/PostProcess/PostProcessMitchellNetravali.h
@@ -6,36 +6,11 @@
 
 #pragma once
 
-#include "CoreMinimal.h"
-#include "RendererInterface.h"
+#include "ScreenPass.h"
 #include "RenderingCompositionGraph.h"
 
-
-class FRCPassMitchellNetravaliDownsample : public TRenderingCompositePassBase<1, 1>
-{
-public:
-	struct FParameters
-	{
-		FIntRect InputViewRect;
-
-		FIntRect OutputViewRect;
-
-		FIntPoint OutputExtent;
-
-	};
-
-
-	FRCPassMitchellNetravaliDownsample(const FParameters& InParams)
-		: Params(InParams)
-	{
-		bIsComputePass = true;
-	}
-
-	// interface FRenderingCompositePass ---------
-	virtual void Process(FRenderingCompositePassContext& Context) override;
-	virtual void Release() override { delete this; }
-	virtual FPooledRenderTargetDesc ComputeOutputDesc(EPassOutputId InPassOutputId) const override;
-
-private:
-	const FParameters Params;
-};
+FRenderingCompositeOutputRef ComputeMitchellNetravaliDownsample(
+	FRenderingCompositionGraph& Graph,
+	FRenderingCompositeOutputRef Input,
+	FIntRect InputViewport,
+	FScreenPassTextureViewport OutputViewport);
\ No newline at end of file
diff --git a/Engine/Source/Runtime/Renderer/Private/PostProcess/PostProcessing.cpp b/Engine/Source/Runtime/Renderer/Private/PostProcess/PostProcessing.cpp
--- a/Engine/Source/Runtime/Renderer/Private/PostProcess/PostProcessing.cpp
+++ b/Engine/Source/Runtime/Renderer/Private/PostProcess/PostProcessing.cpp
@@ -1262,20 +1262,14 @@ void FPostProcessing::Process(FRHICommandListImmediate& RHICmdList, const FViewI
 					FIntPoint QuantizedOutputSize;
 					QuantizeSceneBufferSize(SecondaryViewRectSize, QuantizedOutputSize);
 
-					FRCPassMitchellNetravaliDownsample::FParameters Parameters;
-					Parameters.InputViewRect = TAAParameters.OutputViewRect;
-					Parameters.OutputViewRect.Min.X = 0;
-					Parameters.OutputViewRect.Min.Y = 0;
-					Parameters.OutputViewRect.Max = SecondaryViewRectSize;
+					const FIntRect InputViewport = TAAParameters.OutputViewRect;
 
-					Parameters.OutputExtent.X = FMath::Max(SceneContext.GetBufferSizeXY().X, QuantizedOutputSize.X);
-					Parameters.OutputExtent.Y = FMath::Max(SceneContext.GetBufferSizeXY().Y, QuantizedOutputSize.Y);
+					FScreenPassTextureViewport OutputViewport;
+					OutputViewport.Rect.Max = SecondaryViewRectSize;
+					OutputViewport.Extent.X = FMath::Max(SceneContext.GetBufferSizeXY().X, QuantizedOutputSize.X);
+					OutputViewport.Extent.Y = FMath::Max(SceneContext.GetBufferSizeXY().Y, QuantizedOutputSize.Y);
 
-					FRenderingCompositePass* TemporalAADownsample = Context.Graph.RegisterPass(
-						new(FMemStack::Get()) FRCPassMitchellNetravaliDownsample(Parameters));
-					TemporalAADownsample->SetInput(ePId_Input0, Context.FinalOutput);
-
-					Context.FinalOutput = FRenderingCompositeOutputRef(TemporalAADownsample);
+					Context.FinalOutput = ComputeMitchellNetravaliDownsample(Context.Graph, Context.FinalOutput, InputViewport, OutputViewport);
 				}
 
 				SSRInputChain = AddPostProcessMaterialChain(Context, BL_SSRInput);
