Fixed ensure with post process materials overwriting an existing render target.

diff --git a/Engine/Source/Runtime/RenderCore/Public/RenderGraphPass.h b/Engine/Source/Runtime/RenderCore/Public/RenderGraphPass.h
--- a/Engine/Source/Runtime/RenderCore/Public/RenderGraphPass.h
+++ b/Engine/Source/Runtime/RenderCore/Public/RenderGraphPass.h
@@ -122,11 +122,6 @@ public:
 		: FRDGPassParameterStruct(Parameters, &FParameterStruct::FTypeInfo::GetStructMetadata()->GetLayout())
 	{}
 
-	bool IsCompute() const
-	{
-		return Layout->NumRenderTargets() == 0;
-	}
-
 	uint32 GetParameterCount() const
 	{
 		return Layout->Resources.Num();
diff --git a/Engine/Source/Runtime/Renderer/Private/PostProcess/PostProcessMaterial.cpp b/Engine/Source/Runtime/Renderer/Private/PostProcess/PostProcessMaterial.cpp
--- a/Engine/Source/Runtime/Renderer/Private/PostProcess/PostProcessMaterial.cpp
+++ b/Engine/Source/Runtime/Renderer/Private/PostProcess/PostProcessMaterial.cpp
@@ -398,7 +398,9 @@ FRDGTextureRef ComputePostProcessMaterial(
 	FRDGTextureRef OutputTexture = Inputs.OverrideOutputTexture;
 
 	// We can re-use the scene color texture as the render target if we're not simultaneously reading from it.
-	if (!OutputTexture && !MaterialShaderMap->UsesSceneTexture(PPI_PostProcessInput0))
+	// This is only necessary to do if we're going to be priming content from the render target since it avoids
+	// the copy. Otherwise, we just allocate a new render target.
+	if (!OutputTexture && !MaterialShaderMap->UsesSceneTexture(PPI_PostProcessInput0) && bPrimeOutputColor)
 	{
 		OutputTexture = SceneColorTexture;
 	}
@@ -441,6 +443,12 @@ FRDGTextureRef ComputePostProcessMaterial(
 		OutputLoadAction = ERenderTargetLoadAction::EClear;
 	}
 
+	if (OutputTexture == SceneColorTexture)
+	{
+		ensureMsgf(OutputLoadAction == ERenderTargetLoadAction::ELoad,
+			TEXT("We only want to re-use the input texture if we're going to load its contents. Otherwise RDG will emit a warning."));
+	}
+
 	FPostProcessMaterialParameters* PostProcessMaterialParameters = GraphBuilder.AllocParameters<FPostProcessMaterialParameters>();
 
 	PostProcessMaterialParameters->PostProcessOutput = GetScreenPassTextureViewportParameters(SceneColorViewport);
