Fixes motion blur -rdgdebug validation failure

diff --git a/Engine/Source/Runtime/Renderer/Private/PostProcess/PostProcessMotionBlur.cpp b/Engine/Source/Runtime/Renderer/Private/PostProcess/PostProcessMotionBlur.cpp
--- a/Engine/Source/Runtime/Renderer/Private/PostProcess/PostProcessMotionBlur.cpp
+++ b/Engine/Source/Runtime/Renderer/Private/PostProcess/PostProcessMotionBlur.cpp
@@ -495,6 +495,8 @@ void ComputeMotionBlurVelocity(
 			FComputeShaderUtils::GetGroupCount(Viewports.Velocity.Extent, FComputeShaderUtils::kGolden2DGroupSize));
 	}
 
+	bool ScatterDilatation = IsMotionBlurScatterRequired(View, Viewports.Color);
+
 	FRDGTextureRef VelocityTileTexture =
 		GraphBuilder.CreateTexture(
 			FRDGTextureDesc::Create2DDesc(
@@ -502,7 +504,7 @@ void ComputeMotionBlurVelocity(
 				PF_FloatRGBA,
 				FClearValueBinding::None,
 				GFastVRamConfig.MotionBlur,
-				TexCreate_ShaderResource | TexCreate_UAV | TexCreate_RenderTargetable,
+				TexCreate_ShaderResource | (ScatterDilatation ? TexCreate_RenderTargetable : TexCreate_UAV),
 				false),
 			TEXT("DilatedVelocityTile"));
 
@@ -511,7 +513,7 @@ void ComputeMotionBlurVelocity(
 	VelocityDilateParameters.VelocityTile = Viewports.VelocityTileParameters;
 	VelocityDilateParameters.VelocityTileTexture = VelocityTileTextureSetup;
 
-	if (IsMotionBlurScatterRequired(View, Viewports.Color))
+	if (ScatterDilatation)
 	{
 		FRDGTextureRef VelocityTileDepthTexture =
 			GraphBuilder.CreateTexture(
