Implements BuildHZB with compute, and reduce a ClosestHZB

diff --git a/Engine/Shaders/Private/DiaphragmDOF/DOFReduce.usf b/Engine/Shaders/Private/DiaphragmDOF/DOFReduce.usf
--- a/Engine/Shaders/Private/DiaphragmDOF/DOFReduce.usf
+++ b/Engine/Shaders/Private/DiaphragmDOF/DOFReduce.usf
@@ -6,6 +6,7 @@
 
 #include "DOFDownsample.ush"
 #include "DOFHybridScatterCompilation.ush"
+#include "../ReductionCommon.ush"
 
 
 //------------------------------------------------------- ENUM VALUES
@@ -120,44 +121,6 @@ groupshared float4 GroupSharedArray[THREADGROUP_TOTALSIZE];
 
 //------------------------------------------------------- FUNCTIONS
 
-uint SignedRightShift(uint x, const int bitshift)
-{
-	if (bitshift > 0)
-	{
-		return x << asuint(bitshift);
-	}
-	else if (bitshift < 0)
-	{
-		return x >> asuint(-bitshift);
-	}
-	return x;
-}
-
-// Returns the pixel pos [[0; N[[^2 in a two dimensional tile size of N=2^TileSizeLog2, to
-// store at a given SharedArrayId in [[0; N^2[[, so that a following recursive 2x2 pixel
-// block reduction stays entirely LDS memory banks coherent.
-uint2 InitialTilePixelPositionForReduction2x2(const uint TileSizeLog2, uint SharedArrayId)
-{
-	uint x = 0;
-	uint y = 0;
-
-	UNROLL
-	for (uint i = 0; i < TileSizeLog2; i++)
-	{
-		const uint DestBitId = TileSizeLog2 - 1 - i;
-		const uint DestBitMask = 1 << DestBitId;
-		x |= DestBitMask & SignedRightShift(SharedArrayId, int(DestBitId) - int(i * 2 + 0));
-		y |= DestBitMask & SignedRightShift(SharedArrayId, int(DestBitId) - int(i * 2 + 1));
-	}
-
-	#if 0
-	const uint N = 1 << TileSizeLog2;
-	return uint2(SharedArrayId / N, SharedArrayId - N * (SharedArrayId / N));
-	#endif
-
-	return uint2(x, y);
-}
-
 // Reduce using GroupSharedArray
 void ReduceOperator(FCocDownsampleParams DownsampleParams, uint OutId, uint ReduceBankSize)
 {
@@ -648,8 +611,7 @@ void ReduceCS(
 			const uint ReduceBankSize = TileSize * TileSize;
 		
 			// GroupSharedArray has been written before.
-			if (i == 0)
-				GroupMemoryBarrierWithGroupSync();
+			GroupMemoryBarrierWithGroupSync();
 		
 			if (GroupThreadIndex < ReduceBankSize)
 			{
diff --git a/Engine/Shaders/Private/HZB.usf b/Engine/Shaders/Private/HZB.usf
--- a/Engine/Shaders/Private/HZB.usf
+++ b/Engine/Shaders/Private/HZB.usf
@@ -1,60 +1,196 @@
 // Copyright 1998-2019 Epic Games, Inc. All Rights Reserved.
 
 #include "Common.ush"
+#include "SceneTextureParameters.ush"
+#include "ReductionCommon.ush"
+#include "/Engine/Public/WaveBroadcastIntrinsics.ush"
 
+
+#define MAX_MIP_BATCH_SIZE 4
+#define GROUP_TILE_SIZE 8
+
+
+float4 DispatchThreadIdToBufferUV;
 float2 InvSize;
-float4 InputUvFactorAndOffset;
 float2 InputViewportMaxBound;
 
-Texture2D		Texture;
-SamplerState	TextureSampler;
+Texture2D		ParentTextureMip;
+SamplerState	ParentTextureMipSampler;
 
-void HZBBuildPS(noperspective float4 RenderTargetUVAndScreenPos : TEXCOORD0, out float4 OutColor : SV_Target0)
-{
-	float4 Depth;
 
-#if STAGE == 0
-	const float2 ViewUV = RenderTargetUVAndScreenPos.zw * float2(0.5f, -0.5f) + 0.5f;
-	const float2 InUV = ViewUV * InputUvFactorAndOffset.xy + InputUvFactorAndOffset.zw;
+float4 Gather4(Texture2D Texture, SamplerState TextureSampler, float2 BufferUV)
+#if COMPILER_GLSL || FEATURE_LEVEL < FEATURE_LEVEL_SM5
+{
+	float2 UV[4];
 	
 	// min(..., InputViewportMaxBound) because we don't want to sample outside of the viewport
 	// when the view size has odd dimensions on X/Y axis.
-#if COMPILER_GLSL || COMPILER_GLSL_ES2 || COMPILER_GLSL_ES3_1 || FEATURE_LEVEL < FEATURE_LEVEL_SM5
-	float2 UV[4];
-	UV[0] = min(InUV + float2(-0.25f, -0.25f) * InvSize, InputViewportMaxBound);
-	UV[1] = min(InUV + float2( 0.25f, -0.25f) * InvSize, InputViewportMaxBound);
-	UV[2] = min(InUV + float2(-0.25f,  0.25f) * InvSize, InputViewportMaxBound);
-	UV[3] = min(InUV + float2( 0.25f,  0.25f) * InvSize, InputViewportMaxBound);
-
-	Depth.x = SceneTexturesStruct.SceneDepthTexture.SampleLevel( SceneTexturesStruct.SceneDepthTextureSampler, UV[0], 0 ).r;
-	Depth.y = SceneTexturesStruct.SceneDepthTexture.SampleLevel( SceneTexturesStruct.SceneDepthTextureSampler, UV[1], 0 ).r;
-	Depth.z = SceneTexturesStruct.SceneDepthTexture.SampleLevel( SceneTexturesStruct.SceneDepthTextureSampler, UV[2], 0 ).r;
-	Depth.w = SceneTexturesStruct.SceneDepthTexture.SampleLevel( SceneTexturesStruct.SceneDepthTextureSampler, UV[3], 0 ).r;
+	UV[0] = min(BufferUV + float2(-0.25f, -0.25f) * InvSize, InputViewportMaxBound);
+	UV[1] = min(BufferUV + float2( 0.25f, -0.25f) * InvSize, InputViewportMaxBound);
+	UV[2] = min(BufferUV + float2(-0.25f,  0.25f) * InvSize, InputViewportMaxBound);
+	UV[3] = min(BufferUV + float2( 0.25f,  0.25f) * InvSize, InputViewportMaxBound);
+
+	float4 Out;
+	Out.x = Texture.SampleLevel(TextureSampler, UV[0], 0 ).r;
+	Out.y = Texture.SampleLevel(TextureSampler, UV[1], 0 ).r;
+	Out.z = Texture.SampleLevel(TextureSampler, UV[2], 0 ).r;
+	Out.w = Texture.SampleLevel(TextureSampler, UV[3], 0 ).r;
+
+	return Out;
+}
 #else
-	float2 UV = min(InUV + float2(-0.25f, -0.25f) * InvSize, InputViewportMaxBound);
-	Depth = SceneTexturesStruct.SceneDepthTexture.GatherRed( SceneTexturesStruct.SceneDepthTextureSampler, UV, 0 );
+{
+	float2 UV = min(BufferUV + float2(-0.25f, -0.25f) * InvSize, InputViewportMaxBound - InvSize);
+	return Texture.GatherRed(TextureSampler, UV, 0);
+}
 #endif
 
-#else
-	const float2 InUV = RenderTargetUVAndScreenPos.xy;
+
+#if COMPUTESHADER
+
+RWTexture2D<float> FurthestHZBOutput_0;
+RWTexture2D<float> FurthestHZBOutput_1;
+RWTexture2D<float> FurthestHZBOutput_2;
+RWTexture2D<float> FurthestHZBOutput_3;
+
+RWTexture2D<float> ClosestHZBOutput_0;
+RWTexture2D<float> ClosestHZBOutput_1;
+RWTexture2D<float> ClosestHZBOutput_2;
+RWTexture2D<float> ClosestHZBOutput_3;
+
+groupshared float SharedFurthestDeviceZ[GROUP_TILE_SIZE * GROUP_TILE_SIZE];
+groupshared float SharedClosestDeviceZ[GROUP_TILE_SIZE * GROUP_TILE_SIZE];
+
+void OutputMipLevel(uint MipLevel, uint2 OutputPixelPos, float FurthestDeviceZ, float ClosestDeviceZ)
+{
+	#if DIM_MIP_LEVEL_COUNT >= 2
+	if (MipLevel == 1)
+	{
+		#if DIM_FURTHEST
+			FurthestHZBOutput_1[OutputPixelPos] = FurthestDeviceZ;
+		#endif
+		#if DIM_CLOSEST
+			ClosestHZBOutput_1[OutputPixelPos] = ClosestDeviceZ;
+		#endif
+	}
+	#endif
+	#if DIM_MIP_LEVEL_COUNT >= 3
+	else if (MipLevel == 2)
+	{
+		#if DIM_FURTHEST
+			FurthestHZBOutput_2[OutputPixelPos] = FurthestDeviceZ;
+		#endif
+		#if DIM_CLOSEST
+			ClosestHZBOutput_2[OutputPixelPos] = ClosestDeviceZ;
+		#endif
+	}
+	#endif
+	#if DIM_MIP_LEVEL_COUNT >= 4
+	else if (MipLevel == 3)
+	{
+		#if DIM_FURTHEST
+			FurthestHZBOutput_3[OutputPixelPos] = FurthestDeviceZ;
+		#endif
+		#if DIM_CLOSEST
+			ClosestHZBOutput_3[OutputPixelPos] = ClosestDeviceZ;
+		#endif
+	}		
+	#endif
+}
+
+[numthreads(GROUP_TILE_SIZE, GROUP_TILE_SIZE, 1)]
+void HZBBuildCS(
+	uint2 GroupId : SV_GroupID,
+	uint GroupThreadIndex : SV_GroupIndex)
+{
+	#if DIM_MIP_LEVEL_COUNT == 1
+		uint2 GroupThreadId = uint2(GroupThreadIndex % GROUP_TILE_SIZE, GroupThreadIndex / GROUP_TILE_SIZE);
+	#else
+		uint2 GroupThreadId = InitialTilePixelPositionForReduction2x2(MAX_MIP_BATCH_SIZE - 1, GroupThreadIndex);
+	#endif
+
+	uint2 DispatchThreadId = GROUP_TILE_SIZE * GroupId + GroupThreadId;
+
+	float2 BufferUV = (DispatchThreadId + 0.5) * DispatchThreadIdToBufferUV.xy + DispatchThreadIdToBufferUV.zw;
+	float4 DeviceZ = Gather4(ParentTextureMip, ParentTextureMipSampler, BufferUV);
+
+	float FurthestDeviceZ = min(min(DeviceZ.x, DeviceZ.y), min(DeviceZ.z, DeviceZ.w));
+	float ClosestDeviceZ = max(max(DeviceZ.x, DeviceZ.y), max(DeviceZ.z, DeviceZ.w));
 	
-#if COMPILER_GLSL || COMPILER_GLSL_ES2 || COMPILER_GLSL_ES3_1 || FEATURE_LEVEL < FEATURE_LEVEL_SM5
-	float2 UV[4];
-	UV[0] = InUV + float2(-0.25f, -0.25f) * InvSize;
-	UV[1] = InUV + float2( 0.25f, -0.25f) * InvSize;
-	UV[2] = InUV + float2(-0.25f,  0.25f) * InvSize;
-	UV[3] = InUV + float2( 0.25f,  0.25f) * InvSize;
-
-	Depth.x = Texture.SampleLevel( TextureSampler, UV[0], 0 ).r;
-	Depth.y = Texture.SampleLevel( TextureSampler, UV[1], 0 ).r;
-	Depth.z = Texture.SampleLevel( TextureSampler, UV[2], 0 ).r;
-	Depth.w = Texture.SampleLevel( TextureSampler, UV[3], 0 ).r;
-#else
-	float2 UV = InUV + float2(-0.25f, -0.25f) * InvSize;
-	Depth = Texture.GatherRed( TextureSampler, UV, 0 );
-#endif
+	uint2 OutputPixelPos = DispatchThreadId;
+	
+	#if DIM_FURTHEST
+		FurthestHZBOutput_0[OutputPixelPos] = FurthestDeviceZ;
+	#endif
+	
+	#if DIM_CLOSEST
+		ClosestHZBOutput_0[OutputPixelPos] = ClosestDeviceZ;
+	#endif
 
-#endif
+	#if DIM_MIP_LEVEL_COUNT == 1
+	{
+		// NOP
+	}
+	#else
+	{
+		SharedFurthestDeviceZ[GroupThreadIndex] = FurthestDeviceZ;
+		SharedClosestDeviceZ[GroupThreadIndex] = ClosestDeviceZ;
+	
+		UNROLL
+		for (uint MipLevel = 1; MipLevel < DIM_MIP_LEVEL_COUNT; MipLevel++)
+		{
+			const uint TileSize = GROUP_TILE_SIZE / (1 << MipLevel);
+			const uint ReduceBankSize = TileSize * TileSize;
+			
+			// LDS has been written before.
+			if (MipLevel == 1)
+				GroupMemoryBarrierWithGroupSync();
+
+			BRANCH
+			if (GroupThreadIndex < ReduceBankSize)
+			{
+				float4 ParentFurthestDeviceZ;
+				float4 ParentClosestDeviceZ;
+				ParentFurthestDeviceZ[0] = FurthestDeviceZ;
+				ParentClosestDeviceZ[0] = ClosestDeviceZ;
+
+				UNROLL
+				for (uint i = 1; i < 4; i++)
+				{
+					uint LDSIndex = GroupThreadIndex + i * ReduceBankSize;
+					ParentFurthestDeviceZ[i] = SharedFurthestDeviceZ[LDSIndex];
+					ParentClosestDeviceZ[i] = SharedClosestDeviceZ[LDSIndex];
+				}
+				
+				FurthestDeviceZ = min(min(ParentFurthestDeviceZ.x, ParentFurthestDeviceZ.y), min(ParentFurthestDeviceZ.z, ParentFurthestDeviceZ.w));
+				ClosestDeviceZ =  max(max(ParentClosestDeviceZ.x,  ParentClosestDeviceZ.y),  max(ParentClosestDeviceZ.z,  ParentClosestDeviceZ.w));
+	
+				OutputPixelPos = OutputPixelPos >> 1;
+				OutputMipLevel(MipLevel, OutputPixelPos, FurthestDeviceZ, ClosestDeviceZ);
+				
+				SharedFurthestDeviceZ[GroupThreadIndex] = FurthestDeviceZ;
+				SharedClosestDeviceZ[GroupThreadIndex] = ClosestDeviceZ;
+			}
+		} // for (uint MipLevel = 1; MipLevel < DIM_MIP_LEVEL_COUNT; MipLevel++)
+	}
+	#endif
+} // HZBBuildCS
+
+
+#elif PIXELSHADER
+
+void HZBBuildPS(float4 SvPosition : SV_POSITION, out float4 OutColor : SV_Target0)
+{
+	float2 BufferUV = SvPosition * DispatchThreadIdToBufferUV.xy + DispatchThreadIdToBufferUV.zw;
+	float4 DeviceZ = Gather4(ParentTextureMip, ParentTextureMipSampler, BufferUV);
+
+	float FurthestDeviceZ = min(min(DeviceZ.x, DeviceZ.y), min(DeviceZ.z, DeviceZ.w));
 
-	OutColor = min( min(Depth.x, Depth.y), min(Depth.z, Depth.w) );
+	OutColor = FurthestDeviceZ;
 }
+
+
+#else
+	#error Unknown shader frequency
+
+#endif
diff --git a/Engine/Shaders/Private/ReductionCommon.ush b/Engine/Shaders/Private/ReductionCommon.ush
new file mode 100644
--- /dev/null
+++ b/Engine/Shaders/Private/ReductionCommon.ush
@@ -0,0 +1,60 @@
+// Copyright 1998-2019 Epic Games, Inc. All Rights Reserved.
+
+#pragma once
+
+#include "/Engine/Public/Platform.ush"
+
+
+uint SignedRightShift(uint x, const int bitshift)
+{
+	if (bitshift > 0)
+	{
+		return x << asuint(bitshift);
+	}
+	else if (bitshift < 0)
+	{
+		return x >> asuint(-bitshift);
+	}
+	return x;
+}
+
+
+// Returns the pixel pos [[0; N[[^2 in a two dimensional tile size of N=2^TileSizeLog2, to
+// store at a given SharedArrayId in [[0; N^2[[, so that a following recursive 2x2 pixel
+// block reduction stays entirely LDS memory banks coherent.
+uint2 InitialTilePixelPositionForReduction2x2(const uint TileSizeLog2, uint SharedArrayId)
+{
+	uint x = 0;
+	uint y = 0;
+
+	UNROLL
+	for (uint i = 0; i < TileSizeLog2; i++)
+	{
+		const uint DestBitId = TileSizeLog2 - 1 - i;
+		const uint DestBitMask = 1 << DestBitId;
+		x |= DestBitMask & SignedRightShift(SharedArrayId, int(DestBitId) - int(i * 2 + 0));
+		y |= DestBitMask & SignedRightShift(SharedArrayId, int(DestBitId) - int(i * 2 + 1));
+	}
+
+	#if 0
+	const uint N = 1 << TileSizeLog2;
+	return uint2(SharedArrayId / N, SharedArrayId - N * (SharedArrayId / N));
+	#endif
+
+	return uint2(x, y);
+}
+
+
+uint2 InitialTilePixelPositionForReduction2x2(const uint TileSizeLog2, const uint ReduceCount, uint SharedArrayId)
+{
+	uint2 p = InitialTilePixelPositionForReduction2x2(ReduceCount, SharedArrayId);
+
+	SharedArrayId = SharedArrayId >> (2 * ReduceCount);
+
+	const uint RemainingSize = 1 << (TileSizeLog2 - ReduceCount);
+
+	p.x |= ((SharedArrayId % RemainingSize) << ReduceCount);
+	p.y |= ((SharedArrayId / RemainingSize) << ReduceCount);
+
+	return p;
+}
diff --git a/Engine/Source/Runtime/Renderer/Private/DeferredShadingRenderer.cpp b/Engine/Source/Runtime/Renderer/Private/DeferredShadingRenderer.cpp
--- a/Engine/Source/Runtime/Renderer/Private/DeferredShadingRenderer.cpp
+++ b/Engine/Source/Runtime/Renderer/Private/DeferredShadingRenderer.cpp
@@ -353,7 +353,7 @@ void FDeferredShadingSceneRenderer::RenderFinish(FRHICommandListImmediate& RHICm
 	FSceneRenderTargets::Get(RHICmdList).SetLightAttenuation(0);
 }
 
-void BuildHZB( FRDGBuilder& GraphBuilder, FViewInfo& View );
+void BuildHZB(FRDGBuilder& GraphBuilder, const FSceneTextureParameters& SceneTextures, FViewInfo& View);
 
 /** 
 * Renders the view family. 
@@ -418,9 +418,13 @@ bool FDeferredShadingSceneRenderer::RenderHzb(FRHICommandListImmediate& RHICmdLi
 		if (bSSAO || bHZBOcclusion || bSSR || bSSGI)
 		{
 			FRDGBuilder GraphBuilder(RHICmdList);
+
+			FSceneTextureParameters SceneTextures;
+			SetupSceneTextureParameters(GraphBuilder, &SceneTextures);
+
 			{
 				RDG_EVENT_SCOPE(GraphBuilder, "BuildHZB(ViewId=%d)", ViewIndex);
-				BuildHZB(GraphBuilder, Views[ViewIndex]);
+				BuildHZB(GraphBuilder, SceneTextures, Views[ViewIndex]);
 			}
 			GraphBuilder.Execute();
 		}
diff --git a/Engine/Source/Runtime/Renderer/Private/SceneRendering.h b/Engine/Source/Runtime/Renderer/Private/SceneRendering.h
--- a/Engine/Source/Runtime/Renderer/Private/SceneRendering.h
+++ b/Engine/Source/Runtime/Renderer/Private/SceneRendering.h
@@ -1038,8 +1038,9 @@ public:
 	FOcclusionQueryBatcher IndividualOcclusionQueries;
 	FOcclusionQueryBatcher GroupedOcclusionQueries;
 
-	// Hierarchical Z Buffer
+	// Furthest and closest Hierarchical Z Buffer
 	TRefCountPtr<IPooledRenderTarget> HZB;
+	TRefCountPtr<IPooledRenderTarget> ClosestHZB;
 
 	int32 NumBoxReflectionCaptures;
 	int32 NumSphereReflectionCaptures;
diff --git a/Engine/Source/Runtime/Renderer/Private/SceneTextureReductions.cpp b/Engine/Source/Runtime/Renderer/Private/SceneTextureReductions.cpp
--- a/Engine/Source/Runtime/Renderer/Private/SceneTextureReductions.cpp
+++ b/Engine/Source/Runtime/Renderer/Private/SceneTextureReductions.cpp
@@ -11,33 +11,38 @@
 #include "PostProcess/PostProcessing.h" // for FPostProcessVS
 
 
-BEGIN_SHADER_PARAMETER_STRUCT(FHZBBuildPassParameters, )
-	RENDER_TARGET_BINDING_SLOTS()
-	SHADER_PARAMETER_RDG_TEXTURE_SRV(Texture2D, Texture)
-	SHADER_PARAMETER_SAMPLER(SamplerState, TextureSampler)
+static TAutoConsoleVariable<int32> CVarHZBBuildUseCompute(
+	TEXT("r.HZB.BuildUseCompute"), 1,
+	TEXT("Selects whether HZB should be built with compute."),
+	ECVF_RenderThreadSafe);
+
+
+
+BEGIN_SHADER_PARAMETER_STRUCT(FSharedHZBParameters, )
+	SHADER_PARAMETER(FVector4, DispatchThreadIdToBufferUV)
+	SHADER_PARAMETER(FVector2D, InputViewportMaxBound)
+	SHADER_PARAMETER(FVector2D, InvSize)
+
+	SHADER_PARAMETER_RDG_TEXTURE_SRV(Texture2D, ParentTextureMip)
+	SHADER_PARAMETER_SAMPLER(SamplerState, ParentTextureMipSampler)
+
+	SHADER_PARAMETER_STRUCT_REF(FViewUniformShaderParameters, View)
 END_SHADER_PARAMETER_STRUCT()
 
+
 class FHZBBuildPS : public FGlobalShader
 {
 	DECLARE_GLOBAL_SHADER(FHZBBuildPS);
 	SHADER_USE_PARAMETER_STRUCT(FHZBBuildPS, FGlobalShader)
 
-	class FStageDim : SHADER_PERMUTATION_BOOL("STAGE");
-	using FPermutationDomain = TShaderPermutationDomain<FStageDim>;
-	
 	BEGIN_SHADER_PARAMETER_STRUCT( FParameters, )
-		SHADER_PARAMETER( FVector2D,	InvSize )
-		SHADER_PARAMETER( FVector4,		InputUvFactorAndOffset )
-		SHADER_PARAMETER( FVector2D,	InputViewportMaxBound )
-
-		SHADER_PARAMETER_STRUCT_INCLUDE(FHZBBuildPassParameters, Pass)
-		SHADER_PARAMETER_STRUCT_REF(	FViewUniformShaderParameters,	View )
-		SHADER_PARAMETER_STRUCT_REF(	FSceneTexturesUniformParameters,SceneTextures )
+		SHADER_PARAMETER_STRUCT_INCLUDE(FSharedHZBParameters, Shared)
+		RENDER_TARGET_BINDING_SLOTS()
 	END_SHADER_PARAMETER_STRUCT()
 
 	static bool ShouldCompilePermutation(const FGlobalShaderPermutationParameters& Parameters)
 	{
-		return IsFeatureLevelSupported(Parameters.Platform, ERHIFeatureLevel::SM4);
+		return IsFeatureLevelSupported(Parameters.Platform, ERHIFeatureLevel::SM5);
 	}
 
 	static void ModifyCompilationEnvironment(const FGlobalShaderPermutationParameters& Parameters, FShaderCompilerEnvironment& OutEnvironment)
@@ -46,169 +51,226 @@ class FHZBBuildPS : public FGlobalShader
 	}
 };
 
+class FHZBBuildCS : public FGlobalShader
+{
+	DECLARE_GLOBAL_SHADER(FHZBBuildCS);
+	SHADER_USE_PARAMETER_STRUCT(FHZBBuildCS, FGlobalShader)
+
+	static constexpr int32 kMaxMipBatchSize = 4;
+
+	class FDimFurthest : SHADER_PERMUTATION_BOOL("DIM_FURTHEST");
+	class FDimClosest : SHADER_PERMUTATION_BOOL("DIM_CLOSEST");
+	class FDimMipLevelCount : SHADER_PERMUTATION_RANGE_INT("DIM_MIP_LEVEL_COUNT", 1, kMaxMipBatchSize);
+	using FPermutationDomain = TShaderPermutationDomain<FDimFurthest, FDimClosest, FDimMipLevelCount>;
+
+	BEGIN_SHADER_PARAMETER_STRUCT(FParameters, )
+		SHADER_PARAMETER_STRUCT_INCLUDE(FSharedHZBParameters, Shared)
+
+		SHADER_PARAMETER_RDG_TEXTURE_UAV_ARRAY(RWTexture2D<float>, FurthestHZBOutput, [kMaxMipBatchSize])
+		SHADER_PARAMETER_RDG_TEXTURE_UAV_ARRAY(RWTexture2D<float>, ClosestHZBOutput, [kMaxMipBatchSize])
+	END_SHADER_PARAMETER_STRUCT()
+
+	static bool ShouldCompilePermutation(const FGlobalShaderPermutationParameters& Parameters)
+	{
+		FPermutationDomain PermutationVector(Parameters.PermutationId);
+
+		// Necessarily reduce at least closest of furthest.
+		if (!PermutationVector.Get<FDimFurthest>() && !PermutationVector.Get<FDimClosest>())
+		{
+			return false;
+		}
+
+
+		return IsFeatureLevelSupported(Parameters.Platform, ERHIFeatureLevel::SM5);
+	}
+};
+
 IMPLEMENT_GLOBAL_SHADER(FHZBBuildPS, "/Engine/Private/HZB.usf", "HZBBuildPS", SF_Pixel);
+IMPLEMENT_GLOBAL_SHADER(FHZBBuildCS, "/Engine/Private/HZB.usf", "HZBBuildCS", SF_Compute);
+
 
+bool ShouldRenderScreenSpaceDiffuseIndirect(const FViewInfo& View);
 
-void BuildHZB(FRDGBuilder& GraphBuilder, FViewInfo& View)
+static bool RequireClosestDepthHZB(const FViewInfo& View)
+{
+	return ShouldRenderScreenSpaceDiffuseIndirect(View);
+}
+
+
+void BuildHZB(FRDGBuilder& GraphBuilder, const FSceneTextureParameters& SceneTextures, FViewInfo& View)
 {
 	QUICK_SCOPE_CYCLE_COUNTER(STAT_BuildHZB);
-	
-	// View.ViewRect.{Width,Height}() are most likely to be < 2^24, so the float
-	// conversion won't loss any precision (assuming float have 23bits for mantissa)
-	const int32 NumMipsX = FMath::Max(FPlatformMath::CeilToInt(FMath::Log2(float(View.ViewRect.Width()))) - 1, 1);
-	const int32 NumMipsY = FMath::Max(FPlatformMath::CeilToInt(FMath::Log2(float(View.ViewRect.Height()))) - 1, 1);
-	const uint32 NumMips = FMath::Max(NumMipsX, NumMipsY);
-
-	// Must be power of 2
-	const FIntPoint HZBSize( 1 << NumMipsX, 1 << NumMipsY );
-	View.HZBMipmap0Size = HZBSize;
 
-	//@DW: Configure texture creation
-	FRDGTextureDesc HZBDesc = FRDGTextureDesc::Create2DDesc(HZBSize, PF_R16F, FClearValueBinding::None, TexCreate_None, TexCreate_RenderTargetable | TexCreate_ShaderResource | TexCreate_NoFastClear, false, NumMips);
+	FIntPoint HZBSize;
+	int32 NumMips;
+	{
+		const int32 NumMipsX = FMath::Max(FPlatformMath::CeilToInt(FMath::Log2(float(View.ViewRect.Width()))) - 1, 1);
+		const int32 NumMipsY = FMath::Max(FPlatformMath::CeilToInt(FMath::Log2(float(View.ViewRect.Height()))) - 1, 1);
+
+		NumMips = FMath::Max(NumMipsX, NumMipsY);
+
+		// Must be power of 2
+		HZBSize = FIntPoint(1 << NumMipsX, 1 << NumMipsY);
+	}
+
+	bool bReduceClosestDepth = RequireClosestDepthHZB(View);
+	bool bUseCompute = bReduceClosestDepth || CVarHZBBuildUseCompute.GetValueOnRenderThread();
+
+	FRDGTextureDesc HZBDesc = FRDGTextureDesc::Create2DDesc(
+		HZBSize, PF_R16F,
+		FClearValueBinding::None,
+		TexCreate_None,
+		TexCreate_ShaderResource | (bUseCompute ? TexCreate_UAV : TexCreate_RenderTargetable),
+		/* bInForceSeparateTargetAndShaderResource = */ false,
+		NumMips);
 	HZBDesc.Flags |= GFastVRamConfig.HZB;
 
-	//@DW: Explicit creation of graph resource handles - full support for everything the RHI supports
-	//@DW: Now that we've created a resource handle, it will have to be passed around to other passes for manual wiring or put into a Blackboard structure for automatic wiring
-	FRDGTextureRef HZBTexture = GraphBuilder.CreateTexture(HZBDesc, TEXT("HZB"));
+	/** Closest and furthest HZB are intentionally in separate render target, because majority of the case you only one or the other.
+	 * Keeping them separate avoid doubling the size in cache for this cases, to avoid performance regression.
+	 */
+	FRDGTextureRef FurthestHZBTexture = GraphBuilder.CreateTexture(HZBDesc, TEXT("HZBFurthest"));
 
+	FRDGTextureRef ClosestHZBTexture = nullptr;
+	if (bReduceClosestDepth)
 	{
-		FHZBBuildPassParameters* PassParameters = GraphBuilder.AllocParameters<FHZBBuildPassParameters>();
-		PassParameters->RenderTargets[0] = FRenderTargetBinding( HZBTexture, ERenderTargetLoadAction::ENoAction);
-
-		//@DW - this pass only reads external textures, we don't have any graph inputs
-		GraphBuilder.AddPass(
-			RDG_EVENT_NAME("HZB(mip=0) %dx%d", HZBSize.X, HZBSize.Y),
-			PassParameters,
-			ERDGPassFlags::Raster,
-			[PassParameters, &View, HZBSize](FRHICommandListImmediate& RHICmdList)
-			{
-				FGraphicsPipelineStateInitializer GraphicsPSOInit;
-				GraphicsPSOInit.BlendState = TStaticBlendState<>::GetRHI();
-				GraphicsPSOInit.RasterizerState = TStaticRasterizerState<>::GetRHI();
-				GraphicsPSOInit.DepthStencilState = TStaticDepthStencilState<false, CF_Always>::GetRHI();
-
-				RHICmdList.ApplyCachedRenderTargets(GraphicsPSOInit);
-				GraphicsPSOInit.PrimitiveType = PT_TriangleList;
-
-				FHZBBuildPS::FPermutationDomain PermutationVector;
-				PermutationVector.Set<FHZBBuildPS::FStageDim>(false);
-
-				TShaderMapRef<FPostProcessVS> VertexShader(View.ShaderMap);
-				TShaderMapRef<FHZBBuildPS> PixelShader(View.ShaderMap, PermutationVector);
-
-				GraphicsPSOInit.BoundShaderState.VertexDeclarationRHI = GFilterVertexDeclaration.VertexDeclarationRHI;
-				GraphicsPSOInit.BoundShaderState.VertexShaderRHI = GETSAFERHISHADER_VERTEX(*VertexShader);
-				GraphicsPSOInit.BoundShaderState.PixelShaderRHI = GETSAFERHISHADER_PIXEL(*PixelShader);
-				GraphicsPSOInit.PrimitiveType = PT_TriangleList;
-
-				SetGraphicsPipelineState(RHICmdList, GraphicsPSOInit);
-
-				// Imperfect sampling, doesn't matter too much
-				FSceneRenderTargets& SceneContext = FSceneRenderTargets::Get(RHICmdList);
-				FIntPoint Size = SceneContext.GetBufferSizeXY();
-
-				FHZBBuildPS::FParameters Parameters;
-				Parameters.InvSize = FVector2D(1.0f / Size.X, 1.0f / Size.Y);
-				Parameters.InputUvFactorAndOffset = FVector4(
-					float(2 * HZBSize.X) / float(Size.X),
-					float(2 * HZBSize.Y) / float(Size.Y),
-					float(View.ViewRect.Min.X) / float(Size.X),
-					float(View.ViewRect.Min.Y) / float(Size.Y));
-				Parameters.InputViewportMaxBound = FVector2D(
-					float(View.ViewRect.Max.X) / float(Size.X) - 0.5f * Parameters.InvSize.X,
-					float(View.ViewRect.Max.Y) / float(Size.Y) - 0.5f * Parameters.InvSize.Y);
-		
-				Parameters.Pass = *PassParameters;
-				Parameters.View = View.ViewUniformBuffer;
-				Parameters.SceneTextures = CreateSceneTextureUniformBufferSingleDraw(RHICmdList, ESceneTextureSetupMode::SceneDepth, View.FeatureLevel);
-
-				SetShaderParameters(RHICmdList, *PixelShader, PixelShader->GetPixelShader(), Parameters);
-
-				RHICmdList.SetViewport(0, 0, 0.0f, HZBSize.X, HZBSize.Y, 1.0f);
-
-				DrawRectangle(
-					RHICmdList,
-					0, 0,
-					HZBSize.X, HZBSize.Y,
-					View.ViewRect.Min.X, View.ViewRect.Min.Y,
-					View.ViewRect.Width(), View.ViewRect.Height(),
-					HZBSize,
-					FSceneRenderTargets::Get(RHICmdList).GetBufferSizeXY(),
-					*VertexShader,
-					EDRF_UseTriangleOptimization);
-			});
+		ClosestHZBTexture = GraphBuilder.CreateTexture(HZBDesc, TEXT("HZBClosest"));
 	}
 
-	FIntPoint SrcSize = HZBSize;
-	FIntPoint DstSize = SrcSize / 2;
+	int32 MaxMipBatchSize = bUseCompute ? FHZBBuildCS::kMaxMipBatchSize : 1;
 
-	// Downsampling...
-	for (uint8 MipIndex = 1; MipIndex < NumMips; MipIndex++)
+	auto ReduceMips = [&](
+		FRDGTextureSRVRef ParentTextureMip, int32 StartDestMip, FVector4 DispatchThreadIdToBufferUV, FVector2D InputViewportMaxBound,
+		bool bOutputClosest, bool bOutputFurthest)
 	{
-		SrcSize.X = FMath::Max(SrcSize.X, 1);
-		SrcSize.Y = FMath::Max(SrcSize.Y, 1);
-		DstSize.X = FMath::Max(DstSize.X, 1);
-		DstSize.Y = FMath::Max(DstSize.Y, 1);
-
-		//@DW: Explicit creation of SRV, full configuration of SRV supported
-		FRDGTextureSRVDesc Desc = FRDGTextureSRVDesc::CreateForMipLevel(HZBTexture, MipIndex - 1);
-		FRDGTextureSRVRef ParentMipSRV = GraphBuilder.CreateSRV(Desc);
-
-		FHZBBuildPassParameters* PassParameters = GraphBuilder.AllocParameters<FHZBBuildPassParameters>();
-		PassParameters->RenderTargets[0] = FRenderTargetBinding( HZBTexture, ERenderTargetLoadAction::ENoAction, MipIndex );
-		PassParameters->Texture = ParentMipSRV;
-		PassParameters->TextureSampler = TStaticSamplerState<SF_Point, AM_Clamp, AM_Clamp, AM_Clamp>::GetRHI();
-
-		GraphBuilder.AddPass(
-			RDG_EVENT_NAME("HZB(mip=%d) %dx%d", MipIndex, DstSize.X, DstSize.Y),
-			PassParameters,
-			ERDGPassFlags::Raster | ERDGPassFlags::GenerateMips,
-			[PassParameters, SrcSize, DstSize, &View](FRHICommandListImmediate& RHICmdList)
-			{
-				FHZBBuildPS::FPermutationDomain PermutationVector;
-				PermutationVector.Set<FHZBBuildPS::FStageDim>(true);
+		FIntPoint SrcSize = FIntPoint::DivideAndRoundUp(ParentTextureMip->Desc.Texture->Desc.Extent, 1 << int32(ParentTextureMip->Desc.MipLevel));
 
-				TShaderMapRef<FPostProcessVS> VertexShader(View.ShaderMap);
-				TShaderMapRef<FHZBBuildPS> PixelShader(View.ShaderMap, PermutationVector);
+		FSharedHZBParameters ShaderParameters;
+		ShaderParameters.InvSize = FVector2D(1.0f / SrcSize.X, 1.0f / SrcSize.Y);
+		ShaderParameters.InputViewportMaxBound = InputViewportMaxBound;
+		ShaderParameters.DispatchThreadIdToBufferUV = DispatchThreadIdToBufferUV;
+		ShaderParameters.ParentTextureMip = ParentTextureMip;
+		ShaderParameters.ParentTextureMipSampler = TStaticSamplerState<SF_Point>::GetRHI();
+		ShaderParameters.View = View.ViewUniformBuffer;
 
-				FGraphicsPipelineStateInitializer GraphicsPSOInit;
-				GraphicsPSOInit.BlendState = TStaticBlendState<>::GetRHI();
-				GraphicsPSOInit.RasterizerState = TStaticRasterizerState<>::GetRHI();
-				GraphicsPSOInit.DepthStencilState = TStaticDepthStencilState<false, CF_Always>::GetRHI();
+		FIntPoint DstSize = FIntPoint::DivideAndRoundUp(HZBSize, 1 << StartDestMip);
 
-				RHICmdList.ApplyCachedRenderTargets(GraphicsPSOInit);
+		if (bUseCompute)
+		{
+			int32 EndDestMip = FMath::Min(StartDestMip + FHZBBuildCS::kMaxMipBatchSize, NumMips);
 
-				GraphicsPSOInit.BoundShaderState.VertexDeclarationRHI = GFilterVertexDeclaration.VertexDeclarationRHI;
-				GraphicsPSOInit.BoundShaderState.VertexShaderRHI = GETSAFERHISHADER_VERTEX(*VertexShader);
-				GraphicsPSOInit.BoundShaderState.PixelShaderRHI = GETSAFERHISHADER_PIXEL(*PixelShader);
-				GraphicsPSOInit.PrimitiveType = PT_TriangleList;
+			FHZBBuildCS::FParameters* PassParameters = GraphBuilder.AllocParameters<FHZBBuildCS::FParameters>();
+			PassParameters->Shared = ShaderParameters;
 
-				SetGraphicsPipelineState(RHICmdList, GraphicsPSOInit);
+			for (int32 DestMip = StartDestMip; DestMip < EndDestMip; DestMip++)
+			{
+				if (bOutputFurthest)
+					PassParameters->FurthestHZBOutput[DestMip - StartDestMip] = GraphBuilder.CreateUAV(FRDGTextureUAVDesc(FurthestHZBTexture, DestMip));
+				if (bOutputClosest)
+					PassParameters->ClosestHZBOutput[DestMip - StartDestMip] = GraphBuilder.CreateUAV(FRDGTextureUAVDesc(ClosestHZBTexture, DestMip));
+			}
+
+			FHZBBuildCS::FPermutationDomain PermutationVector;
+			PermutationVector.Set<FHZBBuildCS::FDimMipLevelCount>(EndDestMip - StartDestMip);
+			PermutationVector.Set<FHZBBuildCS::FDimFurthest>(bOutputFurthest);
+			PermutationVector.Set<FHZBBuildCS::FDimClosest>(bOutputClosest);
+
+			TShaderMapRef<FHZBBuildCS> ComputeShader(View.ShaderMap, PermutationVector);
+
+			// TODO(RDG): remove ERDGPassFlags::GenerateMips to use FComputeShaderUtils::AddPass().
+			ClearUnusedGraphResources(*ComputeShader, PassParameters);
+			GraphBuilder.AddPass(
+				RDG_EVENT_NAME("ReduceHZB(mips=[%d;%d]%s%s) %dx%d",
+					StartDestMip, EndDestMip - 1,
+					bOutputClosest ? TEXT(" Closest") : TEXT(""),
+					bOutputFurthest ? TEXT(" Furthest") : TEXT(""),
+					DstSize.X, DstSize.Y),
+				PassParameters,
+				StartDestMip ? (ERDGPassFlags::Compute | ERDGPassFlags::GenerateMips) : ERDGPassFlags::Compute,
+				[PassParameters, ComputeShader, DstSize](FRHICommandList& RHICmdList)
+			{
+				FComputeShaderUtils::Dispatch(RHICmdList, *ComputeShader, *PassParameters, FComputeShaderUtils::GetGroupCount(DstSize, 8));
+			});
+		}
+		else
+		{
+			check(bOutputFurthest);
+			check(!bOutputClosest);
+
+			FHZBBuildPS::FParameters* PassParameters = GraphBuilder.AllocParameters<FHZBBuildPS::FParameters>();
+			PassParameters->Shared = ShaderParameters;
+			PassParameters->RenderTargets[0] = FRenderTargetBinding(FurthestHZBTexture, ERenderTargetLoadAction::ENoAction, StartDestMip);
+
+			TShaderMapRef<FHZBBuildPS> PixelShader(View.ShaderMap);
+
+			// TODO(RDG): remove ERDGPassFlags::GenerateMips to use FPixelShaderUtils::AddFullscreenPass().
+			ClearUnusedGraphResources(*PixelShader, PassParameters);
+			GraphBuilder.AddPass(
+				RDG_EVENT_NAME("DownsampleHZB(mip=%d) %dx%d", StartDestMip, DstSize.X, DstSize.Y),
+				PassParameters,
+				StartDestMip ? (ERDGPassFlags::Raster | ERDGPassFlags::GenerateMips) : ERDGPassFlags::Raster,
+				[PassParameters, &View, PixelShader, DstSize](FRHICommandList& RHICmdList)
+			{
+				FPixelShaderUtils::DrawFullscreenPixelShader(RHICmdList, View.ShaderMap, *PixelShader, *PassParameters, FIntRect(0, 0, DstSize.X, DstSize.Y));
+			});
+		}
+	};
 
-				FHZBBuildPS::FParameters Parameters;
-				Parameters.InvSize = FVector2D(1.0f / SrcSize.X, 1.0f / SrcSize.Y);
-				Parameters.Pass = *PassParameters;
-				Parameters.View = View.ViewUniformBuffer;
+	// Reduce first mips Closesy and furtherest are done at same time.
+	{
+		FIntPoint SrcSize = SceneTextures.SceneDepthBuffer->Desc.Extent;
 
-				SetShaderParameters(RHICmdList, *PixelShader, PixelShader->GetPixelShader(), Parameters);
+		FRDGTextureSRVRef ParentTextureMip = GraphBuilder.CreateSRV(FRDGTextureSRVDesc::Create(SceneTextures.SceneDepthBuffer));
 
-				RHICmdList.SetViewport(0, 0, 0.0f, DstSize.X, DstSize.Y, 1.0f);
+		FVector4 DispatchThreadIdToBufferUV;
+		DispatchThreadIdToBufferUV.X = 2.0f / float(SrcSize.X);
+		DispatchThreadIdToBufferUV.Y = 2.0f / float(SrcSize.Y);
+		DispatchThreadIdToBufferUV.Z = View.ViewRect.Min.X / float(SrcSize.X);
+		DispatchThreadIdToBufferUV.W = View.ViewRect.Min.Y / float(SrcSize.Y);
 
-				DrawRectangle(
-					RHICmdList,
-					0, 0,
-					DstSize.X, DstSize.Y,
-					0, 0,
-					SrcSize.X, SrcSize.Y,
-					DstSize,
-					SrcSize,
-					*VertexShader,
-					EDRF_UseTriangleOptimization);
-			});
+		FVector2D InputViewportMaxBound = FVector2D(
+			float(View.ViewRect.Max.X - 0.5f) / float(SrcSize.X),
+			float(View.ViewRect.Max.Y - 0.5f) / float(SrcSize.Y));
 
-		SrcSize /= 2;
-		DstSize /= 2;
+		ReduceMips(
+			ParentTextureMip,
+			/* StartDestMip = */ 0, DispatchThreadIdToBufferUV, InputViewportMaxBound,
+			/* bOutputClosest = */ bReduceClosestDepth, /* bOutputFurthest = */ true);
 	}
 
-	GraphBuilder.QueueTextureExtraction(HZBTexture, &View.HZB);
-}
+	// Reduce the next mips
+	for (int32 StartDestMip = MaxMipBatchSize; StartDestMip < NumMips; StartDestMip += MaxMipBatchSize)
+	{
+		FIntPoint SrcSize = FIntPoint::DivideAndRoundUp(HZBSize, 1 << int32(StartDestMip - 1));
+
+		FVector4 DispatchThreadIdToBufferUV;
+		DispatchThreadIdToBufferUV.X = 2.0f / float(SrcSize.X);
+		DispatchThreadIdToBufferUV.Y = 2.0f / float(SrcSize.Y);
+		DispatchThreadIdToBufferUV.Z = 0.0f;
+		DispatchThreadIdToBufferUV.W = 0.0f;
+
+		FVector2D InputViewportMaxBound(1.0f, 1.0f);
+		
+		{
+			FRDGTextureSRVRef ParentTextureMip = GraphBuilder.CreateSRV(FRDGTextureSRVDesc::CreateForMipLevel(FurthestHZBTexture, StartDestMip - 1));
+			ReduceMips(ParentTextureMip,
+				StartDestMip, DispatchThreadIdToBufferUV, InputViewportMaxBound,
+				/* bOutputClosest = */ false, /* bOutputFurthest = */ true);
+		}
+
+		if (bReduceClosestDepth)
+		{
+			FRDGTextureSRVRef ParentTextureMip = GraphBuilder.CreateSRV(FRDGTextureSRVDesc::CreateForMipLevel(ClosestHZBTexture, StartDestMip - 1));
+			ReduceMips(ParentTextureMip,
+				StartDestMip, DispatchThreadIdToBufferUV, InputViewportMaxBound,
+				/* bOutputClosest = */ true, /* bOutputFurthest = */ false);
+		}
+	}
+
+	// Update the view.
+	View.HZBMipmap0Size = HZBSize;
+
+	GraphBuilder.QueueTextureExtraction(FurthestHZBTexture, &View.HZB);
+
+	if (ClosestHZBTexture)
+		GraphBuilder.QueueTextureExtraction(ClosestHZBTexture, &View.ClosestHZB);
+} // BuildHZB()
