Uses closest and furthest HZB to do tile classification more conservative about cone angle.

diff --git a/Engine/Shaders/Private/SSRT/SSRTDiffuseIndirect.usf b/Engine/Shaders/Private/SSRT/SSRTDiffuseIndirect.usf
--- a/Engine/Shaders/Private/SSRT/SSRTDiffuseIndirect.usf
+++ b/Engine/Shaders/Private/SSRT/SSRTDiffuseIndirect.usf
@@ -82,17 +82,17 @@
 	#define GroupMemoryBarrierWithGroupSync()
 #endif
 
-Texture2D HZBTexture;
+Texture2D FurthestHZBTexture;
 Texture2D ColorTexture;
 
 #if SUPPORTS_INDEPENDENT_SAMPLERS
 
-#define HZBTextureSampler GlobalPointClampedSampler
+#define FurthestHZBTextureSampler GlobalPointClampedSampler
 #define ColorTextureSampler GlobalBilinearClampedSampler
 
 #else
 
-SamplerState HZBTextureSampler;
+SamplerState FurthestHZBTextureSampler;
 SamplerState ColorTextureSampler;
 
 #endif
@@ -515,7 +515,7 @@ void MainCS(
 			#endif
 			{
 				bHit = CastScreenSpaceRay(
-					HZBTexture, HZBTextureSampler,
+					FurthestHZBTexture, FurthestHZBTextureSampler,
 					Ray, 1, CONFIG_RAY_STEPS, StepOffset,
 					HZBUvFactorAndInvFactor, bDebugPrint,
 					/* out */ HitUVz,
diff --git a/Engine/Shaders/Private/SSRT/SSRTTileClassification.usf b/Engine/Shaders/Private/SSRT/SSRTTileClassification.usf
--- a/Engine/Shaders/Private/SSRT/SSRTTileClassification.usf
+++ b/Engine/Shaders/Private/SSRT/SSRTTileClassification.usf
@@ -10,13 +10,16 @@
 
 float2 SamplePixelToHZBUV;
 
-Texture2D HZBTexture;
+Texture2D FurthestHZBTexture;
+Texture2D ClosestHZBTexture;
 
 #if SUPPORTS_INDEPENDENT_SAMPLERS
-#define HZBTextureSampler GlobalPointClampedSampler
+#define FurthestHZBTextureSampler GlobalPointClampedSampler
+#define ClosestHZBTextureSampler GlobalPointClampedSampler
 
 #else
-SamplerState HZBTextureSampler;
+SamplerState FurthestHZBTextureSampler;
+SamplerState ClosestHZBTextureSampler;
 
 #endif
 
@@ -66,7 +69,7 @@ void MainCS(
 	const float StartLevel = log2(GROUP_TILE_SIZE) - 1;
 
 	// TODO: should be furtherest depth.
-	float TileDeviceZ = HZBTexture.SampleLevel(HZBTextureSampler, RayStartPixel * SamplePixelToHZBUV, StartLevel).r;
+	float TileDeviceZ = FurthestHZBTexture.SampleLevel(FurthestHZBTextureSampler, RayStartPixel * SamplePixelToHZBUV, StartLevel).r;
 	
 	float MaxTheta = 0;
 
@@ -78,15 +81,15 @@ void MainCS(
 	{
 		float SamplePixelDistance = CurrentSampleDistance;
 
-		//if (SamplePixelDistance > 1.5 * MaxSampleDistance)
-		//	break;
+		if (SamplePixelDistance > 1.5 * MaxSampleDistance)
+			break;
 
 		float SampleLevel = StartLevel + i * MipLevelMultiplier;
 
 		float2 SamplePixel = RayStartPixel + RayPixelDirection * SamplePixelDistance;
 
 		// TODO: should be closest depth.
-		float SampleDiviceZ = HZBTexture.SampleLevel(HZBTextureSampler, SamplePixel * SamplePixelToHZBUV, SampleLevel).r;
+		float SampleDiviceZ = ClosestHZBTexture.SampleLevel(ClosestHZBTextureSampler, SamplePixel * SamplePixelToHZBUV, SampleLevel).r;
 
 		// Screen space distance in viewport UV from origin of the ray.
 		float DeltaU = SamplePixelDistance * View.ViewSizeAndInvSize.z;
@@ -107,8 +110,8 @@ void MainCS(
 		TileClassificationBufferOutput[TileIndex + ViewTileCount * DirectionId] = MaxTheta;
 	}
 
-	//if (DirectionId == View.GeneralPurposeTweak)
-	//{
-	//	DebugOutput[TileCoord] = float4(MaxTheta / PI, TileDeviceZ, 0, 0);
-	//}
+	if (DirectionId == View.GeneralPurposeTweak)
+	{
+		DebugOutput[TileCoord] = float4(MaxTheta / PI, TileDeviceZ, 0, 0);
+	}
 } // MainCS()
diff --git a/Engine/Source/Runtime/Renderer/Private/ScreenSpaceRayTracing.cpp b/Engine/Source/Runtime/Renderer/Private/ScreenSpaceRayTracing.cpp
--- a/Engine/Source/Runtime/Renderer/Private/ScreenSpaceRayTracing.cpp
+++ b/Engine/Source/Runtime/Renderer/Private/ScreenSpaceRayTracing.cpp
@@ -251,12 +251,13 @@ class FSSRTDiffuseTileClassificationCS : public FGlobalShader
 	BEGIN_SHADER_PARAMETER_STRUCT(FParameters, )
 		SHADER_PARAMETER(FVector2D, SamplePixelToHZBUV)
 
-		SHADER_PARAMETER_RDG_TEXTURE(Texture2D, HZBTexture)
-		SHADER_PARAMETER_SAMPLER(SamplerState, HZBTextureSampler)
+		SHADER_PARAMETER_RDG_TEXTURE(Texture2D, ClosestHZBTexture)
+		SHADER_PARAMETER_SAMPLER(SamplerState, ClosestHZBTextureSampler)
+
+		SHADER_PARAMETER_RDG_TEXTURE(Texture2D, FurthestHZBTexture)
+		SHADER_PARAMETER_SAMPLER(SamplerState, FurthestHZBTextureSampler)
 
 		SHADER_PARAMETER_STRUCT_REF(FViewUniformShaderParameters, View)
-		SHADER_PARAMETER_STRUCT_INCLUDE(FSceneTextureParameters, SceneTextures)
-		SHADER_PARAMETER_STRUCT_INCLUDE(FSceneTextureSamplerParameters, SceneTextureSamplers)
 
 		SHADER_PARAMETER_STRUCT_INCLUDE(FSSRTTileClassificationParameters, TileClassificationParameters)
 		SHADER_PARAMETER_STRUCT_INCLUDE(FSSRTTileClassificationUAVs, TileClassificationUAVs)
@@ -334,8 +335,8 @@ class FScreenSpaceDiffuseIndirectCS : public FGlobalShader
 		SHADER_PARAMETER( FVector4,		PrevScreenPositionScaleBias )
 		SHADER_PARAMETER( float,		PrevSceneColorPreExposureCorrection )
 		
-		SHADER_PARAMETER_RDG_TEXTURE(Texture2D, HZBTexture)
-		SHADER_PARAMETER_SAMPLER(SamplerState, HZBTextureSampler)
+		SHADER_PARAMETER_RDG_TEXTURE(Texture2D, FurthestHZBTexture)
+		SHADER_PARAMETER_SAMPLER(SamplerState, FurthestHZBTextureSampler)
 
 		SHADER_PARAMETER_RDG_TEXTURE(Texture2D, ColorTexture)
 		SHADER_PARAMETER_SAMPLER(SamplerState, ColorTextureSampler)
@@ -695,7 +696,8 @@ void RenderScreenSpaceDiffuseIndirect(
 	int32 RayCountPerPixel;
 	GetSSRTGIShaderOptionsForQuality(Quality, &GroupSize, &RayCountPerPixel);
 
-	FRDGTexture* HZBTexture = GraphBuilder.RegisterExternalTexture(View.ClosestHZB);
+	FRDGTexture* FurthestHZBTexture = GraphBuilder.RegisterExternalTexture(View.HZB);
+	FRDGTexture* ClosestHZBTexture = GraphBuilder.RegisterExternalTexture(View.ClosestHZB);
 	FRDGTexture* ColorTexture = GraphBuilder.RegisterExternalTexture(TemporalAAHistory.RT[0]);
 
 	// Reproject and reduce previous frame color.
@@ -761,14 +763,15 @@ void RenderScreenSpaceDiffuseIndirect(
 
 		FSSRTDiffuseTileClassificationCS::FParameters* PassParameters = GraphBuilder.AllocParameters<FSSRTDiffuseTileClassificationCS::FParameters>();
 		PassParameters->SamplePixelToHZBUV = FVector2D(
-			0.5f / float(HZBTexture->Desc.Extent.X),
-			0.5f / float(HZBTexture->Desc.Extent.Y));
+			0.5f / float(FurthestHZBTexture->Desc.Extent.X),
+			0.5f / float(FurthestHZBTexture->Desc.Extent.Y));
 
-		PassParameters->HZBTexture = HZBTexture;
-		PassParameters->HZBTextureSampler = TStaticSamplerState<SF_Point>::GetRHI();
+		PassParameters->FurthestHZBTexture = FurthestHZBTexture;
+		PassParameters->FurthestHZBTextureSampler = TStaticSamplerState<SF_Point>::GetRHI();
+
+		PassParameters->ClosestHZBTexture = ClosestHZBTexture;
+		PassParameters->ClosestHZBTextureSampler = TStaticSamplerState<SF_Point>::GetRHI();
 
-		PassParameters->SceneTextures = SceneTextures;
-		SetupSceneTextureSamplers(&PassParameters->SceneTextureSamplers);
 		PassParameters->View = View.ViewUniformBuffer;
 
 		PassParameters->TileClassificationParameters = ClassificationParameters;
@@ -814,8 +817,8 @@ void RenderScreenSpaceDiffuseIndirect(
 
 		FScreenSpaceDiffuseIndirectCS::FParameters* PassParameters = GraphBuilder.AllocParameters<FScreenSpaceDiffuseIndirectCS::FParameters>();
 
-		PassParameters->HZBTexture = HZBTexture;
-		PassParameters->HZBTextureSampler = TStaticSamplerState<SF_Point>::GetRHI();
+		PassParameters->FurthestHZBTexture = FurthestHZBTexture;
+		PassParameters->FurthestHZBTextureSampler = TStaticSamplerState<SF_Point>::GetRHI();
 		PassParameters->ColorTexture = ReducedSceneColor;
 		PassParameters->ColorTextureSampler = TStaticSamplerState<SF_Bilinear>::GetRHI();
 
