Switches SSGI's tile classification to structured buffer to actually benefit scalar load.

diff --git a/Engine/Shaders/Private/SSRT/SSRTTileClassification.usf b/Engine/Shaders/Private/SSRT/SSRTTileClassification.usf
--- a/Engine/Shaders/Private/SSRT/SSRTTileClassification.usf
+++ b/Engine/Shaders/Private/SSRT/SSRTTileClassification.usf
@@ -24,7 +24,7 @@ SamplerState ClosestHZBTextureSampler;
 #endif
 
 
-RWBuffer<float> TileClassificationBufferOutput;
+RWStructuredBuffer<float> TileClassificationBufferOutput;
 RWTexture2D<float4> DebugOutput;
 
 //void StoreTileInfos(uint2 TileCoord, FSSRTTileInfos TileInfos)
@@ -107,7 +107,7 @@ void MainCS(
 	if (TileCoord.x < TileBufferExtent.x && TileCoord.y < TileBufferExtent.y)
 	{
 		uint TileIndex = TileCoord.x + TileCoord.y * TileBufferExtent.x;
-		TileClassificationBufferOutput[TileIndex + ViewTileCount * DirectionId] = MaxTheta;
+		TileClassificationBufferOutput[TileIndex * DIRECTIONALITY_DIVISION + DirectionId] = MaxTheta;
 	}
 
 	if (DirectionId == View.GeneralPurposeTweak)
diff --git a/Engine/Shaders/Private/SSRT/SSRTTileClassificationBuffer.ush b/Engine/Shaders/Private/SSRT/SSRTTileClassificationBuffer.ush
--- a/Engine/Shaders/Private/SSRT/SSRTTileClassificationBuffer.ush
+++ b/Engine/Shaders/Private/SSRT/SSRTTileClassificationBuffer.ush
@@ -14,7 +14,7 @@ uint2 TileBufferExtent;
 uint ViewTileCount;
 uint MaxTileCount;
 
-Buffer<float> TileClassificationBuffer;
+StructuredBuffer<float> TileClassificationBuffer;
 
 
 /** Information about tiles in SGPRs. */
@@ -34,7 +34,7 @@ FSSRTTileInfos LoadTileInfos(uint2 TileCoord)
 	UNROLL_N(DIRECTIONALITY_DIVISION)
 	for (uint i = 0; i < DIRECTIONALITY_DIVISION; i++)
 	{
-		TileInfos.Directionality[i] = TileClassificationBuffer.Load(TileIndex + ViewTileCount * i);
+		TileInfos.Directionality[i] = TileClassificationBuffer[TileIndex * DIRECTIONALITY_DIVISION + i];
 	}
 	return TileInfos;
 }
diff --git a/Engine/Source/Runtime/Renderer/Private/ScreenSpaceRayTracing.cpp b/Engine/Source/Runtime/Renderer/Private/ScreenSpaceRayTracing.cpp
--- a/Engine/Source/Runtime/Renderer/Private/ScreenSpaceRayTracing.cpp
+++ b/Engine/Source/Runtime/Renderer/Private/ScreenSpaceRayTracing.cpp
@@ -160,15 +160,15 @@ BEGIN_SHADER_PARAMETER_STRUCT(FSSRTTileClassificationParameters, )
 END_SHADER_PARAMETER_STRUCT()
 
 BEGIN_SHADER_PARAMETER_STRUCT(FSSRTTileClassificationResources, )
-	SHADER_PARAMETER_RDG_BUFFER(Buffer<float>, TileClassificationBuffer)
+	SHADER_PARAMETER_RDG_BUFFER(StructuredBuffer<float>, TileClassificationBuffer)
 END_SHADER_PARAMETER_STRUCT()
 
 BEGIN_SHADER_PARAMETER_STRUCT(FSSRTTileClassificationSRVs, )
-	SHADER_PARAMETER_RDG_BUFFER_SRV(Buffer<float>, TileClassificationBuffer)
+	SHADER_PARAMETER_RDG_BUFFER_SRV(StructuredBuffer<float>, TileClassificationBuffer)
 END_SHADER_PARAMETER_STRUCT()
 
 BEGIN_SHADER_PARAMETER_STRUCT(FSSRTTileClassificationUAVs, )
-	SHADER_PARAMETER_RDG_BUFFER_UAV(RWBuffer<float>, TileClassificationBufferOutput)
+	SHADER_PARAMETER_RDG_BUFFER_UAV(RWStructuredBuffer<float>, TileClassificationBufferOutput)
 END_SHADER_PARAMETER_STRUCT()
 
 FSSRTTileClassificationResources CreateTileClassificationResources(FRDGBuilder& GraphBuilder, const FViewInfo& View, FIntPoint MaxRenderTargetSize, FSSRTTileClassificationParameters* OutParameters)
@@ -180,7 +180,7 @@ FSSRTTileClassificationResources CreateTileClassificationResources(FRDGBuilder&
 	OutParameters->ViewTileCount = OutParameters->TileBufferExtent.X * OutParameters->TileBufferExtent.Y;
 
 	FSSRTTileClassificationResources Resources;
-	Resources.TileClassificationBuffer = GraphBuilder.CreateBuffer(FRDGBufferDesc::CreateBufferDesc(sizeof(float), MaxTileCount * 8), TEXT("SSRTTileClassification"));
+	Resources.TileClassificationBuffer = GraphBuilder.CreateBuffer(FRDGBufferDesc::CreateStructuredDesc(sizeof(float), MaxTileCount * 8), TEXT("SSRTTileClassification"));
 	return Resources;
 }
 
