Fixed TAA not getting enabled for planar reflections.

diff --git a/Engine/Source/Runtime/Renderer/Private/SceneCaptureRendering.cpp b/Engine/Source/Runtime/Renderer/Private/SceneCaptureRendering.cpp
--- a/Engine/Source/Runtime/Renderer/Private/SceneCaptureRendering.cpp
+++ b/Engine/Source/Runtime/Renderer/Private/SceneCaptureRendering.cpp
@@ -562,6 +562,8 @@ void SetupViewVamilyForSceneCapture(
 		View->bIsSceneCapture = true;
 		// Note: this has to be set before EndFinalPostprocessSettings
 		View->bIsPlanarReflection = bIsPlanarReflection;
+        // Needs to be reconfigured now that bIsPlanarReflection has changed.
+		View->SetupAntiAliasingMethod();
 
 		check(SceneCaptureComponent);
 		for (auto It = SceneCaptureComponent->HiddenComponents.CreateConstIterator(); It; ++It)
