Implements a denoiser for polychromatic penumbra harmonics

Needs bilateral rejection refactor to correctly handle denoising of the specular therm.

diff --git a/Engine/Shaders/Private/ScreenSpaceDenoise/SSDCommon.ush b/Engine/Shaders/Private/ScreenSpaceDenoise/SSDCommon.ush
--- a/Engine/Shaders/Private/ScreenSpaceDenoise/SSDCommon.ush
+++ b/Engine/Shaders/Private/ScreenSpaceDenoise/SSDCommon.ush
@@ -41,6 +41,10 @@
 	// A gray scale valued encoded into a 16bit float only have 10bits mantissa.
 	#define TARGETED_SAMPLE_COUNT 1024
 
+#elif CONFIG_SIGNAL_PROCESSING == SIGNAL_PROCESSING_POLYCHROMATIC_PENUMBRA_HARMONIC
+	// Polychromatic penumbra is encded into 16bit float that only have 10bits mantissa.
+	#define TARGETED_SAMPLE_COUNT 1024
+
 #elif CONFIG_SIGNAL_PROCESSING == SIGNAL_PROCESSING_DIFFUSE_INDIRECT_AND_AO
 	// Indirect diffuse requires really large spatial kernel to clean the noise.
 	#define TARGETED_SAMPLE_COUNT 4096
diff --git a/Engine/Shaders/Private/ScreenSpaceDenoise/SSDComposeHarmonics.usf b/Engine/Shaders/Private/ScreenSpaceDenoise/SSDComposeHarmonics.usf
new file mode 100644
--- /dev/null
+++ b/Engine/Shaders/Private/ScreenSpaceDenoise/SSDComposeHarmonics.usf
@@ -0,0 +1,222 @@
+// Copyright 1998-2019 Epic Games, Inc. All Rights Reserved.
+
+#include "SSDDefinitions.ush"
+#include "SSDPublicHarmonics.ush"
+
+
+//------------------------------------------------------- CONFIGS
+
+#define TILE_PIXEL_SIZE 8
+
+#define CONFIG_SIGNAL_PROCESSING SIGNAL_PROCESSING_POLYCHROMATIC_PENUMBRA_HARMONIC
+#define CONFIG_SIGNAL_BATCH_SIZE 1
+
+#if CONFIG_SIGNAL_PROCESSING == SIGNAL_PROCESSING_POLYCHROMATIC_PENUMBRA_HARMONIC
+	#define COMPILE_SIGNAL_COLOR 1
+
+	//#define CONFIG_VIEW_POSITION_BASED_BILATERAL 2
+	#define CONFIG_MULTIPLEXED_SIGNALS_PER_SIGNAL_DOMAIN 1
+
+	// Input and output layout.
+	#define CONFIG_SIGNAL_INPUT_LAYOUT  SIGNAL_BUFFER_LAYOUT_POLYCHROMATIC_PENUMBRA_INPUT
+	#define CONFIG_SIGNAL_RECONSTRUCTED_LAYOUT  SIGNAL_BUFFER_LAYOUT_POLYCHROMATIC_PENUMBRA_HISTORY
+	#define CONFIG_SIGNAL_OUTPUT_LAYOUT SIGNAL_BUFFER_LAYOUT_POLYCHROMATIC_PENUMBRA_HISTORY
+	#define CONFIG_INPUT_TEXTURE_COUNT 3
+	#define CONFIG_OUTPUT_TEXTURE_COUNT 3
+
+	#define CONFIG_HARMONIC_COUNT DENOISER_HARMONICS_COUNT
+
+	// Denoise diffuse and specular harmonics at the same time.
+	#define MAX_SIGNAL_BATCH_SIZE 2
+
+#else
+	#error Unknown signal processing.
+#endif
+
+
+//------------------------------------------------------- INCLUDES
+
+#include "SSDSignalFramework.ush"
+#include "SSDSignalArray.ush"
+#include "SSDSignalBufferEncoding.ush"
+
+
+//------------------------------------------------------- PARAMETERS
+
+#if !defined(CONFIG_INPUT_TEXTURE_COUNT)
+	#error Missing CONFIG_INPUT_TEXTURE_COUNT
+#endif
+
+FSSDTexture2D SignalHarmonics_0_Textures_0;
+FSSDTexture2D SignalHarmonics_1_Textures_0;
+FSSDTexture2D SignalHarmonics_2_Textures_0;
+FSSDTexture2D SignalHarmonics_3_Textures_0;
+
+#if CONFIG_INPUT_TEXTURE_COUNT > 1
+FSSDTexture2D SignalHarmonics_0_Textures_1;
+FSSDTexture2D SignalHarmonics_1_Textures_1;
+FSSDTexture2D SignalHarmonics_2_Textures_1;
+FSSDTexture2D SignalHarmonics_3_Textures_1;
+#else
+#define SignalHarmonics_0_Textures_1 SignalHarmonics_0_Textures_0
+#define SignalHarmonics_1_Textures_1 SignalHarmonics_1_Textures_0
+#define SignalHarmonics_2_Textures_1 SignalHarmonics_2_Textures_0
+#define SignalHarmonics_3_Textures_1 SignalHarmonics_3_Textures_0
+#endif
+
+#if CONFIG_INPUT_TEXTURE_COUNT > 2
+FSSDTexture2D SignalHarmonics_0_Textures_2;
+FSSDTexture2D SignalHarmonics_1_Textures_2;
+FSSDTexture2D SignalHarmonics_2_Textures_2;
+FSSDTexture2D SignalHarmonics_3_Textures_2;
+#else
+#define SignalHarmonics_0_Textures_2 SignalHarmonics_0_Textures_0
+#define SignalHarmonics_1_Textures_2 SignalHarmonics_1_Textures_0
+#define SignalHarmonics_2_Textures_2 SignalHarmonics_2_Textures_0
+#define SignalHarmonics_3_Textures_2 SignalHarmonics_3_Textures_0
+#endif
+
+#if CONFIG_INPUT_TEXTURE_COUNT > 3
+FSSDTexture2D SignalHarmonics_0_Textures_3;
+FSSDTexture2D SignalHarmonics_1_Textures_3;
+FSSDTexture2D SignalHarmonics_2_Textures_3;
+FSSDTexture2D SignalHarmonics_3_Textures_3;
+#else
+#define SignalHarmonics_0_Textures_3 SignalHarmonics_0_Textures_0
+#define SignalHarmonics_1_Textures_3 SignalHarmonics_1_Textures_0
+#define SignalHarmonics_2_Textures_3 SignalHarmonics_2_Textures_0
+#define SignalHarmonics_3_Textures_3 SignalHarmonics_3_Textures_0
+#endif
+
+
+#if !defined(CONFIG_OUTPUT_TEXTURE_COUNT)
+	#error Missing CONFIG_OUTPUT_TEXTURE_COUNT
+#endif
+
+FSSDRWTexture2D SignalOutput_UAVs_0;
+
+#if CONFIG_OUTPUT_TEXTURE_COUNT > 1
+FSSDRWTexture2D	SignalOutput_UAVs_1;
+#else
+#define SignalOutput_UAVs_1 SignalOutput_UAVs_0
+#endif
+
+#if CONFIG_OUTPUT_TEXTURE_COUNT > 2
+FSSDRWTexture2D	SignalOutput_UAVs_2;
+#else
+#define SignalOutput_UAVs_2 SignalOutput_UAVs_0
+#endif
+
+#if CONFIG_OUTPUT_TEXTURE_COUNT > 3
+FSSDRWTexture2D	SignalOutput_UAVs_3;
+#else
+#define SignalOutput_UAVs_3 SignalOutput_UAVs_0
+#endif
+
+
+//------------------------------------------------------- ENTRY POINT
+
+[numthreads(TILE_PIXEL_SIZE, TILE_PIXEL_SIZE, 1)]
+void MainCS(
+	uint2 DispatchThreadId : SV_DispatchThreadID,
+	uint2 GroupId : SV_GroupID,
+	uint2 GroupThreadId : SV_GroupThreadID,
+	uint GroupThreadIndex : SV_GroupIndex)
+{
+	// Find out scene buffer UV.
+	float2 SceneBufferUV = (DispatchThreadId + 0.5 + View.ViewRectMin.xy) * View.BufferSizeAndInvSize.zw;
+	if (true)
+	{
+		SceneBufferUV = clamp(SceneBufferUV, View.BufferBilinearUVMinMax.xy, View.BufferBilinearUVMinMax.zw);
+	}
+	
+	// Read reference meta data.
+	float2 ScreenPosition = ViewportUVToScreenPos(BufferUVToViewportUV(SceneBufferUV));
+	FSSDSampleSceneInfos RefSceneMetadata = FetchCurrentSceneInfos(ScreenPosition, SceneBufferUV);
+
+	// Sample the reference sample.
+	FSSDSignalArray MultiplexedHarmonics[4];
+	{
+		MultiplexedHarmonics[0] = SampleMultiplexedSignals(
+			SignalHarmonics_0_Textures_0,
+			SignalHarmonics_0_Textures_1,
+			SignalHarmonics_0_Textures_2,
+			SignalHarmonics_0_Textures_3,
+			GlobalPointClampedSampler,
+			CONFIG_SIGNAL_INPUT_LAYOUT,
+			/* MultiplexedSampleId = */ 0,
+			/* bNormalizeSample = */ false,
+			SceneBufferUV,
+			/* MipLevel = */ 0);
+	
+		MultiplexedHarmonics[1] = SampleMultiplexedSignals(
+			SignalHarmonics_1_Textures_0,
+			SignalHarmonics_1_Textures_1,
+			SignalHarmonics_1_Textures_2,
+			SignalHarmonics_1_Textures_3,
+			GlobalPointClampedSampler,
+			CONFIG_SIGNAL_RECONSTRUCTED_LAYOUT,
+			/* MultiplexedSampleId = */ 0,
+			/* bNormalizeSample = */ false,
+			SceneBufferUV,
+			/* MipLevel = */ 0);
+	
+		MultiplexedHarmonics[2] = SampleMultiplexedSignals(
+			SignalHarmonics_2_Textures_0,
+			SignalHarmonics_2_Textures_1,
+			SignalHarmonics_2_Textures_2,
+			SignalHarmonics_2_Textures_3,
+			GlobalPointClampedSampler,
+			CONFIG_SIGNAL_RECONSTRUCTED_LAYOUT,
+			/* MultiplexedSampleId = */ 0,
+			/* bNormalizeSample = */ false,
+			SceneBufferUV,
+			/* MipLevel = */ 0);
+	
+		MultiplexedHarmonics[3] = SampleMultiplexedSignals(
+			SignalHarmonics_3_Textures_0,
+			SignalHarmonics_3_Textures_1,
+			SignalHarmonics_3_Textures_2,
+			SignalHarmonics_3_Textures_3,
+			GlobalPointClampedSampler,
+			CONFIG_SIGNAL_RECONSTRUCTED_LAYOUT,
+			/* MultiplexedSampleId = */ 0,
+			/* bNormalizeSample = */ false,
+			SceneBufferUV,
+			/* MipLevel = */ 0);
+	}
+
+	// Whitelists computation that needs to be done.
+	FSSDSignalArray OutputSamples = CreateSignalArrayFromScalarValue(0.0);
+	UNROLL_N(CONFIG_SIGNAL_BATCH_SIZE)
+	for (uint MultiplexId = 0; MultiplexId < CONFIG_SIGNAL_BATCH_SIZE; MultiplexId++)
+	{
+		const uint LowestFrequencyHarmonicId = CONFIG_HARMONIC_COUNT - 1;
+		float LowestFrequencyWeight;
+		{
+			FSSDSignalSample Sample = MultiplexedHarmonics[LowestFrequencyHarmonicId].Array[MultiplexId];
+			OutputSamples.Array[MultiplexId] = Sample;
+			LowestFrequencyWeight = Sample.SampleCount;
+		}
+
+		UNROLL_N(CONFIG_HARMONIC_COUNT - 1)
+		for (uint HarmonicId = 0; HarmonicId < LowestFrequencyHarmonicId; HarmonicId++)
+		{
+			FSSDSignalSample Sample = MultiplexedHarmonics[HarmonicId].Array[MultiplexId];
+
+			float SampleWeight = -min(-LowestFrequencyWeight / Sample.SampleCount, 0.0);
+
+			OutputSamples.Array[MultiplexId] = AddSignal(
+				OutputSamples.Array[MultiplexId],
+				MulSignal(Sample, SampleWeight));
+		}
+	}
+
+	OutputMultiplexedSignal(
+		SignalOutput_UAVs_0,
+		SignalOutput_UAVs_1,
+		SignalOutput_UAVs_2,
+		SignalOutput_UAVs_3,
+		CONFIG_SIGNAL_OUTPUT_LAYOUT, CONFIG_SIGNAL_BATCH_SIZE,
+		DispatchThreadId, OutputSamples.Array);
+} // MainCS
diff --git a/Engine/Shaders/Private/ScreenSpaceDenoise/SSDDefinitions.ush b/Engine/Shaders/Private/ScreenSpaceDenoise/SSDDefinitions.ush
--- a/Engine/Shaders/Private/ScreenSpaceDenoise/SSDDefinitions.ush
+++ b/Engine/Shaders/Private/ScreenSpaceDenoise/SSDDefinitions.ush
@@ -94,6 +94,13 @@
 	#define SIGNAL_BUFFER_LAYOUT_PENUMBRA_HISTORY 11
 
 
+	/** Buffer layout input for denoising multi polychromatic penumbra. */
+	#define SIGNAL_BUFFER_LAYOUT_POLYCHROMATIC_PENUMBRA_INPUT 0x1000
+
+	/** Internal buffer layout for denoising multi polychromatic penumbra. */
+	#define SIGNAL_BUFFER_LAYOUT_POLYCHROMATIC_PENUMBRA_HISTORY 0x1001
+
+
 	/** Buffer layout for the reflection given by the ray generation shader. */
 	#define SIGNAL_BUFFER_LAYOUT_REFLECTIONS_INPUT 3
 
@@ -132,17 +139,20 @@
 	/** Shadow penumbra denoising. */
 	#define SIGNAL_PROCESSING_MONOCHROMATIC_PENUMBRA 0
 
+	/** Shadow penumbra denoising. */
+	#define SIGNAL_PROCESSING_POLYCHROMATIC_PENUMBRA_HARMONIC 1
+
 	/** Reflection denoising. */
-	#define SIGNAL_PROCESSING_REFLECTIONS 1
+	#define SIGNAL_PROCESSING_REFLECTIONS 2
 
 	/** Ambient occlusion denoising. */
-	#define SIGNAL_PROCESSING_AO 2
+	#define SIGNAL_PROCESSING_AO 3
 
 	/** Diffuse indirect and ambient occlusion. */
-	#define SIGNAL_PROCESSING_DIFFUSE_INDIRECT_AND_AO 3
+	#define SIGNAL_PROCESSING_DIFFUSE_INDIRECT_AND_AO 4
 
 	/** Diffuse indirect using spherical harmonic. */
-	#define SIGNAL_PROCESSING_DIFFUSE_SPHERICAL_HARMONIC 4
+	#define SIGNAL_PROCESSING_DIFFUSE_SPHERICAL_HARMONIC 5
 
 
 
diff --git a/Engine/Shaders/Private/ScreenSpaceDenoise/SSDPublicHarmonics.ush b/Engine/Shaders/Private/ScreenSpaceDenoise/SSDPublicHarmonics.ush
new file mode 100644
--- /dev/null
+++ b/Engine/Shaders/Private/ScreenSpaceDenoise/SSDPublicHarmonics.ush
@@ -0,0 +1,71 @@
+// Copyright 1998-2019 Epic Games, Inc. All Rights Reserved.
+
+#pragma once
+
+
+//------------------------------------------------------- INCLUDE
+
+#include "../Common.ush"
+#include "../LightShaderParameters.ush"
+
+
+//------------------------------------------------------- DEFINES
+
+#define DENOISER_HARMONICS_COUNT 4
+
+
+//------------------------------------------------------- FUNCTIONS
+
+float ComputeLightSampleWorldBluringRadius(float3 TranslatedWorldPosition, uint LightType, FLightShaderParameters Light, float HitT)
+{
+	// Sometime, artists might put occluder very close to the area light compared to they area, that may lead to negative values.
+	// TODO: the correct way to fix this is to move this world bluring radius computation into the RGS, and have DistanceFromLight = Ray's MaxT.
+	const float ClosestLightDistance = 0.001;
+
+	BRANCH
+	if (LightType == LIGHT_TYPE_DIRECTIONAL)
+	{
+		// TODO: compute this on CPU.
+		return tan(asin(Light.SourceRadius)) * HitT;
+	}
+	else if (LightType == LIGHT_TYPE_POINT || LightType == LIGHT_TYPE_SPOT)
+	{
+		float3 PixelToLightWorldVector = (Light.Position + View.PreViewTranslation) - TranslatedWorldPosition;
+		float DistanceFromLight = length(PixelToLightWorldVector);
+
+		return Light.SourceRadius * HitT / max(DistanceFromLight - HitT, ClosestLightDistance);
+	}
+	else // if (LightType == LIGHT_TYPE_RECT)
+	{
+		float3 PixelToLightWorldVector = (Light.Position + View.PreViewTranslation) - TranslatedWorldPosition;
+		float DistanceFromLight = length(PixelToLightWorldVector);
+
+		#if CONFIG_VALU_OPTIMIZATIONS
+			float Scalar = abs(dot(PixelToLightWorldVector, Light.Direction)) * rsqrt(length2(PixelToLightWorldVector));
+		#else
+			float Scalar = abs(dot(normalize(PixelToLightWorldVector), Light.Direction));
+		#endif
+
+		float2 LightDimensions = GetRectLightDimensions(Light);
+
+		// TODO: area light anysotropy.
+		float SmallestLightDimension = min(LightDimensions.x, LightDimensions.y);
+
+		// TODO: HitT on the direction of the ray, but does not DistanceFromLight, witch is bad for large area light.
+		return SmallestLightDimension * HitT / max(DistanceFromLight - HitT, ClosestLightDistance);
+	}
+}
+
+float ProjectWorldDistanceToPixelDistance(float WorldDepth, float WorldDistance)
+{
+	return WorldDistance * View.ViewSizeAndInvSize.x * 0.5 / (WorldDepth * GetTanHalfFieldOfView().x);
+}
+
+float WorldBluringRadiusToHarmonic(float WorldDepth, float WorldBluringRadius)
+{
+	float PixelBluringRadius = ProjectWorldDistanceToPixelDistance(WorldDepth, WorldBluringRadius);
+
+	const float TemporalHistoryRejectionDiameter = 3.0;
+
+	return clamp(log2(PixelBluringRadius - (TemporalHistoryRejectionDiameter - 1.0)), 0.0, DENOISER_HARMONICS_COUNT - 1.0);
+}
diff --git a/Engine/Shaders/Private/ScreenSpaceDenoise/SSDSignalBufferEncoding.ush b/Engine/Shaders/Private/ScreenSpaceDenoise/SSDSignalBufferEncoding.ush
--- a/Engine/Shaders/Private/ScreenSpaceDenoise/SSDSignalBufferEncoding.ush
+++ b/Engine/Shaders/Private/ScreenSpaceDenoise/SSDSignalBufferEncoding.ush
@@ -160,6 +160,48 @@ void DecodeMultiplexedSignals(
 		}
 	}
 	#if COMPILE_SIGNAL_COLOR
+	else if (SignalBufferLayout == SIGNAL_BUFFER_LAYOUT_POLYCHROMATIC_PENUMBRA_INPUT)
+	{
+		// Decode diffuse harmonic
+		{
+			float Weight = RawSample[0].a < 0.0 ? 0 : 1;
+
+			OutSamples.Array[0].SampleCount = Weight;
+			OutSamples.Array[0].MissCount = RawSample[0].a * Weight;
+			OutSamples.Array[0].SceneColor.rgb = RawSample[0].rgb * Weight;
+		}
+		
+		// Decode specular harmonic
+		{
+			float Weight = RawSample[1].a < 0.0 ? 0 : 1;
+
+			OutSamples.Array[1].SampleCount = Weight;
+			OutSamples.Array[1].MissCount = RawSample[1].a * Weight;
+			OutSamples.Array[1].SceneColor.rgb = RawSample[1].rgb * Weight;
+		}
+	}
+	else if (SignalBufferLayout == SIGNAL_BUFFER_LAYOUT_POLYCHROMATIC_PENUMBRA_HISTORY)
+	{
+		// Decode diffuse harmonic
+		{
+			float Weight = RawSample[2].r;
+
+			OutSamples.Array[0].SampleCount = Weight;
+			OutSamples.Array[0].MissCount = RawSample[0].a * Weight;
+			OutSamples.Array[0].SceneColor.rgb = RawSample[0].rgb * Weight;
+		}
+		
+		// Decode specular harmonic
+		{
+			float Weight = RawSample[2].g;
+
+			OutSamples.Array[1].SampleCount = Weight;
+			OutSamples.Array[1].MissCount = RawSample[1].a * Weight;
+			OutSamples.Array[1].SceneColor.rgb = RawSample[1].rgb * Weight;
+		}
+	}
+	#endif
+	#if COMPILE_SIGNAL_COLOR
 	else if (SignalBufferLayout == SIGNAL_BUFFER_LAYOUT_REFLECTIONS_INPUT)
 	{
 		OutSamples.Array[0].SampleCount = RawSample[1].r == DENOISER_INVALID_HIT_DISTANCE ? 0 : 1;
@@ -437,6 +479,30 @@ void EncodeMultiplexedSignals(
 		}
 	}
 	#if COMPILE_SIGNAL_COLOR
+	else if (SignalBufferLayout == SIGNAL_BUFFER_LAYOUT_POLYCHROMATIC_PENUMBRA_HISTORY)
+	{
+		// Encode diffuse harmonic
+		{
+			float NormalizationFactor = SafeRcp(Sample[0].SampleCount);
+
+			OutRawSample[0] = ClampColorForEncoding(Sample[0].SceneColor * NormalizationFactor);
+			OutRawSample[0].a = Sample[0].MissCount * NormalizationFactor;
+			OutRawSample[2].r = Sample[0].SampleCount;
+		}
+
+		// Encode specular harmonic
+		{
+			float NormalizationFactor = SafeRcp(Sample[1].SampleCount);
+
+			OutRawSample[1] = ClampColorForEncoding(Sample[1].SceneColor * NormalizationFactor);
+			OutRawSample[1].a = Sample[1].MissCount * NormalizationFactor;
+			OutRawSample[2].g = Sample[1].SampleCount;
+		}
+
+		OutBufferCount = 3;
+	}
+	#endif
+	#if COMPILE_SIGNAL_COLOR
 	else if (SignalBufferLayout == SIGNAL_BUFFER_LAYOUT_REFLECTIONS_REJECTION)
 	{
 		// Outputs scene color to be compatible with a SSR output.
diff --git a/Engine/Shaders/Private/ScreenSpaceDenoise/SSDSignalFramework.ush b/Engine/Shaders/Private/ScreenSpaceDenoise/SSDSignalFramework.ush
--- a/Engine/Shaders/Private/ScreenSpaceDenoise/SSDSignalFramework.ush
+++ b/Engine/Shaders/Private/ScreenSpaceDenoise/SSDSignalFramework.ush
@@ -53,6 +53,9 @@
 	#error Missing CONFIG_SIGNAL_PROCESSING
 
 #elif CONFIG_SIGNAL_PROCESSING == SIGNAL_PROCESSING_MONOCHROMATIC_PENUMBRA
+	#include "SSDPublicHarmonics.ush"
+
+#elif CONFIG_SIGNAL_PROCESSING == SIGNAL_PROCESSING_POLYCHROMATIC_PENUMBRA_HARMONIC
 	
 #elif CONFIG_SIGNAL_PROCESSING == SIGNAL_PROCESSING_REFLECTIONS
 	#include "../MonteCarlo.ush"
@@ -634,33 +637,19 @@ float GetSignalWorldBluringRadius(FSSDSignalSample Sample, FSSDSampleSceneInfos
 	{
 		return DomainKnowledge.HitDistanceToWorldBluringRadius * Sample.ClosestHitDistance;
 	}
-	else if (DomainKnowledge.LightType == LIGHT_TYPE_POINT || DomainKnowledge.LightType == LIGHT_TYPE_SPOT)
-	{
-		float3 PixelToLightWorldVector = (DomainKnowledge.Light.Position + View.PreViewTranslation) - SceneMetadata.TranslatedWorldPosition;
-		float DistanceFromLight = length(PixelToLightWorldVector);
-
-		return DomainKnowledge.Light.SourceRadius * Sample.ClosestHitDistance / max(DistanceFromLight - Sample.ClosestHitDistance, ClosestLightDistance);
-	}
-	else // if (DomainKnowledge.LightType == LIGHT_TYPE_RECT)
+	else
 	{
-		float3 PixelToLightWorldVector = (DomainKnowledge.Light.Position + View.PreViewTranslation) - SceneMetadata.TranslatedWorldPosition;
-		float DistanceFromLight = length(PixelToLightWorldVector);
-
-		#if CONFIG_VALU_OPTIMIZATIONS
-			float Scalar = abs(dot(PixelToLightWorldVector, DomainKnowledge.Light.Direction)) * rsqrt(length2(PixelToLightWorldVector));
-		#else
-			float Scalar = abs(dot(normalize(PixelToLightWorldVector), DomainKnowledge.Light.Direction));
-		#endif
-
-		float2 LightDimensions = GetRectLightDimensions(DomainKnowledge.Light);
-
-		// TODO: area light anysotropy.
-		float SmallestLightDimension = min(LightDimensions.x, LightDimensions.y);
-
-		// TODO: Sample.ClosestHitDistance depends on the direction of the ray, but does not DistanceFromLight, witch is bad for large area light.
-		return SmallestLightDimension * Sample.ClosestHitDistance / max(DistanceFromLight - Sample.ClosestHitDistance, ClosestLightDistance);
+		return ComputeLightSampleWorldBluringRadius(
+			SceneMetadata.TranslatedWorldPosition,
+			DomainKnowledge.LightType, DomainKnowledge.Light,
+			Sample.ClosestHitDistance);
 	}
 }
+#elif CONFIG_SIGNAL_PROCESSING == SIGNAL_PROCESSING_POLYCHROMATIC_PENUMBRA_HARMONIC
+{
+	// No need to compute the world bluring radius given the hit has been classified by harmonic frequencies already.
+	return 0;
+}
 #elif CONFIG_SIGNAL_PROCESSING == SIGNAL_PROCESSING_REFLECTIONS
 {
 	return 1;
diff --git a/Engine/Shaders/Private/ScreenSpaceDenoise/SSDSpatialAccumulation.usf b/Engine/Shaders/Private/ScreenSpaceDenoise/SSDSpatialAccumulation.usf
--- a/Engine/Shaders/Private/ScreenSpaceDenoise/SSDSpatialAccumulation.usf
+++ b/Engine/Shaders/Private/ScreenSpaceDenoise/SSDSpatialAccumulation.usf
@@ -40,8 +40,6 @@
 
 // Configures all the pass for each individual signals.
 #if CONFIG_SIGNAL_PROCESSING == SIGNAL_PROCESSING_MONOCHROMATIC_PENUMBRA
-	#define COMPILE_SIGNAL_COLOR 0
-
 	//#define CONFIG_VIEW_POSITION_BASED_BILATERAL 2
 	#define CONFIG_VIEW_POSITION_BASED_BILATERAL 5
 	#define CONFIG_MULTIPLEXED_SIGNALS_PER_SIGNAL_DOMAIN 1
@@ -121,6 +119,34 @@
 		// #define CONFIG_ACCUMULATOR_VGPR_COMPRESSION ACCUMULATOR_COMPRESSION_PENUMBRA_DRB
 	#endif
 
+#elif CONFIG_SIGNAL_PROCESSING == SIGNAL_PROCESSING_POLYCHROMATIC_PENUMBRA_HARMONIC
+	#define COMPILE_SIGNAL_COLOR 1
+
+	//#define CONFIG_VIEW_POSITION_BASED_BILATERAL 2
+	#define CONFIG_VIEW_POSITION_BASED_BILATERAL 5
+	//#define CONFIG_MULTIPLEXED_SIGNALS_PER_SIGNAL_DOMAIN 1
+
+	// Denoise diffuse and specular harmonics at the same time.
+	#define MAX_SIGNAL_BATCH_SIZE 2
+
+	// Any world distance depends on the harmonic being processed.
+	#define CONFIG_BILATERAL_DISTANCE_COMPUTATION SIGNAL_WORLD_FREQUENCY_HARMONIC
+
+	#if DIM_STAGE == STAGE_RECONSTRUCTION
+		// Input and output layout.
+		#define CONFIG_SIGNAL_INPUT_LAYOUT  SIGNAL_BUFFER_LAYOUT_POLYCHROMATIC_PENUMBRA_INPUT
+		#define CONFIG_SIGNAL_OUTPUT_LAYOUT SIGNAL_BUFFER_LAYOUT_POLYCHROMATIC_PENUMBRA_HISTORY
+
+		#define CONFIG_INPUT_TEXTURE_COUNT 1
+		#define CONFIG_OUTPUT_TEXTURE_COUNT 3
+
+		#define CONFIG_SAMPLE_SET           SAMPLE_SET_STACKOWIAK_4_SETS
+		//#define CONFIG_MAX_WITH_REF_DISTANCE 1
+
+	#else
+		#error Unexpected stage.
+	#endif
+
 #elif CONFIG_SIGNAL_PROCESSING == SIGNAL_PROCESSING_REFLECTIONS
 	#define COMPILE_SIGNAL_COLOR 1
 
@@ -448,6 +474,8 @@ uint UpscaleFactor;
 	float KernelSpreadFactor;
 #endif
 
+float HarmonicPeriode;
+
 
 #if !defined(CONFIG_INPUT_TEXTURE_COUNT)
 	#error Missing CONFIG_INPUT_TEXTURE_COUNT
@@ -645,6 +673,7 @@ void MainCS(
 		KernelConfig.BufferSizeAndInvSize = View.BufferSizeAndInvSize;
 		KernelConfig.BufferBilinearUVMinMax = View.BufferBilinearUVMinMax;
 		KernelConfig.KernelSpreadFactor = KernelSpreadFactor;
+		KernelConfig.HarmonicPeriode = HarmonicPeriode;
 
 		// VGPRs
 		KernelConfig.BufferUV = SceneBufferUV;
diff --git a/Engine/Shaders/Private/ScreenSpaceDenoise/SSDSpatialKernel.ush b/Engine/Shaders/Private/ScreenSpaceDenoise/SSDSpatialKernel.ush
--- a/Engine/Shaders/Private/ScreenSpaceDenoise/SSDSpatialKernel.ush
+++ b/Engine/Shaders/Private/ScreenSpaceDenoise/SSDSpatialKernel.ush
@@ -23,6 +23,9 @@
 	// Uses FSSDSignalSample::WorldBluringRadius precomputed in the sample.
 	#define SIGNAL_WORLD_FREQUENCY_PRECOMPUTED_BLURING_RADIUS 4
 
+	// Compute based on the harmonic being processed.
+	#define SIGNAL_WORLD_FREQUENCY_HARMONIC 5
+
 
 //------------------------------------------------------- CONFIG DISABLED DEFAULTS
 
@@ -128,6 +131,9 @@ struct FSSDKernelConfig
 	// Multiplier on the sample's offset.
 	float KernelSpreadFactor;
 
+	// The periode of the harmonic being sampled.
+	float HarmonicPeriode;
+
 	
 	// --------------------------- Per lane.
 	
@@ -188,6 +194,7 @@ FSSDKernelConfig CreateKernelConfig()
 	KernelConfig.BufferBilinearUVMinMax = float4(0, 0, 0, 0);
 	KernelConfig.BufferMipLevel = 0.0;
 	KernelConfig.KernelSpreadFactor = 1;
+	KernelConfig.HarmonicPeriode = 1.0;
 	
 	// VGPRs.
 	KernelConfig.BoxKernelRadius = 1;
@@ -448,6 +455,11 @@ void AccumulateSampledMultiplexedSignals(
 		{
 			FinalWorldBluringDistance = SignalConvolutionBluringRadius;
 		}
+		else if (KernelConfig.BilateralDistanceComputation == SIGNAL_WORLD_FREQUENCY_HARMONIC)
+		{
+			FinalWorldBluringDistance = AmendWorldBluringRadiusCausedByPixelSize(
+				RefPixelWorldBluringRadius) * KernelConfig.HarmonicPeriode;
+		}
 		
 		if (KernelConfig.bMaxWithRefBilateralDistance)
 		{
diff --git a/Engine/Shaders/Private/ScreenSpaceDenoise/SSDTemporalAccumulation.usf b/Engine/Shaders/Private/ScreenSpaceDenoise/SSDTemporalAccumulation.usf
--- a/Engine/Shaders/Private/ScreenSpaceDenoise/SSDTemporalAccumulation.usf
+++ b/Engine/Shaders/Private/ScreenSpaceDenoise/SSDTemporalAccumulation.usf
@@ -62,6 +62,23 @@
 
 	#define MAX_SIGNAL_BATCH_SIZE CONFIG_SIGNAL_BATCH_SIZE
 
+#elif CONFIG_SIGNAL_PROCESSING == SIGNAL_PROCESSING_POLYCHROMATIC_PENUMBRA_HARMONIC
+	#define COMPILE_SIGNAL_COLOR 1
+
+	#define CONFIG_VIEW_POSITION_BASED_BILATERAL 2
+	#define CONFIG_MULTIPLEXED_SIGNALS_PER_SIGNAL_DOMAIN 1
+
+	// Input and output layout.
+	#define CONFIG_SIGNAL_INPUT_LAYOUT  SIGNAL_BUFFER_LAYOUT_POLYCHROMATIC_PENUMBRA_HISTORY
+	#define CONFIG_SIGNAL_HISTORY_LAYOUT SIGNAL_BUFFER_LAYOUT_POLYCHROMATIC_PENUMBRA_HISTORY
+	#define CONFIG_INPUT_TEXTURE_COUNT 3
+	#define CONFIG_HISTORY_TEXTURE_COUNT 3
+
+	#define CONFIG_HISTORY_REJECTION HISTORY_REJECTION_VAR_BOUNDARIES
+
+	// Denoise diffuse and specular harmonics at the same time.
+	#define MAX_SIGNAL_BATCH_SIZE 2
+
 #elif CONFIG_SIGNAL_PROCESSING == SIGNAL_PROCESSING_REFLECTIONS
 	#define COMPILE_SIGNAL_COLOR 1
 
@@ -744,7 +761,7 @@ void TemporallyAccumulate(
 						RejectionFactor = max(RejectionFactor, 1 * max(RejectedDiff.SceneColor.y, RejectedDiff.SceneColor.z));
 						RejectionFactor = max(RejectionFactor, 1 * RejectedDiff.SceneColor.a);
 					}
-					#elif CONFIG_SIGNAL_PROCESSING == SIGNAL_PROCESSING_MONOCHROMATIC_PENUMBRA || CONFIG_SIGNAL_PROCESSING == SIGNAL_PROCESSING_AO || CONFIG_SIGNAL_PROCESSING == SIGNAL_PROCESSING_DIFFUSE_INDIRECT_AND_AO
+					#elif CONFIG_SIGNAL_PROCESSING == SIGNAL_PROCESSING_MONOCHROMATIC_PENUMBRA || CONFIG_SIGNAL_PROCESSING == SIGNAL_PROCESSING_POLYCHROMATIC_PENUMBRA_HARMONIC || CONFIG_SIGNAL_PROCESSING == SIGNAL_PROCESSING_AO || CONFIG_SIGNAL_PROCESSING == SIGNAL_PROCESSING_DIFFUSE_INDIRECT_AND_AO
 					{
 						RejectionFactor = abs(ClampedTransformedHistorySample.MissCount - TransformedHistorySample.MissCount);
 					}
diff --git a/Engine/Source/Runtime/Renderer/Private/LightRendering.cpp b/Engine/Source/Runtime/Renderer/Private/LightRendering.cpp
--- a/Engine/Source/Runtime/Renderer/Private/LightRendering.cpp
+++ b/Engine/Source/Runtime/Renderer/Private/LightRendering.cpp
@@ -1047,7 +1047,7 @@ void FDeferredShadingSceneRenderer::RenderLights(FRHICommandListImmediate& RHICm
 									InputParameterCount,
 									View.ViewRect.Width(), View.ViewRect.Height());
 
-								DenoiserToUse->DenoiseShadows(
+								DenoiserToUse->DenoiseMonochromaticShadows(
 									GraphBuilder,
 									View,
 									&View.PrevViewInfo,
@@ -1226,7 +1226,7 @@ void FDeferredShadingSceneRenderer::RenderLights(FRHICommandListImmediate& RHICm
 									InputParameterCount,
 									View.ViewRect.Width(), View.ViewRect.Height());
 
-								DenoiserToUse->DenoiseShadows(
+								DenoiserToUse->DenoiseMonochromaticShadows(
 									GraphBuilder,
 									View,
 									&View.PrevViewInfo,
diff --git a/Engine/Source/Runtime/Renderer/Private/SceneRendering.h b/Engine/Source/Runtime/Renderer/Private/SceneRendering.h
--- a/Engine/Source/Runtime/Renderer/Private/SceneRendering.h
+++ b/Engine/Source/Runtime/Renderer/Private/SceneRendering.h
@@ -783,6 +783,9 @@ struct FPreviousViewInfo
 
 	// History for shadow denoising.
 	TMap<const ULightComponent*, FScreenSpaceFilteringHistory> ShadowHistories;
+
+	// History for denoising all lights penumbra at once.
+	FScreenSpaceFilteringHistory PolychromaticPenumbraHarmonicsHistory;
 };
 
 class FViewCommands
diff --git a/Engine/Source/Runtime/Renderer/Private/ScreenSpaceDenoise.cpp b/Engine/Source/Runtime/Renderer/Private/ScreenSpaceDenoise.cpp
--- a/Engine/Source/Runtime/Renderer/Private/ScreenSpaceDenoise.cpp
+++ b/Engine/Source/Runtime/Renderer/Private/ScreenSpaceDenoise.cpp
@@ -139,6 +139,9 @@ enum class ESignalProcessing
 	// Denoise a shadow mask.
 	MonochromaticPenumbra,
 
+	// Denoise one lighting harmonic when denoising multiple light's penumbra.
+	PolychromaticPenumbraHarmonic,
+
 	// Denoise first bounce specular.
 	Reflections,
 
@@ -167,6 +170,7 @@ static bool UsesConstantPixelDensityPassLayout(ESignalProcessing SignalProcessin
 {
 	return (
 		SignalProcessing == ESignalProcessing::MonochromaticPenumbra ||
+		SignalProcessing == ESignalProcessing::PolychromaticPenumbraHarmonic ||
 		SignalProcessing == ESignalProcessing::Reflections ||
 		SignalProcessing == ESignalProcessing::AmbientOcclusion ||
 		SignalProcessing == ESignalProcessing::DiffuseAndAmbientOcclusion ||
@@ -236,6 +240,7 @@ static int32 SignalMaxBatchSize(ESignalProcessing SignalProcessing)
 	}
 	else if (
 		SignalProcessing == ESignalProcessing::Reflections ||
+		SignalProcessing == ESignalProcessing::PolychromaticPenumbraHarmonic ||
 		SignalProcessing == ESignalProcessing::AmbientOcclusion ||
 		SignalProcessing == ESignalProcessing::DiffuseAndAmbientOcclusion ||
 		SignalProcessing == ESignalProcessing::DiffuseSphericalHarmonic)
@@ -260,6 +265,7 @@ static bool SignalSupportMultiSPP(ESignalProcessing SignalProcessing)
 {
 	return (
 		SignalProcessing == ESignalProcessing::MonochromaticPenumbra ||
+		SignalProcessing == ESignalProcessing::PolychromaticPenumbraHarmonic ||
 		SignalProcessing == ESignalProcessing::Reflections ||
 		SignalProcessing == ESignalProcessing::AmbientOcclusion ||
 		SignalProcessing == ESignalProcessing::DiffuseAndAmbientOcclusion ||
@@ -286,6 +292,12 @@ const TCHAR* const kInjestResourceNames[] = {
 	nullptr,
 	nullptr,
 
+	// PolychromaticPenumbraHarmonic
+	nullptr,
+	nullptr,
+	nullptr,
+	nullptr,
+
 	// Reflections
 	nullptr,
 	nullptr,
@@ -318,6 +330,12 @@ const TCHAR* const kReduceResourceNames[] = {
 	nullptr,
 	nullptr,
 
+	// PolychromaticPenumbraHarmonic
+	nullptr,
+	nullptr,
+	nullptr,
+	nullptr,
+
 	// Reflections
 	nullptr,
 	nullptr,
@@ -350,6 +368,12 @@ const TCHAR* const kReconstructionResourceNames[] = {
 	TEXT("ShadowReconstruction2"),
 	TEXT("ShadowReconstruction3"),
 
+	// PolychromaticPenumbraHarmonic
+	TEXT("PolychromaticPenumbraHarmonicReconstruction0"),
+	TEXT("PolychromaticPenumbraHarmonicReconstruction1"),
+	TEXT("PolychromaticPenumbraHarmonicReconstruction2"),
+	nullptr,
+
 	// Reflections
 	TEXT("ReflectionsReconstruction0"),
 	TEXT("ReflectionsReconstruction1"),
@@ -382,6 +406,12 @@ const TCHAR* const kPreConvolutionResourceNames[] = {
 	TEXT("ShadowPreConvolution2"),
 	TEXT("ShadowPreConvolution3"),
 
+	// PolychromaticPenumbraHarmonic
+	nullptr,
+	nullptr,
+	nullptr,
+	nullptr,
+
 	// Reflections
 	nullptr,
 	nullptr,
@@ -414,6 +444,12 @@ const TCHAR* const kRejectionPreConvolutionResourceNames[] = {
 	TEXT("ShadowRejectionPreConvolution2"),
 	TEXT("ShadowRejectionPreConvolution3"),
 
+	// PolychromaticPenumbraHarmonic
+	nullptr,
+	nullptr,
+	nullptr,
+	nullptr,
+
 	// Reflections
 	TEXT("ReflectionsRejectionPreConvolution0"),
 	TEXT("ReflectionsRejectionPreConvolution1"),
@@ -446,6 +482,12 @@ const TCHAR* const kTemporalAccumulationResourceNames[] = {
 	TEXT("ShadowTemporalAccumulation2"),
 	TEXT("ShadowTemporalAccumulation3"),
 
+	// PolychromaticPenumbraHarmonic
+	TEXT("PolychromaticPenumbraHistory0"),
+	TEXT("PolychromaticPenumbraHistory1"),
+	TEXT("PolychromaticPenumbraHistory2"),
+	nullptr,
+
 	// Reflections
 	TEXT("ReflectionsTemporalAccumulation0"),
 	TEXT("ReflectionsTemporalAccumulation1"),
@@ -478,6 +520,12 @@ const TCHAR* const kHistoryConvolutionResourceNames[] = {
 	TEXT("ShadowHistoryConvolution2"),
 	TEXT("ShadowHistoryConvolution3"),
 
+	// PolychromaticPenumbraHarmonic
+	nullptr,
+	nullptr,
+	nullptr,
+	nullptr,
+
 	// Reflections
 	TEXT("ReflectionsHistoryConvolution0"),
 	TEXT("ReflectionsHistoryConvolution1"),
@@ -510,6 +558,12 @@ const TCHAR* const kDenoiserOutputResourceNames[] = {
 	TEXT("ShadowDenoiserOutput2"),
 	TEXT("ShadowDenoiserOutput3"),
 
+	// PolychromaticPenumbraHarmonic
+	nullptr,
+	nullptr,
+	nullptr,
+	nullptr,
+
 	// Reflections
 	nullptr,
 	nullptr,
@@ -553,6 +607,7 @@ bool ShouldCompileSignalPipeline(ESignalProcessing SignalProcessing, EShaderPlat
 	}
 	else if (
 		SignalProcessing == ESignalProcessing::MonochromaticPenumbra ||
+		SignalProcessing == ESignalProcessing::PolychromaticPenumbraHarmonic ||
 		SignalProcessing == ESignalProcessing::AmbientOcclusion ||
 		SignalProcessing == ESignalProcessing::DiffuseAndAmbientOcclusion)
 	{
@@ -858,6 +913,7 @@ class FSSDSpatialAccumulationCS : public FGlobalShader
 		SHADER_PARAMETER(uint32, MaxSampleCount)
 		SHADER_PARAMETER(int32, UpscaleFactor)
 		SHADER_PARAMETER(float, KernelSpreadFactor)
+		SHADER_PARAMETER(float, HarmonicPeriode)
 		
 		SHADER_PARAMETER_STRUCT_INCLUDE(FSSDCommonParameters, CommonParameters)
 		SHADER_PARAMETER_STRUCT_INCLUDE(FSSDConvolutionMetaData, ConvolutionMetaData)
@@ -917,10 +973,28 @@ class FSSDTemporalAccumulationCS : public FGlobalShader
 	END_SHADER_PARAMETER_STRUCT()
 };
 
+class FSSDComposeHarmonicsCS : public FGlobalShader
+{
+	DECLARE_GLOBAL_SHADER(FSSDComposeHarmonicsCS);
+	SHADER_USE_PARAMETER_STRUCT(FSSDComposeHarmonicsCS, FGlobalShader);
+
+	static bool ShouldCompilePermutation(const FGlobalShaderPermutationParameters& Parameters)
+	{
+		return ShouldCompileSignalPipeline(ESignalProcessing::PolychromaticPenumbraHarmonic, Parameters.Platform);
+	}
+
+	BEGIN_SHADER_PARAMETER_STRUCT(FParameters, )
+		SHADER_PARAMETER_STRUCT_ARRAY(FSSDSignalTextures, SignalHarmonics, [IScreenSpaceDenoiser::kMultiPolychromaticPenumbraHarmonics])
+		SHADER_PARAMETER_STRUCT_REF(FViewUniformShaderParameters, ViewUniformBuffer)
+		SHADER_PARAMETER_STRUCT(FSSDSignalUAVs, SignalOutput)
+	END_SHADER_PARAMETER_STRUCT()
+};
+
 IMPLEMENT_GLOBAL_SHADER(FSSDInjestCS, "/Engine/Private/ScreenSpaceDenoise/SSDInjest.usf", "MainCS", SF_Compute);
 IMPLEMENT_GLOBAL_SHADER(FSSDReduceCS, "/Engine/Private/ScreenSpaceDenoise/SSDReduce.usf", "MainCS", SF_Compute);
 IMPLEMENT_GLOBAL_SHADER(FSSDSpatialAccumulationCS, "/Engine/Private/ScreenSpaceDenoise/SSDSpatialAccumulation.usf", "MainCS", SF_Compute);
 IMPLEMENT_GLOBAL_SHADER(FSSDTemporalAccumulationCS, "/Engine/Private/ScreenSpaceDenoise/SSDTemporalAccumulation.usf", "MainCS", SF_Compute);
+IMPLEMENT_GLOBAL_SHADER(FSSDComposeHarmonicsCS, "/Engine/Private/ScreenSpaceDenoise/SSDComposeHarmonics.usf", "MainCS", SF_Compute);
 
 } // namespace
 
@@ -930,8 +1004,10 @@ struct FSSDConstantPixelDensitySettings
 {
 	ESignalProcessing SignalProcessing;
 	int32 SignalBatchSize = 1;
+	float HarmonicPeriode = 1.0f;
 	int32 MaxInputSPP = 1;
 	float InputResolutionFraction = 1.0f;
+	bool bEnableReconstruction = true;
 	int32 ReconstructionSamples = 1;
 	int32 PreConvolutionCount = 0;
 	bool bUseTemporalAccumulation = false;
@@ -1039,6 +1115,16 @@ static void DenoiseSignalAtConstantPixelDensity(
 			ReconstructionTextureCount = Settings.SignalBatchSize;
 			bHasReconstructionLayoutDifferentFromHistory = true;
 		}
+		else if (Settings.SignalProcessing == ESignalProcessing::PolychromaticPenumbraHarmonic)
+		{
+			ReconstructionTextureCount = 3;
+			ReconstructionDescs[0].Format = PF_FloatRGBA;
+			ReconstructionDescs[1].Format = PF_FloatRGBA;
+			ReconstructionDescs[2].Format = PF_FloatRGBA;
+
+			HistoryTextureCountPerSignal = ReconstructionTextureCount;
+			HistoryDescs = ReconstructionDescs;
+		}
 		else if (Settings.SignalProcessing == ESignalProcessing::Reflections)
 		{
 			ReconstructionDescs[0].Format = HistoryDescs[0].Format = PF_FloatRGBA;
@@ -1228,7 +1314,8 @@ static void DenoiseSignalAtConstantPixelDensity(
 		SignalHistory = NewSignalOutput;
 	}
 
-	// Spatial reconstruction with multiple important sampling to be more precise in the history rejection.
+	// Spatial reconstruction with ratio estimator to be more precise in the history rejection.
+	if (Settings.bEnableReconstruction)
 	{
 		FSSDSignalTextures NewSignalOutput = CreateMultiplexedTextures(
 			GraphBuilder,
@@ -1238,6 +1325,7 @@ static void DenoiseSignalAtConstantPixelDensity(
 		FSSDSpatialAccumulationCS::FParameters* PassParameters = GraphBuilder.AllocParameters<FSSDSpatialAccumulationCS::FParameters>();
 		PassParameters->MaxSampleCount = FMath::Clamp(Settings.ReconstructionSamples, 1, kStackowiakMaxSampleCountPerSet);
 		PassParameters->UpscaleFactor = int32(1.0f / Settings.InputResolutionFraction);
+		PassParameters->HarmonicPeriode = Settings.HarmonicPeriode;
 		PassParameters->CommonParameters = CommonParameters;
 		PassParameters->ConvolutionMetaData = ConvolutionMetaData;
 		PassParameters->SignalInput = SignalHistory;
@@ -1474,7 +1562,7 @@ static void DenoiseSignalAtConstantPixelDensity(
 		SignalHistory = SignalOutput;
 	} // if (MaxPostFilterSampleCount > 1)
 
-	if (!View.bViewStateIsReadOnly)
+	if (!View.bViewStateIsReadOnly && Settings.bUseTemporalAccumulation)
 	{
 		check(View.ViewState);
 
@@ -1483,15 +1571,18 @@ static void DenoiseSignalAtConstantPixelDensity(
 			GraphBuilder.QueueTextureExtraction(SceneTextures.SceneDepthBuffer, &View.ViewState->PrevFrameViewInfo.DepthBuffer);
 
 			// Requires the normal that are in GBuffer A.
-			if (Settings.SignalProcessing == ESignalProcessing::Reflections ||
+			if (Settings.SignalProcessing == ESignalProcessing::MonochromaticPenumbra ||
+				Settings.SignalProcessing == ESignalProcessing::PolychromaticPenumbraHarmonic ||
+				Settings.SignalProcessing == ESignalProcessing::Reflections ||
 				Settings.SignalProcessing == ESignalProcessing::AmbientOcclusion ||
 				Settings.SignalProcessing == ESignalProcessing::DiffuseAndAmbientOcclusion)
 			{
 				GraphBuilder.QueueTextureExtraction(SceneTextures.SceneGBufferA, &View.ViewState->PrevFrameViewInfo.GBufferA);
 			}
 
-			// Reflections requires the roughness that is in GBuffer B.
-			if (Settings.SignalProcessing == ESignalProcessing::Reflections)
+			// Requires the roughness that is in GBuffer B.
+			if (Settings.SignalProcessing == ESignalProcessing::PolychromaticPenumbraHarmonic || 
+				Settings.SignalProcessing == ESignalProcessing::Reflections)
 			{
 				GraphBuilder.QueueTextureExtraction(SceneTextures.SceneGBufferB, &View.ViewState->PrevFrameViewInfo.GBufferB);
 			}
@@ -1594,7 +1685,7 @@ public:
 		return IScreenSpaceDenoiser::EShadowRequirements::ClosestOccluder;
 	}
 
-	virtual void DenoiseShadows(
+	void DenoiseMonochromaticShadows(
 		FRDGBuilder& GraphBuilder,
 		const FViewInfo& View,
 		FPreviousViewInfo* PreviousViewInfos,
@@ -1665,6 +1756,108 @@ public:
 		}
 	}
 
+	FPolychromaticPenumbraOutputs DenoisePolychromaticPenumbraHarmonics(
+		FRDGBuilder& GraphBuilder,
+		const FViewInfo& View,
+		FPreviousViewInfo* PreviousViewInfos,
+		const FSceneTextureParameters& SceneTextures,
+		const FPolychromaticPenumbraHarmonics& Inputs) const override
+	{
+		RDG_GPU_STAT_SCOPE(GraphBuilder, ShadowsDenoiser);
+
+		FSSDComposeHarmonicsCS::FParameters* ComposePassParameters = GraphBuilder.AllocParameters<FSSDComposeHarmonicsCS::FParameters>();
+
+		// Harmonic 0 doesn't need any reconstruction given it's the highest frequency details.
+		{
+			const int32 HarmonicId = 0;
+			ComposePassParameters->SignalHarmonics[HarmonicId].Textures[0] = Inputs.Diffuse[HarmonicId];
+			ComposePassParameters->SignalHarmonics[HarmonicId].Textures[1] = Inputs.Specular[HarmonicId];
+		}
+
+		// Reconstruct each harmonic independently
+		for (int32 HarmonicId = 1; HarmonicId < IScreenSpaceDenoiser::kMultiPolychromaticPenumbraHarmonics; HarmonicId++)
+		{
+			int32 Periode = 1 << HarmonicId;
+
+			FSSDConstantPixelDensitySettings Settings;
+			Settings.SignalProcessing = ESignalProcessing::PolychromaticPenumbraHarmonic;
+			Settings.HarmonicPeriode = Periode;
+			Settings.ReconstructionSamples = Periode * Periode; // TODO: should use preconvolution instead for harmonic 3
+			Settings.bUseTemporalAccumulation = false;
+
+			TStaticArray<FScreenSpaceFilteringHistory*, IScreenSpaceDenoiser::kMaxBatchSize> PrevHistories;
+			TStaticArray<FScreenSpaceFilteringHistory*, IScreenSpaceDenoiser::kMaxBatchSize> NewHistories;
+			PrevHistories[0] = nullptr;
+			NewHistories[0] = nullptr;
+
+			FSSDSignalTextures InputSignal;
+			InputSignal.Textures[0] = Inputs.Diffuse[HarmonicId];
+			InputSignal.Textures[1] = Inputs.Specular[HarmonicId];
+
+			FSSDSignalTextures SignalOutput;
+			DenoiseSignalAtConstantPixelDensity(
+				GraphBuilder, View, SceneTextures,
+				InputSignal, Settings,
+				PrevHistories, NewHistories,
+				/* out */ &SignalOutput);
+
+			ComposePassParameters->SignalHarmonics[HarmonicId] = SignalOutput;
+		}
+
+		// Merges the different harmonics.
+		FSSDSignalTextures ComposedHarmonics;
+		{
+			{
+				FRDGTextureDesc Desc = FRDGTextureDesc::Create2DDesc(
+					SceneTextures.SceneDepthBuffer->Desc.Extent,
+					PF_FloatRGBA,
+					FClearValueBinding::Black,
+					/* InFlags = */ TexCreate_None,
+					/* InTargetableFlags = */ TexCreate_ShaderResource | TexCreate_UAV,
+					/* bInForceSeparateTargetAndShaderResource = */ false);
+
+				ComposedHarmonics.Textures[0] = GraphBuilder.CreateTexture(Desc, TEXT("MultiPolychromaticComposition0"));
+				ComposedHarmonics.Textures[1] = GraphBuilder.CreateTexture(Desc, TEXT("MultiPolychromaticComposition1"));
+				ComposedHarmonics.Textures[2] = GraphBuilder.CreateTexture(Desc, TEXT("MultiPolychromaticComposition2"));
+			}
+
+			ComposePassParameters->ViewUniformBuffer = View.ViewUniformBuffer;
+			ComposePassParameters->SignalOutput = CreateMultiplexedUAVs(GraphBuilder, ComposedHarmonics);
+
+			TShaderMapRef<FSSDComposeHarmonicsCS> ComputeShader(View.ShaderMap);
+			FComputeShaderUtils::AddPass(
+				GraphBuilder,
+				RDG_EVENT_NAME("SSD ComposeHarmonics"),
+				*ComputeShader, ComposePassParameters,
+				FComputeShaderUtils::GetGroupCount(View.ViewRect.Size(), FSSDSpatialAccumulationCS::kGroupSize));
+		}
+
+		FPolychromaticPenumbraOutputs Outputs;
+		{
+			FSSDConstantPixelDensitySettings Settings;
+			Settings.SignalProcessing = ESignalProcessing::PolychromaticPenumbraHarmonic;
+			Settings.bEnableReconstruction = false;
+			Settings.bUseTemporalAccumulation = CVarShadowTemporalAccumulation.GetValueOnRenderThread() != 0;
+
+			TStaticArray<FScreenSpaceFilteringHistory*, IScreenSpaceDenoiser::kMaxBatchSize> PrevHistories;
+			TStaticArray<FScreenSpaceFilteringHistory*, IScreenSpaceDenoiser::kMaxBatchSize> NewHistories;
+			PrevHistories[0] = &PreviousViewInfos->PolychromaticPenumbraHarmonicsHistory;
+			NewHistories[0] = View.ViewState ? &View.ViewState->PrevFrameViewInfo.PolychromaticPenumbraHarmonicsHistory : nullptr;
+
+			FSSDSignalTextures SignalOutput;
+			DenoiseSignalAtConstantPixelDensity(
+				GraphBuilder, View, SceneTextures,
+				ComposedHarmonics, Settings,
+				PrevHistories, NewHistories,
+				/* out */ &SignalOutput);
+
+			Outputs.Diffuse = SignalOutput.Textures[0];
+			Outputs.Specular = SignalOutput.Textures[1];
+		}
+
+		return Outputs;
+	}
+
 	FReflectionsOutputs DenoiseReflections(
 		FRDGBuilder& GraphBuilder,
 		const FViewInfo& View,
diff --git a/Engine/Source/Runtime/Renderer/Private/ScreenSpaceDenoise.h b/Engine/Source/Runtime/Renderer/Private/ScreenSpaceDenoise.h
--- a/Engine/Source/Runtime/Renderer/Private/ScreenSpaceDenoise.h
+++ b/Engine/Source/Runtime/Renderer/Private/ScreenSpaceDenoise.h
@@ -104,7 +104,7 @@ public:
 
 	/** Entry point to denoise the diffuse mask of a shadow. */
 	// TODO: correct specular.
-	virtual void DenoiseShadows(
+	virtual void DenoiseMonochromaticShadows(
 		FRDGBuilder& GraphBuilder,
 		const FViewInfo& View,
 		FPreviousViewInfo* PreviousViewInfos,
@@ -112,7 +112,30 @@ public:
 		const TStaticArray<FShadowParameters, IScreenSpaceDenoiser::kMaxBatchSize>& InputParameters,
 		const int32 InputParameterCount,
 		TStaticArray<FShadowPenumbraOutputs, IScreenSpaceDenoiser::kMaxBatchSize>& Outputs) const = 0;
-	
+
+	// Number of harmonics to be feed when denoising multiple lights.
+	static constexpr int32 kMultiPolychromaticPenumbraHarmonics = 4;
+
+	/** All the inputs to denoise polychromatic penumbra of multiple light. */
+	BEGIN_SHADER_PARAMETER_STRUCT(FPolychromaticPenumbraHarmonics, )
+		SHADER_PARAMETER_RDG_TEXTURE_ARRAY(Texture2D, Diffuse, [kMultiPolychromaticPenumbraHarmonics])
+		SHADER_PARAMETER_RDG_TEXTURE_ARRAY(Texture2D, Specular, [kMultiPolychromaticPenumbraHarmonics])
+	END_SHADER_PARAMETER_STRUCT()
+
+	/** All the outputs when denoising polychromatic penumbra. */
+	BEGIN_SHADER_PARAMETER_STRUCT(FPolychromaticPenumbraOutputs, )
+		SHADER_PARAMETER_RDG_TEXTURE(Texture2D, Diffuse)
+		SHADER_PARAMETER_RDG_TEXTURE(Texture2D, Specular)
+	END_SHADER_PARAMETER_STRUCT()
+
+	/** Entry point to denoise polychromatic penumbra of multiple light. */
+	virtual FPolychromaticPenumbraOutputs DenoisePolychromaticPenumbraHarmonics(
+		FRDGBuilder& GraphBuilder,
+		const FViewInfo& View,
+		FPreviousViewInfo* PreviousViewInfos,
+		const FSceneTextureParameters& SceneTextures,
+		const FPolychromaticPenumbraHarmonics& Inputs) const = 0;
+
 	/** All the inputs of the reflection denoiser. */
 	BEGIN_SHADER_PARAMETER_STRUCT(FReflectionsInputs, )
 		SHADER_PARAMETER_RDG_TEXTURE(Texture2D, Color)
