Refactored first part of post processing into a single RDG builder instance.

diff --git a/Engine/Shaders/Private/PostProcessBloom.usf b/Engine/Shaders/Private/PostProcessBloom.usf
--- a/Engine/Shaders/Private/PostProcessBloom.usf
+++ b/Engine/Shaders/Private/PostProcessBloom.usf
@@ -7,11 +7,18 @@
 #include "Common.ush"
 #include "PostProcessCommon.ush"
 #include "EyeAdaptationCommon.ush"
+#include "ScreenPass.ush"
 
 // for VisualizeBloomOverlayPS
 // xyz:Bloom1Tint.rgb, w:unused
 float4 ColorScale1;
 
+SCREEN_PASS_TEXTURE_VIEWPORT(Input)
+
+Texture2D InputTexture;
+
+SamplerState InputSampler;
+
 void VisualizeBloomSetupPS(
 	noperspective float4 UVAndScreenPos : TEXCOORD0,
 	float4 SvPosition : SV_POSITION, 
@@ -23,7 +30,7 @@ void VisualizeBloomSetupPS(
 
 	int2 PixelPos = (int2)SvPosition.xy;
 
-	int2 ViewportCenter = (int2)(ViewportRect.xy + ViewportSize.xy / 2);
+	int2 ViewportCenter = (int2)(Input_ViewportMin + Input_ViewportSize / 2);
 
 	// cross in x and y
 //	bool bMask = ViewportCenter.x == PixelPos.x || ViewportCenter.y == PixelPos.y;
@@ -36,7 +43,7 @@ void VisualizeBloomSetupPS(
 		OutColor.g = 20;
 	}
 	
-	float3 HDRColor = Texture2DSample(PostprocessInput0, PostprocessInput0Sampler, UV).rgb;
+	float3 HDRColor = Texture2DSample(InputTexture, InputSampler, UV).rgb;
 
 	// reference is in red
 	OutColor.r = dot(HDRColor, 1/3.0f);
diff --git a/Engine/Shaders/Private/PostProcessDownsample.usf b/Engine/Shaders/Private/PostProcessDownsample.usf
--- a/Engine/Shaders/Private/PostProcessDownsample.usf
+++ b/Engine/Shaders/Private/PostProcessDownsample.usf
@@ -5,113 +5,92 @@
 =============================================================================*/
 
 #include "Common.ush"
-#include "PostProcessCommon.ush"
-
-// vertex shader entry point
-void MainDownsampleVS(
-	in float4 InPosition : ATTRIBUTE0,
-	in float2 UV : ATTRIBUTE1,
-	out noperspective float2 OutUV : TEXCOORD0,
-	out float4 OutPosition : SV_POSITION
-	) 
-{
-	DrawRectangle(InPosition, UV, OutPosition, OutUV);
-}
+#include "ScreenPass.ush"
+
+#define DOWNSAMPLE_QUALITY_LOW 0
+#define DOWNSAMPLE_QUALITY_HIGH 1
+
+#ifndef DOWNSAMPLE_QUALITY
+#error DOWNSAMPLE_QUALITY is not defined.
+#endif
+
+Texture2D InputTexture;
 
+SamplerState InputSampler;
 
-float4 BufferBilinearUVMinMax;
+SCREEN_PASS_TEXTURE_VIEWPORT(Input)
 
-float4 SampleInput0(float2 BufferUV)
+bool IsComputeUVOutOfBounds(in float2 UV)
 {
-	#if MANUALLY_CLAMP_UV
-		BufferUV = clamp(BufferUV, BufferBilinearUVMinMax.xy, BufferBilinearUVMinMax.zw);
-	#endif
-	return Texture2DSampleLevel(PostprocessInput0, PostprocessInput0Sampler, BufferUV, 0);
+	float2 CenterDist = abs(UV - 0.5f);
+	return (max(CenterDist.x, CenterDist.y) >= 0.5f);
 }
 
+float4 SampleInput(float2 UV)
+{
+	UV = clamp(UV, Input_UVViewportBilinearMin, Input_UVViewportBilinearMax);
+
+	return Texture2DSampleLevel(InputTexture, InputSampler, UV, 0);
+}
 
-float4 DownsampleCommon(in float2 InUV, in float2 Offset)
+float4 DownsampleCommon(float2 UV)
 {
 	float4 OutColor;
 
-#if METHOD == 0
+#if DOWNSAMPLE_QUALITY == DOWNSAMPLE_QUALITY_LOW
 	// Output: low quality, 1 filtered sample
-	OutColor = SampleInput0(InUV);
+	OutColor = SampleInput(UV);
+#elif DOWNSAMPLE_QUALITY == DOWNSAMPLE_QUALITY_HIGH
+	// Output: float4(RGBA), 4 filtered samples
+	float2 UVs[4];
 
-#elif METHOD == 1
-	// Output: float4(RGBA), 4 filtered samples	
-	float2 UV[4];
-
-	// blur during downsample (4x4 kernel) to get better quality especially for HDR content, can be made an option of this shader
-	UV[0] = InUV + Offset * float2(-1, -1);
-	UV[1] = InUV + Offset * float2( 1, -1);
-	UV[2] = InUV + Offset * float2(-1,  1);
-	UV[3] = InUV + Offset * float2( 1,  1);
+	// Blur during downsample (4x4 kernel) to get better quality especially for HDR content.
+	UVs[0] = UV + Input_ExtentInverse * float2(-1, -1);
+	UVs[1] = UV + Input_ExtentInverse * float2( 1, -1);
+	UVs[2] = UV + Input_ExtentInverse * float2(-1,  1);
+	UVs[3] = UV + Input_ExtentInverse * float2( 1,  1);
 
 	float4 Sample[4];
 
-	UNROLL for(uint i = 0; i < 4; ++i)
+	UNROLL
+	for(uint i = 0; i < 4; ++i)
 	{
-		Sample[i] = SampleInput0(UV[i]);
+		Sample[i] = SampleInput(UVs[i]);
 	}
 
 	OutColor = (Sample[0] + Sample[1] + Sample[2] + Sample[3]) * 0.25f;
 
-	// fixed rarely occuring yellow color tint of the whole viewport (certain view port size, need to investigate more)
+	// Fixed rarely occurring yellow color tint of the whole viewport (certain viewport size, need to investigate more)
 	OutColor.rgb = max(float3(0,0,0), OutColor.rgb);
 #else
-	// Output: float4(RGB,depth), 4 unfiltered samples (no leaking)
-	float2 UV[4];
-
-	// no filtering (2x2 kernel) to get no leaking in Depth of Field
-	UV[0] = InUV + Offset * float2(-1, -1);
-	UV[1] = InUV + Offset * float2( 1, -1);
-	UV[2] = InUV + Offset * float2(-1,  1);
-	UV[3] = InUV + Offset * float2( 1,  1);
-
-	float3 Sample[4];
-
-	UNROLL for(uint i = 0; i < 4; ++i)
-	{
-		Sample[i] = SampleInput0(UV[i]).rgb;
-	}
-
-	// todo: this still leaks as we need to mask in focus pixels
-	OutColor = float4((Sample[0] + Sample[1] + Sample[2] + Sample[3]) * 0.25f, CalcSceneDepth(UV[0]));
+#error Invalid quality level specified.
 #endif
 
 	return OutColor;
 }
 
-// xy: UV offset, zw:unused
-float4 DownsampleParams;
-
 // pixel shader entry point
 void MainPS(noperspective float2 InUV : TEXCOORD0, out float4 OutColor : SV_Target0)
 {
-	float2 Offset = DownsampleParams.xy;
-	OutColor = DownsampleCommon(InUV, Offset);
+	OutColor = DownsampleCommon(InUV);
 }
 
 #if COMPUTESHADER
-// xy: UVOffset, zw: ThreadToUVScale
-float4 DownsampleComputeParams;
-RWTexture2D<float4> OutComputeTex;
+
+SCREEN_PASS_TEXTURE_VIEWPORT(Output)
+
+RWTexture2D<float4> OutComputeTexture;
 
 [numthreads(THREADGROUP_SIZEX, THREADGROUP_SIZEY, 1)]
 void MainCS(uint2 DispatchThreadId : SV_DispatchThreadID)
 {
-	float2 InUV = ((float2)DispatchThreadId + View.ViewRectMin.xy + 0.5f) * DownsampleComputeParams.zw;
-	float2 Offset = DownsampleComputeParams.xy;
+	uint2 iPixelPos = DispatchThreadId + Output_ViewportMin;
+	float2 PixelPos = float2(iPixelPos);
+	float2 UV = (PixelPos + 0.5) * Output_ExtentInverse;
 
-	if (IsComputeUVOutOfBounds(InUV))
+	if (!IsComputeUVOutOfBounds(UV))
 	{
-		return;
+		OutComputeTexture[PixelPos] = DownsampleCommon(UV);
 	}
-
-	float4 OutColor = DownsampleCommon(InUV, Offset);
-
-	uint2 PixelPos = DispatchThreadId + (uint2)View.ViewRectMin.xy;
-	OutComputeTex[PixelPos] = OutColor;
 }
 #endif 
\ No newline at end of file
diff --git a/Engine/Source/Runtime/RenderCore/Public/RenderGraphBuilder.h b/Engine/Source/Runtime/RenderCore/Public/RenderGraphBuilder.h
--- a/Engine/Source/Runtime/RenderCore/Public/RenderGraphBuilder.h
+++ b/Engine/Source/Runtime/RenderCore/Public/RenderGraphBuilder.h
@@ -25,11 +25,27 @@ public:
 		const TRefCountPtr<IPooledRenderTarget>& ExternalPooledTexture,
 		const TCHAR* Name = TEXT("External"));
 
+	/** Variant of RegisterExternalTexture which will returns null (rather than assert) if the external texture is null. */
+	inline FRDGTextureRef TryRegisterExternalTexture(
+		const TRefCountPtr<IPooledRenderTarget>& ExternalPooledTexture,
+		const TCHAR* Name = TEXT("External"))
+	{
+		return ExternalPooledTexture ? RegisterExternalTexture(ExternalPooledTexture, Name) : nullptr;
+	}
+
 	/** Register a external buffer to be tracked by the render graph. */
 	FRDGBufferRef RegisterExternalBuffer(
 		const TRefCountPtr<FPooledRDGBuffer>& ExternalPooledBuffer,
 		const TCHAR* Name = TEXT("External"));
 
+	/** Variant of RegisterExternalBuffer which will return null (rather than assert) if the external buffer is null. */
+	inline FRDGBufferRef TryRegisterExternalBuffer(
+		const TRefCountPtr<FPooledRDGBuffer>& ExternalPooledBuffer,
+		const TCHAR* Name = TEXT("External"))
+	{
+		return ExternalPooledBuffer ? RegisterExternalBuffer(ExternalPooledBuffer, Name) : nullptr;
+	}
+
 	/** Create graph tracked resource from a descriptor with a debug name.
 	 *
 	 *  The debug name is the name used for GPU debugging tools, but also for the VisualizeTexture/Vis command.
diff --git a/Engine/Source/Runtime/Renderer/Private/DebugViewModeRendering.cpp b/Engine/Source/Runtime/Renderer/Private/DebugViewModeRendering.cpp
--- a/Engine/Source/Runtime/Renderer/Private/DebugViewModeRendering.cpp
+++ b/Engine/Source/Runtime/Renderer/Private/DebugViewModeRendering.cpp
@@ -53,6 +53,8 @@ IMPLEMENT_MATERIAL_SHADER_TYPE(,FDebugViewModeDS,TEXT("/Engine/Private/DebugView
 
 ENGINE_API bool GetDebugViewMaterial(const UMaterialInterface* InMaterialInterface, EDebugViewShaderMode InDebugViewMode, ERHIFeatureLevel::Type InFeatureLevel,const FMaterialRenderProxy*& OutMaterialRenderProxy, const FMaterial*& OutMaterial);
 
+extern FRenderingCompositeOutputRef AddTemporalAADebugViewPass(FPostprocessContext& Context);
+
 void FDeferredShadingSceneRenderer::DoDebugViewModePostProcessing(FRHICommandListImmediate& RHICmdList, const FViewInfo& View, TRefCountPtr<IPooledRenderTarget>& VelocityRT)
 {
 	QUICK_SCOPE_CYCLE_COUNTER( STAT_PostProcessing_Process );
@@ -109,24 +111,7 @@ void FDeferredShadingSceneRenderer::DoDebugViewModePostProcessing(FRHICommandLis
 		}
 		case DVSM_RayTracingDebug:
 		{
-			FSceneViewState* ViewState = (FSceneViewState*)Context.View.State;
-			FTAAPassParameters Parameters(Context.View);
-			FRenderingCompositePass* Node = Context.Graph.RegisterPass(new(FMemStack::Get()) FRCPassPostProcessTemporalAA(
-				Context, 
-				Parameters, 
-				Context.View.PrevViewInfo.TemporalAAHistory, 
-				&ViewState->PrevFrameViewInfo.TemporalAAHistory));
-
-			FRenderingCompositePass* VelocityNode = VelocityRT
-				? Context.Graph.RegisterPass(new(FMemStack::Get()) FRCPassPostProcessInput(VelocityRT))
-				: Context.Graph.RegisterPass(new(FMemStack::Get()) FRCPassPostProcessInput(GSystemTextures.BlackDummy));
-
-			FRenderingCompositeOutputRef VelocityRef(VelocityNode);
-
-			Node->SetInput(ePId_Input0, FRenderingCompositeOutputRef(Context.FinalOutput));
-			Node->SetInput(ePId_Input2, VelocityRef);
-
-			Context.FinalOutput = FRenderingCompositeOutputRef(Node);
+			Context.FinalOutput = AddTemporalAADebugViewPass(Context);
 
 			if (View.Family->EngineShowFlags.Tonemapper)
 			{
diff --git a/Engine/Source/Runtime/Renderer/Private/LightShaftRendering.cpp b/Engine/Source/Runtime/Renderer/Private/LightShaftRendering.cpp
--- a/Engine/Source/Runtime/Renderer/Private/LightShaftRendering.cpp
+++ b/Engine/Source/Runtime/Renderer/Private/LightShaftRendering.cpp
@@ -29,6 +29,7 @@
 #include "PostProcess/PostProcessing.h"
 #include "PostProcess/PostProcessTemporalAA.h"
 #include "PipelineStateCache.h"
+#include "SceneTextureParameters.h"
 
 /** Tweaked values from UE3 implementation **/
 extern const float PointLightFadeDistanceIncrease = 200;
@@ -44,6 +45,15 @@ static FAutoConsoleVariableRef CVarLightShaftQuality(
 	TEXT("  1: on (default)"),
 	ECVF_Scalability | ECVF_RenderThreadSafe);
 
+int32 GLightShaftAllowTAA = 1;
+static FAutoConsoleVariableRef CVarLightAllowTAA(
+	TEXT("r.LightShaftAllowTAA"),
+	GLightShaftAllowTAA,
+	TEXT("Allows temporal filtering for lightshafts.\n")
+	TEXT("  0: off\n")
+	TEXT("  1: on (default)"),
+	ECVF_Scalability | ECVF_RenderThreadSafe);
+
 int32 GLightShaftDownsampleFactor = 2;
 static FAutoConsoleVariableRef CVarCacheLightShaftDownsampleFactor(
 	TEXT("r.LightShaftDownSampleFactor"),
@@ -576,41 +586,31 @@ void ApplyTemporalAA(
 	/** Output of Temporal AA for the next step in the pipeline. */
 	TRefCountPtr<IPooledRenderTarget>& HistoryOutput)
 {
-	if( View.AntiAliasingMethod == AAM_TemporalAA &&
-		HistoryState &&
-		HistoryState->RT[0] )
+	if (View.AntiAliasingMethod == AAM_TemporalAA && HistoryState && GLightShaftAllowTAA)
 	{
-		FMemMark Mark(FMemStack::Get());
-		FRenderingCompositePassContext CompositeContext(RHICmdList, View);
-		FPostprocessContext Context(RHICmdList, CompositeContext.Graph, View);
+		FRDGBuilder GraphBuilder(RHICmdList);
 
-		// Nodes for input render targets
-		FRenderingCompositePass* LightShaftSetup = Context.Graph.RegisterPass( new(FMemStack::Get()) FRCPassPostProcessInput( LightShaftsSource ) );
+		FRDGTextureRef LightShaftSetup = GraphBuilder.RegisterExternalTexture(LightShaftsSource, TEXT("LightShaftsSource"));
+
+		FSceneTextureParameters SceneTextures;
+		SetupSceneTextureParameters(GraphBuilder, &SceneTextures);
 
 		FTAAPassParameters TAAParameters(View);
 		TAAParameters.Pass = ETAAPassConfig::LightShaft;
 		TAAParameters.SetupViewRect(View, /** ResolutionDivisor = */ 2);
+		TAAParameters.SceneColorInput = LightShaftSetup;
 
-		// Temporal AA node
-		FRenderingCompositePass* NodeTemporalAA = Context.Graph.RegisterPass( new(FMemStack::Get()) FRCPassPostProcessTemporalAA(Context, TAAParameters,
-			*HistoryState, HistoryState) );
-
-		// Setup inputs on Temporal AA node as the shader expects
-		NodeTemporalAA->SetInput( ePId_Input0, LightShaftSetup );
-
-		// Reuse a render target from the pool with a consistent name, for vis purposes
-		TRefCountPtr<IPooledRenderTarget> NewHistory;
-		AllocateOrReuseLightShaftRenderTarget(RHICmdList, NewHistory, HistoryRTName);
-
-		// Setup the output to write to the new history render target
-		Context.FinalOutput = FRenderingCompositeOutputRef(NodeTemporalAA);
-		Context.FinalOutput.GetOutput()->RenderTargetDesc = NewHistory->GetDesc();
-		Context.FinalOutput.GetOutput()->PooledRenderTarget = NewHistory;
+		FTAAOutputs Outputs = AddTemporalAAPass(
+			GraphBuilder,
+			SceneTextures,
+			View,
+			TAAParameters,
+			*HistoryState,
+			HistoryState);
 
-		// Execute Temporal AA
-		CompositeContext.Process(Context.FinalOutput.GetPass(), TEXT("LightShaftTemporalAA"));
+		GraphBuilder.QueueTextureExtraction(Outputs.SceneColor, &HistoryOutput);
 
-		HistoryOutput = NewHistory;
+		GraphBuilder.Execute();
 	}
 	else
 	{
diff --git a/Engine/Source/Runtime/Renderer/Private/PlanarReflectionRendering.cpp b/Engine/Source/Runtime/Renderer/Private/PlanarReflectionRendering.cpp
--- a/Engine/Source/Runtime/Renderer/Private/PlanarReflectionRendering.cpp
+++ b/Engine/Source/Runtime/Renderer/Private/PlanarReflectionRendering.cpp
@@ -220,10 +220,8 @@ void PrefilterPlanarReflection(FRHICommandListImmediate& RHICmdList, FViewInfo&
 
 	if(View.FeatureLevel >= ERHIFeatureLevel::SM4)
 	{
-		// Note: null velocity buffer, so dynamic object temporal AA will not be correct
-		TRefCountPtr<IPooledRenderTarget> VelocityRT;
 		TRefCountPtr<IPooledRenderTarget> FilteredSceneColor;
-		GPostProcessing.ProcessPlanarReflection(RHICmdList, View, VelocityRT, FilteredSceneColor);
+		GPostProcessing.ProcessPlanarReflection(RHICmdList, View, FilteredSceneColor);
 
 		if (FilteredSceneColor)
 		{
diff --git a/Engine/Source/Runtime/Renderer/Private/PostProcess/DiaphragmDOF.cpp b/Engine/Source/Runtime/Renderer/Private/PostProcess/DiaphragmDOF.cpp
--- a/Engine/Source/Runtime/Renderer/Private/PostProcess/DiaphragmDOF.cpp
+++ b/Engine/Source/Runtime/Renderer/Private/PostProcess/DiaphragmDOF.cpp
@@ -1556,8 +1556,11 @@ FRDGTextureRef DiaphragmDOF::AddPasses(
 		TAAParameters.SceneColorInput = HalfResGatherInputTextures.SceneColor;
 		TAAParameters.SceneMetadataInput = HalfResGatherInputTextures.SeparateCoc;
 
-		FTAAOutputs TAAOutputs = TAAParameters.AddTemporalAAPass(
-			GraphBuilder, SceneTextures, View,
+		FTAAOutputs TAAOutputs = AddTemporalAAPass(
+			GraphBuilder,
+			SceneTextures,
+			View,
+			TAAParameters,
 			View.PrevViewInfo.DOFSetupHistory,
 			&ViewState->PrevFrameViewInfo.DOFSetupHistory);
 		
diff --git a/Engine/Source/Runtime/Renderer/Private/PostProcess/PostProcessBloomSetup.cpp b/Engine/Source/Runtime/Renderer/Private/PostProcess/PostProcessBloomSetup.cpp
--- a/Engine/Source/Runtime/Renderer/Private/PostProcess/PostProcessBloomSetup.cpp
+++ b/Engine/Source/Runtime/Renderer/Private/PostProcess/PostProcessBloomSetup.cpp
@@ -357,140 +357,64 @@ FPooledRenderTargetDesc FRCPassPostProcessBloomSetup::ComputeOutputDesc(EPassOut
 	return Ret;
 }
 
-
-/** Encapsulates the visualize bloom setup pixel shader. */
-class FPostProcessVisualizeBloomSetupPS : public FGlobalShader
+class FVisualizeBloomSetupPS : public FGlobalShader
 {
-	DECLARE_SHADER_TYPE(FPostProcessVisualizeBloomSetupPS, Global);
-
-	static bool ShouldCompilePermutation(const FGlobalShaderPermutationParameters& Parameters)
-	{
-		return IsFeatureLevelSupported(Parameters.Platform, ERHIFeatureLevel::SM4);
-	}
-
-	static void ModifyCompilationEnvironment(const FGlobalShaderPermutationParameters& Parameters, FShaderCompilerEnvironment& OutEnvironment)
-	{
-		FGlobalShader::ModifyCompilationEnvironment(Parameters,OutEnvironment);
-	}
-
-	/** Default constructor. */
-	FPostProcessVisualizeBloomSetupPS() {}
-
 public:
-	FPostProcessPassParameters PostprocessParameter;
-
-	/** Initialization constructor. */
-	FPostProcessVisualizeBloomSetupPS(const ShaderMetaType::CompiledShaderInitializerType& Initializer)
-		: FGlobalShader(Initializer)
-	{
-		PostprocessParameter.Bind(Initializer.ParameterMap);
-	}
+	DECLARE_GLOBAL_SHADER(FVisualizeBloomSetupPS);
+	SHADER_USE_PARAMETER_STRUCT(FVisualizeBloomSetupPS, FGlobalShader);
 
-	// FShader interface.
-	virtual bool Serialize(FArchive& Ar) override
-	{
-		bool bShaderHasOutdatedParameters = FGlobalShader::Serialize(Ar);
-		Ar << PostprocessParameter;
-		return bShaderHasOutdatedParameters;
-	}
+	BEGIN_SHADER_PARAMETER_STRUCT(FParameters, )
+		SHADER_PARAMETER_RDG_TEXTURE(Texture2D, InputTexture)
+		SHADER_PARAMETER_SAMPLER(SamplerState, InputSampler)
+		SHADER_PARAMETER_STRUCT(FScreenPassTextureViewportParameters, Input)
+		RENDER_TARGET_BINDING_SLOTS()
+	END_SHADER_PARAMETER_STRUCT()
 
-	template <typename TRHICmdList>
-	void SetParameters(TRHICmdList& RHICmdList, const FRenderingCompositePassContext& Context)
+	static bool ShouldCompilePermutation(const FGlobalShaderPermutationParameters& Parameters)
 	{
-		FRHIPixelShader* ShaderRHI = GetPixelShader();
-
-		FGlobalShader::SetParameters<FViewUniformShaderParameters>(RHICmdList, ShaderRHI, Context.View.ViewUniformBuffer);
-		PostprocessParameter.SetPS(RHICmdList, ShaderRHI, Context, TStaticSamplerState<SF_Bilinear, AM_Border, AM_Border, AM_Border>::GetRHI());
+		return IsFeatureLevelSupported(Parameters.Platform, ERHIFeatureLevel::SM4);
 	}
 };
 
-IMPLEMENT_SHADER_TYPE(,FPostProcessVisualizeBloomSetupPS,TEXT("/Engine/Private/PostProcessBloom.usf"),TEXT("VisualizeBloomSetupPS"),SF_Pixel);
-
-
+IMPLEMENT_GLOBAL_SHADER(FVisualizeBloomSetupPS, "/Engine/Private/PostProcessBloom.usf", "VisualizeBloomSetupPS", SF_Pixel);
 
-
-void FRCPassPostProcessVisualizeBloomSetup::Process(FRenderingCompositePassContext& Context)
+FRDGTextureRef AddVisualizeBloomSetupPass(
+	FRDGBuilder& GraphBuilder,
+	const FScreenPassViewInfo& ScreenPassView,
+	FRDGTextureRef InputTexture,
+	FIntRect InputViewport)
 {
-	SCOPED_DRAW_EVENT(Context.RHICmdList, VisualizeBloomSetup);
+	check(InputTexture);
 
-	const FPooledRenderTargetDesc* InputDesc = GetInputDesc(ePId_Input0);
+	FRDGTextureDesc OutputDesc = InputTexture->Desc;
+	OutputDesc.Reset();
+	OutputDesc.TargetableFlags &= ~TexCreate_UAV;
+	OutputDesc.TargetableFlags |= TexCreate_RenderTargetable;
 
-	check(InputDesc && "Input is not hooked up correctly");
+	FRDGTextureRef OutputTexture = GraphBuilder.CreateTexture(OutputDesc, TEXT("VisualizeBloomSetup"));
 
-	const FViewInfo& View = Context.View;
-
-	FIntPoint SrcSize = InputDesc->Extent;
-	FIntPoint DestSize = PassOutputs[0].RenderTargetDesc.Extent;
+	const FScreenPassTextureViewport TextureViewport(InputViewport, InputTexture);
 
-	// e.g. 4 means the input texture is 4x smaller than the buffer size
-	uint32 ScaleFactor = FMath::DivideAndRoundUp(FSceneRenderTargets::Get(Context.RHICmdList).GetBufferSizeXY().Y, SrcSize.Y);
+	FVisualizeBloomSetupPS::FParameters* PassParameters = GraphBuilder.AllocParameters<FVisualizeBloomSetupPS::FParameters>();
+	PassParameters->Input = GetScreenPassTextureViewportParameters(TextureViewport);
+	PassParameters->InputTexture = InputTexture;
+	PassParameters->InputSampler = TStaticSamplerState<SF_Bilinear, AM_Border, AM_Border, AM_Border>::GetRHI();
+	PassParameters->RenderTargets[0] = FRenderTargetBinding(OutputTexture, ERenderTargetLoadAction::ENoAction);
 
-	FIntRect SrcRect = FIntRect::DivideAndRoundUp(View.ViewRect, ScaleFactor);
-	FIntRect DestRect = SrcRect;
+	TShaderMapRef<FVisualizeBloomSetupPS> PixelShader(ScreenPassView.View.ShaderMap);
 
-	const FSceneRenderTargetItem& DestRenderTarget = PassOutputs[0].RequestSurface(Context);
-	
-	FRHIRenderPassInfo RPInfo(DestRenderTarget.TargetableTexture, ERenderTargetActions::DontLoad_Store);
-	Context.RHICmdList.BeginRenderPass(RPInfo, TEXT("VisualizeBloom"));
-	{
-		Context.SetViewportAndCallRHI(0, 0, 0.0f, DestRect.Width(), DestRect.Height(), 1.0f);
+	AddDrawScreenPass(
+		GraphBuilder,
+		RDG_EVENT_NAME("VisualizeBloomSetup"),
+		ScreenPassView,
+		TextureViewport,
+		TextureViewport,
+		*PixelShader,
+		PassParameters);
 
-		FGraphicsPipelineStateInitializer GraphicsPSOInit;
-		Context.RHICmdList.ApplyCachedRenderTargets(GraphicsPSOInit);
-
-		// set the state
-		GraphicsPSOInit.BlendState = TStaticBlendState<>::GetRHI();
-		GraphicsPSOInit.RasterizerState = TStaticRasterizerState<>::GetRHI();
-		GraphicsPSOInit.DepthStencilState = TStaticDepthStencilState<false, CF_Always>::GetRHI();
-
-		TShaderMapRef<FPostProcessVS> VertexShader(Context.GetShaderMap());
-		TShaderMapRef<FPostProcessVisualizeBloomSetupPS> PixelShader(Context.GetShaderMap());
-
-		GraphicsPSOInit.BoundShaderState.VertexDeclarationRHI = GFilterVertexDeclaration.VertexDeclarationRHI;
-		GraphicsPSOInit.BoundShaderState.VertexShaderRHI = GETSAFERHISHADER_VERTEX(*VertexShader);
-		GraphicsPSOInit.BoundShaderState.PixelShaderRHI = GETSAFERHISHADER_PIXEL(*PixelShader);
-		GraphicsPSOInit.PrimitiveType = PT_TriangleList;
-
-		SetGraphicsPipelineState(Context.RHICmdList, GraphicsPSOInit);
-
-		PixelShader->SetParameters(Context.RHICmdList, Context);
-		VertexShader->SetParameters(Context);
-
-		// Draw a quad mapping scene color to the view's render target
-		DrawRectangle(
-			Context.RHICmdList,
-			DestRect.Min.X, DestRect.Min.Y,
-			DestRect.Width(), DestRect.Height(),
-			SrcRect.Min.X, SrcRect.Min.Y,
-			SrcRect.Width(), SrcRect.Height(),
-			DestRect.Size(),
-			SrcSize,
-			*VertexShader,
-			EDRF_UseTriangleOptimization);
-	}
-	Context.RHICmdList.EndRenderPass();
-	Context.RHICmdList.CopyToResolveTarget(DestRenderTarget.TargetableTexture, DestRenderTarget.ShaderResourceTexture, FResolveParams());
-}
-
-
-FPooledRenderTargetDesc FRCPassPostProcessVisualizeBloomSetup::ComputeOutputDesc(EPassOutputId InPassOutputId) const
-{
-	FPooledRenderTargetDesc Ret = GetInput(ePId_Input0)->GetOutput()->RenderTargetDesc;
-
-	Ret.Reset();
-	Ret.TargetableFlags &= ~TexCreate_UAV;
-	Ret.TargetableFlags |= TexCreate_RenderTargetable;
-	Ret.DebugName = TEXT("VisualizeBloomSetup");
-
-	return Ret;
+	return OutputTexture;
 }
 
-
-
-
-
-
-
 /** Encapsulates the visualize bloom overlay pixel shader. */
 class FPostProcessVisualizeBloomOverlayPS : public FGlobalShader
 {
diff --git a/Engine/Source/Runtime/Renderer/Private/PostProcess/PostProcessBloomSetup.h b/Engine/Source/Runtime/Renderer/Private/PostProcess/PostProcessBloomSetup.h
--- a/Engine/Source/Runtime/Renderer/Private/PostProcess/PostProcessBloomSetup.h
+++ b/Engine/Source/Runtime/Renderer/Private/PostProcess/PostProcessBloomSetup.h
@@ -9,6 +9,13 @@
 #include "CoreMinimal.h"
 #include "RendererInterface.h"
 #include "PostProcess/RenderingCompositionGraph.h"
+#include "ScreenPass.h"
+
+FRDGTextureRef AddVisualizeBloomSetupPass(
+	FRDGBuilder& GraphBuilder,
+	const FScreenPassViewInfo& ScreenPassView,
+	FRDGTextureRef InputTexture,
+	FIntRect InputViewport);
 
 // ePId_Input0: Half res HDR scene color
 // ePId_Input1: EyeAdaptation
@@ -36,17 +43,6 @@ private:
 	FComputeFenceRHIRef AsyncEndFence;
 };
 
-// ePId_Input0: HDR SceneColor
-// derives from TRenderingCompositePassBase<InputCount, OutputCount> 
-class FRCPassPostProcessVisualizeBloomSetup : public TRenderingCompositePassBase<1, 1>
-{
-public:
-	// interface FRenderingCompositePass ---------
-	virtual void Process(FRenderingCompositePassContext& Context) override;
-	virtual FPooledRenderTargetDesc ComputeOutputDesc(EPassOutputId InPassOutputId) const override;
-	virtual void Release() override { delete this; }
-};
-
 // ePId_Input0: LDR SceneColor
 // ePId_Input1: HDR SceneColor
 // ePId_Input2: BloomOutputCombined
diff --git a/Engine/Source/Runtime/Renderer/Private/PostProcess/PostProcessDownsample.cpp b/Engine/Source/Runtime/Renderer/Private/PostProcess/PostProcessDownsample.cpp
--- a/Engine/Source/Runtime/Renderer/Private/PostProcess/PostProcessDownsample.cpp
+++ b/Engine/Source/Runtime/Renderer/Private/PostProcess/PostProcessDownsample.cpp
@@ -5,166 +5,83 @@
 =============================================================================*/
 
 #include "PostProcess/PostProcessDownsample.h"
-#include "ClearQuad.h"
-#include "StaticBoundShaderState.h"
-#include "SceneUtils.h"
-#include "PostProcess/SceneRenderTargets.h"
-#include "PostProcess/SceneFilterRendering.h"
-#include "SceneRenderTargetParameters.h"
-#include "SceneRendering.h"
-#include "PipelineStateCache.h"
 
+namespace
+{
+	
 const int32 GDownsampleTileSizeX = 8;
 const int32 GDownsampleTileSizeY = 8;
 
-/** Encapsulates the post processing down sample pixel shader. */
-template <uint32 Method, uint32 ManuallyClampUV>
-class FPostProcessDownsamplePS : public FGlobalShader
+TAutoConsoleVariable<int32> CVarDownsampleQuality(
+	TEXT("r.Downsample.Quality"),
+	1,
+	TEXT("Defines the quality in which the Downsample passes. we might add more quality levels later.\n")
+	TEXT(" 0: low quality\n")
+	TEXT(">0: high quality (default: 1)\n"),
+	ECVF_Scalability | ECVF_RenderThreadSafe);
+
+BEGIN_SHADER_PARAMETER_STRUCT(FDownsampleParameters, )
+	SHADER_PARAMETER_STRUCT_REF(FViewUniformShaderParameters, ViewUniformBuffer)
+	SHADER_PARAMETER_STRUCT(FScreenPassTextureViewportParameters, Input)
+	SHADER_PARAMETER_RDG_TEXTURE(Texture2D, InputTexture)
+	SHADER_PARAMETER_SAMPLER(SamplerState, InputSampler)
+END_SHADER_PARAMETER_STRUCT()
+
+FDownsampleParameters GetDownsampleParameters(
+	const FViewInfo& View,
+	const FIntRect InputViewport,
+	FRDGTextureRef InputTexture,
+	EDownsampleQuality DownsampleMethod)
 {
-	DECLARE_SHADER_TYPE(FPostProcessDownsamplePS, Global);
+	check(InputTexture);
 
-	static bool ShouldCompilePermutation(const FGlobalShaderPermutationParameters& Parameters)
-	{
-		return Method != 2 || IsFeatureLevelSupported(Parameters.Platform, ERHIFeatureLevel::SM4);
-	}
+	const FScreenPassTextureViewportParameters InputParameters = GetScreenPassTextureViewportParameters(FScreenPassTextureViewport(InputViewport, InputTexture));
 
-	static void ModifyCompilationEnvironment(const FGlobalShaderPermutationParameters& Parameters, FShaderCompilerEnvironment& OutEnvironment)
-	{
-		FGlobalShader::ModifyCompilationEnvironment(Parameters,OutEnvironment);
-		OutEnvironment.SetDefine(TEXT("METHOD"), Method);
-		OutEnvironment.SetDefine(TEXT("MANUALLY_CLAMP_UV"), ManuallyClampUV);
-	}
+	FDownsampleParameters Parameters;
+	Parameters.ViewUniformBuffer = View.ViewUniformBuffer;
+	Parameters.Input = InputParameters;
+	Parameters.InputTexture = InputTexture;
+	Parameters.InputSampler = TStaticSamplerState<SF_Bilinear, AM_Clamp, AM_Clamp, AM_Clamp>::GetRHI();
+	return Parameters;
+}
 
-	/** Default constructor. */
-	FPostProcessDownsamplePS() {}
+class FDownsampleQualityDimension : SHADER_PERMUTATION_ENUM_CLASS("DOWNSAMPLE_QUALITY", EDownsampleQuality);
+using FDownsamplePermutationDomain = TShaderPermutationDomain<FDownsampleQualityDimension>;
 
+class FDownsamplePS : public FGlobalShader
+{
 public:
-	FPostProcessPassParameters PostprocessParameter;
-	FSceneTextureShaderParameters SceneTextureParameters;
-	FShaderParameter DownsampleParams;
-	FShaderParameter BufferBilinearUVMinMax;
-
-
-	/** Initialization constructor. */
-	FPostProcessDownsamplePS(const ShaderMetaType::CompiledShaderInitializerType& Initializer)
-		: FGlobalShader(Initializer)
-	{
-		PostprocessParameter.Bind(Initializer.ParameterMap);
-		SceneTextureParameters.Bind(Initializer);
-		DownsampleParams.Bind(Initializer.ParameterMap, TEXT("DownsampleParams"));
-		BufferBilinearUVMinMax.Bind(Initializer.ParameterMap, TEXT("BufferBilinearUVMinMax"));
-	}
-
-	// FShader interface.
-	virtual bool Serialize(FArchive& Ar) override
-	{
-		bool bShaderHasOutdatedParameters = FGlobalShader::Serialize(Ar);
-		Ar << PostprocessParameter << SceneTextureParameters << DownsampleParams << BufferBilinearUVMinMax;
-		return bShaderHasOutdatedParameters;
-	}
-
-	template <typename TRHICmdList>
-	void SetParameters(TRHICmdList& RHICmdList, const FRenderingCompositePassContext& Context, const FPooledRenderTargetDesc* InputDesc, const FIntPoint& SrcSize, const FIntRect& SrcRect)
-	{
-		FRHIPixelShader* ShaderRHI = GetPixelShader();
-
-		FGlobalShader::SetParameters<FViewUniformShaderParameters>(RHICmdList, ShaderRHI, Context.View.ViewUniformBuffer);
-		SceneTextureParameters.Set(RHICmdList, ShaderRHI, Context.View.FeatureLevel, ESceneTextureSetupMode::All);
-
-		// filter only if needed for better performance
-		FRHISamplerState* Filter = (Method == 2) ?
-			TStaticSamplerState<SF_Point, AM_Clamp, AM_Clamp, AM_Clamp>::GetRHI():
-			TStaticSamplerState<SF_Bilinear, AM_Clamp, AM_Clamp, AM_Clamp>::GetRHI();
-		
-		{
-			float PixelScale = (Method == 2) ? 0.5f : 1.0f;
- 
-			FVector4 DownsampleParamsValue(PixelScale / InputDesc->Extent.X, PixelScale / InputDesc->Extent.Y, 0, 0);
-			SetShaderValue(RHICmdList, ShaderRHI, DownsampleParams, DownsampleParamsValue);
-		}
-
-		PostprocessParameter.SetPS(RHICmdList, ShaderRHI, Context, Filter);
+	DECLARE_GLOBAL_SHADER(FDownsamplePS);
+	SHADER_USE_PARAMETER_STRUCT(FDownsamplePS, FGlobalShader);
 
-		if (ManuallyClampUV)
-		{
-			FVector4 BufferBilinearUVMinMaxValue(
-				(SrcRect.Min.X + 0.5f) / SrcSize.X, (SrcRect.Min.Y + 0.5f) / SrcSize.Y,
-				(SrcRect.Max.X - 0.5f) / SrcSize.X, (SrcRect.Max.Y - 0.5f) / SrcSize.Y);
-			SetShaderValue(RHICmdList, ShaderRHI, BufferBilinearUVMinMax, BufferBilinearUVMinMaxValue);
-		}
-	}
+	using FPermutationDomain = FDownsamplePermutationDomain;
 
-	static const TCHAR* GetSourceFilename()
-	{
-		return TEXT("/Engine/Private/PostProcessDownsample.usf");
-	}
+	BEGIN_SHADER_PARAMETER_STRUCT(FParameters, )
+		SHADER_PARAMETER_STRUCT_INCLUDE(FDownsampleParameters, Common)
+		RENDER_TARGET_BINDING_SLOTS()
+	END_SHADER_PARAMETER_STRUCT()
 
-	static const TCHAR* GetFunctionName()
+	static bool ShouldCompilePermutation(const FGlobalShaderPermutationParameters& Parameters)
 	{
-		return TEXT("MainPS");
+		return true;
 	}
 };
 
-// #define avoids a lot of code duplication
-#define VARIATION1(A,B) typedef FPostProcessDownsamplePS<A, B> FPostProcessDownsamplePS##A##B; \
-	IMPLEMENT_SHADER_TYPE2(FPostProcessDownsamplePS##A##B, SF_Pixel);
+IMPLEMENT_GLOBAL_SHADER(FDownsamplePS, "/Engine/Private/PostProcessDownsample.usf", "MainPS", SF_Pixel);
 
-VARIATION1(0, 0) VARIATION1(1, 0) VARIATION1(2, 0)
-VARIATION1(0, 1) VARIATION1(1, 1) VARIATION1(2, 1)
-#undef VARIATION1
-
-
-/** Encapsulates the post processing down sample vertex shader. */
-class FPostProcessDownsampleVS : public FGlobalShader
+class FDownsampleCS : public FGlobalShader
 {
-	DECLARE_SHADER_TYPE(FPostProcessDownsampleVS,Global);
 public:
+	DECLARE_GLOBAL_SHADER(FDownsampleCS);
+	SHADER_USE_PARAMETER_STRUCT(FDownsampleCS, FGlobalShader);
 
-	static bool ShouldCompilePermutation(const FGlobalShaderPermutationParameters& Parameters)
-	{
-		return true;
-	}
-
-	/** Default constructor. */
-	FPostProcessDownsampleVS() {}
-	
-	/** Initialization constructor. */
-	FPostProcessDownsampleVS(const ShaderMetaType::CompiledShaderInitializerType& Initializer):
-		FGlobalShader(Initializer)
-	{
-	}
-
-	/** Serializer */
-	virtual bool Serialize(FArchive& Ar) override
-	{
-		bool bShaderHasOutdatedParameters = FGlobalShader::Serialize(Ar);
-
-		return bShaderHasOutdatedParameters;
-	}
-
-	void SetParameters(const FRenderingCompositePassContext& Context)
-	{
-		FRHIVertexShader* ShaderRHI = GetVertexShader();
-
-		FGlobalShader::SetParameters<FViewUniformShaderParameters>(Context.RHICmdList, ShaderRHI, Context.View.ViewUniformBuffer);
-
-		const FPooledRenderTargetDesc* InputDesc = Context.Pass->GetInputDesc(ePId_Input0);
+	using FPermutationDomain = FDownsamplePermutationDomain;
 
-		if(!InputDesc)
-		{
-			// input is not hooked up correctly
-			return;
-		}
-	}
-};
-
-IMPLEMENT_SHADER_TYPE(,FPostProcessDownsampleVS,TEXT("/Engine/Private/PostProcessDownsample.usf"),TEXT("MainDownsampleVS"),SF_Vertex);
-
-/** Encapsulates the post processing down sample compute shader. */
-template <uint32 Method>
-class FPostProcessDownsampleCS : public FGlobalShader
-{
-	DECLARE_SHADER_TYPE(FPostProcessDownsampleCS, Global);
+	BEGIN_SHADER_PARAMETER_STRUCT(FParameters, )
+		SHADER_PARAMETER_STRUCT_INCLUDE(FDownsampleParameters, Common)
+		SHADER_PARAMETER_STRUCT(FScreenPassTextureViewportParameters, Output)
+		SHADER_PARAMETER_RDG_TEXTURE_UAV(RWTexture2D<float4>, OutComputeTexture)
+	END_SHADER_PARAMETER_STRUCT()
 
 	static bool ShouldCompilePermutation(const FGlobalShaderPermutationParameters& Parameters)
 	{
@@ -174,304 +91,135 @@ class FPostProcessDownsampleCS : public FGlobalShader
 	static void ModifyCompilationEnvironment(const FGlobalShaderPermutationParameters& Parameters, FShaderCompilerEnvironment& OutEnvironment)
 	{
 		FGlobalShader::ModifyCompilationEnvironment(Parameters,OutEnvironment);
-		OutEnvironment.SetDefine(TEXT("METHOD"), Method);
 		OutEnvironment.SetDefine(TEXT("THREADGROUP_SIZEX"), GDownsampleTileSizeX);
 		OutEnvironment.SetDefine(TEXT("THREADGROUP_SIZEY"), GDownsampleTileSizeY);
 	}
-
-	/** Default constructor. */
-	FPostProcessDownsampleCS() {}
-
-public:
-	FPostProcessPassParameters PostprocessParameter;
-	FSceneTextureShaderParameters SceneTextureParameters;
-	FShaderParameter DownsampleComputeParams;
-	FShaderParameter OutComputeTex;
-
-	/** Initialization constructor. */
-	FPostProcessDownsampleCS(const ShaderMetaType::CompiledShaderInitializerType& Initializer)
-		: FGlobalShader(Initializer)
-	{
-		PostprocessParameter.Bind(Initializer.ParameterMap);
-		SceneTextureParameters.Bind(Initializer);
-		DownsampleComputeParams.Bind(Initializer.ParameterMap, TEXT("DownsampleComputeParams"));
-		OutComputeTex.Bind(Initializer.ParameterMap, TEXT("OutComputeTex"));
-	}
-
-	// FShader interface.
-	virtual bool Serialize(FArchive& Ar) override
-	{
-		bool bShaderHasOutdatedParameters = FGlobalShader::Serialize(Ar);
-		Ar << PostprocessParameter << SceneTextureParameters << DownsampleComputeParams << OutComputeTex;
-		return bShaderHasOutdatedParameters;
-	}
-
-	template <typename TRHICmdList>
-	void SetParameters(TRHICmdList& RHICmdList, const FRenderingCompositePassContext& Context, const FIntPoint& SrcSize, FRHIUnorderedAccessView* DestUAV)
-	{
-		FRHIComputeShader* ShaderRHI = GetComputeShader();
-		FGlobalShader::SetParameters<FViewUniformShaderParameters>(RHICmdList, ShaderRHI, Context.View.ViewUniformBuffer);
-
-		// filter only if needed for better performance
-		FRHISamplerState* Filter = (Method == 2) ?
-			TStaticSamplerState<SF_Point, AM_Clamp, AM_Clamp, AM_Clamp>::GetRHI():
-			TStaticSamplerState<SF_Bilinear, AM_Clamp, AM_Clamp, AM_Clamp>::GetRHI();
-
-		PostprocessParameter.SetCS(ShaderRHI, Context, RHICmdList, Filter);
-		SceneTextureParameters.Set(RHICmdList, ShaderRHI, Context.View.FeatureLevel, ESceneTextureSetupMode::All);
-		RHICmdList.SetUAVParameter(ShaderRHI, OutComputeTex.GetBaseIndex(), DestUAV);
-		
-		const float PixelScale = (Method == 2) ? 0.5f : 1.0f;
-		FVector4 DownsampleComputeValues(PixelScale / SrcSize.X, PixelScale / SrcSize.Y, 2.f / SrcSize.X, 2.f / SrcSize.Y);
-		SetShaderValue(RHICmdList, ShaderRHI, DownsampleComputeParams, DownsampleComputeValues);
-	}
-
-	template <typename TRHICmdList>
-	void UnsetParameters(TRHICmdList& RHICmdList)
-	{
-		FRHIComputeShader* ShaderRHI = GetComputeShader();
-		RHICmdList.SetUAVParameter(ShaderRHI, OutComputeTex.GetBaseIndex(), NULL);
-	}
 };
 
-// #define avoids a lot of code duplication
-#define VARIATION1(A) typedef FPostProcessDownsampleCS<A> FPostProcessDownsampleCS##A; \
-	IMPLEMENT_SHADER_TYPE(template<>,FPostProcessDownsampleCS##A,TEXT("/Engine/Private/PostProcessDownsample.usf"),TEXT("MainCS"),SF_Compute);
-
-VARIATION1(0)			VARIATION1(1)			VARIATION1(2)
-#undef VARIATION1
+IMPLEMENT_GLOBAL_SHADER(FDownsampleCS, "/Engine/Private/PostProcessDownsample.usf", "MainCS", SF_Compute);
 
+} //! namespace
 
-FRCPassPostProcessDownsample::FRCPassPostProcessDownsample(EPixelFormat InOverrideFormat, uint32 InQuality, bool bInIsComputePass, const TCHAR *InDebugName)
-	: OverrideFormat(InOverrideFormat)
-	, Quality(InQuality)
-	, DebugName(InDebugName)
+EDownsampleQuality GetDownsampleQuality()
 {
-	bIsComputePass = bInIsComputePass;
-	bPreferAsyncCompute = false;
-}
+	const int32 DownsampleQuality = FMath::Clamp(CVarDownsampleQuality.GetValueOnRenderThread(), 0, 1);
 
-template <uint32 Method, uint32 ManuallyClampUV>
-void FRCPassPostProcessDownsample::SetShader(const FRenderingCompositePassContext& Context, const FPooledRenderTargetDesc* InputDesc, const FIntPoint& SrcSize, const FIntRect& SrcRect)
-{
-	FGraphicsPipelineStateInitializer GraphicsPSOInit;
-	Context.RHICmdList.ApplyCachedRenderTargets(GraphicsPSOInit);
-	GraphicsPSOInit.BlendState = TStaticBlendState<>::GetRHI();
-	GraphicsPSOInit.RasterizerState = TStaticRasterizerState<>::GetRHI();
-	GraphicsPSOInit.DepthStencilState = TStaticDepthStencilState<false, CF_Always>::GetRHI();
-	GraphicsPSOInit.PrimitiveType = PT_TriangleList;
-
-	auto ShaderMap = Context.GetShaderMap();
-	TShaderMapRef<FPostProcessDownsampleVS> VertexShader(ShaderMap);
-	TShaderMapRef<FPostProcessDownsamplePS<Method, ManuallyClampUV> > PixelShader(ShaderMap);
-
-	GraphicsPSOInit.BoundShaderState.VertexDeclarationRHI = GFilterVertexDeclaration.VertexDeclarationRHI;
-	GraphicsPSOInit.BoundShaderState.VertexShaderRHI = GETSAFERHISHADER_VERTEX(*VertexShader);
-	GraphicsPSOInit.BoundShaderState.PixelShaderRHI = GETSAFERHISHADER_PIXEL(*PixelShader);
-	SetGraphicsPipelineState(Context.RHICmdList, GraphicsPSOInit);
-
-	PixelShader->SetParameters(Context.RHICmdList, Context, InputDesc, SrcSize, SrcRect);
-	VertexShader->SetParameters(Context);
+	return static_cast<EDownsampleQuality>(DownsampleQuality);
 }
 
-void FRCPassPostProcessDownsample::Process(FRenderingCompositePassContext& Context)
+FDownsamplePassOutputs AddDownsamplePass(
+	FRDGBuilder& GraphBuilder,
+	const FScreenPassViewInfo& ScreenPassView,
+	const FDownsamplePassInputs& Inputs)
 {
-	const FPooledRenderTargetDesc* InputDesc = GetInputDesc(ePId_Input0);
-	AsyncEndFence = FComputeFenceRHIRef();
+	check(Inputs.Texture);
+
+	bool bIsComputePass = ScreenPassView.bUseComputePasses;
 
-	if (!InputDesc)
+	if ((Inputs.Flags & EDownsampleFlags::ForceRaster) == EDownsampleFlags::ForceRaster)
 	{
-		return;
+		bIsComputePass = false;
 	}
 
-	const FViewInfo& View = Context.View;
-	const FSceneViewFamily& ViewFamily = *(View.Family);
+	FRDGTextureRef OutputTexture = nullptr;
 
-	FIntPoint SrcSize = InputDesc->Extent;
-	FIntPoint DestSize = PassOutputs[0].RenderTargetDesc.Extent;
-
-	// e.g. 4 means the input texture is 4x smaller than the buffer size
-	uint32 ScaleFactor = FMath::RoundUpToPowerOfTwo(FMath::DivideAndRoundUp(Context.ReferenceBufferSize.Y, SrcSize.Y));
+	// Construct the output texture to be half resolution (rounded up to even) with an optional format override.
+	{
+		FRDGTextureDesc Desc = Inputs.Texture->Desc;
+		Desc.Reset();
+		Desc.Extent = FIntPoint::DivideAndRoundUp(Desc.Extent, 2);
+		Desc.Extent.X = FMath::Max(1, Desc.Extent.X);
+		Desc.Extent.Y = FMath::Max(1, Desc.Extent.Y);
+		Desc.TargetableFlags &= ~(TexCreate_RenderTargetable | TexCreate_UAV);
+		Desc.TargetableFlags |= bIsComputePass ? TexCreate_UAV : TexCreate_RenderTargetable;
+		Desc.Flags |= GFastVRamConfig.Downsample;
+		Desc.AutoWritable = false;
+		Desc.DebugName = Inputs.Name;
+		Desc.ClearValue = FClearValueBinding(FLinearColor(0, 0, 0, 0));
 
-	FIntRect SrcRect = FIntRect::DivideAndRoundUp(Context.SceneColorViewRect, ScaleFactor);
-	FIntRect DestRect = FIntRect::DivideAndRoundUp(Context.SceneColorViewRect, ScaleFactor * 2);
+		if (Inputs.FormatOverride != PF_Unknown)
+		{
+			Desc.Format = Inputs.FormatOverride;
+		}
 
-	SCOPED_DRAW_EVENTF(Context.RHICmdList, Downsample, TEXT("Downsample%s %dx%d"), bIsComputePass?TEXT("Compute"):TEXT(""), DestRect.Width(), DestRect.Height());
+		OutputTexture = GraphBuilder.CreateTexture(Desc, Inputs.Name);
+	}
 
-	const FSceneRenderTargetItem& DestRenderTarget = PassOutputs[0].RequestSurface(Context);
-	const bool bIsDepthInputAvailable = (GetInputDesc(ePId_Input1) != nullptr);
+	FDownsamplePermutationDomain PermutationVector;
+	PermutationVector.Set<FDownsampleQualityDimension>(Inputs.Quality);
 
-	bool bManuallyClampUV = !(SrcRect.Min == FIntPoint::ZeroValue && SrcRect.Max == SrcSize);
+	const FIntRect OutputViewport = FIntRect::DivideAndRoundUp(Inputs.Viewport, 2);
 
 	if (bIsComputePass)
-	{	
-		DestRect = {Context.SceneColorViewRect.Min, Context.SceneColorViewRect.Min + DestSize};
-
-		// Common setup
-		// #todo-renderpasses remove once everything is renderpasses
-		UnbindRenderTargets(Context.RHICmdList);
-		Context.SetViewportAndCallRHI(DestRect, 0.0f, 1.0f);
-		
-		static FName AsyncEndFenceName(TEXT("AsyncDownsampleEndFence"));
-		AsyncEndFence = Context.RHICmdList.CreateComputeFence(AsyncEndFenceName);
-
-		if (IsAsyncComputePass())
-		{
-			// Async path
-			FRHIAsyncComputeCommandListImmediate& RHICmdListComputeImmediate = FRHICommandListExecutor::GetImmediateAsyncComputeCommandList();
-			{
-	 			SCOPED_COMPUTE_EVENT(RHICmdListComputeImmediate, AsyncDownsample);
-
-				WaitForInputPassComputeFences(RHICmdListComputeImmediate);
-				RHICmdListComputeImmediate.TransitionResource(EResourceTransitionAccess::ERWBarrier, EResourceTransitionPipeline::EGfxToCompute, DestRenderTarget.UAV);
-
-				if (bIsDepthInputAvailable)
-				{
-					DispatchCS<2>(RHICmdListComputeImmediate, Context, SrcSize, DestRect, DestRenderTarget.UAV);
-				}
-				else if (Quality == 0)
-				{
-					DispatchCS<0>(RHICmdListComputeImmediate, Context, SrcSize, DestRect, DestRenderTarget.UAV);
-				}
-				else
-				{
-					DispatchCS<1>(RHICmdListComputeImmediate, Context, SrcSize, DestRect, DestRenderTarget.UAV);
-				}
-
-				RHICmdListComputeImmediate.TransitionResource(EResourceTransitionAccess::EReadable, EResourceTransitionPipeline::EComputeToGfx, DestRenderTarget.UAV, AsyncEndFence);
-			}
-			FRHIAsyncComputeCommandListImmediate::ImmediateDispatch(RHICmdListComputeImmediate);		
-		}
-		else
-		{
-			// Direct path
-			WaitForInputPassComputeFences(Context.RHICmdList);
-			Context.RHICmdList.TransitionResource(EResourceTransitionAccess::ERWBarrier, EResourceTransitionPipeline::EGfxToCompute, DestRenderTarget.UAV);
-
-			if (bIsDepthInputAvailable)
-			{
-				DispatchCS<2>(Context.RHICmdList, Context, SrcSize, DestRect, DestRenderTarget.UAV);
-			}
-			else if (Quality == 0)
-			{
-				DispatchCS<0>(Context.RHICmdList, Context, SrcSize, DestRect, DestRenderTarget.UAV);
-			}
-			else
-			{
-				DispatchCS<1>(Context.RHICmdList, Context, SrcSize, DestRect, DestRenderTarget.UAV);
-			}
-
-			Context.RHICmdList.TransitionResource(EResourceTransitionAccess::EReadable, EResourceTransitionPipeline::EComputeToGfx, DestRenderTarget.UAV, AsyncEndFence);
-		}
+	{
+		FDownsampleCS::FParameters* PassParameters = GraphBuilder.AllocParameters<FDownsampleCS::FParameters>();
+		PassParameters->Common = GetDownsampleParameters(ScreenPassView.View, Inputs.Viewport, Inputs.Texture, Inputs.Quality);
+		PassParameters->Output = GetScreenPassTextureViewportParameters(FScreenPassTextureViewport(OutputViewport, OutputTexture));
+		PassParameters->OutComputeTexture = GraphBuilder.CreateUAV(OutputTexture);
+
+		TShaderMapRef<FDownsampleCS> ComputeShader(ScreenPassView.View.ShaderMap, PermutationVector);
+
+		FComputeShaderUtils::AddPass(
+			GraphBuilder,
+			RDG_EVENT_NAME("Downsample.%s %dx%d (CS)", Inputs.Name, Inputs.Viewport.Width(), Inputs.Viewport.Height()),
+			*ComputeShader,
+			PassParameters,
+			FComputeShaderUtils::GetGroupCount(OutputViewport.Size(), FIntPoint(GDownsampleTileSizeX, GDownsampleTileSizeY)));
 	}
 	else
 	{
-		// #todo-renderpasses only clear dest rectangle if it's been computed
-		FRHIRenderPassInfo RPInfo(DestRenderTarget.TargetableTexture, ERenderTargetActions::Clear_Store);
-		Context.RHICmdList.BeginRenderPass(RPInfo, TEXT("PostProcessDownsample"));
-		{
-			Context.SetViewportAndCallRHI(0, 0, 0.0f, DestSize.X, DestSize.Y, 1.0f);
-			// InflateSize increases the size of the source/dest rectangle to compensate for bilinear reads and UIBlur pass requirements.
-			int32 InflateSize;
-			// if second input is hooked up
-
-			if (bManuallyClampUV)
-			{
-				if (bIsDepthInputAvailable)
-				{
-					// also put depth in alpha
-					InflateSize = 2;
-					SetShader<2, true>(Context, InputDesc, SrcSize, SrcRect);
-				}
-				else if (Quality == 0)
-				{
-					SetShader<0, true>(Context, InputDesc, SrcSize, SrcRect);
-					InflateSize = 1;
-				}
-				else
-				{
-					SetShader<1, true>(Context, InputDesc, SrcSize, SrcRect);
-					InflateSize = 2;
-				}
-			}
-			else
-			{
-				if (bIsDepthInputAvailable)
-				{
-					// also put depth in alpha
-					InflateSize = 2;
-					SetShader<2, false>(Context, InputDesc, SrcSize, SrcRect);
-				}
-				else if (Quality == 0)
-				{
-					SetShader<0, false>(Context, InputDesc, SrcSize, SrcRect);
-					InflateSize = 1;
-				}
-				else
-				{
-					SetShader<1, false>(Context, InputDesc, SrcSize, SrcRect);
-					InflateSize = 2;
-				}
-			}
-
-			TShaderMapRef<FPostProcessDownsampleVS> VertexShader(Context.GetShaderMap());
-
-			SrcRect = DestRect * 2;
-			DrawRectangle(
-				Context.RHICmdList,
-				DestRect.Min.X, DestRect.Min.Y,
-				DestRect.Width(), DestRect.Height(),
-				SrcRect.Min.X, SrcRect.Min.Y,
-				SrcRect.Width(), SrcRect.Height(),
-				DestSize,
-				SrcSize,
-				*VertexShader,
-				EDRF_UseTriangleOptimization);
-		}
-		Context.RHICmdList.EndRenderPass();
-		Context.RHICmdList.CopyToResolveTarget(DestRenderTarget.TargetableTexture, DestRenderTarget.ShaderResourceTexture, FResolveParams());
-	}
-}
+		FDownsamplePS::FParameters* PassParameters = GraphBuilder.AllocParameters<FDownsamplePS::FParameters>();
+		PassParameters->Common = GetDownsampleParameters(ScreenPassView.View, Inputs.Viewport, Inputs.Texture, Inputs.Quality);
+		PassParameters->RenderTargets[0] = FRenderTargetBinding(OutputTexture, ERenderTargetLoadAction::ENoAction);
 
-template <uint32 Method, typename TRHICmdList>
-void FRCPassPostProcessDownsample::DispatchCS(TRHICmdList& RHICmdList, FRenderingCompositePassContext& Context, const FIntPoint& SrcSize, const FIntRect& DestRect, FRHIUnorderedAccessView* DestUAV)
-{
-	auto ShaderMap = Context.GetShaderMap();
-	TShaderMapRef<FPostProcessDownsampleCS<Method>> ComputeShader(ShaderMap);
-	RHICmdList.SetComputeShader(ComputeShader->GetComputeShader());
-
-	ComputeShader->SetParameters(RHICmdList, Context, SrcSize, DestUAV);
+		TShaderMapRef<FDownsamplePS> PixelShader(ScreenPassView.View.ShaderMap, PermutationVector);
 
-	uint32 GroupSizeX = FMath::DivideAndRoundUp(DestRect.Width(), GDownsampleTileSizeX);
-	uint32 GroupSizeY = FMath::DivideAndRoundUp(DestRect.Height(), GDownsampleTileSizeY);
-	DispatchComputeShader(RHICmdList, *ComputeShader, GroupSizeX, GroupSizeY, 1);
+		AddDrawScreenPass(
+			GraphBuilder,
+			RDG_EVENT_NAME("Downsample.%s %dx%d (PS)", Inputs.Name, Inputs.Viewport.Width(), Inputs.Viewport.Height()),
+			ScreenPassView,
+			FScreenPassTextureViewport(OutputViewport, OutputTexture),
+			FScreenPassTextureViewport(Inputs.Viewport, Inputs.Texture),
+			*PixelShader,
+			PassParameters);
+	}
 
-	ComputeShader->UnsetParameters(RHICmdList);
+	FDownsamplePassOutputs Outputs;
+	Outputs.Texture = OutputTexture;
+	Outputs.Viewport = OutputViewport;
+	return Outputs;
 }
 
-FPooledRenderTargetDesc FRCPassPostProcessDownsample::ComputeOutputDesc(EPassOutputId InPassOutputId) const
+FRenderingCompositeOutputRef AddDownsamplePass(
+	FRenderingCompositionGraph& Graph,
+	const TCHAR *InName,
+	FRenderingCompositeOutputRef Input,
+	EDownsampleQuality InQuality,
+	EDownsampleFlags InFlags,
+	EPixelFormat InFormatOverride)
 {
-	FPooledRenderTargetDesc Ret = GetInput(ePId_Input0)->GetOutput()->RenderTargetDesc;
-
-	Ret.Reset();
-
-	Ret.Extent  = FIntPoint::DivideAndRoundUp(Ret.Extent, 2);
+	FRenderingCompositePass* DownsamplePass = Graph.RegisterPass(
+		new(FMemStack::Get()) TRCPassForRDG<1, 1>(
+			[InFormatOverride, InQuality, InFlags, InName] (FRenderingCompositePass* Pass, FRenderingCompositePassContext& InContext)
+	{
+		FRDGBuilder GraphBuilder(InContext.RHICmdList);
 
-	Ret.Extent.X = FMath::Max(1, Ret.Extent.X);
-	Ret.Extent.Y = FMath::Max(1, Ret.Extent.Y);
+		FRDGTextureRef InputTexture = Pass->CreateRDGTextureForRequiredInput(GraphBuilder, ePId_Input0, TEXT("DownsampleInput"));
 
-	if(OverrideFormat != PF_Unknown)
-	{
-		Ret.Format = OverrideFormat;
-	}
+		FDownsamplePassInputs PassInputs;
+		PassInputs.Name = InName;
+		PassInputs.Texture = InputTexture;
+		PassInputs.Viewport = InContext.SceneColorViewRect;
+		PassInputs.FormatOverride = InFormatOverride;
+		PassInputs.Quality = InQuality;
+		PassInputs.Flags = InFlags;
 
-	Ret.TargetableFlags &= ~(TexCreate_RenderTargetable | TexCreate_UAV);
-	Ret.TargetableFlags |= bIsComputePass ? TexCreate_UAV : TexCreate_RenderTargetable;
-	Ret.Flags |= GFastVRamConfig.Downsample;
-	Ret.AutoWritable = false;
-	Ret.DebugName = DebugName;
+		FDownsamplePassOutputs PassOutputs = AddDownsamplePass(GraphBuilder, FScreenPassViewInfo(InContext.View), PassInputs);
 
-	Ret.ClearValue = FClearValueBinding(FLinearColor(0,0,0,0));
+		Pass->ExtractRDGTextureForOutput(GraphBuilder, ePId_Output0, PassOutputs.Texture);
 
-	return Ret;
-}
+		GraphBuilder.Execute();
+	}));
+	DownsamplePass->SetInput(ePId_Input0, Input);
+	return DownsamplePass;
+}
\ No newline at end of file
diff --git a/Engine/Source/Runtime/Renderer/Private/PostProcess/PostProcessDownsample.h b/Engine/Source/Runtime/Renderer/Private/PostProcess/PostProcessDownsample.h
--- a/Engine/Source/Runtime/Renderer/Private/PostProcess/PostProcessDownsample.h
+++ b/Engine/Source/Runtime/Renderer/Private/PostProcess/PostProcessDownsample.h
@@ -6,44 +6,73 @@
 
 #pragma once
 
-#include "CoreMinimal.h"
-#include "RendererInterface.h"
 #include "PostProcess/RenderingCompositionGraph.h"
+#include "ScreenPass.h"
 
-// derives from TRenderingCompositePassBase<InputCount, OutputCount>
-// ePId_Input0: Color input
-// ePId_Input1: optional depth input (then quality is ignores and it uses the a 4 sample unfiltered samples method)
-class FRCPassPostProcessDownsample : public TRenderingCompositePassBase<2, 1>
+enum class EDownsampleFlags : uint8
 {
-public:
-	// constructor
-	// @param InDebugName we store the pointer so don't release this string
-	// @param Quality only used if ePId_Input1 is not set, 0:one filtered sample, 1:four filtered samples
-	FRCPassPostProcessDownsample(EPixelFormat InOverrideFormat = PF_Unknown,
-			uint32 InQuality = 1,
-			bool bInIsComputePass = false, 
-			const TCHAR *InDebugName = TEXT("Downsample"));
+	None = 0,
 
-	// interface FRenderingCompositePass ---------
+	// Forces the downsample pass to run on the raster pipeline, regardless of view settings.
+	ForceRaster = 0x1
+};
+ENUM_CLASS_FLAGS(EDownsampleFlags);
+
+enum class EDownsampleQuality : uint8
+{
+	// Single filtered sample (2x2 tap).
+	Low,
+
+	// Four filtered samples (4x4 tap).
+	High,
+
+	MAX
+};
+
+// Returns the global downsample quality specified by the r.Downsample.Quality CVar.
+EDownsampleQuality GetDownsampleQuality();
+
+// The set of inputs needed to add a downsample pass to RDG.
+struct FDownsamplePassInputs
+{
+	FDownsamplePassInputs() = default;
+
+	// Friendly name of the pass. Used for logging and profiling.
+	const TCHAR* Name = nullptr;
+
+	// Input RDG texture. Must not be null.
+	FRDGTextureRef Texture = nullptr;
+
+	// Input viewport to sample from.
+	FIntRect Viewport;
 
-	virtual void Process(FRenderingCompositePassContext& Context) override;
-	virtual void Release() override { delete this; }
-	virtual FPooledRenderTargetDesc ComputeOutputDesc(EPassOutputId InPassOutputId) const override;
+	// The downsample method to use.
+	EDownsampleQuality Quality = EDownsampleQuality::Low;
 
-	virtual FRHIComputeFence* GetComputePassEndFence() const override { return AsyncEndFence; }
+	// Flags to control how the downsample pass is run.
+	EDownsampleFlags Flags = EDownsampleFlags::None;
 
-private:
-	template <uint32 Method, uint32 ManuallyClampUV>
-	void SetShader(const FRenderingCompositePassContext& Context, const FPooledRenderTargetDesc* InputDesc, const FIntPoint& SrcSize, const FIntRect& SrcRect);
+	// The format to use for the output texture (if unknown, the input format is used).
+	EPixelFormat FormatOverride = PF_Unknown;
+};
 
-	template <uint32 Method, typename TRHICmdList>
-	void DispatchCS(TRHICmdList& RHICmdList, FRenderingCompositePassContext& Context, const FIntPoint& SrcSize, const FIntRect& DestRect, FRHIUnorderedAccessView* DestUAV);
+struct FDownsamplePassOutputs
+{
+	FDownsamplePassOutputs() = default;
 
-	FComputeFenceRHIRef AsyncEndFence;
+	// Half-resolution texture.
+	FRDGTextureRef Texture = nullptr;
 
-	EPixelFormat OverrideFormat;
-	// explained in constructor
-	uint32 Quality;
-	// must be a valid pointer
-	const TCHAR* DebugName;
+	// Half-resolution viewport.
+	FIntRect Viewport;
 };
+
+FDownsamplePassOutputs AddDownsamplePass(FRDGBuilder& GraphBuilder, const FScreenPassViewInfo& ScreenPassView, const FDownsamplePassInputs& Inputs);
+
+FRenderingCompositeOutputRef AddDownsamplePass(
+	FRenderingCompositionGraph& Graph,
+	const TCHAR *Name,
+	FRenderingCompositeOutputRef Input,
+	EDownsampleQuality Quality = EDownsampleQuality::Low,
+	EDownsampleFlags Flags = EDownsampleFlags::None,
+	EPixelFormat FormatOverride = PF_Unknown);
\ No newline at end of file
diff --git a/Engine/Source/Runtime/Renderer/Private/PostProcess/PostProcessMaterial.cpp b/Engine/Source/Runtime/Renderer/Private/PostProcess/PostProcessMaterial.cpp
--- a/Engine/Source/Runtime/Renderer/Private/PostProcess/PostProcessMaterial.cpp
+++ b/Engine/Source/Runtime/Renderer/Private/PostProcess/PostProcessMaterial.cpp
@@ -519,19 +519,28 @@ FRDGTextureRef ComputePostProcessMaterial(
 
 	ScreenPassDraw.Validate();
 
+	const bool bNeedsGBuffer = Material->NeedsGBuffer();
+
+	if (bNeedsGBuffer)
+	{
+		FSceneRenderTargets::Get(GraphBuilder.RHICmdList).AdjustGBufferRefCount(GraphBuilder.RHICmdList, 1);
+	}
+
 	GraphBuilder.AddPass(
 		RDG_EVENT_NAME("PostProcessMaterial"),
 		PostProcessMaterialParameters,
 		ERDGPassFlags::Raster,
-		[ScreenPassView, SceneColorViewport, ScreenPassDraw, PostProcessMaterialParameters, DepthStencilTexture, DepthStencilTextureForSRV, SetupFunction, bFlipYAxis]
-		(FRHICommandListImmediate& RHICmdList)
+		[ScreenPassView, SceneColorViewport, ScreenPassDraw, PostProcessMaterialParameters, DepthStencilTexture, DepthStencilTextureForSRV, SetupFunction, bFlipYAxis, bNeedsGBuffer]
+	(FRHICommandListImmediate& RHICmdList)
 	{
+		FSceneRenderTargets& SceneRenderTargets = FSceneRenderTargets::Get(RHICmdList);
+
 		const bool bUsesDepthStencilCopy = DepthStencilTexture != DepthStencilTextureForSRV;
 
 		if (bUsesDepthStencilCopy)
 		{
 			// Swap in the new copied SRV just prior to rendering.
-			FSceneRenderTargets::Get(RHICmdList).CustomDepth = DepthStencilTextureForSRV->GetPooledRenderTarget();
+			SceneRenderTargets.CustomDepth = DepthStencilTextureForSRV->GetPooledRenderTarget();
 		}
 
 		DrawScreenPass<decltype(SetupFunction)>(
@@ -545,13 +554,129 @@ FRDGTextureRef ComputePostProcessMaterial(
 		if (bUsesDepthStencilCopy)
 		{
 			// Swap back to the up-to-date target.
-			FSceneRenderTargets::Get(RHICmdList).CustomDepth = DepthStencilTexture->GetPooledRenderTarget();
+			SceneRenderTargets.CustomDepth = DepthStencilTexture->GetPooledRenderTarget();
+		}
+
+		if (bNeedsGBuffer)
+		{
+			SceneRenderTargets.AdjustGBufferRefCount(RHICmdList, -1);
 		}
 	});
 
 	return OutputTexture;
 }
 
+static bool IsPostProcessMaterialsEnabledForView(const FViewInfo& View)
+{
+	if (!View.Family->EngineShowFlags.PostProcessing ||
+		!View.Family->EngineShowFlags.PostProcessMaterial ||
+		View.Family->EngineShowFlags.VisualizeShadingModels ||
+		CVarPostProcessingDisableMaterials.GetValueOnRenderThread() != 0)
+	{
+		return false;
+	}
+
+	return true;
+}
+
+static FPostProcessMaterialNode* IteratePostProcessMaterialNodes(const FFinalPostProcessSettings& Dest, EBlendableLocation Location, FBlendableEntry*& Iterator)
+{
+	for (;;)
+	{
+		FPostProcessMaterialNode* DataPtr = Dest.BlendableManager.IterateBlendables<FPostProcessMaterialNode>(Iterator);
+
+		if (!DataPtr || DataPtr->GetLocation() == Location)
+		{
+			return DataPtr;
+		}
+	}
+}
+
+FRDGTextureRef AddPostProcessMaterialChain(
+	FRDGBuilder& GraphBuilder,
+	const FScreenPassViewInfo& ScreenPassView,
+	const FPostProcessMaterialInputs& Inputs,
+	EBlendableLocation Location)
+{
+	const FViewInfo& View = ScreenPassView.View;
+
+	if (!IsPostProcessMaterialsEnabledForView(View))
+	{
+		FIntRect Viewport;
+		FRDGTextureRef Texture;
+		Inputs.GetInput(EPostProcessMaterialInput::SceneColor, Texture, Viewport);
+		return Texture;
+	}
+
+	const FSceneViewFamily& ViewFamily = *View.Family;
+
+	// hard coded - this should be a reasonable limit
+	const uint32 MAX_PPMATERIALNODES = 10;
+	FBlendableEntry* Iterator = nullptr;
+	FPostProcessMaterialNode Nodes[MAX_PPMATERIALNODES];
+	uint32 NodeCount = 0;
+	bool bVisualizingBuffer = false;
+
+	if (ViewFamily.EngineShowFlags.VisualizeBuffer)
+	{
+		// Apply requested material to the full screen
+		UMaterial* Material = GetBufferVisualizationData().GetMaterial(View.CurrentBufferVisualizationMode);
+
+		if (Material && Material->BlendableLocation == Location)
+		{
+			Nodes[0] = FPostProcessMaterialNode(Material, Location, Material->BlendablePriority);
+			++NodeCount;
+			bVisualizingBuffer = true;
+		}
+	}
+	for (; NodeCount < MAX_PPMATERIALNODES; ++NodeCount)
+	{
+		FPostProcessMaterialNode* Data = IteratePostProcessMaterialNodes(View.FinalPostProcessSettings, Location, Iterator);
+
+		if (!Data)
+		{
+			break;
+		}
+
+		check(Data->GetMaterialInterface());
+
+		Nodes[NodeCount] = *Data;
+	}
+
+	::Sort(Nodes, NodeCount, FPostProcessMaterialNode::FCompare());
+
+	const bool bBasePassCanOutputVelocity = FVelocityRendering::BasePassCanOutputVelocity(View.GetFeatureLevel());
+
+	FRDGTextureRef LastSceneColor = nullptr;
+	FIntRect SceneColorViewport;
+	Inputs.GetInput(EPostProcessMaterialInput::SceneColor, LastSceneColor, SceneColorViewport);
+
+	for (uint32 i = 0; i < NodeCount; ++i)
+	{
+		const UMaterialInterface* MaterialInterface = Nodes[i].GetMaterialInterface();
+
+		FPostProcessMaterialInputs LocalInputs = Inputs;
+		LocalInputs.SetInput(EPostProcessMaterialInput::SceneColor, LastSceneColor, SceneColorViewport);
+
+		// This input is only needed for visualization and frame dumping
+		if (!bVisualizingBuffer)
+		{
+			LocalInputs.SetInput(EPostProcessMaterialInput::PreTonemapHDRColor, nullptr, FIntRect());
+			LocalInputs.SetInput(EPostProcessMaterialInput::PostTonemapHDRColor, nullptr, FIntRect());
+		}
+
+		// Velocity is only available when not generated from the base pass.
+		if (bBasePassCanOutputVelocity)
+		{
+			LocalInputs.SetInput(EPostProcessMaterialInput::Velocity, nullptr, FIntRect());
+		}
+
+		LastSceneColor = ComputePostProcessMaterial(GraphBuilder, ScreenPassView, MaterialInterface, Inputs);
+	}
+
+	return LastSceneColor;
+}
+
 FRenderingCompositePass* AddPostProcessMaterialPass(
 	const FPostprocessContext& PostProcessContext,
 	const UMaterialInterface* MaterialInterface,
@@ -617,19 +742,6 @@ FRenderingCompositePass* AddPostProcessMaterialPass(
 	}));
 }
 
-static FPostProcessMaterialNode* IteratePostProcessMaterialNodes(const FFinalPostProcessSettings& Dest, EBlendableLocation Location, FBlendableEntry*& Iterator)
-{
-	for (;;)
-	{
-		FPostProcessMaterialNode* DataPtr = Dest.BlendableManager.IterateBlendables<FPostProcessMaterialNode>(Iterator);
-
-		if (!DataPtr || DataPtr->GetLocation() == Location)
-		{
-			return DataPtr;
-		}
-	}
-}
-
 FRenderingCompositeOutputRef AddPostProcessMaterialReplaceTonemapPass(
 	FPostprocessContext& Context,
 	FRenderingCompositeOutputRef SeparateTranslucency,
@@ -681,10 +793,7 @@ FRenderingCompositeOutputRef AddPostProcessMaterialChain(
 	FRenderingCompositeOutputRef PostTonemapHDRColor,
 	FRenderingCompositeOutputRef PreFlattenVelocity)
 {
-	if (!Context.View.Family->EngineShowFlags.PostProcessing ||
-		!Context.View.Family->EngineShowFlags.PostProcessMaterial ||
-		Context.View.Family->EngineShowFlags.VisualizeShadingModels ||
-		CVarPostProcessingDisableMaterials.GetValueOnRenderThread() != 0)
+	if (!IsPostProcessMaterialsEnabledForView(Context.View))
 	{
 		return Context.FinalOutput;
 	}
diff --git a/Engine/Source/Runtime/Renderer/Private/PostProcess/PostProcessMaterial.h b/Engine/Source/Runtime/Renderer/Private/PostProcess/PostProcessMaterial.h
--- a/Engine/Source/Runtime/Renderer/Private/PostProcess/PostProcessMaterial.h
+++ b/Engine/Source/Runtime/Renderer/Private/PostProcess/PostProcessMaterial.h
@@ -97,6 +97,12 @@ struct FPostProcessMaterialInputs
 	bool bFlipYAxis = false;
 };
 
+FRDGTextureRef AddPostProcessMaterialChain(
+	FRDGBuilder& GraphBuilder,
+	const FScreenPassViewInfo& ScreenPassView,
+	const FPostProcessMaterialInputs& Inputs,
+	EBlendableLocation Location);
+
 FRDGTextureRef ComputePostProcessMaterial(
 	FRDGBuilder& GraphBuilder,
 	const FScreenPassViewInfo& ScreenPassView,
diff --git a/Engine/Source/Runtime/Renderer/Private/PostProcess/PostProcessMitchellNetravali.cpp b/Engine/Source/Runtime/Renderer/Private/PostProcess/PostProcessMitchellNetravali.cpp
--- a/Engine/Source/Runtime/Renderer/Private/PostProcess/PostProcessMitchellNetravali.cpp
+++ b/Engine/Source/Runtime/Renderer/Private/PostProcess/PostProcessMitchellNetravali.cpp
@@ -5,18 +5,10 @@
 =============================================================================*/
 
 #include "PostProcess/PostProcessMitchellNetravali.h"
-#include "StaticBoundShaderState.h"
 #include "SceneUtils.h"
-#include "PostProcess/SceneRenderTargets.h"
-#include "SceneRenderTargetParameters.h"
 #include "SceneRendering.h"
 #include "ScenePrivate.h"
-#include "PostProcess/SceneFilterRendering.h"
-#include "CompositionLighting/PostProcessAmbientOcclusion.h"
-#include "PostProcess/PostProcessTonemap.h"
-#include "ClearQuad.h"
-#include "PipelineStateCache.h"
-#include "PostProcessing.h"
+#include "SceneTextureParameters.h"
 
 class FMitchellNetravaliDownsampleCS : public FGlobalShader
 {
@@ -50,17 +42,6 @@ FRDGTextureRef ComputeMitchellNetravaliDownsample(
 	const FIntRect InputViewport,
 	const FScreenPassTextureViewport OutputViewport)
 {
-	FRDGTextureRef EyeAdaptationTexture = nullptr;
-
-	if (ScreenPassView.View.HasValidEyeAdaptation())
-	{
-		EyeAdaptationTexture = GraphBuilder.RegisterExternalTexture(ScreenPassView.View.GetEyeAdaptation(), TEXT("EyeAdaptation"));
-	}
-	else
-	{
-		EyeAdaptationTexture = GraphBuilder.RegisterExternalTexture(GSystemTextures.WhiteDummy, TEXT("EyeAdaptation"));
-	}
-
 	const FRDGTextureDesc OutputTextureDesc = FRDGTextureDesc::Create2DDesc(
 		OutputViewport.Extent,
 		PF_FloatRGBA,
@@ -77,7 +58,7 @@ FRDGTextureRef ComputeMitchellNetravaliDownsample(
 	PassParameters->InputTexture = InputTexture;
 	PassParameters->InputSampler = TStaticSamplerState<SF_Bilinear, AM_Clamp, AM_Clamp, AM_Clamp>::GetRHI();
 	PassParameters->OutputTexture = GraphBuilder.CreateUAV(OutputTexture);
-	PassParameters->EyeAdaptation = EyeAdaptationTexture;
+	PassParameters->EyeAdaptation = GetEyeAdaptationTexture(GraphBuilder, ScreenPassView.View);
 
 	// Scale / Bias factor to map the dispatch thread id to the input texture UV.
 	PassParameters->DispatchThreadToInputUVScale.X = InputViewport.Width()  / float(OutputViewport.Rect.Width()  * InputTexture->Desc.Extent.X);
@@ -95,32 +76,4 @@ FRDGTextureRef ComputeMitchellNetravaliDownsample(
 		FComputeShaderUtils::GetGroupCount(OutputViewport.Rect.Size(), FComputeShaderUtils::kGolden2DGroupSize));
 
 	return OutputTexture;
-}
-
-FRenderingCompositeOutputRef ComputeMitchellNetravaliDownsample(
-	FRenderingCompositionGraph& Graph,
-	FRenderingCompositeOutputRef Input,
-	FIntRect InputViewport,
-	FScreenPassTextureViewport OutputViewport)
-{
-	FRenderingCompositePass* MitchelNetravaliDownsamplePass = Graph.RegisterPass(new(FMemStack::Get()) TRCPassForRDG<1, 1>(
-		[InputViewport, OutputViewport](FRenderingCompositePass* Pass, FRenderingCompositePassContext& InContext)
-	{
-		FRDGBuilder GraphBuilder(InContext.RHICmdList);
-
-		FRDGTextureRef InputTexture = Pass->CreateRDGTextureForRequiredInput(GraphBuilder, ePId_Input0, TEXT("Input"));
-
-		FRDGTextureRef OutputTexture = ComputeMitchellNetravaliDownsample(GraphBuilder, FScreenPassViewInfo(InContext.View), InputTexture, InputViewport, OutputViewport);
-
-		Pass->ExtractRDGTextureForOutput(GraphBuilder, ePId_Output0, OutputTexture);
-
-		InContext.SceneColorViewRect = OutputViewport.Rect;
-		InContext.ReferenceBufferSize = OutputViewport.Extent;
-
-		GraphBuilder.Execute();
-	}));
-
-	MitchelNetravaliDownsamplePass->SetInput(ePId_Input0, Input);
-
-	return FRenderingCompositeOutputRef(MitchelNetravaliDownsamplePass);
 }
\ No newline at end of file
diff --git a/Engine/Source/Runtime/Renderer/Private/PostProcess/PostProcessMitchellNetravali.h b/Engine/Source/Runtime/Renderer/Private/PostProcess/PostProcessMitchellNetravali.h
--- a/Engine/Source/Runtime/Renderer/Private/PostProcess/PostProcessMitchellNetravali.h
+++ b/Engine/Source/Runtime/Renderer/Private/PostProcess/PostProcessMitchellNetravali.h
@@ -7,10 +7,10 @@
 #pragma once
 
 #include "ScreenPass.h"
-#include "RenderingCompositionGraph.h"
 
-FRenderingCompositeOutputRef ComputeMitchellNetravaliDownsample(
-	FRenderingCompositionGraph& Graph,
-	FRenderingCompositeOutputRef Input,
-	FIntRect InputViewport,
-	FScreenPassTextureViewport OutputViewport);
\ No newline at end of file
+FRDGTextureRef ComputeMitchellNetravaliDownsample(
+	FRDGBuilder& GraphBuilder,
+	const FScreenPassViewInfo& ScreenPassView,
+	FRDGTextureRef InputTexture,
+	const FIntRect InputViewport,
+	const FScreenPassTextureViewport OutputViewport);
\ No newline at end of file
diff --git a/Engine/Source/Runtime/Renderer/Private/PostProcess/PostProcessMotionBlur.cpp b/Engine/Source/Runtime/Renderer/Private/PostProcess/PostProcessMotionBlur.cpp
--- a/Engine/Source/Runtime/Renderer/Private/PostProcess/PostProcessMotionBlur.cpp
+++ b/Engine/Source/Runtime/Renderer/Private/PostProcess/PostProcessMotionBlur.cpp
@@ -8,7 +8,6 @@
 #include "SceneRenderTargetParameters.h"
 #include "ScenePrivate.h"
 #include "PostProcess/SceneFilterRendering.h"
-#include "PostProcess/SceneRenderTargets.h"
 #include "PostProcess/PostProcessing.h"
 #include "CompositionLighting/PostProcessAmbientOcclusion.h"
 #include "DeferredShadingRenderer.h"
@@ -109,6 +108,11 @@
 		&& (CVarAllowMotionBlurInVR->GetInt() != 0 || !(ViewFamily.Views.Num() > 1));
 }
 
+bool IsVisualizeMotionBlurEnabled(const FViewInfo& View)
+{
+	return View.Family->EngineShowFlags.VisualizeMotionBlur && View.GetFeatureLevel() >= ERHIFeatureLevel::SM4;
+}
+
 bool IsMotionBlurScatterRequired(const FViewInfo& View, const FScreenPassTextureViewport& SceneViewport)
 {
 	const FSceneViewState* ViewState = View.ViewState;
@@ -615,7 +619,7 @@
 	*VelocityTileTextureOutput = VelocityTileTexture;
 }
 
-FRDGTextureRef ComputeMotionBlurFilter(
+FRDGTextureRef AddMotionBlurFilterPass(
 	FRDGBuilder& GraphBuilder,
 	const FScreenPassViewInfo& ScreenPassView,
 	const FMotionBlurViewports& Viewports,
@@ -730,75 +734,6 @@
 	return ColorTextureOutput;
 }
 
-FRDGTextureRef ComputeMotionBlur(
-	FRDGBuilder& GraphBuilder,
-	const FScreenPassViewInfo& ScreenPassView,
-	FIntRect ColorViewportRect,
-	FIntRect VelocityViewportRect,
-	FRDGTextureRef ColorTexture,
-	FRDGTextureRef DepthTexture,
-	FRDGTextureRef VelocityTexture)
-{
-	check(ColorTexture);
-	check(DepthTexture);
-	check(VelocityTexture);
-
-	const FMotionBlurViewports Viewports(
-		FScreenPassTextureViewport(ColorViewportRect, ColorTexture),
-		FScreenPassTextureViewport(VelocityViewportRect, VelocityTexture));
-
-	FRDGTextureRef VelocityFlatTexture = nullptr;
-	FRDGTextureRef VelocityTileTexture = nullptr;
-	ComputeMotionBlurVelocity(
-		GraphBuilder,
-		ScreenPassView,
-		Viewports,
-		ColorTexture,
-		DepthTexture,
-		VelocityTexture,
-		&VelocityFlatTexture,
-		&VelocityTileTexture);
-
-	const EMotionBlurQuality MotionBlurQuality = GetMotionBlurQuality();
-
-	const bool bUseSeparableFilter = CVarMotionBlurSeparable.GetValueOnRenderThread() != 0;
-
-	if (bUseSeparableFilter)
-	{
-		FRDGTextureRef MotionBlurFilterTexture = ComputeMotionBlurFilter(
-			GraphBuilder,
-			ScreenPassView,
-			Viewports,
-			ColorTexture,
-			VelocityFlatTexture,
-			VelocityTileTexture,
-			EMotionBlurFilterPass::Separable0,
-			MotionBlurQuality);
-
-		return ComputeMotionBlurFilter(
-			GraphBuilder,
-			ScreenPassView,
-			Viewports,
-			MotionBlurFilterTexture,
-			VelocityFlatTexture,
-			VelocityTileTexture,
-			EMotionBlurFilterPass::Separable1,
-			MotionBlurQuality);
-	}
-	else
-	{
-		return ComputeMotionBlurFilter(
-			GraphBuilder,
-			ScreenPassView,
-			Viewports,
-			ColorTexture,
-			VelocityFlatTexture,
-			VelocityTileTexture,
-			EMotionBlurFilterPass::Unified,
-			MotionBlurQuality);
-	}
-}
-
 FRDGTextureRef VisualizeMotionBlur(
 	FRDGBuilder& GraphBuilder,
 	const FScreenPassViewInfo& ScreenPassView,
@@ -808,17 +743,17 @@
 	FRDGTextureRef DepthTexture,
 	FRDGTextureRef VelocityTexture)
 {
-	check(ColorTexture);
-	check(DepthTexture);
-	check(VelocityTexture);
-
 	const FViewInfo& View = ScreenPassView.View;
 
 	const FScreenPassTextureViewport ColorViewport(ColorViewportRect, ColorTexture);
 	const FScreenPassTextureViewport VelocityViewport(VelocityViewportRect, VelocityTexture);
 	const FMotionBlurViewports Viewports(ColorViewport, VelocityViewport);
 
-	FRDGTextureRef ColorTextureOutput = GraphBuilder.CreateTexture(ColorTexture->Desc, TEXT("MotionBlurVisualize"));
+	FRDGTextureDesc ColorTextureDesc = ColorTexture->Desc;
+	ColorTextureDesc.TargetableFlags &= ~(TexCreate_RenderTargetable | TexCreate_UAV);
+	ColorTextureDesc.TargetableFlags |= TexCreate_RenderTargetable;
 
+	FRDGTextureRef ColorTextureOutput = GraphBuilder.CreateTexture(ColorTextureDesc, TEXT("MotionBlurVisualize"));
+
 	FMotionBlurVisualizePS::FParameters* PassParameters = GraphBuilder.AllocParameters<FMotionBlurVisualizePS::FParameters>();
 	PassParameters->WorldToClipPrev = GetPreviousWorldToClipMatrix(View);
@@ -895,64 +830,89 @@
 	return ColorTextureOutput;
 }
 
-//////////////////////////////////////////////////////////////////////////
-//! Shim methods to hook into the legacy pipeline until the full RDG conversion is complete.
-
-FRenderingCompositeOutputRef ComputeMotionBlurShim(
-	FRenderingCompositionGraph& Graph,
-	FRenderingCompositeOutputRef ColorInput,
-	FRenderingCompositeOutputRef DepthInput,
-	FRenderingCompositeOutputRef VelocityInput,
-	bool bVisualizeMotionBlur)
+FRDGTextureRef AddMotionBlurPass(
+	FRDGBuilder& GraphBuilder,
+	const FScreenPassViewInfo& ScreenPassView,
+	FIntRect ColorViewportRect,
+	FIntRect VelocityViewportRect,
+	FRDGTextureRef ColorTexture,
+	FRDGTextureRef DepthTexture,
+	FRDGTextureRef VelocityTexture)
 {
-	FRenderingCompositePass* MotionBlurPass = Graph.RegisterPass(new(FMemStack::Get()) TRCPassForRDG<3, 1>(
-		[bVisualizeMotionBlur](FRenderingCompositePass* Pass, FRenderingCompositePassContext& InContext)
+	check(ColorTexture);
+	check(DepthTexture);
+	check(VelocityTexture);
+
+	if (IsVisualizeMotionBlurEnabled(ScreenPassView.View))
 	{
-		FRDGBuilder GraphBuilder(InContext.RHICmdList);
+		return VisualizeMotionBlur(
+			GraphBuilder,
+			ScreenPassView,
+			ColorViewportRect,
+			VelocityViewportRect,
+			ColorTexture,
+			DepthTexture,
+			VelocityTexture);
+	}
+	else if (!IsMotionBlurEnabled(ScreenPassView.View))
+	{
+		return ColorTexture;
+	}
 
-		FRDGTextureRef ColorTexture = Pass->CreateRDGTextureForRequiredInput(GraphBuilder, ePId_Input0, TEXT("SceneColor"));
-		FRDGTextureRef DepthTexture = Pass->CreateRDGTextureForRequiredInput(GraphBuilder, ePId_Input1, TEXT("SceneDepth"));
-		FRDGTextureRef VelocityTexture = Pass->CreateRDGTextureForRequiredInput(GraphBuilder, ePId_Input2, TEXT("SceneVelocity"));
+	const FMotionBlurViewports Viewports(
+		FScreenPassTextureViewport(ColorViewportRect, ColorTexture),
+		FScreenPassTextureViewport(VelocityViewportRect, VelocityTexture));
 
-		const FIntRect ColorViewportRect(InContext.SceneColorViewRect);
-		const FIntRect VelocityViewportRect(InContext.View.ViewRect);
+	RDG_EVENT_SCOPE(GraphBuilder, "MotionBlur");
 
-		const FScreenPassViewInfo ScreenPassView(InContext.View);
+	FRDGTextureRef VelocityFlatTexture = nullptr;
+	FRDGTextureRef VelocityTileTexture = nullptr;
+	ComputeMotionBlurVelocity(
+		GraphBuilder,
+		ScreenPassView,
+		Viewports,
+		ColorTexture,
+		DepthTexture,
+		VelocityTexture,
+		&VelocityFlatTexture,
+		&VelocityTileTexture);
 
-		FRDGTextureRef ColorTextureOutput;
+	const EMotionBlurQuality MotionBlurQuality = GetMotionBlurQuality();
 
-		if (bVisualizeMotionBlur)
-		{
-			ColorTextureOutput = VisualizeMotionBlur(
-				GraphBuilder,
-				ScreenPassView,
-				ColorViewportRect,
-				VelocityViewportRect,
-				ColorTexture,
-				DepthTexture,
-				VelocityTexture);
-		}
-		else
-		{
-			ColorTextureOutput = ComputeMotionBlur(
-				GraphBuilder,
-				ScreenPassView,
-				ColorViewportRect,
-				VelocityViewportRect,
-				ColorTexture,
-				DepthTexture,
-				VelocityTexture);
-		}
+	const bool bUseSeparableFilter = CVarMotionBlurSeparable.GetValueOnRenderThread() != 0;
 
-		Pass->ExtractRDGTextureForOutput(GraphBuilder, ePId_Output0, ColorTextureOutput);
+	if (bUseSeparableFilter)
+	{
+		FRDGTextureRef MotionBlurFilterTexture = AddMotionBlurFilterPass(
+			GraphBuilder,
+			ScreenPassView,
+			Viewports,
+			ColorTexture,
+			VelocityFlatTexture,
+			VelocityTileTexture,
+			EMotionBlurFilterPass::Separable0,
+			MotionBlurQuality);
 
-		GraphBuilder.Execute();
-	}));
-
-	MotionBlurPass->SetInput(ePId_Input0, ColorInput);
-	MotionBlurPass->SetInput(ePId_Input1, DepthInput);
-	MotionBlurPass->SetInput(ePId_Input2, VelocityInput);
-	return FRenderingCompositeOutputRef(MotionBlurPass);
+		return AddMotionBlurFilterPass(
+			GraphBuilder,
+			ScreenPassView,
+			Viewports,
+			MotionBlurFilterTexture,
+			VelocityFlatTexture,
+			VelocityTileTexture,
+			EMotionBlurFilterPass::Separable1,
+			MotionBlurQuality);
+	}
+	else
+	{
+		return AddMotionBlurFilterPass(
+			GraphBuilder,
+			ScreenPassView,
+			Viewports,
+			ColorTexture,
+			VelocityFlatTexture,
+			VelocityTileTexture,
+			EMotionBlurFilterPass::Unified,
+			MotionBlurQuality);
+	}
 }
-
-//////////////////////////////////////////////////////////////////////////
diff --git a/Engine/Source/Runtime/Renderer/Private/PostProcess/PostProcessMotionBlur.h b/Engine/Source/Runtime/Renderer/Private/PostProcess/PostProcessMotionBlur.h
--- a/Engine/Source/Runtime/Renderer/Private/PostProcess/PostProcessMotionBlur.h
+++ b/Engine/Source/Runtime/Renderer/Private/PostProcess/PostProcessMotionBlur.h
@@ -7,11 +7,13 @@
 #pragma once
 
 #include "ScreenPass.h"
-#include "RenderingCompositionGraph.h"
 
 // Returns whether motion blur is enabled for the requested view.
 bool IsMotionBlurEnabled(const FViewInfo& View);
 
+// Returns whether visualization of motion blur is enabled for the requested view.
+bool IsVisualizeMotionBlurEnabled(const FViewInfo& View);
+
 // The quality level of the motion blur pass.
 enum class EMotionBlurQuality : uint32
 {
@@ -22,17 +24,14 @@ enum class EMotionBlurQuality : uint32
 	MAX
 };
 
-// Returns the motion blur quality level set by the 'r.MotionBlurQuality' cvar.
+// Returns the motion blur quality level set by the 'r.MotionBlurQuality' CVar.
 EMotionBlurQuality GetMotionBlurQuality();
 
-//////////////////////////////////////////////////////////////////////////
-//! Shim methods to hook into the legacy pipeline until the full RDG conversion is complete.
-
-FRenderingCompositeOutputRef ComputeMotionBlurShim(
-	FRenderingCompositionGraph& Graph,
-	FRenderingCompositeOutputRef ColorInput,
-	FRenderingCompositeOutputRef DepthInput,
-	FRenderingCompositeOutputRef VelocityInput,
-	bool bVisualizeMotionBlur);
-
-//////////////////////////////////////////////////////////////////////////
\ No newline at end of file
+FRDGTextureRef AddMotionBlurPass(
+	FRDGBuilder& GraphBuilder,
+	const FScreenPassViewInfo& ScreenPassView,
+	FIntRect ColorViewportRect,
+	FIntRect VelocityViewportRect,
+	FRDGTextureRef ColorTexture,
+	FRDGTextureRef DepthTexture,
+	FRDGTextureRef VelocityTexture);
\ No newline at end of file
diff --git a/Engine/Source/Runtime/Renderer/Private/PostProcess/PostProcessTemporalAA.cpp b/Engine/Source/Runtime/Renderer/Private/PostProcess/PostProcessTemporalAA.cpp
--- a/Engine/Source/Runtime/Renderer/Private/PostProcess/PostProcessTemporalAA.cpp
+++ b/Engine/Source/Runtime/Renderer/Private/PostProcess/PostProcessTemporalAA.cpp
@@ -57,6 +57,11 @@ static TAutoConsoleVariable<int32> CVarTemporalAAUpsampleFiltered(
 	TEXT("Use filtering to fetch color history during TamporalAA upsampling (see AA_FILTERED define in TAA shader). Disabling this makes TAAU faster, but lower quality. "),
 	ECVF_Scalability | ECVF_RenderThreadSafe);
 
+static TAutoConsoleVariable<float> CVarTemporalAAHistorySP(
+	TEXT("r.TemporalAA.HistoryScreenPercentage"),
+	100.0f,
+	TEXT("Size of temporal AA's history."),
+	ECVF_RenderThreadSafe);
 
 static float CatmullRom( float x )
 {
@@ -67,6 +72,16 @@ static float CatmullRom( float x )
 		return ( 1.5f * ax - 2.5f ) * ax*ax + 1.0f;
 }
 
+float GetTemporalAAHistoryUpscaleFactor(const FViewInfo& View)
+{
+	// We only support history upscale on PC with feature level SM5+
+	if (!IsPCPlatform(View.GetShaderPlatform()) || !IsFeatureLevelSupported(View.GetShaderPlatform(), ERHIFeatureLevel::SM5))
+	{
+		return 1.0f;
+	}
+
+	return FMath::Clamp(CVarTemporalAAHistorySP.GetValueOnRenderThread() / 100.0f, 1.0f, 2.0f);
+}
 
 // ---------------------------------------------------- Shader permutation dimensions
 
@@ -370,70 +385,75 @@ bool FTAAPassParameters::Validate() const
 	return true;
 }
 
-
-FTAAOutputs FTAAPassParameters::AddTemporalAAPass(
+FTAAOutputs AddTemporalAAPass(
 	FRDGBuilder& GraphBuilder,
 	const FSceneTextureParameters& SceneTextures,
 	const FViewInfo& View,
+	const FTAAPassParameters& Inputs,
 	const FTemporalAAHistory& InputHistory,
-	FTemporalAAHistory* OutputHistory) const
+	FTemporalAAHistory* OutputHistory)
 {
-	check(Validate());
+	check(Inputs.Validate());
 
 	// Number of render target in TAA history.
-	const int32 RenderTargetCount = IsDOFTAAConfig(Pass) && FPostProcessing::HasAlphaChannelSupport() ? 2 : 1;
+	const int32 RenderTargetCount = IsDOFTAAConfig(Inputs.Pass) && FPostProcessing::HasAlphaChannelSupport() ? 2 : 1;
 	
 	// Whether this is main TAA pass;
-	bool bIsMainPass = IsMainTAAConfig(Pass);
+	const bool bIsMainPass = IsMainTAAConfig(Inputs.Pass);
 	
 	// Whether to use camera cut shader permutation or not.
 	const bool bCameraCut = !InputHistory.IsValid() || View.bCameraCut;
 
-	FIntPoint OutputExtent = GetOutputExtent();
+	const FIntPoint OutputExtent = Inputs.GetOutputExtent();
 
 	// Src rectangle.
-	FIntRect SrcRect = InputViewRect;
-	FIntRect DestRect = OutputViewRect;
-	FIntRect PracticableSrcRect = FIntRect::DivideAndRoundUp(SrcRect, ResolutionDivisor);
-	FIntRect PracticableDestRect = FIntRect::DivideAndRoundUp(DestRect, ResolutionDivisor);
+	const FIntRect SrcRect = Inputs.InputViewRect;
+	const FIntRect DestRect = Inputs.OutputViewRect;
+	const FIntRect PracticableSrcRect = FIntRect::DivideAndRoundUp(SrcRect, Inputs.ResolutionDivisor);
+	const FIntRect PracticableDestRect = FIntRect::DivideAndRoundUp(DestRect, Inputs.ResolutionDivisor);
+
+	const uint32 PassIndex = static_cast<uint32>(Inputs.Pass);
 
 	// Name of the pass.
-	const TCHAR* PassName = kTAAPassNames[static_cast<int32>(Pass)];
+	const TCHAR* PassName = kTAAPassNames[PassIndex];
 
 	// Create outputs
 	FTAAOutputs Outputs;
+
 	{
 		FRDGTextureDesc SceneColorDesc = FRDGTextureDesc::Create2DDesc(
-			GetOutputExtent(),
+			OutputExtent,
 			PF_FloatRGBA,
 			FClearValueBinding::Black,
 			/* InFlags = */ TexCreate_None,
 			/* InTargetableFlags = */ TexCreate_ShaderResource | TexCreate_UAV,
 			/* bInForceSeparateTargetAndShaderResource = */ false);
 
-		if (bOutputRenderTargetable)
+		if (Inputs.bOutputRenderTargetable)
 		{
 			SceneColorDesc.TargetableFlags |= TexCreate_RenderTargetable;
 		}
 
+		const TCHAR* OutputName = kTAAOutputNames[PassIndex];
+
 		Outputs.SceneColor = GraphBuilder.CreateTexture(
 			SceneColorDesc,
-			kTAAOutputNames[static_cast<int32>(Pass)],
+			OutputName,
 			ERDGResourceFlags::MultiFrame);
 
 		if (RenderTargetCount == 2)
 		{
 			Outputs.SceneMetadata = GraphBuilder.CreateTexture(
 				SceneColorDesc,
-				kTAAOutputNames[static_cast<int32>(Pass)],
+				OutputName,
 				ERDGResourceFlags::MultiFrame);
 		}
 
-		if (bDownsample)
+		if (Inputs.bDownsample)
 		{
-			FRDGTextureDesc HalfResSceneColorDesc = FRDGTextureDesc::Create2DDesc(
+			const FRDGTextureDesc HalfResSceneColorDesc = FRDGTextureDesc::Create2DDesc(
 				SceneColorDesc.Extent / 2,
-				DownsampleOverrideFormat != PF_Unknown ? DownsampleOverrideFormat : SceneColorInput->Desc.Format,
+				Inputs.DownsampleOverrideFormat != PF_Unknown ? Inputs.DownsampleOverrideFormat : Inputs.SceneColorInput->Desc.Format,
 				FClearValueBinding::Black,
 				/* InFlags = */ TexCreate_None,
 				/* InTargetableFlags = */ TexCreate_ShaderResource | TexCreate_Transient | TexCreate_UAV,
@@ -447,14 +467,14 @@ FTAAOutputs FTAAPassParameters::AddTemporalAAPass(
 
 	{
 		FTemporalAACS::FPermutationDomain PermutationVector;
-		PermutationVector.Set<FTAAPassConfigDim>(Pass);
-		PermutationVector.Set<FTAAFastDim>(bUseFast);
-		PermutationVector.Set<FTAADownsampleDim>(bDownsample);
+		PermutationVector.Set<FTAAPassConfigDim>(Inputs.Pass);
+		PermutationVector.Set<FTAAFastDim>(Inputs.bUseFast);
+		PermutationVector.Set<FTAADownsampleDim>(Inputs.bDownsample);
 		PermutationVector.Set<FTAAUpsampleFilteredDim>(true);
 
-		if (IsTAAUpsamplingConfig(Pass))
+		if (IsTAAUpsamplingConfig(Inputs.Pass))
 		{
-			const bool bUpsampleFiltered = CVarTemporalAAUpsampleFiltered.GetValueOnRenderThread() != 0 || Pass != ETAAPassConfig::MainUpsampling;
+			const bool bUpsampleFiltered = CVarTemporalAAUpsampleFiltered.GetValueOnRenderThread() != 0 || Inputs.Pass != ETAAPassConfig::MainUpsampling;
 			PermutationVector.Set<FTAAUpsampleFilteredDim>(bUpsampleFiltered);
 
 			// If screen percentage > 100% on X or Y axes, then use screen percentage range = 2 shader permutation to disable LDS caching.
@@ -466,7 +486,7 @@ FTAAOutputs FTAAPassParameters::AddTemporalAAPass(
 			// If screen percentage < 50% on X and Y axes, then use screen percentage range = 3 shader permutation.
 			else if (SrcRect.Width() * 100 < 50 * DestRect.Width() &&
 				SrcRect.Height() * 100 < 50 * DestRect.Height() &&
-				Pass == ETAAPassConfig::MainSuperSampling)
+				Inputs.Pass == ETAAPassConfig::MainSuperSampling)
 			{
 				PermutationVector.Set<FTAAScreenPercentageDim>(3);
 			}
@@ -481,13 +501,17 @@ FTAAOutputs FTAAPassParameters::AddTemporalAAPass(
 		FTemporalAACS::FParameters* PassParameters = GraphBuilder.AllocParameters<FTemporalAACS::FParameters>();
 
 		// Setups common shader parameters
-		FIntPoint SrcSize = SceneColorInput->Desc.Extent;
+		const FIntPoint InputExtent = Inputs.SceneColorInput->Desc.Extent;
+		const FIntRect InputViewRect = Inputs.InputViewRect;
+		const FIntRect OutputViewRect = Inputs.OutputViewRect;
 
-		if (!IsTAAUpsamplingConfig(Pass))
-			SetupSampleWeightParameters(PassParameters, *this, View.TemporalJitterPixels);
+		if (!IsTAAUpsamplingConfig(Inputs.Pass))
+		{
+			SetupSampleWeightParameters(PassParameters, Inputs, View.TemporalJitterPixels);
+		}
 
-		float ResDivisor = ResolutionDivisor;
-		float ResDivisorInv = 1.0f / ResDivisor;
+		const float ResDivisor = Inputs.ResolutionDivisor;
+		const float ResDivisorInv = 1.0f / ResDivisor;
 
 		PassParameters->ViewUniformBuffer = View.ViewUniformBuffer;
 		PassParameters->CurrentFrameWeight = CVarTemporalAACurrentFrameWeight.GetValueOnRenderThread();
@@ -497,16 +521,22 @@ FTAAOutputs FTAAPassParameters::AddTemporalAAPass(
 		PassParameters->SceneDepthBufferSampler = TStaticSamplerState<SF_Point>::GetRHI();
 		PassParameters->SceneVelocityBufferSampler = TStaticSamplerState<SF_Point>::GetRHI();
 
+		// We need a valid velocity buffer texture. Use black (no velocity) if none exists.
+		if (!PassParameters->SceneTextures.SceneVelocityBuffer)
+		{
+			PassParameters->SceneTextures.SceneVelocityBuffer = GraphBuilder.RegisterExternalTexture(GSystemTextures.BlackDummy);;
+		}
+
 		// Input buffer shader parameters
 		{
 			PassParameters->InputSceneColorSize = FVector4(
-				SceneColorInput->Desc.Extent.X,
-				SceneColorInput->Desc.Extent.Y,
-				1.0f / float(SceneColorInput->Desc.Extent.X),
-				1.0f / float(SceneColorInput->Desc.Extent.Y));
-			PassParameters->InputSceneColor = SceneColorInput;
+				InputExtent.X,
+				InputExtent.Y,
+				1.0f / float(InputExtent.X),
+				1.0f / float(InputExtent.Y));
+			PassParameters->InputSceneColor = Inputs.SceneColorInput;
 			PassParameters->InputSceneColorSampler = TStaticSamplerState<SF_Point>::GetRHI();
-			PassParameters->InputSceneMetadata = SceneMetadataInput;
+			PassParameters->InputSceneMetadata = Inputs.SceneMetadataInput;
 			PassParameters->InputSceneMetadataSampler = TStaticSamplerState<SF_Point>::GetRHI();
 		}
 
@@ -545,9 +575,9 @@ FTAAOutputs FTAAPassParameters::AddTemporalAAPass(
 					(ReferenceViewportExtent.X * 0.5f + ReferenceViewportOffset.X) * InvReferenceBufferSizeX,
 					(ReferenceViewportExtent.Y * 0.5f + ReferenceViewportOffset.Y) * InvReferenceBufferSizeY);
 
-				FIntPoint ViewportOffset = ReferenceViewportOffset / ResolutionDivisor;
-				FIntPoint ViewportExtent = FIntPoint::DivideAndRoundUp(ReferenceViewportExtent, ResolutionDivisor);
-				FIntPoint BufferSize = ReferenceBufferSize / ResolutionDivisor;
+				FIntPoint ViewportOffset = ReferenceViewportOffset / Inputs.ResolutionDivisor;
+				FIntPoint ViewportExtent = FIntPoint::DivideAndRoundUp(ReferenceViewportExtent, Inputs.ResolutionDivisor);
+				FIntPoint BufferSize = ReferenceBufferSize / Inputs.ResolutionDivisor;
 
 				PassParameters->ScreenPosAbsMax = FVector2D(1.0f - 1.0f / float(ViewportExtent.X), 1.0f - 1.0f / float(ViewportExtent.Y));
 
@@ -580,8 +610,8 @@ FTAAOutputs FTAAPassParameters::AddTemporalAAPass(
 		PassParameters->HistoryPreExposureCorrection = View.PreExposure / View.PrevViewInfo.SceneColorPreExposure;
 
 		{
-			float InvSizeX = 1.0f / float(SrcSize.X);
-			float InvSizeY = 1.0f / float(SrcSize.Y);
+			float InvSizeX = 1.0f / float(InputExtent.X);
+			float InvSizeY = 1.0f / float(InputExtent.Y);
 			PassParameters->ViewportUVToInputBufferUV = FVector4(
 				ResDivisorInv * InputViewRect.Width() * InvSizeX,
 				ResDivisorInv * InputViewRect.Height() * InvSizeY,
@@ -594,7 +624,7 @@ FTAAOutputs FTAAPassParameters::AddTemporalAAPass(
 		// Temporal upsample specific shader parameters.
 		{
 			// Temporal AA upscale specific params.
-			float InputViewSizeInvScale = ResolutionDivisor;
+			float InputViewSizeInvScale = Inputs.ResolutionDivisor;
 			float InputViewSizeScale = 1.0f / InputViewSizeInvScale;
 
 			PassParameters->TemporalJitterPixels = InputViewSizeScale * View.TemporalJitterPixels;
@@ -610,16 +640,20 @@ FTAAOutputs FTAAPassParameters::AddTemporalAAPass(
 		{
 			PassParameters->OutComputeTex0 = GraphBuilder.CreateUAV(Outputs.SceneColor);
 			if (Outputs.SceneMetadata)
+			{
 				PassParameters->OutComputeTex1 = GraphBuilder.CreateUAV(Outputs.SceneMetadata);
+			}
 			if (Outputs.DownsampledSceneColor)
+			{
 				PassParameters->OutComputeTexDownsampled = GraphBuilder.CreateUAV(Outputs.DownsampledSceneColor);
+			}
 		}
 
 		TShaderMapRef<FTemporalAACS> ComputeShader(View.ShaderMap, PermutationVector);
 		FComputeShaderUtils::AddPass(
 			GraphBuilder,
 			RDG_EVENT_NAME("TAA %s%s %dx%d -> %dx%d",
-				PassName, bUseFast ? TEXT(" Fast") : TEXT(""),
+				PassName, Inputs.bUseFast ? TEXT(" Fast") : TEXT(""),
 				PracticableSrcRect.Width(), PracticableSrcRect.Height(),
 				PracticableDestRect.Width(), PracticableDestRect.Height()),
 			*ComputeShader,
@@ -632,78 +666,103 @@ FTAAOutputs FTAAPassParameters::AddTemporalAAPass(
 		OutputHistory->SafeRelease();
 
 		GraphBuilder.QueueTextureExtraction(Outputs.SceneColor, &OutputHistory->RT[0]);
+
 		if (Outputs.SceneMetadata)
+		{
 			GraphBuilder.QueueTextureExtraction(Outputs.SceneMetadata, &OutputHistory->RT[1]);
+		}
 
 		OutputHistory->ViewportRect = DestRect;
-		OutputHistory->ReferenceBufferSize = OutputExtent * ResolutionDivisor;
+		OutputHistory->ReferenceBufferSize = OutputExtent * Inputs.ResolutionDivisor;
 	}
 
 	return Outputs;
-} // AddTemporalAAPass()
+}
 
+//////////////////////////////////////////////////////////////////////////
+//! Legacy - Only used by DebugViewModeRendering. Remove after porting.
 
-FRCPassPostProcessTemporalAA::FRCPassPostProcessTemporalAA(
-	const FPostprocessContext& Context,
-	const FTAAPassParameters& InParameters,
-	const FTemporalAAHistory& InInputHistory,
-	FTemporalAAHistory* OutOutputHistory)
-	: SavedParameters(InParameters)
-	, InputHistory(InInputHistory)
-	, OutputHistory(OutOutputHistory)
+class FRCPassPostProcessTemporalAA : public TRenderingCompositePassBase<3, 3>
 {
-	check(SavedParameters.Validate());
+public:
+	FRCPassPostProcessTemporalAA(
+		const FPostprocessContext& Context,
+		const FTAAPassParameters& InParameters,
+		const FTemporalAAHistory& InInputHistory,
+		FTemporalAAHistory* OutOutputHistory)
+		: SavedParameters(InParameters)
+		, InputHistory(InInputHistory)
+		, OutputHistory(OutOutputHistory)
+	{
+		check(InParameters.Pass == ETAAPassConfig::Main);
+		check(SavedParameters.Validate());
 
-	bIsComputePass = true;
-	bPreferAsyncCompute = false;
-}
+		bIsComputePass = true;
+		bPreferAsyncCompute = false;
+	}
 
-void FRCPassPostProcessTemporalAA::Process(FRenderingCompositePassContext& Context)
-{
-	WaitForInputPassComputeFences(Context.RHICmdList);
+	void Process(FRenderingCompositePassContext& Context) override
+	{
+		WaitForInputPassComputeFences(Context.RHICmdList);
 
-	AsyncEndFence = FComputeFenceRHIRef();
-	
-	FRDGBuilder GraphBuilder(Context.RHICmdList);
+		AsyncEndFence = FComputeFenceRHIRef();
 
-	FSceneRenderTargets& SceneContext = FSceneRenderTargets::Get(Context.RHICmdList);
-	FSceneTextureParameters SceneTextures;
-	SetupSceneTextureParameters(GraphBuilder, &SceneTextures);
+		FRDGBuilder GraphBuilder(Context.RHICmdList);
 
-	// FPostProcessing::Process() does a AdjustGBufferRefCount(RHICmdList, -1), therefore need to pass down reference on velocity buffer manually.
-	if (FRDGTextureRef SceneVelocityBuffer = CreateRDGTextureForOptionalInput(GraphBuilder, ePId_Input2, TEXT("SceneVelocity")))
-	{
-		SceneTextures.SceneVelocityBuffer = SceneVelocityBuffer;
-	}
+		FSceneTextureParameters SceneTextures;
+		SetupSceneTextureParameters(GraphBuilder, &SceneTextures);
 
-	FTAAPassParameters Parameters = SavedParameters;
-	Parameters.SceneColorInput = CreateRDGTextureForRequiredInput(GraphBuilder, ePId_Input0, TEXT("SceneColor"));
-	Parameters.SceneMetadataInput = CreateRDGTextureForOptionalInput(GraphBuilder, ePId_Input1, TEXT("SceneColor"));
+		FTAAPassParameters Parameters = SavedParameters;
+		Parameters.SceneColorInput = CreateRDGTextureForRequiredInput(GraphBuilder, ePId_Input0, TEXT("SceneColor"));
 
-	FTAAOutputs Outputs = Parameters.AddTemporalAAPass(
-		GraphBuilder,
-		SceneTextures, Context.View,
-		InputHistory, /* out */ OutputHistory);
-		
-	ExtractRDGTextureForOutput(GraphBuilder, ePId_Output0, Outputs.SceneColor);
-	if (Outputs.SceneMetadata)
-		ExtractRDGTextureForOutput(GraphBuilder, ePId_Output1, Outputs.SceneMetadata);
-	if (Outputs.DownsampledSceneColor)
-		ExtractRDGTextureForOutput(GraphBuilder, ePId_Output2, Outputs.DownsampledSceneColor);
+		FTAAOutputs Outputs = AddTemporalAAPass(
+			GraphBuilder,
+			SceneTextures,
+			Context.View,
+			Parameters,
+			InputHistory,
+			/* out */ OutputHistory);
+
+		ExtractRDGTextureForOutput(GraphBuilder, ePId_Output0, Outputs.SceneColor);
+
+		GraphBuilder.Execute();
+	}
 
-	GraphBuilder.Execute();
+	void Release() override { delete this; }
 
-	// Changes the view rectangle of the scene color and reference buffer size when doing temporal upsample for the
-	// following passes to still work.
-	if (Parameters.Pass == ETAAPassConfig::MainUpsampling || Parameters.Pass == ETAAPassConfig::MainSuperSampling)
+	FPooledRenderTargetDesc ComputeOutputDesc(EPassOutputId InPassOutputId) const override
 	{
-		Context.SceneColorViewRect = Parameters.OutputViewRect;
-		Context.ReferenceBufferSize = Parameters.GetOutputExtent() * Parameters.ResolutionDivisor;
+		// ExtractRDGTextureForOutput() is doing this work for us already.
+		return FPooledRenderTargetDesc();
 	}
-}
 
-FPooledRenderTargetDesc FRCPassPostProcessTemporalAA::ComputeOutputDesc(EPassOutputId InPassOutputId) const
+	FRHIComputeFence* GetComputePassEndFence() const override { return AsyncEndFence; }
+
+private:
+	const FTAAPassParameters SavedParameters;
+
+	FComputeFenceRHIRef AsyncEndFence;
+
+	const FTemporalAAHistory& InputHistory;
+	FTemporalAAHistory* OutputHistory;
+};
+
+FRenderingCompositeOutputRef AddTemporalAADebugViewPass(FPostprocessContext& Context)
 {
-	// ExtractRDGTextureForOutput() is doing this work for us already.
-	return FPooledRenderTargetDesc();
+	check(Context.View.ViewState);
+
+	FSceneViewState* ViewState = Context.View.ViewState;
+
+	FTAAPassParameters Parameters(Context.View);
+
+	FRCPassPostProcessTemporalAA* TemporalAAPass = Context.Graph.RegisterPass(new(FMemStack::Get()) FRCPassPostProcessTemporalAA(
+		Context,
+		Parameters,
+		Context.View.PrevViewInfo.TemporalAAHistory,
+		&ViewState->PrevFrameViewInfo.TemporalAAHistory));
+
+	TemporalAAPass->SetInput(ePId_Input0, Context.FinalOutput);
+	return FRenderingCompositeOutputRef(TemporalAAPass);
 }
+
+//////////////////////////////////////////////////////////////////////////
\ No newline at end of file
diff --git a/Engine/Source/Runtime/Renderer/Private/PostProcess/PostProcessTemporalAA.h b/Engine/Source/Runtime/Renderer/Private/PostProcess/PostProcessTemporalAA.h
--- a/Engine/Source/Runtime/Renderer/Private/PostProcess/PostProcessTemporalAA.h
+++ b/Engine/Source/Runtime/Renderer/Private/PostProcess/PostProcessTemporalAA.h
@@ -6,20 +6,13 @@
 
 #pragma once
 
-#include "CoreMinimal.h"
+#include "SceneRendering.h"
 #include "RenderGraph.h"
 
-// TODO: kill these includes once FRCPassPostProcessTemporalAA is gone.
-#include "RendererInterface.h"
-#include "RenderingCompositionGraph.h"
-
-
-class FViewInfo;
 class FSceneTextureParameters;
 struct FTemporalAAHistory;
 
-
-/** Lists of TAA configurations. */
+/** List of TAA configurations. */
 enum class ETAAPassConfig
 {
 	// Permutations for main scene color TAA.
@@ -56,6 +49,7 @@ static FORCEINLINE bool IsDOFTAAConfig(ETAAPassConfig Pass)
 	return Pass == ETAAPassConfig::DiaphragmDOF || Pass == ETAAPassConfig::DiaphragmDOFUpsampling;
 }
 
+float GetTemporalAAHistoryUpscaleFactor(const FViewInfo& View);
 
 /** GPU Output of the TAA pass. */
 struct FTAAOutputs
@@ -142,40 +136,12 @@ struct FTAAPassParameters
 
 	/** Validate the settings of TAA, to make sure there is no issue. */
 	bool Validate() const;
-
-	
-	/** Apply a temporal AA pass. */
-	FTAAOutputs AddTemporalAAPass(
-		FRDGBuilder& GraphBuilder,
-		const FSceneTextureParameters& SceneTextures,
-		const FViewInfo& View,
-		const FTemporalAAHistory& InputHistory,
-		FTemporalAAHistory* OutputHistory) const;
-}; // struct FTAAPassParameters
-
-
-// DEPRECATED. Use FTAAPassParameters::AddTemporalAAPass() instead.
-class FRCPassPostProcessTemporalAA : public TRenderingCompositePassBase<3, 3>
-{
-public:
-	FRCPassPostProcessTemporalAA(
-		const class FPostprocessContext& Context,
-		const FTAAPassParameters& Parameters,
-		const FTemporalAAHistory& InInputHistory,
-		FTemporalAAHistory* OutOutputHistory);
-
-	// interface FRenderingCompositePass ---------
-	virtual void Process(FRenderingCompositePassContext& Context) override;
-	virtual void Release() override { delete this; }
-	virtual FPooledRenderTargetDesc ComputeOutputDesc(EPassOutputId InPassOutputId) const override;
-
-	virtual FRHIComputeFence* GetComputePassEndFence() const override { return AsyncEndFence; }
-
-private:
-	const FTAAPassParameters SavedParameters;
-
-	FComputeFenceRHIRef AsyncEndFence;
-
-	const FTemporalAAHistory& InputHistory;
-	FTemporalAAHistory* OutputHistory;
 };
+
+FTAAOutputs AddTemporalAAPass(
+	FRDGBuilder& GraphBuilder,
+	const FSceneTextureParameters& SceneTextures,
+	const FViewInfo& View,
+	const FTAAPassParameters& Inputs,
+	const FTemporalAAHistory& InputHistory,
+	FTemporalAAHistory* OutputHistory);
\ No newline at end of file
diff --git a/Engine/Source/Runtime/Renderer/Private/PostProcess/PostProcessing.cpp b/Engine/Source/Runtime/Renderer/Private/PostProcess/PostProcessing.cpp
--- a/Engine/Source/Runtime/Renderer/Private/PostProcess/PostProcessing.cpp
+++ b/Engine/Source/Runtime/Renderer/Private/PostProcess/PostProcessing.cpp
@@ -61,12 +61,6 @@
 /** The global center for all post processing activities. */
 FPostProcessing GPostProcessing;
 
-static TAutoConsoleVariable<int32> CVarUseMobileBloom(
-	TEXT("r.UseMobileBloom"),
-	0,
-	TEXT("HACK: Set to 1 to use mobile bloom."),
-	ECVF_Scalability | ECVF_RenderThreadSafe);
-
 static TAutoConsoleVariable<float> CVarDepthOfFieldNearBlurSizeThreshold(
 	TEXT("r.DepthOfField.NearBlurSizeThreshold"),
 	0.01f,
@@ -101,14 +95,6 @@ static TAutoConsoleVariable<int32> CVarUpscaleQuality(
 	TEXT(" 5: 36-tap Gaussian-filtered unsharp mask (very expensive, but good for extreme upsampling).\n"),
 	ECVF_Scalability | ECVF_RenderThreadSafe);
 
-static TAutoConsoleVariable<int32> CDownsampleQuality(
-	TEXT("r.Downsample.Quality"),
-	3,
-	TEXT("Defines the quality in which the Downsample passes. we might add more quality levels later.\n")
-	TEXT(" 0: low quality\n")
-	TEXT(">0: high quality (default: 3)\n"),
-	ECVF_Scalability | ECVF_RenderThreadSafe);
-
 static TAutoConsoleVariable<float> CVarBloomCross(
 	TEXT("r.Bloom.Cross"),
 	0.0f,
@@ -171,19 +157,12 @@ TAutoConsoleVariable<int32> CVarHalfResFFTBloom(
 	TEXT(" 1: Half-resolution convoltuion that excludes the center of the kernel.\n"),
 	ECVF_Scalability | ECVF_RenderThreadSafe);
 
-TAutoConsoleVariable<int32> CVarTemporalAAAllowDownsampling(
+static TAutoConsoleVariable<int32> CVarTemporalAAAllowDownsampling(
 	TEXT("r.TemporalAA.AllowDownsampling"),
 	1,
 	TEXT("Allows half-resolution color buffer to be produced during TAA. Only possible when motion blur is off and when using compute shaders for post processing."),
 	ECVF_RenderThreadSafe);
 
-static TAutoConsoleVariable<float> CVarTemporalAAHistorySP(
-	TEXT("r.TemporalAA.HistoryScreenPercentage"),
-	100.0f,
-	TEXT("Size of temporal AA's history."),
-	ECVF_RenderThreadSafe
-);
-
 // -------------------------------------------------------
 
 bool ShouldDoComputePostProcessing(const FViewInfo& View)
@@ -191,6 +170,14 @@ bool ShouldDoComputePostProcessing(const FViewInfo& View)
 	return CVarPostProcessingPreferCompute.GetValueOnRenderThread() && View.FeatureLevel >= ERHIFeatureLevel::SM5;
 }
 
+bool IsTemporalAASceneDownsampleAllowed(const FViewInfo& View)
+{
+	return (CVarTemporalAAAllowDownsampling.GetValueOnRenderThread() != 0)
+		&& !IsMotionBlurEnabled(View)
+		&& !IsVisualizeMotionBlurEnabled(View)
+		&& GetDownsampleQuality() == EDownsampleQuality::Low;
+}
+
 FPostprocessContext::FPostprocessContext(FRHICommandListImmediate& InRHICmdList, FRenderingCompositionGraph& InGraph, const FViewInfo& InView)
 : RHICmdList(InRHICmdList)
 	, Graph(InGraph)
@@ -230,15 +217,12 @@ public:
 		// The first down sample is the input
 		PostProcessDownsamples[0] = SourceDownsample;
 
-		const bool bIsComputePass = ShouldDoComputePostProcessing(Context.View);
+		const EDownsampleQuality DownsampleQuality = GetDownsampleQuality();
 
-		const int32 DownsampleQuality = FMath::Clamp(CDownsampleQuality.GetValueOnRenderThread(), 0, 1);
 		// Queue the down samples. 
 		for (int i = 1; i < DownSampleStages; i++)
 		{
-			FRenderingCompositePass* Pass = Context.Graph.RegisterPass(new(FMemStack::Get()) FRCPassPostProcessDownsample(PF_Unknown, DownsampleQuality, bIsComputePass, PassLabels[i]));
-			Pass->SetInput(ePId_Input0, PostProcessDownsamples[i - 1]);
-			PostProcessDownsamples[i] = FRenderingCompositeOutputRef(Pass);
+			PostProcessDownsamples[i] = AddDownsamplePass(Context.Graph, PassLabels[i], PostProcessDownsamples[i - 1], DownsampleQuality);
 
 			// Add log2 data to the alpha channel after doing the 1st (i==1) down sample pass
 			if (bHasLog2Alpha && i == 1 ) {
@@ -248,7 +232,7 @@ public:
 			}
 		}
 
-		// Calculate the final viewrect size (matching FRCPassPostProcessDownsample behavior)
+		// Calculate the final viewrect size (matching AddDownsamplePass behavior)
 		FinalViewRectSize.X = FMath::Max(1, FMath::DivideAndRoundUp(InContext.View.ViewRect.Width(), 1 << DownSampleStages));
 		FinalViewRectSize.Y = FMath::Max(1, FMath::DivideAndRoundUp(InContext.View.ViewRect.Height(), 1 << DownSampleStages));
 	}
@@ -455,15 +439,6 @@ static FRenderingCompositeOutputRef AddPostProcessHistogramEyeAdaptation(FPostpr
 	return FRenderingCompositeOutputRef(Node);
 }
 
-static void AddVisualizeBloomSetup(FPostprocessContext& Context)
-{
-	auto Node = Context.Graph.RegisterPass(new(FMemStack::Get()) FRCPassPostProcessVisualizeBloomSetup());
-
-	Node->SetInput(ePId_Input0, Context.FinalOutput);
-
-	Context.FinalOutput = FRenderingCompositeOutputRef(Node);
-}
-
 static void AddVisualizeBloomOverlay(FPostprocessContext& Context, FRenderingCompositeOutputRef& HDRColor, FRenderingCompositeOutputRef& BloomOutputCombined)
 {
 	auto Node = Context.Graph.RegisterPass(new(FMemStack::Get()) FRCPassPostProcessVisualizeBloomOverlay());
@@ -499,9 +474,7 @@ static bool AddPostProcessDepthOfFieldGaussian(FPostprocessContext& Context, FDe
 		FRenderingCompositeOutputRef SetupInput(Context.FinalOutput);
 		if (bMobileQuality)
 		{
-			FRenderingCompositePass* HalfResFar = Context.Graph.RegisterPass(new(FMemStack::Get()) FRCPassPostProcessDownsample(PF_FloatRGBA, 1, false, TEXT("GausSetupHalfRes")));
-			HalfResFar->SetInput(ePId_Input0, FRenderingCompositeOutputRef(SetupInput));
-			SetupInput = HalfResFar;
+			SetupInput = AddDownsamplePass(Context.Graph, TEXT("GaussianSetupHalfRes"), SetupInput, EDownsampleQuality::High, EDownsampleFlags::ForceRaster, PF_FloatRGBA);
 		}
 
 		FRenderingCompositePass* DOFSetupPass = Context.Graph.RegisterPass(new(FMemStack::Get()) FRCPassPostProcessDOFSetup(bFar, bNear));
@@ -745,48 +718,6 @@ static FRenderingCompositeOutputRef AddBloom(FBloomDownSampleArray& BloomDownSam
 	return BloomOutput;
 }
 
-static FTAAPassParameters MakeTAAPassParametersForView( const FViewInfo& View )
-{
-	FTAAPassParameters Parameters(View);
-
-	Parameters.Pass = View.PrimaryScreenPercentageMethod == EPrimaryScreenPercentageMethod::TemporalUpscale
-		? ETAAPassConfig::MainUpsampling
-		: ETAAPassConfig::Main;
-	
-	Parameters.SetupViewRect(View);
-
-	{
-		static const auto CVar = IConsoleManager::Get().FindTConsoleVariableDataInt(TEXT("r.PostProcessAAQuality"));
-		uint32 Quality = FMath::Clamp(CVar->GetValueOnRenderThread(), 1, 6);
-		Parameters.bUseFast = Quality == 3;
-	}
-
-	return Parameters;
-}
-
-static void AddTemporalAA( FPostprocessContext& Context, FRenderingCompositeOutputRef& VelocityInput, const FTAAPassParameters& Parameters, FRenderingCompositeOutputRef* OutSceneColorHalfRes )
-{
-	check(VelocityInput.IsValid());
-	check(Context.View.ViewState);
-
-	FSceneViewState* ViewState = Context.View.ViewState;
-
-	FRCPassPostProcessTemporalAA* TemporalAAPass = Context.Graph.RegisterPass(new(FMemStack::Get()) FRCPassPostProcessTemporalAA(
-		Context,
-		Parameters,
-		Context.View.PrevViewInfo.TemporalAAHistory,
-		&ViewState->PrevFrameViewInfo.TemporalAAHistory));
-
-	TemporalAAPass->SetInput( ePId_Input0, Context.FinalOutput );
-	TemporalAAPass->SetInput( ePId_Input2, VelocityInput );
-	Context.FinalOutput = FRenderingCompositeOutputRef( TemporalAAPass );
-
-	if (OutSceneColorHalfRes && Parameters.bDownsample)
-	{
-		*OutSceneColorHalfRes = FRenderingCompositeOutputRef(TemporalAAPass, ePId_Output2);
-	}
-}
-
 static void AddGBufferVisualizationOverview(
 	FPostprocessContext& Context,
 	FRenderingCompositeOutputRef SeparateTranslucencyInput,
@@ -870,12 +801,10 @@ static void AddGBufferVisualizationOverview(
 					if (bOverviewModeEnabled)
 					{
 						// Down-sample to 1/2 size
-						FRenderingCompositePass* HalfSize = Context.Graph.RegisterPass(new(FMemStack::Get()) FRCPassPostProcessDownsample(PF_Unknown, 0, false, TEXT("MaterialHalfSize")));
-						HalfSize->SetInput(ePId_Input0, FRenderingCompositeOutputRef(MaterialPass));
+						FRenderingCompositeOutputRef HalfSize = AddDownsamplePass(Context.Graph, TEXT("MaterialHalfSize"), MaterialPass, EDownsampleQuality::Low, EDownsampleFlags::ForceRaster);
 
 						// Down-sample to 1/4 size
-						FRenderingCompositePass* QuarterSize = Context.Graph.RegisterPass(new(FMemStack::Get()) FRCPassPostProcessDownsample(PF_Unknown, 0, false, TEXT("MaterialQuarterSize")));
-						QuarterSize->SetInput(ePId_Input0, FRenderingCompositeOutputRef(HalfSize));
+						FRenderingCompositeOutputRef QuarterSize = AddDownsamplePass(Context.Graph, TEXT("MaterialQuarterSize"), HalfSize, EDownsampleQuality::Low, EDownsampleFlags::ForceRaster);
 
 						// Set whether current buffer is selected
 						bool bIsSelected = false;
@@ -1018,7 +947,7 @@ FRDGTextureRef AddSeparateTranslucencyCompositionPass(FRDGBuilder& GraphBuilder,
 	PassParameters->SeparateTranslucencySampler = TStaticSamplerState<SF_Point>::GetRHI();
 	PassParameters->ViewUniformBuffer = View.ViewUniformBuffer;
 	PassParameters->RenderTargets[0] = FRenderTargetBinding(NewSceneColor, ERenderTargetLoadAction::ENoAction);
-			
+
 	TShaderMapRef<FComposeSeparateTranslucencyPS> PixelShader(View.ShaderMap);
 	FPixelShaderUtils::AddFullscreenPass(
 		GraphBuilder,
@@ -1031,6 +960,137 @@ FRDGTextureRef AddSeparateTranslucencyCompositionPass(FRDGBuilder& GraphBuilder,
 	return NewSceneColor;
 }
 
+bool IsDepthOfFieldEnabledForView(const FViewInfo& View)
+{
+	static const auto CVar = IConsoleManager::Get().FindTConsoleVariableDataInt(TEXT("r.DepthOfFieldQuality"));
+	check(CVar);
+	return
+		View.Family->EngineShowFlags.DepthOfField &&
+		CVar->GetValueOnRenderThread() > 0 &&
+		View.FinalPostProcessSettings.DepthOfFieldFstop > 0 &&
+		View.FinalPostProcessSettings.DepthOfFieldFocalDistance > 0;
+}
+
+EPixelFormat GetSceneColorHalfResFormat()
+{
+	return PF_FloatRGB;
+}
+
+void AddTemporalAAMainPass(
+	FRDGBuilder& GraphBuilder,
+	const FSceneTextureParameters& SceneTextures,
+	const FScreenPassViewInfo& ScreenPassView,
+	FRDGTextureRef InSceneColorTexture,
+	FRDGTextureRef InCustomDepthTexture,
+	FRDGTextureRef* OutSceneColorTexture,
+	FRDGTextureRef* OutSceneColorHalfResTexture,
+	FIntRect* OutSecondaryViewRect)
+{
+	const FViewInfo& View = ScreenPassView.View;
+
+	// Skip if TAA is disabled or if we don't have any persistent view state (required to track history).
+	if (View.AntiAliasingMethod != AAM_TemporalAA || !View.ViewState)
+	{
+		*OutSceneColorTexture = InSceneColorTexture;
+		*OutSceneColorHalfResTexture = nullptr;
+		*OutSecondaryViewRect = View.ViewRect;
+		return;
+	}
+
+	FPreviousViewInfo& PrevViewInfo = View.ViewState->PrevFrameViewInfo;
+
+	// Temporal AA passes.
+	FTAAPassParameters TAAParameters(View);
+
+	TAAParameters.Pass = View.PrimaryScreenPercentageMethod == EPrimaryScreenPercentageMethod::TemporalUpscale
+		? ETAAPassConfig::MainUpsampling
+		: ETAAPassConfig::Main;
+
+	TAAParameters.SetupViewRect(View);
+
+	{
+		static const auto CVar = IConsoleManager::Get().FindTConsoleVariableDataInt(TEXT("r.PostProcessAAQuality"));
+		uint32 Quality = FMath::Clamp(CVar->GetValueOnRenderThread(), 1, 6);
+		TAAParameters.bUseFast = Quality == 3;
+	}
+
+	const FIntRect SecondaryViewRect = TAAParameters.OutputViewRect;
+
+	const float HistoryUpscaleFactor = GetTemporalAAHistoryUpscaleFactor(View);
+
+	// Configures TAA to upscale the history buffer; this is in addition to the secondary screen percentage upscale.
+	// We end up with a scene color that is larger than the secondary screen percentage. We immediately downscale
+	// afterwards using a Mitchel-Netravali filter.
+	if (HistoryUpscaleFactor > 1.0f)
+	{
+		const FIntPoint HistoryViewSize(
+			TAAParameters.OutputViewRect.Width() * HistoryUpscaleFactor,
+			TAAParameters.OutputViewRect.Height() * HistoryUpscaleFactor);
+
+		FIntPoint QuantizedMinHistorySize;
+		QuantizeSceneBufferSize(HistoryViewSize, QuantizedMinHistorySize);
+
+		TAAParameters.Pass = ETAAPassConfig::MainSuperSampling;
+		TAAParameters.bUseFast = false;
+
+		TAAParameters.OutputViewRect.Min.X = 0;
+		TAAParameters.OutputViewRect.Min.Y = 0;
+		TAAParameters.OutputViewRect.Max = HistoryViewSize;
+	}
+
+	TAAParameters.DownsampleOverrideFormat = GetSceneColorHalfResFormat();
+
+	// Downsample pass may be merged with with TemporalAA when there is no motion blur and compute shader is used.
+	// This is currently only possible for r.Downsample.Quality = 0 (box filter).
+	TAAParameters.bDownsample = TAAParameters.bUseFast && IsTemporalAASceneDownsampleAllowed(View);
+
+	TAAParameters.SceneColorInput = InSceneColorTexture;
+
+	const FTAAOutputs TAAOutputs = AddTemporalAAPass(
+		GraphBuilder,
+		SceneTextures,
+		View,
+		TAAParameters,
+		View.PrevViewInfo.TemporalAAHistory,
+		&PrevViewInfo.TemporalAAHistory);
+
+	*OutSceneColorTexture = TAAOutputs.SceneColor;
+	*OutSceneColorHalfResTexture = TAAOutputs.DownsampledSceneColor;
+
+	// If we upscaled the history buffer, downsize back to the secondary screen percentage size.
+	if (HistoryUpscaleFactor > 1.0f)
+	{
+		const FIntRect InputViewport = TAAParameters.OutputViewRect;
+
+		FIntPoint QuantizedOutputSize;
+		QuantizeSceneBufferSize(SecondaryViewRect.Size(), QuantizedOutputSize);
+
+		FScreenPassTextureViewport OutputViewport;
+		OutputViewport.Rect = SecondaryViewRect;
+		OutputViewport.Extent.X = FMath::Max(InSceneColorTexture->Desc.Extent.X, QuantizedOutputSize.X);
+		OutputViewport.Extent.Y = FMath::Max(InSceneColorTexture->Desc.Extent.Y, QuantizedOutputSize.Y);
+
+		*OutSceneColorTexture = ComputeMitchellNetravaliDownsample(GraphBuilder, ScreenPassView, TAAOutputs.SceneColor, InputViewport, OutputViewport);
+	}
+
+	// Post Process Material Chain - SSR Input
+	if (!View.bViewStateIsReadOnly)
+	{
+		FPostProcessMaterialInputs Inputs;
+		Inputs.SetInput(EPostProcessMaterialInput::SceneColor, InSceneColorTexture, SecondaryViewRect);
+		Inputs.CustomDepthTexture = InCustomDepthTexture;
+
+		FRDGTextureRef SSRInputTexture = AddPostProcessMaterialChain(GraphBuilder, ScreenPassView, Inputs, BL_SSRInput);
+
+		if (SSRInputTexture != InSceneColorTexture)
+		{
+			GraphBuilder.QueueTextureExtraction(SSRInputTexture, &PrevViewInfo.CustomSSRInput);
+		}
+	}
+
+	*OutSecondaryViewRect = SecondaryViewRect;
+}
+
 } // namespace
 
 
@@ -1042,36 +1102,36 @@ void FPostProcessing::Process(FRHICommandListImmediate& RHICmdList, const FViewI
 	check(IsInRenderingThread());
 	check(View.VerifyMembersChecks());
 
-	const auto FeatureLevel = View.GetFeatureLevel();
+	const ERHIFeatureLevel::Type FeatureLevel = View.GetFeatureLevel();
 
 	GRenderTargetPool.AddPhaseEvent(TEXT("PostProcessing"));
 
-	// This page: https://udn.epicgames.com/Three/RenderingOverview#Rendering%20state%20defaults 
-	// describes what state a pass can expect and to what state it need to be set back.
-
 	// All post processing is happening on the render thread side. All passes can access FinalPostProcessSettings and all
 	// view settings. Those are copies for the RT then never get access by the main thread again.
 	// Pointers to other structures might be unsafe to touch.
 
-
 	// so that the passes can register themselves to the graph
 	{
 		FMemMark Mark(FMemStack::Get());
 		FRenderingCompositePassContext CompositeContext(RHICmdList, View);
 
 		FPostprocessContext Context(RHICmdList, CompositeContext.Graph, View);
-		
+
 		// not always valid
 		FRenderingCompositeOutputRef HistogramOverScreen;
 		FRenderingCompositeOutputRef Histogram;
 		FRenderingCompositeOutputRef PreTonemapHDRColor;
 		FRenderingCompositeOutputRef PostTonemapHDRColor;
 		FRenderingCompositeOutputRef PreFlattenVelocity;
+		FRenderingCompositeOutputRef SeparateTranslucency;
+		FRenderingCompositeOutputRef BloomOutputCombined;
+		FRenderingCompositeOutputRef CustomDepth;
+		FRenderingCompositeOutputRef SceneColorHalfRes;
 
 		class FAutoExposure
 		{
 		public:
-			FAutoExposure(const FViewInfo& InView) : 
+			FAutoExposure(const FViewInfo& InView) :
 				MethodId(GetAutoExposureMethod(InView))
 			{}
 			// distinguish between Basic and Histogram-based
@@ -1079,224 +1139,192 @@ void FPostProcessing::Process(FRHICommandListImmediate& RHICmdList, const FViewI
 			// not always valid
 			FRenderingCompositeOutputRef EyeAdaptation;
 		} AutoExposure(View);
-	
-		// not always valid
-		FRenderingCompositeOutputRef SeparateTranslucency;
-		// optional
-		FRenderingCompositeOutputRef BloomOutputCombined;
-		// not always valid
-		FRenderingCompositePass* VelocityFlattenPass = 0;
-		// in the following code some feature might set this to false
-		bool bAllowTonemapper = FeatureLevel >= ERHIFeatureLevel::SM4;
-		//
+
 		FRCPassPostProcessUpscale::PaniniParams PaniniConfig(View);
-		//
-		EStereoscopicPass StereoPass = View.StereoPass;
-		//
-		FSceneViewState* ViewState = (FSceneViewState*)Context.View.State;
+
+		const bool bDepthOfFieldEnabled = IsDepthOfFieldEnabledForView(View);
+
+		const bool bVisualizeHDR = View.Family->EngineShowFlags.VisualizeHDR && FeatureLevel >= ERHIFeatureLevel::SM5;
+		const bool bVisualizeBloom = View.Family->EngineShowFlags.VisualizeBloom && FeatureLevel >= ERHIFeatureLevel::SM4;
+		const bool bVisualizeMotionBlur = IsVisualizeMotionBlurEnabled(View);
+		const bool bVisualizeDepthOfField = bDepthOfFieldEnabled && Context.View.Family->EngineShowFlags.VisualizeDOF;
+
+		const bool bTonemapEnabled = FeatureLevel >= ERHIFeatureLevel::SM4 && !bVisualizeBloom && !bVisualizeMotionBlur;
+
+		const bool bHDROutputEnabled = GRHISupportsHDROutput && IsHDREnabled();
+
+		static const auto CVarDumpFramesAsHDR = IConsoleManager::Get().FindTConsoleVariableDataInt(TEXT("r.BufferVisualizationDumpFramesAsHDR"));
+		const bool bHDRTonemapperOutput = bTonemapEnabled && (View.Family->SceneCaptureSource == SCS_FinalColorHDR || GetHighResScreenshotConfig().bCaptureHDR || CVarDumpFramesAsHDR->GetValueOnRenderThread() || bHDROutputEnabled);
+
+		const EAntiAliasingMethod AntiAliasingMethod = View.AntiAliasingMethod;
+
+		FRCPassPostProcessTonemap* Tonemapper = nullptr;
 
 		FSceneRenderTargets& SceneContext = FSceneRenderTargets::Get(RHICmdList);
 
+		// Register textures as inputs into the composition graph.
+		if (SceneContext.SeparateTranslucencyRT)
 		{
-			if (FSceneRenderTargets::Get(RHICmdList).SeparateTranslucencyRT)
-			{
-				FRenderingCompositePass* NodeSeparateTranslucency = Context.Graph.RegisterPass(new(FMemStack::Get()) FRCPassPostProcessInput(FSceneRenderTargets::Get(RHICmdList).SeparateTranslucencyRT));
-				SeparateTranslucency = FRenderingCompositeOutputRef(NodeSeparateTranslucency);
+			SeparateTranslucency = Context.Graph.RegisterPass(new(FMemStack::Get()) FRCPassPostProcessInput(SceneContext.SeparateTranslucencyRT));
 
-				// make sure we only release if this is the last view we're rendering
-				int32 LastView = View.Family->Views.Num() - 1;
-				if (View.Family->Views[LastView] == &View)
-				{
-					// the node keeps another reference so the RT will not be release too early
-					FSceneRenderTargets::Get(RHICmdList).FreeSeparateTranslucency();
-					check(!FSceneRenderTargets::Get(RHICmdList).SeparateTranslucencyRT);
-				}
+			// make sure we only release if this is the last view we're rendering
+			if (View.IsLastInFamily())
+			{
+				// the node keeps another reference so the RT will not be release too early
+				SceneContext.FreeSeparateTranslucency();
 			}
 		}
 
-		bool bVisualizeHDR = View.Family->EngineShowFlags.VisualizeHDR && FeatureLevel >= ERHIFeatureLevel::SM5;
-		bool bVisualizeBloom = View.Family->EngineShowFlags.VisualizeBloom && FeatureLevel >= ERHIFeatureLevel::SM4;
-		bool bVisualizeMotionBlur = View.Family->EngineShowFlags.VisualizeMotionBlur && FeatureLevel >= ERHIFeatureLevel::SM4;
-
-		if(bVisualizeBloom || bVisualizeMotionBlur)
+		if (VelocityRT)
 		{
-			bAllowTonemapper = false;
+			PreFlattenVelocity = Context.Graph.RegisterPass(new(FMemStack::Get()) FRCPassPostProcessInput(VelocityRT));
 		}
 
-		const bool bHDROutputEnabled = GRHISupportsHDROutput && IsHDREnabled();
+		// add the passes we want to add to the graph (commenting a line means the pass is not inserted into the graph) ---------
 
-		static const auto CVarDumpFramesAsHDR = IConsoleManager::Get().FindTConsoleVariableDataInt(TEXT("r.BufferVisualizationDumpFramesAsHDR"));
-    	const bool bHDRTonemapperOutput = bAllowTonemapper && (View.Family->SceneCaptureSource == SCS_FinalColorHDR || GetHighResScreenshotConfig().bCaptureHDR || CVarDumpFramesAsHDR->GetValueOnRenderThread() || bHDROutputEnabled);
+		if (AllowFullPostProcessing(View, FeatureLevel))
+		{
+			// GBuffers are released prior to executing the composition graph. We take a reference here
+			// and then release the reference inside of RDGPass. This allows RDGPass to control lifetime
+			// of the GBuffers internally.
+			SceneContext.AdjustGBufferRefCount(RHICmdList, 1);
+
+			FRenderingCompositePass* RDGPass = Context.Graph.RegisterPass(
+				new(FMemStack::Get()) TRCPassForRDG<5, 2>(
+					[&View, &SceneContext, bVisualizeBloom, bDepthOfFieldEnabled]
+			(FRenderingCompositePass* Pass, FRenderingCompositePassContext& InContext)
+			{
+				FRDGBuilder GraphBuilder(InContext.RHICmdList);
 
+				FRDGTextureRef SceneColorTexture = Pass->CreateRDGTextureForRequiredInput(GraphBuilder, ePId_Input0, TEXT("SceneColor"));
+				FRDGTextureRef SceneDepthTexture = Pass->CreateRDGTextureForRequiredInput(GraphBuilder, ePId_Input1, TEXT("SceneDepth"));
+				FRDGTextureRef SeparateTranslucencyTexture = Pass->CreateRDGTextureForOptionalInput(GraphBuilder, ePId_Input2, TEXT("SeparateTranslucency"));
 
-		FRCPassPostProcessTonemap* Tonemapper = 0;
+				FRDGTextureRef CustomDepthTexture = GraphBuilder.TryRegisterExternalTexture(SceneContext.CustomDepth, TEXT("CustomDepth"));
 
-		FRenderingCompositeOutputRef SSRInputChain;
+				FSceneTextureParameters SceneTextures;
+				SetupSceneTextureParameters(GraphBuilder, &SceneTextures);
 
-		// add the passes we want to add to the graph (commenting a line means the pass is not inserted into the graph) ---------
+				FRDGTextureRef VelocityTexture = SceneTextures.SceneVelocityBuffer;
 
-		if (AllowFullPostProcessing(View, FeatureLevel))
-		{
-			FRenderingCompositeOutputRef VelocityInput;
-			if(VelocityRT)
-			{
-				VelocityInput = Context.Graph.RegisterPass(new(FMemStack::Get()) FRCPassPostProcessInput(VelocityRT));
-				PreFlattenVelocity = VelocityInput;
-			}
+				const FScreenPassViewInfo& ScreenPassView(View);
 
-			Context.FinalOutput = AddPostProcessMaterialChain(Context, BL_BeforeTranslucency, SeparateTranslucency);
-			
-			static const auto CVar = IConsoleManager::Get().FindTConsoleVariableDataInt(TEXT("r.DepthOfFieldQuality"));
-			check(CVar)
-			bool bDepthOfField = 
-				View.Family->EngineShowFlags.DepthOfField &&
-				CVar->GetValueOnRenderThread() > 0 &&
-				View.FinalPostProcessSettings.DepthOfFieldFstop > 0 &&
-				View.FinalPostProcessSettings.DepthOfFieldFocalDistance > 0;
-
-			// Applies DOF and separate translucency,
-			{
-				FRenderingCompositePass* DiaphragmDOFPass = Context.Graph.RegisterPass(
-					new(FMemStack::Get()) TRCPassForRDG<3, 1>([bDepthOfField](FRenderingCompositePass* Pass, FRenderingCompositePassContext& InContext)
+				const FIntRect PrimaryViewRect = View.ViewRect;
+
+				// Post Process Material Chain - Before Translucency
 				{
-					FRDGBuilder GraphBuilder(InContext.RHICmdList);
+					FPostProcessMaterialInputs Inputs;
+					Inputs.SetInput(EPostProcessMaterialInput::SceneColor, SceneColorTexture, PrimaryViewRect);
+					Inputs.SetInput(EPostProcessMaterialInput::SeparateTranslucency, SeparateTranslucencyTexture, PrimaryViewRect);
+					Inputs.SetInput(EPostProcessMaterialInput::Velocity, VelocityTexture, PrimaryViewRect);
+					Inputs.CustomDepthTexture = CustomDepthTexture;
 
-					FSceneTextureParameters SceneTextures;
-					SetupSceneTextureParameters(GraphBuilder, &SceneTextures);
-	
-					FRDGTextureRef SceneColor = Pass->CreateRDGTextureForRequiredInput(GraphBuilder, ePId_Input0, TEXT("SceneColor"));
-					FRDGTextureRef LocalSeparateTranslucency = Pass->CreateRDGTextureForOptionalInput(GraphBuilder, ePId_Input1, TEXT("SeparateTranslucency"));
-					
-					// FPostProcessing::Process() does a AdjustGBufferRefCount(RHICmdList, -1), therefore need to pass down reference on velocity buffer manually.
-					SceneTextures.SceneVelocityBuffer = Pass->CreateRDGTextureForOptionalInput(GraphBuilder, ePId_Input2, TEXT("SceneVelocity"));
+					SceneColorTexture = AddPostProcessMaterialChain(GraphBuilder, ScreenPassView, Inputs, BL_BeforeTranslucency);
+				}
 
-					FRDGTextureRef NewSceneColor = SceneColor;
+				// Diaphragm Depth of Field
+				{
+					FRDGTextureRef LocalSceneColorTexture = SceneColorTexture;
 
-					if (bDepthOfField && DiaphragmDOF::IsSupported(InContext.View.GetShaderPlatform()))
+					if (bDepthOfFieldEnabled && DiaphragmDOF::IsSupported(View.GetShaderPlatform()))
 					{
-						NewSceneColor = DiaphragmDOF::AddPasses(
+						LocalSceneColorTexture = DiaphragmDOF::AddPasses(
 							GraphBuilder,
-							SceneTextures, InContext.View,
-							SceneColor, LocalSeparateTranslucency);
+							SceneTextures,
+							View,
+							SceneColorTexture,
+							SeparateTranslucencyTexture);
 					}
 
-					// DOF passes were not added, therefore need to compose Separate translucency manually. 
-					if (NewSceneColor == SceneColor && LocalSeparateTranslucency)
+					// DOF passes were not added, therefore need to compose Separate translucency manually.
+					if (LocalSceneColorTexture == SceneColorTexture && SeparateTranslucencyTexture)
 					{
-						NewSceneColor = AddSeparateTranslucencyCompositionPass(GraphBuilder, InContext.View, SceneColor, LocalSeparateTranslucency);
+						LocalSceneColorTexture = AddSeparateTranslucencyCompositionPass(GraphBuilder, View, SceneColorTexture, SeparateTranslucencyTexture);
 					}
 
-					Pass->ExtractRDGTextureForOutput(GraphBuilder, ePId_Output0, NewSceneColor);
-
-					GraphBuilder.Execute();
-				}));
-				DiaphragmDOFPass->SetInput(ePId_Input0, Context.FinalOutput);
-				DiaphragmDOFPass->SetInput(ePId_Input1, SeparateTranslucency);
-				DiaphragmDOFPass->SetInput(ePId_Input2, VelocityInput);
-				Context.FinalOutput = FRenderingCompositeOutputRef(DiaphragmDOFPass, ePId_Output0);
-			}
-
-			Context.FinalOutput = AddPostProcessMaterialChain(Context, BL_BeforeTonemapping, SeparateTranslucency, nullptr, nullptr, PreFlattenVelocity);
-
-			EAntiAliasingMethod AntiAliasingMethod = Context.View.AntiAliasingMethod;
-
-			const int32 DownsampleQuality = FMath::Clamp(CDownsampleQuality.GetValueOnRenderThread(), 0, 1);
-
-			FRenderingCompositeOutputRef SceneColorHalfRes;
-			const EPixelFormat SceneColorHalfResFormat = PF_FloatRGB;
-
-			if( AntiAliasingMethod == AAM_TemporalAA && ViewState)
-			{
-				FTAAPassParameters TAAParameters = MakeTAAPassParametersForView(Context.View);
-
-				float HistoryUpscaleSize = FMath::Clamp(CVarTemporalAAHistorySP.GetValueOnRenderThread() / 100.0f, 1.0f, 2.0f);
-				if (!IsPCPlatform(View.GetShaderPlatform()) || !IsFeatureLevelSupported(View.GetShaderPlatform(), ERHIFeatureLevel::SM5))
-				{
-					HistoryUpscaleSize = 1;
+					SceneColorTexture = LocalSceneColorTexture;
 				}
 
-				// Downsample pass may be merged with with TemporalAA when there is no motion blur and compute shader is used.
-				// This is currently only possible for r.Downsample.Quality = 0 (box filter).
-				TAAParameters.bDownsample = (CVarTemporalAAAllowDownsampling.GetValueOnRenderThread() != 0)
-					&& !IsMotionBlurEnabled(View)
-					&& !bVisualizeMotionBlur
-					&& (DownsampleQuality == 0)
-					&& TAAParameters.bUseFast;
+				// Post Process Material Chain - Before Tonemapping
+				{
+					FPostProcessMaterialInputs Inputs;
+					Inputs.SetInput(EPostProcessMaterialInput::SceneColor, SceneColorTexture, PrimaryViewRect);
+					Inputs.SetInput(EPostProcessMaterialInput::SeparateTranslucency, SeparateTranslucencyTexture, PrimaryViewRect);
+					Inputs.SetInput(EPostProcessMaterialInput::Velocity, VelocityTexture, PrimaryViewRect);
+					Inputs.CustomDepthTexture = CustomDepthTexture;
 
-				TAAParameters.DownsampleOverrideFormat = SceneColorHalfResFormat;
+					SceneColorTexture = AddPostProcessMaterialChain(GraphBuilder, ScreenPassView, Inputs, BL_BeforeTonemapping);
+				}
 
-				FIntPoint SecondaryViewRectSize(TAAParameters.OutputViewRect.Width(), TAAParameters.OutputViewRect.Height());
+				// Half-resolution scene color texture.
+				FRDGTextureRef SceneColorHalfResTexture = nullptr;
 
-				if (HistoryUpscaleSize > 1.0f)
-				{
-					FIntPoint HistoryViewSize(
-						SecondaryViewRectSize.X * HistoryUpscaleSize,
-						SecondaryViewRectSize.Y * HistoryUpscaleSize);
+				// Scene color view rectangle after temporal AA upscale to secondary screen percentage.
+				FIntRect SecondaryViewRect;
 
-					FIntPoint QuantizedMinHistorySize;
-					QuantizeSceneBufferSize(HistoryViewSize, QuantizedMinHistorySize);
+				AddTemporalAAMainPass(
+					GraphBuilder,
+					SceneTextures,
+					ScreenPassView,
+					SceneColorTexture,
+					CustomDepthTexture,
+					&SceneColorTexture,
+					&SceneColorHalfResTexture,
+					&SecondaryViewRect);
 
-					TAAParameters.Pass = ETAAPassConfig::MainSuperSampling;
-					TAAParameters.bDownsample = false;
-					TAAParameters.bUseFast = false;
+				// SceneColorTexture is now upscaled to the SecondaryViewRect due to TAAU.
 
-					TAAParameters.OutputViewRect.Min.X = 0;
-					TAAParameters.OutputViewRect.Min.Y = 0;
-					TAAParameters.OutputViewRect.Max = HistoryViewSize;
-				}
-
-				if(VelocityInput.IsValid())
+				if (VelocityTexture)
 				{
-					AddTemporalAA( Context, VelocityInput, TAAParameters, TAAParameters.bDownsample ? &SceneColorHalfRes : nullptr);
+					SceneColorTexture = AddMotionBlurPass(
+						GraphBuilder,
+						ScreenPassView,
+						SecondaryViewRect,
+						PrimaryViewRect,
+						SceneColorTexture,
+						SceneDepthTexture,
+						VelocityTexture);
 				}
-				else
+
+				if (bVisualizeBloom)
 				{
-					// black is how we clear the velocity buffer so this means no velocity
-					FRenderingCompositePass* NoVelocity = Context.Graph.RegisterPass(new(FMemStack::Get()) FRCPassPostProcessInput(GSystemTextures.BlackDummy));
-					FRenderingCompositeOutputRef NoVelocityRef(NoVelocity);
-					AddTemporalAA( Context, NoVelocityRef, TAAParameters, TAAParameters.bDownsample ? &SceneColorHalfRes : nullptr);
+					SceneColorTexture = AddVisualizeBloomSetupPass(GraphBuilder, ScreenPassView, SceneColorTexture, SecondaryViewRect);
 				}
 
-				if (HistoryUpscaleSize > 1.0f)
+				// If TAA didn't do it, downsample the scene color texture by half.
+				if (!SceneColorHalfResTexture)
 				{
-					FIntPoint QuantizedOutputSize;
-					QuantizeSceneBufferSize(SecondaryViewRectSize, QuantizedOutputSize);
+					FDownsamplePassInputs Inputs;
+					Inputs.Name = TEXT("SceneColorHalfRes");
+					Inputs.Texture = SceneColorTexture;
+					Inputs.Viewport = SecondaryViewRect;
+					Inputs.Quality = GetDownsampleQuality();
+					Inputs.FormatOverride = GetSceneColorHalfResFormat();
 
-					const FIntRect InputViewport = TAAParameters.OutputViewRect;
+					FDownsamplePassOutputs Outputs = AddDownsamplePass(GraphBuilder, ScreenPassView, Inputs);
 
-					FScreenPassTextureViewport OutputViewport;
-					OutputViewport.Rect.Max = SecondaryViewRectSize;
-					OutputViewport.Extent.X = FMath::Max(SceneContext.GetBufferSizeXY().X, QuantizedOutputSize.X);
-					OutputViewport.Extent.Y = FMath::Max(SceneContext.GetBufferSizeXY().Y, QuantizedOutputSize.Y);
-
-					Context.FinalOutput = ComputeMitchellNetravaliDownsample(Context.Graph, Context.FinalOutput, InputViewport, OutputViewport);
+					SceneColorHalfResTexture = Outputs.Texture;
 				}
 
-				SSRInputChain = AddPostProcessMaterialChain(Context, BL_SSRInput);
-			}
+				// Release held GBuffer reference taken during composition graph setup. Passes will take their own references during RDG setup.
+				SceneContext.AdjustGBufferRefCount(InContext.RHICmdList, -1);
 
-			// Motion Blur
-			if ((IsMotionBlurEnabled(View) || bVisualizeMotionBlur) && VelocityInput.IsValid())
-			{
-				Context.FinalOutput = ComputeMotionBlurShim(Context.Graph, Context.FinalOutput, Context.SceneDepth, VelocityInput, bVisualizeMotionBlur);
-			}
+				Pass->ExtractRDGTextureForOutput(GraphBuilder, ePId_Output0, SceneColorTexture);
+				Pass->ExtractRDGTextureForOutput(GraphBuilder, ePId_Output1, SceneColorHalfResTexture);
 
-			if(bVisualizeBloom)
-			{
-				AddVisualizeBloomSetup(Context);
-			}
+				GraphBuilder.Execute();
 
-			// Down sample Scene color from full to half res (this may have been done during TAA).
-			if (!SceneColorHalfRes.IsValid())
-			{
-				// doesn't have to be as high quality as the Scene color
-				const bool bIsComputePass = ShouldDoComputePostProcessing(Context.View);
+				InContext.SceneColorViewRect = SecondaryViewRect;
+				InContext.ReferenceBufferSize = SceneColorTexture->Desc.Extent;
+			}));
 
-				FRenderingCompositePass* HalfResPass = Context.Graph.RegisterPass(new(FMemStack::Get()) FRCPassPostProcessDownsample(SceneColorHalfResFormat, DownsampleQuality, bIsComputePass, TEXT("SceneColorHalfRes")));
-				HalfResPass->SetInput(ePId_Input0, FRenderingCompositeOutputRef(Context.FinalOutput));
+			RDGPass->SetInput(ePId_Input0, Context.FinalOutput);
+			RDGPass->SetInput(ePId_Input1, Context.SceneDepth);
+			RDGPass->SetInput(ePId_Input2, SeparateTranslucency);
 
-				SceneColorHalfRes = FRenderingCompositeOutputRef(HalfResPass);
-			}
+			Context.FinalOutput = FRenderingCompositeOutputRef(RDGPass, ePId_Output0);
+
+			SceneColorHalfRes = FRenderingCompositeOutputRef(RDGPass, ePId_Output1);
 
 			{
 				bool bHistogramNeeded = false;
@@ -1309,7 +1337,7 @@ void FPostProcessing::Process(FRHICommandListImmediate& RHICmdList, const FViewI
 					bHistogramNeeded = true;
 				}
 
-				if(!bAllowTonemapper)
+				if(!bTonemapEnabled)
 				{
 					bHistogramNeeded = false;
 				}
@@ -1319,7 +1347,7 @@ void FPostProcessing::Process(FRHICommandListImmediate& RHICmdList, const FViewI
 					bHistogramNeeded = true;
 				}
 
-				if (!GIsHighResScreenshot && bHistogramNeeded && FeatureLevel >= ERHIFeatureLevel::SM5 && IStereoRendering::IsAPrimaryView(StereoPass))
+				if (!GIsHighResScreenshot && bHistogramNeeded && FeatureLevel >= ERHIFeatureLevel::SM5 && IStereoRendering::IsAPrimaryView(View.StereoPass))
 				{
 					FRenderingCompositePass* NodeHistogram = Context.Graph.RegisterPass(new(FMemStack::Get()) FRCPassPostProcessHistogram());
 
@@ -1343,7 +1371,7 @@ void FPostProcessing::Process(FRHICommandListImmediate& RHICmdList, const FViewI
 				if (Context.View.FinalPostProcessSettings.BloomThreshold <= -1 && Context.View.Family->Views.Num() == 1)
 				{
 					if (!GIsHighResScreenshot && View.State &&
-						IStereoRendering::IsAPrimaryView(StereoPass) &&
+						IStereoRendering::IsAPrimaryView(View.StereoPass) &&
 						AutoExposure.MethodId == EAutoExposureMethod::AEM_Basic)
 					{
 						BloomAndEyeDownSamplesPtr = CreateDownSampleArray(Context, SceneColorHalfRes, true /*bGenerateLog2Alpha*/);
@@ -1352,7 +1380,7 @@ void FPostProcessing::Process(FRHICommandListImmediate& RHICmdList, const FViewI
 			}
 
 			// some views don't have a state (thumbnail rendering)
-			if(!GIsHighResScreenshot && View.State && IStereoRendering::IsAPrimaryView(StereoPass))
+			if(!GIsHighResScreenshot && View.State && IStereoRendering::IsAPrimaryView(View.StereoPass))
 			{
 				const bool bUseBasicEyeAdaptation = (AutoExposure.MethodId == EAutoExposureMethod::AEM_Basic);
 
@@ -1380,110 +1408,17 @@ void FPostProcessing::Process(FRHICommandListImmediate& RHICmdList, const FViewI
 
 			if(View.FinalPostProcessSettings.BloomIntensity > 0.0f)
 			{
-				if (CVarUseMobileBloom.GetValueOnRenderThread() == 0)
+				if (!BloomAndEyeDownSamplesPtr.IsValid())
 				{
-					if (!BloomAndEyeDownSamplesPtr.IsValid())
-					{
-						FRenderingCompositeOutputRef HalfResBloomThreshold = RenderHalfResBloomThreshold(Context, SceneColorHalfRes, AutoExposure.EyeAdaptation);
-						BloomAndEyeDownSamplesPtr = CreateDownSampleArray(Context, HalfResBloomThreshold, false /*bGenerateLog2Alpha*/);
-					}
-					BloomOutputCombined = AddBloom(*BloomAndEyeDownSamplesPtr, bVisualizeBloom);
-				}
-				else
-				{
-					FIntPoint PrePostSourceViewportSize = View.ViewRect.Size();
-
-					// Bloom.
-					FRenderingCompositeOutputRef PostProcessDownsample2;
-					FRenderingCompositeOutputRef PostProcessDownsample3;
-					FRenderingCompositeOutputRef PostProcessDownsample4;
-					FRenderingCompositeOutputRef PostProcessDownsample5;
-					FRenderingCompositeOutputRef PostProcessUpsample4;
-					FRenderingCompositeOutputRef PostProcessUpsample3;
-					FRenderingCompositeOutputRef PostProcessUpsample2;
-					FRenderingCompositeOutputRef PostProcessSunMerge;
-
-					float DownScale = 0.66f * 4.0f;
-					// Downsample by 2
-					{
-						FRenderingCompositePass* Pass = Context.Graph.RegisterPass(new(FMemStack::Get()) FRCPassPostProcessBloomDownES2(PrePostSourceViewportSize/4, DownScale));
-						Pass->SetInput(ePId_Input0, SceneColorHalfRes);
-						PostProcessDownsample2 = FRenderingCompositeOutputRef(Pass);
-					}
-
-					// Downsample by 2
-					{
-						FRenderingCompositePass* Pass = Context.Graph.RegisterPass(new(FMemStack::Get()) FRCPassPostProcessBloomDownES2(PrePostSourceViewportSize/8, DownScale));
-						Pass->SetInput(ePId_Input0, PostProcessDownsample2);
-						PostProcessDownsample3 = FRenderingCompositeOutputRef(Pass);
-					}
-
-					// Downsample by 2
-					{
-						FRenderingCompositePass* Pass = Context.Graph.RegisterPass(new(FMemStack::Get()) FRCPassPostProcessBloomDownES2(PrePostSourceViewportSize/16, DownScale));
-						Pass->SetInput(ePId_Input0, PostProcessDownsample3);
-						PostProcessDownsample4 = FRenderingCompositeOutputRef(Pass);
-					}
-
-					// Downsample by 2
-					{
-						FRenderingCompositePass* Pass = Context.Graph.RegisterPass(new(FMemStack::Get()) FRCPassPostProcessBloomDownES2(PrePostSourceViewportSize/32, DownScale));
-						Pass->SetInput(ePId_Input0, PostProcessDownsample4);
-						PostProcessDownsample5 = FRenderingCompositeOutputRef(Pass);
-					}
-
-					const FFinalPostProcessSettings& Settings = Context.View.FinalPostProcessSettings;
-
-					float UpScale = 0.66f * 2.0f;
-					// Upsample by 2
-					{
-						FVector4 TintA = FVector4(Settings.Bloom4Tint.R, Settings.Bloom4Tint.G, Settings.Bloom4Tint.B, 0.0f);
-						FVector4 TintB = FVector4(Settings.Bloom5Tint.R, Settings.Bloom5Tint.G, Settings.Bloom5Tint.B, 0.0f);
-						TintA *= View.FinalPostProcessSettings.BloomIntensity;
-						TintB *= View.FinalPostProcessSettings.BloomIntensity;
-						FRenderingCompositePass* Pass = Context.Graph.RegisterPass(new(FMemStack::Get()) FRCPassPostProcessBloomUpES2(PrePostSourceViewportSize/32, FVector2D(UpScale, UpScale), TintA, TintB));
-						Pass->SetInput(ePId_Input0, PostProcessDownsample4);
-						Pass->SetInput(ePId_Input1, PostProcessDownsample5);
-						PostProcessUpsample4 = FRenderingCompositeOutputRef(Pass);
-					}
-
-					// Upsample by 2
-					{
-						FVector4 TintA = FVector4(Settings.Bloom3Tint.R, Settings.Bloom3Tint.G, Settings.Bloom3Tint.B, 0.0f);
-						TintA *= View.FinalPostProcessSettings.BloomIntensity;
-						FVector4 TintB = FVector4(1.0f, 1.0f, 1.0f, 0.0f);
-						FRenderingCompositePass* Pass = Context.Graph.RegisterPass(new(FMemStack::Get()) FRCPassPostProcessBloomUpES2(PrePostSourceViewportSize/16, FVector2D(UpScale, UpScale), TintA, TintB));
-						Pass->SetInput(ePId_Input0, PostProcessDownsample3);
-						Pass->SetInput(ePId_Input1, PostProcessUpsample4);
-						PostProcessUpsample3 = FRenderingCompositeOutputRef(Pass);
-					}
-
-					// Upsample by 2
-					{
-						FVector4 TintA = FVector4(Settings.Bloom2Tint.R, Settings.Bloom2Tint.G, Settings.Bloom2Tint.B, 0.0f);
-						TintA *= View.FinalPostProcessSettings.BloomIntensity;
-						// Scaling Bloom2 by extra factor to match filter area difference between PC default and mobile.
-						TintA *= 0.5;
-						FVector4 TintB = FVector4(1.0f, 1.0f, 1.0f, 0.0f);
-						FRenderingCompositePass* Pass = Context.Graph.RegisterPass(new(FMemStack::Get()) FRCPassPostProcessBloomUpES2(PrePostSourceViewportSize/8, FVector2D(UpScale, UpScale), TintA, TintB));
-						Pass->SetInput(ePId_Input0, PostProcessDownsample2);
-						Pass->SetInput(ePId_Input1, PostProcessUpsample3);
-						PostProcessUpsample2 = FRenderingCompositeOutputRef(Pass);
-					}
-
-					{
-						FRenderingCompositePass* Pass = Context.Graph.RegisterPass(new(FMemStack::Get()) FRCPassPostProcessSunMergeES2(PrePostSourceViewportSize));
-						Pass->SetInput(ePId_Input1, SceneColorHalfRes);
-						Pass->SetInput(ePId_Input2, PostProcessUpsample2);
-						PostProcessSunMerge = FRenderingCompositeOutputRef(Pass);
-						BloomOutputCombined = PostProcessSunMerge;
-					}
+					FRenderingCompositeOutputRef HalfResBloomThreshold = RenderHalfResBloomThreshold(Context, SceneColorHalfRes, AutoExposure.EyeAdaptation);
+					BloomAndEyeDownSamplesPtr = CreateDownSampleArray(Context, HalfResBloomThreshold, false /*bGenerateLog2Alpha*/);
 				}
+				BloomOutputCombined = AddBloom(*BloomAndEyeDownSamplesPtr, bVisualizeBloom);
 			}
 
 			PreTonemapHDRColor = Context.FinalOutput;
 
-			if(bAllowTonemapper)
+			if(bTonemapEnabled)
 			{
 				{
 					FRenderingCompositeOutputRef FinalOutputPrev = Context.FinalOutput;
@@ -1512,8 +1447,8 @@ void FPostProcessing::Process(FRHICommandListImmediate& RHICmdList, const FViewI
 			{
 				AddPostProcessAA(Context);
 			}
-			
-			if(bDepthOfField && Context.View.Family->EngineShowFlags.VisualizeDOF)
+
+			if(bVisualizeDepthOfField)
 			{
 				FDepthOfFieldStats DepthOfFieldStat;
 
@@ -1521,7 +1456,6 @@ void FPostProcessing::Process(FRHICommandListImmediate& RHICmdList, const FViewI
 				VisualizeNode->SetInput(ePId_Input0, FRenderingCompositeOutputRef(Context.FinalOutput));
 
 				Context.FinalOutput = FRenderingCompositeOutputRef(VisualizeNode);
-				bAllowTonemapper = false;
 			}
 		}
 		else
@@ -1761,7 +1695,7 @@ void FPostProcessing::Process(FRHICommandListImmediate& RHICmdList, const FViewI
 		// If a postprocess material is using a GBuffer it adds the refcount int FRCPassPostProcessMaterial::Process()
 		// and when it gets processed it removes the refcount
 		// We only release the GBuffers after the last view was processed (SplitScreen)
-		if(View.Family->Views[View.Family->Views.Num() - 1] == &View)
+		if(View.IsLastInFamily())
 		{
 			// Generally we no longer need the GBuffers, anyone that wants to keep the GBuffers for longer should have called AdjustGBufferRefCount(1) to keep it for longer
 			// and call AdjustGBufferRefCount(-1) once it's consumed. This needs to happen each frame. PostProcessMaterial do that automatically
@@ -1815,12 +1749,6 @@ void FPostProcessing::Process(FRHICommandListImmediate& RHICmdList, const FViewI
 			TArray<FRenderingCompositePass*> TargetedRoots;
 			TargetedRoots.Add(Context.FinalOutput.GetPass());
 
-			if (SSRInputChain.IsValid())
-			{
-				check(ViewState);
-				TargetedRoots.Add(SSRInputChain.GetPass());
-			}
-
 			// execute the graph/DAG
 			CompositeContext.Process(TargetedRoots, TEXT("PostProcessing"));
 
@@ -1833,12 +1761,6 @@ void FPostProcessing::Process(FRHICommandListImmediate& RHICmdList, const FViewI
 					Context.RHICmdList.WaitComputeFence(ComputeFinalizeFence);
 				}
 			}
-
-			if (SSRInputChain.IsValid() && !View.bViewStateIsReadOnly)
-			{
-				check(ViewState);
-				ViewState->PrevFrameViewInfo.CustomSSRInput = SSRInputChain.GetOutput()->PooledRenderTarget;
-			}
 		}
 	}
 
@@ -2405,44 +2327,44 @@ void FPostProcessing::ProcessES2(FRHICommandListImmediate& RHICmdList, FScene* S
 	SetMobilePassFlipVerticalAxis(nullptr);
 }
 
-void FPostProcessing::ProcessPlanarReflection(FRHICommandListImmediate& RHICmdList, FViewInfo& View, TRefCountPtr<IPooledRenderTarget>& VelocityRT, TRefCountPtr<IPooledRenderTarget>& OutFilteredSceneColor)
+void FPostProcessing::ProcessPlanarReflection(FRHICommandListImmediate& RHICmdList, const FViewInfo& View, TRefCountPtr<IPooledRenderTarget>& OutFilteredSceneColor)
 {
+	FSceneViewState* ViewState = View.ViewState;
+	const EAntiAliasingMethod AntiAliasingMethod = View.AntiAliasingMethod;
+
+	const FSceneRenderTargets& SceneContext = FSceneRenderTargets::Get(RHICmdList);
+
+	if (AntiAliasingMethod == AAM_TemporalAA && ViewState)
 	{
-		FMemMark Mark(FMemStack::Get());
-		FRenderingCompositePassContext CompositeContext(RHICmdList, View);
+		FRDGBuilder GraphBuilder(RHICmdList);
 
-		FPostprocessContext Context(RHICmdList, CompositeContext.Graph, View);
-		FSceneRenderTargets& SceneContext = FSceneRenderTargets::Get(RHICmdList);
+		FSceneTextureParameters SceneTextures;
+		SetupSceneTextureParameters(GraphBuilder, &SceneTextures);
 
-		FRenderingCompositeOutputRef VelocityInput;
-		if(VelocityRT)
-		{
-			VelocityInput = Context.Graph.RegisterPass(new(FMemStack::Get()) FRCPassPostProcessInput(VelocityRT));
-		}
+		// Planar reflections don't support velocity.
+		SceneTextures.SceneVelocityBuffer = nullptr;
 
-		FSceneViewState* ViewState = Context.View.ViewState;
-		EAntiAliasingMethod AntiAliasingMethod = Context.View.AntiAliasingMethod;
+		const FTemporalAAHistory& InputHistory = View.PrevViewInfo.TemporalAAHistory;
+		FTemporalAAHistory* OutputHistory = &ViewState->PrevFrameViewInfo.TemporalAAHistory;
 
-		if (AntiAliasingMethod == AAM_TemporalAA && ViewState)
-		{
-			FTAAPassParameters Parameters = MakeTAAPassParametersForView(Context.View);
+		FTAAPassParameters Parameters(View);
+		Parameters.SceneColorInput = GraphBuilder.RegisterExternalTexture(SceneContext.GetSceneColor(), TEXT("SceneColor"));
 
-			if(VelocityInput.IsValid())
-			{
-				AddTemporalAA( Context, VelocityInput, Parameters, nullptr);
-			}
-			else
-			{
-				// black is how we clear the velocity buffer so this means no velocity
-				FRenderingCompositePass* NoVelocity = Context.Graph.RegisterPass(new(FMemStack::Get()) FRCPassPostProcessInput(GSystemTextures.BlackDummy));
-				FRenderingCompositeOutputRef NoVelocityRef(NoVelocity);
-				AddTemporalAA( Context, NoVelocityRef, Parameters, nullptr );
-			}
-		}
+		FTAAOutputs Outputs = AddTemporalAAPass(
+			GraphBuilder,
+			SceneTextures,
+			View,
+			Parameters,
+			InputHistory,
+			OutputHistory);
 
-		CompositeContext.Process(Context.FinalOutput.GetPass(), TEXT("ProcessPlanarReflection"));
+		GraphBuilder.QueueTextureExtraction(Outputs.SceneColor, &OutFilteredSceneColor);
 
-		OutFilteredSceneColor = Context.FinalOutput.GetOutput()->PooledRenderTarget;
+		GraphBuilder.Execute();
+	}
+	else
+	{
+		OutFilteredSceneColor = SceneContext.GetSceneColor();
 	}
 }
 
diff --git a/Engine/Source/Runtime/Renderer/Private/PostProcess/PostProcessing.h b/Engine/Source/Runtime/Renderer/Private/PostProcess/PostProcessing.h
--- a/Engine/Source/Runtime/Renderer/Private/PostProcess/PostProcessing.h
+++ b/Engine/Source/Runtime/Renderer/Private/PostProcess/PostProcessing.h
@@ -57,7 +57,7 @@ public:
 
 	void ProcessES2(FRHICommandListImmediate& RHICmdList, FScene* Scene, const FViewInfo& View);
 
-	void ProcessPlanarReflection(FRHICommandListImmediate& RHICmdList, FViewInfo& View, TRefCountPtr<IPooledRenderTarget>& VelocityRT, TRefCountPtr<IPooledRenderTarget>& OutFilteredSceneColor);
+	void ProcessPlanarReflection(FRHICommandListImmediate& RHICmdList, const FViewInfo& View, TRefCountPtr<IPooledRenderTarget>& OutFilteredSceneColor);
 
 #if WITH_EDITOR
 	void AddSelectionOutline(FPostprocessContext& Context);
diff --git a/Engine/Source/Runtime/Renderer/Private/PostProcess/RenderingCompositionGraph.cpp b/Engine/Source/Runtime/Renderer/Private/PostProcess/RenderingCompositionGraph.cpp
--- a/Engine/Source/Runtime/Renderer/Private/PostProcess/RenderingCompositionGraph.cpp
+++ b/Engine/Source/Runtime/Renderer/Private/PostProcess/RenderingCompositionGraph.cpp
@@ -1186,6 +1186,8 @@ FRDGTextureRef FRenderingCompositePass::CreateRDGTextureForInputWithFallback(
 
 void FRenderingCompositePass::ExtractRDGTextureForOutput(FRDGBuilder& GraphBuilder, EPassOutputId OutputId, FRDGTextureRef Texture)
 {
+	check(Texture);
+
 	if (FRenderingCompositeOutput* Output = GetOutput(OutputId))
 	{
 		Output->RenderTargetDesc = Texture->Desc;
diff --git a/Engine/Source/Runtime/Renderer/Private/ReflectionEnvironment.cpp b/Engine/Source/Runtime/Renderer/Private/ReflectionEnvironment.cpp
--- a/Engine/Source/Runtime/Renderer/Private/ReflectionEnvironment.cpp
+++ b/Engine/Source/Runtime/Renderer/Private/ReflectionEnvironment.cpp
@@ -739,9 +739,11 @@ void FDeferredShadingSceneRenderer::RenderDeferredReflectionsAndSkyLighting(FRHI
 				TAASettings.SceneColorInput = DenoiserInputs.Color;
 				TAASettings.bOutputRenderTargetable = bComposePlanarReflections;
 				
-				FTAAOutputs TAAOutputs = TAASettings.AddTemporalAAPass(
+				FTAAOutputs TAAOutputs = AddTemporalAAPass(
 					GraphBuilder,
-					SceneTextures, View,
+					SceneTextures,
+					View,
+					TAASettings,
 					View.PrevViewInfo.SSRHistory,
 					&View.ViewState->PrevFrameViewInfo.SSRHistory);
 				
diff --git a/Engine/Source/Runtime/Renderer/Private/SceneRendering.h b/Engine/Source/Runtime/Renderer/Private/SceneRendering.h
--- a/Engine/Source/Runtime/Renderer/Private/SceneRendering.h
+++ b/Engine/Source/Runtime/Renderer/Private/SceneRendering.h
@@ -1209,7 +1209,11 @@ public:
 	* for stereo rendering, this will force the post-processing to use the same render target for both eyes*/
 	FSceneRenderTargetItem* GetTonemappingLUTRenderTarget(FRHICommandList& RHICmdList, const int32 LUTSize, const bool bUseVolumeLUT, const bool bNeedUAV, const bool bNeedFloatOutput) const;
 	
-
+	/** Returns whether this view is the last in the family. */
+	bool IsLastInFamily() const
+	{
+		return Family->Views.Last() == this;
+	}
 
 	/** Instanced stereo and multi-view only need to render the left eye. */
 	bool ShouldRenderView() const 
diff --git a/Engine/Source/Runtime/Renderer/Private/SceneTextureParameters.cpp b/Engine/Source/Runtime/Renderer/Private/SceneTextureParameters.cpp
--- a/Engine/Source/Runtime/Renderer/Private/SceneTextureParameters.cpp
+++ b/Engine/Source/Runtime/Renderer/Private/SceneTextureParameters.cpp
@@ -23,14 +23,14 @@ void SetupSceneTextureParameters(
 
 	// Registers all the scene texture from the scene context. No fallback is provided to catch mistake at shader parameter validation time
 	// when a pass is trying to access a resource before any other pass actually created it.
-	OutTextures->SceneVelocityBuffer = SceneContext.SceneVelocity ? GraphBuilder.RegisterExternalTexture(SceneContext.SceneVelocity, TEXT("VelocityBuffer")) : nullptr;
-	OutTextures->SceneGBufferA = SceneContext.GBufferA ? GraphBuilder.RegisterExternalTexture(SceneContext.GBufferA, TEXT("GBufferA")) : nullptr;
-	OutTextures->SceneGBufferB = SceneContext.GBufferB ? GraphBuilder.RegisterExternalTexture(SceneContext.GBufferB, TEXT("GBufferB")) : nullptr;
-	OutTextures->SceneGBufferC = SceneContext.GBufferC ? GraphBuilder.RegisterExternalTexture(SceneContext.GBufferC, TEXT("GBufferC")) : nullptr;
-	OutTextures->SceneGBufferD = SceneContext.GBufferD ? GraphBuilder.RegisterExternalTexture(SceneContext.GBufferD, TEXT("GBufferD")) : nullptr;
-	OutTextures->SceneGBufferE = SceneContext.GBufferE ? GraphBuilder.RegisterExternalTexture(SceneContext.GBufferE, TEXT("GBufferE")) : nullptr;
+	OutTextures->SceneVelocityBuffer = GraphBuilder.TryRegisterExternalTexture(SceneContext.SceneVelocity, TEXT("VelocityBuffer"));
+	OutTextures->SceneGBufferA = GraphBuilder.TryRegisterExternalTexture(SceneContext.GBufferA, TEXT("GBufferA"));
+	OutTextures->SceneGBufferB = GraphBuilder.TryRegisterExternalTexture(SceneContext.GBufferB, TEXT("GBufferB"));
+	OutTextures->SceneGBufferC = GraphBuilder.TryRegisterExternalTexture(SceneContext.GBufferC, TEXT("GBufferC"));
+	OutTextures->SceneGBufferD = GraphBuilder.TryRegisterExternalTexture(SceneContext.GBufferD, TEXT("GBufferD"));
+	OutTextures->SceneGBufferE = GraphBuilder.TryRegisterExternalTexture(SceneContext.GBufferE, TEXT("GBufferE"));
 
-	// Ligthing channels might be disabled when all lights are on the same channel.
+	// Lighting channels might be disabled when all lights are on the same channel.
 	if (SceneContext.LightingChannels)
 	{
 		OutTextures->SceneLightingChannels = GraphBuilder.RegisterExternalTexture(SceneContext.LightingChannels, TEXT("LightingChannels"));
