Fixed build (bad cvar variable name)

diff --git a/Engine/Source/Runtime/Renderer/Private/SkyAtmosphereRendering.cpp b/Engine/Source/Runtime/Renderer/Private/SkyAtmosphereRendering.cpp
--- a/Engine/Source/Runtime/Renderer/Private/SkyAtmosphereRendering.cpp
+++ b/Engine/Source/Runtime/Renderer/Private/SkyAtmosphereRendering.cpp
@@ -25,7 +25,7 @@ static TAutoConsoleVariable<int32> CVarSkyAtmosphere(
 	ECVF_RenderThreadSafe);
 
 // The project settings (disable runtime and shader code)
-static TAutoConsoleVariable<int32> CVarSupportAtmosphericFog(
+static TAutoConsoleVariable<int32> CVarSupportSkyAtmosphere(
 	TEXT("r.SupportSkyAtmosphere"),
 	1,
 	TEXT("Enables SkyAtmosphere rendering and shader code."),
