Fix shader binding error in denoiser (only when shaders not optimized)
#rnx

diff --git a/Engine/Source/Runtime/Renderer/Private/ScreenSpaceDenoise.cpp b/Engine/Source/Runtime/Renderer/Private/ScreenSpaceDenoise.cpp
--- a/Engine/Source/Runtime/Renderer/Private/ScreenSpaceDenoise.cpp
+++ b/Engine/Source/Runtime/Renderer/Private/ScreenSpaceDenoise.cpp
@@ -985,7 +985,7 @@ class FSSDComposeHarmonicsCS : public FGlobalShader
 
 	BEGIN_SHADER_PARAMETER_STRUCT(FParameters, )
 		SHADER_PARAMETER_STRUCT_ARRAY(FSSDSignalTextures, SignalHarmonics, [IScreenSpaceDenoiser::kMultiPolychromaticPenumbraHarmonics])
-		SHADER_PARAMETER_STRUCT_REF(FViewUniformShaderParameters, ViewUniformBuffer)
+		SHADER_PARAMETER_STRUCT_INCLUDE(FSSDCommonParameters, CommonParameters)
 		SHADER_PARAMETER_STRUCT(FSSDSignalUAVs, SignalOutput)
 	END_SHADER_PARAMETER_STRUCT()
 };
@@ -1821,7 +1821,9 @@ public:
 				ComposedHarmonics.Textures[2] = GraphBuilder.CreateTexture(Desc, TEXT("MultiPolychromaticComposition2"));
 			}
 
-			ComposePassParameters->ViewUniformBuffer = View.ViewUniformBuffer;
+			ComposePassParameters->CommonParameters.ViewUniformBuffer = View.ViewUniformBuffer;
+			ComposePassParameters->CommonParameters.SceneTextures = SceneTextures;
+
 			ComposePassParameters->SignalOutput = CreateMultiplexedUAVs(GraphBuilder, ComposedHarmonics);
 
 			TShaderMapRef<FSSDComposeHarmonicsCS> ComputeShader(View.ShaderMap);
