Fixed -rdgdebug warning in post process material.

diff --git a/Engine/Source/Runtime/Renderer/Private/PostProcess/PostProcessMaterial.cpp b/Engine/Source/Runtime/Renderer/Private/PostProcess/PostProcessMaterial.cpp
--- a/Engine/Source/Runtime/Renderer/Private/PostProcess/PostProcessMaterial.cpp
+++ b/Engine/Source/Runtime/Renderer/Private/PostProcess/PostProcessMaterial.cpp
@@ -454,9 +454,7 @@ FRDGTextureRef ComputePostProcessMaterial(
 	PostProcessMaterialParameters->PostProcessOutput = GetScreenPassTextureViewportParameters(SceneColorViewport);
 	PostProcessMaterialParameters->CustomDepth = DepthStencilTextureForSRV;
 
-	PostProcessMaterialParameters->RenderTargets[0] = FRenderTargetBinding(
-		OutputTexture,
-		OutputLoadAction);
+	PostProcessMaterialParameters->RenderTargets[0] = FRenderTargetBinding(OutputTexture, OutputLoadAction);
 
 	if (DepthStencilTexture)
 	{
@@ -469,7 +467,10 @@ FRDGTextureRef ComputePostProcessMaterial(
 
 	PostProcessMaterialParameters->PostProcessInput_BilinearSampler = TStaticSamplerState<SF_Bilinear, AM_Clamp, AM_Clamp, AM_Clamp>::GetRHI();;
 
-	FRDGTextureRef BlackTexture = GraphBuilder.RegisterExternalTexture(GSystemTextures.BlackDummy);
+	FRDGTextureRef BlackTexture = GraphBuilder.RegisterExternalTexture(GSystemTextures.BlackDummy, TEXT("BlackDummy"));
+
+    // This gets passed in whether or not it's used.
+	GraphBuilder.RemoveUnusedTextureWarning(BlackTexture);
 
 	FRHISamplerState* PointClampSampler = TStaticSamplerState<SF_Point, AM_Clamp, AM_Clamp, AM_Clamp>::GetRHI();
 
