Copying //UE4/Dev-Volumetrics to Dev-Rendering (//UE4/Dev-Rendering)

diff --git a/Engine/Shaders/Private/BasePassPixelShader.usf b/Engine/Shaders/Private/BasePassPixelShader.usf
--- a/Engine/Shaders/Private/BasePassPixelShader.usf
+++ b/Engine/Shaders/Private/BasePassPixelShader.usf
@@ -1144,7 +1144,7 @@ ShadingModel);
 			View.CameraAerialPerspectiveVolume, View.CameraAerialPerspectiveVolumeSampler,
 			ResolvedView.SkyAtmosphereCameraAerialPerspectiveVolumeDepthResolutionInv,
 			ResolvedView.SkyAtmosphereCameraAerialPerspectiveVolumeDepthResolution,
-			ResolvedView.SkyAtmosphereCameraAerialPerspectiveVolumeStartDepth,
+			ResolvedView.SkyAtmosphereAerialPerspectiveStartDepth,
 			ResolvedView.SkyAtmosphereCameraAerialPerspectiveVolumeDepthSliceLength,
 			ResolvedView.SkyAtmosphereCameraAerialPerspectiveVolumeDepthSliceLengthInv, 
 			OneOverPreExposure, Fogging);
diff --git a/Engine/Shaders/Private/BasePassVertexShader.usf b/Engine/Shaders/Private/BasePassVertexShader.usf
--- a/Engine/Shaders/Private/BasePassVertexShader.usf
+++ b/Engine/Shaders/Private/BasePassVertexShader.usf
@@ -138,7 +138,7 @@ void Main(
 			View.CameraAerialPerspectiveVolume, View.CameraAerialPerspectiveVolumeSampler,
 			ResolvedView.SkyAtmosphereCameraAerialPerspectiveVolumeDepthResolutionInv,
 			ResolvedView.SkyAtmosphereCameraAerialPerspectiveVolumeDepthResolution,
-			ResolvedView.SkyAtmosphereCameraAerialPerspectiveVolumeStartDepth,
+			ResolvedView.SkyAtmosphereAerialPerspectiveStartDepth,
 			ResolvedView.SkyAtmosphereCameraAerialPerspectiveVolumeDepthSliceLength,
 			ResolvedView.SkyAtmosphereCameraAerialPerspectiveVolumeDepthSliceLengthInv,
 			OneOverPreExposure, Output.BasePassInterpolants.VertexFog);
diff --git a/Engine/Shaders/Private/MaterialTemplate.ush b/Engine/Shaders/Private/MaterialTemplate.ush
--- a/Engine/Shaders/Private/MaterialTemplate.ush
+++ b/Engine/Shaders/Private/MaterialTemplate.ush
@@ -1084,7 +1084,7 @@ float4 MaterialExpressionSkyAtmosphereAerialPerspective(FMaterialPixelParameters
 		View.CameraAerialPerspectiveVolume, View.CameraAerialPerspectiveVolumeSampler,
 		ResolvedView.SkyAtmosphereCameraAerialPerspectiveVolumeDepthResolutionInv,
 		ResolvedView.SkyAtmosphereCameraAerialPerspectiveVolumeDepthResolution,
-		ResolvedView.SkyAtmosphereCameraAerialPerspectiveVolumeStartDepth,
+		ResolvedView.SkyAtmosphereAerialPerspectiveStartDepth,
 		ResolvedView.SkyAtmosphereCameraAerialPerspectiveVolumeDepthSliceLength,
 		ResolvedView.SkyAtmosphereCameraAerialPerspectiveVolumeDepthSliceLengthInv,
 		OneOverPreExposure);
diff --git a/Engine/Shaders/Private/MobileBasePassVertexShader.usf b/Engine/Shaders/Private/MobileBasePassVertexShader.usf
--- a/Engine/Shaders/Private/MobileBasePassVertexShader.usf
+++ b/Engine/Shaders/Private/MobileBasePassVertexShader.usf
@@ -99,7 +99,7 @@ void Main(
 			View.CameraAerialPerspectiveVolume, View.CameraAerialPerspectiveVolumeSampler,
 			ResolvedView.SkyAtmosphereCameraAerialPerspectiveVolumeDepthResolutionInv,
 			ResolvedView.SkyAtmosphereCameraAerialPerspectiveVolumeDepthResolution,
-			ResolvedView.SkyAtmosphereCameraAerialPerspectiveVolumeStartDepth,
+			ResolvedView.SkyAtmosphereAerialPerspectiveStartDepth,
 			ResolvedView.SkyAtmosphereCameraAerialPerspectiveVolumeDepthSliceLength,
 			ResolvedView.SkyAtmosphereCameraAerialPerspectiveVolumeDepthSliceLengthInv,
 			OneOverPreExposure, VertexFog);
diff --git a/Engine/Shaders/Private/SkyAtmosphere.usf b/Engine/Shaders/Private/SkyAtmosphere.usf
--- a/Engine/Shaders/Private/SkyAtmosphere.usf
+++ b/Engine/Shaders/Private/SkyAtmosphere.usf
@@ -16,18 +16,12 @@
 
 
 
-#ifndef GROUND_GI_ENABLED
-#define GROUND_GI_ENABLED 0
-#endif
-#ifndef TRANSMITANCE_METHOD
-#define TRANSMITANCE_METHOD 2
+#ifndef COLORED_TRANSMITTANCE_ENABLED	// Never used, UE4 does not supports dual blending
+#define COLORED_TRANSMITTANCE_ENABLED 0 
 #endif
 #ifndef MULTISCATTERING_APPROX_ENABLED 
 #define MULTISCATTERING_APPROX_ENABLED 0 
 #endif
-#ifndef COLORED_TRANSMITTANCE_ENABLED	// Never used, UE4 does not supports dual blending
-#define COLORED_TRANSMITTANCE_ENABLED 0 
-#endif
 #ifndef FASTSKY_ENABLED 
 #define FASTSKY_ENABLED 0 
 #endif
@@ -37,15 +31,15 @@
 #ifndef SOURCE_DISK_ENABLED
 #define SOURCE_DISK_ENABLED 0
 #endif
-#ifndef ONLY_SKY_PIXEL_ENABLED
-#define ONLY_SKY_PIXEL_ENABLED 0
-#endif
 #ifndef PER_PIXEL_NOISE
 #define PER_PIXEL_NOISE 0
 #endif
 #ifndef SECOND_ATMOSPHERE_LIGHT_ENABLED
 #define SECOND_ATMOSPHERE_LIGHT_ENABLED 0
 #endif
+#ifndef RENDERSKY_ENABLED
+#define RENDERSKY_ENABLED 0
+#endif
 
 #ifndef TRANSMITTANCE_PASS
 #define TRANSMITTANCE_PASS 0
@@ -149,15 +143,6 @@ bool MoveToTopAtmosphere(inout float3 WorldPos, in float3 WorldDir, in float Atm
 
 
 
-
-
-
-
-
-
-
-
-
 ////////////////////////////////////////////////////////////
 // Participating medium properties
 ////////////////////////////////////////////////////////////
@@ -415,13 +400,8 @@ float SkyAtmosphereNoise(float2 UV)
 
 
 
-
-
-
-
-
 ////////////////////////////////////////////////////////////
-// Main scattering/transmitance integration fucntion
+// Main scattering/transmitance integration function
 ////////////////////////////////////////////////////////////
 
 struct SingleScatteringResult
@@ -687,12 +667,9 @@ SingleScatteringResult IntegrateSingleScatteredLuminance(
 
 
 
-
-
-
-
-
-
+////////////////////////////////////////////////////////////
+// Main sky and atmosphere on opaque ray marching shaders
+////////////////////////////////////////////////////////////
 
 float StartDepthZ;
 
@@ -702,15 +679,9 @@ void SkyAtmosphereVS(
 {
 	float2 UV = -1.0f;
 	UV = VertexId == 1 ? float2(-1.0f, 3.0f) : UV;
-	UV = VertexId == 2 ? float2(3.0f, -1.0f) : UV;
+	UV = VertexId == 2 ? float2( 3.0f,-1.0f) : UV;
 
-#if ONLY_SKY_PIXEL_ENABLED
-	const float Depth = FarDepthValue;
-#else
-	const float Depth = StartDepthZ;
-#endif
-
-	Position = float4(UV, Depth, 1.0f);
+	Position = float4(UV, StartDepthZ, 1.0f);
 }
 
 
@@ -746,17 +717,10 @@ void RenderSkyAtmosphereRayMarchingPS(
 	float3 WorldPos = GetCameraPlanetPos();
 	float3 WorldDir = GetScreenWorldDir(SVPos);
 
-	float3 PreExposedL = 0;
-#if  ONLY_SKY_PIXEL_ENABLED
-	// In this case we assume all pixel have a FarDepthValue to avoid reading the depth buffer. 
-	// Others pixels will be skipped by depth test.
-	float DeviceZ = FarDepthValue;
-#else
-	float DeviceZ = LookupDeviceZ(UvBuffer);
-#endif
-
 	// Get the light disk luminance to draw 
+	float3 PreExposedL = 0;
 	float3 LuminanceScale = 1.0f;
+	float DeviceZ = LookupDeviceZ(UvBuffer);
 	if (DeviceZ == FarDepthValue)
 	{
 		LuminanceScale = SkyAtmosphere.SkyLuminanceFactor;
@@ -764,10 +728,17 @@ void RenderSkyAtmosphereRayMarchingPS(
 #if SECOND_ATMOSPHERE_LIGHT_ENABLED
 		PreExposedL += GetLightDiskLuminance(WorldPos, WorldDir, 1);
 #endif
+
+#if RENDERSKY_ENABLED==0
+		// We should not render the sky and the current pixels are at far depth, so simply early exit.
+		// We enable depth bound when supported to not have to even process those pixels.
+		OutLuminance = PrepareOutput(float3(0.0f, 0.0f, 0.0f), float3(1.0f, 1.0f, 1.0f));
+		return;
+#endif
 	}
 
 	float ViewHeight = length(WorldPos);
-#if FASTSKY_ENABLED
+#if FASTSKY_ENABLED && RENDERSKY_ENABLED
 	if (ViewHeight < Atmosphere.TopRadius && DeviceZ == FarDepthValue)
 	{
 		float2 UV;
@@ -795,9 +766,6 @@ void RenderSkyAtmosphereRayMarchingPS(
 
 #if FASTAERIALPERSPECTIVE_ENABLED
 
-#if  ONLY_SKY_PIXEL_ENABLED==0
-	// We only compute the aerial perspective only when "not only rendering the sky" pixels. We simply remove that code section to simplify the shader.
-
 #if COLORED_TRANSMITTANCE_ENABLED
 #error The FASTAERIALPERSPECTIVE_ENABLED path does not support COLORED_TRANSMITTANCE_ENABLED.
 #else
@@ -812,7 +780,7 @@ void RenderSkyAtmosphereRayMarchingPS(
 		CameraAerialPerspectiveVolumeTexture, CameraAerialPerspectiveVolumeTextureSampler,
 		SkyAtmosphere.CameraAerialPerspectiveVolumeDepthResolutionInv,
 		SkyAtmosphere.CameraAerialPerspectiveVolumeDepthResolution,
-		SkyAtmosphere.CameraAerialPerspectiveVolumeStartDepth,
+		SkyAtmosphere.AerialPerspectiveStartDepth,
 		SkyAtmosphere.CameraAerialPerspectiveVolumeDepthSliceLength,
 		SkyAtmosphere.CameraAerialPerspectiveVolumeDepthSliceLengthInv,
 		OneOverPreExposure);
@@ -824,11 +792,9 @@ void RenderSkyAtmosphereRayMarchingPS(
 	return;
 #endif
 
-#endif // ONLY_SKY_PIXEL_ENABLED==0
-
 #else // FASTAERIALPERSPECTIVE_ENABLED
 
-	WorldPos += WorldDir * SkyAtmosphere.CameraAerialPerspectiveVolumeStartDepth; // Offset according to start depth
+	WorldPos += WorldDir * SkyAtmosphere.AerialPerspectiveStartDepth; // Offset according to start depth
 
 	// Move to top atmosphere as the starting point for ray marching.
 	// This is critical to be after the above to not disrupt above atmosphere tests and voxel selection.
@@ -871,6 +837,10 @@ void RenderSkyAtmosphereRayMarchingPS(
 
 
 
+////////////////////////////////////////////////////////////
+// Transmittance LUT
+////////////////////////////////////////////////////////////
+
 #ifndef THREADGROUP_SIZE
 #define THREADGROUP_SIZE 1
 #endif
@@ -916,6 +886,10 @@ void RenderTransmittanceLutCS(uint3 ThreadId : SV_DispatchThreadID)
 
 
 
+////////////////////////////////////////////////////////////
+// Multi-scattering LUT
+////////////////////////////////////////////////////////////
+
 RWTexture2D<float3> MultiScatteredLuminanceLutUAV;
 
 [numthreads(THREADGROUP_SIZE, THREADGROUP_SIZE, 1)]
@@ -1005,6 +979,10 @@ void RenderMultiScatteredLuminanceLutCS(uint3 ThreadId : SV_DispatchThreadID)
 
 
 
+////////////////////////////////////////////////////////////
+// Sky View LUT
+////////////////////////////////////////////////////////////
+
 RWTexture2D<float3> SkyViewLutUAV;
 
 [numthreads(THREADGROUP_SIZE, THREADGROUP_SIZE, 1)]
@@ -1066,6 +1044,10 @@ void RenderSkyViewLutCS(uint3 ThreadId : SV_DispatchThreadID)
 
 
 
+////////////////////////////////////////////////////////////
+// Distant sky light LUT
+////////////////////////////////////////////////////////////
+
 RWTexture2D<float3> DistantSkyLightLutUAV;
 Buffer<float4> UniformSphereSamplesBuffer;
 float4 AtmosphereLightDirection0;
@@ -1161,6 +1143,10 @@ void RenderDistantSkyLightLutCS(uint3 ThreadId : SV_DispatchThreadID)
 
 
 
+////////////////////////////////////////////////////////////
+// Camera aerial perspective volume LUT
+////////////////////////////////////////////////////////////
+
 RWTexture3D<float4> CameraAerialPerspectiveVolumeUAV;
 
 [numthreads(THREADGROUP_SIZE, THREADGROUP_SIZE, THREADGROUP_SIZE)]
@@ -1177,7 +1163,7 @@ void RenderCameraAerialPerspectiveVolumeCS(uint3 ThreadId : SV_DispatchThreadID)
 	Slice *= Slice;	// squared distribution
 	Slice *= SkyAtmosphere.CameraAerialPerspectiveVolumeDepthResolution;
 
-	float3 RayStartWorldPos = CamPos + SkyAtmosphere.CameraAerialPerspectiveVolumeStartDepth * WorldDir; // Offset according to start depth
+	float3 RayStartWorldPos = CamPos + SkyAtmosphere.AerialPerspectiveStartDepth * WorldDir; // Offset according to start depth
 	float ViewHeight;
 
 
@@ -1230,7 +1216,7 @@ void RenderCameraAerialPerspectiveVolumeCS(uint3 ThreadId : SV_DispatchThreadID)
 		}
 		 
 		WorldDir = normalize(VoxelWorldPos - CamPos);
-		RayStartWorldPos = CamPos + SkyAtmosphere.CameraAerialPerspectiveVolumeStartDepth * WorldDir; // Offset according to start depth
+		RayStartWorldPos = CamPos + SkyAtmosphere.AerialPerspectiveStartDepth * WorldDir; // Offset according to start depth
 		tMax = length(VoxelWorldPos - RayStartWorldPos);
 	}
 	float tMaxMax = tMax;
@@ -1280,9 +1266,9 @@ void RenderCameraAerialPerspectiveVolumeCS(uint3 ThreadId : SV_DispatchThreadID)
 
 
 
-//////////////////////////////////////////////////////////////////////////
-// Debug shaders
-//////////////////////////////////////////////////////////////////////////
+////////////////////////////////////////////////////////////
+// Debug
+////////////////////////////////////////////////////////////
 
 float ViewPortWidth;
 float ViewPortHeight;
diff --git a/Engine/Source/Editor/EditorStyle/Private/SlateEditorStyle.cpp b/Engine/Source/Editor/EditorStyle/Private/SlateEditorStyle.cpp
--- a/Engine/Source/Editor/EditorStyle/Private/SlateEditorStyle.cpp
+++ b/Engine/Source/Editor/EditorStyle/Private/SlateEditorStyle.cpp
@@ -6184,6 +6184,7 @@ void FSlateEditorStyle::FStyle::SetupClassIconsAndThumbnails()
 		Set("ClassIcon.VectorFieldComponent", new IMAGE_BRUSH("Icons/ActorIcons/VectorFieldVolume_16x", Icon16x16));
 		Set("ClassIcon.ArrowComponent", new IMAGE_BRUSH("Icons/ActorIcons/Arrow_16px", Icon16x16));
 		Set("ClassIcon.AtmosphericFogComponent", new IMAGE_BRUSH("Icons/AssetIcons/AtmosphericFog_16x", Icon16x16));
+		Set("ClassIcon.SkyAtmosphereComponent", new IMAGE_BRUSH("Icons/AssetIcons/AtmosphericFog_16x", Icon16x16));
 		Set("ClassIcon.BoxComponent", new IMAGE_BRUSH("Icons/ActorIcons/Box_16px", Icon16x16));
 		Set("ClassIcon.CapsuleComponent", new IMAGE_BRUSH("Icons/ActorIcons/Capsule_16px", Icon16x16));
 		Set("ClassIcon.InstancedStaticMeshComponent", new IMAGE_BRUSH("Icons/ActorIcons/InstancedStaticMesh_16px", Icon16x16));
@@ -6434,6 +6435,14 @@ void FSlateEditorStyle::FStyle::SetupClassIconsAndThumbnails()
 			Set( *FString::Printf(TEXT("ClassIcon.%s"), Type),		new IMAGE_BRUSH(FString::Printf(TEXT("Icons/AssetIcons/%s_%dx"), Type, 16), Icon16x16 ) );
 			Set( *FString::Printf(TEXT("ClassThumbnail.%s"), Type),	new IMAGE_BRUSH(FString::Printf(TEXT("Icons/AssetIcons/%s_%dx"), Type, 64), Icon64x64 ) );
 		}
+
+		{
+			// For now, the SkyAtmosphere component copy the style of the AtmosphericFog component.
+			const TCHAR* Type = TEXT("SkyAtmosphere");
+			const TCHAR* Icon = TEXT("AtmosphericFog");
+			Set(*FString::Printf(TEXT("ClassIcon.%s"), Type), new IMAGE_BRUSH(FString::Printf(TEXT("Icons/AssetIcons/%s_%dx"), Icon, 16), Icon16x16));
+			Set(*FString::Printf(TEXT("ClassThumbnail.%s"), Type), new IMAGE_BRUSH(FString::Printf(TEXT("Icons/AssetIcons/%s_%dx"), Icon, 64), Icon64x64));
+		}
 	}
 #endif
 }
diff --git a/Engine/Source/Editor/PlacementMode/Private/PlacementModeModule.cpp b/Engine/Source/Editor/PlacementMode/Private/PlacementModeModule.cpp
--- a/Engine/Source/Editor/PlacementMode/Private/PlacementModeModule.cpp
+++ b/Engine/Source/Editor/PlacementMode/Private/PlacementModeModule.cpp
@@ -21,6 +21,7 @@
 #include "ActorFactories/ActorFactoryEmptyActor.h"
 #include "ActorFactories/ActorFactoryPawn.h"
 #include "ActorFactories/ActorFactoryExponentialHeightFog.h"
+#include "ActorFactories/ActorFactorySkyAtmosphere.h"
 #include "ActorFactories/ActorFactoryPlayerStart.h"
 #include "ActorFactories/ActorFactoryPointLight.h"
 #include "ActorFactories/ActorFactorySpotLight.h"
@@ -174,12 +175,13 @@ void FPlacementModeModule::StartupModule()
 
 		FPlacementCategory* Category = Categories.Find(CategoryName);
 		Category->Items.Add(CreateID(), MakeShareable(new FPlaceableItem(PPFactory, FAssetData(APostProcessVolume::StaticClass()), SortOrder += 10)));
-		Category->Items.Add(CreateID(), MakeShareable(new FPlaceableItem(*UActorFactoryAtmosphericFog::StaticClass(), SortOrder += 10)));
+		Category->Items.Add(CreateID(), MakeShareable(new FPlaceableItem(*UActorFactorySkyAtmosphere::StaticClass(), SortOrder += 10)));
 		Category->Items.Add(CreateID(), MakeShareable(new FPlaceableItem(*UActorFactoryExponentialHeightFog::StaticClass(), SortOrder += 10)));
 		Category->Items.Add(CreateID(), MakeShareable(new FPlaceableItem(*UActorFactorySphereReflectionCapture::StaticClass(), SortOrder += 10)));
 		Category->Items.Add(CreateID(), MakeShareable(new FPlaceableItem(*UActorFactoryBoxReflectionCapture::StaticClass(), SortOrder += 10)));
 		Category->Items.Add(CreateID(), MakeShareable(new FPlaceableItem(*UActorFactoryPlanarReflection::StaticClass(), SortOrder += 10)));
 		Category->Items.Add(CreateID(), MakeShareable(new FPlaceableItem(*UActorFactoryDeferredDecal::StaticClass(), SortOrder += 10)));
+		Category->Items.Add(CreateID(), MakeShareable(new FPlaceableItem(*UActorFactoryAtmosphericFog::StaticClass(), SortOrder += 10)));
 	}
 
 	RegisterPlacementCategory(
diff --git a/Engine/Source/Editor/UnrealEd/Classes/ActorFactories/ActorFactorySkyAtmosphere.h b/Engine/Source/Editor/UnrealEd/Classes/ActorFactories/ActorFactorySkyAtmosphere.h
new file mode 100644
--- /dev/null
+++ b/Engine/Source/Editor/UnrealEd/Classes/ActorFactories/ActorFactorySkyAtmosphere.h
@@ -0,0 +1,16 @@
+// Copyright 1998-2019 Epic Games, Inc. All Rights Reserved.
+
+
+#pragma once
+
+#include "CoreMinimal.h"
+#include "UObject/ObjectMacros.h"
+#include "ActorFactories/ActorFactory.h"
+#include "ActorFactorySkyAtmosphere.generated.h"
+
+UCLASS(MinimalAPI, config = Editor, collapsecategories, hidecategories = Object)
+class UActorFactorySkyAtmosphere : public UActorFactory
+{
+	GENERATED_UCLASS_BODY()
+};
+
diff --git a/Engine/Source/Editor/UnrealEd/Private/Factories/ActorFactory.cpp b/Engine/Source/Editor/UnrealEd/Private/Factories/ActorFactory.cpp
--- a/Engine/Source/Editor/UnrealEd/Private/Factories/ActorFactory.cpp
+++ b/Engine/Source/Editor/UnrealEd/Private/Factories/ActorFactory.cpp
@@ -16,6 +16,7 @@ ActorFactory.cpp:
 #include "Model.h"
 #include "ActorFactories/ActorFactoryAmbientSound.h"
 #include "ActorFactories/ActorFactoryAtmosphericFog.h"
+#include "ActorFactories/ActorFactorySkyAtmosphere.h"
 #include "ActorFactories/ActorFactoryBlueprint.h"
 #include "ActorFactories/ActorFactoryBoxReflectionCapture.h"
 #include "ActorFactories/ActorFactoryBoxVolume.h"
@@ -92,6 +93,7 @@ ActorFactory.cpp:
 #include "VectorField/VectorFieldVolume.h"
 #include "Components/DecalComponent.h"
 #include "Components/BillboardComponent.h"
+#include "Components/SkyAtmosphereComponent.h"
 #include "Engine/BlueprintGeneratedClass.h"
 #include "Animation/AnimBlueprintGeneratedClass.h"
 #include "Engine/Polys.h"
@@ -1599,6 +1601,16 @@ UActorFactoryAtmosphericFog::UActorFactoryAtmosphericFog(const FObjectInitialize
 	NewActorClass = AAtmosphericFog::StaticClass();
 }
 
+/*-----------------------------------------------------------------------------
+UActorFactorySkyAtmosphere
+-----------------------------------------------------------------------------*/
+UActorFactorySkyAtmosphere::UActorFactorySkyAtmosphere(const FObjectInitializer& ObjectInitializer)
+	: Super(ObjectInitializer)
+{
+	DisplayName = LOCTEXT("SkyAtmosphereDisplayName", "Sky Atmosphere");
+	NewActorClass = ASkyAtmosphere::StaticClass();
+}
+
 /*-----------------------------------------------------------------------------
 UActorFactoryExponentialHeightFog
 -----------------------------------------------------------------------------*/
diff --git a/Engine/Source/Runtime/Engine/Public/SceneView.h b/Engine/Source/Runtime/Engine/Public/SceneView.h
--- a/Engine/Source/Runtime/Engine/Public/SceneView.h
+++ b/Engine/Source/Runtime/Engine/Public/SceneView.h
@@ -711,7 +711,7 @@ enum ETranslucencyVolumeCascade
 	VIEW_UNIFORM_BUFFER_MEMBER(FVector4, SkyViewLutSizeAndInvSize) \
 	VIEW_UNIFORM_BUFFER_MEMBER(float, SkyAtmosphereBottomRadius) \
 	VIEW_UNIFORM_BUFFER_MEMBER(float, SkyAtmosphereTopRadius) \
-	VIEW_UNIFORM_BUFFER_MEMBER(float, SkyAtmosphereCameraAerialPerspectiveVolumeStartDepth) \
+	VIEW_UNIFORM_BUFFER_MEMBER(float, SkyAtmosphereAerialPerspectiveStartDepth) \
 	VIEW_UNIFORM_BUFFER_MEMBER(float, SkyAtmosphereCameraAerialPerspectiveVolumeDepthResolution) \
 	VIEW_UNIFORM_BUFFER_MEMBER(float, SkyAtmosphereCameraAerialPerspectiveVolumeDepthResolutionInv) \
 	VIEW_UNIFORM_BUFFER_MEMBER(float, SkyAtmosphereCameraAerialPerspectiveVolumeDepthSliceLength) \
diff --git a/Engine/Source/Runtime/Renderer/Private/MobileShadingRenderer.cpp b/Engine/Source/Runtime/Renderer/Private/MobileShadingRenderer.cpp
--- a/Engine/Source/Runtime/Renderer/Private/MobileShadingRenderer.cpp
+++ b/Engine/Source/Runtime/Renderer/Private/MobileShadingRenderer.cpp
@@ -547,12 +547,6 @@ void FMobileSceneRenderer::Render(FRHICommandListImmediate& RHICmdList)
 
 	// scene depth is read only and can be fetched
 	RHICmdList.NextSubpass();
-
-	// Draw the sky atmosphere
-	if (bShouldRenderSkyAtmosphere)
-	{
-		RenderSkyAtmosphere(RHICmdList);
-	}
 		
 	// Split if we need to render translucency in a separate render pass
 	if (bRequiresTranslucencyPass)
diff --git a/Engine/Source/Runtime/Renderer/Private/SceneRendering.cpp b/Engine/Source/Runtime/Renderer/Private/SceneRendering.cpp
--- a/Engine/Source/Runtime/Renderer/Private/SceneRendering.cpp
+++ b/Engine/Source/Runtime/Renderer/Private/SceneRendering.cpp
@@ -1229,7 +1229,7 @@ void FViewInfo::SetupUniformBufferParameters(
 
 		FSkyAtmosphereViewSharedUniformShaderParameters OutParameters;
 		SetupSkyAtmosphereViewSharedUniformShaderParameters(*this, OutParameters);
-		ViewUniformShaderParameters.SkyAtmosphereCameraAerialPerspectiveVolumeStartDepth = OutParameters.CameraAerialPerspectiveVolumeStartDepth;
+		ViewUniformShaderParameters.SkyAtmosphereAerialPerspectiveStartDepth = OutParameters.AerialPerspectiveStartDepth;
 		ViewUniformShaderParameters.SkyAtmosphereCameraAerialPerspectiveVolumeDepthResolution = OutParameters.CameraAerialPerspectiveVolumeDepthResolution;
 		ViewUniformShaderParameters.SkyAtmosphereCameraAerialPerspectiveVolumeDepthResolutionInv = OutParameters.CameraAerialPerspectiveVolumeDepthResolutionInv;
 		ViewUniformShaderParameters.SkyAtmosphereCameraAerialPerspectiveVolumeDepthSliceLength = OutParameters.CameraAerialPerspectiveVolumeDepthSliceLength;
@@ -1261,7 +1261,7 @@ void FViewInfo::SetupUniformBufferParameters(
 		ViewUniformShaderParameters.SkyViewLutSizeAndInvSize = FVector4(1.0f, 1.0f, 1.0f, 1.0f);
 		ViewUniformShaderParameters.SkyAtmosphereBottomRadius = 1.0f;
 		ViewUniformShaderParameters.SkyAtmosphereTopRadius = 1.0f;
-		ViewUniformShaderParameters.SkyAtmosphereCameraAerialPerspectiveVolumeStartDepth = 1.0f;
+		ViewUniformShaderParameters.SkyAtmosphereAerialPerspectiveStartDepth = 1.0f;
 		ViewUniformShaderParameters.SkyAtmosphereCameraAerialPerspectiveVolumeDepthResolution = 1.0f;
 		ViewUniformShaderParameters.SkyAtmosphereCameraAerialPerspectiveVolumeDepthResolutionInv = 1.0f;
 		ViewUniformShaderParameters.SkyAtmosphereCameraAerialPerspectiveVolumeDepthSliceLength = 1.0f;
diff --git a/Engine/Source/Runtime/Renderer/Private/SkyAtmosphereRendering.cpp b/Engine/Source/Runtime/Renderer/Private/SkyAtmosphereRendering.cpp
--- a/Engine/Source/Runtime/Renderer/Private/SkyAtmosphereRendering.cpp
+++ b/Engine/Source/Runtime/Renderer/Private/SkyAtmosphereRendering.cpp
@@ -33,6 +33,13 @@ static TAutoConsoleVariable<int32> CVarSupportSkyAtmosphere(
 
 ////////////////////////////////////////////////////////////////////////// Regular sky 
 
+static TAutoConsoleVariable<int32> CVarSkyAtmosphereRelyOnSkyDome(
+	TEXT("r.SkyAtmosphere.RelyOnSkyDome"), 0,
+	TEXT("When true, the sky pixels will not be rendered by the full screen sky and aerialpesepctive pass.\n")
+	TEXT("Sky pixels will be skipped using depth bound test when available and\n")
+	TEXT("only the aerial perspective will be evaluated on opaque mesh pixels.\n"),
+	ECVF_RenderThreadSafe);
+
 static TAutoConsoleVariable<float> CVarSkyAtmosphereSampleCountMin(
 	TEXT("r.SkyAtmosphere.SampleCountMin"), 2.0f,
 	TEXT("The minimum sample count used to compute sky/atmosphere scattering and transmittance.\n")
@@ -106,8 +113,8 @@ static TAutoConsoleVariable<float> CVarSkyAtmosphereAerialPerspectiveLUTDepthRes
 	ECVF_RenderThreadSafe | ECVF_Scalability);
 
 static TAutoConsoleVariable<float> CVarSkyAtmosphereAerialPerspectiveLUTDepth(
-	TEXT("r.SkyAtmosphere.AerialPerspectiveLUT.Depth"), 8.0f,
-	TEXT("The length of the LUT in kilometers (default = 8km). After, the scattering and transmittance will be constantly the last slice."),
+	TEXT("r.SkyAtmosphere.AerialPerspectiveLUT.Depth"), 96.0f,
+	TEXT("The length of the LUT in kilometers (default = 96km to get nice cloud/atmosphere interactions in the distance for default sky). Further than this distance, the last slice is used."),
 	ECVF_RenderThreadSafe | ECVF_Scalability);
 
 static TAutoConsoleVariable<float> CVarSkyAtmosphereAerialPerspectiveLUTSampleCountPerSlice(
@@ -205,12 +212,12 @@ BEGIN_GLOBAL_SHADER_PARAMETER_STRUCT(FSkyAtmosphereInternalCommonParameters, )
 	SHADER_PARAMETER(float, FastSkyDistanceToSampleCountMaxInv)
 
 	SHADER_PARAMETER(FVector4, CameraAerialPerspectiveVolumeSizeAndInvSize)
-	SHADER_PARAMETER(float, CameraAerialPerspectiveVolumeStartDepth)
 	SHADER_PARAMETER(float, CameraAerialPerspectiveVolumeDepthResolution)
 	SHADER_PARAMETER(float, CameraAerialPerspectiveVolumeDepthResolutionInv)
 	SHADER_PARAMETER(float, CameraAerialPerspectiveVolumeDepthSliceLength)
 	SHADER_PARAMETER(float, CameraAerialPerspectiveVolumeDepthSliceLengthInv)
 	SHADER_PARAMETER(float, CameraAerialPerspectiveSampleCountPerSlice)
+	SHADER_PARAMETER(float, AerialPerspectiveStartDepth)
 
 	SHADER_PARAMETER(FVector4, TransmittanceLutSizeAndInvSize)
 	SHADER_PARAMETER(FVector4, MultiScatteredLuminanceLutSizeAndInvSize)
@@ -241,8 +248,8 @@ IMPLEMENT_GLOBAL_SHADER_PARAMETER_STRUCT(FSkyAtmosphereInternalCommonParameters,
 	int32 SkyViewLutHeight = ValidateLUTResolution(CVarSkyAtmosphereFastSkyLUTHeight.GetValueOnRenderThread()); \
 	int32 CameraAerialPerspectiveVolumeScreenResolution = ValidateLUTResolution(CVarSkyAtmosphereAerialPerspectiveLUTWidth.GetValueOnRenderThread()); \
 	int32 CameraAerialPerspectiveVolumeDepthResolution = ValidateLUTResolution(CVarSkyAtmosphereAerialPerspectiveLUTDepthResolution.GetValueOnRenderThread()); \
-	float CameraAerialPerspectiveVolumeStartDepth = CVarSkyAtmosphereAerialPerspectiveStartDepth.GetValueOnRenderThread(); \
-	CameraAerialPerspectiveVolumeStartDepth = CameraAerialPerspectiveVolumeStartDepth < 0.0f ? 0.0f : CameraAerialPerspectiveVolumeStartDepth; \
+	float AerialPerspectiveStartDepth = CVarSkyAtmosphereAerialPerspectiveStartDepth.GetValueOnRenderThread(); \
+	AerialPerspectiveStartDepth = AerialPerspectiveStartDepth < 0.0f ? 0.0f : AerialPerspectiveStartDepth; \
 	float CameraAerialPerspectiveVolumeDepth = CVarSkyAtmosphereAerialPerspectiveLUTDepth.GetValueOnRenderThread(); \
 	CameraAerialPerspectiveVolumeDepth = CameraAerialPerspectiveVolumeDepth < 1.0f ? 1.0f : CameraAerialPerspectiveVolumeDepth;	/* 1 kilometer minimum */ \
 	float CameraAerialPerspectiveVolumeDepthSliceLength = CameraAerialPerspectiveVolumeDepth / CameraAerialPerspectiveVolumeDepthResolution;
@@ -272,12 +279,13 @@ void SetupSkyAtmosphereViewSharedUniformShaderParameters(const class FViewInfo&
 	}
 
 	OutParameters.ApplyCameraAerialPerspectiveVolume = View.SkyAtmosphereCameraAerialPerspectiveVolume == nullptr ? 0.0f : 1.0f;
-	OutParameters.CameraAerialPerspectiveVolumeStartDepth = CameraAerialPerspectiveVolumeStartDepth;
 	OutParameters.CameraAerialPerspectiveVolumeDepthResolution = float(CameraAerialPerspectiveVolumeDepthResolution);
 	OutParameters.CameraAerialPerspectiveVolumeDepthResolutionInv = 1.0f / OutParameters.CameraAerialPerspectiveVolumeDepthResolution;
 	OutParameters.CameraAerialPerspectiveVolumeDepthSliceLength = CameraAerialPerspectiveVolumeDepthSliceLength;
 	OutParameters.CameraAerialPerspectiveVolumeDepthSliceLengthInv = 1.0f / OutParameters.CameraAerialPerspectiveVolumeDepthSliceLength;
 
+	OutParameters.AerialPerspectiveStartDepth = AerialPerspectiveStartDepth;
+
 	SetBlackAlpha13DIfNull(SkyAtmosphereCameraAerialPerspectiveVolume); // Needs to be after we set ApplyCameraAerialPerspectiveVolume
 }
 
@@ -427,20 +435,21 @@ class FSkyPermutationMultiScatteringApprox : SHADER_PERMUTATION_BOOL("MULTISCATT
 class FFastSky : SHADER_PERMUTATION_BOOL("FASTSKY_ENABLED");
 class FFastAerialPespective : SHADER_PERMUTATION_BOOL("FASTAERIALPERSPECTIVE_ENABLED");
 class FSourceDiskEnabled : SHADER_PERMUTATION_BOOL("SOURCE_DISK_ENABLED");
-class FOnlySkyPixel : SHADER_PERMUTATION_BOOL("ONLY_SKY_PIXEL_ENABLED");
 class FSecondAtmosphereLight : SHADER_PERMUTATION_BOOL("SECOND_ATMOSPHERE_LIGHT_ENABLED");
+class FRenderSky : SHADER_PERMUTATION_BOOL("RENDERSKY_ENABLED");
 
 }
 
+//////////////////////////////////////////////////////////////////////////
+
 class FRenderSkyAtmosphereVS : public FGlobalShader
 {
 	DECLARE_GLOBAL_SHADER(FRenderSkyAtmosphereVS);
 	SHADER_USE_PARAMETER_STRUCT(FRenderSkyAtmosphereVS, FGlobalShader);
 
-	using FPermutationDomain = TShaderPermutationDomain<FOnlySkyPixel>;
+	using FPermutationDomain = TShaderPermutationDomain<>;
 
 	BEGIN_SHADER_PARAMETER_STRUCT(FParameters, )
-		SHADER_PARAMETER_STRUCT_REF(FViewUniformShaderParameters, ViewUniformBuffer)
 		SHADER_PARAMETER(float, StartDepthZ)
 	END_SHADER_PARAMETER_STRUCT()
 
@@ -461,20 +470,18 @@ class FRenderSkyAtmosphereVS : public FGlobalShader
 };
 IMPLEMENT_GLOBAL_SHADER(FRenderSkyAtmosphereVS, "/Engine/Private/SkyAtmosphere.usf", "SkyAtmosphereVS", SF_Vertex);
 
-//////////////////////////////////////////////////////////////////////////
-
 class FRenderSkyAtmospherePS : public FGlobalShader
 {
 	DECLARE_GLOBAL_SHADER(FRenderSkyAtmospherePS);
 	SHADER_USE_PARAMETER_STRUCT(FRenderSkyAtmospherePS, FGlobalShader);
 
-	using FPermutationDomain = TShaderPermutationDomain<FSkyPermutationMultiScatteringApprox, FFastSky, FFastAerialPespective, FSourceDiskEnabled, FOnlySkyPixel, FSecondAtmosphereLight>;
+	using FPermutationDomain = TShaderPermutationDomain<FSkyPermutationMultiScatteringApprox, FFastSky, FFastAerialPespective, FSourceDiskEnabled, FSecondAtmosphereLight, FRenderSky>;
 
 	BEGIN_SHADER_PARAMETER_STRUCT(FParameters, )
 		SHADER_PARAMETER_STRUCT_REF(FAtmosphereUniformShaderParameters, Atmosphere)
 		SHADER_PARAMETER_STRUCT_REF(FSkyAtmosphereInternalCommonParameters, SkyAtmosphere)
 		SHADER_PARAMETER_STRUCT_REF(FViewUniformShaderParameters, ViewUniformBuffer)
-//		RENDER_TARGET_BINDING_SLOTS()
+		RENDER_TARGET_BINDING_SLOTS()
 		SHADER_PARAMETER_STRUCT_REF(FSceneTexturesUniformParameters, SceneTextures)
 		SHADER_PARAMETER_TEXTURE(Texture2D<float3>, TransmittanceLutTexture)
 		SHADER_PARAMETER_TEXTURE(Texture2D<float3>, MultiScatteredLuminanceLutTexture)
@@ -488,11 +495,43 @@ class FRenderSkyAtmospherePS : public FGlobalShader
 
 	static FPermutationDomain RemapPermutation(FPermutationDomain PermutationVector)
 	{
+		// If not rendering the sky, ignore the fastsky and sundisk permutations
+		if (PermutationVector.Get<FRenderSky>() == false)
+		{
+			PermutationVector.Set<FFastSky>(false);
+			PermutationVector.Set<FSourceDiskEnabled>(false);
+		}
+
+		// See comment below
+		if (PermutationVector.Get<FFastAerialPespective>())
+		{
+			PermutationVector.Set<FSkyPermutationMultiScatteringApprox>(false);
+		}
+
 		return PermutationVector;
 	}
 
 	static bool ShouldCompilePermutation(const FGlobalShaderPermutationParameters& Parameters)
 	{
+		FPermutationDomain PermutationVector(Parameters.PermutationId);
+
+		// If not rendering the sky, ignore the FFastSky and FSourceDiskEnabled permutations
+		if (PermutationVector.Get<FRenderSky>() == false && (
+			PermutationVector.Get<FFastSky>() ||
+			PermutationVector.Get<FSourceDiskEnabled>() ))
+		{
+			return false;
+		}
+
+		// When rendering using FFastAerialPespective, we can ignore FSkyPermutationMultiScatteringApprox because it has been all baked in already:
+		//	- Aerial perspective will contain it using the camera volume
+		//	- Sky will contain it using SkyView LUT or camera volume. 
+		// This is possible because we do not allow FFastAerialPespective to be used with ray-marched sky. It must be used with the FFastSky SkyView LUT oir nothing.
+		if (PermutationVector.Get<FFastAerialPespective>() && PermutationVector.Get<FSkyPermutationMultiScatteringApprox>())
+		{
+			return false;
+		}
+
 		return ShouldPipelineCompileSkyAtmosphereShader(Parameters.Platform);
 	}
 
@@ -876,7 +915,7 @@ static void SetupSkyAtmosphereInternalCommonParameters(
 	InternalCommonParameters.FastSkySampleCountMax = CVarSkyAtmosphereFastSkyLUTSampleCountMax.GetValueOnRenderThread();
 	float FastSkyDistanceToSampleCountMaxInv = CVarSkyAtmosphereFastSkyLUTDistanceToSampleCountMax.GetValueOnRenderThread();
 
-	InternalCommonParameters.CameraAerialPerspectiveVolumeStartDepth = CameraAerialPerspectiveVolumeStartDepth;
+	InternalCommonParameters.AerialPerspectiveStartDepth = AerialPerspectiveStartDepth;
 	InternalCommonParameters.CameraAerialPerspectiveVolumeDepthResolution = float(CameraAerialPerspectiveVolumeDepthResolution);
 	InternalCommonParameters.CameraAerialPerspectiveVolumeDepthResolutionInv = 1.0f / InternalCommonParameters.CameraAerialPerspectiveVolumeDepthResolution;
 	InternalCommonParameters.CameraAerialPerspectiveVolumeDepthSliceLength = CameraAerialPerspectiveVolumeDepthSliceLength;
@@ -946,46 +985,39 @@ void FSceneRenderer::RenderSkyAtmosphereLookUpTables(FRHICommandListImmediate& R
 	TUniformBufferRef<FSkyAtmosphereInternalCommonParameters> InternalCommonParametersRef = TUniformBufferRef<FSkyAtmosphereInternalCommonParameters>::CreateUniformBufferImmediate(InternalCommonParameters, UniformBuffer_SingleFrame);
 
 	UnbindRenderTargets(RHICmdList);
+	FRDGBuilder GraphBuilder(RHICmdList);
 
 	// Transmittance LUT
 	TShaderMap<FGlobalShaderType>* GlobalShaderMap = GetGlobalShaderMap(FeatureLevel);
 	{
 		TShaderMapRef<FRenderTransmittanceLutCS> ComputeShader(GlobalShaderMap);
-		RHICmdList.SetComputeShader(ComputeShader->GetComputeShader());
-
-		FRenderTransmittanceLutCS::FParameters PassParameters;
-		PassParameters.Atmosphere = Scene->GetSkyAtmosphereSceneInfo()->GetAtmosphereUniformBuffer();
-		PassParameters.SkyAtmosphere = InternalCommonParametersRef;
-		PassParameters.TransmittanceLutUAV = SkyInfo.GetTransmittanceLutTexture()->GetRenderTargetItem().UAV;
 
-		SetShaderParameters(RHICmdList, *ComputeShader, ComputeShader->GetComputeShader(), PassParameters);
+		FRenderTransmittanceLutCS::FParameters * PassParameters = GraphBuilder.AllocParameters<FRenderTransmittanceLutCS::FParameters>();
+		PassParameters->Atmosphere = Scene->GetSkyAtmosphereSceneInfo()->GetAtmosphereUniformBuffer();
+		PassParameters->SkyAtmosphere = InternalCommonParametersRef;
+		PassParameters->TransmittanceLutUAV = SkyInfo.GetTransmittanceLutTexture()->GetRenderTargetItem().UAV;
 
 		FIntVector TextureSize = SkyInfo.GetTransmittanceLutTexture()->GetDesc().GetSize();
 		TextureSize.Z = 1;
 		const FIntVector NumGroups = FIntVector::DivideAndRoundUp(TextureSize, FRenderTransmittanceLutCS::GroupSize);
-		RHICmdList.DispatchComputeShader(NumGroups.X, NumGroups.Y, NumGroups.Z);
-		UnsetShaderUAVs(RHICmdList, *ComputeShader, ComputeShader->GetComputeShader());
+		FComputeShaderUtils::AddPass(GraphBuilder, RDG_EVENT_NAME("TransmittanceLut"), *ComputeShader, PassParameters, NumGroups);
 	}
 
 	// Mean Illuminance LUT
 	{
 		TShaderMapRef<FRenderMultiScatteredLuminanceLutCS> ComputeShader(GlobalShaderMap);
-		RHICmdList.SetComputeShader(ComputeShader->GetComputeShader());
 
-		FRenderMultiScatteredLuminanceLutCS::FParameters PassParameters;
-		PassParameters.Atmosphere = Scene->GetSkyAtmosphereSceneInfo()->GetAtmosphereUniformBuffer();
-		PassParameters.SkyAtmosphere = InternalCommonParametersRef;
-		PassParameters.TransmittanceLutTextureSampler = SamplerLinearClamp;
-		PassParameters.TransmittanceLutTexture = SkyInfo.GetTransmittanceLutTexture()->GetRenderTargetItem().ShaderResourceTexture;
-		PassParameters.MultiScatteredLuminanceLutUAV = SkyInfo.GetMultiScatteredLuminanceLutTexture()->GetRenderTargetItem().UAV;
-
-		SetShaderParameters(RHICmdList, *ComputeShader, ComputeShader->GetComputeShader(), PassParameters);
+		FRenderMultiScatteredLuminanceLutCS::FParameters * PassParameters = GraphBuilder.AllocParameters<FRenderMultiScatteredLuminanceLutCS::FParameters>();
+		PassParameters->Atmosphere = Scene->GetSkyAtmosphereSceneInfo()->GetAtmosphereUniformBuffer();
+		PassParameters->SkyAtmosphere = InternalCommonParametersRef;
+		PassParameters->TransmittanceLutTextureSampler = SamplerLinearClamp;
+		PassParameters->TransmittanceLutTexture = SkyInfo.GetTransmittanceLutTexture()->GetRenderTargetItem().ShaderResourceTexture;
+		PassParameters->MultiScatteredLuminanceLutUAV = SkyInfo.GetMultiScatteredLuminanceLutTexture()->GetRenderTargetItem().UAV;
 
 		FIntVector TextureSize = SkyInfo.GetMultiScatteredLuminanceLutTexture()->GetDesc().GetSize();
 		TextureSize.Z = 1;
 		const FIntVector NumGroups = FIntVector::DivideAndRoundUp(TextureSize, FRenderMultiScatteredLuminanceLutCS::GroupSize);
-		RHICmdList.DispatchComputeShader(NumGroups.X, NumGroups.Y, NumGroups.Z);
-		UnsetShaderUAVs(RHICmdList, *ComputeShader, ComputeShader->GetComputeShader());
+		FComputeShaderUtils::AddPass(GraphBuilder, RDG_EVENT_NAME("MeanIllumLut"), *ComputeShader, PassParameters, NumGroups);
 	}
 
 	// Distant Sky Light LUT
@@ -998,38 +1030,33 @@ void FSceneRenderer::RenderSkyAtmosphereLookUpTables(FRHICommandListImmediate& R
 
 		using FPermutationDomain = TShaderPermutationDomain<FSkyPermutationMultiScatteringApprox, FSecondAtmosphereLight>;
 
-		RHICmdList.SetComputeShader(ComputeShader->GetComputeShader());
-
-		FRenderDistantSkyLightLutCS::FParameters PassParameters;
-		PassParameters.Atmosphere = Scene->GetSkyAtmosphereSceneInfo()->GetAtmosphereUniformBuffer();
-		PassParameters.SkyAtmosphere = InternalCommonParametersRef;
-		PassParameters.TransmittanceLutTextureSampler = SamplerLinearClamp;
-		PassParameters.MultiScatteredLuminanceLutTextureSampler = SamplerLinearClamp;
-		PassParameters.TransmittanceLutTexture = SkyInfo.GetTransmittanceLutTexture()->GetRenderTargetItem().ShaderResourceTexture;
-		PassParameters.MultiScatteredLuminanceLutTexture = SkyInfo.GetMultiScatteredLuminanceLutTexture()->GetRenderTargetItem().ShaderResourceTexture;
-		PassParameters.UniformSphereSamplesBuffer = GUniformSphereSamplesBuffer.UniformSphereSamplesBuffer.SRV;
-		PassParameters.DistantSkyLightLutUAV = SkyInfo.GetDistantSkyLightLutTexture()->GetRenderTargetItem().UAV;
+		FRenderDistantSkyLightLutCS::FParameters * PassParameters = GraphBuilder.AllocParameters<FRenderDistantSkyLightLutCS::FParameters>();
+		PassParameters->Atmosphere = Scene->GetSkyAtmosphereSceneInfo()->GetAtmosphereUniformBuffer();
+		PassParameters->SkyAtmosphere = InternalCommonParametersRef;
+		PassParameters->TransmittanceLutTextureSampler = SamplerLinearClamp;
+		PassParameters->MultiScatteredLuminanceLutTextureSampler = SamplerLinearClamp;
+		PassParameters->TransmittanceLutTexture = SkyInfo.GetTransmittanceLutTexture()->GetRenderTargetItem().ShaderResourceTexture;
+		PassParameters->MultiScatteredLuminanceLutTexture = SkyInfo.GetMultiScatteredLuminanceLutTexture()->GetRenderTargetItem().ShaderResourceTexture;
+		PassParameters->UniformSphereSamplesBuffer = GUniformSphereSamplesBuffer.UniformSphereSamplesBuffer.SRV;
+		PassParameters->DistantSkyLightLutUAV = SkyInfo.GetDistantSkyLightLutTexture()->GetRenderTargetItem().UAV;
 
 		FLightSceneInfo* Light0 = Scene->AtmosphereLights[0];
 		FLightSceneInfo* Light1 = Scene->AtmosphereLights[1];
 		if (Light0)
 		{
-			PassParameters.AtmosphereLightDirection0 = -Light0->Proxy->GetDirection();
-			PassParameters.AtmosphereLightColor0 = Light0->Proxy->GetColor();
+			PassParameters->AtmosphereLightDirection0 = -Light0->Proxy->GetDirection();
+			PassParameters->AtmosphereLightColor0 = Light0->Proxy->GetColor();
 		}
 		if (Light1)
 		{
-			PassParameters.AtmosphereLightDirection1 = -Light1->Proxy->GetDirection();
-			PassParameters.AtmosphereLightColor1 = Light1->Proxy->GetColor();
+			PassParameters->AtmosphereLightDirection1 = -Light1->Proxy->GetDirection();
+			PassParameters->AtmosphereLightColor1 = Light1->Proxy->GetColor();
 		}
-		PassParameters.DistantSkyLightSampleAltitude = CVarSkyAtmosphereDistantSkyLightLUTAltitude.GetValueOnAnyThread();
-
-		SetShaderParameters(RHICmdList, *ComputeShader, ComputeShader->GetComputeShader(), PassParameters);
+		PassParameters->DistantSkyLightSampleAltitude = CVarSkyAtmosphereDistantSkyLightLUTAltitude.GetValueOnAnyThread();
 
 		FIntVector TextureSize = FIntVector(1, 1, 1);
 		const FIntVector NumGroups = FIntVector::DivideAndRoundUp(TextureSize, FRenderDistantSkyLightLutCS::GroupSize);
-		RHICmdList.DispatchComputeShader(NumGroups.X, NumGroups.Y, NumGroups.Z);
-		UnsetShaderUAVs(RHICmdList, *ComputeShader, ComputeShader->GetComputeShader());
+		FComputeShaderUtils::AddPass(GraphBuilder, RDG_EVENT_NAME("DistantSkyLightLut"), *ComputeShader, PassParameters, NumGroups);
 	}
 
 	for (int32 ViewIndex = 0; ViewIndex < Views.Num(); ViewIndex++)
@@ -1047,25 +1074,21 @@ void FSceneRenderer::RenderSkyAtmosphereLookUpTables(FRHICommandListImmediate& R
 			PermutationVector.Set<FSourceDiskEnabled>(bLightDiskEnabled);
 			PermutationVector.Set<FSecondAtmosphereLight>(bSecondAtmosphereLightEnabled);
 			TShaderMapRef<FRenderSkyViewLutCS> ComputeShader(View.ShaderMap, PermutationVector);
-			RHICmdList.SetComputeShader(ComputeShader->GetComputeShader());
 
-			FRenderSkyViewLutCS::FParameters PassParameters;
-			PassParameters.Atmosphere = Scene->GetSkyAtmosphereSceneInfo()->GetAtmosphereUniformBuffer();
-			PassParameters.SkyAtmosphere = InternalCommonParametersRef;
-			PassParameters.ViewUniformBuffer = View.ViewUniformBuffer;
-			PassParameters.TransmittanceLutTextureSampler = SamplerLinearClamp;
-			PassParameters.MultiScatteredLuminanceLutTextureSampler = SamplerLinearClamp;
-			PassParameters.TransmittanceLutTexture = SkyInfo.GetTransmittanceLutTexture()->GetRenderTargetItem().ShaderResourceTexture;
-			PassParameters.MultiScatteredLuminanceLutTexture = SkyInfo.GetMultiScatteredLuminanceLutTexture()->GetRenderTargetItem().ShaderResourceTexture;
-			PassParameters.SkyViewLutUAV = View.SkyAtmosphereViewLutTexture->GetRenderTargetItem().UAV;
-
-			SetShaderParameters(RHICmdList, *ComputeShader, ComputeShader->GetComputeShader(), PassParameters);
+			FRenderSkyViewLutCS::FParameters * PassParameters = GraphBuilder.AllocParameters<FRenderSkyViewLutCS::FParameters>();
+			PassParameters->Atmosphere = Scene->GetSkyAtmosphereSceneInfo()->GetAtmosphereUniformBuffer();
+			PassParameters->SkyAtmosphere = InternalCommonParametersRef;
+			PassParameters->ViewUniformBuffer = View.ViewUniformBuffer;
+			PassParameters->TransmittanceLutTextureSampler = SamplerLinearClamp;
+			PassParameters->MultiScatteredLuminanceLutTextureSampler = SamplerLinearClamp;
+			PassParameters->TransmittanceLutTexture = SkyInfo.GetTransmittanceLutTexture()->GetRenderTargetItem().ShaderResourceTexture;
+			PassParameters->MultiScatteredLuminanceLutTexture = SkyInfo.GetMultiScatteredLuminanceLutTexture()->GetRenderTargetItem().ShaderResourceTexture;
+			PassParameters->SkyViewLutUAV = View.SkyAtmosphereViewLutTexture->GetRenderTargetItem().UAV;
 
 			FIntVector TextureSize = View.SkyAtmosphereViewLutTexture->GetDesc().GetSize();
 			TextureSize.Z = 1;
 			const FIntVector NumGroups = FIntVector::DivideAndRoundUp(TextureSize, FRenderSkyViewLutCS::GroupSize);
-			RHICmdList.DispatchComputeShader(NumGroups.X, NumGroups.Y, NumGroups.Z);
-			UnsetShaderUAVs(RHICmdList, *ComputeShader, ComputeShader->GetComputeShader());
+			FComputeShaderUtils::AddPass(GraphBuilder, RDG_EVENT_NAME("SkyViewLut"), *ComputeShader, PassParameters, NumGroups);
 		}
 
 		// Camera Atmosphere Volume
@@ -1074,34 +1097,33 @@ void FSceneRenderer::RenderSkyAtmosphereLookUpTables(FRHICommandListImmediate& R
 			PermutationVector.Set<FSkyPermutationMultiScatteringApprox>(bMultiScattering);
 			PermutationVector.Set<FSecondAtmosphereLight>(bSecondAtmosphereLightEnabled);
 			TShaderMapRef<FRenderCameraAerialPerspectiveVolumeCS> ComputeShader(View.ShaderMap, PermutationVector);
-			RHICmdList.SetComputeShader(ComputeShader->GetComputeShader());
-
-			FRenderCameraAerialPerspectiveVolumeCS::FParameters PassParameters;
-			PassParameters.Atmosphere = Scene->GetSkyAtmosphereSceneInfo()->GetAtmosphereUniformBuffer();
-			PassParameters.SkyAtmosphere = InternalCommonParametersRef;
-			PassParameters.ViewUniformBuffer = View.ViewUniformBuffer;
-			PassParameters.TransmittanceLutTextureSampler = SamplerLinearClamp;
-			PassParameters.MultiScatteredLuminanceLutTextureSampler = SamplerLinearClamp;
-			PassParameters.TransmittanceLutTexture = SkyInfo.GetTransmittanceLutTexture()->GetRenderTargetItem().ShaderResourceTexture;
-			PassParameters.MultiScatteredLuminanceLutTexture = SkyInfo.GetMultiScatteredLuminanceLutTexture()->GetRenderTargetItem().ShaderResourceTexture;
-			PassParameters.CameraAerialPerspectiveVolumeUAV = View.SkyAtmosphereCameraAerialPerspectiveVolume->GetRenderTargetItem().UAV;
 
-			SetShaderParameters(RHICmdList, *ComputeShader, ComputeShader->GetComputeShader(), PassParameters);
+			FRenderCameraAerialPerspectiveVolumeCS::FParameters * PassParameters = GraphBuilder.AllocParameters<FRenderCameraAerialPerspectiveVolumeCS::FParameters>();
+			PassParameters->Atmosphere = Scene->GetSkyAtmosphereSceneInfo()->GetAtmosphereUniformBuffer();
+			PassParameters->SkyAtmosphere = InternalCommonParametersRef;
+			PassParameters->ViewUniformBuffer = View.ViewUniformBuffer;
+			PassParameters->TransmittanceLutTextureSampler = SamplerLinearClamp;
+			PassParameters->MultiScatteredLuminanceLutTextureSampler = SamplerLinearClamp;
+			PassParameters->TransmittanceLutTexture = SkyInfo.GetTransmittanceLutTexture()->GetRenderTargetItem().ShaderResourceTexture;
+			PassParameters->MultiScatteredLuminanceLutTexture = SkyInfo.GetMultiScatteredLuminanceLutTexture()->GetRenderTargetItem().ShaderResourceTexture;
+			PassParameters->CameraAerialPerspectiveVolumeUAV = View.SkyAtmosphereCameraAerialPerspectiveVolume->GetRenderTargetItem().UAV;
 
 			FIntVector TextureSize = View.SkyAtmosphereCameraAerialPerspectiveVolume->GetDesc().GetSize();
 			const FIntVector NumGroups = FIntVector::DivideAndRoundUp(TextureSize, FRenderCameraAerialPerspectiveVolumeCS::GroupSize);
-			RHICmdList.DispatchComputeShader(NumGroups.X, NumGroups.Y, NumGroups.Z);
-			UnsetShaderUAVs(RHICmdList, *ComputeShader, ComputeShader->GetComputeShader());
+			FComputeShaderUtils::AddPass(GraphBuilder, RDG_EVENT_NAME("CameraVolumeLut"), *ComputeShader, PassParameters, NumGroups);
 		}
 	}
 
-	// TODO
-	// RHICmdList.AutomaticCacheFlushAfterComputeShader(true);
-	// RHICmdList.FlushComputeShaderCache();
+	GraphBuilder.Execute();
+	// TODO have RDG execute those above passes with compute overlap similarly to using AutomaticCacheFlushAfterComputeShader(true);
 }
 
 void FSceneRenderer::RenderSkyAtmosphere(FRHICommandListImmediate& RHICmdList)
 {
+	// We never render the sky pass on such platform as it is too expensive. 
+	// In this case, the sky must be rendered as an opaque mesh using the shader graph as a composition tool.
+	check(!IsMobilePlatform(Scene->GetShaderPlatform()));
+
 	check(ShouldRenderSkyAtmosphere(Scene->GetSkyAtmosphereSceneInfo(), Scene->GetShaderPlatform())); // This should not be called if we should not render SkyAtmosphere
 
 	DECLARE_GPU_STAT(SkyAtmosphere);
@@ -1114,11 +1136,9 @@ void FSceneRenderer::RenderSkyAtmosphere(FRHICommandListImmediate& RHICmdList)
 	const bool bFastSky = CVarSkyAtmosphereFastSkyLUT.GetValueOnRenderThread() > 0;
 	const bool bFastAerialPespective = CVarSkyAtmosphereAerialPerspectiveApplyOnOpaque.GetValueOnRenderThread() > 0;
 	const bool bFastAerialPespectiveDepthTest = CVarSkyAtmosphereAerialPerspectiveDepthTest.GetValueOnRenderThread() > 0;
+	const bool bSkyRelyOnSkyDome = CVarSkyAtmosphereRelyOnSkyDome.GetValueOnAnyThread();
 	const bool bSecondAtmosphereLightEnabled = IsSecondAtmosphereLightEnabled(Scene);
 
-	FSceneRenderTargets& SceneContext = FSceneRenderTargets::Get(RHICmdList);
-
-	FRHIBlendState* PreMultipliedColorTransmittanceBlend = TStaticBlendState<CW_RGB, BO_Add, BF_One, BF_SourceAlpha, BO_Add, BF_Zero, BF_One>::GetRHI();
 	FRHISamplerState* SamplerLinearClamp = TStaticSamplerState<SF_Trilinear>::GetRHI();
 
 	// Initialise common internal parameters
@@ -1126,8 +1146,10 @@ void FSceneRenderer::RenderSkyAtmosphere(FRHICommandListImmediate& RHICmdList)
 	SetupSkyAtmosphereInternalCommonParameters(InternalCommonParameters, *Scene, SkyInfo);
 	TUniformBufferRef<FSkyAtmosphereInternalCommonParameters> InternalCommonParametersRef = TUniformBufferRef<FSkyAtmosphereInternalCommonParameters>::CreateUniformBufferImmediate(InternalCommonParameters, UniformBuffer_SingleFrame);
 
-	const bool bIsMobilePlatformAndOnlySkyPixel = IsMobilePlatform(Scene->GetShaderPlatform());
-	const bool bUseOpaqueDepthTestAtStartDepth = bFastAerialPespectiveDepthTest && !bIsMobilePlatformAndOnlySkyPixel;
+	FRDGBuilder GraphBuilder(RHICmdList);
+	FSceneRenderTargets& SceneContext = FSceneRenderTargets::Get(RHICmdList);
+	FRDGTextureRef SceneColor = GraphBuilder.RegisterExternalTexture(SceneContext.GetSceneColor(), TEXT("SceneColor"));
+	FRDGTextureRef SceneDepth = GraphBuilder.RegisterExternalTexture(SceneContext.SceneDepthZ, TEXT("SceneDepth"));
 
 	for (int32 ViewIndex = 0; ViewIndex < Views.Num(); ViewIndex++)
 	{
@@ -1143,91 +1165,112 @@ void FSceneRenderer::RenderSkyAtmosphere(FRHICommandListImmediate& RHICmdList)
 		const bool ForceRayMarching = (FVector::Distance(ViewOrigin, PlanetOrigin) - TopOfAtmosphere - SafeEdge) > 0.0f;
 
 		// Render the sky, and optionally the atmosphere aerial perspective, on the scene luminance buffer
-		if (!bIsMobilePlatformAndOnlySkyPixel)
-		{
-			// This is skipped on mobile as this pass is part of the BasePass
-			FRHITexture* RTs[1] = { SceneContext.GetSceneColor()->GetRenderTargetItem().TargetableTexture };
-			FRHIRenderPassInfo RPInfo(1, RTs, ERenderTargetActions::DontLoad_Store, SceneContext.GetSceneDepthTexture(), EDepthStencilTargetActions::LoadDepthNotStencil_DontStore, FExclusiveDepthStencil::DepthRead_StencilWrite);
-			RHICmdList.BeginRenderPass(RPInfo, TEXT("DrawUVDisplacement"));
-		}
 		{
-			FRenderSkyAtmospherePS::FPermutationDomain PermutationVector;
-			PermutationVector.Set<FSkyPermutationMultiScatteringApprox>(bMultiScattering);
-			PermutationVector.Set<FFastSky>(bFastSky && !ForceRayMarching);
-			PermutationVector.Set<FFastAerialPespective>(bFastAerialPespective && !ForceRayMarching);
-			PermutationVector.Set<FSourceDiskEnabled>(bLightDiskEnabled);
-			PermutationVector.Set<FOnlySkyPixel>(bIsMobilePlatformAndOnlySkyPixel);
-			PermutationVector.Set<FSecondAtmosphereLight>(bSecondAtmosphereLightEnabled);
-			TShaderMapRef<FRenderSkyAtmospherePS> PixelShader(View.ShaderMap, PermutationVector);
+			FRenderSkyAtmospherePS::FPermutationDomain PsPermutationVector;
+			PsPermutationVector.Set<FSkyPermutationMultiScatteringApprox>(bMultiScattering);
+			PsPermutationVector.Set<FFastSky>(bFastSky && !ForceRayMarching);
+			PsPermutationVector.Set<FFastAerialPespective>(bFastAerialPespective && !ForceRayMarching);
+			PsPermutationVector.Set<FSourceDiskEnabled>(bLightDiskEnabled);
+			PsPermutationVector.Set<FSecondAtmosphereLight>(bSecondAtmosphereLightEnabled);
+			PsPermutationVector.Set<FRenderSky>(!bSkyRelyOnSkyDome);
+			PsPermutationVector = FRenderSkyAtmospherePS::RemapPermutation(PsPermutationVector);
+			TShaderMapRef<FRenderSkyAtmospherePS> PixelShader(View.ShaderMap, PsPermutationVector);
 
 			FRenderSkyAtmosphereVS::FPermutationDomain VsPermutationVector;
-			VsPermutationVector.Set<FOnlySkyPixel>(bIsMobilePlatformAndOnlySkyPixel);
 			TShaderMapRef<FRenderSkyAtmosphereVS> VertexShader(View.ShaderMap, VsPermutationVector);
 
-			FRenderSkyAtmospherePS::FParameters PassParameters;
-			PassParameters.Atmosphere = Scene->GetSkyAtmosphereSceneInfo()->GetAtmosphereUniformBuffer();
-			PassParameters.SkyAtmosphere = InternalCommonParametersRef;
-			PassParameters.ViewUniformBuffer = View.ViewUniformBuffer;
-			//PassParameters->RenderTargets[0] = FRenderTargetBinding(SceneColor, ERenderTargetLoadAction::ELoad, ERenderTargetStoreAction::EStore);
-			PassParameters.SceneTextures = CreateSceneTextureUniformBufferSingleDraw(RHICmdList, ESceneTextureSetupMode::All, View.FeatureLevel);
-			PassParameters.TransmittanceLutTextureSampler = SamplerLinearClamp;
-			PassParameters.MultiScatteredLuminanceLutTextureSampler = SamplerLinearClamp;
-			PassParameters.SkyViewLutTextureSampler = SamplerLinearClamp;
-			PassParameters.CameraAerialPerspectiveVolumeTextureSampler = SamplerLinearClamp;
-
-			PassParameters.TransmittanceLutTexture = SkyInfo.GetTransmittanceLutTexture()->GetRenderTargetItem().ShaderResourceTexture;
-			PassParameters.MultiScatteredLuminanceLutTexture = SkyInfo.GetMultiScatteredLuminanceLutTexture()->GetRenderTargetItem().ShaderResourceTexture;
-			PassParameters.SkyViewLutTexture = View.SkyAtmosphereViewLutTexture->GetRenderTargetItem().ShaderResourceTexture;
-			PassParameters.CameraAerialPerspectiveVolumeTexture = View.SkyAtmosphereCameraAerialPerspectiveVolume->GetRenderTargetItem().ShaderResourceTexture;
-
-			FRenderSkyAtmosphereVS::FParameters VsPassParameters;
-			VsPassParameters.ViewUniformBuffer = View.ViewUniformBuffer;
-			VsPassParameters.StartDepthZ = 0.1f;
-			if (bUseOpaqueDepthTestAtStartDepth)
+			FRenderTargetBindingSlots RenderTargets;
+			RenderTargets[0] = FRenderTargetBinding(SceneColor, ERenderTargetLoadAction::ELoad);
+			RenderTargets.DepthStencil = FDepthStencilBinding(SceneDepth, ERenderTargetLoadAction::ELoad, ERenderTargetLoadAction::ENoAction, FExclusiveDepthStencil::DepthRead_StencilNop);
+
+			FRenderSkyAtmospherePS::FParameters* PsPassParameters = GraphBuilder.AllocParameters<FRenderSkyAtmospherePS::FParameters>();
+			PsPassParameters->Atmosphere = Scene->GetSkyAtmosphereSceneInfo()->GetAtmosphereUniformBuffer();
+			PsPassParameters->SkyAtmosphere = InternalCommonParametersRef;
+			PsPassParameters->ViewUniformBuffer = View.ViewUniformBuffer;
+			PsPassParameters->RenderTargets = RenderTargets;
+			PsPassParameters->SceneTextures = CreateSceneTextureUniformBufferSingleDraw(RHICmdList, ESceneTextureSetupMode::All, View.FeatureLevel);
+			PsPassParameters->TransmittanceLutTextureSampler = SamplerLinearClamp;
+			PsPassParameters->MultiScatteredLuminanceLutTextureSampler = SamplerLinearClamp;
+			PsPassParameters->SkyViewLutTextureSampler = SamplerLinearClamp;
+			PsPassParameters->CameraAerialPerspectiveVolumeTextureSampler = SamplerLinearClamp;
+			PsPassParameters->TransmittanceLutTexture = SkyInfo.GetTransmittanceLutTexture()->GetRenderTargetItem().ShaderResourceTexture;
+			PsPassParameters->MultiScatteredLuminanceLutTexture = SkyInfo.GetMultiScatteredLuminanceLutTexture()->GetRenderTargetItem().ShaderResourceTexture;
+			PsPassParameters->SkyViewLutTexture = View.SkyAtmosphereViewLutTexture->GetRenderTargetItem().ShaderResourceTexture;
+			PsPassParameters->CameraAerialPerspectiveVolumeTexture = View.SkyAtmosphereCameraAerialPerspectiveVolume->GetRenderTargetItem().ShaderResourceTexture;
+			ClearUnusedGraphResources(*PixelShader, PsPassParameters);
+
+			float StartDepthZ = 0.1f;
+			if (bFastAerialPespectiveDepthTest)
 			{
 				const FMatrix ViewProjectionMatrix = View.ViewMatrices.GetProjectionMatrix();
 				float HalfHorizontalFOV = FMath::Atan(1.0f / ViewProjectionMatrix.M[0][0]);
 				float HalfVerticalFOV = FMath::Atan(1.0f / ViewProjectionMatrix.M[1][1]);
 				const float HalfFOV = View.FOV * 0.5f * PI / 180.0f;
-				const float StartDepthView = FMath::Cos(FMath::Max(HalfHorizontalFOV, HalfVerticalFOV)) * InternalCommonParameters.CameraAerialPerspectiveVolumeStartDepth * KmToCm;
+				const float StartDepthView = FMath::Cos(FMath::Max(HalfHorizontalFOV, HalfVerticalFOV)) * InternalCommonParameters.AerialPerspectiveStartDepth * KmToCm;
 				const FVector4 Projected = ViewProjectionMatrix.TransformFVector4(FVector4(0.0f, 0.0f, StartDepthView, 1.0f));
-				VsPassParameters.StartDepthZ = Projected.Z / Projected.W;
+				StartDepthZ = Projected.Z / Projected.W;
 			}
 
-			FGraphicsPipelineStateInitializer GraphicsPSOInit;
-			
-			RHICmdList.ApplyCachedRenderTargets(GraphicsPSOInit);
-
-			GraphicsPSOInit.RasterizerState = TStaticRasterizerState<FM_Solid, CM_None>::GetRHI(); ;
-			GraphicsPSOInit.BlendState = PreMultipliedColorTransmittanceBlend;
-			if (bIsMobilePlatformAndOnlySkyPixel || bUseOpaqueDepthTestAtStartDepth)
-			{
-				GraphicsPSOInit.DepthStencilState = TStaticDepthStencilState<false, CF_DepthNearOrEqual>::GetRHI();
-			}
-			else
+			FIntRect Viewport = View.ViewRect;
+			GraphBuilder.AddPass(
+				RDG_EVENT_NAME("SkyAtmosphere"),
+				PsPassParameters,
+				ERDGPassFlags::Raster,
+				[PsPassParameters, VertexShader, PixelShader, Viewport, bFastAerialPespectiveDepthTest, bSkyRelyOnSkyDome, StartDepthZ](FRHICommandList& RHICmdList)
 			{
-				GraphicsPSOInit.DepthStencilState = TStaticDepthStencilState<false, CF_Always>::GetRHI();
-			}
+				RHICmdList.SetViewport(Viewport.Min.X, Viewport.Min.Y, 0.0f, Viewport.Max.X, Viewport.Max.Y, 1.0f);
 
-			GraphicsPSOInit.BoundShaderState.VertexDeclarationRHI = GEmptyVertexDeclaration.VertexDeclarationRHI;
-			GraphicsPSOInit.BoundShaderState.VertexShaderRHI = GETSAFERHISHADER_VERTEX(*VertexShader);
-			GraphicsPSOInit.BoundShaderState.PixelShaderRHI = GETSAFERHISHADER_PIXEL(*PixelShader);
-			GraphicsPSOInit.PrimitiveType = PT_TriangleList;
+				FGraphicsPipelineStateInitializer GraphicsPSOInit;
+				RHICmdList.ApplyCachedRenderTargets(GraphicsPSOInit);
 
-			RHICmdList.SetViewport(0, 0, 0.0f, View.ViewRect.Width(), View.ViewRect.Height(), 1.0f);
+				GraphicsPSOInit.BlendState = TStaticBlendState<CW_RGB, BO_Add, BF_One, BF_SourceAlpha, BO_Add, BF_Zero, BF_One>::GetRHI();
+				GraphicsPSOInit.RasterizerState = TStaticRasterizerState<FM_Solid, CM_None>::GetRHI();
+				if (bFastAerialPespectiveDepthTest)
+				{
+					GraphicsPSOInit.DepthStencilState = TStaticDepthStencilState<false, CF_DepthNearOrEqual>::GetRHI();
+				}
+				else
+				{
+					GraphicsPSOInit.DepthStencilState = TStaticDepthStencilState<false, CF_Always>::GetRHI();
+				}
 
-			SetGraphicsPipelineState(RHICmdList, GraphicsPSOInit);
+				GraphicsPSOInit.BoundShaderState.VertexDeclarationRHI = GEmptyVertexDeclaration.VertexDeclarationRHI;
+				GraphicsPSOInit.BoundShaderState.VertexShaderRHI = VertexShader->GetVertexShader();
+				GraphicsPSOInit.BoundShaderState.PixelShaderRHI = PixelShader->GetPixelShader();
+				GraphicsPSOInit.PrimitiveType = PT_TriangleList;
 
-			SetShaderParameters(RHICmdList, *PixelShader, PixelShader->GetPixelShader(), PassParameters);
-			SetShaderParameters(RHICmdList, *VertexShader, VertexShader->GetVertexShader(), VsPassParameters);
+				if (bSkyRelyOnSkyDome && GSupportsDepthBoundsTest)
+				{
+					// When we do not render the sky in the sky pass and depth bound test is supported, we take advantage of it in order to skip the processing of sky pixels.
+					GraphicsPSOInit.bDepthBounds = true;
+					if (bool(ERHIZBuffer::IsInverted))
+					{
+						//const float SmallestFloatAbove0 = 1.1754943508e-38;		// 32bit float depth
+						const float SmallestFloatAbove0 = 1.0f / 16777215.0f;		// 24bit norm depth
+						RHICmdList.SetDepthBounds(SmallestFloatAbove0, 1.0f);		// Tested on dx12 PC
+					}
+					else
+					{
+						//const float SmallestFloatBelow1 = 0.9999999404;			// 32bit float depth
+						const float SmallestFloatBelow1 = 16777214.0f / 16777215.0f;// 24bit norm depth
+						RHICmdList.SetDepthBounds(0.0f, SmallestFloatBelow1);		// Untested
+					}
+				}
 
-			RHICmdList.DrawPrimitive(0, 1, 1);
-		}
-		if (!bIsMobilePlatformAndOnlySkyPixel)
-		{
-			RHICmdList.EndRenderPass();
+				SetGraphicsPipelineState(RHICmdList, GraphicsPSOInit);
+
+				SetShaderParameters(RHICmdList, *PixelShader, PixelShader->GetPixelShader(), *PsPassParameters);
+
+				FRenderSkyAtmosphereVS::FParameters VsPassParameters;
+				VsPassParameters.StartDepthZ = StartDepthZ;
+				SetShaderParameters(RHICmdList, *VertexShader, VertexShader->GetVertexShader(), VsPassParameters);
+
+				RHICmdList.DrawPrimitive(0, 1, 1);
+			});
 		}
 	}
+	GraphBuilder.Execute();
+
 
 #if WITH_EDITOR
 	if (CVarSkyAtmosphereFastSkyLUT.GetValueOnAnyThread() == 0 && CVarSkyAtmosphereAerialPerspectiveApplyOnOpaque.GetValueOnAnyThread() > 0)
@@ -1283,7 +1326,7 @@ void FDeferredShadingSceneRenderer::RenderDebugSkyAtmosphere(FRHICommandListImme
 
 		FRDGBuilder GraphBuilder(RHICmdList);
 		FRDGTextureRef SceneColor = GraphBuilder.RegisterExternalTexture(SceneContext.GetSceneColor(), TEXT("SceneColor"));
-		FRDGTextureRef SceneDepth = GraphBuilder.RegisterExternalTexture(SceneContext.SceneDepthZ, TEXT("SceneColor"));
+		FRDGTextureRef SceneDepth = GraphBuilder.RegisterExternalTexture(SceneContext.SceneDepthZ, TEXT("SceneDepth"));
 
 		FRHIBlendState* PreMultipliedColorTransmittanceBlend = TStaticBlendState<CW_RGB, BO_Add, BF_One, BF_SourceAlpha, BO_Add, BF_Zero, BF_One>::GetRHI();
 		FRHIDepthStencilState* DepthStencilStateWrite = TStaticDepthStencilState<true, CF_Always>::GetRHI();
diff --git a/Engine/Source/Runtime/Renderer/Private/SkyAtmosphereRendering.h b/Engine/Source/Runtime/Renderer/Private/SkyAtmosphereRendering.h
--- a/Engine/Source/Runtime/Renderer/Private/SkyAtmosphereRendering.h
+++ b/Engine/Source/Runtime/Renderer/Private/SkyAtmosphereRendering.h
@@ -42,7 +42,7 @@ END_GLOBAL_SHADER_PARAMETER_STRUCT()
 // These parameters are shared on the view global uniform buffer and are dynamically changed with cvars.
 struct FSkyAtmosphereViewSharedUniformShaderParameters
 {
-	float CameraAerialPerspectiveVolumeStartDepth;
+	float AerialPerspectiveStartDepth;
 	float CameraAerialPerspectiveVolumeDepthResolution;
 	float CameraAerialPerspectiveVolumeDepthResolutionInv;
 	float CameraAerialPerspectiveVolumeDepthSliceLength;
