Disable TAA when view doesn't have a view state and add check. It's not a valid state and removes the need to check against this invalid state in the post processing chain. Also moved FViewInfo validation into FSceneView since the checked members are actually members of FSceneView.

diff --git a/Engine/Source/Runtime/Engine/Private/SceneView.cpp b/Engine/Source/Runtime/Engine/Private/SceneView.cpp
--- a/Engine/Source/Runtime/Engine/Private/SceneView.cpp
+++ b/Engine/Source/Runtime/Engine/Private/SceneView.cpp
@@ -765,6 +765,17 @@ FSceneView::FSceneView(const FSceneViewInitOptions& InitOptions)
 #if DO_CHECK
 bool FSceneView::VerifyMembersChecks() const
 {
+	if (PrimaryScreenPercentageMethod == EPrimaryScreenPercentageMethod::TemporalUpscale)
+	{
+		checkf(GetFeatureLevel() >= ERHIFeatureLevel::SM5, TEXT("Temporal upsample is SM5 only."));
+		checkf(AntiAliasingMethod == AAM_TemporalAA, TEXT("ScreenPercentageMethod == EPrimaryScreenPercentageMethod::TemporalUpscale requires AntiAliasingMethod == AAM_TemporalAA"));
+	}
+
+	if (AntiAliasingMethod == AAM_TemporalAA)
+	{
+		checkf(State, TEXT("TemporalAA requires the view to have a valid state."));
+	}
+
 	return true;
 }
 #endif
@@ -828,6 +839,12 @@ void FSceneView::SetupAntiAliasingMethod()
 			}
 		}
 	}
+
+    // TemporalAA requires view state for history.
+	if (AntiAliasingMethod == AAM_TemporalAA && !State)
+	{
+		AntiAliasingMethod = AAM_None;
+	}
 }
 
 static TAutoConsoleVariable<int32> CVarCompensateForFOV(
diff --git a/Engine/Source/Runtime/Renderer/Private/SceneRendering.cpp b/Engine/Source/Runtime/Renderer/Private/SceneRendering.cpp
--- a/Engine/Source/Runtime/Renderer/Private/SceneRendering.cpp
+++ b/Engine/Source/Runtime/Renderer/Private/SceneRendering.cpp
@@ -931,15 +931,7 @@ bool FViewInfo::VerifyMembersChecks() const
 {
 	FSceneView::VerifyMembersChecks();
 
-	// ------------------------------------------- Screen percentage checks.
-	checkf(
-		!(PrimaryScreenPercentageMethod == EPrimaryScreenPercentageMethod::TemporalUpscale && AntiAliasingMethod != AAM_TemporalAA),
-		TEXT("ScreenPercentageMethod == EPrimaryScreenPercentageMethod::TemporalUpscale requires AntiAliasingMethod == AAM_TemporalAA"));
-
-	if (PrimaryScreenPercentageMethod == EPrimaryScreenPercentageMethod::TemporalUpscale)
-	{
-		checkf(GetFeatureLevel() >= ERHIFeatureLevel::SM5, TEXT("Temporal upsample is SM5 only."));
-	}
+	check(ViewState == State);
 
 	return true;
 }
@@ -3446,10 +3438,13 @@ void FRendererModule::CreateAndInitSingleView(FRHICommandListImmediate& RHICmdLi
 
 void FRendererModule::BeginRenderingViewFamily(FCanvas* Canvas, FSceneViewFamily* ViewFamily)
 {
-	UWorld* World = nullptr; 
+	check(Canvas);
+	check(ViewFamily);
 	check(ViewFamily->Scene);
 	check(ViewFamily->GetScreenPercentageInterface());
 
+	UWorld* World = nullptr;
+
 	FScene* const Scene = ViewFamily->Scene->GetRenderScene();
 	if (Scene)
 	{
@@ -3469,9 +3464,9 @@ void FRendererModule::BeginRenderingViewFamily(FCanvas* Canvas, FSceneViewFamily
 
 	ENQUEUE_RENDER_COMMAND(UpdateFastVRamConfig)(
 		[](FRHICommandList& RHICmdList)
-	{
-		GFastVRamConfig.Update();
-	});
+		{
+			GFastVRamConfig.Update();
+		});
 
 	// Flush the canvas first.
 	Canvas->Flush_GameThread();
@@ -3484,8 +3479,8 @@ void FRendererModule::BeginRenderingViewFamily(FCanvas* Canvas, FSceneViewFamily
 	}
 	else
 	{
-	// this is passes to the render thread, better access that than GFrameNumberRenderThread
-	ViewFamily->FrameNumber = GFrameNumber;
+		// this is passes to the render thread, better access that than GFrameNumberRenderThread
+		ViewFamily->FrameNumber = GFrameNumber;
 	}
 
 #if !(UE_BUILD_SHIPPING || UE_BUILD_TEST)
@@ -3502,7 +3497,7 @@ void FRendererModule::BeginRenderingViewFamily(FCanvas* Canvas, FSceneViewFamily
 	}
 
 	if (Scene)
-	{		
+	{
 		// Set the world's "needs full lighting rebuild" flag if the scene has any uncached static lighting interactions.
 		if(World)
 		{
