Minor refactor of post process chain a bit for clarity.

diff --git a/Engine/Source/Runtime/Renderer/Private/PostProcess/DiaphragmDOF.cpp b/Engine/Source/Runtime/Renderer/Private/PostProcess/DiaphragmDOF.cpp
--- a/Engine/Source/Runtime/Renderer/Private/PostProcess/DiaphragmDOF.cpp
+++ b/Engine/Source/Runtime/Renderer/Private/PostProcess/DiaphragmDOF.cpp
@@ -1242,6 +1242,21 @@ IMPLEMENT_GLOBAL_SHADER(FDiaphragmDOFRecombineCS,        "/Engine/Private/Diaphr
 
 } // namespace
 
+bool DiaphragmDOF::IsEnabled(const FViewInfo& View)
+{
+	static const auto CVar = IConsoleManager::Get().FindTConsoleVariableDataInt(TEXT("r.DepthOfFieldQuality"));
+	check(CVar);
+
+	const bool bDepthOfFieldRequestedByCVar = CVar->GetValueOnRenderThread() > 0;
+
+	return
+		DiaphragmDOF::IsSupported(View.GetShaderPlatform()) &&
+		View.Family->EngineShowFlags.DepthOfField &&
+		bDepthOfFieldRequestedByCVar &&
+		View.FinalPostProcessSettings.DepthOfFieldFstop > 0 &&
+		View.FinalPostProcessSettings.DepthOfFieldFocalDistance > 0;
+}
+
 FRDGTextureRef DiaphragmDOF::AddPasses(
 	FRDGBuilder& GraphBuilder,
 	const FSceneTextureParameters& SceneTextures,
diff --git a/Engine/Source/Runtime/Renderer/Private/PostProcess/DiaphragmDOF.h b/Engine/Source/Runtime/Renderer/Private/PostProcess/DiaphragmDOF.h
--- a/Engine/Source/Runtime/Renderer/Private/PostProcess/DiaphragmDOF.h
+++ b/Engine/Source/Runtime/Renderer/Private/PostProcess/DiaphragmDOF.h
@@ -22,7 +22,10 @@ struct FTemporalAAHistory;
 
 namespace DiaphragmDOF
 {
-	
+
+// Whether DOF is enabled for the requested view.
+bool IsEnabled(const FViewInfo& View);
+
 float ComputeFocalLengthFromFov(const FSceneView& View);
 FVector4 CircleDofHalfCoc(const FViewInfo& View);
 
diff --git a/Engine/Source/Runtime/Renderer/Private/PostProcess/PostProcessMotionBlur.cpp b/Engine/Source/Runtime/Renderer/Private/PostProcess/PostProcessMotionBlur.cpp
--- a/Engine/Source/Runtime/Renderer/Private/PostProcess/PostProcessMotionBlur.cpp
+++ b/Engine/Source/Runtime/Renderer/Private/PostProcess/PostProcessMotionBlur.cpp
@@ -734,7 +734,7 @@ FRDGTextureRef AddMotionBlurFilterPass(
 	return ColorTextureOutput;
 }
 
-FRDGTextureRef VisualizeMotionBlur(
+FRDGTextureRef AddVisualizeMotionBlurPass(
 	FRDGBuilder& GraphBuilder,
 	const FScreenPassViewInfo& ScreenPassView,
 	FIntRect ColorViewportRect,
@@ -843,22 +843,6 @@ FRDGTextureRef AddMotionBlurPass(
 	check(DepthTexture);
 	check(VelocityTexture);
 
-	if (IsVisualizeMotionBlurEnabled(ScreenPassView.View))
-	{
-		return VisualizeMotionBlur(
-			GraphBuilder,
-			ScreenPassView,
-			ColorViewportRect,
-			VelocityViewportRect,
-			ColorTexture,
-			DepthTexture,
-			VelocityTexture);
-	}
-	else if (!IsMotionBlurEnabled(ScreenPassView.View))
-	{
-		return ColorTexture;
-	}
-
 	const FMotionBlurViewports Viewports(
 		FScreenPassTextureViewport(ColorViewportRect, ColorTexture),
 		FScreenPassTextureViewport(VelocityViewportRect, VelocityTexture));
diff --git a/Engine/Source/Runtime/Renderer/Private/PostProcess/PostProcessMotionBlur.h b/Engine/Source/Runtime/Renderer/Private/PostProcess/PostProcessMotionBlur.h
--- a/Engine/Source/Runtime/Renderer/Private/PostProcess/PostProcessMotionBlur.h
+++ b/Engine/Source/Runtime/Renderer/Private/PostProcess/PostProcessMotionBlur.h
@@ -28,6 +28,15 @@ enum class EMotionBlurQuality : uint32
 EMotionBlurQuality GetMotionBlurQuality();
 
 FRDGTextureRef AddMotionBlurPass(
+	FRDGBuilder& GraphBuilder,
+	const FScreenPassViewInfo& ScreenPassView,
+	FIntRect ColorViewportRect,
+	FIntRect VelocityViewportRect,
+	FRDGTextureRef ColorTexture,
+	FRDGTextureRef DepthTexture,
+	FRDGTextureRef VelocityTexture);
+
+FRDGTextureRef AddVisualizeMotionBlurPass(
 	FRDGBuilder& GraphBuilder,
 	const FScreenPassViewInfo& ScreenPassView,
 	FIntRect ColorViewportRect,
diff --git a/Engine/Source/Runtime/Renderer/Private/PostProcess/PostProcessTemporalAA.cpp b/Engine/Source/Runtime/Renderer/Private/PostProcess/PostProcessTemporalAA.cpp
--- a/Engine/Source/Runtime/Renderer/Private/PostProcess/PostProcessTemporalAA.cpp
+++ b/Engine/Source/Runtime/Renderer/Private/PostProcess/PostProcessTemporalAA.cpp
@@ -14,15 +14,15 @@
 #include "PostProcess/SceneFilterRendering.h"
 #include "CompositionLighting/PostProcessAmbientOcclusion.h"
 #include "PostProcess/PostProcessTonemap.h"
+#include "PostProcess/PostProcessMitchellNetravali.h"
 #include "ClearQuad.h"
 #include "PipelineStateCache.h"
 #include "PostProcessing.h"
-
-
 #include "RenderGraph.h"
 #include "SceneTextureParameters.h"
 #include "PixelShaderUtils.h"
 
+extern int32 GetPostProcessAAQuality();
 
 const int32 GTemporalAATileSizeX = 8;
 const int32 GTemporalAATileSizeY = 8;
@@ -72,7 +72,7 @@ static float CatmullRom( float x )
 		return ( 1.5f * ax - 2.5f ) * ax*ax + 1.0f;
 }
 
-float GetTemporalAAHistoryUpscaleFactor(const FViewInfo& View)
+static float GetTemporalAAHistoryUpscaleFactor(const FViewInfo& View)
 {
 	// We only support history upscale on PC with feature level SM5+
 	if (!IsPCPlatform(View.GetShaderPlatform()) || !IsFeatureLevelSupported(View.GetShaderPlatform(), ERHIFeatureLevel::SM5))
@@ -679,6 +679,98 @@ FTAAOutputs AddTemporalAAPass(
 	return Outputs;
 }
 
+void AddTemporalAAPass(
+	FRDGBuilder& GraphBuilder,
+	const FSceneTextureParameters& SceneTextures,
+	const FScreenPassViewInfo& ScreenPassView,
+	const bool bAllowDownsampleSceneColor,
+	const EPixelFormat DownsampleOverrideFormat,
+	FRDGTextureRef InSceneColorTexture,
+	FRDGTextureRef* OutSceneColorTexture,
+	FRDGTextureRef* OutSceneColorHalfResTexture,
+	FIntRect* OutSecondaryViewRect)
+{
+	const FViewInfo& View = ScreenPassView.View;
+
+	check(View.AntiAliasingMethod == AAM_TemporalAA && View.ViewState);
+
+	FTAAPassParameters TAAParameters(View);
+
+	TAAParameters.Pass = View.PrimaryScreenPercentageMethod == EPrimaryScreenPercentageMethod::TemporalUpscale
+		? ETAAPassConfig::MainUpsampling
+		: ETAAPassConfig::Main;
+
+	TAAParameters.SetupViewRect(View);
+
+	const int32 LowQualityTemporalAA = 3;
+
+	TAAParameters.bUseFast = GetPostProcessAAQuality() == LowQualityTemporalAA;
+
+	const FIntRect SecondaryViewRect = TAAParameters.OutputViewRect;
+
+	const float HistoryUpscaleFactor = GetTemporalAAHistoryUpscaleFactor(View);
+
+	// Configures TAA to upscale the history buffer; this is in addition to the secondary screen percentage upscale.
+	// We end up with a scene color that is larger than the secondary screen percentage. We immediately downscale
+	// afterwards using a Mitchel-Netravali filter.
+	if (HistoryUpscaleFactor > 1.0f)
+	{
+		const FIntPoint HistoryViewSize(
+			TAAParameters.OutputViewRect.Width() * HistoryUpscaleFactor,
+			TAAParameters.OutputViewRect.Height() * HistoryUpscaleFactor);
+
+		FIntPoint QuantizedMinHistorySize;
+		QuantizeSceneBufferSize(HistoryViewSize, QuantizedMinHistorySize);
+
+		TAAParameters.Pass = ETAAPassConfig::MainSuperSampling;
+		TAAParameters.bUseFast = false;
+
+		TAAParameters.OutputViewRect.Min.X = 0;
+		TAAParameters.OutputViewRect.Min.Y = 0;
+		TAAParameters.OutputViewRect.Max = HistoryViewSize;
+	}
+
+	TAAParameters.DownsampleOverrideFormat = DownsampleOverrideFormat;
+
+	TAAParameters.bDownsample = bAllowDownsampleSceneColor && TAAParameters.bUseFast;
+
+	TAAParameters.SceneColorInput = InSceneColorTexture;
+
+	const FTemporalAAHistory& InputHistory = View.PrevViewInfo.TemporalAAHistory;
+
+	FTemporalAAHistory& OutputHistory = View.ViewState->PrevFrameViewInfo.TemporalAAHistory;
+
+	const FTAAOutputs TAAOutputs = AddTemporalAAPass(
+		GraphBuilder,
+		SceneTextures,
+		View,
+		TAAParameters,
+		InputHistory,
+		&OutputHistory);
+
+	FRDGTextureRef SceneColorTexture = TAAOutputs.SceneColor;
+
+	// If we upscaled the history buffer, downsize back to the secondary screen percentage size.
+	if (HistoryUpscaleFactor > 1.0f)
+	{
+		const FIntRect InputViewport = TAAParameters.OutputViewRect;
+
+		FIntPoint QuantizedOutputSize;
+		QuantizeSceneBufferSize(SecondaryViewRect.Size(), QuantizedOutputSize);
+
+		FScreenPassTextureViewport OutputViewport;
+		OutputViewport.Rect = SecondaryViewRect;
+		OutputViewport.Extent.X = FMath::Max(InSceneColorTexture->Desc.Extent.X, QuantizedOutputSize.X);
+		OutputViewport.Extent.Y = FMath::Max(InSceneColorTexture->Desc.Extent.Y, QuantizedOutputSize.Y);
+
+		SceneColorTexture = ComputeMitchellNetravaliDownsample(GraphBuilder, ScreenPassView, SceneColorTexture, InputViewport, OutputViewport);
+	}
+
+	*OutSceneColorTexture = SceneColorTexture;
+	*OutSceneColorHalfResTexture = TAAOutputs.DownsampledSceneColor;
+	*OutSecondaryViewRect = SecondaryViewRect;
+}
+
 //////////////////////////////////////////////////////////////////////////
 //! Legacy - Only used by DebugViewModeRendering. Remove after porting.
 
diff --git a/Engine/Source/Runtime/Renderer/Private/PostProcess/PostProcessTemporalAA.h b/Engine/Source/Runtime/Renderer/Private/PostProcess/PostProcessTemporalAA.h
--- a/Engine/Source/Runtime/Renderer/Private/PostProcess/PostProcessTemporalAA.h
+++ b/Engine/Source/Runtime/Renderer/Private/PostProcess/PostProcessTemporalAA.h
@@ -6,8 +6,7 @@
 
 #pragma once
 
-#include "SceneRendering.h"
-#include "RenderGraph.h"
+#include "ScreenPass.h"
 
 class FSceneTextureParameters;
 struct FTemporalAAHistory;
@@ -49,8 +48,6 @@ static FORCEINLINE bool IsDOFTAAConfig(ETAAPassConfig Pass)
 	return Pass == ETAAPassConfig::DiaphragmDOF || Pass == ETAAPassConfig::DiaphragmDOFUpsampling;
 }
 
-float GetTemporalAAHistoryUpscaleFactor(const FViewInfo& View);
-
 /** GPU Output of the TAA pass. */
 struct FTAAOutputs
 {
@@ -103,7 +100,7 @@ struct FTAAPassParameters
 	{ }
 
 
-	// Customises the view rectangles for input and output.
+	// Customizes the view rectangles for input and output.
 	FORCEINLINE void SetupViewRect(const FViewInfo& View, int32 InResolutionDivisor = 1)
 	{
 		ResolutionDivisor = InResolutionDivisor;
@@ -131,17 +128,33 @@ struct FTAAPassParameters
 		OutputViewRect.Min = FIntPoint::ZeroValue;
 	}
 	
-	/** Returns the texture resolution that will be outputed. */
+	/** Returns the texture resolution that will be output. */
 	FIntPoint GetOutputExtent() const;
 
 	/** Validate the settings of TAA, to make sure there is no issue. */
 	bool Validate() const;
 };
 
+/** Temporal AA pass which emits a filtered scene color and new history. */
 FTAAOutputs AddTemporalAAPass(
 	FRDGBuilder& GraphBuilder,
 	const FSceneTextureParameters& SceneTextures,
 	const FViewInfo& View,
 	const FTAAPassParameters& Inputs,
 	const FTemporalAAHistory& InputHistory,
-	FTemporalAAHistory* OutputHistory);
\ No newline at end of file
+	FTemporalAAHistory* OutputHistory);
+
+/** Temporal AA helper method which performs filtering on the main pass scene color. Supports upsampled history and,
+ *  if requested, will attempt to perform the scene color downsample. Returns the filtered scene color, the downsampled
+ *  scene color (or null if it was not performed), and the secondary view rect.
+ */
+void AddTemporalAAPass(
+	FRDGBuilder& GraphBuilder,
+	const FSceneTextureParameters& SceneTextures,
+	const FScreenPassViewInfo& ScreenPassView,
+	const bool bAllowDownsampleSceneColor,
+	const EPixelFormat DownsampleOverrideFormat,
+	FRDGTextureRef InSceneColorTexture,
+	FRDGTextureRef* OutSceneColorTexture,
+	FRDGTextureRef* OutSceneColorHalfResTexture,
+	FIntRect* OutSecondaryViewRect);
\ No newline at end of file
diff --git a/Engine/Source/Runtime/Renderer/Private/PostProcess/PostProcessing.cpp b/Engine/Source/Runtime/Renderer/Private/PostProcess/PostProcessing.cpp
--- a/Engine/Source/Runtime/Renderer/Private/PostProcess/PostProcessing.cpp
+++ b/Engine/Source/Runtime/Renderer/Private/PostProcess/PostProcessing.cpp
@@ -37,7 +37,6 @@
 #include "PostProcess/PostProcessDOF.h"
 #include "PostProcess/PostProcessUpscale.h"
 #include "PostProcess/PostProcessHMD.h"
-#include "PostProcess/PostProcessMitchellNetravali.h"
 #include "PostProcess/PostProcessVisualizeComplexity.h"
 #include "PostProcess/PostProcessCompositeEditorPrimitives.h"
 #include "PostProcess/PostProcessShaderPrint.h"
@@ -172,10 +171,26 @@ bool ShouldDoComputePostProcessing(const FViewInfo& View)
 
 bool IsTemporalAASceneDownsampleAllowed(const FViewInfo& View)
 {
-	return (CVarTemporalAAAllowDownsampling.GetValueOnRenderThread() != 0)
-		&& !IsMotionBlurEnabled(View)
-		&& !IsVisualizeMotionBlurEnabled(View)
-		&& GetDownsampleQuality() == EDownsampleQuality::Low;
+	return CVarTemporalAAAllowDownsampling.GetValueOnRenderThread() != 0;
+}
+
+bool IsBufferVisualizationDumpFramesEnabled()
+{
+	static const auto CVarDumpFrames = IConsoleManager::Get().FindTConsoleVariableDataInt(TEXT("r.BufferVisualizationDumpFrames"));
+	return CVarDumpFrames->GetValueOnRenderThread() != 0;
+}
+
+bool IsBufferVisualizationDumpFramesInHDREnabled()
+{
+	static const auto CVarDumpFramesAsHDR = IConsoleManager::Get().FindTConsoleVariableDataInt(TEXT("r.BufferVisualizationDumpFramesAsHDR"));
+	return CVarDumpFramesAsHDR->GetValueOnRenderThread() != 0;
+}
+
+int32 GetPostProcessAAQuality()
+{
+	static const auto PostProcessAAQualityCVar = IConsoleManager::Get().FindTConsoleVariableDataInt(TEXT("r.PostProcessAAQuality"));
+
+	return FMath::Clamp(PostProcessAAQualityCVar->GetValueOnAnyThread(), 0, 6);
 }
 
 FPostprocessContext::FPostprocessContext(FRHICommandListImmediate& InRHICmdList, FRenderingCompositionGraph& InGraph, const FViewInfo& InView)
@@ -400,10 +415,7 @@ void FPostProcessing::AddGammaOnlyTonemapper(FPostprocessContext& Context)
 
 static void AddPostProcessAA(FPostprocessContext& Context)
 {
-	// console variable override
-	static const auto CVar = IConsoleManager::Get().FindTConsoleVariableDataInt(TEXT("r.PostProcessAAQuality")); 
-
-	uint32 Quality = FMath::Clamp(CVar->GetValueOnRenderThread(), 1, 6);
+	const uint32 Quality = GetPostProcessAAQuality();
 
 	FRenderingCompositePass* Node = Context.Graph.RegisterPass(new(FMemStack::Get()) FRCPassPostProcessAA(Quality));
 
@@ -725,14 +737,11 @@ static void AddGBufferVisualizationOverview(
 	FRenderingCompositeOutputRef PostTonemapHDRColorInput,
 	FRenderingCompositeOutputRef PreFlattenVelocity)
 {
-	static const auto CVarDumpFrames = IConsoleManager::Get().FindTConsoleVariableDataInt(TEXT("r.BufferVisualizationDumpFrames"));
-	static const auto CVarDumpFramesAsHDR = IConsoleManager::Get().FindTConsoleVariableDataInt(TEXT("r.BufferVisualizationDumpFramesAsHDR"));
-
-	bool bVisualizationEnabled = Context.View.Family->EngineShowFlags.VisualizeBuffer;
-	bool bOverviewModeEnabled = bVisualizationEnabled && (Context.View.CurrentBufferVisualizationMode == NAME_None);
-	bool bHighResBufferVisualizationDumpRequried = GIsHighResScreenshot && GetHighResScreenshotConfig().bDumpBufferVisualizationTargets;
-	bool bDumpFrames = Context.View.FinalPostProcessSettings.bBufferVisualizationDumpRequired && (CVarDumpFrames->GetValueOnRenderThread() || bHighResBufferVisualizationDumpRequried);
-	bool bCaptureAsHDR = CVarDumpFramesAsHDR->GetValueOnRenderThread() || GetHighResScreenshotConfig().bCaptureHDR;
+	const bool bVisualizationEnabled = Context.View.Family->EngineShowFlags.VisualizeBuffer;
+	const bool bOverviewModeEnabled = bVisualizationEnabled && (Context.View.CurrentBufferVisualizationMode == NAME_None);
+	const bool bHighResBufferVisualizationDumpRequried = GIsHighResScreenshot && GetHighResScreenshotConfig().bDumpBufferVisualizationTargets;
+	bool bDumpFrames = Context.View.FinalPostProcessSettings.bBufferVisualizationDumpRequired && (IsBufferVisualizationDumpFramesEnabled() || bHighResBufferVisualizationDumpRequried);
+	const bool bCaptureAsHDR = IsBufferVisualizationDumpFramesInHDREnabled() || GetHighResScreenshotConfig().bCaptureHDR;
 	FString BaseFilename;
 
 	if (!bDumpFrames)
@@ -960,137 +969,6 @@ FRDGTextureRef AddSeparateTranslucencyCompositionPass(FRDGBuilder& GraphBuilder,
 	return NewSceneColor;
 }
 
-bool IsDepthOfFieldEnabledForView(const FViewInfo& View)
-{
-	static const auto CVar = IConsoleManager::Get().FindTConsoleVariableDataInt(TEXT("r.DepthOfFieldQuality"));
-	check(CVar);
-	return
-		View.Family->EngineShowFlags.DepthOfField &&
-		CVar->GetValueOnRenderThread() > 0 &&
-		View.FinalPostProcessSettings.DepthOfFieldFstop > 0 &&
-		View.FinalPostProcessSettings.DepthOfFieldFocalDistance > 0;
-}
-
-EPixelFormat GetSceneColorHalfResFormat()
-{
-	return PF_FloatRGB;
-}
-
-void AddTemporalAAMainPass(
-	FRDGBuilder& GraphBuilder,
-	const FSceneTextureParameters& SceneTextures,
-	const FScreenPassViewInfo& ScreenPassView,
-	FRDGTextureRef InSceneColorTexture,
-	FRDGTextureRef InCustomDepthTexture,
-	FRDGTextureRef* OutSceneColorTexture,
-	FRDGTextureRef* OutSceneColorHalfResTexture,
-	FIntRect* OutSecondaryViewRect)
-{
-	const FViewInfo& View = ScreenPassView.View;
-
-	// Skip if TAA is disabled or if we don't have any persistent view state (required to track history).
-	if (View.AntiAliasingMethod != AAM_TemporalAA || !View.ViewState)
-	{
-		*OutSceneColorTexture = InSceneColorTexture;
-		*OutSceneColorHalfResTexture = nullptr;
-		*OutSecondaryViewRect = View.ViewRect;
-		return;
-	}
-
-	FPreviousViewInfo& PrevViewInfo = View.ViewState->PrevFrameViewInfo;
-
-	// Temporal AA passes.
-	FTAAPassParameters TAAParameters(View);
-
-	TAAParameters.Pass = View.PrimaryScreenPercentageMethod == EPrimaryScreenPercentageMethod::TemporalUpscale
-		? ETAAPassConfig::MainUpsampling
-		: ETAAPassConfig::Main;
-
-	TAAParameters.SetupViewRect(View);
-
-	{
-		static const auto CVar = IConsoleManager::Get().FindTConsoleVariableDataInt(TEXT("r.PostProcessAAQuality"));
-		uint32 Quality = FMath::Clamp(CVar->GetValueOnRenderThread(), 1, 6);
-		TAAParameters.bUseFast = Quality == 3;
-	}
-
-	const FIntRect SecondaryViewRect = TAAParameters.OutputViewRect;
-
-	const float HistoryUpscaleFactor = GetTemporalAAHistoryUpscaleFactor(View);
-
-	// Configures TAA to upscale the history buffer; this is in addition to the secondary screen percentage upscale.
-	// We end up with a scene color that is larger than the secondary screen percentage. We immediately downscale
-	// afterwards using a Mitchel-Netravali filter.
-	if (HistoryUpscaleFactor > 1.0f)
-	{
-		const FIntPoint HistoryViewSize(
-			TAAParameters.OutputViewRect.Width() * HistoryUpscaleFactor,
-			TAAParameters.OutputViewRect.Height() * HistoryUpscaleFactor);
-
-		FIntPoint QuantizedMinHistorySize;
-		QuantizeSceneBufferSize(HistoryViewSize, QuantizedMinHistorySize);
-
-		TAAParameters.Pass = ETAAPassConfig::MainSuperSampling;
-		TAAParameters.bUseFast = false;
-
-		TAAParameters.OutputViewRect.Min.X = 0;
-		TAAParameters.OutputViewRect.Min.Y = 0;
-		TAAParameters.OutputViewRect.Max = HistoryViewSize;
-	}
-
-	TAAParameters.DownsampleOverrideFormat = GetSceneColorHalfResFormat();
-
-	// Downsample pass may be merged with with TemporalAA when there is no motion blur and compute shader is used.
-	// This is currently only possible for r.Downsample.Quality = 0 (box filter).
-	TAAParameters.bDownsample = TAAParameters.bUseFast && IsTemporalAASceneDownsampleAllowed(View);
-
-	TAAParameters.SceneColorInput = InSceneColorTexture;
-
-	const FTAAOutputs TAAOutputs = AddTemporalAAPass(
-		GraphBuilder,
-		SceneTextures,
-		View,
-		TAAParameters,
-		View.PrevViewInfo.TemporalAAHistory,
-		&PrevViewInfo.TemporalAAHistory);
-
-	*OutSceneColorTexture = TAAOutputs.SceneColor;
-	*OutSceneColorHalfResTexture = TAAOutputs.DownsampledSceneColor;
-
-	// If we upscaled the history buffer, downsize back to the secondary screen percentage size.
-	if (HistoryUpscaleFactor > 1.0f)
-	{
-		const FIntRect InputViewport = TAAParameters.OutputViewRect;
-
-		FIntPoint QuantizedOutputSize;
-		QuantizeSceneBufferSize(SecondaryViewRect.Size(), QuantizedOutputSize);
-
-		FScreenPassTextureViewport OutputViewport;
-		OutputViewport.Rect = SecondaryViewRect;
-		OutputViewport.Extent.X = FMath::Max(InSceneColorTexture->Desc.Extent.X, QuantizedOutputSize.X);
-		OutputViewport.Extent.Y = FMath::Max(InSceneColorTexture->Desc.Extent.Y, QuantizedOutputSize.Y);
-
-		*OutSceneColorTexture = ComputeMitchellNetravaliDownsample(GraphBuilder, ScreenPassView, TAAOutputs.SceneColor, InputViewport, OutputViewport);
-	}
-
-	// Post Process Material Chain - SSR Input
-	if (!View.bViewStateIsReadOnly)
-	{
-		FPostProcessMaterialInputs Inputs;
-		Inputs.SetInput(EPostProcessMaterialInput::SceneColor, InSceneColorTexture, SecondaryViewRect);
-		Inputs.CustomDepthTexture = InCustomDepthTexture;
-
-		FRDGTextureRef SSRInputTexture = AddPostProcessMaterialChain(GraphBuilder, ScreenPassView, Inputs, BL_SSRInput);
-
-		if (SSRInputTexture != InSceneColorTexture)
-		{
-			GraphBuilder.QueueTextureExtraction(SSRInputTexture, &PrevViewInfo.CustomSSRInput);
-		}
-	}
-
-	*OutSecondaryViewRect = SecondaryViewRect;
-}
-
 } // namespace
 
 
@@ -1142,22 +1020,31 @@ void FPostProcessing::Process(FRHICommandListImmediate& RHICmdList, const FViewI
 
 		FRCPassPostProcessUpscale::PaniniParams PaniniConfig(View);
 
-		const bool bDepthOfFieldEnabled = IsDepthOfFieldEnabledForView(View);
+		const bool bDepthOfFieldEnabled = DiaphragmDOF::IsEnabled(View);
 
 		const bool bVisualizeHDR = View.Family->EngineShowFlags.VisualizeHDR && FeatureLevel >= ERHIFeatureLevel::SM5;
+
 		const bool bVisualizeBloom = View.Family->EngineShowFlags.VisualizeBloom && FeatureLevel >= ERHIFeatureLevel::SM4;
+
 		const bool bVisualizeMotionBlur = IsVisualizeMotionBlurEnabled(View);
+
+		// Motion blur gets replaced by the visualization pass.
+		const bool bMotionBlurEnabled = !bVisualizeMotionBlur && IsMotionBlurEnabled(View);
+
 		const bool bVisualizeDepthOfField = bDepthOfFieldEnabled && Context.View.Family->EngineShowFlags.VisualizeDOF;
 
 		const bool bTonemapEnabled = FeatureLevel >= ERHIFeatureLevel::SM4 && !bVisualizeBloom && !bVisualizeMotionBlur;
 
 		const bool bHDROutputEnabled = GRHISupportsHDROutput && IsHDREnabled();
 
-		static const auto CVarDumpFramesAsHDR = IConsoleManager::Get().FindTConsoleVariableDataInt(TEXT("r.BufferVisualizationDumpFramesAsHDR"));
-		const bool bHDRTonemapperOutput = bTonemapEnabled && (View.Family->SceneCaptureSource == SCS_FinalColorHDR || GetHighResScreenshotConfig().bCaptureHDR || CVarDumpFramesAsHDR->GetValueOnRenderThread() || bHDROutputEnabled);
+		const bool bHDRTonemapperOutput = bTonemapEnabled && (View.Family->SceneCaptureSource == SCS_FinalColorHDR || GetHighResScreenshotConfig().bCaptureHDR || IsBufferVisualizationDumpFramesInHDREnabled() || bHDROutputEnabled);
 
 		const EAntiAliasingMethod AntiAliasingMethod = View.AntiAliasingMethod;
 
+		const EDownsampleQuality DownsampleQuality = GetDownsampleQuality();
+
+		const EPixelFormat DownsampleOverrideFormat = PF_FloatRGB;
+
 		FRCPassPostProcessTonemap* Tonemapper = nullptr;
 
 		FSceneRenderTargets& SceneContext = FSceneRenderTargets::Get(RHICmdList);
@@ -1191,7 +1078,7 @@ void FPostProcessing::Process(FRHICommandListImmediate& RHICmdList, const FViewI
 
 			FRenderingCompositePass* RDGPass = Context.Graph.RegisterPass(
 				new(FMemStack::Get()) TRCPassForRDG<5, 2>(
-					[&View, &SceneContext, bVisualizeBloom, bDepthOfFieldEnabled]
+					[&View, &SceneContext, bVisualizeBloom, bMotionBlurEnabled, bVisualizeMotionBlur, bDepthOfFieldEnabled, DownsampleQuality, DownsampleOverrideFormat]
 			(FRenderingCompositePass* Pass, FRenderingCompositePassContext& InContext)
 			{
 				FRDGBuilder GraphBuilder(InContext.RHICmdList);
@@ -1207,6 +1094,13 @@ void FPostProcessing::Process(FRHICommandListImmediate& RHICmdList, const FViewI
 
 				FRDGTextureRef VelocityTexture = SceneTextures.SceneVelocityBuffer;
 
+				// Fallback to a black texture if no velocity.
+				if (!VelocityTexture)
+				{
+					VelocityTexture = GraphBuilder.RegisterExternalTexture(GSystemTextures.BlackDummy, TEXT("BlackDummy"));
+					SceneTextures.SceneVelocityBuffer = VelocityTexture;
+				}
+
 				const FScreenPassViewInfo& ScreenPassView(View);
 
 				const FIntRect PrimaryViewRect = View.ViewRect;
@@ -1226,7 +1120,7 @@ void FPostProcessing::Process(FRHICommandListImmediate& RHICmdList, const FViewI
 				{
 					FRDGTextureRef LocalSceneColorTexture = SceneColorTexture;
 
-					if (bDepthOfFieldEnabled && DiaphragmDOF::IsSupported(View.GetShaderPlatform()))
+					if (bDepthOfFieldEnabled)
 					{
 						LocalSceneColorTexture = DiaphragmDOF::AddPasses(
 							GraphBuilder,
@@ -1260,21 +1154,62 @@ void FPostProcessing::Process(FRHICommandListImmediate& RHICmdList, const FViewI
 				FRDGTextureRef SceneColorHalfResTexture = nullptr;
 
 				// Scene color view rectangle after temporal AA upscale to secondary screen percentage.
-				FIntRect SecondaryViewRect;
+				FIntRect SecondaryViewRect = PrimaryViewRect;
 
-				AddTemporalAAMainPass(
-					GraphBuilder,
-					SceneTextures,
-					ScreenPassView,
-					SceneColorTexture,
-					CustomDepthTexture,
-					&SceneColorTexture,
-					&SceneColorHalfResTexture,
-					&SecondaryViewRect);
+				// Temporal Anti-aliasing. Also may perform a temporal upsample from primary to secondary view rect.
+				if (View.AntiAliasingMethod == AAM_TemporalAA)
+				{
+					// Whether we allow the temporal AA pass to downsample scene color. It may choose not to based on internal context,
+					// in which case the output half resolution texture will remain null.
+					const bool bAllowSceneDownsample =
+						IsTemporalAASceneDownsampleAllowed(View) &&
+						// We can only merge if the normal downsample pass would happen immediately after.
+						!bMotionBlurEnabled && !bVisualizeMotionBlur && !bVisualizeBloom &&
+						// TemporalAA is only able to match the low quality mode (box filter).
+						GetDownsampleQuality() == EDownsampleQuality::Low;
+
+					AddTemporalAAPass(
+						GraphBuilder,
+						SceneTextures,
+						ScreenPassView,
+						bAllowSceneDownsample,
+						DownsampleOverrideFormat,
+						SceneColorTexture,
+						&SceneColorTexture,
+						&SceneColorHalfResTexture,
+						&SecondaryViewRect);
+				}
+
+				//! SceneColorTexture is now upsampled to the SecondaryViewRect. Use SecondaryViewRect for input / output.
 
-				// SceneColorTexture is now upscaled to the SecondaryViewRect due to TAAU.
+				// Post Process Material Chain - SSR Input
+				if (View.ViewState && !View.bViewStateIsReadOnly)
+				{
+					FPostProcessMaterialInputs Inputs;
+					Inputs.SetInput(EPostProcessMaterialInput::SceneColor, SceneColorTexture, SecondaryViewRect);
+					Inputs.CustomDepthTexture = CustomDepthTexture;
+
+					FRDGTextureRef SSRInputTexture = AddPostProcessMaterialChain(GraphBuilder, ScreenPassView, Inputs, BL_SSRInput);
+
+					// Save off SSR post process output for the next frame.
+					if (SSRInputTexture != SceneColorTexture)
+					{
+						GraphBuilder.QueueTextureExtraction(SSRInputTexture, &View.ViewState->PrevFrameViewInfo.CustomSSRInput);
+					}
+				}
 
-				if (VelocityTexture)
+				if (bVisualizeMotionBlur)
+				{
+					SceneColorTexture = AddVisualizeMotionBlurPass(
+						GraphBuilder,
+						ScreenPassView,
+						SecondaryViewRect,
+						PrimaryViewRect,
+						SceneColorTexture,
+						SceneDepthTexture,
+						VelocityTexture);
+				}
+				else if (bMotionBlurEnabled)
 				{
 					SceneColorTexture = AddMotionBlurPass(
 						GraphBuilder,
@@ -1298,8 +1233,8 @@ void FPostProcessing::Process(FRHICommandListImmediate& RHICmdList, const FViewI
 					Inputs.Name = TEXT("SceneColorHalfRes");
 					Inputs.Texture = SceneColorTexture;
 					Inputs.Viewport = SecondaryViewRect;
-					Inputs.Quality = GetDownsampleQuality();
-					Inputs.FormatOverride = GetSceneColorHalfResFormat();
+					Inputs.Quality = DownsampleQuality;
+					Inputs.FormatOverride = DownsampleOverrideFormat;
 
 					FDownsamplePassOutputs Outputs = AddDownsamplePass(GraphBuilder, ScreenPassView, Inputs);
 
@@ -2334,8 +2269,10 @@ void FPostProcessing::ProcessPlanarReflection(FRHICommandListImmediate& RHICmdLi
 
 	const FSceneRenderTargets& SceneContext = FSceneRenderTargets::Get(RHICmdList);
 
-	if (AntiAliasingMethod == AAM_TemporalAA && ViewState)
+	if (AntiAliasingMethod == AAM_TemporalAA)
 	{
+		check(ViewState);
+
 		FRDGBuilder GraphBuilder(RHICmdList);
 
 		FSceneTextureParameters SceneTextures;
