Copying //UE4/Dev-Volumetrics to Dev-Rendering (//UE4/Dev-Rendering)

diff --git a/Engine/Plugins/Runtime/Oculus/OculusVR/Source/OculusMR/Private/OculusMR_BoundaryMeshComponent.cpp b/Engine/Plugins/Runtime/Oculus/OculusVR/Source/OculusMR/Private/OculusMR_BoundaryMeshComponent.cpp
--- a/Engine/Plugins/Runtime/Oculus/OculusVR/Source/OculusMR/Private/OculusMR_BoundaryMeshComponent.cpp
+++ b/Engine/Plugins/Runtime/Oculus/OculusVR/Source/OculusMR/Private/OculusMR_BoundaryMeshComponent.cpp
@@ -189,6 +189,7 @@ public:
 		MaterialRelevance.SetPrimitiveViewRelevance(Result);
 		Result.bOpaqueRelevance = true;
 		Result.bUsesSceneDepth = false;
+		Result.bUsesSkyMaterial = false;
 		Result.bRenderCustomDepth = false;
 		return Result;
 	}
diff --git a/Engine/Shaders/Private/BasePassPixelShader.usf b/Engine/Shaders/Private/BasePassPixelShader.usf
--- a/Engine/Shaders/Private/BasePassPixelShader.usf
+++ b/Engine/Shaders/Private/BasePassPixelShader.usf
@@ -1131,7 +1131,7 @@ ShadingModel);
 	}
 #endif
 
-#if NEEDS_BASEPASS_PIXEL_FOGGING && PROJECT_SUPPORT_SKY_ATMOSPHERE
+#if NEEDS_BASEPASS_PIXEL_FOGGING && PROJECT_SUPPORT_SKY_ATMOSPHERE && MATERIAL_FORCE_SKIP_AERIAL_PERSPECTIVE==0
 	if (ResolvedView.SkyAtmosphereApplyCameraAerialPerspectiveVolume > 0.0f)
 	{
 		const float OneOverPreExposure = USE_PREEXPOSURE ? ResolvedView.OneOverPreExposure : 1.0f;
diff --git a/Engine/Shaders/Private/BasePassVertexShader.usf b/Engine/Shaders/Private/BasePassVertexShader.usf
--- a/Engine/Shaders/Private/BasePassVertexShader.usf
+++ b/Engine/Shaders/Private/BasePassVertexShader.usf
@@ -128,7 +128,7 @@ void Main(
 	Output.BasePassInterpolants.VertexFog = CalculateHeightFog(WorldPosition.xyz - ResolvedView.TranslatedWorldCameraOrigin);
 	#endif
 
-#if PROJECT_SUPPORT_SKY_ATMOSPHERE
+#if PROJECT_SUPPORT_SKY_ATMOSPHERE && MATERIAL_FORCE_SKIP_AERIAL_PERSPECTIVE==0
 	if (ResolvedView.SkyAtmosphereApplyCameraAerialPerspectiveVolume > 0.0f)
 	{
 		const float OneOverPreExposure = USE_PREEXPOSURE ? ResolvedView.OneOverPreExposure : 1.0f;
diff --git a/Engine/Shaders/Private/MobileBasePassVertexShader.usf b/Engine/Shaders/Private/MobileBasePassVertexShader.usf
--- a/Engine/Shaders/Private/MobileBasePassVertexShader.usf
+++ b/Engine/Shaders/Private/MobileBasePassVertexShader.usf
@@ -92,7 +92,7 @@ void Main(
 #if USE_VERTEX_FOG
 	float4 VertexFog = CalculateHeightFog(WorldPosition.xyz - ResolvedView.TranslatedWorldCameraOrigin);
 
-#if PROJECT_SUPPORT_SKY_ATMOSPHERE
+#if PROJECT_SUPPORT_SKY_ATMOSPHERE && MATERIAL_FORCE_SKIP_AERIAL_PERSPECTIVE==0
 	if (ResolvedView.SkyAtmosphereApplyCameraAerialPerspectiveVolume > 0.0f)
 	{
 		const float OneOverPreExposure = USE_PREEXPOSURE ? ResolvedView.OneOverPreExposure : 1.0f;
diff --git a/Engine/Shaders/Private/SkyAtmosphere.usf b/Engine/Shaders/Private/SkyAtmosphere.usf
--- a/Engine/Shaders/Private/SkyAtmosphere.usf
+++ b/Engine/Shaders/Private/SkyAtmosphere.usf
@@ -733,6 +733,11 @@ void RenderSkyAtmosphereRayMarchingPS(
 		// We should not render the sky and the current pixels are at far depth, so simply early exit.
 		// We enable depth bound when supported to not have to even process those pixels.
 		OutLuminance = PrepareOutput(float3(0.0f, 0.0f, 0.0f), float3(1.0f, 1.0f, 1.0f));
+
+		//Now the sky pass can ignore the pixel with depth == far but it will need to alpha clip because not all RHI backend support depthbound tests.
+		// And the depthtest is already setup to avoid writing all the pixel closer than to the camera than the start distance (very good optimisation).
+		// Since this shader does not write to depth or stencil it should still benefit from EArlyZ even with the clip (See AMD depth-in-depth documentation)
+		clip(0);
 		return;
 #endif
 	}
diff --git a/Engine/Shaders/Public/ShaderVersion.ush b/Engine/Shaders/Public/ShaderVersion.ush
--- a/Engine/Shaders/Public/ShaderVersion.ush
+++ b/Engine/Shaders/Public/ShaderVersion.ush
@@ -3,6 +3,6 @@
 // in Platform.ush (which should be included in any shader) it allows to invalidate the shader DDC.
 // 
 // If you are merging streams and there is a conflict with this GUID you should make a new GUID rather than taking one or the other.
-// GUID = 90D1752E117E47A8996AC771A0197283
+// GUID = 7526116FF71A45C498CB373C6683DF8A
 
 
diff --git a/Engine/Source/Runtime/Engine/Classes/Materials/Material.h b/Engine/Source/Runtime/Engine/Classes/Materials/Material.h
--- a/Engine/Source/Runtime/Engine/Classes/Materials/Material.h
+++ b/Engine/Source/Runtime/Engine/Classes/Materials/Material.h
@@ -857,6 +857,10 @@ public:
 	UPROPERTY(EditAnywhere, Category=Translucency, meta=(DisplayName = "Apply Fogging"))
 	uint32 bUseTranslucencyVertexFog:1;
 
+	/** Unlit and Opaque materials can be used as sky material on a sky dome mesh. When IsSky is true, these meshes will not receive any contribution from the aerial perspective. Height and Volumetric fog effects will still be applied. */
+	UPROPERTY(EditAnywhere, Category=Material, AdvancedDisplay)
+	uint32 bIsSky:1;
+
 	/** When true, translucent materials have fog computed for every pixel, which costs more but fixes artifacts due to low tessellation. */
 	UPROPERTY(EditAnywhere, Category=Translucency)
 	uint32 bComputeFogPerPixel:1;
diff --git a/Engine/Source/Runtime/Engine/Classes/Materials/MaterialInterface.h b/Engine/Source/Runtime/Engine/Classes/Materials/MaterialInterface.h
--- a/Engine/Source/Runtime/Engine/Classes/Materials/MaterialInterface.h
+++ b/Engine/Source/Runtime/Engine/Classes/Materials/MaterialInterface.h
@@ -76,6 +76,7 @@ struct ENGINE_API FMaterialRelevance
 	uint8 bDecal : 1;
 	uint8 bTranslucentSurfaceLighting : 1;
 	uint8 bUsesSceneDepth : 1;
+	uint8 bUsesSkyMaterial : 1;
 	uint8 bHasVolumeMaterialDomain : 1;
 	uint8 bUsesDistanceCullFade : 1;
 
diff --git a/Engine/Source/Runtime/Engine/Private/Materials/HLSLMaterialTranslator.h b/Engine/Source/Runtime/Engine/Private/Materials/HLSLMaterialTranslator.h
--- a/Engine/Source/Runtime/Engine/Private/Materials/HLSLMaterialTranslator.h
+++ b/Engine/Source/Runtime/Engine/Private/Materials/HLSLMaterialTranslator.h
@@ -1149,6 +1149,11 @@ ResourcesString = TEXT("");
 				Errorf(TEXT("Only unlit materials can output negative emissive color."));
 			}
 
+			if (Material->IsSky() && (!MaterialShadingModels.IsUnlit() || BlendMode!=BLEND_Opaque))
+			{
+				Errorf(TEXT("Sky materials must be opaque and unlit. They are expected to completely replace the background."));
+			}
+
 			bool bDBufferAllowed = IsUsingDBuffers(Platform);
 			bool bDBufferBlendMode = IsDBufferDecalBlendMode((EDecalBlendMode)Material->GetDecalBlendMode());
 
@@ -1408,6 +1413,7 @@ ResourcesString = TEXT("");
 		OutEnvironment.SetDefine(TEXT("USES_DISTORTION"), Material->IsDistorted()); 
 
 		OutEnvironment.SetDefine(TEXT("MATERIAL_ENABLE_TRANSLUCENCY_FOGGING"), Material->ShouldApplyFogging());
+		OutEnvironment.SetDefine(TEXT("MATERIAL_FORCE_SKIP_AERIAL_PERSPECTIVE"), Material->IsSky());
 		OutEnvironment.SetDefine(TEXT("MATERIAL_COMPUTE_FOG_PER_PIXEL"), Material->ComputeFogPerPixel());
 		OutEnvironment.SetDefine(TEXT("MATERIAL_FULLY_ROUGH"), bIsFullyRough || Material->IsFullyRough());
 
diff --git a/Engine/Source/Runtime/Engine/Private/Materials/Material.cpp b/Engine/Source/Runtime/Engine/Private/Materials/Material.cpp
--- a/Engine/Source/Runtime/Engine/Private/Materials/Material.cpp
+++ b/Engine/Source/Runtime/Engine/Private/Materials/Material.cpp
@@ -875,6 +875,7 @@ UMaterial::UMaterial(const FObjectInitializer& ObjectInitializer)
 	bUseMaterialAttributes = false;
 	bCastRayTracedShadows = true;
 	bUseTranslucencyVertexFog = true;
+	bIsSky = false;
 	BlendableLocation = BL_AfterTonemapping;
 	BlendablePriority = 0;
 	BlendableOutputAlpha = false;
@@ -4063,6 +4064,11 @@ bool UMaterial::CanEditChange(const UProperty* InProperty) const
 			return MaterialDomain != MD_DeferredDecal && MaterialDomain != MD_RuntimeVirtualTexture && IsTranslucentBlendMode(BlendMode);
 		}
 
+		if (PropertyName == GET_MEMBER_NAME_STRING_CHECKED(UMaterial, bIsSky))
+		{
+			return MaterialDomain != MD_DeferredDecal && MaterialDomain != MD_RuntimeVirtualTexture && GetShadingModels().IsUnlit() && BlendMode==BLEND_Opaque;
+		}
+
 		if (PropertyName == GET_MEMBER_NAME_STRING_CHECKED(UMaterial, TranslucencyLightingMode)
 			|| PropertyName == GET_MEMBER_NAME_STRING_CHECKED(UMaterial, TranslucencyDirectionalLightingIntensity)
 			|| PropertyName == GET_MEMBER_NAME_STRING_CHECKED(UMaterial, TranslucentShadowDensityScale)
diff --git a/Engine/Source/Runtime/Engine/Private/Materials/MaterialInterface.cpp b/Engine/Source/Runtime/Engine/Private/Materials/MaterialInterface.cpp
--- a/Engine/Source/Runtime/Engine/Private/Materials/MaterialInterface.cpp
+++ b/Engine/Source/Runtime/Engine/Private/Materials/MaterialInterface.cpp
@@ -50,6 +50,7 @@ void FMaterialRelevance::SetPrimitiveViewRelevance(FPrimitiveViewRelevance& OutV
 	OutViewRelevance.bDecal = bDecal;
 	OutViewRelevance.bTranslucentSurfaceLighting = bTranslucentSurfaceLighting;
 	OutViewRelevance.bUsesSceneDepth = bUsesSceneDepth;
+	OutViewRelevance.bUsesSkyMaterial = bUsesSkyMaterial;
 	OutViewRelevance.bHasVolumeMaterialDomain = bHasVolumeMaterialDomain;
 }
 
@@ -158,6 +159,7 @@ FMaterialRelevance UMaterialInterface::GetRelevance_Internal(const UMaterial* Ma
 			MaterialRelevance.bUsesSceneDepth = MaterialResource->MaterialUsesSceneDepthLookup_GameThread();
 			MaterialRelevance.bHasVolumeMaterialDomain = MaterialResource->IsVolumetricPrimitive();
 			MaterialRelevance.bUsesDistanceCullFade = MaterialResource->MaterialUsesDistanceCullFade_GameThread();
+			MaterialRelevance.bUsesSkyMaterial = Material->bIsSky;
 		}
 		return MaterialRelevance;
 	}
diff --git a/Engine/Source/Runtime/Engine/Private/Materials/MaterialShared.cpp b/Engine/Source/Runtime/Engine/Private/Materials/MaterialShared.cpp
--- a/Engine/Source/Runtime/Engine/Private/Materials/MaterialShared.cpp
+++ b/Engine/Source/Runtime/Engine/Private/Materials/MaterialShared.cpp
@@ -1334,7 +1334,8 @@ FLinearColor FMaterialResource::GetTranslucentMultipleScatteringExtinction() con
 float FMaterialResource::GetTranslucentShadowStartOffset() const { return Material->TranslucentShadowStartOffset; }
 float FMaterialResource::GetRefractionDepthBiasValue() const { return Material->RefractionDepthBias; }
 float FMaterialResource::GetMaxDisplacement() const { return Material->MaxDisplacement; }
-bool FMaterialResource::ShouldApplyFogging() const {return Material->bUseTranslucencyVertexFog;}
+bool FMaterialResource::ShouldApplyFogging() const { return Material->bUseTranslucencyVertexFog; }
+bool FMaterialResource::IsSky() const { return Material->bIsSky; }
 bool FMaterialResource::ComputeFogPerPixel() const {return Material->bComputeFogPerPixel;}
 FString FMaterialResource::GetFriendlyName() const { return *GetNameSafe(Material); } //avoid using the material instance name here, we want materials that share a shadermap to also share a friendly name.
 
diff --git a/Engine/Source/Runtime/Engine/Private/StaticMeshRender.cpp b/Engine/Source/Runtime/Engine/Private/StaticMeshRender.cpp
--- a/Engine/Source/Runtime/Engine/Private/StaticMeshRender.cpp
+++ b/Engine/Source/Runtime/Engine/Private/StaticMeshRender.cpp
@@ -917,6 +917,7 @@ void FStaticMeshSceneProxy::DrawStaticElements(FStaticPrimitiveDrawInterface* PD
 		//Never use the dynamic path in this path, because only unselected elements will use DrawStaticElements
 		bool bIsMeshElementSelected = false;
 		const auto FeatureLevel = GetScene().GetFeatureLevel();
+		const bool IsMobile = IsMobilePlatform(GetScene().GetShaderPlatform());
 
 		//check if a LOD is being forced
 		if (ForcedLodModel > 0) 
diff --git a/Engine/Source/Runtime/Engine/Public/MaterialShared.h b/Engine/Source/Runtime/Engine/Public/MaterialShared.h
--- a/Engine/Source/Runtime/Engine/Public/MaterialShared.h
+++ b/Engine/Source/Runtime/Engine/Public/MaterialShared.h
@@ -1494,6 +1494,7 @@ public:
 	virtual float GetMaxDisplacement() const { return 0.0f; }
 	virtual bool ShouldApplyFogging() const { return false; }
 	virtual bool ComputeFogPerPixel() const { return false; }
+	virtual bool IsSky() const { return false; }
 	virtual FString GetFriendlyName() const = 0;
 	virtual bool HasVertexPositionOffsetConnected() const { return false; }
 	virtual bool HasPixelDepthOffsetConnected() const { return false; }
@@ -2253,6 +2254,7 @@ public:
 	ENGINE_API virtual float GetRefractionDepthBiasValue() const override;
 	ENGINE_API virtual float GetMaxDisplacement() const override;
 	ENGINE_API virtual bool ShouldApplyFogging() const override;
+	ENGINE_API virtual bool IsSky() const override;
 	ENGINE_API virtual bool ComputeFogPerPixel() const override;
 	ENGINE_API virtual bool HasRuntimeVirtualTextureOutput() const override;
 	ENGINE_API virtual bool CastsRayTracedShadows() const override;
diff --git a/Engine/Source/Runtime/Engine/Public/PrimitiveViewRelevance.h b/Engine/Source/Runtime/Engine/Public/PrimitiveViewRelevance.h
--- a/Engine/Source/Runtime/Engine/Public/PrimitiveViewRelevance.h
+++ b/Engine/Source/Runtime/Engine/Public/PrimitiveViewRelevance.h
@@ -75,6 +75,8 @@ struct FPrimitiveViewRelevance
 	uint32 bTranslucentSelfShadow : 1;
 	/** Whether the primitive has materials that read the scene depth. */
 	uint32 bUsesSceneDepth : 1;
+	/** Whether the primitive has materials that read the scene depth. */
+	uint32 bUsesSkyMaterial : 1;
 	/** Whether the view use custom data. */
 	uint32 bUseCustomViewData : 1;
 
diff --git a/Engine/Source/Runtime/Renderer/Private/BasePassRendering.cpp b/Engine/Source/Runtime/Renderer/Private/BasePassRendering.cpp
--- a/Engine/Source/Runtime/Renderer/Private/BasePassRendering.cpp
+++ b/Engine/Source/Runtime/Renderer/Private/BasePassRendering.cpp
@@ -631,6 +631,7 @@ bool FDeferredShadingSceneRenderer::RenderBasePass(FRHICommandListImmediate& RHI
 		SCOPE_CYCLE_COUNTER(STAT_BasePassDrawTime);
 		SCOPED_GPU_STAT(RHICmdList, Basepass);
 
+		FExclusiveDepthStencil::Type BasePassDepthStencilAccess_NoDepthWrite = FExclusiveDepthStencil::Type(BasePassDepthStencilAccess & ~FExclusiveDepthStencil::DepthWrite);
 		if (bParallelBasePass)
 		{
 			check(RHICmdList.IsOutsideRenderPass());
@@ -649,7 +650,8 @@ bool FDeferredShadingSceneRenderer::RenderBasePass(FRHICommandListImmediate& RHI
 
 				SetupBasePassState(BasePassDepthStencilAccess, ViewFamily.EngineShowFlags.ShaderComplexity, DrawRenderState);
 
-				if (View.ShouldRenderView())
+				const bool bShouldRenderView = View.ShouldRenderView();
+				if (bShouldRenderView)
 				{
 					Scene->UniformBuffers.UpdateViewUniformBuffer(View);
 
@@ -661,6 +663,12 @@ bool FDeferredShadingSceneRenderer::RenderBasePass(FRHICommandListImmediate& RHI
 				FSceneRenderTargets::Get(RHICmdList).BeginRenderingGBuffer(RHICmdList, ERenderTargetLoadAction::ELoad, ERenderTargetLoadAction::ELoad, BasePassDepthStencilAccess, this->ViewFamily.EngineShowFlags.ShaderComplexity);
 				RenderEditorPrimitives(RHICmdList, View, BasePassDepthStencilAccess, DrawRenderState, bDirty);
 				RHICmdList.EndRenderPass();
+
+				if (bShouldRenderView)
+				{
+					SetupBasePassState(BasePassDepthStencilAccess_NoDepthWrite, ViewFamily.EngineShowFlags.ShaderComplexity, DrawRenderState);
+					RenderSkyPassViewParallel(RHICmdList, View, BasePassDepthStencilAccess_NoDepthWrite, DrawRenderState);
+				}
 			}
 
 			bDirty = true; // assume dirty since we are not going to wait
@@ -683,7 +691,8 @@ bool FDeferredShadingSceneRenderer::RenderBasePass(FRHICommandListImmediate& RHI
 
 				SetupBasePassState(BasePassDepthStencilAccess, ViewFamily.EngineShowFlags.ShaderComplexity, DrawRenderState);
 
-				if (View.ShouldRenderView())
+				const bool bShouldRenderView = View.ShouldRenderView();
+				if (bShouldRenderView)
 				{
 					Scene->UniformBuffers.UpdateViewUniformBuffer(View);
 
@@ -691,6 +700,12 @@ bool FDeferredShadingSceneRenderer::RenderBasePass(FRHICommandListImmediate& RHI
 				}
 
 				RenderEditorPrimitives(RHICmdList, View, BasePassDepthStencilAccess, DrawRenderState, bDirty);
+
+				if (bShouldRenderView)
+				{
+					SetupBasePassState(BasePassDepthStencilAccess_NoDepthWrite, ViewFamily.EngineShowFlags.ShaderComplexity, DrawRenderState);
+					bDirty |= RenderSkyPassView(RHICmdList, View, BasePassDepthStencilAccess_NoDepthWrite, DrawRenderState);
+				}
 			}
 		}	
 	}
@@ -701,7 +716,7 @@ bool FDeferredShadingSceneRenderer::RenderBasePass(FRHICommandListImmediate& RHI
 	return bDirty;
 }
 
-static void SetupBasePassView(FRHICommandList& RHICmdList, const FViewInfo& View, const FSceneRenderer* SceneRenderer, const bool bIsEditorPrimitivePass = false)
+void SetupBasePassView(FRHICommandList& RHICmdList, const FViewInfo& View, const FSceneRenderer* SceneRenderer, const bool bIsEditorPrimitivePass)
 {
 	if (!View.IsInstancedStereoPass() || bIsEditorPrimitivePass)
 	{
@@ -777,6 +792,22 @@ void FDeferredShadingSceneRenderer::RenderBasePassViewParallel(FViewInfo& View,
 	View.ParallelMeshDrawCommandPasses[EMeshPass::BasePass].DispatchDraw(&ParallelSet, ParentCmdList);
 }
 
+void FDeferredShadingSceneRenderer::RenderSkyPassViewParallel(FRHICommandListImmediate& ParentCmdList, FViewInfo& View, FExclusiveDepthStencil::Type BasePassDepthStencilAccess, const FMeshPassProcessorRenderState& InDrawRenderState)
+{
+	check(ParentCmdList.IsOutsideRenderPass());
+
+	FBasePassParallelCommandListSet ParallelSet(View, ParentCmdList,
+		CVarRHICmdBasePassDeferredContexts.GetValueOnRenderThread() > 0,
+		CVarRHICmdFlushRenderThreadTasksBasePass.GetValueOnRenderThread() == 0 && CVarRHICmdFlushRenderThreadTasks.GetValueOnRenderThread() == 0,
+		this,
+		BasePassDepthStencilAccess,
+		InDrawRenderState);
+
+	// enqueue RHIThread command that blocks on prereq, lock / unlock vertex buffer upload
+
+	View.ParallelMeshDrawCommandPasses[EMeshPass::SkyPass].DispatchDraw(&ParallelSet, ParentCmdList);
+}
+
 bool HasEditorPrimitivesForDPG(const FViewInfo& View, ESceneDepthPriorityGroup DepthPriorityGroup)
 {
 	bool bHasPrimitives = View.SimpleElementCollector.HasPrimitives(DepthPriorityGroup);
@@ -903,6 +934,17 @@ bool FDeferredShadingSceneRenderer::RenderBasePassView(FRHICommandListImmediate&
 	return bDirty;
 }
 
+bool FDeferredShadingSceneRenderer::RenderSkyPassView(FRHICommandListImmediate& RHICmdList, FViewInfo& View, FExclusiveDepthStencil::Type BasePassDepthStencilAccess, const FMeshPassProcessorRenderState& InDrawRenderState)
+{
+	bool bDirty = false;
+	FMeshPassProcessorRenderState DrawRenderState(InDrawRenderState);
+	SetupBasePassView(RHICmdList, View, this);
+
+	View.ParallelMeshDrawCommandPasses[EMeshPass::SkyPass].DispatchDraw(nullptr, RHICmdList);
+
+	return bDirty;
+}
+
 template<typename LightMapPolicyType>
 void FBasePassMeshProcessor::Process(
 	const FMeshBatch& RESTRICT MeshBatch,
@@ -1191,7 +1233,8 @@ void FBasePassMeshProcessor::AddMeshBatch(const FMeshBatch& RESTRICT MeshBatch,
 		// Only draw opaque materials.
 		if (bShouldDraw
 			&& (!PrimitiveSceneProxy || PrimitiveSceneProxy->ShouldRenderInMainPass())
-			&& ShouldIncludeDomainInMeshPass(Material.GetMaterialDomain()))
+			&& ShouldIncludeDomainInMeshPass(Material.GetMaterialDomain())
+			&& !Material.IsSky())
 		{
 			// Check for a cached light-map.
 			const bool bIsLitMaterial = ShadingModels.IsLit();
diff --git a/Engine/Source/Runtime/Renderer/Private/BasePassRendering.h b/Engine/Source/Runtime/Renderer/Private/BasePassRendering.h
--- a/Engine/Source/Runtime/Renderer/Private/BasePassRendering.h
+++ b/Engine/Source/Runtime/Renderer/Private/BasePassRendering.h
@@ -590,6 +590,8 @@ void GetBasePassShaders<FUniformLightMapPolicy>(
 	TBasePassVertexShaderPolicyParamType<FUniformLightMapPolicy>*& VertexShader,
 	TBasePassPixelShaderPolicyParamType<FUniformLightMapPolicy>*& PixelShader
 	);
+
+void SetupBasePassView(FRHICommandList& RHICmdList, const FViewInfo& View, const FSceneRenderer* SceneRenderer, const bool bIsEditorPrimitivePass = false);
 	
 class FBasePassMeshProcessor : public FMeshPassProcessor
 {
diff --git a/Engine/Source/Runtime/Renderer/Private/DeferredShadingRenderer.cpp b/Engine/Source/Runtime/Renderer/Private/DeferredShadingRenderer.cpp
--- a/Engine/Source/Runtime/Renderer/Private/DeferredShadingRenderer.cpp
+++ b/Engine/Source/Runtime/Renderer/Private/DeferredShadingRenderer.cpp
@@ -849,7 +849,7 @@ void FDeferredShadingSceneRenderer::Render(FRHICommandListImmediate& RHICmdList)
 
 	PrepareViewRectsForRendering();
 
-	if (Scene->HasSkyAtmosphere())
+	if (Scene->HasSkyAtmosphere() && ShouldRenderSkyAtmosphere(Scene->SkyAtmosphere, Scene->GetShaderPlatform()))
 	{
 		for (int32 LightIndex = 0; LightIndex < NUM_ATMOSPHERE_LIGHTS; ++LightIndex)
 		{
@@ -864,6 +864,10 @@ void FDeferredShadingSceneRenderer::Render(FRHICommandListImmediate& RHICmdList)
 		// Only one atmospheric light at one time.
 		Scene->GetAtmosphericFogSceneInfo()->PrepareSunLightProxy(*Scene->AtmosphereLights[0]);
 	}
+	else
+	{
+		Scene->ResetAtmosphereLightsProperties();
+	}
 
 	SCOPED_NAMED_EVENT(FDeferredShadingSceneRenderer_Render, FColor::Emerald);
 
@@ -1360,6 +1364,7 @@ void FDeferredShadingSceneRenderer::Render(FRHICommandListImmediate& RHICmdList)
 
 	bool bIsWireframeRenderpass = bIsWireframe && FSceneRenderer::ShouldCompositeEditorPrimitives(Views[0]);
 	bool bRenderLightmapDensity = ViewFamily.EngineShowFlags.LightMapDensity && AllowDebugViewmodes();
+	bool bRenderSkyAtmosphereEditorNotifications = ShouldRenderSkyAtmosphereEditorNotifications();
 	bool bDoParallelBasePass = GRHICommandList.UseParallelAlgorithms() && CVarParallelBasePass.GetValueOnRenderThread();
 	
 	// BASE PASS AND GBUFFER SETUP
@@ -1384,7 +1389,7 @@ void FDeferredShadingSceneRenderer::Render(FRHICommandListImmediate& RHICmdList)
 		SceneContext.BeginRenderingGBuffer(RHICmdList, ColorLoadAction, DepthLoadAction, BasePassDepthStencilAccess, ViewFamily.EngineShowFlags.ShaderComplexity, true, ClearColor);
 		
 		// If we are in wireframe mode or will go wide later this pass is just the clear.
-		if (bIsWireframeRenderpass || bDoParallelBasePass)
+		if (bIsWireframeRenderpass || bRenderSkyAtmosphereEditorNotifications || bDoParallelBasePass)
 		{
 			RHICmdList.EndRenderPass();
 		}
@@ -1394,7 +1399,12 @@ void FDeferredShadingSceneRenderer::Render(FRHICommandListImmediate& RHICmdList)
 		}
 		ServiceLocalQueue();
 	}
-	
+
+	if (bRenderSkyAtmosphereEditorNotifications)
+	{
+		RenderSkyAtmosphereEditorNotifications(RHICmdList);
+	}
+
 	// Wireframe mode requires bRequiresRHIClear to be true. 
 	// Rendering will be very funny without it and the call to BeginRenderingGBuffer will call AllocSceneColor which is needed for the EditorPrimitives resolve.
 	if (bIsWireframeRenderpass)
diff --git a/Engine/Source/Runtime/Renderer/Private/DeferredShadingRenderer.h b/Engine/Source/Runtime/Renderer/Private/DeferredShadingRenderer.h
--- a/Engine/Source/Runtime/Renderer/Private/DeferredShadingRenderer.h
+++ b/Engine/Source/Runtime/Renderer/Private/DeferredShadingRenderer.h
@@ -81,6 +81,11 @@ public:
 	/** Renders the basepass for a given View. */
 	bool RenderBasePassView(FRHICommandListImmediate& RHICmdList, FViewInfo& View, FExclusiveDepthStencil::Type BasePassDepthStencilAccess, const FMeshPassProcessorRenderState& InDrawRenderState);
 
+	/** Renders the sky pass for a given View, in parallel */
+	void RenderSkyPassViewParallel(FRHICommandListImmediate& ParentCmdList, FViewInfo& View, FExclusiveDepthStencil::Type BasePassDepthStencilAccess, const FMeshPassProcessorRenderState& InDrawRenderState);
+	/** Renders the sky pass for a given View. */
+	bool RenderSkyPassView(FRHICommandListImmediate& RHICmdList, FViewInfo& View, FExclusiveDepthStencil::Type BasePassDepthStencilAccess, const FMeshPassProcessorRenderState& InDrawRenderState);
+
 	/** Renders editor primitives for a given View. */
 	void RenderEditorPrimitives(FRHICommandListImmediate& RHICmdList, const FViewInfo& View, FExclusiveDepthStencil::Type BasePassDepthStencilAccess, const FMeshPassProcessorRenderState& DrawRenderState, bool& bOutDirty);
 
diff --git a/Engine/Source/Runtime/Renderer/Private/DepthRendering.cpp b/Engine/Source/Runtime/Renderer/Private/DepthRendering.cpp
--- a/Engine/Source/Runtime/Renderer/Private/DepthRendering.cpp
+++ b/Engine/Source/Runtime/Renderer/Private/DepthRendering.cpp
@@ -783,7 +783,8 @@ void FDepthPassMeshProcessor::AddMeshBatch(const FMeshBatch& RESTRICT MeshBatch,
 
 		if (!bIsTranslucent
 			&& (!PrimitiveSceneProxy || PrimitiveSceneProxy->ShouldRenderInMainPass())
-			&& ShouldIncludeDomainInMeshPass(Material.GetMaterialDomain()))
+			&& ShouldIncludeDomainInMeshPass(Material.GetMaterialDomain())
+			&& !Material.IsSky())
 		{
 			if (BlendMode == BLEND_Opaque
 				&& MeshBatch.VertexFactory->SupportsPositionOnlyStream()
diff --git a/Engine/Source/Runtime/Renderer/Private/LightMapDensityRendering.cpp b/Engine/Source/Runtime/Renderer/Private/LightMapDensityRendering.cpp
--- a/Engine/Source/Runtime/Renderer/Private/LightMapDensityRendering.cpp
+++ b/Engine/Source/Runtime/Renderer/Private/LightMapDensityRendering.cpp
@@ -253,7 +253,8 @@ void FLightmapDensityMeshProcessor::AddMeshBatch(const FMeshBatch& RESTRICT Mesh
 		static const auto SupportLowQualityLightmapsVar = IConsoleManager::Get().FindTConsoleVariableDataInt(TEXT("r.SupportLowQualityLightmaps"));
 		const bool bAllowLowQualityLightMaps = (!SupportLowQualityLightmapsVar) || (SupportLowQualityLightmapsVar->GetValueOnAnyThread() != 0);
 
-		if (!bTranslucentBlendMode || ViewIfDynamicMeshCommand->Family->EngineShowFlags.Wireframe)
+		if ((!bTranslucentBlendMode || ViewIfDynamicMeshCommand->Family->EngineShowFlags.Wireframe)
+			&& !Material->IsSky())
 		{
 			if (!bMaterialMasked && !Material->MaterialModifiesMeshPosition_RenderThread())
 			{
diff --git a/Engine/Source/Runtime/Renderer/Private/MobileShadingRenderer.cpp b/Engine/Source/Runtime/Renderer/Private/MobileShadingRenderer.cpp
--- a/Engine/Source/Runtime/Renderer/Private/MobileShadingRenderer.cpp
+++ b/Engine/Source/Runtime/Renderer/Private/MobileShadingRenderer.cpp
@@ -315,7 +315,7 @@ void FMobileSceneRenderer::Render(FRHICommandListImmediate& RHICmdList)
 
 	PrepareViewRectsForRendering();
 
-	if (Scene->HasSkyAtmosphere())
+	if (Scene->HasSkyAtmosphere() && ShouldRenderSkyAtmosphere(Scene->SkyAtmosphere, Scene->GetShaderPlatform()))
 	{
 		for (int32 LightIndex = 0; LightIndex < NUM_ATMOSPHERE_LIGHTS; ++LightIndex)
 		{
@@ -325,6 +325,10 @@ void FMobileSceneRenderer::Render(FRHICommandListImmediate& RHICmdList)
 			}
 		}
 	}
+	else
+	{
+		Scene->ResetAtmosphereLightsProperties();
+	}
 
 	CSV_SCOPED_TIMING_STAT_EXCLUSIVE(RenderOther);
 	QUICK_SCOPE_CYCLE_COUNTER(STAT_FMobileSceneRenderer_Render);
diff --git a/Engine/Source/Runtime/Renderer/Private/PrimitiveSceneInfo.cpp b/Engine/Source/Runtime/Renderer/Private/PrimitiveSceneInfo.cpp
--- a/Engine/Source/Runtime/Renderer/Private/PrimitiveSceneInfo.cpp
+++ b/Engine/Source/Runtime/Renderer/Private/PrimitiveSceneInfo.cpp
@@ -77,10 +77,12 @@ public:
 
 			const bool bSupportsCachingMeshDrawCommands = SupportsCachingMeshDrawCommands(PrimitiveSceneProxy, *StaticMesh, FeatureLevel);
 
+			bool bUseSkyMaterial = Mesh.MaterialRenderProxy->GetMaterial(FeatureLevel)->IsSky();
 			FStaticMeshBatchRelevance* StaticMeshRelevance = new(PrimitiveSceneInfo->StaticMeshRelevances) FStaticMeshBatchRelevance(
 				*StaticMesh, 
 				ScreenSize, 
-				bSupportsCachingMeshDrawCommands
+				bSupportsCachingMeshDrawCommands,
+				bUseSkyMaterial
 			);
 		}
 	}
diff --git a/Engine/Source/Runtime/Renderer/Private/SceneCore.h b/Engine/Source/Runtime/Renderer/Private/SceneCore.h
--- a/Engine/Source/Runtime/Renderer/Private/SceneCore.h
+++ b/Engine/Source/Runtime/Renderer/Private/SceneCore.h
@@ -168,7 +168,7 @@ private:
 class FStaticMeshBatchRelevance
 {
 public:
-	FStaticMeshBatchRelevance(const FStaticMeshBatch& StaticMesh, float InScreenSize, bool InbSupportsCachingMeshDrawCommands)
+	FStaticMeshBatchRelevance(const FStaticMeshBatch& StaticMesh, float InScreenSize, bool InbSupportsCachingMeshDrawCommands, bool InbUseSkyMaterial)
 		: Id(StaticMesh.Id)
 		, ScreenSize(InScreenSize)
 		, CommandInfosBase(0)
@@ -181,6 +181,7 @@ public:
 		, bUseForMaterial(StaticMesh.bUseForMaterial)
 		, bUseForDepthPass(StaticMesh.bUseForDepthPass)
 		, bUseAsOccluder(StaticMesh.bUseAsOccluder)
+		, bUseSkyMaterial(InbUseSkyMaterial)
 		, bRenderToVirtualTexture(StaticMesh.bRenderToVirtualTexture)
 		, RuntimeVirtualTextureMaterialType(StaticMesh.RuntimeVirtualTextureMaterialType)
 		, bSupportsCachingMeshDrawCommands(InbSupportsCachingMeshDrawCommands)
@@ -218,6 +219,7 @@ public:
 	uint8 bUseForMaterial	: 1; // Whether it can be used in renderpasses requiring material outputs.
 	uint8 bUseForDepthPass	: 1; // Whether it can be used in depth pass.
 	uint8 bUseAsOccluder	: 1; // User hint whether it's a good occluder.
+	uint8 bUseSkyMaterial	: 1; // Whether this batch uses a Sky material or not.
 
 	/** Whether the mesh batch can be used for rendering to a virtual texture. */
 	uint8 bRenderToVirtualTexture : 1;
diff --git a/Engine/Source/Runtime/Renderer/Private/ScenePrivate.h b/Engine/Source/Runtime/Renderer/Private/ScenePrivate.h
--- a/Engine/Source/Runtime/Renderer/Private/ScenePrivate.h
+++ b/Engine/Source/Runtime/Renderer/Private/ScenePrivate.h
@@ -2713,6 +2713,9 @@ public:
 		return (SkyAtmosphere != NULL);
 	}
 
+	// Reset all the light to default state "not being affected by atmosphere". Should only be called from render side.
+	void ResetAtmosphereLightsProperties();
+
 	/**
 	 * Retrieves the lights interacting with the passed in primitive and adds them to the out array.
 	 *
diff --git a/Engine/Source/Runtime/Renderer/Private/SceneRendering.cpp b/Engine/Source/Runtime/Renderer/Private/SceneRendering.cpp
--- a/Engine/Source/Runtime/Renderer/Private/SceneRendering.cpp
+++ b/Engine/Source/Runtime/Renderer/Private/SceneRendering.cpp
@@ -812,6 +812,7 @@ void FViewInfo::Init()
 	bDisableQuerySubmissions = false;
 	bDisableDistanceBasedFadeTransitions = false;	
 	ShadingModelMaskInView = 0;
+	bSceneHasSkyMaterial = 0;
 
 	NumVisibleStaticMeshElements = 0;
 	PrecomputedVisibilityData = 0;
diff --git a/Engine/Source/Runtime/Renderer/Private/SceneRendering.h b/Engine/Source/Runtime/Renderer/Private/SceneRendering.h
--- a/Engine/Source/Runtime/Renderer/Private/SceneRendering.h
+++ b/Engine/Source/Runtime/Renderer/Private/SceneRendering.h
@@ -1017,6 +1017,11 @@ public:
 	 * TODO: Right now decal visibility is computed right before rendering them. Ideally it should be done in InitViews and this flag should be replaced with list of visible decals  
 	 */
 	uint32 bSceneHasDecals : 1;
+	/**
+	 * true if the scene has at least one mesh with a material tagged as sky. 
+	 * This is used to skip the sky rendering part during the SkyAtmosphere pass on non mobile platforms.
+	 */
+	uint32 bSceneHasSkyMaterial : 1;
 	/** Bitmask of all shading models used by primitives in this view */
 	uint16 ShadingModelMaskInView;
 
@@ -1690,6 +1695,11 @@ protected:
 	/** Render the sky atmosphere over the scene.*/
 	void RenderSkyAtmosphere(FRHICommandListImmediate& RHICmdList);
 
+	/** Render notification to artist when a sky material is used but it might comtains the camera (and then the sky/background would look black).*/
+	void RenderSkyAtmosphereEditorNotifications(FRHICommandListImmediate& RHICmdList);
+	/** We should render on screen notification only if any of the scene contains a mesh using a sky material.*/
+	bool ShouldRenderSkyAtmosphereEditorNotifications();
+
 	void ResolveSceneColor(FRHICommandList& RHICmdList);
 
 private:
diff --git a/Engine/Source/Runtime/Renderer/Private/SceneVisibility.cpp b/Engine/Source/Runtime/Renderer/Private/SceneVisibility.cpp
--- a/Engine/Source/Runtime/Renderer/Private/SceneVisibility.cpp
+++ b/Engine/Source/Runtime/Renderer/Private/SceneVisibility.cpp
@@ -1888,6 +1888,7 @@ struct FRelevancePacket
 	bool bUsesLightingChannels;
 	bool bTranslucentSurfaceLighting;
 	bool bUsesSceneDepth;
+	bool bSceneHasSkyMaterial;
 
 	FRelevancePacket(
 		FRHICommandListImmediate& InRHICmdList,
@@ -1924,6 +1925,7 @@ struct FRelevancePacket
 		, bUsesLightingChannels(false)
 		, bTranslucentSurfaceLighting(false)
 		, bUsesSceneDepth(false)
+		, bSceneHasSkyMaterial(false)
 	{
 	}
 
@@ -1936,6 +1938,7 @@ struct FRelevancePacket
 	void ComputeRelevance()
 	{
 		CombinedShadingModelMask = 0;
+		bSceneHasSkyMaterial = 0;
 		bUsesGlobalDistanceField = false;
 		bUsesLightingChannels = false;
 		bTranslucentSurfaceLighting = false;
@@ -2033,6 +2036,7 @@ struct FRelevancePacket
 			bUsesLightingChannels |= ViewRelevance.bUsesLightingChannels;
 			bTranslucentSurfaceLighting |= ViewRelevance.bTranslucentSurfaceLighting;
 			bUsesSceneDepth |= ViewRelevance.bUsesSceneDepth;
+			bSceneHasSkyMaterial |= ViewRelevance.bUsesSkyMaterial;
 
 			if (ViewRelevance.bRenderCustomDepth)
 			{
@@ -2207,7 +2211,7 @@ struct FRelevancePacket
 
 					if (ViewRelevance.bDrawRelevance)
 					{
-						if ((StaticMeshRelevance.bUseForMaterial || StaticMeshRelevance.bUseAsOccluder) 
+						if ((StaticMeshRelevance.bUseForMaterial || StaticMeshRelevance.bUseAsOccluder)
 							&& (ViewRelevance.bRenderInMainPass || ViewRelevance.bRenderCustomDepth) 
 							&& !bHiddenByHLODFade)
 						{
@@ -2226,6 +2230,11 @@ struct FRelevancePacket
 								{
 									DrawCommandPacket.AddCommandsForMesh(PrimitiveIndex, PrimitiveSceneInfo, StaticMeshRelevance, StaticMesh, Scene, bCanCache, EMeshPass::MobileBasePassCSM);
 								}
+								else if(StaticMeshRelevance.bUseSkyMaterial)
+								{
+									// Not needed on Mobile path as in this case everything goes into the regular base pass
+									DrawCommandPacket.AddCommandsForMesh(PrimitiveIndex, PrimitiveSceneInfo, StaticMeshRelevance, StaticMesh, Scene, bCanCache, EMeshPass::SkyPass);
+								}
 
 								if (ViewRelevance.bRenderCustomDepth)
 								{
@@ -2362,6 +2371,7 @@ struct FRelevancePacket
 		WriteView.bUsesLightingChannels |= bUsesLightingChannels;
 		WriteView.bTranslucentSurfaceLighting |= bTranslucentSurfaceLighting;
 		WriteView.bUsesSceneDepth |= bUsesSceneDepth;
+		WriteView.bSceneHasSkyMaterial |= bSceneHasSkyMaterial;
 		VisibleDynamicPrimitivesWithSimpleLights.AppendTo(WriteView.VisibleDynamicPrimitivesWithSimpleLights);
 		WriteView.NumVisibleDynamicPrimitives += NumVisibleDynamicPrimitives;
 		WriteView.NumVisibleDynamicEditorPrimitives += NumVisibleDynamicEditorPrimitives;
diff --git a/Engine/Source/Runtime/Renderer/Private/ShadowDepthRendering.cpp b/Engine/Source/Runtime/Renderer/Private/ShadowDepthRendering.cpp
--- a/Engine/Source/Runtime/Renderer/Private/ShadowDepthRendering.cpp
+++ b/Engine/Source/Runtime/Renderer/Private/ShadowDepthRendering.cpp
@@ -1839,7 +1839,8 @@ void FShadowDepthPassMeshProcessor::AddMeshBatch(const FMeshBatch& RESTRICT Mesh
 		}
 
 		if ((bShouldCastShadow || (bReflectiveShadowmap && (Material.ShouldInjectEmissiveIntoLPV() || Material.ShouldBlockGI())))
-			&& ShouldIncludeDomainInMeshPass(Material.GetMaterialDomain()))
+			&& ShouldIncludeDomainInMeshPass(Material.GetMaterialDomain())
+			&& !Material.IsSky())
 		{
 			const FMaterialRenderProxy* EffectiveMaterialRenderProxy = &MaterialRenderProxy;
 			const FMaterial* EffectiveMaterial = &Material;
diff --git a/Engine/Source/Runtime/Renderer/Private/SkyAtmosphereRendering.cpp b/Engine/Source/Runtime/Renderer/Private/SkyAtmosphereRendering.cpp
--- a/Engine/Source/Runtime/Renderer/Private/SkyAtmosphereRendering.cpp
+++ b/Engine/Source/Runtime/Renderer/Private/SkyAtmosphereRendering.cpp
@@ -33,13 +33,6 @@ static TAutoConsoleVariable<int32> CVarSupportSkyAtmosphere(
 
 ////////////////////////////////////////////////////////////////////////// Regular sky 
 
-static TAutoConsoleVariable<int32> CVarSkyAtmosphereRelyOnSkyDome(
-	TEXT("r.SkyAtmosphere.RelyOnSkyDome"), 0,
-	TEXT("When true, the sky pixels will not be rendered by the full screen sky and aerialpesepctive pass.\n")
-	TEXT("Sky pixels will be skipped using depth bound test when available and\n")
-	TEXT("only the aerial perspective will be evaluated on opaque mesh pixels.\n"),
-	ECVF_RenderThreadSafe);
-
 static TAutoConsoleVariable<float> CVarSkyAtmosphereSampleCountMin(
 	TEXT("r.SkyAtmosphere.SampleCountMin"), 2.0f,
 	TEXT("The minimum sample count used to compute sky/atmosphere scattering and transmittance.\n")
@@ -312,6 +305,11 @@ static void CopyAtmosphereSetupToUniformShaderParameters(FAtmosphereUniformShade
 #undef COPYMACRO
 }
 
+static FLinearColor GetLightDiskLuminance(FLightSceneInfo& Light, FLinearColor LightIlluminance)
+{
+	const float SunSolidAngle = 2.0f * PI * (1.0f - FMath::Cos(Light.Proxy->GetSunLightHalfApexAngleRadian())); // Solid angle from aperture https://en.wikipedia.org/wiki/Solid_angle 
+	return LightIlluminance / SunSolidAngle; // approximation
+}
 
 
 /*=============================================================================
@@ -355,10 +353,7 @@ void FSkyAtmosphereRenderSceneInfo::PrepareSunLightProxy(FLightSceneInfo& SunLig
 
 	FLinearColor SunZenithIlluminance = SunLight.Proxy->GetColor();
 	FLinearColor SunOuterSpaceIlluminance = SunZenithIlluminance / TransmittanceAtZenithFinal;
-
-	// SunDiscScale is only considered as a visual tweak so we do not make it influence the sun disk outerspace luminance.
-	const float SunSolidAngle = 2.0f * PI * (1.0f - FMath::Cos(SunLight.Proxy->GetSunLightHalfApexAngleRadian())); // Solid angle from aperture https://en.wikipedia.org/wiki/Solid_angle 
-	FLinearColor SunDiskOuterSpaceLuminance = SunOuterSpaceIlluminance / SunSolidAngle; // approximation
+	FLinearColor SunDiskOuterSpaceLuminance = GetLightDiskLuminance(SunLight, SunOuterSpaceIlluminance);
 
 	SunLight.Proxy->SetAtmosphereRelatedProperties(TransmittanceTowardSun / TransmittanceAtZenithFinal, SunDiskOuterSpaceLuminance);
 }
@@ -420,6 +415,20 @@ void FScene::RemoveSkyAtmosphere(const USkyAtmosphereComponent* SkyAtmosphereCom
 		} );
 }
 
+void FScene::ResetAtmosphereLightsProperties()
+{
+	// Also rest the current atmospheric light to default atmosphere
+	for (int32 LightIndex = 0; LightIndex < NUM_ATMOSPHERE_LIGHTS; ++LightIndex)
+	{
+		FLightSceneInfo* Light = AtmosphereLights[LightIndex];
+		if (Light)
+		{
+			FLinearColor LightZenithIlluminance = Light->Proxy->GetColor();
+			Light->Proxy->SetAtmosphereRelatedProperties(FLinearColor::White, GetLightDiskLuminance(*Light, LightZenithIlluminance));
+		}
+	}
+}
+
 
 
 /*=============================================================================
@@ -1136,7 +1145,6 @@ void FSceneRenderer::RenderSkyAtmosphere(FRHICommandListImmediate& RHICmdList)
 	const bool bFastSky = CVarSkyAtmosphereFastSkyLUT.GetValueOnRenderThread() > 0;
 	const bool bFastAerialPespective = CVarSkyAtmosphereAerialPerspectiveApplyOnOpaque.GetValueOnRenderThread() > 0;
 	const bool bFastAerialPespectiveDepthTest = CVarSkyAtmosphereAerialPerspectiveDepthTest.GetValueOnRenderThread() > 0;
-	const bool bSkyRelyOnSkyDome = CVarSkyAtmosphereRelyOnSkyDome.GetValueOnAnyThread();
 	const bool bSecondAtmosphereLightEnabled = IsSecondAtmosphereLightEnabled(Scene);
 
 	FRHISamplerState* SamplerLinearClamp = TStaticSamplerState<SF_Trilinear>::GetRHI();
@@ -1157,6 +1165,11 @@ void FSceneRenderer::RenderSkyAtmosphere(FRHICommandListImmediate& RHICmdList)
 
 		const bool bLightDiskEnabled = !View.bIsReflectionCapture;
 
+		// If the scene contains Sky material then it is likely rendering the sky using a sky dome mesh.
+		// In this case we can use a simpler shader during this pass to only render aerial perspective
+		// and sky pixels can likely be optimised out.
+		const bool bRenderSkyPixel = !View.bSceneHasSkyMaterial;
+
 		const FVector ViewOrigin = View.ViewMatrices.GetViewOrigin();
 		const float KmToCm = 100000.0f;
 		const FVector PlanetOrigin = FVector(0.0f, 0.0f, -Atmosphere.BottomRadius * KmToCm);
@@ -1172,7 +1185,7 @@ void FSceneRenderer::RenderSkyAtmosphere(FRHICommandListImmediate& RHICmdList)
 			PsPermutationVector.Set<FFastAerialPespective>(bFastAerialPespective && !ForceRayMarching);
 			PsPermutationVector.Set<FSourceDiskEnabled>(bLightDiskEnabled);
 			PsPermutationVector.Set<FSecondAtmosphereLight>(bSecondAtmosphereLightEnabled);
-			PsPermutationVector.Set<FRenderSky>(!bSkyRelyOnSkyDome);
+			PsPermutationVector.Set<FRenderSky>(bRenderSkyPixel);
 			PsPermutationVector = FRenderSkyAtmospherePS::RemapPermutation(PsPermutationVector);
 			TShaderMapRef<FRenderSkyAtmospherePS> PixelShader(View.ShaderMap, PsPermutationVector);
 
@@ -1213,10 +1226,10 @@ void FSceneRenderer::RenderSkyAtmosphere(FRHICommandListImmediate& RHICmdList)
 
 			FIntRect Viewport = View.ViewRect;
 			GraphBuilder.AddPass(
-				RDG_EVENT_NAME("SkyAtmosphere"),
+				RDG_EVENT_NAME("SkyAtmosphereDraw"),
 				PsPassParameters,
 				ERDGPassFlags::Raster,
-				[PsPassParameters, VertexShader, PixelShader, Viewport, bFastAerialPespectiveDepthTest, bSkyRelyOnSkyDome, StartDepthZ](FRHICommandList& RHICmdList)
+				[PsPassParameters, VertexShader, PixelShader, Viewport, bFastAerialPespectiveDepthTest, bRenderSkyPixel, StartDepthZ](FRHICommandList& RHICmdList)
 			{
 				RHICmdList.SetViewport(Viewport.Min.X, Viewport.Min.Y, 0.0f, Viewport.Max.X, Viewport.Max.Y, 1.0f);
 
@@ -1239,7 +1252,7 @@ void FSceneRenderer::RenderSkyAtmosphere(FRHICommandListImmediate& RHICmdList)
 				GraphicsPSOInit.BoundShaderState.PixelShaderRHI = PixelShader->GetPixelShader();
 				GraphicsPSOInit.PrimitiveType = PT_TriangleList;
 
-				if (bSkyRelyOnSkyDome && GSupportsDepthBoundsTest)
+				if (!bRenderSkyPixel && GSupportsDepthBoundsTest)
 				{
 					// When we do not render the sky in the sky pass and depth bound test is supported, we take advantage of it in order to skip the processing of sky pixels.
 					GraphicsPSOInit.bDepthBounds = true;
@@ -1294,6 +1307,61 @@ void FSceneRenderer::RenderSkyAtmosphere(FRHICommandListImmediate& RHICmdList)
 #endif
 }
 
+bool FSceneRenderer::ShouldRenderSkyAtmosphereEditorNotifications()
+{
+#if WITH_EDITOR
+	bool bAnyViewHasSkyMaterial = false;
+	for (int32 ViewIndex = 0; ViewIndex < Views.Num(); ViewIndex++)
+	{
+		bAnyViewHasSkyMaterial |= Views[ViewIndex].bSceneHasSkyMaterial;
+	}
+	return bAnyViewHasSkyMaterial;
+#endif
+	return false;
+}
+
+void FSceneRenderer::RenderSkyAtmosphereEditorNotifications(FRHICommandListImmediate& RHICmdList)
+{
+#if WITH_EDITOR
+	DECLARE_GPU_STAT(SkyAtmosphereEditor);
+	SCOPED_DRAW_EVENT(RHICmdList, SkyAtmosphereEditor);
+	SCOPED_GPU_STAT(RHICmdList, SkyAtmosphereEditor);
+
+	FSceneRenderTargets& SceneContext = FSceneRenderTargets::Get(RHICmdList);
+	for (int32 ViewIndex = 0; ViewIndex < Views.Num(); ViewIndex++)
+	{
+		FViewInfo& View = Views[ViewIndex];
+
+		if (View.bSceneHasSkyMaterial)
+		{
+			const float ViewPortWidth = float(View.ViewRect.Width());
+			const float ViewPortHeight = float(View.ViewRect.Height());
+
+			FRenderTargetTemp TempRenderTarget(View, (const FTexture2DRHIRef&)SceneContext.GetSceneColor()->GetRenderTargetItem().TargetableTexture);
+			FCanvas Canvas(&TempRenderTarget, NULL, View.Family->CurrentRealTime, ViewFamily.CurrentWorldTime, ViewFamily.DeltaWorldTime, View.GetFeatureLevel());
+
+			FLinearColor TextColor(1.0f, 0.5f, 0.0f);
+			FString Text  = TEXT("Your scene contains a mesh with a sky material assigned to it.");
+			FString Text2 = TEXT("So we optimize by not rendering the sky as a full-screen pass.");
+			FString Text3 = TEXT("However, your sky mesh does not cover that part of the screen!");
+			const int32 TextWidth = Text.Len() * 7;
+			for (int i = 32; i < ViewPortWidth; i += (TextWidth + 32))
+			{
+				for (int j = 32; j < ViewPortHeight; j += 128)
+				{
+					Canvas.DrawShadowedString(i, j     , *Text , GetStatsFont(), TextColor);
+					Canvas.DrawShadowedString(i, j + 16, *Text2, GetStatsFont(), TextColor);
+					Canvas.DrawShadowedString(i, j + 32, *Text3, GetStatsFont(), TextColor);
+				}
+			}
+
+			// Flush immediately to output into the scene color right away
+			Canvas.Flush_RenderThread(RHICmdList);
+		}
+	}
+#endif
+}
+
 
 
 /*=============================================================================
diff --git a/Engine/Source/Runtime/Renderer/Private/SkyPassRendering.cpp b/Engine/Source/Runtime/Renderer/Private/SkyPassRendering.cpp
new file mode 100644
--- /dev/null
+++ b/Engine/Source/Runtime/Renderer/Private/SkyPassRendering.cpp
@@ -0,0 +1,103 @@
+// Copyright 1998-2019 Epic Games, Inc. All Rights Reserved.
+
+/*=============================================================================
+	SkyPassRendering.cpp: Sky pass rendering implementation.
+=============================================================================*/
+
+#include "SkyPassRendering.h"
+#include "BasePassRendering.h"
+#include "DeferredShadingRenderer.h"
+#include "ScenePrivate.h"
+#include "SceneRendering.h"
+#include "MeshPassProcessor.inl"
+
+
+
+FSkyPassMeshProcessor::FSkyPassMeshProcessor(const FScene* Scene, const FSceneView* InViewIfDynamicMeshCommand, const FMeshPassProcessorRenderState& InPassDrawRenderState, FMeshPassDrawListContext* InDrawListContext)
+	: FMeshPassProcessor(Scene, Scene->GetFeatureLevel(), InViewIfDynamicMeshCommand, InDrawListContext)
+	, PassDrawRenderState(InPassDrawRenderState)
+{
+}
+
+void FSkyPassMeshProcessor::AddMeshBatch(const FMeshBatch& RESTRICT MeshBatch, uint64 BatchElementMask, const FPrimitiveSceneProxy* RESTRICT PrimitiveSceneProxy, int32 StaticMeshId)
+{
+	// Determine the mesh's material and blend mode.
+	const FMaterialRenderProxy* FallbackMaterialRenderProxyPtr = nullptr;
+	const FMaterial& Material = MeshBatch.MaterialRenderProxy->GetMaterialWithFallback(FeatureLevel, FallbackMaterialRenderProxyPtr);
+
+	if (Material.IsSky())
+	{
+		const ERasterizerFillMode MeshFillMode = ComputeMeshFillMode(MeshBatch, Material);
+		const ERasterizerCullMode MeshCullMode = ComputeMeshCullMode(MeshBatch, Material);
+		const FMaterialRenderProxy& MaterialRenderProxy = FallbackMaterialRenderProxyPtr ? *FallbackMaterialRenderProxyPtr : *MeshBatch.MaterialRenderProxy;
+		Process(MeshBatch, BatchElementMask, PrimitiveSceneProxy, StaticMeshId, MaterialRenderProxy, Material, MeshFillMode, MeshCullMode);
+	}
+}
+
+void FSkyPassMeshProcessor::Process(
+	const FMeshBatch& MeshBatch,
+	uint64 BatchElementMask,
+	const FPrimitiveSceneProxy* RESTRICT PrimitiveSceneProxy,
+	int32 StaticMeshId,
+	const FMaterialRenderProxy& RESTRICT MaterialRenderProxy,
+	const FMaterial& RESTRICT MaterialResource,
+	ERasterizerFillMode MeshFillMode,
+	ERasterizerCullMode MeshCullMode)
+{
+	FUniformLightMapPolicy NoLightmapPolicy(LMP_NO_LIGHTMAP);
+	typedef FUniformLightMapPolicy LightMapPolicyType;
+	TMeshProcessorShaders<
+		TBasePassVertexShaderPolicyParamType<LightMapPolicyType>,
+		FBaseHS,
+		FBaseDS,
+		TBasePassPixelShaderPolicyParamType<LightMapPolicyType>> SkyPassShaders;
+
+	const FVertexFactory* VertexFactory = MeshBatch.VertexFactory;
+	const bool bRenderSkylight = false;
+	const bool bRenderAtmosphericFog = false;
+	GetBasePassShaders<LightMapPolicyType>(
+		MaterialResource,
+		VertexFactory->GetType(),
+		NoLightmapPolicy,
+		FeatureLevel,
+		bRenderAtmosphericFog,
+		bRenderSkylight,
+		SkyPassShaders.HullShader,
+		SkyPassShaders.DomainShader,
+		SkyPassShaders.VertexShader,
+		SkyPassShaders.PixelShader
+		);
+
+	TBasePassShaderElementData<LightMapPolicyType> ShaderElementData(nullptr);
+	ShaderElementData.InitializeMeshMaterialData(ViewIfDynamicMeshCommand, PrimitiveSceneProxy, MeshBatch, StaticMeshId, false);
+
+	const FMeshDrawCommandSortKey SortKey = CalculateMeshStaticSortKey(SkyPassShaders.VertexShader, SkyPassShaders.PixelShader);
+
+	BuildMeshDrawCommands(
+		MeshBatch,
+		BatchElementMask,
+		PrimitiveSceneProxy,
+		MaterialRenderProxy,
+		MaterialResource,
+		PassDrawRenderState,
+		SkyPassShaders,
+		MeshFillMode,
+		MeshCullMode,
+		SortKey,
+		EMeshPassFeatures::Default,
+		ShaderElementData);
+}
+
+FMeshPassProcessor* CreateSkyPassProcessor(const FScene* Scene, const FSceneView* InViewIfDynamicMeshCommand, FMeshPassDrawListContext* InDrawListContext)
+{
+	FMeshPassProcessorRenderState DrawRenderState(Scene->UniformBuffers.ViewUniformBuffer, Scene->UniformBuffers.OpaqueBasePassUniformBuffer);
+	DrawRenderState.SetInstancedViewUniformBuffer(Scene->UniformBuffers.InstancedViewUniformBuffer);
+
+	FExclusiveDepthStencil::Type BasePassDepthStencilAccess_NoDepthWrite = FExclusiveDepthStencil::Type(Scene->DefaultBasePassDepthStencilAccess & ~FExclusiveDepthStencil::DepthWrite);
+	SetupBasePassState(BasePassDepthStencilAccess_NoDepthWrite, false, DrawRenderState);
+
+	return new(FMemStack::Get()) FSkyPassMeshProcessor(Scene, InViewIfDynamicMeshCommand, DrawRenderState, InDrawListContext);
+}
+
+FRegisterPassProcessorCreateFunction RegisterSkyPass(&CreateSkyPassProcessor, EShadingPath::Deferred, EMeshPass::SkyPass, EMeshPassFlags::MainView);
+// No processor on mobile: sky meshes can be rendered as part of the base pass. The aerial perspective is removed when using shader #define.
diff --git a/Engine/Source/Runtime/Renderer/Private/SkyPassRendering.h b/Engine/Source/Runtime/Renderer/Private/SkyPassRendering.h
new file mode 100644
--- /dev/null
+++ b/Engine/Source/Runtime/Renderer/Private/SkyPassRendering.h
@@ -0,0 +1,37 @@
+// Copyright 1998-2019 Epic Games, Inc. All Rights Reserved.
+
+/*=============================================================================
+	SkyPassRendering.h: Sky pass rendering implementation.
+=============================================================================*/
+
+#pragma once
+
+#include "CoreMinimal.h"
+#include "RendererInterface.h"
+#include "MeshPassProcessor.h"
+
+
+
+class FSkyPassMeshProcessor : public FMeshPassProcessor
+{
+public:
+
+	FSkyPassMeshProcessor(const FScene* Scene, const FSceneView* InViewIfDynamicMeshCommand, const FMeshPassProcessorRenderState& InPassDrawRenderState, FMeshPassDrawListContext* InDrawListContext);
+
+	virtual void AddMeshBatch(const FMeshBatch& RESTRICT MeshBatch, uint64 BatchElementMask, const FPrimitiveSceneProxy* RESTRICT PrimitiveSceneProxy, int32 StaticMeshId = -1) override final;
+
+	FMeshPassProcessorRenderState PassDrawRenderState;
+
+private:
+	void Process(
+		const FMeshBatch& MeshBatch,
+		uint64 BatchElementMask,
+		const FPrimitiveSceneProxy* RESTRICT PrimitiveSceneProxy,
+		int32 StaticMeshId,
+		const FMaterialRenderProxy& RESTRICT MaterialRenderProxy,
+		const FMaterial& RESTRICT MaterialResource,
+		ERasterizerFillMode MeshFillMode,
+		ERasterizerCullMode MeshCullMode);
+};
+
+
diff --git a/Engine/Source/Runtime/Renderer/Private/VelocityRendering.cpp b/Engine/Source/Runtime/Renderer/Private/VelocityRendering.cpp
--- a/Engine/Source/Runtime/Renderer/Private/VelocityRendering.cpp
+++ b/Engine/Source/Runtime/Renderer/Private/VelocityRendering.cpp
@@ -572,7 +572,8 @@ void FVelocityMeshProcessor::AddMeshBatch(const FMeshBatch& RESTRICT MeshBatch,
 		const ERasterizerFillMode MeshFillMode = ComputeMeshFillMode(MeshBatch, *Material);
 		const ERasterizerCullMode MeshCullMode = ComputeMeshCullMode(MeshBatch, *Material);
 
-		if (BlendMode == BLEND_Opaque || BlendMode == BLEND_Masked)
+		if ((BlendMode == BLEND_Opaque || BlendMode == BLEND_Masked)
+			&& !Material->IsSky())
 		{
 			if (Material->WritesEveryPixel() && !Material->IsTwoSided() && !Material->MaterialModifiesMeshPosition_RenderThread())
 			{
diff --git a/Engine/Source/Runtime/Renderer/Public/MeshPassProcessor.h b/Engine/Source/Runtime/Renderer/Public/MeshPassProcessor.h
--- a/Engine/Source/Runtime/Renderer/Public/MeshPassProcessor.h
+++ b/Engine/Source/Runtime/Renderer/Public/MeshPassProcessor.h
@@ -19,6 +19,7 @@ namespace EMeshPass
 	{
 		DepthPass,
 		BasePass,
+		SkyPass,
 		CSMShadowDepth,
 		Distortion,
 		Velocity,
@@ -50,6 +51,7 @@ inline const TCHAR* GetMeshPassName(EMeshPass::Type MeshPass)
 	{
 	case EMeshPass::DepthPass: return TEXT("DepthPass");
 	case EMeshPass::BasePass: return TEXT("BasePass");
+	case EMeshPass::SkyPass: return TEXT("SkyPass");
 	case EMeshPass::CSMShadowDepth: return TEXT("CSMShadowDepth");
 	case EMeshPass::Distortion: return TEXT("Distortion");
 	case EMeshPass::Velocity: return TEXT("Velocity");
