Fixed weird lambda variable shadowing issue.

diff --git a/Engine/Source/Runtime/Renderer/Private/SkyAtmosphereRendering.cpp b/Engine/Source/Runtime/Renderer/Private/SkyAtmosphereRendering.cpp
--- a/Engine/Source/Runtime/Renderer/Private/SkyAtmosphereRendering.cpp
+++ b/Engine/Source/Runtime/Renderer/Private/SkyAtmosphereRendering.cpp
@@ -1229,12 +1229,12 @@ void FSceneRenderer::RenderSkyAtmosphere(FRHICommandListImmediate& RHICmdList)
 				RDG_EVENT_NAME("SkyAtmosphereDraw"),
 				PsPassParameters,
 				ERDGPassFlags::Raster,
-				[PsPassParameters, VertexShader, PixelShader, Viewport, bFastAerialPespectiveDepthTest, bRenderSkyPixel, StartDepthZ](FRHICommandList& RHICmdList)
+				[PsPassParameters, VertexShader, PixelShader, Viewport, bFastAerialPespectiveDepthTest, bRenderSkyPixel, StartDepthZ](FRHICommandList& RHICmdListLambda)
 			{
-				RHICmdList.SetViewport(Viewport.Min.X, Viewport.Min.Y, 0.0f, Viewport.Max.X, Viewport.Max.Y, 1.0f);
+				RHICmdListLambda.SetViewport(Viewport.Min.X, Viewport.Min.Y, 0.0f, Viewport.Max.X, Viewport.Max.Y, 1.0f);
 
 				FGraphicsPipelineStateInitializer GraphicsPSOInit;
-				RHICmdList.ApplyCachedRenderTargets(GraphicsPSOInit);
+				RHICmdListLambda.ApplyCachedRenderTargets(GraphicsPSOInit);
 
 				GraphicsPSOInit.BlendState = TStaticBlendState<CW_RGB, BO_Add, BF_One, BF_SourceAlpha, BO_Add, BF_Zero, BF_One>::GetRHI();
 				GraphicsPSOInit.RasterizerState = TStaticRasterizerState<FM_Solid, CM_None>::GetRHI();
@@ -1260,25 +1260,25 @@ void FSceneRenderer::RenderSkyAtmosphere(FRHICommandListImmediate& RHICmdList)
 					{
 						//const float SmallestFloatAbove0 = 1.1754943508e-38;		// 32bit float depth
 						const float SmallestFloatAbove0 = 1.0f / 16777215.0f;		// 24bit norm depth
-						RHICmdList.SetDepthBounds(SmallestFloatAbove0, 1.0f);		// Tested on dx12 PC
+						RHICmdListLambda.SetDepthBounds(SmallestFloatAbove0, 1.0f);		// Tested on dx12 PC
 					}
 					else
 					{
 						//const float SmallestFloatBelow1 = 0.9999999404;			// 32bit float depth
 						const float SmallestFloatBelow1 = 16777214.0f / 16777215.0f;// 24bit norm depth
-						RHICmdList.SetDepthBounds(0.0f, SmallestFloatBelow1);		// Untested
+						RHICmdListLambda.SetDepthBounds(0.0f, SmallestFloatBelow1);		// Untested
 					}
 				}
 
-				SetGraphicsPipelineState(RHICmdList, GraphicsPSOInit);
+				SetGraphicsPipelineState(RHICmdListLambda, GraphicsPSOInit);
 
-				SetShaderParameters(RHICmdList, *PixelShader, PixelShader->GetPixelShader(), *PsPassParameters);
+				SetShaderParameters(RHICmdListLambda, *PixelShader, PixelShader->GetPixelShader(), *PsPassParameters);
 
 				FRenderSkyAtmosphereVS::FParameters VsPassParameters;
 				VsPassParameters.StartDepthZ = StartDepthZ;
-				SetShaderParameters(RHICmdList, *VertexShader, VertexShader->GetVertexShader(), VsPassParameters);
+				SetShaderParameters(RHICmdListLambda, *VertexShader, VertexShader->GetVertexShader(), VsPassParameters);
 
-				RHICmdList.DrawPrimitive(0, 1, 1);
+				RHICmdListLambda.DrawPrimitive(0, 1, 1);
 			});
 		}
 	}
