Fixes static analysis warning in SceneTextureReductions.cpp

diff --git a/Engine/Source/Runtime/Renderer/Private/SceneTextureReductions.cpp b/Engine/Source/Runtime/Renderer/Private/SceneTextureReductions.cpp
--- a/Engine/Source/Runtime/Renderer/Private/SceneTextureReductions.cpp
+++ b/Engine/Source/Runtime/Renderer/Private/SceneTextureReductions.cpp
@@ -259,6 +259,7 @@ void BuildHZB(FRDGBuilder& GraphBuilder, const FSceneTextureParameters& SceneTex
 
 		if (bReduceClosestDepth)
 		{
+			check(ClosestHZBTexture)
 			FRDGTextureSRVRef ParentTextureMip = GraphBuilder.CreateSRV(FRDGTextureSRVDesc::CreateForMipLevel(ClosestHZBTexture, StartDestMip - 1));
 			ReduceMips(ParentTextureMip,
 				StartDestMip, DispatchThreadIdToBufferUV, InputViewportMaxBound,
