Fixed depth clip issue in sky refelction only.
Renaming.

diff --git a/Engine/Source/Runtime/Renderer/Private/SkyAtmosphereRendering.cpp b/Engine/Source/Runtime/Renderer/Private/SkyAtmosphereRendering.cpp
--- a/Engine/Source/Runtime/Renderer/Private/SkyAtmosphereRendering.cpp
+++ b/Engine/Source/Runtime/Renderer/Private/SkyAtmosphereRendering.cpp
@@ -1143,8 +1143,8 @@ void FSceneRenderer::RenderSkyAtmosphere(FRHICommandListImmediate& RHICmdList)
 	const FAtmosphereSetup& Atmosphere = SkyInfo.GetAtmosphereSetup();
 	const bool bMultiScattering = SkyInfo.IsMultiScatteringEnabled();
 	const bool bFastSky = CVarSkyAtmosphereFastSkyLUT.GetValueOnRenderThread() > 0;
-	const bool bFastAerialPespective = CVarSkyAtmosphereAerialPerspectiveApplyOnOpaque.GetValueOnRenderThread() > 0;
-	const bool bFastAerialPespectiveDepthTest = CVarSkyAtmosphereAerialPerspectiveDepthTest.GetValueOnRenderThread() > 0;
+	const bool bFastAerialPerspective = CVarSkyAtmosphereAerialPerspectiveApplyOnOpaque.GetValueOnRenderThread() > 0;
+	const bool bFastAerialPerspectiveDepthTest = CVarSkyAtmosphereAerialPerspectiveDepthTest.GetValueOnRenderThread() > 0;
 	const bool bSecondAtmosphereLightEnabled = IsSecondAtmosphereLightEnabled(Scene);
 
 	FRHISamplerState* SamplerLinearClamp = TStaticSamplerState<SF_Trilinear>::GetRHI();
@@ -1182,7 +1182,7 @@ void FSceneRenderer::RenderSkyAtmosphere(FRHICommandListImmediate& RHICmdList)
 			FRenderSkyAtmospherePS::FPermutationDomain PsPermutationVector;
 			PsPermutationVector.Set<FSkyPermutationMultiScatteringApprox>(bMultiScattering);
 			PsPermutationVector.Set<FFastSky>(bFastSky && !ForceRayMarching);
-			PsPermutationVector.Set<FFastAerialPespective>(bFastAerialPespective && !ForceRayMarching);
+			PsPermutationVector.Set<FFastAerialPespective>(bFastAerialPerspective && !ForceRayMarching);
 			PsPermutationVector.Set<FSourceDiskEnabled>(bLightDiskEnabled);
 			PsPermutationVector.Set<FSecondAtmosphereLight>(bSecondAtmosphereLightEnabled);
 			PsPermutationVector.Set<FRenderSky>(bRenderSkyPixel);
@@ -1213,14 +1213,15 @@ void FSceneRenderer::RenderSkyAtmosphere(FRHICommandListImmediate& RHICmdList)
 			ClearUnusedGraphResources(*PixelShader, PsPassParameters);
 
 			float StartDepthZ = 0.1f;
-			if (bFastAerialPespectiveDepthTest)
+			if (bFastAerialPerspectiveDepthTest)
 			{
 				const FMatrix ViewProjectionMatrix = View.ViewMatrices.GetProjectionMatrix();
 				float HalfHorizontalFOV = FMath::Atan(1.0f / ViewProjectionMatrix.M[0][0]);
 				float HalfVerticalFOV = FMath::Atan(1.0f / ViewProjectionMatrix.M[1][1]);
-				const float HalfFOV = View.FOV * 0.5f * PI / 180.0f;
-				const float StartDepthView = FMath::Cos(FMath::Max(HalfHorizontalFOV, HalfVerticalFOV)) * InternalCommonParameters.AerialPerspectiveStartDepth * KmToCm;
-				const FVector4 Projected = ViewProjectionMatrix.TransformFVector4(FVector4(0.0f, 0.0f, StartDepthView, 1.0f));
+				// For sky reflection capture, the start depth can be super large. Se we max it to make sure the triangle is never in front the NearClippingDistance.
+				float StartDepthInKm = FMath::Max(InternalCommonParameters.AerialPerspectiveStartDepth, View.NearClippingDistance) * KmToCm;
+				const float StartDepthViewKm = FMath::Cos(FMath::Max(HalfHorizontalFOV, HalfVerticalFOV)) * StartDepthInKm;
+				const FVector4 Projected = ViewProjectionMatrix.TransformFVector4(FVector4(0.0f, 0.0f, StartDepthViewKm, 1.0f));
 				StartDepthZ = Projected.Z / Projected.W;
 			}
 
@@ -1229,7 +1230,7 @@ void FSceneRenderer::RenderSkyAtmosphere(FRHICommandListImmediate& RHICmdList)
 				RDG_EVENT_NAME("SkyAtmosphereDraw"),
 				PsPassParameters,
 				ERDGPassFlags::Raster,
-				[PsPassParameters, VertexShader, PixelShader, Viewport, bFastAerialPespectiveDepthTest, bRenderSkyPixel, StartDepthZ](FRHICommandList& RHICmdListLambda)
+				[PsPassParameters, VertexShader, PixelShader, Viewport, bFastAerialPerspectiveDepthTest, bRenderSkyPixel, StartDepthZ](FRHICommandList& RHICmdListLambda)
 			{
 				RHICmdListLambda.SetViewport(Viewport.Min.X, Viewport.Min.Y, 0.0f, Viewport.Max.X, Viewport.Max.Y, 1.0f);
 
@@ -1238,7 +1239,7 @@ void FSceneRenderer::RenderSkyAtmosphere(FRHICommandListImmediate& RHICmdList)
 
 				GraphicsPSOInit.BlendState = TStaticBlendState<CW_RGB, BO_Add, BF_One, BF_SourceAlpha, BO_Add, BF_Zero, BF_One>::GetRHI();
 				GraphicsPSOInit.RasterizerState = TStaticRasterizerState<FM_Solid, CM_None>::GetRHI();
-				if (bFastAerialPespectiveDepthTest)
+				if (bFastAerialPerspectiveDepthTest)
 				{
 					GraphicsPSOInit.DepthStencilState = TStaticDepthStencilState<false, CF_DepthNearOrEqual>::GetRHI();
 				}
