Documents previous frame information a little bit better.

diff --git a/Engine/Source/Runtime/Renderer/Private/LightRendering.cpp b/Engine/Source/Runtime/Renderer/Private/LightRendering.cpp
--- a/Engine/Source/Runtime/Renderer/Private/LightRendering.cpp
+++ b/Engine/Source/Runtime/Renderer/Private/LightRendering.cpp
@@ -954,7 +954,7 @@ void FDeferredShadingSceneRenderer::RenderLights(FRHICommandListImmediate& RHICm
 
 				// Allocate PreprocessedShadowMaskTextures once so QueueTextureExtraction can deferred write.
 				{
-					if (!View.bViewStateIsReadOnly)
+					if (!View.bStatePrevViewInfoIsReadOnly)
 					{
 						View.ViewState->PrevFrameViewInfo.ShadowHistories.Empty();
 						View.ViewState->PrevFrameViewInfo.ShadowHistories.Reserve(SortedLights.Num());
diff --git a/Engine/Source/Runtime/Renderer/Private/PostProcess/PostProcessEyeAdaptation.cpp b/Engine/Source/Runtime/Renderer/Private/PostProcess/PostProcessEyeAdaptation.cpp
--- a/Engine/Source/Runtime/Renderer/Private/PostProcess/PostProcessEyeAdaptation.cpp
+++ b/Engine/Source/Runtime/Renderer/Private/PostProcess/PostProcessEyeAdaptation.cpp
@@ -505,7 +505,7 @@ void FSceneViewState::UpdatePreExposure(FViewInfo& View)
 	View.PreExposure = PreExposure;
 
 	// Update the pre exposure of all temporal histories.
-	if (!View.bViewStateIsReadOnly)
+	if (!View.bStatePrevViewInfoIsReadOnly)
 	{
 		PrevFrameViewInfo.SceneColorPreExposure = PreExposure;
 	}
diff --git a/Engine/Source/Runtime/Renderer/Private/PostProcess/PostProcessTemporalAA.cpp b/Engine/Source/Runtime/Renderer/Private/PostProcess/PostProcessTemporalAA.cpp
--- a/Engine/Source/Runtime/Renderer/Private/PostProcess/PostProcessTemporalAA.cpp
+++ b/Engine/Source/Runtime/Renderer/Private/PostProcess/PostProcessTemporalAA.cpp
@@ -661,7 +661,7 @@ FTAAOutputs AddTemporalAAPass(
 			FComputeShaderUtils::GetGroupCount(PracticableDestRect.Size(), GTemporalAATileSizeX));
 	}
 	
-	if (!View.bViewStateIsReadOnly)
+	if (!View.bStatePrevViewInfoIsReadOnly)
 	{
 		OutputHistory->SafeRelease();
 
diff --git a/Engine/Source/Runtime/Renderer/Private/PostProcess/PostProcessing.cpp b/Engine/Source/Runtime/Renderer/Private/PostProcess/PostProcessing.cpp
--- a/Engine/Source/Runtime/Renderer/Private/PostProcess/PostProcessing.cpp
+++ b/Engine/Source/Runtime/Renderer/Private/PostProcess/PostProcessing.cpp
@@ -1183,7 +1183,7 @@ void FPostProcessing::Process(FRHICommandListImmediate& RHICmdList, const FViewI
 				//! SceneColorTexture is now upsampled to the SecondaryViewRect. Use SecondaryViewRect for input / output.
 
 				// Post Process Material Chain - SSR Input
-				if (View.ViewState && !View.bViewStateIsReadOnly)
+				if (View.ViewState && !View.bStatePrevViewInfoIsReadOnly)
 				{
 					FPostProcessMaterialInputs Inputs;
 					Inputs.SetInput(EPostProcessMaterialInput::SceneColor, SceneColorTexture, SecondaryViewRect);
diff --git a/Engine/Source/Runtime/Renderer/Private/ScenePrivate.h b/Engine/Source/Runtime/Renderer/Private/ScenePrivate.h
--- a/Engine/Source/Runtime/Renderer/Private/ScenePrivate.h
+++ b/Engine/Source/Runtime/Renderer/Private/ScenePrivate.h
@@ -844,8 +844,15 @@ public:
 
 	// counts up by one each frame, warped in 0..7 range, ResetViewState() puts it back to 0
 	uint32 FrameIndex;
-
-	// Previous frame's view info to use.
+	
+	/** Informations of to persist for the next frame's FViewInfo::PrevViewInfo.
+	 *
+	 * Under normal use case (temporal histories are not frozen), this gets cleared after setting FViewInfo::PrevViewInfo
+	 * after being copied to FViewInfo::PrevViewInfo. New temporal histories get directly written to it.
+	 *
+	 * When temporal histories are frozen (pause command, or r.Test.FreezeTemporalHistories), this keeps it's values, and the currently
+	 * rendering FViewInfo should not update it. Refer to FViewInfo::bStatePrevViewInfoIsReadOnly.
+	 */
 	FPreviousViewInfo PrevFrameViewInfo;
 
 	FHeightfieldLightingAtlas* HeightfieldLightingAtlas;
diff --git a/Engine/Source/Runtime/Renderer/Private/SceneRendering.cpp b/Engine/Source/Runtime/Renderer/Private/SceneRendering.cpp
--- a/Engine/Source/Runtime/Renderer/Private/SceneRendering.cpp
+++ b/Engine/Source/Runtime/Renderer/Private/SceneRendering.cpp
@@ -820,7 +820,7 @@ void FViewInfo::Init()
 
 	bIsViewInfo = true;
 	
-	bViewStateIsReadOnly = true;
+	bStatePrevViewInfoIsReadOnly = true;
 	bUsesGlobalDistanceField = false;
 	bUsesLightingChannels = false;
 	bTranslucentSurfaceLighting = false;
diff --git a/Engine/Source/Runtime/Renderer/Private/SceneRendering.h b/Engine/Source/Runtime/Renderer/Private/SceneRendering.h
--- a/Engine/Source/Runtime/Renderer/Private/SceneRendering.h
+++ b/Engine/Source/Runtime/Renderer/Private/SceneRendering.h
@@ -985,8 +985,8 @@ public:
 	/** Temporal AA jitter at the pixel scale. */
 	FVector2D TemporalJitterPixels;
 
-	/** Whether view state may be updated with this view. */
-	uint32 bViewStateIsReadOnly : 1;
+	/** Whether FSceneViewState::PrevFrameViewInfo can be updated with this view. */
+	uint32 bStatePrevViewInfoIsReadOnly : 1;
 
 	/** true if all PrimitiveVisibilityMap's bits are set to false. */
 	uint32 bHasNoVisiblePrimitive : 1;
@@ -1025,7 +1025,7 @@ public:
 	/** Bitmask of all shading models used by primitives in this view */
 	uint16 ShadingModelMaskInView;
 
-	// Previous frame view info to use for this view.
+	/** Informations from the previous frame to use for this view. */
 	FPreviousViewInfo PrevViewInfo;
 
 	/** The GPU nodes on which to render this view. */
diff --git a/Engine/Source/Runtime/Renderer/Private/SceneVisibility.cpp b/Engine/Source/Runtime/Renderer/Private/SceneVisibility.cpp
--- a/Engine/Source/Runtime/Renderer/Private/SceneVisibility.cpp
+++ b/Engine/Source/Runtime/Renderer/Private/SceneVisibility.cpp
@@ -3021,12 +3021,12 @@ void FSceneRenderer::PreVisibilityFrameSetup(FRHICommandListImmediate& RHICmdLis
 
 		if (ViewState)
 		{
-			check(View.bViewStateIsReadOnly);
-			View.bViewStateIsReadOnly = ViewFamily.bWorldIsPaused || ViewFamily.EngineShowFlags.HitProxies || bFreezeTemporalHistories;
+			check(View.bStatePrevViewInfoIsReadOnly);
+			View.bStatePrevViewInfoIsReadOnly = ViewFamily.bWorldIsPaused || ViewFamily.EngineShowFlags.HitProxies || bFreezeTemporalHistories;
 
 			ViewState->SetupDistanceFieldTemporalOffset(ViewFamily);
 
-			if (!View.bViewStateIsReadOnly && !bFreezeTemporalSequences)
+			if (!View.bStatePrevViewInfoIsReadOnly && !bFreezeTemporalSequences)
 			{
 				ViewState->FrameIndex++;
 			}
@@ -3072,7 +3072,7 @@ void FSceneRenderer::PreVisibilityFrameSetup(FRHICommandListImmediate& RHICmdLis
 			}
 
 			// Updates view state.
-			if (!View.bViewStateIsReadOnly && !bFreezeTemporalSequences)
+			if (!View.bStatePrevViewInfoIsReadOnly && !bFreezeTemporalSequences)
 			{
 				ViewState->TemporalAASampleIndex = TemporalSampleIndex;
 			}
@@ -3254,7 +3254,7 @@ void FSceneRenderer::PreVisibilityFrameSetup(FRHICommandListImmediate& RHICmdLis
 			}
 
 			// Replace previous view info of the view state with this frame, clearing out references over render target.
-			if (!View.bViewStateIsReadOnly)
+			if (!View.bStatePrevViewInfoIsReadOnly)
 			{
 				ViewState->PrevFrameViewInfo = NewPrevViewInfo;
 			}
diff --git a/Engine/Source/Runtime/Renderer/Private/ScreenSpaceDenoise.cpp b/Engine/Source/Runtime/Renderer/Private/ScreenSpaceDenoise.cpp
--- a/Engine/Source/Runtime/Renderer/Private/ScreenSpaceDenoise.cpp
+++ b/Engine/Source/Runtime/Renderer/Private/ScreenSpaceDenoise.cpp
@@ -1567,7 +1567,7 @@ static void DenoiseSignalAtConstantPixelDensity(
 		SignalHistory = SignalOutput;
 	} // if (MaxPostFilterSampleCount > 1)
 
-	if (!View.bViewStateIsReadOnly && Settings.bUseTemporalAccumulation)
+	if (!View.bStatePrevViewInfoIsReadOnly && Settings.bUseTemporalAccumulation)
 	{
 		check(View.ViewState);
 
@@ -1769,7 +1769,7 @@ public:
 			PrevHistories[BatchedSignalId] = PreviousViewInfos->ShadowHistories.Find(Settings.LightSceneInfo[BatchedSignalId]->Proxy->GetLightComponent());
 			NewHistories[BatchedSignalId] = nullptr;
 				
-			if (!View.bViewStateIsReadOnly)
+			if (!View.bStatePrevViewInfoIsReadOnly)
 			{
 				check(View.ViewState);
 				NewHistories[BatchedSignalId] = &View.ViewState->PrevFrameViewInfo.ShadowHistories.FindOrAdd(Settings.LightSceneInfo[BatchedSignalId]->Proxy->GetLightComponent());
