Renames FScreenSpaceFilteringHistory to FScreenSpaceDenoiserHistory

diff --git a/Engine/Source/Runtime/Renderer/Private/SceneRendering.h b/Engine/Source/Runtime/Renderer/Private/SceneRendering.h
--- a/Engine/Source/Runtime/Renderer/Private/SceneRendering.h
+++ b/Engine/Source/Runtime/Renderer/Private/SceneRendering.h
@@ -715,8 +715,8 @@ struct FTemporalAAHistory
 	}
 };
 
-// TODO: merge with FTemporalAAHistory?
-struct FScreenSpaceFilteringHistory
+/** Temporal history for a denoiser. */
+struct FScreenSpaceDenoiserHistory
 {
 	// Number of history render target to store.
 	static constexpr int32 RTCount = 3;
@@ -770,22 +770,22 @@ struct FPreviousViewInfo
 	TRefCountPtr<IPooledRenderTarget> CustomSSRInput;
 
 	// History for the reflections
-	FScreenSpaceFilteringHistory ReflectionsHistory;
+	FScreenSpaceDenoiserHistory ReflectionsHistory;
 	
 	// History for the ambient occlusion
-	FScreenSpaceFilteringHistory AmbientOcclusionHistory;
+	FScreenSpaceDenoiserHistory AmbientOcclusionHistory;
 
 	// History for global illumination
-	FScreenSpaceFilteringHistory DiffuseIndirectHistory;
+	FScreenSpaceDenoiserHistory DiffuseIndirectHistory;
 
 	// History for sky light
-	FScreenSpaceFilteringHistory SkyLightHistory;
+	FScreenSpaceDenoiserHistory SkyLightHistory;
 
 	// History for shadow denoising.
-	TMap<const ULightComponent*, FScreenSpaceFilteringHistory> ShadowHistories;
+	TMap<const ULightComponent*, FScreenSpaceDenoiserHistory> ShadowHistories;
 
 	// History for denoising all lights penumbra at once.
-	FScreenSpaceFilteringHistory PolychromaticPenumbraHarmonicsHistory;
+	FScreenSpaceDenoiserHistory PolychromaticPenumbraHarmonicsHistory;
 };
 
 class FViewCommands
diff --git a/Engine/Source/Runtime/Renderer/Private/ScreenSpaceDenoise.cpp b/Engine/Source/Runtime/Renderer/Private/ScreenSpaceDenoise.cpp
--- a/Engine/Source/Runtime/Renderer/Private/ScreenSpaceDenoise.cpp
+++ b/Engine/Source/Runtime/Renderer/Private/ScreenSpaceDenoise.cpp
@@ -1026,8 +1026,8 @@ static void DenoiseSignalAtConstantPixelDensity(
 	const FSceneTextureParameters& SceneTextures,
 	const FSSDSignalTextures& InputSignal,
 	FSSDConstantPixelDensitySettings Settings,
-	TStaticArray<FScreenSpaceFilteringHistory*, IScreenSpaceDenoiser::kMaxBatchSize> PrevFilteringHistory,
-	TStaticArray<FScreenSpaceFilteringHistory*, IScreenSpaceDenoiser::kMaxBatchSize> NewFilteringHistory,
+	TStaticArray<FScreenSpaceDenoiserHistory*, IScreenSpaceDenoiser::kMaxBatchSize> PrevFilteringHistory,
+	TStaticArray<FScreenSpaceDenoiserHistory*, IScreenSpaceDenoiser::kMaxBatchSize> NewFilteringHistory,
 	FSSDSignalTextures* OutputSignal)
 {
 	check(UsesConstantPixelDensityPassLayout(Settings.SignalProcessing));
@@ -1494,12 +1494,12 @@ static void DenoiseSignalAtConstantPixelDensity(
 		PassParameters->PrevGBufferA = RegisterExternalTextureWithFallback(GraphBuilder, View.PrevViewInfo.GBufferA, GSystemTextures.BlackDummy);
 		PassParameters->PrevGBufferB = RegisterExternalTextureWithFallback(GraphBuilder, View.PrevViewInfo.GBufferB, GSystemTextures.BlackDummy);
 
-		FScreenSpaceFilteringHistory DummyPrevFrameHistory;
+		FScreenSpaceDenoiserHistory DummyPrevFrameHistory;
 
 		// Setup signals' previous frame historu buffers.
 		for (int32 BatchedSignalId = 0; BatchedSignalId < Settings.SignalBatchSize; BatchedSignalId++)
 		{
-			FScreenSpaceFilteringHistory* PrevFrameHistory = PrevFilteringHistory[BatchedSignalId] ? PrevFilteringHistory[BatchedSignalId] : &DummyPrevFrameHistory;
+			FScreenSpaceDenoiserHistory* PrevFrameHistory = PrevFilteringHistory[BatchedSignalId] ? PrevFilteringHistory[BatchedSignalId] : &DummyPrevFrameHistory;
 
 			PassParameters->bCameraCut[BatchedSignalId] = !PrevFrameHistory->IsValid();
 
@@ -1596,7 +1596,7 @@ static void DenoiseSignalAtConstantPixelDensity(
 		// Saves signal histories.
 		for (int32 BatchedSignalId = 0; BatchedSignalId < Settings.SignalBatchSize; BatchedSignalId++)
 		{
-			FScreenSpaceFilteringHistory* NewHistory = NewFilteringHistory[BatchedSignalId];
+			FScreenSpaceDenoiserHistory* NewHistory = NewFilteringHistory[BatchedSignalId];
 			check(NewHistory);
 
 			for (int32 BufferId = 0; BufferId < HistoryTextureCountPerSignal; BufferId++)
@@ -1747,8 +1747,8 @@ public:
 			Settings.MaxInputSPP = FMath::Max(Settings.MaxInputSPP, InputParameters[BatchedSignalId].RayTracingConfig.RayCountPerPixel);
 		}
 
-		TStaticArray<FScreenSpaceFilteringHistory*, IScreenSpaceDenoiser::kMaxBatchSize> PrevHistories;
-		TStaticArray<FScreenSpaceFilteringHistory*, IScreenSpaceDenoiser::kMaxBatchSize> NewHistories;
+		TStaticArray<FScreenSpaceDenoiserHistory*, IScreenSpaceDenoiser::kMaxBatchSize> PrevHistories;
+		TStaticArray<FScreenSpaceDenoiserHistory*, IScreenSpaceDenoiser::kMaxBatchSize> NewHistories;
 		for (int32 BatchedSignalId = 0; BatchedSignalId < InputParameterCount; BatchedSignalId++)
 		{
 			const FShadowParameters& Parameters = InputParameters[BatchedSignalId];
@@ -1825,8 +1825,8 @@ public:
 			Settings.ReconstructionSamples = Periode * Periode; // TODO: should use preconvolution instead for harmonic 3
 			Settings.bUseTemporalAccumulation = false;
 
-			TStaticArray<FScreenSpaceFilteringHistory*, IScreenSpaceDenoiser::kMaxBatchSize> PrevHistories;
-			TStaticArray<FScreenSpaceFilteringHistory*, IScreenSpaceDenoiser::kMaxBatchSize> NewHistories;
+			TStaticArray<FScreenSpaceDenoiserHistory*, IScreenSpaceDenoiser::kMaxBatchSize> PrevHistories;
+			TStaticArray<FScreenSpaceDenoiserHistory*, IScreenSpaceDenoiser::kMaxBatchSize> NewHistories;
 			PrevHistories[0] = nullptr;
 			NewHistories[0] = nullptr;
 
@@ -1860,8 +1860,8 @@ public:
 			Settings.ReconstructionSamples = Periode * Periode; // TODO: should use preconvolution instead for harmonic 3
 			Settings.bUseTemporalAccumulation = false;
 
-			TStaticArray<FScreenSpaceFilteringHistory*, IScreenSpaceDenoiser::kMaxBatchSize> PrevHistories;
-			TStaticArray<FScreenSpaceFilteringHistory*, IScreenSpaceDenoiser::kMaxBatchSize> NewHistories;
+			TStaticArray<FScreenSpaceDenoiserHistory*, IScreenSpaceDenoiser::kMaxBatchSize> PrevHistories;
+			TStaticArray<FScreenSpaceDenoiserHistory*, IScreenSpaceDenoiser::kMaxBatchSize> NewHistories;
 			PrevHistories[0] = nullptr;
 			NewHistories[0] = nullptr;
 
@@ -1935,8 +1935,8 @@ public:
 			Settings.bEnableReconstruction = false;
 			Settings.bUseTemporalAccumulation = CVarShadowTemporalAccumulation.GetValueOnRenderThread() != 0;
 
-			TStaticArray<FScreenSpaceFilteringHistory*, IScreenSpaceDenoiser::kMaxBatchSize> PrevHistories;
-			TStaticArray<FScreenSpaceFilteringHistory*, IScreenSpaceDenoiser::kMaxBatchSize> NewHistories;
+			TStaticArray<FScreenSpaceDenoiserHistory*, IScreenSpaceDenoiser::kMaxBatchSize> PrevHistories;
+			TStaticArray<FScreenSpaceDenoiserHistory*, IScreenSpaceDenoiser::kMaxBatchSize> NewHistories;
 			PrevHistories[0] = &PreviousViewInfos->PolychromaticPenumbraHarmonicsHistory;
 			NewHistories[0] = View.ViewState ? &View.ViewState->PrevFrameViewInfo.PolychromaticPenumbraHarmonicsHistory : nullptr;
 
@@ -1980,8 +1980,8 @@ public:
 		Settings.HistoryConvolutionSampleCount = CVarReflectionHistoryConvolutionSampleCount.GetValueOnRenderThread();
 		Settings.MaxInputSPP = RayTracingConfig.RayCountPerPixel;
 
-		TStaticArray<FScreenSpaceFilteringHistory*, IScreenSpaceDenoiser::kMaxBatchSize> PrevHistories;
-		TStaticArray<FScreenSpaceFilteringHistory*, IScreenSpaceDenoiser::kMaxBatchSize> NewHistories;
+		TStaticArray<FScreenSpaceDenoiserHistory*, IScreenSpaceDenoiser::kMaxBatchSize> PrevHistories;
+		TStaticArray<FScreenSpaceDenoiserHistory*, IScreenSpaceDenoiser::kMaxBatchSize> NewHistories;
 		PrevHistories[0] = &PreviousViewInfos->ReflectionsHistory;
 		NewHistories[0] = View.ViewState ? &View.ViewState->PrevFrameViewInfo.ReflectionsHistory : nullptr;
 
@@ -2021,8 +2021,8 @@ public:
 		Settings.HistoryConvolutionKernelSpreadFactor = CVarAOHistoryConvolutionKernelSpreadFactor.GetValueOnRenderThread();
 		Settings.MaxInputSPP = RayTracingConfig.RayCountPerPixel;
 
-		TStaticArray<FScreenSpaceFilteringHistory*, IScreenSpaceDenoiser::kMaxBatchSize> PrevHistories;
-		TStaticArray<FScreenSpaceFilteringHistory*, IScreenSpaceDenoiser::kMaxBatchSize> NewHistories;
+		TStaticArray<FScreenSpaceDenoiserHistory*, IScreenSpaceDenoiser::kMaxBatchSize> PrevHistories;
+		TStaticArray<FScreenSpaceDenoiserHistory*, IScreenSpaceDenoiser::kMaxBatchSize> NewHistories;
 		PrevHistories[0] = &PreviousViewInfos->AmbientOcclusionHistory;
 		NewHistories[0] = View.ViewState ? &View.ViewState->PrevFrameViewInfo.AmbientOcclusionHistory : nullptr;
 
@@ -2063,8 +2063,8 @@ public:
 		Settings.HistoryConvolutionKernelSpreadFactor = CVarGIHistoryConvolutionKernelSpreadFactor.GetValueOnRenderThread();
 		Settings.MaxInputSPP = Config.RayCountPerPixel;
 
-		TStaticArray<FScreenSpaceFilteringHistory*, IScreenSpaceDenoiser::kMaxBatchSize> PrevHistories;
-		TStaticArray<FScreenSpaceFilteringHistory*, IScreenSpaceDenoiser::kMaxBatchSize> NewHistories;
+		TStaticArray<FScreenSpaceDenoiserHistory*, IScreenSpaceDenoiser::kMaxBatchSize> PrevHistories;
+		TStaticArray<FScreenSpaceDenoiserHistory*, IScreenSpaceDenoiser::kMaxBatchSize> NewHistories;
 		PrevHistories[0] = &PreviousViewInfos->DiffuseIndirectHistory;
 		NewHistories[0] = View.ViewState ? &View.ViewState->PrevFrameViewInfo.DiffuseIndirectHistory : nullptr;
 
@@ -2105,8 +2105,8 @@ public:
 		Settings.HistoryConvolutionKernelSpreadFactor = CVarGIHistoryConvolutionKernelSpreadFactor.GetValueOnRenderThread();
 		Settings.MaxInputSPP = Config.RayCountPerPixel;
 
-		TStaticArray<FScreenSpaceFilteringHistory*, IScreenSpaceDenoiser::kMaxBatchSize> PrevHistories;
-		TStaticArray<FScreenSpaceFilteringHistory*, IScreenSpaceDenoiser::kMaxBatchSize> NewHistories;
+		TStaticArray<FScreenSpaceDenoiserHistory*, IScreenSpaceDenoiser::kMaxBatchSize> PrevHistories;
+		TStaticArray<FScreenSpaceDenoiserHistory*, IScreenSpaceDenoiser::kMaxBatchSize> NewHistories;
 		PrevHistories[0] = &PreviousViewInfos->SkyLightHistory;
 		NewHistories[0] = View.ViewState ? &View.ViewState->PrevFrameViewInfo.SkyLightHistory : nullptr;
 
@@ -2144,8 +2144,8 @@ public:
 		Settings.bUseTemporalAccumulation = CVarGITemporalAccumulation.GetValueOnRenderThread() != 0;
 		Settings.MaxInputSPP = Config.RayCountPerPixel;
 
-		TStaticArray<FScreenSpaceFilteringHistory*, IScreenSpaceDenoiser::kMaxBatchSize> PrevHistories;
-		TStaticArray<FScreenSpaceFilteringHistory*, IScreenSpaceDenoiser::kMaxBatchSize> NewHistories;
+		TStaticArray<FScreenSpaceDenoiserHistory*, IScreenSpaceDenoiser::kMaxBatchSize> PrevHistories;
+		TStaticArray<FScreenSpaceDenoiserHistory*, IScreenSpaceDenoiser::kMaxBatchSize> NewHistories;
 		PrevHistories[0] = &PreviousViewInfos->DiffuseIndirectHistory;
 		NewHistories[0] = View.ViewState ? &View.ViewState->PrevFrameViewInfo.DiffuseIndirectHistory : nullptr;
 
