- Cleaned up eye adaptation / histogram pass selection logic.
 - Fixed eye adaptation to respect override cvar.
 - Fixed case where AEM_Max could be returned from GetEyeAdaptationMode.

diff --git a/Engine/Source/Runtime/Renderer/Private/PostProcess/PostProcessEyeAdaptation.cpp b/Engine/Source/Runtime/Renderer/Private/PostProcess/PostProcessEyeAdaptation.cpp
--- a/Engine/Source/Runtime/Renderer/Private/PostProcess/PostProcessEyeAdaptation.cpp
+++ b/Engine/Source/Runtime/Renderer/Private/PostProcess/PostProcessEyeAdaptation.cpp
@@ -69,6 +69,8 @@ inline static void ComputeEyeAdaptationValues(const ERHIFeatureLevel::Type MinFe
 
 	// ----------
 
+	const EAutoExposureMethod AutoExposureMethod = GetAutoExposureMethod(View);
+
 	// Histogram related values.
 	float LowPercent = FMath::Clamp(Settings.AutoExposureLowPercent, 1.0f, 99.0f) * 0.01f;
 	float HighPercent = FMath::Clamp(Settings.AutoExposureHighPercent, 1.0f, 99.0f) * 0.01f;
@@ -115,7 +117,7 @@ inline static void ComputeEyeAdaptationValues(const ERHIFeatureLevel::Type MinFe
 	// When !EngineShowFlags.EyeAdaptation (from "r.EyeAdaptationQuality 0") or the feature level doesn't support eye adaptation, only Settings.AutoExposureBias controls exposure.
 	else if (EngineShowFlags.EyeAdaptation && View.GetFeatureLevel() >= MinFeatureLevel)
 	{
-		if (Settings.AutoExposureMethod == EAutoExposureMethod::AEM_Manual)
+		if (AutoExposureMethod == EAutoExposureMethod::AEM_Manual)
 		{
 			const float FixedEV100 = FMath::Log2(FMath::Square(Settings.DepthOfFieldFstop) * Settings.CameraShutterSpeed * 100 / FMath::Max(1.f, Settings.CameraISO));
 			MinAverageLuminance = MaxAverageLuminance = EV100ToLuminance(FixedEV100);
@@ -132,7 +134,7 @@ inline static void ComputeEyeAdaptationValues(const ERHIFeatureLevel::Type MinFe
 		}
 
 		// Note that AEM_Histogram implements the calibration constant through LowPercent and HiPercent.
-		if (Settings.AutoExposureMethod == EAutoExposureMethod::AEM_Basic)
+		if (AutoExposureMethod == EAutoExposureMethod::AEM_Basic)
 		{
 			const float CalibrationConstant = FMath::Clamp<float>(Settings.AutoExposureCalibrationConstant, 1.f, 100.f) * 0.01f;
 			AverageLuminanceScale = 1.f / CalibrationConstant;
diff --git a/Engine/Source/Runtime/Renderer/Private/PostProcess/PostProcessEyeAdaptation.h b/Engine/Source/Runtime/Renderer/Private/PostProcess/PostProcessEyeAdaptation.h
--- a/Engine/Source/Runtime/Renderer/Private/PostProcess/PostProcessEyeAdaptation.h
+++ b/Engine/Source/Runtime/Renderer/Private/PostProcess/PostProcessEyeAdaptation.h
@@ -103,43 +103,6 @@ private:
 // Console Variable that is used to over-ride the post process settings.
 extern TAutoConsoleVariable<int32> CVarEyeAdaptationMethodOverride;
 
-// Query the view for the auto exposure method, and allow for CVar override.
-static inline EAutoExposureMethod GetAutoExposureMethod(const FViewInfo& View)
-{
-	EAutoExposureMethod AutoExposureMethodId = View.FinalPostProcessSettings.AutoExposureMethod;
-	const int32 EyeOverride = CVarEyeAdaptationMethodOverride.GetValueOnRenderThread();
-
-	// Early out for common case
-	if (EyeOverride < 0) return AutoExposureMethodId;
-
-	// Additional branching for override.
-	switch (EyeOverride)
-	{
-	case 1:
-	{
-			  AutoExposureMethodId = EAutoExposureMethod::AEM_Histogram;
-			  break;
-	}
-	case 2:
-	{
-			  AutoExposureMethodId = EAutoExposureMethod::AEM_Basic;
-			  break;
-	}
-	case 3:
-	{
-			  AutoExposureMethodId = EAutoExposureMethod::AEM_Manual;
-			  break;
-	}
-	default:
-	{
-			   // Should only happen if the user supplies an override > 3
-			   AutoExposureMethodId = EAutoExposureMethod::AEM_MAX;
-			   break;
-	}
-	}
-	return AutoExposureMethodId;
-}
-
 // @return true if the current feature level supports this auto exposure method.
 static inline bool IsAutoExposureMethodSupported(const ERHIFeatureLevel::Type& FeatureLevel, const EAutoExposureMethod& AutoExposureMethodId)
 {
@@ -160,6 +123,52 @@ static inline bool IsAutoExposureMethodSupported(const ERHIFeatureLevel::Type& F
 	return Result;
 }
 
+// Query the view for the auto exposure method, and allow for CVar override.
+static inline EAutoExposureMethod GetAutoExposureMethod(const FViewInfo& View)
+{
+	EAutoExposureMethod AutoExposureMethodId = View.FinalPostProcessSettings.AutoExposureMethod;
+
+	const int32 EyeOverride = CVarEyeAdaptationMethodOverride.GetValueOnRenderThread();
+
+	if (EyeOverride >= 0)
+	{
+		// Additional branching for override.
+		switch (EyeOverride)
+		{
+		case 1:
+		{
+			// Only override if the platform supports it.
+			if (View.GetFeatureLevel() >= ERHIFeatureLevel::SM5)
+			{
+				AutoExposureMethodId = EAutoExposureMethod::AEM_Histogram;
+			}
+			break;
+		}
+		case 2:
+		{
+			AutoExposureMethodId = EAutoExposureMethod::AEM_Basic;
+			break;
+		}
+		case 3:
+		{
+			AutoExposureMethodId = EAutoExposureMethod::AEM_Manual;
+			break;
+		}
+		}
+	}
+
+	// If auto exposure is disabled, revert to manual mode which will clamp to a reasonable default.
+	if (!View.Family->EngineShowFlags.EyeAdaptation)
+	{
+		AutoExposureMethodId = AEM_Manual;
+	}
+
+	// We should have a valid exposure method at this stage.
+	check(IsAutoExposureMethodSupported(View.GetFeatureLevel(), AutoExposureMethodId));
+
+	return AutoExposureMethodId;
+}
+
 extern TAutoConsoleVariable<float> CVarEyeAdaptationFocus;
 static inline float GetBasicAutoExposureFocus()
 {
diff --git a/Engine/Source/Runtime/Renderer/Private/PostProcess/PostProcessVisualizeHDR.cpp b/Engine/Source/Runtime/Renderer/Private/PostProcess/PostProcessVisualizeHDR.cpp
--- a/Engine/Source/Runtime/Renderer/Private/PostProcess/PostProcessVisualizeHDR.cpp
+++ b/Engine/Source/Runtime/Renderer/Private/PostProcess/PostProcessVisualizeHDR.cpp
@@ -295,13 +295,20 @@ void FRCPassPostProcessVisualizeHDR::Process(FRenderingCompositePassContext& Con
 		Canvas.DrawShadowedString(MinX + XAdd - 5, MaxY + YStep, *Line, GetStatsFont(), FLinearColor(1, 0.3f, 0.3f));
 	}
 	Y += 3 * YStep;
-	if (AutoExposureMethod == EAutoExposureMethod::AEM_Basic)
+	switch (AutoExposureMethod)
 	{
+	case EAutoExposureMethod::AEM_Basic:
 		Line = FString::Printf(TEXT("Basic"));
-	}
-	else
-	{
+		break;
+	case EAutoExposureMethod::AEM_Histogram:
 		Line = FString::Printf(TEXT("Histogram"));
+		break;
+	case EAutoExposureMethod::AEM_Manual:
+		Line = FString::Printf(TEXT("Manual"));
+		break;
+	default:
+		Line = FString::Printf(TEXT("Unknown"));
+		break;
 	}
 	Canvas.DrawShadowedString(X, Y += YStep, TEXT("Auto Exposure Method:"), GetStatsFont(), FLinearColor(1, 1, 1));
 	Canvas.DrawShadowedString(X + ColumnWidth, Y, *Line, GetStatsFont(), FLinearColor(1, 1, 1));
diff --git a/Engine/Source/Runtime/Renderer/Private/PostProcess/PostProcessing.cpp b/Engine/Source/Runtime/Renderer/Private/PostProcess/PostProcessing.cpp
--- a/Engine/Source/Runtime/Renderer/Private/PostProcess/PostProcessing.cpp
+++ b/Engine/Source/Runtime/Renderer/Private/PostProcess/PostProcessing.cpp
@@ -982,6 +982,8 @@ void FPostProcessing::Process(FRHICommandListImmediate& RHICmdList, const FViewI
 
 	const ERHIFeatureLevel::Type FeatureLevel = View.GetFeatureLevel();
 
+	check(FeatureLevel >= ERHIFeatureLevel::SM4);
+
 	GRenderTargetPool.AddPhaseEvent(TEXT("PostProcessing"));
 
 	// All post processing is happening on the render thread side. All passes can access FinalPostProcessSettings and all
@@ -995,7 +997,6 @@ void FPostProcessing::Process(FRHICommandListImmediate& RHICmdList, const FViewI
 
 		FPostprocessContext Context(RHICmdList, CompositeContext.Graph, View);
 
-		// not always valid
 		FRenderingCompositeOutputRef HistogramOverScreen;
 		FRenderingCompositeOutputRef Histogram;
 		FRenderingCompositeOutputRef PreTonemapHDRColor;
@@ -1005,45 +1006,7 @@ void FPostProcessing::Process(FRHICommandListImmediate& RHICmdList, const FViewI
 		FRenderingCompositeOutputRef BloomOutputCombined;
 		FRenderingCompositeOutputRef CustomDepth;
 		FRenderingCompositeOutputRef SceneColorHalfRes;
-
-		class FAutoExposure
-		{
-		public:
-			FAutoExposure(const FViewInfo& InView) :
-				MethodId(GetAutoExposureMethod(InView))
-			{}
-			// distinguish between Basic and Histogram-based
-			EAutoExposureMethod          MethodId;
-			// not always valid
-			FRenderingCompositeOutputRef EyeAdaptation;
-		} AutoExposure(View);
-
-		FRCPassPostProcessUpscale::PaniniParams PaniniConfig(View);
-
-		const bool bDepthOfFieldEnabled = DiaphragmDOF::IsEnabled(View);
-
-		const bool bVisualizeHDR = View.Family->EngineShowFlags.VisualizeHDR && FeatureLevel >= ERHIFeatureLevel::SM5;
-
-		const bool bVisualizeBloom = View.Family->EngineShowFlags.VisualizeBloom && FeatureLevel >= ERHIFeatureLevel::SM4;
-
-		const bool bVisualizeMotionBlur = IsVisualizeMotionBlurEnabled(View);
-
-		// Motion blur gets replaced by the visualization pass.
-		const bool bMotionBlurEnabled = !bVisualizeMotionBlur && IsMotionBlurEnabled(View);
-
-		const bool bVisualizeDepthOfField = bDepthOfFieldEnabled && Context.View.Family->EngineShowFlags.VisualizeDOF;
-
-		const bool bTonemapEnabled = FeatureLevel >= ERHIFeatureLevel::SM4 && !bVisualizeBloom && !bVisualizeMotionBlur;
-
-		const bool bHDROutputEnabled = GRHISupportsHDROutput && IsHDREnabled();
-
-		const bool bHDRTonemapperOutput = bTonemapEnabled && (View.Family->SceneCaptureSource == SCS_FinalColorHDR || GetHighResScreenshotConfig().bCaptureHDR || IsBufferVisualizationDumpFramesInHDREnabled() || bHDROutputEnabled);
-
-		const EAntiAliasingMethod AntiAliasingMethod = View.AntiAliasingMethod;
-
-		const EDownsampleQuality DownsampleQuality = GetDownsampleQuality();
-
-		const EPixelFormat DownsampleOverrideFormat = PF_FloatRGB;
+		FRenderingCompositeOutputRef EyeAdaptation;
 
 		FRCPassPostProcessTonemap* Tonemapper = nullptr;
 
@@ -1067,10 +1030,53 @@ void FPostProcessing::Process(FRHICommandListImmediate& RHICmdList, const FViewI
 			PreFlattenVelocity = Context.Graph.RegisterPass(new(FMemStack::Get()) FRCPassPostProcessInput(VelocityRT));
 		}
 
+		const bool bVisualizeHDR = View.Family->EngineShowFlags.VisualizeHDR && FeatureLevel >= ERHIFeatureLevel::SM5;
+
+		const bool bVisualizeBloom = View.Family->EngineShowFlags.VisualizeBloom;
+
+		const bool bHDROutputEnabled = GRHISupportsHDROutput && IsHDREnabled();
+
 		// add the passes we want to add to the graph (commenting a line means the pass is not inserted into the graph) ---------
 
 		if (AllowFullPostProcessing(View, FeatureLevel))
 		{
+			const EAutoExposureMethod AutoExposureMethod = GetAutoExposureMethod(View);
+
+			const EAntiAliasingMethod AntiAliasingMethod = View.AntiAliasingMethod;
+
+			const EDownsampleQuality DownsampleQuality = GetDownsampleQuality();
+
+			const EPixelFormat DownsampleOverrideFormat = PF_FloatRGB;
+
+			const bool bDepthOfFieldEnabled = DiaphragmDOF::IsEnabled(View);
+
+			const bool bVisualizeDepthOfField = bDepthOfFieldEnabled && Context.View.Family->EngineShowFlags.VisualizeDOF;
+
+			const bool bVisualizeMotionBlur = IsVisualizeMotionBlurEnabled(View);
+
+			// Motion blur gets replaced by the visualization pass.
+			const bool bMotionBlurEnabled = !bVisualizeMotionBlur && IsMotionBlurEnabled(View);
+
+			// Skip tonemapping for visualizers which overwrite the HDR scene color.
+			const bool bTonemapEnabled = !bVisualizeBloom && !bVisualizeMotionBlur;
+
+			// We don't test for the EyeAdaptation engine show flag here. If disabled, the auto exposure pass is still executes but performs a clamp.
+			const bool bEyeAdaptationEnabled =
+				// Skip for transient views.
+				View.ViewState &&
+				// Skip for secondary views in a stereo setup.
+				IStereoRendering::IsAPrimaryView(View.StereoPass, GEngine->StereoRenderingDevice);
+
+			const bool bHistogramEnabled =
+				// Force the histogram on when we are visualizing HDR.
+				bVisualizeHDR ||
+				// Skip if not using histogram eye adaptation.
+				(bEyeAdaptationEnabled && AutoExposureMethod == EAutoExposureMethod::AEM_Histogram &&
+				// Skip if we don't have any exposure range to generate (eye adaptation will clamp).
+				View.FinalPostProcessSettings.AutoExposureMinBrightness < View.FinalPostProcessSettings.AutoExposureMaxBrightness);
+
+			const bool bBloomEnabled = View.FinalPostProcessSettings.BloomIntensity > 0.0f;
+
 			// GBuffers are released prior to executing the composition graph. We take a reference here
 			// and then release the reference inside of RDGPass. This allows RDGPass to control lifetime
 			// of the GBuffers internally.
@@ -1261,91 +1267,65 @@ void FPostProcessing::Process(FRHICommandListImmediate& RHICmdList, const FViewI
 
 			SceneColorHalfRes = FRenderingCompositeOutputRef(RDGPass, ePId_Output1);
 
+			if (bHistogramEnabled)
 			{
-				bool bHistogramNeeded = false;
+				FRenderingCompositePass* NodeHistogram = Context.Graph.RegisterPass(new(FMemStack::Get()) FRCPassPostProcessHistogram());
 
-				if (View.Family->EngineShowFlags.EyeAdaptation && (AutoExposure.MethodId == EAutoExposureMethod::AEM_Histogram)
-					&& View.FinalPostProcessSettings.AutoExposureMinBrightness < View.FinalPostProcessSettings.AutoExposureMaxBrightness
-					&& !View.bIsSceneCapture // Eye adaption is not available for scene captures.
-					&& !bVisualizeBloom)
-				{
-					bHistogramNeeded = true;
-				}
+				NodeHistogram->SetInput(ePId_Input0, SceneColorHalfRes);
 
-				if(!bTonemapEnabled)
-				{
-					bHistogramNeeded = false;
-				}
+				HistogramOverScreen = FRenderingCompositeOutputRef(NodeHistogram);
 
-				if(View.Family->EngineShowFlags.VisualizeHDR)
-				{
-					bHistogramNeeded = true;
-				}
-
-				if (!GIsHighResScreenshot && bHistogramNeeded && FeatureLevel >= ERHIFeatureLevel::SM5 && IStereoRendering::IsAPrimaryView(View.StereoPass, GEngine->StereoRenderingDevice))
-				{
-					FRenderingCompositePass* NodeHistogram = Context.Graph.RegisterPass(new(FMemStack::Get()) FRCPassPostProcessHistogram());
-
-					NodeHistogram->SetInput(ePId_Input0, SceneColorHalfRes);
+				FRenderingCompositePass* NodeHistogramReduce = Context.Graph.RegisterPass(new(FMemStack::Get()) FRCPassPostProcessHistogramReduce());
 
-					HistogramOverScreen = FRenderingCompositeOutputRef(NodeHistogram);
+				NodeHistogramReduce->SetInput(ePId_Input0, NodeHistogram);
 
-					FRenderingCompositePass* NodeHistogramReduce = Context.Graph.RegisterPass(new(FMemStack::Get()) FRCPassPostProcessHistogramReduce());
-
-					NodeHistogramReduce->SetInput(ePId_Input0, NodeHistogram);
-
-					Histogram = FRenderingCompositeOutputRef(NodeHistogramReduce);
-				}
+				Histogram = FRenderingCompositeOutputRef(NodeHistogramReduce);
 			}
 
+			const bool bBasicEyeAdaptationEnabled = bEyeAdaptationEnabled && (AutoExposureMethod == EAutoExposureMethod::AEM_Basic);
+
 			// Compute DownSamples passes used by bloom, tint and eye-adaptation if possible.
 			FBloomDownSampleArray::Ptr BloomAndEyeDownSamplesPtr;
-			if (View.FinalPostProcessSettings.BloomIntensity > 0.f) // do bloom
+
+			if (bBloomEnabled)
 			{
 				// No Threshold:  We can share with Eye-Adaptation.
 				if (Context.View.FinalPostProcessSettings.BloomThreshold <= -1 && Context.View.Family->Views.Num() == 1)
 				{
-					if (!GIsHighResScreenshot && View.State &&
-						IStereoRendering::IsAPrimaryView(View.StereoPass, GEngine->StereoRenderingDevice) &&
-						AutoExposure.MethodId == EAutoExposureMethod::AEM_Basic)
+					if (bBasicEyeAdaptationEnabled)
 					{
 						BloomAndEyeDownSamplesPtr = CreateDownSampleArray(Context, SceneColorHalfRes, true /*bGenerateLog2Alpha*/);
 					}
 				}
 			}
 
-			// some views don't have a state (thumbnail rendering)
-			if(!GIsHighResScreenshot && View.State && IStereoRendering::IsAPrimaryView(View.StereoPass, GEngine->StereoRenderingDevice))
+			if (bEyeAdaptationEnabled)
 			{
-				const bool bUseBasicEyeAdaptation = (AutoExposure.MethodId == EAutoExposureMethod::AEM_Basic);
-
-				if (bUseBasicEyeAdaptation) // log average ps reduction ( non histogram ) 
+				if (bBasicEyeAdaptationEnabled)
 				{
-					
 					if (!BloomAndEyeDownSamplesPtr.IsValid()) 
 					{
 						// need downsamples for eye-adaptation.
 						FBloomDownSampleArray::Ptr EyeDownSamplesPtr = CreateDownSampleArray(Context, SceneColorHalfRes, true /*bGenerateLog2Alpha*/);
-						AutoExposure.EyeAdaptation = AddPostProcessBasicEyeAdaptation(View, *EyeDownSamplesPtr);
+						EyeAdaptation = AddPostProcessBasicEyeAdaptation(View, *EyeDownSamplesPtr);
 					}
 					else
 					{
-						// Use the alpha channel in the last downsample (smallest) to compute eye adaptations values.			
-						AutoExposure.EyeAdaptation = AddPostProcessBasicEyeAdaptation(View, *BloomAndEyeDownSamplesPtr);
+						// Use the alpha channel in the last downsample (smallest) to compute eye adaptations values.
+						EyeAdaptation = AddPostProcessBasicEyeAdaptation(View, *BloomAndEyeDownSamplesPtr);
 					}
 				}
-				else  // Use histogram version version
+				else
 				{
-					// We always add eye adaptation, if the engine show flag is disabled we set the ExposureScale in the texture to a fixed value
-					AutoExposure.EyeAdaptation = AddPostProcessHistogramEyeAdaptation(Context, Histogram);
+					EyeAdaptation = AddPostProcessHistogramEyeAdaptation(Context, Histogram);
 				}
 			}
 
-			if(View.FinalPostProcessSettings.BloomIntensity > 0.0f)
+			if (bBloomEnabled)
 			{
 				if (!BloomAndEyeDownSamplesPtr.IsValid())
 				{
-					FRenderingCompositeOutputRef HalfResBloomThreshold = RenderHalfResBloomThreshold(Context, SceneColorHalfRes, AutoExposure.EyeAdaptation);
+					FRenderingCompositeOutputRef HalfResBloomThreshold = RenderHalfResBloomThreshold(Context, SceneColorHalfRes, EyeAdaptation);
 					BloomAndEyeDownSamplesPtr = CreateDownSampleArray(Context, HalfResBloomThreshold, false /*bGenerateLog2Alpha*/);
 				}
 				BloomOutputCombined = AddBloom(*BloomAndEyeDownSamplesPtr, bVisualizeBloom);
@@ -1355,6 +1335,8 @@ void FPostProcessing::Process(FRHICommandListImmediate& RHICmdList, const FViewI
 
 			if(bTonemapEnabled)
 			{
+				const bool bTonemapOutputInHDR = (View.Family->SceneCaptureSource == SCS_FinalColorHDR || GetHighResScreenshotConfig().bCaptureHDR || IsBufferVisualizationDumpFramesInHDREnabled() || bHDROutputEnabled);
+
 				{
 					FRenderingCompositeOutputRef FinalOutputPrev = Context.FinalOutput;
 
@@ -1363,14 +1345,17 @@ void FPostProcessing::Process(FRHICommandListImmediate& RHICmdList, const FViewI
 					// No-op from post process material pass; run built-in tonemapper instead.
 					if (Context.FinalOutput == FinalOutputPrev)
 					{
-						Tonemapper = AddTonemapper(Context, BloomOutputCombined, AutoExposure.EyeAdaptation, AutoExposure.MethodId, false, bHDRTonemapperOutput);
+						Tonemapper = AddTonemapper(Context, BloomOutputCombined, EyeAdaptation, AutoExposureMethod, false, bTonemapOutputInHDR);
 					}
 				}
 
 				PostTonemapHDRColor = Context.FinalOutput;
 
-				// Add a pass-through as tonemapper will be forced LDR if final pass in chain 
-				if (bHDRTonemapperOutput && !bHDROutputEnabled)
+				// The composition graph will substitute the hardware backbuffer in place of the last render target, which
+				// we don't want to do when outputting HDR from the tonemapper. Instead, to be safe, we perform a copy
+				// which will truncate HDR values to LDR. If this isn't the last pass, we end up eating the copy and the
+				// result will still be in HDR.
+				if (bTonemapOutputInHDR && !bHDROutputEnabled)
 				{
 					FRenderingCompositePass* PassthroughNode = Context.Graph.RegisterPass(new(FMemStack::Get()) FRCPassPostProcessPassThrough(nullptr));
 					PassthroughNode->SetInput(ePId_Input0, FRenderingCompositeOutputRef(Context.FinalOutput));
@@ -1547,7 +1532,7 @@ void FPostProcessing::Process(FRHICommandListImmediate& RHICmdList, const FViewI
 			Node->SetInput(ePId_Input1, Histogram);
 			Node->SetInput(ePId_Input2, PreTonemapHDRColor);
 			Node->SetInput(ePId_Input3, HistogramOverScreen);
-			Node->AddDependency(AutoExposure.EyeAdaptation);
+			Node->AddDependency(EyeAdaptation);
 
 			Context.FinalOutput = FRenderingCompositeOutputRef(Node);
 		}
@@ -1568,7 +1553,9 @@ void FPostProcessing::Process(FRHICommandListImmediate& RHICmdList, const FViewI
 
 		AddHighResScreenshotMask(Context);
 
-		FIntPoint PrimaryUpscaleViewSize = Context.View.GetSecondaryViewRectSize();
+		const FIntPoint PrimaryUpscaleViewSize = Context.View.GetSecondaryViewRectSize();
+
+		const FRCPassPostProcessUpscale::PaniniParams PaniniConfig(View);
 
 		// If the final output is still not unscaled, therefore add Upscale pass.
 		if ((!bUnscaledFinalOutput && View.PrimaryScreenPercentageMethod == EPrimaryScreenPercentageMethod::SpatialUpscale && View.ViewRect.Size() != PrimaryUpscaleViewSize) || PaniniConfig.IsEnabled())
