Cleanup of bloom / eye adaptation downsample sharing logic.

diff --git a/Engine/Source/Runtime/Renderer/Private/PostProcess/PostProcessing.cpp b/Engine/Source/Runtime/Renderer/Private/PostProcess/PostProcessing.cpp
--- a/Engine/Source/Runtime/Renderer/Private/PostProcess/PostProcessing.cpp
+++ b/Engine/Source/Runtime/Renderer/Private/PostProcess/PostProcessing.cpp
@@ -280,34 +280,6 @@ FBloomDownSampleArray::Ptr CreateDownSampleArray(FPostprocessContext& Context, F
 	return FBloomDownSampleArray::Ptr(new FBloomDownSampleArray(Context, SourceToDownSample, bAddLog2));
 }
 
-
-static FRenderingCompositeOutputRef RenderHalfResBloomThreshold(FPostprocessContext& Context, FRenderingCompositeOutputRef SceneColorHalfRes, FRenderingCompositeOutputRef EyeAdaptation)
-{
-	// with multiple view ports the Setup pass also isolates the view from the others which allows for simpler simpler/faster blur passes.
-	if(Context.View.FinalPostProcessSettings.BloomThreshold <= -1 && Context.View.Family->Views.Num() == 1)
-	{
-		// no need for threshold, we don't need this pass
-		return SceneColorHalfRes;
-	}
-	else
-	{
-		// todo: optimize later, the missing node causes some wrong behavior
-		//	if(Context.View.FinalPostProcessSettings.BloomIntensity <= 0.0f)
-		//	{
-		//		// this pass is not required
-		//		return FRenderingCompositeOutputRef();
-		//	}
-		// bloom threshold
-		const bool bIsComputePass = ShouldDoComputePostProcessing(Context.View);
-		FRenderingCompositePass* PostProcessBloomSetup = Context.Graph.RegisterPass(new(FMemStack::Get()) FRCPassPostProcessBloomSetup(bIsComputePass));
-		PostProcessBloomSetup->SetInput(ePId_Input0, SceneColorHalfRes);
-		PostProcessBloomSetup->SetInput(ePId_Input1, EyeAdaptation);
-
-		return FRenderingCompositeOutputRef(PostProcessBloomSetup);
-	}
-}
-
-
 // 2 pass Gaussian blur using uni-linear filtering
 // @param CrossCenterWeight see r.Bloom.Cross (positive for X and Y, otherwise for X only)
 static FRenderingCompositeOutputRef RenderGaussianBlur(
@@ -1048,9 +1020,13 @@ void FPostProcessing::Process(FRHICommandListImmediate& RHICmdList, const FViewI
 
 			const EPixelFormat DownsampleOverrideFormat = PF_FloatRGB;
 
+			const bool bPreferCompute = ShouldDoComputePostProcessing(View);
+
+			const bool bHasViewState = View.ViewState != nullptr;
+
 			const bool bDepthOfFieldEnabled = DiaphragmDOF::IsEnabled(View);
 
-			const bool bVisualizeDepthOfField = bDepthOfFieldEnabled && Context.View.Family->EngineShowFlags.VisualizeDOF;
+			const bool bVisualizeDepthOfField = bDepthOfFieldEnabled && View.Family->EngineShowFlags.VisualizeDOF;
 
 			const bool bVisualizeMotionBlur = IsVisualizeMotionBlurEnabled(View);
 
@@ -1063,7 +1039,7 @@ void FPostProcessing::Process(FRHICommandListImmediate& RHICmdList, const FViewI
 			// We don't test for the EyeAdaptation engine show flag here. If disabled, the auto exposure pass is still executes but performs a clamp.
 			const bool bEyeAdaptationEnabled =
 				// Skip for transient views.
-				View.ViewState &&
+				bHasViewState &&
 				// Skip for secondary views in a stereo setup.
 				IStereoRendering::IsAPrimaryView(View.StereoPass, GEngine->StereoRenderingDevice);
 
@@ -1284,35 +1260,25 @@ void FPostProcessing::Process(FRHICommandListImmediate& RHICmdList, const FViewI
 
 			const bool bBasicEyeAdaptationEnabled = bEyeAdaptationEnabled && (AutoExposureMethod == EAutoExposureMethod::AEM_Basic);
 
-			// Compute DownSamples passes used by bloom, tint and eye-adaptation if possible.
-			FBloomDownSampleArray::Ptr BloomAndEyeDownSamplesPtr;
+			const bool bBloomThresholdEnabled = Context.View.FinalPostProcessSettings.BloomThreshold > 0.0f;
 
-			if (bBloomEnabled)
-			{
-				// No Threshold:  We can share with Eye-Adaptation.
-				if (Context.View.FinalPostProcessSettings.BloomThreshold <= -1 && Context.View.Family->Views.Num() == 1)
-				{
-					if (bBasicEyeAdaptationEnabled)
-					{
-						BloomAndEyeDownSamplesPtr = CreateDownSampleArray(Context, SceneColorHalfRes, true /*bGenerateLog2Alpha*/);
-					}
-				}
-			}
+			FBloomDownSampleArray::Ptr SceneColorDownsampleChain;
 
 			if (bEyeAdaptationEnabled)
 			{
 				if (bBasicEyeAdaptationEnabled)
 				{
-					if (!BloomAndEyeDownSamplesPtr.IsValid()) 
-					{
-						// need downsamples for eye-adaptation.
-						FBloomDownSampleArray::Ptr EyeDownSamplesPtr = CreateDownSampleArray(Context, SceneColorHalfRes, true /*bGenerateLog2Alpha*/);
-						EyeAdaptation = AddPostProcessBasicEyeAdaptation(View, *EyeDownSamplesPtr);
-					}
-					else
+					const bool bGenerageLog2Alpha = true;
+
+					FBloomDownSampleArray::Ptr EyeAdaptationDownsampleChain = CreateDownSampleArray(Context, SceneColorHalfRes, bGenerageLog2Alpha);
+
+					// Use the alpha channel in the last downsample (smallest) to compute eye adaptations values.
+					EyeAdaptation = AddPostProcessBasicEyeAdaptation(View, *EyeAdaptationDownsampleChain);
+
+					// We can reuse the same downsample chain on bloom, but only if we aren't doing thresholding.
+					if (!bBloomThresholdEnabled)
 					{
-						// Use the alpha channel in the last downsample (smallest) to compute eye adaptations values.
-						EyeAdaptation = AddPostProcessBasicEyeAdaptation(View, *BloomAndEyeDownSamplesPtr);
+						SceneColorDownsampleChain = EyeAdaptationDownsampleChain;
 					}
 				}
 				else
@@ -1323,12 +1289,24 @@ void FPostProcessing::Process(FRHICommandListImmediate& RHICmdList, const FViewI
 
 			if (bBloomEnabled)
 			{
-				if (!BloomAndEyeDownSamplesPtr.IsValid())
+				if (!SceneColorDownsampleChain.IsValid())
 				{
-					FRenderingCompositeOutputRef HalfResBloomThreshold = RenderHalfResBloomThreshold(Context, SceneColorHalfRes, EyeAdaptation);
-					BloomAndEyeDownSamplesPtr = CreateDownSampleArray(Context, HalfResBloomThreshold, false /*bGenerateLog2Alpha*/);
+					FRenderingCompositeOutputRef DownsampleInput = SceneColorHalfRes;
+
+					if (bBloomThresholdEnabled)
+					{
+						FRenderingCompositePass* PostProcessBloomSetup = Context.Graph.RegisterPass(new(FMemStack::Get()) FRCPassPostProcessBloomSetup(bPreferCompute));
+						PostProcessBloomSetup->SetInput(ePId_Input0, SceneColorHalfRes);
+						PostProcessBloomSetup->SetInput(ePId_Input1, EyeAdaptation);
+
+						DownsampleInput = FRenderingCompositeOutputRef(PostProcessBloomSetup);
+					}
+
+					const bool bGenerageLog2Alpha = false;
+
+					SceneColorDownsampleChain = CreateDownSampleArray(Context, DownsampleInput, bGenerageLog2Alpha);
 				}
-				BloomOutputCombined = AddBloom(*BloomAndEyeDownSamplesPtr, bVisualizeBloom);
+				BloomOutputCombined = AddBloom(*SceneColorDownsampleChain, bVisualizeBloom);
 			}
 
 			PreTonemapHDRColor = Context.FinalOutput;
