Adds support for denoiser scissoring and viewport offsets.

diff --git a/Engine/Shaders/Private/ScreenSpaceDenoise/SSDCommon.ush b/Engine/Shaders/Private/ScreenSpaceDenoise/SSDCommon.ush
--- a/Engine/Shaders/Private/ScreenSpaceDenoise/SSDCommon.ush
+++ b/Engine/Shaders/Private/ScreenSpaceDenoise/SSDCommon.ush
@@ -60,6 +60,10 @@
 
 //------------------------------------------------------- SHARED SHADER PARAMETERS
 
+uint2 ViewportMin;
+uint2 ViewportMax;
+float4 BufferBilinearUVMinMax;
+
 float4x4  PrevScreenToTranslatedWorld;
 Texture2D PrevDepthBuffer;
 Texture2D PrevGBufferA;
diff --git a/Engine/Shaders/Private/ScreenSpaceDenoise/SSDComposeHarmonics.usf b/Engine/Shaders/Private/ScreenSpaceDenoise/SSDComposeHarmonics.usf
--- a/Engine/Shaders/Private/ScreenSpaceDenoise/SSDComposeHarmonics.usf
+++ b/Engine/Shaders/Private/ScreenSpaceDenoise/SSDComposeHarmonics.usf
@@ -136,10 +136,10 @@ void MainCS(
 	uint GroupThreadIndex : SV_GroupIndex)
 {
 	// Find out scene buffer UV.
-	float2 SceneBufferUV = (DispatchThreadId + 0.5 + View.ViewRectMin.xy) * View.BufferSizeAndInvSize.zw;
+	float2 SceneBufferUV = (ViewportMin + DispatchThreadId + 0.5) * View.BufferSizeAndInvSize.zw;
 	if (true)
 	{
-		SceneBufferUV = clamp(SceneBufferUV, View.BufferBilinearUVMinMax.xy, View.BufferBilinearUVMinMax.zw);
+		SceneBufferUV = clamp(SceneBufferUV, BufferBilinearUVMinMax.xy, BufferBilinearUVMinMax.zw);
 	}
 	
 	// Read reference meta data.
@@ -279,12 +279,23 @@ void MainCS(
 			OutputSamples.Array[MultiplexId].SceneColor = max(OutputSamples.Array[MultiplexId].SceneColor, 0);
 		#endif
 	}
-
-	OutputMultiplexedSignal(
-		SignalOutput_UAVs_0,
-		SignalOutput_UAVs_1,
-		SignalOutput_UAVs_2,
-		SignalOutput_UAVs_3,
-		CONFIG_SIGNAL_OUTPUT_LAYOUT, CONFIG_SIGNAL_BATCH_SIZE,
-		DispatchThreadId, OutputSamples.Array);
+	
+	#if CONFIG_VGPR_OPTIMIZATION
+		// No need to keep DispatchThreadId, while SceneBufferUV is arround at highest VGPR peak because center of the unique pixel to sample.
+		uint2 OutputPixelPostion = uint2(SceneBufferUV * View.BufferSizeAndInvSize.xy);
+	#else
+		uint2 OutputPixelPostion = ViewportMin + DispatchThreadId;
+	#endif 
+		
+	BRANCH
+	if (all(OutputPixelPostion < ViewportMax))
+	{
+		OutputMultiplexedSignal(
+			SignalOutput_UAVs_0,
+			SignalOutput_UAVs_1,
+			SignalOutput_UAVs_2,
+			SignalOutput_UAVs_3,
+			CONFIG_SIGNAL_OUTPUT_LAYOUT, CONFIG_SIGNAL_BATCH_SIZE,
+			OutputPixelPostion, OutputSamples.Array);
+	}
 } // MainCS
diff --git a/Engine/Shaders/Private/ScreenSpaceDenoise/SSDInjest.usf b/Engine/Shaders/Private/ScreenSpaceDenoise/SSDInjest.usf
--- a/Engine/Shaders/Private/ScreenSpaceDenoise/SSDInjest.usf
+++ b/Engine/Shaders/Private/ScreenSpaceDenoise/SSDInjest.usf
@@ -99,10 +99,10 @@ void MainCS(
 	uint GroupThreadIndex : SV_GroupIndex)
 {
 	// Find out scene buffer UV.
-	float2 SceneBufferUV = (DispatchThreadId + 0.5 + View.ViewRectMin.xy) * View.BufferSizeAndInvSize.zw;
+	float2 SceneBufferUV = (ViewportMin + DispatchThreadId + 0.5) * View.BufferSizeAndInvSize.zw;
 	if (true)
 	{
-		SceneBufferUV = clamp(SceneBufferUV, View.BufferBilinearUVMinMax.xy, View.BufferBilinearUVMinMax.zw);
+		SceneBufferUV = clamp(SceneBufferUV, BufferBilinearUVMinMax.xy, BufferBilinearUVMinMax.zw);
 	}
 	
 	// Read reference meta data.
@@ -139,12 +139,23 @@ void MainCS(
 	{
 		OutputSamples.Array[MultiplexId] = MultiplexedSamples.Array[MultiplexId];
 	}
-
-	OutputMultiplexedSignal(
-		SignalOutput_UAVs_0,
-		SignalOutput_UAVs_1,
-		SignalOutput_UAVs_2,
-		SignalOutput_UAVs_3,
-		CONFIG_SIGNAL_OUTPUT_LAYOUT, MultiplexCount,
-		DispatchThreadId, OutputSamples.Array);
+	
+	#if CONFIG_VGPR_OPTIMIZATION
+		// No need to keep DispatchThreadId, while SceneBufferUV is arround at highest VGPR peak because center of the unique pixel to sample.
+		uint2 OutputPixelPostion = uint2(SceneBufferUV * View.BufferSizeAndInvSize.xy);
+	#else
+		uint2 OutputPixelPostion = ViewportMin + DispatchThreadId;
+	#endif 
+		
+	BRANCH
+	if (all(OutputPixelPostion < ViewportMax))
+	{
+		OutputMultiplexedSignal(
+			SignalOutput_UAVs_0,
+			SignalOutput_UAVs_1,
+			SignalOutput_UAVs_2,
+			SignalOutput_UAVs_3,
+			CONFIG_SIGNAL_OUTPUT_LAYOUT, MultiplexCount,
+			OutputPixelPostion, OutputSamples.Array);
+	}
 } // MainCS
diff --git a/Engine/Shaders/Private/ScreenSpaceDenoise/SSDReduce.usf b/Engine/Shaders/Private/ScreenSpaceDenoise/SSDReduce.usf
--- a/Engine/Shaders/Private/ScreenSpaceDenoise/SSDReduce.usf
+++ b/Engine/Shaders/Private/ScreenSpaceDenoise/SSDReduce.usf
@@ -147,10 +147,10 @@ void MainCS(
 	uint GroupThreadIndex : SV_GroupIndex)
 {
 	// Find out scene buffer UV.
-	float2 SceneBufferUV = (DispatchThreadId + 0.5 + View.ViewRectMin.xy) * View.BufferSizeAndInvSize.zw;
+	float2 SceneBufferUV = (ViewportMin + DispatchThreadId + 0.5) * View.BufferSizeAndInvSize.zw;
 	if (true)
 	{
-		SceneBufferUV = clamp(SceneBufferUV, View.BufferBilinearUVMinMax.xy, View.BufferBilinearUVMinMax.zw);
+		SceneBufferUV = clamp(SceneBufferUV, BufferBilinearUVMinMax.xy, BufferBilinearUVMinMax.zw);
 	}
 	
 	// Read reference meta data.
@@ -182,7 +182,7 @@ void MainCS(
 		}
 
 		// TODO: TODO: reduce VGPR pressure here
-		uint2 OutputPixelPostion = DispatchThreadId;
+		uint2 OutputPixelPostion = ViewportMax + DispatchThreadId;
 
 		OutputMultiplexedSignal(
 			SignalOutputMips_0_UAVs_0,
@@ -204,10 +204,10 @@ void MainCS(
 			OutputSamples.Array[MultiplexId] = RefSamples.Array[MultiplexId];
 		}
 
-		uint2 OutputPixelPostion = DispatchThreadId >> MipLevel;
+		uint2 OutputPixelPostion = (ViewportMax + DispatchThreadId) >> MipLevel;
 
 		BRANCH
-		if (all((OutputPixelPostion << MipLevel) == DispatchThreadId))
+		if (all((OutputPixelPostion << MipLevel) == (ViewportMax + DispatchThreadId)))
 		{
 			if (MipLevel == 1)
 			{
diff --git a/Engine/Shaders/Private/ScreenSpaceDenoise/SSDSpatialAccumulation.usf b/Engine/Shaders/Private/ScreenSpaceDenoise/SSDSpatialAccumulation.usf
--- a/Engine/Shaders/Private/ScreenSpaceDenoise/SSDSpatialAccumulation.usf
+++ b/Engine/Shaders/Private/ScreenSpaceDenoise/SSDSpatialAccumulation.usf
@@ -62,6 +62,10 @@
 		#define CONFIG_MAX_WITH_REF_DISTANCE 1
 		#define CONFIG_OUTPUT_MODE          OUTPUT_MODE_DRB
 
+		#if DIM_SIGNAL_BATCH_SIZE > 1
+			#define CONFIG_CLAMP_UV_PER_SIGNAL 1
+		#endif
+
 	#elif DIM_STAGE == STAGE_PRE_CONVOLUTION
 		// Input and output layout.
 		#define CONFIG_SIGNAL_INPUT_LAYOUT  SIGNAL_BUFFER_LAYOUT_PENUMBRA_RECONSTRUCTION
@@ -352,6 +356,11 @@
 
 //------------------------------------------------------- CONFIG DISABLED DEFAULTS
 
+/** Whether should clamp the UV individually per texture. */
+#ifndef CONFIG_CLAMP_UV_PER_SIGNAL
+	#define CONFIG_CLAMP_UV_PER_SIGNAL 0
+#endif
+
 /** Changes the logic controling the number of sample to do. */
 #ifndef CONFIG_SAMPLE_COUNT_POLICY
 	#define CONFIG_SAMPLE_COUNT_POLICY SAMPLE_COUNT_POLICY_DISABLED
@@ -471,6 +480,8 @@ uint UpscaleFactor;
 
 float HarmonicPeriode;
 
+float4 InputBufferUVMinMax[CONFIG_SIGNAL_BATCH_SIZE];
+
 
 #if !defined(CONFIG_INPUT_TEXTURE_COUNT)
 	#error Missing CONFIG_INPUT_TEXTURE_COUNT
@@ -548,10 +559,10 @@ void MainCS(
 	uint GroupThreadIndex : SV_GroupIndex)
 {
 	// Find out scene buffer UV.
-	float2 SceneBufferUV = (DispatchThreadId + 0.5 + View.ViewRectMin.xy) * View.BufferSizeAndInvSize.zw;
+	float2 SceneBufferUV = ((DispatchThreadId + ViewportMin) + 0.5) * View.BufferSizeAndInvSize.zw;
 	if (true)
 	{
-		SceneBufferUV = clamp(SceneBufferUV, View.BufferBilinearUVMinMax.xy, View.BufferBilinearUVMinMax.zw);
+		SceneBufferUV = clamp(SceneBufferUV, BufferBilinearUVMinMax.xy, BufferBilinearUVMinMax.zw);
 	}
 	
 	// Read reference meta data.
@@ -651,6 +662,7 @@ void MainCS(
 		KernelConfig.bNormalizeSample = CONFIG_NORMALIZE_INPUT != 0;
 		KernelConfig.bSampleKernelCenter = CONFIG_UPSCALE;
 		KernelConfig.bForceKernelCenterAccumulation = true;
+		KernelConfig.bClampUVPerMultiplexedSignal = CONFIG_CLAMP_UV_PER_SIGNAL != 0;
 
 		// Reconstruct the spherical harmonic when reconstructing from 1spp.
 		KernelConfig.bComputeSampleColorSH = DIM_STAGE == STAGE_RECONSTRUCTION && DIM_MULTI_SPP == 0;
@@ -666,13 +678,24 @@ void MainCS(
 
 		SetBilateralPreset(CONFIG_BILATERAL_PRESET, /* inout */ KernelConfig);
 
-
 		// SGPRs
 		KernelConfig.BufferSizeAndInvSize = View.BufferSizeAndInvSize;
-		KernelConfig.BufferBilinearUVMinMax = View.BufferBilinearUVMinMax;
+		KernelConfig.BufferBilinearUVMinMax = BufferBilinearUVMinMax;
 		KernelConfig.KernelSpreadFactor = KernelSpreadFactor;
 		KernelConfig.HarmonicPeriode = HarmonicPeriode;
 
+		#if CONFIG_CLAMP_UV_PER_SIGNAL
+		{
+			UNROLL_N(CONFIG_SIGNAL_BATCH_SIZE)
+			for (uint BatchedSignalId = 0; BatchedSignalId < CONFIG_SIGNAL_BATCH_SIZE; BatchedSignalId++)
+			{
+				uint MultiplexId = BatchedSignalId / CONFIG_MULTIPLEXED_SIGNALS_PER_SIGNAL_DOMAIN;
+				KernelConfig.PerSignalUVMinMax[MultiplexId] = InputBufferUVMinMax[MultiplexId];
+			}
+		}
+		#endif
+
+
 		// VGPRs
 		KernelConfig.BufferUV = SceneBufferUV;
 		KernelConfig.RefSceneMetadata = RefSceneMetadata;
@@ -771,7 +794,7 @@ void MainCS(
 			
 				uint RaySubPixelId = View.StateFrameIndex & (UpscaleFactor * UpscaleFactor - 1);
 
-				KernelConfig.BufferUV = ((ClosestRayPixelCoord + (0.5 * KernelSpreadFactor + 0.5)) + View.ViewRectMin.xy) * View.BufferSizeAndInvSize.zw;
+				KernelConfig.BufferUV = ((ViewportMin + ClosestRayPixelCoord + (0.5 * KernelSpreadFactor + 0.5))) * View.BufferSizeAndInvSize.zw;
 			
 				// Sample the center of the kernel by comparing it against the RefSceneMetadata, since it may no match.
 				KernelConfig.bSampleKernelCenter = true;
@@ -785,7 +808,7 @@ void MainCS(
 			#else
 			{
 				// Put the kernel center at the center of the quad. Half pixel shift is done in the sample offsets.
-				KernelConfig.BufferUV = (DispatchThreadId | 1) * View.BufferSizeAndInvSize.zw;
+				KernelConfig.BufferUV = (ViewportMin + (DispatchThreadId | 1)) * View.BufferSizeAndInvSize.zw;
 
 				// Id of the pixel in the quad. This is to match hard coded first samples of the sample set.
 				KernelConfig.SampleTrackId = ((DispatchThreadId.x & 1) | ((DispatchThreadId.y & 1) << 1));
@@ -1024,47 +1047,51 @@ void MainCS(
 		#endif
 	}
 	#else
-		OutputPixelPostion = DispatchThreadId;
+		OutputPixelPostion = ViewportMin + DispatchThreadId;
 	#endif 
 
-	// Output the multiplexed signal.
-	#if DIM_STAGE == STAGE_FINAL_OUTPUT && CONFIG_SIGNAL_PROCESSING == SIGNAL_PROCESSING_MONOCHROMATIC_PENUMBRA
+	BRANCH
+	if (all(OutputPixelPostion < ViewportMax))
 	{
-		UNROLL
-		for (uint MultiplexId = 0; MultiplexId < MultiplexCount; MultiplexId++)
+		// Output the multiplexed signal.
+		#if DIM_STAGE == STAGE_FINAL_OUTPUT && CONFIG_SIGNAL_PROCESSING == SIGNAL_PROCESSING_MONOCHROMATIC_PENUMBRA
 		{
-			float Shadow = GetSamplePenumbraSafe(OutputSamples.Array[MultiplexId]);
-
-			const float ShadowFadeFraction = 1;
-			float SSSTransmission = Shadow;
-
-			// 0 is shadowed, 1 is unshadowed
-			// RETURN_COLOR not needed unless writing to SceneColor;
-			float FadedShadow = lerp(1.0f, Square(Shadow), ShadowFadeFraction);
-			float FadedSSSShadow = lerp(1.0f, Square(SSSTransmission), ShadowFadeFraction);
-
-			// the channel assignment is documented in ShadowRendering.cpp (look for Light Attenuation channel assignment)
-			float4 OutColor = EncodeLightAttenuation(half4(FadedShadow, FadedSSSShadow, FadedShadow, FadedSSSShadow));
-
-			if (MultiplexId == 0)
-				SignalOutput_UAVs_0[OutputPixelPostion] = OutColor;
-			if (MultiplexId == 1)
-				SignalOutput_UAVs_1[OutputPixelPostion] = OutColor;
-			if (MultiplexId == 2)
-				SignalOutput_UAVs_2[OutputPixelPostion] = OutColor;
-			if (MultiplexId == 3)
-				SignalOutput_UAVs_3[OutputPixelPostion] = OutColor;
+			UNROLL
+			for (uint MultiplexId = 0; MultiplexId < MultiplexCount; MultiplexId++)
+			{
+				float Shadow = GetSamplePenumbraSafe(OutputSamples.Array[MultiplexId]);
+
+				const float ShadowFadeFraction = 1;
+				float SSSTransmission = Shadow;
+
+				// 0 is shadowed, 1 is unshadowed
+				// RETURN_COLOR not needed unless writing to SceneColor;
+				float FadedShadow = lerp(1.0f, Square(Shadow), ShadowFadeFraction);
+				float FadedSSSShadow = lerp(1.0f, Square(SSSTransmission), ShadowFadeFraction);
+
+				// the channel assignment is documented in ShadowRendering.cpp (look for Light Attenuation channel assignment)
+				float4 OutColor = EncodeLightAttenuation(half4(FadedShadow, FadedSSSShadow, FadedShadow, FadedSSSShadow));
+
+				if (MultiplexId == 0)
+					SignalOutput_UAVs_0[OutputPixelPostion] = OutColor;
+				if (MultiplexId == 1)
+					SignalOutput_UAVs_1[OutputPixelPostion] = OutColor;
+				if (MultiplexId == 2)
+					SignalOutput_UAVs_2[OutputPixelPostion] = OutColor;
+				if (MultiplexId == 3)
+					SignalOutput_UAVs_3[OutputPixelPostion] = OutColor;
+			}
 		}
+		#else
+		{
+			OutputMultiplexedSignal(
+				SignalOutput_UAVs_0,
+				SignalOutput_UAVs_1,
+				SignalOutput_UAVs_2,
+				SignalOutput_UAVs_3,
+				CONFIG_SIGNAL_OUTPUT_LAYOUT, MultiplexCount,
+				OutputPixelPostion, OutputSamples.Array);
+		}
+		#endif
 	}
-	#else
-	{
-		OutputMultiplexedSignal(
-			SignalOutput_UAVs_0,
-			SignalOutput_UAVs_1,
-			SignalOutput_UAVs_2,
-			SignalOutput_UAVs_3,
-			CONFIG_SIGNAL_OUTPUT_LAYOUT, MultiplexCount,
-			OutputPixelPostion, OutputSamples.Array);
-	}
-	#endif
 } // MainCS
diff --git a/Engine/Shaders/Private/ScreenSpaceDenoise/SSDSpatialKernel.ush b/Engine/Shaders/Private/ScreenSpaceDenoise/SSDSpatialKernel.ush
--- a/Engine/Shaders/Private/ScreenSpaceDenoise/SSDSpatialKernel.ush
+++ b/Engine/Shaders/Private/ScreenSpaceDenoise/SSDSpatialKernel.ush
@@ -112,6 +112,9 @@ struct FSSDKernelConfig
 	// Whether the spherical harmonic of a sample should be computed before accumulation.
 	bool bComputeSampleColorSH;
 
+	// Whether should clamp the UV individually per signal.
+	bool bClampUVPerMultiplexedSignal;
+
 	// The color space that has been encoded in the buffer.
 	uint BufferColorSpace[MAX_SIGNAL_MULTIPLEXING];
 	
@@ -137,6 +140,9 @@ struct FSSDKernelConfig
 	// The periode of the harmonic being sampled.
 	float HarmonicPeriode;
 
+	// Buffer's min and max UV, per texture.
+	float4 PerSignalUVMinMax[MAX_SIGNAL_MULTIPLEXING];
+
 	
 	// --------------------------- Per lane.
 	
@@ -182,6 +188,7 @@ FSSDKernelConfig CreateKernelConfig()
 	KernelConfig.bMinSamplePairInvFrequency = false;
 	KernelConfig.bMaxWithRefBilateralDistance = false;
 	KernelConfig.bComputeSampleColorSH = false;
+	KernelConfig.bClampUVPerMultiplexedSignal = false;
 	
 	{
 		UNROLL_N(MAX_SIGNAL_MULTIPLEXING)
@@ -200,6 +207,14 @@ FSSDKernelConfig CreateKernelConfig()
 	KernelConfig.KernelSpreadFactor = 1;
 	KernelConfig.HarmonicPeriode = 1.0;
 	
+	{
+		UNROLL_N(MAX_SIGNAL_MULTIPLEXING)
+		for (uint MultiplexId = 0; MultiplexId < MAX_SIGNAL_MULTIPLEXING; MultiplexId++)
+		{
+			KernelConfig.PerSignalUVMinMax[MultiplexId] = 0.0;
+		}
+	}
+
 	// VGPRs.
 	KernelConfig.BoxKernelRadius = 1;
 	KernelConfig.SampleCount = 1;
@@ -445,6 +460,21 @@ void SampleMultiplexedSignals(
 		KernelConfig.bNormalizeSample,
 		SampleBufferUV,
 		KernelConfig.BufferMipLevel + MipLevel);
+
+	if (KernelConfig.bClampUVPerMultiplexedSignal)
+	{
+		UNROLL_N(MAX_SIGNAL_MULTIPLEXING)
+		for (uint SignalMultiplexId = 0; SignalMultiplexId < MAX_SIGNAL_MULTIPLEXING; SignalMultiplexId++)
+		{
+			bool bInvalidSample = any(SampleBufferUV != clamp(
+				SampleBufferUV, KernelConfig.PerSignalUVMinMax[SignalMultiplexId].xy, KernelConfig.PerSignalUVMinMax[SignalMultiplexId].zw));
+
+			if (bInvalidSample)
+			{
+				OutMultiplexedSamples.Array[SignalMultiplexId] = CreateSignalSampleFromScalarValue(0.0);
+			}
+		} // for (uint SignalMultiplexId = 0; SignalMultiplexId < MAX_SIGNAL_MULTIPLEXING; SignalMultiplexId++)
+	}
 } // SampleMultiplexedSignals()
 
 /** Accumulate multiplexed samples and their metadata to an accumulator. */
diff --git a/Engine/Shaders/Private/ScreenSpaceDenoise/SSDTemporalAccumulation.usf b/Engine/Shaders/Private/ScreenSpaceDenoise/SSDTemporalAccumulation.usf
--- a/Engine/Shaders/Private/ScreenSpaceDenoise/SSDTemporalAccumulation.usf
+++ b/Engine/Shaders/Private/ScreenSpaceDenoise/SSDTemporalAccumulation.usf
@@ -60,6 +60,10 @@
 		#define CONFIG_REJECTION_SAMPLE_SET               REJECTION_SAMPLE_SET_SOBEK2018
 	#endif
 
+	#if DIM_SIGNAL_BATCH_SIZE > 1
+		#define CONFIG_CLAMP_UV_PER_SIGNAL 1
+	#endif
+
 	#define MAX_SIGNAL_BATCH_SIZE CONFIG_SIGNAL_BATCH_SIZE
 
 #elif CONFIG_SIGNAL_PROCESSING == SIGNAL_PROCESSING_POLYCHROMATIC_PENUMBRA_HARMONIC
@@ -173,6 +177,11 @@
 
 //------------------------------------------------------- CONFIG DISABLED DEFAULTS
 
+/** Whether should clamp the UV individually per texture. */
+#ifndef CONFIG_CLAMP_UV_PER_SIGNAL
+	#define CONFIG_CLAMP_UV_PER_SIGNAL 0
+#endif
+
 /** Whether this pass make use of a tile classification. */
 #ifndef CONFIG_USE_TILE_CLASSIFICATION
 	#define CONFIG_USE_TILE_CLASSIFICATION 0
@@ -259,11 +268,7 @@ int bCameraCut[MAX_SIGNAL_BATCH_SIZE];
 
 float HistoryPreExposureCorrection;
 
-static const uint iMipLevel = 0;
-static const uint iMipLevelPow2 = 1;
-static const float MipLevel = 0;
-static const float MipLevelPow2 = 1;
-static const float InvMipLevelPow2 = 1;
+float4 HistoryBufferUVMinMax[CONFIG_SIGNAL_BATCH_SIZE];
 
 
 #if !defined(CONFIG_INPUT_TEXTURE_COUNT)
@@ -379,10 +384,10 @@ void TemporallyAccumulate(
 	uint GroupThreadIndex : SV_GroupIndex)
 {
 	// Find out scene buffer UV.
-	float2 SceneBufferUV = ((DispatchThreadId + 0.5) * MipLevelPow2 + View.ViewRectMin.xy) * View.BufferSizeAndInvSize.zw;
+	float2 SceneBufferUV = ((ViewportMin + DispatchThreadId) + 0.5) * View.BufferSizeAndInvSize.zw;
 	if (true)
 	{
-		SceneBufferUV = clamp(SceneBufferUV, View.BufferBilinearUVMinMax.xy, View.BufferBilinearUVMinMax.zw);
+		SceneBufferUV = clamp(SceneBufferUV, BufferBilinearUVMinMax.xy, BufferBilinearUVMinMax.zw);
 	}
 	float2 ViewportUV = BufferUVToViewportUV(SceneBufferUV);
 	float2 ScreenPosition = ViewportUVToScreenPos(ViewportUV);
@@ -429,7 +434,7 @@ void TemporallyAccumulate(
 
 
 	// Use the diameter, because that is the distance between two pixel.
-	float PixelWorldBluringRadius = MipLevelPow2 * ComputeWorldBluringRadiusCausedByPixelSize(RefSceneMetadata);
+	float PixelWorldBluringRadius = ComputeWorldBluringRadiusCausedByPixelSize(RefSceneMetadata);
 	float WorldBluringRadius = WorldBluringRadiusToBilateralWorldDistance(PixelWorldBluringRadius);
 	
 	float TargetedSampleCount = TARGETED_SAMPLE_COUNT;
@@ -453,7 +458,7 @@ void TemporallyAccumulate(
 	FSSDSignalArray HistorySamples = CreateSignalArrayFromScalarValue(0.0);
 	{
 		float2 HistoryBufferUV = HistoryScreenPosition * View.ScreenPositionScaleBias.xy + View.ScreenPositionScaleBias.wz;
-		float2 ClampedHistoryBufferUV = clamp(HistoryBufferUV, View.BufferBilinearUVMinMax.xy, View.BufferBilinearUVMinMax.zw);
+		float2 ClampedHistoryBufferUV = clamp(HistoryBufferUV, BufferBilinearUVMinMax.xy, BufferBilinearUVMinMax.zw);
 		bool bIsPreviousFrameOffscreen = any(HistoryBufferUV != ClampedHistoryBufferUV);
 
 		// TODO: global camera cut to save perf that is going to be required for panic post filtering.
@@ -469,6 +474,7 @@ void TemporallyAccumulate(
 			KernelConfig.bUnroll = true;
 			KernelConfig.bPreviousFrameMetadata = true;
 			KernelConfig.BilateralDistanceComputation = SIGNAL_WORLD_FREQUENCY_MIN_METADATA;
+			KernelConfig.bClampUVPerMultiplexedSignal = CONFIG_CLAMP_UV_PER_SIGNAL != 0;
 			
 			if (CONFIG_USE_NEAREST_HISTORY)
 				KernelConfig.SampleSet = SAMPLE_SET_1X1;
@@ -477,7 +483,18 @@ void TemporallyAccumulate(
 
 			// SGPR configuration of the kernel.
 			KernelConfig.BufferSizeAndInvSize = View.BufferSizeAndInvSize;
-			KernelConfig.BufferBilinearUVMinMax = View.BufferBilinearUVMinMax;
+			KernelConfig.BufferBilinearUVMinMax = BufferBilinearUVMinMax;
+			
+			#if CONFIG_CLAMP_UV_PER_SIGNAL
+			{
+				UNROLL_N(CONFIG_SIGNAL_BATCH_SIZE)
+				for (uint BatchedSignalId = 0; BatchedSignalId < CONFIG_SIGNAL_BATCH_SIZE; BatchedSignalId++)
+				{
+					uint MultiplexId = BatchedSignalId / CONFIG_MULTIPLEXED_SIGNALS_PER_SIGNAL_DOMAIN;
+					KernelConfig.PerSignalUVMinMax[MultiplexId] = HistoryBufferUVMinMax[MultiplexId];
+				}
+			}
+			#endif
 			
 			// VGPR configuration of the kernel.
 			KernelConfig.BufferUV = HistoryBufferUV;
@@ -614,7 +631,7 @@ void TemporallyAccumulate(
 		// SGPR configuration of the kernel.
 		{
 			KernelConfig.BufferSizeAndInvSize = View.BufferSizeAndInvSize;
-			KernelConfig.BufferBilinearUVMinMax = View.BufferBilinearUVMinMax;
+			KernelConfig.BufferBilinearUVMinMax = BufferBilinearUVMinMax;
 		}
 
 		// VGPR configuration of the kernel.
@@ -966,9 +983,11 @@ void TemporallyAccumulate(
 		// No need to keep DispatchThreadId, while SceneBufferUV is arround at highest VGPR peak because center of the kernel.
 		uint2 OutputPixelPostion = uint2(SceneBufferUV * View.BufferSizeAndInvSize.xy);
 	#else
-		uint2 OutputPixelPostion = DispatchThreadId;
+		uint2 OutputPixelPostion = ViewportMin + DispatchThreadId;
 	#endif
-
+		
+	BRANCH
+	if (all(OutputPixelPostion < ViewportMax))
 	{
 		OutputMultiplexedSignal(
 			SignalHistoryOutput_UAVs_0,
diff --git a/Engine/Source/Runtime/Renderer/Private/ScreenSpaceDenoise.cpp b/Engine/Source/Runtime/Renderer/Private/ScreenSpaceDenoise.cpp
--- a/Engine/Source/Runtime/Renderer/Private/ScreenSpaceDenoise.cpp
+++ b/Engine/Source/Runtime/Renderer/Private/ScreenSpaceDenoise.cpp
@@ -621,6 +621,10 @@ bool ShouldCompileSignalPipeline(ESignalProcessing SignalProcessing, EShaderPlat
 
 /** Shader parameter structure used for all shaders. */
 BEGIN_SHADER_PARAMETER_STRUCT(FSSDCommonParameters, )
+	SHADER_PARAMETER(FIntPoint, ViewportMin)
+	SHADER_PARAMETER(FIntPoint, ViewportMax)
+	SHADER_PARAMETER(FVector4, BufferBilinearUVMinMax)
+
 	SHADER_PARAMETER_STRUCT_INCLUDE(FSceneTextureParameters, SceneTextures)
 	SHADER_PARAMETER_RDG_TEXTURE(Texture2D, EyeAdaptation)
 	SHADER_PARAMETER_RDG_TEXTURE(Texture2D<uint>, TileClassificationTexture)
@@ -910,6 +914,8 @@ class FSSDSpatialAccumulationCS : public FGlobalShader
 	}
 
 	BEGIN_SHADER_PARAMETER_STRUCT(FParameters, )
+		SHADER_PARAMETER_ARRAY(FVector4, InputBufferUVMinMax, [IScreenSpaceDenoiser::kMaxBatchSize])
+
 		SHADER_PARAMETER(uint32, MaxSampleCount)
 		SHADER_PARAMETER(int32, UpscaleFactor)
 		SHADER_PARAMETER(float, KernelSpreadFactor)
@@ -957,6 +963,8 @@ class FSSDTemporalAccumulationCS : public FGlobalShader
 		SHADER_PARAMETER(FMatrix, PrevScreenToTranslatedWorld)
 		SHADER_PARAMETER(float, HistoryPreExposureCorrection)
 
+		SHADER_PARAMETER_ARRAY(FVector4, HistoryBufferUVMinMax, [IScreenSpaceDenoiser::kMaxBatchSize])
+
 		SHADER_PARAMETER_STRUCT_INCLUDE(FSSDCommonParameters, CommonParameters)
 		SHADER_PARAMETER_STRUCT_INCLUDE(FSSDConvolutionMetaData, ConvolutionMetaData)
 
@@ -1005,6 +1013,7 @@ IMPLEMENT_GLOBAL_SHADER(FSSDComposeHarmonicsCS, "/Engine/Private/ScreenSpaceDeno
 /** Generic settings to denoise signal at constant pixel density across the viewport. */
 struct FSSDConstantPixelDensitySettings
 {
+	FIntRect Viewport;
 	ESignalProcessing SignalProcessing;
 	int32 SignalBatchSize = 1;
 	float HarmonicPeriode = 1.0f;
@@ -1016,6 +1025,7 @@ struct FSSDConstantPixelDensitySettings
 	bool bUseTemporalAccumulation = false;
 	int32 HistoryConvolutionSampleCount = 1;
 	float HistoryConvolutionKernelSpreadFactor = 1.0f;
+	TStaticArray<FIntRect, IScreenSpaceDenoiser::kMaxBatchSize> SignalScissor;
 	TStaticArray<const FLightSceneInfo*, IScreenSpaceDenoiser::kMaxBatchSize> LightSceneInfo;
 };
 
@@ -1031,6 +1041,14 @@ static void DenoiseSignalAtConstantPixelDensity(
 	FSSDSignalTextures* OutputSignal)
 {
 	check(UsesConstantPixelDensityPassLayout(Settings.SignalProcessing));
+	
+	// Make sure the viewport of the denoiser is within the viewport of the view.
+	{
+		FIntRect Union = View.ViewRect;
+		Union.Union(Settings.Viewport);
+		check(Union == View.ViewRect);
+	}
+
 	ensure(Settings.InputResolutionFraction == 1.0f || Settings.InputResolutionFraction == 0.5f);
 	
 	auto GetResourceNames = [&](const TCHAR* const ResourceNames[])
@@ -1042,10 +1060,9 @@ static void DenoiseSignalAtConstantPixelDensity(
 		Settings.MaxInputSPP > 1 || 
 		(CVarShadowUse1SPPCodePath.GetValueOnRenderThread() == 0 && Settings.SignalProcessing == ESignalProcessing::MonochromaticPenumbra));
 
-	const FIntPoint DenoiseResolution = View.ViewRect.Size();
+	FIntPoint BufferExtent = SceneTextures.SceneDepthBuffer->Desc.Extent;
+	const FIntPoint DenoiseResolution = Settings.Viewport.Size();
 	
-	FSceneRenderTargets& SceneContext = FSceneRenderTargets::Get(GraphBuilder.RHICmdList);
-
 	// Number of signal to batch.
 	int32 MaxSignalBatchSize = SignalMaxBatchSize(Settings.SignalProcessing);
 	check(Settings.SignalBatchSize >= 1 && Settings.SignalBatchSize <= MaxSignalBatchSize);
@@ -1075,7 +1092,7 @@ static void DenoiseSignalAtConstantPixelDensity(
 		};
 
 		FRDGTextureDesc RefDesc = FRDGTextureDesc::Create2DDesc(
-			SceneTextures.SceneDepthBuffer->Desc.Extent,
+			BufferExtent,
 			PF_Unknown,
 			FClearValueBinding::Black,
 			/* InFlags = */ TexCreate_None,
@@ -1195,6 +1212,14 @@ static void DenoiseSignalAtConstantPixelDensity(
 	// Setup common shader parameters.
 	FSSDCommonParameters CommonParameters;
 	{
+		CommonParameters.ViewportMin = Settings.Viewport.Min;
+		CommonParameters.ViewportMax = Settings.Viewport.Max;
+		CommonParameters.BufferBilinearUVMinMax = FVector4(
+			float(Settings.Viewport.Min.X + 0.5f) / float(BufferExtent.X),
+			float(Settings.Viewport.Min.Y + 0.5f) / float(BufferExtent.Y),
+			float(Settings.Viewport.Max.X - 0.5f) / float(BufferExtent.X),
+			float(Settings.Viewport.Max.Y - 0.5f) / float(BufferExtent.Y));
+
 		CommonParameters.SceneTextures = SceneTextures;
 		CommonParameters.ViewUniformBuffer = View.ViewUniformBuffer;
 		CommonParameters.EyeAdaptation = GetEyeAdaptationTexture(GraphBuilder, View);
@@ -1328,6 +1353,16 @@ static void DenoiseSignalAtConstantPixelDensity(
 			GetResourceNames(kReconstructionResourceNames));
 
 		FSSDSpatialAccumulationCS::FParameters* PassParameters = GraphBuilder.AllocParameters<FSSDSpatialAccumulationCS::FParameters>();
+		for (int32 BatchedSignalId = 0; BatchedSignalId < Settings.SignalBatchSize; BatchedSignalId++)
+		{
+			FIntRect SignalScissor = Settings.SignalScissor[BatchedSignalId];
+			PassParameters->InputBufferUVMinMax[BatchedSignalId] = FVector4(
+				float(SignalScissor.Min.X + 0.5f) / float(BufferExtent.X),
+				float(SignalScissor.Min.Y + 0.5f) / float(BufferExtent.Y),
+				float(SignalScissor.Max.X - 0.5f) / float(BufferExtent.X),
+				float(SignalScissor.Max.Y - 0.5f) / float(BufferExtent.Y));
+		}
+
 		PassParameters->MaxSampleCount = FMath::Clamp(Settings.ReconstructionSamples, 1, kStackowiakMaxSampleCountPerSet);
 		PassParameters->UpscaleFactor = int32(1.0f / Settings.InputResolutionFraction);
 		PassParameters->HarmonicPeriode = Settings.HarmonicPeriode;
@@ -1349,10 +1384,11 @@ static void DenoiseSignalAtConstantPixelDensity(
 		TShaderMapRef<FSSDSpatialAccumulationCS> ComputeShader(View.ShaderMap, PermutationVector);
 		FComputeShaderUtils::AddPass(
 			GraphBuilder,
-			RDG_EVENT_NAME("SSD SpatialAccumulation(Reconstruction MaxSamples=%i Upscale=%i MultiSPP=%i)",
+			RDG_EVENT_NAME("SSD SpatialAccumulation(Reconstruction MaxSamples=%i Scissor=%ix%i%s%s)",
 				PassParameters->MaxSampleCount,
-				int32(PermutationVector.Get<FSSDSpatialAccumulationCS::FUpscaleDim>()),
-				int32(PermutationVector.Get<FMultiSPPDim>())),
+				DenoiseResolution.X, DenoiseResolution.Y,
+				PermutationVector.Get<FSSDSpatialAccumulationCS::FUpscaleDim>() ? TEXT(" Upscale") : TEXT(""),
+				PermutationVector.Get<FMultiSPPDim>() ? TEXT("") : TEXT(" 1SPP")),
 			*ComputeShader,
 			PassParameters,
 			FComputeShaderUtils::GetGroupCount(DenoiseResolution, FSSDSpatialAccumulationCS::kGroupSize));
@@ -1515,6 +1551,12 @@ static void DenoiseSignalAtConstantPixelDensity(
 					GraphBuilder, PrevFrameHistory->RT[BufferId], GSystemTextures.BlackDummy);
 			}
 
+			PassParameters->HistoryBufferUVMinMax[BatchedSignalId] = FVector4(
+				float(PrevFrameHistory->Scissor.Min.X + 0.5f) / float(BufferExtent.X),
+				float(PrevFrameHistory->Scissor.Min.Y + 0.5f) / float(BufferExtent.Y),
+				float(PrevFrameHistory->Scissor.Max.X - 0.5f) / float(BufferExtent.X),
+				float(PrevFrameHistory->Scissor.Max.Y - 0.5f) / float(BufferExtent.Y));
+
 			// Releases the reference on previous frame so the history's render target can be reused ASAP.
 			PrevFrameHistory->SafeRelease();
 		} // for (uint32 BatchedSignalId = 0; BatchedSignalId < Settings.SignalBatchSize; BatchedSignalId++)
@@ -1604,6 +1646,8 @@ static void DenoiseSignalAtConstantPixelDensity(
 				int32 HistoryBufferId = BatchedSignalId * HistoryTextureCountPerSignal + BufferId;
 				GraphBuilder.QueueTextureExtraction(SignalHistory.Textures[HistoryBufferId], &NewHistory->RT[BufferId]);
 			}
+
+			NewHistory->Scissor = Settings.Viewport;
 		} // for (uint32 BatchedSignalId = 0; BatchedSignalId < Settings.SignalBatchSize; BatchedSignalId++)
 	}
 	else if (HistoryTextureCountPerSignal >= 2)
@@ -1752,8 +1796,33 @@ public:
 		for (int32 BatchedSignalId = 0; BatchedSignalId < InputParameterCount; BatchedSignalId++)
 		{
 			const FShadowParameters& Parameters = InputParameters[BatchedSignalId];
+			const FLightSceneProxy* Proxy = Parameters.LightSceneInfo->Proxy;
+			
+			// Scissor the denoiser.
+			{
+				FIntRect LightScissorRect;
+				if (Proxy->GetScissorRect(/* out */ LightScissorRect, View, View.ViewRect))
+				{
 
-			ensure(IsSupportedLightType(ELightComponentType(Parameters.LightSceneInfo->Proxy->GetLightType())));
+				}
+				else
+				{
+					LightScissorRect = View.ViewRect;
+				}
+
+				if (BatchedSignalId == 0)
+				{
+					Settings.Viewport = LightScissorRect;
+				}
+				else
+				{
+					Settings.Viewport.Union(LightScissorRect);
+				}
+
+				Settings.SignalScissor[BatchedSignalId] = LightScissorRect;
+			}
+
+			ensure(IsSupportedLightType(ELightComponentType(Proxy->GetLightType())));
 
 			Settings.LightSceneInfo[BatchedSignalId] = Parameters.LightSceneInfo;
 			if (Settings.MaxInputSPP == 1 && CVarShadowUse1SPPCodePath.GetValueOnRenderThread() != 0)
@@ -1776,6 +1845,12 @@ public:
 			}
 		}
 
+		// Force viewport to be a multiple of 2, to avoid over frame interference between TAA jitter of the frame, and Stackowiack's SampleTrackId.
+		{
+			Settings.Viewport.Min.X &= ~1;
+			Settings.Viewport.Min.Y &= ~1;
+		}
+
 		FSSDSignalTextures SignalOutput;
 		DenoiseSignalAtConstantPixelDensity(
 			GraphBuilder, View, SceneTextures,
@@ -1820,6 +1895,7 @@ public:
 			int32 Periode = 1 << HarmonicId;
 
 			FSSDConstantPixelDensitySettings Settings;
+			Settings.Viewport = View.ViewRect;
 			Settings.SignalProcessing = ESignalProcessing::PolychromaticPenumbraHarmonic;
 			Settings.HarmonicPeriode = Periode;
 			Settings.ReconstructionSamples = Periode * Periode; // TODO: should use preconvolution instead for harmonic 3
@@ -1855,6 +1931,7 @@ public:
 			int32 Periode = 1 << HarmonicId;
 
 			FSSDConstantPixelDensitySettings Settings;
+			Settings.Viewport = View.ViewRect;
 			Settings.SignalProcessing = ESignalProcessing::PolychromaticPenumbraHarmonic;
 			Settings.HarmonicPeriode = Periode;
 			Settings.ReconstructionSamples = Periode * Periode; // TODO: should use preconvolution instead for harmonic 3
@@ -1889,9 +1966,11 @@ public:
 		// Merges the different harmonics.
 		FSSDSignalTextures ComposedHarmonics;
 		{
+			FIntPoint BufferExtent = SceneTextures.SceneDepthBuffer->Desc.Extent;
+
 			{
 				FRDGTextureDesc Desc = FRDGTextureDesc::Create2DDesc(
-					SceneTextures.SceneDepthBuffer->Desc.Extent,
+					BufferExtent,
 					PF_FloatRGBA,
 					FClearValueBinding::Black,
 					/* InFlags = */ TexCreate_None,
@@ -1904,6 +1983,13 @@ public:
 
 			ComposePassParameters->CommonParameters.ViewUniformBuffer = View.ViewUniformBuffer;
 			ComposePassParameters->CommonParameters.SceneTextures = SceneTextures;
+			ComposePassParameters->CommonParameters.ViewportMin = View.ViewRect.Min;
+			ComposePassParameters->CommonParameters.ViewportMax = View.ViewRect.Max;
+			ComposePassParameters->CommonParameters.BufferBilinearUVMinMax = FVector4(
+				float(View.ViewRect.Min.X + 0.5f) / float(BufferExtent.X),
+				float(View.ViewRect.Min.Y + 0.5f) / float(BufferExtent.Y),
+				float(View.ViewRect.Max.X - 0.5f) / float(BufferExtent.X),
+				float(View.ViewRect.Max.Y - 0.5f) / float(BufferExtent.Y));
 
 			ComposePassParameters->SignalOutput = CreateMultiplexedUAVs(GraphBuilder, ComposedHarmonics);
 
@@ -1931,6 +2017,7 @@ public:
 		FPolychromaticPenumbraOutputs Outputs;
 		{
 			FSSDConstantPixelDensitySettings Settings;
+			Settings.Viewport = View.ViewRect;
 			Settings.SignalProcessing = ESignalProcessing::PolychromaticPenumbraHarmonic;
 			Settings.bEnableReconstruction = false;
 			Settings.bUseTemporalAccumulation = CVarShadowTemporalAccumulation.GetValueOnRenderThread() != 0;
@@ -1973,6 +2060,7 @@ public:
 		InputSignal.Textures[1] = ReflectionInputs.RayHitDistance;
 
 		FSSDConstantPixelDensitySettings Settings;
+		Settings.Viewport = View.ViewRect;
 		Settings.SignalProcessing = ESignalProcessing::Reflections;
 		Settings.InputResolutionFraction = RayTracingConfig.ResolutionFraction;
 		Settings.ReconstructionSamples = CVarReflectionReconstructionSampleCount.GetValueOnRenderThread();
@@ -2013,6 +2101,7 @@ public:
 		InputSignal.Textures[1] = ReflectionInputs.RayHitDistance;
 		
 		FSSDConstantPixelDensitySettings Settings;
+		Settings.Viewport = View.ViewRect;
 		Settings.SignalProcessing = ESignalProcessing::AmbientOcclusion;
 		Settings.InputResolutionFraction = RayTracingConfig.ResolutionFraction;
 		Settings.ReconstructionSamples = CVarAOReconstructionSampleCount.GetValueOnRenderThread();
@@ -2054,6 +2143,7 @@ public:
 		InputSignal.Textures[1] = Inputs.RayHitDistance;
 
 		FSSDConstantPixelDensitySettings Settings;
+		Settings.Viewport = View.ViewRect;
 		Settings.SignalProcessing = ESignalProcessing::DiffuseAndAmbientOcclusion;
 		Settings.InputResolutionFraction = Config.ResolutionFraction;
 		Settings.ReconstructionSamples = CVarGIReconstructionSampleCount.GetValueOnRenderThread();
@@ -2096,6 +2186,7 @@ public:
 		InputSignal.Textures[1] = Inputs.RayHitDistance;
 
 		FSSDConstantPixelDensitySettings Settings;
+		Settings.Viewport = View.ViewRect;
 		Settings.SignalProcessing = ESignalProcessing::DiffuseAndAmbientOcclusion;
 		Settings.InputResolutionFraction = Config.ResolutionFraction;
 		Settings.ReconstructionSamples = CVarGIReconstructionSampleCount.GetValueOnRenderThread();
@@ -2138,6 +2229,7 @@ public:
 			InputSignal.Textures[i] = Inputs.SphericalHarmonic[i];
 
 		FSSDConstantPixelDensitySettings Settings;
+		Settings.Viewport = View.ViewRect;
 		Settings.SignalProcessing = ESignalProcessing::DiffuseSphericalHarmonic;
 		Settings.InputResolutionFraction = Config.ResolutionFraction;
 		Settings.ReconstructionSamples = CVarGIReconstructionSampleCount.GetValueOnRenderThread();
