Adds Scissor to FScreenSpaceDenoiserHistory

diff --git a/Engine/Source/Runtime/Renderer/Private/SceneRendering.h b/Engine/Source/Runtime/Renderer/Private/SceneRendering.h
--- a/Engine/Source/Runtime/Renderer/Private/SceneRendering.h
+++ b/Engine/Source/Runtime/Renderer/Private/SceneRendering.h
@@ -721,6 +721,9 @@ struct FScreenSpaceDenoiserHistory
 	// Number of history render target to store.
 	static constexpr int32 RTCount = 3;
 
+	// Scissor of valid data in the render target;
+	FIntRect Scissor;
+
 	// Render target specific to the history.
 	TRefCountPtr<IPooledRenderTarget> RT[RTCount];
 
