Ported histogram passes to RDG.

diff --git a/Engine/Shaders/Private/PostProcessHistogram.usf b/Engine/Shaders/Private/PostProcessHistogram.usf
--- a/Engine/Shaders/Private/PostProcessHistogram.usf
+++ b/Engine/Shaders/Private/PostProcessHistogram.usf
@@ -5,92 +5,89 @@
 =============================================================================*/
 
 #include "Common.ush"
-#include "PostProcessCommon.ush"
+#include "ScreenPass.ush"
 #include "PostProcessHistogramCommon.ush"
 
-// xy:GatherExtent, zw : unused
-float4 HistogramParameters;
+SCREEN_PASS_TEXTURE_VIEWPORT(Input)
+
+Texture2D InputTexture;
 
 // Output histogram texture (UAV)
 RWTexture2D<float4> HistogramRWTexture;
 
-
-//
+// Number of thread groups in the dispatch
 uint2 ThreadGroupCount;
-//
-uint2 LeftTopOffset;
 
 // THREADGROUP_SIZEX*THREADGROUP_SIZEY histograms of the size HISTOGRAM_SIZE
 groupshared float SharedHistogram[HISTOGRAM_SIZE][THREADGROUP_SIZEX][THREADGROUP_SIZEY];
 
-[numthreads(THREADGROUP_SIZEX, THREADGROUP_SIZEY, 1)]	// dimx,dimy,dimz
+[numthreads(THREADGROUP_SIZEX, THREADGROUP_SIZEY, 1)]
 void MainCS(
 	uint3 GroupId : SV_GroupID,
-	uint3 DispatchThreadId : SV_DispatchThreadID,	// DispatchThreadId = GroupId * int2(dimx,dimy) + GroupThreadId
-    uint3 GroupThreadId : SV_GroupThreadID,	// 0..THREADGROUP_SIZEX 0..THREADGROUP_SIZEY
-	uint GroupIndex: SV_GroupIndex) // SV_GroupIndex = SV_GroupThreadID.z*dimx*dimy + SV_GroupThreadID.y*dimx + SV_GroupThreadID.x
+	uint3 DispatchThreadId : SV_DispatchThreadID,
+    uint3 GroupThreadId : SV_GroupThreadID,
+	uint GroupIndex: SV_GroupIndex)
 {
 	// todo: can be cleared more efficiently
 	// clear all THREADGROUP_SIZEX*THREADGROUP_SIZEY histograms
-	UNROLL for(uint i = 0; i < HISTOGRAM_SIZE; ++i)
+	UNROLL for (uint i = 0; i < HISTOGRAM_SIZE; ++i)
 	{
 		SharedHistogram[i][GroupThreadId.x][GroupThreadId.y] = 0.0f;
 	}
-	 
+	
 	GroupMemoryBarrierWithGroupSync();
 
+	// Each thread group processes LoopX * LoopY texels of the input.
 	uint2 TileSize = uint2(LOOP_SIZEX, LOOP_SIZEY);
-	uint2 LeftTop = DispatchThreadId.xy * TileSize + LeftTopOffset;
 
-	uint2 GatherExtentInt = (uint2)HistogramParameters.xy + (uint2)LeftTopOffset.xy;
+	// Top left input texel for this group.
+	uint2 LeftTop = Input_ViewportMin + DispatchThreadId.xy * TileSize;
 
-	// accumulate all pixels into THREADGROUP_SIZEX*THREADGROUP_SIZEY histograms
+	// Accumulate all pixels into THREADGROUP_SIZEX*THREADGROUP_SIZEY histograms
+	LOOP for (uint y = 0; y < LOOP_SIZEY; ++y)
 	{
-		LOOP for(uint y = 0; y < LOOP_SIZEY; ++y)
+		LOOP for (uint x = 0; x < LOOP_SIZEX; ++x)
 		{
-			LOOP for(uint x = 0; x < LOOP_SIZEX; ++x)
+			uint2 Tile = uint2(x, y);
+			uint2 TexelPos = LeftTop + Tile;
+
+			if(TexelPos.x < Input_ViewportMax.x && TexelPos.y < Input_ViewportMax.y)
 			{
-				uint2 Tile = uint2(x, y);
-				uint2 TexelPos = LeftTop + Tile;
-
-				if(TexelPos.x < GatherExtentInt.x && TexelPos.y < GatherExtentInt.y)
-				{
-					float4 SceneColor = PostprocessInput0.Load(int3(TexelPos, 0));
-	#if USE_PREEXPOSURE
-					SceneColor.xyz *= View.OneOverPreExposure;
-	#endif
-					float LuminanceVal = max(SceneColor.r, max(SceneColor.g, SceneColor.b));
-					float LogLuminance = ComputeHistogramPositionFromLuminance(LuminanceVal);
+				float4 SceneColor = InputTexture.Load(int3(TexelPos, 0));
+#if USE_PREEXPOSURE
+				SceneColor.xyz *= View.OneOverPreExposure;
+#endif
+				float LuminanceVal = max(SceneColor.r, max(SceneColor.g, SceneColor.b));
+				float LogLuminance = ComputeHistogramPositionFromLuminance(LuminanceVal);
 			
-					float fBucket = saturate(LogLuminance) * (HISTOGRAM_SIZE - 1) * 0.9999f;
+				// Map the normalized histogram position into texels.
+				float fBucket = saturate(LogLuminance) * (HISTOGRAM_SIZE - 1) * 0.9999f;
 
-					uint Bucket0 = (uint)(fBucket);
-					uint Bucket1 = Bucket0 + 1;
+				// Find two discrete buckets that straddle the continuous histogram position.
+				uint Bucket0 = (uint)(fBucket);
+				uint Bucket1 = Bucket0 + 1;
 
-					float Weight1 = frac(fBucket);
-					float Weight0 = 1.0f - Weight1;
+				// Weighted blend between the two buckets.
+				float Weight1 = frac(fBucket);
+				float Weight0 = 1.0f - Weight1;
 			
-					// accumulate the weight to the nearby history buckets
-					if (Bucket0 != 0)
-					{
-						SharedHistogram[Bucket0][GroupThreadId.x][GroupThreadId.y] += Weight0;
-					}
-					SharedHistogram[Bucket1][GroupThreadId.x][GroupThreadId.y] += Weight1;
-				}
+				// Accumulate the weight to the nearby history buckets.
+				SharedHistogram[Bucket0][GroupThreadId.x][GroupThreadId.y] += Weight0;
+				SharedHistogram[Bucket1][GroupThreadId.x][GroupThreadId.y] += Weight1;
 			}
 		}
 	}
 
 	GroupMemoryBarrierWithGroupSync();
 
-	// accumulate all histograms into one
-	if(GroupIndex < HISTOGRAM_SIZE / 4) 
+	// Accumulate all histograms into one.
+	if (GroupIndex < HISTOGRAM_SIZE / 4) 
 	{
 		float4 Sum = 0;
 
-		LOOP for(uint y = 0; y < THREADGROUP_SIZEY; ++y)
+		LOOP for (uint y = 0; y < THREADGROUP_SIZEY; ++y)
 		{
-			LOOP for(uint x = 0; x < THREADGROUP_SIZEX; ++x)
+			LOOP for (uint x = 0; x < THREADGROUP_SIZEX; ++x)
 			{
 				Sum += float4(
 					SharedHistogram[GroupIndex * 4 + 0][x][y],
@@ -103,10 +100,10 @@ void MainCS(
 		float2 MaxExtent = float2(THREADGROUP_SIZEX * LOOP_SIZEX, THREADGROUP_SIZEY * LOOP_SIZEY);
 		float Area = MaxExtent.x * MaxExtent.y;
 	
-		// doesn't take borders into account but the error should be minor
+		// Doesn't take borders into account but the error should be minor
 		float NormalizeFactor = 1.0f / Area;
 
-		// output texture with one histogram per line, x and y unwrapped into all the lines
+		// Output texture with one histogram per line, x and y unwrapped into all the lines
 		HistogramRWTexture[uint2(GroupIndex, GroupId.x + GroupId.y * ThreadGroupCount.x)] = Sum * NormalizeFactor;
 	}
 }
diff --git a/Engine/Shaders/Private/PostProcessHistogramReduce.usf b/Engine/Shaders/Private/PostProcessHistogramReduce.usf
--- a/Engine/Shaders/Private/PostProcessHistogramReduce.usf
+++ b/Engine/Shaders/Private/PostProcessHistogramReduce.usf
@@ -5,15 +5,19 @@
 =============================================================================*/
 
 #include "Common.ush"
-#include "PostProcessCommon.ush"
 
-// how many lines of histograms should be compiled into one
+// How many lines of histograms should be compiled into one
 uint LoopSize;
-//
-float4 EyeAdapationTemporalParams;
-// the last frames eye adaptation settings
+
+// The last frames eye adaptation settings
 Texture2D EyeAdaptationTexture;
 
+Texture2D InputTexture;
+
+SamplerState InputSampler;
+
+float2 Input_ExtentInverse;
+
 void MainPS(noperspective float4 UVAndScreenPos : TEXCOORD0, float4 SvPosition: SV_POSITION, out float4 OutColor : SV_Target0)
 {
 	float2 UV = UVAndScreenPos.xy;
@@ -22,9 +26,9 @@ void MainPS(noperspective float4 UVAndScreenPos : TEXCOORD0, float4 SvPosition:
 
 	// accumulate all histograms into a single one
 	// could be optimized using bilinear filtering (half sample count)
-	for(uint y = 0; y < LoopSize; ++y)
+	for (uint y = 0; y < LoopSize; ++y)
 	{
-		SceneColor += Texture2DSample(PostprocessInput0, PostprocessInput0Sampler, UV + float2(0, y * PostprocessInput0Size.w));
+		SceneColor += Texture2DSample(InputTexture, InputSampler, UV + float2(0, y * Input_ExtentInverse.y));
 	}
 
 	if(SvPosition.y < 1.0f)
diff --git a/Engine/Source/Runtime/Renderer/Private/PostProcess/PostProcessEyeAdaptation.cpp b/Engine/Source/Runtime/Renderer/Private/PostProcess/PostProcessEyeAdaptation.cpp
--- a/Engine/Source/Runtime/Renderer/Private/PostProcess/PostProcessEyeAdaptation.cpp
+++ b/Engine/Source/Runtime/Renderer/Private/PostProcess/PostProcessEyeAdaptation.cpp
@@ -59,7 +59,7 @@ TAutoConsoleVariable<int32> CVarEyeAdaptationBasicCompute(
  *   Compute the parameters used for eye-adaptation.  These will default to values
  *   that disable eye-adaptation if the hardware doesn't support the minimum feature level
  */
-inline static void ComputeEyeAdaptationValues(const ERHIFeatureLevel::Type MinFeatureLevel, const FViewInfo& View, FVector4 Out[EYE_ADAPTATION_PARAMS_SIZE])
+void ComputeEyeAdaptationValues(const ERHIFeatureLevel::Type MinFeatureLevel, const FViewInfo& View, FVector4 Out[EYE_ADAPTATION_PARAMS_SIZE])
 {
 	static const auto VarDefaultAutoExposureExtendDefaultLuminanceRange = IConsoleManager::Get().FindTConsoleVariableDataInt(TEXT("r.DefaultFeature.AutoExposure.ExtendDefaultLuminanceRange"));
 	const bool bExtendedLuminanceRange = VarDefaultAutoExposureExtendDefaultLuminanceRange->GetValueOnRenderThread() == 1;
diff --git a/Engine/Source/Runtime/Renderer/Private/PostProcess/PostProcessEyeAdaptation.h b/Engine/Source/Runtime/Renderer/Private/PostProcess/PostProcessEyeAdaptation.h
--- a/Engine/Source/Runtime/Renderer/Private/PostProcess/PostProcessEyeAdaptation.h
+++ b/Engine/Source/Runtime/Renderer/Private/PostProcess/PostProcessEyeAdaptation.h
@@ -34,6 +34,8 @@ FORCEINLINE float Log2ToEV100(float Log2)
 	return Log2 - 0.263f; // Where .263 is log2(1.2)
 }
 
+void ComputeEyeAdaptationValues(const ERHIFeatureLevel::Type MinFeatureLevel, const FViewInfo& View, FVector4 Out[EYE_ADAPTATION_PARAMS_SIZE]);
+
 // Computes the eye-adaptation from HDRHistogram.
 // ePId_Input0: HDRHistogram or nothing
 // derives from TRenderingCompositePassBase<InputCount, OutputCount> 
diff --git a/Engine/Source/Runtime/Renderer/Private/PostProcess/PostProcessHistogram.cpp b/Engine/Source/Runtime/Renderer/Private/PostProcess/PostProcessHistogram.cpp
--- a/Engine/Source/Runtime/Renderer/Private/PostProcess/PostProcessHistogram.cpp
+++ b/Engine/Source/Runtime/Renderer/Private/PostProcess/PostProcessHistogram.cpp
@@ -5,13 +5,38 @@
 =============================================================================*/
 
 #include "PostProcess/PostProcessHistogram.h"
-#include "SceneUtils.h"
 #include "PostProcess/PostProcessEyeAdaptation.h"
 
-/** Encapsulates the post processing histogram compute shader. */
-class FPostProcessHistogramCS : public FGlobalShader
+namespace
 {
-	DECLARE_SHADER_TYPE(FPostProcessHistogramCS, Global);
+
+class FHistogramCS : public FGlobalShader
+{
+public:
+	// Changing these numbers requires Histogram.usf to be recompiled.
+	static const uint32 ThreadGroupSizeX = 8;
+	static const uint32 ThreadGroupSizeY = 4;
+	static const uint32 LoopCountX = 8;
+	static const uint32 LoopCountY = 8;
+	static const uint32 HistogramSize = 64;
+
+	// /4 as we store 4 buckets in one ARGB texel.
+	static const uint32 HistogramTexelCount = HistogramSize / 4;
+
+	// The number of texels on each axis processed by a single thread group.
+	static const FIntPoint TexelsPerThreadGroup;
+
+	DECLARE_GLOBAL_SHADER(FHistogramCS);
+	SHADER_USE_PARAMETER_STRUCT(FHistogramCS, FGlobalShader);
+
+	BEGIN_SHADER_PARAMETER_STRUCT(FParameters, )
+		SHADER_PARAMETER_STRUCT_REF(FViewUniformShaderParameters, View)
+		SHADER_PARAMETER_STRUCT(FScreenPassTextureViewportParameters, Input)
+		SHADER_PARAMETER_RDG_TEXTURE(Texture2D, InputTexture)
+		SHADER_PARAMETER_RDG_TEXTURE_UAV(RWTexture2D<float4>, HistogramRWTexture)
+		SHADER_PARAMETER_ARRAY(FVector4, EyeAdaptationParams, [EYE_ADAPTATION_PARAMS_SIZE])
+		SHADER_PARAMETER(FIntPoint, ThreadGroupCount)
+	END_SHADER_PARAMETER_STRUCT()
 
 	static bool ShouldCompilePermutation(const FGlobalShaderPermutationParameters& Parameters)
 	{
@@ -21,150 +46,151 @@ class FPostProcessHistogramCS : public FGlobalShader
 	static void ModifyCompilationEnvironment(const FGlobalShaderPermutationParameters& Parameters, FShaderCompilerEnvironment& OutEnvironment)
 	{
 		FGlobalShader::ModifyCompilationEnvironment(Parameters, OutEnvironment);
-		OutEnvironment.SetDefine(TEXT("THREADGROUP_SIZEX"), FRCPassPostProcessHistogram::ThreadGroupSizeX);
-		OutEnvironment.SetDefine(TEXT("THREADGROUP_SIZEY"), FRCPassPostProcessHistogram::ThreadGroupSizeY);
-		OutEnvironment.SetDefine(TEXT("LOOP_SIZEX"), FRCPassPostProcessHistogram::LoopCountX);
-		OutEnvironment.SetDefine(TEXT("LOOP_SIZEY"), FRCPassPostProcessHistogram::LoopCountY);
-		OutEnvironment.SetDefine(TEXT("HISTOGRAM_SIZE"), FRCPassPostProcessHistogram::HistogramSize);
+		OutEnvironment.SetDefine(TEXT("THREADGROUP_SIZEX"), ThreadGroupSizeX);
+		OutEnvironment.SetDefine(TEXT("THREADGROUP_SIZEY"), ThreadGroupSizeY);
+		OutEnvironment.SetDefine(TEXT("LOOP_SIZEX"), LoopCountX);
+		OutEnvironment.SetDefine(TEXT("LOOP_SIZEY"), LoopCountY);
+		OutEnvironment.SetDefine(TEXT("HISTOGRAM_SIZE"), HistogramSize);
 		OutEnvironment.SetDefine(TEXT("EYE_ADAPTATION_PARAMS_SIZE"), (uint32)EYE_ADAPTATION_PARAMS_SIZE);
 		OutEnvironment.CompilerFlags.Add( CFLAG_StandardOptimization );
 	}
 
-	/** Default constructor. */
-	FPostProcessHistogramCS() {}
-
-public:
-	FPostProcessPassParameters PostprocessParameter;
-	FShaderResourceParameter HistogramRWTexture;
-	FShaderParameter HistogramParameters;
-	FShaderParameter ThreadGroupCount;
-	FShaderParameter LeftTopOffset;
-	FShaderParameter EyeAdaptationParams;
-
-	/** Initialization constructor. */
-	FPostProcessHistogramCS(const ShaderMetaType::CompiledShaderInitializerType& Initializer)
-		: FGlobalShader(Initializer)
+	// One ThreadGroup processes LoopCountX*LoopCountY blocks of size ThreadGroupSizeX*ThreadGroupSizeY
+	static FIntPoint GetThreadGroupCount(FIntPoint InputExtent)
 	{
-		PostprocessParameter.Bind(Initializer.ParameterMap);
-		HistogramRWTexture.Bind(Initializer.ParameterMap, TEXT("HistogramRWTexture"));
-		HistogramParameters.Bind(Initializer.ParameterMap, TEXT("HistogramParameters"));
-		ThreadGroupCount.Bind(Initializer.ParameterMap, TEXT("ThreadGroupCount"));
-		LeftTopOffset.Bind(Initializer.ParameterMap, TEXT("LeftTopOffset"));
-		EyeAdaptationParams.Bind(Initializer.ParameterMap, TEXT("EyeAdaptationParams"));
+		return FIntPoint::DivideAndRoundUp(InputExtent, TexelsPerThreadGroup);
 	}
+};
 
-	void SetCS(FRHICommandList& RHICmdList, const FRenderingCompositePassContext& Context, FIntPoint ThreadGroupCountValue, FIntPoint LeftTopOffsetValue, FIntPoint GatherExtent)
-	{
-		FRHIComputeShader* ShaderRHI = GetComputeShader();
-
-		FGlobalShader::SetParameters<FViewUniformShaderParameters>(RHICmdList, ShaderRHI, Context.View.ViewUniformBuffer);
-
-		PostprocessParameter.SetCS(ShaderRHI, Context, Context.RHICmdList, TStaticSamplerState<SF_Point,AM_Clamp,AM_Clamp,AM_Clamp>::GetRHI());
-
-		SetShaderValue(RHICmdList, ShaderRHI, ThreadGroupCount, ThreadGroupCountValue);
-		SetShaderValue(RHICmdList, ShaderRHI, LeftTopOffset, LeftTopOffsetValue);
-
-		FVector4 HistogramParametersValue(GatherExtent.X, GatherExtent.Y, 0, 0);
-		SetShaderValue(RHICmdList, ShaderRHI, HistogramParameters, HistogramParametersValue);
+const FIntPoint FHistogramCS::TexelsPerThreadGroup(ThreadGroupSizeX * LoopCountX, ThreadGroupSizeY * LoopCountY);
 
-		{
-			FVector4 Temp[EYE_ADAPTATION_PARAMS_SIZE];
+IMPLEMENT_GLOBAL_SHADER(FHistogramCS, "/Engine/Private/PostProcessHistogram.usf", "MainCS", SF_Compute);
 
-			FRCPassPostProcessEyeAdaptation::ComputeEyeAdaptationParamsValue(Context.View, Temp);
-			SetShaderValueArray(RHICmdList, ShaderRHI, EyeAdaptationParams, Temp, EYE_ADAPTATION_PARAMS_SIZE);
-		}
-	}
+class FHistogramReducePS : public FGlobalShader
+{
+public:
+	DECLARE_GLOBAL_SHADER(FHistogramReducePS);
+	SHADER_USE_PARAMETER_STRUCT(FHistogramReducePS, FGlobalShader);
 	
-	// FShader interface.
-	virtual bool Serialize(FArchive& Ar) override
+	BEGIN_SHADER_PARAMETER_STRUCT(FParameters, )
+		SHADER_PARAMETER_STRUCT(FScreenPassTextureViewportParameters, Input)
+		SHADER_PARAMETER_SAMPLER(SamplerState, InputSampler)
+		SHADER_PARAMETER_RDG_TEXTURE(Texture2D, InputTexture)
+		SHADER_PARAMETER_RDG_TEXTURE(Texture2D, EyeAdaptationTexture)
+		SHADER_PARAMETER(uint32, LoopSize)
+		RENDER_TARGET_BINDING_SLOTS()
+	END_SHADER_PARAMETER_STRUCT()
+
+	// Uses full float4 to get best quality for smooth eye adaptation transitions.
+	static const EPixelFormat OutputFormat = PF_A32B32G32R32F;
+
+	static bool ShouldCompilePermutation(const FGlobalShaderPermutationParameters& Parameters)
 	{
-		bool bShaderHasOutdatedParameters = FGlobalShader::Serialize(Ar);
-		Ar << PostprocessParameter << HistogramRWTexture << HistogramParameters << ThreadGroupCount << LeftTopOffset << EyeAdaptationParams;
-		return bShaderHasOutdatedParameters;
+		return IsFeatureLevelSupported(Parameters.Platform, ERHIFeatureLevel::SM5);
 	}
-};
 
-IMPLEMENT_SHADER_TYPE(,FPostProcessHistogramCS,TEXT("/Engine/Private/PostProcessHistogram.usf"),TEXT("MainCS"),SF_Compute);
-
-void FRCPassPostProcessHistogram::Process(FRenderingCompositePassContext& Context)
-{
-	SCOPED_DRAW_EVENT(Context.RHICmdList, PostProcessHistogram);
-	const FPooledRenderTargetDesc* InputDesc = GetInputDesc(ePId_Input0);
-
-	if(!InputDesc)
+	static void ModifyCompilationEnvironment(const FGlobalShaderPermutationParameters& Parameters, FShaderCompilerEnvironment& OutEnvironment)
 	{
-		// input is not hooked up correctly
-		return;
+		FGlobalShader::ModifyCompilationEnvironment(Parameters, OutEnvironment);
+		OutEnvironment.SetRenderTargetOutputFormat(0, OutputFormat);
 	}
+};
 
-	const FViewInfo& View = Context.View;
-	const FSceneViewFamily& ViewFamily = *(View.Family);
-	
-	FIntPoint SrcSize = InputDesc->Extent;
-	FIntRect DestRect = Context.SceneColorViewRect;
-
-	const FSceneRenderTargetItem& DestRenderTarget = PassOutputs[0].RequestSurface(Context);
-
-	TShaderMapRef<FPostProcessHistogramCS> ComputeShader(Context.GetShaderMap());
-
-	// #todo-renderpasses remove once everything is renderpasses
-	UnbindRenderTargets(Context.RHICmdList);
-    Context.RHICmdList.SetComputeShader(ComputeShader->GetComputeShader());
-    
-
-	// set destination
-	check(DestRenderTarget.UAV);
-	Context.RHICmdList.TransitionResource(EResourceTransitionAccess::ERWBarrier, EResourceTransitionPipeline::EGfxToCompute, DestRenderTarget.UAV);
-	Context.RHICmdList.SetUAVParameter(ComputeShader->GetComputeShader(), ComputeShader->HistogramRWTexture.GetBaseIndex(), DestRenderTarget.UAV);
+IMPLEMENT_GLOBAL_SHADER(FHistogramReducePS, "/Engine/Private/PostProcessHistogramReduce.usf", "MainPS", SF_Pixel);
 
-	FIntPoint GatherExtent = ComputeGatherExtent(Context);
-	FIntPoint ThreadGroupCountValue = ComputeThreadGroupCount(GatherExtent);
+} //! namespace
 
-	ComputeShader->SetCS(Context.RHICmdList, Context, ThreadGroupCountValue, (DestRect.Min + FIntPoint(1, 1)) / 2, GatherExtent);
-	
-	DispatchComputeShader(Context.RHICmdList, *ComputeShader, ThreadGroupCountValue.X, ThreadGroupCountValue.Y, 1);
+FRDGTextureRef AddHistogramPass(
+	FRDGBuilder& GraphBuilder,
+	const FScreenPassViewInfo& ScreenPassView,
+	const FIntRect SceneColorViewportRect,
+	FRDGTextureRef SceneColorTexture,
+	FRDGTextureRef EyeAdaptationTexture)
+{
+	check(SceneColorTexture);
+	check(EyeAdaptationTexture);
 
-	// un-set destination
-	Context.RHICmdList.SetUAVParameter(ComputeShader->GetComputeShader(), ComputeShader->HistogramRWTexture.GetBaseIndex(), NULL);
+	const FIntPoint HistogramThreadGroupCount = FIntPoint::DivideAndRoundUp(SceneColorViewportRect.Size(), FHistogramCS::TexelsPerThreadGroup);
 
-	Context.RHICmdList.TransitionResource(EResourceTransitionAccess::EReadable, EResourceTransitionPipeline::EComputeToGfx, DestRenderTarget.UAV);
-	ensureMsgf(DestRenderTarget.TargetableTexture == DestRenderTarget.ShaderResourceTexture, TEXT("%s should be resolved to a separate SRV"), *DestRenderTarget.TargetableTexture->GetName().ToString());	
-}
+	const uint32 HistogramThreadGroupCountTotal = HistogramThreadGroupCount.X * HistogramThreadGroupCount.Y;
 
+	FRDGTextureRef HistogramTexture = nullptr;
 
-FIntPoint FRCPassPostProcessHistogram::ComputeGatherExtent(const FRenderingCompositePassContext& Context)
-{
-	// we currently assume the input is half res, one full res pixel less to avoid getting bilinear filtered input
-	return (Context.SceneColorViewRect.Size() - FIntPoint(1, 1)) / 2;
-}
+	// First pass outputs one flattened histogram per group.
+	{
+		const FIntPoint TextureExtent = FIntPoint(FHistogramCS::HistogramTexelCount, HistogramThreadGroupCountTotal);
+
+		const FRDGTextureDesc TextureDesc = FPooledRenderTargetDesc::Create2DDesc(
+			TextureExtent,
+			PF_FloatRGBA,
+			FClearValueBinding::None,
+			GFastVRamConfig.Histogram,
+			TexCreate_RenderTargetable | TexCreate_UAV,
+			false);
+
+		HistogramTexture = GraphBuilder.CreateTexture(TextureDesc, TEXT("Histogram"));
+
+		FHistogramCS::FParameters* PassParameters = GraphBuilder.AllocParameters<FHistogramCS::FParameters>();
+		PassParameters->View = ScreenPassView.View.ViewUniformBuffer;
+		PassParameters->Input = GetScreenPassTextureViewportParameters(FScreenPassTextureViewport(SceneColorViewportRect, SceneColorTexture));
+		PassParameters->InputTexture = SceneColorTexture;
+		PassParameters->HistogramRWTexture = GraphBuilder.CreateUAV(HistogramTexture);
+		PassParameters->ThreadGroupCount = HistogramThreadGroupCount;
+		ComputeEyeAdaptationValues(ERHIFeatureLevel::SM5, ScreenPassView.View, &PassParameters->EyeAdaptationParams[0]);
+
+		TShaderMapRef<FHistogramCS> ComputeShader(ScreenPassView.View.ShaderMap);
+
+		FComputeShaderUtils::AddPass(
+			GraphBuilder,
+			RDG_EVENT_NAME("Histogram %dx%d (CS)", SceneColorViewportRect.Width(), SceneColorViewportRect.Height()),
+			*ComputeShader,
+			PassParameters,
+			FIntVector(HistogramThreadGroupCount.X, HistogramThreadGroupCount.Y, 1));
+	}
 
-FIntPoint FRCPassPostProcessHistogram::ComputeThreadGroupCount(FIntPoint PixelExtent)
-{
-	uint32 TexelPerThreadGroupX = ThreadGroupSizeX * LoopCountX;
-	uint32 TexelPerThreadGroupY = ThreadGroupSizeY * LoopCountY;
+	FRDGTextureRef HistogramReduceTexture = nullptr;
 
-	uint32 ThreadGroupCountX = (PixelExtent.X + TexelPerThreadGroupX - 1) / TexelPerThreadGroupX;
-	uint32 ThreadGroupCountY = (PixelExtent.Y + TexelPerThreadGroupY - 1) / TexelPerThreadGroupY;
+	// Second pass further reduces the histogram to a single line. The second line contains the eye adaptation value (two line texture).
+	{
+		const FIntPoint TextureExtent = FIntPoint(FHistogramCS::HistogramTexelCount, 2);
+
+		const FRDGTextureDesc TextureDesc = FRDGTextureDesc::Create2DDesc(
+			TextureExtent,
+			FHistogramReducePS::OutputFormat,
+			FClearValueBinding::None,
+			GFastVRamConfig.HistogramReduce,
+			TexCreate_RenderTargetable,
+			false);
+
+		HistogramReduceTexture = GraphBuilder.CreateTexture(TextureDesc, TEXT("HistogramReduce"));
+
+		const FScreenPassTextureViewport InputViewport(HistogramTexture);
+		const FScreenPassTextureViewport OutputViewport(HistogramReduceTexture);
+
+		FHistogramReducePS::FParameters* PassParameters = GraphBuilder.AllocParameters<FHistogramReducePS::FParameters>();
+		PassParameters->Input = GetScreenPassTextureViewportParameters(InputViewport);
+		PassParameters->InputTexture = HistogramTexture;
+		PassParameters->InputSampler = TStaticSamplerState<SF_Bilinear, AM_Clamp, AM_Clamp, AM_Clamp>::GetRHI();
+		PassParameters->LoopSize = HistogramThreadGroupCountTotal;
+		PassParameters->EyeAdaptationTexture = EyeAdaptationTexture;
+		PassParameters->RenderTargets[0] = FRenderTargetBinding(HistogramReduceTexture, ERenderTargetLoadAction::ENoAction);
+
+		TShaderMapRef<FHistogramReducePS> PixelShader(ScreenPassView.View.ShaderMap);
+
+		AddDrawScreenPass(
+			GraphBuilder,
+			RDG_EVENT_NAME("HistogramReduce %dx%d (PS)", InputViewport.Extent.X, InputViewport.Extent.Y),
+			ScreenPassView,
+			OutputViewport,
+			InputViewport,
+			*PixelShader,
+			PassParameters);
+	}
 
-	return FIntPoint(ThreadGroupCountX, ThreadGroupCountY);
+	return HistogramReduceTexture;
 }
 
-FPooledRenderTargetDesc FRCPassPostProcessHistogram::ComputeOutputDesc(EPassOutputId InPassOutputId) const
+FIntPoint GetHistogramTexelsPerGroup()
 {
-	FPooledRenderTargetDesc UnmodifiedRet = GetInput(ePId_Input0)->GetOutput()->RenderTargetDesc;
-
-	UnmodifiedRet.Reset();
-	FIntPoint PixelExtent = UnmodifiedRet.Extent;
-
-	FIntPoint ThreadGroupCount = ComputeThreadGroupCount(PixelExtent);
-
-	// each ThreadGroup outputs one histogram
-	FIntPoint NewSize = FIntPoint(HistogramTexelCount, ThreadGroupCount.X * ThreadGroupCount.Y);
-
-	// format can be optimized later
-	FPooledRenderTargetDesc Ret(FPooledRenderTargetDesc::Create2DDesc(NewSize, PF_FloatRGBA, FClearValueBinding::None, TexCreate_None, TexCreate_RenderTargetable | TexCreate_UAV, false));
-	Ret.Flags |= GFastVRamConfig.Histogram;
-	Ret.DebugName = TEXT("Histogram");
-
-	return Ret;
-}
+	return FHistogramCS::TexelsPerThreadGroup;
+}
\ No newline at end of file
diff --git a/Engine/Source/Runtime/Renderer/Private/PostProcess/PostProcessHistogram.h b/Engine/Source/Runtime/Renderer/Private/PostProcess/PostProcessHistogram.h
--- a/Engine/Source/Runtime/Renderer/Private/PostProcess/PostProcessHistogram.h
+++ b/Engine/Source/Runtime/Renderer/Private/PostProcess/PostProcessHistogram.h
@@ -6,40 +6,11 @@
 
 #pragma once
 
-#include "CoreMinimal.h"
-#include "RendererInterface.h"
-#include "PostProcess/RenderingCompositionGraph.h"
-
-// derives from TRenderingCompositePassBase<InputCount, OutputCount>
-class FRCPassPostProcessHistogram : public TRenderingCompositePassBase<1, 1>
-{
-public:
-	// interface FRenderingCompositePass ---------
-	virtual void Process(FRenderingCompositePassContext& Context) override;
-	virtual void Release() override { delete this; }
-	virtual FPooledRenderTargetDesc ComputeOutputDesc(EPassOutputId InPassOutputId) const override;
-
-	// -------------------------------------------
-
-	// changing this number require Histogram.usf to be recompiled
-	static const uint32 ThreadGroupSizeX = 8;
-	// changing this number require Histogram.usf to be recompiled
-	static const uint32 ThreadGroupSizeY = 4;
-
-	static const uint32 HistogramSize = 64;
-
-	// /4 as we store 4 buckets in one ARGB texel
-	static const uint32 HistogramTexelCount = HistogramSize / 4;
-
-	// one ThreadGroup processes LoopCountX*LoopCountY blocks of size ThreadGroupSizeX*ThreadGroupSizeY
-
-	// changing this number require Histogram.usf to be recompiled
-	static const uint32 LoopCountX = 8;
-	// changing this number require Histogram.usf to be recompiled
-	static const uint32 LoopCountY = 8;
-
-	// -------------------------------------------
-
-	static FIntPoint ComputeGatherExtent(const FRenderingCompositePassContext& Context);
-	static FIntPoint ComputeThreadGroupCount(FIntPoint PixelExtent);
-};
+#include "ScreenPass.h"
+
+FRDGTextureRef AddHistogramPass(
+	FRDGBuilder& GraphBuilder,
+	const FScreenPassViewInfo& ScreenPassView,
+	const FIntRect SceneColorViewportRect,
+	FRDGTextureRef SceneColorTexture,
+	FRDGTextureRef EyeAdaptationTexture);
\ No newline at end of file
diff --git a/Engine/Source/Runtime/Renderer/Private/PostProcess/PostProcessHistogramReduce.cpp b/Engine/Source/Runtime/Renderer/Private/PostProcess/PostProcessHistogramReduce.cpp
deleted file mode 100644
diff --git a/Engine/Source/Runtime/Renderer/Private/PostProcess/PostProcessHistogramReduce.h b/Engine/Source/Runtime/Renderer/Private/PostProcess/PostProcessHistogramReduce.h
deleted file mode 100644
diff --git a/Engine/Source/Runtime/Renderer/Private/PostProcess/PostProcessTemporalAA.cpp b/Engine/Source/Runtime/Renderer/Private/PostProcess/PostProcessTemporalAA.cpp
--- a/Engine/Source/Runtime/Renderer/Private/PostProcess/PostProcessTemporalAA.cpp
+++ b/Engine/Source/Runtime/Renderer/Private/PostProcess/PostProcessTemporalAA.cpp
@@ -687,8 +687,9 @@ void AddTemporalAAPass(
 	const EPixelFormat DownsampleOverrideFormat,
 	FRDGTextureRef InSceneColorTexture,
 	FRDGTextureRef* OutSceneColorTexture,
+	FIntRect* OutSceneColorViewRect,
 	FRDGTextureRef* OutSceneColorHalfResTexture,
-	FIntRect* OutSecondaryViewRect)
+	FIntRect* OutSceneColorHalfResViewRect)
 {
 	const FViewInfo& View = ScreenPassView.View;
 
@@ -767,8 +768,9 @@ void AddTemporalAAPass(
 	}
 
 	*OutSceneColorTexture = SceneColorTexture;
+	*OutSceneColorViewRect = SecondaryViewRect;
 	*OutSceneColorHalfResTexture = TAAOutputs.DownsampledSceneColor;
-	*OutSecondaryViewRect = SecondaryViewRect;
+	*OutSceneColorHalfResViewRect = FIntRect::DivideAndRoundUp(SecondaryViewRect, 2);
 }
 
 //////////////////////////////////////////////////////////////////////////
diff --git a/Engine/Source/Runtime/Renderer/Private/PostProcess/PostProcessTemporalAA.h b/Engine/Source/Runtime/Renderer/Private/PostProcess/PostProcessTemporalAA.h
--- a/Engine/Source/Runtime/Renderer/Private/PostProcess/PostProcessTemporalAA.h
+++ b/Engine/Source/Runtime/Renderer/Private/PostProcess/PostProcessTemporalAA.h
@@ -156,5 +156,6 @@ void AddTemporalAAPass(
 	const EPixelFormat DownsampleOverrideFormat,
 	FRDGTextureRef InSceneColorTexture,
 	FRDGTextureRef* OutSceneColorTexture,
+	FIntRect* OutSceneColorViewRect,
 	FRDGTextureRef* OutSceneColorHalfResTexture,
-	FIntRect* OutSecondaryViewRect);
\ No newline at end of file
+	FIntRect* OutSceneColorHalfResViewRect);
\ No newline at end of file
diff --git a/Engine/Source/Runtime/Renderer/Private/PostProcess/PostProcessVisualizeHDR.cpp b/Engine/Source/Runtime/Renderer/Private/PostProcess/PostProcessVisualizeHDR.cpp
--- a/Engine/Source/Runtime/Renderer/Private/PostProcess/PostProcessVisualizeHDR.cpp
+++ b/Engine/Source/Runtime/Renderer/Private/PostProcess/PostProcessVisualizeHDR.cpp
@@ -19,6 +19,8 @@
 #include "PipelineStateCache.h"
 #include "Curves/CurveFloat.h"
 
+extern FIntPoint GetHistogramTexelsPerGroup();
+
 /** Encapsulates the post processing eye adaptation pixel shader. */
 class FPostProcessVisualizeHDRPS : public FGlobalShader
 {
@@ -139,12 +141,10 @@ public:
 		}
 
 		{
-			FIntPoint GatherExtent = FRCPassPostProcessHistogram::ComputeGatherExtent(Context);
-
-			uint32 TexelPerThreadGroupX = FRCPassPostProcessHistogram::ThreadGroupSizeX * FRCPassPostProcessHistogram::LoopCountX;
-			uint32 TexelPerThreadGroupY = FRCPassPostProcessHistogram::ThreadGroupSizeY * FRCPassPostProcessHistogram::LoopCountY;
+			// we currently assume the input is half res, one full res pixel less to avoid getting bilinear filtered input
+			FIntPoint GatherExtent = (Context.SceneColorViewRect.Size() - FIntPoint(1, 1)) / 2;
 
-			FIntRect Value(GatherExtent, FIntPoint(TexelPerThreadGroupX, TexelPerThreadGroupY));
+			FIntRect Value(GatherExtent, GetHistogramTexelsPerGroup());
 
 			SetShaderValue(RHICmdList, ShaderRHI, HistogramParams, Value);
 		}
diff --git a/Engine/Source/Runtime/Renderer/Private/PostProcess/PostProcessing.cpp b/Engine/Source/Runtime/Renderer/Private/PostProcess/PostProcessing.cpp
--- a/Engine/Source/Runtime/Renderer/Private/PostProcess/PostProcessing.cpp
+++ b/Engine/Source/Runtime/Renderer/Private/PostProcess/PostProcessing.cpp
@@ -20,7 +20,6 @@
 #include "PostProcess/PostProcessMobile.h"
 #include "PostProcess/PostProcessDownsample.h"
 #include "PostProcess/PostProcessHistogram.h"
-#include "PostProcess/PostProcessHistogramReduce.h"
 #include "PostProcess/PostProcessVisualizeHDR.h"
 #include "PostProcess/VisualizeShadingModels.h"
 #include "PostProcess/PostProcessSelectionOutline.h"
@@ -1059,8 +1058,8 @@ void FPostProcessing::Process(FRHICommandListImmediate& RHICmdList, const FViewI
 			SceneContext.AdjustGBufferRefCount(RHICmdList, 1);
 
 			FRenderingCompositePass* RDGPass = Context.Graph.RegisterPass(
-				new(FMemStack::Get()) TRCPassForRDG<5, 2>(
-					[&View, &SceneContext, bVisualizeBloom, bMotionBlurEnabled, bVisualizeMotionBlur, bDepthOfFieldEnabled, DownsampleQuality, DownsampleOverrideFormat]
+				new(FMemStack::Get()) TRCPassForRDG<5, 3>(
+					[&View, &SceneContext, bVisualizeBloom, bMotionBlurEnabled, bVisualizeMotionBlur, bDepthOfFieldEnabled, bHistogramEnabled, DownsampleQuality, DownsampleOverrideFormat]
 			(FRenderingCompositePass* Pass, FRenderingCompositePassContext& InContext)
 			{
 				FRDGBuilder GraphBuilder(InContext.RHICmdList);
@@ -1132,8 +1131,8 @@ void FPostProcessing::Process(FRHICommandListImmediate& RHICmdList, const FViewI
 					SceneColorTexture = AddPostProcessMaterialChain(GraphBuilder, ScreenPassView, Inputs, BL_BeforeTonemapping);
 				}
 
-				// Half-resolution scene color texture.
 				FRDGTextureRef SceneColorHalfResTexture = nullptr;
+				FIntRect SceneColorHalfResViewRect;
 
 				// Scene color view rectangle after temporal AA upscale to secondary screen percentage.
 				FIntRect SecondaryViewRect = PrimaryViewRect;
@@ -1158,8 +1157,9 @@ void FPostProcessing::Process(FRHICommandListImmediate& RHICmdList, const FViewI
 						DownsampleOverrideFormat,
 						SceneColorTexture,
 						&SceneColorTexture,
+						&SecondaryViewRect,
 						&SceneColorHalfResTexture,
-						&SecondaryViewRect);
+						&SceneColorHalfResViewRect);
 				}
 
 				//! SceneColorTexture is now upsampled to the SecondaryViewRect. Use SecondaryViewRect for input / output.
@@ -1221,6 +1221,19 @@ void FPostProcessing::Process(FRHICommandListImmediate& RHICmdList, const FViewI
 					FDownsamplePassOutputs Outputs = AddDownsamplePass(GraphBuilder, ScreenPassView, Inputs);
 
 					SceneColorHalfResTexture = Outputs.Texture;
+					SceneColorHalfResViewRect = Outputs.Viewport;
+				}
+
+				if (bHistogramEnabled)
+				{
+					FRDGTextureRef HistogramTexture = AddHistogramPass(
+						GraphBuilder,
+						ScreenPassView,
+						SceneColorHalfResViewRect,
+						SceneColorHalfResTexture,
+						GetEyeAdaptationTexture(GraphBuilder, View));
+
+					Pass->ExtractRDGTextureForOutput(GraphBuilder, ePId_Output2, HistogramTexture);
 				}
 
 				// Release held GBuffer reference taken during composition graph setup. Passes will take their own references during RDG setup.
@@ -1245,17 +1258,7 @@ void FPostProcessing::Process(FRHICommandListImmediate& RHICmdList, const FViewI
 
 			if (bHistogramEnabled)
 			{
-				FRenderingCompositePass* NodeHistogram = Context.Graph.RegisterPass(new(FMemStack::Get()) FRCPassPostProcessHistogram());
-
-				NodeHistogram->SetInput(ePId_Input0, SceneColorHalfRes);
-
-				HistogramOverScreen = FRenderingCompositeOutputRef(NodeHistogram);
-
-				FRenderingCompositePass* NodeHistogramReduce = Context.Graph.RegisterPass(new(FMemStack::Get()) FRCPassPostProcessHistogramReduce());
-
-				NodeHistogramReduce->SetInput(ePId_Input0, NodeHistogram);
-
-				Histogram = FRenderingCompositeOutputRef(NodeHistogramReduce);
+				Histogram = FRenderingCompositeOutputRef(RDGPass, ePId_Output2);
 			}
 
 			const bool bBasicEyeAdaptationEnabled = bEyeAdaptationEnabled && (AutoExposureMethod == EAutoExposureMethod::AEM_Basic);
@@ -1509,7 +1512,6 @@ void FPostProcessing::Process(FRHICommandListImmediate& RHICmdList, const FViewI
 			Node->SetInput(ePId_Input0, FRenderingCompositeOutputRef(Context.FinalOutput));
 			Node->SetInput(ePId_Input1, Histogram);
 			Node->SetInput(ePId_Input2, PreTonemapHDRColor);
-			Node->SetInput(ePId_Input3, HistogramOverScreen);
 			Node->AddDependency(EyeAdaptation);
 
 			Context.FinalOutput = FRenderingCompositeOutputRef(Node);
diff --git a/Engine/Source/Runtime/Renderer/Private/ScreenPass.h b/Engine/Source/Runtime/Renderer/Private/ScreenPass.h
--- a/Engine/Source/Runtime/Renderer/Private/ScreenPass.h
+++ b/Engine/Source/Runtime/Renderer/Private/ScreenPass.h
@@ -57,6 +57,11 @@ public:
 		, Extent(InTexture->Desc.Extent)
 	{}
 
+	FScreenPassTextureViewport(FRDGTextureRef InTexture)
+		: Rect(FIntPoint::ZeroValue, InTexture->Desc.Extent)
+		, Extent(InTexture->Desc.Extent)
+	{}
+
 	FScreenPassTextureViewport(const FScreenPassTextureViewport&) = default;
 
 	bool operator==(const FScreenPassTextureViewport& Other) const;
