Copying //UE4/Dev-Volumetrics to Dev-Rendering (//UE4/Dev-Rendering)
(Mobile/Vulkan fixes)

diff --git a/Engine/Shaders/Private/AtmosphereCommon.ush b/Engine/Shaders/Private/AtmosphereCommon.ush
--- a/Engine/Shaders/Private/AtmosphereCommon.ush
+++ b/Engine/Shaders/Private/AtmosphereCommon.ush
@@ -543,7 +543,7 @@ float3 GetSunColor(float3 X, float T, float3 V, float3 S, float Radius, float Mu
     }
 	else 
 	{
-		float SunDiskMask = step(View.AtmosphereLightDiscCosHalfApexAngle[0], dot(V, S));
+		float SunDiskMask = step(View.AtmosphereLightDiscCosHalfApexAngle[0].x, dot(V, S));
         return TransmittanceValue * SunDiskMask * View.AtmosphereLightDiscLuminance[0].rgb; // Eq (9)
     }
 }
diff --git a/Engine/Shaders/Private/MaterialTemplate.ush b/Engine/Shaders/Private/MaterialTemplate.ush
--- a/Engine/Shaders/Private/MaterialTemplate.ush
+++ b/Engine/Shaders/Private/MaterialTemplate.ush
@@ -1030,7 +1030,7 @@ float3 MaterialExpressionSkyAtmosphereLightDiskLuminance(FMaterialPixelParameter
 
 	return GetLightDiskLuminance(PlanetCenterToWorldCameraPos, ViewDir, ResolvedView.SkyAtmosphereBottomRadius, ResolvedView.SkyAtmosphereTopRadius,
 		View.TransmittanceLutTexture, View.TransmittanceLutTextureSampler,
-		ResolvedView.AtmosphereLightDirection[LightIndex].xyz, ResolvedView.AtmosphereLightDiscCosHalfApexAngle[LightIndex], ResolvedView.AtmosphereLightDiscLuminance[LightIndex].xyz);
+		ResolvedView.AtmosphereLightDirection[LightIndex].xyz, ResolvedView.AtmosphereLightDiscCosHalfApexAngle[LightIndex].x, ResolvedView.AtmosphereLightDiscLuminance[LightIndex].xyz);
 #else
 	return float3(0.0f, 0.0f, 0.0f);
 #endif
diff --git a/Engine/Shaders/Private/SkyAtmosphere.usf b/Engine/Shaders/Private/SkyAtmosphere.usf
--- a/Engine/Shaders/Private/SkyAtmosphere.usf
+++ b/Engine/Shaders/Private/SkyAtmosphere.usf
@@ -340,7 +340,7 @@ float3 GetLightDiskLuminance(float3 WorldPos, float3 WorldDir, uint LightIndex)
 		float3 LightDiskLuminance = GetLightDiskLuminance(
 			WorldPos, WorldDir, Atmosphere.BottomRadius, Atmosphere.TopRadius,
 			TransmittanceLutTexture, TransmittanceLutTextureSampler,
-			View.AtmosphereLightDirection[LightIndex].xyz, View.AtmosphereLightDiscCosHalfApexAngle[LightIndex], View.AtmosphereLightDiscLuminance[LightIndex].xyz);
+			View.AtmosphereLightDirection[LightIndex].xyz, View.AtmosphereLightDiscCosHalfApexAngle[LightIndex].x, View.AtmosphereLightDiscLuminance[LightIndex].xyz);
 
 		// Clamp to avoid crazy high values (and exposed 64000.0f luminance is already crazy high, solar system sun is 1.6x10^9). Also this removes +inf float and helps TAA.
 		const float3 MaxLightLuminance = 64000.0f;
@@ -349,7 +349,7 @@ float3 GetLightDiskLuminance(float3 WorldPos, float3 WorldDir, uint LightIndex)
 
 #if 1
 		const float ViewDotLight = dot(WorldDir, View.AtmosphereLightDirection[LightIndex].xyz);
-		const float CosHalfApex = View.AtmosphereLightDiscCosHalfApexAngle[LightIndex];
+		const float CosHalfApex = View.AtmosphereLightDiscCosHalfApexAngle[LightIndex].x;
 		const float HalfCosHalfApex = CosHalfApex + (1.0f - CosHalfApex) * 0.25; // Start fading when at 75% distance from light disk center (in cosine space)
 
 		// Apply smooth fading at edge. This is currently an eye balled fade out that works well in many cases.
diff --git a/Engine/Shaders/Public/ShaderVersion.ush b/Engine/Shaders/Public/ShaderVersion.ush
--- a/Engine/Shaders/Public/ShaderVersion.ush
+++ b/Engine/Shaders/Public/ShaderVersion.ush
@@ -3,6 +3,6 @@
 // in Platform.ush (which should be included in any shader) it allows to invalidate the shader DDC.
 // 
 // If you are merging streams and there is a conflict with this GUID you should make a new GUID rather than taking one or the other.
-// GUID = 7526116FF71A45C498CB373C6683DF8A
+// GUID = 32D42B4F58BC472BB1D08E8F7BC39FB9
 
 
diff --git a/Engine/Source/Runtime/Engine/Private/SceneManagement.cpp b/Engine/Source/Runtime/Engine/Private/SceneManagement.cpp
--- a/Engine/Source/Runtime/Engine/Private/SceneManagement.cpp
+++ b/Engine/Source/Runtime/Engine/Private/SceneManagement.cpp
@@ -736,16 +736,16 @@ FViewUniformShaderParameters::FViewUniformShaderParameters()
 	PreIntegratedBRDFSampler = TStaticSamplerState<SF_Bilinear, AM_Clamp, AM_Clamp, AM_Clamp>::GetRHI();
 
 	TransmittanceLutTexture = GWhiteTexture->TextureRHI;
-	TransmittanceLutTextureSampler = TStaticSamplerState<SF_Bilinear, AM_Clamp, AM_Clamp, AM_Clamp>::GetRHI();
+	TransmittanceLutTextureSampler = TStaticSamplerState<SF_Bilinear>::GetRHI();
 
 	SkyViewLutTexture = GBlackTexture->TextureRHI;
-	SkyViewLutTextureSampler = TStaticSamplerState<SF_Bilinear, AM_Clamp, AM_Clamp, AM_Clamp>::GetRHI();
+	SkyViewLutTextureSampler = TStaticSamplerState<SF_Bilinear>::GetRHI();
 
 	DistantSkyLightLutTexture = GBlackTexture->TextureRHI;
-	DistantSkyLightLutTextureSampler = TStaticSamplerState<SF_Bilinear, AM_Clamp, AM_Clamp, AM_Clamp>::GetRHI();
+	DistantSkyLightLutTextureSampler = TStaticSamplerState<SF_Point, AM_Wrap, AM_Wrap>::GetRHI();
 
 	CameraAerialPerspectiveVolume = GBlackAlpha1VolumeTexture->TextureRHI;
-	CameraAerialPerspectiveVolumeSampler = TStaticSamplerState<SF_Bilinear, AM_Clamp, AM_Clamp, AM_Clamp>::GetRHI();
+	CameraAerialPerspectiveVolumeSampler = TStaticSamplerState<SF_Bilinear>::GetRHI();
 
 	PrimitiveSceneData = GIdentityPrimitiveBuffer.PrimitiveSceneDataBufferSRV;
 	LightmapSceneData = GIdentityPrimitiveBuffer.LightmapSceneDataBufferSRV;
diff --git a/Engine/Source/Runtime/Engine/Public/SceneView.h b/Engine/Source/Runtime/Engine/Public/SceneView.h
--- a/Engine/Source/Runtime/Engine/Public/SceneView.h
+++ b/Engine/Source/Runtime/Engine/Public/SceneView.h
@@ -668,7 +668,16 @@ enum ETranslucencyVolumeCascade
 	VIEW_UNIFORM_BUFFER_MEMBER_ARRAY(FVector4, AtmosphereLightDirection, [NUM_ATMOSPHERE_LIGHTS]) \
 	VIEW_UNIFORM_BUFFER_MEMBER_ARRAY(FLinearColor, AtmosphereLightColor, [NUM_ATMOSPHERE_LIGHTS]) \
 	VIEW_UNIFORM_BUFFER_MEMBER_ARRAY(FLinearColor, AtmosphereLightDiscLuminance, [NUM_ATMOSPHERE_LIGHTS]) \
-	VIEW_UNIFORM_BUFFER_MEMBER_ARRAY(float, AtmosphereLightDiscCosHalfApexAngle, [NUM_ATMOSPHERE_LIGHTS]) \
+	VIEW_UNIFORM_BUFFER_MEMBER_ARRAY(FVector4, AtmosphereLightDiscCosHalfApexAngle, [NUM_ATMOSPHERE_LIGHTS]) \
+	VIEW_UNIFORM_BUFFER_MEMBER(FVector4, SkyViewLutSizeAndInvSize) \
+	VIEW_UNIFORM_BUFFER_MEMBER(float, SkyAtmosphereBottomRadius) \
+	VIEW_UNIFORM_BUFFER_MEMBER(float, SkyAtmosphereTopRadius) \
+	VIEW_UNIFORM_BUFFER_MEMBER(float, SkyAtmosphereAerialPerspectiveStartDepth) \
+	VIEW_UNIFORM_BUFFER_MEMBER(float, SkyAtmosphereCameraAerialPerspectiveVolumeDepthResolution) \
+	VIEW_UNIFORM_BUFFER_MEMBER(float, SkyAtmosphereCameraAerialPerspectiveVolumeDepthResolutionInv) \
+	VIEW_UNIFORM_BUFFER_MEMBER(float, SkyAtmosphereCameraAerialPerspectiveVolumeDepthSliceLength) \
+	VIEW_UNIFORM_BUFFER_MEMBER(float, SkyAtmosphereCameraAerialPerspectiveVolumeDepthSliceLengthInv) \
+	VIEW_UNIFORM_BUFFER_MEMBER(float, SkyAtmosphereApplyCameraAerialPerspectiveVolume) \
 	VIEW_UNIFORM_BUFFER_MEMBER(uint32, AtmosphericFogRenderMask) \
 	VIEW_UNIFORM_BUFFER_MEMBER(uint32, AtmosphericFogInscatterAltitudeSampleNum) \
 	VIEW_UNIFORM_BUFFER_MEMBER(FVector, NormalCurvatureToRoughnessScaleBias) \
@@ -708,15 +717,6 @@ enum ETranslucencyVolumeCascade
 	VIEW_UNIFORM_BUFFER_MEMBER(FMatrix, WorldToVirtualTexture) \
 	VIEW_UNIFORM_BUFFER_MEMBER(FVector4, VirtualTextureParams) \
 	VIEW_UNIFORM_BUFFER_MEMBER_ARRAY(FVector4, XRPassthroughCameraUVs, [2]) \
-	VIEW_UNIFORM_BUFFER_MEMBER(FVector4, SkyViewLutSizeAndInvSize) \
-	VIEW_UNIFORM_BUFFER_MEMBER(float, SkyAtmosphereBottomRadius) \
-	VIEW_UNIFORM_BUFFER_MEMBER(float, SkyAtmosphereTopRadius) \
-	VIEW_UNIFORM_BUFFER_MEMBER(float, SkyAtmosphereAerialPerspectiveStartDepth) \
-	VIEW_UNIFORM_BUFFER_MEMBER(float, SkyAtmosphereCameraAerialPerspectiveVolumeDepthResolution) \
-	VIEW_UNIFORM_BUFFER_MEMBER(float, SkyAtmosphereCameraAerialPerspectiveVolumeDepthResolutionInv) \
-	VIEW_UNIFORM_BUFFER_MEMBER(float, SkyAtmosphereCameraAerialPerspectiveVolumeDepthSliceLength) \
-	VIEW_UNIFORM_BUFFER_MEMBER(float, SkyAtmosphereCameraAerialPerspectiveVolumeDepthSliceLengthInv) \
-	VIEW_UNIFORM_BUFFER_MEMBER(float, SkyAtmosphereApplyCameraAerialPerspectiveVolume) \
 
 #define VIEW_UNIFORM_BUFFER_MEMBER(type, identifier) \
 	SHADER_PARAMETER(type, identifier)
diff --git a/Engine/Source/Runtime/Renderer/Private/SceneRendering.cpp b/Engine/Source/Runtime/Renderer/Private/SceneRendering.cpp
--- a/Engine/Source/Runtime/Renderer/Private/SceneRendering.cpp
+++ b/Engine/Source/Runtime/Renderer/Private/SceneRendering.cpp
@@ -1147,7 +1147,7 @@ void FViewInfo::SetupUniformBufferParameters(
 			ViewUniformShaderParameters.AtmosphericFogSunDiscScale = Scene->AtmosphericFog->SunDiscScale;
 			ViewUniformShaderParameters.AtmosphericFogRenderMask = Scene->AtmosphericFog->RenderFlag & (EAtmosphereRenderFlag::E_DisableGroundScattering | EAtmosphereRenderFlag::E_DisableSunDisk);
 			ViewUniformShaderParameters.AtmosphericFogInscatterAltitudeSampleNum = Scene->AtmosphericFog->InscatterAltitudeSampleNum;
-			ViewUniformShaderParameters.AtmosphereLightDiscCosHalfApexAngle[0] = FMath::Cos(Scene->AtmosphericFog->SunDiscScale * SunLightDiskHalfApexAngleRadian);
+			ViewUniformShaderParameters.AtmosphereLightDiscCosHalfApexAngle[0] = FVector4(FMath::Cos(Scene->AtmosphericFog->SunDiscScale * SunLightDiskHalfApexAngleRadian));
 			ViewUniformShaderParameters.AtmosphereLightDiscLuminance[0] = SunLight ? SunLight->Proxy->GetOuterSpaceLuminance() : FLinearColor::White;
 			ViewUniformShaderParameters.AtmosphereLightColor[0] = SunLight ? SunLight->Proxy->GetColor() : Scene->AtmosphericFog->DefaultSunColor; // Sun light color unaffected by atmosphere transmittance
 			ViewUniformShaderParameters.AtmosphereLightDirection[0] = SunLight ? -SunLight->Proxy->GetDirection() : -Scene->AtmosphericFog->DefaultSunDirection;
@@ -1167,7 +1167,7 @@ void FViewInfo::SetupUniformBufferParameters(
 			ViewUniformShaderParameters.AtmosphericFogSunDiscScale = 1.f;
 			ViewUniformShaderParameters.AtmosphericFogRenderMask = EAtmosphereRenderFlag::E_EnableAll;
 			ViewUniformShaderParameters.AtmosphericFogInscatterAltitudeSampleNum = 0;
-			ViewUniformShaderParameters.AtmosphereLightDiscCosHalfApexAngle[0] = FMath::Cos(SunLightDiskHalfApexAngleRadian);
+			ViewUniformShaderParameters.AtmosphereLightDiscCosHalfApexAngle[0] = FVector4(FMath::Cos(SunLightDiskHalfApexAngleRadian));
 			//Added check so atmospheric light color and vector can use a directional light without needing an atmospheric fog actor in the scene
 			ViewUniformShaderParameters.AtmosphereLightDiscLuminance[0] = SunLight ? SunLight->Proxy->GetOuterSpaceLuminance() : FLinearColor::Black;
 			ViewUniformShaderParameters.AtmosphereLightColor[0] = SunLight ? SunLight->Proxy->GetColor() : FLinearColor::Black;
@@ -1190,7 +1190,7 @@ void FViewInfo::SetupUniformBufferParameters(
 		ViewUniformShaderParameters.AtmosphericFogSunDiscScale = 1.f;
 		ViewUniformShaderParameters.AtmosphericFogRenderMask = EAtmosphereRenderFlag::E_EnableAll;
 		ViewUniformShaderParameters.AtmosphericFogInscatterAltitudeSampleNum = 0;
-		ViewUniformShaderParameters.AtmosphereLightDiscCosHalfApexAngle[0] = 1.0f;
+		ViewUniformShaderParameters.AtmosphereLightDiscCosHalfApexAngle[0] = FVector4(1.0f);
 		ViewUniformShaderParameters.AtmosphereLightDiscLuminance[0] = FLinearColor::Black;
 		ViewUniformShaderParameters.AtmosphereLightColor[0] = FLinearColor::Black;
 		ViewUniformShaderParameters.AtmosphereLightDirection[0] = DefaultSunDirection;
@@ -1235,14 +1235,14 @@ void FViewInfo::SetupUniformBufferParameters(
 			FLightSceneInfo* Light = Scene->AtmosphereLights[Index];
 			if (Light)
 			{
-				ViewUniformShaderParameters.AtmosphereLightDiscCosHalfApexAngle[Index] = FMath::Cos(Light->Proxy->GetSunLightHalfApexAngleRadian());
+				ViewUniformShaderParameters.AtmosphereLightDiscCosHalfApexAngle[Index] = FVector4(FMath::Cos(Light->Proxy->GetSunLightHalfApexAngleRadian()));
 				ViewUniformShaderParameters.AtmosphereLightDiscLuminance[Index] = Light->Proxy->GetOuterSpaceLuminance();
 				ViewUniformShaderParameters.AtmosphereLightColor[Index] = Light->Proxy->GetColor();
 				ViewUniformShaderParameters.AtmosphereLightDirection[Index] = -Light->Proxy->GetDirection();
 			}
 			else
 			{
-				ViewUniformShaderParameters.AtmosphereLightDiscCosHalfApexAngle[Index] = 1.0f;
+				ViewUniformShaderParameters.AtmosphereLightDiscCosHalfApexAngle[Index] = FVector4(1.0f);
 				ViewUniformShaderParameters.AtmosphereLightDiscLuminance[Index] = FLinearColor::Black;
 				ViewUniformShaderParameters.AtmosphereLightColor[Index] = FLinearColor::Black;
 				ViewUniformShaderParameters.AtmosphereLightDirection[Index] = DefaultSunDirection;
@@ -1265,7 +1265,7 @@ void FViewInfo::SetupUniformBufferParameters(
 	ViewUniformShaderParameters.TransmittanceLutTexture = OrWhite2DIfNull(TransmittanceLutTextureFound);
 	ViewUniformShaderParameters.TransmittanceLutTextureSampler = TStaticSamplerState<SF_Bilinear>::GetRHI();
 	ViewUniformShaderParameters.DistantSkyLightLutTexture = OrBlack2DIfNull(DistantSkyLightLutTextureFound);
-	ViewUniformShaderParameters.DistantSkyLightLutTextureSampler = TStaticSamplerState<SF_Bilinear>::GetRHI();
+	ViewUniformShaderParameters.DistantSkyLightLutTextureSampler = TStaticSamplerState<SF_Point, AM_Wrap, AM_Wrap>::GetRHI();
 	ViewUniformShaderParameters.SkyViewLutTexture = OrBlack2DIfNull(SkyViewLutTextureFound);
 	ViewUniformShaderParameters.SkyViewLutTextureSampler = TStaticSamplerState<SF_Bilinear>::GetRHI();
 	ViewUniformShaderParameters.CameraAerialPerspectiveVolume = OrBlack3DAlpha1IfNull(CameraAerialPerspectiveVolumeFound);
diff --git a/Engine/Source/Runtime/Renderer/Private/SkyAtmosphereRendering.cpp b/Engine/Source/Runtime/Renderer/Private/SkyAtmosphereRendering.cpp
--- a/Engine/Source/Runtime/Renderer/Private/SkyAtmosphereRendering.cpp
+++ b/Engine/Source/Runtime/Renderer/Private/SkyAtmosphereRendering.cpp
@@ -205,16 +205,16 @@ BEGIN_GLOBAL_SHADER_PARAMETER_STRUCT(FSkyAtmosphereInternalCommonParameters, )
 	SHADER_PARAMETER(float, FastSkyDistanceToSampleCountMaxInv)
 
 	SHADER_PARAMETER(FVector4, CameraAerialPerspectiveVolumeSizeAndInvSize)
-	SHADER_PARAMETER(float, CameraAerialPerspectiveVolumeDepthResolution)
-	SHADER_PARAMETER(float, CameraAerialPerspectiveVolumeDepthResolutionInv)
-	SHADER_PARAMETER(float, CameraAerialPerspectiveVolumeDepthSliceLength)
-	SHADER_PARAMETER(float, CameraAerialPerspectiveVolumeDepthSliceLengthInv)
+	SHADER_PARAMETER(float, CameraAerialPerspectiveVolumeDepthResolution)		// Also on View UB
+	SHADER_PARAMETER(float, CameraAerialPerspectiveVolumeDepthResolutionInv)	// Also on View UB
+	SHADER_PARAMETER(float, CameraAerialPerspectiveVolumeDepthSliceLength)		// Also on View UB
+	SHADER_PARAMETER(float, CameraAerialPerspectiveVolumeDepthSliceLengthInv)	// Also on View UB
 	SHADER_PARAMETER(float, CameraAerialPerspectiveSampleCountPerSlice)
 	SHADER_PARAMETER(float, AerialPerspectiveStartDepth)
 
 	SHADER_PARAMETER(FVector4, TransmittanceLutSizeAndInvSize)
 	SHADER_PARAMETER(FVector4, MultiScatteredLuminanceLutSizeAndInvSize)
-	SHADER_PARAMETER(FVector4, SkyViewLutSizeAndInvSize)
+	SHADER_PARAMETER(FVector4, SkyViewLutSizeAndInvSize)						// Also on View UB
 
 	SHADER_PARAMETER(float, TransmittanceSampleCount)
 	SHADER_PARAMETER(float, MultiScatteringSampleCount)
