Remove SM4

diff --git a/Engine/Source/Editor/LandscapeEditor/Private/LandscapeEdMode.cpp b/Engine/Source/Editor/LandscapeEditor/Private/LandscapeEdMode.cpp
--- a/Engine/Source/Editor/LandscapeEditor/Private/LandscapeEdMode.cpp
+++ b/Engine/Source/Editor/LandscapeEditor/Private/LandscapeEdMode.cpp
@@ -4340,7 +4340,7 @@ ELandscapeEditingState FEdModeLandscape::GetEditingState() const
 	{
 		return ELandscapeEditingState::Unknown;
 	}
-	else if (World->FeatureLevel < ERHIFeatureLevel::SM4)
+	else if (World->FeatureLevel < ERHIFeatureLevel::SM5)
 	{
 		return ELandscapeEditingState::BadFeatureLevel;
 	}
diff --git a/Engine/Source/Editor/LevelEditor/Private/LevelEditor.cpp b/Engine/Source/Editor/LevelEditor/Private/LevelEditor.cpp
--- a/Engine/Source/Editor/LevelEditor/Private/LevelEditor.cpp
+++ b/Engine/Source/Editor/LevelEditor/Private/LevelEditor.cpp
@@ -1762,12 +1762,6 @@ void FLevelEditorModule::BindGlobalLevelEditorCommands()
 		FCanExecuteAction(),
 		FIsActionChecked::CreateStatic(&FLevelEditorActionCallbacks::IsPreviewPlatformChecked, FPreviewPlatformInfo(ERHIFeatureLevel::SM5, NAME_None)));
 
-	ActionList.MapAction(
-		Commands.PreviewPlatformOverride_SM4,
-		FExecuteAction::CreateStatic(&FLevelEditorActionCallbacks::SetPreviewPlatform, FPreviewPlatformInfo(ERHIFeatureLevel::SM4, NAME_None, true)),
-		FCanExecuteAction(),
-		FIsActionChecked::CreateStatic(&FLevelEditorActionCallbacks::IsPreviewPlatformChecked, FPreviewPlatformInfo(ERHIFeatureLevel::SM4, NAME_None)));
-
 	ActionList.MapAction(
 		Commands.PreviewPlatformOverride_HTML5,
 		FExecuteAction::CreateStatic(&FLevelEditorActionCallbacks::SetPreviewPlatform, FPreviewPlatformInfo(ERHIFeatureLevel::ES2, NAME_None, true)),
diff --git a/Engine/Source/Editor/LevelEditor/Private/LevelEditorActions.cpp b/Engine/Source/Editor/LevelEditor/Private/LevelEditorActions.cpp
--- a/Engine/Source/Editor/LevelEditor/Private/LevelEditorActions.cpp
+++ b/Engine/Source/Editor/LevelEditor/Private/LevelEditorActions.cpp
@@ -3410,7 +3410,6 @@ void FLevelEditorCommands::RegisterCommands()
 	UI_COMMAND(ToggleFeatureLevelPreview, "Preview Mode Toggle", "Toggles the Preview Mode on or off for the currently selected Preview target", EUserInterfaceActionType::ToggleButton, FInputChord());
 
 	UI_COMMAND(PreviewPlatformOverride_SM5, "Shader Model 5", "DirectX 11, OpenGL 4.3+, PS4, XB1", EUserInterfaceActionType::Check, FInputChord());
-	UI_COMMAND(PreviewPlatformOverride_SM4, "Shader Model 4", "DirectX 10, OpenGL 3.3+", EUserInterfaceActionType::Check, FInputChord());
 	UI_COMMAND(PreviewPlatformOverride_HTML5, "HTML5", "HTML5 preview.", EUserInterfaceActionType::Check, FInputChord());
 	UI_COMMAND(PreviewPlatformOverride_AndroidGLES2, "Android ES2", "Mobile preview using Android's quality settings.", EUserInterfaceActionType::Check, FInputChord());
 	UI_COMMAND(PreviewPlatformOverride_AndroidGLES31, "Android ES 3.1", "Mobile preview using Android ES3.1 quality settings.", EUserInterfaceActionType::Check, FInputChord());
diff --git a/Engine/Source/Editor/LevelEditor/Private/LevelEditorToolBar.cpp b/Engine/Source/Editor/LevelEditor/Private/LevelEditorToolBar.cpp
--- a/Engine/Source/Editor/LevelEditor/Private/LevelEditorToolBar.cpp
+++ b/Engine/Source/Editor/LevelEditor/Private/LevelEditorToolBar.cpp
@@ -1293,10 +1293,6 @@ TSharedRef< SWidget > FLevelEditorToolBar::MakeLevelEditorToolBar( const TShared
 
 				switch (GEditor->PreviewPlatform.PreviewFeatureLevel)
 				{
-					case ERHIFeatureLevel::SM4:
-					{
-						return LOCTEXT("PreviewModeSM4_Text", "SM4 Preview");
-					}
 					case ERHIFeatureLevel::ES2:
 					{
 						return LOCTEXT("PreviewModeES2_Text", "ES2 Preview");
@@ -1367,10 +1363,6 @@ TSharedRef< SWidget > FLevelEditorToolBar::MakeLevelEditorToolBar( const TShared
 				}
 				switch (GEditor->PreviewPlatform.PreviewFeatureLevel)
 				{
-					case ERHIFeatureLevel::SM4:
-					{
-						return FSlateIcon(FEditorStyle::GetStyleSetName(), GEditor->IsFeatureLevelPreviewActive() ? "LevelEditor.PreviewMode.SM4.Enabled" : "LevelEditor.PreviewMode.SM4.Disabled");
-					}
 					case ERHIFeatureLevel::ES2:
 					{
 						return FSlateIcon(FEditorStyle::GetStyleSetName(), GEditor->IsFeatureLevelPreviewActive() ? "LevelEditor.PreviewMode.AndroidES2.Enabled" : "LevelEditor.PreviewMode.AndroidES2.Disabled");
@@ -1816,9 +1808,6 @@ static void MakeShaderModelPreviewMenu( FMenuBuilder& MenuBuilder )
 	// SM5
 	MenuBuilder.AddMenuEntry(FLevelEditorCommands::Get().PreviewPlatformOverride_SM5);
 
-	// SM4
-	MenuBuilder.AddMenuEntry(FLevelEditorCommands::Get().PreviewPlatformOverride_SM4);
-
 	// Android
 	bool bAndroidBuildForES31 = false;
 	GConfig->GetBool(TEXT("/Script/AndroidRuntimeSettings.AndroidRuntimeSettings"), TEXT("bBuildForES31"), bAndroidBuildForES31, GEngineIni);
diff --git a/Engine/Source/Editor/LevelEditor/Public/LevelEditorActions.h b/Engine/Source/Editor/LevelEditor/Public/LevelEditorActions.h
--- a/Engine/Source/Editor/LevelEditor/Public/LevelEditorActions.h
+++ b/Engine/Source/Editor/LevelEditor/Public/LevelEditorActions.h
@@ -593,7 +593,6 @@ public:
 	TSharedPtr< FUICommandInfo > ToggleFeatureLevelPreview;
 
 	TSharedPtr< FUICommandInfo > PreviewPlatformOverride_SM5;
-	TSharedPtr< FUICommandInfo > PreviewPlatformOverride_SM4;
 	TSharedPtr< FUICommandInfo > PreviewPlatformOverride_HTML5;
 	TSharedPtr< FUICommandInfo > PreviewPlatformOverride_AndroidGLES2;
 	TSharedPtr< FUICommandInfo > PreviewPlatformOverride_AndroidGLES31;
diff --git a/Engine/Source/Editor/MaterialEditor/Private/MaterialEditor.cpp b/Engine/Source/Editor/MaterialEditor/Private/MaterialEditor.cpp
--- a/Engine/Source/Editor/MaterialEditor/Private/MaterialEditor.cpp
+++ b/Engine/Source/Editor/MaterialEditor/Private/MaterialEditor.cpp
@@ -2507,11 +2507,6 @@ void FMaterialEditor::BindCommands()
 		FExecuteAction::CreateSP(this, &FMaterialEditor::SetFeaturePreview, ERHIFeatureLevel::ES3_1),
 		FCanExecuteAction(),
 		FIsActionChecked::CreateSP(this, &FMaterialEditor::IsFeaturePreviewChecked, ERHIFeatureLevel::ES3_1));
-	ToolkitCommands->MapAction(
-		Commands.FeatureLevel_SM4,
-		FExecuteAction::CreateSP(this, &FMaterialEditor::SetFeaturePreview, ERHIFeatureLevel::SM4),
-		FCanExecuteAction(),
-		FIsActionChecked::CreateSP(this, &FMaterialEditor::IsFeaturePreviewChecked, ERHIFeatureLevel::SM4));
 	ToolkitCommands->MapAction(
 		Commands.FeatureLevel_SM5,
 		FExecuteAction::CreateSP(this, &FMaterialEditor::SetFeaturePreview, ERHIFeatureLevel::SM5),
diff --git a/Engine/Source/Editor/MaterialEditor/Private/MaterialStatsCommon.cpp b/Engine/Source/Editor/MaterialEditor/Private/MaterialStatsCommon.cpp
--- a/Engine/Source/Editor/MaterialEditor/Private/MaterialStatsCommon.cpp
+++ b/Engine/Source/Editor/MaterialEditor/Private/MaterialStatsCommon.cpp
@@ -377,7 +377,7 @@ void FMaterialStatsUtils::GetRepresentativeShaderTypesAndDescriptions(TMap<FName
 		ShaderTypeNamesAndDescriptions.FindOrAdd(FLocalVertexFactoryName)
 			.Add(FRepresentativeShaderInfo(ERepresentativeShader::UIInstancedVertexShader, TSlateMaterialShaderVStrueName, TEXT("Instanced UI Vertex Shader")));
 	}
-	else if (TargetMaterial->GetFeatureLevel() >= ERHIFeatureLevel::SM4)
+	else if (TargetMaterial->GetFeatureLevel() >= ERHIFeatureLevel::SM5)
 	{
 		if (TargetMaterial->GetShadingModels().IsUnlit())
 		{
diff --git a/Engine/Source/Editor/TextureEditor/Private/Models/TextureEditorViewportClient.cpp b/Engine/Source/Editor/TextureEditor/Private/Models/TextureEditorViewportClient.cpp
--- a/Engine/Source/Editor/TextureEditor/Private/Models/TextureEditorViewportClient.cpp
+++ b/Engine/Source/Editor/TextureEditor/Private/Models/TextureEditorViewportClient.cpp
@@ -99,7 +99,7 @@ void FTextureEditorViewportClient::Draw(FViewport* Viewport, FCanvas* Canvas)
 
 	TRefCountPtr<FBatchedElementParameters> BatchedElementParameters;
 
-	if (GMaxRHIFeatureLevel >= ERHIFeatureLevel::SM4)
+	if (GMaxRHIFeatureLevel >= ERHIFeatureLevel::SM5)
 	{
 		if (TextureCube || RTTextureCube)
 		{
diff --git a/Engine/Source/Editor/UnrealEd/Private/MeshPaintRendering.cpp b/Engine/Source/Editor/UnrealEd/Private/MeshPaintRendering.cpp
--- a/Engine/Source/Editor/UnrealEd/Private/MeshPaintRendering.cpp
+++ b/Engine/Source/Editor/UnrealEd/Private/MeshPaintRendering.cpp
@@ -29,7 +29,7 @@ namespace MeshPaintRendering
 
 		static bool ShouldCompilePermutation(const FGlobalShaderPermutationParameters& Parameters)
 		{
-			return IsFeatureLevelSupported(Parameters.Platform, ERHIFeatureLevel::SM4) && !IsConsolePlatform(Parameters.Platform);
+			return IsFeatureLevelSupported(Parameters.Platform, ERHIFeatureLevel::SM5) && !IsConsolePlatform(Parameters.Platform);
 		}
 
 		/** Default constructor. */
@@ -73,7 +73,7 @@ namespace MeshPaintRendering
 
 		static bool ShouldCompilePermutation(const FGlobalShaderPermutationParameters& Parameters)
 		{
-			return IsFeatureLevelSupported(Parameters.Platform, ERHIFeatureLevel::SM4) && !IsConsolePlatform(Parameters.Platform);
+			return IsFeatureLevelSupported(Parameters.Platform, ERHIFeatureLevel::SM5) && !IsConsolePlatform(Parameters.Platform);
 		}
 
 		/** Default constructor. */
@@ -192,7 +192,7 @@ namespace MeshPaintRendering
 
 		static bool ShouldCompilePermutation(const FGlobalShaderPermutationParameters& Parameters)
 		{
-			return IsFeatureLevelSupported(Parameters.Platform, ERHIFeatureLevel::SM4) && !IsConsolePlatform(Parameters.Platform);
+			return IsFeatureLevelSupported(Parameters.Platform, ERHIFeatureLevel::SM5) && !IsConsolePlatform(Parameters.Platform);
 		}
 
 		/** Default constructor. */
@@ -236,7 +236,7 @@ namespace MeshPaintRendering
 
 		static bool ShouldCompilePermutation(const FGlobalShaderPermutationParameters& Parameters)
 		{
-			return IsFeatureLevelSupported(Parameters.Platform, ERHIFeatureLevel::SM4) && !IsConsolePlatform(Parameters.Platform);
+			return IsFeatureLevelSupported(Parameters.Platform, ERHIFeatureLevel::SM5) && !IsConsolePlatform(Parameters.Platform);
 		}
 
 		/** Default constructor. */
diff --git a/Engine/Source/Editor/UnrealEd/Private/SEditorViewportViewMenu.cpp b/Engine/Source/Editor/UnrealEd/Private/SEditorViewportViewMenu.cpp
--- a/Engine/Source/Editor/UnrealEd/Private/SEditorViewportViewMenu.cpp
+++ b/Engine/Source/Editor/UnrealEd/Private/SEditorViewportViewMenu.cpp
@@ -321,7 +321,7 @@ TSharedRef<SWidget> SEditorViewportViewMenu::GenerateViewMenuContent() const
 
 						Menu.BeginSection("OptimizationViewmodes", LOCTEXT("OptimizationSubMenuHeader", "Optimization Viewmodes"));
 						{
-							if (FeatureLevel == ERHIFeatureLevel::SM5 || FeatureLevel == ERHIFeatureLevel::SM4)
+							if (FeatureLevel == ERHIFeatureLevel::SM5)
 							{
 								Menu.AddMenuEntry(BaseViewportCommands.LightComplexityMode, NAME_None, LOCTEXT("LightComplexityViewModeDisplayName", "Light Complexity"));
 								Menu.AddMenuEntry(BaseViewportCommands.LightmapDensityMode, NAME_None, LOCTEXT("LightmapDensityViewModeDisplayName", "Lightmap Density"));
diff --git a/Engine/Source/Editor/UnrealEd/Private/Texture2DPreview.cpp b/Engine/Source/Editor/UnrealEd/Private/Texture2DPreview.cpp
--- a/Engine/Source/Editor/UnrealEd/Private/Texture2DPreview.cpp
+++ b/Engine/Source/Editor/UnrealEd/Private/Texture2DPreview.cpp
@@ -57,7 +57,7 @@ public:
 		{
 			return false;
 		}
-		return IsFeatureLevelSupported(Parameters.Platform, ERHIFeatureLevel::SM4) && !IsConsolePlatform(Parameters.Platform);
+		return IsFeatureLevelSupported(Parameters.Platform, ERHIFeatureLevel::SM5) && !IsConsolePlatform(Parameters.Platform);
 	}
 
 	void SetParameters(FRHICommandList& RHICmdList, const FTexture* TextureValue, const FMatrix& ColorWeightsValue, float GammaValue, float MipLevel, float LayerIndex, bool bIsNormalMap, bool bIsVirtualTexture)
diff --git a/Engine/Source/Editor/UnrealEd/Private/VolumeTexturePreview.cpp b/Engine/Source/Editor/UnrealEd/Private/VolumeTexturePreview.cpp
--- a/Engine/Source/Editor/UnrealEd/Private/VolumeTexturePreview.cpp
+++ b/Engine/Source/Editor/UnrealEd/Private/VolumeTexturePreview.cpp
@@ -68,7 +68,7 @@ public:
 	/** Should the shader be cached? Always. */
 	static bool ShouldCompilePermutation(const FGlobalShaderPermutationParameters& Parameters)
 	{
-		return IsFeatureLevelSupported(Parameters.Platform, ERHIFeatureLevel::SM4) && !IsConsolePlatform(Parameters.Platform);
+		return IsFeatureLevelSupported(Parameters.Platform, ERHIFeatureLevel::SM5) && !IsConsolePlatform(Parameters.Platform);
 	}
 	
 	void SetParameters(FRHICommandList& RHICmdList, const FTexture* TextureValue, int32 SizeZ, const FMatrix& ColorWeightsValue, float GammaValue, float MipLevel, float Opacity, const FRotator& TraceOrientation)
diff --git a/Engine/Source/Runtime/Apple/MetalRHI/Private/MetalCommands.cpp b/Engine/Source/Runtime/Apple/MetalRHI/Private/MetalCommands.cpp
--- a/Engine/Source/Runtime/Apple/MetalRHI/Private/MetalCommands.cpp
+++ b/Engine/Source/Runtime/Apple/MetalRHI/Private/MetalCommands.cpp
@@ -896,7 +896,7 @@ void FMetalDynamicRHI::SetupRecursiveResources()
 					TShaderMapRef<TOneColorPixelShaderMRT<1> > MRTPixelShader(ShaderMap);
 					PixelShader = *MRTPixelShader;
 				}
-				else if (IsFeatureLevelSupported( GMaxRHIShaderPlatform, ERHIFeatureLevel::SM4 ))
+				else if (IsFeatureLevelSupported( GMaxRHIShaderPlatform, ERHIFeatureLevel::SM5 ))
 				{
 					if (NumBuffers == 2)
 					{
diff --git a/Engine/Source/Runtime/Apple/MetalRHI/Private/MetalQuery.cpp b/Engine/Source/Runtime/Apple/MetalRHI/Private/MetalQuery.cpp
--- a/Engine/Source/Runtime/Apple/MetalRHI/Private/MetalQuery.cpp
+++ b/Engine/Source/Runtime/Apple/MetalRHI/Private/MetalQuery.cpp
@@ -179,7 +179,7 @@ void FMetalRenderQuery::Begin(FMetalContext* Context, TSharedPtr<FMetalCommandBu
 			Context->GetQueryBufferPool()->Allocate(Buffer);
 			Buffer.bCompleted = false;
 			
-			if ((GMaxRHIFeatureLevel >= ERHIFeatureLevel::SM4) && GetMetalDeviceContext().SupportsFeature(EMetalFeaturesCountingQueries))
+			if ((GMaxRHIFeatureLevel >= ERHIFeatureLevel::SM5) && GetMetalDeviceContext().SupportsFeature(EMetalFeaturesCountingQueries))
 			{
 				Context->GetCurrentState().SetVisibilityResultMode(mtlpp::VisibilityResultMode::Counting, Buffer.Offset);
 			}
diff --git a/Engine/Source/Runtime/Apple/MetalRHI/Private/MetalRHI.cpp b/Engine/Source/Runtime/Apple/MetalRHI/Private/MetalRHI.cpp
--- a/Engine/Source/Runtime/Apple/MetalRHI/Private/MetalRHI.cpp
+++ b/Engine/Source/Runtime/Apple/MetalRHI/Private/MetalRHI.cpp
@@ -155,7 +155,7 @@ FMetalDynamicRHI::FMetalDynamicRHI(ERHIFeatureLevel::Type RequestedFeatureLevel)
     bool bProjectSupportsMRTs = false;
     GConfig->GetBool(TEXT("/Script/IOSRuntimeSettings.IOSRuntimeSettings"), TEXT("bSupportsMetalMRT"), bProjectSupportsMRTs, GEngineIni);
 
-	bool const bRequestedMetalMRT = ((RequestedFeatureLevel >= ERHIFeatureLevel::SM4) || (!bRequestedFeatureLevel && FParse::Param(FCommandLine::Get(),TEXT("metalmrt"))));
+	bool const bRequestedMetalMRT = ((RequestedFeatureLevel >= ERHIFeatureLevel::SM5) || (!bRequestedFeatureLevel && FParse::Param(FCommandLine::Get(),TEXT("metalmrt"))));
 	bSupportsRHIThread = FParse::Param(FCommandLine::Get(),TEXT("rhithread"));
 
     // only allow GBuffers, etc on A8s (A7s are just not going to cut it)
@@ -201,8 +201,8 @@ FMetalDynamicRHI::FMetalDynamicRHI(ERHIFeatureLevel::Type RequestedFeatureLevel)
 	GShaderPlatformForFeatureLevel[ERHIFeatureLevel::ES2] = SP_METAL;
 	GShaderPlatformForFeatureLevel[ERHIFeatureLevel::ES3_1] = SP_METAL;
 #endif
-	GShaderPlatformForFeatureLevel[ERHIFeatureLevel::SM4] = (GMaxRHIFeatureLevel >= ERHIFeatureLevel::SM4) ? GMaxRHIShaderPlatform : SP_NumPlatforms;
-	GShaderPlatformForFeatureLevel[ERHIFeatureLevel::SM5] = (GMaxRHIFeatureLevel >= ERHIFeatureLevel::SM4) ? GMaxRHIShaderPlatform : SP_NumPlatforms;
+	GShaderPlatformForFeatureLevel[ERHIFeatureLevel::SM4_REMOVED] = SP_NumPlatforms;
+	GShaderPlatformForFeatureLevel[ERHIFeatureLevel::SM5] = (GMaxRHIFeatureLevel >= ERHIFeatureLevel::SM5) ? GMaxRHIShaderPlatform : SP_NumPlatforms;
 
 #else // @todo zebra
 	uint32 DeviceIndex = ((FMetalDeviceContext*)ImmediateContext.Context)->GetDeviceIndex();
@@ -304,7 +304,7 @@ FMetalDynamicRHI::FMetalDynamicRHI(ERHIFeatureLevel::Type RequestedFeatureLevel)
 	
 	GShaderPlatformForFeatureLevel[ERHIFeatureLevel::ES2] = SP_METAL_MACES2;
 	GShaderPlatformForFeatureLevel[ERHIFeatureLevel::ES3_1] = (GMaxRHIFeatureLevel >= ERHIFeatureLevel::ES3_1) ? SP_METAL_MACES3_1 : SP_NumPlatforms;
-	GShaderPlatformForFeatureLevel[ERHIFeatureLevel::SM4] = SP_NumPlatforms;
+	GShaderPlatformForFeatureLevel[ERHIFeatureLevel::SM4_REMOVED] = SP_NumPlatforms;
 	GShaderPlatformForFeatureLevel[ERHIFeatureLevel::SM5] = (GMaxRHIFeatureLevel >= ERHIFeatureLevel::SM5) ? GMaxRHIShaderPlatform : SP_NumPlatforms;
 	
 	// Mac GPUs support layer indexing.
@@ -476,7 +476,7 @@ FMetalDynamicRHI::FMetalDynamicRHI(ERHIFeatureLevel::Type RequestedFeatureLevel)
 	GRHISupportsTextureStreaming = true;
 	GSupportsWideMRT = bCanUseWideMRTs;
 	// GSupportsTransientResourceAliasing = FMetalCommandQueue::SupportsFeature(EMetalFeaturesHeaps) && FMetalCommandQueue::SupportsFeature(EMetalFeaturesFences);
-	GSupportsSeparateRenderTargetBlendState = (GMaxRHIFeatureLevel >= ERHIFeatureLevel::SM4);
+	GSupportsSeparateRenderTargetBlendState = (GMaxRHIFeatureLevel >= ERHIFeatureLevel::SM5);
 
 #if PLATFORM_MAC
 	check(Device.SupportsFeatureSet(mtlpp::FeatureSet::macOS_GPUFamily1_v1));
diff --git a/Engine/Source/Runtime/Apple/MetalRHI/Private/MetalRenderPass.cpp b/Engine/Source/Runtime/Apple/MetalRHI/Private/MetalRenderPass.cpp
--- a/Engine/Source/Runtime/Apple/MetalRHI/Private/MetalRenderPass.cpp
+++ b/Engine/Source/Runtime/Apple/MetalRHI/Private/MetalRenderPass.cpp
@@ -2452,7 +2452,7 @@ void FMetalRenderPass::ConditionalSubmit()
 		{
 			bool bSet = false;
 			State.InvalidateRenderTargets();
-			if (IsFeatureLevelSupported( GMaxRHIShaderPlatform, ERHIFeatureLevel::SM4 ))
+			if (IsFeatureLevelSupported( GMaxRHIShaderPlatform, ERHIFeatureLevel::SM5 ))
 			{
 				bSet = State.SetRenderPassInfo(CurrentRenderTargets, State.GetVisibilityResultsBuffer(), false);
 			}
diff --git a/Engine/Source/Runtime/Core/Public/UObject/RenderingObjectVersion.h b/Engine/Source/Runtime/Core/Public/UObject/RenderingObjectVersion.h
--- a/Engine/Source/Runtime/Core/Public/UObject/RenderingObjectVersion.h
+++ b/Engine/Source/Runtime/Core/Public/UObject/RenderingObjectVersion.h
@@ -94,6 +94,9 @@ struct CORE_API FRenderingObjectVersion
 
 		VolumetricLightmapStreaming,
 
+		//ShaderModel4 support removed from engine
+		RemovedSM4,
+
 		// -----<new versions can be added above this line>-------------------------------------------------
 		VersionPlusOne,
 		LatestVersion = VersionPlusOne - 1
diff --git a/Engine/Source/Runtime/D3D12RHI/Private/D3D12RHI.cpp b/Engine/Source/Runtime/D3D12RHI/Private/D3D12RHI.cpp
--- a/Engine/Source/Runtime/D3D12RHI/Private/D3D12RHI.cpp
+++ b/Engine/Source/Runtime/D3D12RHI/Private/D3D12RHI.cpp
@@ -476,14 +476,6 @@ void FD3D12DynamicRHI::RHIGetSupportedResolution(uint32& Width, uint32& Height)
 
 void FD3D12DynamicRHI::GetBestSupportedMSAASetting(DXGI_FORMAT PlatformFormat, uint32 MSAACount, uint32& OutBestMSAACount, uint32& OutMSAAQualityLevels)
 {
-	//  We disable MSAA for Feature level 10
-	if (GMaxRHIFeatureLevel == ERHIFeatureLevel::SM4)
-	{
-		OutBestMSAACount = 1;
-		OutMSAAQualityLevels = 0;
-		return;
-	}
-
 	// start counting down from current setting (indicated the current "best" count) and move down looking for support
 	for (uint32 SampleCount = MSAACount; SampleCount > 0; SampleCount--)
 	{
diff --git a/Engine/Source/Runtime/D3D12RHI/Private/Windows/WindowsD3D12Device.cpp b/Engine/Source/Runtime/D3D12RHI/Private/Windows/WindowsD3D12Device.cpp
--- a/Engine/Source/Runtime/D3D12RHI/Private/Windows/WindowsD3D12Device.cpp
+++ b/Engine/Source/Runtime/D3D12RHI/Private/Windows/WindowsD3D12Device.cpp
@@ -444,11 +444,6 @@ FDynamicRHI* FD3D12DynamicRHIModule::CreateRHI(ERHIFeatureLevel::Type RequestedF
 			GMaxRHIShaderPlatform = SP_PCD3D_ES3_1;
 		}
 	}
-	else if (RequestedFeatureLevel == ERHIFeatureLevel::SM4)
-	{
-		GMaxRHIFeatureLevel = ERHIFeatureLevel::SM4;
-		GMaxRHIShaderPlatform = SP_PCD3D_SM4;
-	}
 	else
 	{
 		GMaxRHIFeatureLevel = ERHIFeatureLevel::SM5;
@@ -647,7 +642,7 @@ void FD3D12DynamicRHI::Init()
 
 	GShaderPlatformForFeatureLevel[ERHIFeatureLevel::ES2] = SP_PCD3D_ES2;
 	GShaderPlatformForFeatureLevel[ERHIFeatureLevel::ES3_1] = SP_PCD3D_ES3_1;
-	GShaderPlatformForFeatureLevel[ERHIFeatureLevel::SM4] = SP_PCD3D_SM4;
+	GShaderPlatformForFeatureLevel[ERHIFeatureLevel::SM4_REMOVED] = SP_NumPlatforms;
 	GShaderPlatformForFeatureLevel[ERHIFeatureLevel::SM5] = SP_PCD3D_SM5;
 
 	GSupportsEfficientAsyncCompute = GRHISupportsParallelRHIExecute && IsRHIDeviceAMD();
diff --git a/Engine/Source/Runtime/Engine/Private/Atmosphere/Atmosphere.h b/Engine/Source/Runtime/Engine/Private/Atmosphere/Atmosphere.h
--- a/Engine/Source/Runtime/Engine/Private/Atmosphere/Atmosphere.h
+++ b/Engine/Source/Runtime/Engine/Private/Atmosphere/Atmosphere.h
@@ -138,7 +138,7 @@ public:
 	 */
 	virtual void InitRHI() override
 	{
-		if (TextureData && GetFeatureLevel() >= ERHIFeatureLevel::SM4)
+		if (TextureData && GetFeatureLevel() >= ERHIFeatureLevel::SM5)
 		{
 			switch(TexType)
 			{
diff --git a/Engine/Source/Runtime/Engine/Private/Components/LightComponent.cpp b/Engine/Source/Runtime/Engine/Private/Components/LightComponent.cpp
--- a/Engine/Source/Runtime/Engine/Private/Components/LightComponent.cpp
+++ b/Engine/Source/Runtime/Engine/Private/Components/LightComponent.cpp
@@ -27,7 +27,7 @@
 
 void FStaticShadowDepthMap::InitRHI()
 {
-	if (FApp::CanEverRender() && Data && Data->ShadowMapSizeX > 0 && Data->ShadowMapSizeY > 0 && GMaxRHIFeatureLevel >= ERHIFeatureLevel::SM4)
+	if (FApp::CanEverRender() && Data && Data->ShadowMapSizeX > 0 && Data->ShadowMapSizeY > 0 && GMaxRHIFeatureLevel >= ERHIFeatureLevel::SM5)
 	{
 		FRHIResourceCreateInfo CreateInfo;
 		FTexture2DRHIRef Texture2DRHI = RHICreateTexture2D(Data->ShadowMapSizeX, Data->ShadowMapSizeY, PF_R16F, 1, 1, 0, CreateInfo);
diff --git a/Engine/Source/Runtime/Engine/Private/Components/PointLightComponent.cpp b/Engine/Source/Runtime/Engine/Private/Components/PointLightComponent.cpp
--- a/Engine/Source/Runtime/Engine/Private/Components/PointLightComponent.cpp
+++ b/Engine/Source/Runtime/Engine/Private/Components/PointLightComponent.cpp
@@ -54,7 +54,7 @@ void FPointLightSceneProxy::GetLightShaderParameters(FLightShaderParameters& Lig
 */
 bool FPointLightSceneProxy::GetWholeSceneProjectedShadowInitializer(const FSceneViewFamily& ViewFamily, TArray<FWholeSceneProjectedShadowInitializer, TInlineAllocator<6> >& OutInitializers) const
 {
-	if (ViewFamily.GetFeatureLevel() >= ERHIFeatureLevel::SM4
+	if (ViewFamily.GetFeatureLevel() >= ERHIFeatureLevel::SM5
 		&& GAllowPointLightCubemapShadows != 0)
 	{
 		FWholeSceneProjectedShadowInitializer& OutInitializer = *new(OutInitializers) FWholeSceneProjectedShadowInitializer;
diff --git a/Engine/Source/Runtime/Engine/Private/Components/RectLightComponent.cpp b/Engine/Source/Runtime/Engine/Private/Components/RectLightComponent.cpp
--- a/Engine/Source/Runtime/Engine/Private/Components/RectLightComponent.cpp
+++ b/Engine/Source/Runtime/Engine/Private/Components/RectLightComponent.cpp
@@ -246,7 +246,7 @@ void FRectLightSceneProxy::GetLightShaderParameters(FLightShaderParameters& Ligh
 */
 bool FRectLightSceneProxy::GetWholeSceneProjectedShadowInitializer(const FSceneViewFamily& ViewFamily, TArray<FWholeSceneProjectedShadowInitializer, TInlineAllocator<6> >& OutInitializers) const
 {
-	if (ViewFamily.GetFeatureLevel() >= ERHIFeatureLevel::SM4
+	if (ViewFamily.GetFeatureLevel() >= ERHIFeatureLevel::SM5
 		&& GAllowPointLightCubemapShadows != 0)
 	{
 		FWholeSceneProjectedShadowInitializer& OutInitializer = *new(OutInitializers) FWholeSceneProjectedShadowInitializer;
diff --git a/Engine/Source/Runtime/Engine/Private/Components/SkyLightComponent.cpp b/Engine/Source/Runtime/Engine/Private/Components/SkyLightComponent.cpp
--- a/Engine/Source/Runtime/Engine/Private/Components/SkyLightComponent.cpp
+++ b/Engine/Source/Runtime/Engine/Private/Components/SkyLightComponent.cpp
@@ -71,7 +71,7 @@ constexpr EPixelFormat SKYLIGHT_CUBEMAP_FORMAT = PF_FloatRGBA;
 
 void FSkyTextureCubeResource::InitRHI()
 {
-	if (GetFeatureLevel() >= ERHIFeatureLevel::SM4 || GSupportsRenderTargetFormat_PF_FloatRGBA)
+	if (GetFeatureLevel() >= ERHIFeatureLevel::SM5 || GSupportsRenderTargetFormat_PF_FloatRGBA)
 	{
 		FRHIResourceCreateInfo CreateInfo;
 		
diff --git a/Engine/Source/Runtime/Engine/Private/DebugViewModeHelpers.cpp b/Engine/Source/Runtime/Engine/Private/DebugViewModeHelpers.cpp
--- a/Engine/Source/Runtime/Engine/Private/DebugViewModeHelpers.cpp
+++ b/Engine/Source/Runtime/Engine/Private/DebugViewModeHelpers.cpp
@@ -62,7 +62,7 @@ bool AllowDebugViewShaderMode(EDebugViewShaderMode ShaderMode, EShaderPlatform P
 		return false;
 	}
 #else
-	return ShaderMode == DVSM_ShaderComplexity && FeatureLevel != ERHIFeatureLevel::SM4;
+	return ShaderMode == DVSM_ShaderComplexity;
 #endif
 }
 
diff --git a/Engine/Source/Runtime/Engine/Private/GPUSkinVertexFactory.cpp b/Engine/Source/Runtime/Engine/Private/GPUSkinVertexFactory.cpp
--- a/Engine/Source/Runtime/Engine/Private/GPUSkinVertexFactory.cpp
+++ b/Engine/Source/Runtime/Engine/Private/GPUSkinVertexFactory.cpp
@@ -76,7 +76,7 @@ static TAutoConsoleVariable<int32> CVarVelocityTest(
 // These should match USE_BONES_SRV_BUFFER
 static inline bool SupportsBonesBufferSRV(EShaderPlatform Platform)
 {
-	return IsFeatureLevelSupported(Platform, ERHIFeatureLevel::SM4) || IsVulkanPlatform(Platform) || IsMetalPlatform(Platform);
+	return IsFeatureLevelSupported(Platform, ERHIFeatureLevel::SM5) || IsVulkanPlatform(Platform) || IsMetalPlatform(Platform);
 }
 
 static inline bool SupportsBonesBufferSRV(ERHIFeatureLevel::Type InFeatureLevel)
@@ -895,7 +895,7 @@ bool FGPUBaseSkinAPEXClothVertexFactory::ClothShaderType::UpdateClothSimulData(F
 
 	FVertexBufferAndSRV* CurrentClothBuffer = 0;
 
-	if (FeatureLevel >= ERHIFeatureLevel::SM4)
+	if (FeatureLevel >= ERHIFeatureLevel::SM5)
 	{
 		check(IsInRenderingThread());
 		
@@ -988,7 +988,7 @@ void TGPUSkinAPEXClothVertexFactory<bExtraBoneInfluencesT>::ModifyCompilationEnv
 template <bool bExtraBoneInfluencesT>
 bool TGPUSkinAPEXClothVertexFactory<bExtraBoneInfluencesT>::ShouldCompilePermutation(EShaderPlatform Platform, const class FMaterial* Material, const class FShaderType* ShaderType)
 {
-	return GetMaxSupportedFeatureLevel(Platform) >= ERHIFeatureLevel::SM4
+	return GetMaxSupportedFeatureLevel(Platform) >= ERHIFeatureLevel::SM5
 		&& (Material->IsUsedWithAPEXCloth() || Material->IsSpecialEngineMaterial()) 
 		&& Super::ShouldCompilePermutation(Platform, Material, ShaderType);
 }
diff --git a/Engine/Source/Runtime/Engine/Private/GeometryCacheVertexFactory.cpp b/Engine/Source/Runtime/Engine/Private/GeometryCacheVertexFactory.cpp
--- a/Engine/Source/Runtime/Engine/Private/GeometryCacheVertexFactory.cpp
+++ b/Engine/Source/Runtime/Engine/Private/GeometryCacheVertexFactory.cpp
@@ -364,7 +364,7 @@ FVertexFactoryShaderParameters* FGeometryCacheVertexVertexFactory::ConstructShad
 bool FGeometryCacheVertexVertexFactory::ShouldCompilePermutation(EShaderPlatform Platform, const class FMaterial* Material, const class FShaderType* ShaderType)
 {
 	// Should this be platform or mesh type based? Returning true should work in all cases, but maybe too expensive? 
-	// return IsFeatureLevelSupported(Platform, ERHIFeatureLevel::SM4) && !IsConsolePlatform(Platform);
+	// return IsFeatureLevelSupported(Platform, ERHIFeatureLevel::SM5) && !IsConsolePlatform(Platform);
 	// TODO currently GeomCache supports only 4 UVs which could cause compilation errors when trying to compile shaders which use > 4
 	return Material->IsUsedWithGeometryCache() || Material->IsSpecialEngineMaterial();
 }
diff --git a/Engine/Source/Runtime/Engine/Private/InstancedStaticMesh.h b/Engine/Source/Runtime/Engine/Private/InstancedStaticMesh.h
--- a/Engine/Source/Runtime/Engine/Private/InstancedStaticMesh.h
+++ b/Engine/Source/Runtime/Engine/Private/InstancedStaticMesh.h
@@ -205,7 +205,7 @@ public:
 		}
 
 		OutEnvironment.SetDefine(TEXT("USE_INSTANCING"),TEXT("1"));
-		if (IsFeatureLevelSupported(Platform, ERHIFeatureLevel::SM4))
+		if (IsFeatureLevelSupported(Platform, ERHIFeatureLevel::SM5))
 		{
 			OutEnvironment.SetDefine(TEXT("USE_DITHERED_LOD_TRANSITION_FOR_INSTANCED"), ALLOW_DITHERED_LOD_FOR_INSTANCED_STATIC_MESHES);
 		}
diff --git a/Engine/Source/Runtime/Engine/Private/MapBuildData.cpp b/Engine/Source/Runtime/Engine/Private/MapBuildData.cpp
--- a/Engine/Source/Runtime/Engine/Private/MapBuildData.cpp
+++ b/Engine/Source/Runtime/Engine/Private/MapBuildData.cpp
@@ -404,9 +404,9 @@ void UMapBuildDataRegistry::PostLoad()
 	{
 		// We already stripped unneeded formats during cooking, but some cooking targets require multiple formats to be stored
 		// Strip unneeded formats for the current max feature level
-		bool bRetainAllFeatureLevelData = GIsEditor && GMaxRHIFeatureLevel >= ERHIFeatureLevel::SM4;
+		bool bRetainAllFeatureLevelData = GIsEditor && GMaxRHIFeatureLevel >= ERHIFeatureLevel::SM5;
 		bool bEncodedDataRequired = bRetainAllFeatureLevelData || (GMaxRHIFeatureLevel == ERHIFeatureLevel::ES2 || GMaxRHIFeatureLevel == ERHIFeatureLevel::ES3_1);
-		bool bFullDataRequired = GMaxRHIFeatureLevel >= ERHIFeatureLevel::SM4;
+		bool bFullDataRequired = GMaxRHIFeatureLevel >= ERHIFeatureLevel::SM5;
 
 		for (TMap<FGuid, FReflectionCaptureMapBuildData>::TIterator It(ReflectionCaptureBuildData); It; ++It)
 		{
diff --git a/Engine/Source/Runtime/Engine/Private/Materials/HLSLMaterialTranslator.h b/Engine/Source/Runtime/Engine/Private/Materials/HLSLMaterialTranslator.h
--- a/Engine/Source/Runtime/Engine/Private/Materials/HLSLMaterialTranslator.h
+++ b/Engine/Source/Runtime/Engine/Private/Materials/HLSLMaterialTranslator.h
@@ -4279,7 +4279,7 @@ protected:
 		}
 
 		// If mobile, then disabling AutomaticViewMipBias.
-		if (FeatureLevel < ERHIFeatureLevel::SM4)
+		if (FeatureLevel < ERHIFeatureLevel::SM5)
 		{
 			AutomaticViewMipBias = false;
 		}
@@ -4808,7 +4808,7 @@ protected:
 										SceneTextureId == PPI_SceneDepth ||
 										SceneTextureId == PPI_CustomStencil;
 
-		if (!bSupportedOnMobile	&& ErrorUnlessFeatureLevelSupported(ERHIFeatureLevel::SM4) == INDEX_NONE)
+		if (!bSupportedOnMobile	&& ErrorUnlessFeatureLevelSupported(ERHIFeatureLevel::SM5) == INDEX_NONE)
 		{
 			return INDEX_NONE;
 		}
@@ -4840,7 +4840,7 @@ protected:
 
 		AddEstimatedTextureSample();
 
-		if (FeatureLevel >= ERHIFeatureLevel::SM4)
+		if (FeatureLevel >= ERHIFeatureLevel::SM5)
 		{
 			return AddCodeChunk(
 				MCT_Float4,
@@ -4912,7 +4912,7 @@ protected:
 
 			if (bRequiresSM5)
 			{
-				ErrorUnlessFeatureLevelSupported(ERHIFeatureLevel::SM4);
+				ErrorUnlessFeatureLevelSupported(ERHIFeatureLevel::SM5);
 			}
 		}
 
@@ -4978,7 +4978,7 @@ protected:
 			Errorf(TEXT("SceneColor lookups are only available when MaterialDomain = Surface."));
 		}
 
-		if (ErrorUnlessFeatureLevelSupported(ERHIFeatureLevel::SM4) == INDEX_NONE)
+		if (ErrorUnlessFeatureLevelSupported(ERHIFeatureLevel::SM5) == INDEX_NONE)
 		{
 			return INDEX_NONE;
 		}
@@ -5112,7 +5112,7 @@ protected:
 
 	virtual int32 ExternalTexture(const FGuid& ExternalTextureGuid) override
 	{
-		bool bOnlyInPixelShader = GetFeatureLevel() < ERHIFeatureLevel::SM4;
+		bool bOnlyInPixelShader = GetFeatureLevel() < ERHIFeatureLevel::SM5;
 
 		if (bOnlyInPixelShader && ShaderFrequency != SF_Pixel)
 		{
@@ -5124,7 +5124,7 @@ protected:
 
 	virtual int32 ExternalTexture(UTexture* InTexture, int32& TextureReferenceIndex) override
 	{
-		bool bOnlyInPixelShader = GetFeatureLevel() < ERHIFeatureLevel::SM4;
+		bool bOnlyInPixelShader = GetFeatureLevel() < ERHIFeatureLevel::SM5;
 
 		if (bOnlyInPixelShader && ShaderFrequency != SF_Pixel)
 		{
@@ -5139,7 +5139,7 @@ protected:
 
 	virtual int32 ExternalTextureParameter(FName ParameterName, UTexture* DefaultValue, int32& TextureReferenceIndex) override
 	{
-		bool bOnlyInPixelShader = GetFeatureLevel() < ERHIFeatureLevel::SM4;
+		bool bOnlyInPixelShader = GetFeatureLevel() < ERHIFeatureLevel::SM5;
 
 		if (bOnlyInPixelShader && ShaderFrequency != SF_Pixel)
 		{
@@ -6118,7 +6118,7 @@ protected:
 			return NonPixelShaderExpressionError();
 		}
 
-		if (ErrorUnlessFeatureLevelSupported(ERHIFeatureLevel::SM4) == INDEX_NONE)
+		if (ErrorUnlessFeatureLevelSupported(ERHIFeatureLevel::SM5) == INDEX_NONE)
 		{
 			return INDEX_NONE;
 		}
@@ -6141,7 +6141,7 @@ protected:
 			return NonPixelShaderExpressionError();
 		}
 
-		if (ErrorUnlessFeatureLevelSupported(ERHIFeatureLevel::SM4) == INDEX_NONE)
+		if (ErrorUnlessFeatureLevelSupported(ERHIFeatureLevel::SM5) == INDEX_NONE)
 		{
 			return INDEX_NONE;
 		}
@@ -6317,7 +6317,7 @@ protected:
 
 	virtual int32 AntialiasedTextureMask(int32 Tex, int32 UV, float Threshold, uint8 Channel) override
 	{
-		if (ErrorUnlessFeatureLevelSupported(ERHIFeatureLevel::SM4) == INDEX_NONE)
+		if (ErrorUnlessFeatureLevelSupported(ERHIFeatureLevel::SM5) == INDEX_NONE)
 		{
 			return INDEX_NONE;
 		}
@@ -6394,7 +6394,7 @@ protected:
 		// GradientTex3D uses 3D texturing, which is not available on ES2
 		if (NoiseFunction == NOISEFUNCTION_GradientTex3D)
 		{
-			if (ErrorUnlessFeatureLevelSupported(ERHIFeatureLevel::SM4) == INDEX_NONE)
+			if (ErrorUnlessFeatureLevelSupported(ERHIFeatureLevel::SM5) == INDEX_NONE)
 			{
 				Errorf(TEXT("3D textures are not supported for ES2"));
 				return INDEX_NONE;
@@ -6533,7 +6533,7 @@ protected:
 
 	virtual int32 AtmosphericFogColor( int32 WorldPosition ) override
 	{
-		if (ErrorUnlessFeatureLevelSupported(ERHIFeatureLevel::SM4) == INDEX_NONE)
+		if (ErrorUnlessFeatureLevelSupported(ERHIFeatureLevel::SM5) == INDEX_NONE)
 		{
 			return INDEX_NONE;
 		}
@@ -7003,7 +7003,7 @@ protected:
 	 */
 	virtual int32 TextureCoordinateOffset() override
 	{
-		if (FeatureLevel < ERHIFeatureLevel::SM4 && ShaderFrequency == SF_Vertex)
+		if (FeatureLevel < ERHIFeatureLevel::SM5 && ShaderFrequency == SF_Vertex)
 		{
 			return AddInlinedCodeChunk(MCT_Float2, TEXT("Parameters.TexCoordOffset"));
 		}
diff --git a/Engine/Source/Runtime/Engine/Private/Materials/MaterialExpressions.cpp b/Engine/Source/Runtime/Engine/Private/Materials/MaterialExpressions.cpp
--- a/Engine/Source/Runtime/Engine/Private/Materials/MaterialExpressions.cpp
+++ b/Engine/Source/Runtime/Engine/Private/Materials/MaterialExpressions.cpp
@@ -7773,12 +7773,18 @@ void UMaterialExpressionFeatureLevelSwitch::Serialize(FStructuredArchive::FRecor
 {
 	Super::Serialize(Record);
 	FArchive& UnderlyingArchive = Record.GetUnderlyingArchive();
+	UnderlyingArchive.UsingCustomVersion(FRenderingObjectVersion::GUID);
 
 	if (UnderlyingArchive.IsLoading() && UnderlyingArchive.UE4Ver() < VER_UE4_RENAME_SM3_TO_ES3_1)
 	{
 		// Copy the ES2 input to SM3 (since SM3 will now become ES3_1 and we don't want broken content)
 		Inputs[ERHIFeatureLevel::ES3_1] = Inputs[ERHIFeatureLevel::ES2];
 	}
+
+	if (UnderlyingArchive.CustomVer(FRenderingObjectVersion::GUID) < FRenderingObjectVersion::RemovedSM4)
+	{
+		Inputs[ERHIFeatureLevel::SM4_REMOVED] = UMaterialExpressionFeatureLevelSwitch::Default;
+	}
 }
 
 bool UMaterialExpressionFeatureLevelSwitch::NeedsLoadForClient() const
@@ -7819,7 +7825,7 @@ int32 UMaterialExpressionShadingPathSwitch::Compile(class FMaterialCompiler* Com
 	{
 		ShadingPathToCompile = ERHIShadingPath::Forward;
 	}
-	else if (Compiler->GetFeatureLevel() < ERHIFeatureLevel::SM4)
+	else if (Compiler->GetFeatureLevel() < ERHIFeatureLevel::SM5)
 	{
 		ShadingPathToCompile = ERHIShadingPath::Mobile;
 	}
diff --git a/Engine/Source/Runtime/Engine/Private/Materials/MaterialShared.cpp b/Engine/Source/Runtime/Engine/Private/Materials/MaterialShared.cpp
--- a/Engine/Source/Runtime/Engine/Private/Materials/MaterialShared.cpp
+++ b/Engine/Source/Runtime/Engine/Private/Materials/MaterialShared.cpp
@@ -913,7 +913,7 @@ void FMaterial::LegacySerialize(FArchive& Ar)
 		int32 LegacyInt;
 		Ar << LegacyInt;
 
-		FeatureLevel = ERHIFeatureLevel::SM4;
+		FeatureLevel = ERHIFeatureLevel::SM4_REMOVED;
 		QualityLevel = EMaterialQualityLevel::High;
 
 #if !WITH_EDITOR
@@ -1700,7 +1700,7 @@ void FMaterial::SetupMaterialEnvironment(
 		OutEnvironment.SetDefine(TEXT("EDITOR_PRIMITIVE_MATERIAL"),TEXT("1"));
 	}
 
-	if (IsFeatureLevelSupported(Platform, ERHIFeatureLevel::SM4))
+	if (IsFeatureLevelSupported(Platform, ERHIFeatureLevel::SM5))
 	{	
 		static const auto CVar = IConsoleManager::Get().FindTConsoleVariableDataInt(TEXT("r.StencilForLODDither"));
 		OutEnvironment.SetDefine(TEXT("USE_STENCIL_LOD_DITHER_DEFAULT"), CVar->GetValueOnAnyThread() != 0 ? 1 : 0);
@@ -2944,7 +2944,7 @@ bool FMaterial::GetMaterialExpressionSource( FString& OutSource )
 bool FMaterial::WritesEveryPixel(bool bShadowPass) const
 {
 	bool bStencilDitheredLOD = false;
-	if (FeatureLevel >= ERHIFeatureLevel::SM4)
+	if (FeatureLevel >= ERHIFeatureLevel::SM5)
 	{
 		// this option affects only deferred renderer
 		static TConsoleVariableData<int32>* CVarStencilDitheredLOD =
diff --git a/Engine/Source/Runtime/Engine/Private/MeshParticleVertexFactory.cpp b/Engine/Source/Runtime/Engine/Private/MeshParticleVertexFactory.cpp
--- a/Engine/Source/Runtime/Engine/Private/MeshParticleVertexFactory.cpp
+++ b/Engine/Source/Runtime/Engine/Private/MeshParticleVertexFactory.cpp
@@ -82,7 +82,7 @@ public:
 				ShaderBindings.Add(DynamicParameter, FVector4(DynamicVertex->DynamicValue[0], DynamicVertex->DynamicValue[1], DynamicVertex->DynamicValue[2], DynamicVertex->DynamicValue[3]));
 			}
 
-			if (BatchParameters->PrevTransformBuffer && FeatureLevel >= ERHIFeatureLevel::SM4)
+			if (BatchParameters->PrevTransformBuffer && FeatureLevel >= ERHIFeatureLevel::SM5)
 			{
 				ShaderBindings.Add(PrevTransform0, PrevTransformVertex->PrevTransform0);
 				ShaderBindings.Add(PrevTransform1, PrevTransformVertex->PrevTransform1);
@@ -91,7 +91,7 @@ public:
 
 			ShaderBindings.Add(ParticleColor, FVector4(Vertex->Color.Component(0), Vertex->Color.Component(1), Vertex->Color.Component(2), Vertex->Color.Component(3)));
 		}
-		else if (FeatureLevel >= ERHIFeatureLevel::SM4)
+		else if (FeatureLevel >= ERHIFeatureLevel::SM5)
 		{
 			ShaderBindings.Add(PrevTransformBuffer, MeshParticleVF->GetPreviousTransformBufferSRV());
 		}
@@ -199,7 +199,7 @@ void FMeshParticleVertexFactory::InitRHI()
 			}
 
 			// Add a dummy resource to avoid crash due to missing resource
-			if (GMaxRHIFeatureLevel >= ERHIFeatureLevel::SM4)
+			if (GMaxRHIFeatureLevel >= ERHIFeatureLevel::SM5)
 			{
 				PrevTransformBuffer.NumBytes = 0;
 				PrevTransformBuffer.Buffer = GDummyPrevTransformBuffer.GetVB();
diff --git a/Engine/Source/Runtime/Engine/Private/Particles/FXSystemPrivate.h b/Engine/Source/Runtime/Engine/Private/Particles/FXSystemPrivate.h
--- a/Engine/Source/Runtime/Engine/Private/Particles/FXSystemPrivate.h
+++ b/Engine/Source/Runtime/Engine/Private/Particles/FXSystemPrivate.h
@@ -64,13 +64,13 @@ inline bool IsParticleCollisionModeSupported(EShaderPlatform InPlatform, EPartic
 	case PCM_DepthBuffer:
 		// we only need to check for simple forward if we're NOT currently attempting to cache the shader
 		// since SF is a runtime change, we need to compile the shader regardless, because we could be switching to deferred at any time
-		return (IsFeatureLevelSupported(InPlatform, ERHIFeatureLevel::SM4))
+		return (IsFeatureLevelSupported(InPlatform, ERHIFeatureLevel::SM5))
 			&& (bForCaching || !IsSimpleForwardShadingEnabled(InPlatform));
 	case PCM_DistanceField:
 		return IsFeatureLevelSupported(InPlatform, ERHIFeatureLevel::SM5);
 	}
 	check(0);
-	return IsFeatureLevelSupported(InPlatform, ERHIFeatureLevel::SM4);
+	return IsFeatureLevelSupported(InPlatform, ERHIFeatureLevel::SM5);
 }
 
 
diff --git a/Engine/Source/Runtime/Engine/Private/Particles/ParticleSystemRender.cpp b/Engine/Source/Runtime/Engine/Private/Particles/ParticleSystemRender.cpp
--- a/Engine/Source/Runtime/Engine/Private/Particles/ParticleSystemRender.cpp
+++ b/Engine/Source/Runtime/Engine/Private/Particles/ParticleSystemRender.cpp
@@ -1668,7 +1668,7 @@ void FDynamicMeshEmitterData::GetDynamicMeshElementsEmitter(const FParticleSyste
 
 				// For OpenGL & Metal we can't assume that it is OK to leave the PrevTransformBuffer buffer unbound.
 				// Doing so can lead to undefined behaviour if the buffer is referenced in the shader even if protected by a branch that is not meant to be taken.
-				bool const bGeneratePrevTransformBuffer = (FeatureLevel >= ERHIFeatureLevel::SM4) && 
+				bool const bGeneratePrevTransformBuffer = (FeatureLevel >= ERHIFeatureLevel::SM5) && 
 														  (Source.MeshMotionBlurOffset || IsOpenGLPlatform(ShaderPlatform) || IsMetalPlatform(ShaderPlatform) || IsPS4Platform(ShaderPlatform));
 
 
diff --git a/Engine/Source/Runtime/Engine/Private/PrecomputedVolumetricLightmap.cpp b/Engine/Source/Runtime/Engine/Private/PrecomputedVolumetricLightmap.cpp
--- a/Engine/Source/Runtime/Engine/Private/PrecomputedVolumetricLightmap.cpp
+++ b/Engine/Source/Runtime/Engine/Private/PrecomputedVolumetricLightmap.cpp
@@ -218,7 +218,7 @@ FArchive& operator<<(FArchive& Ar,FPrecomputedVolumetricLightmapData& Volume)
 
 	if (Ar.IsLoading())
 	{
-		if (GMaxRHIFeatureLevel >= ERHIFeatureLevel::SM4 && !GIsEditor)
+		if (GMaxRHIFeatureLevel >= ERHIFeatureLevel::SM5 && !GIsEditor)
 		{
 			// drop LQ data for SM4+
 			Volume.BrickData.DiscardLowQualityLayers();
diff --git a/Engine/Source/Runtime/Engine/Private/RawIndexBuffer.cpp b/Engine/Source/Runtime/Engine/Private/RawIndexBuffer.cpp
--- a/Engine/Source/Runtime/Engine/Private/RawIndexBuffer.cpp
+++ b/Engine/Source/Runtime/Engine/Private/RawIndexBuffer.cpp
@@ -71,7 +71,7 @@ void FRawIndexBuffer16or32::CacheOptimize()
 
 void FRawIndexBuffer16or32::ComputeIndexWidth()
 {
-	if (GetFeatureLevel() < ERHIFeatureLevel::SM4)
+	if (GetFeatureLevel() < ERHIFeatureLevel::SM5)
 	{
 		const int32 NumIndices = Indices.Num();
 		bool bShouldUse32Bit = false;
diff --git a/Engine/Source/Runtime/Engine/Private/SkeletalRenderGPUSkin.cpp b/Engine/Source/Runtime/Engine/Private/SkeletalRenderGPUSkin.cpp
--- a/Engine/Source/Runtime/Engine/Private/SkeletalRenderGPUSkin.cpp
+++ b/Engine/Source/Runtime/Engine/Private/SkeletalRenderGPUSkin.cpp
@@ -517,7 +517,7 @@ void FSkeletalMeshObjectGPUSkin::ProcessUpdatedDynamicData(FGPUSkinCache* GPUSki
 		{
 			const FSkelMeshRenderSection& Section = Sections[SectionIdx];
 
-			bool bClothFactory = (FeatureLevel >= ERHIFeatureLevel::SM4) && (DynamicData->ClothingSimData.Num() > 0) && Section.HasClothingData();
+			bool bClothFactory = (FeatureLevel >= ERHIFeatureLevel::SM5) && (DynamicData->ClothingSimData.Num() > 0) && Section.HasClothingData();
 
 			FGPUBaseSkinVertexFactory* VertexFactory;
 			{
@@ -1546,7 +1546,7 @@ void FSkeletalMeshObjectGPUSkin::FVertexFactoryData::InitAPEXClothVertexFactorie
 	ClothVertexFactories.Empty(Sections.Num());
 	for( int32 FactoryIdx=0; FactoryIdx < Sections.Num(); FactoryIdx++ )
 	{
-		if (Sections[FactoryIdx].HasClothingData() && InFeatureLevel >= ERHIFeatureLevel::SM4)
+		if (Sections[FactoryIdx].HasClothingData() && InFeatureLevel >= ERHIFeatureLevel::SM5)
 		{
 			if (VertexBuffers.SkinWeightVertexBuffer->HasExtraBoneInfluences())
 			{
diff --git a/Engine/Source/Runtime/Engine/Private/UnrealClient.cpp b/Engine/Source/Runtime/Engine/Private/UnrealClient.cpp
--- a/Engine/Source/Runtime/Engine/Private/UnrealClient.cpp
+++ b/Engine/Source/Runtime/Engine/Private/UnrealClient.cpp
@@ -666,7 +666,7 @@ int32 FStatUnitData::DrawStat(FViewport* InViewport, FCanvas* InCanvas, int32 In
 		}
 
 		ERHIFeatureLevel::Type FeatureLevel = InCanvas->GetFeatureLevel();
-		if (FeatureLevel >= ERHIFeatureLevel::SM4)
+		if (FeatureLevel >= ERHIFeatureLevel::SM5)
 		{
 			float ResolutionFraction = DynamicResolutionStateInfos.ResolutionFractionApproximation;
 			float ScreenPercentage = ResolutionFraction * 100.0f;
diff --git a/Engine/Source/Runtime/Engine/Private/VectorFieldVisualization.cpp b/Engine/Source/Runtime/Engine/Private/VectorFieldVisualization.cpp
--- a/Engine/Source/Runtime/Engine/Private/VectorFieldVisualization.cpp
+++ b/Engine/Source/Runtime/Engine/Private/VectorFieldVisualization.cpp
@@ -141,7 +141,7 @@ void FVectorFieldVisualizationVertexFactory::ReleaseRHI()
 bool FVectorFieldVisualizationVertexFactory::ShouldCompilePermutation(EShaderPlatform Platform, const class FMaterial* Material, const class FShaderType* ShaderType)
 {
 	// We support vector fields on mobile, but not the visualization, so we limit this VF to Feature Level Preview shader platforms
-	return Material->IsSpecialEngineMaterial() && SupportsGPUParticles(Platform) && (IsFeatureLevelSupported(Platform, ERHIFeatureLevel::SM4) || IsPCPlatform(Platform));
+	return Material->IsSpecialEngineMaterial() && SupportsGPUParticles(Platform) && (IsFeatureLevelSupported(Platform, ERHIFeatureLevel::SM5) || IsPCPlatform(Platform));
 }
 
 /**
diff --git a/Engine/Source/Runtime/Engine/Public/CubemapUnwrapUtils.h b/Engine/Source/Runtime/Engine/Public/CubemapUnwrapUtils.h
--- a/Engine/Source/Runtime/Engine/Public/CubemapUnwrapUtils.h
+++ b/Engine/Source/Runtime/Engine/Public/CubemapUnwrapUtils.h
@@ -147,7 +147,7 @@ public:
 
 	static bool ShouldCompilePermutation(const FGlobalShaderPermutationParameters& Parameters)
 	{
-		return IsFeatureLevelSupported(Parameters.Platform, ERHIFeatureLevel::SM4) && !IsConsolePlatform(Parameters.Platform);
+		return IsFeatureLevelSupported(Parameters.Platform, ERHIFeatureLevel::SM5) && !IsConsolePlatform(Parameters.Platform);
 	}
 
 	FIESLightProfilePS() {}
diff --git a/Engine/Source/Runtime/Engine/Public/MaterialShared.h b/Engine/Source/Runtime/Engine/Public/MaterialShared.h
--- a/Engine/Source/Runtime/Engine/Public/MaterialShared.h
+++ b/Engine/Source/Runtime/Engine/Public/MaterialShared.h
@@ -733,7 +733,7 @@ public:
 	
 	FMaterialShaderMapId()
 		: QualityLevel(EMaterialQualityLevel::High)
-		, FeatureLevel(ERHIFeatureLevel::SM4)
+		, FeatureLevel(ERHIFeatureLevel::SM5)
 #if WITH_EDITOR
 		, Usage(EMaterialShaderMapUsage::Default)
 		, bIsCookedId(false)
@@ -1369,7 +1369,7 @@ public:
 		RenderingThreadShaderMap(NULL),
 		QualityLevel(EMaterialQualityLevel::High),
 		bHasQualityLevelUsage(false),
-		FeatureLevel(ERHIFeatureLevel::SM4),
+		FeatureLevel(ERHIFeatureLevel::SM5),
 		bContainsInlineShaders(false),
 		bLoadedCookedShaderMapId(false)
 	{}
diff --git a/Engine/Source/Runtime/Engine/Public/SceneInterface.h b/Engine/Source/Runtime/Engine/Public/SceneInterface.h
--- a/Engine/Source/Runtime/Engine/Public/SceneInterface.h
+++ b/Engine/Source/Runtime/Engine/Public/SceneInterface.h
@@ -403,7 +403,7 @@ public:
 
 	static EShadingPath GetShadingPath(ERHIFeatureLevel::Type InFeatureLevel)
 	{
-		if (InFeatureLevel >= ERHIFeatureLevel::SM4)
+		if (InFeatureLevel >= ERHIFeatureLevel::SM5)
 		{
 			return EShadingPath::Deferred;
 		}
diff --git a/Engine/Source/Runtime/Engine/Public/ScreenRendering.h b/Engine/Source/Runtime/Engine/Public/ScreenRendering.h
--- a/Engine/Source/Runtime/Engine/Public/ScreenRendering.h
+++ b/Engine/Source/Runtime/Engine/Public/ScreenRendering.h
@@ -215,7 +215,7 @@ public:
 
 	static bool ShouldCompilePermutation(const FGlobalShaderPermutationParameters& Parameters)
 	{
-		return IsFeatureLevelSupported(Parameters.Platform, ERHIFeatureLevel::SM4) && (!bUsingVertexLayers || RHISupportsVertexShaderLayer(Parameters.Platform));
+		return IsFeatureLevelSupported(Parameters.Platform, ERHIFeatureLevel::SM5) && (!bUsingVertexLayers || RHISupportsVertexShaderLayer(Parameters.Platform));
 	}
 
 	TScreenVSForGS(const ShaderMetaType::CompiledShaderInitializerType& Initializer):
diff --git a/Engine/Source/Runtime/Landscape/Private/Landscape.cpp b/Engine/Source/Runtime/Landscape/Private/Landscape.cpp
--- a/Engine/Source/Runtime/Landscape/Private/Landscape.cpp
+++ b/Engine/Source/Runtime/Landscape/Private/Landscape.cpp
@@ -1369,7 +1369,7 @@ FPrimitiveSceneProxy* ULandscapeComponent::CreateSceneProxy()
 {
 	const auto FeatureLevel = GetWorld()->FeatureLevel;
 	FPrimitiveSceneProxy* Proxy = nullptr;
-	if (FeatureLevel >= ERHIFeatureLevel::SM4)
+	if (FeatureLevel >= ERHIFeatureLevel::SM5)
 	{
 		Proxy = new FLandscapeComponentSceneProxy(this);
 	}
@@ -3368,7 +3368,7 @@ void ALandscapeProxy::UpdateBakedTextures()
 {
 	// See if we can render
 	UWorld* World = GetWorld();
-	if (!GIsEditor || GUsingNullRHI || !World || World->IsGameWorld() || World->FeatureLevel < ERHIFeatureLevel::SM4)
+	if (!GIsEditor || GUsingNullRHI || !World || World->IsGameWorld() || World->FeatureLevel < ERHIFeatureLevel::SM5)
 	{
 		return;
 	}
diff --git a/Engine/Source/Runtime/Landscape/Private/LandscapeGrass.cpp b/Engine/Source/Runtime/Landscape/Private/LandscapeGrass.cpp
--- a/Engine/Source/Runtime/Landscape/Private/LandscapeGrass.cpp
+++ b/Engine/Source/Runtime/Landscape/Private/LandscapeGrass.cpp
@@ -205,7 +205,7 @@ static bool ShouldCacheLandscapeGrassShaders(EShaderPlatform Platform, const FMa
 {
 	// We only need grass weight shaders for Landscape vertex factories on desktop platforms
 	return (Material->IsUsedWithLandscape() || Material->IsSpecialEngineMaterial()) &&
-		IsFeatureLevelSupported(Platform, ERHIFeatureLevel::SM4) &&
+		IsFeatureLevelSupported(Platform, ERHIFeatureLevel::SM5) &&
 		((VertexFactoryType == FindVertexFactoryType(FName(TEXT("FLandscapeVertexFactory"), FNAME_Find))) || (VertexFactoryType == FindVertexFactoryType(FName(TEXT("FLandscapeXYOffsetVertexFactory"), FNAME_Find))))
 		&& !IsConsolePlatform(Platform);
 }
@@ -867,7 +867,7 @@ bool ULandscapeComponent::CanRenderGrassMap() const
 {
 	// Check we can render
 	UWorld* ComponentWorld = GetWorld();
-	if (!GIsEditor || GUsingNullRHI || !ComponentWorld || ComponentWorld->IsGameWorld() || ComponentWorld->FeatureLevel < ERHIFeatureLevel::SM4 || !SceneProxy)
+	if (!GIsEditor || GUsingNullRHI || !ComponentWorld || ComponentWorld->IsGameWorld() || ComponentWorld->FeatureLevel < ERHIFeatureLevel::SM5 || !SceneProxy)
 	{
 		return false;
 	}
@@ -2031,7 +2031,7 @@ void ALandscapeProxy::FlushGrassComponents(const TSet<ULandscapeComponent*>* Onl
 			}
 		}
 #if WITH_EDITOR
-		if (GIsEditor && bFlushGrassMaps && GetWorld() && GetWorld()->FeatureLevel >= ERHIFeatureLevel::SM4)
+		if (GIsEditor && bFlushGrassMaps && GetWorld() && GetWorld()->FeatureLevel >= ERHIFeatureLevel::SM5)
 		{
 			for (ULandscapeComponent* Component : *OnlyForComponents)
 			{
@@ -2076,7 +2076,7 @@ void ALandscapeProxy::FlushGrassComponents(const TSet<ULandscapeComponent*>* Onl
 #if WITH_EDITOR
 		UWorld* World = GetWorld();
 
-		if (GIsEditor && bFlushGrassMaps && World != nullptr && World->Scene != nullptr && World->Scene->GetFeatureLevel() >= ERHIFeatureLevel::SM4)
+		if (GIsEditor && bFlushGrassMaps && World != nullptr && World->Scene != nullptr && World->Scene->GetFeatureLevel() >= ERHIFeatureLevel::SM5)
 		{
 			// Clear GrassMaps
 			for (UActorComponent* Component : GetComponents())
diff --git a/Engine/Source/Runtime/Landscape/Private/LandscapeRender.cpp b/Engine/Source/Runtime/Landscape/Private/LandscapeRender.cpp
--- a/Engine/Source/Runtime/Landscape/Private/LandscapeRender.cpp
+++ b/Engine/Source/Runtime/Landscape/Private/LandscapeRender.cpp
@@ -691,7 +691,7 @@ FLandscapeComponentSceneProxy::FLandscapeComponentSceneProxy(ULandscapeComponent
 
 	const auto FeatureLevel = GetScene().GetFeatureLevel();
 
-	if (FeatureLevel >= ERHIFeatureLevel::SM4)
+	if (FeatureLevel >= ERHIFeatureLevel::SM5)
 	{
 		if (InComponent->GetLandscapeProxy()->bUseDynamicMaterialInstance)
 		{
@@ -842,7 +842,7 @@ FLandscapeComponentSceneProxy::FLandscapeComponentSceneProxy(ULandscapeComponent
 
 			bool HasTessellationEnabled = false;
 
-			if (FeatureLevel >= ERHIFeatureLevel::SM4)
+			if (FeatureLevel >= ERHIFeatureLevel::SM5)
 			{
 				HasTessellationEnabled = LandscapeMaterial->D3D11TessellationMode != EMaterialTessellationMode::MTM_NoTessellation;
 			}
@@ -1800,7 +1800,7 @@ void FLandscapeComponentSceneProxy::DrawStaticElements(FStaticPrimitiveDrawInter
 
 		UMaterialInstance* MaterialInstance = Cast<UMaterialInstance>(AvailableMaterials[i]);
 
-		bool HasTessellationEnabled = (FeatureLevel >= ERHIFeatureLevel::SM4) ? MaterialInstance != nullptr && MaterialInstance->GetMaterial()->D3D11TessellationMode != EMaterialTessellationMode::MTM_NoTessellation && MaterialIndexToDisabledTessellationMaterial[i] != INDEX_NONE : false;
+		bool HasTessellationEnabled = (FeatureLevel >= ERHIFeatureLevel::SM5) ? MaterialInstance != nullptr && MaterialInstance->GetMaterial()->D3D11TessellationMode != EMaterialTessellationMode::MTM_NoTessellation && MaterialIndexToDisabledTessellationMaterial[i] != INDEX_NONE : false;
 
 		// Only add normal batch index (for each material)
 		MaterialIndexToStaticMeshBatchLOD[i] = CurrentLODIndex;
@@ -1945,7 +1945,7 @@ uint64 FLandscapeComponentSceneProxy::GetStaticBatchElementVisibility(const FSce
 		const FViewCustomDataLOD* CurrentLODData = (const FViewCustomDataLOD*)ViewCustomData;
 		const auto FeatureLevel = InView.GetFeatureLevel();
 
-		if (FeatureLevel >= ERHIFeatureLevel::SM4)
+		if (FeatureLevel >= ERHIFeatureLevel::SM5)
 		{
 			const int8 CurrentLODIndex = CurrentLODData->SubSections[0].BatchElementCurrentLOD;
 			ULandscapeMaterialInstanceConstant* LandscapeMIC = Cast<ULandscapeMaterialInstanceConstant>(AvailableMaterials[LODIndexToMaterialIndex[CurrentLODIndex]]);
@@ -2199,7 +2199,7 @@ void* FLandscapeComponentSceneProxy::InitViewCustomData(const FSceneView& InView
 	LODData->IsShadowOnly = InIsShadowOnly;
 
 	// Mobile use a different way of calculating the Bias
-	if (GetScene().GetFeatureLevel() >= ERHIFeatureLevel::SM4)
+	if (GetScene().GetFeatureLevel() >= ERHIFeatureLevel::SM5)
 	{
 		LODData->LodBias = GetShaderLODBias();
 	}
@@ -2218,7 +2218,7 @@ void FLandscapeComponentSceneProxy::ComputeTessellationFalloffShaderValues(const
 	const auto FeatureLevel = GetScene().GetFeatureLevel();
 	bool HasTessellationEnabled = false;
 
-	if (FeatureLevel >= ERHIFeatureLevel::SM4)
+	if (FeatureLevel >= ERHIFeatureLevel::SM5)
 	{
 		const int8 CurrentLODIndex = InLODData.SubSections[0].BatchElementCurrentLOD;
 		ULandscapeMaterialInstanceConstant* LandscapeMIC = Cast<ULandscapeMaterialInstanceConstant>(AvailableMaterials[LODIndexToMaterialIndex[CurrentLODIndex]]);
@@ -2543,7 +2543,7 @@ FLODMask FLandscapeComponentSceneProxy::GetCustomLOD(const FSceneView& InView, f
 		const auto FeatureLevel = InView.GetFeatureLevel();
 		bool HasTessellationEnabled = false;
 
-		if (FeatureLevel >= ERHIFeatureLevel::SM4)
+		if (FeatureLevel >= ERHIFeatureLevel::SM5)
 		{
 			ULandscapeMaterialInstanceConstant* LandscapeMIC = Cast<ULandscapeMaterialInstanceConstant>(AvailableMaterials[LODIndexToMaterialIndex[PotentialLOD]]);
 
@@ -2632,7 +2632,7 @@ FLODMask FLandscapeComponentSceneProxy::GetCustomWholeSceneShadowLOD(const FScen
 		const auto FeatureLevel = InView.GetFeatureLevel();
 		bool HasTessellationEnabled = false;
 
-		if (FeatureLevel >= ERHIFeatureLevel::SM4)
+		if (FeatureLevel >= ERHIFeatureLevel::SM5)
 		{
 			ULandscapeMaterialInstanceConstant* LandscapeMIC = Cast<ULandscapeMaterialInstanceConstant>(AvailableMaterials[LODIndexToMaterialIndex[PotentialLOD]]);
 
@@ -2738,7 +2738,7 @@ void FLandscapeComponentSceneProxy::GetDynamicMeshElements(const TArray<const FS
 			const auto FeatureLevel = View->GetFeatureLevel();
 			bool HasTessellationEnabled = false;
 
-			if (FeatureLevel >= ERHIFeatureLevel::SM4)
+			if (FeatureLevel >= ERHIFeatureLevel::SM5)
 			{
 				const int8 CurrentLODIndex = PrimitiveCustomData->SubSections[0].BatchElementCurrentLOD;
 				ULandscapeMaterialInstanceConstant* LandscapeMIC = Cast<ULandscapeMaterialInstanceConstant>(AvailableMaterials[LODIndexToMaterialIndex[CurrentLODIndex]]);
@@ -4627,11 +4627,11 @@ void ULandscapeComponent::GetStreamingRenderAssetInfo(FStreamingTextureLevelCont
 	}
 
 	ERHIFeatureLevel::Type FeatureLevel = LevelContext.GetFeatureLevel();
-	int32 MaterialInstanceCount = FeatureLevel >= ERHIFeatureLevel::SM4 ? GetMaterialInstanceCount() : MobileMaterialInterfaces.Num();
+	int32 MaterialInstanceCount = FeatureLevel >= ERHIFeatureLevel::SM5 ? GetMaterialInstanceCount() : MobileMaterialInterfaces.Num();
 
 	for (int32 MaterialIndex = 0; MaterialIndex < MaterialInstanceCount; ++MaterialIndex)
 	{
-		const UMaterialInterface* MaterialInterface = FeatureLevel >= ERHIFeatureLevel::SM4 ? GetMaterialInstance(MaterialIndex) : MobileMaterialInterfaces[MaterialIndex];
+		const UMaterialInterface* MaterialInterface = FeatureLevel >= ERHIFeatureLevel::SM5 ? GetMaterialInstance(MaterialIndex) : MobileMaterialInterfaces[MaterialIndex];
 
 		// Normal usage...
 		// Enumerate the textures used by the material.
diff --git a/Engine/Source/Runtime/Landscape/Public/LandscapeRender.h b/Engine/Source/Runtime/Landscape/Public/LandscapeRender.h
--- a/Engine/Source/Runtime/Landscape/Public/LandscapeRender.h
+++ b/Engine/Source/Runtime/Landscape/Public/LandscapeRender.h
@@ -187,7 +187,7 @@ public:
 	{
 		// only compile landscape materials for landscape vertex factory
 		// The special engine materials must be compiled for the landscape vertex factory because they are used with it for wireframe, etc.
-		return IsFeatureLevelSupported(Platform, ERHIFeatureLevel::SM4) &&
+		return IsFeatureLevelSupported(Platform, ERHIFeatureLevel::SM5) &&
 			(Material->IsUsedWithLandscape() || Material->IsSpecialEngineMaterial());
 	}
 
diff --git a/Engine/Source/Runtime/Launch/Private/LaunchEngineLoop.cpp b/Engine/Source/Runtime/Launch/Private/LaunchEngineLoop.cpp
--- a/Engine/Source/Runtime/Launch/Private/LaunchEngineLoop.cpp
+++ b/Engine/Source/Runtime/Launch/Private/LaunchEngineLoop.cpp
@@ -4077,7 +4077,7 @@ void FEngineLoop::Tick()
 		#if WITH_PROFILEGPU && !UE_BUILD_SHIPPING
 			// Issue the measurement of the execution time of a basic LongGPUTask unit on the very first frame
 			// The results will be retrived on the first call of IssueScalableLongGPUTask
-			if (GFrameCounter == 0 && IsFeatureLevelSupported(GMaxRHIShaderPlatform, ERHIFeatureLevel::SM4) && FApp::CanEverRender())
+			if (GFrameCounter == 0 && IsFeatureLevelSupported(GMaxRHIShaderPlatform, ERHIFeatureLevel::SM5) && FApp::CanEverRender())
 			{
 				FlushRenderingCommands();
 
diff --git a/Engine/Source/Runtime/NullDrv/Private/NullRHI.cpp b/Engine/Source/Runtime/NullDrv/Private/NullRHI.cpp
--- a/Engine/Source/Runtime/NullDrv/Private/NullRHI.cpp
+++ b/Engine/Source/Runtime/NullDrv/Private/NullRHI.cpp
@@ -20,12 +20,12 @@ void FNullDynamicRHI::Init()
 #if PLATFORM_WINDOWS
 	GShaderPlatformForFeatureLevel[ERHIFeatureLevel::ES2] = SP_PCD3D_ES2;
 	GShaderPlatformForFeatureLevel[ERHIFeatureLevel::ES3_1] = SP_PCD3D_ES3_1;
-	GShaderPlatformForFeatureLevel[ERHIFeatureLevel::SM4] = SP_PCD3D_SM4;
+	GShaderPlatformForFeatureLevel[ERHIFeatureLevel::SM4_REMOVED] = SP_NumPlatforms;
 	GShaderPlatformForFeatureLevel[ERHIFeatureLevel::SM5] = SP_PCD3D_SM5;
 #elif PLATFORM_MAC
     GShaderPlatformForFeatureLevel[ERHIFeatureLevel::ES2] = SP_METAL_MACES2;
     GShaderPlatformForFeatureLevel[ERHIFeatureLevel::ES3_1] = SP_METAL_MACES3_1;
-    GShaderPlatformForFeatureLevel[ERHIFeatureLevel::SM4] = SP_NumPlatforms;
+    GShaderPlatformForFeatureLevel[ERHIFeatureLevel::SM4_REMOVED] = SP_NumPlatforms;
     GShaderPlatformForFeatureLevel[ERHIFeatureLevel::SM5] = SP_METAL_SM5;
 #else
 	GShaderPlatformForFeatureLevel[GetMaxSupportedFeatureLevel(GMaxRHIShaderPlatform)] = GMaxRHIShaderPlatform;
diff --git a/Engine/Source/Runtime/OpenGLDrv/Private/Linux/OpenGLLinux.cpp b/Engine/Source/Runtime/OpenGLDrv/Private/Linux/OpenGLLinux.cpp
--- a/Engine/Source/Runtime/OpenGLDrv/Private/Linux/OpenGLLinux.cpp
+++ b/Engine/Source/Runtime/OpenGLDrv/Private/Linux/OpenGLLinux.cpp
@@ -160,15 +160,7 @@ void PlatformOpenGLVersionFromCommandLine(int& OutMajorVersion, int& OutMinorVer
 	bool bGL4 = IsOpenGL4Forced();
 	if (!bGL3 && !bGL4)
 	{
-		// if neither is forced, go with the first RHI in the list.
-		if (GRequestedFeatureLevel == ERHIFeatureLevel::SM4)
-		{
-			bGL3 = true;
-		}
-		else
-		{
-			bGL4 = true;
-		}
+		bGL4 = true;
 	}
 
 	// between GL3 and GL4, prefer GL3 since it might have been forced as a safety measure.
diff --git a/Engine/Source/Runtime/OpenGLDrv/Private/OpenGLCommands.cpp b/Engine/Source/Runtime/OpenGLDrv/Private/OpenGLCommands.cpp
--- a/Engine/Source/Runtime/OpenGLDrv/Private/OpenGLCommands.cpp
+++ b/Engine/Source/Runtime/OpenGLDrv/Private/OpenGLCommands.cpp
@@ -795,7 +795,7 @@ void FOpenGLDynamicRHI::UpdateSRV(FOpenGLShaderResourceView* SRV)
 	check(SRV);
 	// For Depth/Stencil textures whose Stencil component we wish to sample we must blit the stencil component out to an intermediate texture when we 'Store' the texture.
 #if PLATFORM_DESKTOP || PLATFORM_ANDROIDESDEFERRED || PLATFORM_LUMINGL4
-	if (FOpenGL::GetFeatureLevel() >= ERHIFeatureLevel::SM4 && FOpenGL::SupportsPixelBuffers() && IsValidRef(SRV->Texture2D))
+	if (FOpenGL::GetFeatureLevel() >= ERHIFeatureLevel::SM5 && FOpenGL::SupportsPixelBuffers() && IsValidRef(SRV->Texture2D))
 	{
 		FOpenGLTexture2D* Texture2D = ResourceCast(SRV->Texture2D.GetReference());
 		
diff --git a/Engine/Source/Runtime/OpenGLDrv/Private/OpenGLDevice.cpp b/Engine/Source/Runtime/OpenGLDrv/Private/OpenGLDevice.cpp
--- a/Engine/Source/Runtime/OpenGLDrv/Private/OpenGLDevice.cpp
+++ b/Engine/Source/Runtime/OpenGLDrv/Private/OpenGLDevice.cpp
@@ -872,7 +872,7 @@ static void InitRHICapabilitiesForGL()
 
 	GShaderPlatformForFeatureLevel[ERHIFeatureLevel::ES2] = (GMaxRHIFeatureLevel == ERHIFeatureLevel::ES2) ? GMaxRHIShaderPlatform : SP_OPENGL_PCES2;
 	GShaderPlatformForFeatureLevel[ERHIFeatureLevel::ES3_1] = (GMaxRHIFeatureLevel == ERHIFeatureLevel::ES3_1) ? GMaxRHIShaderPlatform : SP_OPENGL_PCES3_1;
-	GShaderPlatformForFeatureLevel[ERHIFeatureLevel::SM4] = SP_OPENGL_SM4;
+	GShaderPlatformForFeatureLevel[ERHIFeatureLevel::SM4_REMOVED] = SP_NumPlatforms;
 	GShaderPlatformForFeatureLevel[ERHIFeatureLevel::SM5] = OPENGL_ESDEFERRED ? SP_OPENGL_ES31_EXT : SP_OPENGL_SM5;
 
 	// Set to same values as in DX11, as for the time being clip space adjustment are done entirely
@@ -927,7 +927,7 @@ static void InitRHICapabilitiesForGL()
 	SetupTextureFormat( PF_R5G6B5_UNORM,		FOpenGLTextureFormat( ));
 #if PLATFORM_DESKTOP || PLATFORM_ANDROIDESDEFERRED || PLATFORM_LUMINGL4
 	CA_SUPPRESS(6286);
-	if (PLATFORM_DESKTOP || FOpenGL::GetFeatureLevel() >= ERHIFeatureLevel::SM4)
+	if (PLATFORM_DESKTOP || FOpenGL::GetFeatureLevel() >= ERHIFeatureLevel::SM5)
 	{
 		// Not supported for rendering:
 		SetupTextureFormat( PF_G16,				FOpenGLTextureFormat( GL_R16,					GL_R16,					GL_RED,			GL_UNSIGNED_SHORT,					false,	false));
diff --git a/Engine/Source/Runtime/OpenGLDrv/Private/OpenGLTexture.cpp b/Engine/Source/Runtime/OpenGLDrv/Private/OpenGLTexture.cpp
--- a/Engine/Source/Runtime/OpenGLDrv/Private/OpenGLTexture.cpp
+++ b/Engine/Source/Runtime/OpenGLDrv/Private/OpenGLTexture.cpp
@@ -1929,7 +1929,7 @@ FShaderResourceViewRHIRef FOpenGLDynamicRHI::RHICreateShaderResourceView(FRHITex
 
 				// For stencil sampling we have to use a separate single channel texture to blit stencil data into
 #if PLATFORM_DESKTOP || PLATFORM_ANDROIDESDEFERRED
-				if (FOpenGL::GetFeatureLevel() >= ERHIFeatureLevel::SM4 && Format == PF_X24_G8 && FOpenGL::SupportsPixelBuffers())
+				if (FOpenGL::GetFeatureLevel() >= ERHIFeatureLevel::SM5 && Format == PF_X24_G8 && FOpenGL::SupportsPixelBuffers())
 				{
 					check(NumMipLevels == 1 && MipLevel == 0);
 
diff --git a/Engine/Source/Runtime/OpenGLDrv/Private/OpenGLUAV.cpp b/Engine/Source/Runtime/OpenGLDrv/Private/OpenGLUAV.cpp
--- a/Engine/Source/Runtime/OpenGLDrv/Private/OpenGLUAV.cpp
+++ b/Engine/Source/Runtime/OpenGLDrv/Private/OpenGLUAV.cpp
@@ -20,12 +20,6 @@ FShaderResourceViewRHIRef FOpenGLDynamicRHI::RHICreateShaderResourceView(FRHIVer
 
 			const uint32 FormatBPP = GPixelFormats[Format].BlockBytes;
 
-			if (FormatBPP != Stride && GMaxRHIFeatureLevel == ERHIFeatureLevel::SM4)
-			{
-				UE_LOG(LogRHI, Fatal,TEXT("OpenGL 3.2 RHI supports only tightly packed texture buffers!"));
-				return new FOpenGLShaderResourceView(OGLRHI, 0, GL_TEXTURE_BUFFER);
-			}
-
 			const FOpenGLTextureFormat& GLFormat = GOpenGLTextureFormats[Format];
 			FOpenGL::GenTextures(1, &TextureID);
 			
diff --git a/Engine/Source/Runtime/OpenGLDrv/Public/OpenGL.h b/Engine/Source/Runtime/OpenGLDrv/Public/OpenGL.h
--- a/Engine/Source/Runtime/OpenGLDrv/Public/OpenGL.h
+++ b/Engine/Source/Runtime/OpenGLDrv/Public/OpenGL.h
@@ -304,7 +304,7 @@ public:
 	static FORCEINLINE const ANSICHAR* GetStringIndexed(GLenum Name, GLuint Index) UGL_REQUIRED(NULL)
 	static FORCEINLINE GLuint GetMajorVersion() UGL_REQUIRED(0)
 	static FORCEINLINE GLuint GetMinorVersion() UGL_REQUIRED(0)
-	static FORCEINLINE ERHIFeatureLevel::Type GetFeatureLevel() UGL_REQUIRED(ERHIFeatureLevel::SM4)
+	static FORCEINLINE ERHIFeatureLevel::Type GetFeatureLevel() UGL_REQUIRED(ERHIFeatureLevel::SM5)
 	static FORCEINLINE EShaderPlatform GetShaderPlatform() UGL_REQUIRED(SP_OPENGL_SM4)
 	static FORCEINLINE FString GetAdapterName() UGL_REQUIRED(TEXT(""))
 	static FORCEINLINE void BlendFuncSeparatei(GLuint Buf, GLenum SrcRGB, GLenum DstRGB, GLenum SrcAlpha, GLenum DstAlpha) UGL_REQUIRED_VOID
diff --git a/Engine/Source/Runtime/OpenGLDrv/Public/OpenGL3.h b/Engine/Source/Runtime/OpenGLDrv/Public/OpenGL3.h
--- a/Engine/Source/Runtime/OpenGLDrv/Public/OpenGL3.h
+++ b/Engine/Source/Runtime/OpenGLDrv/Public/OpenGL3.h
@@ -759,11 +759,11 @@ struct FOpenGL3 : public FOpenGLBase
 		case 2:
 			return ERHIFeatureLevel::ES2;
 		case 3:
-			return ERHIFeatureLevel::SM4;
+			return ERHIFeatureLevel::ES3_1;
 		case 4:
-			return GetMinorVersion() > 2 ? ERHIFeatureLevel::SM5 : ERHIFeatureLevel::SM4;
+			return GetMinorVersion() > 2 ? ERHIFeatureLevel::SM5 : ERHIFeatureLevel::ES3_1;
 		default:
-			return ERHIFeatureLevel::SM4;
+			return ERHIFeatureLevel::ES3_1;
 		}
 	}
 
diff --git a/Engine/Source/Runtime/RHI/Private/Apple/AppleDynamicRHI.cpp b/Engine/Source/Runtime/RHI/Private/Apple/AppleDynamicRHI.cpp
--- a/Engine/Source/Runtime/RHI/Private/Apple/AppleDynamicRHI.cpp
+++ b/Engine/Source/Runtime/RHI/Private/Apple/AppleDynamicRHI.cpp
@@ -107,7 +107,7 @@ FDynamicRHI* PlatformCreateDynamicRHI()
 #if PLATFORM_MAC
 			if (FParse::Param(FCommandLine::Get(),TEXT("metal")))
 			{
-				RequestedFeatureLevel = ERHIFeatureLevel::SM4;
+				RequestedFeatureLevel = ERHIFeatureLevel::SM5;
 			}
 			else if (FParse::Param(FCommandLine::Get(),TEXT("metalsm5")) || FParse::Param(FCommandLine::Get(),TEXT("metalmrt")))
 			{
@@ -132,7 +132,7 @@ FDynamicRHI* PlatformCreateDynamicRHI()
 		if (Sum == 1)
 		{
 #if PLATFORM_MAC
-			RequestedFeatureLevel = ERHIFeatureLevel::SM4;
+			RequestedFeatureLevel = ERHIFeatureLevel::ES3_1;
 #else
 			RequestedFeatureLevel = ERHIFeatureLevel::ES2;
 #endif
diff --git a/Engine/Source/Runtime/RHI/Private/RHI.cpp b/Engine/Source/Runtime/RHI/Private/RHI.cpp
--- a/Engine/Source/Runtime/RHI/Private/RHI.cpp
+++ b/Engine/Source/Runtime/RHI/Private/RHI.cpp
@@ -576,7 +576,7 @@ FName FeatureLevelNames[] =
 {
 	FName(TEXT("ES2")),
 	FName(TEXT("ES3_1")),
-	FName(TEXT("SM4")),
+	FName(TEXT("SM4_REMOVED")),
 	FName(TEXT("SM5")),
 };
 
diff --git a/Engine/Source/Runtime/RHI/Private/Windows/WindowsDynamicRHI.cpp b/Engine/Source/Runtime/RHI/Private/Windows/WindowsDynamicRHI.cpp
--- a/Engine/Source/Runtime/RHI/Private/Windows/WindowsDynamicRHI.cpp
+++ b/Engine/Source/Runtime/RHI/Private/Windows/WindowsDynamicRHI.cpp
@@ -42,15 +42,13 @@ static IDynamicRHIModule* LoadDynamicRHIModule(ERHIFeatureLevel::Type& DesiredFe
 	
 	// command line overrides
 	bool bForceSM5 = FParse::Param(FCommandLine::Get(), TEXT("sm5"));
-	bool bForceSM4 = FParse::Param(FCommandLine::Get(), TEXT("sm4"));
 	bool bForceVulkan = FParse::Param(FCommandLine::Get(), TEXT("vulkan"));
 	bool bForceOpenGL = FWindowsPlatformMisc::VerifyWindowsVersion(6, 0) == false || FParse::Param(FCommandLine::Get(), TEXT("opengl")) || FParse::Param(FCommandLine::Get(), TEXT("opengl3")) || FParse::Param(FCommandLine::Get(), TEXT("opengl4"));
-	bool bForceD3D10 = FParse::Param(FCommandLine::Get(), TEXT("d3d10")) || FParse::Param(FCommandLine::Get(), TEXT("dx10")) || (bForceSM4 && !bForceVulkan && !bForceOpenGL);
 	bool bForceD3D11 = FParse::Param(FCommandLine::Get(), TEXT("d3d11")) || FParse::Param(FCommandLine::Get(), TEXT("dx11")) || (bForceSM5 && !bForceVulkan && !bForceOpenGL);
 	bool bForceD3D12 = FParse::Param(FCommandLine::Get(), TEXT("d3d12")) || FParse::Param(FCommandLine::Get(), TEXT("dx12"));
 	DesiredFeatureLevel = ERHIFeatureLevel::Num;
 	
-	if(!(bForceVulkan||bForceOpenGL||bForceD3D10||bForceD3D11||bForceD3D12))
+	if(!(bForceVulkan||bForceOpenGL||bForceD3D11||bForceD3D12))
 	{
 		//Default graphics RHI is only used if no command line option is specified
 		FConfigFile EngineSettings;
@@ -80,16 +78,11 @@ static IDynamicRHIModule* LoadDynamicRHIModule(ERHIFeatureLevel::Type& DesiredFe
 
 
 
-	int32 Sum = ((bForceD3D12 ? 1 : 0) + (bForceD3D11 ? 1 : 0) + (bForceD3D10 ? 1 : 0) + (bForceOpenGL ? 1 : 0) + (bForceVulkan ? 1 : 0));
-
-	if (bForceSM5 && bForceSM4)
-	{
-		UE_LOG(LogRHI, Fatal, TEXT("-sm4 and -sm5 are mutually exclusive options, but more than one was specified on the command-line."));
-	}
+	int32 Sum = ((bForceD3D12 ? 1 : 0) + (bForceD3D11 ? 1 : 0) + (bForceOpenGL ? 1 : 0) + (bForceVulkan ? 1 : 0));
 
 	if (Sum > 1)
 	{
-		UE_LOG(LogRHI, Fatal, TEXT("-d3d12, -d3d11, -d3d10, -vulkan, and -opengl[3|4] are mutually exclusive options, but more than one was specified on the command-line."));
+		UE_LOG(LogRHI, Fatal, TEXT("-d3d12, -d3d11, -vulkan, and -opengl[3|4] are mutually exclusive options, but more than one was specified on the command-line."));
 	}
 	else if (Sum == 0)
 	{
@@ -113,12 +106,6 @@ static IDynamicRHIModule* LoadDynamicRHIModule(ERHIFeatureLevel::Type& DesiredFe
 		{
 			DesiredFeatureLevel = ERHIFeatureLevel::SM5;
 		}
-
-		if (bForceSM4)
-		{
-			DesiredFeatureLevel = ERHIFeatureLevel::SM4;
-			bPreferD3D12 = false;
-		}
 	}
 
 	// Load the dynamic RHI module.
diff --git a/Engine/Source/Runtime/RHI/Public/RHI.h b/Engine/Source/Runtime/RHI/Public/RHI.h
--- a/Engine/Source/Runtime/RHI/Public/RHI.h
+++ b/Engine/Source/Runtime/RHI/Public/RHI.h
@@ -157,12 +157,12 @@ inline bool RHISupportsBufferLoadTypeConversion(EShaderPlatform Platform)
 /** Whether the platform supports reading from volume textures (does not cover rendering to volume textures). */
 inline bool RHISupportsVolumeTextures(ERHIFeatureLevel::Type FeatureLevel)
 {
-	return FeatureLevel >= ERHIFeatureLevel::SM4;
+	return FeatureLevel >= ERHIFeatureLevel::SM5;
 }
 
 inline bool RHISupportsVertexShaderLayer(const EShaderPlatform Platform)
 {
-	return IsFeatureLevelSupported(Platform, ERHIFeatureLevel::SM4) && IsMetalPlatform(Platform) && (IsPCPlatform(Platform) || (Platform == SP_METAL_MRT && RHIGetShaderLanguageVersion(Platform) >= 4));
+	return IsFeatureLevelSupported(Platform, ERHIFeatureLevel::SM5) && IsMetalPlatform(Platform) && (IsPCPlatform(Platform) || (Platform == SP_METAL_MRT && RHIGetShaderLanguageVersion(Platform) >= 4));
 }
 
 /** Return true if and only if the GPU support rendering to volume textures (2D Array, 3D) is guaranteed supported for a target platform.
@@ -170,7 +170,7 @@ inline bool RHISupportsVertexShaderLayer(const EShaderPlatform Platform)
 */
 inline bool RHIVolumeTextureRenderingSupportGuaranteed(const EShaderPlatform Platform)
 {
-	return IsFeatureLevelSupported(Platform, ERHIFeatureLevel::SM4)
+	return IsFeatureLevelSupported(Platform, ERHIFeatureLevel::SM5)
 		&& (!IsMetalPlatform(Platform) || RHISupportsVertexShaderLayer(Platform)) // For Metal only shader platforms & versions that support vertex-shader-layer can render to volume textures - this is a compile/cook time check.
 		&& !IsOpenGLPlatform(Platform);		// Apparently, some OpenGL 3.3 cards support SM4 but can't render to volume textures
 }
diff --git a/Engine/Source/Runtime/RHI/Public/RHIDefinitions.h b/Engine/Source/Runtime/RHI/Public/RHIDefinitions.h
--- a/Engine/Source/Runtime/RHI/Public/RHIDefinitions.h
+++ b/Engine/Source/Runtime/RHI/Public/RHIDefinitions.h
@@ -107,11 +107,9 @@ namespace ERHIFeatureLevel
 
 		/**
 		 * Feature level defined by the capabilities of DX10 Shader Model 4.
-		 *    Shaders can reference 128 unique textures, 16 texture samplers, 14 constant buffers
-		 *    Rendering to cubemaps and volume textures
-		 * Geometry shaders are not required for Feature Level SM4, as long as vertex shader render target selection is supported
+		 * SUPPORT FOR THIS FEATURE LEVEL HAS BEEN ENTIRELY REMOVED.
 		 */
-		SM4,
+		SM4_REMOVED,
 
 		/**
 		 * Feature level defined by the capabilities of DX11 Shader Model 5.
@@ -1063,7 +1061,7 @@ inline ERHIFeatureLevel::Type GetMaxSupportedFeatureLevel(EShaderPlatform InShad
 	case SP_VULKAN_SM4:
 	case SP_PCD3D_SM4:
 	case SP_OPENGL_SM4:
-		return ERHIFeatureLevel::SM4;
+		return ERHIFeatureLevel::SM4_REMOVED;
 	case SP_PCD3D_ES2:
 	case SP_OPENGL_PCES2:
 	case SP_OPENGL_ES2_ANDROID:
@@ -1167,7 +1165,7 @@ inline bool RHISupportsComputeShaders(const EShaderPlatform Platform)
 
 inline bool RHISupportsGeometryShaders(const EShaderPlatform Platform)
 {
-	return IsFeatureLevelSupported(Platform, ERHIFeatureLevel::SM4) && !IsMetalPlatform(Platform) && !IsVulkanMobilePlatform(Platform) && Platform != SP_VULKAN_SM5_LUMIN;
+	return IsFeatureLevelSupported(Platform, ERHIFeatureLevel::SM5) && !IsMetalPlatform(Platform) && !IsVulkanMobilePlatform(Platform) && Platform != SP_VULKAN_SM5_LUMIN;
 }
 
 inline bool RHIHasTiledGPU(const EShaderPlatform Platform)
@@ -1223,8 +1221,7 @@ inline int32 GetFeatureLevelMaxNumberOfBones(ERHIFeatureLevel::Type FeatureLevel
 	{
 	case ERHIFeatureLevel::ES2:
 	case ERHIFeatureLevel::ES3_1:
-		return 75;
-	case ERHIFeatureLevel::SM4:
+		return 75;	
 	case ERHIFeatureLevel::SM5:
 		return 256;
 	default:
diff --git a/Engine/Source/Runtime/RenderCore/Private/RenderUtils.cpp b/Engine/Source/Runtime/RenderCore/Private/RenderUtils.cpp
--- a/Engine/Source/Runtime/RenderCore/Private/RenderUtils.cpp
+++ b/Engine/Source/Runtime/RenderCore/Private/RenderUtils.cpp
@@ -344,7 +344,7 @@ public:
 	// FResource interface.
 	virtual void InitRHI() override
 	{
-		if (GetFeatureLevel() >= ERHIFeatureLevel::SM4)
+		if (GetFeatureLevel() >= ERHIFeatureLevel::SM5)
 		{
 			// Create the texture RHI.
 			FBlackVolumeTextureResourceBulkDataInterface BlackTextureBulkData(0);
diff --git a/Engine/Source/Runtime/RenderCore/Private/Shader.cpp b/Engine/Source/Runtime/RenderCore/Private/Shader.cpp
--- a/Engine/Source/Runtime/RenderCore/Private/Shader.cpp
+++ b/Engine/Source/Runtime/RenderCore/Private/Shader.cpp
@@ -2419,7 +2419,7 @@ void ShaderMapAppendKeyString(EShaderPlatform Platform, FString& KeyString)
 		}
 	}
 
-	if (IsFeatureLevelSupported(Platform, ERHIFeatureLevel::SM4))
+	if (IsFeatureLevelSupported(Platform, ERHIFeatureLevel::SM5))
 	{
 		static const auto CVar = IConsoleManager::Get().FindTConsoleVariableDataInt(TEXT("r.StencilForLODDither"));
 		if (CVar && CVar->GetValueOnAnyThread() > 0)
diff --git a/Engine/Source/Runtime/Renderer/Private/AtmosphereRendering.cpp b/Engine/Source/Runtime/Renderer/Private/AtmosphereRendering.cpp
--- a/Engine/Source/Runtime/Renderer/Private/AtmosphereRendering.cpp
+++ b/Engine/Source/Runtime/Renderer/Private/AtmosphereRendering.cpp
@@ -238,7 +238,7 @@ class FAtmosphericFogPS : public FGlobalShader
 public:
 	static bool ShouldCompilePermutation(const FGlobalShaderPermutationParameters& Parameters)
 	{
-		return IsFeatureLevelSupported(Parameters.Platform, ERHIFeatureLevel::SM4);
+		return IsFeatureLevelSupported(Parameters.Platform, ERHIFeatureLevel::SM5);
 	}
 
 	FAtmosphericFogPS() {}
@@ -363,7 +363,7 @@ public:
 
 	static bool ShouldCompilePermutation(const FGlobalShaderPermutationParameters& Parameters)
 	{
-		return IsFeatureLevelSupported(Parameters.Platform, ERHIFeatureLevel::SM4);
+		return IsFeatureLevelSupported(Parameters.Platform, ERHIFeatureLevel::SM5);
 	}
 
 	FAtmosphericVS( )	{ }
@@ -476,7 +476,7 @@ void SetAtmosphericFogShaders(FRHICommandList& RHICmdList, FGraphicsPipelineStat
 void FDeferredShadingSceneRenderer::RenderAtmosphere(FRHICommandListImmediate& RHICmdList, const FLightShaftsOutput& LightShaftsOutput)
 {
 	// Atmospheric fog?
-	if (Scene->GetFeatureLevel() >= ERHIFeatureLevel::SM4 && Scene->HasAtmosphericFog())
+	if (Scene->GetFeatureLevel() >= ERHIFeatureLevel::SM5 && Scene->HasAtmosphericFog())
 	{
 		FSceneRenderTargets& SceneContext = FSceneRenderTargets::Get(RHICmdList);
 		
@@ -524,7 +524,7 @@ class FAtmosphereTransmittancePS : public FGlobalShader
 
 	static bool ShouldCompilePermutation(const FGlobalShaderPermutationParameters& Parameters)
 	{
-		return IsFeatureLevelSupported(Parameters.Platform, ERHIFeatureLevel::SM4);
+		return IsFeatureLevelSupported(Parameters.Platform, ERHIFeatureLevel::SM5);
 	}
 
 	/* Default constructor. */
@@ -556,7 +556,7 @@ class FAtmosphereIrradiance1PS : public FGlobalShader
 
 	static bool ShouldCompilePermutation(const FGlobalShaderPermutationParameters& Parameters)
 	{
-		return IsFeatureLevelSupported(Parameters.Platform, ERHIFeatureLevel::SM4);
+		return IsFeatureLevelSupported(Parameters.Platform, ERHIFeatureLevel::SM5);
 	}
 
 	/* Default constructor. */
@@ -592,7 +592,7 @@ class FAtmosphereIrradianceNPS : public FGlobalShader
 
 	static bool ShouldCompilePermutation(const FGlobalShaderPermutationParameters& Parameters)
 	{
-		return IsFeatureLevelSupported(Parameters.Platform, ERHIFeatureLevel::SM4);
+		return IsFeatureLevelSupported(Parameters.Platform, ERHIFeatureLevel::SM5);
 	}
 
 	/* Default constructor. */
@@ -636,7 +636,7 @@ class FAtmosphereCopyIrradiancePS : public FGlobalShader
 
 	static bool ShouldCompilePermutation(const FGlobalShaderPermutationParameters& Parameters)
 	{
-		return IsFeatureLevelSupported(Parameters.Platform, ERHIFeatureLevel::SM4);
+		return IsFeatureLevelSupported(Parameters.Platform, ERHIFeatureLevel::SM5);
 	}
 
 	/* Default constructor. */
@@ -672,7 +672,7 @@ class FAtmosphereGS : public FGlobalShader
 
 	static bool ShouldCompilePermutation(const FGlobalShaderPermutationParameters& Parameters) 
 	{ 
-		return IsFeatureLevelSupported(Parameters.Platform, ERHIFeatureLevel::SM4) && RHISupportsGeometryShaders(Parameters.Platform);
+		return IsFeatureLevelSupported(Parameters.Platform, ERHIFeatureLevel::SM5) && RHISupportsGeometryShaders(Parameters.Platform);
 	}
 
 	FAtmosphereGS() {}
@@ -705,7 +705,7 @@ class FAtmosphereInscatter1PS : public FGlobalShader
 
 	static bool ShouldCompilePermutation(const FGlobalShaderPermutationParameters& Parameters)
 	{
-		return IsFeatureLevelSupported(Parameters.Platform, ERHIFeatureLevel::SM4);
+		return IsFeatureLevelSupported(Parameters.Platform, ERHIFeatureLevel::SM5);
 	}
 
 	/* Default constructor. */
@@ -748,7 +748,7 @@ class FAtmosphereCopyInscatter1PS : public FGlobalShader
 
 	static bool ShouldCompilePermutation(const FGlobalShaderPermutationParameters& Parameters)
 	{
-		return IsFeatureLevelSupported(Parameters.Platform, ERHIFeatureLevel::SM4);
+		return IsFeatureLevelSupported(Parameters.Platform, ERHIFeatureLevel::SM5);
 	}
 
 	/* Default constructor. */
@@ -796,7 +796,7 @@ class FAtmosphereCopyInscatterNPS : public FGlobalShader
 
 	static bool ShouldCompilePermutation(const FGlobalShaderPermutationParameters& Parameters)
 	{
-		return IsFeatureLevelSupported(Parameters.Platform, ERHIFeatureLevel::SM4);
+		return IsFeatureLevelSupported(Parameters.Platform, ERHIFeatureLevel::SM5);
 	}
 
 	/* Default constructor. */
@@ -842,7 +842,7 @@ class FAtmosphereInscatterSPS : public FGlobalShader
 
 	static bool ShouldCompilePermutation(const FGlobalShaderPermutationParameters& Parameters)
 	{
-		return IsFeatureLevelSupported(Parameters.Platform, ERHIFeatureLevel::SM4);
+		return IsFeatureLevelSupported(Parameters.Platform, ERHIFeatureLevel::SM5);
 	}
 
 	/* Default constructor. */
@@ -894,7 +894,7 @@ class FAtmosphereInscatterNPS : public FGlobalShader
 
 	static bool ShouldCompilePermutation(const FGlobalShaderPermutationParameters& Parameters)
 	{
-		return IsFeatureLevelSupported(Parameters.Platform, ERHIFeatureLevel::SM4);
+		return IsFeatureLevelSupported(Parameters.Platform, ERHIFeatureLevel::SM5);
 	}
 
 	/* Default constructor. */
@@ -942,7 +942,7 @@ class FAtmospherePrecomputeVS : public FGlobalShader
 
 	static bool ShouldCompilePermutation(const FGlobalShaderPermutationParameters& Parameters)
 	{
-		return IsFeatureLevelSupported(Parameters.Platform, ERHIFeatureLevel::SM4);
+		return IsFeatureLevelSupported(Parameters.Platform, ERHIFeatureLevel::SM5);
 	}
 
 	/* Default constructor. */
@@ -969,7 +969,7 @@ class FAtmospherePrecomputeInscatterVS : public FGlobalShader
 	
 	static bool ShouldCompilePermutation(const FGlobalShaderPermutationParameters& Parameters)
 	{
-		return IsFeatureLevelSupported(Parameters.Platform, ERHIFeatureLevel::SM4);
+		return IsFeatureLevelSupported(Parameters.Platform, ERHIFeatureLevel::SM5);
 	}
 	
 	/* Default constructor. */
@@ -1006,7 +1006,7 @@ class FAtmosphereCopyInscatterFPS : public FGlobalShader
 
 	static bool ShouldCompilePermutation(const FGlobalShaderPermutationParameters& Parameters)
 	{
-		return IsFeatureLevelSupported(Parameters.Platform, ERHIFeatureLevel::SM4);
+		return IsFeatureLevelSupported(Parameters.Platform, ERHIFeatureLevel::SM5);
 	}
 
 	/* Default constructor. */
@@ -1051,7 +1051,7 @@ class FAtmosphereCopyInscatterFBackPS : public FGlobalShader
 
 	static bool ShouldCompilePermutation(const FGlobalShaderPermutationParameters& Parameters)
 	{
-		return IsFeatureLevelSupported(Parameters.Platform, ERHIFeatureLevel::SM4);
+		return IsFeatureLevelSupported(Parameters.Platform, ERHIFeatureLevel::SM5);
 	}
 
 	/* Default constructor. */
diff --git a/Engine/Source/Runtime/Renderer/Private/BasePassRendering.h b/Engine/Source/Runtime/Renderer/Private/BasePassRendering.h
--- a/Engine/Source/Runtime/Renderer/Private/BasePassRendering.h
+++ b/Engine/Source/Runtime/Renderer/Private/BasePassRendering.h
@@ -289,7 +289,7 @@ public:
 		bShouldCache &= (bEnableAtmosphericFog && bProjectAllowsAtmosphericFog && IsTranslucentBlendMode(Parameters.Material->GetBlendMode())) || !bEnableAtmosphericFog;
 
 		return bShouldCache
-			&& (IsFeatureLevelSupported(Parameters.Platform, ERHIFeatureLevel::SM4));
+			&& (IsFeatureLevelSupported(Parameters.Platform, ERHIFeatureLevel::SM5));
 	}
 
 	static void ModifyCompilationEnvironment(const FMaterialShaderPermutationParameters& Parameters, FShaderCompilerEnvironment& OutEnvironment)
@@ -498,7 +498,7 @@ public:
 			|| LightMapPolicyType::RequiresSkylight()
 			|| ((bProjectSupportsStationarySkylight || IsForwardShadingEnabled(Parameters.Platform)) && Parameters.Material->GetShadingModels().IsLit());
 		return bCacheShaders
-			&& (IsFeatureLevelSupported(Parameters.Platform, ERHIFeatureLevel::SM4))
+			&& (IsFeatureLevelSupported(Parameters.Platform, ERHIFeatureLevel::SM5))
 			&& TBasePassPixelShaderBaseType<LightMapPolicyType>::ShouldCompilePermutation(Parameters);
 	}
 
diff --git a/Engine/Source/Runtime/Renderer/Private/CompositionLighting/PostProcessAmbientOcclusion.cpp b/Engine/Source/Runtime/Renderer/Private/CompositionLighting/PostProcessAmbientOcclusion.cpp
--- a/Engine/Source/Runtime/Renderer/Private/CompositionLighting/PostProcessAmbientOcclusion.cpp
+++ b/Engine/Source/Runtime/Renderer/Private/CompositionLighting/PostProcessAmbientOcclusion.cpp
@@ -294,7 +294,7 @@ class FPostProcessAmbientOcclusionSetupPS : public FGlobalShader
 
 	static bool ShouldCompilePermutation(const FGlobalShaderPermutationParameters& Parameters)
 	{
-		return IsFeatureLevelSupported(Parameters.Platform, ERHIFeatureLevel::SM4);
+		return IsFeatureLevelSupported(Parameters.Platform, ERHIFeatureLevel::SM5);
 	}
 
 	static void ModifyCompilationEnvironment(const FGlobalShaderPermutationParameters& Parameters, FShaderCompilerEnvironment& OutEnvironment)
@@ -674,7 +674,7 @@ class FPostProcessAmbientOcclusionPSandCS : public FGlobalShader
 		}
 		else
 		{
-			return IsFeatureLevelSupported(Parameters.Platform, ERHIFeatureLevel::SM4);
+			return IsFeatureLevelSupported(Parameters.Platform, ERHIFeatureLevel::SM5);
 		}
 	}
 
diff --git a/Engine/Source/Runtime/Renderer/Private/CompositionLighting/PostProcessDeferredDecals.cpp b/Engine/Source/Runtime/Renderer/Private/CompositionLighting/PostProcessDeferredDecals.cpp
--- a/Engine/Source/Runtime/Renderer/Private/CompositionLighting/PostProcessDeferredDecals.cpp
+++ b/Engine/Source/Runtime/Renderer/Private/CompositionLighting/PostProcessDeferredDecals.cpp
@@ -59,30 +59,7 @@ FRHIBlendState* GetDecalBlendState(const ERHIFeatureLevel::Type SMFeatureLevel,
 	if (InDecalRenderStage == DRS_BeforeBasePass)
 	{
 		// before base pass (for DBuffer decals)
-
-		if (SMFeatureLevel == ERHIFeatureLevel::SM4)
-		{
-			// DX10 doesn't support masking/using different blend modes per MRT.
-			// We set the opacity in the shader to 0 so we can use the same frame buffer blend.
-
-			if (DecalBlendMode == DBM_DBuffer_AlphaComposite)
-			{
-				return TStaticBlendState<
-					CW_RGBA, BO_Add, BF_One, BF_InverseSourceAlpha, BO_Add, BF_Zero, BF_InverseSourceAlpha,
-					CW_RGBA, BO_Add, BF_One, BF_InverseSourceAlpha, BO_Add, BF_Zero, BF_InverseSourceAlpha,
-					CW_RGBA, BO_Add, BF_One, BF_InverseSourceAlpha, BO_Add, BF_Zero, BF_InverseSourceAlpha >::GetRHI();
-			}
-			else
-			{
-				return TStaticBlendState<
-					CW_RGBA, BO_Add, BF_SourceAlpha, BF_InverseSourceAlpha, BO_Add, BF_Zero, BF_InverseSourceAlpha,
-					CW_RGBA, BO_Add, BF_SourceAlpha, BF_InverseSourceAlpha, BO_Add, BF_Zero, BF_InverseSourceAlpha,
-					CW_RGBA, BO_Add, BF_SourceAlpha, BF_InverseSourceAlpha, BO_Add, BF_Zero, BF_InverseSourceAlpha >::GetRHI();
-			}
-		}
-		else
 		{
-			// see DX10 comment above
 			// As we set the opacity in the shader we don't need to set different frame buffer blend modes but we like to hint to the driver that we
 			// don't need to output there. We also could replace this with many SetRenderTarget calls but it might be slower (needs to be tested).
 
@@ -207,15 +184,6 @@ FRHIBlendState* GetDecalBlendState(const ERHIFeatureLevel::Type SMFeatureLevel,
 					>::GetRHI();
 				}
 			}
-			else if (SMFeatureLevel == ERHIFeatureLevel::SM4)
-			{
-				return TStaticBlendState<
-					CW_RGB, BO_Add, BF_SourceAlpha, BF_One, BO_Add, BF_Zero, BF_One,	// Emissive
-					CW_RGB, BO_Add, BF_SourceAlpha, BF_One, BO_Add, BF_Zero, BF_One,	// Normal
-					CW_RGB, BO_Add, BF_SourceAlpha, BF_One, BO_Add, BF_Zero, BF_One,	// Metallic, Specular, Roughness
-					CW_RGB, BO_Add, BF_SourceAlpha, BF_One, BO_Add, BF_Zero, BF_One		// BaseColor
-				>::GetRHI();
-			}
 
 		case DBM_Stain:
 			if (GSupportsSeparateRenderTargetBlendState)
@@ -239,15 +207,6 @@ FRHIBlendState* GetDecalBlendState(const ERHIFeatureLevel::Type SMFeatureLevel,
 					>::GetRHI();
 				}
 			}
-			else if (SMFeatureLevel == ERHIFeatureLevel::SM4)
-			{
-				return TStaticBlendState<
-					CW_RGB, BO_Add, BF_SourceAlpha, BF_InverseSourceAlpha, BO_Add, BF_Zero, BF_One,	// Emissive
-					CW_RGB, BO_Add, BF_SourceAlpha, BF_InverseSourceAlpha, BO_Add, BF_Zero, BF_One,	// Normal
-					CW_RGB, BO_Add, BF_SourceAlpha, BF_InverseSourceAlpha, BO_Add, BF_Zero, BF_One,	// Metallic, Specular, Roughness
-					CW_RGB, BO_Add, BF_SourceAlpha, BF_InverseSourceAlpha, BO_Add, BF_Zero, BF_One		// BaseColor
-				>::GetRHI();
-			}
 
 		case DBM_Normal:
 			return TStaticBlendState< CW_RGB, BO_Add, BF_SourceAlpha, BF_InverseSourceAlpha >::GetRHI();
@@ -270,15 +229,6 @@ FRHIBlendState* GetDecalBlendState(const ERHIFeatureLevel::Type SMFeatureLevel,
 					CW_RGB, BO_Add, BF_One, BF_InverseSourceAlpha, BO_Add, BF_Zero, BF_One	// BaseColor
 				>::GetRHI();
 			}
-			else if (SMFeatureLevel == ERHIFeatureLevel::SM4)
-			{
-				return TStaticBlendState<
-					CW_RGB, BO_Add, BF_One, BF_InverseSourceAlpha, BO_Add, BF_Zero, BF_One,	// Emissive
-					CW_RGB, BO_Add, BF_One, BF_InverseSourceAlpha, BO_Add, BF_Zero, BF_One,	// Normal
-					CW_RGB, BO_Add, BF_One, BF_InverseSourceAlpha, BO_Add, BF_Zero, BF_One,	// Metallic, Specular, Roughness
-					CW_RGB, BO_Add, BF_One, BF_InverseSourceAlpha, BO_Add, BF_Zero, BF_One	// BaseColor
-				>::GetRHI();
-			}
 		default:
 			// the decal type should not be rendered in this pass - internal error
 			check(0);
diff --git a/Engine/Source/Runtime/Renderer/Private/CompositionLighting/PostProcessMeshDecals.cpp b/Engine/Source/Runtime/Renderer/Private/CompositionLighting/PostProcessMeshDecals.cpp
--- a/Engine/Source/Runtime/Renderer/Private/CompositionLighting/PostProcessMeshDecals.cpp
+++ b/Engine/Source/Runtime/Renderer/Private/CompositionLighting/PostProcessMeshDecals.cpp
@@ -33,7 +33,7 @@ class FMeshDecalAccumulatePolicy
 public:
 	static bool ShouldCompilePermutation(const FMeshMaterialShaderPermutationParameters& Parameters)
 	{
-		return Parameters.Material && Parameters.Material->IsDeferredDecal() && IsFeatureLevelSupported(Parameters.Platform, ERHIFeatureLevel::SM4);
+		return Parameters.Material && Parameters.Material->IsDeferredDecal() && IsFeatureLevelSupported(Parameters.Platform, ERHIFeatureLevel::SM5);
 	}
 };
 
diff --git a/Engine/Source/Runtime/Renderer/Private/CompositionLighting/PostProcessPassThrough.h b/Engine/Source/Runtime/Renderer/Private/CompositionLighting/PostProcessPassThrough.h
--- a/Engine/Source/Runtime/Renderer/Private/CompositionLighting/PostProcessPassThrough.h
+++ b/Engine/Source/Runtime/Renderer/Private/CompositionLighting/PostProcessPassThrough.h
@@ -20,7 +20,7 @@ class FPostProcessPassThroughPS : public FGlobalShader
 
 	static bool ShouldCompilePermutation(const FGlobalShaderPermutationParameters& Parameters)
 	{
-		return IsFeatureLevelSupported(Parameters.Platform, ERHIFeatureLevel::SM4);
+		return IsFeatureLevelSupported(Parameters.Platform, ERHIFeatureLevel::SM5);
 	}
 
 	/** Default constructor. */
diff --git a/Engine/Source/Runtime/Renderer/Private/DeferredShadingRenderer.cpp b/Engine/Source/Runtime/Renderer/Private/DeferredShadingRenderer.cpp
--- a/Engine/Source/Runtime/Renderer/Private/DeferredShadingRenderer.cpp
+++ b/Engine/Source/Runtime/Renderer/Private/DeferredShadingRenderer.cpp
@@ -1016,7 +1016,7 @@ void FDeferredShadingSceneRenderer::Render(FRHICommandListImmediate& RHICmdList)
 	bool bCanOverlayRayTracingOutput = CanOverlayRayTracingOutput(Views[0]);// #dxr_todo: UE-72557 multi-view case
 	
 	const bool bRenderDeferredLighting = ViewFamily.EngineShowFlags.Lighting
-		&& FeatureLevel >= ERHIFeatureLevel::SM4
+		&& FeatureLevel >= ERHIFeatureLevel::SM5
 		&& ViewFamily.EngineShowFlags.DeferredLighting
 		&& bUseGBuffer
 		&& bCanOverlayRayTracingOutput;
@@ -1639,7 +1639,7 @@ void FDeferredShadingSceneRenderer::Render(FRHICommandListImmediate& RHICmdList)
 
 	// Pre-lighting composition lighting stage
 	// e.g. deferred decals, SSAO
-	if (FeatureLevel >= ERHIFeatureLevel::SM4)
+	if (FeatureLevel >= ERHIFeatureLevel::SM5)
 	{
 		CSV_SCOPED_TIMING_STAT_EXCLUSIVE(AfterBasePass);
 		SCOPE_CYCLE_COUNTER(STAT_FDeferredShadingSceneRenderer_AfterBasePass);
@@ -2023,7 +2023,7 @@ void FDeferredShadingSceneRenderer::Render(FRHICommandListImmediate& RHICmdList)
 	}
 
 	if (ViewFamily.EngineShowFlags.StationaryLightOverlap &&
-		FeatureLevel >= ERHIFeatureLevel::SM4)
+		FeatureLevel >= ERHIFeatureLevel::SM5)
 	{
 		RenderStationaryLightOverlap(RHICmdList);
 		ServiceLocalQueue();
@@ -2118,7 +2118,7 @@ public:
 
 	static bool ShouldCompilePermutation(const FGlobalShaderPermutationParameters& Parameters)
 	{ 
-		return IsFeatureLevelSupported(Parameters.Platform, ERHIFeatureLevel::SM4);
+		return IsFeatureLevelSupported(Parameters.Platform, ERHIFeatureLevel::SM5);
 	}
 
 	FDownsampleSceneDepthPS(const ShaderMetaType::CompiledShaderInitializerType& Initializer):
@@ -2186,7 +2186,7 @@ IMPLEMENT_SHADER_TYPE(,FDownsampleSceneDepthPS,TEXT("/Engine/Private/DownsampleD
 void FDeferredShadingSceneRenderer::UpdateDownsampledDepthSurface(FRHICommandList& RHICmdList)
 {
 	FSceneRenderTargets& SceneContext = FSceneRenderTargets::Get(RHICmdList);
-	if (SceneContext.UseDownsizedOcclusionQueries() && (FeatureLevel >= ERHIFeatureLevel::SM4))
+	if (SceneContext.UseDownsizedOcclusionQueries() && (FeatureLevel >= ERHIFeatureLevel::SM5))
 	{
 		RHICmdList.TransitionResource( EResourceTransitionAccess::EReadable, SceneContext.GetSceneDepthSurface() );
 
@@ -2262,7 +2262,7 @@ public:
 
 	static bool ShouldCompilePermutation(const FGlobalShaderPermutationParameters& Parameters)
 	{ 
-		return IsFeatureLevelSupported(Parameters.Platform, ERHIFeatureLevel::SM4);
+		return IsFeatureLevelSupported(Parameters.Platform, ERHIFeatureLevel::SM5);
 	}
 
 	static void ModifyCompilationEnvironment(const FGlobalShaderPermutationParameters& Parameters, FShaderCompilerEnvironment& OutEnvironment)
diff --git a/Engine/Source/Runtime/Renderer/Private/DepthRendering.cpp b/Engine/Source/Runtime/Renderer/Private/DepthRendering.cpp
--- a/Engine/Source/Runtime/Renderer/Private/DepthRendering.cpp
+++ b/Engine/Source/Runtime/Renderer/Private/DepthRendering.cpp
@@ -319,7 +319,7 @@ public:
 
 	static bool ShouldCompilePermutation(const FGlobalShaderPermutationParameters& Parameters)
 	{ 
-		return IsFeatureLevelSupported(Parameters.Platform, ERHIFeatureLevel::SM4);
+		return IsFeatureLevelSupported(Parameters.Platform, ERHIFeatureLevel::SM5);
 	}
 
 	FDitheredTransitionStencilPS(const ShaderMetaType::CompiledShaderInitializerType& Initializer)
diff --git a/Engine/Source/Runtime/Renderer/Private/DistanceFieldAmbientOcclusion.cpp b/Engine/Source/Runtime/Renderer/Private/DistanceFieldAmbientOcclusion.cpp
--- a/Engine/Source/Runtime/Renderer/Private/DistanceFieldAmbientOcclusion.cpp
+++ b/Engine/Source/Runtime/Renderer/Private/DistanceFieldAmbientOcclusion.cpp
@@ -727,7 +727,7 @@ bool ShouldRenderDeferredDynamicSkyLight(const FScene* Scene, const FSceneViewFa
 		&& !Scene->SkyLight->bWantsStaticShadowing
 		&& !Scene->SkyLight->bHasStaticLighting
 		&& ViewFamily.EngineShowFlags.SkyLighting
-		&& Scene->GetFeatureLevel() >= ERHIFeatureLevel::SM4
+		&& Scene->GetFeatureLevel() >= ERHIFeatureLevel::SM5
 		&& !IsAnyForwardShadingEnabled(Scene->GetShaderPlatform())
 		&& !ViewFamily.EngineShowFlags.VisualizeLightCulling;
 }
diff --git a/Engine/Source/Runtime/Renderer/Private/FogRendering.cpp b/Engine/Source/Runtime/Renderer/Private/FogRendering.cpp
--- a/Engine/Source/Runtime/Renderer/Private/FogRendering.cpp
+++ b/Engine/Source/Runtime/Renderer/Private/FogRendering.cpp
@@ -98,7 +98,7 @@ public:
 
 	static bool ShouldCompilePermutation(const FGlobalShaderPermutationParameters& Parameters)
 	{
-		return IsFeatureLevelSupported(Parameters.Platform, ERHIFeatureLevel::SM4);
+		return IsFeatureLevelSupported(Parameters.Platform, ERHIFeatureLevel::SM5);
 	}
 
 	FHeightFogVS( )	{ }
@@ -170,7 +170,7 @@ public:
 
 	static bool ShouldCompilePermutation(const FGlobalShaderPermutationParameters& Parameters)
 	{
-		return IsFeatureLevelSupported(Parameters.Platform, ERHIFeatureLevel::SM4);
+		return IsFeatureLevelSupported(Parameters.Platform, ERHIFeatureLevel::SM5);
 	}
 
 	static void ModifyCompilationEnvironment(const FGlobalShaderPermutationParameters& Parameters, FShaderCompilerEnvironment& OutEnvironment)
diff --git a/Engine/Source/Runtime/Renderer/Private/GPUBenchmark.cpp b/Engine/Source/Runtime/Renderer/Private/GPUBenchmark.cpp
--- a/Engine/Source/Runtime/Renderer/Private/GPUBenchmark.cpp
+++ b/Engine/Source/Runtime/Renderer/Private/GPUBenchmark.cpp
@@ -36,7 +36,7 @@ class FPostProcessBenchmarkPS : public FGlobalShader
 
 	static bool ShouldCompilePermutation(const FGlobalShaderPermutationParameters& Parameters)
 	{
-		return IsFeatureLevelSupported(Parameters.Platform, ERHIFeatureLevel::SM4);
+		return IsFeatureLevelSupported(Parameters.Platform, ERHIFeatureLevel::SM5);
 	}
 
 	static void ModifyCompilationEnvironment(const FGlobalShaderPermutationParameters& Parameters, FShaderCompilerEnvironment& OutEnvironment)
diff --git a/Engine/Source/Runtime/Renderer/Private/IndirectLightRendering.cpp b/Engine/Source/Runtime/Renderer/Private/IndirectLightRendering.cpp
--- a/Engine/Source/Runtime/Renderer/Private/IndirectLightRendering.cpp
+++ b/Engine/Source/Runtime/Renderer/Private/IndirectLightRendering.cpp
@@ -51,7 +51,7 @@ class FDiffuseIndirectCompositePS : public FGlobalShader
 			return IsFeatureLevelSupported(Parameters.Platform, ERHIFeatureLevel::SM5);
 		}
 
-		return IsFeatureLevelSupported(Parameters.Platform, ERHIFeatureLevel::SM4);
+		return IsFeatureLevelSupported(Parameters.Platform, ERHIFeatureLevel::SM5);
 	}
 
 	BEGIN_SHADER_PARAMETER_STRUCT(FParameters, )
@@ -78,7 +78,7 @@ class FAmbientCubemapCompositePS : public FGlobalShader
 
 	static bool ShouldCompilePermutation(const FGlobalShaderPermutationParameters& Parameters)
 	{
-		return IsFeatureLevelSupported(Parameters.Platform, ERHIFeatureLevel::SM4);
+		return IsFeatureLevelSupported(Parameters.Platform, ERHIFeatureLevel::SM5);
 	}
 
 	BEGIN_SHADER_PARAMETER_STRUCT(FParameters, )
diff --git a/Engine/Source/Runtime/Renderer/Private/IndirectLightingCache.cpp b/Engine/Source/Runtime/Renderer/Private/IndirectLightingCache.cpp
--- a/Engine/Source/Runtime/Renderer/Private/IndirectLightingCache.cpp
+++ b/Engine/Source/Runtime/Renderer/Private/IndirectLightingCache.cpp
@@ -143,7 +143,7 @@ bool IsIndirectLightingCacheAllowed(ERHIFeatureLevel::Type InFeatureLevel)
 bool CanIndirectLightingCacheUseVolumeTexture(ERHIFeatureLevel::Type InFeatureLevel)
 {
 	// @todo Mac OS X/OpenGL: For OpenGL devices which don't support volume-texture rendering we need to use the simpler point indirect lighting shaders.
-	return InFeatureLevel >= ERHIFeatureLevel::SM4 && GSupportsVolumeTextureRendering;
+	return InFeatureLevel >= ERHIFeatureLevel::SM5 && GSupportsVolumeTextureRendering;
 }
 
 FIndirectLightingCache::FIndirectLightingCache(ERHIFeatureLevel::Type InFeatureLevel)
diff --git a/Engine/Source/Runtime/Renderer/Private/LightFunctionRendering.cpp b/Engine/Source/Runtime/Renderer/Private/LightFunctionRendering.cpp
--- a/Engine/Source/Runtime/Renderer/Private/LightFunctionRendering.cpp
+++ b/Engine/Source/Runtime/Renderer/Private/LightFunctionRendering.cpp
@@ -38,7 +38,7 @@ public:
 	  */
 	static bool ShouldCompilePermutation(const FMaterialShaderPermutationParameters& Parameters)
 	{
-		return Parameters.Material->IsLightFunction() && IsFeatureLevelSupported(Parameters.Platform, ERHIFeatureLevel::SM4);
+		return Parameters.Material->IsLightFunction() && IsFeatureLevelSupported(Parameters.Platform, ERHIFeatureLevel::SM5);
 	}
 
 	FLightFunctionVS( )	{ }
@@ -95,7 +95,7 @@ public:
 	  */
 	static bool ShouldCompilePermutation(const FMaterialShaderPermutationParameters& Parameters)
 	{
-		return Parameters.Material->IsLightFunction() && IsFeatureLevelSupported(Parameters.Platform, ERHIFeatureLevel::SM4);
+		return Parameters.Material->IsLightFunction() && IsFeatureLevelSupported(Parameters.Platform, ERHIFeatureLevel::SM5);
 	}
 
 	FLightFunctionPS() {}
diff --git a/Engine/Source/Runtime/Renderer/Private/LightGridInjection.cpp b/Engine/Source/Runtime/Renderer/Private/LightGridInjection.cpp
--- a/Engine/Source/Runtime/Renderer/Private/LightGridInjection.cpp
+++ b/Engine/Source/Runtime/Renderer/Private/LightGridInjection.cpp
@@ -85,7 +85,7 @@ public:
 
 	virtual void InitRHI()
 	{
-		if (GMaxRHIFeatureLevel >= ERHIFeatureLevel::SM4)
+		if (GMaxRHIFeatureLevel >= ERHIFeatureLevel::SM5)
 		{
 			if (GMaxRHIFeatureLevel >= ERHIFeatureLevel::SM5)
 			{
diff --git a/Engine/Source/Runtime/Renderer/Private/LightMapDensityRendering.cpp b/Engine/Source/Runtime/Renderer/Private/LightMapDensityRendering.cpp
--- a/Engine/Source/Runtime/Renderer/Private/LightMapDensityRendering.cpp
+++ b/Engine/Source/Runtime/Renderer/Private/LightMapDensityRendering.cpp
@@ -61,7 +61,7 @@ bool FDeferredShadingSceneRenderer::RenderLightMapDensities(FRHICommandListImmed
 {
 	bool bDirty = false;
 
-	if (Scene->GetFeatureLevel() >= ERHIFeatureLevel::SM4)
+	if (Scene->GetFeatureLevel() >= ERHIFeatureLevel::SM5)
 	{
 		SCOPED_DRAW_EVENT(RHICmdList, LightMapDensity);
 
@@ -235,7 +235,7 @@ void FLightmapDensityMeshProcessor::Process(
 
 void FLightmapDensityMeshProcessor::AddMeshBatch(const FMeshBatch& RESTRICT MeshBatch, uint64 BatchElementMask, const FPrimitiveSceneProxy* RESTRICT PrimitiveSceneProxy, int32 StaticMeshId)
 {
-	if (FeatureLevel >= ERHIFeatureLevel::SM4 && ViewIfDynamicMeshCommand->Family->EngineShowFlags.LightMapDensity && AllowDebugViewmodes() && MeshBatch.bUseForMaterial)
+	if (FeatureLevel >= ERHIFeatureLevel::SM5 && ViewIfDynamicMeshCommand->Family->EngineShowFlags.LightMapDensity && AllowDebugViewmodes() && MeshBatch.bUseForMaterial)
 	{
 		// Determine the mesh's material and blend mode.
 		const FMaterialRenderProxy* MaterialRenderProxy = nullptr;
diff --git a/Engine/Source/Runtime/Renderer/Private/LightMapDensityRendering.h b/Engine/Source/Runtime/Renderer/Private/LightMapDensityRendering.h
--- a/Engine/Source/Runtime/Renderer/Private/LightMapDensityRendering.h
+++ b/Engine/Source/Runtime/Renderer/Private/LightMapDensityRendering.h
@@ -61,7 +61,7 @@ public:
 		return  AllowDebugViewmodes(Parameters.Platform) 
 				&& (Parameters.Material->IsSpecialEngineMaterial() || Parameters.Material->IsMasked() || Parameters.Material->MaterialMayModifyMeshPosition())
 				&& LightMapPolicyType::ShouldCompilePermutation(Parameters)
-				&& IsFeatureLevelSupported(Parameters.Platform, ERHIFeatureLevel::SM4);
+				&& IsFeatureLevelSupported(Parameters.Platform, ERHIFeatureLevel::SM5);
 	}
 
 	static void ModifyCompilationEnvironment(const FMaterialShaderPermutationParameters& Parameters, FShaderCompilerEnvironment& OutEnvironment)
@@ -184,7 +184,7 @@ public:
 		return	AllowDebugViewmodes(Parameters.Platform) 
 				&& (Parameters.Material->IsSpecialEngineMaterial() || Parameters.Material->IsMasked() || Parameters.Material->MaterialMayModifyMeshPosition())
 				&& LightMapPolicyType::ShouldCompilePermutation(Parameters)
-				&& IsFeatureLevelSupported(Parameters.Platform, ERHIFeatureLevel::SM4);
+				&& IsFeatureLevelSupported(Parameters.Platform, ERHIFeatureLevel::SM5);
 	}
 
 	static void ModifyCompilationEnvironment(const FMaterialShaderPermutationParameters& Parameters, FShaderCompilerEnvironment& OutEnvironment)
diff --git a/Engine/Source/Runtime/Renderer/Private/LightMapRendering.cpp b/Engine/Source/Runtime/Renderer/Private/LightMapRendering.cpp
--- a/Engine/Source/Runtime/Renderer/Private/LightMapRendering.cpp
+++ b/Engine/Source/Runtime/Renderer/Private/LightMapRendering.cpp
@@ -359,7 +359,7 @@ void GetIndirectLightingCacheParameters(
 		// If we are using FCachedVolumeIndirectLightingPolicy then InitViews should have updated the lighting cache which would have initialized it
 		// However the conditions for updating the lighting cache are complex and fail very occasionally in non-reproducible ways
 		// Silently skipping setting the cache texture under failure for now
-		if (FeatureLevel >= ERHIFeatureLevel::SM4 && LightingCache && LightingCache->IsInitialized() && GSupportsVolumeTextureRendering)
+		if (FeatureLevel >= ERHIFeatureLevel::SM5 && LightingCache && LightingCache->IsInitialized() && GSupportsVolumeTextureRendering)
 		{
 			Parameters.IndirectLightingCacheTexture0 = const_cast<FIndirectLightingCache*>(LightingCache)->GetTexture0().ShaderResourceTexture;
 			Parameters.IndirectLightingCacheTexture1 = const_cast<FIndirectLightingCache*>(LightingCache)->GetTexture1().ShaderResourceTexture;
diff --git a/Engine/Source/Runtime/Renderer/Private/LightMapRendering.h b/Engine/Source/Runtime/Renderer/Private/LightMapRendering.h
--- a/Engine/Source/Runtime/Renderer/Private/LightMapRendering.h
+++ b/Engine/Source/Runtime/Renderer/Private/LightMapRendering.h
@@ -207,7 +207,7 @@ public:
 	{
 		return Parameters.Material->GetShadingModels().IsLit() &&
 			IsTranslucentBlendMode(Parameters.Material->GetBlendMode()) &&
-			IsFeatureLevelSupported(Parameters.Platform, ERHIFeatureLevel::SM4);
+			IsFeatureLevelSupported(Parameters.Platform, ERHIFeatureLevel::SM5);
 	}
 
 	static void ModifyCompilationEnvironment(const FMaterialShaderPermutationParameters& Parameters, FShaderCompilerEnvironment& OutEnvironment)
@@ -275,7 +275,7 @@ struct FCachedVolumeIndirectLightingPolicy
 		return Parameters.Material->GetShadingModels().IsLit()
 			&& !IsTranslucentBlendMode(Parameters.Material->GetBlendMode())
 			&& (!AllowStaticLightingVar || AllowStaticLightingVar->GetValueOnAnyThread() != 0)
-			&& IsFeatureLevelSupported(Parameters.Platform, ERHIFeatureLevel::SM4);
+			&& IsFeatureLevelSupported(Parameters.Platform, ERHIFeatureLevel::SM5);
 	}
 
 	static void ModifyCompilationEnvironment(const FMaterialShaderPermutationParameters& Parameters, FShaderCompilerEnvironment& OutEnvironment)
diff --git a/Engine/Source/Runtime/Renderer/Private/LightRendering.cpp b/Engine/Source/Runtime/Renderer/Private/LightRendering.cpp
--- a/Engine/Source/Runtime/Renderer/Private/LightRendering.cpp
+++ b/Engine/Source/Runtime/Renderer/Private/LightRendering.cpp
@@ -259,7 +259,7 @@ class FDeferredLightPS : public FGlobalShader
 			return false;
 		}*/
 
-		return IsFeatureLevelSupported(Parameters.Platform, ERHIFeatureLevel::SM4);
+		return IsFeatureLevelSupported(Parameters.Platform, ERHIFeatureLevel::SM5);
 	}
 
 	FDeferredLightPS(const ShaderMetaType::CompiledShaderInitializerType& Initializer)
@@ -432,7 +432,7 @@ public:
 
 	static bool ShouldCompilePermutation(const FGlobalShaderPermutationParameters& Parameters)
 	{
-		return IsFeatureLevelSupported(Parameters.Platform, ERHIFeatureLevel::SM4);
+		return IsFeatureLevelSupported(Parameters.Platform, ERHIFeatureLevel::SM5);
 	}
 
 	static void ModifyCompilationEnvironment(const FGlobalShaderPermutationParameters& Parameters, FShaderCompilerEnvironment& OutEnvironment)
diff --git a/Engine/Source/Runtime/Renderer/Private/LightRendering.h b/Engine/Source/Runtime/Renderer/Private/LightRendering.h
--- a/Engine/Source/Runtime/Renderer/Private/LightRendering.h
+++ b/Engine/Source/Runtime/Renderer/Private/LightRendering.h
@@ -502,7 +502,7 @@ public:
 
 	static bool ShouldCompilePermutation(const FGlobalShaderPermutationParameters& Parameters)
 	{
-		return bRadialLight ? IsFeatureLevelSupported(Parameters.Platform, ERHIFeatureLevel::SM4) : true;
+		return bRadialLight ? IsFeatureLevelSupported(Parameters.Platform, ERHIFeatureLevel::SM5) : true;
 	}
 
 	static void ModifyCompilationEnvironment(const FGlobalShaderPermutationParameters& Parameters, FShaderCompilerEnvironment& OutEnvironment)
diff --git a/Engine/Source/Runtime/Renderer/Private/LightSceneInfo.cpp b/Engine/Source/Runtime/Renderer/Private/LightSceneInfo.cpp
--- a/Engine/Source/Runtime/Renderer/Private/LightSceneInfo.cpp
+++ b/Engine/Source/Runtime/Renderer/Private/LightSceneInfo.cpp
@@ -73,7 +73,7 @@ void FLightSceneInfo::AddToScene()
 		// Lights that should be baked need to check for interactions to track unbuilt state correctly
 		|| Proxy->HasStaticLighting()
 		// ES2 path supports dynamic point lights in the base pass using forward rendering, so we need to know the primitives
-		|| (Scene->GetFeatureLevel() < ERHIFeatureLevel::SM4 && Proxy->GetLightType() == LightType_Point && Proxy->IsMovable()))
+		|| (Scene->GetFeatureLevel() < ERHIFeatureLevel::SM5 && Proxy->GetLightType() == LightType_Point && Proxy->IsMovable()))
 	{
 		// Add the light to the scene's light octree.
 		Scene->LightOctree.AddElement(LightSceneInfoCompact);
diff --git a/Engine/Source/Runtime/Renderer/Private/LightShaftRendering.cpp b/Engine/Source/Runtime/Renderer/Private/LightShaftRendering.cpp
--- a/Engine/Source/Runtime/Renderer/Private/LightShaftRendering.cpp
+++ b/Engine/Source/Runtime/Renderer/Private/LightShaftRendering.cpp
@@ -274,7 +274,7 @@ public:
 
 	static bool ShouldCompilePermutation(const FGlobalShaderPermutationParameters& Parameters)
 	{
-		return IsFeatureLevelSupported(Parameters.Platform, ERHIFeatureLevel::SM4);
+		return IsFeatureLevelSupported(Parameters.Platform, ERHIFeatureLevel::SM5);
 	}
 
 	/** Default constructor. */
@@ -307,7 +307,7 @@ public:
 
 	static bool ShouldCompilePermutation(const FGlobalShaderPermutationParameters& Parameters)
 	{
-		return IsFeatureLevelSupported(Parameters.Platform, ERHIFeatureLevel::SM4); 
+		return IsFeatureLevelSupported(Parameters.Platform, ERHIFeatureLevel::SM5); 
 	}
 
 	static void ModifyCompilationEnvironment(const FGlobalShaderPermutationParameters& Parameters, FShaderCompilerEnvironment& OutEnvironment)
@@ -381,7 +381,7 @@ public:
 
 	static bool ShouldCompilePermutation(const FGlobalShaderPermutationParameters& Parameters)
 	{
-		return IsFeatureLevelSupported(Parameters.Platform, ERHIFeatureLevel::SM4); 
+		return IsFeatureLevelSupported(Parameters.Platform, ERHIFeatureLevel::SM5); 
 	}
 
 	static void ModifyCompilationEnvironment(const FGlobalShaderPermutationParameters& Parameters, FShaderCompilerEnvironment& OutEnvironment)
@@ -438,7 +438,7 @@ public:
 
 	static bool ShouldCompilePermutation(const FGlobalShaderPermutationParameters& Parameters)
 	{
-		return IsFeatureLevelSupported(Parameters.Platform, ERHIFeatureLevel::SM4); 
+		return IsFeatureLevelSupported(Parameters.Platform, ERHIFeatureLevel::SM5); 
 	}
 
 	/** Default constructor. */
@@ -877,7 +877,7 @@ public:
 
 	static bool ShouldCompilePermutation(const FGlobalShaderPermutationParameters& Parameters)
 	{
-		return IsFeatureLevelSupported(Parameters.Platform, ERHIFeatureLevel::SM4); 
+		return IsFeatureLevelSupported(Parameters.Platform, ERHIFeatureLevel::SM5); 
 	}
 
 	/** Default constructor. */
diff --git a/Engine/Source/Runtime/Renderer/Private/PlanarReflectionRendering.cpp b/Engine/Source/Runtime/Renderer/Private/PlanarReflectionRendering.cpp
--- a/Engine/Source/Runtime/Renderer/Private/PlanarReflectionRendering.cpp
+++ b/Engine/Source/Runtime/Renderer/Private/PlanarReflectionRendering.cpp
@@ -148,7 +148,7 @@ public:
 
 	static bool ShouldCompilePermutation(const FGlobalShaderPermutationParameters& Parameters)
 	{
-		return bEnablePlanarReflectionPrefilter ? IsFeatureLevelSupported(Parameters.Platform, ERHIFeatureLevel::SM4) : true;
+		return bEnablePlanarReflectionPrefilter ? IsFeatureLevelSupported(Parameters.Platform, ERHIFeatureLevel::SM5) : true;
 	}
 
 	static void ModifyCompilationEnvironment(const FGlobalShaderPermutationParameters& Parameters, FShaderCompilerEnvironment& OutEnvironment)
@@ -218,7 +218,7 @@ void PrefilterPlanarReflection(FRHICommandListImmediate& RHICmdList, FViewInfo&
 {
 	FRHITexture* SceneColorInput = FSceneRenderTargets::Get(RHICmdList).GetSceneColorTexture();
 
-	if(View.FeatureLevel >= ERHIFeatureLevel::SM4)
+	if(View.FeatureLevel >= ERHIFeatureLevel::SM5)
 	{
 		TRefCountPtr<IPooledRenderTarget> FilteredSceneColor;
 		GPostProcessing.ProcessPlanarReflection(RHICmdList, View, FilteredSceneColor);
@@ -646,7 +646,7 @@ class FPlanarReflectionPS : public FGlobalShader
 
 	static bool ShouldCompilePermutation(const FGlobalShaderPermutationParameters& Parameters)
 	{
-		return IsFeatureLevelSupported(Parameters.Platform, ERHIFeatureLevel::SM4);
+		return IsFeatureLevelSupported(Parameters.Platform, ERHIFeatureLevel::SM5);
 	}
 
 	BEGIN_SHADER_PARAMETER_STRUCT(FParameters, )
diff --git a/Engine/Source/Runtime/Renderer/Private/PostProcess/PostProcessAA.cpp b/Engine/Source/Runtime/Renderer/Private/PostProcess/PostProcessAA.cpp
--- a/Engine/Source/Runtime/Renderer/Private/PostProcess/PostProcessAA.cpp
+++ b/Engine/Source/Runtime/Renderer/Private/PostProcess/PostProcessAA.cpp
@@ -20,7 +20,7 @@ class FPostProcessAntiAliasingPS : public FGlobalShader
 
 	static bool ShouldCompilePermutation(const FGlobalShaderPermutationParameters& Parameters)
 	{
-		return IsFeatureLevelSupported(Parameters.Platform, ERHIFeatureLevel::SM4);
+		return IsFeatureLevelSupported(Parameters.Platform, ERHIFeatureLevel::SM5);
 	}
 
 	static void ModifyCompilationEnvironment(const FGlobalShaderPermutationParameters& Parameters, FShaderCompilerEnvironment& OutEnvironment)
@@ -165,7 +165,7 @@ public:
 
 	static bool ShouldCompilePermutation(const FGlobalShaderPermutationParameters& Parameters)
 	{
-		return IsFeatureLevelSupported(Parameters.Platform, ERHIFeatureLevel::SM4);
+		return IsFeatureLevelSupported(Parameters.Platform, ERHIFeatureLevel::SM5);
 	}
 
 	/** Default constructor. */
diff --git a/Engine/Source/Runtime/Renderer/Private/PostProcess/PostProcessBloomSetup.cpp b/Engine/Source/Runtime/Renderer/Private/PostProcess/PostProcessBloomSetup.cpp
--- a/Engine/Source/Runtime/Renderer/Private/PostProcess/PostProcessBloomSetup.cpp
+++ b/Engine/Source/Runtime/Renderer/Private/PostProcess/PostProcessBloomSetup.cpp
@@ -24,7 +24,7 @@ class FPostProcessBloomSetupPS : public FGlobalShader
 
 	static bool ShouldCompilePermutation(const FGlobalShaderPermutationParameters& Parameters)
 	{
-		return IsFeatureLevelSupported(Parameters.Platform, ERHIFeatureLevel::SM4);
+		return IsFeatureLevelSupported(Parameters.Platform, ERHIFeatureLevel::SM5);
 	}	
 	
 	static void ModifyCompilationEnvironment(const FGlobalShaderPermutationParameters& Parameters, FShaderCompilerEnvironment& OutEnvironment)
@@ -89,7 +89,7 @@ class FPostProcessBloomSetupVS : public FGlobalShader
 
 	static bool ShouldCompilePermutation(const FGlobalShaderPermutationParameters& Parameters)
 	{
-		return IsFeatureLevelSupported(Parameters.Platform, ERHIFeatureLevel::SM4);
+		return IsFeatureLevelSupported(Parameters.Platform, ERHIFeatureLevel::SM5);
 	}
 
 	/** Default constructor. */
@@ -372,7 +372,7 @@ public:
 
 	static bool ShouldCompilePermutation(const FGlobalShaderPermutationParameters& Parameters)
 	{
-		return IsFeatureLevelSupported(Parameters.Platform, ERHIFeatureLevel::SM4);
+		return IsFeatureLevelSupported(Parameters.Platform, ERHIFeatureLevel::SM5);
 	}
 };
 
@@ -422,7 +422,7 @@ class FPostProcessVisualizeBloomOverlayPS : public FGlobalShader
 
 	static bool ShouldCompilePermutation(const FGlobalShaderPermutationParameters& Parameters)
 	{
-		return IsFeatureLevelSupported(Parameters.Platform, ERHIFeatureLevel::SM4);
+		return IsFeatureLevelSupported(Parameters.Platform, ERHIFeatureLevel::SM5);
 	}
 
 	static void ModifyCompilationEnvironment(const FGlobalShaderPermutationParameters& Parameters, FShaderCompilerEnvironment& OutEnvironment)
diff --git a/Engine/Source/Runtime/Renderer/Private/PostProcess/PostProcessBokehDOF.cpp b/Engine/Source/Runtime/Renderer/Private/PostProcess/PostProcessBokehDOF.cpp
--- a/Engine/Source/Runtime/Renderer/Private/PostProcess/PostProcessBokehDOF.cpp
+++ b/Engine/Source/Runtime/Renderer/Private/PostProcess/PostProcessBokehDOF.cpp
@@ -37,7 +37,7 @@ class PostProcessVisualizeDOFPS : public FGlobalShader
 
 	static bool ShouldCompilePermutation(const FGlobalShaderPermutationParameters& Parameters)
 	{
-		return IsFeatureLevelSupported(Parameters.Platform, ERHIFeatureLevel::SM4);
+		return IsFeatureLevelSupported(Parameters.Platform, ERHIFeatureLevel::SM5);
 	}
 
 	/** Default constructor. */
diff --git a/Engine/Source/Runtime/Renderer/Private/PostProcess/PostProcessBusyWait.cpp b/Engine/Source/Runtime/Renderer/Private/PostProcess/PostProcessBusyWait.cpp
--- a/Engine/Source/Runtime/Renderer/Private/PostProcess/PostProcessBusyWait.cpp
+++ b/Engine/Source/Runtime/Renderer/Private/PostProcess/PostProcessBusyWait.cpp
@@ -29,7 +29,7 @@ class FPostProcessBusyWaitPS : public FGlobalShader
 
 	static bool ShouldCompilePermutation(const FGlobalShaderPermutationParameters& Parameters)
 	{
-		return IsFeatureLevelSupported(Parameters.Platform, ERHIFeatureLevel::SM4);
+		return IsFeatureLevelSupported(Parameters.Platform, ERHIFeatureLevel::SM5);
 	}
 
 	/** Default constructor. */
diff --git a/Engine/Source/Runtime/Renderer/Private/PostProcess/PostProcessDOF.cpp b/Engine/Source/Runtime/Renderer/Private/PostProcess/PostProcessDOF.cpp
--- a/Engine/Source/Runtime/Renderer/Private/PostProcess/PostProcessDOF.cpp
+++ b/Engine/Source/Runtime/Renderer/Private/PostProcess/PostProcessDOF.cpp
@@ -31,7 +31,7 @@ class FPostProcessDOFSetupPS : public FGlobalShader
 	static void ModifyCompilationEnvironment(const FGlobalShaderPermutationParameters& Parameters, FShaderCompilerEnvironment& OutEnvironment)
 	{
 		FGlobalShader::ModifyCompilationEnvironment(Parameters,OutEnvironment);
-		OutEnvironment.SetDefine(TEXT("MOBILE_SHADING"), IsFeatureLevelSupported(Parameters.Platform, ERHIFeatureLevel::SM4) ? 0u : 1u);
+		OutEnvironment.SetDefine(TEXT("MOBILE_SHADING"), IsFeatureLevelSupported(Parameters.Platform, ERHIFeatureLevel::SM5) ? 0u : 1u);
 		OutEnvironment.SetDefine(TEXT("NEAR_BLUR"), (uint32)(NearBlur >= 1));
 		OutEnvironment.SetDefine(TEXT("DOF_VIGNETTE"), (uint32)(NearBlur == 2));
 		OutEnvironment.SetDefine(TEXT("MRT_COUNT"), FarBlur + (NearBlur > 0));
@@ -69,7 +69,7 @@ public:
 
 		FGlobalShader::SetParameters<FViewUniformShaderParameters>(RHICmdList, ShaderRHI, Context.View.ViewUniformBuffer);
 
-		if (Context.GetFeatureLevel() < ERHIFeatureLevel::SM4)
+		if (Context.GetFeatureLevel() < ERHIFeatureLevel::SM5)
 		{
 			// Trying bilin here to attempt to alleviate some issues with 1/4 res input...
 			PostprocessParameter.SetPS(RHICmdList, ShaderRHI, Context, TStaticSamplerState<SF_Bilinear, AM_Border, AM_Border, AM_Clamp>::GetRHI());
@@ -311,7 +311,7 @@ class FPostProcessDOFRecombinePS : public FGlobalShader
 		OutEnvironment.SetDefine(TEXT("FAR_BLUR"), FarBlur);
 		OutEnvironment.SetDefine(TEXT("NEAR_BLUR"), NearBlur);
 		OutEnvironment.SetDefine(TEXT("SEPARATE_TRANSLUCENCY"), SeparateTranslucency);
-		OutEnvironment.SetDefine(TEXT("MOBILE_SHADING"), IsFeatureLevelSupported(Parameters.Platform, ERHIFeatureLevel::SM4) ? 0u : 1u);
+		OutEnvironment.SetDefine(TEXT("MOBILE_SHADING"), IsFeatureLevelSupported(Parameters.Platform, ERHIFeatureLevel::SM5) ? 0u : 1u);
 	}
 
 	/** Default constructor. */
diff --git a/Engine/Source/Runtime/Renderer/Private/PostProcess/PostProcessEyeAdaptation.cpp b/Engine/Source/Runtime/Renderer/Private/PostProcess/PostProcessEyeAdaptation.cpp
--- a/Engine/Source/Runtime/Renderer/Private/PostProcess/PostProcessEyeAdaptation.cpp
+++ b/Engine/Source/Runtime/Renderer/Private/PostProcess/PostProcessEyeAdaptation.cpp
@@ -175,7 +175,7 @@ class FPostProcessEyeAdaptationPS : public FGlobalShader
 
 	static bool ShouldCompilePermutation(const FGlobalShaderPermutationParameters& Parameters)
 	{
-		return IsFeatureLevelSupported(Parameters.Platform, ERHIFeatureLevel::SM4);
+		return IsFeatureLevelSupported(Parameters.Platform, ERHIFeatureLevel::SM5);
 	}
 
 	static void ModifyCompilationEnvironment(const FGlobalShaderPermutationParameters& Parameters, FShaderCompilerEnvironment& OutEnvironment)
diff --git a/Engine/Source/Runtime/Renderer/Private/PostProcess/PostProcessGBufferHints.cpp b/Engine/Source/Runtime/Renderer/Private/PostProcess/PostProcessGBufferHints.cpp
--- a/Engine/Source/Runtime/Renderer/Private/PostProcess/PostProcessGBufferHints.cpp
+++ b/Engine/Source/Runtime/Renderer/Private/PostProcess/PostProcessGBufferHints.cpp
@@ -25,7 +25,7 @@ class FPostProcessGBufferHintsPS : public FGlobalShader
 
 	static bool ShouldCompilePermutation(const FGlobalShaderPermutationParameters& Parameters)
 	{
-		return IsFeatureLevelSupported(Parameters.Platform, ERHIFeatureLevel::SM4);
+		return IsFeatureLevelSupported(Parameters.Platform, ERHIFeatureLevel::SM5);
 	}
 
 	static void ModifyCompilationEnvironment(const FGlobalShaderPermutationParameters& Parameters, FShaderCompilerEnvironment& OutEnvironment)
diff --git a/Engine/Source/Runtime/Renderer/Private/PostProcess/PostProcessLensBlur.cpp b/Engine/Source/Runtime/Renderer/Private/PostProcess/PostProcessLensBlur.cpp
--- a/Engine/Source/Runtime/Renderer/Private/PostProcess/PostProcessLensBlur.cpp
+++ b/Engine/Source/Runtime/Renderer/Private/PostProcess/PostProcessLensBlur.cpp
@@ -21,7 +21,7 @@ class FPostProcessLensBlurVS : public FGlobalShader
 
 	static bool ShouldCompilePermutation(const FGlobalShaderPermutationParameters& Parameters)
 	{
-		return IsFeatureLevelSupported(Parameters.Platform, ERHIFeatureLevel::SM4);
+		return IsFeatureLevelSupported(Parameters.Platform, ERHIFeatureLevel::SM5);
 	}
 
 	/** Default constructor. */
@@ -90,7 +90,7 @@ class FPostProcessLensBlurPS : public FGlobalShader
 
 	static bool ShouldCompilePermutation(const FGlobalShaderPermutationParameters& Parameters)
 	{
-		return IsFeatureLevelSupported(Parameters.Platform, ERHIFeatureLevel::SM4);
+		return IsFeatureLevelSupported(Parameters.Platform, ERHIFeatureLevel::SM5);
 	}
 
 	/** Default constructor. */
diff --git a/Engine/Source/Runtime/Renderer/Private/PostProcess/PostProcessLensFlares.cpp b/Engine/Source/Runtime/Renderer/Private/PostProcess/PostProcessLensFlares.cpp
--- a/Engine/Source/Runtime/Renderer/Private/PostProcess/PostProcessLensFlares.cpp
+++ b/Engine/Source/Runtime/Renderer/Private/PostProcess/PostProcessLensFlares.cpp
@@ -21,7 +21,7 @@ class TPostProcessLensFlareBasePS : public FGlobalShader
 
 	static bool ShouldCompilePermutation(const FGlobalShaderPermutationParameters& Parameters)
 	{
-		return IsFeatureLevelSupported(Parameters.Platform, ERHIFeatureLevel::SM4);
+		return IsFeatureLevelSupported(Parameters.Platform, ERHIFeatureLevel::SM5);
 	}
 
 	/** Default constructor. */
@@ -84,7 +84,7 @@ class FPostProcessLensFlaresPS : public FGlobalShader
 
 	static bool ShouldCompilePermutation(const FGlobalShaderPermutationParameters& Parameters)
 	{
-		return IsFeatureLevelSupported(Parameters.Platform, ERHIFeatureLevel::SM4);
+		return IsFeatureLevelSupported(Parameters.Platform, ERHIFeatureLevel::SM5);
 	}
 
 	/** Default constructor. */
diff --git a/Engine/Source/Runtime/Renderer/Private/PostProcess/PostProcessMaterial.cpp b/Engine/Source/Runtime/Renderer/Private/PostProcess/PostProcessMaterial.cpp
--- a/Engine/Source/Runtime/Renderer/Private/PostProcess/PostProcessMaterial.cpp
+++ b/Engine/Source/Runtime/Renderer/Private/PostProcess/PostProcessMaterial.cpp
@@ -87,25 +87,7 @@ ECustomDepthPolicy GetMaterialCustomDepthPolicy(const FMaterial* Material, ERHIF
 		// Custom stencil texture allocated and available.
 		if (IsCustomDepthEnabled())
 		{
-			if (FeatureLevel == ERHIFeatureLevel::SM4)
-			{
-				const FMaterialShaderMap* MaterialShaderMap = Material->GetRenderingThreadShaderMap();
-
-				const bool bIsDepthStencilReadByShader = MaterialShaderMap->UsesSceneTexture(PPI_CustomDepth) || MaterialShaderMap->UsesSceneTexture(PPI_CustomStencil);
-
-				if (!bIsDepthStencilReadByShader)
-				{
-					return ECustomDepthPolicy::Enabled;
-				}
-				else if (IsPostProcessStencilTestWithCopyAllowed())
-				{
-					return ECustomDepthPolicy::EnabledWithCopy;
-				}
-			}
-			else
-			{
-				return ECustomDepthPolicy::Enabled;
-			}
+			return ECustomDepthPolicy::Enabled;
 		}
 		else
 		{
@@ -227,7 +209,7 @@ public:
 			}
 			else
 			{
-				return IsFeatureLevelSupported(Parameters.Platform, ERHIFeatureLevel::SM4);
+				return IsFeatureLevelSupported(Parameters.Platform, ERHIFeatureLevel::SM5);
 			}
 		}
 		return false;
diff --git a/Engine/Source/Runtime/Renderer/Private/PostProcess/PostProcessMotionBlur.cpp b/Engine/Source/Runtime/Renderer/Private/PostProcess/PostProcessMotionBlur.cpp
--- a/Engine/Source/Runtime/Renderer/Private/PostProcess/PostProcessMotionBlur.cpp
+++ b/Engine/Source/Runtime/Renderer/Private/PostProcess/PostProcessMotionBlur.cpp
@@ -110,7 +110,7 @@ bool IsMotionBlurEnabled(const FViewInfo& View)
 
 bool IsVisualizeMotionBlurEnabled(const FViewInfo& View)
 {
-	return View.Family->EngineShowFlags.VisualizeMotionBlur && View.GetFeatureLevel() >= ERHIFeatureLevel::SM4;
+	return View.Family->EngineShowFlags.VisualizeMotionBlur && View.GetFeatureLevel() >= ERHIFeatureLevel::SM5;
 }
 
 bool IsMotionBlurScatterRequired(const FViewInfo& View, const FScreenPassTextureViewport& SceneViewport)
@@ -386,7 +386,7 @@ public:
 
 	static bool ShouldCompilePermutation(const FGlobalShaderPermutationParameters& Parameters)
 	{
-		return IsFeatureLevelSupported(Parameters.Platform, ERHIFeatureLevel::SM4);
+		return IsFeatureLevelSupported(Parameters.Platform, ERHIFeatureLevel::SM5);
 	}
 };
 
diff --git a/Engine/Source/Runtime/Renderer/Private/PostProcess/PostProcessNoiseBlur.cpp b/Engine/Source/Runtime/Renderer/Private/PostProcess/PostProcessNoiseBlur.cpp
--- a/Engine/Source/Runtime/Renderer/Private/PostProcess/PostProcessNoiseBlur.cpp
+++ b/Engine/Source/Runtime/Renderer/Private/PostProcess/PostProcessNoiseBlur.cpp
@@ -22,7 +22,7 @@ class FPostProcessNoiseBlurPS : public FGlobalShader
 
 	static bool ShouldCompilePermutation(const FGlobalShaderPermutationParameters& Parameters)
 	{
-		return IsFeatureLevelSupported(Parameters.Platform, ERHIFeatureLevel::SM4);
+		return IsFeatureLevelSupported(Parameters.Platform, ERHIFeatureLevel::SM5);
 	}
 
 	static void ModifyCompilationEnvironment(const FGlobalShaderPermutationParameters& Parameters, FShaderCompilerEnvironment& OutEnvironment)
diff --git a/Engine/Source/Runtime/Renderer/Private/PostProcess/PostProcessSubsurface.cpp b/Engine/Source/Runtime/Renderer/Private/PostProcess/PostProcessSubsurface.cpp
--- a/Engine/Source/Runtime/Renderer/Private/PostProcess/PostProcessSubsurface.cpp
+++ b/Engine/Source/Runtime/Renderer/Private/PostProcess/PostProcessSubsurface.cpp
@@ -107,7 +107,7 @@ class FSubsurfaceViewportCopyPS : public FSubsurfaceShader
 
 	static bool ShouldCompilePermutation(const FGlobalShaderPermutationParameters& Parameters)
 	{
-		return IsFeatureLevelSupported(Parameters.Platform, ERHIFeatureLevel::SM4);
+		return IsFeatureLevelSupported(Parameters.Platform, ERHIFeatureLevel::SM5);
 	}
 };
 
diff --git a/Engine/Source/Runtime/Renderer/Private/PostProcess/PostProcessTestImage.cpp b/Engine/Source/Runtime/Renderer/Private/PostProcess/PostProcessTestImage.cpp
--- a/Engine/Source/Runtime/Renderer/Private/PostProcess/PostProcessTestImage.cpp
+++ b/Engine/Source/Runtime/Renderer/Private/PostProcess/PostProcessTestImage.cpp
@@ -24,7 +24,7 @@ class FPostProcessTestImagePS : public FGlobalShader
 
 	static bool ShouldCompilePermutation(const FGlobalShaderPermutationParameters& Parameters)
 	{
-		return IsFeatureLevelSupported(Parameters.Platform, ERHIFeatureLevel::SM4);
+		return IsFeatureLevelSupported(Parameters.Platform, ERHIFeatureLevel::SM5);
 	}
 
 	static void ModifyCompilationEnvironment(const FGlobalShaderPermutationParameters& Parameters, FShaderCompilerEnvironment& OutEnvironment)
diff --git a/Engine/Source/Runtime/Renderer/Private/PostProcess/PostProcessUpscale.cpp b/Engine/Source/Runtime/Renderer/Private/PostProcess/PostProcessUpscale.cpp
--- a/Engine/Source/Runtime/Renderer/Private/PostProcess/PostProcessUpscale.cpp
+++ b/Engine/Source/Runtime/Renderer/Private/PostProcess/PostProcessUpscale.cpp
@@ -149,7 +149,7 @@ class FPostProcessUpscalePS : public FGlobalShader
 			return true;
 		}
 
-		return IsFeatureLevelSupported(Parameters.Platform, ERHIFeatureLevel::SM4);
+		return IsFeatureLevelSupported(Parameters.Platform, ERHIFeatureLevel::SM5);
 	}
 
 	static void ModifyCompilationEnvironment(const FGlobalShaderPermutationParameters& Parameters, FShaderCompilerEnvironment& OutEnvironment)
diff --git a/Engine/Source/Runtime/Renderer/Private/PostProcess/PostProcessWeightedSampleSum.cpp b/Engine/Source/Runtime/Renderer/Private/PostProcess/PostProcessWeightedSampleSum.cpp
--- a/Engine/Source/Runtime/Renderer/Private/PostProcess/PostProcessWeightedSampleSum.cpp
+++ b/Engine/Source/Runtime/Renderer/Private/PostProcess/PostProcessWeightedSampleSum.cpp
@@ -79,7 +79,7 @@ public:
 		{
 			return CompileTimeNumSamples <= MAX_FILTER_COMPILE_TIME_SAMPLES_IOS;
 		}
-		else if( IsFeatureLevelSupported(Parameters.Platform, ERHIFeatureLevel::SM4) )
+		else if( IsFeatureLevelSupported(Parameters.Platform, ERHIFeatureLevel::SM5) )
 		{
 			return CompileTimeNumSamples <= MAX_FILTER_COMPILE_TIME_SAMPLES_SM4;
 		}
@@ -237,7 +237,7 @@ public:
 		{
 			return NumSamples <= MAX_FILTER_COMPILE_TIME_SAMPLES_IOS;
 		}
-		else if (IsFeatureLevelSupported(Parameters.Platform, ERHIFeatureLevel::SM4))
+		else if (IsFeatureLevelSupported(Parameters.Platform, ERHIFeatureLevel::SM5))
 		{
 			return NumSamples <= MAX_FILTER_COMPILE_TIME_SAMPLES_SM4;
 		}
@@ -1312,12 +1312,8 @@ uint32 FRCPassPostProcessWeightedSampleSum::GetMaxNumSamples(ERHIFeatureLevel::T
 	if (IsMetalMRTPlatform(InPlatform))
 	{
 		MaxNumSamples = MAX_FILTER_COMPILE_TIME_SAMPLES_IOS;
-	}
-	else if (InFeatureLevel == ERHIFeatureLevel::SM4)
-	{
-		MaxNumSamples = MAX_FILTER_COMPILE_TIME_SAMPLES_SM4;
-	}
-	else if (InFeatureLevel < ERHIFeatureLevel::SM4)
+	}	
+	else if (InFeatureLevel < ERHIFeatureLevel::SM5)
 	{
 		MaxNumSamples = MAX_FILTER_COMPILE_TIME_SAMPLES_ES2;
 	}
diff --git a/Engine/Source/Runtime/Renderer/Private/PostProcess/PostProcessing.cpp b/Engine/Source/Runtime/Renderer/Private/PostProcess/PostProcessing.cpp
--- a/Engine/Source/Runtime/Renderer/Private/PostProcess/PostProcessing.cpp
+++ b/Engine/Source/Runtime/Renderer/Private/PostProcess/PostProcessing.cpp
@@ -452,7 +452,7 @@ static bool AddPostProcessDepthOfFieldGaussian(FPostprocessContext& Context, FDe
 		const bool bFar = FarSize > 0.0f;
 		const bool bNear = NearSize > 0.0f;
 		const bool bCombinedNearFarPass = bFar && bNear;
-		const bool bMobileQuality = Context.View.FeatureLevel < ERHIFeatureLevel::SM4;
+		const bool bMobileQuality = Context.View.FeatureLevel < ERHIFeatureLevel::SM5;
 
 		FRenderingCompositeOutputRef SetupInput(Context.FinalOutput);
 		if (bMobileQuality)
@@ -509,7 +509,7 @@ static bool AddPostProcessDepthOfFieldGaussian(FPostprocessContext& Context, FDe
 	{
 		GaussianDOFPass(SeparateTranslucencyRef, Out.bFar ? FarSize : 0, Out.bNear ? NearSize : 0);
 
-		const bool bMobileQuality = Context.View.FeatureLevel < ERHIFeatureLevel::SM4;
+		const bool bMobileQuality = Context.View.FeatureLevel < ERHIFeatureLevel::SM5;
 		return SeparateTranslucencyRef.IsValid() && !bMobileQuality;
 	}
 	else
@@ -835,7 +835,7 @@ void FPostProcessing::OverrideRenderTarget(FRenderingCompositeOutputRef It, TRef
 
 bool FPostProcessing::AllowFullPostProcessing(const FViewInfo& View, ERHIFeatureLevel::Type FeatureLevel)
 {
-	if(FeatureLevel >= ERHIFeatureLevel::SM4)
+	if(FeatureLevel >= ERHIFeatureLevel::SM5)
 	{
 		return View.Family->EngineShowFlags.PostProcessing
 			&& !View.Family->EngineShowFlags.VisualizeDistanceFieldAO
@@ -906,7 +906,7 @@ class FComposeSeparateTranslucencyPS : public FGlobalShader
 
 	static bool ShouldCompilePermutation(const FGlobalShaderPermutationParameters& Parameters)
 	{
-		return IsFeatureLevelSupported(Parameters.Platform, ERHIFeatureLevel::SM4);
+		return IsFeatureLevelSupported(Parameters.Platform, ERHIFeatureLevel::SM5);
 	}
 };
 
@@ -953,7 +953,7 @@ void FPostProcessing::Process(FRHICommandListImmediate& RHICmdList, const FViewI
 
 	const ERHIFeatureLevel::Type FeatureLevel = View.GetFeatureLevel();
 
-	check(FeatureLevel >= ERHIFeatureLevel::SM4);
+	check(FeatureLevel >= ERHIFeatureLevel::SM5);
 
 	GRenderTargetPool.AddPhaseEvent(TEXT("PostProcessing"));
 
@@ -1452,7 +1452,7 @@ void FPostProcessing::Process(FRHICommandListImmediate& RHICmdList, const FViewI
 			Context.FinalOutput = FRenderingCompositeOutputRef(Node);
 		}
 #endif
-		if(View.Family->EngineShowFlags.VisualizeShadingModels && FeatureLevel >= ERHIFeatureLevel::SM4)
+		if(View.Family->EngineShowFlags.VisualizeShadingModels && FeatureLevel >= ERHIFeatureLevel::SM5)
 		{
 			ensureMsgf(!bUnscaledFinalOutput, TEXT("VisualizeShadingModels is incompatible with unscaled output."));
 
@@ -1461,7 +1461,7 @@ void FPostProcessing::Process(FRHICommandListImmediate& RHICmdList, const FViewI
 			Context.FinalOutput = FRenderingCompositeOutputRef(Node);
 		}
 
-		if (View.Family->EngineShowFlags.GBufferHints && FeatureLevel >= ERHIFeatureLevel::SM4)
+		if (View.Family->EngineShowFlags.GBufferHints && FeatureLevel >= ERHIFeatureLevel::SM5)
 		{
 			ensureMsgf(!bUnscaledFinalOutput, TEXT("GBufferHints is incompatible with unscaled output."));
 
@@ -1479,7 +1479,7 @@ void FPostProcessing::Process(FRHICommandListImmediate& RHICmdList, const FViewI
 		//No more postprocess Final color should be the real one
 		//The HDR was save before the tonemapping
 		//GBuffer should not be change during post process 
-		if (View.bUsePixelInspector && FeatureLevel >= ERHIFeatureLevel::SM4)
+		if (View.bUsePixelInspector && FeatureLevel >= ERHIFeatureLevel::SM5)
 		{
 			FRenderingCompositePass* Node = Context.Graph.RegisterPass(new(FMemStack::Get()) FRCPassPostProcessBufferInspector(RHICmdList));
 			Node->SetInput(ePId_Input0, Context.FinalOutput);
@@ -1517,7 +1517,7 @@ void FPostProcessing::Process(FRHICommandListImmediate& RHICmdList, const FViewI
 			Context.FinalOutput = FRenderingCompositeOutputRef(Node);
 		}
 
-		if(View.Family->EngineShowFlags.TestImage && FeatureLevel >= ERHIFeatureLevel::SM4)
+		if(View.Family->EngineShowFlags.TestImage && FeatureLevel >= ERHIFeatureLevel::SM5)
 		{
 			FRenderingCompositePass* Node = Context.Graph.RegisterPass(new(FMemStack::Get()) FRCPassPostProcessTestImage());
 			Node->SetInput(ePId_Input0, FRenderingCompositeOutputRef(Context.FinalOutput));
diff --git a/Engine/Source/Runtime/Renderer/Private/PostProcess/SceneRenderTargets.cpp b/Engine/Source/Runtime/Renderer/Private/PostProcess/SceneRenderTargets.cpp
--- a/Engine/Source/Runtime/Renderer/Private/PostProcess/SceneRenderTargets.cpp
+++ b/Engine/Source/Runtime/Renderer/Private/PostProcess/SceneRenderTargets.cpp
@@ -485,7 +485,7 @@ uint16 FSceneRenderTargets::GetNumSceneColorMSAASamples(ERHIFeatureLevel::Type I
 {
 	uint16 NumSamples = 1;
 
-	if (InFeatureLevel >= ERHIFeatureLevel::SM4)
+	if (InFeatureLevel >= ERHIFeatureLevel::SM5)
 	{
 		static IConsoleVariable* CVarDefaultAntiAliasing = IConsoleManager::Get().FindConsoleVariable(TEXT("r.DefaultFeature.AntiAliasing"));
 		EAntiAliasingMethod Method = (EAntiAliasingMethod)CVarDefaultAntiAliasing->GetInt();
@@ -857,7 +857,7 @@ void FSceneRenderTargets::BeginRenderingGBuffer(FRHICommandList& RHICmdList, ERe
 	check(RHICmdList.IsOutsideRenderPass());
 
 	SCOPED_DRAW_EVENT(RHICmdList, BeginRenderingGBuffer);
-	check(CurrentFeatureLevel >= ERHIFeatureLevel::SM4);
+	check(CurrentFeatureLevel >= ERHIFeatureLevel::SM5);
 	AllocSceneColor(RHICmdList);
 
 	const ERenderTargetStoreAction DepthStoreAction = (DepthStencilAccess & FExclusiveDepthStencil::DepthWrite) ? ERenderTargetStoreAction::EStore : ERenderTargetStoreAction::ENoAction;
@@ -985,7 +985,7 @@ int32 FSceneRenderTargets::GetNumGBufferTargets() const
 {
 	int32 NumGBufferTargets = 1;
 
-	if (CurrentFeatureLevel >= ERHIFeatureLevel::SM4)
+	if (CurrentFeatureLevel >= ERHIFeatureLevel::SM5)
 	{
 		const EShaderPlatform ShaderPlatform = GetFeatureLevelShaderPlatform(CurrentFeatureLevel);
 		if (IsUsingGBuffers(ShaderPlatform))
@@ -1174,7 +1174,7 @@ void FSceneRenderTargets::AllocGBufferTargets(FRHICommandList& RHICmdList)
 	// create GBuffer on demand so it can be shared with other pooled RT
 	const EShaderPlatform ShaderPlatform = GetFeatureLevelShaderPlatform(CurrentFeatureLevel);
 	const bool bUseGBuffer = IsUsingGBuffers(ShaderPlatform);
-	const bool bCanReadGBufferUniforms = (bUseGBuffer || IsSimpleForwardShadingEnabled(ShaderPlatform)) && CurrentFeatureLevel >= ERHIFeatureLevel::SM4;
+	const bool bCanReadGBufferUniforms = (bUseGBuffer || IsSimpleForwardShadingEnabled(ShaderPlatform)) && CurrentFeatureLevel >= ERHIFeatureLevel::SM5;
 	if (bUseGBuffer)
 	{
 		// good to see the quality loss due to precision in the gbuffer
@@ -2296,7 +2296,7 @@ void FSceneRenderTargets::AllocateDeferredShadingPathRenderTargets(FRHICommandLi
 	}
 
 	// Create the required render targets if running Highend.
-	if (CurrentFeatureLevel >= ERHIFeatureLevel::SM4)
+	if (CurrentFeatureLevel >= ERHIFeatureLevel::SM5)
 	{
 		// Create the screen space ambient occlusion buffer
 		{
@@ -2394,7 +2394,7 @@ void FSceneRenderTargets::AllocateDeferredShadingPathRenderTargets(FRHICommandLi
 		GRenderTargetPool.FindFreeElement(RHICmdList, Desc, DirectionalOcclusion, TEXT("DirectionalOcclusion"));
 	}
 
-	if (CurrentFeatureLevel >= ERHIFeatureLevel::SM4) 
+	if (CurrentFeatureLevel >= ERHIFeatureLevel::SM5) 
 	{
 		FPooledRenderTargetDesc Desc(FPooledRenderTargetDesc::Create2DDesc(BufferSize, PF_FloatRGBA, FClearValueBinding::Black, TexCreate_None, TexCreate_RenderTargetable, false));
 		if (CurrentFeatureLevel >= ERHIFeatureLevel::SM5)
@@ -2549,7 +2549,7 @@ EPixelFormat FSceneRenderTargets::GetSceneColorFormat(ERHIFeatureLevel::Type InF
 {
 	EPixelFormat SceneColorBufferFormat = PF_FloatRGBA;
 
-	if (InFeatureLevel < ERHIFeatureLevel::SM4)
+	if (InFeatureLevel < ERHIFeatureLevel::SM5)
 	{
 		return GetMobileSceneColorFormat();
 	}
@@ -2749,7 +2749,7 @@ const FUnorderedAccessViewRHIRef& FSceneRenderTargets::GetSceneColorTextureUAV()
 const FTexture2DRHIRef* FSceneRenderTargets::GetActualDepthTexture() const
 {
 	const FTexture2DRHIRef* DepthTexture = NULL;
-	if((CurrentFeatureLevel >= ERHIFeatureLevel::SM4) || IsPCPlatform(GShaderPlatformForFeatureLevel[CurrentFeatureLevel]))
+	if((CurrentFeatureLevel >= ERHIFeatureLevel::SM5) || IsPCPlatform(GShaderPlatformForFeatureLevel[CurrentFeatureLevel]))
 	{
 		if(GSupportsDepthFetchDuringDepthTest)
 		{
diff --git a/Engine/Source/Runtime/Renderer/Private/PostProcess/SubsurfaceCommon.h b/Engine/Source/Runtime/Renderer/Private/PostProcess/SubsurfaceCommon.h
--- a/Engine/Source/Runtime/Renderer/Private/PostProcess/SubsurfaceCommon.h
+++ b/Engine/Source/Runtime/Renderer/Private/PostProcess/SubsurfaceCommon.h
@@ -54,7 +54,7 @@ class FSubsurfaceShader : public FGlobalShader
 public:
 	static bool ShouldCompilePermutation(const FGlobalShaderPermutationParameters& Parameters)
 	{
-		return IsFeatureLevelSupported(Parameters.Platform, ERHIFeatureLevel::SM4);
+		return IsFeatureLevelSupported(Parameters.Platform, ERHIFeatureLevel::SM5);
 	}
 
 	static void ModifyCompilationEnvironment(const FGlobalShaderPermutationParameters& Parameters, FShaderCompilerEnvironment& OutEnvironment)
diff --git a/Engine/Source/Runtime/Renderer/Private/PostProcess/VisualizeShadingModels.cpp b/Engine/Source/Runtime/Renderer/Private/PostProcess/VisualizeShadingModels.cpp
--- a/Engine/Source/Runtime/Renderer/Private/PostProcess/VisualizeShadingModels.cpp
+++ b/Engine/Source/Runtime/Renderer/Private/PostProcess/VisualizeShadingModels.cpp
@@ -23,7 +23,7 @@ class FPostProcessVisualizeShadingModelsPS : public FGlobalShader
 
 	static bool ShouldCompilePermutation(const FGlobalShaderPermutationParameters& Parameters)
 	{
-		return IsFeatureLevelSupported(Parameters.Platform, ERHIFeatureLevel::SM4);
+		return IsFeatureLevelSupported(Parameters.Platform, ERHIFeatureLevel::SM5);
 	}
 
 	/** Default constructor. */
diff --git a/Engine/Source/Runtime/Renderer/Private/ReflectionEnvironment.cpp b/Engine/Source/Runtime/Renderer/Private/ReflectionEnvironment.cpp
--- a/Engine/Source/Runtime/Renderer/Private/ReflectionEnvironment.cpp
+++ b/Engine/Source/Runtime/Renderer/Private/ReflectionEnvironment.cpp
@@ -188,7 +188,7 @@ FVector GetReflectionEnvironmentRoughnessMixingScaleBiasAndLargestWeight()
 
 bool IsReflectionEnvironmentAvailable(ERHIFeatureLevel::Type InFeatureLevel)
 {
-	return (InFeatureLevel >= ERHIFeatureLevel::SM4) && (GetReflectionEnvironmentCVar() != 0);
+	return (InFeatureLevel >= ERHIFeatureLevel::SM5) && (GetReflectionEnvironmentCVar() != 0);
 }
 
 bool IsReflectionCaptureAvailable()
@@ -520,7 +520,7 @@ class FReflectionEnvironmentSkyLightingPS : public FGlobalShader
 
 	static bool ShouldCompilePermutation(const FGlobalShaderPermutationParameters& Parameters)
 	{
-		if (!IsFeatureLevelSupported(Parameters.Platform, ERHIFeatureLevel::SM4))
+		if (!IsFeatureLevelSupported(Parameters.Platform, ERHIFeatureLevel::SM5))
 		{
 			return false;
 		}
diff --git a/Engine/Source/Runtime/Renderer/Private/ReflectionEnvironmentCapture.cpp b/Engine/Source/Runtime/Renderer/Private/ReflectionEnvironmentCapture.cpp
--- a/Engine/Source/Runtime/Renderer/Private/ReflectionEnvironmentCapture.cpp
+++ b/Engine/Source/Runtime/Renderer/Private/ReflectionEnvironmentCapture.cpp
@@ -1648,7 +1648,7 @@ void FScene::UpdateSkyCaptureContents(
 	FSHVectorRGB3& OutIrradianceEnvironmentMap,
 	TArray<FFloat16Color>* OutRadianceMap)
 {	
-	if (GSupportsRenderTargetFormat_PF_FloatRGBA || GetFeatureLevel() >= ERHIFeatureLevel::SM4)
+	if (GSupportsRenderTargetFormat_PF_FloatRGBA || GetFeatureLevel() >= ERHIFeatureLevel::SM5)
 	{
 		QUICK_SCOPE_CYCLE_COUNTER(STAT_UpdateSkyCaptureContents);
 		{
diff --git a/Engine/Source/Runtime/Renderer/Private/SceneCaptureRendering.cpp b/Engine/Source/Runtime/Renderer/Private/SceneCaptureRendering.cpp
--- a/Engine/Source/Runtime/Renderer/Private/SceneCaptureRendering.cpp
+++ b/Engine/Source/Runtime/Renderer/Private/SceneCaptureRendering.cpp
@@ -76,7 +76,7 @@ public:
 
 	static bool ShouldCompilePermutation(const FGlobalShaderPermutationParameters& Parameters) 
 	{ 
-		return IsFeatureLevelSupported(Parameters.Platform, ERHIFeatureLevel::SM4);
+		return IsFeatureLevelSupported(Parameters.Platform, ERHIFeatureLevel::SM5);
 	}
 
 	static void ModifyCompilationEnvironment(const FGlobalShaderPermutationParameters& Parameters, FShaderCompilerEnvironment& OutEnvironment)
diff --git a/Engine/Source/Runtime/Renderer/Private/SceneOcclusion.cpp b/Engine/Source/Runtime/Renderer/Private/SceneOcclusion.cpp
--- a/Engine/Source/Runtime/Renderer/Private/SceneOcclusion.cpp
+++ b/Engine/Source/Runtime/Renderer/Private/SceneOcclusion.cpp
@@ -688,7 +688,7 @@ void FHZBOcclusionTester::SetInvalidFrameNumber()
 
 void FHZBOcclusionTester::InitDynamicRHI()
 {
-	if (GetFeatureLevel() >= ERHIFeatureLevel::SM4)
+	if (GetFeatureLevel() >= ERHIFeatureLevel::SM5)
 	{
 		FRHICommandListImmediate& RHICmdList = FRHICommandListExecutor::GetImmediateCommandList();
 		FPooledRenderTargetDesc Desc( FPooledRenderTargetDesc::Create2DDesc( FIntPoint( SizeX, SizeY ), PF_B8G8R8A8, FClearValueBinding::None, TexCreate_CPUReadback | TexCreate_HideInVisualizeTexture, TexCreate_None, false ) );
@@ -698,7 +698,7 @@ void FHZBOcclusionTester::InitDynamicRHI()
 
 void FHZBOcclusionTester::ReleaseDynamicRHI()
 {
-	if (GetFeatureLevel() >= ERHIFeatureLevel::SM4)
+	if (GetFeatureLevel() >= ERHIFeatureLevel::SM5)
 	{
 		GRenderTargetPool.FreeUnusedResource( ResultsTextureCPU );
 	}
@@ -790,7 +790,7 @@ class FHZBTestPS : public FGlobalShader
 
 	static bool ShouldCompilePermutation(const FGlobalShaderPermutationParameters& Parameters)
 	{
-		return IsFeatureLevelSupported(Parameters.Platform, ERHIFeatureLevel::SM4);
+		return IsFeatureLevelSupported(Parameters.Platform, ERHIFeatureLevel::SM5);
 	}
 
 	FHZBTestPS() {}
diff --git a/Engine/Source/Runtime/Renderer/Private/SceneRendering.cpp b/Engine/Source/Runtime/Renderer/Private/SceneRendering.cpp
--- a/Engine/Source/Runtime/Renderer/Private/SceneRendering.cpp
+++ b/Engine/Source/Runtime/Renderer/Private/SceneRendering.cpp
@@ -853,7 +853,7 @@ void FViewInfo::Init()
 	}
 
 	const int32 MaxMobileShadowCascadeCount = FMath::Clamp(CVarMaxMobileShadowCascades.GetValueOnAnyThread(), 0, MAX_MOBILE_SHADOWCASCADES);
-	const int32 MaxShadowCascadeCountUpperBound = GetFeatureLevel() >= ERHIFeatureLevel::SM4 ? 10 : MaxMobileShadowCascadeCount;
+	const int32 MaxShadowCascadeCountUpperBound = GetFeatureLevel() >= ERHIFeatureLevel::SM5 ? 10 : MaxMobileShadowCascadeCount;
 
 	MaxShadowCascades = FMath::Clamp<int32>(CVarMaxShadowCascades.GetValueOnAnyThread(), 0, MaxShadowCascadeCountUpperBound);
 
@@ -2166,7 +2166,7 @@ void FSceneRenderer::PrepareViewRectsForRendering()
 
 		// Fallback to no anti aliasing.
 		{
-			const bool bWillApplyTemporalAA = GPostProcessing.AllowFullPostProcessing(View, FeatureLevel) || (View.bIsPlanarReflection && FeatureLevel >= ERHIFeatureLevel::SM4);
+			const bool bWillApplyTemporalAA = GPostProcessing.AllowFullPostProcessing(View, FeatureLevel) || (View.bIsPlanarReflection && FeatureLevel >= ERHIFeatureLevel::SM5);
 
 			if (!bWillApplyTemporalAA)
 			{
diff --git a/Engine/Source/Runtime/Renderer/Private/SceneVisibility.cpp b/Engine/Source/Runtime/Renderer/Private/SceneVisibility.cpp
--- a/Engine/Source/Runtime/Renderer/Private/SceneVisibility.cpp
+++ b/Engine/Source/Runtime/Renderer/Private/SceneVisibility.cpp
@@ -3045,7 +3045,7 @@ void FSceneRenderer::PreVisibilityFrameSetup(FRHICommandListImmediate& RHICmdLis
 			// Compute number of TAA samples.
 			int32 TemporalAASamples = CVarTemporalAASamplesValue;
 			{
-				if (Scene->GetFeatureLevel() < ERHIFeatureLevel::SM4)
+				if (Scene->GetFeatureLevel() < ERHIFeatureLevel::SM5)
 				{
 					// Only support 2 samples for mobile temporal AA.
 					TemporalAASamples = 2;
@@ -3079,7 +3079,7 @@ void FSceneRenderer::PreVisibilityFrameSetup(FRHICommandListImmediate& RHICmdLis
 
 			// Choose sub pixel sample coordinate in the temporal sequence.
 			float SampleX, SampleY;
-			if (Scene->GetFeatureLevel() < ERHIFeatureLevel::SM4)
+			if (Scene->GetFeatureLevel() < ERHIFeatureLevel::SM5)
 			{
 				float SamplesX[] = { -8.0f/16.0f, 0.0/16.0f };
 				float SamplesY[] = { /* - */ 0.0f/16.0f, 8.0/16.0f };
diff --git a/Engine/Source/Runtime/Renderer/Private/ScreenSpaceRayTracing.cpp b/Engine/Source/Runtime/Renderer/Private/ScreenSpaceRayTracing.cpp
--- a/Engine/Source/Runtime/Renderer/Private/ScreenSpaceRayTracing.cpp
+++ b/Engine/Source/Runtime/Renderer/Private/ScreenSpaceRayTracing.cpp
@@ -275,7 +275,7 @@ class FScreenSpaceReflectionsStencilPS : public FGlobalShader
 
 	static bool ShouldCompilePermutation(const FGlobalShaderPermutationParameters& Parameters)
 	{
-		return IsFeatureLevelSupported(Parameters.Platform, ERHIFeatureLevel::SM4);
+		return IsFeatureLevelSupported(Parameters.Platform, ERHIFeatureLevel::SM5);
 	}
 
 	static void ModifyCompilationEnvironment(const FGlobalShaderPermutationParameters& Parameters, FShaderCompilerEnvironment& OutEnvironment)
@@ -300,7 +300,7 @@ class FScreenSpaceReflectionsPS : public FGlobalShader
 	static bool ShouldCompilePermutation(const FGlobalShaderPermutationParameters& Parameters)
 	{
 		FPermutationDomain PermutationVector(Parameters.PermutationId);
-		return IsFeatureLevelSupported(Parameters.Platform, ERHIFeatureLevel::SM4);
+		return IsFeatureLevelSupported(Parameters.Platform, ERHIFeatureLevel::SM5);
 	}
 	
 	BEGIN_SHADER_PARAMETER_STRUCT(FParameters, )
diff --git a/Engine/Source/Runtime/Renderer/Private/ShadowDepthRendering.cpp b/Engine/Source/Runtime/Renderer/Private/ShadowDepthRendering.cpp
--- a/Engine/Source/Runtime/Renderer/Private/ShadowDepthRendering.cpp
+++ b/Engine/Source/Runtime/Renderer/Private/ShadowDepthRendering.cpp
@@ -235,14 +235,14 @@ public:
 			return (Material->IsSpecialEngineMaterial()
 				// Masked and WPO materials need their shaders but cannot be used with a position only stream.
 				|| ((!Material->WritesEveryPixel(true) || Material->MaterialMayModifyMeshPosition()) && !bUsePositionOnlyStream))
-				// Only compile one pass point light shaders for feature levels >= SM4
-				&& (ShaderMode != VertexShadowDepth_OnePassPointLight || IsFeatureLevelSupported(Platform, ERHIFeatureLevel::SM4))
+				// Only compile one pass point light shaders for feature levels >= SM5
+				&& (ShaderMode != VertexShadowDepth_OnePassPointLight || IsFeatureLevelSupported(Platform, ERHIFeatureLevel::SM5))
 				// Only compile position-only shaders for vertex factories that support it. (Note: this assumes that a vertex factor which supports PositionOnly, supports also PositionAndNormalOnly)
 				&& (!bUsePositionOnlyStream || Parameters.VertexFactoryType->SupportsPositionOnly())
 				// Don't render ShadowDepth for translucent unlit materials
 				&& Material->ShouldCastDynamicShadows()
-				// Only compile perspective correct light shaders for feature levels >= SM4
-				&& (ShaderMode != VertexShadowDepth_PerspectiveCorrect || IsFeatureLevelSupported(Platform, ERHIFeatureLevel::SM4));
+				// Only compile perspective correct light shaders for feature levels >= SM5
+				&& (ShaderMode != VertexShadowDepth_PerspectiveCorrect || IsFeatureLevelSupported(Platform, ERHIFeatureLevel::SM5));
 		}
 	}
 
@@ -469,7 +469,7 @@ public:
 		const EShaderPlatform Platform = Parameters.Platform;
 		const FMaterial* Material = Parameters.Material;
 
-		if (!IsFeatureLevelSupported(Platform, ERHIFeatureLevel::SM4))
+		if (!IsFeatureLevelSupported(Platform, ERHIFeatureLevel::SM5))
 		{
 			return (Material->IsSpecialEngineMaterial()
 				// Only compile for masked or lit translucent materials
@@ -501,11 +501,11 @@ public:
 				|| (Material->MaterialMayModifyMeshPosition() && Material->IsUsedWithInstancedStaticMeshes())
 				// Perspective correct rendering needs a pixel shader and WPO materials can't be overridden with default material.
 				|| (ShaderMode == PixelShadowDepth_PerspectiveCorrect && Material->MaterialMayModifyMeshPosition()))
-				// Only compile one pass point light shaders for feature levels >= SM4
-				&& (ShaderMode != PixelShadowDepth_OnePassPointLight || IsFeatureLevelSupported(Platform, ERHIFeatureLevel::SM4))
+				// Only compile one pass point light shaders for feature levels >= SM5
+				&& (ShaderMode != PixelShadowDepth_OnePassPointLight || IsFeatureLevelSupported(Platform, ERHIFeatureLevel::SM5))
 				// Don't render ShadowDepth for translucent unlit materials
 				&& Material->ShouldCastDynamicShadows()
-				&& IsFeatureLevelSupported(Platform, ERHIFeatureLevel::SM4);
+				&& IsFeatureLevelSupported(Platform, ERHIFeatureLevel::SM5);
 		}
 	}
 
@@ -865,7 +865,7 @@ public:
 
 	static bool ShouldCompilePermutation(const FGlobalShaderPermutationParameters& Parameters)
 	{
-		return RHISupportsGeometryShaders(Parameters.Platform) && IsFeatureLevelSupported(Parameters.Platform, ERHIFeatureLevel::SM4);
+		return RHISupportsGeometryShaders(Parameters.Platform) && IsFeatureLevelSupported(Parameters.Platform, ERHIFeatureLevel::SM5);
 	}
 
 	FCopyShadowMapsCubeGS(const ShaderMetaType::CompiledShaderInitializerType& Initializer) :
@@ -884,7 +884,7 @@ public:
 
 	static bool ShouldCompilePermutation(const FGlobalShaderPermutationParameters& Parameters)
 	{
-		return IsFeatureLevelSupported(Parameters.Platform, ERHIFeatureLevel::SM4);
+		return IsFeatureLevelSupported(Parameters.Platform, ERHIFeatureLevel::SM5);
 	}
 
 	FCopyShadowMapsCubePS(const ShaderMetaType::CompiledShaderInitializerType& Initializer) :
@@ -924,7 +924,7 @@ public:
 
 	static bool ShouldCompilePermutation(const FGlobalShaderPermutationParameters& Parameters)
 	{
-		return IsFeatureLevelSupported(Parameters.Platform, ERHIFeatureLevel::SM4);
+		return IsFeatureLevelSupported(Parameters.Platform, ERHIFeatureLevel::SM5);
 	}
 
 	FCopyShadowMaps2DPS(const ShaderMetaType::CompiledShaderInitializerType& Initializer) :
diff --git a/Engine/Source/Runtime/Renderer/Private/ShadowRendering.h b/Engine/Source/Runtime/Renderer/Private/ShadowRendering.h
--- a/Engine/Source/Runtime/Renderer/Private/ShadowRendering.h
+++ b/Engine/Source/Runtime/Renderer/Private/ShadowRendering.h
@@ -933,7 +933,7 @@ public:
 
 	static bool ShouldCompilePermutation(const FGlobalShaderPermutationParameters& Parameters)
 	{
-		return IsFeatureLevelSupported(Parameters.Platform, ERHIFeatureLevel::SM4);
+		return IsFeatureLevelSupported(Parameters.Platform, ERHIFeatureLevel::SM5);
 	}
 
 	/**
@@ -1120,7 +1120,7 @@ public:
 
 	static bool ShouldCompilePermutation(const FGlobalShaderPermutationParameters& Parameters)
 	{
-		return IsFeatureLevelSupported(Parameters.Platform, ERHIFeatureLevel::SM4) && TShadowProjectionPS<Quality>::ShouldCompilePermutation(Parameters);
+		return IsFeatureLevelSupported(Parameters.Platform, ERHIFeatureLevel::SM5) && TShadowProjectionPS<Quality>::ShouldCompilePermutation(Parameters);
 	}
 
 	TShadowProjectionFromTranslucencyPS() {}
@@ -1296,7 +1296,7 @@ public:
 
 	static bool ShouldCompilePermutation(const FGlobalShaderPermutationParameters& Parameters)
 	{
-		return IsFeatureLevelSupported(Parameters.Platform, ERHIFeatureLevel::SM4);
+		return IsFeatureLevelSupported(Parameters.Platform, ERHIFeatureLevel::SM5);
 	}
 
 	void SetParameters(
diff --git a/Engine/Source/Runtime/Renderer/Private/ShadowSetup.cpp b/Engine/Source/Runtime/Renderer/Private/ShadowSetup.cpp
--- a/Engine/Source/Runtime/Renderer/Private/ShadowSetup.cpp
+++ b/Engine/Source/Runtime/Renderer/Private/ShadowSetup.cpp
@@ -2011,7 +2011,7 @@ void FSceneRenderer::CreatePerObjectProjectedShadow(
 			}
 
 			if (bTranslucentRelevance
-				&& Scene->GetFeatureLevel() >= ERHIFeatureLevel::SM4
+				&& Scene->GetFeatureLevel() >= ERHIFeatureLevel::SM5
 				&& bCreateTranslucentObjectShadow 
 				&& (bTranslucentShadowIsVisibleThisFrame || bShadowIsPotentiallyVisibleNextFrame))
 			{
@@ -2057,7 +2057,7 @@ void FSceneRenderer::CreatePerObjectProjectedShadow(
 		if (MaxPreFadeAlpha > 1.0f / 256.0f 
 			&& bRenderPreShadow
 			&& bOpaqueRelevance
-			&& Scene->GetFeatureLevel() >= ERHIFeatureLevel::SM4)
+			&& Scene->GetFeatureLevel() >= ERHIFeatureLevel::SM5)
 		{
 			// Round down to the nearest power of two so that resolution changes are always doubling or halving the resolution, which increases filtering stability.
 			int32 PreshadowSizeX = 1 << (FMath::CeilLogTwo(FMath::TruncToInt(MaxDesiredResolution * CVarPreShadowResolutionFactor.GetValueOnRenderThread())) - 1);
@@ -3319,7 +3319,7 @@ void FSceneRenderer::GatherShadowPrimitives(
 static bool NeedsUnatlasedCSMDepthsWorkaround(ERHIFeatureLevel::Type FeatureLevel)
 {
 	// UE-42131: Excluding mobile from this, mobile renderer relies on the depth texture border.
-	return GRHINeedsUnatlasedCSMDepthsWorkaround && (FeatureLevel >= ERHIFeatureLevel::SM4);
+	return GRHINeedsUnatlasedCSMDepthsWorkaround && (FeatureLevel >= ERHIFeatureLevel::SM5);
 }
 
 void FSceneRenderer::AddViewDependentWholeSceneShadowsForView(
@@ -3523,7 +3523,7 @@ void FSceneRenderer::AllocateShadowDepthTargets(FRHICommandListImmediate& RHICmd
 				}
 			}
 
-			if (FeatureLevel < ERHIFeatureLevel::SM4 
+			if (FeatureLevel < ERHIFeatureLevel::SM5
 				// Mobile renderer only supports opaque per-object shadows or CSM
 				&& (!ProjectedShadowInfo->bPerObjectOpaqueShadow && !(ProjectedShadowInfo->bDirectionalLight && ProjectedShadowInfo->bWholeSceneShadow)))
 			{
@@ -3564,7 +3564,7 @@ void FSceneRenderer::AllocateShadowDepthTargets(FRHICommandListImmediate& RHICmd
 
 				bool bNeedsProjection = ProjectedShadowInfo->CacheMode != SDCM_StaticPrimitivesOnly
 					// Mobile rendering only projects opaque per object shadows.
-					&& (FeatureLevel >= ERHIFeatureLevel::SM4 || ProjectedShadowInfo->bPerObjectOpaqueShadow);
+					&& (FeatureLevel >= ERHIFeatureLevel::SM5 || ProjectedShadowInfo->bPerObjectOpaqueShadow);
 
 				if (bNeedsProjection)
 				{
@@ -3964,7 +3964,7 @@ void FSceneRenderer::AllocateRSMDepthTargets(FRHICommandListImmediate& RHICmdLis
 
 void FSceneRenderer::AllocateOnePassPointLightDepthTargets(FRHICommandListImmediate& RHICmdList, const TArray<FProjectedShadowInfo*, SceneRenderingAllocator>& WholeScenePointShadows)
 {
-	if (FeatureLevel >= ERHIFeatureLevel::SM4)
+	if (FeatureLevel >= ERHIFeatureLevel::SM5)
 	{
 		for (int32 ShadowIndex = 0; ShadowIndex < WholeScenePointShadows.Num(); ShadowIndex++)
 		{
@@ -4023,7 +4023,7 @@ TCHAR* const GetTranslucencyShadowTransmissionName(uint32 Id)
 
 void FSceneRenderer::AllocateTranslucentShadowDepthTargets(FRHICommandListImmediate& RHICmdList, TArray<FProjectedShadowInfo*, SceneRenderingAllocator>& TranslucentShadows)
 {
-	if (TranslucentShadows.Num() > 0 && FeatureLevel >= ERHIFeatureLevel::SM4)
+	if (TranslucentShadows.Num() > 0 && FeatureLevel >= ERHIFeatureLevel::SM5)
 	{
 		FSceneRenderTargets& SceneContext = FSceneRenderTargets::Get(RHICmdList);
 		const FIntPoint TranslucentShadowBufferResolution = SceneContext.GetTranslucentShadowDepthTextureResolution();
@@ -4104,7 +4104,7 @@ void FSceneRenderer::InitDynamicShadows(FRHICommandListImmediate& RHICmdList, FG
 	SCOPE_CYCLE_COUNTER(STAT_DynamicShadowSetupTime);
 	SCOPED_NAMED_EVENT(FSceneRenderer_InitDynamicShadows, FColor::Magenta);
 
-	const bool bMobile = FeatureLevel < ERHIFeatureLevel::SM4;
+	const bool bMobile = FeatureLevel < ERHIFeatureLevel::SM5;
 
 	bool bStaticSceneOnly = false;
 
diff --git a/Engine/Source/Runtime/Renderer/Private/SystemTextures.cpp b/Engine/Source/Runtime/Renderer/Private/SystemTextures.cpp
--- a/Engine/Source/Runtime/Renderer/Private/SystemTextures.cpp
+++ b/Engine/Source/Runtime/Renderer/Private/SystemTextures.cpp
@@ -196,7 +196,7 @@ void FSystemTextures::InitializeFeatureLevelDependentTextures(FRHICommandListImm
 		}
 
 	// Create a VolumetricBlackDummy texture
-	if (CurrentFeatureLevel < ERHIFeatureLevel::SM4 && InFeatureLevel >= ERHIFeatureLevel::SM4)
+	if (CurrentFeatureLevel < ERHIFeatureLevel::SM5 && InFeatureLevel >= ERHIFeatureLevel::SM5)
 	{
 		FPooledRenderTargetDesc Desc(FPooledRenderTargetDesc::CreateVolumeDesc(1, 1, 1, PF_B8G8R8A8, FClearValueBinding::Transparent, TexCreate_HideInVisualizeTexture, TexCreate_ShaderResource | TexCreate_RenderTargetable | TexCreate_NoFastClear, false));
 		Desc.AutoWritable = false;
@@ -213,7 +213,7 @@ void FSystemTextures::InitializeFeatureLevelDependentTextures(FRHICommandListImm
 			BlackBytes);
 	}
 
-	if (CurrentFeatureLevel < ERHIFeatureLevel::SM4 && InFeatureLevel >= ERHIFeatureLevel::SM4)
+	if (CurrentFeatureLevel < ERHIFeatureLevel::SM5 && InFeatureLevel >= ERHIFeatureLevel::SM5)
 		{
 	// Create the PerlinNoise3D texture (similar to http://prettyprocs.wordpress.com/2012/10/20/fast-perlin-noise/)
 	{
@@ -550,7 +550,7 @@ void FSystemTextures::InitializeFeatureLevelDependentTextures(FRHICommandListImm
 			RHICmdList.UnlockTexture2D((FTexture2DRHIRef&)LTCAmp->GetRenderTargetItem().ShaderResourceTexture, 0, false);
 		}
 		} // end Create the SSAO randomization texture
-	} // end if (FeatureLevelInitializedTo < ERHIFeatureLevel::SM4 && InFeatureLevel >= ERHIFeatureLevel::SM4)
+	} // end if (FeatureLevelInitializedTo < ERHIFeatureLevel::SM5 && InFeatureLevel >= ERHIFeatureLevel::SM5)
 
 	// Initialize textures only once.
 	FeatureLevelInitializedTo = InFeatureLevel;
diff --git a/Engine/Source/Runtime/Renderer/Private/TranslucentLighting.cpp b/Engine/Source/Runtime/Renderer/Private/TranslucentLighting.cpp
--- a/Engine/Source/Runtime/Renderer/Private/TranslucentLighting.cpp
+++ b/Engine/Source/Runtime/Renderer/Private/TranslucentLighting.cpp
@@ -289,7 +289,7 @@ public:
 
 	static bool ShouldCompilePermutation(const FMeshMaterialShaderPermutationParameters& Parameters)
 	{
-		return IsTranslucentBlendMode(Parameters.Material->GetBlendMode()) && IsFeatureLevelSupported(Parameters.Platform, ERHIFeatureLevel::SM4);
+		return IsTranslucentBlendMode(Parameters.Material->GetBlendMode()) && IsFeatureLevelSupported(Parameters.Platform, ERHIFeatureLevel::SM5);
 	}
 
 	FTranslucencyShadowDepthVS() {}
@@ -339,7 +339,7 @@ public:
 
 	static bool ShouldCompilePermutation(const FMeshMaterialShaderPermutationParameters& Parameters)
 	{
-		return IsTranslucentBlendMode(Parameters.Material->GetBlendMode()) && IsFeatureLevelSupported(Parameters.Platform, ERHIFeatureLevel::SM4);
+		return IsTranslucentBlendMode(Parameters.Material->GetBlendMode()) && IsFeatureLevelSupported(Parameters.Platform, ERHIFeatureLevel::SM5);
 	}
 
 	FTranslucencyShadowDepthPS(const ShaderMetaType::CompiledShaderInitializerType& Initializer):
@@ -635,7 +635,7 @@ public:
 
 	static bool ShouldCompilePermutation(const FGlobalShaderPermutationParameters& Parameters) 
 	{ 
-		return IsFeatureLevelSupported(Parameters.Platform, ERHIFeatureLevel::SM4) && (RHISupportsGeometryShaders(Parameters.Platform) || RHISupportsVertexShaderLayer(Parameters.Platform));
+		return IsFeatureLevelSupported(Parameters.Platform, ERHIFeatureLevel::SM5) && (RHISupportsGeometryShaders(Parameters.Platform) || RHISupportsVertexShaderLayer(Parameters.Platform));
 	}
 
 	FFilterTranslucentVolumePS(const ShaderMetaType::CompiledShaderInitializerType& Initializer):
@@ -764,7 +764,7 @@ public:
 
 	static bool ShouldCompilePermutation(const FGlobalShaderPermutationParameters& Parameters)
 	{
-		return IsFeatureLevelSupported(Parameters.Platform, ERHIFeatureLevel::SM4) && (RHISupportsGeometryShaders(Parameters.Platform) || RHISupportsVertexShaderLayer(Parameters.Platform));
+		return IsFeatureLevelSupported(Parameters.Platform, ERHIFeatureLevel::SM5) && (RHISupportsGeometryShaders(Parameters.Platform) || RHISupportsVertexShaderLayer(Parameters.Platform));
 	}
 
 	FTranslucentObjectShadowingPS(const ShaderMetaType::CompiledShaderInitializerType& Initializer):
@@ -822,7 +822,7 @@ public:
 	static bool ShouldCompilePermutation(const FMaterialShaderPermutationParameters& Parameters)
 	{
 		return (Parameters.Material->IsLightFunction() || Parameters.Material->IsSpecialEngineMaterial()) &&
-			(IsFeatureLevelSupported(Parameters.Platform, ERHIFeatureLevel::SM4) &&
+			(IsFeatureLevelSupported(Parameters.Platform, ERHIFeatureLevel::SM5) &&
 			(RHISupportsGeometryShaders(Parameters.Platform) || RHISupportsVertexShaderLayer(Parameters.Platform)));
 	}
 
@@ -1062,7 +1062,7 @@ class FInjectAmbientCubemapPS : public FGlobalShader
 
 	static bool ShouldCompilePermutation(const FGlobalShaderPermutationParameters& Parameters)
 	{
-		return IsFeatureLevelSupported(Parameters.Platform, ERHIFeatureLevel::SM4);
+		return IsFeatureLevelSupported(Parameters.Platform, ERHIFeatureLevel::SM5);
 	}
 
 	/** Default constructor. */
@@ -1617,7 +1617,7 @@ public:
 
 	static bool ShouldCompilePermutation(const FGlobalShaderPermutationParameters& Parameters) 
 	{ 
-		return IsFeatureLevelSupported(Parameters.Platform, ERHIFeatureLevel::SM4) && (RHISupportsGeometryShaders(Parameters.Platform) || RHISupportsVertexShaderLayer(Parameters.Platform));
+		return IsFeatureLevelSupported(Parameters.Platform, ERHIFeatureLevel::SM5) && (RHISupportsGeometryShaders(Parameters.Platform) || RHISupportsVertexShaderLayer(Parameters.Platform));
 	}
 
 	FSimpleLightTranslucentLightingInjectPS(const ShaderMetaType::CompiledShaderInitializerType& Initializer):
diff --git a/Engine/Source/Runtime/Renderer/Private/TranslucentRendering.cpp b/Engine/Source/Runtime/Renderer/Private/TranslucentRendering.cpp
--- a/Engine/Source/Runtime/Renderer/Private/TranslucentRendering.cpp
+++ b/Engine/Source/Runtime/Renderer/Private/TranslucentRendering.cpp
@@ -236,7 +236,7 @@ class FCopySceneColorPS : public FGlobalShader
 	DECLARE_SHADER_TYPE(FCopySceneColorPS, Global);
 public:
 
-	static bool ShouldCompilePermutation(const FGlobalShaderPermutationParameters& Parameters) { return IsFeatureLevelSupported(Parameters.Platform, ERHIFeatureLevel::SM4); }
+	static bool ShouldCompilePermutation(const FGlobalShaderPermutationParameters& Parameters) { return IsFeatureLevelSupported(Parameters.Platform, ERHIFeatureLevel::SM5); }
 
 	FCopySceneColorPS(const ShaderMetaType::CompiledShaderInitializerType& Initializer) :
 		FGlobalShader(Initializer)
@@ -692,7 +692,7 @@ protected:
 
 	static bool ShouldCompilePermutation(const FGlobalShaderPermutationParameters& Parameters)
 	{
-		return IsFeatureLevelSupported(Parameters.Platform, ERHIFeatureLevel::SM4);
+		return IsFeatureLevelSupported(Parameters.Platform, ERHIFeatureLevel::SM5);
 	}
 
 	/** Default constructor. */
diff --git a/Engine/Source/Runtime/Renderer/Private/VT/RuntimeVirtualTextureRender.cpp b/Engine/Source/Runtime/Renderer/Private/VT/RuntimeVirtualTextureRender.cpp
--- a/Engine/Source/Runtime/Renderer/Private/VT/RuntimeVirtualTextureRender.cpp
+++ b/Engine/Source/Runtime/Renderer/Private/VT/RuntimeVirtualTextureRender.cpp
@@ -26,7 +26,7 @@ namespace RuntimeVirtualTexture
 
 		static bool ShouldCompilePermutation(const FMeshMaterialShaderPermutationParameters& Parameters)
 		{
-			return IsFeatureLevelSupported(Parameters.Platform, ERHIFeatureLevel::SM4) && 
+			return IsFeatureLevelSupported(Parameters.Platform, ERHIFeatureLevel::SM5) && 
 				(Parameters.Material->GetMaterialDomain() == MD_RuntimeVirtualTexture || Parameters.Material->HasRuntimeVirtualTextureOutput());
 		}
 
diff --git a/Engine/Source/Runtime/Renderer/Private/VelocityRendering.cpp b/Engine/Source/Runtime/Renderer/Private/VelocityRendering.cpp
--- a/Engine/Source/Runtime/Renderer/Private/VelocityRendering.cpp
+++ b/Engine/Source/Runtime/Renderer/Private/VelocityRendering.cpp
@@ -82,7 +82,7 @@ public:
 			|| (Material->IsTwoSided() && !IsTranslucentBlendMode(Material->GetBlendMode()))
 			// or if the material modifies meshes
 			|| Material->MaterialMayModifyMeshPosition()))
-			&& IsFeatureLevelSupported(Platform, ERHIFeatureLevel::SM4) 
+			&& IsFeatureLevelSupported(Platform, ERHIFeatureLevel::SM5) 
 			&& !FVelocityRendering::VertexFactoryOnlyOutputsVelocityInBasePass(Platform, Parameters.VertexFactoryType->SupportsStaticLighting());
 	}
 
@@ -183,7 +183,7 @@ public:
 			|| (Material->IsTwoSided() && !IsTranslucentBlendMode(Material->GetBlendMode()))
 			// or if the material modifies meshes
 			|| Material->MaterialMayModifyMeshPosition()))
-			&& IsFeatureLevelSupported(Platform, ERHIFeatureLevel::SM4) && 
+			&& IsFeatureLevelSupported(Platform, ERHIFeatureLevel::SM5) && 
 			!FVelocityRendering::VertexFactoryOnlyOutputsVelocityInBasePass(Platform, Parameters.VertexFactoryType->SupportsStaticLighting());
 	}
 
@@ -409,7 +409,7 @@ void FDeferredShadingSceneRenderer::RenderVelocities(FRHICommandListImmediate& R
 {
 	SCOPED_NAMED_EVENT(FDeferredShadingSceneRenderer_RenderVelocities, FColor::Emerald);
 
-	check(FeatureLevel >= ERHIFeatureLevel::SM4);
+	check(FeatureLevel >= ERHIFeatureLevel::SM5);
 	CSV_SCOPED_TIMING_STAT_EXCLUSIVE(RenderVelocities);
 	SCOPE_CYCLE_COUNTER(STAT_RenderVelocities);
 
diff --git a/Engine/Source/Runtime/Renderer/Private/VisualizeVolumetricLightmap.cpp b/Engine/Source/Runtime/Renderer/Private/VisualizeVolumetricLightmap.cpp
--- a/Engine/Source/Runtime/Renderer/Private/VisualizeVolumetricLightmap.cpp
+++ b/Engine/Source/Runtime/Renderer/Private/VisualizeVolumetricLightmap.cpp
@@ -53,7 +53,7 @@ public:
 
 	static bool ShouldCompilePermutation(const FGlobalShaderPermutationParameters& Parameters)
 	{
-		return IsFeatureLevelSupported(Parameters.Platform, ERHIFeatureLevel::SM4);
+		return IsFeatureLevelSupported(Parameters.Platform, ERHIFeatureLevel::SM5);
 	}
 
 	static void ModifyCompilationEnvironment(const FGlobalShaderPermutationParameters& Parameters, FShaderCompilerEnvironment& OutEnvironment)
@@ -105,7 +105,7 @@ public:
 
 	static bool ShouldCompilePermutation(const FGlobalShaderPermutationParameters& Parameters)
 	{
-		return IsFeatureLevelSupported(Parameters.Platform, ERHIFeatureLevel::SM4);
+		return IsFeatureLevelSupported(Parameters.Platform, ERHIFeatureLevel::SM5);
 	}
 
 	/** Default constructor. */
diff --git a/Engine/Source/Runtime/Renderer/Public/VolumeRendering.h b/Engine/Source/Runtime/Renderer/Public/VolumeRendering.h
--- a/Engine/Source/Runtime/Renderer/Public/VolumeRendering.h
+++ b/Engine/Source/Runtime/Renderer/Public/VolumeRendering.h
@@ -50,7 +50,7 @@ public:
 
 	static bool ShouldCompilePermutation(const FGlobalShaderPermutationParameters& Parameters) 
 	{ 
-		return IsFeatureLevelSupported(Parameters.Platform, ERHIFeatureLevel::SM4); 
+		return IsFeatureLevelSupported(Parameters.Platform, ERHIFeatureLevel::SM5); 
 	}
 
 	static void ModifyCompilationEnvironment(const FGlobalShaderPermutationParameters& Parameters, FShaderCompilerEnvironment& OutEnvironment )
@@ -102,7 +102,7 @@ public:
 
 	static bool ShouldCompilePermutation(const FGlobalShaderPermutationParameters& Parameters) 
 	{ 
-		return IsFeatureLevelSupported(Parameters.Platform, ERHIFeatureLevel::SM4) && RHISupportsGeometryShaders(Parameters.Platform); 
+		return IsFeatureLevelSupported(Parameters.Platform, ERHIFeatureLevel::SM5) && RHISupportsGeometryShaders(Parameters.Platform); 
 	}
 
 	FWriteToSliceGS(const ShaderMetaType::CompiledShaderInitializerType& Initializer):
diff --git a/Engine/Source/Runtime/UtilityShaders/Public/ClearReplacementShaders.h b/Engine/Source/Runtime/UtilityShaders/Public/ClearReplacementShaders.h
--- a/Engine/Source/Runtime/UtilityShaders/Public/ClearReplacementShaders.h
+++ b/Engine/Source/Runtime/UtilityShaders/Public/ClearReplacementShaders.h
@@ -27,7 +27,7 @@ public:
 	
 	static bool ShouldCompilePermutation(const FGlobalShaderPermutationParameters& Parameters)
 	{
-		return IsFeatureLevelSupported(Parameters.Platform, ERHIFeatureLevel::SM4);
+		return IsFeatureLevelSupported(Parameters.Platform, ERHIFeatureLevel::SM5);
 	}
 };
 
@@ -52,7 +52,7 @@ public:
 	
 	static bool ShouldCompilePermutation(const FGlobalShaderPermutationParameters& Parameters)
 	{
-		return IsFeatureLevelSupported(Parameters.Platform, ERHIFeatureLevel::SM4);
+		return IsFeatureLevelSupported(Parameters.Platform, ERHIFeatureLevel::SM5);
 	}
 	
 	static void ModifyCompilationEnvironment(const FGlobalShaderPermutationParameters& Parameters, FShaderCompilerEnvironment& OutEnvironment)
diff --git a/Engine/Source/Runtime/UtilityShaders/Public/ResolveShader.h b/Engine/Source/Runtime/UtilityShaders/Public/ResolveShader.h
--- a/Engine/Source/Runtime/UtilityShaders/Public/ResolveShader.h
+++ b/Engine/Source/Runtime/UtilityShaders/Public/ResolveShader.h
@@ -185,7 +185,7 @@ public:
 	
 	typedef FDummyResolveParameter FParameter;
 	
-	static bool ShouldCompilePermutation(const FGlobalShaderPermutationParameters& Parameters) { return GetMaxSupportedFeatureLevel(Parameters.Platform) <= ERHIFeatureLevel::SM4; }
+	static bool ShouldCompilePermutation(const FGlobalShaderPermutationParameters& Parameters) { return GetMaxSupportedFeatureLevel(Parameters.Platform) <= ERHIFeatureLevel::ES3_1; }
 	
 	FResolveDepthNonMSPS(const ShaderMetaType::CompiledShaderInitializerType& Initializer):
 	FGlobalShader(Initializer)
diff --git a/Engine/Source/Runtime/VulkanRHI/Private/VulkanGenericPlatform.cpp b/Engine/Source/Runtime/VulkanRHI/Private/VulkanGenericPlatform.cpp
--- a/Engine/Source/Runtime/VulkanRHI/Private/VulkanGenericPlatform.cpp
+++ b/Engine/Source/Runtime/VulkanRHI/Private/VulkanGenericPlatform.cpp
@@ -8,7 +8,7 @@ void FVulkanGenericPlatform::SetupFeatureLevels()
 {
 	GShaderPlatformForFeatureLevel[ERHIFeatureLevel::ES2] = SP_VULKAN_PCES3_1;
 	GShaderPlatformForFeatureLevel[ERHIFeatureLevel::ES3_1] = SP_VULKAN_PCES3_1;
-	GShaderPlatformForFeatureLevel[ERHIFeatureLevel::SM4] = SP_VULKAN_SM4;
+	GShaderPlatformForFeatureLevel[ERHIFeatureLevel::SM4_REMOVED] = SP_NumPlatforms;
 	GShaderPlatformForFeatureLevel[ERHIFeatureLevel::SM5] = SP_VULKAN_SM5;
 }
 
diff --git a/Engine/Source/Runtime/VulkanRHI/Private/VulkanRHI.cpp b/Engine/Source/Runtime/VulkanRHI/Private/VulkanRHI.cpp
--- a/Engine/Source/Runtime/VulkanRHI/Private/VulkanRHI.cpp
+++ b/Engine/Source/Runtime/VulkanRHI/Private/VulkanRHI.cpp
@@ -80,11 +80,6 @@ FDynamicRHI* FVulkanDynamicRHIModule::CreateRHI(ERHIFeatureLevel::Type InRequest
 		GMaxRHIFeatureLevel = ERHIFeatureLevel::ES3_1;
 		GMaxRHIShaderPlatform = PLATFORM_LUMIN ? SP_VULKAN_ES3_1_LUMIN : (PLATFORM_ANDROID ? SP_VULKAN_ES3_1_ANDROID : SP_VULKAN_PCES3_1);
 	}
-	else if (InRequestedFeatureLevel == ERHIFeatureLevel::SM4)
-	{
-		GMaxRHIFeatureLevel = ERHIFeatureLevel::SM4;
-		GMaxRHIShaderPlatform = SP_VULKAN_SM4;
-	}
 	else
 	{
 		GMaxRHIFeatureLevel = ERHIFeatureLevel::SM5;
diff --git a/Engine/Source/Runtime/Windows/D3D11RHI/Private/D3D11Device.cpp b/Engine/Source/Runtime/Windows/D3D11RHI/Private/D3D11Device.cpp
--- a/Engine/Source/Runtime/Windows/D3D11RHI/Private/D3D11Device.cpp
+++ b/Engine/Source/Runtime/Windows/D3D11RHI/Private/D3D11Device.cpp
@@ -89,12 +89,7 @@ FD3D11DynamicRHI::FD3D11DynamicRHI(IDXGIFactory1* InDXGIFactory1,D3D_FEATURE_LEV
 	GConfig->GetInt( TEXT( "TextureStreaming" ), TEXT( "PoolSizeVRAMPercentage" ), GPoolSizeVRAMPercentage, GEngineIni );	
 
 	// Initialize the RHI capabilities.
-	check(FeatureLevel == D3D_FEATURE_LEVEL_11_1 || FeatureLevel == D3D_FEATURE_LEVEL_11_0 || FeatureLevel == D3D_FEATURE_LEVEL_10_0 );
-
-	if(FeatureLevel == D3D_FEATURE_LEVEL_10_0)
-	{
-		GSupportsDepthFetchDuringDepthTest = false;
-	}
+	check(FeatureLevel == D3D_FEATURE_LEVEL_11_1 || FeatureLevel == D3D_FEATURE_LEVEL_11_0);
 
 	ERHIFeatureLevel::Type PreviewFeatureLevel;
 	if (RHIGetPreviewFeatureLevel(PreviewFeatureLevel))
@@ -115,11 +110,6 @@ FD3D11DynamicRHI::FD3D11DynamicRHI(IDXGIFactory1* InDXGIFactory1,D3D_FEATURE_LEV
 		GMaxRHIFeatureLevel = ERHIFeatureLevel::SM5;
 		GMaxRHIShaderPlatform = SP_PCD3D_SM5;
 	}
-	else if(FeatureLevel == D3D_FEATURE_LEVEL_10_0)
-	{
-		GMaxRHIFeatureLevel = ERHIFeatureLevel::SM4;
-		GMaxRHIShaderPlatform = SP_PCD3D_SM4;
-	}
 
 	// Initialize the platform pixel format map.
 	GPixelFormats[ PF_Unknown		].PlatformFormat	= DXGI_FORMAT_UNKNOWN;
@@ -209,12 +199,6 @@ FD3D11DynamicRHI::FD3D11DynamicRHI(IDXGIFactory1* InDXGIFactory1,D3D_FEATURE_LEV
 		GRHISupportsMSAADepthSampleAccess = true;
 		GRHISupportsRHIThread = !!EXPERIMENTAL_D3D11_RHITHREAD;
 	}
-	else if (FeatureLevel >= D3D_FEATURE_LEVEL_10_0)
-	{
-		GMaxTextureDimensions = D3D10_REQ_TEXTURE2D_U_OR_V_DIMENSION;
-		GMaxCubeTextureDimensions = D3D10_REQ_TEXTURECUBE_DIMENSION;
-		GMaxTextureArrayLayers = D3D10_REQ_TEXTURE2D_ARRAY_AXIS_DIMENSION;
-	}
 
 	GMaxTextureMipCount = FMath::CeilLogTwo( GMaxTextureDimensions ) + 1;
 	GMaxTextureMipCount = FMath::Min<int32>( MAX_TEXTURE_MIP_COUNT, GMaxTextureMipCount );
@@ -365,14 +349,6 @@ void FD3D11DynamicRHI::RHIGetSupportedResolution( uint32 &Width, uint32 &Height
 
 void FD3D11DynamicRHI::GetBestSupportedMSAASetting( DXGI_FORMAT PlatformFormat, uint32 MSAACount, uint32& OutBestMSAACount, uint32& OutMSAAQualityLevels )
 {
-	//  We disable MSAA for Feature level 10
-	if (GMaxRHIFeatureLevel == ERHIFeatureLevel::SM4)
-	{
-		OutBestMSAACount = 1;
-		OutMSAAQualityLevels = 0;
-		return;
-	}
-
 	// start counting down from current setting (indicated the current "best" count) and move down looking for support
 	for(uint32 IndexCount = MSAACount;IndexCount > 0;IndexCount--)
 	{
diff --git a/Engine/Source/Runtime/Windows/D3D11RHI/Private/D3D11RenderTarget.cpp b/Engine/Source/Runtime/Windows/D3D11RHI/Private/D3D11RenderTarget.cpp
--- a/Engine/Source/Runtime/Windows/D3D11RHI/Private/D3D11RenderTarget.cpp
+++ b/Engine/Source/Runtime/Windows/D3D11RHI/Private/D3D11RenderTarget.cpp
@@ -242,26 +242,6 @@ void FD3D11DynamicRHI::RHICopyToResolveTarget(FRHITexture* SourceTextureRHI, FRH
 					FDummyResolveParameter()
 					);
 			}
-			else if(FeatureLevel == D3D_FEATURE_LEVEL_10_0 
-				&& DestTexture2D->GetDepthStencilView(FExclusiveDepthStencil::DepthWrite_StencilWrite))
-			{
-				D3D11_TEXTURE2D_DESC ResolveTargetDesc;
-
-				DestTexture2D->GetResource()->GetDesc(&ResolveTargetDesc);
-
-				ResolveTextureUsingShader<FResolveDepthNonMSPS>(
-					RHICmdList,
-					SourceTexture2D,
-					DestTexture2D,
-					NULL,
-					DestTexture2D->GetDepthStencilView(FExclusiveDepthStencil::DepthWrite_StencilWrite),
-					ResolveTargetDesc,
-					GetDefaultRect(ResolveParams.Rect,DestTexture2D->GetSizeX(),DestTexture2D->GetSizeY()),
-					GetDefaultRect(ResolveParams.Rect,DestTexture2D->GetSizeX(),DestTexture2D->GetSizeY()),
-					Direct3DDeviceIMContext,
-					FDummyResolveParameter()
-					);
-			}
 			else
 			{
 				DXGI_FORMAT SrcFmt = (DXGI_FORMAT)GPixelFormats[SourceTextureRHI->GetFormat()].PlatformFormat;
diff --git a/Engine/Source/Runtime/Windows/D3D11RHI/Private/Windows/WindowsD3D11Device.cpp b/Engine/Source/Runtime/Windows/D3D11RHI/Private/Windows/WindowsD3D11Device.cpp
--- a/Engine/Source/Runtime/Windows/D3D11RHI/Private/Windows/WindowsD3D11Device.cpp
+++ b/Engine/Source/Runtime/Windows/D3D11RHI/Private/Windows/WindowsD3D11Device.cpp
@@ -77,24 +77,6 @@ struct AmdAgsInfo
 static AmdAgsInfo AmdInfo;
 #endif
 
-static TAutoConsoleVariable<int32> CVarForceAMDToSM4(
-	TEXT("r.ForceAMDToSM4"),
-	0,
-	TEXT("Forces AMD devices to use SM4.0/D3D10.0 feature level."),
-	ECVF_RenderThreadSafe);
-
-static TAutoConsoleVariable<int32> CVarForceIntelToSM4(
-	TEXT("r.ForceIntelToSM4"),
-	0,
-	TEXT("Forces Intel devices to use SM4.0/D3D10.0 feature level."),
-	ECVF_RenderThreadSafe);
-
-static TAutoConsoleVariable<int32> CVarForceNvidiaToSM4(
-	TEXT("r.ForceNvidiaToSM4"),
-	0,
-	TEXT("Forces Nvidia devices to use SM4.0/D3D10.0 feature level."),
-	ECVF_RenderThreadSafe);
-
 static TAutoConsoleVariable<int32> CVarAMDUseMultiThreadedDevice(
 	TEXT("r.AMDD3D11MultiThreadedDevice"),
 	0,
@@ -128,13 +110,6 @@ static FAutoConsoleVariableRef CVarDX11NumGPUs(
  */
 namespace RHIConsoleVariables
 {
-	int32 MinFeatureSetLimit = -1;
-	static FAutoConsoleVariableRef CVarMinFeatureSetLimit(
-		TEXT("RHI.MinFeatureSetLimit"),
-		MinFeatureSetLimit,
-		TEXT("If set to 11, limit D3D RHI to D3D11 and 11.1 feature levels, disallowing 10. Otherwise, it will use default. Changing this at run-time has no effect. (default is -1)")
-		);
-
 	int32 MaxFeatureSetLimit = -1;
 	static FAutoConsoleVariableRef CVarMaxFeatureSetLimit(
 		TEXT("RHI.FeatureSetLimit"),
@@ -250,14 +225,8 @@ static void SafeCreateDXGIFactory(IDXGIFactory1** DXGIFactory1)
  */
 static D3D_FEATURE_LEVEL GetMinAllowedD3DFeatureLevel()
 {
-	// Default to 10.0
-	D3D_FEATURE_LEVEL AllowedFeatureLevel = D3D_FEATURE_LEVEL_10_0;
-
-	if (RHIConsoleVariables::MinFeatureSetLimit == 11)
-	{
-		AllowedFeatureLevel = D3D_FEATURE_LEVEL_11_0;
-	}
-
+	// Default to 11.0
+	D3D_FEATURE_LEVEL AllowedFeatureLevel = D3D_FEATURE_LEVEL_11_0;
 	return AllowedFeatureLevel;
 }
 
@@ -269,24 +238,6 @@ static D3D_FEATURE_LEVEL GetMaxAllowedD3DFeatureLevel()
 {
 	// Default to 11.0
 	D3D_FEATURE_LEVEL AllowedFeatureLevel = D3D_FEATURE_LEVEL_11_0;
-
-	// Use a feature level 10 if specified on the command line.
-	if(FParse::Param(FCommandLine::Get(),TEXT("d3d10")) || 
-		FParse::Param(FCommandLine::Get(),TEXT("dx10")) ||
-		FParse::Param(FCommandLine::Get(),TEXT("sm4")) ||
-		RHIConsoleVariables::MaxFeatureSetLimit == 10)
-	{
-		AllowedFeatureLevel = D3D_FEATURE_LEVEL_10_0;
-	}
-
-	if (bIsQuadBufferStereoEnabled)
-	{
-		if (AllowedFeatureLevel == D3D_FEATURE_LEVEL_10_0)
-		{
-			UE_LOG(LogD3D11RHI, Warning, TEXT("D3D Feature Level overriden from 10.0 to 11.1 due to quad_buffer_stereo"));
-		}
-		AllowedFeatureLevel = D3D_FEATURE_LEVEL_11_1;
-	}
 	return AllowedFeatureLevel;
 }
 
@@ -313,7 +264,6 @@ static bool SafeTestD3D11CreateDevice(IDXGIAdapter* Adapter,D3D_FEATURE_LEVEL Mi
 	{
 		D3D_FEATURE_LEVEL_11_1,
 		D3D_FEATURE_LEVEL_11_0,
-		D3D_FEATURE_LEVEL_10_0
 	};
 	
 	// Trim to allowed feature levels
@@ -944,27 +894,6 @@ void FD3D11DynamicRHIModule::FindAdapter()
 	{
 		UE_LOG(LogD3D11RHI, Error, TEXT("Failed to choose a D3D11 Adapter."));
 	}
-
-	// Workaround to force specific IHVs to SM4.0
-	if (ChosenAdapter.IsValid() && ChosenAdapter.MaxSupportedFeatureLevel != D3D_FEATURE_LEVEL_10_0)
-	{
-		DXGI_ADAPTER_DESC AdapterDesc;
-		ZeroMemory(&AdapterDesc, sizeof(DXGI_ADAPTER_DESC));
-
-		DXGIFactory1->EnumAdapters(ChosenAdapter.AdapterIndex, TempAdapter.GetInitReference());
-		VERIFYD3D11RESULT(TempAdapter->GetDesc(&AdapterDesc));	
-
-		const bool bIsAMD = AdapterDesc.VendorId == 0x1002;
-		const bool bIsIntel = AdapterDesc.VendorId == 0x8086;
-		const bool bIsNVIDIA = AdapterDesc.VendorId == 0x10DE;
-
-		if ((bIsAMD && CVarForceAMDToSM4.GetValueOnGameThread() > 0) ||
-			(bIsIntel && CVarForceIntelToSM4.GetValueOnGameThread() > 0) ||
-			(bIsNVIDIA && CVarForceNvidiaToSM4.GetValueOnGameThread() > 0))
-		{
-			ChosenAdapter.MaxSupportedFeatureLevel = D3D_FEATURE_LEVEL_10_0;
-		}
-	}
 }
 
 FDynamicRHI* FD3D11DynamicRHIModule::CreateRHI(ERHIFeatureLevel::Type RequestedFeatureLevel)
@@ -1602,7 +1531,7 @@ void FD3D11DynamicRHI::InitD3DDevice()
 
 		GShaderPlatformForFeatureLevel[ERHIFeatureLevel::ES2] = SP_PCD3D_ES2;
 		GShaderPlatformForFeatureLevel[ERHIFeatureLevel::ES3_1] = SP_PCD3D_ES3_1;
-		GShaderPlatformForFeatureLevel[ERHIFeatureLevel::SM4] = SP_PCD3D_SM4;
+		GShaderPlatformForFeatureLevel[ERHIFeatureLevel::SM4_REMOVED] = SP_NumPlatforms;
 		GShaderPlatformForFeatureLevel[ERHIFeatureLevel::SM5] = SP_PCD3D_SM5;
 
 		if (IsRHIDeviceAMD() && CVarAMDDisableAsyncTextureCreation.GetValueOnAnyThread())
