Removed more dead SM4 code from post process materials.

diff --git a/Engine/Source/Runtime/Renderer/Private/PostProcess/PostProcessMaterial.cpp b/Engine/Source/Runtime/Renderer/Private/PostProcess/PostProcessMaterial.cpp
--- a/Engine/Source/Runtime/Renderer/Private/PostProcess/PostProcessMaterial.cpp
+++ b/Engine/Source/Runtime/Renderer/Private/PostProcess/PostProcessMaterial.cpp
@@ -29,7 +29,6 @@ TAutoConsoleVariable<int32> CVarPostProcessAllowStencilTest(
 	TEXT("Enables stencil testing in post process materials.\n")
 	TEXT("0: disable stencil testing\n")
 	TEXT("1: allow stencil testing\n")
-	TEXT("2: allow stencil testing and, if necessary, making a copy of custom depth/stencil buffer\n")
 	);
 
 TAutoConsoleVariable<int32> CVarPostProcessAllowBlendModes(
@@ -51,11 +50,6 @@ bool IsPostProcessStencilTestAllowed()
 	return CVarPostProcessAllowStencilTest.GetValueOnRenderThread() != 0;
 }
 
-bool IsPostProcessStencilTestWithCopyAllowed()
-{
-	return CVarPostProcessAllowStencilTest.GetValueOnRenderThread() == 2;
-}
-
 bool IsCustomDepthEnabled()
 {
 	static const IConsoleVariable* CVarCustomDepth = IConsoleManager::Get().FindConsoleVariable(TEXT("r.CustomDepth"));
@@ -71,10 +65,7 @@ enum class ECustomDepthPolicy : uint32
 	Disabled,
 
 	// Custom Depth-Stencil is enabled; potentially simultaneous SRV / DSV usage.
-	Enabled,
-
-	// Custom Depth-Stencil is enabled; but makes a copy of the target to avoid simultaneous SRV / DSV usage.
-	EnabledWithCopy
+	Enabled
 };
 
 ECustomDepthPolicy GetMaterialCustomDepthPolicy(const FMaterial* Material, ERHIFeatureLevel::Type FeatureLevel)
@@ -330,34 +321,14 @@ FRDGTextureRef ComputePostProcessMaterial(
 	FRHIDepthStencilState* DepthStencilState = DefaultDepthStencilState;
 
 	FRDGTextureRef DepthStencilTexture = nullptr;
-	FRDGTextureRef DepthStencilTextureForSRV = nullptr;
 
 	// Allocate custom depth stencil texture(s) and depth stencil state.
 	const ECustomDepthPolicy CustomStencilPolicy = GetMaterialCustomDepthPolicy(Material, FeatureLevel);
 
-	if (CustomStencilPolicy != ECustomDepthPolicy::Disabled)
+	if (CustomStencilPolicy == ECustomDepthPolicy::Enabled)
 	{
 		check(Inputs.CustomDepthTexture);
-
 		DepthStencilTexture = Inputs.CustomDepthTexture;
-
-		if (CustomStencilPolicy == ECustomDepthPolicy::EnabledWithCopy)
-		{
-			// NOTE: SM4 does not allow depth stencil to be a destination for a copy. Therefore,
-			// we create a shader resource texture and copy into that.
-			FRDGTextureDesc CopyDesc = Inputs.CustomDepthTexture->Desc;
-			CopyDesc.Flags = TexCreate_None;
-			CopyDesc.TargetableFlags = TexCreate_ShaderResource;
-
-			DepthStencilTextureForSRV = GraphBuilder.CreateTexture(CopyDesc, TEXT("CustomDepthCopy"));
-
-			AddCopyTexturePass(GraphBuilder, DepthStencilTexture, DepthStencilTextureForSRV);
-		}
-		else
-		{
-			DepthStencilTextureForSRV = DepthStencilTexture;
-		}
-
 		DepthStencilState = GetMaterialStencilState(Material);
 	}
 
@@ -433,7 +404,7 @@ FRDGTextureRef ComputePostProcessMaterial(
 	FPostProcessMaterialParameters* PostProcessMaterialParameters = GraphBuilder.AllocParameters<FPostProcessMaterialParameters>();
 
 	PostProcessMaterialParameters->PostProcessOutput = GetScreenPassTextureViewportParameters(SceneColorViewport);
-	PostProcessMaterialParameters->CustomDepth = DepthStencilTextureForSRV;
+	PostProcessMaterialParameters->CustomDepth = DepthStencilTexture;
 
 	PostProcessMaterialParameters->RenderTargets[0] = FRenderTargetBinding(OutputTexture, OutputLoadAction);
 
@@ -512,19 +483,11 @@ FRDGTextureRef ComputePostProcessMaterial(
 		RDG_EVENT_NAME("PostProcessMaterial"),
 		PostProcessMaterialParameters,
 		ERDGPassFlags::Raster,
-		[ScreenPassView, SceneColorViewport, ScreenPassDraw, PostProcessMaterialParameters, DepthStencilTexture, DepthStencilTextureForSRV, SetupFunction, bFlipYAxis, bNeedsGBuffer]
+		[ScreenPassView, SceneColorViewport, ScreenPassDraw, PostProcessMaterialParameters, DepthStencilTexture, SetupFunction, bFlipYAxis, bNeedsGBuffer]
 	(FRHICommandListImmediate& RHICmdList)
 	{
 		FSceneRenderTargets& SceneRenderTargets = FSceneRenderTargets::Get(RHICmdList);
 
-		const bool bUsesDepthStencilCopy = DepthStencilTexture != DepthStencilTextureForSRV;
-
-		if (bUsesDepthStencilCopy)
-		{
-			// Swap in the new copied SRV just prior to rendering.
-			SceneRenderTargets.CustomDepth = DepthStencilTextureForSRV->GetPooledRenderTarget();
-		}
-
 		DrawScreenPass<decltype(SetupFunction)>(
 			RHICmdList,
 			ScreenPassView,
@@ -533,12 +496,6 @@ FRDGTextureRef ComputePostProcessMaterial(
 			ScreenPassDraw,
 			SetupFunction);
 
-		if (bUsesDepthStencilCopy)
-		{
-			// Swap back to the up-to-date target.
-			SceneRenderTargets.CustomDepth = DepthStencilTexture->GetPooledRenderTarget();
-		}
-
 		if (bNeedsGBuffer)
 		{
 			SceneRenderTargets.AdjustGBufferRefCount(RHICmdList, -1);
