Converted eye adaptation parameters to shader parameter struct for clarity.

diff --git a/Engine/Shaders/Private/PostProcessEyeAdaptation.usf b/Engine/Shaders/Private/PostProcessEyeAdaptation.usf
--- a/Engine/Shaders/Private/PostProcessEyeAdaptation.usf
+++ b/Engine/Shaders/Private/PostProcessEyeAdaptation.usf
@@ -17,7 +17,6 @@ uint4 EyeAdaptionSrcRect;
 // in float Slope 
 float ComputeWeightedTextureAverageAlpha(Texture2D InBuffer, in uint4 ROI, in float Slope)
 {
-	
 	// The inverse of the Region of Interest size.
 	const float InvXMax = 1.f / float(ROI.z - ROI.x);
 	const float InvYMax = 1.f / float(ROI.w - ROI.y);
@@ -85,36 +84,31 @@ float ComputeEyeAdaptation(float OldExposure, float TargetExposure, float FrameT
 {
 	float Diff = TargetExposure - OldExposure;
 
-	const float EyeAdaptionSpeedUp   = EyeAdaptationParams[1].z;
-	const float EyeAdaptionSpeedDown = EyeAdaptationParams[1].w;
-
-    float AdaptionSpeed = (Diff > 0) ? EyeAdaptionSpeedUp : EyeAdaptionSpeedDown;
+	float AdaptionSpeed = (Diff > 0) ? EyeAdaptation_ExposureSpeedUp : EyeAdaptation_ExposureSpeedDown;
 
-    float Factor = 1.0f - exp2(-FrameTime * AdaptionSpeed);
+	float Factor = 1.0f - exp2(-FrameTime * AdaptionSpeed);
 
-	return clamp(OldExposure + Diff * Factor, EyeAdaptationParams[0].z, EyeAdaptationParams[0].w);
+	return clamp(OldExposure + Diff * Factor, EyeAdaptation_MinAverageLuminance, EyeAdaptation_MaxAverageLuminance);
 }
 
 float4 MainEyeAdaptationCommon()
 {
 	float4 OutColor = 0;
 
-	const float ExposureOffsetMultipler = EyeAdaptationParams[1].x;
 	const float AverageSceneLuminance = ComputeEyeAdaptationExposure(PostprocessInput0);
-	const float TargetExposure = clamp(AverageSceneLuminance, EyeAdaptationParams[0].z, EyeAdaptationParams[0].w);
+	const float TargetExposure = clamp(AverageSceneLuminance, EyeAdaptation_MinAverageLuminance, EyeAdaptation_MaxAverageLuminance);
 
 	float OldExposureScale = PostprocessInput0.Load(int3(0, 1, 0)).x;
-	float OldExposure = ExposureOffsetMultipler / ( OldExposureScale != 0 ? OldExposureScale : 1.0f );
-	float FrameTime = EyeAdaptationParams[1].y;
+	float OldExposure = EyeAdaptation_ExposureCompensation / ( OldExposureScale != 0 ? OldExposureScale : 1.0f );
 
 	// eye adaptation changes over time
-	float SmoothedExposure = ComputeEyeAdaptation(OldExposure, TargetExposure, FrameTime);
+	float SmoothedExposure = ComputeEyeAdaptation(OldExposure, TargetExposure, EyeAdaptation_DeltaWorldTime);
 
 	float SmoothedExposureScale = 1.0f / max(0.0001f, SmoothedExposure);
 	float TargetExposureScale =   1.0f / max(0.0001f, TargetExposure);
 
-	OutColor.x = ExposureOffsetMultipler * SmoothedExposureScale;
-	OutColor.y = ExposureOffsetMultipler * TargetExposureScale;
+	OutColor.x = EyeAdaptation_ExposureCompensation * SmoothedExposureScale;
+	OutColor.y = EyeAdaptation_ExposureCompensation * TargetExposureScale;
 	OutColor.z = AverageSceneLuminance;
 	OutColor.w = 0;
 
@@ -148,12 +142,11 @@ void MainBasicEyeAdaptationSetupPS(
 	float2 UV = UVAndScreenPos.xy;
 	OutColor = Texture2DSample(PostprocessInput0, PostprocessInput0Sampler, UV);
 
-	// Use max to insure intensity is never zero (so the following log is well behaved)
-	// NB: log2( EyeAdaptationParams[2].z)  =  -EyeAdaptationParams[2].y / EyeAdaptationParams[2].x
-	float Intensity = max(dot(OutColor.xyz, float3(0.2126, 0.7152, 0.0722)), EyeAdaptationParams[2].z);
+	// Use max to ensure intensity is never zero (so the following log is well behaved)
+	float Intensity = max(dot(OutColor.xyz, float3(0.2126, 0.7152, 0.0722)), EyeAdaptation_LuminanceMin);
 
 	// Store log intensity in the alpha channel: scale to 0,1 range.
-	OutColor.w = EyeAdaptationParams[2].x * log2(Intensity) + EyeAdaptationParams[2].y; 
+	OutColor.w = EyeAdaptation_HistogramScale * log2(Intensity) + EyeAdaptation_HistogramBias; 
 }
 
 
@@ -185,18 +178,13 @@ groupshared float2 SubRectValueWeight[TGSIZE*TGSIZE];
 RWTexture2D<float4> RWOutputComputeTex;
 
 [numthreads(TGSIZE, TGSIZE, 1)]
-void MainLogLuminance2ExposureScaleCS(uint GIndex : SV_GroupIndex, 
-	uint2 GTId : SV_GroupThreadID
-	)
+void MainLogLuminance2ExposureScaleCS(uint GIndex : SV_GroupIndex, uint2 GTId : SV_GroupThreadID)
 {
 	float4 OutColor=0;
 
 #else
 
-void MainLogLuminance2ExposureScalePS( 
-	in float4 UVAndScreenPos : TEXCOORD0, 
-	out float4 OutColor : SV_Target0
-	)
+void MainLogLuminance2ExposureScalePS(in float4 UVAndScreenPos : TEXCOORD0, out float4 OutColor : SV_Target0)
 {
 
 #endif
@@ -212,7 +200,7 @@ void MainLogLuminance2ExposureScalePS(
 	// WeightSlope: 0 is uniform weighting.
 	//              1 decays to the edges.
 	//              >1 focuses more in the center of the screen.
-	float WeightSlope = EyeAdaptationParams[2].w;
+	float WeightSlope = EyeAdaptation_WeightSlope;
 	
 	// Compute scaled Log Luminance Average
 #if COMPUTESHADER
@@ -259,33 +247,28 @@ void MainLogLuminance2ExposureScalePS(
 #endif
 
 	// Correct for [0,1] scaling
-	LogLumAve = (LogLumAve - EyeAdaptationParams[2].y) / EyeAdaptationParams[2].x;
+	LogLumAve = (LogLumAve - EyeAdaptation_HistogramBias) / EyeAdaptation_HistogramScale;
 	
 	// Convert LogLuminanceAverage to Average Intensity
 	const float AverageSceneLuminance = OneOverPreExposure * exp2(LogLumAve);
-	const float LumAve = AverageSceneLuminance * AverageLuminanceScale;
+	const float LumAve = AverageSceneLuminance * EyeAdaptation_CalibrationConstantInverse;
 
-	const float ClampedLumAve = clamp(LumAve, EyeAdaptationParams[0].z, EyeAdaptationParams[0].w );
+	const float ClampedLumAve = clamp(LumAve, EyeAdaptation_MinAverageLuminance, EyeAdaptation_MaxAverageLuminance);
 	
-	// NB: this param value differs by a factor from the value used in the histogram-based alogrithm 
-	const float ExposureOffsetMultipler = EyeAdaptationParams[1].x;
-
 	// The Exposure Scale (and thus intensity) used in the previous frame
 	float ExposureScaleOld = EyeAdaptationTexture.Load(int3(0, 0, 0)).x;
-	float LuminanceAveOld = ExposureOffsetMultipler / (ExposureScaleOld != 0 ? ExposureScaleOld : 1.0f);
+	float LuminanceAveOld = EyeAdaptation_ExposureCompensation / (ExposureScaleOld != 0 ? ExposureScaleOld : 1.0f);
 
 	// Time-based expoential blend of the intensity to allow the eye adaptation to ramp up over a few frames.
-	// NB: This also clamps the SmoothedIn
-	float FrameTime = EyeAdaptationParams[1].y;
-	float SmoothedLuminance = ComputeEyeAdaptation(LuminanceAveOld, ClampedLumAve, FrameTime);
+	float SmoothedLuminance = ComputeEyeAdaptation(LuminanceAveOld, ClampedLumAve, EyeAdaptation_DeltaWorldTime);
 
 	float SmoothedExposureScale = 1.0f / max(0.0001f, SmoothedLuminance);
 	float TargetExposureScale   = 1.0f / max(0.0001f, ClampedLumAve);
 	
 	// Output the number that will rescale the image intensity
-	OutColor.x = ExposureOffsetMultipler * SmoothedExposureScale;
+	OutColor.x = EyeAdaptation_ExposureCompensation * SmoothedExposureScale;
 	// Output the target value
-	OutColor.y = ExposureOffsetMultipler * TargetExposureScale;
+	OutColor.y = EyeAdaptation_ExposureCompensation * TargetExposureScale;
 	OutColor.z = AverageSceneLuminance;
 	OutColor.w = 0;
 		
diff --git a/Engine/Shaders/Private/PostProcessHistogramCommon.ush b/Engine/Shaders/Private/PostProcessHistogramCommon.ush
--- a/Engine/Shaders/Private/PostProcessHistogramCommon.ush
+++ b/Engine/Shaders/Private/PostProcessHistogramCommon.ush
@@ -4,13 +4,19 @@
 	PostProcessHistogramCommon.usf: PostProcessing histogram shared functions and structures.
 =============================================================================*/
 
-// [0] .x:ExposureLowPercent/100, .y:EyeAdaptationHighPercent/100, .z:EyeAdaptationMin, .w:EyeAdaptationMax
-// [1] .x:exp2(ExposureOffset), .y:DeltaWorldTime, .zw: EyeAdaptionSpeedUp/Down
-// [2] .x:Histogram multiply, .y:Histogram add, .z:HistogramMinIntensity w:BasicExposureFocus
-// [3] .y:Calibration Constant
-float4 EyeAdaptationParams[EYE_ADAPTATION_PARAMS_SIZE];
-
-#define AverageLuminanceScale (EyeAdaptationParams[3].y)
+float EyeAdaptation_ExposureLowPercent;
+float EyeAdaptation_ExposureHighPercent;
+float EyeAdaptation_MinAverageLuminance;
+float EyeAdaptation_MaxAverageLuminance;
+float EyeAdaptation_ExposureCompensation;
+float EyeAdaptation_DeltaWorldTime;
+float EyeAdaptation_ExposureSpeedUp;
+float EyeAdaptation_ExposureSpeedDown;
+float EyeAdaptation_HistogramScale;
+float EyeAdaptation_HistogramBias;
+float EyeAdaptation_LuminanceMin;
+float EyeAdaptation_CalibrationConstantInverse;
+float EyeAdaptation_WeightSlope;
 
 // inverse of ComputeLuminanceFromHistogramPosition
 // is executed more often than ComputeLuminanceFromHistogramPosition()
@@ -18,16 +24,16 @@ float4 EyeAdaptationParams[EYE_ADAPTATION_PARAMS_SIZE];
 // @return HistogramPosition 0..1
 float ComputeHistogramPositionFromLuminance(float Luminance)
 {
-	return log2(Luminance) * EyeAdaptationParams[2].x + EyeAdaptationParams[2].y;
+	return log2(Luminance) * EyeAdaptation_HistogramScale + EyeAdaptation_HistogramBias;
 }
 
 // inverse of ComputeHistogramPositionFromLuminance()
 // is not as often executed as ComputeHistogramPositionFromLuminance()
 // @param HistogramPosition 0..1
-// @return Lumiance
+// @return Luminance
 float ComputeLuminanceFromHistogramPosition(float HistogramPosition)
 {
-	return exp2( (HistogramPosition - EyeAdaptationParams[2].y) / EyeAdaptationParams[2].x );
+	return exp2((HistogramPosition - EyeAdaptation_HistogramBias) / EyeAdaptation_HistogramScale);
 }
 
 #define HISTOGRAM_SIZE 64
@@ -123,19 +129,17 @@ float ComputeAverageLuminaneWithoutOutlier(Texture2D HistogramTexture, float Min
 	return SumWithoutOutliers.x / max(0.0001f, SumWithoutOutliers.y);
 }
 
-// used by bloom and tonemapper VS
-// requires EyeAdaptationParams to be set
 float ComputeEyeAdaptationExposure(Texture2D HistogramTexture)
 {
 	float HistogramSum = ComputeHistogramSum(HistogramTexture);
-	return ComputeAverageLuminaneWithoutOutlier(HistogramTexture, HistogramSum * EyeAdaptationParams[0].x, HistogramSum * EyeAdaptationParams[0].y);
+	return ComputeAverageLuminaneWithoutOutlier(HistogramTexture, HistogramSum * EyeAdaptation_ExposureLowPercent, HistogramSum * EyeAdaptation_ExposureHighPercent);
 }
 
 // Computation of triangle function on normalized space
 // used in Basic EyeAdaptation in weighting the center of the screen
 // @return range of 0..1 with 1 nx = 0.5
 // @param nx in 0..1 
-// @param slope assumed possitive
+// @param slope assumed positive
 float AdaptationWeight(float nx, float slope)
 {
 	return max(1.f - 2. * slope * abs(nx - 0.5f), 0.0f);
diff --git a/Engine/Shaders/Private/PostProcessVisualizeHDR.usf b/Engine/Shaders/Private/PostProcessVisualizeHDR.usf
--- a/Engine/Shaders/Private/PostProcessVisualizeHDR.usf
+++ b/Engine/Shaders/Private/PostProcessVisualizeHDR.usf
@@ -38,17 +38,11 @@ uint ComputeAdvice(float3 HDRColor)
 {
 	float Lum = max(HDRColor.r, max(HDRColor.g, HDRColor.b));
 
-	if(Lum < EyeAdaptationParams[0].z)
+	if(Lum < EyeAdaptation_MinAverageLuminance)
 	{
 		return 1;
 	}
 
-	// for some reasons HLSL compiler seems to suppress  the return 1 when we comment in this code
-//	if(Lum > EyeAdaptationParams[0].w)
-//	{
-//		return 2;
-//	}
-
 	return 0;
 }
 
@@ -150,7 +144,7 @@ void MainPS(noperspective float4 UVAndScreenPos : TEXCOORD0, float4 SvPosition :
 	}
 
 	// Compute and apply weight value at this location 
-	float WeightSlope = EyeAdaptationParams[2].w;
+	float WeightSlope = EyeAdaptation_WeightSlope;
 	float2 NdPos = ViewportSize.zw * (PixelPos.xy - ViewportRect.xy);
 	float weight = AdaptationWeight(NdPos.x, WeightSlope) * AdaptationWeight(NdPos.y, WeightSlope);
 	OutColor *= max(weight, 0.05f);  
@@ -238,12 +232,12 @@ void MainPS(noperspective float4 UVAndScreenPos : TEXCOORD0, float4 SvPosition :
 
 	float HistogramSum = ComputeHistogramSum(HistogramTexture);
 
-	if(InsetUV.x < ComputeHistogramPositionFromLuminance(EyeAdaptationParams[0].z))
+	if(InsetUV.x < ComputeHistogramPositionFromLuminance(EyeAdaptation_MinAverageLuminance))
 	{
 		// < min: grey
 		OutColor.xyz = lerp(OutColor.xyz, float3(0.5f, 0.5f, 0.5f), 0.5f);
 	}
-	else if(InsetUV.x < ComputeHistogramPositionFromLuminance(EyeAdaptationParams[0].w))
+	else if(InsetUV.x < ComputeHistogramPositionFromLuminance(EyeAdaptation_MaxAverageLuminance))
 	{
 		// >= min && < max: green
 		OutColor.xyz = lerp(OutColor.xyz, float3(0.5f, 0.8f, 0.5f), 0.5f);
diff --git a/Engine/Source/Runtime/Renderer/Private/PostProcess/PostProcessEyeAdaptation.cpp b/Engine/Source/Runtime/Renderer/Private/PostProcess/PostProcessEyeAdaptation.cpp
--- a/Engine/Source/Runtime/Renderer/Private/PostProcess/PostProcessEyeAdaptation.cpp
+++ b/Engine/Source/Runtime/Renderer/Private/PostProcess/PostProcessEyeAdaptation.cpp
@@ -16,7 +16,8 @@
 
 RENDERCORE_API bool UsePreExposure(EShaderPlatform Platform);
 
-// Initialize the static CVar
+namespace
+{
 TAutoConsoleVariable<float> CVarEyeAdaptationPreExposureOverride(
 	TEXT("r.EyeAdaptation.PreExposureOverride"),
 	0,
@@ -25,11 +26,10 @@ TAutoConsoleVariable<float> CVarEyeAdaptationPreExposureOverride(
 	TEXT("> 0 : Override PreExposure\n"),
 	ECVF_RenderThreadSafe);
 
-// Initialize the static CVar
 TAutoConsoleVariable<int32> CVarEyeAdaptationMethodOverride(
-	TEXT("r.EyeAdaptation.MethodOveride"),
+	TEXT("r.EyeAdaptation.MethodOverride"),
 	-1,
-	TEXT("Overide the camera metering method set in post processing volumes\n")
+	TEXT("Override the camera metering method set in post processing volumes\n")
 	TEXT("-2: override with custom settings (for testing Basic Mode)\n")
 	TEXT("-1: no override\n")
 	TEXT(" 1: Auto Histogram-based\n")
@@ -37,7 +37,6 @@ TAutoConsoleVariable<int32> CVarEyeAdaptationMethodOverride(
 	TEXT(" 3: Manual"),
 	ECVF_Scalability | ECVF_RenderThreadSafe);
 
-// Initialize the static CVar used in computing the weighting focus in basic eye-adaptation
 TAutoConsoleVariable<float> CVarEyeAdaptationFocus(
 	TEXT("r.EyeAdaptation.Focus"),
 	1.0f,
@@ -53,65 +52,146 @@ TAutoConsoleVariable<int32> CVarEyeAdaptationBasicCompute(
 	TEXT("= 0 : Pixel Shader\n")
 	TEXT("> 0 : Compute Shader (default) \n"),
 	ECVF_Scalability | ECVF_RenderThreadSafe);
+}
+
+bool IsAutoExposureMethodSupported(ERHIFeatureLevel::Type FeatureLevel, EAutoExposureMethod AutoExposureMethodId)
+{
+	switch (AutoExposureMethodId)
+	{
+	case EAutoExposureMethod::AEM_Histogram:
+		return FeatureLevel >= ERHIFeatureLevel::SM5;
+	case EAutoExposureMethod::AEM_Basic:
+	case EAutoExposureMethod::AEM_Manual:
+		return FeatureLevel >= ERHIFeatureLevel::ES3_1;
+	}
+	return false;
+}
+
+// Query the view for the auto exposure method, and allow for CVar override.
+EAutoExposureMethod GetAutoExposureMethod(const FViewInfo& View)
+{
+	EAutoExposureMethod AutoExposureMethodId = View.FinalPostProcessSettings.AutoExposureMethod;
+
+	const int32 EyeOverride = CVarEyeAdaptationMethodOverride.GetValueOnRenderThread();
+
+	if (EyeOverride >= 0)
+	{
+		// Additional branching for override.
+		switch (EyeOverride)
+		{
+		case 1:
+		{
+			// Only override if the platform supports it.
+			if (View.GetFeatureLevel() >= ERHIFeatureLevel::SM5)
+			{
+				AutoExposureMethodId = EAutoExposureMethod::AEM_Histogram;
+			}
+			break;
+		}
+		case 2:
+		{
+			AutoExposureMethodId = EAutoExposureMethod::AEM_Basic;
+			break;
+		}
+		case 3:
+		{
+			AutoExposureMethodId = EAutoExposureMethod::AEM_Manual;
+			break;
+		}
+		}
+	}
+
+	// If auto exposure is disabled, revert to manual mode which will clamp to a reasonable default.
+	if (!View.Family->EngineShowFlags.EyeAdaptation)
+	{
+		AutoExposureMethodId = AEM_Manual;
+	}
+
+	// We should have a valid exposure method at this stage.
+	check(IsAutoExposureMethodSupported(View.GetFeatureLevel(), AutoExposureMethodId));
 
-/**
- *   Shared functionality used in computing the eye-adaptation parameters
- *   Compute the parameters used for eye-adaptation.  These will default to values
- *   that disable eye-adaptation if the hardware doesn't support the minimum feature level
- */
-void ComputeEyeAdaptationValues(const ERHIFeatureLevel::Type MinFeatureLevel, const FViewInfo& View, FVector4 Out[EYE_ADAPTATION_PARAMS_SIZE])
+	return AutoExposureMethodId;
+}
+
+bool IsExtendLuminanceRangeEnabled()
 {
 	static const auto VarDefaultAutoExposureExtendDefaultLuminanceRange = IConsoleManager::Get().FindTConsoleVariableDataInt(TEXT("r.DefaultFeature.AutoExposure.ExtendDefaultLuminanceRange"));
-	const bool bExtendedLuminanceRange = VarDefaultAutoExposureExtendDefaultLuminanceRange->GetValueOnRenderThread() == 1;
 
-	const FPostProcessSettings& Settings = View.FinalPostProcessSettings;
-	const FEngineShowFlags& EngineShowFlags = View.Family->EngineShowFlags;
+	return VarDefaultAutoExposureExtendDefaultLuminanceRange->GetValueOnRenderThread() == 1;
+}
+
+float GetLuminanceRangeValue(float EV100Value)
+{
+	return IsExtendLuminanceRangeEnabled() ? EV100ToLuminance(EV100Value) : EV100Value;
+}
 
-	// ----------
+float GetBasicAutoExposureFocus()
+{
+	const float FocusMax = 10.f;
+	const float FocusValue = CVarEyeAdaptationFocus.GetValueOnRenderThread();
+	return FMath::Max(FMath::Min(FocusValue, FocusMax), 0.0f);
+}
 
-	const EAutoExposureMethod AutoExposureMethod = GetAutoExposureMethod(View);
+float GetAutoExposureCompensation(const FViewInfo& View)
+{
+	const FPostProcessSettings& Settings = View.FinalPostProcessSettings;
 
-	// Histogram related values.
-	float LowPercent = FMath::Clamp(Settings.AutoExposureLowPercent, 1.0f, 99.0f) * 0.01f;
-	float HighPercent = FMath::Clamp(Settings.AutoExposureHighPercent, 1.0f, 99.0f) * 0.01f;
-	float HistogramLogMin = bExtendedLuminanceRange ? EV100ToLog2(Settings.HistogramLogMin) : Settings.HistogramLogMin;
-	float HistogramLogMax = bExtendedLuminanceRange ? EV100ToLog2(Settings.HistogramLogMax) : Settings.HistogramLogMax;
-	HistogramLogMin = FMath::Min<float>(HistogramLogMin, HistogramLogMax - 1);
+	const FEngineShowFlags& EngineShowFlags = View.Family->EngineShowFlags;
 
-	// ----------
+	// Skip the exposure bias when any of these flags are set.
+	const bool bSkipExposureBias =
+		View.Family->ExposureSettings.bFixed ||
+		View.Family->UseDebugViewPS() ||
+		!EngineShowFlags.Lighting ||
+		(EngineShowFlags.VisualizeBuffer && View.CurrentBufferVisualizationMode != NAME_None) ||
+		EngineShowFlags.RayTracingDebug ||
+		EngineShowFlags.VisualizeDistanceFieldAO ||
+		EngineShowFlags.VisualizeGlobalDistanceField;
+
+	if (bSkipExposureBias)
+	{
+		return 1.0f;
+	}
 
-	// Those clamps the average luminance computed from the scene color.
-	float MinAverageLuminance = 1;
-	float MaxAverageLuminance = 1;
 	// This scales the average luminance AFTER it gets clamped, affecting the exposure value directly.
 	float AutoExposureBias = Settings.AutoExposureBias;
+
 	if (Settings.AutoExposureBiasCurve)
 	{
-		float AverageSceneLuminance = View.GetLastAverageSceneLuminance();
+		const float AverageSceneLuminance = View.GetLastAverageSceneLuminance();
+
 		if (AverageSceneLuminance > 0)
 		{
 			AutoExposureBias += Settings.AutoExposureBiasCurve->GetFloatValue(LuminanceToEV100(AverageSceneLuminance));
 		}
 	}
 
-	float LocalExposureMultipler = FMath::Pow(2.0f, AutoExposureBias);
-	// This scales the average luminance BEFORE it gets clamped, used to implement the calibration constant for AEM_Basic.
-	float AverageLuminanceScale = 1.f;
-
-	// When any of those flags are set, make sure the tonemapper uses an exposure of 1.
-	if (!EngineShowFlags.Lighting 
-		|| (EngineShowFlags.VisualizeBuffer && View.CurrentBufferVisualizationMode != NAME_None) 
-		|| View.Family->UseDebugViewPS() 
-		|| EngineShowFlags.RayTracingDebug 
-		|| EngineShowFlags.VisualizeDistanceFieldAO
-		|| EngineShowFlags.VisualizeGlobalDistanceField)
-	{
-		LocalExposureMultipler = 1.f;
-	}
-	// Otherwise handle the viewport override settings.
-	else if (View.Family->ExposureSettings.bFixed)
+	return FMath::Pow(2.0f, AutoExposureBias);
+}
+
+FEyeAdaptationParameters GetEyeAdaptationParameters(const FViewInfo& View, ERHIFeatureLevel::Type MinFeatureLevel)
+{
+	const FPostProcessSettings& Settings = View.FinalPostProcessSettings;
+
+	const FEngineShowFlags& EngineShowFlags = View.Family->EngineShowFlags;
+
+	const EAutoExposureMethod AutoExposureMethod = GetAutoExposureMethod(View);
+
+	const float PercentToScale = 0.01f;
+
+	const float ExposureHighPercent = FMath::Clamp(Settings.AutoExposureHighPercent, 1.0f, 99.0f) * PercentToScale;
+	const float ExposureLowPercent = FMath::Min(FMath::Clamp(Settings.AutoExposureLowPercent, 1.0f, 99.0f) * PercentToScale, ExposureHighPercent);
+
+	const float HistogramLogMax = GetLuminanceRangeValue(Settings.HistogramLogMax);
+	const float HistogramLogMin = FMath::Min(GetLuminanceRangeValue(Settings.HistogramLogMin), HistogramLogMax - 1);
+
+	// These clamp the average luminance computed from the scene color.
+	float MinAverageLuminance = 1.0f;
+	float MaxAverageLuminance = 1.0f;
+
+	// Fixed exposure override in effect.
+	if (View.Family->ExposureSettings.bFixed)
 	{
-		LocalExposureMultipler = 1.f;
 		MinAverageLuminance = MaxAverageLuminance = EV100ToLuminance(View.Family->ExposureSettings.FixedEV100);
 	}
 	// When !EngineShowFlags.EyeAdaptation (from "r.EyeAdaptationQuality 0") or the feature level doesn't support eye adaptation, only Settings.AutoExposureBias controls exposure.
@@ -122,51 +202,44 @@ void ComputeEyeAdaptationValues(const ERHIFeatureLevel::Type MinFeatureLevel, co
 			const float FixedEV100 = FMath::Log2(FMath::Square(Settings.DepthOfFieldFstop) * Settings.CameraShutterSpeed * 100 / FMath::Max(1.f, Settings.CameraISO));
 			MinAverageLuminance = MaxAverageLuminance = EV100ToLuminance(FixedEV100);
 		}
-		else if (bExtendedLuminanceRange)
-		{
-			MinAverageLuminance = EV100ToLuminance(Settings.AutoExposureMinBrightness);
-			MaxAverageLuminance = EV100ToLuminance(Settings.AutoExposureMaxBrightness);
-		}
 		else
 		{
-			MinAverageLuminance = Settings.AutoExposureMinBrightness;
-			MaxAverageLuminance = Settings.AutoExposureMaxBrightness;
-		}
-
-		// Note that AEM_Histogram implements the calibration constant through LowPercent and HiPercent.
-		if (AutoExposureMethod == EAutoExposureMethod::AEM_Basic)
-		{
-			const float CalibrationConstant = FMath::Clamp<float>(Settings.AutoExposureCalibrationConstant, 1.f, 100.f) * 0.01f;
-			AverageLuminanceScale = 1.f / CalibrationConstant;
+			MinAverageLuminance = GetLuminanceRangeValue(Settings.AutoExposureMinBrightness);
+			MaxAverageLuminance = GetLuminanceRangeValue(Settings.AutoExposureMaxBrightness);
 		}
 	}
 
-	// ----------
-
-	LowPercent = FMath::Min<float>(LowPercent, HighPercent);
-	MinAverageLuminance = FMath::Min<float>(MinAverageLuminance, MaxAverageLuminance);
-
-	Out[0] = FVector4(LowPercent, HighPercent, MinAverageLuminance, MaxAverageLuminance);
-
-	// ----------
-
-	Out[1] = FVector4(LocalExposureMultipler, View.Family->DeltaWorldTime, Settings.AutoExposureSpeedUp, Settings.AutoExposureSpeedDown);
-
-	// ----------
-
-	float DeltaLog = HistogramLogMax - HistogramLogMin;
-	float Multiply = 1.0f / DeltaLog;
-	float Add = -HistogramLogMin * Multiply;
-	float MinIntensity = FMath::Exp2(HistogramLogMin);
-	Out[2] = FVector4(Multiply, Add, MinIntensity, 0);
-
-	// ----------
-
-	Out[3] = FVector4(0, AverageLuminanceScale, 0, 0);
+	// This scales the average luminance BEFORE it gets clamped. Note that AEM_Histogram implements the calibration constant through ExposureLowPercent and ExposureHighPercent.
+	const float CalibrationConstant = FMath::Clamp(Settings.AutoExposureCalibrationConstant, 1.0f, 100.0f) * PercentToScale;
+
+	const float ExposureCompensation = GetAutoExposureCompensation(View);
+
+	const float WeightSlope = (AutoExposureMethod == EAutoExposureMethod::AEM_Basic) ? GetBasicAutoExposureFocus() : 0.0f;
+
+	const float HistogramLogDelta = HistogramLogMax - HistogramLogMin;
+	const float HistogramScale = 1.0f / HistogramLogDelta;
+	const float HistogramBias = -HistogramLogMin * HistogramScale;
+	const float LuminanceMin = FMath::Exp2(HistogramLogMin);
+
+	FEyeAdaptationParameters Parameters;
+	Parameters.ExposureLowPercent = ExposureLowPercent;
+	Parameters.ExposureHighPercent = ExposureHighPercent;
+	Parameters.MinAverageLuminance = MinAverageLuminance;
+	Parameters.MaxAverageLuminance = MaxAverageLuminance;
+	Parameters.ExposureCompensation = ExposureCompensation;
+	Parameters.DeltaWorldTime = View.Family->DeltaWorldTime;
+	Parameters.ExposureSpeedUp = Settings.AutoExposureSpeedUp;
+	Parameters.ExposureSpeedDown = Settings.AutoExposureSpeedDown;
+	Parameters.HistogramScale = HistogramScale;
+	Parameters.HistogramBias = HistogramBias;
+	Parameters.LuminanceMin = LuminanceMin;
+	Parameters.CalibrationConstantInverse = 1.0f / CalibrationConstant;
+	Parameters.WeightSlope = WeightSlope;
+	return Parameters;
 }
 
 // Basic AutoExposure requires at least ES3_1
-static ERHIFeatureLevel::Type BasicEyeAdaptationMinFeatureLevel = ERHIFeatureLevel::ES3_1;
+const ERHIFeatureLevel::Type BasicEyeAdaptationMinFeatureLevel = ERHIFeatureLevel::ES3_1;
 
 /** Encapsulates the histogram-based post processing eye adaptation pixel shader. */
 class FPostProcessEyeAdaptationPS : public FGlobalShader
@@ -182,13 +255,16 @@ class FPostProcessEyeAdaptationPS : public FGlobalShader
 	{
 		FGlobalShader::ModifyCompilationEnvironment( Parameters, OutEnvironment );
 		OutEnvironment.SetRenderTargetOutputFormat(0, PF_A32B32G32R32F);
-		OutEnvironment.SetDefine(TEXT("EYE_ADAPTATION_PARAMS_SIZE"), (uint32)EYE_ADAPTATION_PARAMS_SIZE);
 	}
 
 	/** Default constructor. */
 	FPostProcessEyeAdaptationPS() {}
 
 public:
+	BEGIN_SHADER_PARAMETER_STRUCT(FParameters, )
+		SHADER_PARAMETER_STRUCT(FEyeAdaptationParameters, EyeAdaptation)
+	END_SHADER_PARAMETER_STRUCT()
+
 	FPostProcessPassParameters PostprocessParameter;
 	FShaderParameter EyeAdaptationParams;
 	FShaderResourceParameter EyeAdaptationTexture;
@@ -197,8 +273,9 @@ public:
 	FPostProcessEyeAdaptationPS(const ShaderMetaType::CompiledShaderInitializerType& Initializer)
 		: FGlobalShader(Initializer)
 	{
+		BindForLegacyShaderParameters<FParameters>(this, Initializer.ParameterMap);
+
 		PostprocessParameter.Bind(Initializer.ParameterMap);
-		EyeAdaptationParams.Bind(Initializer.ParameterMap, TEXT("EyeAdaptationParams"));
 		EyeAdaptationTexture.Bind(Initializer.ParameterMap, TEXT("EyeAdaptationTexture"));
 	}
 
@@ -221,9 +298,9 @@ public:
 		}
 
 		{
-			FVector4 Temp[EYE_ADAPTATION_PARAMS_SIZE];
-			FRCPassPostProcessEyeAdaptation::ComputeEyeAdaptationParamsValue(Context.View, Temp);
-			SetShaderValueArray(RHICmdList, ShaderRHI, EyeAdaptationParams, Temp, EYE_ADAPTATION_PARAMS_SIZE);
+			FParameters PassParameters;
+			PassParameters.EyeAdaptation = GetEyeAdaptationParameters(Context.View);
+			SetShaderParameters(Context.RHICmdList, this, ShaderRHI, PassParameters);
 		}
 	}
 
@@ -252,13 +329,16 @@ class FPostProcessEyeAdaptationCS : public FGlobalShader
 	{
 		FGlobalShader::ModifyCompilationEnvironment(Parameters, OutEnvironment);
 		OutEnvironment.SetRenderTargetOutputFormat(0, PF_A32B32G32R32F);
-		OutEnvironment.SetDefine(TEXT("EYE_ADAPTATION_PARAMS_SIZE"), (uint32)EYE_ADAPTATION_PARAMS_SIZE);
 	}
 
 	/** Default constructor. */
 	FPostProcessEyeAdaptationCS() {}
 
 public:
+	BEGIN_SHADER_PARAMETER_STRUCT(FParameters, )
+		SHADER_PARAMETER_STRUCT(FEyeAdaptationParameters, EyeAdaptation)
+	END_SHADER_PARAMETER_STRUCT()
+
 	FPostProcessPassParameters PostprocessParameter;
 	FRWShaderParameter OutComputeTex;
 	FShaderParameter EyeAdaptationParams;
@@ -267,9 +347,9 @@ public:
 	FPostProcessEyeAdaptationCS(const ShaderMetaType::CompiledShaderInitializerType& Initializer)
 		: FGlobalShader(Initializer)
 	{
+		BindForLegacyShaderParameters<FParameters>(this, Initializer.ParameterMap);
 		PostprocessParameter.Bind(Initializer.ParameterMap);
 		OutComputeTex.Bind(Initializer.ParameterMap, TEXT("OutComputeTex"));
-		EyeAdaptationParams.Bind(Initializer.ParameterMap, TEXT("EyeAdaptationParams"));
 	}
 
 	template <typename TRHICmdList>
@@ -281,11 +361,10 @@ public:
 		FGlobalShader::SetParameters<FViewUniformShaderParameters>(RHICmdList, ShaderRHI, Context.View.ViewUniformBuffer);
 		PostprocessParameter.SetCS(ShaderRHI, Context, RHICmdList, TStaticSamplerState<SF_Bilinear,AM_Clamp,AM_Clamp,AM_Clamp>::GetRHI());
 		OutComputeTex.SetTexture(RHICmdList, ShaderRHI, nullptr, DestUAV);		
-		
-		// PS params
-		FVector4 EyeAdaptationParamValues[EYE_ADAPTATION_PARAMS_SIZE];
-		FRCPassPostProcessEyeAdaptation::ComputeEyeAdaptationParamsValue(Context.View, EyeAdaptationParamValues);
-		SetShaderValueArray(RHICmdList, ShaderRHI, EyeAdaptationParams, EyeAdaptationParamValues, EYE_ADAPTATION_PARAMS_SIZE);
+
+		FParameters PassParameters;
+		PassParameters.EyeAdaptation = GetEyeAdaptationParameters(Context.View);
+		SetShaderParameters(RHICmdList, this, ShaderRHI, PassParameters);
 	}
 
 	template <typename TRHICmdList>
@@ -429,22 +508,15 @@ void FRCPassPostProcessEyeAdaptation::DispatchCS(TRHICmdList& RHICmdList, FRende
 	ComputeShader->UnsetParameters(RHICmdList);
 }
 
-void FRCPassPostProcessEyeAdaptation::ComputeEyeAdaptationParamsValue(const FViewInfo& View, FVector4 Out[EYE_ADAPTATION_PARAMS_SIZE])
-{
-	ComputeEyeAdaptationValues(ERHIFeatureLevel::SM5, View, Out);
-}
-
 float FRCPassPostProcessEyeAdaptation::GetFixedExposure(const FViewInfo& View)
 {
-	FVector4 EyeAdaptationParams[EYE_ADAPTATION_PARAMS_SIZE];
-	FRCPassPostProcessEyeAdaptation::ComputeEyeAdaptationParamsValue(View, EyeAdaptationParams);
+	const FEyeAdaptationParameters Parameters = GetEyeAdaptationParameters(View);
 	
-	// like in PostProcessEyeAdaptation.usf (EyeAdaptationParams[0].ZW : Min/Max Intensity)
-	const float Exposure = (EyeAdaptationParams[0].Z + EyeAdaptationParams[0].W) * 0.5f;
-	const float ExposureOffsetMultipler = EyeAdaptationParams[1].X;
+	const float Exposure = (Parameters.MinAverageLuminance + Parameters.MaxAverageLuminance) * 0.5f;
 
 	const float ExposureScale = 1.0f / FMath::Max(0.0001f, Exposure);
-	return ExposureScale * ExposureOffsetMultipler;
+
+	return ExposureScale * Parameters.ExposureCompensation;
 }
 
 void FSceneViewState::UpdatePreExposure(FViewInfo& View)
@@ -530,6 +602,10 @@ class FPostProcessBasicEyeAdaptationSetupPS : public FGlobalShader
 {
 	DECLARE_SHADER_TYPE(FPostProcessBasicEyeAdaptationSetupPS, Global);
 
+	BEGIN_SHADER_PARAMETER_STRUCT(FParameters, )
+		SHADER_PARAMETER_STRUCT(FEyeAdaptationParameters, EyeAdaptation)
+	END_SHADER_PARAMETER_STRUCT()
+
 	static bool ShouldCompilePermutation(const FGlobalShaderPermutationParameters& Parameters)
 	{
 		return IsFeatureLevelSupported(Parameters.Platform, BasicEyeAdaptationMinFeatureLevel);
@@ -538,7 +614,6 @@ class FPostProcessBasicEyeAdaptationSetupPS : public FGlobalShader
 	static void ModifyCompilationEnvironment(const FGlobalShaderPermutationParameters& Parameters, FShaderCompilerEnvironment& OutEnvironment)
 	{
 		FGlobalShader::ModifyCompilationEnvironment(Parameters, OutEnvironment);
-		OutEnvironment.SetDefine(TEXT("EYE_ADAPTATION_PARAMS_SIZE"), (uint32)EYE_ADAPTATION_PARAMS_SIZE);
 	}
 
 	/** Default constructor. */
@@ -552,8 +627,8 @@ public:
 	FPostProcessBasicEyeAdaptationSetupPS(const ShaderMetaType::CompiledShaderInitializerType& Initializer)
 		: FGlobalShader(Initializer)
 	{
+		BindForLegacyShaderParameters<FParameters>(this, Initializer.ParameterMap);
 		PostprocessParameter.Bind(Initializer.ParameterMap);
-		EyeAdaptationParams.Bind(Initializer.ParameterMap, TEXT("EyeAdaptationParams"));
 	}
 
 	void SetPS(const FRenderingCompositePassContext& Context)
@@ -564,9 +639,9 @@ public:
 		PostprocessParameter.SetPS(Context.RHICmdList, ShaderRHI, Context, TStaticSamplerState<SF_Point, AM_Clamp, AM_Clamp, AM_Clamp>::GetRHI());
 
 		{
-			FVector4 Temp[EYE_ADAPTATION_PARAMS_SIZE];
-			ComputeEyeAdaptationValues(BasicEyeAdaptationMinFeatureLevel, Context.View, Temp);
-			SetShaderValueArray(Context.RHICmdList, ShaderRHI, EyeAdaptationParams, Temp, EYE_ADAPTATION_PARAMS_SIZE);
+			FParameters PassParameters;
+			PassParameters.EyeAdaptation = GetEyeAdaptationParameters(Context.View, BasicEyeAdaptationMinFeatureLevel);
+			SetShaderParameters(Context.RHICmdList, this, ShaderRHI, PassParameters);
 		}
 	}
 
@@ -666,15 +741,19 @@ FPooledRenderTargetDesc FRCPassPostProcessBasicEyeAdaptationSetUp::ComputeOutput
 class FPostProcessLogLuminance2ExposureScaleBase: public FGlobalShader
 {
 public:
+	BEGIN_SHADER_PARAMETER_STRUCT(FParameters, )
+		SHADER_PARAMETER_STRUCT(FEyeAdaptationParameters, EyeAdaptation)
+	END_SHADER_PARAMETER_STRUCT()
 
 	FPostProcessLogLuminance2ExposureScaleBase() {}
 
 	FPostProcessLogLuminance2ExposureScaleBase(const ShaderMetaType::CompiledShaderInitializerType& Initializer)
 		: FGlobalShader(Initializer)
 	{
+		BindForLegacyShaderParameters<FParameters>(this, Initializer.ParameterMap);
+
 		PostprocessParameter.Bind(Initializer.ParameterMap);
 		EyeAdaptationTexture.Bind(Initializer.ParameterMap, TEXT("EyeAdaptationTexture"));
-		EyeAdaptationParams.Bind(Initializer.ParameterMap, TEXT("EyeAdaptationParams"));
 		EyeAdaptationSrcRect.Bind(Initializer.ParameterMap, TEXT("EyeAdaptionSrcRect"));
 	}
 
@@ -683,7 +762,6 @@ public:
 	{
 		FGlobalShader::ModifyCompilationEnvironment(Parameters, OutEnvironment);
 		OutEnvironment.SetRenderTargetOutputFormat(0, PF_A32B32G32R32F);
-		OutEnvironment.SetDefine(TEXT("EYE_ADAPTATION_PARAMS_SIZE"), (uint32)EYE_ADAPTATION_PARAMS_SIZE);
 	}
 
 	// FShader interface.
@@ -750,15 +828,9 @@ public:
 
 		// Pack the eye adaptation parameters for the shader
 		{
-			FVector4 Temp[EYE_ADAPTATION_PARAMS_SIZE];
-			// static computation function
-			ComputeEyeAdaptationValues(BasicEyeAdaptationMinFeatureLevel, Context.View, Temp);
-			// Log-based computation of the exposure scale has a built in scaling.
-			//Temp[1].X *= 0.16;  
-			//Encode the eye-focus slope
-			// Get the focus value for the eye-focus weighting
-			Temp[2].W = GetBasicAutoExposureFocus();
-			SetShaderValueArray(Context.RHICmdList, ShaderRHI, EyeAdaptationParams, Temp, EYE_ADAPTATION_PARAMS_SIZE);
+			FParameters PassParameters;
+			PassParameters.EyeAdaptation = GetEyeAdaptationParameters(Context.View, BasicEyeAdaptationMinFeatureLevel);
+			SetShaderParameters(Context.RHICmdList, this, ShaderRHI, PassParameters);
 		}
 
 		// Set the src extent for the shader
@@ -770,10 +842,8 @@ IMPLEMENT_SHADER_TYPE(, FPostProcessLogLuminance2ExposureScalePS, TEXT("/Engine/
 
 class FPostProcessLogLuminance2ExposureScaleCS : public FPostProcessLogLuminance2ExposureScaleBase
 {
-	DECLARE_SHADER_TYPE(FPostProcessLogLuminance2ExposureScaleCS, Global);
-
-
 public:
+	DECLARE_SHADER_TYPE(FPostProcessLogLuminance2ExposureScaleCS, Global);
 
 	/** Default constructor. */
 	FPostProcessLogLuminance2ExposureScaleCS()
@@ -828,15 +898,9 @@ public:
 
 		// Pack the eye adaptation parameters for the shader
 		{
-			FVector4 Temp[EYE_ADAPTATION_PARAMS_SIZE];
-			// static computation function
-			ComputeEyeAdaptationValues(BasicEyeAdaptationMinFeatureLevel, Context.View, Temp);
-			// Log-based computation of the exposure scale has a built in scaling.
-			//Temp[1].X *= 0.16;  
-			//Encode the eye-focus slope
-			// Get the focus value for the eye-focus weighting
-			Temp[2].W = GetBasicAutoExposureFocus();
-			SetShaderValueArray(Context.RHICmdList, ShaderRHI, EyeAdaptationParams, Temp, EYE_ADAPTATION_PARAMS_SIZE);
+			FParameters PassParameters;
+			PassParameters.EyeAdaptation = GetEyeAdaptationParameters(Context.View, BasicEyeAdaptationMinFeatureLevel);
+			SetShaderParameters(Context.RHICmdList, this, ShaderRHI, PassParameters);
 		}
 
 		// Set the src extent for the shader
diff --git a/Engine/Source/Runtime/Renderer/Private/PostProcess/PostProcessEyeAdaptation.h b/Engine/Source/Runtime/Renderer/Private/PostProcess/PostProcessEyeAdaptation.h
--- a/Engine/Source/Runtime/Renderer/Private/PostProcess/PostProcessEyeAdaptation.h
+++ b/Engine/Source/Runtime/Renderer/Private/PostProcess/PostProcessEyeAdaptation.h
@@ -12,8 +12,6 @@
 #include "SceneRendering.h"
 #include "PostProcess/RenderingCompositionGraph.h"
 
-#define EYE_ADAPTATION_PARAMS_SIZE 4
-
 FORCEINLINE float EV100ToLuminance(float EV100)
 {
 	return 1.2 * FMath::Pow(2.0f, EV100);
@@ -34,7 +32,27 @@ FORCEINLINE float Log2ToEV100(float Log2)
 	return Log2 - 0.263f; // Where .263 is log2(1.2)
 }
 
-void ComputeEyeAdaptationValues(const ERHIFeatureLevel::Type MinFeatureLevel, const FViewInfo& View, FVector4 Out[EYE_ADAPTATION_PARAMS_SIZE]);
+bool IsAutoExposureMethodSupported(ERHIFeatureLevel::Type FeatureLevel, EAutoExposureMethod AutoExposureMethodId);
+
+EAutoExposureMethod GetAutoExposureMethod(const FViewInfo& View);
+
+BEGIN_SHADER_PARAMETER_STRUCT(FEyeAdaptationParameters, )
+	SHADER_PARAMETER(float, ExposureLowPercent)
+	SHADER_PARAMETER(float, ExposureHighPercent)
+	SHADER_PARAMETER(float, MinAverageLuminance)
+	SHADER_PARAMETER(float, MaxAverageLuminance)
+	SHADER_PARAMETER(float, ExposureCompensation)
+	SHADER_PARAMETER(float, DeltaWorldTime)
+	SHADER_PARAMETER(float, ExposureSpeedUp)
+	SHADER_PARAMETER(float, ExposureSpeedDown)
+	SHADER_PARAMETER(float, HistogramScale)
+	SHADER_PARAMETER(float, HistogramBias)
+	SHADER_PARAMETER(float, LuminanceMin)
+	SHADER_PARAMETER(float, CalibrationConstantInverse)
+	SHADER_PARAMETER(float, WeightSlope)
+END_SHADER_PARAMETER_STRUCT()
+
+FEyeAdaptationParameters GetEyeAdaptationParameters(const FViewInfo& ViewInfo, ERHIFeatureLevel::Type MinFeatureLevel = ERHIFeatureLevel::SM5);
 
 // Computes the eye-adaptation from HDRHistogram.
 // ePId_Input0: HDRHistogram or nothing
@@ -49,10 +67,6 @@ public:
 		bPreferAsyncCompute &= (GNumAlternateFrameRenderingGroups == 1); // Can't handle multi-frame updates on async pipe
 	}
 
-	// compute the parameters used for eye-adaptation.  These will default to values
-	// that disable eye-adaptation if the hardware doesn't support SM5 feature-level
-	static void ComputeEyeAdaptationParamsValue(const FViewInfo& View, FVector4 Out[EYE_ADAPTATION_PARAMS_SIZE]);
-	
 	// Computes the a fix exposure to be used to replace the dynamic exposure when it's not available (< SM5).
 	static float GetFixedExposure(const FViewInfo& View);
 
@@ -100,83 +114,4 @@ public:
 
 private:
 	FIntPoint DownsampledViewRect;
-};
-
-// Console Variable that is used to over-ride the post process settings.
-extern TAutoConsoleVariable<int32> CVarEyeAdaptationMethodOverride;
-
-// @return true if the current feature level supports this auto exposure method.
-static inline bool IsAutoExposureMethodSupported(const ERHIFeatureLevel::Type& FeatureLevel, const EAutoExposureMethod& AutoExposureMethodId)
-{
-	bool Result = false;
-
-	switch (AutoExposureMethodId)
-	{
-	case EAutoExposureMethod::AEM_Histogram:
-		Result = FeatureLevel >= ERHIFeatureLevel::SM5;
-		break;
-	case EAutoExposureMethod::AEM_Basic:
-	case EAutoExposureMethod::AEM_Manual:
-		Result = FeatureLevel >= ERHIFeatureLevel::ES3_1;
-		break;
-	default:
-		break;
-	}
-	return Result;
-}
-
-// Query the view for the auto exposure method, and allow for CVar override.
-static inline EAutoExposureMethod GetAutoExposureMethod(const FViewInfo& View)
-{
-	EAutoExposureMethod AutoExposureMethodId = View.FinalPostProcessSettings.AutoExposureMethod;
-
-	const int32 EyeOverride = CVarEyeAdaptationMethodOverride.GetValueOnRenderThread();
-
-	if (EyeOverride >= 0)
-	{
-		// Additional branching for override.
-		switch (EyeOverride)
-		{
-		case 1:
-		{
-			// Only override if the platform supports it.
-			if (View.GetFeatureLevel() >= ERHIFeatureLevel::SM5)
-			{
-				AutoExposureMethodId = EAutoExposureMethod::AEM_Histogram;
-			}
-			break;
-		}
-		case 2:
-		{
-			AutoExposureMethodId = EAutoExposureMethod::AEM_Basic;
-			break;
-		}
-		case 3:
-		{
-			AutoExposureMethodId = EAutoExposureMethod::AEM_Manual;
-			break;
-		}
-		}
-	}
-
-	// If auto exposure is disabled, revert to manual mode which will clamp to a reasonable default.
-	if (!View.Family->EngineShowFlags.EyeAdaptation)
-	{
-		AutoExposureMethodId = AEM_Manual;
-	}
-
-	// We should have a valid exposure method at this stage.
-	check(IsAutoExposureMethodSupported(View.GetFeatureLevel(), AutoExposureMethodId));
-
-	return AutoExposureMethodId;
-}
-
-extern TAutoConsoleVariable<float> CVarEyeAdaptationFocus;
-static inline float GetBasicAutoExposureFocus()
-{
-	// Hard coded value camp.
-	static float clampValue = 10.f;
-	float FocusValue = CVarEyeAdaptationFocus.GetValueOnRenderThread();
-	
-	return FMath::Max(FMath::Min(FocusValue, clampValue), 0.0f);
-}
+};
\ No newline at end of file
diff --git a/Engine/Source/Runtime/Renderer/Private/PostProcess/PostProcessHistogram.cpp b/Engine/Source/Runtime/Renderer/Private/PostProcess/PostProcessHistogram.cpp
--- a/Engine/Source/Runtime/Renderer/Private/PostProcess/PostProcessHistogram.cpp
+++ b/Engine/Source/Runtime/Renderer/Private/PostProcess/PostProcessHistogram.cpp
@@ -32,9 +32,9 @@ public:
 	BEGIN_SHADER_PARAMETER_STRUCT(FParameters, )
 		SHADER_PARAMETER_STRUCT_REF(FViewUniformShaderParameters, View)
 		SHADER_PARAMETER_STRUCT(FScreenPassTextureViewportParameters, Input)
+		SHADER_PARAMETER_STRUCT(FEyeAdaptationParameters, EyeAdaptation)
 		SHADER_PARAMETER_RDG_TEXTURE(Texture2D, InputTexture)
 		SHADER_PARAMETER_RDG_TEXTURE_UAV(RWTexture2D<float4>, HistogramRWTexture)
-		SHADER_PARAMETER_ARRAY(FVector4, EyeAdaptationParams, [EYE_ADAPTATION_PARAMS_SIZE])
 		SHADER_PARAMETER(FIntPoint, ThreadGroupCount)
 	END_SHADER_PARAMETER_STRUCT()
 
@@ -51,7 +51,6 @@ public:
 		OutEnvironment.SetDefine(TEXT("LOOP_SIZEX"), LoopCountX);
 		OutEnvironment.SetDefine(TEXT("LOOP_SIZEY"), LoopCountY);
 		OutEnvironment.SetDefine(TEXT("HISTOGRAM_SIZE"), HistogramSize);
-		OutEnvironment.SetDefine(TEXT("EYE_ADAPTATION_PARAMS_SIZE"), (uint32)EYE_ADAPTATION_PARAMS_SIZE);
 		OutEnvironment.CompilerFlags.Add( CFLAG_StandardOptimization );
 	}
 
@@ -136,7 +135,7 @@ FRDGTextureRef AddHistogramPass(
 		PassParameters->InputTexture = SceneColorTexture;
 		PassParameters->HistogramRWTexture = GraphBuilder.CreateUAV(HistogramTexture);
 		PassParameters->ThreadGroupCount = HistogramThreadGroupCount;
-		ComputeEyeAdaptationValues(ERHIFeatureLevel::SM5, ScreenPassView.View, &PassParameters->EyeAdaptationParams[0]);
+		PassParameters->EyeAdaptation = GetEyeAdaptationParameters(ScreenPassView.View);
 
 		TShaderMapRef<FHistogramCS> ComputeShader(ScreenPassView.View.ShaderMap);
 
diff --git a/Engine/Source/Runtime/Renderer/Private/PostProcess/PostProcessVisualizeHDR.cpp b/Engine/Source/Runtime/Renderer/Private/PostProcess/PostProcessVisualizeHDR.cpp
--- a/Engine/Source/Runtime/Renderer/Private/PostProcess/PostProcessVisualizeHDR.cpp
+++ b/Engine/Source/Runtime/Renderer/Private/PostProcess/PostProcessVisualizeHDR.cpp
@@ -21,6 +21,10 @@
 
 extern FIntPoint GetHistogramTexelsPerGroup();
 
+extern float GetBasicAutoExposureFocus();
+
+extern bool IsExtendLuminanceRangeEnabled();
+
 /** Encapsulates the post processing eye adaptation pixel shader. */
 class FPostProcessVisualizeHDRPS : public FGlobalShader
 {
@@ -38,13 +42,16 @@ class FPostProcessVisualizeHDRPS : public FGlobalShader
 		OutEnvironment.SetDefine(TEXT("USE_SHADOW_TINT"), 1);
 		OutEnvironment.SetDefine(TEXT("USE_CONTRAST"), 1);
 		OutEnvironment.SetDefine(TEXT("USE_APPROXIMATE_SRGB"), (uint32)0);
-		OutEnvironment.SetDefine(TEXT("EYE_ADAPTATION_PARAMS_SIZE"), (uint32)EYE_ADAPTATION_PARAMS_SIZE);
 	}
 
 	/** Default constructor. */
 	FPostProcessVisualizeHDRPS() {}
 
 public:
+	BEGIN_SHADER_PARAMETER_STRUCT(FParameters, )
+		SHADER_PARAMETER_STRUCT(FEyeAdaptationParameters, EyeAdaptation)
+	END_SHADER_PARAMETER_STRUCT()
+
 	FPostProcessPassParameters PostprocessParameter;
 	FShaderParameter EyeAdaptationParams;
 	FShaderResourceParameter MiniFontTexture;
@@ -67,6 +74,8 @@ public:
 	FPostProcessVisualizeHDRPS(const ShaderMetaType::CompiledShaderInitializerType& Initializer)
 		: FGlobalShader(Initializer)
 	{
+		BindForLegacyShaderParameters<FParameters>(this, Initializer.ParameterMap);
+
 		PostprocessParameter.Bind(Initializer.ParameterMap);
 		EyeAdaptationParams.Bind(Initializer.ParameterMap, TEXT("EyeAdaptationParams"));
 		MiniFontTexture.Bind(Initializer.ParameterMap, TEXT("MiniFontTexture"));
@@ -99,19 +108,9 @@ public:
 		PostprocessParameter.SetPS(RHICmdList, ShaderRHI, Context, TStaticSamplerState<SF_Point, AM_Clamp, AM_Clamp, AM_Clamp>::GetRHI());
 
 		{
-			FVector4 Temp[EYE_ADAPTATION_PARAMS_SIZE];
-
-			FRCPassPostProcessEyeAdaptation::ComputeEyeAdaptationParamsValue(Context.View, Temp);
-			if (GetAutoExposureMethod(Context.View) == EAutoExposureMethod::AEM_Basic)
-			{
-				Temp[2].W = GetBasicAutoExposureFocus();
-			}
-			else
-			{
-				Temp[2].W = 0.0f;
-			}
-
-			SetShaderValueArray(Context.RHICmdList, ShaderRHI, EyeAdaptationParams, Temp, EYE_ADAPTATION_PARAMS_SIZE);
+			FParameters PassParameters;
+			PassParameters.EyeAdaptation = GetEyeAdaptationParameters(Context.View);
+			SetShaderParameters(RHICmdList, this, ShaderRHI, PassParameters);
 		}
 
 		SetTextureParameter(RHICmdList, ShaderRHI, MiniFontTexture, GEngine->MiniFontTexture ? GEngine->MiniFontTexture->Resource->TextureRHI : GSystemTextures.WhiteDummy->GetRenderTargetItem().TargetableTexture);
@@ -209,8 +208,7 @@ void FRCPassPostProcessVisualizeHDR::Process(FRenderingCompositePassContext& Con
 		return;
 	}
 
-	static const auto VarDefaultAutoExposureExtendDefaultLuminanceRange = IConsoleManager::Get().FindTConsoleVariableDataInt(TEXT("r.DefaultFeature.AutoExposure.ExtendDefaultLuminanceRange"));
-	const bool bExtendedLuminanceRange = VarDefaultAutoExposureExtendDefaultLuminanceRange->GetValueOnRenderThread() == 1;
+	const bool bExtendedLuminanceRange = IsExtendLuminanceRangeEnabled();
 
 	const FViewInfo& View = Context.View;
 	const FViewInfo& ViewInfo = Context.View;
