Removed some SM5 checks in post processing.

diff --git a/Engine/Source/Runtime/Renderer/Private/PostProcess/PostProcessing.cpp b/Engine/Source/Runtime/Renderer/Private/PostProcess/PostProcessing.cpp
--- a/Engine/Source/Runtime/Renderer/Private/PostProcess/PostProcessing.cpp
+++ b/Engine/Source/Runtime/Renderer/Private/PostProcess/PostProcessing.cpp
@@ -1001,7 +1001,7 @@ void FPostProcessing::Process(FRHICommandListImmediate& RHICmdList, const FViewI
 			PreFlattenVelocity = Context.Graph.RegisterPass(new(FMemStack::Get()) FRCPassPostProcessInput(VelocityRT));
 		}
 
-		const bool bVisualizeHDR = View.Family->EngineShowFlags.VisualizeHDR && FeatureLevel >= ERHIFeatureLevel::SM5;
+		const bool bVisualizeHDR = View.Family->EngineShowFlags.VisualizeHDR;
 
 		const bool bVisualizeBloom = View.Family->EngineShowFlags.VisualizeBloom;
 
@@ -1452,7 +1452,7 @@ void FPostProcessing::Process(FRHICommandListImmediate& RHICmdList, const FViewI
 			Context.FinalOutput = FRenderingCompositeOutputRef(Node);
 		}
 #endif
-		if(View.Family->EngineShowFlags.VisualizeShadingModels && FeatureLevel >= ERHIFeatureLevel::SM5)
+		if(View.Family->EngineShowFlags.VisualizeShadingModels)
 		{
 			ensureMsgf(!bUnscaledFinalOutput, TEXT("VisualizeShadingModels is incompatible with unscaled output."));
 
@@ -1461,7 +1461,7 @@ void FPostProcessing::Process(FRHICommandListImmediate& RHICmdList, const FViewI
 			Context.FinalOutput = FRenderingCompositeOutputRef(Node);
 		}
 
-		if (View.Family->EngineShowFlags.GBufferHints && FeatureLevel >= ERHIFeatureLevel::SM5)
+		if (View.Family->EngineShowFlags.GBufferHints)
 		{
 			ensureMsgf(!bUnscaledFinalOutput, TEXT("GBufferHints is incompatible with unscaled output."));
 
@@ -1479,7 +1479,7 @@ void FPostProcessing::Process(FRHICommandListImmediate& RHICmdList, const FViewI
 		//No more postprocess Final color should be the real one
 		//The HDR was save before the tonemapping
 		//GBuffer should not be change during post process 
-		if (View.bUsePixelInspector && FeatureLevel >= ERHIFeatureLevel::SM5)
+		if (View.bUsePixelInspector)
 		{
 			FRenderingCompositePass* Node = Context.Graph.RegisterPass(new(FMemStack::Get()) FRCPassPostProcessBufferInspector(RHICmdList));
 			Node->SetInput(ePId_Input0, Context.FinalOutput);
@@ -1517,7 +1517,7 @@ void FPostProcessing::Process(FRHICommandListImmediate& RHICmdList, const FViewI
 			Context.FinalOutput = FRenderingCompositeOutputRef(Node);
 		}
 
-		if(View.Family->EngineShowFlags.TestImage && FeatureLevel >= ERHIFeatureLevel::SM5)
+		if(View.Family->EngineShowFlags.TestImage)
 		{
 			FRenderingCompositePass* Node = Context.Graph.RegisterPass(new(FMemStack::Get()) FRCPassPostProcessTestImage());
 			Node->SetInput(ePId_Input0, FRenderingCompositeOutputRef(Context.FinalOutput));
