1. Removed Separable pass.
2. Merge Combined Burley and Separable pass to PostprocessSubsurface.*
3. Clean up.

diff --git a/Engine/Shaders/Private/BurleyNormalizedSSSCommon.ush b/Engine/Shaders/Private/BurleyNormalizedSSSCommon.ush
--- a/Engine/Shaders/Private/BurleyNormalizedSSSCommon.ush
+++ b/Engine/Shaders/Private/BurleyNormalizedSSSCommon.ush
@@ -13,8 +13,8 @@
 //------------------------------------------------------------------------------------------
 // Consistent in BurleyNormalizedSSSCommon.ush and SubsurfaceProfile.cpp
 
-#define SSS_TYPE_SSSS		1.0f
-#define SSS_TYPE_BURLEY	    0.0f
+#define SSS_TYPE_BURLEY	    0
+#define SSS_TYPE_SSSS		1
 
 //make sure UIMax|ClampMax of WorldScale * ENC_WORLDSCALE_TO_UNIT <= 1
 #define ENC_WORLDSCALE_IN_CM_TO_UNIT 0.02f
@@ -51,6 +51,28 @@ bool GetSubsurfaceProfileUseBurley(uint SubsurfaceProfileInt)
 	return abs(Type - SSS_TYPE_BURLEY) < 0.01f;
 }
 
+bool GetSubsurfaceProfileUseBurley(float2 BufferUV)
+{
+	FScreenSpaceData ScreenSpaceData = GetScreenSpaceData(BufferUV);
+	uint SubsurfaceProfileInt = ExtractSubsurfaceProfileInt(ScreenSpaceData.GBuffer);
+	return GetSubsurfaceProfileUseBurley(SubsurfaceProfileInt);
+}
+
+bool GetSubsurfaceProfileUseSeparable(uint SubsurfaceProfileInt)
+{
+	// 0..255, which SubSurface profile to pick
+	float Type = ActualSSProfilesTexture.Load(int3(SSSS_BOUNDARY_COLOR_BLEED_OFFSET, SubsurfaceProfileInt, 0)).a;
+
+	return abs(Type - SSS_TYPE_SSSS) < 0.01f;
+}
+
+bool GetSubsurfaceProfileUseSeparable(float2 BufferUV)
+{
+	FScreenSpaceData ScreenSpaceData = GetScreenSpaceData(BufferUV);
+	uint SubsurfaceProfileInt = ExtractSubsurfaceProfileInt(ScreenSpaceData.GBuffer);
+	return GetSubsurfaceProfileUseSeparable(SubsurfaceProfileInt);
+}
+
 inline float3 Burley_Profile(float Radius, float3 Albedo, float3 S3D, float L)
 {   //R(r)r
 	float3 D = 1 / S3D;
@@ -140,7 +162,7 @@ float4 GetBurleyTransmissionProfile(FGBufferData GBufferData, float Thickness)
 	uint SubsurfaceProfileInt = ExtractSubsurfaceProfileInt(GBufferData);
 	float WorldScale = DecodeWorldScale(ActualSSProfilesTexture.Load(int3(SSSS_SUBSURFACE_COLOR_OFFSET, SubsurfaceProfileInt,0)).a)*10.0f; //in cm. we do not multiply by 100.0f as we did in the subsurface shader.
 	float3 MeanFreePath = DecodeDiffuseMeanFreePath(ActualSSProfilesTexture.Load(int3(SSSS_SUBSURFACE_COLOR_OFFSET, SubsurfaceProfileInt,0)));
-	float3 SurfaceAlbedo = ActualSSProfilesTexture.Load(int3(SSS__SURFACEALBEDO_OFFSET, SubsurfaceProfileInt, 0)).rgb;
+	float3 SurfaceAlbedo = ActualSSProfilesTexture.Load(int3(BSSS_SURFACEALBEDO_OFFSET, SubsurfaceProfileInt, 0)).rgb;
 	bool  UseBurleyTransmittance = GetSubsurfaceTransmittanceProfileUseBurley(SubsurfaceProfileInt);
 
 	//this is slow.
diff --git a/Engine/Shaders/Private/PostProcessSubsurface.usf b/Engine/Shaders/Private/PostProcessSubsurface.usf
--- a/Engine/Shaders/Private/PostProcessSubsurface.usf
+++ b/Engine/Shaders/Private/PostProcessSubsurface.usf
@@ -19,23 +19,31 @@ float4 SubsurfaceParams;
 
 Texture2D SubsurfaceInput0_Texture;
 Texture2D SubsurfaceInput1_Texture;
+Texture2D SubsurfaceInput2_Texture;
 
 SamplerState SubsurfaceSampler0;
 SamplerState SubsurfaceSampler1;
+SamplerState SubsurfaceSampler2;
 
 #if SUPPORTS_INDEPENDENT_SAMPLERS
 	#define SharedSubsurfaceSampler0 SubsurfaceSampler0
 	#define SharedSubsurfaceSampler1 SubsurfaceSampler0
+	#define SharedSubsurfaceSampler2 SubsurfaceSampler0
 #else
 	#define SharedSubsurfaceSampler0 SubsurfaceSampler0
 	#define SharedSubsurfaceSampler1 SubsurfaceSampler1
+	#define SharedSubsurfaceSampler2 SubsurfaceSampler2
 #endif
 
 SCREEN_PASS_TEXTURE_VIEWPORT(SubsurfaceInput0)
 SCREEN_PASS_TEXTURE_VIEWPORT(SubsurfaceInput1)
+SCREEN_PASS_TEXTURE_VIEWPORT(SubsurfaceInput2)
 
-#define SUBSURFACE_DIRECTION_HORIZONTAL 0
-#define SUBSURFACE_DIRECTION_VERTICAL 1
+#define SUBSURFACE_PASS_ONE 0
+#define SUBSURFACE_PASS_TWO 1
+
+#define SUBSURFACE_DIRECTION_HORIZONTAL SUBSURFACE_PASS_ONE
+#define SUBSURFACE_DIRECTION_VERTICAL SUBSURFACE_PASS_TWO
 
 // Controls the quality (number of samples) of the blur kernel.
 #define SUBSURFACE_QUALITY_LOW 0
@@ -380,18 +388,141 @@ float4 SetupSubsurfaceForOnePixel(float2 UVSceneColor)
 	return Ret;
 }
 
+// Ref: https://graphics.pixar.com/library/ApproxBSSRDF/paper.pdf
+#include "SubsurfaceBurleyNormalized.ush"
+
+#ifdef SUBSURFACE_BURLEY_COMPUTE
+//-----------------------------------------------------------------------------------
+// Indirect pass
 
-void SetupPS(in noperspective float4 UVAndScreenPos : TEXCOORD0, out float4 OutColor : SV_Target0)
+// Indirect dispatch parameters
+struct SubsurfaceGroupItem
 {
-	float2 BufferUV = UVAndScreenPos.xy;
+	float2 GroupXY; // Which tile
+	int Count;		// The first element is used as the counter.
+	int Dummy;		// Alignment
+};
+
+RWStructuredBuffer<SubsurfaceGroupItem> RWSeparableGroupBuffer;
+RWStructuredBuffer<SubsurfaceGroupItem> RWBurleyGroupBuffer;
+RWBuffer<uint> RWIndirectDispatchArgsBuffer;
+StructuredBuffer<SubsurfaceGroupItem> GroupBuffer;
+
+
+[numthreads(1, 1, 1)]
+void InitValueBufferCS()
+{
+	// Clear the first buffer element which is used as the buffer counter.
+	RWSeparableGroupBuffer[0].Count = 0;
+	RWBurleyGroupBuffer[0].Count = 0;
+}
+
+// fixed to 64 by now
+[numthreads(1, 1, 1)]
+void BuildIndirectDispatchArgsCS()
+{
+	uint ValueCount = min(SubsurfaceUniformParameters.MaxGroupCount, GroupBuffer[0].Count);
+	RWIndirectDispatchArgsBuffer[0] = ValueCount;// (ValueCount + (THREAD_SIZE_1D*THREAD_SIZE_1D) - 1) / (THREAD_SIZE_1D*THREAD_SIZE_1D);
+	RWIndirectDispatchArgsBuffer[1] = 1;
+	RWIndirectDispatchArgsBuffer[2] = 1;
+}
+
+#define TYPE_SEPARABLE 0x1
+#define TYPE_BURLEY    0x2
+
+// 0, no subsurface
+// 1, Separable
+// 2, Burley
+// 3, Both
+groupshared uint SubsurfaceTypeFlag = 0;
+
+uint GetSubsurfaceType(uint ProfileId)
+{
+	return GetSubsurfaceProfileUseBurley(ProfileId) ? TYPE_BURLEY : TYPE_SEPARABLE;
+}
+
+bool GetSubsurfaceProfileUseBurleyFromCache(float2 BufferUV)
+{
+#if SUBSURFACE_HALF_RES && !defined(SUBSURFACE_SINGLE_PASS)
+	uint Type = Texture2DSample(SubsurfaceInput2_Texture, SubsurfaceSampler2, BufferUV).g;
+	return Type == TYPE_BURLEY;
+#else
+	return GetSubsurfaceProfileUseBurley(BufferUV);
+#endif
+}
+
+void AddSeparableGroup(uint2 GroupXY)
+{
+	//buffer counter is stored in the .Count in the first element.
+	int IndexToStore = 0;
+
+	InterlockedAdd(RWSeparableGroupBuffer[0].Count, 1, IndexToStore);
+
+	if (IndexToStore >= SubsurfaceUniformParameters.MaxGroupCount)
+	{
+		return;
+	}
+
+	RWSeparableGroupBuffer[IndexToStore + 1].GroupXY = GroupXY;
+}
+
+void AddBurleyGroup(uint2 GroupXY)
+{
+	//buffer counter is stored in the .Count in the first element.
+	int IndexToStore = 0;
+
+	InterlockedAdd(RWBurleyGroupBuffer[0].Count, 1, IndexToStore);
+
+	if (IndexToStore >= SubsurfaceUniformParameters.MaxGroupCount)
+	{
+		return;
+	}
+
+	RWBurleyGroupBuffer[IndexToStore + 1].GroupXY = GroupXY;
+}
+
+void AddSubsurfaceComputeGroup(uint GI, uint Type, uint2 GroupXY)
+{
+	InterlockedOr(SubsurfaceTypeFlag, Type);
 
+	GroupMemoryBarrierWithGroupSync();
+
+	//write to buffer using the first thread
+	if (GI == 0 && (SubsurfaceTypeFlag & TYPE_SEPARABLE))
+	{
+		AddSeparableGroup(GroupXY);
+	}
+
+	if (GI == 0 && (SubsurfaceTypeFlag & TYPE_BURLEY))
+	{
+		AddBurleyGroup(GroupXY);
+	}
+}
+
+RWTexture2D<float4> SetupTexture;
+RWTexture2D<float4> ProfileEdgeMask;
+
+[numthreads(THREAD_SIZE_X, THREAD_SIZE_Y, 1)]
+void SetupIndirectCS(uint3 DT_ID : SV_DispatchThreadID,
+	uint3 G_ID : SV_GroupID,
+	uint3 GT_ID : SV_GroupThreadID,
+	uint  GI : SV_GroupIndex)
+{
+	uint2 Pos = DT_ID.xy + SubsurfaceInput0_ViewportMin;
+	bool bHasSubsurface = false;
+	float2 BufferUV = ConvertGridPos2UV(Pos);
+	uint Type = 0;
+	float4 OutColor = 0;
+	float4 OutColor2 = 0;
 #if SUBSURFACE_HALF_RES
+	//Since the tiles only covers the half resolution texture size, we need to scale up UV by 2 to match the scene texture.
+	BufferUV *= 2;
 	// order aligned with Gather() hardware implementation
 	// RGB: color*A, A:weight 0 if no subsurface scattering
 
-	float4 A = SetupSubsurfaceForOnePixel(min(BufferUV + float2(-0.5,  0.5f) * SubsurfaceInput0_ExtentInverse, SubsurfaceInput0_UVViewportBilinearMax));
-	float4 B = SetupSubsurfaceForOnePixel(min(BufferUV + float2( 0.5,  0.5f) * SubsurfaceInput0_ExtentInverse, SubsurfaceInput0_UVViewportBilinearMax));
-	float4 C = SetupSubsurfaceForOnePixel(min(BufferUV + float2( 0.5, -0.5f) * SubsurfaceInput0_ExtentInverse, SubsurfaceInput0_UVViewportBilinearMax));
+	float4 A = SetupSubsurfaceForOnePixel(min(BufferUV + float2(-0.5, 0.5f) * SubsurfaceInput0_ExtentInverse, SubsurfaceInput0_UVViewportBilinearMax));
+	float4 B = SetupSubsurfaceForOnePixel(min(BufferUV + float2(0.5, 0.5f) * SubsurfaceInput0_ExtentInverse, SubsurfaceInput0_UVViewportBilinearMax));
+	float4 C = SetupSubsurfaceForOnePixel(min(BufferUV + float2(0.5, -0.5f) * SubsurfaceInput0_ExtentInverse, SubsurfaceInput0_UVViewportBilinearMax));
 	float4 D = SetupSubsurfaceForOnePixel(min(BufferUV + float2(-0.5, -0.5f) * SubsurfaceInput0_ExtentInverse, SubsurfaceInput0_UVViewportBilinearMax));
 
 	float4 Sum = (A + B) + (C + D);
@@ -406,38 +537,126 @@ void SetupPS(in noperspective float4 UVAndScreenPos : TEXCOORD0, out float4 OutC
 	float SingleDepth = dot(FourDepth, float4(A.a, B.a, C.a, D.a)) * Div;
 
 	OutColor.a = SingleDepth;
+
+	if (OutColor.a > 0)
+	{
+		bHasSubsurface = true;
+		OutColor2 = GetProfileEdgeMask(BufferUV);
+	}
+	//BufferUV += min(float2(-0.25, -0.25f) * SubsurfaceInput0_ExtentInverse, SubsurfaceInput0_UVViewportBilinearMax);
 #else // SUBSURFACE_HALF_RES
 	OutColor = SetupSubsurfaceForOnePixel(BufferUV);
-	if(OutColor.a > 0)
+
+	if (OutColor.a > 0)
 	{
+		bHasSubsurface = true;
+
 		OutColor.a = CalcSceneDepth(BufferUV);
+		// Calculate the mask
+		OutColor2 = GetProfileEdgeMask(BufferUV);
+
+		//@TODO: sparkling
+		//supress the edge color lighting information at edge.
+		//with firefly elimination
 	}
 #endif // SUBSURFACE_HALF_RES
+
+	// setup the Subsurface type
+	uint SelectedProfile = 0;
+	if (bHasSubsurface)
+	{
+#if SUBSURFACE_HALF_RES
+		float Type1 = 0;
+		float2 UVA = min(BufferUV + float2(-0.5, 0.5f) * SubsurfaceInput0_ExtentInverse, SubsurfaceInput0_UVViewportBilinearMax);
+		uint PidA = GetSubsurfaceProfileId(UVA);
+		float2 UVB = min(BufferUV + float2(0.5, 0.5f) * SubsurfaceInput0_ExtentInverse, SubsurfaceInput0_UVViewportBilinearMax);
+		uint PidB = GetSubsurfaceProfileId(UVB);
+		float2 UVC = min(BufferUV + float2(0.5, -0.5f) * SubsurfaceInput0_ExtentInverse, SubsurfaceInput0_UVViewportBilinearMax);
+		uint PidC = GetSubsurfaceProfileId(UVC);
+		float2 UVD = min(BufferUV + float2(-0.5, -0.5f) * SubsurfaceInput0_ExtentInverse, SubsurfaceInput0_UVViewportBilinearMax);
+		uint PidD = GetSubsurfaceProfileId(UVD);
+
+		// Gatter available and use the last subsurface type for the current pixel.
+		SelectedProfile = lerp(SelectedProfile, PidA, A.a);
+		SelectedProfile = lerp(SelectedProfile, PidB, B.a);
+		SelectedProfile = lerp(SelectedProfile, PidC, C.a);
+		SelectedProfile = lerp(SelectedProfile, PidD, D.a);
+
+		// record the uv of this pixel
+		BufferUV = lerp(BufferUV, UVA, A.a);
+		BufferUV = lerp(BufferUV, UVB, B.a);
+		BufferUV = lerp(BufferUV, UVC, C.a);
+		BufferUV = lerp(BufferUV, UVD, D.a);
+
+		Type = GetSubsurfaceType(SelectedProfile);
+#else
+		SelectedProfile = GetSubsurfaceProfileId(BufferUV);
+		Type = GetSubsurfaceType(SelectedProfile);
+#endif
+	}
+
+	if (all(Pos < SubsurfaceInput0_ViewportMax))
+	{
+		SetupTexture[Pos] = OutColor;
+		//cache the edge mask, depth, and the original uv in the buffer.
+		ProfileEdgeMask[Pos] = float4(OutColor2.r, Type, BufferUV.x, BufferUV.y);
+	}
+
+	AddSubsurfaceComputeGroup(GI, Type, G_ID.xy);
 }
 
-// input0 is created by the SetupPS shader
-void MainPS(noperspective float4 UVAndScreenPos : TEXCOORD0, out float4 OutColor : SV_Target0)
+[numthreads(THREAD_SIZE_X*THREAD_SIZE_Y, 1, 1)]
+void MainIndirectDispatchCS(
+	uint3 IndirectG_ID : SV_GroupID,
+	uint  GI : SV_GroupIndex
+)
 {
-	float2 BufferUV = UVAndScreenPos.xy;
+	// The first element is used for counting
+	SubsurfaceGroupItem Item = GroupBuffer[IndirectG_ID.x + 1];
+	uint2 G_ID = Item.GroupXY;
+	uint2 DT_ID = G_ID * float2(THREAD_SIZE_X, THREAD_SIZE_Y) + float2(GI % THREAD_SIZE_X, GI / THREAD_SIZE_X);
+
+
+	//1. Call function based on which function is defined and in which phase
+#if SUBSURFACE_TYPE == SSS_TYPE_BURLEY
+
+	BurleyComputeMain(DT_ID, G_ID, GI);
+
+#elif SUBSURFACE_TYPE == SSS_TYPE_SSSS
 
-#if SUBSURFACE_DIRECTION == 0
+	uint2 Pos = DT_ID.xy*SUBSURFACE_GROUP_SIZE / THREAD_SIZE_1D + /*viewport offset*/SubsurfaceInput0_ViewportMin;
+	float2 BufferUV = ConvertGridPos2UV(Pos);
+
+	BRANCH if (GetSubsurfaceProfileUseBurley(BufferUV))
+	{
+		return;
+	}
+	// pass one and pass two
+#if SUBSURFACE_PASS == SUBSURFACE_PASS_ONE
 	// horizontal
 	float2 ViewportDirectionUV = float2(1, 0) * SUBSURFACE_RADIUS_SCALE;
-#else
+#elif SUBSURFACE_PASS == SUBSURFACE_PASS_TWO
 	// vertical
 	float2 ViewportDirectionUV = float2(0, 1) * SUBSURFACE_RADIUS_SCALE * (SubsurfaceInput0_Extent.x * SubsurfaceInput0_ExtentInverse.y);
 #endif
 
 	ViewportDirectionUV *= (SubsurfaceInput0_ViewportSize.x * SubsurfaceInput0_ExtentInverse.x);
 
-	OutColor = SSSSBlurPS(BufferUV, ViewportDirectionUV, false);
+	float4 OutColor = SSSSBlurPS(BufferUV, ViewportDirectionUV, false);
 
-#if SUBSURFACE_DIRECTION == SUBSURFACE_DIRECTION_VERTICAL
+#if SUBSURFACE_PASS == SUBSURFACE_PASS_TWO
 	// second pass prepares the setup from the recombine pass which doesn't need depth but wants to reconstruct the color
 	OutColor.a = GetMaskFromDepthInAlpha(OutColor.a);
+#endif
+	if (all(Pos < SubsurfaceInput0_ViewportMax))
+	{
+		SSSColorUAV[Pos] = OutColor;
+	}
 #endif
 }
 
+#endif
+
 // Recombines the half res Subsurface filtered lighting contribution (upsampled and renormalized with the alpha)
 // with the SceneColor.
 void SubsurfaceRecombinePS(noperspective float4 UVAndScreenPos : TEXCOORD0, float4 SvPosition : SV_POSITION, out float4 OutColor : SV_Target0)
diff --git a/Engine/Shaders/Private/SubsurfaceBurleyNormalized.usf b/Engine/Shaders/Private/SubsurfaceBurleyNormalized.usf
deleted file mode 100644
diff --git a/Engine/Shaders/Private/SubsurfaceBurleyNormalized.ush b/Engine/Shaders/Private/SubsurfaceBurleyNormalized.ush
new file mode 100644
--- /dev/null
+++ b/Engine/Shaders/Private/SubsurfaceBurleyNormalized.ush
@@ -0,0 +1,765 @@
+// Copyright 1998-2019 Epic Games, Inc. All Rights Reserved.
+
+/*=============================================================================
+	SubsurfaceBurleyNormalized.ush: Screenspace Burley subsurface scattering implementation.
+=============================================================================*/
+
+#include "Random.ush"
+#include "BurleyNormalizedSSSCommon.ush"
+
+// Setup the max number of Burley samples
+#define BURLEY_NUM_SAMPLES	64
+#define BURLEY_INV_NUM_SAMPLES (1.0f/BURLEY_NUM_SAMPLES)
+
+struct FBurleyParameter
+{
+	float4 SurfaceAlbedo;
+	float4 DiffuseMeanFreePath;
+	float  WorldScale;
+};
+
+//Surface albedo and mean free path length
+float4 GetSubsurfaceProfileSurfaceAlbedo(FGBufferData GBufferData)
+{
+	// 0..255, which SubSurface profile to pick
+	uint SubsurfaceProfileInt = ExtractSubsurfaceProfileInt(GBufferData);
+
+	return ActualSSProfilesTexture.Load(int3(BSSS_SURFACEALBEDO_OFFSET, SubsurfaceProfileInt, 0));
+}
+
+float4 GetSubsurfaceProfileDiffuseMeanFreePath(FGBufferData GBufferData)
+{
+	// 0..255, which SubSurface profile to pick
+	uint SubsurfaceProfileInt = ExtractSubsurfaceProfileInt(GBufferData);
+
+	return ActualSSProfilesTexture.Load(int3(BSSS_DMFP_OFFSET, SubsurfaceProfileInt, 0));
+}
+
+float GetSubsurfaceProfileWorldScale(FGBufferData GBufferData)
+{
+	// 0..255, which SubSurface profile to pick
+	uint SubsurfaceProfileInt = ExtractSubsurfaceProfileInt(GBufferData);
+
+	return ActualSSProfilesTexture.Load(int3(SSSS_SUBSURFACE_COLOR_OFFSET, SubsurfaceProfileInt, 0)).a;
+}
+
+float SampleDepthTexture(float2 ScreenUV)
+{
+
+#if SUBSURFACE_HALF_RES && !defined(SUBSURFACE_SINGLE_PASS)
+	ConvertToDeviceZ(Texture2DSample(SubsurfaceInput1_Texture, SubsurfaceSampler1, ScreenUV).g);
+#else
+	float2 FullScreenUV = ScreenUV;
+#endif
+	
+	return Texture2DSampleLevel(SceneTexturesStruct.SceneDepthTexture, SceneTexturesStruct.SceneDepthTextureSampler, ScreenUV, 0).r;
+}
+
+// fetch surface albedo and diffuse mean free path
+// The diffuse mean free path is modulated by world scale and a custom term passed in
+// which is the opacity texture.
+FBurleyParameter GetBurleyParameters(FGBufferData GBuffer)
+{
+	FBurleyParameter BurleyParameter;
+
+	BurleyParameter.SurfaceAlbedo = GetSubsurfaceProfileSurfaceAlbedo(GBuffer);
+	float4 DiffuseMeanFreePath = GetSubsurfaceProfileDiffuseMeanFreePath(GBuffer);
+	BurleyParameter.DiffuseMeanFreePath = DecodeDiffuseMeanFreePath(DiffuseMeanFreePath);
+	float WorldScale = DecodeWorldScale(GetSubsurfaceProfileWorldScale(GBuffer))*BURLEY_CM_2_MM;
+	BurleyParameter.WorldScale = WorldScale;
+	BRANCH if (GBuffer.ShadingModelID == SHADINGMODELID_SUBSURFACE_PROFILE)
+	{
+		// scale with opacity stored in CustomData.a
+		BurleyParameter.DiffuseMeanFreePath *= GBuffer.CustomData.a;
+	}
+	else if (GBuffer.ShadingModelID == SHADINGMODELID_EYE)
+	{
+		//eye 
+	}
+	return BurleyParameter;
+}
+
+float2 Generate2DRandomNumber(int3 Seed)
+{
+	return float2(Rand3DPCG16(Seed).xy) / 0x10000;
+}
+
+float GetCDF(float D, float X, float XI)
+{
+	return 1 - 0.25*exp(-X / D) - 0.75*exp(-X / (3 * D)) - XI;
+}
+float GetCDFDeriv1(float D, float X)
+{
+	return 0.25 / D * (exp(-X / D) + exp(-X / (3 * D)));
+}
+
+float GetCDFDeriv2(float D, float X)
+{
+	return exp(-X / D)*(-0.0833333*exp(2 * X / (3 * D)) - 0.25) / (D*D);
+}
+
+// Without clamp, the result will be zero which will lead to infinity for R/pdf. and Nan for the weighted sum.
+// We need to clamp pdf to a very small value. this clamp is based on value below
+#define CLAMP_PDF 0.00001
+
+// Given D and a random number, use root finding method to calculate the radius
+// in meters
+float RadiusRootFinding(float D, float RandomNumber, float X0)
+{
+	// Make the distribution correct.
+	// r=cdf^(-1)(p)     if p<1-CLAMP_PDF
+	//  =cdf^(-1)(1-ClAMP_PDF) P>=1-CLAMP_PDF
+	// RandomNumber = clamp(RandomNumber,0,1-CLAMP_PDF);
+ 
+	const int NumOfIteration = 3;
+	float Xn = X0;
+
+	UNROLL for (int i = 0; i < NumOfIteration; ++i)
+	{
+		float Fx = GetCDF(D, Xn, RandomNumber);
+		float DFx = GetCDFDeriv1(D, Xn);
+		float DFxx = GetCDFDeriv2(D, Xn);
+		Xn = Xn - 2 * Fx*DFx / (2 * DFx*DFx - Fx * DFxx);
+	}
+
+	return Xn;
+}
+
+// Brian's approximation.
+float RadiusRootFindByApproximation(float D, float RandomNumber)
+{
+	return D * ((2 - 2.6)*RandomNumber - 2)*log(1 - RandomNumber);
+}
+
+// Get the probability to sample a disk.
+float GetPdf(float Radius, float L, float S)
+{
+	//without clamp, the result will be zero which will lead to infinity for R/pdf. and Nan for the weighted sum.
+	//we need to clamp this to a very small pdf.
+
+	float Pdf = GetCDFDeriv1(L / S, Radius);
+	return  max(Pdf, CLAMP_PDF);
+}
+
+struct FBurleySampleInfo
+{
+	float Radius;
+	float Theta;
+	float Pdf;
+};
+
+FBurleySampleInfo GenerateSampleInfo(float2 Rand0T1, float DiffuseMeanFreePathForSample, float SpectralForSample, uint SequenceId)
+{
+	FBurleySampleInfo BurleySampleInfo;
+
+#define SAMPLE_ANGLE_FIBONACCI
+#define ROOT_APROXIMATE
+
+	// Sample radius
+#ifdef ROOT_FINDING
+	float FoundRoot = RadiusRootFinding(DiffuseMeanFreePathForSample / SpectralForSample, Rand0T1.x, DiffuseMeanFreePathForSample);
+#elif defined(ROOT_APROXIMATE)
+	//Approximation
+	float FoundRoot = RadiusRootFindByApproximation(DiffuseMeanFreePathForSample / SpectralForSample, Rand0T1.x);
+#endif
+
+	BurleySampleInfo.Radius = max(FoundRoot, 0.00001f);//in mm
+	
+	// Sample angle
+#ifdef SAMPLE_ANGLE_RANDOM
+	BurleySampleInfo.Theta = Rand0T1.y * 2 * PI;
+#elif defined(SAMPLE_ANGLE_FIBONACCI)
+	//Less variance
+	BurleySampleInfo.Theta = (((SequenceId + 0.5 + Rand0T1.y))*(1 + sqrt(5))*0.5) * 2 * PI;//Fibonacci sequence for angle.
+#endif
+
+	// Estimate Pdf
+	BurleySampleInfo.Pdf = GetPdf(BurleySampleInfo.Radius, DiffuseMeanFreePathForSample, SpectralForSample); //
+
+	return BurleySampleInfo;
+}
+
+float GetProfileMask(float2 BufferUV)
+{
+	FScreenSpaceData ScreenSpaceData = GetScreenSpaceData(BufferUV);
+	float Ret = 0;
+	BRANCH if (UseSubsurfaceProfile(ScreenSpaceData.GBuffer.ShadingModelID))
+	{
+		Ret = 1.0f;
+	}
+
+	return Ret;
+}
+
+float GetProfileEdgeMask(float2 BufferUV)
+{
+#ifdef SUBSURFACE_HALF_RES
+	#define PIXELOFFSET_UVDELTA 0.5f
+#else
+	#define PIXELOFFSET_UVDELTA 1.0f
+#endif
+
+	float P11 = GetProfileMask((BufferUV + float2(0.0f, 0.0f) * SubsurfaceInput0_ExtentInverse));
+	float P00 = GetProfileMask((BufferUV + float2(-PIXELOFFSET_UVDELTA, -PIXELOFFSET_UVDELTA) * SubsurfaceInput0_ExtentInverse));
+	float P01 = GetProfileMask((BufferUV + float2(0.0f, -PIXELOFFSET_UVDELTA) * SubsurfaceInput0_ExtentInverse));
+	float P02 = GetProfileMask((BufferUV + float2(PIXELOFFSET_UVDELTA, -PIXELOFFSET_UVDELTA) * SubsurfaceInput0_ExtentInverse));
+	float P10 = GetProfileMask((BufferUV + float2(-PIXELOFFSET_UVDELTA, 0.0f) * SubsurfaceInput0_ExtentInverse));
+	float P12 = GetProfileMask((BufferUV + float2(PIXELOFFSET_UVDELTA, 0.0f) * SubsurfaceInput0_ExtentInverse));
+	float P20 = GetProfileMask((BufferUV + float2(-PIXELOFFSET_UVDELTA, PIXELOFFSET_UVDELTA) * SubsurfaceInput0_ExtentInverse));
+	float P21 = GetProfileMask((BufferUV + float2(0.0f, PIXELOFFSET_UVDELTA) * SubsurfaceInput0_ExtentInverse));
+	float P22 = GetProfileMask((BufferUV + float2(PIXELOFFSET_UVDELTA, PIXELOFFSET_UVDELTA) * SubsurfaceInput0_ExtentInverse));
+
+	return (P00 + P01 + P02 + P10 + P11 + P12 + P20 + P21 + P22) / 9.0f;
+}
+
+float RadiusRootFindingCM(float D, float RandomNumber, float X0)
+{
+	return RadiusRootFinding(D*0.01, RandomNumber, X0*0.01)*100.0f;
+}
+
+float GetPdfInCM(float Radius, float L, float S)
+{
+	return GetPdf(Radius *0.01f, L*0.01f, S);
+}
+
+int GetNumOfSamplesBasedOnQuality(float2 UV)
+{
+
+	float4 QualityMatrix = Texture2DSample(SubsurfaceInput1_Texture, SharedSubsurfaceSampler1, UV);
+	float VariableVar = QualityMatrix.b;
+
+	#define VARIANCE_LEVEL 0.0001
+	#define HIGH_LUMA_SAMPLE_COUNT 8
+	#define LOW_LUMA_SAMPLE_COUNT 16
+	#define PROFILE_EDGE_SAMPLE_COUNT 32
+
+	//1. Estimate the number of samples required to reach the target variance level
+	float NumOfSamples = clamp((VariableVar / VARIANCE_LEVEL), 8, BURLEY_NUM_SAMPLES);//View.GeneralPurposeTweak
+
+	//1.1 If the luminance is too low, we should increase several number of samples to oberse the world for high lighting condition
+	NumOfSamples = lerp(NumOfSamples, max(LOW_LUMA_SAMPLE_COUNT, NumOfSamples), step(0.5, QualityMatrix.a));
+
+	//2. Clamp down the number of samples when the gamma corrected luminance is too large
+	//   because it does not help to improve the quality, but degrade the performance.
+	NumOfSamples = lerp(NumOfSamples, HIGH_LUMA_SAMPLE_COUNT, step(1.5, QualityMatrix.a)); // 1.5 is an experimental value
+
+	//3. Reduce the number of samples at the edge due to TAA (Flickering)
+	float Mask = Texture2DSample(SubsurfaceInput2_Texture, SharedSubsurfaceSampler2, UV).r;
+	NumOfSamples = lerp(NumOfSamples, PROFILE_EDGE_SAMPLE_COUNT, step(Mask, 0.99));
+
+	return NumOfSamples;
+}
+
+// Texture local cash does not help on 2080TI with regular layout
+#define TEXTURE_CACHE_DISABLED
+
+#ifdef SUBSURFACE_BURLEY_COMPUTE
+
+// Configurations for each small thread
+#define THREAD_SIZE_1D 8
+#define THREAD_SIZE_X THREAD_SIZE_1D
+#define THREAD_SIZE_Y THREAD_SIZE_X
+#define THREAD_TOTAL_SZIE (THREAD_SIZE_X*THREAD_SIZE_Y)
+
+#define THREAD_TEXTURE_BOARDER 1
+#define THREAD_TEXTURE_DIFFUSE_REGION_SIZE   THREAD_SIZE_X
+#define THREAD_TEXTURE_DIMENSION_SIZE (THREAD_TEXTURE_DIFFUSE_REGION_SIZE + 2*THREAD_TEXTURE_BOARDER)
+#define THREADGROUP_TEXTURE_SHARE_TOTALSIZE (THREAD_TEXTURE_DIMENSION_SIZE*THREAD_TEXTURE_DIMENSION_SIZE)
+
+// Configuration of group threads
+#ifndef SUBSURFACE_GROUP_SIZE
+#define SUBSURFACE_GROUP_SIZE 8
+#endif
+
+#define LARGE_GROUP_TOTAL_SIZE (SUBSURFACE_GROUP_SIZE*SUBSURFACE_GROUP_SIZE)
+#define LARGE_GROUP_DIFFUSE_REGION_SIZE SUBSURFACE_GROUP_SIZE
+#define LARGE_GROUP_TEXTURE_DIMENSION_SIZE (LARGE_GROUP_DIFFUSE_REGION_SIZE+2*THREAD_TEXTURE_BOARDER)
+#define LARGE_GROUP_TEXTURE_SHARE_TOTALSIZE (LARGE_GROUP_TEXTURE_DIMENSION_SIZE*LARGE_GROUP_TEXTURE_DIMENSION_SIZE)
+
+#define NUM_OF_PASS_REQUIRED_FILL_SHARED_DIFFUSE_TEXTURE ((LARGE_GROUP_TEXTURE_SHARE_TOTALSIZE+THREAD_TOTAL_SZIE-1)/THREAD_TOTAL_SZIE)
+
+#define LOCALGROUP_RATIO (SUBSURFACE_GROUP_SIZE/THREAD_SIZE_1D)
+
+#ifndef TEXTURE_CACHE_DISABLED 
+groupshared float4 SharedSubsurfaceDiffuseLighting[LARGE_GROUP_TEXTURE_SHARE_TOTALSIZE];
+groupshared uint NotHaveSubsurface = 0;
+#endif
+
+// Ref: https://www.shadertoy.com/view/4sscDn
+static int Masks[] =
+{
+	0x55555555, 0x33333333, 0x0F0F0F0F, 0x00FF00FF, 0x0000FFFF
+};
+
+// Ref: [REDACTED_PROFILE_URL]/libmorton/blob/master/libmorton/include/morton2D_LUTs.h
+static uint MortonLUTX[] =
+{
+	0,1,0,1,2,3,2,3,0,1,0,1,2,3,2,3,
+	4,5,4,5,6,7,6,7,4,5,4,5,6,7,6,7,
+	0,1,0,1,2,3,2,3,0,1,0,1,2,3,2,3,
+	4,5,4,5,6,7,6,7,4,5,4,5,6,7,6,7
+};
+
+static uint MortonLUTY[] =
+{
+	0,0,1,1,0,0,1,1,2,2,3,3,2,2,3,3,
+	0,0,1,1,0,0,1,1,2,2,3,3,2,2,3,3,
+	4,4,5,5,4,4,5,5,6,6,7,7,6,6,7,7,
+	4,4,5,5,4,4,5,5,6,6,7,7,6,6,7,7
+};
+
+//#define MORTON_USE_LUT
+
+// We can move this to cache instead of computing
+uint2 ConvertGroupIndexTo2DUsingMoltonCode(uint GroupIndex, uint2 StartOffset)
+{
+#ifdef MORTON_USE_LUT
+	return uint2(MortonLUTX[GroupIndex], MortonLUTY[GroupIndex]) + StartOffset;
+#else
+	int n = 1;
+	uint2 I = uint2(GroupIndex, GroupIndex >> 1) & Masks[0];
+	UNROLL for (int i = 1; i <= 4; ++i)
+	{
+		I = (I | (I >> n)) & Masks[i];
+		n *= 2;
+	}
+	return I + StartOffset;
+#endif
+}
+
+uint2 ConvertGroupIndexToNormal2DGrid(uint LocalGroupIndex, uint2 StartOffset)
+{
+	return uint2(LocalGroupIndex % THREAD_SIZE_1D,
+		LocalGroupIndex / THREAD_SIZE_1D) + StartOffset;
+}
+
+// BufferUV
+float2 ConvertGridPos2UV(uint2 GridPosition)
+{
+	float2 GripPositionF = float2(GridPosition);
+	return SubsurfaceInput0_ExtentInverse * (GripPositionF + 0.5f);
+}
+// Convert UV to groupThreadIds
+float2 ConvertUVOffset2GTIDOffsetForInput0(float2 UVOffset)
+{
+	return UVOffset * SubsurfaceInput0_Extent;
+}
+
+bool IsIDInsideLocalShared(float2 GroupThreadId)
+{
+	return GroupThreadId.x >= -THREAD_TEXTURE_BOARDER &&
+		(GroupThreadId.x < (LARGE_GROUP_TEXTURE_DIMENSION_SIZE - THREAD_TEXTURE_BOARDER)) &&
+		GroupThreadId.y >= -THREAD_TEXTURE_BOARDER &&
+		(GroupThreadId.y < (LARGE_GROUP_TEXTURE_DIMENSION_SIZE - THREAD_TEXTURE_BOARDER));
+}
+
+#ifndef TEXTURE_CACHE_DISABLED
+float4 SampleSharedDiffuse(float2 SampleGroupThreadId)
+{
+#if NONSUBSURFACE_SAMPLE_STRATEGY == 0
+#define SAMPLDE_POINT
+#elif NONSUBSURFACE_SAMPLE_STRATEGY == 1
+#define SAMPLE_BILINEAR
+#endif
+
+#ifdef SAMPLDE_POINT
+	// We use point sampling by now
+	int2 Id = SampleGroupThreadId + THREAD_TEXTURE_BOARDER;
+	return SharedSubsurfaceDiffuseLighting[Id.x + Id.y*LARGE_GROUP_TEXTURE_DIMENSION_SIZE];
+#elif defined(SAMPLE_BILINEAR)
+
+	//ref: https://en.wikipedia.org/wiki/Bilinear_interpolation
+	int2 Id00 = floor(SampleGroupThreadId) + THREAD_TEXTURE_BOARDER;
+	int2 Id11 = ceil(SampleGroupThreadId) + THREAD_TEXTURE_BOARDER;
+	int2 Id01 = int2(Id00.x, Id11.y);
+	int2 Id10 = int2(Id11.x, Id00.y);
+	float4 Q00 = SharedSubsurfaceDiffuseLighting[Id00.x + Id00.y*LARGE_GROUP_TEXTURE_DIMENSION_SIZE];
+	float4 Q01 = SharedSubsurfaceDiffuseLighting[Id01.x + Id01.y*LARGE_GROUP_TEXTURE_DIMENSION_SIZE];
+	float4 Q10 = SharedSubsurfaceDiffuseLighting[Id10.x + Id10.y*LARGE_GROUP_TEXTURE_DIMENSION_SIZE];
+	float4 Q11 = SharedSubsurfaceDiffuseLighting[Id11.x + Id11.y*LARGE_GROUP_TEXTURE_DIMENSION_SIZE];
+	float x = SampleGroupThreadId.x + THREAD_TEXTURE_BOARDER - Id00.x;
+	float y = SampleGroupThreadId.y + THREAD_TEXTURE_BOARDER - Id11.y;
+	return Q00 * (1 - x)*(1 - y) + Q10 * x*(1 - y) + Q01 * (1 - x)*y + Q11 * x*y;
+#endif
+}
+#endif
+
+#endif
+
+float4 SampleSSSColorConsideringLocalShared(float2 CenterUV, float2 UVOffset, uint2 CenterGroupThreadID, float MipLevel)
+{
+
+#ifdef SUBSURFACE_BURLEY_COMPUTE
+
+#ifdef TEXTURE_CACHE_DISABLED 
+	return Texture2DSampleLevel(SubsurfaceInput0_Texture, SubsurfaceSampler0, CenterUV + UVOffset, MipLevel);
+#else
+	float2 SampleGroupThreadId = ConvertUVOffset2GTIDOffsetForInput0(UVOffset) + CenterGroupThreadID; // Subtract 0.5 is to make sampling match. (0,0), samples at 1/w*0.5;
+	bool bUseLocalShared = MipLevel == 0 && IsIDInsideLocalShared(SampleGroupThreadId);// We will have artifacts if we do not limit the miplevel
+
+	BRANCH
+		if (bUseLocalShared)
+		{
+			return SampleSharedDiffuse(SampleGroupThreadId);
+		}
+		else
+		{
+			return Texture2DSampleLevel(SubsurfaceInput0_Texture, SubsurfaceSampler0, CenterUV + UVOffset, MipLevel);
+		}
+#endif
+
+#else
+	return Texture2DSampleLevel(SubsurfaceInput0_Texture, SubsurfaceSampler0, CenterUV + UVOffset, MipLevel);
+#endif
+}
+
+// If it is in half resolution, we get the buffer from cached location.
+FScreenSpaceData GetScreenSpaceDataFromCache(float2 BufferUV)
+{
+#if SUBSURFACE_HALF_RES && !defined(SUBSURFACE_SINGLE_PASS)
+	float2 FullScreenUV = Texture2DSample(SubsurfaceInput2_Texture, SubsurfaceSampler2, BufferUV).zw;
+#else
+	float2 FullScreenUV = BufferUV;
+#endif
+	return GetScreenSpaceData(FullScreenUV);
+}
+
+uint GetSubsurfaceProfileIdFromCache(float2 BufferUV)
+{
+#if SUBSURFACE_HALF_RES && !defined(SUBSURFACE_SINGLE_PASS)
+	float2 FullScreenUV = Texture2DSample(SubsurfaceInput2_Texture, SubsurfaceSampler2, BufferUV).zw;
+#else
+	float2 FullScreenUV = BufferUV;
+#endif
+	return GetSubsurfaceProfileId(FullScreenUV);
+}
+
+float4 BurleyNormalizedSS(float2 BufferUV, uint2 GroupThreadID)
+{
+	FScreenSpaceData ScreenSpaceData = GetScreenSpaceDataFromCache(BufferUV);
+
+#ifdef SUBSURFACE_SINGLE_PASS
+	float DepthAtDiscCenter = CalcSceneDepth(BufferUV);
+	uint ReconstructMethod = (SUBSURFACE_RECOMBINE_QUALITY == SUBSURFACE_RECOMBINE_QUALITY_HIGH) ? 3 : 1;
+#else
+	float DepthAtDiscCenter = SampleSSSColorConsideringLocalShared(BufferUV, 0, GroupThreadID, 0).w;
+#endif
+
+	float4 OutColor = 0;
+
+	BRANCH if (DepthAtDiscCenter <= 0)
+	{
+		return OutColor;
+	}
+
+	FBurleyParameter BurleyParameter = GetBurleyParameters(ScreenSpaceData.GBuffer);
+	float DiffuseMeanFreePathForSampling = GetDiffuseMeanFreePathForSampling(BurleyParameter.DiffuseMeanFreePath);
+	float A = GetComponentForScalingFactorEstimation(BurleyParameter.SurfaceAlbedo);
+	uint SubsurfaceProfileInt = ExtractSubsurfaceProfileInt(ScreenSpaceData.GBuffer);
+	float3 BoundaryColorBleed = GetSubsurfaceProfileBoundaryColorBleed(ScreenSpaceData.GBuffer).xyz;
+
+#define LIGHT_PERPENDICULAR_DIFFUSE_SURFACE
+
+#ifdef LIGHT_PERPENDICULAR
+	float S = GetPerpendicularScalingFactor(A);
+	float3 S3D = GetPerpendicularScalingFactor3D(BurleyParameter.SurfaceAlbedo);
+#elif defined(LIGHT_DIFFUSESURFACE)
+	float S = GetDiffuseSurfaceScalingFactor(A);
+	float3 S3D = GetDiffuseSurfaceScalingFactor3D(BurleyParameter.SurfaceAlbedo);
+#elif defined(LIGHT_PERPENDICULAR_DIFFUSE_SURFACE)
+	float S = GetSearchLightDiffuseScalingFactor(A);
+	float3 S3D = GetSearchLightDiffuseScalingFactor(BurleyParameter.SurfaceAlbedo.rgb);
+#endif
+
+	int SeedStart = View.FrameNumber;
+	float3 WeightingFactor = 0.0f;
+	float4 RadianceAccumulated = float4(0.0f, 0.0f, 0.0f, 1.0f);
+	float Mask = 1.0f;
+
+#ifdef SUBSURFACE_SINGLE_PASS
+	int NumOfSamples = BURLEY_NUM_SAMPLES;
+	float InvNumOfSamples = BURLEY_INV_NUM_SAMPLES;
+#else
+	int NumOfSamples = GetNumOfSamplesBasedOnQuality(BufferUV);
+	float InvNumOfSamples = 1.0f / NumOfSamples;
+#endif
+
+	LOOP for (int i = 0; i < NumOfSamples; ++i)
+	{
+		// Step 1: sample generation
+		// Create an 2d disk sampling pattern (we can load from the disk as a texture or buffer).
+		int3  Seed3D = int3(BufferUV*SubsurfaceInput0_Extent, SeedStart++);
+		float2 Random0T1 = Generate2DRandomNumber(Seed3D);
+		FBurleySampleInfo BurleySampleInfo = GenerateSampleInfo(Random0T1, DiffuseMeanFreePathForSampling, S, i);
+
+		// Step 2: get the light radiance and depth at the offset
+		// and estimate the scale from the random disk sampling space to sceen space.
+		float YScale = BURLEY_MM_2_CM / DepthAtDiscCenter; //from mm to cm, and from world to screen
+		float XScale = YScale * SubsurfaceInput0_ExtentInverse.x*SubsurfaceInput0_Extent.y;
+
+		// World unit to screen space unit, we use world scale to control
+		// the scaling of the radius unit when mean free path is fixed.
+		float2 UVOffset = float2(XScale, YScale)*BurleySampleInfo.Radius*BurleyParameter.WorldScale;
+		UVOffset.x *= cos(BurleySampleInfo.Theta);
+		UVOffset.y *= sin(BurleySampleInfo.Theta);
+
+		// Sampling
+		{
+			float2 SampledDiscUV = BufferUV + UVOffset;
+
+#ifdef SUBSURFACE_SINGLE_PASS
+			SDiffuseAndSpecular SampledDiffuseAndSpecular = ReconstructLighting(SampledDiscUV, ReconstructMethod);
+			float4 SampledRadianceAndDepth = float4(SampledDiffuseAndSpecular.Diffuse, CalcSceneDepth(SampledDiscUV));
+#else
+
+			// Determine the miplevel with the expected number of samples at the pixel.
+			float MipLevel = 0.5*max(-log(BurleySampleInfo.Pdf*NumOfSamples / (BurleyParameter.WorldScale*BurleyParameter.WorldScale)), 0);
+
+			// If we are using half resolution,we should shift the mip level by -1
+#if SUBSURFACE_HALF_RES
+			MipLevel -= 1;
+#endif
+			MipLevel = clamp(MipLevel, 0, 5);
+
+			float4 SampledRadianceAndDepth = SampleSSSColorConsideringLocalShared(BufferUV, UVOffset, GroupThreadID, MipLevel);
+#endif
+
+			// Step 3: create the bilateral filtering weighted Distance between entry and exit.
+#define USE_BILATERAL_FILTERING
+
+#ifdef USE_BILATERAL_FILTERING
+			// 10.0/depth is the depth in world unit (cm) ???
+			float DeltaDepth = min(abs(10.0f / SampledRadianceAndDepth.w - 10.0f / DepthAtDiscCenter), 60000)*BURLEY_CM_2_MM;//in mm
+#ifdef SUBSURFACE_HALF_RES
+
+#endif
+			float RadiusSampled = sqrt(BurleySampleInfo.Radius * BurleySampleInfo.Radius + DeltaDepth * DeltaDepth); //mm.
+//#define RESAMPLE_PDF
+#ifdef RESAMPLE_PDF
+			BurleySampleInfo.Pdf = GetPdf(RadiusSampled, DiffuseMeanFreePathForSample, S);
+#endif
+#else
+			float RadiusSampled = BurleySampleInfo.Radius;
+#endif				
+
+			// Determine the tint color, if the sampling pixel is not subsurface, we use tint color
+			// to mask out the sampling. Unless we specifically want the shadowing region.
+			uint LocalProfileInt = GetSubsurfaceProfileIdFromCache(SampledDiscUV);
+			float3 TintColor = (LocalProfileInt == SubsurfaceProfileInt) ? 1.0f : BoundaryColorBleed;
+			Mask = UseSubsurfaceProfile(GetScreenSpaceDataFromCache(SampledDiscUV).GBuffer.ShadingModelID) ? 1 : 0;
+
+			// 1 considers the non subsurface neighbor (0,0,0) to add shadow.
+#if NONSUBSURFACE_SAMPLE_STRATEGY == 1
+			// The bilinear filtering in Separable looks better. So I added this to simulate the look
+			// With out this, the hair does not have shadow at the edge. 2 is an experimental value
+			RadiusSampled = (Mask == 0) ? BurleySampleInfo.Radius * 2 : RadiusSampled;
+			Mask = 1;
+#endif
+
+			// Step 4: accumulate radiance from the diffusion profile rR(r)
+			// make sure the DiffuseMeanFreePath is not zero and in m.
+			float3 DiffusionProfile = GetDiffuseReflectProfileWithDiffuseMeanFreePath(BurleyParameter.SurfaceAlbedo.xyz, BurleyParameter.DiffuseMeanFreePath.xyz, S3D.xyz, RadiusSampled);
+			RadianceAccumulated.xyz += (DiffusionProfile / BurleySampleInfo.Pdf) * (SampledRadianceAndDepth.xyz*TintColor*Mask);
+
+			// Use Normal and variance to the bilaterial filter.
+
+
+#ifdef USE_BILATERAL_FILTERING
+			WeightingFactor += (DiffusionProfile / BurleySampleInfo.Pdf)*Mask;
+#endif	
+		}
+
+	}
+
+	// 0.99995f is a compensitation to make it energe conservation.
+	const float EnergyNormalization = 1.0f / 0.99995f;
+
+#if defined(RADIUS_SAMPLE_UNIFORM_DISK)
+	RadianceAccumulated.xyz *= (InvSampleCount*0.5 * 2 * PI);
+#elif !defined(USE_BILATERAL_FILTERING)
+	RadianceAccumulated.xyz *= (InvNumOfSamples*0.5 * 2 * PI) * EnergyNormalization;
+#else
+	// The added epsilon is used to avoid divid by zero.
+	RadianceAccumulated.xyz *= 1.0f / (WeightingFactor + 0.000001f)*EnergyNormalization;
+#endif
+
+	OutColor.xyz = RadianceAccumulated.xyz;
+	OutColor.w = NumOfSamples;
+	return OutColor;
+}
+
+
+#define EXPONENTIAL_WEIGHT 0.2f
+float4 UpdateQualityVariance(float4 SubsurfaceColor, float2 BufferUV)
+{
+	float Weight = EXPONENTIAL_WEIGHT;
+	float WeightFinal = Weight;
+	float4 NewHistory = 0;
+	float2 VelocityOffset = float2(0.0, 0.0);
+	float2 NearestBufferUV = BufferUV;
+	float2 ViewportUV = (BufferUV - SubsurfaceInput0_UVViewportMin.xy)*SubsurfaceInput0_UVViewportSizeInverse.xy;
+
+	//	Get the history with reprojection
+	//  Implement a simplified version of temporal AA
+#define REPROJECTION
+#ifdef REPROJECTION
+
+	// Code adapted from temporal AA
+	float3 PosN;
+	PosN.xy = ViewportUVToScreenPos(ViewportUV);
+	PosN.z = SampleDepthTexture(BufferUV);// Direct sample without converting to world space;
+	float4 ThisClip = float4(PosN.xy, PosN.z, 1.0f);
+	float4 PrevClip = mul(ThisClip, View.ClipToPrevClip);
+	float2 PrevScreen = PrevClip.xy / PrevClip.w;
+	float2 BackN = PosN.xy - PrevScreen;
+
+	// Sample the velocity texture
+	float Velocity = 0;
+
+	float2 BackTemp = BackN * SubsurfaceInput1_ViewportSize.xy;
+
+#if defined(ENABLE_VELOCITY)
+	{
+		float2 VelocityN;
+		VelocityN = Texture2DSampleLevel(SubsurfaceInput2_Texture, SubsurfaceSampler2, NearestBufferUV + VelocityOffset, 0).xy;
+		
+		bool DynamicN = VelocityN.x > 0.0;
+
+		if (DynamicN)
+		{
+			BackN = DecodeVelocityFromTexture(VelocityN);
+		}
+
+		BackTemp = BackN * SubsurfaceInput1_ViewportSize.xy;
+	}
+#endif
+
+	// Update velocity
+	Velocity = sqrt(dot(BackTemp, BackTemp));
+
+	float2 HistoryScreenPosition = (PosN.xy - BackN);
+
+	// Detect if HistoryBufferUV would be outside of the viewport.
+	bool OffScreen = max(abs(HistoryScreenPosition.x), abs(HistoryScreenPosition.y)) >= 1.0;
+
+	float4 History = 0;
+
+	WeightFinal = lerp(Weight, 0.2, saturate(Velocity / 40));
+
+	BRANCH if (!OffScreen)
+	{
+		// ScreenPos to bufferUV
+		float2 HistoryUV = ScreenPosToViewportUV(HistoryScreenPosition);
+		// Convert history uv in viewport to buffer uv
+		float2 HistoryBufferUV = HistoryUV * SubsurfaceInput0_UVViewportSize.xy + SubsurfaceInput0_UVViewportMin.xy;
+		History = Texture2DSample(SubsurfaceInput1_Texture, SubsurfaceSampler1, HistoryBufferUV);
+	}
+	else
+	{
+		WeightFinal = 1;
+	}
+
+#else
+	float4 History = Texture2DSample(SubsurfaceInput1_Texture, SubsurfaceSampler1, BufferUV);
+#endif
+
+	// Calculate the illuminance
+	float NewSurfaceLuminance = Luminance(pow(SubsurfaceColor.rgb, 1 / 2.2));
+
+	NewHistory.a = (1 - WeightFinal) * History.a + WeightFinal * NewSurfaceLuminance;
+
+	float LumDelta = NewSurfaceLuminance - History.a;
+	NewHistory.b = (1 - WeightFinal)*History.b + WeightFinal * (1 - WeightFinal)*LumDelta*LumDelta;
+
+	// Update the random variable variance
+	NewHistory.g = History.b*SubsurfaceColor.a;
+
+	// Update the average number of samples used at each pixel (only for visualization purpose)
+	NewHistory.r = (1 - WeightFinal) * History.r + WeightFinal * SubsurfaceColor.a / BURLEY_NUM_SAMPLES;
+
+	BRANCH if (SubsurfaceColor.a <= 0)
+	{
+		NewHistory = 0;
+	}
+
+	return NewHistory;
+}
+
+
+#ifdef SUBSURFACE_BURLEY_COMPUTE
+
+// Compute shader common data and functions
+
+RWTexture2D<float4> SSSColorUAV;
+RWTexture2D<float4> HistoryUAV;
+
+void BurleyComputeMain(uint2 DT_ID, uint2 G_ID, uint GI)
+{
+
+	int2 TopLeftCorner = G_ID * SUBSURFACE_GROUP_SIZE - THREAD_TEXTURE_BOARDER + SubsurfaceInput0_ViewportMin;
+
+	// Step 1: Read the diffuse lighting into the local share. 64+64+16, three cycles
+#if SUBSURFACE_PASS == SUBSURFACE_PASS_ONE
+
+#ifndef TEXTURE_CACHE_DISABLED
+
+	// We do not need to use unroll if we have a fixed number of thread dimension.
+	UNROLL for (uint i = 0; i < NUM_OF_PASS_REQUIRED_FILL_SHARED_DIFFUSE_TEXTURE; i++)
+	{
+		// Calculate the sample uv for the current thread
+		uint LocalSharedIndex = GI * NUM_OF_PASS_REQUIRED_FILL_SHARED_DIFFUSE_TEXTURE + i;
+		BRANCH if (LocalSharedIndex >= LARGE_GROUP_TEXTURE_SHARE_TOTALSIZE)
+		{
+			break;
+		}
+
+		float2 SampleUV = SubsurfaceInput0_ExtentInverse * (TopLeftCorner + float2(LocalSharedIndex % LARGE_GROUP_TEXTURE_DIMENSION_SIZE,
+			LocalSharedIndex / LARGE_GROUP_TEXTURE_DIMENSION_SIZE) + 0.5f);
+		SampleUV = clamp(SampleUV, 0, 1);
+
+		float4 SubsurfaceColorWithSSSIndicator = Texture2DSample(SubsurfaceInput0_Texture, SubsurfaceSampler0, SampleUV);
+		SharedSubsurfaceDiffuseLighting[LocalSharedIndex] = SubsurfaceColorWithSSSIndicator;
+	}
+
+	GroupMemoryBarrierWithGroupSync();
+
+	// need to perform burley normalized subsurface scattering if it's burley
+	// this would unroll to just 1 sequence.
+	UNROLL for (int i = 0; i < LOCALGROUP_RATIO; ++i)
+		UNROLL for (int j = 0; j < LOCALGROUP_RATIO; ++j)
+			{
+				uint2 Pos = ConvertGroupIndexTo2DUsingMoltonCode(GI, TopLeftCorner + THREAD_TEXTURE_BOARDER + uint2(j, i)*THREAD_SIZE_1D);
+				float2 LBufferUV = ConvertGridPos2UV(Pos);
+				SSSColorUAV[Pos] = BurleyNormalizedSS(LBufferUV, Pos - (TopLeftCorner + THREAD_TEXTURE_BOARDER));
+			}
+
+#else
+	UNROLL for (int i = 0; i < LOCALGROUP_RATIO; ++i)
+		UNROLL for (int j = 0; j < LOCALGROUP_RATIO; ++j)
+			{
+
+				uint2 Pos = ConvertGroupIndexToNormal2DGrid(GI, TopLeftCorner + THREAD_TEXTURE_BOARDER + uint2(j, i)*THREAD_SIZE_1D);
+				float2 LBufferUV = ConvertGridPos2UV(Pos);
+				SSSColorUAV[Pos] = BurleyNormalizedSS(ConvertGridPos2UV(Pos), Pos - (TopLeftCorner + THREAD_TEXTURE_BOARDER));
+			}
+#endif
+#else
+
+	uint2 Pos = DT_ID.xy*SUBSURFACE_GROUP_SIZE / THREAD_SIZE_1D + SubsurfaceInput0_ViewportMin;
+
+	//we only update the variance if the subsurface is burley
+	UNROLL for (int i = 0; i < LOCALGROUP_RATIO; ++i)
+		UNROLL for (int j = 0; j < LOCALGROUP_RATIO; ++j)
+			{
+				uint2 TargetGridPos = Pos + uint2(i, j);
+				float2 LBufferUV = ConvertGridPos2UV(TargetGridPos);
+				// We perform quality variance thread-wise in the second pass
+				float4 SurfaceColor = Texture2DSample(SubsurfaceInput0_Texture, SubsurfaceSampler0, LBufferUV);
+				HistoryUAV[TargetGridPos] = UpdateQualityVariance(SurfaceColor, LBufferUV);
+				// Direct sample the color
+				SSSColorUAV[TargetGridPos] = float4(SurfaceColor.rgb, 1.0f);
+			}
+#endif
+}
+
+#endif
\ No newline at end of file
diff --git a/Engine/Shaders/Private/SubsurfaceProfileCommon.ush b/Engine/Shaders/Private/SubsurfaceProfileCommon.ush
--- a/Engine/Shaders/Private/SubsurfaceProfileCommon.ush
+++ b/Engine/Shaders/Private/SubsurfaceProfileCommon.ush
@@ -9,9 +9,9 @@
 // NOTE: Changing offsets below requires updating all instances of #SSSS_CONSTANTS
 // TODO: This needs to be defined in a single place and shared between C++ and shaders!
 #define SSSS_SUBSURFACE_COLOR_OFFSET			0
-#define SSS__SURFACEALBEDO_OFFSET                (SSSS_SUBSURFACE_COLOR_OFFSET+1)
-#define SSS__DMFP_OFFSET                        (SSS__SURFACEALBEDO_OFFSET+1)
-#define SSSS_TRANSMISSION_OFFSET				(SSS__DMFP_OFFSET+1)
+#define BSSS_SURFACEALBEDO_OFFSET                (SSSS_SUBSURFACE_COLOR_OFFSET+1)
+#define BSSS_DMFP_OFFSET                        (BSSS_SURFACEALBEDO_OFFSET+1)
+#define SSSS_TRANSMISSION_OFFSET				(BSSS_DMFP_OFFSET+1)
 #define SSSS_BOUNDARY_COLOR_BLEED_OFFSET		(SSSS_TRANSMISSION_OFFSET+1)
 #define SSSS_DUAL_SPECULAR_OFFSET				(SSSS_BOUNDARY_COLOR_BLEED_OFFSET+1)
 #define SSSS_KERNEL0_OFFSET						(SSSS_DUAL_SPECULAR_OFFSET+1)
@@ -23,6 +23,8 @@
 #define SSSS_KERNEL_TOTAL_SIZE					(SSSS_KERNEL0_SIZE + SSSS_KERNEL1_SIZE + SSSS_KERNEL2_SIZE)
 #define SSSS_TRANSMISSION_PROFILE_OFFSET		(SSSS_KERNEL0_OFFSET + SSSS_KERNEL_TOTAL_SIZE)
 #define SSSS_TRANSMISSION_PROFILE_SIZE			32
+#define BSSS_TRANSMISSION_PROFILE_OFFSET        (SSSS_TRANSMISSION_PROFILE_OFFSET + SSSS_TRANSMISSION_PROFILE_SIZE)
+#define BSSS_TRANSMISSION_PROFILE_SIZE			SSSS_TRANSMISSION_PROFILE_SIZE
 #define	SSSS_MAX_TRANSMISSION_PROFILE_DISTANCE	5.0f // See MaxTransmissionProfileDistance in ComputeTransmissionProfile(), SeparableSSS.cpp
 #define SSSS_MAX_DUAL_SPECULAR_ROUGHNESS		2.0f
 
diff --git a/Engine/Shaders/Private/TransmissionCommon.ush b/Engine/Shaders/Private/TransmissionCommon.ush
--- a/Engine/Shaders/Private/TransmissionCommon.ush
+++ b/Engine/Shaders/Private/TransmissionCommon.ush
@@ -5,6 +5,7 @@
 =========================================================================*/
 
 #include "SubsurfaceProfileCommon.ush"
+#include "BurleyNormalizedSSSCommon.ush"
 
 SamplerState TransmissionProfilesLinearSampler;
 
@@ -40,7 +41,8 @@ float4 GetTransmissionProfile(FGBufferData GBufferData, float Thickness)
 	uint SubsurfaceProfileInt = ExtractSubsurfaceProfileInt(GBufferData);
 	float3 TextureSize;
 	ActualSSProfilesTexture.GetDimensions(0, TextureSize.x, TextureSize.y,TextureSize.z); // TODO: it's slightly faster to pass texture dimensions through a constant buffer
-	float2 uv = float2((SSSS_TRANSMISSION_PROFILE_OFFSET + Thickness / SSSS_MAX_TRANSMISSION_PROFILE_DISTANCE * (SSSS_TRANSMISSION_PROFILE_SIZE-1) + 0.5f), SubsurfaceProfileInt + 0.5f) / (TextureSize.xy);
+	float ProfileOffset = GetSubsurfaceTransmittanceProfileUseBurley(SubsurfaceProfileInt)? BSSS_TRANSMISSION_PROFILE_OFFSET: SSSS_TRANSMISSION_PROFILE_OFFSET;
+	float2 uv = float2((ProfileOffset + Thickness / SSSS_MAX_TRANSMISSION_PROFILE_DISTANCE * (SSSS_TRANSMISSION_PROFILE_SIZE-1) + 0.5f), SubsurfaceProfileInt + 0.5f) / (TextureSize.xy);
 	return Texture2DSampleLevel(ActualSSProfilesTexture, TransmissionProfilesLinearSampler, uv, 0);
 }
 
diff --git a/Engine/Source/Runtime/Engine/Private/Rendering/SubsurfaceProfile.cpp b/Engine/Source/Runtime/Engine/Private/Rendering/SubsurfaceProfile.cpp
--- a/Engine/Source/Runtime/Engine/Private/Rendering/SubsurfaceProfile.cpp
+++ b/Engine/Source/Runtime/Engine/Private/Rendering/SubsurfaceProfile.cpp
@@ -127,9 +127,9 @@ static float GetNextSmallerPositiveFloat(float x)
 // NOTE: Changing offsets below requires updating all instances of #SSSS_CONSTANTS
 // TODO: This needs to be defined in a single place and shared between C++ and shaders!
 #define SSSS_SUBSURFACE_COLOR_OFFSET			0
-#define SSS__SURFACEALBEDO_OFFSET                (SSSS_SUBSURFACE_COLOR_OFFSET+1)
-#define SSS__DMFP_OFFSET                        (SSS__SURFACEALBEDO_OFFSET+1)
-#define SSSS_TRANSMISSION_OFFSET				(SSS__DMFP_OFFSET+1)
+#define BSSS_SURFACEALBEDO_OFFSET                (SSSS_SUBSURFACE_COLOR_OFFSET+1)
+#define BSSS_DMFP_OFFSET                        (BSSS_SURFACEALBEDO_OFFSET+1)
+#define SSSS_TRANSMISSION_OFFSET				(BSSS_DMFP_OFFSET+1)
 #define SSSS_BOUNDARY_COLOR_BLEED_OFFSET		(SSSS_TRANSMISSION_OFFSET+1)
 #define SSSS_DUAL_SPECULAR_OFFSET				(SSSS_BOUNDARY_COLOR_BLEED_OFFSET+1)
 #define SSSS_KERNEL0_OFFSET						(SSSS_DUAL_SPECULAR_OFFSET+1)
@@ -141,14 +141,16 @@ static float GetNextSmallerPositiveFloat(float x)
 #define SSSS_KERNEL_TOTAL_SIZE					(SSSS_KERNEL0_SIZE + SSSS_KERNEL1_SIZE + SSSS_KERNEL2_SIZE)
 #define SSSS_TRANSMISSION_PROFILE_OFFSET		(SSSS_KERNEL0_OFFSET + SSSS_KERNEL_TOTAL_SIZE)
 #define SSSS_TRANSMISSION_PROFILE_SIZE			32
+#define BSSS_TRANSMISSION_PROFILE_OFFSET        (SSSS_TRANSMISSION_PROFILE_OFFSET + SSSS_TRANSMISSION_PROFILE_SIZE)
+#define BSSS_TRANSMISSION_PROFILE_SIZE			SSSS_TRANSMISSION_PROFILE_SIZE
 #define	SSSS_MAX_TRANSMISSION_PROFILE_DISTANCE	5.0f // See MaxTransmissionProfileDistance in ComputeTransmissionProfile(), SeparableSSS.cpp
 #define	SSSS_MAX_DUAL_SPECULAR_ROUGHNESS		2.0f
 
 //------------------------------------------------------------------------------------------
 // Consistent in BurleyNormalizedSSSCommon.ush and SubsurfaceProfile.cpp
 
-#define SSS_TYPE_SSSS		1.0f
-#define SSS_TYPE_BURLEY	    0.0f
+#define SSS_TYPE_BURLEY	    0
+#define SSS_TYPE_SSSS		1
 
 //make sure UIMax|ClampMax of WorldScale * ENC_WORLDSCALE_TO_UNIT <= 1
 #define ENC_WORLDSCALE_IN_CM_TO_UNIT 0.02f
@@ -213,7 +215,7 @@ void FSubsurfaceProfileTexture::CreateTexture(FRHICommandListImmediate& RHICmdLi
 	const bool b16Bit = true;
 
 	// Each row of the texture contains SSS parameters, followed by 3 precomputed kernels. Texture must be wide enough to fit all data.
-	const uint32 Width = SSSS_TRANSMISSION_PROFILE_OFFSET + SSSS_TRANSMISSION_PROFILE_SIZE;
+	const uint32 Width = BSSS_TRANSMISSION_PROFILE_OFFSET + BSSS_TRANSMISSION_PROFILE_SIZE;
 
 	// at minimum 64 lines (less reallocations)
 	FPooledRenderTargetDesc Desc(FPooledRenderTargetDesc::Create2DDesc(FIntPoint(Width, FMath::Max(Height, (uint32)64)), PF_B8G8R8A8, FClearValueBinding::None, 0, TexCreate_None, false));
@@ -234,14 +236,6 @@ void FSubsurfaceProfileTexture::CreateTexture(FRHICommandListImmediate& RHICmdLi
 	const float FloatScale = GetNextSmallerPositiveFloat(0x10000);
 	check((int32)GetNextSmallerPositiveFloat(0x10000) == 0xffff);
 
-	const auto CVarSSSType = IConsoleManager::Get().FindTConsoleVariableDataInt(TEXT("r.SubsurfaceScattering.Type"));
-	check(CVarSSSType);
-	const bool bBurleyIsTurnedOn = CVarSSSType->GetValueOnAnyThread() == 1;
-
-	const auto CVarPassType = IConsoleManager::Get().FindTConsoleVariableDataInt(TEXT("r.SSS.Burley.PassType"));
-	check(CVarPassType);
-	const bool bSinglePass = CVarPassType->GetValueOnAnyThread() == 1;
-
 	for (uint32 y = 0; y < Height; ++y)
 	{
 		FSubsurfaceProfileStruct Data = SubsurfaceProfileEntries[y].Settings;
@@ -261,13 +255,12 @@ void FSubsurfaceProfileTexture::CreateTexture(FRHICommandListImmediate& RHICmdLi
 		
 		FLinearColor DifffuseMeanFreePath = Data.MeanFreePathColor*Data.MeanFreePathDistance*10.0f; // convert cm to mm.
 		SetupSurfaceAlbedoAndDiffuseMeanFreePath(Data.SurfaceAlbedo, DifffuseMeanFreePath);
-		TextureRow[SSS__SURFACEALBEDO_OFFSET] = Data.SurfaceAlbedo;
-		TextureRow[SSS__DMFP_OFFSET] = EncodeDiffuseMeanFreePath(DifffuseMeanFreePath);
+		TextureRow[BSSS_SURFACEALBEDO_OFFSET] = Data.SurfaceAlbedo;
+		TextureRow[BSSS_DMFP_OFFSET] = EncodeDiffuseMeanFreePath(DifffuseMeanFreePath);
 
 		TextureRow[SSSS_BOUNDARY_COLOR_BLEED_OFFSET] = Data.BoundaryColorBleed;
 
-		bool bUseBurley = bSinglePass ? bBurleyIsTurnedOn : (bBurleyIsTurnedOn && Data.bEnableBurley);
-		TextureRow[SSSS_BOUNDARY_COLOR_BLEED_OFFSET].A = bUseBurley ? SSS_TYPE_BURLEY : SSS_TYPE_SSSS;
+		TextureRow[SSSS_BOUNDARY_COLOR_BLEED_OFFSET].A = Data.bEnableBurley ? SSS_TYPE_BURLEY : SSS_TYPE_SSSS;
 
 		float MaterialRoughnessToAverage = Data.Roughness0 * (1.0f - Data.LobeMix) + Data.Roughness1 * Data.LobeMix;
 		float AverageToRoughness0 = Data.Roughness0 / MaterialRoughnessToAverage;
@@ -288,15 +281,11 @@ void FSubsurfaceProfileTexture::CreateTexture(FRHICommandListImmediate& RHICmdLi
 		ComputeMirroredSSSKernel(&TextureRow[SSSS_KERNEL1_OFFSET], SSSS_KERNEL1_SIZE, Data.SubsurfaceColor, Data.FalloffColor);
 		ComputeMirroredSSSKernel(&TextureRow[SSSS_KERNEL2_OFFSET], SSSS_KERNEL2_SIZE, Data.SubsurfaceColor, Data.FalloffColor);
 
-		if (bUseBurley)
-		{
-			ComputeTransmissionProfileBurley(&TextureRow[SSSS_TRANSMISSION_PROFILE_OFFSET], SSSS_TRANSMISSION_PROFILE_SIZE, 
-				Data.SubsurfaceColor, Data.FalloffColor, Data.ExtinctionScale,Data.SurfaceAlbedo, DifffuseMeanFreePath);
-		}
-		else
-		{
-			ComputeTransmissionProfile(&TextureRow[SSSS_TRANSMISSION_PROFILE_OFFSET], SSSS_TRANSMISSION_PROFILE_SIZE, Data.SubsurfaceColor, Data.FalloffColor, Data.ExtinctionScale);
-		}
+		ComputeTransmissionProfile(&TextureRow[SSSS_TRANSMISSION_PROFILE_OFFSET], SSSS_TRANSMISSION_PROFILE_SIZE, Data.SubsurfaceColor, Data.FalloffColor, Data.ExtinctionScale);
+
+		ComputeTransmissionProfileBurley(&TextureRow[BSSS_TRANSMISSION_PROFILE_OFFSET], BSSS_TRANSMISSION_PROFILE_SIZE,
+			Data.SubsurfaceColor, Data.FalloffColor, Data.ExtinctionScale, Data.SurfaceAlbedo, DifffuseMeanFreePath);
+
 		// could be lower than 1 (but higher than 0) to range compress for better quality (for 8 bit)
 		const float TableMaxRGB = 1.0f;
 		const float TableMaxA = 3.0f;
@@ -499,6 +488,6 @@ void USubsurfaceProfile::PostLoad()
 		MapFallOffColor2SurfaceAlbedoAndDiffuseMeanFreePath(Settings.FalloffColor.R, Settings.SurfaceAlbedo.R, Settings.MeanFreePathColor.R);
 		MapFallOffColor2SurfaceAlbedoAndDiffuseMeanFreePath(Settings.FalloffColor.G, Settings.SurfaceAlbedo.G, Settings.MeanFreePathColor.G);
 		MapFallOffColor2SurfaceAlbedoAndDiffuseMeanFreePath(Settings.FalloffColor.B, Settings.SurfaceAlbedo.B, Settings.MeanFreePathColor.B);
-		Settings.WorldScale = Settings.ScatterRadius;
+		Settings.WorldScale = Settings.ScatterRadius * 0.63; // Experimental value
 	}
 }
\ No newline at end of file
diff --git a/Engine/Source/Runtime/RenderCore/Private/GenerateMips.cpp b/Engine/Source/Runtime/RenderCore/Private/GenerateMips.cpp
--- a/Engine/Source/Runtime/RenderCore/Private/GenerateMips.cpp
+++ b/Engine/Source/Runtime/RenderCore/Private/GenerateMips.cpp
@@ -166,4 +166,52 @@ void FGenerateMips::Execute(FRHICommandListImmediate& RHICmdList, FRHITexture* I
 			RHICmdList.GenerateMips(InTexture);
 		}
 	}
-}
\ No newline at end of file
+}
+
+void FGenerateMips::Execute(FRDGBuilder* GraphBuilder, FRDGTextureRef InGraphTexture, FRHISamplerState* InSampler)
+{
+	check(IsInRenderingThread());
+	check(GraphBuilder);
+	check(InGraphTexture);
+	check(InSampler);
+
+	TShaderMapRef<FGenerateMipsCS> ComputeShader(GetGlobalShaderMap(ERHIFeatureLevel::SM5));
+	
+	//Loop through each level of the mips that require creation and add a dispatch pass per level,.
+	for (uint8 MipLevel = 1; MipLevel < InGraphTexture->Desc.NumMips; MipLevel++)
+	{
+		int DestTextureSizeX = InGraphTexture->Desc.Extent.X >> MipLevel;
+		int DestTextureSizeY = InGraphTexture->Desc.Extent.Y >> MipLevel;
+
+		//Create the RDG viewable SRV, of a complete Mip, to read from
+		FRDGTextureSRVDesc SRVDesc = FRDGTextureSRVDesc::CreateForMipLevel(InGraphTexture, MipLevel - 1);
+		//Create the RDG writeable UAV for the next mip level to be written to.
+		FRDGTextureUAVDesc UAVDesc(InGraphTexture, MipLevel);
+
+		FGenerateMipsCS::FParameters* PassParameters = GraphBuilder->AllocParameters<FGenerateMipsCS::FParameters>();
+		//Texel size is 1/the total length of a side.
+		PassParameters->TexelSize = FVector2D(1.0f / DestTextureSizeX, 1.0f / DestTextureSizeY);
+		PassParameters->MipInSRV = GraphBuilder->CreateSRV(SRVDesc);
+		PassParameters->MipOutUAV = GraphBuilder->CreateUAV(UAVDesc);
+		PassParameters->MipSampler = InSampler;
+
+		//Dispatch count is the destination's mip texture dimensions, so only the number required is executed.
+		FIntVector GenMipsGroupCount(
+			FMath::Max((DestTextureSizeX + MIPSSHADER_NUMTHREADS -1)/ MIPSSHADER_NUMTHREADS, 1),
+			FMath::Max((DestTextureSizeY + MIPSSHADER_NUMTHREADS -1)/ MIPSSHADER_NUMTHREADS, 1),
+			1);
+		
+		//Pass added per mip level to be written.
+		ClearUnusedGraphResources(*ComputeShader, PassParameters);
+
+		GraphBuilder->AddPass(
+			RDG_EVENT_NAME("Generate2DTextureMips DestMipLevel=%d", MipLevel),
+			PassParameters,
+			ERDGPassFlags::Compute | ERDGPassFlags::GenerateMips,
+			[PassParameters, ComputeShader, GenMipsGroupCount](FRHICommandList& RHICmdList)
+		{
+			FComputeShaderUtils::Dispatch(RHICmdList, *ComputeShader, *PassParameters, GenMipsGroupCount);
+		});
+
+	}
+}
diff --git a/Engine/Source/Runtime/RenderCore/Public/GenerateMips.h b/Engine/Source/Runtime/RenderCore/Public/GenerateMips.h
--- a/Engine/Source/Runtime/RenderCore/Public/GenerateMips.h
+++ b/Engine/Source/Runtime/RenderCore/Public/GenerateMips.h
@@ -22,6 +22,7 @@ public:
 	static void Execute(FRHICommandListImmediate& RHICmdList, FRHITexture* InTexture,
 		const FGenerateMipsParams& InParams = FGenerateMipsParams());
 
+	static void Execute(class FRDGBuilder* GraphBuilder, class FRDGTexture* InGraphTexture, FRHISamplerState* InSampler);
 private:
 	static void Compute(FRHICommandListImmediate& RHIImmCmdList, FRHITexture* InTexture);
 	static FGenerateMipsStruct* SetupTexture(FRHITexture* InTexture,
diff --git a/Engine/Source/Runtime/Renderer/Private/PostProcess/PostProcessSubsurface.cpp b/Engine/Source/Runtime/Renderer/Private/PostProcess/PostProcessSubsurface.cpp
--- a/Engine/Source/Runtime/Renderer/Private/PostProcess/PostProcessSubsurface.cpp
+++ b/Engine/Source/Runtime/Renderer/Private/PostProcess/PostProcessSubsurface.cpp
@@ -2,39 +2,242 @@
 
 /*=============================================================================
 	PostProcessSubsurface.cpp: Screenspace subsurface scattering implementation.
+		Indirect dispatch implementation high level description
+	   1. Initialize counters
+	   2. Setup pass: record the tiles that need to draw Burley and Separable in two different buffer.
+	   3. Indirect dispatch Burley.
+	   4. Indirect dispatch Separable.
+	   5. Recombine.
 =============================================================================*/
 
 #include "PostProcess/PostProcessSubsurface.h"
 #include "PostProcess/SceneRenderTargets.h"
-#include "PostProcess/SubsurfaceCommon.h"
-#include "PostProcess/SubsurfaceBurleyNormalized.h"
-#include "PostProcess/SubsurfaceSeparable.h"
 #include "Engine/SubsurfaceProfile.h"
 #include "CanvasTypes.h"
 #include "RenderTargetTemp.h"
+#include "SystemTextures.h"
+#include "GenerateMips.h"
+#include "ClearQuad.h"
+
 
 namespace
 {
-	//only add ubber subsurface console variables here
-	static TAutoConsoleVariable<int32> CVarSubsurfaceScatteringType(
-		TEXT("r.SubsurfaceScattering.Type"),
+
+	// Subsurface common parameters
+	TAutoConsoleVariable<int32> CVarSubsurfaceScattering(
+		TEXT("r.SubsurfaceScattering"),
+		1,
+		TEXT(" 0: disabled\n")
+		TEXT(" 1: enabled (default)"),
+		ECVF_RenderThreadSafe | ECVF_Scalability);
+
+	TAutoConsoleVariable<float> CVarSSSScale(
+		TEXT("r.SSS.Scale"),
+		1.0f,
+		TEXT("Affects the Screen space Separable subsurface scattering pass ")
+		TEXT("(use shadingmodel SubsurfaceProfile, get near to the object as the default)\n")
+		TEXT("is human skin which only scatters about 1.2cm)\n")
+		TEXT(" 0: off (if there is no object on the screen using this pass it should automatically disable the post process pass)\n")
+		TEXT("<1: scale scatter radius down (for testing)\n")
+		TEXT(" 1: use given radius form the Subsurface scattering asset (default)\n")
+		TEXT(">1: scale scatter radius up (for testing)"),
+		ECVF_Scalability | ECVF_RenderThreadSafe);
+
+	TAutoConsoleVariable<int32> CVarSSSHalfRes(
+		TEXT("r.SSS.HalfRes"),
+		1,
+		TEXT(" 0: full quality (not optimized, as reference)\n")
+		TEXT(" 1: parts of the algorithm runs in half resolution which is lower quality but faster (default)"),
+		ECVF_RenderThreadSafe | ECVF_Scalability);
+
+	TAutoConsoleVariable<int32> CVarSSSQuality(
+		TEXT("r.SSS.Quality"),
 		0,
-		TEXT("Subsurface type:\n")
-		TEXT(" 0. Separable SSS (default).\n")
-		TEXT(" 1. Burley Normalized.")
-		TEXT("Else. Separable SSS. Make sure we always have SSS"),
-		ECVF_Scalability | ECVF_RenderThreadSafe
+		TEXT("Defines the quality of the recombine pass when using the SubsurfaceScatteringProfile shading model\n")
+		TEXT(" 0: low (faster, default)\n")
+		TEXT(" 1: high (sharper details but slower)\n")
+		TEXT("-1: auto, 1 if TemporalAA is disabled (without TemporalAA the quality is more noticable)"),
+		ECVF_RenderThreadSafe | ECVF_Scalability);
+
+	TAutoConsoleVariable<int32> CVarSSSFilter(
+		TEXT("r.SSS.Filter"),
+		1,
+		TEXT("Defines the filter method for Screenspace Subsurface Scattering feature.\n")
+		TEXT(" 0: point filter (useful for testing, could be cleaner)\n")
+		TEXT(" 1: bilinear filter"),
+		ECVF_RenderThreadSafe | ECVF_Scalability);
+
+	TAutoConsoleVariable<int32> CVarSSSSampleSet(
+		TEXT("r.SSS.SampleSet"),
+		2,
+		TEXT("Defines how many samples we use for Separable Screenspace Subsurface Scattering feature.\n")
+		TEXT(" 0: lowest quality (6*2+1)\n")
+		TEXT(" 1: medium quality (9*2+1)\n")
+		TEXT(" 2: high quality (13*2+1) (default)"),
+		ECVF_RenderThreadSafe | ECVF_Scalability);
+
+	TAutoConsoleVariable<int32> CVarSSSBurleyUpdateParameter(
+		TEXT("r.SSS.Burley.AlwaysUpdateParametersFromSeparable"),
+		1,
+		TEXT("0: Will not update parameters when the program loads.")
+		TEXT("1: Always update from the separable when the program loads. (default)"),
+		ECVF_RenderThreadSafe | ECVF_Scalability
 	);
+
+	TAutoConsoleVariable<int32> CVarSSSCheckerboard(
+		TEXT("r.SSS.Checkerboard"),
+		2,
+		TEXT("Enables or disables checkerboard rendering for subsurface profile rendering.\n")
+		TEXT("This is necessary if SceneColor does not include a floating point alpha channel (e.g 32-bit formats)\n")
+		TEXT(" 0: Disabled (high quality) \n")
+		TEXT(" 1: Enabled (low quality). Surface lighting will be at reduced resolution.\n")
+		TEXT(" 2: Automatic. Non-checkerboard lighting will be applied if we have a suitable rendertarget format\n"),
+		ECVF_RenderThreadSafe);
+}
+
+// Define the size of subsurface group. Set to 16 to use LDS.
+#define SUBSURFACE_GROUP_SIZE        8
+
+// Define to use a custom ps to clear UAV.
+#define USE_CUSTOM_ClearUAV
+
+ENGINE_API IPooledRenderTarget* GetSubsufaceProfileTexture_RT(FRHICommandListImmediate& RHICmdList);
+
+enum class ESubsurfaceMode : uint32
+{
+	// Performs a full resolution scattering filter.
+	FullRes,
+
+	// Performs a half resolution scattering filter.
+	HalfRes,
+
+	// Reconstructs lighting, but does not perform scattering.
+	Bypass,
+
+	MAX
+};
+
+// Returns the [0, N] clamped value of the 'r.SSS.Scale' CVar.
+float GetSubsurfaceRadiusScale()
+{
+	static const auto CVar = IConsoleManager::Get().FindTConsoleVariableDataFloat(TEXT("r.SSS.Scale"));
+	check(CVar);
+
+	return FMath::Max(0.0f, CVar->GetValueOnRenderThread());
+}
+
+int32 GetSSSFilter()
+{
+	return CVarSSSFilter.GetValueOnRenderThread();
+}
+
+int32 GetSSSSampleSet()
+{
+	return CVarSSSSampleSet.GetValueOnRenderThread();
+}
+
+int32 GetSSSQuality()
+{
+	return CVarSSSQuality.GetValueOnRenderThread();
+}
+
+// Returns the SS profile texture with a black fallback texture if none exists yet.
+// Actually we do not need this for the burley normalized SSS.
+FRHITexture* GetSubsurfaceProfileTexture(FRHICommandListImmediate& RHICmdList)
+{
+	const IPooledRenderTarget* ProfileTextureTarget = GetSubsufaceProfileTexture_RT(RHICmdList);
+
+	if (!ProfileTextureTarget)
+	{
+		// No subsurface profile was used yet
+		ProfileTextureTarget = GSystemTextures.BlackDummy;
+	}
+
+	return ProfileTextureTarget->GetRenderTargetItem().ShaderResourceTexture;
+}
+
+// Returns the current subsurface mode required by the current view.
+ESubsurfaceMode GetSubsurfaceModeForView(const FViewInfo& View)
+{
+	const float Radius = GetSubsurfaceRadiusScale();
+	const bool bShowSubsurfaceScattering = Radius > 0 && View.Family->EngineShowFlags.SubsurfaceScattering;
+
+	if (bShowSubsurfaceScattering)
+	{
+		const bool bHalfRes = CVarSSSHalfRes.GetValueOnRenderThread() != 0;
+		if (bHalfRes)
+		{
+			return ESubsurfaceMode::HalfRes;
+		}
+		else
+		{
+			return ESubsurfaceMode::FullRes;
+		}
+	}
+	else
+	{
+		return ESubsurfaceMode::Bypass;
+	}
+}
+
+// A shader parameter struct for a single subsurface input texture.
+BEGIN_SHADER_PARAMETER_STRUCT(FSubsurfaceInput, )
+SHADER_PARAMETER_STRUCT_INCLUDE(FScreenPassTextureViewportParameters, Viewport)
+SHADER_PARAMETER_RDG_TEXTURE(Texture2D, Texture)
+END_SHADER_PARAMETER_STRUCT()
+
+BEGIN_SHADER_PARAMETER_STRUCT(FSubsurfaceSRVInput, )
+SHADER_PARAMETER_STRUCT_INCLUDE(FScreenPassTextureViewportParameters, Viewport)
+SHADER_PARAMETER_RDG_TEXTURE_SRV(Texture2D, Texture)
+END_SHADER_PARAMETER_STRUCT();
+
+// Set of common shader parameters shared by all subsurface shaders.
+BEGIN_SHADER_PARAMETER_STRUCT(FSubsurfaceParameters, )
+SHADER_PARAMETER(FVector4, SubsurfaceParams)
+SHADER_PARAMETER_STRUCT_REF(FSceneTexturesUniformParameters, SceneUniformBuffer)
+SHADER_PARAMETER_STRUCT_REF(FViewUniformShaderParameters, ViewUniformBuffer)
+SHADER_PARAMETER_SAMPLER(SamplerState, BilinearTextureSampler)
+SHADER_PARAMETER_TEXTURE(Texture2D, SSProfilesTexture)
+END_SHADER_PARAMETER_STRUCT()
+
+FSubsurfaceParameters GetSubsurfaceCommonParameters(FRHICommandListImmediate& RHICmdList, const FViewInfo& View)
+{
+	FSceneRenderTargets& SceneContext = FSceneRenderTargets::Get(RHICmdList);
+
+	const float DistanceToProjectionWindow = View.ViewMatrices.GetProjectionMatrix().M[0][0];
+	const float SSSScaleZ = DistanceToProjectionWindow * GetSubsurfaceRadiusScale();
+	const float SSSScaleX = SSSScaleZ / SUBSURFACE_KERNEL_SIZE * 0.5f;
+
+	FSubsurfaceParameters Parameters;
+	Parameters.SubsurfaceParams = FVector4(SSSScaleX, SSSScaleZ, 0, 0);
+	Parameters.ViewUniformBuffer = View.ViewUniformBuffer;
+	Parameters.SceneUniformBuffer = CreateSceneTextureUniformBuffer(
+		SceneContext, View.FeatureLevel, ESceneTextureSetupMode::All, EUniformBufferUsage::UniformBuffer_SingleFrame);
+	Parameters.BilinearTextureSampler = TStaticSamplerState<SF_Bilinear>::GetRHI();
+	Parameters.SSProfilesTexture = GetSubsurfaceProfileTexture(RHICmdList);
+	return Parameters;
+}
+
+FSubsurfaceInput GetSubsurfaceInput(FRDGTextureRef Texture, const FScreenPassTextureViewportParameters& ViewportParameters)
+{
+	FSubsurfaceInput Input;
+	Input.Texture = Texture;
+	Input.Viewport = ViewportParameters;
+	return Input;
+}
+
+FSubsurfaceSRVInput GetSubsurfaceSRVInput(FRDGTextureSRVRef Texture, const FScreenPassTextureViewportParameters& ViewportParameters)
+{
+	FSubsurfaceSRVInput Input;
+	Input.Texture = Texture;
+	Input.Viewport = ViewportParameters;
+	return Input;
 }
 
 bool IsSubsurfaceEnabled()
 {
-	static const auto CVarSubsurfaceScattering = IConsoleManager::Get().FindTConsoleVariableDataInt(TEXT("r.SubsurfaceScattering"));
-	static const auto CVarSSSScale = IConsoleManager::Get().FindTConsoleVariableDataFloat(TEXT("r.SSS.Scale"));
-	check(CVarSubsurfaceScattering);
-	check(CVarSSSScale);
-	const bool bEnabled = CVarSubsurfaceScattering->GetValueOnAnyThread() != 0;
-	const bool bHasScale = CVarSSSScale->GetValueOnAnyThread() > 0.0f;
+	const bool bEnabled = CVarSubsurfaceScattering.GetValueOnAnyThread() != 0;
+	const bool bHasScale = CVarSSSScale.GetValueOnAnyThread() > 0.0f;
 	return (bEnabled && bHasScale);
 }
 
@@ -69,12 +272,52 @@ uint32 GetSubsurfaceRequiredViewMask(const TArray<FViewInfo>& Views)
 
 bool IsSubsurfaceCheckerboardFormat(EPixelFormat SceneColorFormat)
 {
-	// No checkerboard for Burley
-	if (CVarSubsurfaceScatteringType.GetValueOnRenderThread() == 1)
+	int CVarValue = CVarSSSCheckerboard.GetValueOnRenderThread();
+	if (CVarValue == 0)
+	{
 		return false;
-	return IsSeparableSubsurfaceCheckerboardFormat(SceneColorFormat);
+	}
+	else if (CVarValue == 1)
+	{
+		return true;
+	}
+	else if (CVarValue == 2)
+	{
+		switch (SceneColorFormat)
+		{
+		case PF_A32B32G32R32F:
+		case PF_FloatRGBA:
+			return false;
+		default:
+			return true;
+		}
+	}
+	return true;
 }
 
+
+// Base class for a subsurface shader.
+class FSubsurfaceShader : public FGlobalShader
+{
+public:
+	static bool ShouldCompilePermutation(const FGlobalShaderPermutationParameters& Parameters)
+	{
+		return IsFeatureLevelSupported(Parameters.Platform, ERHIFeatureLevel::SM5);
+	}
+
+	static void ModifyCompilationEnvironment(const FGlobalShaderPermutationParameters& Parameters, FShaderCompilerEnvironment& OutEnvironment)
+	{
+		FGlobalShader::ModifyCompilationEnvironment(Parameters, OutEnvironment);
+		OutEnvironment.SetDefine(TEXT("SUBSURFACE_RADIUS_SCALE"), SUBSURFACE_RADIUS_SCALE);
+		OutEnvironment.SetDefine(TEXT("SUBSURFACE_KERNEL_SIZE"), SUBSURFACE_KERNEL_SIZE);
+	}
+
+	FSubsurfaceShader() = default;
+	FSubsurfaceShader(const ShaderMetaType::CompiledShaderInitializerType& Initializer)
+		: FGlobalShader(Initializer)
+	{}
+};
+
 // Encapsulates the post processing subsurface scattering common pixel shader.
 class FSubsurfaceVisualizePS : public FSubsurfaceShader
 {
@@ -113,6 +356,397 @@ class FSubsurfaceViewportCopyPS : public FSubsurfaceShader
 
 IMPLEMENT_GLOBAL_SHADER(FSubsurfaceViewportCopyPS, "/Engine/Private/PostProcessSubsurface.usf", "SubsurfaceViewportCopyPS", SF_Pixel);
 
+
+//-------------------------------------------------------------------------------------------
+// Indirect dispatch class and functions
+//-------------------------------------------------------------------------------------------
+
+struct FSubsurfaceGroupItem
+{
+	FVector2D GroupXY;		// which tile. 
+	int32 Count;			// The first element is used as the counter.
+	int32 Dummy;			// alignment
+};
+
+// Subsurface uniform buffer layout
+BEGIN_GLOBAL_SHADER_PARAMETER_STRUCT(FSubsurfaceUniformParameters, )
+SHADER_PARAMETER(int32, MaxGroupCount)
+END_GLOBAL_SHADER_PARAMETER_STRUCT()
+
+IMPLEMENT_GLOBAL_SHADER_PARAMETER_STRUCT(FSubsurfaceUniformParameters, "SubsurfaceUniformParameters");
+typedef TUniformBufferRef<FSubsurfaceUniformParameters> FSubsurfaceUniformRef;
+
+// Return a uniform buffer with values filled and with single frame lifetime
+FSubsurfaceUniformRef CreateUniformBuffer(FViewInfo const& View, int32 MaxGroupCount)
+{
+	FSubsurfaceUniformParameters Parameters;
+	Parameters.MaxGroupCount = MaxGroupCount;
+	return FSubsurfaceUniformRef::CreateUniformBufferImmediate(Parameters, UniformBuffer_SingleFrame);
+}
+
+class FSubsurfaceInitValueBufferCS : public FSubsurfaceShader
+{
+public:
+	DECLARE_GLOBAL_SHADER(FSubsurfaceInitValueBufferCS);
+	SHADER_USE_PARAMETER_STRUCT(FSubsurfaceInitValueBufferCS, FSubsurfaceShader);
+
+	static void ModifyCompilationEnvironment(const FGlobalShaderPermutationParameters& Parameters, FShaderCompilerEnvironment& OutEnvironment)
+	{
+		FGlobalShader::ModifyCompilationEnvironment(Parameters, OutEnvironment);
+		OutEnvironment.SetDefine(TEXT("SUBSURFACE_RADIUS_SCALE"), SUBSURFACE_RADIUS_SCALE);
+		OutEnvironment.SetDefine(TEXT("SUBSURFACE_KERNEL_SIZE"), SUBSURFACE_KERNEL_SIZE);
+		OutEnvironment.SetDefine(TEXT("SUBSURFACE_BURLEY_COMPUTE"), 1);
+	}
+
+	BEGIN_SHADER_PARAMETER_STRUCT(FParameters, )
+		SHADER_PARAMETER_RDG_BUFFER_UAV(RWStructuredBuffer<FSubsurfaceGroupItem>, RWSeparableGroupBuffer)
+		SHADER_PARAMETER_RDG_BUFFER_UAV(RWStructuredBuffer<FSubsurfaceGroupItem>, RWBurleyGroupBuffer)
+		END_SHADER_PARAMETER_STRUCT()
+
+		static bool ShouldCompilePermutation(const FGlobalShaderPermutationParameters& Parameters)
+	{
+		return IsFeatureLevelSupported(Parameters.Platform, ERHIFeatureLevel::SM5);
+	}
+};
+
+IMPLEMENT_GLOBAL_SHADER(FSubsurfaceInitValueBufferCS, "/Engine/Private/PostProcessSubsurface.usf", "InitValueBufferCS", SF_Compute);
+
+class FSubsurfaceBuildIndirectDispatchArgsCS : public FSubsurfaceShader
+{
+public:
+	DECLARE_GLOBAL_SHADER(FSubsurfaceBuildIndirectDispatchArgsCS);
+	SHADER_USE_PARAMETER_STRUCT(FSubsurfaceBuildIndirectDispatchArgsCS, FSubsurfaceShader)
+
+		static void ModifyCompilationEnvironment(const FGlobalShaderPermutationParameters& Parameters, FShaderCompilerEnvironment& OutEnvironment)
+	{
+		FGlobalShader::ModifyCompilationEnvironment(Parameters, OutEnvironment);
+		OutEnvironment.SetDefine(TEXT("SUBSURFACE_RADIUS_SCALE"), SUBSURFACE_RADIUS_SCALE);
+		OutEnvironment.SetDefine(TEXT("SUBSURFACE_KERNEL_SIZE"), SUBSURFACE_KERNEL_SIZE);
+		OutEnvironment.SetDefine(TEXT("SUBSURFACE_BURLEY_COMPUTE"), 1);
+	}
+
+	BEGIN_SHADER_PARAMETER_STRUCT(FParameters, )
+		SHADER_PARAMETER_STRUCT_REF(FSubsurfaceUniformParameters, SubsurfaceUniformParameters)
+		SHADER_PARAMETER_RDG_BUFFER_UAV(RWStructuredBuffer<uint>, RWIndirectDispatchArgsBuffer)
+		SHADER_PARAMETER_RDG_BUFFER_SRV(StructuredBuffer<FSubsurfaceGroupItem>, GroupBuffer)
+		END_SHADER_PARAMETER_STRUCT()
+
+		static bool ShouldCompilePermutation(const FGlobalShaderPermutationParameters& Parameters)
+	{
+		return IsFeatureLevelSupported(Parameters.Platform, ERHIFeatureLevel::SM5);
+	}
+};
+
+IMPLEMENT_GLOBAL_SHADER(FSubsurfaceBuildIndirectDispatchArgsCS, "/Engine/Private/PostProcessSubsurface.usf", "BuildIndirectDispatchArgsCS", SF_Compute);
+
+class FSubsurfaceIndirectDispatchSetupCS : public FSubsurfaceShader
+{
+public:
+	DECLARE_GLOBAL_SHADER(FSubsurfaceIndirectDispatchSetupCS);
+	SHADER_USE_PARAMETER_STRUCT(FSubsurfaceIndirectDispatchSetupCS, FSubsurfaceShader)
+
+		static void ModifyCompilationEnvironment(const FGlobalShaderPermutationParameters& Parameters, FShaderCompilerEnvironment& OutEnvironment)
+	{
+		FGlobalShader::ModifyCompilationEnvironment(Parameters, OutEnvironment);
+		OutEnvironment.SetDefine(TEXT("SUBSURFACE_RADIUS_SCALE"), SUBSURFACE_RADIUS_SCALE);
+		OutEnvironment.SetDefine(TEXT("SUBSURFACE_KERNEL_SIZE"), SUBSURFACE_KERNEL_SIZE);
+		OutEnvironment.SetDefine(TEXT("SUBSURFACE_BURLEY_COMPUTE"), 1);
+	}
+
+	BEGIN_SHADER_PARAMETER_STRUCT(FParameters, )
+		SHADER_PARAMETER_STRUCT_INCLUDE(FSubsurfaceParameters, Subsurface)
+		SHADER_PARAMETER_STRUCT(FSubsurfaceInput, SubsurfaceInput0)
+		SHADER_PARAMETER_SAMPLER(SamplerState, SubsurfaceSampler0)
+		SHADER_PARAMETER_RDG_TEXTURE_UAV(RWTexture2D, SetupTexture)
+		SHADER_PARAMETER_RDG_TEXTURE_UAV(RWTexture2D, ProfileEdgeMask)
+		SHADER_PARAMETER_RDG_BUFFER_UAV(RWStructuredBuffer<FSubsurfaceGroupItem>, RWSeparableGroupBuffer)
+		SHADER_PARAMETER_RDG_BUFFER_UAV(RWStructuredBuffer<FSubsurfaceGroupItem>, RWBurleyGroupBuffer)
+		SHADER_PARAMETER_STRUCT_REF(FSubsurfaceUniformParameters, SubsurfaceUniformParameters)
+		END_SHADER_PARAMETER_STRUCT()
+
+		class FDimensionHalfRes : SHADER_PERMUTATION_BOOL("SUBSURFACE_HALF_RES");
+	class FDimensionCheckerboard : SHADER_PERMUTATION_BOOL("SUBSURFACE_PROFILE_CHECKERBOARD");
+	using FPermutationDomain = TShaderPermutationDomain<FDimensionHalfRes, FDimensionCheckerboard>;
+};
+
+IMPLEMENT_GLOBAL_SHADER(FSubsurfaceIndirectDispatchSetupCS, "/Engine/Private/PostProcessSubsurface.usf", "SetupIndirectCS", SF_Compute);
+
+class FSubsurfaceIndirectDispatchCS : public FSubsurfaceShader
+{
+public:
+	DECLARE_GLOBAL_SHADER(FSubsurfaceIndirectDispatchCS);
+	SHADER_USE_PARAMETER_STRUCT(FSubsurfaceIndirectDispatchCS, FSubsurfaceShader);
+
+	static void ModifyCompilationEnvironment(const FGlobalShaderPermutationParameters& Parameters, FShaderCompilerEnvironment& OutEnvironment)
+	{
+		FGlobalShader::ModifyCompilationEnvironment(Parameters, OutEnvironment);
+		OutEnvironment.SetDefine(TEXT("SUBSURFACE_RADIUS_SCALE"), SUBSURFACE_RADIUS_SCALE);
+		OutEnvironment.SetDefine(TEXT("SUBSURFACE_KERNEL_SIZE"), SUBSURFACE_KERNEL_SIZE);
+		OutEnvironment.SetDefine(TEXT("SUBSURFACE_BURLEY_COMPUTE"), 1);
+		OutEnvironment.SetDefine(TEXT("ENABLE_VELOCITY"), 1);
+		OutEnvironment.SetDefine(TEXT("SUBSURFACE_GROUP_SIZE"), SUBSURFACE_GROUP_SIZE);
+	}
+
+	BEGIN_SHADER_PARAMETER_STRUCT(FParameters, )
+		SHADER_PARAMETER_STRUCT_INCLUDE(FSubsurfaceParameters, Subsurface)
+		SHADER_PARAMETER_RDG_TEXTURE_UAV(RWTexture2D, SSSColorUAV)
+		SHADER_PARAMETER_RDG_TEXTURE_UAV(RWTexture2D, HistoryUAV)
+		SHADER_PARAMETER_RDG_BUFFER_SRV(StructuredBuffer<FSubsurfaceGroupItem>, GroupBuffer)
+		SHADER_PARAMETER_RDG_BUFFER(StructuredBuffer<uint>, IndirectDispatchArgsBuffer)
+		SHADER_PARAMETER_STRUCT(FSubsurfaceInput, SubsurfaceInput0)
+		SHADER_PARAMETER_SAMPLER(SamplerState, SubsurfaceSampler0)
+		SHADER_PARAMETER_STRUCT(FSubsurfaceInput, SubsurfaceInput1)	// History
+		SHADER_PARAMETER_SAMPLER(SamplerState, SubsurfaceSampler1)
+		SHADER_PARAMETER_STRUCT(FSubsurfaceInput, SubsurfaceInput2)   // Profile mask | Velocity
+		SHADER_PARAMETER_SAMPLER(SamplerState, SubsurfaceSampler2)
+		END_SHADER_PARAMETER_STRUCT()
+
+		// Direction of the 1D separable filter.
+		enum class EDirection : uint32
+	{
+		Horizontal,
+		Vertical,
+		MAX
+	};
+
+	enum class ESubsurfacePass : uint32
+	{
+		PassOne,    // Burley sampling   (or Horizontal) pass   pass one
+		PassTwo,	// Variance updating (or   Vertical) pass pass two
+		MAX
+	};
+
+	// Controls the quality (number of samples) of the blur kernel.
+	enum class EQuality : uint32
+	{
+		Low,
+		Medium,
+		High,
+		MAX
+	};
+
+	// Need this to make hair denser in Burley to simulate the bilinear effect in Separable.
+	// In bilinear setting, Separable samples horizontally and vertically sequentially on the same texture
+	// This will blur the edge between non-subsurface scattering and subsurface scattering. For hair region
+	// there are many small edges interleaving. Thus Bilinear makes the hair dense with small amount of hairs.
+	// Since Burley only samples in one pass, it yields the correct result but the hair looks less dense.
+	// So, for better appearance quality with less hair, we need to consider the non-subsurface samples, which
+	// is set dark in setup pass, in the weighting function. To achieve a denser look when bilinear is enabled,
+	// we consider the non-subsurface dark samples.
+	enum class ENonSubsurfaceSampleStrategy : uint32
+	{
+		Ignore,
+		Considered,
+		MAX
+	};
+
+	enum class ESubsurfaceType : uint32
+	{
+		BURLEY,
+		SEPARABLE,
+		MAX
+	};
+
+	class FSubsurfacePassFunction : SHADER_PERMUTATION_ENUM_CLASS("SUBSURFACE_PASS", ESubsurfacePass);
+	class FDimensionQuality : SHADER_PERMUTATION_ENUM_CLASS("SUBSURFACE_QUALITY", EQuality);
+	class FNonSubsurfaceSamplerStrategy : SHADER_PERMUTATION_ENUM_CLASS("NONSUBSURFACE_SAMPLE_STRATEGY", ENonSubsurfaceSampleStrategy);
+	class FSubsurfaceType : SHADER_PERMUTATION_ENUM_CLASS("SUBSURFACE_TYPE", ESubsurfaceType);
+	class FDimensionHalfRes : SHADER_PERMUTATION_BOOL("SUBSURFACE_HALF_RES");
+	using FPermutationDomain = TShaderPermutationDomain<FSubsurfacePassFunction, FDimensionQuality,
+		FNonSubsurfaceSamplerStrategy, FSubsurfaceType, FDimensionHalfRes>;
+
+	// Returns the sampler state based on the requested SSS filter CVar setting.
+	static FRHISamplerState* GetSamplerState()
+	{
+		if (GetSSSFilter())
+		{	// Trilinear is used for mipmap sampling
+			return TStaticSamplerState<SF_Trilinear, AM_Border, AM_Border, AM_Border>::GetRHI();//SF_Bilinear
+		}
+		else
+		{
+			return TStaticSamplerState<SF_Point, AM_Border, AM_Border, AM_Border>::GetRHI();
+		}
+	}
+
+	// Returns the SSS quality level requested by the SSS SampleSet CVar setting.
+	static EQuality GetQuality()
+	{
+		return static_cast<FSubsurfaceIndirectDispatchCS::EQuality>(
+			FMath::Clamp(
+				GetSSSSampleSet(),
+				static_cast<int32>(FSubsurfaceIndirectDispatchCS::EQuality::Low),
+				static_cast<int32>(FSubsurfaceIndirectDispatchCS::EQuality::High)));
+	}
+
+	static ENonSubsurfaceSampleStrategy GetStrategy()
+	{
+		if (GetSSSFilter())
+		{
+			return ENonSubsurfaceSampleStrategy::Considered;
+		}
+		else
+		{
+			return ENonSubsurfaceSampleStrategy::Ignore;
+		}
+	}
+};
+
+IMPLEMENT_GLOBAL_SHADER(FSubsurfaceIndirectDispatchCS, "/Engine/Private/PostProcessSubsurface.usf", "MainIndirectDispatchCS", SF_Compute);
+
+// resolve textures that is not SRV
+// Encapsulates a simple copy pixel shader.
+class FSubsurfaceSRVResolvePS : public FSubsurfaceShader
+{
+	DECLARE_GLOBAL_SHADER(FSubsurfaceSRVResolvePS);
+	SHADER_USE_PARAMETER_STRUCT(FSubsurfaceSRVResolvePS, FSubsurfaceShader);
+
+	BEGIN_SHADER_PARAMETER_STRUCT(FParameters, )
+		SHADER_PARAMETER_RDG_TEXTURE(Texture2D, SubsurfaceInput0_Texture)
+		SHADER_PARAMETER_SAMPLER(SamplerState, SubsurfaceSampler0)
+		RENDER_TARGET_BINDING_SLOTS()
+		END_SHADER_PARAMETER_STRUCT();
+
+	static bool ShouldCompilePermutation(const FGlobalShaderPermutationParameters& Parameters)
+	{
+		return IsFeatureLevelSupported(Parameters.Platform, ERHIFeatureLevel::SM5);
+	}
+};
+
+IMPLEMENT_GLOBAL_SHADER(FSubsurfaceSRVResolvePS, "/Engine/Private/PostProcessSubsurface.usf", "SubsurfaceViewportCopyPS", SF_Pixel);
+
+//-------------------------------------------------------------------------------------------
+// End indirect dispatch class
+//-------------------------------------------------------------------------------------------
+FRDGTextureRef ResolveTextureToSRV(FRDGBuilder& GraphBuilder, FRDGTextureRef InputTexture, const FScreenPassViewInfo& ScreenPassView,
+	const FScreenPassTextureViewport& SceneViewport)
+{
+	const FViewInfo& View = ScreenPassView.View;
+	FRDGTextureDesc SRVDesc = InputTexture->Desc;
+
+	// if this texture can be used as SRV, we ignore this function call.
+	if (SRVDesc.TargetableFlags &TexCreate_ShaderResource)
+	{
+		return InputTexture;
+	}
+
+	SRVDesc.TargetableFlags |= TexCreate_ShaderResource;
+	FRDGTextureRef SRVTextureOutput = GraphBuilder.CreateTexture(SRVDesc, InputTexture->Desc.DebugName);
+
+	FSubsurfaceSRVResolvePS::FParameters* PassParameters = GraphBuilder.AllocParameters<FSubsurfaceSRVResolvePS::FParameters>();
+	PassParameters->RenderTargets[0] = FRenderTargetBinding(SRVTextureOutput, ERenderTargetLoadAction::ENoAction);
+	PassParameters->SubsurfaceInput0_Texture = InputTexture;
+	PassParameters->SubsurfaceSampler0 = TStaticSamplerState<SF_Point, AM_Clamp, AM_Clamp, AM_Clamp>::GetRHI();
+
+	TShaderMapRef<FSubsurfaceSRVResolvePS> PixelShader(View.ShaderMap);
+
+	AddDrawScreenPass(GraphBuilder, RDG_EVENT_NAME("SubsurfaceTextureResolve"), ScreenPassView, SceneViewport, SceneViewport, *PixelShader, PassParameters);
+
+	return SRVTextureOutput;
+}
+
+FRDGTextureRef CreateBlackUAVTexture(FRDGBuilder& GraphBuilder, FRDGTextureDesc SRVDesc, const TCHAR* Name, const FScreenPassViewInfo& ScreenPassView,
+	const FScreenPassTextureViewport& SceneViewport)
+{
+#ifdef USE_CUSTOM_ClearUAV
+	const FViewInfo& View = ScreenPassView.View;
+
+	SRVDesc.TargetableFlags |= TexCreate_ShaderResource | TexCreate_UAV;
+	FRDGTextureRef SRVTextureOutput = GraphBuilder.CreateTexture(SRVDesc, Name);
+
+	FSubsurfaceSRVResolvePS::FParameters* PassParameters = GraphBuilder.AllocParameters<FSubsurfaceSRVResolvePS::FParameters>();
+	PassParameters->RenderTargets[0] = FRenderTargetBinding(SRVTextureOutput, ERenderTargetLoadAction::ENoAction);
+	PassParameters->SubsurfaceInput0_Texture = GraphBuilder.RegisterExternalTexture(GSystemTextures.BlackDummy);
+	PassParameters->SubsurfaceSampler0 = TStaticSamplerState<SF_Point, AM_Clamp, AM_Clamp, AM_Clamp>::GetRHI();
+
+	TShaderMapRef<FSubsurfaceSRVResolvePS> PixelShader(View.ShaderMap);
+
+	AddDrawScreenPass(GraphBuilder, RDG_EVENT_NAME("ClearUAV"), ScreenPassView, SceneViewport, SceneViewport, *PixelShader, PassParameters);
+#else
+	FRDGTextureRef SRVTextureOutput = GraphBuilder.CreateTexture(SRVDesc, Name);
+	FRDGTextureUAVDesc UAVClearDesc(SRVTextureOutput, 0);
+
+	ClearUAV(GraphBuilder, FRDGEventName(TEXT("ClearUAV")), GraphBuilder.CreateUAV(UAVClearDesc), FLinearColor::Black);
+#endif
+
+	return SRVTextureOutput;
+}
+
+// Helper function to use external textures for the current GraphBuilder.
+// When the texture is null, we use BlackDummy.
+FRDGTextureRef RegisterExternalRenderTarget(FRDGBuilder& GraphBuilder, TRefCountPtr<IPooledRenderTarget>* PtrExternalTexture, FIntPoint CurentViewExtent, const TCHAR* Name)
+{
+	FRDGTextureRef RegisteredTexture = NULL;
+
+	if (!PtrExternalTexture || !(*PtrExternalTexture))
+	{
+		RegisteredTexture = GraphBuilder.RegisterExternalTexture(GSystemTextures.BlackDummy, Name);
+	}
+	else
+	{
+		if (CurentViewExtent != (*PtrExternalTexture)->GetDesc().Extent)
+		{
+			RegisteredTexture = GraphBuilder.RegisterExternalTexture(GSystemTextures.BlackDummy, Name);
+		}
+		else
+		{
+			RegisteredTexture = GraphBuilder.RegisterExternalTexture(*PtrExternalTexture, Name);
+		}
+	}
+
+	return RegisteredTexture;
+}
+
+// Encapsulates the post processing subsurface recombine pixel shader.
+class FSubsurfaceRecombinePS : public FSubsurfaceShader
+{
+	DECLARE_GLOBAL_SHADER(FSubsurfaceRecombinePS);
+	SHADER_USE_PARAMETER_STRUCT(FSubsurfaceRecombinePS, FSubsurfaceShader);
+
+	BEGIN_SHADER_PARAMETER_STRUCT(FParameters, )
+		SHADER_PARAMETER_STRUCT_INCLUDE(FSubsurfaceParameters, Subsurface)
+		SHADER_PARAMETER_STRUCT(FSubsurfaceInput, SubsurfaceInput0)
+		SHADER_PARAMETER_STRUCT(FSubsurfaceInput, SubsurfaceInput1)
+		SHADER_PARAMETER_SAMPLER(SamplerState, SubsurfaceSampler0)
+		SHADER_PARAMETER_SAMPLER(SamplerState, SubsurfaceSampler1)
+		RENDER_TARGET_BINDING_SLOTS()
+		END_SHADER_PARAMETER_STRUCT();
+
+	// Controls the quality of lighting reconstruction.
+	enum class EQuality : uint32
+	{
+		Low,
+		High,
+		MAX
+	};
+
+	class FDimensionMode : SHADER_PERMUTATION_ENUM_CLASS("SUBSURFACE_RECOMBINE_MODE", ESubsurfaceMode);
+	class FDimensionQuality : SHADER_PERMUTATION_ENUM_CLASS("SUBSURFACE_RECOMBINE_QUALITY", EQuality);
+	class FDimensionCheckerboard : SHADER_PERMUTATION_BOOL("SUBSURFACE_PROFILE_CHECKERBOARD");
+	using FPermutationDomain = TShaderPermutationDomain<FDimensionMode, FDimensionQuality, FDimensionCheckerboard>;
+
+	// Returns the Recombine quality level requested by the SSS Quality CVar setting.
+	static EQuality GetQuality(const FViewInfo& View)
+	{
+		const uint32 QualityCVar = GetSSSQuality();
+
+		// Quality is forced to high when the CVar is set to 'auto' and TAA is NOT enabled.
+		// TAA improves quality through temporal filtering, making it less necessary to use
+		// high quality mode.
+		const bool bUseHighQuality = (QualityCVar == -1 && View.AntiAliasingMethod != AAM_TemporalAA);
+
+		if (QualityCVar == 1 || bUseHighQuality)
+		{
+			return EQuality::High;
+		}
+		else
+		{
+			return EQuality::Low;
+		}
+	}
+};
+
+IMPLEMENT_GLOBAL_SHADER(FSubsurfaceRecombinePS, "/Engine/Private/PostProcessSubsurface.usf", "SubsurfaceRecombinePS", SF_Pixel);
+
 void ComputeSubsurfaceForView(
 	FRDGBuilder& GraphBuilder,
 	const FScreenPassViewInfo& ScreenPassView,
@@ -121,17 +755,296 @@ void ComputeSubsurfaceForView(
 	FRDGTextureRef SceneTextureOutput,
 	ERenderTargetLoadAction SceneTextureLoadAction)
 {
-	// Compute subsurface based on cmd configuration
-	static const auto CVar = IConsoleManager::Get().FindTConsoleVariableDataInt(TEXT("r.SubsurfaceScattering.Type"));
-	check(CVar);
+	check(SceneTexture);
+	check(SceneTextureOutput);
+	check(SceneViewport.Extent == SceneTexture->Desc.Extent);
+
+	const FViewInfo& View = ScreenPassView.View;
+
+	const FSceneViewFamily* ViewFamily = View.Family;
+
+	const FRDGTextureDesc& SceneTextureDesc = SceneTexture->Desc;
+
+	const ESubsurfaceMode SubsurfaceMode = GetSubsurfaceModeForView(View);
 
-	if (CVar->GetValueOnRenderThread() == 1)	// Use Burley normalized SSS
+	const bool bHalfRes = (SubsurfaceMode == ESubsurfaceMode::HalfRes);
+
+	const bool bCheckerboard = IsSubsurfaceCheckerboardFormat(SceneTextureDesc.Format);
+
+	const uint32 ScaleFactor = bHalfRes ? 2 : 1;
+
+	/**
+	 * All subsurface passes within the screen-space subsurface effect can operate at half or full resolution,
+	 * depending on the subsurface mode. The values are precomputed and shared among all Subsurface textures.
+	 */
+	const FScreenPassTextureViewport SubsurfaceViewport = FScreenPassTextureViewport::CreateDownscaled(SceneViewport, ScaleFactor);
+
+	const FIntPoint TileDimension = FIntPoint::DivideAndRoundUp(SubsurfaceViewport.Extent, SUBSURFACE_GROUP_SIZE);
+	const int32 MaxGroupCount = TileDimension.X*TileDimension.Y;
+
+	const FRDGTextureDesc SubsurfaceTextureDescriptor = FRDGTextureDesc::Create2DDesc(
+		SubsurfaceViewport.Extent,
+		PF_FloatRGBA,
+		FClearValueBinding(),
+		TexCreate_None,
+		TexCreate_RenderTargetable | TexCreate_ShaderResource | TexCreate_UAV,
+		false);
+
+	const FRDGTextureDesc SubsurfaceTextureWith6MipsDescriptor = FRDGTextureDesc::Create2DDesc(
+		SubsurfaceViewport.Extent,
+		PF_FloatRGBA,
+		FClearValueBinding(),
+		TexCreate_None,
+		TexCreate_RenderTargetable | TexCreate_ShaderResource | TexCreate_UAV,
+		false,
+		6);
+
+	const FSubsurfaceParameters SubsurfaceCommonParameters = GetSubsurfaceCommonParameters(GraphBuilder.RHICmdList, View);
+	const FScreenPassTextureViewportParameters SubsurfaceViewportParameters = GetScreenPassTextureViewportParameters(SubsurfaceViewport);
+	const FScreenPassTextureViewportParameters SceneViewportParameters = GetScreenPassTextureViewportParameters(SceneViewport);
+
+	FRDGTextureRef SetupTexture = SceneTexture;
+	FRDGTextureRef SubsurfaceSubpassOneTex = nullptr;
+	FRDGTextureRef SubsurfaceSubpassTwoTex = nullptr;
+
+	FRHISamplerState* PointClampSampler = TStaticSamplerState<SF_Point, AM_Clamp, AM_Clamp, AM_Clamp>::GetRHI();
+	FRHISamplerState* BilinearBorderSampler = TStaticSamplerState<SF_Bilinear, AM_Border, AM_Border, AM_Border>::GetRHI();
+
+	//History texture
+	FSceneViewState* ViewState = (FSceneViewState*)View.State;
+	TRefCountPtr<IPooledRenderTarget>* QualityHistoryState = ViewState ? &ViewState->SubsurfaceScatteringQualityHistoryRT : NULL;
+
+	//allocate/reallocate the quality history texture. 
+	FRDGTextureRef QualityHistoryTexture = RegisterExternalRenderTarget(GraphBuilder, QualityHistoryState, SubsurfaceTextureDescriptor.Extent, TEXT("QualityHistoryTexture"));
+	FRDGTextureRef NewQualityHistoryTexture = nullptr;
+
+	//-----------------------------------------------------------------------------------------------------------------------------------
+	//
+	/**
+	 * When in bypass mode, the setup and convolution passes are skipped, but lighting
+	 * reconstruction is still performed in the recombine pass.
+	 */
+	if (SubsurfaceMode != ESubsurfaceMode::Bypass)
 	{
-		ComputeBurleySubsurfaceForView(GraphBuilder, ScreenPassView, SceneViewport, SceneTexture, SceneTextureOutput, SceneTextureLoadAction);
+
+		SetupTexture = GraphBuilder.CreateTexture(SubsurfaceTextureWith6MipsDescriptor, TEXT("SubsurfaceSetupTexture"));
+
+		// Profile masks (between subsurface and non-subsurface) to reduce convergence.
+		// Notes: variance based sampling algorithm samples too many samples at the edge because it has a high variance due to TAA jittering.
+		// we need to identify them and limit the number of samples at the edge.
+		// We also record the profile of the pixel for half resolution usage.
+		const FRDGTextureDesc ProfileEdgeMaskDescriptor = FRDGTextureDesc::Create2DDesc(
+			SubsurfaceViewport.Extent,
+			PF_FloatRGBA,
+			FClearValueBinding(),
+			TexCreate_None,
+			TexCreate_RenderTargetable | TexCreate_ShaderResource | TexCreate_UAV,
+			false);
+
+		FRDGTextureRef ProfileEdgeMask = GraphBuilder.CreateTexture(ProfileEdgeMaskDescriptor, TEXT("ProfileEdgeMask"));
+		FSceneRenderTargets& SceneContext = FSceneRenderTargets::Get(GraphBuilder.RHICmdList);
+		FRDGTextureRef VelocityTexture = RegisterExternalRenderTarget(GraphBuilder, &(SceneContext.SceneVelocity), SubsurfaceTextureDescriptor.Extent, TEXT("Velocity")); 
+		FSubsurfaceUniformRef UniformBuffer = CreateUniformBuffer(View, MaxGroupCount);
+		
+		NewQualityHistoryTexture = CreateBlackUAVTexture(GraphBuilder, SubsurfaceTextureDescriptor, TEXT("SubsurfaceQualityHistoryState"),
+			ScreenPassView, SceneViewport);
+
+		// Initialize the group buffer
+		FRDGBufferRef SeparableGroupBuffer = GraphBuilder.CreateBuffer(FRDGBufferDesc::CreateStructuredDesc(sizeof(FSubsurfaceGroupItem), MaxGroupCount + 1), TEXT("SeparableGroupBuffer"));;
+		FRDGBufferRef BurleyGroupBuffer = GraphBuilder.CreateBuffer(FRDGBufferDesc::CreateStructuredDesc(sizeof(FSubsurfaceGroupItem), MaxGroupCount + 1), TEXT("BurleyGroupBuffer"));;
+		FRDGBufferRef SeparableIndirectDispatchArgsBuffer = GraphBuilder.CreateBuffer(FRDGBufferDesc::CreateIndirectDesc(4), TEXT("SeprableIndirectDispatchArgs"));
+		FRDGBufferRef BurleyIndirectDispatchArgsBuffer = GraphBuilder.CreateBuffer(FRDGBufferDesc::CreateIndirectDesc(4), TEXT("BurleyIndirectDispatchArgs"));
+
+		// Initialize the group counters
+		{
+			typedef FSubsurfaceInitValueBufferCS SHADER;
+			TShaderMapRef<SHADER> ComputeShader(View.ShaderMap);
+			SHADER::FParameters* PassParameters = GraphBuilder.AllocParameters<SHADER::FParameters>();
+			PassParameters->RWBurleyGroupBuffer = GraphBuilder.CreateUAV(BurleyGroupBuffer, EPixelFormat::PF_R32_UINT);
+			PassParameters->RWSeparableGroupBuffer = GraphBuilder.CreateUAV(SeparableGroupBuffer, EPixelFormat::PF_R32_UINT);
+			FComputeShaderUtils::AddPass(GraphBuilder, RDG_EVENT_NAME("InitGroupCounter"), *ComputeShader, PassParameters, FIntVector(1, 1, 1));
+		}
+
+		// Call the indirect setup
+		{
+			FRDGTextureSRVDesc SceneTextureSRVDesc = FRDGTextureSRVDesc::Create(SceneTexture);
+			FRDGTextureUAVDesc SetupTextureOutDesc(SetupTexture, 0);
+			FRDGTextureUAVDesc ProfileEdgeMaskDesc(ProfileEdgeMask, 0);
+
+			typedef FSubsurfaceIndirectDispatchSetupCS SHADER;
+			SHADER::FParameters* PassParameters = GraphBuilder.AllocParameters<SHADER::FParameters>();
+			PassParameters->Subsurface = SubsurfaceCommonParameters;
+			PassParameters->SubsurfaceInput0 = GetSubsurfaceInput(SceneTexture, SceneViewportParameters);
+			PassParameters->SubsurfaceSampler0 = PointClampSampler;
+			PassParameters->ProfileEdgeMask = GraphBuilder.CreateUAV(ProfileEdgeMaskDesc);
+			PassParameters->SetupTexture = GraphBuilder.CreateUAV(SetupTextureOutDesc);
+			PassParameters->RWBurleyGroupBuffer = GraphBuilder.CreateUAV(BurleyGroupBuffer, EPixelFormat::PF_R32_UINT);
+			PassParameters->RWSeparableGroupBuffer = GraphBuilder.CreateUAV(SeparableGroupBuffer, EPixelFormat::PF_R32_UINT);
+			PassParameters->SubsurfaceUniformParameters = UniformBuffer;
+
+			SHADER::FPermutationDomain ComputeShaderPermutationVector;
+			ComputeShaderPermutationVector.Set<SHADER::FDimensionHalfRes>(bHalfRes);
+			ComputeShaderPermutationVector.Set<SHADER::FDimensionCheckerboard>(bCheckerboard);
+			TShaderMapRef<SHADER> ComputeShader(View.ShaderMap, ComputeShaderPermutationVector);
+
+			FIntPoint ComputeGroupCount = FIntPoint::DivideAndRoundUp(SubsurfaceViewport.Extent, SUBSURFACE_GROUP_SIZE);
+
+			FComputeShaderUtils::AddPass(GraphBuilder, RDG_EVENT_NAME("SubsurfaceSetup"), *ComputeShader, PassParameters, FIntVector(ComputeGroupCount.X, ComputeGroupCount.Y, 1));
+		}
+
+		{
+			//Generate mipmaps for the diffuse scene color and depth, use bilinear filter
+			FGenerateMips::Execute(&GraphBuilder, SetupTexture, BilinearBorderSampler);
+		}
+
+		typedef FSubsurfaceIndirectDispatchCS SHADER;
+
+		FRHISamplerState* SubsurfaceSamplerState = SHADER::GetSamplerState();
+		const SHADER::EQuality SubsurfaceQuality = SHADER::GetQuality();
+
+		// Store the buffer
+		const FRDGBufferRef SubsurfaceBufferUsage[] = { BurleyGroupBuffer,                      SeparableGroupBuffer };
+		const FRDGBufferRef  SubsurfaceBufferArgs[] = { BurleyIndirectDispatchArgsBuffer,       SeparableIndirectDispatchArgsBuffer };
+		const TCHAR*		  SubsurfacePhaseName[] = { TEXT("BuildBurleyIndirectDispatchArgs"),TEXT("BuildSeparableIndirectDispatchArgs") };
+
+		// Setup the indirect arguments.
+		{
+			const int NumOfSubsurfaceType = 2;
+
+			for (int SubsurfaceTypeIndex = 0; SubsurfaceTypeIndex < NumOfSubsurfaceType; ++SubsurfaceTypeIndex)
+			{
+				typedef FSubsurfaceBuildIndirectDispatchArgsCS ARGSETUPSHADER;
+				ARGSETUPSHADER::FParameters* PassParameters = GraphBuilder.AllocParameters<ARGSETUPSHADER::FParameters>();
+				PassParameters->SubsurfaceUniformParameters = UniformBuffer;
+				PassParameters->RWIndirectDispatchArgsBuffer = GraphBuilder.CreateUAV(SubsurfaceBufferArgs[SubsurfaceTypeIndex], EPixelFormat::PF_R32_UINT);
+				PassParameters->GroupBuffer = GraphBuilder.CreateSRV(SubsurfaceBufferUsage[SubsurfaceTypeIndex], EPixelFormat::PF_R32_UINT);
+
+				TShaderMapRef<ARGSETUPSHADER> ComputeShader(View.ShaderMap);
+				FComputeShaderUtils::AddPass(GraphBuilder, FRDGEventName(SubsurfacePhaseName[SubsurfaceTypeIndex]), *ComputeShader, PassParameters, FIntVector(1, 1, 1));
+			}
+		}
+
+		// Major pass to combine Burley and Separable
+		{
+			SubsurfaceSubpassOneTex = CreateBlackUAVTexture(GraphBuilder, SubsurfaceTextureWith6MipsDescriptor, TEXT("SubsurfaceSubpassOneTex"),
+				ScreenPassView, SceneViewport);
+			SubsurfaceSubpassTwoTex = CreateBlackUAVTexture(GraphBuilder, SubsurfaceTextureWith6MipsDescriptor, TEXT("SubsurfaceSubpassTwoTex"),
+				ScreenPassView, SceneViewport);
+
+			struct FSubsurfacePassInfo
+			{
+				FSubsurfacePassInfo(const TCHAR* InName, FRDGTextureRef InInput, FRDGTextureRef InOutput,
+					SHADER::ESubsurfaceType InSurfaceType, SHADER::ESubsurfacePass InSurfacePass)
+					: Name(InName), Input(InInput), Output(InOutput), SurfaceType(InSurfaceType), SubsurfacePass(InSurfacePass)
+				{}
+
+				const TCHAR* Name;
+				FRDGTextureRef Input;
+				FRDGTextureRef Output;
+				SHADER::ESubsurfaceType SurfaceType;
+				SHADER::ESubsurfacePass SubsurfacePass;
+			};
+
+			const int NumOfSubsurfacePass = 4;
+
+			const FSubsurfacePassInfo SubsurfacePassInfos[NumOfSubsurfacePass] =
+			{
+				{	TEXT("SubsurfacePassOne_Burley"),				SetupTexture, SubsurfaceSubpassOneTex, SHADER::ESubsurfaceType::BURLEY	 , SHADER::ESubsurfacePass::PassOne}, //Burley main pass
+				{	TEXT("SubsurfacePassTwo_SepHon"),				SetupTexture, SubsurfaceSubpassOneTex, SHADER::ESubsurfaceType::SEPARABLE, SHADER::ESubsurfacePass::PassOne}, //Separable horizontal
+				{ TEXT("SubsurfacePassThree_SepVer"),	 SubsurfaceSubpassOneTex, SubsurfaceSubpassTwoTex, SHADER::ESubsurfaceType::SEPARABLE, SHADER::ESubsurfacePass::PassTwo}, //Separable Vertical
+				{	 TEXT("SubsurfacePassFour_BVar"),    SubsurfaceSubpassOneTex, SubsurfaceSubpassTwoTex, SHADER::ESubsurfaceType::BURLEY	 , SHADER::ESubsurfacePass::PassTwo}  //Burley Variance
+			};
+
+			//Dispatch the two phase for both SSS
+			for (int PassIndex = 0; PassIndex < NumOfSubsurfacePass; ++PassIndex)
+			{
+				const FSubsurfacePassInfo& PassInfo = SubsurfacePassInfos[PassIndex];
+
+				const SHADER::ESubsurfaceType SubsurfaceType = PassInfo.SurfaceType;
+				const auto SubsurfacePassFunction = PassInfo.SubsurfacePass;
+				const int SubsurfaceTypeIndex = static_cast<int>(SubsurfaceType);
+				FRDGTextureRef TextureInput = PassInfo.Input;
+				FRDGTextureRef TextureOutput = PassInfo.Output;
+
+				FRDGTextureUAVDesc SSSColorUAVDesc(TextureOutput, 0);
+				FRDGTextureSRVDesc InputSRVDesc = FRDGTextureSRVDesc::Create(TextureInput);
+
+				SHADER::FParameters* PassParameters = GraphBuilder.AllocParameters<SHADER::FParameters>();
+				PassParameters->Subsurface = SubsurfaceCommonParameters;
+				PassParameters->SSSColorUAV = GraphBuilder.CreateUAV(SSSColorUAVDesc);
+				PassParameters->SubsurfaceInput0 = GetSubsurfaceInput(TextureInput, SubsurfaceViewportParameters);
+				PassParameters->SubsurfaceSampler0 = SubsurfaceSamplerState;
+				PassParameters->GroupBuffer = GraphBuilder.CreateSRV(SubsurfaceBufferUsage[SubsurfaceTypeIndex], EPixelFormat::PF_R32_UINT);
+				PassParameters->IndirectDispatchArgsBuffer = SubsurfaceBufferArgs[SubsurfaceTypeIndex]; // need to maintain a reference to the indirect dispatch args.
+
+				if (SubsurfaceType == SHADER::ESubsurfaceType::SEPARABLE)
+				{
+					PassParameters->SubsurfaceInput2 = GetSubsurfaceInput(ProfileEdgeMask, SubsurfaceViewportParameters);
+					PassParameters->SubsurfaceSampler2 = PointClampSampler;
+				}
+
+				if (SubsurfacePassFunction == SHADER::ESubsurfacePass::PassOne && SubsurfaceType == SHADER::ESubsurfaceType::BURLEY)
+				{
+					PassParameters->SubsurfaceInput1 = GetSubsurfaceInput(QualityHistoryTexture, SubsurfaceViewportParameters);
+					PassParameters->SubsurfaceSampler1 = PointClampSampler;
+					PassParameters->SubsurfaceInput2 = GetSubsurfaceInput(ProfileEdgeMask, SubsurfaceViewportParameters);
+					PassParameters->SubsurfaceSampler2 = PointClampSampler;
+				}
+
+				if (SubsurfacePassFunction == SHADER::ESubsurfacePass::PassTwo && SubsurfaceType == SHADER::ESubsurfaceType::BURLEY)
+				{
+					PassParameters->HistoryUAV = GraphBuilder.CreateUAV(FRDGTextureUAVDesc(NewQualityHistoryTexture, 0));
+					PassParameters->SubsurfaceInput1 = GetSubsurfaceInput(QualityHistoryTexture, SubsurfaceViewportParameters);
+					PassParameters->SubsurfaceSampler1 = PointClampSampler;
+					PassParameters->SubsurfaceInput2 = GetSubsurfaceInput(VelocityTexture, SubsurfaceViewportParameters);
+					PassParameters->SubsurfaceSampler2 = PointClampSampler;
+				}
+
+				SHADER::FPermutationDomain ComputeShaderPermutationVector;
+				ComputeShaderPermutationVector.Set<SHADER::FSubsurfacePassFunction>(SubsurfacePassFunction);
+				ComputeShaderPermutationVector.Set<SHADER::FDimensionQuality>(SHADER::GetQuality());
+				ComputeShaderPermutationVector.Set<SHADER::FNonSubsurfaceSamplerStrategy>(SHADER::GetStrategy());
+				ComputeShaderPermutationVector.Set<SHADER::FSubsurfaceType>(SubsurfaceType);
+				ComputeShaderPermutationVector.Set<SHADER::FDimensionHalfRes>(bHalfRes);
+				TShaderMapRef<SHADER> ComputeShader(View.ShaderMap, ComputeShaderPermutationVector);
+
+				FComputeShaderUtils::AddPass(GraphBuilder, FRDGEventName(PassInfo.Name), *ComputeShader, PassParameters, SubsurfaceBufferArgs[SubsurfaceTypeIndex], 0);
+			}
+		}
+
+		// Replace the history
+		if (QualityHistoryState)
+			GraphBuilder.QueueTextureExtraction(NewQualityHistoryTexture, QualityHistoryState, true);
 	}
-	else										// Use separable SSS.
+
+	// Recombines scattering result with scene color.
+	{
+	FSubsurfaceRecombinePS::FParameters* PassParameters = GraphBuilder.AllocParameters<FSubsurfaceRecombinePS::FParameters>();
+	PassParameters->Subsurface = SubsurfaceCommonParameters;
+	PassParameters->RenderTargets[0] = FRenderTargetBinding(SceneTextureOutput, SceneTextureLoadAction);
+	PassParameters->SubsurfaceInput0 = GetSubsurfaceInput(SceneTexture, SceneViewportParameters);
+	PassParameters->SubsurfaceSampler0 = BilinearBorderSampler;
+
+	// Scattering output target is only used when scattering is enabled.
+	if (SubsurfaceMode != ESubsurfaceMode::Bypass)
 	{
-		ComputeSeparableSubsurfaceForView(GraphBuilder, ScreenPassView, SceneViewport, SceneTexture, SceneTextureOutput, SceneTextureLoadAction);
+		PassParameters->SubsurfaceInput1 = GetSubsurfaceInput(SubsurfaceSubpassTwoTex, SubsurfaceViewportParameters);
+		PassParameters->SubsurfaceSampler1 = BilinearBorderSampler;
+	}
+
+	const FSubsurfaceRecombinePS::EQuality RecombineQuality = FSubsurfaceRecombinePS::GetQuality(View);
+
+	FSubsurfaceRecombinePS::FPermutationDomain PixelShaderPermutationVector;
+	PixelShaderPermutationVector.Set<FSubsurfaceRecombinePS::FDimensionMode>(SubsurfaceMode);
+	PixelShaderPermutationVector.Set<FSubsurfaceRecombinePS::FDimensionQuality>(RecombineQuality);
+	PixelShaderPermutationVector.Set<FSubsurfaceRecombinePS::FDimensionCheckerboard>(bCheckerboard);
+	TShaderMapRef<FSubsurfaceRecombinePS> PixelShader(View.ShaderMap, PixelShaderPermutationVector);
+
+	/**
+	 * See the related comment above in the prepare pass. The scene viewport is used as both the target and
+	 * texture viewport in order to ensure that the correct pixel is sampled for checkerboard rendering.
+	 */
+	AddDrawScreenPass(GraphBuilder, RDG_EVENT_NAME("SubsurfaceRecombine"), ScreenPassView, SceneViewport, SceneViewport, *PixelShader, PassParameters);
 	}
 }
 
diff --git a/Engine/Source/Runtime/Renderer/Private/PostProcess/SubsurfaceBurleyNormalized.cpp b/Engine/Source/Runtime/Renderer/Private/PostProcess/SubsurfaceBurleyNormalized.cpp
deleted file mode 100644
diff --git a/Engine/Source/Runtime/Renderer/Private/PostProcess/SubsurfaceBurleyNormalized.h b/Engine/Source/Runtime/Renderer/Private/PostProcess/SubsurfaceBurleyNormalized.h
deleted file mode 100644
diff --git a/Engine/Source/Runtime/Renderer/Private/PostProcess/SubsurfaceCommon.cpp b/Engine/Source/Runtime/Renderer/Private/PostProcess/SubsurfaceCommon.cpp
deleted file mode 100644
diff --git a/Engine/Source/Runtime/Renderer/Private/PostProcess/SubsurfaceCommon.h b/Engine/Source/Runtime/Renderer/Private/PostProcess/SubsurfaceCommon.h
deleted file mode 100644
diff --git a/Engine/Source/Runtime/Renderer/Private/PostProcess/SubsurfaceSeparable.cpp b/Engine/Source/Runtime/Renderer/Private/PostProcess/SubsurfaceSeparable.cpp
deleted file mode 100644
diff --git a/Engine/Source/Runtime/Renderer/Private/PostProcess/SubsurfaceSeparable.h b/Engine/Source/Runtime/Renderer/Private/PostProcess/SubsurfaceSeparable.h
deleted file mode 100644
diff --git a/Engine/Source/Runtime/Renderer/Private/ScenePrivate.h b/Engine/Source/Runtime/Renderer/Private/ScenePrivate.h
--- a/Engine/Source/Runtime/Renderer/Private/ScenePrivate.h
+++ b/Engine/Source/Runtime/Renderer/Private/ScenePrivate.h
@@ -872,6 +872,9 @@ public:
 	TRefCountPtr<IPooledRenderTarget> MobileAaColor0;
 	TRefCountPtr<IPooledRenderTarget> MobileAaColor1;
 
+	// Burley Subsurface scattering variance texture from the last frame.
+	TRefCountPtr<IPooledRenderTarget> SubsurfaceScatteringQualityHistoryRT;
+
 	// Pre-computed filter in spectral (i.e. FFT) domain along with data to determine if we need to up date it
 	struct {
 		/// @cond DOXYGEN_WARNINGS
