Ported histogram eye adaptation to RDG.

diff --git a/Engine/Shaders/Private/PostProcessEyeAdaptation.usf b/Engine/Shaders/Private/PostProcessEyeAdaptation.usf
--- a/Engine/Shaders/Private/PostProcessEyeAdaptation.usf
+++ b/Engine/Shaders/Private/PostProcessEyeAdaptation.usf
@@ -91,14 +91,16 @@ float ComputeEyeAdaptation(float OldExposure, float TargetExposure, float FrameT
 	return clamp(OldExposure + Diff * Factor, EyeAdaptation_MinAverageLuminance, EyeAdaptation_MaxAverageLuminance);
 }
 
+Texture2D HistogramTexture;
+
 float4 MainEyeAdaptationCommon()
 {
 	float4 OutColor = 0;
 
-	const float AverageSceneLuminance = ComputeEyeAdaptationExposure(PostprocessInput0);
+	const float AverageSceneLuminance = ComputeEyeAdaptationExposure(HistogramTexture);
 	const float TargetExposure = clamp(AverageSceneLuminance, EyeAdaptation_MinAverageLuminance, EyeAdaptation_MaxAverageLuminance);
 
-	float OldExposureScale = PostprocessInput0.Load(int3(0, 1, 0)).x;
+	float OldExposureScale = HistogramTexture.Load(int3(0, 1, 0)).x;
 	float OldExposure = EyeAdaptation_ExposureCompensation / ( OldExposureScale != 0 ? OldExposureScale : 1.0f );
 
 	// eye adaptation changes over time
@@ -122,15 +124,15 @@ void MainPS(float4 UVAndScreenPos : TEXCOORD0,
 }
 
 #if COMPUTESHADER
-RWTexture2D<float4> RWOutComputeTex;
+RWTexture2D<float4> RWEyeAdaptationTexture;
 
 [numthreads(1, 1, 1)]
 void MainCS(uint2 DispatchThreadId : SV_DispatchThreadID)
 {
 	float4 OutColor = MainEyeAdaptationCommon();
 
-	uint2 PixelPos = DispatchThreadId + (uint2)View.ViewRectMin.xy;
-	RWOutComputeTex[PixelPos] = OutColor;
+	uint2 PixelPos = DispatchThreadId;
+	RWEyeAdaptationTexture[PixelPos] = OutColor;
 }
 #endif
 
diff --git a/Engine/Source/Runtime/Renderer/Private/PostProcess/PostProcessBloomSetup.cpp b/Engine/Source/Runtime/Renderer/Private/PostProcess/PostProcessBloomSetup.cpp
--- a/Engine/Source/Runtime/Renderer/Private/PostProcess/PostProcessBloomSetup.cpp
+++ b/Engine/Source/Runtime/Renderer/Private/PostProcess/PostProcessBloomSetup.cpp
@@ -64,7 +64,7 @@ public:
 
 		PostprocessParameter.SetPS(RHICmdList, ShaderRHI, Context, TStaticSamplerState<SF_Bilinear, AM_Clamp, AM_Clamp, AM_Clamp>::GetRHI());
 
-		const float FixedExposure = FRCPassPostProcessEyeAdaptation::GetFixedExposure(Context.View);
+		const float FixedExposure = GetEyeAdaptationFixedExposure(Context.View);
 
 		FVector4 BloomThresholdValue(Settings.BloomThreshold, 0, 0, FixedExposure);
 		SetShaderValue(RHICmdList, ShaderRHI, BloomThreshold, BloomThresholdValue);
diff --git a/Engine/Source/Runtime/Renderer/Private/PostProcess/PostProcessEyeAdaptation.cpp b/Engine/Source/Runtime/Renderer/Private/PostProcess/PostProcessEyeAdaptation.cpp
--- a/Engine/Source/Runtime/Renderer/Private/PostProcess/PostProcessEyeAdaptation.cpp
+++ b/Engine/Source/Runtime/Renderer/Private/PostProcess/PostProcessEyeAdaptation.cpp
@@ -238,13 +238,23 @@ FEyeAdaptationParameters GetEyeAdaptationParameters(const FViewInfo& View, ERHIF
 	return Parameters;
 }
 
-// Basic AutoExposure requires at least ES3_1
+float GetEyeAdaptationFixedExposure(const FViewInfo& View)
+{
+	const FEyeAdaptationParameters Parameters = GetEyeAdaptationParameters(View);
+
+	const float Exposure = (Parameters.MinAverageLuminance + Parameters.MaxAverageLuminance) * 0.5f;
+
+	const float ExposureScale = 1.0f / FMath::Max(0.0001f, Exposure);
+
+	return ExposureScale * Parameters.ExposureCompensation;
+}
+
 const ERHIFeatureLevel::Type BasicEyeAdaptationMinFeatureLevel = ERHIFeatureLevel::ES3_1;
 
-/** Encapsulates the histogram-based post processing eye adaptation pixel shader. */
-class FPostProcessEyeAdaptationPS : public FGlobalShader
+class FEyeAdaptationShader : public FGlobalShader
 {
-	DECLARE_SHADER_TYPE(FPostProcessEyeAdaptationPS, Global);
+public:
+	static const EPixelFormat OutputFormat = PF_A32B32G32R32F;
 
 	static bool ShouldCompilePermutation(const FGlobalShaderPermutationParameters& Parameters)
 	{
@@ -253,270 +263,122 @@ class FPostProcessEyeAdaptationPS : public FGlobalShader
 
 	static void ModifyCompilationEnvironment(const FGlobalShaderPermutationParameters& Parameters, FShaderCompilerEnvironment& OutEnvironment)
 	{
-		FGlobalShader::ModifyCompilationEnvironment( Parameters, OutEnvironment );
-		OutEnvironment.SetRenderTargetOutputFormat(0, PF_A32B32G32R32F);
+		FGlobalShader::ModifyCompilationEnvironment(Parameters, OutEnvironment);
+		OutEnvironment.SetRenderTargetOutputFormat(0, OutputFormat);
 	}
 
-	/** Default constructor. */
-	FPostProcessEyeAdaptationPS() {}
+	FEyeAdaptationShader() = default;
 
+	FEyeAdaptationShader(const CompiledShaderInitializerType& Initializer)
+		: FGlobalShader(Initializer)
+	{}
+};
+
+class FEyeAdaptationPS : public FEyeAdaptationShader
+{
 public:
+	DECLARE_GLOBAL_SHADER(FEyeAdaptationPS);
+	SHADER_USE_PARAMETER_STRUCT(FEyeAdaptationPS, FEyeAdaptationShader);
+
 	BEGIN_SHADER_PARAMETER_STRUCT(FParameters, )
 		SHADER_PARAMETER_STRUCT(FEyeAdaptationParameters, EyeAdaptation)
+		SHADER_PARAMETER_RDG_TEXTURE(Texture2D, HistogramTexture)
+		RENDER_TARGET_BINDING_SLOTS()
 	END_SHADER_PARAMETER_STRUCT()
-
-	FPostProcessPassParameters PostprocessParameter;
-	FShaderParameter EyeAdaptationParams;
-	FShaderResourceParameter EyeAdaptationTexture;
-
-	/** Initialization constructor. */
-	FPostProcessEyeAdaptationPS(const ShaderMetaType::CompiledShaderInitializerType& Initializer)
-		: FGlobalShader(Initializer)
-	{
-		BindForLegacyShaderParameters<FParameters>(this, Initializer.ParameterMap);
-
-		PostprocessParameter.Bind(Initializer.ParameterMap);
-		EyeAdaptationTexture.Bind(Initializer.ParameterMap, TEXT("EyeAdaptationTexture"));
-	}
-
-	template <typename TRHICmdList>
-	void SetPS(const FRenderingCompositePassContext& Context, TRHICmdList& RHICmdList, IPooledRenderTarget* LastEyeAdaptation)
-	{
-		FRHIPixelShader* ShaderRHI = GetPixelShader();
-
-		FGlobalShader::SetParameters<FViewUniformShaderParameters>(RHICmdList, ShaderRHI, Context.View.ViewUniformBuffer);
-
-		PostprocessParameter.SetPS(RHICmdList, ShaderRHI, Context, TStaticSamplerState<SF_Bilinear,AM_Clamp,AM_Clamp,AM_Clamp>::GetRHI());
-
-		if (Context.View.HasValidEyeAdaptation())
-		{
-			SetTextureParameter(Context.RHICmdList, ShaderRHI, EyeAdaptationTexture, LastEyeAdaptation->GetRenderTargetItem().TargetableTexture);
-		}
-		else // some views don't have a state, thumbnail rendering?
-		{
-			SetTextureParameter(Context.RHICmdList, ShaderRHI, EyeAdaptationTexture, GWhiteTexture->TextureRHI);
-		}
-
-		{
-			FParameters PassParameters;
-			PassParameters.EyeAdaptation = GetEyeAdaptationParameters(Context.View);
-			SetShaderParameters(Context.RHICmdList, this, ShaderRHI, PassParameters);
-		}
-	}
-
-	// FShader interface.
-	virtual bool Serialize(FArchive& Ar) override
-	{
-		bool bShaderHasOutdatedParameters = FGlobalShader::Serialize(Ar);
-		Ar << PostprocessParameter << EyeAdaptationParams;
-		return bShaderHasOutdatedParameters;
-	}
 };
 
-IMPLEMENT_SHADER_TYPE(,FPostProcessEyeAdaptationPS,TEXT("/Engine/Private/PostProcessEyeAdaptation.usf"),TEXT("MainPS"),SF_Pixel);
+IMPLEMENT_GLOBAL_SHADER(FEyeAdaptationPS, "/Engine/Private/PostProcessEyeAdaptation.usf", "MainPS", SF_Pixel);
 
-/** Encapsulates the histogram-based post processing eye adaptation compute shader. */
-class FPostProcessEyeAdaptationCS : public FGlobalShader
+class FEyeAdaptationCS : public FEyeAdaptationShader
 {
-	DECLARE_SHADER_TYPE(FPostProcessEyeAdaptationCS, Global);
-
-	static bool ShouldCompilePermutation(const FGlobalShaderPermutationParameters& Parameters)
-	{
-		return IsFeatureLevelSupported(Parameters.Platform, ERHIFeatureLevel::SM5);
-	}
-
-	static void ModifyCompilationEnvironment(const FGlobalShaderPermutationParameters& Parameters, FShaderCompilerEnvironment& OutEnvironment )
-	{
-		FGlobalShader::ModifyCompilationEnvironment(Parameters, OutEnvironment);
-		OutEnvironment.SetRenderTargetOutputFormat(0, PF_A32B32G32R32F);
-	}
-
-	/** Default constructor. */
-	FPostProcessEyeAdaptationCS() {}
-
 public:
+	DECLARE_GLOBAL_SHADER(FEyeAdaptationCS);
+	SHADER_USE_PARAMETER_STRUCT(FEyeAdaptationCS, FEyeAdaptationShader);
+
 	BEGIN_SHADER_PARAMETER_STRUCT(FParameters, )
 		SHADER_PARAMETER_STRUCT(FEyeAdaptationParameters, EyeAdaptation)
+		SHADER_PARAMETER_RDG_TEXTURE(Texture2D, HistogramTexture)
+		SHADER_PARAMETER_RDG_TEXTURE_UAV(RWTexture2D<float4>, RWEyeAdaptationTexture)
 	END_SHADER_PARAMETER_STRUCT()
-
-	FPostProcessPassParameters PostprocessParameter;
-	FRWShaderParameter OutComputeTex;
-	FShaderParameter EyeAdaptationParams;
-
-	/** Initialization constructor. */
-	FPostProcessEyeAdaptationCS(const ShaderMetaType::CompiledShaderInitializerType& Initializer)
-		: FGlobalShader(Initializer)
-	{
-		BindForLegacyShaderParameters<FParameters>(this, Initializer.ParameterMap);
-		PostprocessParameter.Bind(Initializer.ParameterMap);
-		OutComputeTex.Bind(Initializer.ParameterMap, TEXT("OutComputeTex"));
-	}
-
-	template <typename TRHICmdList>
-	void SetParameters(TRHICmdList& RHICmdList, const FRenderingCompositePassContext& Context, FRHIUnorderedAccessView* DestUAV, IPooledRenderTarget* LastEyeAdaptation)
-	{
-		FRHIComputeShader* ShaderRHI = GetComputeShader();
-
-		// CS params
-		FGlobalShader::SetParameters<FViewUniformShaderParameters>(RHICmdList, ShaderRHI, Context.View.ViewUniformBuffer);
-		PostprocessParameter.SetCS(ShaderRHI, Context, RHICmdList, TStaticSamplerState<SF_Bilinear,AM_Clamp,AM_Clamp,AM_Clamp>::GetRHI());
-		OutComputeTex.SetTexture(RHICmdList, ShaderRHI, nullptr, DestUAV);		
-
-		FParameters PassParameters;
-		PassParameters.EyeAdaptation = GetEyeAdaptationParameters(Context.View);
-		SetShaderParameters(RHICmdList, this, ShaderRHI, PassParameters);
-	}
-
-	template <typename TRHICmdList>
-	void UnsetParameters(TRHICmdList& RHICmdList)
-	{
-		FRHIComputeShader* ShaderRHI = GetComputeShader();
-		OutComputeTex.UnsetUAV(RHICmdList, ShaderRHI);
-	}
-
-	// FShader interface.
-	virtual bool Serialize(FArchive& Ar) override
-	{
-		bool bShaderHasOutdatedParameters = FGlobalShader::Serialize(Ar);
-		Ar << PostprocessParameter << OutComputeTex << EyeAdaptationParams;
-		return bShaderHasOutdatedParameters;
-	}
 };
 
-IMPLEMENT_SHADER_TYPE(,FPostProcessEyeAdaptationCS,TEXT("/Engine/Private/PostProcessEyeAdaptation.usf"),TEXT("MainCS"),SF_Compute);
+IMPLEMENT_GLOBAL_SHADER(FEyeAdaptationCS, "/Engine/Private/PostProcessEyeAdaptation.usf", "MainCS", SF_Compute);
 
-void FRCPassPostProcessEyeAdaptation::Process(FRenderingCompositePassContext& Context)
+FRDGTextureRef AddHistogramEyeAdaptationPass(
+	FRDGBuilder& GraphBuilder,
+	const FScreenPassViewInfo& ScreenPassView,
+	const FEyeAdaptationParameters& EyeAdaptationParameters,
+	FRDGTextureRef HistogramTexture)
 {
-	SCOPED_DRAW_EVENTF(Context.RHICmdList, PostProcessEyeAdaptation, TEXT("PostProcessEyeAdaptation%s"), bIsComputePass?TEXT("Compute"):TEXT(""));
-	AsyncEndFence = FComputeFenceRHIRef();
+	const FViewInfo& View = ScreenPassView.View;
 
-	const FViewInfo& View = Context.View;
-	const FSceneViewFamily& ViewFamily = *(View.Family);
+	View.SwapEyeAdaptationRTs(GraphBuilder.RHICmdList);
+	View.SetValidEyeAdaptation();
 
-	// Get the custom 1x1 target used to store exposure value and Toggle the two render targets used to store new and old.
-	Context.View.SwapEyeAdaptationRTs(Context.RHICmdList);
-
-	IPooledRenderTarget* LastEyeAdaptation = Context.View.GetLastEyeAdaptationRT(Context.RHICmdList);
-	IPooledRenderTarget* EyeAdaptation = Context.View.GetEyeAdaptation(Context.RHICmdList);
-	check(EyeAdaptation);
-
-	FIntPoint DestSize = EyeAdaptation->GetDesc().Extent;
+	FRDGTextureRef OutputTexture = GraphBuilder.RegisterExternalTexture(View.GetEyeAdaptation(GraphBuilder.RHICmdList), TEXT("EyeAdaptation"), ERDGResourceFlags::MultiFrame);
 
-	const FSceneRenderTargetItem& DestRenderTarget = EyeAdaptation->GetRenderTargetItem();
-
-	static auto* RenderPassCVar = IConsoleManager::Get().FindTConsoleVariableDataInt(TEXT("r.RHIRenderPasses"));
-	if (bIsComputePass)
+	if (ScreenPassView.bUseComputePasses)
 	{
-		FIntRect DestRect(0, 0, DestSize.X, DestSize.Y);
-
-		// Common setup
-		// #todo-renderpasses remove once everything is renderpasses
-		UnbindRenderTargets(Context.RHICmdList);
-		Context.SetViewportAndCallRHI(DestRect, 0.0f, 1.0f);
-		
-		static FName AsyncEndFenceName(TEXT("AsyncEyeAdaptationEndFence"));
-		AsyncEndFence = Context.RHICmdList.CreateComputeFence(AsyncEndFenceName);
+		FEyeAdaptationCS::FParameters* PassParameters = GraphBuilder.AllocParameters<FEyeAdaptationCS::FParameters>();
+		PassParameters->EyeAdaptation = GetEyeAdaptationParameters(View);
+		PassParameters->HistogramTexture = HistogramTexture;
+		PassParameters->RWEyeAdaptationTexture = GraphBuilder.CreateUAV(OutputTexture);
 
-		if (IsAsyncComputePass())
-		{
-			// Async path
-			FRHIAsyncComputeCommandListImmediate& RHICmdListComputeImmediate = FRHICommandListExecutor::GetImmediateAsyncComputeCommandList();
-			{
-				SCOPED_COMPUTE_EVENT(RHICmdListComputeImmediate, AsyncEyeAdaptation);
-				WaitForInputPassComputeFences(RHICmdListComputeImmediate);
-				RHICmdListComputeImmediate.TransitionResource(EResourceTransitionAccess::ERWBarrier, EResourceTransitionPipeline::EGfxToCompute, DestRenderTarget.UAV);
-				DispatchCS(RHICmdListComputeImmediate, Context, DestRenderTarget.UAV, LastEyeAdaptation);
-				RHICmdListComputeImmediate.TransitionResource(EResourceTransitionAccess::EReadable, EResourceTransitionPipeline::EComputeToGfx, DestRenderTarget.UAV, AsyncEndFence);
-			}
-			FRHIAsyncComputeCommandListImmediate::ImmediateDispatch(RHICmdListComputeImmediate);
-		}
-		else
-		{
-			// Direct path
-			WaitForInputPassComputeFences(Context.RHICmdList);
-			Context.RHICmdList.BeginUpdateMultiFrameResource(DestRenderTarget.ShaderResourceTexture);
+		TShaderMapRef<FEyeAdaptationCS> ComputeShader(ScreenPassView.View.ShaderMap);
 
-			Context.RHICmdList.TransitionResource(EResourceTransitionAccess::ERWBarrier, EResourceTransitionPipeline::EGfxToCompute, DestRenderTarget.UAV);
-			DispatchCS(Context.RHICmdList, Context, DestRenderTarget.UAV, LastEyeAdaptation);
-			Context.RHICmdList.TransitionResource(EResourceTransitionAccess::EReadable, EResourceTransitionPipeline::EComputeToGfx, DestRenderTarget.UAV, AsyncEndFence);
-
-			Context.RHICmdList.EndUpdateMultiFrameResource(DestRenderTarget.ShaderResourceTexture);
-		}
+		FComputeShaderUtils::AddPass(
+			GraphBuilder,
+			RDG_EVENT_NAME("HistogramEyeAdaptation (CS)"),
+			*ComputeShader,
+			PassParameters,
+			FIntVector(1, 1, 1));
 	}
 	else
 	{
-		// Inform MultiGPU systems that we're starting to update this texture for this frame
-		Context.RHICmdList.BeginUpdateMultiFrameResource(DestRenderTarget.ShaderResourceTexture);
+		FEyeAdaptationPS::FParameters* PassParameters = GraphBuilder.AllocParameters<FEyeAdaptationPS::FParameters>();
+		PassParameters->EyeAdaptation = GetEyeAdaptationParameters(View);
+		PassParameters->HistogramTexture = HistogramTexture;
+		PassParameters->RenderTargets[0] = FRenderTargetBinding(OutputTexture, ERenderTargetLoadAction::ELoad);
 
-		// we render to our own output render target, not the intermediate one created by the compositing system
-		// Set the view family's render target/viewport.
+		TShaderMapRef<FEyeAdaptationPS> PixelShader(ScreenPassView.View.ShaderMap);
 
-		FRHIRenderPassInfo RPInfo(DestRenderTarget.TargetableTexture, ERenderTargetActions::Load_Store);
-		TransitionRenderPassTargets(Context.RHICmdList, RPInfo);
-		Context.RHICmdList.BeginRenderPass(RPInfo, TEXT("EyeAdaptation"));
-		{
-			Context.SetViewportAndCallRHI(0, 0, 0.0f, DestSize.X, DestSize.Y, 1.0f);
-
-			FGraphicsPipelineStateInitializer GraphicsPSOInit;
-			Context.RHICmdList.ApplyCachedRenderTargets(GraphicsPSOInit);
-			GraphicsPSOInit.BlendState = TStaticBlendState<>::GetRHI();
-			GraphicsPSOInit.RasterizerState = TStaticRasterizerState<>::GetRHI();
-			GraphicsPSOInit.DepthStencilState = TStaticDepthStencilState<false, CF_Always>::GetRHI();
-
-			TShaderMapRef<FPostProcessVS> VertexShader(Context.GetShaderMap());
-			TShaderMapRef<FPostProcessEyeAdaptationPS> PixelShader(Context.GetShaderMap());
+		AddDrawScreenPass(
+			GraphBuilder,
+			RDG_EVENT_NAME("HistogramEyeAdaptation (PS)"),
+			ScreenPassView,
+			FScreenPassTextureViewport(OutputTexture),
+			FScreenPassTextureViewport(HistogramTexture),
+			*PixelShader,
+			PassParameters);
+	}
 
-			GraphicsPSOInit.BoundShaderState.VertexDeclarationRHI = GFilterVertexDeclaration.VertexDeclarationRHI;
-			GraphicsPSOInit.BoundShaderState.VertexShaderRHI = GETSAFERHISHADER_VERTEX(*VertexShader);
-			GraphicsPSOInit.BoundShaderState.PixelShaderRHI = GETSAFERHISHADER_PIXEL(*PixelShader);
-			GraphicsPSOInit.PrimitiveType = PT_TriangleList;
+	return OutputTexture;
+}
 
-			SetGraphicsPipelineState(Context.RHICmdList, GraphicsPSOInit);
+FRenderingCompositeOutputRef AddHistogramEyeAdaptationPass(FPostprocessContext& Context, FRenderingCompositeOutputRef Histogram)
+{
+	const FViewInfo& View = Context.View;
 
-			PixelShader->SetPS(Context, Context.RHICmdList, LastEyeAdaptation);
+	FRenderingCompositePass* Pass = Context.Graph.RegisterPass(
+		new(FMemStack::Get()) TRCPassForRDG<1, 1>(
+			[&View](FRenderingCompositePass* Pass, FRenderingCompositePassContext& InContext)
+	{
+		FRDGBuilder GraphBuilder(InContext.RHICmdList);
 
-			// Draw a quad mapping scene color to the view's render target
-			DrawRectangle(
-				Context.RHICmdList,
-				0, 0,
-				DestSize.X, DestSize.Y,
-				0, 0,
-				DestSize.X, DestSize.Y,
-				DestSize,
-				DestSize,
-				*VertexShader,
-				EDRF_UseTriangleOptimization);
-		}
-		Context.RHICmdList.EndRenderPass();
-		Context.RHICmdList.CopyToResolveTarget(DestRenderTarget.TargetableTexture, DestRenderTarget.ShaderResourceTexture, FResolveParams());
-		// Inform MultiGPU systems that we've finished updating this texture for this frame
-		Context.RHICmdList.EndUpdateMultiFrameResource(DestRenderTarget.ShaderResourceTexture);
-	}
+		FRDGTextureRef HistogramTexture = Pass->CreateRDGTextureForRequiredInput(GraphBuilder, ePId_Input0, TEXT("Histogram"));
 
-	Context.View.SetValidEyeAdaptation();
-}
+		const FEyeAdaptationParameters EyeAdaptationParameters = GetEyeAdaptationParameters(View);
 
-template <typename TRHICmdList>
-void FRCPassPostProcessEyeAdaptation::DispatchCS(TRHICmdList& RHICmdList, FRenderingCompositePassContext& Context, FRHIUnorderedAccessView* DestUAV, IPooledRenderTarget* LastEyeAdaptation)
-{
-	auto ShaderMap = Context.GetShaderMap();
-	TShaderMapRef<FPostProcessEyeAdaptationCS> ComputeShader(ShaderMap);
-	RHICmdList.SetComputeShader(ComputeShader->GetComputeShader());
+		FRDGTextureRef OutputTexture = AddHistogramEyeAdaptationPass(GraphBuilder, FScreenPassViewInfo(View), EyeAdaptationParameters, HistogramTexture);
 
-	ComputeShader->SetParameters(RHICmdList, Context, DestUAV, LastEyeAdaptation);
-	DispatchComputeShader(RHICmdList, *ComputeShader, 1, 1, 1);
-	ComputeShader->UnsetParameters(RHICmdList);
-}
+		Pass->ExtractRDGTextureForOutput(GraphBuilder, ePId_Output0, OutputTexture);
 
-float FRCPassPostProcessEyeAdaptation::GetFixedExposure(const FViewInfo& View)
-{
-	const FEyeAdaptationParameters Parameters = GetEyeAdaptationParameters(View);
-	
-	const float Exposure = (Parameters.MinAverageLuminance + Parameters.MaxAverageLuminance) * 0.5f;
+		GraphBuilder.Execute();
+	}));
 
-	const float ExposureScale = 1.0f / FMath::Max(0.0001f, Exposure);
+	Pass->SetInput(ePId_Input0, Histogram);
 
-	return ExposureScale * Parameters.ExposureCompensation;
+	return FRenderingCompositeOutputRef(Pass);
 }
 
 void FSceneViewState::UpdatePreExposure(FViewInfo& View)
@@ -548,7 +410,7 @@ void FSceneViewState::UpdatePreExposure(FViewInfo& View)
 		if (!IsMobileHDR())
 		{
 			// In gamma space, the exposure is fully applied in the pre-exposure (no post-exposure compensation)
-			PreExposure = FRCPassPostProcessEyeAdaptation::GetFixedExposure(View);
+			PreExposure = GetEyeAdaptationFixedExposure(View);
 		}
 	}
 	else if (bIsPreExposureRelevant)
@@ -585,18 +447,6 @@ void FSceneViewState::UpdatePreExposure(FViewInfo& View)
 	}
 }
 
-FPooledRenderTargetDesc FRCPassPostProcessEyeAdaptation::ComputeOutputDesc(EPassOutputId InPassOutputId) const
-{
-	// Specify invalid description to avoid getting intermediate rendertargets created.
-	// We want to use ViewState->GetEyeAdaptation() instead
-	FPooledRenderTargetDesc Ret;
-
-	Ret.DebugName = TEXT("EyeAdaptation");
-
-	return Ret;
-}
-
-
 /** Encapsulates the post process computation of Log2 Luminance pixel shader. */
 class FPostProcessBasicEyeAdaptationSetupPS : public FGlobalShader
 {
diff --git a/Engine/Source/Runtime/Renderer/Private/PostProcess/PostProcessEyeAdaptation.h b/Engine/Source/Runtime/Renderer/Private/PostProcess/PostProcessEyeAdaptation.h
--- a/Engine/Source/Runtime/Renderer/Private/PostProcess/PostProcessEyeAdaptation.h
+++ b/Engine/Source/Runtime/Renderer/Private/PostProcess/PostProcessEyeAdaptation.h
@@ -32,8 +32,10 @@ FORCEINLINE float Log2ToEV100(float Log2)
 	return Log2 - 0.263f; // Where .263 is log2(1.2)
 }
 
+// Returns whether the auto exposure method is supported by the feature level.
 bool IsAutoExposureMethodSupported(ERHIFeatureLevel::Type FeatureLevel, EAutoExposureMethod AutoExposureMethodId);
 
+// Returns the auto exposure method enabled by the view (including CVar override).
 EAutoExposureMethod GetAutoExposureMethod(const FViewInfo& View);
 
 BEGIN_SHADER_PARAMETER_STRUCT(FEyeAdaptationParameters, )
@@ -54,35 +56,8 @@ END_SHADER_PARAMETER_STRUCT()
 
 FEyeAdaptationParameters GetEyeAdaptationParameters(const FViewInfo& ViewInfo, ERHIFeatureLevel::Type MinFeatureLevel = ERHIFeatureLevel::SM5);
 
-// Computes the eye-adaptation from HDRHistogram.
-// ePId_Input0: HDRHistogram or nothing
-// derives from TRenderingCompositePassBase<InputCount, OutputCount> 
-class FRCPassPostProcessEyeAdaptation : public TRenderingCompositePassBase<1, 1>
-{
-public:
-	FRCPassPostProcessEyeAdaptation(bool bInIsComputePass)
-	{
-		bIsComputePass = bInIsComputePass;
-		bPreferAsyncCompute = false;
-		bPreferAsyncCompute &= (GNumAlternateFrameRenderingGroups == 1); // Can't handle multi-frame updates on async pipe
-	}
-
-	// Computes the a fix exposure to be used to replace the dynamic exposure when it's not available (< SM5).
-	static float GetFixedExposure(const FViewInfo& View);
-
-	// interface FRenderingCompositePass ---------
-	virtual void Process(FRenderingCompositePassContext& Context) override;
-	virtual void Release() override { delete this; }
-	virtual FPooledRenderTargetDesc ComputeOutputDesc(EPassOutputId InPassOutputId) const override;
-
-	virtual FRHIComputeFence* GetComputePassEndFence() const override { return AsyncEndFence; }
-
-private:
-	template <typename TRHICmdList>
-	void DispatchCS(TRHICmdList& RHICmdList, FRenderingCompositePassContext& Context, FRHIUnorderedAccessView* DestUAV, IPooledRenderTarget* LastEyeAdaptation);
-
-	FComputeFenceRHIRef AsyncEndFence;
-};
+// Computes the a fixed exposure to be used to replace the dynamic exposure when it's not supported (< SM5).
+float GetEyeAdaptationFixedExposure(const FViewInfo& View);
 
 // Write Log2(Luminance) in the alpha channel.
 // ePId_Input0: Half-Res HDR scene color
@@ -114,4 +89,6 @@ public:
 
 private:
 	FIntPoint DownsampledViewRect;
-};
\ No newline at end of file
+};
+
+FRenderingCompositeOutputRef AddHistogramEyeAdaptationPass(FPostprocessContext& Context, FRenderingCompositeOutputRef Histogram);
\ No newline at end of file
diff --git a/Engine/Source/Runtime/Renderer/Private/PostProcess/PostProcessTonemap.cpp b/Engine/Source/Runtime/Renderer/Private/PostProcess/PostProcessTonemap.cpp
--- a/Engine/Source/Runtime/Renderer/Private/PostProcess/PostProcessTonemap.cpp
+++ b/Engine/Source/Runtime/Renderer/Private/PostProcess/PostProcessTonemap.cpp
@@ -883,7 +883,7 @@ public:
 
 		{
 			// Compute a CPU-based default.  NB: reverts to "1" if SM5 feature level is not supported
-			float FixedExposure = FRCPassPostProcessEyeAdaptation::GetFixedExposure(Context.View);
+			float FixedExposure = GetEyeAdaptationFixedExposure(Context.View);
 			// Load a default value 
 			SetShaderValue(RHICmdList, ShaderRHI, DefaultEyeExposure, FixedExposure);
 		}
@@ -1581,7 +1581,7 @@ public:
 
 		SetShaderValue(RHICmdList, ShaderRHI, SRGBAwareTargetParam, bSRGBAwareTarget ? 1.0f : 0.0f );
 
-		float FixedExposure = FRCPassPostProcessEyeAdaptation::GetFixedExposure(Context.View);
+		float FixedExposure = GetEyeAdaptationFixedExposure(Context.View);
 		SetShaderValue(RHICmdList, ShaderRHI, DefaultEyeExposure, FixedExposure);
 	}
 };
diff --git a/Engine/Source/Runtime/Renderer/Private/PostProcess/PostProcessTonemap.h b/Engine/Source/Runtime/Renderer/Private/PostProcess/PostProcessTonemap.h
--- a/Engine/Source/Runtime/Renderer/Private/PostProcess/PostProcessTonemap.h
+++ b/Engine/Source/Runtime/Renderer/Private/PostProcess/PostProcessTonemap.h
@@ -197,7 +197,7 @@ public:
 		if (!PermutationVector.Get<FTonemapperVSUseAutoExposure>())
 		{
 			// Compute a CPU-based default.  NB: reverts to "1" if SM5 feature level is not supported
-			float FixedExposure = FRCPassPostProcessEyeAdaptation::GetFixedExposure(Context.View);
+			float FixedExposure = GetEyeAdaptationFixedExposure(Context.View);
 			// Load a default value 
 			SetShaderValue(Context.RHICmdList, ShaderRHI, DefaultEyeExposure, FixedExposure);
 		}
diff --git a/Engine/Source/Runtime/Renderer/Private/PostProcess/PostProcessing.cpp b/Engine/Source/Runtime/Renderer/Private/PostProcess/PostProcessing.cpp
--- a/Engine/Source/Runtime/Renderer/Private/PostProcess/PostProcessing.cpp
+++ b/Engine/Source/Runtime/Renderer/Private/PostProcess/PostProcessing.cpp
@@ -413,15 +413,6 @@ static FRenderingCompositeOutputRef AddPostProcessBasicEyeAdaptation(const FView
 	return FRenderingCompositeOutputRef(Node);
 }
 
-static FRenderingCompositeOutputRef AddPostProcessHistogramEyeAdaptation(FPostprocessContext& Context, FRenderingCompositeOutputRef& Histogram)
-{
-	const bool bIsComputePass = ShouldDoComputePostProcessing(Context.View);
-	FRenderingCompositePass* Node = Context.Graph.RegisterPass(new(FMemStack::Get()) FRCPassPostProcessEyeAdaptation(bIsComputePass));
-
-	Node->SetInput(ePId_Input0, Histogram);
-	return FRenderingCompositeOutputRef(Node);
-}
-
 static void AddVisualizeBloomOverlay(FPostprocessContext& Context, FRenderingCompositeOutputRef& HDRColor, FRenderingCompositeOutputRef& BloomOutputCombined)
 {
 	auto Node = Context.Graph.RegisterPass(new(FMemStack::Get()) FRCPassPostProcessVisualizeBloomOverlay());
@@ -1286,7 +1277,7 @@ void FPostProcessing::Process(FRHICommandListImmediate& RHICmdList, const FViewI
 				}
 				else
 				{
-					EyeAdaptation = AddPostProcessHistogramEyeAdaptation(Context, Histogram);
+					EyeAdaptation = AddHistogramEyeAdaptationPass(Context, Histogram);
 				}
 			}
 
