Add null check before executing histogram-based eye adaptation pass

diff --git a/Engine/Source/Runtime/Renderer/Private/PostProcess/PostProcessing.cpp b/Engine/Source/Runtime/Renderer/Private/PostProcess/PostProcessing.cpp
--- a/Engine/Source/Runtime/Renderer/Private/PostProcess/PostProcessing.cpp
+++ b/Engine/Source/Runtime/Renderer/Private/PostProcess/PostProcessing.cpp
@@ -1275,7 +1275,7 @@ void FPostProcessing::Process(FRHICommandListImmediate& RHICmdList, const FViewI
 						SceneColorDownsampleChain = EyeAdaptationDownsampleChain;
 					}
 				}
-				else
+				else if (bHistogramEnabled)
 				{
 					EyeAdaptation = AddHistogramEyeAdaptationPass(Context, Histogram);
 				}
