Fixed build warnings.

diff --git a/Engine/Source/Runtime/Renderer/Private/PostProcess/PostProcessEyeAdaptation.cpp b/Engine/Source/Runtime/Renderer/Private/PostProcess/PostProcessEyeAdaptation.cpp
--- a/Engine/Source/Runtime/Renderer/Private/PostProcess/PostProcessEyeAdaptation.cpp
+++ b/Engine/Source/Runtime/Renderer/Private/PostProcess/PostProcessEyeAdaptation.cpp
@@ -359,13 +359,18 @@ FRenderingCompositeOutputRef AddHistogramEyeAdaptationPass(FPostprocessContext&
 {
 	const FViewInfo& View = Context.View;
 
-	FRenderingCompositePass* Pass = Context.Graph.RegisterPass(
+	FRenderingCompositePass* EyeAdaptationPass = Context.Graph.RegisterPass(
 		new(FMemStack::Get()) TRCPassForRDG<1, 1>(
 			[&View](FRenderingCompositePass* Pass, FRenderingCompositePassContext& InContext)
 	{
 		FRDGBuilder GraphBuilder(InContext.RHICmdList);
 
-		FRDGTextureRef HistogramTexture = Pass->CreateRDGTextureForRequiredInput(GraphBuilder, ePId_Input0, TEXT("Histogram"));
+		FRDGTextureRef HistogramTexture = Pass->CreateRDGTextureForOptionalInput(GraphBuilder, ePId_Input0, TEXT("Histogram"));
+
+		if (!HistogramTexture)
+		{
+			HistogramTexture = GraphBuilder.RegisterExternalTexture(GSystemTextures.BlackDummy);
+		}
 
 		const FEyeAdaptationParameters EyeAdaptationParameters = GetEyeAdaptationParameters(View);
 
@@ -376,9 +381,9 @@ FRenderingCompositeOutputRef AddHistogramEyeAdaptationPass(FPostprocessContext&
 		GraphBuilder.Execute();
 	}));
 
-	Pass->SetInput(ePId_Input0, Histogram);
+	EyeAdaptationPass->SetInput(ePId_Input0, Histogram);
 
-	return FRenderingCompositeOutputRef(Pass);
+	return FRenderingCompositeOutputRef(EyeAdaptationPass);
 }
 
 void FSceneViewState::UpdatePreExposure(FViewInfo& View)
diff --git a/Engine/Source/Runtime/Renderer/Private/PostProcess/PostProcessing.cpp b/Engine/Source/Runtime/Renderer/Private/PostProcess/PostProcessing.cpp
--- a/Engine/Source/Runtime/Renderer/Private/PostProcess/PostProcessing.cpp
+++ b/Engine/Source/Runtime/Renderer/Private/PostProcess/PostProcessing.cpp
@@ -1275,7 +1275,8 @@ void FPostProcessing::Process(FRHICommandListImmediate& RHICmdList, const FViewI
 						SceneColorDownsampleChain = EyeAdaptationDownsampleChain;
 					}
 				}
-				else if (bHistogramEnabled)
+				// We run histogram eye adaptation even if no histogram exists to support the manual clamping mode.
+				else
 				{
 					EyeAdaptation = AddHistogramEyeAdaptationPass(Context, Histogram);
 				}
