Fixing merge compilation error

diff --git a/Engine/Source/Runtime/Renderer/Private/PostProcess/PostProcessMotionBlur.cpp b/Engine/Source/Runtime/Renderer/Private/PostProcess/PostProcessMotionBlur.cpp
--- a/Engine/Source/Runtime/Renderer/Private/PostProcess/PostProcessMotionBlur.cpp
+++ b/Engine/Source/Runtime/Renderer/Private/PostProcess/PostProcessMotionBlur.cpp
@@ -539,6 +539,7 @@ void ComputeMotionBlurVelocity(
 				ERenderTargetLoadAction::EClear,
 				ERenderTargetStoreAction::EStore,
 				ERenderTargetLoadAction::ENoAction,
+				ERenderTargetStoreAction::ENoAction,
 				FExclusiveDepthStencil::DepthWrite_StencilNop);
 
 		PassParameters->RenderTargets[0] =
