Fixing merge compilation error

diff --git a/Engine/Source/Runtime/Renderer/Private/PostProcess/PostProcessMotionBlur.cpp b/Engine/Source/Runtime/Renderer/Private/PostProcess/PostProcessMotionBlur.cpp
--- a/Engine/Source/Runtime/Renderer/Private/PostProcess/PostProcessMotionBlur.cpp
+++ b/Engine/Source/Runtime/Renderer/Private/PostProcess/PostProcessMotionBlur.cpp
@@ -537,9 +537,7 @@ void ComputeMotionBlurVelocity(
 			FDepthStencilBinding(
 				VelocityTileDepthTexture,
 				ERenderTargetLoadAction::EClear,
-				ERenderTargetStoreAction::EStore,
 				ERenderTargetLoadAction::ENoAction,
-				ERenderTargetStoreAction::ENoAction,
 				FExclusiveDepthStencil::DepthWrite_StencilNop);
 
 		PassParameters->RenderTargets[0] =
