Added helper function to FScreenPassView to compute the correct load action when a pass is writing to every pixel. Helps determine whether to clear if using an HMD mask.

diff --git a/Engine/Source/Runtime/Renderer/Private/ScreenPass.cpp b/Engine/Source/Runtime/Renderer/Private/ScreenPass.cpp
--- a/Engine/Source/Runtime/Renderer/Private/ScreenPass.cpp
+++ b/Engine/Source/Runtime/Renderer/Private/ScreenPass.cpp
@@ -137,6 +137,11 @@ FScreenPassViewInfo::FScreenPassViewInfo(const FViewInfo& InView)
 	, bUseComputePasses(ShouldDoComputePostProcessing(InView))
 {}
 
+ERenderTargetLoadAction FScreenPassViewInfo::GetOverwriteLoadAction() const
+{
+	return bHasHMDMask ? ERenderTargetLoadAction::EClear : ERenderTargetLoadAction::ENoAction;
+}
+
 void SetScreenPassPipelineState(FRHICommandListImmediate& RHICmdList, const FScreenPassDrawInfo& ScreenPassDraw)
 {
 	FRHIPixelShader* PixelShaderRHI = GETSAFERHISHADER_PIXEL(ScreenPassDraw.PixelShader);
diff --git a/Engine/Source/Runtime/Renderer/Private/ScreenPass.h b/Engine/Source/Runtime/Renderer/Private/ScreenPass.h
--- a/Engine/Source/Runtime/Renderer/Private/ScreenPass.h
+++ b/Engine/Source/Runtime/Renderer/Private/ScreenPass.h
@@ -144,6 +144,10 @@ public:
 	FScreenPassViewInfo(const FViewInfo& InView);
 	FScreenPassViewInfo(const FScreenPassViewInfo&) = default;
 
+	// Returns the load action we should use when we expect to overwrite all relevant pixels.
+	// Takes into account the HMD mesh.
+	ERenderTargetLoadAction GetOverwriteLoadAction() const;
+
 	const FViewInfo& View;
 
 	// The vertex shader used by draw screen pass. Cached here to avoid many lookups.
