Ported basic eye adaptation setup pass to RDG.

diff --git a/Engine/Shaders/Private/PostProcessEyeAdaptation.usf b/Engine/Shaders/Private/PostProcessEyeAdaptation.usf
--- a/Engine/Shaders/Private/PostProcessEyeAdaptation.usf
+++ b/Engine/Shaders/Private/PostProcessEyeAdaptation.usf
@@ -136,13 +136,16 @@ void MainCS(uint2 DispatchThreadId : SV_DispatchThreadID)
 }
 #endif
 
+Texture2D SceneColorTexture;
+
+SamplerState SceneColorSampler;
 
 void MainBasicEyeAdaptationSetupPS(
 	noperspective float4 UVAndScreenPos : TEXCOORD0,
 	out float4 OutColor : SV_Target0)
 {
 	float2 UV = UVAndScreenPos.xy;
-	OutColor = Texture2DSample(PostprocessInput0, PostprocessInput0Sampler, UV);
+	OutColor = Texture2DSample(SceneColorTexture, SceneColorSampler, UV);
 
 	// Use max to ensure intensity is never zero (so the following log is well behaved)
 	float Intensity = max(dot(OutColor.xyz, float3(0.2126, 0.7152, 0.0722)), EyeAdaptation_LuminanceMin);
diff --git a/Engine/Source/Runtime/Renderer/Private/PostProcess/PostProcessEyeAdaptation.cpp b/Engine/Source/Runtime/Renderer/Private/PostProcess/PostProcessEyeAdaptation.cpp
--- a/Engine/Source/Runtime/Renderer/Private/PostProcess/PostProcessEyeAdaptation.cpp
+++ b/Engine/Source/Runtime/Renderer/Private/PostProcess/PostProcessEyeAdaptation.cpp
@@ -146,7 +146,9 @@ float GetAutoExposureCompensation(const FViewInfo& View)
 		(EngineShowFlags.VisualizeBuffer && View.CurrentBufferVisualizationMode != NAME_None) ||
 		EngineShowFlags.RayTracingDebug ||
 		EngineShowFlags.VisualizeDistanceFieldAO ||
-		EngineShowFlags.VisualizeGlobalDistanceField;
+		EngineShowFlags.VisualizeGlobalDistanceField ||
+		EngineShowFlags.CollisionVisibility ||
+		EngineShowFlags.CollisionPawn;
 
 	if (bSkipExposureBias)
 	{
@@ -249,7 +251,9 @@ float GetEyeAdaptationFixedExposure(const FViewInfo& View)
 	return ExposureScale * Parameters.ExposureCompensation;
 }
 
-const ERHIFeatureLevel::Type BasicEyeAdaptationMinFeatureLevel = ERHIFeatureLevel::ES3_1;
+//////////////////////////////////////////////////////////////////////////
+//! Histogram Eye Adaptation
+//////////////////////////////////////////////////////////////////////////
 
 class FEyeAdaptationShader : public FGlobalShader
 {
@@ -268,7 +272,6 @@ public:
 	}
 
 	FEyeAdaptationShader() = default;
-
 	FEyeAdaptationShader(const CompiledShaderInitializerType& Initializer)
 		: FGlobalShader(Initializer)
 	{}
@@ -386,213 +389,112 @@ FRenderingCompositeOutputRef AddHistogramEyeAdaptationPass(FPostprocessContext&
 	return FRenderingCompositeOutputRef(EyeAdaptationPass);
 }
 
-void FSceneViewState::UpdatePreExposure(FViewInfo& View)
-{
-	const FSceneViewFamily& ViewFamily = *View.Family;
-	// One could use the IsRichView functionality to check if we need to update pre-exposure, 
-	// but this is too limiting for certain view. For instance shader preview doesn't have 
-	// volumetric lighting enabled, which makes the view be flagged as rich, and not updating 
-	// the pre-exposition value.
-	const bool bIsPreExposureRelevant =
-		ViewFamily.EngineShowFlags.EyeAdaptation && // Controls whether scene luminance is computed at all.
-		ViewFamily.EngineShowFlags.Lighting &&
-		ViewFamily.EngineShowFlags.PostProcessing &&
-		ViewFamily.bResolveScene &&
-		!ViewFamily.EngineShowFlags.LightMapDensity &&
-		!ViewFamily.EngineShowFlags.StationaryLightOverlap &&
-		!ViewFamily.EngineShowFlags.LightComplexity &&
-		!ViewFamily.EngineShowFlags.LODColoration &&
-		!ViewFamily.EngineShowFlags.HLODColoration &&
-		!ViewFamily.EngineShowFlags.LevelColoration &&
-		!ViewFamily.EngineShowFlags.VisualizeBloom;
-
-
-	PreExposure = 1.f;
-	bUpdateLastExposure = false;
-
-	if (IsMobilePlatform(View.GetShaderPlatform()))
-	{
-		if (!IsMobileHDR())
-		{
-			// In gamma space, the exposure is fully applied in the pre-exposure (no post-exposure compensation)
-			PreExposure = GetEyeAdaptationFixedExposure(View);
-		}
-	}
-	else if (bIsPreExposureRelevant)
-	{
-		if (UsePreExposure(View.GetShaderPlatform()))
-		{
-			const float PreExposureOverride = CVarEyeAdaptationPreExposureOverride.GetValueOnRenderThread();
-			const float LastExposure = View.GetLastEyeAdaptationExposure();
-			if (PreExposureOverride > 0)
-			{
-				PreExposure = PreExposureOverride;
-			}
-			else if (LastExposure > 0)
-			{
-				PreExposure = LastExposure;
-			}
-
-			bUpdateLastExposure = true;
-		}
-		// The exposure compensation curves require the scene average luminance
-		else if (View.FinalPostProcessSettings.AutoExposureBiasCurve)
-		{
-			bUpdateLastExposure = true;
-		}
-	}
-
-	// Update the pre-exposure value on the actual view
-	View.PreExposure = PreExposure;
+//////////////////////////////////////////////////////////////////////////
+//! Basic Eye Adaptation
+//////////////////////////////////////////////////////////////////////////
 
-	// Update the pre exposure of all temporal histories.
-	if (!View.bStatePrevViewInfoIsReadOnly)
-	{
-		PrevFrameViewInfo.SceneColorPreExposure = PreExposure;
-	}
-}
+const ERHIFeatureLevel::Type BasicEyeAdaptationMinFeatureLevel = ERHIFeatureLevel::ES3_1;
 
-/** Encapsulates the post process computation of Log2 Luminance pixel shader. */
-class FPostProcessBasicEyeAdaptationSetupPS : public FGlobalShader
+class FBasicEyeAdaptationShader : public FGlobalShader
 {
-	DECLARE_SHADER_TYPE(FPostProcessBasicEyeAdaptationSetupPS, Global);
-
-	BEGIN_SHADER_PARAMETER_STRUCT(FParameters, )
-		SHADER_PARAMETER_STRUCT(FEyeAdaptationParameters, EyeAdaptation)
-	END_SHADER_PARAMETER_STRUCT()
-
+public:
 	static bool ShouldCompilePermutation(const FGlobalShaderPermutationParameters& Parameters)
 	{
 		return IsFeatureLevelSupported(Parameters.Platform, BasicEyeAdaptationMinFeatureLevel);
 	}
 
-	static void ModifyCompilationEnvironment(const FGlobalShaderPermutationParameters& Parameters, FShaderCompilerEnvironment& OutEnvironment)
-	{
-		FGlobalShader::ModifyCompilationEnvironment(Parameters, OutEnvironment);
-	}
-
-	/** Default constructor. */
-	FPostProcessBasicEyeAdaptationSetupPS() {}
-
-public:
-	FPostProcessPassParameters PostprocessParameter;
-	FShaderParameter EyeAdaptationParams;
-
-	/** Initialization constructor. */
-	FPostProcessBasicEyeAdaptationSetupPS(const ShaderMetaType::CompiledShaderInitializerType& Initializer)
+	FBasicEyeAdaptationShader() = default;
+	FBasicEyeAdaptationShader(const CompiledShaderInitializerType& Initializer)
 		: FGlobalShader(Initializer)
-	{
-		BindForLegacyShaderParameters<FParameters>(this, Initializer.ParameterMap);
-		PostprocessParameter.Bind(Initializer.ParameterMap);
-	}
-
-	void SetPS(const FRenderingCompositePassContext& Context)
-	{
-		FRHIPixelShader* ShaderRHI = GetPixelShader();
-
-		FGlobalShader::SetParameters<FViewUniformShaderParameters>(Context.RHICmdList, ShaderRHI, Context.View.ViewUniformBuffer);
-		PostprocessParameter.SetPS(Context.RHICmdList, ShaderRHI, Context, TStaticSamplerState<SF_Point, AM_Clamp, AM_Clamp, AM_Clamp>::GetRHI());
+	{}
+};
 
-		{
-			FParameters PassParameters;
-			PassParameters.EyeAdaptation = GetEyeAdaptationParameters(Context.View, BasicEyeAdaptationMinFeatureLevel);
-			SetShaderParameters(Context.RHICmdList, this, ShaderRHI, PassParameters);
-		}
-	}
+/** Computes scaled and biased luma for the input scene color and puts it in the alpha channel. */
+class FBasicEyeAdaptationSetupPS : public FBasicEyeAdaptationShader
+{
+public:
+	DECLARE_GLOBAL_SHADER(FBasicEyeAdaptationSetupPS);
+	SHADER_USE_PARAMETER_STRUCT(FBasicEyeAdaptationSetupPS, FBasicEyeAdaptationShader);
 
-	// FShader interface.
-	virtual bool Serialize(FArchive& Ar) override
-	{
-		bool bShaderHasOutdatedParameters = FGlobalShader::Serialize(Ar);
-		Ar << PostprocessParameter << EyeAdaptationParams;
-		return bShaderHasOutdatedParameters;
-	}
+	BEGIN_SHADER_PARAMETER_STRUCT(FParameters, )
+		SHADER_PARAMETER_STRUCT(FEyeAdaptationParameters, EyeAdaptation)
+		SHADER_PARAMETER_RDG_TEXTURE(Texture2D, SceneColorTexture)
+		SHADER_PARAMETER_SAMPLER(SamplerState, SceneColorSampler)
+		RENDER_TARGET_BINDING_SLOTS()
+	END_SHADER_PARAMETER_STRUCT()
 };
 
-IMPLEMENT_SHADER_TYPE(, FPostProcessBasicEyeAdaptationSetupPS, TEXT("/Engine/Private/PostProcessEyeAdaptation.usf"), TEXT("MainBasicEyeAdaptationSetupPS"), SF_Pixel);
+IMPLEMENT_GLOBAL_SHADER(FBasicEyeAdaptationSetupPS, "/Engine/Private/PostProcessEyeAdaptation.usf", "MainBasicEyeAdaptationSetupPS", SF_Pixel);
 
-void FRCPassPostProcessBasicEyeAdaptationSetUp::Process(FRenderingCompositePassContext& Context)
+FRDGTextureRef AddBasicEyeAdaptationSetupPass(
+	FRDGBuilder& GraphBuilder,
+	const FScreenPassViewInfo& ScreenPassView,
+	const FEyeAdaptationParameters& EyeAdaptationParameters,
+	FRDGTextureRef SceneColorTexture,
+	FIntRect SceneColorViewRect)
 {
-	const FPooledRenderTargetDesc* InputDesc = GetInputDesc(ePId_Input0);
+	check(SceneColorTexture);
 
-	if (!InputDesc)
-	{
-		// input is not hooked up correctly
-		return;
-	}
+	FRDGTextureDesc OutputDesc = SceneColorTexture->Desc;
+	OutputDesc.Reset();
+	OutputDesc.DebugName = TEXT("EyeAdaptationBasicSetup");
+	// Require alpha channel for log2 information.
+	OutputDesc.Format = PF_FloatRGBA;
+	OutputDesc.Flags |= GFastVRamConfig.EyeAdaptation;
 
-	const FSceneRenderTargetItem& DestRenderTarget = PassOutputs[0].RequestSurface(Context);
+	FRDGTextureRef OutputTexture = GraphBuilder.CreateTexture(OutputDesc, TEXT("BasicEyeAdaptationSetup"));
 
-	ERenderTargetLoadAction LoadAction = Context.GetLoadActionForRenderTarget(DestRenderTarget);
+	const FScreenPassTextureViewport Viewport(SceneColorViewRect, OutputTexture);
 
-	FRHIRenderPassInfo RPInfo(DestRenderTarget.TargetableTexture, MakeRenderTargetActions(LoadAction, ERenderTargetStoreAction::EStore));
-	Context.RHICmdList.BeginRenderPass(RPInfo, TEXT("PostProcessBasicEyeAdaptationSetUp"));
-	{
-		const FViewInfo& View = Context.View;
-		const FSceneViewFamily& ViewFamily = *(View.Family);
+	FBasicEyeAdaptationSetupPS::FParameters* PassParameters = GraphBuilder.AllocParameters<FBasicEyeAdaptationSetupPS::FParameters>();
+	PassParameters->EyeAdaptation = GetEyeAdaptationParameters(ScreenPassView.View, BasicEyeAdaptationMinFeatureLevel);
+	PassParameters->SceneColorTexture = SceneColorTexture;
+	PassParameters->SceneColorSampler = TStaticSamplerState<SF_Point, AM_Clamp, AM_Clamp, AM_Clamp>::GetRHI();
+	PassParameters->RenderTargets[0] = FRenderTargetBinding(OutputTexture, ScreenPassView.GetOverwriteLoadAction());
 
-		FIntPoint SrcSize = InputDesc->Extent;
-		FIntPoint DestSize = PassOutputs[0].RenderTargetDesc.Extent;
+	TShaderMapRef<FBasicEyeAdaptationSetupPS> PixelShader(ScreenPassView.View.ShaderMap);
 
-		// e.g. 4 means the input texture is 4x smaller than the buffer size
-		uint32 ScaleFactor = Context.ReferenceBufferSize.X / SrcSize.X;
+	AddDrawScreenPass(
+		GraphBuilder,
+		RDG_EVENT_NAME("BasicEyeAdaptationSetup (PS) %dx%d", Viewport.Rect.Width(), Viewport.Rect.Height()),
+		ScreenPassView,
+		Viewport,
+		Viewport,
+		*PixelShader,
+		PassParameters);
 
-		FIntRect SrcRect = Context.SceneColorViewRect / ScaleFactor;
-		FIntRect DestRect = SrcRect;
-
-		SCOPED_DRAW_EVENTF(Context.RHICmdList, PostProcessBasicEyeAdaptationSetup, TEXT("PostProcessBasicEyeAdaptationSetup %dx%d"), DestRect.Width(), DestRect.Height());
+	return OutputTexture;
+}
 
-		Context.SetViewportAndCallRHI(0, 0, 0.0f, DestSize.X, DestSize.Y, 1.0f);
+FRenderingCompositeOutputRef AddBasicEyeAdaptationSetupPass(
+	FPostprocessContext& Context,
+	FRenderingCompositeOutputRef SceneColor,
+	FIntRect SceneColorViewRect)
+{
+	const FViewInfo& View = Context.View;
 
-		FGraphicsPipelineStateInitializer GraphicsPSOInit;
-		Context.RHICmdList.ApplyCachedRenderTargets(GraphicsPSOInit);
-		GraphicsPSOInit.BlendState = TStaticBlendState<>::GetRHI();
-		GraphicsPSOInit.RasterizerState = TStaticRasterizerState<>::GetRHI();
-		GraphicsPSOInit.DepthStencilState = TStaticDepthStencilState<false, CF_Always>::GetRHI();
+	FRenderingCompositePass* Pass = Context.Graph.RegisterPass(
+		new(FMemStack::Get()) TRCPassForRDG<1, 1>(
+			[&View, SceneColorViewRect](FRenderingCompositePass* InPass, FRenderingCompositePassContext& InContext)
+	{
+		FRDGBuilder GraphBuilder(InContext.RHICmdList);
 
-		TShaderMapRef<FPostProcessVS> VertexShader(Context.GetShaderMap());
-		TShaderMapRef<FPostProcessBasicEyeAdaptationSetupPS> PixelShader(Context.GetShaderMap());
+		const FEyeAdaptationParameters EyeAdaptationParameters = GetEyeAdaptationParameters(View);
 
-		GraphicsPSOInit.BoundShaderState.VertexDeclarationRHI = GFilterVertexDeclaration.VertexDeclarationRHI;
-		GraphicsPSOInit.BoundShaderState.VertexShaderRHI = GETSAFERHISHADER_VERTEX(*VertexShader);
-		GraphicsPSOInit.BoundShaderState.PixelShaderRHI = GETSAFERHISHADER_PIXEL(*PixelShader);
-		GraphicsPSOInit.PrimitiveType = PT_TriangleList;
+		FRDGTextureRef SceneColorTexture = InPass->CreateRDGTextureForRequiredInput(GraphBuilder, ePId_Input0, TEXT("SceneColor"));
 
-		SetGraphicsPipelineState(Context.RHICmdList, GraphicsPSOInit);
+		FRDGTextureRef OutputTexture = AddBasicEyeAdaptationSetupPass(GraphBuilder, FScreenPassViewInfo(View), EyeAdaptationParameters, SceneColorTexture, SceneColorViewRect);
 
-		PixelShader->SetPS(Context);
+		InPass->ExtractRDGTextureForOutput(GraphBuilder, ePId_Output0, OutputTexture);
 
-		DrawPostProcessPass(
-			Context.RHICmdList,
-			DestRect.Min.X, DestRect.Min.Y,
-			DestRect.Width(), DestRect.Height(),
-			SrcRect.Min.X, SrcRect.Min.Y,
-			SrcRect.Width(), SrcRect.Height(),
-			DestSize,
-			SrcSize,
-			*VertexShader,
-			View.StereoPass,
-			Context.HasHmdMesh(),
-			EDRF_UseTriangleOptimization);
-	}
-	Context.RHICmdList.EndRenderPass();
-	Context.RHICmdList.CopyToResolveTarget(DestRenderTarget.TargetableTexture, DestRenderTarget.ShaderResourceTexture, FResolveParams());
-}
+		GraphBuilder.Execute();
+	}));
 
-FPooledRenderTargetDesc FRCPassPostProcessBasicEyeAdaptationSetUp::ComputeOutputDesc(EPassOutputId InPassOutputId) const
-{
-	FPooledRenderTargetDesc Ret = GetInput(ePId_Input0)->GetOutput()->RenderTargetDesc;
+	Pass->SetInput(ePId_Input0, SceneColor);
 
-	Ret.Reset();
-	Ret.DebugName = TEXT("EyeAdaptationBasicSetup");
-	// Require alpha channel for log2 information.
-	Ret.Format = PF_FloatRGBA;
-	Ret.Flags |= GFastVRamConfig.EyeAdaptation;
-	return Ret;
+	return FRenderingCompositeOutputRef(Pass);
 }
 
-
-
 class FPostProcessLogLuminance2ExposureScaleBase: public FGlobalShader
 {
 public:
@@ -950,3 +852,69 @@ TRefCountPtr<IPooledRenderTarget>&  FSceneViewState::FEyeAdaptationRTManager::Ge
 
 	return PooledRenderTarget[BufferNumber];
 }
+
+void FSceneViewState::UpdatePreExposure(FViewInfo& View)
+{
+	const FSceneViewFamily& ViewFamily = *View.Family;
+
+	// One could use the IsRichView functionality to check if we need to update pre-exposure, 
+	// but this is too limiting for certain view. For instance shader preview doesn't have 
+	// volumetric lighting enabled, which makes the view be flagged as rich, and not updating 
+	// the pre-exposition value.
+	const bool bIsPreExposureRelevant =
+		ViewFamily.EngineShowFlags.EyeAdaptation && // Controls whether scene luminance is computed at all.
+		ViewFamily.EngineShowFlags.Lighting &&
+		ViewFamily.EngineShowFlags.PostProcessing &&
+		ViewFamily.bResolveScene &&
+		!ViewFamily.EngineShowFlags.LightMapDensity &&
+		!ViewFamily.EngineShowFlags.StationaryLightOverlap &&
+		!ViewFamily.EngineShowFlags.LightComplexity &&
+		!ViewFamily.EngineShowFlags.LODColoration &&
+		!ViewFamily.EngineShowFlags.HLODColoration &&
+		!ViewFamily.EngineShowFlags.LevelColoration &&
+		!ViewFamily.EngineShowFlags.VisualizeBloom;
+
+	PreExposure = 1.f;
+	bUpdateLastExposure = false;
+
+	if (IsMobilePlatform(View.GetShaderPlatform()))
+	{
+		if (!IsMobileHDR())
+		{
+			// In gamma space, the exposure is fully applied in the pre-exposure (no post-exposure compensation)
+			PreExposure = GetEyeAdaptationFixedExposure(View);
+		}
+	}
+	else if (bIsPreExposureRelevant)
+	{
+		if (UsePreExposure(View.GetShaderPlatform()))
+		{
+			const float PreExposureOverride = CVarEyeAdaptationPreExposureOverride.GetValueOnRenderThread();
+			const float LastExposure = View.GetLastEyeAdaptationExposure();
+			if (PreExposureOverride > 0)
+			{
+				PreExposure = PreExposureOverride;
+			}
+			else if (LastExposure > 0)
+			{
+				PreExposure = LastExposure;
+			}
+
+			bUpdateLastExposure = true;
+		}
+		// The exposure compensation curves require the scene average luminance
+		else if (View.FinalPostProcessSettings.AutoExposureBiasCurve)
+		{
+			bUpdateLastExposure = true;
+		}
+	}
+
+	// Update the pre-exposure value on the actual view
+	View.PreExposure = PreExposure;
+
+	// Update the pre exposure of all temporal histories.
+	if (!View.bStatePrevViewInfoIsReadOnly)
+	{
+		PrevFrameViewInfo.SceneColorPreExposure = PreExposure;
+	}
+}
\ No newline at end of file
diff --git a/Engine/Source/Runtime/Renderer/Private/PostProcess/PostProcessEyeAdaptation.h b/Engine/Source/Runtime/Renderer/Private/PostProcess/PostProcessEyeAdaptation.h
--- a/Engine/Source/Runtime/Renderer/Private/PostProcess/PostProcessEyeAdaptation.h
+++ b/Engine/Source/Runtime/Renderer/Private/PostProcess/PostProcessEyeAdaptation.h
@@ -59,18 +59,10 @@ FEyeAdaptationParameters GetEyeAdaptationParameters(const FViewInfo& ViewInfo, E
 // Computes the a fixed exposure to be used to replace the dynamic exposure when it's not supported (< SM5).
 float GetEyeAdaptationFixedExposure(const FViewInfo& View);
 
-// Write Log2(Luminance) in the alpha channel.
-// ePId_Input0: Half-Res HDR scene color
-// derives from TRenderingCompositePassBase<InputCount, OutputCount> 
-class FRCPassPostProcessBasicEyeAdaptationSetUp : public TRenderingCompositePassBase<1, 1>
-{
-public:
-	
-	// interface FRenderingCompositePass ---------
-	virtual void Process(FRenderingCompositePassContext& Context) override;
-	virtual FPooledRenderTargetDesc ComputeOutputDesc(EPassOutputId InPassOutputId) const override;
-	virtual void Release() override { delete this; }
-};
+FRenderingCompositeOutputRef AddBasicEyeAdaptationSetupPass(
+	FPostprocessContext& Context,
+	FRenderingCompositeOutputRef SceneColor,
+	FIntRect SceneColorViewRect);
 
 // ePId_Input0: Downsampled SceneColor Log
 // derives from TRenderingCompositePassBase<InputCount, OutputCount> 
diff --git a/Engine/Source/Runtime/Renderer/Private/PostProcess/PostProcessing.cpp b/Engine/Source/Runtime/Renderer/Private/PostProcess/PostProcessing.cpp
--- a/Engine/Source/Runtime/Renderer/Private/PostProcess/PostProcessing.cpp
+++ b/Engine/Source/Runtime/Renderer/Private/PostProcess/PostProcessing.cpp
@@ -233,22 +233,33 @@ public:
 
 		const EDownsampleQuality DownsampleQuality = GetDownsampleQuality();
 
+		const auto DivideAndRoundUpClamped = [](const FIntRect Rect, uint32 DownsampleFactor)
+		{
+			FIntRect Out = FIntRect::DivideAndRoundUp(Rect, DownsampleFactor);
+			Out.Max.X = FMath::Max(1, Out.Max.X);
+			Out.Max.Y = FMath::Max(1, Out.Max.Y);
+			return Out;
+		};
+
+		FIntRect InputViewRect = DivideAndRoundUpClamped(InContext.View.ViewRect, 2);
+
 		// Queue the down samples. 
 		for (int i = 1; i < DownSampleStages; i++)
 		{
 			PostProcessDownsamples[i] = AddDownsamplePass(Context.Graph, PassLabels[i], PostProcessDownsamples[i - 1], DownsampleQuality);
 
+			InputViewRect = DivideAndRoundUpClamped(InputViewRect, 2);
+
 			// Add log2 data to the alpha channel after doing the 1st (i==1) down sample pass
-			if (bHasLog2Alpha && i == 1 ) {
-				FRenderingCompositePass* BasicEyeSetupPass = Context.Graph.RegisterPass(new(FMemStack::Get()) FRCPassPostProcessBasicEyeAdaptationSetUp());
-				BasicEyeSetupPass->SetInput(ePId_Input0, PostProcessDownsamples[i]);
-				PostProcessDownsamples[i] = FRenderingCompositeOutputRef(BasicEyeSetupPass);
+			if (bHasLog2Alpha && i == 1 )
+			{
+				PostProcessDownsamples[i] = AddBasicEyeAdaptationSetupPass(InContext, PostProcessDownsamples[i], InputViewRect);
 			}
 		}
 
 		// Calculate the final viewrect size (matching AddDownsamplePass behavior)
-		FinalViewRectSize.X = FMath::Max(1, FMath::DivideAndRoundUp(InContext.View.ViewRect.Width(), 1 << DownSampleStages));
-		FinalViewRectSize.Y = FMath::Max(1, FMath::DivideAndRoundUp(InContext.View.ViewRect.Height(), 1 << DownSampleStages));
+		FinalViewRectSize.X = InputViewRect.Width();
+		FinalViewRectSize.Y = InputViewRect.Height();
 	}
 
 	// The number of elements in the array.
