Fix subsurface shader compilation for vulkan.

diff --git a/Engine/Shaders/Private/PostProcessSubsurface.usf b/Engine/Shaders/Private/PostProcessSubsurface.usf
--- a/Engine/Shaders/Private/PostProcessSubsurface.usf
+++ b/Engine/Shaders/Private/PostProcessSubsurface.usf
@@ -396,32 +396,25 @@ float4 SetupSubsurfaceForOnePixel(float2 UVSceneColor)
 // Indirect pass
 
 // Indirect dispatch parameters
-struct SubsurfaceGroupItem
-{
-	float2 GroupXY; // Which tile
-	int Count;		// The first element is used as the counter.
-	int Dummy;		// Alignment
-};
-
-RWStructuredBuffer<SubsurfaceGroupItem> RWSeparableGroupBuffer;
-RWStructuredBuffer<SubsurfaceGroupItem> RWBurleyGroupBuffer;
+RWStructuredBuffer<uint> RWSeparableGroupBuffer;
+RWStructuredBuffer<uint> RWBurleyGroupBuffer;
 RWBuffer<uint> RWIndirectDispatchArgsBuffer;
-StructuredBuffer<SubsurfaceGroupItem> GroupBuffer;
+StructuredBuffer<uint> GroupBuffer;
 
 
 [numthreads(1, 1, 1)]
 void InitValueBufferCS()
 {
 	// Clear the first buffer element which is used as the buffer counter.
-	RWSeparableGroupBuffer[0].Count = 0;
-	RWBurleyGroupBuffer[0].Count = 0;
+	RWSeparableGroupBuffer[0] = 0;
+	RWBurleyGroupBuffer[0] = 0;
 }
 
 // fixed to 64 by now
 [numthreads(1, 1, 1)]
 void BuildIndirectDispatchArgsCS()
 {
-	uint ValueCount = min(SubsurfaceUniformParameters.MaxGroupCount, GroupBuffer[0].Count);
+	uint ValueCount = min(SubsurfaceUniformParameters.MaxGroupCount, GroupBuffer[0]);
 	RWIndirectDispatchArgsBuffer[0] = ValueCount;// (ValueCount + (THREAD_SIZE_1D*THREAD_SIZE_1D) - 1) / (THREAD_SIZE_1D*THREAD_SIZE_1D);
 	RWIndirectDispatchArgsBuffer[1] = 1;
 	RWIndirectDispatchArgsBuffer[2] = 1;
@@ -456,14 +449,15 @@ void AddSeparableGroup(uint2 GroupXY)
 	//buffer counter is stored in the .Count in the first element.
 	int IndexToStore = 0;
 
-	InterlockedAdd(RWSeparableGroupBuffer[0].Count, 1, IndexToStore);
+	InterlockedAdd(RWSeparableGroupBuffer[0], 1, IndexToStore);
 
 	if (IndexToStore >= SubsurfaceUniformParameters.MaxGroupCount)
 	{
 		return;
 	}
 
-	RWSeparableGroupBuffer[IndexToStore + 1].GroupXY = GroupXY;
+	RWSeparableGroupBuffer[2 *(IndexToStore + 1) + 0] = GroupXY.x;
+	RWSeparableGroupBuffer[2 *(IndexToStore + 1) + 1] = GroupXY.y;
 }
 
 void AddBurleyGroup(uint2 GroupXY)
@@ -471,14 +465,15 @@ void AddBurleyGroup(uint2 GroupXY)
 	//buffer counter is stored in the .Count in the first element.
 	int IndexToStore = 0;
 
-	InterlockedAdd(RWBurleyGroupBuffer[0].Count, 1, IndexToStore);
+	InterlockedAdd(RWBurleyGroupBuffer[0], 1, IndexToStore);
 
 	if (IndexToStore >= SubsurfaceUniformParameters.MaxGroupCount)
 	{
 		return;
 	}
 
-	RWBurleyGroupBuffer[IndexToStore + 1].GroupXY = GroupXY;
+	RWBurleyGroupBuffer[2*(IndexToStore + 1 ) + 0] = GroupXY.x;
+	RWBurleyGroupBuffer[2*(IndexToStore + 1 ) + 1] = GroupXY.y;
 }
 
 void AddSubsurfaceComputeGroup(uint GI, uint Type, uint2 GroupXY)
@@ -612,8 +607,7 @@ void MainIndirectDispatchCS(
 )
 {
 	// The first element is used for counting
-	SubsurfaceGroupItem Item = GroupBuffer[IndirectG_ID.x + 1];
-	uint2 G_ID = Item.GroupXY;
+	uint2 G_ID = uint2(GroupBuffer[2*(IndirectG_ID.x + 1) + 0], GroupBuffer[2*(IndirectG_ID.x + 1)+1]);
 	uint2 DT_ID = G_ID * float2(THREAD_SIZE_X, THREAD_SIZE_Y) + float2(GI % THREAD_SIZE_X, GI / THREAD_SIZE_X);
 
 
diff --git a/Engine/Shaders/Private/SeparableSSS.ush b/Engine/Shaders/Private/SeparableSSS.ush
--- a/Engine/Shaders/Private/SeparableSSS.ush
+++ b/Engine/Shaders/Private/SeparableSSS.ush
@@ -121,7 +121,7 @@
 #define SSSS_FLATTEN FLATTEN
 #define SSSS_BRANCH BRANCH
 #define SSSS_UNROLL UNROLL
-
+#define SSSS_COMPUTESHADER COMPUTESHADER
 
 //-----------------------------------------------------------------------------
 // Separable SSS Transmittance Function
@@ -203,10 +203,11 @@ float4 SSSSBlurPS(float2 BufferUV, float2 dir, bool initStencil)
 	// 0..1
 	float SSSStrength = GetSubsurfaceStrength(BufferUV);
 
+#if !SSSS_COMPUTESHADER
 	// Initialize the stencil buffer in case it was not already available:
     if (initStencil) // (Checked in compile time, it's optimized away)
         if (SSSStrength < 1 / 256.0f) discard;
-
+#endif
 	float SSSScaleX = SubsurfaceParams.x;
 
 	float scale = SSSScaleX / OutDepth;
diff --git a/Engine/Shaders/Private/SubsurfaceBurleyNormalized.ush b/Engine/Shaders/Private/SubsurfaceBurleyNormalized.ush
--- a/Engine/Shaders/Private/SubsurfaceBurleyNormalized.ush
+++ b/Engine/Shaders/Private/SubsurfaceBurleyNormalized.ush
@@ -3,6 +3,7 @@
 /*=============================================================================
 	SubsurfaceBurleyNormalized.ush: Screenspace Burley subsurface scattering implementation.
 =============================================================================*/
+#pragma once
 
 #include "Random.ush"
 #include "BurleyNormalizedSSSCommon.ush"
@@ -726,22 +727,26 @@ void BurleyComputeMain(uint2 DT_ID, uint2 G_ID, uint GI)
 	// need to perform burley normalized subsurface scattering if it's burley
 	// this would unroll to just 1 sequence.
 	UNROLL for (int i = 0; i < LOCALGROUP_RATIO; ++i)
+	{
 		UNROLL for (int j = 0; j < LOCALGROUP_RATIO; ++j)
-			{
-				uint2 Pos = ConvertGroupIndexTo2DUsingMoltonCode(GI, TopLeftCorner + THREAD_TEXTURE_BOARDER + uint2(j, i)*THREAD_SIZE_1D);
-				float2 LBufferUV = ConvertGridPos2UV(Pos);
-				SSSColorUAV[Pos] = BurleyNormalizedSS(LBufferUV, Pos - (TopLeftCorner + THREAD_TEXTURE_BOARDER));
-			}
+		{
+			uint2 Pos = ConvertGroupIndexTo2DUsingMoltonCode(GI, TopLeftCorner + THREAD_TEXTURE_BOARDER + uint2(j, i)*THREAD_SIZE_1D);
+			float2 LBufferUV = ConvertGridPos2UV(Pos);
+			SSSColorUAV[Pos] = BurleyNormalizedSS(LBufferUV, Pos - (TopLeftCorner + THREAD_TEXTURE_BOARDER));
+		}
+	}
 
 #else
 	UNROLL for (int i = 0; i < LOCALGROUP_RATIO; ++i)
+	{
 		UNROLL for (int j = 0; j < LOCALGROUP_RATIO; ++j)
-			{
+		{
 
-				uint2 Pos = ConvertGroupIndexToNormal2DGrid(GI, TopLeftCorner + THREAD_TEXTURE_BOARDER + uint2(j, i)*THREAD_SIZE_1D);
-				float2 LBufferUV = ConvertGridPos2UV(Pos);
-				SSSColorUAV[Pos] = BurleyNormalizedSS(ConvertGridPos2UV(Pos), Pos - (TopLeftCorner + THREAD_TEXTURE_BOARDER));
-			}
+			uint2 Pos = ConvertGroupIndexToNormal2DGrid(GI, TopLeftCorner + THREAD_TEXTURE_BOARDER + uint2(j, i)*THREAD_SIZE_1D);
+			float2 LBufferUV = ConvertGridPos2UV(Pos);
+			SSSColorUAV[Pos] = BurleyNormalizedSS(ConvertGridPos2UV(Pos), Pos - (TopLeftCorner + THREAD_TEXTURE_BOARDER));
+		}
+	}
 #endif
 #else
 
@@ -749,16 +754,18 @@ void BurleyComputeMain(uint2 DT_ID, uint2 G_ID, uint GI)
 
 	//we only update the variance if the subsurface is burley
 	UNROLL for (int i = 0; i < LOCALGROUP_RATIO; ++i)
+	{
 		UNROLL for (int j = 0; j < LOCALGROUP_RATIO; ++j)
-			{
-				uint2 TargetGridPos = Pos + uint2(i, j);
-				float2 LBufferUV = ConvertGridPos2UV(TargetGridPos);
-				// We perform quality variance thread-wise in the second pass
-				float4 SurfaceColor = Texture2DSample(SubsurfaceInput0_Texture, SubsurfaceSampler0, LBufferUV);
-				HistoryUAV[TargetGridPos] = UpdateQualityVariance(SurfaceColor, LBufferUV);
-				// Direct sample the color
-				SSSColorUAV[TargetGridPos] = float4(SurfaceColor.rgb, 1.0f);
-			}
+		{
+			uint2 TargetGridPos = Pos + uint2(i, j);
+			float2 LBufferUV = ConvertGridPos2UV(TargetGridPos);
+			// We perform quality variance thread-wise in the second pass
+			float4 SurfaceColor = Texture2DSample(SubsurfaceInput0_Texture, SubsurfaceSampler0, LBufferUV);
+			HistoryUAV[TargetGridPos] = UpdateQualityVariance(SurfaceColor, LBufferUV);
+			// Direct sample the color
+			SSSColorUAV[TargetGridPos] = float4(SurfaceColor.rgb, 1.0f);
+		}
+	}
 #endif
 }
 
diff --git a/Engine/Source/Runtime/Renderer/Private/PostProcess/PostProcessSubsurface.cpp b/Engine/Source/Runtime/Renderer/Private/PostProcess/PostProcessSubsurface.cpp
--- a/Engine/Source/Runtime/Renderer/Private/PostProcess/PostProcessSubsurface.cpp
+++ b/Engine/Source/Runtime/Renderer/Private/PostProcess/PostProcessSubsurface.cpp
@@ -361,16 +361,9 @@ IMPLEMENT_GLOBAL_SHADER(FSubsurfaceViewportCopyPS, "/Engine/Private/PostProcessS
 // Indirect dispatch class and functions
 //-------------------------------------------------------------------------------------------
 
-struct FSubsurfaceGroupItem
-{
-	FVector2D GroupXY;		// which tile. 
-	int32 Count;			// The first element is used as the counter.
-	int32 Dummy;			// alignment
-};
-
 // Subsurface uniform buffer layout
 BEGIN_GLOBAL_SHADER_PARAMETER_STRUCT(FSubsurfaceUniformParameters, )
-SHADER_PARAMETER(int32, MaxGroupCount)
+SHADER_PARAMETER(uint32, MaxGroupCount)
 END_GLOBAL_SHADER_PARAMETER_STRUCT()
 
 IMPLEMENT_GLOBAL_SHADER_PARAMETER_STRUCT(FSubsurfaceUniformParameters, "SubsurfaceUniformParameters");
@@ -399,8 +392,8 @@ public:
 	}
 
 	BEGIN_SHADER_PARAMETER_STRUCT(FParameters, )
-		SHADER_PARAMETER_RDG_BUFFER_UAV(RWStructuredBuffer<FSubsurfaceGroupItem>, RWSeparableGroupBuffer)
-		SHADER_PARAMETER_RDG_BUFFER_UAV(RWStructuredBuffer<FSubsurfaceGroupItem>, RWBurleyGroupBuffer)
+		SHADER_PARAMETER_RDG_BUFFER_UAV(RWStructuredBuffer<uint>, RWSeparableGroupBuffer)
+		SHADER_PARAMETER_RDG_BUFFER_UAV(RWStructuredBuffer<uint>, RWBurleyGroupBuffer)
 		END_SHADER_PARAMETER_STRUCT()
 
 		static bool ShouldCompilePermutation(const FGlobalShaderPermutationParameters& Parameters)
@@ -428,7 +421,7 @@ public:
 	BEGIN_SHADER_PARAMETER_STRUCT(FParameters, )
 		SHADER_PARAMETER_STRUCT_REF(FSubsurfaceUniformParameters, SubsurfaceUniformParameters)
 		SHADER_PARAMETER_RDG_BUFFER_UAV(RWStructuredBuffer<uint>, RWIndirectDispatchArgsBuffer)
-		SHADER_PARAMETER_RDG_BUFFER_SRV(StructuredBuffer<FSubsurfaceGroupItem>, GroupBuffer)
+		SHADER_PARAMETER_RDG_BUFFER_SRV(StructuredBuffer<uint>, GroupBuffer)
 		END_SHADER_PARAMETER_STRUCT()
 
 		static bool ShouldCompilePermutation(const FGlobalShaderPermutationParameters& Parameters)
@@ -459,8 +452,8 @@ public:
 		SHADER_PARAMETER_SAMPLER(SamplerState, SubsurfaceSampler0)
 		SHADER_PARAMETER_RDG_TEXTURE_UAV(RWTexture2D, SetupTexture)
 		SHADER_PARAMETER_RDG_TEXTURE_UAV(RWTexture2D, ProfileEdgeMask)
-		SHADER_PARAMETER_RDG_BUFFER_UAV(RWStructuredBuffer<FSubsurfaceGroupItem>, RWSeparableGroupBuffer)
-		SHADER_PARAMETER_RDG_BUFFER_UAV(RWStructuredBuffer<FSubsurfaceGroupItem>, RWBurleyGroupBuffer)
+		SHADER_PARAMETER_RDG_BUFFER_UAV(RWStructuredBuffer<uint>, RWSeparableGroupBuffer)
+		SHADER_PARAMETER_RDG_BUFFER_UAV(RWStructuredBuffer<uint>, RWBurleyGroupBuffer)
 		SHADER_PARAMETER_STRUCT_REF(FSubsurfaceUniformParameters, SubsurfaceUniformParameters)
 		END_SHADER_PARAMETER_STRUCT()
 
@@ -491,7 +484,7 @@ public:
 		SHADER_PARAMETER_STRUCT_INCLUDE(FSubsurfaceParameters, Subsurface)
 		SHADER_PARAMETER_RDG_TEXTURE_UAV(RWTexture2D, SSSColorUAV)
 		SHADER_PARAMETER_RDG_TEXTURE_UAV(RWTexture2D, HistoryUAV)
-		SHADER_PARAMETER_RDG_BUFFER_SRV(StructuredBuffer<FSubsurfaceGroupItem>, GroupBuffer)
+		SHADER_PARAMETER_RDG_BUFFER_SRV(StructuredBuffer<uint>, GroupBuffer)
 		SHADER_PARAMETER_RDG_BUFFER(StructuredBuffer<uint>, IndirectDispatchArgsBuffer)
 		SHADER_PARAMETER_STRUCT(FSubsurfaceInput, SubsurfaceInput0)
 		SHADER_PARAMETER_SAMPLER(SamplerState, SubsurfaceSampler0)
@@ -850,8 +843,8 @@ void ComputeSubsurfaceForView(
 			ScreenPassView, SceneViewport);
 
 		// Initialize the group buffer
-		FRDGBufferRef SeparableGroupBuffer = GraphBuilder.CreateBuffer(FRDGBufferDesc::CreateStructuredDesc(sizeof(FSubsurfaceGroupItem), MaxGroupCount + 1), TEXT("SeparableGroupBuffer"));;
-		FRDGBufferRef BurleyGroupBuffer = GraphBuilder.CreateBuffer(FRDGBufferDesc::CreateStructuredDesc(sizeof(FSubsurfaceGroupItem), MaxGroupCount + 1), TEXT("BurleyGroupBuffer"));;
+		FRDGBufferRef SeparableGroupBuffer = GraphBuilder.CreateBuffer(FRDGBufferDesc::CreateStructuredDesc(sizeof(uint32), 2*(MaxGroupCount + 1)), TEXT("SeparableGroupBuffer"));;
+		FRDGBufferRef BurleyGroupBuffer = GraphBuilder.CreateBuffer(FRDGBufferDesc::CreateStructuredDesc(sizeof(uint32), 2*(MaxGroupCount + 1)), TEXT("BurleyGroupBuffer"));;
 		FRDGBufferRef SeparableIndirectDispatchArgsBuffer = GraphBuilder.CreateBuffer(FRDGBufferDesc::CreateIndirectDesc(4), TEXT("SeprableIndirectDispatchArgs"));
 		FRDGBufferRef BurleyIndirectDispatchArgsBuffer = GraphBuilder.CreateBuffer(FRDGBufferDesc::CreateIndirectDesc(4), TEXT("BurleyIndirectDispatchArgs"));
 
