Ported Basic Eye Adaptation to RDG

diff --git a/Engine/Shaders/Private/PostProcessEyeAdaptation.usf b/Engine/Shaders/Private/PostProcessEyeAdaptation.usf
--- a/Engine/Shaders/Private/PostProcessEyeAdaptation.usf
+++ b/Engine/Shaders/Private/PostProcessEyeAdaptation.usf
@@ -8,37 +8,49 @@
 #include "PostProcessCommon.ush"
 #include "PostProcessHistogramCommon.ush"
 
-// prior frame's eye adaptation exposure scale
 Texture2D EyeAdaptationTexture;
-uint4 EyeAdaptionSrcRect;
 
-// Compute the average alpha stored in a 2D texture.
-// in uint4 ROI:  region of interest.  (min, min, max, max).
-// in float Slope 
-float ComputeWeightedTextureAverageAlpha(Texture2D InBuffer, in uint4 ROI, in float Slope)
+Texture2D HistogramTexture;
+
+Texture2D ColorTexture;
+
+SamplerState ColorSampler;
+
+#if COMPUTESHADER
+
+RWTexture2D<float4> RWEyeAdaptationTexture;
+
+#endif
+
+uint2 Color_ViewportMin;
+uint2 Color_ViewportMax;
+
+float ComputeWeightedTextureAverageAlpha(
+	Texture2D Texture,
+	uint2 RectMin,
+	uint2 RectMax,
+	float Slope)
 {
 	// The inverse of the Region of Interest size.
-	const float InvXMax = 1.f / float(ROI.z - ROI.x);
-	const float InvYMax = 1.f / float(ROI.w - ROI.y);
-
-	// use product of linear weight in x and y.  
+	const float InvRectWidth = 1.0f / float(RectMax.x - RectMin.x);
+	const float InvRectHeight = 1.0f / float(RectMax.x - RectMin.y);
 
+	// use product of linear weight in x and y.
 	float Average = 0.0f;
 	float WeightTotal = 0.0f;
 
-
-	for (uint i = ROI.x; i < ROI.z; ++i)
+	for (uint i = RectMin.x; i < RectMax.x; ++i)
 	{
-		float IWeight = AdaptationWeight( i * InvXMax, Slope);  
-		for (uint j = ROI.y; j < ROI.w; ++j)
+		float IWeight = AdaptationWeight(i * InvRectWidth, Slope);
+		for (uint j = RectMin.y; j < RectMax.y; ++j)
 		{
-			float JWeight = AdaptationWeight( j * InvYMax, Slope);
+			float JWeight = AdaptationWeight(j * InvRectHeight, Slope);
 			float Weight = max(IWeight * JWeight, 0.05f);
-			
+
 			WeightTotal += Weight;
 
-			//accumulate values from alpha channel
-			float Sample = InBuffer.Load(int3(i, j, 0)).w;
+			// Accumulate values from alpha channel.
+			float Sample = Texture.Load(int3(i, j, 0)).w;
 			Average += Weight * Sample;
 		}
 	}
@@ -47,53 +59,53 @@ float ComputeWeightedTextureAverageAlpha(Texture2D InBuffer, in uint4 ROI, in fl
 	return Average;
 }
 
-float2 ComputeWeightedTextureAverageAlphaSubRegion(Texture2D InBuffer, in uint4 SubROI, in uint4 MainROI, in float Slope)
+float2 ComputeWeightedTextureAverageAlphaSubRegion(
+	Texture2D Texture,
+	uint2 SubRectMin,
+	uint2 SubRectMax,
+	uint2 RectMin,
+	uint2 RectMax,
+	float Slope)
 {
 	// The inverse of the Region of Interest size.
-	const float InvXMax = 1.f / float(MainROI.z - MainROI.x);
-	const float InvYMax = 1.f / float(MainROI.w - MainROI.y);
-
-	// use product of linear weight in x and y.  
+	const float InvRectWidth = 1.f / float(RectMax.x - RectMin.x);
+	const float InvRectHeight = 1.f / float(RectMax.y - RectMin.y);
 
+	// use product of linear weight in x and y.
 	float Value = 0.0f;
-	float WeightTotal = 0.0000000001f;	
+	float WeightTotal = 0.0000000001f;
 
-	for (uint i = SubROI.x; i < SubROI.z; ++i)
+	for (uint i = SubRectMin.x; i < SubRectMax.x; ++i)
 	{
-		float IWeight = AdaptationWeight( i * InvXMax, Slope);  
-		for (uint j = SubROI.y; j < SubROI.w; ++j)
+		float IWeight = AdaptationWeight(i * InvRectWidth, Slope);
+		for (uint j = SubRectMin.y; j < SubRectMax.y; ++j)
 		{
-			float JWeight = AdaptationWeight( j * InvYMax, Slope);
+			float JWeight = AdaptationWeight(j * InvRectHeight, Slope);
 			float Weight = max(IWeight * JWeight, 0.05f);
-			
+
 			WeightTotal += Weight;
 
-			//accumulate values from alpha channel
-			float Sample = InBuffer.Load(int3(i, j, 0)).w;
+			// Accumulate values from alpha channel.
+			float Sample = Texture.Load(int3(i, j, 0)).w;
 			Value += Weight * Sample;
 		}
 	}
 
-	return  float2(Value, WeightTotal);
+	return float2(Value, WeightTotal);
 }
 
-
-// @param FrameTime in seconds
-// @return smoothed exposure
 float ComputeEyeAdaptation(float OldExposure, float TargetExposure, float FrameTime)
 {
-	float Diff = TargetExposure - OldExposure;
+	const float Diff = TargetExposure - OldExposure;
 
-	float AdaptionSpeed = (Diff > 0) ? EyeAdaptation_ExposureSpeedUp : EyeAdaptation_ExposureSpeedDown;
+	const float AdaptionSpeed = (Diff > 0) ? EyeAdaptation_ExposureSpeedUp : EyeAdaptation_ExposureSpeedDown;
 
-	float Factor = 1.0f - exp2(-FrameTime * AdaptionSpeed);
+	const float Factor = 1.0f - exp2(-FrameTime * AdaptionSpeed);
 
 	return clamp(OldExposure + Diff * Factor, EyeAdaptation_MinAverageLuminance, EyeAdaptation_MaxAverageLuminance);
 }
 
-Texture2D HistogramTexture;
-
-float4 MainEyeAdaptationCommon()
+float4 EyeAdaptationCommon()
 {
 	float4 OutColor = 0;
 
@@ -101,7 +113,7 @@ float4 MainEyeAdaptationCommon()
 	const float TargetExposure = clamp(AverageSceneLuminance, EyeAdaptation_MinAverageLuminance, EyeAdaptation_MaxAverageLuminance);
 
 	float OldExposureScale = HistogramTexture.Load(int3(0, 1, 0)).x;
-	float OldExposure = EyeAdaptation_ExposureCompensation / ( OldExposureScale != 0 ? OldExposureScale : 1.0f );
+	float OldExposure = EyeAdaptation_ExposureCompensation / (OldExposureScale != 0 ? OldExposureScale : 1.0f);
 
 	// eye adaptation changes over time
 	float SmoothedExposure = ComputeEyeAdaptation(OldExposure, TargetExposure, EyeAdaptation_DeltaWorldTime);
@@ -117,35 +129,29 @@ float4 MainEyeAdaptationCommon()
 	return OutColor;
 }
 
-void MainPS(float4 UVAndScreenPos : TEXCOORD0, 
+void EyeAdaptationPS(float4 UVAndScreenPos : TEXCOORD0, 
 	out float4 OutColor : SV_Target0)
 {
-	OutColor = MainEyeAdaptationCommon();
+	OutColor = EyeAdaptationCommon();
 }
 
 #if COMPUTESHADER
-RWTexture2D<float4> RWEyeAdaptationTexture;
-
 [numthreads(1, 1, 1)]
-void MainCS(uint2 DispatchThreadId : SV_DispatchThreadID)
+void EyeAdaptationCS(uint2 DispatchThreadId : SV_DispatchThreadID)
 {
-	float4 OutColor = MainEyeAdaptationCommon();
+	float4 OutColor = EyeAdaptationCommon();
 
 	uint2 PixelPos = DispatchThreadId;
 	RWEyeAdaptationTexture[PixelPos] = OutColor;
 }
 #endif
 
-Texture2D SceneColorTexture;
-
-SamplerState SceneColorSampler;
-
-void MainBasicEyeAdaptationSetupPS(
+void BasicEyeAdaptationSetupPS(
 	noperspective float4 UVAndScreenPos : TEXCOORD0,
 	out float4 OutColor : SV_Target0)
 {
 	float2 UV = UVAndScreenPos.xy;
-	OutColor = Texture2DSample(SceneColorTexture, SceneColorSampler, UV);
+	OutColor = Texture2DSample(ColorTexture, ColorSampler, UV);
 
 	// Use max to ensure intensity is never zero (so the following log is well behaved)
 	float Intensity = max(dot(OutColor.xyz, float3(0.2126, 0.7152, 0.0722)), EyeAdaptation_LuminanceMin);
@@ -154,42 +160,20 @@ void MainBasicEyeAdaptationSetupPS(
 	OutColor.w = EyeAdaptation_HistogramScale * log2(Intensity) + EyeAdaptation_HistogramBias; 
 }
 
-
-// vertex shader entry point
-void MainReductionVS(
-	in float4 InPosition : ATTRIBUTE0,
-	in float2 UV : ATTRIBUTE1,
-	out float2 OutUV : TEXCOORD0,
-	out float4 OutPosition : SV_POSITION
-	)
-{
-	DrawRectangle(InPosition, UV, OutPosition, OutUV);
-}
-
-// pixel shader entry point
-void MainReductionPS(
-	noperspective float2 InUV : TEXCOORD0,
-	out float4 OutColor : SV_Target0
-	)
-{
-	// Output: low quality, 1 filtered sample
-	OutColor = Texture2DSample(PostprocessInput0, PostprocessInput0Sampler, InUV);
-}
-
 #if COMPUTESHADER
 
 #define TGSIZE 16
+
 groupshared float2 SubRectValueWeight[TGSIZE*TGSIZE];
-RWTexture2D<float4> RWOutputComputeTex;
 
 [numthreads(TGSIZE, TGSIZE, 1)]
-void MainLogLuminance2ExposureScaleCS(uint GIndex : SV_GroupIndex, uint2 GTId : SV_GroupThreadID)
+void BasicEyeAdaptationCS(uint GIndex : SV_GroupIndex, uint2 GTId : SV_GroupThreadID)
 {
 	float4 OutColor=0;
 
 #else
 
-void MainLogLuminance2ExposureScalePS(in float4 UVAndScreenPos : TEXCOORD0, out float4 OutColor : SV_Target0)
+void BasicEyeAdaptationPS(in float4 UVAndScreenPos : TEXCOORD0, out float4 OutColor : SV_Target0)
 {
 
 #endif
@@ -201,21 +185,24 @@ void MainLogLuminance2ExposureScalePS(in float4 UVAndScreenPos : TEXCOORD0, out
 #endif
 
 	// Weighting is a triangle function focused at the center of the screen
-	// with the product of linear decay in x and y.  
+	// with the product of linear decay in x and y.
 	// WeightSlope: 0 is uniform weighting.
 	//              1 decays to the edges.
 	//              >1 focuses more in the center of the screen.
-	float WeightSlope = EyeAdaptation_WeightSlope;
-	
+	const float WeightSlope = EyeAdaptation_WeightSlope;
+
 	// Compute scaled Log Luminance Average
 #if COMPUTESHADER
 	// There are TGSIZE*TGSIZE threads. Each thread will calculate the luminance for its own subregions from a TGSIZE*TGSIZE screen grid
-	uint4 SubRect = uint4(	((EyeAdaptionSrcRect.z-EyeAdaptionSrcRect.x) * GTId.x)		/ TGSIZE,
-							((EyeAdaptionSrcRect.w-EyeAdaptionSrcRect.y) * GTId.y)		/ TGSIZE, 
-							((EyeAdaptionSrcRect.z-EyeAdaptionSrcRect.x) * (GTId.x+1))	/ TGSIZE  ,
-							((EyeAdaptionSrcRect.w-EyeAdaptionSrcRect.y) * (GTId.y+1))	/ TGSIZE  );
+	const uint2 SubRectMin = uint2(
+		((Color_ViewportMax.x - Color_ViewportMin.y) * GTId.x) / TGSIZE,
+		((Color_ViewportMax.y - Color_ViewportMin.y) * GTId.y) / TGSIZE);
 
-	float2 LogLumAveWithWeight = ComputeWeightedTextureAverageAlphaSubRegion(PostprocessInput0, SubRect,  EyeAdaptionSrcRect, WeightSlope);
+	const uint2 SubRectMax = uint2(
+		((Color_ViewportMax.x - Color_ViewportMin.y) * (GTId.x + 1)) / TGSIZE,
+		((Color_ViewportMax.y - Color_ViewportMin.y) * (GTId.y + 1)) / TGSIZE);
+
+	const float2 LogLumAveWithWeight = ComputeWeightedTextureAverageAlphaSubRegion(ColorTexture, SubRectMin, SubRectMax, Color_ViewportMin, Color_ViewportMax, WeightSlope);
 
 	// Store in LDS
 	SubRectValueWeight[GIndex] = LogLumAveWithWeight;
@@ -248,7 +235,7 @@ void MainLogLuminance2ExposureScalePS(in float4 UVAndScreenPos : TEXCOORD0, out
 
 	float LogLumAve = SubRectValueWeight[0].x /  SubRectValueWeight[0].y;
 #else
-	float LogLumAve = ComputeWeightedTextureAverageAlpha(PostprocessInput0,  EyeAdaptionSrcRect, WeightSlope);
+	float LogLumAve = ComputeWeightedTextureAverageAlpha(ColorTexture, Color_ViewportMin, Color_ViewportMax, WeightSlope);
 #endif
 
 	// Correct for [0,1] scaling
@@ -261,14 +248,14 @@ void MainLogLuminance2ExposureScalePS(in float4 UVAndScreenPos : TEXCOORD0, out
 	const float ClampedLumAve = clamp(LumAve, EyeAdaptation_MinAverageLuminance, EyeAdaptation_MaxAverageLuminance);
 	
 	// The Exposure Scale (and thus intensity) used in the previous frame
-	float ExposureScaleOld = EyeAdaptationTexture.Load(int3(0, 0, 0)).x;
-	float LuminanceAveOld = EyeAdaptation_ExposureCompensation / (ExposureScaleOld != 0 ? ExposureScaleOld : 1.0f);
+	const float ExposureScaleOld = EyeAdaptationTexture.Load(int3(0, 0, 0)).x;
+	const float LuminanceAveOld = EyeAdaptation_ExposureCompensation / (ExposureScaleOld != 0 ? ExposureScaleOld : 1.0f);
 
 	// Time-based expoential blend of the intensity to allow the eye adaptation to ramp up over a few frames.
-	float SmoothedLuminance = ComputeEyeAdaptation(LuminanceAveOld, ClampedLumAve, EyeAdaptation_DeltaWorldTime);
+	const float SmoothedLuminance = ComputeEyeAdaptation(LuminanceAveOld, ClampedLumAve, EyeAdaptation_DeltaWorldTime);
 
-	float SmoothedExposureScale = 1.0f / max(0.0001f, SmoothedLuminance);
-	float TargetExposureScale   = 1.0f / max(0.0001f, ClampedLumAve);
+	const float SmoothedExposureScale = 1.0f / max(0.0001f, SmoothedLuminance);
+	const float TargetExposureScale   = 1.0f / max(0.0001f, ClampedLumAve);
 	
 	// Output the number that will rescale the image intensity
 	OutColor.x = EyeAdaptation_ExposureCompensation * SmoothedExposureScale;
@@ -276,11 +263,11 @@ void MainLogLuminance2ExposureScalePS(in float4 UVAndScreenPos : TEXCOORD0, out
 	OutColor.y = EyeAdaptation_ExposureCompensation * TargetExposureScale;
 	OutColor.z = AverageSceneLuminance;
 	OutColor.w = 0;
-		
+
 #if COMPUTESHADER
 	if(GIndex==0)
 	{
-		RWOutputComputeTex[uint2(0,0)] = OutColor;
+		RWEyeAdaptationTexture[uint2(0,0)] = OutColor;
 	}
 #endif
 }
\ No newline at end of file
diff --git a/Engine/Source/Runtime/Renderer/Private/PostProcess/PostProcessEyeAdaptation.cpp b/Engine/Source/Runtime/Renderer/Private/PostProcess/PostProcessEyeAdaptation.cpp
--- a/Engine/Source/Runtime/Renderer/Private/PostProcess/PostProcessEyeAdaptation.cpp
+++ b/Engine/Source/Runtime/Renderer/Private/PostProcess/PostProcessEyeAdaptation.cpp
@@ -7,6 +7,7 @@
 #include "PostProcess/PostProcessEyeAdaptation.h"
 #include "PostProcess/SceneFilterRendering.h"
 #include "PostProcess/PostProcessing.h"
+#include "SceneTextureParameters.h"
 #include "ClearQuad.h"
 #include "PipelineStateCache.h"
 #include "Math/UnrealMathUtility.h"
@@ -258,6 +259,11 @@ float GetEyeAdaptationFixedExposure(const FViewInfo& View)
 class FEyeAdaptationShader : public FGlobalShader
 {
 public:
+	BEGIN_SHADER_PARAMETER_STRUCT(FParameters, )
+		SHADER_PARAMETER_STRUCT(FEyeAdaptationParameters, EyeAdaptation)
+		SHADER_PARAMETER_RDG_TEXTURE(Texture2D, HistogramTexture)
+	END_SHADER_PARAMETER_STRUCT()
+
 	static const EPixelFormat OutputFormat = PF_A32B32G32R32F;
 
 	static bool ShouldCompilePermutation(const FGlobalShaderPermutationParameters& Parameters)
@@ -279,33 +285,33 @@ public:
 
 class FEyeAdaptationPS : public FEyeAdaptationShader
 {
+	using Super = FEyeAdaptationShader;
 public:
 	DECLARE_GLOBAL_SHADER(FEyeAdaptationPS);
-	SHADER_USE_PARAMETER_STRUCT(FEyeAdaptationPS, FEyeAdaptationShader);
+	SHADER_USE_PARAMETER_STRUCT(FEyeAdaptationPS, Super);
 
 	BEGIN_SHADER_PARAMETER_STRUCT(FParameters, )
-		SHADER_PARAMETER_STRUCT(FEyeAdaptationParameters, EyeAdaptation)
-		SHADER_PARAMETER_RDG_TEXTURE(Texture2D, HistogramTexture)
+		SHADER_PARAMETER_STRUCT_INCLUDE(Super::FParameters, Base)
 		RENDER_TARGET_BINDING_SLOTS()
 	END_SHADER_PARAMETER_STRUCT()
 };
 
-IMPLEMENT_GLOBAL_SHADER(FEyeAdaptationPS, "/Engine/Private/PostProcessEyeAdaptation.usf", "MainPS", SF_Pixel);
+IMPLEMENT_GLOBAL_SHADER(FEyeAdaptationPS, "/Engine/Private/PostProcessEyeAdaptation.usf", "EyeAdaptationPS", SF_Pixel);
 
 class FEyeAdaptationCS : public FEyeAdaptationShader
 {
+	using Super = FEyeAdaptationShader;
 public:
 	DECLARE_GLOBAL_SHADER(FEyeAdaptationCS);
-	SHADER_USE_PARAMETER_STRUCT(FEyeAdaptationCS, FEyeAdaptationShader);
+	SHADER_USE_PARAMETER_STRUCT(FEyeAdaptationCS, Super);
 
 	BEGIN_SHADER_PARAMETER_STRUCT(FParameters, )
-		SHADER_PARAMETER_STRUCT(FEyeAdaptationParameters, EyeAdaptation)
-		SHADER_PARAMETER_RDG_TEXTURE(Texture2D, HistogramTexture)
+		SHADER_PARAMETER_STRUCT_INCLUDE(Super::FParameters, Base)
 		SHADER_PARAMETER_RDG_TEXTURE_UAV(RWTexture2D<float4>, RWEyeAdaptationTexture)
 	END_SHADER_PARAMETER_STRUCT()
 };
 
-IMPLEMENT_GLOBAL_SHADER(FEyeAdaptationCS, "/Engine/Private/PostProcessEyeAdaptation.usf", "MainCS", SF_Compute);
+IMPLEMENT_GLOBAL_SHADER(FEyeAdaptationCS, "/Engine/Private/PostProcessEyeAdaptation.usf", "EyeAdaptationCS", SF_Compute);
 
 FRDGTextureRef AddHistogramEyeAdaptationPass(
 	FRDGBuilder& GraphBuilder,
@@ -320,11 +326,14 @@ FRDGTextureRef AddHistogramEyeAdaptationPass(
 
 	FRDGTextureRef OutputTexture = GraphBuilder.RegisterExternalTexture(View.GetEyeAdaptation(GraphBuilder.RHICmdList), TEXT("EyeAdaptation"), ERDGResourceFlags::MultiFrame);
 
+	FEyeAdaptationShader::FParameters PassBaseParameters;
+	PassBaseParameters.EyeAdaptation = GetEyeAdaptationParameters(View);
+	PassBaseParameters.HistogramTexture = HistogramTexture;
+
 	if (ScreenPassView.bUseComputePasses)
 	{
 		FEyeAdaptationCS::FParameters* PassParameters = GraphBuilder.AllocParameters<FEyeAdaptationCS::FParameters>();
-		PassParameters->EyeAdaptation = GetEyeAdaptationParameters(View);
-		PassParameters->HistogramTexture = HistogramTexture;
+		PassParameters->Base = PassBaseParameters;
 		PassParameters->RWEyeAdaptationTexture = GraphBuilder.CreateUAV(OutputTexture);
 
 		TShaderMapRef<FEyeAdaptationCS> ComputeShader(ScreenPassView.View.ShaderMap);
@@ -339,8 +348,7 @@ FRDGTextureRef AddHistogramEyeAdaptationPass(
 	else
 	{
 		FEyeAdaptationPS::FParameters* PassParameters = GraphBuilder.AllocParameters<FEyeAdaptationPS::FParameters>();
-		PassParameters->EyeAdaptation = GetEyeAdaptationParameters(View);
-		PassParameters->HistogramTexture = HistogramTexture;
+		PassParameters->Base = PassBaseParameters;
 		PassParameters->RenderTargets[0] = FRenderTargetBinding(OutputTexture, ERenderTargetLoadAction::ELoad);
 
 		TShaderMapRef<FEyeAdaptationPS> PixelShader(ScreenPassView.View.ShaderMap);
@@ -395,36 +403,27 @@ FRenderingCompositeOutputRef AddHistogramEyeAdaptationPass(FPostprocessContext&
 
 const ERHIFeatureLevel::Type BasicEyeAdaptationMinFeatureLevel = ERHIFeatureLevel::ES3_1;
 
-class FBasicEyeAdaptationShader : public FGlobalShader
-{
-public:
-	static bool ShouldCompilePermutation(const FGlobalShaderPermutationParameters& Parameters)
-	{
-		return IsFeatureLevelSupported(Parameters.Platform, BasicEyeAdaptationMinFeatureLevel);
-	}
-
-	FBasicEyeAdaptationShader() = default;
-	FBasicEyeAdaptationShader(const CompiledShaderInitializerType& Initializer)
-		: FGlobalShader(Initializer)
-	{}
-};
-
 /** Computes scaled and biased luma for the input scene color and puts it in the alpha channel. */
-class FBasicEyeAdaptationSetupPS : public FBasicEyeAdaptationShader
+class FBasicEyeAdaptationSetupPS : public FGlobalShader
 {
 public:
 	DECLARE_GLOBAL_SHADER(FBasicEyeAdaptationSetupPS);
-	SHADER_USE_PARAMETER_STRUCT(FBasicEyeAdaptationSetupPS, FBasicEyeAdaptationShader);
+	SHADER_USE_PARAMETER_STRUCT(FBasicEyeAdaptationSetupPS, FGlobalShader);
 
 	BEGIN_SHADER_PARAMETER_STRUCT(FParameters, )
 		SHADER_PARAMETER_STRUCT(FEyeAdaptationParameters, EyeAdaptation)
-		SHADER_PARAMETER_RDG_TEXTURE(Texture2D, SceneColorTexture)
-		SHADER_PARAMETER_SAMPLER(SamplerState, SceneColorSampler)
+		SHADER_PARAMETER_RDG_TEXTURE(Texture2D, ColorTexture)
+		SHADER_PARAMETER_SAMPLER(SamplerState, ColorSampler)
 		RENDER_TARGET_BINDING_SLOTS()
 	END_SHADER_PARAMETER_STRUCT()
+
+	static bool ShouldCompilePermutation(const FGlobalShaderPermutationParameters& Parameters)
+	{
+		return IsFeatureLevelSupported(Parameters.Platform, BasicEyeAdaptationMinFeatureLevel);
+	}
 };
 
-IMPLEMENT_GLOBAL_SHADER(FBasicEyeAdaptationSetupPS, "/Engine/Private/PostProcessEyeAdaptation.usf", "MainBasicEyeAdaptationSetupPS", SF_Pixel);
+IMPLEMENT_GLOBAL_SHADER(FBasicEyeAdaptationSetupPS, "/Engine/Private/PostProcessEyeAdaptation.usf", "BasicEyeAdaptationSetupPS", SF_Pixel);
 
 FRDGTextureRef AddBasicEyeAdaptationSetupPass(
 	FRDGBuilder& GraphBuilder,
@@ -448,8 +447,8 @@ FRDGTextureRef AddBasicEyeAdaptationSetupPass(
 
 	FBasicEyeAdaptationSetupPS::FParameters* PassParameters = GraphBuilder.AllocParameters<FBasicEyeAdaptationSetupPS::FParameters>();
 	PassParameters->EyeAdaptation = GetEyeAdaptationParameters(ScreenPassView.View, BasicEyeAdaptationMinFeatureLevel);
-	PassParameters->SceneColorTexture = SceneColorTexture;
-	PassParameters->SceneColorSampler = TStaticSamplerState<SF_Point, AM_Clamp, AM_Clamp, AM_Clamp>::GetRHI();
+	PassParameters->ColorTexture = SceneColorTexture;
+	PassParameters->ColorSampler = TStaticSamplerState<SF_Point, AM_Clamp, AM_Clamp, AM_Clamp>::GetRHI();
 	PassParameters->RenderTargets[0] = FRenderTargetBinding(OutputTexture, ScreenPassView.GetOverwriteLoadAction());
 
 	TShaderMapRef<FBasicEyeAdaptationSetupPS> PixelShader(ScreenPassView.View.ShaderMap);
@@ -495,302 +494,171 @@ FRenderingCompositeOutputRef AddBasicEyeAdaptationSetupPass(
 	return FRenderingCompositeOutputRef(Pass);
 }
 
-class FPostProcessLogLuminance2ExposureScaleBase: public FGlobalShader
+class FBasicEyeAdaptationShader : public FGlobalShader
 {
 public:
 	BEGIN_SHADER_PARAMETER_STRUCT(FParameters, )
+		SHADER_PARAMETER_STRUCT_REF(FViewUniformShaderParameters, View)
 		SHADER_PARAMETER_STRUCT(FEyeAdaptationParameters, EyeAdaptation)
+		SHADER_PARAMETER_STRUCT(FScreenPassTextureViewportParameters, Color)
+		SHADER_PARAMETER_RDG_TEXTURE(Texture2D, ColorTexture)
+		SHADER_PARAMETER_RDG_TEXTURE(Texture2D, EyeAdaptationTexture)
 	END_SHADER_PARAMETER_STRUCT()
 
-	FPostProcessLogLuminance2ExposureScaleBase() {}
-
-	FPostProcessLogLuminance2ExposureScaleBase(const ShaderMetaType::CompiledShaderInitializerType& Initializer)
-		: FGlobalShader(Initializer)
-	{
-		BindForLegacyShaderParameters<FParameters>(this, Initializer.ParameterMap);
-
-		PostprocessParameter.Bind(Initializer.ParameterMap);
-		EyeAdaptationTexture.Bind(Initializer.ParameterMap, TEXT("EyeAdaptationTexture"));
-		EyeAdaptationSrcRect.Bind(Initializer.ParameterMap, TEXT("EyeAdaptionSrcRect"));
-	}
+	static const EPixelFormat OutputFormat = PF_A32B32G32R32F;
 
-	/** Static Shader boilerplate */
 	static void ModifyCompilationEnvironment(const FGlobalShaderPermutationParameters& Parameters, FShaderCompilerEnvironment& OutEnvironment)
 	{
 		FGlobalShader::ModifyCompilationEnvironment(Parameters, OutEnvironment);
-		OutEnvironment.SetRenderTargetOutputFormat(0, PF_A32B32G32R32F);
+		OutEnvironment.SetRenderTargetOutputFormat(0, OutputFormat);
 	}
 
-	// FShader interface.
-	virtual bool Serialize(FArchive& Ar) override
+	static bool ShouldCompilePermutation(const FGlobalShaderPermutationParameters& Parameters)
 	{
-		bool bShaderHasOutdatedParameters = FGlobalShader::Serialize(Ar);
-		Ar << PostprocessParameter << EyeAdaptationTexture
-			<< EyeAdaptationParams << EyeAdaptationSrcRect;
-		return bShaderHasOutdatedParameters;
+		return IsFeatureLevelSupported(Parameters.Platform, BasicEyeAdaptationMinFeatureLevel);
 	}
 
-protected:
-	FPostProcessPassParameters	PostprocessParameter;
-	FShaderResourceParameter	EyeAdaptationTexture;
-	FShaderParameter			EyeAdaptationParams;
-	FShaderParameter			EyeAdaptationSrcRect;
+	FBasicEyeAdaptationShader() = default;
+	FBasicEyeAdaptationShader(const CompiledShaderInitializerType& Initializer)
+		: FGlobalShader(Initializer)
+	{}
 };
 
-
-
-/** Encapsulates the post process computation of the exposure scale pixel shader. */
-class FPostProcessLogLuminance2ExposureScalePS : public FPostProcessLogLuminance2ExposureScaleBase
+class FBasicEyeAdaptationPS : public FBasicEyeAdaptationShader
 {
-	DECLARE_SHADER_TYPE(FPostProcessLogLuminance2ExposureScalePS, Global);
-
+	using Super = FBasicEyeAdaptationShader;
 public:
+	DECLARE_GLOBAL_SHADER(FBasicEyeAdaptationPS);
+	SHADER_USE_PARAMETER_STRUCT(FBasicEyeAdaptationPS, Super);
 
-	static bool ShouldCompilePermutation(const FGlobalShaderPermutationParameters& Parameters)
-	{
-		return IsFeatureLevelSupported(Parameters.Platform, BasicEyeAdaptationMinFeatureLevel);
-	}
-
+	BEGIN_SHADER_PARAMETER_STRUCT(FParameters, )
+		SHADER_PARAMETER_STRUCT_INCLUDE(Super::FParameters, Base)
+		RENDER_TARGET_BINDING_SLOTS()
+	END_SHADER_PARAMETER_STRUCT()
+};
 
-	/** Default constructor. */
-	FPostProcessLogLuminance2ExposureScalePS()
-		: FPostProcessLogLuminance2ExposureScaleBase()
-	{}
+IMPLEMENT_GLOBAL_SHADER(FBasicEyeAdaptationPS, "/Engine/Private/PostProcessEyeAdaptation.usf", "BasicEyeAdaptationPS", SF_Pixel);
 
-	/** Initialization constructor. */
-	FPostProcessLogLuminance2ExposureScalePS(const ShaderMetaType::CompiledShaderInitializerType& Initializer)
-		: FPostProcessLogLuminance2ExposureScaleBase(Initializer)
-	{
-	}
+class FBasicEyeAdaptationCS : public FBasicEyeAdaptationShader
+{
+	using Super = FBasicEyeAdaptationShader;
+public:
+	DECLARE_GLOBAL_SHADER(FBasicEyeAdaptationCS);
+	SHADER_USE_PARAMETER_STRUCT(FBasicEyeAdaptationCS, Super);
 
+	BEGIN_SHADER_PARAMETER_STRUCT(FParameters, )
+		SHADER_PARAMETER_STRUCT_INCLUDE(Super::FParameters, Base)
+		SHADER_PARAMETER_RDG_TEXTURE_UAV(RWTexture2D<float4>, RWEyeAdaptationTexture)
+	END_SHADER_PARAMETER_STRUCT()
 
-	void SetPS(const FRenderingCompositePassContext& Context, const FIntRect& SrcRect, IPooledRenderTarget* LastEyeAdaptation)
+	static bool ShouldCompilePermutation(const FGlobalShaderPermutationParameters& Parameters)
 	{
-		FRHIPixelShader* ShaderRHI = GetPixelShader();
-
-		FGlobalShader::SetParameters<FViewUniformShaderParameters>(Context.RHICmdList, ShaderRHI, Context.View.ViewUniformBuffer);
-
-		PostprocessParameter.SetPS(Context.RHICmdList, ShaderRHI, Context, TStaticSamplerState<SF_Bilinear, AM_Clamp, AM_Clamp, AM_Clamp>::GetRHI());
-
-		// Associate the eye adaptation buffer from the previous frame with a texture to be read in this frame. 
-		
-		if (Context.View.HasValidEyeAdaptation())
-		{
-			SetTextureParameter(Context.RHICmdList, ShaderRHI, EyeAdaptationTexture, LastEyeAdaptation->GetRenderTargetItem().TargetableTexture);
-		}
-		else // some views don't have a state, thumbnail rendering?
-		{
-			SetTextureParameter(Context.RHICmdList, ShaderRHI, EyeAdaptationTexture, GWhiteTexture->TextureRHI);
-		}
-
-		// Pack the eye adaptation parameters for the shader
-		{
-			FParameters PassParameters;
-			PassParameters.EyeAdaptation = GetEyeAdaptationParameters(Context.View, BasicEyeAdaptationMinFeatureLevel);
-			SetShaderParameters(Context.RHICmdList, this, ShaderRHI, PassParameters);
-		}
-
-		// Set the src extent for the shader
-		SetShaderValue(Context.RHICmdList, ShaderRHI, EyeAdaptationSrcRect, SrcRect);
+		return IsFeatureLevelSupported(Parameters.Platform, ERHIFeatureLevel::SM5);
 	}
 };
-IMPLEMENT_SHADER_TYPE(, FPostProcessLogLuminance2ExposureScalePS, TEXT("/Engine/Private/PostProcessEyeAdaptation.usf"), TEXT("MainLogLuminance2ExposureScalePS"), SF_Pixel);
 
+IMPLEMENT_GLOBAL_SHADER(FBasicEyeAdaptationCS, "/Engine/Private/PostProcessEyeAdaptation.usf", "BasicEyeAdaptationCS", SF_Compute);
 
-class FPostProcessLogLuminance2ExposureScaleCS : public FPostProcessLogLuminance2ExposureScaleBase
+FRDGTextureRef AddBasicEyeAdaptationPass(
+	FRDGBuilder& GraphBuilder,
+	const FScreenPassViewInfo& ScreenPassView,
+	const FEyeAdaptationParameters& EyeAdaptationParameters,
+	FRDGTextureRef SceneColorTexture,
+	FIntRect SceneColorViewRect,
+	FRDGTextureRef EyeAdaptationTexture)
 {
-public:
-	DECLARE_SHADER_TYPE(FPostProcessLogLuminance2ExposureScaleCS, Global);
+	const FViewInfo& View = ScreenPassView.View;
 
-	/** Default constructor. */
-	FPostProcessLogLuminance2ExposureScaleCS()
-		: FPostProcessLogLuminance2ExposureScaleBase()
-	{}
+	View.SwapEyeAdaptationRTs(GraphBuilder.RHICmdList);
+	View.SetValidEyeAdaptation();
 
-	/** Initialization constructor. */
-	FPostProcessLogLuminance2ExposureScaleCS(const ShaderMetaType::CompiledShaderInitializerType& Initializer)
-		: FPostProcessLogLuminance2ExposureScaleBase(Initializer)
-	{
-		OutComputeTex.Bind(Initializer.ParameterMap, TEXT("OutputComputeTex"));
-	}
+	const FScreenPassTextureViewport SceneColorViewport(SceneColorViewRect, SceneColorTexture);
 
-	/** Static Shader boilerplate */
-	static bool ShouldCompilePermutation(const FGlobalShaderPermutationParameters& Parameters)
-	{
-		return IsFeatureLevelSupported(Parameters.Platform, ERHIFeatureLevel::SM5);
-	}
+	FRDGTextureRef OutputTexture = GraphBuilder.RegisterExternalTexture(View.GetEyeAdaptation(GraphBuilder.RHICmdList), TEXT("EyeAdaptation"), ERDGResourceFlags::MultiFrame);
 
-	// FShader interface.
-	virtual bool Serialize(FArchive& Ar) override
+	FBasicEyeAdaptationShader::FParameters PassBaseParameters;
+	PassBaseParameters.View = ScreenPassView.View.ViewUniformBuffer;
+	PassBaseParameters.EyeAdaptation = GetEyeAdaptationParameters(View);
+	PassBaseParameters.Color = GetScreenPassTextureViewportParameters(SceneColorViewport);
+	PassBaseParameters.ColorTexture = SceneColorTexture;
+	PassBaseParameters.EyeAdaptationTexture = EyeAdaptationTexture;
+
+	if (ScreenPassView.bUseComputePasses)
 	{
-		bool bShaderHasOutdatedParameters = FGlobalShader::Serialize(Ar);
-		Ar << PostprocessParameter << EyeAdaptationTexture
-			<< EyeAdaptationParams << EyeAdaptationSrcRect << OutComputeTex;
-		return bShaderHasOutdatedParameters;
-	}
+		FBasicEyeAdaptationCS::FParameters* PassParameters = GraphBuilder.AllocParameters<FBasicEyeAdaptationCS::FParameters>();
+		PassParameters->Base = PassBaseParameters;
+		PassParameters->RWEyeAdaptationTexture = GraphBuilder.CreateUAV(OutputTexture);
 
+		TShaderMapRef<FBasicEyeAdaptationCS> ComputeShader(ScreenPassView.View.ShaderMap);
 
-	void UnsetParameters(const FRenderingCompositePassContext& Context)
-	{
-		FRHIComputeShader* ShaderRHI = GetComputeShader();
-		OutComputeTex.UnsetUAV(Context.RHICmdList, ShaderRHI);
+		FComputeShaderUtils::AddPass(
+			GraphBuilder,
+			RDG_EVENT_NAME("BasicEyeAdaptation (CS)"),
+			*ComputeShader,
+			PassParameters,
+			FIntVector(1, 1, 1));
 	}
-
-	void SetParameters(const FRenderingCompositePassContext& Context, const FIntRect& SrcRect, IPooledRenderTarget* LastEyeAdaptation, FRHIUnorderedAccessView* DestUAV)
+	else
 	{
-		FRHIComputeShader* ShaderRHI = GetComputeShader();
-
-		FGlobalShader::SetParameters<FViewUniformShaderParameters>(Context.RHICmdList, ShaderRHI, Context.View.ViewUniformBuffer);
-		PostprocessParameter.SetCS(ShaderRHI, Context, Context.RHICmdList, TStaticSamplerState<SF_Bilinear, AM_Clamp, AM_Clamp, AM_Clamp>::GetRHI());
+		FBasicEyeAdaptationPS::FParameters* PassParameters = GraphBuilder.AllocParameters<FBasicEyeAdaptationPS::FParameters>();
+		PassParameters->Base = PassBaseParameters;
+		PassParameters->RenderTargets[0] = FRenderTargetBinding(OutputTexture, ERenderTargetLoadAction::ELoad);
 
-		// Associate the eye adaptation buffer from the previous frame with a texture to be read in this frame. 
-		if (Context.View.HasValidEyeAdaptation())
-		{
-			SetTextureParameter(Context.RHICmdList, ShaderRHI, EyeAdaptationTexture, LastEyeAdaptation->GetRenderTargetItem().TargetableTexture);
-		}
-		else // some views don't have a state, thumbnail rendering?
-		{
-			SetTextureParameter(Context.RHICmdList, ShaderRHI, EyeAdaptationTexture, GWhiteTexture->TextureRHI);
-		}
+		TShaderMapRef<FBasicEyeAdaptationPS> PixelShader(ScreenPassView.View.ShaderMap);
 
-		// Pack the eye adaptation parameters for the shader
-		{
-			FParameters PassParameters;
-			PassParameters.EyeAdaptation = GetEyeAdaptationParameters(Context.View, BasicEyeAdaptationMinFeatureLevel);
-			SetShaderParameters(Context.RHICmdList, this, ShaderRHI, PassParameters);
-		}
+		const FScreenPassTextureViewport OutputViewport(OutputTexture);
 
-		// Set the src extent for the shader
-		SetShaderValue(Context.RHICmdList, ShaderRHI, EyeAdaptationSrcRect, SrcRect);
-		OutComputeTex.SetTexture(Context.RHICmdList, ShaderRHI, nullptr, DestUAV);
+		AddDrawScreenPass(
+			GraphBuilder,
+			RDG_EVENT_NAME("BasicEyeAdaptation (PS)"),
+			ScreenPassView,
+			OutputViewport,
+			OutputViewport,
+			*PixelShader,
+			PassParameters);
 	}
-	private:
-
-	FRWShaderParameter			OutComputeTex;
-};
-IMPLEMENT_SHADER_TYPE(, FPostProcessLogLuminance2ExposureScaleCS, TEXT("/Engine/Private/PostProcessEyeAdaptation.usf"), TEXT("MainLogLuminance2ExposureScaleCS"), SF_Compute);
 
+	return OutputTexture;
+}
 
-void FRCPassPostProcessBasicEyeAdaptation::Process(FRenderingCompositePassContext& Context)
+FRenderingCompositeOutputRef AddBasicEyeAdaptationPass(
+	FPostprocessContext& Context,
+	FRenderingCompositeOutputRef SceneColor,
+	FIntRect SceneColorViewRect)
 {
 	const FViewInfo& View = Context.View;
-	const FSceneViewFamily& ViewFamily = *(View.Family);
-
-	// Get the custom 1x1 target used to store exposure value and Toggle the two render targets used to store new and old.
-	Context.View.SwapEyeAdaptationRTs(Context.RHICmdList);
-	IPooledRenderTarget* EyeAdaptationThisFrameRT = Context.View.GetEyeAdaptationRT(Context.RHICmdList);
-	IPooledRenderTarget* EyeAdaptationLastFrameRT = Context.View.GetLastEyeAdaptationRT(Context.RHICmdList);
-
-	check(EyeAdaptationThisFrameRT && EyeAdaptationLastFrameRT);
-
-	FIntPoint DestSize = EyeAdaptationThisFrameRT->GetDesc().Extent;
-
-	// The input texture sample size.  Averaged in the pixel shader.
-	FIntPoint SrcSize = GetInputDesc(ePId_Input0)->Extent;
-
-	// Compute the region of interest in the source texture.
-	uint32 ScaleFactor = FMath::DivideAndRoundUp(FSceneRenderTargets::Get(Context.RHICmdList).GetBufferSizeXY().Y, SrcSize.Y);
 
-	FIntRect SrcRect = View.ViewRect / ScaleFactor;
-
-	SCOPED_DRAW_EVENTF(Context.RHICmdList, PostProcessBasicEyeAdaptation, TEXT("PostProcessBasicEyeAdaptation %dx%d"), SrcSize.X, SrcSize.Y);
-
-	const FSceneRenderTargetItem& DestRenderTarget = EyeAdaptationThisFrameRT->GetRenderTargetItem();
-
-	// Inform MultiGPU systems that we're starting to update this texture for this frame
-	Context.RHICmdList.BeginUpdateMultiFrameResource(DestRenderTarget.ShaderResourceTexture);
+	FRenderingCompositePass* Pass = Context.Graph.RegisterPass(
+		new(FMemStack::Get()) TRCPassForRDG<1, 1>(
+			[&View, SceneColorViewRect](FRenderingCompositePass* InPass, FRenderingCompositePassContext& InContext)
+	{
+		FRDGBuilder GraphBuilder(InContext.RHICmdList);
 
-	// we render to our own output render target, not the intermediate one created by the compositing system
-	// Set the view family's render target/viewport.
+		const FEyeAdaptationParameters EyeAdaptationParameters = GetEyeAdaptationParameters(View);
 
-	FRHIRenderPassInfo RPInfo(DestRenderTarget.TargetableTexture, ERenderTargetActions::Load_Store);
-	TransitionRenderPassTargets(Context.RHICmdList, RPInfo);
-		
-	bool IsCompute = (CVarEyeAdaptationBasicCompute.GetValueOnRenderThread()>0) && (Context.View.GetFeatureLevel() >= ERHIFeatureLevel::SM5);
-	if (IsCompute)
-	{
-		WaitForInputPassComputeFences(Context.RHICmdList);
-
-		// #todo-renderpasses remove once everything is renderpasses
-		UnbindRenderTargets(Context.RHICmdList);
-		TShaderMapRef<FPostProcessLogLuminance2ExposureScaleCS> ComputeShader(Context.GetShaderMap());
-		Context.RHICmdList.SetComputeShader(ComputeShader->GetComputeShader());
-
-		Context.RHICmdList.TransitionResource(EResourceTransitionAccess::ERWBarrier, EResourceTransitionPipeline::EGfxToCompute, DestRenderTarget.UAV);
-		ComputeShader->SetParameters(Context, SrcRect, EyeAdaptationLastFrameRT, DestRenderTarget.UAV);
-		DispatchComputeShader(Context.RHICmdList, *ComputeShader, 1, 1, 1);
-		ComputeShader->UnsetParameters(Context);
-		Context.RHICmdList.TransitionResource(EResourceTransitionAccess::EReadable, EResourceTransitionPipeline::EComputeToGfx, DestRenderTarget.UAV);
-	}
-	else
-	{
-		Context.RHICmdList.BeginRenderPass(RPInfo, TEXT("BasicEyeAdaptation"));
-		{
-			Context.SetViewportAndCallRHI(0, 0, 0.0f, DestSize.X, DestSize.Y, 1.0f);
-
-			FGraphicsPipelineStateInitializer GraphicsPSOInit;
-			Context.RHICmdList.ApplyCachedRenderTargets(GraphicsPSOInit);
-			GraphicsPSOInit.BlendState = TStaticBlendState<>::GetRHI();
-			GraphicsPSOInit.RasterizerState = TStaticRasterizerState<>::GetRHI();
-			GraphicsPSOInit.DepthStencilState = TStaticDepthStencilState<false, CF_Always>::GetRHI();
-
-			TShaderMapRef<FPostProcessVS> VertexShader(Context.GetShaderMap());
-			TShaderMapRef<FPostProcessLogLuminance2ExposureScalePS> PixelShader(Context.GetShaderMap());
-
-			GraphicsPSOInit.BoundShaderState.VertexDeclarationRHI = GFilterVertexDeclaration.VertexDeclarationRHI;
-			GraphicsPSOInit.BoundShaderState.VertexShaderRHI = GETSAFERHISHADER_VERTEX(*VertexShader);
-			GraphicsPSOInit.BoundShaderState.PixelShaderRHI = GETSAFERHISHADER_PIXEL(*PixelShader);
-			GraphicsPSOInit.PrimitiveType = PT_TriangleList;
-
-			SetGraphicsPipelineState(Context.RHICmdList, GraphicsPSOInit);
-
-			// Set the parameters used by the pixel shader.
-
-			PixelShader->SetPS(Context, SrcRect, EyeAdaptationLastFrameRT);
-
-			// Draw a quad mapping scene color to the view's render target
-			DrawRectangle(
-				Context.RHICmdList,
-				0, 0,
-				DestSize.X, DestSize.Y,
-				0, 0,
-				DestSize.X, DestSize.Y,
-				DestSize,
-				DestSize,
-				*VertexShader,
-				EDRF_UseTriangleOptimization);
-		}
-		Context.RHICmdList.EndRenderPass();
-		Context.RHICmdList.CopyToResolveTarget(DestRenderTarget.TargetableTexture, DestRenderTarget.ShaderResourceTexture, FResolveParams());
-	}
+		FRDGTextureRef SceneColorTexture = InPass->CreateRDGTextureForRequiredInput(GraphBuilder, ePId_Input0, TEXT("SceneColor"));
 
-	// Inform MultiGPU systems that we've finished with this texture for this frame
-	Context.RHICmdList.EndUpdateMultiFrameResource(DestRenderTarget.ShaderResourceTexture);
+		FRDGTextureRef EyeAdaptationTexture = GetEyeAdaptationTexture(GraphBuilder, View);
 
-	Context.View.SetValidEyeAdaptation();
-}
+		FRDGTextureRef OutputTexture = AddBasicEyeAdaptationPass(
+			GraphBuilder,
+			FScreenPassViewInfo(View),
+			EyeAdaptationParameters,
+			SceneColorTexture,
+			SceneColorViewRect,
+			EyeAdaptationTexture);
 
+		InPass->ExtractRDGTextureForOutput(GraphBuilder, ePId_Output0, OutputTexture);
 
-FPooledRenderTargetDesc FRCPassPostProcessBasicEyeAdaptation::ComputeOutputDesc(EPassOutputId InPassOutputId) const
-{
-	// Specify invalid description to avoid getting intermediate rendertargets created.
-	// We want to use ViewState->GetEyeAdaptation() instead
-	FPooledRenderTargetDesc Ret;
+		GraphBuilder.Execute();
+	}));
 
-	Ret.DebugName = TEXT("EyeAdaptationBasic");
-	Ret.Flags |= GFastVRamConfig.EyeAdaptation;
+	Pass->SetInput(ePId_Input0, SceneColor);
 
-	if (GMaxRHIFeatureLevel >= ERHIFeatureLevel::SM5)
-	{
-		Ret.Flags |= TexCreate_UAV;
-	}
-	return Ret;
+	return FRenderingCompositeOutputRef(Pass);
 }
 
-//
-FSceneViewState::FEyeAdaptationRTManager::~FEyeAdaptationRTManager()
-{
-}
+FSceneViewState::FEyeAdaptationRTManager::~FEyeAdaptationRTManager() {}
 
 void FSceneViewState::FEyeAdaptationRTManager::SafeRelease()
 {
@@ -834,7 +702,7 @@ void FSceneViewState::FEyeAdaptationRTManager::SwapRTs(bool bInUpdateLastExposur
 	CurrentBuffer = 1 - CurrentBuffer;
 }
 
-TRefCountPtr<IPooledRenderTarget>&  FSceneViewState::FEyeAdaptationRTManager::GetRTRef(FRHICommandList* RHICmdList, const int BufferNumber)
+TRefCountPtr<IPooledRenderTarget>& FSceneViewState::FEyeAdaptationRTManager::GetRTRef(FRHICommandList* RHICmdList, const int BufferNumber)
 {
 	check(BufferNumber == 0 || BufferNumber == 1);
 
diff --git a/Engine/Source/Runtime/Renderer/Private/PostProcess/PostProcessEyeAdaptation.h b/Engine/Source/Runtime/Renderer/Private/PostProcess/PostProcessEyeAdaptation.h
--- a/Engine/Source/Runtime/Renderer/Private/PostProcess/PostProcessEyeAdaptation.h
+++ b/Engine/Source/Runtime/Renderer/Private/PostProcess/PostProcessEyeAdaptation.h
@@ -64,23 +64,9 @@ FRenderingCompositeOutputRef AddBasicEyeAdaptationSetupPass(
 	FRenderingCompositeOutputRef SceneColor,
 	FIntRect SceneColorViewRect);
 
-// ePId_Input0: Downsampled SceneColor Log
-// derives from TRenderingCompositePassBase<InputCount, OutputCount> 
-class FRCPassPostProcessBasicEyeAdaptation : public TRenderingCompositePassBase<1, 1>
-{
-public:
-	FRCPassPostProcessBasicEyeAdaptation(FIntPoint InDownsampledViewRect)
-	: DownsampledViewRect(InDownsampledViewRect) 
-	{
-	}
-
-	// interface FRenderingCompositePass ---------
-	virtual void Process(FRenderingCompositePassContext& Context) override;
-	virtual void Release() override { delete this; }
-	virtual FPooledRenderTargetDesc ComputeOutputDesc(EPassOutputId InPassOutputId) const override;
-
-private:
-	FIntPoint DownsampledViewRect;
-};
+FRenderingCompositeOutputRef AddBasicEyeAdaptationPass(
+	FPostprocessContext& Context,
+	FRenderingCompositeOutputRef SceneColor,
+	FIntRect SceneColorViewRect);
 
 FRenderingCompositeOutputRef AddHistogramEyeAdaptationPass(FPostprocessContext& Context, FRenderingCompositeOutputRef Histogram);
\ No newline at end of file
diff --git a/Engine/Source/Runtime/Renderer/Private/PostProcess/PostProcessing.cpp b/Engine/Source/Runtime/Renderer/Private/PostProcess/PostProcessing.cpp
--- a/Engine/Source/Runtime/Renderer/Private/PostProcess/PostProcessing.cpp
+++ b/Engine/Source/Runtime/Renderer/Private/PostProcess/PostProcessing.cpp
@@ -257,17 +257,15 @@ public:
 			}
 		}
 
-		// Calculate the final viewrect size (matching AddDownsamplePass behavior)
-		FinalViewRectSize.X = InputViewRect.Width();
-		FinalViewRectSize.Y = InputViewRect.Height();
+		FinalViewRect = InputViewRect;
 	}
 
 	// The number of elements in the array.
 	inline static int32 Num() { return DownSampleStages; }
 
-	FIntPoint GetFinalViewRectSize() const
+	FIntRect GetFinalViewRect() const
 	{
-		return FinalViewRectSize;
+		return FinalViewRect;
 	}
 
 	// Member data kept public for simplicity
@@ -279,7 +277,7 @@ private:
 	// no default constructor.
 	TBloomDownSampleArray() {};
 
-	FIntPoint	FinalViewRectSize;
+	FIntRect FinalViewRect;
 };
 
 // Standard DownsampleArray shared by Bloom, Tint, and Eye-Adaptation. 
@@ -416,12 +414,9 @@ static FRenderingCompositeOutputRef AddPostProcessBasicEyeAdaptation(const FView
 	static const int32 FinalDSIdx = FBloomDownSampleArray::Num() - 1;
 	FRenderingCompositeOutputRef PostProcessPriorReduction = BloomAndEyeDownSamples.PostProcessDownsamples[FinalDSIdx];
 
-	const FIntPoint DownsampledViewRectSize = BloomAndEyeDownSamples.GetFinalViewRectSize();
+	const FIntRect DownsampledViewRect = BloomAndEyeDownSamples.GetFinalViewRect();
 
-	// Compute the eye adaptation value based on average luminance from log2 luminance buffer, history, and specific shader parameters.
-	FRenderingCompositePass* Node = Context.Graph.RegisterPass(new(FMemStack::Get()) FRCPassPostProcessBasicEyeAdaptation(DownsampledViewRectSize));
-	Node->SetInput(ePId_Input0, PostProcessPriorReduction);
-	return FRenderingCompositeOutputRef(Node);
+	return AddBasicEyeAdaptationPass(Context, PostProcessPriorReduction, DownsampledViewRect);
 }
 
 static void AddVisualizeBloomOverlay(FPostprocessContext& Context, FRenderingCompositeOutputRef& HDRColor, FRenderingCompositeOutputRef& BloomOutputCombined)
