DR - Batching up Add/Remove Scene primitive.

diff --git a/Engine/Source/Runtime/Engine/Public/SceneInterface.h b/Engine/Source/Runtime/Engine/Public/SceneInterface.h
--- a/Engine/Source/Runtime/Engine/Public/SceneInterface.h
+++ b/Engine/Source/Runtime/Engine/Public/SceneInterface.h
@@ -62,6 +62,10 @@ public:
 	virtual void RemovePrimitive(UPrimitiveComponent* Primitive) = 0;
 	/** Called when a primitive is being unregistered and will not be immediately re-registered. */
 	virtual void ReleasePrimitive(UPrimitiveComponent* Primitive) = 0;
+	/**
+	* Updates all primitive scene info additions, remobals and translation changes
+	*/
+	virtual void UpdateAllPrimitiveSceneInfos(FRHICommandListImmediate& RHICmdList) = 0;
 	/** 
 	 * Updates the transform of a primitive which has already been added to the scene. 
 	 * 
diff --git a/Engine/Source/Runtime/Launch/Private/LaunchEngineLoop.cpp b/Engine/Source/Runtime/Launch/Private/LaunchEngineLoop.cpp
--- a/Engine/Source/Runtime/Launch/Private/LaunchEngineLoop.cpp
+++ b/Engine/Source/Runtime/Launch/Private/LaunchEngineLoop.cpp
@@ -4132,6 +4132,20 @@ void FEngineLoop::Tick()
 			GEngine->UpdateTimeAndHandleMaxTickRate();
 		}
 
+		for (const FWorldContext& Context : GEngine->GetWorldContexts())
+		{
+			UWorld* CurrentWorld = Context.World();
+			if (CurrentWorld)
+			{
+				FSceneInterface* Scene = CurrentWorld->Scene;
+				ENQUEUE_RENDER_COMMAND(UpdateScenePrimitives)(
+					[Scene](FRHICommandListImmediate& RHICmdList)
+				{
+					Scene->UpdateAllPrimitiveSceneInfos(RHICmdList);
+				});
+			}
+		}
+
 		// beginning of RHI frame
 		ENQUEUE_RENDER_COMMAND(BeginFrame)([CurrentFrameCounter](FRHICommandListImmediate& RHICmdList)
 		{
diff --git a/Engine/Source/Runtime/RenderCore/Private/RenderCore.cpp b/Engine/Source/Runtime/RenderCore/Private/RenderCore.cpp
--- a/Engine/Source/Runtime/RenderCore/Private/RenderCore.cpp
+++ b/Engine/Source/Runtime/RenderCore/Private/RenderCore.cpp
@@ -172,6 +172,7 @@ DEFINE_STAT(STAT_AddSceneLightTime);
 
 DEFINE_STAT(STAT_RemoveScenePrimitiveTime);
 DEFINE_STAT(STAT_AddScenePrimitiveRenderThreadTime);
+DEFINE_STAT(STAT_UpdateScenePrimitiveRenderThreadTime);
 DEFINE_STAT(STAT_UpdatePrimitiveTransformRenderThreadTime);
 
 DEFINE_STAT(STAT_RemoveScenePrimitiveGT);
diff --git a/Engine/Source/Runtime/RenderCore/Public/RenderCore.h b/Engine/Source/Runtime/RenderCore/Public/RenderCore.h
--- a/Engine/Source/Runtime/RenderCore/Public/RenderCore.h
+++ b/Engine/Source/Runtime/RenderCore/Public/RenderCore.h
@@ -115,6 +115,7 @@ DECLARE_DWORD_COUNTER_STAT_EXTERN(TEXT("Num shadowed"),STAT_NumShadowedLights,ST
 // It should only contain stats that are likely to track a lot of time in a typical scene, not edge case stats that are rarely non-zero.
 // Use a more detailed profiler (like an instruction trace or sampling capture on Xbox 360) to track down where time is going in more detail.
 DECLARE_CYCLE_STAT_EXTERN(TEXT("AddPrimitive"),STAT_AddScenePrimitiveRenderThreadTime,STATGROUP_SceneUpdate, RENDERCORE_API);
+DECLARE_CYCLE_STAT_EXTERN(TEXT("UpdatePrimitive"), STAT_UpdateScenePrimitiveRenderThreadTime, STATGROUP_SceneUpdate, RENDERCORE_API);
 DECLARE_CYCLE_STAT_EXTERN(TEXT("AddLight"),STAT_AddSceneLightTime,STATGROUP_SceneUpdate, RENDERCORE_API);
 DECLARE_CYCLE_STAT_EXTERN(TEXT("RemovePrimitive"),STAT_RemoveScenePrimitiveTime,STATGROUP_SceneUpdate, RENDERCORE_API);
 DECLARE_CYCLE_STAT_EXTERN(TEXT("RemoveLight"),STAT_RemoveSceneLightTime,STATGROUP_SceneUpdate, RENDERCORE_API);
diff --git a/Engine/Source/Runtime/Renderer/Private/DeferredShadingRenderer.cpp b/Engine/Source/Runtime/Renderer/Private/DeferredShadingRenderer.cpp
--- a/Engine/Source/Runtime/Renderer/Private/DeferredShadingRenderer.cpp
+++ b/Engine/Source/Runtime/Renderer/Private/DeferredShadingRenderer.cpp
@@ -843,6 +843,8 @@ static TAutoConsoleVariable<float> CVarStallInitViews(
 
 void FDeferredShadingSceneRenderer::Render(FRHICommandListImmediate& RHICmdList)
 {
+	Scene->UpdateAllPrimitiveSceneInfos(RHICmdList);
+
 	check(RHICmdList.IsOutsideRenderPass());
 
 	CSV_SCOPED_TIMING_STAT_EXCLUSIVE(RenderOther);
@@ -936,7 +938,7 @@ void FDeferredShadingSceneRenderer::Render(FRHICommandListImmediate& RHICmdList)
 
 	// Find the visible primitives.
 	RHICmdList.ImmediateFlush(EImmediateFlushType::DispatchToRHIThread);
-	
+
 	bool bDoInitViewAftersPrepass = false;
 	{
 		SCOPED_GPU_STAT(RHICmdList, VisibilityCommands);
diff --git a/Engine/Source/Runtime/Renderer/Private/GPUScene.cpp b/Engine/Source/Runtime/Renderer/Private/GPUScene.cpp
--- a/Engine/Source/Runtime/Renderer/Private/GPUScene.cpp
+++ b/Engine/Source/Runtime/Renderer/Private/GPUScene.cpp
@@ -151,7 +151,7 @@ int32 FGrowOnlySpanAllocator::SearchFreeList(int32 Num)
 	return INDEX_NONE;
 }
 
-void UpdateGPUScene(FRHICommandList& RHICmdList, FScene& Scene)
+void UpdateGPUScene(FRHICommandListImmediate& RHICmdList, FScene& Scene)
 {
 	if (UseGPUScene(GMaxRHIShaderPlatform, Scene.GetFeatureLevel()))
 	{
@@ -190,14 +190,13 @@ void UpdateGPUScene(FRHICommandList& RHICmdList, FScene& Scene)
 
 		const int32 NumPrimitiveDataUploads = Scene.GPUScene.PrimitivesToUpdate.Num();
 
+		int32 NumLightmapDataUploads = 0;
 		if (NumPrimitiveDataUploads > 0)
 		{
 			SCOPED_DRAW_EVENTF(RHICmdList, UpdateGPUScene, TEXT("UpdateGPUScene PrimitivesToUpdate = %u"), NumPrimitiveDataUploads);
 
 			FScatterUploadBuilder PrimitivesUploadBuilder(NumPrimitiveDataUploads, FPrimitiveSceneShaderData::PrimitiveDataStrideInFloat4s, Scene.GPUScene.PrimitivesUploadScatterBuffer, Scene.GPUScene.PrimitivesUploadDataBuffer);
 
-			int32 NumLightmapDataUploads = 0;
-
 			for (int32 Index : Scene.GPUScene.PrimitivesToUpdate)
 			{
 				// PrimitivesToUpdate may contain a stale out of bounds index, as we don't remove update request on primitive removal from scene.
@@ -221,26 +220,31 @@ void UpdateGPUScene(FRHICommandList& RHICmdList, FScene& Scene)
 			{
 				RHICmdList.TransitionResource(EResourceTransitionAccess::EWritable, EResourceTransitionPipeline::EGfxToCompute, Scene.GPUScene.PrimitiveBuffer.UAV);
 			}
-			
+
 			PrimitivesUploadBuilder.UploadTo_Flush(RHICmdList, Scene.GPUScene.PrimitiveBuffer);
 
 			RHICmdList.TransitionResource(EResourceTransitionAccess::EReadable, EResourceTransitionPipeline::EComputeToGfx, Scene.GPUScene.PrimitiveBuffer.UAV);
+		}
 
-			if (GGPUSceneValidatePrimitiveBuffer && Scene.GPUScene.PrimitiveBuffer.NumBytes > 0)
-			{
-				UE_LOG(LogRenderer, Warning, TEXT("r.GPUSceneValidatePrimitiveBuffer enabled, doing slow readback from GPU"));
-				FPrimitiveSceneShaderData* InstanceBufferCopy = (FPrimitiveSceneShaderData*)RHILockStructuredBuffer(Scene.GPUScene.PrimitiveBuffer.Buffer, 0, Scene.GPUScene.PrimitiveBuffer.NumBytes, RLM_ReadOnly);
+		if (GGPUSceneValidatePrimitiveBuffer && Scene.GPUScene.PrimitiveBuffer.NumBytes > 0)
+		{
+			//UE_LOG(LogRenderer, Warning, TEXT("r.GPUSceneValidatePrimitiveBuffer enabled, doing slow readback from GPU"));
+			FPrimitiveSceneShaderData* InstanceBufferCopy = (FPrimitiveSceneShaderData*)RHILockStructuredBuffer(Scene.GPUScene.PrimitiveBuffer.Buffer, 0, Scene.GPUScene.PrimitiveBuffer.NumBytes, RLM_ReadOnly);
 
-				for (int32 Index = 0; Index < Scene.PrimitiveSceneProxies.Num(); Index++)
+			for (int32 Index = 0; Index < Scene.PrimitiveSceneProxies.Num(); Index++)
+			{
+				FPrimitiveSceneShaderData PrimitiveSceneData(Scene.PrimitiveSceneProxies[Index]);
+				for (int i = 0; i < FPrimitiveSceneShaderData::PrimitiveDataStrideInFloat4s; i++)
 				{
-					FPrimitiveSceneShaderData PrimitiveSceneData(Scene.PrimitiveSceneProxies[Index]);
-
-					check(FMemory::Memcmp(&InstanceBufferCopy[Index], &PrimitiveSceneData, sizeof(FPrimitiveSceneShaderData)) == 0);
+					check(PrimitiveSceneData.Data[i] == InstanceBufferCopy[Index].Data[i]);
 				}
-
-				RHIUnlockStructuredBuffer(Scene.GPUScene.PrimitiveBuffer.Buffer);
 			}
 
+			RHIUnlockStructuredBuffer(Scene.GPUScene.PrimitiveBuffer.Buffer);
+		}
+
+		if (NumPrimitiveDataUploads > 0)
+		{
 			if (NumLightmapDataUploads > 0)
 			{
 				FScatterUploadBuilder LightmapDataUploadBuilder(NumLightmapDataUploads, FLightmapSceneShaderData::LightmapDataStrideInFloat4s, Scene.GPUScene.LightmapUploadScatterBuffer, Scene.GPUScene.LightmapUploadDataBuffer);
@@ -299,7 +303,7 @@ void UpdateGPUScene(FRHICommandList& RHICmdList, FScene& Scene)
 	checkSlow(Scene.GPUScene.PrimitivesToUpdate.Num() == 0);
 }
 
-void UploadDynamicPrimitiveShaderDataForView(FRHICommandList& RHICmdList, FScene& Scene, FViewInfo& View)
+void UploadDynamicPrimitiveShaderDataForView(FRHICommandListImmediate& RHICmdList, FScene& Scene, FViewInfo& View)
 {
 	if (UseGPUScene(GMaxRHIShaderPlatform, Scene.GetFeatureLevel()))
 	{
diff --git a/Engine/Source/Runtime/Renderer/Private/GPUScene.h b/Engine/Source/Runtime/Renderer/Private/GPUScene.h
--- a/Engine/Source/Runtime/Renderer/Private/GPUScene.h
+++ b/Engine/Source/Runtime/Renderer/Private/GPUScene.h
@@ -10,7 +10,7 @@ class FRHICommandList;
 class FScene;
 class FViewInfo;
 
-extern void UploadDynamicPrimitiveShaderDataForView(FRHICommandList& RHICmdList, FScene& Scene, FViewInfo& View);
-extern void UpdateGPUScene(FRHICommandList& RHICmdList, FScene& Scene);
+extern void UploadDynamicPrimitiveShaderDataForView(FRHICommandListImmediate& RHICmdList, FScene& Scene, FViewInfo& View);
+extern void UpdateGPUScene(FRHICommandListImmediate& RHICmdList, FScene& Scene);
 extern void AddPrimitiveToUpdateGPU(FScene& Scene, int32 PrimitiveId);
 
diff --git a/Engine/Source/Runtime/Renderer/Private/RendererScene.cpp b/Engine/Source/Runtime/Renderer/Private/RendererScene.cpp
--- a/Engine/Source/Runtime/Renderer/Private/RendererScene.cpp
+++ b/Engine/Source/Runtime/Renderer/Private/RendererScene.cpp
@@ -554,7 +554,7 @@ SIZE_T FScene::GetSizeBytes() const
 		+ PrimitiveOctree.GetSizeBytes();
 }
 
-void FScene::CheckPrimitiveArrays()
+void FScene::CheckPrimitiveArrays(int MaxTypeOffsetIndex)
 {
 	check(Primitives.Num() == PrimitiveTransforms.Num());
 	check(Primitives.Num() == PrimitiveSceneProxies.Num());
@@ -567,6 +567,31 @@ void FScene::CheckPrimitiveArrays()
 	check(Primitives.Num() == PrimitiveVirtualTextureLod.Num());
 	check(Primitives.Num() == PrimitiveOcclusionBounds.Num());
 	check(Primitives.Num() == PrimitivesNeedingStaticMeshUpdate.Num());
+
+#if UE_BUILD_DEBUG
+	MaxTypeOffsetIndex = MaxTypeOffsetIndex == -1 ? TypeOffsetTable.Num() : MaxTypeOffsetIndex;
+	for (int i = 0; i < MaxTypeOffsetIndex; i++)
+	{
+		for (int j = i + 1; j < MaxTypeOffsetIndex; j++)
+		{
+			check(TypeOffsetTable[i].PrimitiveSceneProxyType != TypeOffsetTable[j].PrimitiveSceneProxyType);
+			check(TypeOffsetTable[i].Offset <= TypeOffsetTable[j].Offset);
+		}
+	}
+
+	uint32 NextOffset = 0;
+	for (int i = 0; i < MaxTypeOffsetIndex; i++)
+	{
+		const FTypeOffsetTableEntry& Entry = TypeOffsetTable[i];
+		for (uint32 Index = NextOffset; Index < Entry.Offset; Index++)
+		{
+			checkSlow(Primitives[Index]->Proxy == PrimitiveSceneProxies[Index]);
+			SIZE_T TypeHash = PrimitiveSceneProxies[Index]->GetTypeHash();
+			checkfSlow(TypeHash == Entry.PrimitiveSceneProxyType, TEXT("TypeHash: %i not matching, expected: %i"), TypeHash, Entry.PrimitiveSceneProxyType);
+		}
+		NextOffset = Entry.Offset;
+	}
+#endif
 }
 
 
@@ -774,149 +799,16 @@ static void TBitArraySwapElements(TBitArray<>& Array, int32 i1, int32 i2)
 	BitRef2 = Bit1;
 }
 
-void FScene::AddPrimitiveSceneInfo_RenderThread(FRHICommandListImmediate& RHICmdList, FPrimitiveSceneInfo* PrimitiveSceneInfo)
+void FScene::AddPrimitiveSceneInfo_RenderThread(FPrimitiveSceneInfo* PrimitiveSceneInfo, const TOptional<FTransform>& PreviousTransform)
 {
-	SCOPE_CYCLE_COUNTER(STAT_AddScenePrimitiveRenderThreadTime);
-	
-	CheckPrimitiveArrays();
-
-	Primitives.Add(PrimitiveSceneInfo);
-	const FMatrix LocalToWorld = PrimitiveSceneInfo->Proxy->GetLocalToWorld();
-	PrimitiveTransforms.Add(LocalToWorld);
-	PrimitiveSceneProxies.Add(PrimitiveSceneInfo->Proxy);
-	PrimitiveBounds.AddUninitialized();
-	PrimitiveFlagsCompact.AddUninitialized();
-	PrimitiveVisibilityIds.AddUninitialized();
-	PrimitiveOcclusionFlags.AddUninitialized();
-	PrimitiveComponentIds.AddUninitialized();
-	PrimitiveVirtualTextureFlags.AddUninitialized();
-	PrimitiveVirtualTextureLod.AddUninitialized();
-	PrimitiveOcclusionBounds.AddUninitialized();
-	PrimitivesNeedingStaticMeshUpdate.Add(false);
-	
-	const int SourceIndex = PrimitiveSceneProxies.Num() - 1;
-	PrimitiveSceneInfo->PackedIndex = SourceIndex;
-
+	check(IsInRenderingThread());
+	check(PrimitiveSceneInfo->PackedIndex == INDEX_NONE);
+	check(!AddedPrimitiveSceneInfos.Contains(PrimitiveSceneInfo));
+	AddedPrimitiveSceneInfos.Add(PrimitiveSceneInfo);
+	if (PreviousTransform.IsSet())
 	{
-		bool EntryFound = false;
-		int BroadIndex = - 1;
-		SIZE_T InsertProxyHash = PrimitiveSceneInfo->Proxy->GetTypeHash();
-		//broad phase search for a matching type
-		for (BroadIndex = TypeOffsetTable.Num() - 1; BroadIndex >= 0; BroadIndex--)
-		{
-			// example how the prefix sum of the tails could look like
-			// PrimitiveSceneProxies[0,0,0,6,6,6,6,6,2,2,2,2,1,1,1,7,4,8]
-			// TypeOffsetTable[3,8,12,15,16,17,18]
-
-			if (TypeOffsetTable[BroadIndex].PrimitiveSceneProxyType == InsertProxyHash)
-			{
-				EntryFound = true;
-				break;
-			}
-		}
-
-		//new type encountered
-		if (EntryFound == false)
-		{
-			BroadIndex = TypeOffsetTable.Num();
-			if (BroadIndex)
-			{
-				FTypeOffsetTableEntry Entry = TypeOffsetTable[BroadIndex - 1];
-				//adding to the end of the list and offset of the tail (will will be incremented once during the while loop)
-				TypeOffsetTable.Push(FTypeOffsetTableEntry(InsertProxyHash, Entry.Offset));
-			}
-			else
-			{
-				//starting with an empty list and offset zero (will will be incremented once during the while loop)
-				TypeOffsetTable.Push(FTypeOffsetTableEntry(InsertProxyHash, 0));
-			}
-		}
-
-		while (BroadIndex < TypeOffsetTable.Num())
-		{
-			FTypeOffsetTableEntry& NextEntry = TypeOffsetTable[BroadIndex++];
-			int DestIndex = NextEntry.Offset++; //prepare swap and increment
-
-			// example swap chain of inserting a type of 6 at the end
-			// PrimitiveSceneProxies[0,0,0,6,6,6,6,6,2,2,2,2,1,1,1,7,4,8,6]
-			// PrimitiveSceneProxies[0,0,0,6,6,6,6,6,6,2,2,2,1,1,1,7,4,8,2]
-			// PrimitiveSceneProxies[0,0,0,6,6,6,6,6,6,2,2,2,2,1,1,7,4,8,1]
-			// PrimitiveSceneProxies[0,0,0,6,6,6,6,6,6,2,2,2,2,1,1,1,4,8,7]
-			// PrimitiveSceneProxies[0,0,0,6,6,6,6,6,6,2,2,2,2,1,1,1,7,8,4]
-			// PrimitiveSceneProxies[0,0,0,6,6,6,6,6,6,2,2,2,2,1,1,1,7,4,8]
-
-			if (DestIndex != SourceIndex)
-			{
-				checkfSlow(SourceIndex > DestIndex, TEXT("Corrupted Prefix Sum [%d, %d]"), SourceIndex, DestIndex);
-				Primitives[DestIndex]->PackedIndex = SourceIndex;
-				Primitives[SourceIndex]->PackedIndex = DestIndex;
-
-				TArraySwapElements(Primitives, DestIndex, SourceIndex);
-				TArraySwapElements(PrimitiveTransforms, DestIndex, SourceIndex);
-				TArraySwapElements(PrimitiveSceneProxies, DestIndex, SourceIndex);
-				TArraySwapElements(PrimitiveBounds, DestIndex, SourceIndex);
-				TArraySwapElements(PrimitiveFlagsCompact, DestIndex, SourceIndex);
-				TArraySwapElements(PrimitiveVisibilityIds, DestIndex, SourceIndex);
-				TArraySwapElements(PrimitiveOcclusionFlags, DestIndex, SourceIndex);
-				TArraySwapElements(PrimitiveComponentIds, DestIndex, SourceIndex);
-				TArraySwapElements(PrimitiveVirtualTextureFlags, DestIndex, SourceIndex);
-				TArraySwapElements(PrimitiveVirtualTextureLod, DestIndex, SourceIndex);
-				TArraySwapElements(PrimitiveOcclusionBounds, DestIndex, SourceIndex);
-				TBitArraySwapElements(PrimitivesNeedingStaticMeshUpdate, DestIndex, SourceIndex);
-
-				AddPrimitiveToUpdateGPU(*this, DestIndex);
-			}
-		}
+		OverridenPreviousTransforms.Add(PrimitiveSceneInfo, PreviousTransform.GetValue().ToMatrixWithScale());
 	}
-
-	CheckPrimitiveArrays();
-
-	// Add the primitive to its shadow parent's linked list of children.
-	// Note: must happen before AddToScene because AddToScene depends on LightingAttachmentRoot
-	PrimitiveSceneInfo->LinkAttachmentGroup();
-
-	// Set lod Parent information if valid
-	PrimitiveSceneInfo->LinkLODParentComponent();
-
-	if (GIsEditor)
-	{
-		PrimitiveSceneInfo->AddToScene(RHICmdList, true);
-	}
-	else
-	{
-		const bool bAddToDrawLists = !(CVarDoLazyStaticMeshUpdate.GetValueOnRenderThread());
-		if (bAddToDrawLists)
-		{
-			PrimitiveSceneInfo->AddToScene(RHICmdList, true);
-		}
-		else
-		{
-			PrimitiveSceneInfo->AddToScene(RHICmdList, true, false);
-			PrimitiveSceneInfo->BeginDeferredUpdateStaticMeshes();
-		}
-	}
-
-	if (PrimitiveSceneInfo->Proxy->IsMovable() && GetFeatureLevel() > ERHIFeatureLevel::ES3_1)
-	{
-		// We must register the initial LocalToWorld with the velocity state. 
-		// In the case of a moving component with MarkRenderStateDirty() called every frame, UpdateTransform will never happen.
-		VelocityData.UpdateTransform(PrimitiveSceneInfo, LocalToWorld, LocalToWorld);
-	}
-
-	AddPrimitiveToUpdateGPU(*this, SourceIndex);
-	bPathTracingNeedsInvalidation = true;
-
-	DistanceFieldSceneData.AddPrimitive(PrimitiveSceneInfo);
-
-	// Flush virtual textures touched by primitive
-	PrimitiveSceneInfo->FlushRuntimeVirtualTexture();
-
-	// LOD Parent, if this is LOD parent, we should update Proxy Scene Info
-	// LOD parent gets removed WHEN no children is accessing
-	// LOD parent can be recreated as scene updates
-	// I update if the parent component ID is still valid
-	// @Todo : really remove it if you know this is being destroyed - should happen from game thread as streaming in/out
-	SceneLODHierarchy.UpdateNodeSceneInfo(PrimitiveSceneInfo->PrimitiveComponentId, PrimitiveSceneInfo);
 }
 
 /**
@@ -1270,12 +1162,7 @@ void FScene::AddPrimitive(UPrimitiveComponent* Primitive)
 			// Create any RenderThreadResources required.
 			SceneProxy->CreateRenderThreadResources();
 
-			Scene->AddPrimitiveSceneInfo_RenderThread(RHICmdList, PrimitiveSceneInfo);
-
-			if (PreviousTransform.IsSet())
-			{
-				Scene->VelocityData.OverridePreviousTransform(PrimitiveSceneInfo->PrimitiveComponentId, PreviousTransform->ToMatrixWithScale());
-			}
+			Scene->AddPrimitiveSceneInfo_RenderThread(PrimitiveSceneInfo, PreviousTransform);
 		});
 
 }
@@ -1288,51 +1175,30 @@ static FAutoConsoleVariableRef CVarWarningOnRedundantTransformUpdate(
 	ECVF_Default
 );
 
-
-void FScene::UpdatePrimitiveTransform_RenderThread(FRHICommandListImmediate& RHICmdList, FPrimitiveSceneProxy* PrimitiveSceneProxy, const FBoxSphereBounds& WorldBounds, const FBoxSphereBounds& LocalBounds, const FMatrix& LocalToWorld, const FVector& AttachmentRootPosition)
+void FScene::UpdatePrimitiveTransform_RenderThread(FPrimitiveSceneProxy* PrimitiveSceneProxy, const FBoxSphereBounds& WorldBounds, const FBoxSphereBounds& LocalBounds, const FMatrix& LocalToWorld, const FVector& AttachmentRootPosition, const TOptional<FTransform>& PreviousTransform)
 {
-	SCOPE_CYCLE_COUNTER(STAT_UpdatePrimitiveTransformRenderThreadTime);
-
-	FPrimitiveSceneInfo* PrimitiveSceneInfo = PrimitiveSceneProxy->GetPrimitiveSceneInfo();
-
-	const bool bUpdateStaticDrawLists = !PrimitiveSceneProxy->StaticElementsAlwaysUseProxyPrimitiveUniformBuffer();
-
-	PrimitiveSceneInfo->FlushRuntimeVirtualTexture();
-
-	// Remove the primitive from the scene at its old location
-	// (note that the octree update relies on the bounds not being modified yet).
-	PrimitiveSceneInfo->RemoveFromScene(bUpdateStaticDrawLists);
-
-	if (PrimitiveSceneInfo->Proxy->IsMovable() && GetFeatureLevel() > ERHIFeatureLevel::ES3_1)
-	{
-		VelocityData.UpdateTransform(PrimitiveSceneInfo, LocalToWorld, PrimitiveSceneProxy->GetLocalToWorld());
-	}
-
+	check(IsInRenderingThread());
 	if (GWarningOnRedundantTransformUpdate && PrimitiveSceneProxy->WouldSetTransformBeRedundant(LocalToWorld, WorldBounds, LocalBounds, AttachmentRootPosition))
 	{
 		UE_LOG(LogRenderer, Warning, TEXT("Redundant UpdatePrimitiveTransform_RenderThread Owner: %s, Resource: %s, Level: %s"), *PrimitiveSceneProxy->GetOwnerName().ToString(), *PrimitiveSceneProxy->GetResourceName().ToString(), *PrimitiveSceneProxy->GetLevelName().ToString());
 	}
-	// Update the primitive transform.
-	PrimitiveSceneProxy->SetTransform(LocalToWorld, WorldBounds, LocalBounds, AttachmentRootPosition);
-	PrimitiveTransforms[PrimitiveSceneInfo->PackedIndex] = LocalToWorld;
 
-	if (GetFeatureLevel() < ERHIFeatureLevel::SM5
-		&& (PrimitiveSceneProxy->IsMovable() || PrimitiveSceneProxy->NeedsUnbuiltPreviewLighting() || PrimitiveSceneProxy->GetLightmapType() == ELightmapType::ForceVolumetric))
+	if (AddedPrimitiveSceneInfos.Contains(PrimitiveSceneProxy->GetPrimitiveSceneInfo()))
 	{
-		PrimitiveSceneInfo->MarkIndirectLightingCacheBufferDirty();
+		check(PrimitiveSceneProxy->GetPrimitiveSceneInfo()->PackedIndex == INDEX_NONE);
+	}
+	else
+	{
+		check(PrimitiveSceneProxy->GetPrimitiveSceneInfo()->PackedIndex != INDEX_NONE);
 	}
 
-	AddPrimitiveToUpdateGPU(*this, PrimitiveSceneInfo->PackedIndex);
-
-	DistanceFieldSceneData.UpdatePrimitive(PrimitiveSceneInfo);
-
-	// If the primitive has static mesh elements, it should have returned true from ShouldRecreateProxyOnUpdateTransform!
-	check(!(bUpdateStaticDrawLists && PrimitiveSceneInfo->StaticMeshes.Num()));
-
-	// Re-add the primitive to the scene with the new transform.
-	PrimitiveSceneInfo->AddToScene(RHICmdList, bUpdateStaticDrawLists);
+	check(!RemovedPrimitiveSceneInfos.Contains(PrimitiveSceneProxy->GetPrimitiveSceneInfo()));
+	UpdatedTransforms.Add(PrimitiveSceneProxy, { WorldBounds, LocalBounds, LocalToWorld, AttachmentRootPosition });
 
-	PrimitiveSceneInfo->FlushRuntimeVirtualTexture();
+	if (PreviousTransform.IsSet())
+	{
+		OverridenPreviousTransforms.Add(PrimitiveSceneProxy->GetPrimitiveSceneInfo(), PreviousTransform.GetValue().ToMatrixWithScale());
+	}
 }
 
 void FScene::UpdatePrimitiveTransform(UPrimitiveComponent* Primitive)
@@ -1400,12 +1266,7 @@ void FScene::UpdatePrimitiveTransform(UPrimitiveComponent* Primitive)
 				[UpdateParams](FRHICommandListImmediate& RHICmdList)
 				{
 					FScopeCycleCounter Context(UpdateParams.PrimitiveSceneProxy->GetStatId());
-					UpdateParams.Scene->UpdatePrimitiveTransform_RenderThread(RHICmdList, UpdateParams.PrimitiveSceneProxy, UpdateParams.WorldBounds, UpdateParams.LocalBounds, UpdateParams.LocalToWorld, UpdateParams.AttachmentRootPosition);
-
-					if (UpdateParams.PreviousTransform.IsSet())
-					{
-						UpdateParams.Scene->VelocityData.OverridePreviousTransform(UpdateParams.PrimitiveSceneProxy->PrimitiveComponentId, UpdateParams.PreviousTransform->ToMatrixWithScale());
-					}
+					UpdateParams.Scene->UpdatePrimitiveTransform_RenderThread(UpdateParams.PrimitiveSceneProxy, UpdateParams.WorldBounds, UpdateParams.LocalBounds, UpdateParams.LocalToWorld, UpdateParams.AttachmentRootPosition, UpdateParams.PreviousTransform);
 				});
 		}
 	}
@@ -1430,13 +1291,12 @@ void FScene::UpdatePrimitiveLightingAttachmentRoot(UPrimitiveComponent* Primitiv
 	if (Primitive->SceneProxy)
 	{
 		FPrimitiveSceneProxy* Proxy = Primitive->SceneProxy;
+		FScene* Scene = this;
 		ENQUEUE_RENDER_COMMAND(UpdatePrimitiveAttachment)(
-			[Proxy,NewComponentId](FRHICommandList&)
+			[Scene, Proxy, NewComponentId](FRHICommandList&)
 			{
 				FPrimitiveSceneInfo* PrimitiveInfo = Proxy->GetPrimitiveSceneInfo();
-				PrimitiveInfo->UnlinkAttachmentGroup();
-				PrimitiveInfo->LightingAttachmentRoot = NewComponentId;
-				PrimitiveInfo->LinkAttachmentGroup();
+				Scene->UpdatedAttachmentRoots.Add(PrimitiveInfo, NewComponentId);
 			});
 	}
 }
@@ -1489,19 +1349,7 @@ void FScene::UpdateCustomPrimitiveData(UPrimitiveComponent* Primitive)
 		ENQUEUE_RENDER_COMMAND(UpdateCustomPrimitiveDataCommand)(
 			[UpdateParams](FRHICommandListImmediate& RHICmdList)
 			{
-				FScopeCycleCounter Context(UpdateParams.PrimitiveSceneProxy->GetStatId());
-				UpdateParams.PrimitiveSceneProxy->CustomPrimitiveData = UpdateParams.CustomPrimitiveData;
-
-				// No need to do any of this if GPUScene isn't used (the custom primitive data will make it to the primitive uniform buffer through FPrimitiveSceneProxy::UpdateUniformBuffer if that's the case)
-				if (UseGPUScene(GMaxRHIShaderPlatform, UpdateParams.Scene->GetFeatureLevel()))
-				{
-					AddPrimitiveToUpdateGPU(*UpdateParams.Scene, UpdateParams.PrimitiveSceneProxy->GetPrimitiveSceneInfo()->PackedIndex);
-				}
-				else
-				{
-					// Make sure the uniform buffer is updated before rendering
-					UpdateParams.PrimitiveSceneProxy->GetPrimitiveSceneInfo()->SetNeedsUniformBufferUpdate(true);
-				}
+				UpdateParams.Scene->UpdatedCustomPrimitiveParams.Add(UpdateParams.PrimitiveSceneProxy, UpdateParams.CustomPrimitiveData);
 			});
 	}
 }
@@ -1520,7 +1368,7 @@ void FScene::UpdatePrimitiveDistanceFieldSceneData_GameThread(UPrimitiveComponen
 				if (PrimitiveSceneProxy && PrimitiveSceneProxy->GetPrimitiveSceneInfo())
 				{
 					FPrimitiveSceneInfo* Info = PrimitiveSceneProxy->GetPrimitiveSceneInfo();
-					DistanceFieldSceneData.UpdatePrimitive(Info);
+					this->DistanceFieldSceneDataUpdates.Add(Info);
 				}
 			});
 	}
@@ -1537,131 +1385,29 @@ FPrimitiveSceneInfo* FScene::GetPrimitiveSceneInfo(int32 PrimitiveIndex)
 
 void FScene::RemovePrimitiveSceneInfo_RenderThread(FPrimitiveSceneInfo* PrimitiveSceneInfo)
 {
-	SCOPE_CYCLE_COUNTER(STAT_RemoveScenePrimitiveTime);
-
-	// clear it up, parent is getting removed
-	SceneLODHierarchy.UpdateNodeSceneInfo(PrimitiveSceneInfo->PrimitiveComponentId, nullptr);
-
-	CheckPrimitiveArrays();
+	check(IsInRenderingThread());
 
-	int32 PrimitiveIndex = PrimitiveSceneInfo->PackedIndex;
+	if (AddedPrimitiveSceneInfos.Remove(PrimitiveSceneInfo))
 	{
-		int BroadIndex = -1;
-		SIZE_T InsertProxyHash = PrimitiveSceneInfo->Proxy->GetTypeHash();
-		//broad phase search for a matching type
-		for (BroadIndex = TypeOffsetTable.Num() - 1; BroadIndex >= 0; BroadIndex--)
-		{
-			// example how the prefix sum of the tails could look like
-			// PrimitiveSceneProxies[0,0,0,6,6,6,6,6,2,2,2,2,1,1,1,7,4,8]
-			// TypeOffsetTable[3,8,12,15,16,17,18]
-
-			if (TypeOffsetTable[BroadIndex].PrimitiveSceneProxyType == InsertProxyHash)
-			{
-				const int InsertionOffset = TypeOffsetTable[BroadIndex].Offset;
-				const int PrevOffset = BroadIndex > 0 ? TypeOffsetTable[BroadIndex - 1].Offset : 0;
-				checkfSlow(PrimitiveIndex >= PrevOffset && PrimitiveIndex < InsertionOffset, TEXT("PrimitiveIndex %d not in Bucket Range [%d, %d]"), PrimitiveIndex, PrevOffset, InsertionOffset);
-				break;
-			}
-		}
+		check(PrimitiveSceneInfo->PackedIndex == INDEX_NONE);
+		UpdatedTransforms.Remove(PrimitiveSceneInfo->Proxy);
+		UpdatedCustomPrimitiveParams.Remove(PrimitiveSceneInfo->Proxy);
+		OverridenPreviousTransforms.Remove(PrimitiveSceneInfo);
+		DistanceFieldSceneDataUpdates.Remove(PrimitiveSceneInfo);
+		UpdatedAttachmentRoots.Remove(PrimitiveSceneInfo);
 
-		int SourceIndex = PrimitiveIndex;
-		const int SavedBroadIndex = BroadIndex;
-		while ( BroadIndex < TypeOffsetTable.Num() )
-		{
-			FTypeOffsetTableEntry& NextEntry = TypeOffsetTable[BroadIndex++];
-			int DestIndex = --NextEntry.Offset; //decrement and prepare swap 
-
-			// example swap chain of removing X 
-			// PrimitiveSceneProxies[0,0,0,6,X,6,6,6,2,2,2,2,1,1,1,7,4,8]
-			// PrimitiveSceneProxies[0,0,0,6,6,6,6,6,X,2,2,2,1,1,1,7,4,8]
-			// PrimitiveSceneProxies[0,0,0,6,6,6,6,6,2,2,2,X,1,1,1,7,4,8]
-			// PrimitiveSceneProxies[0,0,0,6,6,6,6,6,2,2,2,1,1,1,X,7,4,8]
-			// PrimitiveSceneProxies[0,0,0,6,6,6,6,6,2,2,2,1,1,1,7,X,4,8]
-			// PrimitiveSceneProxies[0,0,0,6,6,6,6,6,2,2,2,1,1,1,7,4,X,8]
-			// PrimitiveSceneProxies[0,0,0,6,6,6,6,6,2,2,2,1,1,1,7,4,8,X]
-
-			if (DestIndex != SourceIndex)
-			{
-				checkfSlow(DestIndex > SourceIndex, TEXT("Corrupted Prefix Sum [%d, %d]"), DestIndex, SourceIndex);
-				Primitives[DestIndex]->PackedIndex = SourceIndex;
-				Primitives[SourceIndex]->PackedIndex = DestIndex;
-
-				TArraySwapElements(Primitives, DestIndex, SourceIndex);
-				TArraySwapElements(PrimitiveTransforms, DestIndex, SourceIndex);
-				TArraySwapElements(PrimitiveSceneProxies, DestIndex, SourceIndex);
-				TArraySwapElements(PrimitiveBounds, DestIndex, SourceIndex);
-				TArraySwapElements(PrimitiveFlagsCompact, DestIndex, SourceIndex);
-				TArraySwapElements(PrimitiveVisibilityIds, DestIndex, SourceIndex);
-				TArraySwapElements(PrimitiveOcclusionFlags, DestIndex, SourceIndex);
-				TArraySwapElements(PrimitiveComponentIds, DestIndex, SourceIndex);
-				TArraySwapElements(PrimitiveVirtualTextureFlags, DestIndex, SourceIndex);
-				TArraySwapElements(PrimitiveVirtualTextureLod, DestIndex, SourceIndex);
-				TArraySwapElements(PrimitiveOcclusionBounds, DestIndex, SourceIndex);
-				TBitArraySwapElements(PrimitivesNeedingStaticMeshUpdate, DestIndex, SourceIndex);
-				SourceIndex = DestIndex;
-
-				AddPrimitiveToUpdateGPU(*this, DestIndex);
-			}
-		}
-
-		const int PreviousOffset = SavedBroadIndex > 0 ? TypeOffsetTable[SavedBroadIndex - 1].Offset : 0;
-		const int CurrentOffset = TypeOffsetTable[SavedBroadIndex].Offset;
-
-		checkfSlow(PreviousOffset <= CurrentOffset, TEXT("Corrupted Bucket [%d, %d]"), PreviousOffset, CurrentOffset);
-		if (CurrentOffset - PreviousOffset == 0)
 		{
-			// remove empty OffsetTable entries e.g.
-			// TypeOffsetTable[3,8,12,15,15,17,18]
-			// TypeOffsetTable[3,8,12,15,17,18]
-			TypeOffsetTable.RemoveAt(SavedBroadIndex);
+			SCOPED_NAMED_EVENT(FScene_DeletePrimitiveSceneInfo, FColor::Red);
+			delete PrimitiveSceneInfo->Proxy;
+			delete PrimitiveSceneInfo;
 		}
-
-		checkfSlow((TypeOffsetTable.Num() == 0 && Primitives.Num() == 1) || TypeOffsetTable[TypeOffsetTable.Num() - 1].Offset == Primitives.Num() - 1, TEXT("Corrupted Tail Offset [%d, %d]"), TypeOffsetTable[TypeOffsetTable.Num() - 1].Offset, Primitives.Num() - 1);
-		checkfSlow(Primitives[Primitives.Num() - 1] == PrimitiveSceneInfo, TEXT("Removed item should be at the end"));
-
-		Primitives.Pop();
-		PrimitiveTransforms.Pop();
-		PrimitiveSceneProxies.Pop();
-		PrimitiveBounds.Pop();
-		PrimitiveFlagsCompact.Pop();
-		PrimitiveVisibilityIds.Pop();
-		PrimitiveOcclusionFlags.Pop();
-		PrimitiveComponentIds.Pop();
-		PrimitiveVirtualTextureFlags.Pop();
-		PrimitiveVirtualTextureLod.Pop();
-		PrimitiveOcclusionBounds.Pop();
-		PrimitivesNeedingStaticMeshUpdate.RemoveAt(PrimitivesNeedingStaticMeshUpdate.Num()-1);
 	}
-
-	PrimitiveSceneInfo->PackedIndex = MAX_int32;
-	
-	CheckPrimitiveArrays();
-
-	if (PrimitiveSceneInfo->Proxy->IsMovable())
+	else
 	{
-		// Remove primitive's motion blur information.
-		VelocityData.RemoveFromScene(PrimitiveSceneInfo->PrimitiveComponentId);
+		check(PrimitiveSceneInfo->PackedIndex != INDEX_NONE);
+		check(!RemovedPrimitiveSceneInfos.Contains(PrimitiveSceneInfo));
+		RemovedPrimitiveSceneInfos.Add(PrimitiveSceneInfo);
 	}
-
-	// Unlink the primitive from its shadow parent.
-	PrimitiveSceneInfo->UnlinkAttachmentGroup();
-
-	// Unlink the LOD parent info if valid
-	PrimitiveSceneInfo->UnlinkLODParentComponent();
-
-	// Flush virtual textures touched by primitive
-	PrimitiveSceneInfo->FlushRuntimeVirtualTexture();
-
-	// Remove the primitive from the scene.
-	PrimitiveSceneInfo->RemoveFromScene(true);
-
-	// Update the primitive that was swapped to this index
-	AddPrimitiveToUpdateGPU(*this, PrimitiveIndex);
-
-	DistanceFieldSceneData.RemovePrimitive(PrimitiveSceneInfo);
-
-	// free the primitive scene proxy.
-	delete PrimitiveSceneInfo->Proxy;
 }
 
 void FScene::RemovePrimitive( UPrimitiveComponent* Primitive )
@@ -1683,14 +1429,9 @@ void FScene::RemovePrimitive( UPrimitiveComponent* Primitive )
 		ENQUEUE_RENDER_COMMAND(FRemovePrimitiveCommand)(
 			[Scene, PrimitiveSceneInfo, AttachmentCounter](FRHICommandList&)
 			{
-				FScopeCycleCounter Context(PrimitiveSceneInfo->Proxy->GetStatId());
 				Scene->RemovePrimitiveSceneInfo_RenderThread(PrimitiveSceneInfo);
 				AttachmentCounter->Decrement();
 			});
-
-		// Delete the PrimitiveSceneInfo on the game thread after the rendering thread has processed its removal.
-		// This must be done on the game thread because the hit proxy references (and possibly other members) need to be freed on the game thread.
-		BeginCleanup(PrimitiveSceneInfo);
 	}
 }
 
@@ -3658,6 +3399,517 @@ bool FScene::AddPixelInspectorRequest(FPixelInspectorRequest *PixelInspectorRequ
 }
 #endif //WITH_EDITOR
 
+
+struct FPrimitiveArraySortKey
+{
+	inline bool operator()(const FPrimitiveSceneInfo& A, const FPrimitiveSceneInfo& B) const
+	{
+		return A.Proxy->GetTypeHash() < B.Proxy->GetTypeHash();
+	}
+};
+
+void FScene::UpdateAllPrimitiveSceneInfos(FRHICommandListImmediate& RHICmdList)
+{
+	SCOPED_NAMED_EVENT(FScene_UpdateAllPrimitiveSceneInfos, FColor::Orange);
+	SCOPE_CYCLE_COUNTER(STAT_UpdateScenePrimitiveRenderThreadTime);
+
+	check(IsInRenderingThread());
+
+	TArray<FPrimitiveSceneInfo*> RemovedLocalPrimitiveSceneInfos(RemovedPrimitiveSceneInfos.Array());
+	RemovedLocalPrimitiveSceneInfos.Sort(FPrimitiveArraySortKey());
+
+	TArray<FPrimitiveSceneInfo*> AddedLocalPrimitiveSceneInfos(AddedPrimitiveSceneInfos.Array());
+	AddedLocalPrimitiveSceneInfos.Sort(FPrimitiveArraySortKey());
+
+	TSet<FPrimitiveSceneInfo*> DeletedSceneInfos;
+	DeletedSceneInfos.Reserve(RemovedLocalPrimitiveSceneInfos.Num());
+
+	{
+		SCOPED_NAMED_EVENT(FScene_RemovePrimitiveSceneInfos, FColor::Red);
+		SCOPE_CYCLE_COUNTER(STAT_RemoveScenePrimitiveTime);
+		for (FPrimitiveSceneInfo* PrimitiveSceneInfo : RemovedLocalPrimitiveSceneInfos)
+		{
+			// clear it up, parent is getting removed
+			SceneLODHierarchy.UpdateNodeSceneInfo(PrimitiveSceneInfo->PrimitiveComponentId, nullptr);
+		}
+
+		while (RemovedLocalPrimitiveSceneInfos.Num())
+		{
+			int StartIndex = RemovedLocalPrimitiveSceneInfos.Num() - 1;
+			SIZE_T InsertProxyHash = RemovedLocalPrimitiveSceneInfos[StartIndex]->Proxy->GetTypeHash();
+
+			while (StartIndex > 0 && RemovedLocalPrimitiveSceneInfos[StartIndex - 1]->Proxy->GetTypeHash() == InsertProxyHash)
+			{
+				StartIndex--;
+			}
+
+			int BroadIndex = -1;
+			//broad phase search for a matching type
+			for (BroadIndex = TypeOffsetTable.Num() - 1; BroadIndex >= 0; BroadIndex--)
+			{
+				// example how the prefix sum of the tails could look like
+				// PrimitiveSceneProxies[0,0,0,6,6,6,6,6,2,2,2,2,1,1,1,7,4,8]
+				// TypeOffsetTable[3,8,12,15,16,17,18]
+
+				if (TypeOffsetTable[BroadIndex].PrimitiveSceneProxyType == InsertProxyHash)
+				{
+					const int InsertionOffset = TypeOffsetTable[BroadIndex].Offset;
+					const int PrevOffset = BroadIndex > 0 ? TypeOffsetTable[BroadIndex - 1].Offset : 0;
+					for (int CheckIndex = StartIndex; CheckIndex < RemovedLocalPrimitiveSceneInfos.Num(); CheckIndex++)
+					{
+						int32 PrimitiveIndex = RemovedLocalPrimitiveSceneInfos[CheckIndex]->PackedIndex;
+						checkfSlow(PrimitiveIndex >= PrevOffset && PrimitiveIndex < InsertionOffset, TEXT("PrimitiveIndex %d not in Bucket Range [%d, %d]"), PrimitiveIndex, PrevOffset, InsertionOffset);
+					}
+					break;
+				}
+			}
+
+			{
+				SCOPED_NAMED_EVENT(FScene_SwapPrimitiveSceneInfos, FColor::Turquoise);
+
+				for (int CheckIndex = StartIndex; CheckIndex < RemovedLocalPrimitiveSceneInfos.Num(); CheckIndex++)
+				{
+					int SourceIndex = RemovedLocalPrimitiveSceneInfos[CheckIndex]->PackedIndex;
+
+					for (int TypeIndex = BroadIndex; TypeIndex < TypeOffsetTable.Num(); TypeIndex++)
+					{
+						FTypeOffsetTableEntry& NextEntry = TypeOffsetTable[TypeIndex];
+						int DestIndex = --NextEntry.Offset; //decrement and prepare swap 
+
+						// example swap chain of removing X 
+						// PrimitiveSceneProxies[0,0,0,6,X,6,6,6,2,2,2,2,1,1,1,7,4,8]
+						// PrimitiveSceneProxies[0,0,0,6,6,6,6,6,X,2,2,2,1,1,1,7,4,8]
+						// PrimitiveSceneProxies[0,0,0,6,6,6,6,6,2,2,2,X,1,1,1,7,4,8]
+						// PrimitiveSceneProxies[0,0,0,6,6,6,6,6,2,2,2,1,1,1,X,7,4,8]
+						// PrimitiveSceneProxies[0,0,0,6,6,6,6,6,2,2,2,1,1,1,7,X,4,8]
+						// PrimitiveSceneProxies[0,0,0,6,6,6,6,6,2,2,2,1,1,1,7,4,X,8]
+						// PrimitiveSceneProxies[0,0,0,6,6,6,6,6,2,2,2,1,1,1,7,4,8,X]
+
+						if (DestIndex != SourceIndex)
+						{
+							checkfSlow(DestIndex > SourceIndex, TEXT("Corrupted Prefix Sum [%d, %d]"), DestIndex, SourceIndex);
+							Primitives[DestIndex]->PackedIndex = SourceIndex;
+							Primitives[SourceIndex]->PackedIndex = DestIndex;
+
+							TArraySwapElements(Primitives, DestIndex, SourceIndex);
+							TArraySwapElements(PrimitiveTransforms, DestIndex, SourceIndex);
+							TArraySwapElements(PrimitiveSceneProxies, DestIndex, SourceIndex);
+							TArraySwapElements(PrimitiveBounds, DestIndex, SourceIndex);
+							TArraySwapElements(PrimitiveFlagsCompact, DestIndex, SourceIndex);
+							TArraySwapElements(PrimitiveVisibilityIds, DestIndex, SourceIndex);
+							TArraySwapElements(PrimitiveOcclusionFlags, DestIndex, SourceIndex);
+							TArraySwapElements(PrimitiveComponentIds, DestIndex, SourceIndex);
+							TArraySwapElements(PrimitiveVirtualTextureFlags, DestIndex, SourceIndex);
+							TArraySwapElements(PrimitiveVirtualTextureLod, DestIndex, SourceIndex);
+							TArraySwapElements(PrimitiveOcclusionBounds, DestIndex, SourceIndex);
+							TBitArraySwapElements(PrimitivesNeedingStaticMeshUpdate, DestIndex, SourceIndex);
+
+							AddPrimitiveToUpdateGPU(*this, SourceIndex);
+							AddPrimitiveToUpdateGPU(*this, DestIndex);
+
+							SourceIndex = DestIndex;
+						}
+					}
+				}
+			}
+
+			const int PreviousOffset = BroadIndex > 0 ? TypeOffsetTable[BroadIndex - 1].Offset : 0;
+			const int CurrentOffset = TypeOffsetTable[BroadIndex].Offset;
+
+			checkfSlow(PreviousOffset <= CurrentOffset, TEXT("Corrupted Bucket [%d, %d]"), PreviousOffset, CurrentOffset);
+			if (CurrentOffset - PreviousOffset == 0)
+			{
+				// remove empty OffsetTable entries e.g.
+				// TypeOffsetTable[3,8,12,15,15,17,18]
+				// TypeOffsetTable[3,8,12,15,17,18]
+				TypeOffsetTable.RemoveAt(BroadIndex);
+			}
+
+			checkfSlow((TypeOffsetTable.Num() == 0 && Primitives.Num() == (RemovedLocalPrimitiveSceneInfos.Num() - StartIndex)) || TypeOffsetTable[TypeOffsetTable.Num() - 1].Offset == Primitives.Num() - (RemovedLocalPrimitiveSceneInfos.Num() - StartIndex), TEXT("Corrupted Tail Offset [%d, %d]"), TypeOffsetTable[TypeOffsetTable.Num() - 1].Offset, Primitives.Num() - (RemovedLocalPrimitiveSceneInfos.Num() - StartIndex));
+
+			for (int CheckIndex = StartIndex; CheckIndex < RemovedLocalPrimitiveSceneInfos.Num(); CheckIndex++)
+			{
+				checkfSlow(RemovedLocalPrimitiveSceneInfos[CheckIndex]->PackedIndex >= Primitives.Num() - RemovedLocalPrimitiveSceneInfos.Num(), TEXT("Removed item should be at the end"));
+			}
+
+			for (int RemoveIndex = StartIndex; RemoveIndex < RemovedLocalPrimitiveSceneInfos.Num(); RemoveIndex++)
+			{
+				int SourceIndex = RemovedLocalPrimitiveSceneInfos[RemoveIndex]->PackedIndex;
+				check(SourceIndex >= (Primitives.Num() - RemovedLocalPrimitiveSceneInfos.Num() + StartIndex));
+				Primitives.Pop();
+				PrimitiveTransforms.Pop();
+				PrimitiveSceneProxies.Pop();
+				PrimitiveBounds.Pop();
+				PrimitiveFlagsCompact.Pop();
+				PrimitiveVisibilityIds.Pop();
+				PrimitiveOcclusionFlags.Pop();
+				PrimitiveComponentIds.Pop();
+				PrimitiveVirtualTextureFlags.Pop();
+				PrimitiveVirtualTextureLod.Pop();
+				PrimitiveOcclusionBounds.Pop();
+				PrimitivesNeedingStaticMeshUpdate.RemoveAt(PrimitivesNeedingStaticMeshUpdate.Num() - 1);
+			}
+
+			CheckPrimitiveArrays();
+
+			for (int RemoveIndex = StartIndex; RemoveIndex < RemovedLocalPrimitiveSceneInfos.Num(); RemoveIndex++)
+			{
+				FPrimitiveSceneInfo* PrimitiveSceneInfo = RemovedLocalPrimitiveSceneInfos[RemoveIndex];
+				FScopeCycleCounter Context(PrimitiveSceneInfo->Proxy->GetStatId());
+				int32 PrimitiveIndex = PrimitiveSceneInfo->PackedIndex;
+				PrimitiveSceneInfo->PackedIndex = INDEX_NONE;
+
+				if (PrimitiveSceneInfo->Proxy->IsMovable())
+				{
+					// Remove primitive's motion blur information.
+					VelocityData.RemoveFromScene(PrimitiveSceneInfo->PrimitiveComponentId);
+				}
+
+				// Unlink the primitive from its shadow parent.
+				PrimitiveSceneInfo->UnlinkAttachmentGroup();
+
+				// Unlink the LOD parent info if valid
+				PrimitiveSceneInfo->UnlinkLODParentComponent();
+
+				// Flush virtual textures touched by primitive
+				PrimitiveSceneInfo->FlushRuntimeVirtualTexture();
+
+				// Remove the primitive from the scene.
+				PrimitiveSceneInfo->RemoveFromScene(true);
+
+				// Update the primitive that was swapped to this index
+				AddPrimitiveToUpdateGPU(*this, PrimitiveIndex);
+
+				DistanceFieldSceneData.RemovePrimitive(PrimitiveSceneInfo);
+
+				DeletedSceneInfos.Add(PrimitiveSceneInfo);
+			}
+			RemovedLocalPrimitiveSceneInfos.RemoveAt(StartIndex, RemovedLocalPrimitiveSceneInfos.Num() - StartIndex);
+		}
+	}
+	{
+		SCOPED_NAMED_EVENT(FScene_AddPrimitiveSceneInfos, FColor::Green);
+		SCOPE_CYCLE_COUNTER(STAT_AddScenePrimitiveRenderThreadTime);
+		if (AddedLocalPrimitiveSceneInfos.Num())
+		{
+			Primitives.Reserve(Primitives.Num() + AddedLocalPrimitiveSceneInfos.Num());
+			PrimitiveTransforms.Reserve(PrimitiveTransforms.Num() + AddedLocalPrimitiveSceneInfos.Num());
+			PrimitiveSceneProxies.Reserve(PrimitiveSceneProxies.Num() + AddedLocalPrimitiveSceneInfos.Num());
+			PrimitiveBounds.Reserve(PrimitiveBounds.Num() + AddedLocalPrimitiveSceneInfos.Num());
+			PrimitiveFlagsCompact.Reserve(PrimitiveFlagsCompact.Num() + AddedLocalPrimitiveSceneInfos.Num());
+			PrimitiveVisibilityIds.Reserve(PrimitiveVisibilityIds.Num() + AddedLocalPrimitiveSceneInfos.Num());
+			PrimitiveOcclusionFlags.Reserve(PrimitiveOcclusionFlags.Num() + AddedLocalPrimitiveSceneInfos.Num());
+			PrimitiveComponentIds.Reserve(PrimitiveComponentIds.Num() + AddedLocalPrimitiveSceneInfos.Num());
+			PrimitiveVirtualTextureFlags.Reserve(PrimitiveVirtualTextureFlags.Num() + AddedLocalPrimitiveSceneInfos.Num());
+			PrimitiveVirtualTextureLod.Reserve(PrimitiveVirtualTextureLod.Num() + AddedLocalPrimitiveSceneInfos.Num());
+			PrimitiveOcclusionBounds.Reserve(PrimitiveOcclusionBounds.Num() + AddedLocalPrimitiveSceneInfos.Num());
+			PrimitivesNeedingStaticMeshUpdate.Reserve(PrimitivesNeedingStaticMeshUpdate.Num() + AddedLocalPrimitiveSceneInfos.Num());
+		}
+
+		while (AddedLocalPrimitiveSceneInfos.Num())
+		{
+			int StartIndex = AddedLocalPrimitiveSceneInfos.Num() - 1;
+			SIZE_T InsertProxyHash = AddedLocalPrimitiveSceneInfos[StartIndex]->Proxy->GetTypeHash();
+
+			while (StartIndex > 0 && AddedLocalPrimitiveSceneInfos[StartIndex - 1]->Proxy->GetTypeHash() == InsertProxyHash)
+			{
+				StartIndex--;
+			}
+
+			for (int AddIndex = StartIndex; AddIndex < AddedLocalPrimitiveSceneInfos.Num(); AddIndex++)
+			{
+				FPrimitiveSceneInfo* PrimitiveSceneInfo = AddedLocalPrimitiveSceneInfos[AddIndex];
+				Primitives.Add(PrimitiveSceneInfo);
+				const FMatrix LocalToWorld = PrimitiveSceneInfo->Proxy->GetLocalToWorld();
+				PrimitiveTransforms.Add(LocalToWorld);
+				PrimitiveSceneProxies.Add(PrimitiveSceneInfo->Proxy);
+				PrimitiveBounds.AddUninitialized();
+				PrimitiveFlagsCompact.AddUninitialized();
+				PrimitiveVisibilityIds.AddUninitialized();
+				PrimitiveOcclusionFlags.AddUninitialized();
+				PrimitiveComponentIds.AddUninitialized();
+				PrimitiveVirtualTextureFlags.AddUninitialized();
+				PrimitiveVirtualTextureLod.AddUninitialized();
+				PrimitiveOcclusionBounds.AddUninitialized();
+				PrimitivesNeedingStaticMeshUpdate.Add(false);
+
+				const int SourceIndex = PrimitiveSceneProxies.Num() - 1;
+				PrimitiveSceneInfo->PackedIndex = SourceIndex;
+
+				AddPrimitiveToUpdateGPU(*this, SourceIndex);
+			}
+
+			bool EntryFound = false;
+			int BroadIndex = -1;
+			//broad phase search for a matching type
+			for (BroadIndex = TypeOffsetTable.Num() - 1; BroadIndex >= 0; BroadIndex--)
+			{
+				// example how the prefix sum of the tails could look like
+				// PrimitiveSceneProxies[0,0,0,6,6,6,6,6,2,2,2,2,1,1,1,7,4,8]
+				// TypeOffsetTable[3,8,12,15,16,17,18]
+
+				if (TypeOffsetTable[BroadIndex].PrimitiveSceneProxyType == InsertProxyHash)
+				{
+					EntryFound = true;
+					break;
+				}
+			}
+
+			//new type encountered
+			if (EntryFound == false)
+			{
+				BroadIndex = TypeOffsetTable.Num();
+				if (BroadIndex)
+				{
+					FTypeOffsetTableEntry Entry = TypeOffsetTable[BroadIndex - 1];
+					//adding to the end of the list and offset of the tail (will will be incremented once during the while loop)
+					TypeOffsetTable.Push(FTypeOffsetTableEntry(InsertProxyHash, Entry.Offset));
+				}
+				else
+				{
+					//starting with an empty list and offset zero (will will be incremented once during the while loop)
+					TypeOffsetTable.Push(FTypeOffsetTableEntry(InsertProxyHash, 0));
+				}
+			}
+
+			{
+				SCOPED_NAMED_EVENT(FScene_SwapPrimitiveSceneInfos, FColor::Turquoise);
+
+				for (int AddIndex = StartIndex; AddIndex < AddedLocalPrimitiveSceneInfos.Num(); AddIndex++)
+				{
+					int SourceIndex = AddedLocalPrimitiveSceneInfos[AddIndex]->PackedIndex;
+
+					for (int TypeIndex = BroadIndex; TypeIndex < TypeOffsetTable.Num(); TypeIndex++)
+					{
+						FTypeOffsetTableEntry& NextEntry = TypeOffsetTable[TypeIndex];
+						int DestIndex = NextEntry.Offset++; //prepare swap and increment
+
+						// example swap chain of inserting a type of 6 at the end
+						// PrimitiveSceneProxies[0,0,0,6,6,6,6,6,2,2,2,2,1,1,1,7,4,8,6]
+						// PrimitiveSceneProxies[0,0,0,6,6,6,6,6,6,2,2,2,1,1,1,7,4,8,2]
+						// PrimitiveSceneProxies[0,0,0,6,6,6,6,6,6,2,2,2,2,1,1,7,4,8,1]
+						// PrimitiveSceneProxies[0,0,0,6,6,6,6,6,6,2,2,2,2,1,1,1,4,8,7]
+						// PrimitiveSceneProxies[0,0,0,6,6,6,6,6,6,2,2,2,2,1,1,1,7,8,4]
+						// PrimitiveSceneProxies[0,0,0,6,6,6,6,6,6,2,2,2,2,1,1,1,7,4,8]
+
+						if (DestIndex != SourceIndex)
+						{
+							checkfSlow(SourceIndex > DestIndex, TEXT("Corrupted Prefix Sum [%d, %d]"), SourceIndex, DestIndex);
+							Primitives[DestIndex]->PackedIndex = SourceIndex;
+							Primitives[SourceIndex]->PackedIndex = DestIndex;
+
+							TArraySwapElements(Primitives, DestIndex, SourceIndex);
+							TArraySwapElements(PrimitiveTransforms, DestIndex, SourceIndex);
+							TArraySwapElements(PrimitiveSceneProxies, DestIndex, SourceIndex);
+							TArraySwapElements(PrimitiveBounds, DestIndex, SourceIndex);
+							TArraySwapElements(PrimitiveFlagsCompact, DestIndex, SourceIndex);
+							TArraySwapElements(PrimitiveVisibilityIds, DestIndex, SourceIndex);
+							TArraySwapElements(PrimitiveOcclusionFlags, DestIndex, SourceIndex);
+							TArraySwapElements(PrimitiveComponentIds, DestIndex, SourceIndex);
+							TArraySwapElements(PrimitiveVirtualTextureFlags, DestIndex, SourceIndex);
+							TArraySwapElements(PrimitiveVirtualTextureLod, DestIndex, SourceIndex);
+							TArraySwapElements(PrimitiveOcclusionBounds, DestIndex, SourceIndex);
+							TBitArraySwapElements(PrimitivesNeedingStaticMeshUpdate, DestIndex, SourceIndex);
+
+							AddPrimitiveToUpdateGPU(*this, DestIndex);
+						}
+					}
+				}
+			}
+
+			CheckPrimitiveArrays();
+
+			for (int AddIndex = StartIndex; AddIndex < AddedLocalPrimitiveSceneInfos.Num(); AddIndex++)
+			{
+				FPrimitiveSceneInfo* PrimitiveSceneInfo = AddedLocalPrimitiveSceneInfos[AddIndex];
+				FScopeCycleCounter Context(PrimitiveSceneInfo->Proxy->GetStatId());
+				int32 PrimitiveIndex = PrimitiveSceneInfo->PackedIndex;
+
+				// Add the primitive to its shadow parent's linked list of children.
+				// Note: must happen before AddToScene because AddToScene depends on LightingAttachmentRoot
+				PrimitiveSceneInfo->LinkAttachmentGroup();
+
+				// Set lod Parent information if valid
+				PrimitiveSceneInfo->LinkLODParentComponent();
+
+				{
+					SCOPED_NAMED_EVENT(FScene_AddPrimitiveSceneInfoToScene, FColor::Turquoise);
+					if (GIsEditor)
+					{
+						PrimitiveSceneInfo->AddToScene(RHICmdList, true);
+					}
+					else
+					{
+						const bool bAddToDrawLists = !(CVarDoLazyStaticMeshUpdate.GetValueOnRenderThread());
+						if (bAddToDrawLists)
+						{
+							PrimitiveSceneInfo->AddToScene(RHICmdList, true);
+						}
+						else
+						{
+							PrimitiveSceneInfo->AddToScene(RHICmdList, true, false);
+							PrimitiveSceneInfo->BeginDeferredUpdateStaticMeshes();
+						}
+					}
+				}
+
+				if (PrimitiveSceneInfo->Proxy->IsMovable() && GetFeatureLevel() > ERHIFeatureLevel::ES3_1)
+				{
+					// We must register the initial LocalToWorld with the velocity state. 
+					// In the case of a moving component with MarkRenderStateDirty() called every frame, UpdateTransform will never happen.
+					VelocityData.UpdateTransform(PrimitiveSceneInfo, PrimitiveTransforms[PrimitiveIndex], PrimitiveTransforms[PrimitiveIndex]);
+				}
+
+				AddPrimitiveToUpdateGPU(*this, PrimitiveIndex);
+				bPathTracingNeedsInvalidation = true;
+
+				DistanceFieldSceneData.AddPrimitive(PrimitiveSceneInfo);
+
+				// Flush virtual textures touched by primitive
+				PrimitiveSceneInfo->FlushRuntimeVirtualTexture();
+
+				// LOD Parent, if this is LOD parent, we should update Proxy Scene Info
+				// LOD parent gets removed WHEN no children is accessing
+				// LOD parent can be recreated as scene updates
+				// I update if the parent component ID is still valid
+				// @Todo : really remove it if you know this is being destroyed - should happen from game thread as streaming in/out
+				SceneLODHierarchy.UpdateNodeSceneInfo(PrimitiveSceneInfo->PrimitiveComponentId, PrimitiveSceneInfo);
+			}
+			AddedLocalPrimitiveSceneInfos.RemoveAt(StartIndex, AddedLocalPrimitiveSceneInfos.Num() - StartIndex);
+		}
+	}
+	{
+		SCOPED_NAMED_EVENT(FScene_AddPrimitiveSceneInfos, FColor::Yellow);
+		SCOPE_CYCLE_COUNTER(STAT_UpdatePrimitiveTransformRenderThreadTime);
+
+		for (const auto& Transform : UpdatedTransforms)
+		{
+			FPrimitiveSceneProxy* PrimitiveSceneProxy = Transform.Key;
+			if (DeletedSceneInfos.Contains(PrimitiveSceneProxy->GetPrimitiveSceneInfo()))
+			{
+				continue;
+			}
+			check(PrimitiveSceneProxy->GetPrimitiveSceneInfo()->PackedIndex != INDEX_NONE);
+
+			const FBoxSphereBounds& WorldBounds = Transform.Value.WorldBounds;
+			const FBoxSphereBounds& LocalBounds = Transform.Value.LocalBounds;
+			const FMatrix& LocalToWorld = Transform.Value.LocalToWorld;
+			const FVector& AttachmentRootPosition = Transform.Value.AttachmentRootPosition;
+			FScopeCycleCounter Context(PrimitiveSceneProxy->GetStatId());
+
+			FPrimitiveSceneInfo* PrimitiveSceneInfo = PrimitiveSceneProxy->GetPrimitiveSceneInfo();
+
+			const bool bUpdateStaticDrawLists = !PrimitiveSceneProxy->StaticElementsAlwaysUseProxyPrimitiveUniformBuffer();
+
+			PrimitiveSceneInfo->FlushRuntimeVirtualTexture();
+
+			// Remove the primitive from the scene at its old location
+			// (note that the octree update relies on the bounds not being modified yet).
+			PrimitiveSceneInfo->RemoveFromScene(bUpdateStaticDrawLists);
+
+			if (PrimitiveSceneInfo->Proxy->IsMovable() && GetFeatureLevel() > ERHIFeatureLevel::ES3_1)
+			{
+				VelocityData.UpdateTransform(PrimitiveSceneInfo, LocalToWorld, PrimitiveSceneProxy->GetLocalToWorld());
+			}
+
+			// Update the primitive transform.
+			PrimitiveSceneProxy->SetTransform(LocalToWorld, WorldBounds, LocalBounds, AttachmentRootPosition);
+			PrimitiveTransforms[PrimitiveSceneInfo->PackedIndex] = LocalToWorld;
+
+			if (GetFeatureLevel() < ERHIFeatureLevel::SM5
+				&& (PrimitiveSceneProxy->IsMovable() || PrimitiveSceneProxy->NeedsUnbuiltPreviewLighting() || PrimitiveSceneProxy->GetLightmapType() == ELightmapType::ForceVolumetric))
+			{
+				PrimitiveSceneInfo->MarkIndirectLightingCacheBufferDirty();
+			}
+
+			AddPrimitiveToUpdateGPU(*this, PrimitiveSceneInfo->PackedIndex);
+
+			DistanceFieldSceneData.UpdatePrimitive(PrimitiveSceneInfo);
+
+			// If the primitive has static mesh elements, it should have returned true from ShouldRecreateProxyOnUpdateTransform!
+			check(!(bUpdateStaticDrawLists && PrimitiveSceneInfo->StaticMeshes.Num()));
+
+			// Re-add the primitive to the scene with the new transform.
+			PrimitiveSceneInfo->AddToScene(RHICmdList, bUpdateStaticDrawLists);
+
+			PrimitiveSceneInfo->FlushRuntimeVirtualTexture();
+		}
+
+		for (const auto& Transform : OverridenPreviousTransforms)
+		{
+			FPrimitiveSceneInfo* PrimitiveSceneInfo = Transform.Key;
+			VelocityData.OverridePreviousTransform(PrimitiveSceneInfo->PrimitiveComponentId, Transform.Value);
+		}
+	}
+
+	for (const auto& Attachments : UpdatedAttachmentRoots)
+	{
+		FPrimitiveSceneInfo* PrimitiveSceneInfo = Attachments.Key;
+		if (DeletedSceneInfos.Contains(PrimitiveSceneInfo))
+		{
+			continue;
+		}
+
+		PrimitiveSceneInfo->UnlinkAttachmentGroup();
+		PrimitiveSceneInfo->LightingAttachmentRoot = Attachments.Value;
+		PrimitiveSceneInfo->LinkAttachmentGroup();
+	}
+
+	for (const auto& CustomParams : UpdatedCustomPrimitiveParams)
+	{
+		FPrimitiveSceneProxy* PrimitiveSceneProxy = CustomParams.Key;
+		if (DeletedSceneInfos.Contains(PrimitiveSceneProxy->GetPrimitiveSceneInfo()))
+		{
+			continue;
+		}
+
+		FScopeCycleCounter Context(PrimitiveSceneProxy->GetStatId());
+		PrimitiveSceneProxy->CustomPrimitiveData = CustomParams.Value;
+
+		// No need to do any of this if GPUScene isn't used (the custom primitive data will make it to the primitive uniform buffer through FPrimitiveSceneProxy::UpdateUniformBuffer if that's the case)
+		if (UseGPUScene(GMaxRHIShaderPlatform, GetFeatureLevel()))
+		{
+			AddPrimitiveToUpdateGPU(*this, PrimitiveSceneProxy->GetPrimitiveSceneInfo()->PackedIndex);
+		}
+		else
+		{
+			// Make sure the uniform buffer is updated before rendering
+			PrimitiveSceneProxy->GetPrimitiveSceneInfo()->SetNeedsUniformBufferUpdate(true);
+		}
+	}
+
+	for (FPrimitiveSceneInfo* PrimitiveSceneInfo : DistanceFieldSceneDataUpdates)
+	{
+		if (DeletedSceneInfos.Contains(PrimitiveSceneInfo))
+		{
+			continue;
+		}
+
+		DistanceFieldSceneData.UpdatePrimitive(PrimitiveSceneInfo);
+	}
+
+	{
+		SCOPED_NAMED_EVENT(FScene_DeletePrimitiveSceneInfo, FColor::Red);
+		for (FPrimitiveSceneInfo* PrimitiveSceneInfo : DeletedSceneInfos)
+		{
+			// free the primitive scene proxy.
+			delete PrimitiveSceneInfo->Proxy;
+			delete PrimitiveSceneInfo;
+			// Delete the PrimitiveSceneInfo on the game thread after the rendering thread has processed its removal.
+			// This must be done on the game thread because the hit proxy references (and possibly other members) need to be freed on the game thread.
+			//BeginCleanup(PrimitiveSceneInfo);
+		}
+	}
+
+	UpdatedAttachmentRoots.Reset();
+	UpdatedTransforms.Reset();
+	UpdatedCustomPrimitiveParams.Reset();
+	OverridenPreviousTransforms.Reset();
+	DistanceFieldSceneDataUpdates.Reset();
+	RemovedPrimitiveSceneInfos.Reset();
+	AddedPrimitiveSceneInfos.Reset();
+}
+
 /**
  * Dummy NULL scene interface used by dedicated servers.
  */
@@ -3685,6 +3937,7 @@ public:
 	virtual void AddPrimitive(UPrimitiveComponent* Primitive) override {}
 	virtual void RemovePrimitive(UPrimitiveComponent* Primitive) override {}
 	virtual void ReleasePrimitive(UPrimitiveComponent* Primitive) override {}
+	virtual void UpdateAllPrimitiveSceneInfos(FRHICommandListImmediate& RHICmdList) override {}
 	virtual FPrimitiveSceneInfo* GetPrimitiveSceneInfo(int32 PrimiteIndex) override { return NULL; }
 
 	/** Updates the transform of a primitive which has already been added to the scene. */
diff --git a/Engine/Source/Runtime/Renderer/Private/ScenePrivate.h b/Engine/Source/Runtime/Renderer/Private/ScenePrivate.h
--- a/Engine/Source/Runtime/Renderer/Private/ScenePrivate.h
+++ b/Engine/Source/Runtime/Renderer/Private/ScenePrivate.h
@@ -2653,6 +2653,7 @@ public:
 	virtual void AddPrimitive(UPrimitiveComponent* Primitive) override;
 	virtual void RemovePrimitive(UPrimitiveComponent* Primitive) override;
 	virtual void ReleasePrimitive(UPrimitiveComponent* Primitive) override;
+	virtual void UpdateAllPrimitiveSceneInfos(FRHICommandListImmediate& RHICmdList) override;
 	virtual void UpdatePrimitiveTransform(UPrimitiveComponent* Primitive) override;
 	virtual void UpdatePrimitiveAttachment(UPrimitiveComponent* Primitive) override;
 	virtual void UpdateCustomPrimitiveData(UPrimitiveComponent* Primitive) override;
@@ -2902,7 +2903,7 @@ private:
 	/**
 	 * Ensures the packed primitive arrays contain the same number of elements.
 	 */
-	void CheckPrimitiveArrays();
+	void CheckPrimitiveArrays(int MaxTypeOffsetIndex = -1);
 
 	/**
 	 * Retrieves the lights interacting with the passed in primitive and adds them to the out array.
@@ -2916,7 +2917,7 @@ private:
 	 * Adds a primitive to the scene.  Called in the rendering thread by AddPrimitive.
 	 * @param PrimitiveSceneInfo - The primitive being added.
 	 */
-	void AddPrimitiveSceneInfo_RenderThread(FRHICommandListImmediate& RHICmdList, FPrimitiveSceneInfo* PrimitiveSceneInfo);
+	void AddPrimitiveSceneInfo_RenderThread(FPrimitiveSceneInfo* PrimitiveSceneInfo, const TOptional<FTransform>& PreviousTransform);
 
 	/**
 	 * Removes a primitive from the scene.  Called in the rendering thread by RemovePrimitive.
@@ -2925,7 +2926,7 @@ private:
 	void RemovePrimitiveSceneInfo_RenderThread(FPrimitiveSceneInfo* PrimitiveSceneInfo);
 
 	/** Updates a primitive's transform, called on the rendering thread. */
-	void UpdatePrimitiveTransform_RenderThread(FRHICommandListImmediate& RHICmdList, FPrimitiveSceneProxy* PrimitiveSceneProxy, const FBoxSphereBounds& WorldBounds, const FBoxSphereBounds& LocalBounds, const FMatrix& LocalToWorld, const FVector& OwnerPosition);
+	void UpdatePrimitiveTransform_RenderThread(FPrimitiveSceneProxy* PrimitiveSceneProxy, const FBoxSphereBounds& WorldBounds, const FBoxSphereBounds& LocalBounds, const FMatrix& LocalToWorld, const FVector& OwnerPosition, const TOptional<FTransform>& PreviousTransform);
 
 	/** Updates a single primitive's lighting attachment root. */
 	void UpdatePrimitiveLightingAttachmentRoot(UPrimitiveComponent* Primitive);
@@ -2998,6 +2999,22 @@ private:
 	void ProcessAtmosphereLightAddition_RenderThread(FLightSceneInfo* LightSceneInfo);
 
 private:
+	struct FUpdateTransformCommand
+	{
+		FBoxSphereBounds WorldBounds;
+		FBoxSphereBounds LocalBounds; 
+		FMatrix LocalToWorld; 
+		FVector AttachmentRootPosition;
+	};
+
+	TMap<FPrimitiveSceneInfo*, FPrimitiveComponentId> UpdatedAttachmentRoots;
+	TMap<FPrimitiveSceneProxy*, FCustomPrimitiveData> UpdatedCustomPrimitiveParams;
+	TMap<FPrimitiveSceneProxy*, FUpdateTransformCommand> UpdatedTransforms;
+	TMap<FPrimitiveSceneInfo*, FMatrix> OverridenPreviousTransforms;
+	TSet<FPrimitiveSceneInfo*> AddedPrimitiveSceneInfos;
+	TSet<FPrimitiveSceneInfo*> RemovedPrimitiveSceneInfos;
+	TSet<FPrimitiveSceneInfo*> DistanceFieldSceneDataUpdates;
+
 	/** 
 	 * The number of visible lights in the scene
 	 * Note: This is tracked on the game thread!
