Fixed eye adaptation turning the screen white.

diff --git a/Engine/Source/Runtime/Renderer/Private/PostProcess/PostProcessEyeAdaptation.cpp b/Engine/Source/Runtime/Renderer/Private/PostProcess/PostProcessEyeAdaptation.cpp
--- a/Engine/Source/Runtime/Renderer/Private/PostProcess/PostProcessEyeAdaptation.cpp
+++ b/Engine/Source/Runtime/Renderer/Private/PostProcess/PostProcessEyeAdaptation.cpp
@@ -123,7 +123,7 @@ bool IsExtendLuminanceRangeEnabled()
 
 float GetLuminanceRangeValue(float EV100Value)
 {
-	return IsExtendLuminanceRangeEnabled() ? EV100ToLuminance(EV100Value) : EV100Value;
+	return IsExtendLuminanceRangeEnabled() ? EV100ToLog2(EV100Value) : EV100Value;
 }
 
 float GetBasicAutoExposureFocus()
