Fixed another eye adaptation artifact when Min / Max luminance values are set to 0.

diff --git a/Engine/Source/Runtime/Renderer/Private/PostProcess/PostProcessEyeAdaptation.cpp b/Engine/Source/Runtime/Renderer/Private/PostProcess/PostProcessEyeAdaptation.cpp
--- a/Engine/Source/Runtime/Renderer/Private/PostProcess/PostProcessEyeAdaptation.cpp
+++ b/Engine/Source/Runtime/Renderer/Private/PostProcess/PostProcessEyeAdaptation.cpp
@@ -121,11 +121,6 @@ bool IsExtendLuminanceRangeEnabled()
 	return VarDefaultAutoExposureExtendDefaultLuminanceRange->GetValueOnRenderThread() == 1;
 }
 
-float GetLuminanceRangeValue(float EV100Value)
-{
-	return IsExtendLuminanceRangeEnabled() ? EV100ToLog2(EV100Value) : EV100Value;
-}
-
 float GetBasicAutoExposureFocus()
 {
 	const float FocusMax = 10.f;
@@ -180,13 +175,15 @@ FEyeAdaptationParameters GetEyeAdaptationParameters(const FViewInfo& View, ERHIF
 
 	const EAutoExposureMethod AutoExposureMethod = GetAutoExposureMethod(View);
 
+	const bool bExtendedLuminanceRange = IsExtendLuminanceRangeEnabled();
+
 	const float PercentToScale = 0.01f;
 
 	const float ExposureHighPercent = FMath::Clamp(Settings.AutoExposureHighPercent, 1.0f, 99.0f) * PercentToScale;
 	const float ExposureLowPercent = FMath::Min(FMath::Clamp(Settings.AutoExposureLowPercent, 1.0f, 99.0f) * PercentToScale, ExposureHighPercent);
 
-	const float HistogramLogMax = GetLuminanceRangeValue(Settings.HistogramLogMax);
-	const float HistogramLogMin = FMath::Min(GetLuminanceRangeValue(Settings.HistogramLogMin), HistogramLogMax - 1);
+	const float HistogramLogMax = bExtendedLuminanceRange ? EV100ToLog2(Settings.HistogramLogMax) : Settings.HistogramLogMax;
+	const float HistogramLogMin = FMath::Min(bExtendedLuminanceRange ? EV100ToLog2(Settings.HistogramLogMin) : Settings.HistogramLogMin, HistogramLogMax - 1);
 
 	// These clamp the average luminance computed from the scene color.
 	float MinAverageLuminance = 1.0f;
@@ -205,13 +202,20 @@ FEyeAdaptationParameters GetEyeAdaptationParameters(const FViewInfo& View, ERHIF
 			const float FixedEV100 = FMath::Log2(FMath::Square(Settings.DepthOfFieldFstop) * Settings.CameraShutterSpeed * 100 / FMath::Max(1.f, Settings.CameraISO));
 			MinAverageLuminance = MaxAverageLuminance = EV100ToLuminance(FixedEV100);
 		}
+		else if (bExtendedLuminanceRange)
+		{
+			MinAverageLuminance = EV100ToLuminance(Settings.AutoExposureMinBrightness);
+			MaxAverageLuminance = EV100ToLuminance(Settings.AutoExposureMaxBrightness);
+		}
 		else
 		{
-			MinAverageLuminance = GetLuminanceRangeValue(Settings.AutoExposureMinBrightness);
-			MaxAverageLuminance = GetLuminanceRangeValue(Settings.AutoExposureMaxBrightness);
+			MinAverageLuminance = Settings.AutoExposureMinBrightness;
+			MaxAverageLuminance = Settings.AutoExposureMaxBrightness;
 		}
 	}
 
+	MinAverageLuminance = FMath::Min(MinAverageLuminance, MaxAverageLuminance);
+
 	// This scales the average luminance BEFORE it gets clamped. Note that AEM_Histogram implements the calibration constant through ExposureLowPercent and ExposureHighPercent.
 	const float CalibrationConstant = FMath::Clamp(Settings.AutoExposureCalibrationConstant, 1.0f, 100.0f) * PercentToScale;
 
