DR - Fix materialEditor crash

diff --git a/Engine/Source/Runtime/Renderer/Private/RendererScene.cpp b/Engine/Source/Runtime/Renderer/Private/RendererScene.cpp
--- a/Engine/Source/Runtime/Renderer/Private/RendererScene.cpp
+++ b/Engine/Source/Runtime/Renderer/Private/RendererScene.cpp
@@ -4120,6 +4120,7 @@ void UpdateStaticMeshesForMaterials(const TArray<const FMaterial*>& MaterialReso
 					ENQUEUE_RENDER_COMMAND(FUpdateStaticMeshesForMaterials)(
 						[SceneProxy](FRHICommandListImmediate& RHICmdList)
 						{
+							SceneProxy->GetScene().UpdateAllPrimitiveSceneInfos(RHICmdList);
 							// Defer the caching until the next render tick, to make sure that all render components queued
 							// for re-creation are processed. Otherwise, we may end up caching mesh commands from stale data.
 							bool bReAddToDrawLists = false;
