Fixed crash in Infiltrator Demo when selecting reflective shield material.

diff --git a/Engine/Source/Runtime/Renderer/Private/PrimitiveSceneInfo.cpp b/Engine/Source/Runtime/Renderer/Private/PrimitiveSceneInfo.cpp
--- a/Engine/Source/Runtime/Renderer/Private/PrimitiveSceneInfo.cpp
+++ b/Engine/Source/Runtime/Renderer/Private/PrimitiveSceneInfo.cpp
@@ -505,6 +505,9 @@ void FPrimitiveSceneInfo::AddToScene(FRHICommandListImmediate& RHICmdList, bool
 {
 	check(IsInRenderingThread());
 
+	// The scene should have allocated an index.
+	check(IsIndexValid());
+
 	// Create an indirect lighting cache uniform buffer if we attaching a primitive that may require it, as it may be stored inside a cached mesh command.
 	if (IsIndirectLightingCacheAllowed(Scene->GetFeatureLevel())
 		&& Proxy->WillEverBeLit()
@@ -770,26 +773,34 @@ void FPrimitiveSceneInfo::UpdateRuntimeVirtualTextureFlags()
 
 bool FPrimitiveSceneInfo::NeedsUpdateStaticMeshes()
 {
-	return Scene->PrimitivesNeedingStaticMeshUpdate[PackedIndex];
+	if (IsIndexValid())
+	{
+		return Scene->PrimitivesNeedingStaticMeshUpdate[PackedIndex];
+	}
+	return false;
 }
 
 void FPrimitiveSceneInfo::UpdateStaticMeshes(FRHICommandListImmediate& RHICmdList, bool bReAddToDrawLists)
 {
 	QUICK_SCOPE_CYCLE_COUNTER(STAT_FPrimitiveSceneInfo_UpdateStaticMeshes);
-	const bool bNeedsStaticMeshUpdate = !bReAddToDrawLists;
-	Scene->PrimitivesNeedingStaticMeshUpdate[PackedIndex] = bNeedsStaticMeshUpdate;
 
-	if (!bNeedsStaticMeshUpdate && bNeedsStaticMeshUpdateWithoutVisibilityCheck)
+	if (IsIndexValid())
 	{
-		Scene->PrimitivesNeedingStaticMeshUpdateWithoutVisibilityCheck.Remove(this);
+		const bool bNeedsStaticMeshUpdate = !bReAddToDrawLists;
+		Scene->PrimitivesNeedingStaticMeshUpdate[PackedIndex] = bNeedsStaticMeshUpdate;
 
-		bNeedsStaticMeshUpdateWithoutVisibilityCheck = false;
-	}
+		if (!bNeedsStaticMeshUpdate && bNeedsStaticMeshUpdateWithoutVisibilityCheck)
+		{
+			Scene->PrimitivesNeedingStaticMeshUpdateWithoutVisibilityCheck.Remove(this);
 
-	RemoveCachedMeshDrawCommands();
-	if (bReAddToDrawLists)
-	{
-		CacheMeshDrawCommands(RHICmdList);
+			bNeedsStaticMeshUpdateWithoutVisibilityCheck = false;
+		}
+
+		RemoveCachedMeshDrawCommands();
+		if (bReAddToDrawLists)
+		{
+			CacheMeshDrawCommands(RHICmdList);
+		}
 	}
 }
 
