Fixed bug with post process passes not using the correct scene color view rect.

diff --git a/Engine/Source/Runtime/Renderer/Private/PostProcess/PostProcessDownsample.cpp b/Engine/Source/Runtime/Renderer/Private/PostProcess/PostProcessDownsample.cpp
--- a/Engine/Source/Runtime/Renderer/Private/PostProcess/PostProcessDownsample.cpp
+++ b/Engine/Source/Runtime/Renderer/Private/PostProcess/PostProcessDownsample.cpp
@@ -194,22 +194,25 @@ FRenderingCompositeOutputRef AddDownsamplePass(
 	FRenderingCompositionGraph& Graph,
 	const TCHAR *InName,
 	FRenderingCompositeOutputRef Input,
+	uint32 SceneColorDownsampleFactor,
 	EDownsampleQuality InQuality,
 	EDownsampleFlags InFlags,
 	EPixelFormat InFormatOverride)
 {
 	FRenderingCompositePass* DownsamplePass = Graph.RegisterPass(
 		new(FMemStack::Get()) TRCPassForRDG<1, 1>(
-			[InFormatOverride, InQuality, InFlags, InName] (FRenderingCompositePass* Pass, FRenderingCompositePassContext& InContext)
+			[InFormatOverride, InQuality, SceneColorDownsampleFactor, InFlags, InName] (FRenderingCompositePass* Pass, FRenderingCompositePassContext& InContext)
 	{
 		FRDGBuilder GraphBuilder(InContext.RHICmdList);
 
+		const FIntRect SceneColorViewRect = InContext.GetDownsampledSceneColorViewRect(SceneColorDownsampleFactor);
+
 		FRDGTextureRef InputTexture = Pass->CreateRDGTextureForRequiredInput(GraphBuilder, ePId_Input0, TEXT("DownsampleInput"));
 
 		FDownsamplePassInputs PassInputs;
 		PassInputs.Name = InName;
 		PassInputs.Texture = InputTexture;
-		PassInputs.Viewport = InContext.SceneColorViewRect;
+		PassInputs.Viewport = SceneColorViewRect;
 		PassInputs.FormatOverride = InFormatOverride;
 		PassInputs.Quality = InQuality;
 		PassInputs.Flags = InFlags;
diff --git a/Engine/Source/Runtime/Renderer/Private/PostProcess/PostProcessDownsample.h b/Engine/Source/Runtime/Renderer/Private/PostProcess/PostProcessDownsample.h
--- a/Engine/Source/Runtime/Renderer/Private/PostProcess/PostProcessDownsample.h
+++ b/Engine/Source/Runtime/Renderer/Private/PostProcess/PostProcessDownsample.h
@@ -73,6 +73,7 @@ FRenderingCompositeOutputRef AddDownsamplePass(
 	FRenderingCompositionGraph& Graph,
 	const TCHAR *Name,
 	FRenderingCompositeOutputRef Input,
+	uint32 SceneColorDownsampleFactor,
 	EDownsampleQuality Quality = EDownsampleQuality::Low,
 	EDownsampleFlags Flags = EDownsampleFlags::None,
 	EPixelFormat FormatOverride = PF_Unknown);
\ No newline at end of file
diff --git a/Engine/Source/Runtime/Renderer/Private/PostProcess/PostProcessEyeAdaptation.cpp b/Engine/Source/Runtime/Renderer/Private/PostProcess/PostProcessEyeAdaptation.cpp
--- a/Engine/Source/Runtime/Renderer/Private/PostProcess/PostProcessEyeAdaptation.cpp
+++ b/Engine/Source/Runtime/Renderer/Private/PostProcess/PostProcessEyeAdaptation.cpp
@@ -472,18 +472,20 @@ FRDGTextureRef AddBasicEyeAdaptationSetupPass(
 FRenderingCompositeOutputRef AddBasicEyeAdaptationSetupPass(
 	FPostprocessContext& Context,
 	FRenderingCompositeOutputRef SceneColor,
-	FIntRect SceneColorViewRect)
+	uint32 SceneColorDownsampleFactor)
 {
 	const FViewInfo& View = Context.View;
 
 	FRenderingCompositePass* Pass = Context.Graph.RegisterPass(
 		new(FMemStack::Get()) TRCPassForRDG<1, 1>(
-			[&View, SceneColorViewRect](FRenderingCompositePass* InPass, FRenderingCompositePassContext& InContext)
+			[&View, SceneColorDownsampleFactor](FRenderingCompositePass* InPass, FRenderingCompositePassContext& InContext)
 	{
 		FRDGBuilder GraphBuilder(InContext.RHICmdList);
 
 		const FEyeAdaptationParameters EyeAdaptationParameters = GetEyeAdaptationParameters(View);
 
+		const FIntRect SceneColorViewRect = InContext.GetDownsampledSceneColorViewRect(SceneColorDownsampleFactor);
+
 		FRDGTextureRef SceneColorTexture = InPass->CreateRDGTextureForRequiredInput(GraphBuilder, ePId_Input0, TEXT("SceneColor"));
 
 		FRDGTextureRef OutputTexture = AddBasicEyeAdaptationSetupPass(GraphBuilder, FScreenPassViewInfo(View), EyeAdaptationParameters, SceneColorTexture, SceneColorViewRect);
@@ -628,18 +630,20 @@ FRDGTextureRef AddBasicEyeAdaptationPass(
 FRenderingCompositeOutputRef AddBasicEyeAdaptationPass(
 	FPostprocessContext& Context,
 	FRenderingCompositeOutputRef SceneColor,
-	FIntRect SceneColorViewRect)
+	uint32 SceneColorDownsampleFactor)
 {
 	const FViewInfo& View = Context.View;
 
 	FRenderingCompositePass* Pass = Context.Graph.RegisterPass(
 		new(FMemStack::Get()) TRCPassForRDG<1, 1>(
-			[&View, SceneColorViewRect](FRenderingCompositePass* InPass, FRenderingCompositePassContext& InContext)
+			[&View, SceneColorDownsampleFactor](FRenderingCompositePass* InPass, FRenderingCompositePassContext& InContext)
 	{
 		FRDGBuilder GraphBuilder(InContext.RHICmdList);
 
 		const FEyeAdaptationParameters EyeAdaptationParameters = GetEyeAdaptationParameters(View);
 
+		const FIntRect SceneColorViewRect = InContext.GetDownsampledSceneColorViewRect(SceneColorDownsampleFactor);
+
 		FRDGTextureRef SceneColorTexture = InPass->CreateRDGTextureForRequiredInput(GraphBuilder, ePId_Input0, TEXT("SceneColor"));
 
 		FRDGTextureRef EyeAdaptationTexture = GetEyeAdaptationTexture(GraphBuilder, View);
diff --git a/Engine/Source/Runtime/Renderer/Private/PostProcess/PostProcessEyeAdaptation.h b/Engine/Source/Runtime/Renderer/Private/PostProcess/PostProcessEyeAdaptation.h
--- a/Engine/Source/Runtime/Renderer/Private/PostProcess/PostProcessEyeAdaptation.h
+++ b/Engine/Source/Runtime/Renderer/Private/PostProcess/PostProcessEyeAdaptation.h
@@ -62,11 +62,11 @@ float GetEyeAdaptationFixedExposure(const FViewInfo& View);
 FRenderingCompositeOutputRef AddBasicEyeAdaptationSetupPass(
 	FPostprocessContext& Context,
 	FRenderingCompositeOutputRef SceneColor,
-	FIntRect SceneColorViewRect);
+	uint32 SceneColorDownsampleFactor);
 
 FRenderingCompositeOutputRef AddBasicEyeAdaptationPass(
 	FPostprocessContext& Context,
 	FRenderingCompositeOutputRef SceneColor,
-	FIntRect SceneColorViewRect);
+	uint32 SceneColorDownsampleFactor);
 
 FRenderingCompositeOutputRef AddHistogramEyeAdaptationPass(FPostprocessContext& Context, FRenderingCompositeOutputRef Histogram);
\ No newline at end of file
diff --git a/Engine/Source/Runtime/Renderer/Private/PostProcess/PostProcessing.cpp b/Engine/Source/Runtime/Renderer/Private/PostProcess/PostProcessing.cpp
--- a/Engine/Source/Runtime/Renderer/Private/PostProcess/PostProcessing.cpp
+++ b/Engine/Source/Runtime/Renderer/Private/PostProcess/PostProcessing.cpp
@@ -233,40 +233,27 @@ public:
 
 		const EDownsampleQuality DownsampleQuality = GetDownsampleQuality();
 
-		const auto DivideAndRoundUpClamped = [](const FIntRect Rect, uint32 DownsampleFactor)
-		{
-			FIntRect Out = FIntRect::DivideAndRoundUp(Rect, DownsampleFactor);
-			Out.Max.X = FMath::Max(1, Out.Max.X);
-			Out.Max.Y = FMath::Max(1, Out.Max.Y);
-			return Out;
-		};
-
-		FIntRect InputViewRect = DivideAndRoundUpClamped(InContext.View.ViewRect, 2);
-
 		// Queue the down samples. 
 		for (int i = 1; i < DownSampleStages; i++)
 		{
-			PostProcessDownsamples[i] = AddDownsamplePass(Context.Graph, PassLabels[i], PostProcessDownsamples[i - 1], DownsampleQuality);
+			uint32 DownsampleFactor = 1 << i;
 
-			InputViewRect = DivideAndRoundUpClamped(InputViewRect, 2);
+			PostProcessDownsamples[i] = AddDownsamplePass(Context.Graph, PassLabels[i], PostProcessDownsamples[i - 1], DownsampleFactor, DownsampleQuality);
 
 			// Add log2 data to the alpha channel after doing the 1st (i==1) down sample pass
 			if (bHasLog2Alpha && i == 1 )
 			{
-				PostProcessDownsamples[i] = AddBasicEyeAdaptationSetupPass(InContext, PostProcessDownsamples[i], InputViewRect);
+				DownsampleFactor <<= 1;
+
+				PostProcessDownsamples[i] = AddBasicEyeAdaptationSetupPass(InContext, PostProcessDownsamples[i], DownsampleFactor);
 			}
 		}
-
-		FinalViewRect = InputViewRect;
 	}
 
 	// The number of elements in the array.
 	inline static int32 Num() { return DownSampleStages; }
 
-	FIntRect GetFinalViewRect() const
-	{
-		return FinalViewRect;
-	}
+	inline static uint32 GetFinalDownsampleFactor() { return 1 << DownSampleStages; }
 
 	// Member data kept public for simplicity
 	bool bHasLog2Alpha;
@@ -276,8 +263,6 @@ public:
 private:
 	// no default constructor.
 	TBloomDownSampleArray() {};
-
-	FIntRect FinalViewRect;
 };
 
 // Standard DownsampleArray shared by Bloom, Tint, and Eye-Adaptation. 
@@ -414,9 +399,7 @@ static FRenderingCompositeOutputRef AddPostProcessBasicEyeAdaptation(const FView
 	static const int32 FinalDSIdx = FBloomDownSampleArray::Num() - 1;
 	FRenderingCompositeOutputRef PostProcessPriorReduction = BloomAndEyeDownSamples.PostProcessDownsamples[FinalDSIdx];
 
-	const FIntRect DownsampledViewRect = BloomAndEyeDownSamples.GetFinalViewRect();
-
-	return AddBasicEyeAdaptationPass(Context, PostProcessPriorReduction, DownsampledViewRect);
+	return AddBasicEyeAdaptationPass(Context, PostProcessPriorReduction, BloomAndEyeDownSamples.GetFinalDownsampleFactor());
 }
 
 static void AddVisualizeBloomOverlay(FPostprocessContext& Context, FRenderingCompositeOutputRef& HDRColor, FRenderingCompositeOutputRef& BloomOutputCombined)
@@ -454,7 +437,9 @@ static bool AddPostProcessDepthOfFieldGaussian(FPostprocessContext& Context, FDe
 		FRenderingCompositeOutputRef SetupInput(Context.FinalOutput);
 		if (bMobileQuality)
 		{
-			SetupInput = AddDownsamplePass(Context.Graph, TEXT("GaussianSetupHalfRes"), SetupInput, EDownsampleQuality::High, EDownsampleFlags::ForceRaster, PF_FloatRGBA);
+			const uint32 SetupInputDownsampleFactor = 1;
+
+			SetupInput = AddDownsamplePass(Context.Graph, TEXT("GaussianSetupHalfRes"), SetupInput, SetupInputDownsampleFactor, EDownsampleQuality::High, EDownsampleFlags::ForceRaster, PF_FloatRGBA);
 		}
 
 		FRenderingCompositePass* DOFSetupPass = Context.Graph.RegisterPass(new(FMemStack::Get()) FRCPassPostProcessDOFSetup(bFar, bNear));
@@ -778,10 +763,10 @@ static void AddGBufferVisualizationOverview(
 					if (bOverviewModeEnabled)
 					{
 						// Down-sample to 1/2 size
-						FRenderingCompositeOutputRef HalfSize = AddDownsamplePass(Context.Graph, TEXT("MaterialHalfSize"), MaterialPass, EDownsampleQuality::Low, EDownsampleFlags::ForceRaster);
+						FRenderingCompositeOutputRef HalfSize = AddDownsamplePass(Context.Graph, TEXT("MaterialHalfSize"), MaterialPass, 2, EDownsampleQuality::Low, EDownsampleFlags::ForceRaster);
 
 						// Down-sample to 1/4 size
-						FRenderingCompositeOutputRef QuarterSize = AddDownsamplePass(Context.Graph, TEXT("MaterialQuarterSize"), HalfSize, EDownsampleQuality::Low, EDownsampleFlags::ForceRaster);
+						FRenderingCompositeOutputRef QuarterSize = AddDownsamplePass(Context.Graph, TEXT("MaterialQuarterSize"), HalfSize, 4, EDownsampleQuality::Low, EDownsampleFlags::ForceRaster);
 
 						// Set whether current buffer is selected
 						bool bIsSelected = false;
diff --git a/Engine/Source/Runtime/Renderer/Private/PostProcess/RenderingCompositionGraph.h b/Engine/Source/Runtime/Renderer/Private/PostProcess/RenderingCompositionGraph.h
--- a/Engine/Source/Runtime/Renderer/Private/PostProcess/RenderingCompositionGraph.h
+++ b/Engine/Source/Runtime/Renderer/Private/PostProcess/RenderingCompositionGraph.h
@@ -174,6 +174,14 @@ struct FRenderingCompositePassContext
 		return SceneColorViewRect;
 	}
 
+	FIntRect GetDownsampledSceneColorViewRect(uint32 DownsampleFactor) const
+	{
+		FIntRect Rect = FIntRect::DivideAndRoundUp(SceneColorViewRect, DownsampleFactor);
+		Rect.Max.X = FMath::Max(Rect.Max.X, 1);
+		Rect.Max.Y = FMath::Max(Rect.Max.Y, 1);
+		return Rect;
+	}
+
 	/** Returns the LoadAction that should be use for a given render target. */
 	ERenderTargetLoadAction GetLoadActionForRenderTarget(const FSceneRenderTargetItem& DestRenderTarget) const;
 
