- Fix bug where non subsurface pixel and viewport edge will affect the scattering.
- First iteration of performance optimization for sss on console.

diff --git a/Engine/Shaders/Private/SubsurfaceBurleyNormalized.ush b/Engine/Shaders/Private/SubsurfaceBurleyNormalized.ush
--- a/Engine/Shaders/Private/SubsurfaceBurleyNormalized.ush
+++ b/Engine/Shaders/Private/SubsurfaceBurleyNormalized.ush
@@ -357,17 +357,12 @@ bool IsIDInsideLocalShared(float2 GroupThreadId)
 #ifndef TEXTURE_CACHE_DISABLED
 float4 SampleSharedDiffuse(float2 SampleGroupThreadId)
 {
-#if NONSUBSURFACE_SAMPLE_STRATEGY == 0
-#define SAMPLDE_POINT
-#elif NONSUBSURFACE_SAMPLE_STRATEGY == 1
-#define SAMPLE_BILINEAR
-#endif
 
-#ifdef SAMPLDE_POINT
+#if SUBSURFACE_SAMPLER_TYPE == 0
 	// We use point sampling by now
 	int2 Id = SampleGroupThreadId + THREAD_TEXTURE_BOARDER;
 	return SharedSubsurfaceDiffuseLighting[Id.x + Id.y*LARGE_GROUP_TEXTURE_DIMENSION_SIZE];
-#elif defined(SAMPLE_BILINEAR)
+#elif SUBSURFACE_SAMPLER_TYPE == 1
 
 	//ref: https://en.wikipedia.org/wiki/Bilinear_interpolation
 	int2 Id00 = floor(SampleGroupThreadId) + THREAD_TEXTURE_BOARDER;
@@ -389,11 +384,12 @@ float4 SampleSharedDiffuse(float2 SampleGroupThreadId)
 
 float4 SampleSSSColorConsideringLocalShared(float2 CenterUV, float2 UVOffset, uint2 CenterGroupThreadID, float MipLevel)
 {
-
+	float2 ClampedUV = clamp(CenterUV+ UVOffset, SubsurfaceInput0_UVViewportBilinearMin, SubsurfaceInput0_UVViewportBilinearMax);
 #ifdef SUBSURFACE_BURLEY_COMPUTE
 
-#ifdef TEXTURE_CACHE_DISABLED 
-	return Texture2DSampleLevel(SubsurfaceInput0_Texture, SubsurfaceSampler0, CenterUV + UVOffset, MipLevel);
+#ifdef TEXTURE_CACHE_DISABLED
+	
+	return Texture2DSampleLevel(SubsurfaceInput0_Texture, SubsurfaceSampler0, ClampedUV, MipLevel);
 #else
 	float2 SampleGroupThreadId = ConvertUVOffset2GTIDOffsetForInput0(UVOffset) + CenterGroupThreadID; // Subtract 0.5 is to make sampling match. (0,0), samples at 1/w*0.5;
 	bool bUseLocalShared = MipLevel == 0 && IsIDInsideLocalShared(SampleGroupThreadId);// We will have artifacts if we do not limit the miplevel
@@ -405,12 +401,12 @@ float4 SampleSSSColorConsideringLocalShared(float2 CenterUV, float2 UVOffset, ui
 		}
 		else
 		{
-			return Texture2DSampleLevel(SubsurfaceInput0_Texture, SubsurfaceSampler0, CenterUV + UVOffset, MipLevel);
+			return Texture2DSampleLevel(SubsurfaceInput0_Texture, SubsurfaceSampler0, ClampedUV, MipLevel);
 		}
 #endif
 
 #else
-	return Texture2DSampleLevel(SubsurfaceInput0_Texture, SubsurfaceSampler0, CenterUV + UVOffset, MipLevel);
+	return Texture2DSampleLevel(SubsurfaceInput0_Texture, SubsurfaceSampler0, ClampedUV, MipLevel);
 #endif
 }
 
@@ -549,14 +545,6 @@ float4 BurleyNormalizedSS(float2 BufferUV, uint2 GroupThreadID)
 			float3 TintColor = (LocalProfileInt == SubsurfaceProfileInt) ? 1.0f : BoundaryColorBleed;
 			Mask = UseSubsurfaceProfile(GetScreenSpaceDataFromCache(SampledDiscUV).GBuffer.ShadingModelID) ? 1 : 0;
 
-			// 1 considers the non subsurface neighbor (0,0,0) to add shadow.
-#if NONSUBSURFACE_SAMPLE_STRATEGY == 1
-			// The bilinear filtering in Separable looks better. So I added this to simulate the look
-			// With out this, the hair does not have shadow at the edge. 2 is an experimental value
-			RadiusSampled = (Mask == 0) ? BurleySampleInfo.Radius * 2 : RadiusSampled;
-			Mask = 1;
-#endif
-
 			// Step 4: accumulate radiance from the diffusion profile rR(r)
 			// make sure the DiffuseMeanFreePath is not zero and in m.
 			float3 DiffusionProfile = GetDiffuseReflectProfileWithDiffuseMeanFreePath(BurleyParameter.SurfaceAlbedo.xyz, BurleyParameter.DiffuseMeanFreePath.xyz, S3D.xyz, RadiusSampled);
@@ -571,7 +559,7 @@ float4 BurleyNormalizedSS(float2 BufferUV, uint2 GroupThreadID)
 		}
 
 	}
-
+	 
 	// 0.99995f is a compensitation to make it energe conservation.
 	const float EnergyNormalization = 1.0f / 0.99995f;
 
diff --git a/Engine/Source/Runtime/Renderer/Private/PostProcess/PostProcessSubsurface.cpp b/Engine/Source/Runtime/Renderer/Private/PostProcess/PostProcessSubsurface.cpp
--- a/Engine/Source/Runtime/Renderer/Private/PostProcess/PostProcessSubsurface.cpp
+++ b/Engine/Source/Runtime/Renderer/Private/PostProcess/PostProcessSubsurface.cpp
@@ -518,18 +518,10 @@ public:
 		MAX
 	};
 
-	// Need this to make hair denser in Burley to simulate the bilinear effect in Separable.
-	// In bilinear setting, Separable samples horizontally and vertically sequentially on the same texture
-	// This will blur the edge between non-subsurface scattering and subsurface scattering. For hair region
-	// there are many small edges interleaving. Thus Bilinear makes the hair dense with small amount of hairs.
-	// Since Burley only samples in one pass, it yields the correct result but the hair looks less dense.
-	// So, for better appearance quality with less hair, we need to consider the non-subsurface samples, which
-	// is set dark in setup pass, in the weighting function. To achieve a denser look when bilinear is enabled,
-	// we consider the non-subsurface dark samples.
-	enum class ENonSubsurfaceSampleStrategy : uint32
+	enum class ESubsurfaceSamplerType : uint32
 	{
-		Ignore,
-		Considered,
+		PointSampler,
+		NonPointSampler,	//bilinear on LDS or trilinear on texture. 
 		MAX
 	};
 
@@ -542,11 +534,11 @@ public:
 
 	class FSubsurfacePassFunction : SHADER_PERMUTATION_ENUM_CLASS("SUBSURFACE_PASS", ESubsurfacePass);
 	class FDimensionQuality : SHADER_PERMUTATION_ENUM_CLASS("SUBSURFACE_QUALITY", EQuality);
-	class FNonSubsurfaceSamplerStrategy : SHADER_PERMUTATION_ENUM_CLASS("NONSUBSURFACE_SAMPLE_STRATEGY", ENonSubsurfaceSampleStrategy);
+	class FSubsurfaceSamplerType : SHADER_PERMUTATION_ENUM_CLASS("SUBSURFACE_SAMPLER_TYPE", ESubsurfaceSamplerType);
 	class FSubsurfaceType : SHADER_PERMUTATION_ENUM_CLASS("SUBSURFACE_TYPE", ESubsurfaceType);
 	class FDimensionHalfRes : SHADER_PERMUTATION_BOOL("SUBSURFACE_HALF_RES");
 	using FPermutationDomain = TShaderPermutationDomain<FSubsurfacePassFunction, FDimensionQuality,
-		FNonSubsurfaceSamplerStrategy, FSubsurfaceType, FDimensionHalfRes>;
+		FSubsurfaceSamplerType, FSubsurfaceType, FDimensionHalfRes>;
 
 	// Returns the sampler state based on the requested SSS filter CVar setting.
 	static FRHISamplerState* GetSamplerState()
@@ -571,15 +563,15 @@ public:
 				static_cast<int32>(FSubsurfaceIndirectDispatchCS::EQuality::High)));
 	}
 
-	static ENonSubsurfaceSampleStrategy GetStrategy()
+	static ESubsurfaceSamplerType GetSamplerType()
 	{
 		if (GetSSSFilter())
 		{
-			return ENonSubsurfaceSampleStrategy::Considered;
+			return ESubsurfaceSamplerType::NonPointSampler;
 		}
 		else
 		{
-			return ENonSubsurfaceSampleStrategy::Ignore;
+			return ESubsurfaceSamplerType::PointSampler;
 		}
 	}
 };
@@ -775,6 +767,14 @@ void ComputeSubsurfaceForView(
 	const FIntPoint TileDimension = FIntPoint::DivideAndRoundUp(SubsurfaceViewport.Extent, SUBSURFACE_GROUP_SIZE);
 	const int32 MaxGroupCount = TileDimension.X*TileDimension.Y;
 
+	const FRDGTextureDesc SceneTextureDescriptor = FRDGTextureDesc::Create2DDesc(
+		SceneViewport.Extent,
+		PF_FloatRGBA,
+		FClearValueBinding(),
+		TexCreate_None,
+		TexCreate_RenderTargetable | TexCreate_ShaderResource | TexCreate_UAV,
+		false);
+
 	const FRDGTextureDesc SubsurfaceTextureDescriptor = FRDGTextureDesc::Create2DDesc(
 		SubsurfaceViewport.Extent,
 		PF_FloatRGBA,
@@ -808,7 +808,7 @@ void ComputeSubsurfaceForView(
 	TRefCountPtr<IPooledRenderTarget>* QualityHistoryState = ViewState ? &ViewState->SubsurfaceScatteringQualityHistoryRT : NULL;
 
 	//allocate/reallocate the quality history texture. 
-	FRDGTextureRef QualityHistoryTexture = RegisterExternalRenderTarget(GraphBuilder, QualityHistoryState, SubsurfaceTextureDescriptor.Extent, TEXT("QualityHistoryTexture"));
+	FRDGTextureRef QualityHistoryTexture = RegisterExternalRenderTarget(GraphBuilder, QualityHistoryState, SceneTextureDescriptor.Extent, TEXT("QualityHistoryTexture"));
 	FRDGTextureRef NewQualityHistoryTexture = nullptr;
 
 	//-----------------------------------------------------------------------------------------------------------------------------------
@@ -996,7 +996,7 @@ void ComputeSubsurfaceForView(
 				SHADER::FPermutationDomain ComputeShaderPermutationVector;
 				ComputeShaderPermutationVector.Set<SHADER::FSubsurfacePassFunction>(SubsurfacePassFunction);
 				ComputeShaderPermutationVector.Set<SHADER::FDimensionQuality>(SHADER::GetQuality());
-				ComputeShaderPermutationVector.Set<SHADER::FNonSubsurfaceSamplerStrategy>(SHADER::GetStrategy());
+				ComputeShaderPermutationVector.Set<SHADER::FSubsurfaceSamplerType>(SHADER::GetSamplerType());
 				ComputeShaderPermutationVector.Set<SHADER::FSubsurfaceType>(SubsurfaceType);
 				ComputeShaderPermutationVector.Set<SHADER::FDimensionHalfRes>(bHalfRes);
 				TShaderMapRef<SHADER> ComputeShader(View.ShaderMap, ComputeShaderPermutationVector);
@@ -1004,10 +1004,6 @@ void ComputeSubsurfaceForView(
 				FComputeShaderUtils::AddPass(GraphBuilder, FRDGEventName(PassInfo.Name), *ComputeShader, PassParameters, SubsurfaceBufferArgs[SubsurfaceTypeIndex], 0);
 			}
 		}
-
-		// Replace the history
-		if (QualityHistoryState)
-			GraphBuilder.QueueTextureExtraction(NewQualityHistoryTexture, QualityHistoryState, true);
 	}
 
 	// Recombines scattering result with scene color.
@@ -1039,6 +1035,11 @@ void ComputeSubsurfaceForView(
 	 */
 	AddDrawScreenPass(GraphBuilder, RDG_EVENT_NAME("SubsurfaceRecombine"), ScreenPassView, SceneViewport, SceneViewport, *PixelShader, PassParameters);
 	}
+
+	if (SubsurfaceMode != ESubsurfaceMode::Bypass && QualityHistoryState)
+	{
+		GraphBuilder.QueueTextureExtraction(NewQualityHistoryTexture, QualityHistoryState, true);
+	}
 }
 
 FRDGTextureRef ComputeSubsurface(
