Fixed SSS removing the UAV flag.

diff --git a/Engine/Source/Runtime/Renderer/Private/PostProcess/PostProcessSubsurface.cpp b/Engine/Source/Runtime/Renderer/Private/PostProcess/PostProcessSubsurface.cpp
--- a/Engine/Source/Runtime/Renderer/Private/PostProcess/PostProcessSubsurface.cpp
+++ b/Engine/Source/Runtime/Renderer/Private/PostProcess/PostProcessSubsurface.cpp
@@ -19,10 +19,8 @@
 #include "GenerateMips.h"
 #include "ClearQuad.h"
 
-
 namespace
 {
-
 	// Subsurface common parameters
 	TAutoConsoleVariable<int32> CVarSubsurfaceScattering(
 		TEXT("r.SubsurfaceScattering"),
@@ -95,11 +93,11 @@ namespace
 		ECVF_RenderThreadSafe);
 }
 
-// Define the size of subsurface group. Set to 16 to use LDS.
-#define SUBSURFACE_GROUP_SIZE        8
-
 // Define to use a custom ps to clear UAV.
-#define USE_CUSTOM_ClearUAV
+#define USE_CUSTOM_CLEAR_UAV
+
+// Define the size of subsurface group. Set to 16 to use LDS.
+const uint32 kSubsurfaceGroupSize = 8;
 
 ENGINE_API IPooledRenderTarget* GetSubsufaceProfileTexture_RT(FRHICommandListImmediate& RHICmdList);
 
@@ -182,22 +180,22 @@ ESubsurfaceMode GetSubsurfaceModeForView(const FViewInfo& View)
 
 // A shader parameter struct for a single subsurface input texture.
 BEGIN_SHADER_PARAMETER_STRUCT(FSubsurfaceInput, )
-SHADER_PARAMETER_STRUCT_INCLUDE(FScreenPassTextureViewportParameters, Viewport)
-SHADER_PARAMETER_RDG_TEXTURE(Texture2D, Texture)
+	SHADER_PARAMETER_STRUCT_INCLUDE(FScreenPassTextureViewportParameters, Viewport)
+	SHADER_PARAMETER_RDG_TEXTURE(Texture2D, Texture)
 END_SHADER_PARAMETER_STRUCT()
 
 BEGIN_SHADER_PARAMETER_STRUCT(FSubsurfaceSRVInput, )
-SHADER_PARAMETER_STRUCT_INCLUDE(FScreenPassTextureViewportParameters, Viewport)
-SHADER_PARAMETER_RDG_TEXTURE_SRV(Texture2D, Texture)
+	SHADER_PARAMETER_STRUCT_INCLUDE(FScreenPassTextureViewportParameters, Viewport)
+	SHADER_PARAMETER_RDG_TEXTURE_SRV(Texture2D, Texture)
 END_SHADER_PARAMETER_STRUCT();
 
 // Set of common shader parameters shared by all subsurface shaders.
 BEGIN_SHADER_PARAMETER_STRUCT(FSubsurfaceParameters, )
-SHADER_PARAMETER(FVector4, SubsurfaceParams)
-SHADER_PARAMETER_STRUCT_REF(FSceneTexturesUniformParameters, SceneUniformBuffer)
-SHADER_PARAMETER_STRUCT_REF(FViewUniformShaderParameters, ViewUniformBuffer)
-SHADER_PARAMETER_SAMPLER(SamplerState, BilinearTextureSampler)
-SHADER_PARAMETER_TEXTURE(Texture2D, SSProfilesTexture)
+	SHADER_PARAMETER(FVector4, SubsurfaceParams)
+	SHADER_PARAMETER_STRUCT_REF(FSceneTexturesUniformParameters, SceneUniformBuffer)
+	SHADER_PARAMETER_STRUCT_REF(FViewUniformShaderParameters, ViewUniformBuffer)
+	SHADER_PARAMETER_SAMPLER(SamplerState, BilinearTextureSampler)
+	SHADER_PARAMETER_TEXTURE(Texture2D, SSProfilesTexture)
 END_SHADER_PARAMETER_STRUCT()
 
 FSubsurfaceParameters GetSubsurfaceCommonParameters(FRHICommandListImmediate& RHICmdList, const FViewInfo& View)
@@ -363,7 +361,7 @@ IMPLEMENT_GLOBAL_SHADER(FSubsurfaceViewportCopyPS, "/Engine/Private/PostProcessS
 
 // Subsurface uniform buffer layout
 BEGIN_GLOBAL_SHADER_PARAMETER_STRUCT(FSubsurfaceUniformParameters, )
-SHADER_PARAMETER(uint32, MaxGroupCount)
+	SHADER_PARAMETER(uint32, MaxGroupCount)
 END_GLOBAL_SHADER_PARAMETER_STRUCT()
 
 IMPLEMENT_GLOBAL_SHADER_PARAMETER_STRUCT(FSubsurfaceUniformParameters, "SubsurfaceUniformParameters");
@@ -394,9 +392,9 @@ public:
 	BEGIN_SHADER_PARAMETER_STRUCT(FParameters, )
 		SHADER_PARAMETER_RDG_BUFFER_UAV(RWStructuredBuffer<uint>, RWSeparableGroupBuffer)
 		SHADER_PARAMETER_RDG_BUFFER_UAV(RWStructuredBuffer<uint>, RWBurleyGroupBuffer)
-		END_SHADER_PARAMETER_STRUCT()
+	END_SHADER_PARAMETER_STRUCT()
 
-		static bool ShouldCompilePermutation(const FGlobalShaderPermutationParameters& Parameters)
+	static bool ShouldCompilePermutation(const FGlobalShaderPermutationParameters& Parameters)
 	{
 		return IsFeatureLevelSupported(Parameters.Platform, ERHIFeatureLevel::SM5);
 	}
@@ -410,7 +408,7 @@ public:
 	DECLARE_GLOBAL_SHADER(FSubsurfaceBuildIndirectDispatchArgsCS);
 	SHADER_USE_PARAMETER_STRUCT(FSubsurfaceBuildIndirectDispatchArgsCS, FSubsurfaceShader)
 
-		static void ModifyCompilationEnvironment(const FGlobalShaderPermutationParameters& Parameters, FShaderCompilerEnvironment& OutEnvironment)
+	static void ModifyCompilationEnvironment(const FGlobalShaderPermutationParameters& Parameters, FShaderCompilerEnvironment& OutEnvironment)
 	{
 		FGlobalShader::ModifyCompilationEnvironment(Parameters, OutEnvironment);
 		OutEnvironment.SetDefine(TEXT("SUBSURFACE_RADIUS_SCALE"), SUBSURFACE_RADIUS_SCALE);
@@ -422,9 +420,9 @@ public:
 		SHADER_PARAMETER_STRUCT_REF(FSubsurfaceUniformParameters, SubsurfaceUniformParameters)
 		SHADER_PARAMETER_RDG_BUFFER_UAV(RWStructuredBuffer<uint>, RWIndirectDispatchArgsBuffer)
 		SHADER_PARAMETER_RDG_BUFFER_SRV(StructuredBuffer<uint>, GroupBuffer)
-		END_SHADER_PARAMETER_STRUCT()
+	END_SHADER_PARAMETER_STRUCT()
 
-		static bool ShouldCompilePermutation(const FGlobalShaderPermutationParameters& Parameters)
+	static bool ShouldCompilePermutation(const FGlobalShaderPermutationParameters& Parameters)
 	{
 		return IsFeatureLevelSupported(Parameters.Platform, ERHIFeatureLevel::SM5);
 	}
@@ -438,7 +436,7 @@ public:
 	DECLARE_GLOBAL_SHADER(FSubsurfaceIndirectDispatchSetupCS);
 	SHADER_USE_PARAMETER_STRUCT(FSubsurfaceIndirectDispatchSetupCS, FSubsurfaceShader)
 
-		static void ModifyCompilationEnvironment(const FGlobalShaderPermutationParameters& Parameters, FShaderCompilerEnvironment& OutEnvironment)
+	static void ModifyCompilationEnvironment(const FGlobalShaderPermutationParameters& Parameters, FShaderCompilerEnvironment& OutEnvironment)
 	{
 		FGlobalShader::ModifyCompilationEnvironment(Parameters, OutEnvironment);
 		OutEnvironment.SetDefine(TEXT("SUBSURFACE_RADIUS_SCALE"), SUBSURFACE_RADIUS_SCALE);
@@ -455,9 +453,9 @@ public:
 		SHADER_PARAMETER_RDG_BUFFER_UAV(RWStructuredBuffer<uint>, RWSeparableGroupBuffer)
 		SHADER_PARAMETER_RDG_BUFFER_UAV(RWStructuredBuffer<uint>, RWBurleyGroupBuffer)
 		SHADER_PARAMETER_STRUCT_REF(FSubsurfaceUniformParameters, SubsurfaceUniformParameters)
-		END_SHADER_PARAMETER_STRUCT()
+	END_SHADER_PARAMETER_STRUCT()
 
-		class FDimensionHalfRes : SHADER_PERMUTATION_BOOL("SUBSURFACE_HALF_RES");
+	class FDimensionHalfRes : SHADER_PERMUTATION_BOOL("SUBSURFACE_HALF_RES");
 	class FDimensionCheckerboard : SHADER_PERMUTATION_BOOL("SUBSURFACE_PROFILE_CHECKERBOARD");
 	using FPermutationDomain = TShaderPermutationDomain<FDimensionHalfRes, FDimensionCheckerboard>;
 };
@@ -477,7 +475,7 @@ public:
 		OutEnvironment.SetDefine(TEXT("SUBSURFACE_KERNEL_SIZE"), SUBSURFACE_KERNEL_SIZE);
 		OutEnvironment.SetDefine(TEXT("SUBSURFACE_BURLEY_COMPUTE"), 1);
 		OutEnvironment.SetDefine(TEXT("ENABLE_VELOCITY"), 1);
-		OutEnvironment.SetDefine(TEXT("SUBSURFACE_GROUP_SIZE"), SUBSURFACE_GROUP_SIZE);
+		OutEnvironment.SetDefine(TEXT("SUBSURFACE_GROUP_SIZE"), kSubsurfaceGroupSize);
 	}
 
 	BEGIN_SHADER_PARAMETER_STRUCT(FParameters, )
@@ -492,10 +490,10 @@ public:
 		SHADER_PARAMETER_SAMPLER(SamplerState, SubsurfaceSampler1)
 		SHADER_PARAMETER_STRUCT(FSubsurfaceInput, SubsurfaceInput2)   // Profile mask | Velocity
 		SHADER_PARAMETER_SAMPLER(SamplerState, SubsurfaceSampler2)
-		END_SHADER_PARAMETER_STRUCT()
+	END_SHADER_PARAMETER_STRUCT()
 
-		// Direction of the 1D separable filter.
-		enum class EDirection : uint32
+	// Direction of the 1D separable filter.
+	enum class EDirection : uint32
 	{
 		Horizontal,
 		Vertical,
@@ -537,8 +535,7 @@ public:
 	class FSubsurfaceSamplerType : SHADER_PERMUTATION_ENUM_CLASS("SUBSURFACE_SAMPLER_TYPE", ESubsurfaceSamplerType);
 	class FSubsurfaceType : SHADER_PERMUTATION_ENUM_CLASS("SUBSURFACE_TYPE", ESubsurfaceType);
 	class FDimensionHalfRes : SHADER_PERMUTATION_BOOL("SUBSURFACE_HALF_RES");
-	using FPermutationDomain = TShaderPermutationDomain<FSubsurfacePassFunction, FDimensionQuality,
-		FSubsurfaceSamplerType, FSubsurfaceType, FDimensionHalfRes>;
+	using FPermutationDomain = TShaderPermutationDomain<FSubsurfacePassFunction, FDimensionQuality, FSubsurfaceSamplerType, FSubsurfaceType, FDimensionHalfRes>;
 
 	// Returns the sampler state based on the requested SSS filter CVar setting.
 	static FRHISamplerState* GetSamplerState()
@@ -589,7 +586,7 @@ class FSubsurfaceSRVResolvePS : public FSubsurfaceShader
 		SHADER_PARAMETER_RDG_TEXTURE(Texture2D, SubsurfaceInput0_Texture)
 		SHADER_PARAMETER_SAMPLER(SamplerState, SubsurfaceSampler0)
 		RENDER_TARGET_BINDING_SLOTS()
-		END_SHADER_PARAMETER_STRUCT();
+	END_SHADER_PARAMETER_STRUCT();
 
 	static bool ShouldCompilePermutation(const FGlobalShaderPermutationParameters& Parameters)
 	{
@@ -632,7 +629,7 @@ FRDGTextureRef ResolveTextureToSRV(FRDGBuilder& GraphBuilder, FRDGTextureRef Inp
 FRDGTextureRef CreateBlackUAVTexture(FRDGBuilder& GraphBuilder, FRDGTextureDesc SRVDesc, const TCHAR* Name, const FScreenPassViewInfo& ScreenPassView,
 	const FScreenPassTextureViewport& SceneViewport)
 {
-#ifdef USE_CUSTOM_ClearUAV
+#ifdef USE_CUSTOM_CLEAR_UAV
 	const FViewInfo& View = ScreenPassView.View;
 
 	SRVDesc.TargetableFlags |= TexCreate_ShaderResource | TexCreate_UAV;
@@ -764,7 +761,7 @@ void ComputeSubsurfaceForView(
 	 */
 	const FScreenPassTextureViewport SubsurfaceViewport = FScreenPassTextureViewport::CreateDownscaled(SceneViewport, ScaleFactor);
 
-	const FIntPoint TileDimension = FIntPoint::DivideAndRoundUp(SubsurfaceViewport.Extent, SUBSURFACE_GROUP_SIZE);
+	const FIntPoint TileDimension = FIntPoint::DivideAndRoundUp(SubsurfaceViewport.Extent, kSubsurfaceGroupSize);
 	const int32 MaxGroupCount = TileDimension.X*TileDimension.Y;
 
 	const FRDGTextureDesc SceneTextureDescriptor = FRDGTextureDesc::Create2DDesc(
@@ -811,15 +808,12 @@ void ComputeSubsurfaceForView(
 	FRDGTextureRef QualityHistoryTexture = RegisterExternalRenderTarget(GraphBuilder, QualityHistoryState, SceneTextureDescriptor.Extent, TEXT("QualityHistoryTexture"));
 	FRDGTextureRef NewQualityHistoryTexture = nullptr;
 
-	//-----------------------------------------------------------------------------------------------------------------------------------
-	//
 	/**
 	 * When in bypass mode, the setup and convolution passes are skipped, but lighting
 	 * reconstruction is still performed in the recombine pass.
 	 */
 	if (SubsurfaceMode != ESubsurfaceMode::Bypass)
 	{
-
 		SetupTexture = GraphBuilder.CreateTexture(SubsurfaceTextureWith6MipsDescriptor, TEXT("SubsurfaceSetupTexture"));
 
 		// Profile masks (between subsurface and non-subsurface) to reduce convergence.
@@ -880,13 +874,13 @@ void ComputeSubsurfaceForView(
 			ComputeShaderPermutationVector.Set<SHADER::FDimensionCheckerboard>(bCheckerboard);
 			TShaderMapRef<SHADER> ComputeShader(View.ShaderMap, ComputeShaderPermutationVector);
 
-			FIntPoint ComputeGroupCount = FIntPoint::DivideAndRoundUp(SubsurfaceViewport.Extent, SUBSURFACE_GROUP_SIZE);
+			FIntPoint ComputeGroupCount = FIntPoint::DivideAndRoundUp(SubsurfaceViewport.Extent, kSubsurfaceGroupSize);
 
 			FComputeShaderUtils::AddPass(GraphBuilder, RDG_EVENT_NAME("SubsurfaceSetup"), *ComputeShader, PassParameters, FIntVector(ComputeGroupCount.X, ComputeGroupCount.Y, 1));
 		}
 
 		{
-			//Generate mipmaps for the diffuse scene color and depth, use bilinear filter
+			// Generate mipmaps for the diffuse scene color and depth, use bilinear filter
 			FGenerateMips::Execute(&GraphBuilder, SetupTexture, BilinearBorderSampler);
 		}
 
@@ -1008,32 +1002,32 @@ void ComputeSubsurfaceForView(
 
 	// Recombines scattering result with scene color.
 	{
-	FSubsurfaceRecombinePS::FParameters* PassParameters = GraphBuilder.AllocParameters<FSubsurfaceRecombinePS::FParameters>();
-	PassParameters->Subsurface = SubsurfaceCommonParameters;
-	PassParameters->RenderTargets[0] = FRenderTargetBinding(SceneTextureOutput, SceneTextureLoadAction);
-	PassParameters->SubsurfaceInput0 = GetSubsurfaceInput(SceneTexture, SceneViewportParameters);
-	PassParameters->SubsurfaceSampler0 = BilinearBorderSampler;
-
-	// Scattering output target is only used when scattering is enabled.
-	if (SubsurfaceMode != ESubsurfaceMode::Bypass)
-	{
-		PassParameters->SubsurfaceInput1 = GetSubsurfaceInput(SubsurfaceSubpassTwoTex, SubsurfaceViewportParameters);
-		PassParameters->SubsurfaceSampler1 = BilinearBorderSampler;
-	}
+		FSubsurfaceRecombinePS::FParameters* PassParameters = GraphBuilder.AllocParameters<FSubsurfaceRecombinePS::FParameters>();
+		PassParameters->Subsurface = SubsurfaceCommonParameters;
+		PassParameters->RenderTargets[0] = FRenderTargetBinding(SceneTextureOutput, SceneTextureLoadAction);
+		PassParameters->SubsurfaceInput0 = GetSubsurfaceInput(SceneTexture, SceneViewportParameters);
+		PassParameters->SubsurfaceSampler0 = BilinearBorderSampler;
+
+		// Scattering output target is only used when scattering is enabled.
+		if (SubsurfaceMode != ESubsurfaceMode::Bypass)
+		{
+			PassParameters->SubsurfaceInput1 = GetSubsurfaceInput(SubsurfaceSubpassTwoTex, SubsurfaceViewportParameters);
+			PassParameters->SubsurfaceSampler1 = BilinearBorderSampler;
+		}
 
-	const FSubsurfaceRecombinePS::EQuality RecombineQuality = FSubsurfaceRecombinePS::GetQuality(View);
+		const FSubsurfaceRecombinePS::EQuality RecombineQuality = FSubsurfaceRecombinePS::GetQuality(View);
 
-	FSubsurfaceRecombinePS::FPermutationDomain PixelShaderPermutationVector;
-	PixelShaderPermutationVector.Set<FSubsurfaceRecombinePS::FDimensionMode>(SubsurfaceMode);
-	PixelShaderPermutationVector.Set<FSubsurfaceRecombinePS::FDimensionQuality>(RecombineQuality);
-	PixelShaderPermutationVector.Set<FSubsurfaceRecombinePS::FDimensionCheckerboard>(bCheckerboard);
-	TShaderMapRef<FSubsurfaceRecombinePS> PixelShader(View.ShaderMap, PixelShaderPermutationVector);
+		FSubsurfaceRecombinePS::FPermutationDomain PixelShaderPermutationVector;
+		PixelShaderPermutationVector.Set<FSubsurfaceRecombinePS::FDimensionMode>(SubsurfaceMode);
+		PixelShaderPermutationVector.Set<FSubsurfaceRecombinePS::FDimensionQuality>(RecombineQuality);
+		PixelShaderPermutationVector.Set<FSubsurfaceRecombinePS::FDimensionCheckerboard>(bCheckerboard);
+		TShaderMapRef<FSubsurfaceRecombinePS> PixelShader(View.ShaderMap, PixelShaderPermutationVector);
 
-	/**
-	 * See the related comment above in the prepare pass. The scene viewport is used as both the target and
-	 * texture viewport in order to ensure that the correct pixel is sampled for checkerboard rendering.
-	 */
-	AddDrawScreenPass(GraphBuilder, RDG_EVENT_NAME("SubsurfaceRecombine"), ScreenPassView, SceneViewport, SceneViewport, *PixelShader, PassParameters);
+		/**
+		 * See the related comment above in the prepare pass. The scene viewport is used as both the target and
+		 * texture viewport in order to ensure that the correct pixel is sampled for checkerboard rendering.
+		 */
+		AddDrawScreenPass(GraphBuilder, RDG_EVENT_NAME("SubsurfaceRecombine"), ScreenPassView, SceneViewport, SceneViewport, *PixelShader, PassParameters);
 	}
 
 	if (SubsurfaceMode != ESubsurfaceMode::Bypass && QualityHistoryState)
@@ -1058,7 +1052,6 @@ FRDGTextureRef ComputeSubsurface(
 	}
 
 	FRDGTextureDesc SceneColorDesc = SceneTexture->Desc;
-	SceneColorDesc.TargetableFlags &= ~TexCreate_UAV;
 	SceneColorDesc.TargetableFlags |= TexCreate_RenderTargetable;
 	FRDGTextureRef SceneTextureOutput = GraphBuilder.CreateTexture(SceneColorDesc, TEXT("SceneColorSubsurface"));
 
