Fixed r.rendertargetpool.events visualization.

diff --git a/Engine/Source/Runtime/Renderer/Private/VisualizeTexturePresent.cpp b/Engine/Source/Runtime/Renderer/Private/VisualizeTexturePresent.cpp
--- a/Engine/Source/Runtime/Renderer/Private/VisualizeTexturePresent.cpp
+++ b/Engine/Source/Runtime/Renderer/Private/VisualizeTexturePresent.cpp
@@ -364,7 +364,9 @@ void FVisualizeTexturePresent::PresentContent(FRHICommandListImmediate& RHICmdLi
 				}
 			}
 
-			Canvas.Flush_RenderThread(RHICmdList);
+			const bool bForce = false;
+			const bool bInsideRenderPass = true;
+			Canvas.Flush_RenderThread(RHICmdList, bForce, bInsideRenderPass);
 
 			GRenderTargetPool.CurrentEventRecordingTime = 0;
 			GRenderTargetPool.RenderTargetPoolEvents.Empty();
