Adds sample per pixel in reflection ray generation pass's name for profiling purposes

diff --git a/Engine/Source/Runtime/Renderer/Private/RayTracing/RayTracingReflections.cpp b/Engine/Source/Runtime/Renderer/Private/RayTracing/RayTracingReflections.cpp
--- a/Engine/Source/Runtime/Renderer/Private/RayTracing/RayTracingReflections.cpp
+++ b/Engine/Source/Runtime/Renderer/Private/RayTracing/RayTracingReflections.cpp
@@ -390,7 +390,7 @@ void FDeferredShadingSceneRenderer::RenderRayTracingReflections(
 		else
 		{
 			GraphBuilder.AddPass(
-				RDG_EVENT_NAME("ReflectionRayTracing %dx%d", RayTracingResolution.X, RayTracingResolution.Y),
+				RDG_EVENT_NAME("ReflectionRayTracing(spp=%d) %dx%d", SamplePerPixel, RayTracingResolution.X, RayTracingResolution.Y),
 				PassParameters,
 				ERDGPassFlags::Compute,
 				[PassParameters, this, &View, RayGenShader, RayTracingResolution, DeferredMaterialBufferNumElements, DeferredMaterialMode](FRHICommandList& RHICmdList)
