Optimization to remove un-needed transition type for distortion rendering

diff --git a/Engine/Source/Runtime/Renderer/Private/DistortionRendering.cpp b/Engine/Source/Runtime/Renderer/Private/DistortionRendering.cpp
--- a/Engine/Source/Runtime/Renderer/Private/DistortionRendering.cpp
+++ b/Engine/Source/Runtime/Renderer/Private/DistortionRendering.cpp
@@ -592,7 +592,8 @@ void FSceneRenderer::RenderDistortion(FRHICommandListImmediate& RHICmdList)
 
 		// Create a texture to store the resolved light attenuation values, and a render-targetable surface to hold the unresolved light attenuation values.
 		{
-			FPooledRenderTargetDesc Desc(FPooledRenderTargetDesc::Create2DDesc(SceneContext.GetBufferSizeXY(), PF_B8G8R8A8, FClearValueBinding::Transparent, TexCreate_None, TexCreate_RenderTargetable, false));
+			uint32 BaseFlags = (GSupportsRenderTargetWriteMask) ? TexCreate_NoFastClearFinalize : TexCreate_None;
+			FPooledRenderTargetDesc Desc(FPooledRenderTargetDesc::Create2DDesc(SceneContext.GetBufferSizeXY(), PF_B8G8R8A8, FClearValueBinding::Transparent, BaseFlags, TexCreate_RenderTargetable, false));
 			Desc.Flags |= GFastVRamConfig.Distortion;
 			Desc.NumSamples = MSAACount;
 			GRenderTargetPool.FindFreeElement(RHICmdList, Desc, DistortionRT, TEXT("Distortion"));
