Backing out null check on UpdateStaticMeshes. Another higher level fix was added instead.

diff --git a/Engine/Source/Runtime/Renderer/Private/PrimitiveSceneInfo.cpp b/Engine/Source/Runtime/Renderer/Private/PrimitiveSceneInfo.cpp
--- a/Engine/Source/Runtime/Renderer/Private/PrimitiveSceneInfo.cpp
+++ b/Engine/Source/Runtime/Renderer/Private/PrimitiveSceneInfo.cpp
@@ -505,9 +505,6 @@ void FPrimitiveSceneInfo::AddToScene(FRHICommandListImmediate& RHICmdList, bool
 {
 	check(IsInRenderingThread());
 
-	// The scene should have allocated an index.
-	check(IsIndexValid());
-
 	// Create an indirect lighting cache uniform buffer if we attaching a primitive that may require it, as it may be stored inside a cached mesh command.
 	if (IsIndirectLightingCacheAllowed(Scene->GetFeatureLevel())
 		&& Proxy->WillEverBeLit()
@@ -773,34 +770,26 @@ void FPrimitiveSceneInfo::UpdateRuntimeVirtualTextureFlags()
 
 bool FPrimitiveSceneInfo::NeedsUpdateStaticMeshes()
 {
-	if (IsIndexValid())
-	{
-		return Scene->PrimitivesNeedingStaticMeshUpdate[PackedIndex];
-	}
-	return false;
+	return Scene->PrimitivesNeedingStaticMeshUpdate[PackedIndex];
 }
 
 void FPrimitiveSceneInfo::UpdateStaticMeshes(FRHICommandListImmediate& RHICmdList, bool bReAddToDrawLists)
 {
 	QUICK_SCOPE_CYCLE_COUNTER(STAT_FPrimitiveSceneInfo_UpdateStaticMeshes);
+	const bool bNeedsStaticMeshUpdate = !bReAddToDrawLists;
+	Scene->PrimitivesNeedingStaticMeshUpdate[PackedIndex] = bNeedsStaticMeshUpdate;
 
-	if (IsIndexValid())
+	if (!bNeedsStaticMeshUpdate && bNeedsStaticMeshUpdateWithoutVisibilityCheck)
 	{
-		const bool bNeedsStaticMeshUpdate = !bReAddToDrawLists;
-		Scene->PrimitivesNeedingStaticMeshUpdate[PackedIndex] = bNeedsStaticMeshUpdate;
+		Scene->PrimitivesNeedingStaticMeshUpdateWithoutVisibilityCheck.Remove(this);
 
-		if (!bNeedsStaticMeshUpdate && bNeedsStaticMeshUpdateWithoutVisibilityCheck)
-		{
-			Scene->PrimitivesNeedingStaticMeshUpdateWithoutVisibilityCheck.Remove(this);
-
-			bNeedsStaticMeshUpdateWithoutVisibilityCheck = false;
-		}
+		bNeedsStaticMeshUpdateWithoutVisibilityCheck = false;
+	}
 
-		RemoveCachedMeshDrawCommands();
-		if (bReAddToDrawLists)
-		{
-			CacheMeshDrawCommands(RHICmdList);
-		}
+	RemoveCachedMeshDrawCommands();
+	if (bReAddToDrawLists)
+	{
+		CacheMeshDrawCommands(RHICmdList);
 	}
 }
 
