Double sided geometry ray traced shadow CVar should effect reflection

diff --git a/Engine/Shaders/Private/RayTracing/RayTracingLightingCommon.ush b/Engine/Shaders/Private/RayTracing/RayTracingLightingCommon.ush
--- a/Engine/Shaders/Private/RayTracing/RayTracingLightingCommon.ush
+++ b/Engine/Shaders/Private/RayTracing/RayTracingLightingCommon.ush
@@ -78,7 +78,11 @@ float TraceShadowRay(
 	in FRayCone ShadowRayCone,
 	in RaytracingAccelerationStructure TLAS)
 {
-	uint RayFlags = RAY_FLAG_CULL_BACK_FACING_TRIANGLES | RAY_FLAG_ACCEPT_FIRST_HIT_AND_END_SEARCH | RAY_FLAG_SKIP_CLOSEST_HIT_SHADER;
+	uint RayFlags = RAY_FLAG_ACCEPT_FIRST_HIT_AND_END_SEARCH | RAY_FLAG_SKIP_CLOSEST_HIT_SHADER;
+
+#if !ENABLE_TWO_SIDED_GEOMETRY
+	RayFlags |= RAY_FLAG_CULL_BACK_FACING_TRIANGLES;
+#endif
 
 	FMaterialClosestHitPayload Payload = TraceRayInternal(
 		TLAS,      // AccelerationStructure
diff --git a/Engine/Source/Runtime/Renderer/Private/RayTracing/RayTracingReflections.cpp b/Engine/Source/Runtime/Renderer/Private/RayTracing/RayTracingReflections.cpp
--- a/Engine/Source/Runtime/Renderer/Private/RayTracing/RayTracingReflections.cpp
+++ b/Engine/Source/Runtime/Renderer/Private/RayTracing/RayTracingReflections.cpp
@@ -150,8 +150,9 @@ class FRayTracingReflectionsRGS : public FGlobalShader
 
 	class FDeferredMaterialMode : SHADER_PERMUTATION_ENUM_CLASS("DIM_DEFERRED_MATERIAL_MODE", EDeferredMaterialMode);
 	class FHybrid : SHADER_PERMUTATION_BOOL("DIM_HYBRID");
+	class FEnableTwoSidedGeometryForShadowDim : SHADER_PERMUTATION_BOOL("ENABLE_TWO_SIDED_GEOMETRY");
 
-	using FPermutationDomain = TShaderPermutationDomain<FDenoiserOutput, FDeferredMaterialMode, FHybrid>;
+	using FPermutationDomain = TShaderPermutationDomain<FDenoiserOutput, FDeferredMaterialMode, FHybrid, FEnableTwoSidedGeometryForShadowDim>;
 
 	BEGIN_SHADER_PARAMETER_STRUCT(FParameters, )
 		SHADER_PARAMETER(int32, SamplesPerPixel)
@@ -214,6 +215,7 @@ void FDeferredShadingSceneRenderer::PrepareRayTracingReflections(const FViewInfo
 	{
 		{
 			FRayTracingReflectionsRGS::FPermutationDomain PermutationVector;
+			PermutationVector.Set<FRayTracingReflectionsRGS::FEnableTwoSidedGeometryForShadowDim>(EnableRayTracingShadowTwoSidedGeometry());
 			PermutationVector.Set<FRayTracingReflectionsRGS::FDeferredMaterialMode>(EDeferredMaterialMode::Gather);
 			PermutationVector.Set<FRayTracingReflectionsRGS::FHybrid>(bHybridReflections);
 			auto RayGenShader = View.ShaderMap->GetShader<FRayTracingReflectionsRGS>(PermutationVector);
@@ -222,6 +224,7 @@ void FDeferredShadingSceneRenderer::PrepareRayTracingReflections(const FViewInfo
 		
 		{
 			FRayTracingReflectionsRGS::FPermutationDomain PermutationVector;
+			PermutationVector.Set<FRayTracingReflectionsRGS::FEnableTwoSidedGeometryForShadowDim>(EnableRayTracingShadowTwoSidedGeometry());
 			PermutationVector.Set<FRayTracingReflectionsRGS::FDeferredMaterialMode>(EDeferredMaterialMode::Shade);
 			PermutationVector.Set<FRayTracingReflectionsRGS::FHybrid>(bHybridReflections);
 			auto RayGenShader = View.ShaderMap->GetShader<FRayTracingReflectionsRGS>(PermutationVector);
@@ -231,6 +234,7 @@ void FDeferredShadingSceneRenderer::PrepareRayTracingReflections(const FViewInfo
 	else
 	{
 		FRayTracingReflectionsRGS::FPermutationDomain PermutationVector;
+		PermutationVector.Set<FRayTracingReflectionsRGS::FEnableTwoSidedGeometryForShadowDim>(EnableRayTracingShadowTwoSidedGeometry());
 		PermutationVector.Set<FRayTracingReflectionsRGS::FDeferredMaterialMode>(EDeferredMaterialMode::None);
 		auto RayGenShader = View.ShaderMap->GetShader<FRayTracingReflectionsRGS>(PermutationVector);
 		OutRayGenShaders.Add(RayGenShader->GetRayTracingShader());
@@ -359,8 +363,9 @@ void FDeferredShadingSceneRenderer::RenderRayTracingReflections(
 		FRayTracingReflectionsRGS::FPermutationDomain PermutationVector;
 		PermutationVector.Set<FRayTracingReflectionsRGS::FDeferredMaterialMode>(DeferredMaterialMode);
 		PermutationVector.Set<FRayTracingReflectionsRGS::FHybrid>(bHybridReflections);
-
+		PermutationVector.Set<FRayTracingReflectionsRGS::FEnableTwoSidedGeometryForShadowDim>(EnableRayTracingShadowTwoSidedGeometry());
 		auto RayGenShader = View.ShaderMap->GetShader<FRayTracingReflectionsRGS>(PermutationVector);
+
 		ClearUnusedGraphResources(RayGenShader, PassParameters);
 
 		if (DeferredMaterialMode == EDeferredMaterialMode::Gather)
diff --git a/Engine/Source/Runtime/Renderer/Private/RayTracing/RayTracingShadows.cpp b/Engine/Source/Runtime/Renderer/Private/RayTracing/RayTracingShadows.cpp
--- a/Engine/Source/Runtime/Renderer/Private/RayTracing/RayTracingShadows.cpp
+++ b/Engine/Source/Runtime/Renderer/Private/RayTracing/RayTracingShadows.cpp
@@ -41,6 +41,11 @@ static TAutoConsoleVariable<int32> CVarRayTracingShadowsEnableTwoSidedGeometry(
 	ECVF_RenderThreadSafe
 );
 
+bool EnableRayTracingShadowTwoSidedGeometry()
+{
+	return CVarRayTracingShadowsEnableTwoSidedGeometry.GetValueOnRenderThread() != 0;
+}
+
 class FOcclusionRGS : public FGlobalShader
 {
 	DECLARE_GLOBAL_SHADER(FOcclusionRGS)
diff --git a/Engine/Source/Runtime/Renderer/Private/RayTracing/RaytracingOptions.h b/Engine/Source/Runtime/Renderer/Private/RayTracing/RaytracingOptions.h
--- a/Engine/Source/Runtime/Renderer/Private/RayTracing/RaytracingOptions.h
+++ b/Engine/Source/Runtime/Renderer/Private/RayTracing/RaytracingOptions.h
@@ -13,6 +13,16 @@ class FViewInfo;
 class FLightSceneInfoCompact;
 class FLightSceneInfo;
 
+// be sure to also update the definition in the `RayTracingPrimaryVisibility.usf`
+enum class ERayTracingPrimaryVisibilityFlag: uint32 
+{
+	None                      =      0,
+	UseGBufferForMaxDistance  = 1 << 0,
+	ConsiderSurfaceScatter	  = 1 << 1,
+	AllowSkipSkySample		  = 1 << 2,
+};
+
+ENUM_CLASS_FLAGS(ERayTracingPrimaryVisibilityFlag);
 
 #if RHI_RAYTRACING
 extern int32 GetForceRayTracingEffectsCVarValue();
@@ -27,6 +37,7 @@ extern bool ShouldRenderRayTracingShadows(const FLightSceneInfoCompact& LightInf
 extern bool ShouldRenderRayTracingStochasticRectLight(const FLightSceneInfo& LightInfo);
 extern bool CanOverlayRayTracingOutput(const FViewInfo& View);
 
+extern bool EnableRayTracingShadowTwoSidedGeometry();
 extern float GetRaytracingMaxNormalBias();
 
 #else
