Fix when compiling shaders with DXC - materials that kill() all fragments by using a constant comparison value (present in Infiltrator and more) will be reduced to just the kill() operation and references to global parameters & the parameters themselves can be removed.

diff --git a/Engine/Source/Runtime/Renderer/Private/SceneHitProxyRendering.cpp b/Engine/Source/Runtime/Renderer/Private/SceneHitProxyRendering.cpp
--- a/Engine/Source/Runtime/Renderer/Private/SceneHitProxyRendering.cpp
+++ b/Engine/Source/Runtime/Renderer/Private/SceneHitProxyRendering.cpp
@@ -163,7 +163,7 @@ public:
 	FHitProxyPS(const ShaderMetaType::CompiledShaderInitializerType& Initializer):
 		FMeshMaterialShader(Initializer)
 	{
-		HitProxyId.Bind(Initializer.ParameterMap,TEXT("HitProxyId"), SPF_Mandatory);
+		HitProxyId.Bind(Initializer.ParameterMap,TEXT("HitProxyId"), SPF_Optional); // There is no way to guarantee that this parameter will be preserved in a material that kill()s all fragments as the optimiser can remove the global - this happens in various projects.
 		PassUniformBuffer.Bind(Initializer.ParameterMap, FSceneTexturesUniformParameters::StaticStructMetadata.GetShaderVariableName());
 	}
 
@@ -898,4 +898,4 @@ FMeshPassProcessor* CreateEditorSelectionPassProcessor(const FScene* Scene, cons
 FRegisterPassProcessorCreateFunction RegisterEditorSelectionPass(&CreateEditorSelectionPassProcessor, EShadingPath::Deferred, EMeshPass::EditorSelection, EMeshPassFlags::MainView);
 FRegisterPassProcessorCreateFunction RegisterMobileEditorSelectionPass(&CreateEditorSelectionPassProcessor, EShadingPath::Mobile, EMeshPass::EditorSelection, EMeshPassFlags::MainView);
 
-#endif
\ No newline at end of file
+#endif
