Fixing compiles, definitely breaking sky. Seb save us please. :)

diff --git a/Engine/Source/Runtime/Renderer/Private/MobileShadingRenderer.cpp b/Engine/Source/Runtime/Renderer/Private/MobileShadingRenderer.cpp
--- a/Engine/Source/Runtime/Renderer/Private/MobileShadingRenderer.cpp
+++ b/Engine/Source/Runtime/Renderer/Private/MobileShadingRenderer.cpp
@@ -373,7 +373,7 @@ void FMobileSceneRenderer::Render(FRHICommandListImmediate& RHICmdList)
 
 	// Generate the Sky/Atmosphere look up tables
 	// Do compute work first, before any graphics work
-	const bool bShouldRenderSkyAtmosphere = ShouldRenderSkyAtmosphere(Scene);
+	const bool bShouldRenderSkyAtmosphere = ShouldRenderSkyAtmosphere(Scene, ViewFamily.EngineShowFlags);
 	if (bShouldRenderSkyAtmosphere)
 	{
 		RenderSkyAtmosphereLookUpTables(RHICmdList);
diff --git a/Engine/Source/Runtime/Renderer/Private/SceneRendering.cpp b/Engine/Source/Runtime/Renderer/Private/SceneRendering.cpp
--- a/Engine/Source/Runtime/Renderer/Private/SceneRendering.cpp
+++ b/Engine/Source/Runtime/Renderer/Private/SceneRendering.cpp
@@ -1253,8 +1253,8 @@ void FViewInfo::SetupUniformBufferParameters(
 		ViewUniformShaderParameters.SkyViewLutSizeAndInvSize = FVector4(SkyViewLutWidth, SkyViewLutHeight, 1.0f / SkyViewLutWidth, 1.0f / SkyViewLutHeight);
 
 		CameraAerialPerspectiveVolumeFound = this->SkyAtmosphereCameraAerialPerspectiveVolume->GetRenderTargetItem().ShaderResourceTexture;
-		DistantSkyLightLutTextureFound = SkyAtmosphere->GetDistantSkyLightLutTextureRHI();
-		ViewUniformShaderParameters.SkyAtmosphereSkyLuminanceFactor = SkyAtmosphere->GetSkyLuminanceFactor();
+		//DistantSkyLightLutTextureFound = SkyAtmosphere->GetDistantSkyLightLutTextureRHI();
+		//ViewUniformShaderParameters.SkyAtmosphereSkyLuminanceFactor = SkyAtmosphere->GetSkyLuminanceFactor();
 
 		// Now initialize remaining view parameters.
 
