Fixed mergedown issues to match 8361260

diff --git a/Engine/Source/Runtime/Renderer/Private/SceneRendering.cpp b/Engine/Source/Runtime/Renderer/Private/SceneRendering.cpp
--- a/Engine/Source/Runtime/Renderer/Private/SceneRendering.cpp
+++ b/Engine/Source/Runtime/Renderer/Private/SceneRendering.cpp
@@ -1319,11 +1319,6 @@ void FViewInfo::SetupUniformBufferParameters(
 		ViewUniformShaderParameters.SkyAtmosphereCameraAerialPerspectiveVolumeDepthSliceLengthInv = 1.0f;
 		ViewUniformShaderParameters.SkyAtmosphereApplyCameraAerialPerspectiveVolume = 1.0f;
 		ViewUniformShaderParameters.SkyWorldCameraOrigin = ViewUniformShaderParameters.WorldCameraOrigin;
-
-		for (uint8 Index = 0; Index < NUM_ATMOSPHERE_LIGHTS; ++Index)
-		{
-			ClearAtmosphereLightData(Index);
-		}
 	}
 
 	for (uint8 Index = AtmosphereLightDataClearStartIndex; Index < NUM_ATMOSPHERE_LIGHTS; ++Index)
