Fixed bad merge issue

diff --git a/Engine/Source/Runtime/Renderer/Private/RayTracing/RayTracingGlobalIllumination.cpp b/Engine/Source/Runtime/Renderer/Private/RayTracing/RayTracingGlobalIllumination.cpp
--- a/Engine/Source/Runtime/Renderer/Private/RayTracing/RayTracingGlobalIllumination.cpp
+++ b/Engine/Source/Runtime/Renderer/Private/RayTracing/RayTracingGlobalIllumination.cpp
@@ -404,23 +404,6 @@ void FDeferredShadingSceneRenderer::RenderRayTracingGlobalIllumination(
 			});
 		}
 		else
-
-		if (GRayTracingGlobalIlluminationTileSize <= 0)
-		{
-			GraphBuilder.AddPass(
-				RDG_EVENT_NAME("GlobalIlluminationRayTracing %dx%d", RayTracingResolution.X, RayTracingResolution.Y),
-				PassParameters,
-				ERDGPassFlags::Compute,
-				[PassParameters, this, &View, RayGenerationShader, RayTracingResolution](FRHICommandList& RHICmdList)
-			{
-				FRHIRayTracingScene* RayTracingSceneRHI = View.RayTracingScene.RayTracingSceneRHI;
-
-				FRayTracingShaderBindingsWriter GlobalResources;
-				SetShaderParameters(GlobalResources, *RayGenerationShader, *PassParameters);
-				RHICmdList.RayTraceDispatch(View.RayTracingMaterialPipeline, RayGenerationShader->GetRayTracingShader(), RayTracingSceneRHI, GlobalResources, RayTracingResolution.X, RayTracingResolution.Y);
-			});
-		}
-		else
 		{
 			int32 TileSize = FMath::Max(32, GRayTracingGlobalIlluminationTileSize);
 			int32 NumTilesX = FMath::DivideAndRoundUp(RayTracingResolution.X, TileSize);
