Fixed lens flares looking dimmer.

diff --git a/Engine/Source/Runtime/Renderer/Private/PostProcess/PostProcessLensFlares.cpp b/Engine/Source/Runtime/Renderer/Private/PostProcess/PostProcessLensFlares.cpp
--- a/Engine/Source/Runtime/Renderer/Private/PostProcess/PostProcessLensFlares.cpp
+++ b/Engine/Source/Runtime/Renderer/Private/PostProcess/PostProcessLensFlares.cpp
@@ -145,6 +145,9 @@ FRDGTextureRef AddLensFlaresPass(
 	{
 		FRDGTextureDesc BlurOutputDesc = Inputs.FlareTexture->Desc;
 		BlurOutputDesc.Reset();
+
+		// More precision for additive blend.
+		BlurOutputDesc.Format = PF_FloatRGBA;
 		BlurOutputDesc.ClearValue = FClearValueBinding(FLinearColor::Transparent);
 
 		BlurOutputTexture = GraphBuilder.CreateTexture(BlurOutputDesc, TEXT("LensFlareBlur"));
diff --git a/Engine/Source/Runtime/Renderer/Private/PostProcess/PostProcessLensFlares.h b/Engine/Source/Runtime/Renderer/Private/PostProcess/PostProcessLensFlares.h
--- a/Engine/Source/Runtime/Renderer/Private/PostProcess/PostProcessLensFlares.h
+++ b/Engine/Source/Runtime/Renderer/Private/PostProcess/PostProcessLensFlares.h
@@ -9,6 +9,7 @@ enum class ELensFlareQuality : uint32
 	Disabled,
 	Low,
 	High,
+	VeryHigh,
 	MAX
 };
 
