Fixing FScreenPassViewInfo lifetime issues for lambda capture.

diff --git a/Engine/Source/Runtime/Renderer/Private/PostProcess/PostProcessCombineLUTs.cpp b/Engine/Source/Runtime/Renderer/Private/PostProcess/PostProcessCombineLUTs.cpp
--- a/Engine/Source/Runtime/Renderer/Private/PostProcess/PostProcessCombineLUTs.cpp
+++ b/Engine/Source/Runtime/Renderer/Private/PostProcess/PostProcessCombineLUTs.cpp
@@ -543,7 +543,9 @@ FRenderingCompositeOutputRef AddCombineLUTPass(FRenderingCompositionGraph& Graph
 	{
 		FRDGBuilder GraphBuilder(InContext.RHICmdList);
 
-		FRDGTextureRef OutputTexture = AddCombineLUTPass(GraphBuilder, FScreenPassViewInfo(InContext.View));
+		FScreenPassViewInfo ScreenPassView(InContext.View);
+
+		FRDGTextureRef OutputTexture = AddCombineLUTPass(GraphBuilder, ScreenPassView);
 
 		InPass->ExtractRDGTextureForOutput(GraphBuilder, ePId_Output0, OutputTexture);
 
diff --git a/Engine/Source/Runtime/Renderer/Private/PostProcess/PostProcessDownsample.cpp b/Engine/Source/Runtime/Renderer/Private/PostProcess/PostProcessDownsample.cpp
--- a/Engine/Source/Runtime/Renderer/Private/PostProcess/PostProcessDownsample.cpp
+++ b/Engine/Source/Runtime/Renderer/Private/PostProcess/PostProcessDownsample.cpp
@@ -217,7 +217,9 @@ FRenderingCompositeOutputRef AddDownsamplePass(
 		PassInputs.Quality = InQuality;
 		PassInputs.Flags = InFlags;
 
-		FDownsamplePassOutputs PassOutputs = AddDownsamplePass(GraphBuilder, FScreenPassViewInfo(InContext.View), PassInputs);
+		FScreenPassViewInfo ScreenPassView(InContext.View);
+
+		FDownsamplePassOutputs PassOutputs = AddDownsamplePass(GraphBuilder, ScreenPassView, PassInputs);
 
 		Pass->ExtractRDGTextureForOutput(GraphBuilder, ePId_Output0, PassOutputs.Texture);
 
diff --git a/Engine/Source/Runtime/Renderer/Private/PostProcess/PostProcessMaterial.cpp b/Engine/Source/Runtime/Renderer/Private/PostProcess/PostProcessMaterial.cpp
--- a/Engine/Source/Runtime/Renderer/Private/PostProcess/PostProcessMaterial.cpp
+++ b/Engine/Source/Runtime/Renderer/Private/PostProcess/PostProcessMaterial.cpp
@@ -668,7 +668,9 @@ FRenderingCompositePass* AddPostProcessMaterialPass(
 			Inputs.CustomDepthTexture = GraphBuilder.RegisterExternalTexture(CustomDepthTarget, TEXT("CustomDepth"));
 		}
 
-		FRDGTextureRef OutputTexture = ComputePostProcessMaterial(GraphBuilder, FScreenPassViewInfo(InContext.View), MaterialInterface, Inputs);
+		FScreenPassViewInfo ScreenPassView(InContext.View);
+
+		FRDGTextureRef OutputTexture = ComputePostProcessMaterial(GraphBuilder, ScreenPassView, MaterialInterface, Inputs);
 
 		Pass->ExtractRDGTextureForOutput(GraphBuilder, ePId_Output0, OutputTexture);
 
diff --git a/Engine/Source/Runtime/Renderer/Private/PostProcess/PostProcessTestImage.cpp b/Engine/Source/Runtime/Renderer/Private/PostProcess/PostProcessTestImage.cpp
--- a/Engine/Source/Runtime/Renderer/Private/PostProcess/PostProcessTestImage.cpp
+++ b/Engine/Source/Runtime/Renderer/Private/PostProcess/PostProcessTestImage.cpp
@@ -124,7 +124,9 @@ FRenderingCompositeOutputRef AddTestImagePass(FRenderingCompositionGraph& Graph,
 
 		FRDGTextureRef SceneTextureOutput = InPass->FindOrCreateRDGTextureForOutput(GraphBuilder, ePId_Output0, OutputTextureDesc, TEXT("TestImage"));
 
-		AddTestImagePass(GraphBuilder, FScreenPassViewInfo(InContext.View), SceneTextureOutput, InContext.SceneColorViewRect);
+		FScreenPassViewInfo ScreenPassView(InContext.View);
+
+		AddTestImagePass(GraphBuilder, ScreenPassView, SceneTextureOutput, InContext.SceneColorViewRect);
 
 		InPass->ExtractRDGTextureForOutput(GraphBuilder, ePId_Output0, SceneTextureOutput);
 
diff --git a/Engine/Source/Runtime/Renderer/Private/PostProcess/PostProcessWeightedSampleSum.cpp b/Engine/Source/Runtime/Renderer/Private/PostProcess/PostProcessWeightedSampleSum.cpp
--- a/Engine/Source/Runtime/Renderer/Private/PostProcess/PostProcessWeightedSampleSum.cpp
+++ b/Engine/Source/Runtime/Renderer/Private/PostProcess/PostProcessWeightedSampleSum.cpp
@@ -617,7 +617,9 @@ FRenderingCompositeOutputRef AddGaussianBlurPass(
 		PassInputs.CrossCenterWeight = CrossCenterWeight;
 		PassInputs.KernelSizePercent = KernelSizePercent;
 
-		FRDGTextureRef OutputTexture = AddGaussianBlurPass(GraphBuilder, FScreenPassViewInfo(InContext.View), PassInputs);
+		FScreenPassViewInfo ScreenPassView(InContext.View);
+
+		FRDGTextureRef OutputTexture = AddGaussianBlurPass(GraphBuilder, ScreenPassView, PassInputs);
 
 		InPass->ExtractRDGTextureForOutput(GraphBuilder, ePId_Output0, OutputTexture);
 
