Fixed merge issue that was causing separate translucency to render incorrectly in multiple viewports.

diff --git a/Engine/Source/Runtime/Renderer/Private/PostProcess/PostProcessing.cpp b/Engine/Source/Runtime/Renderer/Private/PostProcess/PostProcessing.cpp
--- a/Engine/Source/Runtime/Renderer/Private/PostProcess/PostProcessing.cpp
+++ b/Engine/Source/Runtime/Renderer/Private/PostProcess/PostProcessing.cpp
@@ -973,11 +973,11 @@ void FPostProcessing::Process(FRHICommandListImmediate& RHICmdList, const FViewI
 				// the node keeps another reference so the RT will not be release too early
 				SceneContext.FreeSeparateTranslucency();
 			}
-			else
-			{
-				FRenderingCompositePass* NodeSeparateTranslucency = Context.Graph.RegisterPass(new (FMemStack::Get()) FRCPassPostProcessInput(FSceneRenderTargets::Get(RHICmdList).GetSeparateTranslucencyDummy()));
-				SeparateTranslucency = FRenderingCompositeOutputRef(NodeSeparateTranslucency);
-			}
+		}
+		else
+		{
+			FRenderingCompositePass* NodeSeparateTranslucency = Context.Graph.RegisterPass(new (FMemStack::Get()) FRCPassPostProcessInput(FSceneRenderTargets::Get(RHICmdList).GetSeparateTranslucencyDummy()));
+			SeparateTranslucency = FRenderingCompositeOutputRef(NodeSeparateTranslucency);
 		}
 
 		if (VelocityRT)
