Added legacy tonemapper back into CombineLUTs to fix green tint issue on post process volume.

diff --git a/Engine/Shaders/Private/PostProcessCombineLUTs.usf b/Engine/Shaders/Private/PostProcessCombineLUTs.usf
--- a/Engine/Shaders/Private/PostProcessCombineLUTs.usf
+++ b/Engine/Shaders/Private/PostProcessCombineLUTs.usf
@@ -28,6 +28,8 @@ float LUTWeights[5];
 half3 ColorScale;
 half4 OverlayColor;
 
+uint bUseMobileTonemapper;
+
 static const float LUTSize = 32;
 
 // Accurate for 1000K < Temp < 15000K
@@ -269,6 +271,7 @@ float4 CombineLUTsCommon(float2 InUV, uint InLayerIndex)
 	float3 ColorAP1 = mul( sRGB_2_AP1, BalancedColor );
 
 	// Expand bright saturated colors outside the sRGB gamut to fake wide gamut rendering.
+	if (!bUseMobileTonemapper)
 	{
 		float  LumaAP1 = dot( ColorAP1, AP1_RGB2Y );
 		float3 ChromaAP1 = ColorAP1 / LumaAP1;
@@ -328,6 +331,15 @@ float4 CombineLUTsCommon(float2 InUV, uint InLayerIndex)
 	// Convert from AP1 to sRGB and clip out-of-gamut values
 	float3 FilmColor = max(0, mul( AP1_2_sRGB, ColorAP1 ));
 
+#if FEATURE_LEVEL > FEATURE_LEVEL_ES3_1 // Mobile path uses separate shader for legacy tone mapper
+	BRANCH
+	if (bUseMobileTonemapper)
+	{
+		// Legacy tone mapper
+		FilmColor = FilmPostProcess(GradedColor);
+	}
+#endif // FEATURE_LEVEL > FEATURE_LEVEL_ES3_1
+
 #if BLENDCOUNT > 1
 	{
 		// Legacy LDR LUT color grading
diff --git a/Engine/Source/Editor/DetailCustomizations/Private/PostProcessSettingsCustomization.cpp b/Engine/Source/Editor/DetailCustomizations/Private/PostProcessSettingsCustomization.cpp
--- a/Engine/Source/Editor/DetailCustomizations/Private/PostProcessSettingsCustomization.cpp
+++ b/Engine/Source/Editor/DetailCustomizations/Private/PostProcessSettingsCustomization.cpp
@@ -120,8 +120,10 @@ void FPostProcessSettingsCustomization::CustomizeChildren( TSharedRef<IPropertyH
 	static const FName LegacyTonemapperName("LegacyTonemapper");
 	static const FName TonemapperCategory("Film");
 	static const FName MobileTonemapperCategory("Mobile Tonemapper");
+	const bool bDesktopTonemapperFilm = VarTonemapperFilm->GetValueOnGameThread() == 1;
 	const bool bMobileTonemapperFilm = VarMobileTonemapperFilm->GetValueOnGameThread() == 1;
-	const bool bUsingLegacyTonemapper = !bMobileTonemapperFilm;	// Are any platforms use legacy/ES2 tonemapper
+	const bool bUsingFilmTonemapper = bDesktopTonemapperFilm || bMobileTonemapperFilm;		// Are any platforms use film tonemapper
+	const bool bUsingLegacyTonemapper = !bDesktopTonemapperFilm || !bMobileTonemapperFilm;	// Are any platforms use legacy/ES2 tonemapper
 
 	static const auto VarDefaultAutoExposureExtendDefaultLuminanceRange = IConsoleManager::Get().FindTConsoleVariableDataInt(TEXT("r.DefaultFeature.AutoExposure.ExtendDefaultLuminanceRange"));
 	const bool bExtendedLuminanceRange = VarDefaultAutoExposureExtendDefaultLuminanceRange->GetValueOnGameThread() == 1;
@@ -149,16 +151,19 @@ void FPostProcessSettingsCustomization::CustomizeChildren( TSharedRef<IPropertyH
 
 					// Hide in case no platforms use legacy/ES2 tonemapper
 					// Hide in case no platforms use film tonemapper
-					if ((bIsLegacyTonemapperPropery && !bUsingLegacyTonemapper))
+					if ((bIsLegacyTonemapperPropery && !bUsingLegacyTonemapper) || (!bIsLegacyTonemapperPropery && !bUsingFilmTonemapper))
 					{
 						ChildHandle->MarkHiddenByCustomization();
 						continue;
 					}
 
 					// In case platforms use different tonemappers, place mobile settings into separate category
-					if (!bMobileTonemapperFilm && bIsLegacyTonemapperPropery)
+					if (bMobileTonemapperFilm != bDesktopTonemapperFilm)
 					{
-						CategoryFName = MobileTonemapperCategory;
+						if (bMobileTonemapperFilm == !bIsLegacyTonemapperPropery)
+						{
+							CategoryFName = MobileTonemapperCategory;
+						}
 					}
 				}
 				else if (CategoryFName == ExposureCategory && bExtendedLuminanceRange)
diff --git a/Engine/Source/Runtime/Renderer/Private/PostProcess/PostProcessCombineLUTs.cpp b/Engine/Source/Runtime/Renderer/Private/PostProcess/PostProcessCombineLUTs.cpp
--- a/Engine/Source/Runtime/Renderer/Private/PostProcess/PostProcessCombineLUTs.cpp
+++ b/Engine/Source/Runtime/Renderer/Private/PostProcess/PostProcessCombineLUTs.cpp
@@ -1,6 +1,7 @@
 // Copyright 1998-2019 Epic Games, Inc. All Rights Reserved.
 
 #include "PostProcess/PostProcessCombineLUTs.h"
+#include "PostProcess/PostProcessTonemap.h"
 
 namespace
 {
@@ -32,6 +33,12 @@ FAutoConsoleVariableRef CVarLUTSize(
 	TEXT("Size of film LUT"),
 	ECVF_RenderThreadSafe);
 
+TAutoConsoleVariable<int32> CVarTonemapperFilm(
+	TEXT("r.TonemapperFilm"),
+	1,
+	TEXT("Use new film tone mapper"),
+	ECVF_RenderThreadSafe);
+
 TAutoConsoleVariable<int32> CVarMobileTonemapperFilm(
 	TEXT("r.Mobile.TonemapperFilm"),
 	0,
@@ -134,8 +141,10 @@ BEGIN_SHADER_PARAMETER_STRUCT(FCombineLUTParameters, )
 	SHADER_PARAMETER(float, ExpandGamut)
 	SHADER_PARAMETER(uint32, OutputDevice)
 	SHADER_PARAMETER(uint32, OutputGamut)
+	SHADER_PARAMETER(uint32, bUseMobileTonemapper)
 	SHADER_PARAMETER_STRUCT_INCLUDE(FColorRemapParameters, ColorRemap)
 	SHADER_PARAMETER_STRUCT_INCLUDE(FFilmTonemapParameters, FilmTonemap)
+	SHADER_PARAMETER_STRUCT_INCLUDE(FMobileFilmTonemapParameters, MobileFilmTonemap)
 END_SHADER_PARAMETER_STRUCT()
 
 void GetCombineLUTParameters(
@@ -205,8 +214,16 @@ void GetCombineLUTParameters(
 	Parameters.BlueCorrection = Settings.BlueCorrection;
 	Parameters.ExpandGamut = Settings.ExpandGamut;
 
+	Parameters.bUseMobileTonemapper = CVarTonemapperFilm.GetValueOnRenderThread() == 0;
+
 	Parameters.FilmTonemap = GetFilmTonemapParameters(Settings);
 
+	Parameters.MobileFilmTonemap = GetMobileFilmTonemapParameters(
+		Settings,
+		/* UseColorMatrix = */ true,
+		/* UseShadowTint = */ true,
+		/* UseContrast = */ true);
+
 	{
 		static TConsoleVariableData<int32>* CVarOutputGamut = IConsoleManager::Get().FindTConsoleVariableDataInt(TEXT("r.HDR.Display.ColorGamut"));
 		static TConsoleVariableData<int32>* CVarOutputDevice = IConsoleManager::Get().FindTConsoleVariableDataInt(TEXT("r.HDR.Display.OutputDevice"));
diff --git a/Engine/Source/Runtime/Renderer/Private/PostProcess/PostProcessTonemap.cpp b/Engine/Source/Runtime/Renderer/Private/PostProcess/PostProcessTonemap.cpp
--- a/Engine/Source/Runtime/Renderer/Private/PostProcess/PostProcessTonemap.cpp
+++ b/Engine/Source/Runtime/Renderer/Private/PostProcess/PostProcessTonemap.cpp
@@ -316,27 +316,26 @@ void GrainPostSettings(FVector* RESTRICT const Constant, const FPostProcessSetti
 	Constant->Z = GrainJitter;
 }
 
-// This code is shared by PostProcessTonemap and VisualizeHDR.
-void FilmPostSetConstants(FVector4* RESTRICT const Constants, const FPostProcessSettings* RESTRICT const FinalPostProcessSettings, bool bMobile, bool UseColorMatrix, bool UseShadowTint, bool UseContrast)
+FMobileFilmTonemapParameters GetMobileFilmTonemapParameters(const FPostProcessSettings& PostProcessSettings, bool bUseColorMatrix, bool bUseShadowTint, bool bUseContrast)
 {
 	// Must insure inputs are in correct range (else possible generation of NaNs).
 	float InExposure = 1.0f;
-	FVector InWhitePoint(FinalPostProcessSettings->FilmWhitePoint);
-	float InSaturation = FMath::Clamp(FinalPostProcessSettings->FilmSaturation, 0.0f, 2.0f);
+	FVector InWhitePoint(PostProcessSettings.FilmWhitePoint);
+	float InSaturation = FMath::Clamp(PostProcessSettings.FilmSaturation, 0.0f, 2.0f);
 	FVector InLuma = FVector(1.0f/3.0f, 1.0f/3.0f, 1.0f/3.0f);
-	FVector InMatrixR(FinalPostProcessSettings->FilmChannelMixerRed);
-	FVector InMatrixG(FinalPostProcessSettings->FilmChannelMixerGreen);
-	FVector InMatrixB(FinalPostProcessSettings->FilmChannelMixerBlue);
-	float InContrast = FMath::Clamp(FinalPostProcessSettings->FilmContrast, 0.0f, 1.0f) + 1.0f;
-	float InDynamicRange = powf(2.0f, FMath::Clamp(FinalPostProcessSettings->FilmDynamicRange, 1.0f, 4.0f));
-	float InToe = (1.0f - FMath::Clamp(FinalPostProcessSettings->FilmToeAmount, 0.0f, 1.0f)) * 0.18f;
+	FVector InMatrixR(PostProcessSettings.FilmChannelMixerRed);
+	FVector InMatrixG(PostProcessSettings.FilmChannelMixerGreen);
+	FVector InMatrixB(PostProcessSettings.FilmChannelMixerBlue);
+	float InContrast = FMath::Clamp(PostProcessSettings.FilmContrast, 0.0f, 1.0f) + 1.0f;
+	float InDynamicRange = powf(2.0f, FMath::Clamp(PostProcessSettings.FilmDynamicRange, 1.0f, 4.0f));
+	float InToe = (1.0f - FMath::Clamp(PostProcessSettings.FilmToeAmount, 0.0f, 1.0f)) * 0.18f;
 	InToe = FMath::Clamp(InToe, 0.18f/8.0f, 0.18f * (15.0f/16.0f));
-	float InHeal = 1.0f - (FMath::Max(1.0f/32.0f, 1.0f - FMath::Clamp(FinalPostProcessSettings->FilmHealAmount, 0.0f, 1.0f)) * (1.0f - 0.18f)); 
-	FVector InShadowTint(FinalPostProcessSettings->FilmShadowTint);
-	float InShadowTintBlend = FMath::Clamp(FinalPostProcessSettings->FilmShadowTintBlend, 0.0f, 1.0f) * 64.0f;
+	float InHeal = 1.0f - (FMath::Max(1.0f/32.0f, 1.0f - FMath::Clamp(PostProcessSettings.FilmHealAmount, 0.0f, 1.0f)) * (1.0f - 0.18f)); 
+	FVector InShadowTint(PostProcessSettings.FilmShadowTint);
+	float InShadowTintBlend = FMath::Clamp(PostProcessSettings.FilmShadowTintBlend, 0.0f, 1.0f) * 64.0f;
 
 	// Shadow tint amount enables turning off shadow tinting.
-	float InShadowTintAmount = FMath::Clamp(FinalPostProcessSettings->FilmShadowTintAmount, 0.0f, 1.0f);
+	float InShadowTintAmount = FMath::Clamp(PostProcessSettings.FilmShadowTintAmount, 0.0f, 1.0f);
 	InShadowTint = InWhitePoint + (InShadowTint - InWhitePoint) * InShadowTintAmount;
 
 	// Make sure channel mixer inputs sum to 1 (+ smart dealing with all zeros).
@@ -370,13 +369,13 @@ void FilmPostSetConstants(FVector4* RESTRICT const Constants, const FPostProcess
 	FVector OutColorShadow_Tint1 = InWhitePoint;
 	FVector OutColorShadow_Tint2 = InShadowTint - InWhitePoint;
 
-	if(UseColorMatrix)
+	if (bUseColorMatrix)
 	{
 		// Final color matrix effected by saturation and exposure.
 		OutMatrixR = (ColorMatrixLuma + ((InMatrixR - ColorMatrixLuma) * InSaturation)) * InExposure;
 		OutMatrixG = (ColorMatrixLuma + ((InMatrixG - ColorMatrixLuma) * InSaturation)) * InExposure;
 		OutMatrixB = (ColorMatrixLuma + ((InMatrixB - ColorMatrixLuma) * InSaturation)) * InExposure;
-		if(!UseShadowTint)
+		if(!bUseShadowTint)
 		{
 			OutMatrixR = OutMatrixR * InWhitePoint.X;
 			OutMatrixG = OutMatrixG * InWhitePoint.Y;
@@ -386,7 +385,7 @@ void FilmPostSetConstants(FVector4* RESTRICT const Constants, const FPostProcess
 	else
 	{
 		// No color matrix fast path.
-		if(!UseShadowTint)
+		if (!bUseShadowTint)
 		{
 			OutMatrixB = InExposure * InWhitePoint;
 		}
@@ -433,7 +432,7 @@ void FilmPostSetConstants(FVector4* RESTRICT const Constants, const FPostProcess
 	float FilmHiG = (-FilmHiYS + (FilmSlope*FilmHeal)) / (FilmSlope*FilmHeal);
 	float FilmLoG = (-FilmLoYS + (FilmSlope*FilmToe)) / (FilmSlope*FilmToe);
 
-	if(UseContrast)
+	if (bUseContrast)
 	{
 		// Constants.
 		OutColorCurveCh1 = FilmHiYS/FilmHiG;
@@ -444,7 +443,7 @@ void FilmPostSetConstants(FVector4* RESTRICT const Constants, const FPostProcess
 		OutColorCurveCm0Cd0 = FilmLoX;
 		OutColorCurveCd3Cm3 = FilmLoY - FilmLoX*FilmSlope;
 		// Handle these separate in case of FilmLoG being 0.
-		if(FilmLoG != 0.0f)
+		if (FilmLoG != 0.0f)
 		{
 			OutColorCurveCd1 = -FilmLoYS/FilmLoG;
 			OutColorCurveCd2 = FilmLoYS/(FilmSlope*FilmLoG);
@@ -473,14 +472,16 @@ void FilmPostSetConstants(FVector4* RESTRICT const Constants, const FPostProcess
 		OutColorCurveCd2 = 0.0f;
 	}
 
-	Constants[0] = FVector4(OutMatrixR, OutColorCurveCd1);
-	Constants[1] = FVector4(OutMatrixG, OutColorCurveCd3Cm3);
-	Constants[2] = FVector4(OutMatrixB, OutColorCurveCm2); 
-	Constants[3] = FVector4(OutColorCurveCm0Cd0, OutColorCurveCd2, OutColorCurveCh0Cm1, OutColorCurveCh3); 
-	Constants[4] = FVector4(OutColorCurveCh1, OutColorCurveCh2, 0.0f, 0.0f);
-	Constants[5] = FVector4(OutColorShadow_Luma, 0.0f);
-	Constants[6] = FVector4(OutColorShadow_Tint1, 0.0f);
-	Constants[7] = FVector4(OutColorShadow_Tint2, 0.0f);
+	FMobileFilmTonemapParameters Parameters;
+	Parameters.ColorMatrixR_ColorCurveCd1 = FVector4(OutMatrixR, OutColorCurveCd1);
+	Parameters.ColorMatrixG_ColorCurveCd3Cm3 = FVector4(OutMatrixG, OutColorCurveCd3Cm3);
+	Parameters.ColorMatrixB_ColorCurveCm2 = FVector4(OutMatrixB, OutColorCurveCm2);
+	Parameters.ColorCurve_Cm0Cd0_Cd2_Ch0Cm1_Ch3 = FVector4(OutColorCurveCm0Cd0, OutColorCurveCd2, OutColorCurveCh0Cm1, OutColorCurveCh3);
+	Parameters.ColorCurve_Ch1_Ch2 = FVector4(OutColorCurveCh1, OutColorCurveCh2, 0.0f, 0.0f);
+	Parameters.ColorShadow_Luma = FVector4(OutColorShadow_Luma, 0.0f);
+	Parameters.ColorShadow_Tint1 = FVector4(OutColorShadow_Tint1, 0.0f);
+	Parameters.ColorShadow_Tint2 = FVector4(OutColorShadow_Tint2, 0.0f);
+	return Parameters;
 }
 
 
@@ -1585,22 +1586,20 @@ public:
 		}
 
 		{
-			FVector4 Constants[8];
-
-			FilmPostSetConstants(Constants, &Context.View.FinalPostProcessSettings,
-				/* bMobile = */ true,
+			const FMobileFilmTonemapParameters Parameters = GetMobileFilmTonemapParameters(
+				Context.View.FinalPostProcessSettings,
 				/* UseColorMatrix = */ PermutationVector.Get<FTonemapperColorMatrixDim>(),
 				/* UseShadowTint = */ PermutationVector.Get<FTonemapperShadowTintDim>(),
 				/* UseContrast = */ PermutationVector.Get<FTonemapperContrastDim>());
 
-			SetShaderValue(RHICmdList, ShaderRHI, ColorMatrixR_ColorCurveCd1, Constants[0]);
-			SetShaderValue(RHICmdList, ShaderRHI, ColorMatrixG_ColorCurveCd3Cm3, Constants[1]);
-			SetShaderValue(RHICmdList, ShaderRHI, ColorMatrixB_ColorCurveCm2, Constants[2]); 
-			SetShaderValue(RHICmdList, ShaderRHI, ColorCurve_Cm0Cd0_Cd2_Ch0Cm1_Ch3, Constants[3]); 
-			SetShaderValue(RHICmdList, ShaderRHI, ColorCurve_Ch1_Ch2, Constants[4]);
-			SetShaderValue(RHICmdList, ShaderRHI, ColorShadow_Luma, Constants[5]);
-			SetShaderValue(RHICmdList, ShaderRHI, ColorShadow_Tint1, Constants[6]);
-			SetShaderValue(RHICmdList, ShaderRHI, ColorShadow_Tint2, Constants[7]);
+			SetShaderValue(RHICmdList, ShaderRHI, ColorMatrixR_ColorCurveCd1, Parameters.ColorMatrixR_ColorCurveCd1);
+			SetShaderValue(RHICmdList, ShaderRHI, ColorMatrixG_ColorCurveCd3Cm3, Parameters.ColorMatrixG_ColorCurveCd3Cm3);
+			SetShaderValue(RHICmdList, ShaderRHI, ColorMatrixB_ColorCurveCm2, Parameters.ColorMatrixB_ColorCurveCm2); 
+			SetShaderValue(RHICmdList, ShaderRHI, ColorCurve_Cm0Cd0_Cd2_Ch0Cm1_Ch3, Parameters.ColorCurve_Cm0Cd0_Cd2_Ch0Cm1_Ch3); 
+			SetShaderValue(RHICmdList, ShaderRHI, ColorCurve_Ch1_Ch2, Parameters.ColorCurve_Ch1_Ch2);
+			SetShaderValue(RHICmdList, ShaderRHI, ColorShadow_Luma, Parameters.ColorShadow_Luma);
+			SetShaderValue(RHICmdList, ShaderRHI, ColorShadow_Tint1, Parameters.ColorShadow_Tint1);
+			SetShaderValue(RHICmdList, ShaderRHI, ColorShadow_Tint2, Parameters.ColorShadow_Tint2);
 		}
 
 		SetShaderValue(RHICmdList, ShaderRHI, SRGBAwareTargetParam, bSRGBAwareTarget ? 1.0f : 0.0f );
diff --git a/Engine/Source/Runtime/Renderer/Private/PostProcess/PostProcessTonemap.h b/Engine/Source/Runtime/Renderer/Private/PostProcess/PostProcessTonemap.h
--- a/Engine/Source/Runtime/Renderer/Private/PostProcess/PostProcessTonemap.h
+++ b/Engine/Source/Runtime/Renderer/Private/PostProcess/PostProcessTonemap.h
@@ -26,7 +26,18 @@ static void GrainRandomFromFrame(FVector* RESTRICT const Constant, uint32 FrameN
 }
 
 
-void FilmPostSetConstants(FVector4* RESTRICT const Constants, const FPostProcessSettings* RESTRICT const FinalPostProcessSettings, bool bMobile, bool UseColorMatrix, bool UseShadowTint, bool UseContrast);
+BEGIN_SHADER_PARAMETER_STRUCT(FMobileFilmTonemapParameters, )
+	SHADER_PARAMETER(FVector4, ColorMatrixR_ColorCurveCd1)
+	SHADER_PARAMETER(FVector4, ColorMatrixG_ColorCurveCd3Cm3)
+	SHADER_PARAMETER(FVector4, ColorMatrixB_ColorCurveCm2)
+	SHADER_PARAMETER(FVector4, ColorCurve_Cm0Cd0_Cd2_Ch0Cm1_Ch3)
+	SHADER_PARAMETER(FVector4, ColorCurve_Ch1_Ch2)
+	SHADER_PARAMETER(FVector4, ColorShadow_Luma)
+	SHADER_PARAMETER(FVector4, ColorShadow_Tint1)
+	SHADER_PARAMETER(FVector4, ColorShadow_Tint2)
+END_SHADER_PARAMETER_STRUCT()
+
+FMobileFilmTonemapParameters GetMobileFilmTonemapParameters(const FPostProcessSettings& PostProcessSettings, bool bUseColorMatrix, bool bUseShadowTint, bool bUseContrast);
 
 // derives from TRenderingCompositePassBase<InputCount, OutputCount>
 // ePId_Input0: SceneColor
diff --git a/Engine/Source/Runtime/Renderer/Private/PostProcess/PostProcessVisualizeHDR.cpp b/Engine/Source/Runtime/Renderer/Private/PostProcess/PostProcessVisualizeHDR.cpp
--- a/Engine/Source/Runtime/Renderer/Private/PostProcess/PostProcessVisualizeHDR.cpp
+++ b/Engine/Source/Runtime/Renderer/Private/PostProcess/PostProcessVisualizeHDR.cpp
@@ -155,20 +155,20 @@ public:
 		}
 
 		{
-			FVector4 Constants[8];
-			FilmPostSetConstants(Constants, &Context.View.FinalPostProcessSettings,
-				/* bMobile = */ false,
+			const FMobileFilmTonemapParameters Parameters = GetMobileFilmTonemapParameters(
+				Context.View.FinalPostProcessSettings,
 				/* UseColorMatrix = */ true,
 				/* UseShadowTint = */ true,
 				/* UseContrast = */ true);
-			SetShaderValue(RHICmdList, ShaderRHI, ColorMatrixR_ColorCurveCd1, Constants[0]);
-			SetShaderValue(RHICmdList, ShaderRHI, ColorMatrixG_ColorCurveCd3Cm3, Constants[1]);
-			SetShaderValue(RHICmdList, ShaderRHI, ColorMatrixB_ColorCurveCm2, Constants[2]);
-			SetShaderValue(RHICmdList, ShaderRHI, ColorCurve_Cm0Cd0_Cd2_Ch0Cm1_Ch3, Constants[3]);
-			SetShaderValue(RHICmdList, ShaderRHI, ColorCurve_Ch1_Ch2, Constants[4]);
-			SetShaderValue(RHICmdList, ShaderRHI, ColorShadow_Luma, Constants[5]);
-			SetShaderValue(RHICmdList, ShaderRHI, ColorShadow_Tint1, Constants[6]);
-			SetShaderValue(RHICmdList, ShaderRHI, ColorShadow_Tint2, Constants[7]);
+
+			SetShaderValue(RHICmdList, ShaderRHI, ColorMatrixR_ColorCurveCd1, Parameters.ColorMatrixR_ColorCurveCd1);
+			SetShaderValue(RHICmdList, ShaderRHI, ColorMatrixG_ColorCurveCd3Cm3, Parameters.ColorMatrixG_ColorCurveCd3Cm3);
+			SetShaderValue(RHICmdList, ShaderRHI, ColorMatrixB_ColorCurveCm2, Parameters.ColorMatrixB_ColorCurveCm2);
+			SetShaderValue(RHICmdList, ShaderRHI, ColorCurve_Cm0Cd0_Cd2_Ch0Cm1_Ch3, Parameters.ColorCurve_Cm0Cd0_Cd2_Ch0Cm1_Ch3);
+			SetShaderValue(RHICmdList, ShaderRHI, ColorCurve_Ch1_Ch2, Parameters.ColorCurve_Ch1_Ch2);
+			SetShaderValue(RHICmdList, ShaderRHI, ColorShadow_Luma, Parameters.ColorShadow_Luma);
+			SetShaderValue(RHICmdList, ShaderRHI, ColorShadow_Tint1, Parameters.ColorShadow_Tint1);
+			SetShaderValue(RHICmdList, ShaderRHI, ColorShadow_Tint2, Parameters.ColorShadow_Tint2);
 		}
 	}
 	
