Disables experimental tile classification pass for SSGI

$jira UE-79911

diff --git a/Engine/Source/Runtime/Renderer/Private/ScreenSpaceRayTracing.cpp b/Engine/Source/Runtime/Renderer/Private/ScreenSpaceRayTracing.cpp
--- a/Engine/Source/Runtime/Renderer/Private/ScreenSpaceRayTracing.cpp
+++ b/Engine/Source/Runtime/Renderer/Private/ScreenSpaceRayTracing.cpp
@@ -259,7 +259,7 @@ class FSSRTDiffuseTileClassificationCS : public FGlobalShader
 
 	static bool ShouldCompilePermutation(const FGlobalShaderPermutationParameters& Parameters)
 	{
-		return IsFeatureLevelSupported(Parameters.Platform, ERHIFeatureLevel::SM5);
+		return false; // Parameters.Platform == SP_PCD3D_SM5;
 	}
 
 	BEGIN_SHADER_PARAMETER_STRUCT(FParameters, )
@@ -810,6 +810,7 @@ void RenderScreenSpaceDiffuseIndirect(
 	// Tile classify.
 	FSSRTTileClassificationParameters ClassificationParameters;
 	FSSRTTileClassificationResources ClassificationResources;
+	#if 0
 	{
 		ClassificationResources = CreateTileClassificationResources(GraphBuilder, View, SceneTextures.SceneDepthBuffer->Desc.Extent, &ClassificationParameters);
 
@@ -853,6 +854,7 @@ void RenderScreenSpaceDiffuseIndirect(
 			PassParameters,
 			FComputeShaderUtils::GetGroupCount(ThreadCount, 8));
 	}
+	#endif
 
 	{
 		// Allocate outputs.
@@ -903,8 +905,8 @@ void RenderScreenSpaceDiffuseIndirect(
 		SetupSceneTextureSamplers(&PassParameters->SceneTextureSamplers);
 		PassParameters->View = View.ViewUniformBuffer;
 	
-		PassParameters->ClassificationParameters = ClassificationParameters;
-		PassParameters->ClassificationSRVs = CreateSRVs(GraphBuilder, ClassificationResources);
+		//PassParameters->ClassificationParameters = ClassificationParameters;
+		//PassParameters->ClassificationSRVs = CreateSRVs(GraphBuilder, ClassificationResources);
 
 		PassParameters->IndirectDiffuseOutput = GraphBuilder.CreateUAV(OutDenoiserInputs->Color);
 		PassParameters->AmbientOcclusionOutput = GraphBuilder.CreateUAV(OutDenoiserInputs->AmbientOcclusionMask);
