Fixing merge conflict

diff --git a/Engine/Source/Runtime/Renderer/Private/RayTracing/RayTracingGlobalIllumination.cpp b/Engine/Source/Runtime/Renderer/Private/RayTracing/RayTracingGlobalIllumination.cpp
--- a/Engine/Source/Runtime/Renderer/Private/RayTracing/RayTracingGlobalIllumination.cpp
+++ b/Engine/Source/Runtime/Renderer/Private/RayTracing/RayTracingGlobalIllumination.cpp
@@ -21,14 +21,15 @@
 #include "Raytracing/RaytracingOptions.h"
 #include "BlueNoise.h"
 #include "SceneTextureParameters.h"
+#include "RayTracingDefinitions.h"
 
-static int32 GRayTracingGlobalIllumination = -1;
-static FAutoConsoleVariableRef CVarRayTracingGlobalIllumination(
+static TAutoConsoleVariable<int32> CVarRayTracingGlobalIllumination(
 	TEXT("r.RayTracing.GlobalIllumination"),
-	GRayTracingGlobalIllumination,
+	-1,
 	TEXT("-1: Value driven by postprocess volume (default) \n")
 	TEXT(" 0: ray tracing global illumination off \n")
-	TEXT(" 1: ray tracing global illumination enabled")
+	TEXT(" 1: ray tracing global illumination enabled"),
+	ECVF_RenderThreadSafe
 );
 
 static int32 GRayTracingGlobalIlluminationSamplesPerPixel = -1;
@@ -112,9 +113,31 @@ static FAutoConsoleVariableRef CVarRayTracingGlobalIlluminationTileSize(
 	TEXT("Render ray traced global illumination in NxN piel tiles, where each tile is submitted as separate GPU command buffer, allowing high quality rendering without triggering timeout detection. (default = 0, tiling disabled)")
 );
 
-static const int32 GLightCountMax = 64;
+static TAutoConsoleVariable<int32> CVarRayTracingGlobalIlluminationMaxLightCount(
+	TEXT("r.RayTracing.GlobalIllumination.MaxLightCount"),
+	RAY_TRACING_LIGHT_COUNT_MAXIMUM,
+	TEXT("Enables two-sided geometry when tracing GI rays (default = 256)"),
+	ECVF_RenderThreadSafe
+);
+
+static TAutoConsoleVariable<int32> CVarRayTracingGlobalIlluminationEnableFinalGather(
+	TEXT("r.RayTracing.GlobalIllumination.EnableFinalGather"),
+	0,
+	TEXT("Enables final gather algorithm for 1-bounce global illumination (default = 0)"),
+	ECVF_RenderThreadSafe
+);
+
+static float GRayTracingGlobalIlluminationFinalGatherDistance = 10.0;
+static FAutoConsoleVariableRef CVarRayTracingGlobalIlluminationFinalGatherDistance(
+	TEXT("r.RayTracing.GlobalIllumination.FinalGatherDistance"),
+	GRayTracingGlobalIlluminationFinalGatherDistance,
+	TEXT("Maximum world-space distance for valid, reprojected final gather points (default = 10)")
+);
 
-DECLARE_GPU_STAT_NAMED(RayTracingGlobalIllumination, TEXT("Ray Tracing Global Illumination"));
+
+DECLARE_GPU_STAT_NAMED(RayTracingGIBruteForce, TEXT("Ray Tracing GI: Brute Force"));
+DECLARE_GPU_STAT_NAMED(RayTracingGIFinalGather, TEXT("Ray Tracing GI: Final Gather"));
+DECLARE_GPU_STAT_NAMED(RayTracingGICreateGatherPoints, TEXT("Ray Tracing GI: Create Gather Points"));
 
 void SetupLightParameters(
 	const TSparseArray<FLightSceneInfoCompact>& Lights,
@@ -130,11 +153,13 @@ void SetupLightParameters(
 	LightParameters->Color[SkyLightIndex] = FVector(1.0);
 	LightParameters->Count++;
 
+	uint32 MaxLightCount = FMath::Min(CVarRayTracingGlobalIlluminationMaxLightCount.GetValueOnRenderThread(), RAY_TRACING_LIGHT_COUNT_MAXIMUM);
 	for (auto Light : Lights)
 	{
-		if (LightParameters->Count >= GLightCountMax) break;
+		if (LightParameters->Count >= MaxLightCount) break;
 
 		if (Light.LightSceneInfo->Proxy->HasStaticLighting() && Light.LightSceneInfo->IsPrecomputedLightingValid()) continue;
+		if (!Light.LightSceneInfo->Proxy->AffectGlobalIllumination()) continue;
 
 		FLightShaderParameters LightShaderParameters;
 		Light.LightSceneInfo->Proxy->GetLightShaderParameters(LightShaderParameters);
@@ -160,6 +185,8 @@ void SetupLightParameters(
 			LightParameters->Color[LightParameters->Count] = LightShaderParameters.Color;
 			LightParameters->Dimensions[LightParameters->Count] = FVector(2.0f * LightShaderParameters.SourceRadius, 2.0f * LightShaderParameters.SourceLength, 0.0f);
 			LightParameters->Attenuation[LightParameters->Count] = 1.0 / LightShaderParameters.InvRadius;
+			LightParameters->RectLightBarnCosAngle[LightParameters->Count] = LightShaderParameters.RectLightBarnCosAngle;
+			LightParameters->RectLightBarnLength[LightParameters->Count] = LightShaderParameters.RectLightBarnLength;
 			break;
 		}
 		case LightType_Point:
@@ -192,19 +219,33 @@ void SetupLightParameters(
 	}
 }
 
+int32 GetRayTracingGlobalIlluminationSamplesPerPixel(const FViewInfo& View)
+{
+	int32 SamplesPerPixel = GRayTracingGlobalIlluminationSamplesPerPixel > -1 ? GRayTracingGlobalIlluminationSamplesPerPixel : View.FinalPostProcessSettings.RayTracingGISamplesPerPixel;
+	return SamplesPerPixel;
+}
+
 bool ShouldRenderRayTracingGlobalIllumination(const FViewInfo& View)
 {
 	if (!IsRayTracingEnabled())
 	{
 		return (false);
 	}
-	else if (GetForceRayTracingEffectsCVarValue() >= 0)
+
+	if (GetRayTracingGlobalIlluminationSamplesPerPixel(View) <= 0)
+	{
+		return false;
+	}
+
+	if (GetForceRayTracingEffectsCVarValue() >= 0)
 	{
 		return GetForceRayTracingEffectsCVarValue() > 0;
 	}
-	else if (GRayTracingGlobalIllumination >= 0)
+
+	int32 CVarRayTracingGlobalIlluminationValue = CVarRayTracingGlobalIllumination.GetValueOnRenderThread();
+	if (CVarRayTracingGlobalIlluminationValue >= 0)
 	{
-		return (GRayTracingGlobalIllumination > 0);
+		return (CVarRayTracingGlobalIlluminationValue > 0);
 	}
 	else 
 	{
@@ -256,7 +297,168 @@ class FGlobalIlluminationRGS : public FGlobalShader
 	END_SHADER_PARAMETER_STRUCT()
 };
 
+class FRayTracingGlobalIlluminationCompositePS : public FGlobalShader
+{
+	DECLARE_GLOBAL_SHADER(FRayTracingGlobalIlluminationCompositePS)
+	SHADER_USE_PARAMETER_STRUCT(FRayTracingGlobalIlluminationCompositePS, FGlobalShader)
+
+	static bool ShouldCompilePermutation(const FGlobalShaderPermutationParameters& Parameters)
+	{
+		return ShouldCompileRayTracingShadersForProject(Parameters.Platform);
+	}
+
+	BEGIN_SHADER_PARAMETER_STRUCT(FParameters, )
+		RENDER_TARGET_BINDING_SLOTS()
+		SHADER_PARAMETER_RDG_TEXTURE(Texture2D, GlobalIlluminationTexture)
+		SHADER_PARAMETER_SAMPLER(SamplerState, GlobalIlluminationSampler)
+		SHADER_PARAMETER_STRUCT_REF(FViewUniformShaderParameters, ViewUniformBuffer)
+		SHADER_PARAMETER_STRUCT_INCLUDE(FSceneTextureParameters, SceneTextures)
+	END_SHADER_PARAMETER_STRUCT()
+};
+
+class FRayTracingGlobalIlluminationSceneColorCompositePS : public FGlobalShader
+{
+	DECLARE_GLOBAL_SHADER(FRayTracingGlobalIlluminationSceneColorCompositePS)
+	SHADER_USE_PARAMETER_STRUCT(FRayTracingGlobalIlluminationSceneColorCompositePS, FGlobalShader)
+
+	static bool ShouldCompilePermutation(const FGlobalShaderPermutationParameters& Parameters)
+	{
+		return ShouldCompileRayTracingShadersForProject(Parameters.Platform);
+	}
+
+	BEGIN_SHADER_PARAMETER_STRUCT(FParameters, )
+		RENDER_TARGET_BINDING_SLOTS()
+		SHADER_PARAMETER_RDG_TEXTURE(Texture2D, GlobalIlluminationTexture)
+		SHADER_PARAMETER_SAMPLER(SamplerState, GlobalIlluminationSampler)
+		SHADER_PARAMETER_STRUCT_INCLUDE(FSceneTextureParameters, SceneTextures)
+	END_SHADER_PARAMETER_STRUCT()
+};
+
+class FRayTracingGlobalIlluminationCHS : public FGlobalShader
+{
+	DECLARE_GLOBAL_SHADER(FRayTracingGlobalIlluminationCHS)
+	SHADER_USE_ROOT_PARAMETER_STRUCT(FRayTracingGlobalIlluminationCHS, FGlobalShader)
+
+	static bool ShouldCompilePermutation(const FGlobalShaderPermutationParameters& Parameters)
+	{
+		return ShouldCompileRayTracingShadersForProject(Parameters.Platform);
+	}
+
+	using FParameters = FEmptyShaderParameters;
+};
+
 IMPLEMENT_GLOBAL_SHADER(FGlobalIlluminationRGS, "/Engine/Private/RayTracing/RayTracingGlobalIlluminationRGS.usf", "GlobalIlluminationRGS", SF_RayGen);
+IMPLEMENT_GLOBAL_SHADER(FRayTracingGlobalIlluminationCHS, "/Engine/Private/RayTracing/RayTracingGlobalIlluminationRGS.usf", "RayTracingGlobalIlluminationCHS", SF_RayHitGroup);
+IMPLEMENT_GLOBAL_SHADER(FRayTracingGlobalIlluminationCompositePS, "/Engine/Private/RayTracing/RayTracingGlobalIlluminationCompositePS.usf", "GlobalIlluminationCompositePS", SF_Pixel);
+IMPLEMENT_GLOBAL_SHADER(FRayTracingGlobalIlluminationSceneColorCompositePS, "/Engine/Private/RayTracing/RayTracingGlobalIlluminationCompositePS.usf", "GlobalIlluminationSceneColorCompositePS", SF_Pixel);
+
+struct GatherPoints
+{
+	FVector CreationPoint[16];
+	FVector Position[16];
+	FVector Irradiance[16];
+};
+
+class FRayTracingGlobalIlluminationCreateGatherPointsRGS : public FGlobalShader
+{
+	DECLARE_GLOBAL_SHADER(FRayTracingGlobalIlluminationCreateGatherPointsRGS)
+	SHADER_USE_ROOT_PARAMETER_STRUCT(FRayTracingGlobalIlluminationCreateGatherPointsRGS, FGlobalShader)
+
+	class FUseAttenuationTermDim : SHADER_PERMUTATION_BOOL("USE_ATTENUATION_TERM");
+	class FEnableTwoSidedGeometryDim : SHADER_PERMUTATION_BOOL("ENABLE_TWO_SIDED_GEOMETRY");
+
+	using FPermutationDomain = TShaderPermutationDomain<FUseAttenuationTermDim, FEnableTwoSidedGeometryDim>;
+
+	static bool ShouldCompilePermutation(const FGlobalShaderPermutationParameters& Parameters)
+	{
+		return ShouldCompileRayTracingShadersForProject(Parameters.Platform);
+	}
+
+	BEGIN_SHADER_PARAMETER_STRUCT(FParameters, )
+		SHADER_PARAMETER(uint32, SamplesPerPixel)
+		SHADER_PARAMETER(uint32, SampleIndex)
+		SHADER_PARAMETER(uint32, MaxBounces)
+		SHADER_PARAMETER(uint32, UpscaleFactor)
+		SHADER_PARAMETER(uint32, TileOffsetX)
+		SHADER_PARAMETER(uint32, TileOffsetY)
+		SHADER_PARAMETER(float, MaxRayDistanceForGI)
+		SHADER_PARAMETER(float, NextEventEstimationSamples)
+		SHADER_PARAMETER(float, DiffuseThreshold)
+		SHADER_PARAMETER(float, MaxNormalBias)
+		SHADER_PARAMETER(bool, EvalSkyLight)
+		SHADER_PARAMETER(bool, UseRussianRoulette)
+
+		// Scene data
+		SHADER_PARAMETER_SRV(RaytracingAccelerationStructure, TLAS)
+		SHADER_PARAMETER_STRUCT_REF(FViewUniformShaderParameters, ViewUniformBuffer)
+
+		// Sampling sequence
+		SHADER_PARAMETER_STRUCT_REF(FHaltonIteration, HaltonIteration)
+		SHADER_PARAMETER_STRUCT_REF(FHaltonPrimes, HaltonPrimes)
+		SHADER_PARAMETER_STRUCT_REF(FBlueNoise, BlueNoise)
+
+		// Light data
+		SHADER_PARAMETER_STRUCT_REF(FPathTracingLightData, LightParameters)
+		SHADER_PARAMETER_STRUCT_REF(FSkyLightData, SkyLight)
+
+		// Shading data
+		SHADER_PARAMETER_STRUCT_INCLUDE(FSceneTextureParameters, SceneTextures)
+		SHADER_PARAMETER_RDG_TEXTURE(Texture2D, SSProfilesTexture)
+		SHADER_PARAMETER_SAMPLER(SamplerState, TransmissionProfilesLinearSampler)
+
+		SHADER_PARAMETER(FIntPoint, GatherPointsResolution)
+		// Output
+		SHADER_PARAMETER_RDG_BUFFER_UAV(RWStructuredBuffer<GatherPoints>, RWGatherPointsBuffer)
+		END_SHADER_PARAMETER_STRUCT()
+};
+
+IMPLEMENT_GLOBAL_SHADER(FRayTracingGlobalIlluminationCreateGatherPointsRGS, "/Engine/Private/RayTracing/RayTracingCreateGatherPointsRGS.usf", "RayTracingCreateGatherPointsRGS", SF_RayGen);
+
+class FRayTracingGlobalIlluminationFinalGatherRGS : public FGlobalShader
+{
+	DECLARE_GLOBAL_SHADER(FRayTracingGlobalIlluminationFinalGatherRGS)
+	SHADER_USE_ROOT_PARAMETER_STRUCT(FRayTracingGlobalIlluminationFinalGatherRGS, FGlobalShader)
+
+	class FUseAttenuationTermDim : SHADER_PERMUTATION_BOOL("USE_ATTENUATION_TERM");
+	class FEnableTwoSidedGeometryDim : SHADER_PERMUTATION_BOOL("ENABLE_TWO_SIDED_GEOMETRY");
+
+	using FPermutationDomain = TShaderPermutationDomain<FUseAttenuationTermDim, FEnableTwoSidedGeometryDim>;
+
+	static bool ShouldCompilePermutation(const FGlobalShaderPermutationParameters& Parameters)
+	{
+		return ShouldCompileRayTracingShadersForProject(Parameters.Platform);
+	}
+
+	BEGIN_SHADER_PARAMETER_STRUCT(FParameters, )
+		SHADER_PARAMETER(uint32, SampleIndex)
+		SHADER_PARAMETER(uint32, SamplesPerPixel)
+		SHADER_PARAMETER(uint32, UpscaleFactor)
+		SHADER_PARAMETER(uint32, TileOffsetX)
+		SHADER_PARAMETER(uint32, TileOffsetY)
+		SHADER_PARAMETER(float, DiffuseThreshold)
+		SHADER_PARAMETER(float, MaxNormalBias)
+		SHADER_PARAMETER(float, FinalGatherDistance)
+
+		// Scene data
+		SHADER_PARAMETER_SRV(RaytracingAccelerationStructure, TLAS)
+		SHADER_PARAMETER_STRUCT_REF(FViewUniformShaderParameters, ViewUniformBuffer)
+
+		// Shading data
+		SHADER_PARAMETER_STRUCT_INCLUDE(FSceneTextureParameters, SceneTextures)
+		SHADER_PARAMETER_RDG_TEXTURE(Texture2D, SSProfilesTexture)
+		SHADER_PARAMETER_SAMPLER(SamplerState, TransmissionProfilesLinearSampler)
+
+		// Gather points
+		SHADER_PARAMETER_RDG_BUFFER_SRV(StructuredBuffer<GatherPoints>, GatherPointsBuffer)
+		SHADER_PARAMETER(FIntPoint, GatherPointsResolution)
+
+		// Output
+		SHADER_PARAMETER_RDG_TEXTURE_UAV(RWTexture2D<float4>, RWGlobalIlluminationUAV)
+		SHADER_PARAMETER_RDG_TEXTURE_UAV(RWTexture2D<float>, RWRayDistanceUAV)
+		END_SHADER_PARAMETER_STRUCT()
+};
+
+IMPLEMENT_GLOBAL_SHADER(FRayTracingGlobalIlluminationFinalGatherRGS, "/Engine/Private/RayTracing/RayTracingFinalGatherRGS.usf", "RayTracingFinalGatherRGS", SF_RayGen);
 
 void FDeferredShadingSceneRenderer::PrepareRayTracingGlobalIllumination(const FViewInfo& View, TArray<FRHIRayTracingShader*>& OutRayGenShaders)
 {
@@ -270,183 +472,641 @@ void FDeferredShadingSceneRenderer::PrepareRayTracingGlobalIllumination(const FV
 			PermutationVector.Set<FGlobalIlluminationRGS::FEnableTwoSidedGeometryDim>(EnableTwoSidedGeometry == 1);
 			TShaderMapRef<FGlobalIlluminationRGS> RayGenerationShader(View.ShaderMap, PermutationVector);
 			OutRayGenShaders.Add(RayGenerationShader->GetRayTracingShader());
+
+			FRayTracingGlobalIlluminationCreateGatherPointsRGS::FPermutationDomain CreateGatherPointsPermutationVector;
+			CreateGatherPointsPermutationVector.Set<FRayTracingGlobalIlluminationCreateGatherPointsRGS::FUseAttenuationTermDim>(UseAttenuationTerm == 1);
+			CreateGatherPointsPermutationVector.Set<FRayTracingGlobalIlluminationCreateGatherPointsRGS::FEnableTwoSidedGeometryDim>(EnableTwoSidedGeometry == 1);
+			TShaderMapRef<FRayTracingGlobalIlluminationCreateGatherPointsRGS> CreateGatherPointsRayGenerationShader(View.ShaderMap, CreateGatherPointsPermutationVector);
+			OutRayGenShaders.Add(CreateGatherPointsRayGenerationShader->GetRayTracingShader());
+
+			FRayTracingGlobalIlluminationFinalGatherRGS::FPermutationDomain GatherPassPermutationVector;
+			GatherPassPermutationVector.Set<FRayTracingGlobalIlluminationFinalGatherRGS::FUseAttenuationTermDim>(UseAttenuationTerm == 1);
+			GatherPassPermutationVector.Set<FRayTracingGlobalIlluminationFinalGatherRGS::FEnableTwoSidedGeometryDim>(EnableTwoSidedGeometry == 1);
+			TShaderMapRef<FRayTracingGlobalIlluminationFinalGatherRGS> GatherPassRayGenerationShader(View.ShaderMap, GatherPassPermutationVector);
+			OutRayGenShaders.Add(GatherPassRayGenerationShader->GetRayTracingShader());
 		}
 	}
 
 }
 
-#endif // RHI_RAYTRACING
-
 void FDeferredShadingSceneRenderer::RenderRayTracingGlobalIllumination(
-	FRDGBuilder& GraphBuilder, 
-	FSceneTextureParameters& SceneTextures,
-	FViewInfo& View,
-	IScreenSpaceDenoiser::FAmbientOcclusionRayTracingConfig* OutRayTracingConfig,
-	IScreenSpaceDenoiser::FDiffuseIndirectInputs* OutDenoiserInputs)
-#if RHI_RAYTRACING
-{	
-	OutRayTracingConfig->ResolutionFraction = 1.0;
-	if (GRayTracingGlobalIlluminationDenoiser != 0)
+	FRHICommandListImmediate& RHICmdList, 
+	TRefCountPtr<IPooledRenderTarget>& GlobalIlluminationRT
+)
+{
+	SCOPED_GPU_STAT(RHICmdList, RayTracingGIBruteForce);
+
+	bool bAnyViewWithRTGI = false;
+	for (int32 ViewIndex = 0; ViewIndex < Views.Num(); ++ViewIndex)
 	{
-		OutRayTracingConfig->ResolutionFraction = FMath::Clamp(GRayTracingGlobalIlluminationScreenPercentage / 100.0, 0.25, 1.0);
+		FViewInfo& View = Views[ViewIndex];
+		bAnyViewWithRTGI = bAnyViewWithRTGI || ShouldRenderRayTracingGlobalIllumination(View);
 	}
 
-	RDG_GPU_STAT_SCOPE(GraphBuilder, RayTracingGlobalIllumination);
-	RDG_EVENT_SCOPE(GraphBuilder, "Ray Tracing Global Illumination");
+	if (!bAnyViewWithRTGI)
+	{
+		return;
+	}
 
-	int32 RayTracingGISamplesPerPixel = GRayTracingGlobalIlluminationSamplesPerPixel > -1 ? GRayTracingGlobalIlluminationSamplesPerPixel : View.FinalPostProcessSettings.RayTracingGISamplesPerPixel;
-	OutRayTracingConfig->RayCountPerPixel = RayTracingGISamplesPerPixel;
-	
-	int32 UpscaleFactor = int32(1.0 / OutRayTracingConfig->ResolutionFraction);
+	IScreenSpaceDenoiser::FAmbientOcclusionRayTracingConfig RayTracingConfig;
+	RayTracingConfig.ResolutionFraction = 1.0;
+	if (GRayTracingGlobalIlluminationDenoiser != 0)
+	{
+		RayTracingConfig.ResolutionFraction = FMath::Clamp(GRayTracingGlobalIlluminationScreenPercentage / 100.0, 0.25, 1.0);
+	}
 
-	// Allocate input for the denoiser.
+	int32 UpscaleFactor = int32(1.0 / RayTracingConfig.ResolutionFraction);
+
+	FSceneRenderTargets& SceneContext = FSceneRenderTargets::Get(RHICmdList);
 	{
-		FRDGTextureDesc Desc = FRDGTextureDesc::Create2DDesc(
-			SceneTextures.SceneDepthBuffer->Desc.Extent / UpscaleFactor,
-			PF_FloatRGBA,
-			FClearValueBinding::None,
-			/* InFlags = */ TexCreate_None,
-			/* InTargetableFlags = */ TexCreate_ShaderResource | TexCreate_RenderTargetable | TexCreate_UAV,
-			/* bInForceSeparateTargetAndShaderResource = */ false);
+		FPooledRenderTargetDesc Desc = SceneContext.GetSceneColor()->GetDesc();
+		Desc.Format = PF_FloatRGBA;
+		Desc.Flags &= ~(TexCreate_FastVRAM | TexCreate_Transient);
+		GRenderTargetPool.FindFreeElement(RHICmdList, Desc, GlobalIlluminationRT, TEXT("RayTracingGlobalIllumination"));
+	}
 
-		OutDenoiserInputs->Color = GraphBuilder.CreateTexture(Desc, TEXT("RayTracingDiffuseIndirect"));
+	FRDGBuilder GraphBuilder(RHICmdList);
+
+	FRDGTextureRef GlobalIlluminationTexture;
+	{
+		FRDGTextureDesc Desc = SceneContext.GetSceneColor()->GetDesc();
+		Desc.Extent /= UpscaleFactor;
+		Desc.Format = PF_FloatRGBA;
+		Desc.Flags &= ~(TexCreate_FastVRAM | TexCreate_Transient);
+		GlobalIlluminationTexture = GraphBuilder.CreateTexture(Desc, TEXT("RayTracingGlobalIllumination"));
+	}
+	FRDGTextureUAV* GlobalIlluminationUAV = GraphBuilder.CreateUAV(FRDGTextureUAVDesc(GlobalIlluminationTexture));
 
+	FRDGTextureRef RayDistanceTexture;
+	{
+		FRDGTextureDesc Desc = SceneContext.GetSceneColor()->GetDesc();
+		Desc.Extent /= UpscaleFactor;
 		Desc.Format = PF_G16R16;
-		OutDenoiserInputs->RayHitDistance = GraphBuilder.CreateTexture(Desc, TEXT("RayTracingDiffuseIndirectHitDistance"));
+		Desc.Flags &= ~(TexCreate_FastVRAM | TexCreate_Transient);
+		RayDistanceTexture = GraphBuilder.CreateTexture(Desc, TEXT("RayTracingGlobalIlluminationRayDistance"));
 	}
+	FRDGTextureUAV* RayDistanceUAV = GraphBuilder.CreateUAV(FRDGTextureUAVDesc(RayDistanceTexture));
+
+	TRefCountPtr<IPooledRenderTarget>& AmbientOcclusionRT = SceneContext.ScreenSpaceAO;
 
-	// Ray generation pass
+	for (int32 ViewIndex = 0; ViewIndex < Views.Num(); ++ViewIndex)
 	{
-		uint32 IterationCount = FMath::Max(RayTracingGISamplesPerPixel, 1);
-		uint32 SequenceCount = 1;
-		uint32 DimensionCount = 24;
-		FHaltonSequenceIteration HaltonSequenceIteration(Scene->HaltonSequence, IterationCount, SequenceCount, DimensionCount, View.ViewState ? (View.ViewState->FrameIndex % 1024) : 0);
+		FViewInfo& View = Views[ViewIndex];
+		if (ShouldRenderRayTracingGlobalIllumination(View))
+		{
+			RenderRayTracingGlobalIllumination(
+				RHICmdList, 
+				GraphBuilder, 
+				View, 
+				RayTracingConfig, 
+				UpscaleFactor, 
+				GlobalIlluminationRT, 
+				AmbientOcclusionRT, 
+				GlobalIlluminationUAV, 
+				RayDistanceUAV, 
+				GlobalIlluminationTexture, 
+				RayDistanceTexture
+			);
+		}
+	}
 
-		FHaltonIteration HaltonIteration;
-		InitializeHaltonSequenceIteration(HaltonSequenceIteration, HaltonIteration);
+	GraphBuilder.Execute();
+	SceneContext.bScreenSpaceAOIsValid = true;
+	GVisualizeTexture.SetCheckPoint(RHICmdList, GlobalIlluminationRT);
+}
 
-		FHaltonPrimes HaltonPrimes;
-		InitializeHaltonPrimes(Scene->HaltonPrimesResource, HaltonPrimes);
+void FDeferredShadingSceneRenderer::RenderRayTracingGlobalIllumination(
+	FRHICommandListImmediate& RHICmdList,
+	FRDGBuilder& GraphBuilder,
+	FViewInfo& View,
+	IScreenSpaceDenoiser::FAmbientOcclusionRayTracingConfig& RayTracingConfig,
+	int32 UpscaleFactor,
+	TRefCountPtr<IPooledRenderTarget>& GlobalIlluminationRT,
+	TRefCountPtr<IPooledRenderTarget>& AmbientOcclusionRT,
+	FRDGTextureUAV* GlobalIlluminationUAV,
+	FRDGTextureUAV* RayDistanceUAV,
+	const FRDGTextureRef GlobalIlluminationTexture,
+	const FRDGTextureRef RayDistanceTexture
+)
+{	
+	FSceneRenderTargets& SceneContext = FSceneRenderTargets::Get(RHICmdList);
 
-		FBlueNoise BlueNoise;
-		InitializeBlueNoise(BlueNoise);
+	RDG_EVENT_SCOPE(GraphBuilder, "RTGI");
 
-		FPathTracingLightData LightParameters;
-		SetupLightParameters(Scene->Lights, View, &LightParameters);
+	RayTracingConfig.RayCountPerPixel = GetRayTracingGlobalIlluminationSamplesPerPixel(View);
+
+	FSceneTextureParameters SceneTextures;
+	SetupSceneTextureParameters(GraphBuilder, &SceneTextures);
+
+	if (Scene->SkyLight && Scene->SkyLight->ShouldRebuildCdf())
+	{
+		BuildSkyLightCdfs(RHICmdList, Scene->SkyLight);
+	}
+
+	// Ray generation
+	bool bIsValid;
+	if (CVarRayTracingGlobalIlluminationEnableFinalGather.GetValueOnRenderThread() != 0)
+	{
+		bIsValid = RenderRayTracingGlobalIlluminationFinalGather(GraphBuilder, SceneTextures, View, RayTracingConfig, UpscaleFactor, GlobalIlluminationUAV, RayDistanceUAV);
+	}
+	else
+	{
+		bIsValid = RenderRayTracingGlobalIlluminationBruteForce(RHICmdList, GraphBuilder, SceneTextures, View, RayTracingConfig, UpscaleFactor, GlobalIlluminationUAV, RayDistanceUAV);
+	}
+
+	// Denoising
+
+	FRDGTextureRef ResultTexture; //#dxr_todo review
+
+	if (GRayTracingGlobalIlluminationDenoiser != 0 && bIsValid)
+	{
+		const IScreenSpaceDenoiser* DefaultDenoiser = IScreenSpaceDenoiser::GetDefaultDenoiser();
+		const IScreenSpaceDenoiser* DenoiserToUse = GRayTracingGlobalIlluminationDenoiser == 1 ? DefaultDenoiser : GScreenSpaceDenoiser;
+
+		IScreenSpaceDenoiser::FDiffuseIndirectInputs DenoiserInputs;
+		DenoiserInputs.Color = GlobalIlluminationTexture;
+		DenoiserInputs.RayHitDistance = RayDistanceTexture;
 
-		if (Scene->SkyLight && Scene->SkyLight->ShouldRebuildCdf())
-		{
-			// TODO: should be converted to RDG.
-			BuildSkyLightCdfs(GraphBuilder.RHICmdList, Scene->SkyLight);
-		}
-		FSkyLightData SkyLightParameters;
-		SetupSkyLightParameters(*Scene, &SkyLightParameters);
-
-		FGlobalIlluminationRGS::FParameters* PassParameters = GraphBuilder.AllocParameters<FGlobalIlluminationRGS::FParameters>();
-		PassParameters->SamplesPerPixel = RayTracingGISamplesPerPixel;
-		PassParameters->MaxBounces = GRayTracingGlobalIlluminationMaxBounces > -1? GRayTracingGlobalIlluminationMaxBounces : View.FinalPostProcessSettings.RayTracingGIMaxBounces;
-		PassParameters->MaxNormalBias = GetRaytracingMaxNormalBias();
-		float MaxRayDistanceForGI = GRayTracingGlobalIlluminationMaxRayDistance;
-		if (MaxRayDistanceForGI == -1.0)
 		{
-			MaxRayDistanceForGI = View.FinalPostProcessSettings.AmbientOcclusionRadius;
+			RDG_EVENT_SCOPE(GraphBuilder, "%s%s(DiffuseIndirect) %dx%d",
+				DenoiserToUse != DefaultDenoiser ? TEXT("ThirdParty ") : TEXT(""),
+				DenoiserToUse->GetDebugName(),
+				View.ViewRect.Width(), View.ViewRect.Height());
+
+			IScreenSpaceDenoiser::FDiffuseIndirectOutputs DenoiserOutputs = DenoiserToUse->DenoiseDiffuseIndirect(
+				GraphBuilder,
+				View,
+				&View.PrevViewInfo,
+				SceneTextures,
+				DenoiserInputs,
+				RayTracingConfig);
+
+			ResultTexture = DenoiserOutputs.Color;
 		}
-		PassParameters->MaxRayDistanceForGI = MaxRayDistanceForGI;
-		PassParameters->MaxRayDistanceForAO = View.FinalPostProcessSettings.AmbientOcclusionRadius;
-		PassParameters->UpscaleFactor = UpscaleFactor;
-		PassParameters->EvalSkyLight = GRayTracingGlobalIlluminationEvalSkyLight != 0;
-		PassParameters->UseRussianRoulette = GRayTracingGlobalIlluminationUseRussianRoulette != 0;
-		PassParameters->DiffuseThreshold = GRayTracingGlobalIlluminationDiffuseThreshold;
-		PassParameters->NextEventEstimationSamples = GRayTracingGlobalIlluminationNextEventEstimationSamples;
-		PassParameters->TLAS = View.RayTracingScene.RayTracingSceneRHI->GetShaderResourceView();
+	}
+	else
+	{
+		ResultTexture = GlobalIlluminationTexture;
+	}
+
+	// Compositing
+	if (bIsValid)
+	{
+		FRayTracingGlobalIlluminationCompositePS::FParameters *PassParameters = GraphBuilder.AllocParameters<FRayTracingGlobalIlluminationCompositePS::FParameters>();
+		PassParameters->GlobalIlluminationTexture = ResultTexture;
+		PassParameters->GlobalIlluminationSampler = TStaticSamplerState<SF_Point, AM_Clamp, AM_Clamp, AM_Clamp>::GetRHI();
 		PassParameters->ViewUniformBuffer = View.ViewUniformBuffer;
-		PassParameters->HaltonIteration = CreateUniformBufferImmediate(HaltonIteration, EUniformBufferUsage::UniformBuffer_SingleDraw);
-		PassParameters->HaltonPrimes = CreateUniformBufferImmediate(HaltonPrimes, EUniformBufferUsage::UniformBuffer_SingleDraw);
-		PassParameters->BlueNoise = CreateUniformBufferImmediate(BlueNoise, EUniformBufferUsage::UniformBuffer_SingleDraw);
-		PassParameters->LightParameters = CreateUniformBufferImmediate(LightParameters, EUniformBufferUsage::UniformBuffer_SingleDraw);
+		PassParameters->RenderTargets[0] = FRenderTargetBinding(GraphBuilder.RegisterExternalTexture(GlobalIlluminationRT), ERenderTargetLoadAction::ELoad, ERenderTargetStoreAction::EStore);
+		PassParameters->RenderTargets[1] = FRenderTargetBinding(GraphBuilder.RegisterExternalTexture(AmbientOcclusionRT), ERenderTargetLoadAction::ELoad, ERenderTargetStoreAction::EStore);
 		PassParameters->SceneTextures = SceneTextures;
-		PassParameters->SkyLight = CreateUniformBufferImmediate(SkyLightParameters, EUniformBufferUsage::UniformBuffer_SingleDraw);
 
-		// TODO: should be converted to RDG
-		TRefCountPtr<IPooledRenderTarget> SubsurfaceProfileRT((IPooledRenderTarget*) GetSubsufaceProfileTexture_RT(GraphBuilder.RHICmdList));
-		if (!SubsurfaceProfileRT)
-		{
-			SubsurfaceProfileRT = GSystemTextures.BlackDummy;
-		}
-		PassParameters->SSProfilesTexture = GraphBuilder.RegisterExternalTexture(SubsurfaceProfileRT);
-		PassParameters->TransmissionProfilesLinearSampler = TStaticSamplerState<SF_Bilinear, AM_Clamp, AM_Clamp, AM_Clamp>::GetRHI();
-		PassParameters->RWGlobalIlluminationUAV = GraphBuilder.CreateUAV(OutDenoiserInputs->Color);
-		PassParameters->RWRayDistanceUAV = GraphBuilder.CreateUAV(OutDenoiserInputs->RayHitDistance);
-		PassParameters->TileOffsetX = 0;
-		PassParameters->TileOffsetY = 0;
-
-		FGlobalIlluminationRGS::FPermutationDomain PermutationVector;
-		PermutationVector.Set<FGlobalIlluminationRGS::FUseAttenuationTermDim>(true);
-		PermutationVector.Set<FGlobalIlluminationRGS::FEnableTwoSidedGeometryDim>(CVarRayTracingGlobalIlluminationEnableTwoSidedGeometry.GetValueOnRenderThread() != 0);
-		TShaderMapRef<FGlobalIlluminationRGS> RayGenerationShader(GetGlobalShaderMap(FeatureLevel), PermutationVector);
-		ClearUnusedGraphResources(*RayGenerationShader, PassParameters);
-
-		FIntPoint RayTracingResolution = FIntPoint::DivideAndRoundUp(View.ViewRect.Size(), UpscaleFactor);
-
-		if (GRayTracingGlobalIlluminationTileSize <= 0)
-		{
-			GraphBuilder.AddPass(
-				RDG_EVENT_NAME("GlobalIlluminationRayTracing %dx%d", RayTracingResolution.X, RayTracingResolution.Y),
-				PassParameters,
-				ERDGPassFlags::Compute,
-				[PassParameters, this, &View, RayGenerationShader, RayTracingResolution](FRHICommandList& RHICmdList)
+		GraphBuilder.AddPass(
+			RDG_EVENT_NAME("GlobalIlluminationComposite"),
+			PassParameters,
+			ERDGPassFlags::Raster,
+			[this, &SceneContext, &View, &SceneTextures, PassParameters](FRHICommandListImmediate& RHICmdList)
 			{
-				FRayTracingShaderBindingsWriter GlobalResources;
-				SetShaderParameters(GlobalResources, *RayGenerationShader, *PassParameters);
+				TShaderMapRef<FPostProcessVS> VertexShader(View.ShaderMap);
+				TShaderMapRef<FRayTracingGlobalIlluminationCompositePS> PixelShader(View.ShaderMap);
+				FGraphicsPipelineStateInitializer GraphicsPSOInit;
+				RHICmdList.ApplyCachedRenderTargets(GraphicsPSOInit);
 
-				FRHIRayTracingScene* RayTracingSceneRHI = View.RayTracingScene.RayTracingSceneRHI;
-				RHICmdList.RayTraceDispatch(View.RayTracingMaterialPipeline, RayGenerationShader->GetRayTracingShader(), RayTracingSceneRHI, GlobalResources, RayTracingResolution.X, RayTracingResolution.Y);
-			});
-		}
-		else
+				// Additive blending
+				GraphicsPSOInit.BlendState = TStaticBlendState<CW_RGBA, BO_Add, BF_One, BF_Zero, BO_Add, BF_One, BF_Zero>::GetRHI();
+				//GraphicsPSOInit.BlendState = TStaticBlendState<CW_RGBA, BO_Add, BF_One, BF_One, BO_Add, BF_One, BF_One>::GetRHI();
+				GraphicsPSOInit.RasterizerState = TStaticRasterizerState<FM_Solid, CM_None>::GetRHI();
+				GraphicsPSOInit.DepthStencilState = TStaticDepthStencilState<false, CF_Always>::GetRHI();
+
+				GraphicsPSOInit.BoundShaderState.VertexDeclarationRHI = GFilterVertexDeclaration.VertexDeclarationRHI;
+				GraphicsPSOInit.BoundShaderState.VertexShaderRHI = GETSAFERHISHADER_VERTEX(*VertexShader);
+				GraphicsPSOInit.BoundShaderState.PixelShaderRHI = GETSAFERHISHADER_PIXEL(*PixelShader);
+				GraphicsPSOInit.PrimitiveType = PT_TriangleList;
+				SetGraphicsPipelineState(RHICmdList, GraphicsPSOInit);
+
+				SetShaderParameters(RHICmdList, *PixelShader, PixelShader->GetPixelShader(), *PassParameters);
+
+				RHICmdList.SetViewport(View.ViewRect.Min.X, View.ViewRect.Min.Y, 0.0f, View.ViewRect.Max.X, View.ViewRect.Max.Y, 1.0f);
+
+				DrawRectangle(
+					RHICmdList,
+					0, 0,
+					View.ViewRect.Width(), View.ViewRect.Height(),
+					View.ViewRect.Min.X, View.ViewRect.Min.Y,
+					View.ViewRect.Width(), View.ViewRect.Height(),
+					FIntPoint(View.ViewRect.Width(), View.ViewRect.Height()),
+					SceneContext.GetBufferSizeXY(),
+					*VertexShader
+				);
+			}
+		);
+	}
+}
+
+void FDeferredShadingSceneRenderer::RayTracingGlobalIlluminationCreateGatherPoints(
+	FRDGBuilder& GraphBuilder,
+	FSceneTextureParameters& SceneTextures,
+	FViewInfo& View,
+	int32 UpscaleFactor,
+	FRDGBufferRef& GatherPointsBuffer,
+	FIntPoint& GatherPointsResolution
+)
+#if RHI_RAYTRACING
+{
+	RDG_GPU_STAT_SCOPE(GraphBuilder, RayTracingGICreateGatherPoints);
+	RDG_EVENT_SCOPE(GraphBuilder, "Ray Tracing GI: Create Gather Points");
+
+	int32 GatherSamples = GetRayTracingGlobalIlluminationSamplesPerPixel(View);
+	int32 SamplesPerPixel = 1;
+
+	uint32 IterationCount = SamplesPerPixel;
+	uint32 SequenceCount = 1;
+	uint32 DimensionCount = 24;
+	int32 FrameIndex = View.ViewState->FrameIndex % 1024;
+	FHaltonSequenceIteration HaltonSequenceIteration(Scene->HaltonSequence, IterationCount, SequenceCount, DimensionCount, FrameIndex);
+
+	FHaltonIteration HaltonIteration;
+	InitializeHaltonSequenceIteration(HaltonSequenceIteration, HaltonIteration);
+
+	FHaltonPrimes HaltonPrimes;
+	InitializeHaltonPrimes(Scene->HaltonPrimesResource, HaltonPrimes);
+
+	FBlueNoise BlueNoise;
+	InitializeBlueNoise(BlueNoise);
+
+	FPathTracingLightData LightParameters;
+	SetupLightParameters(Scene->Lights, View, &LightParameters);
+
+	FSkyLightData SkyLightParameters;
+	SetupSkyLightParameters(*Scene, &SkyLightParameters);
+
+	FRayTracingGlobalIlluminationCreateGatherPointsRGS::FParameters* PassParameters = GraphBuilder.AllocParameters<FRayTracingGlobalIlluminationCreateGatherPointsRGS::FParameters>();
+	PassParameters->SampleIndex = (FrameIndex * SamplesPerPixel) % GatherSamples;
+	PassParameters->SamplesPerPixel = SamplesPerPixel;
+	PassParameters->MaxBounces = 1;
+	PassParameters->MaxNormalBias = GetRaytracingMaxNormalBias();
+	PassParameters->MaxRayDistanceForGI = GRayTracingGlobalIlluminationMaxRayDistance;
+	PassParameters->EvalSkyLight = GRayTracingGlobalIlluminationEvalSkyLight != 0;
+	PassParameters->UseRussianRoulette = GRayTracingGlobalIlluminationUseRussianRoulette != 0;
+	PassParameters->DiffuseThreshold = GRayTracingGlobalIlluminationDiffuseThreshold;
+	PassParameters->NextEventEstimationSamples = GRayTracingGlobalIlluminationNextEventEstimationSamples;
+	PassParameters->UpscaleFactor = UpscaleFactor;
+	PassParameters->TileOffsetX = 0;
+	PassParameters->TileOffsetY = 0;
+
+	// Global
+	PassParameters->TLAS = View.RayTracingScene.RayTracingSceneRHI->GetShaderResourceView();
+	PassParameters->ViewUniformBuffer = View.ViewUniformBuffer;
+
+	// Sampling sequence
+	PassParameters->HaltonIteration = CreateUniformBufferImmediate(HaltonIteration, EUniformBufferUsage::UniformBuffer_SingleDraw);
+	PassParameters->HaltonPrimes = CreateUniformBufferImmediate(HaltonPrimes, EUniformBufferUsage::UniformBuffer_SingleDraw);
+	PassParameters->BlueNoise = CreateUniformBufferImmediate(BlueNoise, EUniformBufferUsage::UniformBuffer_SingleDraw);
+
+	// Light data
+	PassParameters->LightParameters = CreateUniformBufferImmediate(LightParameters, EUniformBufferUsage::UniformBuffer_SingleDraw);
+	PassParameters->SceneTextures = SceneTextures;
+	PassParameters->SkyLight = CreateUniformBufferImmediate(SkyLightParameters, EUniformBufferUsage::UniformBuffer_SingleDraw);
+
+	// Shading data
+	TRefCountPtr<IPooledRenderTarget> SubsurfaceProfileRT((IPooledRenderTarget*)GetSubsufaceProfileTexture_RT(GraphBuilder.RHICmdList));
+	if (!SubsurfaceProfileRT)
+	{
+		SubsurfaceProfileRT = GSystemTextures.BlackDummy;
+	}
+	PassParameters->SSProfilesTexture = GraphBuilder.RegisterExternalTexture(SubsurfaceProfileRT);
+	PassParameters->TransmissionProfilesLinearSampler = TStaticSamplerState<SF_Bilinear, AM_Clamp, AM_Clamp, AM_Clamp>::GetRHI();
+
+	// Output
+	FIntPoint LocalGatherPointsResolution = FIntPoint::DivideAndRoundUp(View.ViewRect.Size(), UpscaleFactor);
+	if (GatherPointsResolution != LocalGatherPointsResolution)
+	{
+		GatherPointsResolution = LocalGatherPointsResolution;
+		FRDGBufferDesc BufferDesc = FRDGBufferDesc::CreateStructuredDesc(sizeof(GatherPoints), GatherPointsResolution.X * GatherPointsResolution.Y);
+		GatherPointsBuffer = GraphBuilder.CreateBuffer(BufferDesc, TEXT("GatherPointsBuffer"), ERDGResourceFlags::MultiFrame);
+	}
+	else
+	{
+		GatherPointsBuffer = GraphBuilder.RegisterExternalBuffer(((FSceneViewState*)View.State)->GatherPointsBuffer, TEXT("GatherPointsBuffer"));
+	}
+	PassParameters->GatherPointsResolution = GatherPointsResolution;
+	PassParameters->RWGatherPointsBuffer = GraphBuilder.CreateUAV(GatherPointsBuffer, EPixelFormat::PF_R32_UINT);
+
+	FRayTracingGlobalIlluminationCreateGatherPointsRGS::FPermutationDomain PermutationVector;
+	PermutationVector.Set<FRayTracingGlobalIlluminationCreateGatherPointsRGS::FUseAttenuationTermDim>(true);
+	PermutationVector.Set<FRayTracingGlobalIlluminationCreateGatherPointsRGS::FEnableTwoSidedGeometryDim>(CVarRayTracingGlobalIlluminationEnableTwoSidedGeometry.GetValueOnRenderThread() != 0);
+	TShaderMapRef<FRayTracingGlobalIlluminationCreateGatherPointsRGS> RayGenerationShader(GetGlobalShaderMap(FeatureLevel), PermutationVector);
+	ClearUnusedGraphResources(*RayGenerationShader, PassParameters);
+
+	GraphBuilder.AddPass(
+		RDG_EVENT_NAME("GatherPoints %d%d", GatherPointsResolution.X, GatherPointsResolution.Y),
+		PassParameters,
+		ERDGPassFlags::Compute,
+		[PassParameters, this, &View, RayGenerationShader, GatherPointsResolution](FRHICommandList& RHICmdList)
+	{
+		FRHIRayTracingScene* RayTracingSceneRHI = View.RayTracingScene.RayTracingSceneRHI;
+		FRayTracingShaderBindingsWriter GlobalResources;
+		SetShaderParameters(GlobalResources, *RayGenerationShader, *PassParameters);
+		RHICmdList.RayTraceDispatch(View.RayTracingMaterialPipeline, RayGenerationShader->GetRayTracingShader(), RayTracingSceneRHI, GlobalResources, GatherPointsResolution.X, GatherPointsResolution.Y);
+	});
+}
+#else
+{
+	unimplemented();
+}
+#endif
+
+bool FDeferredShadingSceneRenderer::RenderRayTracingGlobalIlluminationFinalGather(
+	FRDGBuilder& GraphBuilder,
+	FSceneTextureParameters& SceneTextures,
+	FViewInfo& View,
+	const IScreenSpaceDenoiser::FAmbientOcclusionRayTracingConfig& RayTracingConfig,
+	int32 UpscaleFactor,
+	// Output
+	FRDGTextureUAV* GlobalIlluminationUAV,
+	FRDGTextureUAV* RayDistanceUAV)
+#if RHI_RAYTRACING
+{
+	FSceneViewState* SceneViewState = (FSceneViewState*)View.State;
+	if (!SceneViewState) return false;
+
+	int32 SamplesPerPixel = GetRayTracingGlobalIlluminationSamplesPerPixel(View);
+	if (SamplesPerPixel <= 0) return false;
+
+	// Generate gather points
+	FRDGBufferRef GatherPointsBuffer;
+	RayTracingGlobalIlluminationCreateGatherPoints(GraphBuilder, SceneTextures, View, UpscaleFactor, GatherPointsBuffer, SceneViewState->GatherPointsResolution);
+
+	// Perform gather
+	RDG_GPU_STAT_SCOPE(GraphBuilder, RayTracingGIFinalGather);
+	RDG_EVENT_SCOPE(GraphBuilder, "Ray Tracing GI: Final Gather");
+
+	FRayTracingGlobalIlluminationFinalGatherRGS::FParameters* PassParameters = GraphBuilder.AllocParameters<FRayTracingGlobalIlluminationFinalGatherRGS::FParameters>();
+	int32 SampleIndex = View.ViewState->FrameIndex % SamplesPerPixel;
+	PassParameters->SampleIndex = SampleIndex;
+	PassParameters->SamplesPerPixel = SamplesPerPixel;
+	PassParameters->DiffuseThreshold = GRayTracingGlobalIlluminationDiffuseThreshold;
+	PassParameters->MaxNormalBias = GetRaytracingMaxNormalBias();
+	PassParameters->FinalGatherDistance = GRayTracingGlobalIlluminationFinalGatherDistance;
+	PassParameters->UpscaleFactor = UpscaleFactor;
+	PassParameters->TileOffsetX = 0;
+	PassParameters->TileOffsetY = 0;
+
+	// Scene data
+	PassParameters->TLAS = View.RayTracingScene.RayTracingSceneRHI->GetShaderResourceView();
+	PassParameters->ViewUniformBuffer = View.ViewUniformBuffer;
+
+	// Shading data
+	PassParameters->SceneTextures = SceneTextures;
+	TRefCountPtr<IPooledRenderTarget> SubsurfaceProfileRT((IPooledRenderTarget*)GetSubsufaceProfileTexture_RT(GraphBuilder.RHICmdList));
+	if (!SubsurfaceProfileRT)
+	{
+		SubsurfaceProfileRT = GSystemTextures.BlackDummy;
+	}
+	PassParameters->SSProfilesTexture = GraphBuilder.RegisterExternalTexture(SubsurfaceProfileRT);
+	PassParameters->TransmissionProfilesLinearSampler = TStaticSamplerState<SF_Bilinear, AM_Clamp, AM_Clamp, AM_Clamp>::GetRHI();
+
+	// Gather points
+	PassParameters->GatherPointsResolution = SceneViewState->GatherPointsResolution;
+	PassParameters->GatherPointsBuffer = GraphBuilder.CreateSRV(GatherPointsBuffer);
+
+	// Output
+	PassParameters->RWGlobalIlluminationUAV = GlobalIlluminationUAV;
+	PassParameters->RWRayDistanceUAV = RayDistanceUAV;
+
+	FRayTracingGlobalIlluminationFinalGatherRGS::FPermutationDomain PermutationVector;
+	PermutationVector.Set<FRayTracingGlobalIlluminationFinalGatherRGS::FUseAttenuationTermDim>(true);
+	PermutationVector.Set<FRayTracingGlobalIlluminationFinalGatherRGS::FEnableTwoSidedGeometryDim>(CVarRayTracingGlobalIlluminationEnableTwoSidedGeometry.GetValueOnRenderThread() != 0);
+	TShaderMapRef<FRayTracingGlobalIlluminationFinalGatherRGS> RayGenerationShader(GetGlobalShaderMap(FeatureLevel), PermutationVector);
+	ClearUnusedGraphResources(*RayGenerationShader, PassParameters);
+
+	FIntPoint RayTracingResolution = FIntPoint::DivideAndRoundUp(View.ViewRect.Size(), UpscaleFactor);
+	GraphBuilder.AddPass(
+		RDG_EVENT_NAME("GlobalIlluminationRayTracing %dx%d", RayTracingResolution.X, RayTracingResolution.Y),
+		PassParameters,
+		ERDGPassFlags::Compute,
+		[PassParameters, this, &View, RayGenerationShader, RayTracingResolution](FRHICommandList& RHICmdList)
+	{
+		FRHIRayTracingScene* RayTracingSceneRHI = View.RayTracingScene.RayTracingSceneRHI;
+
+		FRayTracingShaderBindingsWriter GlobalResources;
+		SetShaderParameters(GlobalResources, *RayGenerationShader, *PassParameters);
+		RHICmdList.RayTraceDispatch(View.RayTracingMaterialPipeline, RayGenerationShader->GetRayTracingShader(), RayTracingSceneRHI, GlobalResources, RayTracingResolution.X, RayTracingResolution.Y);
+	});
+
+
+	GraphBuilder.QueueBufferExtraction(GatherPointsBuffer, &SceneViewState->GatherPointsBuffer);
+	return true;
+}
+#else
+{
+	unimplemented();
+	return true;
+}
+#endif
+
+bool FDeferredShadingSceneRenderer::RenderRayTracingGlobalIlluminationBruteForce(
+	FRHICommandListImmediate& RHICmdList,
+	FRDGBuilder& GraphBuilder,
+	FSceneTextureParameters& SceneTextures,
+	FViewInfo& View,
+	const IScreenSpaceDenoiser::FAmbientOcclusionRayTracingConfig& RayTracingConfig,
+	int32 UpscaleFactor,
+	FRDGTextureUAV* GlobalIlluminationUAV,
+	FRDGTextureUAV* RayDistanceUAV)
+#if RHI_RAYTRACING
+{
+	if (!View.ViewState) return false;
+
+	RDG_GPU_STAT_SCOPE(GraphBuilder, RayTracingGIBruteForce);
+	RDG_EVENT_SCOPE(GraphBuilder, "Ray Tracing GI: Brute Force");
+
+	int32 RayTracingGISamplesPerPixel = GetRayTracingGlobalIlluminationSamplesPerPixel(View);
+	uint32 IterationCount = FMath::Max(RayTracingGISamplesPerPixel, 1);
+	uint32 SequenceCount = 1;
+	uint32 DimensionCount = 24;
+	FHaltonSequenceIteration HaltonSequenceIteration(Scene->HaltonSequence, IterationCount, SequenceCount, DimensionCount, View.ViewState ? (View.ViewState->FrameIndex % 1024) : 0);
+
+	FHaltonIteration HaltonIteration;
+	InitializeHaltonSequenceIteration(HaltonSequenceIteration, HaltonIteration);
+
+	FHaltonPrimes HaltonPrimes;
+	InitializeHaltonPrimes(Scene->HaltonPrimesResource, HaltonPrimes);
+
+	FBlueNoise BlueNoise;
+	InitializeBlueNoise(BlueNoise);
+
+	FPathTracingLightData LightParameters;
+	SetupLightParameters(Scene->Lights, View, &LightParameters);
+
+	FSkyLightData SkyLightParameters;
+	SetupSkyLightParameters(*Scene, &SkyLightParameters);
+
+	FGlobalIlluminationRGS::FParameters* PassParameters = GraphBuilder.AllocParameters<FGlobalIlluminationRGS::FParameters>();
+	PassParameters->SamplesPerPixel = RayTracingGISamplesPerPixel;
+	PassParameters->MaxBounces = GRayTracingGlobalIlluminationMaxBounces > -1 ? GRayTracingGlobalIlluminationMaxBounces : View.FinalPostProcessSettings.RayTracingGIMaxBounces;
+	PassParameters->MaxNormalBias = GetRaytracingMaxNormalBias();
+	float MaxRayDistanceForGI = GRayTracingGlobalIlluminationMaxRayDistance;
+	if (MaxRayDistanceForGI == -1.0)
+	{
+		MaxRayDistanceForGI = View.FinalPostProcessSettings.AmbientOcclusionRadius;
+	}
+	PassParameters->MaxRayDistanceForGI = MaxRayDistanceForGI;
+	PassParameters->MaxRayDistanceForAO = View.FinalPostProcessSettings.AmbientOcclusionRadius;
+	PassParameters->UpscaleFactor = UpscaleFactor;
+	PassParameters->EvalSkyLight = GRayTracingGlobalIlluminationEvalSkyLight != 0;
+	PassParameters->UseRussianRoulette = GRayTracingGlobalIlluminationUseRussianRoulette != 0;
+	PassParameters->DiffuseThreshold = GRayTracingGlobalIlluminationDiffuseThreshold;
+	PassParameters->NextEventEstimationSamples = GRayTracingGlobalIlluminationNextEventEstimationSamples;
+	PassParameters->TLAS = View.RayTracingScene.RayTracingSceneRHI->GetShaderResourceView();
+	PassParameters->ViewUniformBuffer = View.ViewUniformBuffer;
+	PassParameters->HaltonIteration = CreateUniformBufferImmediate(HaltonIteration, EUniformBufferUsage::UniformBuffer_SingleDraw);
+	PassParameters->HaltonPrimes = CreateUniformBufferImmediate(HaltonPrimes, EUniformBufferUsage::UniformBuffer_SingleDraw);
+	PassParameters->BlueNoise = CreateUniformBufferImmediate(BlueNoise, EUniformBufferUsage::UniformBuffer_SingleDraw);
+	PassParameters->LightParameters = CreateUniformBufferImmediate(LightParameters, EUniformBufferUsage::UniformBuffer_SingleDraw);
+	PassParameters->SceneTextures = SceneTextures;
+	PassParameters->SkyLight = CreateUniformBufferImmediate(SkyLightParameters, EUniformBufferUsage::UniformBuffer_SingleDraw);
+	TRefCountPtr<IPooledRenderTarget> SubsurfaceProfileRT((IPooledRenderTarget*)GetSubsufaceProfileTexture_RT(RHICmdList));
+	if (!SubsurfaceProfileRT)
+	{
+		SubsurfaceProfileRT = GSystemTextures.BlackDummy;
+	}
+	PassParameters->SSProfilesTexture = GraphBuilder.RegisterExternalTexture(SubsurfaceProfileRT);
+	PassParameters->TransmissionProfilesLinearSampler = TStaticSamplerState<SF_Bilinear, AM_Clamp, AM_Clamp, AM_Clamp>::GetRHI();
+	PassParameters->RWGlobalIlluminationUAV = GlobalIlluminationUAV;
+	PassParameters->RWRayDistanceUAV = RayDistanceUAV;
+	PassParameters->TileOffsetX = 0;
+	PassParameters->TileOffsetY = 0;
+
+	FGlobalIlluminationRGS::FPermutationDomain PermutationVector;
+	PermutationVector.Set<FGlobalIlluminationRGS::FUseAttenuationTermDim>(true);
+	PermutationVector.Set<FGlobalIlluminationRGS::FEnableTwoSidedGeometryDim>(CVarRayTracingGlobalIlluminationEnableTwoSidedGeometry.GetValueOnRenderThread() != 0);
+	TShaderMapRef<FGlobalIlluminationRGS> RayGenerationShader(GetGlobalShaderMap(FeatureLevel), PermutationVector);
+	ClearUnusedGraphResources(*RayGenerationShader, PassParameters);
+
+	FIntPoint RayTracingResolution = FIntPoint::DivideAndRoundUp(View.ViewRect.Size(), UpscaleFactor);
+
+	if (GRayTracingGlobalIlluminationTileSize <= 0)
+	{
+		GraphBuilder.AddPass(
+			RDG_EVENT_NAME("GlobalIlluminationRayTracing %dx%d", RayTracingResolution.X, RayTracingResolution.Y),
+			PassParameters,
+			ERDGPassFlags::Compute,
+			[PassParameters, this, &View, RayGenerationShader, RayTracingResolution](FRHICommandList& RHICmdList)
+		{
+			FRayTracingShaderBindingsWriter GlobalResources;
+			SetShaderParameters(GlobalResources, *RayGenerationShader, *PassParameters);
+
+			FRHIRayTracingScene* RayTracingSceneRHI = View.RayTracingScene.RayTracingSceneRHI;
+			RHICmdList.RayTraceDispatch(View.RayTracingMaterialPipeline, RayGenerationShader->GetRayTracingShader(), RayTracingSceneRHI, GlobalResources, RayTracingResolution.X, RayTracingResolution.Y);
+		});
+	}
+	else
+	{
+		int32 TileSize = FMath::Max(32, GRayTracingGlobalIlluminationTileSize);
+		int32 NumTilesX = FMath::DivideAndRoundUp(RayTracingResolution.X, TileSize);
+		int32 NumTilesY = FMath::DivideAndRoundUp(RayTracingResolution.Y, TileSize);
+		for (int32 Y = 0; Y < NumTilesY; ++Y)
 		{
-			int32 TileSize = FMath::Max(32, GRayTracingGlobalIlluminationTileSize);
-			int32 NumTilesX = FMath::DivideAndRoundUp(RayTracingResolution.X, TileSize);
-			int32 NumTilesY = FMath::DivideAndRoundUp(RayTracingResolution.Y, TileSize);
-			for (int32 Y = 0; Y < NumTilesY; ++Y)
+			for (int32 X = 0; X < NumTilesX; ++X)
 			{
-				for (int32 X = 0; X < NumTilesX; ++X)
+				FGlobalIlluminationRGS::FParameters* TilePassParameters = PassParameters;
+
+				if (X > 0 || Y > 0)
 				{
-					FGlobalIlluminationRGS::FParameters* TilePassParameters = PassParameters;
-
-					if (X > 0 || Y > 0)
-					{
-						TilePassParameters = GraphBuilder.AllocParameters<FGlobalIlluminationRGS::FParameters>();
-						*TilePassParameters = *PassParameters;
-
-						TilePassParameters->TileOffsetX = X * TileSize;
-						TilePassParameters->TileOffsetY = Y * TileSize;
-					}
-
-					int32 DispatchSizeX = FMath::Min<int32>(TileSize, RayTracingResolution.X - TilePassParameters->TileOffsetX);
-					int32 DispatchSizeY = FMath::Min<int32>(TileSize, RayTracingResolution.Y - TilePassParameters->TileOffsetY);
-
-					GraphBuilder.AddPass(
-						RDG_EVENT_NAME("GlobalIlluminationRayTracing %dx%d (tile %dx%d)", DispatchSizeX, DispatchSizeY, X, Y),
-						TilePassParameters,
-						ERDGPassFlags::Compute,
-						[TilePassParameters, this, &View, RayGenerationShader, DispatchSizeX, DispatchSizeY](FRHICommandList& RHICmdList)
-					{
-						FRHIRayTracingScene* RayTracingSceneRHI = View.RayTracingScene.RayTracingSceneRHI;
-
-						FRayTracingShaderBindingsWriter GlobalResources;
-						SetShaderParameters(GlobalResources, *RayGenerationShader, *TilePassParameters);
-						RHICmdList.RayTraceDispatch(View.RayTracingMaterialPipeline, RayGenerationShader->GetRayTracingShader(), RayTracingSceneRHI, 
-							GlobalResources, DispatchSizeX, DispatchSizeY);
-						RHICmdList.SubmitCommandsHint();
-					});
+					TilePassParameters = GraphBuilder.AllocParameters<FGlobalIlluminationRGS::FParameters>();
+					*TilePassParameters = *PassParameters;
+
+					TilePassParameters->TileOffsetX = X * TileSize;
+					TilePassParameters->TileOffsetY = Y * TileSize;
 				}
+
+				int32 DispatchSizeX = FMath::Min<int32>(TileSize, RayTracingResolution.X - TilePassParameters->TileOffsetX);
+				int32 DispatchSizeY = FMath::Min<int32>(TileSize, RayTracingResolution.Y - TilePassParameters->TileOffsetY);
+
+				GraphBuilder.AddPass(
+					RDG_EVENT_NAME("GlobalIlluminationRayTracing %dx%d (tile %dx%d)", DispatchSizeX, DispatchSizeY, X, Y),
+					TilePassParameters,
+					ERDGPassFlags::Compute,
+					[TilePassParameters, this, &View, RayGenerationShader, DispatchSizeX, DispatchSizeY](FRHICommandList& RHICmdList)
+				{
+					FRHIRayTracingScene* RayTracingSceneRHI = View.RayTracingScene.RayTracingSceneRHI;
+
+					FRayTracingShaderBindingsWriter GlobalResources;
+					SetShaderParameters(GlobalResources, *RayGenerationShader, *TilePassParameters);
+					RHICmdList.RayTraceDispatch(View.RayTracingMaterialPipeline, RayGenerationShader->GetRayTracingShader(), RayTracingSceneRHI,
+						GlobalResources, DispatchSizeX, DispatchSizeY);
+					RHICmdList.SubmitCommandsHint();
+				});
 			}
 		}
 	}
+	return true;
 }
 #else
 {
 	unimplemented();
+	return false;
+}
+#endif // RHI_RAYTRACING
+
+void FDeferredShadingSceneRenderer::CompositeGlobalIllumination(
+	FRHICommandListImmediate& RHICmdList,
+	const FViewInfo& View,
+	TRefCountPtr<IPooledRenderTarget>& GlobalIlluminationRT
+)
+{
+	FSceneRenderTargets& SceneContext = FSceneRenderTargets::Get(RHICmdList);
+
+	FRDGBuilder GraphBuilder(RHICmdList);
+	FSceneTextureParameters SceneTextures;
+	SetupSceneTextureParameters(GraphBuilder, &SceneTextures);
+
+	FRayTracingGlobalIlluminationSceneColorCompositePS::FParameters *PassParameters = GraphBuilder.AllocParameters<FRayTracingGlobalIlluminationSceneColorCompositePS::FParameters>();
+	PassParameters->GlobalIlluminationTexture = GraphBuilder.RegisterExternalTexture(GlobalIlluminationRT);
+	PassParameters->GlobalIlluminationSampler = TStaticSamplerState<SF_Point, AM_Clamp, AM_Clamp, AM_Clamp>::GetRHI();
+	PassParameters->RenderTargets[0] = FRenderTargetBinding(GraphBuilder.RegisterExternalTexture(SceneContext.GetSceneColor()), ERenderTargetLoadAction::ELoad, ERenderTargetStoreAction::EStore);
+	PassParameters->SceneTextures = SceneTextures;
+
+	GraphBuilder.AddPass(
+		RDG_EVENT_NAME("GlobalIlluminationComposite"),
+		PassParameters,
+		ERDGPassFlags::Raster,
+		[this, &SceneContext, &View, PassParameters](FRHICommandListImmediate& RHICmdList)
+		{
+			TShaderMapRef<FPostProcessVS> VertexShader(View.ShaderMap);
+			TShaderMapRef<FRayTracingGlobalIlluminationSceneColorCompositePS> PixelShader(View.ShaderMap);
+			FGraphicsPipelineStateInitializer GraphicsPSOInit;
+			RHICmdList.ApplyCachedRenderTargets(GraphicsPSOInit);
+
+			// Additive blending
+			GraphicsPSOInit.BlendState = TStaticBlendState<CW_RGBA, BO_Add, BF_One, BF_One, BO_Add, BF_One, BF_One>::GetRHI();
+			//GraphicsPSOInit.BlendState = TStaticBlendState<CW_RGBA, BO_Add, BF_One, BF_One, BO_Add, BF_One, BF_One>::GetRHI();
+			GraphicsPSOInit.RasterizerState = TStaticRasterizerState<FM_Solid, CM_None>::GetRHI();
+			GraphicsPSOInit.DepthStencilState = TStaticDepthStencilState<false, CF_Always>::GetRHI();
+
+			GraphicsPSOInit.BoundShaderState.VertexDeclarationRHI = GFilterVertexDeclaration.VertexDeclarationRHI;
+			GraphicsPSOInit.BoundShaderState.VertexShaderRHI = GETSAFERHISHADER_VERTEX(*VertexShader);
+			GraphicsPSOInit.BoundShaderState.PixelShaderRHI = GETSAFERHISHADER_PIXEL(*PixelShader);
+			GraphicsPSOInit.PrimitiveType = PT_TriangleList;
+			SetGraphicsPipelineState(RHICmdList, GraphicsPSOInit);
+
+			SetShaderParameters(RHICmdList, *PixelShader, PixelShader->GetPixelShader(), *PassParameters);
+
+			RHICmdList.SetViewport(View.ViewRect.Min.X, View.ViewRect.Min.Y, 0.0f, View.ViewRect.Max.X, View.ViewRect.Max.Y, 1.0f);
+
+			DrawRectangle(
+				RHICmdList,
+				0, 0,
+				View.ViewRect.Width(), View.ViewRect.Height(),
+				View.ViewRect.Min.X, View.ViewRect.Min.Y,
+				View.ViewRect.Width(), View.ViewRect.Height(),
+				FIntPoint(View.ViewRect.Width(), View.ViewRect.Height()),
+				SceneContext.GetBufferSizeXY(),
+				*VertexShader
+			);
+		}
+	);
+	GraphBuilder.Execute();
 }
+
 #endif // RHI_RAYTRACING
