Add "Final Gather" style algorithm for ray tracing global illumination.

Invoked with the following CVar: r.RayTracing.GlobalIllumination.EnableFinalGather 1

diff --git a/Engine/Shaders/Private/RayTracing/RayTracingCreateGatherPointsRGS.usf b/Engine/Shaders/Private/RayTracing/RayTracingCreateGatherPointsRGS.usf
--- a/Engine/Shaders/Private/RayTracing/RayTracingCreateGatherPointsRGS.usf
+++ b/Engine/Shaders/Private/RayTracing/RayTracingCreateGatherPointsRGS.usf
@@ -27,8 +27,8 @@ float MaxRayDistanceForGI;
 float NextEventEstimationSamples;
 float DiffuseThreshold;
 float MaxNormalBias;
-bool EvalSkyLight;
-bool UseRussianRoulette;
+uint EvalSkyLight;
+uint UseRussianRoulette;
 
 int2 GatherPointsResolution;
 RWStructuredBuffer<FGatherPoints> RWGatherPointsBuffer;
@@ -115,6 +115,7 @@ void RayTracingCreateGatherPointsRGS()
 
 	FQuasiRandom QuasiRandom;
 	QuasiRandom_Initialize(QuasiRandom, PixelCoord, 0, SamplesPerPixel);
+	//QuasiRandom_Initialize(QuasiRandom, DispatchThreadId, 0, SamplesPerPixel);
 
 	for (uint Index = 0; Index < SamplesPerPixel; ++Index)
 	{
@@ -161,6 +162,7 @@ void RayTracingCreateGatherPointsRGS()
 				RayCone
 			);
 
+			// TODO: impossible based on formulation
 			if (Payload.HitT > MaxRayDistanceForGI) continue;
 
 			// Update intersection
diff --git a/Engine/Source/Runtime/Renderer/Private/DeferredShadingRenderer.h b/Engine/Source/Runtime/Renderer/Private/DeferredShadingRenderer.h
--- a/Engine/Source/Runtime/Renderer/Private/DeferredShadingRenderer.h
+++ b/Engine/Source/Runtime/Renderer/Private/DeferredShadingRenderer.h
@@ -499,33 +499,20 @@ private:
 	void RenderRayTracingStochasticRectLight(FRHICommandListImmediate& RHICmdList, const FLightSceneInfo& RectLightSceneInfo, TRefCountPtr<IPooledRenderTarget>& RectLightRT, TRefCountPtr<IPooledRenderTarget>& HitDistanceRT);
 	void CompositeRayTracingSkyLight(FRHICommandListImmediate& RHICmdList, TRefCountPtr<IPooledRenderTarget>& SkyLightRT, TRefCountPtr<IPooledRenderTarget>& HitDistanceRT);
 	
-	void RenderRayTracingGlobalIllumination(
+	bool RenderRayTracingGlobalIllumination(
 		FRDGBuilder& GraphBuilder, 
 		FSceneTextureParameters& SceneTextures,
 		FViewInfo& View,
 		IScreenSpaceDenoiser::FAmbientOcclusionRayTracingConfig* RayTracingConfig,
 		IScreenSpaceDenoiser::FDiffuseIndirectInputs* OutDenoiserInputs);
 	
-	void RenderRayTracingAmbientOcclusion(
-		FRDGBuilder& GraphBuilder,
-		FViewInfo& View,
-		const FSceneTextureParameters& SceneTextures,
-		FRDGTextureRef* OutAmbientOcclusionTexture);
-	
-
-#if RHI_RAYTRACING
-	template <int TextureImportanceSampling> void RenderRayTracingRectLightInternal(FRHICommandListImmediate& RHICmdList, const TArray<FViewInfo>& Views, const FLightSceneInfo& RectLightSceneInfo, TRefCountPtr<IPooledRenderTarget>& ScreenShadowMaskTexture, TRefCountPtr<IPooledRenderTarget>& RayDistanceTexture);
-	void VisualizeRectLightMipTree(FRHICommandListImmediate& RHICmdList, const FViewInfo& View, const FRWBuffer& RectLightMipTree, const FIntVector& RectLightMipTreeDimensions);
-	
-	bool RenderRayTracingGlobalIlluminationBruteForce(
-		FRHICommandListImmediate& RHICmdList,
+	void RenderRayTracingGlobalIlluminationBruteForce(
 		FRDGBuilder& GraphBuilder,
 		FSceneTextureParameters& SceneTextures,
 		FViewInfo& View,
 		const IScreenSpaceDenoiser::FAmbientOcclusionRayTracingConfig& RayTracingConfig,
 		int32 UpscaleFactor,
-		FRDGTextureUAV* GlobalIlluminationUAV,
-		FRDGTextureUAV* RayDistanceUAV);
+		IScreenSpaceDenoiser::FDiffuseIndirectInputs* OutDenoiserInputs);
 
 	void RayTracingGlobalIlluminationCreateGatherPoints(
 		FRDGBuilder& GraphBuilder,
@@ -535,15 +522,25 @@ private:
 		FRDGBufferRef& GatherPointsBuffer,
 		FIntPoint& GatherPointsResolution);
 
-	bool RenderRayTracingGlobalIlluminationFinalGather(
+	void RenderRayTracingGlobalIlluminationFinalGather(
 		FRDGBuilder& GraphBuilder,
 		FSceneTextureParameters& SceneTextures,
 		FViewInfo& View,
 		const IScreenSpaceDenoiser::FAmbientOcclusionRayTracingConfig& RayTracingConfig,
 		int32 UpscaleFactor,
-		FRDGTextureUAV* GlobalIlluminationUAV,
-		FRDGTextureUAV* RayDistanceUAV);
+		IScreenSpaceDenoiser::FDiffuseIndirectInputs* OutDenoiserInputs);
+	
+	void RenderRayTracingAmbientOcclusion(
+		FRDGBuilder& GraphBuilder,
+		FViewInfo& View,
+		const FSceneTextureParameters& SceneTextures,
+		FRDGTextureRef* OutAmbientOcclusionTexture);
+	
 
+#if RHI_RAYTRACING
+	template <int TextureImportanceSampling> void RenderRayTracingRectLightInternal(FRHICommandListImmediate& RHICmdList, const TArray<FViewInfo>& Views, const FLightSceneInfo& RectLightSceneInfo, TRefCountPtr<IPooledRenderTarget>& ScreenShadowMaskTexture, TRefCountPtr<IPooledRenderTarget>& RayDistanceTexture);
+	void VisualizeRectLightMipTree(FRHICommandListImmediate& RHICmdList, const FViewInfo& View, const FRWBuffer& RectLightMipTree, const FIntVector& RectLightMipTreeDimensions);
+	
 	void BuildSkyLightCdfs(FRHICommandListImmediate& RHICmdList, FSkyLightSceneProxy* SkyLight);
 	void BuildSkyLightMipTree(FRHICommandListImmediate& RHICmdList, FTextureRHIRef SkyLightTexture, FRWBuffer& SkyLightMipTreePosX, FRWBuffer& SkyLightMipTreePosY, FRWBuffer& SkyLightMipTreePosZ, FRWBuffer& SkyLightMipTreeNegX, FRWBuffer& SkyLightMipTreeNegY, FRWBuffer& SkyLightMipTreeNegZ, FIntVector& SkyLightMipTreeDimensions);
 	void BuildSkyLightMipTreePdf(
diff --git a/Engine/Source/Runtime/Renderer/Private/IndirectLightRendering.cpp b/Engine/Source/Runtime/Renderer/Private/IndirectLightRendering.cpp
--- a/Engine/Source/Runtime/Renderer/Private/IndirectLightRendering.cpp
+++ b/Engine/Source/Runtime/Renderer/Private/IndirectLightRendering.cpp
@@ -136,7 +136,11 @@ void FDeferredShadingSceneRenderer::RenderDiffuseIndirectAndAmbientOcclusion(FRH
 		IScreenSpaceDenoiser::FDiffuseIndirectInputs DenoiserInputs;
 		if (bApplyRTGI)
 		{
-			RenderRayTracingGlobalIllumination(GraphBuilder, SceneTextures, View, /* out */ &RayTracingConfig, /* out */ &DenoiserInputs);
+			bool bIsValid = RenderRayTracingGlobalIllumination(GraphBuilder, SceneTextures, View, /* out */ &RayTracingConfig, /* out */ &DenoiserInputs);
+			if (!bIsValid)
+			{
+				DenoiseMode = 0;
+			}
 		}
 		else if (bApplySSGI)
 		{
@@ -169,16 +173,16 @@ void FDeferredShadingSceneRenderer::RenderDiffuseIndirectAndAmbientOcclusion(FRH
 				DenoiserToUse->GetDebugName(),
 				View.ViewRect.Width(), View.ViewRect.Height());
 
-			DenoiserOutputs = DenoiserToUse->DenoiseDiffuseIndirect(
-				GraphBuilder,
-				View,
-				&View.PrevViewInfo,
-				SceneTextures,
-				DenoiserInputs,
-				RayTracingConfig);
-		}
-		else
-		{
+				DenoiserOutputs = DenoiserToUse->DenoiseDiffuseIndirect(
+					GraphBuilder,
+					View,
+					&View.PrevViewInfo,
+					SceneTextures,
+					DenoiserInputs,
+					RayTracingConfig);
+			}
+			else
+			{
 			DenoiserOutputs.Color = DenoiserInputs.Color;
 			DenoiserOutputs.AmbientOcclusionMask = DenoiserInputs.AmbientOcclusionMask;
 		}
diff --git a/Engine/Source/Runtime/Renderer/Private/RayTracing/RayTracingGlobalIllumination.cpp b/Engine/Source/Runtime/Renderer/Private/RayTracing/RayTracingGlobalIllumination.cpp
--- a/Engine/Source/Runtime/Renderer/Private/RayTracing/RayTracingGlobalIllumination.cpp
+++ b/Engine/Source/Runtime/Renderer/Private/RayTracing/RayTracingGlobalIllumination.cpp
@@ -46,11 +46,11 @@ static FAutoConsoleVariableRef CVarRayTracingGlobalIlluminationMaxRayDistance(
 	TEXT("Max ray distance (default = 1.0e27)")
 );
 
-static int32 GRayTracingGlobalIlluminationMaxBounces = -1;
-static FAutoConsoleVariableRef CVarRayTracingGlobalIlluminationMaxBounces(
-	TEXT("r.RayTracing.GlobalIllumination.MaxBounces"),
-	GRayTracingGlobalIlluminationMaxBounces,
-	TEXT("Max bounces (default = -1 (driven by postprocesing volume))")
+static TAutoConsoleVariable<int32> CVarRayTracingGlobalIlluminationMaxBounces(
+	TEXT("r.RayTracing.GlobalIlluminGlobalIllumination.MaxBouncesation"),
+	-1,
+	TEXT("Max bounces (default = -1 (driven by postprocesing volume))"),
+	ECVF_RenderThreadSafe
 );
 
 static int32 GRayTracingGlobalIlluminationNextEventEstimationSamples = 2;
@@ -99,6 +99,13 @@ static FAutoConsoleVariableRef CVarRayTracingGlobalIlluminationScreenPercentage(
 	TEXT("Screen percentage for ray tracing global illumination (default = 50)")
 );
 
+static TAutoConsoleVariable<int32> CVarRayTracingGlobalIlluminationEnableLightAttenuation(
+	TEXT("r.RayTracing.GlobalIllumination.EnableLightAttenuation"),
+	1,
+	TEXT("Enables light attenuation when calculating irradiance during next-event estimation (default = 1)"),
+	ECVF_RenderThreadSafe
+);
+
 static TAutoConsoleVariable<int32> CVarRayTracingGlobalIlluminationEnableTwoSidedGeometry(
 	TEXT("r.RayTracing.GlobalIllumination.EnableTwoSidedGeometry"),
 	1,
@@ -159,7 +166,6 @@ void SetupLightParameters(
 		if (LightParameters->Count >= MaxLightCount) break;
 
 		if (Light.LightSceneInfo->Proxy->HasStaticLighting() && Light.LightSceneInfo->IsPrecomputedLightingValid()) continue;
-		if (!Light.LightSceneInfo->Proxy->AffectGlobalIllumination()) continue;
 
 		FLightShaderParameters LightShaderParameters;
 		Light.LightSceneInfo->Proxy->GetLightShaderParameters(LightShaderParameters);
@@ -236,7 +242,7 @@ bool ShouldRenderRayTracingGlobalIllumination(const FViewInfo& View)
 	{
 		return false;
 	}
-
+	
 	if (GetForceRayTracingEffectsCVarValue() >= 0)
 	{
 		return GetForceRayTracingEffectsCVarValue() > 0;
@@ -276,8 +282,8 @@ class FGlobalIlluminationRGS : public FGlobalShader
 		SHADER_PARAMETER(float, MaxRayDistanceForAO)
 		SHADER_PARAMETER(float, NextEventEstimationSamples)
 		SHADER_PARAMETER(float, DiffuseThreshold)
-		SHADER_PARAMETER(bool, EvalSkyLight)
-		SHADER_PARAMETER(bool, UseRussianRoulette)
+		SHADER_PARAMETER(uint32, EvalSkyLight)
+		SHADER_PARAMETER(uint32, UseRussianRoulette)
 		SHADER_PARAMETER(float, MaxNormalBias)
 		SHADER_PARAMETER(uint32, TileOffsetX)
 		SHADER_PARAMETER(uint32, TileOffsetY)
@@ -297,60 +303,7 @@ class FGlobalIlluminationRGS : public FGlobalShader
 	END_SHADER_PARAMETER_STRUCT()
 };
 
-class FRayTracingGlobalIlluminationCompositePS : public FGlobalShader
-{
-	DECLARE_GLOBAL_SHADER(FRayTracingGlobalIlluminationCompositePS)
-	SHADER_USE_PARAMETER_STRUCT(FRayTracingGlobalIlluminationCompositePS, FGlobalShader)
-
-	static bool ShouldCompilePermutation(const FGlobalShaderPermutationParameters& Parameters)
-	{
-		return ShouldCompileRayTracingShadersForProject(Parameters.Platform);
-	}
-
-	BEGIN_SHADER_PARAMETER_STRUCT(FParameters, )
-		RENDER_TARGET_BINDING_SLOTS()
-		SHADER_PARAMETER_RDG_TEXTURE(Texture2D, GlobalIlluminationTexture)
-		SHADER_PARAMETER_SAMPLER(SamplerState, GlobalIlluminationSampler)
-		SHADER_PARAMETER_STRUCT_REF(FViewUniformShaderParameters, ViewUniformBuffer)
-		SHADER_PARAMETER_STRUCT_INCLUDE(FSceneTextureParameters, SceneTextures)
-	END_SHADER_PARAMETER_STRUCT()
-};
-
-class FRayTracingGlobalIlluminationSceneColorCompositePS : public FGlobalShader
-{
-	DECLARE_GLOBAL_SHADER(FRayTracingGlobalIlluminationSceneColorCompositePS)
-	SHADER_USE_PARAMETER_STRUCT(FRayTracingGlobalIlluminationSceneColorCompositePS, FGlobalShader)
-
-	static bool ShouldCompilePermutation(const FGlobalShaderPermutationParameters& Parameters)
-	{
-		return ShouldCompileRayTracingShadersForProject(Parameters.Platform);
-	}
-
-	BEGIN_SHADER_PARAMETER_STRUCT(FParameters, )
-		RENDER_TARGET_BINDING_SLOTS()
-		SHADER_PARAMETER_RDG_TEXTURE(Texture2D, GlobalIlluminationTexture)
-		SHADER_PARAMETER_SAMPLER(SamplerState, GlobalIlluminationSampler)
-		SHADER_PARAMETER_STRUCT_INCLUDE(FSceneTextureParameters, SceneTextures)
-	END_SHADER_PARAMETER_STRUCT()
-};
-
-class FRayTracingGlobalIlluminationCHS : public FGlobalShader
-{
-	DECLARE_GLOBAL_SHADER(FRayTracingGlobalIlluminationCHS)
-	SHADER_USE_ROOT_PARAMETER_STRUCT(FRayTracingGlobalIlluminationCHS, FGlobalShader)
-
-	static bool ShouldCompilePermutation(const FGlobalShaderPermutationParameters& Parameters)
-	{
-		return ShouldCompileRayTracingShadersForProject(Parameters.Platform);
-	}
-
-	using FParameters = FEmptyShaderParameters;
-};
-
 IMPLEMENT_GLOBAL_SHADER(FGlobalIlluminationRGS, "/Engine/Private/RayTracing/RayTracingGlobalIlluminationRGS.usf", "GlobalIlluminationRGS", SF_RayGen);
-IMPLEMENT_GLOBAL_SHADER(FRayTracingGlobalIlluminationCHS, "/Engine/Private/RayTracing/RayTracingGlobalIlluminationRGS.usf", "RayTracingGlobalIlluminationCHS", SF_RayHitGroup);
-IMPLEMENT_GLOBAL_SHADER(FRayTracingGlobalIlluminationCompositePS, "/Engine/Private/RayTracing/RayTracingGlobalIlluminationCompositePS.usf", "GlobalIlluminationCompositePS", SF_Pixel);
-IMPLEMENT_GLOBAL_SHADER(FRayTracingGlobalIlluminationSceneColorCompositePS, "/Engine/Private/RayTracing/RayTracingGlobalIlluminationCompositePS.usf", "GlobalIlluminationSceneColorCompositePS", SF_Pixel);
 
 struct GatherPoints
 {
@@ -364,7 +317,7 @@ class FRayTracingGlobalIlluminationCreateGatherPointsRGS : public FGlobalShader
 	DECLARE_GLOBAL_SHADER(FRayTracingGlobalIlluminationCreateGatherPointsRGS)
 	SHADER_USE_ROOT_PARAMETER_STRUCT(FRayTracingGlobalIlluminationCreateGatherPointsRGS, FGlobalShader)
 
-	class FUseAttenuationTermDim : SHADER_PERMUTATION_BOOL("USE_ATTENUATION_TERM");
+		class FUseAttenuationTermDim : SHADER_PERMUTATION_BOOL("USE_ATTENUATION_TERM");
 	class FEnableTwoSidedGeometryDim : SHADER_PERMUTATION_BOOL("ENABLE_TWO_SIDED_GEOMETRY");
 
 	using FPermutationDomain = TShaderPermutationDomain<FUseAttenuationTermDim, FEnableTwoSidedGeometryDim>;
@@ -385,8 +338,8 @@ class FRayTracingGlobalIlluminationCreateGatherPointsRGS : public FGlobalShader
 		SHADER_PARAMETER(float, NextEventEstimationSamples)
 		SHADER_PARAMETER(float, DiffuseThreshold)
 		SHADER_PARAMETER(float, MaxNormalBias)
-		SHADER_PARAMETER(bool, EvalSkyLight)
-		SHADER_PARAMETER(bool, UseRussianRoulette)
+		SHADER_PARAMETER(uint32, EvalSkyLight)
+		SHADER_PARAMETER(uint32, UseRussianRoulette)
 
 		// Scene data
 		SHADER_PARAMETER_SRV(RaytracingAccelerationStructure, TLAS)
@@ -419,7 +372,7 @@ class FRayTracingGlobalIlluminationFinalGatherRGS : public FGlobalShader
 	DECLARE_GLOBAL_SHADER(FRayTracingGlobalIlluminationFinalGatherRGS)
 	SHADER_USE_ROOT_PARAMETER_STRUCT(FRayTracingGlobalIlluminationFinalGatherRGS, FGlobalShader)
 
-	class FUseAttenuationTermDim : SHADER_PERMUTATION_BOOL("USE_ATTENUATION_TERM");
+		class FUseAttenuationTermDim : SHADER_PERMUTATION_BOOL("USE_ATTENUATION_TERM");
 	class FEnableTwoSidedGeometryDim : SHADER_PERMUTATION_BOOL("ENABLE_TWO_SIDED_GEOMETRY");
 
 	using FPermutationDomain = TShaderPermutationDomain<FUseAttenuationTermDim, FEnableTwoSidedGeometryDim>;
@@ -468,7 +421,7 @@ void FDeferredShadingSceneRenderer::PrepareRayTracingGlobalIllumination(const FV
 		for (int EnableTwoSidedGeometry = 0; EnableTwoSidedGeometry < 2; ++EnableTwoSidedGeometry)
 		{
 			FGlobalIlluminationRGS::FPermutationDomain PermutationVector;
-			PermutationVector.Set<FGlobalIlluminationRGS::FUseAttenuationTermDim>(UseAttenuationTerm == 1);
+			PermutationVector.Set<FGlobalIlluminationRGS::FUseAttenuationTermDim>(CVarRayTracingGlobalIlluminationEnableLightAttenuation.GetValueOnRenderThread() != 0);
 			PermutationVector.Set<FGlobalIlluminationRGS::FEnableTwoSidedGeometryDim>(EnableTwoSidedGeometry == 1);
 			TShaderMapRef<FGlobalIlluminationRGS> RayGenerationShader(View.ShaderMap, PermutationVector);
 			OutRayGenShaders.Add(RayGenerationShader->GetRayTracingShader());
@@ -489,218 +442,64 @@ void FDeferredShadingSceneRenderer::PrepareRayTracingGlobalIllumination(const FV
 
 }
 
-void FDeferredShadingSceneRenderer::RenderRayTracingGlobalIllumination(
-	FRHICommandListImmediate& RHICmdList, 
-	TRefCountPtr<IPooledRenderTarget>& GlobalIlluminationRT
-)
-{
-	SCOPED_GPU_STAT(RHICmdList, RayTracingGIBruteForce);
+#endif // RHI_RAYTRACING
 
-	bool bAnyViewWithRTGI = false;
-	for (int32 ViewIndex = 0; ViewIndex < Views.Num(); ++ViewIndex)
-	{
-		FViewInfo& View = Views[ViewIndex];
-		bAnyViewWithRTGI = bAnyViewWithRTGI || ShouldRenderRayTracingGlobalIllumination(View);
-	}
+bool FDeferredShadingSceneRenderer::RenderRayTracingGlobalIllumination(
+	FRDGBuilder& GraphBuilder, 
+	FSceneTextureParameters& SceneTextures,
+	FViewInfo& View,
+	IScreenSpaceDenoiser::FAmbientOcclusionRayTracingConfig* OutRayTracingConfig,
+	IScreenSpaceDenoiser::FDiffuseIndirectInputs* OutDenoiserInputs)
+#if RHI_RAYTRACING
+{	
+	if (!View.ViewState) return false;
 
-	if (!bAnyViewWithRTGI)
-	{
-		return;
-	}
+	int32 RayTracingGISamplesPerPixel = GetRayTracingGlobalIlluminationSamplesPerPixel(View);
+	if (RayTracingGISamplesPerPixel <= 0) return false;
 
-	IScreenSpaceDenoiser::FAmbientOcclusionRayTracingConfig RayTracingConfig;
-	RayTracingConfig.ResolutionFraction = 1.0;
+	OutRayTracingConfig->ResolutionFraction = 1.0;
 	if (GRayTracingGlobalIlluminationDenoiser != 0)
 	{
-		RayTracingConfig.ResolutionFraction = FMath::Clamp(GRayTracingGlobalIlluminationScreenPercentage / 100.0, 0.25, 1.0);
+		OutRayTracingConfig->ResolutionFraction = FMath::Clamp(GRayTracingGlobalIlluminationScreenPercentage / 100.0, 0.25, 1.0);
 	}
 
-	int32 UpscaleFactor = int32(1.0 / RayTracingConfig.ResolutionFraction);
+	OutRayTracingConfig->RayCountPerPixel = RayTracingGISamplesPerPixel;
+	
+	int32 UpscaleFactor = int32(1.0 / OutRayTracingConfig->ResolutionFraction);
 
-	FSceneRenderTargets& SceneContext = FSceneRenderTargets::Get(RHICmdList);
+	// Allocate input for the denoiser.
 	{
-		FPooledRenderTargetDesc Desc = SceneContext.GetSceneColor()->GetDesc();
-		Desc.Format = PF_FloatRGBA;
-		Desc.Flags &= ~(TexCreate_FastVRAM | TexCreate_Transient);
-		GRenderTargetPool.FindFreeElement(RHICmdList, Desc, GlobalIlluminationRT, TEXT("RayTracingGlobalIllumination"));
-	}
+		FRDGTextureDesc Desc = FRDGTextureDesc::Create2DDesc(
+			SceneTextures.SceneDepthBuffer->Desc.Extent / UpscaleFactor,
+			PF_FloatRGBA,
+			FClearValueBinding::None,
+			/* InFlags = */ TexCreate_None,
+			/* InTargetableFlags = */ TexCreate_ShaderResource | TexCreate_RenderTargetable | TexCreate_UAV,
+			/* bInForceSeparateTargetAndShaderResource = */ false);
 
-	FRDGBuilder GraphBuilder(RHICmdList);
+		OutDenoiserInputs->Color = GraphBuilder.CreateTexture(Desc, TEXT("RayTracingDiffuseIndirect"));
 
-	FRDGTextureRef GlobalIlluminationTexture;
-	{
-		FRDGTextureDesc Desc = SceneContext.GetSceneColor()->GetDesc();
-		Desc.Extent /= UpscaleFactor;
-		Desc.Format = PF_FloatRGBA;
-		Desc.Flags &= ~(TexCreate_FastVRAM | TexCreate_Transient);
-		GlobalIlluminationTexture = GraphBuilder.CreateTexture(Desc, TEXT("RayTracingGlobalIllumination"));
-	}
-	FRDGTextureUAV* GlobalIlluminationUAV = GraphBuilder.CreateUAV(FRDGTextureUAVDesc(GlobalIlluminationTexture));
-
-	FRDGTextureRef RayDistanceTexture;
-	{
-		FRDGTextureDesc Desc = SceneContext.GetSceneColor()->GetDesc();
-		Desc.Extent /= UpscaleFactor;
 		Desc.Format = PF_G16R16;
-		Desc.Flags &= ~(TexCreate_FastVRAM | TexCreate_Transient);
-		RayDistanceTexture = GraphBuilder.CreateTexture(Desc, TEXT("RayTracingGlobalIlluminationRayDistance"));
-	}
-	FRDGTextureUAV* RayDistanceUAV = GraphBuilder.CreateUAV(FRDGTextureUAVDesc(RayDistanceTexture));
-
-	TRefCountPtr<IPooledRenderTarget>& AmbientOcclusionRT = SceneContext.ScreenSpaceAO;
-
-	for (int32 ViewIndex = 0; ViewIndex < Views.Num(); ++ViewIndex)
-	{
-		FViewInfo& View = Views[ViewIndex];
-		if (ShouldRenderRayTracingGlobalIllumination(View))
-		{
-			RenderRayTracingGlobalIllumination(
-				RHICmdList, 
-				GraphBuilder, 
-				View, 
-				RayTracingConfig, 
-				UpscaleFactor, 
-				GlobalIlluminationRT, 
-				AmbientOcclusionRT, 
-				GlobalIlluminationUAV, 
-				RayDistanceUAV, 
-				GlobalIlluminationTexture, 
-				RayDistanceTexture
-			);
-		}
-	}
-
-	GraphBuilder.Execute();
-	SceneContext.bScreenSpaceAOIsValid = true;
-	GVisualizeTexture.SetCheckPoint(RHICmdList, GlobalIlluminationRT);
-}
-
-void FDeferredShadingSceneRenderer::RenderRayTracingGlobalIllumination(
-	FRHICommandListImmediate& RHICmdList,
-	FRDGBuilder& GraphBuilder,
-	FViewInfo& View,
-	IScreenSpaceDenoiser::FAmbientOcclusionRayTracingConfig& RayTracingConfig,
-	int32 UpscaleFactor,
-	TRefCountPtr<IPooledRenderTarget>& GlobalIlluminationRT,
-	TRefCountPtr<IPooledRenderTarget>& AmbientOcclusionRT,
-	FRDGTextureUAV* GlobalIlluminationUAV,
-	FRDGTextureUAV* RayDistanceUAV,
-	const FRDGTextureRef GlobalIlluminationTexture,
-	const FRDGTextureRef RayDistanceTexture
-)
-{	
-	FSceneRenderTargets& SceneContext = FSceneRenderTargets::Get(RHICmdList);
-
-	RDG_EVENT_SCOPE(GraphBuilder, "RTGI");
-
-	RayTracingConfig.RayCountPerPixel = GetRayTracingGlobalIlluminationSamplesPerPixel(View);
-
-	FSceneTextureParameters SceneTextures;
-	SetupSceneTextureParameters(GraphBuilder, &SceneTextures);
-
-	if (Scene->SkyLight && Scene->SkyLight->ShouldRebuildCdf())
-	{
-		BuildSkyLightCdfs(RHICmdList, Scene->SkyLight);
+		OutDenoiserInputs->RayHitDistance = GraphBuilder.CreateTexture(Desc, TEXT("RayTracingDiffuseIndirectHitDistance"));
 	}
 
-	// Ray generation
-	bool bIsValid;
+	// Ray generation pass
 	if (CVarRayTracingGlobalIlluminationEnableFinalGather.GetValueOnRenderThread() != 0)
 	{
-		bIsValid = RenderRayTracingGlobalIlluminationFinalGather(GraphBuilder, SceneTextures, View, RayTracingConfig, UpscaleFactor, GlobalIlluminationUAV, RayDistanceUAV);
-	}
-	else
-	{
-		bIsValid = RenderRayTracingGlobalIlluminationBruteForce(RHICmdList, GraphBuilder, SceneTextures, View, RayTracingConfig, UpscaleFactor, GlobalIlluminationUAV, RayDistanceUAV);
-	}
-
-	// Denoising
-
-	FRDGTextureRef ResultTexture; //#dxr_todo review
-
-	if (GRayTracingGlobalIlluminationDenoiser != 0 && bIsValid)
-	{
-		const IScreenSpaceDenoiser* DefaultDenoiser = IScreenSpaceDenoiser::GetDefaultDenoiser();
-		const IScreenSpaceDenoiser* DenoiserToUse = GRayTracingGlobalIlluminationDenoiser == 1 ? DefaultDenoiser : GScreenSpaceDenoiser;
-
-		IScreenSpaceDenoiser::FDiffuseIndirectInputs DenoiserInputs;
-		DenoiserInputs.Color = GlobalIlluminationTexture;
-		DenoiserInputs.RayHitDistance = RayDistanceTexture;
-
-		{
-			RDG_EVENT_SCOPE(GraphBuilder, "%s%s(DiffuseIndirect) %dx%d",
-				DenoiserToUse != DefaultDenoiser ? TEXT("ThirdParty ") : TEXT(""),
-				DenoiserToUse->GetDebugName(),
-				View.ViewRect.Width(), View.ViewRect.Height());
-
-			IScreenSpaceDenoiser::FDiffuseIndirectOutputs DenoiserOutputs = DenoiserToUse->DenoiseDiffuseIndirect(
-				GraphBuilder,
-				View,
-				&View.PrevViewInfo,
-				SceneTextures,
-				DenoiserInputs,
-				RayTracingConfig);
-
-			ResultTexture = DenoiserOutputs.Color;
-		}
+		RenderRayTracingGlobalIlluminationFinalGather(GraphBuilder, SceneTextures, View, *OutRayTracingConfig, UpscaleFactor, OutDenoiserInputs);
 	}
 	else
 	{
-		ResultTexture = GlobalIlluminationTexture;
-	}
-
-	// Compositing
-	if (bIsValid)
-	{
-		FRayTracingGlobalIlluminationCompositePS::FParameters *PassParameters = GraphBuilder.AllocParameters<FRayTracingGlobalIlluminationCompositePS::FParameters>();
-		PassParameters->GlobalIlluminationTexture = ResultTexture;
-		PassParameters->GlobalIlluminationSampler = TStaticSamplerState<SF_Point, AM_Clamp, AM_Clamp, AM_Clamp>::GetRHI();
-		PassParameters->ViewUniformBuffer = View.ViewUniformBuffer;
-		PassParameters->RenderTargets[0] = FRenderTargetBinding(GraphBuilder.RegisterExternalTexture(GlobalIlluminationRT), ERenderTargetLoadAction::ELoad, ERenderTargetStoreAction::EStore);
-		PassParameters->RenderTargets[1] = FRenderTargetBinding(GraphBuilder.RegisterExternalTexture(AmbientOcclusionRT), ERenderTargetLoadAction::ELoad, ERenderTargetStoreAction::EStore);
-		PassParameters->SceneTextures = SceneTextures;
-
-		GraphBuilder.AddPass(
-			RDG_EVENT_NAME("GlobalIlluminationComposite"),
-			PassParameters,
-			ERDGPassFlags::Raster,
-			[this, &SceneContext, &View, &SceneTextures, PassParameters](FRHICommandListImmediate& RHICmdList)
-			{
-				TShaderMapRef<FPostProcessVS> VertexShader(View.ShaderMap);
-				TShaderMapRef<FRayTracingGlobalIlluminationCompositePS> PixelShader(View.ShaderMap);
-				FGraphicsPipelineStateInitializer GraphicsPSOInit;
-				RHICmdList.ApplyCachedRenderTargets(GraphicsPSOInit);
-
-				// Additive blending
-				GraphicsPSOInit.BlendState = TStaticBlendState<CW_RGBA, BO_Add, BF_One, BF_Zero, BO_Add, BF_One, BF_Zero>::GetRHI();
-				//GraphicsPSOInit.BlendState = TStaticBlendState<CW_RGBA, BO_Add, BF_One, BF_One, BO_Add, BF_One, BF_One>::GetRHI();
-				GraphicsPSOInit.RasterizerState = TStaticRasterizerState<FM_Solid, CM_None>::GetRHI();
-				GraphicsPSOInit.DepthStencilState = TStaticDepthStencilState<false, CF_Always>::GetRHI();
-
-				GraphicsPSOInit.BoundShaderState.VertexDeclarationRHI = GFilterVertexDeclaration.VertexDeclarationRHI;
-				GraphicsPSOInit.BoundShaderState.VertexShaderRHI = GETSAFERHISHADER_VERTEX(*VertexShader);
-				GraphicsPSOInit.BoundShaderState.PixelShaderRHI = GETSAFERHISHADER_PIXEL(*PixelShader);
-				GraphicsPSOInit.PrimitiveType = PT_TriangleList;
-				SetGraphicsPipelineState(RHICmdList, GraphicsPSOInit);
-
-				SetShaderParameters(RHICmdList, *PixelShader, PixelShader->GetPixelShader(), *PassParameters);
-
-				RHICmdList.SetViewport(View.ViewRect.Min.X, View.ViewRect.Min.Y, 0.0f, View.ViewRect.Max.X, View.ViewRect.Max.Y, 1.0f);
-
-				DrawRectangle(
-					RHICmdList,
-					0, 0,
-					View.ViewRect.Width(), View.ViewRect.Height(),
-					View.ViewRect.Min.X, View.ViewRect.Min.Y,
-					View.ViewRect.Width(), View.ViewRect.Height(),
-					FIntPoint(View.ViewRect.Width(), View.ViewRect.Height()),
-					SceneContext.GetBufferSizeXY(),
-					*VertexShader
-				);
-			}
-		);
+		RenderRayTracingGlobalIlluminationBruteForce(GraphBuilder, SceneTextures, View, *OutRayTracingConfig, UpscaleFactor, OutDenoiserInputs);
 	}
+	return true;
+}
+#else
+{
+	unimplemented();
+	return false;
 }
+#endif // RHI_RAYTRACING
 
 void FDeferredShadingSceneRenderer::RayTracingGlobalIlluminationCreateGatherPoints(
 	FRDGBuilder& GraphBuilder,
@@ -815,25 +614,19 @@ void FDeferredShadingSceneRenderer::RayTracingGlobalIlluminationCreateGatherPoin
 }
 #endif
 
-bool FDeferredShadingSceneRenderer::RenderRayTracingGlobalIlluminationFinalGather(
+void FDeferredShadingSceneRenderer::RenderRayTracingGlobalIlluminationFinalGather(
 	FRDGBuilder& GraphBuilder,
 	FSceneTextureParameters& SceneTextures,
 	FViewInfo& View,
 	const IScreenSpaceDenoiser::FAmbientOcclusionRayTracingConfig& RayTracingConfig,
 	int32 UpscaleFactor,
 	// Output
-	FRDGTextureUAV* GlobalIlluminationUAV,
-	FRDGTextureUAV* RayDistanceUAV)
+	IScreenSpaceDenoiser::FDiffuseIndirectInputs* OutDenoiserInputs)
 #if RHI_RAYTRACING
 {
-	FSceneViewState* SceneViewState = (FSceneViewState*)View.State;
-	if (!SceneViewState) return false;
-
-	int32 SamplesPerPixel = GetRayTracingGlobalIlluminationSamplesPerPixel(View);
-	if (SamplesPerPixel <= 0) return false;
-
 	// Generate gather points
 	FRDGBufferRef GatherPointsBuffer;
+	FSceneViewState* SceneViewState = (FSceneViewState*)View.State;
 	RayTracingGlobalIlluminationCreateGatherPoints(GraphBuilder, SceneTextures, View, UpscaleFactor, GatherPointsBuffer, SceneViewState->GatherPointsResolution);
 
 	// Perform gather
@@ -841,6 +634,7 @@ bool FDeferredShadingSceneRenderer::RenderRayTracingGlobalIlluminationFinalGathe
 	RDG_EVENT_SCOPE(GraphBuilder, "Ray Tracing GI: Final Gather");
 
 	FRayTracingGlobalIlluminationFinalGatherRGS::FParameters* PassParameters = GraphBuilder.AllocParameters<FRayTracingGlobalIlluminationFinalGatherRGS::FParameters>();
+	int32 SamplesPerPixel = GetRayTracingGlobalIlluminationSamplesPerPixel(View);
 	int32 SampleIndex = View.ViewState->FrameIndex % SamplesPerPixel;
 	PassParameters->SampleIndex = SampleIndex;
 	PassParameters->SamplesPerPixel = SamplesPerPixel;
@@ -870,8 +664,8 @@ bool FDeferredShadingSceneRenderer::RenderRayTracingGlobalIlluminationFinalGathe
 	PassParameters->GatherPointsBuffer = GraphBuilder.CreateSRV(GatherPointsBuffer);
 
 	// Output
-	PassParameters->RWGlobalIlluminationUAV = GlobalIlluminationUAV;
-	PassParameters->RWRayDistanceUAV = RayDistanceUAV;
+	PassParameters->RWGlobalIlluminationUAV = GraphBuilder.CreateUAV(OutDenoiserInputs->Color);
+	PassParameters->RWRayDistanceUAV = GraphBuilder.CreateUAV(OutDenoiserInputs->RayHitDistance);
 
 	FRayTracingGlobalIlluminationFinalGatherRGS::FPermutationDomain PermutationVector;
 	PermutationVector.Set<FRayTracingGlobalIlluminationFinalGatherRGS::FUseAttenuationTermDim>(true);
@@ -894,37 +688,31 @@ bool FDeferredShadingSceneRenderer::RenderRayTracingGlobalIlluminationFinalGathe
 	});
 
 
-	GraphBuilder.QueueBufferExtraction(GatherPointsBuffer, &SceneViewState->GatherPointsBuffer);
-	return true;
+	GraphBuilder.QueueBufferExtraction(GatherPointsBuffer, &SceneViewState->GatherPointsBuffer, FRDGResourceState::EAccess::Read, FRDGResourceState::EPipeline::Compute);
 }
 #else
 {
 	unimplemented();
-	return true;
 }
 #endif
 
-bool FDeferredShadingSceneRenderer::RenderRayTracingGlobalIlluminationBruteForce(
-	FRHICommandListImmediate& RHICmdList,
+void FDeferredShadingSceneRenderer::RenderRayTracingGlobalIlluminationBruteForce(
 	FRDGBuilder& GraphBuilder,
 	FSceneTextureParameters& SceneTextures,
 	FViewInfo& View,
 	const IScreenSpaceDenoiser::FAmbientOcclusionRayTracingConfig& RayTracingConfig,
 	int32 UpscaleFactor,
-	FRDGTextureUAV* GlobalIlluminationUAV,
-	FRDGTextureUAV* RayDistanceUAV)
+	IScreenSpaceDenoiser::FDiffuseIndirectInputs* OutDenoiserInputs)
 #if RHI_RAYTRACING
 {
-	if (!View.ViewState) return false;
-
 	RDG_GPU_STAT_SCOPE(GraphBuilder, RayTracingGIBruteForce);
 	RDG_EVENT_SCOPE(GraphBuilder, "Ray Tracing GI: Brute Force");
 
 	int32 RayTracingGISamplesPerPixel = GetRayTracingGlobalIlluminationSamplesPerPixel(View);
-	uint32 IterationCount = FMath::Max(RayTracingGISamplesPerPixel, 1);
+	uint32 IterationCount = RayTracingGISamplesPerPixel;
 	uint32 SequenceCount = 1;
 	uint32 DimensionCount = 24;
-	FHaltonSequenceIteration HaltonSequenceIteration(Scene->HaltonSequence, IterationCount, SequenceCount, DimensionCount, View.ViewState ? (View.ViewState->FrameIndex % 1024) : 0);
+	FHaltonSequenceIteration HaltonSequenceIteration(Scene->HaltonSequence, IterationCount, SequenceCount, DimensionCount, View.ViewState->FrameIndex % 1024);
 
 	FHaltonIteration HaltonIteration;
 	InitializeHaltonSequenceIteration(HaltonSequenceIteration, HaltonIteration);
@@ -938,12 +726,18 @@ bool FDeferredShadingSceneRenderer::RenderRayTracingGlobalIlluminationBruteForce
 	FPathTracingLightData LightParameters;
 	SetupLightParameters(Scene->Lights, View, &LightParameters);
 
+	if (Scene->SkyLight && Scene->SkyLight->ShouldRebuildCdf())
+	{
+		// TODO: should be converted to RDG.
+		BuildSkyLightCdfs(GraphBuilder.RHICmdList, Scene->SkyLight);
+	}
 	FSkyLightData SkyLightParameters;
 	SetupSkyLightParameters(*Scene, &SkyLightParameters);
 
 	FGlobalIlluminationRGS::FParameters* PassParameters = GraphBuilder.AllocParameters<FGlobalIlluminationRGS::FParameters>();
 	PassParameters->SamplesPerPixel = RayTracingGISamplesPerPixel;
-	PassParameters->MaxBounces = GRayTracingGlobalIlluminationMaxBounces > -1 ? GRayTracingGlobalIlluminationMaxBounces : View.FinalPostProcessSettings.RayTracingGIMaxBounces;
+	int32 CVarRayTracingGlobalIlluminationMaxBouncesValue = CVarRayTracingGlobalIlluminationMaxBounces.GetValueOnRenderThread();
+	PassParameters->MaxBounces = CVarRayTracingGlobalIlluminationMaxBouncesValue > -1? CVarRayTracingGlobalIlluminationMaxBouncesValue : View.FinalPostProcessSettings.RayTracingGIMaxBounces;
 	PassParameters->MaxNormalBias = GetRaytracingMaxNormalBias();
 	float MaxRayDistanceForGI = GRayTracingGlobalIlluminationMaxRayDistance;
 	if (MaxRayDistanceForGI == -1.0)
@@ -965,15 +759,17 @@ bool FDeferredShadingSceneRenderer::RenderRayTracingGlobalIlluminationBruteForce
 	PassParameters->LightParameters = CreateUniformBufferImmediate(LightParameters, EUniformBufferUsage::UniformBuffer_SingleDraw);
 	PassParameters->SceneTextures = SceneTextures;
 	PassParameters->SkyLight = CreateUniformBufferImmediate(SkyLightParameters, EUniformBufferUsage::UniformBuffer_SingleDraw);
-	TRefCountPtr<IPooledRenderTarget> SubsurfaceProfileRT((IPooledRenderTarget*)GetSubsufaceProfileTexture_RT(RHICmdList));
+
+	// TODO: should be converted to RDG
+	TRefCountPtr<IPooledRenderTarget> SubsurfaceProfileRT((IPooledRenderTarget*) GetSubsufaceProfileTexture_RT(GraphBuilder.RHICmdList));
 	if (!SubsurfaceProfileRT)
 	{
 		SubsurfaceProfileRT = GSystemTextures.BlackDummy;
 	}
 	PassParameters->SSProfilesTexture = GraphBuilder.RegisterExternalTexture(SubsurfaceProfileRT);
 	PassParameters->TransmissionProfilesLinearSampler = TStaticSamplerState<SF_Bilinear, AM_Clamp, AM_Clamp, AM_Clamp>::GetRHI();
-	PassParameters->RWGlobalIlluminationUAV = GlobalIlluminationUAV;
-	PassParameters->RWRayDistanceUAV = RayDistanceUAV;
+	PassParameters->RWGlobalIlluminationUAV = GraphBuilder.CreateUAV(OutDenoiserInputs->Color);
+	PassParameters->RWRayDistanceUAV = GraphBuilder.CreateUAV(OutDenoiserInputs->RayHitDistance);
 	PassParameters->TileOffsetX = 0;
 	PassParameters->TileOffsetY = 0;
 
@@ -993,10 +789,10 @@ bool FDeferredShadingSceneRenderer::RenderRayTracingGlobalIlluminationBruteForce
 			ERDGPassFlags::Compute,
 			[PassParameters, this, &View, RayGenerationShader, RayTracingResolution](FRHICommandList& RHICmdList)
 		{
+			FRHIRayTracingScene* RayTracingSceneRHI = View.RayTracingScene.RayTracingSceneRHI;
+
 			FRayTracingShaderBindingsWriter GlobalResources;
 			SetShaderParameters(GlobalResources, *RayGenerationShader, *PassParameters);
-
-			FRHIRayTracingScene* RayTracingSceneRHI = View.RayTracingScene.RayTracingSceneRHI;
 			RHICmdList.RayTraceDispatch(View.RayTracingMaterialPipeline, RayGenerationShader->GetRayTracingShader(), RayTracingSceneRHI, GlobalResources, RayTracingResolution.X, RayTracingResolution.Y);
 		});
 	}
@@ -1033,80 +829,16 @@ bool FDeferredShadingSceneRenderer::RenderRayTracingGlobalIlluminationBruteForce
 
 					FRayTracingShaderBindingsWriter GlobalResources;
 					SetShaderParameters(GlobalResources, *RayGenerationShader, *TilePassParameters);
-					RHICmdList.RayTraceDispatch(View.RayTracingMaterialPipeline, RayGenerationShader->GetRayTracingShader(), RayTracingSceneRHI,
+					RHICmdList.RayTraceDispatch(View.RayTracingMaterialPipeline, RayGenerationShader->GetRayTracingShader(), RayTracingSceneRHI, 
 						GlobalResources, DispatchSizeX, DispatchSizeY);
 					RHICmdList.SubmitCommandsHint();
 				});
 			}
 		}
 	}
-	return true;
 }
 #else
 {
 	unimplemented();
-	return false;
-}
-#endif // RHI_RAYTRACING
-
-void FDeferredShadingSceneRenderer::CompositeGlobalIllumination(
-	FRHICommandListImmediate& RHICmdList,
-	const FViewInfo& View,
-	TRefCountPtr<IPooledRenderTarget>& GlobalIlluminationRT
-)
-{
-	FSceneRenderTargets& SceneContext = FSceneRenderTargets::Get(RHICmdList);
-
-	FRDGBuilder GraphBuilder(RHICmdList);
-	FSceneTextureParameters SceneTextures;
-	SetupSceneTextureParameters(GraphBuilder, &SceneTextures);
-
-	FRayTracingGlobalIlluminationSceneColorCompositePS::FParameters *PassParameters = GraphBuilder.AllocParameters<FRayTracingGlobalIlluminationSceneColorCompositePS::FParameters>();
-	PassParameters->GlobalIlluminationTexture = GraphBuilder.RegisterExternalTexture(GlobalIlluminationRT);
-	PassParameters->GlobalIlluminationSampler = TStaticSamplerState<SF_Point, AM_Clamp, AM_Clamp, AM_Clamp>::GetRHI();
-	PassParameters->RenderTargets[0] = FRenderTargetBinding(GraphBuilder.RegisterExternalTexture(SceneContext.GetSceneColor()), ERenderTargetLoadAction::ELoad, ERenderTargetStoreAction::EStore);
-	PassParameters->SceneTextures = SceneTextures;
-
-	GraphBuilder.AddPass(
-		RDG_EVENT_NAME("GlobalIlluminationComposite"),
-		PassParameters,
-		ERDGPassFlags::Raster,
-		[this, &SceneContext, &View, PassParameters](FRHICommandListImmediate& RHICmdList)
-		{
-			TShaderMapRef<FPostProcessVS> VertexShader(View.ShaderMap);
-			TShaderMapRef<FRayTracingGlobalIlluminationSceneColorCompositePS> PixelShader(View.ShaderMap);
-			FGraphicsPipelineStateInitializer GraphicsPSOInit;
-			RHICmdList.ApplyCachedRenderTargets(GraphicsPSOInit);
-
-			// Additive blending
-			GraphicsPSOInit.BlendState = TStaticBlendState<CW_RGBA, BO_Add, BF_One, BF_One, BO_Add, BF_One, BF_One>::GetRHI();
-			//GraphicsPSOInit.BlendState = TStaticBlendState<CW_RGBA, BO_Add, BF_One, BF_One, BO_Add, BF_One, BF_One>::GetRHI();
-			GraphicsPSOInit.RasterizerState = TStaticRasterizerState<FM_Solid, CM_None>::GetRHI();
-			GraphicsPSOInit.DepthStencilState = TStaticDepthStencilState<false, CF_Always>::GetRHI();
-
-			GraphicsPSOInit.BoundShaderState.VertexDeclarationRHI = GFilterVertexDeclaration.VertexDeclarationRHI;
-			GraphicsPSOInit.BoundShaderState.VertexShaderRHI = GETSAFERHISHADER_VERTEX(*VertexShader);
-			GraphicsPSOInit.BoundShaderState.PixelShaderRHI = GETSAFERHISHADER_PIXEL(*PixelShader);
-			GraphicsPSOInit.PrimitiveType = PT_TriangleList;
-			SetGraphicsPipelineState(RHICmdList, GraphicsPSOInit);
-
-			SetShaderParameters(RHICmdList, *PixelShader, PixelShader->GetPixelShader(), *PassParameters);
-
-			RHICmdList.SetViewport(View.ViewRect.Min.X, View.ViewRect.Min.Y, 0.0f, View.ViewRect.Max.X, View.ViewRect.Max.Y, 1.0f);
-
-			DrawRectangle(
-				RHICmdList,
-				0, 0,
-				View.ViewRect.Width(), View.ViewRect.Height(),
-				View.ViewRect.Min.X, View.ViewRect.Min.Y,
-				View.ViewRect.Width(), View.ViewRect.Height(),
-				FIntPoint(View.ViewRect.Width(), View.ViewRect.Height()),
-				SceneContext.GetBufferSizeXY(),
-				*VertexShader
-			);
-		}
-	);
-	GraphBuilder.Execute();
 }
-
 #endif // RHI_RAYTRACING
