Fixed: iOS First Person Template body lags behind viewport movement

diff --git a/Engine/Source/Runtime/Renderer/Private/MobileShadingRenderer.cpp b/Engine/Source/Runtime/Renderer/Private/MobileShadingRenderer.cpp
--- a/Engine/Source/Runtime/Renderer/Private/MobileShadingRenderer.cpp
+++ b/Engine/Source/Runtime/Renderer/Private/MobileShadingRenderer.cpp
@@ -308,6 +308,8 @@ void FMobileSceneRenderer::Render(FRHICommandListImmediate& RHICmdList)
 {
 	RHICmdList.SetCurrentStat(GET_STATID(STAT_CLMM_SceneStart));
 
+	Scene->UpdateAllPrimitiveSceneInfos(RHICmdList);
+
 	PrepareViewRectsForRendering();
 
 	if (ShouldRenderSkyAtmosphere(Scene, ViewFamily.EngineShowFlags))
