Merge of fix to lens flares sampling from wrong downsample stage.

diff --git a/Engine/Source/Runtime/Renderer/Private/PostProcess/PostProcessing.cpp b/Engine/Source/Runtime/Renderer/Private/PostProcess/PostProcessing.cpp
--- a/Engine/Source/Runtime/Renderer/Private/PostProcess/PostProcessing.cpp
+++ b/Engine/Source/Runtime/Renderer/Private/PostProcess/PostProcessing.cpp
@@ -473,7 +473,7 @@ FBloomOutputs AddBloomPass(
 		}
 
 		// The quality level controls which downsample stage we use as the flare input texture.
-		const uint32 LensFlareDownsampleStageIndex = static_cast<uint32>(ELensFlareQuality::High) - static_cast<uint32>(LensFlareQuality);
+		const uint32 LensFlareDownsampleStageIndex = static_cast<uint32>(ELensFlareQuality::MAX) - static_cast<uint32>(LensFlareQuality) - 1;
 
 		FLensFlareInputs LensFlareInputs;
 		LensFlareInputs.BloomTexture = BloomOutputTexture;
