DRPS - Fix deprecation warnings
#rnx

diff --git a/Engine/Source/Runtime/Renderer/Private/PostProcess/PostProcessBloomSetup.cpp b/Engine/Source/Runtime/Renderer/Private/PostProcess/PostProcessBloomSetup.cpp
--- a/Engine/Source/Runtime/Renderer/Private/PostProcess/PostProcessBloomSetup.cpp
+++ b/Engine/Source/Runtime/Renderer/Private/PostProcess/PostProcessBloomSetup.cpp
@@ -248,8 +248,8 @@ FBloomOutputs AddBloomPass(FRDGBuilder& GraphBuilder, const FViewInfo& View, con
 				6  // Q5
 			};
 
-			static_assert(ARRAY_COUNT(BloomStages) == BloomQualityCountMax, "Array must be one less than the number of bloom quality entries.");
-			static_assert(ARRAY_COUNT(BloomQualityToSceneDownsampleStage) == BloomQualityCountMax, "Array must be one less than the number of bloom quality entries.");
+			static_assert(UE_ARRAY_COUNT(BloomStages) == BloomQualityCountMax, "Array must be one less than the number of bloom quality entries.");
+			static_assert(UE_ARRAY_COUNT(BloomQualityToSceneDownsampleStage) == BloomQualityCountMax, "Array must be one less than the number of bloom quality entries.");
 
 			// Use bloom quality to select the number of downsample stages to use for bloom.
 			const uint32 BloomStageCount = BloomQualityToSceneDownsampleStage[BloomQualityIndex];
diff --git a/Engine/Source/Runtime/Renderer/Private/PostProcess/PostProcessDownsample.cpp b/Engine/Source/Runtime/Renderer/Private/PostProcess/PostProcessDownsample.cpp
--- a/Engine/Source/Runtime/Renderer/Private/PostProcess/PostProcessDownsample.cpp
+++ b/Engine/Source/Runtime/Renderer/Private/PostProcess/PostProcessDownsample.cpp
@@ -203,7 +203,7 @@ void FSceneDownsampleChain::Init(
 		TEXT("Scene(1/32)"),
 		TEXT("Scene(1/64)")
 	};
-	static_assert(ARRAY_COUNT(PassNames) == StageCount, "PassNames size must equal StageCount");
+	static_assert(UE_ARRAY_COUNT(PassNames) == StageCount, "PassNames size must equal StageCount");
 
 	// The first stage is the input.
 	Textures[0] = HalfResolutionSceneColor;
diff --git a/Engine/Source/Runtime/Renderer/Private/PostProcess/PostProcessMaterial.cpp b/Engine/Source/Runtime/Renderer/Private/PostProcess/PostProcessMaterial.cpp
--- a/Engine/Source/Runtime/Renderer/Private/PostProcess/PostProcessMaterial.cpp
+++ b/Engine/Source/Runtime/Renderer/Private/PostProcess/PostProcessMaterial.cpp
@@ -689,7 +689,7 @@ FScreenPassTexture AddHighResolutionScreenshotMaskPass(
 		TEXT("CaptureRegionMaterial")
 	};
 
-	static_assert(ARRAY_COUNT(PassNames) == static_cast<uint32>(EPass::MAX), "Pass names array doesn't match pass enum");
+	static_assert(UE_ARRAY_COUNT(PassNames) == static_cast<uint32>(EPass::MAX), "Pass names array doesn't match pass enum");
 
 	TOverridePassSequence<EPass> PassSequence(Inputs.OverrideOutput);
 	PassSequence.SetEnabled(EPass::Material, Inputs.Material != nullptr);
diff --git a/Engine/Source/Runtime/Renderer/Private/PostProcess/PostProcessUpscale.cpp b/Engine/Source/Runtime/Renderer/Private/PostProcess/PostProcessUpscale.cpp
--- a/Engine/Source/Runtime/Renderer/Private/PostProcess/PostProcessUpscale.cpp
+++ b/Engine/Source/Runtime/Renderer/Private/PostProcess/PostProcessUpscale.cpp
@@ -221,7 +221,7 @@ FScreenPassTexture AddUpscalePass(FRDGBuilder& GraphBuilder, const FViewInfo& Vi
 	TShaderMapRef<FUpscalePS> PixelShader(View.ShaderMap, PixelPermutationVector);
 
 	const TCHAR* const StageNames[] = { TEXT("PrimaryToSecondary"), TEXT("PrimaryToOutput"), TEXT("SecondaryToOutput") };
-	static_assert(ARRAY_COUNT(StageNames) == static_cast<uint32>(EUpscaleStage::MAX), "StageNames does not match EUpscaleStage");
+	static_assert(UE_ARRAY_COUNT(StageNames) == static_cast<uint32>(EUpscaleStage::MAX), "StageNames does not match EUpscaleStage");
 	const TCHAR* StageName = StageNames[static_cast<uint32>(Inputs.Stage)];
 
 	GraphBuilder.AddPass(
diff --git a/Engine/Source/Runtime/Renderer/Private/PostProcess/PostProcessing.cpp b/Engine/Source/Runtime/Renderer/Private/PostProcess/PostProcessing.cpp
--- a/Engine/Source/Runtime/Renderer/Private/PostProcess/PostProcessing.cpp
+++ b/Engine/Source/Runtime/Renderer/Private/PostProcess/PostProcessing.cpp
@@ -286,10 +286,10 @@ void AddPostProcessingPasses(FRDGBuilder& GraphBuilder, const FViewInfo& View, c
 		TEXT("SecondaryUpscale")
 	};
 
-	static_assert(static_cast<uint32>(EPass::MAX) == ARRAY_COUNT(PassNames), "EPass does not match PassNames.");
+	static_assert(static_cast<uint32>(EPass::MAX) == UE_ARRAY_COUNT(PassNames), "EPass does not match PassNames.");
 
 	TOverridePassSequence<EPass> PassSequence(ViewFamilyOutput);
-	PassSequence.SetNames(PassNames, ARRAY_COUNT(PassNames));
+	PassSequence.SetNames(PassNames, UE_ARRAY_COUNT(PassNames));
 	PassSequence.SetEnabled(EPass::VisualizeStationaryLightOverlap, EngineShowFlags.StationaryLightOverlap);
 	PassSequence.SetEnabled(EPass::VisualizeLightCulling, EngineShowFlags.VisualizeLightCulling);
 #if WITH_EDITOR
@@ -901,10 +901,10 @@ void AddDebugPostProcessingPasses(FRDGBuilder& GraphBuilder, const FViewInfo& Vi
 		TEXT("SecondaryUpscale")
 	};
 
-	static_assert(static_cast<uint32>(EPass::MAX) == ARRAY_COUNT(PassNames), "EPass does not match PassNames.");
+	static_assert(static_cast<uint32>(EPass::MAX) == UE_ARRAY_COUNT(PassNames), "EPass does not match PassNames.");
 
 	TOverridePassSequence<EPass> PassSequence(ViewFamilyOutput);
-	PassSequence.SetNames(PassNames, ARRAY_COUNT(PassNames));
+	PassSequence.SetNames(PassNames, UE_ARRAY_COUNT(PassNames));
 	PassSequence.SetEnabled(EPass::Visualize, true);
 	PassSequence.SetEnabled(EPass::TonemapAfter, bTonemapAfter);
 	PassSequence.SetEnabled(EPass::SelectionOutline, GIsEditor);
