Merged fix for VR into renderplat-staging.

diff --git a/Engine/Source/Runtime/Renderer/Private/PostProcess/SceneRenderTargets.cpp b/Engine/Source/Runtime/Renderer/Private/PostProcess/SceneRenderTargets.cpp
--- a/Engine/Source/Runtime/Renderer/Private/PostProcess/SceneRenderTargets.cpp
+++ b/Engine/Source/Runtime/Renderer/Private/PostProcess/SceneRenderTargets.cpp
@@ -2260,15 +2260,19 @@ void FSceneRenderTargets::AllocateCommonDepthTargets(FRHICommandList& RHICmdList
 		if (SceneDepthZ && bHMDAllocatedDepthTarget)
 		{
 			const uint32 OldElementSize = SceneDepthZ->ComputeMemorySize();
-		
-			// If SRT and texture are different (MSAA), only modify the resolve render target, to avoid creating a swapchain of MSAA textures
-			if (SceneDepthZ->GetRenderTargetItem().ShaderResourceTexture == SceneDepthZ->GetRenderTargetItem().TargetableTexture)
-			{
-				SceneDepthZ->GetRenderTargetItem().ShaderResourceTexture = SceneDepthZ->GetRenderTargetItem().TargetableTexture = SRTex;
-			}
-			else
+
 			{
-				SceneDepthZ->GetRenderTargetItem().ShaderResourceTexture = SRTex;
+				FSceneRenderTargetItem& Item = SceneDepthZ->GetRenderTargetItem();
+
+				// If SRT and texture are different (MSAA), only modify the resolve render target, to avoid creating a swapchain of MSAA textures
+				if (Item.ShaderResourceTexture == Item.TargetableTexture)
+				{
+					Item.TargetableTexture = SRTex;
+				}
+
+				Item.ShaderResourceTexture = SRTex;
+				Item.SRVs.Empty();
+				Item.HTileSRV = nullptr;
 			}
 
 			GRenderTargetPool.UpdateElementSize(SceneDepthZ, OldElementSize);
