Fixed null pointer dereference.

diff --git a/Engine/Source/Runtime/Renderer/Private/PostProcess/PostProcessUpscale.cpp b/Engine/Source/Runtime/Renderer/Private/PostProcess/PostProcessUpscale.cpp
--- a/Engine/Source/Runtime/Renderer/Private/PostProcess/PostProcessUpscale.cpp
+++ b/Engine/Source/Runtime/Renderer/Private/PostProcess/PostProcessUpscale.cpp
@@ -238,7 +238,7 @@ FScreenPassTexture AddUpscalePass(FRDGBuilder& GraphBuilder, const FViewInfo& Vi
 		{
 			TShaderMapRef<FUpscaleVS> TypedVertexShader(View.ShaderMap);
 			SetScreenPassPipelineState(RHICmdList, FScreenPassPipelineState(*TypedVertexShader, *PixelShader));
-			SetShaderParameters(RHICmdList, *TypedVertexShader, VertexShader->GetVertexShader(), *PassParameters);
+			SetShaderParameters(RHICmdList, *TypedVertexShader, TypedVertexShader->GetVertexShader(), *PassParameters);
 			VertexShader = *TypedVertexShader;
 		}
 		else
