Merged fixes for EngineTest.

diff --git a/Engine/Shaders/Private/PostProcessHistogram.usf b/Engine/Shaders/Private/PostProcessHistogram.usf
--- a/Engine/Shaders/Private/PostProcessHistogram.usf
+++ b/Engine/Shaders/Private/PostProcessHistogram.usf
@@ -72,7 +72,10 @@ void MainCS(
 				float Weight0 = 1.0f - Weight1;
 			
 				// Accumulate the weight to the nearby history buckets.
-				SharedHistogram[Bucket0][GroupThreadId.x][GroupThreadId.y] += Weight0;
+				if (Bucket0 != 0) // Ignore entirely black pixels.
+				{
+					SharedHistogram[Bucket0][GroupThreadId.x][GroupThreadId.y] += Weight0;
+				}
 				SharedHistogram[Bucket1][GroupThreadId.x][GroupThreadId.y] += Weight1;
 			}
 		}
diff --git a/Engine/Shaders/Private/PostProcessMaterialShaders.usf b/Engine/Shaders/Private/PostProcessMaterialShaders.usf
--- a/Engine/Shaders/Private/PostProcessMaterialShaders.usf
+++ b/Engine/Shaders/Private/PostProcessMaterialShaders.usf
@@ -200,7 +200,7 @@ void MainPS(
 #endif
 	OutColor = float4(GetMaterialEmissive(PixelMaterialInputs), Alpha );
 
-#if !MATERIALBLENDING_MODULATE && USE_PREEXPOSURE && !POST_PROCESS_MATERIAL_AFTER_TAA_UPSAMPLE
+#if !MATERIALBLENDING_MODULATE && USE_PREEXPOSURE && !POST_PROCESS_MATERIAL_BEFORE_TONEMAP
 	OutColor.xyz *= View.PreExposure;
 #endif
 } 
diff --git a/Engine/Source/Runtime/Renderer/Private/PostProcess/PostProcessMaterial.cpp b/Engine/Source/Runtime/Renderer/Private/PostProcess/PostProcessMaterial.cpp
--- a/Engine/Source/Runtime/Renderer/Private/PostProcess/PostProcessMaterial.cpp
+++ b/Engine/Source/Runtime/Renderer/Private/PostProcess/PostProcessMaterial.cpp
@@ -211,6 +211,9 @@ public:
 		FMaterialShader::ModifyCompilationEnvironment(Parameters, OutEnvironment);
 		OutEnvironment.SetDefine(TEXT("POST_PROCESS_MATERIAL"), 1);
 
+		EBlendableLocation Location = EBlendableLocation(Parameters.Material->GetBlendableLocation());
+		OutEnvironment.SetDefine(TEXT("POST_PROCESS_MATERIAL_BEFORE_TONEMAP"), (Location == BL_AfterTonemapping || Location == BL_ReplacingTonemapper) ? 0 : 1);
+
 		FPermutationDomain PermutationVector(Parameters.PermutationId);
 		if (PermutationVector.Get<FMobileDimension>())
 		{
