Actually fixed black screenshots in EngineTest

diff --git a/Engine/Source/Runtime/Renderer/Private/PostProcess/PostProcessing.cpp b/Engine/Source/Runtime/Renderer/Private/PostProcess/PostProcessing.cpp
--- a/Engine/Source/Runtime/Renderer/Private/PostProcess/PostProcessing.cpp
+++ b/Engine/Source/Runtime/Renderer/Private/PostProcess/PostProcessing.cpp
@@ -623,6 +623,8 @@ void AddPostProcessingPasses(FRDGBuilder& GraphBuilder, const FViewInfo& View, c
 		{
 			FPostProcessMaterialInputs PassInputs = GetPostProcessMaterialInputs(SceneColor);
 			PassSequence.AcceptOverrideIfLastPass(EPass::PostProcessMaterialAfterTonemapping, PassInputs.OverrideOutput);
+			PassInputs.SetInput(EPostProcessMaterialInput::PreTonemapHDRColor, SceneColorBeforeTonemap);
+			PassInputs.SetInput(EPostProcessMaterialInput::PostTonemapHDRColor, SceneColorAfterTonemap);
 
 			SceneColor = AddPostProcessMaterialChain(GraphBuilder, View, PassInputs, PostProcessMaterialAfterTonemappingChain);
 		}
