Merged translucent velocity implementation.

diff --git a/Engine/Source/Runtime/Engine/Classes/Materials/Material.h b/Engine/Source/Runtime/Engine/Classes/Materials/Material.h
--- a/Engine/Source/Runtime/Engine/Classes/Materials/Material.h
+++ b/Engine/Source/Runtime/Engine/Classes/Materials/Material.h
@@ -411,10 +411,7 @@ private:
 	TEnumAsByte<enum EMaterialShadingModel> ShadingModel; 
 
 public:
-	/**
-	* If true, translucent materials will cast dynamic shadows according to their opacity.
-	* OpacityMaskClipValue is used as the threshold value.
-	*/
+	/** Whether the material should cast shadows as masked even though it has a translucent blend mode. */
 	UPROPERTY(EditAnywhere, Category = Material, AdvancedDisplay)
 	uint8 bCastDynamicShadowAsMasked : 1;
 
@@ -430,8 +427,12 @@ private:
 
 public:
 
-	/** If BlendMode is BLEND_Masked, the surface is not rendered where OpacityMask < OpacityMaskClipValue. */
-	UPROPERTY(EditAnywhere, Category=Material, AdvancedDisplay)
+	/**
+	 * If BlendMode is BLEND_Masked, the surface is not rendered where OpacityMask < OpacityMaskClipValue.
+	 * If BlendMode is BLEND_Translucent, BLEND_Additive, or BLEND_Modulate, and "Output Velocity" is enabled,
+	 * the object velocity is not rendered where Opacity < OpacityMaskClipValue.
+	 */
+	UPROPERTY(EditAnywhere, Category = Material, AdvancedDisplay)
 	float OpacityMaskClipValue;
 
 	/** Adds to world position in the vertex shader. */
diff --git a/Engine/Source/Runtime/Engine/Private/Materials/Material.cpp b/Engine/Source/Runtime/Engine/Private/Materials/Material.cpp
--- a/Engine/Source/Runtime/Engine/Private/Materials/Material.cpp
+++ b/Engine/Source/Runtime/Engine/Private/Materials/Material.cpp
@@ -4301,7 +4301,8 @@ bool UMaterial::CanEditChange(const UProperty* InProperty) const
 			|| PropertyName == GET_MEMBER_NAME_STRING_CHECKED(UMaterial, bContactShadows)
 			|| PropertyName == GET_MEMBER_NAME_STRING_CHECKED(UMaterial, bDisableDepthTest)
 			|| PropertyName == GET_MEMBER_NAME_STRING_CHECKED(UMaterial, bUseTranslucencyVertexFog)
-			|| PropertyName == GET_MEMBER_NAME_STRING_CHECKED(UMaterial, bComputeFogPerPixel))
+			|| PropertyName == GET_MEMBER_NAME_STRING_CHECKED(UMaterial, bComputeFogPerPixel)
+			|| PropertyName == GET_MEMBER_NAME_STRING_CHECKED(UMaterial, bOutputTranslucentVelocity))
 		{
 			return MaterialDomain != MD_DeferredDecal && MaterialDomain != MD_RuntimeVirtualTexture && IsTranslucentBlendMode(BlendMode);
 		}
diff --git a/Engine/Source/Runtime/Engine/Private/Materials/MaterialInterface.cpp b/Engine/Source/Runtime/Engine/Private/Materials/MaterialInterface.cpp
--- a/Engine/Source/Runtime/Engine/Private/Materials/MaterialInterface.cpp
+++ b/Engine/Source/Runtime/Engine/Private/Materials/MaterialInterface.cpp
@@ -40,7 +40,7 @@ void FMaterialRelevance::SetPrimitiveViewRelevance(FPrimitiveViewRelevance& OutV
 {
 	OutViewRelevance.bOpaqueRelevance = bOpaque;
 	OutViewRelevance.bMaskedRelevance = bMasked;
-	OutViewRelevance.bOutputsTranslucentVelocityRelevance = bOutputsTranslucentVelocity;
+	OutViewRelevance.bTranslucentVelocityRelevance = bOutputsTranslucentVelocity;
 	OutViewRelevance.bDistortionRelevance = bDistortion;
 	OutViewRelevance.bHairStrandsRelevance = bHairStrands;
 	OutViewRelevance.bSeparateTranslucencyRelevance = bSeparateTranslucency;
diff --git a/Engine/Source/Runtime/Engine/Private/SceneManagement.cpp b/Engine/Source/Runtime/Engine/Private/SceneManagement.cpp
--- a/Engine/Source/Runtime/Engine/Private/SceneManagement.cpp
+++ b/Engine/Source/Runtime/Engine/Private/SceneManagement.cpp
@@ -198,7 +198,6 @@ FMeshBatchAndRelevance::FMeshBatchAndRelevance(const FMeshBatch& InMesh, const F
 	EBlendMode BlendMode = Material->GetBlendMode();
 	bHasOpaqueMaterial = (BlendMode == BLEND_Opaque);
 	bHasMaskedMaterial = (BlendMode == BLEND_Masked);
-	bHasTranslucentMaterialWithVelocity = Material->IsTranslucencyWritingVelocity();
 	bRenderInMainPass = PrimitiveSceneProxy->ShouldRenderInMainPass();
 }
 
diff --git a/Engine/Source/Runtime/Engine/Public/PrimitiveViewRelevance.h b/Engine/Source/Runtime/Engine/Public/PrimitiveViewRelevance.h
--- a/Engine/Source/Runtime/Engine/Public/PrimitiveViewRelevance.h
+++ b/Engine/Source/Runtime/Engine/Public/PrimitiveViewRelevance.h
@@ -23,7 +23,7 @@ struct FPrimitiveViewRelevance
 	/** The primitive has one or more masked elements. */
 	uint32 bMaskedRelevance : 1;
 	/** The primitive has one or more translucent elements which output velocity. */
-	uint32 bOutputsTranslucentVelocityRelevance : 1;
+	uint32 bTranslucentVelocityRelevance : 1;
 	/** The primitive has one or more distortion elements. */
 	uint32 bDistortionRelevance : 1;
 	/** The primitive has one or more elements that have SeparateTranslucency. */
@@ -100,6 +100,11 @@ struct FPrimitiveViewRelevance
 		return bSeparateTranslucencyRelevance || bNormalTranslucencyRelevance;
 	}
 
+	bool HasVelocity() const
+	{
+		return bVelocityRelevance || bTranslucentVelocityRelevance;
+	}
+
 	/** Default constructor */
 	FPrimitiveViewRelevance()
 	{
diff --git a/Engine/Source/Runtime/Engine/Public/SceneManagement.h b/Engine/Source/Runtime/Engine/Public/SceneManagement.h
--- a/Engine/Source/Runtime/Engine/Public/SceneManagement.h
+++ b/Engine/Source/Runtime/Engine/Public/SceneManagement.h
@@ -2049,7 +2049,6 @@ private:
 	 */
 	uint32 bHasOpaqueMaterial : 1;
 	uint32 bHasMaskedMaterial : 1;
-	uint32 bHasTranslucentMaterialWithVelocity : 1;
 	uint32 bRenderInMainPass : 1;
 
 public:
@@ -2058,7 +2057,6 @@ public:
 	bool GetHasOpaqueMaterial() const { return bHasOpaqueMaterial; }
 	bool GetHasMaskedMaterial() const { return bHasMaskedMaterial; }
 	bool GetHasOpaqueOrMaskedMaterial() const { return bHasOpaqueMaterial || bHasMaskedMaterial; }
-	bool GetHasTranslucentMaterialWithVelocity() const { return bHasTranslucentMaterialWithVelocity; }
 	bool GetRenderInMainPass() const { return bRenderInMainPass; }
 };
 
diff --git a/Engine/Source/Runtime/Renderer/Private/DeferredShadingRenderer.cpp b/Engine/Source/Runtime/Renderer/Private/DeferredShadingRenderer.cpp
--- a/Engine/Source/Runtime/Renderer/Private/DeferredShadingRenderer.cpp
+++ b/Engine/Source/Runtime/Renderer/Private/DeferredShadingRenderer.cpp
@@ -1711,9 +1711,12 @@ void FDeferredShadingSceneRenderer::Render(FRHICommandListImmediate& RHICmdList)
 	// If bBasePassCanOutputVelocity is set, basepass fully writes the velocity buffer unless bUseSelectiveBasePassOutputs is enabled.
 	if (bShouldRenderVelocities && (!bBasePassCanOutputVelocity || bUseSelectiveBasePassOutputs))
 	{
+		// We only need to clear if the base pass didn't already render velocities.
+		const bool bClearVelocityRT = !bBasePassCanOutputVelocity;
+
 		// Render the velocities of movable objects
 		RHICmdList.SetCurrentStat(GET_STATID(STAT_CLM_Velocity));
-		RenderVelocities(RHICmdList, SceneContext.SceneVelocity);
+		RenderVelocities(RHICmdList, SceneContext.SceneVelocity, EVelocityPass::Opaque, bClearVelocityRT);
 		RHICmdList.SetCurrentStat(GET_STATID(STAT_CLM_AfterVelocity));
 		ServiceLocalQueue();
 	}
@@ -2144,17 +2147,18 @@ void FDeferredShadingSceneRenderer::Render(FRHICommandListImmediate& RHICmdList)
 			RHICmdList.SetCurrentStat(GET_STATID(STAT_CLM_AfterTranslucency));
 		}
 
-		/*if (bShouldRenderVelocities)
+		if (bShouldRenderVelocities)
 		{
+			const bool bClearVelocityRT = false;
+
 			// Render the velocities of movable objects
 			RHICmdList.SetCurrentStat(GET_STATID(STAT_CLM_TranslucentVelocity));
-			RenderVelocities(RHICmdList, VelocityRT);
+			RenderVelocities(RHICmdList, SceneContext.SceneVelocity, EVelocityPass::Translucent, bClearVelocityRT);
 			RHICmdList.SetCurrentStat(GET_STATID(STAT_CLM_AfterTranslucentVelocity));
 			ServiceLocalQueue();
-		}*/
+		}
 
 		checkSlow(RHICmdList.IsOutsideRenderPass());
-
 	}
 
 	if (HairDatas)
diff --git a/Engine/Source/Runtime/Renderer/Private/DeferredShadingRenderer.h b/Engine/Source/Runtime/Renderer/Private/DeferredShadingRenderer.h
--- a/Engine/Source/Runtime/Renderer/Private/DeferredShadingRenderer.h
+++ b/Engine/Source/Runtime/Renderer/Private/DeferredShadingRenderer.h
@@ -20,6 +20,8 @@
 enum class ERayTracingPrimaryRaysFlag : uint32;
 #endif
 
+enum class EVelocityPass : uint32;
+
 class FSceneTextureParameters;
 class FDistanceFieldAOParameters;
 class UStaticMeshComponent;
@@ -314,11 +316,11 @@ private:
 	bool ShouldRenderVelocities() const;
 
 	/** Renders the velocities of movable objects for the motion blur effect. */
-	void RenderVelocities(FRHICommandListImmediate& RHICmdList, TRefCountPtr<IPooledRenderTarget>& VelocityRT);
+	void RenderVelocities(FRHICommandListImmediate& RHICmdList, TRefCountPtr<IPooledRenderTarget>& VelocityRT, EVelocityPass VelocityPass, bool bClearVelocityRT);
 
 	/** Renders the velocities of movable objects for the motion blur effect. */
-	void RenderVelocitiesInner(FRHICommandListImmediate& RHICmdList, TRefCountPtr<IPooledRenderTarget>& VelocityRT);
-	void RenderVelocitiesInnerParallel(FRHICommandListImmediate& RHICmdList, TRefCountPtr<IPooledRenderTarget>& VelocityRT);
+	void RenderVelocitiesInner(FRHICommandListImmediate& RHICmdList, TRefCountPtr<IPooledRenderTarget>& VelocityRT, EVelocityPass VelocityPass);
+	void RenderVelocitiesInnerParallel(FRHICommandListImmediate& RHICmdList, TRefCountPtr<IPooledRenderTarget>& VelocityRT, EVelocityPass VelocityPass);
 
 	/** Renders world-space lightmap density instead of the normal color. */
 	bool RenderLightMapDensities(FRHICommandListImmediate& RHICmdList);
diff --git a/Engine/Source/Runtime/Renderer/Private/SceneRendering.h b/Engine/Source/Runtime/Renderer/Private/SceneRendering.h
--- a/Engine/Source/Runtime/Renderer/Private/SceneRendering.h
+++ b/Engine/Source/Runtime/Renderer/Private/SceneRendering.h
@@ -151,17 +151,6 @@ public:
 	TArray<FProjectedShadowInfo*,SceneRenderingAllocator> OccludedPerObjectShadows;
 };
 
-enum class EVelocityPass : uint32
-{
-	// Renders a separate velocity pass for opaques.
-	Opaque = 0,
-
-	// Renders a separate velocity / depth pass for translucency AFTER the translucent pass.
-	Translucency,
-
-	Count
-};
-
 // Stores the primitive count of each translucency pass (redundant, could be computed after sorting but this way we touch less memory)
 struct FTranslucenyPrimCount
 {
diff --git a/Engine/Source/Runtime/Renderer/Private/SceneVisibility.cpp b/Engine/Source/Runtime/Renderer/Private/SceneVisibility.cpp
--- a/Engine/Source/Runtime/Renderer/Private/SceneVisibility.cpp
+++ b/Engine/Source/Runtime/Renderer/Private/SceneVisibility.cpp
@@ -2290,11 +2290,26 @@ struct FRelevancePacket
 								}
 #endif
 
-								if (ViewRelevance.bVelocityRelevance
-									&& FVelocityRendering::PrimitiveHasVelocity(View.GetFeatureLevel(), PrimitiveSceneInfo)
-									&& FVelocityRendering::PrimitiveHasVelocityForView(View, Bounds.BoxSphereBounds, PrimitiveSceneInfo))
+								if (ViewRelevance.HasVelocity())
 								{
-									DrawCommandPacket.AddCommandsForMesh(PrimitiveIndex, PrimitiveSceneInfo, StaticMeshRelevance, StaticMesh, Scene, bCanCache, EMeshPass::Velocity);
+									const FPrimitiveSceneProxy* PrimitiveSceneProxy = PrimitiveSceneInfo->Proxy;
+
+									if (FVelocityMeshProcessor::PrimitiveHasVelocityForView(View, PrimitiveSceneProxy))
+									{
+										if (ViewRelevance.bVelocityRelevance &&
+											FOpaqueVelocityMeshProcessor::PrimitiveCanHaveVelocity(View.GetShaderPlatform(), PrimitiveSceneProxy) &&
+											FOpaqueVelocityMeshProcessor::PrimitiveHasVelocityForFrame(PrimitiveSceneProxy))
+										{
+											DrawCommandPacket.AddCommandsForMesh(PrimitiveIndex, PrimitiveSceneInfo, StaticMeshRelevance, StaticMesh, Scene, bCanCache, EMeshPass::Velocity);
+										}
+
+										if (ViewRelevance.bTranslucentVelocityRelevance &&
+											FTranslucentVelocityMeshProcessor::PrimitiveCanHaveVelocity(View.GetShaderPlatform(), PrimitiveSceneProxy) &&
+											FTranslucentVelocityMeshProcessor::PrimitiveHasVelocityForFrame(PrimitiveSceneProxy))
+										{
+											DrawCommandPacket.AddCommandsForMesh(PrimitiveIndex, PrimitiveSceneInfo, StaticMeshRelevance, StaticMesh, Scene, bCanCache, EMeshPass::TranslucentVelocity);
+										}
+									}
 								}
 
 								++NumVisibleStaticMeshElements;
@@ -2709,19 +2724,23 @@ void ComputeDynamicMeshRelevance(EShadingPath ShadingPath, bool bAddLightmapDens
 			}
 #endif
 
-			if (ViewRelevance.bVelocityRelevance
-				&& FVelocityRendering::PrimitiveHasVelocity(View.GetFeatureLevel(), PrimitiveSceneInfo)
-				&& FVelocityRendering::PrimitiveHasVelocityForView(View, Bounds.BoxSphereBounds, PrimitiveSceneInfo))
+			if (ViewRelevance.bVelocityRelevance)
 			{
 				PassMask.Set(EMeshPass::Velocity);
 				View.NumVisibleDynamicMeshElements[EMeshPass::Velocity] += NumElements;
 			}
-		}
 
-		if (ViewRelevance.bUsesSingleLayerWaterMaterial)
-		{
-			PassMask.Set(EMeshPass::SingleLayerWaterPass);
-			View.NumVisibleDynamicMeshElements[EMeshPass::SingleLayerWaterPass] += NumElements;
+			if (ViewRelevance.bTranslucentVelocityRelevance)
+			{
+				PassMask.Set(EMeshPass::TranslucentVelocity);
+				View.NumVisibleDynamicMeshElements[EMeshPass::TranslucentVelocity] += NumElements;
+			}
+
+			if (ViewRelevance.bUsesSingleLayerWaterMaterial)
+			{
+				PassMask.Set(EMeshPass::SingleLayerWaterPass);
+				View.NumVisibleDynamicMeshElements[EMeshPass::SingleLayerWaterPass] += NumElements;
+			}
 		}
 	}
 
diff --git a/Engine/Source/Runtime/Renderer/Private/VelocityRendering.cpp b/Engine/Source/Runtime/Renderer/Private/VelocityRendering.cpp
--- a/Engine/Source/Runtime/Renderer/Private/VelocityRendering.cpp
+++ b/Engine/Source/Runtime/Renderer/Private/VelocityRendering.cpp
@@ -1,9 +1,5 @@
 // Copyright 1998-2019 Epic Games, Inc. All Rights Reserved.
 
-/*=============================================================================
-	VelocityRendering.cpp: Velocity rendering implementation.
-=============================================================================*/
-
 #include "VelocityRendering.h"
 #include "SceneUtils.h"
 #include "Materials/Material.h"
@@ -54,147 +50,101 @@ bool IsParallelVelocity()
 	return GRHICommandList.UseParallelAlgorithms() && CVarParallelVelocity.GetValueOnRenderThread();
 }
 
-//=============================================================================
-/** Encapsulates the Velocity vertex shader. */
 class FVelocityVS : public FMeshMaterialShader
 {
-	DECLARE_SHADER_TYPE(FVelocityVS,MeshMaterial);
-
 public:
-
-	bool SupportsVelocity() const
-	{
-		return GPUSkinCachePreviousPositionBuffer.IsBound() ||
-			PrevTransformBuffer.IsBound() || 
-			(PrevTransform0.IsBound() && PrevTransform1.IsBound() && PrevTransform2.IsBound()) ||
-			//@todo MeshCommandPipeline - now that PreviousLocalToWorld is in the primitive uniform buffer, we can't look at whether the shader bound it to cull what gets rendered in velocity pass
-			true;
-	}
+	DECLARE_MESH_MATERIAL_SHADER(FVelocityVS);
 
 	static bool ShouldCompilePermutation(const FMeshMaterialShaderPermutationParameters& Parameters)
 	{
-		const EShaderPlatform Platform = Parameters.Platform;
 		const FMaterial* Material = Parameters.Material;
 
-		//Only compile the velocity shaders for the default material or if it's masked,
-		return ((Material->IsSpecialEngineMaterial() || !Material->WritesEveryPixel() 
-			//or if the material is opaque and two-sided,
-			|| (Material->IsTwoSided() && !IsTranslucentBlendMode(Material->GetBlendMode()))
-			// or if the material modifies meshes
-			|| Material->MaterialMayModifyMeshPosition()))
-			&& IsFeatureLevelSupported(Platform, ERHIFeatureLevel::SM5) 
-			&& !FVelocityRendering::VertexFactoryOnlyOutputsVelocityInBasePass(Platform, Parameters.VertexFactoryType->SupportsStaticLighting());
-	}
-
-protected:
-	FVelocityVS(const ShaderMetaType::CompiledShaderInitializerType& Initializer):
-		FMeshMaterialShader(Initializer)
-	{
-		GPUSkinCachePreviousPositionBuffer.Bind(Initializer.ParameterMap, TEXT("GPUSkinCachePreviousPositionBuffer"));
-		PrevTransform0.Bind(Initializer.ParameterMap, TEXT("PrevTransform0"));
-		PrevTransform1.Bind(Initializer.ParameterMap, TEXT("PrevTransform1"));
-		PrevTransform2.Bind(Initializer.ParameterMap, TEXT("PrevTransform2"));
-		PrevTransformBuffer.Bind(Initializer.ParameterMap, TEXT("PrevTransformBuffer"));
-		PassUniformBuffer.Bind(Initializer.ParameterMap, FSceneTexturesUniformParameters::StaticStructMetadata.GetShaderVariableName());
-	}
+		// Compile for default material.
+		const bool bIsDefault = Material->IsSpecialEngineMaterial();
 
-	FVelocityVS() {}
+		// Compile for masked materials.
+		const bool bIsMasked = !Material->WritesEveryPixel();
 
-	virtual bool Serialize(FArchive& Ar) override
-	{
-		bool bShaderHasOutdatedParameters = FMeshMaterialShader::Serialize(Ar);
+		// Compile for opaque and two-sided materials.
+		const bool bIsOpaqueAndTwoSided = (Material->IsTwoSided() && !IsTranslucentBlendMode(Material->GetBlendMode()));
 
-		Ar << GPUSkinCachePreviousPositionBuffer;
-		Ar << PrevTransform0;
-		Ar << PrevTransform1;
-		Ar << PrevTransform2;
-		Ar << PrevTransformBuffer;
+		// Compile for materials which modify meshes.
+		const bool bMayModifyMeshes = Material->MaterialMayModifyMeshPosition();
 
-		return bShaderHasOutdatedParameters;
+		// Compile if supported by the hardware.
+		const bool bIsFeatureSupported = IsFeatureLevelSupported(Parameters.Platform, ERHIFeatureLevel::SM5);
+
+		/**
+		 * Any material with a vertex factory incompatible with base pass velocity generation must generate
+		 * permutations for this shader. Shaders which don't fall into this set are considered "simple" enough
+		 * to swap against the default material. This massively simplifies the calculations.
+		 */
+		const bool bIsSeparateVelocityPassRequired = (bIsDefault || bIsMasked || bIsOpaqueAndTwoSided || bMayModifyMeshes) &&
+			FVelocityRendering::IsSeparateVelocityPassRequiredByVertexFactory(Parameters.Platform, Parameters.VertexFactoryType->SupportsStaticLighting());
+
+		// The material may explicitly override and request that it be rendered into the velocity pass.
+		const bool bIsSeparateVelocityPassRequiredByMaterial = Material->IsTranslucencyWritingVelocity();
+
+		return bIsFeatureSupported && (bIsSeparateVelocityPassRequired || bIsSeparateVelocityPassRequiredByMaterial);
 	}
 
-private:
-	FShaderResourceParameter GPUSkinCachePreviousPositionBuffer;
-	FShaderParameter PrevTransform0;
-	FShaderParameter PrevTransform1;
-	FShaderParameter PrevTransform2;
-	FShaderResourceParameter PrevTransformBuffer;
+	FVelocityVS() = default;
+	FVelocityVS(const ShaderMetaType::CompiledShaderInitializerType& Initializer)
+		: FMeshMaterialShader(Initializer)
+	{
+		PassUniformBuffer.Bind(Initializer.ParameterMap, FSceneTexturesUniformParameters::StaticStructMetadata.GetShaderVariableName());
+	}
 };
 
-
-//=============================================================================
-/** Encapsulates the Velocity hull shader. */
 class FVelocityHS : public FBaseHS
 {
-	DECLARE_SHADER_TYPE(FVelocityHS,MeshMaterial);
-
-protected:
-	FVelocityHS(const ShaderMetaType::CompiledShaderInitializerType& Initializer):
-		FBaseHS(Initializer)
-	{}
-
-	FVelocityHS() {}
+public:
+	DECLARE_MESH_MATERIAL_SHADER(FVelocityHS);
 
 	static bool ShouldCompilePermutation(const FMeshMaterialShaderPermutationParameters& Parameters)
 	{
-		return FBaseHS::ShouldCompilePermutation(Parameters) &&
-				FVelocityVS::ShouldCompilePermutation(Parameters); // same rules as VS
+		return FBaseHS::ShouldCompilePermutation(Parameters) && FVelocityVS::ShouldCompilePermutation(Parameters);
 	}
+
+	FVelocityHS() = default;
+	FVelocityHS(const ShaderMetaType::CompiledShaderInitializerType& Initializer)
+		: FBaseHS(Initializer)
+	{}
 };
 
-//=============================================================================
-/** Encapsulates the Velocity domain shader. */
 class FVelocityDS : public FBaseDS
 {
-	DECLARE_SHADER_TYPE(FVelocityDS,MeshMaterial);
-
-protected:
-	FVelocityDS(const ShaderMetaType::CompiledShaderInitializerType& Initializer):
-		FBaseDS(Initializer)
-	{}
-
-	FVelocityDS() {}
+public:
+	DECLARE_MESH_MATERIAL_SHADER(FVelocityDS);
 
 	static bool ShouldCompilePermutation(const FMeshMaterialShaderPermutationParameters& Parameters)
 	{
-		return FBaseDS::ShouldCompilePermutation(Parameters) &&
-				FVelocityVS::ShouldCompilePermutation(Parameters); // same rules as VS
+		return FBaseDS::ShouldCompilePermutation(Parameters) && FVelocityVS::ShouldCompilePermutation(Parameters);
 	}
-};
 
-IMPLEMENT_MATERIAL_SHADER_TYPE(,FVelocityVS,TEXT("/Engine/Private/VelocityShader.usf"),TEXT("MainVertexShader"),SF_Vertex); 
-IMPLEMENT_MATERIAL_SHADER_TYPE(,FVelocityHS,TEXT("/Engine/Private/VelocityShader.usf"),TEXT("MainHull"),SF_Hull); 
-IMPLEMENT_MATERIAL_SHADER_TYPE(,FVelocityDS,TEXT("/Engine/Private/VelocityShader.usf"),TEXT("MainDomain"),SF_Domain);
+	FVelocityDS() = default;
+	FVelocityDS(const ShaderMetaType::CompiledShaderInitializerType& Initializer)
+		: FBaseDS(Initializer)
+	{}
+};
 
-//=============================================================================
-/** Encapsulates the Velocity pixel shader. */
 class FVelocityPS : public FMeshMaterialShader
 {
-	DECLARE_SHADER_TYPE(FVelocityPS,MeshMaterial);
 public:
+	DECLARE_MESH_MATERIAL_SHADER(FVelocityPS);
+
 	static bool ShouldCompilePermutation(const FMeshMaterialShaderPermutationParameters& Parameters)
 	{
-		const EShaderPlatform Platform = Parameters.Platform;
-		const FMaterial* Material = Parameters.Material;
-
-		//Only compile the velocity shaders for the default material or if it's masked,
-		return ((Material->IsSpecialEngineMaterial() || !Material->WritesEveryPixel() 
-			//or if the material is opaque and two-sided,
-			|| (Material->IsTwoSided() && !IsTranslucentBlendMode(Material->GetBlendMode()))
-			// or if the material modifies meshes
-			|| Material->MaterialMayModifyMeshPosition()))
-			&& IsFeatureLevelSupported(Platform, ERHIFeatureLevel::SM5) && 
-			!FVelocityRendering::VertexFactoryOnlyOutputsVelocityInBasePass(Platform, Parameters.VertexFactoryType->SupportsStaticLighting());
+		return FVelocityVS::ShouldCompilePermutation(Parameters);
 	}
 
-	static void ModifyCompilationEnvironment(const FMaterialShaderPermutationParameters& Parameters, FShaderCompilerEnvironment& OutEnvironment )
+	static void ModifyCompilationEnvironment(const FMaterialShaderPermutationParameters& Parameters, FShaderCompilerEnvironment& OutEnvironment)
 	{
 		FMaterialShader::ModifyCompilationEnvironment(Parameters, OutEnvironment);
 		OutEnvironment.SetRenderTargetOutputFormat(0, PF_G16R16);
 	}
 
-	FVelocityPS() {}
-
+	FVelocityPS() = default;
 	FVelocityPS(const ShaderMetaType::CompiledShaderInitializerType& Initializer)
 		: FMeshMaterialShader(Initializer)
 	{
@@ -202,11 +152,30 @@ public:
 	}
 };
 
-IMPLEMENT_MATERIAL_SHADER_TYPE(,FVelocityPS,TEXT("/Engine/Private/VelocityShader.usf"),TEXT("MainPixelShader"),SF_Pixel);
-
+IMPLEMENT_MESH_MATERIAL_SHADER(FVelocityVS, "/Engine/Private/VelocityShader.usf", "MainVertexShader", SF_Vertex); 
+IMPLEMENT_MESH_MATERIAL_SHADER(FVelocityHS, "/Engine/Private/VelocityShader.usf", "MainHull", SF_Hull); 
+IMPLEMENT_MESH_MATERIAL_SHADER(FVelocityDS, "/Engine/Private/VelocityShader.usf", "MainDomain", SF_Domain);
+IMPLEMENT_MESH_MATERIAL_SHADER(FVelocityPS, "/Engine/Private/VelocityShader.usf", "MainPixelShader", SF_Pixel);
 IMPLEMENT_SHADERPIPELINE_TYPE_VSPS(VelocityPipeline, FVelocityVS, FVelocityPS, true);
 
-static void BeginVelocityRendering(FRHICommandList& RHICmdList, TRefCountPtr<IPooledRenderTarget>& VelocityRT, bool bPerformClear)
+EMeshPass::Type GetMeshPassFromVelocityPass(EVelocityPass VelocityPass)
+{
+	switch (VelocityPass)
+	{
+	case EVelocityPass::Opaque:
+		return EMeshPass::Velocity;
+	case EVelocityPass::Translucent:
+		return EMeshPass::TranslucentVelocity;
+	}
+	check(false);
+	return EMeshPass::Velocity;
+}
+
+static void BeginVelocityRendering(
+	FRHICommandList& RHICmdList,
+	TRefCountPtr<IPooledRenderTarget>& VelocityRT,
+	EVelocityPass VelocityPass,
+	bool bPerformClear)
 {
 	check(RHICmdList.IsOutsideRenderPass());
 
@@ -216,7 +185,7 @@ static void BeginVelocityRendering(FRHICommandList& RHICmdList, TRefCountPtr<IPo
 	FRHIRenderPassInfo RPInfo(VelocityTexture, ERenderTargetActions::Load_Store);
 	RPInfo.DepthStencilRenderTarget.Action = MakeDepthStencilTargetActions(ERenderTargetActions::Load_Store, ERenderTargetActions::Load_Store);
 	RPInfo.DepthStencilRenderTarget.DepthStencilTarget = DepthTexture;
-	RPInfo.DepthStencilRenderTarget.ExclusiveDepthStencil = FExclusiveDepthStencil::DepthRead_StencilWrite;
+	RPInfo.DepthStencilRenderTarget.ExclusiveDepthStencil = VelocityPass == EVelocityPass::Opaque ? FExclusiveDepthStencil::DepthRead_StencilWrite : FExclusiveDepthStencil::DepthWrite_StencilWrite;
 
 	if (bPerformClear)
 	{
@@ -233,7 +202,13 @@ static void BeginVelocityRendering(FRHICommandList& RHICmdList, TRefCountPtr<IPo
 	}
 }
 
-static void SetVelocitiesState(FRHICommandList& RHICmdList, const FViewInfo& View, const FSceneRenderer* SceneRender, FMeshPassProcessorRenderState& DrawRenderState, TRefCountPtr<IPooledRenderTarget>& VelocityRT)
+static void SetVelocitiesState(
+	FRHICommandList& RHICmdList,
+	const FViewInfo& View,
+	const FSceneRenderer* SceneRender,
+	FMeshPassProcessorRenderState& DrawRenderState,
+	TRefCountPtr<IPooledRenderTarget>& VelocityRT,
+	EVelocityPass VelocityPass)
 {
 	const FIntPoint BufferSize = FSceneRenderTargets::Get(RHICmdList).GetBufferSizeXY();
 	const FIntPoint VelocityBufferSize = BufferSize;		// full resolution so we can reuse the existing full res z buffer
@@ -269,17 +244,25 @@ static void SetVelocitiesState(FRHICommandList& RHICmdList, const FViewInfo& Vie
 	}
 
 	DrawRenderState.SetBlendState(TStaticBlendState<CW_RGBA>::GetRHI());
-	DrawRenderState.SetDepthStencilState(TStaticDepthStencilState<false, CF_DepthNearOrEqual>::GetRHI());
-	//TODO Where does this state go?
-	//RHICmdList.SetRasterizerState(GetStaticRasterizerState<true>(FM_Solid, CM_CW));
+
+	switch (VelocityPass)
+	{
+	case EVelocityPass::Opaque:
+		DrawRenderState.SetDepthStencilState(TStaticDepthStencilState<false, CF_DepthNearOrEqual>::GetRHI());
+		break;
+
+	case EVelocityPass::Translucent:
+		DrawRenderState.SetDepthStencilState(TStaticDepthStencilState<true, CF_DepthNearOrEqual>::GetRHI());
+		break;
+	}
 }
 
 DECLARE_CYCLE_STAT(TEXT("Velocity"), STAT_CLP_Velocity, STATGROUP_ParallelCommandListMarkers);
 
-
 class FVelocityPassParallelCommandListSet : public FParallelCommandListSet
 {
 	TRefCountPtr<IPooledRenderTarget>& VelocityRT;
+	EVelocityPass VelocityPass;
 
 public:
 	FVelocityPassParallelCommandListSet(
@@ -289,9 +272,11 @@ public:
 		bool bInParallelExecute,
 		bool bInCreateSceneContext,
 		const FMeshPassProcessorRenderState& InDrawRenderState,
-		TRefCountPtr<IPooledRenderTarget>& InVelocityRT)
+		TRefCountPtr<IPooledRenderTarget>& InVelocityRT,
+		EVelocityPass InVelocityPass)
 		: FParallelCommandListSet(GET_STATID(STAT_CLP_Velocity), InView, InSceneRenderer, InParentCmdList, bInParallelExecute, bInCreateSceneContext, InDrawRenderState)
 		, VelocityRT(InVelocityRT)
+		, VelocityPass(InVelocityPass)
 	{
 	}
 
@@ -303,8 +288,8 @@ public:
 	virtual void SetStateOnCommandList(FRHICommandList& CmdList) override
 	{
 		FParallelCommandListSet::SetStateOnCommandList(CmdList);
-		BeginVelocityRendering(CmdList, VelocityRT, false);
-		SetVelocitiesState(CmdList, View, SceneRenderer, DrawRenderState, VelocityRT);
+		BeginVelocityRendering(CmdList, VelocityRT, VelocityPass, false);
+		SetVelocitiesState(CmdList, View, SceneRenderer, DrawRenderState, VelocityRT, VelocityPass);
 	}
 };
 
@@ -313,7 +298,7 @@ static TAutoConsoleVariable<int32> CVarRHICmdFlushRenderThreadTasksVelocityPass(
 	0,
 	TEXT("Wait for completion of parallel render thread tasks at the end of the velocity pass.  A more granular version of r.RHICmdFlushRenderThreadTasks. If either r.RHICmdFlushRenderThreadTasks or r.RHICmdFlushRenderThreadTasksVelocityPass is > 0 we will flush."));
 
-void FDeferredShadingSceneRenderer::RenderVelocitiesInnerParallel(FRHICommandListImmediate& RHICmdList, TRefCountPtr<IPooledRenderTarget>& VelocityRT)
+void FDeferredShadingSceneRenderer::RenderVelocitiesInnerParallel(FRHICommandListImmediate& RHICmdList, TRefCountPtr<IPooledRenderTarget>& VelocityRT, EVelocityPass VelocityPass)
 {
 	// Parallel rendering requires its own renderpasses so we cannot have an active one at this point
 	check(RHICmdList.IsOutsideRenderPass());
@@ -343,15 +328,17 @@ void FDeferredShadingSceneRenderer::RenderVelocitiesInnerParallel(FRHICommandLis
 				CVarRHICmdVelocityPassDeferredContexts.GetValueOnRenderThread() > 0,
 				CVarRHICmdFlushRenderThreadTasksVelocityPass.GetValueOnRenderThread() == 0 && CVarRHICmdFlushRenderThreadTasks.GetValueOnRenderThread() == 0,
 				DrawRenderState,
-				VelocityRT);
+				VelocityRT,
+				VelocityPass);
 
-			// Draw velocities.
-			View.ParallelMeshDrawCommandPasses[EMeshPass::Velocity].DispatchDraw(&ParallelCommandListSet, RHICmdList);
+			const EMeshPass::Type MeshPass = GetMeshPassFromVelocityPass(VelocityPass);
+
+			View.ParallelMeshDrawCommandPasses[MeshPass].DispatchDraw(&ParallelCommandListSet, RHICmdList);
 		}
 	}
 }
 
-void FDeferredShadingSceneRenderer::RenderVelocitiesInner(FRHICommandListImmediate& RHICmdList, TRefCountPtr<IPooledRenderTarget>& VelocityRT)
+void FDeferredShadingSceneRenderer::RenderVelocitiesInner(FRHICommandListImmediate& RHICmdList, TRefCountPtr<IPooledRenderTarget>& VelocityRT, EVelocityPass VelocityPass)
 {
 	check(RHICmdList.IsInsideRenderPass());
 	for(int32 ViewIndex = 0; ViewIndex < Views.Num(); ViewIndex++)
@@ -371,16 +358,18 @@ void FDeferredShadingSceneRenderer::RenderVelocitiesInner(FRHICommandListImmedia
 
 			Scene->UniformBuffers.UpdateViewUniformBuffer(View);
 
-			SetVelocitiesState(RHICmdList, View, this, DrawRenderState, VelocityRT);
+			SetVelocitiesState(RHICmdList, View, this, DrawRenderState, VelocityRT, VelocityPass);
+
+			const EMeshPass::Type MeshPass = GetMeshPassFromVelocityPass(VelocityPass);
 
-			View.ParallelMeshDrawCommandPasses[EMeshPass::Velocity].DispatchDraw(nullptr, RHICmdList);
+			View.ParallelMeshDrawCommandPasses[MeshPass].DispatchDraw(nullptr, RHICmdList);
 		}
 	}
 }
 
 bool FDeferredShadingSceneRenderer::ShouldRenderVelocities() const
 {
-	if (!GPixelFormats[PF_G16R16].Supported)
+	if (!FVelocityRendering::IsSeparateVelocityPassSupported())
 	{
 		return false;
 	}
@@ -405,7 +394,7 @@ bool FDeferredShadingSceneRenderer::ShouldRenderVelocities() const
 	return bNeedsVelocity;
 }
 
-void FDeferredShadingSceneRenderer::RenderVelocities(FRHICommandListImmediate& RHICmdList, TRefCountPtr<IPooledRenderTarget>& VelocityRT)
+void FDeferredShadingSceneRenderer::RenderVelocities(FRHICommandListImmediate& RHICmdList, TRefCountPtr<IPooledRenderTarget>& VelocityRT, EVelocityPass VelocityPass, bool bClearVelocityRT)
 {
 	SCOPED_NAMED_EVENT(FDeferredShadingSceneRenderer_RenderVelocities, FColor::Emerald);
 
@@ -421,8 +410,11 @@ void FDeferredShadingSceneRenderer::RenderVelocities(FRHICommandListImmediate& R
 	SCOPED_DRAW_EVENT(RHICmdList, RenderVelocities);
 	SCOPED_GPU_STAT(RHICmdList, RenderVelocities);
 
-	FPooledRenderTargetDesc Desc = FVelocityRendering::GetRenderTargetDesc();
-	GRenderTargetPool.FindFreeElement(RHICmdList, Desc, VelocityRT, TEXT("Velocity"));
+	if (!VelocityRT)
+	{
+		FPooledRenderTargetDesc Desc = FVelocityRendering::GetRenderTargetDesc();
+		GRenderTargetPool.FindFreeElement(RHICmdList, Desc, VelocityRT, TEXT("Velocity"));
+	}
 
 	{
 		static const auto MotionBlurDebugVar = IConsoleManager::Get().FindTConsoleVariableDataInt(TEXT("r.MotionBlurDebug"));
@@ -433,10 +425,7 @@ void FDeferredShadingSceneRenderer::RenderVelocities(FRHICommandListImmediate& R
 		}
 	}
 
-	{
-		// In this case, basepass also outputs some of the velocities, so append is already started, and don't clear the buffer.
-		BeginVelocityRendering(RHICmdList, VelocityRT, !FVelocityRendering::BasePassCanOutputVelocity(FeatureLevel));
-	}
+	BeginVelocityRendering(RHICmdList, VelocityRT, VelocityPass, bClearVelocityRT);
 
 	{
 		if (IsParallelVelocity())
@@ -445,11 +434,11 @@ void FDeferredShadingSceneRenderer::RenderVelocities(FRHICommandListImmediate& R
 			RHICmdList.EndRenderPass();
 
 			// Now do parallel encoding.
-			RenderVelocitiesInnerParallel(RHICmdList, VelocityRT);
+			RenderVelocitiesInnerParallel(RHICmdList, VelocityRT, VelocityPass);
 		}
 		else
 		{
-			RenderVelocitiesInner(RHICmdList, VelocityRT);
+			RenderVelocitiesInner(RHICmdList, VelocityRT, VelocityPass);
 			RHICmdList.EndRenderPass();
 		}
 
@@ -467,6 +456,11 @@ FPooledRenderTargetDesc FVelocityRendering::GetRenderTargetDesc()
 	return FPooledRenderTargetDesc(FPooledRenderTargetDesc::Create2DDesc(VelocityBufferSize, PF_G16R16, FClearValueBinding::Transparent, TexCreate_None, TexCreate_RenderTargetable | TexCreate_UAV | TexCreate_ShaderResource, false));
 }
 
+bool FVelocityRendering::IsSeparateVelocityPassSupported()
+{
+	return GPixelFormats[PF_G16R16].Supported;
+}
+
 bool FVelocityRendering::BasePassCanOutputVelocity(EShaderPlatform ShaderPlatform)
 {
 	return IsUsingBasePassVelocity(ShaderPlatform);
@@ -478,21 +472,34 @@ bool FVelocityRendering::BasePassCanOutputVelocity(ERHIFeatureLevel::Type Featur
 	return BasePassCanOutputVelocity(ShaderPlatform);
 }
 
-bool FVelocityRendering::VertexFactoryOnlyOutputsVelocityInBasePass(EShaderPlatform ShaderPlatform, bool bVertexFactorySupportsStaticLighting)
+bool FVelocityRendering::IsSeparateVelocityPassRequiredByVertexFactory(EShaderPlatform ShaderPlatform, bool bVertexFactoryUsesStaticLighting)
 {
-	return BasePassCanOutputVelocity(ShaderPlatform) && !(IsUsingSelectiveBasePassOutputs(ShaderPlatform) && bVertexFactorySupportsStaticLighting);;
+	// A separate pass is required if the base pass can't do it.
+	const bool bBasePassVelocityNotSupported = !BasePassCanOutputVelocity(ShaderPlatform);
+
+	// Meshes with static lighting need a separate velocity pass, but only if we are using selective render target outputs.
+	const bool bVertexFactoryRequiresSeparateVelocityPass = IsUsingSelectiveBasePassOutputs(ShaderPlatform) && bVertexFactoryUsesStaticLighting;
+
+	return bBasePassVelocityNotSupported || bVertexFactoryRequiresSeparateVelocityPass;
 }
 
-bool FVelocityRendering::PrimitiveHasVelocity(ERHIFeatureLevel::Type FeatureLevel, const FPrimitiveSceneInfo* PrimitiveSceneInfo)
+bool FVelocityMeshProcessor::PrimitiveHasVelocityForView(const FViewInfo& View, const FPrimitiveSceneProxy* PrimitiveSceneProxy)
 {
-	// No velocity if motionblur is off, or if it's a non-moving object (treat as background in that case).
-	if (!GPixelFormats[PF_G16R16].Supported || !PrimitiveSceneInfo->Proxy->IsMovable())
+	// Skip camera cuts which effectively reset velocity for the new frame.
+	if (View.bCameraCut && !View.PreviousViewTransform.IsSet())
 	{
 		return false;
 	}
 
-	// If the base pass is allowed to render velocity in the GBuffer, only mesh with static lighting need the velocity pass.
-	if (FVelocityRendering::BasePassCanOutputVelocity(FeatureLevel) && !(IsUsingSelectiveBasePassOutputs(GetFeatureLevelShaderPlatform(FeatureLevel)) && PrimitiveSceneInfo->Proxy->HasStaticLighting()))
+	const FBoxSphereBounds& PrimitiveBounds = PrimitiveSceneProxy->GetBounds();
+	const float LODFactorDistanceSquared = (PrimitiveBounds.Origin - View.ViewMatrices.GetViewOrigin()).SizeSquared() * FMath::Square(View.LODDistanceFactor);
+
+	// The minimum projected screen radius for a primitive to be drawn in the velocity pass, as a fraction of half the horizontal screen width (likely to be 0.08f)
+	float MinScreenRadiusForVelocityPass = View.FinalPostProcessSettings.MotionBlurPerObjectSize * (2.0f / 100.0f);
+	float MinScreenRadiusForVelocityPassSquared = FMath::Square(MinScreenRadiusForVelocityPass);
+
+	// Skip primitives that only cover a small amount of screen space, motion blur on them won't be noticeable.
+	if (FMath::Square(PrimitiveBounds.SphereRadius) <= MinScreenRadiusForVelocityPassSquared * LODFactorDistanceSquared)
 	{
 		return false;
 	}
@@ -500,37 +507,42 @@ bool FVelocityRendering::PrimitiveHasVelocity(ERHIFeatureLevel::Type FeatureLeve
 	return true;
 }
 
-bool FVelocityRendering::PrimitiveHasVelocityForView(const FViewInfo& View, const FBoxSphereBounds& Bounds, const FPrimitiveSceneInfo* PrimitiveSceneInfo)
+bool FOpaqueVelocityMeshProcessor::PrimitiveCanHaveVelocity(EShaderPlatform ShaderPlatform, const FPrimitiveSceneProxy* PrimitiveSceneProxy)
 {
-	// Skip camera cuts which effectively reset velocity for the new frame.
-	bool bRawCameraCut = View.bCameraCut && !View.PreviousViewTransform.IsSet();
-	if (bRawCameraCut)
+	if (!FVelocityRendering::IsSeparateVelocityPassSupported())
 	{
 		return false;
 	}
 
-	const float LODFactorDistanceSquared = (Bounds.Origin - View.ViewMatrices.GetViewOrigin()).SizeSquared() * FMath::Square(View.LODDistanceFactor);
-
-	// The minimum projected screen radius for a primitive to be drawn in the velocity pass, as a fraction of half the horizontal screen width (likely to be 0.08f)
-	float MinScreenRadiusForVelocityPass = View.FinalPostProcessSettings.MotionBlurPerObjectSize * (2.0f / 100.0f);
-	float MinScreenRadiusForVelocityPassSquared = FMath::Square(MinScreenRadiusForVelocityPass);
-
-	// Skip primitives that only cover a small amount of screenspace, motion blur on them won't be noticeable.
-	if (FMath::Square(Bounds.SphereRadius) <= MinScreenRadiusForVelocityPassSquared * LODFactorDistanceSquared)
+	if (!PrimitiveSceneProxy->IsMovable())
 	{
 		return false;
 	}
 
-	if (PrimitiveSceneInfo->Proxy->AlwaysHasVelocity())
+	/**
+	 * Whether the vertex factory for this primitive requires that it render in the separate velocity pass, as opposed to the base pass.
+	 * In cases where the base pass is rendering opaque velocity for a particular mesh batch, we want to filter it out from this pass,
+	 * which performs a separate draw call to render velocity.
+	 */
+	const bool bIsSeparateVelocityPassRequiredByVertexFactory =
+		FVelocityRendering::IsSeparateVelocityPassRequiredByVertexFactory(ShaderPlatform, PrimitiveSceneProxy->HasStaticLighting());
+
+	if (!bIsSeparateVelocityPassRequiredByVertexFactory)
 	{
-		return true;
+		return false;
 	}
 
-	// check if the primitive has moved
+	return true;
+}
+
+bool FOpaqueVelocityMeshProcessor::PrimitiveHasVelocityForFrame(const FPrimitiveSceneProxy* PrimitiveSceneProxy)
+{
+	if (!PrimitiveSceneProxy->AlwaysHasVelocity())
 	{
+		// Check if the primitive has moved.
+		const FPrimitiveSceneInfo* PrimitiveSceneInfo = PrimitiveSceneProxy->GetPrimitiveSceneInfo();
 		const FScene* Scene = PrimitiveSceneInfo->Scene;
-
-		const FMatrix& LocalToWorld = PrimitiveSceneInfo->Proxy->GetLocalToWorld();
+		const FMatrix& LocalToWorld = PrimitiveSceneProxy->GetLocalToWorld();
 		FMatrix PreviousLocalToWorld = LocalToWorld;
 		Scene->VelocityData.GetComponentPreviousLocalToWorld(PrimitiveSceneInfo->PrimitiveComponentId, PreviousLocalToWorld);
 
@@ -544,54 +556,122 @@ bool FVelocityRendering::PrimitiveHasVelocityForView(const FViewInfo& View, cons
 	return true;
 }
 
-
-void FVelocityMeshProcessor::AddMeshBatch(const FMeshBatch& RESTRICT MeshBatch, uint64 BatchElementMask, const FPrimitiveSceneProxy* RESTRICT PrimitiveSceneProxy, int32 StaticMeshId)
+void FOpaqueVelocityMeshProcessor::AddMeshBatch(
+	const FMeshBatch& RESTRICT MeshBatch,
+	uint64 BatchElementMask,
+	const FPrimitiveSceneProxy* RESTRICT PrimitiveSceneProxy,
+	int32 StaticMeshId)
 {
-	bool bRequiresSeparateVelocity = false;
+	const EShaderPlatform ShaderPlatform = GetFeatureLevelShaderPlatform(FeatureLevel);
 
-	if (FVelocityRendering::PrimitiveHasVelocity(FeatureLevel, PrimitiveSceneProxy->GetPrimitiveSceneInfo()))
+	if (!PrimitiveCanHaveVelocity(ShaderPlatform, PrimitiveSceneProxy))
 	{
-		bRequiresSeparateVelocity = true;
+		return;
+	}
 
-		// Cached mesh commands have identical check inside MarkRelevant.
-		if (ViewIfDynamicMeshCommand)
+	if (ViewIfDynamicMeshCommand)
+	{
+		if (!PrimitiveHasVelocityForFrame(PrimitiveSceneProxy))
 		{
-			checkSlow(ViewIfDynamicMeshCommand->bIsViewInfo);
-			FViewInfo* ViewInfo = (FViewInfo*)ViewIfDynamicMeshCommand;
+			return;
+		}
 
-			bRequiresSeparateVelocity = FVelocityRendering::PrimitiveHasVelocityForView(*ViewInfo, PrimitiveSceneProxy->GetBounds(), PrimitiveSceneProxy->GetPrimitiveSceneInfo());
+		checkSlow(ViewIfDynamicMeshCommand->bIsViewInfo);
+		FViewInfo* ViewInfo = (FViewInfo*)ViewIfDynamicMeshCommand;
+
+		if (!PrimitiveHasVelocityForView(*ViewInfo, PrimitiveSceneProxy))
+		{
+			return;
 		}
 	}
 
-	if (MeshBatch.bUseForMaterial && bRequiresSeparateVelocity)
+	const FMaterialRenderProxy* MaterialRenderProxy = MeshBatch.MaterialRenderProxy;
+	const FMaterial* Material = &MaterialRenderProxy->GetMaterialWithFallback(FeatureLevel, MaterialRenderProxy);
+	const EBlendMode BlendMode = Material->GetBlendMode();
+	const bool bIsNotTranslucent = BlendMode == BLEND_Opaque || BlendMode == BLEND_Masked;
+
+	if (MeshBatch.bUseForMaterial && bIsNotTranslucent && ShouldIncludeMaterialInDefaultOpaquePass(*Material))
 	{
-		// Determine the mesh's material and blend mode.
-		const FMaterialRenderProxy* MaterialRenderProxy = nullptr;
-		const FMaterial* Material = &MeshBatch.MaterialRenderProxy->GetMaterialWithFallback(FeatureLevel, MaterialRenderProxy);
-		const EBlendMode BlendMode = Material->GetBlendMode();
+		// This is specifically done *before* the material swap, as swapped materials may have different fill / cull modes.
 		const ERasterizerFillMode MeshFillMode = ComputeMeshFillMode(MeshBatch, *Material);
 		const ERasterizerCullMode MeshCullMode = ComputeMeshCullMode(MeshBatch, *Material);
 
-		if ((BlendMode == BLEND_Opaque || BlendMode == BLEND_Masked)
-			&& ShouldIncludeMaterialInDefaultOpaquePass(*Material)) // No velocity for water as of now
+		/**
+		 * Materials without masking or custom vertex modifications can be swapped out
+		 * for the default material, which simplifies the shader. However, the default
+		 * material also does not support being two-sided.
+		 */
+		const bool bSwapWithDefaultMaterial = Material->WritesEveryPixel() && !Material->IsTwoSided() && !Material->MaterialModifiesMeshPosition_RenderThread();
+
+		if (bSwapWithDefaultMaterial)
 		{
-			if (Material->WritesEveryPixel() && !Material->IsTwoSided() && !Material->MaterialModifiesMeshPosition_RenderThread())
-			{
-				// Default material doesn't handle masked or mesh-mod, and doesn't have the correct bIsTwoSided setting.
-				MaterialRenderProxy = UMaterial::GetDefaultMaterial(MD_Surface)->GetRenderProxy();
-				Material = MaterialRenderProxy->GetMaterial(FeatureLevel);
-			}
+			MaterialRenderProxy = UMaterial::GetDefaultMaterial(MD_Surface)->GetRenderProxy();
+			Material = MaterialRenderProxy->GetMaterial(FeatureLevel);
+		}
+
+		check(Material && MaterialRenderProxy);
 
-			if (!MaterialRenderProxy)
-			{
-				MaterialRenderProxy = MeshBatch.MaterialRenderProxy;
-			}
+		Process(MeshBatch, BatchElementMask, StaticMeshId, PrimitiveSceneProxy, *MaterialRenderProxy, *Material, MeshFillMode, MeshCullMode);
+	}
+}
+
+bool FTranslucentVelocityMeshProcessor::PrimitiveCanHaveVelocity(EShaderPlatform ShaderPlatform, const FPrimitiveSceneProxy* PrimitiveSceneProxy)
+{
+	/**
+	 * Velocity for translucency is always relevant because the pass also writes depth.
+	 * Therefore, the primitive can't be filtered based on motion, or it will break post
+	 * effects like depth of field which rely on depth information.
+	 */
+	return FVelocityRendering::IsSeparateVelocityPassSupported();
+}
 
-			check(Material && MaterialRenderProxy);
+bool FTranslucentVelocityMeshProcessor::PrimitiveHasVelocityForFrame(const FPrimitiveSceneProxy* PrimitiveSceneProxy)
+{
+	return true;
+}
 
-			Process(MeshBatch, BatchElementMask, StaticMeshId, PrimitiveSceneProxy, *MaterialRenderProxy, *Material, MeshFillMode, MeshCullMode);
+void FTranslucentVelocityMeshProcessor::AddMeshBatch(
+	const FMeshBatch& RESTRICT MeshBatch,
+	uint64 BatchElementMask,
+	const FPrimitiveSceneProxy* RESTRICT PrimitiveSceneProxy,
+	int32 StaticMeshId)
+{
+	const EShaderPlatform ShaderPlatform = GetFeatureLevelShaderPlatform(FeatureLevel);
+
+	if (!PrimitiveCanHaveVelocity(ShaderPlatform, PrimitiveSceneProxy))
+	{
+		return;
+	}
+
+	if (ViewIfDynamicMeshCommand)
+	{
+		if (!PrimitiveHasVelocityForFrame(PrimitiveSceneProxy))
+		{
+			return;
+		}
+
+		checkSlow(ViewIfDynamicMeshCommand->bIsViewInfo);
+		FViewInfo* ViewInfo = (FViewInfo*)ViewIfDynamicMeshCommand;
+
+		if (!PrimitiveHasVelocityForView(*ViewInfo, PrimitiveSceneProxy))
+		{
+			return;
 		}
 	}
+
+	const FMaterialRenderProxy* MaterialRenderProxy = MeshBatch.MaterialRenderProxy;
+	const FMaterial* Material = &MeshBatch.MaterialRenderProxy->GetMaterialWithFallback(FeatureLevel, MaterialRenderProxy);
+
+	// Whether the primitive is marked to write translucent velocity / depth.
+	const bool bMaterialWritesVelocity = Material->IsTranslucencyWritingVelocity();
+
+	if (MeshBatch.bUseForMaterial && bMaterialWritesVelocity)
+	{
+		const ERasterizerFillMode MeshFillMode = ComputeMeshFillMode(MeshBatch, *Material);
+		const ERasterizerCullMode MeshCullMode = ComputeMeshCullMode(MeshBatch, *Material);
+
+		Process(MeshBatch, BatchElementMask, StaticMeshId, PrimitiveSceneProxy, *MaterialRenderProxy, *Material, MeshFillMode, MeshCullMode);
+	}
 }
 
 void GetVelocityPassShaders(
@@ -690,13 +770,32 @@ FVelocityMeshProcessor::FVelocityMeshProcessor(const FScene* Scene, const FScene
 	PassDrawRenderState.SetPassUniformBuffer(Scene->UniformBuffers.VelocityPassUniformBuffer);
 }
 
+FOpaqueVelocityMeshProcessor::FOpaqueVelocityMeshProcessor(const FScene* Scene, const FSceneView* InViewIfDynamicMeshCommand, const FMeshPassProcessorRenderState& InPassDrawRenderState, FMeshPassDrawListContext* InDrawListContext)
+	: FVelocityMeshProcessor(Scene, InViewIfDynamicMeshCommand, InPassDrawRenderState, InDrawListContext)
+{}
+
 FMeshPassProcessor* CreateVelocityPassProcessor(const FScene* Scene, const FSceneView* InViewIfDynamicMeshCommand, FMeshPassDrawListContext* InDrawListContext)
 {
 	FMeshPassProcessorRenderState VelocityPassState;
 	VelocityPassState.SetBlendState(TStaticBlendState<CW_RGBA>::GetRHI());
 	VelocityPassState.SetDepthStencilState(TStaticDepthStencilState<false, CF_DepthNearOrEqual>::GetRHI());
 
-	return new(FMemStack::Get()) FVelocityMeshProcessor(Scene, InViewIfDynamicMeshCommand, VelocityPassState, InDrawListContext);
+	return new(FMemStack::Get()) FOpaqueVelocityMeshProcessor(Scene, InViewIfDynamicMeshCommand, VelocityPassState, InDrawListContext);
+}
+
+FRegisterPassProcessorCreateFunction RegisterVelocityPass(&CreateVelocityPassProcessor, EShadingPath::Deferred,  EMeshPass::Velocity, EMeshPassFlags::CachedMeshCommands | EMeshPassFlags::MainView);
+
+FTranslucentVelocityMeshProcessor::FTranslucentVelocityMeshProcessor(const FScene* Scene, const FSceneView* InViewIfDynamicMeshCommand, const FMeshPassProcessorRenderState& InPassDrawRenderState, FMeshPassDrawListContext* InDrawListContext)
+	: FVelocityMeshProcessor(Scene, InViewIfDynamicMeshCommand, InPassDrawRenderState, InDrawListContext)
+{}
+
+FMeshPassProcessor* CreateTranslucentVelocityPassProcessor(const FScene* Scene, const FSceneView* InViewIfDynamicMeshCommand, FMeshPassDrawListContext* InDrawListContext)
+{
+	FMeshPassProcessorRenderState VelocityPassState;
+	VelocityPassState.SetBlendState(TStaticBlendState<CW_RGBA>::GetRHI());
+	VelocityPassState.SetDepthStencilState(TStaticDepthStencilState<true, CF_DepthNearOrEqual>::GetRHI());
+
+	return new(FMemStack::Get()) FTranslucentVelocityMeshProcessor(Scene, InViewIfDynamicMeshCommand, VelocityPassState, InDrawListContext);
 }
 
-FRegisterPassProcessorCreateFunction RegisterVelocityPass(&CreateVelocityPassProcessor, EShadingPath::Deferred,  EMeshPass::Velocity, EMeshPassFlags::CachedMeshCommands | EMeshPassFlags::MainView);
\ No newline at end of file
+FRegisterPassProcessorCreateFunction RegisterTranslucentVelocityPass(&CreateTranslucentVelocityPassProcessor, EShadingPath::Deferred, EMeshPass::TranslucentVelocity, EMeshPassFlags::CachedMeshCommands | EMeshPassFlags::MainView);
\ No newline at end of file
diff --git a/Engine/Source/Runtime/Renderer/Private/VelocityRendering.h b/Engine/Source/Runtime/Renderer/Private/VelocityRendering.h
--- a/Engine/Source/Runtime/Renderer/Private/VelocityRendering.h
+++ b/Engine/Source/Runtime/Renderer/Private/VelocityRendering.h
@@ -1,9 +1,5 @@
 // Copyright 1998-2019 Epic Games, Inc. All Rights Reserved.
 
-/*=============================================================================
-	VelocityRendering.h: Velocity rendering definitions.
-=============================================================================*/
-
 #pragma once
 
 #include "CoreMinimal.h"
@@ -18,39 +14,57 @@ class FScene;
 class FStaticMeshBatch;
 class FViewInfo;
 
+struct FPrimitiveViewRelevance;
+
+enum class EVelocityPass : uint32
+{
+	// Renders a separate velocity pass for opaques.
+	Opaque = 0,
+
+	// Renders a separate velocity / depth pass for translucency AFTER the translucent pass.
+	Translucent,
+
+	Count
+};
+
+EMeshPass::Type GetMeshPassFromVelocityPass(EVelocityPass VelocityPass);
+
 // Group Velocity Rendering accessors, types, etc.
 struct FVelocityRendering
 {
 	static FPooledRenderTargetDesc GetRenderTargetDesc();
 
-	/** Returns true the velocity can be outputted in the BasePass. */
+	/** Returns true if the separate velocity pass is enabled. */
+	static bool IsSeparateVelocityPassSupported();
+
+	/** Returns true if the velocity can be output in the BasePass. */
 	static bool BasePassCanOutputVelocity(EShaderPlatform ShaderPlatform);
 
-	/** Returns true the velocity can be outputted in the BasePass. Only valid for the current platform. */
+	/** Returns true if the velocity can be output in the BasePass. Only valid for the current platform. */
 	static bool BasePassCanOutputVelocity(ERHIFeatureLevel::Type FeatureLevel);
 
-	/** Returns true the velocity is output in the BasePass. */
-	static bool VertexFactoryOnlyOutputsVelocityInBasePass(EShaderPlatform ShaderPlatform, bool bVertexFactorySupportsStaticLighting);
-
-	/** Returns true if the object needs to be rendered in the velocity pass (is not moving like the world, needed for motionblur and TemporalAA). */
-	static bool PrimitiveHasVelocity(ERHIFeatureLevel::Type FeatureLevel, const FPrimitiveSceneInfo* PrimitiveSceneInfo);
-
-	/** Returns true if the object needs to be rendered in the velocity pass for a given view. */
-	static bool PrimitiveHasVelocityForView(const FViewInfo& View, const FBoxSphereBounds& Bounds, const FPrimitiveSceneInfo* PrimitiveSceneInfo);
+	/** Returns true if a separate velocity pass is required (i.e. not rendered by the base pass) given the provided vertex factory settings. */
+	static bool IsSeparateVelocityPassRequiredByVertexFactory(EShaderPlatform ShaderPlatform, bool bVertexFactoryUsesStaticLighting);
 };
 
-
+/**
+ * Base velocity mesh pass processor class. Used for both opaque and translucent velocity passes.
+ */
 class FVelocityMeshProcessor : public FMeshPassProcessor
 {
 public:
-
-	FVelocityMeshProcessor(const FScene* Scene, const FSceneView* InViewIfDynamicMeshCommand, const FMeshPassProcessorRenderState& InPassDrawRenderState, FMeshPassDrawListContext* InDrawListContext);
-
-	virtual void AddMeshBatch(const FMeshBatch& RESTRICT MeshBatch, uint64 BatchElementMask, const FPrimitiveSceneProxy* RESTRICT PrimitiveSceneProxy, int32 StaticMeshId = -1) override final;
+	FVelocityMeshProcessor(
+		const FScene* Scene,
+		const FSceneView* InViewIfDynamicMeshCommand,
+		const FMeshPassProcessorRenderState& InPassDrawRenderState,
+		FMeshPassDrawListContext* InDrawListContext);
 
 	FMeshPassProcessorRenderState PassDrawRenderState;
 
-private:
+	/* Checks whether the primitive should emit velocity for the current view by comparing screen space size against a threshold. */
+	static bool PrimitiveHasVelocityForView(const FViewInfo& View, const FPrimitiveSceneProxy* PrimitiveSceneProxy);
+
+protected:
 	void Process(
 		const FMeshBatch& MeshBatch,
 		uint64 BatchElementMask,
@@ -60,4 +74,57 @@ private:
 		const FMaterial& RESTRICT MaterialResource,
 		ERasterizerFillMode MeshFillMode,
 		ERasterizerCullMode MeshCullMode);
+};
+
+/**
+ * Velocity pass processor for rendering opaques into a separate velocity pass (i.e. separate from the base pass).
+ */
+class FOpaqueVelocityMeshProcessor : public FVelocityMeshProcessor
+{
+public:
+	FOpaqueVelocityMeshProcessor(
+		const FScene* Scene,
+		const FSceneView* InViewIfDynamicMeshCommand,
+		const FMeshPassProcessorRenderState& InPassDrawRenderState,
+		FMeshPassDrawListContext* InDrawListContext);
+
+	/** Returns true if the object is capable of having velocity for any frame. */
+	static bool PrimitiveCanHaveVelocity(EShaderPlatform ShaderPlatform, const FPrimitiveSceneProxy* PrimitiveSceneProxy);
+
+	/** Returns true if the primitive has velocity for the current frame. */
+	static bool PrimitiveHasVelocityForFrame(const FPrimitiveSceneProxy* PrimitiveSceneProxy);
+
+private:
+	void AddMeshBatch(
+		const FMeshBatch& RESTRICT MeshBatch,
+		uint64 BatchElementMask,
+		const FPrimitiveSceneProxy* RESTRICT PrimitiveSceneProxy,
+		int32 StaticMeshId = -1) override final;
+};
+
+/**
+ * Velocity pass processor for rendering translucent object velocity and depth. This pass is rendered AFTER the
+ * translucent pass so that depth can safely be written.
+ */
+class FTranslucentVelocityMeshProcessor : public FVelocityMeshProcessor
+{
+public:
+	FTranslucentVelocityMeshProcessor(
+		const FScene* Scene,
+		const FSceneView* InViewIfDynamicMeshCommand,
+		const FMeshPassProcessorRenderState& InPassDrawRenderState,
+		FMeshPassDrawListContext* InDrawListContext);
+
+	/** Returns true if the object is capable of having velocity for any frame. */
+	static bool PrimitiveCanHaveVelocity(EShaderPlatform ShaderPlatform, const FPrimitiveSceneProxy* PrimitiveSceneProxy);
+
+	/** Returns true if the primitive has velocity for the current frame. */
+	static bool PrimitiveHasVelocityForFrame(const FPrimitiveSceneProxy* PrimitiveSceneProxy);
+
+private:
+	void AddMeshBatch(
+		const FMeshBatch& RESTRICT MeshBatch,
+		uint64 BatchElementMask,
+		const FPrimitiveSceneProxy* RESTRICT PrimitiveSceneProxy,
+		int32 StaticMeshId = -1) override final;
 };
\ No newline at end of file
diff --git a/Engine/Source/Runtime/Renderer/Public/MeshPassProcessor.h b/Engine/Source/Runtime/Renderer/Public/MeshPassProcessor.h
--- a/Engine/Source/Runtime/Renderer/Public/MeshPassProcessor.h
+++ b/Engine/Source/Runtime/Renderer/Public/MeshPassProcessor.h
@@ -24,6 +24,7 @@ namespace EMeshPass
 		CSMShadowDepth,
 		Distortion,
 		Velocity,
+		TranslucentVelocity,
 		TranslucencyStandard,
 		TranslucencyAfterDOF,
 		TranslucencyAll, /** Drawing all translucency, regardless of separate or standard.  Used when drawing translucency outside of the main renderer, eg FRendererModule::DrawTile. */
@@ -57,6 +58,7 @@ inline const TCHAR* GetMeshPassName(EMeshPass::Type MeshPass)
 	case EMeshPass::CSMShadowDepth: return TEXT("CSMShadowDepth");
 	case EMeshPass::Distortion: return TEXT("Distortion");
 	case EMeshPass::Velocity: return TEXT("Velocity");
+	case EMeshPass::TranslucentVelocity: return TEXT("TranslucentVelocity");
 	case EMeshPass::TranslucencyStandard: return TEXT("TranslucencyStandard");
 	case EMeshPass::TranslucencyAfterDOF: return TEXT("TranslucencyAfterDOF");
 	case EMeshPass::TranslucencyAll: return TEXT("TranslucencyAll");
