Deletes experimental warning of SSGI

diff --git a/Engine/Source/Runtime/Renderer/Private/IndirectLightRendering.cpp b/Engine/Source/Runtime/Renderer/Private/IndirectLightRendering.cpp
--- a/Engine/Source/Runtime/Renderer/Private/IndirectLightRendering.cpp
+++ b/Engine/Source/Runtime/Renderer/Private/IndirectLightRendering.cpp
@@ -144,13 +144,6 @@ void FDeferredShadingSceneRenderer::RenderDiffuseIndirectAndAmbientOcclusion(FRH
 		}
 		else if (bApplySSGI)
 		{
-			static bool bWarnExperimental = false;
-			if (!bWarnExperimental)
-			{
-				UE_LOG(LogRenderer, Warning, TEXT("SSGI is experimental."));
-				bWarnExperimental = true;
-			}
-
 			RenderScreenSpaceDiffuseIndirect(GraphBuilder, SceneTextures, SceneColor, View, /* out */ &RayTracingConfig, /* out */ &DenoiserInputs);
 		}
 		else
