Fixed crash on mobile when using a material with "Enable Stencil Test" set to true.

Updated the code which detects when it's possible to skip custom depth rendering on mobile to take into account the new "Enable Stencil Test" material flag. Also fixed a couple of issues with the OpenGL RHI:
* the post processing material writes directly to the backbuffer on mobile, but it's not possible to use that with a custom depth/stencil buffer on OpenGL. Added a hack in PostProcessMaterial to render to an intermediate buffer when that's the case and copy to the backbuffer at the end.
* glCopyImageSubData() doesn't like copying from a 2D texture to a renderbuffer on mobile, so made it use glBlitFramebuffer() when the bind targets are different (which happens when copying from a render target to the backbuffer)

diff --git a/Engine/Source/Runtime/OpenGLDrv/Private/OpenGLDevice.cpp b/Engine/Source/Runtime/OpenGLDrv/Private/OpenGLDevice.cpp
--- a/Engine/Source/Runtime/OpenGLDrv/Private/OpenGLDevice.cpp
+++ b/Engine/Source/Runtime/OpenGLDrv/Private/OpenGLDevice.cpp
@@ -854,6 +854,9 @@ static void InitRHICapabilitiesForGL()
 	GRHISupportsInstancing = FOpenGL::SupportsInstancing(); // HTML5 supports it with ANGLE_instanced_arrays or WebGL 2.0+. Android supports it with OpenGL ES3.0+
 	GSupportsTimestampRenderQueries = FOpenGL::SupportsTimestampQueries();
 
+	// It's not possible to create a framebuffer with the backbuffer as the color attachment and a custom renderbuffer as the depth/stencil surface.
+	GRHISupportsBackBufferWithCustomDepthStencil = false;
+
 	GSupportsHDR32bppEncodeModeIntrinsic = FOpenGL::SupportsHDR32bppEncodeModeIntrinsic();
 
 	GRHISupportsLazyShaderCodeLoading = true;
diff --git a/Engine/Source/Runtime/OpenGLDrv/Private/OpenGLRenderTarget.cpp b/Engine/Source/Runtime/OpenGLDrv/Private/OpenGLRenderTarget.cpp
--- a/Engine/Source/Runtime/OpenGLDrv/Private/OpenGLRenderTarget.cpp
+++ b/Engine/Source/Runtime/OpenGLDrv/Private/OpenGLRenderTarget.cpp
@@ -486,8 +486,12 @@ void FOpenGLDynamicRHI::RHICopyToResolveTarget(FRHITexture* SourceTextureRHI, FR
 
 		const bool bTrueBlit = !SourceTextureRHI->IsMultisampled()
 			&& !DestTextureRHI->IsMultisampled()
-			&& SourceTextureRHI->GetFormat() == DestTextureRHI->GetFormat();
-
+			&& SourceTextureRHI->GetFormat() == DestTextureRHI->GetFormat()
+#if PLATFORM_ANDROID
+			&& SourceTexture->Target == DestTexture->Target // glCopyImageSubData() doesn't like copying from a texture to a renderbuffer on Android
+#endif
+			;
+		
 		if ( !bTrueBlit || !FOpenGL::SupportsCopyImage() )
 		{
 			// Color buffers can be GL_NONE for attachment purposes if they aren't used as render targets
diff --git a/Engine/Source/Runtime/RHI/Private/RHI.cpp b/Engine/Source/Runtime/RHI/Private/RHI.cpp
--- a/Engine/Source/Runtime/RHI/Private/RHI.cpp
+++ b/Engine/Source/Runtime/RHI/Private/RHI.cpp
@@ -529,6 +529,8 @@ bool GRHISupportsUAVFormatAliasing = false;
 bool GRHISupportsMSAADepthSampleAccess = false;
 bool GRHISupportsResolveCubemapFaces = false;
 
+bool GRHISupportsBackBufferWithCustomDepthStencil = true;
+
 bool GRHIIsHDREnabled = false;
 bool GRHISupportsHDROutput = false;
 EPixelFormat GRHIHDRDisplayOutputFormat = PF_FloatRGBA;
diff --git a/Engine/Source/Runtime/RHI/Public/RHI.h b/Engine/Source/Runtime/RHI/Public/RHI.h
--- a/Engine/Source/Runtime/RHI/Public/RHI.h
+++ b/Engine/Source/Runtime/RHI/Public/RHI.h
@@ -534,6 +534,9 @@ extern RHI_API bool GRHISupportsParallelRHIExecute;
 /** Whether or not the RHI can perform MSAA sample load. */
 extern RHI_API bool GRHISupportsMSAADepthSampleAccess;
 
+/** Whether or not the RHI can render to the backbuffer with a custom depth/stencil surface bound. */
+extern RHI_API bool GRHISupportsBackBufferWithCustomDepthStencil;
+
 /** Whether or not HDR is currently enabled */
 extern RHI_API bool GRHIIsHDREnabled;
 
diff --git a/Engine/Source/Runtime/Renderer/Private/MobileShadingRenderer.cpp b/Engine/Source/Runtime/Renderer/Private/MobileShadingRenderer.cpp
--- a/Engine/Source/Runtime/Renderer/Private/MobileShadingRenderer.cpp
+++ b/Engine/Source/Runtime/Renderer/Private/MobileShadingRenderer.cpp
@@ -88,9 +88,12 @@ TGlobalResource<FGlobalDynamicReadBuffer> FMobileSceneRenderer::DynamicReadBuffe
 
 static bool UsesCustomDepthStencilLookup(const FViewInfo& View)
 {
+	if (View.bUsesCustomDepthStencil)
+	{
+		return true;
+	}
+
 	// Find out whether post-process materials use CustomDepth/Stencil lookups
-	bool bPPUsesCustomDepth = false;
-	bool bPPUsesCustomStencil = false;
 	const FBlendableManager& BlendableManager = View.FinalPostProcessSettings.BlendableManager;
 	FBlendableEntry* BlendableIt = nullptr;
 
@@ -103,15 +106,21 @@ static bool UsesCustomDepthStencilLookup(const FViewInfo& View)
 
 			const FMaterial* Material = Proxy->GetMaterial(View.GetFeatureLevel());
 			check(Material);
-			const FMaterialShaderMap* MaterialShaderMap = Material->GetRenderingThreadShaderMap();
+			if (Material->IsStencilTestEnabled())
+			{
+				return true;
+			}
 
-			bPPUsesCustomDepth|= MaterialShaderMap->UsesSceneTexture(PPI_CustomDepth);
-			bPPUsesCustomStencil|= MaterialShaderMap->UsesSceneTexture(PPI_CustomStencil);
+			const FMaterialShaderMap* MaterialShaderMap = Material->GetRenderingThreadShaderMap();
+			if (MaterialShaderMap->UsesSceneTexture(PPI_CustomDepth) || MaterialShaderMap->UsesSceneTexture(PPI_CustomStencil))
+			{
+				return true;
+			}
 		}
 	}
 
 	//TODO: check if translucency uses CustomDepth
-	return bPPUsesCustomDepth || bPPUsesCustomStencil || View.bUsesCustomDepthStencil;
+	return false;
 }
 
 
@@ -450,7 +459,11 @@ void FMobileSceneRenderer::Render(FRHICommandListImmediate& RHICmdList)
 		bool bUsesCustomDepthStencil = false;
 		for (int32 ViewIndex = 0; ViewIndex < Views.Num(); ViewIndex++) 
 		{
-			bUsesCustomDepthStencil = UsesCustomDepthStencilLookup(Views[ViewIndex]);
+			if (UsesCustomDepthStencilLookup(Views[ViewIndex]))
+			{
+				bUsesCustomDepthStencil = true;
+				break;
+			}
 		}
 
 		if (bUsesCustomDepthStencil)
diff --git a/Engine/Source/Runtime/Renderer/Private/PostProcess/PostProcessMaterial.cpp b/Engine/Source/Runtime/Renderer/Private/PostProcess/PostProcessMaterial.cpp
--- a/Engine/Source/Runtime/Renderer/Private/PostProcess/PostProcessMaterial.cpp
+++ b/Engine/Source/Runtime/Renderer/Private/PostProcess/PostProcessMaterial.cpp
@@ -330,6 +330,12 @@ FScreenPassTexture AddPostProcessMaterialPass(
 	// or we need to retain previously rendered views.
 	const bool bPrimeOutputColor = bIsCompositeWithInput || !View.IsFirstInFamily();
 
+	// Inputs.OverrideOutput is used to force drawing directly to the backbuffer. OpenGL doesn't support using the backbuffer color target with a custom depth/stencil
+	// buffer, so in that case we must draw to an intermediate target and copy to the backbuffer at the end. Ideally, we would test if Inputs.OverrideOutput.Texture
+	// is actually the backbuffer (as returned by AndroidEGL::GetOnScreenColorRenderBuffer() and such), but it's not worth doing all the plumbing and increasing the
+	// RHI surface area just for this hack.
+	const bool bForceIntermediateRT = !GRHISupportsBackBufferWithCustomDepthStencil && (DepthStencilTexture != nullptr) && Inputs.OverrideOutput.IsValid();
+
 	FScreenPassRenderTarget Output = Inputs.OverrideOutput;
 
 	// We can re-use the scene color texture as the render target if we're not simultaneously reading from it.
@@ -342,7 +348,7 @@ FScreenPassTexture AddPostProcessMaterialPass(
 	else
 	{
 		// Allocate new transient output texture if none exists.
-		if (!Output.IsValid())
+		if (!Output.IsValid() || bForceIntermediateRT)
 		{
 			FRDGTextureDesc OutputDesc = SceneColor.Texture->Desc;
 			OutputDesc.Reset();
@@ -356,7 +362,7 @@ FScreenPassTexture AddPostProcessMaterialPass(
 			Output = FScreenPassRenderTarget(GraphBuilder.CreateTexture(OutputDesc, TEXT("PostProcessMaterial")), SceneColor.ViewRect, View.GetOverwriteLoadAction());
 		}
 
-		if (bPrimeOutputColor)
+		if (bPrimeOutputColor || bForceIntermediateRT)
 		{
 			// Copy existing contents to new output and use load-action to preserve untouched pixels.
 			AddDrawTexturePass(GraphBuilder, View, SceneColor.Texture, Output.Texture);
@@ -458,6 +464,14 @@ FScreenPassTexture AddPostProcessMaterialPass(
 		}
 	});
 
+	if (bForceIntermediateRT)
+	{
+		// We shouldn't get here unless we had an override target.
+		check(Inputs.OverrideOutput.IsValid());
+		AddDrawTexturePass(GraphBuilder, View, Output.Texture, Inputs.OverrideOutput.Texture);
+		Output = Inputs.OverrideOutput;
+	}
+
 	return MoveTemp(Output);
 }
 
