Edigrated 9454753

  - Fixed a FPrimitiveSceneInfo memory leak in FScene class preventing destruction (and unregistration from the global list) of hit proxies referencing deleted actors. The level editor FViewport::GetHitProxyMap() would query the global list and return the dangling actors.
  - Added assert to ~FScene() to ensure that the list of removed FPrimitiveSceneInfo was cleaned up.

diff --git a/Engine/Source/Runtime/Renderer/Private/RendererScene.cpp b/Engine/Source/Runtime/Renderer/Private/RendererScene.cpp
--- a/Engine/Source/Runtime/Renderer/Private/RendererScene.cpp
+++ b/Engine/Source/Runtime/Renderer/Private/RendererScene.cpp
@@ -1090,6 +1090,8 @@ FScene::~FScene()
 
 	BeginReleaseResource(&HaltonPrimesResource);
 #endif
+
+	checkf(RemovedPrimitiveSceneInfos.Num() == 0, TEXT("Leaking %d FPrimitiveSceneInfo instances."), RemovedPrimitiveSceneInfos.Num()); // Ensure UpdateAllPrimitiveSceneInfos() is called before destruction.
 }
 
 void FScene::AddPrimitive(UPrimitiveComponent* Primitive)
@@ -3068,6 +3070,8 @@ void FScene::Release()
 	ENQUEUE_RENDER_COMMAND(FReleaseCommand)(
 		[Scene](FRHICommandListImmediate& RHICmdList)
 		{
+			// Update one more time to clear RemovedPrimitiveSceneInfos and prevent leaking FPrimitiveSceneInfo instances.
+			Scene->UpdateAllPrimitiveSceneInfos(RHICmdList);
 			delete Scene;
 		});
 }
