Texture LOD for ray tracing shaders is now a per-project setting, disabled by default.

(ushell-p4-cherrypick of 9538473 by [REDACTED_PERSON])

diff --git a/Engine/Source/Runtime/Engine/Classes/Engine/RendererSettings.h b/Engine/Source/Runtime/Engine/Classes/Engine/RendererSettings.h
--- a/Engine/Source/Runtime/Engine/Classes/Engine/RendererSettings.h
+++ b/Engine/Source/Runtime/Engine/Classes/Engine/RendererSettings.h
@@ -622,6 +622,12 @@ class ENGINE_API URendererSettings : public UDeveloperSettings
 		ConfigRestartRequired = true))
 		uint32 bEnableRayTracing : 1;
 
+	UPROPERTY(config, EditAnywhere, Category = RayTracing, meta = (
+		ConsoleVariable = "r.RayTracing.UseTextureLod", DisplayName = "Texture LOD",
+		ToolTip = "Enable automatic texture mip level selection in ray tracing material shaders. Unchecked: highest resolution mip level is used for all texture (default). Checked: texture LOD is approximated based on total ray length, output resolution and texel density at hit point (ray cone method).",
+		ConfigRestartRequired = true))
+		uint32 bEnableRayTracingTextureLOD : 1;
+
 	/**
 	"Stationary skylight requires permutations of the basepass shaders.  Disabling will reduce the number of shader permutations required per material. Changing this setting requires restarting the editor."
 	*/
diff --git a/Engine/Source/Runtime/Engine/Private/RendererSettings.cpp b/Engine/Source/Runtime/Engine/Private/RendererSettings.cpp
--- a/Engine/Source/Runtime/Engine/Private/RendererSettings.cpp
+++ b/Engine/Source/Runtime/Engine/Private/RendererSettings.cpp
@@ -73,6 +73,7 @@ URendererSettings::URendererSettings(const FObjectInitializer& ObjectInitializer
 	GPUSimulationTextureSizeX = 1024;
 	GPUSimulationTextureSizeY = 1024;
 	bEnableRayTracing = 0;
+	bEnableRayTracingTextureLOD = 0; 
 	bLPV = 1;
 }
 
diff --git a/Engine/Source/Runtime/Renderer/Private/RayTracing/RayTracingMaterialHitShaders.cpp b/Engine/Source/Runtime/Renderer/Private/RayTracing/RayTracingMaterialHitShaders.cpp
--- a/Engine/Source/Runtime/Renderer/Private/RayTracing/RayTracingMaterialHitShaders.cpp
+++ b/Engine/Source/Runtime/Renderer/Private/RayTracing/RayTracingMaterialHitShaders.cpp
@@ -17,13 +17,13 @@ static FAutoConsoleVariableRef CVarEnableRayTracingMaterials(
 	ECVF_RenderThreadSafe
 );
 
-static TAutoConsoleVariable<int32> CVarRayTracingTextureLodCvar(
+static TAutoConsoleVariable<int32> CVarRayTracingTextureLod(
 	TEXT("r.RayTracing.UseTextureLod"),
 	0,
-	TEXT("0 to disable texture LOD.\n")
-	TEXT(" 0: off\n")
-	TEXT(" 1: on"),
-	ECVF_RenderThreadSafe);
+	TEXT("Enable automatic texture mip level selection in ray tracing material shaders.\n")
+	TEXT(" 0: highest resolution mip level is used for all texture (default).\n")
+	TEXT(" 1: texture LOD is approximated based on total ray length, output resolution and texel density at hit point (ray cone method)."),
+	ECVF_ReadOnly);
 
 static bool IsSupportedVertexFactoryType(const FVertexFactoryType* VertexFactoryType)
 {
@@ -119,7 +119,8 @@ public:
 		return IsSupportedVertexFactoryType(Parameters.VertexFactoryType)
 			&& (Parameters.Material->IsMasked() == UseAnyHitShader)
 			&& LightMapPolicyType::ShouldCompilePermutation(Parameters)
-			&& ShouldCompileRayTracingShadersForProject(Parameters.Platform);
+			&& ShouldCompileRayTracingShadersForProject(Parameters.Platform)
+			&& (bool)CVarRayTracingTextureLod.GetValueOnAnyThread() == UseRayConeTextureLod;
 	}
 
 	static void ModifyCompilationEnvironment(const FMaterialShaderPermutationParameters& Parameters, FShaderCompilerEnvironment& OutEnvironment)
@@ -275,22 +276,24 @@ void FRayTracingMeshProcessor::Process(
 		FMeshMaterialShader,
 		FMaterialCHS> RayTracingShaders;
 
+	const bool bUseTextureLOD = CVarRayTracingTextureLod.GetValueOnAnyThread();
+
 	switch (LightMapPolicy.GetIndirectPolicy())
 	{
 	case LMP_PRECOMPUTED_IRRADIANCE_VOLUME_INDIRECT_LIGHTING:
-		RayTracingShaders.RayHitGroupShader = GetMaterialHitShader<TUniformLightMapPolicy<LMP_PRECOMPUTED_IRRADIANCE_VOLUME_INDIRECT_LIGHTING>>(MaterialResource, VertexFactory, CVarRayTracingTextureLodCvar.GetValueOnRenderThread() != 0);
+		RayTracingShaders.RayHitGroupShader = GetMaterialHitShader<TUniformLightMapPolicy<LMP_PRECOMPUTED_IRRADIANCE_VOLUME_INDIRECT_LIGHTING>>(MaterialResource, VertexFactory, bUseTextureLOD);
 		break;
 	case LMP_LQ_LIGHTMAP:
-		RayTracingShaders.RayHitGroupShader = GetMaterialHitShader<TUniformLightMapPolicy<LMP_LQ_LIGHTMAP>>(MaterialResource, VertexFactory, CVarRayTracingTextureLodCvar.GetValueOnRenderThread() != 0);
+		RayTracingShaders.RayHitGroupShader = GetMaterialHitShader<TUniformLightMapPolicy<LMP_LQ_LIGHTMAP>>(MaterialResource, VertexFactory, bUseTextureLOD);
 		break;
 	case LMP_HQ_LIGHTMAP:
-		RayTracingShaders.RayHitGroupShader = GetMaterialHitShader<TUniformLightMapPolicy<LMP_HQ_LIGHTMAP>>(MaterialResource, VertexFactory, CVarRayTracingTextureLodCvar.GetValueOnRenderThread() != 0);
+		RayTracingShaders.RayHitGroupShader = GetMaterialHitShader<TUniformLightMapPolicy<LMP_HQ_LIGHTMAP>>(MaterialResource, VertexFactory, bUseTextureLOD);
 		break;
 	case LMP_DISTANCE_FIELD_SHADOWS_AND_HQ_LIGHTMAP:
-		RayTracingShaders.RayHitGroupShader = GetMaterialHitShader<TUniformLightMapPolicy<LMP_DISTANCE_FIELD_SHADOWS_AND_HQ_LIGHTMAP>>(MaterialResource, VertexFactory, CVarRayTracingTextureLodCvar.GetValueOnRenderThread() != 0);
+		RayTracingShaders.RayHitGroupShader = GetMaterialHitShader<TUniformLightMapPolicy<LMP_DISTANCE_FIELD_SHADOWS_AND_HQ_LIGHTMAP>>(MaterialResource, VertexFactory, bUseTextureLOD);
 		break;
 	case LMP_NO_LIGHTMAP:
-		RayTracingShaders.RayHitGroupShader = GetMaterialHitShader<TUniformLightMapPolicy<LMP_NO_LIGHTMAP>>(MaterialResource, VertexFactory, CVarRayTracingTextureLodCvar.GetValueOnRenderThread() != 0);
+		RayTracingShaders.RayHitGroupShader = GetMaterialHitShader<TUniformLightMapPolicy<LMP_NO_LIGHTMAP>>(MaterialResource, VertexFactory, bUseTextureLOD);
 		break;
 	default:
 		check(false);
