Renderer settings shouldn't be hidden behind defines. Moving raytracing's texture lod cvar out. Fixes all non-raytracing platforms.

(ushell-p4-cherrypick of 9552152 by [REDACTED_PERSON])

diff --git a/Engine/Source/Runtime/Renderer/Private/DeferredShadingRenderer.cpp b/Engine/Source/Runtime/Renderer/Private/DeferredShadingRenderer.cpp
--- a/Engine/Source/Runtime/Renderer/Private/DeferredShadingRenderer.cpp
+++ b/Engine/Source/Runtime/Renderer/Private/DeferredShadingRenderer.cpp
@@ -131,6 +131,15 @@ static TAutoConsoleVariable<int32> CVarRayTracing(
 	TEXT(" 1: on"),
 	ECVF_RenderThreadSafe | ECVF_ReadOnly);
 
+int32 GRayTracingUseTextureLod = 0;
+static TAutoConsoleVariable<int32> CVarRayTracingTextureLod(
+	TEXT("r.RayTracing.UseTextureLod"),
+	GRayTracingUseTextureLod,
+	TEXT("Enable automatic texture mip level selection in ray tracing material shaders.\n")
+	TEXT(" 0: highest resolution mip level is used for all texture (default).\n")
+	TEXT(" 1: texture LOD is approximated based on total ray length, output resolution and texel density at hit point (ray cone method)."),
+	ECVF_RenderThreadSafe | ECVF_ReadOnly);
+
 static int32 GForceAllRayTracingEffects = -1;
 static TAutoConsoleVariable<int32> CVarForceAllRayTracingEffects(
 	TEXT("r.RayTracing.ForceAllRayTracingEffects"),
diff --git a/Engine/Source/Runtime/Renderer/Private/RayTracing/RayTracingMaterialHitShaders.cpp b/Engine/Source/Runtime/Renderer/Private/RayTracing/RayTracingMaterialHitShaders.cpp
--- a/Engine/Source/Runtime/Renderer/Private/RayTracing/RayTracingMaterialHitShaders.cpp
+++ b/Engine/Source/Runtime/Renderer/Private/RayTracing/RayTracingMaterialHitShaders.cpp
@@ -17,13 +17,8 @@ static FAutoConsoleVariableRef CVarEnableRayTracingMaterials(
 	ECVF_RenderThreadSafe
 );
 
-static TAutoConsoleVariable<int32> CVarRayTracingTextureLod(
-	TEXT("r.RayTracing.UseTextureLod"),
-	0,
-	TEXT("Enable automatic texture mip level selection in ray tracing material shaders.\n")
-	TEXT(" 0: highest resolution mip level is used for all texture (default).\n")
-	TEXT(" 1: texture LOD is approximated based on total ray length, output resolution and texel density at hit point (ray cone method)."),
-	ECVF_ReadOnly);
+// CVar defined in DeferredShadingRenderer.cpp
+extern int32 GRayTracingUseTextureLod;
 
 static bool IsSupportedVertexFactoryType(const FVertexFactoryType* VertexFactoryType)
 {
@@ -120,7 +115,7 @@ public:
 			&& (Parameters.Material->IsMasked() == UseAnyHitShader)
 			&& LightMapPolicyType::ShouldCompilePermutation(Parameters)
 			&& ShouldCompileRayTracingShadersForProject(Parameters.Platform)
-			&& (bool)CVarRayTracingTextureLod.GetValueOnAnyThread() == UseRayConeTextureLod;
+			&& (bool)GRayTracingUseTextureLod == UseRayConeTextureLod;
 	}
 
 	static void ModifyCompilationEnvironment(const FMaterialShaderPermutationParameters& Parameters, FShaderCompilerEnvironment& OutEnvironment)
@@ -276,7 +271,7 @@ void FRayTracingMeshProcessor::Process(
 		FMeshMaterialShader,
 		FMaterialCHS> RayTracingShaders;
 
-	const bool bUseTextureLOD = CVarRayTracingTextureLod.GetValueOnAnyThread();
+	const bool bUseTextureLOD = bool(GRayTracingUseTextureLod);
 
 	switch (LightMapPolicy.GetIndirectPolicy())
 	{
