Add a new method to reduce temporal aliasing in interactive Path Tracing views as well as a CVar to toggle it.

(ushell-p4-cherrypick of 9420691 by riley.alston)

diff --git a/Engine/Shaders/Private/PathTracing/PathTracing.usf b/Engine/Shaders/Private/PathTracing/PathTracing.usf
--- a/Engine/Shaders/Private/PathTracing/PathTracing.usf
+++ b/Engine/Shaders/Private/PathTracing/PathTracing.usf
@@ -42,7 +42,7 @@ RAY_TRACING_ENTRY_RAYGEN(PathTracingMainRG)
 	float2 BufferSize = View.BufferSizeAndInvSize.xy; 
 	uint LinearIndex = LaunchIndex.x + (LaunchIndex.y * BufferSize.x); 
 	RandomSequence RandSequence;
-	RandomSequence_Initialize(RandSequence, LinearIndex, AdaptiveSamplingData.Iteration);
+	RandomSequence_Initialize(RandSequence, LinearIndex, AdaptiveSamplingData.TemporalSeed);
 	RandSequence.Type = AdaptiveSamplingData.RandomSequence;
 	uint SampleIndex = 0;
 	uint SampleCount = 1;
diff --git a/Engine/Source/Runtime/Renderer/Private/PathTracing.cpp b/Engine/Source/Runtime/Renderer/Private/PathTracing.cpp
--- a/Engine/Source/Runtime/Renderer/Private/PathTracing.cpp
+++ b/Engine/Source/Runtime/Renderer/Private/PathTracing.cpp
@@ -28,13 +28,22 @@ TAutoConsoleVariable<int32> CVarPathTracingSamplesPerPixel(
 	ECVF_RenderThreadSafe
 );
 
+TAutoConsoleVariable<int32> CVarPathTracingFrameIndependentTemporalSeed(
+	TEXT("r.PathTracing.FrameIndependentTemporalSeed"),
+	1,
+	TEXT("Indicates to use different temporal seed for each sample across frames rather than resetting the sequence at the start of each frame\n")
+	TEXT("0: off\n")
+	TEXT("1: on (default)\n"),
+	ECVF_RenderThreadSafe
+);
+
 TAutoConsoleVariable<int32> CVarPathTracingRandomSequence(
 	TEXT("r.PathTracing.RandomSequence"),
 	2,
 	TEXT("Changes the underlying random sequence\n")
-	TEXT("0: LCG (default\n")
+	TEXT("0: LCG\n")
 	TEXT("1: Halton\n")
-	TEXT("2: Scrambled Halton\n"),
+	TEXT("2: Scrambled Halton (default)\n"),
 	ECVF_RenderThreadSafe
 );
 
@@ -119,6 +128,7 @@ public:
 		const TSparseArray<FLightSceneInfoCompact>& Lights,
 		// Adaptive sampling
 		uint32 Iteration,
+		uint32 FrameIndependentTemporalSeed,
 		FIntVector VarianceDimensions,
 		const FRWBuffer& VarianceMipTree,
 		// Output
@@ -237,6 +247,16 @@ public:
 
 		// Adaptive sampling
 		{
+			uint32 TemporalSeed;
+			if (CVarPathTracingFrameIndependentTemporalSeed.GetValueOnRenderThread() == 0)
+			{
+				TemporalSeed = Iteration;
+			}
+			else
+			{
+				TemporalSeed = FrameIndependentTemporalSeed;
+			}
+
 			FPathTracingAdaptiveSamplingData AdaptiveSamplingData;
 			AdaptiveSamplingData.MaxNormalBias = GetRaytracingMaxNormalBias();
 			AdaptiveSamplingData.UseAdaptiveSampling = CVarPathTracingAdaptiveSampling.GetValueOnRenderThread();
@@ -244,6 +264,7 @@ public:
 			if (VarianceMipTree.NumBytes > 0)
 			{
 				AdaptiveSamplingData.Iteration = Iteration;
+				AdaptiveSamplingData.TemporalSeed = TemporalSeed;
 				AdaptiveSamplingData.VarianceDimensions = VarianceDimensions;
 				AdaptiveSamplingData.VarianceMipTree = VarianceMipTree.SRV;
 				AdaptiveSamplingData.MinimumSamplesPerPixel = CVarPathTracingAdaptiveSamplingMinimumSamplesPerPixel.GetValueOnRenderThread();
@@ -252,6 +273,7 @@ public:
 			{
 				AdaptiveSamplingData.UseAdaptiveSampling = 0;
 				AdaptiveSamplingData.Iteration = Iteration;
+				AdaptiveSamplingData.TemporalSeed = TemporalSeed;
 				AdaptiveSamplingData.VarianceDimensions = FIntVector(1, 1, 1);
 				AdaptiveSamplingData.VarianceMipTree = RHICreateShaderResourceView(GBlackTexture->TextureRHI->GetTexture2D(), 0);
 				AdaptiveSamplingData.MinimumSamplesPerPixel = CVarPathTracingAdaptiveSamplingMinimumSamplesPerPixel.GetValueOnRenderThread();
@@ -407,6 +429,8 @@ void FDeferredShadingSceneRenderer::RenderPathTracing(FRHICommandListImmediate&
 
 	// The local iteration counter.
 	static int32 SPPCount = 0;
+	// The frame independent temporal seed, not reset at the beginning of each frame unlike SPPCount to allow for less temporal aliasing.
+	static uint32 FrameIndependentTemporalSeed = 0;
 
 	// Conditionally rebuild sky light CDFs
 	if (Scene->SkyLight && Scene->SkyLight->ShouldRebuildCdf())
@@ -457,7 +481,7 @@ void FDeferredShadingSceneRenderer::RenderPathTracing(FRHICommandListImmediate&
 		View.ViewUniformBuffer,
 		SceneTexturesUniformBuffer,
 		Scene->Lights,
-		SPPCount, ViewState->VarianceMipTreeDimensions, *ViewState->VarianceMipTree,
+		SPPCount, FrameIndependentTemporalSeed, ViewState->VarianceMipTreeDimensions, *ViewState->VarianceMipTree,
 		RadianceRT->GetRenderTargetItem().UAV,
 		SampleCountRT->GetRenderTargetItem().UAV,
 		PixelPositionRT->GetRenderTargetItem().UAV,
@@ -583,6 +607,8 @@ void FDeferredShadingSceneRenderer::RenderPathTracing(FRHICommandListImmediate&
 			SPPCount = 0;
 		}
 
+		++FrameIndependentTemporalSeed;
+
 		PixelShader->SetParameters(RHICmdList, View, RadianceRedTexture, RadianceGreenTexture, RadianceBlueTexture, RadianceAlphaTexture, SampleCountTexture, CumulativeRadianceTexture, CumulativeSampleCount);
 		DrawRectangle(
 			RHICmdList,
diff --git a/Engine/Source/Runtime/Renderer/Private/PathTracingUniformBuffers.h b/Engine/Source/Runtime/Renderer/Private/PathTracingUniformBuffers.h
--- a/Engine/Source/Runtime/Renderer/Private/PathTracingUniformBuffers.h
+++ b/Engine/Source/Runtime/Renderer/Private/PathTracingUniformBuffers.h
@@ -39,6 +39,7 @@ BEGIN_GLOBAL_SHADER_PARAMETER_STRUCT(FPathTracingAdaptiveSamplingData, )
 	SHADER_PARAMETER(uint32, RandomSequence)
 	SHADER_PARAMETER(uint32, MinimumSamplesPerPixel)
 	SHADER_PARAMETER(uint32, Iteration)
+	SHADER_PARAMETER(uint32, TemporalSeed)
 	SHADER_PARAMETER(float, MaxNormalBias)
 	SHADER_PARAMETER(FIntVector, VarianceDimensions)
 	SHADER_PARAMETER_SRV(Buffer<float>, VarianceMipTree)
