Translucent velocity now supports SceneDepth.

diff --git a/Engine/Source/Runtime/Renderer/Private/VelocityRendering.cpp b/Engine/Source/Runtime/Renderer/Private/VelocityRendering.cpp
--- a/Engine/Source/Runtime/Renderer/Private/VelocityRendering.cpp
+++ b/Engine/Source/Runtime/Renderer/Private/VelocityRendering.cpp
@@ -317,7 +317,7 @@ void FDeferredShadingSceneRenderer::RenderVelocitiesInnerParallel(FRHICommandLis
 
 			FSceneTexturesUniformParameters SceneTextureParameters;
 			FSceneRenderTargets& SceneContext = FSceneRenderTargets::Get(RHICmdList);
-			SetupSceneTextureUniformParameters(SceneContext, View.FeatureLevel, ESceneTextureSetupMode::None, SceneTextureParameters);
+			SetupSceneTextureUniformParameters(SceneContext, View.FeatureLevel, VelocityPass == EVelocityPass::Opaque ? ESceneTextureSetupMode::None : ESceneTextureSetupMode::SceneDepth, SceneTextureParameters);
 			TUniformBufferRef<FSceneTexturesUniformParameters> PassUniformBuffer = TUniformBufferRef<FSceneTexturesUniformParameters>::CreateUniformBufferImmediate(SceneTextureParameters, UniformBuffer_SingleFrame);
 
 			FMeshPassProcessorRenderState DrawRenderState(View, PassUniformBuffer);
