Fixed vertically flipped image on Android/OpenGL if a postprocess material was set to render after the tonemapper.

diff --git a/Engine/Source/Runtime/OpenGLDrv/Private/OpenGLRenderTarget.cpp b/Engine/Source/Runtime/OpenGLDrv/Private/OpenGLRenderTarget.cpp
--- a/Engine/Source/Runtime/OpenGLDrv/Private/OpenGLRenderTarget.cpp
+++ b/Engine/Source/Runtime/OpenGLDrv/Private/OpenGLRenderTarget.cpp
@@ -460,25 +460,24 @@ void FOpenGLDynamicRHI::RHICopyToResolveTarget(FRHITexture* SourceTextureRHI, FR
 		uint32 DestIndex = ResolveParams.DestArrayIndex * (bDestCubemap ? 6 : 1) + (bDestCubemap ? uint32(ResolveParams.CubeFace) : 0);
 		uint32 SrcIndex  = ResolveParams.SourceArrayIndex * (bSrcCubemap ? 6 : 1) + (bSrcCubemap ? uint32(ResolveParams.CubeFace) : 0);
 
-		uint32 BaseX = 0;
-		uint32 BaseY = 0;
-		uint32 SizeX = 0;
-		uint32 SizeY = 0;
-		if (ResolveParams.Rect.IsValid())
+		FIntRect SrcRect(ResolveParams.Rect.X1, ResolveParams.Rect.Y1, ResolveParams.Rect.X2, ResolveParams.Rect.Y2);
+		if (SrcRect.IsEmpty())
 		{
-			BaseX = ResolveParams.Rect.X1;
-			BaseY = ResolveParams.Rect.Y1;
-			SizeX = ResolveParams.Rect.X2 - ResolveParams.Rect.X1;
-			SizeY = ResolveParams.Rect.Y2 - ResolveParams.Rect.Y1;
+			// Empty rect mans that the entire source is to be copied. Note that we can't use ResolveParams.Rect.IsValid(), because it
+			// returns false if the rectangle is "inside out" (e.g. X1 > X2), and we want to perform flipping when that's the case.
+			SrcRect.Max.X = GetOpenGLTextureSizeXFromRHITexture(SourceTextureRHI);
+			SrcRect.Max.Y = GetOpenGLTextureSizeYFromRHITexture(SourceTextureRHI);
+			SrcRect.Max.X = FMath::Max<int32>(1, SrcRect.Max.X >> ResolveParams.MipIndex);
+			SrcRect.Max.Y = FMath::Max<int32>(1, SrcRect.Max.Y >> ResolveParams.MipIndex);
 		}
-		else
-		{
-			// Invalid rect mans that the entire source is to be copied
-			SizeX = GetOpenGLTextureSizeXFromRHITexture(SourceTextureRHI);
-			SizeY = GetOpenGLTextureSizeYFromRHITexture(SourceTextureRHI);
 
-			SizeX = FMath::Max<uint32>(1, SizeX >> ResolveParams.MipIndex);
-			SizeY = FMath::Max<uint32>(1, SizeY >> ResolveParams.MipIndex);
+		FIntRect DestRect(ResolveParams.DestRect.X1, ResolveParams.DestRect.Y1, ResolveParams.DestRect.X2, ResolveParams.DestRect.Y2);
+		if(DestRect.IsEmpty())
+		{
+			DestRect.Max.X = GetOpenGLTextureSizeXFromRHITexture(DestTextureRHI);
+			DestRect.Max.Y = GetOpenGLTextureSizeYFromRHITexture(DestTextureRHI);
+			DestRect.Max.X = FMath::Max<int32>(1, DestRect.Max.X >> ResolveParams.MipIndex);
+			DestRect.Max.Y = FMath::Max<int32>(1, DestRect.Max.Y >> ResolveParams.MipIndex);
 		}
 
 		GPUProfilingData.RegisterGPUWork();
@@ -487,6 +486,9 @@ void FOpenGLDynamicRHI::RHICopyToResolveTarget(FRHITexture* SourceTextureRHI, FR
 		const bool bTrueBlit = !SourceTextureRHI->IsMultisampled()
 			&& !DestTextureRHI->IsMultisampled()
 			&& SourceTextureRHI->GetFormat() == DestTextureRHI->GetFormat()
+			&& SrcRect.Size() == DestRect.Size()
+			&& SrcRect.Width() > 0
+			&& SrcRect.Height() > 0
 #if PLATFORM_ANDROID
 			&& SourceTexture->Target == DestTexture->Target // glCopyImageSubData() doesn't like copying from a texture to a renderbuffer on Android
 #endif
@@ -529,8 +531,8 @@ void FOpenGLDynamicRHI::RHICopyToResolveTarget(FRHITexture* SourceTextureRHI, FR
 			}
 
 			FOpenGL::BlitFramebuffer(
-				BaseX, BaseY, BaseX + SizeX, BaseY + SizeY,
-				BaseX, BaseY, BaseX + SizeX, BaseY + SizeY,
+				SrcRect.Min.X, SrcRect.Min.Y, SrcRect.Max.X, SrcRect.Max.Y,
+				DestRect.Min.X, DestRect.Min.Y, DestRect.Max.X, DestRect.Max.Y,
 				Mask,
 				GL_NEAREST
 				);
@@ -543,17 +545,17 @@ void FOpenGLDynamicRHI::RHICopyToResolveTarget(FRHITexture* SourceTextureRHI, FR
 			FOpenGL::CopyImageSubData(	SourceTexture->Resource,
 										SourceTexture->Target,
 										MipmapLevel,
-										BaseX,
-										BaseY,
+										SrcRect.Min.X,
+										SrcRect.Min.Y,
 										SrcIndex,
 										DestTexture->Resource,
 										DestTexture->Target,
 										MipmapLevel,
-										BaseX,
-										BaseY,
+										DestRect.Min.X,
+										DestRect.Min.Y,
 										DestIndex,
-										SizeX,
-										SizeY,
+										SrcRect.Width(),
+										SrcRect.Height(),
 										1);
 		}
 
diff --git a/Engine/Source/Runtime/RHI/Public/RHIContext.h b/Engine/Source/Runtime/RHI/Public/RHIContext.h
--- a/Engine/Source/Runtime/RHI/Public/RHIContext.h
+++ b/Engine/Source/Runtime/RHI/Public/RHIContext.h
@@ -633,12 +633,18 @@ public:
 			int32 DestArrayIndex = CopyInfo.DestSliceIndex + ArrayIndex;
 			for (int32 FaceIndex = 0; FaceIndex < NumFaces; ++FaceIndex)
 			{
-				FResolveParams ResolveParams(FResolveRect(),
+				FResolveParams ResolveParams(FResolveRect(0, 0, 0, 0),
 					bIsCube ? (ECubeFace)FaceIndex : CubeFace_PosX,
 					CopyInfo.SourceMipIndex,
 					SourceArrayIndex,
-					DestArrayIndex
+					DestArrayIndex,
+					FResolveRect(0, 0, 0, 0)
 				);
+				if (CopyInfo.Size != FIntVector::ZeroValue)
+				{
+					ResolveParams.Rect = FResolveRect(CopyInfo.SourcePosition.X, CopyInfo.SourcePosition.Y, CopyInfo.SourcePosition.X + CopyInfo.Size.X, CopyInfo.SourcePosition.Y + CopyInfo.Size.Y);
+					ResolveParams.DestRect = FResolveRect(CopyInfo.DestPosition.X, CopyInfo.DestPosition.Y, CopyInfo.DestPosition.X + CopyInfo.Size.X, CopyInfo.DestPosition.Y + CopyInfo.Size.Y);
+				}
 				RHICopyToResolveTarget(SourceTexture, DestTexture, ResolveParams);
 			}
 		}
diff --git a/Engine/Source/Runtime/RenderCore/Private/RenderGraphUtils.cpp b/Engine/Source/Runtime/RenderCore/Private/RenderGraphUtils.cpp
--- a/Engine/Source/Runtime/RenderCore/Private/RenderGraphUtils.cpp
+++ b/Engine/Source/Runtime/RenderCore/Private/RenderGraphUtils.cpp
@@ -121,6 +121,32 @@ void AddCopyTexturePass(
 	});
 }
 
+BEGIN_SHADER_PARAMETER_STRUCT(FCopyToResolveTargetParameters, )
+	SHADER_PARAMETER_RDG_TEXTURE(, Input)
+	RDG_TEXTURE_COPY_DEST(Output)
+END_SHADER_PARAMETER_STRUCT()
+
+void AddCopyToResolveTargetPass(
+	FRDGBuilder& GraphBuilder,
+	FRDGTextureRef InputTexture,
+	FRDGTextureRef OutputTexture,
+	const FResolveParams& ResolveParams)
+{
+	FCopyToResolveTargetParameters* Parameters = GraphBuilder.AllocParameters<FCopyToResolveTargetParameters>();
+	Parameters->Input = InputTexture;
+	Parameters->Output = OutputTexture;
+
+	GraphBuilder.AddPass(RDG_EVENT_NAME("CopyTexture"), Parameters, ERDGPassFlags::Copy,
+		[InputTexture, OutputTexture, ResolveParams](FRHICommandList& RHICmdList)
+	{
+		// Manually mark as used since we aren't invoking any shaders.
+		InputTexture->MarkResourceAsUsed();
+		OutputTexture->MarkResourceAsUsed();
+
+		RHICmdList.CopyToResolveTarget(InputTexture->GetRHI(), OutputTexture->GetRHI(), ResolveParams);
+	});
+}
+
 BEGIN_SHADER_PARAMETER_STRUCT(FClearBufferUAVParameters, )
 	SHADER_PARAMETER_RDG_BUFFER_UAV(RWBuffer<uint>, BufferUAV)
 END_SHADER_PARAMETER_STRUCT()
diff --git a/Engine/Source/Runtime/RenderCore/Public/RenderGraphUtils.h b/Engine/Source/Runtime/RenderCore/Public/RenderGraphUtils.h
--- a/Engine/Source/Runtime/RenderCore/Public/RenderGraphUtils.h
+++ b/Engine/Source/Runtime/RenderCore/Public/RenderGraphUtils.h
@@ -213,6 +213,15 @@ inline void AddCopyTexturePass(
 	AddCopyTexturePass(GraphBuilder, InputTexture, OutputTexture, CopyInfo);
 }
 
+/** Adds a render graph pass to resolve from one texture to another. Uses RHICopyToResolveTarget under the hood.
+ *  The formats of the two textures don't need to match.
+ */
+RENDERCORE_API void AddCopyToResolveTargetPass(
+	FRDGBuilder& GraphBuilder,
+	FRDGTextureRef InputTexture,
+	FRDGTextureRef OutputTexture,
+	const FResolveParams& ResolveParams);
+
 /** Adds a render graph pass to clear a texture or buffer UAV with a single typed value. */
 RENDERCORE_API void AddClearUAVPass(FRDGBuilder& GraphBuilder, FRDGBufferUAVRef BufferUAV, uint32 Value);
 
diff --git a/Engine/Source/Runtime/Renderer/Private/PostProcess/PostProcessMaterial.cpp b/Engine/Source/Runtime/Renderer/Private/PostProcess/PostProcessMaterial.cpp
--- a/Engine/Source/Runtime/Renderer/Private/PostProcess/PostProcessMaterial.cpp
+++ b/Engine/Source/Runtime/Renderer/Private/PostProcess/PostProcessMaterial.cpp
@@ -288,6 +288,41 @@ TGlobalResource<FPostProcessMaterialVertexDeclaration> GPostProcessMaterialVerte
 
 } //! namespace
 
+static void AddCopyAndFlipTexturePass(FRDGBuilder& GraphBuilder, const FViewInfo& View, FRDGTextureRef SrcTexture, FRDGTextureRef DestTexture)
+{
+	if (IsOpenGLPlatform(GShaderPlatformForFeatureLevel[View.GetFeatureLevel()]))
+	{
+		// The OpenGL RHI can copy and flip at the same time by using an upside down destination rectangle.
+		FResolveParams ResolveParams;
+		FMemory::Memzero(ResolveParams);
+		ResolveParams.Rect.X1 = 0;
+		ResolveParams.Rect.X2 = SrcTexture->Desc.Extent.X;
+		ResolveParams.Rect.Y1 = 0;
+		ResolveParams.Rect.Y2 = SrcTexture->Desc.Extent.Y;
+		ResolveParams.DestRect.X1 = 0;
+		ResolveParams.DestRect.X2 = DestTexture->Desc.Extent.X;
+		ResolveParams.DestRect.Y1 = DestTexture->Desc.Extent.Y - 1;
+		ResolveParams.DestRect.Y2 = -1;
+		AddCopyToResolveTargetPass(GraphBuilder, SrcTexture, DestTexture, ResolveParams);
+		return;
+	}
+
+	// Other RHIs can't flip and copy at the same time, so we'll use a pixel shader to perform the copy, together with the FlipYAxis
+	// flag on the screen pass. This path is only taken when using the mobile preview feature in the editor with
+	// r.Mobile.ForceRHISwitchVerticalAxis set to 1, so we don't care about it being sub-optimal.
+	FIntPoint Size = SrcTexture->Desc.Extent;
+	const FScreenPassTextureViewport InputViewport(SrcTexture->Desc.Extent, FIntRect(FIntPoint::ZeroValue, Size));
+	const FScreenPassTextureViewport OutputViewport(DestTexture->Desc.Extent, FIntRect(FIntPoint::ZeroValue, Size));
+	TShaderMapRef<FCopyRectPS> PixelShader(View.ShaderMap);
+
+	FCopyRectPS::FParameters* Parameters = GraphBuilder.AllocParameters<FCopyRectPS::FParameters>();
+	Parameters->InputTexture = SrcTexture;
+	Parameters->InputSampler = TStaticSamplerState<>::GetRHI();
+	Parameters->RenderTargets[0] = FRenderTargetBinding(DestTexture, ERenderTargetLoadAction::ENoAction);
+
+	AddDrawScreenPass(GraphBuilder, RDG_EVENT_NAME("DrawTexture"), View, OutputViewport, InputViewport, *PixelShader, Parameters, EScreenPassDrawFlags::FlipYAxis);
+}
+
 FScreenPassTexture AddPostProcessMaterialPass(
 	FRDGBuilder& GraphBuilder,
 	const FViewInfo& View,
@@ -334,14 +369,20 @@ FScreenPassTexture AddPostProcessMaterialPass(
 	// buffer, so in that case we must draw to an intermediate target and copy to the backbuffer at the end. Ideally, we would test if Inputs.OverrideOutput.Texture
 	// is actually the backbuffer (as returned by AndroidEGL::GetOnScreenColorRenderBuffer() and such), but it's not worth doing all the plumbing and increasing the
 	// RHI surface area just for this hack.
-	const bool bForceIntermediateRT = !GRHISupportsBackBufferWithCustomDepthStencil && (DepthStencilTexture != nullptr) && Inputs.OverrideOutput.IsValid();
+	//
+	// The other case when we must render to an intermediate target is when we have to flip the image vertically because we're the last postprocess pass on mobile OpenGL.
+	// We can't simply output a flipped image, because the parts of the input image which show through the stencil mask or are blended in must also be flipped. In that case,
+	// we render normally to the intermediate target and flip the image when we copy to the output target.
+	const bool bForceIntermediateRT =
+		((DepthStencilTexture != nullptr) && !GRHISupportsBackBufferWithCustomDepthStencil && Inputs.OverrideOutput.IsValid()) ||
+		(bIsCompositeWithInput && Inputs.bFlipYAxis);
 
 	FScreenPassRenderTarget Output = Inputs.OverrideOutput;
 
 	// We can re-use the scene color texture as the render target if we're not simultaneously reading from it.
 	// This is only necessary to do if we're going to be priming content from the render target since it avoids
 	// the copy. Otherwise, we just allocate a new render target.
-	if (!Output.IsValid() && !MaterialShaderMap->UsesSceneTexture(PPI_PostProcessInput0) && bPrimeOutputColor)
+	if (!Output.IsValid() && !MaterialShaderMap->UsesSceneTexture(PPI_PostProcessInput0) && bPrimeOutputColor && !bForceIntermediateRT)
 	{
 		Output = FScreenPassRenderTarget(SceneColor, ERenderTargetLoadAction::ELoad);
 	}
@@ -414,7 +455,7 @@ FScreenPassTexture AddPostProcessMaterialPass(
 
 	const bool bIsMobile = FeatureLevel <= ERHIFeatureLevel::ES3_1;
 
-	PostProcessMaterialParameters->bFlipYAxis = Inputs.bFlipYAxis;
+	PostProcessMaterialParameters->bFlipYAxis = Inputs.bFlipYAxis && !bForceIntermediateRT;
 
 	FPostProcessMaterialShader::FPermutationDomain PermutationVector;
 	PermutationVector.Set<FPostProcessMaterialShader::FMobileDimension>(bIsMobile);
@@ -426,7 +467,7 @@ FScreenPassTexture AddPostProcessMaterialPass(
 
 	EScreenPassDrawFlags ScreenPassFlags = EScreenPassDrawFlags::AllowHMDHiddenAreaMask;
 
-	if (Inputs.bFlipYAxis)
+	if (PostProcessMaterialParameters->bFlipYAxis)
 	{
 		ScreenPassFlags |= EScreenPassDrawFlags::FlipYAxis;
 	}
@@ -466,10 +507,27 @@ FScreenPassTexture AddPostProcessMaterialPass(
 
 	if (bForceIntermediateRT)
 	{
-		// We shouldn't get here unless we had an override target.
-		check(Inputs.OverrideOutput.IsValid());
-		AddDrawTexturePass(GraphBuilder, View, Output.Texture, Inputs.OverrideOutput.Texture);
-		Output = Inputs.OverrideOutput;
+		if (!Inputs.bFlipYAxis)
+		{
+			// We shouldn't get here unless we had an override target.
+			check(Inputs.OverrideOutput.IsValid());
+			AddDrawTexturePass(GraphBuilder, View, Output.Texture, Inputs.OverrideOutput.Texture);
+			Output = Inputs.OverrideOutput;
+		}
+		else
+		{
+			FScreenPassRenderTarget TempTarget = Output;
+			if (Inputs.OverrideOutput.IsValid())
+			{
+				Output = Inputs.OverrideOutput;
+			}
+			else
+			{
+				Output = FScreenPassRenderTarget(SceneColor, ERenderTargetLoadAction::ENoAction);
+			}
+
+			AddCopyAndFlipTexturePass(GraphBuilder, View, TempTarget.Texture, Output.Texture);
+		}
 	}
 
 	return MoveTemp(Output);
@@ -623,7 +681,7 @@ FRenderingCompositePass* AddPostProcessMaterialPass(
 			Inputs.SetInput((EPostProcessMaterialInput)InputIndex, FScreenPassTexture(InputTexture, InputViewportRect));
 		}
 
-		Inputs.bFlipYAxis = ShouldMobilePassFlipVerticalAxis(Pass);
+		Inputs.bFlipYAxis = ShouldMobilePassFlipVerticalAxis(InContext, Pass);
 
 		if (TRefCountPtr<IPooledRenderTarget> CustomDepthTarget = FSceneRenderTargets::Get(InContext.RHICmdList).CustomDepth)
 		{
diff --git a/Engine/Source/Runtime/Renderer/Private/PostProcess/PostProcessMobile.cpp b/Engine/Source/Runtime/Renderer/Private/PostProcess/PostProcessMobile.cpp
--- a/Engine/Source/Runtime/Renderer/Private/PostProcess/PostProcessMobile.cpp
+++ b/Engine/Source/Runtime/Renderer/Private/PostProcess/PostProcessMobile.cpp
@@ -33,9 +33,9 @@ void SetMobilePassFlipVerticalAxis(const FRenderingCompositePass* FlipPass)
 	GMobilePassShouldFlipVerticalAxis = FlipPass;
 }
 
-bool ShouldMobilePassFlipVerticalAxis(const FRenderingCompositePass* ShouldFlipPass)
+bool ShouldMobilePassFlipVerticalAxis(const FRenderingCompositePassContext& Context, const FRenderingCompositePass* ShouldFlipPass)
 {
-	return GMobilePassShouldFlipVerticalAxis == ShouldFlipPass;
+	return RHINeedsToSwitchVerticalAxis(Context.GetShaderPlatform()) && (GMobilePassShouldFlipVerticalAxis == ShouldFlipPass);
 }
 
 //
diff --git a/Engine/Source/Runtime/Renderer/Private/PostProcess/PostProcessMobile.h b/Engine/Source/Runtime/Renderer/Private/PostProcess/PostProcessMobile.h
--- a/Engine/Source/Runtime/Renderer/Private/PostProcess/PostProcessMobile.h
+++ b/Engine/Source/Runtime/Renderer/Private/PostProcess/PostProcessMobile.h
@@ -17,7 +17,7 @@ float GetMobileDepthOfFieldScale(const FViewInfo& View);
 
 // Used to indicate the final PP stage which needs to be flipped on platforms that 'RHINeedsToSwitchVerticalAxis'
 void SetMobilePassFlipVerticalAxis(const FRenderingCompositePass* FlipPass);
-bool ShouldMobilePassFlipVerticalAxis(const FRenderingCompositePass* ShouldFlipPass);
+bool ShouldMobilePassFlipVerticalAxis(const FRenderingCompositePassContext& Context, const FRenderingCompositePass* ShouldFlipPass);
 
 class FRCPassPostProcessBloomSetupES2 : public TRenderingCompositePassBase<1, 1>
 {
diff --git a/Engine/Source/Runtime/Renderer/Private/PostProcess/PostProcessTonemap.cpp b/Engine/Source/Runtime/Renderer/Private/PostProcess/PostProcessTonemap.cpp
--- a/Engine/Source/Runtime/Renderer/Private/PostProcess/PostProcessTonemap.cpp
+++ b/Engine/Source/Runtime/Renderer/Private/PostProcess/PostProcessTonemap.cpp
@@ -956,7 +956,7 @@ void FRCPassPostProcessTonemap::Process(FRenderingCompositePassContext& Context)
 	Inputs.OverrideOutput.ViewRect = OutputViewRect;
 	Inputs.OverrideOutput.LoadAction = OutputLoadAction;
 	Inputs.bWriteAlphaChannel = View.AntiAliasingMethod == AAM_FXAA || IsPostProcessingWithAlphaChannelSupported();
-	Inputs.bFlipYAxis = RHINeedsToSwitchVerticalAxis(GShaderPlatformForFeatureLevel[Context.GetFeatureLevel()]) && ShouldMobilePassFlipVerticalAxis(this);
+	Inputs.bFlipYAxis = ShouldMobilePassFlipVerticalAxis(Context, this);
 	Inputs.bGammaOnly = bDoGammaOnly;
 
 	AddTonemapPass(GraphBuilder, View, Inputs);
@@ -1495,7 +1495,7 @@ void FRCPassPostProcessTonemapES2::Process(FRenderingCompositePassContext& Conte
 	{
 		Context.SetViewportAndCallRHI(DestRect);
 
-		bool bNeedsToSwitchVerticalAxis = RHINeedsToSwitchVerticalAxis(GShaderPlatformForFeatureLevel[Context.GetFeatureLevel()]) && ShouldMobilePassFlipVerticalAxis(this);
+		bool bNeedsToSwitchVerticalAxis = ShouldMobilePassFlipVerticalAxis(Context, this);
 
 		auto VertexShaderPermutationVector = FPostProcessTonemapVS_ES2::BuildPermutationVector(bNeedsToSwitchVerticalAxis);
 		auto PixelShaderPermutationVector = FPostProcessTonemapPS_ES2::BuildPermutationVector(View, bNeedsToSwitchVerticalAxis);
