Fix for water material to execute with RDG in dev-render

(ushell-p4-cherrypick of 9547180 by [REDACTED_PERSON])

diff --git a/Engine/Source/Runtime/Renderer/Private/SingleLayerWaterRendering.cpp b/Engine/Source/Runtime/Renderer/Private/SingleLayerWaterRendering.cpp
--- a/Engine/Source/Runtime/Renderer/Private/SingleLayerWaterRendering.cpp
+++ b/Engine/Source/Runtime/Renderer/Private/SingleLayerWaterRendering.cpp
@@ -451,11 +451,11 @@ void FDeferredShadingSceneRenderer::CopySingleLayerWaterTextures(FRHICommandList
 
 	if (bCopyColor)
 	{
-		FPooledRenderTargetDesc ColorDesc(FPooledRenderTargetDesc::Create2DDesc(RefractionResolution, SceneContext.GetSceneColorFormat(), SceneContext.GetDefaultColorClear(), TexCreate_None, TexCreate_RenderTargetable, false));
+		FPooledRenderTargetDesc ColorDesc(FPooledRenderTargetDesc::Create2DDesc(RefractionResolution, SceneContext.GetSceneColorFormat(), SceneContext.GetDefaultColorClear(), TexCreate_None, TexCreate_RenderTargetable | TexCreate_ShaderResource, false));
 		GRenderTargetPool.FindFreeElement(RHICmdList, ColorDesc, PassData.SceneColorWithoutSingleLayerWater, TEXT("SceneColorWithoutSingleLayerWater"));
 	}
 
-	FPooledRenderTargetDesc DepthDesc(FPooledRenderTargetDesc::Create2DDesc(RefractionResolution, GSingleLayerWaterRefractionFullPrecision ? PF_R32_FLOAT : PF_R16F, SceneContext.GetDefaultColorClear(), TexCreate_None, TexCreate_RenderTargetable, false));
+	FPooledRenderTargetDesc DepthDesc(FPooledRenderTargetDesc::Create2DDesc(RefractionResolution, GSingleLayerWaterRefractionFullPrecision ? PF_R32_FLOAT : PF_R16F, SceneContext.GetDefaultColorClear(), TexCreate_None, TexCreate_RenderTargetable | TexCreate_ShaderResource, false));
 	GRenderTargetPool.FindFreeElement(RHICmdList, DepthDesc, PassData.SceneDepthWithoutSingleLayerWater, TEXT("SceneDepthWithoutSingleLayerWater"));
 
 	// For now support only the 1st view
