Added some more missing TexCreate_ShaderResource...

(ushell-p4-cherrypick of 9566732 by [REDACTED_PERSON])

diff --git a/Engine/Source/Runtime/Renderer/Private/DistanceFieldAmbientOcclusion.cpp b/Engine/Source/Runtime/Renderer/Private/DistanceFieldAmbientOcclusion.cpp
--- a/Engine/Source/Runtime/Renderer/Private/DistanceFieldAmbientOcclusion.cpp
+++ b/Engine/Source/Runtime/Renderer/Private/DistanceFieldAmbientOcclusion.cpp
@@ -849,7 +849,7 @@ bool FDeferredShadingSceneRenderer::RenderDistanceFieldLighting(
 
 			{
 				const FIntPoint BufferSize = GetBufferSizeForAO();
-				FPooledRenderTargetDesc Desc(FPooledRenderTargetDesc::Create2DDesc(BufferSize, PF_FloatRGBA, FClearValueBinding::Transparent, TexCreate_None, TexCreate_RenderTargetable | TexCreate_UAV, false));
+				FPooledRenderTargetDesc Desc(FPooledRenderTargetDesc::Create2DDesc(BufferSize, PF_FloatRGBA, FClearValueBinding::Transparent, TexCreate_None, TexCreate_RenderTargetable | TexCreate_UAV | TexCreate_ShaderResource, false));
 				Desc.Flags |= GFastVRamConfig.DistanceFieldNormal;
 				GRenderTargetPool.FindFreeElement(RHICmdList, Desc, DistanceFieldNormal, TEXT("DistanceFieldNormal"));
 			}
diff --git a/Engine/Source/Runtime/Renderer/Private/DistanceFieldScreenGridLighting.cpp b/Engine/Source/Runtime/Renderer/Private/DistanceFieldScreenGridLighting.cpp
--- a/Engine/Source/Runtime/Renderer/Private/DistanceFieldScreenGridLighting.cpp
+++ b/Engine/Source/Runtime/Renderer/Private/DistanceFieldScreenGridLighting.cpp
@@ -673,7 +673,7 @@ void FDeferredShadingSceneRenderer::RenderDistanceFieldAOScreenGrid(
 	TRefCountPtr<IPooledRenderTarget> DownsampledBentNormal;
 
 	{
-		FPooledRenderTargetDesc Desc(FPooledRenderTargetDesc::Create2DDesc(ConeTraceBufferSize, PF_FloatRGBA, FClearValueBinding::None, TexCreate_None, TexCreate_RenderTargetable | TexCreate_UAV, false));
+		FPooledRenderTargetDesc Desc(FPooledRenderTargetDesc::Create2DDesc(ConeTraceBufferSize, PF_FloatRGBA, FClearValueBinding::None, TexCreate_None, TexCreate_RenderTargetable | TexCreate_UAV | TexCreate_ShaderResource, false));
 		Desc.Flags |= GFastVRamConfig.DistanceFieldAODownsampledBentNormal;
 		GRenderTargetPool.FindFreeElement(RHICmdList, Desc, DownsampledBentNormal, TEXT("DownsampledBentNormal"));
 	}
