Fixed render target pool SRV flags. Refactored water to use RDG for render target creation.

diff --git a/Engine/Source/Runtime/Renderer/Private/LightShaftRendering.cpp b/Engine/Source/Runtime/Renderer/Private/LightShaftRendering.cpp
--- a/Engine/Source/Runtime/Renderer/Private/LightShaftRendering.cpp
+++ b/Engine/Source/Runtime/Renderer/Private/LightShaftRendering.cpp
@@ -480,7 +480,7 @@ void AllocateOrReuseLightShaftRenderTarget(FRHICommandListImmediate& RHICmdList,
 		EPixelFormat LightShaftFilterBufferFormat = PF_FloatRGB;
 		const FIntPoint BufferSize = FSceneRenderTargets::Get(RHICmdList).GetBufferSizeXY();
 		FIntPoint LightShaftSize(FMath::Max<uint32>(BufferSize.X / GetLightShaftDownsampleFactor(), 1), FMath::Max<uint32>(BufferSize.Y / GetLightShaftDownsampleFactor(), 1));
-		FPooledRenderTargetDesc Desc(FPooledRenderTargetDesc::Create2DDesc(LightShaftSize, LightShaftFilterBufferFormat, FClearValueBinding::White, TexCreate_None, TexCreate_RenderTargetable, false));
+		FPooledRenderTargetDesc Desc(FPooledRenderTargetDesc::Create2DDesc(LightShaftSize, LightShaftFilterBufferFormat, FClearValueBinding::White, TexCreate_None, TexCreate_ShaderResource | TexCreate_RenderTargetable, false));
 		Desc.AutoWritable = false;
 		GRenderTargetPool.FindFreeElement(RHICmdList, Desc, Target, Name);
 
diff --git a/Engine/Source/Runtime/Renderer/Private/SingleLayerWaterRendering.cpp b/Engine/Source/Runtime/Renderer/Private/SingleLayerWaterRendering.cpp
--- a/Engine/Source/Runtime/Renderer/Private/SingleLayerWaterRendering.cpp
+++ b/Engine/Source/Runtime/Renderer/Private/SingleLayerWaterRendering.cpp
@@ -449,14 +449,18 @@ void FDeferredShadingSceneRenderer::CopySingleLayerWaterTextures(FRHICommandList
 	const int32 RefractionDownsampleFactor = FMath::Clamp(GSingleLayerWaterRefractionDownsampleFactor, 1, 8);
 	const FIntPoint RefractionResolution = FIntPoint::DivideAndRoundDown(SceneContext.GetBufferSizeXY(), RefractionDownsampleFactor);
 
+	FRDGBuilder GraphBuilder(RHICmdList);
+
+	FRDGTextureRef SceneColorWithoutSingleLayerWaterTexture = nullptr;
+
 	if (bCopyColor)
 	{
-		FPooledRenderTargetDesc ColorDesc(FPooledRenderTargetDesc::Create2DDesc(RefractionResolution, SceneContext.GetSceneColorFormat(), SceneContext.GetDefaultColorClear(), TexCreate_None, TexCreate_RenderTargetable | TexCreate_ShaderResource, false));
-		GRenderTargetPool.FindFreeElement(RHICmdList, ColorDesc, PassData.SceneColorWithoutSingleLayerWater, TEXT("SceneColorWithoutSingleLayerWater"));
+		const FRDGTextureDesc ColorDesc = FRDGTextureDesc::Create2DDesc(RefractionResolution, SceneContext.GetSceneColorFormat(), SceneContext.GetDefaultColorClear(), TexCreate_None, TexCreate_ShaderResource | TexCreate_RenderTargetable, false);
+		SceneColorWithoutSingleLayerWaterTexture = GraphBuilder.CreateTexture(ColorDesc, TEXT("SceneColorWithoutSingleLayerWater"));
 	}
 
-	FPooledRenderTargetDesc DepthDesc(FPooledRenderTargetDesc::Create2DDesc(RefractionResolution, GSingleLayerWaterRefractionFullPrecision ? PF_R32_FLOAT : PF_R16F, SceneContext.GetDefaultColorClear(), TexCreate_None, TexCreate_RenderTargetable | TexCreate_ShaderResource, false));
-	GRenderTargetPool.FindFreeElement(RHICmdList, DepthDesc, PassData.SceneDepthWithoutSingleLayerWater, TEXT("SceneDepthWithoutSingleLayerWater"));
+	const FRDGTextureDesc DepthDesc(FRDGTextureDesc::Create2DDesc(RefractionResolution, GSingleLayerWaterRefractionFullPrecision ? PF_R32_FLOAT : PF_R16F, SceneContext.GetDefaultColorClear(), TexCreate_None, TexCreate_ShaderResource | TexCreate_RenderTargetable, false));
+	FRDGTextureRef SceneDepthWithoutSingleLayerWaterTexture = GraphBuilder.CreateTexture(DepthDesc, TEXT("SceneDepthWithoutSingleLayerWater"));
 
 	// For now support only the 1st view
 	const FViewInfo& View = Views[0];
@@ -471,12 +475,11 @@ void FDeferredShadingSceneRenderer::CopySingleLayerWaterTextures(FRHICommandList
 	PassParameters->SceneDepthCopyDownsampleSampler = TStaticSamplerState<SF_Bilinear, AM_Clamp, AM_Clamp, AM_Clamp>::GetRHI();
 	PassParameters->SVPositionToSourceTextureUV = FVector2D(RefractionDownsampleFactor / float(SceneContext.GetBufferSizeXY().X), RefractionDownsampleFactor / float(SceneContext.GetBufferSizeXY().Y));
 
-	FRDGTextureRef TargetSceneDepthTexture = GraphBuilder.RegisterExternalTexture(PassData.SceneDepthWithoutSingleLayerWater);
-	PassParameters->RenderTargets[0] = FRenderTargetBinding(TargetSceneDepthTexture, ERenderTargetLoadAction::ELoad);
+	PassParameters->RenderTargets[0] = FRenderTargetBinding(SceneDepthWithoutSingleLayerWaterTexture, ERenderTargetLoadAction::ENoAction);
+
 	if (bCopyColor)
 	{
-		FRDGTextureRef TargetColorTexture = GraphBuilder.RegisterExternalTexture(PassData.SceneColorWithoutSingleLayerWater);
-		PassParameters->RenderTargets[1] = FRenderTargetBinding(TargetColorTexture, ERenderTargetLoadAction::ELoad);
+		PassParameters->RenderTargets[1] = FRenderTargetBinding(SceneColorWithoutSingleLayerWaterTexture, ERenderTargetLoadAction::ENoAction);
 	}
 
 	FWaterRefractionCopyPS::FPermutationDomain PermutationVector;
@@ -496,13 +499,15 @@ void FDeferredShadingSceneRenderer::CopySingleLayerWaterTextures(FRHICommandList
 		PixelShader,
 		PassParameters,
 		RefractionViewRect);
-	GraphBuilder.Execute();
+
+	GraphBuilder.QueueTextureExtraction(SceneDepthWithoutSingleLayerWaterTexture, &PassData.SceneDepthWithoutSingleLayerWater, true);
 
 	if (bCopyColor)
 	{
-		RHICmdList.TransitionResource(EResourceTransitionAccess::EReadable, PassData.SceneColorWithoutSingleLayerWater->GetRenderTargetItem().TargetableTexture.GetReference());
+		GraphBuilder.QueueTextureExtraction(SceneColorWithoutSingleLayerWaterTexture, &PassData.SceneColorWithoutSingleLayerWater, true);
 	}
-	RHICmdList.TransitionResource(EResourceTransitionAccess::EReadable, PassData.SceneDepthWithoutSingleLayerWater->GetRenderTargetItem().TargetableTexture.GetReference());
+
+	GraphBuilder.Execute();
 }
 
 void FDeferredShadingSceneRenderer::BeginRenderingWaterGBuffer(FRHICommandList& RHICmdList, FExclusiveDepthStencil::Type DepthStencilAccess, bool bBindQuadOverdrawBuffers, EShaderPlatform InShaderPlatform)
