Fixed build break.

diff --git a/Engine/Source/Runtime/Renderer/Private/SingleLayerWaterRendering.cpp b/Engine/Source/Runtime/Renderer/Private/SingleLayerWaterRendering.cpp
--- a/Engine/Source/Runtime/Renderer/Private/SingleLayerWaterRendering.cpp
+++ b/Engine/Source/Runtime/Renderer/Private/SingleLayerWaterRendering.cpp
@@ -465,8 +465,6 @@ void FDeferredShadingSceneRenderer::CopySingleLayerWaterTextures(FRHICommandList
 	// For now support only the 1st view
 	const FViewInfo& View = Views[0];
 
-	FRDGBuilder GraphBuilder(RHICmdList);
-
 	FWaterRefractionCopyPS::FParameters* PassParameters = GraphBuilder.AllocParameters<FWaterRefractionCopyPS::FParameters>();
 	PassParameters->ViewUniformBuffer = View.ViewUniformBuffer;
 	PassParameters->SceneColorCopyDownsampleTexture = GraphBuilder.RegisterExternalTexture(SceneContext.GetSceneColor());
