Fix for Linux Build break

diff --git a/Engine/Source/Runtime/Renderer/Private/CompositionLighting/PostProcessAmbientOcclusion.cpp b/Engine/Source/Runtime/Renderer/Private/CompositionLighting/PostProcessAmbientOcclusion.cpp
--- a/Engine/Source/Runtime/Renderer/Private/CompositionLighting/PostProcessAmbientOcclusion.cpp
+++ b/Engine/Source/Runtime/Renderer/Private/CompositionLighting/PostProcessAmbientOcclusion.cpp
@@ -2506,10 +2506,12 @@ public:
 		FRHIPixelShader* ShaderRHI = GetPixelShader();
 
 		FGlobalShader::SetParameters<FViewUniformShaderParameters>(Context.RHICmdList, ShaderRHI, Context.View.ViewUniformBuffer);
-		ScreenSpaceAOParams.Set(RHICmdList, Context.View, ShaderRHI, InputTextureSize);
+		//ScreenSpaceAOParams.Set(Context.RHICmdList, Context.View, ShaderRHI, InputTextureSize);
 		PostprocessParameter.SetPS(Context.RHICmdList, ShaderRHI, Context, TStaticSamplerState<SF_Bilinear, AM_Clamp, AM_Clamp, AM_Clamp>::GetRHI());
 	
-		SetTextureParameter(RHICmdList, ShaderRHI, ZReadTexture, ZReadTextureSampler, TStaticSamplerState<SF_Point, AM_Wrap, AM_Wrap, AM_Wrap>::GetRHI(),
+		FSceneRenderTargets& SceneContext = FSceneRenderTargets::Get(Context.RHICmdList);
+
+		SetTextureParameter(Context.RHICmdList, ShaderRHI, ZReadTexture, ZReadTextureSampler, TStaticSamplerState<SF_Point, AM_Wrap, AM_Wrap, AM_Wrap>::GetRHI(),
 			SceneContext.ScreenSpaceGTAODepths->GetRenderTargetItem().ShaderResourceTexture);
 	}
 
@@ -2675,7 +2677,7 @@ public:
 
 		FSceneRenderTargets& SceneContext = FSceneRenderTargets::Get(Context.RHICmdList);
 
-		SetTextureParameter(RHICmdList, ShaderRHI, ZReadTexture, ZReadTextureSampler, TStaticSamplerState<SF_Point, AM_Wrap, AM_Wrap, AM_Wrap>::GetRHI(),
+		SetTextureParameter(Context.RHICmdList, ShaderRHI, ZReadTexture, ZReadTextureSampler, TStaticSamplerState<SF_Point, AM_Wrap, AM_Wrap, AM_Wrap>::GetRHI(),
 			SceneContext.ScreenSpaceGTAODepths->GetRenderTargetItem().ShaderResourceTexture);
 
 	}
