Perform fast-clear eliminate on distortion RT to fix reading invalid data left by previous renders.
#rnx

diff --git a/Engine/Source/Runtime/Renderer/Private/DistortionRendering.cpp b/Engine/Source/Runtime/Renderer/Private/DistortionRendering.cpp
--- a/Engine/Source/Runtime/Renderer/Private/DistortionRendering.cpp
+++ b/Engine/Source/Runtime/Renderer/Private/DistortionRendering.cpp
@@ -581,8 +581,7 @@ void FSceneRenderer::RenderDistortion(FRHICommandListImmediate& RHICmdList)
 
 		// Create a texture to store the resolved light attenuation values, and a render-targetable surface to hold the unresolved light attenuation values.
 		{
-			uint32 BaseFlags = RHISupportsRenderTargetWriteMask(GMaxRHIShaderPlatform) ? TexCreate_NoFastClearFinalize : TexCreate_None;
-			FPooledRenderTargetDesc Desc(FPooledRenderTargetDesc::Create2DDesc(SceneContext.GetBufferSizeXY(), PF_B8G8R8A8, FClearValueBinding::Transparent, BaseFlags, TexCreate_RenderTargetable, false));
+			FPooledRenderTargetDesc Desc(FPooledRenderTargetDesc::Create2DDesc(SceneContext.GetBufferSizeXY(), PF_B8G8R8A8, FClearValueBinding::Transparent, TexCreate_None, TexCreate_RenderTargetable, false));
 			Desc.Flags |= GFastVRamConfig.Distortion;
 			Desc.NumSamples = MSAACount;
 			GRenderTargetPool.FindFreeElement(RHICmdList, Desc, DistortionRT, TEXT("Distortion"));
@@ -649,12 +648,8 @@ void FSceneRenderer::RenderDistortion(FRHICommandListImmediate& RHICmdList)
 
 			if (bDirty)
 			{
-				// Use a metadata transition to skip the fast clear eliminate, since we don't need pixels with no stencil set to be cleared
-				RHICmdList.TransitionResource(
-					RHISupportsRenderTargetWriteMask(GMaxRHIShaderPlatform)
-						? EResourceTransitionAccess::EMetaData
-						: EResourceTransitionAccess::EReadable
-				, DistortionRT->GetRenderTargetItem().TargetableTexture );
+				// Ideally we skip the EliminateFastClear since we don't need pixels with no stencil set to be cleared
+				RHICmdList.TransitionResource(EResourceTransitionAccess::EReadable, DistortionRT->GetRenderTargetItem().TargetableTexture);
 				// to be able to observe results with VisualizeTexture
 				GVisualizeTexture.SetCheckPoint(RHICmdList, DistortionRT);
 			}
