Added low quality setting for Distance Field shadow (20 steps, no SSS)

diff --git a/Engine/Shaders/Private/DistanceFieldShadowing.usf b/Engine/Shaders/Private/DistanceFieldShadowing.usf
--- a/Engine/Shaders/Private/DistanceFieldShadowing.usf
+++ b/Engine/Shaders/Private/DistanceFieldShadowing.usf
@@ -569,7 +569,7 @@ void DistanceFieldShadowingCS(
 			float SubsurfaceDensity = 0;
 			bool bUseSubsurfaceTransmission = false;
 
-#if !FORWARD_SHADING && DF_SHADOW_QUALITY > 1
+#if !FORWARD_SHADING && DF_SHADOW_QUALITY > 2
 			FGBufferData GBufferData = GetGBufferData(ScreenUV);
 
 			BRANCH
diff --git a/Engine/Shaders/Private/DistanceFieldShadowingShared.ush b/Engine/Shaders/Private/DistanceFieldShadowingShared.ush
--- a/Engine/Shaders/Private/DistanceFieldShadowingShared.ush
+++ b/Engine/Shaders/Private/DistanceFieldShadowingShared.ush
@@ -119,10 +119,12 @@ float ShadowRayTraceThroughCulledObjects(
 			if (IntersectionTimes.x < IntersectionTimes.y)
 			{
 				float SampleRayTime = IntersectionTimes.x * VolumeRayLength;
-#if DF_SHADOW_QUALITY > 1
+#if DF_SHADOW_QUALITY == 3
 				uint MaxSteps = 64;
-#else
+#elif DF_SHADOW_QUALITY == 2
 				uint MaxSteps = 32;
+#else
+				uint MaxSteps = 20;
 #endif
 				float MinStepSize = 1.0f / (4 * MaxSteps);
 
diff --git a/Engine/Source/Runtime/Renderer/Private/DistanceFieldShadowing.cpp b/Engine/Source/Runtime/Renderer/Private/DistanceFieldShadowing.cpp
--- a/Engine/Source/Runtime/Renderer/Private/DistanceFieldShadowing.cpp
+++ b/Engine/Source/Runtime/Renderer/Private/DistanceFieldShadowing.cpp
@@ -36,12 +36,12 @@ FAutoConsoleVariableRef CVarDistanceFieldShadowing(
 	ECVF_Scalability | ECVF_RenderThreadSafe
 	);
 
-int32 GDFShadowQuality = 2;
+int32 GDFShadowQuality = 3;
 FAutoConsoleVariableRef CVarDFShadowQuality(
 	TEXT("r.DFShadowQuality"),
 	GDFShadowQuality,
 	TEXT("Defines the distance field shadow method which allows to adjust for quality or performance.\n")
-	TEXT(" 0:off, 1:medium (less samples, no SSS), 2:high (default)"),
+	TEXT(" 0:off, 1:low (20 steps, no SSS), 2:medium (32 steps, no SSS), 3:high (64 steps, SSS, default)"),
 	ECVF_Scalability | ECVF_RenderThreadSafe
 );
 
@@ -500,8 +500,10 @@ private:
 #define VARIATION(A) \
 	typedef TDistanceFieldShadowingCS<A, 1> FDistanceFieldShadowingCS##A##1; \
 	typedef TDistanceFieldShadowingCS<A, 2> FDistanceFieldShadowingCS##A##2; \
+    typedef TDistanceFieldShadowingCS<A, 3> FDistanceFieldShadowingCS##A##3; \
 	IMPLEMENT_SHADER_TYPE2(FDistanceFieldShadowingCS##A##1, SF_Compute); \
-	IMPLEMENT_SHADER_TYPE2(FDistanceFieldShadowingCS##A##2, SF_Compute);
+	IMPLEMENT_SHADER_TYPE2(FDistanceFieldShadowingCS##A##2, SF_Compute); \
+    IMPLEMENT_SHADER_TYPE2(FDistanceFieldShadowingCS##A##3, SF_Compute);
 
 VARIATION(DFS_DirectionalLightScatterTileCulling)
 VARIATION(DFS_DirectionalLightTiledCulling)
@@ -839,7 +841,7 @@ void CullDistanceFieldObjectsForLight(
 
 int32 GetDFShadowQuality()
 {
-	return FMath::Clamp(GDFShadowQuality, 0, 2);
+	return FMath::Clamp(GDFShadowQuality, 0, 3);
 }
 
 bool SupportsDistanceFieldShadows(ERHIFeatureLevel::Type FeatureLevel, EShaderPlatform ShaderPlatform)
@@ -881,7 +883,7 @@ bool FDeferredShadingSceneRenderer::ShouldPrepareForDistanceFieldShadows() const
 		&& SupportsDistanceFieldShadows(Scene->GetFeatureLevel(), Scene->GetShaderPlatform());
 }
 
-template<typename TRHICommandList, EDistanceFieldShadowingType DFSType, uint32 DFSQuality>
+template<typename TRHICommandList, EDistanceFieldShadowingType DFSType>
 void RayTraceShadowsDispatch(TRHICommandList& RHICmdList, const FViewInfo& View, FProjectedShadowInfo* ProjectedShadowInfo, FLightTileIntersectionResources* TileIntersectionResources)
 {
 	FIntRect ScissorRect;
@@ -893,12 +895,32 @@ void RayTraceShadowsDispatch(TRHICommandList& RHICmdList, const FViewInfo& View,
 	uint32 GroupSizeX = FMath::DivideAndRoundUp(ScissorRect.Size().X / GetDFShadowDownsampleFactor(), GDistanceFieldShadowTileSizeX);
 	uint32 GroupSizeY = FMath::DivideAndRoundUp(ScissorRect.Size().Y / GetDFShadowDownsampleFactor(), GDistanceFieldShadowTileSizeY);
 
-	TShaderMapRef<TDistanceFieldShadowingCS<DFSType, DFSQuality> > ComputeShader(View.ShaderMap);
-	RHICmdList.SetComputeShader(ComputeShader->GetComputeShader());
-	FSceneRenderTargetItem& RayTracedShadowsRTI = ProjectedShadowInfo->RayTracedShadowsRT->GetRenderTargetItem();
-	ComputeShader->SetParameters(RHICmdList, View, ProjectedShadowInfo, RayTracedShadowsRTI, FVector2D(GroupSizeX, GroupSizeY), ScissorRect, TileIntersectionResources);
-	DispatchComputeShader(RHICmdList, *ComputeShader, GroupSizeX, GroupSizeY, 1);
-	ComputeShader->UnsetParameters(RHICmdList, RayTracedShadowsRTI);
+	auto DispatchTemplatedCS = [&](auto CS)
+	{
+		RHICmdList.SetComputeShader(CS->GetComputeShader());
+		FSceneRenderTargetItem& RayTracedShadowsRTI = ProjectedShadowInfo->RayTracedShadowsRT->GetRenderTargetItem();
+		CS->SetParameters(RHICmdList, View, ProjectedShadowInfo, RayTracedShadowsRTI, FVector2D(GroupSizeX, GroupSizeY), ScissorRect, TileIntersectionResources);
+		DispatchComputeShader(RHICmdList, *CS, GroupSizeX, GroupSizeY, 1);
+		CS->UnsetParameters(RHICmdList, RayTracedShadowsRTI);
+	};
+
+	int32 const DFShadowQuality = GetDFShadowQuality();
+	
+	if (DFShadowQuality == 1)
+	{
+		TShaderMapRef<TDistanceFieldShadowingCS<DFSType, 1> > ComputeShader(View.ShaderMap);
+		DispatchTemplatedCS(ComputeShader);
+	}
+	else if (DFShadowQuality == 2)
+	{
+		TShaderMapRef<TDistanceFieldShadowingCS<DFSType, 2> > ComputeShader(View.ShaderMap);
+		DispatchTemplatedCS(ComputeShader);
+	}
+	else if (DFShadowQuality == 3)
+	{ 
+		TShaderMapRef<TDistanceFieldShadowingCS<DFSType, 3> > ComputeShader(View.ShaderMap);
+		DispatchTemplatedCS(ComputeShader);
+	}
 }
 
 template<typename TRHICommandList>
@@ -907,36 +929,15 @@ void RayTraceShadows(TRHICommandList& RHICmdList, const FViewInfo& View, FProjec
 	int32 const DFShadowQuality = GetDFShadowQuality();
 	if (ProjectedShadowInfo->bDirectionalLight && GShadowScatterTileCulling)
 	{
-		if (DFShadowQuality == 1)
-		{
-			RayTraceShadowsDispatch<TRHICommandList, DFS_DirectionalLightScatterTileCulling, 1>(RHICmdList, View, ProjectedShadowInfo, TileIntersectionResources);
-		}
-		else
-		{
-			RayTraceShadowsDispatch<TRHICommandList, DFS_DirectionalLightScatterTileCulling, 2>(RHICmdList, View, ProjectedShadowInfo, TileIntersectionResources);
-		}
+		RayTraceShadowsDispatch<TRHICommandList, DFS_DirectionalLightScatterTileCulling>(RHICmdList, View, ProjectedShadowInfo, TileIntersectionResources);
 	}
 	else if (ProjectedShadowInfo->bDirectionalLight)
 	{
-		if (DFShadowQuality == 1)
-		{
-			RayTraceShadowsDispatch<TRHICommandList, DFS_DirectionalLightTiledCulling, 1>(RHICmdList, View, ProjectedShadowInfo, TileIntersectionResources);
-		}
-		else
-		{
-			RayTraceShadowsDispatch<TRHICommandList, DFS_DirectionalLightTiledCulling, 2>(RHICmdList, View, ProjectedShadowInfo, TileIntersectionResources);
-		}
+		RayTraceShadowsDispatch<TRHICommandList, DFS_DirectionalLightTiledCulling>(RHICmdList, View, ProjectedShadowInfo, TileIntersectionResources);
 	}
 	else
 	{
-		if (DFShadowQuality == 1)
-		{
-			RayTraceShadowsDispatch<TRHICommandList, DFS_PointLightTiledCulling, 1>(RHICmdList, View, ProjectedShadowInfo, TileIntersectionResources);
-		}
-		else
-		{
-			RayTraceShadowsDispatch<TRHICommandList, DFS_PointLightTiledCulling, 2>(RHICmdList, View, ProjectedShadowInfo, TileIntersectionResources);
-		}
+		RayTraceShadowsDispatch<TRHICommandList, DFS_PointLightTiledCulling>(RHICmdList, View, ProjectedShadowInfo, TileIntersectionResources);
 	}
 }
 
