Fix DX shader validation issue.

(ushell-p4-cherrypick of 9442457 by Charles.deRousiers)

diff --git a/Engine/Shaders/Private/HairStrands/HairStrandsDeepTransmittanceMask.usf b/Engine/Shaders/Private/HairStrands/HairStrandsDeepTransmittanceMask.usf
--- a/Engine/Shaders/Private/HairStrands/HairStrandsDeepTransmittanceMask.usf
+++ b/Engine/Shaders/Private/HairStrands/HairStrandsDeepTransmittanceMask.usf
@@ -53,7 +53,7 @@ Texture3D<uint> 	Voxel_DensityTexture7;
 
 StructuredBuffer<FNodeData> HairVisibilityNodeData;
 StructuredBuffer<uint>		HairVisibilityNodeCoord;
-StructuredBuffer<uint4>		IndirectArgsBuffer;
+Buffer<uint4>				IndirectArgsBuffer;
 
 Texture3D<float4>			HairLUTTexture;
 RWStructuredBuffer<float4>	OutputColor;
diff --git a/Engine/Shaders/Private/HairStrands/HairStrandsVisibilityDepthPS.usf b/Engine/Shaders/Private/HairStrands/HairStrandsVisibilityDepthPS.usf
--- a/Engine/Shaders/Private/HairStrands/HairStrandsVisibilityDepthPS.usf
+++ b/Engine/Shaders/Private/HairStrands/HairStrandsVisibilityDepthPS.usf
@@ -12,8 +12,8 @@ Texture2DMS<float, HAIR_VISIBILITY_MAXSAMPLECOUNT> HairVisibilityDepthTexture;
 void MainPS(
 	in FScreenVertexOutput Input,
 	out float SvDepth		: SV_Depth,
-	out uint4  OutGBufferB	: SV_Target0,
-	out uint4  OutColor		: SV_Target1
+	out float4  OutGBufferB	: SV_Target0,
+	out float4 OutColor		: SV_Target1
 )
 {
 	float MaxDepth = 0;
diff --git a/Engine/Source/Runtime/Renderer/Private/HairStrands/HairStrandsTransmittance.cpp b/Engine/Source/Runtime/Renderer/Private/HairStrands/HairStrandsTransmittance.cpp
--- a/Engine/Source/Runtime/Renderer/Private/HairStrands/HairStrandsTransmittance.cpp
+++ b/Engine/Source/Runtime/Renderer/Private/HairStrands/HairStrandsTransmittance.cpp
@@ -449,7 +449,10 @@ static FHairStrandsTransmittanceMaskData RenderHairStrandsTransmittanceMask(
 		memset(Params.Cluster_MinAABBs, 0, sizeof(FVector4) * FHairStrandsDeepShadowData::MaxClusterCount);
 		memset(Params.Cluster_MaxAABBs, 0, sizeof(FVector4) * FHairStrandsDeepShadowData::MaxClusterCount);
 
-		FRDGTextureRef DefaultDensityTexture = GraphBuilder.RegisterExternalTexture(GSystemTextures.BlackDummy, TEXT("Voxel_DefaultDensityTexture"));
+		TRefCountPtr<IPooledRenderTarget> DummyVoxelResources;
+		FPooledRenderTargetDesc Desc(FPooledRenderTargetDesc::CreateVolumeDesc(1, 1, 1, PF_R32_UINT, FClearValueBinding::Black, TexCreate_None, TexCreate_UAV | TexCreate_ShaderResource, false, 1));
+		GRenderTargetPool.FindFreeElement(RHICmdList, Desc, DummyVoxelResources, TEXT("DummyDensityTexture"));
+		FRDGTextureRef DefaultDensityTexture = GraphBuilder.RegisterExternalTexture(DummyVoxelResources, TEXT("Voxel_DefaultDensityTexture"));
 		for (uint32 TexIt = 0; TexIt < FHairStrandsDeepShadowData::MaxClusterCount; ++TexIt)
 		{
 			Params.Cluster_DensityTextures[TexIt] = DefaultDensityTexture;
