* Plugged hair AO to the correct placed (instead of overriting the scene color).
 * Plugged user control on AO (power/instensity)

(ushell-p4-cherrypick of 9451778 by Charles.deRousiers)

diff --git a/Engine/Shaders/Private/HairStrands/HairStrandsEnvironmentLighting.usf b/Engine/Shaders/Private/HairStrands/HairStrandsEnvironmentLighting.usf
--- a/Engine/Shaders/Private/HairStrands/HairStrandsEnvironmentLighting.usf
+++ b/Engine/Shaders/Private/HairStrands/HairStrandsEnvironmentLighting.usf
@@ -53,6 +53,9 @@ float Voxel_TanConeAngle;
 float Voxel_DensityScale;
 float Voxel_DepthBiasScale;
 
+float AO_Power;
+float AO_Intensity;
+
 float3 Voxel_MinAABB;
 float3 Voxel_MaxAABB;
 uint   Voxel_Resolution;
@@ -417,7 +420,12 @@ float3 TraceAO(float3 WorldPosition, float3 WorldNormal, float4 SvPosition)
 	Occlusion /= SampleCount;
 	BentNormal /= SampleCount;
 
-	return saturate(Occlusion);
+	// User adjust AO
+	// abs() to prevent shader warning
+	float AO = 1 - Occlusion;
+	AO = 1 - (1 - pow(abs(AO), AO_Power)) * AO_Intensity;
+
+	return saturate(AO);
 }
 
 //////////////////////////////////////////////////////////////////////////////////////////////////////////////////////////
@@ -434,13 +442,15 @@ float3 TotWorldPosition(float2 UV, float Depth)
 
 void MainPS(
 	in float4 SvPosition : SV_Position,
-	out float4 OutColor : SV_Target0
 #if PERMUTATION_RENDER_MODE == RENDER_MODE_LIGHTING
-	, out float4 OutColorSubPixel : SV_Target1
+	out float4 OutColor : SV_Target0,
+	out float4 OutColorSubPixel : SV_Target1
+#endif
+#if PERMUTATION_RENDER_MODE == RENDER_MODE_AO
+	out float4 OutAO : SV_Target0
 #endif
 )
 {
-	OutColor = 0.0f;
 	const float2 BufferUV = SvPositionToBufferUV(SvPosition);
 	const float2 ScreenPosition = SvPositionToScreenPosition(SvPosition).xy;
 	const uint3 PixelCoord = uint3(floor(SvPosition.xy), 0);
@@ -448,6 +458,7 @@ void MainPS(
 	FNodeDesc NodeDesc = DecodeNodeDesc(HairVisibilityNodeOffsetAndCount.Load(PixelCoord));
 
 #if PERMUTATION_RENDER_MODE == RENDER_MODE_LIGHTING
+	OutColor = 0.0f;
 	OutColorSubPixel = 0.0f;
 
 	if (NodeDesc.Count > 0)
@@ -521,6 +532,8 @@ void MainPS(
 
 
 #if PERMUTATION_RENDER_MODE == RENDER_MODE_AO
+	OutAO = 1;
+
 	const FGBufferData GBuffer = GetGBufferDataFromSceneTextures(BufferUV);
 	const float3 WorldPosition = TotWorldPosition(BufferUV, GBuffer.Depth);
 	const float3 WorldNormal = GBuffer.WorldNormal;
@@ -555,7 +568,7 @@ void MainPS(
 	if (bIsValid)
 	{
 		const float AO = TraceAO(WorldPosition, WorldNormal, SvPosition);
-		OutColor = float4(0, 0, 0, 1-AO);
+		OutAO = AO;
 	}
 	else
 	{
diff --git a/Engine/Source/Runtime/Renderer/Private/DeferredShadingRenderer.cpp b/Engine/Source/Runtime/Renderer/Private/DeferredShadingRenderer.cpp
--- a/Engine/Source/Runtime/Renderer/Private/DeferredShadingRenderer.cpp
+++ b/Engine/Source/Runtime/Renderer/Private/DeferredShadingRenderer.cpp
@@ -1899,6 +1899,15 @@ void FDeferredShadingSceneRenderer::Render(FRHICommandListImmediate& RHICmdList)
 		ServiceLocalQueue();
 
 		HairDatas = &HairDatasStorage;
+
+		if (SceneContext.bScreenSpaceAOIsValid && SceneContext.ScreenSpaceAO)
+		{
+			RenderHairStrandsAmbientOcclusion(
+				RHICmdList,
+				Views,
+				HairDatas,
+				SceneContext.ScreenSpaceAO);
+		}
 	}
 
 	// Render lighting.
diff --git a/Engine/Source/Runtime/Renderer/Private/HairStrands/HairStrandsEnvironment.cpp b/Engine/Source/Runtime/Renderer/Private/HairStrands/HairStrandsEnvironment.cpp
--- a/Engine/Source/Runtime/Renderer/Private/HairStrands/HairStrandsEnvironment.cpp
+++ b/Engine/Source/Runtime/Renderer/Private/HairStrands/HairStrandsEnvironment.cpp
@@ -16,13 +16,14 @@
 #include "HairStrandsRendering.h"
 #include "ReflectionEnvironment.h"
 #include "ScenePrivate.h"
+#include "RenderGraphEvent.h"
 
 ///////////////////////////////////////////////////////////////////////////////////////////////////
 
 static int32 GHairSkylightingEnable = 1;
 static FAutoConsoleVariableRef CVarHairSkylightingEnable(TEXT("r.HairStrands.SkyLightingEnable"), GHairSkylightingEnable, TEXT("Enable sky lighting on hair."));
 
-static int32 GHairSkyAOEnable = 0;
+static int32 GHairSkyAOEnable = 1;
 static FAutoConsoleVariableRef CVarHairSkyAOEnable(TEXT("r.HairStrands.SkyAOEnable"), GHairSkyAOEnable, TEXT("Enable (sky) AO on hair."));
 
 static float GHairSkylightingConeAngle = 3;
@@ -58,6 +59,9 @@ class FHairEnvironmentLightingPS : public FGlobalShader
 		SHADER_PARAMETER(float, Voxel_DepthBiasScale)
 		SHADER_PARAMETER(float, Voxel_TanConeAngle)
 
+		SHADER_PARAMETER(float, AO_Power)
+		SHADER_PARAMETER(float, AO_Intensity)
+
 		SHADER_PARAMETER_TEXTURE(Texture2D, PreIntegratedGF)
 		SHADER_PARAMETER_SAMPLER(SamplerState, PreIntegratedGFSampler)
 
@@ -85,6 +89,140 @@ class FHairEnvironmentLightingPS : public FGlobalShader
 
 IMPLEMENT_GLOBAL_SHADER(FHairEnvironmentLightingPS, "/Engine/Private/HairStrands/HairStrandsEnvironmentLighting.usf", "MainPS", SF_Pixel);
 
+enum class EEnvRenderMode : uint8
+{
+	Lighting,
+	AO
+};
+
+static void AddHairStrandsEnvironmentPass(
+	FRDGBuilder& GraphBuilder,
+	const FViewInfo& View,
+	EEnvRenderMode RenderMode,
+	const FHairStrandsVisibilityData* HairVisibilityData,
+	const FHairStrandsClusterData* ClusterData,
+	FRDGTextureRef Output0,
+	FRDGTextureRef Output1)
+{
+	FSceneTextureParameters SceneTextures;
+	SetupSceneTextureParameters(GraphBuilder, &SceneTextures);
+
+	// Render the reflection environment with tiled deferred culling
+	const bool bHasBoxCaptures = (View.NumBoxReflectionCaptures > 0);
+	const bool bHasSphereCaptures = (View.NumSphereReflectionCaptures > 0);
+	FHairEnvironmentLightingPS::FParameters* PassParameters = GraphBuilder.AllocParameters<FHairEnvironmentLightingPS::FParameters>();
+
+	const FHairLUT InHairLUT = GetHairLUT(GraphBuilder.RHICmdList, View);
+	PassParameters->HairEnergyLUTTexture = GraphBuilder.RegisterExternalTexture(InHairLUT.Textures[HairLUTType_MeanEnergy], TEXT("HairMeanEnergyLUTTexture"));;
+	PassParameters->HairScatteringLUTTexture = GraphBuilder.RegisterExternalTexture(InHairLUT.Textures[HairLUTType_DualScattering], TEXT("HairScatteringEnergyLUTTexture"));
+	PassParameters->HairLUTSampler = TStaticSamplerState<SF_Bilinear, AM_Clamp, AM_Clamp, AM_Clamp>::GetRHI();
+
+	if (ClusterData)
+	{
+		PassParameters->Voxel_MinAABB = ClusterData->GetMinBound();
+		PassParameters->Voxel_MaxAABB = ClusterData->GetMaxBound();
+		PassParameters->Voxel_Resolution = ClusterData->GetResolution();
+		PassParameters->Voxel_DensityTexture = GraphBuilder.RegisterExternalTexture(ClusterData->VoxelResources.DensityTexture);
+		PassParameters->Voxel_DensityScale = GetHairStrandsVoxelizationDensityScale();
+		PassParameters->Voxel_DepthBiasScale = GetHairStrandsVoxelizationDepthBiasScale();
+		PassParameters->Voxel_TanConeAngle = FMath::Tan(FMath::DegreesToRadians(GetHairStrandsSkyLightingConeAngle()));
+	}
+	else
+	{
+		// .. todo dummy 3D texture
+		PassParameters->Voxel_DensityTexture = GraphBuilder.RegisterExternalTexture(GSystemTextures.BlackDummy);
+	}
+
+	PassParameters->PreIntegratedGF = GSystemTextures.PreintegratedGF->GetRenderTargetItem().ShaderResourceTexture;
+	PassParameters->PreIntegratedGFSampler = TStaticSamplerState<SF_Bilinear, AM_Clamp, AM_Clamp, AM_Clamp>::GetRHI();
+
+	PassParameters->SceneTextures = SceneTextures;
+	SetupSceneTextureSamplers(&PassParameters->SceneTextureSamplers);
+
+	PassParameters->ViewUniformBuffer = View.ViewUniformBuffer;
+	PassParameters->ReflectionCaptureData = View.ReflectionCaptureUniformBuffer;
+	{
+		FReflectionUniformParameters ReflectionUniformParameters;
+		SetupReflectionUniformParameters(View, ReflectionUniformParameters);
+		PassParameters->ReflectionsParameters = CreateUniformBufferImmediate(ReflectionUniformParameters, UniformBuffer_SingleDraw);
+	}
+	PassParameters->ForwardLightData = View.ForwardLightingResources->ForwardLightDataUniformBuffer;
+
+
+	// Bind hair data
+	if (HairVisibilityData)
+	{
+		PassParameters->HairCategorizationTexture = HairVisibilityData->CategorizationTexture->GetRenderTargetItem().ShaderResourceTexture;
+		PassParameters->HairVisibilityNodeOffsetAndCount = HairVisibilityData->NodeIndex->GetRenderTargetItem().ShaderResourceTexture;
+		PassParameters->HairVisibilityNodeData = HairVisibilityData->NodeDataSRV;
+	}
+
+	PassParameters->AO_Power = 0;
+	PassParameters->AO_Intensity = 0;
+	if (RenderMode == EEnvRenderMode::AO)
+	{
+		const FFinalPostProcessSettings& Settings = View.FinalPostProcessSettings;
+		PassParameters->AO_Power = Settings.AmbientOcclusionPower;
+		PassParameters->AO_Intensity = Settings.AmbientOcclusionIntensity;
+	}
+
+	check(Output0);
+	PassParameters->RenderTargets[0] = FRenderTargetBinding(Output0, ERenderTargetLoadAction::ELoad);
+	if (Output1)
+	{
+		PassParameters->RenderTargets[1] = FRenderTargetBinding(Output1, ERenderTargetLoadAction::ELoad);
+	}
+
+	FHairEnvironmentLightingPS::FPermutationDomain PermutationVector;
+	PermutationVector.Set<FHairEnvironmentLightingPS::FRenderMode>(RenderMode == EEnvRenderMode::AO ? 1 : 0);
+	TShaderMapRef<FHairEnvironmentLightingPS> PixelShader(View.ShaderMap, PermutationVector);
+	ClearUnusedGraphResources(*PixelShader, PassParameters);
+
+	GraphBuilder.AddPass(
+		(RenderMode == EEnvRenderMode::Lighting) ? 
+		RDG_EVENT_NAME("HairStrandsEnvironment %dx%d", View.ViewRect.Width(), View.ViewRect.Height()) : 
+		RDG_EVENT_NAME("HairStrandsAO %dx%d", View.ViewRect.Width(), View.ViewRect.Height()),
+		PassParameters,
+		ERDGPassFlags::Raster,
+		[PassParameters, &View, PixelShader, RenderMode](FRHICommandList& InRHICmdList)
+	{
+		InRHICmdList.SetViewport(View.ViewRect.Min.X, View.ViewRect.Min.Y, 0.0f, View.ViewRect.Max.X, View.ViewRect.Max.Y, 1.0f);
+
+		FGraphicsPipelineStateInitializer GraphicsPSOInit;
+		FPixelShaderUtils::InitFullscreenPipelineState(InRHICmdList, View.ShaderMap, *PixelShader, GraphicsPSOInit);
+
+		if (RenderMode == EEnvRenderMode::AO)
+		{
+			GraphicsPSOInit.BlendState = TStaticBlendState<
+			CW_RGBA, BO_Min, BF_SourceColor, BF_DestColor, BO_Add, BF_Zero, BF_DestAlpha
+			>::GetRHI();
+		}
+		else //if (RenderMode == EEnvRenderMode::Lighting)
+		{
+			GraphicsPSOInit.BlendState = TStaticBlendState<
+			CW_RGBA, BO_Add, BF_One, BF_One, BO_Add, BF_One, BF_One,
+			CW_RGBA, BO_Add, BF_One, BF_One, BO_Max, BF_SourceAlpha, BF_DestAlpha
+			>::GetRHI();
+		}
+
+		SetGraphicsPipelineState(InRHICmdList, GraphicsPSOInit);
+		SetShaderParameters(InRHICmdList, *PixelShader, PixelShader->GetPixelShader(), *PassParameters);
+		FPixelShaderUtils::DrawFullscreenTriangle(InRHICmdList);
+	});
+}
+
+static const FHairStrandsClusterData* GetClusterData(const FHairStrandsDatas* HairDatas, uint32 ViewIndex, uint32 ClusterIndex)
+{
+	return
+		(
+			HairDatas && 
+			ViewIndex < uint32(HairDatas->HairClusterPerViews.Views.Num()) &&
+			ClusterIndex < uint32(HairDatas->HairClusterPerViews.Views[ViewIndex].Datas.Num())
+		) ?
+		&HairDatas->HairClusterPerViews.Views[ViewIndex].Datas[ClusterIndex] :
+		nullptr;
+}
+
 void RenderHairStrandsEnvironmentLighting(
 	FRDGBuilder& GraphBuilder,
 	const uint32 ViewIndex,
@@ -93,125 +231,56 @@ void RenderHairStrandsEnvironmentLighting(
 	FRDGTextureRef SceneColorTexture,
 	FRDGTextureRef SceneColorSubPixelTexture)
 {
+	if (!GetHairStrandsSkyLightingEnable())
+		return;
+
+	check(ViewIndex < uint32(Views.Num()));
 	const FViewInfo& View = Views[ViewIndex];
 	const FHairStrandsVisibilityData* HairVisibilityData = HairDatas ? &HairDatas->HairVisibilityViews.HairDatas[ViewIndex] : nullptr;
-	const bool bRenderHairLighting = HairVisibilityData && HairVisibilityData->NodeIndex && HairVisibilityData->NodeDataSRV && (GetHairStrandsSkyLightingEnable() || GetHairStrandsSkyAOEnable());
+	const bool bRenderHairLighting = HairVisibilityData && HairVisibilityData->NodeIndex && HairVisibilityData->NodeDataSRV;
 	if (!bRenderHairLighting)
 	{
 		return;
 	}
 
+	// @hair_todo: add support for multiple clusters on screen
+	const FHairStrandsClusterData* ClusterData = GetClusterData(HairDatas, ViewIndex, 0);
+
 	// @hair_todo: 
 	// * Add support for : BentNormal, global AO, SSR, DFSO
 	// * Add local reflection probe, current take into account only the sky lighting
+	//RDG_GPU_STAT_SCOPE(GraphBuilder, HairStrandEnvironment);
+	AddHairStrandsEnvironmentPass(GraphBuilder, View, EEnvRenderMode::Lighting, HairVisibilityData, ClusterData, SceneColorTexture, SceneColorSubPixelTexture);
+}
 
-	FSceneTextureParameters SceneTextures;
-	SetupSceneTextureParameters(GraphBuilder, &SceneTextures);
-
-	RDG_GPU_STAT_SCOPE(GraphBuilder, HairStrandsReflectionEnvironment);
-
-	enum class EEnvRenderMode : uint8
-	{
-		Lighting,
-		AO
-	};
-	const FHairLUT InHairLUT = GetHairLUT(GraphBuilder.RHICmdList, View);
+void RenderHairStrandsAmbientOcclusion(
+	FRHICommandListImmediate& RHICmdList,
+	const TArray<FViewInfo>& Views,
+	const FHairStrandsDatas* HairDatas,
+	const TRefCountPtr<IPooledRenderTarget>& InAOTexture)
+{
+	if (!GetHairStrandsSkyAOEnable() || Views.Num() == 0 || !InAOTexture)
+		return;
 
-	// Setup the parameters of the shader.
-	auto RenderEnv = [&](EEnvRenderMode RenderMode)
+	for (uint32 ViewIndex = 0; ViewIndex < uint32(Views.Num()); ++ViewIndex)
 	{
-		// Render the reflection environment with tiled deferred culling
-		const bool bHasBoxCaptures = (View.NumBoxReflectionCaptures > 0);
-		const bool bHasSphereCaptures = (View.NumSphereReflectionCaptures > 0);
-		FHairEnvironmentLightingPS::FParameters* PassParameters = GraphBuilder.AllocParameters<FHairEnvironmentLightingPS::FParameters>();
-				
-		PassParameters->HairEnergyLUTTexture = GraphBuilder.RegisterExternalTexture(InHairLUT.Textures[HairLUTType_MeanEnergy], TEXT("HairMeanEnergyLUTTexture"));;
-		PassParameters->HairScatteringLUTTexture = GraphBuilder.RegisterExternalTexture(InHairLUT.Textures[HairLUTType_DualScattering], TEXT("HairScatteringEnergyLUTTexture"));
-		PassParameters->HairLUTSampler = TStaticSamplerState<SF_Bilinear, AM_Clamp, AM_Clamp, AM_Clamp>::GetRHI();
-
-		if (HairDatas && ViewIndex < uint32(HairDatas->HairClusterPerViews.Views.Num()) && HairDatas->HairClusterPerViews.Views[ViewIndex].Datas.Num() > 0)
+		const FViewInfo& View = Views[ViewIndex];
+		const FHairStrandsVisibilityData* HairVisibilityData = HairDatas ? &HairDatas->HairVisibilityViews.HairDatas[ViewIndex] : nullptr;
+		const bool bRenderHairLighting = HairVisibilityData && HairVisibilityData->NodeIndex && HairVisibilityData->NodeDataSRV;
+		if (!bRenderHairLighting)
 		{
-			const uint32 ClusterIndex = 0; // Support only one cluster at the moment
-			const FHairStrandsClusterData& ClusterData = HairDatas->HairClusterPerViews.Views[ViewIndex].Datas[ClusterIndex];
-
-			PassParameters->Voxel_MinAABB = ClusterData.GetMinBound();
-			PassParameters->Voxel_MaxAABB = ClusterData.GetMaxBound();
-			PassParameters->Voxel_Resolution = ClusterData.GetResolution();
-			PassParameters->Voxel_DensityTexture = GraphBuilder.RegisterExternalTexture(ClusterData.VoxelResources.DensityTexture);
-			PassParameters->Voxel_DensityScale = GetHairStrandsVoxelizationDensityScale();
-			PassParameters->Voxel_DepthBiasScale = GetHairStrandsVoxelizationDepthBiasScale();
-			PassParameters->Voxel_TanConeAngle = FMath::Tan(FMath::DegreesToRadians(GetHairStrandsSkyLightingConeAngle()));
-		}
-		else
-		{
-			PassParameters->Voxel_DensityTexture = GraphBuilder.RegisterExternalTexture(GSystemTextures.BlackDummy);
+			return;
 		}
 
 
-		PassParameters->PreIntegratedGF = GSystemTextures.PreintegratedGF->GetRenderTargetItem().ShaderResourceTexture;
-		PassParameters->PreIntegratedGFSampler = TStaticSamplerState<SF_Bilinear, AM_Clamp, AM_Clamp, AM_Clamp>::GetRHI();
+		FRDGBuilder GraphBuilder(RHICmdList);
+		FRDGTextureRef AOTexture = InAOTexture ? GraphBuilder.RegisterExternalTexture(InAOTexture, TEXT("AOTexture")) : nullptr;
 
-		PassParameters->SceneTextures = SceneTextures;
-		SetupSceneTextureSamplers(&PassParameters->SceneTextureSamplers);
+		// @hair_todo: add support for multiple clusters on screen
+		const FHairStrandsClusterData* ClusterData = GetClusterData(HairDatas, ViewIndex, 0);
 
-		PassParameters->ViewUniformBuffer = View.ViewUniformBuffer;
-		PassParameters->ReflectionCaptureData = View.ReflectionCaptureUniformBuffer;
-		{
-			FReflectionUniformParameters ReflectionUniformParameters;
-			SetupReflectionUniformParameters(View, ReflectionUniformParameters);
-			PassParameters->ReflectionsParameters = CreateUniformBufferImmediate(ReflectionUniformParameters, UniformBuffer_SingleDraw);
-		}
-		PassParameters->ForwardLightData = View.ForwardLightingResources->ForwardLightDataUniformBuffer;
+		AddHairStrandsEnvironmentPass(GraphBuilder, View, EEnvRenderMode::AO, HairVisibilityData, ClusterData, AOTexture, nullptr);
 
-
-		// Bind hair data
-		if (HairVisibilityData)
-		{
-			PassParameters->HairCategorizationTexture = HairVisibilityData->CategorizationTexture->GetRenderTargetItem().ShaderResourceTexture;
-			PassParameters->HairVisibilityNodeOffsetAndCount = HairVisibilityData->NodeIndex->GetRenderTargetItem().ShaderResourceTexture;
-			PassParameters->HairVisibilityNodeData = HairVisibilityData->NodeDataSRV;
-		}
-
-		PassParameters->RenderTargets[0] = FRenderTargetBinding(SceneColorTexture, ERenderTargetLoadAction::ELoad);
-		if (RenderMode == EEnvRenderMode::Lighting)
-			PassParameters->RenderTargets[1] = FRenderTargetBinding(SceneColorSubPixelTexture, ERenderTargetLoadAction::ELoad);
-
-		FHairEnvironmentLightingPS::FPermutationDomain PermutationVector;
-		PermutationVector.Set<FHairEnvironmentLightingPS::FRenderMode>(RenderMode == EEnvRenderMode::AO ? 1 : 0);
-		TShaderMapRef<FHairEnvironmentLightingPS> PixelShader(View.ShaderMap, PermutationVector);
-		ClearUnusedGraphResources(*PixelShader, PassParameters);
-
-		GraphBuilder.AddPass(
-			RDG_EVENT_NAME("HairStrandsEnvironment %dx%d", View.ViewRect.Width(), View.ViewRect.Height()),
-			PassParameters,
-			ERDGPassFlags::Raster,
-			[PassParameters, &View, PixelShader, bRenderHairLighting, RenderMode](FRHICommandList& InRHICmdList)
-		{
-			InRHICmdList.SetViewport(View.ViewRect.Min.X, View.ViewRect.Min.Y, 0.0f, View.ViewRect.Max.X, View.ViewRect.Max.Y, 1.0f);
-
-			FGraphicsPipelineStateInitializer GraphicsPSOInit;
-			FPixelShaderUtils::InitFullscreenPipelineState(InRHICmdList, View.ShaderMap, *PixelShader, GraphicsPSOInit);
-
-			if (RenderMode == EEnvRenderMode::AO)
-				GraphicsPSOInit.BlendState = TStaticBlendState<
-				CW_RGBA, BO_Add, BF_Zero, BF_SourceAlpha, BO_Add, BF_Zero, BF_One
-				>::GetRHI();
-			else //if (RenderMode == EEnvRenderMode::Lighting)
-				GraphicsPSOInit.BlendState = TStaticBlendState<
-				CW_RGBA, BO_Add, BF_One, BF_One, BO_Add, BF_One, BF_One,
-				CW_RGBA, BO_Add, BF_One, BF_One, BO_Max, BF_SourceAlpha, BF_DestAlpha
-				>::GetRHI();
-
-			SetGraphicsPipelineState(InRHICmdList, GraphicsPSOInit);
-			SetShaderParameters(InRHICmdList, *PixelShader, PixelShader->GetPixelShader(), *PassParameters);
-			FPixelShaderUtils::DrawFullscreenTriangle(InRHICmdList);
-		});
-	}; // hair
-
-	// @hair_todo: move hair AO actually into the AO pass.
-	if (GetHairStrandsSkyAOEnable())
-		RenderEnv(EEnvRenderMode::AO);
-
-	if (GetHairStrandsSkyLightingEnable())
-		RenderEnv(EEnvRenderMode::Lighting);
+		GraphBuilder.Execute();
+	}
 }
diff --git a/Engine/Source/Runtime/Renderer/Private/HairStrands/HairStrandsEnvironment.h b/Engine/Source/Runtime/Renderer/Private/HairStrands/HairStrandsEnvironment.h
--- a/Engine/Source/Runtime/Renderer/Private/HairStrands/HairStrandsEnvironment.h
+++ b/Engine/Source/Runtime/Renderer/Private/HairStrands/HairStrandsEnvironment.h
@@ -9,7 +9,13 @@
 #include "CoreMinimal.h"
 #include "RendererInterface.h"
 #include "RenderGraphResources.h"
-#
+
+void RenderHairStrandsAmbientOcclusion(
+	FRHICommandListImmediate& RHICmdList,
+	const TArray<FViewInfo>& Views,
+	const struct FHairStrandsDatas* HairDatas,
+	const TRefCountPtr<IPooledRenderTarget>& InAOTexture);
+
 void RenderHairStrandsEnvironmentLighting(
 	FRDGBuilder& GraphBuilder,
 	const uint32 ViewIndex,
