Add rotation component into groom skin.

(ushell-p4-cherrypick of 9478990 by Charles.deRousiers)

diff --git a/Engine/Plugins/Experimental/HairStrands/Source/HairStrandsCore/Private/HairStrandsRendering.cpp b/Engine/Plugins/Experimental/HairStrands/Source/HairStrandsCore/Private/HairStrandsRendering.cpp
--- a/Engine/Plugins/Experimental/HairStrands/Source/HairStrandsCore/Private/HairStrandsRendering.cpp
+++ b/Engine/Plugins/Experimental/HairStrands/Source/HairStrandsCore/Private/HairStrandsRendering.cpp
@@ -114,6 +114,7 @@ class FHairInterpolationCS : public FGlobalShader
 	BEGIN_SHADER_PARAMETER_STRUCT(FParameters, )
 		SHADER_PARAMETER(uint32, VertexCount)
 		SHADER_PARAMETER(uint32, DispatchCountX)
+		SHADER_PARAMETER(FVector, HairWorldOffset)
 
 		SHADER_PARAMETER(FVector, RestPositionWorldCenter)
 		SHADER_PARAMETER(FVector, DeformedPositionWorldCenter)
@@ -152,6 +153,7 @@ IMPLEMENT_GLOBAL_SHADER(FHairInterpolationCS, "/Engine/Private/HairStrands/HairS
 static void AddHairStrandsInterpolationPass(
 	FRDGBuilder& GraphBuilder,
 	const FHairStrandsProjectionHairData::HairGroup& InHairData,
+	const FVector& HairWorldOffset, 
 	const int32 LODIndex,
 	const uint32 VertexCount,
 	const FShaderResourceViewRHIRef& RenderRestPosePositionBuffer,
@@ -180,6 +182,7 @@ static void AddHairStrandsInterpolationPass(
 		Parameters->OutRenderAttributeBuffer = OutRenderAttributeBuffer;
 	}
 	Parameters->VertexCount = VertexCount;
+	Parameters->HairWorldOffset = HairWorldOffset;
 	Parameters->DispatchCountX = DispatchCount.X;
 	Parameters->OutRenderDeformedPositionCenter = FVector::ZeroVector;
 
@@ -447,6 +450,7 @@ void ComputeHairStrandsInterpolation(
 			AddHairStrandsInterpolationPass(
 				GraphBuilder,
 				InHairDatas.HairGroups[GroupIndex],
+				Input.HairWorldOffset,
 				LODIndex,
 				Input.RenderVertexCount,
 				Input.RenderRestPosePositionBuffer->SRV,
diff --git a/Engine/Shaders/Private/HairStrands/HairStrandsInterpolation.usf b/Engine/Shaders/Private/HairStrands/HairStrandsInterpolation.usf
--- a/Engine/Shaders/Private/HairStrands/HairStrandsInterpolation.usf
+++ b/Engine/Shaders/Private/HairStrands/HairStrandsInterpolation.usf
@@ -1,9 +1,11 @@
 // Copyright 1998-2019 Epic Games, Inc. All Rights Reserved.
 
 #include "../Common.ush"
+#include "HairStrandsMeshProjectionCommon.ush"
 
 uint VertexCount;
 uint DispatchCountX;
+float3 HairWorldOffset;
 
 Buffer<uint4>  RenderRestPosePositionBuffer;
 Buffer<uint4>  SimRestPosePositionBuffer;
@@ -29,76 +31,6 @@ Buffer<float4>  DeformedPosition2Buffer;
 
 Buffer<uint>	RootToTriangleIndex;
 Buffer<uint>	VertexToRootIndexBuffer;
-
-#define QUATERNION_IDENTITY float4(0,0,0,1)
-#define SMALL_NUMBER 1e-8
-
-float4 NormalizeQuat(in float4 Quat)
-{
-	float SquaredNorm = dot(Quat, Quat);
-	return (SquaredNorm >= SMALL_NUMBER) ? Quat / sqrt(SquaredNorm) : QUATERNION_IDENTITY;
-}
-
-float4 InverseQuat(in float4 Quat)
-{
-	return float4(-Quat.x, -Quat.y, -Quat.z, Quat.w);
-}
-
-float4 MultiplyQuat(in float4 QuatA, in float4 QuatB)
-{
-	return float4(
-		QuatB.xyz * QuatA.w + QuatA.xyz * QuatB.w + cross(QuatA.xyz, QuatB.xyz),
-		QuatA.w * QuatB.w - dot(QuatA.xyz, QuatB.xyz));
-}
-
-float3 RotateVectorByQuat(in float3 Vector, in float4 Quat)
-{
-	float3 T = 2.0 * cross(Quat.xyz, Vector);
-	return Vector + Quat.w * T + cross(Quat.xyz, T);
-}
-
-float3 UnrotateVectorByQuat(in float3 Vector, in float4 Quat)
-{
-	float3 T = 2.0 * cross(Quat.xyz, Vector);
-	return Vector - Quat.w * T + cross(Quat.xyz, T);
-}
-
-float4 FindQuatBetweenHelper(in float3 A, in float3 B, in float NormAB)
-{
-	float W = NormAB + dot(A, B);
-	float4 Quat = (W > 1e-6*NormAB) ? float4(A.y*B.z - A.z*B.y, A.z*B.x - A.x*B.z, A.x*B.y - A.y*B.x, W) :
-		(abs(A.x) > abs(A.y)) ? float4(-A.z, 0.0, A.x, 0.0) : float4(0.0, -A.z, A.y, 0.0);
-	return NormalizeQuat(Quat);
-}
-
-float4 FindQuatBetweenNormals(in float3 NormalA, in float3 NormalB)
-{
-	float NormAB = 1.0;
-	return FindQuatBetweenHelper(NormalA, NormalB, NormAB);
-}
-
-void GetTriangleTransformation(uint VertexIndex, out float3 T, out float4 Q)
-{
-	const uint RootIndex = VertexToRootIndexBuffer[VertexIndex];
-//	const uint TriangleIndex = RootToTriangleIndex[RootIndex];
-
-	float3 Rest_P1 = RestPosition1Buffer[RootIndex].xyz + RestPositionWorldCenter;
-	float3 Rest_P0 = RestPosition0Buffer[RootIndex].xyz + RestPositionWorldCenter;
-	float3 Rest_P2 = RestPosition2Buffer[RootIndex].xyz + RestPositionWorldCenter;
-	float3 Rest_E0 = Rest_P1 - Rest_P0;
-	float3 Rest_E1 = Rest_P2 - Rest_P0;
-	float3 Rest_N = normalize(cross(Rest_E0, Rest_E1));
-
-	float3 Deformed_P0	= DeformedPosition0Buffer[RootIndex].xyz + DeformedPositionWorldCenter;
-	float3 Deformed_P1	= DeformedPosition1Buffer[RootIndex].xyz + DeformedPositionWorldCenter;
-	float3 Deformed_P2	= DeformedPosition2Buffer[RootIndex].xyz + DeformedPositionWorldCenter;
-	float3 Deformed_E0	= Deformed_P1 - Deformed_P0;
-	float3 Deformed_E1	= Deformed_P2 - Deformed_P0;
-	float3 Deformed_N	= normalize(cross(Deformed_E0, Deformed_E1));
-
-	T = Deformed_P0 - Rest_P0;
-	Q = FindQuatBetweenNormals(Rest_N, Deformed_N);
-}
 #endif
 
 #if PERMUTATION_DEBUG == 1
@@ -132,8 +64,7 @@ void MainCS(uint2 DispatchThreadId : SV_DispatchThreadID)
 		GuideVertexWeights.z = saturate(1.0f - (GuideVertexWeights.x + GuideVertexWeights.y));
 
 		float3 CurrOffset = 0;
-		float3 PrevOffset = 0;
-		const float3 ControlPoint = f16tof32(RenderRestPosePositionBuffer[VertexIndex].xyz);
+		float3 ControlPoint = f16tof32(RenderRestPosePositionBuffer[VertexIndex].xyz);
 		[unroll]
 		for (uint K = 0; K < 3; ++K)
 		{
@@ -149,10 +80,22 @@ void MainCS(uint2 DispatchThreadId : SV_DispatchThreadID)
 		// hair_todo: applied the transformation in local space and then rotation everything in global space
 		#if PERMUTATION_DYNAMIC_GEOMETRY == 1
 		{
-			float3 DynTranslation = 0;
-			float4 DynRotation = 0;
-			GetTriangleTransformation(VertexIndex, DynTranslation, DynRotation);
-			CurrOffset += DynTranslation;
+			const uint RootIndex						= VertexToRootIndexBuffer[VertexIndex];
+			const FHairMeshTriangle RestTriangle		= GetTriangleTransformation(RootIndex, RestPosition0Buffer, RestPosition1Buffer, RestPosition2Buffer, RestPositionWorldCenter);
+			const FHairMeshTriangle DeformedTriangle	= GetTriangleTransformation(RootIndex, DeformedPosition0Buffer, DeformedPosition1Buffer, DeformedPosition2Buffer, DeformedPositionWorldCenter);
+
+			ControlPoint = TransformPoint(ControlPoint + HairWorldOffset, RestTriangle, DeformedTriangle) - HairWorldOffset;
+
+			// Debug
+			#if 0
+			{
+				// Show the triangle normal use by hair fiber
+				//ControlPoint = VertexIndex % 2 == 0 ? RestTriangle.P0 : RestTriangle.P0 + RestTriangle.N;
+
+				// Measure the stretch of the projection (to identify if there any local transform issue during the projeciton steps
+				ControlPoint = VertexIndex % 2 == 0 ? ControlPoint  : RestTriangle.P0 - HairWorldOffset;
+			}
+			#endif
 		}
 		#endif	
 
diff --git a/Engine/Shaders/Private/HairStrands/HairStrandsMeshProjectionCommon.ush b/Engine/Shaders/Private/HairStrands/HairStrandsMeshProjectionCommon.ush
--- a/Engine/Shaders/Private/HairStrands/HairStrandsMeshProjectionCommon.ush
+++ b/Engine/Shaders/Private/HairStrands/HairStrandsMeshProjectionCommon.ush
@@ -43,11 +43,55 @@ FHairMeshTriangle GetTriangleData(uint TriangleIndex, float4x4 LocalToWorld, Buf
 
 	Out.E0 = Out.P1 - Out.P0;
 	Out.E1 = Out.P2 - Out.P0;
-	Out.N = normalize(cross(Out.E0, Out.E1));
+	Out.N = normalize(cross(Out.E1, Out.E0));
 
 	return Out;
 }
 
+FHairMeshTriangle GetTriangleTransformation(
+	uint RootIndex,
+	Buffer<float4>  InPosition0Buffer,
+	Buffer<float4>  InPosition1Buffer,
+	Buffer<float4>  InPosition2Buffer, 
+	float3 CenterOffset=float3(0,0,0))
+{
+	FHairMeshTriangle Out;
+
+	// Build orthonormal frame
+	Out.P0 = InPosition0Buffer[RootIndex].xyz + CenterOffset;
+	Out.P1 = InPosition1Buffer[RootIndex].xyz + CenterOffset;
+	Out.P2 = InPosition2Buffer[RootIndex].xyz + CenterOffset;
+	Out.E0 = Out.P1 - Out.P0;
+	Out.E1 = Out.P2 - Out.P0;
+
+	Out.N  = normalize(cross(Out.E1, Out.E0));
+	Out.E0 = normalize(Out.E0); 
+	Out.E1 = normalize(cross(Out.E0, Out.N));
+
+	return Out;
+}
+
+float3 ToTriangle(const float3 P, const FHairMeshTriangle T)
+{
+	float3 LocalP = P - T.P0;
+	float3 Out;
+	Out.x = dot(LocalP, T.E0);
+	Out.y = dot(LocalP, T.E1);
+	Out.z = dot(LocalP, T.N);
+	return Out;
+}
+
+float3 ToWorld(const float3 P, const FHairMeshTriangle T)
+{
+	return P.x * T.E0 + P.y * T.E1 + P.z * T.N + T.P0;
+}
+
+float3 TransformPoint(float3 P, const FHairMeshTriangle RestTriangle, const FHairMeshTriangle DeformedTriangle)
+{
+	float3 LocalP = ToTriangle(P, RestTriangle);
+	return ToWorld(LocalP, DeformedTriangle);
+}
+
 struct FHairProjectionData
 {
 	uint TriangleIndex;
diff --git a/Engine/Shaders/Private/HairStrands/HairStrandsMeshProjectionHairDebug.usf b/Engine/Shaders/Private/HairStrands/HairStrandsMeshProjectionHairDebug.usf
--- a/Engine/Shaders/Private/HairStrands/HairStrandsMeshProjectionHairDebug.usf
+++ b/Engine/Shaders/Private/HairStrands/HairStrandsMeshProjectionHairDebug.usf
@@ -7,15 +7,15 @@
 
 #define INPUT_ROOT_FRAME 0
 #define INPUT_ROOT_TRIANGLE 1
-
 uint MaxRootCount;
+float4x4 RootLocalToWorld;
+uint			DeformedFrameEnable;
 
 #if PERMUTATION_INPUT_TYPE == INPUT_ROOT_FRAME
-float4x4 RootLocalToWorld;
 Buffer<float4> RootPositionBuffer;
 Buffer<float4> RootNormalBuffer;
+#endif
 
-uint			DeformedFrameEnable;
 Buffer<float4>  RestPosition0Buffer;
 Buffer<float4>  RestPosition1Buffer;
 Buffer<float4>  RestPosition2Buffer;
@@ -23,35 +23,6 @@ Buffer<float4>  RestPosition2Buffer;
 Buffer<float4>  DeformedPosition0Buffer;
 Buffer<float4>  DeformedPosition1Buffer;
 Buffer<float4>  DeformedPosition2Buffer;
-#endif
-
-#if PERMUTATION_INPUT_TYPE == INPUT_ROOT_TRIANGLE
-Buffer<float4> RootTrianglePosition0;
-Buffer<float4> RootTrianglePosition1;
-Buffer<float4> RootTrianglePosition2;
-#endif
-
-#if PERMUTATION_INPUT_TYPE == INPUT_ROOT_FRAME
-void GetTriangleTransformation(uint RootIndex, out float3 T, out float4 Q)
-{
-	float3 Rest_P1 = RestPosition1Buffer[RootIndex].xyz;
-	float3 Rest_P0 = RestPosition0Buffer[RootIndex].xyz;
-	float3 Rest_P2 = RestPosition2Buffer[RootIndex].xyz;
-	float3 Rest_E0 = Rest_P1 - Rest_P0;
-	float3 Rest_E1 = Rest_P2 - Rest_P0;
-	float3 Rest_N = normalize(cross(Rest_E0, Rest_E1));
-
-	float3 Deformed_P0 = DeformedPosition0Buffer[RootIndex].xyz;
-	float3 Deformed_P1 = DeformedPosition1Buffer[RootIndex].xyz;
-	float3 Deformed_P2 = DeformedPosition2Buffer[RootIndex].xyz;
-	float3 Deformed_E0 = Deformed_P1 - Deformed_P0;
-	float3 Deformed_E1 = Deformed_P2 - Deformed_P0;
-	float3 Deformed_N = normalize(cross(Deformed_E0, Deformed_E1));
-
-	T = Deformed_P0 - Rest_P0;
-	Q = 0;// FindQuatBetweenNormals(Rest_N, Deformed_N);
-}
-#endif
 
 void MainVS(
 	uint VertexId : SV_VertexID,
@@ -70,36 +41,27 @@ void MainVS(
 
 	float3 N;
 	float3 P;
-
-	if (DeformedFrameEnable)
 	{
-		float3 DynTranslation = 0;
-		float4 DynRotation = 0;
-		GetTriangleTransformation(LineIndex, DynTranslation, DynRotation);
+		const FHairMeshTriangle RestTriangle     = GetTriangleTransformation(LineIndex, RestPosition0Buffer, RestPosition1Buffer, RestPosition2Buffer);
+		const FHairMeshTriangle DeformedTriangle = GetTriangleTransformation(LineIndex, DeformedPosition0Buffer, DeformedPosition1Buffer, DeformedPosition2Buffer);
 
-		float3 P0 = RootPositionBuffer[LineIndex];
-		P0 = mul(float4(P0, 1), RootLocalToWorld).xyz;
-		P0 += DynTranslation;
+		float3 P0 = RootPositionBuffer[LineIndex].xyz;
+		float3 LN = RootNormalBuffer[LineIndex].xyz;
+		float3 P1 = P0 + LN;
 
-		N = RootNormalBuffer[LineIndex];
-		N = mul(float4(N, 0), RootLocalToWorld).xyz;
 
-		const float3 P1 = P0 + N;
-		P = LineVertex == 0 ? P0 : P1;
-	}
-	else
-	{
-		float3 P0 = RootPositionBuffer[LineIndex];
-		P0 = mul(float4(P0,1), RootLocalToWorld).xyz;
-	
-		N  = RootNormalBuffer[LineIndex];
-		N = mul(float4(N, 0), RootLocalToWorld).xyz;
+		if (DeformedFrameEnable > 0)
+		{
+			P0 = TransformPoint(P0, RestTriangle, DeformedTriangle);
+			P1 = TransformPoint(P1, RestTriangle, DeformedTriangle);
+		}
 
-		const float3 P1 = P0 + N;
+		P0 = mul(float4(P0, 1), RootLocalToWorld).xyz;
+		P1 = mul(float4(P1, 1), RootLocalToWorld).xyz;
 
+		N = normalize(P0 - P1);
 		P = LineVertex == 0 ? P0 : P1;
 	}
-
 #endif
 
 #if PERMUTATION_INPUT_TYPE == INPUT_ROOT_TRIANGLE
@@ -108,21 +70,39 @@ void MainVS(
 
 	if (TriangleIndex > MaxRootCount) return;
 
-	const float3 P0 = RootTrianglePosition0[TriangleIndex];
-	const float3 P1 = RootTrianglePosition1[TriangleIndex];
-	const float3 P2 = RootTrianglePosition2[TriangleIndex];
+	const FHairMeshTriangle RestTriangle     = GetTriangleTransformation(TriangleIndex, RestPosition0Buffer, RestPosition1Buffer, RestPosition2Buffer);
+	const FHairMeshTriangle DeformedTriangle = GetTriangleTransformation(TriangleIndex, DeformedPosition0Buffer, DeformedPosition1Buffer, DeformedPosition2Buffer);
 
-	const float3 E0 = P1 - P0;
-	const float3 E1 = P2 - P0;
-	const float3 N = normalize(cross(E0, E1));
+	FHairMeshTriangle Tri;
+	if (DeformedFrameEnable > 0)
+		Tri = DeformedTriangle;
+	else
+		Tri = RestTriangle;
 
+	float3 N = Tri.N;
 	float3 P = 0;
 	if (TriangleVertex == 0)
-		P = P0;
+		P = Tri.P0;
 	if (TriangleVertex == 1)
-		P = P1;
+		P = Tri.P1;
 	if (TriangleVertex == 2)
-		P = P2;
+		P = Tri.P2;
+	
+	// Debug: To test if the rest-triangle-to-deformed-triangle transformation works correctly
+	#if 0
+	if (DeformedFrameEnable > 0)
+	{
+		if (TriangleVertex == 0)
+			P = RestTriangle.P0;
+		if (TriangleVertex == 1)
+			P = RestTriangle.P1;
+		if (TriangleVertex == 2)
+			P = RestTriangle.P2;
+		P = TransformPoint(P, RestTriangle, DeformedTriangle);
+	}
+	#endif
+
+	P = mul(float4(P, 1), RootLocalToWorld).xyz;
 #endif
 
 	OutFaceNormal = N;
diff --git a/Engine/Source/Runtime/Renderer/Private/HairStrands/HairStrandsDebug.cpp b/Engine/Source/Runtime/Renderer/Private/HairStrands/HairStrandsDebug.cpp
--- a/Engine/Source/Runtime/Renderer/Private/HairStrands/HairStrandsDebug.cpp
+++ b/Engine/Source/Runtime/Renderer/Private/HairStrands/HairStrandsDebug.cpp
@@ -634,10 +634,6 @@ BEGIN_SHADER_PARAMETER_STRUCT(FHairProjectionHairDebugParameters, )
 	SHADER_PARAMETER(uint32, DeformedFrameEnable)
 	SHADER_PARAMETER(FMatrix, RootLocalToWorld)
 
-	SHADER_PARAMETER_SRV(StructuredBuffer, RootTrianglePosition0)
-	SHADER_PARAMETER_SRV(StructuredBuffer, RootTrianglePosition2)
-	SHADER_PARAMETER_SRV(StructuredBuffer, RootTrianglePosition1)
-
 	SHADER_PARAMETER_SRV(StructuredBuffer, RestPosition0Buffer)
 	SHADER_PARAMETER_SRV(StructuredBuffer, RestPosition1Buffer)
 	SHADER_PARAMETER_SRV(StructuredBuffer, RestPosition2Buffer)
@@ -720,25 +716,12 @@ static void AddDebugProjectionHairPass(
 	if (EDebugProjectionHairType::HairFrame == GeometryType && (!HairData.RootPositionBuffer || !HairData.RootNormalBuffer))
 		return;
 
-	if (EDebugProjectionHairType::HairTriangle == GeometryType && HairStrandsTriangleType::RestPose == PoseType && (
-		!HairData.LODDatas[LODIndex].RestRootTrianglePosition0Buffer ||
-		!HairData.LODDatas[LODIndex].RestRootTrianglePosition1Buffer ||
-		!HairData.LODDatas[LODIndex].RestRootTrianglePosition2Buffer))
-		return;
-
-	if (EDebugProjectionHairType::HairTriangle == GeometryType && HairStrandsTriangleType::DeformedPose == PoseType && (
-		!HairData.LODDatas[LODIndex].DeformedRootTrianglePosition0Buffer ||
-		!HairData.LODDatas[LODIndex].DeformedRootTrianglePosition1Buffer ||
-		!HairData.LODDatas[LODIndex].DeformedRootTrianglePosition2Buffer))
-		return;
-
-	if (EDebugProjectionHairType::HairFrame == GeometryType && (
-		!HairData.LODDatas[LODIndex].RestRootTrianglePosition0Buffer ||
+	if (!HairData.LODDatas[LODIndex].RestRootTrianglePosition0Buffer ||
 		!HairData.LODDatas[LODIndex].RestRootTrianglePosition1Buffer ||
 		!HairData.LODDatas[LODIndex].RestRootTrianglePosition2Buffer ||
 		!HairData.LODDatas[LODIndex].DeformedRootTrianglePosition0Buffer ||
 		!HairData.LODDatas[LODIndex].DeformedRootTrianglePosition1Buffer ||
-		!HairData.LODDatas[LODIndex].DeformedRootTrianglePosition2Buffer))
+		!HairData.LODDatas[LODIndex].DeformedRootTrianglePosition2Buffer)
 		return;
 
 	const FIntRect Viewport = View->ViewRect;
@@ -747,41 +730,22 @@ static void AddDebugProjectionHairPass(
 	FHairProjectionHairDebugParameters* Parameters = GraphBuilder.AllocParameters<FHairProjectionHairDebugParameters>();
 	Parameters->OutputResolution = Resolution;
 	Parameters->MaxRootCount = HairData.RootCount;
+	Parameters->RootLocalToWorld = HairData.LocalToWorld.ToMatrixWithScale();
+	Parameters->DeformedFrameEnable = PoseType == HairStrandsTriangleType::DeformedPose;
 
 	if (EDebugProjectionHairType::HairFrame == GeometryType)
 	{
-		Parameters->RootLocalToWorld = HairData.LocalToWorld.ToMatrixWithScale();
 		Parameters->RootPositionBuffer = HairData.RootPositionBuffer;
 		Parameters->RootNormalBuffer = HairData.RootNormalBuffer;
-		Parameters->DeformedFrameEnable = PoseType == HairStrandsTriangleType::DeformedPose;
-		//if (PoseType == HairStrandsTriangleType::DeformedPose)
-		{
-			Parameters->RestPosition0Buffer = HairData.LODDatas[LODIndex].RestRootTrianglePosition0Buffer->SRV;
-			Parameters->RestPosition1Buffer = HairData.LODDatas[LODIndex].RestRootTrianglePosition1Buffer->SRV;
-			Parameters->RestPosition2Buffer = HairData.LODDatas[LODIndex].RestRootTrianglePosition2Buffer->SRV;
-
-			Parameters->DeformedPosition0Buffer = HairData.LODDatas[LODIndex].DeformedRootTrianglePosition0Buffer->SRV;
-			Parameters->DeformedPosition1Buffer = HairData.LODDatas[LODIndex].DeformedRootTrianglePosition1Buffer->SRV;
-			Parameters->DeformedPosition2Buffer = HairData.LODDatas[LODIndex].DeformedRootTrianglePosition2Buffer->SRV;
-		}
 	}
 
-	if (EDebugProjectionHairType::HairTriangle == GeometryType)
-	{
-		if (PoseType == HairStrandsTriangleType::RestPose)
-		{
-			Parameters->RootTrianglePosition0 = HairData.LODDatas[LODIndex].RestRootTrianglePosition0Buffer->SRV;
-			Parameters->RootTrianglePosition1 = HairData.LODDatas[LODIndex].RestRootTrianglePosition1Buffer->SRV;
-			Parameters->RootTrianglePosition2 = HairData.LODDatas[LODIndex].RestRootTrianglePosition2Buffer->SRV;
-		}
-
-		if (PoseType == HairStrandsTriangleType::DeformedPose)
-		{
-			Parameters->RootTrianglePosition0 = HairData.LODDatas[LODIndex].DeformedRootTrianglePosition0Buffer->SRV;
-			Parameters->RootTrianglePosition1 = HairData.LODDatas[LODIndex].DeformedRootTrianglePosition1Buffer->SRV;
-			Parameters->RootTrianglePosition2 = HairData.LODDatas[LODIndex].DeformedRootTrianglePosition2Buffer->SRV;
-		}
-	}
+	Parameters->RestPosition0Buffer = HairData.LODDatas[LODIndex].RestRootTrianglePosition0Buffer->SRV;
+	Parameters->RestPosition1Buffer = HairData.LODDatas[LODIndex].RestRootTrianglePosition1Buffer->SRV;
+	Parameters->RestPosition2Buffer = HairData.LODDatas[LODIndex].RestRootTrianglePosition2Buffer->SRV;
+		
+	Parameters->DeformedPosition0Buffer = HairData.LODDatas[LODIndex].DeformedRootTrianglePosition0Buffer->SRV;
+	Parameters->DeformedPosition1Buffer = HairData.LODDatas[LODIndex].DeformedRootTrianglePosition1Buffer->SRV;
+	Parameters->DeformedPosition2Buffer = HairData.LODDatas[LODIndex].DeformedRootTrianglePosition2Buffer->SRV;
 
 	Parameters->ViewUniformBuffer = View->ViewUniformBuffer;
 	Parameters->RenderTargets[0] = FRenderTargetBinding(ColorTarget, ERenderTargetLoadAction::ELoad, 0);
