Disable rect optim by default, since it does not work properly at the moment.

(ushell-p4-cherrypick of 9479462 by Charles.deRousiers)

diff --git a/Engine/Source/Runtime/Renderer/Private/HairStrands/HairStrandsUtils.cpp b/Engine/Source/Runtime/Renderer/Private/HairStrands/HairStrandsUtils.cpp
--- a/Engine/Source/Runtime/Renderer/Private/HairStrands/HairStrandsUtils.cpp
+++ b/Engine/Source/Runtime/Renderer/Private/HairStrands/HairStrandsUtils.cpp
@@ -27,7 +27,7 @@ static FAutoConsoleVariableRef CVarStrandHairShadowRasterizationScale(TEXT("r.Ha
 static float GDeepShadowAABBScale = 1.0f;
 static FAutoConsoleVariableRef CVarDeepShadowAABBScale(TEXT("r.HairStrands.DeepShadow.AABBScale"), GDeepShadowAABBScale, TEXT("Scaling value for loosing/tighting deep shadow bounding volume"));
 
-static int32 GHairVisibilityRectOptimEnable = 1;
+static int32 GHairVisibilityRectOptimEnable = 0;
 static FAutoConsoleVariableRef CVarHairVisibilityRectOptimEnable(TEXT("r.HairStrands.RectLightingOptim"), GHairVisibilityRectOptimEnable, TEXT("Hair Visibility use projected view rect to light only relevant pixels"));
 
 float SampleCountToSubPixelSize(uint32 SamplePerPixelCount)
