Add support for multiple hair clusterw during the AO and env lighting passes.

(ushell-p4-cherrypick of 9491453 by Charles.deRousiers)

diff --git a/Engine/Shaders/Private/HairStrands/HairStrandsEnvironmentLighting.usf b/Engine/Shaders/Private/HairStrands/HairStrandsEnvironmentLighting.usf
--- a/Engine/Shaders/Private/HairStrands/HairStrandsEnvironmentLighting.usf
+++ b/Engine/Shaders/Private/HairStrands/HairStrandsEnvironmentLighting.usf
@@ -49,6 +49,7 @@ Texture3D<float4>			HairScatteringLUTTexture;
 Texture3D<float4>			HairEnergyLUTTexture;
 SamplerState 				HairLUTSampler;
 
+uint  Voxel_ClusterId;
 float Voxel_TanConeAngle;
 float Voxel_DensityScale;
 float Voxel_DepthBiasScale;
@@ -468,6 +469,7 @@ void MainPS(
 		// * Change this to per sample evaluation once we have better hair transmittance evaluation for sky
 		float Coverage = 1;
 		float ClosestDepth = 0;
+		bool bIsClusterCompatible = false;
 		FGBufferData HairGBuffer = (FGBufferData)0;
 		Coverage = 0;
 
@@ -482,6 +484,7 @@ void MainPS(
 			float			SampleRoughness;
 			float			SampleSpecular;
 			DecodeNodeData(NodeData, SampleDepth, SampleVisData, SampleBaseColor, SampleRoughness, SampleSpecular);
+			bIsClusterCompatible = bIsClusterCompatible || SampleVisData.ClusterId == Voxel_ClusterId;
 
 			// Inject material data from the visibility/mini-gbuffer for hair
 			const float SceneDepth = ConvertFromDeviceZ(SampleDepth);
@@ -509,6 +512,7 @@ void MainPS(
 		// Lit fully covered or partially covered pixel
 		// * Fully covered pixel output their result into the main color target
 		// * Partially covered pixel output their result into the sub pixel color target. the main color target will store the underlying pixel color 
+		if (bIsClusterCompatible)
 		{
 			const float3 SkyLighting = ReflectionEnvironment(HairGBuffer, BufferUV, ScreenPosition, SvPosition);
 
@@ -540,6 +544,7 @@ void MainPS(
 
 	// Trace AO for all pixel which are partially covered or not covered at all.
 	bool bIsValid = false;
+	bool bIsClusterCompatible = false;
 	if (NodeDesc.Count < 255)
 	{
 		float ClosestDepth = 0;
@@ -555,6 +560,8 @@ void MainPS(
 			float			SampleRoughness;
 			float			SampleSpecular;
 			DecodeNodeData(NodeData, SampleDepth, SampleVisData, SampleBaseColor, SampleRoughness, SampleSpecular);
+			bIsClusterCompatible = bIsClusterCompatible || SampleVisData.ClusterId == Voxel_ClusterId;
+
 			const float LocalCoverage = From8bitCoverage(SampleVisData.Coverage8bit);
 			Coverage8Bits += SampleVisData.Coverage8bit;
 		}
@@ -565,7 +572,7 @@ void MainPS(
 		bIsValid = GBuffer.ShadingModelID != SHADINGMODELID_UNLIT;
 	}
 
-	if (bIsValid)
+	if (bIsValid && bIsClusterCompatible)
 	{
 		const float AO = TraceAO(WorldPosition, WorldNormal, SvPosition);
 		OutAO = AO;
diff --git a/Engine/Source/Runtime/Renderer/Private/HairStrands/HairStrandsEnvironment.cpp b/Engine/Source/Runtime/Renderer/Private/HairStrands/HairStrandsEnvironment.cpp
--- a/Engine/Source/Runtime/Renderer/Private/HairStrands/HairStrandsEnvironment.cpp
+++ b/Engine/Source/Runtime/Renderer/Private/HairStrands/HairStrandsEnvironment.cpp
@@ -53,6 +53,7 @@ class FHairEnvironmentLightingPS : public FGlobalShader
 	BEGIN_SHADER_PARAMETER_STRUCT(FParameters, )
 
 		SHADER_PARAMETER(FVector, Voxel_MinAABB)
+		SHADER_PARAMETER(uint32, Voxel_ClusterId)
 		SHADER_PARAMETER(FVector, Voxel_MaxAABB)
 		SHADER_PARAMETER(uint32, Voxel_Resolution)
 		SHADER_PARAMETER(float, Voxel_DensityScale)
@@ -99,8 +100,8 @@ static void AddHairStrandsEnvironmentPass(
 	FRDGBuilder& GraphBuilder,
 	const FViewInfo& View,
 	EEnvRenderMode RenderMode,
-	const FHairStrandsVisibilityData* HairVisibilityData,
-	const FHairStrandsClusterData* ClusterData,
+	const FHairStrandsVisibilityData& VisibilityData,
+	const FHairStrandsClusterData& ClusterData,
 	FRDGTextureRef Output0,
 	FRDGTextureRef Output1)
 {
@@ -117,21 +118,14 @@ static void AddHairStrandsEnvironmentPass(
 	PassParameters->HairScatteringLUTTexture = GraphBuilder.RegisterExternalTexture(InHairLUT.Textures[HairLUTType_DualScattering], TEXT("HairScatteringEnergyLUTTexture"));
 	PassParameters->HairLUTSampler = TStaticSamplerState<SF_Bilinear, AM_Clamp, AM_Clamp, AM_Clamp>::GetRHI();
 
-	if (ClusterData)
-	{
-		PassParameters->Voxel_MinAABB = ClusterData->GetMinBound();
-		PassParameters->Voxel_MaxAABB = ClusterData->GetMaxBound();
-		PassParameters->Voxel_Resolution = ClusterData->GetResolution();
-		PassParameters->Voxel_DensityTexture = GraphBuilder.RegisterExternalTexture(ClusterData->VoxelResources.DensityTexture);
-		PassParameters->Voxel_DensityScale = GetHairStrandsVoxelizationDensityScale();
-		PassParameters->Voxel_DepthBiasScale = GetHairStrandsVoxelizationDepthBiasScale();
-		PassParameters->Voxel_TanConeAngle = FMath::Tan(FMath::DegreesToRadians(GetHairStrandsSkyLightingConeAngle()));
-	}
-	else
-	{
-		// .. todo dummy 3D texture
-		PassParameters->Voxel_DensityTexture = GraphBuilder.RegisterExternalTexture(GSystemTextures.BlackDummy);
-	}
+	PassParameters->Voxel_ClusterId = ClusterData.ClusterId;
+	PassParameters->Voxel_MinAABB = ClusterData.GetMinBound();
+	PassParameters->Voxel_MaxAABB = ClusterData.GetMaxBound();
+	PassParameters->Voxel_Resolution = ClusterData.GetResolution();
+	PassParameters->Voxel_DensityTexture = GraphBuilder.RegisterExternalTexture(ClusterData.VoxelResources.DensityTexture);
+	PassParameters->Voxel_DensityScale = GetHairStrandsVoxelizationDensityScale();
+	PassParameters->Voxel_DepthBiasScale = GetHairStrandsVoxelizationDepthBiasScale();
+	PassParameters->Voxel_TanConeAngle = FMath::Tan(FMath::DegreesToRadians(GetHairStrandsSkyLightingConeAngle()));
 
 	PassParameters->PreIntegratedGF = GSystemTextures.PreintegratedGF->GetRenderTargetItem().ShaderResourceTexture;
 	PassParameters->PreIntegratedGFSampler = TStaticSamplerState<SF_Bilinear, AM_Clamp, AM_Clamp, AM_Clamp>::GetRHI();
@@ -150,12 +144,9 @@ static void AddHairStrandsEnvironmentPass(
 
 
 	// Bind hair data
-	if (HairVisibilityData)
-	{
-		PassParameters->HairCategorizationTexture = HairVisibilityData->CategorizationTexture->GetRenderTargetItem().ShaderResourceTexture;
-		PassParameters->HairVisibilityNodeOffsetAndCount = HairVisibilityData->NodeIndex->GetRenderTargetItem().ShaderResourceTexture;
-		PassParameters->HairVisibilityNodeData = HairVisibilityData->NodeDataSRV;
-	}
+	PassParameters->HairCategorizationTexture = VisibilityData.CategorizationTexture->GetRenderTargetItem().ShaderResourceTexture;
+	PassParameters->HairVisibilityNodeOffsetAndCount = VisibilityData.NodeIndex->GetRenderTargetItem().ShaderResourceTexture;
+	PassParameters->HairVisibilityNodeData = VisibilityData.NodeDataSRV;
 
 	PassParameters->AO_Power = 0;
 	PassParameters->AO_Intensity = 0;
@@ -210,19 +201,7 @@ static void AddHairStrandsEnvironmentPass(
 		FPixelShaderUtils::DrawFullscreenTriangle(InRHICmdList);
 	});
 }
-
-static const FHairStrandsClusterData* GetClusterData(const FHairStrandsDatas* HairDatas, uint32 ViewIndex, uint32 ClusterIndex)
-{
-	return
-		(
-			HairDatas && 
-			ViewIndex < uint32(HairDatas->HairClusterPerViews.Views.Num()) &&
-			ClusterIndex < uint32(HairDatas->HairClusterPerViews.Views[ViewIndex].Datas.Num())
-		) ?
-		&HairDatas->HairClusterPerViews.Views[ViewIndex].Datas[ClusterIndex] :
-		nullptr;
-}
-
+	
 void RenderHairStrandsEnvironmentLighting(
 	FRDGBuilder& GraphBuilder,
 	const uint32 ViewIndex,
@@ -231,26 +210,26 @@ void RenderHairStrandsEnvironmentLighting(
 	FRDGTextureRef SceneColorTexture,
 	FRDGTextureRef SceneColorSubPixelTexture)
 {
-	if (!GetHairStrandsSkyLightingEnable())
+	if (!GetHairStrandsSkyLightingEnable() || !HairDatas)
 		return;
 
 	check(ViewIndex < uint32(Views.Num()));
-	const FViewInfo& View = Views[ViewIndex];
-	const FHairStrandsVisibilityData* HairVisibilityData = HairDatas ? &HairDatas->HairVisibilityViews.HairDatas[ViewIndex] : nullptr;
-	const bool bRenderHairLighting = HairVisibilityData && HairVisibilityData->NodeIndex && HairVisibilityData->NodeDataSRV;
+	check(ViewIndex < uint32(HairDatas->HairVisibilityViews.HairDatas.Num()));	
+	const FHairStrandsVisibilityData& VisibilityData = HairDatas->HairVisibilityViews.HairDatas[ViewIndex];
+	const bool bRenderHairLighting = VisibilityData.NodeIndex && VisibilityData.NodeDataSRV;
 	if (!bRenderHairLighting)
 	{
 		return;
 	}
 
-	// @hair_todo: add support for multiple clusters on screen
-	const FHairStrandsClusterData* ClusterData = GetClusterData(HairDatas, ViewIndex, 0);
-
-	// @hair_todo: 
-	// * Add support for : BentNormal, global AO, SSR, DFSO
-	// * Add local reflection probe, current take into account only the sky lighting
-	//RDG_GPU_STAT_SCOPE(GraphBuilder, HairStrandEnvironment);
-	AddHairStrandsEnvironmentPass(GraphBuilder, View, EEnvRenderMode::Lighting, HairVisibilityData, ClusterData, SceneColorTexture, SceneColorSubPixelTexture);
+	const FViewInfo& View = Views[ViewIndex];
+	for (const FHairStrandsClusterData& ClusterData : HairDatas->HairClusterPerViews.Views[ViewIndex].Datas)
+	{
+		// @hair_todo: 
+		// * Add support for : BentNormal, global AO, SSR, DFSO
+		// * Add local reflection probe, current take into account only the sky lighting
+		AddHairStrandsEnvironmentPass(GraphBuilder, View, EEnvRenderMode::Lighting, VisibilityData, ClusterData, SceneColorTexture, SceneColorSubPixelTexture);
+	}
 }
 
 void RenderHairStrandsAmbientOcclusion(
@@ -259,28 +238,30 @@ void RenderHairStrandsAmbientOcclusion(
 	const FHairStrandsDatas* HairDatas,
 	const TRefCountPtr<IPooledRenderTarget>& InAOTexture)
 {
-	if (!GetHairStrandsSkyAOEnable() || Views.Num() == 0 || !InAOTexture)
+	if (!GetHairStrandsSkyAOEnable() || Views.Num() == 0 || !InAOTexture || !HairDatas)
 		return;
 
 	for (uint32 ViewIndex = 0; ViewIndex < uint32(Views.Num()); ++ViewIndex)
 	{
-		const FViewInfo& View = Views[ViewIndex];
-		const FHairStrandsVisibilityData* HairVisibilityData = HairDatas ? &HairDatas->HairVisibilityViews.HairDatas[ViewIndex] : nullptr;
-		const bool bRenderHairLighting = HairVisibilityData && HairVisibilityData->NodeIndex && HairVisibilityData->NodeDataSRV;
+		check(ViewIndex < uint32(HairDatas->HairVisibilityViews.HairDatas.Num()));
+		const FHairStrandsVisibilityData& VisibilityData = HairDatas->HairVisibilityViews.HairDatas[ViewIndex];
+		const bool bRenderHairLighting = VisibilityData.NodeIndex && VisibilityData.NodeDataSRV;
 		if (!bRenderHairLighting)
 		{
-			return;
+			continue;
 		}
 
+		if (ViewIndex > uint32(HairDatas->HairClusterPerViews.Views.Num()))
+			continue;
 
+		const FViewInfo& View = Views[ViewIndex];
 		FRDGBuilder GraphBuilder(RHICmdList);
 		FRDGTextureRef AOTexture = InAOTexture ? GraphBuilder.RegisterExternalTexture(InAOTexture, TEXT("AOTexture")) : nullptr;
+		for (const FHairStrandsClusterData& ClusterData : HairDatas->HairClusterPerViews.Views[ViewIndex].Datas)
+		{
+			AddHairStrandsEnvironmentPass(GraphBuilder, View, EEnvRenderMode::AO, VisibilityData, ClusterData, AOTexture, nullptr);
 
-		// @hair_todo: add support for multiple clusters on screen
-		const FHairStrandsClusterData* ClusterData = GetClusterData(HairDatas, ViewIndex, 0);
-
-		AddHairStrandsEnvironmentPass(GraphBuilder, View, EEnvRenderMode::AO, HairVisibilityData, ClusterData, AOTexture, nullptr);
-
+		}
 		GraphBuilder.Execute();
 	}
 }
