Fix crash when debug information are display with non-skinned element.

(ushell-p4-cherrypick of 9530730 by Charles.deRousiers)

diff --git a/Engine/Source/Runtime/Renderer/Private/HairStrands/HairStrandsInterface.cpp b/Engine/Source/Runtime/Renderer/Private/HairStrands/HairStrandsInterface.cpp
--- a/Engine/Source/Runtime/Renderer/Private/HairStrands/HairStrandsInterface.cpp
+++ b/Engine/Source/Runtime/Renderer/Private/HairStrands/HairStrandsInterface.cpp
@@ -82,9 +82,9 @@ void AddHairStrandsProjectionQuery(
 	uint64 Id,
 	EWorldType::Type WorldType,
 	int32 LODIndex,
-	const FVector& RestRootCenter)
+	const FVector& RestPositionOffset)
 {
-	GHairManager.ProjectionsQueries.Add({ Id, WorldType, LODIndex, RestRootCenter, false });
+	GHairManager.ProjectionsQueries.Add({ Id, WorldType, LODIndex, RestPositionOffset, false });
 }
 
 void RegisterHairStrands(
@@ -369,15 +369,21 @@ FHairStrandsDebugInfos GetHairStandsDebugInfos()
 		Info.Id = E.Id;
 		Info.WorldType = E.WorldType;
 
-		check(Info.HairGroups.Num() == E.RenProjectionHairDatas.HairGroups.Num());
-
 		const uint32 GroupCount = Info.HairGroups.Num();
 		for (uint32 GroupIt=0; GroupIt < GroupCount; ++GroupIt)
 		{		
 			FHairStrandsDebugInfo::HairGroup& GroupInfo = Info.HairGroups[GroupIt];
-			FHairStrandsProjectionHairData::HairGroup& ProjectionHair = E.RenProjectionHairDatas.HairGroups[GroupIt];
-			GroupInfo.LODCount = ProjectionHair.LODDatas.Num();
-			GroupInfo.bHasSkinInterpolation = ProjectionHair.LODDatas.Num() > 0;
+			if (GroupIt < uint32(E.RenProjectionHairDatas.HairGroups.Num()))
+			{
+				FHairStrandsProjectionHairData::HairGroup& ProjectionHair = E.RenProjectionHairDatas.HairGroups[GroupIt];
+				GroupInfo.LODCount = ProjectionHair.LODDatas.Num();
+				GroupInfo.bHasSkinInterpolation = ProjectionHair.LODDatas.Num() > 0;
+			}
+			else
+			{
+				GroupInfo.LODCount = 0;
+				GroupInfo.bHasSkinInterpolation = false;
+			}
 		}
 	}
 
