Split hair component and hair cluster debug information to display component information even when there is no hair cluster visible on screen

(ushell-p4-cherrypick of 9529367 by Charles.deRousiers)

(ushell-p4-cherrypick of 9530788 by Charles.deRousiers)

diff --git a/Engine/Source/Runtime/Renderer/Private/DeferredShadingRenderer.cpp b/Engine/Source/Runtime/Renderer/Private/DeferredShadingRenderer.cpp
--- a/Engine/Source/Runtime/Renderer/Private/DeferredShadingRenderer.cpp
+++ b/Engine/Source/Runtime/Renderer/Private/DeferredShadingRenderer.cpp
@@ -2254,6 +2254,10 @@ void FDeferredShadingSceneRenderer::Render(FRHICommandListImmediate& RHICmdList)
 	if (HairDatas)
 	{
 		RenderHairComposeSubPixel(RHICmdList, Views, HairDatas);
+	}
+
+	if (IsHairStrandsEnable(Scene->GetShaderPlatform()))
+	{
 		RenderHairStrandsDebugInfo(RHICmdList, Views, HairDatas);
 	}
 
diff --git a/Engine/Source/Runtime/Renderer/Private/HairStrands/HairStrandsDebug.cpp b/Engine/Source/Runtime/Renderer/Private/HairStrands/HairStrandsDebug.cpp
--- a/Engine/Source/Runtime/Renderer/Private/HairStrands/HairStrandsDebug.cpp
+++ b/Engine/Source/Runtime/Renderer/Private/HairStrands/HairStrandsDebug.cpp
@@ -826,12 +826,9 @@ void RenderHairStrandsDebugInfo(FRHICommandListImmediate& RHICmdList, TArray<FVi
 	const float YStep = 14;
 	const float ColumnWidth = 200;
 
-	if (Views.Num() == 0 || !HairDatas)
+	if (Views.Num() == 0)
 		return;
 
-	const FHairStrandsDeepShadowViews& InDomViews = HairDatas->DeepShadowViews;
-	const FHairStrandsClusterViews& InClusterViews = HairDatas->HairClusterPerViews;
-
 	// Only render debug information for the main view
 	const uint32 ViewIndex = 0;
 	FViewInfo& View = Views[ViewIndex];
@@ -842,38 +839,25 @@ void RenderHairStrandsDebugInfo(FRHICommandListImmediate& RHICmdList, TArray<FVi
 	const EHairStrandsDebugMode StrandsDebugMode = GetHairStrandsDebugStrandsMode();
 	const EHairDebugMode HairDebugMode = GetHairDebugMode();
 
+	float ClusterY = 38;
 	if (HairDebugMode == EHairDebugMode::ClusterData)
 	{
-		FViewElementPDI ShadowFrustumPDI(&View, nullptr, nullptr);
-		const FHairStrandsClusterDatas& ClusterDatas = InClusterViews.Views[ViewIndex];
-		for (const FHairStrandsClusterData& ClusterData : ClusterDatas.Datas)
-		{
-			const FBox ClusterBox(ClusterData.GetMinBound(), ClusterData.GetMaxBound());
-			DrawWireBox(&ShadowFrustumPDI, ClusterBox, FColor::Red, 0);
-		}
-
+		// Component part of the clusters
 		FRenderTargetTemp TempRenderTarget(View, (const FTexture2DRHIRef&)SceneTargets.GetSceneColor()->GetRenderTargetItem().TargetableTexture);
 		FCanvas Canvas(&TempRenderTarget, NULL, ViewFamily.CurrentRealTime, ViewFamily.CurrentWorldTime, ViewFamily.DeltaWorldTime, View.FeatureLevel);
 
 		float X = 20;
-		float Y = 38;
+		float Y = ClusterY;
 		FLinearColor InactiveColor(0.5, 0.5, 0.5);
 		FLinearColor DebugColor(1, 1, 0);
 		FString Line;
 
-		//const FVector4 HairComponent = GetHairComponents();
-		//const uint32 bHairR = HairComponent.X > 0 ? 1 : 0;
-		//const uint32 bHairTT = HairComponent.Y > 0 ? 1 : 0;
-		//const uint32 bHairTRT = HairComponent.Z > 0 ? 1 : 0;
-		//const uint32 bHairGlobalScattering = FMath::FloorToInt(HairComponent.W / 10.0f) > 0 ? 1 : 0;
-		//const uint32 bHairLocalScattering = FMath::Frac(HairComponent.W / 10.0f)*10.f > 0 ? 1 : 0;
-
 		const FHairStrandsDebugInfos DebugInfos = GetHairStandsDebugInfos();
 
-		Line = FString::Printf(TEXT("----------------------------------------------------------------"));				
+		Line = FString::Printf(TEXT("----------------------------------------------------------------"));
 		Canvas.DrawShadowedString(X, Y += YStep, *Line, GetStatsFont(), DebugColor);
-		
-		Line = FString::Printf(TEXT("Registered component count : %d"), DebugInfos.Num());									
+
+		Line = FString::Printf(TEXT("Registered component count : %d"), DebugInfos.Num());
 		Canvas.DrawShadowedString(X, Y += YStep, *Line, GetStatsFont(), DebugColor);
 
 		for (const FHairStrandsDebugInfo& DebugInfo : DebugInfos)
@@ -897,6 +881,37 @@ void RenderHairStrandsDebugInfo(FRHICommandListImmediate& RHICmdList, TArray<FVi
 			}
 		}
 
+		Canvas.Flush_RenderThread(RHICmdList);
+
+		ClusterY = Y;
+	}
+
+	if (!HairDatas)
+		return;
+
+	const FHairStrandsDeepShadowViews& InDomViews = HairDatas->DeepShadowViews;
+	const FHairStrandsClusterViews& InClusterViews = HairDatas->HairClusterPerViews;
+
+	if (HairDebugMode == EHairDebugMode::ClusterData)
+	{
+		// Clusters part of the clusters
+		FViewElementPDI ShadowFrustumPDI(&View, nullptr, nullptr);
+		const FHairStrandsClusterDatas& ClusterDatas = InClusterViews.Views[ViewIndex];
+		for (const FHairStrandsClusterData& ClusterData : ClusterDatas.Datas)
+		{
+			const FBox ClusterBox(ClusterData.GetMinBound(), ClusterData.GetMaxBound());
+			DrawWireBox(&ShadowFrustumPDI, ClusterBox, FColor::Red, 0);
+		}
+
+		FRenderTargetTemp TempRenderTarget(View, (const FTexture2DRHIRef&)SceneTargets.GetSceneColor()->GetRenderTargetItem().TargetableTexture);
+		FCanvas Canvas(&TempRenderTarget, NULL, ViewFamily.CurrentRealTime, ViewFamily.CurrentWorldTime, ViewFamily.DeltaWorldTime, View.FeatureLevel);
+
+		float X = 20;
+		float Y = ClusterY;
+		FLinearColor InactiveColor(0.5, 0.5, 0.5);
+		FLinearColor DebugColor(1, 1, 0);
+		FString Line;
+
 		Line = FString::Printf(TEXT("----------------------------------------------------------------"));
 		Canvas.DrawShadowedString(X, Y += YStep, *Line, GetStatsFont(), DebugColor);
 
