Fix rotation issue with hair skinning.

The projection and the mesh triangle update were wrongly done in world space.

(ushell-p4-cherrypick of 9541576 by Charles.deRousiers)

diff --git a/Engine/Shaders/Private/HairStrands/HairStrandsMeshProjection.usf b/Engine/Shaders/Private/HairStrands/HairStrandsMeshProjection.usf
--- a/Engine/Shaders/Private/HairStrands/HairStrandsMeshProjection.usf
+++ b/Engine/Shaders/Private/HairStrands/HairStrandsMeshProjection.usf
@@ -6,9 +6,7 @@
 uint bClear;
 uint MaxRootCount;
 float DistanceThreshold;
-float4x4 RootLocalToWorld;
 
-float4x4 MeshLocalToWorld;
 uint MeshNumPrimitives;
 uint MeshSectionIndex;
 uint MeshMaxIndexCount;
@@ -136,9 +134,7 @@ void MainCS(uint3 DispatchThreadId : SV_DispatchThreadID)
 
 	if (RootIndex >= MaxRootCount) return;
 
-	float3 RootPosition = RootPositionBuffer[RootIndex].xyz;
-	RootPosition = mul(float4(RootPosition, 1), RootLocalToWorld).xyz;
-
+	const float3 RootPosition = RootPositionBuffer[RootIndex].xyz;
 	const float3 RootNormal = RootNormalBuffer[RootIndex].xyz;
 
 	float MinDistance = bClear ? 999999999 : OutRootTriangleDistance[RootIndex];
@@ -146,7 +142,7 @@ void MainCS(uint3 DispatchThreadId : SV_DispatchThreadID)
 	float3 TriangleBarycentric = float3(0,0,0);
 	for (uint PrimitiveIt = 0; PrimitiveIt < MeshNumPrimitives; ++PrimitiveIt)
 	{
-		const FHairMeshTriangle Tri = GetTriangleData(PrimitiveIt, MeshLocalToWorld, MeshPositionBuffer, MeshIndexBuffer, MeshIndexOffset, MeshMaxIndexCount, MeshMaxVertexCount);
+		const FHairMeshTriangle Tri = GetTriangleData(PrimitiveIt, MeshPositionBuffer, MeshIndexBuffer, MeshIndexOffset, MeshMaxIndexCount, MeshMaxVertexCount);
 		
 		const float3 B = 0.5f;
 		FTrianglePoint C = ComputeClosestPoint(Tri, RootPosition);
diff --git a/Engine/Shaders/Private/HairStrands/HairStrandsMeshProjectionCommon.ush b/Engine/Shaders/Private/HairStrands/HairStrandsMeshProjectionCommon.ush
--- a/Engine/Shaders/Private/HairStrands/HairStrandsMeshProjectionCommon.ush
+++ b/Engine/Shaders/Private/HairStrands/HairStrandsMeshProjectionCommon.ush
@@ -25,7 +25,7 @@ float3 GetVertexPosition(uint VId, Buffer<float> InPositionBuffer, Buffer<uint>
 	return P;
 }
 
-FHairMeshTriangle GetTriangleData(uint TriangleIndex, float4x4 LocalToWorld, Buffer<float> InPositionBuffer, Buffer<uint> InIndexBuffer, uint InIndexOffset, uint InMaxIndexCount, uint InMaxVertexCount)
+FHairMeshTriangle GetTriangleData(uint TriangleIndex, Buffer<float> InPositionBuffer, Buffer<uint> InIndexBuffer, uint InIndexOffset, uint InMaxIndexCount, uint InMaxVertexCount)
 {
 	const uint BaseVertexId = TriangleIndex*3;
 	const uint VId0 = BaseVertexId;
@@ -36,11 +36,6 @@ FHairMeshTriangle GetTriangleData(uint TriangleIndex, float4x4 LocalToWorld, Buf
 	Out.P0 = GetVertexPosition(VId0, InPositionBuffer, InIndexBuffer, InIndexOffset, InMaxIndexCount, InMaxVertexCount);
 	Out.P1 = GetVertexPosition(VId1, InPositionBuffer, InIndexBuffer, InIndexOffset, InMaxIndexCount, InMaxVertexCount);
 	Out.P2 = GetVertexPosition(VId2, InPositionBuffer, InIndexBuffer, InIndexOffset, InMaxIndexCount, InMaxVertexCount);
-
-	Out.P0 = mul(float4(Out.P0,1), LocalToWorld).xyz;
-	Out.P1 = mul(float4(Out.P1,1), LocalToWorld).xyz;
-	Out.P2 = mul(float4(Out.P2,1), LocalToWorld).xyz;
-
 	Out.E0 = Out.P1 - Out.P0;
 	Out.E1 = Out.P2 - Out.P0;
 	Out.N = normalize(cross(Out.E1, Out.E0));
diff --git a/Engine/Shaders/Private/HairStrands/HairStrandsMeshProjectionMeshDebug.usf b/Engine/Shaders/Private/HairStrands/HairStrandsMeshProjectionMeshDebug.usf
--- a/Engine/Shaders/Private/HairStrands/HairStrandsMeshProjectionMeshDebug.usf
+++ b/Engine/Shaders/Private/HairStrands/HairStrandsMeshProjectionMeshDebug.usf
@@ -33,7 +33,7 @@ void MainVS(
 #if PERMUTATION_INPUT_TYPE == 1
 	const uint TriangleVertexIndex = (VertexId%3);
 	const uint TriangleIndex = VertexId / 3;
-	const FHairMeshTriangle Tri = GetTriangleData(TriangleIndex, LocalToWorld, InputVertexBuffer, InputIndexBuffer, IndexOffset, MaxIndexCount, MaxVertexCount);
+	const FHairMeshTriangle Tri = GetTriangleData(TriangleIndex, InputVertexBuffer, InputIndexBuffer, IndexOffset, MaxIndexCount, MaxVertexCount);
 	OutFaceNormal = Tri.N;
 
 	float3 P = 0;
@@ -44,7 +44,8 @@ void MainVS(
 	if (TriangleVertexIndex == 2)
 		P = Tri.P2;
 #endif
-
+	// TODO: the normal is in local space, not in world space so it won't appear correct
+	P = mul(float4(P, 1), LocalToWorld).xyz;
 	OutPosition = mul(float4(P,1), View.WorldToClip);
 }
 
diff --git a/Engine/Shaders/Private/HairStrands/HairStrandsMeshUpdate.usf b/Engine/Shaders/Private/HairStrands/HairStrandsMeshUpdate.usf
--- a/Engine/Shaders/Private/HairStrands/HairStrandsMeshUpdate.usf
+++ b/Engine/Shaders/Private/HairStrands/HairStrandsMeshUpdate.usf
@@ -4,7 +4,6 @@
 #include "HairStrandsMeshProjectionCommon.ush"
 
 float3 RootTrianglePositionOffset;
-float4x4 MeshLocalToWorld;
 uint MeshSectionIndex;
 uint MeshMaxIndexCount;
 uint MeshMaxVertexCount;
@@ -12,7 +11,6 @@ uint MeshIndexOffset;
 Buffer<uint> MeshIndexBuffer;
 Buffer<float> MeshPositionBuffer;
 
-float4x4 RootLocalToWorld;
 uint MaxRootCount;
 Buffer<uint> RootTriangleIndex;
 
@@ -31,7 +29,7 @@ void MainCS(uint3 DispatchThreadId : SV_DispatchThreadID)
 	if (TriangleIndex.SectionIndex != MeshSectionIndex)
 		return;
 
-	const FHairMeshTriangle Tri = GetTriangleData(TriangleIndex.TriangleIndex, MeshLocalToWorld, MeshPositionBuffer, MeshIndexBuffer, MeshIndexOffset, MeshMaxIndexCount, MeshMaxVertexCount);
+	const FHairMeshTriangle Tri = GetTriangleData(TriangleIndex.TriangleIndex, MeshPositionBuffer, MeshIndexBuffer, MeshIndexOffset, MeshMaxIndexCount, MeshMaxVertexCount);
 
 	OutRootTrianglePosition0[RootIndex] = (float4(Tri.P0 - RootTrianglePositionOffset,1));
 	OutRootTrianglePosition1[RootIndex] = (float4(Tri.P1 - RootTrianglePositionOffset,1));
diff --git a/Engine/Source/Runtime/Renderer/Private/HairStrands/HairStrandsMeshProjection.cpp b/Engine/Source/Runtime/Renderer/Private/HairStrands/HairStrandsMeshProjection.cpp
--- a/Engine/Source/Runtime/Renderer/Private/HairStrands/HairStrandsMeshProjection.cpp
+++ b/Engine/Source/Runtime/Renderer/Private/HairStrands/HairStrandsMeshProjection.cpp
@@ -20,9 +20,7 @@ class FHairMeshProjectionCS : public FGlobalShader
 		SHADER_PARAMETER(uint32, bClear)
 		SHADER_PARAMETER(uint32, MaxRootCount)
 		SHADER_PARAMETER(uint32, MaxSpinLockCount)
-		SHADER_PARAMETER(FMatrix, RootLocalToWorld)
 
-		SHADER_PARAMETER(FMatrix, MeshLocalToWorld)
 		SHADER_PARAMETER(uint32, MeshNumPrimitives)
 		SHADER_PARAMETER(uint32, MeshSectionIndex)
 		SHADER_PARAMETER(uint32, MeshMaxIndexCount)
@@ -67,21 +65,12 @@ static void AddHairStrandMeshProjectionPass(
 		return;
 	}
 
-	// Change the transform to be just a relative transform from the mesh
-	FTransform RootTranform = RootData.LocalToWorld;
-	FTransform MeshTranform = MeshSectionData.LocalToWorld;
-	RootTranform.SetLocation(RootTranform.GetLocation() - MeshTranform.GetLocation());
-	MeshTranform.SetLocation(FVector(0, 0, 0));
-
 	FHairMeshProjectionCS::FParameters* Parameters = GraphBuilder.AllocParameters<FHairMeshProjectionCS::FParameters>();
 	Parameters->bClear				= bClear ? 1 : 0;
 	Parameters->MaxRootCount		= RootData.RootCount;
 	Parameters->MaxSpinLockCount	= FMath::Clamp(GHairProjectionMaxSpinLockCount, 0, 100000);
 	Parameters->RootPositionBuffer	= RootData.RootPositionBuffer;
 	Parameters->RootNormalBuffer	= RootData.RootNormalBuffer;
-	Parameters->RootLocalToWorld	= RootTranform.ToMatrixWithScale();
-
-	Parameters->MeshLocalToWorld	= MeshTranform.ToMatrixWithScale();
 	Parameters->MeshNumPrimitives	= MeshSectionData.NumPrimitives;
 	Parameters->MeshSectionIndex	= MeshSectionData.SectionIndex;
 	Parameters->MeshMaxIndexCount	= MeshSectionData.TotalIndexCount;
@@ -141,9 +130,7 @@ class FHairUpdateMeshTriangleCS : public FGlobalShader
 	BEGIN_SHADER_PARAMETER_STRUCT(FParameters, )
 		SHADER_PARAMETER(FVector, RootTrianglePositionOffset)
 		SHADER_PARAMETER(uint32, MaxRootCount)
-		SHADER_PARAMETER(FMatrix, RootLocalToWorld)
-
-		SHADER_PARAMETER(FMatrix, MeshLocalToWorld)
+		
 		SHADER_PARAMETER(uint32, MeshSectionIndex)
 		SHADER_PARAMETER(uint32, MeshMaxIndexCount)
 		SHADER_PARAMETER(uint32, MeshMaxVertexCount)
@@ -178,18 +165,9 @@ static void AddHairStrandUpdateMeshTrianglesPass(
 	FHairStrandsProjectionHairData::LODData& LODData = RootData.LODDatas[LODIndex];
 	check(LODData.LODIndex == LODIndex);
 
-	// Change the transform to be just a relative transform from the mesh
-	FTransform RootTranform = RootData.LocalToWorld;
-	FTransform MeshTranform = MeshSectionData.LocalToWorld;
-	RootTranform.SetLocation(RootTranform.GetLocation() - MeshTranform.GetLocation());
-	MeshTranform.SetLocation(FVector(0, 0, 0));
-
 	FHairUpdateMeshTriangleCS::FParameters* Parameters = GraphBuilder.AllocParameters<FHairUpdateMeshTriangleCS::FParameters>();
 	Parameters->MaxRootCount		= RootData.RootCount;
 	Parameters->RootTriangleIndex	= LODData.RootTriangleIndexBuffer->SRV;
-	Parameters->RootLocalToWorld	= RootTranform.ToMatrixWithScale();
-	
-	Parameters->MeshLocalToWorld	= MeshTranform.ToMatrixWithScale();
 	Parameters->MeshSectionIndex	= MeshSectionData.SectionIndex;
 	Parameters->MeshMaxIndexCount	= MeshSectionData.TotalIndexCount;
 	Parameters->MeshMaxVertexCount	= MeshSectionData.TotalVertexCount;
