Remove tesselation stages from the HairStrands vertex factory as this is unessary/unused and it simplifies the code.

(ushell-p4-cherrypick of 9547581 by Charles.deRousiers)

diff --git a/Engine/Plugins/Experimental/HairStrands/Source/HairStrandsCore/Public/HairStrandsVertexFactory.h b/Engine/Plugins/Experimental/HairStrands/Source/HairStrandsCore/Public/HairStrandsVertexFactory.h
--- a/Engine/Plugins/Experimental/HairStrands/Source/HairStrandsCore/Public/HairStrandsVertexFactory.h
+++ b/Engine/Plugins/Experimental/HairStrands/Source/HairStrandsCore/Public/HairStrandsVertexFactory.h
@@ -66,7 +66,7 @@ public:
 	virtual void InitRHI() override;
 	virtual void ReleaseRHI() override;
 
-	static bool SupportsTessellationShaders() { return true; }
+	static bool SupportsTessellationShaders() { return false; }
 
 	static FVertexFactoryShaderParameters* ConstructShaderParameters(EShaderFrequency ShaderFrequency);
 
diff --git a/Engine/Shaders/Private/HairStrands/HairStrandsVisibilityPS.usf b/Engine/Shaders/Private/HairStrands/HairStrandsVisibilityPS.usf
--- a/Engine/Shaders/Private/HairStrands/HairStrandsVisibilityPS.usf
+++ b/Engine/Shaders/Private/HairStrands/HairStrandsVisibilityPS.usf
@@ -24,7 +24,6 @@ uint HairVisibilityPass_HairClusterIndex;
 void MainVisibility(
 	centroid FVertexFactoryInterpolantsVSToPS Interpolants,
 	in float4 SvPosition	: SV_Position,
-	float4 PixelPosition : TEXCOORD6,
 	centroid in float WorldStrandRadius : TEXCOORD8,
 	const float4 PackedVelocityA : TEXCOORD9,
 	const float4 PackedVelocityC : TEXCOORD12
diff --git a/Engine/Shaders/Private/HairStrands/HairStrandsVisibilityVS.usf b/Engine/Shaders/Private/HairStrands/HairStrandsVisibilityVS.usf
--- a/Engine/Shaders/Private/HairStrands/HairStrandsVisibilityVS.usf
+++ b/Engine/Shaders/Private/HairStrands/HairStrandsVisibilityVS.usf
@@ -9,57 +9,12 @@ struct FHairVisibilityAccumulateVSToPS
 {
 	FVertexFactoryInterpolantsVSToPS FactoryInterpolants;
 	float4 Position : SV_POSITION;
-	float4 PixelPosition : TEXCOORD6;
 	float WorldStrandRadius : TEXCOORD8;
 	float4 PackedVelocityA : TEXCOORD9;
 	float4 PackedVelocityC : TEXCOORD12;
 };
 
-#if USING_TESSELLATION	
-	struct FHairVisibilityAccumulateVSToDS
-	{
-		FVertexFactoryInterpolantsVSToDS FactoryInterpolants;
-		float4 Position : VS_To_DS_Position;
-		OPTIONAL_VertexID_VS_To_DS
-	};
-	
-	#define VS_OUTPUT_TYPE FHairVisibilityAccumulateVSToDS
-#else
-	#define VS_OUTPUT_TYPE FHairVisibilityAccumulateVSToPS
-#endif
-
-#if USING_TESSELLATION	
-	#define FPassSpecificVSToDS FHairVisibilityAccumulateVSToDS
-	#define FPassSpecificVSToPS FHairVisibilityAccumulateVSToPS
-	#define INTERPOLATE_MEMBER(member) O.member = a.member * aInterp + b.member * bInterp
-
-	FHairVisibilityAccumulateVSToDS PassInterpolate(FHairVisibilityAccumulateVSToDS a, float aInterp, FHairVisibilityAccumulateVSToDS b, float bInterp)
-	{
-		FHairVisibilityAccumulateVSToDS O;
-		
-		O.FactoryInterpolants = VertexFactoryInterpolate(a.FactoryInterpolants, aInterp, b.FactoryInterpolants, bInterp);
-
-		return O;
-	}
-
-	#undef INTERPOLATE_MEMBER
-	
-	FHairVisibilityAccumulateVSToPS PassFinalizeTessellationOutput(FHairVisibilityAccumulateVSToDS Interpolants, float4 WorldPosition, FMaterialTessellationParameters MaterialParameters)
-	{
-		FHairVisibilityAccumulateVSToPS O;
-		
-		O.FactoryInterpolants = VertexFactoryAssignInterpolants(Interpolants.FactoryInterpolants);
-		
-		// Finally, transform position to clip-space
-		O.Position = mul(WorldPosition, HairVisibilityPass.WorldToClipMatrix); // mul(WorldPosition, View.TranslatedWorldToClip);
-			
-		O.PixelPosition = WorldPosition;
-		
-		return O;
-	}
-	
-#include "Tessellation.ush"
-#endif
+#define VS_OUTPUT_TYPE FHairVisibilityAccumulateVSToPS
 
 #if VERTEXSHADER
 
@@ -116,18 +71,8 @@ void Main(
 	//	PrevWorldPosition.xyz += GetMaterialPreviousWorldPositionOffset(VertexParameters); 
 	//}
 
-#if USING_TESSELLATION
-	Output.Position = WorldPosition;
-	Output.FactoryInterpolants = VertexFactoryGetInterpolantsVSToDS(Input, VFIntermediates, VertexParameters);
-	
-#else // !USING_TESSELLATION
-
-//	Output.Position = mul(WorldPosition, HairVisibilityPass.WorldToClipMatrix);
 	Output.Position = mul(WorldPosition, View.TranslatedWorldToClip);
 	Output.FactoryInterpolants = VertexFactoryGetInterpolantsVSToPS(Input, VFIntermediates, VertexParameters);
-	Output.PixelPosition = WorldPosition;
-
-#endif // USING_TESSELLATION
 
 #ifdef VF_STRAND_HAIR
 	Output.WorldStrandRadius = VFIntermediates.HairDimensions.y;
diff --git a/Engine/Source/Runtime/Renderer/Private/HairStrands/HairStrandsVisibility.cpp b/Engine/Source/Runtime/Renderer/Private/HairStrands/HairStrandsVisibility.cpp
--- a/Engine/Source/Runtime/Renderer/Private/HairStrands/HairStrandsVisibility.cpp
+++ b/Engine/Source/Runtime/Renderer/Private/HairStrands/HairStrandsVisibility.cpp
@@ -121,49 +121,7 @@ protected:
 	}
 };
 
-class FHairVisibilityHS : public FBaseHS
-{
-	DECLARE_SHADER_TYPE(FHairVisibilityHS,MeshMaterial);
-
-protected:
-
-	FHairVisibilityHS(const ShaderMetaType::CompiledShaderInitializerType& Initializer)
-		: FBaseHS(Initializer)
-	{
-		PassUniformBuffer.Bind(Initializer.ParameterMap, FHairVisibilityPassUniformParameters::StaticStructMetadata.GetShaderVariableName());
-	}
-
-	FHairVisibilityHS() {}
-
-	static bool ShouldCompilePermutation(const FMeshMaterialShaderPermutationParameters& Parameters)
-	{
-		return FBaseHS::ShouldCompilePermutation(Parameters) && IsCompatibleWithHairVisibility(Parameters);
-	}
-};
-
-class FHairVisibilityDS : public FBaseDS
-{
-	DECLARE_SHADER_TYPE(FHairVisibilityDS,MeshMaterial);
-
-protected:
-
-	FHairVisibilityDS(const ShaderMetaType::CompiledShaderInitializerType& Initializer)
-		: FBaseDS(Initializer)
-	{
-		PassUniformBuffer.Bind(Initializer.ParameterMap, FHairVisibilityPassUniformParameters::StaticStructMetadata.GetShaderVariableName());
-	}
-
-	FHairVisibilityDS() {}
-
-	static bool ShouldCompilePermutation(const FMeshMaterialShaderPermutationParameters& Parameters)
-	{
-		return FBaseDS::ShouldCompilePermutation(Parameters) && IsCompatibleWithHairVisibility(Parameters);
-	}
-};
-
 IMPLEMENT_MATERIAL_SHADER_TYPE(, FHairVisibilityVS, TEXT("/Engine/Private/HairStrands/HairStrandsVisibilityVS.usf"), TEXT("Main"), SF_Vertex);
-IMPLEMENT_MATERIAL_SHADER_TYPE(, FHairVisibilityHS, TEXT("/Engine/Private/HairStrands/HairStrandsVisibilityVS.usf"), TEXT("MainHull"), SF_Hull);
-IMPLEMENT_MATERIAL_SHADER_TYPE(, FHairVisibilityDS, TEXT("/Engine/Private/HairStrands/HairStrandsVisibilityVS.usf"), TEXT("MainDomain"), SF_Domain);
 
 /////////////////////////////////////////////////////////////////////////////////////////
 
@@ -303,23 +261,11 @@ void FHairVisibilityProcessor::Process(
 
 	TMeshProcessorShaders<
 		FHairVisibilityVS,
-		FHairVisibilityHS,
-		FHairVisibilityDS,
+		FMeshMaterialShader,
+		FMeshMaterialShader,
 		FHairVisibilityPS<TRenderMode>> PassShaders;
 	{
-		const EMaterialTessellationMode MaterialTessellationMode = MaterialResource.GetTessellationMode();
 		FVertexFactoryType* VertexFactoryType = VertexFactory->GetType();
-
-		const bool bNeedsHSDS = RHISupportsTessellation(GShaderPlatformForFeatureLevel[FeatureLevel])
-			&& VertexFactoryType->SupportsTessellationShaders()
-			&& MaterialTessellationMode != MTM_NoTessellation;
-
-		if (bNeedsHSDS)
-		{
-			PassShaders.DomainShader = MaterialResource.GetShader<FHairVisibilityDS>(VertexFactoryType);
-			PassShaders.HullShader = MaterialResource.GetShader<FHairVisibilityHS>(VertexFactoryType);
-		}
-
 		PassShaders.VertexShader = MaterialResource.GetShader<FHairVisibilityVS>(VertexFactoryType);
 		PassShaders.PixelShader = MaterialResource.GetShader<FHairVisibilityPS<TRenderMode>>(VertexFactoryType);
 	}
