Add optional coverage pass (no plugged yet).

(ushell-p4-cherrypick of 9550639 by Charles.deRousiers)

diff --git a/Engine/Shaders/Private/HairStrands/HairStrandsVisibilityPS.usf b/Engine/Shaders/Private/HairStrands/HairStrandsVisibilityPS.usf
--- a/Engine/Shaders/Private/HairStrands/HairStrandsVisibilityPS.usf
+++ b/Engine/Shaders/Private/HairStrands/HairStrandsVisibilityPS.usf
@@ -12,28 +12,47 @@
 #include "/Engine/Generated/Material.ush"
 #include "/Engine/Generated/VertexFactory.ush"
 
-#define SUPPORT_OPACITY_MASK 0
-#define SUPPORT_MATERIAL_PROPERTY 1
+#define RENDER_MODE_MSAA 0
+#define RENDER_MODE_COVERAGE 1
 
 #ifndef HAIR_RENDER_MODE
 #error Undefined hair render mode
 #endif
 
+#define SUPPORT_OPACITY_MASK 0
+
+#if HAIR_RENDER_MODE == RENDER_MODE_MSAA
+#define SUPPORT_MATERIAL_PROPERTY 1
+#endif
+
+#if HAIR_RENDER_MODE == RENDER_MODE_COVERAGE
+#define SUPPORT_MATERIAL_PROPERTY 0
+#endif
+
+
 uint HairVisibilityPass_HairClusterIndex;
 
 void MainVisibility(
 	centroid FVertexFactoryInterpolantsVSToPS Interpolants,
 	in float4 SvPosition	: SV_Position,
+#if HAIR_RENDER_MODE == RENDER_MODE_MSAA
 	centroid in float WorldStrandRadius : TEXCOORD8,
 	const float4 PackedVelocityA : TEXCOORD9,
 	const float4 PackedVelocityC : TEXCOORD12
+#elif HAIR_RENDER_MODE == RENDER_MODE_COVERAGE
+	centroid in float WorldStrandRadius : TEXCOORD8
+#endif
 	OPTIONAL_IsFrontFace,
-	out uint4  OutColor0	: SV_Target0
-	#if SUPPORT_MATERIAL_PROPERTY
+#if HAIR_RENDER_MODE == RENDER_MODE_COVERAGE
+	out float OutColor0 : SV_Target0
+#elif HAIR_RENDER_MODE == RENDER_MODE_MSAA
+	out uint4 OutColor0	: SV_Target0
+  #if SUPPORT_MATERIAL_PROPERTY
 	, out float4 OutColor1	: SV_Target1
 	, out float4 OutColor2	: SV_Target2
 	, out float2 OutColor3	: SV_Target3
-	#endif
+  #endif
+#endif
 )
 {
 	ResolvedView = ResolveView();
@@ -55,18 +74,22 @@ void MainVisibility(
 		Coverage = StrandRealRadius / max(StrandRealRadius, PixelRadius);
 	}
 
+#if HAIR_RENDER_MODE == RENDER_MODE_COVERAGE
+	OutColor0 = Coverage;
+#endif
+
+#if HAIR_RENDER_MODE == RENDER_MODE_MSAA
 	FVisibilityData VisibilityData;
-#if VF_STRAND_HAIR
+  #if VF_STRAND_HAIR
 	VisibilityData.Tangent = normalize(Interpolants.TangentToWorld2.xyz);
 	VisibilityData.PrimitiveId = Interpolants.HairPrimitiveId;
-#else
+  #else
 	VisibilityData.Tangent = float3(0, 0, 1);
 	VisibilityData.PrimitiveId = 0;
-#endif
+  #endif
 	VisibilityData.ClusterId = HairVisibilityPass_HairClusterIndex;
 	VisibilityData.Coverage8bit = To8bitCoverage(Coverage);
 
-
 	half3 BaseColor = 0;
 	float Roughness = 0;
 	float Specular  = 0;
@@ -97,6 +120,8 @@ void MainVisibility(
 	OutColor1 = float4(sqrt(BaseColor), Roughness);
 	OutColor2 = float4(Specular,0,0,0);
 	OutColor3 = EncodedVelocity;
+#endif
+
 }
 
 
diff --git a/Engine/Shaders/Private/HairStrands/HairStrandsVisibilityVS.usf b/Engine/Shaders/Private/HairStrands/HairStrandsVisibilityVS.usf
--- a/Engine/Shaders/Private/HairStrands/HairStrandsVisibilityVS.usf
+++ b/Engine/Shaders/Private/HairStrands/HairStrandsVisibilityVS.usf
@@ -5,13 +5,18 @@
 #include "/Engine/Generated/Material.ush"
 #include "/Engine/Generated/VertexFactory.ush"
 
+#define RENDER_MODE_MSAA 0
+#define RENDER_MODE_COVERAGE 1
+
 struct FHairVisibilityAccumulateVSToPS
 {
 	FVertexFactoryInterpolantsVSToPS FactoryInterpolants;
 	float4 Position : SV_POSITION;
 	float WorldStrandRadius : TEXCOORD8;
+#if HAIR_RENDER_MODE == RENDER_MODE_MSAA
 	float4 PackedVelocityA : TEXCOORD9;
 	float4 PackedVelocityC : TEXCOORD12;
+#endif
 };
 
 #define VS_OUTPUT_TYPE FHairVisibilityAccumulateVSToPS
@@ -42,6 +47,7 @@ void Main(
 	// The velocity is used to adapt/increase the minimum rasterization size in order to avoid aliasing under heavy-motion.
 	// The higher the velocity, the larger a strand becomes.
 	float VelocityScale = 0;
+	#if HAIR_RENDER_MODE == RENDER_MODE_MSAA
 	{
 		const float4 ScreenPos			= mul(float4(WorldPositionRaw.xyz, 1), ResolvedView.TranslatedWorldToClip);
 		const float4 PrevScreenPosObj	= mul(float4(PrevWorldPosition.xyz, 1), ResolvedView.PrevTranslatedWorldToClip);
@@ -51,6 +57,7 @@ void Main(
 		const float2 ScreenVelocity = Calculate2DVelocity(Output.PackedVelocityA, Output.PackedVelocityC);
 		VelocityScale = saturate(max(abs(ScreenVelocity.x), abs(ScreenVelocity.y)) / HairVisibilityPass.HairStrandsVelocityScale);
 	}
+	#endif
 
 #if VF_STRAND_HAIR
 	float4 WorldPosition = VertexFactoryGetWorldPosition(Input, VFIntermediates, VelocityScale);
diff --git a/Engine/Source/Runtime/Renderer/Private/HairStrands/HairStrandsVisibility.cpp b/Engine/Source/Runtime/Renderer/Private/HairStrands/HairStrandsVisibility.cpp
--- a/Engine/Source/Runtime/Renderer/Private/HairStrands/HairStrandsVisibility.cpp
+++ b/Engine/Source/Runtime/Renderer/Private/HairStrands/HairStrandsVisibility.cpp
@@ -28,6 +28,9 @@ IMPLEMENT_GLOBAL_SHADER_PARAMETER_STRUCT(FHairVisibilityPassUniformParameters, "
 
 /////////////////////////////////////////////////////////////////////////////////////////
 
+static int32 GHairStrandsCoveragePassEnable = 0;
+static FAutoConsoleVariableRef CVarHairStrandsCoveragePassEnable(TEXT("r.HairStrands.CoveragePass"), GHairStrandsCoveragePassEnable, TEXT("Enable accurate coverage pass"));
+
 static int32 GHairVisibilitySampleCount = 8;
 static FAutoConsoleVariableRef CVarHairVisibilitySampleCount(TEXT("r.HairStrands.VisibilitySampleCount"), GHairVisibilitySampleCount, TEXT("Hair strands visibility sample count"));
 
@@ -45,6 +48,7 @@ static FAutoConsoleVariableRef CVarHairVelocityMagnitudeScale(TEXT("r.HairStrand
 enum EHairVisibilityRenderMode
 {
 	HairVisibilityRenderMode_MSAA,
+	HairVisibilityRenderMode_Coverage,
 	HairVisibilityRenderModeCount
 };
 
@@ -98,6 +102,7 @@ static FIntPoint GetVendorOptimalGroupSize2D()
 	}
 }
 
+template<EHairVisibilityRenderMode RenderMode>
 class FHairVisibilityVS : public FMeshMaterialShader
 {
 	DECLARE_SHADER_TYPE(FHairVisibilityVS,MeshMaterial);
@@ -117,15 +122,21 @@ protected:
 
 	static bool ShouldCompilePermutation(const FMeshMaterialShaderPermutationParameters& Parameters)
 	{
-		return IsCompatibleWithHairVisibility(Parameters);
+		return IsCompatibleWithHairVisibility(Parameters) && Parameters.VertexFactoryType->GetFName() == FName(TEXT("FHairStrandsVertexFactory"));
 	}
-};
 
-IMPLEMENT_MATERIAL_SHADER_TYPE(, FHairVisibilityVS, TEXT("/Engine/Private/HairStrands/HairStrandsVisibilityVS.usf"), TEXT("Main"), SF_Vertex);
+	static void ModifyCompilationEnvironment(const FMaterialShaderPermutationParameters& Parameters, FShaderCompilerEnvironment& OutEnvironment)
+	{
+		FMaterialShader::ModifyCompilationEnvironment(Parameters, OutEnvironment);
+		const uint32 RenderModeValue = uint32(RenderMode);
+		OutEnvironment.SetDefine(TEXT("HAIR_RENDER_MODE"), RenderModeValue);
+	}
+};
+IMPLEMENT_MATERIAL_SHADER_TYPE(template<>, FHairVisibilityVS<HairVisibilityRenderMode_MSAA>, TEXT("/Engine/Private/HairStrands/HairStrandsVisibilityVS.usf"), TEXT("Main"), SF_Vertex);
+IMPLEMENT_MATERIAL_SHADER_TYPE(template<>, FHairVisibilityVS<HairVisibilityRenderMode_Coverage>, TEXT("/Engine/Private/HairStrands/HairStrandsVisibilityVS.usf"), TEXT("Main"), SF_Vertex);
 
 /////////////////////////////////////////////////////////////////////////////////////////
 
-
 class FHairVisibilityShaderElementData : public FMeshMaterialShaderElementData
 {
 public:
@@ -133,7 +144,6 @@ public:
 	uint32 HairClusterId;
 };
 
-
 template<EHairVisibilityRenderMode RenderMode>
 class FHairVisibilityPS : public FMeshMaterialShader
 {
@@ -154,13 +164,13 @@ public:
 
 	static bool ShouldCompilePermutation(const FMeshMaterialShaderPermutationParameters& Parameters)
 	{
-		return IsCompatibleWithHairStrands(Parameters.Material, Parameters.Platform);
+		return IsCompatibleWithHairStrands(Parameters.Material, Parameters.Platform) && Parameters.VertexFactoryType->GetFName() == FName(TEXT("FHairStrandsVertexFactory"));
 	}
 
 	static void ModifyCompilationEnvironment(const FMaterialShaderPermutationParameters& Parameters, FShaderCompilerEnvironment& OutEnvironment)	
 	{
 		FMaterialShader::ModifyCompilationEnvironment(Parameters, OutEnvironment);
-		const uint32 RenderModeValue = RenderMode == HairVisibilityRenderMode_MSAA ? 0 : 1;
+		const uint32 RenderModeValue = uint32(RenderMode);
 		OutEnvironment.SetDefine(TEXT("HAIR_RENDER_MODE"), RenderModeValue);
 	}
 
@@ -188,6 +198,7 @@ public:
 	FShaderParameter HairVisibilityPass_HairClusterIndex;
 };
 IMPLEMENT_MATERIAL_SHADER_TYPE(template<>, FHairVisibilityPS<HairVisibilityRenderMode_MSAA>, TEXT("/Engine/Private/HairStrands/HairStrandsVisibilityPS.usf"), TEXT("MainVisibility"), SF_Pixel);
+IMPLEMENT_MATERIAL_SHADER_TYPE(template<>, FHairVisibilityPS<HairVisibilityRenderMode_Coverage>, TEXT("/Engine/Private/HairStrands/HairStrandsVisibilityPS.usf"), TEXT("MainVisibility"), SF_Pixel);
 
 /////////////////////////////////////////////////////////////////////////////////////////
 
@@ -228,12 +239,16 @@ void FHairVisibilityProcessor::AddMeshBatch(const FMeshBatch& RESTRICT MeshBatch
 
 void FHairVisibilityProcessor::AddMeshBatch(const FMeshBatch& RESTRICT MeshBatch, uint64 BatchElementMask, const FPrimitiveSceneProxy* RESTRICT PrimitiveSceneProxy, int32 StaticMeshId, uint32 HairClusterId)
 {
+	static const FVertexFactoryType* CompatibleVF = FVertexFactoryType::GetVFByName(TEXT("FHairStrandsVertexFactory"));
+
 	// Determine the mesh's material and blend mode.
 	const FMaterialRenderProxy* FallbackMaterialRenderProxyPtr = nullptr;
 	const FMaterial& Material = MeshBatch.MaterialRenderProxy->GetMaterialWithFallback(FeatureLevel, FallbackMaterialRenderProxyPtr);
 	const bool bIsCompatible = IsCompatibleWithHairStrands(&Material, FeatureLevel);
+	const bool bIsHairStrandsFactory = MeshBatch.VertexFactory->GetType()->GetId() == CompatibleVF->GetId();
 
-	if (bIsCompatible
+	if (bIsCompatible 
+		&& bIsHairStrandsFactory
 		&& (!PrimitiveSceneProxy || PrimitiveSceneProxy->ShouldRenderInMainPass())
 		&& ShouldIncludeDomainInMeshPass(Material.GetMaterialDomain()))
 	{
@@ -242,6 +257,8 @@ void FHairVisibilityProcessor::AddMeshBatch(const FMeshBatch& RESTRICT MeshBatch
 		const ERasterizerCullMode MeshCullMode = ComputeMeshCullMode(MeshBatch, Material);
 		if (RenderMode == HairVisibilityRenderMode_MSAA)
 			Process<HairVisibilityRenderMode_MSAA>(MeshBatch, BatchElementMask, PrimitiveSceneProxy, StaticMeshId, MaterialRenderProxy, Material, HairClusterId, MeshFillMode, MeshCullMode);
+		if (RenderMode == HairVisibilityRenderMode_Coverage)
+			Process<HairVisibilityRenderMode_Coverage>(MeshBatch, BatchElementMask, PrimitiveSceneProxy, StaticMeshId, MaterialRenderProxy, Material, HairClusterId, MeshFillMode, MeshCullMode);
 	}
 }
 
@@ -260,18 +277,17 @@ void FHairVisibilityProcessor::Process(
 	const FVertexFactory* VertexFactory = MeshBatch.VertexFactory;
 
 	TMeshProcessorShaders<
-		FHairVisibilityVS,
+		FHairVisibilityVS<TRenderMode>,
 		FMeshMaterialShader,
 		FMeshMaterialShader,
 		FHairVisibilityPS<TRenderMode>> PassShaders;
 	{
 		FVertexFactoryType* VertexFactoryType = VertexFactory->GetType();
-		PassShaders.VertexShader = MaterialResource.GetShader<FHairVisibilityVS>(VertexFactoryType);
-		PassShaders.PixelShader = MaterialResource.GetShader<FHairVisibilityPS<TRenderMode>>(VertexFactoryType);
+		PassShaders.VertexShader = MaterialResource.GetShader<FHairVisibilityVS<TRenderMode>>(VertexFactoryType);
+		PassShaders.PixelShader  = MaterialResource.GetShader<FHairVisibilityPS<TRenderMode>>(VertexFactoryType);
 	}
 
 	FMeshPassProcessorRenderState DrawRenderState(PassDrawRenderState);
-
 	FHairVisibilityShaderElementData ShaderElementData(HairClusterId);
 	ShaderElementData.InitializeMeshMaterialData(ViewIfDynamicMeshCommand, PrimitiveSceneProxy, MeshBatch, StaticMeshId, false);
 
@@ -737,7 +753,7 @@ static FRDGTextureRef AddHairVisibilityFillOpaqueDepth(
 
 ///////////////////////////////////////////////////////////////////////////////////////////////////
 
-struct FVisibilityMSAAPassParameters
+struct FVisibilityPassParameters
 {
 	BEGIN_SHADER_PARAMETER_STRUCT(FParameters, )
 		RENDER_TARGET_BINDING_SLOTS()
@@ -754,7 +770,7 @@ static void AddHairVisibilityCommonPass(
 	TPassParameters* PassParameters)
 {
 	GraphBuilder.AddPass(
-		RDG_EVENT_NAME("HairStrandsVisibilityPass"),
+		RenderMode == HairVisibilityRenderMode_MSAA ? RDG_EVENT_NAME("HairStrandsVisibilityPass") : RDG_EVENT_NAME("HairStrandsCoveragePass"),
 		PassParameters,
 		ERDGPassFlags::Raster,
 		[PassParameters, Scene = Scene, ViewInfo, &ClusterDatas, RenderMode](FRHICommandListImmediate& RHICmdList)
@@ -769,8 +785,10 @@ static void AddHairVisibilityCommonPass(
 			const float VelocityMagnitudeScale = FMath::Clamp(GHairVelocityMagnitudeScale, 0, 512) * FMath::Min(PixelVelocity.X, PixelVelocity.Y);
 
 			// Set the sample count to one as we want the size of the pixel
+			const uint32 HairVisibilitySampleCount = RenderMode == HairVisibilityRenderMode_MSAA ? GetHairVisibilitySampleCount() : 1;
+			const float RasterizationScaleOverride = RenderMode == HairVisibilityRenderMode_MSAA ? 0 : 1.35f;
 			FHairVisibilityPassUniformParameters PassUniformParameters;
-			FMinHairRadiusAtDepth1 MinHairRadius = ComputeMinStrandRadiusAtDepth1(FIntPoint(ViewInfo->UnconstrainedViewRect.Width(), ViewInfo->UnconstrainedViewRect.Height()), ViewInfo->FOV, 1, 0);
+			FMinHairRadiusAtDepth1 MinHairRadius = ComputeMinStrandRadiusAtDepth1(FIntPoint(ViewInfo->UnconstrainedViewRect.Width(), ViewInfo->UnconstrainedViewRect.Height()), ViewInfo->FOV, HairVisibilitySampleCount, RasterizationScaleOverride);
 			PassUniformParameters.MinStrandRadius_Primary = MinHairRadius.Primary;
 			PassUniformParameters.MinStrandRadius_Velocity = MinHairRadius.Velocity;
 			PassUniformParameters.HairStrandsVelocityScale = VelocityMagnitudeScale;
@@ -780,14 +798,20 @@ static void AddHairVisibilityCommonPass(
 		FMeshPassProcessorRenderState DrawRenderState(*ViewInfo, PassUniformBuffer);
 		{
 			RHICmdList.SetViewport(0, 0, 0.0f, Resolution.X, Resolution.Y, 1.0f);
-			DrawRenderState.SetBlendState(TStaticBlendState<
-				CW_RGBA, BO_Add, BF_One, BF_Zero, BO_Add, BF_One, BF_Zero,
-				CW_RGBA, BO_Add, BF_One, BF_Zero, BO_Add, BF_One, BF_Zero>::GetRHI());
 			if (RenderMode == HairVisibilityRenderMode_MSAA)
+			{
+				DrawRenderState.SetBlendState(TStaticBlendState<
+					CW_RGBA, BO_Add, BF_One, BF_Zero, BO_Add, BF_One, BF_Zero,
+					CW_RGBA, BO_Add, BF_One, BF_Zero, BO_Add, BF_One, BF_Zero>::GetRHI());
 				DrawRenderState.SetDepthStencilState(TStaticDepthStencilState<true, CF_DepthNearOrEqual>::GetRHI());
+			}
+			else if (RenderMode == HairVisibilityRenderMode_Coverage)
+			{
+				DrawRenderState.SetBlendState(TStaticBlendState<CW_RED, BO_Add, BF_One, BF_One, BO_Add, BF_One, BF_One>::GetRHI());
+				DrawRenderState.SetDepthStencilState(TStaticDepthStencilState<false, CF_DepthNearOrEqual>::GetRHI());
+			}
 
-			FDynamicMeshDrawCommandStorage DynamicMeshDrawCommandStorage; // << Were would thid be stored?
-			// Use  void DrawDynamicMeshPass(const FSceneView& View, FRHICommandList& RHICmdList, const LambdaType& BuildPassProcessorLambda)
+			FDynamicMeshDrawCommandStorage DynamicMeshDrawCommandStorage;
 			FMeshCommandOneFrameArray VisibleMeshDrawCommands;
 			FGraphicsMinimalPipelineStateSet PipelineStateSet;
 			FDynamicPassMeshDrawListContext ShadowContext(DynamicMeshDrawCommandStorage, VisibleMeshDrawCommands, PipelineStateSet);
@@ -899,7 +923,7 @@ static void AddHairVisibilityMSAAPass(
 		AddClearGraphicPass(GraphBuilder, RDG_EVENT_NAME("HairStrandsClearVisibilityMSAAVelocity"), ViewInfo, 0, VisibilityVelocityTexture);
 	}
 
-	FVisibilityMSAAPassParameters::FParameters* PassParameters = GraphBuilder.AllocParameters<FVisibilityMSAAPassParameters::FParameters>();
+	FVisibilityPassParameters::FParameters* PassParameters = GraphBuilder.AllocParameters<FVisibilityPassParameters::FParameters>();
 	PassParameters->RenderTargets[0] = FRenderTargetBinding(VisibilityIdTexture, ERenderTargetLoadAction::ELoad, 0);
 	PassParameters->RenderTargets[1] = FRenderTargetBinding(VisibilityMaterialTexture,  LoadAction, 0);
 	PassParameters->RenderTargets[2] = FRenderTargetBinding(VisibilityAttributeTexture, LoadAction, 0);
@@ -913,6 +937,45 @@ static void AddHairVisibilityMSAAPass(
 	AddHairVisibilityCommonPass(GraphBuilder, Scene, ViewInfo, ClusterDatas, HairVisibilityRenderMode_MSAA, PassParameters);
 }
 
+///////////////////////////////////////////////////////////////////////////////////////////////////
+
+static FRDGTextureRef AddHairCoveragePass(
+	FRDGBuilder& GraphBuilder,
+	const FScene* Scene,
+	const FViewInfo* ViewInfo,
+	const FHairStrandsClusterDatas& ClusterDatas,
+	const FIntPoint& Resolution,
+	FRDGTextureRef SceneDepthTexture)
+{
+	check(SceneDepthTexture->Desc.Extent == Resolution);
+
+	FRDGTextureRef CoverageTexture = nullptr;
+	{
+		FRDGTextureDesc Desc;
+		Desc.Extent.X = Resolution.X;
+		Desc.Extent.Y = Resolution.Y;
+		Desc.Depth = 0;
+		Desc.Format = PF_R32_FLOAT;
+		Desc.NumMips = 1;
+		Desc.NumSamples = 1;
+		Desc.Flags = TexCreate_None;
+		Desc.TargetableFlags = TexCreate_RenderTargetable | TexCreate_ShaderResource;
+		Desc.bForceSharedTargetAndShaderResource = true;
+		CoverageTexture = GraphBuilder.CreateTexture(Desc, TEXT("HairCoverageTexture"));
+	}
+
+	FVisibilityPassParameters::FParameters* PassParameters = GraphBuilder.AllocParameters<FVisibilityPassParameters::FParameters>();
+	PassParameters->RenderTargets[0] = FRenderTargetBinding(CoverageTexture, ERenderTargetLoadAction::EClear, 0);
+
+	PassParameters->RenderTargets.DepthStencil = FDepthStencilBinding(
+		SceneDepthTexture,
+		ERenderTargetLoadAction::ELoad,
+		ERenderTargetLoadAction::ENoAction,
+		FExclusiveDepthStencil::DepthRead_StencilNop);
+	AddHairVisibilityCommonPass(GraphBuilder, Scene, ViewInfo, ClusterDatas, HairVisibilityRenderMode_Coverage, PassParameters);
+
+	return CoverageTexture;
+}
 
 ///////////////////////////////////////////////////////////////////////////////////////////////////
 class FHairVisibilityDepthPS : public FGlobalShader
@@ -1048,9 +1111,20 @@ FHairStrandsVisibilityViews RenderHairStrandsVisibilityBuffer(
 				FRDGTextureRef VelocityTexture;
 			} MsaaVisibilityResources;
 
+			FRDGTextureRef CoverageTexture = nullptr;
+			if (GHairStrandsCoveragePassEnable > 0)
+			{
+				CoverageTexture = AddHairCoveragePass(
+					GraphBuilder,
+					Scene,
+					&View,
+					ClusterDatas,
+					Resolution, 
+					SceneDepthTexture);
+			}
+
 			if (RenderMode == HairVisibilityRenderMode_MSAA)
 			{
-				
 				MsaaVisibilityResources.DepthTexture = AddHairVisibilityFillOpaqueDepth(
 					GraphBuilder,
 					View,
@@ -1076,6 +1150,10 @@ FHairStrandsVisibilityViews RenderHairStrandsVisibilityBuffer(
 				GraphBuilder.QueueTextureExtraction(MsaaVisibilityResources.AttributeTexture, &VisibilityData.AttributeTexture);
 				GraphBuilder.QueueTextureExtraction(MsaaVisibilityResources.VelocityTexture, &VisibilityData.VelocityTexture);
 				GraphBuilder.QueueTextureExtraction(MsaaVisibilityResources.DepthTexture, &VisibilityData.DepthTexture);
+				if (CoverageTexture)
+				{
+					GraphBuilder.QueueTextureExtraction(CoverageTexture, &VisibilityData.CoverageTexture);
+				}
 				
 				{
 					FRDGTextureRef CompactNodeIndex;
diff --git a/Engine/Source/Runtime/Renderer/Private/HairStrands/HairStrandsVisibility.h b/Engine/Source/Runtime/Renderer/Private/HairStrands/HairStrandsVisibility.h
--- a/Engine/Source/Runtime/Renderer/Private/HairStrands/HairStrandsVisibility.h
+++ b/Engine/Source/Runtime/Renderer/Private/HairStrands/HairStrandsVisibility.h
@@ -18,6 +18,7 @@ struct FHairStrandsVisibilityData
 	TRefCountPtr<IPooledRenderTarget> AttributeTexture;
 	TRefCountPtr<IPooledRenderTarget> VelocityTexture;
 	TRefCountPtr<IPooledRenderTarget> CategorizationTexture;
+	TRefCountPtr<IPooledRenderTarget> CoverageTexture; // temp
 
 	TRefCountPtr<IPooledRenderTarget> NodeIndex;
 	TRefCountPtr<FPooledRDGBuffer>	  NodeData;
