Add debug visualization of the sim strands.

(ushell-p4-cherrypick of 9551437 by Charles.deRousiers)

diff --git a/Engine/Source/Runtime/Renderer/Private/HairStrands/HairStrandsDebug.cpp b/Engine/Source/Runtime/Renderer/Private/HairStrands/HairStrandsDebug.cpp
--- a/Engine/Source/Runtime/Renderer/Private/HairStrands/HairStrandsDebug.cpp
+++ b/Engine/Source/Runtime/Renderer/Private/HairStrands/HairStrandsDebug.cpp
@@ -35,21 +35,33 @@ static int32 GHairStrandsDebugStrandsMode = 0;
 static FAutoConsoleVariableRef CVarDebugPhysicsStrand(TEXT("r.HairStrands.StrandsMode"), GHairStrandsDebugStrandsMode, TEXT("Render debug mode for hair strands. 0:off, 1:simulation strands, 2:render strands with colored simulation strands influence, 3:hair UV, 4:hair root UV, 5: hair seed, 6: dimensions"));
 
 static int32 GHairDebugMeshProjection_SkinCacheMesh = 0;
-static int32 GHairDebugMeshProjection_HairRestTriangles = 0;
-static int32 GHairDebugMeshProjection_HairRestFrames = 0;
-static int32 GHairDebugMeshProjection_HairDeformedTriangles = 0;
-static int32 GHairDebugMeshProjection_HairDeformedFrames = 0;
-static FAutoConsoleVariableRef CVarHairDebugMeshProjection_SkinCacheMesh(TEXT("r.HairStrands.MeshProjection.DebugSkinCache"),						GHairDebugMeshProjection_SkinCacheMesh, TEXT("Render debug mes projection"));
-static FAutoConsoleVariableRef CVarHairDebugMeshProjection_HairRestTriangles(TEXT("r.HairStrands.MeshProjection.DebugHairRestTriangles"),			GHairDebugMeshProjection_HairRestTriangles, TEXT("Render debug mes projection"));
-static FAutoConsoleVariableRef CVarHairDebugMeshProjection_HairRestFrames(TEXT("r.HairStrands.MeshProjection.DebugHairRestFrames"),					GHairDebugMeshProjection_HairRestFrames, TEXT("Render debug mes projection"));
-static FAutoConsoleVariableRef CVarHairDebugMeshProjection_HairDeformedTriangles(TEXT("r.HairStrands.MeshProjection.DebugHairDeformedTriangles"),	GHairDebugMeshProjection_HairDeformedTriangles, TEXT("Render debug mes projection"));
-static FAutoConsoleVariableRef CVarHairDebugMeshProjection_HairDeformedFrames(TEXT("r.HairStrands.MeshProjection.DebugHairDeformedFrames"),			GHairDebugMeshProjection_HairDeformedFrames, TEXT("Render debug mes projection"));
+
+static int32 GHairDebugMeshProjection_Sim_HairRestTriangles = 0;
+static int32 GHairDebugMeshProjection_Sim_HairRestFrames = 0;
+static int32 GHairDebugMeshProjection_Sim_HairDeformedTriangles = 0;
+static int32 GHairDebugMeshProjection_Sim_HairDeformedFrames = 0;
+
+static int32 GHairDebugMeshProjection_Render_HairRestTriangles = 0;
+static int32 GHairDebugMeshProjection_Render_HairRestFrames = 0;
+static int32 GHairDebugMeshProjection_Render_HairDeformedTriangles = 0;
+static int32 GHairDebugMeshProjection_Render_HairDeformedFrames = 0;
+
+static FAutoConsoleVariableRef CVarHairDebugMeshProjection_SkinCacheMesh(					TEXT("r.HairStrands.MeshProjection.DebugSkinCache"),					GHairDebugMeshProjection_SkinCacheMesh, TEXT("Render debug mes projection"));
+static FAutoConsoleVariableRef CVarHairDebugMeshProjection_Render_HairRestTriangles(		TEXT("r.HairStrands.MeshProjection.Render.DebugHairRestTriangles"),		GHairDebugMeshProjection_Render_HairRestTriangles, TEXT("Render debug mes projection"));
+static FAutoConsoleVariableRef CVarHairDebugMeshProjection_Render_HairRestFrames(			TEXT("r.HairStrands.MeshProjection.Render.DebugHairRestFrames"),		GHairDebugMeshProjection_Render_HairRestFrames, TEXT("Render debug mes projection"));
+static FAutoConsoleVariableRef CVarHairDebugMeshProjection_Render_HairDeformedTriangles(	TEXT("r.HairStrands.MeshProjection.Render.DebugHairDeformedTriangles"),	GHairDebugMeshProjection_Render_HairDeformedTriangles, TEXT("Render debug mes projection"));
+static FAutoConsoleVariableRef CVarHairDebugMeshProjection_Render_HairDeformedFrames(		TEXT("r.HairStrands.MeshProjection.Render.DebugHairDeformedFrames"),	GHairDebugMeshProjection_Render_HairDeformedFrames, TEXT("Render debug mes projection"));
+
+static FAutoConsoleVariableRef CVarHairDebugMeshProjection_Sim_HairRestTriangles(			TEXT("r.HairStrands.MeshProjection.Sim.DebugHairRestTriangles"),		GHairDebugMeshProjection_Sim_HairRestTriangles, TEXT("Render debug mes projection"));
+static FAutoConsoleVariableRef CVarHairDebugMeshProjection_Sim_HairRestFrames(				TEXT("r.HairStrands.MeshProjection.Sim.DebugHairRestFrames"),			GHairDebugMeshProjection_Sim_HairRestFrames, TEXT("Render debug mes projection"));
+static FAutoConsoleVariableRef CVarHairDebugMeshProjection_Sim_HairDeformedTriangles(		TEXT("r.HairStrands.MeshProjection.Sim.DebugHairDeformedTriangles"),	GHairDebugMeshProjection_Sim_HairDeformedTriangles, TEXT("Render debug mes projection"));
+static FAutoConsoleVariableRef CVarHairDebugMeshProjection_Sim_HairDeformedFrames(			TEXT("r.HairStrands.MeshProjection.Sim.DebugHairDeformedFrames"),		GHairDebugMeshProjection_Sim_HairDeformedFrames, TEXT("Render debug mes projection"));
 
 
 // Helper functions for accessing interpolation data for debug purpose.
 // Definitions is in HairStrandsInterface.cpp
 void GetGroomInterpolationData(const EWorldType::Type WorldType, FHairStrandsProjectionMeshData& OutGeometries);
-void GetGroomInterpolationData(const EWorldType::Type WorldType, FHairStrandsProjectionHairData& OutHairData, TArray<int32>& OutLODIndices);
+void GetGroomInterpolationData(const EWorldType::Type WorldType, const bool bRenderData, FHairStrandsProjectionHairData& OutHairData, TArray<int32>& OutLODIndices);
 
 enum class EHairDebugMode : uint8
 {
@@ -707,8 +719,8 @@ static void AddDebugProjectionHairPass(
 	const HairStrandsTriangleType PoseType,
 	const int32 LODIndex,
 	const FHairStrandsProjectionHairData::HairGroup& HairData,
-	FRDGTextureRef& ColorTarget,
-	FRDGTextureRef& DepthTexture)
+	FRDGTextureRef ColorTarget,
+	FRDGTextureRef DepthTexture)
 {
 	const EPrimitiveType PrimitiveType = GeometryType == EDebugProjectionHairType::HairFrame ? PT_LineList : PT_TriangleList;
 	const uint32 PrimitiveCount = HairData.RootCount;
@@ -1044,8 +1056,6 @@ void RenderHairStrandsDebugInfo(FRHICommandListImmediate& RHICmdList, TArray<FVi
 		FRDGTextureRef SceneColorTexture = GraphBuilder.RegisterExternalTexture(SceneTargets.GetSceneColor(), TEXT("SceneColorTexture"));
 		if (ViewIndex < uint32(HairDatas->HairVisibilityViews.HairDatas.Num()))
 		{
-			FHairStrandsProjectionMeshData MeshProjectionData;
-			GetGroomInterpolationData(WorldType, MeshProjectionData);
 
 			bool bClearDepth = true;
 			FRDGTextureRef DepthTexture;
@@ -1065,6 +1075,8 @@ void RenderHairStrandsDebugInfo(FRHICommandListImmediate& RHICmdList, TArray<FVi
 
 			if (GHairDebugMeshProjection_SkinCacheMesh > 0)
 			{
+				FHairStrandsProjectionMeshData MeshProjectionData;
+				GetGroomInterpolationData(WorldType, MeshProjectionData);
 				for (FHairStrandsProjectionMeshData::Section& Section : MeshProjectionData.Sections)
 				{
 					AddDebugProjectionMeshPass(GraphBuilder, &View, bClearDepth, Section, SceneColorTexture, DepthTexture);
@@ -1072,35 +1084,65 @@ void RenderHairStrandsDebugInfo(FRHICommandListImmediate& RHICmdList, TArray<FVi
 				}
 			}
 
-			FHairStrandsProjectionHairData HairProjectionDatas;
-			TArray<int32> HairLODIndices;
-			GetGroomInterpolationData(WorldType, HairProjectionDatas, HairLODIndices);
-			check(HairProjectionDatas.HairGroups.Num() == HairLODIndices.Num());
-			for (int32 HairIndex=0; HairIndex < HairProjectionDatas.HairGroups.Num(); ++HairIndex)
+			FViewInfo* LocalView = &View;
+			auto RenderProjectionData = [&GraphBuilder, LocalView, WorldType, &bClearDepth, SceneColorTexture, DepthTexture](bool bRenderData, bool bRestTriangle, bool bRestFrame, bool bDeformedTriangle, bool bDeformedFrame)
 			{
-				const FHairStrandsProjectionHairData::HairGroup& Data = HairProjectionDatas.HairGroups[HairIndex];
-				const int32 LODIndex = HairLODIndices[HairIndex];
-
-				if (GHairDebugMeshProjection_HairRestTriangles > 0)
+				FHairStrandsProjectionHairData HairProjectionDatas;
+				TArray<int32> HairLODIndices;
+				GetGroomInterpolationData(WorldType, bRenderData, HairProjectionDatas, HairLODIndices);
+				check(HairProjectionDatas.HairGroups.Num() == HairLODIndices.Num());
+				for (int32 HairIndex=0; HairIndex < HairProjectionDatas.HairGroups.Num(); ++HairIndex)
 				{
-					AddDebugProjectionHairPass(GraphBuilder, &View, bClearDepth, EDebugProjectionHairType::HairTriangle, HairStrandsTriangleType::RestPose, LODIndex, Data, SceneColorTexture, DepthTexture);
-					bClearDepth = false;
-				}
-				if (GHairDebugMeshProjection_HairRestFrames > 0)
-				{
-					AddDebugProjectionHairPass(GraphBuilder, &View, bClearDepth, EDebugProjectionHairType::HairFrame, HairStrandsTriangleType::RestPose, LODIndex, Data, SceneColorTexture, DepthTexture);
-					bClearDepth = false;
-				}
-				if (GHairDebugMeshProjection_HairDeformedTriangles > 0)
-				{
-					AddDebugProjectionHairPass(GraphBuilder, &View, bClearDepth, EDebugProjectionHairType::HairTriangle, HairStrandsTriangleType::DeformedPose, LODIndex, Data, SceneColorTexture, DepthTexture);
-					bClearDepth = false;
-				}
-				if (GHairDebugMeshProjection_HairDeformedFrames > 0)
-				{
-					AddDebugProjectionHairPass(GraphBuilder, &View, bClearDepth, EDebugProjectionHairType::HairFrame, HairStrandsTriangleType::DeformedPose, LODIndex, Data, SceneColorTexture, DepthTexture);
-					bClearDepth = false;
+					const FHairStrandsProjectionHairData::HairGroup& Data = HairProjectionDatas.HairGroups[HairIndex];
+					const int32 LODIndex = HairLODIndices[HairIndex];
+
+					if (bRestTriangle)
+					{
+						AddDebugProjectionHairPass(GraphBuilder, LocalView, bClearDepth, EDebugProjectionHairType::HairTriangle, HairStrandsTriangleType::RestPose, LODIndex, Data, SceneColorTexture, DepthTexture);
+						bClearDepth = false;
+					}
+					if (bRestFrame)
+					{
+						AddDebugProjectionHairPass(GraphBuilder, LocalView, bClearDepth, EDebugProjectionHairType::HairFrame, HairStrandsTriangleType::RestPose, LODIndex, Data, SceneColorTexture, DepthTexture);
+						bClearDepth = false;
+					}
+					if (bDeformedTriangle)
+					{
+						AddDebugProjectionHairPass(GraphBuilder, LocalView, bClearDepth, EDebugProjectionHairType::HairTriangle, HairStrandsTriangleType::DeformedPose, LODIndex, Data, SceneColorTexture, DepthTexture);
+						bClearDepth = false;
+					}
+					if (bDeformedFrame)
+					{
+						AddDebugProjectionHairPass(GraphBuilder, LocalView, bClearDepth, EDebugProjectionHairType::HairFrame, HairStrandsTriangleType::DeformedPose, LODIndex, Data, SceneColorTexture, DepthTexture);
+						bClearDepth = false;
+					}
 				}
+			};
+
+			if (GHairDebugMeshProjection_Render_HairRestTriangles > 0 || 
+				GHairDebugMeshProjection_Render_HairRestFrames > 0 || 
+				GHairDebugMeshProjection_Render_HairDeformedTriangles > 0 || 
+				GHairDebugMeshProjection_Render_HairDeformedFrames > 0)
+			{
+				RenderProjectionData(
+					true, 
+					GHairDebugMeshProjection_Render_HairRestTriangles > 0, 
+					GHairDebugMeshProjection_Render_HairRestFrames > 0, 
+					GHairDebugMeshProjection_Render_HairDeformedTriangles > 0, 
+					GHairDebugMeshProjection_Render_HairDeformedFrames > 0);
+			}
+
+			if (GHairDebugMeshProjection_Sim_HairRestTriangles > 0 || 
+				GHairDebugMeshProjection_Sim_HairRestFrames > 0 || 
+				GHairDebugMeshProjection_Sim_HairDeformedTriangles > 0 || 
+				GHairDebugMeshProjection_Sim_HairDeformedFrames > 0)
+			{
+				RenderProjectionData(
+					false, 
+					GHairDebugMeshProjection_Sim_HairRestTriangles > 0,
+					GHairDebugMeshProjection_Sim_HairRestFrames > 0, 
+					GHairDebugMeshProjection_Sim_HairDeformedTriangles > 0, 
+					GHairDebugMeshProjection_Sim_HairDeformedFrames > 0);
 			}
 		}
 		GraphBuilder.Execute();
diff --git a/Engine/Source/Runtime/Renderer/Private/HairStrands/HairStrandsInterface.cpp b/Engine/Source/Runtime/Renderer/Private/HairStrands/HairStrandsInterface.cpp
--- a/Engine/Source/Runtime/Renderer/Private/HairStrands/HairStrandsInterface.cpp
+++ b/Engine/Source/Runtime/Renderer/Private/HairStrands/HairStrandsInterface.cpp
@@ -340,7 +340,7 @@ void GetGroomInterpolationData(const EWorldType::Type WorldType, FHairStrandsPro
 	}
 }
 
-void GetGroomInterpolationData(const EWorldType::Type WorldType, FHairStrandsProjectionHairData& Out, TArray<int32>& OutLODIndices)
+void GetGroomInterpolationData(const EWorldType::Type WorldType, const bool bRenderData, FHairStrandsProjectionHairData& Out, TArray<int32>& OutLODIndices)
 {
 	for (FHairStrandsManager::Element& E : GHairManager.Elements)
 	{
@@ -350,10 +350,21 @@ void GetGroomInterpolationData(const EWorldType::Type WorldType, FHairStrandsPro
 		const bool bHasDynamicMesh = E.CachedGeometry.Sections.Num() > 0;
 		if (bHasDynamicMesh)
 		{
-			for (FHairStrandsProjectionHairData::HairGroup& ProjectionHairData : E.RenProjectionHairDatas.HairGroups)
+			if (bRenderData)
+			{
+				for (FHairStrandsProjectionHairData::HairGroup& ProjectionHairData : E.RenProjectionHairDatas.HairGroups)
+				{
+					Out.HairGroups.Add(ProjectionHairData);
+					OutLODIndices.Add(E.FrameLODIndex);
+				}
+			}
+			else
 			{
-				Out.HairGroups.Add(ProjectionHairData);
-				OutLODIndices.Add(E.FrameLODIndex);
+				for (FHairStrandsProjectionHairData::HairGroup& ProjectionHairData : E.SimProjectionHairDatas.HairGroups)
+				{
+					Out.HairGroups.Add(ProjectionHairData);
+					OutLODIndices.Add(E.FrameLODIndex);
+				}
 			}
 		}
 	}
