Add debug view for accurate coverage pass.

(ushell-p4-cherrypick of 9568239 by Charles.deRousiers)

diff --git a/Engine/Shaders/Private/HairStrands/HairStrandsDebug.usf b/Engine/Shaders/Private/HairStrands/HairStrandsDebug.usf
--- a/Engine/Shaders/Private/HairStrands/HairStrandsDebug.usf
+++ b/Engine/Shaders/Private/HairStrands/HairStrandsDebug.usf
@@ -7,12 +7,14 @@ float2 OutputResolution;
 uint FastResolveMask;
 
 Texture2D<uint4> CategorizationTexture;
+Texture2D<float> CoverageTexture;
 Texture2D<uint2> DepthStencilTexture;
 SamplerState LinearSampler;
 
 #define DEBUG_MODE_SAMPLE_COUNT 0 
 #define DEBUG_MODE_COVERAGE_TYPE 1
 #define DEBUG_MODE_TAA_RESOLVE_TYPE 2
+#define DEBUG_MODE_COVERAGE 3
 
 float3 HUEtoRGB(in float H)
 {
@@ -61,5 +63,11 @@ void MainPS(
 	const float3 Color = bIsFastResolve ? float3(1,1,0) : float3(0.15f, 0.25f, 0.75f);
 #endif
 
+#if PERMUTATION_DEBUG_MODE == DEBUG_MODE_COVERAGE
+	const float Coverage = CoverageTexture.Load(uint3(PixelPos, 0));
+	const float3 Color = GetColorCode(saturate(Coverage));
+#endif
+	
+
 	OutColor = float4(Color, 1);
 }
diff --git a/Engine/Shaders/Private/HairStrands/HairStrandsDebugPrint.usf b/Engine/Shaders/Private/HairStrands/HairStrandsDebugPrint.usf
--- a/Engine/Shaders/Private/HairStrands/HairStrandsDebugPrint.usf
+++ b/Engine/Shaders/Private/HairStrands/HairStrandsDebugPrint.usf
@@ -14,9 +14,9 @@ SamplerState LinearSampler;
 Texture2D<uint> HairVisibilityNodeOffsetAndCount;
 StructuredBuffer<FNodeData> HairVisibilityNodeData;
 
-#define DEBUG_MODE_SAMPLE_COUNT 0 
-#define DEBUG_MODE_COVERAGE_TYPE 1
-#define DEBUG_MODE_TAA_RESOLVE_TYPE 2
+#if PERMUTATION_COVERAGE == 1
+Texture2D<float> CoverageTexture;
+#endif 
 
 #define MATERIAL_DATA 0
 #define DETAIL_DATA 0
@@ -105,8 +105,9 @@ void MainCS(uint GroupIndex : SV_GroupIndex, uint3 DispatchThreadId : SV_Dispatc
 	Pos = ShaderPrintNewline(Pos);
 	Pos.x = OriginalPos.x;
 
-	// Coverage
-	#if 0
+	// Accurate coverage
+	#if PERMUTATION_COVERAGE == 1
+	const float PixelCoverage = CoverageTexture.Load(uint3(PixelPos,0));
 	Pos = ShaderPrintSymbol(Pos, _C_);
 	Pos = ShaderPrintSymbol(Pos, _O_);
 	Pos = ShaderPrintSymbol(Pos, _V_);
@@ -116,7 +117,7 @@ void MainCS(uint GroupIndex : SV_GroupIndex, uint3 DispatchThreadId : SV_Dispatc
 	Pos = ShaderPrintSymbol(Pos, _G_);
 	Pos = ShaderPrintSymbol(Pos, _E_);
 	Pos = ShaderPrintSymbol(Pos, _SPC_);
-	Pos = ShaderPrint(Pos, CatData.PixelCoverage);
+	Pos = ShaderPrint(Pos, PixelCoverage);
 
 	Pos = ShaderPrintNewline(Pos);
 	Pos.x = OriginalPos.x;
diff --git a/Engine/Source/Runtime/Renderer/Private/HairStrands/HairStrandsDebug.cpp b/Engine/Source/Runtime/Renderer/Private/HairStrands/HairStrandsDebug.cpp
--- a/Engine/Source/Runtime/Renderer/Private/HairStrands/HairStrandsDebug.cpp
+++ b/Engine/Source/Runtime/Renderer/Private/HairStrands/HairStrandsDebug.cpp
@@ -78,7 +78,8 @@ enum class EHairDebugMode : uint8
 	VoxelsTangent,
 	VoxelsBaseColor,
 	VoxelsRoughness,
-	MeshProjection
+	MeshProjection,
+	Coverage
 };
 
 static EHairDebugMode GetHairDebugMode()
@@ -95,9 +96,10 @@ static EHairDebugMode GetHairDebugMode()
 	case 7:  return EHairDebugMode::CoverageType;
 	case 8:  return EHairDebugMode::VoxelsDensity;
 	case 9:  return EHairDebugMode::VoxelsTangent;
-	case 10:  return EHairDebugMode::VoxelsBaseColor;
+	case 10: return EHairDebugMode::VoxelsBaseColor;
 	case 11: return EHairDebugMode::VoxelsRoughness;
 	case 12: return EHairDebugMode::MeshProjection;
+	case 13: return EHairDebugMode::Coverage;
 	default: return EHairDebugMode::None;
 	};
 }
@@ -119,6 +121,7 @@ static const TCHAR* ToString(EHairDebugMode DebugMode)
 	case EHairDebugMode::VoxelsBaseColor: return TEXT("Hair base color volume");
 	case EHairDebugMode::VoxelsRoughness: return TEXT("Hair roughness volume");
 	case EHairDebugMode::MeshProjection: return TEXT("Hair mesh projection");
+	case EHairDebugMode::Coverage: return TEXT("Hair coverage");
 	default: return TEXT("None");
 	};
 }
@@ -163,10 +166,14 @@ class FHairDebugPrintCS : public FGlobalShader
 	DECLARE_GLOBAL_SHADER(FHairDebugPrintCS);
 	SHADER_USE_PARAMETER_STRUCT(FHairDebugPrintCS, FGlobalShader);
 
+	class FCoverage : SHADER_PERMUTATION_INT("PERMUTATION_COVERAGE", 2);
+	using FPermutationDomain = TShaderPermutationDomain<FCoverage>;
+
 	BEGIN_SHADER_PARAMETER_STRUCT(FParameters, )
 		SHADER_PARAMETER(FIntPoint, PixelCoord)
 		SHADER_PARAMETER(FIntPoint, MaxResolution)
 		SHADER_PARAMETER(uint32, FastResolveMask)
+		SHADER_PARAMETER_RDG_TEXTURE(Texture2D, CoverageTexture)
 		SHADER_PARAMETER_RDG_TEXTURE(Texture2D, CategorizationTexture)
 		SHADER_PARAMETER_RDG_TEXTURE(Texture2D, HairVisibilityNodeOffsetAndCount)
 		SHADER_PARAMETER_RDG_BUFFER_SRV(StructuredBuffer, HairVisibilityNodeData)
@@ -192,7 +199,8 @@ static void AddDebugHairPrintPass(
 	const TRefCountPtr<IPooledRenderTarget>& InCategorizationTexture = VisibilityData.CategorizationTexture;
 	const TRefCountPtr<IPooledRenderTarget>& InNodeIndex = VisibilityData.NodeIndex;
 	const TRefCountPtr<FPooledRDGBuffer>& InNodeData = VisibilityData.NodeData;
-
+	const TRefCountPtr<IPooledRenderTarget>& InCoverageTexture = VisibilityData.CoverageTexture;
+	
 	if (!InCategorizationTexture || !InNodeIndex || !InNodeData || !InDepthStencilTexture) return;
 
 	FRDGTextureRef CategorizationTexture = InCategorizationTexture ? GraphBuilder.RegisterExternalTexture(InCategorizationTexture, TEXT("CategorizationTexture")) : nullptr;
@@ -202,17 +210,22 @@ static void AddDebugHairPrintPass(
 	const FIntRect Viewport = View->ViewRect;
 	const FIntPoint Resolution(Viewport.Width(), Viewport.Height());
 
+	FHairDebugPrintCS::FPermutationDomain PermutationVector;
+	PermutationVector.Set<FHairDebugPrintCS::FCoverage>(InCoverageTexture ? 1 : 0);
+	FRDGTextureRef CoverageTexture = InCoverageTexture ? GraphBuilder.RegisterExternalTexture(InCoverageTexture, TEXT("CoverageTexture")) : nullptr;
+
 	FHairDebugPrintCS::FParameters* Parameters = GraphBuilder.AllocParameters<FHairDebugPrintCS::FParameters>();
 	Parameters->MaxResolution = CategorizationTexture->Desc.Extent;
 	Parameters->PixelCoord = View->CursorPos;
 	Parameters->FastResolveMask = STENCIL_TEMPORAL_RESPONSIVE_AA_MASK;
 	Parameters->CategorizationTexture = CategorizationTexture;
+	Parameters->CoverageTexture = CoverageTexture;
 	Parameters->HairVisibilityNodeData = GraphBuilder.CreateSRV(NodeData);
 	Parameters->HairVisibilityNodeOffsetAndCount = NodeIndex;
 	Parameters->DepthStencilTexture = InDepthStencilTexture;
 	Parameters->LinearSampler = TStaticSamplerState<SF_Bilinear, AM_Clamp, AM_Clamp, AM_Clamp>::GetRHI();
 	ShaderPrint::SetParameters(*View, Parameters->ShaderPrintUniformBuffer);
-	TShaderMapRef<FHairDebugPrintCS> ComputeShader(View->ShaderMap);
+	TShaderMapRef<FHairDebugPrintCS> ComputeShader(View->ShaderMap, PermutationVector);
 
 	ClearUnusedGraphResources(*ComputeShader, Parameters);
 
@@ -230,13 +243,14 @@ class FHairDebugPS : public FGlobalShader
 	DECLARE_GLOBAL_SHADER(FHairDebugPS);
 	SHADER_USE_PARAMETER_STRUCT(FHairDebugPS, FGlobalShader);
 
-	class FDebugMode : SHADER_PERMUTATION_INT("PERMUTATION_DEBUG_MODE", 3);
+	class FDebugMode : SHADER_PERMUTATION_INT("PERMUTATION_DEBUG_MODE", 4);
 	using FPermutationDomain = TShaderPermutationDomain<FDebugMode>;
 
 	BEGIN_SHADER_PARAMETER_STRUCT(FParameters, )
 		SHADER_PARAMETER(FVector2D, OutputResolution)
 		SHADER_PARAMETER(uint32, FastResolveMask)
 		SHADER_PARAMETER_RDG_TEXTURE(Texture2D, CategorizationTexture)
+		SHADER_PARAMETER_RDG_TEXTURE(Texture2D, CoverageTexture)
 		SHADER_PARAMETER_SRV(Texture2D, DepthStencilTexture)
 		SHADER_PARAMETER_SAMPLER(SamplerState, LinearSampler)
 		SHADER_PARAMETER_STRUCT_REF(FViewUniformShaderParameters, ViewUniformBuffer)
@@ -254,16 +268,19 @@ static void AddDebugHairPass(
 	const FViewInfo* View,
 	const EHairDebugMode InDebugMode,
 	const TRefCountPtr<IPooledRenderTarget>& InCategorizationTexture,
+	const TRefCountPtr<IPooledRenderTarget>& InCoverageTexture,
 	const FShaderResourceViewRHIRef& InDepthStencilTexture,
 	FRDGTextureRef& OutTarget)
 {
 	check(OutTarget);
-	check(InDebugMode == EHairDebugMode::TAAResolveType || InDebugMode == EHairDebugMode::SamplePerPixel || InDebugMode == EHairDebugMode::CoverageType);
+	check(InDebugMode == EHairDebugMode::TAAResolveType || InDebugMode == EHairDebugMode::SamplePerPixel || InDebugMode == EHairDebugMode::CoverageType || InDebugMode == EHairDebugMode::Coverage);
 
 	if (!InCategorizationTexture) return;
 	if (InDebugMode == EHairDebugMode::TAAResolveType && !InDepthStencilTexture) return;
+	if (InDebugMode == EHairDebugMode::Coverage && !InCoverageTexture) return;
 
 	FRDGTextureRef CategorizationTexture = InCategorizationTexture ? GraphBuilder.RegisterExternalTexture(InCategorizationTexture, TEXT("CategorizationTexture")) : nullptr;
+	FRDGTextureRef CoverageTexture = InCoverageTexture ? GraphBuilder.RegisterExternalTexture(InCoverageTexture, TEXT("CoverageTexture")) : nullptr;
 
 	const FIntRect Viewport = View->ViewRect;
 	const FIntPoint Resolution(Viewport.Width(), Viewport.Height());
@@ -272,6 +289,7 @@ static void AddDebugHairPass(
 	Parameters->OutputResolution = Resolution;
 	Parameters->FastResolveMask = STENCIL_TEMPORAL_RESPONSIVE_AA_MASK;
 	Parameters->CategorizationTexture = CategorizationTexture;
+	Parameters->CoverageTexture = CoverageTexture;
 	Parameters->DepthStencilTexture = InDepthStencilTexture;
 	Parameters->LinearSampler = TStaticSamplerState<SF_Bilinear, AM_Clamp, AM_Clamp, AM_Clamp>::GetRHI();
 	Parameters->RenderTargets[0] = FRenderTargetBinding(OutTarget, ERenderTargetLoadAction::ELoad, 0);
@@ -284,6 +302,7 @@ static void AddDebugHairPass(
 		case EHairDebugMode::SamplePerPixel:	DebugPermutation = 0; break;
 		case EHairDebugMode::CoverageType:		DebugPermutation = 1; break;
 		case EHairDebugMode::TAAResolveType:	DebugPermutation = 2; break;
+		case EHairDebugMode::Coverage:			DebugPermutation = 3; break;
 	};
 	PermutationVector.Set<FHairDebugPS::FDebugMode>(DebugPermutation);
 	TShaderMapRef<FHairDebugPS> PixelShader(View->ShaderMap, PermutationVector);
@@ -1099,14 +1118,14 @@ void RenderHairStrandsDebugInfo(FRHICommandListImmediate& RHICmdList, TArray<FVi
 		}
 	}
 
-	if (HairDebugMode == EHairDebugMode::TAAResolveType || HairDebugMode == EHairDebugMode::SamplePerPixel || HairDebugMode == EHairDebugMode::CoverageType)
+	if (HairDebugMode == EHairDebugMode::TAAResolveType || HairDebugMode == EHairDebugMode::SamplePerPixel || HairDebugMode == EHairDebugMode::CoverageType || HairDebugMode == EHairDebugMode::Coverage)
 	{
 		FRDGBuilder GraphBuilder(RHICmdList);
 		FRDGTextureRef SceneColorTexture = GraphBuilder.RegisterExternalTexture(SceneTargets.GetSceneColor(), TEXT("SceneColorTexture"));
 		if (ViewIndex < uint32(HairDatas->HairVisibilityViews.HairDatas.Num()))
 		{
 			const FHairStrandsVisibilityData& VisibilityData = HairDatas->HairVisibilityViews.HairDatas[ViewIndex];
-			AddDebugHairPass(GraphBuilder, &View, HairDebugMode, VisibilityData.CategorizationTexture, SceneTargets.SceneStencilSRV, SceneColorTexture);
+			AddDebugHairPass(GraphBuilder, &View, HairDebugMode, VisibilityData.CategorizationTexture, VisibilityData.CoverageTexture, SceneTargets.SceneStencilSRV, SceneColorTexture);
 			AddDebugHairPrintPass(GraphBuilder, &View, HairDebugMode, VisibilityData,  SceneTargets.SceneStencilSRV);
 		}
 
