Add pixel coverage into compacted visibility data.

(ushell-p4-cherrypick of 9568424 by Charles.deRousiers)

diff --git a/Engine/Shaders/Private/HairStrands/HairStrandsDebug.usf b/Engine/Shaders/Private/HairStrands/HairStrandsDebug.usf
--- a/Engine/Shaders/Private/HairStrands/HairStrandsDebug.usf
+++ b/Engine/Shaders/Private/HairStrands/HairStrandsDebug.usf
@@ -64,8 +64,7 @@ void MainPS(
 #endif
 
 #if PERMUTATION_DEBUG_MODE == DEBUG_MODE_COVERAGE
-	const float Coverage = CoverageTexture.Load(uint3(PixelPos, 0));
-	const float3 Color = GetColorCode(saturate(Coverage));
+	const float3 Color = GetColorCode(saturate(CatData.PixelCoverage));
 #endif
 	
 
diff --git a/Engine/Shaders/Private/HairStrands/HairStrandsDebugPrint.usf b/Engine/Shaders/Private/HairStrands/HairStrandsDebugPrint.usf
--- a/Engine/Shaders/Private/HairStrands/HairStrandsDebugPrint.usf
+++ b/Engine/Shaders/Private/HairStrands/HairStrandsDebugPrint.usf
@@ -14,10 +14,6 @@ SamplerState LinearSampler;
 Texture2D<uint> HairVisibilityNodeOffsetAndCount;
 StructuredBuffer<FNodeData> HairVisibilityNodeData;
 
-#if PERMUTATION_COVERAGE == 1
-Texture2D<float> CoverageTexture;
-#endif 
-
 #define MATERIAL_DATA 0
 #define DETAIL_DATA 0
 
@@ -105,9 +101,7 @@ void MainCS(uint GroupIndex : SV_GroupIndex, uint3 DispatchThreadId : SV_Dispatc
 	Pos = ShaderPrintNewline(Pos);
 	Pos.x = OriginalPos.x;
 
-	// Accurate coverage
-	#if PERMUTATION_COVERAGE == 1
-	const float PixelCoverage = CoverageTexture.Load(uint3(PixelPos,0));
+	// Accurate coverage (if coverage pass is activated)
 	Pos = ShaderPrintSymbol(Pos, _C_);
 	Pos = ShaderPrintSymbol(Pos, _O_);
 	Pos = ShaderPrintSymbol(Pos, _V_);
@@ -117,11 +111,10 @@ void MainCS(uint GroupIndex : SV_GroupIndex, uint3 DispatchThreadId : SV_Dispatc
 	Pos = ShaderPrintSymbol(Pos, _G_);
 	Pos = ShaderPrintSymbol(Pos, _E_);
 	Pos = ShaderPrintSymbol(Pos, _SPC_);
-	Pos = ShaderPrint(Pos, PixelCoverage);
+	Pos = ShaderPrint(Pos, CatData.PixelCoverage);
 
 	Pos = ShaderPrintNewline(Pos);
 	Pos.x = OriginalPos.x;
-	#endif
 
 	// 8 bit Coverage (cateogory texture)
 	Pos = ShaderPrintSymbol(Pos, _C_);
@@ -141,6 +134,8 @@ void MainCS(uint GroupIndex : SV_GroupIndex, uint3 DispatchThreadId : SV_Dispatc
 	Pos = ShaderPrintSymbol(Pos, _SPC_);
 	Pos = ShaderPrint(Pos, CatData.TotalCoverage8bit);
 	
+	Pos = ShaderPrintNewline(Pos);
+	Pos.x = OriginalPos.x;
 
 	// 8 bit Coverage (summed)
 	Pos = ShaderPrintSymbol(Pos, _S_);
diff --git a/Engine/Shaders/Private/HairStrands/HairStrandsVisibilityCommon.ush b/Engine/Shaders/Private/HairStrands/HairStrandsVisibilityCommon.ush
--- a/Engine/Shaders/Private/HairStrands/HairStrandsVisibilityCommon.ush
+++ b/Engine/Shaders/Private/HairStrands/HairStrandsVisibilityCommon.ush
@@ -53,10 +53,11 @@ uint4 EncodeVisibilityData(FVisibilityData Data)
 
 	// Tangent and Coverage
 	// Encode into 2 x 16bits
+	// Coverage is clamped o 255 as we only have 8 bits to store its value
 	float3 T = saturate(EncodeTangent(Data.Tangent));
 	Output.xy = uint2(
 		(0xFF & uint(T.x * 255)) | ((0xFF & uint(T.y * 255)) << 8),
-		(0xFF & uint(T.z * 255)) | ((0xFF & Data.Coverage8bit) << 8));
+		(0xFF & uint(T.z * 255)) | (min(0xFF, Data.Coverage8bit) << 8));
 
 	// PrimitiveId and ClusterId
 	// Encode into 2 x 16bits
@@ -255,6 +256,7 @@ FNodeData ConvertToNodeData(float Depth, FVisibilityData VisibilityData, float3
  {
 	 uint  TotalCoverage8bit;
 	 uint  SampleCount;
+	 float PixelCoverage;
 	 float ClosestDepth;
  };
 
@@ -263,15 +265,16 @@ FNodeData ConvertToNodeData(float Depth, FVisibilityData VisibilityData, float3
 	 const uint ClosestDepthUint = asuint(Data.ClosestDepth);
 	 const uint ClosestDepthLow  = ClosestDepthUint & 0xFFFF;
 	 const uint ClosestDepthHigh = (ClosestDepthUint >> 16) & 0xFFFF;
-
-	 return uint4(Data.TotalCoverage8bit, Data.SampleCount, ClosestDepthLow, ClosestDepthHigh);
+	 const uint ComponentX		 = min(Data.TotalCoverage8bit, 0xFF) | ((Data.SampleCount & 0xFF)<<8);
+	 return uint4(ComponentX, f32tof16(Data.PixelCoverage), ClosestDepthLow, ClosestDepthHigh);
  }
 
  FCategorizationData DecodeCategorizationData(uint4 Data)
  {
 	 FCategorizationData Out;
-	 Out.TotalCoverage8bit			= Data.x;
-	 Out.SampleCount				= Data.y;
+	 Out.TotalCoverage8bit			= Data.x & 0xFF;
+	 Out.SampleCount				=(Data.x>>8) & 0xFF;
+	 Out.PixelCoverage				= f16tof32(Data.y);
 	 const uint ClosestDepthLow		= Data.z & 0xFFFF;
 	 const uint ClosestDepthHigh	=(Data.w & 0xFFFF) << 16;
 	 Out.ClosestDepth				= asfloat(ClosestDepthHigh | ClosestDepthLow);
diff --git a/Engine/Shaders/Private/HairStrands/HairStrandsVisibilityCompaction.usf b/Engine/Shaders/Private/HairStrands/HairStrandsVisibilityCompaction.usf
--- a/Engine/Shaders/Private/HairStrands/HairStrandsVisibilityCompaction.usf
+++ b/Engine/Shaders/Private/HairStrands/HairStrandsVisibilityCompaction.usf
@@ -86,6 +86,10 @@ Texture2DMS<float, HAIR_VISIBILITY_MAXSAMPLECOUNT> MSAA_DepthTexture;
 Texture2DMS<float4,HAIR_VISIBILITY_MAXSAMPLECOUNT> MSAA_MaterialTexture;
 Texture2DMS<float4,HAIR_VISIBILITY_MAXSAMPLECOUNT> MSAA_AttributeTexture;
 
+#if PERMUTATION_COVERAGE == 1
+Texture2D<float> CoverageTexture;
+#endif
+
 RWTexture2D<uint> OutCompactNodeCounter;
 RWTexture2D<uint> OutCompactNodeIndex;
 RWStructuredBuffer<FNodeData> OutCompactNodeData;
@@ -227,6 +231,12 @@ void MainCS(uint GroupIndex : SV_GroupIndex, uint3 DispatchThreadId : SV_Dispatc
 				OutCompactNodeCoord[StoreIndex] = (PixelCoord.x & 0xFFFF) | ((PixelCoord.y & 0xFFFF) << 16);
 			}
 
+			#if PERMUTATION_COVERAGE == 1
+			Categorization.PixelCoverage = CoverageTexture.Load(uint3(PixelCoord,0));
+			#else
+			Categorization.PixelCoverage = SampleDesc.ValidSampleCount / float(HAIR_VISIBILITY_MAXSAMPLECOUNT);
+			#endif
+
 			#if PERMUTATION_VELOCITY == VELOCITY_TYPE_AVG
 			OutVelocityTexture[PixelCoord] = EncodeVelocityToTexture(AverageVelocity / max(ValidVelocityCount, 1));
 			#endif		
diff --git a/Engine/Source/Runtime/Renderer/Private/HairStrands/HairStrandsDebug.cpp b/Engine/Source/Runtime/Renderer/Private/HairStrands/HairStrandsDebug.cpp
--- a/Engine/Source/Runtime/Renderer/Private/HairStrands/HairStrandsDebug.cpp
+++ b/Engine/Source/Runtime/Renderer/Private/HairStrands/HairStrandsDebug.cpp
@@ -166,14 +166,10 @@ class FHairDebugPrintCS : public FGlobalShader
 	DECLARE_GLOBAL_SHADER(FHairDebugPrintCS);
 	SHADER_USE_PARAMETER_STRUCT(FHairDebugPrintCS, FGlobalShader);
 
-	class FCoverage : SHADER_PERMUTATION_INT("PERMUTATION_COVERAGE", 2);
-	using FPermutationDomain = TShaderPermutationDomain<FCoverage>;
-
 	BEGIN_SHADER_PARAMETER_STRUCT(FParameters, )
 		SHADER_PARAMETER(FIntPoint, PixelCoord)
 		SHADER_PARAMETER(FIntPoint, MaxResolution)
 		SHADER_PARAMETER(uint32, FastResolveMask)
-		SHADER_PARAMETER_RDG_TEXTURE(Texture2D, CoverageTexture)
 		SHADER_PARAMETER_RDG_TEXTURE(Texture2D, CategorizationTexture)
 		SHADER_PARAMETER_RDG_TEXTURE(Texture2D, HairVisibilityNodeOffsetAndCount)
 		SHADER_PARAMETER_RDG_BUFFER_SRV(StructuredBuffer, HairVisibilityNodeData)
@@ -199,7 +195,6 @@ static void AddDebugHairPrintPass(
 	const TRefCountPtr<IPooledRenderTarget>& InCategorizationTexture = VisibilityData.CategorizationTexture;
 	const TRefCountPtr<IPooledRenderTarget>& InNodeIndex = VisibilityData.NodeIndex;
 	const TRefCountPtr<FPooledRDGBuffer>& InNodeData = VisibilityData.NodeData;
-	const TRefCountPtr<IPooledRenderTarget>& InCoverageTexture = VisibilityData.CoverageTexture;
 	
 	if (!InCategorizationTexture || !InNodeIndex || !InNodeData || !InDepthStencilTexture) return;
 
@@ -210,22 +205,17 @@ static void AddDebugHairPrintPass(
 	const FIntRect Viewport = View->ViewRect;
 	const FIntPoint Resolution(Viewport.Width(), Viewport.Height());
 
-	FHairDebugPrintCS::FPermutationDomain PermutationVector;
-	PermutationVector.Set<FHairDebugPrintCS::FCoverage>(InCoverageTexture ? 1 : 0);
-	FRDGTextureRef CoverageTexture = InCoverageTexture ? GraphBuilder.RegisterExternalTexture(InCoverageTexture, TEXT("CoverageTexture")) : nullptr;
-
 	FHairDebugPrintCS::FParameters* Parameters = GraphBuilder.AllocParameters<FHairDebugPrintCS::FParameters>();
 	Parameters->MaxResolution = CategorizationTexture->Desc.Extent;
 	Parameters->PixelCoord = View->CursorPos;
 	Parameters->FastResolveMask = STENCIL_TEMPORAL_RESPONSIVE_AA_MASK;
 	Parameters->CategorizationTexture = CategorizationTexture;
-	Parameters->CoverageTexture = CoverageTexture;
 	Parameters->HairVisibilityNodeData = GraphBuilder.CreateSRV(NodeData);
 	Parameters->HairVisibilityNodeOffsetAndCount = NodeIndex;
 	Parameters->DepthStencilTexture = InDepthStencilTexture;
 	Parameters->LinearSampler = TStaticSamplerState<SF_Bilinear, AM_Clamp, AM_Clamp, AM_Clamp>::GetRHI();
 	ShaderPrint::SetParameters(*View, Parameters->ShaderPrintUniformBuffer);
-	TShaderMapRef<FHairDebugPrintCS> ComputeShader(View->ShaderMap, PermutationVector);
+	TShaderMapRef<FHairDebugPrintCS> ComputeShader(View->ShaderMap);
 
 	ClearUnusedGraphResources(*ComputeShader, Parameters);
 
@@ -250,7 +240,6 @@ class FHairDebugPS : public FGlobalShader
 		SHADER_PARAMETER(FVector2D, OutputResolution)
 		SHADER_PARAMETER(uint32, FastResolveMask)
 		SHADER_PARAMETER_RDG_TEXTURE(Texture2D, CategorizationTexture)
-		SHADER_PARAMETER_RDG_TEXTURE(Texture2D, CoverageTexture)
 		SHADER_PARAMETER_SRV(Texture2D, DepthStencilTexture)
 		SHADER_PARAMETER_SAMPLER(SamplerState, LinearSampler)
 		SHADER_PARAMETER_STRUCT_REF(FViewUniformShaderParameters, ViewUniformBuffer)
@@ -268,7 +257,6 @@ static void AddDebugHairPass(
 	const FViewInfo* View,
 	const EHairDebugMode InDebugMode,
 	const TRefCountPtr<IPooledRenderTarget>& InCategorizationTexture,
-	const TRefCountPtr<IPooledRenderTarget>& InCoverageTexture,
 	const FShaderResourceViewRHIRef& InDepthStencilTexture,
 	FRDGTextureRef& OutTarget)
 {
@@ -277,10 +265,8 @@ static void AddDebugHairPass(
 
 	if (!InCategorizationTexture) return;
 	if (InDebugMode == EHairDebugMode::TAAResolveType && !InDepthStencilTexture) return;
-	if (InDebugMode == EHairDebugMode::Coverage && !InCoverageTexture) return;
 
 	FRDGTextureRef CategorizationTexture = InCategorizationTexture ? GraphBuilder.RegisterExternalTexture(InCategorizationTexture, TEXT("CategorizationTexture")) : nullptr;
-	FRDGTextureRef CoverageTexture = InCoverageTexture ? GraphBuilder.RegisterExternalTexture(InCoverageTexture, TEXT("CoverageTexture")) : nullptr;
 
 	const FIntRect Viewport = View->ViewRect;
 	const FIntPoint Resolution(Viewport.Width(), Viewport.Height());
@@ -289,7 +275,6 @@ static void AddDebugHairPass(
 	Parameters->OutputResolution = Resolution;
 	Parameters->FastResolveMask = STENCIL_TEMPORAL_RESPONSIVE_AA_MASK;
 	Parameters->CategorizationTexture = CategorizationTexture;
-	Parameters->CoverageTexture = CoverageTexture;
 	Parameters->DepthStencilTexture = InDepthStencilTexture;
 	Parameters->LinearSampler = TStaticSamplerState<SF_Bilinear, AM_Clamp, AM_Clamp, AM_Clamp>::GetRHI();
 	Parameters->RenderTargets[0] = FRenderTargetBinding(OutTarget, ERenderTargetLoadAction::ELoad, 0);
@@ -1125,7 +1110,7 @@ void RenderHairStrandsDebugInfo(FRHICommandListImmediate& RHICmdList, TArray<FVi
 		if (ViewIndex < uint32(HairDatas->HairVisibilityViews.HairDatas.Num()))
 		{
 			const FHairStrandsVisibilityData& VisibilityData = HairDatas->HairVisibilityViews.HairDatas[ViewIndex];
-			AddDebugHairPass(GraphBuilder, &View, HairDebugMode, VisibilityData.CategorizationTexture, VisibilityData.CoverageTexture, SceneTargets.SceneStencilSRV, SceneColorTexture);
+			AddDebugHairPass(GraphBuilder, &View, HairDebugMode, VisibilityData.CategorizationTexture, SceneTargets.SceneStencilSRV, SceneColorTexture);
 			AddDebugHairPrintPass(GraphBuilder, &View, HairDebugMode, VisibilityData,  SceneTargets.SceneStencilSRV);
 		}
 
diff --git a/Engine/Source/Runtime/Renderer/Private/HairStrands/HairStrandsVisibility.cpp b/Engine/Source/Runtime/Renderer/Private/HairStrands/HairStrandsVisibility.cpp
--- a/Engine/Source/Runtime/Renderer/Private/HairStrands/HairStrandsVisibility.cpp
+++ b/Engine/Source/Runtime/Renderer/Private/HairStrands/HairStrandsVisibility.cpp
@@ -470,7 +470,8 @@ class FHairVisibilityPrimitiveIdCompactionCS : public FGlobalShader
 
 	class FVendor   : SHADER_PERMUTATION_INT("PERMUTATION_VENDOR", HairVisibilityVendorCount);
 	class FVelocity : SHADER_PERMUTATION_INT("PERMUTATION_VELOCITY", 4);
-	using FPermutationDomain = TShaderPermutationDomain<FVendor, FVelocity>;
+	class FCoverage : SHADER_PERMUTATION_INT("PERMUTATION_COVERAGE", 2);
+	using FPermutationDomain = TShaderPermutationDomain<FVendor, FVelocity, FCoverage>;
 
 	BEGIN_SHADER_PARAMETER_STRUCT(FParameters, )
 		SHADER_PARAMETER(FIntPoint, OutputResolution)
@@ -483,6 +484,7 @@ class FHairVisibilityPrimitiveIdCompactionCS : public FGlobalShader
 		SHADER_PARAMETER_RDG_TEXTURE(Texture2D, MSAA_MaterialTexture)
 		SHADER_PARAMETER_RDG_TEXTURE(Texture2D, MSAA_AttributeTexture)
 		SHADER_PARAMETER_RDG_TEXTURE(Texture2D, MSAA_VelocityTexture)
+		SHADER_PARAMETER_RDG_TEXTURE(Texture2D, CoverageTexture)
 		
 		SHADER_PARAMETER_RDG_TEXTURE_UAV(Texture2D, OutCompactNodeCounter)
 		SHADER_PARAMETER_RDG_TEXTURE_UAV(Texture2D, OutCompactNodeIndex)
@@ -510,6 +512,7 @@ static void AddHairVisibilityPrimitiveIdCompactionPass(
 	const FRDGTextureRef& MSAA_MaterialTexture,
 	const FRDGTextureRef& MSAA_AttributeTexture,
 	const FRDGTextureRef& MSAA_VelocityTexture,
+	const FRDGTextureRef& CoverageTexture,
 	FRDGTextureRef& OutCompactNodeIndex,
 	FRDGBufferRef& OutCompactNodeData,
 	FRDGBufferRef& OutCompactNodeCoord,
@@ -594,12 +597,14 @@ static void AddHairVisibilityPrimitiveIdCompactionPass(
 	FHairVisibilityPrimitiveIdCompactionCS::FPermutationDomain PermutationVector;
 	PermutationVector.Set<FHairVisibilityPrimitiveIdCompactionCS::FVendor>(GetVendor());
 	PermutationVector.Set<FHairVisibilityPrimitiveIdCompactionCS::FVelocity>(bWriteOutVelocity ? FMath::Clamp(GHairVelocityType+1, 0, 3) : 0);
+	PermutationVector.Set<FHairVisibilityPrimitiveIdCompactionCS::FCoverage>(CoverageTexture ? 1 : 0);
 
 	FHairVisibilityPrimitiveIdCompactionCS::FParameters* Parameters = GraphBuilder.AllocParameters<FHairVisibilityPrimitiveIdCompactionCS::FParameters>();
 	Parameters->MSAA_DepthTexture = MSAA_DepthTexture;
 	Parameters->MSAA_IDTexture = MSAA_IDTexture;
 	Parameters->MSAA_MaterialTexture = MSAA_MaterialTexture;
 	Parameters->MSAA_AttributeTexture = MSAA_AttributeTexture;
+	Parameters->CoverageTexture = CoverageTexture;
 	Parameters->OutputResolution = Resolution;
 	Parameters->MaxNodeCount = MaxNodeCount;
 	Parameters->HairVisibilitySampleCount = HairVisibilitySampleCount;
@@ -1150,10 +1155,6 @@ FHairStrandsVisibilityViews RenderHairStrandsVisibilityBuffer(
 				GraphBuilder.QueueTextureExtraction(MsaaVisibilityResources.AttributeTexture, &VisibilityData.AttributeTexture);
 				GraphBuilder.QueueTextureExtraction(MsaaVisibilityResources.VelocityTexture, &VisibilityData.VelocityTexture);
 				GraphBuilder.QueueTextureExtraction(MsaaVisibilityResources.DepthTexture, &VisibilityData.DepthTexture);
-				if (CoverageTexture)
-				{
-					GraphBuilder.QueueTextureExtraction(CoverageTexture, &VisibilityData.CoverageTexture);
-				}
 				
 				{
 					FRDGTextureRef CompactNodeIndex;
@@ -1172,6 +1173,7 @@ FHairStrandsVisibilityViews RenderHairStrandsVisibilityBuffer(
 						MsaaVisibilityResources.MaterialTexture,
 						MsaaVisibilityResources.AttributeTexture,
 						MsaaVisibilityResources.VelocityTexture,
+						CoverageTexture,
 						CompactNodeIndex,
 						CompactNodeData,
 						CompactNodeCoord,
diff --git a/Engine/Source/Runtime/Renderer/Private/HairStrands/HairStrandsVisibility.h b/Engine/Source/Runtime/Renderer/Private/HairStrands/HairStrandsVisibility.h
--- a/Engine/Source/Runtime/Renderer/Private/HairStrands/HairStrandsVisibility.h
+++ b/Engine/Source/Runtime/Renderer/Private/HairStrands/HairStrandsVisibility.h
@@ -18,7 +18,6 @@ struct FHairStrandsVisibilityData
 	TRefCountPtr<IPooledRenderTarget> AttributeTexture;
 	TRefCountPtr<IPooledRenderTarget> VelocityTexture;
 	TRefCountPtr<IPooledRenderTarget> CategorizationTexture;
-	TRefCountPtr<IPooledRenderTarget> CoverageTexture; // temp
 
 	TRefCountPtr<IPooledRenderTarget> NodeIndex;
 	TRefCountPtr<FPooledRDGBuffer>	  NodeData;
