Add extra information about CS thread id/group in the hair pixel inspector.

This is useful when debugging hair compute shaders.

(ushell-p4-cherrypick of 9583100 by Charles.deRousiers)

diff --git a/Engine/Shaders/Private/HairStrands/HairStrandsDebugPrint.usf b/Engine/Shaders/Private/HairStrands/HairStrandsDebugPrint.usf
--- a/Engine/Shaders/Private/HairStrands/HairStrandsDebugPrint.usf
+++ b/Engine/Shaders/Private/HairStrands/HairStrandsDebugPrint.usf
@@ -7,6 +7,7 @@
 int2 PixelCoord;
 int2 MaxResolution;
 uint FastResolveMask;
+uint2 GroupSize; 
 
 Texture2D<uint4> CategorizationTexture;
 Texture2D<uint2> DepthStencilTexture;
@@ -70,6 +71,35 @@ void MainCS(uint GroupIndex : SV_GroupIndex, uint3 DispatchThreadId : SV_Dispatc
 		Total8BitCoverage += SampleVisData.Coverage8bit;
 	}
 
+	// Pixel coord
+	Pos = ShaderPrintSymbol(Pos, _P_);
+	Pos = ShaderPrintSymbol(Pos, _I_);
+	Pos = ShaderPrintSymbol(Pos, _X_);
+	Pos = ShaderPrintSymbol(Pos, _E_);
+	Pos = ShaderPrintSymbol(Pos, _L_);
+	Pos = ShaderPrintSymbol(Pos, _SPC_);
+	Pos = ShaderPrint(Pos, PixelCoord);
+
+	Pos = ShaderPrintNewline(Pos);
+	Pos.x = OriginalPos.x;
+
+#if MATERIAL_DATA
+	// Dispatch thread id
+	Pos = ShaderPrintSymbol(Pos, _T_);
+	Pos = ShaderPrintSymbol(Pos, _H_);
+	Pos = ShaderPrintSymbol(Pos, _R_);
+	Pos = ShaderPrintSymbol(Pos, _E_);
+	Pos = ShaderPrintSymbol(Pos, _A_);
+	Pos = ShaderPrintSymbol(Pos, _D_);
+	Pos = ShaderPrintSymbol(Pos, _SPC_);
+	Pos = ShaderPrintSymbol(Pos, _I_);
+	Pos = ShaderPrintSymbol(Pos, _D_);
+	Pos = ShaderPrintSymbol(Pos, _SPC_);
+	Pos = ShaderPrint(Pos, uint4(PixelCoord.x/GroupSize.x, PixelCoord.x % GroupSize.x, PixelCoord.y / GroupSize.y, PixelCoord.y % GroupSize.y));
+
+	Pos = ShaderPrintNewline(Pos);
+	Pos.x = OriginalPos.x;
+#endif
 
 	// Max sample
 	Pos = ShaderPrintSymbol(Pos, _M_);
diff --git a/Engine/Source/Runtime/Renderer/Private/HairStrands/HairStrandsDebug.cpp b/Engine/Source/Runtime/Renderer/Private/HairStrands/HairStrandsDebug.cpp
--- a/Engine/Source/Runtime/Renderer/Private/HairStrands/HairStrandsDebug.cpp
+++ b/Engine/Source/Runtime/Renderer/Private/HairStrands/HairStrandsDebug.cpp
@@ -167,6 +167,7 @@ class FHairDebugPrintCS : public FGlobalShader
 	SHADER_USE_PARAMETER_STRUCT(FHairDebugPrintCS, FGlobalShader);
 
 	BEGIN_SHADER_PARAMETER_STRUCT(FParameters, )
+		SHADER_PARAMETER(FIntPoint, GroupSize)
 		SHADER_PARAMETER(FIntPoint, PixelCoord)
 		SHADER_PARAMETER(FIntPoint, MaxResolution)
 		SHADER_PARAMETER(uint32, FastResolveMask)
@@ -206,6 +207,7 @@ static void AddDebugHairPrintPass(
 	const FIntPoint Resolution(Viewport.Width(), Viewport.Height());
 
 	FHairDebugPrintCS::FParameters* Parameters = GraphBuilder.AllocParameters<FHairDebugPrintCS::FParameters>();
+	Parameters->GroupSize = GetVendorOptimalGroupSize2D();
 	Parameters->MaxResolution = CategorizationTexture->Desc.Extent;
 	Parameters->PixelCoord = View->CursorPos;
 	Parameters->FastResolveMask = STENCIL_TEMPORAL_RESPONSIVE_AA_MASK;
diff --git a/Engine/Source/Runtime/Renderer/Private/HairStrands/HairStrandsUtils.cpp b/Engine/Source/Runtime/Renderer/Private/HairStrands/HairStrandsUtils.cpp
--- a/Engine/Source/Runtime/Renderer/Private/HairStrands/HairStrandsUtils.cpp
+++ b/Engine/Source/Runtime/Renderer/Private/HairStrands/HairStrandsUtils.cpp
@@ -195,4 +195,31 @@ bool IsHairStrandsViewRectOptimEnable()
 bool IsHairStrandsSupported(const EShaderPlatform Platform)
 {
 	return IsD3DPlatform(Platform, false) && IsPCPlatform(Platform) && GetMaxSupportedFeatureLevel(Platform) == ERHIFeatureLevel::SM5;
+}
+
+EHairVisibilityVendor GetVendor()
+{
+	return IsRHIDeviceAMD() ? HairVisibilityVendor_AMD : (IsRHIDeviceNVIDIA() ? HairVisibilityVendor_NVIDIA : HairVisibilityVendor_INTEL);
+}
+
+uint32 GetVendorOptimalGroupSize1D()
+{
+	switch (GetVendor())
+	{
+	case HairVisibilityVendor_AMD:		return 64;
+	case HairVisibilityVendor_NVIDIA:	return 32;
+	case HairVisibilityVendor_INTEL:	return 64;
+	default:							return 64;
+	}
+}
+
+FIntPoint GetVendorOptimalGroupSize2D()
+{
+	switch (GetVendor())
+	{
+	case HairVisibilityVendor_AMD:		return FIntPoint(8, 8);
+	case HairVisibilityVendor_NVIDIA:	return FIntPoint(8, 4);
+	case HairVisibilityVendor_INTEL:	return FIntPoint(8, 8);
+	default:							return FIntPoint(8, 8);
+	}
 }
\ No newline at end of file
diff --git a/Engine/Source/Runtime/Renderer/Private/HairStrands/HairStrandsUtils.h b/Engine/Source/Runtime/Renderer/Private/HairStrands/HairStrandsUtils.h
--- a/Engine/Source/Runtime/Renderer/Private/HairStrands/HairStrandsUtils.h
+++ b/Engine/Source/Runtime/Renderer/Private/HairStrands/HairStrandsUtils.h
@@ -38,4 +38,17 @@ FIntRect ComputeVisibleHairStrandsClustersRect(const FIntRect& ViewRect, const s
 
 bool IsHairStrandsViewRectOptimEnable();
 
-RENDERER_API bool IsHairStrandsSupported(const EShaderPlatform Platform);
\ No newline at end of file
+enum EHairVisibilityVendor
+{
+	HairVisibilityVendor_AMD,
+	HairVisibilityVendor_NVIDIA,
+	HairVisibilityVendor_INTEL,
+	HairVisibilityVendorCount
+};
+
+EHairVisibilityVendor GetVendor();
+uint32 GetVendorOptimalGroupSize1D();
+FIntPoint GetVendorOptimalGroupSize2D();
+
+RENDERER_API bool IsHairStrandsSupported(const EShaderPlatform Platform);
+
diff --git a/Engine/Source/Runtime/Renderer/Private/HairStrands/HairStrandsVisibility.cpp b/Engine/Source/Runtime/Renderer/Private/HairStrands/HairStrandsVisibility.cpp
--- a/Engine/Source/Runtime/Renderer/Private/HairStrands/HairStrandsVisibility.cpp
+++ b/Engine/Source/Runtime/Renderer/Private/HairStrands/HairStrandsVisibility.cpp
@@ -52,14 +52,6 @@ enum EHairVisibilityRenderMode
 	HairVisibilityRenderModeCount
 };
 
-enum EHairVisibilityVendor
-{
-	HairVisibilityVendor_AMD,
-	HairVisibilityVendor_NVIDIA,
-	HairVisibilityVendor_INTEL,
-	HairVisibilityVendorCount
-};
-
 EHairVisibilityRenderMode GetHairVisibilityRenderMode()
 {
 	return HairVisibilityRenderMode_MSAA;
@@ -75,32 +67,6 @@ static bool IsCompatibleWithHairVisibility(const FMeshMaterialShaderPermutationP
 	return IsCompatibleWithHairStrands(Parameters.Material, Parameters.Platform);
 }
 
-static EHairVisibilityVendor GetVendor()
-{
-	return IsRHIDeviceAMD() ? HairVisibilityVendor_AMD : (IsRHIDeviceNVIDIA() ? HairVisibilityVendor_NVIDIA : HairVisibilityVendor_INTEL);
-}
-
-static uint32 GetVendorOptimalGroupSize1D()
-{
-	switch (GetVendor())
-	{
-	case HairVisibilityVendor_AMD:		return 64;
-	case HairVisibilityVendor_NVIDIA:	return 32;
-	case HairVisibilityVendor_INTEL:	return 64;
-	default:							return 64;
-	}
-}
-
-static FIntPoint GetVendorOptimalGroupSize2D()
-{
-	switch (GetVendor())
-	{
-	case HairVisibilityVendor_AMD:		return FIntPoint(8, 8);
-	case HairVisibilityVendor_NVIDIA:	return FIntPoint(8, 4);
-	case HairVisibilityVendor_INTEL:	return FIntPoint(8, 8);
-	default:							return FIntPoint(8, 8);
-	}
-}
 
 template<EHairVisibilityRenderMode RenderMode>
 class FHairVisibilityVS : public FMeshMaterialShader
