Add material compaction pass.

This add a permutation to the compaction pass. This process the unique sub-pixel sample, and merge them if they are close enough and if their tangent is similar. Threshold value are exposed through CVAR.

(ushell-p4-cherrypick of 9602313 by Charles.deRousiers)

diff --git a/Engine/Shaders/Private/HairStrands/HairStrandsVisibilityCompaction.usf b/Engine/Shaders/Private/HairStrands/HairStrandsVisibilityCompaction.usf
--- a/Engine/Shaders/Private/HairStrands/HairStrandsVisibilityCompaction.usf
+++ b/Engine/Shaders/Private/HairStrands/HairStrandsVisibilityCompaction.usf
@@ -10,6 +10,11 @@
 #define VENDOR_NVIDIA 1
 #define VENDOR_OTHER 2
 
+#if PERMUTATION_MATERIAL_COMPACTION == 1
+float DepthTheshold;
+float CosTangentThreshold;
+#endif
+
 void ComputeUniqueSamplesWithoutDepthTest(
 	out uint3 OutSamples[HAIR_VISIBILITY_MAXSAMPLECOUNT],
 	out FSampleSetDesc OutSet,
@@ -24,6 +29,10 @@ void ComputeUniqueSamplesWithoutDepthTest(
 
 	OutSet.HairSampleCount = clamp(SampleCount, 0, HAIR_VISIBILITY_MAXSAMPLECOUNT);
 	const uint InvalidSample = 0xFFFF;
+
+#if PERMUTATION_MATERIAL_COMPACTION == 1
+	float4 UniqueTangentAndDepths[HAIR_VISIBILITY_MAXSAMPLECOUNT];
+#endif
 	for (uint SampleIt = 0; SampleIt < OutSet.HairSampleCount; ++SampleIt)
 	{
 		const uint4 EncodedData = InMSAA_IDTexture.Load(PixelCoord, SampleIt);
@@ -37,11 +46,25 @@ void ComputeUniqueSamplesWithoutDepthTest(
 
 		++OutSet.ValidSampleCount;
 
+	#if PERMUTATION_MATERIAL_COMPACTION == 1
+		const float SceneDepth = ConvertFromDeviceZ(SampleDepth);
+	#endif
+
 		const FVisibilityData VisibilityData = DecodeVisibilityData(EncodedData);
 		bool bAlreadyExist = false;
 		for (uint UniqueIt = 0; UniqueIt < OutSet.UniqueSampleCount; ++UniqueIt)
 		{
-			if (VisibilityData.PrimitiveId == OutSamples[UniqueIt].x)
+		#if PERMUTATION_MATERIAL_COMPACTION == 1
+			const float SceneDepth = ConvertFromDeviceZ(SampleDepth);
+			const bool bIsSimilar  = 
+				VisibilityData.PrimitiveId == OutSamples[UniqueIt].x ||
+				(abs(dot(UniqueTangentAndDepths[UniqueIt].xyz, VisibilityData.Tangent)) > CosTangentThreshold &&
+					abs(UniqueTangentAndDepths[UniqueIt].w - SceneDepth) < DepthTheshold);
+		#else
+			const bool bIsSimilar = VisibilityData.PrimitiveId == OutSamples[UniqueIt].x;
+		#endif
+
+			if (bIsSimilar)
 			{
 				OutSamples[UniqueIt].y += 1;
 				bAlreadyExist = true;
@@ -51,6 +74,10 @@ void ComputeUniqueSamplesWithoutDepthTest(
 
 		if (!bAlreadyExist)
 		{
+		#if PERMUTATION_MATERIAL_COMPACTION == 1
+			UniqueTangentAndDepths[OutSet.UniqueSampleCount] = float4(VisibilityData.Tangent, SceneDepth);
+		#endif
+
 			OutSamples[OutSet.UniqueSampleCount].x = VisibilityData.PrimitiveId;
 			OutSamples[OutSet.UniqueSampleCount].y = 1;
 			OutSamples[OutSet.UniqueSampleCount].z = SampleIt;
@@ -191,7 +218,6 @@ void MainCS(uint GroupIndex : SV_GroupIndex, uint3 DispatchThreadId : SV_Dispatc
 					BaseColorAndRoughness.xyz = BaseColorAndRoughness.xyz * BaseColorAndRoughness.xyz; // Cheap gamma encoding since this is encoded into 8bit target
 					const float Specular = MSAA_AttributeTexture.Load(PixelCoord, SampleIt).x;
 
-
 					OutNodeData = ConvertToNodeData(
 						SampleDepth,
 						VisibilityData,
diff --git a/Engine/Source/Runtime/Renderer/Private/HairStrands/HairStrandsVisibility.cpp b/Engine/Source/Runtime/Renderer/Private/HairStrands/HairStrandsVisibility.cpp
--- a/Engine/Source/Runtime/Renderer/Private/HairStrands/HairStrandsVisibility.cpp
+++ b/Engine/Source/Runtime/Renderer/Private/HairStrands/HairStrandsVisibility.cpp
@@ -31,6 +31,15 @@ IMPLEMENT_GLOBAL_SHADER_PARAMETER_STRUCT(FHairVisibilityPassUniformParameters, "
 static int32 GHairStrandsCoveragePassEnable = 0;
 static FAutoConsoleVariableRef CVarHairStrandsCoveragePassEnable(TEXT("r.HairStrands.CoveragePass"), GHairStrandsCoveragePassEnable, TEXT("Enable accurate coverage pass"));
 
+static int32 GHairStrandsMaterialCompactionEnable = 0;
+static FAutoConsoleVariableRef CVarHairStrandsMaterialCompactionEnable(TEXT("r.HairStrands.MaterialCompaction"), GHairStrandsMaterialCompactionEnable, TEXT("Enable extra compaction based on material properties in order to reduce sample per pixel and improve performance."));
+
+static float GHairStrandsMaterialCompactionDepthThreshold = 1.f;
+static float GHairStrandsMaterialCompactionTangentThreshold = 10.f;
+static FAutoConsoleVariableRef CVarHairStrandsMaterialCompactionDepthThreshold(TEXT("r.HairStrands.MaterialCompaction.DepthThreshold"), GHairStrandsMaterialCompactionDepthThreshold, TEXT("Compaction threshold for depth value for material compaction (in centimeters). Default 1 cm."));
+static FAutoConsoleVariableRef CVarHairStrandsMaterialCompactionTangentThreshold(TEXT("r.HairStrands.MaterialCompaction.TangentThreshold"), GHairStrandsMaterialCompactionTangentThreshold, TEXT("Compaciton threshold for tangent value for material compaction (in degrees). Default 10 deg."));
+
+
 static int32 GHairVisibilitySampleCount = 8;
 static FAutoConsoleVariableRef CVarHairVisibilitySampleCount(TEXT("r.HairStrands.VisibilitySampleCount"), GHairVisibilitySampleCount, TEXT("Hair strands visibility sample count"));
 
@@ -437,13 +446,16 @@ class FHairVisibilityPrimitiveIdCompactionCS : public FGlobalShader
 	class FVendor   : SHADER_PERMUTATION_INT("PERMUTATION_VENDOR", HairVisibilityVendorCount);
 	class FVelocity : SHADER_PERMUTATION_INT("PERMUTATION_VELOCITY", 4);
 	class FCoverage : SHADER_PERMUTATION_INT("PERMUTATION_COVERAGE", 2);
-	using FPermutationDomain = TShaderPermutationDomain<FVendor, FVelocity, FCoverage>;
+	class FMaterial : SHADER_PERMUTATION_INT("PERMUTATION_MATERIAL_COMPACTION", 2);
+	using FPermutationDomain = TShaderPermutationDomain<FVendor, FVelocity, FCoverage, FMaterial>;
 
 	BEGIN_SHADER_PARAMETER_STRUCT(FParameters, )
 		SHADER_PARAMETER(FIntPoint, OutputResolution)
 		SHADER_PARAMETER(uint32, MaxNodeCount)
 		SHADER_PARAMETER(uint32, HairVisibilitySampleCount)
 		SHADER_PARAMETER(FIntPoint, ResolutionOffset)
+		SHADER_PARAMETER(float, DepthTheshold)
+		SHADER_PARAMETER(float, CosTangentThreshold)
 
 		SHADER_PARAMETER_RDG_TEXTURE(Texture2D, MSAA_DepthTexture)
 		SHADER_PARAMETER_RDG_TEXTURE(Texture2D, MSAA_IDTexture)
@@ -459,8 +471,9 @@ class FHairVisibilityPrimitiveIdCompactionCS : public FGlobalShader
 		SHADER_PARAMETER_RDG_BUFFER_UAV(StructuredBuffer, OutCompactNodeCoord)
 		SHADER_PARAMETER_RDG_TEXTURE_UAV(Texture2D, OutVelocityTexture)
 
+		SHADER_PARAMETER_STRUCT_REF(FViewUniformShaderParameters, ViewUniformBuffer)
 		SHADER_PARAMETER_STRUCT_REF(FSceneTexturesUniformParameters, SceneTexturesStruct)
-		END_SHADER_PARAMETER_STRUCT()
+	END_SHADER_PARAMETER_STRUCT()
 
 public:
 	static bool ShouldCompilePermutation(const FGlobalShaderPermutationParameters& Parameters) { return IsHairStrandsSupported(Parameters.Platform); }
@@ -564,6 +577,7 @@ static void AddHairVisibilityPrimitiveIdCompactionPass(
 	PermutationVector.Set<FHairVisibilityPrimitiveIdCompactionCS::FVendor>(GetVendor());
 	PermutationVector.Set<FHairVisibilityPrimitiveIdCompactionCS::FVelocity>(bWriteOutVelocity ? FMath::Clamp(GHairVelocityType+1, 0, 3) : 0);
 	PermutationVector.Set<FHairVisibilityPrimitiveIdCompactionCS::FCoverage>(CoverageTexture ? 1 : 0);
+	PermutationVector.Set<FHairVisibilityPrimitiveIdCompactionCS::FMaterial>(GHairStrandsMaterialCompactionEnable ? 1 : 0);
 
 	FHairVisibilityPrimitiveIdCompactionCS::FParameters* Parameters = GraphBuilder.AllocParameters<FHairVisibilityPrimitiveIdCompactionCS::FParameters>();
 	Parameters->MSAA_DepthTexture = MSAA_DepthTexture;
@@ -573,8 +587,11 @@ static void AddHairVisibilityPrimitiveIdCompactionPass(
 	Parameters->CoverageTexture = CoverageTexture;
 	Parameters->OutputResolution = Resolution;
 	Parameters->MaxNodeCount = MaxNodeCount;
+	Parameters->DepthTheshold = FMath::Clamp(GHairStrandsMaterialCompactionDepthThreshold, 0.f, 100.f);
+	Parameters->CosTangentThreshold = FMath::Cos(FMath::DegreesToRadians(FMath::Clamp(GHairStrandsMaterialCompactionTangentThreshold, 0.f, 90.f)));
 	Parameters->HairVisibilitySampleCount = HairVisibilitySampleCount;
 	Parameters->SceneTexturesStruct = CreateUniformBufferImmediate(SceneTextures, EUniformBufferUsage::UniformBuffer_SingleDraw);
+	Parameters->ViewUniformBuffer = View.ViewUniformBuffer;
 	Parameters->OutCompactNodeCounter = GraphBuilder.CreateUAV(CompactCounter);
 	Parameters->OutCompactNodeIndex = GraphBuilder.CreateUAV(OutCompactNodeIndex);
 	Parameters->OutCompactNodeData = GraphBuilder.CreateUAV(OutCompactNodeData);
@@ -623,7 +640,7 @@ class FHairVisibilityFillOpaqueDepthPS : public FGlobalShader
 
 		SHADER_PARAMETER_STRUCT_REF(FViewUniformShaderParameters, ViewUniformBuffer)
 		RENDER_TARGET_BINDING_SLOTS()
-		END_SHADER_PARAMETER_STRUCT()
+	END_SHADER_PARAMETER_STRUCT()
 
 public:
 	static bool ShouldCompilePermutation(const FGlobalShaderPermutationParameters& Parameters) { return IsHairStrandsSupported(Parameters.Platform); }
