Fix RDG validation error.

(ushell-p4-cherrypick of 9618667 by Charles.deRousiers)

diff --git a/Engine/Source/Runtime/Renderer/Private/HairStrands/HairStrandsLUT.cpp b/Engine/Source/Runtime/Renderer/Private/HairStrands/HairStrandsLUT.cpp
--- a/Engine/Source/Runtime/Renderer/Private/HairStrands/HairStrandsLUT.cpp
+++ b/Engine/Source/Runtime/Renderer/Private/HairStrands/HairStrandsLUT.cpp
@@ -58,7 +58,7 @@ static FRDGTextureRef AddHairLUTPass(
 	OutputDesc.Format = PF_FloatRGBA;
 	OutputDesc.NumMips = 1;
 	OutputDesc.Flags = TexCreate_ShaderResource;
-	OutputDesc.TargetableFlags = TexCreate_UAV;
+	OutputDesc.TargetableFlags = TexCreate_UAV | TexCreate_ShaderResource;
 	FRDGTextureRef HairLUTTexture = GraphBuilder.CreateTexture(OutputDesc, TEXT("HairLUT"));
 
 	FHairLUTCS::FParameters* Parameters = GraphBuilder.AllocParameters<FHairLUTCS::FParameters>();
diff --git a/Engine/Source/Runtime/Renderer/Private/HairStrands/HairStrandsVisibility.cpp b/Engine/Source/Runtime/Renderer/Private/HairStrands/HairStrandsVisibility.cpp
--- a/Engine/Source/Runtime/Renderer/Private/HairStrands/HairStrandsVisibility.cpp
+++ b/Engine/Source/Runtime/Renderer/Private/HairStrands/HairStrandsVisibility.cpp
@@ -516,7 +516,7 @@ static void AddHairVisibilityPrimitiveIdCompactionPass(
 		Desc.NumMips = 1;
 		Desc.NumSamples = 1;
 		Desc.Flags = TexCreate_None;
-		Desc.TargetableFlags = TexCreate_UAV;
+		Desc.TargetableFlags = TexCreate_UAV | TexCreate_ShaderResource;
 		Desc.ClearValue = FClearValueBinding(0);
 		CompactCounter = GraphBuilder.CreateTexture(Desc, TEXT("HairVisibilityCompactCounter"));
 	}
@@ -539,7 +539,7 @@ static void AddHairVisibilityPrimitiveIdCompactionPass(
 		OutputDesc.Extent = Resolution;
 		OutputDesc.Format = PF_R16G16B16A16_UINT;
 		OutputDesc.NumMips = 1;
-		OutputDesc.TargetableFlags = TexCreate_UAV;
+		OutputDesc.TargetableFlags = TexCreate_UAV | TexCreate_ShaderResource;
 		OutCategorizationTexture = GraphBuilder.CreateTexture(OutputDesc, TEXT("CategorizationTexture"));
 	}
 	
diff --git a/Engine/Source/Runtime/Renderer/Private/PostProcess/SceneRenderTargets.cpp b/Engine/Source/Runtime/Renderer/Private/PostProcess/SceneRenderTargets.cpp
--- a/Engine/Source/Runtime/Renderer/Private/PostProcess/SceneRenderTargets.cpp
+++ b/Engine/Source/Runtime/Renderer/Private/PostProcess/SceneRenderTargets.cpp
@@ -2500,7 +2500,7 @@ void FSceneRenderTargets::AllocateDeferredShadingPathRenderTargets(FRHICommandLi
 
 	if (CurrentFeatureLevel >= ERHIFeatureLevel::SM5)
 	{
-		FPooledRenderTargetDesc Desc(FPooledRenderTargetDesc::Create2DDesc(BufferSize, PF_FloatRGBA,  FClearValueBinding::Transparent, TexCreate_None, TexCreate_RenderTargetable, false));
+		FPooledRenderTargetDesc Desc(FPooledRenderTargetDesc::Create2DDesc(BufferSize, PF_FloatRGBA,  FClearValueBinding::Transparent, TexCreate_None, TexCreate_RenderTargetable | TexCreate_ShaderResource, false));
 		GRenderTargetPool.FindFreeElement(RHICmdList, Desc, SceneColorSubPixel, TEXT("SceneColorSubPixel"), true);
 	}
 
