Fix more RDG validation error.

(ushell-p4-cherrypick of 9621699 by Charles.deRousiers)

diff --git a/Engine/Source/Runtime/Renderer/Private/HairStrands/HairStrandsDeepShadow.cpp b/Engine/Source/Runtime/Renderer/Private/HairStrands/HairStrandsDeepShadow.cpp
--- a/Engine/Source/Runtime/Renderer/Private/HairStrands/HairStrandsDeepShadow.cpp
+++ b/Engine/Source/Runtime/Renderer/Private/HairStrands/HairStrandsDeepShadow.cpp
@@ -174,12 +174,12 @@ FHairStrandsDeepShadowViews RenderHairStrandsDeepShadows(
 	bool bClearLayerAtlasTexture = true;
 	{
 		{
-			FPooledRenderTargetDesc ShadowDesc(FPooledRenderTargetDesc::Create2DDesc(AtlasResolution, PF_DepthStencil, FClearValueBinding::DepthFar, TexCreate_None, TexCreate_DepthStencilTargetable, false));
+			FPooledRenderTargetDesc ShadowDesc(FPooledRenderTargetDesc::Create2DDesc(AtlasResolution, PF_DepthStencil, FClearValueBinding::DepthFar, TexCreate_None, TexCreate_DepthStencilTargetable | TexCreate_ShaderResource, false));
 			GRenderTargetPool.FindFreeElement(RHICmdList, ShadowDesc, FrontDepthAtlasTexture, TEXT("ShadowDepth"));
 		}
 
 		{
-			FPooledRenderTargetDesc Desc(FPooledRenderTargetDesc::Create2DDesc(AtlasResolution, PF_FloatRGBA, FClearValueBinding::Transparent, TexCreate_None, TexCreate_RenderTargetable, false));
+			FPooledRenderTargetDesc Desc(FPooledRenderTargetDesc::Create2DDesc(AtlasResolution, PF_FloatRGBA, FClearValueBinding::Transparent, TexCreate_None, TexCreate_RenderTargetable | TexCreate_ShaderResource, false));
 			GRenderTargetPool.FindFreeElement(RHICmdList, Desc, DeepShadowLayersAtlasTexture, TEXT("DeepShadowLayers"));
 		}
 	}
