Fix Vulkan queries to allow accurate timers and GPU stats on mobile - integration of CL from Oculus.

diff --git a/Engine/Source/Runtime/Renderer/Private/MobileShadingRenderer.cpp b/Engine/Source/Runtime/Renderer/Private/MobileShadingRenderer.cpp
--- a/Engine/Source/Runtime/Renderer/Private/MobileShadingRenderer.cpp
+++ b/Engine/Source/Runtime/Renderer/Private/MobileShadingRenderer.cpp
@@ -72,6 +72,8 @@ static TAutoConsoleVariable<int32> CVarMobileAdrenoOcclusionMode(
 	TEXT("1: Render occlusion queries after translucency and a flush, which can help Adreno devices in GL mode."),
 	ECVF_RenderThreadSafe);
 
+DECLARE_GPU_STAT_NAMED(MobileSceneRender, TEXT("Mobile Scene Render"));
+
 DECLARE_CYCLE_STAT(TEXT("SceneStart"), STAT_CLMM_SceneStart, STATGROUP_CommandListMarkers);
 DECLARE_CYCLE_STAT(TEXT("SceneEnd"), STAT_CLMM_SceneEnd, STATGROUP_CommandListMarkers);
 DECLARE_CYCLE_STAT(TEXT("InitViews"), STAT_CLMM_InitViews, STATGROUP_CommandListMarkers);
@@ -328,6 +330,9 @@ void FMobileSceneRenderer::Render(FRHICommandListImmediate& RHICmdList)
 {
 	RHICmdList.SetCurrentStat(GET_STATID(STAT_CLMM_SceneStart));
 
+	SCOPED_DRAW_EVENT(RHICmdList, MobileSceneRender);
+	SCOPED_GPU_STAT(RHICmdList, MobileSceneRender);
+
 	Scene->UpdateAllPrimitiveSceneInfos(RHICmdList);
 
 	PrepareViewRectsForRendering();
diff --git a/Engine/Source/Runtime/VulkanRHI/Private/VulkanQuery.cpp b/Engine/Source/Runtime/VulkanRHI/Private/VulkanQuery.cpp
--- a/Engine/Source/Runtime/VulkanRHI/Private/VulkanQuery.cpp
+++ b/Engine/Source/Runtime/VulkanRHI/Private/VulkanQuery.cpp
@@ -609,7 +609,7 @@ FVulkanTimingQuery::FVulkanTimingQuery(FVulkanDevice* InDevice)
 	: FVulkanRenderQuery(RQT_AbsoluteTime)
 	, Pool(InDevice, 4)
 {
-	Pool.ResultsBuffer = InDevice->GetStagingManager().AcquireBuffer(Pool.BufferSize * sizeof(uint64), VK_BUFFER_USAGE_TRANSFER_DST_BIT, VK_MEMORY_PROPERTY_HOST_CACHED_BIT);
+	Pool.ResultsBuffer = InDevice->GetStagingManager().AcquireBuffer(Pool.BufferSize * sizeof(uint64), VK_BUFFER_USAGE_TRANSFER_DST_BIT, VK_MEMORY_PROPERTY_HOST_COHERENT_BIT);
 }
 
 
diff --git a/Engine/Source/Runtime/VulkanRHI/Private/VulkanUtil.cpp b/Engine/Source/Runtime/VulkanRHI/Private/VulkanUtil.cpp
--- a/Engine/Source/Runtime/VulkanRHI/Private/VulkanUtil.cpp
+++ b/Engine/Source/Runtime/VulkanRHI/Private/VulkanUtil.cpp
@@ -149,7 +149,7 @@ void FVulkanGPUTiming::Initialize()
 	{
 		check(!Pool);
 		Pool = new FVulkanTimingQueryPool(Device, 8);
-		Pool->ResultsBuffer = Device->GetStagingManager().AcquireBuffer(8 * sizeof(uint64), VK_BUFFER_USAGE_TRANSFER_DST_BIT, VK_MEMORY_PROPERTY_HOST_CACHED_BIT);
+		Pool->ResultsBuffer = Device->GetStagingManager().AcquireBuffer(8 * sizeof(uint64), VK_BUFFER_USAGE_TRANSFER_DST_BIT, VK_MEMORY_PROPERTY_HOST_COHERENT_BIT);
 	}
 }
 
