Fix hair raytracing geometry.

diff --git a/Engine/Plugins/Experimental/HairStrands/Source/HairStrandsCore/Private/HairStrandsRendering.cpp b/Engine/Plugins/Experimental/HairStrands/Source/HairStrandsCore/Private/HairStrandsRendering.cpp
--- a/Engine/Plugins/Experimental/HairStrands/Source/HairStrandsCore/Private/HairStrandsRendering.cpp
+++ b/Engine/Plugins/Experimental/HairStrands/Source/HairStrandsCore/Private/HairStrandsRendering.cpp
@@ -353,7 +353,7 @@ static void BuildHairAccelerationStructure(FRHICommandList& RHICmdList, uint32 R
 	Initializer.IndexBuffer = nullptr;
 	Initializer.IndexBufferOffset = 0;
 	Initializer.GeometryType = RTGT_Triangles;
-	Initializer.TotalPrimitiveCount = RaytracingVertexCount/3;
+	Initializer.TotalPrimitiveCount = RaytracingVertexCount / 3;
 	Initializer.bFastBuild = true;
 	Initializer.bAllowUpdate = true;
 
@@ -361,7 +361,7 @@ static void BuildHairAccelerationStructure(FRHICommandList& RHICmdList, uint32 R
 	Segment.VertexBuffer = PositionBuffer;
 	Segment.VertexBufferStride = FHairStrandsRaytracingFormat::SizeInByte;
 	Segment.VertexBufferElementType = FHairStrandsRaytracingFormat::VertexElementType;
-	Segment.NumPrimitives = RaytracingVertexCount;
+	Segment.NumPrimitives = RaytracingVertexCount / 3;
 	Initializer.Segments.Add(Segment);
 
 	OutRayTracingGeometry->SetInitializer(Initializer);
diff --git a/Engine/Source/Runtime/Renderer/Private/DeferredShadingRenderer.cpp b/Engine/Source/Runtime/Renderer/Private/DeferredShadingRenderer.cpp
--- a/Engine/Source/Runtime/Renderer/Private/DeferredShadingRenderer.cpp
+++ b/Engine/Source/Runtime/Renderer/Private/DeferredShadingRenderer.cpp
@@ -813,7 +813,10 @@ bool FDeferredShadingSceneRenderer::GatherRayTracingWorldInstances(FRHICommandLi
 						RayTracingInstance.Mask = Instance.Mask;
 						RayTracingInstance.bForceOpaque = Instance.bForceOpaque;
 
-						check(Instance.Materials.Num() == Instance.Geometry->Initializer.Segments.Num() || (Instance.Geometry->Initializer.Segments.Num() == 0 && Instance.Materials.Num() == 1));
+						// Thin geometries like hair don't have material, as they only support shadow at the moment.
+						check(Instance.Materials.Num() == Instance.Geometry->Initializer.Segments.Num() || 
+							 (Instance.Geometry->Initializer.Segments.Num() == 0 && Instance.Materials.Num() == 1) || 
+							 (Instance.Materials.Num() == 0 && (Instance.Mask & RAY_TRACING_MASK_THIN_SHADOW) > 0));
 
 						RayTracingInstance.Transforms.SetNumUninitialized(Instance.InstanceTransforms.Num());
 						FMemory::Memcpy(RayTracingInstance.Transforms.GetData(), Instance.InstanceTransforms.GetData(), Instance.InstanceTransforms.Num() * sizeof(RayTracingInstance.Transforms[0]));
