Add boundary check for transmittance mask evaluation.

diff --git a/Engine/Shaders/Private/HairStrands/HairStrandsDeepTransmittanceMask.usf b/Engine/Shaders/Private/HairStrands/HairStrandsDeepTransmittanceMask.usf
--- a/Engine/Shaders/Private/HairStrands/HairStrandsDeepTransmittanceMask.usf
+++ b/Engine/Shaders/Private/HairStrands/HairStrandsDeepTransmittanceMask.usf
@@ -29,6 +29,7 @@ float3 ToWorldPosition(float2 UV, float Depth)
 uint2				DeepShadow_Resolution;
 float4x4			DeepShadow_WorldToLightTransforms[MAX_HAIR_CLUSTER_COUNT];
 float3				LightDirection;
+uint				MaxVisibilityNodeCount;
 float4				LightPosition;
 float				DepthBiasScale;
 float				DensityScale;
@@ -74,6 +75,9 @@ void MainCS(uint2 DispatchThreadId : SV_DispatchThreadID)
 	// Note: Use a fixed group count width (HAIR_VISIBILITY_GROUP_COUNT_WIDTH) for avoiding loading the indirect args (to avoid dep. memory fetch)
 	const uint LocalOffset = DispatchThreadId.x + DispatchThreadId.y * TRANSMITTANCE_GROUP_SIZE * HAIR_VISIBILITY_GROUP_COUNT_WIDTH;
 
+	if (LocalOffset >= MaxVisibilityNodeCount)
+		return;
+
 	const uint PackedCoord	 = HairVisibilityNodeCoord[LocalOffset];
 	const uint2 PixelCoord	 = uint2((PackedCoord & 0xFFFF), ((PackedCoord>>16) & 0xFFFF));
 
diff --git a/Engine/Source/Runtime/Renderer/Private/HairStrands/HairStrandsTransmittance.cpp b/Engine/Source/Runtime/Renderer/Private/HairStrands/HairStrandsTransmittance.cpp
--- a/Engine/Source/Runtime/Renderer/Private/HairStrands/HairStrandsTransmittance.cpp
+++ b/Engine/Source/Runtime/Renderer/Private/HairStrands/HairStrandsTransmittance.cpp
@@ -79,6 +79,7 @@ class FDeepTransmittanceMaskCS : public FGlobalShader
 		SHADER_PARAMETER_ARRAY(FMatrix, DeepShadow_WorldToLightTransforms, [FHairStrandsDeepShadowData::MaxClusterCount])
 		SHADER_PARAMETER(FIntPoint, DeepShadow_Resolution)
 		SHADER_PARAMETER(FVector, LightDirection)
+		SHADER_PARAMETER(uint32, MaxVisibilityNodeCount)
 		SHADER_PARAMETER(FVector4, LightPosition)
 		SHADER_PARAMETER(float, DepthBiasScale)
 		SHADER_PARAMETER(float, DensityScale)
@@ -176,6 +177,7 @@ static FRDGBufferRef AddDeepShadowTransmittanceMaskPass(
 	Parameters->DeepShadow_DebugMode = GetDeepShadowDebugMode();
 	Parameters->DeepShadow_ShadowToWorld = Params.DeepShadow_ShadowToWorld;
 	Parameters->IndirectArgsBuffer = IndirectArgsBuffer;
+	Parameters->MaxVisibilityNodeCount = Params.HairVisibilityNodeData->Desc.NumElements;
 
 	memcpy(&(Parameters->DeepShadow_AtlasSlotOffsets[0]), Params.DeepShadow_AtlasSlotOffsets, sizeof(FIntVector4) * FHairStrandsDeepShadowData::MaxClusterCount);
 	memcpy(&(Parameters->DeepShadow_WorldToLightTransforms[0]), Params.DeepShadow_WorldToLightTransforms, sizeof(FMatrix) * FHairStrandsDeepShadowData::MaxClusterCount);
