Fix CIS issue.

diff --git a/Engine/Source/Runtime/Renderer/Private/HairStrands/HairStrandsDebug.cpp b/Engine/Source/Runtime/Renderer/Private/HairStrands/HairStrandsDebug.cpp
--- a/Engine/Source/Runtime/Renderer/Private/HairStrands/HairStrandsDebug.cpp
+++ b/Engine/Source/Runtime/Renderer/Private/HairStrands/HairStrandsDebug.cpp
@@ -208,7 +208,7 @@ static void AddDebugHairPrintPass(
 
 	FHairDebugPrintCS::FParameters* Parameters = GraphBuilder.AllocParameters<FHairDebugPrintCS::FParameters>();
 	Parameters->GroupSize = GetVendorOptimalGroupSize2D();
-	Parameters->MaxResolution = CategorizationTexture->Desc.Extent;
+	Parameters->MaxResolution = CategorizationTexture ? CategorizationTexture->Desc.Extent : FIntPoint(0,0);
 	Parameters->PixelCoord = View->CursorPos;
 	Parameters->FastResolveMask = STENCIL_TEMPORAL_RESPONSIVE_AA_MASK;
 	Parameters->CategorizationTexture = CategorizationTexture;
