Shadow ray tracing optimizations

* Set up ray tracing hints for shadow ray generation shader
* Added a single ray per pixel specialization for shadow ray generation shader

diff --git a/Engine/Shaders/Private/RayTracing/RayTracingOcclusionRGS.usf b/Engine/Shaders/Private/RayTracing/RayTracingOcclusionRGS.usf
--- a/Engine/Shaders/Private/RayTracing/RayTracingOcclusionRGS.usf
+++ b/Engine/Shaders/Private/RayTracing/RayTracingOcclusionRGS.usf
@@ -124,7 +124,7 @@ RAY_TRACING_ENTRY_RAYGEN(OcclusionRGS)
 	float ClosestRayDistance = DENOISER_INVALID_HIT_DISTANCE;
 	float HitCount = 0.0;
 	float RayCount = 0.0;
-	uint LocalSamplesPerPixel = SamplesPerPixel;
+	uint LocalSamplesPerPixel = ENABLE_MULTIPLE_SAMPLES_PER_PIXEL ? SamplesPerPixel : 1;
 
 	if (all(PixelCoord >= LightScissor.xy) && all(PixelCoord <= LightScissor.zw))
 	{
@@ -178,7 +178,7 @@ RAY_TRACING_ENTRY_RAYGEN(OcclusionRGS)
 		const bool bTraceRay = bIsValidPixel && (LightChannel & LightingChannelMask) != 0;
 		if (!bTraceRay)
 		{
-			LocalSamplesPerPixel = 0.0;
+			LocalSamplesPerPixel = 0;
 		}
 
 		// Cheap solution for double sided foliage to have valid lighing through the denoiser
@@ -217,14 +217,16 @@ RAY_TRACING_ENTRY_RAYGEN(OcclusionRGS)
 			bApplyNormalCulling = false;
 		}
 
-		LOOP
-		for (uint SampleIndex = 0; SampleIndex < LocalSamplesPerPixel; ++SampleIndex)
+#if ENABLE_MULTIPLE_SAMPLES_PER_PIXEL
+		LOOP for (uint SampleIndex = 0; SampleIndex < LocalSamplesPerPixel; ++SampleIndex)
+#else // ENABLE_MULTIPLE_SAMPLES_PER_PIXEL
+		do if (bTraceRay)
+#endif // ENABLE_MULTIPLE_SAMPLES_PER_PIXEL
 		{
 			uint DummyVariable;
 			float2 RandSample = RandomSequence_GenerateSample2D(RandSequence, DummyVariable);
 
-			BRANCH
-			if (LocalSamplesPerPixel == 1 && 1)
+			#if !ENABLE_MULTIPLE_SAMPLES_PER_PIXEL
 			{
 				const uint TemporalBits = 10;
 				uint FrameIndexMod1024 = ReverseBitsN(GetPowerOfTwoModulatedFrameIndex(1 << TemporalBits), TemporalBits);
@@ -234,6 +236,7 @@ RAY_TRACING_ENTRY_RAYGEN(OcclusionRGS)
 
 				RandSample = SobolIndexToUniformUnitSquare(SobolFrame);
 			}
+			#endif // !ENABLE_MULTIPLE_SAMPLES_PER_PIXEL
 
 			RayDesc Ray;
 			bool bIsValidRay = GenerateOcclusionRay(
@@ -300,6 +303,10 @@ RAY_TRACING_ENTRY_RAYGEN(OcclusionRGS)
 				ClosestRayDistance = (ClosestRayDistance == DENOISER_INVALID_HIT_DISTANCE) ? DENOISER_MISS_HIT_DISTANCE : ClosestRayDistance;
 			}
 		}
+#if !ENABLE_MULTIPLE_SAMPLES_PER_PIXEL
+		while (0);
+#endif //!ENABLE_MULTIPLE_SAMPLES_PER_PIXEL
+
 	}
 
 	float Shadow = (LocalSamplesPerPixel > 0) ? Visibility / LocalSamplesPerPixel : Visibility;
diff --git a/Engine/Source/Runtime/Renderer/Private/RayTracing/RayTracingShadows.cpp b/Engine/Source/Runtime/Renderer/Private/RayTracing/RayTracingShadows.cpp
--- a/Engine/Source/Runtime/Renderer/Private/RayTracing/RayTracingShadows.cpp
+++ b/Engine/Source/Runtime/Renderer/Private/RayTracing/RayTracingShadows.cpp
@@ -54,15 +54,35 @@ class FOcclusionRGS : public FGlobalShader
 	class FLightTypeDim : SHADER_PERMUTATION_INT("LIGHT_TYPE", LightType_MAX);
 	class FDenoiserOutputDim : SHADER_PERMUTATION_INT("DIM_DENOISER_OUTPUT", 3);
 	class FEnableTwoSidedGeometryDim : SHADER_PERMUTATION_BOOL("ENABLE_TWO_SIDED_GEOMETRY");
+	class FEnableMultipleSamplesPerPixel : SHADER_PERMUTATION_BOOL("ENABLE_MULTIPLE_SAMPLES_PER_PIXEL");
 	class FHairLighting : SHADER_PERMUTATION_INT("USE_HAIR_LIGHTING", 2);
 
-	using FPermutationDomain = TShaderPermutationDomain<FLightTypeDim, FDenoiserOutputDim, FEnableTwoSidedGeometryDim, FHairLighting>;
+	using FPermutationDomain = TShaderPermutationDomain<FLightTypeDim, FDenoiserOutputDim, FEnableTwoSidedGeometryDim, FHairLighting, FEnableMultipleSamplesPerPixel>;
 
 	static bool ShouldCompilePermutation(const FGlobalShaderPermutationParameters& Parameters)
 	{
 		return ShouldCompileRayTracingShadersForProject(Parameters.Platform);
 	}
 
+	static void ModifyCompilationEnvironment(const FGlobalShaderPermutationParameters& Parameters, FShaderCompilerEnvironment& OutEnvironment)
+	{
+		FGlobalShader::ModifyCompilationEnvironment(Parameters, OutEnvironment);
+
+		OutEnvironment.SetDefine(TEXT("UE_RAY_TRACING_DYNAMIC_CLOSEST_HIT_SHADER"), 0);
+		OutEnvironment.SetDefine(TEXT("UE_RAY_TRACING_DYNAMIC_ANY_HIT_SHADER"),     1);
+		OutEnvironment.SetDefine(TEXT("UE_RAY_TRACING_DYNAMIC_MISS_SHADER"),        0);
+
+		FPermutationDomain PermutationVector(Parameters.PermutationId);
+		if (PermutationVector.Get<FLightTypeDim>() == LightType_Directional)
+		{
+			OutEnvironment.SetDefine(TEXT("UE_RAY_TRACING_COHERENT_RAYS"), 1);
+		}
+		else
+		{
+			OutEnvironment.SetDefine(TEXT("UE_RAY_TRACING_COHERENT_RAYS"), 0);
+		}
+	}
+
 	BEGIN_SHADER_PARAMETER_STRUCT(FParameters, )
 		SHADER_PARAMETER(uint32, SamplesPerPixel)
 		SHADER_PARAMETER(float, NormalBias)
@@ -100,6 +120,7 @@ void FDeferredShadingSceneRenderer::PrepareRayTracingShadows(const FViewInfo& Vi
 		IScreenSpaceDenoiser::EShadowRequirements::PenumbraAndClosestOccluder,
 	};
 
+	for (int32 MultiSPP = 0; MultiSPP < 1; ++MultiSPP)
 	for (int32 HairLighting = 0; HairLighting < 2; ++HairLighting)
 	for (int32 LightType = 0; LightType < LightType_MAX; ++LightType)
 	{
@@ -110,6 +131,7 @@ void FDeferredShadingSceneRenderer::PrepareRayTracingShadows(const FViewInfo& Vi
 			PermutationVector.Set<FOcclusionRGS::FDenoiserOutputDim>((int32)DenoiserRequirement);
 			PermutationVector.Set<FOcclusionRGS::FEnableTwoSidedGeometryDim>(CVarRayTracingShadowsEnableTwoSidedGeometry.GetValueOnRenderThread() != 0);
 			PermutationVector.Set<FOcclusionRGS::FHairLighting>(HairLighting);
+			PermutationVector.Set<FOcclusionRGS::FEnableMultipleSamplesPerPixel>(MultiSPP);
 
 			TShaderMapRef<FOcclusionRGS> RayGenerationShader(View.ShaderMap, PermutationVector);
 			OutRayGenShaders.Add(RayGenerationShader->GetRayTracingShader());
@@ -221,6 +243,8 @@ void FDeferredShadingSceneRenderer::RenderRayTracingShadows(
 		PermutationVector.Set<FOcclusionRGS::FEnableTwoSidedGeometryDim>(CVarRayTracingShadowsEnableTwoSidedGeometry.GetValueOnRenderThread() != 0);
 		PermutationVector.Set<FOcclusionRGS::FHairLighting>(bUseHairLighting? 1 : 0);
 
+		PermutationVector.Set<FOcclusionRGS::FEnableMultipleSamplesPerPixel>(RayTracingConfig.RayCountPerPixel > 1);
+
 		TShaderMapRef<FOcclusionRGS> RayGenerationShader(GetGlobalShaderMap(FeatureLevel), PermutationVector);
 
 		ClearUnusedGraphResources(*RayGenerationShader, PassParameters);
