Added multiview support for ray tracing.
TODO: denoisers do not work in multiview mode yet

diff --git a/Engine/Source/Runtime/Renderer/Private/DeferredShadingRenderer.h b/Engine/Source/Runtime/Renderer/Private/DeferredShadingRenderer.h
--- a/Engine/Source/Runtime/Renderer/Private/DeferredShadingRenderer.h
+++ b/Engine/Source/Runtime/Renderer/Private/DeferredShadingRenderer.h
@@ -504,6 +504,9 @@ private:
 		int32 SamplePerPixel,
 		int32 HeightFog,
 		float ResolutionFraction,
+		FRDGTextureUAV* ColorOutputUAV,
+		FRDGTextureUAV* RayHitDistanceOutputUAV,
+		FRDGTextureUAV* RayImaginaryDepthOutputUAV,
 		IScreenSpaceDenoiser::FReflectionsInputs* OutDenoiserInputs);
 
 	void RenderRayTracingShadows(
@@ -515,8 +518,8 @@ private:
 		const IScreenSpaceDenoiser::EShadowRequirements DenoiserRequirements,
 		const bool bSubPixelShadowMask,
 		FRDGTextureRef HairCategorizationTexture,
-		FRDGTextureRef* OutShadowMask,
-		FRDGTextureRef* OutRayHitDistance);
+		FRDGTextureUAV* OutShadowMaskUAV,
+		FRDGTextureUAV* OutRayHitDistanceUAV);
 
 	void RenderRayTracingStochasticRectLight(FRHICommandListImmediate& RHICmdList, const FLightSceneInfo& RectLightSceneInfo, TRefCountPtr<IPooledRenderTarget>& RectLightRT, TRefCountPtr<IPooledRenderTarget>& HitDistanceRT);
 	void CompositeRayTracingSkyLight(FRHICommandListImmediate& RHICmdList, TRefCountPtr<IPooledRenderTarget>& SkyLightRT, TRefCountPtr<IPooledRenderTarget>& HitDistanceRT);
diff --git a/Engine/Source/Runtime/Renderer/Private/LightRendering.cpp b/Engine/Source/Runtime/Renderer/Private/LightRendering.cpp
--- a/Engine/Source/Runtime/Renderer/Private/LightRendering.cpp
+++ b/Engine/Source/Runtime/Renderer/Private/LightRendering.cpp
@@ -1148,7 +1148,9 @@ void FDeferredShadingSceneRenderer::RenderLights(FRHICommandListImmediate& RHICm
 			const int32 MaxDenoisingBatchSize = FMath::Clamp(CVarMaxShadowDenoisingBatchSize.GetValueOnRenderThread(), 1, IScreenSpaceDenoiser::kMaxBatchSize);
 			const int32 MaxRTShadowBatchSize = CVarMaxShadowRayTracingBatchSize.GetValueOnRenderThread();
 			const bool bDoShadowDenoisingBatching = DenoiserMode != 0 && MaxDenoisingBatchSize > 1;
-			const bool dDoShadowBatching = bDoShadowDenoisingBatching || MaxRTShadowBatchSize > 1;
+
+			//#dxr_todo: support multiview for the batching case
+			const bool dDoShadowBatching = (bDoShadowDenoisingBatching || MaxRTShadowBatchSize > 1) && Views.Num() == 1;
 
 			// Optimisations: batches all shadow ray tracing denoising. Definitely could be smarter to avoid high VGPR pressure if this entire
 			// function was converted to render graph, and want least intrusive change as possible. So right not it trades render target memory pressure
@@ -1227,14 +1229,17 @@ void FDeferredShadingSceneRenderer::RenderLights(FRHICommandListImmediate& RHICm
 
 						const bool bDenoiserCompatible = !LightRequiresDenosier(LightSceneInfo) || IScreenSpaceDenoiser::EShadowRequirements::PenumbraAndClosestOccluder == DenoiserToUse->GetShadowRequirements(View, LightSceneInfo, RayTracingConfig);
 
-						const bool bWantsBatchedShadow = OcclusionType == FLightOcclusionType::Raytraced && bDenoiserCompatible &&
-							SortedLightInfo.SortKey.Fields.bShadowed && !ShouldRenderRayTracingStochasticRectLight(LightSceneInfo);
+						const bool bWantsBatchedShadow = OcclusionType == FLightOcclusionType::Raytraced && 
+							dDoShadowBatching &&
+							bDenoiserCompatible &&
+							SortedLightInfo.SortKey.Fields.bShadowed && 
+							!ShouldRenderRayTracingStochasticRectLight(LightSceneInfo);
 
 						// determine if this light doesn't yet have a precomuted shadow and execute a batch to amortize costs if one is needed
 						if (
 							RHI_RAYTRACING &&
 							bWantsBatchedShadow &&
-							!PreprocessedShadowMaskTextures[LightIndex - AttenuationLightStart])
+							(PreprocessedShadowMaskTextures.Num() == 0 || !PreprocessedShadowMaskTextures[LightIndex - AttenuationLightStart]))
 						{
 							SCOPED_DRAW_EVENT(RHICmdList, ShadowBatch);
 							TStaticArray<IScreenSpaceDenoiser::FShadowVisibilityParameters, IScreenSpaceDenoiser::kMaxBatchSize> DenoisingQueue;
@@ -1320,9 +1325,35 @@ void FDeferredShadingSceneRenderer::RenderLights(FRHICommandListImmediate& RHICm
 								}
 
 								// Ray trace the shadow.
+								//#dxr_todo: support multiview for the batching case
+								FRDGTextureRef RayTracingShadowMaskTexture;
+								{
+									FRDGTextureDesc Desc = FRDGTextureDesc::Create2DDesc(
+										SceneTextures.SceneDepthBuffer->Desc.Extent,
+										PF_FloatRGBA,
+										FClearValueBinding::Black,
+										TexCreate_None,
+										TexCreate_ShaderResource | TexCreate_RenderTargetable | TexCreate_UAV,
+										/* bInForceSeparateTargetAndShaderResource = */ false);
+									RayTracingShadowMaskTexture = GraphBuilder.CreateTexture(Desc, TEXT("RayTracingOcclusion"));
+								}
+
+								FRDGTextureRef RayDistanceTexture;
+								{
+									FRDGTextureDesc Desc = FRDGTextureDesc::Create2DDesc(
+										SceneTextures.SceneDepthBuffer->Desc.Extent,
+										PF_R16F,
+										FClearValueBinding::Black,
+										TexCreate_None,
+										TexCreate_ShaderResource | TexCreate_RenderTargetable | TexCreate_UAV,
+										/* bInForceSeparateTargetAndShaderResource = */ false);
+									RayDistanceTexture = GraphBuilder.CreateTexture(Desc, TEXT("RayTracingOcclusionDistance"));
+								}
+
+								FRDGTextureUAV* RayTracingShadowMaskUAV = GraphBuilder.CreateUAV(FRDGTextureUAVDesc(RayTracingShadowMaskTexture));
+								FRDGTextureUAV* RayHitDistanceUAV = GraphBuilder.CreateUAV(FRDGTextureUAVDesc(RayDistanceTexture));
 								FRDGTextureRef HairCategorizationTexture = GetHairStrandsCategorizationTexture(ViewIndex, GraphBuilder);
-								FRDGTextureRef ShadowMask;
-								FRDGTextureRef RayHitDistance;
+
 								{
 									FString BatchLightNameWithLevel;
 									GetLightNameForDrawEvent(BatchLightSceneInfo.Proxy, BatchLightNameWithLevel);
@@ -1337,16 +1368,42 @@ void FDeferredShadingSceneRenderer::RenderLights(FRHICommandListImmediate& RHICm
 										DenoiserRequirements,
 										false,
 										HairCategorizationTexture,
-										&ShadowMask,
-										&RayHitDistance);
+										RayTracingShadowMaskUAV,
+										RayHitDistanceUAV);
 								}
 
 								// Ray trace the shadow cast by opaque geometries on to hair strands geometries
 								// Note: No denoiser is required on this output, as the hair strands are geometrically noisy, and so hard to denoise
 								if (HasHairStrandsClusters(ViewIndex))
 								{
-									FRDGTextureRef SubPixelShadowMask = nullptr;
-									FRDGTextureRef SubPixelRayHitDistance = nullptr;
+									//#dxr_todo: support multiview for the batching case
+									FRDGTextureRef SubPixelRayTracingShadowMaskTexture;
+									{
+										FRDGTextureDesc Desc = FRDGTextureDesc::Create2DDesc(
+											SceneTextures.SceneDepthBuffer->Desc.Extent,
+											PF_FloatRGBA,
+											FClearValueBinding::Black,
+											TexCreate_None,
+											TexCreate_ShaderResource | TexCreate_RenderTargetable | TexCreate_UAV,
+											/* bInForceSeparateTargetAndShaderResource = */ false);
+										SubPixelRayTracingShadowMaskTexture = GraphBuilder.CreateTexture(Desc, TEXT("RayTracingOcclusion"));
+									}
+
+									FRDGTextureRef SubPixelRayDistanceTexture;
+									{
+										FRDGTextureDesc Desc = FRDGTextureDesc::Create2DDesc(
+											SceneTextures.SceneDepthBuffer->Desc.Extent,
+											PF_R16F,
+											FClearValueBinding::Black,
+											TexCreate_None,
+											TexCreate_ShaderResource | TexCreate_RenderTargetable | TexCreate_UAV,
+											/* bInForceSeparateTargetAndShaderResource = */ false);
+										SubPixelRayDistanceTexture = GraphBuilder.CreateTexture(Desc, TEXT("RayTracingOcclusionDistance"));
+									}
+
+									FRDGTextureUAV* SubPixelRayTracingShadowMaskUAV = GraphBuilder.CreateUAV(FRDGTextureUAVDesc(SubPixelRayTracingShadowMaskTexture));
+									FRDGTextureUAV* SubPixelRayHitDistanceUAV = GraphBuilder.CreateUAV(FRDGTextureUAVDesc(SubPixelRayDistanceTexture));
+
 									RenderRayTracingShadows(
 										GraphBuilder,
 										SceneTextures,
@@ -1356,11 +1413,13 @@ void FDeferredShadingSceneRenderer::RenderLights(FRHICommandListImmediate& RHICm
 										IScreenSpaceDenoiser::EShadowRequirements::Bailout,
 										true,
 										HairCategorizationTexture,
-										&SubPixelShadowMask,
-										&SubPixelRayHitDistance);
+										SubPixelRayTracingShadowMaskUAV,
+										SubPixelRayHitDistanceUAV);
+
 									TRefCountPtr<IPooledRenderTarget>* RefDestination = &PreprocessedShadowMaskSubPixelTextures[LightIndex - AttenuationLightStart];
 									check(*RefDestination == nullptr);
-									GraphBuilder.QueueTextureExtraction(SubPixelShadowMask, RefDestination);
+
+									GraphBuilder.QueueTextureExtraction(SubPixelRayTracingShadowMaskTexture, RefDestination);
 								}
 
 								bool bBatchFull = false;
@@ -1374,8 +1433,8 @@ void FDeferredShadingSceneRenderer::RenderLights(FRHICommandListImmediate& RHICm
 										{
 											DenoisingQueue[i].LightSceneInfo = &BatchLightSceneInfo;
 											DenoisingQueue[i].RayTracingConfig = RayTracingConfig;
-											DenoisingQueue[i].InputTextures.Mask = ShadowMask;
-											DenoisingQueue[i].InputTextures.ClosestOccluder = RayHitDistance;
+											DenoisingQueue[i].InputTextures.Mask = RayTracingShadowMaskTexture;
+											DenoisingQueue[i].InputTextures.ClosestOccluder = RayDistanceTexture;
 											LightIndices[i] = LightBatchIndex;
 
 											// If queue for this light type is full, quick of the batch.
@@ -1394,7 +1453,7 @@ void FDeferredShadingSceneRenderer::RenderLights(FRHICommandListImmediate& RHICm
 								}
 								else
 								{
-									GraphBuilder.QueueTextureExtraction(ShadowMask, &PreprocessedShadowMaskTextures[LightBatchIndex - AttenuationLightStart]);
+									GraphBuilder.QueueTextureExtraction(RayTracingShadowMaskTexture, &PreprocessedShadowMaskTextures[LightBatchIndex - AttenuationLightStart]);
 								}
 
 								// terminate batch if we filled a denoiser batch or hit our max light batch
@@ -1431,90 +1490,145 @@ void FDeferredShadingSceneRenderer::RenderLights(FRHICommandListImmediate& RHICm
 					}
 					else if (OcclusionType == FLightOcclusionType::Raytraced)
 					{
-						const uint32 ViewIndex = 0;
-						FViewInfo& View = Views[ViewIndex];
-
-						IScreenSpaceDenoiser::FShadowRayTracingConfig RayTracingConfig;
-						RayTracingConfig.RayCountPerPixel = LightSceneInfo.Proxy->GetSamplesPerPixel();
-
-						IScreenSpaceDenoiser::EShadowRequirements DenoiserRequirements = IScreenSpaceDenoiser::EShadowRequirements::Bailout;
-						if (DenoiserMode != 0 && LightRequiresDenosier(LightSceneInfo))
-						{
-							DenoiserRequirements = DenoiserToUse->GetShadowRequirements(View, LightSceneInfo, RayTracingConfig);
-						}
-
 						FRDGBuilder GraphBuilder(RHICmdList);
 
 						FSceneTextureParameters SceneTextures;
 						SetupSceneTextureParameters(GraphBuilder, &SceneTextures);
 
-						FRDGTextureRef HairCategorizationTexture = GetHairStrandsCategorizationTexture(ViewIndex, GraphBuilder);
-						FRDGTextureRef ShadowMask;
-						FRDGTextureRef RayHitDistance;
-						RenderRayTracingShadows(
-							GraphBuilder,
-							SceneTextures,
-							View,
-							LightSceneInfo,
-							RayTracingConfig,
-							DenoiserRequirements,
-							false,
-							HairCategorizationTexture,
-							&ShadowMask,
-							&RayHitDistance);
-
-						if (DenoiserRequirements != IScreenSpaceDenoiser::EShadowRequirements::Bailout)
+						FRDGTextureRef RayTracingShadowMaskTexture;
 						{
-							TStaticArray<IScreenSpaceDenoiser::FShadowVisibilityParameters, IScreenSpaceDenoiser::kMaxBatchSize> InputParameters;
-							TStaticArray<IScreenSpaceDenoiser::FShadowVisibilityOutputs, IScreenSpaceDenoiser::kMaxBatchSize> Outputs;
-
-							InputParameters[0].InputTextures.Mask = ShadowMask;
-							InputParameters[0].InputTextures.ClosestOccluder = RayHitDistance;
-							InputParameters[0].LightSceneInfo = &LightSceneInfo;
-							InputParameters[0].RayTracingConfig = RayTracingConfig;
+							FRDGTextureDesc Desc = FRDGTextureDesc::Create2DDesc(
+								SceneTextures.SceneDepthBuffer->Desc.Extent,
+								PF_FloatRGBA,
+								FClearValueBinding::Black,
+								TexCreate_None,
+								TexCreate_ShaderResource | TexCreate_RenderTargetable | TexCreate_UAV,
+								/* bInForceSeparateTargetAndShaderResource = */ false);
+							RayTracingShadowMaskTexture = GraphBuilder.CreateTexture(Desc, TEXT("RayTracingOcclusion"));
+						}
 
-							int32 InputParameterCount = 1;
+						FRDGTextureRef RayDistanceTexture;
+						{
+							FRDGTextureDesc Desc = FRDGTextureDesc::Create2DDesc(
+								SceneTextures.SceneDepthBuffer->Desc.Extent,
+								PF_R16F,
+								FClearValueBinding::Black,
+								TexCreate_None,
+								TexCreate_ShaderResource | TexCreate_RenderTargetable | TexCreate_UAV,
+								/* bInForceSeparateTargetAndShaderResource = */ false);
+							RayDistanceTexture = GraphBuilder.CreateTexture(Desc, TEXT("RayTracingOcclusionDistance"));
+						}
 
-							RDG_EVENT_SCOPE(GraphBuilder, "%s%s(Shadow BatchSize=%d) %dx%d",
-								DenoiserToUse != DefaultDenoiser ? TEXT("ThirdParty ") : TEXT(""),
-								DenoiserToUse->GetDebugName(),
-								InputParameterCount,
-								View.ViewRect.Width(), View.ViewRect.Height());
+						FRDGTextureUAV* RayTracingShadowMaskUAV = GraphBuilder.CreateUAV(FRDGTextureUAVDesc(RayTracingShadowMaskTexture));
+						FRDGTextureUAV* RayHitDistanceUAV = GraphBuilder.CreateUAV(FRDGTextureUAVDesc(RayDistanceTexture));
 
-							DenoiserToUse->DenoiseShadowVisibilityMasks(
-								GraphBuilder,
-								View,
-								&View.PrevViewInfo,
-								SceneTextures,
-								InputParameters,
-								InputParameterCount,
-								Outputs);
 
-							GraphBuilder.QueueTextureExtraction(Outputs[0].Mask, &ScreenShadowMaskTexture);
+						FRDGTextureRef SubPixelRayTracingShadowMaskTexture;
+						{
+							FRDGTextureDesc Desc = FRDGTextureDesc::Create2DDesc(
+								SceneTextures.SceneDepthBuffer->Desc.Extent,
+								PF_FloatRGBA,
+								FClearValueBinding::Black,
+								TexCreate_None,
+								TexCreate_ShaderResource | TexCreate_RenderTargetable | TexCreate_UAV,
+								/* bInForceSeparateTargetAndShaderResource = */ false);
+							SubPixelRayTracingShadowMaskTexture = GraphBuilder.CreateTexture(Desc, TEXT("RayTracingOcclusion"));
 						}
-						else
+
+						FRDGTextureRef SubPixelRayDistanceTexture;
 						{
-							GraphBuilder.QueueTextureExtraction(ShadowMask, &ScreenShadowMaskTexture);
+							FRDGTextureDesc Desc = FRDGTextureDesc::Create2DDesc(
+								SceneTextures.SceneDepthBuffer->Desc.Extent,
+								PF_R16F,
+								FClearValueBinding::Black,
+								TexCreate_None,
+								TexCreate_ShaderResource | TexCreate_RenderTargetable | TexCreate_UAV,
+								/* bInForceSeparateTargetAndShaderResource = */ false);
+							SubPixelRayDistanceTexture = GraphBuilder.CreateTexture(Desc, TEXT("RayTracingOcclusionDistance"));
 						}
 
-						// Ray trace the shadow cast by opaque geometries on to hair strands geometries
-						// Note: No denoiser is required on this output, as the hair strands are geometrically noisy, and so hard to denoise
-						if (HasHairStrandsClusters(ViewIndex))
+						FRDGTextureUAV* SubPixelRayTracingShadowMaskUAV = GraphBuilder.CreateUAV(FRDGTextureUAVDesc(SubPixelRayTracingShadowMaskTexture));
+						FRDGTextureUAV* SubPixelRayHitDistanceUAV = GraphBuilder.CreateUAV(FRDGTextureUAVDesc(SubPixelRayDistanceTexture));
+
+						for (int32 ViewIndex = 0; ViewIndex < Views.Num(); ViewIndex++)
 						{
-							FRDGTextureRef SubPixelShadowMask = nullptr;
-							FRDGTextureRef SubPixelRayHitDistance = nullptr;
+							FViewInfo& View = Views[ViewIndex];
+
+							IScreenSpaceDenoiser::FShadowRayTracingConfig RayTracingConfig;
+							RayTracingConfig.RayCountPerPixel = LightSceneInfo.Proxy->GetSamplesPerPixel();
+
+							IScreenSpaceDenoiser::EShadowRequirements DenoiserRequirements = IScreenSpaceDenoiser::EShadowRequirements::Bailout;
+							if (DenoiserMode != 0 && LightRequiresDenosier(LightSceneInfo))
+							{
+								DenoiserRequirements = DenoiserToUse->GetShadowRequirements(View, LightSceneInfo, RayTracingConfig);
+							}
+
+							FRDGTextureRef HairCategorizationTexture = GetHairStrandsCategorizationTexture(ViewIndex, GraphBuilder);
+
 							RenderRayTracingShadows(
 								GraphBuilder,
 								SceneTextures,
 								View,
 								LightSceneInfo,
 								RayTracingConfig,
-								IScreenSpaceDenoiser::EShadowRequirements::Bailout,
-								true,
+								DenoiserRequirements,
+								false,
 								HairCategorizationTexture,
-								&SubPixelShadowMask,
-								&SubPixelRayHitDistance);
-							GraphBuilder.QueueTextureExtraction(SubPixelShadowMask, &ScreenShadowMaskSubPixelTexture);
+								RayTracingShadowMaskUAV,
+								RayHitDistanceUAV);
+
+							if (DenoiserRequirements != IScreenSpaceDenoiser::EShadowRequirements::Bailout)
+							{
+								TStaticArray<IScreenSpaceDenoiser::FShadowVisibilityParameters, IScreenSpaceDenoiser::kMaxBatchSize> InputParameters;
+								TStaticArray<IScreenSpaceDenoiser::FShadowVisibilityOutputs, IScreenSpaceDenoiser::kMaxBatchSize> Outputs;
+
+								InputParameters[0].InputTextures.Mask = RayTracingShadowMaskTexture;
+								InputParameters[0].InputTextures.ClosestOccluder = RayDistanceTexture;
+								InputParameters[0].LightSceneInfo = &LightSceneInfo;
+								InputParameters[0].RayTracingConfig = RayTracingConfig;
+
+								int32 InputParameterCount = 1;
+
+								RDG_EVENT_SCOPE(GraphBuilder, "%s%s(Shadow BatchSize=%d) %dx%d",
+									DenoiserToUse != DefaultDenoiser ? TEXT("ThirdParty ") : TEXT(""),
+									DenoiserToUse->GetDebugName(),
+									InputParameterCount,
+									View.ViewRect.Width(), View.ViewRect.Height());
+
+								DenoiserToUse->DenoiseShadowVisibilityMasks(
+									GraphBuilder,
+									View,
+									&View.PrevViewInfo,
+									SceneTextures,
+									InputParameters,
+									InputParameterCount,
+									Outputs);
+
+								GraphBuilder.QueueTextureExtraction(Outputs[0].Mask, &ScreenShadowMaskTexture);
+							}
+							else
+							{
+									GraphBuilder.QueueTextureExtraction(RayTracingShadowMaskTexture, &ScreenShadowMaskTexture);
+							}
+
+							// Ray trace the shadow cast by opaque geometries on to hair strands geometries
+							// Note: No denoiser is required on this output, as the hair strands are geometrically noisy, and so hard to denoise
+							if (HasHairStrandsClusters(ViewIndex))
+							{
+								RenderRayTracingShadows(
+									GraphBuilder,
+									SceneTextures,
+									View,
+									LightSceneInfo,
+									RayTracingConfig,
+									IScreenSpaceDenoiser::EShadowRequirements::Bailout,
+									true,
+									HairCategorizationTexture,
+									SubPixelRayTracingShadowMaskUAV,
+									SubPixelRayHitDistanceUAV);
+
+								GraphBuilder.QueueTextureExtraction(SubPixelRayTracingShadowMaskTexture, &ScreenShadowMaskSubPixelTexture);
+							}
 						}
 
 						GraphBuilder.Execute();
diff --git a/Engine/Source/Runtime/Renderer/Private/RayTracing/RayTracingReflections.cpp b/Engine/Source/Runtime/Renderer/Private/RayTracing/RayTracingReflections.cpp
--- a/Engine/Source/Runtime/Renderer/Private/RayTracing/RayTracingReflections.cpp
+++ b/Engine/Source/Runtime/Renderer/Private/RayTracing/RayTracingReflections.cpp
@@ -251,6 +251,9 @@ void FDeferredShadingSceneRenderer::RenderRayTracingReflections(
 	int32 SamplePerPixel,
 	int32 HeightFog,
 	float ResolutionFraction,
+	FRDGTextureUAV* ColorOutputUAV,
+	FRDGTextureUAV* RayHitDistanceOutputUAV,
+	FRDGTextureUAV* RayImaginaryDepthOutputUAV,
 	IScreenSpaceDenoiser::FReflectionsInputs* OutDenoiserInputs)
 #if RHI_RAYTRACING
 {
@@ -266,22 +269,6 @@ void FDeferredShadingSceneRenderer::RenderRayTracingReflections(
 	ensureMsgf(FComputeShaderUtils::kGolden2DGroupSize % UpscaleFactor == 0, TEXT("Reflection ray tracing will have uv misalignement."));
 	FIntPoint RayTracingResolution = FIntPoint::DivideAndRoundUp(View.ViewRect.Size(), UpscaleFactor);
 
-	{
-		FRDGTextureDesc Desc = FRDGTextureDesc::Create2DDesc(
-			SceneTextures.SceneDepthBuffer->Desc.Extent / UpscaleFactor,
-			PF_FloatRGBA,
-			FClearValueBinding::None,
-			/* InFlags = */ TexCreate_None,
-			/* InTargetableFlags = */ TexCreate_ShaderResource | TexCreate_RenderTargetable | TexCreate_UAV,
-			/* bInForceSeparateTargetAndShaderResource = */ false);
-
-		OutDenoiserInputs->Color = GraphBuilder.CreateTexture(Desc, TEXT("RayTracingReflections"));
-		
-		Desc.Format = PF_R16F;
-		OutDenoiserInputs->RayHitDistance = GraphBuilder.CreateTexture(Desc, TEXT("RayTracingReflectionsHitDistance"));
-		OutDenoiserInputs->RayImaginaryDepth = GraphBuilder.CreateTexture(Desc, TEXT("RayTracingReflectionsImaginaryDepth"));
-	}
-
 	// When deferred materials are used, we need to dispatch the reflection shader twice:
 	// - First pass gathers reflected ray hit data and sorts it by hit shader ID.
 	// - Second pass re-traces the reflected ray and performs full shading.
@@ -346,9 +333,9 @@ void FDeferredShadingSceneRenderer::RenderRayTracingReflections(
 
 	CommonParameters.ReflectionStruct = CreateReflectionUniformBuffer(View, EUniformBufferUsage::UniformBuffer_SingleFrame);
 	CommonParameters.FogUniformParameters = CreateFogUniformBuffer(View, EUniformBufferUsage::UniformBuffer_SingleFrame);
-	CommonParameters.ColorOutput = GraphBuilder.CreateUAV(OutDenoiserInputs->Color);
-	CommonParameters.RayHitDistanceOutput = GraphBuilder.CreateUAV(OutDenoiserInputs->RayHitDistance);
-	CommonParameters.RayImaginaryDepthOutput = GraphBuilder.CreateUAV(OutDenoiserInputs->RayImaginaryDepth);
+	CommonParameters.ColorOutput = ColorOutputUAV;
+	CommonParameters.RayHitDistanceOutput = RayHitDistanceOutputUAV;
+	CommonParameters.RayImaginaryDepthOutput = RayImaginaryDepthOutputUAV;
 	CommonParameters.SortTileSize = SortTileSize;
 	CommonParameters.ReflectionCapture = View.ReflectionCaptureUniformBuffer;
 	CommonParameters.Forward = View.ForwardLightingResources->ForwardLightDataUniformBuffer;
diff --git a/Engine/Source/Runtime/Renderer/Private/RayTracing/RayTracingShadows.cpp b/Engine/Source/Runtime/Renderer/Private/RayTracing/RayTracingShadows.cpp
--- a/Engine/Source/Runtime/Renderer/Private/RayTracing/RayTracingShadows.cpp
+++ b/Engine/Source/Runtime/Renderer/Private/RayTracing/RayTracingShadows.cpp
@@ -150,8 +150,8 @@ void FDeferredShadingSceneRenderer::RenderRayTracingShadows(
 	const IScreenSpaceDenoiser::EShadowRequirements DenoiserRequirements,
 	const bool bSubPixelShadowMask,
 	FRDGTextureRef HairCategorizationTexture,
-	FRDGTextureRef* OutShadowMask,
-	FRDGTextureRef* OutRayHitDistance)
+	FRDGTextureUAV* OutShadowMaskUAV,
+	FRDGTextureUAV* OutRayHitDistanceUAV)
 #if RHI_RAYTRACING
 {
 	FLightSceneProxy* LightSceneProxy = LightSceneInfo.Proxy;
@@ -186,7 +186,10 @@ void FDeferredShadingSceneRenderer::RenderRayTracingShadows(
 	FIntPoint PixelOffset = { 0, 0 };
 
 	// whether to clip the dispatch to a subrect, requires that denoising doesn't need the whole buffer
-	const bool bClipDispatch = DenoiserRequirements == IScreenSpaceDenoiser::EShadowRequirements::Bailout;
+	
+	//#dxr_todo: reenable Clip Dispatch with multiview support
+	//const bool bClipDispatch = DenoiserRequirements == IScreenSpaceDenoiser::EShadowRequirements::Bailout;
+	const bool bClipDispatch = false;//  DenoiserRequirements == IScreenSpaceDenoiser::EShadowRequirements::Bailout;
 
 	if (LightSceneProxy->GetScissorRect(ScissorRect, View, View.ViewRect))
 	{
@@ -208,8 +211,8 @@ void FDeferredShadingSceneRenderer::RenderRayTracingShadows(
 	{
 		const bool bUseHairLighting = HairCategorizationTexture != nullptr;
 		FOcclusionRGS::FParameters* PassParameters = GraphBuilder.AllocParameters<FOcclusionRGS::FParameters>();
-		PassParameters->RWOcclusionMaskUAV = GraphBuilder.CreateUAV(FRDGTextureUAVDesc(ScreenShadowMaskTexture));
-		PassParameters->RWRayDistanceUAV = GraphBuilder.CreateUAV(FRDGTextureUAVDesc(RayDistanceTexture));
+		PassParameters->RWOcclusionMaskUAV = OutShadowMaskUAV;
+		PassParameters->RWRayDistanceUAV = OutRayHitDistanceUAV;
 		PassParameters->SamplesPerPixel = RayTracingConfig.RayCountPerPixel;
 		PassParameters->NormalBias = GetRaytracingMaxNormalBias();
 		PassParameters->LightingChannelMask = LightSceneProxy->GetLightingChannelMask();
@@ -290,9 +293,6 @@ void FDeferredShadingSceneRenderer::RenderRayTracingShadows(
 			}
 		});
 	}
-
-	*OutShadowMask = ScreenShadowMaskTexture;
-	*OutRayHitDistance = RayDistanceTexture;
 }
 #else // !RHI_RAYTRACING
 {
diff --git a/Engine/Source/Runtime/Renderer/Private/ReflectionEnvironment.cpp b/Engine/Source/Runtime/Renderer/Private/ReflectionEnvironment.cpp
--- a/Engine/Source/Runtime/Renderer/Private/ReflectionEnvironment.cpp
+++ b/Engine/Source/Runtime/Renderer/Private/ReflectionEnvironment.cpp
@@ -658,6 +658,31 @@ void FDeferredShadingSceneRenderer::RenderDeferredReflectionsAndSkyLighting(FRHI
 	FSceneTextureParameters SceneTextures;
 	SetupSceneTextureParameters(GraphBuilder, &SceneTextures);
 
+	IScreenSpaceDenoiser::FReflectionsInputs DenoiserInputs;
+	IScreenSpaceDenoiser::FReflectionsRayTracingConfig RayTracingConfig;
+	RayTracingConfig.ResolutionFraction = FMath::Clamp(CVarReflectionScreenPercentage.GetValueOnRenderThread() / 100.0f, 0.25f, 1.0f);
+	int32 UpscaleFactor = int32(1.0f / RayTracingConfig.ResolutionFraction);
+
+	{
+		FRDGTextureDesc Desc = FRDGTextureDesc::Create2DDesc(
+			SceneTextures.SceneDepthBuffer->Desc.Extent / UpscaleFactor,
+			PF_FloatRGBA,
+			FClearValueBinding::None,
+			/* InFlags = */ TexCreate_None,
+			/* InTargetableFlags = */ TexCreate_ShaderResource | TexCreate_RenderTargetable | TexCreate_UAV,
+			/* bInForceSeparateTargetAndShaderResource = */ false);
+
+		DenoiserInputs.Color = GraphBuilder.CreateTexture(Desc, TEXT("RayTracingReflections"));
+
+		Desc.Format = PF_R16F;
+		DenoiserInputs.RayHitDistance = GraphBuilder.CreateTexture(Desc, TEXT("RayTracingReflectionsHitDistance"));
+		DenoiserInputs.RayImaginaryDepth = GraphBuilder.CreateTexture(Desc, TEXT("RayTracingReflectionsImaginaryDepth"));
+	}
+
+	FRDGTextureUAV* ReflectionColorOutputUAV = GraphBuilder.CreateUAV(FRDGTextureUAVDesc(DenoiserInputs.Color));
+	FRDGTextureUAV* RayHitDistanceOutputUAV = GraphBuilder.CreateUAV(FRDGTextureUAVDesc(DenoiserInputs.RayHitDistance));
+	FRDGTextureUAV* RayImaginaryDepthOutputUAV = GraphBuilder.CreateUAV(FRDGTextureUAVDesc(DenoiserInputs.RayImaginaryDepth));
+
 	uint32 ViewIndex = 0;
 	for (FViewInfo& View : Views)
 	{
@@ -672,20 +697,17 @@ void FDeferredShadingSceneRenderer::RenderDeferredReflectionsAndSkyLighting(FRHI
 		{
 			int32 DenoiserMode = CVarUseReflectionDenoiser.GetValueOnRenderThread();
 			
+			RayTracingConfig.RayCountPerPixel = GRayTracingReflectionsSamplesPerPixel > -1 ? GRayTracingReflectionsSamplesPerPixel : View.FinalPostProcessSettings.RayTracingReflectionsSamplesPerPixel;
+			
 			bool bDenoise = false;
 			bool bTemporalFilter = false;
 
 			// Traces the reflections, either using screen space reflection, or ray tracing.
-			IScreenSpaceDenoiser::FReflectionsInputs DenoiserInputs;
-			IScreenSpaceDenoiser::FReflectionsRayTracingConfig RayTracingConfig;
 			if (bRayTracedReflections)
 			{
 				RDG_EVENT_SCOPE(GraphBuilder, "RayTracingReflections");
 				RDG_GPU_STAT_SCOPE(GraphBuilder, RayTracingReflections);
 
-				RayTracingConfig.ResolutionFraction = FMath::Clamp(CVarReflectionScreenPercentage.GetValueOnRenderThread() / 100.0f, 0.25f, 1.0f);
-				RayTracingConfig.RayCountPerPixel = GRayTracingReflectionsSamplesPerPixel > -1 ? GRayTracingReflectionsSamplesPerPixel : View.FinalPostProcessSettings.RayTracingReflectionsSamplesPerPixel;
-
 				bDenoise = DenoiserMode != 0;
 				
 				if (!bDenoise)
@@ -698,6 +720,9 @@ void FDeferredShadingSceneRenderer::RenderDeferredReflectionsAndSkyLighting(FRHI
 					SceneTextures,
 					View,
 					RayTracingConfig.RayCountPerPixel, GRayTracingReflectionsHeightFog, RayTracingConfig.ResolutionFraction,
+					ReflectionColorOutputUAV,
+					RayHitDistanceOutputUAV,
+					RayImaginaryDepthOutputUAV,
 					&DenoiserInputs);
 			}
 			else if (bScreenSpaceReflections)
