Split AO and sky lighting passes.

diff --git a/Engine/Shaders/Private/HairStrands/HairStrandsEnvironmentAO.usf b/Engine/Shaders/Private/HairStrands/HairStrandsEnvironmentAO.usf
new file mode 100644
--- /dev/null
+++ b/Engine/Shaders/Private/HairStrands/HairStrandsEnvironmentAO.usf
@@ -0,0 +1,142 @@
+// Copyright 1998-2019 Epic Games, Inc. All Rights Reserved.
+  
+#include "../Common.ush"
+#include "../MonteCarlo.ush"
+#include "../DeferredShadingCommon.ush"  
+#include "../SceneTextureParameters.ush"
+#include "HairStrandsCommon.ush"
+#include "HairStrandsVisibilityCommon.ush"
+#include "HairStrandsDeepVoxelCommon.ush"
+
+Texture2D<uint4> HairCategorizationTexture;
+
+uint  Voxel_ClusterId;
+float Voxel_TanConeAngle;
+float Voxel_DensityScale;
+float Voxel_DepthBiasScale;
+
+float AO_Power;
+float AO_Intensity;
+
+float3 Voxel_MinAABB;
+float3 Voxel_MaxAABB;
+uint   Voxel_Resolution;
+Texture3D<uint> Voxel_DensityTexture;
+
+float InternalComputeHairCountVoxel(float3 WorldPosition, float3 LightDirection, float3 InVoxelAABBMin, float3 InVoxelAABBMax, float InVoxelResolution, Texture3D<uint> InDensityTexture, float DensityScale, float CountThreshold, bool bUseConeTracing)
+{
+	if (bUseConeTracing)
+		return ComputeHairCountVoxelCone(
+			WorldPosition,
+			LightDirection,
+			InVoxelAABBMin,
+			InVoxelAABBMax,
+			InVoxelResolution,
+			InDensityTexture,
+			DensityScale,
+			CountThreshold,
+			Voxel_TanConeAngle);
+	else
+		return ComputeHairCountVoxel(
+			WorldPosition,
+			LightDirection,
+			InVoxelAABBMin,
+			InVoxelAABBMax,
+			InVoxelResolution,
+			InDensityTexture,
+			DensityScale,
+			CountThreshold);
+}
+
+float TraceAO(float3 WorldPosition, float3 WorldNormal, float4 SvPosition)
+{
+	const float3 VoxelSize = (Voxel_MaxAABB - Voxel_MinAABB) / Voxel_Resolution;
+	const float  VoxelMaxSize = max(VoxelSize.x, max(VoxelSize.y, VoxelSize.z));
+	const float3 VoxelOffset = VoxelMaxSize * GetHairVoxelJitter(SvPosition.xy, View.StateFrameIndexMod8);
+	
+	float3 UnoccludedN = 0;
+	uint TotalHairCount = 0;
+	// Poisson disk position http://developer.download.nvidia.com/whitepapers/2008/PCSS_Integration.pdf
+	float2 PoissonDisk[16] =
+	{
+		float2(-0.94201624, -0.39906216),
+		float2(0.94558609, -0.76890725),
+		float2(-0.094184101, -0.92938870),
+		float2(0.34495938, 0.29387760),
+		float2(-0.91588581, 0.45771432),
+		float2(-0.81544232, -0.87912464),
+		float2(-0.38277543, 0.27676845),
+		float2(0.97484398, 0.75648379),
+		float2(0.44323325, -0.97511554),
+		float2(0.53742981, -0.47373420),
+		float2(-0.26496911, -0.41893023),
+		float2(0.79197514, 0.19090188),
+		float2(-0.24188840, 0.99706507),
+		float2(-0.81409955, 0.91437590),
+		float2(0.19984126, 0.78641367),
+		float2(0.14383161, -0.14100790)
+	};
+	const uint SampleCount = 16;
+
+	float3 BentNormal = 0;
+	float Occlusion = 0;
+	for (uint i = 0; i < SampleCount; ++i)
+	{
+		float2 Rand2 = (PoissonDisk[i] + float2(1, 1)) * 0.5f;
+		Rand2 = frac(Rand2 + VoxelOffset.xy);
+		const float3 LocalL  = CosineSampleHemisphere(saturate(Rand2)).xyz;
+		const float3 SampleL = TangentToWorld(LocalL.xyz, WorldNormal);
+
+		// Depth bias
+		// Origin is shifted voxels away towards the light + a constant bias of the size of the voxel
+		const float3 SampleWorldPosition = WorldPosition + VoxelOffset /*+ VoxelMaxSize*Voxel_DepthBiasScale*SampleL*/;;
+
+		// Compute the number of hair count between light & shading point
+		const float HairCount = InternalComputeHairCountVoxel(SampleWorldPosition, SampleL, Voxel_MinAABB, Voxel_MaxAABB, Voxel_Resolution, Voxel_DensityTexture, Voxel_DensityScale, 0, true);
+		BentNormal += SampleL * saturate(1 - HairCount);
+		Occlusion += saturate(HairCount);
+	}
+	Occlusion /= SampleCount;
+	BentNormal /= SampleCount;
+
+	// User adjust AO
+	const float AO = 1 - (1 - pow(saturate(1 - Occlusion), AO_Power)) * AO_Intensity;
+	return saturate(AO);
+}
+
+float3 TotWorldPosition(float2 UV, float Depth)
+{
+	float2 ScreenPosition = (UV - View.ScreenPositionScaleBias.wz) / View.ScreenPositionScaleBias.xy;
+	float4 HomogeneousWorldPosition = mul(float4(ScreenPosition * Depth, Depth, 1), View.ScreenToWorld);
+	float3 WorldPosition = HomogeneousWorldPosition.xyz / HomogeneousWorldPosition.w;
+
+	return WorldPosition;
+}
+
+void MainPS(
+	in float4 SvPosition : SV_Position,
+	out float4 OutAO : SV_Target0
+)
+{
+	const float2 BufferUV = SvPositionToBufferUV(SvPosition);
+	const uint3 PixelCoord = uint3(floor(SvPosition.xy), 0);
+
+	OutAO = 1;
+
+	const FGBufferData GBuffer = GetGBufferDataFromSceneTextures(BufferUV);
+	const float3 WorldPosition = TotWorldPosition(BufferUV, GBuffer.Depth);
+	const float3 WorldNormal = GBuffer.WorldNormal;
+
+	// Trace AO for all pixel which are partially covered or not covered at all.
+	const FCategorizationData CatData = DecodeCategorizationData(HairCategorizationTexture.Load(PixelCoord));
+	const bool bIsValid = (CatData.PixelCoverage > 0 && CatData.PixelCoverage < 1) || (GBuffer.ShadingModelID != SHADINGMODELID_UNLIT);
+
+	if (bIsValid)
+	{
+		OutAO = TraceAO(WorldPosition, WorldNormal, SvPosition);
+	}
+	else
+	{
+		discard;
+	}
+}
diff --git a/Engine/Shaders/Private/HairStrands/HairStrandsEnvironmentLighting.usf b/Engine/Shaders/Private/HairStrands/HairStrandsEnvironmentLighting.usf
--- a/Engine/Shaders/Private/HairStrands/HairStrandsEnvironmentLighting.usf
+++ b/Engine/Shaders/Private/HairStrands/HairStrandsEnvironmentLighting.usf
@@ -15,9 +15,6 @@
 #include "../LightGridCommon.ush"
 #include "../SceneTextureParameters.ush"
 
-#define RENDER_MODE_LIGHTING 0
-#define RENDER_MODE_AO 1
-
 #include "../HairBsdf.ush"
 #include "HairStrandsVisibilityCommon.ush"
 #include "HairStrandsDeepVoxelCommon.ush"
@@ -366,90 +363,10 @@ float3 ReflectionEnvironment(FGBufferData GBuffer, float2 BufferUV, float2 Scree
 	return -min(-Color.rgb, 0.0);
 }
 
-//////////////////////////////////////////////////////////////////////////////////////////////////////////////////////////
-// AO
-//////////////////////////////////////////////////////////////////////////////////////////////////////////////////////////
-
-float3 TraceAO(float3 WorldPosition, float3 WorldNormal, float4 SvPosition)
-{
-	const float3 VoxelSize = (Voxel_MaxAABB - Voxel_MinAABB) / Voxel_Resolution;
-	const float  VoxelMaxSize = max(VoxelSize.x, max(VoxelSize.y, VoxelSize.z));
-	const float3 VoxelOffset = VoxelMaxSize * GetHairVoxelJitter(SvPosition.xy, View.StateFrameIndexMod8);
-	
-	float3 UnoccludedN = 0;
-	uint TotalHairCount = 0;
-	// Poisson disk position http://developer.download.nvidia.com/whitepapers/2008/PCSS_Integration.pdf
-	float2 PoissonDisk[16] =
-	{
-		float2(-0.94201624, -0.39906216),
-		float2(0.94558609, -0.76890725),
-		float2(-0.094184101, -0.92938870),
-		float2(0.34495938, 0.29387760),
-		float2(-0.91588581, 0.45771432),
-		float2(-0.81544232, -0.87912464),
-		float2(-0.38277543, 0.27676845),
-		float2(0.97484398, 0.75648379),
-		float2(0.44323325, -0.97511554),
-		float2(0.53742981, -0.47373420),
-		float2(-0.26496911, -0.41893023),
-		float2(0.79197514, 0.19090188),
-		float2(-0.24188840, 0.99706507),
-		float2(-0.81409955, 0.91437590),
-		float2(0.19984126, 0.78641367),
-		float2(0.14383161, -0.14100790)
-	};
-	const uint SampleCount = 16;
-
-	float3 BentNormal = 0;
-	float Occlusion = 0;
-	for (uint i = 0; i < SampleCount; ++i)
-	{
-		float2 Rand2 = (PoissonDisk[i] + float2(1, 1)) * 0.5f;
-		Rand2 = frac(Rand2 + VoxelOffset.xy);
-		const float3 LocalL  = CosineSampleHemisphere(saturate(Rand2)).xyz;
-		const float3 SampleL = TangentToWorld(LocalL.xyz, WorldNormal);
-
-		// Depth bias
-		// Origin is shifted voxels away towards the light + a constant bias of the size of the voxel
-		const float3 SampleWorldPosition = WorldPosition + VoxelOffset /*+ VoxelMaxSize*Voxel_DepthBiasScale*SampleL*/;;
-
-		// Compute the number of hair count between light & shading point
-		const float HairCount = InternalComputeHairCountVoxel(SampleWorldPosition, SampleL, Voxel_MinAABB, Voxel_MaxAABB, Voxel_Resolution, Voxel_DensityTexture, Voxel_DensityScale, 0, true);
-		BentNormal += SampleL * saturate(1 - HairCount);
-		Occlusion += saturate(HairCount);
-	}
-	Occlusion /= SampleCount;
-	BentNormal /= SampleCount;
-
-	// User adjust AO
-	// abs() to prevent shader warning
-	float AO = 1 - Occlusion;
-	AO = 1 - (1 - pow(abs(AO), AO_Power)) * AO_Intensity;
-
-	return saturate(AO);
-}
-
-//////////////////////////////////////////////////////////////////////////////////////////////////////////////////////////
-//////////////////////////////////////////////////////////////////////////////////////////////////////////////////////////
-//////////////////////////////////////////////////////////////////////////////////////////////////////////////////////////
-float3 TotWorldPosition(float2 UV, float Depth)
-{
-	float2 ScreenPosition = (UV - View.ScreenPositionScaleBias.wz) / View.ScreenPositionScaleBias.xy;
-	float4 HomogeneousWorldPosition = mul(float4(ScreenPosition * Depth, Depth, 1), View.ScreenToWorld);
-	float3 WorldPosition = HomogeneousWorldPosition.xyz / HomogeneousWorldPosition.w;
-
-	return WorldPosition;
-}
-
 void MainPS(
 	in float4 SvPosition : SV_Position,
-#if PERMUTATION_RENDER_MODE == RENDER_MODE_LIGHTING
 	out float4 OutColor : SV_Target0,
 	out float4 OutColorSubPixel : SV_Target1
-#endif
-#if PERMUTATION_RENDER_MODE == RENDER_MODE_AO
-	out float4 OutAO : SV_Target0
-#endif
 )
 {
 	const float2 BufferUV = SvPositionToBufferUV(SvPosition);
@@ -458,7 +375,6 @@ void MainPS(
 
 	FNodeDesc NodeDesc = DecodeNodeDesc(HairVisibilityNodeOffsetAndCount.Load(PixelCoord));
 
-#if PERMUTATION_RENDER_MODE == RENDER_MODE_LIGHTING
 	OutColor = 0.0f;
 	OutColorSubPixel = 0.0f;
 
@@ -542,50 +458,4 @@ void MainPS(
 			}
 		}
 	}
-#endif
-
-
-#if PERMUTATION_RENDER_MODE == RENDER_MODE_AO
-	OutAO = 1;
-
-	const FGBufferData GBuffer = GetGBufferDataFromSceneTextures(BufferUV);
-	const float3 WorldPosition = TotWorldPosition(BufferUV, GBuffer.Depth);
-	const float3 WorldNormal = GBuffer.WorldNormal;
-
-	// Trace AO for all pixel which are partially covered or not covered at all.
-	const FCategorizationData CatData = DecodeCategorizationData(HairCategorizationTexture.Load(PixelCoord));
-	bool bIsValid = false;
-	bool bIsClusterCompatible = false;
-	if (CatData.PixelCoverage > 0 && CatData.PixelCoverage < 1)
-	{
-		LOOP
-		for (uint SampleIt = 0; SampleIt < NodeDesc.Count; SampleIt++)
-		{
-			const uint LocalOffset = NodeDesc.Offset + SampleIt;
-			const FNodeData NodeData = HairVisibilityNodeData[LocalOffset];
-			float			SampleDepth;
-			FVisibilityData SampleVisData;
-			float3			SampleBaseColor;
-			float			SampleRoughness;
-			float			SampleSpecular;
-			DecodeNodeData(NodeData, SampleDepth, SampleVisData, SampleBaseColor, SampleRoughness, SampleSpecular);
-			bIsClusterCompatible = bIsClusterCompatible || SampleVisData.ClusterId == Voxel_ClusterId;
-		}
-		bIsValid = true;
-	}
-	else
-	{
-		bIsValid = GBuffer.ShadingModelID != SHADINGMODELID_UNLIT;
-	}
-
-	if (bIsValid && bIsClusterCompatible)
-	{
-		const float AO = TraceAO(WorldPosition, WorldNormal, SvPosition);
-		OutAO = AO;
-	}
-	else
-	{
-		discard;
-	}
-#endif
 }
diff --git a/Engine/Source/Runtime/Renderer/Private/HairStrands/HairStrandsEnvironment.cpp b/Engine/Source/Runtime/Renderer/Private/HairStrands/HairStrandsEnvironment.cpp
--- a/Engine/Source/Runtime/Renderer/Private/HairStrands/HairStrandsEnvironment.cpp
+++ b/Engine/Source/Runtime/Renderer/Private/HairStrands/HairStrandsEnvironment.cpp
@@ -36,19 +36,107 @@ static bool GetHairStrandsSkyLightingEnable() { return GHairSkylightingEnable >
 static bool GetHairStrandsSkyAOEnable() { return GHairSkyAOEnable > 0; }
 static float GetHairStrandsSkyLightingConeAngle() { return FMath::Max(0.f, GHairSkylightingConeAngle); }
 
+DECLARE_GPU_STAT_NAMED(HairStrandsReflectionEnvironment, TEXT("Hair Strands Reflection Environment"));
+
 ///////////////////////////////////////////////////////////////////////////////////////////////////
+// AO
+class FHairEnvironmentAO : public FGlobalShader
+{
+	DECLARE_GLOBAL_SHADER(FHairEnvironmentAO);
+	SHADER_USE_PARAMETER_STRUCT(FHairEnvironmentAO, FGlobalShader)
 
-DECLARE_GPU_STAT_NAMED(HairStrandsReflectionEnvironment, TEXT("Hair Strands Reflection Environment"));
+	static bool ShouldCompilePermutation(const FGlobalShaderPermutationParameters& Parameters)
+	{
+		return IsHairStrandsSupported(Parameters.Platform);
+	}
+
+	BEGIN_SHADER_PARAMETER_STRUCT(FParameters, )
 
+		SHADER_PARAMETER(FVector, Voxel_MinAABB)
+		SHADER_PARAMETER(uint32, Voxel_ClusterId)
+		SHADER_PARAMETER(FVector, Voxel_MaxAABB)
+		SHADER_PARAMETER(uint32, Voxel_Resolution)
+		SHADER_PARAMETER(float, Voxel_DensityScale)
+		SHADER_PARAMETER(float, Voxel_DepthBiasScale)
+		SHADER_PARAMETER(float, Voxel_TanConeAngle)
+		SHADER_PARAMETER(float, AO_Power)	
+		SHADER_PARAMETER(float, AO_Intensity)
+
+		SHADER_PARAMETER_STRUCT_INCLUDE(FSceneTextureParameters, SceneTextures)
+		SHADER_PARAMETER_STRUCT_INCLUDE(FSceneTextureSamplerParameters, SceneTextureSamplers)
+
+		SHADER_PARAMETER_RDG_TEXTURE(Texture2D, HairCategorizationTexture)
+		SHADER_PARAMETER_RDG_TEXTURE(Texture3D, Voxel_DensityTexture)
+
+		SHADER_PARAMETER_STRUCT_REF(FViewUniformShaderParameters, ViewUniformBuffer)
+
+		RENDER_TARGET_BINDING_SLOTS()
+	END_SHADER_PARAMETER_STRUCT()
+};
+
+IMPLEMENT_GLOBAL_SHADER(FHairEnvironmentAO, "/Engine/Private/HairStrands/HairStrandsEnvironmentAO.usf", "MainPS", SF_Pixel);
+
+static void AddHairStrandsEnvironmentAOPass(
+	FRDGBuilder& GraphBuilder,
+	const FViewInfo& View,
+	const FHairStrandsVisibilityData& VisibilityData,
+	const FHairStrandsClusterData& ClusterData,
+	FRDGTextureRef Output)
+{
+	FSceneTextureParameters SceneTextures;
+	SetupSceneTextureParameters(GraphBuilder, &SceneTextures);
+
+	FHairEnvironmentAO::FParameters* PassParameters = GraphBuilder.AllocParameters<FHairEnvironmentAO::FParameters>();
+	PassParameters->Voxel_ClusterId = ClusterData.ClusterId;
+	PassParameters->Voxel_MinAABB = ClusterData.GetMinBound();
+	PassParameters->Voxel_MaxAABB = ClusterData.GetMaxBound();
+	PassParameters->Voxel_Resolution = ClusterData.GetResolution();
+	PassParameters->Voxel_DensityTexture = GraphBuilder.RegisterExternalTexture(ClusterData.VoxelResources.DensityTexture);
+	PassParameters->Voxel_DensityScale = GetHairStrandsVoxelizationDensityScale();
+	PassParameters->Voxel_DepthBiasScale = GetHairStrandsVoxelizationDepthBiasScale();
+	PassParameters->Voxel_TanConeAngle = FMath::Tan(FMath::DegreesToRadians(GetHairStrandsSkyLightingConeAngle()));
+	PassParameters->SceneTextures = SceneTextures;
+	SetupSceneTextureSamplers(&PassParameters->SceneTextureSamplers);
+
+	PassParameters->ViewUniformBuffer = View.ViewUniformBuffer;
+	PassParameters->HairCategorizationTexture = GraphBuilder.RegisterExternalTexture(VisibilityData.CategorizationTexture);
+	const FFinalPostProcessSettings& Settings = View.FinalPostProcessSettings;
+	PassParameters->AO_Power = Settings.AmbientOcclusionPower;
+	PassParameters->AO_Intensity = Settings.AmbientOcclusionIntensity;
+
+	check(Output);
+	PassParameters->RenderTargets[0] = FRenderTargetBinding(Output, ERenderTargetLoadAction::ELoad);
+
+	TShaderMapRef<FHairEnvironmentAO> PixelShader(View.ShaderMap);
+	ClearUnusedGraphResources(*PixelShader, PassParameters);
+
+	GraphBuilder.AddPass(
+		RDG_EVENT_NAME("HairStrandsAO %dx%d", View.ViewRect.Width(), View.ViewRect.Height()),
+		PassParameters,
+		ERDGPassFlags::Raster,
+		[PassParameters, &View, PixelShader](FRHICommandList& InRHICmdList)
+	{
+		InRHICmdList.SetViewport(View.ViewRect.Min.X, View.ViewRect.Min.Y, 0.0f, View.ViewRect.Max.X, View.ViewRect.Max.Y, 1.0f);
+
+		FGraphicsPipelineStateInitializer GraphicsPSOInit;
+		FPixelShaderUtils::InitFullscreenPipelineState(InRHICmdList, View.ShaderMap, *PixelShader, GraphicsPSOInit);
+		GraphicsPSOInit.BlendState = TStaticBlendState<CW_RGBA, BO_Min, BF_SourceColor, BF_DestColor, BO_Add, BF_Zero, BF_DestAlpha>::GetRHI();
+		SetGraphicsPipelineState(InRHICmdList, GraphicsPSOInit);
+		SetShaderParameters(InRHICmdList, *PixelShader, PixelShader->GetPixelShader(), *PassParameters);
+		FPixelShaderUtils::DrawFullscreenTriangle(InRHICmdList);
+	});
+}
+
+
+///////////////////////////////////////////////////////////////////////////////////////////////////
+// Lighting PS
 class FHairEnvironmentLightingPS : public FGlobalShader
 {
 	DECLARE_GLOBAL_SHADER(FHairEnvironmentLightingPS);
 	SHADER_USE_PARAMETER_STRUCT(FHairEnvironmentLightingPS, FGlobalShader)
 
-
-	class FRenderMode : SHADER_PERMUTATION_INT("PERMUTATION_RENDER_MODE", 2);
 	class FPerSample : SHADER_PERMUTATION_INT("PERMUTATION_PER_SAMPLE", 2);
-	using FPermutationDomain = TShaderPermutationDomain<FRenderMode, FPerSample>;
+	using FPermutationDomain = TShaderPermutationDomain<FPerSample>;
 	static bool ShouldCompilePermutation(const FGlobalShaderPermutationParameters& Parameters)
 	{
 		return IsHairStrandsSupported(Parameters.Platform);
@@ -64,9 +152,6 @@ class FHairEnvironmentLightingPS : public FGlobalShader
 		SHADER_PARAMETER(float, Voxel_DepthBiasScale)
 		SHADER_PARAMETER(float, Voxel_TanConeAngle)
 
-		SHADER_PARAMETER(float, AO_Power)
-		SHADER_PARAMETER(float, AO_Intensity)
-
 		SHADER_PARAMETER_TEXTURE(Texture2D, PreIntegratedGF)
 		SHADER_PARAMETER_SAMPLER(SamplerState, PreIntegratedGFSampler)
 
@@ -89,21 +174,14 @@ class FHairEnvironmentLightingPS : public FGlobalShader
 		SHADER_PARAMETER_STRUCT_REF(FForwardLightData, ForwardLightData)
 
 		RENDER_TARGET_BINDING_SLOTS()
-		END_SHADER_PARAMETER_STRUCT()
-}; // FHairEnvironmentLightingPS
+	END_SHADER_PARAMETER_STRUCT()
+};
 
 IMPLEMENT_GLOBAL_SHADER(FHairEnvironmentLightingPS, "/Engine/Private/HairStrands/HairStrandsEnvironmentLighting.usf", "MainPS", SF_Pixel);
 
-enum class EEnvRenderMode : uint8
-{
-	Lighting,
-	AO
-};
-
 static void AddHairStrandsEnvironmentPass(
 	FRDGBuilder& GraphBuilder,
 	const FViewInfo& View,
-	EEnvRenderMode RenderMode,
 	const FHairStrandsVisibilityData& VisibilityData,
 	const FHairStrandsClusterData& ClusterData,
 	FRDGTextureRef Output0,
@@ -152,54 +230,30 @@ static void AddHairStrandsEnvironmentPass(
 	PassParameters->HairVisibilityNodeOffsetAndCount = VisibilityData.NodeIndex->GetRenderTargetItem().ShaderResourceTexture;
 	PassParameters->HairVisibilityNodeData = VisibilityData.NodeDataSRV;
 
-	PassParameters->AO_Power = 0;
-	PassParameters->AO_Intensity = 0;
-	if (RenderMode == EEnvRenderMode::AO)
-	{
-		const FFinalPostProcessSettings& Settings = View.FinalPostProcessSettings;
-		PassParameters->AO_Power = Settings.AmbientOcclusionPower;
-		PassParameters->AO_Intensity = Settings.AmbientOcclusionIntensity;
-	}
-
-	check(Output0);
+	check(Output0); check(Output1);
 	PassParameters->RenderTargets[0] = FRenderTargetBinding(Output0, ERenderTargetLoadAction::ELoad);
-	if (Output1)
-	{
-		PassParameters->RenderTargets[1] = FRenderTargetBinding(Output1, ERenderTargetLoadAction::ELoad);
-	}
+	PassParameters->RenderTargets[1] = FRenderTargetBinding(Output1, ERenderTargetLoadAction::ELoad);
 
 	FHairEnvironmentLightingPS::FPermutationDomain PermutationVector;
-	PermutationVector.Set<FHairEnvironmentLightingPS::FRenderMode>(RenderMode == EEnvRenderMode::AO ? 1 : 0);
 	PermutationVector.Set<FHairEnvironmentLightingPS::FPerSample>(GHairSkylightingPerSample > 0 ? 1 : 0);
 	TShaderMapRef<FHairEnvironmentLightingPS> PixelShader(View.ShaderMap, PermutationVector);
 	ClearUnusedGraphResources(*PixelShader, PassParameters);
 
 	GraphBuilder.AddPass(
-		(RenderMode == EEnvRenderMode::Lighting) ? 
-		RDG_EVENT_NAME("HairStrandsEnvironment %dx%d", View.ViewRect.Width(), View.ViewRect.Height()) : 
-		RDG_EVENT_NAME("HairStrandsAO %dx%d", View.ViewRect.Width(), View.ViewRect.Height()),
+		RDG_EVENT_NAME("HairStrandsEnvironment %dx%d", View.ViewRect.Width(), View.ViewRect.Height()),
 		PassParameters,
 		ERDGPassFlags::Raster,
-		[PassParameters, &View, PixelShader, RenderMode](FRHICommandList& InRHICmdList)
+		[PassParameters, &View, PixelShader](FRHICommandList& InRHICmdList)
 	{
 		InRHICmdList.SetViewport(View.ViewRect.Min.X, View.ViewRect.Min.Y, 0.0f, View.ViewRect.Max.X, View.ViewRect.Max.Y, 1.0f);
 
 		FGraphicsPipelineStateInitializer GraphicsPSOInit;
 		FPixelShaderUtils::InitFullscreenPipelineState(InRHICmdList, View.ShaderMap, *PixelShader, GraphicsPSOInit);
 
-		if (RenderMode == EEnvRenderMode::AO)
-		{
-			GraphicsPSOInit.BlendState = TStaticBlendState<
-			CW_RGBA, BO_Min, BF_SourceColor, BF_DestColor, BO_Add, BF_Zero, BF_DestAlpha
-			>::GetRHI();
-		}
-		else //if (RenderMode == EEnvRenderMode::Lighting)
-		{
-			GraphicsPSOInit.BlendState = TStaticBlendState<
+		GraphicsPSOInit.BlendState = TStaticBlendState<
 			CW_RGBA, BO_Add, BF_One, BF_One, BO_Add, BF_One, BF_One,
 			CW_RGBA, BO_Add, BF_One, BF_One, BO_Max, BF_SourceAlpha, BF_DestAlpha
 			>::GetRHI();
-		}
 
 		SetGraphicsPipelineState(InRHICmdList, GraphicsPSOInit);
 		SetShaderParameters(InRHICmdList, *PixelShader, PixelShader->GetPixelShader(), *PassParameters);
@@ -233,7 +287,7 @@ void RenderHairStrandsEnvironmentLighting(
 		// @hair_todo: 
 		// * Add support for : BentNormal, global AO, SSR, DFSO
 		// * Add local reflection probe, current take into account only the sky lighting
-		AddHairStrandsEnvironmentPass(GraphBuilder, View, EEnvRenderMode::Lighting, VisibilityData, ClusterData, SceneColorTexture, SceneColorSubPixelTexture);
+		AddHairStrandsEnvironmentPass(GraphBuilder, View, VisibilityData, ClusterData, SceneColorTexture, SceneColorSubPixelTexture);
 	}
 }
 
@@ -264,7 +318,7 @@ void RenderHairStrandsAmbientOcclusion(
 		FRDGTextureRef AOTexture = InAOTexture ? GraphBuilder.RegisterExternalTexture(InAOTexture, TEXT("AOTexture")) : nullptr;
 		for (const FHairStrandsClusterData& ClusterData : HairDatas->HairClusterPerViews.Views[ViewIndex].Datas)
 		{
-			AddHairStrandsEnvironmentPass(GraphBuilder, View, EEnvRenderMode::AO, VisibilityData, ClusterData, AOTexture, nullptr);
+			AddHairStrandsEnvironmentAOPass(GraphBuilder, View, VisibilityData, ClusterData, AOTexture);
 
 		}
 		GraphBuilder.Execute();
