Fixed compiler warning due to implicit conversion from int to bool

diff --git a/Engine/Source/Runtime/Renderer/Private/RayTracing/RayTracingShadows.cpp b/Engine/Source/Runtime/Renderer/Private/RayTracing/RayTracingShadows.cpp
--- a/Engine/Source/Runtime/Renderer/Private/RayTracing/RayTracingShadows.cpp
+++ b/Engine/Source/Runtime/Renderer/Private/RayTracing/RayTracingShadows.cpp
@@ -131,7 +131,7 @@ void FDeferredShadingSceneRenderer::PrepareRayTracingShadows(const FViewInfo& Vi
 			PermutationVector.Set<FOcclusionRGS::FDenoiserOutputDim>((int32)DenoiserRequirement);
 			PermutationVector.Set<FOcclusionRGS::FEnableTwoSidedGeometryDim>(CVarRayTracingShadowsEnableTwoSidedGeometry.GetValueOnRenderThread() != 0);
 			PermutationVector.Set<FOcclusionRGS::FHairLighting>(HairLighting);
-			PermutationVector.Set<FOcclusionRGS::FEnableMultipleSamplesPerPixel>(MultiSPP);
+			PermutationVector.Set<FOcclusionRGS::FEnableMultipleSamplesPerPixel>(MultiSPP != 0);
 
 			TShaderMapRef<FOcclusionRGS> RayGenerationShader(View.ShaderMap, PermutationVector);
 			OutRayGenShaders.Add(RayGenerationShader->GetRayTracingShader());
