Address some performance concerns with virtual texture on PC
Use event queries instead of the frame counting fence to ensure that we don't needlessly sync GPU and CPU
Reduce max load of virtual texure updates
Add some more CVS stats to help track down any problems in the live environment

diff --git a/Engine/Source/Runtime/Renderer/Private/VT/VirtualTextureFeedback.cpp b/Engine/Source/Runtime/Renderer/Private/VT/VirtualTextureFeedback.cpp
--- a/Engine/Source/Runtime/Renderer/Private/VT/VirtualTextureFeedback.cpp
+++ b/Engine/Source/Runtime/Renderer/Private/VT/VirtualTextureFeedback.cpp
@@ -1,9 +1,132 @@
 // Copyright 1998-2019 Epic Games, Inc. All Rights Reserved.
 
 #include "VirtualTextureFeedback.h"
-#include "VisualizeTexture.h"
 
 #include "ClearQuad.h"
+#include "VisualizeTexture.h"
+
+#if PLATFORM_WINDOWS
+// Temporary use of Querys until RHI has a good fence on D3D11
+#define USE_RHI_FENCES 1
+#else
+#define USE_RHI_FENCES 1
+#endif
+
+#if USE_RHI_FENCES
+class FFeedbackFences
+{
+public:
+	FGPUFenceRHIRef Fences[FVirtualTextureFeedback::TargetCapacity];
+
+	void Init(FRHICommandListImmediate& RHICmdList)
+	{
+	}
+
+	void ReleaseAll()
+	{
+		for (int i = 0; i < FVirtualTextureFeedback::TargetCapacity; ++i)
+		{
+			Fences[i].SafeRelease();
+		}
+	}
+
+	void Allocate(FRHICommandListImmediate& RHICmdList, int32 Index)
+	{
+		if (!Fences[Index])
+		{
+			Fences[Index] = RHICmdList.CreateGPUFence(FName(""));
+		}
+		Fences[Index]->Clear();
+	}
+
+	void Write(FRHICommandListImmediate& RHICmdList, int32 Index)
+	{
+		RHICmdList.WriteGPUFence(Fences[Index]);
+	}
+
+	bool Poll(FRHICommandListImmediate& RHICmdList, int32 Index)
+	{
+		return Fences[Index]->Poll();
+	}
+
+	FGPUFenceRHIRef GetMapFence(int32 Index)
+	{
+		return Fences[Index];
+	}
+
+	void Release(int32 Index)
+	{
+		Fences[Index].SafeRelease();
+	}
+};
+#else
+class FFeedbackFences
+{
+public:
+	FRenderQueryPoolRHIRef FenceQueryPool;
+	FGPUFenceRHIRef DummyFence;
+	FRHIPooledRenderQuery Fences[FVirtualTextureFeedback::TargetCapacity];
+
+	void Init(FRHICommandListImmediate& RHICmdList)
+	{
+		if (!FenceQueryPool.IsValid())
+		{
+			FenceQueryPool = RHICreateRenderQueryPool(RQT_AbsoluteTime, FVirtualTextureFeedback::TargetCapacity);
+		}
+
+		if (!DummyFence.IsValid())
+		{
+			DummyFence = RHICmdList.CreateGPUFence(FName());
+			RHICmdList.WriteGPUFence(DummyFence);
+		}
+	}
+
+	void ReleaseAll()
+	{
+		for (int i = 0; i < FVirtualTextureFeedback::TargetCapacity; ++i)
+		{
+			if (Fences[i].IsValid())
+			{
+				Fences[i].ReleaseQuery();
+			}
+		}
+
+		DummyFence.SafeRelease();
+
+		FenceQueryPool.SafeRelease();
+	}
+
+	void Allocate(FRHICommandListImmediate& RHICmdList, int32 Index)
+	{
+		if (!Fences[Index].IsValid())
+		{
+			Fences[Index] = FenceQueryPool->AllocateQuery();
+		}
+	}
+	
+	void Write(FRHICommandListImmediate& RHICmdList, int32 Index)
+	{
+		RHICmdList.EndRenderQuery(Fences[Index].GetQuery());
+	}
+
+	bool Poll(FRHICommandListImmediate& RHICmdList, int32 Index)
+	{
+		uint64 Dummy;
+		return RHICmdList.GetRenderQueryResult(Fences[Index].GetQuery(), Dummy, false);
+	}
+
+	FGPUFenceRHIRef GetMapFence(int32 Index)
+	{
+		return DummyFence;
+	}
+
+	void Release(int32 Index)
+	{
+		Fences[Index].ReleaseQuery();
+	}
+};
+#endif
+
 
 FVirtualTextureFeedback::FVirtualTextureFeedback()
 	: Size( 0, 0 )
@@ -11,19 +134,30 @@ FVirtualTextureFeedback::FVirtualTextureFeedback()
 	, CPUReadIndex(0)
 	, PendingTargetCount(0)
 {
+	FeedBackFences = new FFeedbackFences;
+}
+
+FVirtualTextureFeedback::~FVirtualTextureFeedback()
+{
+	delete FeedBackFences;
 }
 
 void FVirtualTextureFeedback::ReleaseResources()
 {
 	GRenderTargetPool.FreeUnusedResource( FeedbackTextureGPU );
+	FeedbackTextureGPU.SafeRelease();
+	
 	for (int i = 0; i < TargetCapacity; ++i)
 	{
 		if (FeedbackTextureCPU[i].TextureCPU.IsValid())
 		{
 			GRenderTargetPool.FreeUnusedResource(FeedbackTextureCPU[i].TextureCPU);
+			FeedbackTextureCPU[i].TextureCPU.SafeRelease();
 		}
-		FeedbackTextureCPU[i].GPUFenceRHI.SafeRelease();
 	}
+
+	FeedBackFences->ReleaseAll();
+
 	CPUReadIndex = 0u;
 	GPUWriteIndex = 0u;
 	PendingTargetCount = 0u;
@@ -40,6 +174,8 @@ void FVirtualTextureFeedback::CreateResourceGPU( FRHICommandListImmediate& RHICm
 	const uint32 ClearValue[4] = { ~0u, ~0u, ~0u, ~0u };
 	ClearUAV( RHICmdList, FeedbackTextureGPU->GetRenderTargetItem(), ClearValue );
 	RHICmdList.TransitionResource( EResourceTransitionAccess::ERWNoBarrier, EResourceTransitionPipeline::EGfxToGfx, FeedbackTextureGPU->GetRenderTargetItem().UAV );
+
+	FeedBackFences->Init(RHICmdList);
 }
 
 void FVirtualTextureFeedback::MakeSnapshot(const FVirtualTextureFeedback& SnapshotSource)
@@ -61,17 +197,14 @@ void FVirtualTextureFeedback::TransferGPUToCPU( FRHICommandListImmediate& RHICmd
 	if (PendingTargetCount >= TargetCapacity)
 	{
 		// If we have too many pending transfers, start throwing away the oldest
-
-		// will need to allocate a new fence,
-		// since the previous fence will still be set on the old CopyToResolveTarget command (which we will not ignore/discard)
-		FeedBackItem& DiscardFeedbackEntry = FeedbackTextureCPU[CPUReadIndex];
-		DiscardFeedbackEntry.GPUFenceRHI.SafeRelease();
+		// We will need to allocate a new fence, since the previous fence will still be set on the old CopyToResolveTarget command (which we will now ignore/discard)
+		FeedBackFences->Release(CPUReadIndex);
 
 		--PendingTargetCount;
 		CPUReadIndex = (CPUReadIndex + 1) % TargetCapacity;
 	}
 
-	FeedBackItem& FeedbackEntryCPU = FeedbackTextureCPU[GPUWriteIndex];
+	FFeedBackItem& FeedbackEntryCPU = FeedbackTextureCPU[GPUWriteIndex];
 	FeedbackEntryCPU.Rect.Min = FIntPoint(Rect.Min.X / 16, Rect.Min.Y / 16);
 	FeedbackEntryCPU.Rect.Max = FIntPoint((Rect.Max.X + 15) / 16, (Rect.Max.Y + 15) / 16);
 
@@ -79,42 +212,40 @@ void FVirtualTextureFeedback::TransferGPUToCPU( FRHICommandListImmediate& RHICmd
 
 	const TCHAR* DebugNames[TargetCapacity] = { TEXT("VTFeedbackCPU_0") , TEXT("VTFeedbackCPU_1") , TEXT("VTFeedbackCPU_2") , TEXT("VTFeedbackCPU_3") };
 	GRenderTargetPool.FindFreeElement(RHICmdList, Desc, FeedbackEntryCPU.TextureCPU, DebugNames[GPUWriteIndex]);
-
-	if (!FeedbackEntryCPU.GPUFenceRHI)
-	{
-		FeedbackEntryCPU.GPUFenceRHI = RHICmdList.CreateGPUFence(DebugNames[GPUWriteIndex]);
-	}
+	
+	FeedBackFences->Allocate(RHICmdList, GPUWriteIndex);
 
 	// We only need to transfer 1 copy of the data, so restrict mask to the first active GPU.
 	FeedbackEntryCPU.GPUMask = FRHIGPUMask::FromIndex(RHICmdList.GetGPUMask().GetFirstIndex());
 	SCOPED_GPU_MASK(RHICmdList, FeedbackEntryCPU.GPUMask);
 
 	// Transfer memory GPU -> CPU
-	FeedbackEntryCPU.GPUFenceRHI->Clear();
 	RHICmdList.CopyToResolveTarget(
 		FeedbackTextureGPU->GetRenderTargetItem().TargetableTexture,
 		FeedbackEntryCPU.TextureCPU->GetRenderTargetItem().ShaderResourceTexture,
 		FResolveParams());
-	RHICmdList.WriteGPUFence(FeedbackEntryCPU.GPUFenceRHI);
 	
-	GRenderTargetPool.FreeUnusedResource( FeedbackTextureGPU );
+	FeedBackFences->Write(RHICmdList, GPUWriteIndex);
 
 	GPUWriteIndex = (GPUWriteIndex + 1) % TargetCapacity;
 	++PendingTargetCount;
 }
 
-bool FVirtualTextureFeedback::CanMap()
+bool FVirtualTextureFeedback::CanMap(FRHICommandListImmediate& RHICmdList)
 {
-	const FeedBackItem& FeedbackEntryCPU = FeedbackTextureCPU[CPUReadIndex];
-	return (PendingTargetCount > 0u && FeedbackEntryCPU.TextureCPU.IsValid() && FeedbackEntryCPU.GPUFenceRHI->Poll());
+	const FFeedBackItem& FeedbackEntryCPU = FeedbackTextureCPU[CPUReadIndex];
+	return (
+		PendingTargetCount > 0u &&
+		FeedbackEntryCPU.TextureCPU.IsValid() &&
+		FeedBackFences->Poll(RHICmdList, CPUReadIndex));
 }
 
 bool FVirtualTextureFeedback::Map(FRHICommandListImmediate& RHICmdList, MapResult& OutResult)
 {
-	const FeedBackItem& FeedbackEntryCPU = FeedbackTextureCPU[CPUReadIndex];
+	const FFeedBackItem& FeedbackEntryCPU = FeedbackTextureCPU[CPUReadIndex];
 	if (PendingTargetCount > 0u &&
 		FeedbackEntryCPU.TextureCPU.IsValid() &&
-		FeedbackEntryCPU.GPUFenceRHI->Poll(FeedbackEntryCPU.GPUMask))
+		FeedbackEntryCPU.GPUFenceRHI->Poll())
 	{
 		SCOPED_GPU_MASK(RHICmdList, FeedbackEntryCPU.GPUMask);
 
@@ -122,7 +253,7 @@ bool FVirtualTextureFeedback::Map(FRHICommandListImmediate& RHICmdList, MapResul
 		OutResult.Rect = FeedbackEntryCPU.Rect;
 
 		int32 LockHeight = 0;
-		RHICmdList.MapStagingSurface(FeedbackEntryCPU.TextureCPU->GetRenderTargetItem().ShaderResourceTexture, FeedbackEntryCPU.GPUFenceRHI, *(void**)&OutResult.Buffer, OutResult.Pitch, LockHeight);
+		RHICmdList.MapStagingSurface(FeedbackEntryCPU.TextureCPU->GetRenderTargetItem().ShaderResourceTexture, FeedBackFences->GetMapFence(CPUReadIndex), *(void**)&OutResult.Buffer, OutResult.Pitch, LockHeight);
 
 		check(PendingTargetCount > 0u);
 		--PendingTargetCount;
diff --git a/Engine/Source/Runtime/Renderer/Private/VT/VirtualTextureFeedback.h b/Engine/Source/Runtime/Renderer/Private/VT/VirtualTextureFeedback.h
--- a/Engine/Source/Runtime/Renderer/Private/VT/VirtualTextureFeedback.h
+++ b/Engine/Source/Runtime/Renderer/Private/VT/VirtualTextureFeedback.h
@@ -18,7 +18,7 @@ class FVirtualTextureFeedback
 {
 public:
 	FVirtualTextureFeedback();
-	~FVirtualTextureFeedback() {}
+	~FVirtualTextureFeedback();
 
 	static const uint32 TargetCapacity = 4u;
 
@@ -32,26 +32,29 @@ public:
 		int32 Pitch;
 	};
 
-	void			CreateResourceGPU( FRHICommandListImmediate& RHICmdList, FIntPoint InSize );
+	void			CreateResourceGPU(FRHICommandListImmediate& RHICmdList, FIntPoint InSize);
 	void			ReleaseResources();
 	void			MakeSnapshot(const FVirtualTextureFeedback& SnapshotSource);
 
-	void			TransferGPUToCPU( FRHICommandListImmediate& RHICmdList, FIntRect const& Rect);
+	void			TransferGPUToCPU(FRHICommandListImmediate& RHICmdList, FIntRect const& Rect);
 
-	bool			CanMap();
+	uint32			GetPendingTargetCount() const { return PendingTargetCount; }
+
+	bool			CanMap(FRHICommandListImmediate& RHICmdList);
 	bool			Map(FRHICommandListImmediate& RHICmdList, MapResult& OutResult);
 	void			Unmap(FRHICommandListImmediate& RHICmdList, int32 MapHandle);
 
 private:
-	struct FeedBackItem
+	struct FFeedBackItem
 	{
 		FIntRect Rect;
 		TRefCountPtr< IPooledRenderTarget > TextureCPU;
-		FGPUFenceRHIRef GPUFenceRHI;
 		FRHIGPUMask GPUMask;
 	};
 
-	FeedBackItem FeedbackTextureCPU[TargetCapacity];
+	FFeedBackItem FeedbackTextureCPU[TargetCapacity];
+
+	class FFeedbackFences* FeedBackFences;
 
 	FIntPoint Size;
 	uint32 GPUWriteIndex;
diff --git a/Engine/Source/Runtime/Renderer/Private/VT/VirtualTextureSystem.cpp b/Engine/Source/Runtime/Renderer/Private/VT/VirtualTextureSystem.cpp
--- a/Engine/Source/Runtime/Renderer/Private/VT/VirtualTextureSystem.cpp
+++ b/Engine/Source/Runtime/Renderer/Private/VT/VirtualTextureSystem.cpp
@@ -24,11 +24,14 @@ DECLARE_CYCLE_STAT(TEXT("Gather Requests"), STAT_ProcessRequests_Gather, STATGRO
 DECLARE_CYCLE_STAT(TEXT("Sort Requests"), STAT_ProcessRequests_Sort, STATGROUP_VirtualTexturing);
 DECLARE_CYCLE_STAT(TEXT("Submit Requests"), STAT_ProcessRequests_Submit, STATGROUP_VirtualTexturing);
 DECLARE_CYCLE_STAT(TEXT("Map Requests"), STAT_ProcessRequests_Map, STATGROUP_VirtualTexturing);
+DECLARE_CYCLE_STAT(TEXT("Map New VTs"), STAT_ProcessRequests_MapNew, STATGROUP_VirtualTexturing);
 DECLARE_CYCLE_STAT(TEXT("Finalize Requests"), STAT_ProcessRequests_Finalize, STATGROUP_VirtualTexturing);
 DECLARE_CYCLE_STAT(TEXT("Merge Unique Pages"), STAT_ProcessRequests_MergePages, STATGROUP_VirtualTexturing);
 DECLARE_CYCLE_STAT(TEXT("Merge Requests"), STAT_ProcessRequests_MergeRequests, STATGROUP_VirtualTexturing);
 DECLARE_CYCLE_STAT(TEXT("Submit Tasks"), STAT_ProcessRequests_SubmitTasks, STATGROUP_VirtualTexturing);
+DECLARE_CYCLE_STAT(TEXT("Wait Tasks"), STAT_ProcessRequests_WaitTasks, STATGROUP_VirtualTexturing);
 
+DECLARE_CYCLE_STAT(TEXT("Feedback Map"), STAT_FeedbackMap, STATGROUP_VirtualTexturing);
 DECLARE_CYCLE_STAT(TEXT("Feedback Analysis"), STAT_FeedbackAnalysis, STATGROUP_VirtualTexturing);
 DECLARE_CYCLE_STAT(TEXT("Page Table Updates"), STAT_PageTableUpdates, STATGROUP_VirtualTexturing);
 DECLARE_CYCLE_STAT(TEXT("Flush Cache"), STAT_FlushCache, STATGROUP_VirtualTexturing);
@@ -67,13 +70,13 @@ static TAutoConsoleVariable<int32> CVarVTVerbose(
 
 static TAutoConsoleVariable<int32> CVarVTNumFeedbackTasks(
 	TEXT("r.VT.NumFeedbackTasks"),
-	4,
+	1,
 	TEXT("Number of tasks to create to process virtual texture updates."),
 	ECVF_RenderThreadSafe
 );
 static TAutoConsoleVariable<int32> CVarVTNumGatherTasks(
 	TEXT("r.VT.NumGatherTasks"),
-	4,
+	1,
 	TEXT("Number of tasks to create to process virtual texture updates."),
 	ECVF_RenderThreadSafe
 );
@@ -908,7 +911,7 @@ void FVirtualTextureSystem::Update(FRHICommandListImmediate& RHICmdList, ERHIFea
 		// Only flush if we know that there is GPU feedback available to refill the visible data this frame
 		// This prevents bugs when low frame rate causes feedback buffer to stall so that the physical cache isn't filled immediately which causes visible glitching
 		FSceneRenderTargets& SceneContext = FSceneRenderTargets::Get(RHICmdList);
-		if (SceneContext.VirtualTextureFeedback.CanMap())
+		if (SceneContext.VirtualTextureFeedback.CanMap(RHICmdList))
 		{
 			// Each RVT will call FVirtualTextureSystem::FlushCache()
 			Scene->FlushDirtyRuntimeVirtualTextures();
@@ -930,13 +933,23 @@ void FVirtualTextureSystem::Update(FRHICommandListImmediate& RHICmdList, ERHIFea
 
 		if (CVarVTEnableFeedBack.GetValueOnRenderThread())
 		{
-			for (FeedbackBufferCount = 0; FeedbackBufferCount < FVirtualTextureFeedback::TargetCapacity; ++FeedbackBufferCount)
+			SCOPE_CYCLE_COUNTER(STAT_FeedbackMap);
+
+#if WITH_EDITOR
+			uint32 MaxFeedbackTargetCount = FVirtualTextureFeedback::TargetCapacity;
+#else
+			uint32 MaxFeedbackTargetCount = 1;
+#endif
+
+			for (FeedbackBufferCount = 0; FeedbackBufferCount < MaxFeedbackTargetCount; ++FeedbackBufferCount)
 			{
 				if (!SceneContext.VirtualTextureFeedback.Map(RHICmdList, MappedFeedbackBuffers[FeedbackBufferCount]))
 				{
 					break;
 				}
 			}
+
+			const int32 PendingFeedbackBuffers = SceneContext.VirtualTextureFeedback.GetPendingTargetCount();
 		}
 
 		// Create tasks to read all the buffers
@@ -951,7 +964,7 @@ void FVirtualTextureSystem::Update(FRHICommandListImmediate& RHICmdList, ERHIFea
 
 			// Give each task a section of a feedback buffer to analyze
 			//todo[vt]: For buffers of different sizes we will have different task payload sizes which is not efficient
-			const uint32 FeedbackTasksPerBuffer = MaxNumFeedbackTasks / FeedbackBufferCount;
+			const uint32 FeedbackTasksPerBuffer = FMath::Max(MaxNumFeedbackTasks / FeedbackBufferCount, 1u);
 			const uint32 FeedbackRowsPerTask = FMath::DivideAndRoundUp((uint32)FeedbackBuffer.Rect.Size().Y, FeedbackTasksPerBuffer);
 			const uint32 NumRows = (uint32)FeedbackBuffer.Rect.Size().Y;
 			
@@ -1002,6 +1015,8 @@ void FVirtualTextureSystem::Update(FRHICommandListImmediate& RHICmdList, ERHIFea
 			// Wait for them to complete
 			if (Tasks.Num() > 0)
 			{
+				SCOPE_CYCLE_COUNTER(STAT_ProcessRequests_WaitTasks);
+
 				FTaskGraphInterface::Get().WaitUntilTasksComplete(Tasks, ENamedThreads::GetRenderThread_Local());
 			}
 		}
@@ -1172,6 +1187,8 @@ void FVirtualTextureSystem::GatherRequests(FUniqueRequestList* MergedRequestList
 		// Wait for them to complete
 		if (Tasks.Num() > 0)
 		{
+			SCOPE_CYCLE_COUNTER(STAT_ProcessRequests_WaitTasks);
+
 			FTaskGraphInterface::Get().WaitUntilTasksComplete(Tasks, ENamedThreads::GetRenderThread_Local());
 		}
 	}
@@ -1955,6 +1972,8 @@ void FVirtualTextureSystem::SubmitRequests(FRHICommandListImmediate& RHICmdList,
 
 	// Map any resident tiles to newly allocated VTs
 	{
+		SCOPE_CYCLE_COUNTER(STAT_ProcessRequests_MapNew);
+
 		uint32 Index = 0u;
 		while (Index < (uint32)AllocatedVTsToMap.Num())
 		{
