Fix define setting where new use of queries for VT feedback was disabled

diff --git a/Engine/Source/Runtime/Renderer/Private/VT/VirtualTextureFeedback.cpp b/Engine/Source/Runtime/Renderer/Private/VT/VirtualTextureFeedback.cpp
--- a/Engine/Source/Runtime/Renderer/Private/VT/VirtualTextureFeedback.cpp
+++ b/Engine/Source/Runtime/Renderer/Private/VT/VirtualTextureFeedback.cpp
@@ -7,7 +7,7 @@
 
 #if PLATFORM_WINDOWS
 // Temporary use of Querys until RHI has a good fence on D3D11
-#define USE_RHI_FENCES 1
+#define USE_RHI_FENCES 0
 #else
 #define USE_RHI_FENCES 1
 #endif
