Fix bad merge

diff --git a/Engine/Source/Runtime/Renderer/Private/VT/VirtualTextureFeedback.cpp b/Engine/Source/Runtime/Renderer/Private/VT/VirtualTextureFeedback.cpp
--- a/Engine/Source/Runtime/Renderer/Private/VT/VirtualTextureFeedback.cpp
+++ b/Engine/Source/Runtime/Renderer/Private/VT/VirtualTextureFeedback.cpp
@@ -245,7 +245,7 @@ bool FVirtualTextureFeedback::Map(FRHICommandListImmediate& RHICmdList, MapResul
 	const FFeedBackItem& FeedbackEntryCPU = FeedbackTextureCPU[CPUReadIndex];
 	if (PendingTargetCount > 0u &&
 		FeedbackEntryCPU.TextureCPU.IsValid() &&
-		FeedbackEntryCPU.GPUFenceRHI->Poll())
+		FeedBackFences->Poll(RHICmdList, CPUReadIndex))
 	{
 		SCOPED_GPU_MASK(RHICmdList, FeedbackEntryCPU.GPUMask);
 
