Ensure the depth stencil action is set correctly when decals are active on mobile
Fixes decals appearing corrupted on Lumin when decals are rendered in a subpass

diff --git a/Engine/Source/Runtime/Renderer/Private/MobileShadingRenderer.cpp b/Engine/Source/Runtime/Renderer/Private/MobileShadingRenderer.cpp
--- a/Engine/Source/Runtime/Renderer/Private/MobileShadingRenderer.cpp
+++ b/Engine/Source/Runtime/Renderer/Private/MobileShadingRenderer.cpp
@@ -527,9 +527,9 @@ void FMobileSceneRenderer::Render(FRHICommandListImmediate& RHICmdList)
 			DepthTargetAction = EDepthStencilTargetActions::ClearDepthStencil_StoreDepthStencil;
 		}
 						
-		if (bKeepDepthContent && !bMobileMSAA)
+		if ((bKeepDepthContent && !bMobileMSAA) || Scene->Decals.Num() > 0)
 		{
-			// store depth if post-processing/capture needs it
+			// store depth if post-processing/capture/decals needs it
 			DepthTargetAction = EDepthStencilTargetActions::ClearDepthStencil_StoreDepthStencil;
 		}
 	}
