Fix static analysis warning and fix up for loop to coding standard
#rnx

diff --git a/Engine/Source/Runtime/Renderer/Private/RayTracing/RayTracingShadows.cpp b/Engine/Source/Runtime/Renderer/Private/RayTracing/RayTracingShadows.cpp
--- a/Engine/Source/Runtime/Renderer/Private/RayTracing/RayTracingShadows.cpp
+++ b/Engine/Source/Runtime/Renderer/Private/RayTracing/RayTracingShadows.cpp
@@ -121,21 +121,25 @@ void FDeferredShadingSceneRenderer::PrepareRayTracingShadows(const FViewInfo& Vi
 		IScreenSpaceDenoiser::EShadowRequirements::PenumbraAndClosestOccluder,
 	};
 
-	for (int32 MultiSPP = 0; MultiSPP < 1; ++MultiSPP)
-	for (int32 HairLighting = 0; HairLighting < 2; ++HairLighting)
-	for (int32 LightType = 0; LightType < LightType_MAX; ++LightType)
+	for (int32 MultiSPP = 0; MultiSPP < 2; ++MultiSPP)
 	{
-		for (IScreenSpaceDenoiser::EShadowRequirements DenoiserRequirement : DenoiserRequirements)
+		for (int32 HairLighting = 0; HairLighting < 2; ++HairLighting)
 		{
-			FOcclusionRGS::FPermutationDomain PermutationVector;
-			PermutationVector.Set<FOcclusionRGS::FLightTypeDim>(LightType);
-			PermutationVector.Set<FOcclusionRGS::FDenoiserOutputDim>((int32)DenoiserRequirement);
-			PermutationVector.Set<FOcclusionRGS::FEnableTwoSidedGeometryDim>(CVarRayTracingShadowsEnableTwoSidedGeometry.GetValueOnRenderThread() != 0);
-			PermutationVector.Set<FOcclusionRGS::FHairLighting>(HairLighting);
-			PermutationVector.Set<FOcclusionRGS::FEnableMultipleSamplesPerPixel>(MultiSPP != 0);
-
-			TShaderMapRef<FOcclusionRGS> RayGenerationShader(View.ShaderMap, PermutationVector);
-			OutRayGenShaders.Add(RayGenerationShader->GetRayTracingShader());
+			for (int32 LightType = 0; LightType < LightType_MAX; ++LightType)
+			{
+				for (IScreenSpaceDenoiser::EShadowRequirements DenoiserRequirement : DenoiserRequirements)
+				{
+					FOcclusionRGS::FPermutationDomain PermutationVector;
+					PermutationVector.Set<FOcclusionRGS::FLightTypeDim>(LightType);
+					PermutationVector.Set<FOcclusionRGS::FDenoiserOutputDim>((int32)DenoiserRequirement);
+					PermutationVector.Set<FOcclusionRGS::FEnableTwoSidedGeometryDim>(CVarRayTracingShadowsEnableTwoSidedGeometry.GetValueOnRenderThread() != 0);
+					PermutationVector.Set<FOcclusionRGS::FHairLighting>(HairLighting);
+					PermutationVector.Set<FOcclusionRGS::FEnableMultipleSamplesPerPixel>(MultiSPP != 0);
+
+					TShaderMapRef<FOcclusionRGS> RayGenerationShader(View.ShaderMap, PermutationVector);
+					OutRayGenShaders.Add(RayGenerationShader->GetRayTracingShader());
+				}
+			}
 		}
 	}
 }
