Add compute shader permutation for sky lighting (+ composition pass).

This allows to reduce the sky lighting by +60%, by doing one sampler per thread.

diff --git a/Engine/Shaders/Private/HairStrands/HairStrandsEnvironmentLighting.usf b/Engine/Shaders/Private/HairStrands/HairStrandsEnvironmentLighting.usf
--- a/Engine/Shaders/Private/HairStrands/HairStrandsEnvironmentLighting.usf
+++ b/Engine/Shaders/Private/HairStrands/HairStrandsEnvironmentLighting.usf
@@ -88,7 +88,7 @@ float InternalComputeHairCountVoxel(float3 WorldPosition, float3 LightDirection,
 //////////////////////////////////////////////////////////////////////////////////////////////////////////////////////////
 //////////////////////////////////////////////////////////////////////////////////////////////////////////////////////////
 
-float3 ReflectionEnvironment(FGBufferData GBuffer, float2 BufferUV, float2 ScreenPosition, float4 SvPosition)
+float3 ReflectionEnvironment(FGBufferData GBuffer, float2 ScreenPosition, float2 SvPosition)
 {
 #if USE_PREEXPOSURE
 	const float PreExposure = View.PreExposure;
@@ -101,7 +101,7 @@ float3 ReflectionEnvironment(FGBufferData GBuffer, float2 BufferUV, float2 Scree
 
 	const float3 VoxelSize = (Voxel_MaxAABB - Voxel_MinAABB) / Voxel_Resolution;
 	const float  VoxelMaxSize = max(VoxelSize.x, max(VoxelSize.y, VoxelSize.z));
-	const float3 VoxelOffset = VoxelMaxSize * GetHairVoxelJitter(SvPosition.xy, View.StateFrameIndexMod8);
+	const float3 VoxelOffset = VoxelMaxSize * GetHairVoxelJitter(SvPosition, View.StateFrameIndexMod8);
 
 	const float3 V = -CameraToPixel;
 	const float3 T = GBuffer.WorldNormal;
@@ -363,13 +363,15 @@ float3 ReflectionEnvironment(FGBufferData GBuffer, float2 BufferUV, float2 Scree
 	return -min(-Color.rgb, 0.0);
 }
 
+////////////////////////////////////////////////////////////////////////////////////////////////////////////////////////////////////////////
+
+#if PIXELSHADER
 void MainPS(
 	in float4 SvPosition : SV_Position,
 	out float4 OutColor : SV_Target0,
 	out float4 OutColorSubPixel : SV_Target1
 )
 {
-	const float2 BufferUV = SvPositionToBufferUV(SvPosition);
 	const float2 ScreenPosition = SvPositionToScreenPosition(SvPosition).xy;
 	const uint3 PixelCoord = uint3(floor(SvPosition.xy), 0);
 
@@ -429,7 +431,7 @@ void MainPS(
 			if (bIsClusterCompatible = SampleVisData.ClusterId == Voxel_ClusterId)
 			{
 				const float SampleCoverage = From8bitCoverage(SampleVisData.Coverage8bit);
-				SkyLighting += ReflectionEnvironment(HairGBuffer, BufferUV, ScreenPosition, SvPosition) * SampleCoverage;
+				SkyLighting += ReflectionEnvironment(HairGBuffer, ScreenPosition, SvPosition.xy) * SampleCoverage;
 			}
 			#endif
 		}
@@ -440,7 +442,7 @@ void MainPS(
 		if (bIsClusterCompatible)
 		{
 			#if PERMUTATION_PER_SAMPLE == 0
-			SkyLighting = ReflectionEnvironment(HairGBuffer, BufferUV, ScreenPosition, SvPosition) * saturate(CatData.PixelCoverage);
+			SkyLighting = ReflectionEnvironment(HairGBuffer, ScreenPosition, SvPosition.xy) * saturate(CatData.PixelCoverage);
 			#endif
 
 			// Scene color alpha is used for ScreenSpaceSubsurfaceScattering (if that is not needed it can be disabled with SUBSURFACE_CHANNEL_MODE)
@@ -459,3 +461,76 @@ void MainPS(
 		}
 	}
 }
+#endif
+
+////////////////////////////////////////////////////////////////////////////////////////////////////////////////////////////////////////////
+
+#if COMPUTESHADER
+uint MaxVisibilityNodeCount;
+StructuredBuffer<uint>		HairVisibilityNodeCoord;
+Buffer<uint4>				IndirectArgsBuffer;
+RWStructuredBuffer<float4> OutLightingBuffer;
+
+#if PERMUTATION_GROUP_SIZE == 0
+#define LIGHTING_GROUP_SIZE 64
+#elif PERMUTATION_GROUP_SIZE == 1
+#define LIGHTING_GROUP_SIZE 32
+#else
+#error Unknown group size	
+#endif
+
+[numthreads(LIGHTING_GROUP_SIZE, 1, 1)]
+void MainCS(uint3 DispatchThreadId : SV_DispatchThreadID)
+{
+	// Note: Use a fixed group count width (HAIR_VISIBILITY_GROUP_COUNT_WIDTH) for avoiding loading the indirect args (to avoid dep. memory fetch)
+	const uint LocalOffset = DispatchThreadId.x + DispatchThreadId.y * LIGHTING_GROUP_SIZE * HAIR_VISIBILITY_GROUP_COUNT_WIDTH;
+
+	const uint PackedCoord = HairVisibilityNodeCoord[LocalOffset];
+	const uint2 PixelCoord = uint2((PackedCoord & 0xFFFF), ((PackedCoord >> 16) & 0xFFFF));
+	const float2 SvPosition = PixelCoord + float2(0.5f, 0.5f);
+
+	const float2 UV = (PixelCoord + float2(0.5f, 0.5f)) / float2(View.BufferSizeAndInvSize.xy);
+	const float2 ScreenPosition = (UV - View.ScreenPositionScaleBias.wz) / View.ScreenPositionScaleBias.xy;
+
+	if (LocalOffset > MaxVisibilityNodeCount)
+		return;
+
+	const FNodeData NodeData = HairVisibilityNodeData[LocalOffset];
+	float			SampleDepth;
+	FVisibilityData SampleVisData;
+	float3			SampleBaseColor;
+	float			SampleRoughness;
+	float			SampleSpecular;
+	DecodeNodeData(NodeData, SampleDepth, SampleVisData, SampleBaseColor, SampleRoughness, SampleSpecular);
+
+	FGBufferData GBuffer;
+	GBuffer.ShadingModelID = SHADINGMODELID_HAIR;
+	GBuffer.WorldNormal = SampleVisData.Tangent;
+	GBuffer.BaseColor = SampleBaseColor;
+	GBuffer.Roughness = SampleRoughness;
+	GBuffer.Specular = SampleSpecular;
+	GBuffer.Depth = ConvertFromDeviceZ(SampleDepth);
+	GBuffer.GBufferAO = 1;
+	GBuffer.CustomData = 0;
+	GBuffer.IndirectIrradiance = 1000000;
+
+	const bool bIsClusterCompatible = SampleVisData.ClusterId == Voxel_ClusterId;
+	if (bIsClusterCompatible)
+	{
+		const float SampleCoverage = From8bitCoverage(SampleVisData.Coverage8bit);
+		const float3 SkyLighting = ReflectionEnvironment(GBuffer, ScreenPosition, SvPosition);
+
+		FLightAccumulator LightAccumulator = (FLightAccumulator)0;
+		LightAccumulator_Add(LightAccumulator, SkyLighting, SkyLighting, 1.0f, false);
+		float4 OutColor = LightAccumulator_GetResult(LightAccumulator);
+		OutColor.xyz *= saturate(SampleCoverage);
+		OutColor.a = saturate(SampleCoverage);
+
+		OutLightingBuffer[LocalOffset] = OutColor;
+	}
+	else
+	{
+		OutLightingBuffer[LocalOffset] = 0;
+	}
+}
+#endif
\ No newline at end of file
diff --git a/Engine/Shaders/Private/HairStrands/HairStrandsEnvironmentLightingCompose.usf b/Engine/Shaders/Private/HairStrands/HairStrandsEnvironmentLightingCompose.usf
new file mode 100644
--- /dev/null
+++ b/Engine/Shaders/Private/HairStrands/HairStrandsEnvironmentLightingCompose.usf
@@ -0,0 +1,38 @@
+// Copyright 1998-2019 Epic Games, Inc. All Rights Reserved.
+  
+#include "../Common.ush"
+#include "HairStrandsVisibilityCommon.ush"
+
+Buffer<float4>	 HairLightingSampleBuffer;
+Texture2D<uint4> HairCategorizationTexture;
+Texture2D<uint>	 HairVisibilityNodeOffsetAndCount;
+
+void MainPS(
+	in FScreenVertexOutput Input,
+	out float4 OutColor : SV_Target0,
+	out float4 OutColorSubPixel : SV_Target1)
+{
+	const uint3 PixelCoord = uint3(floor(Input.Position.xy), 0);
+
+	FNodeDesc NodeDesc = DecodeNodeDesc(HairVisibilityNodeOffsetAndCount.Load(PixelCoord));
+
+	OutColor = 0.0f;
+	OutColorSubPixel = 0.0f;
+
+	if (NodeDesc.Count > 0)
+	{
+		const FCategorizationData CatData = DecodeCategorizationData(HairCategorizationTexture.Load(PixelCoord));
+		const float PixelCoverage = min(CatData.PixelCoverage, 1);
+
+		LOOP
+		for (uint SampleIt = 0; SampleIt < NodeDesc.Count; SampleIt++)
+		{
+			const uint LocalOffset = NodeDesc.Offset + SampleIt;
+			const float4 LightingSample = HairLightingSampleBuffer[LocalOffset];
+			if (PixelCoverage >= 1)
+				OutColor += LightingSample;
+			else
+				OutColorSubPixel += LightingSample;
+		}		
+	}
+}
diff --git a/Engine/Shaders/Private/HairStrands/HairStrandsVisibilityCopyIndirectArg.usf b/Engine/Shaders/Private/HairStrands/HairStrandsVisibilityCopyIndirectArg.usf
--- a/Engine/Shaders/Private/HairStrands/HairStrandsVisibilityCopyIndirectArg.usf
+++ b/Engine/Shaders/Private/HairStrands/HairStrandsVisibilityCopyIndirectArg.usf
@@ -23,7 +23,7 @@ void CopyCS(uint GroupIndex : SV_GroupIndex, uint3 DispatchThreadId : SV_Dispatc
 	OutArgBuffer[2] = 1;
 #else
 	// Note: Use a fixed group count width (64) for avoiding loading the indirect args (to avoid dep. memory fetch)
-	// This value is in sync with DeepTransmittanceMask.usf
+	// This value is in sync with DeepTransmittanceMask.usf | HairStrandsEnvironementLightingCS.usf
 	if (TotalGroup < HAIR_VISIBILITY_GROUP_COUNT_WIDTH)
 	{
 		OutArgBuffer[0] = TotalGroup;
diff --git a/Engine/Source/Runtime/Renderer/Private/HairStrands/HairStrandsEnvironment.cpp b/Engine/Source/Runtime/Renderer/Private/HairStrands/HairStrandsEnvironment.cpp
--- a/Engine/Source/Runtime/Renderer/Private/HairStrands/HairStrandsEnvironment.cpp
+++ b/Engine/Source/Runtime/Renderer/Private/HairStrands/HairStrandsEnvironment.cpp
@@ -17,6 +17,7 @@
 #include "ReflectionEnvironment.h"
 #include "ScenePrivate.h"
 #include "RenderGraphEvent.h"
+#include "PostProcess/PostProcessing.h"
 
 ///////////////////////////////////////////////////////////////////////////////////////////////////
 
@@ -29,9 +30,12 @@ static FAutoConsoleVariableRef CVarHairSkyAOEnable(TEXT("r.HairStrands.SkyAOEnab
 static float GHairSkylightingConeAngle = 3;
 static FAutoConsoleVariableRef CVarHairSkylightingConeAngle(TEXT("r.HairStrands.SkyLightingConeAngle"), GHairSkylightingConeAngle, TEXT("Cone angle for tracing sky lighting on hair."));
 
-static float GHairSkylightingPerSample = 0;
+static float GHairSkylightingPerSample = 1;
 static FAutoConsoleVariableRef CVarHairSkylightingPerSample(TEXT("r.HairStrands.SkyLightingPerSample"), GHairSkylightingPerSample, TEXT("Evaluate sky lighting per hair sample."));
 
+static float GHairStrandsSkyLightingCompute = 1;
+static FAutoConsoleVariableRef GVarHairStrandsSkyLightingCompute(TEXT("r.HairStrands.SkyLightingCompute"), GHairStrandsSkyLightingCompute, TEXT("Evaluate sky lighting using a compute shader."));
+
 static bool GetHairStrandsSkyLightingEnable() { return GHairSkylightingEnable > 0; }
 static bool GetHairStrandsSkyAOEnable() { return GHairSkyAOEnable > 0; }
 static float GetHairStrandsSkyLightingConeAngle() { return FMath::Max(0.f, GHairSkylightingConeAngle); }
@@ -127,6 +131,211 @@ static void AddHairStrandsEnvironmentAOPass(
 	});
 }
 
+///////////////////////////////////////////////////////////////////////////////////////////////////
+
+class FHairEnvironmentLightingComposePS : public FGlobalShader
+{
+	DECLARE_GLOBAL_SHADER(FHairEnvironmentLightingComposePS);
+	SHADER_USE_PARAMETER_STRUCT(FHairEnvironmentLightingComposePS, FGlobalShader);
+
+	BEGIN_SHADER_PARAMETER_STRUCT(FParameters, )
+		SHADER_PARAMETER_RDG_TEXTURE(Texture2D, HairCategorizationTexture)
+		SHADER_PARAMETER_RDG_TEXTURE(Texture2D, HairVisibilityNodeOffsetAndCount)
+		SHADER_PARAMETER_RDG_BUFFER_SRV(Buffer, HairLightingSampleBuffer)
+		RENDER_TARGET_BINDING_SLOTS()
+	END_SHADER_PARAMETER_STRUCT()
+
+public:
+	static bool ShouldCompilePermutation(const FGlobalShaderPermutationParameters& Parameters) { return Parameters.Platform == SP_PCD3D_SM5; }
+};
+
+IMPLEMENT_GLOBAL_SHADER(FHairEnvironmentLightingComposePS, "/Engine/Private/HairStrands/HairStrandsEnvironmentLightingCompose.usf", "MainPS", SF_Pixel);
+
+static void AddHairEnvironmentLightingComposePass(
+	FRDGBuilder& GraphBuilder,
+	const FViewInfo& View,
+	const FHairStrandsVisibilityData& VisibilityData,
+	const FHairStrandsClusterData& ClusterData,
+	const FRDGBufferRef SampleLightingBuffer,
+	FRDGTextureRef& OutColorTexture,
+	FRDGTextureRef& OutSubColorTexture)
+{
+	FHairEnvironmentLightingComposePS::FParameters* Parameters = GraphBuilder.AllocParameters<FHairEnvironmentLightingComposePS::FParameters>();
+	Parameters->HairCategorizationTexture = GraphBuilder.RegisterExternalTexture(VisibilityData.CategorizationTexture);
+	Parameters->HairVisibilityNodeOffsetAndCount = GraphBuilder.RegisterExternalTexture(VisibilityData.NodeIndex);
+	Parameters->HairLightingSampleBuffer  = GraphBuilder.CreateSRV(SampleLightingBuffer);
+	Parameters->RenderTargets[0] = FRenderTargetBinding(OutColorTexture, ERenderTargetLoadAction::ELoad);
+	Parameters->RenderTargets[1] = FRenderTargetBinding(OutSubColorTexture, ERenderTargetLoadAction::ELoad);
+
+	TShaderMapRef<FPostProcessVS> VertexShader(View.ShaderMap);
+	TShaderMapRef<FHairEnvironmentLightingComposePS> PixelShader(View.ShaderMap);
+	const TShaderMap<FGlobalShaderType>* GlobalShaderMap = View.ShaderMap;
+	const FIntRect Viewport = View.ViewRect;
+	const FIntPoint Resolution = OutColorTexture->Desc.Extent;
+	const FViewInfo* CapturedView = &View;
+
+	ClearUnusedGraphResources(*PixelShader, Parameters);
+
+	GraphBuilder.AddPass(
+		RDG_EVENT_NAME("HairStrandsEnvironmentLightingCompose"),
+		Parameters,
+		ERDGPassFlags::Raster,
+		[Parameters, VertexShader, PixelShader, Viewport, Resolution, CapturedView](FRHICommandList& RHICmdList)
+	{
+		FGraphicsPipelineStateInitializer GraphicsPSOInit;
+		RHICmdList.ApplyCachedRenderTargets(GraphicsPSOInit);
+		GraphicsPSOInit.BlendState = TStaticBlendState<
+			CW_RGBA, BO_Add, BF_One, BF_One, BO_Add, BF_One, BF_One,
+			CW_RGBA, BO_Add, BF_One, BF_One, BO_Max, BF_SourceAlpha, BF_DestAlpha>::GetRHI();
+		GraphicsPSOInit.RasterizerState = TStaticRasterizerState<>::GetRHI();
+		GraphicsPSOInit.DepthStencilState = TStaticDepthStencilState<false, CF_Always>::GetRHI();
+		GraphicsPSOInit.BoundShaderState.VertexDeclarationRHI = GFilterVertexDeclaration.VertexDeclarationRHI;
+		GraphicsPSOInit.BoundShaderState.VertexShaderRHI = GETSAFERHISHADER_VERTEX(*VertexShader);
+		GraphicsPSOInit.BoundShaderState.PixelShaderRHI = GETSAFERHISHADER_PIXEL(*PixelShader);
+		GraphicsPSOInit.PrimitiveType = PT_TriangleList;
+		SetGraphicsPipelineState(RHICmdList, GraphicsPSOInit);
+
+		VertexShader->SetParameters(RHICmdList, CapturedView->ViewUniformBuffer);
+		RHICmdList.SetViewport(Viewport.Min.X, Viewport.Min.Y, 0.0f, Viewport.Max.X, Viewport.Max.Y, 1.0f);
+		SetShaderParameters(RHICmdList, *PixelShader, PixelShader->GetPixelShader(), *Parameters);
+		DrawRectangle(
+			RHICmdList,
+			0, 0,
+			Viewport.Width(), Viewport.Height(),
+			Viewport.Min.X,   Viewport.Min.Y,
+			Viewport.Width(), Viewport.Height(),
+			Viewport.Size(),
+			Resolution,
+			*VertexShader,
+			EDRF_UseTriangleOptimization);
+	});
+}
+
+
+///////////////////////////////////////////////////////////////////////////////////////////////////
+
+class FHairEnvironmentLightingCS : public FGlobalShader
+{
+	DECLARE_GLOBAL_SHADER(FHairEnvironmentLightingCS);
+	SHADER_USE_PARAMETER_STRUCT(FHairEnvironmentLightingCS, FGlobalShader)
+
+	class FGroupSize : SHADER_PERMUTATION_INT("PERMUTATION_GROUP_SIZE", 2);
+	using FPermutationDomain = TShaderPermutationDomain<FGroupSize>;
+
+	static bool ShouldCompilePermutation(const FGlobalShaderPermutationParameters& Parameters)
+	{
+		return IsHairStrandsSupported(Parameters.Platform);
+	}
+
+	BEGIN_SHADER_PARAMETER_STRUCT(FParameters, )
+
+		SHADER_PARAMETER(FVector, Voxel_MinAABB)
+		SHADER_PARAMETER(uint32, Voxel_ClusterId)
+		SHADER_PARAMETER(FVector, Voxel_MaxAABB)
+		SHADER_PARAMETER(uint32, Voxel_Resolution)
+		SHADER_PARAMETER(float, Voxel_DensityScale)
+		SHADER_PARAMETER(float, Voxel_DepthBiasScale)
+		SHADER_PARAMETER(float, Voxel_TanConeAngle)
+
+		SHADER_PARAMETER(uint32, MaxVisibilityNodeCount)
+
+		SHADER_PARAMETER_TEXTURE(Texture2D, PreIntegratedGF)
+		SHADER_PARAMETER_SAMPLER(SamplerState, PreIntegratedGFSampler)
+
+		SHADER_PARAMETER_RDG_BUFFER_SRV(StructuredBuffer, HairVisibilityNodeData)
+		SHADER_PARAMETER_RDG_BUFFER_SRV(StructuredBuffer, HairVisibilityNodeCoord)
+
+		SHADER_PARAMETER_STRUCT_INCLUDE(FSceneTextureParameters, SceneTextures)
+		SHADER_PARAMETER_STRUCT_INCLUDE(FSceneTextureSamplerParameters, SceneTextureSamplers)
+
+		SHADER_PARAMETER_RDG_TEXTURE(TEXTURE3D, HairEnergyLUTTexture)
+		SHADER_PARAMETER_RDG_TEXTURE(TEXTURE3D, HairScatteringLUTTexture)
+		SHADER_PARAMETER_SAMPLER(SamplerState, HairLUTSampler)
+
+		SHADER_PARAMETER_RDG_TEXTURE(Texture3D, Voxel_DensityTexture)
+		SHADER_PARAMETER_RDG_BUFFER(StructuredBuffer, IndirectArgsBuffer)
+		SHADER_PARAMETER_RDG_BUFFER_UAV(RWBuffer, OutLightingBuffer)
+
+		SHADER_PARAMETER_STRUCT_REF(FViewUniformShaderParameters, ViewUniformBuffer)
+		SHADER_PARAMETER_STRUCT_REF(FReflectionUniformParameters, ReflectionsParameters)
+		SHADER_PARAMETER_STRUCT_REF(FReflectionCaptureShaderData, ReflectionCaptureData)
+		SHADER_PARAMETER_STRUCT_REF(FForwardLightData, ForwardLightData)
+
+		END_SHADER_PARAMETER_STRUCT()
+};
+
+IMPLEMENT_GLOBAL_SHADER(FHairEnvironmentLightingCS, "/Engine/Private/HairStrands/HairStrandsEnvironmentLighting.usf", "MainCS", SF_Compute);
+
+static FRDGBufferRef AddHairStrandsEnvironmentLightingPassCS(
+	FRDGBuilder& GraphBuilder,
+	const FViewInfo& View,
+	const FHairStrandsVisibilityData& VisibilityData,
+	const FHairStrandsClusterData& ClusterData,
+	const uint32 NodeGroupSize,
+	FRDGBufferRef IndirectArgsBuffer)
+{
+	FRDGBufferRef OutBuffer = GraphBuilder.CreateBuffer(FRDGBufferDesc::CreateStructuredDesc(
+		4 * sizeof(float),
+		VisibilityData.NodeData->Desc.NumElements),
+		TEXT("HairSkyLightingNodeData"));
+
+	FSceneTextureParameters SceneTextures;
+	SetupSceneTextureParameters(GraphBuilder, &SceneTextures);
+
+	// Render the reflection environment with tiled deferred culling
+	const bool bHasBoxCaptures = (View.NumBoxReflectionCaptures > 0);
+	const bool bHasSphereCaptures = (View.NumSphereReflectionCaptures > 0);
+	FHairEnvironmentLightingCS::FParameters* PassParameters = GraphBuilder.AllocParameters<FHairEnvironmentLightingCS::FParameters>();
+
+	const FHairLUT InHairLUT = GetHairLUT(GraphBuilder.RHICmdList, View);
+	PassParameters->HairEnergyLUTTexture = GraphBuilder.RegisterExternalTexture(InHairLUT.Textures[HairLUTType_MeanEnergy], TEXT("HairMeanEnergyLUTTexture"));;
+	PassParameters->HairScatteringLUTTexture = GraphBuilder.RegisterExternalTexture(InHairLUT.Textures[HairLUTType_DualScattering], TEXT("HairScatteringEnergyLUTTexture"));
+	PassParameters->HairLUTSampler = TStaticSamplerState<SF_Bilinear, AM_Clamp, AM_Clamp, AM_Clamp>::GetRHI();
+	PassParameters->Voxel_ClusterId = ClusterData.ClusterId;
+	PassParameters->Voxel_MinAABB = ClusterData.GetMinBound();
+	PassParameters->Voxel_MaxAABB = ClusterData.GetMaxBound();
+	PassParameters->Voxel_Resolution = ClusterData.GetResolution();
+	PassParameters->Voxel_DensityTexture = GraphBuilder.RegisterExternalTexture(ClusterData.VoxelResources.DensityTexture);
+	PassParameters->Voxel_DensityScale = GetHairStrandsVoxelizationDensityScale();
+	PassParameters->Voxel_DepthBiasScale = GetHairStrandsVoxelizationDepthBiasScale();
+	PassParameters->Voxel_TanConeAngle = FMath::Tan(FMath::DegreesToRadians(GetHairStrandsSkyLightingConeAngle()));
+	PassParameters->MaxVisibilityNodeCount = VisibilityData.NodeData->Desc.NumElements;
+	PassParameters->PreIntegratedGF = GSystemTextures.PreintegratedGF->GetRenderTargetItem().ShaderResourceTexture;
+	PassParameters->PreIntegratedGFSampler = TStaticSamplerState<SF_Bilinear, AM_Clamp, AM_Clamp, AM_Clamp>::GetRHI();
+	PassParameters->SceneTextures = SceneTextures;
+	SetupSceneTextureSamplers(&PassParameters->SceneTextureSamplers);
+	PassParameters->ViewUniformBuffer = View.ViewUniformBuffer;
+	PassParameters->ReflectionCaptureData = View.ReflectionCaptureUniformBuffer;
+	{
+		FReflectionUniformParameters ReflectionUniformParameters;
+		SetupReflectionUniformParameters(View, ReflectionUniformParameters);
+		PassParameters->ReflectionsParameters = CreateUniformBufferImmediate(ReflectionUniformParameters, UniformBuffer_SingleDraw);
+	}
+	PassParameters->ForwardLightData = View.ForwardLightingResources->ForwardLightDataUniformBuffer;
+	PassParameters->IndirectArgsBuffer = IndirectArgsBuffer;
+	PassParameters->OutLightingBuffer = GraphBuilder.CreateUAV(OutBuffer);
+
+	// Bind hair data
+	FRDGBufferRef InHairVisibilityNodeData = GraphBuilder.RegisterExternalBuffer(VisibilityData.NodeData, TEXT("HairVisibilityNodeData"));
+	FRDGBufferRef InHairVisibilityNodeCoord = GraphBuilder.RegisterExternalBuffer(VisibilityData.NodeCoord, TEXT("HairVisibilityNodeCoord"));
+	PassParameters->HairVisibilityNodeData = GraphBuilder.CreateSRV(InHairVisibilityNodeData);
+	PassParameters->HairVisibilityNodeCoord = GraphBuilder.CreateSRV(InHairVisibilityNodeCoord);
+
+	check(NodeGroupSize == 64 || NodeGroupSize == 32);
+	FHairEnvironmentLightingCS::FPermutationDomain PermutationVector;
+	PermutationVector.Set<FHairEnvironmentLightingCS::FGroupSize>(NodeGroupSize == 64 ? 0 : (NodeGroupSize == 32 ? 1 : 2));
+
+	TShaderMapRef<FHairEnvironmentLightingCS> ComputeShader(View.ShaderMap, PermutationVector);
+	FComputeShaderUtils::AddPass(
+		GraphBuilder,
+		RDG_EVENT_NAME("HairStrandsEnvironmentCS %dx%d", View.ViewRect.Width(), View.ViewRect.Height()),
+		*ComputeShader,
+		PassParameters,
+		IndirectArgsBuffer,
+		0);
+
+	return OutBuffer;
+}
 
 ///////////////////////////////////////////////////////////////////////////////////////////////////
 // Lighting PS
@@ -179,7 +388,7 @@ class FHairEnvironmentLightingPS : public FGlobalShader
 
 IMPLEMENT_GLOBAL_SHADER(FHairEnvironmentLightingPS, "/Engine/Private/HairStrands/HairStrandsEnvironmentLighting.usf", "MainPS", SF_Pixel);
 
-static void AddHairStrandsEnvironmentPass(
+static void AddHairStrandsEnvironmentLightingPass(
 	FRDGBuilder& GraphBuilder,
 	const FViewInfo& View,
 	const FHairStrandsVisibilityData& VisibilityData,
@@ -240,7 +449,7 @@ static void AddHairStrandsEnvironmentPass(
 	ClearUnusedGraphResources(*PixelShader, PassParameters);
 
 	GraphBuilder.AddPass(
-		RDG_EVENT_NAME("HairStrandsEnvironment %dx%d", View.ViewRect.Width(), View.ViewRect.Height()),
+		RDG_EVENT_NAME("HairStrandsEnvironmentPS %dx%d", View.ViewRect.Width(), View.ViewRect.Height()),
 		PassParameters,
 		ERDGPassFlags::Raster,
 		[PassParameters, &View, PixelShader](FRHICommandList& InRHICmdList)
@@ -280,6 +489,8 @@ void RenderHairStrandsEnvironmentLighting(
 	{
 		return;
 	}
+	
+	FRDGBufferRef NodeIndirectArgBuffer = GraphBuilder.RegisterExternalBuffer(VisibilityData.NodeIndirectArg, TEXT("HairNodeIndirectArgBuffer"));
 
 	const FViewInfo& View = Views[ViewIndex];
 	for (const FHairStrandsClusterData& ClusterData : HairDatas->HairClusterPerViews.Views[ViewIndex].Datas)
@@ -287,7 +498,16 @@ void RenderHairStrandsEnvironmentLighting(
 		// @hair_todo: 
 		// * Add support for : BentNormal, global AO, SSR, DFSO
 		// * Add local reflection probe, current take into account only the sky lighting
-		AddHairStrandsEnvironmentPass(GraphBuilder, View, VisibilityData, ClusterData, SceneColorTexture, SceneColorSubPixelTexture);
+
+		if (GHairStrandsSkyLightingCompute)
+		{
+			FRDGBufferRef SampleLightingBuffer = AddHairStrandsEnvironmentLightingPassCS(GraphBuilder, View, VisibilityData, ClusterData, VisibilityData.NodeGroupSize, NodeIndirectArgBuffer);
+			AddHairEnvironmentLightingComposePass(GraphBuilder, View, VisibilityData, ClusterData, SampleLightingBuffer, SceneColorTexture, SceneColorSubPixelTexture);
+		}
+		else
+		{
+			AddHairStrandsEnvironmentLightingPass(GraphBuilder, View, VisibilityData, ClusterData, SceneColorTexture, SceneColorSubPixelTexture);
+		}
 	}
 }
 
