Fixed performance regressions in Motion Blur.

diff --git a/Engine/Source/Runtime/Renderer/Private/PostProcess/PostProcessMotionBlur.cpp b/Engine/Source/Runtime/Renderer/Private/PostProcess/PostProcessMotionBlur.cpp
--- a/Engine/Source/Runtime/Renderer/Private/PostProcess/PostProcessMotionBlur.cpp
+++ b/Engine/Source/Runtime/Renderer/Private/PostProcess/PostProcessMotionBlur.cpp
@@ -219,6 +219,8 @@ public:
 class FMotionBlurVelocityFlattenCS : public FMotionBlurShader
 {
 public:
+	static const uint32 ThreadGroupSize = 16;
+
 	DECLARE_GLOBAL_SHADER(FMotionBlurVelocityFlattenCS);
 	SHADER_USE_PARAMETER_STRUCT(FMotionBlurVelocityFlattenCS, FMotionBlurShader);
 
@@ -244,6 +246,8 @@ END_SHADER_PARAMETER_STRUCT()
 class FMotionBlurVelocityDilateGatherCS : public FMotionBlurShader
 {
 public:
+	static const uint32 ThreadGroupSize = 16;
+
 	DECLARE_GLOBAL_SHADER(FMotionBlurVelocityDilateGatherCS);
 	SHADER_USE_PARAMETER_STRUCT(FMotionBlurVelocityDilateGatherCS, FMotionBlurShader);
 
@@ -486,7 +490,7 @@ void AddMotionBlurVelocityPass(
 			RDG_EVENT_NAME("Velocity Flatten"),
 			*ComputeShader,
 			PassParameters,
-			FComputeShaderUtils::GetGroupCount(Viewports.Velocity.Rect.Size(), FComputeShaderUtils::kGolden2DGroupSize));
+			FComputeShaderUtils::GetGroupCount(Viewports.Velocity.Rect.Size(), FMotionBlurVelocityFlattenCS::ThreadGroupSize));
 	}
 
 	bool ScatterDilatation = IsMotionBlurScatterRequired(View, Viewports.Color);
@@ -602,7 +606,7 @@ void AddMotionBlurVelocityPass(
 			RDG_EVENT_NAME("VelocityTileGatherCS %dx%d", VelocityTileCount.X, VelocityTileCount.Y),
 			*ComputeShader,
 			PassParameters,
-			FComputeShaderUtils::GetGroupCount(VelocityTileCount, FComputeShaderUtils::kGolden2DGroupSize));
+			FComputeShaderUtils::GetGroupCount(VelocityTileCount, FMotionBlurVelocityDilateGatherCS::ThreadGroupSize));
 	}
 
 	*VelocityFlatTextureOutput = VelocityFlatTexture;
