Time slice the VirtualTexture MaxUploadsPerFrame
This reduces the per frame uploads for very high frame rates
Disabled by default with a CVar to enable

diff --git a/Engine/Source/Runtime/Renderer/Private/VT/VirtualTextureSystem.cpp b/Engine/Source/Runtime/Renderer/Private/VT/VirtualTextureSystem.cpp
--- a/Engine/Source/Runtime/Renderer/Private/VT/VirtualTextureSystem.cpp
+++ b/Engine/Source/Runtime/Renderer/Private/VT/VirtualTextureSystem.cpp
@@ -88,6 +88,13 @@ static TAutoConsoleVariable<int32> CVarVTPageUpdateFlushCount(
 );
 static const int32 MaxNumTasks = 16;
 
+static TAutoConsoleVariable<float> CVarVTPageUpdateTimeSlice(
+	TEXT("r.VT.PageUpdateTimeSlice"),
+	0.f,
+	TEXT("Assumed frame time to slice MaxUploadsPerFrame. 0 is disabled."),
+	ECVF_RenderThreadSafe
+);
+
 static FORCEINLINE uint32 EncodePage(uint32 ID, uint32 vLevel, uint32 vTileX, uint32 vTileY)
 {
 	uint32 Page;
@@ -868,6 +875,35 @@ void FVirtualTextureSystem::FeedbackAnalysisTask(const FFeedbackAnalysisParamete
 	}
 }
 
+/** Helper to reduce MaxUploadsPerFrame at higher frame rates. */
+class FVirtualTextureTimeSlicer
+{
+public:
+	/** Call once per frame to update current time deltas. */
+	void Update()
+	{
+		const float Time = FApp::GetCurrentTime() - GStartTime;
+		TimeDelta = Time - LastTime;
+		LastTime = Time;
+	}
+
+	/** Slice a value for the current time slice according to the BaseTimeSlice. This clamps for time deltas greater than the BaseTimeSlice. */
+	int32 Slice(int32 InValue, float BaseTimeSlice)
+	{
+		if (BaseTimeSlice <= 0.f)
+		{
+			return InValue;
+		}
+		
+		const float ClampedRatio = FMath::Clamp(TimeDelta / BaseTimeSlice, 0.0001f, 1.f);
+		return FMath::CeilToInt(ClampedRatio * InValue);
+	}
+
+private:
+	float LastTime = 0.f;
+	float TimeDelta = 0.f;
+};
+
 void FVirtualTextureSystem::Update(FRHICommandListImmediate& RHICmdList, ERHIFeatureLevel::Type FeatureLevel, FScene* Scene)
 {
 	check(IsInRenderingThread());
@@ -875,6 +911,9 @@ void FVirtualTextureSystem::Update(FRHICommandListImmediate& RHICmdList, ERHIFea
 	CSV_SCOPED_TIMING_STAT_EXCLUSIVE(VirtualTextureSystem_Update);
 	SCOPE_CYCLE_COUNTER(STAT_VirtualTextureSystem_Update);
 	SCOPED_GPU_STAT(RHICmdList, VirtualTexture);
+	
+	static FVirtualTextureTimeSlicer TimeSlicer;
+	TimeSlicer.Update();
 
 	if (bFlushCaches)
 	{
@@ -1109,7 +1148,9 @@ void FVirtualTextureSystem::Update(FRHICommandListImmediate& RHICmdList, ERHIFea
 		// Limit the number of uploads (account for MappedTilesToProduce this frame)
 		// Are all pages equal? Should there be different limits on different types of pages?
 		const int32 MaxNumUploads = VirtualTextureScalability::GetMaxUploadsPerFrame();
-		const int32 MaxRequestUploads = FMath::Max(MaxNumUploads - MappedTilesToProduce.Num(), 1);
+		const float TimeSlice = CVarVTPageUpdateTimeSlice.GetValueOnRenderThread();
+		const int32 MaxNumUploadsForTimeSlice = TimeSlicer.Slice(MaxNumUploads, TimeSlice);
+		const int32 MaxRequestUploads = FMath::Max(MaxNumUploadsForTimeSlice - MappedTilesToProduce.Num(), 1);
 
 		MergedRequestList->SortRequests(Producers, MemStack, MaxRequestUploads);
 	}
