Avoid creating too many uniform buffers for deep shadow maps.
Also simplified function parameters.

diff --git a/Engine/Shaders/Private/HairStrands/HairStrandsDeepShadowVS.usf b/Engine/Shaders/Private/HairStrands/HairStrandsDeepShadowVS.usf
--- a/Engine/Shaders/Private/HairStrands/HairStrandsDeepShadowVS.usf
+++ b/Engine/Shaders/Private/HairStrands/HairStrandsDeepShadowVS.usf
@@ -56,7 +56,7 @@ void Main(
   #ifdef VF_STRAND_HAIR	
 	FHairRenderingInfo HairRenderingInfo = GetHairRenderInfo(ResolvedView.HairRenderInfo);
 	const float CurrentDepth = Output.Position.z / Output.Position.w;
-	const float PixelRadius = HairRenderingInfo.IsOrthoView ? HairRenderingInfo.RadiusAtDepth1Primary : (CurrentDepth * HairRenderingInfo.RadiusAtDepth1Primary);
+	const float PixelRadius = HairRenderingInfo.bIsOrthoView ? HairRenderingInfo.RadiusAtDepth1Primary : (CurrentDepth * HairRenderingInfo.RadiusAtDepth1Primary);
 	const float StrandRealRadius = VFIntermediates.HairDimensions.y;
 	const float Coverage = StrandRealRadius / max(StrandRealRadius, PixelRadius);
 	Output.HairCoverage = Coverage * VFIntermediates.HairDensity;
diff --git a/Engine/Shaders/Private/HairStrands/HairStrandsVertexFactory.ush b/Engine/Shaders/Private/HairStrands/HairStrandsVertexFactory.ush
--- a/Engine/Shaders/Private/HairStrands/HairStrandsVertexFactory.ush
+++ b/Engine/Shaders/Private/HairStrands/HairStrandsVertexFactory.ush
@@ -367,7 +367,7 @@ float4 ComputeViewAlignedWorldPosition(FVertexFactoryInput Input, float3 WorldTa
 	// Minimal radius to snap the strand to a sample/pixel center (to avoid aliasing)
 	const float RadiusAtDepth1 = lerp(HairRenderingInfo.RadiusAtDepth1Primary, HairRenderingInfo.RadiusAtDepth1Velocity, VelocityScale);
 	const float DistanceToCamera = length(ResolvedView.TranslatedWorldCameraOrigin - WorldPosition.xyz);
-	const float MinStrandHairRadius = HairRenderingInfo.IsOrthoView ? RadiusAtDepth1 : (DistanceToCamera * RadiusAtDepth1);
+	const float MinStrandHairRadius = HairRenderingInfo.bIsOrthoView ? RadiusAtDepth1 : (DistanceToCamera * RadiusAtDepth1);
 	const float3 ViewDir = -ResolvedView.ViewForward;
 	const float3 Right = normalize(cross(WorldTangent, ViewDir));
 	const float3 OutWorldPosition = WorldPosition.xyz + (bIsLeft ? -Right : Right) * max(WorldStrandRadius, MinStrandHairRadius);
diff --git a/Engine/Shaders/Private/HairStrands/HairStrandsVisibilityCommon.ush b/Engine/Shaders/Private/HairStrands/HairStrandsVisibilityCommon.ush
--- a/Engine/Shaders/Private/HairStrands/HairStrandsVisibilityCommon.ush
+++ b/Engine/Shaders/Private/HairStrands/HairStrandsVisibilityCommon.ush
@@ -315,7 +315,7 @@ FNodeData ConvertToNodeData(float Depth, FVisibilityData VisibilityData, float3
  {
 	 float RadiusAtDepth1Primary;
 	 float RadiusAtDepth1Velocity;
-	 bool IsOrthoView;
+	 bool bIsOrthoView;
  };
 
  FHairRenderingInfo GetHairRenderInfo(float4 ViewHairRenderInfo)
@@ -323,6 +323,6 @@ FNodeData ConvertToNodeData(float Depth, FVisibilityData VisibilityData, float3
 	 FHairRenderingInfo Info;
 	 Info.RadiusAtDepth1Primary = ViewHairRenderInfo.x;
 	 Info.RadiusAtDepth1Velocity = ViewHairRenderInfo.y;
-	 Info.IsOrthoView = ViewHairRenderInfo.z > 0.0f ? true : false;
+	 Info.bIsOrthoView = ViewHairRenderInfo.z > 0.0f ? true : false;
 	 return Info;
  }
diff --git a/Engine/Source/Runtime/Renderer/Private/HairStrands/HairStrandsDeepShadow.cpp b/Engine/Source/Runtime/Renderer/Private/HairStrands/HairStrandsDeepShadow.cpp
--- a/Engine/Source/Runtime/Renderer/Private/HairStrands/HairStrandsDeepShadow.cpp
+++ b/Engine/Source/Runtime/Renderer/Private/HairStrands/HairStrandsDeepShadow.cpp
@@ -50,8 +50,9 @@ static void RenderDeepShadowFrontDepth(
 	const FScene* Scene,
 	const FViewInfo* ViewInfo,
 	const FHairStrandsClusterData::TPrimitiveInfos& PrimitiveSceneInfo,
-	const FHairStrandsLightDesc& LightDesc,
 	const FIntRect& AtlasRect,
+	const TUniformBufferRef<FViewUniformShaderParameters>& ViewUniformShaderParameters,
+	const TUniformBufferRef<FDeepShadowPassUniformParameters>& DeepShadowPassUniformParameters,
 	const bool bClearOutput,
 	TRefCountPtr<IPooledRenderTarget>& outShadowDepthRT)
 {
@@ -69,10 +70,10 @@ static void RenderDeepShadowFrontDepth(
 		Scene,
 		ViewInfo,
 		PrimitiveSceneInfo,
-		LightDesc,
 		EHairStrandsRasterPassType::FrontDepth,
-		nullptr,
-		AtlasRect);
+		AtlasRect,
+		ViewUniformShaderParameters,
+		DeepShadowPassUniformParameters);
 	RHICmdList.EndRenderPass();
 }
 
@@ -83,9 +84,9 @@ static void RenderDeepShadowLayers(
 	const FScene* Scene,
 	const FViewInfo* ViewInfo,
 	const FHairStrandsClusterData::TPrimitiveInfos& PrimitiveSceneInfo,
-	const FHairStrandsLightDesc& LightDesc,
 	const FIntRect& AtlasRect,
-	FRHITexture* ShadowDepthTexture,
+	const TUniformBufferRef<FViewUniformShaderParameters>& ViewUniformShaderParameters,
+	const TUniformBufferRef<FDeepShadowPassUniformParameters>& DeepShadowPassUniformParameters,
 	const bool bClearOutput,
 	TRefCountPtr<IPooledRenderTarget>& outDeepShadowLayersRT)
 {
@@ -100,10 +101,10 @@ static void RenderDeepShadowLayers(
 		Scene,
 		ViewInfo,
 		PrimitiveSceneInfo,
-		LightDesc,
 		EHairStrandsRasterPassType::DeepOpacityMap,
-		ShadowDepthTexture,
-		AtlasRect);
+		AtlasRect,
+		ViewUniformShaderParameters,
+		DeepShadowPassUniformParameters);
 	RHICmdList.EndRenderPass();
 }
 
@@ -191,8 +192,8 @@ FHairStrandsDeepShadowViews RenderHairStrandsDeepShadows(
 	uint32 AtlasSlotIndex = 0;
 	for (int32 ViewIndex = 0; ViewIndex < Views.Num(); ViewIndex++)
 	{
-		const FViewInfo& View = Views[ViewIndex];
-		if (View.Family)
+		const FViewInfo& ViewInfo = Views[ViewIndex];
+		if (ViewInfo.Family)
 		{
 			FHairStrandsDeepShadowDatas& DeepShadowDatas = DeepShadowsPerView.Views[ViewIndex];
 
@@ -235,31 +236,49 @@ FHairStrandsDeepShadowViews RenderHairStrandsDeepShadows(
 					DomData.LayersTexture = DeepShadowLayersAtlasTexture;
 					AtlasSlotIndex++;
 
-					FHairStrandsLightDesc LightDesc;
-					LightDesc.bIsOrtho = LightType == ELightComponentType::LightType_Directional;
-					LightDesc.LightDirection = DomData.LightDirection;
-					LightDesc.LightPosition = DomData.LightPosition;
-					LightDesc.WorldToLightClipTransform = DomData.WorldToLightTransform;
-					LightDesc.MinStrandRadiusAtDepth1 = MinStrandRadiusAtDepth1;
+					TUniformBufferRef<FDeepShadowPassUniformParameters> DeepShadowPassUniformParameters;
+					{
+						FDeepShadowPassUniformParameters PassParameters;
+						PassParameters.WorldToClipMatrix = DomData.WorldToLightTransform;;
+						PassParameters.SliceValue = FVector4(1, 1, 1, 1);
+						PassParameters.FrontDepthTexture = DomData.DepthTexture->GetRenderTargetItem().TargetableTexture.GetReference();
+						PassParameters.AtlasRect = DomData.AtlasRect;
+						PassParameters.VoxelMinAABB = FVector::ZeroVector;
+						PassParameters.VoxelMaxAABB = FVector::ZeroVector;
+						PassParameters.VoxelResolution = 0;
+						DeepShadowPassUniformParameters = CreateUniformBufferImmediate(PassParameters, UniformBuffer_SingleFrame, EUniformBufferValidation::None);
+					}
+
+					TUniformBufferRef<FViewUniformShaderParameters> ViewUniformShaderParameters;
+					{
+						const float bIsOrtho = LightType == ELightComponentType::LightType_Directional;
+
+						ViewInfo.CachedViewUniformShaderParameters->HairRenderInfo.X = MinStrandRadiusAtDepth1.Primary;
+						ViewInfo.CachedViewUniformShaderParameters->HairRenderInfo.Y = MinStrandRadiusAtDepth1.Primary;
+						ViewInfo.CachedViewUniformShaderParameters->HairRenderInfo.Z = bIsOrtho;
+						ViewInfo.CachedViewUniformShaderParameters->ViewForward = DomData.LightDirection;
+						ViewUniformShaderParameters = TUniformBufferRef<FViewUniformShaderParameters>::CreateUniformBufferImmediate(*ViewInfo.CachedViewUniformShaderParameters, UniformBuffer_SingleFrame);
+					}
 
 					RenderDeepShadowFrontDepth(
 						RHICmdList,
 						Scene,
-						&View,
+						&ViewInfo,
 						Cluster.PrimitivesInfos,
-						LightDesc,
 						DomData.AtlasRect,
+						ViewUniformShaderParameters,
+						DeepShadowPassUniformParameters,
 						bClearFrontDepthAtlasTexture,
 						DomData.DepthTexture);
 
 					RenderDeepShadowLayers(
 						RHICmdList,
 						Scene,
-						&View,
+						&ViewInfo,
 						Cluster.PrimitivesInfos,
-						LightDesc,
 						DomData.AtlasRect,
-						DomData.DepthTexture->GetRenderTargetItem().TargetableTexture.GetReference(),
+						ViewUniformShaderParameters,
+						DeepShadowPassUniformParameters,
 						bClearLayerAtlasTexture,
 						DomData.LayersTexture);
 
diff --git a/Engine/Source/Runtime/Renderer/Private/HairStrands/HairStrandsRasterCommon.cpp b/Engine/Source/Runtime/Renderer/Private/HairStrands/HairStrandsRasterCommon.cpp
--- a/Engine/Source/Runtime/Renderer/Private/HairStrands/HairStrandsRasterCommon.cpp
+++ b/Engine/Source/Runtime/Renderer/Private/HairStrands/HairStrandsRasterCommon.cpp
@@ -12,23 +12,8 @@
 #include "ScenePrivate.h"
 
 // this is temporary until we split the voxelize and DOM path
-
-
-/////////////////////////////////////////////////////////////////////////////////////////
-
-BEGIN_GLOBAL_SHADER_PARAMETER_STRUCT(FDeepShadowPassUniformParameters, )
-	//SHADER_PARAMETER_STRUCT(FSceneTexturesUniformParameters, SceneTextures)
-	SHADER_PARAMETER(FMatrix, WorldToClipMatrix)
-	SHADER_PARAMETER(FVector4, SliceValue)
-	SHADER_PARAMETER(FIntRect, AtlasRect)
-	SHADER_PARAMETER(FVector, VoxelMinAABB)
-	SHADER_PARAMETER(FVector, VoxelMaxAABB)
-	SHADER_PARAMETER(uint32, VoxelResolution)
-	SHADER_PARAMETER_TEXTURE(Texture2D<float>, FrontDepthTexture)
-END_GLOBAL_SHADER_PARAMETER_STRUCT()
 IMPLEMENT_GLOBAL_SHADER_PARAMETER_STRUCT(FDeepShadowPassUniformParameters, "DeepShadowPass");
 
-
 /////////////////////////////////////////////////////////////////////////////////////////
 
 class FDeepShadowDepthMeshVS : public FMeshMaterialShader
@@ -379,99 +364,61 @@ void RasterHairStrands(
 	const FScene* Scene,
 	const FViewInfo* ViewInfo,
 	const FHairStrandsClusterData::TPrimitiveInfos& PrimitiveSceneInfos,
-	const FHairStrandsLightDesc& LightDesc,
 	const EHairStrandsRasterPassType ShadowPassType,
-	FRHITexture* DeepShadowDepthTexture,
 	const FIntRect& AtlasRect,
-	const FVector MinAABB,
-	const FVector MaxAABB)
+	const TUniformBufferRef<FViewUniformShaderParameters>& ViewUniformShaderParameters,
+	const TUniformBufferRef<FDeepShadowPassUniformParameters>& DeepShadowPassUniformParameters)
 {
 	check(RHICmdList.IsInsideRenderPass());
 	check(IsInRenderingThread());
 
 	SCOPE_CYCLE_COUNTER(STAT_RenderPerObjectShadowDepthsTime);
 
-	TUniformBufferRef<FDeepShadowPassUniformParameters> PassUniformBuffer;
-	{
-		if (ShadowPassType == EHairStrandsRasterPassType::DeepOpacityMap)
-		{
-			check(DeepShadowDepthTexture);
-		}
+	FMeshPassProcessorRenderState DrawRenderState(*ViewInfo, DeepShadowPassUniformParameters);
+	DrawRenderState.SetViewUniformBuffer(ViewUniformShaderParameters);
 
-		FDeepShadowPassUniformParameters PassParameters;
-		PassParameters.WorldToClipMatrix = LightDesc.WorldToLightClipTransform;
-		PassParameters.SliceValue = FVector4(1, 1, 1, 1);
-		PassParameters.FrontDepthTexture = DeepShadowDepthTexture;
-		PassParameters.AtlasRect = AtlasRect;
-		PassParameters.VoxelMinAABB = MinAABB;
-		PassParameters.VoxelMaxAABB = MaxAABB;
-		PassParameters.VoxelResolution = AtlasRect.Width();
-		PassUniformBuffer = CreateUniformBufferImmediate(PassParameters, UniformBuffer_SingleDraw, EUniformBufferValidation::None);
-	}
-
-	FMeshPassProcessorRenderState DrawRenderState(*ViewInfo, PassUniformBuffer);
+	RHICmdList.SetViewport(AtlasRect.Min.X, AtlasRect.Min.Y, 0.0f, AtlasRect.Max.X, AtlasRect.Max.Y, 1.0f);
 
-	if (ShadowPassType == EHairStrandsRasterPassType::DeepOpacityMap
-		|| ShadowPassType == EHairStrandsRasterPassType::FrontDepth
-		|| ShadowPassType == EHairStrandsRasterPassType::VoxelizationMaterial 
-		|| ShadowPassType == EHairStrandsRasterPassType::Voxelization)
+	if (ShadowPassType == EHairStrandsRasterPassType::DeepOpacityMap)
 	{
-		TUniformBufferRef<FViewUniformShaderParameters> ViewPassUniformBuffer;
-		// Create a ViewUniformBuffer by patching the cached one
-
-		ViewInfo->CachedViewUniformShaderParameters->HairRenderInfo.X = LightDesc.MinStrandRadiusAtDepth1.Primary;
-		ViewInfo->CachedViewUniformShaderParameters->HairRenderInfo.Y = LightDesc.MinStrandRadiusAtDepth1.Velocity;
-		ViewInfo->CachedViewUniformShaderParameters->HairRenderInfo.Z = LightDesc.bIsOrtho ? 1.0f : 0.0f;
-		ViewInfo->CachedViewUniformShaderParameters->ViewForward = LightDesc.LightDirection;
-		ViewPassUniformBuffer = TUniformBufferRef<FViewUniformShaderParameters>::CreateUniformBufferImmediate(*ViewInfo->CachedViewUniformShaderParameters, UniformBuffer_SingleFrame);
-		DrawRenderState.SetViewUniformBuffer(ViewPassUniformBuffer);
+		DrawRenderState.SetBlendState(TStaticBlendState<
+			CW_RGBA, BO_Add, BF_One, BF_One, BO_Add, BF_One, BF_One,
+			CW_RGBA, BO_Add, BF_One, BF_One, BO_Add, BF_One, BF_One>::GetRHI());
+		DrawRenderState.SetDepthStencilState(TStaticDepthStencilState<false, CF_Always>::GetRHI());
 	}
-
-
+	else if (ShadowPassType == EHairStrandsRasterPassType::FrontDepth)
 	{
-		RHICmdList.SetViewport(AtlasRect.Min.X, AtlasRect.Min.Y, 0.0f, AtlasRect.Max.X, AtlasRect.Max.Y, 1.0f);
-
-		if (ShadowPassType == EHairStrandsRasterPassType::DeepOpacityMap)
-		{
-			DrawRenderState.SetBlendState(TStaticBlendState<
-				CW_RGBA, BO_Add, BF_One, BF_One, BO_Add, BF_One, BF_One,
-				CW_RGBA, BO_Add, BF_One, BF_One, BO_Add, BF_One, BF_One>::GetRHI());
-			DrawRenderState.SetDepthStencilState(TStaticDepthStencilState<false, CF_Always>::GetRHI());
-		}
-		else if (ShadowPassType == EHairStrandsRasterPassType::FrontDepth)
-		{
-			DrawRenderState.SetBlendState(TStaticBlendState<
-				CW_RGBA, BO_Add, BF_One, BF_Zero, BO_Add, BF_One, BF_Zero,
-				CW_RGBA, BO_Add, BF_One, BF_Zero, BO_Add, BF_One, BF_Zero>::GetRHI());
-			DrawRenderState.SetDepthStencilState(TStaticDepthStencilState<true, CF_DepthNearOrEqual>::GetRHI());
-		}
-		else if (ShadowPassType == EHairStrandsRasterPassType::Voxelization || ShadowPassType == EHairStrandsRasterPassType::VoxelizationMaterial)
-		{
-			DrawRenderState.SetBlendState(TStaticBlendState<
-				CW_RGBA, BO_Add, BF_One, BF_Zero, BO_Add, BF_One, BF_Zero>::GetRHI());
-			DrawRenderState.SetDepthStencilState(TStaticDepthStencilState<false, CF_Always>::GetRHI());
-		}
+		DrawRenderState.SetBlendState(TStaticBlendState<
+			CW_RGBA, BO_Add, BF_One, BF_Zero, BO_Add, BF_One, BF_Zero,
+			CW_RGBA, BO_Add, BF_One, BF_Zero, BO_Add, BF_One, BF_Zero>::GetRHI());
+		DrawRenderState.SetDepthStencilState(TStaticDepthStencilState<true, CF_DepthNearOrEqual>::GetRHI());
+	}
+	else if (ShadowPassType == EHairStrandsRasterPassType::Voxelization || ShadowPassType == EHairStrandsRasterPassType::VoxelizationMaterial)
+	{
+		DrawRenderState.SetBlendState(TStaticBlendState<
+			CW_RGBA, BO_Add, BF_One, BF_Zero, BO_Add, BF_One, BF_Zero>::GetRHI());
+		DrawRenderState.SetDepthStencilState(TStaticDepthStencilState<false, CF_Always>::GetRHI());
+	}
 
-		FDynamicMeshDrawCommandStorage DynamicMeshDrawCommandStorage; // << Were would thid be stored?
-		FMeshCommandOneFrameArray VisibleMeshDrawCommands;
-		FGraphicsMinimalPipelineStateSet GraphicsMinimalPipelineStateSet;
-		FDynamicPassMeshDrawListContext ShadowContext(DynamicMeshDrawCommandStorage, VisibleMeshDrawCommands, GraphicsMinimalPipelineStateSet);
+	FDynamicMeshDrawCommandStorage DynamicMeshDrawCommandStorage; // << Were would thid be stored?
+	FMeshCommandOneFrameArray VisibleMeshDrawCommands;
+	FGraphicsMinimalPipelineStateSet GraphicsMinimalPipelineStateSet;
+	FDynamicPassMeshDrawListContext ShadowContext(DynamicMeshDrawCommandStorage, VisibleMeshDrawCommands, GraphicsMinimalPipelineStateSet);
 
 
-		FDeepShadowMeshProcessor DeepShadowMeshProcessor(Scene, ViewInfo /* is a SceneView */, DrawRenderState, &ShadowContext, ShadowPassType);
+	FDeepShadowMeshProcessor DeepShadowMeshProcessor(Scene, ViewInfo /* is a SceneView */, DrawRenderState, &ShadowContext, ShadowPassType);
 
-		for (const FMeshBatchAndRelevance& MeshBatchAndRelevance : PrimitiveSceneInfos)
-		{
-			const FMeshBatch& MeshBatch = *MeshBatchAndRelevance.Mesh;
-			const uint64 BatchElementMask = ~0ull;
-			DeepShadowMeshProcessor.AddMeshBatch(MeshBatch, BatchElementMask, MeshBatchAndRelevance.PrimitiveSceneProxy);
-		}
+	for (const FMeshBatchAndRelevance& MeshBatchAndRelevance : PrimitiveSceneInfos)
+	{
+		const FMeshBatch& MeshBatch = *MeshBatchAndRelevance.Mesh;
+		const uint64 BatchElementMask = ~0ull;
+		DeepShadowMeshProcessor.AddMeshBatch(MeshBatch, BatchElementMask, MeshBatchAndRelevance.PrimitiveSceneProxy);
+	}
 
-		if (VisibleMeshDrawCommands.Num() > 0)
-		{
-			FRHIVertexBuffer* PrimitiveIdVertexBuffer = nullptr;
-			SortAndMergeDynamicPassMeshDrawCommands(ViewInfo->GetFeatureLevel(), VisibleMeshDrawCommands, DynamicMeshDrawCommandStorage, PrimitiveIdVertexBuffer, 1);
-			SubmitMeshDrawCommands(VisibleMeshDrawCommands, GraphicsMinimalPipelineStateSet, PrimitiveIdVertexBuffer, 0, false, 1, RHICmdList);
-		}
+	if (VisibleMeshDrawCommands.Num() > 0)
+	{
+		FRHIVertexBuffer* PrimitiveIdVertexBuffer = nullptr;
+		SortAndMergeDynamicPassMeshDrawCommands(ViewInfo->GetFeatureLevel(), VisibleMeshDrawCommands, DynamicMeshDrawCommandStorage, PrimitiveIdVertexBuffer, 1);
+		SubmitMeshDrawCommands(VisibleMeshDrawCommands, GraphicsMinimalPipelineStateSet, PrimitiveIdVertexBuffer, 0, false, 1, RHICmdList);
 	}
 }
diff --git a/Engine/Source/Runtime/Renderer/Private/HairStrands/HairStrandsRasterCommon.h b/Engine/Source/Runtime/Renderer/Private/HairStrands/HairStrandsRasterCommon.h
--- a/Engine/Source/Runtime/Renderer/Private/HairStrands/HairStrandsRasterCommon.h
+++ b/Engine/Source/Runtime/Renderer/Private/HairStrands/HairStrandsRasterCommon.h
@@ -18,15 +18,23 @@ enum class EHairStrandsRasterPassType : uint8
 	VoxelizationMaterial,
 };
 
+BEGIN_GLOBAL_SHADER_PARAMETER_STRUCT(FDeepShadowPassUniformParameters, )
+	SHADER_PARAMETER(FMatrix, WorldToClipMatrix)
+	SHADER_PARAMETER(FVector4, SliceValue)
+	SHADER_PARAMETER(FIntRect, AtlasRect)
+	SHADER_PARAMETER(FVector, VoxelMinAABB)
+	SHADER_PARAMETER(FVector, VoxelMaxAABB)
+	SHADER_PARAMETER(uint32, VoxelResolution)
+	SHADER_PARAMETER_TEXTURE(Texture2D<float>, FrontDepthTexture)
+END_GLOBAL_SHADER_PARAMETER_STRUCT()
+
 // to be removed
 void RasterHairStrands(
 	FRHICommandList& RHICmdList,
 	const FScene* Scene,
 	const FViewInfo* ViewInfo,
 	const FHairStrandsClusterData::TPrimitiveInfos& PrimitiveSceneInfos,
-	const struct FHairStrandsLightDesc& LightDesc,
 	const EHairStrandsRasterPassType ShadowPassType,
-	FRHITexture* DeepShadowDepthTexture,
 	const FIntRect& AtlasRect,
-	const FVector MinAABB = FVector::ZeroVector,
-	const FVector MaxAABB = FVector::ZeroVector);
\ No newline at end of file
+	const TUniformBufferRef<FViewUniformShaderParameters>& ViewUniformShaderParameters,
+	const TUniformBufferRef<FDeepShadowPassUniformParameters>& DeepShadowPassUniformParameters);
\ No newline at end of file
diff --git a/Engine/Source/Runtime/Renderer/Private/HairStrands/HairStrandsUtils.h b/Engine/Source/Runtime/Renderer/Private/HairStrands/HairStrandsUtils.h
--- a/Engine/Source/Runtime/Renderer/Private/HairStrands/HairStrandsUtils.h
+++ b/Engine/Source/Runtime/Renderer/Private/HairStrands/HairStrandsUtils.h
@@ -11,15 +11,6 @@
 #include "HairStrandsInterface.h"
 #include "SceneTypes.h"
 
-struct FHairStrandsLightDesc
-{
-	FMatrix WorldToLightClipTransform;
-	FVector LightDirection;
-	FVector LightPosition;
-	FMinHairRadiusAtDepth1 MinStrandRadiusAtDepth1;
-	bool bIsOrtho;
-};
-
 FIntRect ComputeProjectedScreenRect(const FBox& B, const FViewInfo& View);
 
 void ComputeWorldToLightClip(
diff --git a/Engine/Source/Runtime/Renderer/Private/HairStrands/HairStrandsVoxelization.cpp b/Engine/Source/Runtime/Renderer/Private/HairStrands/HairStrandsVoxelization.cpp
--- a/Engine/Source/Runtime/Renderer/Private/HairStrands/HairStrandsVoxelization.cpp
+++ b/Engine/Source/Runtime/Renderer/Private/HairStrands/HairStrandsVoxelization.cpp
@@ -303,24 +303,40 @@ static void RenderVoxelization(
 	}
 
 	const FIntRect ViewportRect(0, 0, ResolutionDim, ResolutionDim);
-
-	FHairStrandsLightDesc LightDesc;
-	LightDesc.bIsOrtho = true;
-	LightDesc.LightDirection = FVector(1, 0, 0);
-	LightDesc.LightPosition = FVector(0, 0, 0);
+	const FVector LightDirection = FVector(1, 0, 0);
+	const FVector LightPosition = FVector(0, 0, 0);
+	const FSphere SphereBound = Bounds.GetSphere();
+	const float SphereRadius = SphereBound.W * GHairVoxelizationAABBScale;
+	const float RadiusAtDepth1 = GStrandHairVoxelizationRasterizationScale * SphereRadius / FMath::Min(Resolution.X, Resolution.Y);
+	const float bIsOrtho = 1.0f;
+
+	TUniformBufferRef<FDeepShadowPassUniformParameters> DeepShadowPassUniformParameters;
 	{
-		const FSphere SphereBound = Bounds.GetSphere();
-		const float SphereRadius = SphereBound.W * GHairVoxelizationAABBScale;
-
 		FReversedZOrthoMatrix OrthoMatrix(SphereRadius, SphereRadius, 1.f / (2 * SphereRadius), 0);
-		FLookAtMatrix LookAt(SphereBound.Center - LightDesc.LightDirection * SphereRadius, SphereBound.Center, FVector(0, 0, 1));
-		LightDesc.WorldToLightClipTransform = LookAt * OrthoMatrix;
-		LightDesc.MinStrandRadiusAtDepth1.Primary = GStrandHairVoxelizationRasterizationScale * SphereRadius / FMath::Min(Resolution.X, Resolution.Y);
-		LightDesc.MinStrandRadiusAtDepth1.Velocity = LightDesc.MinStrandRadiusAtDepth1.Primary;
+		FLookAtMatrix LookAt(SphereBound.Center - LightDirection * SphereRadius, SphereBound.Center, FVector(0, 0, 1));
 
-		VoxelResources.WorldToClip = LightDesc.WorldToLightClipTransform;
+		VoxelResources.WorldToClip = LookAt * OrthoMatrix;
 		VoxelResources.MinAABB = Bounds.GetSphere().Center - SphereRadius;
 		VoxelResources.MaxAABB = Bounds.GetSphere().Center + SphereRadius;
+
+		FDeepShadowPassUniformParameters PassParameters;
+		PassParameters.WorldToClipMatrix = VoxelResources.WorldToClip;;
+		PassParameters.SliceValue = FVector4(1, 1, 1, 1);
+		PassParameters.FrontDepthTexture = nullptr;
+		PassParameters.AtlasRect = ViewportRect;
+		PassParameters.VoxelMinAABB = VoxelResources.MinAABB;
+		PassParameters.VoxelMaxAABB = VoxelResources.MaxAABB;
+		PassParameters.VoxelResolution = ViewportRect.Width();
+		DeepShadowPassUniformParameters = CreateUniformBufferImmediate(PassParameters, UniformBuffer_SingleDraw, EUniformBufferValidation::None);
+	}
+
+	TUniformBufferRef<FViewUniformShaderParameters> ViewUniformShaderParameters;
+	{
+		ViewInfo->CachedViewUniformShaderParameters->HairRenderInfo.X = RadiusAtDepth1;
+		ViewInfo->CachedViewUniformShaderParameters->HairRenderInfo.Y = RadiusAtDepth1;
+		ViewInfo->CachedViewUniformShaderParameters->HairRenderInfo.Z = bIsOrtho;
+		ViewInfo->CachedViewUniformShaderParameters->ViewForward = LightDirection;
+		ViewUniformShaderParameters = TUniformBufferRef<FViewUniformShaderParameters>::CreateUniformBufferImmediate(*ViewInfo->CachedViewUniformShaderParameters, UniformBuffer_SingleDraw);
 	}
 
 	RHICmdList.BeginRenderPass(RPInfo, TEXT("HairVoxelization"));
@@ -329,12 +345,10 @@ static void RenderVoxelization(
 		Scene,
 		ViewInfo,
 		PrimitiveSceneInfo,
-		LightDesc,
 		bVoxelizeMaterial ? EHairStrandsRasterPassType::VoxelizationMaterial : EHairStrandsRasterPassType::Voxelization,
-		nullptr,
 		ViewportRect,
-		VoxelResources.MinAABB,
-		VoxelResources.MaxAABB);
+		ViewUniformShaderParameters,
+		DeepShadowPassUniformParameters);
 	RHICmdList.EndRenderPass();
 
 	if (GHairVoxelInjectOpaqueDepthEnable)
