Fix missing barrier in water pass. Render targets must be made writable.

diff --git a/Engine/Source/Runtime/Renderer/Private/SingleLayerWaterRendering.cpp b/Engine/Source/Runtime/Renderer/Private/SingleLayerWaterRendering.cpp
--- a/Engine/Source/Runtime/Renderer/Private/SingleLayerWaterRendering.cpp
+++ b/Engine/Source/Runtime/Renderer/Private/SingleLayerWaterRendering.cpp
@@ -540,6 +540,19 @@ void FDeferredShadingSceneRenderer::BeginRenderingWaterGBuffer(FRHICommandList&
 		SceneContext.BindVirtualTextureFeedbackUAV(RPInfo);
 	}
 
+	// Make the render targets writable
+	FRHITexture* TransitionRTs[MaxSimultaneousRenderTargets];
+	int32 NumColorRenderTargets = 0;
+	for (int32 Index = 0; Index < MaxSimultaneousRenderTargets; ++Index)
+	{
+		if (RPInfo.ColorRenderTargets[Index].RenderTarget)
+		{
+			TransitionRTs[NumColorRenderTargets++] = RPInfo.ColorRenderTargets[Index].RenderTarget;
+		}
+	}
+
+	RHICmdList.TransitionResources(EResourceTransitionAccess::EWritable, TransitionRTs, NumColorRenderTargets);
+
 	// Begin the pass
 	RHICmdList.BeginRenderPass(RPInfo, TEXT("WaterGBuffer"));
 
