Fix for MAC build.

diff --git a/Engine/Source/Runtime/Renderer/Private/SceneRendering.cpp b/Engine/Source/Runtime/Renderer/Private/SceneRendering.cpp
--- a/Engine/Source/Runtime/Renderer/Private/SceneRendering.cpp
+++ b/Engine/Source/Runtime/Renderer/Private/SceneRendering.cpp
@@ -58,7 +58,7 @@
 #include "IHeadMountedDisplay.h"
 #include "DiaphragmDOF.h" 
 #include "SingleLayerWaterRendering.h" 
-#include "HairStrands\HairStrandsUtils.h"
+#include "HairStrands/HairStrandsUtils.h"
 
 /*-----------------------------------------------------------------------------
 	Globals
