Fixed visibility pass to use the View hair render info from 9730968

diff --git a/Engine/Shaders/Private/HairStrands/HairStrandsVisibilityCommon.ush b/Engine/Shaders/Private/HairStrands/HairStrandsVisibilityCommon.ush
--- a/Engine/Shaders/Private/HairStrands/HairStrandsVisibilityCommon.ush
+++ b/Engine/Shaders/Private/HairStrands/HairStrandsVisibilityCommon.ush
@@ -315,7 +315,8 @@ FNodeData ConvertToNodeData(float Depth, FVisibilityData VisibilityData, float3
  {
 	 float RadiusAtDepth1Primary;
 	 float RadiusAtDepth1Velocity;
-	 bool bIsOrthoView;
+	 bool  bIsOrthoView;
+	 float VelocityMagnitudeScale;
  };
 
  FHairRenderingInfo GetHairRenderInfo(float4 ViewHairRenderInfo)
@@ -324,5 +325,6 @@ FNodeData ConvertToNodeData(float Depth, FVisibilityData VisibilityData, float3
 	 Info.RadiusAtDepth1Primary = ViewHairRenderInfo.x;
 	 Info.RadiusAtDepth1Velocity = ViewHairRenderInfo.y;
 	 Info.bIsOrthoView = ViewHairRenderInfo.z > 0.0f ? true : false;
+	 Info.VelocityMagnitudeScale = ViewHairRenderInfo.w;
 	 return Info;
  }
diff --git a/Engine/Shaders/Private/HairStrands/HairStrandsVisibilityPS.usf b/Engine/Shaders/Private/HairStrands/HairStrandsVisibilityPS.usf
--- a/Engine/Shaders/Private/HairStrands/HairStrandsVisibilityPS.usf
+++ b/Engine/Shaders/Private/HairStrands/HairStrandsVisibilityPS.usf
@@ -68,8 +68,9 @@ void MainVisibility(
 
 	float Coverage = 1;
 	{
+		FHairRenderingInfo HairRenderingInfo = GetHairRenderInfo(ResolvedView.HairRenderInfo);
 		const float SceneDepth = ConvertFromDeviceZ(SvPosition.z); // Linear depth in world unit
-		const float PixelRadius = SceneDepth * HairVisibilityPass.MinStrandRadius_Primary; // Not correct but the coverage is not used (we count instead the number of sub-sample covered)
+		const float PixelRadius = SceneDepth * HairRenderingInfo.RadiusAtDepth1Primary; // Not correct but the coverage is not used (we count instead the number of sub-sample covered)
 		const float StrandRealRadius = WorldStrandRadius;
 		Coverage = StrandRealRadius / max(StrandRealRadius, PixelRadius);
 	}
diff --git a/Engine/Shaders/Private/HairStrands/HairStrandsVisibilityVS.usf b/Engine/Shaders/Private/HairStrands/HairStrandsVisibilityVS.usf
--- a/Engine/Shaders/Private/HairStrands/HairStrandsVisibilityVS.usf
+++ b/Engine/Shaders/Private/HairStrands/HairStrandsVisibilityVS.usf
@@ -55,7 +55,8 @@ void Main(
 		Output.PackedVelocityC			= INVARIANT(PrevScreenPosObj);
 
 		const float2 ScreenVelocity = Calculate2DVelocity(Output.PackedVelocityA, Output.PackedVelocityC);
-		VelocityScale = saturate(max(abs(ScreenVelocity.x), abs(ScreenVelocity.y)) / HairVisibilityPass.HairStrandsVelocityScale);
+		FHairRenderingInfo HairRenderingInfo = GetHairRenderInfo(ResolvedView.HairRenderInfo);
+		VelocityScale = saturate(max(abs(ScreenVelocity.x), abs(ScreenVelocity.y)) / HairRenderingInfo.VelocityMagnitudeScale);
 	}
 	#endif
 
diff --git a/Engine/Source/Runtime/Renderer/Private/HairStrands/HairStrandsVisibility.cpp b/Engine/Source/Runtime/Renderer/Private/HairStrands/HairStrandsVisibility.cpp
--- a/Engine/Source/Runtime/Renderer/Private/HairStrands/HairStrandsVisibility.cpp
+++ b/Engine/Source/Runtime/Renderer/Private/HairStrands/HairStrandsVisibility.cpp
@@ -18,16 +18,6 @@ DECLARE_GPU_STAT(HairStrandsVisibility);
 
 /////////////////////////////////////////////////////////////////////////////////////////
 
-BEGIN_GLOBAL_SHADER_PARAMETER_STRUCT(FHairVisibilityPassUniformParameters, )
-SHADER_PARAMETER(float, MinStrandRadius_Primary)
-SHADER_PARAMETER(float, MinStrandRadius_Velocity)
-SHADER_PARAMETER(float, HairStrandsVelocityScale)
-SHADER_PARAMETER_TEXTURE(Texture2D<float>, MainDepthTexture)
-END_GLOBAL_SHADER_PARAMETER_STRUCT()
-IMPLEMENT_GLOBAL_SHADER_PARAMETER_STRUCT(FHairVisibilityPassUniformParameters, "HairVisibilityPass");
-
-/////////////////////////////////////////////////////////////////////////////////////////
-
 static int32 GHairStrandsCoveragePassEnable = 0;
 static FAutoConsoleVariableRef CVarHairStrandsCoveragePassEnable(TEXT("r.HairStrands.CoveragePass"), GHairStrandsCoveragePassEnable, TEXT("Enable accurate coverage pass"));
 
@@ -49,9 +39,6 @@ static FAutoConsoleVariableRef CVarHairClearVisibilityBuffer(TEXT("r.HairStrands
 static int32 GHairVelocityType = 1; // default is 
 static FAutoConsoleVariableRef CVarHairVelocityType(TEXT("r.HairStrands.VelocityType"), GHairVelocityType, TEXT("Type of velocity filtering (0:avg, 1:closest, 2:max). Default is 1."));
 
-static int32 GHairVelocityMagnitudeScale = 100; // Tuned by eye, based on heavy motion (strong head shack)
-static FAutoConsoleVariableRef CVarHairVelocityMagnitudeScale(TEXT("r.HairStrands.VelocityMagnitudeScale"), GHairVelocityMagnitudeScale, TEXT("Velocity magnitude (in pixel) at which a hair will reach its pic velocity-rasterization-scale under motion to reduce aliasing. Default is 100."));
-
 /////////////////////////////////////////////////////////////////////////////////////////
 
 enum EHairVisibilityRenderMode
@@ -89,8 +76,6 @@ protected:
 	{
 		ERHIFeatureLevel::Type FeatureLevel = GetMaxSupportedFeatureLevel((EShaderPlatform)Initializer.Target.Platform);
 		check(FSceneInterface::GetShadingPath(FeatureLevel) != EShadingPath::Mobile);
-		// deferred
-		PassUniformBuffer.Bind(Initializer.ParameterMap, FHairVisibilityPassUniformParameters::StaticStructMetadata.GetShaderVariableName());
 	}
 
 	FHairVisibilityVS() {}
@@ -131,7 +116,6 @@ public:
 	{
 		ERHIFeatureLevel::Type FeatureLevel = GetMaxSupportedFeatureLevel((EShaderPlatform)Initializer.Target.Platform);
 		check(FSceneInterface::GetShadingPath(FeatureLevel) != EShadingPath::Mobile);
-		PassUniformBuffer.Bind(Initializer.ParameterMap, FHairVisibilityPassUniformParameters::StaticStructMetadata.GetShaderVariableName());
 		HairVisibilityPass_HairClusterIndex.Bind(Initializer.ParameterMap, TEXT("HairVisibilityPass_HairClusterIndex"));
 	}
 
@@ -766,26 +750,27 @@ static void AddHairVisibilityCommonPass(
 		check(RHICmdList.IsInsideRenderPass());
 		check(IsInRenderingThread());
 
-		const FIntPoint Resolution(ViewInfo->ViewRect.Width(), ViewInfo->ViewRect.Height());
-		TUniformBufferRef<FHairVisibilityPassUniformParameters> PassUniformBuffer;
+		FMeshPassProcessorRenderState DrawRenderState(*ViewInfo);
+
+		if(RenderMode == HairVisibilityRenderMode_Coverage)
 		{
-			FVector2D PixelVelocity(1.f / (Resolution.X * 2), 1.f / (Resolution.Y * 2));
-			const float VelocityMagnitudeScale = FMath::Clamp(GHairVelocityMagnitudeScale, 0, 512) * FMath::Min(PixelVelocity.X, PixelVelocity.Y);
-
-			// Set the sample count to one as we want the size of the pixel
-			const uint32 HairVisibilitySampleCount = RenderMode == HairVisibilityRenderMode_MSAA ? GetHairVisibilitySampleCount() : 1;
-			const float RasterizationScaleOverride = RenderMode == HairVisibilityRenderMode_MSAA ? 0 : 1.35f;
-			FHairVisibilityPassUniformParameters PassUniformParameters;
-			FMinHairRadiusAtDepth1 MinHairRadius = ComputeMinStrandRadiusAtDepth1(FIntPoint(ViewInfo->UnconstrainedViewRect.Width(), ViewInfo->UnconstrainedViewRect.Height()), ViewInfo->FOV, HairVisibilitySampleCount, RasterizationScaleOverride);
-			PassUniformParameters.MinStrandRadius_Primary = MinHairRadius.Primary;
-			PassUniformParameters.MinStrandRadius_Velocity = MinHairRadius.Velocity;
-			PassUniformParameters.HairStrandsVelocityScale = VelocityMagnitudeScale;
-			PassUniformBuffer = CreateUniformBufferImmediate(PassUniformParameters, UniformBuffer_SingleDraw, EUniformBufferValidation::None);
+			// In the case we render coverage, we need to override some hair parameters to account for the change in MSAA sample count.
+			const uint32 HairVisibilitySampleCount = 1;
+			const float RasterizationScaleOverride = 1.35f;
+			FMinHairRadiusAtDepth1 MinHairRadius = ComputeMinStrandRadiusAtDepth1(
+				FIntPoint(ViewInfo->UnconstrainedViewRect.Width(), ViewInfo->UnconstrainedViewRect.Height()), ViewInfo->FOV, HairVisibilitySampleCount, RasterizationScaleOverride);
+
+			TUniformBufferRef<FViewUniformShaderParameters> ViewUniformShaderParameters;
+			ViewInfo->CachedViewUniformShaderParameters->HairRenderInfo.X = MinHairRadius.Primary;
+			ViewInfo->CachedViewUniformShaderParameters->HairRenderInfo.Y = MinHairRadius.Velocity;
+			ViewInfo->CachedViewUniformShaderParameters->HairRenderInfo.Z = ViewInfo->IsPerspectiveProjection() ? 0.0f : 1.0f;
+			ViewInfo->CachedViewUniformShaderParameters->ViewForward = ViewInfo->ViewMatrices.GetOverriddenTranslatedViewMatrix().GetColumn(2);
+			ViewUniformShaderParameters = TUniformBufferRef<FViewUniformShaderParameters>::CreateUniformBufferImmediate(*ViewInfo->CachedViewUniformShaderParameters, UniformBuffer_SingleFrame);
+			DrawRenderState.SetViewUniformBuffer(ViewUniformShaderParameters);
 		}
 
-		FMeshPassProcessorRenderState DrawRenderState(*ViewInfo, PassUniformBuffer);
 		{
-			RHICmdList.SetViewport(0, 0, 0.0f, Resolution.X, Resolution.Y, 1.0f);
+			RHICmdList.SetViewport(0, 0, 0.0f, ViewInfo->ViewRect.Width(), ViewInfo->ViewRect.Height(), 1.0f);
 			if (RenderMode == HairVisibilityRenderMode_MSAA)
 			{
 				DrawRenderState.SetBlendState(TStaticBlendState<
diff --git a/Engine/Source/Runtime/Renderer/Private/SceneRendering.cpp b/Engine/Source/Runtime/Renderer/Private/SceneRendering.cpp
--- a/Engine/Source/Runtime/Renderer/Private/SceneRendering.cpp
+++ b/Engine/Source/Runtime/Renderer/Private/SceneRendering.cpp
@@ -329,6 +329,11 @@ TAutoConsoleVariable<int32> CVarTransientResourceAliasing_Buffers(
 	TEXT("Enables transient resource aliasing for specified buffers. Used only if GSupportsTransientResourceAliasing is true.\n"),
 	ECVF_ReadOnly);
 
+static TAutoConsoleVariable<int32> CVarHairVelocityMagnitudeScale(
+	TEXT("r.HairStrands.VelocityMagnitudeScale"), 
+	100,  // Tuned by eye, based on heavy motion (strong head shack)
+	TEXT("Velocity magnitude (in pixel) at which a hair will reach its pic velocity-rasterization-scale under motion to reduce aliasing. Default is 100."));
+
 #if !UE_BUILD_SHIPPING
 
 static TAutoConsoleVariable<int32> CVarTestInternalViewRectOffset(
@@ -1671,9 +1676,14 @@ void FViewInfo::SetupUniformBufferParameters(
 	{
 		const FMinHairRadiusAtDepth1 MinRadiusAtDepth1 = ComputeMinStrandRadiusAtDepth1(
 			FIntPoint(UnconstrainedViewRect.Width(), UnconstrainedViewRect.Height()), FOV, GetHairVisibilitySampleCount(), 0.0f);
+
+		FVector2D PixelVelocity(1.f / (ViewRect.Width() * 2), 1.f / (ViewRect.Height() * 2));
+		const float VelocityMagnitudeScale = FMath::Clamp(CVarHairVelocityMagnitudeScale.GetValueOnAnyThread(), 0, 512) * FMath::Min(PixelVelocity.X, PixelVelocity.Y);
+
 		ViewUniformShaderParameters.HairRenderInfo.X = MinRadiusAtDepth1.Primary;
 		ViewUniformShaderParameters.HairRenderInfo.Y = MinRadiusAtDepth1.Velocity;
 		ViewUniformShaderParameters.HairRenderInfo.Z = IsPerspectiveProjection() ? 0.0f : 1.0f;
+		ViewUniformShaderParameters.HairRenderInfo.W = VelocityMagnitudeScale;
 	}
 }
 
