Fix DX validation error.

diff --git a/Engine/Shaders/Private/HairStrands/HairStrandsEnvironmentLightingCompose.usf b/Engine/Shaders/Private/HairStrands/HairStrandsEnvironmentLightingCompose.usf
--- a/Engine/Shaders/Private/HairStrands/HairStrandsEnvironmentLightingCompose.usf
+++ b/Engine/Shaders/Private/HairStrands/HairStrandsEnvironmentLightingCompose.usf
@@ -3,9 +3,9 @@
 #include "../Common.ush"
 #include "HairStrandsVisibilityCommon.ush"
 
-Buffer<float4>	 HairLightingSampleBuffer;
-Texture2D<uint4> HairCategorizationTexture;
-Texture2D<uint>	 HairVisibilityNodeOffsetAndCount;
+StructuredBuffer<float4>	HairLightingSampleBuffer;
+Texture2D<uint4>			HairCategorizationTexture;
+Texture2D<uint>				HairVisibilityNodeOffsetAndCount;
 
 void MainPS(
 	in FScreenVertexOutput Input,
diff --git a/Engine/Source/Runtime/Renderer/Private/HairStrands/HairStrandsEnvironment.cpp b/Engine/Source/Runtime/Renderer/Private/HairStrands/HairStrandsEnvironment.cpp
--- a/Engine/Source/Runtime/Renderer/Private/HairStrands/HairStrandsEnvironment.cpp
+++ b/Engine/Source/Runtime/Renderer/Private/HairStrands/HairStrandsEnvironment.cpp
@@ -163,7 +163,7 @@ static void AddHairEnvironmentLightingComposePass(
 	FHairEnvironmentLightingComposePS::FParameters* Parameters = GraphBuilder.AllocParameters<FHairEnvironmentLightingComposePS::FParameters>();
 	Parameters->HairCategorizationTexture = GraphBuilder.RegisterExternalTexture(VisibilityData.CategorizationTexture);
 	Parameters->HairVisibilityNodeOffsetAndCount = GraphBuilder.RegisterExternalTexture(VisibilityData.NodeIndex);
-	Parameters->HairLightingSampleBuffer  = GraphBuilder.CreateSRV(SampleLightingBuffer);
+	Parameters->HairLightingSampleBuffer  = GraphBuilder.CreateSRV(FRDGBufferSRVDesc(SampleLightingBuffer, EPixelFormat::PF_A32B32G32R32F));
 	Parameters->RenderTargets[0] = FRenderTargetBinding(OutColorTexture, ERenderTargetLoadAction::ELoad);
 	Parameters->RenderTargets[1] = FRenderTargetBinding(OutSubColorTexture, ERenderTargetLoadAction::ELoad);
 
