Change light samples to be stored into  a 16bits buffer.

diff --git a/Engine/Shaders/Private/HairStrands/HairStrandsEnvironmentLighting.usf b/Engine/Shaders/Private/HairStrands/HairStrandsEnvironmentLighting.usf
--- a/Engine/Shaders/Private/HairStrands/HairStrandsEnvironmentLighting.usf
+++ b/Engine/Shaders/Private/HairStrands/HairStrandsEnvironmentLighting.usf
@@ -465,9 +465,9 @@ void MainPS(
 
 #if COMPUTESHADER
 uint MaxVisibilityNodeCount;
-StructuredBuffer<uint>		HairVisibilityNodeCoord;
-Buffer<uint4>				IndirectArgsBuffer;
-RWStructuredBuffer<float4> OutLightingBuffer;
+StructuredBuffer<uint>	HairVisibilityNodeCoord;
+Buffer<uint4>			IndirectArgsBuffer;
+RWBuffer<float4>		OutLightingBuffer;
 
 #if PERMUTATION_GROUP_SIZE == 0
 #define LIGHTING_GROUP_SIZE 64
diff --git a/Engine/Shaders/Private/HairStrands/HairStrandsEnvironmentLightingCompose.usf b/Engine/Shaders/Private/HairStrands/HairStrandsEnvironmentLightingCompose.usf
--- a/Engine/Shaders/Private/HairStrands/HairStrandsEnvironmentLightingCompose.usf
+++ b/Engine/Shaders/Private/HairStrands/HairStrandsEnvironmentLightingCompose.usf
@@ -3,9 +3,9 @@
 #include "../Common.ush"
 #include "HairStrandsVisibilityCommon.ush"
 
-StructuredBuffer<float4>	HairLightingSampleBuffer;
-Texture2D<uint4>			HairCategorizationTexture;
-Texture2D<uint>				HairVisibilityNodeOffsetAndCount;
+Buffer<float4>		HairLightingSampleBuffer;
+Texture2D<uint4>	HairCategorizationTexture;
+Texture2D<uint>		HairVisibilityNodeOffsetAndCount;
 
 void MainPS(
 	in FScreenVertexOutput Input,
diff --git a/Engine/Source/Runtime/Renderer/Private/HairStrands/HairStrandsEnvironment.cpp b/Engine/Source/Runtime/Renderer/Private/HairStrands/HairStrandsEnvironment.cpp
--- a/Engine/Source/Runtime/Renderer/Private/HairStrands/HairStrandsEnvironment.cpp
+++ b/Engine/Source/Runtime/Renderer/Private/HairStrands/HairStrandsEnvironment.cpp
@@ -163,7 +163,7 @@ static void AddHairEnvironmentLightingComposePass(
 	FHairEnvironmentLightingComposePS::FParameters* Parameters = GraphBuilder.AllocParameters<FHairEnvironmentLightingComposePS::FParameters>();
 	Parameters->HairCategorizationTexture = GraphBuilder.RegisterExternalTexture(VisibilityData.CategorizationTexture);
 	Parameters->HairVisibilityNodeOffsetAndCount = GraphBuilder.RegisterExternalTexture(VisibilityData.NodeIndex);
-	Parameters->HairLightingSampleBuffer  = GraphBuilder.CreateSRV(FRDGBufferSRVDesc(SampleLightingBuffer, EPixelFormat::PF_A32B32G32R32F));
+	Parameters->HairLightingSampleBuffer  = GraphBuilder.CreateSRV(FRDGBufferSRVDesc(SampleLightingBuffer, EPixelFormat::PF_FloatRGBA));
 	Parameters->RenderTargets[0] = FRenderTargetBinding(OutColorTexture, ERenderTargetLoadAction::ELoad);
 	Parameters->RenderTargets[1] = FRenderTargetBinding(OutSubColorTexture, ERenderTargetLoadAction::ELoad);
 
@@ -274,8 +274,8 @@ static FRDGBufferRef AddHairStrandsEnvironmentLightingPassCS(
 	const uint32 NodeGroupSize,
 	FRDGBufferRef IndirectArgsBuffer)
 {
-	FRDGBufferRef OutBuffer = GraphBuilder.CreateBuffer(FRDGBufferDesc::CreateStructuredDesc(
-		4 * sizeof(float),
+	FRDGBufferRef OutBuffer = GraphBuilder.CreateBuffer(FRDGBufferDesc::CreateBufferDesc(
+		4 * sizeof(uint16),
 		VisibilityData.NodeData->Desc.NumElements),
 		TEXT("HairSkyLightingNodeData"));
 
@@ -313,7 +313,7 @@ static FRDGBufferRef AddHairStrandsEnvironmentLightingPassCS(
 	}
 	PassParameters->ForwardLightData = View.ForwardLightingResources->ForwardLightDataUniformBuffer;
 	PassParameters->IndirectArgsBuffer = IndirectArgsBuffer;
-	PassParameters->OutLightingBuffer = GraphBuilder.CreateUAV(OutBuffer);
+	PassParameters->OutLightingBuffer = GraphBuilder.CreateUAV(FRDGBufferUAVDesc(OutBuffer, PF_FloatRGBA));
 
 	// Bind hair data
 	FRDGBufferRef InHairVisibilityNodeData = GraphBuilder.RegisterExternalBuffer(VisibilityData.NodeData, TEXT("HairVisibilityNodeData"));
