Do not override raster scale during hair coverage pass for a better match.

diff --git a/Engine/Source/Runtime/Renderer/Private/HairStrands/HairStrandsVisibility.cpp b/Engine/Source/Runtime/Renderer/Private/HairStrands/HairStrandsVisibility.cpp
--- a/Engine/Source/Runtime/Renderer/Private/HairStrands/HairStrandsVisibility.cpp
+++ b/Engine/Source/Runtime/Renderer/Private/HairStrands/HairStrandsVisibility.cpp
@@ -754,9 +754,9 @@ static void AddHairVisibilityCommonPass(
 
 		if(RenderMode == HairVisibilityRenderMode_Coverage)
 		{
-			// In the case we render coverage, we need to override some hair parameters to account for the change in MSAA sample count.
-			const uint32 HairVisibilitySampleCount = 1;
-			const float RasterizationScaleOverride = 1.35f;
+			// In the case we render coverage, we need to override some view uniform shader parameters to account for the change in MSAA sample count.
+			const uint32 HairVisibilitySampleCount = 1;		// The coverage pass does not use MSAA
+			const float RasterizationScaleOverride = 0.0f;	// no override
 			FMinHairRadiusAtDepth1 MinHairRadius = ComputeMinStrandRadiusAtDepth1(
 				FIntPoint(ViewInfo->UnconstrainedViewRect.Width(), ViewInfo->UnconstrainedViewRect.Height()), ViewInfo->FOV, HairVisibilitySampleCount, RasterizationScaleOverride);
 
