Added support for the Coverage pass, for more accurate and consistent hair coverage at any distance.

diff --git a/Engine/Shaders/Private/DeferredLightPixelShaders.usf b/Engine/Shaders/Private/DeferredLightPixelShaders.usf
--- a/Engine/Shaders/Private/DeferredLightPixelShaders.usf
+++ b/Engine/Shaders/Private/DeferredLightPixelShaders.usf
@@ -141,7 +141,7 @@ void ComputeHairLighting(
 	in FInputParams InputParams,
 	in FNodeDesc NodeDesc,
 	in FCategorizationData CatData,
-	inout float OutOpaqueWeight,
+	inout float OutHairCoverage,
 	inout float4 OutColor, 
 	inout float4 OutColorSubPixel)
 {
@@ -216,23 +216,20 @@ void ComputeHairLighting(
 		const float LocalCoverage = From8bitCoverage(SampleVisData.Coverage8bit);
 		LocalColor += PercepualColor * LocalCoverage;
 	}
+	const float PixelCoverage = saturate(CatData.PixelCoverage);
+	OutHairCoverage = PixelCoverage;
 
-	const float PixelCoverage = min(CatData.PixelCoverage, 1);
-	float Coverage = 0;
 	OutColor = 0;
 	OutColorSubPixel = 0;
-	if (PixelCoverage < 1)
+	if (OutHairCoverage < 1)
 	{
-		OutColorSubPixel = LocalColor; // * PixelCoverage <-- this would be need if Coverage8bit was storing the ratio (sampleCoverage / validCoverage)
-		OutColorSubPixel.a = PixelCoverage;
-		Coverage = 0;
+		OutColorSubPixel = LocalColor * PixelCoverage;
+		OutColorSubPixel.a = OutHairCoverage;
 	}
 	else
 	{
 		OutColor = LocalColor;
-		Coverage = 1;
 	}
-	OutOpaqueWeight = 1 - Coverage;
 }
 
 #endif // USE_HAIR_LIGHTING
@@ -275,7 +272,7 @@ void DeferredLightPixelMain(
 	
 	// Only light pixels marked as using deferred shading
 	OutColorSubPixel = 0;
-	float OpaqueWeight = 1;
+	float HairCoverage = 0;
 	BRANCH 
 	if (NodeDesc.Count > 0
 	#if USE_LIGHTING_CHANNELS
@@ -283,7 +280,7 @@ void DeferredLightPixelMain(
 	#endif
 		)
 	{
-		ComputeHairLighting(InputParams, NodeDesc, CatData, OpaqueWeight, OutColor, OutColorSubPixel);
+		ComputeHairLighting(InputParams, NodeDesc, CatData, HairCoverage, OutColor, OutColorSubPixel);
 	#if USE_PREEXPOSURE
 		OutColorSubPixel.rgb *= View.PreExposure;
 	#endif
@@ -312,7 +309,8 @@ void DeferredLightPixelMain(
 		const float  Attenuation = ComputeLightProfileMultiplier(DerivedParams.WorldPosition, DeferredLightUniforms.Position, -DeferredLightUniforms.Direction, DeferredLightUniforms.Tangent);
 
 #if USE_HAIR_LIGHTING
-		OutColor += LinearToPerceptual(Radiance * Attenuation) * OpaqueWeight;
+		float OpaqueVisibility = 1.0f - HairCoverage;
+		OutColor += LinearToPerceptual(Radiance * Attenuation) * OpaqueVisibility;
 		OutColor  = PerceptualToLinear(OutColor);
 #else
 		OutColor += (Radiance * Attenuation);
diff --git a/Engine/Shaders/Private/HairStrands/HairStrandsEnvironmentLighting.usf b/Engine/Shaders/Private/HairStrands/HairStrandsEnvironmentLighting.usf
--- a/Engine/Shaders/Private/HairStrands/HairStrandsEnvironmentLighting.usf
+++ b/Engine/Shaders/Private/HairStrands/HairStrandsEnvironmentLighting.usf
@@ -515,14 +515,14 @@ void MainCS(uint3 DispatchThreadId : SV_DispatchThreadID)
 	const bool bIsClusterCompatible = SampleVisData.ClusterId == Voxel_ClusterId;
 	if (bIsClusterCompatible)
 	{
-		const float SampleCoverage = From8bitCoverage(SampleVisData.Coverage8bit);
+		const float SampleCoverage = saturate(From8bitCoverage(SampleVisData.Coverage8bit));
 		const float3 SkyLighting = ReflectionEnvironment(GBuffer, ScreenPosition, SvPosition);
 
 		FLightAccumulator LightAccumulator = (FLightAccumulator)0;
 		LightAccumulator_Add(LightAccumulator, SkyLighting, SkyLighting, 1.0f, false);
 		float4 OutColor = LightAccumulator_GetResult(LightAccumulator);
-		OutColor.xyz *= saturate(SampleCoverage);
-		OutColor.a = saturate(SampleCoverage);
+		OutColor.xyz *= SampleCoverage;
+		OutColor.a = SampleCoverage;
 
 		OutLightingBuffer[LocalOffset] = OutColor;
 	}
diff --git a/Engine/Shaders/Private/HairStrands/HairStrandsEnvironmentLightingCompose.usf b/Engine/Shaders/Private/HairStrands/HairStrandsEnvironmentLightingCompose.usf
--- a/Engine/Shaders/Private/HairStrands/HairStrandsEnvironmentLightingCompose.usf
+++ b/Engine/Shaders/Private/HairStrands/HairStrandsEnvironmentLightingCompose.usf
@@ -32,7 +32,7 @@ void MainPS(
 			if (PixelCoverage >= 1)
 				OutColor += LightingSample;
 			else
-				OutColorSubPixel += LightingSample;
+				OutColorSubPixel += LightingSample * PixelCoverage;
 		}		
 	}
 }
diff --git a/Engine/Shaders/Private/HairStrands/HairStrandsVisibilityCompaction.usf b/Engine/Shaders/Private/HairStrands/HairStrandsVisibilityCompaction.usf
--- a/Engine/Shaders/Private/HairStrands/HairStrandsVisibilityCompaction.usf
+++ b/Engine/Shaders/Private/HairStrands/HairStrandsVisibilityCompaction.usf
@@ -210,7 +210,8 @@ void MainCS(uint GroupIndex : SV_GroupIndex, uint3 DispatchThreadId : SV_Dispatc
 					const float SampleDepth = MSAA_DepthTexture.Load(PixelCoord, SampleIt);
 					const uint4 EncodedData = MSAA_IDTexture.Load(PixelCoord, SampleIt);
 					FVisibilityData VisibilityData = DecodeVisibilityData(EncodedData);
-					VisibilityData.Coverage8bit = To8bitCoverage(Samples[OutIndex].y / float(SampleDesc.HairSampleCount));
+					// VisibilityData.Coverage8bit is a weight normalising to 1 the contribution of all the compacted samples. Because later it is weighted by Categorization.PixelCoverage.
+					VisibilityData.Coverage8bit = To8bitCoverage(Samples[OutIndex].y / float(SampleDesc.ValidSampleCount));
 					Categorization.TotalCoverage8bit += VisibilityData.Coverage8bit;
 					Categorization.ClosestDepth = max(Categorization.ClosestDepth, SampleDepth); // Inverse-Z
 
diff --git a/Engine/Shaders/Private/HairStrands/HairStrandsVisibilityDepthPS.usf b/Engine/Shaders/Private/HairStrands/HairStrandsVisibilityDepthPS.usf
--- a/Engine/Shaders/Private/HairStrands/HairStrandsVisibilityDepthPS.usf
+++ b/Engine/Shaders/Private/HairStrands/HairStrandsVisibilityDepthPS.usf
@@ -9,6 +9,10 @@
 
 Texture2DMS<float, HAIR_VISIBILITY_MAXSAMPLECOUNT> HairVisibilityDepthTexture;
 
+#if PERMUTATION_COVERAGE == 1
+Texture2D<float> CoverageTexture;
+#endif
+
 void MainPS(
 	in FScreenVertexOutput Input,
 	out float SvDepth		: SV_Depth,
@@ -29,7 +33,11 @@ void MainPS(
 		ValidSampleCount += bIsValid ? 1 : 0;
 	}
 
-	if (ValidSampleCount < HAIR_VISIBILITY_MAXSAMPLECOUNT || MaxDepth == 0)
+	if (ValidSampleCount < HAIR_VISIBILITY_MAXSAMPLECOUNT
+#if PERMUTATION_COVERAGE == 1
+		|| CoverageTexture.Load(uint3(PixelCoord, 0))<1.0f
+#endif
+		)
 		discard;
 
 	SvDepth = MinDepth;
diff --git a/Engine/Source/Runtime/Renderer/Private/HairStrands/HairStrandsVisibility.cpp b/Engine/Source/Runtime/Renderer/Private/HairStrands/HairStrandsVisibility.cpp
--- a/Engine/Source/Runtime/Renderer/Private/HairStrands/HairStrandsVisibility.cpp
+++ b/Engine/Source/Runtime/Renderer/Private/HairStrands/HairStrandsVisibility.cpp
@@ -956,8 +956,12 @@ class FHairVisibilityDepthPS : public FGlobalShader
 	DECLARE_GLOBAL_SHADER(FHairVisibilityDepthPS);
 	SHADER_USE_PARAMETER_STRUCT(FHairVisibilityDepthPS, FGlobalShader);
 
+	class FCoverage : SHADER_PERMUTATION_INT("PERMUTATION_COVERAGE", 2);
+	using FPermutationDomain = TShaderPermutationDomain<FCoverage>;
+
 	BEGIN_SHADER_PARAMETER_STRUCT(FParameters, )
 		SHADER_PARAMETER_RDG_TEXTURE(Texture2D, HairVisibilityDepthTexture)
+		SHADER_PARAMETER_RDG_TEXTURE(Texture2D, CoverageTexture)
 		RENDER_TARGET_BINDING_SLOTS()
 	END_SHADER_PARAMETER_STRUCT()
 
@@ -971,12 +975,14 @@ static void AddHairVisibilityColorAndDepthPatchPass(
 	FRDGBuilder& GraphBuilder,
 	const FViewInfo& View,
 	const FRDGTextureRef& VisibilityDepthTexture,
+	const FRDGTextureRef& CoverageTexture,
 	FRDGTextureRef& OutGBufferBTexture,
 	FRDGTextureRef& OutColorTexture,
 	FRDGTextureRef& OutDepthTexture)
 {
 	FHairVisibilityDepthPS::FParameters* Parameters = GraphBuilder.AllocParameters<FHairVisibilityDepthPS::FParameters>();
 	Parameters->HairVisibilityDepthTexture = VisibilityDepthTexture;
+	Parameters->CoverageTexture = CoverageTexture;
 	Parameters->RenderTargets[0] = FRenderTargetBinding(OutGBufferBTexture, ERenderTargetLoadAction::ELoad);
 	Parameters->RenderTargets[1] = FRenderTargetBinding(OutColorTexture, ERenderTargetLoadAction::ELoad);
 	Parameters->RenderTargets.DepthStencil = FDepthStencilBinding(
@@ -986,7 +992,9 @@ static void AddHairVisibilityColorAndDepthPatchPass(
 		FExclusiveDepthStencil::DepthWrite_StencilNop);
 
 	TShaderMapRef<FPostProcessVS> VertexShader(View.ShaderMap);
-	TShaderMapRef<FHairVisibilityDepthPS> PixelShader(View.ShaderMap);
+	FHairVisibilityDepthPS::FPermutationDomain PermutationVector;
+	PermutationVector.Set<FHairVisibilityDepthPS::FCoverage>(CoverageTexture ? 1 : 0);
+	TShaderMapRef<FHairVisibilityDepthPS> PixelShader(View.ShaderMap, PermutationVector);
 	const TShaderMap<FGlobalShaderType>* GlobalShaderMap = View.ShaderMap;
 	const FIntRect Viewport = View.ViewRect;
 	const FIntPoint Resolution = OutDepthTexture->Desc.Extent;
@@ -1164,6 +1172,7 @@ FHairStrandsVisibilityViews RenderHairStrandsVisibilityBuffer(
 				GraphBuilder,
 				View,
 				MsaaVisibilityResources.DepthTexture,
+				CoverageTexture,
 				SceneGBufferBTexture,
 				SceneColorTexture,
 				SceneDepthTexture);
