Back out of changelist 9694919

diff --git a/Engine/Source/Runtime/Renderer/Private/MobileShadingRenderer.cpp b/Engine/Source/Runtime/Renderer/Private/MobileShadingRenderer.cpp
--- a/Engine/Source/Runtime/Renderer/Private/MobileShadingRenderer.cpp
+++ b/Engine/Source/Runtime/Renderer/Private/MobileShadingRenderer.cpp
@@ -527,9 +527,9 @@ void FMobileSceneRenderer::Render(FRHICommandListImmediate& RHICmdList)
 			DepthTargetAction = EDepthStencilTargetActions::ClearDepthStencil_StoreDepthStencil;
 		}
 						
-		if ((bKeepDepthContent && !bMobileMSAA) || Scene->Decals.Num() > 0)
+		if (bKeepDepthContent && !bMobileMSAA)
 		{
-			// store depth if post-processing/capture/decals needs it
+			// store depth if post-processing/capture needs it
 			DepthTargetAction = EDepthStencilTargetActions::ClearDepthStencil_StoreDepthStencil;
 		}
 	}
