Texture buffer-> 2D Texture for GPU Scene Mobile device
This was done because of MALI GPU limitation of 64kb texture buffer

- Changed default Precision of compute shaders to highp
- Added Texture Buffer GPU Support mobile
- Texture path controlled by GPUSceneUseTexture2D(requires recompile of shaders)

diff --git a/Engine/Shaders/Private/ByteBuffer.usf b/Engine/Shaders/Private/ByteBuffer.usf
--- a/Engine/Shaders/Private/ByteBuffer.usf
+++ b/Engine/Shaders/Private/ByteBuffer.usf
@@ -6,9 +6,13 @@ float4 Value;
 uint Size;
 uint SrcOffset;
 uint DstOffset;
+uint FloatsPerLine;
+
 
 StructuredBuffer<float4> SrcBuffer;
 RWStructuredBuffer<float4> DstBuffer;
+Texture2D<float4>	SrcTexture;
+RWTexture2D<float4>	DstTexture;
 
 [numthreads(THREADGROUP_SIZE, 1, 1)]
 void MemsetBufferCS( uint3 DispatchThreadId : SV_DispatchThreadID ) 
@@ -33,6 +37,20 @@ void MemcpyBufferCS(uint3 DispatchThreadId : SV_DispatchThreadID)
 	}
 }
 
+[numthreads(THREADGROUP_SIZE, 1, 1)]
+void MemcpyTextureToTextureCS(uint3 DispatchThreadId : SV_DispatchThreadID) 
+{
+	uint2 IndexTexture;
+	IndexTexture.y = DispatchThreadId.x / (FloatsPerLine);
+	IndexTexture.x = DispatchThreadId.x % (FloatsPerLine);
+
+	if (IndexTexture.x * IndexTexture.y  < Size)
+	{
+		float4 SrcValue = SrcTexture.Load(float3(IndexTexture.x, IndexTexture.y, 0)); 
+		DstTexture[IndexTexture.xy] = SrcValue;
+	}
+}
+
 uint NumScatters;
 Buffer<float4> UploadBuffer;
 Buffer<uint> ScatterBuffer;
@@ -48,4 +66,24 @@ void ScatterCopyCS( uint3 DispatchThreadId : SV_DispatchThreadID )
 		uint SrcIndex = ScatterIndex;
 		DstBuffer[DestIndex] = UploadBuffer.Load(SrcIndex);
 	}
+}
+
+[numthreads(THREADGROUP_SIZE, 1, 1)]
+void ScatterCopyTextureCS( uint3 DispatchThreadId : SV_DispatchThreadID ) 
+{
+	uint ScatterIndex = DispatchThreadId.x;
+	
+	if (ScatterIndex < NumScatters)
+	{
+		uint DestIndex = ScatterBuffer.Load(ScatterIndex);
+		uint2 IndexTexture;
+		IndexTexture.y = DestIndex / (FloatsPerLine);
+		IndexTexture.x = DestIndex % (FloatsPerLine);
+		uint SrcIndex = ScatterIndex;
+
+		if (IndexTexture.x * IndexTexture.y  < Size)
+		{
+			DstTexture[IndexTexture.xy] = UploadBuffer.Load(SrcIndex);
+		}
+	}
 }
\ No newline at end of file
diff --git a/Engine/Shaders/Private/SceneData.ush b/Engine/Shaders/Private/SceneData.ush
--- a/Engine/Shaders/Private/SceneData.ush
+++ b/Engine/Shaders/Private/SceneData.ush
@@ -45,7 +45,7 @@ struct FPrimitiveSceneData
 #define PRIMITIVE_SCENE_DATA_STRIDE 35
 
 // Fetch from scene primitive buffer
-FPrimitiveSceneData GetPrimitiveData(uint PrimitiveId) 
+FPrimitiveSceneData GetPrimitiveData_Buffer(uint PrimitiveId)
 {
 	// Note: layout must match FPrimitiveSceneShaderData in C++
 	// Relying on optimizer to remove unused loads
@@ -110,6 +110,93 @@ FPrimitiveSceneData GetPrimitiveData(uint PrimitiveId)
 	return PrimitiveData;
 }
 
+// Fetch from scene primitive texture2D
+FPrimitiveSceneData GetPrimitiveData_Texture2D(uint PrimitiveId)
+{
+	// Note: layout must match FPrimitiveSceneShaderData in C++
+	// Relying on optimizer to remove unused loads
+	FPrimitiveSceneData PrimitiveData;
+
+	uint LinePrimitive = (PrimitiveId & 0xFFFF0000) >> 16;
+	uint OffsetPrimitive = PrimitiveId & 0xFFFF;
+	OffsetPrimitive = OffsetPrimitive * PRIMITIVE_SCENE_DATA_STRIDE;
+
+	PrimitiveData.LocalToWorld[0] = View.PrimitiveSceneDataTexture.Load(float3(OffsetPrimitive + 0, LinePrimitive, 0));
+	PrimitiveData.LocalToWorld[1] = View.PrimitiveSceneDataTexture.Load(float3(OffsetPrimitive + 1, LinePrimitive, 0));
+	PrimitiveData.LocalToWorld[2] = View.PrimitiveSceneDataTexture.Load(float3(OffsetPrimitive + 2, LinePrimitive, 0));
+	PrimitiveData.LocalToWorld[3] = View.PrimitiveSceneDataTexture.Load(float3(OffsetPrimitive + 3, LinePrimitive, 0));
+
+
+	PrimitiveData.InvNonUniformScaleAndDeterminantSign = View.PrimitiveSceneDataTexture.Load(float3(OffsetPrimitive + 4, LinePrimitive, 0));
+	PrimitiveData.ObjectWorldPositionAndRadius = View.PrimitiveSceneDataTexture.Load(float3(OffsetPrimitive + 5, LinePrimitive, 0));
+
+	PrimitiveData.WorldToLocal[0] = View.PrimitiveSceneDataTexture.Load(float3(OffsetPrimitive + 6, LinePrimitive, 0));
+	PrimitiveData.WorldToLocal[1] = View.PrimitiveSceneDataTexture.Load(float3(OffsetPrimitive + 7, LinePrimitive, 0));
+	PrimitiveData.WorldToLocal[2] = View.PrimitiveSceneDataTexture.Load(float3(OffsetPrimitive + 8, LinePrimitive, 0));
+	PrimitiveData.WorldToLocal[3] = View.PrimitiveSceneDataTexture.Load(float3(OffsetPrimitive + 9, LinePrimitive, 0));
+
+	PrimitiveData.PreviousLocalToWorld[0] = View.PrimitiveSceneDataTexture.Load(float3(OffsetPrimitive + 10, LinePrimitive, 0));
+	PrimitiveData.PreviousLocalToWorld[1] = View.PrimitiveSceneDataTexture.Load(float3(OffsetPrimitive + 11, LinePrimitive, 0));
+	PrimitiveData.PreviousLocalToWorld[2] = View.PrimitiveSceneDataTexture.Load(float3(OffsetPrimitive + 12, LinePrimitive, 0));
+	PrimitiveData.PreviousLocalToWorld[3] = View.PrimitiveSceneDataTexture.Load(float3(OffsetPrimitive + 13, LinePrimitive, 0));
+
+	PrimitiveData.PreviousWorldToLocal[0] = View.PrimitiveSceneDataTexture.Load(float3(OffsetPrimitive + 14, LinePrimitive, 0));
+	PrimitiveData.PreviousWorldToLocal[1] = View.PrimitiveSceneDataTexture.Load(float3(OffsetPrimitive + 15, LinePrimitive, 0));
+	PrimitiveData.PreviousWorldToLocal[2] = View.PrimitiveSceneDataTexture.Load(float3(OffsetPrimitive + 16, LinePrimitive, 0));
+	PrimitiveData.PreviousWorldToLocal[3] = View.PrimitiveSceneDataTexture.Load(float3(OffsetPrimitive + 17, LinePrimitive, 0));
+
+
+	PrimitiveData.ActorWorldPosition = View.PrimitiveSceneDataTexture.Load(float3(OffsetPrimitive + 18, LinePrimitive, 0)).xyz;
+	PrimitiveData.UseSingleSampleShadowFromStationaryLights = View.PrimitiveSceneDataTexture.Load(float3(OffsetPrimitive + 18, LinePrimitive, 0)).w;
+
+	PrimitiveData.ObjectBounds = View.PrimitiveSceneDataTexture.Load(float3(OffsetPrimitive + 19, LinePrimitive, 0)).xyz;
+	PrimitiveData.LpvBiasMultiplier = View.PrimitiveSceneDataTexture.Load(float3(OffsetPrimitive + 19, LinePrimitive, 0)).w;
+
+	PrimitiveData.DecalReceiverMask = View.PrimitiveSceneDataTexture.Load(float3(OffsetPrimitive + 20, LinePrimitive, 0)).x;
+	PrimitiveData.PerObjectGBufferData = View.PrimitiveSceneDataTexture.Load(float3(OffsetPrimitive + 20, LinePrimitive, 0)).y;
+	PrimitiveData.UseVolumetricLightmapShadowFromStationaryLights = View.PrimitiveSceneDataTexture.Load(float3(OffsetPrimitive + 20, LinePrimitive, 0)).z;
+	PrimitiveData.DrawsVelocity = View.PrimitiveSceneDataTexture.Load(float3(OffsetPrimitive + 20, LinePrimitive, 0)).w;
+
+	PrimitiveData.ObjectOrientation = View.PrimitiveSceneDataTexture.Load(float3(OffsetPrimitive + 21, LinePrimitive, 0));
+	PrimitiveData.NonUniformScale = View.PrimitiveSceneDataTexture.Load(float3(OffsetPrimitive + 22, LinePrimitive, 0));
+
+	PrimitiveData.LocalObjectBoundsMin = View.PrimitiveSceneDataTexture.Load(float3(OffsetPrimitive + 23, LinePrimitive, 0)).xyz;
+	PrimitiveData.LightingChannelMask = asuint(View.PrimitiveSceneDataTexture.Load(float3(OffsetPrimitive + 23, LinePrimitive, 0)).w);
+
+	PrimitiveData.LocalObjectBoundsMax = View.PrimitiveSceneDataTexture.Load(float3(OffsetPrimitive + 24, LinePrimitive, 0)).xyz;
+	PrimitiveData.LightmapDataIndex = asuint(View.PrimitiveSceneDataTexture.Load(float3(OffsetPrimitive + 24, LinePrimitive, 0)).w);
+
+	PrimitiveData.PreSkinnedLocalBoundsMin = View.PrimitiveSceneDataTexture.Load(float3(OffsetPrimitive + 25, LinePrimitive, 0)).xyz;
+	PrimitiveData.SingleCaptureIndex = asuint(View.PrimitiveSceneDataTexture.Load(float3(OffsetPrimitive + 25, LinePrimitive, 0)).w);
+
+	PrimitiveData.PreSkinnedLocalBoundsMax = View.PrimitiveSceneDataTexture.Load(float3(OffsetPrimitive + 26, LinePrimitive, 0)).xyz;
+	PrimitiveData.OutputVelocity = asuint(View.PrimitiveSceneDataTexture.Load(float3(OffsetPrimitive + 26, LinePrimitive, 0)).w);
+
+	UNROLL
+	for (int i = 0; i < NUM_CUSTOM_PRIMITIVE_DATA; i++)
+	{
+		PrimitiveData.CustomPrimitiveData[i] = View.PrimitiveSceneDataTexture.Load(float3(OffsetPrimitive + 27 + i, LinePrimitive, 0));
+	}
+
+	return PrimitiveData;
+}
+
+FPrimitiveSceneData GetPrimitiveData(uint PrimitiveId) 
+{
+	// Note: layout must match FPrimitiveSceneShaderData in C++
+	// Relying on optimizer to remove unused loads
+
+	#if VF_GPU_SCENE_BUFFER
+
+	return GetPrimitiveData_Buffer(PrimitiveId);
+	
+	#else
+	
+	return GetPrimitiveData_Texture2D(PrimitiveId);
+	
+	#endif
+}
+
 #else
 
 // Route to Primitive uniform buffer
diff --git a/Engine/Source/Developer/ShaderFormatOpenGL/Private/GlslBackend.cpp b/Engine/Source/Developer/ShaderFormatOpenGL/Private/GlslBackend.cpp
--- a/Engine/Source/Developer/ShaderFormatOpenGL/Private/GlslBackend.cpp
+++ b/Engine/Source/Developer/ShaderFormatOpenGL/Private/GlslBackend.cpp
@@ -3585,7 +3585,7 @@ char* FGlslCodeBackend::GenerateCode(exec_list* ir, _mesa_glsl_parse_state* stat
 	FixRedundantCasts(ir);
 //IRDump(ir);
 
-	const bool bDefaultPrecisionIsHalf = ((HlslCompileFlags & HLSLCC_UseFullPrecisionInPS) == 0);
+	const bool bDefaultPrecisionIsHalf = ((HlslCompileFlags & HLSLCC_UseFullPrecisionInPS) == 0 && Frequency != HSF_ComputeShader);
 	const bool bUsesExternalTexture = ((HlslCompileFlags & HLSLCC_UsesExternalTexture) == HLSLCC_UsesExternalTexture);
 	
 	FBreakPrecisionChangesVisitor BreakPrecisionChangesVisitor(state, bDefaultPrecisionIsHalf);
diff --git a/Engine/Source/Developer/ShaderFormatOpenGL/Private/ShaderFormatOpenGL.cpp b/Engine/Source/Developer/ShaderFormatOpenGL/Private/ShaderFormatOpenGL.cpp
--- a/Engine/Source/Developer/ShaderFormatOpenGL/Private/ShaderFormatOpenGL.cpp
+++ b/Engine/Source/Developer/ShaderFormatOpenGL/Private/ShaderFormatOpenGL.cpp
@@ -24,7 +24,7 @@ class FShaderFormatGLSL : public IShaderFormat
 	{
 		/** Version for shader format, this becomes part of the DDC key. */
 		UE_SHADER_GLSL_VER = 82,
-		UE_SHADER_GLSL_ANDROID_VER = 81,
+		UE_SHADER_GLSL_ANDROID_VER = 82,
 	};
 
 	void CheckFormat(FName Format) const
diff --git a/Engine/Source/Runtime/Core/Private/HAL/ConsoleManager.cpp b/Engine/Source/Runtime/Core/Private/HAL/ConsoleManager.cpp
--- a/Engine/Source/Runtime/Core/Private/HAL/ConsoleManager.cpp
+++ b/Engine/Source/Runtime/Core/Private/HAL/ConsoleManager.cpp
@@ -2158,6 +2158,13 @@ static TAutoConsoleVariable<int32> CVarMobileSupportGPUScene(
 	ECVF_ReadOnly | ECVF_RenderThreadSafe
 );
 
+static TAutoConsoleVariable<int32> CVarMobileGPUSceneUseTexture2D(
+	TEXT("r.Mobile.UseGPUSceneTexture"),
+	0,
+	TEXT("Use a Texture2D instead of TextureBuffer for GPUScene.\n"),
+	ECVF_RenderThreadSafe
+);
+
 static TAutoConsoleVariable<int32> CVarSetClearSceneMethod(
 	TEXT("r.ClearSceneMethod"),
 	1,
diff --git a/Engine/Source/Runtime/Engine/Private/LocalVertexFactory.cpp b/Engine/Source/Runtime/Engine/Private/LocalVertexFactory.cpp
--- a/Engine/Source/Runtime/Engine/Private/LocalVertexFactory.cpp
+++ b/Engine/Source/Runtime/Engine/Private/LocalVertexFactory.cpp
@@ -206,6 +206,7 @@ void FLocalVertexFactory::ModifyCompilationEnvironment(const FVertexFactoryType*
 	}
 
 	OutEnvironment.SetDefine(TEXT("VF_SUPPORTS_PRIMITIVE_SCENE_DATA"), Type->SupportsPrimitiveIdStream() && UseGPUScene(Platform, GetMaxSupportedFeatureLevel(Platform)));
+	OutEnvironment.SetDefine(TEXT("VF_GPU_SCENE_BUFFER"), Type->SupportsPrimitiveIdStream() && UseGPUScene(Platform, GetMaxSupportedFeatureLevel(Platform)) && !GPUSceneUseTexture2D(Platform));
 }
 
 void FLocalVertexFactory::ValidateCompiledResult(const FVertexFactoryType* Type, EShaderPlatform Platform, const FShaderParameterMap& ParameterMap, TArray<FString>& OutErrors)
diff --git a/Engine/Source/Runtime/Engine/Private/PrimitiveUniformShaderParameters.cpp b/Engine/Source/Runtime/Engine/Private/PrimitiveUniformShaderParameters.cpp
--- a/Engine/Source/Runtime/Engine/Private/PrimitiveUniformShaderParameters.cpp
+++ b/Engine/Source/Runtime/Engine/Private/PrimitiveUniformShaderParameters.cpp
@@ -4,13 +4,21 @@
 #include "PrimitiveSceneProxy.h"
 #include "PrimitiveSceneInfo.h"
 
-void FSinglePrimitiveStructuredBuffer::InitRHI() 
+void FSinglePrimitiveStructured::InitRHI() 
 {
 	if (RHISupportsComputeShaders(GMaxRHIShaderPlatform))
 	{
 		FRHIResourceCreateInfo CreateInfo;
-		PrimitiveSceneDataBufferRHI = RHICreateStructuredBuffer(sizeof(FVector4), FPrimitiveSceneShaderData::PrimitiveDataStrideInFloat4s * sizeof(FVector4), BUF_Static | BUF_ShaderResource, CreateInfo);
-		PrimitiveSceneDataBufferSRV = RHICreateShaderResourceView(PrimitiveSceneDataBufferRHI);
+
+		{	
+			PrimitiveSceneDataBufferRHI = RHICreateStructuredBuffer(sizeof(FVector4), FPrimitiveSceneShaderData::PrimitiveDataStrideInFloat4s * sizeof(FVector4), BUF_Static | BUF_ShaderResource, CreateInfo);
+			PrimitiveSceneDataBufferSRV = RHICreateShaderResourceView(PrimitiveSceneDataBufferRHI);
+		}
+
+		{
+			PrimitiveSceneDataTextureRHI = RHICreateTexture2D(FPrimitiveSceneShaderData::PrimitiveDataStrideInFloat4s, 1, PF_A32B32G32R32F, 1, 1, TexCreate_ShaderResource | TexCreate_UAV, CreateInfo);
+			PrimitiveSceneDataTextureSRV = RHICreateShaderResourceView(PrimitiveSceneDataTextureRHI, 0);
+		}
 
 		LightmapSceneDataBufferRHI = RHICreateStructuredBuffer(sizeof(FVector4), FLightmapSceneShaderData::LightmapDataStrideInFloat4s * sizeof(FVector4), BUF_Static | BUF_ShaderResource, CreateInfo);
 		LightmapSceneDataBufferSRV = RHICreateShaderResourceView(LightmapSceneDataBufferRHI);
@@ -19,15 +27,25 @@ void FSinglePrimitiveStructuredBuffer::InitRHI()
 	UploadToGPU();
 }
 
-void FSinglePrimitiveStructuredBuffer::UploadToGPU()
+void FSinglePrimitiveStructured::UploadToGPU()
 {
 	if (RHISupportsComputeShaders(GMaxRHIShaderPlatform))
 	{
 		void* LockedData = nullptr;
 
-		LockedData = RHILockStructuredBuffer(PrimitiveSceneDataBufferRHI, 0, FPrimitiveSceneShaderData::PrimitiveDataStrideInFloat4s * sizeof(FVector4), RLM_WriteOnly);
-		FPlatformMemory::Memcpy(LockedData, PrimitiveSceneData.Data, FPrimitiveSceneShaderData::PrimitiveDataStrideInFloat4s * sizeof(FVector4));
-		RHIUnlockStructuredBuffer(PrimitiveSceneDataBufferRHI);
+		if (!GPUSceneUseTexture2D(ShaderPlatform))
+		{
+			LockedData = RHILockStructuredBuffer(PrimitiveSceneDataBufferRHI, 0, FPrimitiveSceneShaderData::PrimitiveDataStrideInFloat4s * sizeof(FVector4), RLM_WriteOnly);
+			FPlatformMemory::Memcpy(LockedData, PrimitiveSceneData.Data, FPrimitiveSceneShaderData::PrimitiveDataStrideInFloat4s * sizeof(FVector4));
+			RHIUnlockStructuredBuffer(PrimitiveSceneDataBufferRHI);
+		}
+		else
+		{
+			uint32 SrcStride;
+			void* LockedData = RHILockTexture2D(PrimitiveSceneDataTextureRHI, 0, RLM_WriteOnly, SrcStride, false);
+			FPlatformMemory::Memcpy(LockedData, PrimitiveSceneData.Data, FPrimitiveSceneShaderData::PrimitiveDataStrideInFloat4s * sizeof(FVector4));
+			RHIUnlockTexture2D(PrimitiveSceneDataTextureRHI, 0, false);
+		}
 
 		LockedData = RHILockStructuredBuffer(LightmapSceneDataBufferRHI, 0, FLightmapSceneShaderData::LightmapDataStrideInFloat4s * sizeof(FVector4), RLM_WriteOnly);
 		FPlatformMemory::Memcpy(LockedData, LightmapSceneData.Data, FLightmapSceneShaderData::LightmapDataStrideInFloat4s * sizeof(FVector4));
@@ -35,8 +53,8 @@ void FSinglePrimitiveStructuredBuffer::UploadToGPU()
 	}
 }
 
-TGlobalResource<FSinglePrimitiveStructuredBuffer> GIdentityPrimitiveBuffer;
-TGlobalResource<FSinglePrimitiveStructuredBuffer> GTilePrimitiveBuffer;
+TGlobalResource<FSinglePrimitiveStructured> GIdentityPrimitiveBuffer;
+TGlobalResource<FSinglePrimitiveStructured> GTilePrimitiveBuffer;
 
 FPrimitiveSceneShaderData::FPrimitiveSceneShaderData(const FPrimitiveSceneProxy* RESTRICT Proxy)
 {
@@ -128,3 +146,10 @@ void FPrimitiveSceneShaderData::Setup(const FPrimitiveUniformShaderParameters& P
 		Data[CustomPrimitiveDataStartIndex + i] = PrimitiveUniformShaderParameters.CustomPrimitiveData[i];
 	}
 }
+
+uint16 FPrimitiveSceneShaderData::GetPrimitivesPerTextureLine()
+{
+	// @todo texture size limit over 65536, revisit this in the future :). Currently you can have(with primitiveData = 35 floats4) a max of 122,683,392 primitives
+	uint16 PrimitivesPerTextureLine = FMath::Min((int32)MAX_uint16, (int32)GMaxTextureDimensions) / (FPrimitiveSceneShaderData::PrimitiveDataStrideInFloat4s);
+	return PrimitivesPerTextureLine;
+}
\ No newline at end of file
diff --git a/Engine/Source/Runtime/Engine/Private/SceneManagement.cpp b/Engine/Source/Runtime/Engine/Private/SceneManagement.cpp
--- a/Engine/Source/Runtime/Engine/Private/SceneManagement.cpp
+++ b/Engine/Source/Runtime/Engine/Private/SceneManagement.cpp
@@ -747,6 +747,7 @@ FViewUniformShaderParameters::FViewUniformShaderParameters()
 	CameraAerialPerspectiveVolume = GBlackAlpha1VolumeTexture->TextureRHI;
 	CameraAerialPerspectiveVolumeSampler = TStaticSamplerState<SF_Bilinear>::GetRHI();
 
+	PrimitiveSceneDataTexture = GIdentityPrimitiveBuffer.PrimitiveSceneDataTextureRHI;
 	PrimitiveSceneData = GIdentityPrimitiveBuffer.PrimitiveSceneDataBufferSRV;
 	LightmapSceneData = GIdentityPrimitiveBuffer.LightmapSceneDataBufferSRV;
 
diff --git a/Engine/Source/Runtime/Engine/Public/PrimitiveUniformShaderParameters.h b/Engine/Source/Runtime/Engine/Public/PrimitiveUniformShaderParameters.h
--- a/Engine/Source/Runtime/Engine/Public/PrimitiveUniformShaderParameters.h
+++ b/Engine/Source/Runtime/Engine/Public/PrimitiveUniformShaderParameters.h
@@ -224,9 +224,10 @@ struct FPrimitiveSceneShaderData
 	ENGINE_API FPrimitiveSceneShaderData(const class FPrimitiveSceneProxy* RESTRICT Proxy);
 
 	ENGINE_API void Setup(const FPrimitiveUniformShaderParameters& PrimitiveUniformShaderParameters);
+	ENGINE_API static uint16 GetPrimitivesPerTextureLine();
 };
 
-class ENGINE_VTABLE FSinglePrimitiveStructuredBuffer : public FRenderResource
+class ENGINE_VTABLE FSinglePrimitiveStructured : public FRenderResource
 {
 public:
 
@@ -236,18 +237,25 @@ public:
 	{
 		PrimitiveSceneDataBufferRHI.SafeRelease();
 		PrimitiveSceneDataBufferSRV.SafeRelease();
+		PrimitiveSceneDataTextureRHI.SafeRelease();
+		PrimitiveSceneDataTextureSRV.SafeRelease();
 		LightmapSceneDataBufferRHI.SafeRelease();
 		LightmapSceneDataBufferSRV.SafeRelease();
 	}
 
 	ENGINE_API void UploadToGPU();
 
+	EShaderPlatform ShaderPlatform;
+
 	FPrimitiveSceneShaderData PrimitiveSceneData;
 	FLightmapSceneShaderData LightmapSceneData;
 
 	FStructuredBufferRHIRef PrimitiveSceneDataBufferRHI;
 	FShaderResourceViewRHIRef PrimitiveSceneDataBufferSRV;
 
+	FTexture2DRHIRef PrimitiveSceneDataTextureRHI;
+	FShaderResourceViewRHIRef PrimitiveSceneDataTextureSRV;
+
 	FStructuredBufferRHIRef LightmapSceneDataBufferRHI;
 	FShaderResourceViewRHIRef LightmapSceneDataBufferSRV;
 };
@@ -256,5 +264,5 @@ public:
 * Default Primitive data buffer.  
 * This is used when the VF is used for rendering outside normal mesh passes, where there is no valid scene.
 */
-extern ENGINE_API TGlobalResource<FSinglePrimitiveStructuredBuffer> GIdentityPrimitiveBuffer;
-extern ENGINE_API TGlobalResource<FSinglePrimitiveStructuredBuffer> GTilePrimitiveBuffer;
+extern ENGINE_API TGlobalResource<FSinglePrimitiveStructured> GIdentityPrimitiveBuffer;
+extern ENGINE_API TGlobalResource<FSinglePrimitiveStructured> GTilePrimitiveBuffer;
diff --git a/Engine/Source/Runtime/Engine/Public/SceneView.h b/Engine/Source/Runtime/Engine/Public/SceneView.h
--- a/Engine/Source/Runtime/Engine/Public/SceneView.h
+++ b/Engine/Source/Runtime/Engine/Public/SceneView.h
@@ -795,6 +795,7 @@ BEGIN_GLOBAL_SHADER_PARAMETER_STRUCT_WITH_CONSTRUCTOR(FViewUniformShaderParamete
 	SHADER_PARAMETER_TEXTURE(Texture2D, PreIntegratedBRDF)
 	SHADER_PARAMETER_SAMPLER(SamplerState, PreIntegratedBRDFSampler)
 	SHADER_PARAMETER_SRV(StructuredBuffer<float4>, PrimitiveSceneData)
+	SHADER_PARAMETER_TEXTURE(Texture2D<float4>, PrimitiveSceneDataTexture)
 	SHADER_PARAMETER_SRV(StructuredBuffer<float4>, LightmapSceneData)
 
 	SHADER_PARAMETER_TEXTURE(Texture2D, TransmittanceLutTexture)
diff --git a/Engine/Source/Runtime/OpenGLDrv/Private/Android/AndroidOpenGL.h b/Engine/Source/Runtime/OpenGLDrv/Private/Android/AndroidOpenGL.h
--- a/Engine/Source/Runtime/OpenGLDrv/Private/Android/AndroidOpenGL.h
+++ b/Engine/Source/Runtime/OpenGLDrv/Private/Android/AndroidOpenGL.h
@@ -78,6 +78,9 @@ extern PFNBLITFRAMEBUFFERNVPROC					glBlitFramebufferNV ;
 #define GL_SAMPLES_PASSED_EXT					0x8914
 #define GL_ANY_SAMPLES_PASSED_EXT				0x8C2F
 
+#ifndef GL_MAX_TEXTURE_BUFFER_SIZE
+#define GL_MAX_TEXTURE_BUFFER_SIZE				0x8C2B
+#endif
 
 typedef void (GL_APIENTRYP PFNGLGENQUERIESEXTPROC) (GLsizei n, GLuint *ids);
 typedef void (GL_APIENTRYP PFNGLDELETEQUERIESEXTPROC) (GLsizei n, const GLuint *ids);
diff --git a/Engine/Source/Runtime/OpenGLDrv/Private/OpenGLCommands.cpp b/Engine/Source/Runtime/OpenGLDrv/Private/OpenGLCommands.cpp
--- a/Engine/Source/Runtime/OpenGLDrv/Private/OpenGLCommands.cpp
+++ b/Engine/Source/Runtime/OpenGLDrv/Private/OpenGLCommands.cpp
@@ -719,7 +719,7 @@ void FOpenGLDynamicRHI::SetupTexturesForDraw( FOpenGLContextState& ContextState,
 			
 				if (ContextState.SamplerStates[TextureStageIndex] != PendingSampler)
 				{
-					FOpenGL::BindSampler(TextureStageIndex, PendingSampler->Resource);
+					FOpenGL::BindSampler(TextureStageIndex, PendingSampler ? PendingSampler->Resource : 0);
 					ContextState.SamplerStates[TextureStageIndex] = PendingSampler;
 				}
 			}
@@ -2474,6 +2474,9 @@ FORCEINLINE void SetResource(FOpenGLDynamicRHI* RESTRICT OpenGLRHI, uint32 BindI
 	{
 		OpenGLRHI->InternalSetShaderTexture(Texture, nullptr, GetFirstTextureUnit<Frequency>() + BindIndex, 0, 0, 0, -1);
 	}
+
+	// clear any previous sampler state
+	OpenGLRHI->InternalSetSamplerStates(GetFirstTextureUnit<Frequency>() + BindIndex, nullptr);
 }
 
 template <EShaderFrequency Frequency>
diff --git a/Engine/Source/Runtime/OpenGLDrv/Private/OpenGLDevice.cpp b/Engine/Source/Runtime/OpenGLDrv/Private/OpenGLDevice.cpp
--- a/Engine/Source/Runtime/OpenGLDrv/Private/OpenGLDevice.cpp
+++ b/Engine/Source/Runtime/OpenGLDrv/Private/OpenGLDevice.cpp
@@ -696,6 +696,9 @@ static void InitRHICapabilitiesForGL()
 #define LOG_AND_GET_GL_INT_TEMP(IntEnum,Default) GLint Value_##IntEnum = Default; if (IntEnum) {glGetIntegerv(IntEnum, &Value_##IntEnum); glGetError();} else {Value_##IntEnum = Default;} UE_LOG(LogRHI, Log, TEXT("  " #IntEnum ": %d"), Value_##IntEnum)
 #endif
 	LOG_AND_GET_GL_INT_TEMP(GL_MAX_TEXTURE_SIZE, 0);
+#if defined(GL_MAX_TEXTURE_BUFFER_SIZE)
+	LOG_AND_GET_GL_INT_TEMP(GL_MAX_TEXTURE_BUFFER_SIZE, 0);
+#endif
 	LOG_AND_GET_GL_INT_TEMP(GL_MAX_CUBE_MAP_TEXTURE_SIZE, 0);
 #if defined(GL_MAX_ARRAY_TEXTURE_LAYERS) && GL_MAX_ARRAY_TEXTURE_LAYERS
 	LOG_AND_GET_GL_INT_TEMP(GL_MAX_ARRAY_TEXTURE_LAYERS, 0);
@@ -737,6 +740,9 @@ static void InitRHICapabilitiesForGL()
 	GMaxOpenGLColorSamples = Value_GL_MAX_COLOR_TEXTURE_SAMPLES;
 	GMaxOpenGLDepthSamples = Value_GL_MAX_DEPTH_TEXTURE_SAMPLES;
 	GMaxOpenGLIntegerSamples = Value_GL_MAX_INTEGER_SAMPLES;
+#if defined(GL_MAX_TEXTURE_BUFFER_SIZE)
+	GMaxTextureBufferSize = Value_GL_MAX_TEXTURE_BUFFER_SIZE;
+#endif
 
 	// Verify some assumptions.
 	// Android seems like reports one color attachment even when it supports MRT
diff --git a/Engine/Source/Runtime/RHI/Private/RHI.cpp b/Engine/Source/Runtime/RHI/Private/RHI.cpp
--- a/Engine/Source/Runtime/RHI/Private/RHI.cpp
+++ b/Engine/Source/Runtime/RHI/Private/RHI.cpp
@@ -545,6 +545,7 @@ FVertexElementTypeSupportInfo GVertexElementTypeSupport;
 RHI_API int32 volatile GCurrentTextureMemorySize = 0;
 RHI_API int32 volatile GCurrentRendertargetMemorySize = 0;
 RHI_API int64 GTexturePoolSize = 0 * 1024 * 1024;
+RHI_API int64 GMaxTextureBufferSize = 0;
 RHI_API int32 GPoolSizeVRAMPercentage = 0;
 
 RHI_API EShaderPlatform GShaderPlatformForFeatureLevel[ERHIFeatureLevel::Num] = {SP_NumPlatforms,SP_NumPlatforms,SP_NumPlatforms,SP_NumPlatforms};
diff --git a/Engine/Source/Runtime/RHI/Public/RHI.h b/Engine/Source/Runtime/RHI/Public/RHI.h
--- a/Engine/Source/Runtime/RHI/Public/RHI.h
+++ b/Engine/Source/Runtime/RHI/Public/RHI.h
@@ -478,6 +478,8 @@ extern RHI_API volatile int32 GCurrentTextureMemorySize;
 extern RHI_API volatile int32 GCurrentRendertargetMemorySize;
 /** Current texture streaming pool size, in bytes. 0 means unlimited. */
 extern RHI_API int64 GTexturePoolSize;
+/** Maximum texture buffer size. Relevant for OpenGL. 0 means we do not take it into account*/
+extern RHI_API int64 GMaxTextureBufferSize;
 /** In percent. If non-zero, the texture pool size is a percentage of GTotalGraphicsMemory. */
 extern RHI_API int32 GPoolSizeVRAMPercentage;
 
diff --git a/Engine/Source/Runtime/RHI/Public/RHIUtilities.h b/Engine/Source/Runtime/RHI/Public/RHIUtilities.h
--- a/Engine/Source/Runtime/RHI/Public/RHIUtilities.h
+++ b/Engine/Source/Runtime/RHI/Public/RHIUtilities.h
@@ -73,7 +73,8 @@ struct FTextureRWBuffer2D
 	}
 
 	// @param AdditionalUsage passed down to RHICreateVertexBuffer(), get combined with "BUF_UnorderedAccess | BUF_ShaderResource" e.g. BUF_Static
-	void Initialize(const uint32 BytesPerElement, const uint32 SizeX, const uint32 SizeY, const EPixelFormat Format)
+	const static uint32 DefaultTextureInitFlag = TexCreate_ShaderResource | TexCreate_UAV;
+	void Initialize(const uint32 BytesPerElement, const uint32 SizeX, const uint32 SizeY, const EPixelFormat Format, uint32 Flags = DefaultTextureInitFlag)
 	{
 		check(GMaxRHIFeatureLevel == ERHIFeatureLevel::SM5
 			|| IsVulkanPlatform(GMaxRHIShaderPlatform)
@@ -88,7 +89,7 @@ struct FTextureRWBuffer2D
 			SizeX, SizeY, Format, //PF_R32_FLOAT,
 			/*NumMips=*/ 1,
 			1,
-			/*Flags=*/ TexCreate_ShaderResource | TexCreate_UAV,
+			Flags,
 			/*BulkData=*/ CreateInfo);
 
 							
diff --git a/Engine/Source/Runtime/RenderCore/Private/RenderUtils.cpp b/Engine/Source/Runtime/RenderCore/Private/RenderUtils.cpp
--- a/Engine/Source/Runtime/RenderCore/Private/RenderUtils.cpp
+++ b/Engine/Source/Runtime/RenderCore/Private/RenderUtils.cpp
@@ -915,6 +915,23 @@ RENDERCORE_API bool MobileSupportsGPUScene(EShaderPlatform Platform)
 	return (CVar && CVar->GetValueOnAnyThread() != 0) ? true : false;
 }
 
+RENDERCORE_API bool GPUSceneUseTexture2D(EShaderPlatform Platform)
+{
+	if (IsMobilePlatform(Platform))
+	{
+		static TConsoleVariableData<int32>* CVar = IConsoleManager::Get().FindTConsoleVariableDataInt(TEXT("r.Mobile.UseGPUSceneTexture"));
+		if (Platform == SP_OPENGL_ES3_1_ANDROID)
+		{
+			return true;
+		}
+		else
+		{
+			return (CVar && CVar->GetValueOnAnyThread() != 0) ? true : false;
+		}
+	}
+	return false;
+}
+
 RENDERCORE_API bool AllowPixelDepthOffset(EShaderPlatform Platform)
 {
 	if (IsMobilePlatform(Platform))
diff --git a/Engine/Source/Runtime/RenderCore/Private/Shader.cpp b/Engine/Source/Runtime/RenderCore/Private/Shader.cpp
--- a/Engine/Source/Runtime/RenderCore/Private/Shader.cpp
+++ b/Engine/Source/Runtime/RenderCore/Private/Shader.cpp
@@ -2290,9 +2290,21 @@ void ShaderMapAppendKeyString(EShaderPlatform Platform, FString& KeyString)
 		{
 			// make it per shader platform ?
 			static IConsoleVariable* CVar = IConsoleManager::Get().FindConsoleVariable(TEXT("r.Mobile.SupportGPUScene"));
-			KeyString += (CVar && CVar->GetInt() != 0) ? TEXT("_MobGPUSc") : TEXT("");
+			bool bMobileGpuScene = (CVar && CVar->GetInt() != 0);
+			KeyString += bMobileGpuScene ? TEXT("_MobGPUSc") : TEXT("");
+			if (bMobileGpuScene)
+			{
+				// Mobile specific verify if we are using texturebuffer or texture2D
+				if (!GPUSceneUseTexture2D(Platform))
+				{
+					KeyString += TEXT("_TexBuf");
+				}
+				else
+				{
+					KeyString += TEXT("_Tex2D");
+				}
+			}
 		}
-		
 	}
 
 	const FName ShaderFormatName = LegacyShaderPlatformToShaderFormat(Platform);
diff --git a/Engine/Source/Runtime/RenderCore/Public/RenderUtils.h b/Engine/Source/Runtime/RenderCore/Public/RenderUtils.h
--- a/Engine/Source/Runtime/RenderCore/Public/RenderUtils.h
+++ b/Engine/Source/Runtime/RenderCore/Public/RenderUtils.h
@@ -435,6 +435,8 @@ RENDERCORE_API bool IsSimpleForwardShadingEnabled(EShaderPlatform Platform);
 
 RENDERCORE_API bool MobileSupportsGPUScene(EShaderPlatform Platform);
 
+RENDERCORE_API bool GPUSceneUseTexture2D(EShaderPlatform Platform);
+
 RENDERCORE_API bool AllowPixelDepthOffset(EShaderPlatform Platform);
 
 /** Returns if ForwardShading is enabled. Only valid for the current platform (otherwise call ITargetPlatform::UsesForwardShading()). */
diff --git a/Engine/Source/Runtime/Renderer/Private/ByteBuffer.cpp b/Engine/Source/Runtime/Renderer/Private/ByteBuffer.cpp
--- a/Engine/Source/Runtime/Renderer/Private/ByteBuffer.cpp
+++ b/Engine/Source/Runtime/Renderer/Private/ByteBuffer.cpp
@@ -151,6 +151,84 @@ void MemcpyBuffer(FRHICommandList& RHICmdList, const FRWBufferStructured& SrcBuf
 	RHICmdList.TransitionResource(EResourceTransitionAccess::ERWBarrier, EResourceTransitionPipeline::EComputeToCompute, DstBuffer.UAV);
 }
 
+class FMemcpyTextureToTextureCS : public FGlobalShader
+{
+	DECLARE_SHADER_TYPE(FMemcpyTextureToTextureCS, Global)
+
+	FMemcpyTextureToTextureCS() {}
+
+	static bool ShouldCompilePermutation(const FGlobalShaderPermutationParameters& Parameters)
+	{
+		return RHISupportsComputeShaders(Parameters.Platform);
+	}
+
+	static void ModifyCompilationEnvironment(const FGlobalShaderPermutationParameters& Parameters, FShaderCompilerEnvironment& OutEnvironment)
+	{
+		FGlobalShader::ModifyCompilationEnvironment(Parameters, OutEnvironment);
+		OutEnvironment.SetDefine(TEXT("THREADGROUP_SIZE"), ThreadGroupSize);
+	}
+
+public:
+	enum { ThreadGroupSize = 64 };
+
+	FShaderParameter			Size;
+	FShaderParameter			SrcOffset;
+	FShaderParameter			DstOffset;
+	FShaderParameter			FloatsPerLine;
+	FShaderResourceParameter	SrcTexture;
+	FShaderResourceParameter	DstTexture;
+
+	FMemcpyTextureToTextureCS(const ShaderMetaType::CompiledShaderInitializerType& Initializer)
+		: FGlobalShader(Initializer)
+	{
+		Size.Bind(Initializer.ParameterMap, TEXT("Size"));
+		SrcOffset.Bind(Initializer.ParameterMap, TEXT("SrcOffset"));
+		DstOffset.Bind(Initializer.ParameterMap, TEXT("DstOffset"));
+		FloatsPerLine.Bind(Initializer.ParameterMap, TEXT("FloatsPerLine"));
+		SrcTexture.Bind(Initializer.ParameterMap, TEXT("SrcTexture"));
+		DstTexture.Bind(Initializer.ParameterMap, TEXT("DstTexture"));
+	}
+
+	virtual bool Serialize(FArchive& Ar) override
+	{
+		bool bShaderHasOutdatedParameters = FGlobalShader::Serialize(Ar);
+		Ar << Size;
+		Ar << SrcOffset;
+		Ar << DstOffset;
+		Ar << FloatsPerLine;
+		Ar << SrcTexture;
+		Ar << DstTexture;
+		return bShaderHasOutdatedParameters;
+	}
+};
+
+IMPLEMENT_SHADER_TYPE(, FMemcpyTextureToTextureCS, TEXT("/Engine/Private/ByteBuffer.usf"), TEXT("MemcpyTextureToTextureCS"), SF_Compute);
+
+void MemcpyTextureToTexture(FRHICommandList& RHICmdList, const FTextureRWBuffer2D& SrcBuffer, const FTextureRWBuffer2D& DstBuffer, uint32 SrcOffset, uint32 DstOffset, uint32 NumFloat4s, uint32 FloatsPerLine)
+{
+	auto ShaderMap = GetGlobalShaderMap(GMaxRHIFeatureLevel);
+
+	TShaderMapRef< FMemcpyTextureToTextureCS > ComputeShader(ShaderMap);
+
+	FRHIComputeShader* ShaderRHI = ComputeShader->GetComputeShader();
+	RHICmdList.SetComputeShader(ShaderRHI);
+
+	RHICmdList.TransitionResource(EResourceTransitionAccess::EWritable, EResourceTransitionPipeline::EGfxToCompute, DstBuffer.UAV);
+
+	SetShaderValue(RHICmdList, ShaderRHI, ComputeShader->SrcOffset, SrcOffset);
+	SetShaderValue(RHICmdList, ShaderRHI, ComputeShader->DstOffset, DstOffset);
+	SetShaderValue(RHICmdList, ShaderRHI, ComputeShader->FloatsPerLine, FloatsPerLine);
+	SetShaderValue(RHICmdList, ShaderRHI, ComputeShader->Size, NumFloat4s);
+	SetSRVParameter(RHICmdList, ShaderRHI, ComputeShader->SrcTexture, SrcBuffer.SRV);
+	SetUAVParameter(RHICmdList, ShaderRHI, ComputeShader->DstTexture, DstBuffer.UAV);
+
+	RHICmdList.DispatchComputeShader(FMath::DivideAndRoundUp< uint32 >(NumFloat4s, FMemcpyBufferCS::ThreadGroupSize), 1, 1);
+
+	SetUAVParameter(RHICmdList, ShaderRHI, ComputeShader->DstTexture, FUnorderedAccessViewRHIRef());
+
+	RHICmdList.TransitionResource(EResourceTransitionAccess::ERWBarrier, EResourceTransitionPipeline::EComputeToCompute, DstBuffer.UAV);
+}
+
 bool ResizeBufferIfNeeded(FRHICommandList& RHICmdList, FRWBufferStructured& Buffer, uint32 NumFloat4s)
 {
 	EPixelFormat BufferFormat = PF_A32B32G32R32F;
@@ -176,6 +254,31 @@ bool ResizeBufferIfNeeded(FRHICommandList& RHICmdList, FRWBufferStructured& Buff
 	return false;
 }
 
+bool ResizeTextureIfNeeded(FRHICommandList& RHICmdList, FTextureRWBuffer2D& Texture, uint32 NumFloat4s, uint32 PrimitiveStride)
+{
+	EPixelFormat BufferFormat = PF_A32B32G32R32F;
+	uint32 BytesPerElement = GPixelFormats[BufferFormat].BlockBytes;
+
+	uint32 TotalSize = 0;
+	uint16 PrimitivesPerTextureLine = FPrimitiveSceneShaderData::GetPrimitivesPerTextureLine();
+	uint32 SizeY = (NumFloat4s / (PrimitiveStride * PrimitivesPerTextureLine)) + 1;
+
+	if (Texture.NumBytes == 0)
+	{
+		Texture.Initialize(BytesPerElement, PrimitiveStride * PrimitivesPerTextureLine, SizeY, PF_A32B32G32R32F, TexCreate_RenderTargetable | TexCreate_UAV);
+	}
+	else if ((SizeY * PrimitiveStride * PrimitivesPerTextureLine * BytesPerElement) != Texture.NumBytes)
+	{
+		FTextureRWBuffer2D NewTexture;
+		NewTexture.Initialize(BytesPerElement, PrimitiveStride * PrimitivesPerTextureLine, SizeY, PF_A32B32G32R32F, TexCreate_RenderTargetable | TexCreate_UAV);
+		MemcpyTextureToTexture(RHICmdList, Texture, NewTexture, 0, 0, NumFloat4s, PrimitiveStride * PrimitivesPerTextureLine);
+		Texture = NewTexture;
+		return true;
+	}
+
+	return false;
+}
+
 class FScatterCopyCS : public FGlobalShader
 {
 	DECLARE_SHADER_TYPE(FScatterCopyCS,Global)
@@ -223,6 +326,59 @@ public:
 
 IMPLEMENT_SHADER_TYPE(, FScatterCopyCS, TEXT("/Engine/Private/ByteBuffer.usf"), TEXT("ScatterCopyCS"), SF_Compute );
 
+class FScatterCopyTextureCS : public FGlobalShader
+{
+	DECLARE_SHADER_TYPE(FScatterCopyTextureCS, Global)
+
+	FScatterCopyTextureCS() {}
+
+	static bool ShouldCompilePermutation(const FGlobalShaderPermutationParameters& Parameters)
+	{
+		return RHISupportsComputeShaders(Parameters.Platform);
+	}
+
+	static void ModifyCompilationEnvironment(const FGlobalShaderPermutationParameters& Parameters, FShaderCompilerEnvironment& OutEnvironment)
+	{
+		FGlobalShader::ModifyCompilationEnvironment(Parameters, OutEnvironment);
+		OutEnvironment.SetDefine(TEXT("THREADGROUP_SIZE"), ThreadGroupSize);
+	}
+
+public:
+	enum { ThreadGroupSize = 64 };
+
+	FShaderParameter			NumScatters;
+	FShaderParameter			FloatsPerLine;
+	FShaderParameter			Size;
+	FShaderResourceParameter	ScatterBuffer;
+	FShaderResourceParameter	UploadBuffer;
+	FShaderResourceParameter	DstTexture;
+
+	FScatterCopyTextureCS(const ShaderMetaType::CompiledShaderInitializerType& Initializer)
+		: FGlobalShader(Initializer)
+	{
+		NumScatters.Bind(Initializer.ParameterMap, TEXT("NumScatters"));
+		FloatsPerLine.Bind(Initializer.ParameterMap, TEXT("FloatsPerLine"));
+		Size.Bind(Initializer.ParameterMap, TEXT("Size"));
+		ScatterBuffer.Bind(Initializer.ParameterMap, TEXT("ScatterBuffer"));
+		UploadBuffer.Bind(Initializer.ParameterMap, TEXT("UploadBuffer"));
+		DstTexture.Bind(Initializer.ParameterMap, TEXT("DstTexture"));
+	}
+
+	virtual bool Serialize(FArchive& Ar) override
+	{
+		bool bShaderHasOutdatedParameters = FGlobalShader::Serialize(Ar);
+		Ar << NumScatters;
+		Ar << FloatsPerLine;
+		Ar << Size;
+		Ar << ScatterBuffer;
+		Ar << UploadBuffer;
+		Ar << DstTexture;
+		return bShaderHasOutdatedParameters;
+	}
+};
+
+IMPLEMENT_SHADER_TYPE(, FScatterCopyTextureCS, TEXT("/Engine/Private/ByteBuffer.usf"), TEXT("ScatterCopyTextureCS"), SF_Compute);
+
 FScatterUploadBuilder::FScatterUploadBuilder(uint32 NumUploads, uint32 InStrideInFloat4s, FReadBuffer& InScatterBuffer, FReadBuffer& InUploadBuffer)
 	: ScatterBuffer(InScatterBuffer)
 	, UploadBuffer(InUploadBuffer)
@@ -289,6 +445,39 @@ void FScatterUploadBuilder::UploadTo(FRHICommandList& RHICmdList, FRWBufferStruc
 	SetUAVParameter(RHICmdList, ShaderRHI, ComputeShader->DstBuffer, FUnorderedAccessViewRHIRef());
 }
 
+int32 FTextureScatterUploadBuilder::GetMaxPrimitivesUpdate(uint32 NumUploads, uint32 InStrideInFloat4s)
+{
+	return GMaxTextureBufferSize == 0 ? NumUploads : FMath::Min((uint32)(GMaxTextureBufferSize / InStrideInFloat4s), NumUploads);
+}
+
+void FTextureScatterUploadBuilder::TextureUploadTo(FRHICommandList& RHICmdList, FTextureRWBuffer2D& DstTexture, uint32 NumFloat4, uint32 FloatsPerLine)
+{
+	RHIUnlockVertexBuffer(ScatterBuffer.Buffer);
+	RHIUnlockVertexBuffer(UploadBuffer.Buffer);
+
+	ScatterData = nullptr;
+	UploadData = nullptr;
+
+	auto ShaderMap = GetGlobalShaderMap(GMaxRHIFeatureLevel);
+
+	TShaderMapRef<FScatterCopyTextureCS> ComputeShader(ShaderMap);
+
+	FRHIComputeShader* ShaderRHI = ComputeShader->GetComputeShader();
+	
+	RHICmdList.SetComputeShader(ShaderRHI);
+
+	SetShaderValue(RHICmdList, ShaderRHI, ComputeShader->NumScatters, NumScatters);
+	SetSRVParameter(RHICmdList, ShaderRHI, ComputeShader->ScatterBuffer, ScatterBuffer.SRV);
+	SetShaderValue(RHICmdList, ShaderRHI, ComputeShader->FloatsPerLine, FloatsPerLine);
+	SetShaderValue(RHICmdList, ShaderRHI, ComputeShader->Size, NumFloat4);
+	SetSRVParameter(RHICmdList, ShaderRHI, ComputeShader->UploadBuffer, UploadBuffer.SRV);
+	SetUAVParameter(RHICmdList, ShaderRHI, ComputeShader->DstTexture, DstTexture.UAV);
+
+	RHICmdList.DispatchComputeShader(FMath::DivideAndRoundUp<uint32>(NumScatters, FScatterCopyTextureCS::ThreadGroupSize), 1, 1);
+
+	SetUAVParameter(RHICmdList, ShaderRHI, ComputeShader->DstTexture, FUnorderedAccessViewRHIRef());
+}
+
 void FScatterUploadBuilder::UploadTo_Flush(FRHICommandList& RHICmdList, FRWBufferStructured& DstBuffer)
 {
 	RHIUnlockVertexBuffer(ScatterBuffer.Buffer);
@@ -314,4 +503,33 @@ void FScatterUploadBuilder::UploadTo_Flush(FRHICommandList& RHICmdList, FRWBuffe
 	FRHICommandListExecutor::GetImmediateCommandList().ImmediateFlush(EImmediateFlushType::DispatchToRHIThread);
 
 	SetUAVParameter(RHICmdList, ShaderRHI, ComputeShader->DstBuffer, FUnorderedAccessViewRHIRef());
+}
+
+void FTextureScatterUploadBuilder::TextureUploadTo_Flush(FRHICommandList& RHICmdList, FTextureRWBuffer2D& DstTexture, uint32 NumFloat4, uint32 FloatsPerLine)
+{
+	RHIUnlockVertexBuffer(ScatterBuffer.Buffer);
+	RHIUnlockVertexBuffer(UploadBuffer.Buffer);
+
+	ScatterData = nullptr;
+	UploadData = nullptr;
+
+	auto ShaderMap = GetGlobalShaderMap(GMaxRHIFeatureLevel);
+
+	TShaderMapRef<FScatterCopyTextureCS> ComputeShader(ShaderMap);
+
+	FRHIComputeShader* ShaderRHI = ComputeShader->GetComputeShader();
+	RHICmdList.SetComputeShader(ShaderRHI);
+
+	SetShaderValue(RHICmdList, ShaderRHI, ComputeShader->NumScatters, NumScatters);
+	SetSRVParameter(RHICmdList, ShaderRHI, ComputeShader->ScatterBuffer, ScatterBuffer.SRV);
+	SetShaderValue(RHICmdList, ShaderRHI, ComputeShader->FloatsPerLine, FloatsPerLine);
+	SetShaderValue(RHICmdList, ShaderRHI, ComputeShader->Size, NumFloat4);
+	SetSRVParameter(RHICmdList, ShaderRHI, ComputeShader->UploadBuffer, UploadBuffer.SRV);
+	SetUAVParameter(RHICmdList, ShaderRHI, ComputeShader->DstTexture, DstTexture.UAV);
+
+	RHICmdList.DispatchComputeShader(FMath::DivideAndRoundUp<uint32>(NumScatters, FScatterCopyTextureCS::ThreadGroupSize), 1, 1);
+
+	FRHICommandListExecutor::GetImmediateCommandList().ImmediateFlush(EImmediateFlushType::DispatchToRHIThread);
+
+	SetUAVParameter(RHICmdList, ShaderRHI, ComputeShader->DstTexture, FUnorderedAccessViewRHIRef());
 }
\ No newline at end of file
diff --git a/Engine/Source/Runtime/Renderer/Private/ByteBuffer.h b/Engine/Source/Runtime/Renderer/Private/ByteBuffer.h
--- a/Engine/Source/Runtime/Renderer/Private/ByteBuffer.h
+++ b/Engine/Source/Runtime/Renderer/Private/ByteBuffer.h
@@ -4,10 +4,13 @@
 
 #include "CoreMinimal.h"
 #include "RHI.h"
+#include "PrimitiveUniformShaderParameters.h"
 
 void MemsetBuffer(FRHICommandList& RHICmdList, const FRWBufferStructured& DstBuffer, const FVector4& Value, uint32 NumFloat4s, uint32 DstOffsetInFloat4s);
 void MemcpyBuffer(FRHICommandList& RHICmdList, const FRWBufferStructured& SrcBuffer, const FRWBufferStructured& DstBuffer, uint32 NumFloat4s, uint32 SrcOffset = 0, uint32 DstOffset = 0);
+void MemcpyTextureToTexture(FRHICommandList& RHICmdList, const FTextureRWBuffer2D& SrcBuffer, const FTextureRWBuffer2D& DstBuffer, uint32 SrcOffset, uint32 DstOffset, uint32 NumFloat4s, uint32 FloatsPerLine);
 bool ResizeBufferIfNeeded(FRHICommandList& RHICmdList, FRWBufferStructured& Buffer, uint32 NumFloat4s);
+bool ResizeTextureIfNeeded(FRHICommandList& RHICmdList, FTextureRWBuffer2D& Texture, uint32 NumFloat4s, uint32 PrimitiveStride);
 
 class FScatterUploadBuilder
 {
@@ -46,4 +49,16 @@ public:
 		UploadData += StrideInFloat4s;
 		NumScatters += StrideInFloat4s;
 	}
+};
+
+class FTextureScatterUploadBuilder : public FScatterUploadBuilder
+{
+public:
+	FTextureScatterUploadBuilder(uint32 NumUploads, uint32 InStrideInFloat4s, FReadBuffer& InScatterBuffer, FReadBuffer& InUploadBuffer)
+		:	FScatterUploadBuilder(NumUploads, InStrideInFloat4s, InScatterBuffer, InUploadBuffer) 
+	{}
+
+	static int32 GetMaxPrimitivesUpdate(uint32 NumUploads, uint32 InStrideInFloat4s);
+	void TextureUploadTo(FRHICommandList& RHICmdList, FTextureRWBuffer2D& DstTexture, uint32 NumFloat4, uint32 FloatsPerLine);
+	void TextureUploadTo_Flush(FRHICommandList& RHICmdList, FTextureRWBuffer2D& DstTexture, uint32 NumFloat4, uint32 FloatsPerLine);
 };
\ No newline at end of file
diff --git a/Engine/Source/Runtime/Renderer/Private/GPUScene.cpp b/Engine/Source/Runtime/Renderer/Private/GPUScene.cpp
--- a/Engine/Source/Runtime/Renderer/Private/GPUScene.cpp
+++ b/Engine/Source/Runtime/Renderer/Private/GPUScene.cpp
@@ -151,7 +151,7 @@ int32 FGrowOnlySpanAllocator::SearchFreeList(int32 Num)
 	return INDEX_NONE;
 }
 
-void UpdateGPUScene(FRHICommandListImmediate& RHICmdList, FScene& Scene)
+static void UpdateGPUSceneBuffer(FRHICommandListImmediate& RHICmdList, FScene& Scene)
 {
 	if (UseGPUScene(GMaxRHIShaderPlatform, Scene.GetFeatureLevel()))
 	{
@@ -308,7 +308,171 @@ void UpdateGPUScene(FRHICommandListImmediate& RHICmdList, FScene& Scene)
 	checkSlow(Scene.GPUScene.PrimitivesToUpdate.Num() == 0);
 }
 
-void UploadDynamicPrimitiveShaderDataForView(FRHICommandListImmediate& RHICmdList, FScene& Scene, FViewInfo& View)
+static void UpdateGPUSceneTexture(FRHICommandListImmediate& RHICmdList, FScene& Scene)
+{
+	if (UseGPUScene(GMaxRHIShaderPlatform, Scene.GetFeatureLevel()))
+	{
+		CSV_SCOPED_TIMING_STAT_EXCLUSIVE(UpdateGPUSceneTexture);
+		QUICK_SCOPE_CYCLE_COUNTER(STAT_UpdateGPUSceneTexture);
+
+		if (GGPUSceneUploadEveryFrame || Scene.GPUScene.bUpdateAllPrimitives)
+		{
+			for (int32 Index : Scene.GPUScene.PrimitivesToUpdate)
+			{
+				Scene.GPUScene.PrimitivesMarkedToUpdate[Index] = false;
+			}
+			Scene.GPUScene.PrimitivesToUpdate.Reset();
+
+			for (int32 i = 0; i < Scene.Primitives.Num(); i++)
+			{
+				Scene.GPUScene.PrimitivesToUpdate.Add(i);
+			}
+
+			Scene.GPUScene.bUpdateAllPrimitives = false;
+		}
+
+		bool bResizedLightmapData = false;
+		bool bResizedPrimitiveTextureData = false;
+
+		{
+			const int32 NumPrimitiveEntries = Scene.Primitives.Num();
+			const uint32 PrimitiveSceneNumFloat4s = NumPrimitiveEntries * FPrimitiveSceneShaderData::PrimitiveDataStrideInFloat4s;
+			// Reserve enough space
+			bResizedPrimitiveTextureData = ResizeTextureIfNeeded(RHICmdList, Scene.GPUScene.PrimitiveTexture, PrimitiveSceneNumFloat4s, FPrimitiveSceneShaderData::PrimitiveDataStrideInFloat4s);
+		}
+
+		{
+			const int32 NumLightmapDataEntries = Scene.GPUScene.LightmapDataAllocator.GetMaxSize();
+			const uint32 LightmapDataNumFloat4s = NumLightmapDataEntries * FLightmapSceneShaderData::LightmapDataStrideInFloat4s;
+			bResizedLightmapData = ResizeBufferIfNeeded(RHICmdList, Scene.GPUScene.LightmapDataBuffer, FMath::RoundUpToPowerOfTwo(LightmapDataNumFloat4s));
+		}
+
+		const int32 NumPrimitiveDataUploads = Scene.GPUScene.PrimitivesToUpdate.Num();
+
+		int32 NumLightmapDataUploads = 0;
+
+		if (NumPrimitiveDataUploads > 0)
+		{
+			const int32 MaxPrimitivesUploads = FTextureScatterUploadBuilder::GetMaxPrimitivesUpdate(NumPrimitiveDataUploads, FPrimitiveSceneShaderData::PrimitiveDataStrideInFloat4s);
+			for (int32 PrimitiveOffset = 0; PrimitiveOffset < NumPrimitiveDataUploads; PrimitiveOffset += MaxPrimitivesUploads)
+			{
+				SCOPED_DRAW_EVENTF(RHICmdList, UpdateGPUScene, TEXT("UpdateGPUScene PrimitivesToUpdate and Offset = %u %u"), NumPrimitiveDataUploads, PrimitiveOffset);
+
+				FTextureScatterUploadBuilder PrimitivesUploadBuilder(MaxPrimitivesUploads, FPrimitiveSceneShaderData::PrimitiveDataStrideInFloat4s, Scene.GPUScene.PrimitivesUploadScatterBuffer, Scene.GPUScene.PrimitivesUploadDataBuffer);
+
+				for (int32 IndexUpdate = 0; (IndexUpdate < MaxPrimitivesUploads) && ((IndexUpdate + PrimitiveOffset) < NumPrimitiveDataUploads); ++IndexUpdate)
+				{
+					int32 Index = Scene.GPUScene.PrimitivesToUpdate[IndexUpdate + PrimitiveOffset];
+					// PrimitivesToUpdate may contain a stale out of bounds index, as we don't remove update request on primitive removal from scene.
+					if (Index < Scene.PrimitiveSceneProxies.Num())
+					{
+						FPrimitiveSceneProxy* PrimitiveSceneProxy = Scene.PrimitiveSceneProxies[Index];
+						NumLightmapDataUploads += PrimitiveSceneProxy->GetPrimitiveSceneInfo()->GetNumLightmapDataEntries();
+
+						FPrimitiveSceneShaderData PrimitiveSceneData(PrimitiveSceneProxy);
+						PrimitivesUploadBuilder.Add(Index, &PrimitiveSceneData.Data[0]);
+					}
+					Scene.GPUScene.PrimitivesMarkedToUpdate[Index] = false;
+				}
+
+				{
+					if (bResizedPrimitiveTextureData)
+					{
+						RHICmdList.TransitionResource(EResourceTransitionAccess::ERWBarrier, EResourceTransitionPipeline::EComputeToCompute, Scene.GPUScene.PrimitiveTexture.UAV);
+					}
+					else
+					{
+						RHICmdList.TransitionResource(EResourceTransitionAccess::EWritable, EResourceTransitionPipeline::EGfxToCompute, Scene.GPUScene.PrimitiveTexture.UAV);
+					}
+				}
+
+				{
+					uint16 PrimitivesPerTextureLine = FPrimitiveSceneShaderData::GetPrimitivesPerTextureLine();
+					PrimitivesUploadBuilder.TextureUploadTo_Flush(RHICmdList, Scene.GPUScene.PrimitiveTexture, Scene.Primitives.Num() * FPrimitiveSceneShaderData::PrimitiveDataStrideInFloat4s, PrimitivesPerTextureLine * FPrimitiveSceneShaderData::PrimitiveDataStrideInFloat4s);
+				}
+			}
+			RHICmdList.TransitionResource(EResourceTransitionAccess::EReadable, EResourceTransitionPipeline::EComputeToGfx, Scene.GPUScene.PrimitiveTexture.UAV);
+		}
+
+		
+		if (NumPrimitiveDataUploads > 0)
+		{
+			if (NumLightmapDataUploads > 0)
+			{
+				const int32 MaxLightmapsUploads = FTextureScatterUploadBuilder::GetMaxPrimitivesUpdate(NumLightmapDataUploads, FPrimitiveSceneShaderData::PrimitiveDataStrideInFloat4s);
+				for (int32 PrimitiveOffset = 0; PrimitiveOffset < NumPrimitiveDataUploads; PrimitiveOffset += MaxLightmapsUploads)
+				{
+					FScatterUploadBuilder LightmapDataUploadBuilder(NumLightmapDataUploads, FLightmapSceneShaderData::LightmapDataStrideInFloat4s, Scene.GPUScene.LightmapUploadScatterBuffer, Scene.GPUScene.LightmapUploadDataBuffer);
+
+					for (int32 IndexUpdate = 0; (IndexUpdate < MaxLightmapsUploads) && ((IndexUpdate + PrimitiveOffset) < NumPrimitiveDataUploads); ++IndexUpdate)
+					{
+						int32 Index = Scene.GPUScene.PrimitivesToUpdate[IndexUpdate + PrimitiveOffset];
+						// PrimitivesToUpdate may contain a stale out of bounds index, as we don't remove update request on primitive removal from scene.
+						if (Index < Scene.PrimitiveSceneProxies.Num())
+						{
+							FPrimitiveSceneProxy* PrimitiveSceneProxy = Scene.PrimitiveSceneProxies[Index];
+
+							FPrimitiveSceneProxy::FLCIArray LCIs;
+							PrimitiveSceneProxy->GetLCIs(LCIs);
+
+							check(LCIs.Num() == PrimitiveSceneProxy->GetPrimitiveSceneInfo()->GetNumLightmapDataEntries());
+							const int32 LightmapDataOffset = PrimitiveSceneProxy->GetPrimitiveSceneInfo()->GetLightmapDataOffset();
+
+							for (int32 i = 0; i < LCIs.Num(); i++)
+							{
+								FLightmapSceneShaderData LightmapSceneData(LCIs[i], Scene.GetFeatureLevel());
+								LightmapDataUploadBuilder.Add(LightmapDataOffset + i, &LightmapSceneData.Data[0]);
+							}
+						}
+					}
+
+					if (bResizedLightmapData)
+					{
+						RHICmdList.TransitionResource(EResourceTransitionAccess::ERWBarrier, EResourceTransitionPipeline::EComputeToCompute, Scene.GPUScene.LightmapDataBuffer.UAV);
+					}
+					else
+					{
+						RHICmdList.TransitionResource(EResourceTransitionAccess::EWritable, EResourceTransitionPipeline::EGfxToCompute, Scene.GPUScene.LightmapDataBuffer.UAV);
+					}
+
+					LightmapDataUploadBuilder.UploadTo(RHICmdList, Scene.GPUScene.LightmapDataBuffer);
+	
+				}
+				RHICmdList.TransitionResource(EResourceTransitionAccess::EReadable, EResourceTransitionPipeline::EComputeToGfx, Scene.GPUScene.LightmapDataBuffer.UAV);
+			}
+
+			Scene.GPUScene.PrimitivesToUpdate.Reset();
+			
+			if (Scene.GPUScene.PrimitivesUploadDataBuffer.NumBytes > (uint32)GGPUSceneMaxPooledUploadBufferSize || Scene.GPUScene.PrimitivesUploadScatterBuffer.NumBytes > (uint32)GGPUSceneMaxPooledUploadBufferSize)
+			{
+				Scene.GPUScene.PrimitivesUploadDataBuffer.Release();
+				Scene.GPUScene.PrimitivesUploadScatterBuffer.Release();
+			}
+
+			if (Scene.GPUScene.LightmapUploadDataBuffer.NumBytes > (uint32)GGPUSceneMaxPooledUploadBufferSize || Scene.GPUScene.LightmapUploadScatterBuffer.NumBytes > (uint32)GGPUSceneMaxPooledUploadBufferSize)
+			{
+				Scene.GPUScene.LightmapUploadDataBuffer.Release();
+				Scene.GPUScene.LightmapUploadScatterBuffer.Release();
+			}
+		}
+	}
+
+	checkSlow(Scene.GPUScene.PrimitivesToUpdate.Num() == 0);
+}
+
+void UpdateGPUScene(FRHICommandListImmediate& RHICmdList, FScene& Scene)
+{
+	if (!GPUSceneUseTexture2D(Scene.GetShaderPlatform()))
+	{
+		UpdateGPUSceneBuffer(RHICmdList, Scene);
+	}
+	else
+	{
+		UpdateGPUSceneTexture(RHICmdList, Scene);
+	}
+}
+
+static void UploadDynamicPrimitiveShaderDataBufferForView(FRHICommandListImmediate& RHICmdList, FScene& Scene, FViewInfo& View)
 {
 	if (UseGPUScene(GMaxRHIShaderPlatform, Scene.GetFeatureLevel()))
 	{
@@ -369,6 +533,94 @@ void UploadDynamicPrimitiveShaderDataForView(FRHICommandListImmediate& RHICmdLis
 	}
 }
 
+static void UploadDynamicPrimitiveShaderDataTextureForView(FRHICommandListImmediate& RHICmdList, FScene& Scene, FViewInfo& View)
+{
+	uint16 PrimitivesPerTextureLine = FPrimitiveSceneShaderData::GetPrimitivesPerTextureLine();
+
+	if (UseGPUScene(GMaxRHIShaderPlatform, Scene.GetFeatureLevel()))
+	{
+		CSV_SCOPED_TIMING_STAT_EXCLUSIVE(UploadDynamicPrimitiveShaderDataTextureForView);
+		QUICK_SCOPE_CYCLE_COUNTER(STAT_UploadDynamicPrimitiveShaderDataTextureForView);
+
+		const int32 NumPrimitiveDataUploads = View.DynamicPrimitiveShaderData.Num();
+		if (NumPrimitiveDataUploads > 0)
+		{
+			FTextureRWBuffer2D& ViewPrimitiveShaderDataTexture = View.ViewState ? View.ViewState->PrimitiveShaderDataTexture : View.OneFramePrimitiveShaderDataTexture;
+
+			const int32 NumPrimitiveEntries = Scene.Primitives.Num() + View.DynamicPrimitiveShaderData.Num();
+			const uint32 PrimitiveSceneNumFloat4s = NumPrimitiveEntries * FPrimitiveSceneShaderData::PrimitiveDataStrideInFloat4s;
+
+			uint32 ViewPrimitiveSceneNumFloat4s = FMath::RoundUpToPowerOfTwo(PrimitiveSceneNumFloat4s);
+			uint32 BytesPerElement = GPixelFormats[PF_A32B32G32R32F].BlockBytes;
+
+			// Reserve enough space
+			if (ViewPrimitiveSceneNumFloat4s * BytesPerElement != ViewPrimitiveShaderDataTexture.NumBytes)
+			{
+				ViewPrimitiveShaderDataTexture.Release();
+				ViewPrimitiveShaderDataTexture.Initialize(BytesPerElement, FPrimitiveSceneShaderData::PrimitiveDataStrideInFloat4s * PrimitivesPerTextureLine, NumPrimitiveEntries / PrimitivesPerTextureLine + 1, PF_A32B32G32R32F, TexCreate_RenderTargetable | TexCreate_UAV);
+			}
+
+			// Copy scene primitive data into view primitive data
+			{
+				MemcpyTextureToTexture(RHICmdList, Scene.GPUScene.PrimitiveTexture, ViewPrimitiveShaderDataTexture, 0, 0, Scene.Primitives.Num() * FPrimitiveSceneShaderData::PrimitiveDataStrideInFloat4s, PrimitivesPerTextureLine * FPrimitiveSceneShaderData::PrimitiveDataStrideInFloat4s);
+			}
+
+			// Append View.DynamicPrimitiveShaderData to the end of the view primitive data texture
+			if (NumPrimitiveDataUploads > 0)
+			{
+				int32 MaxPrimitivesUploads = FTextureScatterUploadBuilder::GetMaxPrimitivesUpdate(NumPrimitiveDataUploads, FPrimitiveSceneShaderData::PrimitiveDataStrideInFloat4s);
+				for (int32 PrimitiveOffset = 0; PrimitiveOffset < NumPrimitiveDataUploads; PrimitiveOffset += MaxPrimitivesUploads)
+				{
+					FTextureScatterUploadBuilder PrimitivesUploadBuilder(MaxPrimitivesUploads, FPrimitiveSceneShaderData::PrimitiveDataStrideInFloat4s, Scene.GPUScene.PrimitivesUploadScatterBuffer, Scene.GPUScene.PrimitivesUploadDataBuffer);
+
+					for (int32 IndexUpdate = 0; (IndexUpdate < MaxPrimitivesUploads) && ((IndexUpdate + PrimitiveOffset) < NumPrimitiveDataUploads); ++IndexUpdate)
+					{
+						int32 DynamicUploadIndex = IndexUpdate + PrimitiveOffset;
+						FPrimitiveSceneShaderData PrimitiveSceneData(View.DynamicPrimitiveShaderData[DynamicUploadIndex]);
+						// Place dynamic primitive shader data just after scene primitive data
+						PrimitivesUploadBuilder.Add(Scene.Primitives.Num() + DynamicUploadIndex, &PrimitiveSceneData.Data[0]);
+					}
+
+					{
+						RHICmdList.TransitionResource(EResourceTransitionAccess::ERWBarrier, EResourceTransitionPipeline::EComputeToCompute, ViewPrimitiveShaderDataTexture.UAV);
+						PrimitivesUploadBuilder.TextureUploadTo(RHICmdList, ViewPrimitiveShaderDataTexture, NumPrimitiveEntries * FPrimitiveSceneShaderData::PrimitiveDataStrideInFloat4s, PrimitivesPerTextureLine * FPrimitiveSceneShaderData::PrimitiveDataStrideInFloat4s);
+					}
+				}
+			}
+
+			{
+				RHICmdList.TransitionResource(EResourceTransitionAccess::EReadable, EResourceTransitionPipeline::EComputeToGfx, ViewPrimitiveShaderDataTexture.UAV);
+				View.CachedViewUniformShaderParameters->PrimitiveSceneDataTexture = ViewPrimitiveShaderDataTexture.Buffer;
+			}
+
+		}
+		else
+		{
+			// No dynamic primitives for this view, we just use Scene.GPUScene.PrimitiveTexture.
+			{
+				View.CachedViewUniformShaderParameters->PrimitiveSceneDataTexture = Scene.GPUScene.PrimitiveTexture.Buffer;
+			}
+
+		}
+
+		// Update view uniform buffer
+		View.CachedViewUniformShaderParameters->LightmapSceneData = Scene.GPUScene.LightmapDataBuffer.SRV;
+		View.ViewUniformBuffer.UpdateUniformBufferImmediate(*View.CachedViewUniformShaderParameters);
+	}
+}
+
+void UploadDynamicPrimitiveShaderDataForView(FRHICommandListImmediate& RHICmdList, FScene& Scene, FViewInfo& View)
+{
+	if (!GPUSceneUseTexture2D(Scene.GetShaderPlatform()))
+	{
+		UploadDynamicPrimitiveShaderDataBufferForView(RHICmdList, Scene, View);
+	}
+	else
+	{
+		UploadDynamicPrimitiveShaderDataTextureForView(RHICmdList, Scene, View);
+	}
+}
+
 void AddPrimitiveToUpdateGPU(FScene& Scene, int32 PrimitiveId)
 {
 	if (UseGPUScene(GMaxRHIShaderPlatform, Scene.GetFeatureLevel()))
diff --git a/Engine/Source/Runtime/Renderer/Private/MeshDrawCommands.cpp b/Engine/Source/Runtime/Renderer/Private/MeshDrawCommands.cpp
--- a/Engine/Source/Runtime/Renderer/Private/MeshDrawCommands.cpp
+++ b/Engine/Source/Runtime/Renderer/Private/MeshDrawCommands.cpp
@@ -366,6 +366,7 @@ void BuildMeshDrawCommandPrimitiveIdBuffer(
 	int32& MaxInstances,
 	int32& VisibleMeshDrawCommandsNum,
 	int32& NewPassVisibleMeshDrawCommandsNum,
+	EShaderPlatform ShaderPlatform,
 	uint32 InstanceFactor = 1
 	)
 {
@@ -451,7 +452,16 @@ void BuildMeshDrawCommandPrimitiveIdBuffer(
 			{
 				//@todo - refactor into memcpy
 				checkSlow(PrimitiveIdIndex < MaxPrimitiveId);
-				PrimitiveIds[PrimitiveIdIndex] = VisibleMeshDrawCommand.DrawPrimitiveId;
+				if (!GPUSceneUseTexture2D(ShaderPlatform))
+				{
+					PrimitiveIds[PrimitiveIdIndex] = VisibleMeshDrawCommand.DrawPrimitiveId;
+				}
+				else
+				{
+					//Packing for mobile texture2D GPUScene. Must be in sync with SceneData.ush
+					uint16 PrimitivesPerTextureLine = FPrimitiveSceneShaderData::GetPrimitivesPerTextureLine();
+					PrimitiveIds[PrimitiveIdIndex] = ((VisibleMeshDrawCommand.DrawPrimitiveId / PrimitivesPerTextureLine) << 16) | (VisibleMeshDrawCommand.DrawPrimitiveId % PrimitivesPerTextureLine);
+				}
 			}
 		}
 
@@ -845,6 +855,7 @@ public:
 					Context.MaxInstances,
 					Context.VisibleMeshDrawCommandsNum,
 					Context.NewPassVisibleMeshDrawCommandsNum,
+					Context.ShaderPlatform,
 					Context.InstanceFactor
 				);
 			}
@@ -905,6 +916,7 @@ void SortAndMergeDynamicPassMeshDrawCommands(
 				MaxInstances,
 				VisibleMeshDrawCommandsNum,
 				NewPassVisibleMeshDrawCommandsNum,
+				GShaderPlatformForFeatureLevel[FeatureLevel],
 				InstanceFactor
 			);
 
@@ -942,6 +954,7 @@ void FParallelMeshDrawCommandPass::DispatchPassSetup(
 
 	TaskContext.View = &View;
 	TaskContext.ShadingPath = Scene->GetShadingPath();
+	TaskContext.ShaderPlatform = Scene->GetShaderPlatform();
 	TaskContext.PassType = PassType;
 	TaskContext.bUseGPUScene = UseGPUScene(GMaxRHIShaderPlatform, View.GetFeatureLevel());
 	TaskContext.bDynamicInstancing = IsDynamicInstancingEnabled(View.GetFeatureLevel());
diff --git a/Engine/Source/Runtime/Renderer/Private/MeshDrawCommands.h b/Engine/Source/Runtime/Renderer/Private/MeshDrawCommands.h
--- a/Engine/Source/Runtime/Renderer/Private/MeshDrawCommands.h
+++ b/Engine/Source/Runtime/Renderer/Private/MeshDrawCommands.h
@@ -73,6 +73,7 @@ public:
 
 	const FViewInfo* View;
 	EShadingPath ShadingPath;
+	EShaderPlatform ShaderPlatform;
 	EMeshPass::Type PassType;
 	bool bUseGPUScene;
 	bool bDynamicInstancing;
diff --git a/Engine/Source/Runtime/Renderer/Private/Renderer.cpp b/Engine/Source/Runtime/Renderer/Private/Renderer.cpp
--- a/Engine/Source/Runtime/Renderer/Private/Renderer.cpp
+++ b/Engine/Source/Runtime/Renderer/Private/Renderer.cpp
@@ -120,7 +120,7 @@ void FRendererModule::DrawTileMesh(FRHICommandListImmediate& RHICmdList, FMeshPa
 		extern FForwardLightingViewResources* GetMinimalDummyForwardLightingResources();
 		View.ForwardLightingResources = GetMinimalDummyForwardLightingResources();
 
-		FSinglePrimitiveStructuredBuffer& SinglePrimitiveStructuredBuffer = GTilePrimitiveBuffer;
+		FSinglePrimitiveStructured& SinglePrimitiveStructured = GTilePrimitiveBuffer;
 
 		if (Mesh.VertexFactory->GetPrimitiveIdStreamIndex(EVertexInputStreamType::PositionOnly) >= 0)
 		{
@@ -140,16 +140,25 @@ void FRendererModule::DrawTileMesh(FRHICommandListImmediate& RHICmdList, FMeshPa
 				PrimitiveParams.LightmapDataIndex = 0;
 
 				// Now we just need to fill out the first entry of primitive data in a buffer and bind it
-				SinglePrimitiveStructuredBuffer.PrimitiveSceneData = FPrimitiveSceneShaderData(PrimitiveParams);
+				SinglePrimitiveStructured.PrimitiveSceneData = FPrimitiveSceneShaderData(PrimitiveParams);
+				SinglePrimitiveStructured.ShaderPlatform = View.GetShaderPlatform();
 
 				// Set up the parameters for the LightmapSceneData from the given LCI data 
 				FPrecomputedLightingUniformParameters LightmapParams;
 				GetPrecomputedLightingParameters(FeatureLevel, LightmapParams, Mesh.LCI);
-				SinglePrimitiveStructuredBuffer.LightmapSceneData = FLightmapSceneShaderData(LightmapParams);
+				SinglePrimitiveStructured.LightmapSceneData = FLightmapSceneShaderData(LightmapParams);
 
-				SinglePrimitiveStructuredBuffer.UploadToGPU();
-				View.PrimitiveSceneDataOverrideSRV = SinglePrimitiveStructuredBuffer.PrimitiveSceneDataBufferSRV;
-				View.LightmapSceneDataOverrideSRV = SinglePrimitiveStructuredBuffer.LightmapSceneDataBufferSRV;
+				SinglePrimitiveStructured.UploadToGPU();
+
+				if (!GPUSceneUseTexture2D(View.GetShaderPlatform()))
+				{
+					View.PrimitiveSceneDataOverrideSRV = SinglePrimitiveStructured.PrimitiveSceneDataBufferSRV;
+				}
+				else
+				{
+					View.PrimitiveSceneDataTextureOverrideRHI = SinglePrimitiveStructured.PrimitiveSceneDataTextureRHI;
+				}
+				View.LightmapSceneDataOverrideSRV = SinglePrimitiveStructured.LightmapSceneDataBufferSRV;
 			}
 		}
 
diff --git a/Engine/Source/Runtime/Renderer/Private/ScenePrivate.h b/Engine/Source/Runtime/Renderer/Private/ScenePrivate.h
--- a/Engine/Source/Runtime/Renderer/Private/ScenePrivate.h
+++ b/Engine/Source/Runtime/Renderer/Private/ScenePrivate.h
@@ -958,7 +958,9 @@ public:
 	FRWBuffer CapsuleTileIntersectionCountsBuffer;
 
 	/** Contains both DynamicPrimitiveShaderData (per view) and primitive shader data (per scene).  Stored in ViewState for pooling only (contents are not persistent). */
+	/** Only one of the resources(TextureBuffer or Texture2D) will be used depending on the Mobile.UseGPUSceneTexture cvar */
 	FRWBufferStructured PrimitiveShaderDataBuffer;
+	FTextureRWBuffer2D PrimitiveShaderDataTexture;
 
 	/** Timestamp queries around separate translucency, used for auto-downsampling. */
 	FRenderQueryPoolRHIRef TimerQueryPool;
@@ -1667,7 +1669,9 @@ public:
 	TBitArray<> PrimitivesMarkedToUpdate;
 
 	/** GPU mirror of Primitives */
+	/** Only one of the resources(TextureBuffer or Texture2D) will be used depending on the Mobile.UseGPUSceneTexture cvar */
 	FRWBufferStructured PrimitiveBuffer;
+	FTextureRWBuffer2D PrimitiveTexture;
 
 	FGrowOnlySpanAllocator LightmapDataAllocator;
 
diff --git a/Engine/Source/Runtime/Renderer/Private/SceneRendering.cpp b/Engine/Source/Runtime/Renderer/Private/SceneRendering.cpp
--- a/Engine/Source/Runtime/Renderer/Private/SceneRendering.cpp
+++ b/Engine/Source/Runtime/Renderer/Private/SceneRendering.cpp
@@ -914,6 +914,7 @@ void FViewInfo::Init()
 	}
 
 	PrimitiveSceneDataOverrideSRV = nullptr;
+	PrimitiveSceneDataTextureOverrideRHI = nullptr;
 	LightmapSceneDataOverrideSRV = nullptr;
 
 	DitherFadeInUniformBuffer = nullptr;
@@ -1662,6 +1663,16 @@ void FViewInfo::SetupUniformBufferParameters(
 			}
 		}
 
+		if (PrimitiveSceneDataTextureOverrideRHI)
+		{
+			ViewUniformShaderParameters.PrimitiveSceneDataTexture = PrimitiveSceneDataTextureOverrideRHI;
+		}
+		else
+		{
+			const FTextureRWBuffer2D& ViewPrimitiveShaderDataTexture = ViewState ? ViewState->PrimitiveShaderDataTexture : OneFramePrimitiveShaderDataTexture;
+			ViewUniformShaderParameters.PrimitiveSceneDataTexture = OrBlack2DIfNull(ViewPrimitiveShaderDataTexture.Buffer);
+		}
+		
 		if (LightmapSceneDataOverrideSRV)
 		{
 			ViewUniformShaderParameters.LightmapSceneData = LightmapSceneDataOverrideSRV;
@@ -1768,6 +1779,9 @@ FViewInfo* FViewInfo::CreateSnapshot() const
 
 	FRWBufferStructured NullOneFramePrimitiveShaderDataBuffer;
 	FMemory::Memcpy(Result->OneFramePrimitiveShaderDataBuffer, NullOneFramePrimitiveShaderDataBuffer);
+	
+	FTextureRWBuffer2D NullOneFramePrimitiveShaderDataTexture;
+	FMemory::Memcpy(Result->OneFramePrimitiveShaderDataTexture, NullOneFramePrimitiveShaderDataTexture);
 
 	TStaticArray<FParallelMeshDrawCommandPass, EMeshPass::Num> NullParallelMeshDrawCommandPasses;
 	FMemory::Memcpy(Result->ParallelMeshDrawCommandPasses, NullParallelMeshDrawCommandPasses);
@@ -1798,6 +1812,7 @@ void FViewInfo::DestroyAllSnapshots()
 		Snapshot->CachedViewUniformShaderParameters.Reset();
 		Snapshot->DynamicPrimitiveShaderData.Empty();
 		Snapshot->OneFramePrimitiveShaderDataBuffer.Release();
+		Snapshot->OneFramePrimitiveShaderDataTexture.Release();
 
 		for (int32 Index = 0; Index < Snapshot->ParallelMeshDrawCommandPasses.Num(); ++Index)
 		{
diff --git a/Engine/Source/Runtime/Renderer/Private/SceneRendering.h b/Engine/Source/Runtime/Renderer/Private/SceneRendering.h
--- a/Engine/Source/Runtime/Renderer/Private/SceneRendering.h
+++ b/Engine/Source/Runtime/Renderer/Private/SceneRendering.h
@@ -974,7 +974,9 @@ public:
 	/** Tracks dynamic primitive data for upload to GPU Scene, when enabled. */
 	TArray<FPrimitiveUniformShaderParameters> DynamicPrimitiveShaderData;
 
+	/** Only one of the resources(TextureBuffer or Texture2D) will be used depending on the Mobile.UseGPUSceneTexture cvar */
 	FRWBufferStructured OneFramePrimitiveShaderDataBuffer;
+	FTextureRWBuffer2D OneFramePrimitiveShaderDataTexture;
 
 	TStaticArray<FParallelMeshDrawCommandPass, EMeshPass::Num> ParallelMeshDrawCommandPasses;
 	
@@ -1138,7 +1140,9 @@ public:
 
 	TArray<FPrimitiveSceneInfo*, SceneRenderingAllocator> IndirectShadowPrimitives;
 
+	/** Only one of the resources(TextureBuffer or Texture2D) will be used depending on the Mobile.UseGPUSceneTexture cvar */
 	FShaderResourceViewRHIRef PrimitiveSceneDataOverrideSRV;
+	FTexture2DRHIRef PrimitiveSceneDataTextureOverrideRHI;
 	FShaderResourceViewRHIRef LightmapSceneDataOverrideSRV;
 
 	FRWBufferStructured ShaderPrintValueBuffer;
diff --git a/Engine/Source/Runtime/Renderer/Private/ShadowDepthRendering.cpp b/Engine/Source/Runtime/Renderer/Private/ShadowDepthRendering.cpp
--- a/Engine/Source/Runtime/Renderer/Private/ShadowDepthRendering.cpp
+++ b/Engine/Source/Runtime/Renderer/Private/ShadowDepthRendering.cpp
@@ -1107,8 +1107,6 @@ void FProjectedShadowInfo::SetupShadowUniformBuffers(FRHICommandListImmediate& R
 
 		ShadowDepthPassUniformBuffer.UpdateUniformBufferImmediate(ShadowDepthPassParameters);
 
-		UploadDynamicPrimitiveShaderDataForView(RHICmdList, *Scene, *ShadowDepthView);
-
 		if (DependentView)
 		{
 			extern TSet<IPersistentViewUniformBufferExtension*> PersistentViewUniformBufferExtensions;
@@ -1119,6 +1117,9 @@ void FProjectedShadowInfo::SetupShadowUniformBuffers(FRHICommandListImmediate& R
 			}
 		}
 	}
+	
+	// This needs to be done for both mobile and deferred
+	UploadDynamicPrimitiveShaderDataForView(RHICmdList, *Scene, *ShadowDepthView);
 }
 
 void FProjectedShadowInfo::TransitionCachedShadowmap(FRHICommandListImmediate& RHICmdList, FScene* Scene)
