Remove unused runtime virtual texture formats
Default to world space normals when writing to virtual texture through standard material attributes
Will require shader change for impacted materials, but almost all of our current VT materials use the VT output node and so are unaffected
This is necessary ground work for channel selection when writing to runtime virtual textures

diff --git a/Engine/Shaders/Private/VirtualTextureCompress.usf b/Engine/Shaders/Private/VirtualTextureCompress.usf
--- a/Engine/Shaders/Private/VirtualTextureCompress.usf
+++ b/Engine/Shaders/Private/VirtualTextureCompress.usf
@@ -127,82 +127,9 @@ void CompressBaseColorCS(
 	OutCompressTexture0_u2[ThreadId.xy + DestRect.xy] = CompressBC1Block_SRGB(BlockBaseColor);
 }
 
-/** Compress base color and normal to BC1/BC5 target pair. */
-[numthreads(8,8,1)]
-void CompressBaseColorNormalCS(
-	uint3 ThreadId : SV_DispatchThreadID 
-	)
-{
-	uint2 SamplePos = ThreadId.xy * 4;
-	if (any(SamplePos >= DestRect.zw))
-		return;
-
-	float2 TexelUVSize = 1.f / float2(DestRect.zw);
-	float2 SampleUV = (float2(SamplePos) + 0.5f) * TexelUVSize;
-
-	float3 BlockBaseColor[16];
-	ReadBlockRGB(RenderTexture0, TextureSampler0, SampleUV, TexelUVSize, BlockBaseColor);
-
-	float BlockNormalX[16];
-	float BlockNormalY[16];
-	ReadBlockRG(RenderTexture1, TextureSampler1, SampleUV, TexelUVSize, BlockNormalX, BlockNormalY);
-
-	OutCompressTexture0_u2[ThreadId.xy + DestRect.xy] = CompressBC1Block_SRGB(BlockBaseColor);
-	OutCompressTexture1[ThreadId.xy + DestRect.xy] = CompressBC5Block(BlockNormalX, BlockNormalY);
-}
-
-/** Compress base color, normal, roughness and specular to BC1/BC3 target pair. */
-[numthreads(8,8,1)]
-void CompressBaseColorNormalSpecularCS(
-	uint3 ThreadId : SV_DispatchThreadID 
-	)
-{
-	uint2 SamplePos = ThreadId.xy * 4;
-	if (any(SamplePos >= DestRect.zw))
-		return;
-
-	float2 TexelUVSize = 1.f / float2(DestRect.zw);
-	float2 SampleUV = (float2(SamplePos) + 0.5f) * TexelUVSize;
-
-	float3 BlockBaseColor[16];
-	ReadBlockRGB(RenderTexture0, TextureSampler0, SampleUV, TexelUVSize, BlockBaseColor);
-
-	float BlockNormalX[16];
-	float BlockNormalY[16];
-	ReadBlockRG(RenderTexture1, TextureSampler1, SampleUV, TexelUVSize, BlockNormalX, BlockNormalY);
-
-	float3 BlockSpecular[16];
-	ReadBlockRGB(RenderTexture2, TextureSampler2, SampleUV, TexelUVSize, BlockSpecular);
-
-	for (int i=0; i<16; i++)
-	{
-		BlockSpecular[i] = float3(BlockSpecular[i].x, BlockNormalY[i], BlockSpecular[i].y);
-	}
-
-	OutCompressTexture0_u2[ThreadId.xy + DestRect.xy] = CompressBC1Block_SRGB(BlockBaseColor);
-	OutCompressTexture1[ThreadId.xy + DestRect.xy] = CompressBC3Block(BlockSpecular, BlockNormalX);
-}
-
-/** Combine and copy normal, roughness and specular to a single packed target (without compression).
-  * This path is used when we disable texture compression, because we need to keep the same final channel layout. */
-[numthreads(8,8,1)]
-void CopyNormalSpecularCS(
-	uint3 ThreadId : SV_DispatchThreadID 
-	)
-{
-	uint2 SamplePos = ThreadId.xy;
-	if (any(ThreadId.xy >= DestRect.zw))
-		return;
-
-	float2 NormalXY = RenderTexture1[SamplePos].xy;
-	float2 RoughnessSpecular = RenderTexture2[SamplePos].xy;
-	
-	OutCopyTexture0[ThreadId.xy + DestRect.xy] = float4(RoughnessSpecular.x, NormalXY.y, RoughnessSpecular.y, NormalXY.x);
-}
-
 /** Compress base color, normal, roughness and specular to BC3 target pair. */
 [numthreads(8,8,1)]
-void CompressBaseColorNormalSpecularExCS(
+void CompressBaseColorNormalSpecularCS(
 	uint3 ThreadId : SV_DispatchThreadID 
 	)
 {
@@ -235,7 +162,7 @@ void CompressBaseColorNormalSpecularExCS(
 /** Copy normal X and Y to the alpha channels of the final target textures (without compression)
   * This path is used when we disable texture compression, because we need to keep the same final channel layout. */
 [numthreads(8,8,1)]
-void CopyNormalSpecularExCS(
+void CopyNormalSpecularCS(
 	uint3 ThreadId : SV_DispatchThreadID 
 	)
 {
diff --git a/Engine/Shaders/Private/VirtualTextureMaterial.usf b/Engine/Shaders/Private/VirtualTextureMaterial.usf
--- a/Engine/Shaders/Private/VirtualTextureMaterial.usf
+++ b/Engine/Shaders/Private/VirtualTextureMaterial.usf
@@ -79,7 +79,7 @@ void FPixelShaderInOut_MainPS(
 	half3 BaseColor = GetMaterialBaseColor(PixelMaterialInputs);
 	half Specular = GetMaterialSpecular(PixelMaterialInputs);
 	half Roughness = GetMaterialRoughness(PixelMaterialInputs);
-	half3 Normal = mul(MaterialParameters.WorldNormal, (float3x3)ResolvedView.WorldToVirtualTexture);
+	half3 Normal = MaterialParameters.WorldNormal;
 	float WorldHeight = MaterialParameters.AbsoluteWorldPosition.z;
 	float Opacity = GetMaterialOpacity(PixelMaterialInputs);
 #endif
diff --git a/Engine/Source/Runtime/Engine/Private/Materials/MaterialExpressions.cpp b/Engine/Source/Runtime/Engine/Private/Materials/MaterialExpressions.cpp
--- a/Engine/Source/Runtime/Engine/Private/Materials/MaterialExpressions.cpp
+++ b/Engine/Source/Runtime/Engine/Private/Materials/MaterialExpressions.cpp
@@ -2112,6 +2112,13 @@ UObject* UMaterialExpressionRuntimeVirtualTextureSample::GetReferencedTexture()
 void UMaterialExpressionRuntimeVirtualTextureSample::PostLoad()
 {
 	Super::PostLoad();
+
+	// Convert BaseColor_Normal_DEPRECATED
+	if (MaterialType == ERuntimeVirtualTextureMaterialType::BaseColor_Normal_DEPRECATED)
+	{
+		MaterialType = ERuntimeVirtualTextureMaterialType::BaseColor_Normal_Specular;
+	}
+
 	InitOutputs();
 }
 
@@ -2181,9 +2188,7 @@ int32 UMaterialExpressionRuntimeVirtualTextureSample::Compile(class FMaterialCom
 	switch (MaterialType)
 	{
 	case ERuntimeVirtualTextureMaterialType::BaseColor: bIsBaseColorValid = true; break;
-	case ERuntimeVirtualTextureMaterialType::BaseColor_Normal: bIsBaseColorValid = bIsNormalValid = true; break;
-	case ERuntimeVirtualTextureMaterialType::BaseColor_Normal_Specular:
-	case ERuntimeVirtualTextureMaterialType::BaseColor_Normal_Specular_Ex: bIsBaseColorValid = bIsNormalValid = bIsSpecularValid = true; break;
+	case ERuntimeVirtualTextureMaterialType::BaseColor_Normal_Specular: bIsBaseColorValid = bIsNormalValid = bIsSpecularValid = true; break;
 	case ERuntimeVirtualTextureMaterialType::BaseColor_Normal_Specular_YCoCg: bIsBaseColorValid = bIsNormalValid = bIsSpecularValid = true; break;
 	case ERuntimeVirtualTextureMaterialType::WorldHeight: bIsWorldHeightValid = true; break;
 	}
@@ -2209,8 +2214,7 @@ int32 UMaterialExpressionRuntimeVirtualTextureSample::Compile(class FMaterialCom
 		{
 			switch (MaterialType)
 			{
-			case ERuntimeVirtualTextureMaterialType::BaseColor_Normal_Specular:
-			case ERuntimeVirtualTextureMaterialType::BaseColor_Normal_Specular_Ex:  UnpackTarget = 1; UnpackMask = 0x1; break;
+			case ERuntimeVirtualTextureMaterialType::BaseColor_Normal_Specular:  UnpackTarget = 1; UnpackMask = 0x1; break;
 			case ERuntimeVirtualTextureMaterialType::BaseColor_Normal_Specular_YCoCg: UnpackTarget = 2; UnpackMask = 0x1; break;
 			}
 		}
@@ -2224,8 +2228,7 @@ int32 UMaterialExpressionRuntimeVirtualTextureSample::Compile(class FMaterialCom
 		{
 			switch (MaterialType)
 			{
-			case ERuntimeVirtualTextureMaterialType::BaseColor_Normal_Specular: UnpackTarget = 1; UnpackMask = 0x4; break;
-			case ERuntimeVirtualTextureMaterialType::BaseColor_Normal_Specular_Ex: UnpackTarget = 1; UnpackMask = 0x2; break;
+			case ERuntimeVirtualTextureMaterialType::BaseColor_Normal_Specular: UnpackTarget = 1; UnpackMask = 0x2; break;
 			case ERuntimeVirtualTextureMaterialType::BaseColor_Normal_Specular_YCoCg: UnpackTarget = 2; UnpackMask = 0x2; break;
 			}
 		}
@@ -2239,9 +2242,7 @@ int32 UMaterialExpressionRuntimeVirtualTextureSample::Compile(class FMaterialCom
 		{
 			switch (MaterialType)
 			{
-			case ERuntimeVirtualTextureMaterialType::BaseColor_Normal: UnpackType = EVirtualTextureUnpackType::NormalBC5; break;
-			case ERuntimeVirtualTextureMaterialType::BaseColor_Normal_Specular: UnpackType = EVirtualTextureUnpackType::NormalBC3; break;
-			case ERuntimeVirtualTextureMaterialType::BaseColor_Normal_Specular_Ex: UnpackType = EVirtualTextureUnpackType::NormalBC3BC3; break;
+			case ERuntimeVirtualTextureMaterialType::BaseColor_Normal_Specular: UnpackType = EVirtualTextureUnpackType::NormalBC3BC3; break;
 			case ERuntimeVirtualTextureMaterialType::BaseColor_Normal_Specular_YCoCg: UnpackType = EVirtualTextureUnpackType::NormalBC5BC1; break;
 			}
 		}
diff --git a/Engine/Source/Runtime/Engine/Private/VT/RuntimeVirtualTexture.cpp b/Engine/Source/Runtime/Engine/Private/VT/RuntimeVirtualTexture.cpp
--- a/Engine/Source/Runtime/Engine/Private/VT/RuntimeVirtualTexture.cpp
+++ b/Engine/Source/Runtime/Engine/Private/VT/RuntimeVirtualTexture.cpp
@@ -262,9 +262,7 @@ int32 URuntimeVirtualTexture::GetLayerCount(ERuntimeVirtualTextureMaterialType I
 	case ERuntimeVirtualTextureMaterialType::BaseColor:
 	case ERuntimeVirtualTextureMaterialType::WorldHeight:
 		return 1;
-	case ERuntimeVirtualTextureMaterialType::BaseColor_Normal:
 	case ERuntimeVirtualTextureMaterialType::BaseColor_Normal_Specular:
-	case ERuntimeVirtualTextureMaterialType::BaseColor_Normal_Specular_Ex:
 		return 2;
 	case ERuntimeVirtualTextureMaterialType::BaseColor_Normal_Specular_YCoCg:
 		return 3;
@@ -289,10 +287,8 @@ EPixelFormat URuntimeVirtualTexture::GetLayerFormat(int32 LayerIndex) const
 		switch (MaterialType)
 		{
 		case ERuntimeVirtualTextureMaterialType::BaseColor:
-		case ERuntimeVirtualTextureMaterialType::BaseColor_Normal:
-		case ERuntimeVirtualTextureMaterialType::BaseColor_Normal_Specular:
 			return bCompressTextures ? PF_DXT1 : PF_B8G8R8A8;
-		case ERuntimeVirtualTextureMaterialType::BaseColor_Normal_Specular_Ex:
+		case ERuntimeVirtualTextureMaterialType::BaseColor_Normal_Specular:
 		case ERuntimeVirtualTextureMaterialType::BaseColor_Normal_Specular_YCoCg:
 			return bCompressTextures ? PF_DXT5 : PF_B8G8R8A8;
 		case ERuntimeVirtualTextureMaterialType::WorldHeight:
@@ -305,12 +301,10 @@ EPixelFormat URuntimeVirtualTexture::GetLayerFormat(int32 LayerIndex) const
 	{
 		switch (MaterialType)
 		{
-		case ERuntimeVirtualTextureMaterialType::BaseColor_Normal:
-		case ERuntimeVirtualTextureMaterialType::BaseColor_Normal_Specular_YCoCg:
-			return bCompressTextures ? PF_BC5 : PF_B8G8R8A8;
 		case ERuntimeVirtualTextureMaterialType::BaseColor_Normal_Specular:
-		case ERuntimeVirtualTextureMaterialType::BaseColor_Normal_Specular_Ex:
 			return bCompressTextures ? PF_DXT5 : PF_B8G8R8A8;
+		case ERuntimeVirtualTextureMaterialType::BaseColor_Normal_Specular_YCoCg:
+			return bCompressTextures ? PF_BC5 : PF_B8G8R8A8;
 		default:
 			break;
 		}
@@ -336,9 +330,7 @@ bool URuntimeVirtualTexture::IsLayerSRGB(int32 LayerIndex) const
 	switch (MaterialType)
 	{
 	case ERuntimeVirtualTextureMaterialType::BaseColor:
-	case ERuntimeVirtualTextureMaterialType::BaseColor_Normal:
 	case ERuntimeVirtualTextureMaterialType::BaseColor_Normal_Specular:
-	case ERuntimeVirtualTextureMaterialType::BaseColor_Normal_Specular_Ex:
 		// Only BaseColor layer is sRGB
 		return LayerIndex == 0;
 	case ERuntimeVirtualTextureMaterialType::BaseColor_Normal_Specular_YCoCg:
@@ -454,6 +446,12 @@ void URuntimeVirtualTexture::PostLoad()
 		Size_DEPRECATED = -1;
 	}
 
+	// Convert BaseColor_Normal_DEPRECATED
+	if (MaterialType == ERuntimeVirtualTextureMaterialType::BaseColor_Normal_DEPRECATED)
+	{
+		MaterialType = ERuntimeVirtualTextureMaterialType::BaseColor_Normal_Specular;
+	}
+
 	Super::PostLoad();
 }
 
diff --git a/Engine/Source/Runtime/Engine/Public/VT/RuntimeVirtualTextureEnum.h b/Engine/Source/Runtime/Engine/Public/VT/RuntimeVirtualTextureEnum.h
--- a/Engine/Source/Runtime/Engine/Public/VT/RuntimeVirtualTextureEnum.h
+++ b/Engine/Source/Runtime/Engine/Public/VT/RuntimeVirtualTextureEnum.h
@@ -17,10 +17,9 @@ UENUM()
 enum class ERuntimeVirtualTextureMaterialType : uint8
 {
 	BaseColor UMETA(DisplayName = "Base Color"),
-	BaseColor_Normal UMETA(DisplayName = "Base Color, Normal"),
-	BaseColor_Normal_Specular UMETA(DisplayName = "[BC1|BC3] Base Color, Normal, Roughness, Specular"),
-	BaseColor_Normal_Specular_Ex UMETA(DisplayName = "[BC3|BC3] Base Color, Normal, Roughness, Specular"),
-	BaseColor_Normal_Specular_YCoCg UMETA(DisplayName = "[BC3|BC5|BC1] Base Color, Normal, Roughness, Specular"),
+	BaseColor_Normal_DEPRECATED UMETA(Hidden),
+	BaseColor_Normal_Specular UMETA(DisplayName = "Base Color, Normal, Roughness, Specular"),
+	BaseColor_Normal_Specular_YCoCg UMETA(DisplayName = "YCoCg Base Color, Normal, Roughness, Specular", ToolTip="Base Color is stored in YCoCg space. This requires more memory but may provide better quality."),
 	WorldHeight UMETA(DisplayName = "World Height"),
 
 	Count UMETA(Hidden),
diff --git a/Engine/Source/Runtime/Renderer/Private/VT/RuntimeVirtualTextureRender.cpp b/Engine/Source/Runtime/Renderer/Private/VT/RuntimeVirtualTextureRender.cpp
--- a/Engine/Source/Runtime/Renderer/Private/VT/RuntimeVirtualTextureRender.cpp
+++ b/Engine/Source/Runtime/Renderer/Private/VT/RuntimeVirtualTextureRender.cpp
@@ -78,23 +78,6 @@ namespace RuntimeVirtualTexture
 		}
 	};
 
-	/** Specialization for ERuntimeVirtualTextureMaterialType::BaseColor_Normal */
-	class FMaterialPolicy_BaseColorNormal
-	{
-	public:
-		static void ModifyCompilationEnvironment(FShaderCompilerEnvironment& OutEnvironment)
-		{
-			OutEnvironment.SetDefine(TEXT("OUT_BASECOLOR_NORMAL"), 1);
-		}
-
-		static FRHIBlendState* GetBlendState()
-		{
-			return TStaticBlendState<
-				CW_RGBA, BO_Add, BF_One, BF_InverseSourceAlpha, BO_Add, BF_Zero, BF_One,
-				CW_RGBA, BO_Add, BF_One, BF_InverseSourceAlpha, BO_Add, BF_Zero, BF_One	>::GetRHI();
-		}
-	};
-
 	/** Specialization for ERuntimeVirtualTextureMaterialType::BaseColor_Normal_Specular */
 	class FMaterialPolicy_BaseColorNormalSpecular
 	{
@@ -181,7 +164,6 @@ namespace RuntimeVirtualTexture
 	IMPLEMENT_MATERIAL_SHADER_TYPE(template<>,TVirtualTexturePS##PolicyName, TEXT("/Engine/Private/VirtualTextureMaterial.usf"), TEXT("MainPS"), SF_Pixel);
 
 	IMPLEMENT_VIRTUALTEXTURE_SHADER_TYPE(FMaterialPolicy_BaseColor, BaseColor);
-	IMPLEMENT_VIRTUALTEXTURE_SHADER_TYPE(FMaterialPolicy_BaseColorNormal, BaseColorNormal);
 	IMPLEMENT_VIRTUALTEXTURE_SHADER_TYPE(FMaterialPolicy_BaseColorNormalSpecular, BaseColorNormalSpecular);
 	IMPLEMENT_VIRTUALTEXTURE_SHADER_TYPE(FMaterialPolicy_WorldHeight, WorldHeight);
 
@@ -262,11 +244,7 @@ namespace RuntimeVirtualTexture
 					case ERuntimeVirtualTextureMaterialType::BaseColor:
 						Process<FMaterialPolicy_BaseColor>(MeshBatch, BatchElementMask, StaticMeshId, PrimitiveSceneProxy, MaterialRenderProxy, Material);
 						break;
-					case ERuntimeVirtualTextureMaterialType::BaseColor_Normal:
-						Process<FMaterialPolicy_BaseColorNormal>(MeshBatch, BatchElementMask, StaticMeshId, PrimitiveSceneProxy, MaterialRenderProxy, Material);
-						break;
 					case ERuntimeVirtualTextureMaterialType::BaseColor_Normal_Specular:
-					case ERuntimeVirtualTextureMaterialType::BaseColor_Normal_Specular_Ex:
 					case ERuntimeVirtualTextureMaterialType::BaseColor_Normal_Specular_YCoCg:
 						Process<FMaterialPolicy_BaseColorNormalSpecular>(MeshBatch, BatchElementMask, StaticMeshId, PrimitiveSceneProxy, MaterialRenderProxy, Material);
 						break;
@@ -481,24 +459,15 @@ namespace RuntimeVirtualTexture
 	typedef FShader_VirtualTextureCompress_CS< ERuntimeVirtualTextureMaterialType::BaseColor, false > FShader_VirtualTextureCompress_BaseColor_CS;
 	IMPLEMENT_SHADER_TYPE(template<>, FShader_VirtualTextureCompress_BaseColor_CS, TEXT("/Engine/Private/VirtualTextureCompress.usf"), TEXT("CompressBaseColorCS"), SF_Compute);
 	
-	typedef FShader_VirtualTextureCompress_CS< ERuntimeVirtualTextureMaterialType::BaseColor_Normal, false > FShader_VirtualTextureCompress_BaseColorNormal_CS;
-	IMPLEMENT_SHADER_TYPE(template<>, FShader_VirtualTextureCompress_BaseColorNormal_CS, TEXT("/Engine/Private/VirtualTextureCompress.usf"), TEXT("CompressBaseColorNormalCS"), SF_Compute);
-
 	typedef FShader_VirtualTextureCompress_CS< ERuntimeVirtualTextureMaterialType::BaseColor_Normal_Specular, false > FShader_VirtualTextureCompress_BaseColorNormalSpecular_CS;
 	IMPLEMENT_SHADER_TYPE(template<>, FShader_VirtualTextureCompress_BaseColorNormalSpecular_CS, TEXT("/Engine/Private/VirtualTextureCompress.usf"), TEXT("CompressBaseColorNormalSpecularCS"), SF_Compute);
 
-	typedef FShader_VirtualTextureCompress_CS< ERuntimeVirtualTextureMaterialType::BaseColor_Normal_Specular_Ex, false > FShader_VirtualTextureCompress_BaseColorNormalSpecularEx_CS;
-	IMPLEMENT_SHADER_TYPE(template<>, FShader_VirtualTextureCompress_BaseColorNormalSpecularEx_CS, TEXT("/Engine/Private/VirtualTextureCompress.usf"), TEXT("CompressBaseColorNormalSpecularExCS"), SF_Compute);
-
 	typedef FShader_VirtualTextureCompress_CS< ERuntimeVirtualTextureMaterialType::BaseColor_Normal_Specular_YCoCg, false > FShader_VirtualTextureCompress_BaseColorNormalSpecularYCoCg_CS;
 	IMPLEMENT_SHADER_TYPE(template<>, FShader_VirtualTextureCompress_BaseColorNormalSpecularYCoCg_CS, TEXT("/Engine/Private/VirtualTextureCompress.usf"), TEXT("CompressBaseColorNormalSpecularYCoCgCS"), SF_Compute);
 
 	typedef FShader_VirtualTextureCompress_CS< ERuntimeVirtualTextureMaterialType::BaseColor_Normal_Specular, true > FShader_VirtualTextureCompress_BaseColorNormalSpecular_CopyOnly_CS;
 	IMPLEMENT_SHADER_TYPE(template<>, FShader_VirtualTextureCompress_BaseColorNormalSpecular_CopyOnly_CS, TEXT("/Engine/Private/VirtualTextureCompress.usf"), TEXT("CopyNormalSpecularCS"), SF_Compute);
 
-	typedef FShader_VirtualTextureCompress_CS< ERuntimeVirtualTextureMaterialType::BaseColor_Normal_Specular_Ex, true > FShader_VirtualTextureCompress_BaseColorNormalSpecularEx_CopyOnly_CS;
-	IMPLEMENT_SHADER_TYPE(template<>, FShader_VirtualTextureCompress_BaseColorNormalSpecularEx_CopyOnly_CS, TEXT("/Engine/Private/VirtualTextureCompress.usf"), TEXT("CopyNormalSpecularExCS"), SF_Compute);
-
 	typedef FShader_VirtualTextureCompress_CS< ERuntimeVirtualTextureMaterialType::BaseColor_Normal_Specular_YCoCg, true > FShader_VirtualTextureCompress_BaseColorNormalSpecularYCoCg_CopyOnly_CS;
 	IMPLEMENT_SHADER_TYPE(template<>, FShader_VirtualTextureCompress_BaseColorNormalSpecularYCoCg_CopyOnly_CS, TEXT("/Engine/Private/VirtualTextureCompress.usf"), TEXT("CopyNormalSpecularYCoCgCS"), SF_Compute);
 
@@ -530,15 +499,9 @@ namespace RuntimeVirtualTexture
 		case ERuntimeVirtualTextureMaterialType::BaseColor:
 			AddCompressOrCopyPass<ERuntimeVirtualTextureMaterialType::BaseColor>(GraphBuilder, FeatureLevel, Parameters, GroupCount);
 			break;
-		case ERuntimeVirtualTextureMaterialType::BaseColor_Normal:
-			AddCompressOrCopyPass<ERuntimeVirtualTextureMaterialType::BaseColor_Normal>(GraphBuilder, FeatureLevel, Parameters, GroupCount);
-			break;
 		case ERuntimeVirtualTextureMaterialType::BaseColor_Normal_Specular:
 			AddCompressOrCopyPass<ERuntimeVirtualTextureMaterialType::BaseColor_Normal_Specular>(GraphBuilder, FeatureLevel, Parameters, GroupCount);
 			break;
-		case ERuntimeVirtualTextureMaterialType::BaseColor_Normal_Specular_Ex:
-			AddCompressOrCopyPass<ERuntimeVirtualTextureMaterialType::BaseColor_Normal_Specular_Ex>(GraphBuilder, FeatureLevel, Parameters, GroupCount);
-			break;
 		case ERuntimeVirtualTextureMaterialType::BaseColor_Normal_Specular_YCoCg:
 			AddCompressOrCopyPass<ERuntimeVirtualTextureMaterialType::BaseColor_Normal_Specular_YCoCg>(GraphBuilder, FeatureLevel, Parameters, GroupCount);
 			break;
@@ -550,16 +513,13 @@ namespace RuntimeVirtualTexture
 	{
 		const FIntVector GroupCount((TextureSize.X + 7) / 8, (TextureSize.Y + 7) / 8, 1);
 
-		// Needed for BaseColor_Normal_Specular where we need to pack normal and specular into one VT layer
+		// Needed for BaseColor_Normal_Specular where we pack layers
 		// Needed for WorldHeight thumbnails where we need to copy G16 to the RGBA8 thumbnail
 		switch (MaterialType)
 		{
 		case ERuntimeVirtualTextureMaterialType::BaseColor_Normal_Specular:
 			AddCompressOrCopyPass<ERuntimeVirtualTextureMaterialType::BaseColor_Normal_Specular, true>(GraphBuilder, FeatureLevel, Parameters, GroupCount);
 			break;
-		case ERuntimeVirtualTextureMaterialType::BaseColor_Normal_Specular_Ex:
-			AddCompressOrCopyPass<ERuntimeVirtualTextureMaterialType::BaseColor_Normal_Specular_Ex, true>(GraphBuilder, FeatureLevel, Parameters, GroupCount);
-			break;
 		case ERuntimeVirtualTextureMaterialType::BaseColor_Normal_Specular_YCoCg:
 			AddCompressOrCopyPass<ERuntimeVirtualTextureMaterialType::BaseColor_Normal_Specular_YCoCg, true>(GraphBuilder, FeatureLevel, Parameters, GroupCount);
 			break;
@@ -579,7 +539,6 @@ namespace RuntimeVirtualTexture
 			bRenderPass = OutputTexture0 != nullptr;
 			bCompressPass = bRenderPass && (OutputTexture0->GetFormat() == PF_DXT1 || OutputTexture0->GetFormat() == PF_DXT5 || OutputTexture0->GetFormat() == PF_BC5);
 			bCopyPass = bRenderPass && !bCompressPass && MaterialType == ERuntimeVirtualTextureMaterialType::BaseColor_Normal_Specular;
-			bCopyPass |= bRenderPass && !bCompressPass && MaterialType == ERuntimeVirtualTextureMaterialType::BaseColor_Normal_Specular_Ex;
 			bCopyPass |= bRenderPass && !bCompressPass && MaterialType == ERuntimeVirtualTextureMaterialType::BaseColor_Normal_Specular_YCoCg;
 			bCopyPass |= bRenderPass && MaterialType == ERuntimeVirtualTextureMaterialType::WorldHeight && OutputTexture0->GetFormat() == PF_B8G8R8A8;
 		
@@ -595,36 +554,7 @@ namespace RuntimeVirtualTexture
 					OutputAlias0 = CompressTexture0_u2 = GraphBuilder.CreateTexture(FPooledRenderTargetDesc::Create2DDesc(TextureSize / 4, PF_R32G32_UINT, FClearValueBinding::None, TexCreate_None, TexCreate_UAV, false), TEXT("CompressTexture0"));
 				}
 				break;
-			case ERuntimeVirtualTextureMaterialType::BaseColor_Normal:
-				if (bRenderPass)
-				{
-					OutputAlias0 = RenderTexture0 = GraphBuilder.CreateTexture(FPooledRenderTargetDesc::Create2DDesc(TextureSize, PF_B8G8R8A8, FClearValueBinding::Black, TexCreate_SRGB, TexCreate_RenderTargetable, false), TEXT("RenderTexture0"));
-					OutputAlias1 = RenderTexture1 = GraphBuilder.CreateTexture(FPooledRenderTargetDesc::Create2DDesc(TextureSize, PF_B8G8R8A8, FClearValueBinding::DefaultNormal8Bit, TexCreate_None, TexCreate_RenderTargetable, false), TEXT("RenderTexture1"));
-				}
-				if (bCompressPass)
-				{
-					OutputAlias0 = CompressTexture0_u4 = GraphBuilder.CreateTexture(FPooledRenderTargetDesc::Create2DDesc(TextureSize / 4, PF_R32G32B32A32_UINT, FClearValueBinding::None, TexCreate_None, TexCreate_UAV, false), TEXT("CompressTexture0"));
-					OutputAlias1 = CompressTexture1 = GraphBuilder.CreateTexture(FPooledRenderTargetDesc::Create2DDesc(TextureSize / 4, PF_R32G32B32A32_UINT, FClearValueBinding::None, TexCreate_None, TexCreate_UAV, false), TEXT("CompressTexture1"));
-				}
-				break;
 			case ERuntimeVirtualTextureMaterialType::BaseColor_Normal_Specular:
-				if (bRenderPass)
-				{
-					OutputAlias0 = RenderTexture0 = GraphBuilder.CreateTexture(FPooledRenderTargetDesc::Create2DDesc(TextureSize, PF_B8G8R8A8, FClearValueBinding::Black, TexCreate_SRGB, TexCreate_RenderTargetable, false), TEXT("RenderTexture0"));
-					RenderTexture1 = GraphBuilder.CreateTexture(FPooledRenderTargetDesc::Create2DDesc(TextureSize, PF_B8G8R8A8, FClearValueBinding::DefaultNormal8Bit, TexCreate_None, TexCreate_RenderTargetable, false), TEXT("RenderTexture1"));
-					RenderTexture2 = GraphBuilder.CreateTexture(FPooledRenderTargetDesc::Create2DDesc(TextureSize, PF_B8G8R8A8, FClearValueBinding::DefaultNormal8Bit, TexCreate_None, TexCreate_RenderTargetable, false), TEXT("RenderTexture2"));
-				}
-				if (bCompressPass)
-				{
-					OutputAlias0 = CompressTexture0_u2 = GraphBuilder.CreateTexture(FPooledRenderTargetDesc::Create2DDesc(TextureSize / 4, PF_R32G32_UINT, FClearValueBinding::None, TexCreate_None, TexCreate_UAV, false), TEXT("CompressTexture0"));
-					OutputAlias1 = CompressTexture1 = GraphBuilder.CreateTexture(FPooledRenderTargetDesc::Create2DDesc(TextureSize / 4, PF_R32G32B32A32_UINT, FClearValueBinding::None, TexCreate_None, TexCreate_UAV, false), TEXT("CompressTexture1"));
-				}
-				if (bCopyPass)
-				{
-					OutputAlias1 = CopyTexture0 = GraphBuilder.CreateTexture(FPooledRenderTargetDesc::Create2DDesc(TextureSize, PF_B8G8R8A8, FClearValueBinding::None, TexCreate_None, TexCreate_UAV, false), TEXT("CopyTexture0"));
-				}
-				break;
-			case ERuntimeVirtualTextureMaterialType::BaseColor_Normal_Specular_Ex:
 				if (bRenderPass)
 				{
 					OutputAlias0 = RenderTexture0 = GraphBuilder.CreateTexture(FPooledRenderTargetDesc::Create2DDesc(TextureSize, PF_B8G8R8A8, FClearValueBinding::Black, TexCreate_SRGB, TexCreate_RenderTargetable, false), TEXT("RenderTexture0"));
