Fixes shader compilation failure in DOF due to more than 8 UAV in reduce pass when enabling alpha channel on project

diff --git a/Engine/Shaders/Private/DiaphragmDOF/DOFReduce.usf b/Engine/Shaders/Private/DiaphragmDOF/DOFReduce.usf
--- a/Engine/Shaders/Private/DiaphragmDOF/DOFReduce.usf
+++ b/Engine/Shaders/Private/DiaphragmDOF/DOFReduce.usf
@@ -538,12 +538,17 @@ void ReduceCS(
 		GatheredColor.rgb *= (1 - ScatterFactor);
 	}
 	#endif // DIM_HYBRID_SCATTER_FGD || DIM_HYBRID_SCATTER_BGD
-
-	// Unconditionaly output mip 0 in output buffer that is always a multiple of group size,
+	
 	// for gathering pass to avoid having to clamp UV in gathering pass. 
 	uint2 OutputPixelId = DispatchThreadId;
-	OutputMipLevel(/* MipLevel = */ 0, OutputPixelId, GatheredColor, CocRadius);
 	
+	// Output mip 0 in output buffer that is always a multiple of group size. Unless reducing lower mips.
+	#if !DIM_DISABLE_OUTPUT_MIP0
+	{
+		OutputMipLevel(/* MipLevel = */ 0, OutputPixelId, GatheredColor, CocRadius);
+	}
+	#endif
+
 	FCocDownsampleParams DownsampleParams;
 	DownsampleParams.CocRadiusMultiplier = 1.0;
 	DownsampleParams.FrameExposureScale = FrameExposureScale;
diff --git a/Engine/Source/Runtime/Renderer/Private/PostProcess/DiaphragmDOF.cpp b/Engine/Source/Runtime/Renderer/Private/PostProcess/DiaphragmDOF.cpp
--- a/Engine/Source/Runtime/Renderer/Private/PostProcess/DiaphragmDOF.cpp
+++ b/Engine/Source/Runtime/Renderer/Private/PostProcess/DiaphragmDOF.cpp
@@ -495,6 +495,11 @@ BEGIN_SHADER_PARAMETER_STRUCT(FDOFGatherInputTextures, )
 	SHADER_PARAMETER_RDG_TEXTURE(Texture2D, SeparateCoc)
 END_SHADER_PARAMETER_STRUCT()
 
+BEGIN_SHADER_PARAMETER_STRUCT(FDOFGatherInputSRVs, )
+	SHADER_PARAMETER_RDG_TEXTURE_SRV(Texture2D, SceneColor)
+	SHADER_PARAMETER_RDG_TEXTURE_SRV(Texture2D, SeparateCoc)
+END_SHADER_PARAMETER_STRUCT()
+
 BEGIN_SHADER_PARAMETER_STRUCT(FDOFGatherInputUAVs, )
 	SHADER_PARAMETER_RDG_TEXTURE_UAV(RWTexture2D<float4>, SceneColor)
 	SHADER_PARAMETER_RDG_TEXTURE_UAV(RWTexture2D<float4>, SeparateCoc)
@@ -511,6 +516,17 @@ FDOFGatherInputTextures CreateTextures(FRDGBuilder& GraphBuilder, const FDOFGath
 	return Textures;
 }
 
+FDOFGatherInputSRVs CreateSRVs(FRDGBuilder& GraphBuilder, const FDOFGatherInputTextures& Textures, uint8 MipLevel = 0)
+{
+	FDOFGatherInputSRVs SRVs;
+	SRVs.SceneColor = GraphBuilder.CreateSRV(FRDGTextureSRVDesc::CreateForMipLevel(Textures.SceneColor, MipLevel));
+	if (Textures.SeparateCoc)
+	{
+		SRVs.SeparateCoc = GraphBuilder.CreateSRV(FRDGTextureSRVDesc::CreateForMipLevel(Textures.SeparateCoc, MipLevel));
+	}
+	return SRVs;
+}
+
 FDOFGatherInputUAVs CreateUAVs(FRDGBuilder& GraphBuilder, const FDOFGatherInputTextures& Textures, uint8 MipLevel = 0)
 {
 	FDOFGatherInputUAVs UAVs;
@@ -769,17 +785,52 @@ class FDiaphragmDOFReduceCS : public FDiaphragmDOFShader
 	class FReduceMipCount : SHADER_PERMUTATION_RANGE_INT("DIM_REDUCE_MIP_COUNT", 2, 3);
 	class FHybridScatterForeground : SHADER_PERMUTATION_BOOL("DIM_HYBRID_SCATTER_FGD");
 	class FHybridScatterBackground : SHADER_PERMUTATION_BOOL("DIM_HYBRID_SCATTER_BGD");
+	class FDisableOutputMip0 : SHADER_PERMUTATION_BOOL("DIM_DISABLE_OUTPUT_MIP0");
 
 	using FPermutationDomain = TShaderPermutationDomain<
 		FReduceMipCount,
 		FHybridScatterForeground,
 		FHybridScatterBackground,
-		FDDOFRGBColorBufferDim>;
+		FDDOFRGBColorBufferDim,
+		FDisableOutputMip0>;
+
+
+	/** Returns the number of mip level the reduce pass is able to output. */
+	static int32 GetMaxReductionMipLevelCount()
+	{
+		//return 2; // TODO
+
+		// Can only have 8 UAVs, but need to output 3x UAV for hybird scatter + 2 UAV per mips for RGBA + SeparateCoc.
+		return IsPostProcessingWithAlphaChannelSupported() ? 2 : kMaxMipLevelCount;
+	}
 
 	static bool ShouldCompilePermutation(const FGlobalShaderPermutationParameters& Parameters)
 	{
 		FPermutationDomain PermutationVector(Parameters.PermutationId);
 
+		if (PermutationVector.Get<FDisableOutputMip0>())
+		{
+			// Hybrid scatter is done only on first reduction, that needs to output mip0
+			if (PermutationVector.Get<FHybridScatterForeground>() || PermutationVector.Get<FHybridScatterBackground>())
+			{
+				return false;
+			}
+
+			// Do not output mip level that are more than supported.
+			if (PermutationVector.Get<FReduceMipCount>() > (GetMaxReductionMipLevelCount() + 1))
+			{
+				return false;
+			}
+		}
+		else
+		{
+			// Do not output mip level that are more than supported.
+			if (PermutationVector.Get<FReduceMipCount>() > GetMaxReductionMipLevelCount())
+			{
+				return false;
+			}
+		}
+
 		// Do not compile storing Coc independently of RGB if not supported.
 		if (PermutationVector.Get<FDDOFRGBColorBufferDim>() && !SupportsRGBColorBuffer(Parameters.Platform))
 		{
@@ -809,7 +860,7 @@ class FDiaphragmDOFReduceCS : public FDiaphragmDOFShader
 		SHADER_PARAMETER_STRUCT_INCLUDE(FDOFCommonShaderParameters, CommonParameters)
 
 		SHADER_PARAMETER(FVector4, GatherInputSize)
-		SHADER_PARAMETER_STRUCT(FDOFGatherInputTextures, GatherInput)
+		SHADER_PARAMETER_STRUCT(FDOFGatherInputSRVs, GatherInput)
 		
 		SHADER_PARAMETER(FVector4, QuarterResGatherInputSize)
 		SHADER_PARAMETER_STRUCT(FDOFGatherInputTextures, QuarterResGatherInput)
@@ -1780,7 +1831,10 @@ FRDGTextureRef DiaphragmDOF::AddPasses(
 		FIntPoint SrcSize = HalfResGatherInputDescs.SceneColor.Extent;
 		
 		// Compute the number of mip level required by the gathering pass.
-		const int32 MipLevelCount = FMath::Clamp(FMath::CeilToInt(FMath::Log2(MaxBluringRadius * 0.5 / HalfResRingCount)) + (bUseLowAccumulatorQuality ? 1 : 0), 2, kMaxMipLevelCount);
+		const int32 MipLevelCount = FMath::Max(FMath::CeilToInt(FMath::Log2(MaxBluringRadius * 0.5 / HalfResRingCount)) + (bUseLowAccumulatorQuality ? 1 : 0), 2);
+
+		// Maximum number of mip level that can be done.
+		const int32 MaxReductionMipLevelCount = FDiaphragmDOFReduceCS::GetMaxReductionMipLevelCount();
 
 		// Maximum number of scattering group per draw instance.
 		// TODO: depends.
@@ -1866,12 +1920,17 @@ FRDGTextureRef DiaphragmDOF::AddPasses(
 				BackgroundScatterDrawListBuffer = GraphBuilder.CreateBuffer(DrawListDescs, TEXT("DOFBackgroundDrawList"));
 		}
 		
-		// Adds the reduce pass.
+		// Number of mip level that has been reduced.
+		int32 ReducedMipLevelCount;
+
+		// Adds the first reduce pass.
 		{
+			int32 ProcessingMipLevelCount = FMath::Min(MipLevelCount, MaxReductionMipLevelCount);
+
 			FIntPoint PassViewSize = PreprocessViewSize;
 
 			FDiaphragmDOFReduceCS::FPermutationDomain PermutationVector;
-			PermutationVector.Set<FDiaphragmDOFReduceCS::FReduceMipCount>(MipLevelCount);
+			PermutationVector.Set<FDiaphragmDOFReduceCS::FReduceMipCount>(ProcessingMipLevelCount);
 			PermutationVector.Set<FDiaphragmDOFReduceCS::FHybridScatterForeground>(bForegroundHybridScattering);
 			PermutationVector.Set<FDiaphragmDOFReduceCS::FHybridScatterBackground>(bBackgroundHybridScattering);
 			PermutationVector.Set<FDDOFRGBColorBufferDim>(bRGBBufferSeparateCocBuffer);
@@ -1890,12 +1949,12 @@ FRDGTextureRef DiaphragmDOF::AddPasses(
 			PassParameters->CommonParameters = CommonParameters;
 
 			PassParameters->GatherInputSize = FVector4(SrcSize.X, SrcSize.Y, 1.0f / SrcSize.X, 1.0f / SrcSize.Y);
-			PassParameters->GatherInput = HalfResGatherInputTextures;
+			PassParameters->GatherInput = CreateSRVs(GraphBuilder, HalfResGatherInputTextures);
 			
 			PassParameters->QuarterResGatherInputSize = FVector4(SrcSize.X / 2, SrcSize.Y / 2, 2.0f / SrcSize.X, 2.0f / SrcSize.Y);
 			PassParameters->QuarterResGatherInput = QuarterResGatherInputTextures;
 
-			for (int32 MipLevel = 0; MipLevel < MipLevelCount; MipLevel++)
+			for (int32 MipLevel = 0; MipLevel < ProcessingMipLevelCount; MipLevel++)
 			{
 				PassParameters->OutputMips[MipLevel] = CreateUAVs(GraphBuilder, ReducedGatherInputTextures, MipLevel);
 			}
@@ -1914,8 +1973,8 @@ FRDGTextureRef DiaphragmDOF::AddPasses(
 			TShaderMapRef<FDiaphragmDOFReduceCS> ComputeShader(View.ShaderMap, PermutationVector);
 			FComputeShaderUtils::AddPass(
 				GraphBuilder,
-				RDG_EVENT_NAME("DOF Reduce(Mips=%d FgdScatter=%s BgdScatter=%s%s) %dx%d",
-					MipLevelCount,
+				RDG_EVENT_NAME("DOF Reduce(Mips=[0;%d] FgdScatter=%s BgdScatter=%s%s) %dx%d",
+					ProcessingMipLevelCount - 1,
 					bForegroundHybridScattering ? TEXT("Yes") : TEXT("No"),
 					bBackgroundHybridScattering ? TEXT("Yes") : TEXT("No"),
 					bRGBBufferSeparateCocBuffer ? TEXT(" R11G11B10") : TEXT(""),
@@ -1923,6 +1982,58 @@ FRDGTextureRef DiaphragmDOF::AddPasses(
 				*ComputeShader,
 				PassParameters,
 				FComputeShaderUtils::GetGroupCount(PassViewSize, kDefaultGroupSize));
+
+			ReducedMipLevelCount = ProcessingMipLevelCount;
+		}
+
+		// Complete the reduction.
+		while (ReducedMipLevelCount < MipLevelCount)
+		{
+			int32 ProcessingMipLevelCount = FMath::Min(MipLevelCount - ReducedMipLevelCount, MaxReductionMipLevelCount);
+			int32 InputMipLevel = ReducedMipLevelCount - 1;
+
+			FIntPoint PassViewSize = FIntPoint::DivideAndRoundUp(PreprocessViewSize, 1 << InputMipLevel);
+
+			FDiaphragmDOFReduceCS::FPermutationDomain PermutationVector;
+			PermutationVector.Set<FDiaphragmDOFReduceCS::FReduceMipCount>(ProcessingMipLevelCount + 1);
+			PermutationVector.Set<FDDOFRGBColorBufferDim>(bRGBBufferSeparateCocBuffer);
+			PermutationVector.Set<FDiaphragmDOFReduceCS::FDisableOutputMip0>(true);
+
+			FDiaphragmDOFReduceCS::FParameters* PassParameters = GraphBuilder.AllocParameters<FDiaphragmDOFReduceCS::FParameters>();
+			PassParameters->ViewportRect = FIntRect(0, 0, PassViewSize.X, PassViewSize.Y);
+			PassParameters->MaxInputBufferUV = FVector2D(
+				(PreprocessViewSize.X - 0.5f) / SrcSize.X,
+				(PreprocessViewSize.Y - 0.5f) / SrcSize.Y);
+
+			PassParameters->EyeAdaptation = GetEyeAdaptationTexture(GraphBuilder, View);
+			PassParameters->CommonParameters = CommonParameters;
+
+			float InputMipLevelPow2 = 1 << InputMipLevel;
+			PassParameters->GatherInputSize = FVector4(SrcSize.X / InputMipLevelPow2, SrcSize.Y / InputMipLevelPow2, InputMipLevelPow2 / SrcSize.X, InputMipLevelPow2 / SrcSize.Y);
+			PassParameters->GatherInput = CreateSRVs(GraphBuilder, ReducedGatherInputTextures, InputMipLevel);
+
+			for (int32 MipLevel = 0; MipLevel < ProcessingMipLevelCount; MipLevel++)
+			{
+				PassParameters->OutputMips[1 + MipLevel] = CreateUAVs(GraphBuilder, ReducedGatherInputTextures, ReducedMipLevelCount + MipLevel);
+			}
+
+			// TODO(RDG): ERDGPassFlags::GenerateMips has no reason to exist.
+			TShaderMapRef<FDiaphragmDOFReduceCS> ComputeShader(View.ShaderMap, PermutationVector);
+			ClearUnusedGraphResources(*ComputeShader, PassParameters);
+			GraphBuilder.AddPass(
+				RDG_EVENT_NAME("DOF Reduce(Mips=[%d;%d]%s) %dx%d",
+					ReducedMipLevelCount,
+					ReducedMipLevelCount + ProcessingMipLevelCount - 1,
+					bRGBBufferSeparateCocBuffer ? TEXT(" R11G11B10") : TEXT(""),
+					PassViewSize.X, PassViewSize.Y),
+				PassParameters,
+				ERDGPassFlags::Compute | ERDGPassFlags::GenerateMips,
+				[PassParameters, ComputeShader, PassViewSize](FRHICommandList& RHICmdList)
+			{
+				FComputeShaderUtils::Dispatch(RHICmdList, *ComputeShader, *PassParameters, FComputeShaderUtils::GetGroupCount(PassViewSize, kDefaultGroupSize));
+			});
+
+			ReducedMipLevelCount += ProcessingMipLevelCount;
 		}
 		
 		// Pack multiple scattering group on same primitive instance to increase wave occupancy in the scattering vertex shader.
