Fix bug computing wrong sizes for mips of non-power-2 UDIM texture, when computing VT address wrapping

diff --git a/Engine/Source/Runtime/Renderer/Private/VT/AllocatedVirtualTexture.cpp b/Engine/Source/Runtime/Renderer/Private/VT/AllocatedVirtualTexture.cpp
--- a/Engine/Source/Runtime/Renderer/Private/VT/AllocatedVirtualTexture.cpp
+++ b/Engine/Source/Runtime/Renderer/Private/VT/AllocatedVirtualTexture.cpp
@@ -66,8 +66,9 @@ FAllocatedVirtualTexture::FAllocatedVirtualTexture(FVirtualTextureSystem* InSyst
 			const uint32 Local_vLevel = MaxLevel - MipBias;
 			check(Local_vLevel <= Producer->GetMaxLevel());
 
-			const uint32 RootWidthInTiles = FMath::Max(Producer->GetWidthInTiles() >> Local_vLevel, 1u);
-			const uint32 RootHeightInTiles = FMath::Max(Producer->GetHeightInTiles() >> Local_vLevel, 1u);
+			const uint32 MipScaleFactor = (1u << Local_vLevel);
+			const uint32 RootWidthInTiles = FMath::DivideAndRoundUp(Producer->GetWidthInTiles(), MipScaleFactor);
+			const uint32 RootHeightInTiles = FMath::DivideAndRoundUp(Producer->GetHeightInTiles(), MipScaleFactor);
 			for (uint32 TileY = 0u; TileY < RootHeightInTiles; ++TileY)
 			{
 				for (uint32 TileX = 0u; TileX < RootWidthInTiles; ++TileX)
@@ -114,8 +115,9 @@ void FAllocatedVirtualTexture::Release(FVirtualTextureSystem* System)
 			const uint32 Local_vLevel = MaxLevel - MipBias;
 			check(Local_vLevel <= Producer->GetMaxLevel());
 
-			const uint32 RootWidthInTiles = FMath::Max(Producer->GetWidthInTiles() >> Local_vLevel, 1u);
-			const uint32 RootHeightInTiles = FMath::Max(Producer->GetHeightInTiles() >> Local_vLevel, 1u);
+			const uint32 MipScaleFactor = (1u << Local_vLevel);
+			const uint32 RootWidthInTiles = FMath::DivideAndRoundUp(Producer->GetWidthInTiles(), MipScaleFactor);
+			const uint32 RootHeightInTiles = FMath::DivideAndRoundUp(Producer->GetHeightInTiles(), MipScaleFactor);
 			for (uint32 TileY = 0u; TileY < RootHeightInTiles; ++TileY)
 			{
 				for (uint32 TileX = 0u; TileX < RootWidthInTiles; ++TileX)
diff --git a/Engine/Source/Runtime/Renderer/Private/VT/VirtualTextureSystem.cpp b/Engine/Source/Runtime/Renderer/Private/VT/VirtualTextureSystem.cpp
--- a/Engine/Source/Runtime/Renderer/Private/VT/VirtualTextureSystem.cpp
+++ b/Engine/Source/Runtime/Renderer/Private/VT/VirtualTextureSystem.cpp
@@ -1469,8 +1469,9 @@ void FVirtualTextureSystem::GatherRequestsTask(const FGatherRequestsParameters&
 			// For square textures, this could simply be (Local_vAddress % NumTilesInMip), but that doesn't work for non-square
 			// Possible there is a more clever approach to take here
 			{
-				const uint32 ProducerMipWidthInTiles = FMath::Max(Producer->GetWidthInTiles() >> Local_vLevel, 1u);
-				const uint32 ProducerMipHeightInTiles = FMath::Max(Producer->GetHeightInTiles() >> Local_vLevel, 1u);
+				const uint32 MipScaleFactor = (1u << Local_vLevel);
+				const uint32 ProducerMipWidthInTiles = FMath::DivideAndRoundUp(Producer->GetWidthInTiles(), MipScaleFactor);
+				const uint32 ProducerMipHeightInTiles = FMath::DivideAndRoundUp(Producer->GetHeightInTiles(), MipScaleFactor);
 				uint32 Local_vTileX = FMath::ReverseMortonCode2(Local_vAddress);
 				uint32 Local_vTileY = FMath::ReverseMortonCode2(Local_vAddress >> 1);
 				Local_vTileX %= ProducerMipWidthInTiles;
