Don't merge duplicate streaming VT producers when creating page table
- Keep merging in place for RVT

diff --git a/Engine/Source/Runtime/Engine/Private/VT/RuntimeVirtualTexture.cpp b/Engine/Source/Runtime/Engine/Private/VT/RuntimeVirtualTexture.cpp
--- a/Engine/Source/Runtime/Engine/Private/VT/RuntimeVirtualTexture.cpp
+++ b/Engine/Source/Runtime/Engine/Private/VT/RuntimeVirtualTexture.cpp
@@ -138,6 +138,7 @@ protected:
 			VTDesc.TileBorderSize = InDesc.TileBorderSize;
 			VTDesc.NumTextureLayers = InDesc.NumTextureLayers;
 			VTDesc.bPrivateSpace = InPrivateSpace;
+			VTDesc.bShareDuplicateLayers = true;
 
 			for (uint32 LayerIndex = 0u; LayerIndex < VTDesc.NumTextureLayers; ++LayerIndex)
 			{
diff --git a/Engine/Source/Runtime/RenderCore/Public/VirtualTexturing.h b/Engine/Source/Runtime/RenderCore/Public/VirtualTexturing.h
--- a/Engine/Source/Runtime/RenderCore/Public/VirtualTexturing.h
+++ b/Engine/Source/Runtime/RenderCore/Public/VirtualTexturing.h
@@ -63,6 +63,13 @@ struct FAllocatedVTDescription
 			 * The system only supports a limited number of unique page tables, so this must be used carefully
 			 */
 			uint8 bPrivateSpace : 1;
+
+			/**
+			 * If the AllocatedVT has the same producer mapped to multiple layers, should those be merged into a single page table layer?
+			 * This can make for more efficient page tables when enabled, but certain code may make assumption that number of layers
+			 * specified when allocating VT exactly matches the resulting page page
+			 */
+			uint8 bShareDuplicateLayers : 1;
 		};
 	};
 };
diff --git a/Engine/Source/Runtime/Renderer/Private/VT/AllocatedVirtualTexture.cpp b/Engine/Source/Runtime/Renderer/Private/VT/AllocatedVirtualTexture.cpp
--- a/Engine/Source/Runtime/Renderer/Private/VT/AllocatedVirtualTexture.cpp
+++ b/Engine/Source/Runtime/Renderer/Private/VT/AllocatedVirtualTexture.cpp
@@ -198,13 +198,16 @@ uint32 FAllocatedVirtualTexture::AddUniqueProducer(FVirtualTextureProducerHandle
 
 uint32 FAllocatedVirtualTexture::AddUniquePhysicalSpace(FVirtualTexturePhysicalSpace* InPhysicalSpace, uint32 InUniqueProducerIndex, uint32 InProducerPhysicalSpaceIndex)
 {
-	for (int32 Index = 0u; Index < UniquePageTableLayers.Num(); ++Index)
+	if (Description.bShareDuplicateLayers)
 	{
-		if (UniquePageTableLayers[Index].PhysicalSpace == InPhysicalSpace && 
-			UniquePageTableLayers[Index].UniqueProducerIndex == InUniqueProducerIndex &&
-			UniquePageTableLayers[Index].ProducerPhysicalGroupIndex == InProducerPhysicalSpaceIndex)
+		for (int32 Index = 0u; Index < UniquePageTableLayers.Num(); ++Index)
 		{
-			return Index;
+			if (UniquePageTableLayers[Index].PhysicalSpace == InPhysicalSpace &&
+				UniquePageTableLayers[Index].UniqueProducerIndex == InUniqueProducerIndex &&
+				UniquePageTableLayers[Index].ProducerPhysicalGroupIndex == InProducerPhysicalSpaceIndex)
+			{
+				return Index;
+			}
 		}
 	}
 	const uint32 Index = UniquePageTableLayers.AddDefaulted();
diff --git a/Engine/Source/Runtime/Renderer/Private/VT/VirtualTextureSystem.cpp b/Engine/Source/Runtime/Renderer/Private/VT/VirtualTextureSystem.cpp
--- a/Engine/Source/Runtime/Renderer/Private/VT/VirtualTextureSystem.cpp
+++ b/Engine/Source/Runtime/Renderer/Private/VT/VirtualTextureSystem.cpp
@@ -469,17 +469,24 @@ IAllocatedVirtualTexture* FVirtualTextureSystem::AllocateVirtualTexture(const FA
 
 	// Sum the total number of physical groups from all producers
 	uint32 NumPhysicalGroups = 0;
-	TArray<FVirtualTextureProducer*> UniqueProducers;
-	for (uint32 LayerIndex = 0u; LayerIndex < Desc.NumTextureLayers; ++LayerIndex)
+	if (Desc.bShareDuplicateLayers)
 	{
-		if (ProducerForLayer[LayerIndex] != nullptr)
+		TArray<FVirtualTextureProducer*> UniqueProducers;
+		for (uint32 LayerIndex = 0u; LayerIndex < Desc.NumTextureLayers; ++LayerIndex)
+		{
+			if (ProducerForLayer[LayerIndex] != nullptr)
+			{
+				UniqueProducers.AddUnique(ProducerForLayer[LayerIndex]);
+			}
+		}
+		for (int32 ProducerIndex = 0u; ProducerIndex < UniqueProducers.Num(); ++ProducerIndex)
 		{
-			UniqueProducers.AddUnique(ProducerForLayer[LayerIndex]);
+			NumPhysicalGroups += UniqueProducers[ProducerIndex]->GetNumPhysicalGroups();
 		}
 	}
-	for (int32 ProducerIndex = 0u; ProducerIndex < UniqueProducers.Num(); ++ProducerIndex)
+	else
 	{
-		NumPhysicalGroups += UniqueProducers[ProducerIndex]->GetNumPhysicalGroups();
+		NumPhysicalGroups = Desc.NumTextureLayers;
 	}
 
 	FVTSpaceDescription SpaceDesc;
