Properly handle missing VT producers when creating allocated VTs
This can happen if textures used as material parameters are not properly converted to VT

diff --git a/Engine/Source/Runtime/Renderer/Private/VT/AllocatedVirtualTexture.cpp b/Engine/Source/Runtime/Renderer/Private/VT/AllocatedVirtualTexture.cpp
--- a/Engine/Source/Runtime/Renderer/Private/VT/AllocatedVirtualTexture.cpp
+++ b/Engine/Source/Runtime/Renderer/Private/VT/AllocatedVirtualTexture.cpp
@@ -27,13 +27,18 @@ FAllocatedVirtualTexture::FAllocatedVirtualTexture(FVirtualTextureSystem* InSyst
 	for (uint32 LayerIndex = 0u; LayerIndex < Description.NumTextureLayers; ++LayerIndex)
 	{
 		FVirtualTextureProducer const* Producer = InProducers[LayerIndex];
-		if (Producer)
+		// Can't have missing entries for null producers if we're not merging duplicate layers
+		if (Producer || !Description.bShareDuplicateLayers)
 		{
-			const uint32 UniqueProducerIndex = AddUniqueProducer(InDesc.ProducerHandle[LayerIndex], Producer->GetDescription());
-
+			const uint32 UniqueProducerIndex = AddUniqueProducer(InDesc.ProducerHandle[LayerIndex], Producer);
 			const int32 ProducerLayerIndex = InDesc.ProducerLayerIndex[LayerIndex];
-			const uint32 ProducerPhysicalGroupIndex = Producer->GetPhysicalGroupIndexForTextureLayer(ProducerLayerIndex);
-			FVirtualTexturePhysicalSpace* PhysicalSpace = Producer->GetPhysicalSpaceForPhysicalGroup(ProducerPhysicalGroupIndex);
+			uint32 ProducerPhysicalGroupIndex = 0u;
+			FVirtualTexturePhysicalSpace* PhysicalSpace = nullptr;
+			if (Producer)
+			{
+				ProducerPhysicalGroupIndex = Producer->GetPhysicalGroupIndexForTextureLayer(ProducerLayerIndex);
+				PhysicalSpace = Producer->GetPhysicalSpaceForPhysicalGroup(ProducerPhysicalGroupIndex);
+			}
 			const uint32 UniquePhysicalSpaceIndex = AddUniquePhysicalSpace(PhysicalSpace, UniqueProducerIndex, ProducerPhysicalGroupIndex);
 			UniquePageTableLayers[UniquePhysicalSpaceIndex].ProducerTextureLayerMask |= 1 << ProducerLayerIndex;
 			const uint8 PageTableLayerLocalIndex = UniquePageTableLayers[UniquePhysicalSpaceIndex].TextureLayerCount++;
@@ -136,7 +141,10 @@ void FAllocatedVirtualTexture::Release(FVirtualTextureSystem* System)
 		TArray<FVirtualTexturePhysicalSpace*> UniquePhysicalSpaces;
 		for (int32 PageTableIndex = 0u; PageTableIndex < UniquePageTableLayers.Num(); ++PageTableIndex)
 		{
-			UniquePhysicalSpaces.Add(UniquePageTableLayers[PageTableIndex].PhysicalSpace);
+			if (UniquePageTableLayers[PageTableIndex].PhysicalSpace)
+			{
+				UniquePhysicalSpaces.Add(UniquePageTableLayers[PageTableIndex].PhysicalSpace);
+			}
 		}
 
 		for (FVirtualTexturePhysicalSpace* PhysicalSpace : UniquePhysicalSpaces)
@@ -163,7 +171,7 @@ void FAllocatedVirtualTexture::Release(FVirtualTextureSystem* System)
 	delete this;
 }
 
-uint32 FAllocatedVirtualTexture::AddUniqueProducer(FVirtualTextureProducerHandle const& InHandle, FVTProducerDescription const& InProducerDesc)
+uint32 FAllocatedVirtualTexture::AddUniqueProducer(FVirtualTextureProducerHandle const& InHandle, const FVirtualTextureProducer* InProducer)
 {
 	for (int32 Index = 0u; Index < UniqueProducers.Num(); ++Index)
 	{
@@ -175,20 +183,25 @@ uint32 FAllocatedVirtualTexture::AddUniqueProducer(FVirtualTextureProducerHandle
 	const uint32 Index = UniqueProducers.AddDefaulted();
 	check(Index < VIRTUALTEXTURE_SPACE_MAXLAYERS);
 	
-	// maybe these values should just be set by producers, rather than also set on AllocatedVT desc
-	check(InProducerDesc.Dimensions == Description.Dimensions);
-	check(InProducerDesc.TileSize == Description.TileSize);
-	check(InProducerDesc.TileBorderSize == Description.TileBorderSize);
+	uint32 MipBias = 0u;
+	if (InProducer)
+	{
+		const FVTProducerDescription& ProducerDesc = InProducer->GetDescription();
+		// maybe these values should just be set by producers, rather than also set on AllocatedVT desc
+		check(ProducerDesc.Dimensions == Description.Dimensions);
+		check(ProducerDesc.TileSize == Description.TileSize);
+		check(ProducerDesc.TileBorderSize == Description.TileBorderSize);
 
-	const uint32 BlockSizeInTiles = FMath::Max(BlockWidthInTiles, BlockHeightInTiles);
-	const uint32 ProducerBlockSizeInTiles = FMath::Max(InProducerDesc.BlockWidthInTiles, InProducerDesc.BlockHeightInTiles);
-	const uint32 MipBias = FMath::CeilLogTwo(BlockSizeInTiles / ProducerBlockSizeInTiles);
+		const uint32 BlockSizeInTiles = FMath::Max(BlockWidthInTiles, BlockHeightInTiles);
+		const uint32 ProducerBlockSizeInTiles = FMath::Max(ProducerDesc.BlockWidthInTiles, ProducerDesc.BlockHeightInTiles);
+		MipBias = FMath::CeilLogTwo(BlockSizeInTiles / ProducerBlockSizeInTiles);
 
-	check((BlockSizeInTiles / ProducerBlockSizeInTiles) * ProducerBlockSizeInTiles == BlockSizeInTiles);
-	check(InProducerDesc.BlockWidthInTiles << MipBias == BlockWidthInTiles);
-	check(InProducerDesc.BlockHeightInTiles << MipBias == BlockHeightInTiles);
+		check((BlockSizeInTiles / ProducerBlockSizeInTiles) * ProducerBlockSizeInTiles == BlockSizeInTiles);
+		check(ProducerDesc.BlockWidthInTiles << MipBias == BlockWidthInTiles);
+		check(ProducerDesc.BlockHeightInTiles << MipBias == BlockHeightInTiles);
 
-	MaxLevel = FMath::Max<uint32>(MaxLevel, InProducerDesc.MaxLevel + MipBias);
+		MaxLevel = FMath::Max<uint32>(MaxLevel, ProducerDesc.MaxLevel + MipBias);
+	}
 
 	UniqueProducers[Index].Handle = InHandle;
 	UniqueProducers[Index].MipBias = MipBias;
diff --git a/Engine/Source/Runtime/Renderer/Private/VT/AllocatedVirtualTexture.h b/Engine/Source/Runtime/Renderer/Private/VT/AllocatedVirtualTexture.h
--- a/Engine/Source/Runtime/Renderer/Private/VT/AllocatedVirtualTexture.h
+++ b/Engine/Source/Runtime/Renderer/Private/VT/AllocatedVirtualTexture.h
@@ -59,7 +59,7 @@ public:
 	inline uint32 GetProducerPhysicalGroupIndexForPageTableLayer(uint32 InLayerIndex) const { return UniquePageTableLayers[InLayerIndex].ProducerPhysicalGroupIndex; }
 
 private:
-	uint32 AddUniqueProducer(FVirtualTextureProducerHandle const& InHandle, FVTProducerDescription const& InProducerDesc);
+	uint32 AddUniqueProducer(FVirtualTextureProducerHandle const& InHandle, const FVirtualTextureProducer* InProducer);
 	uint32 AddUniquePhysicalSpace(FVirtualTexturePhysicalSpace* InPhysicalSpace, uint32 InUniqueProducerIndex, uint32 InProducerPhysicalSpaceIndex);
 
 	FThreadSafeCounter RefCount;
