Fixed high resolution screenshot mask.

diff --git a/Engine/Source/Runtime/Renderer/Private/PostProcess/PostProcessMaterial.cpp b/Engine/Source/Runtime/Renderer/Private/PostProcess/PostProcessMaterial.cpp
--- a/Engine/Source/Runtime/Renderer/Private/PostProcess/PostProcessMaterial.cpp
+++ b/Engine/Source/Runtime/Renderer/Private/PostProcess/PostProcessMaterial.cpp
@@ -374,7 +374,7 @@ FScreenPassTexture AddPostProcessMaterialPass(
 	// We can't simply output a flipped image, because the parts of the input image which show through the stencil mask or are blended in must also be flipped. In that case,
 	// we render normally to the intermediate target and flip the image when we copy to the output target.
 	const bool bForceIntermediateRT =
-		((DepthStencilTexture != nullptr) && !GRHISupportsBackBufferWithCustomDepthStencil && Inputs.OverrideOutput.IsValid()) ||
+		(DepthStencilTexture != nullptr && !GRHISupportsBackBufferWithCustomDepthStencil && Inputs.OverrideOutput.IsValid()) ||
 		(bIsCompositeWithInput && Inputs.bFlipYAxis);
 
 	FScreenPassRenderTarget Output = Inputs.OverrideOutput;
@@ -382,7 +382,7 @@ FScreenPassTexture AddPostProcessMaterialPass(
 	// We can re-use the scene color texture as the render target if we're not simultaneously reading from it.
 	// This is only necessary to do if we're going to be priming content from the render target since it avoids
 	// the copy. Otherwise, we just allocate a new render target.
-	if (!Output.IsValid() && !MaterialShaderMap->UsesSceneTexture(PPI_PostProcessInput0) && bPrimeOutputColor && !bForceIntermediateRT)
+	if (!Output.IsValid() && !MaterialShaderMap->UsesSceneTexture(PPI_PostProcessInput0) && bPrimeOutputColor && !bForceIntermediateRT && Inputs.bAllowSceneColorInputAsOutput)
 	{
 		Output = FScreenPassRenderTarget(SceneColor, ERenderTargetLoadAction::ELoad);
 	}
@@ -785,13 +785,26 @@ FScreenPassTexture AddHighResolutionScreenshotMaskPass(
 
 	if (PassSequence.IsEnabled(EPass::MaskMaterial))
 	{
+		PassSequence.AcceptPass(EPass::MaskMaterial);
+
 		FPostProcessMaterialInputs PassInputs;
-		PassSequence.AcceptOverrideIfLastPass(EPass::MaskMaterial, PassInputs.OverrideOutput);
 		PassInputs.SetInput(EPostProcessMaterialInput::SceneColor, Output);
 
-		Output = AddPostProcessMaterialPass(GraphBuilder, View, PassInputs, Inputs.MaskMaterial);
+		// Disallow the scene color input as output optimization since we need to not pollute the scene texture.
+		PassInputs.bAllowSceneColorInputAsOutput = false;
+
+		FScreenPassTexture MaskOutput = AddPostProcessMaterialPass(GraphBuilder, View, PassInputs, Inputs.MaskMaterial);
+		AddDumpToColorArrayPass(GraphBuilder, MaskOutput, FScreenshotRequest::GetHighresScreenshotMaskColorArray());
 
-		AddDumpToColorArrayPass(GraphBuilder, Output, FScreenshotRequest::GetHighresScreenshotMaskColorArray());
+		// The mask material pass is actually outputting to system memory. If we're the last pass in the chain
+		// and the override output is valid, we need to perform a copy of the input to the output. Since we can't
+		// sample from the override output (since it might be the backbuffer), we still need to participate in
+		// the pass sequence.
+		if (PassSequence.IsLastPass(EPass::MaskMaterial) && Inputs.OverrideOutput.IsValid())
+		{
+			AddDrawTexturePass(GraphBuilder, View, Output, Inputs.OverrideOutput);
+			Output = Inputs.OverrideOutput;
+		}
 	}
 
 	if (PassSequence.IsEnabled(EPass::CaptureRegionMaterial))
diff --git a/Engine/Source/Runtime/Renderer/Private/PostProcess/PostProcessMaterial.h b/Engine/Source/Runtime/Renderer/Private/PostProcess/PostProcessMaterial.h
--- a/Engine/Source/Runtime/Renderer/Private/PostProcess/PostProcessMaterial.h
+++ b/Engine/Source/Runtime/Renderer/Private/PostProcess/PostProcessMaterial.h
@@ -96,6 +96,13 @@ struct FPostProcessMaterialInputs
 
 	/** Performs a vertical axis flip if the RHI allows it. */
 	bool bFlipYAxis = false;
+
+	/** Allows (but doesn't guarantee) an optimization where, if possible, the scene color input is reused as
+	 *  the output. This can elide a copy in certain circumstances; for example, when the scene color input isn't
+	 *  actually used by the post process material and no special depth-stencil / blend composition is required.
+	 *  Set this to false when you need to guarantee creation of a dedicated output texture.
+	 */
+	bool bAllowSceneColorInputAsOutput = true;
 };
 
 FScreenPassTexture AddPostProcessMaterialPass(
