Fixed multi-light contribution on hair with ray tracing.

diff --git a/Engine/Source/Runtime/Renderer/Private/LightRendering.cpp b/Engine/Source/Runtime/Renderer/Private/LightRendering.cpp
--- a/Engine/Source/Runtime/Renderer/Private/LightRendering.cpp
+++ b/Engine/Source/Runtime/Renderer/Private/LightRendering.cpp
@@ -1402,7 +1402,7 @@ void FDeferredShadingSceneRenderer::RenderLights(FRHICommandListImmediate& RHICm
 										SubPixelRayTracingShadowMaskUAV,
 										SubPixelRayHitDistanceUAV);
 
-									TRefCountPtr<IPooledRenderTarget>* RefDestination = &PreprocessedShadowMaskSubPixelTextures[LightIndex - AttenuationLightStart];
+									TRefCountPtr<IPooledRenderTarget>* RefDestination = &PreprocessedShadowMaskSubPixelTextures[LightBatchIndex - AttenuationLightStart];
 									check(*RefDestination == nullptr);
 
 									GraphBuilder.QueueTextureExtraction(SubPixelRayTracingShadowMaskTexture, RefDestination);
