Added fading out of inset shadows with distance.
- Distance and fade range is the same as CSM (not including far/ray traced cascade).

diff --git a/Engine/Shaders/Private/ShadowProjectionPixelShader.usf b/Engine/Shaders/Private/ShadowProjectionPixelShader.usf
--- a/Engine/Shaders/Private/ShadowProjectionPixelShader.usf
+++ b/Engine/Shaders/Private/ShadowProjectionPixelShader.usf
@@ -51,6 +51,9 @@ float ShadowFadeFraction;
 float ShadowSharpen;
 float4 LightPositionAndInvRadius;
 
+float PerObjectShadowFadeStart;
+float InvPerObjectShadowFadeLength;
+
 #if USE_TRANSMISSION
 	#include "TransmissionThickness.ush"
 #else
@@ -159,6 +162,8 @@ void Main(
 	// Must not clamp for SSS shadowing, the subsurface gradient must continue past the far plane
 	float LightSpacePixelDepthForSSS = ShadowZ;
 
+	// fade out per-object shadow before cut-off, goes from 1 at start of fade plane to 0 at the end of the fade length.
+	float PerObjectDistanceFadeFraction = 1.0f - saturate((LightSpacePixelDepthForSSS - PerObjectShadowFadeStart) * InvPerObjectShadowFadeLength);
 
 	float Shadow = 1;
 	float SSSTransmission = 1;
@@ -305,7 +310,7 @@ void Main(
 
 	// 0 is shadowed, 1 is unshadowed
 	// RETURN_COLOR not needed unless writing to SceneColor;
-	float FadedShadow = lerp(1.0f, Square(Shadow), ShadowFadeFraction);
+	float FadedShadow = lerp(1.0f, Square(Shadow), ShadowFadeFraction * PerObjectDistanceFadeFraction);
 
 #if FORWARD_SHADING
 	float LightInfluenceMask = GetLightInfluenceMask(WorldPosition);
@@ -314,7 +319,7 @@ void Main(
 	// Write into all channels, the write mask will constrain to the correct one
 	OutColor = EncodeLightAttenuation(FadedShadow);
 #else
-	float FadedSSSShadow = lerp(1.0f, Square(SSSTransmission), ShadowFadeFraction);
+	float FadedSSSShadow = lerp(1.0f, Square(SSSTransmission), ShadowFadeFraction * PerObjectDistanceFadeFraction);
 
 	// the channel assignment is documented in ShadowRendering.cpp (look for Light Attenuation channel assignment)
 	OutColor = EncodeLightAttenuation(half4(FadedShadow, FadedSSSShadow, FadedShadow, FadedSSSShadow));
diff --git a/Engine/Source/Runtime/Engine/Private/Components/DirectionalLightComponent.cpp b/Engine/Source/Runtime/Engine/Private/Components/DirectionalLightComponent.cpp
--- a/Engine/Source/Runtime/Engine/Private/Components/DirectionalLightComponent.cpp
+++ b/Engine/Source/Runtime/Engine/Private/Components/DirectionalLightComponent.cpp
@@ -44,7 +44,13 @@ static TAutoConsoleVariable<float> CVarPerObjectCastDistanceRadiusScale(
 	TEXT("PerObjectCastDistanceRadiusScale The scale factor multiplied with the radius of the object to calculate the maximum distance a per-object directional shadow can reach. This will only take effect after a certain (large) radius. Default is 8 times the object radius."),
 	ECVF_RenderThreadSafe
 	);
-
+static TAutoConsoleVariable<float> CVarPerObjectCastDistanceMin(
+	TEXT("r.Shadow.PerObjectCastDistanceMin"),
+	(float)HALF_WORLD_MAX / 32.0f,
+	TEXT("Minimum cast distance for Per-Object shadows, i.e., CastDistDance = Max(r.Shadow.PerObjectCastDistanceRadiusScale * object-radius, r.Shadow.PerObjectCastDistanceMin).\n")
+	TEXT("  Default: HALF_WORLD_MAX / 32.0f"),
+	ECVF_RenderThreadSafe
+	);
 static TAutoConsoleVariable<int32> CVarMaxNumFarShadowCascades(
 	TEXT("r.Shadow.MaxNumFarShadowCascades"),
 	10,
@@ -398,7 +404,7 @@ public:
 		// Reduce casting distance on a directional light
 		// This is necessary to improve floating point precision in several places, especially when deriving frustum verts from InvReceiverMatrix
 		// This takes the object size into account to ensure that large objects get an extended distance
-		OutInitializer.MaxDistanceToCastInLightW = FMath::Clamp(SubjectBounds.SphereRadius * CVarPerObjectCastDistanceRadiusScale.GetValueOnRenderThread(), (float)HALF_WORLD_MAX / 32.0f, (float)WORLD_MAX);
+		OutInitializer.MaxDistanceToCastInLightW = FMath::Clamp(SubjectBounds.SphereRadius * CVarPerObjectCastDistanceRadiusScale.GetValueOnRenderThread(), CVarPerObjectCastDistanceMin.GetValueOnRenderThread(), (float)WORLD_MAX);
 
 		return true;
 	}
diff --git a/Engine/Source/Runtime/Renderer/Private/ShadowRendering.h b/Engine/Source/Runtime/Renderer/Private/ShadowRendering.h
--- a/Engine/Source/Runtime/Renderer/Private/ShadowRendering.h
+++ b/Engine/Source/Runtime/Renderer/Private/ShadowRendering.h
@@ -339,6 +339,10 @@ public:
 	FComputeFenceRHIRef RayTracedShadowsEndFence;
 	TRefCountPtr<IPooledRenderTarget> RayTracedShadowsRT;
 
+	/** Controls fading out of per-object shadows in the distance to avoid casting super-sharp shadows far away. */
+	float PerObjectShadowFadeStart;
+	float InvPerObjectShadowFadeLength;
+
 public:
 
 	// default constructor
@@ -784,6 +788,8 @@ public:
 		ShadowTileOffsetAndSizeParam.Bind(ParameterMap, TEXT("ShadowTileOffsetAndSize"));
 		LightPositionOrDirection.Bind(ParameterMap, TEXT("LightPositionOrDirection"));
 		HairCategorizationTexture.Bind(ParameterMap, TEXT("HairCategorizationTexture"));
+		PerObjectShadowFadeStart.Bind(ParameterMap, TEXT("PerObjectShadowFadeStart"));
+		InvPerObjectShadowFadeLength.Bind(ParameterMap, TEXT("InvPerObjectShadowFadeLength"));
 	}
 
 	void Set(FRHICommandList& RHICmdList, FShader* Shader, const FSceneView& View, const FProjectedShadowInfo* ShadowInfo, const FHairStrandsVisibilityData* HairVisibilityData)
@@ -879,6 +885,9 @@ public:
 		{
 			SetTextureParameter(RHICmdList, ShaderRHI, HairCategorizationTexture, HairVisibilityData->CategorizationTexture->GetRenderTargetItem().ShaderResourceTexture);
 		}
+
+		SetShaderValue(RHICmdList, ShaderRHI, PerObjectShadowFadeStart, ShadowInfo->PerObjectShadowFadeStart);
+		SetShaderValue(RHICmdList, ShaderRHI, InvPerObjectShadowFadeLength, ShadowInfo->InvPerObjectShadowFadeLength);
 	}
 
 	/** Serializer. */
@@ -896,6 +905,8 @@ public:
 		Ar << P.ShadowTileOffsetAndSizeParam;
 		Ar << P.LightPositionOrDirection;
 		Ar << P.HairCategorizationTexture;
+		Ar << P.PerObjectShadowFadeStart;
+		Ar << P.InvPerObjectShadowFadeLength;
 		return Ar;
 	}
 
@@ -913,6 +924,8 @@ private:
 	FShaderParameter ShadowTileOffsetAndSizeParam;
 	FShaderParameter LightPositionOrDirection;
 	FShaderResourceParameter HairCategorizationTexture;
+	FShaderParameter PerObjectShadowFadeStart;
+	FShaderParameter InvPerObjectShadowFadeLength;
 };
 
 /**
diff --git a/Engine/Source/Runtime/Renderer/Private/ShadowSetup.cpp b/Engine/Source/Runtime/Renderer/Private/ShadowSetup.cpp
--- a/Engine/Source/Runtime/Renderer/Private/ShadowSetup.cpp
+++ b/Engine/Source/Runtime/Renderer/Private/ShadowSetup.cpp
@@ -515,6 +515,8 @@ FProjectedShadowInfo::FProjectedShadowInfo()
 	, bPerObjectOpaqueShadow(false)
 	, bTransmission(false)
 	, bHairStrandsDeepShadow(false)
+	, PerObjectShadowFadeStart(WORLD_MAX)
+	, InvPerObjectShadowFadeLength(0.0f)
 	, LightSceneInfo(0)
 	, ParentSceneInfo(0)
 	, NumDynamicSubjectMeshElements(0)
@@ -622,6 +624,19 @@ bool FProjectedShadowInfo::SetupPerObjectProjection(
 			SubjectAndReceiverMatrix = SubjectMatrix;
 			ReceiverMatrix = PostSubjectMatrix;
 			MaxSubjectDepth = MaxSubjectAndReceiverDepth;
+
+			if (bDirectionalLight)
+			{
+				// No room to fade out if the end of receiver range is inside the subject range, it will just clip.
+				if (MaxSubjectZ < MaxReceiverZ)
+				{
+					float ShadowSubjectRange = MaxSubjectZ - MinSubjectZ;
+					float FadeLength = FMath::Min(ShadowSubjectRange, MaxReceiverZ - MaxSubjectZ);
+					//Initializer.MaxDistanceToCastInLightW / 16.0f;
+					PerObjectShadowFadeStart = (MaxReceiverZ - MinSubjectZ - FadeLength) / ShadowSubjectRange;
+					InvPerObjectShadowFadeLength = ShadowSubjectRange / FMath::Max(0.000001f, FadeLength);
+				}
+			}
 		}
 
 		InvMaxSubjectDepth = 1.0f / MaxSubjectDepth;
@@ -1918,6 +1933,7 @@ void FSceneRenderer::CreatePerObjectProjectedShadow(
 
 	const bool bRenderPreShadow = 
 		CVarAllowPreshadows.GetValueOnRenderThread() 
+		&& LightSceneInfo->Proxy->HasStaticShadowing()
 		// Preshadow only affects the subject's pixels
 		&& bSubjectIsVisible 
 		// Only objects with dynamic lighting should create a preshadow
