Another fix for another hair issue.

diff --git a/Engine/Source/Runtime/Renderer/Private/LightRendering.cpp b/Engine/Source/Runtime/Renderer/Private/LightRendering.cpp
--- a/Engine/Source/Runtime/Renderer/Private/LightRendering.cpp
+++ b/Engine/Source/Runtime/Renderer/Private/LightRendering.cpp
@@ -1394,8 +1394,8 @@ void FDeferredShadingSceneRenderer::RenderLights(FRHICommandListImmediate& RHICm
 										GraphBuilder,
 										SceneTextures,
 										View,
-										LightSceneInfo,
-										RayTracingConfig,
+										BatchLightSceneInfo,
+										BatchRayTracingConfig,
 										IScreenSpaceDenoiser::EShadowRequirements::Bailout,
 										true,
 										HairCategorizationTexture,
