Fix some missing texture creation flags.
Enable UAV creation for VT physical texture where supported.

diff --git a/Engine/Source/Runtime/Renderer/Private/VT/RuntimeVirtualTextureRender.cpp b/Engine/Source/Runtime/Renderer/Private/VT/RuntimeVirtualTextureRender.cpp
--- a/Engine/Source/Runtime/Renderer/Private/VT/RuntimeVirtualTextureRender.cpp
+++ b/Engine/Source/Runtime/Renderer/Private/VT/RuntimeVirtualTextureRender.cpp
@@ -547,7 +547,7 @@ namespace RuntimeVirtualTexture
 			case ERuntimeVirtualTextureMaterialType::BaseColor:
 				if (bRenderPass)
 				{
-					OutputAlias0 = RenderTexture0 = GraphBuilder.CreateTexture(FPooledRenderTargetDesc::Create2DDesc(TextureSize, PF_B8G8R8A8, FClearValueBinding::Black, TexCreate_SRGB, TexCreate_RenderTargetable, false), TEXT("RenderTexture0"));
+					OutputAlias0 = RenderTexture0 = GraphBuilder.CreateTexture(FPooledRenderTargetDesc::Create2DDesc(TextureSize, PF_B8G8R8A8, FClearValueBinding::Black, TexCreate_SRGB, TexCreate_RenderTargetable | TexCreate_ShaderResource, false), TEXT("RenderTexture0"));
 				}
 				if (bCompressPass)
 				{
@@ -557,9 +557,9 @@ namespace RuntimeVirtualTexture
 			case ERuntimeVirtualTextureMaterialType::BaseColor_Normal_Specular:
 				if (bRenderPass)
 				{
-					OutputAlias0 = RenderTexture0 = GraphBuilder.CreateTexture(FPooledRenderTargetDesc::Create2DDesc(TextureSize, PF_B8G8R8A8, FClearValueBinding::Black, TexCreate_SRGB, TexCreate_RenderTargetable, false), TEXT("RenderTexture0"));
-					RenderTexture1 = GraphBuilder.CreateTexture(FPooledRenderTargetDesc::Create2DDesc(TextureSize, PF_B8G8R8A8, FClearValueBinding::DefaultNormal8Bit, TexCreate_None, TexCreate_RenderTargetable, false), TEXT("RenderTexture1"));
-					RenderTexture2 = GraphBuilder.CreateTexture(FPooledRenderTargetDesc::Create2DDesc(TextureSize, PF_B8G8R8A8, FClearValueBinding::DefaultNormal8Bit, TexCreate_None, TexCreate_RenderTargetable, false), TEXT("RenderTexture2"));
+					OutputAlias0 = RenderTexture0 = GraphBuilder.CreateTexture(FPooledRenderTargetDesc::Create2DDesc(TextureSize, PF_B8G8R8A8, FClearValueBinding::Black, TexCreate_SRGB, TexCreate_RenderTargetable | TexCreate_ShaderResource, false), TEXT("RenderTexture0"));
+					RenderTexture1 = GraphBuilder.CreateTexture(FPooledRenderTargetDesc::Create2DDesc(TextureSize, PF_B8G8R8A8, FClearValueBinding::DefaultNormal8Bit, TexCreate_None, TexCreate_RenderTargetable | TexCreate_ShaderResource, false), TEXT("RenderTexture1"));
+					RenderTexture2 = GraphBuilder.CreateTexture(FPooledRenderTargetDesc::Create2DDesc(TextureSize, PF_B8G8R8A8, FClearValueBinding::DefaultNormal8Bit, TexCreate_None, TexCreate_RenderTargetable | TexCreate_ShaderResource, false), TEXT("RenderTexture2"));
 				}
 				if (bCompressPass)
 				{
@@ -575,9 +575,9 @@ namespace RuntimeVirtualTexture
 			case ERuntimeVirtualTextureMaterialType::BaseColor_Normal_Specular_YCoCg:
 				if (bRenderPass)
 				{
-					OutputAlias0 = RenderTexture0 = GraphBuilder.CreateTexture(FPooledRenderTargetDesc::Create2DDesc(TextureSize, PF_B8G8R8A8, FClearValueBinding::Black, TexCreate_SRGB, TexCreate_RenderTargetable, false), TEXT("RenderTexture0"));
-					RenderTexture1 = GraphBuilder.CreateTexture(FPooledRenderTargetDesc::Create2DDesc(TextureSize, PF_B8G8R8A8, FClearValueBinding::DefaultNormal8Bit, TexCreate_None, TexCreate_RenderTargetable, false), TEXT("RenderTexture1"));
-					RenderTexture2 = GraphBuilder.CreateTexture(FPooledRenderTargetDesc::Create2DDesc(TextureSize, PF_B8G8R8A8, FClearValueBinding::DefaultNormal8Bit, TexCreate_None, TexCreate_RenderTargetable, false), TEXT("RenderTexture2"));
+					OutputAlias0 = RenderTexture0 = GraphBuilder.CreateTexture(FPooledRenderTargetDesc::Create2DDesc(TextureSize, PF_B8G8R8A8, FClearValueBinding::Black, TexCreate_SRGB, TexCreate_RenderTargetable | TexCreate_ShaderResource, false), TEXT("RenderTexture0"));
+					RenderTexture1 = GraphBuilder.CreateTexture(FPooledRenderTargetDesc::Create2DDesc(TextureSize, PF_B8G8R8A8, FClearValueBinding::DefaultNormal8Bit, TexCreate_None, TexCreate_RenderTargetable | TexCreate_ShaderResource, false), TEXT("RenderTexture1"));
+					RenderTexture2 = GraphBuilder.CreateTexture(FPooledRenderTargetDesc::Create2DDesc(TextureSize, PF_B8G8R8A8, FClearValueBinding::DefaultNormal8Bit, TexCreate_None, TexCreate_RenderTargetable | TexCreate_ShaderResource, false), TEXT("RenderTexture2"));
 				}
 				if (bCompressPass)
 				{
@@ -595,7 +595,7 @@ namespace RuntimeVirtualTexture
 			case ERuntimeVirtualTextureMaterialType::WorldHeight:
 				if (bRenderPass)
 				{
-					OutputAlias0 = RenderTexture0 = GraphBuilder.CreateTexture(FPooledRenderTargetDesc::Create2DDesc(TextureSize, PF_G16, FClearValueBinding::Black, TexCreate_None, TexCreate_RenderTargetable, false), TEXT("RenderTexture0"));
+					OutputAlias0 = RenderTexture0 = GraphBuilder.CreateTexture(FPooledRenderTargetDesc::Create2DDesc(TextureSize, PF_G16, FClearValueBinding::Black, TexCreate_None, TexCreate_RenderTargetable | TexCreate_ShaderResource, false), TEXT("RenderTexture0"));
 				}
 				if (bCopyPass)
 				{
diff --git a/Engine/Source/Runtime/Renderer/Private/VT/VirtualTextureFeedback.cpp b/Engine/Source/Runtime/Renderer/Private/VT/VirtualTextureFeedback.cpp
--- a/Engine/Source/Runtime/Renderer/Private/VT/VirtualTextureFeedback.cpp
+++ b/Engine/Source/Runtime/Renderer/Private/VT/VirtualTextureFeedback.cpp
@@ -167,7 +167,7 @@ void FVirtualTextureFeedback::CreateResourceGPU( FRHICommandListImmediate& RHICm
 {
 	Size = InSize;
 
-	FPooledRenderTargetDesc Desc( FPooledRenderTargetDesc::Create2DDesc( Size, PF_R32_UINT, FClearValueBinding::None, TexCreate_None, TexCreate_UAV, false ) );
+	FPooledRenderTargetDesc Desc( FPooledRenderTargetDesc::Create2DDesc( Size, PF_R32_UINT, FClearValueBinding::None, TexCreate_None, TexCreate_ShaderResource | TexCreate_UAV, false ) );
 	GRenderTargetPool.FindFreeElement( RHICmdList, Desc, FeedbackTextureGPU, TEXT("VTFeedbackGPU") );
 
 	// Clear to default value
diff --git a/Engine/Source/Runtime/Renderer/Private/VT/VirtualTexturePhysicalSpace.cpp b/Engine/Source/Runtime/Renderer/Private/VT/VirtualTexturePhysicalSpace.cpp
--- a/Engine/Source/Runtime/Renderer/Private/VT/VirtualTexturePhysicalSpace.cpp
+++ b/Engine/Source/Runtime/Renderer/Private/VT/VirtualTexturePhysicalSpace.cpp
@@ -7,8 +7,6 @@
 #include "Stats/Stats.h"
 #include "VT/VirtualTexturePoolConfig.h"
 
-// Will turn this on when RHI support is reviewed and submitted
-#define VIRTUALTEXTURE_UAV_ALIASING 0
 
 FVirtualTexturePhysicalSpace::FVirtualTexturePhysicalSpace(const FVTPhysicalSpaceDescription& InDesc, uint16 InID)
 	: Description(InDesc)
@@ -34,13 +32,6 @@ FVirtualTexturePhysicalSpace::FVirtualTexturePhysicalSpace(const FVTPhysicalSpac
 	const uint32 MaxTiles = FMath::Max((uint32)(PoolSizeInBytes / TileSizeBytes), 1u);
 	TextureSizeInTiles = FMath::FloorToInt(FMath::Sqrt((float)MaxTiles));
 	
-// 	const bool bForce16BitPageTable = false;
-// 	if (bForce16BitPageTable)
-// 	{
-// 		// 16 bit page tables support max size of 64x64 (4096 tiles)
-// 		TextureSizeInTiles = FMath::Min(64u, TextureSizeInTiles);
-// 	}
-
 	if (TextureSizeInTiles * InDesc.TileSize > GetMax2DTextureDimension())
 	{
 		// A good option to support extremely large caches would be to allow additional slices in an array here for caches...
@@ -63,13 +54,9 @@ FVirtualTexturePhysicalSpace::~FVirtualTexturePhysicalSpace()
 
 EPixelFormat GetUnorderedAccessViewFormat(EPixelFormat InFormat)
 {
-#if VIRTUALTEXTURE_UAV_ALIASING
 	// Use alias formats for compressed textures on APIs where that is possible
 	// This allows us to compress runtime data directly to the physical texture
 	const bool bUAVAliasForCompressedTextures = GRHISupportsUAVFormatAliasing;
-#else
-	const bool bUAVAliasForCompressedTextures = false;
-#endif
 
 	switch (InFormat)
 	{
@@ -102,7 +89,7 @@ void FVirtualTexturePhysicalSpace::InitRHI()
 			FormatSRV,
 			FClearValueBinding::None,
 			TexCreate_None,
-			bCreateAliasedUAV ? TexCreate_UAV : TexCreate_None,
+			bCreateAliasedUAV ? TexCreate_ShaderResource | TexCreate_UAV : TexCreate_ShaderResource,
 			false);
 
 		GRenderTargetPool.FindFreeElement(RHICmdList, Desc, PooledRenderTarget[Layer], TEXT("PhysicalTexture"));
@@ -116,12 +103,10 @@ void FVirtualTexturePhysicalSpace::InitRHI()
 		SRVCreateInfo.SRGBOverride = SRGBO_ForceEnable;
 		TextureSRV_SRGB[Layer] = RHICreateShaderResourceView(TextureRHI, SRVCreateInfo);
 
-#if VIRTUALTEXTURE_UAV_ALIASING
 		if (bCreateAliasedUAV)
 		{
 			TextureUAV[Layer] = RHICreateUnorderedAccessView(TextureRHI, 0, FormatUAV);
 		}
-#endif
 	}
 }
 
