A (temporary) support for stationnary/static lighting for hair rendering.

This allows to have some lighting onto hair cluster/primitives. However static geometry won't cast any occlusion/shadow onto hair, as currently there is no setup for computing pre-shadow for hair. The main limitation is that hair are neither visible inot the main pass nor into the shadow (both are handled with a custom manner). The side effect of this is that it is difficult to render pre-shadow for it. This needs to be address later one.

diff --git a/Engine/Shaders/Private/DeferredLightPixelShaders.usf b/Engine/Shaders/Private/DeferredLightPixelShaders.usf
--- a/Engine/Shaders/Private/DeferredLightPixelShaders.usf
+++ b/Engine/Shaders/Private/DeferredLightPixelShaders.usf
@@ -109,7 +109,7 @@ uint HairTransmittanceBufferMaxCount;
 StructuredBuffer<float4> HairTransmittanceBuffer;
 Texture3D<float4> HairLUTTexture;
 SamplerState HairLUTSampler;
-uint HairTransmittanceEnable;
+uint HairShadowMaskValid;
 uint HairComponents; // R, TT, TRT, LocalScattering, GlobalScattering, MultipleScattering
 
 
@@ -174,7 +174,7 @@ void ComputeHairLighting(
 		HairScreenSpaceData.GBuffer.GBufferAO = 1;
 		HairScreenSpaceData.GBuffer.CustomData = 0;
 		HairScreenSpaceData.GBuffer.IndirectIrradiance = 0;
-		HairScreenSpaceData.GBuffer.PrecomputedShadowFactors = 0;
+		HairScreenSpaceData.GBuffer.PrecomputedShadowFactors = 1;
 
 		const FDerivedParams DerivedParams = GetDerivedParams(InputParams, SceneDepth);
 		FDeferredLightData LightData = SetupLightDataForStandardDeferred();
@@ -204,7 +204,7 @@ void ComputeHairLighting(
 
 		const float Dither = InterleavedGradientNoise(InputParams.PixelPos, View.StateFrameIndexMod8);
 		const uint3 PixelCoord = uint3(floor(InputParams.PixelPos), 0);
-		const float4 LightAttenuation = ScreenShadowMaskSubPixelTexture.Load(PixelCoord);
+		const float4 LightAttenuation = HairShadowMaskValid ? ScreenShadowMaskSubPixelTexture.Load(PixelCoord) : 1;
 
 		FRectTexture RectTexture = InitRectTexture(DeferredLightUniforms.SourceTexture);
 		float SurfaceShadow = 1.0f;
diff --git a/Engine/Source/Runtime/Renderer/Private/LightRendering.cpp b/Engine/Source/Runtime/Renderer/Private/LightRendering.cpp
--- a/Engine/Source/Runtime/Renderer/Private/LightRendering.cpp
+++ b/Engine/Source/Runtime/Renderer/Private/LightRendering.cpp
@@ -314,6 +314,7 @@ class FDeferredLightPS : public FGlobalShader
 		HairLUTTexture.Bind(Initializer.ParameterMap, TEXT("HairLUTTexture"));
 		HairLUTSampler.Bind(Initializer.ParameterMap, TEXT("HairLUTSampler"));
 		HairComponents.Bind(Initializer.ParameterMap, TEXT("HairComponents"));
+		HairShadowMaskValid.Bind(Initializer.ParameterMap, TEXT("HairShadowMaskValid"));
 
 		HairCategorizationTexture.Bind(Initializer.ParameterMap, TEXT("HairCategorizationTexture"));
 		HairVisibilityNodeOffsetAndCount.Bind(Initializer.ParameterMap, TEXT("HairVisibilityNodeOffsetAndCount"));
@@ -371,6 +372,7 @@ public:
 		Ar << HairLUTTexture;
 		Ar << HairLUTSampler;
 		Ar << HairComponents;
+		Ar << HairShadowMaskValid;
 
 		return bShaderHasOutdatedParameters;
 	}
@@ -484,7 +486,7 @@ private:
 
 		if (ScreenShadowMaskSubPixelTexture.IsBound())
 		{
-			if (RenderLightParams && RenderLightParams->ScreenShadowMaskSubPixelTexture)
+			if (RenderLightParams)
 			{
 				SetTextureParameter(
 					RHICmdList,
@@ -493,6 +495,13 @@ private:
 					LightAttenuationTextureSampler,
 					TStaticSamplerState<SF_Point, AM_Clamp, AM_Clamp, AM_Clamp>::GetRHI(),
 					(RenderLightParams && RenderLightParams->ScreenShadowMaskSubPixelTexture) ? RenderLightParams->ScreenShadowMaskSubPixelTexture->GetRenderTargetItem().ShaderResourceTexture : GWhiteTexture->TextureRHI);
+
+				uint32 InHairShadowMaskValid = RenderLightParams->ScreenShadowMaskSubPixelTexture ? 1 : 0;
+				SetShaderValue(
+					RHICmdList,
+					ShaderRHI,
+					HairShadowMaskValid,
+					InHairShadowMaskValid);
 			}
 		}
 
@@ -584,6 +593,7 @@ private:
 	FShaderResourceParameter HairLUTTexture;
 	FShaderResourceParameter HairLUTSampler;
 	FShaderParameter HairComponents;
+	FShaderParameter HairShadowMaskValid;
 	};
 
 IMPLEMENT_GLOBAL_SHADER(FDeferredLightPS, "/Engine/Private/DeferredLightPixelShaders.usf", "DeferredLightPixelMain", SF_Pixel);
@@ -1008,6 +1018,7 @@ void FDeferredShadingSceneRenderer::RenderLights(FRHICommandListImmediate& RHICm
 				SceneContext.FinishRenderingSceneColor(RHICmdList);
 			}
 
+			if (!bUseHairLighting)
 			{
 				SCOPED_DRAW_EVENT(RHICmdList, StandardDeferredLighting);
 
@@ -1026,6 +1037,30 @@ void FDeferredShadingSceneRenderer::RenderLights(FRHICommandListImmediate& RHICm
 
 				SceneContext.FinishRenderingSceneColor(RHICmdList);
 			}
+			else 
+			// Add a special version when hair rendering is enabled for getting lighting on hair. 
+			// This is a temporary solution as normally we should render a pre-shadow when a hair cluster is visible on screen
+			{
+				SCOPED_DRAW_EVENT(RHICmdList, StandardDeferredLighting);
+
+				// make sure we don't clear the depth
+
+				// Draw non-shadowed non-light function lights without changing render targets between them
+				for (int32 LightIndex = StandardDeferredStart; LightIndex < AttenuationLightStart; LightIndex++)
+				{
+					const FSortedLightSceneInfo& SortedLightInfo = SortedLights[LightIndex];
+					const FLightSceneInfo* const LightSceneInfo = SortedLightInfo.LightSceneInfo;
+
+					FHairStrandsTransmittanceMaskData TransmittanceMaskData;
+					TransmittanceMaskData = RenderHairStrandsTransmittanceMask(RHICmdList, Views, LightSceneInfo, HairDatas);
+
+					// Render the light to the scene color buffer, using a 1x1 white texture as input
+					SceneContext.BeginRenderingSceneColor(RHICmdList, ESimpleRenderTargetMode::EExistingColorAndDepth, FExclusiveDepthStencil::DepthRead_StencilWrite, true, true);
+					RenderLight(RHICmdList, LightSceneInfo, nullptr, nullptr, &TransmittanceMaskData, InHairVisibilityViews, false, false);
+					SceneContext.FinishRenderingSceneColor(RHICmdList);
+				}
+
+			}
 
 			if (GUseTranslucentLightingVolumes && GSupportsVolumeTextureRendering)
 			{
@@ -1912,7 +1947,7 @@ void FDeferredShadingSceneRenderer::RenderLight(FRHICommandList& RHICmdList, con
 	FGraphicsPipelineStateInitializer GraphicsPSOInit;
 	RHICmdList.ApplyCachedRenderTargets(GraphicsPSOInit);
 
-	const bool bHairRenderingEnabled = ScreenShadowMaskSubPixelTexture && InTransmittanceMaskData && InHairVisibilityViews && (LightSceneInfo->Proxy->CastsHairStrandsDeepShadow() || IsHairStrandsVoxelizationEnable());
+	const bool bHairRenderingEnabled = InTransmittanceMaskData && InHairVisibilityViews && (LightSceneInfo->Proxy->CastsHairStrandsDeepShadow() || IsHairStrandsVoxelizationEnable());
 	if (bHairRenderingEnabled)
 	{
 		GraphicsPSOInit.BlendState = TStaticBlendState<
