Fixed crash in velocity pass

diff --git a/Engine/Source/Runtime/Renderer/Private/VelocityRendering.cpp b/Engine/Source/Runtime/Renderer/Private/VelocityRendering.cpp
--- a/Engine/Source/Runtime/Renderer/Private/VelocityRendering.cpp
+++ b/Engine/Source/Runtime/Renderer/Private/VelocityRendering.cpp
@@ -317,10 +317,10 @@ void FDeferredShadingSceneRenderer::RenderVelocitiesInnerParallel(FRHICommandLis
 
 			FSceneTexturesUniformParameters SceneTextureParameters;
 			FSceneRenderTargets& SceneContext = FSceneRenderTargets::Get(RHICmdList);
-			SetupSceneTextureUniformParameters(SceneContext, View.FeatureLevel, VelocityPass == EVelocityPass::Opaque ? ESceneTextureSetupMode::None : ESceneTextureSetupMode::SceneDepth, SceneTextureParameters);
-			TUniformBufferRef<FSceneTexturesUniformParameters> PassUniformBuffer = TUniformBufferRef<FSceneTexturesUniformParameters>::CreateUniformBufferImmediate(SceneTextureParameters, UniformBuffer_SingleFrame);
+			SetupSceneTextureUniformParameters(SceneContext, View.FeatureLevel, VelocityPass == EVelocityPass::Opaque ? ESceneTextureSetupMode::None : ESceneTextureSetupMode::All, SceneTextureParameters);
+			Scene->UniformBuffers.VelocityPassUniformBuffer.UpdateUniformBufferImmediate(SceneTextureParameters);
 
-			FMeshPassProcessorRenderState DrawRenderState(View, PassUniformBuffer);
+			FMeshPassProcessorRenderState DrawRenderState(View, Scene->UniformBuffers.VelocityPassUniformBuffer);
 
 			FVelocityPassParallelCommandListSet ParallelCommandListSet(View,
 				this,
@@ -346,11 +346,11 @@ void FDeferredShadingSceneRenderer::RenderVelocitiesInner(FRHICommandListImmedia
 		const FViewInfo& View = Views[ViewIndex];
 		
 		FSceneTexturesUniformParameters SceneTextureParameters;
-		FSceneRenderTargets& SceneContext = FSceneRenderTargets::Get(RHICmdList);
-		SetupSceneTextureUniformParameters(SceneContext, View.FeatureLevel, ESceneTextureSetupMode::None, SceneTextureParameters);
-		TUniformBufferRef<FSceneTexturesUniformParameters> PassUniformBuffer = TUniformBufferRef<FSceneTexturesUniformParameters>::CreateUniformBufferImmediate(SceneTextureParameters, UniformBuffer_SingleFrame);	
+		FSceneRenderTargets& SceneContext = FSceneRenderTargets::Get(RHICmdList);		
+		SetupSceneTextureUniformParameters(SceneContext, View.FeatureLevel, VelocityPass == EVelocityPass::Opaque ? ESceneTextureSetupMode::None : ESceneTextureSetupMode::All, SceneTextureParameters);
+		Scene->UniformBuffers.VelocityPassUniformBuffer.UpdateUniformBufferImmediate(SceneTextureParameters);
 
-		FMeshPassProcessorRenderState DrawRenderState(View, PassUniformBuffer);
+		FMeshPassProcessorRenderState DrawRenderState(View, Scene->UniformBuffers.VelocityPassUniformBuffer);
 
 		if (View.ShouldRenderView())
 		{
