Fixes a crashes in SSGI when disabling TAA

diff --git a/Engine/Source/Runtime/Renderer/Private/ScreenSpaceRayTracing.cpp b/Engine/Source/Runtime/Renderer/Private/ScreenSpaceRayTracing.cpp
--- a/Engine/Source/Runtime/Renderer/Private/ScreenSpaceRayTracing.cpp
+++ b/Engine/Source/Runtime/Renderer/Private/ScreenSpaceRayTracing.cpp
@@ -111,7 +111,7 @@ bool ShouldRenderScreenSpaceDiffuseIndirect(const FViewInfo& View)
 		return false;
 	}
 
-	return View.PrevViewInfo.TemporalAAHistory.IsValid();
+	return View.PrevViewInfo.TemporalAAHistory.IsValid() && View.AntiAliasingMethod == AAM_TemporalAA;
 }
 
 bool IsSSRTemporalPassRequired(const FViewInfo& View)
