Try to keep LightMapPolicyType selection coherence between mobile and PC.

diff --git a/Engine/Source/Runtime/Renderer/Private/LightMapRendering.h b/Engine/Source/Runtime/Renderer/Private/LightMapRendering.h
--- a/Engine/Source/Runtime/Renderer/Private/LightMapRendering.h
+++ b/Engine/Source/Runtime/Renderer/Private/LightMapRendering.h
@@ -702,7 +702,6 @@ enum ELightMapPolicyType
 	LMP_MOBILE_MOVABLE_DIRECTIONAL_LIGHT_AND_SH_INDIRECT,
 	LMP_MOBILE_MOVABLE_DIRECTIONAL_LIGHT_CSM_AND_SH_INDIRECT,
 	LMP_MOBILE_DIRECTIONAL_LIGHT_CSM_AND_SH_INDIRECT,
-	LMP_MOBILE_MOVABLE_DIRECTIONAL_LIGHT,
 	LMP_MOBILE_MOVABLE_DIRECTIONAL_LIGHT_CSM,
 	LMP_MOBILE_MOVABLE_DIRECTIONAL_LIGHT_WITH_LIGHTMAP,
 	LMP_MOBILE_MOVABLE_DIRECTIONAL_LIGHT_CSM_WITH_LIGHTMAP,
@@ -840,8 +839,6 @@ public:
 			return FMobileDirectionalLightCSMAndSHIndirectPolicy::ShouldCompilePermutation(Parameters);
 		case LMP_MOBILE_MOVABLE_DIRECTIONAL_LIGHT_CSM_AND_SH_INDIRECT:
 			return FMobileMovableDirectionalLightCSMAndSHIndirectPolicy::ShouldCompilePermutation(Parameters);
-		case LMP_MOBILE_MOVABLE_DIRECTIONAL_LIGHT:
-			return FMobileMovableDirectionalLightLightingPolicy::ShouldCompilePermutation(Parameters);
 		case LMP_MOBILE_MOVABLE_DIRECTIONAL_LIGHT_CSM:
 			return FMobileMovableDirectionalLightCSMLightingPolicy::ShouldCompilePermutation(Parameters);
 		case LMP_MOBILE_MOVABLE_DIRECTIONAL_LIGHT_WITH_LIGHTMAP:
@@ -926,9 +923,6 @@ public:
 		case LMP_MOBILE_MOVABLE_DIRECTIONAL_LIGHT_CSM_AND_SH_INDIRECT:
 			FMobileMovableDirectionalLightCSMAndSHIndirectPolicy::ModifyCompilationEnvironment(Parameters, OutEnvironment);
 			break;
-		case LMP_MOBILE_MOVABLE_DIRECTIONAL_LIGHT:
-			FMobileMovableDirectionalLightLightingPolicy::ModifyCompilationEnvironment(Parameters, OutEnvironment);
-			break;
 		case LMP_MOBILE_MOVABLE_DIRECTIONAL_LIGHT_CSM:
 			FMobileMovableDirectionalLightCSMLightingPolicy::ModifyCompilationEnvironment(Parameters, OutEnvironment);
 			break;
diff --git a/Engine/Source/Runtime/Renderer/Private/MobileBasePass.cpp b/Engine/Source/Runtime/Renderer/Private/MobileBasePass.cpp
--- a/Engine/Source/Runtime/Renderer/Private/MobileBasePass.cpp
+++ b/Engine/Source/Runtime/Renderer/Private/MobileBasePass.cpp
@@ -83,9 +83,6 @@ void GetMobileBasePassShaders(
 	case LMP_MOBILE_DIRECTIONAL_LIGHT_CSM_AND_SH_INDIRECT:
 		GetUniformMobileBasePassShaders<LMP_MOBILE_DIRECTIONAL_LIGHT_CSM_AND_SH_INDIRECT, NumMovablePointLights>(Material, VertexFactoryType, bEnableSkyLight, VertexShader, PixelShader);
 		break;
-	case LMP_MOBILE_MOVABLE_DIRECTIONAL_LIGHT:
-		GetUniformMobileBasePassShaders<LMP_MOBILE_MOVABLE_DIRECTIONAL_LIGHT, NumMovablePointLights>(Material, VertexFactoryType, bEnableSkyLight, VertexShader, PixelShader);
-		break;
 	case LMP_MOBILE_MOVABLE_DIRECTIONAL_LIGHT_CSM:
 		GetUniformMobileBasePassShaders<LMP_MOBILE_MOVABLE_DIRECTIONAL_LIGHT_CSM, NumMovablePointLights>(Material, VertexFactoryType, bEnableSkyLight, VertexShader, PixelShader);
 		break;
@@ -257,6 +254,9 @@ ELightMapPolicyType MobileBasePass::SelectMeshLightmapPolicy(
 
 		const bool bMovableWithCSM = bUseMovableLight && MobileDirectionalLight->ShouldRenderViewIndependentWholeSceneShadows() && bPrimReceivesCSM;
 
+		const FIndirectLightingCacheAllocation* IndirectLightingCacheAllocation = PrimitiveSceneProxy->GetPrimitiveSceneInfo()->IndirectLightingCacheAllocation;
+		const bool bPrimitiveUsesILC = PrimitiveSceneProxy->GetIndirectLightingCacheQuality() != ILCQ_Off;
+
 		if (LightMapInteraction.GetType() == LMIT_Texture && ReadOnlyCVARCache.bAllowStaticLighting && ReadOnlyCVARCache.bEnableLowQualityLightmaps)
 		{
 			// Lightmap path
@@ -303,8 +303,12 @@ ELightMapPolicyType MobileBasePass::SelectMeshLightmapPolicy(
 			&& PrimitiveSceneProxy
 			// Movable objects need to get their GI from the indirect lighting cache
 			&& (PrimitiveSceneProxy->IsMovable() 
-			// Static objects with lightmap type set to ForceVolumetric also need to get GI from ILC
-			|| LightMapInteraction.GetType() == LMIT_GlobalVolume) )
+				// Static objects with lightmap type set to ForceVolumetric also need to get GI from ILC
+				|| LightMapInteraction.GetType() == LMIT_GlobalVolume)
+			// if there is no valid ILC data, skip SH calculation by using Movable Directional Light permutation.
+			&& (IndirectLightingCacheAllocation && IndirectLightingCacheAllocation->IsValid())
+			// Use the indirect lighting cache shaders if the object does not turn off ILC 
+			&& bPrimitiveUsesILC)
 		{
 			if (bUseMovableLight)
 			{
@@ -329,17 +333,10 @@ ELightMapPolicyType MobileBasePass::SelectMeshLightmapPolicy(
 				}
 			}
 		}
-		else if (bUseMovableLight)
+		else if (bMovableWithCSM || bUseStaticAndCSM)
 		{
 			// final determination of whether CSMs are rendered can be view dependent, thus we always need to clear the CSMs even if we're not going to render to them based on the condition below.
-			if (MobileDirectionalLight && bMovableWithCSM)
-			{
-				SelectedLightmapPolicy = LMP_MOBILE_MOVABLE_DIRECTIONAL_LIGHT_CSM;
-			}
-			else
-			{
-				SelectedLightmapPolicy = LMP_MOBILE_MOVABLE_DIRECTIONAL_LIGHT;
-			}
+			SelectedLightmapPolicy = LMP_MOBILE_MOVABLE_DIRECTIONAL_LIGHT_CSM;
 		}
 	}
 		
diff --git a/Engine/Source/Runtime/Renderer/Private/MobileBasePassRendering.cpp b/Engine/Source/Runtime/Renderer/Private/MobileBasePassRendering.cpp
--- a/Engine/Source/Runtime/Renderer/Private/MobileBasePassRendering.cpp
+++ b/Engine/Source/Runtime/Renderer/Private/MobileBasePassRendering.cpp
@@ -84,7 +84,6 @@ IMPLEMENT_MOBILE_SHADING_BASEPASS_LIGHTMAPPED_SHADER_TYPE(TUniformLightMapPolicy
 IMPLEMENT_MOBILE_SHADING_BASEPASS_LIGHTMAPPED_SHADER_TYPE(TUniformLightMapPolicy<LMP_MOBILE_MOVABLE_DIRECTIONAL_LIGHT_AND_SH_INDIRECT>, FMobileMovableDirectionalLightAndSHIndirectPolicy);
 IMPLEMENT_MOBILE_SHADING_BASEPASS_LIGHTMAPPED_SHADER_TYPE(TUniformLightMapPolicy<LMP_MOBILE_MOVABLE_DIRECTIONAL_LIGHT_CSM_AND_SH_INDIRECT>, FMobileMovableDirectionalLightCSMAndSHIndirectPolicy);
 IMPLEMENT_MOBILE_SHADING_BASEPASS_LIGHTMAPPED_SHADER_TYPE(TUniformLightMapPolicy<LMP_MOBILE_DIRECTIONAL_LIGHT_CSM_AND_SH_INDIRECT>, FMobileDirectionalLightCSMAndSHIndirectPolicy);
-IMPLEMENT_MOBILE_SHADING_BASEPASS_LIGHTMAPPED_SHADER_TYPE(TUniformLightMapPolicy<LMP_MOBILE_MOVABLE_DIRECTIONAL_LIGHT>, FMobileMovableDirectionalLightLightingPolicy);
 IMPLEMENT_MOBILE_SHADING_BASEPASS_LIGHTMAPPED_SHADER_TYPE(TUniformLightMapPolicy<LMP_MOBILE_MOVABLE_DIRECTIONAL_LIGHT_CSM>, FMobileMovableDirectionalLightCSMLightingPolicy);
 IMPLEMENT_MOBILE_SHADING_BASEPASS_LIGHTMAPPED_SHADER_TYPE(TUniformLightMapPolicy<LMP_MOBILE_MOVABLE_DIRECTIONAL_LIGHT_WITH_LIGHTMAP>, FMobileMovableDirectionalLightWithLightmapPolicy);
 IMPLEMENT_MOBILE_SHADING_BASEPASS_LIGHTMAPPED_SHADER_TYPE(TUniformLightMapPolicy<LMP_MOBILE_MOVABLE_DIRECTIONAL_LIGHT_CSM_WITH_LIGHTMAP>, FMobileMovableDirectionalLightCSMWithLightmapPolicy);
@@ -240,7 +239,7 @@ void SetupMobileSkyReflectionUniformParameters(FSkyLightSceneProxy* SkyLight, FM
 		SkyMaxMipIndex = FMath::Log2(CaptureTexture->GetSizeX());
 		InvBrightness = FMath::Max(FMath::Min(1.0f / SkyLight->AverageBrightness, 65504.f), -65504.f);
 	}
-		
+	
 	Parameters.Params = FVector4(InvBrightness, SkyMaxMipIndex, 0.f, 0.f);
 	Parameters.Texture = CaptureTexture->TextureRHI;
 	Parameters.TextureSampler = CaptureTexture->SamplerStateRHI;
