Fixed overwritten view uniform data that should not be (for other passes creating other morphed view uniform buffers)

diff --git a/Engine/Source/Runtime/Renderer/Private/HairStrands/HairStrandsDeepShadow.cpp b/Engine/Source/Runtime/Renderer/Private/HairStrands/HairStrandsDeepShadow.cpp
--- a/Engine/Source/Runtime/Renderer/Private/HairStrands/HairStrandsDeepShadow.cpp
+++ b/Engine/Source/Runtime/Renderer/Private/HairStrands/HairStrandsDeepShadow.cpp
@@ -253,11 +253,17 @@ FHairStrandsDeepShadowViews RenderHairStrandsDeepShadows(
 					{
 						const float bIsOrtho = LightType == ELightComponentType::LightType_Directional;
 
+						// Save some status we need to restore
+						const FVector SavedViewForward = ViewInfo.CachedViewUniformShaderParameters->ViewForward;
+						// Update our view parameters
 						ViewInfo.CachedViewUniformShaderParameters->HairRenderInfo.X = MinStrandRadiusAtDepth1.Primary;
 						ViewInfo.CachedViewUniformShaderParameters->HairRenderInfo.Y = MinStrandRadiusAtDepth1.Primary;
 						ViewInfo.CachedViewUniformShaderParameters->HairRenderInfo.Z = bIsOrtho;
 						ViewInfo.CachedViewUniformShaderParameters->ViewForward = DomData.LightDirection;
+						// Create the uniform buffer
 						ViewUniformShaderParameters = TUniformBufferRef<FViewUniformShaderParameters>::CreateUniformBufferImmediate(*ViewInfo.CachedViewUniformShaderParameters, UniformBuffer_SingleFrame);
+						// Restore the view cached parameters
+						ViewInfo.CachedViewUniformShaderParameters->ViewForward = SavedViewForward;
 					}
 
 					RenderDeepShadowFrontDepth(
diff --git a/Engine/Source/Runtime/Renderer/Private/HairStrands/HairStrandsVisibility.cpp b/Engine/Source/Runtime/Renderer/Private/HairStrands/HairStrandsVisibility.cpp
--- a/Engine/Source/Runtime/Renderer/Private/HairStrands/HairStrandsVisibility.cpp
+++ b/Engine/Source/Runtime/Renderer/Private/HairStrands/HairStrandsVisibility.cpp
@@ -761,10 +761,11 @@ static void AddHairVisibilityCommonPass(
 				FIntPoint(ViewInfo->UnconstrainedViewRect.Width(), ViewInfo->UnconstrainedViewRect.Height()), ViewInfo->FOV, HairVisibilitySampleCount, RasterizationScaleOverride);
 
 			TUniformBufferRef<FViewUniformShaderParameters> ViewUniformShaderParameters;
+			// Update our view parameters
 			ViewInfo->CachedViewUniformShaderParameters->HairRenderInfo.X = MinHairRadius.Primary;
 			ViewInfo->CachedViewUniformShaderParameters->HairRenderInfo.Y = MinHairRadius.Velocity;
 			ViewInfo->CachedViewUniformShaderParameters->HairRenderInfo.Z = ViewInfo->IsPerspectiveProjection() ? 0.0f : 1.0f;
-			ViewInfo->CachedViewUniformShaderParameters->ViewForward = ViewInfo->ViewMatrices.GetOverriddenTranslatedViewMatrix().GetColumn(2);
+			// Create and set the uniform buffer
 			ViewUniformShaderParameters = TUniformBufferRef<FViewUniformShaderParameters>::CreateUniformBufferImmediate(*ViewInfo->CachedViewUniformShaderParameters, UniformBuffer_SingleFrame);
 			DrawRenderState.SetViewUniformBuffer(ViewUniformShaderParameters);
 		}
diff --git a/Engine/Source/Runtime/Renderer/Private/HairStrands/HairStrandsVoxelization.cpp b/Engine/Source/Runtime/Renderer/Private/HairStrands/HairStrandsVoxelization.cpp
--- a/Engine/Source/Runtime/Renderer/Private/HairStrands/HairStrandsVoxelization.cpp
+++ b/Engine/Source/Runtime/Renderer/Private/HairStrands/HairStrandsVoxelization.cpp
@@ -332,11 +332,17 @@ static void RenderVoxelization(
 
 	TUniformBufferRef<FViewUniformShaderParameters> ViewUniformShaderParameters;
 	{
+		// Save some status we need to restore
+		const FVector SavedViewForward = ViewInfo->CachedViewUniformShaderParameters->ViewForward;
+		// Update our view parameters
 		ViewInfo->CachedViewUniformShaderParameters->HairRenderInfo.X = RadiusAtDepth1;
 		ViewInfo->CachedViewUniformShaderParameters->HairRenderInfo.Y = RadiusAtDepth1;
 		ViewInfo->CachedViewUniformShaderParameters->HairRenderInfo.Z = bIsOrtho;
 		ViewInfo->CachedViewUniformShaderParameters->ViewForward = LightDirection;
+		// Create the uniform buffer
 		ViewUniformShaderParameters = TUniformBufferRef<FViewUniformShaderParameters>::CreateUniformBufferImmediate(*ViewInfo->CachedViewUniformShaderParameters, UniformBuffer_SingleDraw);
+		// Restore the view cached parameters
+		ViewInfo->CachedViewUniformShaderParameters->ViewForward = SavedViewForward;
 	}
 
 	RHICmdList.BeginRenderPass(RPInfo, TEXT("HairVoxelization"));
