Try to update all PassUniformBuffers before start rendering to avoid GPU flush.

diff --git a/Engine/Source/Runtime/Renderer/Private/DistortionRendering.cpp b/Engine/Source/Runtime/Renderer/Private/DistortionRendering.cpp
--- a/Engine/Source/Runtime/Renderer/Private/DistortionRendering.cpp
+++ b/Engine/Source/Runtime/Renderer/Private/DistortionRendering.cpp
@@ -70,7 +70,7 @@ void SetupDistortionPassUniformBuffer(FRHICommandListImmediate& RHICmdList, cons
 void SetupMobileDistortionPassUniformBuffer(FRHICommandListImmediate& RHICmdList, const FViewInfo& View, FMobileDistortionPassUniformParameters& DistortionPassParameters)
 {
 	FSceneRenderTargets& SceneRenderTargets = FSceneRenderTargets::Get(RHICmdList);
-	SetupMobileSceneTextureUniformParameters(SceneRenderTargets, View.FeatureLevel, true, DistortionPassParameters.SceneTextures);
+	SetupMobileSceneTextureUniformParameters(SceneRenderTargets, View.FeatureLevel, true, View.bUsesCustomDepthStencil, DistortionPassParameters.SceneTextures);
 	SetupDistortionParams(DistortionPassParameters.DistortionParams, View);
 }
 
diff --git a/Engine/Source/Runtime/Renderer/Private/MobileBasePassRendering.cpp b/Engine/Source/Runtime/Renderer/Private/MobileBasePassRendering.cpp
--- a/Engine/Source/Runtime/Renderer/Private/MobileBasePassRendering.cpp
+++ b/Engine/Source/Runtime/Renderer/Private/MobileBasePassRendering.cpp
@@ -162,7 +162,7 @@ void SetupMobileBasePassUniformParameters(
 	SetupPlanarReflectionUniformParameters(View, ReflectionSceneProxy, BasePassParameters.PlanarReflection);
 
 	FSceneRenderTargets& SceneContext = FSceneRenderTargets::Get(RHICmdList);
-	SetupMobileSceneTextureUniformParameters(SceneContext, View.FeatureLevel, bTranslucentPass, BasePassParameters.SceneTextures);
+	SetupMobileSceneTextureUniformParameters(SceneContext, View.FeatureLevel, bTranslucentPass, View.bUsesCustomDepthStencil, BasePassParameters.SceneTextures);
 }
 
 void CreateMobileBasePassUniformBuffer(
diff --git a/Engine/Source/Runtime/Renderer/Private/MobileShadingRenderer.cpp b/Engine/Source/Runtime/Renderer/Private/MobileShadingRenderer.cpp
--- a/Engine/Source/Runtime/Renderer/Private/MobileShadingRenderer.cpp
+++ b/Engine/Source/Runtime/Renderer/Private/MobileShadingRenderer.cpp
@@ -305,7 +305,6 @@ void FMobileSceneRenderer::InitViews(FRHICommandListImmediate& RHICmdList)
 
 	// update buffers used in cached mesh path
 	// in case there are multiple views, these buffers will be updated before rendering each view
-	// OpaqueBasePassUniformBuffer, TranslucentBasePassUniformBuffer and DistortionPassUniformBuffer may depend on custom depth render target, so they should be updated after custom depth rendering 
 	if (Views.Num() > 0)
 	{
 		const FViewInfo& View = Views[0];
@@ -314,6 +313,10 @@ void FMobileSceneRenderer::InitViews(FRHICommandListImmediate& RHICmdList)
 		UpdateOpaqueBasePassUniformBuffer(RHICmdList, View);
 		UpdateTranslucentBasePassUniformBuffer(RHICmdList, View);
 		UpdateDirectionalLightUniformBuffers(RHICmdList, View);
+
+		FMobileDistortionPassUniformParameters DistortionPassParameters;
+		SetupMobileDistortionPassUniformBuffer(RHICmdList, View, DistortionPassParameters);
+		Scene->UniformBuffers.MobileDistortionPassUniformBuffer.UpdateUniformBufferImmediate(DistortionPassParameters);
 	}
 	UpdateSkyReflectionUniformBuffer();
 
@@ -476,16 +479,6 @@ void FMobileSceneRenderer::Render(FRHICommandListImmediate& RHICmdList)
 			RenderCustomDepthPass(RHICmdList);
 		}
 	}
-
-	// Update OpaqueBasePassUniformBuffer, TranslucentBasePassUniformBuffer and DistortionUniformBuffer right after CustomDepthPass and before BeginRenderPass of MobileBasePass
-	{
-		UpdateOpaqueBasePassUniformBuffer(RHICmdList, View);
-		UpdateTranslucentBasePassUniformBuffer(RHICmdList, View);
-
-		FMobileDistortionPassUniformParameters DistortionPassParameters;
-		SetupMobileDistortionPassUniformBuffer(RHICmdList, View, DistortionPassParameters);
-		Scene->UniformBuffers.MobileDistortionPassUniformBuffer.UpdateUniformBufferImmediate(DistortionPassParameters);
-	}
 		
 	//
 	FRHITexture* SceneColor = nullptr;
diff --git a/Engine/Source/Runtime/Renderer/Private/PostProcess/SceneRenderTargets.cpp b/Engine/Source/Runtime/Renderer/Private/PostProcess/SceneRenderTargets.cpp
--- a/Engine/Source/Runtime/Renderer/Private/PostProcess/SceneRenderTargets.cpp
+++ b/Engine/Source/Runtime/Renderer/Private/PostProcess/SceneRenderTargets.cpp
@@ -3195,6 +3195,7 @@ void SetupMobileSceneTextureUniformParameters(
 	FSceneRenderTargets& SceneContext,
 	ERHIFeatureLevel::Type FeatureLevel,
 	bool bSceneTexturesValid,
+	bool bCustomDepthIsValid,
 	FMobileSceneTextureUniformParameters& SceneTextureParameters)
 {
 	FRHITexture* BlackDefault2D = GSystemTextures.BlackDummy->GetRenderTargetItem().ShaderResourceTexture;
@@ -3213,7 +3214,8 @@ void SetupMobileSceneTextureUniformParameters(
 	FRHITexture* CustomDepth = DepthDefault;
 
 	// if there is no custom depth it's better to have the far distance there
-	IPooledRenderTarget* CustomDepthTarget = SceneContext.bCustomDepthIsValid ? SceneContext.CustomDepth.GetReference() : 0;
+	// we should update all pass uniform buffers at the start of frame on mobile, SceneContext.bCustomDepthIsValid is invalid at InitView, so pass the parameter from View.bUsesCustomDepthStencil
+	IPooledRenderTarget* CustomDepthTarget = bCustomDepthIsValid ? SceneContext.CustomDepth.GetReference() : 0;
 	if (CustomDepthTarget)
 	{
 		CustomDepth = CustomDepthTarget->GetRenderTargetItem().ShaderResourceTexture;
@@ -3238,7 +3240,7 @@ TUniformBufferRef<FMobileSceneTextureUniformParameters> CreateMobileSceneTexture
 {
 	FSceneRenderTargets& SceneContext = FSceneRenderTargets::Get(RHICmdList);
 	FMobileSceneTextureUniformParameters SceneTextureParameters;
-	SetupMobileSceneTextureUniformParameters(SceneContext, FeatureLevel, true, SceneTextureParameters);
+	SetupMobileSceneTextureUniformParameters(SceneContext, FeatureLevel, true, true, SceneTextureParameters);
 	return TUniformBufferRef<FMobileSceneTextureUniformParameters>::CreateUniformBufferImmediate(SceneTextureParameters, UniformBuffer_SingleDraw);
 }
 
diff --git a/Engine/Source/Runtime/Renderer/Private/SceneRendering.cpp b/Engine/Source/Runtime/Renderer/Private/SceneRendering.cpp
--- a/Engine/Source/Runtime/Renderer/Private/SceneRendering.cpp
+++ b/Engine/Source/Runtime/Renderer/Private/SceneRendering.cpp
@@ -3091,7 +3091,7 @@ void FSceneRenderer::RenderCustomDepthPass(FRHICommandListImmediate& RHICmdList)
 				if (FSceneInterface::GetShadingPath(View.FeatureLevel) == EShadingPath::Mobile)
 				{
 					FMobileSceneTextureUniformParameters SceneTextureParameters;
-					SetupMobileSceneTextureUniformParameters(SceneContext, View.FeatureLevel, false, SceneTextureParameters);
+					SetupMobileSceneTextureUniformParameters(SceneContext, View.FeatureLevel, false, false, SceneTextureParameters);
 					Scene->UniformBuffers.MobileCustomDepthPassUniformBuffer.UpdateUniformBufferImmediate(SceneTextureParameters);
 					PassUniformBuffer = Scene->UniformBuffers.MobileCustomDepthPassUniformBuffer;
 				}
diff --git a/Engine/Source/Runtime/Renderer/Private/ShadowDepthRendering.cpp b/Engine/Source/Runtime/Renderer/Private/ShadowDepthRendering.cpp
--- a/Engine/Source/Runtime/Renderer/Private/ShadowDepthRendering.cpp
+++ b/Engine/Source/Runtime/Renderer/Private/ShadowDepthRendering.cpp
@@ -98,7 +98,7 @@ void SetupShadowDepthPassUniformBuffer(
 	FMobileShadowDepthPassUniformParameters& ShadowDepthPassParameters)
 {
 	FSceneRenderTargets& SceneRenderTargets = FSceneRenderTargets::Get(RHICmdList);
-	SetupMobileSceneTextureUniformParameters(SceneRenderTargets, View.FeatureLevel, false, ShadowDepthPassParameters.SceneTextures);
+	SetupMobileSceneTextureUniformParameters(SceneRenderTargets, View.FeatureLevel, false, false, ShadowDepthPassParameters.SceneTextures);
 
 	ShadowDepthPassParameters.ProjectionMatrix = FTranslationMatrix(ShadowInfo->PreShadowTranslation - View.ViewMatrices.GetPreViewTranslation()) * ShadowInfo->SubjectAndReceiverMatrix;
 	ShadowDepthPassParameters.ViewMatrix = ShadowInfo->ShadowViewMatrix;
diff --git a/Engine/Source/Runtime/Renderer/Public/SceneRenderTargetParameters.h b/Engine/Source/Runtime/Renderer/Public/SceneRenderTargetParameters.h
--- a/Engine/Source/Runtime/Renderer/Public/SceneRenderTargetParameters.h
+++ b/Engine/Source/Runtime/Renderer/Public/SceneRenderTargetParameters.h
@@ -119,6 +119,7 @@ extern void SetupMobileSceneTextureUniformParameters(
 	FSceneRenderTargets& SceneContext,
 	ERHIFeatureLevel::Type FeatureLevel,
 	bool bSceneTexturesValid,
+	bool bCustomDepthIsValid,
 	FMobileSceneTextureUniformParameters& SceneTextureParameters);
 
 template< typename TRHICmdList >
