Change hair to mesh projection pass to iteration over a smaller amount of triangles (done in multiple passes) to avoid TDR on mesh with large number of triangles.

diff --git a/Engine/Shaders/Private/HairStrands/HairStrandsMeshProjection.usf b/Engine/Shaders/Private/HairStrands/HairStrandsMeshProjection.usf
--- a/Engine/Shaders/Private/HairStrands/HairStrandsMeshProjection.usf
+++ b/Engine/Shaders/Private/HairStrands/HairStrandsMeshProjection.usf
@@ -7,12 +7,12 @@ uint bClear;
 uint MaxRootCount;
 float DistanceThreshold;
 
-uint MeshNumPrimitives;
+uint MeshPrimitiveOffset_Iteration;
+uint MeshPrimitiveCount_Iteration;
 uint MeshSectionIndex;
 uint MeshMaxIndexCount;
 uint MeshMaxVertexCount;
 uint MeshIndexOffset;
-uint MaxSpinLockCount;
 Buffer<uint> MeshIndexBuffer;
 Buffer<float> MeshPositionBuffer;
 
@@ -140,7 +140,7 @@ void MainCS(uint3 DispatchThreadId : SV_DispatchThreadID)
 	float MinDistance = bClear ? 999999999 : OutRootTriangleDistance[RootIndex];
 	int TriangleIndex = -1;
 	float3 TriangleBarycentric = float3(0,0,0);
-	for (uint PrimitiveIt = 0; PrimitiveIt < MeshNumPrimitives; ++PrimitiveIt)
+	for (uint PrimitiveIt = 0; PrimitiveIt < MeshPrimitiveCount_Iteration; ++PrimitiveIt)
 	{
 		const FHairMeshTriangle Tri = GetTriangleData(PrimitiveIt, MeshPositionBuffer, MeshIndexBuffer, MeshIndexOffset, MeshMaxIndexCount, MeshMaxVertexCount);
 		
@@ -151,7 +151,7 @@ void MainCS(uint3 DispatchThreadId : SV_DispatchThreadID)
 		{			
 			MinDistance = abs(Distance);
 			TriangleBarycentric = C.Barycentric;
-			TriangleIndex = PrimitiveIt;
+			TriangleIndex = MeshPrimitiveOffset_Iteration + PrimitiveIt;
 		}
 	}
 
diff --git a/Engine/Source/Runtime/Renderer/Private/HairStrands/HairStrandsMeshProjection.cpp b/Engine/Source/Runtime/Renderer/Private/HairStrands/HairStrandsMeshProjection.cpp
--- a/Engine/Source/Runtime/Renderer/Private/HairStrands/HairStrandsMeshProjection.cpp
+++ b/Engine/Source/Runtime/Renderer/Private/HairStrands/HairStrandsMeshProjection.cpp
@@ -7,8 +7,9 @@
 #include "MeshPassProcessor.h"
 #include "RenderGraphUtils.h"
 
-static int32 GHairProjectionMaxSpinLockCount = 4096;
-static FAutoConsoleVariableRef CVarHairProjectionMaxSpinLockCount(TEXT("r.HairStrands.Projection.MaxSpinLockCount"), GHairProjectionMaxSpinLockCount, TEXT("Change the spin lock count for writing hair projection data"));
+static int32 GHairProjectionMaxTrianglePerProjectionIteration = 8;
+static FAutoConsoleVariableRef CVarHairProjectionMaxTrianglePerProjectionIteration(TEXT("r.HairStrands.Projection.MaxTrianglePerIteration"), GHairProjectionMaxTrianglePerProjectionIteration, TEXT("Change the number of triangles which are iterated over during one projection iteration step. In kilo triangle (e.g., 8 == 8000 triangles). Default is 8."));
+
 
 ///////////////////////////////////////////////////////////////////////////////////////////////////
 class FHairMeshProjectionCS : public FGlobalShader
@@ -19,9 +20,9 @@ class FHairMeshProjectionCS : public FGlobalShader
 	BEGIN_SHADER_PARAMETER_STRUCT(FParameters, )
 		SHADER_PARAMETER(uint32, bClear)
 		SHADER_PARAMETER(uint32, MaxRootCount)
-		SHADER_PARAMETER(uint32, MaxSpinLockCount)
 
-		SHADER_PARAMETER(uint32, MeshNumPrimitives)
+		SHADER_PARAMETER(uint32, MeshPrimitiveOffset_Iteration)
+		SHADER_PARAMETER(uint32, MeshPrimitiveCount_Iteration)
 		SHADER_PARAMETER(uint32, MeshSectionIndex)
 		SHADER_PARAMETER(uint32, MeshMaxIndexCount)
 		SHADER_PARAMETER(uint32, MeshMaxVertexCount)
@@ -65,33 +66,50 @@ static void AddHairStrandMeshProjectionPass(
 		return;
 	}
 
-	FHairMeshProjectionCS::FParameters* Parameters = GraphBuilder.AllocParameters<FHairMeshProjectionCS::FParameters>();
-	Parameters->bClear				= bClear ? 1 : 0;
-	Parameters->MaxRootCount		= RootData.RootCount;
-	Parameters->MaxSpinLockCount	= FMath::Clamp(GHairProjectionMaxSpinLockCount, 0, 100000);
-	Parameters->RootPositionBuffer	= RootData.RootPositionBuffer;
-	Parameters->RootNormalBuffer	= RootData.RootNormalBuffer;
-	Parameters->MeshNumPrimitives	= MeshSectionData.NumPrimitives;
-	Parameters->MeshSectionIndex	= MeshSectionData.SectionIndex;
-	Parameters->MeshMaxIndexCount	= MeshSectionData.TotalIndexCount;
-	Parameters->MeshMaxVertexCount	= MeshSectionData.TotalVertexCount;
-	Parameters->MeshIndexOffset		= MeshSectionData.IndexBaseIndex;
-	Parameters->MeshIndexBuffer		= MeshSectionData.IndexBuffer;
-	Parameters->MeshPositionBuffer	= MeshSectionData.PositionBuffer;
-
-	Parameters->OutRootTriangleIndex		= RootData.LODDatas[LODIndex].RootTriangleIndexBuffer->UAV;
-	Parameters->OutRootTriangleBarycentrics = RootData.LODDatas[LODIndex].RootTriangleBarycentricBuffer->UAV;
-	Parameters->OutRootTriangleDistance		= GraphBuilder.CreateUAV(FRDGBufferUAVDesc(RootDistanceBuffer, PF_R32_FLOAT));
+	// For projecting hair onto a skeletal mesh, 1 thread is spawn for each hair which iterates over all triangles.
+	// To avoid TDR, we split projection into multiple passes when the mesh is too large.
+	uint32 MeshPassNumPrimitive = 1024 * FMath::Clamp(GHairProjectionMaxTrianglePerProjectionIteration, 1, 256);
+	uint32 MeshPassCount = 1;
+	if (MeshSectionData.NumPrimitives < MeshPassNumPrimitive)
+	{
+		MeshPassNumPrimitive = MeshSectionData.NumPrimitives;
+	}
+	else
+	{
+		MeshPassCount = FMath::CeilToInt(MeshSectionData.NumPrimitives / float(MeshPassNumPrimitive));
+	}
 
-	const FIntVector DispatchGroupCount = FComputeShaderUtils::GetGroupCount(RootData.RootCount, 32);
-	check(DispatchGroupCount.X < 65536);
-	TShaderMapRef<FHairMeshProjectionCS> ComputeShader(ShaderMap);
-	FComputeShaderUtils::AddPass(
-		GraphBuilder,
-		RDG_EVENT_NAME("HairStrandsMeshProjection"),
-		*ComputeShader,
-		Parameters,
-		DispatchGroupCount);
+	FRDGBufferUAVRef DistanceUAV = GraphBuilder.CreateUAV(FRDGBufferUAVDesc(RootDistanceBuffer, PF_R32_FLOAT));
+	for (uint32 MeshPassIt=0;MeshPassIt<MeshPassCount;++MeshPassIt)
+	{
+		FHairMeshProjectionCS::FParameters* Parameters = GraphBuilder.AllocParameters<FHairMeshProjectionCS::FParameters>();
+		Parameters->bClear				= bClear && MeshPassIt == 0 ? 1 : 0;
+		Parameters->MaxRootCount		= RootData.RootCount;
+		Parameters->RootPositionBuffer	= RootData.RootPositionBuffer;
+		Parameters->RootNormalBuffer	= RootData.RootNormalBuffer;
+		Parameters->MeshSectionIndex	= MeshSectionData.SectionIndex;
+		Parameters->MeshMaxIndexCount	= MeshSectionData.TotalIndexCount;
+		Parameters->MeshMaxVertexCount	= MeshSectionData.TotalVertexCount;
+		Parameters->MeshIndexOffset		= MeshSectionData.IndexBaseIndex + (MeshPassNumPrimitive * MeshPassIt * 3);
+		Parameters->MeshIndexBuffer		= MeshSectionData.IndexBuffer;
+		Parameters->MeshPositionBuffer	= MeshSectionData.PositionBuffer;
+		Parameters->MeshPrimitiveOffset_Iteration	= MeshPassNumPrimitive * MeshPassIt;
+		Parameters->MeshPrimitiveCount_Iteration	= (MeshPassIt < MeshPassCount-1) ? MeshPassNumPrimitive : (MeshSectionData.NumPrimitives - MeshPassNumPrimitive * MeshPassIt);
+
+		Parameters->OutRootTriangleIndex		= RootData.LODDatas[LODIndex].RootTriangleIndexBuffer->UAV;
+		Parameters->OutRootTriangleBarycentrics = RootData.LODDatas[LODIndex].RootTriangleBarycentricBuffer->UAV;
+		Parameters->OutRootTriangleDistance		= DistanceUAV;
+
+		const FIntVector DispatchGroupCount = FComputeShaderUtils::GetGroupCount(RootData.RootCount, 32);
+		check(DispatchGroupCount.X < 65536);
+		TShaderMapRef<FHairMeshProjectionCS> ComputeShader(ShaderMap);
+		FComputeShaderUtils::AddPass(
+			GraphBuilder,
+			RDG_EVENT_NAME("HairStrandsMeshProjection"),
+			*ComputeShader,
+			Parameters,
+			DispatchGroupCount);
+	}
 }
 
 void ProjectHairStrandsOntoMesh(
