Added extra logging to track down crash in GetShaderBindings.

diff --git a/Engine/Source/Runtime/Renderer/Private/ShaderBaseClasses.cpp b/Engine/Source/Runtime/Renderer/Private/ShaderBaseClasses.cpp
--- a/Engine/Source/Runtime/Renderer/Private/ShaderBaseClasses.cpp
+++ b/Engine/Source/Runtime/Renderer/Private/ShaderBaseClasses.cpp
@@ -423,7 +423,10 @@ void FMaterialShader::GetShaderBindings(
 	check(Material.GetRenderingThreadShaderMap() && Material.GetRenderingThreadShaderMap()->IsValidForRendering() && Material.GetFeatureLevel() == FeatureLevel);
 
 	const FUniformExpressionCache& UniformExpressionCache = MaterialRenderProxy.UniformExpressionCache[FeatureLevel];
-	check(UniformExpressionCache.bUpToDate && UniformExpressionCache.UniformBuffer);
+
+	// Trying to get context for UE-75466.
+	checkf(UniformExpressionCache.bUpToDate, TEXT("UniformExpressionCache should be up to date, RenderProxy=%s Material=%s FeatureLevel=%d"), *MaterialRenderProxy.GetFriendlyName(), *Material.GetFriendlyName(), FeatureLevel);
+	checkf(UniformExpressionCache.UniformBuffer, TEXT("NULL UniformBuffer, RenderProxy=%s Material=%s FeatureLevel=%d"), *MaterialRenderProxy.GetFriendlyName(), *Material.GetFriendlyName(), FeatureLevel);
 
 #if !(UE_BUILD_TEST || UE_BUILD_SHIPPING || !WITH_EDITOR)
 	VerifyExpressionAndShaderMaps(&MaterialRenderProxy, Material, &UniformExpressionCache);
