Fixes GPU restart when we perform reflection captures via SceneCapture components.

CaptureSceneToScratchCubemap performs a copy that requires depth as it has a distance fade factor. So we must use a real depth buffer for that pass.

diff --git a/Engine/Source/Runtime/Renderer/Private/MobileShadingRenderer.cpp b/Engine/Source/Runtime/Renderer/Private/MobileShadingRenderer.cpp
--- a/Engine/Source/Runtime/Renderer/Private/MobileShadingRenderer.cpp
+++ b/Engine/Source/Runtime/Renderer/Private/MobileShadingRenderer.cpp
@@ -393,7 +393,10 @@ void FMobileSceneRenderer::Render(FRHICommandListImmediate& RHICmdList)
 	const bool bForceDepthResolve = CVarMobileForceDepthResolve.GetValueOnRenderThread() == 1;
 	const bool bSeparateTranslucencyActive = IsMobileSeparateTranslucencyActive(View);
 	bool bKeepDepthContent = bRenderToSceneColor &&
-		(bForceDepthResolve || bSeparateTranslucencyActive || (View.bIsSceneCapture && (ViewFamily.SceneCaptureSource == ESceneCaptureSource::SCS_SceneColorHDR || ViewFamily.SceneCaptureSource == ESceneCaptureSource::SCS_SceneColorSceneDepth)));
+		(bForceDepthResolve ||
+		 bSeparateTranslucencyActive ||
+		 View.bIsReflectionCapture ||
+		 (View.bIsSceneCapture && (ViewFamily.SceneCaptureSource == ESceneCaptureSource::SCS_SceneColorHDR || ViewFamily.SceneCaptureSource == ESceneCaptureSource::SCS_SceneColorSceneDepth)));
 
 	// Whether to submit cmdbuffer with offscreen rendering before doing post-processing
 	bool bSubmitOffscreenRendering = !bGammaSpace || bRenderToSceneColor;
