* Add debug plot for hair BSDF.
 * Add roughness remapping from hair roughness to GGX

The roughness remapping is a coarse remapping based on the longitudinal component, which is a sensible, yet imprecise, fit due to the anisotropic nature of the exitent lobes.

diff --git a/Engine/Shaders/Private/HairStrands/HairStrandsBsdfPlot.usf b/Engine/Shaders/Private/HairStrands/HairStrandsBsdfPlot.usf
new file mode 100644
--- /dev/null
+++ b/Engine/Shaders/Private/HairStrands/HairStrandsBsdfPlot.usf
@@ -0,0 +1,121 @@
+// Copyright 1998-2019 Epic Games, Inc. All Rights Reserved.
+
+#define HAIR_CUSTOM_BSDF 1
+
+#include "../Common.ush"
+#include "../SceneTexturesCommon.ush"
+#include "../ShadingModels.ush"
+#include "../DeferredShadingCommon.ush"
+#include "../HairBsdf.ush"
+#include "HairStrandsVisibilityCommon.ush"
+#include "HairStrandsEnvironmentLightingCommon.ush"
+
+int2 InputCoord;
+int2 OutputOffset;
+int2 OutputResolution;
+int2 MaxResolution;
+uint HairComponents;
+float Roughness;
+float BaseColor;
+float Exposure;
+
+RWTexture2D<float4> OutputTexture;
+
+#define COMPARE_ROUGHNESS 0
+
+#if COMPARE_ROUGHNESS
+float DefaultLit(float3 N, float3 V, float3 L, float InRoughnesss)
+{
+	FGBufferData GBuffer;
+	GBuffer.Roughness = InRoughnesss;
+	GBuffer.SpecularColor = 1;
+	GBuffer.DiffuseColor = 0;
+
+	FAreaLight AreaLight;
+	AreaLight.bIsRect = false;
+	AreaLight.FalloffColor = 1;
+	AreaLight.SphereSinAlpha = 0;
+	AreaLight.SphereSinAlphaSoft = 0;
+	AreaLight.LineCosSubtended = 1;
+	FShadowTerms Shadow;
+
+	if (dot(V,N) < 0)
+		V.z = -V.z;
+
+	const float NoL = saturate(dot(N, L));
+	const float3 Fs = DefaultLitBxDF(GBuffer, N, V, L, 1, NoL, AreaLight, Shadow).Specular;
+	return Fs.x;
+}
+#endif
+
+void MainPS(
+	in FScreenVertexOutput Input,
+	out float4 OutColor : SV_Target0)
+{
+	const int2 PixelCoord = Input.Position.xy;
+	if (PixelCoord.x < OutputOffset.x || PixelCoord.x >= OutputOffset.x+OutputResolution.x ||
+		PixelCoord.y < OutputOffset.y || PixelCoord.y >= OutputOffset.x+OutputResolution.y )
+		discard;
+
+	const float2 UV = ((PixelCoord-OutputOffset) + float2(0.5f,0.5f)) / float2(OutputResolution);
+	const float2 InputUV = saturate(float2(InputCoord - OutputOffset) / float2(OutputResolution));
+
+	float V_Phi   = (UV.x*2-1) * PI; //[-Pi..Pi]
+	float V_Theta = (UV.y*2-1) * PI * 0.5f;
+	float L_Phi   = (InputUV.x*2-1) * PI; //[-Pi..Pi]
+	float L_Theta = (InputUV.y*2-1) * PI * 0.5f;
+
+	const float3 Axis_U = float3(0,0,1); // Tangent
+	const float3 Axis_V = float3(1,0,0); // Perp. to the view direction
+	const float3 Axis_W = float3(0,1,0); // Ortho.
+
+	const float3 L_Local = float3(sin(L_Theta), cos(L_Theta)*cos(L_Phi), cos(L_Theta)*sin(L_Phi));
+	const float3 V_Local = float3(sin(V_Theta), cos(V_Theta)*cos(V_Phi), cos(V_Theta)*sin(V_Phi));
+	const float3 T 		 = float3(1,0,0);
+
+	const float3 V_World = 
+		V_Local.x * Axis_U + 
+		V_Local.y * Axis_V + 
+		V_Local.z * Axis_W;
+
+	FHairTransmittanceData TransmittanceData = (FHairTransmittanceData)0;
+	TransmittanceData.LocalScattering = 0;
+	TransmittanceData.GlobalScattering = 1;
+	TransmittanceData.ScatteringComponent = HairComponents;
+	TransmittanceData.OpaqueVisibility = 1;
+
+	FGBufferData HairGBuffer = (FGBufferData)0;
+	HairGBuffer.WorldNormal = T;
+	HairGBuffer.BaseColor = BaseColor.xxx;
+	HairGBuffer.Roughness = Roughness;
+	HairGBuffer.ShadingModelID = SHADINGMODELID_HAIR;
+	HairGBuffer.DiffuseColor = 0;
+	HairGBuffer.SpecularColor = 0;
+	HairGBuffer.Specular = 0.5f;
+	HairGBuffer.Metallic = 0;
+	HairGBuffer.Depth = 0.5f;
+	HairGBuffer.GBufferAO = 1;
+	HairGBuffer.CustomData = 0;
+	HairGBuffer.IndirectIrradiance = 1000000;
+
+	OutColor = float4(0,0,0,1);
+
+#if COMPARE_ROUGHNESS
+	const float3 Fs_Lit_R   = DefaultLit(float3(0, 0,-1), V_Local, L_Local, HairToGGXRoughness_R  (HairGBuffer.Roughness));
+	const float3 Fs_Lit_TT  = DefaultLit(float3(0, 0,-1), V_Local, L_Local, HairToGGXRoughness_TT (HairGBuffer.Roughness));
+	const float3 Fs_Lit_TRT = DefaultLit(float3(0, 0,-1), V_Local, L_Local, HairToGGXRoughness_TRT(HairGBuffer.Roughness));
+#endif
+	const float3 Fs_Hair    = HairShading(HairGBuffer, L_Local, V_Local, T, 1, TransmittanceData, 1, 0, uint2(0,0) );
+
+#if COMPARE_ROUGHNESS
+	OutColor.xyz += Fs_Lit_R   * float3(1, 0, 0);
+	OutColor.xyz += Fs_Lit_TRT * float3(0, 1, 1);
+	OutColor.xyz += Fs_Lit_TT  * float3(0, 0, 1);
+#endif
+	OutColor.xyz += Fs_Hair    * float3(0, 1, 0);
+	OutColor.xyz *= Exposure;
+
+	//const float3 ColorV = (V_World+1)* 0.5f;
+	//OutColor = float4(ColorV, 1);
+	//OutColor = float4(UV, 0, 1);
+}
diff --git a/Engine/Shaders/Private/HairStrands/HairStrandsEnvironmentLightingCommon.ush b/Engine/Shaders/Private/HairStrands/HairStrandsEnvironmentLightingCommon.ush
new file mode 100644
--- /dev/null
+++ b/Engine/Shaders/Private/HairStrands/HairStrandsEnvironmentLightingCommon.ush
@@ -0,0 +1,40 @@
+// Copyright 1998-2019 Epic Games, Inc. All Rights Reserved.
+
+// Roughness remapping functions of each term of the hair BSDF to the default lit model (GGX based)
+// This remapping is done based on the hair longitudinal component. This means it does not take into 
+// around the anistotropic streatching which happen along the azimuthal plane
+float HairToGGXRoughness_R(float R_Roughness)
+{
+	const float X = saturate(R_Roughness) - 0.51f;
+	return clamp(X, 0.1f, 1);
+}
+
+float HairToGGXRoughness_TRT(float R_Roughness)
+{
+	const float X = saturate(R_Roughness);
+	const float X2 = X * X;
+	const float X3 = X2 * X;
+	const float X4 = X2 * X2;
+
+	const float A = -0.207704f;
+	const float B = 6.95267f;
+	const float C = -10.1783f;
+	const float D = 4.24751f;
+	const float F = 3.37235f;
+	const float G = -1.96947f;
+	return (A * X + B * X2 + C * X3 + D * X4) / (G * X + F * X2);
+}
+
+float HairToGGXRoughness_TT(float TT_Roughness)
+{
+	const float X = saturate(TT_Roughness);
+	const float X2 = X * X;
+	const float X3 = X2 * X;
+
+	const float A = 0.70505f;
+	const float B = -0.117104f;
+	const float C = 0.0259047f;
+
+	return A * X + B * X2 + C * X3;
+}
+
diff --git a/Engine/Source/Runtime/Renderer/Private/HairStrands/HairStrandsDebug.cpp b/Engine/Source/Runtime/Renderer/Private/HairStrands/HairStrandsDebug.cpp
--- a/Engine/Source/Runtime/Renderer/Private/HairStrands/HairStrandsDebug.cpp
+++ b/Engine/Source/Runtime/Renderer/Private/HairStrands/HairStrandsDebug.cpp
@@ -35,6 +35,18 @@ static FAutoConsoleVariableRef CVarDeepShadowStats(TEXT("r.HairStrands.DebugMode
 static int32 GHairStrandsDebugStrandsMode = 0;
 static FAutoConsoleVariableRef CVarDebugPhysicsStrand(TEXT("r.HairStrands.StrandsMode"), GHairStrandsDebugStrandsMode, TEXT("Render debug mode for hair strands. 0:off, 1:simulation strands, 2:render strands with colored simulation strands influence, 3:hair UV, 4:hair root UV, 5: hair seed, 6: dimensions"));
 
+static int32 GHairStrandsDebugPlotBsdf = 0;
+static FAutoConsoleVariableRef CVarHairStrandsDebugBSDF(TEXT("r.HairStrands.PlotBsdf"), GHairStrandsDebugPlotBsdf, TEXT("Debug view for visualizing hair BSDF."));
+
+static float GHairStrandsDebugPlotBsdfRoughness = 0.3f;
+static FAutoConsoleVariableRef CVarHairStrandsDebugBSDFRoughness(TEXT("r.HairStrands.PlotBsdf.Roughness"), GHairStrandsDebugPlotBsdfRoughness, TEXT("Change the roughness of the debug BSDF plot."));
+
+static float GHairStrandsDebugPlotBsdfBaseColor = 1;
+static FAutoConsoleVariableRef CVarHairStrandsDebugBSDFAbsorption(TEXT("r.HairStrands.PlotBsdf.BaseColor"), GHairStrandsDebugPlotBsdfBaseColor, TEXT("Change the base color / absorption of the debug BSDF plot."));
+
+static float GHairStrandsDebugPlotBsdfExposure = 1.1f;
+static FAutoConsoleVariableRef CVarHairStrandsDebugBSDFExposure(TEXT("r.HairStrands.PlotBsdf.Exposure"), GHairStrandsDebugPlotBsdfExposure, TEXT("Change the exposure of the plot."));
+
 static int32 GHairDebugMeshProjection_SkinCacheMesh = 0;
 
 static int32 GHairDebugMeshProjection_Sim_HairRestTriangles = 0;
@@ -587,6 +599,95 @@ static void AddVoxelRaymarchingPass(
 	}
 }
 	
+///////////////////////////////////////////////////////////////////////////////////////////////////
+class FHairStrandsBSDFPlotPS : public FGlobalShader
+{
+	DECLARE_GLOBAL_SHADER(FHairStrandsBSDFPlotPS);
+	SHADER_USE_PARAMETER_STRUCT(FHairStrandsBSDFPlotPS, FGlobalShader);
+
+	BEGIN_SHADER_PARAMETER_STRUCT(FParameters, )
+		SHADER_PARAMETER(FIntPoint, InputCoord)
+		SHADER_PARAMETER(FIntPoint, OutputOffset)
+		SHADER_PARAMETER(FIntPoint, OutputResolution)
+		SHADER_PARAMETER(FIntPoint, MaxResolution)
+		SHADER_PARAMETER(uint32, HairComponents)
+		SHADER_PARAMETER(float, Roughness)
+		SHADER_PARAMETER(float, BaseColor)
+		SHADER_PARAMETER(float, Exposure)
+		RENDER_TARGET_BINDING_SLOTS()
+	END_SHADER_PARAMETER_STRUCT()
+
+public:
+	static bool ShouldCompilePermutation(const FGlobalShaderPermutationParameters& Parameters) { return IsHairStrandsSupported(Parameters.Platform); }
+};
+
+IMPLEMENT_GLOBAL_SHADER(FHairStrandsBSDFPlotPS, "/Engine/Private/HairStrands/HairStrandsBsdfPlot.usf", "MainPS", SF_Pixel);
+
+static void AddPlotBSDFPass(
+	FRDGBuilder& GraphBuilder,
+	const FViewInfo& View,
+	FRDGTextureRef& OutputTexture)
+{
+	
+	FSceneTextureParameters SceneTextures;
+	SetupSceneTextureParameters(GraphBuilder, &SceneTextures);
+
+	const FIntPoint Resolution(OutputTexture->Desc.Extent);
+	FHairStrandsBSDFPlotPS::FParameters* Parameters = GraphBuilder.AllocParameters<FHairStrandsBSDFPlotPS::FParameters>();
+	Parameters->InputCoord = View.CursorPos;
+	Parameters->OutputOffset = FIntPoint(100,100);
+	Parameters->OutputResolution = FIntPoint(256, 256);
+	Parameters->MaxResolution = OutputTexture->Desc.Extent;
+	Parameters->HairComponents = ToBitfield(GetHairComponents());
+	Parameters->Roughness = GHairStrandsDebugPlotBsdfRoughness;
+	Parameters->BaseColor = GHairStrandsDebugPlotBsdfBaseColor;
+	Parameters->Exposure = GHairStrandsDebugPlotBsdfExposure;
+	Parameters->RenderTargets[0] = FRenderTargetBinding(OutputTexture, ERenderTargetLoadAction::ELoad);
+
+	const FIntPoint OutputResolution = SceneTextures.SceneDepthBuffer->Desc.Extent;
+	TShaderMapRef<FPostProcessVS> VertexShader(View.ShaderMap);
+	TShaderMapRef<FHairStrandsBSDFPlotPS> PixelShader(View.ShaderMap);
+	const TShaderMap<FGlobalShaderType>* GlobalShaderMap = View.ShaderMap;
+	const FIntRect Viewport = View.ViewRect;
+	const FViewInfo* CapturedView = &View;
+
+	ClearUnusedGraphResources(*PixelShader, Parameters);
+
+	GraphBuilder.AddPass(
+		RDG_EVENT_NAME("HairStrandsBsdfPlot"),
+		Parameters,
+		ERDGPassFlags::Raster,
+		[Parameters, VertexShader, PixelShader, Viewport, Resolution, CapturedView](FRHICommandList& RHICmdList)
+	{
+		FGraphicsPipelineStateInitializer GraphicsPSOInit;
+		RHICmdList.ApplyCachedRenderTargets(GraphicsPSOInit);
+		GraphicsPSOInit.BlendState = TStaticBlendState<CW_RGBA, BO_Add, BF_One, BF_Zero, BO_Add, BF_One, BF_Zero>::GetRHI();
+		GraphicsPSOInit.RasterizerState = TStaticRasterizerState<>::GetRHI();
+		GraphicsPSOInit.DepthStencilState = TStaticDepthStencilState<false, CF_Always>::GetRHI();
+
+		GraphicsPSOInit.BoundShaderState.VertexDeclarationRHI = GFilterVertexDeclaration.VertexDeclarationRHI;
+		GraphicsPSOInit.BoundShaderState.VertexShaderRHI = GETSAFERHISHADER_VERTEX(*VertexShader);
+		GraphicsPSOInit.BoundShaderState.PixelShaderRHI = GETSAFERHISHADER_PIXEL(*PixelShader);
+		GraphicsPSOInit.PrimitiveType = PT_TriangleList;
+		SetGraphicsPipelineState(RHICmdList, GraphicsPSOInit);
+
+		VertexShader->SetParameters(RHICmdList, CapturedView->ViewUniformBuffer);
+		RHICmdList.SetViewport(Viewport.Min.X, Viewport.Min.Y, 0.0f, Viewport.Max.X, Viewport.Max.Y, 1.0f);
+		SetShaderParameters(RHICmdList, *PixelShader, PixelShader->GetPixelShader(), *Parameters);
+
+		DrawRectangle(
+			RHICmdList,
+			0, 0,
+			Viewport.Width(), Viewport.Height(),
+			Viewport.Min.X, Viewport.Min.Y,
+			Viewport.Width(), Viewport.Height(),
+			Viewport.Size(),
+			Resolution,
+			*VertexShader,
+			EDRF_UseTriangleOptimization);
+	});
+}
+	
 ///////////////////////////////////////////////////////////////////////////////////////////////////
 BEGIN_SHADER_PARAMETER_STRUCT(FHairProjectionMeshDebugParameters, )
 	SHADER_PARAMETER(FMatrix, LocalToWorld)
@@ -931,6 +1032,14 @@ void RenderHairStrandsDebugInfo(FRHICommandListImmediate& RHICmdList, TArray<FVi
 	const EHairStrandsDebugMode StrandsDebugMode = GetHairStrandsDebugStrandsMode();
 	const EHairDebugMode HairDebugMode = GetHairDebugMode();
 
+	if (GHairStrandsDebugPlotBsdf > 0)
+	{
+		FRDGBuilder GraphBuilder(RHICmdList);
+		FRDGTextureRef SceneColorTexture = GraphBuilder.RegisterExternalTexture(SceneTargets.GetSceneColor(), TEXT("SceneColorTexture"));
+		AddPlotBSDFPass(GraphBuilder, View, SceneColorTexture);
+		GraphBuilder.Execute();		
+	}
+
 	float ClusterY = 38;
 	if (HairDebugMode == EHairDebugMode::ClusterData)
 	{
