Added console command & option to change the number of samples for hair sky lighting integration.

diff --git a/Engine/Shaders/Private/HairStrands/HairStrandsEnvironmentLighting.usf b/Engine/Shaders/Private/HairStrands/HairStrandsEnvironmentLighting.usf
--- a/Engine/Shaders/Private/HairStrands/HairStrandsEnvironmentLighting.usf
+++ b/Engine/Shaders/Private/HairStrands/HairStrandsEnvironmentLighting.usf
@@ -47,6 +47,8 @@ Texture3D<float4>			HairScatteringLUTTexture;
 Texture3D<float4>			HairEnergyLUTTexture;
 SamplerState 				HairLUTSampler;
 
+uint  MultipleScatterSampleCount;
+
 uint  Voxel_ClusterId;
 float Voxel_TanConeAngle;
 float Voxel_DensityScale;
@@ -138,7 +140,7 @@ float3 ReflectionEnvironment(FGBufferData GBuffer, float2 ScreenPosition, float2
 		};
 
 		float3 AvgR = 0;
-		const uint SampleCount = 16;
+		const uint SampleCount = clamp(MultipleScatterSampleCount, 1, 16);
 		for (uint i = 0; i < SampleCount; ++i)
 		{
 			const float2 Rand2 = (PoissonDisk[i] + float2(1, 1)) * 0.5f;
diff --git a/Engine/Source/Runtime/Renderer/Private/HairStrands/HairStrandsEnvironment.cpp b/Engine/Source/Runtime/Renderer/Private/HairStrands/HairStrandsEnvironment.cpp
--- a/Engine/Source/Runtime/Renderer/Private/HairStrands/HairStrandsEnvironment.cpp
+++ b/Engine/Source/Runtime/Renderer/Private/HairStrands/HairStrandsEnvironment.cpp
@@ -36,6 +36,9 @@ static FAutoConsoleVariableRef CVarHairSkylightingPerSample(TEXT("r.HairStrands.
 static float GHairStrandsSkyLightingCompute = 1;
 static FAutoConsoleVariableRef GVarHairStrandsSkyLightingCompute(TEXT("r.HairStrands.SkyLightingCompute"), GHairStrandsSkyLightingCompute, TEXT("Evaluate sky lighting using a compute shader."));
 
+static int32 GHairStrandsSkyLightingSampleCount = 16;
+static FAutoConsoleVariableRef GVarHairStrandsSkyLightingSampleCount(TEXT("r.HairStrands.SkyLightingSampleCount"), GHairStrandsSkyLightingSampleCount, TEXT("Number of samples used for evaluating multiple scattering (in [1,16], default is set to 16)."));
+
 static bool GetHairStrandsSkyLightingEnable() { return GHairSkylightingEnable > 0; }
 static bool GetHairStrandsSkyAOEnable() { return GHairSkyAOEnable > 0; }
 static float GetHairStrandsSkyLightingConeAngle() { return FMath::Max(0.f, GHairSkylightingConeAngle); }
@@ -63,7 +66,7 @@ class FHairEnvironmentAO : public FGlobalShader
 		SHADER_PARAMETER(float, Voxel_DensityScale)
 		SHADER_PARAMETER(float, Voxel_DepthBiasScale)
 		SHADER_PARAMETER(float, Voxel_TanConeAngle)
-		SHADER_PARAMETER(float, AO_Power)	
+		SHADER_PARAMETER(float, AO_Power)
 		SHADER_PARAMETER(float, AO_Intensity)
 
 		SHADER_PARAMETER_STRUCT_INCLUDE(FSceneTextureParameters, SceneTextures)
@@ -238,6 +241,7 @@ class FHairEnvironmentLightingCS : public FGlobalShader
 		SHADER_PARAMETER(float, Voxel_TanConeAngle)
 
 		SHADER_PARAMETER(uint32, MaxVisibilityNodeCount)
+		SHADER_PARAMETER(uint32, MultipleScatterSampleCount)
 
 		SHADER_PARAMETER_TEXTURE(Texture2D, PreIntegratedGF)
 		SHADER_PARAMETER_SAMPLER(SamplerState, PreIntegratedGFSampler)
@@ -300,6 +304,7 @@ static FRDGBufferRef AddHairStrandsEnvironmentLightingPassCS(
 	PassParameters->Voxel_DepthBiasScale = GetHairStrandsVoxelizationDepthBiasScale();
 	PassParameters->Voxel_TanConeAngle = FMath::Tan(FMath::DegreesToRadians(GetHairStrandsSkyLightingConeAngle()));
 	PassParameters->MaxVisibilityNodeCount = VisibilityData.NodeData->Desc.NumElements;
+	PassParameters->MultipleScatterSampleCount = FMath::Clamp(uint32(GHairStrandsSkyLightingSampleCount), 1u, 16u);
 	PassParameters->PreIntegratedGF = GSystemTextures.PreintegratedGF->GetRenderTargetItem().ShaderResourceTexture;
 	PassParameters->PreIntegratedGFSampler = TStaticSamplerState<SF_Bilinear, AM_Clamp, AM_Clamp, AM_Clamp>::GetRHI();
 	PassParameters->SceneTextures = SceneTextures;
@@ -360,6 +365,7 @@ class FHairEnvironmentLightingPS : public FGlobalShader
 		SHADER_PARAMETER(float, Voxel_DensityScale)
 		SHADER_PARAMETER(float, Voxel_DepthBiasScale)
 		SHADER_PARAMETER(float, Voxel_TanConeAngle)
+		SHADER_PARAMETER(uint32, MultipleScatterSampleCount)
 
 		SHADER_PARAMETER_TEXTURE(Texture2D, PreIntegratedGF)
 		SHADER_PARAMETER_SAMPLER(SamplerState, PreIntegratedGFSampler)
@@ -432,7 +438,7 @@ static void AddHairStrandsEnvironmentLightingPass(
 		PassParameters->ReflectionsParameters = CreateUniformBufferImmediate(ReflectionUniformParameters, UniformBuffer_SingleDraw);
 	}
 	PassParameters->ForwardLightData = View.ForwardLightingResources->ForwardLightDataUniformBuffer;
-
+	PassParameters->MultipleScatterSampleCount = FMath::Clamp(uint32(GHairStrandsSkyLightingSampleCount), 1u, 16u);
 
 	// Bind hair data
 	PassParameters->HairCategorizationTexture = VisibilityData.CategorizationTexture->GetRenderTargetItem().ShaderResourceTexture;
