Fixed issue where eye adaptation was being applied to debug view modes.

diff --git a/Engine/Source/Runtime/Renderer/Private/PostProcess/PostProcessEyeAdaptation.cpp b/Engine/Source/Runtime/Renderer/Private/PostProcess/PostProcessEyeAdaptation.cpp
--- a/Engine/Source/Runtime/Renderer/Private/PostProcess/PostProcessEyeAdaptation.cpp
+++ b/Engine/Source/Runtime/Renderer/Private/PostProcess/PostProcessEyeAdaptation.cpp
@@ -137,23 +137,6 @@ float GetAutoExposureCompensation(const FViewInfo& View)
 
 	const FEngineShowFlags& EngineShowFlags = View.Family->EngineShowFlags;
 
-	// Skip the exposure bias when any of these flags are set.
-	const bool bSkipExposureBias =
-		View.Family->ExposureSettings.bFixed ||
-		View.Family->UseDebugViewPS() ||
-		!EngineShowFlags.Lighting ||
-		(EngineShowFlags.VisualizeBuffer && View.CurrentBufferVisualizationMode != NAME_None) ||
-		EngineShowFlags.RayTracingDebug ||
-		EngineShowFlags.VisualizeDistanceFieldAO ||
-		EngineShowFlags.VisualizeGlobalDistanceField ||
-		EngineShowFlags.CollisionVisibility ||
-		EngineShowFlags.CollisionPawn;
-
-	if (bSkipExposureBias)
-	{
-		return 1.0f;
-	}
-
 	// This scales the average luminance AFTER it gets clamped, affecting the exposure value directly.
 	float AutoExposureBias = Settings.AutoExposureBias;
 
@@ -191,30 +174,48 @@ FEyeAdaptationParameters GetEyeAdaptationParameters(const FViewInfo& View, ERHIF
 	// These clamp the average luminance computed from the scene color.
 	float MinAverageLuminance = 1.0f;
 	float MaxAverageLuminance = 1.0f;
+	float ExposureCompensation = 1.0f;
 
-	// Fixed exposure override in effect.
-	if (View.Family->ExposureSettings.bFixed)
-	{
-		MinAverageLuminance = MaxAverageLuminance = EV100ToLuminance(View.Family->ExposureSettings.FixedEV100);
-	}
-	// When !EngineShowFlags.EyeAdaptation (from "r.EyeAdaptationQuality 0") or the feature level doesn't support eye adaptation, only Settings.AutoExposureBias controls exposure.
-	else if (EngineShowFlags.EyeAdaptation && View.GetFeatureLevel() >= MinFeatureLevel)
+	// Force an exposure of 1 when any of these flags are set.
+	const bool bUseDebugExposure =
+		View.Family->ExposureSettings.bFixed ||
+		View.Family->UseDebugViewPS() ||
+		!EngineShowFlags.Lighting ||
+		(EngineShowFlags.VisualizeBuffer && View.CurrentBufferVisualizationMode != NAME_None) ||
+		EngineShowFlags.RayTracingDebug ||
+		EngineShowFlags.VisualizeDistanceFieldAO ||
+		EngineShowFlags.VisualizeGlobalDistanceField ||
+		EngineShowFlags.CollisionVisibility ||
+		EngineShowFlags.CollisionPawn;
+
+	if (!bUseDebugExposure)
 	{
-		if (AutoExposureMethod == EAutoExposureMethod::AEM_Manual)
-		{
-			const float FixedEV100 = FMath::Log2(FMath::Square(Settings.DepthOfFieldFstop) * Settings.CameraShutterSpeed * 100 / FMath::Max(1.f, Settings.CameraISO));
-			MinAverageLuminance = MaxAverageLuminance = EV100ToLuminance(FixedEV100);
-		}
-		else if (bExtendedLuminanceRange)
+		// Fixed exposure override in effect.
+		if (View.Family->ExposureSettings.bFixed)
 		{
-			MinAverageLuminance = EV100ToLuminance(Settings.AutoExposureMinBrightness);
-			MaxAverageLuminance = EV100ToLuminance(Settings.AutoExposureMaxBrightness);
+			MinAverageLuminance = MaxAverageLuminance = EV100ToLuminance(View.Family->ExposureSettings.FixedEV100);
 		}
-		else
+		// When !EngineShowFlags.EyeAdaptation (from "r.EyeAdaptationQuality 0") or the feature level doesn't support eye adaptation, only Settings.AutoExposureBias controls exposure.
+		else if (EngineShowFlags.EyeAdaptation && View.GetFeatureLevel() >= MinFeatureLevel)
 		{
-			MinAverageLuminance = Settings.AutoExposureMinBrightness;
-			MaxAverageLuminance = Settings.AutoExposureMaxBrightness;
+			if (AutoExposureMethod == EAutoExposureMethod::AEM_Manual)
+			{
+				const float FixedEV100 = FMath::Log2(FMath::Square(Settings.DepthOfFieldFstop) * Settings.CameraShutterSpeed * 100 / FMath::Max(1.f, Settings.CameraISO));
+				MinAverageLuminance = MaxAverageLuminance = EV100ToLuminance(FixedEV100);
+			}
+			else if (bExtendedLuminanceRange)
+			{
+				MinAverageLuminance = EV100ToLuminance(Settings.AutoExposureMinBrightness);
+				MaxAverageLuminance = EV100ToLuminance(Settings.AutoExposureMaxBrightness);
+			}
+			else
+			{
+				MinAverageLuminance = Settings.AutoExposureMinBrightness;
+				MaxAverageLuminance = Settings.AutoExposureMaxBrightness;
+			}
 		}
+
+		ExposureCompensation = GetAutoExposureCompensation(View);
 	}
 
 	MinAverageLuminance = FMath::Min(MinAverageLuminance, MaxAverageLuminance);
@@ -222,8 +223,6 @@ FEyeAdaptationParameters GetEyeAdaptationParameters(const FViewInfo& View, ERHIF
 	// This scales the average luminance BEFORE it gets clamped. Note that AEM_Histogram implements the calibration constant through ExposureLowPercent and ExposureHighPercent.
 	const float CalibrationConstant = FMath::Clamp(Settings.AutoExposureCalibrationConstant, 1.0f, 100.0f) * PercentToScale;
 
-	const float ExposureCompensation = GetAutoExposureCompensation(View);
-
 	const float WeightSlope = (AutoExposureMethod == EAutoExposureMethod::AEM_Basic) ? GetBasicAutoExposureFocus() : 0.0f;
 
 	const float HistogramLogDelta = HistogramLogMax - HistogramLogMin;
