* Add hair components enable/disable option on sky lighting to isolate components (for debug purpose)
 * Add MS jittering over the spherical domain. This reduce the TRT artifact, but increase tremendously the amount of noise. Disable by default.
 * Rename the hair sky command, to be more consistent/searchable.

diff --git a/Engine/Shaders/Private/HairStrands/HairStrandsEnvironmentLighting.usf b/Engine/Shaders/Private/HairStrands/HairStrandsEnvironmentLighting.usf
--- a/Engine/Shaders/Private/HairStrands/HairStrandsEnvironmentLighting.usf
+++ b/Engine/Shaders/Private/HairStrands/HairStrandsEnvironmentLighting.usf
@@ -48,6 +48,8 @@ Texture3D<float4>			HairEnergyLUTTexture;
 SamplerState 				HairLUTSampler;
 
 uint  MultipleScatterSampleCount;
+uint  JitterSphericalIntegration;
+uint  HairComponents;
 
 uint  Voxel_ClusterId;
 float Voxel_TanConeAngle;
@@ -114,9 +116,10 @@ float3 ReflectionEnvironment(FGBufferData GBuffer, float2 ScreenPosition, float2
 	////////////////////////////////////////////////////////////////////////////////////////////////////////////////////////////////////////////
 	////////////////////////////////////////////////////////////////////////////////////////////////////////////////////////////////////////////
 	// Trace/sample transmittance in various direction to have estimate of where is occlusion & computing the integrated dual scattering value
-	float3 UnoccludedN = 0;
+	float3 UnoccludedN = -V;
 	float3 AvgDualScattering = 0;
 	#if 1
+	if ((HairComponents & HAIR_COMPONENT_R) || (HairComponents & HAIR_COMPONENT_TRT) || (HairComponents & HAIR_COMPONENT_TT))
 	{
 		// Poisson disk position http://developer.download.nvidia.com/whitepapers/2008/PCSS_Integration.pdf
 		float2 PoissonDisk[16] =
@@ -139,12 +142,20 @@ float3 ReflectionEnvironment(FGBufferData GBuffer, float2 ScreenPosition, float2
 			float2(0.14383161, -0.14100790)
 		};
 
+		float2 RandOffset = 0;
+		if (JitterSphericalIntegration > 0)
+		{
+			RandOffset += float2(
+				InterleavedGradientNoise(SvPosition, View.StateFrameIndexMod8 * 11),
+				InterleavedGradientNoise(SvPosition, View.StateFrameIndexMod8 * 7));
+		}
+
 		float3 AvgR = 0;
 		const uint SampleCount = clamp(MultipleScatterSampleCount, 1, 16);
 		for (uint i = 0; i < SampleCount; ++i)
 		{
-			const float2 Rand2 = (PoissonDisk[i] + float2(1, 1)) * 0.5f;
-			const float3 SampleL = UniformSampleSphere(saturate(Rand2)).xyz;
+			const float2 Rand2 = (PoissonDisk[i] + float2(1, 1)) * 0.5f + RandOffset;
+			const float3 SampleL = UniformSampleSphere(frac(Rand2)).xyz;
 
 			// Depth bias
 			// Origin is shifted voxels away towards the light + a constant bias of the size of the voxel
@@ -158,7 +169,7 @@ float3 ReflectionEnvironment(FGBufferData GBuffer, float2 ScreenPosition, float2
 			const float ScalarTransmittance = max(SampleTransmittance.x, max(SampleTransmittance.y, SampleTransmittance.z));
 			AvgR += SampleL * SampleTransmittance.xyz;
 
-			const uint ScatteringComponent = HAIR_COMPONENT_R | HAIR_COMPONENT_TRT | HAIR_COMPONENT_LS | HAIR_COMPONENT_GS;// | HAIR_MULTISCATTER;
+			const uint ScatteringComponent = (HairComponents & HAIR_COMPONENT_R) | (HairComponents & HAIR_COMPONENT_TRT) | HAIR_COMPONENT_LS | HAIR_COMPONENT_GS;// | HAIR_MULTISCATTER;
 			const FHairTransmittanceData TransmittanceData = GetHairTransmittance(V, SampleL, GBuffer, SampleTransmittance, HairScatteringLUTTexture, HairLUTSampler, ScatteringComponent);
 
 			const float Area = 0.2;
@@ -191,6 +202,7 @@ float3 ReflectionEnvironment(FGBufferData GBuffer, float2 ScreenPosition, float2
 	// MS ==============================================================================
 	float3 ColorMS = 0;
 	#if 1
+	if ((HairComponents & HAIR_COMPONENT_R) || (HairComponents & HAIR_COMPONENT_TRT) || (HairComponents & HAIR_COMPONENT_TT))
 	{
 		float3 N = UnoccludedN;
 		float3 R = 2 * dot(V, N) * N - V;
@@ -230,6 +242,7 @@ float3 ReflectionEnvironment(FGBufferData GBuffer, float2 ScreenPosition, float2
 	// R ==============================================================================
 	float3 ColorR = 0;
 	#if 1
+	if ((HairComponents & HAIR_COMPONENT_R) || (HairComponents & HAIR_COMPONENT_TRT))
 	{
 		float3 R = 0;
 		{
@@ -287,7 +300,7 @@ float3 ReflectionEnvironment(FGBufferData GBuffer, float2 ScreenPosition, float2
 		#if 1
 		{			
 			TransmittanceData.GlobalScattering = Transmittance.xyz;
-			TransmittanceData.ScatteringComponent = HAIR_COMPONENT_R | HAIR_COMPONENT_TRT | HAIR_COMPONENT_MULTISCATTER;
+			TransmittanceData.ScatteringComponent = (HairComponents & HAIR_COMPONENT_R) | (HairComponents & HAIR_COMPONENT_TRT) | HAIR_COMPONENT_MULTISCATTER;
 
 			const float Area = 0.2;
 			const float BackLit = 0;
@@ -304,6 +317,7 @@ float3 ReflectionEnvironment(FGBufferData GBuffer, float2 ScreenPosition, float2
 	// TT ==============================================================================
 	float3 ColorTT = 0;
 	#if 1
+	if (HairComponents & HAIR_COMPONENT_TT)
 	{
 		const float3 L = -V;
 		const float SinLightAngle = dot(L, T);
diff --git a/Engine/Source/Runtime/Renderer/Private/HairStrands/HairStrandsEnvironment.cpp b/Engine/Source/Runtime/Renderer/Private/HairStrands/HairStrandsEnvironment.cpp
--- a/Engine/Source/Runtime/Renderer/Private/HairStrands/HairStrandsEnvironment.cpp
+++ b/Engine/Source/Runtime/Renderer/Private/HairStrands/HairStrandsEnvironment.cpp
@@ -22,22 +22,25 @@
 ///////////////////////////////////////////////////////////////////////////////////////////////////
 
 static int32 GHairSkylightingEnable = 1;
-static FAutoConsoleVariableRef CVarHairSkylightingEnable(TEXT("r.HairStrands.SkyLightingEnable"), GHairSkylightingEnable, TEXT("Enable sky lighting on hair."));
+static FAutoConsoleVariableRef CVarHairSkylightingEnable(TEXT("r.HairStrands.SkyLighting"), GHairSkylightingEnable, TEXT("Enable sky lighting on hair."));
 
 static int32 GHairSkyAOEnable = 1;
 static FAutoConsoleVariableRef CVarHairSkyAOEnable(TEXT("r.HairStrands.SkyAOEnable"), GHairSkyAOEnable, TEXT("Enable (sky) AO on hair."));
 
 static float GHairSkylightingConeAngle = 3;
-static FAutoConsoleVariableRef CVarHairSkylightingConeAngle(TEXT("r.HairStrands.SkyLightingConeAngle"), GHairSkylightingConeAngle, TEXT("Cone angle for tracing sky lighting on hair."));
+static FAutoConsoleVariableRef CVarHairSkylightingConeAngle(TEXT("r.HairStrands.SkyLighting.ConeAngle"), GHairSkylightingConeAngle, TEXT("Cone angle for tracing sky lighting on hair."));
 
 static float GHairSkylightingPerSample = 1;
-static FAutoConsoleVariableRef CVarHairSkylightingPerSample(TEXT("r.HairStrands.SkyLightingPerSample"), GHairSkylightingPerSample, TEXT("Evaluate sky lighting per hair sample."));
+static FAutoConsoleVariableRef CVarHairSkylightingPerSample(TEXT("r.HairStrands.SkyLighting.PerSample"), GHairSkylightingPerSample, TEXT("Evaluate sky lighting per hair sample."));
 
 static float GHairStrandsSkyLightingCompute = 1;
-static FAutoConsoleVariableRef GVarHairStrandsSkyLightingCompute(TEXT("r.HairStrands.SkyLightingCompute"), GHairStrandsSkyLightingCompute, TEXT("Evaluate sky lighting using a compute shader."));
+static FAutoConsoleVariableRef GVarHairStrandsSkyLightingCompute(TEXT("r.HairStrands.SkyLighting.Compute"), GHairStrandsSkyLightingCompute, TEXT("Evaluate sky lighting using a compute shader."));
 
 static int32 GHairStrandsSkyLightingSampleCount = 16;
-static FAutoConsoleVariableRef GVarHairStrandsSkyLightingSampleCount(TEXT("r.HairStrands.SkyLightingSampleCount"), GHairStrandsSkyLightingSampleCount, TEXT("Number of samples used for evaluating multiple scattering (in [1,16], default is set to 16)."));
+static FAutoConsoleVariableRef GVarHairStrandsSkyLightingSampleCount(TEXT("r.HairStrands.SkyLighting.SampleCount"), GHairStrandsSkyLightingSampleCount, TEXT("Number of samples used for evaluating multiple scattering (in [1,16], default is set to 16)."));
+
+static int32 GHairStrandsSkyLightingJitterSphericalIntegration = 0;
+static FAutoConsoleVariableRef GVarHairStrandsSkyLightingJitterSphericalIntegration(TEXT("r.HairStrands.SkyLighting.JitterIntegration"), GHairStrandsSkyLightingJitterSphericalIntegration, TEXT("Jitter the sphereical integration for the multiple scattering term. The result is more correct, but noiser as well."));
 
 static bool GetHairStrandsSkyLightingEnable() { return GHairSkylightingEnable > 0; }
 static bool GetHairStrandsSkyAOEnable() { return GHairSkyAOEnable > 0; }
@@ -242,6 +245,8 @@ class FHairEnvironmentLightingCS : public FGlobalShader
 
 		SHADER_PARAMETER(uint32, MaxVisibilityNodeCount)
 		SHADER_PARAMETER(uint32, MultipleScatterSampleCount)
+		SHADER_PARAMETER(uint32, HairComponents)
+		SHADER_PARAMETER(uint32, JitterSphericalIntegration)
 
 		SHADER_PARAMETER_TEXTURE(Texture2D, PreIntegratedGF)
 		SHADER_PARAMETER_SAMPLER(SamplerState, PreIntegratedGFSampler)
@@ -305,6 +310,8 @@ static FRDGBufferRef AddHairStrandsEnvironmentLightingPassCS(
 	PassParameters->Voxel_TanConeAngle = FMath::Tan(FMath::DegreesToRadians(GetHairStrandsSkyLightingConeAngle()));
 	PassParameters->MaxVisibilityNodeCount = VisibilityData.NodeData->Desc.NumElements;
 	PassParameters->MultipleScatterSampleCount = FMath::Clamp(uint32(GHairStrandsSkyLightingSampleCount), 1u, 16u);
+	PassParameters->JitterSphericalIntegration = GHairStrandsSkyLightingJitterSphericalIntegration ? 1 : 0;
+	PassParameters->HairComponents = ToBitfield(GetHairComponents());
 	PassParameters->PreIntegratedGF = GSystemTextures.PreintegratedGF->GetRenderTargetItem().ShaderResourceTexture;
 	PassParameters->PreIntegratedGFSampler = TStaticSamplerState<SF_Bilinear, AM_Clamp, AM_Clamp, AM_Clamp>::GetRHI();
 	PassParameters->SceneTextures = SceneTextures;
@@ -366,6 +373,8 @@ class FHairEnvironmentLightingPS : public FGlobalShader
 		SHADER_PARAMETER(float, Voxel_DepthBiasScale)
 		SHADER_PARAMETER(float, Voxel_TanConeAngle)
 		SHADER_PARAMETER(uint32, MultipleScatterSampleCount)
+		SHADER_PARAMETER(uint32, HairComponents)
+		SHADER_PARAMETER(uint32, JitterSphericalIntegration)
 
 		SHADER_PARAMETER_TEXTURE(Texture2D, PreIntegratedGF)
 		SHADER_PARAMETER_SAMPLER(SamplerState, PreIntegratedGFSampler)
@@ -439,6 +448,8 @@ static void AddHairStrandsEnvironmentLightingPass(
 	}
 	PassParameters->ForwardLightData = View.ForwardLightingResources->ForwardLightDataUniformBuffer;
 	PassParameters->MultipleScatterSampleCount = FMath::Clamp(uint32(GHairStrandsSkyLightingSampleCount), 1u, 16u);
+	PassParameters->HairComponents = ToBitfield(GetHairComponents());
+	PassParameters->JitterSphericalIntegration = GHairStrandsSkyLightingJitterSphericalIntegration ? 1 : 0;
 
 	// Bind hair data
 	PassParameters->HairCategorizationTexture = VisibilityData.CategorizationTexture->GetRenderTargetItem().ShaderResourceTexture;
