Fixed RDG validation assert.

diff --git a/Engine/Source/Runtime/Renderer/Private/MobileShadingRenderer.cpp b/Engine/Source/Runtime/Renderer/Private/MobileShadingRenderer.cpp
--- a/Engine/Source/Runtime/Renderer/Private/MobileShadingRenderer.cpp
+++ b/Engine/Source/Runtime/Renderer/Private/MobileShadingRenderer.cpp
@@ -781,6 +781,7 @@ void FMobileSceneRenderer::BasicPostProcess(FRHICommandListImmediate& RHICmdList
 	// todo: this should come from View.Family->RenderTarget
 	Desc.Format = PF_B8G8R8A8;
 	Desc.NumMips = 1;
+	Desc.TargetableFlags |= TexCreate_RenderTargetable;
 
 	GRenderTargetPool.CreateUntrackedElement(Desc, Temp, Item);
 
