RL - Add missing GPU Scopes for Hair and Virtual Texture

diff --git a/Engine/Source/Runtime/Renderer/Private/DeferredShadingRenderer.cpp b/Engine/Source/Runtime/Renderer/Private/DeferredShadingRenderer.cpp
--- a/Engine/Source/Runtime/Renderer/Private/DeferredShadingRenderer.cpp
+++ b/Engine/Source/Runtime/Renderer/Private/DeferredShadingRenderer.cpp
@@ -213,6 +213,8 @@ DECLARE_GPU_STAT(PostRenderOpsFX);
 DECLARE_GPU_STAT(HZB);
 DECLARE_GPU_STAT_NAMED(Unaccounted, TEXT("[unaccounted]"));
 DECLARE_GPU_STAT(WaterRendering);
+DECLARE_GPU_STAT(HairRendering);
+DECLARE_GPU_STAT(VirtualTextureFeedback);
 
 const TCHAR* GetDepthPassReason(bool bDitheredLODTransitionsUseStencil, EShaderPlatform ShaderPlatform)
 {
@@ -1897,6 +1899,7 @@ void FDeferredShadingSceneRenderer::Render(FRHICommandListImmediate& RHICmdList)
 	const bool bIsViewCompatible = Views.Num() > 0 && Views[0].Family->ViewMode == VMI_Lit; 
 	if (IsHairStrandsEnable(Scene->GetShaderPlatform()) && bIsViewCompatible)
 	{
+		SCOPED_GPU_STAT(RHICmdList, HairRendering);
 		HairDatasStorage.HairClusterPerViews = CreateHairStrandsClusters(RHICmdList, Scene, Views);
 		VoxelizeHairStrands(RHICmdList, Scene, Views, HairDatasStorage.HairClusterPerViews);
 		HairDatasStorage.DeepShadowViews = RenderHairStrandsDeepShadows(RHICmdList, Scene, Views, HairDatasStorage.HairClusterPerViews);
@@ -2256,14 +2259,17 @@ void FDeferredShadingSceneRenderer::Render(FRHICommandListImmediate& RHICmdList)
 		checkSlow(RHICmdList.IsOutsideRenderPass());
 	}
 
-	if (HairDatas)
 	{
-		RenderHairComposeSubPixel(RHICmdList, Views, HairDatas);
-	}
+		SCOPED_GPU_STAT(RHICmdList, HairRendering);
+		if (HairDatas)
+		{
+			RenderHairComposeSubPixel(RHICmdList, Views, HairDatas);
+		}
 
-	if (IsHairStrandsEnable(Scene->GetShaderPlatform()))
-	{
-		RenderHairStrandsDebugInfo(RHICmdList, Views, HairDatas);
+		if (IsHairStrandsEnable(Scene->GetShaderPlatform()))
+		{
+			RenderHairStrandsDebugInfo(RHICmdList, Views, HairDatas);
+		}
 	}
 
 	checkSlow(RHICmdList.IsOutsideRenderPass());
@@ -2278,6 +2284,7 @@ void FDeferredShadingSceneRenderer::Render(FRHICommandListImmediate& RHICmdList)
 
 	if (bUseVirtualTexturing)
 	{
+		SCOPED_GPU_STAT(RHICmdList, VirtualTextureFeedback);
 		// No pass after this can make VT page requests
 		SceneContext.VirtualTextureFeedback.TransferGPUToCPU(RHICmdList, Views[0].ViewRect);
 	}
