Remove unwanted CSV marker.
Remove hacky check for GPUScene buffer when preparing runtime virtual texture rendering.
Add GPUScene update in runtime virtual texture rendering (which should only be needed by RVT streaming mip build path).

diff --git a/Engine/Source/Runtime/Renderer/Private/VT/RuntimeVirtualTextureRender.cpp b/Engine/Source/Runtime/Renderer/Private/VT/RuntimeVirtualTextureRender.cpp
--- a/Engine/Source/Runtime/Renderer/Private/VT/RuntimeVirtualTextureRender.cpp
+++ b/Engine/Source/Runtime/Renderer/Private/VT/RuntimeVirtualTextureRender.cpp
@@ -635,8 +635,7 @@ namespace RuntimeVirtualTexture
 		// Test scene is loaded and has been updated once by main rendering passes
 		// This function gets called on the main thread, so accessing scene frame number is not strictly thread safe, but we can probably
 		// assume that frame number is always increasing, and so the test is conservative
-		return Scene != nullptr && Scene->GetRenderScene() != nullptr && Scene->GetRenderScene()->GetFrameNumber() > 1 &&
-			Scene->GetRenderScene()->GPUScene.PrimitiveBuffer.SRV != nullptr; // todo[vt]: Shouldn't need this low level check, but sometimes this isn't created yet and so we crash in rendering.
+		return Scene != nullptr && Scene->GetRenderScene() != nullptr && Scene->GetRenderScene()->GetFrameNumber() > 1;
 	}
 
 	void RenderPage(
@@ -835,10 +834,13 @@ namespace RuntimeVirtualTexture
 
 	void RenderPages(FRHICommandListImmediate& RHICmdList, FRenderPageBatchDesc const& InDesc)
 	{
-		CSV_SCOPED_TIMING_STAT_EXCLUSIVE(STAT_VirtualTextureSystem_RVT_RenderPages);
 		SCOPED_DRAW_EVENT(RHICmdList, RuntimeVirtualTextureRenderPages);
 		check(InDesc.NumPageDescs <= EMaxRenderPageBatch);
 
+		// Make sure GPUScene is up to date. 
+		// Usually this is a no-op since we updated before calling, but RuntimeVirtualTexture::BuildStreamedMips() needs this.
+		UpdateGPUScene(RHICmdList, *InDesc.Scene->GetRenderScene());
+
 		for (int32 PageIndex = 0; PageIndex < InDesc.NumPageDescs; ++PageIndex)
 		{
 			FRenderPageDesc const& PageDesc = InDesc.PageDescs[PageIndex];
