Add alternative voxel mipmap method which mipmap 2 levels at a time, and cache data through LDS.

Disabled by default through console command as this method is slower (at least on Nvidia HW).

diff --git a/Engine/Shaders/Private/HairStrands/HairStrandsVoxelMip.usf b/Engine/Shaders/Private/HairStrands/HairStrandsVoxelMip.usf
--- a/Engine/Shaders/Private/HairStrands/HairStrandsVoxelMip.usf
+++ b/Engine/Shaders/Private/HairStrands/HairStrandsVoxelMip.usf
@@ -5,9 +5,79 @@ uint SourceMip;
 uint TargetMip;
 
 Texture3D<uint> InDensityTexture;
-RWTexture3D<uint> OutDensityTexture;
+RWTexture3D<uint> OutDensityTexture0;
 
-// #hair_todo: do two mips at a time through LDS?
+// Write two levels at a time, and store/cache thread result within LDS
+#if PERMUTATION_METHOD == 1
+
+RWTexture3D<uint> OutDensityTexture1;
+groupshared uint Values0[64];
+groupshared uint Values1[8];
+
+uint ToMorton(uint3 Index)
+{
+	return (Index.x & 0x1) | ((Index.y & 0x1)<<1) | ((Index.z& 0x1)<<2);
+}
+
+uint3 FromMorton(uint Index)
+{
+	return uint3((Index & 0x1) > 0 ? 1 : 0, (Index & 0x2) > 0 ? 1 : 0, (Index & 0x4) > 0 ? 1 : 0);
+}
+
+[numthreads(4, 4, 4)]
+void MainCS(uint3 DispatchThreadId : SV_DispatchThreadID, uint3 GroupThreadId : SV_GroupThreadID , uint GroupIndex : SV_GroupIndex, uint3 GroupId : SV_GroupId)
+{
+	{
+		const uint3 LocalBase	= uint3(GroupThreadId.x >> 1, GroupThreadId.y >> 1, GroupThreadId.z>> 1);
+		const uint3 LocalOffset = uint3(GroupThreadId.x % 2,  GroupThreadId.y % 2,  GroupThreadId.z % 2);
+
+		const uint Base   = ToMorton(LocalBase);
+		const uint Offset = ToMorton(LocalOffset);
+		const uint OutLocal = Base*8 + Offset;
+		Values0[OutLocal] = InDensityTexture[DispatchThreadId];
+	}
+
+	GroupMemoryBarrierWithGroupSync();
+
+	if (GroupIndex < 8)
+	{
+		const uint3 OutLocalCoord = FromMorton(GroupIndex);
+		const uint3 OutVoxelCoord = (GroupId << 1) + OutLocalCoord;
+
+		uint Total = 0;
+		const uint Offset = GroupIndex * 8;
+		Total += Values0[Offset + 0];
+		Total += Values0[Offset + 1];
+		Total += Values0[Offset + 2];
+		Total += Values0[Offset + 3];
+		Total += Values0[Offset + 4];
+		Total += Values0[Offset + 5];
+		Total += Values0[Offset + 6];
+		Total += Values0[Offset + 7];
+		Values1[GroupIndex] = Total;
+		OutDensityTexture0[OutVoxelCoord] = Total;
+	}
+
+	GroupMemoryBarrierWithGroupSync();
+
+	if (GroupIndex == 0)
+	{
+		const uint3 OutVoxelCoord = GroupId;
+		uint Total = 0;
+		Total += Values1[0];
+		Total += Values1[1];
+		Total += Values1[2];
+		Total += Values1[3];
+		Total += Values1[4];
+		Total += Values1[5];
+		Total += Values1[6];
+		Total += Values1[7];
+		OutDensityTexture1[OutVoxelCoord] = Total;
+	}
+}
+#endif
+
+#if PERMUTATION_METHOD == 0
 [numthreads(4, 4, 4)]
 void MainCS(uint3 DispatchThreadId : SV_DispatchThreadID)
 {
@@ -34,5 +104,6 @@ void MainCS(uint3 DispatchThreadId : SV_DispatchThreadID)
 		InDensityTexture[InVoxelCoord6] +
 		InDensityTexture[InVoxelCoord7];
 
-	OutDensityTexture[OutVoxelCoord] = Total;
+	OutDensityTexture0[OutVoxelCoord] = Total;
 }
+#endif
diff --git a/Engine/Source/Runtime/Renderer/Private/HairStrands/HairStrandsVoxelization.cpp b/Engine/Source/Runtime/Renderer/Private/HairStrands/HairStrandsVoxelization.cpp
--- a/Engine/Source/Runtime/Renderer/Private/HairStrands/HairStrandsVoxelization.cpp
+++ b/Engine/Source/Runtime/Renderer/Private/HairStrands/HairStrandsVoxelization.cpp
@@ -35,6 +35,9 @@ static FAutoConsoleVariableRef CVarHairVoxelizationDepthBiasScale(TEXT("r.HairSt
 static int32 GHairVoxelInjectOpaqueDepthEnable = 1;
 static FAutoConsoleVariableRef CVarHairVoxelInjectOpaqueDepthEnable(TEXT("r.HairStrands.Voxelization.InjectOpaqueDepth"), GHairVoxelInjectOpaqueDepthEnable, TEXT("Inject opaque geometry depth into the voxel volume for acting as occluder."));
 
+static int32 GHairStrandsVoxelMipMethod = 0;
+static FAutoConsoleVariableRef CVarHairVoxelMipMethod(TEXT("r.HairStrands.Voxelization.MipMethod"), GHairStrandsVoxelMipMethod, TEXT("Voxel mip methods (0 : one level per pass, 1 : two levels per pass."));
+
 float GetHairStrandsVoxelizationDensityScale() { return FMath::Max(0.0f, GHairVoxelizationDensityScale); }
 float GetHairStrandsVoxelizationDepthBiasScale() { return FMath::Max(0.0f, GHairVoxelizationDepthBiasScale); }
 
@@ -173,6 +176,9 @@ class FVoxelGenerateMipCS : public FGlobalShader
 	DECLARE_GLOBAL_SHADER(FVoxelGenerateMipCS);
 	SHADER_USE_PARAMETER_STRUCT(FVoxelGenerateMipCS, FGlobalShader);
 
+	class FMethod : SHADER_PERMUTATION_INT("PERMUTATION_METHOD", 2);
+	using FPermutationDomain = TShaderPermutationDomain<FMethod>;
+
 	BEGIN_SHADER_PARAMETER_STRUCT(FParameters, )
 		SHADER_PARAMETER_STRUCT_INCLUDE(FSceneTextureParameters, SceneTextures)
 
@@ -181,7 +187,8 @@ class FVoxelGenerateMipCS : public FGlobalShader
 		SHADER_PARAMETER(uint32, TargetMip)
 
 		SHADER_PARAMETER_RDG_TEXTURE_SRV(Texture3D, InDensityTexture)
-		SHADER_PARAMETER_RDG_TEXTURE_UAV(RWTexture3D, OutDensityTexture)
+		SHADER_PARAMETER_RDG_TEXTURE_UAV(RWTexture3D, OutDensityTexture0)
+		SHADER_PARAMETER_RDG_TEXTURE_UAV(RWTexture3D, OutDensityTexture1)
 
 	END_SHADER_PARAMETER_STRUCT()
 
@@ -202,24 +209,36 @@ static void AddVoxelGenerateMipPass(
 	TRefCountPtr<IPooledRenderTarget> DensityTexture = Cluster.VoxelResources.DensityTexture;
 	check(DensityTexture->GetDesc().Extent.X == DensityTexture->GetDesc().Extent.Y);
 
+	const uint32 NumLevelPerPass = GHairStrandsVoxelMipMethod > 0 ? 2 : 1;
+
 	const uint32 VoxelResolution = DensityTexture->GetDesc().Extent.X;
 	const uint32 MipCount = DensityTexture->GetDesc().NumMips;
 	FRDGTextureRef VoxelDensityTexture = GraphBuilder.RegisterExternalTexture(DensityTexture, TEXT("HairVoxelDensityTexture"));
-	for (uint32 MipIt = 0; MipIt < MipCount-1; ++MipIt)
+	for (uint32 MipIt = 0; MipIt < MipCount-1; MipIt += NumLevelPerPass)
 	{
 		FVoxelGenerateMipCS::FParameters* Parameters = GraphBuilder.AllocParameters<FVoxelGenerateMipCS::FParameters>();
 		Parameters->InDensityTexture = GraphBuilder.CreateSRV(FRDGTextureSRVDesc::CreateForMipLevel(VoxelDensityTexture, MipIt));
-		Parameters->OutDensityTexture = GraphBuilder.CreateUAV(FRDGTextureUAVDesc(VoxelDensityTexture, MipIt+1));
+		Parameters->OutDensityTexture0 = GraphBuilder.CreateUAV(FRDGTextureUAVDesc(VoxelDensityTexture, MipIt+1));
+		if (NumLevelPerPass > 1)
+		{
+			Parameters->OutDensityTexture1 = GraphBuilder.CreateUAV(FRDGTextureUAVDesc(VoxelDensityTexture, MipIt+2));
+		}
 		Parameters->VoxelResolution = VoxelResolution;
 		Parameters->SourceMip = MipIt;
 		Parameters->TargetMip = MipIt+1;
 
-		uint32 SourceResolution = VoxelResolution >> (MipIt);
-		uint32 TargetResolution = VoxelResolution >> (MipIt + 1);
+		const uint32 SourceResolution = VoxelResolution >> (MipIt);
+		const uint32 TargetResolution = VoxelResolution >> (MipIt + 1);
 
-		TShaderMapRef<FVoxelGenerateMipCS> ComputeShader(View.ShaderMap);
+		FVoxelGenerateMipCS::FPermutationDomain PermutationVector;
+		PermutationVector.Set<FVoxelGenerateMipCS::FMethod>(NumLevelPerPass == 2 ? 1 : 0);
+		
+		TShaderMapRef<FVoxelGenerateMipCS> ComputeShader(View.ShaderMap, PermutationVector);
 		const TShaderMap<FGlobalShaderType>* GlobalShaderMap = View.ShaderMap;
-		const FIntVector DispatchCount = FComputeShaderUtils::GetGroupCount(FIntVector(TargetResolution, TargetResolution, TargetResolution), FIntVector(4, 4, 4));
+		const FIntVector DispatchCount = 
+			NumLevelPerPass == 1 ? 
+			FComputeShaderUtils::GetGroupCount(FIntVector(TargetResolution, TargetResolution, TargetResolution), FIntVector(4, 4, 4)):
+			FComputeShaderUtils::GetGroupCount(FIntVector(SourceResolution, SourceResolution, SourceResolution), FIntVector(4, 4, 4));
 		
 		ClearUnusedGraphResources(*ComputeShader, Parameters);
 		GraphBuilder.AddPass(
