Add tiling support to ray tracing reflections pass. r.RayTracing.Reflections.RenderTileSize

Old Tile parameterto control reflections material sorting renamed to "SortTile" for clarity.

diff --git a/Engine/Shaders/Private/RayTracing/RayTracingCreateGatherPointsRGS.usf b/Engine/Shaders/Private/RayTracing/RayTracingCreateGatherPointsRGS.usf
--- a/Engine/Shaders/Private/RayTracing/RayTracingCreateGatherPointsRGS.usf
+++ b/Engine/Shaders/Private/RayTracing/RayTracingCreateGatherPointsRGS.usf
@@ -21,8 +21,8 @@ uint SamplesPerPixel;
 uint SampleIndex;
 uint MaxBounces;
 uint UpscaleFactor;
-uint TileOffsetX;
-uint TileOffsetY;
+uint RenderTileOffsetX;
+uint RenderTileOffsetY;
 float MaxRayDistanceForGI;
 float NextEventEstimationSamples;
 float DiffuseThreshold;
@@ -62,7 +62,7 @@ void GenerateCosineNormalRay(
 
 RAY_TRACING_ENTRY_RAYGEN(RayTracingCreateGatherPointsRGS)
 {
-	uint2 DispatchThreadId = DispatchRaysIndex().xy + View.ViewRectMin + uint2(TileOffsetX, TileOffsetY);
+	uint2 DispatchThreadId = DispatchRaysIndex().xy + View.ViewRectMin + uint2(RenderTileOffsetX, RenderTileOffsetY);
 	uint GatherPointsIndex = DispatchThreadIdToLinearIndex(DispatchThreadId);
 #if 1
 	// Default initialization
diff --git a/Engine/Shaders/Private/RayTracing/RayTracingGlobalIlluminationRGS.usf b/Engine/Shaders/Private/RayTracing/RayTracingGlobalIlluminationRGS.usf
--- a/Engine/Shaders/Private/RayTracing/RayTracingGlobalIlluminationRGS.usf
+++ b/Engine/Shaders/Private/RayTracing/RayTracingGlobalIlluminationRGS.usf
@@ -30,8 +30,8 @@ float DiffuseThreshold;
 uint EvalSkyLight;
 uint UseRussianRoulette;
 float MaxNormalBias;
-uint TileOffsetX;
-uint TileOffsetY;
+uint RenderTileOffsetX;
+uint RenderTileOffsetY;
 
 void GenerateCosineNormalRay(
 	float3 WorldPosition,
@@ -57,7 +57,7 @@ void GenerateCosineNormalRay(
 
 RAY_TRACING_ENTRY_RAYGEN(GlobalIlluminationRGS)
 {
-	uint2 DispatchThreadId = DispatchRaysIndex().xy + View.ViewRectMin + uint2(TileOffsetX, TileOffsetY);
+	uint2 DispatchThreadId = DispatchRaysIndex().xy + View.ViewRectMin + uint2(RenderTileOffsetX, RenderTileOffsetY);
 	RWGlobalIlluminationUAV[DispatchThreadId] = 0.0;
 	RWRayDistanceUAV[DispatchThreadId] = float2(-1.0, 0.0);
 
diff --git a/Engine/Shaders/Private/RayTracing/RayTracingReflections.usf b/Engine/Shaders/Private/RayTracing/RayTracingReflections.usf
--- a/Engine/Shaders/Private/RayTracing/RayTracingReflections.usf
+++ b/Engine/Shaders/Private/RayTracing/RayTracingReflections.usf
@@ -88,6 +88,8 @@ uint2 TileAlignedResolution;
 uint TestPathRoughness;
 float MinClearCoatLevel;
 uint MaxUnderCoatBounces;
+uint RenderTileOffsetX;
+uint RenderTileOffsetY;
 uint EnableTranslucency;
 
 float ReflectionMinRayDistance;
@@ -105,7 +107,7 @@ RWTexture2D<float> RayHitDistanceOutput;
 RWTexture2D<float> RayImaginaryDepthOutput;
 
 
-// Material buffer and tile size are used for sorted materials.
+// Material buffer and sort tile size are used for sorted materials.
 RWStructuredBuffer<FDeferredMaterialPayload> MaterialBuffer;
 
 // Generate a random direction to sample according to world normal and roughness.
@@ -209,7 +211,7 @@ RAY_TRACING_ENTRY_RAYGEN(RayTracingReflectionsRGS)
 		return;
 	}
 #else
-	uint2 DispatchThreadId = DispatchRaysIndex().xy;
+	uint2 DispatchThreadId = DispatchRaysIndex().xy + uint2(RenderTileOffsetX, RenderTileOffsetY);
 #endif
 
 	DispatchThreadId += View.ViewRectMin;
@@ -236,7 +238,7 @@ RAY_TRACING_ENTRY_RAYGEN(RayTracingReflectionsRGS)
 		LocalSamplesPerPixel = 0;
 	}
 
-	// Mask out samples from rounding to tile boundaries
+	// Mask out samples from rounding to sort tile boundaries
 	if (DIM_DEFERRED_MATERIAL_MODE == DEFERRED_MATERIAL_MODE_GATHER && any(DispatchThreadId >= RayTracingResolution))
 	{
 		LocalSamplesPerPixel = 0;
diff --git a/Engine/Source/Runtime/Renderer/Private/RayTracing/RayTracingGlobalIllumination.cpp b/Engine/Source/Runtime/Renderer/Private/RayTracing/RayTracingGlobalIllumination.cpp
--- a/Engine/Source/Runtime/Renderer/Private/RayTracing/RayTracingGlobalIllumination.cpp
+++ b/Engine/Source/Runtime/Renderer/Private/RayTracing/RayTracingGlobalIllumination.cpp
@@ -113,11 +113,11 @@ static TAutoConsoleVariable<int32> CVarRayTracingGlobalIlluminationEnableTwoSide
 	ECVF_RenderThreadSafe
 );
 
-static int32 GRayTracingGlobalIlluminationTileSize = 0;
-static FAutoConsoleVariableRef CVarRayTracingGlobalIlluminationTileSize(
-	TEXT("r.RayTracing.GlobalIllumination.TileSize"),
-	GRayTracingGlobalIlluminationTileSize,
-	TEXT("Render ray traced global illumination in NxN piel tiles, where each tile is submitted as separate GPU command buffer, allowing high quality rendering without triggering timeout detection. (default = 0, tiling disabled)")
+static int32 GRayTracingGlobalIlluminationRenderTileSize = 0;
+static FAutoConsoleVariableRef CVarRayTracingGlobalIlluminationRenderTileSize(
+	TEXT("r.RayTracing.GlobalIllumination.RenderTileSize"),
+	GRayTracingGlobalIlluminationRenderTileSize,
+	TEXT("Render ray traced global illumination in NxN pixel tiles, where each tile is submitted as separate GPU command buffer, allowing high quality rendering without triggering timeout detection. (default = 0, tiling disabled)")
 );
 
 static TAutoConsoleVariable<int32> CVarRayTracingGlobalIlluminationMaxLightCount(
@@ -286,8 +286,8 @@ class FGlobalIlluminationRGS : public FGlobalShader
 		SHADER_PARAMETER(uint32, EvalSkyLight)
 		SHADER_PARAMETER(uint32, UseRussianRoulette)
 		SHADER_PARAMETER(float, MaxNormalBias)
-		SHADER_PARAMETER(uint32, TileOffsetX)
-		SHADER_PARAMETER(uint32, TileOffsetY)
+		SHADER_PARAMETER(uint32, RenderTileOffsetX)
+		SHADER_PARAMETER(uint32, RenderTileOffsetY)
 
 		SHADER_PARAMETER_SRV(RaytracingAccelerationStructure, TLAS)
 		SHADER_PARAMETER_RDG_TEXTURE_UAV(RWTexture2D<float4>, RWGlobalIlluminationUAV)
@@ -333,8 +333,8 @@ class FRayTracingGlobalIlluminationCreateGatherPointsRGS : public FGlobalShader
 		SHADER_PARAMETER(uint32, SampleIndex)
 		SHADER_PARAMETER(uint32, MaxBounces)
 		SHADER_PARAMETER(uint32, UpscaleFactor)
-		SHADER_PARAMETER(uint32, TileOffsetX)
-		SHADER_PARAMETER(uint32, TileOffsetY)
+		SHADER_PARAMETER(uint32, RenderTileOffsetX)
+		SHADER_PARAMETER(uint32, RenderTileOffsetY)
 		SHADER_PARAMETER(float, MaxRayDistanceForGI)
 		SHADER_PARAMETER(float, NextEventEstimationSamples)
 		SHADER_PARAMETER(float, DiffuseThreshold)
@@ -387,8 +387,8 @@ class FRayTracingGlobalIlluminationFinalGatherRGS : public FGlobalShader
 		SHADER_PARAMETER(uint32, SampleIndex)
 		SHADER_PARAMETER(uint32, SamplesPerPixel)
 		SHADER_PARAMETER(uint32, UpscaleFactor)
-		SHADER_PARAMETER(uint32, TileOffsetX)
-		SHADER_PARAMETER(uint32, TileOffsetY)
+		SHADER_PARAMETER(uint32, RenderTileOffsetX)
+		SHADER_PARAMETER(uint32, RenderTileOffsetY)
 		SHADER_PARAMETER(float, DiffuseThreshold)
 		SHADER_PARAMETER(float, MaxNormalBias)
 		SHADER_PARAMETER(float, FinalGatherDistance)
@@ -555,8 +555,8 @@ void FDeferredShadingSceneRenderer::RayTracingGlobalIlluminationCreateGatherPoin
 	PassParameters->DiffuseThreshold = GRayTracingGlobalIlluminationDiffuseThreshold;
 	PassParameters->NextEventEstimationSamples = GRayTracingGlobalIlluminationNextEventEstimationSamples;
 	PassParameters->UpscaleFactor = UpscaleFactor;
-	PassParameters->TileOffsetX = 0;
-	PassParameters->TileOffsetY = 0;
+	PassParameters->RenderTileOffsetX = 0;
+	PassParameters->RenderTileOffsetY = 0;
 
 	// Global
 	PassParameters->TLAS = View.RayTracingScene.RayTracingSceneRHI->GetShaderResourceView();
@@ -648,8 +648,8 @@ void FDeferredShadingSceneRenderer::RenderRayTracingGlobalIlluminationFinalGathe
 	PassParameters->MaxNormalBias = GetRaytracingMaxNormalBias();
 	PassParameters->FinalGatherDistance = GRayTracingGlobalIlluminationFinalGatherDistance;
 	PassParameters->UpscaleFactor = UpscaleFactor;
-	PassParameters->TileOffsetX = 0;
-	PassParameters->TileOffsetY = 0;
+	PassParameters->RenderTileOffsetX = 0;
+	PassParameters->RenderTileOffsetY = 0;
 
 	// Scene data
 	PassParameters->TLAS = View.RayTracingScene.RayTracingSceneRHI->GetShaderResourceView();
@@ -776,8 +776,8 @@ void FDeferredShadingSceneRenderer::RenderRayTracingGlobalIlluminationBruteForce
 	PassParameters->TransmissionProfilesLinearSampler = TStaticSamplerState<SF_Bilinear, AM_Clamp, AM_Clamp, AM_Clamp>::GetRHI();
 	PassParameters->RWGlobalIlluminationUAV = GraphBuilder.CreateUAV(OutDenoiserInputs->Color);
 	PassParameters->RWRayDistanceUAV = GraphBuilder.CreateUAV(OutDenoiserInputs->RayHitDistance);
-	PassParameters->TileOffsetX = 0;
-	PassParameters->TileOffsetY = 0;
+	PassParameters->RenderTileOffsetX = 0;
+	PassParameters->RenderTileOffsetY = 0;
 
 	FGlobalIlluminationRGS::FPermutationDomain PermutationVector;
 	PermutationVector.Set<FGlobalIlluminationRGS::FUseAttenuationTermDim>(true);
@@ -787,7 +787,7 @@ void FDeferredShadingSceneRenderer::RenderRayTracingGlobalIlluminationBruteForce
 
 	FIntPoint RayTracingResolution = FIntPoint::DivideAndRoundUp(View.ViewRect.Size(), UpscaleFactor);
 
-	if (GRayTracingGlobalIlluminationTileSize <= 0)
+	if (GRayTracingGlobalIlluminationRenderTileSize <= 0)
 	{
 		GraphBuilder.AddPass(
 			RDG_EVENT_NAME("GlobalIlluminationRayTracing %dx%d", RayTracingResolution.X, RayTracingResolution.Y),
@@ -804,9 +804,9 @@ void FDeferredShadingSceneRenderer::RenderRayTracingGlobalIlluminationBruteForce
 	}
 	else
 	{
-		int32 TileSize = FMath::Max(32, GRayTracingGlobalIlluminationTileSize);
-		int32 NumTilesX = FMath::DivideAndRoundUp(RayTracingResolution.X, TileSize);
-		int32 NumTilesY = FMath::DivideAndRoundUp(RayTracingResolution.Y, TileSize);
+		int32 RenderTileSize = FMath::Max(32, GRayTracingGlobalIlluminationRenderTileSize);
+		int32 NumTilesX = FMath::DivideAndRoundUp(RayTracingResolution.X, RenderTileSize);
+		int32 NumTilesY = FMath::DivideAndRoundUp(RayTracingResolution.Y, RenderTileSize);
 		for (int32 Y = 0; Y < NumTilesY; ++Y)
 		{
 			for (int32 X = 0; X < NumTilesX; ++X)
@@ -818,12 +818,12 @@ void FDeferredShadingSceneRenderer::RenderRayTracingGlobalIlluminationBruteForce
 					TilePassParameters = GraphBuilder.AllocParameters<FGlobalIlluminationRGS::FParameters>();
 					*TilePassParameters = *PassParameters;
 
-					TilePassParameters->TileOffsetX = X * TileSize;
-					TilePassParameters->TileOffsetY = Y * TileSize;
+					TilePassParameters->RenderTileOffsetX = X * RenderTileSize;
+					TilePassParameters->RenderTileOffsetY = Y * RenderTileSize;
 				}
 
-				int32 DispatchSizeX = FMath::Min<int32>(TileSize, RayTracingResolution.X - TilePassParameters->TileOffsetX);
-				int32 DispatchSizeY = FMath::Min<int32>(TileSize, RayTracingResolution.Y - TilePassParameters->TileOffsetY);
+				int32 DispatchSizeX = FMath::Min<int32>(RenderTileSize, RayTracingResolution.X - TilePassParameters->RenderTileOffsetX);
+				int32 DispatchSizeY = FMath::Min<int32>(RenderTileSize, RayTracingResolution.Y - TilePassParameters->RenderTileOffsetY);
 
 				GraphBuilder.AddPass(
 					RDG_EVENT_NAME("GlobalIlluminationRayTracing %dx%d (tile %dx%d)", DispatchSizeX, DispatchSizeY, X, Y),
diff --git a/Engine/Source/Runtime/Renderer/Private/RayTracing/RayTracingReflections.cpp b/Engine/Source/Runtime/Renderer/Private/RayTracing/RayTracingReflections.cpp
--- a/Engine/Source/Runtime/Renderer/Private/RayTracing/RayTracingReflections.cpp
+++ b/Engine/Source/Runtime/Renderer/Private/RayTracing/RayTracingReflections.cpp
@@ -116,6 +116,12 @@ static TAutoConsoleVariable<int32> CVarRayTracingReflectionsSortTileSize(
 	TEXT("  Default 64\n"),
 	ECVF_RenderThreadSafe);
 
+static TAutoConsoleVariable<int32> CVarRayTracingReflectionsRenderTileSize(
+	TEXT("r.RayTracing.Reflections.RenderTileSize"),
+	0,
+	TEXT("Render ray traced reflections in NxN pixel tiles, where each tile is submitted as separate GPU command buffer, allowing high quality rendering without triggering timeout detection (default = 0, tiling disabled)"),
+	ECVF_RenderThreadSafe);
+
 static TAutoConsoleVariable<int32> CVarRayTracingReflectionsSortSize(
 	TEXT("r.RayTracing.Reflections.SortSize"),
 	5,
@@ -182,6 +188,8 @@ class FRayTracingReflectionsRGS : public FGlobalShader
 		SHADER_PARAMETER(int32, TestPathRoughness)
 		SHADER_PARAMETER(float, MinClearCoatLevel)
 		SHADER_PARAMETER(int32, MaxUnderCoatBounces)
+		SHADER_PARAMETER(uint32, RenderTileOffsetX)
+		SHADER_PARAMETER(uint32, RenderTileOffsetY)
 		SHADER_PARAMETER(uint32, EnableTranslucency)
 
 		SHADER_PARAMETER_SRV(RaytracingAccelerationStructure, TLAS)
@@ -322,6 +330,8 @@ void FDeferredShadingSceneRenderer::RenderRayTracingReflections(
 	CommonParameters.TestPathRoughness = CVarRayTracingReflectionsTestPathRoughness.GetValueOnRenderThread();
 	CommonParameters.MinClearCoatLevel = CVarRayTracingReflectionsMinClearCoatLevel.GetValueOnRenderThread();
 	CommonParameters.MaxUnderCoatBounces = CVarRayTracingReflectionsMaxUnderCoatBounces.GetValueOnRenderThread();
+	CommonParameters.RenderTileOffsetX = 0;
+	CommonParameters.RenderTileOffsetY = 0;
 	CommonParameters.EnableTranslucency = EnableTranslucency;
 
 	CommonParameters.TLAS = View.RayTracingScene.RayTracingSceneRHI->GetShaderResourceView();
@@ -406,28 +416,68 @@ void FDeferredShadingSceneRenderer::RenderRayTracingReflections(
 		}
 		else
 		{
-			GraphBuilder.AddPass(
-				RDG_EVENT_NAME("ReflectionRayTracing(spp=%d) %dx%d", SamplePerPixel, RayTracingResolution.X, RayTracingResolution.Y),
-				PassParameters,
-				ERDGPassFlags::Compute,
-				[PassParameters, this, &View, RayGenShader, RayTracingResolution, DeferredMaterialBufferNumElements, DeferredMaterialMode](FRHICommandList& RHICmdList)
+			// Add optional tiling behavior to avoid TDR events in expensive passes
+			int32 RenderTileSize = CVarRayTracingReflectionsRenderTileSize.GetValueOnRenderThread();
+			if (NumPasses > 1 || RenderTileSize <= 0)
 			{
-				FRayTracingShaderBindingsWriter GlobalResources;
-				SetShaderParameters(GlobalResources, RayGenShader, *PassParameters);
-
-				FRHIRayTracingScene* RayTracingSceneRHI = View.RayTracingScene.RayTracingSceneRHI;
-
-				if (DeferredMaterialMode == EDeferredMaterialMode::Shade)
+				GraphBuilder.AddPass(
+					RDG_EVENT_NAME("ReflectionRayTracing(spp=%d) %dx%d", SamplePerPixel, RayTracingResolution.X, RayTracingResolution.Y),
+					PassParameters,
+					ERDGPassFlags::Compute,
+					[PassParameters, this, &View, RayGenShader, RayTracingResolution, DeferredMaterialBufferNumElements, DeferredMaterialMode](FRHICommandList& RHICmdList)
 				{
-					// Shading pass for sorted materials uses 1D dispatch over all elements in the material buffer.
-					// This can be reduced to the number of output pixels if sorting pass guarantees that all invalid entries are moved to the end.
-					RHICmdList.RayTraceDispatch(View.RayTracingMaterialPipeline, RayGenShader->GetRayTracingShader(), RayTracingSceneRHI, GlobalResources, DeferredMaterialBufferNumElements, 1);
-				}
-				else // EDeferredMaterialMode::None
+					FRayTracingShaderBindingsWriter GlobalResources;
+					SetShaderParameters(GlobalResources, RayGenShader, *PassParameters);
+
+					FRHIRayTracingScene* RayTracingSceneRHI = View.RayTracingScene.RayTracingSceneRHI;
+
+					if (DeferredMaterialMode == EDeferredMaterialMode::Shade)
+					{
+						// Shading pass for sorted materials uses 1D dispatch over all elements in the material buffer.
+						// This can be reduced to the number of output pixels if sorting pass guarantees that all invalid entries are moved to the end.
+						RHICmdList.RayTraceDispatch(View.RayTracingMaterialPipeline, RayGenShader->GetRayTracingShader(), RayTracingSceneRHI, GlobalResources, DeferredMaterialBufferNumElements, 1);
+					}
+					else // EDeferredMaterialMode::None
+					{
+						RHICmdList.RayTraceDispatch(View.RayTracingMaterialPipeline, RayGenShader->GetRayTracingShader(), RayTracingSceneRHI, GlobalResources, RayTracingResolution.X, RayTracingResolution.Y);
+					}
+				});
+			}
+			else
+			{
+				RenderTileSize = FMath::Max(RenderTileSize, 32);
+				int32 NumTilesX = FMath::DivideAndRoundUp(RayTracingResolution.X, RenderTileSize);
+				int32 NumTilesY = FMath::DivideAndRoundUp(RayTracingResolution.Y, RenderTileSize);
+				for (int32 Y = 0; Y < NumTilesY; ++Y)
 				{
-					RHICmdList.RayTraceDispatch(View.RayTracingMaterialPipeline, RayGenShader->GetRayTracingShader(), RayTracingSceneRHI, GlobalResources, RayTracingResolution.X, RayTracingResolution.Y);
+					FRayTracingReflectionsRGS::FParameters* TilePassParameters = PassParameters;
+					for (int32 X = 0; X < NumTilesX; ++X)
+					{
+						if (X > 0 || Y > 0)
+						{
+							TilePassParameters = GraphBuilder.AllocParameters<FRayTracingReflectionsRGS::FParameters>();
+							*TilePassParameters = *PassParameters;
+							TilePassParameters->RenderTileOffsetX = X * RenderTileSize;
+							TilePassParameters->RenderTileOffsetY = Y * RenderTileSize;
+						}
+
+						int32 DispatchSizeX = FMath::Min<int32>(RenderTileSize, RayTracingResolution.X - TilePassParameters->RenderTileOffsetX);
+						int32 DispatchSizeY = FMath::Min<int32>(RenderTileSize, RayTracingResolution.Y - TilePassParameters->RenderTileOffsetY);
+						GraphBuilder.AddPass(
+							RDG_EVENT_NAME("ReflectionRayTracing(spp=%d) %dx%d", SamplePerPixel, DispatchSizeX, DispatchSizeY),
+							TilePassParameters,
+							ERDGPassFlags::Compute,
+							[TilePassParameters, this, &View, RayGenShader, DispatchSizeX, DispatchSizeY, DeferredMaterialBufferNumElements, DeferredMaterialMode](FRHICommandList& RHICmdList)
+						{
+							FRayTracingShaderBindingsWriter GlobalResources;
+							SetShaderParameters(GlobalResources, RayGenShader, *TilePassParameters);
+
+							FRHIRayTracingScene* RayTracingSceneRHI = View.RayTracingScene.RayTracingSceneRHI;
+							RHICmdList.RayTraceDispatch(View.RayTracingMaterialPipeline, RayGenShader->GetRayTracingShader(), RayTracingSceneRHI, GlobalResources, DispatchSizeX, DispatchSizeY);
+						});
+					}
 				}
-			});
+			}
 		}
 	}
 
