Fixed regression in translucency pass

diff --git a/Engine/Source/Runtime/Renderer/Private/RayTracing/RayTracingTranslucency.cpp b/Engine/Source/Runtime/Renderer/Private/RayTracing/RayTracingTranslucency.cpp
--- a/Engine/Source/Runtime/Renderer/Private/RayTracing/RayTracingTranslucency.cpp
+++ b/Engine/Source/Runtime/Renderer/Private/RayTracing/RayTracingTranslucency.cpp
@@ -239,8 +239,7 @@ void FDeferredShadingSceneRenderer::RenderRayTracingTranslucency(FRHICommandList
 			GraphBuilder,
 			View, &DenoiserInputs.Color, &DenoiserInputs.RayHitDistance,
 			TranslucencySPP, GRayTracingTranslucencyHeightFog, ResolutionFraction, 
-			ERayTracingPrimaryRaysFlag::AllowSkipSkySample
-		  & ERayTracingPrimaryRaysFlag::UseGBufferForMaxDistance);
+			ERayTracingPrimaryRaysFlag::AllowSkipSkySample | ERayTracingPrimaryRaysFlag::UseGBufferForMaxDistance);
 
 		//#dxr_todo: UE-72581 : replace DenoiserInputs with DenoiserOutputs in the following lines!
 		TRefCountPtr<IPooledRenderTarget> TranslucencyColor = GSystemTextures.BlackDummy;
