Added CVar to toggle ray tracing debug visualization modes. Default value is now BaseColor instead of showing a black viewport.

diff --git a/Engine/Source/Runtime/Renderer/Private/RayTracing/RayTracingDebug.cpp b/Engine/Source/Runtime/Renderer/Private/RayTracing/RayTracingDebug.cpp
--- a/Engine/Source/Runtime/Renderer/Private/RayTracing/RayTracingDebug.cpp
+++ b/Engine/Source/Runtime/Renderer/Private/RayTracing/RayTracingDebug.cpp
@@ -14,10 +14,17 @@
 #include "RayTracingDebugDefinitions.h"
 #include "RayTracing/RayTracingLighting.h"
 #include "RayTracing/RaytracingOptions.h"
+
 #define LOCTEXT_NAMESPACE "RayTracingDebugVisualizationMenuCommands"
 
 DECLARE_GPU_STAT(RayTracingDebug);
 
+static TAutoConsoleVariable<FString> CVarRayTracingDebugMode(
+	TEXT("r.RayTracing.DebugVisualizationMode"),
+	TEXT(""),
+	TEXT("Sets the ray tracing debug visualization mode (default = None - Driven by viewport menu) .\n")
+	);
+
 class FRayTracingDebugRGS : public FGlobalShader
 {
 	DECLARE_GLOBAL_SHADER(FRayTracingDebugRGS)
@@ -70,10 +77,27 @@ void FDeferredShadingSceneRenderer::RenderRayTracingDebug(FRHICommandListImmedia
 		RayTracingDebugVisualizationModes.Emplace(FName(*LOCTEXT("World Position", "World Position").ToString()),								RAY_TRACING_DEBUG_VIZ_WORLD_POSITION);
 		RayTracingDebugVisualizationModes.Emplace(FName(*LOCTEXT("HitKind", "HitKind").ToString()),												RAY_TRACING_DEBUG_VIZ_HITKIND);
 		RayTracingDebugVisualizationModes.Emplace(FName(*LOCTEXT("Barycentrics", "Barycentrics").ToString()),									RAY_TRACING_DEBUG_VIZ_BARYCENTRICS);
-		RayTracingDebugVisualizationModes.Emplace(FName(*LOCTEXT("PrimaryRays", "PrimaryRays").ToString()),							RAY_TRACING_DEBUG_VIZ_PRIMARY_RAYS);
+		RayTracingDebugVisualizationModes.Emplace(FName(*LOCTEXT("PrimaryRays", "PrimaryRays").ToString()),										RAY_TRACING_DEBUG_VIZ_PRIMARY_RAYS);
+	}
+
+	uint32 DebugVisualizationMode;
+	
+	FString ConsoleViewMode = CVarRayTracingDebugMode.GetValueOnRenderThread();
+
+	if (!ConsoleViewMode.IsEmpty())
+	{
+		DebugVisualizationMode = RayTracingDebugVisualizationModes.FindRef(FName(*ConsoleViewMode));
+	}
+	else if(View.CurrentRayTracingDebugVisualizationMode != NAME_None)
+	{
+		DebugVisualizationMode = RayTracingDebugVisualizationModes.FindRef(View.CurrentRayTracingDebugVisualizationMode);
+	}
+	else
+	{
+		// Set useful default value
+		DebugVisualizationMode = RAY_TRACING_DEBUG_VIZ_BASE_COLOR;
 	}
 
-	uint32 DebugVisualizationMode = RayTracingDebugVisualizationModes.FindRef(View.CurrentRayTracingDebugVisualizationMode);
 
 	if (DebugVisualizationMode == RAY_TRACING_DEBUG_VIZ_BARYCENTRICS)
 	{
