Add optional pass filtering pass for filtering out opaque depth fragment which are not useful.

diff --git a/Engine/Shaders/Private/HairStrands/HairStrandsVoxelOpaque.usf b/Engine/Shaders/Private/HairStrands/HairStrandsVoxelOpaque.usf
--- a/Engine/Shaders/Private/HairStrands/HairStrandsVoxelOpaque.usf
+++ b/Engine/Shaders/Private/HairStrands/HairStrandsVoxelOpaque.usf
@@ -5,6 +5,7 @@
 #include "../SceneTextureParameters.ush"
 
 
+#if PIXELSHADER
 float3 VoxelMinAABB;
 uint   VoxelResolution;
 float3 VoxelMaxAABB;
@@ -70,6 +71,46 @@ void MainPS(
 	if (ClosestDepth < VoxelDepth && abs(ClosestDepth - VoxelDepth) < DepthBias)
 	{
 		OutColor0 = float4(1,0,0,1);
-		DensityTexture[VoxelPos] = 100000;
+		InterlockedOr(DensityTexture[VoxelPos], 0x10000000);
 	}
 }
+#endif
+
+#if COMPUTESHADER
+
+uint				VoxelResolution;
+RWTexture3D<uint>	VoxelTexture;
+
+[numthreads(4, 4, 4)]
+void MainCS(uint3 DispatchThreadId : SV_DispatchThreadID)
+{
+	const uint  OpaqueMask	 =  0x10000000;
+	const uint  DensityMask	 = ~OpaqueMask;
+	const uint3 VoxelCoord   = DispatchThreadId;
+	const uint  VoxelDensity = VoxelTexture[VoxelCoord];
+	const bool  bHasOpaque	 = (OpaqueMask & VoxelDensity) > 0;
+
+	// If the current voxel does not contains any opaque data, there is nothing to filter.
+	if (!bHasOpaque)
+	{
+		return;
+	}
+
+	const int FilterSize = 3;
+	for (int z = -FilterSize; z <= FilterSize; z++)
+	for (int y = -FilterSize; y <= FilterSize; y++)
+	for (int x = -FilterSize; x <= FilterSize; x++)
+	{
+		const int3 Offset = int3(x, y, z);
+		const int3 Sample = clamp(VoxelCoord + Offset, 0, VoxelResolution-1);
+		const uint Density = VoxelTexture[Sample] & DensityMask;
+		if (Density > 0)
+		{
+			return;
+		}
+	}
+
+	// No hair density has been found, so this opaque voxel is a false positive
+	InterlockedAnd(VoxelTexture[VoxelCoord], DensityMask);
+}
+#endif
\ No newline at end of file
diff --git a/Engine/Source/Runtime/Renderer/Private/HairStrands/HairStrandsVoxelization.cpp b/Engine/Source/Runtime/Renderer/Private/HairStrands/HairStrandsVoxelization.cpp
--- a/Engine/Source/Runtime/Renderer/Private/HairStrands/HairStrandsVoxelization.cpp
+++ b/Engine/Source/Runtime/Renderer/Private/HairStrands/HairStrandsVoxelization.cpp
@@ -35,6 +35,9 @@ static FAutoConsoleVariableRef CVarHairVoxelizationDepthBiasScale(TEXT("r.HairSt
 static int32 GHairVoxelInjectOpaqueDepthEnable = 1;
 static FAutoConsoleVariableRef CVarHairVoxelInjectOpaqueDepthEnable(TEXT("r.HairStrands.Voxelization.InjectOpaqueDepth"), GHairVoxelInjectOpaqueDepthEnable, TEXT("Inject opaque geometry depth into the voxel volume for acting as occluder."));
 
+static int32 GHairVoxelFilterOpaqueDepthEnable = 0;
+static FAutoConsoleVariableRef CVarHairVoxelFilterOpaqueDepthEnable(TEXT("r.HairStrands.Voxelization.FilterOpaqueDepth"), GHairVoxelFilterOpaqueDepthEnable, TEXT("Filter opaque geometry depth into the voxel volume for acting as occluder."));
+
 static int32 GHairStrandsVoxelMipMethod = 0;
 static FAutoConsoleVariableRef CVarHairVoxelMipMethod(TEXT("r.HairStrands.Voxelization.MipMethod"), GHairStrandsVoxelMipMethod, TEXT("Voxel mip methods (0 : one level per pass, 1 : two levels per pass."));
 
@@ -170,6 +173,58 @@ static void AddVoxelInjectOpaquePass(
 	});
 }
 
+///////////////////////////////////////////////////////////////////////////////////////////////////
+class FVoxelFilterDepthCS : public FGlobalShader
+{
+	DECLARE_GLOBAL_SHADER(FVoxelFilterDepthCS);
+	SHADER_USE_PARAMETER_STRUCT(FVoxelFilterDepthCS, FGlobalShader);
+
+	BEGIN_SHADER_PARAMETER_STRUCT(FParameters, )
+		SHADER_PARAMETER_STRUCT_INCLUDE(FSceneTextureParameters, SceneTextures)
+
+		SHADER_PARAMETER(uint32, VoxelResolution)
+		SHADER_PARAMETER_RDG_TEXTURE_UAV(RWTexture3D, VoxelTexture)
+
+	END_SHADER_PARAMETER_STRUCT()
+
+public:
+	static bool ShouldCompilePermutation(const FGlobalShaderPermutationParameters& Parameters) { return IsHairStrandsSupported(Parameters.Platform); }
+};
+
+IMPLEMENT_GLOBAL_SHADER(FVoxelFilterDepthCS, "/Engine/Private/HairStrands/HairStrandsVoxelOpaque.usf", "MainCS", SF_Compute);
+
+static void AddFilterVoxelOpaqueDepthPass(
+	FRDGBuilder& GraphBuilder,
+	const FViewInfo& View,
+	FHairStrandsClusterData& Cluster)
+{
+	if (!Cluster.VoxelResources.DensityTexture)
+		return;
+
+	TRefCountPtr<IPooledRenderTarget> DensityTexture = Cluster.VoxelResources.DensityTexture;
+	check(DensityTexture->GetDesc().Extent.X == DensityTexture->GetDesc().Extent.Y);
+
+	const uint32 VoxelResolution = DensityTexture->GetDesc().Extent.X;
+	FRDGTextureRef VoxelDensityTexture = GraphBuilder.RegisterExternalTexture(DensityTexture, TEXT("HairVoxelDensityTexture"));
+	FVoxelFilterDepthCS::FParameters* Parameters = GraphBuilder.AllocParameters<FVoxelFilterDepthCS::FParameters>();
+	Parameters->VoxelTexture = GraphBuilder.CreateUAV(FRDGTextureUAVDesc(VoxelDensityTexture));
+	Parameters->VoxelResolution = VoxelResolution;
+
+	TShaderMapRef<FVoxelFilterDepthCS> ComputeShader(View.ShaderMap);
+	const TShaderMap<FGlobalShaderType>* GlobalShaderMap = View.ShaderMap;
+	const FIntVector DispatchCount = FComputeShaderUtils::GetGroupCount(FIntVector(VoxelResolution, VoxelResolution, VoxelResolution), FIntVector(4, 4, 4));
+
+	ClearUnusedGraphResources(*ComputeShader, Parameters);
+	GraphBuilder.AddPass(
+		RDG_EVENT_NAME("HairStrandsVoxelFilterDepth"),
+		Parameters,
+		ERDGPassFlags::Compute,
+		[Parameters, ComputeShader, DispatchCount](FRHICommandList& RHICmdList)
+	{
+		FComputeShaderUtils::Dispatch(RHICmdList, *ComputeShader, *Parameters, DispatchCount);
+	});
+}
+
 ///////////////////////////////////////////////////////////////////////////////////////////////////
 class FVoxelGenerateMipCS : public FGlobalShader
 {
@@ -380,6 +435,10 @@ static void RenderVoxelization(
 	{
 		FRDGBuilder GraphBuilder(RHICmdList);
 		AddVoxelInjectOpaquePass(GraphBuilder, *ViewInfo, Cluster);
+		if (GHairVoxelFilterOpaqueDepthEnable)
+		{
+			AddFilterVoxelOpaqueDepthPass(GraphBuilder, *ViewInfo, Cluster);
+		}
 		GraphBuilder.Execute();
 	}
 
