Upgrades to GTAO Implementation.
Using r.GTAO.UseNormals we now have the ability to use the derived view normals from the depth buffer.

With this option enabled using r.AmbientOcclusion.Compute 2 or 3 allows for the following passes to run compute

Horizon Search and Integate combined, Spatial filter.

The temporal filter requires access to the Velocity buffer and so needs to be on the graphics pipe.
This is then followed by the upsample if needed.

Renamed many of the internal function as there was confusion with the 6 different shaders and the 3 different ways they can run across Async Compute and graphics pipe.

Also fixed some haloing bugs on PC due to depth buffer no depth samples being passes into bilateral filter.

diff --git a/Engine/Shaders/Private/PostProcessAmbientOcclusion.usf b/Engine/Shaders/Private/PostProcessAmbientOcclusion.usf
--- a/Engine/Shaders/Private/PostProcessAmbientOcclusion.usf
+++ b/Engine/Shaders/Private/PostProcessAmbientOcclusion.usf
@@ -921,16 +921,16 @@ void MainSSAOSmoothCS(
 
 #if SHADER_QUALITY == 0
 	// very low
-	#define GTAO_NUMTAPS 5
+	#define GTAO_NUMTAPS 4
 #elif SHADER_QUALITY == 1
 	// low
-	#define GTAO_NUMTAPS 9
+	#define GTAO_NUMTAPS 7
 #elif SHADER_QUALITY == 2
 	// medium
-	#define GTAO_NUMTAPS 15
+	#define GTAO_NUMTAPS 12
 #elif SHADER_QUALITY == 3
 	// high
-	#define GTAO_NUMTAPS 25
+	#define GTAO_NUMTAPS 20
 #else // SHADER_QUALITY == 4
 	// very high
 	#define GTAO_NUMTAPS 35
@@ -992,40 +992,37 @@ float3 GetViewSpacePosFromHZB(float2 UV, float MipLevel)
 }
 
 
-float3 GetNormal(float2 UV, float3 ViewSpacePosMid, bool bConstructFromDepth)
+float3 GetNormal(float2 UV, float3 ViewSpacePosMid)
 {
 	float3 ViewSpaceNormal;
 
-	if(bConstructFromDepth)
-	{
-		// Get the normal derived from the depth buffer
-		float2 DeltaUV = View.BufferSizeAndInvSize.zw;
-	
-		float DeviceZ		= Texture2DSampleLevel(SceneTexturesStruct.SceneDepthTexture, SceneTexturesStruct.SceneDepthTextureSampler, UV,0).r;									
-		float DeviceZLeft	= Texture2DSampleLevel(SceneTexturesStruct.SceneDepthTexture, SceneTexturesStruct.SceneDepthTextureSampler, UV + float2(-DeltaUV.x,  0.0f),0).r;		
-		float DeviceZTop	= Texture2DSampleLevel(SceneTexturesStruct.SceneDepthTexture, SceneTexturesStruct.SceneDepthTextureSampler, UV + float2( 0.0f     , -DeltaUV.y),0).r;	
-		float DeviceZRight	= Texture2DSampleLevel(SceneTexturesStruct.SceneDepthTexture, SceneTexturesStruct.SceneDepthTextureSampler, UV + float2( DeltaUV.x,  0.0f),0).r;		
-		float DeviceZBottom = Texture2DSampleLevel(SceneTexturesStruct.SceneDepthTexture, SceneTexturesStruct.SceneDepthTextureSampler, UV + float2( 0.0f     ,  DeltaUV.y),0).r;	
+#if USE_NORMALBUFFER
 
-		float DeviceZDdx	= TakeSmallerAbsDelta(DeviceZLeft, DeviceZ, DeviceZRight);
-		float DeviceZDdy	= TakeSmallerAbsDelta(DeviceZTop, DeviceZ, DeviceZBottom);
+	// Get the normal from the normal buffer
+	float3 WorldNormal		= GetGBufferData(UV, false).WorldNormal;
+	ViewSpaceNormal	= normalize(mul(WorldNormal, (float3x3)View.TranslatedWorldToView));
 
-		float ZRight		= ConvertFromDeviceZ(DeviceZ + DeviceZDdx);
-		float ZDown			= ConvertFromDeviceZ(DeviceZ + DeviceZDdy);
+#else
+	// Get the normal derived from the depth buffer
+	float2 DeltaUV = View.BufferSizeAndInvSize.zw;
+	
+	float DeviceZ		= Texture2DSampleLevel(SceneTexturesStruct.SceneDepthTexture, SceneTexturesStruct.SceneDepthTextureSampler, UV,0).r;									
+	float DeviceZLeft	= Texture2DSampleLevel(SceneTexturesStruct.SceneDepthTexture, SceneTexturesStruct.SceneDepthTextureSampler, UV + float2(-DeltaUV.x,  0.0f),0).r;		
+	float DeviceZTop	= Texture2DSampleLevel(SceneTexturesStruct.SceneDepthTexture, SceneTexturesStruct.SceneDepthTextureSampler, UV + float2( 0.0f     , -DeltaUV.y),0).r;	
+	float DeviceZRight	= Texture2DSampleLevel(SceneTexturesStruct.SceneDepthTexture, SceneTexturesStruct.SceneDepthTextureSampler, UV + float2( DeltaUV.x,  0.0f),0).r;		
+	float DeviceZBottom = Texture2DSampleLevel(SceneTexturesStruct.SceneDepthTexture, SceneTexturesStruct.SceneDepthTextureSampler, UV + float2( 0.0f     ,  DeltaUV.y),0).r;	
 
-		float3 Right		= ScreenToViewPos(UV+ float2( DeltaUV.x,      0.0f) , ZRight)-ViewSpacePosMid;
-		float3 Down			= ScreenToViewPos(UV+ float2(      0.0f, DeltaUV.y) , ZDown) -ViewSpacePosMid;
-		
-		ViewSpaceNormal = normalize(cross(Right, Down));
+	float DeviceZDdx	= TakeSmallerAbsDelta(DeviceZLeft, DeviceZ, DeviceZRight);
+	float DeviceZDdy	= TakeSmallerAbsDelta(DeviceZTop, DeviceZ, DeviceZBottom);
 
-	}
-	else
-	{
-		// Get the normal from the normal buffer
-		float3 WorldNormal		= GetGBufferData(UV, false).WorldNormal;
-		ViewSpaceNormal	= normalize(mul(WorldNormal, (float3x3)View.TranslatedWorldToView));
+	float ZRight		= ConvertFromDeviceZ(DeviceZ + DeviceZDdx);
+	float ZDown			= ConvertFromDeviceZ(DeviceZ + DeviceZDdy);
 
-	}
+	float3 Right		= ScreenToViewPos(UV+ float2( DeltaUV.x,      0.0f) , ZRight)-ViewSpacePosMid;
+	float3 Down			= ScreenToViewPos(UV+ float2(      0.0f, DeltaUV.y) , ZDown) -ViewSpacePosMid;
+		
+	ViewSpaceNormal = normalize(cross(Right, Down));
+#endif
 
 	return ViewSpaceNormal;
 }
@@ -1345,7 +1342,8 @@ void GTAOCombinedPSandCS(in float2 UV, out float4 OutColor, inout float Depth)
 	float3 ViewSpacePos		= ScreenToViewPos(UV,SceneDepth);
 
 	// Get Normal and transform into viewspace
-	float3 ViewSpaceNormal	= GetNormal(UV, ViewSpacePos, false);
+	float3 ViewSpaceNormal	= GetNormal(UV, ViewSpacePos);
+
 #if GTAO_OPT
 	float3 ViewDir			= float3(0,0,-1);
 #else
@@ -1598,7 +1596,7 @@ float2 ReprojectPos(float2 UV, float Depth)
 	return PrevUV;
 }
 
-void GTAOTemporalFilterPSAndCS(float2 UV, out float OutAO, out float OutDepth)
+void GTAOTemporalFilterPSAndCS(float2 UV, inout float OutAO, inout float OutDepth)
 {
 	float NewAO			= Texture2DSample(PostprocessInput0, PostprocessInput0Sampler, UV).r;
 
@@ -1671,36 +1669,17 @@ void GTAOTemporalFilterCS(
 	float OutDepth		= 0;
 	int2   PixelPos		= DispatchThreadId + ScreenSpaceAOParams[5].zw; 
 	float2 PixelCenter	= (float2)PixelPos + float2(0.5, 0.5);
-	float2 BufferUV		= SvPositionToBufferUV(float4(PixelCenter, 0, 0));
+	float2 BufferUV		= PixelCenter.xy * PostprocessInput0Size.zw;
 
 	GTAOTemporalFilterPSAndCS(BufferUV, OutColor, OutDepth);
 
-	OutTexture[PixelPos] = OutColor;
+	OutTexture[PixelPos]		= OutColor;
+	DepthOutTexture[PixelPos]   = OutDepth;
 }
-#endif
-
-
-
-
-
-
 
-float4 GTAOUpsamplePSAndCS(float2 UV)
-{
-	float AO = saturate(PI_HALF * Texture2DSample(PostprocessInput0, PostprocessInput0Sampler, UV).r);
-
-	float AmbientOcclusionIntensity = ScreenSpaceAOParams[0].w;
-	float AmbientOcclusionPower		= ScreenSpaceAOParams[0].x;
-	// user adjust AO
-	AO = 1 - (1 - pow(abs(AO), AmbientOcclusionPower)) * AmbientOcclusionIntensity;
+#endif
 
-	return saturate(AO);
-}
 
-void GTAOUpsamplePS(in noperspective float4 UVAndScreenPos : TEXCOORD0,  out float4 OutColor : SV_Target0)
-{
-	OutColor = GTAOUpsamplePSAndCS(UVAndScreenPos.xy);
-}
 
 #if COMPUTE_SHADER
 
@@ -1753,6 +1732,8 @@ void GTAOUpsampleCS(
 {
 	float2 OutColor		= 0;
 
+
+
 	// Get Pixel Pos of the group we are processing
 	int2   PixelPos				= DispatchThreadId + ScreenSpaceAOParams[5].zw; 
 
@@ -1763,14 +1744,14 @@ void GTAOUpsampleCS(
 	int2   PosInGroup			= GroupThreadId.xy;	// Position with a group
 
 	// Get Full Rez Depth
-	float DeviceZ  = PostprocessInput0.Load(int3(PixelPos, 0)).r;
+	float DeviceZ  = PostprocessInput1.Load(int3(PixelPos, 0)).r;
 	float FullResZ = ConvertFromDeviceZ(DeviceZ); 
 	
 	// Get all the low res samples
 	float HalfZ   = ZReadTexture.Load(int3(HalfGroupOriginM1 +PosInGroup.xy , 0)).r;
 	LowResZArray[PosInGroup.x + (PosInGroup.y * THREADGROUP_SIZEX)] = HalfZ;
 
-	float AO      = saturate(PostprocessInput1.Load(int3(HalfGroupOriginM1 +PosInGroup.xy , 0)).r);
+	float AO      = saturate(PostprocessInput0.Load(int3(HalfGroupOriginM1 +PosInGroup.xy , 0)).r);
 	LowResAOArray[PosInGroup.x + (PosInGroup.y * THREADGROUP_SIZEX)] = AO;
 
 	GroupMemoryBarrierWithGroupSync();
@@ -1808,11 +1789,23 @@ void GTAOUpsampleCS(
 }
 #endif
 
-void GTAOSpatialFilterPS(in noperspective float4 UVAndScreenPos : TEXCOORD0,  out float4 OutColor : SV_Target0)
+
+
+float4 GTAOUpsamplePSAndCS(float2 UV)
 {
-	OutColor = Texture2DSample(PostprocessInput0, PostprocessInput0Sampler, UVAndScreenPos.xy);
+	float AO = saturate( Texture2DSample(PostprocessInput0, PostprocessInput0Sampler, UV).r);
+	return float4(AO,0,0,0);
 }
 
+void GTAOUpsamplePS(in noperspective float4 UVAndScreenPos : TEXCOORD0,  out float4 OutColor : SV_Target0)
+{
+	OutColor = GTAOUpsamplePSAndCS(UVAndScreenPos.xy);
+}
+
+
+
+
+
 
 
 #if COMPUTE_SHADER
@@ -1870,13 +1863,13 @@ void GTAOSpatialFilterCS(
 	uint XPos = pixIdx%11;
 	uint YPos = pixIdx/11;
 	SetAOVal(PostprocessInput0.Load(int3(FullGroupOriginM2 + int2(XPos,YPos), 0)).r, XPos, YPos);
-	SetZVal( ZReadTexture.Load(    int3(FullGroupOriginM2 + int2(XPos,YPos), 0)).r, XPos,YPos);
+	SetZVal( ZReadTexture.Load(     int3(FullGroupOriginM2 + int2(XPos,YPos), 0)).r, XPos, YPos);
 
 	pixIdx = 1+(GroupIndex*2);
 	XPos = pixIdx%11;
 	YPos = pixIdx/11;
 	SetAOVal(PostprocessInput0.Load(int3(FullGroupOriginM2 + int2(XPos,YPos), 0)).r, XPos, YPos);
-	SetZVal( ZReadTexture.Load(    int3(FullGroupOriginM2 + int2(XPos,YPos), 0)).r, XPos,YPos);
+	SetZVal( ZReadTexture.Load(     int3(FullGroupOriginM2 + int2(XPos,YPos), 0)).r, XPos, YPos);
 
 	GroupMemoryBarrierWithGroupSync();
 
@@ -1943,7 +1936,11 @@ void GTAOSpatialFilterCS(
 	// user adjust AO
 	SumAO = 1 - (1 - pow(abs(SumAO), AmbientOcclusionPower)) * AmbientOcclusionIntensity;
 
-
 	OutTexture[PixelPos] = SumAO ;
 }
-#endif
\ No newline at end of file
+#endif
+
+void GTAOSpatialFilterPS(in noperspective float4 UVAndScreenPos : TEXCOORD0,  out float4 OutColor : SV_Target0)
+{
+	OutColor = Texture2DSample(PostprocessInput0, PostprocessInput0Sampler, UVAndScreenPos.xy);
+}
\ No newline at end of file
diff --git a/Engine/Source/Runtime/Renderer/Private/CompositionLighting/CompositionLighting.cpp b/Engine/Source/Runtime/Renderer/Private/CompositionLighting/CompositionLighting.cpp
--- a/Engine/Source/Runtime/Renderer/Private/CompositionLighting/CompositionLighting.cpp
+++ b/Engine/Source/Runtime/Renderer/Private/CompositionLighting/CompositionLighting.cpp
@@ -138,41 +138,72 @@ bool ShouldRenderScreenSpaceAmbientOcclusion(const FViewInfo& View)
 	return bEnabled;
 }
 
-static FRenderingCompositeOutputRef AddPostProcessingGTAOAsyncHorizonSearch(FRHICommandListImmediate& RHICmdList, FPostprocessContext& Context)
+// Async Passes of the GTAO.
+// This can either just be the Horizon search if GBuffer Normals are needed or it can be
+// Combined Horizon search and Integrate followed by the Spatial filter if no normals are needed
+static FRenderingCompositeOutputRef AddPostProcessingGTAOAsyncPasses(FRHICommandListImmediate& RHICmdList, FPostprocessContext& Context, EGTAOType GTAOType)
 {
-	FRenderingCompositePass* FinalOutputPass;
+	check((GTAOType == EGTAOType::EAsyncHorizonSearch) || (GTAOType == EGTAOType::EAsyncCombinedSpatial));
 
+	FRenderingCompositePass* FinalOutputPass;
 
 	FSceneRenderTargets& SceneContext = FSceneRenderTargets::Get(Context.RHICmdList);
 	uint32 DownsampleFactor = CVarGTAODownsample.GetValueOnRenderThread() > 0 ? 2 : 1;
 
 	FIntPoint BufferSize	 = SceneContext.GetBufferSizeXY();
 	FIntPoint HorizonBufferSize = FIntPoint::DivideAndRoundUp(BufferSize, DownsampleFactor);
-	FPooledRenderTargetDesc Desc(FPooledRenderTargetDesc::Create2DDesc(HorizonBufferSize, PF_R8G8, FClearValueBinding::White, TexCreate_None, TexCreate_RenderTargetable, false));
+
+	FPooledRenderTargetDesc Desc(FPooledRenderTargetDesc::Create2DDesc(HorizonBufferSize, PF_R32_FLOAT, FClearValueBinding::White, TexCreate_None, TexCreate_RenderTargetable, false));
 	if (SceneContext.GetCurrentFeatureLevel() >= ERHIFeatureLevel::SM5)
 	{
 		Desc.TargetableFlags |= TexCreate_UAV;
 	}
-	GRenderTargetPool.FindFreeElement(RHICmdList, Desc, SceneContext.ScreenSpaceGTAOHorizons, TEXT("ScreenSpaceGTAOHorizons"));
-
-	Desc.Format = PF_R32_FLOAT;
 	GRenderTargetPool.FindFreeElement(RHICmdList, Desc, SceneContext.ScreenSpaceGTAODepths, TEXT("ScreenSpaceGTAODepths"));
 
+	Desc.Format = PF_R8G8;
+	GRenderTargetPool.FindFreeElement(RHICmdList, Desc, SceneContext.ScreenSpaceGTAOHorizons, TEXT("ScreenSpaceGTAOHorizons"));
+
 	FRenderingCompositePass* HZBInput = Context.Graph.RegisterPass(new FRCPassPostProcessInput(const_cast<FViewInfo&>(Context.View).HZB));
-	FRenderingCompositePass* AmbientOcclusionHorizonSearch = Context.Graph.RegisterPass(new(FMemStack::Get()) FRCPassPostProcessAmbientOcclusion_HorizonSearch(Context.View, DownsampleFactor, ESSAOType::EAsyncCS));
+	
+	if (GTAOType == EGTAOType::EAsyncHorizonSearch)
+	{
+		FRenderingCompositePass* AmbientOcclusionHorizonSearch = Context.Graph.RegisterPass(new(FMemStack::Get()) FRCPassPostProcessAmbientOcclusion_HorizonSearch(Context.View, DownsampleFactor, GTAOType));
 
-	AmbientOcclusionHorizonSearch->SetInput(ePId_Input0, Context.SceneDepth);
-	AmbientOcclusionHorizonSearch->SetInput(ePId_Input1, HZBInput);
+		AmbientOcclusionHorizonSearch->SetInput(ePId_Input0, Context.SceneDepth);
+		AmbientOcclusionHorizonSearch->SetInput(ePId_Input1, HZBInput);
+
+		FinalOutputPass = AmbientOcclusionHorizonSearch;
+
+	}
+	else  // (GTAOType == EGTAOType::EAsyncCombinedSpatial)
+	{
+		FRenderingCompositePass* AmbientOcclusionGTAO = Context.Graph.RegisterPass(new(FMemStack::Get()) FRCPassPostProcessAmbientOcclusion_GTAOHorizonSearchIntegrate(Context.View, DownsampleFactor, false, GTAOType));
+		AmbientOcclusionGTAO->SetInput(ePId_Input0, Context.SceneDepth);
+		AmbientOcclusionGTAO->SetInput(ePId_Input1, HZBInput);
+		FinalOutputPass = AmbientOcclusionGTAO;
+
+		// Add spatial Filter
+		if (CVarGTAOSpatialFilter.GetValueOnRenderThread() == 1)
+		{
+			FRenderingCompositePass* SpatialPass;
+			SpatialPass = Context.Graph.RegisterPass(new (FMemStack::Get()) FRCPassPostProcessAmbientOcclusion_GTAO_SpatialFilter(Context.View, DownsampleFactor, GTAOType));
+			SpatialPass->SetInput(ePId_Input0, FinalOutputPass);
+			SpatialPass->SetInput(ePId_Input1, HZBInput);
+			FinalOutputPass = SpatialPass;
+		}
+	}
 
-	FinalOutputPass = AmbientOcclusionHorizonSearch;
 
 	Context.FinalOutput = FRenderingCompositeOutputRef(FinalOutputPass);
 	return FRenderingCompositeOutputRef(FinalOutputPass);
 }
 
 
-static FRenderingCompositeOutputRef AddPostProcessingGTAOCombined(FRHICommandListImmediate& RHICmdList, FPostprocessContext& Context)
+
+// The whole GTAO stack is run on the Gfx Pipe
+static FRenderingCompositeOutputRef AddPostProcessingGTAOAllPasses(FRHICommandListImmediate& RHICmdList, FPostprocessContext& Context, EGTAOType GTAOType)
 {
+	// This can run on async compute if available and we don't use the per pixel normals
 	FRenderingCompositePass* FinalOutputPass;
 
 	FRenderingCompositePass* HZBInput = Context.Graph.RegisterPass(new FRCPassPostProcessInput(const_cast<FViewInfo&>(Context.View).HZB));
@@ -190,13 +221,14 @@ static FRenderingCompositeOutputRef AddPostProcessingGTAOCombined(FRHICommandLis
 		GRenderTargetPool.FindFreeElement(RHICmdList, Desc, SceneContext.ScreenSpaceGTAODepths, TEXT("ScreenSpaceGTAODepths"));
 	}
 
-	if(CVarGTAOCombined.GetValueOnRenderThread() ==1)
+	//if(CVarGTAOCombined.GetValueOnRenderThread() ==1 || (GTAOType== EGTAOType::ESplitNoNormalAsync) )
 	{
-		FRenderingCompositePass* AmbientOcclusionGTAO = Context.Graph.RegisterPass(new(FMemStack::Get()) FRCPassPostProcessAmbientOcclusion_GTAOCombined(Context.View, DownsampleFactor, false));
+		FRenderingCompositePass* AmbientOcclusionGTAO = Context.Graph.RegisterPass(new(FMemStack::Get()) FRCPassPostProcessAmbientOcclusion_GTAOHorizonSearchIntegrate(Context.View, DownsampleFactor, false, GTAOType));
 		AmbientOcclusionGTAO->SetInput(ePId_Input0, Context.SceneDepth);
 		AmbientOcclusionGTAO->SetInput(ePId_Input1, HZBInput);
 		FinalOutputPass = AmbientOcclusionGTAO;
 	}
+/*
 	else
 	{
 		FPooledRenderTargetDesc Desc(FPooledRenderTargetDesc::Create2DDesc(HorizonBufferSize, PF_R8G8, FClearValueBinding::White, TexCreate_None, TexCreate_RenderTargetable, false));
@@ -205,7 +237,7 @@ static FRenderingCompositeOutputRef AddPostProcessingGTAOCombined(FRHICommandLis
 			Desc.TargetableFlags |= TexCreate_UAV;
 		}
 		GRenderTargetPool.FindFreeElement(RHICmdList, Desc, SceneContext.ScreenSpaceGTAOHorizons, TEXT("ScreenSpaceGTAOHorizons"));
-		FRenderingCompositePass* AmbientOcclusionHorizonSearch = Context.Graph.RegisterPass(new(FMemStack::Get()) FRCPassPostProcessAmbientOcclusion_HorizonSearch(Context.View, DownsampleFactor, ESSAOType::ECS));
+		FRenderingCompositePass* AmbientOcclusionHorizonSearch = Context.Graph.RegisterPass(new(FMemStack::Get()) FRCPassPostProcessAmbientOcclusion_HorizonSearch(Context.View, DownsampleFactor, EGTAOType::ESplitUseNormalAsync));
 
 		AmbientOcclusionHorizonSearch->SetInput(ePId_Input0, Context.SceneDepth);
 		AmbientOcclusionHorizonSearch->SetInput(ePId_Input1, HZBInput);
@@ -217,6 +249,7 @@ static FRenderingCompositeOutputRef AddPostProcessingGTAOCombined(FRHICommandLis
 		AmbientOcclusionInnerIntegrate->SetInput(ePId_Input1, FinalOutputPass);
 		FinalOutputPass = AmbientOcclusionInnerIntegrate;
 	}
+*/
 
 	SceneContext.bScreenSpaceAOIsValid = true;
 
@@ -226,7 +259,7 @@ static FRenderingCompositeOutputRef AddPostProcessingGTAOCombined(FRHICommandLis
 	if (CVarGTAOSpatialFilter.GetValueOnRenderThread() == 1)
 	{
 		FRenderingCompositePass* SpatialPass;
-		SpatialPass = Context.Graph.RegisterPass(new (FMemStack::Get()) FRCPassPostProcessAmbientOcclusion_GTAO_SpatialFilter(Context.View, DownsampleFactor));
+		SpatialPass = Context.Graph.RegisterPass(new (FMemStack::Get()) FRCPassPostProcessAmbientOcclusion_GTAO_SpatialFilter(Context.View, DownsampleFactor, GTAOType));
 		SpatialPass->SetInput(ePId_Input0, FinalOutputPass);
 		SpatialPass->SetInput(ePId_Input1, HZBInput);
 		FinalOutputPass = SpatialPass;
@@ -238,21 +271,19 @@ static FRenderingCompositeOutputRef AddPostProcessingGTAOCombined(FRHICommandLis
 		// Add temporal filter
 		FRenderingCompositePass* TemporalPass;
 		TemporalPass = Context.Graph.RegisterPass(new (FMemStack::Get()) FRCPassPostProcessAmbientOcclusion_GTAO_TemporalFilter(Context.View, DownsampleFactor,
-			Context.View.PrevViewInfo.GTAOHistory,
-			&ViewState->PrevFrameViewInfo.GTAOHistory));
+			Context.View.PrevViewInfo.GTAOHistory, &ViewState->PrevFrameViewInfo.GTAOHistory, GTAOType));
 
 		TemporalPass->SetInput(ePId_Input0, FinalOutputPass);
 		FinalOutputPass = TemporalPass;
-
 	}
 	
 	// If downsampled then upsample the final result
 	if(DownsampleFactor > 1)
 	{
 		FRenderingCompositePass* UpsamplePass;
-		UpsamplePass = Context.Graph.RegisterPass(new (FMemStack::Get()) FRCPassPostProcessAmbientOcclusion_GTAO_Upsample(Context.View, DownsampleFactor));
-		UpsamplePass->SetInput(ePId_Input0, Context.SceneDepth);
-		UpsamplePass->SetInput(ePId_Input1, FRenderingCompositeOutputRef(FinalOutputPass, ePId_Output0));
+		UpsamplePass = Context.Graph.RegisterPass(new (FMemStack::Get()) FRCPassPostProcessAmbientOcclusion_GTAO_Upsample(Context.View, DownsampleFactor, GTAOType));
+		UpsamplePass->SetInput(ePId_Input0, FRenderingCompositeOutputRef(FinalOutputPass, ePId_Output0));
+		UpsamplePass->SetInput(ePId_Input1, Context.SceneDepth);
 		FinalOutputPass = UpsamplePass;
 	}
 
@@ -261,57 +292,72 @@ static FRenderingCompositeOutputRef AddPostProcessingGTAOCombined(FRHICommandLis
 }
 
 
-static FRenderingCompositeOutputRef AddPostProcessingGTAOIntegration(FRHICommandListImmediate& RHICmdList, FPostprocessContext& Context)
+// These are the passes run after Async where some are run before on the Async pipe
+static FRenderingCompositeOutputRef AddPostProcessingGTAOPostAsync(FRHICommandListImmediate& RHICmdList, FPostprocessContext& Context, EGTAOType GTAOType)
 {
-	FRenderingCompositePass* FinalOutputPass;
+	FSceneRenderTargets& SceneContext = FSceneRenderTargets::Get(RHICmdList);
+	SceneContext.bScreenSpaceAOIsValid = true; 
+	
+	FRenderingCompositePass* FinalOutputPass =nullptr;
 
 	FRenderingCompositePass* HZBInput = Context.Graph.RegisterPass(new FRCPassPostProcessInput(const_cast<FViewInfo&>(Context.View).HZB));
 	uint32 DownsampleFactor = CVarGTAODownsample.GetValueOnRenderThread() > 0 ? 2 : 1;
-
-	FRenderingCompositePass* AmbientOcclusionInnerIntegrate = Context.Graph.RegisterPass(new(FMemStack::Get()) FRCPassPostProcessAmbientOcclusion_InnerIntegrate(Context.View, DownsampleFactor, false));
-	AmbientOcclusionInnerIntegrate->SetInput(ePId_Input0, Context.SceneDepth);
-	FinalOutputPass = AmbientOcclusionInnerIntegrate;
-
-	FSceneRenderTargets& SceneContext = FSceneRenderTargets::Get(RHICmdList);
-	SceneContext.bScreenSpaceAOIsValid = true;
-
 	FSceneViewState* ViewState = Context.View.ViewState;
 
-	// Add spatial Filter
-	if (CVarGTAOSpatialFilter.GetValueOnRenderThread() == 1)
+	// If we run Just the async Horizon Search Earlier then need to do all the other passes
+	if (GTAOType == EGTAOType::EAsyncHorizonSearch)
 	{
-		FRenderingCompositePass* SpatialPass;
-		SpatialPass = Context.Graph.RegisterPass(new (FMemStack::Get()) FRCPassPostProcessAmbientOcclusion_GTAO_SpatialFilter(Context.View, DownsampleFactor));
-		SpatialPass->SetInput(ePId_Input0, FinalOutputPass);
-		SpatialPass->SetInput(ePId_Input1, HZBInput);
-		FinalOutputPass = SpatialPass;
-	}
+		FRenderingCompositePass* AmbientOcclusionInnerIntegrate = Context.Graph.RegisterPass(new(FMemStack::Get()) FRCPassPostProcessAmbientOcclusion_GTAOInnerIntegrate(Context.View, DownsampleFactor, false));
+		AmbientOcclusionInnerIntegrate->SetInput(ePId_Input0, Context.SceneDepth);
+		FinalOutputPass = AmbientOcclusionInnerIntegrate;
+
+		// Add spatial Filter
+		if (CVarGTAOSpatialFilter.GetValueOnRenderThread() == 1)
+		{
+			FRenderingCompositePass* SpatialPass;
+			SpatialPass = Context.Graph.RegisterPass(new (FMemStack::Get()) FRCPassPostProcessAmbientOcclusion_GTAO_SpatialFilter(Context.View, DownsampleFactor, GTAOType));
+			// TODO . If the 
 
-//	bool bNeedsUpsample = DownsampleFactor != 1;
+			SpatialPass->SetInput(ePId_Input0, FinalOutputPass);
+			SpatialPass->SetInput(ePId_Input1, HZBInput);
+			FinalOutputPass = SpatialPass;
+		}
+	}
 
 	// Add temporal filter
 	if (ViewState && CVarGTAOTemporalFilter.GetValueOnRenderThread() == 1)
 	{
 		FRenderingCompositePass* TemporalPass;
 		TemporalPass = Context.Graph.RegisterPass(new (FMemStack::Get()) FRCPassPostProcessAmbientOcclusion_GTAO_TemporalFilter(Context.View, DownsampleFactor,
-			Context.View.PrevViewInfo.GTAOHistory,
-			&ViewState->PrevFrameViewInfo.GTAOHistory));
+			Context.View.PrevViewInfo.GTAOHistory, &ViewState->PrevFrameViewInfo.GTAOHistory, GTAOType));
 
-		TemporalPass->SetInput(ePId_Input0, FinalOutputPass);
-		FinalOutputPass = TemporalPass;
+		// If the Spatial Filter is running as part of the async then we'll render to the R channel of the horizons texture so it can be read in as part of the temporal
+		if (GTAOType == EGTAOType::EAsyncCombinedSpatial)
+		{
+			// The Spatial Filter Pass was stored in the horizons texture in the previous graph on the async pipe
+			FRenderingCompositePass* HorizonsTexture = Context.Graph.RegisterPass(new(FMemStack::Get()) FRCPassPostProcessInput(SceneContext.ScreenSpaceGTAOHorizons));
+			TemporalPass->SetInput(ePId_Input0, HorizonsTexture);
+		}
+		else
+		{
+			TemporalPass->SetInput(ePId_Input0, FinalOutputPass);
+		}
 
+
+		FinalOutputPass = TemporalPass;
 	}
+
+	//Upsample Pass
 	{
 		FRenderingCompositePass* UpsamplePass;
-		UpsamplePass = Context.Graph.RegisterPass(new (FMemStack::Get()) FRCPassPostProcessAmbientOcclusion_GTAO_Upsample(Context.View, DownsampleFactor));
-		UpsamplePass->SetInput(ePId_Input0, Context.SceneDepth);
-		UpsamplePass->SetInput(ePId_Input1, FRenderingCompositeOutputRef(FinalOutputPass, ePId_Output0));
+		UpsamplePass = Context.Graph.RegisterPass(new (FMemStack::Get()) FRCPassPostProcessAmbientOcclusion_GTAO_Upsample(Context.View, DownsampleFactor, GTAOType));
+		UpsamplePass->SetInput(ePId_Input0, FRenderingCompositeOutputRef(FinalOutputPass, ePId_Output0));
+		UpsamplePass->SetInput(ePId_Input1, Context.SceneDepth);
 		FinalOutputPass = UpsamplePass;
 	}
 
 	Context.FinalOutput = FRenderingCompositeOutputRef(FinalOutputPass);
 	return FRenderingCompositeOutputRef(FinalOutputPass);
-
 }
 
 
@@ -449,7 +495,7 @@ void FCompositionLighting::ProcessBeforeBasePass(FRHICommandListImmediate& RHICm
 		if (SSAOLevels)
 		{
 
-			if (FSSAOHelper::GetGTAOPassType(View) != EGTAOType::ECombinedNonAsync)
+			if (FSSAOHelper::GetGTAOPassType(View) != EGTAOType::ENonAsync)
 			{
 				AddPostProcessingAmbientOcclusion(RHICmdList, Context, SSAOLevels);
 			}
@@ -536,9 +582,9 @@ void FCompositionLighting::ProcessAfterBasePass(FRHICommandListImmediate& RHICmd
 			{
 				if(!FSSAOHelper::IsAmbientOcclusionAsyncCompute(Context.View, SSAOLevels))
 				{
-					if (FSSAOHelper::GetGTAOPassType(View) == EGTAOType::ECombinedNonAsync)
+					if (FSSAOHelper::GetGTAOPassType(View) == EGTAOType::ENonAsync)
 					{
-						AmbientOcclusion = AddPostProcessingGTAOCombined(RHICmdList, Context);
+						AmbientOcclusion = AddPostProcessingGTAOAllPasses(RHICmdList, Context, EGTAOType::ENonAsync);
 					}
 					else
 					{
@@ -556,9 +602,11 @@ void FCompositionLighting::ProcessAfterBasePass(FRHICommandListImmediate& RHICmd
 				else
 				{
 					// If doing the Split GTAO method then we need to do the second part here.
-					if (FSSAOHelper::GetGTAOPassType(View) == EGTAOType::ESplitAsync)
+					if ( (FSSAOHelper::GetGTAOPassType(View) == EGTAOType::EAsyncHorizonSearch) || 
+						 (FSSAOHelper::GetGTAOPassType(View) == EGTAOType::EAsyncCombinedSpatial) )
+
 					{
-						AmbientOcclusion = AddPostProcessingGTAOIntegration(RHICmdList, Context);
+						AmbientOcclusion = AddPostProcessingGTAOPostAsync(RHICmdList, Context, FSSAOHelper::GetGTAOPassType(View));
 					}
 
 					ensureMsgf(
@@ -668,9 +716,10 @@ void FCompositionLighting::ProcessAsyncSSAO(FRHICommandListImmediate& RHICmdList
 
 				FPostprocessContext Context(RHICmdList, CompositeContext.Graph, View);
 
-				if (FSSAOHelper::GetGTAOPassType(View) == EGTAOType::ESplitAsync)
+				if ( (FSSAOHelper::GetGTAOPassType(View) == EGTAOType::EAsyncHorizonSearch) || 
+					 (FSSAOHelper::GetGTAOPassType(View) == EGTAOType::EAsyncCombinedSpatial) )
 				{
-					FRenderingCompositeOutputRef AmbientOcclusion = AddPostProcessingGTAOAsyncHorizonSearch(RHICmdList, Context);
+					FRenderingCompositeOutputRef AmbientOcclusion = AddPostProcessingGTAOAsyncPasses(RHICmdList, Context, FSSAOHelper::GetGTAOPassType(View));
 					Context.FinalOutput = FRenderingCompositeOutputRef(AmbientOcclusion);
 				}
 				else
diff --git a/Engine/Source/Runtime/Renderer/Private/CompositionLighting/PostProcessAmbientOcclusion.cpp b/Engine/Source/Runtime/Renderer/Private/CompositionLighting/PostProcessAmbientOcclusion.cpp
--- a/Engine/Source/Runtime/Renderer/Private/CompositionLighting/PostProcessAmbientOcclusion.cpp
+++ b/Engine/Source/Runtime/Renderer/Private/CompositionLighting/PostProcessAmbientOcclusion.cpp
@@ -20,12 +20,12 @@ DECLARE_GPU_STAT_NAMED(SSAOSetup, TEXT("ScreenSpace AO Setup") );
 DECLARE_GPU_STAT_NAMED(SSAO, TEXT("ScreenSpace AO") );
 DECLARE_GPU_STAT_NAMED(BasePassAO, TEXT("BasePass AO") );
 DECLARE_GPU_STAT_NAMED(SSAOSmooth, TEXT("SSAO smooth"));
-DECLARE_GPU_STAT_NAMED(GTAO_HorizonSearch,	TEXT("GTAO HorizonSearch"));
-DECLARE_GPU_STAT_NAMED(GTAO_Combined,		TEXT("GTAO Combined"));
-DECLARE_GPU_STAT_NAMED(GTAO_InnerIntegrate, TEXT("GTAO InnerIntegrate"));
-DECLARE_GPU_STAT_NAMED(GTAO_TemporalFilter, TEXT("GTAO Temportal Filter"));
-DECLARE_GPU_STAT_NAMED(GTAO_SpatialFilter,	TEXT("GTAO Spatial Filter"));
-DECLARE_GPU_STAT_NAMED(GTAO_Upsample,		TEXT("GTAO Upsample"));
+DECLARE_GPU_STAT_NAMED(GTAO_HorizonSearch,				TEXT("GTAO HorizonSearch"));
+DECLARE_GPU_STAT_NAMED(GTAO_HorizonSearchIntegrate,		TEXT("GTAO HorizonSearch And Integrate"));
+DECLARE_GPU_STAT_NAMED(GTAO_InnerIntegrate,				TEXT("GTAO InnerIntegrate"));
+DECLARE_GPU_STAT_NAMED(GTAO_TemporalFilter,				TEXT("GTAO Temportal Filter"));
+DECLARE_GPU_STAT_NAMED(GTAO_SpatialFilter,				TEXT("GTAO Spatial Filter"));
+DECLARE_GPU_STAT_NAMED(GTAO_Upsample,					TEXT("GTAO Upsample"));
 
 // Tile size for the AmbientOcclusion compute shader, tweaked for 680 GTX. */
 // see GCN Performance Tip 21 http://developer.amd.com/wordpress/media/2013/05/GCNPerformanceTweets.pdf 
@@ -101,6 +101,13 @@ static TAutoConsoleVariable<int32> CVarAmbientOcclusionMethod(
 	TEXT("1: GTAO\n "),
 	ECVF_RenderThreadSafe | ECVF_Scalability);
 
+static TAutoConsoleVariable<int32> CVarGTAOUseNormals(
+	TEXT("r.GTAO.UseNormals"),
+	1,
+	TEXT("Whether to use GBuffer Normals or Depth Derived normals \n ")
+	TEXT("0: Off \n ")
+	TEXT("1: On (default)\n "),
+	ECVF_RenderThreadSafe | ECVF_Scalability);
 
 float FSSAOHelper::GetAmbientOcclusionQualityRT(const FSceneView& View)
 {
@@ -223,23 +230,39 @@ uint32 FSSAOHelper::ComputeAmbientOcclusionPassCount(const FViewInfo& View)
 	return Ret;
 }
 
+
+// Helper function to get what type of method we are using
+// EGTAOType::EOff					: This is when r.AmbientOcclusion.Method == 0
+// EGTAOType::EAsyncHorizonSearch   : This is when We need GBuffer normals and the hardware Supports Async Compute. The trace pass is on the 
+//									  Async Pipe and the Intergrate, Spatial and temporal Filters are on the Gfx pipe after the base pass
+// EGTAOType::EAsyncCombinedSpatial	: This is when we use Derived normals from the depth buffer and the hardware Supports Async Compute. 
+//									  All passes will be on the async compute pipe
+// EGTAOType::ENonAsync				: All passes are are on the graphics pipe. Can use either gbuffer normals or derived depth normals.
+
 EGTAOType FSSAOHelper::GetGTAOPassType(const FViewInfo& View)
 {
-	int32 Method = CVarAmbientOcclusionMethod.GetValueOnRenderThread();
+	int32 Method		= CVarAmbientOcclusionMethod.GetValueOnRenderThread();
+	int32 UseNormals	= CVarGTAOUseNormals.GetValueOnRenderThread();
 
 	if (Method == 1)
 	{
 		if (IsAmbientOcclusionAsyncCompute(View, 1) && GSupportsEfficientAsyncCompute)
 		{
-			return EGTAOType::ESplitAsync;
+			if (UseNormals)
+			{
+				return EGTAOType::EAsyncHorizonSearch;
+			}
+			else
+			{
+				return EGTAOType::EAsyncCombinedSpatial;
+			}
 		}
 		else
 		{
-			return EGTAOType::ECombinedNonAsync;
+			return EGTAOType::ENonAsync;
 		}
 	}
 	return EGTAOType::EOff;
-
 }
 
 
@@ -1215,6 +1238,7 @@ public:
 	template<typename ShaderRHIParamRef, typename TRHICmdList>
 	void Set(TRHICmdList& RHICmdList, const FViewInfo& View, FIntPoint DestSize, const ShaderRHIParamRef ShaderRHI) const
 	{
+		static const auto CVarTemporalFilter = IConsoleManager::Get().FindTConsoleVariableDataInt(TEXT("r.GTAO.TemporalFilter"));
 		const FFinalPostProcessSettings& Settings = View.FinalPostProcessSettings;
 
 		uint32 TemporalFrame = 0;
@@ -1222,7 +1246,7 @@ public:
 		
 		const FSceneViewState* ViewState = static_cast<const FSceneViewState*>(View.State);
 
-		if (ViewState)
+		if (ViewState && CVarTemporalFilter->GetValueOnRenderThread() > 0)
 		{
 			TemporalFrame = ViewState->GetCurrentUnclampedTemporalAASampleIndex();
 			Frame = ViewState->GetFrameIndex();
@@ -1514,10 +1538,10 @@ public:
 		RHICmdList.SetUAVParameter(ShaderRHI, DepthOutTexture.GetBaseIndex(), DepthOutUAV);
 
 		FVector4 HZBRemappingValue = GetHZBRemapVal(Context, DestSize, InputTextureSize);
-		SetShaderValue(Context.RHICmdList, ShaderRHI, HZBRemapping, HZBRemappingValue);
+		SetShaderValue(RHICmdList, ShaderRHI, HZBRemapping, HZBRemappingValue);
 
 		const FSceneRenderTargetItem& GTAORandomization = GSystemTextures.GTAORandomization->GetRenderTargetItem();
-		SetTextureParameter(Context.RHICmdList, ShaderRHI, RandomNormalTexture, RandomNormalTextureSampler, TStaticSamplerState<SF_Point, AM_Wrap, AM_Wrap, AM_Wrap>::GetRHI(), GTAORandomization.ShaderResourceTexture);
+		SetTextureParameter(RHICmdList, ShaderRHI, RandomNormalTexture, RandomNormalTextureSampler, TStaticSamplerState<SF_Point, AM_Wrap, AM_Wrap, AM_Wrap>::GetRHI(), GTAORandomization.ShaderResourceTexture);
 	}
 
 	void SetParametersPS(const FRenderingCompositePassContext& Context, FIntPoint DestSize, FIntPoint InputTextureSize)
@@ -1672,7 +1696,7 @@ public:
 		const FSceneRenderTargetItem& GTAORandomization = GSystemTextures.GTAORandomization->GetRenderTargetItem();
 		SetTextureParameter(RHICmdList, ShaderRHI, RandomNormalTexture, RandomNormalTextureSampler, TStaticSamplerState<SF_Point, AM_Wrap, AM_Wrap, AM_Wrap>::GetRHI(), GTAORandomization.ShaderResourceTexture);
 
-		FSceneRenderTargets& SceneContext = FSceneRenderTargets::Get(Context.RHICmdList);
+		FSceneRenderTargets& SceneContext = FSceneRenderTargets::Get(RHICmdList);
 
 		SetTextureParameter(RHICmdList, ShaderRHI, HorizonsTexture, HorizonsTextureSampler, TStaticSamplerState<SF_Point, AM_Wrap, AM_Wrap, AM_Wrap>::GetRHI(), 
 			SceneContext.ScreenSpaceGTAOHorizons->GetRenderTargetItem().ShaderResourceTexture);
@@ -1724,32 +1748,33 @@ public:
 IMPLEMENT_SHADER_TYPE2(FPostProcessGTAOInnerIntegratePSandCS<false>, SF_Pixel);
 IMPLEMENT_SHADER_TYPE2(FPostProcessGTAOInnerIntegratePSandCS<true>, SF_Compute);
 
-FRCPassPostProcessAmbientOcclusion_GTAOCombined::FRCPassPostProcessAmbientOcclusion_GTAOCombined(const FSceneView& View, uint32 InDownScaleFactor, bool FinalOutput)
-	:bFinalOutput(FinalOutput),
+FRCPassPostProcessAmbientOcclusion_GTAOHorizonSearchIntegrate::FRCPassPostProcessAmbientOcclusion_GTAOHorizonSearchIntegrate(const FSceneView& View, uint32 InDownScaleFactor, bool FinalOutput, const EGTAOType InAOType)
+	: AOType(InAOType),
+	bFinalOutput(FinalOutput),
 	DownScaleFactor(InDownScaleFactor)
 {
-
 }
 
-template <uint32 ShaderQuality, uint32 UseNormals>
-void FRCPassPostProcessAmbientOcclusion_GTAOCombined::DispatchCS(const FRenderingCompositePassContext& Context, FIntRect ViewRect, FIntPoint DestSize, FIntPoint TexSize)
+
+template <uint32 ShaderQuality, uint32 UseNormals, typename TRHICmdList>
+void FRCPassPostProcessAmbientOcclusion_GTAOHorizonSearchIntegrate::DispatchCS(TRHICmdList& RHICmdList, const FRenderingCompositePassContext& Context, FIntRect ViewRect, FIntPoint DestSize, FIntPoint TexSize)
 {
 	TShaderMapRef<FPostProcessGTAOCombinedPSandCS<true, ShaderQuality, UseNormals>> ComputeShader(Context.GetShaderMap());
 	FSceneRenderTargets& SceneContext = FSceneRenderTargets::Get(Context.RHICmdList);
 	const FSceneRenderTargetItem& DestRenderTarget = PassOutputs[0].RequestSurface(Context);
 	const FSceneRenderTargetItem& DestRenderTarget1 = SceneContext.ScreenSpaceGTAODepths->GetRenderTargetItem();
 
-	Context.RHICmdList.SetComputeShader(ComputeShader->GetComputeShader());
-	ComputeShader->SetParametersCS(Context.RHICmdList, Context, DestSize, TexSize, DownScaleFactor, DestRenderTarget.UAV, DestRenderTarget1.UAV);
+	RHICmdList.SetComputeShader(ComputeShader->GetComputeShader());
+	ComputeShader->SetParametersCS(RHICmdList, Context, DestSize, TexSize, DownScaleFactor, DestRenderTarget.UAV, DestRenderTarget1.UAV);
 
 	uint32 GroupSizeX = FMath::DivideAndRoundUp(ViewRect.Width(), 8);
 	uint32 GroupSizeY = FMath::DivideAndRoundUp(ViewRect.Height(), 8);
-	DispatchComputeShader(Context.RHICmdList, *ComputeShader, GroupSizeX, GroupSizeY, 1);
-	ComputeShader->UnsetParameters(Context.RHICmdList);
+	DispatchComputeShader(RHICmdList, *ComputeShader, GroupSizeX, GroupSizeY, 1);
+	ComputeShader->UnsetParameters(RHICmdList);
 }
 
 template <uint32 ShaderQuality, uint32 UseNormals>
-FShader* FRCPassPostProcessAmbientOcclusion_GTAOCombined::SetShaderPS(const FRenderingCompositePassContext& Context, FGraphicsPipelineStateInitializer& GraphicsPSOInit, FIntPoint DestSize)
+FShader* FRCPassPostProcessAmbientOcclusion_GTAOHorizonSearchIntegrate::SetShaderPS(const FRenderingCompositePassContext& Context, FGraphicsPipelineStateInitializer& GraphicsPSOInit, FIntPoint DestSize)
 {
 	TShaderMapRef<FPostProcessVS> VertexShader(Context.GetShaderMap());
 	TShaderMapRef<FPostProcessGTAOCombinedPSandCS<false, ShaderQuality, UseNormals> > PixelShader(Context.GetShaderMap());
@@ -1769,11 +1794,9 @@ FShader* FRCPassPostProcessAmbientOcclusion_GTAOCombined::SetShaderPS(const FRen
 	return *VertexShader;
 }
 
-
-
-void FRCPassPostProcessAmbientOcclusion_GTAOCombined::Process(FRenderingCompositePassContext& Context)
+void FRCPassPostProcessAmbientOcclusion_GTAOHorizonSearchIntegrate::Process(FRenderingCompositePassContext& Context)
 {
-	SCOPED_GPU_STAT(Context.RHICmdList, GTAO_Combined);
+	SCOPED_GPU_STAT(Context.RHICmdList, GTAO_HorizonSearchIntegrate);
 	const FViewInfo& View = Context.View;
 
 	// Get Size of destination
@@ -1785,38 +1808,71 @@ void FRCPassPostProcessAmbientOcclusion_GTAOCombined::Process(FRenderingComposit
 	FIntPoint InputTexSize = InputDesc0 ? InputDesc0->Extent : SceneContext.GetBufferSizeXY();
 	FIntPoint OutputTexSize = PassOutputs[0].RenderTargetDesc.Extent;
 	FIntRect ViewRect = View.ViewRect / DownScaleFactor;
+   
+	static const auto CVarCompute = IConsoleManager::Get().FindTConsoleVariableDataInt(TEXT("r.AmbientOcclusion.Compute"));
+	static const auto CVarNormals = IConsoleManager::Get().FindTConsoleVariableDataInt(TEXT("r.GTAO.UseNormals"));
 
-	   
-	static const auto CVar = IConsoleManager::Get().FindTConsoleVariableDataInt(TEXT("r.AmbientOcclusion.Compute"));
-
+	bool bUseNormals = CVarNormals->GetValueOnRenderThread() >= 1;
 	const int32 ShaderQuality = FSSAOHelper::GetAmbientOcclusionShaderLevel(Context.View);
 
-	if (CVar->GetValueOnRenderThread() >= 1)
+	if (1) //CVarCompute->GetValueOnRenderThread() >= 1)
 	{
-		Context.RHICmdList.TransitionResource(EResourceTransitionAccess::ERWBarrier, EResourceTransitionPipeline::EGfxToCompute, DestRenderTarget.UAV, nullptr);
-
 		// Compute  Version
 		UnbindRenderTargets(Context.RHICmdList);
 
-		// 0..4, 0:low 4:high
-		switch(ShaderQuality)
+#define SET_SHADER_CASE_CS(RHICmdList,ShaderQualityCase) \
+		case ShaderQualityCase: \
+		if (bUseNormals)		\
+		{						\
+			DispatchCS<ShaderQualityCase, 1>(RHICmdList, Context, ViewRect, OutputTexSize, InputTexSize);\
+		} \
+		else \
+		{ \
+			DispatchCS<ShaderQualityCase, 0>(RHICmdList, Context, ViewRect, OutputTexSize, InputTexSize);\
+		} \
+		break
+
+		// If on the Async Pipe
+		if (AOType == EGTAOType::EAsyncCombinedSpatial)
 		{
-		case 0:
-			DispatchCS<0,0>(Context, ViewRect, OutputTexSize, InputTexSize);
-			break;
-		case 1:
-			DispatchCS<1, 0>(Context, ViewRect, OutputTexSize, InputTexSize);
-			break;
-		case 2:
-			DispatchCS<2, 0>(Context, ViewRect, OutputTexSize, InputTexSize);
-			break;
-		case 3:
-			DispatchCS<3, 0>(Context, ViewRect, OutputTexSize, InputTexSize);
-			break;
-		case 4:
-			DispatchCS<4, 0>(Context, ViewRect, OutputTexSize, InputTexSize);
-			break;
+			static FName AsyncStartFenceName(TEXT("AsyncStartFence"));
+			FComputeFenceRHIRef AsyncStartFence = Context.RHICmdList.CreateComputeFence(AsyncStartFenceName);
+			FRHIAsyncComputeCommandListImmediate& RHICmdListComputeImmediate = FRHICommandListExecutor::GetImmediateAsyncComputeCommandList();
+
+			//Fence to let us know when the Gfx pipe is done with the RT we want to write to.
+			Context.RHICmdList.TransitionResource(EResourceTransitionAccess::ERWBarrier, EResourceTransitionPipeline::EGfxToCompute, DestRenderTarget.UAV, AsyncStartFence);
+
+			SCOPED_COMPUTE_EVENT(RHICmdListComputeImmediate, AsyncSSAO);
+			RHICmdListComputeImmediate.WaitComputeFence(AsyncStartFence);
+
+			// 0..4, 0:low 4:high
+			switch (ShaderQuality)
+			{
+				SET_SHADER_CASE_CS(RHICmdListComputeImmediate, 0);
+				SET_SHADER_CASE_CS(RHICmdListComputeImmediate, 1);
+				SET_SHADER_CASE_CS(RHICmdListComputeImmediate, 2);
+				SET_SHADER_CASE_CS(RHICmdListComputeImmediate, 3);
+				SET_SHADER_CASE_CS(RHICmdListComputeImmediate, 4);
+			}
 		}
+		else
+		{
+			Context.RHICmdList.TransitionResource(EResourceTransitionAccess::ERWBarrier, EResourceTransitionPipeline::EGfxToCompute, DestRenderTarget.UAV, nullptr);
+
+			// 0..4, 0:low 4:high
+			switch (ShaderQuality)
+			{
+				SET_SHADER_CASE_CS(Context.RHICmdList, 0);
+				SET_SHADER_CASE_CS(Context.RHICmdList, 1);
+				SET_SHADER_CASE_CS(Context.RHICmdList, 2);
+				SET_SHADER_CASE_CS(Context.RHICmdList, 3);
+				SET_SHADER_CASE_CS(Context.RHICmdList, 4);
+			}
+		}
+
+
+#undef SET_SHADER_CASE_CS
+
 	}
 	else
 	{
@@ -1840,25 +1896,28 @@ void FRCPassPostProcessAmbientOcclusion_GTAOCombined::Process(FRenderingComposit
 
 			FShader* VertexShader = 0;
 
+#define SET_SHADER_CASE_PS(ShaderQualityCase) \
+			case ShaderQualityCase: \
+			if (bUseNormals)		\
+			{						\
+				VertexShader = SetShaderPS<ShaderQualityCase, 1>(Context, GraphicsPSOInit, OutputTexSize);\
+			} \
+			else \
+			{ \
+				VertexShader = SetShaderPS<ShaderQualityCase, 0>(Context, GraphicsPSOInit, OutputTexSize);\
+			} \
+			break
+
+			// 0..4, 0:low 4:high
 			switch (ShaderQuality)
 			{
-			case 0:
-				VertexShader = SetShaderPS<0, 0>(Context, GraphicsPSOInit, OutputTexSize);
-				break;
-			case 1:
-				VertexShader = SetShaderPS<1, 0>(Context, GraphicsPSOInit, OutputTexSize);
-				break;
-			case 2:
-				VertexShader = SetShaderPS<2, 0>(Context, GraphicsPSOInit, OutputTexSize);
-				break;
-			case 3:
-				VertexShader = SetShaderPS<3, 0>(Context, GraphicsPSOInit, OutputTexSize);
-				break;
-			case 4:
-				VertexShader = SetShaderPS<4, 0>(Context, GraphicsPSOInit, OutputTexSize);
-				break;
+				SET_SHADER_CASE_PS(0);
+				SET_SHADER_CASE_PS(1);
+				SET_SHADER_CASE_PS(2);
+				SET_SHADER_CASE_PS(3);
+				SET_SHADER_CASE_PS(4);
 			}
-
+#undef SET_SHADER_CASE_PS
 
 			DrawRectangle(
 				Context.RHICmdList,
@@ -1868,8 +1927,6 @@ void FRCPassPostProcessAmbientOcclusion_GTAOCombined::Process(FRenderingComposit
 				OutputTexSize,
 				VertexShader,
 				EDRF_UseTriangleOptimization);
-
-
 		}
 		Context.RHICmdList.EndRenderPass();
 	}
@@ -1878,7 +1935,7 @@ void FRCPassPostProcessAmbientOcclusion_GTAOCombined::Process(FRenderingComposit
 }
 
 
-FPooledRenderTargetDesc FRCPassPostProcessAmbientOcclusion_GTAOCombined::ComputeOutputDesc(EPassOutputId InPassOutputId) const
+FPooledRenderTargetDesc FRCPassPostProcessAmbientOcclusion_GTAOHorizonSearchIntegrate::ComputeOutputDesc(EPassOutputId InPassOutputId) const
 {
 	FPooledRenderTargetDesc Ret;
 
@@ -1897,13 +1954,13 @@ FPooledRenderTargetDesc FRCPassPostProcessAmbientOcclusion_GTAOCombined::Compute
 }
 
 
-FRCPassPostProcessAmbientOcclusion_InnerIntegrate::FRCPassPostProcessAmbientOcclusion_InnerIntegrate(const FSceneView& View, uint32 InDownScaleFactor, bool FinalOutput)
+FRCPassPostProcessAmbientOcclusion_GTAOInnerIntegrate::FRCPassPostProcessAmbientOcclusion_GTAOInnerIntegrate(const FSceneView& View, uint32 InDownScaleFactor, bool FinalOutput)
 	:bFinalOutput(FinalOutput),
 	DownScaleFactor(InDownScaleFactor)
 {
 }
 
-void FRCPassPostProcessAmbientOcclusion_InnerIntegrate::Process(FRenderingCompositePassContext& Context)
+void FRCPassPostProcessAmbientOcclusion_GTAOInnerIntegrate::Process(FRenderingCompositePassContext& Context)
 {
 	SCOPED_GPU_STAT(Context.RHICmdList, GTAO_InnerIntegrate);
 	const FViewInfo& View = Context.View;
@@ -1976,7 +2033,7 @@ void FRCPassPostProcessAmbientOcclusion_InnerIntegrate::Process(FRenderingCompos
 }
 
 
-FPooledRenderTargetDesc FRCPassPostProcessAmbientOcclusion_InnerIntegrate::ComputeOutputDesc(EPassOutputId InPassOutputId) const
+FPooledRenderTargetDesc FRCPassPostProcessAmbientOcclusion_GTAOInnerIntegrate::ComputeOutputDesc(EPassOutputId InPassOutputId) const
 {
 	FPooledRenderTargetDesc Ret;
 
@@ -1995,7 +2052,7 @@ FPooledRenderTargetDesc FRCPassPostProcessAmbientOcclusion_InnerIntegrate::Compu
 }
 
 
-FRCPassPostProcessAmbientOcclusion_HorizonSearch::FRCPassPostProcessAmbientOcclusion_HorizonSearch(const FSceneView& View, uint32 InDownScaleFactor, const ESSAOType InAOType)
+FRCPassPostProcessAmbientOcclusion_HorizonSearch::FRCPassPostProcessAmbientOcclusion_HorizonSearch(const FSceneView& View, uint32 InDownScaleFactor, const EGTAOType InAOType)
 	: AOType(InAOType),
 	DownScaleFactor(InDownScaleFactor)
 {
@@ -2066,7 +2123,7 @@ void FRCPassPostProcessAmbientOcclusion_HorizonSearch::Process(FRenderingComposi
 	static const auto CVar = IConsoleManager::Get().FindTConsoleVariableDataInt(TEXT("r.AmbientOcclusion.Compute"));
 	const int32 ShaderQuality = FSSAOHelper::GetAmbientOcclusionShaderLevel(Context.View);
 
-	if (AOType == ESSAOType::EAsyncCS)
+	if (AOType == EGTAOType::EAsyncHorizonSearch)
 	{
 		static FName AsyncStartFenceName(TEXT("AsyncStartFence"));
 		FComputeFenceRHIRef AsyncStartFence = Context.RHICmdList.CreateComputeFence(AsyncStartFenceName);
@@ -2216,7 +2273,9 @@ public:
 	FShaderParameter				PrevScreenPositionScaleBias;
 	FShaderParameter				BlendParams;
 	FScreenSpaceAOParameters		ScreenSpaceAOParams;
-	FShaderParameter OutTexture;
+
+	FShaderResourceParameter		OutTexture;
+	FShaderResourceParameter 		DepthOutTexture;
 
 
 	/** Initialization constructor. */
@@ -2237,9 +2296,69 @@ public:
 
 		PrevScreenPositionScaleBias.Bind(Initializer.ParameterMap,  TEXT("PrevScreenPositionScaleBias"));
 		BlendParams.Bind(Initializer.ParameterMap,					TEXT("BlendParams"));
-		OutTexture.Bind(Initializer.ParameterMap, TEXT("OutTexture"));
 
 		ScreenSpaceAOParams.Bind(Initializer.ParameterMap);
+
+		if (bComputeShader)
+		{
+			OutTexture.Bind(Initializer.ParameterMap,				TEXT("OutTexture"));
+			DepthOutTexture.Bind(Initializer.ParameterMap,			TEXT("DepthOutTexture"));
+		}
+	}
+
+	template <typename TRHICmdList>
+	void SetParametersCS(TRHICmdList& RHICmdList, FRenderingCompositePassContext& Context, FIntPoint DestSize, FIntPoint InputTextureSize, bool bCameraCut, const FGTAOTAAHistory& InputHistory, TRefCountPtr<IPooledRenderTarget>& VelocityRT,
+		FRHIUnorderedAccessView *OutUAV0, FRHIUnorderedAccessView *OutUAV1)
+	{
+		const FFinalPostProcessSettings& Settings = Context.View.FinalPostProcessSettings;
+		FRHIComputeShader* ShaderRHI = GetComputeShader();
+
+		PostprocessParameter.SetCS(ShaderRHI, Context, RHICmdList, TStaticSamplerState<SF_Bilinear, AM_Clamp, AM_Clamp, AM_Clamp>::GetRHI());
+
+		FGlobalShader::SetParameters<FViewUniformShaderParameters>(RHICmdList, ShaderRHI, Context.View.ViewUniformBuffer);
+		SceneTextureParameters.Set(RHICmdList, ShaderRHI, Context.View.FeatureLevel, ESceneTextureSetupMode::All);
+
+		FIntPoint ViewportOffset = InputHistory.ViewportRect.Min;
+		FIntPoint ViewportExtent = InputHistory.ViewportRect.Size();
+		FIntPoint BufferSize = InputHistory.ReferenceBufferSize;
+
+		FVector4 PrevScreenPositionScaleBiasValue = FVector4(
+			ViewportExtent.X * 0.5f / BufferSize.X,
+			-ViewportExtent.Y * 0.5f / BufferSize.Y,
+			(ViewportExtent.X * 0.5f + ViewportOffset.X) / BufferSize.X,
+			(ViewportExtent.Y * 0.5f + ViewportOffset.Y) / BufferSize.Y);
+
+		SetShaderValue(RHICmdList, ShaderRHI, PrevScreenPositionScaleBias, PrevScreenPositionScaleBiasValue);
+
+		FVector4 BlendParamsValue = FVector4(bCameraCut ? 1.0f : 0.0f, 0.0f, 0.0f, 0.0f);
+		SetShaderValue(RHICmdList, ShaderRHI, BlendParams, BlendParamsValue);
+
+		if (InputHistory.IsValid())
+		{
+			SetTextureParameter(RHICmdList, ShaderRHI, HistoryTexture, HistoryTextureSampler,
+				TStaticSamplerState<SF_Bilinear, AM_Border, AM_Border, AM_Border, 0, 0, 0xffffffff >::GetRHI(),
+				InputHistory.RT[0]->GetRenderTargetItem().TargetableTexture);
+
+
+			SetTextureParameter(RHICmdList, ShaderRHI, DepthHistoryTexture, DepthHistoryTextureSampler,
+				TStaticSamplerState<SF_Bilinear>::GetRHI(),
+				InputHistory.Depth[0]->GetRenderTargetItem().TargetableTexture);
+		}
+		else
+		{
+			// Need to bind a white dummy
+			SetTextureParameter(RHICmdList, ShaderRHI, HistoryTexture, HistoryTextureSampler,
+				TStaticSamplerState<SF_Bilinear>::GetRHI(),
+				GSystemTextures.WhiteDummy->GetRenderTargetItem().ShaderResourceTexture);
+
+		}
+
+		SetTextureParameter(RHICmdList, ShaderRHI, SceneVelocityTexture, SceneVelocityTextureSampler,
+			TStaticSamplerState<SF_Point>::GetRHI(), VelocityRT->GetRenderTargetItem().ShaderResourceTexture);
+
+		// Bind the output UAVs
+		RHICmdList.SetUAVParameter(ShaderRHI, OutTexture.GetBaseIndex(),	  OutUAV0);
+		RHICmdList.SetUAVParameter(ShaderRHI, DepthOutTexture.GetBaseIndex(), OutUAV1);
 	}
 
 
@@ -2268,7 +2387,6 @@ public:
 		FVector4 BlendParamsValue = FVector4(bCameraCut ? 1.0f : 0.0f, 0.0f, 0.0f, 0.0f);
 		SetShaderValue(Context.RHICmdList, ShaderRHI, BlendParams, BlendParamsValue);
 
-
 		if (InputHistory.IsValid())
 		{
 			SetTextureParameter(Context.RHICmdList, ShaderRHI, HistoryTexture, HistoryTextureSampler,
@@ -2299,6 +2417,7 @@ public:
 	{
 		FRHIComputeShader* ShaderRHI = GetComputeShader();
 		RHICmdList.SetUAVParameter(ShaderRHI, OutTexture.GetBaseIndex(), nullptr);
+		RHICmdList.SetUAVParameter(ShaderRHI, DepthOutTexture.GetBaseIndex(), nullptr);
 	}
 
 
@@ -2307,7 +2426,7 @@ public:
 	{
 		bool bShaderHasOutdatedParameters = FGlobalShader::Serialize(Ar);
 		Ar << DepthHistoryTexture << DepthHistoryTextureSampler << HistoryTexture << HistoryTextureSampler << PostprocessParameter << SceneTextureParameters << PrevScreenPositionScaleBias << OutTexture;
-		Ar << SceneVelocityTexture << SceneVelocityTextureSampler << BlendParams;
+		Ar << SceneVelocityTexture << SceneVelocityTextureSampler << BlendParams << DepthOutTexture;
 		return bShaderHasOutdatedParameters;
 	}
 
@@ -2323,13 +2442,15 @@ public:
 };
 
 IMPLEMENT_SHADER_TYPE2(FPostProcessGTAOTemporalFilterPSandCS<false>, SF_Pixel); 
-IMPLEMENT_SHADER_TYPE2(FPostProcessGTAOTemporalFilterPSandCS<true>, SF_Compute);
+IMPLEMENT_SHADER_TYPE2(FPostProcessGTAOTemporalFilterPSandCS<true>,  SF_Compute);
 
 
 FRCPassPostProcessAmbientOcclusion_GTAO_TemporalFilter::FRCPassPostProcessAmbientOcclusion_GTAO_TemporalFilter(const FSceneView& View, uint32 InDownScaleFactor,
 	const FGTAOTAAHistory& InInputHistory,
-	FGTAOTAAHistory* OutOutputHistory)
-	: InputHistory(InInputHistory),
+	FGTAOTAAHistory* OutOutputHistory,
+	const EGTAOType InAOType)
+	: AOType(InAOType),
+	InputHistory(InInputHistory),
 	OutputHistory(OutOutputHistory),
 	DownScaleFactor(InDownScaleFactor)
 {
@@ -2343,7 +2464,6 @@ void FRCPassPostProcessAmbientOcclusion_GTAO_TemporalFilter::Process(FRenderingC
 	const FPooledRenderTargetDesc* InputDesc0 = GetInputDesc(ePId_Input0);
 	FSceneRenderTargets& SceneContext = FSceneRenderTargets::Get(Context.RHICmdList);
 
-
 	const FSceneRenderTargetItem& DestRenderTarget0 = (DownScaleFactor==1) ? SceneContext.ScreenSpaceAO->GetRenderTargetItem() : PassOutputs[0].RequestSurface(Context);
 	const FSceneRenderTargetItem& DestRenderTarget1 = PassOutputs[1].RequestSurface(Context);
 
@@ -2363,7 +2483,25 @@ void FRCPassPostProcessAmbientOcclusion_GTAO_TemporalFilter::Process(FRenderingC
 	OutputHistory->ViewportRect = OutputViewRect;
 	OutputHistory->ReferenceBufferSize = OutputTexSize;
 
-	
+	if (1)
+	{
+		// Compute  Version
+		UnbindRenderTargets(Context.RHICmdList);
+		TShaderMapRef<FPostProcessGTAOTemporalFilterPSandCS<true>> ComputeShader(Context.GetShaderMap());
+		uint32 GroupSizeX = FMath::DivideAndRoundUp(OutputViewRect.Width(),  8);
+		uint32 GroupSizeY = FMath::DivideAndRoundUp(OutputViewRect.Height(), 8);
+
+		Context.RHICmdList.TransitionResource(EResourceTransitionAccess::ERWBarrier, EResourceTransitionPipeline::EComputeToCompute, DestRenderTarget0.UAV, nullptr);
+		Context.RHICmdList.TransitionResource(EResourceTransitionAccess::ERWBarrier, EResourceTransitionPipeline::EComputeToCompute, DestRenderTarget1.UAV, nullptr);
+
+		Context.RHICmdList.SetComputeShader(ComputeShader->GetComputeShader());
+		ComputeShader->SetParametersCS(Context.RHICmdList, Context, OutputTexSize, InputTexSize, bCameraCut, InputHistory,
+			bCameraCut ? GSystemTextures.BlackDummy : SceneContext.SceneVelocity, DestRenderTarget0.UAV, DestRenderTarget1.UAV);
+
+		DispatchComputeShader(Context.RHICmdList, *ComputeShader, GroupSizeX, GroupSizeY, 1);
+		ComputeShader->UnsetParameters(Context.RHICmdList);
+	}
+	else
 	{
 		//FRHIRenderPassInfo RPInfo(DestRenderTarget.TargetableTexture, ERenderTargetActions::Load_Store);
 		FRHITexture* RenderTargets[2] =
@@ -2383,7 +2521,6 @@ void FRCPassPostProcessAmbientOcclusion_GTAO_TemporalFilter::Process(FRenderingC
 
 			DrawClearQuad(Context.RHICmdList, FLinearColor::White);
 
-
 			FGraphicsPipelineStateInitializer GraphicsPSOInit;
 			Context.RHICmdList.ApplyCachedRenderTargets(GraphicsPSOInit);
 
@@ -2417,9 +2554,13 @@ void FRCPassPostProcessAmbientOcclusion_GTAO_TemporalFilter::Process(FRenderingC
 
 		}
 		Context.RHICmdList.EndRenderPass();
+
+		Context.RHICmdList.TransitionResource(EResourceTransitionAccess::EReadable, DestRenderTarget0.TargetableTexture);
+		Context.RHICmdList.TransitionResource(EResourceTransitionAccess::EReadable, DestRenderTarget1.TargetableTexture);
 	}
-	Context.RHICmdList.TransitionResource(EResourceTransitionAccess::EReadable, DestRenderTarget0.TargetableTexture);
-	Context.RHICmdList.TransitionResource(EResourceTransitionAccess::EReadable, DestRenderTarget1.TargetableTexture);
+
+
+
 }
 
 
@@ -2438,7 +2579,6 @@ FPooledRenderTargetDesc FRCPassPostProcessAmbientOcclusion_GTAO_TemporalFilter::
 		Ret.Format = PF_R32_FLOAT;
 	}
 
-
 	Ret.ClearValue = FClearValueBinding::None;
 	Ret.TargetableFlags &= ~TexCreate_DepthStencilTargetable;
 	Ret.TargetableFlags |= TexCreate_UAV;
@@ -2558,8 +2698,9 @@ public:
 IMPLEMENT_SHADER_TYPE2(FPostProcessGTAOSpatialFilterPSandCS<true>,  SF_Compute);
 IMPLEMENT_SHADER_TYPE2(FPostProcessGTAOSpatialFilterPSandCS<false>, SF_Pixel);
 
-FRCPassPostProcessAmbientOcclusion_GTAO_SpatialFilter::FRCPassPostProcessAmbientOcclusion_GTAO_SpatialFilter(const FSceneView& View, uint32 InDownScaleFactor)
-	:DownScaleFactor(InDownScaleFactor)
+FRCPassPostProcessAmbientOcclusion_GTAO_SpatialFilter::FRCPassPostProcessAmbientOcclusion_GTAO_SpatialFilter(const FSceneView& View, uint32 InDownScaleFactor, const EGTAOType InAOType)
+	: AOType(InAOType),
+	DownScaleFactor(InDownScaleFactor)
 {
 }
 
@@ -2577,20 +2718,33 @@ void FRCPassPostProcessAmbientOcclusion_GTAO_SpatialFilter::Process(FRenderingCo
 	FIntRect InputViewRect = View.ViewRect / DownScaleFactor;
 	FIntRect OutputViewRect = InputViewRect;
 
-	const FSceneRenderTargetItem& DestRenderTarget = PassOutputs[0].RequestSurface(Context);
-
 	// Compute  Version
-	Context.RHICmdList.TransitionResource(EResourceTransitionAccess::ERWBarrier, EResourceTransitionPipeline::EGfxToCompute, DestRenderTarget.UAV, nullptr);
-
 	UnbindRenderTargets(Context.RHICmdList);
-
 	TShaderMapRef<FPostProcessGTAOSpatialFilterPSandCS<true>> ComputeShader(Context.GetShaderMap());
-	Context.RHICmdList.SetComputeShader(ComputeShader->GetComputeShader());
-	ComputeShader->SetParametersCS(Context.RHICmdList, Context, InputTexSize, OutputViewRect, DestRenderTarget.UAV);
-
 	uint32 GroupSizeX = FMath::DivideAndRoundUp(OutputViewRect.Width(), 8);
 	uint32 GroupSizeY = FMath::DivideAndRoundUp(OutputViewRect.Height(), 8);
-	DispatchComputeShader(Context.RHICmdList, *ComputeShader, GroupSizeX, GroupSizeY, 1);
+
+	if (AOType == EGTAOType::EAsyncCombinedSpatial)
+	{
+		// If the Spatial Filter is running as part of the async then we'll render to the R channel of the horizons texture so it can be read in as part of the temporal
+		const FSceneRenderTargetItem& DestRenderTarget = SceneContext.ScreenSpaceGTAOHorizons->GetRenderTargetItem();
+
+		FRHIAsyncComputeCommandListImmediate& RHICmdListComputeImmediate = FRHICommandListExecutor::GetImmediateAsyncComputeCommandList();
+
+		RHICmdListComputeImmediate.SetComputeShader(ComputeShader->GetComputeShader());
+		RHICmdListComputeImmediate.TransitionResource(EResourceTransitionAccess::ERWBarrier, EResourceTransitionPipeline::EComputeToCompute, DestRenderTarget.UAV, nullptr);
+		ComputeShader->SetParametersCS(RHICmdListComputeImmediate, Context, InputTexSize, OutputViewRect, DestRenderTarget.UAV);
+		DispatchComputeShader(RHICmdListComputeImmediate, *ComputeShader, GroupSizeX, GroupSizeY, 1);
+	}
+	else
+	{
+		const FSceneRenderTargetItem& DestRenderTarget = PassOutputs[0].RequestSurface(Context);
+
+		Context.RHICmdList.SetComputeShader(ComputeShader->GetComputeShader());
+		Context.RHICmdList.TransitionResource(EResourceTransitionAccess::ERWBarrier, EResourceTransitionPipeline::EGfxToCompute, DestRenderTarget.UAV, nullptr);
+		ComputeShader->SetParametersCS(Context.RHICmdList, Context, InputTexSize, OutputViewRect, DestRenderTarget.UAV);
+		DispatchComputeShader(Context.RHICmdList, *ComputeShader, GroupSizeX, GroupSizeY, 1);
+	}
 
 	ComputeShader->UnsetParameters(Context.RHICmdList);
 }
@@ -2723,11 +2877,14 @@ public:
 IMPLEMENT_SHADER_TYPE2(FPostProcessGTAOUpsamplePSandCS<true>, SF_Compute);
 IMPLEMENT_SHADER_TYPE2(FPostProcessGTAOUpsamplePSandCS<false>, SF_Pixel);
 
-FRCPassPostProcessAmbientOcclusion_GTAO_Upsample::FRCPassPostProcessAmbientOcclusion_GTAO_Upsample(const FSceneView& View, uint32 InDownScaleFactor)
-	:DownScaleFactor(InDownScaleFactor)
+FRCPassPostProcessAmbientOcclusion_GTAO_Upsample::FRCPassPostProcessAmbientOcclusion_GTAO_Upsample(const FSceneView& View, uint32 InDownScaleFactor, const EGTAOType InAOType)
+	:AOType(InAOType),
+	 DownScaleFactor(InDownScaleFactor)
 {
 }
 
+
+
 void FRCPassPostProcessAmbientOcclusion_GTAO_Upsample::Process(FRenderingCompositePassContext& Context)
 {
 	SCOPED_GPU_STAT(Context.RHICmdList, GTAO_Upsample);
@@ -2740,22 +2897,61 @@ void FRCPassPostProcessAmbientOcclusion_GTAO_Upsample::Process(FRenderingComposi
 	FIntPoint InputTexSize = InputDesc0->Extent;	// This is the full res Z
 	FIntPoint OutputTexSize = SceneContext.GetBufferSizeXY();
 
-	// Compute  Version
-	Context.RHICmdList.TransitionResource(EResourceTransitionAccess::ERWBarrier, EResourceTransitionPipeline::EGfxToCompute, DestRenderTarget.UAV, nullptr);
+	if (0)
+	{
+		UnbindRenderTargets(Context.RHICmdList);
+		TShaderMapRef<FPostProcessGTAOUpsamplePSandCS<true>> ComputeShader(Context.GetShaderMap());
 
-	UnbindRenderTargets(Context.RHICmdList);
+		uint32 GroupSizeX = FMath::DivideAndRoundUp(View.ViewRect.Width(), 8);
+		uint32 GroupSizeY = FMath::DivideAndRoundUp(View.ViewRect.Height(), 8);
 
-	TShaderMapRef<FPostProcessGTAOUpsamplePSandCS<true>> ComputeShader(Context.GetShaderMap());
-	Context.RHICmdList.SetComputeShader(ComputeShader->GetComputeShader());
-	ComputeShader->SetParametersCS(Context.RHICmdList, Context, InputTexSize, DestRenderTarget.UAV);
+		Context.RHICmdList.TransitionResource(EResourceTransitionAccess::ERWBarrier, EResourceTransitionPipeline::EComputeToCompute, DestRenderTarget.UAV, nullptr);
+		Context.RHICmdList.SetComputeShader(ComputeShader->GetComputeShader());
+		ComputeShader->SetParametersCS(Context.RHICmdList, Context, InputTexSize, DestRenderTarget.UAV);
+		DispatchComputeShader(Context.RHICmdList, *ComputeShader, GroupSizeX, GroupSizeY, 1);
+		ComputeShader->UnsetParameters(Context.RHICmdList);
+		Context.RHICmdList.TransitionResource(EResourceTransitionAccess::EReadable, DestRenderTarget.TargetableTexture);
+	}
+	else
+	{
+		FIntRect ViewRect = View.ViewRect;
 
-	uint32 GroupSizeX = FMath::DivideAndRoundUp(View.ViewRect.Width(), 8);
-	uint32 GroupSizeY = FMath::DivideAndRoundUp(View.ViewRect.Height(), 8);
-	DispatchComputeShader(Context.RHICmdList, *ComputeShader, GroupSizeX, GroupSizeY, 1);
+		FRHIRenderPassInfo RPInfo(DestRenderTarget.TargetableTexture, ERenderTargetActions::Load_Store);
+		Context.RHICmdList.BeginRenderPass(RPInfo, TEXT("GTAOUpsample"));
+		{
+			Context.SetViewportAndCallRHI(ViewRect);
 
-	ComputeShader->UnsetParameters(Context.RHICmdList);
+			TShaderMapRef<FPostProcessGTAOUpsamplePSandCS<false> > PixelShader(Context.GetShaderMap());
+			TShaderMapRef<FPostProcessVS> VertexShader(Context.GetShaderMap());
 
-	Context.RHICmdList.TransitionResource(EResourceTransitionAccess::EReadable, DestRenderTarget.TargetableTexture);
+			FGraphicsPipelineStateInitializer GraphicsPSOInit;
+			Context.RHICmdList.ApplyCachedRenderTargets(GraphicsPSOInit);
+
+			// set the state
+			GraphicsPSOInit.BlendState			= TStaticBlendState<>::GetRHI();
+			GraphicsPSOInit.RasterizerState		= TStaticRasterizerState<>::GetRHI();
+			GraphicsPSOInit.DepthStencilState	= TStaticDepthStencilState<false, CF_Always>::GetRHI();
+			GraphicsPSOInit.PrimitiveType		= PT_TriangleList;
+			GraphicsPSOInit.bDepthBounds		= false;
+
+			GraphicsPSOInit.BoundShaderState.VertexDeclarationRHI = GFilterVertexDeclaration.VertexDeclarationRHI;
+			GraphicsPSOInit.BoundShaderState.VertexShaderRHI = GETSAFERHISHADER_VERTEX(*VertexShader);
+			GraphicsPSOInit.BoundShaderState.PixelShaderRHI  = GETSAFERHISHADER_PIXEL(*PixelShader);
+			SetGraphicsPipelineState(Context.RHICmdList, GraphicsPSOInit);
+
+			PixelShader->SetParametersPS(Context, InputTexSize);
+
+			DrawRectangle(
+				Context.RHICmdList,
+				0, 0, ViewRect.Width(), ViewRect.Height(),
+				ViewRect.Min.X, ViewRect.Min.Y, ViewRect.Width(), ViewRect.Height(),
+				ViewRect.Size(),
+				OutputTexSize,
+				*VertexShader,
+				EDRF_UseTriangleOptimization);
+		}
+		Context.RHICmdList.EndRenderPass();
+	}
 }
 
 FPooledRenderTargetDesc FRCPassPostProcessAmbientOcclusion_GTAO_Upsample::ComputeOutputDesc(EPassOutputId InPassOutputId) const
diff --git a/Engine/Source/Runtime/Renderer/Private/CompositionLighting/PostProcessAmbientOcclusion.h b/Engine/Source/Runtime/Renderer/Private/CompositionLighting/PostProcessAmbientOcclusion.h
--- a/Engine/Source/Runtime/Renderer/Private/CompositionLighting/PostProcessAmbientOcclusion.h
+++ b/Engine/Source/Runtime/Renderer/Private/CompositionLighting/PostProcessAmbientOcclusion.h
@@ -30,14 +30,16 @@ enum class EGTAOType
 	// Not on (use legacy if at all)
 	EOff,
 
-	// non async compute shader where the Horizon Search and Inner Integrate are in separate passes
-	ESplitNonAsync,
+	// Async compute shader where the Horizon Search and Inner Integrate are combined and the spatial filter is run on the Async Pipe
+	// Temporal and Upsample are run on the GFX Pipe as Temporal requires velocity Buffer
+	EAsyncCombinedSpatial,
 
-	// Async compute shader where the Horizon Search and Inner Integrate are in separate passes
-	ESplitAsync,
+	// Async compute shader where the Horizon Search is run on the Async compute pipe
+	// Integrate, Spatial, Temporal and Upsample are run on the GFX pipe as these require GBuffer channels
+	EAsyncHorizonSearch,
 
-	// Combined version which has to be run after base pass
-	ECombinedNonAsync,
+	// Non async version where all passes are run on the GFX Pipe
+	ENonAsync,
 };
 
 class FSSAOHelper
@@ -148,7 +150,7 @@ class FRCPassPostProcessAmbientOcclusion_HorizonSearch : public TRenderingCompos
 {
 public:
 
-	FRCPassPostProcessAmbientOcclusion_HorizonSearch(const FSceneView& View, uint32 DownScaleFactor, const ESSAOType AOType );
+	FRCPassPostProcessAmbientOcclusion_HorizonSearch(const FSceneView& View, uint32 DownScaleFactor, const EGTAOType AOType );
 
 	// interface FRenderingCompositePass ---------
 	virtual void Process(FRenderingCompositePassContext& Context) override;
@@ -162,16 +164,16 @@ public:
 	FShader* SetShaderPS(const FRenderingCompositePassContext& Context, FGraphicsPipelineStateInitializer& GraphicsPSOInit, FIntPoint DestSize);
 
 private:
-	const ESSAOType AOType;
+	const EGTAOType AOType;
 	uint32		DownScaleFactor;
 };
 
 
-class FRCPassPostProcessAmbientOcclusion_InnerIntegrate : public TRenderingCompositePassBase<2, 1>
+class FRCPassPostProcessAmbientOcclusion_GTAOInnerIntegrate : public TRenderingCompositePassBase<2, 1>
 {
 public:
 
-	FRCPassPostProcessAmbientOcclusion_InnerIntegrate(const FSceneView& View, uint32 DownScaleFactor, bool FinalOutput);
+	FRCPassPostProcessAmbientOcclusion_GTAOInnerIntegrate(const FSceneView& View, uint32 DownScaleFactor, bool FinalOutput);
 
 	// interface FRenderingCompositePass ---------
 	virtual void Process(FRenderingCompositePassContext& Context) override;
@@ -183,25 +185,25 @@ private:
 	uint32		DownScaleFactor;
 };
 
-
-class FRCPassPostProcessAmbientOcclusion_GTAOCombined : public TRenderingCompositePassBase<2, 2>
+class FRCPassPostProcessAmbientOcclusion_GTAOHorizonSearchIntegrate : public TRenderingCompositePassBase<2, 2>
 {
 public:
 
-	FRCPassPostProcessAmbientOcclusion_GTAOCombined(const FSceneView& View, uint32 DownScaleFactor, bool FinalOutput);
+	FRCPassPostProcessAmbientOcclusion_GTAOHorizonSearchIntegrate(const FSceneView& View, uint32 DownScaleFactor, bool FinalOutput, const EGTAOType AOType);
 
 	// interface FRenderingCompositePass ---------
 	virtual void Process(FRenderingCompositePassContext& Context) override;
 	virtual void Release() override { delete this; }
 	virtual FPooledRenderTargetDesc ComputeOutputDesc(EPassOutputId InPassOutputId) const override;
 
-	template <uint32 ShaderQuality, uint32 UseNormals>
-	void DispatchCS(const FRenderingCompositePassContext& Context, FIntRect ViewRect, FIntPoint DestSize, FIntPoint TexSize);
+	template <uint32 ShaderQuality, uint32 UseNormals, typename TRHICmdList>
+	void DispatchCS(TRHICmdList& RHICmdList, const FRenderingCompositePassContext& Context, FIntRect ViewRect, FIntPoint DestSize, FIntPoint TexSize);
 
 	template <uint32 ShaderQuality, uint32 UseNormals>
 	FShader* SetShaderPS(const FRenderingCompositePassContext& Context, FGraphicsPipelineStateInitializer& GraphicsPSOInit, FIntPoint DestSize);
 
 private:
+	const EGTAOType AOType;
 	const bool bFinalOutput;
 	uint32		DownScaleFactor;
 };
@@ -213,7 +215,7 @@ class FRCPassPostProcessAmbientOcclusion_GTAO_TemporalFilter : public TRendering
 {
 public:
 
-	FRCPassPostProcessAmbientOcclusion_GTAO_TemporalFilter(const FSceneView& View, uint32 DownScaleFactor, const FGTAOTAAHistory& InInputHistory, FGTAOTAAHistory* OutOutputHistory);
+	FRCPassPostProcessAmbientOcclusion_GTAO_TemporalFilter(const FSceneView& View, uint32 DownScaleFactor, const FGTAOTAAHistory& InInputHistory, FGTAOTAAHistory* OutOutputHistory, const EGTAOType AOType);
 
 	// interface FRenderingCompositePass ---------
 	virtual void Process(FRenderingCompositePassContext& Context) override;
@@ -224,7 +226,7 @@ public:
 	void DispatchCS(const FRenderingCompositePassContext& Context, FIntRect OutputViewRect, FIntPoint OutputTexSize);
 
 private:
-
+	const EGTAOType AOType;
 	const FGTAOTAAHistory& InputHistory;
 	FGTAOTAAHistory* OutputHistory;
 	uint32		DownScaleFactor;
@@ -234,7 +236,7 @@ class FRCPassPostProcessAmbientOcclusion_GTAO_SpatialFilter : public TRenderingC
 {
 public:
 
-	FRCPassPostProcessAmbientOcclusion_GTAO_SpatialFilter(const FSceneView& View, uint32 DownScaleFactor);
+	FRCPassPostProcessAmbientOcclusion_GTAO_SpatialFilter(const FSceneView& View, uint32 DownScaleFactor, const EGTAOType AOType);
 
 	// interface FRenderingCompositePass ---------
 	virtual void Process(FRenderingCompositePassContext& Context) override;
@@ -242,15 +244,15 @@ public:
 	virtual FPooledRenderTargetDesc ComputeOutputDesc(EPassOutputId InPassOutputId) const override;
 
 private:
+	const EGTAOType AOType;
 	uint32		DownScaleFactor;
-
 };
 
 class FRCPassPostProcessAmbientOcclusion_GTAO_Upsample : public TRenderingCompositePassBase<2, 1>
 {
 public:
 
-	FRCPassPostProcessAmbientOcclusion_GTAO_Upsample(const FSceneView& View, uint32 DownScaleFactor);
+	FRCPassPostProcessAmbientOcclusion_GTAO_Upsample(const FSceneView& View, uint32 DownScaleFactor, const EGTAOType AOType);
 
 	// interface FRenderingCompositePass ---------
 	virtual void Process(FRenderingCompositePassContext& Context) override;
@@ -258,6 +260,6 @@ public:
 	virtual FPooledRenderTargetDesc ComputeOutputDesc(EPassOutputId InPassOutputId) const override;
 
 private:
-	uint32		DownScaleFactor;
-
+	const EGTAOType AOType;
+	uint32			DownScaleFactor;
 };
\ No newline at end of file
