Fix water shading in multi viewport rendering.

Color and depth downsample now happens in multiple passes to avoid downscaling across views.

#rnx

diff --git a/Engine/Shaders/Private/SingleLayerWaterComposite.usf b/Engine/Shaders/Private/SingleLayerWaterComposite.usf
--- a/Engine/Shaders/Private/SingleLayerWaterComposite.usf
+++ b/Engine/Shaders/Private/SingleLayerWaterComposite.usf
@@ -257,7 +257,7 @@ void WaterTileVS(
 	uint TileData = TileListData[InstanceId.x];
 	uint2 TileOrigin = uint2((TileData & 0x0000FFFF) * WORK_TILE_SIZE, (TileData >> 16) * WORK_TILE_SIZE);
 
-	uint2 TileVertex = TileOrigin + View.ViewRectMin.xy;
+	uint2 TileVertex = TileOrigin;
 	TileVertex.x += VertexId == 1 || VertexId == 2 || VertexId == 4 ? WORK_TILE_SIZE : 0;
 	TileVertex.y += VertexId == 2 || VertexId == 4 || VertexId == 5 ? WORK_TILE_SIZE : 0;
 
diff --git a/Engine/Source/Runtime/Renderer/Private/DeferredShadingRenderer.cpp b/Engine/Source/Runtime/Renderer/Private/DeferredShadingRenderer.cpp
--- a/Engine/Source/Runtime/Renderer/Private/DeferredShadingRenderer.cpp
+++ b/Engine/Source/Runtime/Renderer/Private/DeferredShadingRenderer.cpp
@@ -2033,6 +2033,7 @@ void FDeferredShadingSceneRenderer::Render(FRHICommandListImmediate& RHICmdList)
 
 		// Copy the texture to be available for the water surface to refrace
 		FSingleLayerWaterPassData SingleLayerWaterPassData;
+		SingleLayerWaterPassData.ViewData.SetNum(Views.Num());
 		CopySingleLayerWaterTextures(RHICmdList, SingleLayerWaterPassData);
 
 		// Make the Depth texture writable since the water GBuffer pass will update it
diff --git a/Engine/Source/Runtime/Renderer/Private/DeferredShadingRenderer.h b/Engine/Source/Runtime/Renderer/Private/DeferredShadingRenderer.h
--- a/Engine/Source/Runtime/Renderer/Private/DeferredShadingRenderer.h
+++ b/Engine/Source/Runtime/Renderer/Private/DeferredShadingRenderer.h
@@ -118,7 +118,7 @@ public:
 	*/
 	bool RenderSingleLayerWaterPass(FRHICommandListImmediate& RHICmdList, FSingleLayerWaterPassData& PassData, FExclusiveDepthStencil::Type WaterPassDepthStencilAccess, bool bParallel);
 	/** Renders the water draw pass for a given View. */
-	bool RenderSingleLayerWaterPassView(FRHICommandListImmediate& RHICmdList, FViewInfo& View, FSingleLayerWaterPassData& PassData, const FMeshPassProcessorRenderState& InDrawRenderState, bool bParallel);
+	bool RenderSingleLayerWaterPassView(FRHICommandListImmediate& RHICmdList, FViewInfo& View, const FMeshPassProcessorRenderState& InDrawRenderState, bool bParallel);
 	/** Render, denoise and composite the scene SSR and under water effect.*/
 	void RenderSingleLayerWaterReflections(FRHICommandListImmediate& RHICmdList, FSingleLayerWaterPassData& PassData);
 
diff --git a/Engine/Source/Runtime/Renderer/Private/SingleLayerWaterRendering.cpp b/Engine/Source/Runtime/Renderer/Private/SingleLayerWaterRendering.cpp
--- a/Engine/Source/Runtime/Renderer/Private/SingleLayerWaterRendering.cpp
+++ b/Engine/Source/Runtime/Renderer/Private/SingleLayerWaterRendering.cpp
@@ -462,41 +462,43 @@ void FDeferredShadingSceneRenderer::CopySingleLayerWaterTextures(FRHICommandList
 	const FRDGTextureDesc DepthDesc(FRDGTextureDesc::Create2DDesc(RefractionResolution, GSingleLayerWaterRefractionFullPrecision ? PF_R32_FLOAT : PF_R16F, SceneContext.GetDefaultColorClear(), TexCreate_None, TexCreate_ShaderResource | TexCreate_RenderTargetable, false));
 	FRDGTextureRef SceneDepthWithoutSingleLayerWaterTexture = GraphBuilder.CreateTexture(DepthDesc, TEXT("SceneDepthWithoutSingleLayerWater"));
 
-	// For now support only the 1st view
-	const FViewInfo& View = Views[0];
-
-	FWaterRefractionCopyPS::FParameters* PassParameters = GraphBuilder.AllocParameters<FWaterRefractionCopyPS::FParameters>();
-	PassParameters->ViewUniformBuffer = View.ViewUniformBuffer;
-	PassParameters->SceneColorCopyDownsampleTexture = GraphBuilder.RegisterExternalTexture(SceneContext.GetSceneColor());
-	PassParameters->SceneColorCopyDownsampleSampler = TStaticSamplerState<SF_Bilinear, AM_Clamp, AM_Clamp, AM_Clamp>::GetRHI();
-	PassParameters->SceneDepthCopyDownsampleTexture = GraphBuilder.RegisterExternalTexture(SceneContext.SceneDepthZ);
-	PassParameters->SceneDepthCopyDownsampleSampler = TStaticSamplerState<SF_Bilinear, AM_Clamp, AM_Clamp, AM_Clamp>::GetRHI();
-	PassParameters->SVPositionToSourceTextureUV = FVector2D(RefractionDownsampleFactor / float(SceneContext.GetBufferSizeXY().X), RefractionDownsampleFactor / float(SceneContext.GetBufferSizeXY().Y));
-
-	PassParameters->RenderTargets[0] = FRenderTargetBinding(SceneDepthWithoutSingleLayerWaterTexture, ERenderTargetLoadAction::ENoAction);
-
-	if (bCopyColor)
+	for (int32 ViewIndex = 0; ViewIndex < Views.Num(); ViewIndex++)
 	{
-		PassParameters->RenderTargets[1] = FRenderTargetBinding(SceneColorWithoutSingleLayerWaterTexture, ERenderTargetLoadAction::ENoAction);
-	}
-
-	FWaterRefractionCopyPS::FPermutationDomain PermutationVector;
-	PermutationVector.Set<FWaterRefractionCopyPS::FDownsampleRefraction>(RefractionDownsampleFactor > 1);
-	PermutationVector.Set<FWaterRefractionCopyPS::FDownsampleColor>(bCopyColor);
-	auto PixelShader = View.ShaderMap->GetShader<FWaterRefractionCopyPS>(PermutationVector);
-
-	const FIntRect RefractionViewRect = FIntRect(FIntPoint(0, 0), FIntPoint::DivideAndRoundDown(View.ViewRect.Size(), RefractionDownsampleFactor));
-
-	PassData.SceneWithoutSingleLayerWaterMaxUV.X = (RefractionViewRect.Max.X - 0.5f) / RefractionResolution.X;
-	PassData.SceneWithoutSingleLayerWaterMaxUV.Y = (RefractionViewRect.Max.Y - 0.5f) / RefractionResolution.Y;
-
-	FPixelShaderUtils::AddFullscreenPass(
-		GraphBuilder,
-		View.ShaderMap,
-		RDG_EVENT_NAME("Water Refraction Copy"),
-		PixelShader,
-		PassParameters,
-		RefractionViewRect);
+		const FViewInfo& View = Views[ViewIndex];
+
+	    FWaterRefractionCopyPS::FParameters* PassParameters = GraphBuilder.AllocParameters<FWaterRefractionCopyPS::FParameters>();
+	    PassParameters->ViewUniformBuffer = View.ViewUniformBuffer;
+	    PassParameters->SceneColorCopyDownsampleTexture = GraphBuilder.RegisterExternalTexture(SceneContext.GetSceneColor());
+	    PassParameters->SceneColorCopyDownsampleSampler = TStaticSamplerState<SF_Bilinear, AM_Clamp, AM_Clamp, AM_Clamp>::GetRHI();
+	    PassParameters->SceneDepthCopyDownsampleTexture = GraphBuilder.RegisterExternalTexture(SceneContext.SceneDepthZ);
+	    PassParameters->SceneDepthCopyDownsampleSampler = TStaticSamplerState<SF_Bilinear, AM_Clamp, AM_Clamp, AM_Clamp>::GetRHI();
+	    PassParameters->SVPositionToSourceTextureUV = FVector2D(RefractionDownsampleFactor / float(SceneContext.GetBufferSizeXY().X), RefractionDownsampleFactor / float(SceneContext.GetBufferSizeXY().Y));
+    
+	    PassParameters->RenderTargets[0] = FRenderTargetBinding(SceneDepthWithoutSingleLayerWaterTexture, ERenderTargetLoadAction::ENoAction);
+    
+	    if (bCopyColor)
+	    {
+		    PassParameters->RenderTargets[1] = FRenderTargetBinding(SceneColorWithoutSingleLayerWaterTexture, ERenderTargetLoadAction::ENoAction);
+	    }
+    
+	    FWaterRefractionCopyPS::FPermutationDomain PermutationVector;
+	    PermutationVector.Set<FWaterRefractionCopyPS::FDownsampleRefraction>(RefractionDownsampleFactor > 1);
+	    PermutationVector.Set<FWaterRefractionCopyPS::FDownsampleColor>(bCopyColor);
+	    auto PixelShader = View.ShaderMap->GetShader<FWaterRefractionCopyPS>(PermutationVector);
+    
+		const FIntRect RefractionViewRect = FIntRect(FIntPoint::DivideAndRoundDown(View.ViewRect.Min, RefractionDownsampleFactor), FIntPoint::DivideAndRoundDown(View.ViewRect.Max, RefractionDownsampleFactor));
+
+		PassData.ViewData[ViewIndex].SceneWithoutSingleLayerWaterMaxUV.X = (RefractionViewRect.Size().X - 0.5f) / RefractionViewRect.Size().X;
+		PassData.ViewData[ViewIndex].SceneWithoutSingleLayerWaterMaxUV.Y = (RefractionViewRect.Size().Y - 0.5f) / RefractionViewRect.Size().Y;
+
+	    FPixelShaderUtils::AddFullscreenPass(
+		    GraphBuilder,
+		    View.ShaderMap,
+		    RDG_EVENT_NAME("Water Refraction Copy"),
+		    PixelShader,
+		    PassParameters,
+		    RefractionViewRect);
+    }
 
 	GraphBuilder.QueueTextureExtraction(SceneDepthWithoutSingleLayerWaterTexture, &PassData.SceneDepthWithoutSingleLayerWater, true);
 
@@ -617,7 +619,7 @@ void FDeferredShadingSceneRenderer::RenderSingleLayerWaterReflections(FRHIComman
 			Parameters.PreIntegratedGFSampler = TStaticSamplerState<SF_Bilinear, AM_Clamp, AM_Clamp, AM_Clamp>::GetRHI();
 			Parameters.SceneNoWaterDepthTexture = GraphBuilder.RegisterExternalTexture(PassData.SceneDepthWithoutSingleLayerWater ? PassData.SceneDepthWithoutSingleLayerWater : GSystemTextures.BlackDummy);
 			Parameters.SceneNoWaterDepthSampler = TStaticSamplerState<SF_Point>::GetRHI();
-			Parameters.SceneNoWaterMaxUV = PassData.SceneWithoutSingleLayerWaterMaxUV;
+			Parameters.SceneNoWaterMaxUV = PassData.ViewData[ViewIndex].SceneWithoutSingleLayerWaterMaxUV;
 			Parameters.SceneTextures = SceneTextures;
 			SetupSceneTextureSamplers(&Parameters.SceneTextureSamplers);
 			Parameters.ViewUniformBuffer = View.ViewUniformBuffer;
@@ -820,7 +822,7 @@ bool FDeferredShadingSceneRenderer::RenderSingleLayerWaterPass(FRHICommandListIm
 			CreateOpaqueBasePassUniformBuffer(RHICmdList, 
 				View, 
 				WhiteForwardScreenSpaceShadowMask, 
-				&PassData.SceneWithoutSingleLayerWaterMaxUV,
+				&PassData.ViewData[ViewIndex].SceneWithoutSingleLayerWaterMaxUV,
 				PassData.SceneColorWithoutSingleLayerWater.IsValid() ? PassData.SceneColorWithoutSingleLayerWater : GSystemTextures.BlackDummy,
 				PassData.SceneDepthWithoutSingleLayerWater,
 				WaterPassUniformBuffer);
@@ -833,7 +835,7 @@ bool FDeferredShadingSceneRenderer::RenderSingleLayerWaterPass(FRHICommandListIm
 			{
 				Scene->UniformBuffers.UpdateViewUniformBuffer(View);
 
-				bDirty |= RenderSingleLayerWaterPassView(RHICmdList, View, PassData, DrawRenderState, bParallel);
+				bDirty |= RenderSingleLayerWaterPassView(RHICmdList, View, DrawRenderState, bParallel);
 			}
 		}
 	}
@@ -844,7 +846,7 @@ bool FDeferredShadingSceneRenderer::RenderSingleLayerWaterPass(FRHICommandListIm
 	return bDirty;
 }
 
-bool FDeferredShadingSceneRenderer::RenderSingleLayerWaterPassView(FRHICommandListImmediate& RHICmdList, FViewInfo& View, FSingleLayerWaterPassData& PassData, const FMeshPassProcessorRenderState& InDrawRenderState, bool bParallel)
+bool FDeferredShadingSceneRenderer::RenderSingleLayerWaterPassView(FRHICommandListImmediate& RHICmdList, FViewInfo& View, const FMeshPassProcessorRenderState& InDrawRenderState, bool bParallel)
 {
 	if (!bParallel)
 	{
diff --git a/Engine/Source/Runtime/Renderer/Private/SingleLayerWaterRendering.h b/Engine/Source/Runtime/Renderer/Private/SingleLayerWaterRendering.h
--- a/Engine/Source/Runtime/Renderer/Private/SingleLayerWaterRendering.h
+++ b/Engine/Source/Runtime/Renderer/Private/SingleLayerWaterRendering.h
@@ -17,7 +17,13 @@ struct FSingleLayerWaterPassData
 {
 	TRefCountPtr<IPooledRenderTarget> SceneColorWithoutSingleLayerWater;
 	TRefCountPtr<IPooledRenderTarget> SceneDepthWithoutSingleLayerWater;
-	FVector2D SceneWithoutSingleLayerWaterMaxUV;
+
+	struct FSingleLayerWaterPassViewData
+	{
+		FVector2D SceneWithoutSingleLayerWaterMaxUV;
+	};
+
+	TArray<FSingleLayerWaterPassViewData> ViewData;
 };
 
 class FSingleLayerWaterPassMeshProcessor : public FMeshPassProcessor
