Prevent compiling mobile base pass pixel shaders for decal domain materials. Fixes crash in SCW.

diff --git a/Engine/Source/Runtime/Renderer/Private/MobileBasePassRendering.h b/Engine/Source/Runtime/Renderer/Private/MobileBasePassRendering.h
--- a/Engine/Source/Runtime/Renderer/Private/MobileBasePassRendering.h
+++ b/Engine/Source/Runtime/Renderer/Private/MobileBasePassRendering.h
@@ -394,7 +394,13 @@ class TMobileBasePassPS : public TMobileBasePassPSBaseType<LightMapPolicyType>
 public:
 
 	static bool ShouldCompilePermutation(const FMeshMaterialShaderPermutationParameters& Parameters)
-	{		
+	{	
+		// Don't compile base pass for decal domain materials
+		if (Parameters.Material->GetMaterialDomain() == MD_DeferredDecal && !IsPCPlatform(Parameters.Platform))
+		{
+			return false;
+		}
+
 		// We compile the point light shader combinations based on the project settings
 		static auto* MobileDynamicPointLightsUseStaticBranchCVar = IConsoleManager::Get().FindTConsoleVariableDataInt(TEXT("r.MobileDynamicPointLightsUseStaticBranch"));
 		static auto* MobileNumDynamicPointLightsCVar = IConsoleManager::Get().FindTConsoleVariableDataInt(TEXT("r.MobileNumDynamicPointLights"));
