Remove PS version of hair sky lighting.

diff --git a/Engine/Shaders/Private/HairStrands/HairStrandsEnvironmentLighting.usf b/Engine/Shaders/Private/HairStrands/HairStrandsEnvironmentLighting.usf
--- a/Engine/Shaders/Private/HairStrands/HairStrandsEnvironmentLighting.usf
+++ b/Engine/Shaders/Private/HairStrands/HairStrandsEnvironmentLighting.usf
@@ -66,8 +66,6 @@ float3 GatherRadiance(float3 RayDirection, float Roughness, uint2 PixelPosition,
 		EyeIndex);
 }
 
-Texture2D<uint4>			HairCategorizationTexture;
-Texture2D<uint>				HairVisibilityNodeOffsetAndCount;
 StructuredBuffer<FNodeData> HairVisibilityNodeData;
 
 Texture3D<float4>			HairScatteringLUTTexture;
@@ -83,9 +81,6 @@ float Voxel_TanConeAngle;
 float Voxel_DensityScale;
 float Voxel_DepthBiasScale;
 
-float AO_Power;
-float AO_Intensity;
-
 float3 Voxel_MinAABB;
 float3 Voxel_MaxAABB;
 uint   Voxel_Resolution;
@@ -412,106 +407,6 @@ float3 ReflectionEnvironment(FGBufferData GBuffer, float2 ScreenPosition, float2
 
 ////////////////////////////////////////////////////////////////////////////////////////////////////////////////////////////////////////////
 
-#if PIXELSHADER
-void MainPS(
-	in float4 SvPosition : SV_Position,
-	out float4 OutColor : SV_Target0,
-	out float4 OutColorSubPixel : SV_Target1
-)
-{
-	const float2 ScreenPosition = SvPositionToScreenPosition(SvPosition).xy;
-	const uint3 PixelCoord = uint3(floor(SvPosition.xy), 0);
-
-	FNodeDesc NodeDesc = DecodeNodeDesc(HairVisibilityNodeOffsetAndCount.Load(PixelCoord));
-
-	OutColor = 0.0f;
-	OutColorSubPixel = 0.0f;
-
-	if (NodeDesc.Count > 0)
-	{
-		const FCategorizationData CatData = DecodeCategorizationData(HairCategorizationTexture.Load(PixelCoord));
-
-		// Note:
-		// * Currently: compute take the closest sample
-		// * Change this to per sample evaluation once we have better hair transmittance evaluation for sky
-		float ClosestDepth = 0;
-		bool bIsClusterCompatible = false;
-		FGBufferData HairGBuffer = (FGBufferData)0;
-		float3 SkyLighting = 0;
-
-		LOOP
-		for (uint SampleIt = 0; SampleIt < NodeDesc.Count; SampleIt++)
-		{
-			const uint LocalOffset = NodeDesc.Offset + SampleIt;
-			const FNodeData NodeData = HairVisibilityNodeData[LocalOffset];
-			float			SampleDepth;
-			FVisibilityData SampleVisData;
-			float3			SampleBaseColor;
-			float			SampleRoughness;
-			float			SampleSpecular;
-			DecodeNodeData(NodeData, SampleDepth, SampleVisData, SampleBaseColor, SampleRoughness, SampleSpecular);
-			bIsClusterCompatible = bIsClusterCompatible || SampleVisData.ClusterId == Voxel_ClusterId;
-
-			// Inject material data from the visibility/mini-gbuffer for hair
-			const float SceneDepth = ConvertFromDeviceZ(SampleDepth);
-			#if PERMUTATION_PER_SAMPLE == 0
-			if (SampleDepth > ClosestDepth)
-			#endif
-			{
-				HairGBuffer.WorldNormal = SampleVisData.Tangent;
-				HairGBuffer.BaseColor = SampleBaseColor;
-				HairGBuffer.Roughness = SampleRoughness;
-				HairGBuffer.ShadingModelID = SHADINGMODELID_HAIR;
-				HairGBuffer.DiffuseColor = 0;
-				HairGBuffer.SpecularColor = 0;
-				HairGBuffer.Specular = SampleSpecular;
-				HairGBuffer.Metallic = 0;
-				HairGBuffer.Depth = SceneDepth;
-				HairGBuffer.GBufferAO = 1;
-				HairGBuffer.CustomData = 0;
-				HairGBuffer.IndirectIrradiance = 1000000;
-			}
-
-			#if PERMUTATION_PER_SAMPLE == 0
-			ClosestDepth = max(ClosestDepth, SampleDepth);  // Inverse-Z
-			#else
-			if (bIsClusterCompatible = SampleVisData.ClusterId == Voxel_ClusterId)
-			{
-				const float SampleCoverage = From8bitCoverage(SampleVisData.Coverage8bit);
-				SkyLighting += ReflectionEnvironment(HairGBuffer, ScreenPosition, SvPosition.xy) * SampleCoverage;
-			}
-			#endif
-		}
-		
-		// Lit fully covered or partially covered pixel
-		// * Fully covered pixel output their result into the main color target
-		// * Partially covered pixel output their result into the sub pixel color target. the main color target will store the underlying pixel color 
-		if (bIsClusterCompatible)
-		{
-			#if PERMUTATION_PER_SAMPLE == 0
-			SkyLighting = ReflectionEnvironment(HairGBuffer, ScreenPosition, SvPosition.xy) * saturate(CatData.PixelCoverage);
-			#endif
-
-			// Scene color alpha is used for ScreenSpaceSubsurfaceScattering (if that is not needed it can be disabled with SUBSURFACE_CHANNEL_MODE)
-			FLightAccumulator LightAccumulator = (FLightAccumulator)0;
-			LightAccumulator_Add(LightAccumulator, SkyLighting, SkyLighting, 1.0f, false);
-			const float4 LocalOutColor = LightAccumulator_GetResult(LightAccumulator);
-			if (CatData.PixelCoverage < 1)
-			{
-				OutColorSubPixel = LocalOutColor; // * PixelCoverage <-- this would be need if Coverage8bit was storing the ratio (sampleCoverage / validCoverage)
-				OutColorSubPixel.a = CatData.PixelCoverage;
-			}
-			else
-			{
-				OutColor += LocalOutColor;
-			}
-		}
-	}
-}
-#endif
-
-////////////////////////////////////////////////////////////////////////////////////////////////////////////////////////////////////////////
-
 #if COMPUTESHADER
 uint MaxVisibilityNodeCount;
 StructuredBuffer<uint>	HairVisibilityNodeCoord;
diff --git a/Engine/Source/Runtime/Renderer/Private/HairStrands/HairStrandsEnvironment.cpp b/Engine/Source/Runtime/Renderer/Private/HairStrands/HairStrandsEnvironment.cpp
--- a/Engine/Source/Runtime/Renderer/Private/HairStrands/HairStrandsEnvironment.cpp
+++ b/Engine/Source/Runtime/Renderer/Private/HairStrands/HairStrandsEnvironment.cpp
@@ -30,12 +30,6 @@ static FAutoConsoleVariableRef CVarHairSkyAOEnable(TEXT("r.HairStrands.SkyAOEnab
 static float GHairSkylightingConeAngle = 3;
 static FAutoConsoleVariableRef CVarHairSkylightingConeAngle(TEXT("r.HairStrands.SkyLighting.ConeAngle"), GHairSkylightingConeAngle, TEXT("Cone angle for tracing sky lighting on hair."));
 
-static float GHairSkylightingPerSample = 1;
-static FAutoConsoleVariableRef CVarHairSkylightingPerSample(TEXT("r.HairStrands.SkyLighting.PerSample"), GHairSkylightingPerSample, TEXT("Evaluate sky lighting per hair sample."));
-
-static float GHairStrandsSkyLightingCompute = 1;
-static FAutoConsoleVariableRef GVarHairStrandsSkyLightingCompute(TEXT("r.HairStrands.SkyLighting.Compute"), GHairStrandsSkyLightingCompute, TEXT("Evaluate sky lighting using a compute shader."));
-
 static int32 GHairStrandsSkyLightingSampleCount = 16;
 static FAutoConsoleVariableRef GVarHairStrandsSkyLightingSampleCount(TEXT("r.HairStrands.SkyLighting.SampleCount"), GHairStrandsSkyLightingSampleCount, TEXT("Number of samples used for evaluating multiple scattering (in [1,16], default is set to 16)."));
 
@@ -355,149 +349,7 @@ static FRDGBufferRef AddHairStrandsEnvironmentLightingPassCS(
 }
 
 ///////////////////////////////////////////////////////////////////////////////////////////////////
-// Lighting PS
-class FHairEnvironmentLightingPS : public FGlobalShader
-{
-	DECLARE_GLOBAL_SHADER(FHairEnvironmentLightingPS);
-	SHADER_USE_PARAMETER_STRUCT(FHairEnvironmentLightingPS, FGlobalShader)
-
-	class FPerSample : SHADER_PERMUTATION_INT("PERMUTATION_PER_SAMPLE", 2);
-	using FPermutationDomain = TShaderPermutationDomain<FPerSample>;
-
-	static void ModifyCompilationEnvironment(const FGlobalShaderPermutationParameters& Parameters, FShaderCompilerEnvironment& OutEnvironment)
-	{
-		FForwardLightingParameters::ModifyCompilationEnvironment(Parameters.Platform, OutEnvironment);
-	}
-
-	static bool ShouldCompilePermutation(const FGlobalShaderPermutationParameters& Parameters)
-	{
-		return IsHairStrandsSupported(Parameters.Platform);
-	}
-
-	BEGIN_SHADER_PARAMETER_STRUCT(FParameters, )
-
-		SHADER_PARAMETER(FVector, Voxel_MinAABB)
-		SHADER_PARAMETER(uint32, Voxel_ClusterId)
-		SHADER_PARAMETER(FVector, Voxel_MaxAABB)
-		SHADER_PARAMETER(uint32, Voxel_Resolution)
-		SHADER_PARAMETER(float, Voxel_DensityScale)
-		SHADER_PARAMETER(float, Voxel_DepthBiasScale)
-		SHADER_PARAMETER(float, Voxel_TanConeAngle)
-		SHADER_PARAMETER(uint32, MultipleScatterSampleCount)
-		SHADER_PARAMETER(uint32, HairComponents)
-		SHADER_PARAMETER(uint32, JitterSphericalIntegration)
-
-		SHADER_PARAMETER_TEXTURE(Texture2D, PreIntegratedGF)
-		SHADER_PARAMETER_SAMPLER(SamplerState, PreIntegratedGFSampler)
-
-		SHADER_PARAMETER_TEXTURE(Texture2D, HairCategorizationTexture)
-		SHADER_PARAMETER_TEXTURE(Texture2D, HairVisibilityNodeOffsetAndCount)
-		SHADER_PARAMETER_SRV(Buffer, HairVisibilityNodeData)
-
-		SHADER_PARAMETER_STRUCT_INCLUDE(FSceneTextureParameters, SceneTextures)
-		SHADER_PARAMETER_STRUCT_INCLUDE(FSceneTextureSamplerParameters, SceneTextureSamplers)
-
-		SHADER_PARAMETER_RDG_TEXTURE(TEXTURE3D, HairEnergyLUTTexture)
-		SHADER_PARAMETER_RDG_TEXTURE(TEXTURE3D, HairScatteringLUTTexture)
-		SHADER_PARAMETER_SAMPLER(SamplerState, HairLUTSampler)
-
-		SHADER_PARAMETER_RDG_TEXTURE(Texture3D, Voxel_DensityTexture)
-
-		SHADER_PARAMETER_STRUCT_REF(FViewUniformShaderParameters, ViewUniformBuffer)
-		SHADER_PARAMETER_STRUCT_REF(FReflectionUniformParameters, ReflectionsParameters)
-		SHADER_PARAMETER_STRUCT_REF(FReflectionCaptureShaderData, ReflectionCaptureData)
-		SHADER_PARAMETER_STRUCT_REF(FForwardLightData, ForwardLightData)
-
-		RENDER_TARGET_BINDING_SLOTS()
-	END_SHADER_PARAMETER_STRUCT()
-};
-
-IMPLEMENT_GLOBAL_SHADER(FHairEnvironmentLightingPS, "/Engine/Private/HairStrands/HairStrandsEnvironmentLighting.usf", "MainPS", SF_Pixel);
-
-static void AddHairStrandsEnvironmentLightingPass(
-	FRDGBuilder& GraphBuilder,
-	const FViewInfo& View,
-	const FHairStrandsVisibilityData& VisibilityData,
-	const FHairStrandsClusterData& ClusterData,
-	FRDGTextureRef Output0,
-	FRDGTextureRef Output1)
-{
-	FSceneTextureParameters SceneTextures;
-	SetupSceneTextureParameters(GraphBuilder, &SceneTextures);
-
-	// Render the reflection environment with tiled deferred culling
-	const bool bHasBoxCaptures = (View.NumBoxReflectionCaptures > 0);
-	const bool bHasSphereCaptures = (View.NumSphereReflectionCaptures > 0);
-	FHairEnvironmentLightingPS::FParameters* PassParameters = GraphBuilder.AllocParameters<FHairEnvironmentLightingPS::FParameters>();
 
-	const FHairLUT InHairLUT = GetHairLUT(GraphBuilder.RHICmdList, View);
-	PassParameters->HairEnergyLUTTexture = GraphBuilder.RegisterExternalTexture(InHairLUT.Textures[HairLUTType_MeanEnergy], TEXT("HairMeanEnergyLUTTexture"));;
-	PassParameters->HairScatteringLUTTexture = GraphBuilder.RegisterExternalTexture(InHairLUT.Textures[HairLUTType_DualScattering], TEXT("HairScatteringEnergyLUTTexture"));
-	PassParameters->HairLUTSampler = TStaticSamplerState<SF_Bilinear, AM_Clamp, AM_Clamp, AM_Clamp>::GetRHI();
-
-	PassParameters->Voxel_ClusterId = ClusterData.ClusterId;
-	PassParameters->Voxel_MinAABB = ClusterData.GetMinBound();
-	PassParameters->Voxel_MaxAABB = ClusterData.GetMaxBound();
-	PassParameters->Voxel_Resolution = ClusterData.GetResolution();
-	PassParameters->Voxel_DensityTexture = GraphBuilder.RegisterExternalTexture(ClusterData.VoxelResources.DensityTexture);
-	PassParameters->Voxel_DensityScale = GetHairStrandsVoxelizationDensityScale();
-	PassParameters->Voxel_DepthBiasScale = GetHairStrandsVoxelizationDepthBiasScale();
-	PassParameters->Voxel_TanConeAngle = FMath::Tan(FMath::DegreesToRadians(GetHairStrandsSkyLightingConeAngle()));
-
-	PassParameters->PreIntegratedGF = GSystemTextures.PreintegratedGF->GetRenderTargetItem().ShaderResourceTexture;
-	PassParameters->PreIntegratedGFSampler = TStaticSamplerState<SF_Bilinear, AM_Clamp, AM_Clamp, AM_Clamp>::GetRHI();
-
-	PassParameters->SceneTextures = SceneTextures;
-	SetupSceneTextureSamplers(&PassParameters->SceneTextureSamplers);
-
-	PassParameters->ViewUniformBuffer = View.ViewUniformBuffer;
-	PassParameters->ReflectionCaptureData = View.ReflectionCaptureUniformBuffer;
-	{
-		FReflectionUniformParameters ReflectionUniformParameters;
-		SetupReflectionUniformParameters(View, ReflectionUniformParameters);
-		PassParameters->ReflectionsParameters = CreateUniformBufferImmediate(ReflectionUniformParameters, UniformBuffer_SingleDraw);
-	}
-	PassParameters->ForwardLightData = View.ForwardLightingResources->ForwardLightDataUniformBuffer;
-	PassParameters->MultipleScatterSampleCount = FMath::Clamp(uint32(GHairStrandsSkyLightingSampleCount), 1u, 16u);
-	PassParameters->HairComponents = ToBitfield(GetHairComponents());
-	PassParameters->JitterSphericalIntegration = GHairStrandsSkyLightingJitterSphericalIntegration ? 1 : 0;
-
-	// Bind hair data
-	PassParameters->HairCategorizationTexture = VisibilityData.CategorizationTexture->GetRenderTargetItem().ShaderResourceTexture;
-	PassParameters->HairVisibilityNodeOffsetAndCount = VisibilityData.NodeIndex->GetRenderTargetItem().ShaderResourceTexture;
-	PassParameters->HairVisibilityNodeData = VisibilityData.NodeDataSRV;
-
-	check(Output0); check(Output1);
-	PassParameters->RenderTargets[0] = FRenderTargetBinding(Output0, ERenderTargetLoadAction::ELoad);
-	PassParameters->RenderTargets[1] = FRenderTargetBinding(Output1, ERenderTargetLoadAction::ELoad);
-
-	FHairEnvironmentLightingPS::FPermutationDomain PermutationVector;
-	PermutationVector.Set<FHairEnvironmentLightingPS::FPerSample>(GHairSkylightingPerSample > 0 ? 1 : 0);
-	TShaderMapRef<FHairEnvironmentLightingPS> PixelShader(View.ShaderMap, PermutationVector);
-	ClearUnusedGraphResources(*PixelShader, PassParameters);
-
-	GraphBuilder.AddPass(
-		RDG_EVENT_NAME("HairStrandsEnvironmentPS %dx%d", View.ViewRect.Width(), View.ViewRect.Height()),
-		PassParameters,
-		ERDGPassFlags::Raster,
-		[PassParameters, &View, PixelShader](FRHICommandList& InRHICmdList)
-	{
-		InRHICmdList.SetViewport(View.ViewRect.Min.X, View.ViewRect.Min.Y, 0.0f, View.ViewRect.Max.X, View.ViewRect.Max.Y, 1.0f);
-
-		FGraphicsPipelineStateInitializer GraphicsPSOInit;
-		FPixelShaderUtils::InitFullscreenPipelineState(InRHICmdList, View.ShaderMap, *PixelShader, GraphicsPSOInit);
-
-		GraphicsPSOInit.BlendState = TStaticBlendState<
-			CW_RGBA, BO_Add, BF_One, BF_One, BO_Add, BF_One, BF_One,
-			CW_RGBA, BO_Add, BF_One, BF_One, BO_Max, BF_SourceAlpha, BF_DestAlpha
-			>::GetRHI();
-
-		SetGraphicsPipelineState(InRHICmdList, GraphicsPSOInit);
-		SetShaderParameters(InRHICmdList, *PixelShader, PixelShader->GetPixelShader(), *PassParameters);
-		FPixelShaderUtils::DrawFullscreenTriangle(InRHICmdList);
-	});
-}
-	
 void RenderHairStrandsEnvironmentLighting(
 	FRDGBuilder& GraphBuilder,
 	const uint32 ViewIndex,
@@ -523,19 +375,8 @@ void RenderHairStrandsEnvironmentLighting(
 	const FViewInfo& View = Views[ViewIndex];
 	for (const FHairStrandsClusterData& ClusterData : HairDatas->HairClusterPerViews.Views[ViewIndex].Datas)
 	{
-		// @hair_todo: 
-		// * Add support for : BentNormal, global AO, SSR, DFSO
-		// * Add local reflection probe, current take into account only the sky lighting
-
-		if (GHairStrandsSkyLightingCompute)
-		{
-			FRDGBufferRef SampleLightingBuffer = AddHairStrandsEnvironmentLightingPassCS(GraphBuilder, View, VisibilityData, ClusterData, VisibilityData.NodeGroupSize, NodeIndirectArgBuffer);
-			AddHairEnvironmentLightingComposePass(GraphBuilder, View, VisibilityData, ClusterData, SampleLightingBuffer, SceneColorTexture, SceneColorSubPixelTexture);
-		}
-		else
-		{
-			AddHairStrandsEnvironmentLightingPass(GraphBuilder, View, VisibilityData, ClusterData, SceneColorTexture, SceneColorSubPixelTexture);
-		}
+		FRDGBufferRef SampleLightingBuffer = AddHairStrandsEnvironmentLightingPassCS(GraphBuilder, View, VisibilityData, ClusterData, VisibilityData.NodeGroupSize, NodeIndirectArgBuffer);
+		AddHairEnvironmentLightingComposePass(GraphBuilder, View, VisibilityData, ClusterData, SampleLightingBuffer, SceneColorTexture, SceneColorSubPixelTexture);
 	}
 }
 
