Add permutation for removing limit of sample count for sky lighting integration.

diff --git a/Engine/Shaders/Private/HairStrands/HairStrandsEnvironmentLighting.usf b/Engine/Shaders/Private/HairStrands/HairStrandsEnvironmentLighting.usf
--- a/Engine/Shaders/Private/HairStrands/HairStrandsEnvironmentLighting.usf
+++ b/Engine/Shaders/Private/HairStrands/HairStrandsEnvironmentLighting.usf
@@ -145,6 +145,7 @@ float3 ReflectionEnvironment(FGBufferData GBuffer, float2 ScreenPosition, float2
 	#if 1
 	if ((HairComponents & HAIR_COMPONENT_R) || (HairComponents & HAIR_COMPONENT_TRT) || (HairComponents & HAIR_COMPONENT_TT))
 	{
+		#if PERMUTATION_SAMPLESET == 0
 		// Poisson disk position http://developer.download.nvidia.com/whitepapers/2008/PCSS_Integration.pdf
 		float2 PoissonDisk[16] =
 		{
@@ -165,6 +166,7 @@ float3 ReflectionEnvironment(FGBufferData GBuffer, float2 ScreenPosition, float2
 			float2(0.19984126, 0.78641367),
 			float2(0.14383161, -0.14100790)
 		};
+		#endif
 
 		float2 RandOffset = 0;
 		if (JitterSphericalIntegration > 0)
@@ -178,7 +180,11 @@ float3 ReflectionEnvironment(FGBufferData GBuffer, float2 ScreenPosition, float2
 		const uint SampleCount = clamp(MultipleScatterSampleCount, 1, 16);
 		for (uint i = 0; i < SampleCount; ++i)
 		{
+		#if PERMUTATION_SAMPLESET == 0
 			const float2 Rand2 = (PoissonDisk[i] + float2(1, 1)) * 0.5f + RandOffset;
+		#else // PERMUTATION_SAMPLESET == 1
+			const float2 Rand2 = Hammersley(i, SampleCount, uint2(0,0)) + RandOffset;
+		#endif		
 			const float3 SampleL = UniformSampleSphere(frac(Rand2)).xyz;
 
 			// Depth bias
diff --git a/Engine/Source/Runtime/Renderer/Private/HairStrands/HairStrandsEnvironment.cpp b/Engine/Source/Runtime/Renderer/Private/HairStrands/HairStrandsEnvironment.cpp
--- a/Engine/Source/Runtime/Renderer/Private/HairStrands/HairStrandsEnvironment.cpp
+++ b/Engine/Source/Runtime/Renderer/Private/HairStrands/HairStrandsEnvironment.cpp
@@ -31,7 +31,7 @@ static float GHairSkylightingConeAngle = 3;
 static FAutoConsoleVariableRef CVarHairSkylightingConeAngle(TEXT("r.HairStrands.SkyLighting.ConeAngle"), GHairSkylightingConeAngle, TEXT("Cone angle for tracing sky lighting on hair."));
 
 static int32 GHairStrandsSkyLightingSampleCount = 16;
-static FAutoConsoleVariableRef GVarHairStrandsSkyLightingSampleCount(TEXT("r.HairStrands.SkyLighting.SampleCount"), GHairStrandsSkyLightingSampleCount, TEXT("Number of samples used for evaluating multiple scattering (in [1,16], default is set to 16)."));
+static FAutoConsoleVariableRef GVarHairStrandsSkyLightingSampleCount(TEXT("r.HairStrands.SkyLighting.SampleCount"), GHairStrandsSkyLightingSampleCount, TEXT("Number of samples used for evaluating multiple scattering and visible area (default is set to 16)."));
 
 static int32 GHairStrandsSkyLightingJitterSphericalIntegration = 0;
 static FAutoConsoleVariableRef GVarHairStrandsSkyLightingJitterSphericalIntegration(TEXT("r.HairStrands.SkyLighting.JitterIntegration"), GHairStrandsSkyLightingJitterSphericalIntegration, TEXT("Jitter the sphereical integration for the multiple scattering term. The result is more correct, but noiser as well."));
@@ -220,7 +220,8 @@ class FHairEnvironmentLightingCS : public FGlobalShader
 	SHADER_USE_PARAMETER_STRUCT(FHairEnvironmentLightingCS, FGlobalShader)
 
 	class FGroupSize : SHADER_PERMUTATION_INT("PERMUTATION_GROUP_SIZE", 2);
-	using FPermutationDomain = TShaderPermutationDomain<FGroupSize>;
+	class FSampleSet : SHADER_PERMUTATION_INT("PERMUTATION_SAMPLESET", 2);
+	using FPermutationDomain = TShaderPermutationDomain<FGroupSize, FSampleSet>;
 
 	static void ModifyCompilationEnvironment(const FGlobalShaderPermutationParameters& Parameters, FShaderCompilerEnvironment& OutEnvironment)
 	{
@@ -308,7 +309,7 @@ static FRDGBufferRef AddHairStrandsEnvironmentLightingPassCS(
 	PassParameters->Voxel_DepthBiasScale = GetHairStrandsVoxelizationDepthBiasScale();
 	PassParameters->Voxel_TanConeAngle = FMath::Tan(FMath::DegreesToRadians(GetHairStrandsSkyLightingConeAngle()));
 	PassParameters->MaxVisibilityNodeCount = VisibilityData.NodeData->Desc.NumElements;
-	PassParameters->MultipleScatterSampleCount = FMath::Clamp(uint32(GHairStrandsSkyLightingSampleCount), 1u, 16u);
+	PassParameters->MultipleScatterSampleCount = FMath::Max(uint32(GHairStrandsSkyLightingSampleCount), 1u);
 	PassParameters->JitterSphericalIntegration = GHairStrandsSkyLightingJitterSphericalIntegration ? 1 : 0;
 	PassParameters->HairComponents = ToBitfield(GetHairComponents());
 	PassParameters->PreIntegratedGF = GSystemTextures.PreintegratedGF->GetRenderTargetItem().ShaderResourceTexture;
@@ -335,6 +336,7 @@ static FRDGBufferRef AddHairStrandsEnvironmentLightingPassCS(
 	check(NodeGroupSize == 64 || NodeGroupSize == 32);
 	FHairEnvironmentLightingCS::FPermutationDomain PermutationVector;
 	PermutationVector.Set<FHairEnvironmentLightingCS::FGroupSize>(NodeGroupSize == 64 ? 0 : (NodeGroupSize == 32 ? 1 : 2));
+	PermutationVector.Set<FHairEnvironmentLightingCS::FSampleSet>(PassParameters->MultipleScatterSampleCount <= 16 ? 0 : 1);
 
 	TShaderMapRef<FHairEnvironmentLightingCS> ComputeShader(View.ShaderMap, PermutationVector);
 	FComputeShaderUtils::AddPass(
