Added per pixel linked list option for high quality hair rendering (e.g. when rendering movies). Not there yet, that is a first checkin.

diff --git a/Engine/Shaders/Private/HairStrands/HairStrandsVisibilityCommon.ush b/Engine/Shaders/Private/HairStrands/HairStrandsVisibilityCommon.ush
--- a/Engine/Shaders/Private/HairStrands/HairStrandsVisibilityCommon.ush
+++ b/Engine/Shaders/Private/HairStrands/HairStrandsVisibilityCommon.ush
@@ -179,65 +179,6 @@ FNodeData ConvertToNodeData(float Depth, FVisibilityData VisibilityData, float3
 	 uint   HairSampleCount;
  };
 
- void ComputeUniqueSamples(
-	 inout FSampleDesc OutSamples[HAIR_VISIBILITY_MAXSAMPLECOUNT],
-	 inout FSampleSetDesc OutSet,
-	 const uint2 PixelCoord,
-	 const uint SampleCount,
-	 Texture2D<float4> InSceneDepthTexture,
-	 Texture2DMS<uint4, HAIR_VISIBILITY_MAXSAMPLECOUNT>  InMSAA_IDTexture,
-	 Texture2DMS<float4, HAIR_VISIBILITY_MAXSAMPLECOUNT> InMSAA_MaterialTexture,
-//	 Texture2DMS<float4, HAIR_VISIBILITY_MAXSAMPLECOUNT> InMSAA_AttributeTexture,
-	 Texture2DMS<float, HAIR_VISIBILITY_MAXSAMPLECOUNT>  InMSAA_DepthTexture)
- {
-	 OutSet.UniqueSampleCount = 0;
-	 OutSet.ValidSampleCount = 0;
-
-	 OutSet.HairSampleCount = clamp(SampleCount, 0, HAIR_VISIBILITY_MAXSAMPLECOUNT);
-	 const uint InvalidSample = 0xFFFF;
-	 const float OpaqueDepth = InSceneDepthTexture.Load(uint3(PixelCoord, 0)).r;
-	 for (uint SampleIt = 0; SampleIt < OutSet.HairSampleCount; ++SampleIt)
-	 {
-		 const uint4 EncodedData = InMSAA_IDTexture.Load(PixelCoord, SampleIt);
-		 const bool bIsValid = IsValidHairData(EncodedData);
-		 if (!bIsValid)
-			 continue;
-
-		 const float SampleDepth = InMSAA_DepthTexture.Load(PixelCoord, SampleIt);
-		 if (OpaqueDepth > SampleDepth) // Reverse-Z
-			 continue;
-
-		 ++OutSet.ValidSampleCount;
-
-		 const FVisibilityData VisibilityData = DecodeVisibilityData(EncodedData);
-		 bool bAlreadyExist = false;
-		 for (uint UniqueIt = 0; UniqueIt < OutSet.UniqueSampleCount; ++UniqueIt)
-		 {
-			 if (VisibilityData.PrimitiveId == OutSamples[UniqueIt].SampleId)
-			 {
-				 OutSamples[UniqueIt].Weights += 1;
-				 bAlreadyExist = true;
-				 break;
-			 }
-		 }
-
-		 if (!bAlreadyExist)
-		 {
-			 //const float Specular = InMSAA_AttributeTexture.Load(PixelCoord, SampleIt).x;
-			 const float4 BaseColorAndRoughness = InMSAA_MaterialTexture.Load(PixelCoord, SampleIt);
-			 OutSamples[OutSet.UniqueSampleCount].SampleId = VisibilityData.PrimitiveId;
-			 OutSamples[OutSet.UniqueSampleCount].Tangents = VisibilityData.Tangent;
-			 OutSamples[OutSet.UniqueSampleCount].Coverage8bit = VisibilityData.Coverage8bit;
-			 OutSamples[OutSet.UniqueSampleCount].Depths = SampleDepth;
-			 OutSamples[OutSet.UniqueSampleCount].Weights = 1;
-			 OutSamples[OutSet.UniqueSampleCount].ClusterId = VisibilityData.ClusterId;
-			 OutSamples[OutSet.UniqueSampleCount].BaseColorAndRoughness = float4(BaseColorAndRoughness.xyz * BaseColorAndRoughness.xyz, BaseColorAndRoughness.a);
-			 OutSamples[OutSet.UniqueSampleCount].Specular = 0.5f; // Specular;
-			 ++OutSet.UniqueSampleCount;
-		 }
-	 }
- }
-
  uint To8bitCoverage(float Coverage)
  {
 	 return min(uint(Coverage * 0x100), 0xFFu);
@@ -349,4 +290,45 @@ FNodeData ConvertToNodeData(float Depth, FVisibilityData VisibilityData, float3
 	 Out.UCoord		= (Attributes.y & 0xFF) / 255.f;
 	 Out.Seed		= ((Attributes.y>>8) & 0xFF) / 255.f;
 	 return Out;
- }
\ No newline at end of file
+ }
+
+ //////////
+ // PPLL
+
+// Footprint: 160bits | 20bytes
+ struct FPPLLNodeData
+ {
+	 float Depth;
+	 uint PrimitiveID_ClusterID;
+	 uint Tangent_Coverage8bit;
+	 uint BaseColor_Roughness;
+	 uint NextNodeIndex;
+ };
+
+ FPPLLNodeData ConvertToPPLLNodeData(float Depth, FVisibilityData VisibilityData, float3 BaseColor, float Roughness, uint NextNodeIndex)
+ {
+	 // 4x16bits
+	 const uint4 EncodedVisibilityData = EncodeVisibilityData(VisibilityData);
+
+	 FPPLLNodeData Out;
+	 Out.Depth = Depth;																								// 32bits float
+	 Out.Tangent_Coverage8bit = (0xFFFF & EncodedVisibilityData.x) | ((0xFFFF & EncodedVisibilityData.y) << 16);	// 32bits uint
+	 Out.PrimitiveID_ClusterID = (0xFFFF & EncodedVisibilityData.z) | ((0xFFFF & EncodedVisibilityData.w) << 16);	// 32bits uint
+	 Out.BaseColor_Roughness = float4ToUint(float4(sqrt(BaseColor), Roughness));									// 32bits uint
+	 Out.NextNodeIndex = NextNodeIndex;																				// 32bits next node
+
+	 return Out;
+ }
+
+ FNodeData ConvertToNodeData(FPPLLNodeData In)
+ {
+	 FNodeData Out;
+	 Out.Depth = In.Depth;
+	 Out.PrimitiveID_ClusterID = In.PrimitiveID_ClusterID;
+	 Out.Tangent_Coverage8bit = In.Tangent_Coverage8bit;
+	 Out.BaseColor_Roughness = In.BaseColor_Roughness;
+	 Out.Specular = 0.5f;		// TODO
+
+	 return Out;
+ }
+
diff --git a/Engine/Shaders/Private/HairStrands/HairStrandsVisibilityCompaction.usf b/Engine/Shaders/Private/HairStrands/HairStrandsVisibilityCompaction.usf
--- a/Engine/Shaders/Private/HairStrands/HairStrandsVisibilityCompaction.usf
+++ b/Engine/Shaders/Private/HairStrands/HairStrandsVisibilityCompaction.usf
@@ -15,6 +15,156 @@
 float CosTangentThreshold;
 #endif
 
+#if PERMUTATION_VENDOR == VENDOR_AMD
+#define TILE_PIXEL_SIZE_X 8
+#define TILE_PIXEL_SIZE_Y 8
+#elif PERMUTATION_VENDOR == VENDOR_NVIDIA
+#define TILE_PIXEL_SIZE_X 8
+#define TILE_PIXEL_SIZE_Y 4
+#else //VENDOR_OTHER
+#define TILE_PIXEL_SIZE_X 8
+#define TILE_PIXEL_SIZE_Y 8
+#endif
+
+uint2 OutputResolution;
+uint MaxNodeCount;
+uint HairVisibilitySampleCount;	// MSAA sample count
+uint2 ResolutionOffset;
+
+#if PERMUTATION_COVERAGE == 1
+Texture2D<float> CoverageTexture;
+#endif
+
+RWTexture2D<uint> OutCompactNodeCounter;
+RWTexture2D<uint> OutCompactNodeIndex;
+RWStructuredBuffer<FNodeData> OutCompactNodeData;
+RWStructuredBuffer<uint> OutCompactNodeCoord;
+RWTexture2D<uint4> OutCategorizationTexture;
+
+#define VELOCITY_TYPE_NONE 0
+#define VELOCITY_TYPE_AVG 1  
+#define VELOCITY_TYPE_CLOSEST 2
+#define VELOCITY_TYPE_MAX 3
+#if PERMUTATION_VELOCITY != VELOCITY_TYPE_NONE
+	Texture2DMS<float2, HAIR_VISIBILITY_MAXSAMPLECOUNT> MSAA_VelocityTexture;
+	RWTexture2D<float2> OutVelocityTexture;
+#endif
+
+
+#if PERMUTATION_PPLL
+
+Texture2D<uint> PPLLCounter;
+Texture2D<uint> PPLLNodeIndex;
+StructuredBuffer<FPPLLNodeData> PPLLNodeData;
+
+
+[numthreads(TILE_PIXEL_SIZE_X, TILE_PIXEL_SIZE_Y, 1)]
+void MainCS(uint GroupIndex : SV_GroupIndex, uint3 DispatchThreadId : SV_DispatchThreadID)
+{
+	uint2 PixelCoord = DispatchThreadId.xy + ResolutionOffset;
+	const uint FirstNodeIndex = PPLLNodeIndex[PixelCoord];
+
+	FCategorizationData Categorization = (FCategorizationData)0;
+	uint OffsetInStorage = 0;
+	uint PixelSamplecount = 0;
+
+	if (FirstNodeIndex != 0xFFFFFFFF && all(PixelCoord < OutputResolution))
+	{
+		//////////
+		// First pass: Pre compute some values based on the linked list
+		float ValidPixelSampleTotalCoverage = 0.0f;
+		float Transmittance = 1.0f;
+
+		uint SortedIndex[HAIR_VISIBILITY_MAXSAMPLECOUNT];
+		float SortedDepth[HAIR_VISIBILITY_MAXSAMPLECOUNT];
+
+		uint PixelNodeCount = 0;
+		uint NodeIndex = FirstNodeIndex;
+		while (NodeIndex != 0xFFFFFFFF)// && Transmittance > 0.0f)// && PixelNodeCount < HAIR_VISIBILITY_MAXSAMPLECOUNT)
+		{
+			FPPLLNodeData PPLLNode = PPLLNodeData[NodeIndex];
+
+			// Decode data we need
+			float			Depth = 0.0f;
+			FVisibilityData	VisibilityData = (FVisibilityData)0;
+			float3			BaseColor = 0.0f;
+			float			Roughness = 0.0f;
+			float			Specular = 0;
+			DecodeNodeData(ConvertToNodeData(PPLLNode), Depth, VisibilityData, BaseColor, Roughness, Specular);
+
+			const float Coverage = From8bitCoverage(VisibilityData.Coverage8bit);
+			Transmittance *= 1.0f - Coverage;
+			if (PixelNodeCount < HAIR_VISIBILITY_MAXSAMPLECOUNT)
+			{
+				ValidPixelSampleTotalCoverage += Coverage;
+
+				// Always insert up to HAIR_VISIBILITY_MAXSAMPLECOUNT data
+				SortedIndex[PixelNodeCount] = NodeIndex;
+				SortedDepth[PixelNodeCount] = PPLLNode.Depth;
+				// TODO sort and keep closest, merge furthest
+
+			}
+
+			PixelNodeCount++;
+			NodeIndex = PPLLNode.NextNodeIndex;
+		}
+
+		//////////
+		// Second pass: write out compact pixel sample data
+		PixelSamplecount = min(PixelNodeCount, HAIR_VISIBILITY_MAXSAMPLECOUNT);
+		InterlockedAdd(OutCompactNodeCounter[uint2(0, 0)], PixelSamplecount, OffsetInStorage);
+		for (uint i = 0; i < PixelSamplecount; ++i)
+		{
+			const uint NodePPLLIndex = SortedIndex[i];
+			FPPLLNodeData PPLLNode = PPLLNodeData[NodePPLLIndex];
+
+			FNodeData OutNodeData = ConvertToNodeData(PPLLNode);
+
+			float			Depth = 0.0f;
+			FVisibilityData	VisibilityData = (FVisibilityData)0;
+			float3			BaseColor = 0.0f;
+			float			Roughness = 0.0f;
+			float			Specular = 0.0f;
+			DecodeNodeData(OutNodeData, Depth, VisibilityData, BaseColor, Roughness, Specular);
+
+			// Coverage8bit is a weight normalising to 1 the contribution of all the compacted samples. Because later it is weighted by Categorization.PixelCoverage.
+			uint Coverage8bit = To8bitCoverage(From8bitCoverage(VisibilityData.Coverage8bit) / float(ValidPixelSampleTotalCoverage));
+			// Patch the coverage on the out node
+			VisibilityData.Coverage8bit = Coverage8bit;
+			OutNodeData = ConvertToNodeData(Depth, VisibilityData, BaseColor, Roughness, Specular);
+
+			Categorization.TotalCoverage8bit += Coverage8bit;
+			Categorization.ClosestDepth = max(Categorization.ClosestDepth, OutNodeData.Depth); // Inverse-Z
+
+			OutCompactNodeData[OffsetInStorage + i] = OutNodeData;
+			OutCompactNodeCoord[OffsetInStorage + i] = (PixelCoord.x & 0xFFFF) | ((PixelCoord.y & 0xFFFF) << 16);
+		}
+
+	#if PERMUTATION_COVERAGE == 1
+		Categorization.PixelCoverage = CoverageTexture.Load(uint3(PixelCoord, 0));
+	#else
+		Categorization.PixelCoverage = 1.0f - Transmittance;
+	#endif
+	}
+
+	FNodeDesc NodeDesc;
+	NodeDesc.Count = PixelSamplecount;
+	NodeDesc.Offset = OffsetInStorage;
+	OutCompactNodeIndex[PixelCoord] = EncodeNodeDesc(NodeDesc);
+
+	Categorization.SampleCount = PixelSamplecount;
+	OutCategorizationTexture[PixelCoord] = EncodeCategorizationData(Categorization);
+
+
+
+}
+
+
+#else // PERMUTATION_PPLL //////////////////////////////////////////////////////////////////////////////////////////////////////////////
+
+
+
+
 void ComputeUniqueSamplesWithoutDepthTest(
 	inout uint3 OutSamples[HAIR_VISIBILITY_MAXSAMPLECOUNT],
 	inout FSampleSetDesc OutSet,
@@ -86,52 +236,15 @@
 	}
 }
 
-
-//////////////////////////////////////////////////////////////////////////////////////////////////////////////
-#if PERMUTATION_VENDOR == VENDOR_AMD
-#define TILE_PIXEL_SIZE_X 8
-#define TILE_PIXEL_SIZE_Y 8
-#elif PERMUTATION_VENDOR == VENDOR_NVIDIA
-#define TILE_PIXEL_SIZE_X 8
-#define TILE_PIXEL_SIZE_Y 4
-#else //VENDOR_OTHER
-#define TILE_PIXEL_SIZE_X 8
-#define TILE_PIXEL_SIZE_Y 8
-#endif
-
 #ifndef HAIR_VISIBILITY_MAXSAMPLECOUNT
 #error HAIR_VISIBILITY_MAXSAMPLECOUNT is not defined
 #endif
 
-uint2 OutputResolution;
-uint MaxNodeCount;
-uint HairVisibilitySampleCount;
-uint2 ResolutionOffset;
-
 Texture2DMS<uint4, HAIR_VISIBILITY_MAXSAMPLECOUNT> MSAA_IDTexture;
 Texture2DMS<float, HAIR_VISIBILITY_MAXSAMPLECOUNT> MSAA_DepthTexture;
 Texture2DMS<float4,HAIR_VISIBILITY_MAXSAMPLECOUNT> MSAA_MaterialTexture;
 Texture2DMS<float4,HAIR_VISIBILITY_MAXSAMPLECOUNT> MSAA_AttributeTexture;
 
-#if PERMUTATION_COVERAGE == 1
-Texture2D<float> CoverageTexture;
-#endif
-
-RWTexture2D<uint> OutCompactNodeCounter;
-RWTexture2D<uint> OutCompactNodeIndex;
-RWStructuredBuffer<FNodeData> OutCompactNodeData;
-RWStructuredBuffer<uint> OutCompactNodeCoord;
-RWTexture2D<uint4> OutCategorizationTexture;
-
-#define VELOCITY_TYPE_NONE 0
-#define VELOCITY_TYPE_AVG 1  
-#define VELOCITY_TYPE_CLOSEST 2
-#define VELOCITY_TYPE_MAX 3
-#if PERMUTATION_VELOCITY != VELOCITY_TYPE_NONE
-	Texture2DMS<float2, HAIR_VISIBILITY_MAXSAMPLECOUNT> MSAA_VelocityTexture;
-	RWTexture2D<float2> OutVelocityTexture;
-#endif
-
 groupshared uint AllocationNodeCount;
 groupshared uint AllocationNodeOffset;
 
@@ -284,3 +397,9 @@
 	OutCompactNodeIndex[PixelCoord] = EncodeNodeDesc(NodeDesc);
 	OutCategorizationTexture[PixelCoord] = EncodeCategorizationData(Categorization);
 }
+
+
+
+#endif // PERMUTATION_PPLL
+
+
diff --git a/Engine/Shaders/Private/HairStrands/HairStrandsVisibilityComposeSubPixelPS.usf b/Engine/Shaders/Private/HairStrands/HairStrandsVisibilityComposeSubPixelPS.usf
--- a/Engine/Shaders/Private/HairStrands/HairStrandsVisibilityComposeSubPixelPS.usf
+++ b/Engine/Shaders/Private/HairStrands/HairStrandsVisibilityComposeSubPixelPS.usf
@@ -8,12 +8,17 @@
 #include "HairStrandsVisibilityCommon.ush"
 
 Texture2D<float4> SubPixelColorTexture;
+#if PERMUTATION_PPLL
+Texture2D<uint4> CategorisationTexture;
+#else
 Texture2DMS<float, HAIR_VISIBILITY_MAXSAMPLECOUNT> HairVisibilityDepthTexture;
+#endif
 
 void SubColorPS(
 	in FScreenVertexOutput Input,
-	out float4 OutColor  : SV_Target0,
-	out float  OutDepth : SV_Depth)
+	out float4 OutColor  : SV_Target0
+	, out float  OutDepth : SV_Depth
+)
 {
 	const uint2 PixelCoord = floor(Input.Position.xy);
 	OutColor = SubPixelColorTexture.Load(uint3(PixelCoord, 0));
@@ -22,6 +27,13 @@ void SubColorPS(
 	//if (OutColor.a == 0)
 	//	discard;
 
+#if PERMUTATION_PPLL
+
+	const FCategorizationData CategorisationData = DecodeCategorizationData(CategorisationTexture.Load(uint3(Input.Position.xy, 0)));
+	OutDepth = CategorisationData.ClosestDepth;
+
+#else
+
 	// Brute force reading of the depth value. This could be pre-computed during the categorize pass or the compaction pass
 	// This must match what is done in VisibilityDepthPS to fonction. TODO: do not write depth here as it is already in. 
 	OutDepth = 0;
@@ -30,6 +42,8 @@ void SubColorPS(
 		const float SampleDepth = HairVisibilityDepthTexture.Load(PixelCoord, SampleIt);
 		OutDepth = max(OutDepth, SampleDepth); // HAS_INVERTED_Z_BUFFER
 	}
+
+#endif
 }
 
 float VelocityThreshold;
diff --git a/Engine/Shaders/Private/HairStrands/HairStrandsVisibilityDepthPS.usf b/Engine/Shaders/Private/HairStrands/HairStrandsVisibilityDepthPS.usf
--- a/Engine/Shaders/Private/HairStrands/HairStrandsVisibilityDepthPS.usf
+++ b/Engine/Shaders/Private/HairStrands/HairStrandsVisibilityDepthPS.usf
@@ -7,7 +7,11 @@
 #include "../DeferredShadingCommon.ush"
 #include "HairStrandsVisibilityCommon.ush"
 
+#if PERMUTATION_PPLL
+Texture2D<uint4> CategorisationTexture;
+#else
 Texture2DMS<float, HAIR_VISIBILITY_MAXSAMPLECOUNT> HairVisibilityDepthTexture;
+#endif
 
 #if PERMUTATION_COVERAGE == 1
 Texture2D<float> CoverageTexture;
@@ -20,6 +24,17 @@ void MainPS(
 	out float4 OutColor		: SV_Target1
 )
 {
+#if PERMUTATION_PPLL
+
+	const FCategorizationData CategorisationData = DecodeCategorizationData(CategorisationTexture.Load(uint3(Input.Position.xy, 0)));
+	if (CategorisationData.PixelCoverage < 1.0f)
+	{
+		discard;
+	}
+	SvDepth = CategorisationData.ClosestDepth;
+
+#else // PERMUTATION_PPLL
+
 	// HAS_INVERTED_Z_BUFFER
 	float ClosestDepth = 0;
 	float FarthestDepth = 1;
@@ -49,6 +64,8 @@ void MainPS(
 	SvDepth = FarthestDepth; // Farthest is fine to keep because all samples contains hair data
 #endif
 
+#endif // PERMUTATION_PPLL
+
 	OutGBufferB = 0;
 	OutColor = 0;
 }
diff --git a/Engine/Shaders/Private/HairStrands/HairStrandsVisibilityPPLLDebug.usf b/Engine/Shaders/Private/HairStrands/HairStrandsVisibilityPPLLDebug.usf
new file mode 100644
--- /dev/null
+++ b/Engine/Shaders/Private/HairStrands/HairStrandsVisibilityPPLLDebug.usf
@@ -0,0 +1,37 @@
+// Copyright 1998-2019 Epic Games, Inc. All Rights Reserved.
+
+#include "../Common.ush"
+
+#include "HairStrandsVisibilityCommon.ush"
+
+uint MaxPPLLNodeCount;
+Texture2D<uint> PPLLCounter;
+Texture2D<uint> PPLLNodeIndex;
+StructuredBuffer<FPPLLNodeData> PPLLNodeData;
+
+void VisibilityDebugPPLLPS(
+	in float4 SVPos : SV_POSITION,
+	out float4 OutLuminance : SV_Target0)
+{
+	uint2 PixPos = SVPos.xy;
+
+	uint NextNodeIndex = PPLLNodeIndex[PixPos];
+	if (NextNodeIndex != 0xFFFFFFFF)
+	{
+		uint StrandCount = 0;
+		do
+		{
+			StrandCount++;
+			FPPLLNodeData Node = PPLLNodeData[NextNodeIndex];
+			NextNodeIndex = Node.NextNodeIndex;
+		} while (NextNodeIndex != 0xFFFFFFFF);
+
+		float l = saturate(StrandCount / 16.0f);
+
+		OutLuminance = float4(l, 1.0f - l, 0.0f, 0.0f);
+	}
+	else
+	{
+		OutLuminance = float4(0.0f, 0.0f, 0.0f, 1.0f);
+	}
+}
diff --git a/Engine/Shaders/Private/HairStrands/HairStrandsVisibilityPS.usf b/Engine/Shaders/Private/HairStrands/HairStrandsVisibilityPS.usf
--- a/Engine/Shaders/Private/HairStrands/HairStrandsVisibilityPS.usf
+++ b/Engine/Shaders/Private/HairStrands/HairStrandsVisibilityPS.usf
@@ -8,12 +8,14 @@
 #include "HairStrandsDeepShadowCommon.ush"
 #include "../DeferredShadingCommon.ush"
 #include "HairStrandsVisibilityCommon.ush"
-#include "../VelocityCommon.ush"  
+#include "../VelocityCommon.ush"
+
 #include "/Engine/Generated/Material.ush"
 #include "/Engine/Generated/VertexFactory.ush"
 
 #define RENDER_MODE_MSAA 0
 #define RENDER_MODE_COVERAGE 1
+#define RENDER_MODE_PPLL 2
 
 #ifndef HAIR_RENDER_MODE
 #error Undefined hair render mode
@@ -21,7 +23,7 @@
 
 #define SUPPORT_OPACITY_MASK 0
 
-#if HAIR_RENDER_MODE == RENDER_MODE_MSAA
+#if HAIR_RENDER_MODE == RENDER_MODE_MSAA || HAIR_RENDER_MODE == RENDER_MODE_PPLL
 #define SUPPORT_MATERIAL_PROPERTY 1
 #endif
 
@@ -29,23 +31,36 @@
 #define SUPPORT_MATERIAL_PROPERTY 0
 #endif
 
+#if HAIR_RENDER_MODE == RENDER_MODE_PPLL 
+uint HairVisibilityPass_MaxPPLLNodeCount;
+RWTexture2D<uint> PPLLCounter;
+RWTexture2D<uint> PPLLNodeIndex;
+RWStructuredBuffer<FPPLLNodeData> PPLLNodes;
+#endif
 
 uint HairVisibilityPass_HairClusterIndex;
 
+#if HAIR_RENDER_MODE == RENDER_MODE_PPLL
+EARLYDEPTHSTENCIL	// Required for UAV operations to not happen.
+#endif
 void MainVisibility(
 	centroid FVertexFactoryInterpolantsVSToPS Interpolants,
-	in float4 SvPosition	: SV_Position,
-#if HAIR_RENDER_MODE == RENDER_MODE_MSAA
+	in float4 SvPosition	: SV_Position
+#if HAIR_RENDER_MODE == RENDER_MODE_MSAA || HAIR_RENDER_MODE == RENDER_MODE_PPLL
+	,
 	centroid in float WorldStrandRadius : TEXCOORD8,
 	const float4 PackedVelocityA : TEXCOORD9,
 	const float4 PackedVelocityC : TEXCOORD12
-#elif HAIR_RENDER_MODE == RENDER_MODE_COVERAGE
-	centroid in float WorldStrandRadius : TEXCOORD8
+#elif HAIR_RENDER_MODE == RENDER_MODE_COVERAGE 
+	,
+	in float WorldStrandRadius : TEXCOORD8
 #endif
-	OPTIONAL_IsFrontFace,
+	OPTIONAL_IsFrontFace
 #if HAIR_RENDER_MODE == RENDER_MODE_COVERAGE
+	,
 	out float OutColor0 : SV_Target0
 #elif HAIR_RENDER_MODE == RENDER_MODE_MSAA
+	,
 	out uint4 OutColor0	: SV_Target0
   #if SUPPORT_MATERIAL_PROPERTY
 	, out float4 OutColor1	: SV_Target1
@@ -79,7 +94,7 @@ void MainVisibility(
 	OutColor0 = Coverage;
 #endif
 
-#if HAIR_RENDER_MODE == RENDER_MODE_MSAA
+#if HAIR_RENDER_MODE == RENDER_MODE_MSAA || HAIR_RENDER_MODE == RENDER_MODE_PPLL
 	FVisibilityData VisibilityData;
   #if VF_STRAND_HAIR
 	VisibilityData.Tangent = normalize(Interpolants.TangentToWorld2.xyz);
@@ -95,7 +110,7 @@ void MainVisibility(
 	float Roughness = 0;
 	float Specular  = 0;
 	float2 EncodedVelocity = 0;
-	#if SUPPORT_MATERIAL_PROPERTY
+#if SUPPORT_MATERIAL_PROPERTY
 	{
 		FMaterialPixelParameters MaterialParameters = GetMaterialPixelParameters(Interpolants, SvPosition);
 		FPixelMaterialInputs PixelMaterialInputs;
@@ -116,13 +131,36 @@ void MainVisibility(
 		const float2 Velocity = Calculate2DVelocity(PackedVelocityA, PackedVelocityC);
 		EncodedVelocity = EncodeVelocityToTexture(Velocity);
 	}
-	#endif
+#endif
+#endif
+
+#if HAIR_RENDER_MODE == RENDER_MODE_MSAA
 	OutColor0 = EncodeVisibilityData(VisibilityData);
 	OutColor1 = float4(sqrt(BaseColor), Roughness);
 	OutColor2 = float4(Specular,0,0,0);
 	OutColor3 = EncodedVelocity;
 #endif
 
+#if HAIR_RENDER_MODE == RENDER_MODE_PPLL
+	const uint2 PixelCoord = SvPosition.xy;
+
+	float2 NDC = SvPosition.xy / SvPosition.w;
+	float2 ScreenUV = NDC * ResolvedView.ScreenPositionScaleBias.xy + ResolvedView.ScreenPositionScaleBias.wz;
+
+	// Allocate space for a new node
+	uint NodeIndex;
+	InterlockedAdd(PPLLCounter[uint2(0, 0)], 1, NodeIndex);
+
+	if (NodeIndex < HairVisibilityPass_MaxPPLLNodeCount)
+	{
+		// If we can append new node, add it and make it point to the node previously in head of list
+		uint NextNodeIndex;
+		InterlockedExchange(PPLLNodeIndex[PixelCoord], NodeIndex, NextNodeIndex);
+		const FPPLLNodeData Node = ConvertToPPLLNodeData(SvPosition.z, VisibilityData, BaseColor, Roughness, NextNodeIndex);
+		PPLLNodes[NodeIndex] = Node;
+	}
+#endif
+
 }
 
 
diff --git a/Engine/Shaders/Private/HairStrands/HairStrandsVisibilityVS.usf b/Engine/Shaders/Private/HairStrands/HairStrandsVisibilityVS.usf
--- a/Engine/Shaders/Private/HairStrands/HairStrandsVisibilityVS.usf
+++ b/Engine/Shaders/Private/HairStrands/HairStrandsVisibilityVS.usf
@@ -13,7 +13,7 @@ struct FHairVisibilityAccumulateVSToPS
 	FVertexFactoryInterpolantsVSToPS FactoryInterpolants;
 	float4 Position : SV_POSITION;
 	float WorldStrandRadius : TEXCOORD8;
-#if HAIR_RENDER_MODE == RENDER_MODE_MSAA
+#if HAIR_RENDER_MODE == RENDER_MODE_MSAA || HAIR_RENDER_MODE == RENDER_MODE_PPLL
 	float4 PackedVelocityA : TEXCOORD9;
 	float4 PackedVelocityC : TEXCOORD12;
 #endif
@@ -47,7 +47,7 @@ void Main(
 	// The velocity is used to adapt/increase the minimum rasterization size in order to avoid aliasing under heavy-motion.
 	// The higher the velocity, the larger a strand becomes.
 	float VelocityScale = 0;
-	#if HAIR_RENDER_MODE == RENDER_MODE_MSAA
+	#if HAIR_RENDER_MODE == RENDER_MODE_MSAA || HAIR_RENDER_MODE == RENDER_MODE_PPLL
 	{
 		const float4 ScreenPos			= mul(float4(WorldPositionRaw.xyz, 1), ResolvedView.TranslatedWorldToClip);
 		const float4 PrevScreenPosObj	= mul(float4(PrevWorldPosition.xyz, 1), ResolvedView.PrevTranslatedWorldToClip);
diff --git a/Engine/Shaders/Public/ShaderVersion.ush b/Engine/Shaders/Public/ShaderVersion.ush
--- a/Engine/Shaders/Public/ShaderVersion.ush
+++ b/Engine/Shaders/Public/ShaderVersion.ush
@@ -3,4 +3,4 @@
 // in Platform.ush (which should be included in any shader) it allows to invalidate the shader DDC.
 // 
 // If you are merging streams and there is a conflict with this GUID you should make a new GUID rather than taking one or the other.
-// GUID = FDD7186B-678E-403A-8A1C-F24AAEDE76C5
+// GUID = 6925F773-21F5-4603-97EA-A8D76B115A4E
diff --git a/Engine/Source/Runtime/Renderer/Private/HairStrands/HairStrandsComposition.cpp b/Engine/Source/Runtime/Renderer/Private/HairStrands/HairStrandsComposition.cpp
--- a/Engine/Source/Runtime/Renderer/Private/HairStrands/HairStrandsComposition.cpp
+++ b/Engine/Source/Runtime/Renderer/Private/HairStrands/HairStrandsComposition.cpp
@@ -21,9 +21,13 @@ class FHairVisibilityComposeSubPixelPS : public FGlobalShader
 	DECLARE_GLOBAL_SHADER(FHairVisibilityComposeSubPixelPS);
 	SHADER_USE_PARAMETER_STRUCT(FHairVisibilityComposeSubPixelPS, FGlobalShader);
 
+	class FPPLL : SHADER_PERMUTATION_INT("PERMUTATION_PPLL", 2);
+	using FPermutationDomain = TShaderPermutationDomain<FPPLL>;
+
 	BEGIN_SHADER_PARAMETER_STRUCT(FParameters, )
 		SHADER_PARAMETER_RDG_TEXTURE(Texture2D, SubPixelColorTexture)
 		SHADER_PARAMETER_RDG_TEXTURE(Texture2D, HairVisibilityDepthTexture)
+		SHADER_PARAMETER_RDG_TEXTURE(Texture2D, CategorisationTexture)
 		RENDER_TARGET_BINDING_SLOTS()
 	END_SHADER_PARAMETER_STRUCT()
 
@@ -38,22 +42,23 @@ static void AddHairVisibilityComposeSubPixelPass(
 	const FViewInfo& View,
 	const FRDGTextureRef& SubPixelColorTexture,
 	const FRDGTextureRef& HairVisibilityDepthTexture,
+	const FRDGTextureRef CategorisationTexture,
 	FRDGTextureRef& OutColorTexture,
 	FRDGTextureRef& OutDepthTexture)
 {
+	const bool bPPLL = HairVisibilityDepthTexture == nullptr;
+
 	FHairVisibilityComposeSubPixelPS::FParameters* Parameters = GraphBuilder.AllocParameters<FHairVisibilityComposeSubPixelPS::FParameters>();
 	Parameters->SubPixelColorTexture = SubPixelColorTexture;
+	Parameters->CategorisationTexture = CategorisationTexture;
 	Parameters->HairVisibilityDepthTexture = HairVisibilityDepthTexture;
 	Parameters->RenderTargets[0] = FRenderTargetBinding(OutColorTexture, ERenderTargetLoadAction::ELoad);
-	Parameters->RenderTargets.DepthStencil = FDepthStencilBinding(
-		OutDepthTexture,
-		ERenderTargetLoadAction::ELoad,
-		ERenderTargetLoadAction::ENoAction,
-		FExclusiveDepthStencil::DepthWrite_StencilNop);
-	
+	Parameters->RenderTargets.DepthStencil = FDepthStencilBinding(OutDepthTexture, ERenderTargetLoadAction::ELoad, ERenderTargetLoadAction::ENoAction, FExclusiveDepthStencil::DepthWrite_StencilNop);
+
+	FHairVisibilityComposeSubPixelPS::FPermutationDomain PermutationVector;
+	PermutationVector.Set<FHairVisibilityComposeSubPixelPS::FPPLL>(bPPLL ? 1 : 0);	// TODO we need a permutation for each MSAA count we accept
+	TShaderMapRef<FHairVisibilityComposeSubPixelPS> PixelShader(View.ShaderMap, PermutationVector);
 	TShaderMapRef<FPostProcessVS> VertexShader(View.ShaderMap);
-	TShaderMapRef<FHairVisibilityComposeSubPixelPS> PixelShader(View.ShaderMap);
-	const TShaderMap<FGlobalShaderType>* GlobalShaderMap = View.ShaderMap;
 	const FIntRect Viewport = View.ViewRect;
 	const FIntPoint Resolution = OutColorTexture->Desc.Extent;
 	const FViewInfo* CapturedView = &View;
@@ -228,11 +233,18 @@ void RenderHairComposeSubPixel(
 		const FViewInfo& View = Views[ViewIndex];
 		if (View.Family)
 		{
-			if (ViewIndex < HairVisibilityViews.HairDatas.Num() && HairVisibilityViews.HairDatas[ViewIndex].DepthTexture)
+			if (ViewIndex < HairVisibilityViews.HairDatas.Num())
 			{
-				FRDGTextureRef HairVisibilityDepthTexture = GraphBuilder.RegisterExternalTexture(HairVisibilityViews.HairDatas[ViewIndex].DepthTexture, TEXT("HairVisibilityDepthTexture"));
-				FRDGTextureRef HairVisibilityVelocityTexture = GraphBuilder.RegisterExternalTexture(HairVisibilityViews.HairDatas[ViewIndex].VelocityTexture, TEXT("HairVisibilityVelocityTexture"));
-				
+				TRefCountPtr<IPooledRenderTarget> DepthTexture = HairVisibilityViews.HairDatas[ViewIndex].DepthTexture;
+				TRefCountPtr<IPooledRenderTarget> CategorisationTexture = HairVisibilityViews.HairDatas[ViewIndex].CategorizationTexture;
+				if (!DepthTexture && !CategorisationTexture)
+				{
+					continue; // Automatically skip for any view not rendering hair
+				}
+
+				const FRDGTextureRef RDGHairVisibilityDepthTexture = DepthTexture ? GraphBuilder.RegisterExternalTexture(DepthTexture, TEXT("HairVisibilityDepthTexture")) : nullptr;
+				const FRDGTextureRef RDGCategorisationTexture = CategorisationTexture ? GraphBuilder.RegisterExternalTexture(CategorisationTexture, TEXT("HairVisibilityCategorisationTexture")) : nullptr;
+
 				// #hair_todo : compose partially covered hair with transparent surface: this can be done by 
 				// rendering quad(s) covering the hair at the correct depth. This will be sorted with other 
 				// transparent surface, which should make the overall sorting workable
@@ -240,16 +252,21 @@ void RenderHairComposeSubPixel(
 					GraphBuilder,
 					View,
 					SceneColorSubPixelTexture,
-					HairVisibilityDepthTexture,
+					RDGHairVisibilityDepthTexture,
+					RDGCategorisationTexture,
 					SceneColorTexture,
 					SceneColorDepth);
 
-				AddHairVisibilityFastResolvePass(
-					GraphBuilder,
-					View,
-					HairVisibilityVelocityTexture,
-					HairVisibilityDepthTexture,
-					SceneColorDepth);
+				if (HairVisibilityViews.HairDatas[ViewIndex].VelocityTexture)
+				{
+					FRDGTextureRef RDGHairVisibilityVelocityTexture = GraphBuilder.RegisterExternalTexture(HairVisibilityViews.HairDatas[ViewIndex].VelocityTexture, TEXT("HairVisibilityVelocityTexture"));
+					AddHairVisibilityFastResolvePass(
+						GraphBuilder,
+						View,
+						RDGHairVisibilityVelocityTexture,
+						RDGHairVisibilityDepthTexture,
+						SceneColorDepth);
+				}
 			}
 		}
 	}
diff --git a/Engine/Source/Runtime/Renderer/Private/HairStrands/HairStrandsDebug.cpp b/Engine/Source/Runtime/Renderer/Private/HairStrands/HairStrandsDebug.cpp
--- a/Engine/Source/Runtime/Renderer/Private/HairStrands/HairStrandsDebug.cpp
+++ b/Engine/Source/Runtime/Renderer/Private/HairStrands/HairStrandsDebug.cpp
@@ -22,6 +22,7 @@
 #include "RenderTargetTemp.h"
 #include "CanvasTypes.h"
 #include "ShaderPrintParameters.h"
+#include "PixelShaderUtils.h"
 
 static int32 GDeepShadowDebugIndex = 0;
 static float GDeepShadowDebugScale = 20;
@@ -70,6 +71,8 @@ static FAutoConsoleVariableRef CVarHairDebugMeshProjection_Sim_HairRestFrames(
 static FAutoConsoleVariableRef CVarHairDebugMeshProjection_Sim_HairDeformedTriangles(		TEXT("r.HairStrands.MeshProjection.Sim.DebugHairDeformedTriangles"),	GHairDebugMeshProjection_Sim_HairDeformedTriangles, TEXT("Render debug mes projection"));
 static FAutoConsoleVariableRef CVarHairDebugMeshProjection_Sim_HairDeformedFrames(			TEXT("r.HairStrands.MeshProjection.Sim.DebugHairDeformedFrames"),		GHairDebugMeshProjection_Sim_HairDeformedFrames, TEXT("Render debug mes projection"));
 
+static int32 GHairStrandsDebugPPLL = 0;
+static FAutoConsoleVariableRef CVarHairStrandsDebugPPLL(									TEXT("r.HairStrands.DebugPPLL"),										GHairStrandsDebugPPLL, TEXT("Draw debug per pixel light list rendering."));
 
 // Helper functions for accessing interpolation data for debug purpose.
 // Definitions is in HairStrandsInterface.cpp
@@ -1014,6 +1017,44 @@ const TCHAR* ToString(EWorldType::Type Type)
 	}
 }
 
+///////////////////////////////////////////////////////////////////////////////////////////////////
+
+class FHairVisibilityDebugPPLLPS : public FGlobalShader
+{
+	DECLARE_GLOBAL_SHADER(FHairVisibilityDebugPPLLPS);
+	SHADER_USE_PARAMETER_STRUCT(FHairVisibilityDebugPPLLPS, FGlobalShader);
+
+	using FPermutationDomain = TShaderPermutationDomain<>;
+
+	BEGIN_SHADER_PARAMETER_STRUCT(FParameters, )
+		SHADER_PARAMETER_STRUCT_REF(FViewUniformShaderParameters, ViewUniformBuffer)
+		SHADER_PARAMETER(uint32, MaxPPLLNodeCount)
+		SHADER_PARAMETER_RDG_TEXTURE(Texture2D, PPLLCounter)
+		SHADER_PARAMETER_RDG_TEXTURE(Texture2D, PPLLNodeIndex)
+		SHADER_PARAMETER_RDG_BUFFER_SRV(StructuredBuffer, PPLLNodeData)
+		RENDER_TARGET_BINDING_SLOTS()
+	END_SHADER_PARAMETER_STRUCT()
+
+		static FPermutationDomain RemapPermutation(FPermutationDomain PermutationVector)
+	{
+		return PermutationVector;
+	}
+
+	static bool ShouldCompilePermutation(const FGlobalShaderPermutationParameters& Parameters)
+	{
+		return IsHairStrandsSupported(Parameters.Platform);
+	}
+
+	static void ModifyCompilationEnvironment(const FGlobalShaderPermutationParameters& Parameters, FShaderCompilerEnvironment& OutEnvironment)
+	{
+		FGlobalShader::ModifyCompilationEnvironment(Parameters, OutEnvironment);
+		OutEnvironment.SetDefine(TEXT("DEBUG_PPLL_PS"), 1);
+	}
+};
+IMPLEMENT_GLOBAL_SHADER(FHairVisibilityDebugPPLLPS, "/Engine/Private/HairStrands/HairStrandsVisibilityPPLLDebug.usf", "VisibilityDebugPPLLPS", SF_Pixel);
+
+///////////////////////////////////////////////////////////////////////////////////////////////////
+
 void RenderHairStrandsDebugInfo(FRHICommandListImmediate& RHICmdList, TArray<FViewInfo>& Views, const FHairStrandsDatas* HairDatas)
 {
 	const float YStep = 14;
@@ -1346,6 +1387,38 @@ void RenderHairStrandsDebugInfo(FRHICommandListImmediate& RHICmdList, TArray<FVi
 		GraphBuilder.Execute();
 	}
 
+	if (ViewIndex < uint32(HairDatas->HairVisibilityViews.HairDatas.Num()))
+	{
+		const FHairStrandsVisibilityData& VisibilityData = HairDatas->HairVisibilityViews.HairDatas[ViewIndex];
+		if (GHairStrandsDebugPPLL && VisibilityData.PPLLNodeCounterTexture) // Check if PPLL rendering is used and its debuig view is enabled.
+		{
+			FRDGBuilder GraphBuilder(RHICmdList);
+			FRDGTextureRef SceneColorTexture = GraphBuilder.RegisterExternalTexture(SceneTargets.GetSceneColor(), TEXT("SceneColorTexture"));
+			FRDGTextureRef PPLLNodeCounterTexture = GraphBuilder.RegisterExternalTexture(VisibilityData.PPLLNodeCounterTexture, TEXT("PPLLNodeCounterTexture"));
+			FRDGTextureRef PPLLNodeIndexTexture = GraphBuilder.RegisterExternalTexture(VisibilityData.PPLLNodeIndexTexture, TEXT("PPLLNodeIndexTexture"));
+			FRDGBufferRef  PPLLNodeDataBuffer = GraphBuilder.RegisterExternalBuffer(VisibilityData.PPLLNodeDataBuffer, TEXT("PPLLNodeDataBuffer"));
+
+			FRHIBlendState* PreMultipliedColorTransmittanceBlend = TStaticBlendState<CW_RGB, BO_Add, BF_One, BF_SourceAlpha, BO_Add, BF_Zero, BF_One>::GetRHI();
+			FRHIDepthStencilState* DepthStencilStateWrite = TStaticDepthStencilState<false, CF_Always>::GetRHI();
+			FRHISamplerState* SamplerLinearClamp = TStaticSamplerState<SF_Trilinear>::GetRHI();
+
+			FHairVisibilityDebugPPLLPS::FPermutationDomain PermutationVector;
+			TShaderMapRef<FHairVisibilityDebugPPLLPS> PixelShader(View.ShaderMap, PermutationVector);
+
+			FHairVisibilityDebugPPLLPS::FParameters* PassParameters = GraphBuilder.AllocParameters<FHairVisibilityDebugPPLLPS::FParameters>();
+			PassParameters->MaxPPLLNodeCount = GetMaxNodePerPixel(VisibilityData.PPLLNodeCounterTexture->GetDesc().Extent);
+			PassParameters->PPLLCounter = PPLLNodeCounterTexture;
+			PassParameters->PPLLNodeIndex = PPLLNodeIndexTexture;
+			PassParameters->PPLLNodeData = GraphBuilder.CreateSRV(FRDGBufferSRVDesc(PPLLNodeDataBuffer));
+			PassParameters->ViewUniformBuffer = View.ViewUniformBuffer;
+			PassParameters->RenderTargets[0] = FRenderTargetBinding(SceneColorTexture, ERenderTargetLoadAction::ELoad);
+
+			FPixelShaderUtils::AddFullscreenPass<FHairVisibilityDebugPPLLPS>(GraphBuilder, View.ShaderMap, RDG_EVENT_NAME("HairPPLLDebug"), *PixelShader, PassParameters,
+				View.ViewRect, PreMultipliedColorTransmittanceBlend, nullptr, DepthStencilStateWrite);
+			GraphBuilder.Execute();
+		}
+	}
+
 	// Text
 	if (HairDebugMode == EHairDebugMode::LightBounds || HairDebugMode == EHairDebugMode::DeepOpacityMaps)
 	{
diff --git a/Engine/Source/Runtime/Renderer/Private/HairStrands/HairStrandsVisibility.cpp b/Engine/Source/Runtime/Renderer/Private/HairStrands/HairStrandsVisibility.cpp
--- a/Engine/Source/Runtime/Renderer/Private/HairStrands/HairStrandsVisibility.cpp
+++ b/Engine/Source/Runtime/Renderer/Private/HairStrands/HairStrandsVisibility.cpp
@@ -36,21 +36,32 @@ static FAutoConsoleVariableRef CVarHairVisibilitySampleCount(TEXT("r.HairStrands
 static int32 GHairClearVisibilityBuffer = 0;
 static FAutoConsoleVariableRef CVarHairClearVisibilityBuffer(TEXT("r.HairStrands.VisibilityClear"), GHairClearVisibilityBuffer, TEXT("Clear hair strands visibility buffer"));
 
+static TAutoConsoleVariable<int32> CVarHairVelocityMagnitudeScale(
+	TEXT("r.HairStrands.VelocityMagnitudeScale"),
+	100,  // Tuned by eye, based on heavy motion (strong head shack)
+	TEXT("Velocity magnitude (in pixel) at which a hair will reach its pic velocity-rasterization-scale under motion to reduce aliasing. Default is 100."));
+
 static int32 GHairVelocityType = 1; // default is 
 static FAutoConsoleVariableRef CVarHairVelocityType(TEXT("r.HairStrands.VelocityType"), GHairVelocityType, TEXT("Type of velocity filtering (0:avg, 1:closest, 2:max). Default is 1."));
 
+static int32 GHairVisibilityPPLL = 0;
+static FAutoConsoleVariableRef CVarGHairVisibilityPPLL(TEXT("r.HairStrands.VisibilityPPLL"), GHairVisibilityPPLL, TEXT("Hair Visibility uses per pixel linked list"));
+static int32 GHairVisibilityPPLLMaxPixelNodeCount = 16;
+static FAutoConsoleVariableRef CVarGHairVisibilityPPLLMaxPixelNodeCount(TEXT("r.HairStrands.VisibilityPPLLMaxPixelNodeCount"), GHairVisibilityPPLLMaxPixelNodeCount, TEXT("The total maximum number of node allowed for all linked list will be width*height*VisibilityPPLLMaxPixelNodeCount."));
+
 /////////////////////////////////////////////////////////////////////////////////////////
 
 enum EHairVisibilityRenderMode
 {
 	HairVisibilityRenderMode_MSAA,
 	HairVisibilityRenderMode_Coverage,
+	HairVisibilityRenderMode_PPLL,
 	HairVisibilityRenderModeCount
 };
 
 EHairVisibilityRenderMode GetHairVisibilityRenderMode()
 {
-	return HairVisibilityRenderMode_MSAA;
+	return GHairVisibilityPPLL >0 ? HairVisibilityRenderMode_PPLL : HairVisibilityRenderMode_MSAA;
 }
 
 uint32 GetHairVisibilitySampleCount()
@@ -58,11 +69,42 @@ uint32 GetHairVisibilitySampleCount()
 	return GetHairVisibilityRenderMode() == HairVisibilityRenderMode_MSAA ? FMath::Clamp(GHairVisibilitySampleCount,1,16) : 1;
 }
 
+void SetUpViewHairRenderInfo(const FViewInfo& ViewInfo, bool bEnableMSAA, FVector4& OutHairRenderInfo)
+{
+	FVector2D PixelVelocity(1.f / (ViewInfo.ViewRect.Width() * 2), 1.f / (ViewInfo.ViewRect.Height() * 2));
+	const float VelocityMagnitudeScale = FMath::Clamp(CVarHairVelocityMagnitudeScale.GetValueOnAnyThread(), 0, 512) * FMath::Min(PixelVelocity.X, PixelVelocity.Y);
+
+	// In the case we render coverage, we need to override some view uniform shader parameters to account for the change in MSAA sample count.
+	const uint32 HairVisibilitySampleCount = bEnableMSAA ? GetHairVisibilitySampleCount() : 1;	// The coverage pass does not use MSAA
+	const float RasterizationScaleOverride = 0.0f;	// no override
+	FMinHairRadiusAtDepth1 MinHairRadius = ComputeMinStrandRadiusAtDepth1(
+		FIntPoint(ViewInfo.UnconstrainedViewRect.Width(), ViewInfo.UnconstrainedViewRect.Height()), ViewInfo.FOV, HairVisibilitySampleCount, RasterizationScaleOverride);
+
+	TUniformBufferRef<FViewUniformShaderParameters> ViewUniformShaderParameters;
+	// Update our view parameters
+	OutHairRenderInfo.X = MinHairRadius.Primary;
+	OutHairRenderInfo.Y = MinHairRadius.Velocity;
+	OutHairRenderInfo.Z = ViewInfo.IsPerspectiveProjection() ? 0.0f : 1.0f;
+	OutHairRenderInfo.W = VelocityMagnitudeScale;
+}
+
+uint32 GetMaxNodePerPixel(FIntPoint Resolution)
+{
+	return Resolution.X * Resolution.Y * GHairVisibilityPPLLMaxPixelNodeCount;
+}
+
 static bool IsCompatibleWithHairVisibility(const FMeshMaterialShaderPermutationParameters& Parameters)
 {
 	return IsCompatibleWithHairStrands(Parameters.Material, Parameters.Platform);
 }
 
+BEGIN_SHADER_PARAMETER_STRUCT(FVisibilityPassParameters, )
+	SHADER_PARAMETER(uint32, HairVisibilityPass_MaxPPLLNodeCount)
+	SHADER_PARAMETER_RDG_TEXTURE_UAV(Texture2D, PPLLCounter)
+	SHADER_PARAMETER_RDG_TEXTURE_UAV(Texture2D, PPLLNodeIndex)
+	SHADER_PARAMETER_RDG_BUFFER_UAV(StructuredBuffer, PPLLNodeData)
+	RENDER_TARGET_BINDING_SLOTS()
+END_SHADER_PARAMETER_STRUCT()
 
 template<EHairVisibilityRenderMode RenderMode>
 class FHairVisibilityVS : public FMeshMaterialShader
@@ -94,14 +136,16 @@ protected:
 };
 IMPLEMENT_MATERIAL_SHADER_TYPE(template<>, FHairVisibilityVS<HairVisibilityRenderMode_MSAA>, TEXT("/Engine/Private/HairStrands/HairStrandsVisibilityVS.usf"), TEXT("Main"), SF_Vertex);
 IMPLEMENT_MATERIAL_SHADER_TYPE(template<>, FHairVisibilityVS<HairVisibilityRenderMode_Coverage>, TEXT("/Engine/Private/HairStrands/HairStrandsVisibilityVS.usf"), TEXT("Main"), SF_Vertex);
+IMPLEMENT_MATERIAL_SHADER_TYPE(template<>, FHairVisibilityVS<HairVisibilityRenderMode_PPLL>, TEXT("/Engine/Private/HairStrands/HairStrandsVisibilityVS.usf"), TEXT("Main"), SF_Vertex);
 
 /////////////////////////////////////////////////////////////////////////////////////////
 
 class FHairVisibilityShaderElementData : public FMeshMaterialShaderElementData
 {
 public:
-	FHairVisibilityShaderElementData(uint32 InHairClusterId) : HairClusterId(InHairClusterId) { }
+	FHairVisibilityShaderElementData(uint32 InHairClusterId, uint32 InMaxPPLLNodeCount) : HairClusterId(InHairClusterId), MaxPPLLNodeCount(InMaxPPLLNodeCount) { }
 	uint32 HairClusterId;
+	uint32 MaxPPLLNodeCount;
 };
 
 template<EHairVisibilityRenderMode RenderMode>
@@ -117,6 +161,10 @@ public:
 		ERHIFeatureLevel::Type FeatureLevel = GetMaxSupportedFeatureLevel((EShaderPlatform)Initializer.Target.Platform);
 		check(FSceneInterface::GetShadingPath(FeatureLevel) != EShadingPath::Mobile);
 		HairVisibilityPass_HairClusterIndex.Bind(Initializer.ParameterMap, TEXT("HairVisibilityPass_HairClusterIndex"));
+		HairVisibilityPass_MaxPPLLNodeCount.Bind(Initializer.ParameterMap, TEXT("HairVisibilityPass_MaxPPLLNodeCount"));
+		PPLLCounter.Bind(Initializer.ParameterMap, TEXT("PPLLCounter"));
+		PPLLNodeIndex.Bind(Initializer.ParameterMap, TEXT("PPLLNodeIndex"));
+		PPLLNodes.Bind(Initializer.ParameterMap, TEXT("PPLLNodes"));
 	}
 
 	FHairVisibilityPS() {}
@@ -137,6 +185,10 @@ public:
 	{
 		bool bShaderHasOutdatedParameters = FMeshMaterialShader::Serialize(Ar);
 		Ar << HairVisibilityPass_HairClusterIndex;
+		Ar << HairVisibilityPass_MaxPPLLNodeCount;
+		Ar << PPLLCounter;
+		Ar << PPLLNodeIndex;
+		Ar << PPLLNodes;
 		return bShaderHasOutdatedParameters;
 	}
 
@@ -152,12 +204,19 @@ public:
 	{
 		FMeshMaterialShader::GetShaderBindings(Scene, FeatureLevel, PrimitiveSceneProxy, MaterialRenderProxy, Material, DrawRenderState, ShaderElementData, ShaderBindings);
 		ShaderBindings.Add(HairVisibilityPass_HairClusterIndex, ShaderElementData.HairClusterId);
+		ShaderBindings.Add(HairVisibilityPass_MaxPPLLNodeCount, ShaderElementData.MaxPPLLNodeCount);
 	}
 
 	FShaderParameter HairVisibilityPass_HairClusterIndex;
+	FShaderParameter HairVisibilityPass_MaxPPLLNodeCount;
+	// This is on FVisibilityPassParameters but needed to avoid the compiler to assert with unbound parameters 
+	FShaderParameter PPLLCounter;
+	FShaderParameter PPLLNodeIndex;
+	FShaderParameter PPLLNodes;
 };
 IMPLEMENT_MATERIAL_SHADER_TYPE(template<>, FHairVisibilityPS<HairVisibilityRenderMode_MSAA>, TEXT("/Engine/Private/HairStrands/HairStrandsVisibilityPS.usf"), TEXT("MainVisibility"), SF_Pixel);
 IMPLEMENT_MATERIAL_SHADER_TYPE(template<>, FHairVisibilityPS<HairVisibilityRenderMode_Coverage>, TEXT("/Engine/Private/HairStrands/HairStrandsVisibilityPS.usf"), TEXT("MainVisibility"), SF_Pixel);
+IMPLEMENT_MATERIAL_SHADER_TYPE(template<>, FHairVisibilityPS<HairVisibilityRenderMode_PPLL>, TEXT("/Engine/Private/HairStrands/HairStrandsVisibilityPS.usf"), TEXT("MainVisibility"), SF_Pixel);
 
 /////////////////////////////////////////////////////////////////////////////////////////
 
@@ -174,6 +233,8 @@ public:
 	virtual void AddMeshBatch(const FMeshBatch& RESTRICT MeshBatch, uint64 BatchElementMask, const FPrimitiveSceneProxy* RESTRICT PrimitiveSceneProxy, int32 StaticMeshId = -1) override final;
 	void AddMeshBatch(const FMeshBatch& RESTRICT MeshBatch, uint64 BatchElementMask, const FPrimitiveSceneProxy* RESTRICT PrimitiveSceneProxy, int32 StaticMeshId, uint32 HairClusterId);
 
+	void SetMaxPPLLNodeCount(uint32 NewValue) { MaxPPLLNodeCount = NewValue; }
+
 private:
 	template<EHairVisibilityRenderMode RenderMode>
 	void Process(
@@ -189,6 +250,7 @@ private:
 
 	const EHairVisibilityRenderMode RenderMode;
 	FMeshPassProcessorRenderState PassDrawRenderState;
+	uint32 MaxPPLLNodeCount = 0;
 };
 
 void FHairVisibilityProcessor::AddMeshBatch(const FMeshBatch& RESTRICT MeshBatch, uint64 BatchElementMask, const FPrimitiveSceneProxy* RESTRICT PrimitiveSceneProxy, int32 StaticMeshId)
@@ -218,6 +280,8 @@ void FHairVisibilityProcessor::AddMeshBatch(const FMeshBatch& RESTRICT MeshBatch
 			Process<HairVisibilityRenderMode_MSAA>(MeshBatch, BatchElementMask, PrimitiveSceneProxy, StaticMeshId, MaterialRenderProxy, Material, HairClusterId, MeshFillMode, MeshCullMode);
 		if (RenderMode == HairVisibilityRenderMode_Coverage)
 			Process<HairVisibilityRenderMode_Coverage>(MeshBatch, BatchElementMask, PrimitiveSceneProxy, StaticMeshId, MaterialRenderProxy, Material, HairClusterId, MeshFillMode, MeshCullMode);
+		if (RenderMode == HairVisibilityRenderMode_PPLL)
+			Process<HairVisibilityRenderMode_PPLL>(MeshBatch, BatchElementMask, PrimitiveSceneProxy, StaticMeshId, MaterialRenderProxy, Material, HairClusterId, MeshFillMode, MeshCullMode);
 	}
 }
 
@@ -247,7 +311,7 @@ void FHairVisibilityProcessor::Process(
 	}
 
 	FMeshPassProcessorRenderState DrawRenderState(PassDrawRenderState);
-	FHairVisibilityShaderElementData ShaderElementData(HairClusterId);
+	FHairVisibilityShaderElementData ShaderElementData(HairClusterId, MaxPPLLNodeCount);
 	ShaderElementData.InitializeMeshMaterialData(ViewIfDynamicMeshCommand, PrimitiveSceneProxy, MeshBatch, StaticMeshId, false);
 
 	BuildMeshDrawCommands(
@@ -431,7 +495,8 @@ class FHairVisibilityPrimitiveIdCompactionCS : public FGlobalShader
 	class FVelocity : SHADER_PERMUTATION_INT("PERMUTATION_VELOCITY", 4);
 	class FCoverage : SHADER_PERMUTATION_INT("PERMUTATION_COVERAGE", 2);
 	class FMaterial : SHADER_PERMUTATION_INT("PERMUTATION_MATERIAL_COMPACTION", 2);
-	using FPermutationDomain = TShaderPermutationDomain<FVendor, FVelocity, FCoverage, FMaterial>;
+	class FPPLL     : SHADER_PERMUTATION_INT("PERMUTATION_PPLL", 2);
+	using FPermutationDomain = TShaderPermutationDomain<FVendor, FVelocity, FCoverage, FMaterial, FPPLL>;
 
 	BEGIN_SHADER_PARAMETER_STRUCT(FParameters, )
 		SHADER_PARAMETER(FIntPoint, OutputResolution)
@@ -441,11 +506,17 @@ class FHairVisibilityPrimitiveIdCompactionCS : public FGlobalShader
 		SHADER_PARAMETER(float, DepthTheshold)
 		SHADER_PARAMETER(float, CosTangentThreshold)
 
+		// Available for the MSAA path
 		SHADER_PARAMETER_RDG_TEXTURE(Texture2D, MSAA_DepthTexture)
 		SHADER_PARAMETER_RDG_TEXTURE(Texture2D, MSAA_IDTexture)
 		SHADER_PARAMETER_RDG_TEXTURE(Texture2D, MSAA_MaterialTexture)
 		SHADER_PARAMETER_RDG_TEXTURE(Texture2D, MSAA_AttributeTexture)
 		SHADER_PARAMETER_RDG_TEXTURE(Texture2D, MSAA_VelocityTexture)
+		// Available for the PPLL path
+		SHADER_PARAMETER_RDG_TEXTURE(Texture2D, PPLLCounter)
+		SHADER_PARAMETER_RDG_TEXTURE(Texture2D, PPLLNodeIndex)
+		SHADER_PARAMETER_RDG_BUFFER_SRV(StructuredBuffer, PPLLNodeData)
+
 		SHADER_PARAMETER_RDG_TEXTURE(Texture2D, CoverageTexture)
 		
 		SHADER_PARAMETER_RDG_TEXTURE_UAV(Texture2D, OutCompactNodeCounter)
@@ -466,16 +537,12 @@ public:
 IMPLEMENT_GLOBAL_SHADER(FHairVisibilityPrimitiveIdCompactionCS, "/Engine/Private/HairStrands/HairStrandsVisibilityCompaction.usf", "MainCS", SF_Compute);
 
 static void AddHairVisibilityPrimitiveIdCompactionPass(
+	const bool bUsePPLL,
 	FRDGBuilder& GraphBuilder,
 	const FViewInfo& View,
 	const FHairStrandsClusterDatas& ClusterDatas,
 	const uint32 NodeGroupSize,
-	const FRDGTextureRef& MSAA_DepthTexture,
-	const FRDGTextureRef& MSAA_IDTexture,
-	const FRDGTextureRef& MSAA_MaterialTexture,
-	const FRDGTextureRef& MSAA_AttributeTexture,
-	const FRDGTextureRef& MSAA_VelocityTexture,
-	const FRDGTextureRef& CoverageTexture,
+	FHairVisibilityPrimitiveIdCompactionCS::FParameters* PassParameters,
 	FRDGTextureRef& OutCompactNodeIndex,
 	FRDGBufferRef& OutCompactNodeData,
 	FRDGBufferRef& OutCompactNodeCoord,
@@ -483,12 +550,23 @@ static void AddHairVisibilityPrimitiveIdCompactionPass(
 	FRDGTextureRef& OutVelocityTexture,
 	FRDGBufferRef& OutIndirectArgsBuffer)
 {
-	check(MSAA_DepthTexture);
-	check(MSAA_IDTexture);
-	check(MSAA_MaterialTexture);
-	check(MSAA_AttributeTexture);
+	FIntPoint Resolution;
+	if (bUsePPLL)
+	{
+		check(PassParameters->PPLLCounter);
+		check(PassParameters->PPLLNodeIndex);
+		check(PassParameters->PPLLNodeData);
+		Resolution = PassParameters->PPLLNodeIndex->Desc.Extent;
+	}
+	else
+	{
+		check(PassParameters->MSAA_DepthTexture);
+		check(PassParameters->MSAA_IDTexture);
+		check(PassParameters->MSAA_MaterialTexture);
+		check(PassParameters->MSAA_AttributeTexture);
+		Resolution = PassParameters->MSAA_DepthTexture->Desc.Extent;
+	}
 
-	const FIntPoint Resolution = MSAA_DepthTexture->Desc.Extent;
 	
 	FRDGTextureRef CompactCounter;
 	{
@@ -531,7 +609,8 @@ static void AddHairVisibilityPrimitiveIdCompactionPass(
 	AddClearUAVPass(GraphBuilder, RDG_EVENT_NAME("HairStrandsClearCompactionOffsetAndCount"), OutCompactNodeIndex, 0);
 	AddClearUAVPass(GraphBuilder, RDG_EVENT_NAME("HairStrandsClearCategorizationTexture"), OutCategorizationTexture, 0);
 
-	const uint32 HairVisibilitySampleCount = GetHairVisibilitySampleCount();
+	const float PPLLHardCodedPixelSampleCount = 8;
+	const uint32 HairVisibilitySampleCount = bUsePPLL ? PPLLHardCodedPixelSampleCount : GetHairVisibilitySampleCount();
 	const uint32 SampleCount = FMath::RoundUpToPowerOfTwo(HairVisibilitySampleCount);
 	const uint32 MaxNodeCount = Resolution.X * Resolution.Y * SampleCount;
 	{
@@ -560,32 +639,26 @@ static void AddHairVisibilityPrimitiveIdCompactionPass(
 	FHairVisibilityPrimitiveIdCompactionCS::FPermutationDomain PermutationVector;
 	PermutationVector.Set<FHairVisibilityPrimitiveIdCompactionCS::FVendor>(GetVendor());
 	PermutationVector.Set<FHairVisibilityPrimitiveIdCompactionCS::FVelocity>(bWriteOutVelocity ? FMath::Clamp(GHairVelocityType+1, 0, 3) : 0);
-	PermutationVector.Set<FHairVisibilityPrimitiveIdCompactionCS::FCoverage>(CoverageTexture ? 1 : 0);
+	PermutationVector.Set<FHairVisibilityPrimitiveIdCompactionCS::FCoverage>(PassParameters->CoverageTexture ? 1 : 0);
 	PermutationVector.Set<FHairVisibilityPrimitiveIdCompactionCS::FMaterial>(GHairStrandsMaterialCompactionEnable ? 1 : 0);
-
-	FHairVisibilityPrimitiveIdCompactionCS::FParameters* Parameters = GraphBuilder.AllocParameters<FHairVisibilityPrimitiveIdCompactionCS::FParameters>();
-	Parameters->MSAA_DepthTexture = MSAA_DepthTexture;
-	Parameters->MSAA_IDTexture = MSAA_IDTexture;
-	Parameters->MSAA_MaterialTexture = MSAA_MaterialTexture;
-	Parameters->MSAA_AttributeTexture = MSAA_AttributeTexture;
-	Parameters->CoverageTexture = CoverageTexture;
-	Parameters->OutputResolution = Resolution;
-	Parameters->MaxNodeCount = MaxNodeCount;
-	Parameters->DepthTheshold = FMath::Clamp(GHairStrandsMaterialCompactionDepthThreshold, 0.f, 100.f);
-	Parameters->CosTangentThreshold = FMath::Cos(FMath::DegreesToRadians(FMath::Clamp(GHairStrandsMaterialCompactionTangentThreshold, 0.f, 90.f)));
-	Parameters->HairVisibilitySampleCount = HairVisibilitySampleCount;
-	Parameters->SceneTexturesStruct = CreateUniformBufferImmediate(SceneTextures, EUniformBufferUsage::UniformBuffer_SingleDraw);
-	Parameters->ViewUniformBuffer = View.ViewUniformBuffer;
-	Parameters->OutCompactNodeCounter = GraphBuilder.CreateUAV(CompactCounter);
-	Parameters->OutCompactNodeIndex = GraphBuilder.CreateUAV(OutCompactNodeIndex);
-	Parameters->OutCompactNodeData = GraphBuilder.CreateUAV(OutCompactNodeData);
-	Parameters->OutCompactNodeCoord = GraphBuilder.CreateUAV(OutCompactNodeCoord);
-	Parameters->OutCategorizationTexture = GraphBuilder.CreateUAV(OutCategorizationTexture);
+	PermutationVector.Set<FHairVisibilityPrimitiveIdCompactionCS::FPPLL>(bUsePPLL ? 1 : 0);
+
+	PassParameters->OutputResolution = Resolution;
+	PassParameters->MaxNodeCount = MaxNodeCount;
+	PassParameters->DepthTheshold = FMath::Clamp(GHairStrandsMaterialCompactionDepthThreshold, 0.f, 100.f);
+	PassParameters->CosTangentThreshold = FMath::Cos(FMath::DegreesToRadians(FMath::Clamp(GHairStrandsMaterialCompactionTangentThreshold, 0.f, 90.f)));
+	PassParameters->HairVisibilitySampleCount = HairVisibilitySampleCount;
+	PassParameters->SceneTexturesStruct = CreateUniformBufferImmediate(SceneTextures, EUniformBufferUsage::UniformBuffer_SingleDraw);
+	PassParameters->ViewUniformBuffer = View.ViewUniformBuffer;
+	PassParameters->OutCompactNodeCounter = GraphBuilder.CreateUAV(CompactCounter);
+	PassParameters->OutCompactNodeIndex = GraphBuilder.CreateUAV(OutCompactNodeIndex);
+	PassParameters->OutCompactNodeData = GraphBuilder.CreateUAV(OutCompactNodeData);
+	PassParameters->OutCompactNodeCoord = GraphBuilder.CreateUAV(OutCompactNodeCoord);
+	PassParameters->OutCategorizationTexture = GraphBuilder.CreateUAV(OutCategorizationTexture);
 
 	if (bWriteOutVelocity)
 	{
-		Parameters->MSAA_VelocityTexture = MSAA_VelocityTexture;
-		Parameters->OutVelocityTexture = GraphBuilder.CreateUAV(OutVelocityTexture);
+		PassParameters->OutVelocityTexture = GraphBuilder.CreateUAV(OutVelocityTexture);
 	}
  
 	FIntRect TotalRect = ComputeVisibleHairStrandsClustersRect(View.ViewRect, ClusterDatas);
@@ -598,14 +671,14 @@ static void AddHairVisibilityPrimitiveIdCompactionPass(
 	TotalRect.Max.Y = FMath::CeilToInt(float(TotalRect.Max.Y) / float(GroupSize.Y)) * GroupSize.Y;
 	
 	FIntPoint RectResolution(TotalRect.Width(), TotalRect.Height());
-	Parameters->ResolutionOffset = FIntPoint(TotalRect.Min.X, TotalRect.Min.Y);
+	PassParameters->ResolutionOffset = FIntPoint(TotalRect.Min.X, TotalRect.Min.Y);
 
 	TShaderMapRef<FHairVisibilityPrimitiveIdCompactionCS> ComputeShader(View.ShaderMap, PermutationVector);
 	FComputeShaderUtils::AddPass(
 		GraphBuilder,
 		RDG_EVENT_NAME("HairStrandsVisibilityCompaction"),
 		*ComputeShader,
-		Parameters,
+		PassParameters,
 		FComputeShaderUtils::GetGroupCount(RectResolution, GroupSize));
 
 	OutIndirectArgsBuffer = AddCopyIndirectArgPass(GraphBuilder, &View, NodeGroupSize, CompactCounter);
@@ -725,13 +798,6 @@ static FRDGTextureRef AddHairVisibilityFillOpaqueDepth(
 
 ///////////////////////////////////////////////////////////////////////////////////////////////////
 
-struct FVisibilityPassParameters
-{
-	BEGIN_SHADER_PARAMETER_STRUCT(FParameters, )
-		RENDER_TARGET_BINDING_SLOTS()
-	END_SHADER_PARAMETER_STRUCT()
-};
-
 template <typename TPassParameters>
 static void AddHairVisibilityCommonPass(
 	FRDGBuilder& GraphBuilder,
@@ -742,7 +808,8 @@ static void AddHairVisibilityCommonPass(
 	TPassParameters* PassParameters)
 {
 	GraphBuilder.AddPass(
-		RenderMode == HairVisibilityRenderMode_MSAA ? RDG_EVENT_NAME("HairStrandsVisibilityPass") : RDG_EVENT_NAME("HairStrandsCoveragePass"),
+		RenderMode == HairVisibilityRenderMode_PPLL ? RDG_EVENT_NAME("HairStrandsPPLLPass") : 
+		(RenderMode == HairVisibilityRenderMode_MSAA ? RDG_EVENT_NAME("HairStrandsVisibilityPass") : RDG_EVENT_NAME("HairStrandsCoveragePass")),
 		PassParameters,
 		ERDGPassFlags::Raster,
 		[PassParameters, Scene = Scene, ViewInfo, &ClusterDatas, RenderMode](FRHICommandListImmediate& RHICmdList)
@@ -752,19 +819,11 @@ static void AddHairVisibilityCommonPass(
 
 		FMeshPassProcessorRenderState DrawRenderState(*ViewInfo);
 
-		if(RenderMode == HairVisibilityRenderMode_Coverage)
+		if(RenderMode == HairVisibilityRenderMode_Coverage || RenderMode == HairVisibilityRenderMode_PPLL)
 		{
-			// In the case we render coverage, we need to override some view uniform shader parameters to account for the change in MSAA sample count.
-			const uint32 HairVisibilitySampleCount = 1;		// The coverage pass does not use MSAA
-			const float RasterizationScaleOverride = 0.0f;	// no override
-			FMinHairRadiusAtDepth1 MinHairRadius = ComputeMinStrandRadiusAtDepth1(
-				FIntPoint(ViewInfo->UnconstrainedViewRect.Width(), ViewInfo->UnconstrainedViewRect.Height()), ViewInfo->FOV, HairVisibilitySampleCount, RasterizationScaleOverride);
-
 			TUniformBufferRef<FViewUniformShaderParameters> ViewUniformShaderParameters;
-			// Update our view parameters
-			ViewInfo->CachedViewUniformShaderParameters->HairRenderInfo.X = MinHairRadius.Primary;
-			ViewInfo->CachedViewUniformShaderParameters->HairRenderInfo.Y = MinHairRadius.Velocity;
-			ViewInfo->CachedViewUniformShaderParameters->HairRenderInfo.Z = ViewInfo->IsPerspectiveProjection() ? 0.0f : 1.0f;
+			const bool bEnableMSAA = false;
+			SetUpViewHairRenderInfo(*ViewInfo, bEnableMSAA, ViewInfo->CachedViewUniformShaderParameters->HairRenderInfo);
 			// Create and set the uniform buffer
 			ViewUniformShaderParameters = TUniformBufferRef<FViewUniformShaderParameters>::CreateUniformBufferImmediate(*ViewInfo->CachedViewUniformShaderParameters, UniformBuffer_SingleFrame);
 			DrawRenderState.SetViewUniformBuffer(ViewUniformShaderParameters);
@@ -784,12 +843,23 @@ static void AddHairVisibilityCommonPass(
 				DrawRenderState.SetBlendState(TStaticBlendState<CW_RED, BO_Add, BF_One, BF_One, BO_Add, BF_One, BF_One>::GetRHI());
 				DrawRenderState.SetDepthStencilState(TStaticDepthStencilState<false, CF_DepthNearOrEqual>::GetRHI());
 			}
+			else if (RenderMode == HairVisibilityRenderMode_PPLL)
+			{
+				DrawRenderState.SetBlendState(TStaticBlendState<>::GetRHI());
+				DrawRenderState.SetDepthStencilState(TStaticDepthStencilState<false, CF_DepthNearOrEqual>::GetRHI());
+			}
 
 			FDynamicMeshDrawCommandStorage DynamicMeshDrawCommandStorage;
 			FMeshCommandOneFrameArray VisibleMeshDrawCommands;
 			FGraphicsMinimalPipelineStateSet PipelineStateSet;
 			FDynamicPassMeshDrawListContext ShadowContext(DynamicMeshDrawCommandStorage, VisibleMeshDrawCommands, PipelineStateSet);
-			FHairVisibilityProcessor MeshProcessor(Scene, ViewInfo /* is a SceneView */, DrawRenderState, RenderMode, &ShadowContext);
+			FHairVisibilityProcessor MeshProcessor(Scene, ViewInfo, DrawRenderState, RenderMode, &ShadowContext);
+			if (RenderMode == HairVisibilityRenderMode_PPLL)
+			{
+				// Work around because the value on the PassParamter is not taken into account for the PS global constant buffer.
+				// TODO fix that by having FVisibilityPassParameters be a global parameter structure.
+				MeshProcessor.SetMaxPPLLNodeCount(PassParameters->HairVisibilityPass_MaxPPLLNodeCount);
+			}
 
 			for (const FHairStrandsClusterData& ClusterData : ClusterDatas.Datas)
 			{
@@ -897,7 +967,7 @@ static void AddHairVisibilityMSAAPass(
 		AddClearGraphicPass(GraphBuilder, RDG_EVENT_NAME("HairStrandsClearVisibilityMSAAVelocity"), ViewInfo, 0, VisibilityVelocityTexture);
 	}
 
-	FVisibilityPassParameters::FParameters* PassParameters = GraphBuilder.AllocParameters<FVisibilityPassParameters::FParameters>();
+	FVisibilityPassParameters* PassParameters = GraphBuilder.AllocParameters<FVisibilityPassParameters>();
 	PassParameters->RenderTargets[0] = FRenderTargetBinding(VisibilityIdTexture, ERenderTargetLoadAction::ELoad, 0);
 	PassParameters->RenderTargets[1] = FRenderTargetBinding(VisibilityMaterialTexture,  LoadAction, 0);
 	PassParameters->RenderTargets[2] = FRenderTargetBinding(VisibilityAttributeTexture, LoadAction, 0);
@@ -911,6 +981,71 @@ static void AddHairVisibilityMSAAPass(
 	AddHairVisibilityCommonPass(GraphBuilder, Scene, ViewInfo, ClusterDatas, HairVisibilityRenderMode_MSAA, PassParameters);
 }
 
+static void AddHairVisibilityPPLLPass(
+	FRDGBuilder& GraphBuilder,
+	const FScene* Scene,
+	const FViewInfo* ViewInfo,
+	const FHairStrandsClusterDatas& ClusterDatas,
+	const FIntPoint& Resolution,
+	FRDGTextureRef& ViewZDepthTexture,
+	FRDGTextureRef& VisibilityPPLLNodeCounter,
+	FRDGTextureRef& VisibilityPPLLNodeIndex,
+	FRDGBufferRef& VisibilityPPLLNodeData)
+{
+	{
+		FRDGTextureDesc Desc;
+		Desc.Extent.X = 1;
+		Desc.Extent.Y = 1;
+		Desc.Depth = 0;
+		Desc.Format = PF_R32_UINT;
+		Desc.NumMips = 1;
+		Desc.NumSamples = 1;
+		Desc.Flags = TexCreate_None;
+		Desc.TargetableFlags = TexCreate_UAV | TexCreate_ShaderResource;
+		Desc.ClearValue = FClearValueBinding(0);
+		VisibilityPPLLNodeCounter = GraphBuilder.CreateTexture(Desc, TEXT("HairVisibilityPPLLCounter"));
+	}
+
+	{
+		FRDGTextureDesc Desc;
+		Desc.Extent.X = Resolution.X;
+		Desc.Extent.Y = Resolution.Y;
+		Desc.Depth = 0;
+		Desc.Format = PF_R32_UINT;
+		Desc.NumMips = 1;
+		Desc.NumSamples = 1;
+		Desc.Flags = TexCreate_None;
+		Desc.TargetableFlags = TexCreate_UAV | TexCreate_ShaderResource;
+		Desc.ClearValue = FClearValueBinding(0);
+		VisibilityPPLLNodeIndex = GraphBuilder.CreateTexture(Desc, TEXT("HairVisibilityPPLLNodeIndex"));
+	}
+
+	const uint32 MaxNodeCount = GetMaxNodePerPixel(Resolution);
+	{
+		// Example: 20bytes * 8spp = 160bytes per pixel = 316.4Mb @ 1080p
+		struct NodeData
+		{
+			uint32 Depth;
+			uint32 PrimitiveId_ClusterId;
+			uint32 Tangent_Coverage;
+			uint32 BaseColor_Roughness;
+			uint32 NextNodeIndex;
+		};
+
+		VisibilityPPLLNodeData = GraphBuilder.CreateBuffer(FRDGBufferDesc::CreateStructuredDesc(sizeof(NodeData), MaxNodeCount), TEXT("HairVisibilityPPLLNodeData"));
+	}
+	AddClearUAVPass(GraphBuilder, RDG_EVENT_NAME("ClearHairVisibilityPPLLCounter"), VisibilityPPLLNodeCounter, 0);
+	AddClearUAVPass(GraphBuilder, RDG_EVENT_NAME("ClearHairVisibilityPPLLNodeIndex"), VisibilityPPLLNodeIndex, 0xFFFFFFFF);
+
+	FVisibilityPassParameters* PassParameters = GraphBuilder.AllocParameters<FVisibilityPassParameters>();
+	PassParameters->PPLLCounter = GraphBuilder.CreateUAV(FRDGTextureUAVDesc(VisibilityPPLLNodeCounter, 0));
+	PassParameters->PPLLNodeIndex = GraphBuilder.CreateUAV(FRDGTextureUAVDesc(VisibilityPPLLNodeIndex, 0));
+	PassParameters->PPLLNodeData = GraphBuilder.CreateUAV(FRDGBufferUAVDesc(VisibilityPPLLNodeData));
+	PassParameters->HairVisibilityPass_MaxPPLLNodeCount = MaxNodeCount;
+	PassParameters->RenderTargets.DepthStencil = FDepthStencilBinding(ViewZDepthTexture, ERenderTargetLoadAction::ELoad, ERenderTargetLoadAction::ENoAction, FExclusiveDepthStencil::DepthRead_StencilNop);
+	AddHairVisibilityCommonPass(GraphBuilder, Scene, ViewInfo, ClusterDatas, HairVisibilityRenderMode_PPLL, PassParameters);
+}
+
 ///////////////////////////////////////////////////////////////////////////////////////////////////
 
 static FRDGTextureRef AddHairCoveragePass(
@@ -938,7 +1073,7 @@ static FRDGTextureRef AddHairCoveragePass(
 		CoverageTexture = GraphBuilder.CreateTexture(Desc, TEXT("HairCoverageTexture"));
 	}
 
-	FVisibilityPassParameters::FParameters* PassParameters = GraphBuilder.AllocParameters<FVisibilityPassParameters::FParameters>();
+	FVisibilityPassParameters* PassParameters = GraphBuilder.AllocParameters<FVisibilityPassParameters>();
 	PassParameters->RenderTargets[0] = FRenderTargetBinding(CoverageTexture, ERenderTargetLoadAction::EClear, 0);
 
 	PassParameters->RenderTargets.DepthStencil = FDepthStencilBinding(
@@ -958,10 +1093,12 @@ class FHairVisibilityDepthPS : public FGlobalShader
 	SHADER_USE_PARAMETER_STRUCT(FHairVisibilityDepthPS, FGlobalShader);
 
 	class FCoverage : SHADER_PERMUTATION_INT("PERMUTATION_COVERAGE", 2);
-	using FPermutationDomain = TShaderPermutationDomain<FCoverage>;
+	class FPPLL : SHADER_PERMUTATION_INT("PERMUTATION_PPLL", 2);
+	using FPermutationDomain = TShaderPermutationDomain<FCoverage, FPPLL>;
 
 	BEGIN_SHADER_PARAMETER_STRUCT(FParameters, )
 		SHADER_PARAMETER_RDG_TEXTURE(Texture2D, HairVisibilityDepthTexture)
+		SHADER_PARAMETER_RDG_TEXTURE(Texture2D, CategorisationTexture)
 		SHADER_PARAMETER_RDG_TEXTURE(Texture2D, CoverageTexture)
 		RENDER_TARGET_BINDING_SLOTS()
 	END_SHADER_PARAMETER_STRUCT()
@@ -976,6 +1113,7 @@ static void AddHairVisibilityColorAndDepthPatchPass(
 	FRDGBuilder& GraphBuilder,
 	const FViewInfo& View,
 	const FRDGTextureRef& VisibilityDepthTexture,
+	const FRDGTextureRef& CategorisationTexture,
 	const FRDGTextureRef& CoverageTexture,
 	FRDGTextureRef& OutGBufferBTexture,
 	FRDGTextureRef& OutColorTexture,
@@ -983,6 +1121,7 @@ static void AddHairVisibilityColorAndDepthPatchPass(
 {
 	FHairVisibilityDepthPS::FParameters* Parameters = GraphBuilder.AllocParameters<FHairVisibilityDepthPS::FParameters>();
 	Parameters->HairVisibilityDepthTexture = VisibilityDepthTexture;
+	Parameters->CategorisationTexture = CategorisationTexture;
 	Parameters->CoverageTexture = CoverageTexture;
 	Parameters->RenderTargets[0] = FRenderTargetBinding(OutGBufferBTexture, ERenderTargetLoadAction::ELoad);
 	Parameters->RenderTargets[1] = FRenderTargetBinding(OutColorTexture, ERenderTargetLoadAction::ELoad);
@@ -995,6 +1134,7 @@ static void AddHairVisibilityColorAndDepthPatchPass(
 	TShaderMapRef<FPostProcessVS> VertexShader(View.ShaderMap);
 	FHairVisibilityDepthPS::FPermutationDomain PermutationVector;
 	PermutationVector.Set<FHairVisibilityDepthPS::FCoverage>(CoverageTexture ? 1 : 0);
+	PermutationVector.Set<FHairVisibilityDepthPS::FPPLL>(CategorisationTexture ? 1 : 0);
 	TShaderMapRef<FHairVisibilityDepthPS> PixelShader(View.ShaderMap, PermutationVector);
 	const TShaderMap<FGlobalShaderType>* GlobalShaderMap = View.ShaderMap;
 	const FIntRect Viewport = View.ViewRect;
@@ -1084,15 +1224,6 @@ FHairStrandsVisibilityViews RenderHairStrandsVisibilityBuffer(
 			
 			const EHairVisibilityRenderMode RenderMode = GetHairVisibilityRenderMode();
 
-			struct FRDGMsaaVisibilityResources
-			{
-				FRDGTextureRef DepthTexture;
-				FRDGTextureRef IdTexture;
-				FRDGTextureRef MaterialTexture;
-				FRDGTextureRef AttributeTexture;
-				FRDGTextureRef VelocityTexture;
-			} MsaaVisibilityResources;
-
 			FRDGTextureRef CoverageTexture = nullptr;
 			if (GHairStrandsCoveragePassEnable > 0)
 			{
@@ -1107,6 +1238,15 @@ FHairStrandsVisibilityViews RenderHairStrandsVisibilityBuffer(
 
 			if (RenderMode == HairVisibilityRenderMode_MSAA)
 			{
+				struct FRDGMsaaVisibilityResources
+				{
+					FRDGTextureRef DepthTexture;
+					FRDGTextureRef IdTexture;
+					FRDGTextureRef MaterialTexture;
+					FRDGTextureRef AttributeTexture;
+					FRDGTextureRef VelocityTexture;
+				} MsaaVisibilityResources;
+
 				MsaaVisibilityResources.DepthTexture = AddHairVisibilityFillOpaqueDepth(
 					GraphBuilder,
 					View,
@@ -1134,23 +1274,26 @@ FHairStrandsVisibilityViews RenderHairStrandsVisibilityBuffer(
 				GraphBuilder.QueueTextureExtraction(MsaaVisibilityResources.DepthTexture, &VisibilityData.DepthTexture);
 				
 				{
+					FHairVisibilityPrimitiveIdCompactionCS::FParameters* PassParameters = GraphBuilder.AllocParameters<FHairVisibilityPrimitiveIdCompactionCS::FParameters>();
+					PassParameters->MSAA_DepthTexture = MsaaVisibilityResources.DepthTexture;
+					PassParameters->MSAA_IDTexture = MsaaVisibilityResources.IdTexture;
+					PassParameters->MSAA_MaterialTexture = MsaaVisibilityResources.MaterialTexture;
+					PassParameters->MSAA_AttributeTexture = MsaaVisibilityResources.MaterialTexture;
+					PassParameters->MSAA_VelocityTexture = MsaaVisibilityResources.VelocityTexture;
+					PassParameters->CoverageTexture = CoverageTexture;
+
 					FRDGTextureRef CompactNodeIndex;
 					FRDGBufferRef CompactNodeData;
 					FRDGBufferRef CompactNodeCoord;
 					FRDGTextureRef CategorizationTexture;
 					FRDGBufferRef IndirectArgsBuffer;
-
 					AddHairVisibilityPrimitiveIdCompactionPass(
+						false, // bUsePPLL
 						GraphBuilder,
 						View,
 						ClusterDatas,
 						VisibilityData.NodeGroupSize,
-						MsaaVisibilityResources.DepthTexture,
-						MsaaVisibilityResources.IdTexture,
-						MsaaVisibilityResources.MaterialTexture,
-						MsaaVisibilityResources.AttributeTexture,
-						MsaaVisibilityResources.VelocityTexture,
-						CoverageTexture,
+						PassParameters,
 						CompactNodeIndex,
 						CompactNodeData,
 						CompactNodeCoord,
@@ -1163,20 +1306,86 @@ FHairStrandsVisibilityViews RenderHairStrandsVisibilityBuffer(
 					GraphBuilder.QueueBufferExtraction(CompactNodeCoord,		&VisibilityData.NodeCoord,					FRDGResourceState::EAccess::Read, FRDGResourceState::EPipeline::Graphics);
 					GraphBuilder.QueueBufferExtraction(IndirectArgsBuffer,		&VisibilityData.NodeIndirectArg,			FRDGResourceState::EAccess::Read, FRDGResourceState::EPipeline::Compute);
 				}
+
+				// For fully covered pixels, write: 
+				// * black color into the scene color
+				// * closest depth
+				// * unlit shading model ID 
+				AddHairVisibilityColorAndDepthPatchPass(
+					GraphBuilder,
+					View,
+					MsaaVisibilityResources.DepthTexture,
+					nullptr,
+					CoverageTexture,
+					SceneGBufferBTexture,
+					SceneColorTexture,
+					SceneDepthTexture);
+			}
+			else if (RenderMode == HairVisibilityRenderMode_PPLL)
+			{
+				// In this pas we reuse the scene depth buffer to cull hair pixels out.
+				// Pixel data is accumulated in buffer containing data organized in a linked list with node scattered in memory according to pixel shader execution. 
+				// This with up to width * height * GHairVisibilityPPLLGlobalMaxPixelNodeCount node total maximum.
+				// After we have that a node sorting pass happening and we finally output all the data once into the common compaction node list.
+
+				FRDGTextureRef PPLLNodeCounterTexture;
+				FRDGTextureRef PPLLNodeIndexTexture;
+				FRDGBufferRef PPLLNodeDataBuffer;
+				FRDGTextureRef ViewZDepthTexture = GraphBuilder.RegisterExternalTexture(SceneContext.SceneDepthZ);
+
+				// Linked list generation pass
+				AddHairVisibilityPPLLPass(GraphBuilder, Scene, &View, ClusterDatas, Resolution, ViewZDepthTexture, PPLLNodeCounterTexture, PPLLNodeIndexTexture, PPLLNodeDataBuffer);
+
+				// Linked list sorting pass and compaction into common representation
+				FRDGTextureRef CategorizationTexture;
+				{
+					FHairVisibilityPrimitiveIdCompactionCS::FParameters* PassParameters = GraphBuilder.AllocParameters<FHairVisibilityPrimitiveIdCompactionCS::FParameters>();
+					PassParameters->PPLLCounter  = PPLLNodeCounterTexture;
+					PassParameters->PPLLNodeIndex= PPLLNodeIndexTexture;
+					PassParameters->PPLLNodeData = GraphBuilder.CreateSRV(PPLLNodeDataBuffer);
+					PassParameters->CoverageTexture = CoverageTexture;
+
+					FRDGTextureRef CompactNodeIndex;
+					FRDGBufferRef CompactNodeData;
+					FRDGBufferRef CompactNodeCoord;
+					FRDGBufferRef IndirectArgsBuffer;
+					AddHairVisibilityPrimitiveIdCompactionPass(
+						true, // bUsePPLL
+						GraphBuilder,
+						View,
+						ClusterDatas,
+						VisibilityData.NodeGroupSize,
+						PassParameters,
+						CompactNodeIndex,
+						CompactNodeData,
+						CompactNodeCoord,
+						CategorizationTexture,
+						SceneVelocityTexture,
+						IndirectArgsBuffer);
+					GraphBuilder.QueueTextureExtraction(CompactNodeIndex, &VisibilityData.NodeIndex);
+					GraphBuilder.QueueTextureExtraction(CategorizationTexture, &VisibilityData.CategorizationTexture);
+					GraphBuilder.QueueBufferExtraction(CompactNodeData, &VisibilityData.NodeData, FRDGResourceState::EAccess::Read, FRDGResourceState::EPipeline::Graphics);
+					GraphBuilder.QueueBufferExtraction(CompactNodeCoord, &VisibilityData.NodeCoord, FRDGResourceState::EAccess::Read, FRDGResourceState::EPipeline::Graphics);
+					GraphBuilder.QueueBufferExtraction(IndirectArgsBuffer, &VisibilityData.NodeIndirectArg, FRDGResourceState::EAccess::Read, FRDGResourceState::EPipeline::Compute);
+				}
+
+				AddHairVisibilityColorAndDepthPatchPass(
+					GraphBuilder,
+					View,
+					nullptr,
+					CategorizationTexture,
+					CoverageTexture,
+					SceneGBufferBTexture,
+					SceneColorTexture,
+					SceneDepthTexture);
+
+#if WITH_EDITOR
+				// Extract texture for debug visualization
+				GraphBuilder.QueueTextureExtraction(PPLLNodeCounterTexture, &VisibilityData.PPLLNodeCounterTexture);
+				GraphBuilder.QueueTextureExtraction(PPLLNodeIndexTexture, &VisibilityData.PPLLNodeIndexTexture);
+				GraphBuilder.QueueBufferExtraction(PPLLNodeDataBuffer, &VisibilityData.PPLLNodeDataBuffer, FRDGResourceState::EAccess::Read, FRDGResourceState::EPipeline::Graphics);
+#endif
 			}
-		
-			// For fully covered pixels, write: 
-			// * black color into the scene color
-			// * closest depth
-			// * unlit shading model ID 
-			AddHairVisibilityColorAndDepthPatchPass(
-				GraphBuilder,
-				View,
-				MsaaVisibilityResources.DepthTexture,
-				CoverageTexture,
-				SceneGBufferBTexture,
-				SceneColorTexture,
-				SceneDepthTexture);
 
 			GraphBuilder.Execute();
 
diff --git a/Engine/Source/Runtime/Renderer/Private/HairStrands/HairStrandsVisibility.h b/Engine/Source/Runtime/Renderer/Private/HairStrands/HairStrandsVisibility.h
--- a/Engine/Source/Runtime/Renderer/Private/HairStrands/HairStrandsVisibility.h
+++ b/Engine/Source/Runtime/Renderer/Private/HairStrands/HairStrandsVisibility.h
@@ -10,6 +10,10 @@
 #include "RendererInterface.h"
 #include "RenderGraphResources.h"
 
+class FViewInfo;
+struct FVector;
+struct FIntPoint;
+
 struct FHairStrandsVisibilityData
 {
 	TRefCountPtr<IPooledRenderTarget> DepthTexture;
@@ -19,6 +23,10 @@ struct FHairStrandsVisibilityData
 	TRefCountPtr<IPooledRenderTarget> VelocityTexture;
 	TRefCountPtr<IPooledRenderTarget> CategorizationTexture;
 
+	TRefCountPtr<IPooledRenderTarget> PPLLNodeCounterTexture;
+	TRefCountPtr<IPooledRenderTarget> PPLLNodeIndexTexture;
+	TRefCountPtr<FPooledRDGBuffer>	  PPLLNodeDataBuffer;
+
 	TRefCountPtr<IPooledRenderTarget> NodeIndex;
 	TRefCountPtr<FPooledRDGBuffer>	  NodeData;
 	FShaderResourceViewRHIRef		  NodeDataSRV;
@@ -42,3 +50,7 @@ FHairStrandsVisibilityViews RenderHairStrandsVisibilityBuffer(
 	TRefCountPtr<IPooledRenderTarget> DepthTexture,
 	TRefCountPtr<IPooledRenderTarget> VelocityTexture,
 	const struct FHairStrandsClusterViews& ClusterViews);
+
+void SetUpViewHairRenderInfo(const FViewInfo& ViewInfo, bool bEnableMSAA, FVector4& OutHairRenderInfo);
+
+uint32 GetMaxNodePerPixel(FIntPoint Resolution);
diff --git a/Engine/Source/Runtime/Renderer/Private/SceneRendering.cpp b/Engine/Source/Runtime/Renderer/Private/SceneRendering.cpp
--- a/Engine/Source/Runtime/Renderer/Private/SceneRendering.cpp
+++ b/Engine/Source/Runtime/Renderer/Private/SceneRendering.cpp
@@ -57,8 +57,8 @@
 #include "IXRCamera.h"
 #include "IHeadMountedDisplay.h"
 #include "DiaphragmDOF.h" 
-#include "SingleLayerWaterRendering.h" 
-#include "HairStrands/HairStrandsUtils.h"
+#include "SingleLayerWaterRendering.h"
+#include "HairStrands/HairStrandsVisibility.h"
 
 /*-----------------------------------------------------------------------------
 	Globals
@@ -329,11 +329,6 @@ TAutoConsoleVariable<int32> CVarTransientResourceAliasing_Buffers(
 	TEXT("Enables transient resource aliasing for specified buffers. Used only if GSupportsTransientResourceAliasing is true.\n"),
 	ECVF_ReadOnly);
 
-static TAutoConsoleVariable<int32> CVarHairVelocityMagnitudeScale(
-	TEXT("r.HairStrands.VelocityMagnitudeScale"), 
-	100,  // Tuned by eye, based on heavy motion (strong head shack)
-	TEXT("Velocity magnitude (in pixel) at which a hair will reach its pic velocity-rasterization-scale under motion to reduce aliasing. Default is 100."));
-
 #if !UE_BUILD_SHIPPING
 
 static TAutoConsoleVariable<int32> CVarTestInternalViewRectOffset(
@@ -1685,16 +1680,8 @@ void FViewInfo::SetupUniformBufferParameters(
 
 	// Deep opacity maps info
 	{
-		const FMinHairRadiusAtDepth1 MinRadiusAtDepth1 = ComputeMinStrandRadiusAtDepth1(
-			FIntPoint(UnconstrainedViewRect.Width(), UnconstrainedViewRect.Height()), FOV, GetHairVisibilitySampleCount(), 0.0f);
-
-		FVector2D PixelVelocity(1.f / (ViewRect.Width() * 2), 1.f / (ViewRect.Height() * 2));
-		const float VelocityMagnitudeScale = FMath::Clamp(CVarHairVelocityMagnitudeScale.GetValueOnAnyThread(), 0, 512) * FMath::Min(PixelVelocity.X, PixelVelocity.Y);
-
-		ViewUniformShaderParameters.HairRenderInfo.X = MinRadiusAtDepth1.Primary;
-		ViewUniformShaderParameters.HairRenderInfo.Y = MinRadiusAtDepth1.Velocity;
-		ViewUniformShaderParameters.HairRenderInfo.Z = IsPerspectiveProjection() ? 0.0f : 1.0f;
-		ViewUniformShaderParameters.HairRenderInfo.W = VelocityMagnitudeScale;
+		const bool bEnableMSAA = true;
+		SetUpViewHairRenderInfo(*this, bEnableMSAA, ViewUniformShaderParameters.HairRenderInfo);
 	}
 }
 
