Add sample count option for AO cast by hair, and better sequence offset.

diff --git a/Engine/Shaders/Private/HairStrands/HairStrandsEnvironmentAO.usf b/Engine/Shaders/Private/HairStrands/HairStrandsEnvironmentAO.usf
--- a/Engine/Shaders/Private/HairStrands/HairStrandsEnvironmentAO.usf
+++ b/Engine/Shaders/Private/HairStrands/HairStrandsEnvironmentAO.usf
@@ -17,6 +17,7 @@ float Voxel_DepthBiasScale;
 
 float AO_Power;
 float AO_Intensity;
+uint  AO_SampleCount;
 
 float3 Voxel_MinAABB;
 float3 Voxel_MaxAABB;
@@ -56,6 +57,7 @@ float TraceAO(float3 WorldPosition, float3 WorldNormal, float4 SvPosition)
 	
 	float3 UnoccludedN = 0;
 	uint TotalHairCount = 0;
+#if PERMUTATION_SAMPLESET == 0
 	// Poisson disk position http://developer.download.nvidia.com/whitepapers/2008/PCSS_Integration.pdf
 	float2 PoissonDisk[16] =
 	{
@@ -76,20 +78,30 @@ float TraceAO(float3 WorldPosition, float3 WorldNormal, float4 SvPosition)
 		float2(0.19984126, 0.78641367),
 		float2(0.14383161, -0.14100790)
 	};
-	const uint SampleCount = 16;
+	const uint SampleCount = clamp(AO_SampleCount, 1, 16);
+#else
+	const uint SampleCount = max(AO_SampleCount, 1);
+#endif
 
 	float3 BentNormal = 0;
 	float Occlusion = 0;
 	for (uint i = 0; i < SampleCount; ++i)
 	{
-		float2 Rand2 = (PoissonDisk[i] + float2(1, 1)) * 0.5f;
-		Rand2 = frac(Rand2 + VoxelOffset.xy);
-		const float3 LocalL  = CosineSampleHemisphere(saturate(Rand2)).xyz;
+	#if PERMUTATION_SAMPLESET == 0
+		const float2 Rand2 = (PoissonDisk[i] + float2(1, 1)) * 0.5f;
+	#else // PERMUTATION_SAMPLESET == 1
+		const float2 Rand2 = Hammersley(i, SampleCount, uint2(0, 0));
+	#endif
+		const float2 SequenceOffset = float2(
+			InterleavedGradientNoise(SvPosition.xy, View.StateFrameIndexMod8 + i * 13),
+			InterleavedGradientNoise(SvPosition.xy, View.StateFrameIndexMod8 + i * 117));
+
+		const float3 LocalL = CosineSampleHemisphere(frac(Rand2 + SequenceOffset)).xyz;
 		const float3 SampleL = TangentToWorld(LocalL.xyz, WorldNormal);
 
 		// Depth bias
-		// Origin is shifted voxels away towards the light + a constant bias of the size of the voxel
-		const float3 SampleWorldPosition = WorldPosition + VoxelOffset /*+ VoxelMaxSize*Voxel_DepthBiasScale*SampleL*/;;
+		const float3 SampleVoxelOffset = VoxelMaxSize * GetHairVoxelJitter(SvPosition, View.StateFrameIndexMod8 + i * 17);
+		const float3 SampleWorldPosition = WorldPosition + VoxelOffset;
 
 		// Compute the number of hair count between light & shading point
 		const float HairCount = InternalComputeHairCountVoxel(SampleWorldPosition, SampleL, Voxel_MinAABB, Voxel_MaxAABB, Voxel_Resolution, Voxel_DensityTexture, Voxel_DensityScale, 0, true);
diff --git a/Engine/Source/Runtime/Renderer/Private/HairStrands/HairStrandsEnvironment.cpp b/Engine/Source/Runtime/Renderer/Private/HairStrands/HairStrandsEnvironment.cpp
--- a/Engine/Source/Runtime/Renderer/Private/HairStrands/HairStrandsEnvironment.cpp
+++ b/Engine/Source/Runtime/Renderer/Private/HairStrands/HairStrandsEnvironment.cpp
@@ -25,16 +25,19 @@ static int32 GHairSkylightingEnable = 1;
 static FAutoConsoleVariableRef CVarHairSkylightingEnable(TEXT("r.HairStrands.SkyLighting"), GHairSkylightingEnable, TEXT("Enable sky lighting on hair."));
 
 static int32 GHairSkyAOEnable = 1;
-static FAutoConsoleVariableRef CVarHairSkyAOEnable(TEXT("r.HairStrands.SkyAOEnable"), GHairSkyAOEnable, TEXT("Enable (sky) AO on hair."));
+static FAutoConsoleVariableRef CVarHairSkyAOEnable(TEXT("r.HairStrands.SkyAO"), GHairSkyAOEnable, TEXT("Enable (sky) AO on hair."));
 
 static float GHairSkylightingConeAngle = 3;
 static FAutoConsoleVariableRef CVarHairSkylightingConeAngle(TEXT("r.HairStrands.SkyLighting.ConeAngle"), GHairSkylightingConeAngle, TEXT("Cone angle for tracing sky lighting on hair."));
 
 static int32 GHairStrandsSkyLightingSampleCount = 16;
-static FAutoConsoleVariableRef GVarHairStrandsSkyLightingSampleCount(TEXT("r.HairStrands.SkyLighting.SampleCount"), GHairStrandsSkyLightingSampleCount, TEXT("Number of samples used for evaluating multiple scattering and visible area (default is set to 16)."));
+static FAutoConsoleVariableRef CVarHairStrandsSkyLightingSampleCount(TEXT("r.HairStrands.SkyLighting.SampleCount"), GHairStrandsSkyLightingSampleCount, TEXT("Number of samples used for evaluating multiple scattering and visible area (default is set to 16)."));
 
 static int32 GHairStrandsSkyLightingJitterSphericalIntegration = 0;
-static FAutoConsoleVariableRef GVarHairStrandsSkyLightingJitterSphericalIntegration(TEXT("r.HairStrands.SkyLighting.JitterIntegration"), GHairStrandsSkyLightingJitterSphericalIntegration, TEXT("Jitter the sphereical integration for the multiple scattering term. The result is more correct, but noiser as well."));
+static FAutoConsoleVariableRef CVarHairStrandsSkyLightingJitterSphericalIntegration(TEXT("r.HairStrands.SkyLighting.JitterIntegration"), GHairStrandsSkyLightingJitterSphericalIntegration, TEXT("Jitter the sphereical integration for the multiple scattering term. The result is more correct, but noiser as well."));
+
+static int32 GHairStrandsSkyAOSampleCount = 16;
+static FAutoConsoleVariableRef CVarHairStrandsSkyAOSampleCount(TEXT("r.HairStrands.SkyAO.SampleCount"), GHairStrandsSkyAOSampleCount, TEXT("Number of samples used for evaluating hair AO (default is set to 16)."));
 
 static bool GetHairStrandsSkyLightingEnable() { return GHairSkylightingEnable > 0; }
 static bool GetHairStrandsSkyAOEnable() { return GHairSkyAOEnable > 0; }
@@ -49,6 +52,9 @@ class FHairEnvironmentAO : public FGlobalShader
 	DECLARE_GLOBAL_SHADER(FHairEnvironmentAO);
 	SHADER_USE_PARAMETER_STRUCT(FHairEnvironmentAO, FGlobalShader)
 
+	class FSampleSet : SHADER_PERMUTATION_INT("PERMUTATION_SAMPLESET", 2);
+	using FPermutationDomain = TShaderPermutationDomain<FSampleSet>;
+
 	static bool ShouldCompilePermutation(const FGlobalShaderPermutationParameters& Parameters)
 	{
 		return IsHairStrandsSupported(Parameters.Platform);
@@ -65,7 +71,7 @@ class FHairEnvironmentAO : public FGlobalShader
 		SHADER_PARAMETER(float, Voxel_TanConeAngle)
 		SHADER_PARAMETER(float, AO_Power)
 		SHADER_PARAMETER(float, AO_Intensity)
-
+		SHADER_PARAMETER(uint32, AO_SampleCount)
 		SHADER_PARAMETER_STRUCT_INCLUDE(FSceneTextureParameters, SceneTextures)
 		SHADER_PARAMETER_STRUCT_INCLUDE(FSceneTextureSamplerParameters, SceneTextureSamplers)
 
@@ -107,11 +113,14 @@ static void AddHairStrandsEnvironmentAOPass(
 	const FFinalPostProcessSettings& Settings = View.FinalPostProcessSettings;
 	PassParameters->AO_Power = Settings.AmbientOcclusionPower;
 	PassParameters->AO_Intensity = Settings.AmbientOcclusionIntensity;
-
+	PassParameters->AO_SampleCount = FMath::Max(uint32(GHairStrandsSkyAOSampleCount), 1u);
 	check(Output);
 	PassParameters->RenderTargets[0] = FRenderTargetBinding(Output, ERenderTargetLoadAction::ELoad);
 
-	TShaderMapRef<FHairEnvironmentAO> PixelShader(View.ShaderMap);
+	FHairEnvironmentAO::FPermutationDomain PermutationVector;
+	PermutationVector.Set<FHairEnvironmentAO::FSampleSet>(PassParameters->AO_SampleCount <= 16 ? 0 : 1);
+
+	TShaderMapRef<FHairEnvironmentAO> PixelShader(View.ShaderMap, PermutationVector);
 	ClearUnusedGraphResources(*PixelShader, PassParameters);
 
 	GraphBuilder.AddPass(
