Fixed wrong hair RT assignement.

diff --git a/Engine/Source/Runtime/Renderer/Private/HairStrands/HairStrandsVisibility.cpp b/Engine/Source/Runtime/Renderer/Private/HairStrands/HairStrandsVisibility.cpp
--- a/Engine/Source/Runtime/Renderer/Private/HairStrands/HairStrandsVisibility.cpp
+++ b/Engine/Source/Runtime/Renderer/Private/HairStrands/HairStrandsVisibility.cpp
@@ -1279,7 +1279,7 @@ FHairStrandsVisibilityViews RenderHairStrandsVisibilityBuffer(
 					PassParameters->MSAA_DepthTexture = MsaaVisibilityResources.DepthTexture;
 					PassParameters->MSAA_IDTexture = MsaaVisibilityResources.IdTexture;
 					PassParameters->MSAA_MaterialTexture = MsaaVisibilityResources.MaterialTexture;
-					PassParameters->MSAA_AttributeTexture = MsaaVisibilityResources.MaterialTexture;
+					PassParameters->MSAA_AttributeTexture = MsaaVisibilityResources.AttributeTexture;
 					PassParameters->MSAA_VelocityTexture = MsaaVisibilityResources.VelocityTexture;
 					PassParameters->CoverageTexture = CoverageTexture;
 
