Fixed performance drops with scene capture component when raytracing enabled

diff --git a/Engine/Source/Runtime/Renderer/Private/DeferredShadingRenderer.cpp b/Engine/Source/Runtime/Renderer/Private/DeferredShadingRenderer.cpp
--- a/Engine/Source/Runtime/Renderer/Private/DeferredShadingRenderer.cpp
+++ b/Engine/Source/Runtime/Renderer/Private/DeferredShadingRenderer.cpp
@@ -685,7 +685,8 @@ bool FDeferredShadingSceneRenderer::GatherRayTracingWorldInstances(FRHICommandLi
 					continue;
 				}
 
-				if (View.bIsReflectionCapture && !SceneInfo->bIsVisibleInReflectionCaptures)
+				if ((View.bIsReflectionCapture && !SceneInfo->bIsVisibleInReflectionCaptures)
+					|| View.bIsSceneCapture)
 				{
 					continue;
 				}
