RL - Add missing GPU Scopes

diff --git a/Engine/Source/Runtime/Renderer/Private/DeferredShadingRenderer.cpp b/Engine/Source/Runtime/Renderer/Private/DeferredShadingRenderer.cpp
--- a/Engine/Source/Runtime/Renderer/Private/DeferredShadingRenderer.cpp
+++ b/Engine/Source/Runtime/Renderer/Private/DeferredShadingRenderer.cpp
@@ -210,11 +210,14 @@ DECLARE_GPU_STAT(AllocateRendertargets);
 DECLARE_GPU_STAT(FrameRenderFinish);
 DECLARE_GPU_STAT(SortLights);
 DECLARE_GPU_STAT(PostRenderOpsFX);
+DECLARE_GPU_STAT(GPUSceneUpdate);
 DECLARE_GPU_STAT(HZB);
 DECLARE_GPU_STAT_NAMED(Unaccounted, TEXT("[unaccounted]"));
 DECLARE_GPU_STAT(WaterRendering);
 DECLARE_GPU_STAT(HairRendering);
-DECLARE_GPU_STAT(VirtualTextureFeedback);
+DECLARE_GPU_STAT(VirtualTextureUpdate);
+DECLARE_GPU_STAT(UploadDynamicBuffers);
+DECLARE_GPU_STAT(PostOpaqueExtensions);
 
 const TCHAR* GetDepthPassReason(bool bDitheredLODTransitionsUseStencil, EShaderPlatform ShaderPlatform)
 {
@@ -1041,6 +1044,7 @@ void FDeferredShadingSceneRenderer::Render(FRHICommandListImmediate& RHICmdList)
 	const bool bUseVirtualTexturing = UseVirtualTexturing(FeatureLevel);
 	if (bUseVirtualTexturing)
 	{
+		SCOPED_GPU_STAT(RHICmdList, VirtualTextureUpdate);
 		// AllocateResources needs to be called before RHIBeginScene
 		FVirtualTextureSystem::Get().AllocateResources(RHICmdList, FeatureLevel);
 		FVirtualTextureSystem::Get().CallPendingCallbacks();
@@ -1115,46 +1119,50 @@ void FDeferredShadingSceneRenderer::Render(FRHICommandListImmediate& RHICmdList)
 		}	
 	}
 
-	if (GDoPrepareDistanceFieldSceneAfterRHIFlush && (GRHINeedsExtraDeletionLatency || !GRHICommandList.Bypass()))
 	{
-		// we will probably stall on occlusion queries, so might as well have the RHI thread and GPU work while we wait.
-		SCOPE_CYCLE_COUNTER(STAT_PostInitViews_FlushDel);
-		RHICmdList.ImmediateFlush(EImmediateFlushType::FlushRHIThreadFlushResources);
-	}
+		SCOPED_GPU_STAT(RHICmdList, GPUSceneUpdate);
 
+		if (GDoPrepareDistanceFieldSceneAfterRHIFlush && (GRHINeedsExtraDeletionLatency || !GRHICommandList.Bypass()))
+		{
+			// we will probably stall on occlusion queries, so might as well have the RHI thread and GPU work while we wait.
+			SCOPE_CYCLE_COUNTER(STAT_PostInitViews_FlushDel);
+			RHICmdList.ImmediateFlush(EImmediateFlushType::FlushRHIThreadFlushResources);
+		}
 
-	UpdateGPUScene(RHICmdList, *Scene);
+		UpdateGPUScene(RHICmdList, *Scene);
 
-	if (bUseVirtualTexturing)
-	{
-		FVirtualTextureSystem::Get().Update(RHICmdList, FeatureLevel, Scene);
-	}
+		if (bUseVirtualTexturing)
+		{
+			SCOPED_GPU_STAT(RHICmdList, VirtualTextureUpdate);
+			FVirtualTextureSystem::Get().Update(RHICmdList, FeatureLevel, Scene);
+		}
 
-	for (int32 ViewIndex = 0; ViewIndex < Views.Num(); ViewIndex++)
-	{
-		FViewInfo& View = Views[ViewIndex];
-		SCOPED_GPU_MASK(RHICmdList, View.GPUMask);
-		ShaderPrint::BeginView(RHICmdList, View);
-	}
+		for (int32 ViewIndex = 0; ViewIndex < Views.Num(); ViewIndex++)
+		{
+			FViewInfo& View = Views[ViewIndex];
+			SCOPED_GPU_MASK(RHICmdList, View.GPUMask);
+			ShaderPrint::BeginView(RHICmdList, View);
+		}
 
-	for (int32 ViewIndex = 0; ViewIndex < Views.Num(); ViewIndex++)
-	{
-		FViewInfo& View = Views[ViewIndex];
-		SCOPED_GPU_MASK(RHICmdList, View.GPUMask);
-		UploadDynamicPrimitiveShaderDataForView(RHICmdList, *Scene, View);
-	}	
-	
-	if (!bDoInitViewAftersPrepass)
-	{
-		bool bSplitDispatch = !GDoPrepareDistanceFieldSceneAfterRHIFlush;
-		PrepareDistanceFieldScene(RHICmdList, bSplitDispatch);
-	}
+		for (int32 ViewIndex = 0; ViewIndex < Views.Num(); ViewIndex++)
+		{
+			FViewInfo& View = Views[ViewIndex];
+			SCOPED_GPU_MASK(RHICmdList, View.GPUMask);
+			UploadDynamicPrimitiveShaderDataForView(RHICmdList, *Scene, View);
+		}
 
-	if (!GDoPrepareDistanceFieldSceneAfterRHIFlush && (GRHINeedsExtraDeletionLatency || !GRHICommandList.Bypass()))
-	{
-		// we will probably stall on occlusion queries, so might as well have the RHI thread and GPU work while we wait.
-		SCOPE_CYCLE_COUNTER(STAT_PostInitViews_FlushDel);
-		FRHICommandListExecutor::GetImmediateCommandList().ImmediateFlush(EImmediateFlushType::FlushRHIThreadFlushResources);
+		if (!bDoInitViewAftersPrepass)
+		{
+			bool bSplitDispatch = !GDoPrepareDistanceFieldSceneAfterRHIFlush;
+			PrepareDistanceFieldScene(RHICmdList, bSplitDispatch);
+		}
+
+		if (!GDoPrepareDistanceFieldSceneAfterRHIFlush && (GRHINeedsExtraDeletionLatency || !GRHICommandList.Bypass()))
+		{
+			// we will probably stall on occlusion queries, so might as well have the RHI thread and GPU work while we wait.
+			SCOPE_CYCLE_COUNTER(STAT_PostInitViews_FlushDel);
+			FRHICommandListExecutor::GetImmediateCommandList().ImmediateFlush(EImmediateFlushType::FlushRHIThreadFlushResources);
+		}
 	}
 
 	static const auto ClearMethodCVar = IConsoleManager::Get().FindTConsoleVariableDataInt(TEXT("r.ClearSceneMethod"));
@@ -1340,10 +1348,18 @@ void FDeferredShadingSceneRenderer::Render(FRHICommandListImmediate& RHICmdList)
 			bDidAfterTaskWork = true; // only do this once
 			if (bDoInitViewAftersPrepass)
 			{
-				InitViewsPossiblyAfterPrepass(RHICmdList, ILCTaskData, UpdateViewCustomDataEvents);
-				PrepareDistanceFieldScene(RHICmdList, false);
+				{
+					SCOPED_GPU_STAT(RHICmdList, VisibilityCommands);
+					InitViewsPossiblyAfterPrepass(RHICmdList, ILCTaskData, UpdateViewCustomDataEvents);
+				}
+				
+				{
+					SCOPED_GPU_STAT(RHICmdList, GPUSceneUpdate);
+					PrepareDistanceFieldScene(RHICmdList, false);
+				}
 
 				{
+					SCOPED_GPU_STAT(RHICmdList, UploadDynamicBuffers);
 					SCOPE_CYCLE_COUNTER(STAT_FDeferredShadingSceneRenderer_FGlobalDynamicVertexBuffer_Commit);
 					DynamicVertexBufferForInitShadows.Commit();
 					DynamicIndexBufferForInitShadows.Commit();
@@ -2128,6 +2144,7 @@ void FDeferredShadingSceneRenderer::Render(FRHICommandListImmediate& RHICmdList)
 	IRendererModule& RendererModule = GetRendererModule();
 	if (RendererModule.HasPostOpaqueExtentions())
 	{
+		SCOPED_GPU_STAT(RHICmdList, PostOpaqueExtensions);
 		FSceneTexturesUniformParameters SceneTextureParameters;
 		SetupSceneTextureUniformParameters(SceneContext, FeatureLevel, ESceneTextureSetupMode::SceneDepth | ESceneTextureSetupMode::GBuffers, SceneTextureParameters);
 		TUniformBufferRef<FSceneTexturesUniformParameters> SceneTextureUniformBuffer = TUniformBufferRef<FSceneTexturesUniformParameters>::CreateUniformBufferImmediate(SceneTextureParameters, UniformBuffer_SingleFrame);
@@ -2286,7 +2303,7 @@ void FDeferredShadingSceneRenderer::Render(FRHICommandListImmediate& RHICmdList)
 
 	if (bUseVirtualTexturing)
 	{
-		SCOPED_GPU_STAT(RHICmdList, VirtualTextureFeedback);
+		SCOPED_GPU_STAT(RHICmdList, VirtualTextureUpdate);
 		// No pass after this can make VT page requests
 		SceneContext.VirtualTextureFeedback.TransferGPUToCPU(RHICmdList, Views[0].ViewRect);
 	}
@@ -2419,35 +2436,37 @@ void FDeferredShadingSceneRenderer::Render(FRHICommandListImmediate& RHICmdList)
 		SceneContext.AdjustGBufferRefCount(RHICmdList, -1);
 	}
 
-	for (int32 ViewIndex = 0; ViewIndex < Views.Num(); ViewIndex++)
 	{
-		ShaderPrint::EndView(Views[ViewIndex]);
-	}
+		SCOPED_DRAW_EVENT(RHICmdList, AfterPostProcessing);
+		for (int32 ViewIndex = 0; ViewIndex < Views.Num(); ViewIndex++)
+		{
+			ShaderPrint::EndView(Views[ViewIndex]);
+		}
 
 #if WITH_MGPU
-	DoCrossGPUTransfers(RHICmdList, RenderTargetGPUMask);
+		DoCrossGPUTransfers(RHICmdList, RenderTargetGPUMask);
 #endif
 
-	//grab the new transform out of the proxies for next frame
-	SceneContext.SceneVelocity.SafeRelease();
+		//grab the new transform out of the proxies for next frame
+		SceneContext.SceneVelocity.SafeRelease();
 
-	// Invalidate the lighting channels
-	SceneContext.LightingChannels.SafeRelease();
+		// Invalidate the lighting channels
+		SceneContext.LightingChannels.SafeRelease();
 
 
 #if RHI_RAYTRACING
-	// Release resources that were bound to the ray tracing scene to allow them to be immediately recycled.
-	for (int32 ViewIndex = 0; ViewIndex < Views.Num(); ++ViewIndex)
-	{
-		FViewInfo& View = Views[ViewIndex];
-		if (View.RayTracingScene.RayTracingSceneRHI)
+		// Release resources that were bound to the ray tracing scene to allow them to be immediately recycled.
+		for (int32 ViewIndex = 0; ViewIndex < Views.Num(); ++ViewIndex)
 		{
-			RHICmdList.ClearRayTracingBindings(View.RayTracingScene.RayTracingSceneRHI);
-			View.RayTracingScene.RayTracingSceneRHI.SafeRelease();
+			FViewInfo& View = Views[ViewIndex];
+			if (View.RayTracingScene.RayTracingSceneRHI)
+			{
+				RHICmdList.ClearRayTracingBindings(View.RayTracingScene.RayTracingSceneRHI);
+				View.RayTracingScene.RayTracingSceneRHI.SafeRelease();
+			}
 		}
-	}
 #endif //  RHI_RAYTRACING
-
+	}
 	{
 		SCOPE_CYCLE_COUNTER(STAT_FDeferredShadingSceneRenderer_RenderFinish);
 		SCOPED_GPU_STAT(RHICmdList, FrameRenderFinish);
