Fix crash when switching the feature level in the editor viewport when fog volumes are used.

diff --git a/Engine/Source/Runtime/Renderer/Private/VolumetricFogVoxelization.cpp b/Engine/Source/Runtime/Renderer/Private/VolumetricFogVoxelization.cpp
--- a/Engine/Source/Runtime/Renderer/Private/VolumetricFogVoxelization.cpp
+++ b/Engine/Source/Runtime/Renderer/Private/VolumetricFogVoxelization.cpp
@@ -123,7 +123,9 @@ public:
 		VertexBuffer->Buffers.StaticMeshVertexBuffer.BindPackedTexCoordVertexBuffer(this, NewData);
 		VertexBuffer->Buffers.StaticMeshVertexBuffer.BindLightMapVertexBuffer(this, NewData, 0);
 		FColorVertexBuffer::BindDefaultColorVertexBuffer(this, NewData, FColorVertexBuffer::NullBindStride::ZeroForDefaultBufferBind);
-		SetData(NewData);
+		// Don't call SetData(), because that ends up calling UpdateRHI(), and if the resource has already been initialized
+		// (e.g. when switching the feature level in the editor), that calls InitRHI(), resulting in an infinite loop.
+		Data = NewData;
 		FLocalVertexFactory::InitRHI();
 	}
 
