Added velocity output for PPLL pass.
+ some renaming

diff --git a/Engine/Shaders/Private/HairStrands/HairStrandsVisibilityCommon.ush b/Engine/Shaders/Private/HairStrands/HairStrandsVisibilityCommon.ush
--- a/Engine/Shaders/Private/HairStrands/HairStrandsVisibilityCommon.ush
+++ b/Engine/Shaders/Private/HairStrands/HairStrandsVisibilityCommon.ush
@@ -295,7 +295,7 @@ FNodeData ConvertToNodeData(float Depth, FVisibilityData VisibilityData, float3
  //////////
  // PPLL
 
-// Footprint: 192bits | 20bytes
+// Footprint: 224bits | 28bytes
  struct FPPLLNodeData
  {
 	 // (backloged)TODO try f16(depth) f16(Specular) and make sure depth does not lead to artefact
@@ -305,20 +305,32 @@ FNodeData ConvertToNodeData(float Depth, FVisibilityData VisibilityData, float3
 	 uint Tangent_Coverage8bit;
 	 uint BaseColor_Roughness;
 	 uint NextNodeIndex;
+	 uint PackedVelocity;
  };
 
- FPPLLNodeData ConvertToPPLLNodeData(float Depth, FVisibilityData VisibilityData, float3 BaseColor, float Roughness, float Specular, uint NextNodeIndex)
+ uint PackVelocity(float2 EncodedVelocity)
+ {
+	 return f32tof16(EncodedVelocity.x) << 16 | f32tof16(EncodedVelocity.y);
+ }
+
+ float2 UnpackVelocity(uint PackedVelocity)
+ {
+	 return float2(f16tof32(PackedVelocity >> 16), f16tof32(PackedVelocity));
+ }
+
+ FPPLLNodeData ConvertToPPLLNodeData(float Depth, FVisibilityData VisibilityData, float3 BaseColor, float Roughness, float Specular, float2 Velocity, uint NextNodeIndex)
  {
 	 // 4x16bits
 	 const uint4 EncodedVisibilityData = EncodeVisibilityData(VisibilityData);
 
 	 FPPLLNodeData Out;
-	 Out.Depth = Depth;																								// 32bits float
-	 Out.Tangent_Coverage8bit = (0xFFFF & EncodedVisibilityData.x) | ((0xFFFF & EncodedVisibilityData.y) << 16);	// 32bits uint
-	 Out.PrimitiveID_ClusterID = (0xFFFF & EncodedVisibilityData.z) | ((0xFFFF & EncodedVisibilityData.w) << 16);	// 32bits uint
-	 Out.BaseColor_Roughness = float4ToUint(float4(sqrt(BaseColor), Roughness));									// 32bits uint
-	 Out.Specular = Specular;																						// 32 bits float
-	 Out.NextNodeIndex = NextNodeIndex;																				// 32bits next node
+	 Out.Depth = Depth;
+	 Out.Tangent_Coverage8bit = (0xFFFF & EncodedVisibilityData.x) | ((0xFFFF & EncodedVisibilityData.y) << 16);
+	 Out.PrimitiveID_ClusterID = (0xFFFF & EncodedVisibilityData.z) | ((0xFFFF & EncodedVisibilityData.w) << 16);
+	 Out.BaseColor_Roughness = float4ToUint(float4(sqrt(BaseColor), Roughness));
+	 Out.Specular = Specular;
+	 Out.NextNodeIndex = NextNodeIndex;
+	 Out.PackedVelocity = PackVelocity(Velocity);
 
 	 return Out;
  }
diff --git a/Engine/Shaders/Private/HairStrands/HairStrandsVisibilityCompaction.usf b/Engine/Shaders/Private/HairStrands/HairStrandsVisibilityCompaction.usf
--- a/Engine/Shaders/Private/HairStrands/HairStrandsVisibilityCompaction.usf
+++ b/Engine/Shaders/Private/HairStrands/HairStrandsVisibilityCompaction.usf
@@ -68,6 +68,17 @@ void MainCS(uint GroupIndex : SV_GroupIndex, uint3 DispatchThreadId : SV_Dispatc
 	uint OffsetInStorage = 0;
 	uint PixelSamplecount = 0;
 
+#if PERMUTATION_VELOCITY == VELOCITY_TYPE_AVG
+	float2 AverageVelocity = 0;
+	uint ValidVelocityCount = 0;
+#elif PERMUTATION_VELOCITY == VELOCITY_TYPE_CLOSEST
+	float2 ClosestVelocity = 0;
+	float ClosestDepth = 0;
+#elif PERMUTATION_VELOCITY == VELOCITY_TYPE_MAX
+	float2 MaxVelocity = 0;
+	float MaxVelocityMagnitude2 = 0;
+#endif
+
 	if (FirstNodeIndex != 0xFFFFFFFF && all(PixelCoord < OutputResolution))
 	{
 		//////////
@@ -102,7 +113,6 @@ void MainCS(uint GroupIndex : SV_GroupIndex, uint3 DispatchThreadId : SV_Dispatc
 				SortedIndex[PixelNodeCount] = NodeIndex;
 				SortedDepth[PixelNodeCount] = PPLLNode.Depth;
 				// TODO sort and keep closest, merge furthest
-
 			}
 
 			PixelNodeCount++;
@@ -138,6 +148,28 @@ void MainCS(uint GroupIndex : SV_GroupIndex, uint3 DispatchThreadId : SV_Dispatc
 
 			OutCompactNodeData[OffsetInStorage + i] = OutNodeData;
 			OutCompactNodeCoord[OffsetInStorage + i] = (PixelCoord.x & 0xFFFF) | ((PixelCoord.y & 0xFFFF) << 16);
+
+		#if PERMUTATION_VELOCITY != VELOCITY_TYPE_NONE
+			const float2 Velocity = UnpackVelocity(PPLLNode.PackedVelocity) * 0.0f;	// TODO: investigate why TAA jitter is still visible. Disabled for now.
+		#endif
+		#if PERMUTATION_VELOCITY == VELOCITY_TYPE_AVG
+			AverageVelocity += Velocity;
+			++ValidVelocityCount;
+		#elif PERMUTATION_VELOCITY == VELOCITY_TYPE_CLOSEST
+			if (ClosestDepth < Categorization.ClosestDepth) // Inverse-Z
+			{
+				ClosestVelocity = Velocity;
+				ClosestDepth = Categorization.ClosestDepth;
+			}
+		#elif PERMUTATION_VELOCITY == VELOCITY_TYPE_MAX
+			const float2 CurrentVelocity = Velocity;
+			const float CurrentVelocityMagnitude2 = dot(CurrentVelocity, CurrentVelocity);
+			if (CurrentVelocityMagnitude2 > MaxVelocityMagnitude2)
+			{
+				MaxVelocity = Velocity;
+				MaxVelocityMagnitude2 = CurrentVelocityMagnitude2;
+			}
+		#endif		
 		}
 
 	#if PERMUTATION_COVERAGE == 1
@@ -145,6 +177,14 @@ void MainCS(uint GroupIndex : SV_GroupIndex, uint3 DispatchThreadId : SV_Dispatc
 	#else
 		Categorization.PixelCoverage = 1.0f - Transmittance;
 	#endif
+
+	#if PERMUTATION_VELOCITY == VELOCITY_TYPE_AVG
+		OutVelocityTexture[PixelCoord] = EncodeVelocityToTexture(AverageVelocity / max(ValidVelocityCount, 1));
+	#elif PERMUTATION_VELOCITY == VELOCITY_TYPE_CLOSEST
+		OutVelocityTexture[PixelCoord] = EncodeVelocityToTexture(ClosestVelocity);
+	#elif PERMUTATION_VELOCITY == VELOCITY_TYPE_MAX
+		OutVelocityTexture[PixelCoord] = EncodeVelocityToTexture(MaxVelocity);
+	#endif		
 	}
 
 	FNodeDesc NodeDesc;
@@ -154,9 +194,6 @@ void MainCS(uint GroupIndex : SV_GroupIndex, uint3 DispatchThreadId : SV_Dispatc
 
 	Categorization.SampleCount = PixelSamplecount;
 	OutCategorizationTexture[PixelCoord] = EncodeCategorizationData(Categorization);
-
-
-
 }
 
 
diff --git a/Engine/Shaders/Private/HairStrands/HairStrandsVisibilityPS.usf b/Engine/Shaders/Private/HairStrands/HairStrandsVisibilityPS.usf
--- a/Engine/Shaders/Private/HairStrands/HairStrandsVisibilityPS.usf
+++ b/Engine/Shaders/Private/HairStrands/HairStrandsVisibilityPS.usf
@@ -109,6 +109,7 @@ void MainVisibility(
 	half3 BaseColor = 0;
 	float Roughness = 0;
 	float Specular  = 0;
+	float2 Velocity = 0;
 	float2 EncodedVelocity = 0;
 #if SUPPORT_MATERIAL_PROPERTY
 	{
@@ -128,7 +129,7 @@ void MainVisibility(
 		VisibilityData.Tangent = mul(normalize(GetMaterialNormalRaw(PixelMaterialInputs)), MaterialParameters.TangentToWorld);
 	}
 	{
-		const float2 Velocity = Calculate2DVelocity(PackedVelocityA, PackedVelocityC);
+		Velocity = Calculate2DVelocity(PackedVelocityA, PackedVelocityC);
 		EncodedVelocity = EncodeVelocityToTexture(Velocity);
 	}
 #endif
@@ -156,7 +157,7 @@ void MainVisibility(
 		// If we can append new node, add it and make it point to the node previously in head of list
 		uint NextNodeIndex;
 		InterlockedExchange(PPLLNodeIndex[PixelCoord], NodeIndex, NextNodeIndex);
-		const FPPLLNodeData Node = ConvertToPPLLNodeData(SvPosition.z, VisibilityData, BaseColor, Roughness, Specular, NextNodeIndex);
+		const FPPLLNodeData Node = ConvertToPPLLNodeData(SvPosition.z, VisibilityData, BaseColor, Roughness, Specular, Velocity, NextNodeIndex);
 		PPLLNodes[NodeIndex] = Node;
 	}
 #endif
diff --git a/Engine/Source/Runtime/Renderer/Private/HairStrands/HairStrandsVisibility.cpp b/Engine/Source/Runtime/Renderer/Private/HairStrands/HairStrandsVisibility.cpp
--- a/Engine/Source/Runtime/Renderer/Private/HairStrands/HairStrandsVisibility.cpp
+++ b/Engine/Source/Runtime/Renderer/Private/HairStrands/HairStrandsVisibility.cpp
@@ -887,11 +887,11 @@ static void AddHairVisibilityMSAAPass(
 	const FViewInfo* ViewInfo,
 	const FHairStrandsClusterDatas& ClusterDatas,
 	const FIntPoint& Resolution,
-	FRDGTextureRef& VisibilityIdTexture,
-	FRDGTextureRef& VisibilityMaterialTexture,
-	FRDGTextureRef& VisibilityAttributeTexture,
-	FRDGTextureRef& VisibilityVelocityTexture,
-	FRDGTextureRef& VisibilityDepthTexture)
+	FRDGTextureRef& OutVisibilityIdTexture,
+	FRDGTextureRef& OutVisibilityMaterialTexture,
+	FRDGTextureRef& OutVisibilityAttributeTexture,
+	FRDGTextureRef& OutVisibilityVelocityTexture,
+	FRDGTextureRef& OutVisibilityDepthTexture)
 {
 	const uint32 MSAASampleCount = FMath::RoundUpToPowerOfTwo(FMath::Clamp(GHairVisibilitySampleCount, 1, 16));
 
@@ -906,7 +906,7 @@ static void AddHairVisibilityMSAAPass(
 		Desc.Flags = TexCreate_None;
 		Desc.TargetableFlags = TexCreate_RenderTargetable | TexCreate_ShaderResource;
 		Desc.bForceSharedTargetAndShaderResource = true;
-		VisibilityIdTexture = GraphBuilder.CreateTexture(Desc, TEXT("HairVisibilityIDTexture"));
+		OutVisibilityIdTexture = GraphBuilder.CreateTexture(Desc, TEXT("HairVisibilityIDTexture"));
 	}
 
 	{
@@ -921,7 +921,7 @@ static void AddHairVisibilityMSAAPass(
 		Desc.TargetableFlags = TexCreate_RenderTargetable | TexCreate_ShaderResource;
 		Desc.ClearValue = FClearValueBinding(FLinearColor(0, 0, 0, 0));
 		Desc.bForceSharedTargetAndShaderResource = true;
-		VisibilityMaterialTexture = GraphBuilder.CreateTexture(Desc, TEXT("HairVisibilityMaterialTexture"));
+		OutVisibilityMaterialTexture = GraphBuilder.CreateTexture(Desc, TEXT("HairVisibilityMaterialTexture"));
 	}
 
 	{
@@ -936,7 +936,7 @@ static void AddHairVisibilityMSAAPass(
 		Desc.TargetableFlags = TexCreate_RenderTargetable | TexCreate_ShaderResource;
 		Desc.ClearValue = FClearValueBinding(FLinearColor(0, 0, 0, 0));
 		Desc.bForceSharedTargetAndShaderResource = true;
-		VisibilityAttributeTexture = GraphBuilder.CreateTexture(Desc, TEXT("HairVisibilityAttributeTexture"));
+		OutVisibilityAttributeTexture = GraphBuilder.CreateTexture(Desc, TEXT("HairVisibilityAttributeTexture"));
 	}
 
 	{
@@ -951,9 +951,9 @@ static void AddHairVisibilityMSAAPass(
 		Desc.TargetableFlags = TexCreate_RenderTargetable | TexCreate_ShaderResource;
 		Desc.ClearValue = FClearValueBinding(FLinearColor(0, 0, 0, 0));
 		Desc.bForceSharedTargetAndShaderResource = true;
-		VisibilityVelocityTexture = GraphBuilder.CreateTexture(Desc, TEXT("HairVisibilityVelocityTexture"));
+		OutVisibilityVelocityTexture = GraphBuilder.CreateTexture(Desc, TEXT("HairVisibilityVelocityTexture"));
 	}
-	AddClearGraphicPass(GraphBuilder, RDG_EVENT_NAME("HairStrandsClearVisibilityMSAAIdTexture"), ViewInfo, 0xFFFFFFFF, VisibilityIdTexture);
+	AddClearGraphicPass(GraphBuilder, RDG_EVENT_NAME("HairStrandsClearVisibilityMSAAIdTexture"), ViewInfo, 0xFFFFFFFF, OutVisibilityIdTexture);
 
 	// Manually clear RTs as using the Clear action on the RT, issue a global clean on all targets, while still need a special clear 
 	// for the PrimitiveId buffer
@@ -962,19 +962,19 @@ static void AddHairVisibilityMSAAPass(
 	if (GHairClearVisibilityBuffer)
 	{
 		LoadAction = ERenderTargetLoadAction::ELoad;
-		AddClearGraphicPass(GraphBuilder, RDG_EVENT_NAME("HairStrandsClearVisibilityMSAAMaterial"), ViewInfo, 0, VisibilityMaterialTexture);
-		AddClearGraphicPass(GraphBuilder, RDG_EVENT_NAME("HairStrandsClearVisibilityMSAAAttribute"), ViewInfo, 0, VisibilityAttributeTexture);
-		AddClearGraphicPass(GraphBuilder, RDG_EVENT_NAME("HairStrandsClearVisibilityMSAAVelocity"), ViewInfo, 0, VisibilityVelocityTexture);
+		AddClearGraphicPass(GraphBuilder, RDG_EVENT_NAME("HairStrandsClearVisibilityMSAAMaterial"), ViewInfo, 0, OutVisibilityMaterialTexture);
+		AddClearGraphicPass(GraphBuilder, RDG_EVENT_NAME("HairStrandsClearVisibilityMSAAAttribute"), ViewInfo, 0, OutVisibilityAttributeTexture);
+		AddClearGraphicPass(GraphBuilder, RDG_EVENT_NAME("HairStrandsClearVisibilityMSAAVelocity"), ViewInfo, 0, OutVisibilityVelocityTexture);
 	}
 
 	FVisibilityPassParameters* PassParameters = GraphBuilder.AllocParameters<FVisibilityPassParameters>();
-	PassParameters->RenderTargets[0] = FRenderTargetBinding(VisibilityIdTexture, ERenderTargetLoadAction::ELoad, 0);
-	PassParameters->RenderTargets[1] = FRenderTargetBinding(VisibilityMaterialTexture,  LoadAction, 0);
-	PassParameters->RenderTargets[2] = FRenderTargetBinding(VisibilityAttributeTexture, LoadAction, 0);
-	PassParameters->RenderTargets[3] = FRenderTargetBinding(VisibilityVelocityTexture,  LoadAction, 0);
+	PassParameters->RenderTargets[0] = FRenderTargetBinding(OutVisibilityIdTexture, ERenderTargetLoadAction::ELoad, 0);
+	PassParameters->RenderTargets[1] = FRenderTargetBinding(OutVisibilityMaterialTexture,  LoadAction, 0);
+	PassParameters->RenderTargets[2] = FRenderTargetBinding(OutVisibilityAttributeTexture, LoadAction, 0);
+	PassParameters->RenderTargets[3] = FRenderTargetBinding(OutVisibilityVelocityTexture,  LoadAction, 0);
 
 	PassParameters->RenderTargets.DepthStencil = FDepthStencilBinding(
-		VisibilityDepthTexture,
+		OutVisibilityDepthTexture,
 		ERenderTargetLoadAction::ELoad,
 		ERenderTargetLoadAction::ENoAction,
 		FExclusiveDepthStencil::DepthWrite_StencilNop);
@@ -987,10 +987,10 @@ static void AddHairVisibilityPPLLPass(
 	const FViewInfo* ViewInfo,
 	const FHairStrandsClusterDatas& ClusterDatas,
 	const FIntPoint& Resolution,
-	FRDGTextureRef& ViewZDepthTexture,
-	FRDGTextureRef& VisibilityPPLLNodeCounter,
-	FRDGTextureRef& VisibilityPPLLNodeIndex,
-	FRDGBufferRef& VisibilityPPLLNodeData)
+	FRDGTextureRef& InViewZDepthTexture,
+	FRDGTextureRef& OutVisibilityPPLLNodeCounter,
+	FRDGTextureRef& OutVisibilityPPLLNodeIndex,
+	FRDGBufferRef&  OutVisibilityPPLLNodeData)
 {
 	{
 		FRDGTextureDesc Desc;
@@ -1003,7 +1003,7 @@ static void AddHairVisibilityPPLLPass(
 		Desc.Flags = TexCreate_None;
 		Desc.TargetableFlags = TexCreate_UAV | TexCreate_ShaderResource;
 		Desc.ClearValue = FClearValueBinding(0);
-		VisibilityPPLLNodeCounter = GraphBuilder.CreateTexture(Desc, TEXT("HairVisibilityPPLLCounter"));
+		OutVisibilityPPLLNodeCounter = GraphBuilder.CreateTexture(Desc, TEXT("HairVisibilityPPLLCounter"));
 	}
 
 	{
@@ -1017,12 +1017,12 @@ static void AddHairVisibilityPPLLPass(
 		Desc.Flags = TexCreate_None;
 		Desc.TargetableFlags = TexCreate_UAV | TexCreate_ShaderResource;
 		Desc.ClearValue = FClearValueBinding(0);
-		VisibilityPPLLNodeIndex = GraphBuilder.CreateTexture(Desc, TEXT("HairVisibilityPPLLNodeIndex"));
+		OutVisibilityPPLLNodeIndex = GraphBuilder.CreateTexture(Desc, TEXT("HairVisibilityPPLLNodeIndex"));
 	}
 
 	const uint32 MaxNodeCount = GetMaxNodePerPixel(Resolution);
 	{
-		// Example: 24bytes * 8spp = 192bytes per pixel = 379Mb @ 1080p
+		// Example: 28bytes * 8spp = 224bytes per pixel = 442Mb @ 1080p
 		struct PPLLNodeData
 		{
 			uint32 Depth;
@@ -1031,19 +1031,20 @@ static void AddHairVisibilityPPLLPass(
 			uint32 BaseColor_Roughness;
 			uint32 Specular;
 			uint32 NextNodeIndex;
+			uint32 PackedVelocity;
 		};
 
-		VisibilityPPLLNodeData = GraphBuilder.CreateBuffer(FRDGBufferDesc::CreateStructuredDesc(sizeof(PPLLNodeData), MaxNodeCount), TEXT("HairVisibilityPPLLNodeData"));
+		OutVisibilityPPLLNodeData = GraphBuilder.CreateBuffer(FRDGBufferDesc::CreateStructuredDesc(sizeof(PPLLNodeData), MaxNodeCount), TEXT("HairVisibilityPPLLNodeData"));
 	}
-	AddClearUAVPass(GraphBuilder, RDG_EVENT_NAME("ClearHairVisibilityPPLLCounter"), VisibilityPPLLNodeCounter, 0);
-	AddClearUAVPass(GraphBuilder, RDG_EVENT_NAME("ClearHairVisibilityPPLLNodeIndex"), VisibilityPPLLNodeIndex, 0xFFFFFFFF);
+	AddClearUAVPass(GraphBuilder, RDG_EVENT_NAME("ClearHairVisibilityPPLLCounter"), OutVisibilityPPLLNodeCounter, 0);
+	AddClearUAVPass(GraphBuilder, RDG_EVENT_NAME("ClearHairVisibilityPPLLNodeIndex"), OutVisibilityPPLLNodeIndex, 0xFFFFFFFF);
 
 	FVisibilityPassParameters* PassParameters = GraphBuilder.AllocParameters<FVisibilityPassParameters>();
-	PassParameters->PPLLCounter = GraphBuilder.CreateUAV(FRDGTextureUAVDesc(VisibilityPPLLNodeCounter, 0));
-	PassParameters->PPLLNodeIndex = GraphBuilder.CreateUAV(FRDGTextureUAVDesc(VisibilityPPLLNodeIndex, 0));
-	PassParameters->PPLLNodeData = GraphBuilder.CreateUAV(FRDGBufferUAVDesc(VisibilityPPLLNodeData));
+	PassParameters->PPLLCounter = GraphBuilder.CreateUAV(FRDGTextureUAVDesc(OutVisibilityPPLLNodeCounter, 0));
+	PassParameters->PPLLNodeIndex = GraphBuilder.CreateUAV(FRDGTextureUAVDesc(OutVisibilityPPLLNodeIndex, 0));
+	PassParameters->PPLLNodeData = GraphBuilder.CreateUAV(FRDGBufferUAVDesc(OutVisibilityPPLLNodeData));
 	PassParameters->HairVisibilityPass_MaxPPLLNodeCount = MaxNodeCount;
-	PassParameters->RenderTargets.DepthStencil = FDepthStencilBinding(ViewZDepthTexture, ERenderTargetLoadAction::ELoad, ERenderTargetLoadAction::ENoAction, FExclusiveDepthStencil::DepthRead_StencilNop);
+	PassParameters->RenderTargets.DepthStencil = FDepthStencilBinding(InViewZDepthTexture, ERenderTargetLoadAction::ELoad, ERenderTargetLoadAction::ENoAction, FExclusiveDepthStencil::DepthRead_StencilNop);
 	AddHairVisibilityCommonPass(GraphBuilder, Scene, ViewInfo, ClusterDatas, HairVisibilityRenderMode_PPLL, PassParameters);
 }
 
