Remove use of MSAA depth from hair composition passes. We are only using the common categorisation texture now.

diff --git a/Engine/Shaders/Private/HairStrands/HairStrandsVisibilityComposeSubPixelPS.usf b/Engine/Shaders/Private/HairStrands/HairStrandsVisibilityComposeSubPixelPS.usf
--- a/Engine/Shaders/Private/HairStrands/HairStrandsVisibilityComposeSubPixelPS.usf
+++ b/Engine/Shaders/Private/HairStrands/HairStrandsVisibilityComposeSubPixelPS.usf
@@ -8,11 +8,7 @@
 #include "HairStrandsVisibilityCommon.ush"
 
 Texture2D<float4> SubPixelColorTexture;
-#if PERMUTATION_PPLL
 Texture2D<uint4> CategorisationTexture;
-#else
-Texture2DMS<float, HAIR_VISIBILITY_MAXSAMPLECOUNT> HairVisibilityDepthTexture;
-#endif
 
 void SubColorPS(
 	in FScreenVertexOutput Input,
@@ -27,23 +23,8 @@ void SubColorPS(
 	//if (OutColor.a == 0)
 	//	discard;
 
-#if PERMUTATION_PPLL
-
 	const FCategorizationData CategorisationData = DecodeCategorizationData(CategorisationTexture.Load(uint3(Input.Position.xy, 0)));
 	OutDepth = CategorisationData.ClosestDepth;
-
-#else
-
-	// Brute force reading of the depth value. This could be pre-computed during the categorize pass or the compaction pass
-	// This must match what is done in VisibilityDepthPS to fonction. TODO: do not write depth here as it is already in. 
-	OutDepth = 0;
-	for (uint SampleIt = 0; SampleIt < HAIR_VISIBILITY_MAXSAMPLECOUNT; ++SampleIt)
-	{
-		const float SampleDepth = HairVisibilityDepthTexture.Load(PixelCoord, SampleIt);
-		OutDepth = max(OutDepth, SampleDepth); // HAS_INVERTED_Z_BUFFER
-	}
-
-#endif
 }
 
 float VelocityThreshold;
diff --git a/Engine/Source/Runtime/Renderer/Private/HairStrands/HairStrandsComposition.cpp b/Engine/Source/Runtime/Renderer/Private/HairStrands/HairStrandsComposition.cpp
--- a/Engine/Source/Runtime/Renderer/Private/HairStrands/HairStrandsComposition.cpp
+++ b/Engine/Source/Runtime/Renderer/Private/HairStrands/HairStrandsComposition.cpp
@@ -21,12 +21,10 @@ class FHairVisibilityComposeSubPixelPS : public FGlobalShader
 	DECLARE_GLOBAL_SHADER(FHairVisibilityComposeSubPixelPS);
 	SHADER_USE_PARAMETER_STRUCT(FHairVisibilityComposeSubPixelPS, FGlobalShader);
 
-	class FPPLL : SHADER_PERMUTATION_INT("PERMUTATION_PPLL", 2);
-	using FPermutationDomain = TShaderPermutationDomain<FPPLL>;
+	using FPermutationDomain = TShaderPermutationDomain<>;
 
 	BEGIN_SHADER_PARAMETER_STRUCT(FParameters, )
 		SHADER_PARAMETER_RDG_TEXTURE(Texture2D, SubPixelColorTexture)
-		SHADER_PARAMETER_RDG_TEXTURE(Texture2D, HairVisibilityDepthTexture)
 		SHADER_PARAMETER_RDG_TEXTURE(Texture2D, CategorisationTexture)
 		RENDER_TARGET_BINDING_SLOTS()
 	END_SHADER_PARAMETER_STRUCT()
@@ -41,22 +39,17 @@ static void AddHairVisibilityComposeSubPixelPass(
 	FRDGBuilder& GraphBuilder,
 	const FViewInfo& View,
 	const FRDGTextureRef& SubPixelColorTexture,
-	const FRDGTextureRef& HairVisibilityDepthTexture,
 	const FRDGTextureRef CategorisationTexture,
 	FRDGTextureRef& OutColorTexture,
 	FRDGTextureRef& OutDepthTexture)
 {
-	const bool bPPLL = HairVisibilityDepthTexture == nullptr;
-
 	FHairVisibilityComposeSubPixelPS::FParameters* Parameters = GraphBuilder.AllocParameters<FHairVisibilityComposeSubPixelPS::FParameters>();
 	Parameters->SubPixelColorTexture = SubPixelColorTexture;
 	Parameters->CategorisationTexture = CategorisationTexture;
-	Parameters->HairVisibilityDepthTexture = HairVisibilityDepthTexture;
 	Parameters->RenderTargets[0] = FRenderTargetBinding(OutColorTexture, ERenderTargetLoadAction::ELoad);
 	Parameters->RenderTargets.DepthStencil = FDepthStencilBinding(OutDepthTexture, ERenderTargetLoadAction::ELoad, ERenderTargetLoadAction::ENoAction, FExclusiveDepthStencil::DepthWrite_StencilNop);
 
 	FHairVisibilityComposeSubPixelPS::FPermutationDomain PermutationVector;
-	PermutationVector.Set<FHairVisibilityComposeSubPixelPS::FPPLL>(bPPLL ? 1 : 0);	// TODO we need a permutation for each MSAA count we accept
 	TShaderMapRef<FHairVisibilityComposeSubPixelPS> PixelShader(View.ShaderMap, PermutationVector);
 	TShaderMapRef<FPostProcessVS> VertexShader(View.ShaderMap);
 	const FIntRect Viewport = View.ViewRect;
@@ -111,7 +104,6 @@ class FHairVisibilityFastResolvePS : public FGlobalShader
 	BEGIN_SHADER_PARAMETER_STRUCT(FParameters, )
 		SHADER_PARAMETER(float, VelocityThreshold)
 		SHADER_PARAMETER_RDG_TEXTURE(Texture2D, HairVisibilityVelocityTexture)
-		SHADER_PARAMETER_RDG_TEXTURE(Texture2D, HairVisibilityDepthTexture)
 		RENDER_TARGET_BINDING_SLOTS()
 	END_SHADER_PARAMETER_STRUCT()
 
@@ -125,7 +117,6 @@ static void AddHairVisibilityFastResolvePass(
 	FRDGBuilder& GraphBuilder,
 	const FViewInfo& View,
 	const FRDGTextureRef& HairVisibilityVelocityTexture,
-	const FRDGTextureRef& HairVisibilityDepthTexture,
 	FRDGTextureRef& OutDepthTexture)
 {
 	const FIntPoint Resolution = OutDepthTexture->Desc.Extent;
@@ -148,7 +139,6 @@ static void AddHairVisibilityFastResolvePass(
 
 	FHairVisibilityFastResolvePS::FParameters* Parameters = GraphBuilder.AllocParameters<FHairVisibilityFastResolvePS::FParameters>();
 	Parameters->HairVisibilityVelocityTexture = HairVisibilityVelocityTexture;
-	Parameters->HairVisibilityDepthTexture = HairVisibilityDepthTexture;
 	Parameters->VelocityThreshold = VelocityThreshold;
 	Parameters->RenderTargets[0] = FRenderTargetBinding(DummyTexture, ERenderTargetLoadAction::ENoAction);
 	Parameters->RenderTargets.DepthStencil = FDepthStencilBinding(
@@ -235,14 +225,11 @@ void RenderHairComposeSubPixel(
 		{
 			if (ViewIndex < HairVisibilityViews.HairDatas.Num())
 			{
-				TRefCountPtr<IPooledRenderTarget> DepthTexture = HairVisibilityViews.HairDatas[ViewIndex].DepthTexture;
 				TRefCountPtr<IPooledRenderTarget> CategorisationTexture = HairVisibilityViews.HairDatas[ViewIndex].CategorizationTexture;
-				if (!DepthTexture && !CategorisationTexture)
+				if (!CategorisationTexture)
 				{
 					continue; // Automatically skip for any view not rendering hair
 				}
-
-				const FRDGTextureRef RDGHairVisibilityDepthTexture = DepthTexture ? GraphBuilder.RegisterExternalTexture(DepthTexture, TEXT("HairVisibilityDepthTexture")) : nullptr;
 				const FRDGTextureRef RDGCategorisationTexture = CategorisationTexture ? GraphBuilder.RegisterExternalTexture(CategorisationTexture, TEXT("HairVisibilityCategorisationTexture")) : nullptr;
 
 				// #hair_todo : compose partially covered hair with transparent surface: this can be done by 
@@ -252,7 +239,6 @@ void RenderHairComposeSubPixel(
 					GraphBuilder,
 					View,
 					SceneColorSubPixelTexture,
-					RDGHairVisibilityDepthTexture,
 					RDGCategorisationTexture,
 					SceneColorTexture,
 					SceneColorDepth);
@@ -264,7 +250,6 @@ void RenderHairComposeSubPixel(
 						GraphBuilder,
 						View,
 						RDGHairVisibilityVelocityTexture,
-						RDGHairVisibilityDepthTexture,
 						SceneColorDepth);
 				}
 			}
