Work around a driver bug on Magic Leap resulting in corrupted deferred decals

diff --git a/Engine/Config/BaseDeviceProfiles.ini b/Engine/Config/BaseDeviceProfiles.ini
--- a/Engine/Config/BaseDeviceProfiles.ini
+++ b/Engine/Config/BaseDeviceProfiles.ini
@@ -667,6 +667,7 @@ BaseProfileName=
 +CVars=r.LightShaftQuality=0
 +CVars=r.RefractionQuality=0
 +CVars=r.ShadowQuality=2
++CVars=r.Mobile.ForceDepthResolve=1
 
 [LuminClient DeviceProfile]
 DeviceType=Lumin
diff --git a/Engine/Source/Runtime/Renderer/Private/MobileShadingRenderer.cpp b/Engine/Source/Runtime/Renderer/Private/MobileShadingRenderer.cpp
--- a/Engine/Source/Runtime/Renderer/Private/MobileShadingRenderer.cpp
+++ b/Engine/Source/Runtime/Renderer/Private/MobileShadingRenderer.cpp
@@ -392,11 +392,10 @@ void FMobileSceneRenderer::Render(FRHICommandListImmediate& RHICmdList)
 	// See CVarMobileForceDepthResolve use in ConditionalResolveSceneDepth.
 	const bool bForceDepthResolve = CVarMobileForceDepthResolve.GetValueOnRenderThread() == 1;
 	const bool bSeparateTranslucencyActive = IsMobileSeparateTranslucencyActive(View);
-	bool bKeepDepthContent = bRenderToSceneColor &&
-		(bForceDepthResolve ||
-		 bSeparateTranslucencyActive ||
+	bool bKeepDepthContent = bForceDepthResolve || (bRenderToSceneColor &&
+		(bSeparateTranslucencyActive ||
 		 View.bIsReflectionCapture ||
-		 (View.bIsSceneCapture && (ViewFamily.SceneCaptureSource == ESceneCaptureSource::SCS_SceneColorHDR || ViewFamily.SceneCaptureSource == ESceneCaptureSource::SCS_SceneColorSceneDepth)));
+		 (View.bIsSceneCapture && (ViewFamily.SceneCaptureSource == ESceneCaptureSource::SCS_SceneColorHDR || ViewFamily.SceneCaptureSource == ESceneCaptureSource::SCS_SceneColorSceneDepth))));
 
 	// Whether to submit cmdbuffer with offscreen rendering before doing post-processing
 	bool bSubmitOffscreenRendering = !bGammaSpace || bRenderToSceneColor;
