Now allowing for more than 8 nodes for high quality ppll mode (16 and 32).

diff --git a/Engine/Shaders/Private/HairStrands/HairStrandsVisibilityCompaction.usf b/Engine/Shaders/Private/HairStrands/HairStrandsVisibilityCompaction.usf
--- a/Engine/Shaders/Private/HairStrands/HairStrandsVisibilityCompaction.usf
+++ b/Engine/Shaders/Private/HairStrands/HairStrandsVisibilityCompaction.usf
@@ -28,7 +28,7 @@ float CosTangentThreshold;
 
 uint2 OutputResolution;
 uint MaxNodeCount;
-uint HairVisibilitySampleCount;	// MSAA sample count
+uint HairVisibilityMSAASampleCount;
 uint2 ResolutionOffset;
 
 #if PERMUTATION_COVERAGE == 1
@@ -56,6 +56,7 @@ groupshared uint AllocationNodeOffset;
 
 
 #if PERMUTATION_PPLL
+#define PPLL_MAXELEMENTCOUNT PERMUTATION_PPLL
 
 Texture2D<uint> PPLLCounter;
 Texture2D<uint> PPLLNodeIndex;
@@ -85,8 +86,8 @@ void MainCS(uint GroupIndex : SV_GroupIndex, uint3 DispatchThreadId : SV_Dispatc
 	
 	const bool bExecute = FirstNodeIndex != 0xFFFFFFFF && all(PixelCoord < OutputResolution);
 
-	uint SortedIndex[HAIR_VISIBILITY_MAXSAMPLECOUNT];
-	float SortedDepth[HAIR_VISIBILITY_MAXSAMPLECOUNT];
+	uint SortedIndex[PPLL_MAXELEMENTCOUNT];
+	float SortedDepth[PPLL_MAXELEMENTCOUNT];
 	float ValidPixelSampleTotalCoverage = 0.0f;
 	float Transmittance = 1.0f;
 	uint PixelNodeCount = 0;
@@ -109,14 +110,14 @@ void MainCS(uint GroupIndex : SV_GroupIndex, uint3 DispatchThreadId : SV_Dispatc
 
 			const float Coverage = From8bitCoverage(VisibilityData.Coverage8bit);
 			Transmittance *= 1.0f - Coverage;
-			if (PixelNodeCount < HAIR_VISIBILITY_MAXSAMPLECOUNT)
+			if (PixelNodeCount < PPLL_MAXELEMENTCOUNT)
 			{
 				ValidPixelSampleTotalCoverage += Coverage;
 
-				// Always insert up to HAIR_VISIBILITY_MAXSAMPLECOUNT data
+				// Always insert up to PPLL_MAXELEMENTCOUNT data
 				SortedIndex[PixelNodeCount] = NodeIndex;
 				SortedDepth[PixelNodeCount] = PPLLNode.Depth;
-				// TODO sort and keep closest, merge furthest
+				// TODO sort and keep closest, merge furthest 
 			}
 
 			PixelNodeCount++;
@@ -133,7 +134,7 @@ void MainCS(uint GroupIndex : SV_GroupIndex, uint3 DispatchThreadId : SV_Dispatc
 	GroupMemoryBarrierWithGroupSync();
 	// Now notify how many nodes this thread wants to allocate
 	uint PixelDataSubOffsetInStorage = 0; 
-	PixelSamplecount = min(PixelNodeCount, HAIR_VISIBILITY_MAXSAMPLECOUNT);
+	PixelSamplecount = min(PixelNodeCount, PPLL_MAXELEMENTCOUNT);
 	InterlockedAdd(AllocationNodeCount, PixelSamplecount, PixelDataSubOffsetInStorage);
 	GroupMemoryBarrierWithGroupSync();
 	// Now allocate all the nodes for this group contiguously in memory
@@ -317,7 +318,7 @@ void MainCS(uint GroupIndex : SV_GroupIndex, uint3 DispatchThreadId : SV_Dispatc
 	}
 
 	uint2 PixelCoord = DispatchThreadId.xy + ResolutionOffset;
-	if (HairVisibilitySampleCount == 0 || PixelCoord.x >= OutputResolution.x || PixelCoord.y >= OutputResolution.y)
+	if (HairVisibilityMSAASampleCount == 0 || PixelCoord.x >= OutputResolution.x || PixelCoord.y >= OutputResolution.y)
 		PixelCoord = uint2(0, 0);
 
 	FSampleSetDesc SampleDesc;
@@ -328,7 +329,7 @@ void MainCS(uint GroupIndex : SV_GroupIndex, uint3 DispatchThreadId : SV_Dispatc
 		Samples,
 		SampleDesc,
 		PixelCoord,
-		HairVisibilitySampleCount,
+		HairVisibilityMSAASampleCount,
 		OpaqueDepth,
 		MSAA_IDTexture, 
 		MSAA_DepthTexture);
diff --git a/Engine/Source/Runtime/Renderer/Private/HairStrands/HairStrandsVisibility.cpp b/Engine/Source/Runtime/Renderer/Private/HairStrands/HairStrandsVisibility.cpp
--- a/Engine/Source/Runtime/Renderer/Private/HairStrands/HairStrandsVisibility.cpp
+++ b/Engine/Source/Runtime/Renderer/Private/HairStrands/HairStrandsVisibility.cpp
@@ -64,6 +64,28 @@ EHairVisibilityRenderMode GetHairVisibilityRenderMode()
 	return GHairVisibilityPPLL >0 ? HairVisibilityRenderMode_PPLL : HairVisibilityRenderMode_MSAA;
 }
 
+uint32 GetPPLLMaxElementCount()
+{
+	// The following must match the FPPLL permutation of FHairVisibilityPrimitiveIdCompactionCS.
+	if (GHairVisibilityPPLLMaxPixelNodeCount == 0)
+	{
+		return 0;
+	}
+	else if (GHairVisibilityPPLLMaxPixelNodeCount <= 8)
+	{
+		return 8;
+	}
+	else if (GHairVisibilityPPLLMaxPixelNodeCount <= 16)
+	{
+		return 16;
+	}
+	else //if (GHairVisibilityPPLLMaxPixelNodeCount <= 32)
+	{
+		return 32;
+	}
+	// If more is needed: please check out EncodeNodeDesc from HairStrandsVisibilityCommon.ush to verify node count representation limitations.
+}
+
 uint32 GetHairVisibilitySampleCount()
 {
 	return GetHairVisibilityRenderMode() == HairVisibilityRenderMode_MSAA ? FMath::Clamp(GHairVisibilitySampleCount,1,16) : 1;
@@ -495,13 +517,13 @@ class FHairVisibilityPrimitiveIdCompactionCS : public FGlobalShader
 	class FVelocity : SHADER_PERMUTATION_INT("PERMUTATION_VELOCITY", 4);
 	class FCoverage : SHADER_PERMUTATION_INT("PERMUTATION_COVERAGE", 2);
 	class FMaterial : SHADER_PERMUTATION_INT("PERMUTATION_MATERIAL_COMPACTION", 2);
-	class FPPLL     : SHADER_PERMUTATION_INT("PERMUTATION_PPLL", 2);
+	class FPPLL     : SHADER_PERMUTATION_SPARSE_INT("PERMUTATION_PPLL", 0, 8, 16, 32); // See GetPPLLMaxElementCount
 	using FPermutationDomain = TShaderPermutationDomain<FVendor, FVelocity, FCoverage, FMaterial, FPPLL>;
 
 	BEGIN_SHADER_PARAMETER_STRUCT(FParameters, )
 		SHADER_PARAMETER(FIntPoint, OutputResolution)
 		SHADER_PARAMETER(uint32, MaxNodeCount)
-		SHADER_PARAMETER(uint32, HairVisibilitySampleCount)
+		SHADER_PARAMETER(uint32, HairVisibilityMSAASampleCount)
 		SHADER_PARAMETER(FIntPoint, ResolutionOffset)
 		SHADER_PARAMETER(float, DepthTheshold)
 		SHADER_PARAMETER(float, CosTangentThreshold)
@@ -609,10 +631,10 @@ static void AddHairVisibilityPrimitiveIdCompactionPass(
 	AddClearUAVPass(GraphBuilder, RDG_EVENT_NAME("HairStrandsClearCompactionOffsetAndCount"), OutCompactNodeIndex, 0);
 	AddClearUAVPass(GraphBuilder, RDG_EVENT_NAME("HairStrandsClearCategorizationTexture"), OutCategorizationTexture, 0);
 
-	const float PPLLHardCodedPixelSampleCount = 8;
-	const uint32 HairVisibilitySampleCount = bUsePPLL ? PPLLHardCodedPixelSampleCount : GetHairVisibilitySampleCount();
-	const uint32 SampleCount = FMath::RoundUpToPowerOfTwo(HairVisibilitySampleCount);
-	const uint32 MaxNodeCount = Resolution.X * Resolution.Y * SampleCount;
+	const uint32 HairVisibilityMSAASampleCount = GetHairVisibilitySampleCount();
+	const uint32 SampleCount = FMath::RoundUpToPowerOfTwo(HairVisibilityMSAASampleCount);
+	const uint32 PPLLMaxElementCount = GetPPLLMaxElementCount();
+	const uint32 MaxNodeCount = Resolution.X * Resolution.Y * (GetHairVisibilityRenderMode() == HairVisibilityRenderMode_MSAA ? SampleCount : PPLLMaxElementCount);
 	{
 		struct NodeData
 		{
@@ -641,13 +663,13 @@ static void AddHairVisibilityPrimitiveIdCompactionPass(
 	PermutationVector.Set<FHairVisibilityPrimitiveIdCompactionCS::FVelocity>(bWriteOutVelocity ? FMath::Clamp(GHairVelocityType+1, 0, 3) : 0);
 	PermutationVector.Set<FHairVisibilityPrimitiveIdCompactionCS::FCoverage>(PassParameters->CoverageTexture ? 1 : 0);
 	PermutationVector.Set<FHairVisibilityPrimitiveIdCompactionCS::FMaterial>(GHairStrandsMaterialCompactionEnable ? 1 : 0);
-	PermutationVector.Set<FHairVisibilityPrimitiveIdCompactionCS::FPPLL>(bUsePPLL ? 1 : 0);
+	PermutationVector.Set<FHairVisibilityPrimitiveIdCompactionCS::FPPLL>(bUsePPLL ? PPLLMaxElementCount : 0);
 
 	PassParameters->OutputResolution = Resolution;
 	PassParameters->MaxNodeCount = MaxNodeCount;
 	PassParameters->DepthTheshold = FMath::Clamp(GHairStrandsMaterialCompactionDepthThreshold, 0.f, 100.f);
 	PassParameters->CosTangentThreshold = FMath::Cos(FMath::DegreesToRadians(FMath::Clamp(GHairStrandsMaterialCompactionTangentThreshold, 0.f, 90.f)));
-	PassParameters->HairVisibilitySampleCount = HairVisibilitySampleCount;
+	PassParameters->HairVisibilityMSAASampleCount = HairVisibilityMSAASampleCount;
 	PassParameters->SceneTexturesStruct = CreateUniformBufferImmediate(SceneTextures, EUniformBufferUsage::UniformBuffer_SingleDraw);
 	PassParameters->ViewUniformBuffer = View.ViewUniformBuffer;
 	PassParameters->OutCompactNodeCounter = GraphBuilder.CreateUAV(CompactCounter);
