Enable Dynamic shadows with multiview on mobile. Split DirectionalLightScreenToShadow's view dependent part into ViewUniformShaderParameters.MobileMultiviewShadowTransform,  so multiview can work correctly.

[Merge 9702473 from Oculus partner branch]

diff --git a/Engine/Shaders/Private/MobileBasePassPixelShader.usf b/Engine/Shaders/Private/MobileBasePassPixelShader.usf
--- a/Engine/Shaders/Private/MobileBasePassPixelShader.usf
+++ b/Engine/Shaders/Private/MobileBasePassPixelShader.usf
@@ -545,7 +545,7 @@ void Main(
 		half Shadow = 1.0f;
 	#if !MATERIAL_SHADINGMODEL_UNLIT
 		Shadow = GetPrimaryPrecomputedShadowMask(Interpolants).r;
-		#if DIRECTIONAL_LIGHT_CSM && !MOBILE_MULTI_VIEW && !MATERIAL_SHADINGMODEL_SINGLELAYERWATER
+		#if DIRECTIONAL_LIGHT_CSM && !MATERIAL_SHADINGMODEL_SINGLELAYERWATER
 		// Cascaded Shadow Map
 		{
 			FPCFSamplerSettings Settings;
@@ -563,7 +563,12 @@ void Main(
 			{
 				if (MaterialParameters.ScreenPosition.w < MobileDirectionalLight.DirectionalLightShadowDistances[i])
 				{
+				#if MOBILE_MULTI_VIEW
+					ShadowPosition = mul(float4(MaterialParameters.ScreenPosition.xyw, 1), ResolvedView.MobileMultiviewShadowTransform);
+					ShadowPosition = mul(ShadowPosition, MobileDirectionalLight.DirectionalLightScreenToShadow[i]);
+				#else
 					ShadowPosition = mul(float4(MaterialParameters.ScreenPosition.xyw, 1), MobileDirectionalLight.DirectionalLightScreenToShadow[i]);
+				#endif
 					break; // position found.
 				}
 			}
diff --git a/Engine/Source/Runtime/Engine/Private/SceneView.cpp b/Engine/Source/Runtime/Engine/Private/SceneView.cpp
--- a/Engine/Source/Runtime/Engine/Private/SceneView.cpp
+++ b/Engine/Source/Runtime/Engine/Private/SceneView.cpp
@@ -2296,6 +2296,14 @@ void FSceneView::SetupCommonViewUniformBufferParameters(
 		FPlane(0, 0, ProjectionMatrixUnadjustedForRHI.M[3][2], 0))
 		* InViewMatrices.GetInvTranslatedViewProjectionMatrix();
 
+	ViewUniformShaderParameters.MobileMultiviewShadowTransform = FMatrix(
+		FPlane(1, 0, 0, 0),
+		FPlane(0, 1, 0, 0),
+		FPlane(0, 0, InViewMatrices.GetProjectionMatrix().M[2][2], 1),
+		FPlane(0, 0, InViewMatrices.GetProjectionMatrix().M[3][2], 0)) *
+		InViewMatrices.GetInvTranslatedViewProjectionMatrix() *
+		FTranslationMatrix(-InViewMatrices.GetPreViewTranslation());
+
 	ViewUniformShaderParameters.PrevScreenToTranslatedWorld = FMatrix(
 		FPlane(1, 0, 0, 0),
 		FPlane(0, 1, 0, 0),
diff --git a/Engine/Source/Runtime/Engine/Public/SceneView.h b/Engine/Source/Runtime/Engine/Public/SceneView.h
--- a/Engine/Source/Runtime/Engine/Public/SceneView.h
+++ b/Engine/Source/Runtime/Engine/Public/SceneView.h
@@ -575,6 +575,7 @@ enum ETranslucencyVolumeCascade
 	VIEW_UNIFORM_BUFFER_MEMBER(FMatrix, SVPositionToTranslatedWorld) \
 	VIEW_UNIFORM_BUFFER_MEMBER(FMatrix, ScreenToWorld) \
 	VIEW_UNIFORM_BUFFER_MEMBER(FMatrix, ScreenToTranslatedWorld) \
+	VIEW_UNIFORM_BUFFER_MEMBER(FMatrix, MobileMultiviewShadowTransform) \
 	VIEW_UNIFORM_BUFFER_MEMBER_EX(FVector, ViewForward, EShaderPrecisionModifier::Half) \
 	VIEW_UNIFORM_BUFFER_MEMBER_EX(FVector, ViewUp, EShaderPrecisionModifier::Half) \
 	VIEW_UNIFORM_BUFFER_MEMBER_EX(FVector, ViewRight, EShaderPrecisionModifier::Half) \
diff --git a/Engine/Source/Runtime/Renderer/Private/ShadowRendering.cpp b/Engine/Source/Runtime/Renderer/Private/ShadowRendering.cpp
--- a/Engine/Source/Runtime/Renderer/Private/ShadowRendering.cpp
+++ b/Engine/Source/Runtime/Renderer/Private/ShadowRendering.cpp
@@ -1178,21 +1178,26 @@ FMatrix FProjectedShadowInfo::GetScreenToShadowMatrix(const FSceneView& View, ui
 	const float InvBufferResolutionY = 1.0f / (float)ShadowBufferResolution.Y;
 	const float ShadowResolutionFractionY = 0.5f * (float)TileResolutionY * InvBufferResolutionY;
 	// Calculate the matrix to transform a screenspace position into shadow map space
-	FMatrix ScreenToShadow = 
+
+	FMatrix ScreenToShadow;
+	FMatrix ViewDependentTransform =
 		// Z of the position being transformed is actually view space Z, 
-		// Transform it into post projection space by applying the projection matrix,
-		// Which is the required space before applying View.InvTranslatedViewProjectionMatrix
+			// Transform it into post projection space by applying the projection matrix,
+			// Which is the required space before applying View.InvTranslatedViewProjectionMatrix
 		FMatrix(
 			FPlane(1,0,0,0),
 			FPlane(0,1,0,0),
 			FPlane(0,0,View.ViewMatrices.GetProjectionMatrix().M[2][2],1),
-			FPlane(0,0,View.ViewMatrices.GetProjectionMatrix().M[3][2],0)) * 
+			FPlane(0,0,View.ViewMatrices.GetProjectionMatrix().M[3][2],0)) *
 		// Transform the post projection space position into translated world space
 		// Translated world space is normal world space translated to the view's origin, 
 		// Which prevents floating point imprecision far from the world origin.
 		View.ViewMatrices.GetInvTranslatedViewProjectionMatrix() *
+		FTranslationMatrix(-View.ViewMatrices.GetPreViewTranslation());
+
+	FMatrix ShadowMapDependentTransform =
 		// Translate to the origin of the shadow's translated world space
-		FTranslationMatrix(PreShadowTranslation - View.ViewMatrices.GetPreViewTranslation()) *
+		FTranslationMatrix(PreShadowTranslation) *
 		// Transform into the shadow's post projection space
 		// This has to be the same transform used to render the shadow depths
 		SubjectAndReceiverMatrix *
@@ -1207,8 +1212,19 @@ FMatrix FProjectedShadowInfo::GetScreenToShadowMatrix(const FSceneView& View, ui
 				(TileOffsetY + BorderSize) * InvBufferResolutionY + ShadowResolutionFractionY,
 				0,
 				1
-			)
-		);
+				)
+			);
+
+	if (View.bIsMobileMultiViewEnabled && View.Family->Views.Num() > 0)
+	{
+		// In Multiview, we split ViewDependentTransform out into ViewUniformShaderParameters.MobileMultiviewShadowTransform
+		// So we can multiply it later in shader.
+		ScreenToShadow = ShadowMapDependentTransform;
+	}
+	else
+	{
+		ScreenToShadow = ViewDependentTransform * ShadowMapDependentTransform;
+	}
 	return ScreenToShadow;
 }
 
