Added option to use full precision (32f per-channel) sky LUT. (r.SkyAtmosphere.LUT32)

diff --git a/Engine/Source/Runtime/Renderer/Private/SkyAtmosphereRendering.cpp b/Engine/Source/Runtime/Renderer/Private/SkyAtmosphereRendering.cpp
--- a/Engine/Source/Runtime/Renderer/Private/SkyAtmosphereRendering.cpp
+++ b/Engine/Source/Runtime/Renderer/Private/SkyAtmosphereRendering.cpp
@@ -194,6 +194,11 @@ static TAutoConsoleVariable<float> CVarSkyAtmosphereDistantSkyLightLUTAltitude(
 
 ////////////////////////////////////////////////////////////////////////// Debug / Visualisation
 
+static TAutoConsoleVariable<int32> CVarSkyAtmosphereLUT32(
+	TEXT("r.SkyAtmosphere.LUT32"), 0,
+	TEXT("Use full 32bit per-channel precision for all sky LUTs.\n"),
+	ECVF_RenderThreadSafe | ECVF_Scalability);
+
 TAutoConsoleVariable<int32> CVarSkyAtmosphereVisualize(
 	TEXT("r.SkyAtmosphere.Visualize"),
 	0,
@@ -895,10 +900,20 @@ static EPixelFormat GetSkyLutTextureFormat(ERHIFeatureLevel::Type FeatureLevel)
 		// TODO: check if need this for Metal, Vulkan
 		TextureLUTFormat = PF_FloatRGBA;
 	}
+	
+	if (CVarSkyAtmosphereLUT32.GetValueOnAnyThread() != 0)
+	{
+		TextureLUTFormat = PF_A32B32G32R32F;
+	}
+
 	return TextureLUTFormat;
 }
 static EPixelFormat GetSkyLutSmallTextureFormat()
 {
+	if (CVarSkyAtmosphereLUT32.GetValueOnAnyThread() != 0)
+	{
+		return PF_A32B32G32R32F;
+	}
 	return PF_R8G8B8A8;
 }
 
@@ -956,6 +971,7 @@ void InitSkyAtmosphereForView(FRHICommandListImmediate& RHICmdList, const FScene
 
 		EPixelFormat TextureLUTFormat = GetSkyLutTextureFormat(Scene->GetFeatureLevel());
 		EPixelFormat TextureLUTSmallFormat = GetSkyLutSmallTextureFormat();
+		EPixelFormat TextureAerialLUTFormat = (CVarSkyAtmosphereLUT32.GetValueOnAnyThread() != 0) ? PF_A32B32G32R32F : PF_FloatRGBA;
 
 		//
 		// Initialise transient per view resources.
@@ -967,7 +983,7 @@ void InitSkyAtmosphereForView(FRHICommandListImmediate& RHICmdList, const FScene
 
 		FPooledRenderTargetDesc SkyAtmosphereCameraAerialPerspectiveVolumeDesc = FPooledRenderTargetDesc::CreateVolumeDesc(
 			CameraAerialPerspectiveVolumeScreenResolution, CameraAerialPerspectiveVolumeScreenResolution, CameraAerialPerspectiveVolumeDepthResolution,
-			PF_FloatRGBA, FClearValueBinding::None, TexCreate_HideInVisualizeTexture, TexCreate_ShaderResource | TexCreate_UAV, false);
+			TextureAerialLUTFormat, FClearValueBinding::None, TexCreate_HideInVisualizeTexture, TexCreate_ShaderResource | TexCreate_UAV, false);
 
 		// Set textures and data that will be needed later on the view.
 		View.SkyAtmosphereUniformShaderParameters = SkyInfo.GetAtmosphereShaderParameters();
