Updated/Clarified PPLL list element and render node count variable.
Nomber of PPLL list element and render node are now decoupled.

diff --git a/Engine/Shaders/Private/HairStrands/HairStrandsVisibilityPPLLDebug.usf b/Engine/Shaders/Private/HairStrands/HairStrandsVisibilityPPLLDebug.usf
--- a/Engine/Shaders/Private/HairStrands/HairStrandsVisibilityPPLLDebug.usf
+++ b/Engine/Shaders/Private/HairStrands/HairStrandsVisibilityPPLLDebug.usf
@@ -4,7 +4,9 @@
 
 #include "HairStrandsVisibilityCommon.ush"
 
-uint MaxPPLLNodeCount;
+float PPLLMeanListElementCountPerPixel;
+float PPLLMaxTotalListElementCount;
+
 Texture2D<uint> PPLLCounter;
 Texture2D<uint> PPLLNodeIndex;
 StructuredBuffer<FPPLLNodeData> PPLLNodeData;
@@ -26,7 +28,7 @@ void VisibilityDebugPPLLPS(
 			NextNodeIndex = Node.NextNodeIndex;
 		} while (NextNodeIndex != 0xFFFFFFFF);
 
-		float l = saturate(StrandCount / 16.0f);
+		float l = saturate(StrandCount / PPLLMeanListElementCountPerPixel);
 
 		OutLuminance = float4(l, 1.0f - l, 0.0f, 0.0f);
 	}
@@ -34,4 +36,27 @@ void VisibilityDebugPPLLPS(
 	{
 		OutLuminance = float4(0.0f, 0.0f, 0.0f, 1.0f);
 	}
+
+	const float2 Origin = 100.0f;
+	const float2 Size = float2(200.0f, 20.0f);
+	if (all(PixPos >= Origin) && all(PixPos <= (Origin + Size)))
+	{
+		float2 LocalPos = PixPos - Origin;
+		float PPLLListElementcount = PPLLCounter[uint2(0, 0)];
+		float MemoryUsage = PPLLListElementcount / PPLLMaxTotalListElementCount;
+
+		if (MemoryUsage > LocalPos.x / Size.x)
+		{
+			OutLuminance = float4(MemoryUsage, 1.0f - MemoryUsage, 1.0f - MemoryUsage, 0.0f);
+		}
+		else
+		{
+			OutLuminance = float4(0.2f, 0.2f, 0.2f, 0.0f);
+		}
+
+		if (any(PixPos == Origin) || any(PixPos == (Origin + Size)))
+		{
+			OutLuminance = float4(1.0f, 1.0f, 1.0f, 0.0f);
+		}
+	}
 }
diff --git a/Engine/Source/Runtime/Renderer/Private/HairStrands/HairStrandsDebug.cpp b/Engine/Source/Runtime/Renderer/Private/HairStrands/HairStrandsDebug.cpp
--- a/Engine/Source/Runtime/Renderer/Private/HairStrands/HairStrandsDebug.cpp
+++ b/Engine/Source/Runtime/Renderer/Private/HairStrands/HairStrandsDebug.cpp
@@ -1028,7 +1028,8 @@ class FHairVisibilityDebugPPLLPS : public FGlobalShader
 
 	BEGIN_SHADER_PARAMETER_STRUCT(FParameters, )
 		SHADER_PARAMETER_STRUCT_REF(FViewUniformShaderParameters, ViewUniformBuffer)
-		SHADER_PARAMETER(uint32, MaxPPLLNodeCount)
+		SHADER_PARAMETER(float, PPLLMeanListElementCountPerPixel)
+		SHADER_PARAMETER(float, PPLLMaxTotalListElementCount)
 		SHADER_PARAMETER_RDG_TEXTURE(Texture2D, PPLLCounter)
 		SHADER_PARAMETER_RDG_TEXTURE(Texture2D, PPLLNodeIndex)
 		SHADER_PARAMETER_RDG_BUFFER_SRV(StructuredBuffer, PPLLNodeData)
@@ -1406,7 +1407,9 @@ void RenderHairStrandsDebugInfo(FRHICommandListImmediate& RHICmdList, TArray<FVi
 			TShaderMapRef<FHairVisibilityDebugPPLLPS> PixelShader(View.ShaderMap, PermutationVector);
 
 			FHairVisibilityDebugPPLLPS::FParameters* PassParameters = GraphBuilder.AllocParameters<FHairVisibilityDebugPPLLPS::FParameters>();
-			PassParameters->MaxPPLLNodeCount = GetMaxNodePerPixel(VisibilityData.PPLLNodeCounterTexture->GetDesc().Extent);
+
+			PassParameters->PPLLMeanListElementCountPerPixel = float(GetPPLLMeanListElementCountPerPixel());
+			PassParameters->PPLLMaxTotalListElementCount = float(GetPPLLMaxTotalListElementCount(VisibilityData.PPLLNodeIndexTexture->GetDesc().Extent));
 			PassParameters->PPLLCounter = PPLLNodeCounterTexture;
 			PassParameters->PPLLNodeIndex = PPLLNodeIndexTexture;
 			PassParameters->PPLLNodeData = GraphBuilder.CreateSRV(FRDGBufferSRVDesc(PPLLNodeDataBuffer));
diff --git a/Engine/Source/Runtime/Renderer/Private/HairStrands/HairStrandsVisibility.cpp b/Engine/Source/Runtime/Renderer/Private/HairStrands/HairStrandsVisibility.cpp
--- a/Engine/Source/Runtime/Renderer/Private/HairStrands/HairStrandsVisibility.cpp
+++ b/Engine/Source/Runtime/Renderer/Private/HairStrands/HairStrandsVisibility.cpp
@@ -46,8 +46,11 @@ static FAutoConsoleVariableRef CVarHairVelocityType(TEXT("r.HairStrands.Velocity
 
 static int32 GHairVisibilityPPLL = 0;
 static FAutoConsoleVariableRef CVarGHairVisibilityPPLL(TEXT("r.HairStrands.VisibilityPPLL"), GHairVisibilityPPLL, TEXT("Hair Visibility uses per pixel linked list"));
-static int32 GHairVisibilityPPLLMaxPixelNodeCount = 16;
-static FAutoConsoleVariableRef CVarGHairVisibilityPPLLMaxPixelNodeCount(TEXT("r.HairStrands.VisibilityPPLLMaxPixelNodeCount"), GHairVisibilityPPLLMaxPixelNodeCount, TEXT("The total maximum number of node allowed for all linked list will be width*height*VisibilityPPLLMaxPixelNodeCount."));
+static int32 GHairVisibilityPPLLMeanListElementCountPerPixel = 16;
+static FAutoConsoleVariableRef CVarGHairVisibilityPPLLMeanListElementCountPerPixel(TEXT("r.HairStrands.VisibilityPPLLMeanListElementCountPerPixel"), GHairVisibilityPPLLMeanListElementCountPerPixel, TEXT("The mean maximum number of node allowed for all linked list element. It will be width*height*VisibilityPPLLMeanListElementCountPerPixel."));
+static int32 GHairVisibilityPPLLMaxRenderNodePerPixel = 16;
+static FAutoConsoleVariableRef CVarGHairVisibilityPPLLMeanNodeCountPerPixel(TEXT("r.HairStrands.VisibilityPPLLMaxRenderNodePerPixel"), GHairVisibilityPPLLMaxRenderNodePerPixel, TEXT("The maximum number of node allowed to be independently shaded and composited per pixel. Total amount of node will be width*height*VisibilityPPLLMaxRenderNodePerPixel. The last node is used to aggregate all furthest strands to shade into a single one."));
+
 
 /////////////////////////////////////////////////////////////////////////////////////////
 
@@ -64,27 +67,40 @@ EHairVisibilityRenderMode GetHairVisibilityRenderMode()
 	return GHairVisibilityPPLL >0 ? HairVisibilityRenderMode_PPLL : HairVisibilityRenderMode_MSAA;
 }
 
-uint32 GetPPLLMaxElementCount()
+uint32 GetPPLLMeanListElementCountPerPixel()
+{
+	return GHairVisibilityPPLLMeanListElementCountPerPixel;
+}
+uint32 GetPPLLMaxTotalListElementCount(FIntPoint Resolution)
+{
+	return Resolution.X * Resolution.Y * GetPPLLMeanListElementCountPerPixel();
+}
+
+uint32 GetPPLLMaxRenderNodePerPixel()
 {
 	// The following must match the FPPLL permutation of FHairVisibilityPrimitiveIdCompactionCS.
-	if (GHairVisibilityPPLLMaxPixelNodeCount == 0)
+	if (GHairVisibilityPPLLMaxRenderNodePerPixel == 0)
 	{
 		return 0;
 	}
-	else if (GHairVisibilityPPLLMaxPixelNodeCount <= 8)
+	else if (GHairVisibilityPPLLMaxRenderNodePerPixel <= 8)
 	{
 		return 8;
 	}
-	else if (GHairVisibilityPPLLMaxPixelNodeCount <= 16)
+	else if (GHairVisibilityPPLLMaxRenderNodePerPixel <= 16)
 	{
 		return 16;
 	}
-	else //if (GHairVisibilityPPLLMaxPixelNodeCount <= 32)
+	else //if (GHairVisibilityPPLLMaxRenderNodePerPixel <= 32)
 	{
 		return 32;
 	}
 	// If more is needed: please check out EncodeNodeDesc from HairStrandsVisibilityCommon.ush to verify node count representation limitations.
 }
+uint32 GetPPLLMaxTotalRenderNode(FIntPoint Resolution)
+{
+	return Resolution.X * Resolution.Y * GetPPLLMaxRenderNodePerPixel();
+}
 
 uint32 GetHairVisibilitySampleCount()
 {
@@ -110,11 +126,6 @@ void SetUpViewHairRenderInfo(const FViewInfo& ViewInfo, bool bEnableMSAA, FVecto
 	OutHairRenderInfo.W = VelocityMagnitudeScale;
 }
 
-uint32 GetMaxNodePerPixel(FIntPoint Resolution)
-{
-	return Resolution.X * Resolution.Y * GHairVisibilityPPLLMaxPixelNodeCount;
-}
-
 static bool IsCompatibleWithHairVisibility(const FMeshMaterialShaderPermutationParameters& Parameters)
 {
 	return IsCompatibleWithHairStrands(Parameters.Material, Parameters.Platform);
@@ -517,7 +528,7 @@ class FHairVisibilityPrimitiveIdCompactionCS : public FGlobalShader
 	class FVelocity : SHADER_PERMUTATION_INT("PERMUTATION_VELOCITY", 4);
 	class FCoverage : SHADER_PERMUTATION_INT("PERMUTATION_COVERAGE", 2);
 	class FMaterial : SHADER_PERMUTATION_INT("PERMUTATION_MATERIAL_COMPACTION", 2);
-	class FPPLL     : SHADER_PERMUTATION_SPARSE_INT("PERMUTATION_PPLL", 0, 8, 16, 32); // See GetPPLLMaxElementCount
+	class FPPLL     : SHADER_PERMUTATION_SPARSE_INT("PERMUTATION_PPLL", 0, 8, 16, 32); // See GetPPLLMaxRenderNodePerPixel
 	using FPermutationDomain = TShaderPermutationDomain<FVendor, FVelocity, FCoverage, FMaterial, FPPLL>;
 
 	BEGIN_SHADER_PARAMETER_STRUCT(FParameters, )
@@ -631,10 +642,11 @@ static void AddHairVisibilityPrimitiveIdCompactionPass(
 	AddClearUAVPass(GraphBuilder, RDG_EVENT_NAME("HairStrandsClearCompactionOffsetAndCount"), OutCompactNodeIndex, 0);
 	AddClearUAVPass(GraphBuilder, RDG_EVENT_NAME("HairStrandsClearCategorizationTexture"), OutCategorizationTexture, 0);
 
+	// Select render node count according to current mode
 	const uint32 HairVisibilityMSAASampleCount = GetHairVisibilitySampleCount();
 	const uint32 SampleCount = FMath::RoundUpToPowerOfTwo(HairVisibilityMSAASampleCount);
-	const uint32 PPLLMaxElementCount = GetPPLLMaxElementCount();
-	const uint32 MaxNodeCount = Resolution.X * Resolution.Y * (GetHairVisibilityRenderMode() == HairVisibilityRenderMode_MSAA ? SampleCount : PPLLMaxElementCount);
+	const uint32 PPLLMaxRenderNodePerPixel = GetPPLLMaxRenderNodePerPixel();
+	const uint32 MaxRenderNodeCount = Resolution.X * Resolution.Y * (GetHairVisibilityRenderMode() == HairVisibilityRenderMode_MSAA ? SampleCount : PPLLMaxRenderNodePerPixel);
 	{
 		struct NodeData
 		{
@@ -645,12 +657,12 @@ static void AddHairVisibilityPrimitiveIdCompactionPass(
 			uint32 Specular;
 		};
 
-		OutCompactNodeData = GraphBuilder.CreateBuffer(FRDGBufferDesc::CreateStructuredDesc(sizeof(NodeData), MaxNodeCount), TEXT("HairVisibilityPrimitiveIdCompactNodeData"));
+		OutCompactNodeData = GraphBuilder.CreateBuffer(FRDGBufferDesc::CreateStructuredDesc(sizeof(NodeData), MaxRenderNodeCount), TEXT("HairVisibilityPrimitiveIdCompactNodeData"));
 	}
 
 	{
 		// Pixel coord of the node. Stored as 2*uint16, packed into a single uint32
-		OutCompactNodeCoord = GraphBuilder.CreateBuffer(FRDGBufferDesc::CreateStructuredDesc(sizeof(uint32), MaxNodeCount), TEXT("HairVisibilityPrimitiveIdCompactNodeCoord"));
+		OutCompactNodeCoord = GraphBuilder.CreateBuffer(FRDGBufferDesc::CreateStructuredDesc(sizeof(uint32), MaxRenderNodeCount), TEXT("HairVisibilityPrimitiveIdCompactNodeCoord"));
 	}
 
 	FSceneRenderTargets& SceneContext = FSceneRenderTargets::Get(GraphBuilder.RHICmdList);
@@ -663,10 +675,10 @@ static void AddHairVisibilityPrimitiveIdCompactionPass(
 	PermutationVector.Set<FHairVisibilityPrimitiveIdCompactionCS::FVelocity>(bWriteOutVelocity ? FMath::Clamp(GHairVelocityType+1, 0, 3) : 0);
 	PermutationVector.Set<FHairVisibilityPrimitiveIdCompactionCS::FCoverage>(PassParameters->CoverageTexture ? 1 : 0);
 	PermutationVector.Set<FHairVisibilityPrimitiveIdCompactionCS::FMaterial>(GHairStrandsMaterialCompactionEnable ? 1 : 0);
-	PermutationVector.Set<FHairVisibilityPrimitiveIdCompactionCS::FPPLL>(bUsePPLL ? PPLLMaxElementCount : 0);
+	PermutationVector.Set<FHairVisibilityPrimitiveIdCompactionCS::FPPLL>(bUsePPLL ? PPLLMaxRenderNodePerPixel : 0);
 
 	PassParameters->OutputResolution = Resolution;
-	PassParameters->MaxNodeCount = MaxNodeCount;
+	PassParameters->MaxNodeCount = MaxRenderNodeCount;
 	PassParameters->DepthTheshold = FMath::Clamp(GHairStrandsMaterialCompactionDepthThreshold, 0.f, 100.f);
 	PassParameters->CosTangentThreshold = FMath::Cos(FMath::DegreesToRadians(FMath::Clamp(GHairStrandsMaterialCompactionTangentThreshold, 0.f, 90.f)));
 	PassParameters->HairVisibilityMSAASampleCount = HairVisibilityMSAASampleCount;
@@ -1042,7 +1054,7 @@ static void AddHairVisibilityPPLLPass(
 		OutVisibilityPPLLNodeIndex = GraphBuilder.CreateTexture(Desc, TEXT("HairVisibilityPPLLNodeIndex"));
 	}
 
-	const uint32 MaxNodeCount = GetMaxNodePerPixel(Resolution);
+	const uint32 PPLLMaxTotalListElementCount = GetPPLLMaxTotalListElementCount(Resolution);
 	{
 		// Example: 28bytes * 8spp = 224bytes per pixel = 442Mb @ 1080p
 		struct PPLLNodeData
@@ -1056,7 +1068,7 @@ static void AddHairVisibilityPPLLPass(
 			uint32 PackedVelocity;
 		};
 
-		OutVisibilityPPLLNodeData = GraphBuilder.CreateBuffer(FRDGBufferDesc::CreateStructuredDesc(sizeof(PPLLNodeData), MaxNodeCount), TEXT("HairVisibilityPPLLNodeData"));
+		OutVisibilityPPLLNodeData = GraphBuilder.CreateBuffer(FRDGBufferDesc::CreateStructuredDesc(sizeof(PPLLNodeData), PPLLMaxTotalListElementCount), TEXT("HairVisibilityPPLLNodeData"));
 	}
 	AddClearUAVPass(GraphBuilder, RDG_EVENT_NAME("ClearHairVisibilityPPLLCounter"), OutVisibilityPPLLNodeCounter, 0);
 	AddClearUAVPass(GraphBuilder, RDG_EVENT_NAME("ClearHairVisibilityPPLLNodeIndex"), OutVisibilityPPLLNodeIndex, 0xFFFFFFFF);
@@ -1065,7 +1077,7 @@ static void AddHairVisibilityPPLLPass(
 	PassParameters->PPLLCounter = GraphBuilder.CreateUAV(FRDGTextureUAVDesc(OutVisibilityPPLLNodeCounter, 0));
 	PassParameters->PPLLNodeIndex = GraphBuilder.CreateUAV(FRDGTextureUAVDesc(OutVisibilityPPLLNodeIndex, 0));
 	PassParameters->PPLLNodeData = GraphBuilder.CreateUAV(FRDGBufferUAVDesc(OutVisibilityPPLLNodeData));
-	PassParameters->HairVisibilityPass_MaxPPLLNodeCount = MaxNodeCount;
+	PassParameters->HairVisibilityPass_MaxPPLLNodeCount = PPLLMaxTotalListElementCount;
 	PassParameters->RenderTargets.DepthStencil = FDepthStencilBinding(InViewZDepthTexture, ERenderTargetLoadAction::ELoad, ERenderTargetLoadAction::ENoAction, FExclusiveDepthStencil::DepthRead_StencilNop);
 	AddHairVisibilityCommonPass(GraphBuilder, Scene, ViewInfo, ClusterDatas, HairVisibilityRenderMode_PPLL, PassParameters);
 }
diff --git a/Engine/Source/Runtime/Renderer/Private/HairStrands/HairStrandsVisibility.h b/Engine/Source/Runtime/Renderer/Private/HairStrands/HairStrandsVisibility.h
--- a/Engine/Source/Runtime/Renderer/Private/HairStrands/HairStrandsVisibility.h
+++ b/Engine/Source/Runtime/Renderer/Private/HairStrands/HairStrandsVisibility.h
@@ -53,4 +53,7 @@ FHairStrandsVisibilityViews RenderHairStrandsVisibilityBuffer(
 
 void SetUpViewHairRenderInfo(const FViewInfo& ViewInfo, bool bEnableMSAA, FVector4& OutHairRenderInfo);
 
-uint32 GetMaxNodePerPixel(FIntPoint Resolution);
+
+uint32 GetPPLLMeanListElementCountPerPixel();
+uint32 GetPPLLMaxTotalListElementCount(FIntPoint Resolution);
+
