Fix usage of non-UAV compatible formats that were breaking Windows7 build PCs

diff --git a/Engine/Shaders/Private/VirtualTextureCompress.usf b/Engine/Shaders/Private/VirtualTextureCompress.usf
--- a/Engine/Shaders/Private/VirtualTextureCompress.usf
+++ b/Engine/Shaders/Private/VirtualTextureCompress.usf
@@ -20,11 +20,17 @@ RWTexture2D<uint2> OutCompressTexture0_u2;
 RWTexture2D<uint4> OutCompressTexture0_u4;
 RWTexture2D<uint4> OutCompressTexture1;
 RWTexture2D<uint2> OutCompressTexture2;
-RWTexture2D<float4> OutCopyTexture0;
-RWTexture2D<float4> OutCopyTexture1;
-RWTexture2D<float4> OutCopyTexture2;
 
-// Read a 4x4 color block ready for BC1 compression
+
+/** Trivial vertex shader for copy shaders. */
+void CopyVS(in uint VertexId : SV_VertexID, out float4 OutPosition : SV_POSITION, out float2 OutTexCoord : TEXCOORD0)
+{
+	OutTexCoord = float2(((VertexId + 1) / 3) & 1, VertexId & 1);
+	OutPosition.xy = float2(OutTexCoord.x, 1.f - OutTexCoord.y) * 2.f - 1.f;
+	OutPosition.zw = float2(0, 1);
+}
+
+/** Read a 4x4 color block ready for BC1 compression. */
 void ReadBlockRGB( Texture2D<float4> SourceTexture, SamplerState TextureSampler, float2 UV, float2 TexelUVSize, out float3 Block[16] )
 {
 	{
@@ -68,7 +74,7 @@ void ReadBlockRGB( Texture2D<float4> SourceTexture, SamplerState TextureSampler,
 	}
 }
 
-// Read a 4x4 block of XY channels from a normal texture ready for BC5 compression
+/** Read a 4x4 block of XY channels from a normal texture ready for BC5 compression. */
 void ReadBlockRG( Texture2D<float4> SourceTexture, SamplerState TextureSampler, float2 UV, float2 TexelUVSize, out float BlockR[16], out float BlockG[16] )
 {
 	{
@@ -127,6 +133,16 @@ void CompressBaseColorCS(
 	OutCompressTexture0_u2[ThreadId.xy + DestRect.xy] = CompressBC1Block_SRGB(BlockBaseColor);
 }
 
+/** Simple copy used for thumbnails. */
+void CopyBaseColorPS(
+	in float4 InPosition : SV_POSITION, 
+	in noperspective float2 InTexCoord : TEXCOORD0, 
+	out float4 OutColor0 : SV_Target0)
+{
+	float3 BaseColor = RenderTexture0.SampleLevel(TextureSampler0, InTexCoord, 0).xyz;
+	OutColor0 = float4(BaseColor, 1.f);
+}
+
 /** Compress base color, normal, roughness and specular to BC3 target pair. */
 [numthreads(8,8,1)]
 void CompressBaseColorNormalSpecularCS(
@@ -161,23 +177,20 @@ void CompressBaseColorNormalSpecularCS(
 
 /** Copy normal X and Y to the alpha channels of the final target textures (without compression)
   * This path is used when we disable texture compression, because we need to keep the same final channel layout. */
-[numthreads(8,8,1)]
-void CopyNormalSpecularCS(
-	uint3 ThreadId : SV_DispatchThreadID 
-	)
+void CopyBaseColorNormalSpecularPS(
+	in float4 InPosition : SV_POSITION, 
+	in noperspective float2 InTexCoord : TEXCOORD0, 
+	out float4 OutColor0 : SV_Target0,
+	out float4 OutColor1 : SV_Target1)
 {
-	uint2 SamplePos = ThreadId.xy;
-	if (any(ThreadId.xy >= DestRect.zw))
-		return;
-
-	float3 BaseColor = LinearToSrgb(RenderTexture0[SamplePos].xyz);
-	float2 NormalXY = RenderTexture1[SamplePos].xy;
-	float3 RoughnessSpecularNormalZ = RenderTexture2[SamplePos].xyz;
+	float3 BaseColor = RenderTexture0.SampleLevel(TextureSampler0, InTexCoord, 0).xyz;
+	float2 NormalXY = RenderTexture1.SampleLevel(TextureSampler1, InTexCoord, 0).xy;
+	float3 RoughnessSpecularNormalZ = RenderTexture2.SampleLevel(TextureSampler2, InTexCoord, 0).xyz;
 	
 	RoughnessSpecularNormalZ.z = round(RoughnessSpecularNormalZ.z);
 	
-	OutCopyTexture0[ThreadId.xy + DestRect.xy] = float4(BaseColor, NormalXY.x);
-	OutCopyTexture1[ThreadId.xy + DestRect.xy] = float4(RoughnessSpecularNormalZ, NormalXY.y);
+	OutColor0 = float4(BaseColor, NormalXY.x);
+	OutColor1 = float4(RoughnessSpecularNormalZ, NormalXY.y);
 }
 
 /** Compress base color, normal, roughness and specular to BC3, BC5, BC1 target set. */
@@ -215,37 +228,31 @@ void CompressBaseColorNormalSpecularYCoCgCS(
 
 /** Copy normal X and Y to the alpha channels of the final target textures (without compression)
   * This path is used when we disable texture compression, because we need to keep the same final channel layout. */
-[numthreads(8,8,1)]
-void CopyNormalSpecularYCoCgCS(
-	uint3 ThreadId : SV_DispatchThreadID 
-	)
+void CopyBaseColorNormalSpecularYCoCgPS(
+	in float4 InPosition : SV_POSITION, 
+	in noperspective float2 InTexCoord : TEXCOORD0, 
+	out float4 OutColor0 : SV_Target0,
+	out float4 OutColor1 : SV_Target1,
+	out float4 OutColor2 : SV_Target2)
 {
-	uint2 SamplePos = ThreadId.xy;
-	if (any(ThreadId.xy >= DestRect.zw))
-		return;
-
-	float3 YCoCg = RGB2YCoCg(RenderTexture0[SamplePos].xyz);
-	float2 NormalXY = RenderTexture1[SamplePos].xy;
-	float3 RoughnessSpecularNormalZ = RenderTexture2[SamplePos].xyz;
+	float3 YCoCg = RGB2YCoCg(RenderTexture0.SampleLevel(TextureSampler0, InTexCoord, 0).xyz);
+	float2 NormalXY = RenderTexture1.SampleLevel(TextureSampler1, InTexCoord, 0).xy;
+	float3 RoughnessSpecularNormalZ = RenderTexture2.SampleLevel(TextureSampler2, InTexCoord, 0).xyz;
 	
 	RoughnessSpecularNormalZ.z = round(RoughnessSpecularNormalZ.z);
 	
-	OutCopyTexture0[ThreadId.xy + DestRect.xy] = float4(YCoCg.yz, 0, YCoCg.x);
-	OutCopyTexture1[ThreadId.xy + DestRect.xy] = float4(NormalXY.xy, 0, 1);
-	OutCopyTexture2[ThreadId.xy + DestRect.xy] = float4(RoughnessSpecularNormalZ, 1);
+	OutColor0 = float4(YCoCg.yz, 0, YCoCg.x);
+	OutColor1 = float4(NormalXY.xy, 0, 1);
+	OutColor2 = float4(RoughnessSpecularNormalZ, 1);
 }
 
 /** Copy world height (without compression)
   * This path is used when want to copy the world height data into a thumbnail. */
-[numthreads(8,8,1)]
-void CopyWorldHeightCS(
-	uint3 ThreadId : SV_DispatchThreadID 
-	)
+void CopyWorldHeightPS(
+	in float4 InPosition : SV_POSITION, 
+	in noperspective float2 InTexCoord : TEXCOORD0, 
+	out float4 OutColor0 : SV_Target0)
 {
-	uint2 SamplePos = ThreadId.xy;
-	if (any(ThreadId.xy >= DestRect.zw))
-		return;
-
-	float WorldHeight = RenderTexture0[SamplePos].r;
-	OutCopyTexture0[ThreadId.xy + DestRect.xy] = float4(WorldHeight.xxx, 1.f);
+	float WorldHeight = RenderTexture0.SampleLevel(TextureSampler0, InTexCoord, 0).r;
+	OutColor0 = float4(WorldHeight.xxx, 1.f);
 }
diff --git a/Engine/Source/Editor/VirtualTexturingEditor/Private/RuntimeVirtualTextureBuild.cpp b/Engine/Source/Editor/VirtualTexturingEditor/Private/RuntimeVirtualTextureBuild.cpp
--- a/Engine/Source/Editor/VirtualTexturingEditor/Private/RuntimeVirtualTextureBuild.cpp
+++ b/Engine/Source/Editor/VirtualTexturingEditor/Private/RuntimeVirtualTextureBuild.cpp
@@ -200,6 +200,7 @@ namespace RuntimeVirtualTexture
 					Desc.MaterialType = MaterialType;
 					Desc.MaxLevel = MaxLevel;
 					Desc.bClearTextures = true;
+					Desc.bIsThumbnails = false;
 					Desc.DebugType = DebugType;
 					Desc.NumPageDescs = 1;
 					Desc.Targets[0].Texture = RenderTileResources.GetRenderTarget(0);
diff --git a/Engine/Source/Editor/VirtualTexturingEditor/Private/RuntimeVirtualTextureThumbnailRenderer.cpp b/Engine/Source/Editor/VirtualTexturingEditor/Private/RuntimeVirtualTextureThumbnailRenderer.cpp
--- a/Engine/Source/Editor/VirtualTexturingEditor/Private/RuntimeVirtualTextureThumbnailRenderer.cpp
+++ b/Engine/Source/Editor/VirtualTexturingEditor/Private/RuntimeVirtualTextureThumbnailRenderer.cpp
@@ -82,6 +82,7 @@ void URuntimeVirtualTextureThumbnailRenderer::Draw(UObject* Object, int32 X, int
 		Desc.MaterialType = MaterialType;
 		Desc.MaxLevel = MaxLevel;
 		Desc.bClearTextures = true;
+		Desc.bIsThumbnails = true;
 		Desc.DebugType = ERuntimeVirtualTextureDebugType::None;
 		Desc.NumPageDescs = 1;
 		Desc.Targets[0].Texture = RenderTarget->GetRenderTargetTexture();
diff --git a/Engine/Source/Runtime/Renderer/Private/VT/RuntimeVirtualTextureProducer.cpp b/Engine/Source/Runtime/Renderer/Private/VT/RuntimeVirtualTextureProducer.cpp
--- a/Engine/Source/Runtime/Renderer/Private/VT/RuntimeVirtualTextureProducer.cpp
+++ b/Engine/Source/Runtime/Renderer/Private/VT/RuntimeVirtualTextureProducer.cpp
@@ -62,6 +62,7 @@ void FRuntimeVirtualTextureFinalizer::Finalize(FRHICommandListImmediate& RHICmdL
 	RenderPageBatchDesc.MaterialType = MaterialType;
 	RenderPageBatchDesc.MaxLevel = Desc.MaxLevel;
 	RenderPageBatchDesc.bClearTextures = bClearTextures;
+	RenderPageBatchDesc.bIsThumbnails = false;
 	RenderPageBatchDesc.DebugType = ERuntimeVirtualTextureDebugType::None;
 	
 	for (int LayerIndex = 0; LayerIndex < RuntimeVirtualTexture::MaxTextureLayers; ++LayerIndex)
diff --git a/Engine/Source/Runtime/Renderer/Private/VT/RuntimeVirtualTextureRender.cpp b/Engine/Source/Runtime/Renderer/Private/VT/RuntimeVirtualTextureRender.cpp
--- a/Engine/Source/Runtime/Renderer/Private/VT/RuntimeVirtualTextureRender.cpp
+++ b/Engine/Source/Runtime/Renderer/Private/VT/RuntimeVirtualTextureRender.cpp
@@ -421,9 +421,6 @@ namespace RuntimeVirtualTexture
 			SHADER_PARAMETER_RDG_TEXTURE_UAV(RWTexture<uint4>, OutCompressTexture0_u4)
 			SHADER_PARAMETER_RDG_TEXTURE_UAV(RWTexture<uint4>, OutCompressTexture1)
 			SHADER_PARAMETER_RDG_TEXTURE_UAV(RWTexture<uint2>, OutCompressTexture2)
-			SHADER_PARAMETER_RDG_TEXTURE_UAV(RWTexture<float4>, OutCopyTexture0)
-			SHADER_PARAMETER_RDG_TEXTURE_UAV(RWTexture<float4>, OutCopyTexture1)
-			SHADER_PARAMETER_RDG_TEXTURE_UAV(RWTexture<float4>, OutCopyTexture2)
 		END_SHADER_PARAMETER_STRUCT()
 
 		static bool ShouldCompilePermutation(FGlobalShaderPermutationParameters const& Parameters)
@@ -441,11 +438,11 @@ namespace RuntimeVirtualTexture
 		}
 	};
 
-	template< ERuntimeVirtualTextureMaterialType MaterialType, bool CopyOnly >
+	template< ERuntimeVirtualTextureMaterialType MaterialType >
 	class FShader_VirtualTextureCompress_CS : public FShader_VirtualTextureCompress
 	{
 	public:
-		typedef FShader_VirtualTextureCompress_CS< MaterialType, CopyOnly > ClassName; // typedef is only so that we can use in DECLARE_SHADER_TYPE macro
+		typedef FShader_VirtualTextureCompress_CS< MaterialType > ClassName; // typedef is only so that we can use in DECLARE_SHADER_TYPE macro
 		DECLARE_SHADER_TYPE( ClassName, Global );
 
 		FShader_VirtualTextureCompress_CS()
@@ -456,31 +453,17 @@ namespace RuntimeVirtualTexture
 		{}
 	};
 
-	typedef FShader_VirtualTextureCompress_CS< ERuntimeVirtualTextureMaterialType::BaseColor, false > FShader_VirtualTextureCompress_BaseColor_CS;
-	IMPLEMENT_SHADER_TYPE(template<>, FShader_VirtualTextureCompress_BaseColor_CS, TEXT("/Engine/Private/VirtualTextureCompress.usf"), TEXT("CompressBaseColorCS"), SF_Compute);
-	
-	typedef FShader_VirtualTextureCompress_CS< ERuntimeVirtualTextureMaterialType::BaseColor_Normal_Specular, false > FShader_VirtualTextureCompress_BaseColorNormalSpecular_CS;
-	IMPLEMENT_SHADER_TYPE(template<>, FShader_VirtualTextureCompress_BaseColorNormalSpecular_CS, TEXT("/Engine/Private/VirtualTextureCompress.usf"), TEXT("CompressBaseColorNormalSpecularCS"), SF_Compute);
+	IMPLEMENT_SHADER_TYPE(template<>, FShader_VirtualTextureCompress_CS< ERuntimeVirtualTextureMaterialType::BaseColor >, TEXT("/Engine/Private/VirtualTextureCompress.usf"), TEXT("CompressBaseColorCS"), SF_Compute);
+	IMPLEMENT_SHADER_TYPE(template<>, FShader_VirtualTextureCompress_CS< ERuntimeVirtualTextureMaterialType::BaseColor_Normal_Specular >, TEXT("/Engine/Private/VirtualTextureCompress.usf"), TEXT("CompressBaseColorNormalSpecularCS"), SF_Compute);
+	IMPLEMENT_SHADER_TYPE(template<>, FShader_VirtualTextureCompress_CS< ERuntimeVirtualTextureMaterialType::BaseColor_Normal_Specular_YCoCg >, TEXT("/Engine/Private/VirtualTextureCompress.usf"), TEXT("CompressBaseColorNormalSpecularYCoCgCS"), SF_Compute);
 
-	typedef FShader_VirtualTextureCompress_CS< ERuntimeVirtualTextureMaterialType::BaseColor_Normal_Specular_YCoCg, false > FShader_VirtualTextureCompress_BaseColorNormalSpecularYCoCg_CS;
-	IMPLEMENT_SHADER_TYPE(template<>, FShader_VirtualTextureCompress_BaseColorNormalSpecularYCoCg_CS, TEXT("/Engine/Private/VirtualTextureCompress.usf"), TEXT("CompressBaseColorNormalSpecularYCoCgCS"), SF_Compute);
 
-	typedef FShader_VirtualTextureCompress_CS< ERuntimeVirtualTextureMaterialType::BaseColor_Normal_Specular, true > FShader_VirtualTextureCompress_BaseColorNormalSpecular_CopyOnly_CS;
-	IMPLEMENT_SHADER_TYPE(template<>, FShader_VirtualTextureCompress_BaseColorNormalSpecular_CopyOnly_CS, TEXT("/Engine/Private/VirtualTextureCompress.usf"), TEXT("CopyNormalSpecularCS"), SF_Compute);
-
-	typedef FShader_VirtualTextureCompress_CS< ERuntimeVirtualTextureMaterialType::BaseColor_Normal_Specular_YCoCg, true > FShader_VirtualTextureCompress_BaseColorNormalSpecularYCoCg_CopyOnly_CS;
-	IMPLEMENT_SHADER_TYPE(template<>, FShader_VirtualTextureCompress_BaseColorNormalSpecularYCoCg_CopyOnly_CS, TEXT("/Engine/Private/VirtualTextureCompress.usf"), TEXT("CopyNormalSpecularYCoCgCS"), SF_Compute);
-
-	typedef FShader_VirtualTextureCompress_CS< ERuntimeVirtualTextureMaterialType::WorldHeight, true > FShader_VirtualTextureCompress_WorldHeight_CopyOnly_CS;
-	IMPLEMENT_SHADER_TYPE(template<>, FShader_VirtualTextureCompress_WorldHeight_CopyOnly_CS, TEXT("/Engine/Private/VirtualTextureCompress.usf"), TEXT("CopyWorldHeightCS"), SF_Compute);
-
-
-	/** Set up the BC compression pass for the specific MaterialType. */
-	template< ERuntimeVirtualTextureMaterialType MaterialType, bool CopyOnly = false >
-	void AddCompressOrCopyPass( FRDGBuilder& GraphBuilder, ERHIFeatureLevel::Type FeatureLevel, FShader_VirtualTextureCompress::FParameters* Parameters, FIntVector GroupCount )
+	/** Add the BC compression pass to the graph. */
+	template< ERuntimeVirtualTextureMaterialType MaterialType >
+	void AddCompressPass(FRDGBuilder& GraphBuilder, ERHIFeatureLevel::Type FeatureLevel, FShader_VirtualTextureCompress::FParameters* Parameters, FIntVector GroupCount)
 	{
 		TShaderMap<FGlobalShaderType>* GlobalShaderMap = GetGlobalShaderMap(FeatureLevel);
-		TShaderMapRef< FShader_VirtualTextureCompress_CS< MaterialType, CopyOnly > > ComputeShader(GlobalShaderMap);
+		TShaderMapRef< FShader_VirtualTextureCompress_CS< MaterialType > > ComputeShader(GlobalShaderMap);
 
 		FComputeShaderUtils::AddPass(
 			GraphBuilder,
@@ -488,7 +471,7 @@ namespace RuntimeVirtualTexture
 			*ComputeShader, Parameters, GroupCount);
 	}
 
-	/** Set up the BC compression pass. */
+	/** Set up the BC compression pass for the given MaterialType. */
 	void AddCompressPass(FRDGBuilder& GraphBuilder, ERHIFeatureLevel::Type FeatureLevel, FShader_VirtualTextureCompress::FParameters* Parameters, FIntPoint TextureSize, ERuntimeVirtualTextureMaterialType MaterialType)
 	{
 		const FIntVector GroupCount(((TextureSize.X / 4) + 7) / 8, ((TextureSize.Y / 4) + 7) / 8, 1);
@@ -497,34 +480,143 @@ namespace RuntimeVirtualTexture
 		switch (MaterialType)
 		{
 		case ERuntimeVirtualTextureMaterialType::BaseColor:
-			AddCompressOrCopyPass<ERuntimeVirtualTextureMaterialType::BaseColor>(GraphBuilder, FeatureLevel, Parameters, GroupCount);
+			AddCompressPass<ERuntimeVirtualTextureMaterialType::BaseColor>(GraphBuilder, FeatureLevel, Parameters, GroupCount);
 			break;
 		case ERuntimeVirtualTextureMaterialType::BaseColor_Normal_Specular:
-			AddCompressOrCopyPass<ERuntimeVirtualTextureMaterialType::BaseColor_Normal_Specular>(GraphBuilder, FeatureLevel, Parameters, GroupCount);
+			AddCompressPass<ERuntimeVirtualTextureMaterialType::BaseColor_Normal_Specular>(GraphBuilder, FeatureLevel, Parameters, GroupCount);
 			break;
 		case ERuntimeVirtualTextureMaterialType::BaseColor_Normal_Specular_YCoCg:
-			AddCompressOrCopyPass<ERuntimeVirtualTextureMaterialType::BaseColor_Normal_Specular_YCoCg>(GraphBuilder, FeatureLevel, Parameters, GroupCount);
+			AddCompressPass<ERuntimeVirtualTextureMaterialType::BaseColor_Normal_Specular_YCoCg>(GraphBuilder, FeatureLevel, Parameters, GroupCount);
 			break;
 		}
 	}
 
-	/** Set up the copy pass used when BC compression is disabled */
-	void AddCopyPass(FRDGBuilder& GraphBuilder, ERHIFeatureLevel::Type FeatureLevel, FShader_VirtualTextureCompress::FParameters* Parameters, FIntPoint TextureSize, ERuntimeVirtualTextureMaterialType MaterialType)
+
+	/** Copy shaders are used when compression is disabled. These are used to ensure that the channel layout is the same as with compression. */
+	class FShader_VirtualTextureCopy : public FGlobalShader
+	{
+	public:
+		BEGIN_SHADER_PARAMETER_STRUCT(FParameters, )
+			RENDER_TARGET_BINDING_SLOTS()
+			SHADER_PARAMETER(FIntVector4, DestRect)
+			SHADER_PARAMETER_RDG_TEXTURE(Texture2D<float4>, RenderTexture0)
+			SHADER_PARAMETER_SAMPLER(SamplerState, TextureSampler0)
+			SHADER_PARAMETER_RDG_TEXTURE(Texture2D<float4>, RenderTexture1)
+			SHADER_PARAMETER_SAMPLER(SamplerState, TextureSampler1)
+			SHADER_PARAMETER_RDG_TEXTURE(Texture2D<float4>, RenderTexture2)
+			SHADER_PARAMETER_SAMPLER(SamplerState, TextureSampler2)
+		END_SHADER_PARAMETER_STRUCT()
+
+		static bool ShouldCompilePermutation(FGlobalShaderPermutationParameters const& Parameters)
+		{
+			return IsFeatureLevelSupported(Parameters.Platform, ERHIFeatureLevel::SM4);
+		}
+
+		FShader_VirtualTextureCopy()
+		{}
+
+		FShader_VirtualTextureCopy(const ShaderMetaType::CompiledShaderInitializerType& Initializer)
+			: FGlobalShader(Initializer)
+		{
+			Bindings.BindForLegacyShaderParameters(this, Initializer.ParameterMap, *FParameters::FTypeInfo::GetStructMetadata());
+		}
+	};
+
+	class FShader_VirtualTextureCopy_VS : public FShader_VirtualTextureCopy
 	{
-		const FIntVector GroupCount((TextureSize.X + 7) / 8, (TextureSize.Y + 7) / 8, 1);
+	public:
+		DECLARE_SHADER_TYPE(FShader_VirtualTextureCopy_VS, Global);
+
+		FShader_VirtualTextureCopy_VS()
+		{}
 
-		// Needed for BaseColor_Normal_Specular where we pack layers
-		// Needed for WorldHeight thumbnails where we need to copy G16 to the RGBA8 thumbnail
+		FShader_VirtualTextureCopy_VS(const ShaderMetaType::CompiledShaderInitializerType& Initializer)
+			: FShader_VirtualTextureCopy(Initializer)
+		{}
+	};
+
+	IMPLEMENT_SHADER_TYPE(, FShader_VirtualTextureCopy_VS, TEXT("/Engine/Private/VirtualTextureCompress.usf"), TEXT("CopyVS"), SF_Vertex);
+
+	template< ERuntimeVirtualTextureMaterialType MaterialType >
+	class FShader_VirtualTextureCopy_PS : public FShader_VirtualTextureCopy
+	{
+	public:
+		typedef FShader_VirtualTextureCopy_PS< MaterialType > ClassName; // typedef is only so that we can use in DECLARE_SHADER_TYPE macro
+		DECLARE_SHADER_TYPE(ClassName, Global);
+
+		FShader_VirtualTextureCopy_PS()
+		{}
+
+		FShader_VirtualTextureCopy_PS(const ShaderMetaType::CompiledShaderInitializerType& Initializer)
+			: FShader_VirtualTextureCopy(Initializer)
+		{}
+	};
+
+	IMPLEMENT_SHADER_TYPE(template<>, FShader_VirtualTextureCopy_PS< ERuntimeVirtualTextureMaterialType::BaseColor >, TEXT("/Engine/Private/VirtualTextureCompress.usf"), TEXT("CopyBaseColorNormalSpecularPS"), SF_Pixel);
+	IMPLEMENT_SHADER_TYPE(template<>, FShader_VirtualTextureCopy_PS< ERuntimeVirtualTextureMaterialType::BaseColor_Normal_Specular >, TEXT("/Engine/Private/VirtualTextureCompress.usf"), TEXT("CopyBaseColorNormalSpecularPS"), SF_Pixel);
+	IMPLEMENT_SHADER_TYPE(template<>, FShader_VirtualTextureCopy_PS< ERuntimeVirtualTextureMaterialType::BaseColor_Normal_Specular_YCoCg >, TEXT("/Engine/Private/VirtualTextureCompress.usf"), TEXT("CopyBaseColorNormalSpecularYCoCgPS"), SF_Pixel);
+	IMPLEMENT_SHADER_TYPE(template<>, FShader_VirtualTextureCopy_PS< ERuntimeVirtualTextureMaterialType::WorldHeight >, TEXT("/Engine/Private/VirtualTextureCompress.usf"), TEXT("CopyWorldHeightPS"), SF_Pixel);
+
+
+	/** Add the copy pass to the graph. */
+	template< ERuntimeVirtualTextureMaterialType MaterialType >
+	void AddCopyPass(FRDGBuilder& GraphBuilder, ERHIFeatureLevel::Type FeatureLevel, FShader_VirtualTextureCopy::FParameters* Parameters, FIntPoint TextureSize)
+	{
+		TShaderMap<FGlobalShaderType>* GlobalShaderMap = GetGlobalShaderMap(FeatureLevel);
+		TShaderMapRef< FShader_VirtualTextureCopy_VS > VertexShader(GlobalShaderMap);
+		TShaderMapRef< FShader_VirtualTextureCopy_PS< MaterialType > > PixelShader(GlobalShaderMap);
+
+		GraphBuilder.AddPass(
+			RDG_EVENT_NAME("VirtualTextureCopy"),
+			Parameters,
+			ERDGPassFlags::Raster,
+			[VertexShader, PixelShader, Parameters, TextureSize](FRHICommandListImmediate& RHICmdList)
+		{
+			FGraphicsPipelineStateInitializer GraphicsPSOInit;
+			RHICmdList.ApplyCachedRenderTargets(GraphicsPSOInit);
+			GraphicsPSOInit.DepthStencilState = TStaticDepthStencilState<false, CF_Always>::GetRHI();
+			GraphicsPSOInit.BlendState = TStaticBlendState<CW_RGBA, BO_Add, BF_One, BF_Zero, BO_Add, BF_One, BF_Zero>::GetRHI();
+			GraphicsPSOInit.RasterizerState = TStaticRasterizerState<>::GetRHI();
+			GraphicsPSOInit.PrimitiveType = PT_TriangleList;
+			GraphicsPSOInit.BoundShaderState.VertexDeclarationRHI = GetVertexDeclarationFVector4();
+			GraphicsPSOInit.BoundShaderState.VertexShaderRHI = GETSAFERHISHADER_VERTEX(*VertexShader);
+			GraphicsPSOInit.BoundShaderState.PixelShaderRHI = GETSAFERHISHADER_PIXEL(*PixelShader);
+			SetGraphicsPipelineState(RHICmdList, GraphicsPSOInit);
+
+			SetShaderParameters(RHICmdList, *VertexShader, VertexShader->GetVertexShader(), *Parameters);
+			SetShaderParameters(RHICmdList, *PixelShader, PixelShader->GetPixelShader(), *Parameters);
+
+			RHICmdList.SetViewport(0.0f, 0.0f, 0.0f, TextureSize[0], TextureSize[1], 1.0f);
+			RHICmdList.DrawIndexedPrimitive(GTwoTrianglesIndexBuffer.IndexBufferRHI, 0, 0, 4, 0, 2, 1);
+		});
+	}
+
+	/** Set up the copy pass for the given MaterialType. */
+	void AddCopyPass(FRDGBuilder& GraphBuilder, ERHIFeatureLevel::Type FeatureLevel, FShader_VirtualTextureCopy::FParameters* Parameters, FIntPoint TextureSize, ERuntimeVirtualTextureMaterialType MaterialType)
+	{
 		switch (MaterialType)
 		{
 		case ERuntimeVirtualTextureMaterialType::BaseColor_Normal_Specular:
-			AddCompressOrCopyPass<ERuntimeVirtualTextureMaterialType::BaseColor_Normal_Specular, true>(GraphBuilder, FeatureLevel, Parameters, GroupCount);
+			AddCopyPass<ERuntimeVirtualTextureMaterialType::BaseColor_Normal_Specular>(GraphBuilder, FeatureLevel, Parameters, TextureSize);
 			break;
 		case ERuntimeVirtualTextureMaterialType::BaseColor_Normal_Specular_YCoCg:
-			AddCompressOrCopyPass<ERuntimeVirtualTextureMaterialType::BaseColor_Normal_Specular_YCoCg, true>(GraphBuilder, FeatureLevel, Parameters, GroupCount);
+			AddCopyPass<ERuntimeVirtualTextureMaterialType::BaseColor_Normal_Specular_YCoCg>(GraphBuilder, FeatureLevel, Parameters, TextureSize);
+			break;
+		}
+	}
+
+	/** Set up the copy pass for the given MaterialType. */
+	void AddCopyThumbnailPass(FRDGBuilder& GraphBuilder, ERHIFeatureLevel::Type FeatureLevel, FShader_VirtualTextureCopy::FParameters* Parameters, FIntPoint TextureSize, ERuntimeVirtualTextureMaterialType MaterialType)
+	{
+		switch (MaterialType)
+		{
+		case ERuntimeVirtualTextureMaterialType::BaseColor:
+		case ERuntimeVirtualTextureMaterialType::BaseColor_Normal_Specular:
+		case ERuntimeVirtualTextureMaterialType::BaseColor_Normal_Specular_YCoCg:
+			AddCopyPass<ERuntimeVirtualTextureMaterialType::BaseColor>(GraphBuilder, FeatureLevel, Parameters, TextureSize);
 			break;
 		case ERuntimeVirtualTextureMaterialType::WorldHeight:
-			AddCompressOrCopyPass<ERuntimeVirtualTextureMaterialType::WorldHeight, true>(GraphBuilder, FeatureLevel, Parameters, GroupCount);
+			AddCopyPass<ERuntimeVirtualTextureMaterialType::WorldHeight>(GraphBuilder, FeatureLevel, Parameters, TextureSize);
 			break;
 		}
 	}
@@ -533,14 +625,12 @@ namespace RuntimeVirtualTexture
 	/** Structure to localize the setup of our render graph based on the virtual texture setup. */
 	struct FRenderGraphSetup
 	{
-		//todo[vt]: Add flag to disable the clear render target behavior and win some performance when we can. This could be driven a UI on the VT or the VT Plane?
-		FRenderGraphSetup(FRDGBuilder& GraphBuilder, ERuntimeVirtualTextureMaterialType MaterialType, FRHITexture2D* OutputTexture0, FIntPoint TextureSize)
+		FRenderGraphSetup(FRDGBuilder& GraphBuilder, ERuntimeVirtualTextureMaterialType MaterialType, FRHITexture2D* OutputTexture0, FIntPoint TextureSize, bool bIsThumbnails)
 		{
 			bRenderPass = OutputTexture0 != nullptr;
-			bCompressPass = bRenderPass && (OutputTexture0->GetFormat() == PF_DXT1 || OutputTexture0->GetFormat() == PF_DXT5 || OutputTexture0->GetFormat() == PF_BC5);
-			bCopyPass = bRenderPass && !bCompressPass && MaterialType == ERuntimeVirtualTextureMaterialType::BaseColor_Normal_Specular;
-			bCopyPass |= bRenderPass && !bCompressPass && MaterialType == ERuntimeVirtualTextureMaterialType::BaseColor_Normal_Specular_YCoCg;
-			bCopyPass |= bRenderPass && MaterialType == ERuntimeVirtualTextureMaterialType::WorldHeight && OutputTexture0->GetFormat() == PF_B8G8R8A8;
+			bCopyThumbnailPass = bRenderPass && bIsThumbnails;
+			bCompressPass = bRenderPass && !bCopyThumbnailPass && (OutputTexture0->GetFormat() == PF_DXT1 || OutputTexture0->GetFormat() == PF_DXT5 || OutputTexture0->GetFormat() == PF_BC5);
+			bCopyPass = bRenderPass && !bCopyThumbnailPass && !bCompressPass && (MaterialType == ERuntimeVirtualTextureMaterialType::BaseColor_Normal_Specular || MaterialType == ERuntimeVirtualTextureMaterialType::BaseColor_Normal_Specular_YCoCg);
 		
 			switch (MaterialType)
 			{
@@ -553,6 +643,10 @@ namespace RuntimeVirtualTexture
 				{
 					OutputAlias0 = CompressTexture0_u2 = GraphBuilder.CreateTexture(FPooledRenderTargetDesc::Create2DDesc(TextureSize / 4, PF_R32G32_UINT, FClearValueBinding::None, TexCreate_None, TexCreate_UAV, false), TEXT("CompressTexture0"));
 				}
+				if (bCopyThumbnailPass)
+				{
+					OutputAlias0 = CopyTexture0 = GraphBuilder.CreateTexture(FPooledRenderTargetDesc::Create2DDesc(TextureSize, PF_B8G8R8A8, FClearValueBinding::Black, TexCreate_SRGB, TexCreate_RenderTargetable | TexCreate_ShaderResource, false), TEXT("CopyTexture0"));
+				}
 				break;
 			case ERuntimeVirtualTextureMaterialType::BaseColor_Normal_Specular:
 				if (bRenderPass)
@@ -568,8 +662,12 @@ namespace RuntimeVirtualTexture
 				}
 				if (bCopyPass)
 				{
-					OutputAlias0 = CopyTexture0 = GraphBuilder.CreateTexture(FPooledRenderTargetDesc::Create2DDesc(TextureSize, PF_B8G8R8A8, FClearValueBinding::None, TexCreate_None, TexCreate_UAV, false), TEXT("CopyTexture0"));
-					OutputAlias1 = CopyTexture1 = GraphBuilder.CreateTexture(FPooledRenderTargetDesc::Create2DDesc(TextureSize, PF_B8G8R8A8, FClearValueBinding::None, TexCreate_None, TexCreate_UAV, false), TEXT("CopyTexture1"));
+					OutputAlias0 = CopyTexture0 = GraphBuilder.CreateTexture(FPooledRenderTargetDesc::Create2DDesc(TextureSize, PF_B8G8R8A8, FClearValueBinding::Black, TexCreate_SRGB, TexCreate_RenderTargetable | TexCreate_ShaderResource, false), TEXT("CopyTexture0"));
+					OutputAlias1 = CopyTexture1 = GraphBuilder.CreateTexture(FPooledRenderTargetDesc::Create2DDesc(TextureSize, PF_B8G8R8A8, FClearValueBinding::Black, TexCreate_None, TexCreate_RenderTargetable | TexCreate_ShaderResource, false), TEXT("CopyTexture1"));
+				}
+				if (bCopyThumbnailPass)
+				{
+					OutputAlias0 = CopyTexture0 = GraphBuilder.CreateTexture(FPooledRenderTargetDesc::Create2DDesc(TextureSize, PF_B8G8R8A8, FClearValueBinding::Black, TexCreate_SRGB, TexCreate_RenderTargetable | TexCreate_ShaderResource, false), TEXT("CopyTexture0"));
 				}
 				break;
 			case ERuntimeVirtualTextureMaterialType::BaseColor_Normal_Specular_YCoCg:
@@ -587,9 +685,13 @@ namespace RuntimeVirtualTexture
 				}
 				if (bCopyPass)
 				{
-					OutputAlias0 = CopyTexture0 = GraphBuilder.CreateTexture(FPooledRenderTargetDesc::Create2DDesc(TextureSize, PF_B8G8R8A8, FClearValueBinding::None, TexCreate_None, TexCreate_UAV, false), TEXT("CopyTexture0"));
-					OutputAlias1 = CopyTexture1 = GraphBuilder.CreateTexture(FPooledRenderTargetDesc::Create2DDesc(TextureSize, PF_B8G8R8A8, FClearValueBinding::None, TexCreate_None, TexCreate_UAV, false), TEXT("CopyTexture1"));
-					OutputAlias2 = CopyTexture2 = GraphBuilder.CreateTexture(FPooledRenderTargetDesc::Create2DDesc(TextureSize, PF_B8G8R8A8, FClearValueBinding::None, TexCreate_None, TexCreate_UAV, false), TEXT("CopyTexture2"));
+					OutputAlias0 = CopyTexture0 = GraphBuilder.CreateTexture(FPooledRenderTargetDesc::Create2DDesc(TextureSize, PF_B8G8R8A8, FClearValueBinding::Black, TexCreate_None, TexCreate_RenderTargetable | TexCreate_ShaderResource, false), TEXT("CopyTexture0"));
+					OutputAlias1 = CopyTexture1 = GraphBuilder.CreateTexture(FPooledRenderTargetDesc::Create2DDesc(TextureSize, PF_B8G8R8A8, FClearValueBinding::Black, TexCreate_None, TexCreate_RenderTargetable | TexCreate_ShaderResource, false), TEXT("CopyTexture1"));
+					OutputAlias2 = CopyTexture2 = GraphBuilder.CreateTexture(FPooledRenderTargetDesc::Create2DDesc(TextureSize, PF_B8G8R8A8, FClearValueBinding::Black, TexCreate_None, TexCreate_RenderTargetable | TexCreate_ShaderResource, false), TEXT("CopyTexture2"));
+				}
+				if (bCopyThumbnailPass)
+				{
+					OutputAlias0 = CopyTexture0 = GraphBuilder.CreateTexture(FPooledRenderTargetDesc::Create2DDesc(TextureSize, PF_B8G8R8A8, FClearValueBinding::Black, TexCreate_SRGB, TexCreate_RenderTargetable | TexCreate_ShaderResource, false), TEXT("CopyTexture0"));
 				}
 				break;
 			case ERuntimeVirtualTextureMaterialType::WorldHeight:
@@ -597,9 +699,9 @@ namespace RuntimeVirtualTexture
 				{
 					OutputAlias0 = RenderTexture0 = GraphBuilder.CreateTexture(FPooledRenderTargetDesc::Create2DDesc(TextureSize, PF_G16, FClearValueBinding::Black, TexCreate_None, TexCreate_RenderTargetable | TexCreate_ShaderResource, false), TEXT("RenderTexture0"));
 				}
-				if (bCopyPass)
+				if (bCopyThumbnailPass)
 				{
-					OutputAlias0 = CopyTexture0 = GraphBuilder.CreateTexture(FPooledRenderTargetDesc::Create2DDesc(TextureSize, PF_B8G8R8A8, FClearValueBinding::None, TexCreate_None, TexCreate_UAV, false), TEXT("CopyTexture0"));
+					OutputAlias0 = CopyTexture0 = GraphBuilder.CreateTexture(FPooledRenderTargetDesc::Create2DDesc(TextureSize, PF_B8G8R8A8, FClearValueBinding::Black, TexCreate_None, TexCreate_RenderTargetable | TexCreate_ShaderResource, false), TEXT("CopyTexture0"));
 				}
 				break;
 			}
@@ -609,6 +711,7 @@ namespace RuntimeVirtualTexture
 		bool bRenderPass = false;
 		bool bCompressPass = false;
 		bool bCopyPass = false;
+		bool bCopyThumbnailPass = false;
 
 		/** Render graph textures needed for this virtual texture layout. */
 		FRDGTextureRef RenderTexture0 = nullptr;
@@ -644,6 +747,7 @@ namespace RuntimeVirtualTexture
 		uint32 RuntimeVirtualTextureMask,
 		ERuntimeVirtualTextureMaterialType MaterialType,
 		bool bClearTextures,
+		bool bIsThumbnails,
 		FRHITexture2D* OutputTexture0,
 		FRHIUnorderedAccessView* OutputUAV0,
 		FBox2D const& DestBox0,
@@ -715,7 +819,7 @@ namespace RuntimeVirtualTexture
 		// Build graph
 		FMemMark Mark(FMemStack::Get());
 		FRDGBuilder GraphBuilder(RHICmdList);
-		FRenderGraphSetup GraphSetup(GraphBuilder, MaterialType, OutputTexture0, TextureSize);
+		FRenderGraphSetup GraphSetup(GraphBuilder, MaterialType, OutputTexture0, TextureSize, bIsThumbnails);
 
 		// Draw Pass
 		if (GraphSetup.bRenderPass)
@@ -756,9 +860,12 @@ namespace RuntimeVirtualTexture
 		}
 		
 		// Copy Pass
-		if (GraphSetup.bCopyPass)
+		if (GraphSetup.bCopyPass || GraphSetup.bCopyThumbnailPass)
 		{
-			FShader_VirtualTextureCompress::FParameters* PassParameters = GraphBuilder.AllocParameters<FShader_VirtualTextureCompress::FParameters>();
+			FShader_VirtualTextureCopy::FParameters* PassParameters = GraphBuilder.AllocParameters<FShader_VirtualTextureCopy::FParameters>();
+			PassParameters->RenderTargets[0] = GraphSetup.CopyTexture0 ? FRenderTargetBinding(GraphSetup.CopyTexture0, ERenderTargetLoadAction::ENoAction, ERenderTargetStoreAction::EStore) : FRenderTargetBinding();
+			PassParameters->RenderTargets[1] = GraphSetup.CopyTexture1 ? FRenderTargetBinding(GraphSetup.CopyTexture1, ERenderTargetLoadAction::ENoAction, ERenderTargetStoreAction::EStore) : FRenderTargetBinding();
+			PassParameters->RenderTargets[2] = GraphSetup.CopyTexture2 ? FRenderTargetBinding(GraphSetup.CopyTexture2, ERenderTargetLoadAction::ENoAction, ERenderTargetStoreAction::EStore) : FRenderTargetBinding();
 			PassParameters->DestRect = FIntVector4(0, 0, TextureSize.X, TextureSize.Y);
 			PassParameters->RenderTexture0 = GraphSetup.RenderTexture0;
 			PassParameters->TextureSampler0 = TStaticSamplerState<SF_Point, AM_Clamp, AM_Clamp, AM_Clamp>::GetRHI();
@@ -766,11 +873,15 @@ namespace RuntimeVirtualTexture
 			PassParameters->TextureSampler1 = TStaticSamplerState<SF_Point, AM_Clamp, AM_Clamp, AM_Clamp>::GetRHI();
 			PassParameters->RenderTexture2 = GraphSetup.RenderTexture2;
 			PassParameters->TextureSampler2 = TStaticSamplerState<SF_Point, AM_Clamp, AM_Clamp, AM_Clamp>::GetRHI();
-			PassParameters->OutCopyTexture0 = GraphSetup.CopyTexture0 ? GraphBuilder.CreateUAV(FRDGTextureUAVDesc(GraphSetup.CopyTexture0)) : nullptr;
-			PassParameters->OutCopyTexture1 = GraphSetup.CopyTexture1 ? GraphBuilder.CreateUAV(FRDGTextureUAVDesc(GraphSetup.CopyTexture1)) : nullptr;
-			PassParameters->OutCopyTexture2 = GraphSetup.CopyTexture2 ? GraphBuilder.CreateUAV(FRDGTextureUAVDesc(GraphSetup.CopyTexture2)) : nullptr;
 
-			AddCopyPass(GraphBuilder, View->GetFeatureLevel(), PassParameters, TextureSize, MaterialType);
+			if (GraphSetup.bCopyPass)
+			{
+				AddCopyPass(GraphBuilder, View->GetFeatureLevel(), PassParameters, TextureSize, MaterialType);
+			}
+			else
+			{
+				AddCopyThumbnailPass(GraphBuilder, View->GetFeatureLevel(), PassParameters, TextureSize, MaterialType);
+			}
 		}
 
 		// Set up the output to capture
@@ -851,6 +962,7 @@ namespace RuntimeVirtualTexture
 				InDesc.RuntimeVirtualTextureMask,
 				InDesc.MaterialType,
 				InDesc.bClearTextures,
+				InDesc.bIsThumbnails,
 				InDesc.Targets[0].Texture, InDesc.Targets[0].UAV, PageDesc.DestBox[0],
 				InDesc.Targets[1].Texture, InDesc.Targets[1].UAV, PageDesc.DestBox[1],
 				InDesc.Targets[2].Texture, InDesc.Targets[2].UAV, PageDesc.DestBox[2],
diff --git a/Engine/Source/Runtime/Renderer/Public/VT/RuntimeVirtualTextureRender.h b/Engine/Source/Runtime/Renderer/Public/VT/RuntimeVirtualTextureRender.h
--- a/Engine/Source/Runtime/Renderer/Public/VT/RuntimeVirtualTextureRender.h
+++ b/Engine/Source/Runtime/Renderer/Public/VT/RuntimeVirtualTextureRender.h
@@ -60,6 +60,8 @@ namespace RuntimeVirtualTexture
 		uint8 MaxLevel;
 		/** Set to true to clear before rendering. */
 		bool bClearTextures;
+		/** Set to true for thumbnail rendering. */
+		bool bIsThumbnails;
 		/** Debug visualization to render with. */
 		ERuntimeVirtualTextureDebugType DebugType;
 
