Fixed EV100 slider in the editor.

diff --git a/Engine/Source/Runtime/Renderer/Private/PostProcess/PostProcessEyeAdaptation.cpp b/Engine/Source/Runtime/Renderer/Private/PostProcess/PostProcessEyeAdaptation.cpp
--- a/Engine/Source/Runtime/Renderer/Private/PostProcess/PostProcessEyeAdaptation.cpp
+++ b/Engine/Source/Runtime/Renderer/Private/PostProcess/PostProcessEyeAdaptation.cpp
@@ -174,48 +174,44 @@ FEyeAdaptationParameters GetEyeAdaptationParameters(const FViewInfo& View, ERHIF
 	// These clamp the average luminance computed from the scene color.
 	float MinAverageLuminance = 1.0f;
 	float MaxAverageLuminance = 1.0f;
-	float ExposureCompensation = 1.0f;
+	float ExposureCompensation = GetAutoExposureCompensation(View);
 
 	// Force an exposure of 1 when any of these flags are set.
-	const bool bUseDebugExposure =
-		View.Family->ExposureSettings.bFixed ||
-		View.Family->UseDebugViewPS() ||
+	if (View.Family->UseDebugViewPS() ||
 		!EngineShowFlags.Lighting ||
 		(EngineShowFlags.VisualizeBuffer && View.CurrentBufferVisualizationMode != NAME_None) ||
 		EngineShowFlags.RayTracingDebug ||
 		EngineShowFlags.VisualizeDistanceFieldAO ||
 		EngineShowFlags.VisualizeGlobalDistanceField ||
 		EngineShowFlags.CollisionVisibility ||
-		EngineShowFlags.CollisionPawn;
-
-	if (!bUseDebugExposure)
+		EngineShowFlags.CollisionPawn)
+	{
+		ExposureCompensation = 1.0f;
+	}
+	// Fixed exposure override in effect.
+	else if (View.Family->ExposureSettings.bFixed)
 	{
-		// Fixed exposure override in effect.
-		if (View.Family->ExposureSettings.bFixed)
+		ExposureCompensation = 1.0f;
+		MinAverageLuminance = MaxAverageLuminance = EV100ToLuminance(View.Family->ExposureSettings.FixedEV100);
+	}
+	// When !EngineShowFlags.EyeAdaptation (from "r.EyeAdaptationQuality 0") or the feature level doesn't support eye adaptation, only Settings.AutoExposureBias controls exposure.
+	else if (EngineShowFlags.EyeAdaptation && View.GetFeatureLevel() >= MinFeatureLevel)
+	{
+		if (AutoExposureMethod == EAutoExposureMethod::AEM_Manual)
 		{
-			MinAverageLuminance = MaxAverageLuminance = EV100ToLuminance(View.Family->ExposureSettings.FixedEV100);
+			const float FixedEV100 = FMath::Log2(FMath::Square(Settings.DepthOfFieldFstop) * Settings.CameraShutterSpeed * 100 / FMath::Max(1.f, Settings.CameraISO));
+			MinAverageLuminance = MaxAverageLuminance = EV100ToLuminance(FixedEV100);
 		}
-		// When !EngineShowFlags.EyeAdaptation (from "r.EyeAdaptationQuality 0") or the feature level doesn't support eye adaptation, only Settings.AutoExposureBias controls exposure.
-		else if (EngineShowFlags.EyeAdaptation && View.GetFeatureLevel() >= MinFeatureLevel)
+		else if (bExtendedLuminanceRange)
 		{
-			if (AutoExposureMethod == EAutoExposureMethod::AEM_Manual)
-			{
-				const float FixedEV100 = FMath::Log2(FMath::Square(Settings.DepthOfFieldFstop) * Settings.CameraShutterSpeed * 100 / FMath::Max(1.f, Settings.CameraISO));
-				MinAverageLuminance = MaxAverageLuminance = EV100ToLuminance(FixedEV100);
-			}
-			else if (bExtendedLuminanceRange)
-			{
-				MinAverageLuminance = EV100ToLuminance(Settings.AutoExposureMinBrightness);
-				MaxAverageLuminance = EV100ToLuminance(Settings.AutoExposureMaxBrightness);
-			}
-			else
-			{
-				MinAverageLuminance = Settings.AutoExposureMinBrightness;
-				MaxAverageLuminance = Settings.AutoExposureMaxBrightness;
-			}
+			MinAverageLuminance = EV100ToLuminance(Settings.AutoExposureMinBrightness);
+			MaxAverageLuminance = EV100ToLuminance(Settings.AutoExposureMaxBrightness);
+		}
+		else
+		{
+			MinAverageLuminance = Settings.AutoExposureMinBrightness;
+			MaxAverageLuminance = Settings.AutoExposureMaxBrightness;
 		}
-
-		ExposureCompensation = GetAutoExposureCompensation(View);
 	}
 
 	MinAverageLuminance = FMath::Min(MinAverageLuminance, MaxAverageLuminance);
