Fix bad merge

diff --git a/Engine/Source/Runtime/Renderer/Private/VT/RuntimeVirtualTextureRender.cpp b/Engine/Source/Runtime/Renderer/Private/VT/RuntimeVirtualTextureRender.cpp
--- a/Engine/Source/Runtime/Renderer/Private/VT/RuntimeVirtualTextureRender.cpp
+++ b/Engine/Source/Runtime/Renderer/Private/VT/RuntimeVirtualTextureRender.cpp
@@ -509,7 +509,7 @@ namespace RuntimeVirtualTexture
 
 		static bool ShouldCompilePermutation(FGlobalShaderPermutationParameters const& Parameters)
 		{
-			return IsFeatureLevelSupported(Parameters.Platform, ERHIFeatureLevel::SM4);
+			return IsFeatureLevelSupported(Parameters.Platform, ERHIFeatureLevel::SM5);
 		}
 
 		FShader_VirtualTextureCopy()
@@ -863,9 +863,9 @@ namespace RuntimeVirtualTexture
 		if (GraphSetup.bCopyPass || GraphSetup.bCopyThumbnailPass)
 		{
 			FShader_VirtualTextureCopy::FParameters* PassParameters = GraphBuilder.AllocParameters<FShader_VirtualTextureCopy::FParameters>();
-			PassParameters->RenderTargets[0] = GraphSetup.CopyTexture0 ? FRenderTargetBinding(GraphSetup.CopyTexture0, ERenderTargetLoadAction::ENoAction, ERenderTargetStoreAction::EStore) : FRenderTargetBinding();
-			PassParameters->RenderTargets[1] = GraphSetup.CopyTexture1 ? FRenderTargetBinding(GraphSetup.CopyTexture1, ERenderTargetLoadAction::ENoAction, ERenderTargetStoreAction::EStore) : FRenderTargetBinding();
-			PassParameters->RenderTargets[2] = GraphSetup.CopyTexture2 ? FRenderTargetBinding(GraphSetup.CopyTexture2, ERenderTargetLoadAction::ENoAction, ERenderTargetStoreAction::EStore) : FRenderTargetBinding();
+			PassParameters->RenderTargets[0] = GraphSetup.CopyTexture0 ? FRenderTargetBinding(GraphSetup.CopyTexture0, ERenderTargetLoadAction::ENoAction) : FRenderTargetBinding();
+			PassParameters->RenderTargets[1] = GraphSetup.CopyTexture1 ? FRenderTargetBinding(GraphSetup.CopyTexture1, ERenderTargetLoadAction::ENoAction) : FRenderTargetBinding();
+			PassParameters->RenderTargets[2] = GraphSetup.CopyTexture2 ? FRenderTargetBinding(GraphSetup.CopyTexture2, ERenderTargetLoadAction::ENoAction) : FRenderTargetBinding();
 			PassParameters->DestRect = FIntVector4(0, 0, TextureSize.X, TextureSize.Y);
 			PassParameters->RenderTexture0 = GraphSetup.RenderTexture0;
 			PassParameters->TextureSampler0 = TStaticSamplerState<SF_Point, AM_Clamp, AM_Clamp, AM_Clamp>::GetRHI();
