Fixed sun light shaft disapearing when a cubemap is used and "Override Light Colors with Fog Inscattering Colors" is enabled.

diff --git a/Engine/Shaders/Private/VolumetricFog.usf b/Engine/Shaders/Private/VolumetricFog.usf
--- a/Engine/Shaders/Private/VolumetricFog.usf
+++ b/Engine/Shaders/Private/VolumetricFog.usf
@@ -440,7 +440,7 @@ float GetLightFunction(float3 WorldPosition)
 
 float SkyLightVolumetricScatteringIntensity;
 float4 SkySH[3];
-float UseHeightFogColors;
+float2 UseHeightFogColors;	// x=override directional light using height fog inscattering color, y=override sky light using heigh fog inscattering cubemap
 float UseDirectionalLightShadowing;
 float StaticLightingScatteringIntensity;
 
@@ -495,7 +495,7 @@ void LightScatteringCS(
 
 			float3 DirectionalLightColor = ForwardLightData.DirectionalLightColor;
 
-			if (UseHeightFogColors > 0)
+			if (UseHeightFogColors.x > 0)
 			{
 				// Attempt to maintain intensity ratio between sky and sun
 				DirectionalLightColor = VolumetricFog.HeightFogDirectionalLightInscatteringColor * Luminance(ForwardLightData.DirectionalLightColor);
@@ -524,7 +524,7 @@ void LightScatteringCS(
 		{
 			float3 SkyLighting;
 
-			if (UseHeightFogColors > 0)
+			if (UseHeightFogColors.y > 0)
 			{
 				float3 HeightFogInscatteringColor = ComputeInscatteringColor(CameraVector, CameraVectorLength);
 				float ScalarFactor = SHAmbientFunction();
diff --git a/Engine/Source/Runtime/Renderer/Private/VolumetricFog.cpp b/Engine/Source/Runtime/Renderer/Private/VolumetricFog.cpp
--- a/Engine/Source/Runtime/Renderer/Private/VolumetricFog.cpp
+++ b/Engine/Source/Runtime/Renderer/Private/VolumetricFog.cpp
@@ -385,6 +385,16 @@ FIntPoint CalculateVolumetricFogBoundsForLight(const FSphere& LightBounds, const
 	return VolumeZBounds;
 }
 
+static bool OverrideDirectionalLightInScatteringUsingHeightFog(const FViewInfo& View, const FExponentialHeightFogSceneInfo& FogInfo)
+{
+	return FogInfo.bOverrideLightColorsWithFogInscatteringColors && View.bUseDirectionalInscattering && !View.FogInscatteringColorCubemap;
+}
+
+static bool OverrideSkyLightInScatteringUsingHeightFog(const FViewInfo& View, const FExponentialHeightFogSceneInfo& FogInfo)
+{
+	return FogInfo.bOverrideLightColorsWithFogInscatteringColors;
+}
+
 /**  */
 class FCircleRasterizeVertexBuffer : public FVertexBuffer
 {
@@ -692,9 +702,12 @@ public:
 
 		SetShaderValue(RHICmdList, ShaderRHI, PhaseG, FogInfo.VolumetricFogScatteringDistribution);
 		SetShaderValue(RHICmdList, ShaderRHI, InverseSquaredLightDistanceBiasScale, GInverseSquaredLightDistanceBiasScale);
-		SetShaderValue(RHICmdList, ShaderRHI, UseHeightFogColors, FogInfo.bOverrideLightColorsWithFogInscatteringColors ? 1.0f : 0.0f);
 		SetShaderValue(RHICmdList, ShaderRHI, UseDirectionalLightShadowing, bUseDirectionalLightShadowing ? 1.0f : 0.0f);
 
+		SetShaderValue(RHICmdList, ShaderRHI, UseHeightFogColors, FVector2D(
+			OverrideDirectionalLightInScatteringUsingHeightFog(View, FogInfo) ? 1.0f : 0.0f,
+			OverrideSkyLightInScatteringUsingHeightFog(View, FogInfo) ? 1.0f : 0.0f ));
+
 		AOParameters.Set(RHICmdList, ShaderRHI, AOParameterData);
 		GlobalDistanceFieldParameters.Set(RHICmdList, ShaderRHI, View.GlobalDistanceFieldInfo.ParameterData);
 
@@ -902,8 +915,7 @@ void SetupVolumetricFogGlobalData(const FViewInfo& View, FVolumetricFogGlobalDat
 	Parameters.HeightFogInscatteringColor = View.ExponentialFogColor;
 
 	Parameters.HeightFogDirectionalLightInscatteringColor = FVector::ZeroVector;
-
-	if (View.bUseDirectionalInscattering && !View.FogInscatteringColorCubemap)
+	if (OverrideDirectionalLightInScatteringUsingHeightFog(View, FogInfo))
 	{
 		Parameters.HeightFogDirectionalLightInscatteringColor = FVector(View.DirectionalInscatteringColor);
 	}
