Refactor ray tracing reflections in multiview mode to remove unnecesary work and keep coherency with raster path.

diff --git a/Engine/Source/Runtime/Renderer/Private/DeferredShadingRenderer.h b/Engine/Source/Runtime/Renderer/Private/DeferredShadingRenderer.h
--- a/Engine/Source/Runtime/Renderer/Private/DeferredShadingRenderer.h
+++ b/Engine/Source/Runtime/Renderer/Private/DeferredShadingRenderer.h
@@ -504,9 +504,6 @@ private:
 		int32 SamplePerPixel,
 		int32 HeightFog,
 		float ResolutionFraction,
-		FRDGTextureUAV* ColorOutputUAV,
-		FRDGTextureUAV* RayHitDistanceOutputUAV,
-		FRDGTextureUAV* RayImaginaryDepthOutputUAV,
 		IScreenSpaceDenoiser::FReflectionsInputs* OutDenoiserInputs);
 
 	void RenderRayTracingShadows(
diff --git a/Engine/Source/Runtime/Renderer/Private/RayTracing/RayTracingReflections.cpp b/Engine/Source/Runtime/Renderer/Private/RayTracing/RayTracingReflections.cpp
--- a/Engine/Source/Runtime/Renderer/Private/RayTracing/RayTracingReflections.cpp
+++ b/Engine/Source/Runtime/Renderer/Private/RayTracing/RayTracingReflections.cpp
@@ -271,9 +271,6 @@ void FDeferredShadingSceneRenderer::RenderRayTracingReflections(
 	int32 SamplePerPixel,
 	int32 HeightFog,
 	float ResolutionFraction,
-	FRDGTextureUAV* ColorOutputUAV,
-	FRDGTextureUAV* RayHitDistanceOutputUAV,
-	FRDGTextureUAV* RayImaginaryDepthOutputUAV,
 	IScreenSpaceDenoiser::FReflectionsInputs* OutDenoiserInputs)
 #if RHI_RAYTRACING
 {
@@ -290,6 +287,22 @@ void FDeferredShadingSceneRenderer::RenderRayTracingReflections(
 	ensureMsgf(FComputeShaderUtils::kGolden2DGroupSize % UpscaleFactor == 0, TEXT("Reflection ray tracing will have uv misalignement."));
 	FIntPoint RayTracingResolution = FIntPoint::DivideAndRoundUp(View.ViewRect.Size(), UpscaleFactor);
 
+	{
+		FRDGTextureDesc Desc = FRDGTextureDesc::Create2DDesc(
+			SceneTextures.SceneDepthBuffer->Desc.Extent / UpscaleFactor,
+			PF_FloatRGBA,
+			FClearValueBinding::None,
+			/* InFlags = */ TexCreate_None,
+			/* InTargetableFlags = */ TexCreate_ShaderResource | TexCreate_RenderTargetable | TexCreate_UAV,
+			/* bInForceSeparateTargetAndShaderResource = */ false);
+
+		OutDenoiserInputs->Color = GraphBuilder.CreateTexture(Desc, TEXT("RayTracingReflections"));
+
+		Desc.Format = PF_R16F;
+		OutDenoiserInputs->RayHitDistance = GraphBuilder.CreateTexture(Desc, TEXT("RayTracingReflectionsHitDistance"));
+		OutDenoiserInputs->RayImaginaryDepth = GraphBuilder.CreateTexture(Desc, TEXT("RayTracingReflectionsImaginaryDepth"));
+	}
+
 	// When deferred materials are used, we need to dispatch the reflection shader twice:
 	// - First pass gathers reflected ray hit data and sorts it by hit shader ID.
 	// - Second pass re-traces the reflected ray and performs full shading.
@@ -357,9 +370,9 @@ void FDeferredShadingSceneRenderer::RenderRayTracingReflections(
 
 	CommonParameters.ReflectionStruct = CreateReflectionUniformBuffer(View, EUniformBufferUsage::UniformBuffer_SingleFrame);
 	CommonParameters.FogUniformParameters = CreateFogUniformBuffer(View, EUniformBufferUsage::UniformBuffer_SingleFrame);
-	CommonParameters.ColorOutput = ColorOutputUAV;
-	CommonParameters.RayHitDistanceOutput = RayHitDistanceOutputUAV;
-	CommonParameters.RayImaginaryDepthOutput = RayImaginaryDepthOutputUAV;
+	CommonParameters.ColorOutput = GraphBuilder.CreateUAV(OutDenoiserInputs->Color);
+	CommonParameters.RayHitDistanceOutput = GraphBuilder.CreateUAV(OutDenoiserInputs->RayHitDistance);
+	CommonParameters.RayImaginaryDepthOutput = GraphBuilder.CreateUAV(OutDenoiserInputs->RayImaginaryDepth);
 	CommonParameters.SortTileSize = SortTileSize;
 	CommonParameters.ReflectionCapture = View.ReflectionCaptureUniformBuffer;
 	CommonParameters.Forward = View.ForwardLightingResources->ForwardLightDataUniformBuffer;
diff --git a/Engine/Source/Runtime/Renderer/Private/ReflectionEnvironment.cpp b/Engine/Source/Runtime/Renderer/Private/ReflectionEnvironment.cpp
--- a/Engine/Source/Runtime/Renderer/Private/ReflectionEnvironment.cpp
+++ b/Engine/Source/Runtime/Renderer/Private/ReflectionEnvironment.cpp
@@ -658,31 +658,6 @@ void FDeferredShadingSceneRenderer::RenderDeferredReflectionsAndSkyLighting(FRHI
 	FSceneTextureParameters SceneTextures;
 	SetupSceneTextureParameters(GraphBuilder, &SceneTextures);
 
-	IScreenSpaceDenoiser::FReflectionsInputs DenoiserInputs;
-	IScreenSpaceDenoiser::FReflectionsRayTracingConfig RayTracingConfig;
-	RayTracingConfig.ResolutionFraction = FMath::Clamp(CVarReflectionScreenPercentage.GetValueOnRenderThread() / 100.0f, 0.25f, 1.0f);
-	int32 UpscaleFactor = int32(1.0f / RayTracingConfig.ResolutionFraction);
-
-	{
-		FRDGTextureDesc Desc = FRDGTextureDesc::Create2DDesc(
-			SceneTextures.SceneDepthBuffer->Desc.Extent / UpscaleFactor,
-			PF_FloatRGBA,
-			FClearValueBinding::None,
-			/* InFlags = */ TexCreate_None,
-			/* InTargetableFlags = */ TexCreate_ShaderResource | TexCreate_RenderTargetable | TexCreate_UAV,
-			/* bInForceSeparateTargetAndShaderResource = */ false);
-
-		DenoiserInputs.Color = GraphBuilder.CreateTexture(Desc, TEXT("RayTracingReflections"));
-
-		Desc.Format = PF_R16F;
-		DenoiserInputs.RayHitDistance = GraphBuilder.CreateTexture(Desc, TEXT("RayTracingReflectionsHitDistance"));
-		DenoiserInputs.RayImaginaryDepth = GraphBuilder.CreateTexture(Desc, TEXT("RayTracingReflectionsImaginaryDepth"));
-	}
-
-	FRDGTextureUAV* ReflectionColorOutputUAV = GraphBuilder.CreateUAV(FRDGTextureUAVDesc(DenoiserInputs.Color));
-	FRDGTextureUAV* RayHitDistanceOutputUAV = GraphBuilder.CreateUAV(FRDGTextureUAVDesc(DenoiserInputs.RayHitDistance));
-	FRDGTextureUAV* RayImaginaryDepthOutputUAV = GraphBuilder.CreateUAV(FRDGTextureUAVDesc(DenoiserInputs.RayImaginaryDepth));
-
 	uint32 ViewIndex = 0;
 	for (FViewInfo& View : Views)
 	{
@@ -696,13 +671,13 @@ void FDeferredShadingSceneRenderer::RenderDeferredReflectionsAndSkyLighting(FRHI
 		if (bRayTracedReflections || bScreenSpaceReflections)
 		{
 			int32 DenoiserMode = CVarUseReflectionDenoiser.GetValueOnRenderThread();
-			
-			RayTracingConfig.RayCountPerPixel = GRayTracingReflectionsSamplesPerPixel > -1 ? GRayTracingReflectionsSamplesPerPixel : View.FinalPostProcessSettings.RayTracingReflectionsSamplesPerPixel;
-			
+					
 			bool bDenoise = false;
 			bool bTemporalFilter = false;
 
 			// Traces the reflections, either using screen space reflection, or ray tracing.
+			IScreenSpaceDenoiser::FReflectionsInputs DenoiserInputs;
+			IScreenSpaceDenoiser::FReflectionsRayTracingConfig RayTracingConfig;
 			if (bRayTracedReflections)
 			{
 				RDG_EVENT_SCOPE(GraphBuilder, "RayTracingReflections");
@@ -710,6 +685,9 @@ void FDeferredShadingSceneRenderer::RenderDeferredReflectionsAndSkyLighting(FRHI
 
 				bDenoise = DenoiserMode != 0;
 				
+				RayTracingConfig.ResolutionFraction = FMath::Clamp(CVarReflectionScreenPercentage.GetValueOnRenderThread() / 100.0f, 0.25f, 1.0f);
+				RayTracingConfig.RayCountPerPixel = GRayTracingReflectionsSamplesPerPixel > -1 ? GRayTracingReflectionsSamplesPerPixel : View.FinalPostProcessSettings.RayTracingReflectionsSamplesPerPixel;
+
 				if (!bDenoise)
 				{
 					RayTracingConfig.ResolutionFraction = 1.0f;
@@ -720,9 +698,6 @@ void FDeferredShadingSceneRenderer::RenderDeferredReflectionsAndSkyLighting(FRHI
 					SceneTextures,
 					View,
 					RayTracingConfig.RayCountPerPixel, GRayTracingReflectionsHeightFog, RayTracingConfig.ResolutionFraction,
-					ReflectionColorOutputUAV,
-					RayHitDistanceOutputUAV,
-					RayImaginaryDepthOutputUAV,
 					&DenoiserInputs);
 			}
 			else if (bScreenSpaceReflections)
