Updated hair PPLL debug view with node usage.

diff --git a/Engine/Shaders/Private/HairStrands/HairStrandsVisibilityPPLLDebug.usf b/Engine/Shaders/Private/HairStrands/HairStrandsVisibilityPPLLDebug.usf
--- a/Engine/Shaders/Private/HairStrands/HairStrandsVisibilityPPLLDebug.usf
+++ b/Engine/Shaders/Private/HairStrands/HairStrandsVisibilityPPLLDebug.usf
@@ -1,7 +1,7 @@
 // Copyright 1998-2019 Epic Games, Inc. All Rights Reserved.
 
 #include "../Common.ush"
-
+#include "../ShaderPrintCommon.ush"
 #include "HairStrandsVisibilityCommon.ush"
 
 float PPLLMeanListElementCountPerPixel;
@@ -11,13 +11,15 @@ Texture2D<uint> PPLLCounter;
 Texture2D<uint> PPLLNodeIndex;
 StructuredBuffer<FPPLLNodeData> PPLLNodeData;
 
-void VisibilityDebugPPLLPS(
-	in float4 SVPos : SV_POSITION,
-	out float4 OutLuminance : SV_Target0)
+RWTexture2D<float4> SceneColorTextureUAV;
+
+[numthreads(8, 8, 1)]
+void VisibilityDebugPPLLCS(uint GroupIndex : SV_GroupIndex, uint3 DispatchThreadId : SV_DispatchThreadID)
 {
-	uint2 PixPos = SVPos.xy;
+	int2 PixelPos = DispatchThreadId.xy;
+	float PPLLListElementcount = PPLLCounter[uint2(0, 0)];
 
-	uint NextNodeIndex = PPLLNodeIndex[PixPos];
+	uint NextNodeIndex = PPLLNodeIndex[PixelPos];
 	if (NextNodeIndex != 0xFFFFFFFF)
 	{
 		uint StrandCount = 0;
@@ -30,33 +32,105 @@ void VisibilityDebugPPLLPS(
 
 		float l = saturate(StrandCount / PPLLMeanListElementCountPerPixel);
 
-		OutLuminance = float4(l, 1.0f - l, 0.0f, 0.0f);
-	}
-	else
-	{
-		OutLuminance = float4(0.0f, 0.0f, 0.0f, 1.0f);
+		SceneColorTextureUAV[PixelPos] = float4(l, 1.0f - l, 0.0f, 0.0f);
 	}
 
-	const float2 Origin = 100.0f;
-	const float2 Size = float2(200.0f, 20.0f);
-	if (all(PixPos >= Origin) && all(PixPos <= (Origin + Size)))
+	const float2 BarOrigin = 100.0f;
+	const float2 BarSize = float2(200.0f, 20.0f);
+	if (all(PixelPos >= BarOrigin) && all(PixelPos <= (BarOrigin + BarSize)))
 	{
-		float2 LocalPos = PixPos - Origin;
-		float PPLLListElementcount = PPLLCounter[uint2(0, 0)];
+		float2 LocalPos = PixelPos - BarOrigin;
 		float MemoryUsage = PPLLListElementcount / PPLLMaxTotalListElementCount;
 
-		if (MemoryUsage > LocalPos.x / Size.x)
+		if (MemoryUsage > LocalPos.x / BarSize.x)
 		{
-			OutLuminance = float4(MemoryUsage, 1.0f - MemoryUsage, 1.0f - MemoryUsage, 0.0f);
+			SceneColorTextureUAV[PixelPos] = float4(MemoryUsage, 1.0f - MemoryUsage, 1.0f - MemoryUsage, 0.0f);
 		}
 		else
 		{
-			OutLuminance = float4(0.2f, 0.2f, 0.2f, 0.0f);
+			SceneColorTextureUAV[PixelPos] = float4(0.2f, 0.2f, 0.2f, 0.0f);
 		}
 
-		if (any(PixPos == Origin) || any(PixPos == (Origin + Size)))
+		if (any(PixelPos == BarOrigin) || any(PixelPos == (BarOrigin + BarSize)))
 		{
-			OutLuminance = float4(1.0f, 1.0f, 1.0f, 0.0f);
+			SceneColorTextureUAV[PixelPos] = float4(1.0f, 1.0f, 1.0f, 0.0f);
 		}
 	}
+
+	if (DispatchThreadId.x == 0 && DispatchThreadId.y == 0)
+	{
+		ShaderPrintFilter(true);
+
+		float2 StartPos = (BarOrigin + BarSize*float2(0.0f, 2.0f)) * View.ViewSizeAndInvSize.zw;
+		float2 Pos = StartPos;
+
+		Pos = ShaderPrintSymbol(Pos, _M_);
+		Pos = ShaderPrintSymbol(Pos, _E_);
+		Pos = ShaderPrintSymbol(Pos, _A_);
+		Pos = ShaderPrintSymbol(Pos, _N_);
+		Pos = ShaderPrintSymbol(Pos, _SPC_);
+		Pos = ShaderPrintSymbol(Pos, _P_);
+		Pos = ShaderPrintSymbol(Pos, _P_);
+		Pos = ShaderPrintSymbol(Pos, _L_);
+		Pos = ShaderPrintSymbol(Pos, _L_);
+		Pos = ShaderPrintSymbol(Pos, _SPC_);
+		Pos = ShaderPrintSymbol(Pos, _N_);
+		Pos = ShaderPrintSymbol(Pos, _O_);
+		Pos = ShaderPrintSymbol(Pos, _D_);
+		Pos = ShaderPrintSymbol(Pos, _E_);
+		Pos = ShaderPrintSymbol(Pos, _SPC_);
+		Pos = ShaderPrintSymbol(Pos, _P_);
+		Pos = ShaderPrintSymbol(Pos, _E_);
+		Pos = ShaderPrintSymbol(Pos, _R_);
+		Pos = ShaderPrintSymbol(Pos, _SPC_);
+		Pos = ShaderPrintSymbol(Pos, _P_);
+		Pos = ShaderPrintSymbol(Pos, _I_);
+		Pos = ShaderPrintSymbol(Pos, _X_);
+		Pos = ShaderPrintSymbol(Pos, _E_);
+		Pos = ShaderPrintSymbol(Pos, _L_);
+		Pos = ShaderPrintSymbol(Pos, _SPC_);
+		Pos = ShaderPrint(Pos, PPLLMeanListElementCountPerPixel);
+		Pos = ShaderPrintNewline(Pos); Pos.x = StartPos.x;
+
+		Pos = ShaderPrintSymbol(Pos, _P_);
+		Pos = ShaderPrintSymbol(Pos, _P_);
+		Pos = ShaderPrintSymbol(Pos, _L_);
+		Pos = ShaderPrintSymbol(Pos, _L_);
+		Pos = ShaderPrintSymbol(Pos, _SPC_);
+		Pos = ShaderPrintSymbol(Pos, _N_);
+		Pos = ShaderPrintSymbol(Pos, _O_);
+		Pos = ShaderPrintSymbol(Pos, _D_);
+		Pos = ShaderPrintSymbol(Pos, _E_);
+		Pos = ShaderPrintSymbol(Pos, _S_);
+		Pos = ShaderPrintSymbol(Pos, _SPC_);
+		Pos = ShaderPrintSymbol(Pos, _U_);
+		Pos = ShaderPrintSymbol(Pos, _S_);
+		Pos = ShaderPrintSymbol(Pos, _E_);
+		Pos = ShaderPrintSymbol(Pos, _D_);
+		Pos = ShaderPrintSymbol(Pos, _SPC_);
+		Pos = ShaderPrintSymbol(Pos, _SPC_);
+		Pos = ShaderPrint(Pos, PPLLListElementcount);
+		Pos = ShaderPrintNewline(Pos); Pos.x = StartPos.x;
+
+		
+		Pos = ShaderPrintSymbol(Pos, _SPC_);
+		Pos = ShaderPrintSymbol(Pos, _SPC_);
+		Pos = ShaderPrintSymbol(Pos, _SPC_);
+		Pos = ShaderPrintSymbol(Pos, _SPC_);
+		Pos = ShaderPrintSymbol(Pos, _SPC_);
+		Pos = ShaderPrintSymbol(Pos, _SPC_);
+		Pos = ShaderPrintSymbol(Pos, _SPC_);
+		Pos = ShaderPrintSymbol(Pos, _SPC_);
+		Pos = ShaderPrintSymbol(Pos, _SPC_);
+		Pos = ShaderPrintSymbol(Pos, _SPC_);
+		Pos = ShaderPrintSymbol(Pos, _SPC_);
+		Pos = ShaderPrintSymbol(Pos, _T_);
+		Pos = ShaderPrintSymbol(Pos, _O_);
+		Pos = ShaderPrintSymbol(Pos, _T_);
+		Pos = ShaderPrintSymbol(Pos, _A_);
+		Pos = ShaderPrintSymbol(Pos, _L_);
+		Pos = ShaderPrintSymbol(Pos, _SPC_);
+		Pos = ShaderPrint(Pos, PPLLMaxTotalListElementCount);
+		Pos = ShaderPrintNewline(Pos);
+	}
 }
diff --git a/Engine/Source/Runtime/Renderer/Private/HairStrands/HairStrandsDebug.cpp b/Engine/Source/Runtime/Renderer/Private/HairStrands/HairStrandsDebug.cpp
--- a/Engine/Source/Runtime/Renderer/Private/HairStrands/HairStrandsDebug.cpp
+++ b/Engine/Source/Runtime/Renderer/Private/HairStrands/HairStrandsDebug.cpp
@@ -22,7 +22,7 @@
 #include "RenderTargetTemp.h"
 #include "CanvasTypes.h"
 #include "ShaderPrintParameters.h"
-#include "PixelShaderUtils.h"
+#include "RenderGraphUtils.h"
 
 static int32 GDeepShadowDebugIndex = 0;
 static float GDeepShadowDebugScale = 20;
@@ -1019,10 +1019,10 @@ const TCHAR* ToString(EWorldType::Type Type)
 
 ///////////////////////////////////////////////////////////////////////////////////////////////////
 
-class FHairVisibilityDebugPPLLPS : public FGlobalShader
+class FHairVisibilityDebugPPLLCS : public FGlobalShader
 {
-	DECLARE_GLOBAL_SHADER(FHairVisibilityDebugPPLLPS);
-	SHADER_USE_PARAMETER_STRUCT(FHairVisibilityDebugPPLLPS, FGlobalShader);
+	DECLARE_GLOBAL_SHADER(FHairVisibilityDebugPPLLCS);
+	SHADER_USE_PARAMETER_STRUCT(FHairVisibilityDebugPPLLCS, FGlobalShader);
 
 	using FPermutationDomain = TShaderPermutationDomain<>;
 
@@ -1033,7 +1033,8 @@ class FHairVisibilityDebugPPLLPS : public FGlobalShader
 		SHADER_PARAMETER_RDG_TEXTURE(Texture2D, PPLLCounter)
 		SHADER_PARAMETER_RDG_TEXTURE(Texture2D, PPLLNodeIndex)
 		SHADER_PARAMETER_RDG_BUFFER_SRV(StructuredBuffer, PPLLNodeData)
-		RENDER_TARGET_BINDING_SLOTS()
+		SHADER_PARAMETER_RDG_TEXTURE_UAV(Texture2D, SceneColorTextureUAV)
+		SHADER_PARAMETER_STRUCT_INCLUDE(ShaderPrint::FShaderParameters, ShaderPrintParameters)
 	END_SHADER_PARAMETER_STRUCT()
 
 		static FPermutationDomain RemapPermutation(FPermutationDomain PermutationVector)
@@ -1052,7 +1053,7 @@ class FHairVisibilityDebugPPLLPS : public FGlobalShader
 		OutEnvironment.SetDefine(TEXT("DEBUG_PPLL_PS"), 1);
 	}
 };
-IMPLEMENT_GLOBAL_SHADER(FHairVisibilityDebugPPLLPS, "/Engine/Private/HairStrands/HairStrandsVisibilityPPLLDebug.usf", "VisibilityDebugPPLLPS", SF_Pixel);
+IMPLEMENT_GLOBAL_SHADER(FHairVisibilityDebugPPLLCS, "/Engine/Private/HairStrands/HairStrandsVisibilityPPLLDebug.usf", "VisibilityDebugPPLLCS", SF_Compute);
 
 ///////////////////////////////////////////////////////////////////////////////////////////////////
 
@@ -1391,7 +1392,7 @@ void RenderHairStrandsDebugInfo(FRHICommandListImmediate& RHICmdList, TArray<FVi
 	if (ViewIndex < uint32(HairDatas->HairVisibilityViews.HairDatas.Num()))
 	{
 		const FHairStrandsVisibilityData& VisibilityData = HairDatas->HairVisibilityViews.HairDatas[ViewIndex];
-		if (GHairStrandsDebugPPLL && VisibilityData.PPLLNodeCounterTexture) // Check if PPLL rendering is used and its debuig view is enabled.
+		if (GHairStrandsDebugPPLL && VisibilityData.PPLLNodeCounterTexture) // Check if PPLL rendering is used and its debug view is enabled.
 		{
 			FRDGBuilder GraphBuilder(RHICmdList);
 			FRDGTextureRef SceneColorTexture = GraphBuilder.RegisterExternalTexture(SceneTargets.GetSceneColor(), TEXT("SceneColorTexture"));
@@ -1399,25 +1400,21 @@ void RenderHairStrandsDebugInfo(FRHICommandListImmediate& RHICmdList, TArray<FVi
 			FRDGTextureRef PPLLNodeIndexTexture = GraphBuilder.RegisterExternalTexture(VisibilityData.PPLLNodeIndexTexture, TEXT("PPLLNodeIndexTexture"));
 			FRDGBufferRef  PPLLNodeDataBuffer = GraphBuilder.RegisterExternalBuffer(VisibilityData.PPLLNodeDataBuffer, TEXT("PPLLNodeDataBuffer"));
 
-			FRHIBlendState* PreMultipliedColorTransmittanceBlend = TStaticBlendState<CW_RGB, BO_Add, BF_One, BF_SourceAlpha, BO_Add, BF_Zero, BF_One>::GetRHI();
-			FRHIDepthStencilState* DepthStencilStateWrite = TStaticDepthStencilState<false, CF_Always>::GetRHI();
-			FRHISamplerState* SamplerLinearClamp = TStaticSamplerState<SF_Trilinear>::GetRHI();
-
-			FHairVisibilityDebugPPLLPS::FPermutationDomain PermutationVector;
-			TShaderMapRef<FHairVisibilityDebugPPLLPS> PixelShader(View.ShaderMap, PermutationVector);
-
-			FHairVisibilityDebugPPLLPS::FParameters* PassParameters = GraphBuilder.AllocParameters<FHairVisibilityDebugPPLLPS::FParameters>();
-
+			FHairVisibilityDebugPPLLCS::FParameters* PassParameters = GraphBuilder.AllocParameters<FHairVisibilityDebugPPLLCS::FParameters>();
 			PassParameters->PPLLMeanListElementCountPerPixel = float(GetPPLLMeanListElementCountPerPixel());
 			PassParameters->PPLLMaxTotalListElementCount = float(GetPPLLMaxTotalListElementCount(VisibilityData.PPLLNodeIndexTexture->GetDesc().Extent));
 			PassParameters->PPLLCounter = PPLLNodeCounterTexture;
 			PassParameters->PPLLNodeIndex = PPLLNodeIndexTexture;
 			PassParameters->PPLLNodeData = GraphBuilder.CreateSRV(FRDGBufferSRVDesc(PPLLNodeDataBuffer));
 			PassParameters->ViewUniformBuffer = View.ViewUniformBuffer;
-			PassParameters->RenderTargets[0] = FRenderTargetBinding(SceneColorTexture, ERenderTargetLoadAction::ELoad);
-
-			FPixelShaderUtils::AddFullscreenPass<FHairVisibilityDebugPPLLPS>(GraphBuilder, View.ShaderMap, RDG_EVENT_NAME("HairPPLLDebug"), *PixelShader, PassParameters,
-				View.ViewRect, PreMultipliedColorTransmittanceBlend, nullptr, DepthStencilStateWrite);
+			PassParameters->SceneColorTextureUAV = GraphBuilder.CreateUAV(SceneColorTexture);
+			ShaderPrint::SetParameters(View, PassParameters->ShaderPrintParameters);
+
+			FHairVisibilityDebugPPLLCS::FPermutationDomain PermutationVector;
+			TShaderMapRef<FHairVisibilityDebugPPLLCS> ComputeShader(View.ShaderMap, PermutationVector);
+			FIntVector TextureSize = SceneColorTexture->Desc.GetSize(); TextureSize.Z = 1;
+			FComputeShaderUtils::AddPass(GraphBuilder, RDG_EVENT_NAME("HairPPLLDebug"), *ComputeShader, PassParameters,
+				FIntVector::DivideAndRoundUp(TextureSize, FIntVector(8, 8, 1)));
 			GraphBuilder.Execute();
 		}
 	}
