Controls SSGI in effects scalability settings, and enable with r.SSGI.Enable.

diff --git a/Engine/Config/BaseScalability.ini b/Engine/Config/BaseScalability.ini
--- a/Engine/Config/BaseScalability.ini
+++ b/Engine/Config/BaseScalability.ini
@@ -394,6 +394,7 @@ r.SSS.Scale=0
 r.SSS.SampleSet=0
 r.SSS.Quality=0
 r.SSS.HalfRes=1
+r.SSGI.Quality=0
 r.EmitterSpawnRateScale=0.125
 r.ParticleLightQuality=0
 r.DetailLevel=0
@@ -423,6 +424,7 @@ r.SSS.Scale=0.75
 r.SSS.SampleSet=0
 r.SSS.Quality=0
 r.SSS.HalfRes=1
+r.SSGI.Quality=1
 r.EmitterSpawnRateScale=0.25
 r.ParticleLightQuality=0
 r.DetailLevel=1
@@ -452,6 +454,7 @@ r.SSS.Scale=1
 r.SSS.SampleSet=1
 r.SSS.Quality=-1
 r.SSS.HalfRes=1
+r.SSGI.Quality=2
 r.EmitterSpawnRateScale=0.5
 r.ParticleLightQuality=1
 r.DetailLevel=2
@@ -481,6 +484,7 @@ r.SSS.Scale=1
 r.SSS.SampleSet=2
 r.SSS.Quality=1
 r.SSS.HalfRes=1
+r.SSGI.Quality=3
 r.EmitterSpawnRateScale=1.0
 r.ParticleLightQuality=2
 r.DetailLevel=3
@@ -510,6 +514,7 @@ r.SSS.Scale=1
 r.SSS.SampleSet=2
 r.SSS.Quality=1
 r.SSS.HalfRes=0
+r.SSGI.Quality=4
 r.EmitterSpawnRateScale=1.0
 r.ParticleLightQuality=2
 r.DetailLevel=4
diff --git a/Engine/Source/Runtime/Renderer/Private/ScreenSpaceRayTracing.cpp b/Engine/Source/Runtime/Renderer/Private/ScreenSpaceRayTracing.cpp
--- a/Engine/Source/Runtime/Renderer/Private/ScreenSpaceRayTracing.cpp
+++ b/Engine/Source/Runtime/Renderer/Private/ScreenSpaceRayTracing.cpp
@@ -43,20 +43,24 @@ static TAutoConsoleVariable<int32> CVarSSRStencil(
 	TEXT(" 0 is off (default), 1 is on"),
 	ECVF_RenderThreadSafe);
 
+static TAutoConsoleVariable<int32> CVarSSGIEnable(
+	TEXT("r.SSGI.Enable"), 0,
+	TEXT("Whether to enable SSGI (defaults to 0).\n"),
+	ECVF_RenderThreadSafe);
+
 static TAutoConsoleVariable<int32> CVarSSGIHalfResolution(
 	TEXT("r.SSGI.HalfRes"), 0,
-	TEXT("Whether to do SSGI at half resolution.\n"),
+	TEXT("Whether to do SSGI at half resolution (defaults to 0).\n"),
 	ECVF_Scalability | ECVF_RenderThreadSafe);
 
 static TAutoConsoleVariable<int32> CVarSSGIQuality(
-	TEXT("r.SSGI.Quality"),
-	0,
-	TEXT("Whether to use screen space diffuse indirect and at what quality setting.\n"),
+	TEXT("r.SSGI.Quality"), 4,
+	TEXT("Quality setting to control number of ray shot with SSGI, between 1 and 4 (defaults to 4).\n"),
 	ECVF_Scalability | ECVF_RenderThreadSafe);
 
 static TAutoConsoleVariable<float> CVarSSGIMinimumLuminance(
 	TEXT("r.SSGI.MinimumLuminance"), 0.5f,
-	TEXT(""),
+	TEXT("Experimental"),
 	ECVF_Scalability | ECVF_RenderThreadSafe);
 
 
@@ -99,6 +103,11 @@ bool ShouldRenderScreenSpaceReflections(const FViewInfo& View)
 
 bool ShouldRenderScreenSpaceDiffuseIndirect(const FViewInfo& View)
 {
+	if (CVarSSGIEnable.GetValueOnRenderThread() <= 0)
+	{
+		return false;
+	}
+
 	int Quality = CVarSSGIQuality.GetValueOnRenderThread();
 
 	if (Quality <= 0)
