Fixes a crash on r.SSR.ExperimentalDenoiser=1 code path

diff --git a/Engine/Source/Runtime/Renderer/Private/ScreenSpaceDenoise.cpp b/Engine/Source/Runtime/Renderer/Private/ScreenSpaceDenoise.cpp
--- a/Engine/Source/Runtime/Renderer/Private/ScreenSpaceDenoise.cpp
+++ b/Engine/Source/Runtime/Renderer/Private/ScreenSpaceDenoise.cpp
@@ -2335,7 +2335,8 @@ public:
 
 		// Imaginary depth is only used for Nvidia denoiser.
 		// TODO(Denoiser): permutation to not generate it?
-		GraphBuilder.RemoveUnusedTextureWarning(ReflectionInputs.RayImaginaryDepth);
+		if (ReflectionInputs.RayImaginaryDepth)
+			GraphBuilder.RemoveUnusedTextureWarning(ReflectionInputs.RayImaginaryDepth);
 
 		FSSDSignalTextures InputSignal;
 		InputSignal.Textures[0] = ReflectionInputs.Color;
