RL - HitProxies have to be deleted on the RenderingThread. Refcounting them should also only be done on the RT.

diff --git a/Engine/Source/Runtime/Renderer/Private/RendererScene.cpp b/Engine/Source/Runtime/Renderer/Private/RendererScene.cpp
--- a/Engine/Source/Runtime/Renderer/Private/RendererScene.cpp
+++ b/Engine/Source/Runtime/Renderer/Private/RendererScene.cpp
@@ -1444,21 +1444,15 @@ void FScene::RemovePrimitive( UPrimitiveComponent* Primitive )
 		// Disassociate the primitive's scene proxy.
 		Primitive->SceneProxy = NULL;
 
-		// Delete the PrimitiveSceneInfo on the game thread after the rendering thread has processed its removal.
-		// This must be done on the game thread because the hit proxy references (and possibly other members) need to be freed on the game thread.
-		struct DeferDeleteHitProxies : FDeferredCleanupInterface
-		{
-			DeferDeleteHitProxies(TArray<TRefCountPtr<HHitProxy>>&& InHitProxies) : HitProxies(MoveTemp(InHitProxies)) {}
-			TArray<TRefCountPtr<HHitProxy>> HitProxies;
-		};
-		BeginCleanup(new DeferDeleteHitProxies(MoveTemp(PrimitiveSceneInfo->HitProxies)));
-
 		// Send a command to the rendering thread to remove the primitive from the scene.
 		FScene* Scene = this;
 		FThreadSafeCounter* AttachmentCounter = &Primitive->AttachmentCounter;
 		ENQUEUE_RENDER_COMMAND(FRemovePrimitiveCommand)(
 			[Scene, PrimitiveSceneInfo, AttachmentCounter](FRHICommandList&)
 			{
+				PrimitiveSceneInfo->DefaultDynamicHitProxyId = FHitProxyId();
+				PrimitiveSceneInfo->DefaultDynamicHitProxy = nullptr;
+				PrimitiveSceneInfo->HitProxies.Reset();
 				PrimitiveSceneInfo->Proxy->DestroyRenderThreadResources();
 				Scene->RemovePrimitiveSceneInfo_RenderThread(PrimitiveSceneInfo);
 				AttachmentCounter->Decrement();
