Removed some shader permutations.
Make sure Coverage pass is never ran in PPLL mode.

diff --git a/Engine/Shaders/Private/HairStrands/HairStrandsVisibilityCompaction.usf b/Engine/Shaders/Private/HairStrands/HairStrandsVisibilityCompaction.usf
--- a/Engine/Shaders/Private/HairStrands/HairStrandsVisibilityCompaction.usf
+++ b/Engine/Shaders/Private/HairStrands/HairStrandsVisibilityCompaction.usf
@@ -31,10 +31,6 @@ uint MaxNodeCount;
 uint HairVisibilityMSAASampleCount;
 uint2 ResolutionOffset;
 
-#if PERMUTATION_COVERAGE == 1
-Texture2D<float> CoverageTexture;
-#endif
-
 RWTexture2D<uint> OutCompactNodeCounter;
 RWTexture2D<uint> OutCompactNodeIndex;
 RWStructuredBuffer<FNodeData> OutCompactNodeData;
@@ -46,7 +42,6 @@ RWTexture2D<uint4> OutCategorizationTexture;
 #define VELOCITY_TYPE_CLOSEST 2
 #define VELOCITY_TYPE_MAX 3
 #if PERMUTATION_VELOCITY != VELOCITY_TYPE_NONE
-	Texture2DMS<float2, HAIR_VISIBILITY_MAXSAMPLECOUNT> MSAA_VelocityTexture;
 	RWTexture2D<float2> OutVelocityTexture;
 #endif
 
@@ -257,11 +252,7 @@ void MainCS(uint GroupIndex : SV_GroupIndex, uint3 DispatchThreadId : SV_Dispatc
 		#endif		
 		}
 
-	#if PERMUTATION_COVERAGE == 1
-		Categorization.PixelCoverage = CoverageTexture.Load(uint3(PixelCoord, 0));
-	#else
 		Categorization.PixelCoverage = 1.0f - TotalTransmittance;
-	#endif
 
 	#if PERMUTATION_VELOCITY == VELOCITY_TYPE_AVG
 		OutVelocityTexture[PixelCoord] = EncodeVelocityToTexture(AverageVelocity / max(ValidVelocityCount, 1));
@@ -367,6 +358,14 @@ Texture2DMS<float, HAIR_VISIBILITY_MAXSAMPLECOUNT> MSAA_DepthTexture;
 Texture2DMS<float4,HAIR_VISIBILITY_MAXSAMPLECOUNT> MSAA_MaterialTexture;
 Texture2DMS<float4,HAIR_VISIBILITY_MAXSAMPLECOUNT> MSAA_AttributeTexture;
 
+#if PERMUTATION_VELOCITY != VELOCITY_TYPE_NONE
+Texture2DMS<float2, HAIR_VISIBILITY_MAXSAMPLECOUNT> MSAA_VelocityTexture;
+#endif
+
+#if PERMUTATION_COVERAGE == 1
+Texture2D<float> CoverageTexture;
+#endif
+
 [numthreads(TILE_PIXEL_SIZE_X, TILE_PIXEL_SIZE_Y, 1)]
 void MainCS(uint GroupIndex : SV_GroupIndex, uint3 DispatchThreadId : SV_DispatchThreadID)
 {
diff --git a/Engine/Shaders/Private/HairStrands/HairStrandsVisibilityDepthPS.usf b/Engine/Shaders/Private/HairStrands/HairStrandsVisibilityDepthPS.usf
--- a/Engine/Shaders/Private/HairStrands/HairStrandsVisibilityDepthPS.usf
+++ b/Engine/Shaders/Private/HairStrands/HairStrandsVisibilityDepthPS.usf
@@ -7,15 +7,16 @@
 #include "../DeferredShadingCommon.ush"
 #include "HairStrandsVisibilityCommon.ush"
 
+
 #if PERMUTATION_PPLL
-Texture2D<uint4> CategorisationTexture;
+	Texture2D<uint4> CategorisationTexture;
 #else
-Texture2DMS<float, HAIR_VISIBILITY_MAXSAMPLECOUNT> HairVisibilityDepthTexture;
+	Texture2DMS<float, HAIR_VISIBILITY_MAXSAMPLECOUNT> HairVisibilityDepthTexture;
+	#if PERMUTATION_COVERAGE
+		Texture2D<float> CoverageTexture;
+	#endif
 #endif
 
-#if PERMUTATION_COVERAGE == 1
-Texture2D<float> CoverageTexture;
-#endif
 
 void MainPS(
 	in FScreenVertexOutput Input,
diff --git a/Engine/Source/Runtime/Renderer/Private/HairStrands/HairStrandsVisibility.cpp b/Engine/Source/Runtime/Renderer/Private/HairStrands/HairStrandsVisibility.cpp
--- a/Engine/Source/Runtime/Renderer/Private/HairStrands/HairStrandsVisibility.cpp
+++ b/Engine/Source/Runtime/Renderer/Private/HairStrands/HairStrandsVisibility.cpp
@@ -564,7 +564,25 @@ class FHairVisibilityPrimitiveIdCompactionCS : public FGlobalShader
 	END_SHADER_PARAMETER_STRUCT()
 
 public:
-	static bool ShouldCompilePermutation(const FGlobalShaderPermutationParameters& Parameters) { return IsHairStrandsSupported(Parameters.Platform); }
+
+	static FPermutationDomain RemapPermutation(FPermutationDomain PermutationVector)
+	{
+		if (PermutationVector.Get<FPPLL>() > 0)
+		{
+			PermutationVector.Set<FCoverage>(0);
+		}
+		return PermutationVector;
+	}
+
+	static bool ShouldCompilePermutation(const FGlobalShaderPermutationParameters& Parameters)
+	{
+		FPermutationDomain PermutationVector(Parameters.PermutationId);
+		if (PermutationVector.Get<FPPLL>() > 0 && PermutationVector.Get<FCoverage>() > 0)
+		{
+			return false;
+		}
+		return IsHairStrandsSupported(Parameters.Platform);
+	}
 };
 
 IMPLEMENT_GLOBAL_SHADER(FHairVisibilityPrimitiveIdCompactionCS, "/Engine/Private/HairStrands/HairStrandsVisibilityCompaction.usf", "MainCS", SF_Compute);
@@ -676,6 +694,7 @@ static void AddHairVisibilityPrimitiveIdCompactionPass(
 	PermutationVector.Set<FHairVisibilityPrimitiveIdCompactionCS::FCoverage>(PassParameters->CoverageTexture ? 1 : 0);
 	PermutationVector.Set<FHairVisibilityPrimitiveIdCompactionCS::FMaterial>(GHairStrandsMaterialCompactionEnable ? 1 : 0);
 	PermutationVector.Set<FHairVisibilityPrimitiveIdCompactionCS::FPPLL>(bUsePPLL ? PPLLMaxRenderNodePerPixel : 0);
+	PermutationVector = FHairVisibilityPrimitiveIdCompactionCS::RemapPermutation(PermutationVector);
 
 	PassParameters->OutputResolution = Resolution;
 	PassParameters->MaxNodeCount = MaxRenderNodeCount;
@@ -1140,12 +1159,31 @@ class FHairVisibilityDepthPS : public FGlobalShader
 	END_SHADER_PARAMETER_STRUCT()
 
 public:
-	static bool ShouldCompilePermutation(const FGlobalShaderPermutationParameters& Parameters) { return IsHairStrandsSupported(Parameters.Platform); }
+		
+	static FPermutationDomain RemapPermutation(FPermutationDomain PermutationVector)
+	{
+		if (PermutationVector.Get<FPPLL>() > 0)
+		{
+			PermutationVector.Set<FCoverage>(0);
+		}
+		return PermutationVector;
+	}
+
+	static bool ShouldCompilePermutation(const FGlobalShaderPermutationParameters& Parameters)
+	{
+		FPermutationDomain PermutationVector(Parameters.PermutationId);
+		if (PermutationVector.Get<FPPLL>() > 0 && PermutationVector.Get<FCoverage>() > 0)
+		{
+			return false;
+		}
+		return IsHairStrandsSupported(Parameters.Platform);
+	}
 };
 
 IMPLEMENT_GLOBAL_SHADER(FHairVisibilityDepthPS, "/Engine/Private/HairStrands/HairStrandsVisibilityDepthPS.usf", "MainPS", SF_Pixel);
 
 static void AddHairVisibilityColorAndDepthPatchPass(
+	const bool bUsePPLL,
 	FRDGBuilder& GraphBuilder,
 	const FViewInfo& View,
 	const FRDGTextureRef& VisibilityDepthTexture,
@@ -1155,6 +1193,15 @@ static void AddHairVisibilityColorAndDepthPatchPass(
 	FRDGTextureRef& OutColorTexture,
 	FRDGTextureRef& OutDepthTexture)
 {
+	if (bUsePPLL)
+	{
+		check(CategorisationTexture);
+	}
+	else
+	{
+		check(VisibilityDepthTexture);
+	}
+
 	FHairVisibilityDepthPS::FParameters* Parameters = GraphBuilder.AllocParameters<FHairVisibilityDepthPS::FParameters>();
 	Parameters->HairVisibilityDepthTexture = VisibilityDepthTexture;
 	Parameters->CategorisationTexture = CategorisationTexture;
@@ -1170,7 +1217,8 @@ static void AddHairVisibilityColorAndDepthPatchPass(
 	TShaderMapRef<FPostProcessVS> VertexShader(View.ShaderMap);
 	FHairVisibilityDepthPS::FPermutationDomain PermutationVector;
 	PermutationVector.Set<FHairVisibilityDepthPS::FCoverage>(CoverageTexture ? 1 : 0);
-	PermutationVector.Set<FHairVisibilityDepthPS::FPPLL>(CategorisationTexture ? 1 : 0);
+	PermutationVector.Set<FHairVisibilityDepthPS::FPPLL>(bUsePPLL ? 1 : 0);
+	PermutationVector = FHairVisibilityDepthPS::RemapPermutation(PermutationVector);
 	TShaderMapRef<FHairVisibilityDepthPS> PixelShader(View.ShaderMap, PermutationVector);
 	const TShaderMap<FGlobalShaderType>* GlobalShaderMap = View.ShaderMap;
 	const FIntRect Viewport = View.ViewRect;
@@ -1260,8 +1308,9 @@ FHairStrandsVisibilityViews RenderHairStrandsVisibilityBuffer(
 			
 			const EHairVisibilityRenderMode RenderMode = GetHairVisibilityRenderMode();
 
+			// Run the coverage pass if needed (not in PPLL mode that is already a high quality render path)
 			FRDGTextureRef CoverageTexture = nullptr;
-			if (GHairStrandsCoveragePassEnable > 0)
+			if (GHairStrandsCoveragePassEnable > 0 && RenderMode != HairVisibilityRenderMode_PPLL)
 			{
 				CoverageTexture = AddHairCoveragePass(
 					GraphBuilder,
@@ -1348,6 +1397,7 @@ FHairStrandsVisibilityViews RenderHairStrandsVisibilityBuffer(
 				// * closest depth
 				// * unlit shading model ID 
 				AddHairVisibilityColorAndDepthPatchPass(
+					false, // bUsePPLL
 					GraphBuilder,
 					View,
 					MsaaVisibilityResources.DepthTexture,
@@ -1406,6 +1456,7 @@ FHairStrandsVisibilityViews RenderHairStrandsVisibilityBuffer(
 				}
 
 				AddHairVisibilityColorAndDepthPatchPass(
+					true, // bUsePPLL
 					GraphBuilder,
 					View,
 					nullptr,
