Improved coverage by using a transmittance with a mult operator instead (result in of higher quality, closee to PPLL now in the raster path).
Renamed everything appropriately.

diff --git a/Engine/Config/BaseScalability.ini b/Engine/Config/BaseScalability.ini
--- a/Engine/Config/BaseScalability.ini
+++ b/Engine/Config/BaseScalability.ini
@@ -578,7 +578,7 @@ grass.DensityScale=1.0
 ;todo
 
 [ShadingQuality@Cine]
-r.HairStrands.CoveragePass=1
+r.HairStrands.ViewTransmittancePass=1
 r.HairStrands.SkyLighting.SampleCount=64
 r.HairStrands.SkyLighting.JitterIntegration=1
 r.HairStrands.SkyAO.SampleCount=32
diff --git a/Engine/Shaders/Private/HairStrands/HairStrandsDebug.usf b/Engine/Shaders/Private/HairStrands/HairStrandsDebug.usf
--- a/Engine/Shaders/Private/HairStrands/HairStrandsDebug.usf
+++ b/Engine/Shaders/Private/HairStrands/HairStrandsDebug.usf
@@ -7,7 +7,7 @@ float2 OutputResolution;
 uint FastResolveMask;
 
 Texture2D<uint4> CategorizationTexture;
-Texture2D<float> CoverageTexture;
+Texture2D<float> ViewTransmittanceTexture;
 Texture2D<uint2> DepthStencilTexture;
 SamplerState LinearSampler;
 
diff --git a/Engine/Shaders/Private/HairStrands/HairStrandsVisibilityCommon.ush b/Engine/Shaders/Private/HairStrands/HairStrandsVisibilityCommon.ush
--- a/Engine/Shaders/Private/HairStrands/HairStrandsVisibilityCommon.ush
+++ b/Engine/Shaders/Private/HairStrands/HairStrandsVisibilityCommon.ush
@@ -3,7 +3,7 @@
 #pragma once
 
 #define RENDER_MODE_MSAA 0
-#define RENDER_MODE_COVERAGE 1
+#define RENDER_MODE_TRANSMITTANCE 1
 #define RENDER_MODE_PPLL 2
 
 struct FVisibilityData
diff --git a/Engine/Shaders/Private/HairStrands/HairStrandsVisibilityCompaction.usf b/Engine/Shaders/Private/HairStrands/HairStrandsVisibilityCompaction.usf
--- a/Engine/Shaders/Private/HairStrands/HairStrandsVisibilityCompaction.usf
+++ b/Engine/Shaders/Private/HairStrands/HairStrandsVisibilityCompaction.usf
@@ -362,8 +362,8 @@ Texture2DMS<float4,HAIR_VISIBILITY_MAXSAMPLECOUNT> MSAA_AttributeTexture;
 Texture2DMS<float2, HAIR_VISIBILITY_MAXSAMPLECOUNT> MSAA_VelocityTexture;
 #endif
 
-#if PERMUTATION_COVERAGE == 1
-Texture2D<float> CoverageTexture;
+#if PERMUTATION_VIEWTRANSMITTANCE == 1
+Texture2D<float> ViewTransmittanceTexture;
 #endif
 
 [numthreads(TILE_PIXEL_SIZE_X, TILE_PIXEL_SIZE_Y, 1)]
@@ -489,8 +489,8 @@ void MainCS(uint GroupIndex : SV_GroupIndex, uint3 DispatchThreadId : SV_Dispatc
 				OutCompactNodeCoord[StoreIndex] = (PixelCoord.x & 0xFFFF) | ((PixelCoord.y & 0xFFFF) << 16);
 			}
 
-			#if PERMUTATION_COVERAGE == 1
-			Categorization.PixelCoverage = CoverageTexture.Load(uint3(PixelCoord,0));
+			#if PERMUTATION_VIEWTRANSMITTANCE == 1
+			Categorization.PixelCoverage = 1.0f - ViewTransmittanceTexture.Load(uint3(PixelCoord,0));
 			#else
 			Categorization.PixelCoverage = SampleDesc.ValidSampleCount / float(HAIR_VISIBILITY_MAXSAMPLECOUNT);
 			#endif
diff --git a/Engine/Shaders/Private/HairStrands/HairStrandsVisibilityDepthPS.usf b/Engine/Shaders/Private/HairStrands/HairStrandsVisibilityDepthPS.usf
--- a/Engine/Shaders/Private/HairStrands/HairStrandsVisibilityDepthPS.usf
+++ b/Engine/Shaders/Private/HairStrands/HairStrandsVisibilityDepthPS.usf
@@ -12,8 +12,8 @@
 	Texture2D<uint4> CategorisationTexture;
 #else
 	Texture2DMS<float, HAIR_VISIBILITY_MAXSAMPLECOUNT> HairVisibilityDepthTexture;
-	#if PERMUTATION_COVERAGE
-		Texture2D<float> CoverageTexture;
+	#if PERMUTATION_VIEWTRANSMITTANCE
+		Texture2D<float> ViewTransmittanceTexture;
 	#endif
 #endif
 
@@ -51,8 +51,8 @@ void MainPS(
 		ValidSampleCount += bIsValid ? 1 : 0;
 	}
 
-#if PERMUTATION_COVERAGE == 1
-	if (CoverageTexture.Load(uint3(PixelCoord, 0)) < 1.0f || ValidSampleCount == 0)
+#if PERMUTATION_VIEWTRANSMITTANCE == 1
+	if ((1.0f - ViewTransmittanceTexture.Load(uint3(PixelCoord, 0))) < 1.0f || ValidSampleCount == 0)
 	{
 		discard; // If coverage is 1. we skip writting depth to match all the other passes relying on coverage==1 meaning opaque.
 	}
diff --git a/Engine/Shaders/Private/HairStrands/HairStrandsVisibilityPS.usf b/Engine/Shaders/Private/HairStrands/HairStrandsVisibilityPS.usf
--- a/Engine/Shaders/Private/HairStrands/HairStrandsVisibilityPS.usf
+++ b/Engine/Shaders/Private/HairStrands/HairStrandsVisibilityPS.usf
@@ -23,7 +23,7 @@
 #define SUPPORT_MATERIAL_PROPERTY 1
 #endif
 
-#if HAIR_RENDER_MODE == RENDER_MODE_COVERAGE
+#if HAIR_RENDER_MODE == RENDER_MODE_TRANSMITTANCE
 #define SUPPORT_MATERIAL_PROPERTY 0
 #endif
 
@@ -46,11 +46,11 @@ void MainVisibility(
 	, centroid in float WorldStrandRadius : TEXCOORD8
 	, const float4 PackedVelocityA : TEXCOORD9
 	, const float4 PackedVelocityC : TEXCOORD12
-#elif HAIR_RENDER_MODE == RENDER_MODE_COVERAGE 
+#elif HAIR_RENDER_MODE == RENDER_MODE_TRANSMITTANCE 
 	, centroid in float WorldStrandRadius : TEXCOORD8
 #endif
 	  OPTIONAL_IsFrontFace
-#if HAIR_RENDER_MODE == RENDER_MODE_COVERAGE
+#if HAIR_RENDER_MODE == RENDER_MODE_TRANSMITTANCE
 	, out float OutColor0 : SV_Target0
 #elif HAIR_RENDER_MODE == RENDER_MODE_MSAA
 	, out uint4 OutColor0	: SV_Target0
@@ -82,8 +82,8 @@ void MainVisibility(
 		Coverage = StrandRealRadius / max(StrandRealRadius, PixelRadius);
 	}
 
-#if HAIR_RENDER_MODE == RENDER_MODE_COVERAGE
-	OutColor0 = Coverage;
+#if HAIR_RENDER_MODE == RENDER_MODE_TRANSMITTANCE
+	OutColor0 = saturate(1.0f - Coverage);
 #endif
 
 #if HAIR_RENDER_MODE == RENDER_MODE_MSAA || HAIR_RENDER_MODE == RENDER_MODE_PPLL
diff --git a/Engine/Source/Runtime/Renderer/Private/HairStrands/HairStrandsVisibility.cpp b/Engine/Source/Runtime/Renderer/Private/HairStrands/HairStrandsVisibility.cpp
--- a/Engine/Source/Runtime/Renderer/Private/HairStrands/HairStrandsVisibility.cpp
+++ b/Engine/Source/Runtime/Renderer/Private/HairStrands/HairStrandsVisibility.cpp
@@ -18,8 +18,8 @@ DECLARE_GPU_STAT(HairStrandsVisibility);
 
 /////////////////////////////////////////////////////////////////////////////////////////
 
-static int32 GHairStrandsCoveragePassEnable = 0;
-static FAutoConsoleVariableRef CVarHairStrandsCoveragePassEnable(TEXT("r.HairStrands.CoveragePass"), GHairStrandsCoveragePassEnable, TEXT("Enable accurate coverage pass"));
+static int32 GHairStrandsViewTransmittancePassEnable = 0;
+static FAutoConsoleVariableRef CVarHairStrandsCoveragePassEnable(TEXT("r.HairStrands.ViewTransmittancePass"), GHairStrandsViewTransmittancePassEnable, TEXT("Enable accurate transmittance pass for better rendering of small scale hair strand."));
 
 static int32 GHairStrandsMaterialCompactionEnable = 0;
 static FAutoConsoleVariableRef CVarHairStrandsMaterialCompactionEnable(TEXT("r.HairStrands.MaterialCompaction"), GHairStrandsMaterialCompactionEnable, TEXT("Enable extra compaction based on material properties in order to reduce sample per pixel and improve performance."));
@@ -57,7 +57,7 @@ static FAutoConsoleVariableRef CVarGHairVisibilityPPLLMeanNodeCountPerPixel(TEXT
 enum EHairVisibilityRenderMode
 {
 	HairVisibilityRenderMode_MSAA,
-	HairVisibilityRenderMode_Coverage,
+	HairVisibilityRenderMode_Transmittance,
 	HairVisibilityRenderMode_PPLL,
 	HairVisibilityRenderModeCount
 };
@@ -168,7 +168,7 @@ protected:
 	}
 };
 IMPLEMENT_MATERIAL_SHADER_TYPE(template<>, FHairVisibilityVS<HairVisibilityRenderMode_MSAA>, TEXT("/Engine/Private/HairStrands/HairStrandsVisibilityVS.usf"), TEXT("Main"), SF_Vertex);
-IMPLEMENT_MATERIAL_SHADER_TYPE(template<>, FHairVisibilityVS<HairVisibilityRenderMode_Coverage>, TEXT("/Engine/Private/HairStrands/HairStrandsVisibilityVS.usf"), TEXT("Main"), SF_Vertex);
+IMPLEMENT_MATERIAL_SHADER_TYPE(template<>, FHairVisibilityVS<HairVisibilityRenderMode_Transmittance>, TEXT("/Engine/Private/HairStrands/HairStrandsVisibilityVS.usf"), TEXT("Main"), SF_Vertex);
 IMPLEMENT_MATERIAL_SHADER_TYPE(template<>, FHairVisibilityVS<HairVisibilityRenderMode_PPLL>, TEXT("/Engine/Private/HairStrands/HairStrandsVisibilityVS.usf"), TEXT("Main"), SF_Vertex);
 
 /////////////////////////////////////////////////////////////////////////////////////////
@@ -248,7 +248,7 @@ public:
 	FShaderParameter PPLLNodes;
 };
 IMPLEMENT_MATERIAL_SHADER_TYPE(template<>, FHairVisibilityPS<HairVisibilityRenderMode_MSAA>, TEXT("/Engine/Private/HairStrands/HairStrandsVisibilityPS.usf"), TEXT("MainVisibility"), SF_Pixel);
-IMPLEMENT_MATERIAL_SHADER_TYPE(template<>, FHairVisibilityPS<HairVisibilityRenderMode_Coverage>, TEXT("/Engine/Private/HairStrands/HairStrandsVisibilityPS.usf"), TEXT("MainVisibility"), SF_Pixel);
+IMPLEMENT_MATERIAL_SHADER_TYPE(template<>, FHairVisibilityPS<HairVisibilityRenderMode_Transmittance>, TEXT("/Engine/Private/HairStrands/HairStrandsVisibilityPS.usf"), TEXT("MainVisibility"), SF_Pixel);
 IMPLEMENT_MATERIAL_SHADER_TYPE(template<>, FHairVisibilityPS<HairVisibilityRenderMode_PPLL>, TEXT("/Engine/Private/HairStrands/HairStrandsVisibilityPS.usf"), TEXT("MainVisibility"), SF_Pixel);
 
 /////////////////////////////////////////////////////////////////////////////////////////
@@ -311,8 +311,8 @@ void FHairVisibilityProcessor::AddMeshBatch(const FMeshBatch& RESTRICT MeshBatch
 		const ERasterizerCullMode MeshCullMode = ComputeMeshCullMode(MeshBatch, Material);
 		if (RenderMode == HairVisibilityRenderMode_MSAA)
 			Process<HairVisibilityRenderMode_MSAA>(MeshBatch, BatchElementMask, PrimitiveSceneProxy, StaticMeshId, MaterialRenderProxy, Material, HairClusterId, MeshFillMode, MeshCullMode);
-		if (RenderMode == HairVisibilityRenderMode_Coverage)
-			Process<HairVisibilityRenderMode_Coverage>(MeshBatch, BatchElementMask, PrimitiveSceneProxy, StaticMeshId, MaterialRenderProxy, Material, HairClusterId, MeshFillMode, MeshCullMode);
+		if (RenderMode == HairVisibilityRenderMode_Transmittance)
+			Process<HairVisibilityRenderMode_Transmittance>(MeshBatch, BatchElementMask, PrimitiveSceneProxy, StaticMeshId, MaterialRenderProxy, Material, HairClusterId, MeshFillMode, MeshCullMode);
 		if (RenderMode == HairVisibilityRenderMode_PPLL)
 			Process<HairVisibilityRenderMode_PPLL>(MeshBatch, BatchElementMask, PrimitiveSceneProxy, StaticMeshId, MaterialRenderProxy, Material, HairClusterId, MeshFillMode, MeshCullMode);
 	}
@@ -526,10 +526,10 @@ class FHairVisibilityPrimitiveIdCompactionCS : public FGlobalShader
 
 	class FVendor   : SHADER_PERMUTATION_INT("PERMUTATION_VENDOR", HairVisibilityVendorCount);
 	class FVelocity : SHADER_PERMUTATION_INT("PERMUTATION_VELOCITY", 4);
-	class FCoverage : SHADER_PERMUTATION_INT("PERMUTATION_COVERAGE", 2);
+	class FViewTransmittance : SHADER_PERMUTATION_INT("PERMUTATION_VIEWTRANSMITTANCE", 2);
 	class FMaterial : SHADER_PERMUTATION_INT("PERMUTATION_MATERIAL_COMPACTION", 2);
 	class FPPLL     : SHADER_PERMUTATION_SPARSE_INT("PERMUTATION_PPLL", 0, 8, 16, 32); // See GetPPLLMaxRenderNodePerPixel
-	using FPermutationDomain = TShaderPermutationDomain<FVendor, FVelocity, FCoverage, FMaterial, FPPLL>;
+	using FPermutationDomain = TShaderPermutationDomain<FVendor, FVelocity, FViewTransmittance, FMaterial, FPPLL>;
 
 	BEGIN_SHADER_PARAMETER_STRUCT(FParameters, )
 		SHADER_PARAMETER(FIntPoint, OutputResolution)
@@ -550,7 +550,7 @@ class FHairVisibilityPrimitiveIdCompactionCS : public FGlobalShader
 		SHADER_PARAMETER_RDG_TEXTURE(Texture2D, PPLLNodeIndex)
 		SHADER_PARAMETER_RDG_BUFFER_SRV(StructuredBuffer, PPLLNodeData)
 
-		SHADER_PARAMETER_RDG_TEXTURE(Texture2D, CoverageTexture)
+		SHADER_PARAMETER_RDG_TEXTURE(Texture2D, ViewTransmittanceTexture)
 		
 		SHADER_PARAMETER_RDG_TEXTURE_UAV(Texture2D, OutCompactNodeCounter)
 		SHADER_PARAMETER_RDG_TEXTURE_UAV(Texture2D, OutCompactNodeIndex)
@@ -569,7 +569,7 @@ public:
 	{
 		if (PermutationVector.Get<FPPLL>() > 0)
 		{
-			PermutationVector.Set<FCoverage>(0);
+			PermutationVector.Set<FViewTransmittance>(0);
 		}
 		return PermutationVector;
 	}
@@ -577,7 +577,7 @@ public:
 	static bool ShouldCompilePermutation(const FGlobalShaderPermutationParameters& Parameters)
 	{
 		FPermutationDomain PermutationVector(Parameters.PermutationId);
-		if (PermutationVector.Get<FPPLL>() > 0 && PermutationVector.Get<FCoverage>() > 0)
+		if (PermutationVector.Get<FPPLL>() > 0 && PermutationVector.Get<FViewTransmittance>() > 0)
 		{
 			return false;
 		}
@@ -691,7 +691,7 @@ static void AddHairVisibilityPrimitiveIdCompactionPass(
 	FHairVisibilityPrimitiveIdCompactionCS::FPermutationDomain PermutationVector;
 	PermutationVector.Set<FHairVisibilityPrimitiveIdCompactionCS::FVendor>(GetVendor());
 	PermutationVector.Set<FHairVisibilityPrimitiveIdCompactionCS::FVelocity>(bWriteOutVelocity ? FMath::Clamp(GHairVelocityType+1, 0, 3) : 0);
-	PermutationVector.Set<FHairVisibilityPrimitiveIdCompactionCS::FCoverage>(PassParameters->CoverageTexture ? 1 : 0);
+	PermutationVector.Set<FHairVisibilityPrimitiveIdCompactionCS::FViewTransmittance>(PassParameters->ViewTransmittanceTexture ? 1 : 0);
 	PermutationVector.Set<FHairVisibilityPrimitiveIdCompactionCS::FMaterial>(GHairStrandsMaterialCompactionEnable ? 1 : 0);
 	PermutationVector.Set<FHairVisibilityPrimitiveIdCompactionCS::FPPLL>(bUsePPLL ? PPLLMaxRenderNodePerPixel : 0);
 	PermutationVector = FHairVisibilityPrimitiveIdCompactionCS::RemapPermutation(PermutationVector);
@@ -862,7 +862,7 @@ static void AddHairVisibilityCommonPass(
 {
 	GraphBuilder.AddPass(
 		RenderMode == HairVisibilityRenderMode_PPLL ? RDG_EVENT_NAME("HairStrandsPPLLPass") : 
-		(RenderMode == HairVisibilityRenderMode_MSAA ? RDG_EVENT_NAME("HairStrandsVisibilityPass") : RDG_EVENT_NAME("HairStrandsCoveragePass")),
+		(RenderMode == HairVisibilityRenderMode_MSAA ? RDG_EVENT_NAME("HairStrandsVisibilityPass") : RDG_EVENT_NAME("HairStrandsTransmittancePass")),
 		PassParameters,
 		ERDGPassFlags::Raster,
 		[PassParameters, Scene = Scene, ViewInfo, &ClusterDatas, RenderMode](FRHICommandListImmediate& RHICmdList)
@@ -872,7 +872,7 @@ static void AddHairVisibilityCommonPass(
 
 		FMeshPassProcessorRenderState DrawRenderState(*ViewInfo);
 
-		if(RenderMode == HairVisibilityRenderMode_Coverage || RenderMode == HairVisibilityRenderMode_PPLL)
+		if(RenderMode == HairVisibilityRenderMode_Transmittance || RenderMode == HairVisibilityRenderMode_PPLL)
 		{
 			TUniformBufferRef<FViewUniformShaderParameters> ViewUniformShaderParameters;
 			const bool bEnableMSAA = false;
@@ -891,9 +891,9 @@ static void AddHairVisibilityCommonPass(
 					CW_RGBA, BO_Add, BF_One, BF_Zero, BO_Add, BF_One, BF_Zero>::GetRHI());
 				DrawRenderState.SetDepthStencilState(TStaticDepthStencilState<true, CF_DepthNearOrEqual>::GetRHI());
 			}
-			else if (RenderMode == HairVisibilityRenderMode_Coverage)
+			else if (RenderMode == HairVisibilityRenderMode_Transmittance)
 			{
-				DrawRenderState.SetBlendState(TStaticBlendState<CW_RED, BO_Add, BF_One, BF_One, BO_Add, BF_One, BF_One>::GetRHI());
+				DrawRenderState.SetBlendState(TStaticBlendState<CW_RED, BO_Add, BF_DestColor, BF_Zero, BO_Add, BF_Zero, BF_Zero>::GetRHI());
 				DrawRenderState.SetDepthStencilState(TStaticDepthStencilState<false, CF_DepthNearOrEqual>::GetRHI());
 			}
 			else if (RenderMode == HairVisibilityRenderMode_PPLL)
@@ -1103,7 +1103,7 @@ static void AddHairVisibilityPPLLPass(
 
 ///////////////////////////////////////////////////////////////////////////////////////////////////
 
-static FRDGTextureRef AddHairCoveragePass(
+static FRDGTextureRef AddHairViewTransmittancePass(
 	FRDGBuilder& GraphBuilder,
 	const FScene* Scene,
 	const FViewInfo* ViewInfo,
@@ -1113,7 +1113,7 @@ static FRDGTextureRef AddHairCoveragePass(
 {
 	check(SceneDepthTexture->Desc.Extent == Resolution);
 
-	FRDGTextureRef CoverageTexture = nullptr;
+	FRDGTextureRef ViewTransmittanceTexture = nullptr;
 	{
 		FRDGTextureDesc Desc;
 		Desc.Extent.X = Resolution.X;
@@ -1125,20 +1125,21 @@ static FRDGTextureRef AddHairCoveragePass(
 		Desc.Flags = TexCreate_None;
 		Desc.TargetableFlags = TexCreate_RenderTargetable | TexCreate_ShaderResource;
 		Desc.bForceSharedTargetAndShaderResource = true;
-		CoverageTexture = GraphBuilder.CreateTexture(Desc, TEXT("HairCoverageTexture"));
+		Desc.ClearValue = FClearValueBinding(FLinearColor(1.0f, 1.0f, 1.0f, 1.0f)); // Clear to transmittance =1
+		ViewTransmittanceTexture = GraphBuilder.CreateTexture(Desc, TEXT("HairViewTransmittanceTexture"));
 	}
 
 	FVisibilityPassParameters* PassParameters = GraphBuilder.AllocParameters<FVisibilityPassParameters>();
-	PassParameters->RenderTargets[0] = FRenderTargetBinding(CoverageTexture, ERenderTargetLoadAction::EClear, 0);
+	PassParameters->RenderTargets[0] = FRenderTargetBinding(ViewTransmittanceTexture, ERenderTargetLoadAction::EClear, 0);
 
 	PassParameters->RenderTargets.DepthStencil = FDepthStencilBinding(
 		SceneDepthTexture,
 		ERenderTargetLoadAction::ELoad,
 		ERenderTargetLoadAction::ENoAction,
 		FExclusiveDepthStencil::DepthRead_StencilNop);
-	AddHairVisibilityCommonPass(GraphBuilder, Scene, ViewInfo, ClusterDatas, HairVisibilityRenderMode_Coverage, PassParameters);
+	AddHairVisibilityCommonPass(GraphBuilder, Scene, ViewInfo, ClusterDatas, HairVisibilityRenderMode_Transmittance, PassParameters);
 
-	return CoverageTexture;
+	return ViewTransmittanceTexture;
 }
 
 ///////////////////////////////////////////////////////////////////////////////////////////////////
@@ -1147,14 +1148,14 @@ class FHairVisibilityDepthPS : public FGlobalShader
 	DECLARE_GLOBAL_SHADER(FHairVisibilityDepthPS);
 	SHADER_USE_PARAMETER_STRUCT(FHairVisibilityDepthPS, FGlobalShader);
 
-	class FCoverage : SHADER_PERMUTATION_INT("PERMUTATION_COVERAGE", 2);
+	class FViewTransmittance : SHADER_PERMUTATION_INT("PERMUTATION_VIEWTRANSMITTANCE", 2);
 	class FPPLL : SHADER_PERMUTATION_INT("PERMUTATION_PPLL", 2);
-	using FPermutationDomain = TShaderPermutationDomain<FCoverage, FPPLL>;
+	using FPermutationDomain = TShaderPermutationDomain<FViewTransmittance, FPPLL>;
 
 	BEGIN_SHADER_PARAMETER_STRUCT(FParameters, )
 		SHADER_PARAMETER_RDG_TEXTURE(Texture2D, HairVisibilityDepthTexture)
 		SHADER_PARAMETER_RDG_TEXTURE(Texture2D, CategorisationTexture)
-		SHADER_PARAMETER_RDG_TEXTURE(Texture2D, CoverageTexture)
+		SHADER_PARAMETER_RDG_TEXTURE(Texture2D, ViewTransmittanceTexture)
 		RENDER_TARGET_BINDING_SLOTS()
 	END_SHADER_PARAMETER_STRUCT()
 
@@ -1164,7 +1165,7 @@ public:
 	{
 		if (PermutationVector.Get<FPPLL>() > 0)
 		{
-			PermutationVector.Set<FCoverage>(0);
+			PermutationVector.Set<FViewTransmittance>(0);
 		}
 		return PermutationVector;
 	}
@@ -1172,7 +1173,7 @@ public:
 	static bool ShouldCompilePermutation(const FGlobalShaderPermutationParameters& Parameters)
 	{
 		FPermutationDomain PermutationVector(Parameters.PermutationId);
-		if (PermutationVector.Get<FPPLL>() > 0 && PermutationVector.Get<FCoverage>() > 0)
+		if (PermutationVector.Get<FPPLL>() > 0 && PermutationVector.Get<FViewTransmittance>() > 0)
 		{
 			return false;
 		}
@@ -1188,7 +1189,7 @@ static void AddHairVisibilityColorAndDepthPatchPass(
 	const FViewInfo& View,
 	const FRDGTextureRef& VisibilityDepthTexture,
 	const FRDGTextureRef& CategorisationTexture,
-	const FRDGTextureRef& CoverageTexture,
+	const FRDGTextureRef& ViewTransmittanceTexture,
 	FRDGTextureRef& OutGBufferBTexture,
 	FRDGTextureRef& OutColorTexture,
 	FRDGTextureRef& OutDepthTexture)
@@ -1205,7 +1206,7 @@ static void AddHairVisibilityColorAndDepthPatchPass(
 	FHairVisibilityDepthPS::FParameters* Parameters = GraphBuilder.AllocParameters<FHairVisibilityDepthPS::FParameters>();
 	Parameters->HairVisibilityDepthTexture = VisibilityDepthTexture;
 	Parameters->CategorisationTexture = CategorisationTexture;
-	Parameters->CoverageTexture = CoverageTexture;
+	Parameters->ViewTransmittanceTexture = ViewTransmittanceTexture;
 	Parameters->RenderTargets[0] = FRenderTargetBinding(OutGBufferBTexture, ERenderTargetLoadAction::ELoad);
 	Parameters->RenderTargets[1] = FRenderTargetBinding(OutColorTexture, ERenderTargetLoadAction::ELoad);
 	Parameters->RenderTargets.DepthStencil = FDepthStencilBinding(
@@ -1216,7 +1217,7 @@ static void AddHairVisibilityColorAndDepthPatchPass(
 
 	TShaderMapRef<FPostProcessVS> VertexShader(View.ShaderMap);
 	FHairVisibilityDepthPS::FPermutationDomain PermutationVector;
-	PermutationVector.Set<FHairVisibilityDepthPS::FCoverage>(CoverageTexture ? 1 : 0);
+	PermutationVector.Set<FHairVisibilityDepthPS::FViewTransmittance>(ViewTransmittanceTexture ? 1 : 0);
 	PermutationVector.Set<FHairVisibilityDepthPS::FPPLL>(bUsePPLL ? 1 : 0);
 	PermutationVector = FHairVisibilityDepthPS::RemapPermutation(PermutationVector);
 	TShaderMapRef<FHairVisibilityDepthPS> PixelShader(View.ShaderMap, PermutationVector);
@@ -1308,11 +1309,11 @@ FHairStrandsVisibilityViews RenderHairStrandsVisibilityBuffer(
 			
 			const EHairVisibilityRenderMode RenderMode = GetHairVisibilityRenderMode();
 
-			// Run the coverage pass if needed (not in PPLL mode that is already a high quality render path)
-			FRDGTextureRef CoverageTexture = nullptr;
-			if (GHairStrandsCoveragePassEnable > 0 && RenderMode != HairVisibilityRenderMode_PPLL)
+			// Run the view transmittance pass if needed (not in PPLL mode that is already a high quality render path)
+			FRDGTextureRef ViewTransmittanceTexture = nullptr;
+			if (GHairStrandsViewTransmittancePassEnable > 0 && RenderMode != HairVisibilityRenderMode_PPLL)
 			{
-				CoverageTexture = AddHairCoveragePass(
+				ViewTransmittanceTexture = AddHairViewTransmittancePass(
 					GraphBuilder,
 					Scene,
 					&View,
@@ -1365,7 +1366,7 @@ FHairStrandsVisibilityViews RenderHairStrandsVisibilityBuffer(
 					PassParameters->MSAA_MaterialTexture = MsaaVisibilityResources.MaterialTexture;
 					PassParameters->MSAA_AttributeTexture = MsaaVisibilityResources.AttributeTexture;
 					PassParameters->MSAA_VelocityTexture = MsaaVisibilityResources.VelocityTexture;
-					PassParameters->CoverageTexture = CoverageTexture;
+					PassParameters->ViewTransmittanceTexture = ViewTransmittanceTexture;
 
 					FRDGTextureRef CompactNodeIndex;
 					FRDGBufferRef CompactNodeData;
@@ -1402,7 +1403,7 @@ FHairStrandsVisibilityViews RenderHairStrandsVisibilityBuffer(
 					View,
 					MsaaVisibilityResources.DepthTexture,
 					nullptr,
-					CoverageTexture,
+					ViewTransmittanceTexture,
 					SceneGBufferBTexture,
 					SceneColorTexture,
 					SceneDepthTexture);
@@ -1429,7 +1430,7 @@ FHairStrandsVisibilityViews RenderHairStrandsVisibilityBuffer(
 					PassParameters->PPLLCounter  = PPLLNodeCounterTexture;
 					PassParameters->PPLLNodeIndex= PPLLNodeIndexTexture;
 					PassParameters->PPLLNodeData = GraphBuilder.CreateSRV(PPLLNodeDataBuffer);
-					PassParameters->CoverageTexture = CoverageTexture;
+					PassParameters->ViewTransmittanceTexture = ViewTransmittanceTexture;
 
 					FRDGTextureRef CompactNodeIndex;
 					FRDGBufferRef CompactNodeData;
@@ -1461,7 +1462,7 @@ FHairStrandsVisibilityViews RenderHairStrandsVisibilityBuffer(
 					View,
 					nullptr,
 					CategorizationTexture,
-					CoverageTexture,
+					ViewTransmittanceTexture,
 					SceneGBufferBTexture,
 					SceneColorTexture,
 					SceneDepthTexture);
