Simplified opaque scene color hair patching using Categorisation texture only on all the cases. No need to read msaa source now.

diff --git a/Engine/Shaders/Private/HairStrands/HairStrandsVisibilityDepthPS.usf b/Engine/Shaders/Private/HairStrands/HairStrandsVisibilityDepthPS.usf
--- a/Engine/Shaders/Private/HairStrands/HairStrandsVisibilityDepthPS.usf
+++ b/Engine/Shaders/Private/HairStrands/HairStrandsVisibilityDepthPS.usf
@@ -7,16 +7,7 @@
 #include "../DeferredShadingCommon.ush"
 #include "HairStrandsVisibilityCommon.ush"
 
-
-#if PERMUTATION_PPLL
-	Texture2D<uint4> CategorisationTexture;
-#else
-	Texture2DMS<float, HAIR_VISIBILITY_MAXSAMPLECOUNT> HairVisibilityDepthTexture;
-	#if PERMUTATION_VIEWTRANSMITTANCE
-		Texture2D<float> ViewTransmittanceTexture;
-	#endif
-#endif
-
+Texture2D<uint4> CategorisationTexture;
 
 void MainPS(
 	in FScreenVertexOutput Input,
@@ -25,50 +16,13 @@ void MainPS(
 	out float4 OutColor		: SV_Target1
 )
 {
-#if PERMUTATION_PPLL
-
 	const FCategorizationData CategorisationData = DecodeCategorizationData(CategorisationTexture.Load(uint3(Input.Position.xy, 0)));
 	if (CategorisationData.PixelCoverage < 1.0f)
 	{
 		discard;
 	}
-	SvDepth = CategorisationData.ClosestDepth;
-
-#else // PERMUTATION_PPLL
-
-	// HAS_INVERTED_Z_BUFFER
-	float ClosestDepth = 0;
-	float FarthestDepth = 1;
-
-	const uint2 PixelCoord = floor(Input.Position.xy);
-	uint ValidSampleCount = 0;
-	for (uint SampleIt=0; SampleIt<HAIR_VISIBILITY_MAXSAMPLECOUNT; ++SampleIt)
-	{
-		const float SampleDepth = HairVisibilityDepthTexture.Load(PixelCoord, SampleIt).x;
-		FarthestDepth = min(FarthestDepth, SampleDepth);
-		ClosestDepth  = max(ClosestDepth, SampleDepth);
-		const bool bIsValid = SampleDepth > 0.0f;
-		ValidSampleCount += bIsValid ? 1 : 0;
-	}
-
-#if PERMUTATION_VIEWTRANSMITTANCE == 1
-	if ((1.0f - ViewTransmittanceTexture.Load(uint3(PixelCoord, 0))) < 1.0f || ValidSampleCount == 0)
-	{
-		discard; // If coverage is 1. we skip writting depth to match all the other passes relying on coverage==1 meaning opaque.
-	}
-	SvDepth = ClosestDepth; // Keep closest sample (some mayb be far plane but we only want hair samples)
-#else
-	if (ValidSampleCount < HAIR_VISIBILITY_MAXSAMPLECOUNT)
-	{
-		discard; // If we do not have all samples we skip to ignore depth from far clip plane
-	}
-	SvDepth = FarthestDepth; // Farthest is fine to keep because all samples contains hair data
-#endif
-
-#endif // PERMUTATION_PPLL
 
+	SvDepth = CategorisationData.ClosestDepth;
 	OutGBufferB = 0;
 	OutColor = 0;
 }
-
-
diff --git a/Engine/Source/Runtime/Renderer/Private/HairStrands/HairStrandsVisibility.cpp b/Engine/Source/Runtime/Renderer/Private/HairStrands/HairStrandsVisibility.cpp
--- a/Engine/Source/Runtime/Renderer/Private/HairStrands/HairStrandsVisibility.cpp
+++ b/Engine/Source/Runtime/Renderer/Private/HairStrands/HairStrandsVisibility.cpp
@@ -1148,14 +1148,10 @@ class FHairVisibilityDepthPS : public FGlobalShader
 	DECLARE_GLOBAL_SHADER(FHairVisibilityDepthPS);
 	SHADER_USE_PARAMETER_STRUCT(FHairVisibilityDepthPS, FGlobalShader);
 
-	class FViewTransmittance : SHADER_PERMUTATION_INT("PERMUTATION_VIEWTRANSMITTANCE", 2);
-	class FPPLL : SHADER_PERMUTATION_INT("PERMUTATION_PPLL", 2);
-	using FPermutationDomain = TShaderPermutationDomain<FViewTransmittance, FPPLL>;
+	using FPermutationDomain = TShaderPermutationDomain<>;
 
 	BEGIN_SHADER_PARAMETER_STRUCT(FParameters, )
-		SHADER_PARAMETER_RDG_TEXTURE(Texture2D, HairVisibilityDepthTexture)
 		SHADER_PARAMETER_RDG_TEXTURE(Texture2D, CategorisationTexture)
-		SHADER_PARAMETER_RDG_TEXTURE(Texture2D, ViewTransmittanceTexture)
 		RENDER_TARGET_BINDING_SLOTS()
 	END_SHADER_PARAMETER_STRUCT()
 
@@ -1163,20 +1159,11 @@ public:
 		
 	static FPermutationDomain RemapPermutation(FPermutationDomain PermutationVector)
 	{
-		if (PermutationVector.Get<FPPLL>() > 0)
-		{
-			PermutationVector.Set<FViewTransmittance>(0);
-		}
 		return PermutationVector;
 	}
 
 	static bool ShouldCompilePermutation(const FGlobalShaderPermutationParameters& Parameters)
 	{
-		FPermutationDomain PermutationVector(Parameters.PermutationId);
-		if (PermutationVector.Get<FPPLL>() > 0 && PermutationVector.Get<FViewTransmittance>() > 0)
-		{
-			return false;
-		}
 		return IsHairStrandsSupported(Parameters.Platform);
 	}
 };
@@ -1184,29 +1171,15 @@ public:
 IMPLEMENT_GLOBAL_SHADER(FHairVisibilityDepthPS, "/Engine/Private/HairStrands/HairStrandsVisibilityDepthPS.usf", "MainPS", SF_Pixel);
 
 static void AddHairVisibilityColorAndDepthPatchPass(
-	const bool bUsePPLL,
 	FRDGBuilder& GraphBuilder,
 	const FViewInfo& View,
-	const FRDGTextureRef& VisibilityDepthTexture,
 	const FRDGTextureRef& CategorisationTexture,
-	const FRDGTextureRef& ViewTransmittanceTexture,
 	FRDGTextureRef& OutGBufferBTexture,
 	FRDGTextureRef& OutColorTexture,
 	FRDGTextureRef& OutDepthTexture)
 {
-	if (bUsePPLL)
-	{
-		check(CategorisationTexture);
-	}
-	else
-	{
-		check(VisibilityDepthTexture);
-	}
-
 	FHairVisibilityDepthPS::FParameters* Parameters = GraphBuilder.AllocParameters<FHairVisibilityDepthPS::FParameters>();
-	Parameters->HairVisibilityDepthTexture = VisibilityDepthTexture;
 	Parameters->CategorisationTexture = CategorisationTexture;
-	Parameters->ViewTransmittanceTexture = ViewTransmittanceTexture;
 	Parameters->RenderTargets[0] = FRenderTargetBinding(OutGBufferBTexture, ERenderTargetLoadAction::ELoad);
 	Parameters->RenderTargets[1] = FRenderTargetBinding(OutColorTexture, ERenderTargetLoadAction::ELoad);
 	Parameters->RenderTargets.DepthStencil = FDepthStencilBinding(
@@ -1217,8 +1190,6 @@ static void AddHairVisibilityColorAndDepthPatchPass(
 
 	TShaderMapRef<FPostProcessVS> VertexShader(View.ShaderMap);
 	FHairVisibilityDepthPS::FPermutationDomain PermutationVector;
-	PermutationVector.Set<FHairVisibilityDepthPS::FViewTransmittance>(ViewTransmittanceTexture ? 1 : 0);
-	PermutationVector.Set<FHairVisibilityDepthPS::FPPLL>(bUsePPLL ? 1 : 0);
 	PermutationVector = FHairVisibilityDepthPS::RemapPermutation(PermutationVector);
 	TShaderMapRef<FHairVisibilityDepthPS> PixelShader(View.ShaderMap, PermutationVector);
 	const TShaderMap<FGlobalShaderType>* GlobalShaderMap = View.ShaderMap;
@@ -1358,7 +1329,8 @@ FHairStrandsVisibilityViews RenderHairStrandsVisibilityBuffer(
 				GraphBuilder.QueueTextureExtraction(MsaaVisibilityResources.AttributeTexture, &VisibilityData.AttributeTexture);
 				GraphBuilder.QueueTextureExtraction(MsaaVisibilityResources.VelocityTexture, &VisibilityData.VelocityTexture);
 				GraphBuilder.QueueTextureExtraction(MsaaVisibilityResources.DepthTexture, &VisibilityData.DepthTexture);
-				
+
+				FRDGTextureRef CategorizationTexture;
 				{
 					FHairVisibilityPrimitiveIdCompactionCS::FParameters* PassParameters = GraphBuilder.AllocParameters<FHairVisibilityPrimitiveIdCompactionCS::FParameters>();
 					PassParameters->MSAA_DepthTexture = MsaaVisibilityResources.DepthTexture;
@@ -1371,7 +1343,6 @@ FHairStrandsVisibilityViews RenderHairStrandsVisibilityBuffer(
 					FRDGTextureRef CompactNodeIndex;
 					FRDGBufferRef CompactNodeData;
 					FRDGBufferRef CompactNodeCoord;
-					FRDGTextureRef CategorizationTexture;
 					FRDGBufferRef IndirectArgsBuffer;
 					AddHairVisibilityPrimitiveIdCompactionPass(
 						false, // bUsePPLL
@@ -1398,12 +1369,9 @@ FHairStrandsVisibilityViews RenderHairStrandsVisibilityBuffer(
 				// * closest depth
 				// * unlit shading model ID 
 				AddHairVisibilityColorAndDepthPatchPass(
-					false, // bUsePPLL
 					GraphBuilder,
 					View,
-					MsaaVisibilityResources.DepthTexture,
-					nullptr,
-					ViewTransmittanceTexture,
+					CategorizationTexture,
 					SceneGBufferBTexture,
 					SceneColorTexture,
 					SceneDepthTexture);
@@ -1457,12 +1425,9 @@ FHairStrandsVisibilityViews RenderHairStrandsVisibilityBuffer(
 				}
 
 				AddHairVisibilityColorAndDepthPatchPass(
-					true, // bUsePPLL
 					GraphBuilder,
 					View,
-					nullptr,
 					CategorizationTexture,
-					ViewTransmittanceTexture,
 					SceneGBufferBTexture,
 					SceneColorTexture,
 					SceneDepthTexture);
