Fixed selection outline still being rendered in the viewport after deselecting "Selection" from the Show -> Advanced menu.

PostProcessSelectionOutline used to check the Selection engine show flag and do nothing if it's false, but this was lost when the system was converted to RDG (there's a separate SelectionOutline flag wich was being checked in PostProcessing.cpp, but it's not sufficient).

diff --git a/Engine/Source/Runtime/Renderer/Private/PostProcess/PostProcessing.cpp b/Engine/Source/Runtime/Renderer/Private/PostProcess/PostProcessing.cpp
--- a/Engine/Source/Runtime/Renderer/Private/PostProcess/PostProcessing.cpp
+++ b/Engine/Source/Runtime/Renderer/Private/PostProcess/PostProcessing.cpp
@@ -292,7 +292,7 @@ void AddPostProcessingPasses(FRDGBuilder& GraphBuilder, const FViewInfo& View, c
 	PassSequence.SetEnabled(EPass::VisualizeStationaryLightOverlap, EngineShowFlags.StationaryLightOverlap);
 	PassSequence.SetEnabled(EPass::VisualizeLightCulling, EngineShowFlags.VisualizeLightCulling);
 #if WITH_EDITOR
-	PassSequence.SetEnabled(EPass::SelectionOutline, GIsEditor && EngineShowFlags.SelectionOutline && !EngineShowFlags.Wireframe && !bVisualizeHDR);
+	PassSequence.SetEnabled(EPass::SelectionOutline, GIsEditor && EngineShowFlags.Selection && EngineShowFlags.SelectionOutline && !EngineShowFlags.Wireframe && !bVisualizeHDR);
 	PassSequence.SetEnabled(EPass::EditorPrimitive, FSceneRenderer::ShouldCompositeEditorPrimitives(View));
 #else
 	PassSequence.SetEnabled(EPass::SelectionOutline, false);
@@ -1672,6 +1672,7 @@ void FPostProcessing::ProcessES2(FRHICommandListImmediate& RHICmdList, FScene* S
 		// Show the selection outline if it is in the editor and we aren't in wireframe 
 		// If the engine is in demo mode and game view is on we also do not show the selection outline
 		if ( GIsEditor
+			&& View.Family->EngineShowFlags.Selection
 			&& View.Family->EngineShowFlags.SelectionOutline
 			&& !(View.Family->EngineShowFlags.Wireframe)
 			)
