Removed the per instance hacked MaxPPLLNodeCount.
It is now in a global shader parameter struct.
Kept RDG texture in a regular structure as those cannot be added to global struct. Fine.

diff --git a/Engine/Shaders/Private/HairStrands/HairStrandsVisibilityPS.usf b/Engine/Shaders/Private/HairStrands/HairStrandsVisibilityPS.usf
--- a/Engine/Shaders/Private/HairStrands/HairStrandsVisibilityPS.usf
+++ b/Engine/Shaders/Private/HairStrands/HairStrandsVisibilityPS.usf
@@ -28,7 +28,6 @@
 #endif
 
 #if HAIR_RENDER_MODE == RENDER_MODE_PPLL 
-uint HairVisibilityPass_MaxPPLLNodeCount;
 RWTexture2D<uint> PPLLCounter;
 RWTexture2D<uint> PPLLNodeIndex;
 RWStructuredBuffer<FPPLLNodeData> PPLLNodes;
@@ -144,7 +143,7 @@ void MainVisibility(
 	uint NodeIndex;
 	InterlockedAdd(PPLLCounter[uint2(0, 0)], 1, NodeIndex);
 
-	if (NodeIndex < HairVisibilityPass_MaxPPLLNodeCount)
+	if (NodeIndex < VisibilityPassGlobalParameters.MaxPPLLNodeCount)
 	{
 		// If we can append new node, add it and make it point to the node previously in head of list
 		uint NextNodeIndex;
diff --git a/Engine/Source/Runtime/Renderer/Private/HairStrands/HairStrandsVisibility.cpp b/Engine/Source/Runtime/Renderer/Private/HairStrands/HairStrandsVisibility.cpp
--- a/Engine/Source/Runtime/Renderer/Private/HairStrands/HairStrandsVisibility.cpp
+++ b/Engine/Source/Runtime/Renderer/Private/HairStrands/HairStrandsVisibility.cpp
@@ -131,8 +131,13 @@ static bool IsCompatibleWithHairVisibility(const FMeshMaterialShaderPermutationP
 	return IsCompatibleWithHairStrands(Parameters.Material, Parameters.Platform);
 }
 
+BEGIN_GLOBAL_SHADER_PARAMETER_STRUCT(FVisibilityPassGlobalParameters, )
+	SHADER_PARAMETER(uint32, MaxPPLLNodeCount)
+END_GLOBAL_SHADER_PARAMETER_STRUCT()
+IMPLEMENT_GLOBAL_SHADER_PARAMETER_STRUCT(FVisibilityPassGlobalParameters, "VisibilityPassGlobalParameters");
+
 BEGIN_SHADER_PARAMETER_STRUCT(FVisibilityPassParameters, )
-	SHADER_PARAMETER(uint32, HairVisibilityPass_MaxPPLLNodeCount)
+	SHADER_PARAMETER_STRUCT_REF(FVisibilityPassGlobalParameters, VisibilityPassGlobalParameters)
 	SHADER_PARAMETER_RDG_TEXTURE_UAV(Texture2D, PPLLCounter)
 	SHADER_PARAMETER_RDG_TEXTURE_UAV(Texture2D, PPLLNodeIndex)
 	SHADER_PARAMETER_RDG_BUFFER_UAV(StructuredBuffer, PPLLNodeData)
@@ -151,6 +156,7 @@ protected:
 	{
 		ERHIFeatureLevel::Type FeatureLevel = GetMaxSupportedFeatureLevel((EShaderPlatform)Initializer.Target.Platform);
 		check(FSceneInterface::GetShadingPath(FeatureLevel) != EShadingPath::Mobile);
+		PassUniformBuffer.Bind(Initializer.ParameterMap, FVisibilityPassGlobalParameters::StaticStructMetadata.GetShaderVariableName());
 	}
 
 	FHairVisibilityVS() {}
@@ -176,9 +182,8 @@ IMPLEMENT_MATERIAL_SHADER_TYPE(template<>, FHairVisibilityVS<HairVisibilityRende
 class FHairVisibilityShaderElementData : public FMeshMaterialShaderElementData
 {
 public:
-	FHairVisibilityShaderElementData(uint32 InHairClusterId, uint32 InMaxPPLLNodeCount) : HairClusterId(InHairClusterId), MaxPPLLNodeCount(InMaxPPLLNodeCount) { }
+	FHairVisibilityShaderElementData(uint32 InHairClusterId) : HairClusterId(InHairClusterId) { }
 	uint32 HairClusterId;
-	uint32 MaxPPLLNodeCount;
 };
 
 template<EHairVisibilityRenderMode RenderMode>
@@ -194,7 +199,6 @@ public:
 		ERHIFeatureLevel::Type FeatureLevel = GetMaxSupportedFeatureLevel((EShaderPlatform)Initializer.Target.Platform);
 		check(FSceneInterface::GetShadingPath(FeatureLevel) != EShadingPath::Mobile);
 		HairVisibilityPass_HairClusterIndex.Bind(Initializer.ParameterMap, TEXT("HairVisibilityPass_HairClusterIndex"));
-		HairVisibilityPass_MaxPPLLNodeCount.Bind(Initializer.ParameterMap, TEXT("HairVisibilityPass_MaxPPLLNodeCount"));
 		PPLLCounter.Bind(Initializer.ParameterMap, TEXT("PPLLCounter"));
 		PPLLNodeIndex.Bind(Initializer.ParameterMap, TEXT("PPLLNodeIndex"));
 		PPLLNodes.Bind(Initializer.ParameterMap, TEXT("PPLLNodes"));
@@ -218,7 +222,6 @@ public:
 	{
 		bool bShaderHasOutdatedParameters = FMeshMaterialShader::Serialize(Ar);
 		Ar << HairVisibilityPass_HairClusterIndex;
-		Ar << HairVisibilityPass_MaxPPLLNodeCount;
 		Ar << PPLLCounter;
 		Ar << PPLLNodeIndex;
 		Ar << PPLLNodes;
@@ -237,11 +240,9 @@ public:
 	{
 		FMeshMaterialShader::GetShaderBindings(Scene, FeatureLevel, PrimitiveSceneProxy, MaterialRenderProxy, Material, DrawRenderState, ShaderElementData, ShaderBindings);
 		ShaderBindings.Add(HairVisibilityPass_HairClusterIndex, ShaderElementData.HairClusterId);
-		ShaderBindings.Add(HairVisibilityPass_MaxPPLLNodeCount, ShaderElementData.MaxPPLLNodeCount);
 	}
 
 	FShaderParameter HairVisibilityPass_HairClusterIndex;
-	FShaderParameter HairVisibilityPass_MaxPPLLNodeCount;
 	// This is on FVisibilityPassParameters but needed to avoid the compiler to assert with unbound parameters 
 	FShaderParameter PPLLCounter;
 	FShaderParameter PPLLNodeIndex;
@@ -344,7 +345,7 @@ void FHairVisibilityProcessor::Process(
 	}
 
 	FMeshPassProcessorRenderState DrawRenderState(PassDrawRenderState);
-	FHairVisibilityShaderElementData ShaderElementData(HairClusterId, MaxPPLLNodeCount);
+	FHairVisibilityShaderElementData ShaderElementData(HairClusterId);
 	ShaderElementData.InitializeMeshMaterialData(ViewIfDynamicMeshCommand, PrimitiveSceneProxy, MeshBatch, StaticMeshId, false);
 
 	BuildMeshDrawCommands(
@@ -851,14 +852,13 @@ static FRDGTextureRef AddHairVisibilityFillOpaqueDepth(
 
 ///////////////////////////////////////////////////////////////////////////////////////////////////
 
-template <typename TPassParameters>
 static void AddHairVisibilityCommonPass(
 	FRDGBuilder& GraphBuilder,
 	const FScene* Scene,
 	const FViewInfo* ViewInfo,
 	const FHairStrandsClusterDatas& ClusterDatas,
 	const EHairVisibilityRenderMode RenderMode,
-	TPassParameters* PassParameters)
+	FVisibilityPassParameters* PassParameters)
 {
 	GraphBuilder.AddPass(
 		RenderMode == HairVisibilityRenderMode_PPLL ? RDG_EVENT_NAME("HairStrandsPPLLPass") : 
@@ -907,12 +907,6 @@ static void AddHairVisibilityCommonPass(
 			FGraphicsMinimalPipelineStateSet PipelineStateSet;
 			FDynamicPassMeshDrawListContext ShadowContext(DynamicMeshDrawCommandStorage, VisibleMeshDrawCommands, PipelineStateSet);
 			FHairVisibilityProcessor MeshProcessor(Scene, ViewInfo, DrawRenderState, RenderMode, &ShadowContext);
-			if (RenderMode == HairVisibilityRenderMode_PPLL)
-			{
-				// Work around because the value on the PassParamter is not taken into account for the PS global constant buffer.
-				// TODO fix that by having FVisibilityPassParameters be a global parameter structure.
-				MeshProcessor.SetMaxPPLLNodeCount(PassParameters->HairVisibilityPass_MaxPPLLNodeCount);
-			}
 
 			for (const FHairStrandsClusterData& ClusterData : ClusterDatas.Datas)
 			{
@@ -1092,11 +1086,15 @@ static void AddHairVisibilityPPLLPass(
 	AddClearUAVPass(GraphBuilder, RDG_EVENT_NAME("ClearHairVisibilityPPLLCounter"), OutVisibilityPPLLNodeCounter, 0);
 	AddClearUAVPass(GraphBuilder, RDG_EVENT_NAME("ClearHairVisibilityPPLLNodeIndex"), OutVisibilityPPLLNodeIndex, 0xFFFFFFFF);
 
+	FVisibilityPassGlobalParameters VisibilityPassGlobalParameters;
+	VisibilityPassGlobalParameters.MaxPPLLNodeCount = PPLLMaxTotalListElementCount;
+	TUniformBufferRef<FVisibilityPassGlobalParameters> VisibilityPassGlobalParametersUB = TUniformBufferRef<FVisibilityPassGlobalParameters>::CreateUniformBufferImmediate(VisibilityPassGlobalParameters, UniformBuffer_SingleFrame);
+
 	FVisibilityPassParameters* PassParameters = GraphBuilder.AllocParameters<FVisibilityPassParameters>();
+	PassParameters->VisibilityPassGlobalParameters = VisibilityPassGlobalParametersUB;
 	PassParameters->PPLLCounter = GraphBuilder.CreateUAV(FRDGTextureUAVDesc(OutVisibilityPPLLNodeCounter, 0));
 	PassParameters->PPLLNodeIndex = GraphBuilder.CreateUAV(FRDGTextureUAVDesc(OutVisibilityPPLLNodeIndex, 0));
 	PassParameters->PPLLNodeData = GraphBuilder.CreateUAV(FRDGBufferUAVDesc(OutVisibilityPPLLNodeData));
-	PassParameters->HairVisibilityPass_MaxPPLLNodeCount = PPLLMaxTotalListElementCount;
 	PassParameters->RenderTargets.DepthStencil = FDepthStencilBinding(InViewZDepthTexture, ERenderTargetLoadAction::ELoad, ERenderTargetLoadAction::ENoAction, FExclusiveDepthStencil::DepthRead_StencilNop);
 	AddHairVisibilityCommonPass(GraphBuilder, Scene, ViewInfo, ClusterDatas, HairVisibilityRenderMode_PPLL, PassParameters);
 }
