Setting InnerSplitIndex to -1 for local lights to avoid local lights doing cascade calculations with uninitialized planes.

diff --git a/Engine/Source/Runtime/Renderer/Private/VolumetricFog.cpp b/Engine/Source/Runtime/Renderer/Private/VolumetricFog.cpp
--- a/Engine/Source/Runtime/Renderer/Private/VolumetricFog.cpp
+++ b/Engine/Source/Runtime/Renderer/Private/VolumetricFog.cpp
@@ -312,7 +312,7 @@ public:
 		SetShaderValue(RHICmdList, ShaderRHI, PhaseG, FogInfo.VolumetricFogScatteringDistribution);
 		SetShaderValue(RHICmdList, ShaderRHI, InverseSquaredLightDistanceBiasScale, GInverseSquaredLightDistanceBiasScale);
 
-		VolumeShadowingParameters.Set(RHICmdList, ShaderRHI, View, LightSceneInfo, ShadowMap, 0, bDynamicallyShadowed);
+		VolumeShadowingParameters.Set(RHICmdList, ShaderRHI, View, LightSceneInfo, ShadowMap, -1, bDynamicallyShadowed);
 	}
 
 	virtual bool Serialize(FArchive& Ar) override
