Added comment for VolumeShadowingParameters.Set

#rnx

diff --git a/Engine/Source/Runtime/Renderer/Private/VolumetricFog.cpp b/Engine/Source/Runtime/Renderer/Private/VolumetricFog.cpp
--- a/Engine/Source/Runtime/Renderer/Private/VolumetricFog.cpp
+++ b/Engine/Source/Runtime/Renderer/Private/VolumetricFog.cpp
@@ -312,6 +312,7 @@ public:
 		SetShaderValue(RHICmdList, ShaderRHI, PhaseG, FogInfo.VolumetricFogScatteringDistribution);
 		SetShaderValue(RHICmdList, ShaderRHI, InverseSquaredLightDistanceBiasScale, GInverseSquaredLightDistanceBiasScale);
 
+		// Pass -1 to InnerSplitIndex to avoid doing unecessary and uninitialized shadow cascade calculations for local lights 
 		VolumeShadowingParameters.Set(RHICmdList, ShaderRHI, View, LightSceneInfo, ShadowMap, -1, bDynamicallyShadowed);
 	}
 
