Fixed bug with Dithered LOD when EarlyZPassOnlyMaterialMasking is set.

diff --git a/Engine/Source/Runtime/Renderer/Private/BasePassRendering.cpp b/Engine/Source/Runtime/Renderer/Private/BasePassRendering.cpp
--- a/Engine/Source/Runtime/Renderer/Private/BasePassRendering.cpp
+++ b/Engine/Source/Runtime/Renderer/Private/BasePassRendering.cpp
@@ -127,31 +127,6 @@ IMPLEMENT_BASEPASS_LIGHTMAPPED_SHADER_TYPE( TUniformLightMapPolicy<LMP_DISTANCE_
 
 DECLARE_GPU_STAT(Basepass);
 
-void SetBasePassDitheredLODTransitionState(const FSceneView* SceneView, const FMeshBatch& RESTRICT Mesh, int32 StaticMeshId, FMeshPassProcessorRenderState& DrawRenderState)
-{
-	if (SceneView && StaticMeshId >= 0 && Mesh.bDitheredLODTransition)
-	{
-		checkSlow(SceneView->bIsViewInfo);
-		const FViewInfo* ViewInfo = (FViewInfo*)SceneView;
-
-		if (ViewInfo->bAllowStencilDither)
-		{
-			if (ViewInfo->StaticMeshFadeOutDitheredLODMap[StaticMeshId] || ViewInfo->StaticMeshFadeInDitheredLODMap[StaticMeshId])
-			{
-				const uint32 RestoreStencilRef = DrawRenderState.GetStencilRef();
-				DrawRenderState.SetDepthStencilState(
-					TStaticDepthStencilState<
-					false, CF_Equal,
-					true, CF_Always, SO_Keep, SO_Keep, SO_Replace,
-					false, CF_Always, SO_Keep, SO_Keep, SO_Keep,
-					0xFF, GET_STENCIL_BIT_MASK(RECEIVE_DECAL, 1) | STENCIL_LIGHTING_CHANNELS_MASK(0x7)
-					>::GetRHI());
-				DrawRenderState.SetStencilRef(RestoreStencilRef);
-			}
-		}
-	}
-}
-
 void SetTranslucentRenderState(FMeshPassProcessorRenderState& DrawRenderState, const FMaterial& Material)
 {
 	switch (Material.GetBlendMode())
@@ -255,25 +230,26 @@ FMeshDrawCommandSortKey CalculateBasePassMeshStaticSortKey(EDepthDrawingMode Ear
 	return SortKey;
 }
 
-void SetDepthStencilStateForBasePass(FMeshPassProcessorRenderState& DrawRenderState, ERHIFeatureLevel::Type FeatureLevel, const FMeshBatch& Mesh, const FPrimitiveSceneProxy* PrimitiveSceneProxy, bool bEnableReceiveDecalOutput, bool bUseDebugViewPS, FRHIDepthStencilState* LodFadeOverrideDepthStencilState)
+void SetDepthStencilStateForBasePass(
+	const FSceneView* ViewIfDynamicMeshCommand,
+	FMeshPassProcessorRenderState& DrawRenderState,
+	ERHIFeatureLevel::Type FeatureLevel,
+	const FMeshBatch& Mesh,
+	int32 StaticMeshId,
+	const FPrimitiveSceneProxy* PrimitiveSceneProxy,
+	bool bEnableReceiveDecalOutput)
 {
 	static IConsoleVariable* EarlyZPassOnlyMaterialMaskingCVar = IConsoleManager::Get().FindConsoleVariable(TEXT("r.EarlyZPassOnlyMaterialMasking"));
-	bool bMaskInEarlyPass = (EarlyZPassOnlyMaterialMaskingCVar && Mesh.MaterialRenderProxy->GetMaterial(FeatureLevel)->IsMasked() && EarlyZPassOnlyMaterialMaskingCVar->GetInt());
+	const bool bMaskInEarlyPass = (Mesh.MaterialRenderProxy->GetMaterial(FeatureLevel)->IsMasked() || Mesh.bDitheredLODTransition) && EarlyZPassOnlyMaterialMaskingCVar->GetInt();
 
-	if (bEnableReceiveDecalOutput && !bUseDebugViewPS)
+	if (bEnableReceiveDecalOutput)
 	{
 		// Set stencil value for this draw call
 		// This is effectively extending the GBuffer using the stencil bits
 		const uint8 StencilValue = GET_STENCIL_BIT_MASK(RECEIVE_DECAL, PrimitiveSceneProxy ? !!PrimitiveSceneProxy->ReceivesDecals() : 0x00)
 			| STENCIL_LIGHTING_CHANNELS_MASK(PrimitiveSceneProxy ? PrimitiveSceneProxy->GetLightingChannelStencilValue() : 0x00);
 
-		if (LodFadeOverrideDepthStencilState != nullptr)
-		{
-			//@TODO: Handle bMaskInEarlyPass in this case (used when a LODTransition is specified)
-			DrawRenderState.SetDepthStencilState(LodFadeOverrideDepthStencilState);
-			DrawRenderState.SetStencilRef(StencilValue);
-		}
-		else if (bMaskInEarlyPass)
+		if (bMaskInEarlyPass)
 		{
 			DrawRenderState.SetDepthStencilState(TStaticDepthStencilState<
 				false, CF_Equal,
@@ -308,6 +284,28 @@ void SetDepthStencilStateForBasePass(FMeshPassProcessorRenderState& DrawRenderSt
 	{
 		DrawRenderState.SetDepthStencilState(TStaticDepthStencilState<false, CF_Equal>::GetRHI());
 	}
+
+	if (ViewIfDynamicMeshCommand && StaticMeshId >= 0 && Mesh.bDitheredLODTransition)
+	{
+		checkSlow(ViewIfDynamicMeshCommand->bIsViewInfo);
+		const FViewInfo* ViewInfo = static_cast<const FViewInfo*>(ViewIfDynamicMeshCommand);
+
+		if (ViewInfo->bAllowStencilDither)
+		{
+			if (ViewInfo->StaticMeshFadeOutDitheredLODMap[StaticMeshId] || ViewInfo->StaticMeshFadeInDitheredLODMap[StaticMeshId])
+			{
+				const uint32 RestoreStencilRef = DrawRenderState.GetStencilRef();
+				DrawRenderState.SetDepthStencilState(
+					TStaticDepthStencilState<
+					false, CF_Equal,
+					true, CF_Always, SO_Keep, SO_Keep, SO_Replace,
+					false, CF_Always, SO_Keep, SO_Keep, SO_Keep,
+					0xFF, GET_STENCIL_BIT_MASK(RECEIVE_DECAL, 1) | STENCIL_LIGHTING_CHANNELS_MASK(0x7)
+					>::GetRHI());
+				DrawRenderState.SetStencilRef(RestoreStencilRef);
+			}
+		}
+	}
 }
 
 void SetupBasePassState(FExclusiveDepthStencil::Type BasePassDepthStencilAccess, const bool bShaderComplexity, FMeshPassProcessorRenderState& DrawRenderState)
@@ -998,15 +996,20 @@ void FBasePassMeshProcessor::Process(
 
 	FMeshPassProcessorRenderState DrawRenderState(PassDrawRenderState);
 
-	SetDepthStencilStateForBasePass(DrawRenderState, FeatureLevel, MeshBatch, PrimitiveSceneProxy, bEnableReceiveDecalOutput, false, nullptr);
+	SetDepthStencilStateForBasePass(
+		ViewIfDynamicMeshCommand,
+		DrawRenderState,
+		FeatureLevel,
+		MeshBatch,
+		StaticMeshId,
+		PrimitiveSceneProxy,
+		bEnableReceiveDecalOutput);
 
 	if (bTranslucentBasePass)
 	{
 		SetTranslucentRenderState(DrawRenderState, MaterialResource);
 	}
 
-	SetBasePassDitheredLODTransitionState(ViewIfDynamicMeshCommand, MeshBatch, StaticMeshId, DrawRenderState);
-
 	TBasePassShaderElementData<LightMapPolicyType> ShaderElementData(LightMapElementData);
 	ShaderElementData.InitializeMeshMaterialData(ViewIfDynamicMeshCommand, PrimitiveSceneProxy, MeshBatch, StaticMeshId, true);
 
diff --git a/Engine/Source/Runtime/Renderer/Private/BasePassRendering.h b/Engine/Source/Runtime/Renderer/Private/BasePassRendering.h
--- a/Engine/Source/Runtime/Renderer/Private/BasePassRendering.h
+++ b/Engine/Source/Runtime/Renderer/Private/BasePassRendering.h
@@ -674,7 +674,6 @@ private:
 
 ENUM_CLASS_FLAGS(FBasePassMeshProcessor::EFlags);
 
-extern void SetDepthStencilStateForBasePass(FMeshPassProcessorRenderState& DrawRenderState, ERHIFeatureLevel::Type FeatureLevel, const FMeshBatch& Mesh, const FPrimitiveSceneProxy* PrimitiveSceneProxy, bool bEnableReceiveDecalOutput, bool bUseDebugViewPS, FRHIDepthStencilState* LodFadeOverrideDepthStencilState);
 extern void SetupBasePassState(FExclusiveDepthStencil::Type BasePassDepthStencilAccess, const bool bShaderComplexity, FMeshPassProcessorRenderState& DrawRenderState);
 extern FMeshDrawCommandSortKey CalculateTranslucentMeshStaticSortKey(const FPrimitiveSceneProxy* RESTRICT PrimitiveSceneProxy, uint16 MeshIdInPrimitive);
 
