Fixed issue where capturing a pre / post tonemap hdr color for movie output would produce black.

diff --git a/Engine/Source/Runtime/Renderer/Private/PostProcess/PostProcessVisualizeBuffer.cpp b/Engine/Source/Runtime/Renderer/Private/PostProcess/PostProcessVisualizeBuffer.cpp
--- a/Engine/Source/Runtime/Renderer/Private/PostProcess/PostProcessVisualizeBuffer.cpp
+++ b/Engine/Source/Runtime/Renderer/Private/PostProcess/PostProcessVisualizeBuffer.cpp
@@ -361,6 +361,8 @@ FScreenPassTexture AddVisualizeGBufferOverviewPass(
 		FPostProcessMaterialInputs PostProcessMaterialInputs;
 		PostProcessMaterialInputs.SetInput(EPostProcessMaterialInput::SceneColor, Inputs.SceneColor);
 		PostProcessMaterialInputs.SetInput(EPostProcessMaterialInput::SeparateTranslucency, Inputs.SeparateTranslucency);
+		PostProcessMaterialInputs.SetInput(EPostProcessMaterialInput::PreTonemapHDRColor, Inputs.SceneColorBeforeTonemap);
+		PostProcessMaterialInputs.SetInput(EPostProcessMaterialInput::PostTonemapHDRColor, Inputs.SceneColorAfterTonemap);
 		PostProcessMaterialInputs.SetInput(EPostProcessMaterialInput::Velocity, Inputs.Velocity);
 		PostProcessMaterialInputs.OutputFormat = OutputFormat;
 
