Enabled hair transmittance pass by default
Fixed undo cine settings when changing quality settings

diff --git a/Engine/Config/BaseScalability.ini b/Engine/Config/BaseScalability.ini
--- a/Engine/Config/BaseScalability.ini
+++ b/Engine/Config/BaseScalability.ini
@@ -566,19 +566,26 @@ grass.DensityScale=1.0
 ;-----------------------------------------------------------------------------------------------------------------
 
 [ShadingQuality@0]
-;todo
+r.HairStrands.SkyLighting.SampleCount=8
+r.HairStrands.SkyLighting.JitterIntegration=0
+r.HairStrands.SkyAO.SampleCount=8
 
 [ShadingQuality@1]
-;todo
+r.HairStrands.SkyLighting.SampleCount=16
+r.HairStrands.SkyLighting.JitterIntegration=0
+r.HairStrands.SkyAO.SampleCount=16
 
 [ShadingQuality@2]
-;todo
+r.HairStrands.SkyLighting.SampleCount=16
+r.HairStrands.SkyLighting.JitterIntegration=0
+r.HairStrands.SkyAO.SampleCount=16
 
 [ShadingQuality@3]
-;todo
+r.HairStrands.SkyLighting.SampleCount=16
+r.HairStrands.SkyLighting.JitterIntegration=0
+r.HairStrands.SkyAO.SampleCount=16
 
 [ShadingQuality@Cine]
-r.HairStrands.ViewTransmittancePass=1
 r.HairStrands.SkyLighting.SampleCount=64
 r.HairStrands.SkyLighting.JitterIntegration=1
-r.HairStrands.SkyAO.SampleCount=32
+r.HairStrands.SkyAO.SampleCount=32
\ No newline at end of file
diff --git a/Engine/Source/Runtime/Renderer/Private/HairStrands/HairStrandsEnvironment.cpp b/Engine/Source/Runtime/Renderer/Private/HairStrands/HairStrandsEnvironment.cpp
--- a/Engine/Source/Runtime/Renderer/Private/HairStrands/HairStrandsEnvironment.cpp
+++ b/Engine/Source/Runtime/Renderer/Private/HairStrands/HairStrandsEnvironment.cpp
@@ -31,13 +31,13 @@ static float GHairSkylightingConeAngle = 3;
 static FAutoConsoleVariableRef CVarHairSkylightingConeAngle(TEXT("r.HairStrands.SkyLighting.ConeAngle"), GHairSkylightingConeAngle, TEXT("Cone angle for tracing sky lighting on hair."));
 
 static int32 GHairStrandsSkyLightingSampleCount = 16;
-static FAutoConsoleVariableRef CVarHairStrandsSkyLightingSampleCount(TEXT("r.HairStrands.SkyLighting.SampleCount"), GHairStrandsSkyLightingSampleCount, TEXT("Number of samples used for evaluating multiple scattering and visible area (default is set to 16)."));
+static FAutoConsoleVariableRef CVarHairStrandsSkyLightingSampleCount(TEXT("r.HairStrands.SkyLighting.SampleCount"), GHairStrandsSkyLightingSampleCount, TEXT("Number of samples used for evaluating multiple scattering and visible area (default is set to 16)."), ECVF_Scalability | ECVF_RenderThreadSafe);
 
 static int32 GHairStrandsSkyLightingJitterSphericalIntegration = 0;
-static FAutoConsoleVariableRef CVarHairStrandsSkyLightingJitterSphericalIntegration(TEXT("r.HairStrands.SkyLighting.JitterIntegration"), GHairStrandsSkyLightingJitterSphericalIntegration, TEXT("Jitter the sphereical integration for the multiple scattering term. The result is more correct, but noiser as well."));
+static FAutoConsoleVariableRef CVarHairStrandsSkyLightingJitterSphericalIntegration(TEXT("r.HairStrands.SkyLighting.JitterIntegration"), GHairStrandsSkyLightingJitterSphericalIntegration, TEXT("Jitter the sphereical integration for the multiple scattering term. The result is more correct, but noiser as well."), ECVF_Scalability | ECVF_RenderThreadSafe);
 
 static int32 GHairStrandsSkyAOSampleCount = 16;
-static FAutoConsoleVariableRef CVarHairStrandsSkyAOSampleCount(TEXT("r.HairStrands.SkyAO.SampleCount"), GHairStrandsSkyAOSampleCount, TEXT("Number of samples used for evaluating hair AO (default is set to 16)."));
+static FAutoConsoleVariableRef CVarHairStrandsSkyAOSampleCount(TEXT("r.HairStrands.SkyAO.SampleCount"), GHairStrandsSkyAOSampleCount, TEXT("Number of samples used for evaluating hair AO (default is set to 16)."), ECVF_Scalability | ECVF_RenderThreadSafe);
 
 static bool GetHairStrandsSkyLightingEnable() { return GHairSkylightingEnable > 0; }
 static bool GetHairStrandsSkyAOEnable() { return GHairSkyAOEnable > 0; }
diff --git a/Engine/Source/Runtime/Renderer/Private/HairStrands/HairStrandsVisibility.cpp b/Engine/Source/Runtime/Renderer/Private/HairStrands/HairStrandsVisibility.cpp
--- a/Engine/Source/Runtime/Renderer/Private/HairStrands/HairStrandsVisibility.cpp
+++ b/Engine/Source/Runtime/Renderer/Private/HairStrands/HairStrandsVisibility.cpp
@@ -18,7 +18,7 @@ DECLARE_GPU_STAT(HairStrandsVisibility);
 
 /////////////////////////////////////////////////////////////////////////////////////////
 
-static int32 GHairStrandsViewTransmittancePassEnable = 0;
+static int32 GHairStrandsViewTransmittancePassEnable = 1;
 static FAutoConsoleVariableRef CVarHairStrandsCoveragePassEnable(TEXT("r.HairStrands.ViewTransmittancePass"), GHairStrandsViewTransmittancePassEnable, TEXT("Enable accurate transmittance pass for better rendering of small scale hair strand."));
 
 static int32 GHairStrandsMaterialCompactionEnable = 0;
@@ -45,7 +45,7 @@ static int32 GHairVelocityType = 1; // default is
 static FAutoConsoleVariableRef CVarHairVelocityType(TEXT("r.HairStrands.VelocityType"), GHairVelocityType, TEXT("Type of velocity filtering (0:avg, 1:closest, 2:max). Default is 1."));
 
 static int32 GHairVisibilityPPLL = 0;
-static FAutoConsoleVariableRef CVarGHairVisibilityPPLL(TEXT("r.HairStrands.VisibilityPPLL"), GHairVisibilityPPLL, TEXT("Hair Visibility uses per pixel linked list"));
+static FAutoConsoleVariableRef CVarGHairVisibilityPPLL(TEXT("r.HairStrands.VisibilityPPLL"), GHairVisibilityPPLL, TEXT("Hair Visibility uses per pixel linked list"), ECVF_Scalability | ECVF_RenderThreadSafe);
 static int32 GHairVisibilityPPLLMeanListElementCountPerPixel = 16;
 static FAutoConsoleVariableRef CVarGHairVisibilityPPLLMeanListElementCountPerPixel(TEXT("r.HairStrands.VisibilityPPLLMeanListElementCountPerPixel"), GHairVisibilityPPLLMeanListElementCountPerPixel, TEXT("The mean maximum number of node allowed for all linked list element. It will be width*height*VisibilityPPLLMeanListElementCountPerPixel."));
 static int32 GHairVisibilityPPLLMaxRenderNodePerPixel = 16;
