- FShadowCascadeSettings::ShadowSplitIndex now initialized in FDirectionalLightSceneProxy::GetShadowSplitBounds
- Added ensure to catch if this was not done before using the near and far planes to set clipping planes in FVolumeShadowingParameters
- Changed a few -1 to INDEX_NONE

diff --git a/Engine/Source/Runtime/Engine/Private/Components/DirectionalLightComponent.cpp b/Engine/Source/Runtime/Engine/Private/Components/DirectionalLightComponent.cpp
--- a/Engine/Source/Runtime/Engine/Private/Components/DirectionalLightComponent.cpp
+++ b/Engine/Source/Runtime/Engine/Private/Components/DirectionalLightComponent.cpp
@@ -819,6 +819,7 @@ private:
 			OutCascadeSettings->FadePlaneOffset = FadePlane;
 			OutCascadeSettings->FadePlaneLength = SplitFar - FadePlane;
 			OutCascadeSettings->CascadeBiasDistribution = ShadowCascadeBiasDistribution;
+			OutCascadeSettings->ShadowSplitIndex = (int32)ShadowSplitIndex;
 		}
 
 		const FSphere CascadeSphere = FDirectionalLightSceneProxy::GetShadowSplitBoundsDepthRange(View, View.ViewMatrices.GetViewOrigin(), SplitNear, SplitFar, OutCascadeSettings);
diff --git a/Engine/Source/Runtime/Renderer/Private/TranslucentLighting.cpp b/Engine/Source/Runtime/Renderer/Private/TranslucentLighting.cpp
--- a/Engine/Source/Runtime/Renderer/Private/TranslucentLighting.cpp
+++ b/Engine/Source/Runtime/Renderer/Private/TranslucentLighting.cpp
@@ -1520,7 +1520,7 @@ static void InjectTranslucentLightArray(FRHICommandListImmediate& RHICmdList, co
 						{
 							// no shadows
 							SetInjectionShader<LightType_Directional, false>(RHICmdList, GraphicsPSOInit, View, InjectionData.LightFunctionMaterialProxy, LightSceneInfo,
-								InjectionData.ProjectedShadowInfo, -1, VolumeCascadeIndex,
+								InjectionData.ProjectedShadowInfo, INDEX_NONE, VolumeCascadeIndex,
 								*VertexShader, *GeometryShader, InjectionData.bApplyLightFunction, false);
 						}
 					}
@@ -1534,13 +1534,13 @@ static void InjectTranslucentLightArray(FRHICommandListImmediate& RHICmdList, co
 						if (InjectionData.ProjectedShadowInfo)
 						{
 							SetInjectionShader<LightType_Point, true>(RHICmdList, GraphicsPSOInit, View, InjectionData.LightFunctionMaterialProxy, LightSceneInfo,
-								InjectionData.ProjectedShadowInfo, -1, VolumeCascadeIndex,
+								InjectionData.ProjectedShadowInfo, INDEX_NONE, VolumeCascadeIndex,
 								*VertexShader, *GeometryShader, InjectionData.bApplyLightFunction, bInverseSquared);
 						}
 						else
 						{
 							SetInjectionShader<LightType_Point, false>(RHICmdList, GraphicsPSOInit, View, InjectionData.LightFunctionMaterialProxy, LightSceneInfo,
-								InjectionData.ProjectedShadowInfo, -1, VolumeCascadeIndex,
+								InjectionData.ProjectedShadowInfo, INDEX_NONE, VolumeCascadeIndex,
 								*VertexShader, *GeometryShader, InjectionData.bApplyLightFunction, bInverseSquared);
 						}
 					}
diff --git a/Engine/Source/Runtime/Renderer/Private/VolumeLighting.h b/Engine/Source/Runtime/Renderer/Private/VolumeLighting.h
--- a/Engine/Source/Runtime/Renderer/Private/VolumeLighting.h
+++ b/Engine/Source/Runtime/Renderer/Private/VolumeLighting.h
@@ -57,11 +57,12 @@ public:
 		// .zw:DistanceFadeMAD to use MAD for efficiency in the shader, default to ignore the plane
 		FVector4 ShadowInjectParamValue(1, 1, 0, 0);
 
-		if (InnerSplitIndex >= 0)
+		if (InnerSplitIndex != INDEX_NONE)
 		{
 			FShadowCascadeSettings ShadowCascadeSettings;
 
-			LightSceneInfo->Proxy->GetShadowSplitBounds(View, (uint32)InnerSplitIndex, LightSceneInfo->IsPrecomputedLightingValid(), &ShadowCascadeSettings);
+			LightSceneInfo->Proxy->GetShadowSplitBounds(View, InnerSplitIndex, LightSceneInfo->IsPrecomputedLightingValid(), &ShadowCascadeSettings);
+			ensureMsgf(ShadowCascadeSettings.ShadowSplitIndex != INDEX_NONE, TEXT("FLightSceneProxy::GetShadowSplitBounds did not return an initialized ShadowCascadeSettings"));
 
 			// near cascade plane
 			{
diff --git a/Engine/Source/Runtime/Renderer/Private/VolumetricFog.cpp b/Engine/Source/Runtime/Renderer/Private/VolumetricFog.cpp
--- a/Engine/Source/Runtime/Renderer/Private/VolumetricFog.cpp
+++ b/Engine/Source/Runtime/Renderer/Private/VolumetricFog.cpp
@@ -312,8 +312,7 @@ public:
 		SetShaderValue(RHICmdList, ShaderRHI, PhaseG, FogInfo.VolumetricFogScatteringDistribution);
 		SetShaderValue(RHICmdList, ShaderRHI, InverseSquaredLightDistanceBiasScale, GInverseSquaredLightDistanceBiasScale);
 
-		// Pass -1 to InnerSplitIndex to avoid doing unecessary and uninitialized shadow cascade calculations for local lights 
-		VolumeShadowingParameters.Set(RHICmdList, ShaderRHI, View, LightSceneInfo, ShadowMap, -1, bDynamicallyShadowed);
+		VolumeShadowingParameters.Set(RHICmdList, ShaderRHI, View, LightSceneInfo, ShadowMap, INDEX_NONE, bDynamicallyShadowed);
 	}
 
 	virtual bool Serialize(FArchive& Ar) override
