Fixed ray tracing shadows denoiser in multiview mode. Now each view renders in a separate tile and they are later composed into the final texture.

Fixed batch shadows mode was always applying denoiser even if it was disabled.

diff --git a/Engine/Source/Runtime/Renderer/Private/LightRendering.cpp b/Engine/Source/Runtime/Renderer/Private/LightRendering.cpp
--- a/Engine/Source/Runtime/Renderer/Private/LightRendering.cpp
+++ b/Engine/Source/Runtime/Renderer/Private/LightRendering.cpp
@@ -18,6 +18,7 @@
 #include "RayTracing/RaytracingOptions.h"
 #include "SceneTextureParameters.h"
 #include "HairStrands/HairStrandsRendering.h"
+#include "ScreenPass.h"
 
 // ENABLE_DEBUG_DISCARD_PROP is used to test the lighting code by allowing to discard lights to see how performance scales
 // It ought never to be enabled in a shipping build, and is probably only really useful when woring on the shading code.
@@ -1329,7 +1330,7 @@ void FDeferredShadingSceneRenderer::RenderLights(FRHICommandListImmediate& RHICm
 								if (BatchOcclusionType != FLightOcclusionType::Raytraced)
 									continue;
 
-								const bool bRequiresDenoiser = LightRequiresDenosier(BatchLightSceneInfo);
+								const bool bRequiresDenoiser = LightRequiresDenosier(BatchLightSceneInfo) && DenoiserMode > 0;
 
 								IScreenSpaceDenoiser::FShadowRayTracingConfig BatchRayTracingConfig;
 								BatchRayTracingConfig.RayCountPerPixel = BatchLightSceneInfo.Proxy->GetSamplesPerPixel();
@@ -1571,6 +1572,20 @@ void FDeferredShadingSceneRenderer::RenderLights(FRHICommandListImmediate& RHICm
 						FRDGTextureUAV* SubPixelRayTracingShadowMaskUAV = GraphBuilder.CreateUAV(FRDGTextureUAVDesc(SubPixelRayTracingShadowMaskTexture));
 						FRDGTextureUAV* SubPixelRayHitDistanceUAV = GraphBuilder.CreateUAV(FRDGTextureUAVDesc(SubPixelRayDistanceTexture));
 
+						FRDGTextureRef RayTracingShadowMaskTileTexture;
+						{
+							FRDGTextureDesc Desc = FRDGTextureDesc::Create2DDesc(
+								SceneTextures.SceneDepthBuffer->Desc.Extent,
+								PF_FloatRGBA,
+								FClearValueBinding::Black,
+								TexCreate_None,
+								TexCreate_ShaderResource | TexCreate_RenderTargetable | TexCreate_UAV,
+								/* bInForceSeparateTargetAndShaderResource = */ false);
+							RayTracingShadowMaskTileTexture = GraphBuilder.CreateTexture(Desc, TEXT("RayTracingOcclusionTile"));
+						}
+
+						bool bIsMultiview = Views.Num() > 0;
+
 						for (int32 ViewIndex = 0; ViewIndex < Views.Num(); ViewIndex++)
 						{
 							FViewInfo& View = Views[ViewIndex];
@@ -1625,11 +1640,19 @@ void FDeferredShadingSceneRenderer::RenderLights(FRHICommandListImmediate& RHICm
 									InputParameterCount,
 									Outputs);
 
-								GraphBuilder.QueueTextureExtraction(Outputs[0].Mask, &ScreenShadowMaskTexture);
+								if (bIsMultiview)
+								{
+									AddDrawTexturePass(GraphBuilder, View, Outputs[0].Mask, RayTracingShadowMaskTileTexture, View.ViewRect.Min, View.ViewRect.Min, View.ViewRect.Size());
+									GraphBuilder.QueueTextureExtraction(RayTracingShadowMaskTileTexture, &ScreenShadowMaskTexture);
+								}
+								else
+								{
+									GraphBuilder.QueueTextureExtraction(Outputs[0].Mask, &ScreenShadowMaskTexture);
+								}
 							}
 							else
 							{
-									GraphBuilder.QueueTextureExtraction(RayTracingShadowMaskTexture, &ScreenShadowMaskTexture);
+								GraphBuilder.QueueTextureExtraction(RayTracingShadowMaskTexture, &ScreenShadowMaskTexture);
 							}
 
 							// Ray trace the shadow cast by opaque geometries on to hair strands geometries
