Editgrate CL 10058993

#rnx

diff --git a/Engine/Source/Runtime/Renderer/Private/RendererScene.cpp b/Engine/Source/Runtime/Renderer/Private/RendererScene.cpp
--- a/Engine/Source/Runtime/Renderer/Private/RendererScene.cpp
+++ b/Engine/Source/Runtime/Renderer/Private/RendererScene.cpp
@@ -1419,6 +1419,15 @@ void FScene::RemovePrimitiveSceneInfo_RenderThread(FPrimitiveSceneInfo* Primitiv
 
 		{
 			SCOPED_NAMED_EVENT(FScene_DeletePrimitiveSceneInfo, FColor::Red);
+			// Delete the PrimitiveSceneInfo on the game thread after the rendering thread has processed its removal.
+			// This must be done on the game thread because the hit proxy references (and possibly other members) need to be freed on the game thread.
+			struct DeferDeleteHitProxies : FDeferredCleanupInterface
+			{
+				DeferDeleteHitProxies(TArray<TRefCountPtr<HHitProxy>>&& InHitProxies) : HitProxies(MoveTemp(InHitProxies)) {}
+				TArray<TRefCountPtr<HHitProxy>> HitProxies;
+			};
+
+			BeginCleanup(new DeferDeleteHitProxies(MoveTemp(PrimitiveSceneInfo->HitProxies)));
 			delete PrimitiveSceneInfo->Proxy;
 			delete PrimitiveSceneInfo;
 		}
@@ -1450,9 +1459,6 @@ void FScene::RemovePrimitive( UPrimitiveComponent* Primitive )
 		ENQUEUE_RENDER_COMMAND(FRemovePrimitiveCommand)(
 			[Scene, PrimitiveSceneInfo, AttachmentCounter](FRHICommandList&)
 			{
-				PrimitiveSceneInfo->DefaultDynamicHitProxyId = FHitProxyId();
-				PrimitiveSceneInfo->DefaultDynamicHitProxy = nullptr;
-				PrimitiveSceneInfo->HitProxies.Reset();
 				PrimitiveSceneInfo->Proxy->DestroyRenderThreadResources();
 				Scene->RemovePrimitiveSceneInfo_RenderThread(PrimitiveSceneInfo);
 				AttachmentCounter->Decrement();
@@ -3926,6 +3932,16 @@ void FScene::UpdateAllPrimitiveSceneInfos(FRHICommandListImmediate& RHICmdList)
 		SCOPED_NAMED_EVENT(FScene_DeletePrimitiveSceneInfo, FColor::Red);
 		for (FPrimitiveSceneInfo* PrimitiveSceneInfo : DeletedSceneInfos)
 		{
+			// It is possible that hte HitProxies list isn't empty if PrimitiveSceneInfo was Added/Removed in same frame
+			// Delete the PrimitiveSceneInfo on the game thread after the rendering thread has processed its removal.
+			// This must be done on the game thread because the hit proxy references (and possibly other members) need to be freed on the game thread.
+			struct DeferDeleteHitProxies : FDeferredCleanupInterface
+			{
+				DeferDeleteHitProxies(TArray<TRefCountPtr<HHitProxy>>&& InHitProxies) : HitProxies(MoveTemp(InHitProxies)) {}
+				TArray<TRefCountPtr<HHitProxy>> HitProxies;
+			};
+
+			BeginCleanup(new DeferDeleteHitProxies(MoveTemp(PrimitiveSceneInfo->HitProxies)));
 			// free the primitive scene proxy.
 			delete PrimitiveSceneInfo->Proxy;
 			delete PrimitiveSceneInfo;
