Material sorting enabled by default in ray tracing reflections

diff --git a/Engine/Source/Runtime/Renderer/Private/RayTracing/RayTracingReflections.cpp b/Engine/Source/Runtime/Renderer/Private/RayTracing/RayTracingReflections.cpp
--- a/Engine/Source/Runtime/Renderer/Private/RayTracing/RayTracingReflections.cpp
+++ b/Engine/Source/Runtime/Renderer/Private/RayTracing/RayTracingReflections.cpp
@@ -94,10 +94,10 @@ static FAutoConsoleVariableRef CVarRayTracingReflectionsMaxRayDistance(
 
 static TAutoConsoleVariable<int32> CVarRayTracingReflectionsSortMaterials(
 	TEXT("r.RayTracing.Reflections.SortMaterials"),
-	0,
+	1,
 	TEXT("Sets whether refected materials will be sorted before shading\n")
-	TEXT("0: Disabled (Default)\n ")
-	TEXT("1: Enabled, using Trace->Sort->Trace\n"),
+	TEXT("0: Disabled\n ")
+	TEXT("1: Enabled, using Trace->Sort->Trace (Default)\n"),
 	ECVF_RenderThreadSafe);
 
 static TAutoConsoleVariable<int32> CVarRayTracingReflectionsHybrid(
