Try to fix a bug that two-sided shading models are appearing very dark on mobile

diff --git a/Engine/Source/Runtime/Renderer/Private/MobileBasePass.cpp b/Engine/Source/Runtime/Renderer/Private/MobileBasePass.cpp
--- a/Engine/Source/Runtime/Renderer/Private/MobileBasePass.cpp
+++ b/Engine/Source/Runtime/Renderer/Private/MobileBasePass.cpp
@@ -619,7 +619,7 @@ void FMobileBasePassMeshProcessor::Process(
 	
 	//The stationary skylight contribution has been added to the LowQuality Lightmap for StaticMeshActor on mobile, so we should skip the sky light spherical harmonic contribution for it.
 	//Enable skylight if LowQualityLightmaps is disabled or the Lightmap has not been built.
-	bool bEnableSkyLight = ShadingModels.IsLit() && Scene && Scene->ShouldRenderSkylightInBasePass(BlendMode) && !(FReadOnlyCVARCache::Get().bAllowStaticLighting && FReadOnlyCVARCache::Get().bEnableLowQualityLightmaps && PrimitiveSceneProxy && PrimitiveSceneProxy->IsStatic() && MeshBatch.LCI && MeshBatch.LCI->GetLightMapInteraction(FeatureLevel).GetType() == LMIT_Texture);
+	bool bEnableSkyLight = ShadingModels.IsLit() && Scene && Scene->ShouldRenderSkylightInBasePass(BlendMode) && !(FReadOnlyCVARCache::Get().bAllowStaticLighting && FReadOnlyCVARCache::Get().bEnableLowQualityLightmaps && PrimitiveSceneProxy && PrimitiveSceneProxy->IsStatic() && MeshBatch.LCI && MeshBatch.LCI->GetLightMapInteraction(FeatureLevel).GetType() == LMIT_Texture && !MaterialResource.IsTwoSided());
 	int32 NumMovablePointLights = MobileBasePass::CalcNumMovablePointLights(MaterialResource, PrimitiveSceneProxy);
 
 	MobileBasePass::GetShaders(
