Try to fix a bug that indirect lighting missing on mobile platforms

diff --git a/Engine/Source/Runtime/Renderer/Private/MobileBasePass.cpp b/Engine/Source/Runtime/Renderer/Private/MobileBasePass.cpp
--- a/Engine/Source/Runtime/Renderer/Private/MobileBasePass.cpp
+++ b/Engine/Source/Runtime/Renderer/Private/MobileBasePass.cpp
@@ -306,7 +306,9 @@ ELightMapPolicyType MobileBasePass::SelectMeshLightmapPolicy(
 				// Static objects with lightmap type set to ForceVolumetric also need to get GI from ILC
 				|| LightMapInteraction.GetType() == LMIT_GlobalVolume)
 			// if there is no valid ILC data, skip SH calculation by using Movable Directional Light permutation.
-			&& (IndirectLightingCacheAllocation && IndirectLightingCacheAllocation->IsValid())
+			&& (IndirectLightingCacheAllocation && IndirectLightingCacheAllocation->IsValid()
+				// if there is no valid VLM data, skip SH calculation by using Movable Directional Light permutation.
+				|| (Scene && Scene->VolumetricLightmapSceneData.HasData()))
 			// Use the indirect lighting cache shaders if the object does not turn off ILC 
 			&& bPrimitiveUsesILC)
 		{
