TLAS, SBT and other expensive ray tracing routines are now skipped when all Ray tracing passes are disabled

diff --git a/Engine/Source/Runtime/Renderer/Private/DeferredShadingRenderer.cpp b/Engine/Source/Runtime/Renderer/Private/DeferredShadingRenderer.cpp
--- a/Engine/Source/Runtime/Renderer/Private/DeferredShadingRenderer.cpp
+++ b/Engine/Source/Runtime/Renderer/Private/DeferredShadingRenderer.cpp
@@ -594,12 +594,12 @@ void FDeferredShadingSceneRenderer::PrepareDistanceFieldScene(FRHICommandListImm
 
 bool FDeferredShadingSceneRenderer::GatherRayTracingWorldInstances(FRHICommandListImmediate& RHICmdList)
 {
-	if (!IsRayTracingEnabled())
+	if (!IsRayTracingEnabled() || GetForceRayTracingEffectsCVarValue() == 0 || Views.Num() == 0)
 	{
 		return false;
 	}
 
-	if (GetForceRayTracingEffectsCVarValue() == 0 && Views.Num() > 0 && CanOverlayRayTracingOutput(Views[0])) // #dxr_todo: UE-72557 multi-view case
+	if (!AnyRayTracingPassEnabled(Views[0]))
 	{
 		return false;
 	}
@@ -854,12 +854,12 @@ bool FDeferredShadingSceneRenderer::GatherRayTracingWorldInstances(FRHICommandLi
 
 bool FDeferredShadingSceneRenderer::DispatchRayTracingWorldUpdates(FRHICommandListImmediate& RHICmdList)
 {
-	if (!IsRayTracingEnabled())
+	if (!IsRayTracingEnabled() || GetForceRayTracingEffectsCVarValue() == 0 || Views.Num() == 0)
 	{
 		return false;
 	}
 
-	if (GetForceRayTracingEffectsCVarValue() == 0 && Views.Num() > 0 && CanOverlayRayTracingOutput(Views[0])) // #dxr_todo: UE-72557 multi-view case
+	if (!AnyRayTracingPassEnabled(Views[0]))
 	{
 		return false;
 	}
@@ -2741,6 +2741,36 @@ void FDeferredShadingSceneRenderer::CopyStencilToLightingChannelTexture(FRHIComm
 
 #if RHI_RAYTRACING
 
+bool AnyRayTracingPassEnabled(const FViewInfo& View)
+{
+	static auto CVarRayTracingSkyLight = IConsoleManager::Get().FindConsoleVariable(TEXT("r.RayTracing.SkyLight"));
+	static auto CVarRayTracingShadows = IConsoleManager::Get().FindConsoleVariable(TEXT("r.RayTracing.Shadows"));
+	static auto CVarStochasticRectLight = IConsoleManager::Get().FindConsoleVariable(TEXT("r.RayTracing.StochasticRectLight"));
+
+	const bool bRayTracingSkyLight = CVarRayTracingSkyLight != nullptr && CVarRayTracingSkyLight->GetInt() > 0;
+	const bool bRayTracingShadows = CVarRayTracingShadows != nullptr && CVarRayTracingShadows->GetInt() > 0;
+	const bool bRayTracingStochasticRectLight = CVarStochasticRectLight != nullptr && CVarStochasticRectLight->GetInt() > 0;
+
+	if (
+		ShouldRenderRayTracingAmbientOcclusion(View)
+		|| ShouldRenderRayTracingReflections(View)
+		|| ShouldRenderRayTracingGlobalIllumination(View)
+		|| ShouldRenderRayTracingTranslucency(View)
+		|| bRayTracingSkyLight
+		|| bRayTracingShadows
+		|| bRayTracingStochasticRectLight
+		|| View.RayTracingRenderMode == ERayTracingRenderMode::PathTracing
+		|| View.RayTracingRenderMode == ERayTracingRenderMode::RayTracingDebug
+		)
+	{
+		return true;
+	}
+	else
+	{
+		return false;
+	}	
+}
+
 int32 GetForceRayTracingEffectsCVarValue()
 {
 	if (IsRayTracingEnabled())
diff --git a/Engine/Source/Runtime/Renderer/Private/RayTracing/RaytracingOptions.h b/Engine/Source/Runtime/Renderer/Private/RayTracing/RaytracingOptions.h
--- a/Engine/Source/Runtime/Renderer/Private/RayTracing/RaytracingOptions.h
+++ b/Engine/Source/Runtime/Renderer/Private/RayTracing/RaytracingOptions.h
@@ -42,6 +42,7 @@ struct FRayTracingPrimaryRaysOptions
 
 
 #if RHI_RAYTRACING
+extern bool AnyRayTracingPassEnabled(const FViewInfo& View);
 extern int32 GetForceRayTracingEffectsCVarValue();
 extern FRayTracingPrimaryRaysOptions GetRayTracingTranslucencyOptions();
 
@@ -60,6 +61,11 @@ extern float GetRaytracingMaxNormalBias();
 
 #else
 
+FORCEINLINE bool AnyRayTracingPassEnabled(const FViewInfo& View)
+{
+	return 0;
+}
+
 FORCEINLINE int32 GetForceRayTracingEffectsCVarValue()
 {
 	return 0;
