Extending ray tracing global illumination gather points to 32 samples per pixel.

diff --git a/Engine/Shaders/Private/RayTracing/RayTracingCreateGatherPointsRGS.usf b/Engine/Shaders/Private/RayTracing/RayTracingCreateGatherPointsRGS.usf
--- a/Engine/Shaders/Private/RayTracing/RayTracingCreateGatherPointsRGS.usf
+++ b/Engine/Shaders/Private/RayTracing/RayTracingCreateGatherPointsRGS.usf
@@ -117,7 +117,7 @@ RAY_TRACING_ENTRY_RAYGEN(RayTracingCreateGatherPointsRGS)
 
 	for (uint Index = 0; Index < SamplesPerPixel; ++Index)
 	{
-		uint SampleIndexOffset = (SampleIndex + Index) % 16;
+		uint SampleIndexOffset = (SampleIndex + Index) % MAXIMUM_GATHER_POINTS_PER_PIXEL;
 #if 0
 		// Adaptive rejection optimization
 		{
diff --git a/Engine/Shaders/Private/RayTracing/RayTracingFinalGatherRGS.usf b/Engine/Shaders/Private/RayTracing/RayTracingFinalGatherRGS.usf
--- a/Engine/Shaders/Private/RayTracing/RayTracingFinalGatherRGS.usf
+++ b/Engine/Shaders/Private/RayTracing/RayTracingFinalGatherRGS.usf
@@ -86,7 +86,7 @@ RAY_TRACING_ENTRY_RAYGEN(RayTracingFinalGatherRGS)
 	uint GatherPointsIndex = DispatchThreadIdToLinearIndex(DispatchThreadId);
 
 	// Define rejection criteria and evaluate
-	bool bShouldReject[16];
+	bool bShouldReject[MAXIMUM_GATHER_POINTS_PER_PIXEL];
 	FRejectionCriteria RejectionCriteria = CreateRejectionCriteria(WorldPosition, FinalGatherDistance);
 	CreateRejectionMask(RejectionCriteria, GatherPointsBuffer, GatherPointsIndex, 0, GatherPointsResolution.x * GatherPointsResolution.y, SamplesPerPixel, bShouldReject);
 
diff --git a/Engine/Shaders/Private/RayTracing/RayTracingGatherPoints.ush b/Engine/Shaders/Private/RayTracing/RayTracingGatherPoints.ush
--- a/Engine/Shaders/Private/RayTracing/RayTracingGatherPoints.ush
+++ b/Engine/Shaders/Private/RayTracing/RayTracingGatherPoints.ush
@@ -1,5 +1,7 @@
 // Copyright 1998-2019 Epic Games, Inc. All Rights Reserved.
 
+#define MAXIMUM_GATHER_POINTS_PER_PIXEL 32
+
 struct FGatherPoints
 {
 	float3 CreationPoint;
@@ -70,7 +72,15 @@ FRejectionCriteria CreateRejectionCriteria(float3 CreationPoint, float DistanceT
 	return RejectionCriteria;
 }
 
-void CreateRejectionMask(FRejectionCriteria RejectionCriteria, StructuredBuffer<FGatherPoints> GatherBuffer, uint BaseIndex, uint SampleIndex, uint ElementCount, int SampleCount, out bool bShouldReject[16])
+void CreateRejectionMask(
+	FRejectionCriteria RejectionCriteria,
+	StructuredBuffer<FGatherPoints> GatherBuffer,
+	uint BaseIndex,
+	uint SampleIndex,
+	uint ElementCount,
+	int SampleCount,
+	out bool bShouldReject[MAXIMUM_GATHER_POINTS_PER_PIXEL]
+)
 {
 	for (int Index = 0; Index < SampleCount; ++Index)
 	{
diff --git a/Engine/Source/Runtime/Renderer/Private/RayTracing/RayTracingGlobalIllumination.cpp b/Engine/Source/Runtime/Renderer/Private/RayTracing/RayTracingGlobalIllumination.cpp
--- a/Engine/Source/Runtime/Renderer/Private/RayTracing/RayTracingGlobalIllumination.cpp
+++ b/Engine/Source/Runtime/Renderer/Private/RayTracing/RayTracingGlobalIllumination.cpp
@@ -306,13 +306,16 @@ class FGlobalIlluminationRGS : public FGlobalShader
 
 IMPLEMENT_GLOBAL_SHADER(FGlobalIlluminationRGS, "/Engine/Private/RayTracing/RayTracingGlobalIlluminationRGS.usf", "GlobalIlluminationRGS", SF_RayGen);
 
+// Note: This constant must match the definition in RayTracingGatherPoints.ush
+constexpr int32 MAXIMUM_GATHER_POINTS_PER_PIXEL = 32;
+
 // Placeholder structure used to allocate gather points buffer
 // #dxr_todo: rework to minimal set, based on active samples-per-pixel
 struct FGatherPoints
 {
-	FVector CreationPoint[16];
-	FVector Position[16];
-	FIntPoint Irradiance[16];
+	FVector CreationPoint[MAXIMUM_GATHER_POINTS_PER_PIXEL];
+	FVector Position[MAXIMUM_GATHER_POINTS_PER_PIXEL];
+	FIntPoint Irradiance[MAXIMUM_GATHER_POINTS_PER_PIXEL];
 };
 
 class FRayTracingGlobalIlluminationCreateGatherPointsRGS : public FGlobalShader
@@ -320,7 +323,7 @@ class FRayTracingGlobalIlluminationCreateGatherPointsRGS : public FGlobalShader
 	DECLARE_GLOBAL_SHADER(FRayTracingGlobalIlluminationCreateGatherPointsRGS)
 	SHADER_USE_ROOT_PARAMETER_STRUCT(FRayTracingGlobalIlluminationCreateGatherPointsRGS, FGlobalShader)
 
-		class FUseAttenuationTermDim : SHADER_PERMUTATION_BOOL("USE_ATTENUATION_TERM");
+	class FUseAttenuationTermDim : SHADER_PERMUTATION_BOOL("USE_ATTENUATION_TERM");
 	class FEnableTwoSidedGeometryDim : SHADER_PERMUTATION_BOOL("ENABLE_TWO_SIDED_GEOMETRY");
 
 	using FPermutationDomain = TShaderPermutationDomain<FUseAttenuationTermDim, FEnableTwoSidedGeometryDim>;
@@ -642,7 +645,7 @@ void FDeferredShadingSceneRenderer::RenderRayTracingGlobalIlluminationFinalGathe
 	RDG_EVENT_SCOPE(GraphBuilder, "Ray Tracing GI: Final Gather");
 
 	FRayTracingGlobalIlluminationFinalGatherRGS::FParameters* PassParameters = GraphBuilder.AllocParameters<FRayTracingGlobalIlluminationFinalGatherRGS::FParameters>();
-	int32 SamplesPerPixel = GetRayTracingGlobalIlluminationSamplesPerPixel(View);
+	int32 SamplesPerPixel = FMath::Min(GetRayTracingGlobalIlluminationSamplesPerPixel(View), MAXIMUM_GATHER_POINTS_PER_PIXEL);
 	int32 SampleIndex = View.ViewState->FrameIndex % SamplesPerPixel;
 	PassParameters->SampleIndex = SampleIndex;
 	PassParameters->SamplesPerPixel = SamplesPerPixel;
