Strictly bound gather points buffer allocation to active samples per pixel

diff --git a/Engine/Shaders/Private/RayTracing/RayTracingCreateGatherPointsRGS.usf b/Engine/Shaders/Private/RayTracing/RayTracingCreateGatherPointsRGS.usf
--- a/Engine/Shaders/Private/RayTracing/RayTracingCreateGatherPointsRGS.usf
+++ b/Engine/Shaders/Private/RayTracing/RayTracingCreateGatherPointsRGS.usf
@@ -17,6 +17,7 @@
 
 RaytracingAccelerationStructure TLAS;
 
+uint GatherSamplesPerPixel;
 uint SamplesPerPixel;
 uint SampleIndex;
 uint MaxBounces;
@@ -117,7 +118,7 @@ RAY_TRACING_ENTRY_RAYGEN(RayTracingCreateGatherPointsRGS)
 
 	for (uint Index = 0; Index < SamplesPerPixel; ++Index)
 	{
-		uint SampleIndexOffset = (SampleIndex + Index) % MAXIMUM_GATHER_POINTS_PER_PIXEL;
+		uint SampleIndexOffset = (SampleIndex + Index) % GatherSamplesPerPixel;
 #if 0
 		// Adaptive rejection optimization
 		{
diff --git a/Engine/Source/Runtime/Renderer/Private/DeferredShadingRenderer.h b/Engine/Source/Runtime/Renderer/Private/DeferredShadingRenderer.h
--- a/Engine/Source/Runtime/Renderer/Private/DeferredShadingRenderer.h
+++ b/Engine/Source/Runtime/Renderer/Private/DeferredShadingRenderer.h
@@ -542,7 +542,7 @@ private:
 		FViewInfo& View,
 		int32 UpscaleFactor,
 		FRDGBufferRef& GatherPointsBuffer,
-		FIntPoint& GatherPointsResolution);
+		FIntVector& GatherPointsResolution);
 
 	void RenderRayTracingGlobalIlluminationFinalGather(
 		FRDGBuilder& GraphBuilder,
diff --git a/Engine/Source/Runtime/Renderer/Private/RayTracing/RayTracingGlobalIllumination.cpp b/Engine/Source/Runtime/Renderer/Private/RayTracing/RayTracingGlobalIllumination.cpp
--- a/Engine/Source/Runtime/Renderer/Private/RayTracing/RayTracingGlobalIllumination.cpp
+++ b/Engine/Source/Runtime/Renderer/Private/RayTracing/RayTracingGlobalIllumination.cpp
@@ -309,13 +309,11 @@ IMPLEMENT_GLOBAL_SHADER(FGlobalIlluminationRGS, "/Engine/Private/RayTracing/RayT
 // Note: This constant must match the definition in RayTracingGatherPoints.ush
 constexpr int32 MAXIMUM_GATHER_POINTS_PER_PIXEL = 32;
 
-// Placeholder structure used to allocate gather points buffer
-// #dxr_todo: rework to minimal set, based on active samples-per-pixel
-struct FGatherPoints
+struct FGatherPoint
 {
-	FVector CreationPoint[MAXIMUM_GATHER_POINTS_PER_PIXEL];
-	FVector Position[MAXIMUM_GATHER_POINTS_PER_PIXEL];
-	FIntPoint Irradiance[MAXIMUM_GATHER_POINTS_PER_PIXEL];
+	FVector CreationPoint;
+	FVector Position;
+	FIntPoint Irradiance;
 };
 
 class FRayTracingGlobalIlluminationCreateGatherPointsRGS : public FGlobalShader
@@ -334,6 +332,7 @@ class FRayTracingGlobalIlluminationCreateGatherPointsRGS : public FGlobalShader
 	}
 
 	BEGIN_SHADER_PARAMETER_STRUCT(FParameters, )
+		SHADER_PARAMETER(uint32, GatherSamplesPerPixel)
 		SHADER_PARAMETER(uint32, SamplesPerPixel)
 		SHADER_PARAMETER(uint32, SampleIndex)
 		SHADER_PARAMETER(uint32, MaxBounces)
@@ -513,7 +512,7 @@ void FDeferredShadingSceneRenderer::RayTracingGlobalIlluminationCreateGatherPoin
 	FViewInfo& View,
 	int32 UpscaleFactor,
 	FRDGBufferRef& GatherPointsBuffer,
-	FIntPoint& GatherPointsResolution
+	FIntVector& GatherPointsResolution
 )
 #if RHI_RAYTRACING
 {
@@ -551,6 +550,7 @@ void FDeferredShadingSceneRenderer::RayTracingGlobalIlluminationCreateGatherPoin
 
 	FRayTracingGlobalIlluminationCreateGatherPointsRGS::FParameters* PassParameters = GraphBuilder.AllocParameters<FRayTracingGlobalIlluminationCreateGatherPointsRGS::FParameters>();
 	PassParameters->SampleIndex = (FrameIndex * SamplesPerPixel) % GatherSamples;
+	PassParameters->GatherSamplesPerPixel = GatherSamples;
 	PassParameters->SamplesPerPixel = SamplesPerPixel;
 	PassParameters->MaxBounces = 1;
 	PassParameters->MaxNormalBias = GetRaytracingMaxNormalBias();
@@ -587,18 +587,19 @@ void FDeferredShadingSceneRenderer::RayTracingGlobalIlluminationCreateGatherPoin
 	PassParameters->TransmissionProfilesLinearSampler = TStaticSamplerState<SF_Bilinear, AM_Clamp, AM_Clamp, AM_Clamp>::GetRHI();
 
 	// Output
-	FIntPoint LocalGatherPointsResolution = FIntPoint::DivideAndRoundUp(View.ViewRect.Size(), UpscaleFactor);
+	FIntPoint DispatchResolution = FIntPoint::DivideAndRoundUp(View.ViewRect.Size(), UpscaleFactor);
+	FIntVector LocalGatherPointsResolution(DispatchResolution.X, DispatchResolution.Y, GatherSamples);
 	if (GatherPointsResolution != LocalGatherPointsResolution)
 	{
 		GatherPointsResolution = LocalGatherPointsResolution;
-		FRDGBufferDesc BufferDesc = FRDGBufferDesc::CreateStructuredDesc(sizeof(FGatherPoints), GatherPointsResolution.X * GatherPointsResolution.Y);
+		FRDGBufferDesc BufferDesc = FRDGBufferDesc::CreateStructuredDesc(sizeof(FGatherPoint), GatherPointsResolution.X * GatherPointsResolution.Y * GatherPointsResolution.Z);
 		GatherPointsBuffer = GraphBuilder.CreateBuffer(BufferDesc, TEXT("GatherPointsBuffer"), ERDGResourceFlags::MultiFrame);
 	}
 	else
 	{
 		GatherPointsBuffer = GraphBuilder.RegisterExternalBuffer(((FSceneViewState*)View.State)->GatherPointsBuffer, TEXT("GatherPointsBuffer"));
 	}
-	PassParameters->GatherPointsResolution = GatherPointsResolution;
+	PassParameters->GatherPointsResolution = FIntPoint(GatherPointsResolution.X, GatherPointsResolution.Y);
 	PassParameters->RWGatherPointsBuffer = GraphBuilder.CreateUAV(GatherPointsBuffer, EPixelFormat::PF_R32_UINT);
 
 	FRayTracingGlobalIlluminationCreateGatherPointsRGS::FPermutationDomain PermutationVector;
@@ -671,7 +672,7 @@ void FDeferredShadingSceneRenderer::RenderRayTracingGlobalIlluminationFinalGathe
 	PassParameters->TransmissionProfilesLinearSampler = TStaticSamplerState<SF_Bilinear, AM_Clamp, AM_Clamp, AM_Clamp>::GetRHI();
 
 	// Gather points
-	PassParameters->GatherPointsResolution = SceneViewState->GatherPointsResolution;
+	PassParameters->GatherPointsResolution = FIntPoint(SceneViewState->GatherPointsResolution.X, SceneViewState->GatherPointsResolution.Y);
 	PassParameters->GatherPointsBuffer = GraphBuilder.CreateSRV(GatherPointsBuffer);
 
 	// Output
diff --git a/Engine/Source/Runtime/Renderer/Private/RendererScene.cpp b/Engine/Source/Runtime/Renderer/Private/RendererScene.cpp
--- a/Engine/Source/Runtime/Renderer/Private/RendererScene.cpp
+++ b/Engine/Source/Runtime/Renderer/Private/RendererScene.cpp
@@ -223,7 +223,7 @@ FSceneViewState::FSceneViewState()
 	RayCountGPUReadback = new FRHIGPUBufferReadback(TEXT("Ray Count Readback"));
 
 	GatherPointsBuffer = nullptr;
-	GatherPointsResolution = FIntPoint(0, 0);
+	GatherPointsResolution = FIntVector(0, 0, 0);
 #endif
 }
 
diff --git a/Engine/Source/Runtime/Renderer/Private/ScenePrivate.h b/Engine/Source/Runtime/Renderer/Private/ScenePrivate.h
--- a/Engine/Source/Runtime/Renderer/Private/ScenePrivate.h
+++ b/Engine/Source/Runtime/Renderer/Private/ScenePrivate.h
@@ -930,7 +930,7 @@ public:
 	FIESLightProfileResource IESLightProfileResources;
 
 	TRefCountPtr<FPooledRDGBuffer> GatherPointsBuffer;
-	FIntPoint GatherPointsResolution;
+	FIntVector GatherPointsResolution;
 #endif
 
 	TUniquePtr<FForwardLightingViewResources> ForwardLightingResources;
