Implement functionality for the Sky Light to respect the Affect Global Illumination flag as well as Partial functionality for the Affect Reflection flag (which depends on changes not in 4.24 for now).

(ushell-p4-cherrypick of 10086385 by riley.alston)

diff --git a/Engine/Shaders/Private/BasePassPixelShader.usf b/Engine/Shaders/Private/BasePassPixelShader.usf
--- a/Engine/Shaders/Private/BasePassPixelShader.usf
+++ b/Engine/Shaders/Private/BasePassPixelShader.usf
@@ -273,7 +273,7 @@ void GetSkyLighting(FMaterialPixelParameters MaterialParameters, VTPageTableResu
 	
 	#if HQ_TEXTURE_LIGHTMAP || CACHED_POINT_INDIRECT_LIGHTING || CACHED_VOLUME_INDIRECT_LIGHTING || PRECOMPUTED_IRRADIANCE_VOLUME_LIGHTING
 		BRANCH
-		if (View.SkyLightParameters.x > 0)
+		if (ShouldSkyLightApplyPrecomputedBentNormalShadowing())
 		{
 			float3 NormalizedBentNormal = SkyLightingNormal;
 
diff --git a/Engine/Shaders/Private/RayTracing/RayTracingMaterialHitShaders.usf b/Engine/Shaders/Private/RayTracing/RayTracingMaterialHitShaders.usf
--- a/Engine/Shaders/Private/RayTracing/RayTracingMaterialHitShaders.usf
+++ b/Engine/Shaders/Private/RayTracing/RayTracingMaterialHitShaders.usf
@@ -44,7 +44,7 @@ void GetSkyLighting(FMaterialPixelParameters MaterialParameters, VTPageTableResu
 	
 	#if HQ_TEXTURE_LIGHTMAP || CACHED_POINT_INDIRECT_LIGHTING || CACHED_VOLUME_INDIRECT_LIGHTING || PRECOMPUTED_IRRADIANCE_VOLUME_LIGHTING
 		BRANCH
-		if (View.SkyLightParameters.x > 0)
+		if (ShouldSkyLightApplyPrecomputedBentNormalShadowing())
 		{
 			#if PRECOMPUTED_IRRADIANCE_VOLUME_LIGHTING
 			
diff --git a/Engine/Shaders/Private/ReflectionEnvironmentShared.ush b/Engine/Shaders/Private/ReflectionEnvironmentShared.ush
--- a/Engine/Shaders/Private/ReflectionEnvironmentShared.ush
+++ b/Engine/Shaders/Private/ReflectionEnvironmentShared.ush
@@ -52,6 +52,18 @@ float3 GetSkyLightReflectionSupportingBlend(float3 ReflectionVector, float Rough
 	return Reflection;
 }
 
+bool ShouldSkyLightApplyPrecomputedBentNormalShadowing() {
+	return View.SkyLightApplyPrecomputedBentNormalShadowingFlag != 0.0f;
+}
+
+bool ShouldSkyLightAffectReflection() {
+	return View.SkyLightAffectReflectionFlag != 0.0f;
+}
+
+bool ShouldSkyLightAffectGlobalIllumination() {
+	return View.SkyLightAffectGlobalIlluminationFlag != 0.0f;
+}
+
 /** 
  * Computes sky diffuse lighting from the SH irradiance map.  
  * This has the SH basis evaluation and diffuse convolution weights combined for minimal ALU's - see "Stupid Spherical Harmonics (SH) Tricks" 
diff --git a/Engine/Shaders/Public/ShaderVersion.ush b/Engine/Shaders/Public/ShaderVersion.ush
--- a/Engine/Shaders/Public/ShaderVersion.ush
+++ b/Engine/Shaders/Public/ShaderVersion.ush
@@ -3,4 +3,4 @@
 // in Platform.ush (which should be included in any shader) it allows to invalidate the shader DDC.
 // 
 // If you are merging streams and there is a conflict with this GUID you should make a new GUID rather than taking one or the other.
-// GUID = 6925F773-21F5-4603-97EA-A8D76B115A4E
+// GUID = 054EAB614D124BB38BF05ED86E984C59
\ No newline at end of file
diff --git a/Engine/Source/Runtime/Engine/Private/Components/SkyLightComponent.cpp b/Engine/Source/Runtime/Engine/Private/Components/SkyLightComponent.cpp
--- a/Engine/Source/Runtime/Engine/Private/Components/SkyLightComponent.cpp
+++ b/Engine/Source/Runtime/Engine/Private/Components/SkyLightComponent.cpp
@@ -174,6 +174,8 @@ FSkyLightSceneProxy::FSkyLightSceneProxy(const USkyLightComponent* InLightCompon
 	, bHasStaticLighting(InLightComponent->HasStaticLighting())
 	, bCastVolumetricShadow(InLightComponent->bCastVolumetricShadow)
 	, bCastRayTracedShadow(InLightComponent->bCastRaytracedShadow)
+	, bAffectReflection(InLightComponent->bAffectReflection)
+	, bAffectGlobalIllumination(InLightComponent->bAffectGlobalIllumination)
 	, OcclusionCombineMode(InLightComponent->OcclusionCombineMode)
 	, IndirectLightingIntensity(InLightComponent->IndirectLightingIntensity)
 	, VolumetricScatteringIntensity(FMath::Max(InLightComponent->VolumetricScatteringIntensity, 0.0f))
@@ -236,6 +238,8 @@ USkyLightComponent::USkyLightComponent(const FObjectInitializer& ObjectInitializ
 	BlendDestinationAverageBrightness = 1.0f;
 	bCastVolumetricShadow = true;
 	bCastRaytracedShadow = false;
+	bAffectReflection = true;
+	bAffectGlobalIllumination = true;
 	SamplesPerPixel = 1;
 }
 
diff --git a/Engine/Source/Runtime/Engine/Public/SceneManagement.h b/Engine/Source/Runtime/Engine/Public/SceneManagement.h
--- a/Engine/Source/Runtime/Engine/Public/SceneManagement.h
+++ b/Engine/Source/Runtime/Engine/Public/SceneManagement.h
@@ -1068,6 +1068,8 @@ public:
 	uint8 bHasStaticLighting:1;
 	uint8 bCastVolumetricShadow:1;
 	uint8 bCastRayTracedShadow:1;
+	uint8 bAffectReflection:1;
+	uint8 bAffectGlobalIllumination:1;
 	TEnumAsByte<EOcclusionCombineMode> OcclusionCombineMode;
 	float AverageBrightness;
 	float IndirectLightingIntensity;
diff --git a/Engine/Source/Runtime/Engine/Public/SceneView.h b/Engine/Source/Runtime/Engine/Public/SceneView.h
--- a/Engine/Source/Runtime/Engine/Public/SceneView.h
+++ b/Engine/Source/Runtime/Engine/Public/SceneView.h
@@ -690,7 +690,9 @@ enum ETranslucencyVolumeCascade
 	VIEW_UNIFORM_BUFFER_MEMBER(float, RenderingReflectionCaptureMask) \
 	VIEW_UNIFORM_BUFFER_MEMBER(FLinearColor, AmbientCubemapTint) \
 	VIEW_UNIFORM_BUFFER_MEMBER(float, AmbientCubemapIntensity) \
-	VIEW_UNIFORM_BUFFER_MEMBER(float, SkyLightParameters) \
+	VIEW_UNIFORM_BUFFER_MEMBER(float, SkyLightApplyPrecomputedBentNormalShadowingFlag) \
+	VIEW_UNIFORM_BUFFER_MEMBER(float, SkyLightAffectReflectionFlag) \
+	VIEW_UNIFORM_BUFFER_MEMBER(float, SkyLightAffectGlobalIlluminationFlag) \
 	VIEW_UNIFORM_BUFFER_MEMBER(FLinearColor, SkyLightColor) \
 	VIEW_UNIFORM_BUFFER_MEMBER_ARRAY(FVector4, SkyIrradianceEnvironmentMap, [7]) \
 	VIEW_UNIFORM_BUFFER_MEMBER(float, MobilePreviewMode) \
diff --git a/Engine/Source/Runtime/Renderer/Private/RayTracing/RayTracingGlobalIllumination.cpp b/Engine/Source/Runtime/Renderer/Private/RayTracing/RayTracingGlobalIllumination.cpp
--- a/Engine/Source/Runtime/Renderer/Private/RayTracing/RayTracingGlobalIllumination.cpp
+++ b/Engine/Source/Runtime/Renderer/Private/RayTracing/RayTracingGlobalIllumination.cpp
@@ -147,21 +147,30 @@ DECLARE_GPU_STAT_NAMED(RayTracingGIFinalGather, TEXT("Ray Tracing GI: Final Gath
 DECLARE_GPU_STAT_NAMED(RayTracingGICreateGatherPoints, TEXT("Ray Tracing GI: Create Gather Points"));
 
 void SetupLightParameters(
-	const TSparseArray<FLightSceneInfoCompact>& Lights,
+	const FScene& Scene,
 	const FViewInfo& View,
 	FPathTracingLightData* LightParameters)
 {
 	LightParameters->Count = 0;
 
+	// Get the SkyLight color
+
+	FSkyLightSceneProxy* SkyLight = Scene.SkyLight;
+	FVector SkyLightColor = FVector(0.0f, 0.0f, 0.0f);
+	if (SkyLight && SkyLight->bAffectGlobalIllumination)
+	{
+		SkyLightColor = FVector(1.0f);
+	}
+
 	// Prepend SkyLight to light buffer
 	// WARNING: Until ray payload encodes Light data buffer, the execution depends on this ordering!
 	uint32 SkyLightIndex = 0;
 	LightParameters->Type[SkyLightIndex] = 0;
-	LightParameters->Color[SkyLightIndex] = FVector(1.0);
+	LightParameters->Color[SkyLightIndex] = SkyLightColor;
 	LightParameters->Count++;
 
 	uint32 MaxLightCount = FMath::Min(CVarRayTracingGlobalIlluminationMaxLightCount.GetValueOnRenderThread(), RAY_TRACING_LIGHT_COUNT_MAXIMUM);
-	for (auto Light : Lights)
+	for (auto Light : Scene.Lights)
 	{
 		if (LightParameters->Count >= MaxLightCount) break;
 
@@ -226,6 +235,21 @@ void SetupLightParameters(
 	}
 }
 
+void SetupGlobalIlluminationSkyLightParameters(
+	const FScene& Scene,
+	FSkyLightData* SkyLightData)
+{
+	FSkyLightSceneProxy* SkyLight = Scene.SkyLight;
+
+	SetupSkyLightParameters(Scene, SkyLightData);
+
+	// Override the Sky Light color if it should not affect global illumination
+	if (SkyLight && !SkyLight->bAffectGlobalIllumination)
+	{
+		SkyLightData->Color = FVector(0.0f);
+	}
+}
+
 int32 GetRayTracingGlobalIlluminationSamplesPerPixel(const FViewInfo& View)
 {
 	int32 SamplesPerPixel = GRayTracingGlobalIlluminationSamplesPerPixel > -1 ? GRayTracingGlobalIlluminationSamplesPerPixel : View.FinalPostProcessSettings.RayTracingGISamplesPerPixel;
@@ -538,7 +562,7 @@ void FDeferredShadingSceneRenderer::RayTracingGlobalIlluminationCreateGatherPoin
 	InitializeBlueNoise(BlueNoise);
 
 	FPathTracingLightData LightParameters;
-	SetupLightParameters(Scene->Lights, View, &LightParameters);
+	SetupLightParameters(*Scene, View, &LightParameters);
 
 	if (Scene->SkyLight && Scene->SkyLight->ShouldRebuildCdf())
 	{
@@ -546,7 +570,7 @@ void FDeferredShadingSceneRenderer::RayTracingGlobalIlluminationCreateGatherPoin
 		BuildSkyLightCdfs(GraphBuilder.RHICmdList, Scene->SkyLight);
 	}
 	FSkyLightData SkyLightParameters;
-	SetupSkyLightParameters(*Scene, &SkyLightParameters);
+	SetupGlobalIlluminationSkyLightParameters(*Scene, &SkyLightParameters);
 
 	FRayTracingGlobalIlluminationCreateGatherPointsRGS::FParameters* PassParameters = GraphBuilder.AllocParameters<FRayTracingGlobalIlluminationCreateGatherPointsRGS::FParameters>();
 	PassParameters->SampleIndex = (FrameIndex * SamplesPerPixel) % GatherSamples;
@@ -736,7 +760,7 @@ void FDeferredShadingSceneRenderer::RenderRayTracingGlobalIlluminationBruteForce
 	InitializeBlueNoise(BlueNoise);
 
 	FPathTracingLightData LightParameters;
-	SetupLightParameters(Scene->Lights, View, &LightParameters);
+	SetupLightParameters(*Scene, View, &LightParameters);
 
 	if (Scene->SkyLight && Scene->SkyLight->ShouldRebuildCdf())
 	{
@@ -744,7 +768,7 @@ void FDeferredShadingSceneRenderer::RenderRayTracingGlobalIlluminationBruteForce
 		BuildSkyLightCdfs(GraphBuilder.RHICmdList, Scene->SkyLight);
 	}
 	FSkyLightData SkyLightParameters;
-	SetupSkyLightParameters(*Scene, &SkyLightParameters);
+	SetupGlobalIlluminationSkyLightParameters(*Scene, &SkyLightParameters);
 
 	FGlobalIlluminationRGS::FParameters* PassParameters = GraphBuilder.AllocParameters<FGlobalIlluminationRGS::FParameters>();
 	PassParameters->SamplesPerPixel = RayTracingGISamplesPerPixel;
diff --git a/Engine/Source/Runtime/Renderer/Private/SceneRendering.cpp b/Engine/Source/Runtime/Renderer/Private/SceneRendering.cpp
--- a/Engine/Source/Runtime/Renderer/Private/SceneRendering.cpp
+++ b/Engine/Source/Runtime/Renderer/Private/SceneRendering.cpp
@@ -1554,12 +1554,16 @@ void FViewInfo::SetupUniformBufferParameters(
 			SkyLight->bCastShadows 
 			&& SkyLight->bWantsStaticShadowing;
 
-		ViewUniformShaderParameters.SkyLightParameters = bApplyPrecomputedBentNormalShadowing ? 1 : 0;
+		ViewUniformShaderParameters.SkyLightApplyPrecomputedBentNormalShadowingFlag = bApplyPrecomputedBentNormalShadowing ? 1.0f : 0.0f;
+		ViewUniformShaderParameters.SkyLightAffectReflectionFlag = SkyLight->bAffectReflection ? 1.0f : 0.0f;
+		ViewUniformShaderParameters.SkyLightAffectGlobalIlluminationFlag = SkyLight->bAffectGlobalIllumination ? 1.0f : 0.0f;
 	}
 	else
 	{
 		ViewUniformShaderParameters.SkyLightColor = FLinearColor::Black;
-		ViewUniformShaderParameters.SkyLightParameters = 0;
+		ViewUniformShaderParameters.SkyLightApplyPrecomputedBentNormalShadowingFlag = 0.0f;
+		ViewUniformShaderParameters.SkyLightAffectReflectionFlag = 0.0f;
+		ViewUniformShaderParameters.SkyLightAffectGlobalIlluminationFlag = 0.0f;
 	}
 
 	// Make sure there's no padding since we're going to cast to FVector4*
