Fix sky light color not being set properly in ray tracing global illumination.

(ushell-p4-cherrypick of 10086959 by riley.alston)

diff --git a/Engine/Source/Runtime/Renderer/Private/RayTracing/RayTracingGlobalIllumination.cpp b/Engine/Source/Runtime/Renderer/Private/RayTracing/RayTracingGlobalIllumination.cpp
--- a/Engine/Source/Runtime/Renderer/Private/RayTracing/RayTracingGlobalIllumination.cpp
+++ b/Engine/Source/Runtime/Renderer/Private/RayTracing/RayTracingGlobalIllumination.cpp
@@ -159,7 +159,7 @@ void SetupLightParameters(
 	FVector SkyLightColor = FVector(0.0f, 0.0f, 0.0f);
 	if (SkyLight && SkyLight->bAffectGlobalIllumination)
 	{
-		SkyLightColor = FVector(1.0f);
+		SkyLightColor = FVector(SkyLight->GetEffectiveLightColor());
 	}
 
 	// Prepend SkyLight to light buffer
