Lots of vulkan validation work:
-Debug Names for images for helping track down validation errors
-When the reflection environment rounds up to 2 cubemaps, explicitly clear it to avoid a validation error
-Shadowmask texture is now properly transitioned to read.
-Add Missing transtion for LightScatteringHistoryTexture
-Add Ensure when updating descriptorset, to catch validation errors when texture is set
-Fixed an issue in transition resources, where it would use the incorrect no. of subresources
-Add VK_QUERY_RESULT_WAIT_BIT to ensure Query results are actually ready when copied back to result buffer

diff --git a/Engine/Source/Runtime/Renderer/Private/LightRendering.cpp b/Engine/Source/Runtime/Renderer/Private/LightRendering.cpp
--- a/Engine/Source/Runtime/Renderer/Private/LightRendering.cpp
+++ b/Engine/Source/Runtime/Renderer/Private/LightRendering.cpp
@@ -1202,7 +1202,7 @@ void FDeferredShadingSceneRenderer::RenderLights(FRHICommandListImmediate& RHICm
 			} // if (RHI_RAYTRACING)
 
 			bool bDirectLighting = ViewFamily.EngineShowFlags.DirectLighting;
-
+			bool bShadowMaskReadable = false;
 			TRefCountPtr<IPooledRenderTarget> ScreenShadowMaskTexture;
 			TRefCountPtr<IPooledRenderTarget> ScreenShadowMaskSubPixelTexture;
 
@@ -1226,6 +1226,7 @@ void FDeferredShadingSceneRenderer::RenderLights(FRHICommandListImmediate& RHICm
 				if ((bDrawShadows || bDrawLightFunction || bDrawPreviewIndicator) && !ScreenShadowMaskTexture.IsValid())
 				{
 					SceneContext.AllocateScreenShadowMask(RHICmdList, ScreenShadowMaskTexture);
+					bShadowMaskReadable = false;
 					if (bUseHairLighting)
 						SceneContext.AllocateScreenShadowMask(RHICmdList, ScreenShadowMaskSubPixelTexture);
 				}
@@ -1774,6 +1775,12 @@ void FDeferredShadingSceneRenderer::RenderLights(FRHICommandListImmediate& RHICm
 						RHICmdList.CopyToResolveTarget(ScreenShadowMaskSubPixelTexture->GetRenderTargetItem().TargetableTexture, ScreenShadowMaskSubPixelTexture->GetRenderTargetItem().ShaderResourceTexture, FResolveParams(FResolveRect()));
 				}
 
+				if(bUsedShadowMaskTexture || !bShadowMaskReadable)
+				{
+					RHICmdList.TransitionResource(EResourceTransitionAccess::EReadable, ScreenShadowMaskTexture->GetRenderTargetItem().ShaderResourceTexture);
+					bShadowMaskReadable = true;
+				}
+
 				if(bDirectLighting && !bInjectedTranslucentVolume)
 				{
 					for (int32 ViewIndex = 0; ViewIndex < Views.Num(); ViewIndex++)
diff --git a/Engine/Source/Runtime/Renderer/Private/ReflectionEnvironment.cpp b/Engine/Source/Runtime/Renderer/Private/ReflectionEnvironment.cpp
--- a/Engine/Source/Runtime/Renderer/Private/ReflectionEnvironment.cpp
+++ b/Engine/Source/Runtime/Renderer/Private/ReflectionEnvironment.cpp
@@ -289,6 +289,56 @@ TUniformBufferRef<FReflectionUniformParameters> CreateReflectionUniformBuffer(co
 	return CreateUniformBufferImmediate(ReflectionStruct, Usage);
 }
 
+void ClearCubeArray(FRHICommandListImmediate& RHICmdList, TRefCountPtr<IPooledRenderTarget>& CubeArray)
+{
+	const FPooledRenderTargetDesc& Desc = CubeArray.GetReference()->GetDesc();
+	int32 NumMips = Desc.NumMips;
+	int32 CubemapSize = Desc.Extent.X;
+	int32 NumCaptures = Desc.ArraySize;
+	TRefCountPtr<IPooledRenderTarget> TempCube;
+
+	uint32 CubeTexFlags = TexCreate_TargetArraySlicesIndependently;
+	FPooledRenderTargetDesc Desc2(FPooledRenderTargetDesc::CreateCubemapDesc(CubemapSize, Desc.Format, FClearValueBinding(FLinearColor(0, 0, 0, 0)), CubeTexFlags, TexCreate_RenderTargetable, false, 1, NumMips));
+	GRenderTargetPool.FindFreeElement(RHICmdList, Desc2, TempCube, TEXT("TempCube"), true, ERenderTargetTransience::Transient);
+
+	FSceneRenderTargetItem& RT0 = TempCube->GetRenderTargetItem();
+	{
+		SCOPED_DRAW_EVENT(RHICmdList, ClearCubeArray);
+
+		RHICmdList.TransitionResource(EResourceTransitionAccess::EWritable, RT0.TargetableTexture);
+
+		for (int32 MipIndex = 0; MipIndex < NumMips; MipIndex++)
+		{
+			for (int32 CubeFace = 0; CubeFace < CubeFace_MAX; CubeFace++)
+			{
+				FRHIRenderPassInfo RPInfo(RT0.TargetableTexture, ERenderTargetActions::Clear_Store, nullptr, MipIndex, CubeFace);
+
+				RHICmdList.BeginRenderPass(RPInfo, TEXT("ClearCubeFace"));
+				RHICmdList.EndRenderPass();
+			}
+		}
+
+		FResolveParams ResolveParams(FResolveRect(), CubeFace_PosX, -1, -1, -1);
+		RHICmdList.CopyToResolveTarget(RT0.TargetableTexture, RT0.TargetableTexture, ResolveParams);
+
+		FSceneRenderTargetItem& DestCube = CubeArray.GetReference()->GetRenderTargetItem();
+
+		// GPU copy back to the scene's texture array, which is not a render target
+		for (int32 CaptureIndex = 0; CaptureIndex < NumCaptures; CaptureIndex++)
+		{
+			for (int32 MipIndex = 0; MipIndex < NumMips; MipIndex++)
+			{
+				for (int32 CubeFace = 0; CubeFace < CubeFace_MAX; CubeFace++)
+				{
+					RHICmdList.CopyToResolveTarget(RT0.TargetableTexture, DestCube.ShaderResourceTexture, FResolveParams(FResolveRect(), (ECubeFace)CubeFace, MipIndex, 0, CaptureIndex));
+				}
+			}
+		}
+	}
+}
+
+
+
 void FReflectionEnvironmentCubemapArray::InitDynamicRHI()
 {
 	if (GetFeatureLevel() >= ERHIFeatureLevel::SM5)
@@ -318,6 +368,10 @@ void FReflectionEnvironmentCubemapArray::InitDynamicRHI()
 
 		// Allocate TextureCubeArray for the scene's reflection captures
 		GRenderTargetPool.FindFreeElement(RHICmdList, Desc, ReflectionEnvs, TEXT("ReflectionEnvs"));
+		if(MaxCubemaps == 1)
+		{
+			ClearCubeArray(RHICmdList, ReflectionEnvs); // When rounded up to two cubemaps, we need to clear all , to avoid validation problems, which happen because we don't fill out the final cube
+		}
 	}
 }
 
diff --git a/Engine/Source/Runtime/Renderer/Private/VolumetricFog.cpp b/Engine/Source/Runtime/Renderer/Private/VolumetricFog.cpp
--- a/Engine/Source/Runtime/Renderer/Private/VolumetricFog.cpp
+++ b/Engine/Source/Runtime/Renderer/Private/VolumetricFog.cpp
@@ -1161,6 +1161,7 @@ void FDeferredShadingSceneRenderer::ComputeVolumetricFog(FRHICommandListImmediat
 						? View.ViewState->LightScatteringHistory->GetRenderTargetItem().ShaderResourceTexture
 						: GBlackVolumeTexture->TextureRHI;
 
+					RHICmdList.TransitionResource(EResourceTransitionAccess::EReadable, LightScatteringHistoryTexture);
 					ComputeShader->SetParameters(RHICmdList, View, IntegrationData, FogInfo, LightScatteringHistoryTexture, bUseDirectionalLightShadowing, LightFunctionWorldToShadow);
 
 					SetShaderParameters(RHICmdList, ComputeShader, ComputeShader->GetComputeShader(), *PassParameters);
diff --git a/Engine/Source/Runtime/VulkanRHI/Private/VulkanCommandBuffer.cpp b/Engine/Source/Runtime/VulkanRHI/Private/VulkanCommandBuffer.cpp
--- a/Engine/Source/Runtime/VulkanRHI/Private/VulkanCommandBuffer.cpp
+++ b/Engine/Source/Runtime/VulkanRHI/Private/VulkanCommandBuffer.cpp
@@ -170,7 +170,7 @@ void FVulkanCmdBuffer::End()
 		uint64 Index = Query.Index;
 		VkBuffer BufferHandle = Query.BufferHandle;
 		VkQueryPool PoolHandle = Query.PoolHandle;
-		VulkanRHI::vkCmdCopyQueryPoolResults(GetHandle(), PoolHandle, Index, Query.Count, BufferHandle, sizeof(uint64) * Index, sizeof(uint64), VK_QUERY_RESULT_64_BIT);
+		VulkanRHI::vkCmdCopyQueryPoolResults(GetHandle(), PoolHandle, Index, Query.Count, BufferHandle, sizeof(uint64) * Index, sizeof(uint64), VK_QUERY_RESULT_64_BIT|VK_QUERY_RESULT_WAIT_BIT);
 		VulkanRHI::vkCmdResetQueryPool(GetHandle(), PoolHandle, Index, Query.Count);
 	}
 
diff --git a/Engine/Source/Runtime/VulkanRHI/Private/VulkanDescriptorSets.h b/Engine/Source/Runtime/VulkanRHI/Private/VulkanDescriptorSets.h
--- a/Engine/Source/Runtime/VulkanRHI/Private/VulkanDescriptorSets.h
+++ b/Engine/Source/Runtime/VulkanRHI/Private/VulkanDescriptorSets.h
@@ -1095,6 +1095,12 @@ protected:
 		if (DescriptorType == VK_DESCRIPTOR_TYPE_SAMPLED_IMAGE)
 		{
 			check(WriteDescriptors[DescriptorIndex].descriptorType == VK_DESCRIPTOR_TYPE_SAMPLED_IMAGE || WriteDescriptors[DescriptorIndex].descriptorType == VK_DESCRIPTOR_TYPE_COMBINED_IMAGE_SAMPLER);
+			ensureMsgf(Layout == VK_IMAGE_LAYOUT_SHARED_PRESENT_KHR ||
+				  Layout == VK_IMAGE_LAYOUT_DEPTH_READ_ONLY_STENCIL_ATTACHMENT_OPTIMAL ||
+				  Layout == VK_IMAGE_LAYOUT_DEPTH_ATTACHMENT_STENCIL_READ_ONLY_OPTIMAL ||
+				  Layout == VK_IMAGE_LAYOUT_DEPTH_STENCIL_READ_ONLY_OPTIMAL || 
+				  Layout == VK_IMAGE_LAYOUT_SHADER_READ_ONLY_OPTIMAL ||
+				  Layout == VK_IMAGE_LAYOUT_GENERAL, TEXT("Invalid Layout %d, Index %d, Type %d\n"), Layout, DescriptorIndex, WriteDescriptors[DescriptorIndex].descriptorType);
 		}
 		else
 		{
diff --git a/Engine/Source/Runtime/VulkanRHI/Private/VulkanDevice.cpp b/Engine/Source/Runtime/VulkanRHI/Private/VulkanDevice.cpp
--- a/Engine/Source/Runtime/VulkanRHI/Private/VulkanDevice.cpp
+++ b/Engine/Source/Runtime/VulkanRHI/Private/VulkanDevice.cpp
@@ -418,6 +418,11 @@ void FVulkanDevice::SetupDrawMarkers()
 			bDebugMarkersFound = true;
 		}
 	}
+	FVulkanDynamicRHI* RHI = GVulkanRHI;
+	if (RHI->SupportsDebugUtilsExt())
+	{
+		DebugMarkers.SetDebugName = (PFN_vkSetDebugUtilsObjectNameEXT)(void*)VulkanRHI::vkGetInstanceProcAddr(RHI->GetInstance(), "vkSetDebugUtilsObjectNameEXT");
+	}
 
 	if (bDebugMarkersFound)
 	{
@@ -1252,4 +1257,18 @@ void FVulkanDevice::ReleaseDeferredContext(FVulkanCommandListContext* InContext)
 		FScopeLock Lock(&GContextCS);
 		CommandContexts.Push(InContext);
 	}
+}
+
+void FVulkanDevice::VulkanSetObjectName(VkObjectType Type, uint64_t Handle, const TCHAR* Name)
+{
+	if(DebugMarkers.SetDebugName)
+	{
+		FTCHARToUTF8 Converter(Name);
+		VkDebugUtilsObjectNameInfoEXT Info;
+		ZeroVulkanStruct(Info, VK_STRUCTURE_TYPE_DEBUG_UTILS_OBJECT_NAME_INFO_EXT);
+		Info.objectType = Type;
+		Info.objectHandle = Handle;
+		Info.pObjectName = Converter.Get();
+		DebugMarkers.SetDebugName(Device, &Info);
+	}
 }
\ No newline at end of file
diff --git a/Engine/Source/Runtime/VulkanRHI/Private/VulkanDevice.h b/Engine/Source/Runtime/VulkanRHI/Private/VulkanDevice.h
--- a/Engine/Source/Runtime/VulkanRHI/Private/VulkanDevice.h
+++ b/Engine/Source/Runtime/VulkanRHI/Private/VulkanDevice.h
@@ -18,6 +18,14 @@ class FVulkanOcclusionQueryPool;
 class FOLDVulkanQueryPool;
 #endif
 
+#define VULKAN_USE_DEBUG_NAMES 0
+
+#ifdef VULKAN_USE_DEBUG_NAMES
+#define VULKAN_SET_DEBUG_NAME(Device, Type, Handle, Format, ...) Device.VulkanSetObjectName(Type, (uint64)Handle, *FString::Printf(Format, __VA_ARGS__))
+#else
+#define VULKAN_SET_DEBUG_NAME(Device, Type, Handle, Format, ...) do{}while(0)
+#endif
+
 struct FOptionalVulkanDeviceExtensions
 {
 	union
@@ -301,7 +309,7 @@ public:
 
 	FVulkanCommandListContext* AcquireDeferredContext();
 	void ReleaseDeferredContext(FVulkanCommandListContext* InContext);
-
+	void VulkanSetObjectName(VkObjectType Type, uint64_t Handle, const TCHAR* Name);
 	inline const FOptionalVulkanDeviceExtensions& GetOptionalExtensions() const
 	{
 		return OptionalDeviceExtensions;
@@ -430,11 +438,11 @@ private:
 		PFN_vkCmdDebugMarkerBeginEXT		CmdBegin = nullptr;
 		PFN_vkCmdDebugMarkerEndEXT			CmdEnd = nullptr;
 		PFN_vkDebugMarkerSetObjectNameEXT	CmdSetObjectName = nullptr;
+		PFN_vkSetDebugUtilsObjectNameEXT	SetDebugName = nullptr;
 
 #if 0//VULKAN_SUPPORTS_DEBUG_UTILS
 		PFN_vkCmdBeginDebugUtilsLabelEXT	CmdBeginDebugLabel = nullptr;
 		PFN_vkCmdEndDebugUtilsLabelEXT		CmdEndDebugLabel = nullptr;
-		PFN_vkSetDebugUtilsObjectNameEXT	SetDebugName = nullptr;
 #endif
 	} DebugMarkers;
 	friend class FVulkanCommandListContext;
diff --git a/Engine/Source/Runtime/VulkanRHI/Private/VulkanRenderTarget.cpp b/Engine/Source/Runtime/VulkanRHI/Private/VulkanRenderTarget.cpp
--- a/Engine/Source/Runtime/VulkanRHI/Private/VulkanRenderTarget.cpp
+++ b/Engine/Source/Runtime/VulkanRHI/Private/VulkanRenderTarget.cpp
@@ -1398,7 +1398,7 @@ void FVulkanCommandListContext::TransitionResources(const FPendingTransition& Pe
 				// ensure(SrcLayout != VK_IMAGE_LAYOUT_UNDEFINED || bIsDepthStencil);
 				VkImageLayout DstLayout = bIsDepthStencil ? VK_IMAGE_LAYOUT_DEPTH_STENCIL_READ_ONLY_OPTIMAL : VK_IMAGE_LAYOUT_SHADER_READ_ONLY_OPTIMAL;
 
-				int32 BarrierIndex = Barrier.AddImageBarrier(VulkanTexture->Surface.Image, VulkanTexture->Surface.GetFullAspectMask(), VulkanTexture->Surface.GetNumMips(), VulkanTexture->Surface.GetNumberOfArrayLevels());
+				int32 BarrierIndex = Barrier.AddImageBarrier(VulkanTexture->Surface.Image, VulkanTexture->Surface.GetFullAspectMask(), VulkanTexture->Surface.GetNumMips(), VulkanTexture->Surface.NumArrayLevels);
 				Barrier.SetTransition(BarrierIndex, VulkanRHI::GetImageLayoutFromVulkanLayout(SrcLayout), VulkanRHI::GetImageLayoutFromVulkanLayout(DstLayout));
 
 				SrcLayout = DstLayout;
@@ -1439,7 +1439,7 @@ void FVulkanCommandListContext::TransitionResources(const FPendingTransition& Pe
 
 				if (SrcLayout != FinalLayout)
 				{
-					int32 BarrierIndex = Barrier.AddImageBarrier(Surface.Image, AspectMask, Surface.GetNumMips(), Surface.GetNumberOfArrayLevels());
+					int32 BarrierIndex = Barrier.AddImageBarrier(Surface.Image, AspectMask, Surface.GetNumMips(), Surface.NumArrayLevels);
 					//todo: TransitionAndLayoutManager should use EImageLayoutBarrier type?
 					Barrier.SetTransition(BarrierIndex, VulkanRHI::GetImageLayoutFromVulkanLayout(SrcLayout), VulkanRHI::GetImageLayoutFromVulkanLayout(FinalLayout));
 					SrcLayout = FinalLayout;
diff --git a/Engine/Source/Runtime/VulkanRHI/Private/VulkanTexture.cpp b/Engine/Source/Runtime/VulkanRHI/Private/VulkanTexture.cpp
--- a/Engine/Source/Runtime/VulkanRHI/Private/VulkanTexture.cpp
+++ b/Engine/Source/Runtime/VulkanRHI/Private/VulkanTexture.cpp
@@ -514,6 +514,7 @@ FVulkanSurface::FVulkanSurface(FVulkanDevice& InDevice, VkImageViewType Resource
 	// Fetch image size
 	VulkanRHI::vkGetImageMemoryRequirements(InDevice.GetInstanceHandle(), Image, &MemoryRequirements);
 
+	VULKAN_SET_DEBUG_NAME(InDevice, VK_OBJECT_TYPE_IMAGE, Image, TEXT("(FVulkanSurface*)0x%p"), this);
 	uint32 LayerCount = (ResourceType == VK_IMAGE_VIEW_TYPE_CUBE || ResourceType == VK_IMAGE_VIEW_TYPE_CUBE_ARRAY) ? 6 : 1;
 	NumArrayLevels = ArraySize * LayerCount;
 
@@ -607,6 +608,7 @@ FVulkanSurface::FVulkanSurface(FVulkanDevice& InDevice, VkImageViewType Resource
 	, FullAspectMask(0)
 	, PartialAspectMask(0)
 {
+	VULKAN_SET_DEBUG_NAME(InDevice, VK_OBJECT_TYPE_IMAGE, Image, TEXT("(FVulkanSurface*)0x%p"), this);
 	StorageFormat = UEToVkTextureFormat(PixelFormat, false);
 
 	checkf(PixelFormat == PF_Unknown || StorageFormat != VK_FORMAT_UNDEFINED, TEXT("PixelFormat %d, is not supported for images"), (int32)PixelFormat);
