Implemented basic tiled mGPU support for the path tracer.

diff --git a/Engine/Shaders/Private/PathTracing/PathTracing.usf b/Engine/Shaders/Private/PathTracing/PathTracing.usf
--- a/Engine/Shaders/Private/PathTracing/PathTracing.usf
+++ b/Engine/Shaders/Private/PathTracing/PathTracing.usf
@@ -27,7 +27,7 @@ RaytracingAccelerationStructure TLAS;
 
 RAY_TRACING_ENTRY_RAYGEN(PathTracingMainRG)
 {
-	uint2 LaunchIndex = DispatchRaysIndex().xy + View.ViewRectMin;
+	uint2 LaunchIndex = DispatchRaysIndex().xy + View.ViewRectMin + PathTracingData.TileOffset;
 	
 	bool IsUnidirectionalEnabled = true;
 	bool IsNextEventEstimationEnabled = true;
diff --git a/Engine/Source/Runtime/Renderer/Private/DeferredShadingRenderer.cpp b/Engine/Source/Runtime/Renderer/Private/DeferredShadingRenderer.cpp
--- a/Engine/Source/Runtime/Renderer/Private/DeferredShadingRenderer.cpp
+++ b/Engine/Source/Runtime/Renderer/Private/DeferredShadingRenderer.cpp
@@ -1000,7 +1000,16 @@ void FDeferredShadingSceneRenderer::Render(FRHICommandListImmediate& RHICmdList)
 	SCOPED_NAMED_EVENT(FDeferredShadingSceneRenderer_Render, FColor::Emerald);
 
 #if WITH_MGPU
-	const FRHIGPUMask RenderTargetGPUMask = (GNumExplicitGPUsForRendering > 1 && ViewFamily.RenderTarget) ? ViewFamily.RenderTarget->GetGPUMask(RHICmdList) : FRHIGPUMask::GPU0();
+	FRHIGPUMask RenderTargetGPUMask = (GNumExplicitGPUsForRendering > 1 && ViewFamily.RenderTarget) ? ViewFamily.RenderTarget->GetGPUMask(RHICmdList) : FRHIGPUMask::GPU0();
+	
+	{
+		auto CVar = IConsoleManager::Get().FindConsoleVariable(TEXT("r.PathTracing.GPUCount"));
+		if (CVar && CVar->GetInt() > 1)
+		{
+			RenderTargetGPUMask = FRHIGPUMask::All(); // Broadcast to all GPUs 
+		}
+	}
+
 	ComputeViewGPUMasks(RenderTargetGPUMask);
 #endif // WITH_MGPU
 
diff --git a/Engine/Source/Runtime/Renderer/Private/PathTracing.cpp b/Engine/Source/Runtime/Renderer/Private/PathTracing.cpp
--- a/Engine/Source/Runtime/Renderer/Private/PathTracing.cpp
+++ b/Engine/Source/Runtime/Renderer/Private/PathTracing.cpp
@@ -77,6 +77,14 @@ TAutoConsoleVariable<int32> CVarPathTracingRayCountFrequency(
 	ECVF_RenderThreadSafe
 );
 
+// r.PathTracing.GPUCount is ready only because ComputeViewGPUMasks results cannot change after UE has been launched
+TAutoConsoleVariable<int32> CVarPathTracingGPUCount(
+	TEXT("r.PathTracing.GPUCount"),
+	1,
+	TEXT("Sets the amount of GPUs used for computing the path tracing pass (default = 1 GPU)"),
+	ECVF_RenderThreadSafe | ECVF_ReadOnly
+);
+
 IMPLEMENT_GLOBAL_SHADER_PARAMETER_STRUCT(FPathTracingData, "PathTracingData");
 IMPLEMENT_GLOBAL_SHADER_PARAMETER_STRUCT(FPathTracingLightData, "SceneLightsData");
 IMPLEMENT_GLOBAL_SHADER_PARAMETER_STRUCT(FPathTracingAdaptiveSamplingData, "AdaptiveSamplingData");
@@ -131,6 +139,7 @@ public:
 		uint32 FrameIndependentTemporalSeed,
 		FIntVector VarianceDimensions,
 		const FRWBuffer& VarianceMipTree,
+		const FIntVector& TileOffset,
 		// Output
 		FRHIUnorderedAccessView* RadianceUAV,
 		FRHIUnorderedAccessView* SampleCountUAV,
@@ -154,6 +163,8 @@ public:
 				PrevMaxBounces = PathTracingData.MaxBounces;
 			}
 
+			PathTracingData.TileOffset = TileOffset;
+
 			FUniformBufferRHIRef PathTracingDataUniformBuffer = RHICreateUniformBuffer(&PathTracingData, FPathTracingData::StaticStructMetadata.GetLayout(), EUniformBufferUsage::UniformBuffer_SingleDraw);
 			GlobalResources.Set(PathTracingParameters, PathTracingDataUniformBuffer);
 		}
@@ -473,23 +484,133 @@ void FDeferredShadingSceneRenderer::RenderPathTracing(FRHICommandListImmediate&
 	SetupSceneTextureUniformParameters(SceneContext, FeatureLevel, ESceneTextureSetupMode::All, SceneTextures);
 	FRHIUniformBuffer* SceneTexturesUniformBuffer = RHICreateUniformBuffer(&SceneTextures, FSceneTexturesUniformParameters::StaticStructMetadata.GetLayout(), EUniformBufferUsage::UniformBuffer_SingleDraw);
 
-	RayGenShader->SetParameters(
-		Scene,
-		View,
-		GlobalResources,
-		View.RayTracingScene,
-		View.ViewUniformBuffer,
-		SceneTexturesUniformBuffer,
-		Scene->Lights,
-		SPPCount, FrameIndependentTemporalSeed, ViewState->VarianceMipTreeDimensions, *ViewState->VarianceMipTree,
-		RadianceRT->GetRenderTargetItem().UAV,
-		SampleCountRT->GetRenderTargetItem().UAV,
-		PixelPositionRT->GetRenderTargetItem().UAV,
-		RayCountPerPixelRT->GetRenderTargetItem().UAV
-	);
-
 	FRHIRayTracingScene* RayTracingSceneRHI = View.RayTracingScene.RayTracingSceneRHI;
-	RHICmdList.RayTraceDispatch(View.RayTracingMaterialPipeline, RayGenShader->GetRayTracingShader(), RayTracingSceneRHI, GlobalResources, View.ViewRect.Size().X, View.ViewRect.Size().Y);
+	
+	int32 GPUCount = CVarPathTracingGPUCount.GetValueOnRenderThread();
+	uint32 MainGPUIndex = 0; // Default GPU for rendering
+	bool bDoMGPUPathTracing = GNumExplicitGPUsForRendering > 1 && GPUCount > 1;
+
+	if (bDoMGPUPathTracing)
+	{
+		//#dxr-todo: Set minimum tile size for mGPU
+		int32 TileSizeX = FIntPoint::DivideAndRoundUp(View.ViewRect.Size(), GPUCount).X;
+
+		FRHIUnorderedAccessView* UAVs[] = {
+				RadianceRT->GetRenderTargetItem().UAV,
+				SampleCountRT->GetRenderTargetItem().UAV,
+				PixelPositionRT->GetRenderTargetItem().UAV,
+				RayCountPerPixelRT->GetRenderTargetItem().UAV
+		};
+
+		{
+			// Begin mGPU fence
+			FRHIGPUMask GPUMask = FRHIGPUMask::All();
+			FComputeFenceRHIRef Fence = RHICmdList.CreateComputeFence(TEXT("PathTracingRayGen_Fence_Begin"));
+
+			RHICmdList.TransitionResources(EResourceTransitionAccess::ERWBarrier, EResourceTransitionPipeline::EGfxToGfx, UAVs, UE_ARRAY_COUNT(UAVs), Fence);
+
+			for (uint32 GPUIndex : GPUMask)
+			{
+				if (GPUIndex == MainGPUIndex) 
+					continue;
+
+				SCOPED_GPU_MASK(RHICmdList, FRHIGPUMask::FromIndex(GPUIndex));
+
+				RHICmdList.WaitComputeFence(Fence);
+			}
+		}
+
+		for (int32 GPUIndex = 0; GPUIndex < GPUCount; ++GPUIndex)
+		{
+			SCOPED_GPU_MASK(RHICmdList, FRHIGPUMask::FromIndex(GPUIndex));
+
+			FIntVector TileOffset;
+			TileOffset.X = GPUIndex * TileSizeX;
+			TileOffset.Y = 0; //vertical tiles only 
+
+			RayGenShader->SetParameters(
+				Scene,
+				View,
+				GlobalResources,
+				View.RayTracingScene,
+				View.ViewUniformBuffer,
+				SceneTexturesUniformBuffer,
+				Scene->Lights,
+				SPPCount, FrameIndependentTemporalSeed, ViewState->VarianceMipTreeDimensions, *ViewState->VarianceMipTree, TileOffset,
+				RadianceRT->GetRenderTargetItem().UAV,
+				SampleCountRT->GetRenderTargetItem().UAV,
+				PixelPositionRT->GetRenderTargetItem().UAV,
+				RayCountPerPixelRT->GetRenderTargetItem().UAV
+			);
+
+			int32 DispatchSizeX = FMath::Min<int32>(TileSizeX, View.ViewRect.Size().X - TileOffset.X);
+			int32 DispatchSizeY = View.ViewRect.Size().Y;
+
+			RHICmdList.RayTraceDispatch(
+				View.RayTracingMaterialPipeline, 
+				RayGenShader->GetRayTracingShader(), 
+				RayTracingSceneRHI, GlobalResources, 
+				DispatchSizeX, DispatchSizeY
+			);
+
+			FIntRect GPURect;
+			GPURect.Min.X = TileOffset.X;
+			GPURect.Min.Y = TileOffset.Y;
+			GPURect.Max.X = TileOffset.X + DispatchSizeX;
+			GPURect.Max.Y = TileOffset.Y + DispatchSizeY;
+
+			if (GPUIndex > 0)
+			{
+				RHICmdList.TransferTexture(RadianceRT->GetRenderTargetItem().TargetableTexture->GetTexture2D(), GPURect, 1, 0, true);
+				RHICmdList.TransferTexture(SampleCountRT->GetRenderTargetItem().TargetableTexture->GetTexture2D(), GPURect, 1, 0, true);
+				RHICmdList.TransferTexture(PixelPositionRT->GetRenderTargetItem().TargetableTexture->GetTexture2D(), GPURect, 1, 0, true);
+				RHICmdList.TransferTexture(RayCountPerPixelRT->GetRenderTargetItem().TargetableTexture->GetTexture2D(), GPURect, 1, 0, true);
+			}
+		}
+
+		{
+			// End mGPU fence
+			FRHIGPUMask GPUMask = FRHIGPUMask::All();
+			FComputeFenceRHIRef Fence = RHICmdList.CreateComputeFence(TEXT("PathTracingRayGen_Fence_End"));
+
+			RHICmdList.TransitionResources(EResourceTransitionAccess::ERWBarrier, EResourceTransitionPipeline::EGfxToGfx, UAVs, UE_ARRAY_COUNT(UAVs), Fence);
+
+			for (uint32 GPUIndex : GPUMask)
+			{
+				if (GPUIndex == MainGPUIndex) 
+					continue;
+
+				SCOPED_GPU_MASK(RHICmdList, FRHIGPUMask::FromIndex(MainGPUIndex));
+
+				RHICmdList.WaitComputeFence(Fence); 
+			}
+		}
+	}
+	else
+	{
+		FIntVector TileOffset;
+		TileOffset.X = 0;
+		TileOffset.Y = 0; 
+
+		RayGenShader->SetParameters(
+			Scene,
+			View,
+			GlobalResources,
+			View.RayTracingScene,
+			View.ViewUniformBuffer,
+			SceneTexturesUniformBuffer,
+			Scene->Lights,
+			SPPCount, FrameIndependentTemporalSeed, ViewState->VarianceMipTreeDimensions, *ViewState->VarianceMipTree, TileOffset,
+			RadianceRT->GetRenderTargetItem().UAV,
+			SampleCountRT->GetRenderTargetItem().UAV,
+			PixelPositionRT->GetRenderTargetItem().UAV,
+			RayCountPerPixelRT->GetRenderTargetItem().UAV
+		);
+	
+		RHICmdList.RayTraceDispatch(View.RayTracingMaterialPipeline, RayGenShader->GetRayTracingShader(), RayTracingSceneRHI, GlobalResources, View.ViewRect.Size().X, View.ViewRect.Size().Y);
+	}
+	
+
 
 	// Save RayTracingIndirect for compositing
 	RHICmdList.CopyToResolveTarget(RadianceRT->GetRenderTargetItem().TargetableTexture, RadianceRT->GetRenderTargetItem().ShaderResourceTexture, FResolveParams());
@@ -497,8 +618,10 @@ void FDeferredShadingSceneRenderer::RenderPathTracing(FRHICommandListImmediate&
 	RHICmdList.CopyToResolveTarget(PixelPositionRT->GetRenderTargetItem().TargetableTexture, PixelPositionRT->GetRenderTargetItem().ShaderResourceTexture, FResolveParams());
 	RHICmdList.CopyToResolveTarget(RayCountPerPixelRT->GetRenderTargetItem().TargetableTexture, RayCountPerPixelRT->GetRenderTargetItem().ShaderResourceTexture, FResolveParams());
 
-	// Run ray counter shader
-	if (SPPCount % CVarPathTracingRayCountFrequency.GetValueOnRenderThread() == 0)
+	// Single GPU for launching compute shaders
+	SCOPED_GPU_MASK(RHICmdList, FRHIGPUMask::FromIndex(MainGPUIndex));
+
+	if ((SPPCount % CVarPathTracingRayCountFrequency.GetValueOnRenderThread() == 0))
 	{
 		ComputeRayCount(RHICmdList, View, RayCountPerPixelRT->GetRenderTargetItem().ShaderResourceTexture);
 	}
@@ -522,7 +645,7 @@ void FDeferredShadingSceneRenderer::RenderPathTracing(FRHICommandListImmediate&
 	ClearUAV(RHICmdList, RadianceSortedGreenRT->GetRenderTargetItem(), FLinearColor::Black);
 	ClearUAV(RHICmdList, RadianceSortedBlueRT->GetRenderTargetItem(), FLinearColor::Black);
 	ClearUAV(RHICmdList, RadianceSortedAlphaRT->GetRenderTargetItem(), FLinearColor::Black);
-	ClearUAV(RHICmdList, SampleCountSortedRT->GetRenderTargetItem(), FLinearColor::Black);
+	ClearUAV(RHICmdList, SampleCountSortedRT->GetRenderTargetItem(), FLinearColor::Black);	
 
 	ComputePathCompaction(
 		RHICmdList,
@@ -543,6 +666,7 @@ void FDeferredShadingSceneRenderer::RenderPathTracing(FRHICommandListImmediate&
 	RHICmdList.CopyToResolveTarget(RadianceSortedAlphaRT->GetRenderTargetItem().TargetableTexture, RadianceSortedAlphaRT->GetRenderTargetItem().ShaderResourceTexture, FResolveParams());
 	RHICmdList.CopyToResolveTarget(SampleCountSortedRT->GetRenderTargetItem().TargetableTexture, SampleCountSortedRT->GetRenderTargetItem().ShaderResourceTexture, FResolveParams());
 
+
 	// Construct render targets for compositing
 	TRefCountPtr<IPooledRenderTarget> OutputRadianceRT;
 	TRefCountPtr<IPooledRenderTarget> OutputSampleCountRT;
@@ -642,6 +766,5 @@ void FDeferredShadingSceneRenderer::RenderPathTracing(FRHICommandListImmediate&
 	VisualizeVarianceMipTree(RHICmdList, View, *ViewState->VarianceMipTree, ViewState->VarianceMipTreeDimensions);
 
 	ResolveSceneColor(RHICmdList);
-
 }
 #endif
\ No newline at end of file
diff --git a/Engine/Source/Runtime/Renderer/Private/PathTracingUniformBuffers.h b/Engine/Source/Runtime/Renderer/Private/PathTracingUniformBuffers.h
--- a/Engine/Source/Runtime/Renderer/Private/PathTracingUniformBuffers.h
+++ b/Engine/Source/Runtime/Renderer/Private/PathTracingUniformBuffers.h
@@ -11,6 +11,7 @@
 
 BEGIN_GLOBAL_SHADER_PARAMETER_STRUCT(FPathTracingData, )
 	SHADER_PARAMETER(uint32, MaxBounces)
+	SHADER_PARAMETER(FIntVector, TileOffset)
 END_GLOBAL_SHADER_PARAMETER_STRUCT()
 
 
