Improvements in multi sample ray tracing reflections. Multi spp is now handled via a high level loop. The reflections shader adds an extra stage that reads the current accumulated color/hit distance on anything beyond pass 0, and only applies the final weighting on the final pass.

#[fyi] yuriy.odonnell

diff --git a/Engine/Shaders/Private/RayTracing/RayTracingReflections.usf b/Engine/Shaders/Private/RayTracing/RayTracingReflections.usf
--- a/Engine/Shaders/Private/RayTracing/RayTracingReflections.usf
+++ b/Engine/Shaders/Private/RayTracing/RayTracingReflections.usf
@@ -20,6 +20,12 @@
 // Reserved material sorting key for rays that can fetch screen-space data instead of tracing rays in the shading pass.
 #define RAY_TRACING_DEFERRED_MATERIAL_KEY_SCREEN_SPACE 0x80000000
 
+// ESamplePhase
+#define SAMPLE_PHASE_MONOLITHIC 0 
+#define SAMPLE_PHASE_INIT       1
+#define SAMPLE_PHASE_ACCUM      2
+#define SAMPLE_PHASE_RESOLVE    3
+
 #include "../Common.ush"
 
 #define SUPPORT_CONTACT_SHADOWS		0
@@ -50,7 +56,7 @@
 #include "RayTracingSpotLight.ush"
 #include "RayTracingPointLight.ush"
 #include "RayTracingCommon.ush"
-#include "RayTracingLightingCommon.ush" 
+#include "RayTracingLightingCommon.ush"
 
 
 // Setup the reflection composite to sample the sky and reflection captures.
@@ -91,6 +97,8 @@ uint MaxUnderCoatBounces;
 uint RenderTileOffsetX;
 uint RenderTileOffsetY;
 uint EnableTranslucency;
+uint SampleMode;
+uint SampleOffset;
 
 float ReflectionMinRayDistance;
 float ReflectionMaxRayDistance;
@@ -121,7 +129,7 @@ float3 GenerateReflectionDirection(inout RandomSequence RandSequence, uint2 Disp
 		float2 E = RandomSequence_GenerateSample2D(RandSequence, DummyVariable);
 
 		// Sobol can only be used on the first sample of the first bounce
-		if (BounceIndex == 0 && SamplePerPixel == 1 && RetryCount == 0)
+		if (BounceIndex == 0 && SamplePerPixel == 1 && RetryCount == 0 && SampleMode == SAMPLE_PHASE_MONOLITHIC)
 		{
 			E = Rand1SPPDenoiserInput(DispatchThreadId);
 			E.y *= 1 - GGX_IMPORTANT_SAMPLE_BIAS;
@@ -225,7 +233,7 @@ RAY_TRACING_ENTRY_RAYGEN(RayTracingReflectionsRGS)
 	uint LinearIndex = PixelCoord.y * View.BufferSizeAndInvSize.x + PixelCoord.x; // TODO(Denoiser): PixelCoord or DispatchThreadId
 	
 	RandomSequence RandSequence;
-	RandomSequence_Initialize(RandSequence, LinearIndex, View.StateFrameIndex);
+	RandomSequence_Initialize(RandSequence, LinearIndex, View.StateFrameIndex + SampleOffset * 16);
 
 	float2 InvBufferSize = View.BufferSizeAndInvSize.zw;
 	float2 UV = (float2(PixelCoord) + 0.5) * InvBufferSize;
@@ -273,7 +281,6 @@ RAY_TRACING_ENTRY_RAYGEN(RayTracingReflectionsRGS)
 
 	if (bIsValidPixel)
 	{
-		for (uint SampleIndex = 0; SampleIndex < LocalSamplesPerPixel; ++SampleIndex)
 		{
 			bool bIsValidSample = false;
 			bool bNeedsCapture = true;
@@ -369,7 +376,7 @@ RAY_TRACING_ENTRY_RAYGEN(RayTracingReflectionsRGS)
 					FixSampleDirectionIfNeeded(GBufferData.WorldNormal, TopLayerRay.Direction);
 				}
 
-				if (BounceIndex == 0 && SampleIndex == 0)
+				if (BounceIndex == 0)
 				{
 					if (DIM_DEFERRED_MATERIAL_MODE == DEFERRED_MATERIAL_MODE_GATHER)
 					{
@@ -811,7 +818,41 @@ RAY_TRACING_ENTRY_RAYGEN(RayTracingReflectionsRGS)
 	}
 	else
 	{
-		if (ReflectedWeight > 0)
+		if (SampleMode == SAMPLE_PHASE_ACCUM || SampleMode == SAMPLE_PHASE_RESOLVE)
+		{
+			// read last results
+			float4 PreviousColor = ColorOutput[DispatchThreadId];
+			float PreviousHitDistance = RayHitDistanceOutput[DispatchThreadId];
+
+#if USE_PREEXPOSURE
+			PreviousColor.xyz /= View.PreExposure;
+#endif
+
+			if (ReflectedWeight > 0)
+			{
+				// Accumulate with prior samples
+				ReflectedColor += PreviousColor;
+				ReflectedWeight += PreviousColor.a;
+
+				if (HitDistance > 0.0 && PreviousHitDistance > 0.0)
+				{
+					HitDistance = min(HitDistance, PreviousHitDistance);
+				}
+				else
+				{
+					HitDistance = max(HitDistance, PreviousHitDistance);
+				}
+			}
+			else
+			{
+				// no data this pass, use accumulated data
+				ReflectedColor = PreviousColor;
+				ReflectedWeight = PreviousColor.a;
+				HitDistance = PreviousHitDistance;
+			}
+		}
+
+		if (ReflectedWeight > 0 && (SampleMode == SAMPLE_PHASE_MONOLITHIC || SampleMode == SAMPLE_PHASE_RESOLVE))
 		{
 			float DistanceFadeRange = LocalMaxRayDistance * 0.1; // Fade out reflections over 10% of the total maximum ray length (arbitrarily chosen)
 			float DistanceFade = bAllowSkySampling ? 1.0 : saturate((LocalMaxRayDistance - HitDistance) * rcp(DistanceFadeRange));
@@ -819,7 +860,7 @@ RAY_TRACING_ENTRY_RAYGEN(RayTracingReflectionsRGS)
 			ReflectedColor *= DistanceFade * max(RoughnessFade, ClearCoatRoughnessFade) * rcp(LocalSamplesPerPixel);
 		}
 
-		if (EnableTranslucency)
+		if (EnableTranslucency && (SampleMode == SAMPLE_PHASE_MONOLITHIC || SampleMode == SAMPLE_PHASE_RESOLVE))
 		{
 			float3 WorldPosition;
 			float3 CameraDirection;
diff --git a/Engine/Source/Runtime/Renderer/Private/RayTracing/RayTracingReflections.cpp b/Engine/Source/Runtime/Renderer/Private/RayTracing/RayTracingReflections.cpp
--- a/Engine/Source/Runtime/Renderer/Private/RayTracing/RayTracingReflections.cpp
+++ b/Engine/Source/Runtime/Renderer/Private/RayTracing/RayTracingReflections.cpp
@@ -154,6 +154,16 @@ static TAutoConsoleVariable<int32> CVarRayTracingReflectionsMaxUnderCoatBounces(
 	ECVF_RenderThreadSafe);
 
 
+
+// ESamplePhase
+enum class ESamplePhase
+{
+	Monlithic = 0, //single pass for all samples
+	Init = 1,      // First sample of the set initialize the accumulators
+	Accum = 2,     // Intermediate sample, accumulate results
+	Resolve = 3    // Final sample, apply weighting
+};
+
 static const int32 GReflectionLightCountMaximum = 64;
 
 class FRayTracingReflectionsRGS : public FGlobalShader
@@ -191,6 +201,8 @@ class FRayTracingReflectionsRGS : public FGlobalShader
 		SHADER_PARAMETER(uint32, RenderTileOffsetX)
 		SHADER_PARAMETER(uint32, RenderTileOffsetY)
 		SHADER_PARAMETER(uint32, EnableTranslucency)
+		SHADER_PARAMETER(int32, SampleMode)
+		SHADER_PARAMETER(int32, SampleOffset)
 
 		SHADER_PARAMETER_SRV(RaytracingAccelerationStructure, TLAS)
 		SHADER_PARAMETER_RDG_TEXTURE(Texture2D, SceneColor)
@@ -345,7 +357,8 @@ void FDeferredShadingSceneRenderer::RenderRayTracingReflections(
 	CommonParameters.MaxUnderCoatBounces = CVarRayTracingReflectionsMaxUnderCoatBounces.GetValueOnRenderThread();
 	CommonParameters.RenderTileOffsetX = 0;
 	CommonParameters.RenderTileOffsetY = 0;
-	CommonParameters.EnableTranslucency = EnableTranslucency;
+	CommonParameters.EnableTranslucency = EnableTranslucency; 
+	CommonParameters.SampleMode = (int32)ESamplePhase::Monlithic;
 
 	CommonParameters.TLAS = View.RayTracingScene.RayTracingSceneRHI->GetShaderResourceView();
 	CommonParameters.ViewUniformBuffer = View.ViewUniformBuffer;
@@ -377,117 +390,128 @@ void FDeferredShadingSceneRenderer::RenderRayTracingReflections(
 	CommonParameters.ReflectionCapture = View.ReflectionCaptureUniformBuffer;
 	CommonParameters.Forward = View.ForwardLightingResources->ForwardLightDataUniformBuffer;
 
-	for (uint32 PassIndex = 0; PassIndex < NumPasses; ++PassIndex)
+	for (int32 SamplePassIndex = 0; SamplePassIndex < SamplePerPixel; SamplePassIndex++)
 	{
-		FRayTracingReflectionsRGS::FParameters* PassParameters = GraphBuilder.AllocParameters<FRayTracingReflectionsRGS::FParameters>();
-		*PassParameters = CommonParameters;
-
-		const EDeferredMaterialMode DeferredMaterialMode = DeferredMaterialModes[PassIndex];
-
-		if (DeferredMaterialMode != EDeferredMaterialMode::None)
+		if (SamplePerPixel > 1)
 		{
-			if (DeferredMaterialMode == EDeferredMaterialMode::Gather)
-			{
-				FRDGBufferDesc Desc = FRDGBufferDesc::CreateStructuredDesc(sizeof(FDeferredMaterialPayload), DeferredMaterialBufferNumElements);
-				DeferredMaterialBuffer = GraphBuilder.CreateBuffer(Desc, TEXT("RayTracingReflectionsMaterialBuffer"));
-			}
-
-			PassParameters->MaterialBuffer = GraphBuilder.CreateUAV(DeferredMaterialBuffer);
+			CommonParameters.SampleMode = (int32)ESamplePhase::Accum;
+			CommonParameters.SampleMode = SamplePassIndex == 0 ? (int32)ESamplePhase::Init : CommonParameters.SampleMode;
+			CommonParameters.SampleMode = SamplePassIndex == (SamplePerPixel - 1) ? (int32)ESamplePhase::Resolve : CommonParameters.SampleMode;
 		}
+		CommonParameters.SampleOffset = SamplePassIndex;
 
-		FRayTracingReflectionsRGS::FPermutationDomain PermutationVector;
-		PermutationVector.Set<FRayTracingReflectionsRGS::FDeferredMaterialMode>(DeferredMaterialMode);
-		PermutationVector.Set<FRayTracingReflectionsRGS::FHybrid>(bHybridReflections);
-		PermutationVector.Set<FRayTracingReflectionsRGS::FEnableTwoSidedGeometryForShadowDim>(EnableRayTracingShadowTwoSidedGeometry());
-		auto RayGenShader = View.ShaderMap->GetShader<FRayTracingReflectionsRGS>(PermutationVector);
+		for (uint32 PassIndex = 0; PassIndex < NumPasses; ++PassIndex)
+		{
+			FRayTracingReflectionsRGS::FParameters* PassParameters = GraphBuilder.AllocParameters<FRayTracingReflectionsRGS::FParameters>();
+			*PassParameters = CommonParameters;
 
-		ClearUnusedGraphResources(RayGenShader, PassParameters);
+			const EDeferredMaterialMode DeferredMaterialMode = DeferredMaterialModes[PassIndex];
 
-		if (DeferredMaterialMode == EDeferredMaterialMode::Gather)
-		{
-			GraphBuilder.AddPass(
-				RDG_EVENT_NAME("ReflectionRayTracingGatherMaterials %dx%d", TileAlignedResolution.X, TileAlignedResolution.Y),
-				PassParameters,
-				ERDGPassFlags::Compute,
-				[PassParameters, this, &View, RayGenShader, TileAlignedResolution](FRHICommandList& RHICmdList)
+			if (DeferredMaterialMode != EDeferredMaterialMode::None)
 			{
-				FRayTracingPipelineState* Pipeline = BindRayTracingDeferredMaterialGatherPipeline(RHICmdList, View, RayGenShader->GetRayTracingShader());
+				if (DeferredMaterialMode == EDeferredMaterialMode::Gather)
+				{
+					FRDGBufferDesc Desc = FRDGBufferDesc::CreateStructuredDesc(sizeof(FDeferredMaterialPayload), DeferredMaterialBufferNumElements);
+					DeferredMaterialBuffer = GraphBuilder.CreateBuffer(Desc, TEXT("RayTracingReflectionsMaterialBuffer"));
+				}
 
-				FRayTracingShaderBindingsWriter GlobalResources;
-				SetShaderParameters(GlobalResources, RayGenShader, *PassParameters);
+				PassParameters->MaterialBuffer = GraphBuilder.CreateUAV(DeferredMaterialBuffer);
+			}
 
-				FRHIRayTracingScene* RayTracingSceneRHI = View.RayTracingScene.RayTracingSceneRHI;
-				RHICmdList.RayTraceDispatch(Pipeline, RayGenShader->GetRayTracingShader(), RayTracingSceneRHI, GlobalResources, TileAlignedResolution.X, TileAlignedResolution.Y);
-			});
+			FRayTracingReflectionsRGS::FPermutationDomain PermutationVector;
+			PermutationVector.Set<FRayTracingReflectionsRGS::FDeferredMaterialMode>(DeferredMaterialMode);
+			PermutationVector.Set<FRayTracingReflectionsRGS::FHybrid>(bHybridReflections);
+			PermutationVector.Set<FRayTracingReflectionsRGS::FEnableTwoSidedGeometryForShadowDim>(EnableRayTracingShadowTwoSidedGeometry());
+			auto RayGenShader = View.ShaderMap->GetShader<FRayTracingReflectionsRGS>(PermutationVector);
 
-			// A material sorting pass
-			const uint32 SortSize = CVarRayTracingReflectionsSortSize.GetValueOnRenderThread();
-			if (SortSize)
-			{
-				SortDeferredMaterials(GraphBuilder, View, SortSize, DeferredMaterialBufferNumElements, DeferredMaterialBuffer);
-			}
-		}
-		else
-		{
-			// Add optional tiling behavior to avoid TDR events in expensive passes
-			int32 RenderTileSize = CVarRayTracingReflectionsRenderTileSize.GetValueOnRenderThread();
-			if (NumPasses > 1 || RenderTileSize <= 0)
+			ClearUnusedGraphResources(RayGenShader, PassParameters);
+
+			if (DeferredMaterialMode == EDeferredMaterialMode::Gather)
 			{
 				GraphBuilder.AddPass(
-					RDG_EVENT_NAME("ReflectionRayTracing(spp=%d) %dx%d", SamplePerPixel, RayTracingResolution.X, RayTracingResolution.Y),
+					RDG_EVENT_NAME("ReflectionRayTracingGatherMaterials %dx%d", TileAlignedResolution.X, TileAlignedResolution.Y),
 					PassParameters,
 					ERDGPassFlags::Compute,
-					[PassParameters, this, &View, RayGenShader, RayTracingResolution, DeferredMaterialBufferNumElements, DeferredMaterialMode](FRHICommandList& RHICmdList)
+					[PassParameters, this, &View, RayGenShader, TileAlignedResolution](FRHICommandList& RHICmdList)
 				{
+					FRayTracingPipelineState* Pipeline = BindRayTracingDeferredMaterialGatherPipeline(RHICmdList, View, RayGenShader->GetRayTracingShader());
+
 					FRayTracingShaderBindingsWriter GlobalResources;
 					SetShaderParameters(GlobalResources, RayGenShader, *PassParameters);
 
 					FRHIRayTracingScene* RayTracingSceneRHI = View.RayTracingScene.RayTracingSceneRHI;
-
-					if (DeferredMaterialMode == EDeferredMaterialMode::Shade)
-					{
-						// Shading pass for sorted materials uses 1D dispatch over all elements in the material buffer.
-						// This can be reduced to the number of output pixels if sorting pass guarantees that all invalid entries are moved to the end.
-						RHICmdList.RayTraceDispatch(View.RayTracingMaterialPipeline, RayGenShader->GetRayTracingShader(), RayTracingSceneRHI, GlobalResources, DeferredMaterialBufferNumElements, 1);
-					}
-					else // EDeferredMaterialMode::None
-					{
-						RHICmdList.RayTraceDispatch(View.RayTracingMaterialPipeline, RayGenShader->GetRayTracingShader(), RayTracingSceneRHI, GlobalResources, RayTracingResolution.X, RayTracingResolution.Y);
-					}
+					RHICmdList.RayTraceDispatch(Pipeline, RayGenShader->GetRayTracingShader(), RayTracingSceneRHI, GlobalResources, TileAlignedResolution.X, TileAlignedResolution.Y);
 				});
+
+				// A material sorting pass
+				const uint32 SortSize = CVarRayTracingReflectionsSortSize.GetValueOnRenderThread();
+				if (SortSize)
+				{
+					SortDeferredMaterials(GraphBuilder, View, SortSize, DeferredMaterialBufferNumElements, DeferredMaterialBuffer);
+				}
 			}
 			else
 			{
-				RenderTileSize = FMath::Max(RenderTileSize, 32);
-				int32 NumTilesX = FMath::DivideAndRoundUp(RayTracingResolution.X, RenderTileSize);
-				int32 NumTilesY = FMath::DivideAndRoundUp(RayTracingResolution.Y, RenderTileSize);
-				for (int32 Y = 0; Y < NumTilesY; ++Y)
+				// Add optional tiling behavior to avoid TDR events in expensive passes
+				int32 RenderTileSize = CVarRayTracingReflectionsRenderTileSize.GetValueOnRenderThread();
+				if (NumPasses > 1 || RenderTileSize <= 0)
 				{
-					FRayTracingReflectionsRGS::FParameters* TilePassParameters = PassParameters;
-					for (int32 X = 0; X < NumTilesX; ++X)
+					GraphBuilder.AddPass(
+						RDG_EVENT_NAME("ReflectionRayTracing(spp=%d) %dx%d", SamplePerPixel, RayTracingResolution.X, RayTracingResolution.Y),
+						PassParameters,
+						ERDGPassFlags::Compute,
+						[PassParameters, this, &View, RayGenShader, RayTracingResolution, DeferredMaterialBufferNumElements, DeferredMaterialMode](FRHICommandList& RHICmdList)
 					{
-						if (X > 0 || Y > 0)
+						FRayTracingShaderBindingsWriter GlobalResources;
+						SetShaderParameters(GlobalResources, RayGenShader, *PassParameters);
+
+						FRHIRayTracingScene* RayTracingSceneRHI = View.RayTracingScene.RayTracingSceneRHI;
+
+						if (DeferredMaterialMode == EDeferredMaterialMode::Shade)
 						{
-							TilePassParameters = GraphBuilder.AllocParameters<FRayTracingReflectionsRGS::FParameters>();
-							*TilePassParameters = *PassParameters;
-							TilePassParameters->RenderTileOffsetX = X * RenderTileSize;
-							TilePassParameters->RenderTileOffsetY = Y * RenderTileSize;
+							// Shading pass for sorted materials uses 1D dispatch over all elements in the material buffer.
+							// This can be reduced to the number of output pixels if sorting pass guarantees that all invalid entries are moved to the end.
+							RHICmdList.RayTraceDispatch(View.RayTracingMaterialPipeline, RayGenShader->GetRayTracingShader(), RayTracingSceneRHI, GlobalResources, DeferredMaterialBufferNumElements, 1);
 						}
-
-						int32 DispatchSizeX = FMath::Min<int32>(RenderTileSize, RayTracingResolution.X - TilePassParameters->RenderTileOffsetX);
-						int32 DispatchSizeY = FMath::Min<int32>(RenderTileSize, RayTracingResolution.Y - TilePassParameters->RenderTileOffsetY);
-						GraphBuilder.AddPass(
-							RDG_EVENT_NAME("ReflectionRayTracing(spp=%d) %dx%d", SamplePerPixel, DispatchSizeX, DispatchSizeY),
-							TilePassParameters,
-							ERDGPassFlags::Compute,
-							[TilePassParameters, this, &View, RayGenShader, DispatchSizeX, DispatchSizeY, DeferredMaterialBufferNumElements, DeferredMaterialMode](FRHICommandList& RHICmdList)
+						else // EDeferredMaterialMode::None
 						{
-							FRayTracingShaderBindingsWriter GlobalResources;
-							SetShaderParameters(GlobalResources, RayGenShader, *TilePassParameters);
-
-							FRHIRayTracingScene* RayTracingSceneRHI = View.RayTracingScene.RayTracingSceneRHI;
-							RHICmdList.RayTraceDispatch(View.RayTracingMaterialPipeline, RayGenShader->GetRayTracingShader(), RayTracingSceneRHI, GlobalResources, DispatchSizeX, DispatchSizeY);
-						});
+							RHICmdList.RayTraceDispatch(View.RayTracingMaterialPipeline, RayGenShader->GetRayTracingShader(), RayTracingSceneRHI, GlobalResources, RayTracingResolution.X, RayTracingResolution.Y);
+						}
+					});
+				}
+				else
+				{
+					RenderTileSize = FMath::Max(RenderTileSize, 32);
+					int32 NumTilesX = FMath::DivideAndRoundUp(RayTracingResolution.X, RenderTileSize);
+					int32 NumTilesY = FMath::DivideAndRoundUp(RayTracingResolution.Y, RenderTileSize);
+					for (int32 Y = 0; Y < NumTilesY; ++Y)
+					{
+						FRayTracingReflectionsRGS::FParameters* TilePassParameters = PassParameters;
+						for (int32 X = 0; X < NumTilesX; ++X)
+						{
+							if (X > 0 || Y > 0)
+							{
+								TilePassParameters = GraphBuilder.AllocParameters<FRayTracingReflectionsRGS::FParameters>();
+								*TilePassParameters = *PassParameters;
+								TilePassParameters->RenderTileOffsetX = X * RenderTileSize;
+								TilePassParameters->RenderTileOffsetY = Y * RenderTileSize;
+							}
+
+							int32 DispatchSizeX = FMath::Min<int32>(RenderTileSize, RayTracingResolution.X - TilePassParameters->RenderTileOffsetX);
+							int32 DispatchSizeY = FMath::Min<int32>(RenderTileSize, RayTracingResolution.Y - TilePassParameters->RenderTileOffsetY);
+							GraphBuilder.AddPass(
+								RDG_EVENT_NAME("ReflectionRayTracing(spp=%d) %dx%d", SamplePerPixel, DispatchSizeX, DispatchSizeY),
+								TilePassParameters,
+								ERDGPassFlags::Compute,
+								[TilePassParameters, this, &View, RayGenShader, DispatchSizeX, DispatchSizeY, DeferredMaterialBufferNumElements, DeferredMaterialMode](FRHICommandList& RHICmdList)
+							{
+								FRayTracingShaderBindingsWriter GlobalResources;
+								SetShaderParameters(GlobalResources, RayGenShader, *TilePassParameters);
+
+								FRHIRayTracingScene* RayTracingSceneRHI = View.RayTracingScene.RayTracingSceneRHI;
+								RHICmdList.RayTraceDispatch(View.RayTracingMaterialPipeline, RayGenShader->GetRayTracingShader(), RayTracingSceneRHI, GlobalResources, DispatchSizeX, DispatchSizeY);
+							});
+						}
 					}
 				}
 			}
