Fix virtual texture feedback for multiple views
Also remove virtual texture time slice option which was never used/needed

diff --git a/Engine/Source/Runtime/Renderer/Private/DeferredShadingRenderer.cpp b/Engine/Source/Runtime/Renderer/Private/DeferredShadingRenderer.cpp
--- a/Engine/Source/Runtime/Renderer/Private/DeferredShadingRenderer.cpp
+++ b/Engine/Source/Runtime/Renderer/Private/DeferredShadingRenderer.cpp
@@ -2314,7 +2314,13 @@ void FDeferredShadingSceneRenderer::Render(FRHICommandListImmediate& RHICmdList)
 	{
 		SCOPED_GPU_STAT(RHICmdList, VirtualTextureUpdate);
 		// No pass after this can make VT page requests
-		SceneContext.VirtualTextureFeedback.TransferGPUToCPU(RHICmdList, Views[0].ViewRect);
+		TArray<FIntRect, TInlineAllocator<FVirtualTextureFeedback::MaxRectPerTarget>> ViewRects;
+		ViewRects.AddUninitialized(Views.Num());
+		for (int32 ViewIndex = 0; ViewIndex < Views.Num(); ++ViewIndex)
+		{
+			ViewRects[ViewIndex] = Views[ViewIndex].ViewRect;
+		}
+		SceneContext.VirtualTextureFeedback.TransferGPUToCPU(RHICmdList, ViewRects);
 	}
 
 #if RHI_RAYTRACING
diff --git a/Engine/Source/Runtime/Renderer/Private/VT/UniquePageList.h b/Engine/Source/Runtime/Renderer/Private/VT/UniquePageList.h
--- a/Engine/Source/Runtime/Renderer/Private/VT/UniquePageList.h
+++ b/Engine/Source/Runtime/Renderer/Private/VT/UniquePageList.h
@@ -28,21 +28,27 @@ private:
 		MaxUniquePages	= 4*1024,
 	};
 
+	bool bInitialized;
+	uint32 NumPages;
+	uint32 MaxNumCollisions;
 	uint16 HashIndices[HashSize];
 	uint32 Pages[ MaxUniquePages ];
 	uint16 Counts[ MaxUniquePages ];
-	uint32 NumPages;
-	uint32 MaxNumCollisions;
 };
 
 FUniquePageList::FUniquePageList()
-	: NumPages( 0 )
-	, MaxNumCollisions(0u)
+	: bInitialized( false )
+	, NumPages( 0 )
+	, MaxNumCollisions( 0 )
 {}
 
 void FUniquePageList::Initialize()
 {
-	FMemory::Memset(HashIndices, 0xff);
+	if (!bInitialized)
+	{
+		FMemory::Memset(HashIndices, 0xff);
+		bInitialized = true;
+	}
 }
 
 void FUniquePageList::Add( uint32 Page, uint32 Count )
diff --git a/Engine/Source/Runtime/Renderer/Private/VT/VirtualTextureFeedback.cpp b/Engine/Source/Runtime/Renderer/Private/VT/VirtualTextureFeedback.cpp
--- a/Engine/Source/Runtime/Renderer/Private/VT/VirtualTextureFeedback.cpp
+++ b/Engine/Source/Runtime/Renderer/Private/VT/VirtualTextureFeedback.cpp
@@ -188,7 +188,7 @@ void FVirtualTextureFeedback::MakeSnapshot(const FVirtualTextureFeedback& Snapsh
 	}
 }
 
-void FVirtualTextureFeedback::TransferGPUToCPU( FRHICommandListImmediate& RHICmdList, FIntRect const& Rect )
+void FVirtualTextureFeedback::TransferGPUToCPU( FRHICommandListImmediate& RHICmdList, TArrayView<FIntRect> const& ViewRects )
 {
 	RHICmdList.TransitionResource( EResourceTransitionAccess::EReadable, EResourceTransitionPipeline::EGfxToGfx, FeedbackTextureGPU->GetRenderTargetItem().UAV );
 
@@ -205,8 +205,14 @@ void FVirtualTextureFeedback::TransferGPUToCPU( FRHICommandListImmediate& RHICmd
 	}
 
 	FFeedBackItem& FeedbackEntryCPU = FeedbackTextureCPU[GPUWriteIndex];
-	FeedbackEntryCPU.Rect.Min = FIntPoint(Rect.Min.X / 16, Rect.Min.Y / 16);
-	FeedbackEntryCPU.Rect.Max = FIntPoint((Rect.Max.X + 15) / 16, (Rect.Max.Y + 15) / 16);
+	FeedbackEntryCPU.NumRects = FMath::Min((int32)MaxRectPerTarget, ViewRects.Num());
+	for (int32 RectIndex = 0; RectIndex < FeedbackEntryCPU.NumRects; ++RectIndex)
+	{
+		FIntRect const& Rect = ViewRects[RectIndex];
+		//todo[vt]: Value of 16 has to match r.vt.FeedbackFactor
+		FeedbackEntryCPU.Rects[RectIndex].Min = FIntPoint(Rect.Min.X / 16, Rect.Min.Y / 16);
+		FeedbackEntryCPU.Rects[RectIndex].Max = FIntPoint((Rect.Max.X + 15) / 16, (Rect.Max.Y + 15) / 16);
+	}
 
 	const FPooledRenderTargetDesc Desc(FPooledRenderTargetDesc::Create2DDesc(Size, PF_R32_UINT, FClearValueBinding::None, TexCreate_CPUReadback | TexCreate_HideInVisualizeTexture, TexCreate_None, false));
 
@@ -250,7 +256,12 @@ bool FVirtualTextureFeedback::Map(FRHICommandListImmediate& RHICmdList, MapResul
 		SCOPED_GPU_MASK(RHICmdList, FeedbackEntryCPU.GPUMask);
 
 		OutResult.MapHandle = CPUReadIndex;
-		OutResult.Rect = FeedbackEntryCPU.Rect;
+		
+		OutResult.NumRects = FeedbackEntryCPU.NumRects;
+		for (int32 i = 0; i < FeedbackEntryCPU.NumRects; ++i)
+		{
+			OutResult.Rects[i] = FeedbackEntryCPU.Rects[i];
+		}
 
 		int32 LockHeight = 0;
 		RHICmdList.MapStagingSurface(FeedbackEntryCPU.TextureCPU->GetRenderTargetItem().ShaderResourceTexture, FeedBackFences->GetMapFence(CPUReadIndex), *(void**)&OutResult.Buffer, OutResult.Pitch, LockHeight);
diff --git a/Engine/Source/Runtime/Renderer/Private/VT/VirtualTextureFeedback.h b/Engine/Source/Runtime/Renderer/Private/VT/VirtualTextureFeedback.h
--- a/Engine/Source/Runtime/Renderer/Private/VT/VirtualTextureFeedback.h
+++ b/Engine/Source/Runtime/Renderer/Private/VT/VirtualTextureFeedback.h
@@ -21,6 +21,7 @@ public:
 	~FVirtualTextureFeedback();
 
 	static const uint32 TargetCapacity = 4u;
+	static const uint32 MaxRectPerTarget = 4u;
 
 	TRefCountPtr< IPooledRenderTarget >	FeedbackTextureGPU;
 
@@ -28,15 +29,16 @@ public:
 	{
 		int32 MapHandle;
 		uint32* RESTRICT Buffer;
-		FIntRect Rect;
 		int32 Pitch;
+		int32 NumRects;
+		FIntRect Rects[MaxRectPerTarget];
 	};
 
 	void			CreateResourceGPU(FRHICommandListImmediate& RHICmdList, FIntPoint InSize);
 	void			ReleaseResources();
 	void			MakeSnapshot(const FVirtualTextureFeedback& SnapshotSource);
 
-	void			TransferGPUToCPU(FRHICommandListImmediate& RHICmdList, FIntRect const& Rect);
+	void			TransferGPUToCPU(FRHICommandListImmediate& RHICmdList, TArrayView<FIntRect> const& ViewRects);
 
 	uint32			GetPendingTargetCount() const { return PendingTargetCount; }
 
@@ -47,7 +49,8 @@ public:
 private:
 	struct FFeedBackItem
 	{
-		FIntRect Rect;
+		int32 NumRects;
+		FIntRect Rects[MaxRectPerTarget];
 		TRefCountPtr< IPooledRenderTarget > TextureCPU;
 		FRHIGPUMask GPUMask;
 	};
diff --git a/Engine/Source/Runtime/Renderer/Private/VT/VirtualTextureSystem.cpp b/Engine/Source/Runtime/Renderer/Private/VT/VirtualTextureSystem.cpp
--- a/Engine/Source/Runtime/Renderer/Private/VT/VirtualTextureSystem.cpp
+++ b/Engine/Source/Runtime/Renderer/Private/VT/VirtualTextureSystem.cpp
@@ -54,6 +54,13 @@ DECLARE_MEMORY_STAT_POOL(TEXT("Total Pagetable Memory"), STAT_TotalPagetableMemo
 
 DECLARE_GPU_STAT( VirtualTexture );
 
+static TAutoConsoleVariable<int32> CVarVTVerbose(
+	TEXT("r.VT.Verbose"),
+	0,
+	TEXT("Be pedantic about certain things that shouln't occur unless something is wrong. This may cause a lot of logspam 100's of lines per frame."),
+	ECVF_RenderThreadSafe
+);
+
 static TAutoConsoleVariable<int32> CVarVTEnableFeedBack(
 	TEXT("r.VT.EnableFeedBack"),
 	1,
@@ -61,23 +68,22 @@ static TAutoConsoleVariable<int32> CVarVTEnableFeedBack(
 	ECVF_RenderThreadSafe
 );
 
-static TAutoConsoleVariable<int32> CVarVTVerbose(
-	TEXT("r.VT.Verbose"),
+static TAutoConsoleVariable<int32> CVarVTParallelFeedbackTasks(
+	TEXT("r.VT.ParallelFeedbackTasks"),
 	0,
-	TEXT("Be pedantic about certain things that shouln't occur unless something is wrong. This may cause a lot of logspam 100's of lines per frame."),
+	TEXT("Use worker threads for virtual texture feedback tasks."),
 	ECVF_RenderThreadSafe
 );
-
 static TAutoConsoleVariable<int32> CVarVTNumFeedbackTasks(
 	TEXT("r.VT.NumFeedbackTasks"),
 	1,
-	TEXT("Number of tasks to create to process virtual texture updates."),
+	TEXT("Number of tasks to create to read virtual texture feedback."),
 	ECVF_RenderThreadSafe
 );
 static TAutoConsoleVariable<int32> CVarVTNumGatherTasks(
 	TEXT("r.VT.NumGatherTasks"),
 	1,
-	TEXT("Number of tasks to create to process virtual texture updates."),
+	TEXT("Number of tasks to create to combine virtual texture feedback."),
 	ECVF_RenderThreadSafe
 );
 static TAutoConsoleVariable<int32> CVarVTPageUpdateFlushCount(
@@ -86,14 +92,6 @@ static TAutoConsoleVariable<int32> CVarVTPageUpdateFlushCount(
 	TEXT("Number of page updates to buffer before attempting to flush by taking a lock."),
 	ECVF_RenderThreadSafe
 );
-static const int32 MaxNumTasks = 16;
-
-static TAutoConsoleVariable<float> CVarVTPageUpdateTimeSlice(
-	TEXT("r.VT.PageUpdateTimeSlice"),
-	0.f,
-	TEXT("Assumed frame time to slice MaxUploadsPerFrame. 0 is disabled."),
-	ECVF_RenderThreadSafe
-);
 
 static FORCEINLINE uint32 EncodePage(uint32 ID, uint32 vLevel, uint32 vTileX, uint32 vTileY)
 {
@@ -144,10 +142,15 @@ public:
 
 	FFeedbackAnalysisParameters Parameters;
 
+	static void DoTask(FFeedbackAnalysisParameters& InParams)
+	{
+		InParams.UniquePageList->Initialize();
+		InParams.System->FeedbackAnalysisTask(InParams);
+	}
+
 	void DoTask(ENamedThreads::Type CurrentThread, const FGraphEventRef& MyCompletionGraphEvent)
 	{
-		Parameters.UniquePageList->Initialize();
-		Parameters.System->FeedbackAnalysisTask(Parameters);
+		DoTask(Parameters);
 	}
 
 	static ESubsequentsMode::Type GetSubsequentsMode() { return ESubsequentsMode::TrackSubsequents; }
@@ -882,35 +885,6 @@ void FVirtualTextureSystem::FeedbackAnalysisTask(const FFeedbackAnalysisParamete
 	}
 }
 
-/** Helper to reduce MaxUploadsPerFrame at higher frame rates. */
-class FVirtualTextureTimeSlicer
-{
-public:
-	/** Call once per frame to update current time deltas. */
-	void Update()
-	{
-		const float Time = FApp::GetCurrentTime() - GStartTime;
-		TimeDelta = Time - LastTime;
-		LastTime = Time;
-	}
-
-	/** Slice a value for the current time slice according to the BaseTimeSlice. This clamps for time deltas greater than the BaseTimeSlice. */
-	int32 Slice(int32 InValue, float BaseTimeSlice)
-	{
-		if (BaseTimeSlice <= 0.f)
-		{
-			return InValue;
-		}
-		
-		const float ClampedRatio = FMath::Clamp(TimeDelta / BaseTimeSlice, 0.0001f, 1.f);
-		return FMath::CeilToInt(ClampedRatio * InValue);
-	}
-
-private:
-	float LastTime = 0.f;
-	float TimeDelta = 0.f;
-};
-
 void FVirtualTextureSystem::Update(FRHICommandListImmediate& RHICmdList, ERHIFeatureLevel::Type FeatureLevel, FScene* Scene)
 {
 	check(IsInRenderingThread());
@@ -919,9 +893,6 @@ void FVirtualTextureSystem::Update(FRHICommandListImmediate& RHICmdList, ERHIFea
 	SCOPE_CYCLE_COUNTER(STAT_VirtualTextureSystem_Update);
 	SCOPED_GPU_STAT(RHICmdList, VirtualTexture);
 	
-	static FVirtualTextureTimeSlicer TimeSlicer;
-	TimeSlicer.Update();
-
 	if (bFlushCaches)
 	{
 		SCOPE_CYCLE_COUNTER(STAT_FlushCache);
@@ -975,6 +946,7 @@ void FVirtualTextureSystem::Update(FRHICommandListImmediate& RHICmdList, ERHIFea
 
 		// Gather all outstanding feedback buffers
 		uint32 FeedbackBufferCount = 0;
+		uint32 FeedbackRectCount = 0;
 		FVirtualTextureFeedback::MapResult MappedFeedbackBuffers[FVirtualTextureFeedback::TargetCapacity];
 
 		if (CVarVTEnableFeedBack.GetValueOnRenderThread())
@@ -987,66 +959,78 @@ void FVirtualTextureSystem::Update(FRHICommandListImmediate& RHICmdList, ERHIFea
 			uint32 MaxFeedbackTargetCount = 1;
 #endif
 
-			for (FeedbackBufferCount = 0; FeedbackBufferCount < MaxFeedbackTargetCount; ++FeedbackBufferCount)
+			while (FeedbackBufferCount < MaxFeedbackTargetCount)
 			{
 				if (!SceneContext.VirtualTextureFeedback.Map(RHICmdList, MappedFeedbackBuffers[FeedbackBufferCount]))
 				{
 					break;
 				}
+				
+				FeedbackRectCount += MappedFeedbackBuffers[FeedbackBufferCount].NumRects;
+				FeedbackBufferCount++;
 			}
-
-			const int32 PendingFeedbackBuffers = SceneContext.VirtualTextureFeedback.GetPendingTargetCount();
 		}
 
 		// Create tasks to read all the buffers
 		uint32 NumFeedbackTasks = 0u;
 
 		FFeedbackAnalysisParameters FeedbackAnalysisParameters[MaxNumTasks];
-		const uint32 MaxNumFeedbackTasks = FMath::Clamp(CVarVTNumFeedbackTasks.GetValueOnRenderThread(), 1, MaxNumTasks);
-			
-		for (uint32 i=0; i<FeedbackBufferCount; ++i)
+		static_assert(MaxNumTasks >= FVirtualTextureFeedback::TargetCapacity * FVirtualTextureFeedback::MaxRectPerTarget, "MaxNumTasks is too small");
+
+		const uint32 MaxNumFeedbackTasks = FMath::Clamp((uint32)CVarVTNumFeedbackTasks.GetValueOnRenderThread(), 1u, MaxNumTasks);
+		const uint32 FeedbackTasksPerRect = FMath::Max(MaxNumFeedbackTasks / FMath::Max(FeedbackRectCount, 1u), 1u);
+
+		for (uint32 FeedbackBufferIndex = 0; FeedbackBufferIndex < FeedbackBufferCount; ++FeedbackBufferIndex)
 		{
-			FVirtualTextureFeedback::MapResult& FeedbackBuffer = MappedFeedbackBuffers[i];
+			FVirtualTextureFeedback::MapResult const& FeedbackBuffer = MappedFeedbackBuffers[FeedbackBufferIndex];
 
-			// Give each task a section of a feedback buffer to analyze
-			//todo[vt]: For buffers of different sizes we will have different task payload sizes which is not efficient
-			const uint32 FeedbackTasksPerBuffer = FMath::Max(MaxNumFeedbackTasks / FeedbackBufferCount, 1u);
-			const uint32 FeedbackRowsPerTask = FMath::DivideAndRoundUp((uint32)FeedbackBuffer.Rect.Size().Y, FeedbackTasksPerBuffer);
-			const uint32 NumRows = (uint32)FeedbackBuffer.Rect.Size().Y;
-			
-			uint32 CurrentRow = 0;
-			while (CurrentRow < NumRows)
+			for (int32 RectIndex = 0; RectIndex < FeedbackBuffer.NumRects; ++RectIndex)
 			{
-				const uint32 CurrentHeight = FMath::Min(FeedbackRowsPerTask, NumRows - CurrentRow);
-				if (CurrentHeight > 0u)
+				FIntRect const& Rect = FeedbackBuffer.Rects[RectIndex];
+
+				// Give each task a section of a feedback rect to analyze
+				//todo[vt]: For buffers/rects of different sizes we will have different task payload sizes which is not efficient
+				const uint32 FeedbackRowsPerTask = FMath::DivideAndRoundUp((uint32)Rect.Size().Y, FeedbackTasksPerRect);
+				const uint32 NumRows = (uint32)Rect.Size().Y;
+
+				uint32 CurrentRow = 0;
+				while (CurrentRow < NumRows)
 				{
-					const uint32 TaskIndex = NumFeedbackTasks++;
-					FFeedbackAnalysisParameters& Params = FeedbackAnalysisParameters[TaskIndex];
-					Params.System = this;
-					if (TaskIndex == 0u)
-					{
-						Params.UniquePageList = MergedUniquePageList;
-					}
-					else
+					const uint32 CurrentHeight = FMath::Min(FeedbackRowsPerTask, NumRows - CurrentRow);
+					if (CurrentHeight > 0u)
 					{
-						Params.UniquePageList = new(MemStack) FUniquePageList;
+						const uint32 TaskIndex = NumFeedbackTasks++;
+						FFeedbackAnalysisParameters& Params = FeedbackAnalysisParameters[TaskIndex];
+						Params.System = this;
+						if (TaskIndex == 0u)
+						{
+							Params.UniquePageList = MergedUniquePageList;
+						}
+						else
+						{
+							Params.UniquePageList = new(MemStack) FUniquePageList;
+						}
+						Params.FeedbackBuffer = FeedbackBuffer.Buffer + (Rect.Min.Y + CurrentRow) * FeedbackBuffer.Pitch + Rect.Min.X;
+						Params.FeedbackWidth = Rect.Size().X;
+						Params.FeedbackHeight = CurrentHeight;
+						Params.FeedbackPitch = FeedbackBuffer.Pitch;
+						CurrentRow += CurrentHeight;
 					}
-					Params.FeedbackBuffer = FeedbackBuffer.Buffer + (FeedbackBuffer.Rect.Min.Y + CurrentRow) * FeedbackBuffer.Pitch + FeedbackBuffer.Rect.Min.X;
-					Params.FeedbackWidth = FeedbackBuffer.Rect.Size().X;
-					Params.FeedbackHeight = CurrentHeight;
-					Params.FeedbackPitch = FeedbackBuffer.Pitch;
-					CurrentRow += CurrentHeight;
 				}
 			}
 		}
 
 		// Kick the tasks
+		const bool bParallelTasks = CVarVTParallelFeedbackTasks.GetValueOnRenderThread() != 0;
+		const int32 LocalFeedbackTaskCount = bParallelTasks ? 1 : NumFeedbackTasks;
+		const int32 WorkerFeedbackTaskCount = NumFeedbackTasks - LocalFeedbackTaskCount;
+
 		FGraphEventArray Tasks;
-		if(NumFeedbackTasks > 1u)
+		if(WorkerFeedbackTaskCount > 0)
 		{
 			SCOPE_CYCLE_COUNTER(STAT_ProcessRequests_SubmitTasks);
-			Tasks.Reserve(NumFeedbackTasks - 1u);
-			for (uint32 TaskIndex = 1u; TaskIndex < NumFeedbackTasks; ++TaskIndex)
+			Tasks.Reserve(WorkerFeedbackTaskCount);
+			for (uint32 TaskIndex = LocalFeedbackTaskCount; TaskIndex < NumFeedbackTasks; ++TaskIndex)
 			{
 				Tasks.Add(TGraphTask<FFeedbackAnalysisTask>::CreateTask().ConstructAndDispatchWhenReady(FeedbackAnalysisParameters[TaskIndex]));
 			}
@@ -1056,10 +1040,11 @@ void FVirtualTextureSystem::Update(FRHICommandListImmediate& RHICmdList, ERHIFea
 		{
 			SCOPE_CYCLE_COUNTER(STAT_FeedbackAnalysis);
 
-			FeedbackAnalysisTask(FeedbackAnalysisParameters[0]);
-
-			// Wait for them to complete
-			if (Tasks.Num() > 0)
+			for (int32 TaskIndex = 0; TaskIndex < LocalFeedbackTaskCount; ++TaskIndex)
+			{
+				FFeedbackAnalysisTask::DoTask(FeedbackAnalysisParameters[TaskIndex]);
+			}
+			if (WorkerFeedbackTaskCount > 0)
 			{
 				SCOPE_CYCLE_COUNTER(STAT_ProcessRequests_WaitTasks);
 
@@ -1155,9 +1140,7 @@ void FVirtualTextureSystem::Update(FRHICommandListImmediate& RHICmdList, ERHIFea
 		// Limit the number of uploads (account for MappedTilesToProduce this frame)
 		// Are all pages equal? Should there be different limits on different types of pages?
 		const int32 MaxNumUploads = VirtualTextureScalability::GetMaxUploadsPerFrame();
-		const float TimeSlice = CVarVTPageUpdateTimeSlice.GetValueOnRenderThread();
-		const int32 MaxNumUploadsForTimeSlice = TimeSlicer.Slice(MaxNumUploads, TimeSlice);
-		const int32 MaxRequestUploads = FMath::Max(MaxNumUploadsForTimeSlice - MappedTilesToProduce.Num(), 1);
+		const int32 MaxRequestUploads = FMath::Max(MaxNumUploads - MappedTilesToProduce.Num(), 1);
 
 		MergedRequestList->SortRequests(Producers, MemStack, MaxRequestUploads);
 	}
@@ -1176,7 +1159,7 @@ void FVirtualTextureSystem::GatherRequests(FUniqueRequestList* MergedRequestList
 {
 	FMemMark GatherMark(MemStack);
 
-	const uint32 MaxNumGatherTasks = FMath::Clamp(CVarVTNumGatherTasks.GetValueOnRenderThread(), 1, MaxNumTasks);
+	const uint32 MaxNumGatherTasks = FMath::Clamp((uint32)CVarVTNumGatherTasks.GetValueOnRenderThread(), 1u, MaxNumTasks);
 	const uint32 PageUpdateFlushCount = FMath::Min<uint32>(CVarVTPageUpdateFlushCount.GetValueOnRenderThread(), FPageUpdateBuffer::PageCapacity);
 
 	FGatherRequestsParameters GatherRequestsParameters[MaxNumTasks];
diff --git a/Engine/Source/Runtime/Renderer/Private/VT/VirtualTextureSystem.h b/Engine/Source/Runtime/Renderer/Private/VT/VirtualTextureSystem.h
--- a/Engine/Source/Runtime/Renderer/Private/VT/VirtualTextureSystem.h
+++ b/Engine/Source/Runtime/Renderer/Private/VT/VirtualTextureSystem.h
@@ -95,6 +95,7 @@ private:
 
 	uint32	Frame;
 
+	static const uint32 MaxNumTasks = 16;
 	static const uint32 MaxSpaces = 16;
 	TUniquePtr<FVirtualTextureSpace> Spaces[MaxSpaces];
 	TArray<FVirtualTexturePhysicalSpace*> PhysicalSpaces;
