Fix uninitialized screen shadow mask data being incorrectly sampled when a light that doesn't require it is evaluated after a light that does.
Fix attempting to transition null screen shadow mask resource in certain conditions.

diff --git a/Engine/Source/Runtime/Renderer/Private/LightRendering.cpp b/Engine/Source/Runtime/Renderer/Private/LightRendering.cpp
--- a/Engine/Source/Runtime/Renderer/Private/LightRendering.cpp
+++ b/Engine/Source/Runtime/Renderer/Private/LightRendering.cpp
@@ -1734,7 +1734,9 @@ void FDeferredShadingSceneRenderer::RenderLights(FRHICommandListImmediate& RHICm
 
 						ClearShadowMask(ScreenShadowMaskTexture);
 						if (ScreenShadowMaskSubPixelTexture)
+						{
 							ClearShadowMask(ScreenShadowMaskSubPixelTexture);
+						}
 
 						RenderShadowProjections(RHICmdList, &LightSceneInfo, ScreenShadowMaskTexture, ScreenShadowMaskSubPixelTexture, HairDatas, bInjectedTranslucentVolume);
 					}
@@ -1770,15 +1772,28 @@ void FDeferredShadingSceneRenderer::RenderLights(FRHICommandListImmediate& RHICm
 
 				if (bUsedShadowMaskTexture)
 				{
+					check(ScreenShadowMaskTexture);
 					RHICmdList.CopyToResolveTarget(ScreenShadowMaskTexture->GetRenderTargetItem().TargetableTexture, ScreenShadowMaskTexture->GetRenderTargetItem().ShaderResourceTexture, FResolveParams(FResolveRect()));
 					if (ScreenShadowMaskSubPixelTexture)
+					{
 						RHICmdList.CopyToResolveTarget(ScreenShadowMaskSubPixelTexture->GetRenderTargetItem().TargetableTexture, ScreenShadowMaskSubPixelTexture->GetRenderTargetItem().ShaderResourceTexture, FResolveParams(FResolveRect()));
-				}
+					}
 
-				if(bUsedShadowMaskTexture || !bShadowMaskReadable)
-				{
-					RHICmdList.TransitionResource(EResourceTransitionAccess::EReadable, ScreenShadowMaskTexture->GetRenderTargetItem().ShaderResourceTexture);
-					bShadowMaskReadable = true;
+					if (!bShadowMaskReadable)
+					{
+						RHICmdList.TransitionResource(EResourceTransitionAccess::EReadable, ScreenShadowMaskTexture->GetRenderTargetItem().ShaderResourceTexture);
+						if (ScreenShadowMaskSubPixelTexture)
+						{
+							RHICmdList.TransitionResource(EResourceTransitionAccess::EReadable, ScreenShadowMaskSubPixelTexture->GetRenderTargetItem().ShaderResourceTexture);
+						}
+						bShadowMaskReadable = true;
+					}
+
+					GVisualizeTexture.SetCheckPoint(RHICmdList, ScreenShadowMaskTexture);
+					if (ScreenShadowMaskSubPixelTexture)
+					{
+						GVisualizeTexture.SetCheckPoint(RHICmdList, ScreenShadowMaskSubPixelTexture);
+					}
 				}
 
 				if(bDirectLighting && !bInjectedTranslucentVolume)
@@ -1791,10 +1806,6 @@ void FDeferredShadingSceneRenderer::RenderLights(FRHICommandListImmediate& RHICm
 					}
 				}
 
-				GVisualizeTexture.SetCheckPoint(RHICmdList, ScreenShadowMaskTexture);
-				if (ScreenShadowMaskSubPixelTexture)
-					GVisualizeTexture.SetCheckPoint(RHICmdList, ScreenShadowMaskSubPixelTexture);
-
 				FHairStrandsTransmittanceMaskData TransmittanceMaskData;
 				if (bDrawDeepShadow)
 				{
@@ -1817,7 +1828,15 @@ void FDeferredShadingSceneRenderer::RenderLights(FRHICommandListImmediate& RHICm
 					// Render the light to the scene color buffer, conditionally using the attenuation buffer or a 1x1 white texture as input 
 					if (bDirectLighting)
 					{
-						RenderLight(RHICmdList, &LightSceneInfo, ScreenShadowMaskTexture, ScreenShadowMaskSubPixelTexture, &TransmittanceMaskData, InHairVisibilityViews, false, true);
+						// ScreenShadowMaskTexture might have been created for a previous light, but only use it if we wrote valid data into it for this light
+						IPooledRenderTarget* LightShadowMaskTexture = nullptr;
+						IPooledRenderTarget* LightShadowMaskSubPixelTexture = nullptr;
+						if (bUsedShadowMaskTexture)
+						{
+							LightShadowMaskTexture = ScreenShadowMaskTexture;
+							LightShadowMaskSubPixelTexture = ScreenShadowMaskSubPixelTexture;
+						}
+						RenderLight(RHICmdList, &LightSceneInfo, LightShadowMaskTexture, LightShadowMaskSubPixelTexture, &TransmittanceMaskData, InHairVisibilityViews, false, true);
 					}
 
 					SceneContext.FinishRenderingSceneColor(RHICmdList);
