Disables SSRT debug output

diff --git a/Engine/Source/Runtime/Renderer/Private/ScreenSpaceRayTracing.cpp b/Engine/Source/Runtime/Renderer/Private/ScreenSpaceRayTracing.cpp
--- a/Engine/Source/Runtime/Renderer/Private/ScreenSpaceRayTracing.cpp
+++ b/Engine/Source/Runtime/Renderer/Private/ScreenSpaceRayTracing.cpp
@@ -504,7 +504,7 @@ void GetSSRTGIShaderOptionsForQuality(int32 Quality, FIntPoint* OutGroupSize, in
 }
 
 FRDGTextureUAV* CreateScreenSpaceRayTracingDebugUAV(FRDGBuilder& GraphBuilder, const FRDGTextureDesc& Desc, const TCHAR* Name, bool bClear = false)
-#if 1
+#if 0
 {
 	FRDGTextureDesc DebugDesc = FRDGTextureDesc::Create2DDesc(
 		Desc.Extent,
