Workaround depth submission bug in Oculus runtime by outputing alpha in tonemapper when view is a stereo view.

diff --git a/Engine/Source/Runtime/Renderer/Private/PostProcess/PostProcessTonemap.cpp b/Engine/Source/Runtime/Renderer/Private/PostProcess/PostProcessTonemap.cpp
--- a/Engine/Source/Runtime/Renderer/Private/PostProcess/PostProcessTonemap.cpp
+++ b/Engine/Source/Runtime/Renderer/Private/PostProcess/PostProcessTonemap.cpp
@@ -837,7 +837,10 @@ FScreenPassTexture AddTonemapPass(FRDGBuilder& GraphBuilder, const FViewInfo& Vi
 		TShaderMapRef<FTonemapVS> VertexShader(View.ShaderMap, VertexPermutationVector);
 		TShaderMapRef<FTonemapPS> PixelShader(View.ShaderMap, DesktopPermutationVector);
 
-		FRHIBlendState* BlendState = Inputs.bWriteAlphaChannel ? FScreenPassPipelineState::FDefaultBlendState::GetRHI() : TStaticBlendStateWriteMask<CW_RGB>::GetRHI();
+		// If this is a stereo view, there's a good chance we need alpha out of the tonemapper
+		// @todo: Remove this once Oculus fix the bug in their runtime that requires alpha here.
+		const bool bIsStereo = IStereoRendering::IsStereoEyeView(View);
+		FRHIBlendState* BlendState = Inputs.bWriteAlphaChannel || bIsStereo ? FScreenPassPipelineState::FDefaultBlendState::GetRHI() : TStaticBlendStateWriteMask<CW_RGB>::GetRHI();
 		FRHIDepthStencilState* DepthStencilState = FScreenPassPipelineState::FDefaultDepthStencilState::GetRHI();
 
 		EScreenPassDrawFlags DrawFlags = EScreenPassDrawFlags::AllowHMDHiddenAreaMask;
