Removed performance warning

diff --git a/Engine/Source/Runtime/Renderer/Private/DeferredShadingRenderer.cpp b/Engine/Source/Runtime/Renderer/Private/DeferredShadingRenderer.cpp
--- a/Engine/Source/Runtime/Renderer/Private/DeferredShadingRenderer.cpp
+++ b/Engine/Source/Runtime/Renderer/Private/DeferredShadingRenderer.cpp
@@ -1003,7 +1003,7 @@ void FDeferredShadingSceneRenderer::Render(FRHICommandListImmediate& RHICmdList)
 	FRHIGPUMask RenderTargetGPUMask = (GNumExplicitGPUsForRendering > 1 && ViewFamily.RenderTarget) ? ViewFamily.RenderTarget->GetGPUMask(RHICmdList) : FRHIGPUMask::GPU0();
 	
 	{
-		auto CVar = IConsoleManager::Get().FindConsoleVariable(TEXT("r.PathTracing.GPUCount"));
+		static auto CVar = IConsoleManager::Get().FindConsoleVariable(TEXT("r.PathTracing.GPUCount"));
 		if (CVar && CVar->GetInt() > 1)
 		{
 			RenderTargetGPUMask = FRHIGPUMask::All(); // Broadcast to all GPUs 
