Takes down leaking when reprojecting previous frame for SSGI

diff --git a/Engine/Shaders/Private/SSRT/SSRTDiffuseIndirect.usf b/Engine/Shaders/Private/SSRT/SSRTDiffuseIndirect.usf
--- a/Engine/Shaders/Private/SSRT/SSRTDiffuseIndirect.usf
+++ b/Engine/Shaders/Private/SSRT/SSRTDiffuseIndirect.usf
@@ -36,7 +36,7 @@
 #elif QUALITY == 3
 	#define CONFIG_RAY_STEPS 8
 	#define CONFIG_RAY_COUNT 16
-	#define CONFIG_BACKFACE_MODULATION 1
+	#define CONFIG_BACKFACE_MODULATION 0
 
 	#define TILE_PIXEL_SIZE_X 4
 	#define TILE_PIXEL_SIZE_X_LOG 2
@@ -47,7 +47,7 @@
 #elif QUALITY == 4
 	#define CONFIG_RAY_STEPS 12
 	#define CONFIG_RAY_COUNT 32
-	#define CONFIG_BACKFACE_MODULATION 1
+	#define CONFIG_BACKFACE_MODULATION 0
 
 	#define TILE_PIXEL_SIZE_X 4
 	#define TILE_PIXEL_SIZE_X_LOG 2
@@ -99,9 +99,8 @@ SamplerState ColorTextureSampler;
 #endif
 
 float4 HZBUvFactorAndInvFactor;
+float4 ColorBufferScaleBias;
 
-float4 PrevScreenPositionScaleBias;
-float  PrevSceneColorPreExposureCorrection;
 float PixelPositionToFullResPixel;
 float2 FullResPixelOffset;
 
@@ -539,7 +538,12 @@ void MainCS(
 			BRANCH
 			if (bHit)
 			{
-				float4 SampleColor = ColorTexture.SampleLevel(ColorTextureSampler, HitUVz.xy, Level);
+				float2 ReducedColorUV =  HitUVz.xy * ColorBufferScaleBias.xy + ColorBufferScaleBias.zw;
+				Level = min(Level, View.GeneralPurposeTweak);
+
+
+
+				float4 SampleColor = ColorTexture.SampleLevel(ColorTextureSampler, ReducedColorUV, Level);
 				
 				float SampleColorWeight = 1.0;
 
diff --git a/Engine/Shaders/Private/SSRT/SSRTPrevFrameReduction.usf b/Engine/Shaders/Private/SSRT/SSRTPrevFrameReduction.usf
--- a/Engine/Shaders/Private/SSRT/SSRTPrevFrameReduction.usf
+++ b/Engine/Shaders/Private/SSRT/SSRTPrevFrameReduction.usf
@@ -10,30 +10,66 @@
 #define GROUP_PIXEL_COUNT GROUP_TILE_SIZE * GROUP_TILE_SIZE
 
 
-uint bIsTemporalAAHistory;
-Texture2D PrevTemporalAAHistory;
+Texture2D PrevSceneColor;
+Texture2D PrevSceneDepth;
+Texture2D FurthestHZBTexture;
 
 Texture2D HigherMipTexture;
+Texture2D HigherAlphaMipTexture;
 
 #if SUPPORTS_INDEPENDENT_SAMPLERS
-#define PrevTemporalAAHistorySampler GlobalBilinearClampedSampler
+#if DIM_LEAK_FREE
+#define PrevSceneColorSampler GlobalPointClampedSampler
+#else
+#define PrevSceneColorSampler GlobalBilinearClampedSampler
+#endif
+#define PrevSceneDepthSampler GlobalPointClampedSampler
+#define FurthestHZBTextureSampler GlobalPointClampedSampler
 #define HigherMipTextureSampler GlobalPointClampedSampler
+#define HigherAlphaMipTextureSampler GlobalPointClampedSampler
 #else
-SamplerState PrevTemporalAAHistorySampler;
+SamplerState PrevSceneColorSampler;
+SamplerState PrevSceneDepthSampler;
+SamplerState FurthestHZBTextureSampler;
 SamplerState HigherMipTextureSampler;
+SamplerState HigherAlphaMipTextureSampler;
 #endif
 
 float4 PrevScreenPositionScaleBias;
-float  PrevSceneColorPreExposureCorrection;
-float  MinimumLuminance;
-float TexelSize;
+float2 ViewportUVToHZBBufferUV;
+float2 ReducedSceneColorSize;
+float2 ReducedSceneColorTexelSize;
+float PrevSceneColorPreExposureCorrection;
+float MinimumLuminance;
+float HigherMipDownScaleFactor;
 
 RWTexture2D<float4>	ReducedSceneColorOutput_0;
 RWTexture2D<float4>	ReducedSceneColorOutput_1;
 RWTexture2D<float4>	ReducedSceneColorOutput_2;
 
+RWTexture2D<float>	ReducedSceneAlphaOutput_0;
+RWTexture2D<float>	ReducedSceneAlphaOutput_1;
+RWTexture2D<float>	ReducedSceneAlphaOutput_2;
+
 groupshared uint SharedMemory[GROUP_PIXEL_COUNT * 2];
 
+#if DIM_LEAK_FREE
+groupshared float SharedFurthestDepth[GROUP_PIXEL_COUNT];
+#endif
+
+
+float3 ComputeTranslatedWorldPosition(float2 ScreenPosition, float WorldDepth, bool bIsPrevFrame)
+{
+	float4 ClipPosition = float4(ScreenPosition * (View.ViewToClip[3][3] < 1.0f ? WorldDepth : 1.0f), WorldDepth, 1);
+
+	if (bIsPrevFrame)
+	{
+		return mul(ClipPosition, View.PrevScreenToTranslatedWorld).xyz + (View.PreViewTranslation.xyz - View.PrevPreViewTranslation.xyz);
+	}
+	
+	return mul(ClipPosition, View.ScreenToTranslatedWorld).xyz;
+}
+
 
 [numthreads(GROUP_TILE_SIZE, GROUP_TILE_SIZE, 1)]
 void MainCS(
@@ -42,88 +78,167 @@ void MainCS(
 	uint2 GroupThreadId : SV_GroupThreadID,
 	uint GroupThreadIndex : SV_GroupIndex)
 {
-	float2 BufferUV = (float2(DispatchThreadId) + 0.5) * View.BufferSizeAndInvSize.zw;
-	float2 ViewportUV = (float2(DispatchThreadId) + 0.5f) * View.ViewSizeAndInvSize.zw;
+	#if DIM_LOWER_MIPS
+		float2 SceneBufferUV = HigherMipDownScaleFactor * (2.0 * float2(DispatchThreadId) + 1.0) * View.BufferSizeAndInvSize.zw;
+	#else
+		float2 SceneBufferUV = (float2(DispatchThreadId) + 0.5) * View.BufferSizeAndInvSize.zw;
+	#endif
+
+	SceneBufferUV = clamp(SceneBufferUV, View.BufferBilinearUVMinMax.xy, View.BufferBilinearUVMinMax.zw);
+	float2 ViewportUV = BufferUVToViewportUV(SceneBufferUV);
     float2 ScreenPosition = ViewportUVToScreenPos(ViewportUV);
 	
 	float2 PrevFrameUV;
 	float2 PrevScreen;
 	bool bIsSky = false;
 
-	float4 PrevSceneColor;
+	float4 PrevColor;
+
+	float WorldDepth;
 
 	#if DIM_LOWER_MIPS
 	{
-		PrevFrameUV = TexelSize * (float2(DispatchThreadId) + 0.5) * View.BufferSizeAndInvSize.zw;
+		PrevFrameUV = HigherMipDownScaleFactor * (float2(DispatchThreadId) * 2.0 + 1.0) * ReducedSceneColorTexelSize;
 		PrevScreen = 0.0; // unused.
 		
-		PrevSceneColor = HigherMipTexture.SampleLevel(HigherMipTextureSampler, PrevFrameUV, 0);
+		#if DIM_LEAK_FREE
+		{
+			{
+				float HZBDeviceZ = FurthestHZBTexture.SampleLevel(FurthestHZBTextureSampler, ViewportUV * ViewportUVToHZBBufferUV, 2.0).r;
+				WorldDepth = ConvertFromDeviceZ(HZBDeviceZ);
+			}
+
+			float WorldDepthToPixelWorldRadius = GetTanHalfFieldOfView().x * View.ViewSizeAndInvSize.z * 100;
+
+			float WorldBluringRadius = WorldDepthToPixelWorldRadius * WorldDepth;
+			float InvSquareWorldBluringRadius = rcp(WorldBluringRadius * WorldBluringRadius);
+
+			PrevColor = 0.0;
+
+			UNROLL_N(4)
+			for (uint i = 0; i < 4; i++)
+			{
+				const float2 TexelOffset = float2(i % 2, i / 2) - 0.5;
+
+				float2 HZBBufferUV = (ViewportUV + TexelOffset * HigherMipDownScaleFactor * View.ViewSizeAndInvSize.zw) * ViewportUVToHZBBufferUV;
+				float SampleDeviceZ = FurthestHZBTexture.SampleLevel(FurthestHZBTextureSampler, HZBBufferUV, 1.0).r;
+
+				float SampleDist = WorldDepth - ConvertFromDeviceZ(SampleDeviceZ);
+
+				float SampleWeight = 0.25 * saturate(1 - SampleDist * SampleDist * InvSquareWorldBluringRadius);
+
+				float2 SampleUV = HigherMipDownScaleFactor * (2.0 * float2(DispatchThreadId) + 1.0 + TexelOffset) * 0.5 * ReducedSceneColorTexelSize;
+				SampleUV = clamp(SampleUV, View.BufferBilinearUVMinMax.xy, View.BufferBilinearUVMinMax.zw);
+
+				float4 SampleColor = float4(
+					HigherMipTexture.SampleLevel(HigherMipTextureSampler, SampleUV, 0).rgb,
+					HigherAlphaMipTexture.SampleLevel(HigherAlphaMipTextureSampler, SampleUV, 0).r);
+
+				PrevColor += SampleColor * SampleWeight;
+			}
+		}
+		#else
+		{
+			WorldDepth = 0.0;
+			PrevColor = float4(HigherMipTexture.SampleLevel(HigherMipTextureSampler, PrevFrameUV, 0).rgb, 1);
+		}
+		#endif
 	}
 	#else
 	{
-		if (bIsTemporalAAHistory)
+		float DeviceZ = SampleDeviceZFromSceneTextures(SceneBufferUV);
+		WorldDepth = ConvertFromDeviceZ(DeviceZ);
+
+		// Camera motion for pixel (in ScreenPos space).
+		float4 ThisClip = float4(ScreenPosition, DeviceZ, 1);
+		float4 PrevClip = mul(ThisClip, View.ClipToPrevClip);
+		PrevScreen = PrevClip.xy / PrevClip.w;
+
+		bIsSky = WorldDepth > 100 * 1000;
+
+		float2 EncodedVelocity = SceneVelocityBuffer.SampleLevel(SceneVelocityBufferSampler, SceneBufferUV, 0).xy;
+		if (EncodedVelocity.x > 0.0)
+		{
+			PrevScreen = ThisClip.xy - DecodeVelocityFromTexture(EncodedVelocity);
+		}
+
+		PrevFrameUV = PrevScreen.xy * PrevScreenPositionScaleBias.xy + PrevScreenPositionScaleBias.zw;
+
+		// Fetch color.
+		#if DIM_LEAK_FREE
 		{
-			float DeviceZ = SceneDepthBuffer.SampleLevel(SceneDepthBufferSampler, BufferUV, 0).r;
+			float3 RefWorldPosition = ComputeTranslatedWorldPosition(ScreenPosition, WorldDepth, /* bIsPrevFrame = */ false);
 
-			// Camera motion for pixel (in ScreenPos space).
-			float4 ThisClip = float4(ScreenPosition, DeviceZ, 1);
-			float4 PrevClip = mul(ThisClip, View.ClipToPrevClip);
-			PrevScreen = PrevClip.xy / PrevClip.w;
+			float NoV = dot(View.TranslatedWorldCameraOrigin - normalize(RefWorldPosition), GetGBufferDataFromSceneTextures(SceneBufferUV).WorldNormal);
+				
+			float WorldDepthToPixelWorldRadius = GetTanHalfFieldOfView().x * View.ViewSizeAndInvSize.z * 100; // / max(NoV, 0.01);
 
-			bIsSky = ConvertFromDeviceZ(DeviceZ) > 100 * 1000;
+			float WorldBluringRadius = WorldDepthToPixelWorldRadius * WorldDepth;
+			float InvSquareWorldBluringRadius = rcp(WorldBluringRadius * WorldBluringRadius);
 
-			float2 EncodedVelocity = SceneVelocityBuffer.SampleLevel(SceneVelocityBufferSampler, BufferUV, 0).xy;
-			if (EncodedVelocity.x > 0.0)
+			//PrevColor = 0.0;
+
+			//UNROLL_N(4)
+			//for (uint i = 0; i < 4; i++)
 			{
-				PrevScreen = ThisClip.xy - DecodeVelocityFromTexture(EncodedVelocity);
-			}
+				float2 SampleUV = PrevFrameUV; // + View.BufferSizeAndInvSize.zw * (float2(i % 2, i / 2) - 0.5);
+				SampleUV = clamp(SampleUV, View.BufferBilinearUVMinMax.xy, View.BufferBilinearUVMinMax.zw);
 
-			PrevFrameUV = PrevScreen.xy * PrevScreenPositionScaleBias.xy + PrevScreenPositionScaleBias.zw;
+				float PrevDeviceZ = PrevSceneDepth.SampleLevel(PrevSceneDepthSampler, SampleUV, 0).r;
+
+				float3 SampleWorldPosition = ComputeTranslatedWorldPosition(PrevScreen.xy, ConvertFromDeviceZ(PrevDeviceZ), /* bIsPrevFrame = */ true);
+
+				float SampleDistSquare = length2(RefWorldPosition - SampleWorldPosition);
+
+				float SampleWeight = saturate(1 - SampleDistSquare * InvSquareWorldBluringRadius);
+
+				PrevColor = float4(PrevSceneColor.SampleLevel(PrevSceneColorSampler, SampleUV, 0).rgb * SampleWeight, SampleWeight);
+			}
 		}
-		else
+		#else
 		{
-			PrevFrameUV = BufferUV;
-			PrevScreen = ScreenPosition;
+			PrevColor = float4(PrevSceneColor.SampleLevel(PrevSceneColorSampler, PrevFrameUV, 0).rgb, 1.0);
 		}
-		
-		// Fetch color.
-		PrevSceneColor = float4(
-			PrevTemporalAAHistory.SampleLevel(PrevTemporalAAHistorySampler, PrevFrameUV, 0).rgb, 1.0);
-		
+		#endif
+
 		#if CONFIG_COLOR_TILE_CLASSIFICATION
 		{
 			if (bIsSky)
-				PrevSceneColor = 0;
+				PrevColor = 0;
 		}
 		#endif
 
 		// Correct scene color exposure.
 		#if USE_PREEXPOSURE
-			PrevSceneColor.rgb *= PrevSceneColorPreExposureCorrection;
+			PrevColor.rgb *= PrevSceneColorPreExposureCorrection;
 		#endif
 
 		// Apply vignette to the color.
 		{
 			float Vignette = min(ComputeHitVignetteFromScreenPos(ScreenPosition), ComputeHitVignetteFromScreenPos(PrevScreen));
-			PrevSceneColor *= Vignette;
+			PrevColor *= Vignette;
 		}
 
 		// Apply luminance to mask in only significant lighting worth tracing rays towards.
 		#if CONFIG_COLOR_TILE_CLASSIFICATION
 		if (MinimumLuminance > 0.0)
 		{
-			float ColorLuminance = Luminance(PrevSceneColor.rgb);
+			float ColorLuminance = Luminance(PrevColor.rgb);
 
-			PrevSceneColor.rgb *= saturate((ColorLuminance - MinimumLuminance) / MinimumLuminance);
+			PrevColor.rgb *= saturate((ColorLuminance - MinimumLuminance) / MinimumLuminance);
 		}
 		#endif
 	}
 	#endif
-	
+
 	// Output mip 0
 	#if DIM_LOWER_MIPS
 	{
-		ReducedSceneColorOutput_0[DispatchThreadId] = PrevSceneColor;
+		ReducedSceneColorOutput_0[DispatchThreadId] = float4(PrevColor.rgb, 0);
+
+		#if DIM_LEAK_FREE
+			ReducedSceneAlphaOutput_0[DispatchThreadId] = PrevColor.a;
+		#endif
 	}
 	#endif
 
@@ -131,8 +246,12 @@ void MainCS(
 	{
 		// Store to LDS
 		{
-			SharedMemory[GROUP_PIXEL_COUNT * 0 | GroupThreadIndex] = (f32tof16(PrevSceneColor.r) << 0) | (f32tof16(PrevSceneColor.g) << 16);
-			SharedMemory[GROUP_PIXEL_COUNT * 1 | GroupThreadIndex] = (f32tof16(PrevSceneColor.b) << 0) | (f32tof16(PrevSceneColor.a) << 16);
+			SharedMemory[GROUP_PIXEL_COUNT * 0 | GroupThreadIndex] = (f32tof16(PrevColor.r) << 0) | (f32tof16(PrevColor.g) << 16);
+			SharedMemory[GROUP_PIXEL_COUNT * 1 | GroupThreadIndex] = (f32tof16(PrevColor.b) << 0) | (f32tof16(PrevColor.a) << 16);
+
+			#if DIM_LEAK_FREE
+				SharedFurthestDepth[GroupThreadIndex] = WorldDepth;
+			#endif
 		}
 	
 		GroupMemoryBarrierWithGroupSync();
@@ -141,8 +260,6 @@ void MainCS(
 		UNROLL
 		for (uint MipLevel = 1; MipLevel < 3; MipLevel++)
 		{
-			float4 PrevSceneColor = 0;
-			
 			const uint ReductionAmount = 1 << MipLevel;
 			const uint NumberPixelInMip = GROUP_PIXEL_COUNT / (ReductionAmount * ReductionAmount);
 			
@@ -152,11 +269,41 @@ void MainCS(
 					GroupThreadIndex % (GROUP_TILE_SIZE / ReductionAmount),
 					GroupThreadIndex / (GROUP_TILE_SIZE / ReductionAmount));
 
+				// Ray marchsing against FurthestHZB to avoid self intersections, so match here to remain conservative.
+				#if DIM_LEAK_FREE
+				float FurthestDepth;
+				{
+					UNROLL_N(2)
+					for (uint x = 0; x < 2; x++)
+					{
+						UNROLL_N(2)
+						for (uint y = 0; y < 2; y++)
+						{
+							uint2 Coord = OutputCoord * 2 + uint2(x, y);
+							uint LDSIndex = Coord.x + Coord.y * ((2 * GROUP_TILE_SIZE) / ReductionAmount);
+
+							float NeighborDepth = SharedFurthestDepth[LDSIndex];
+
+							if (x == 0 && y == 0)
+								FurthestDepth = NeighborDepth;
+							else
+								FurthestDepth = max(FurthestDepth, NeighborDepth);
+						}
+					}
+				}
+
+				float WorldDepthToPixelWorldRadius = GetTanHalfFieldOfView().x * View.ViewSizeAndInvSize.z * 100;
+
+				float WorldBluringRadius = WorldDepthToPixelWorldRadius * FurthestDepth;
+				float InvSquareWorldBluringRadius = rcp(WorldBluringRadius * WorldBluringRadius);
+
+				#endif
+
 				float4 ReducedColor = 0;
 
 				UNROLL
 				for (uint x = 0; x < 2; x++)
-                {
+				{
 					UNROLL
 					for (uint y = 0; y < 2; y++)
 					{
@@ -172,7 +319,18 @@ void MainCS(
 						Color.b = f16tof32(Raw1 >> 0);
 						Color.a = f16tof32(Raw1 >> 16);
 
-						ReducedColor += Color;
+
+						float SampleWeight = 1.0;
+						#if DIM_LEAK_FREE
+						{
+							float NeighborDepth = SharedFurthestDepth[LDSIndex];
+							float SampleDist = (FurthestDepth - NeighborDepth);
+
+							SampleWeight = saturate(1 - (SampleDist * SampleDist) * InvSquareWorldBluringRadius);
+						}
+						#endif
+
+						ReducedColor += Color * SampleWeight;
 					}
                 }
 
@@ -182,15 +340,27 @@ void MainCS(
 
 				if (MipLevel == 1)
                 {
-					ReducedSceneColorOutput_1[OutputPosition] = ReducedColor;
+					ReducedSceneColorOutput_1[OutputPosition] = float4(ReducedColor.rgb, 0);
+					#if DIM_LEAK_FREE
+						ReducedSceneAlphaOutput_1[OutputPosition] = ReducedColor.a;
+					#endif
                 }
 				else if (MipLevel == 2)
                 {
-					ReducedSceneColorOutput_2[OutputPosition] = ReducedColor;
+					ReducedSceneColorOutput_2[OutputPosition] = float4(ReducedColor.rgb, 0);
+					#if DIM_LEAK_FREE
+						ReducedSceneAlphaOutput_2[OutputPosition] = ReducedColor.a;
+					#endif
                 }
 				
 				SharedMemory[GROUP_PIXEL_COUNT * 0 | GroupThreadIndex] = (f32tof16(ReducedColor.r) << 0) | (f32tof16(ReducedColor.g) << 16);
 				SharedMemory[GROUP_PIXEL_COUNT * 1 | GroupThreadIndex] = (f32tof16(ReducedColor.b) << 0) | (f32tof16(ReducedColor.a) << 16);
+
+				#if DIM_LEAK_FREE
+				{
+					SharedFurthestDepth[GroupThreadIndex] = FurthestDepth;
+				}
+				#endif
             } // if (GroupThreadIndex < NumberPixelInMip)
         } // for (uint MipLevel = 1; MipLevel < 3; MipLevel++)
     }
diff --git a/Engine/Source/Runtime/Renderer/Private/DeferredShadingRenderer.cpp b/Engine/Source/Runtime/Renderer/Private/DeferredShadingRenderer.cpp
--- a/Engine/Source/Runtime/Renderer/Private/DeferredShadingRenderer.cpp
+++ b/Engine/Source/Runtime/Renderer/Private/DeferredShadingRenderer.cpp
@@ -2377,6 +2377,19 @@ void FDeferredShadingSceneRenderer::Render(FRHICommandListImmediate& RHICmdList)
 	// Resolve the scene color for post processing.
 	ResolveSceneColor(RHICmdList);
 
+	// Keep scene color and depth for next frame screen space ray tracing.
+	{
+		for (const FViewInfo& View : Views)
+		{
+			if (ShouldKeepBleedFreeSceneColor(View))
+			{
+				FSceneViewState* ViewState = View.ViewState;
+				ViewState->PrevFrameViewInfo.DepthBuffer = SceneContext.SceneDepthZ;
+				ViewState->PrevFrameViewInfo.ScreenSpaceRayTracingInput = SceneContext.GetSceneColor();
+			}
+		}
+	}
+
 	GetRendererModule().RenderPostResolvedSceneColorExtension(RHICmdList, SceneContext);
 
 	CopySceneCaptureComponentToTarget(RHICmdList);
diff --git a/Engine/Source/Runtime/Renderer/Private/SceneRendering.h b/Engine/Source/Runtime/Renderer/Private/SceneRendering.h
--- a/Engine/Source/Runtime/Renderer/Private/SceneRendering.h
+++ b/Engine/Source/Runtime/Renderer/Private/SceneRendering.h
@@ -786,6 +786,9 @@ struct FPreviousViewInfo
 	// DeviceZ as float16, and normal in view space.
 	TRefCountPtr<IPooledRenderTarget> CompressedDepthViewNormal;
 
+	// Bleed free scene color to use for screen space ray tracing.
+	TRefCountPtr<IPooledRenderTarget> ScreenSpaceRayTracingInput;
+
 	// Temporal AA result of last frame
 	FTemporalAAHistory TemporalAAHistory;
 
diff --git a/Engine/Source/Runtime/Renderer/Private/ScreenSpaceRayTracing.cpp b/Engine/Source/Runtime/Renderer/Private/ScreenSpaceRayTracing.cpp
--- a/Engine/Source/Runtime/Renderer/Private/ScreenSpaceRayTracing.cpp
+++ b/Engine/Source/Runtime/Renderer/Private/ScreenSpaceRayTracing.cpp
@@ -48,6 +48,11 @@ static TAutoConsoleVariable<int32> CVarSSGIEnable(
 	TEXT("Whether to enable SSGI (defaults to 0).\n"),
 	ECVF_RenderThreadSafe);
 
+static TAutoConsoleVariable<int32> CVarSSGILeakFreeReprojection(
+	TEXT("r.SSGI.LeakFreeReprojection"), 1,
+	TEXT("Whether use a more expensive but leak free reprojection of previous frame's scene color.\n"),
+	ECVF_RenderThreadSafe);
+
 static TAutoConsoleVariable<int32> CVarSSGIHalfResolution(
 	TEXT("r.SSGI.HalfRes"), 0,
 	TEXT("Whether to do SSGI at half resolution (defaults to 0).\n"),
@@ -68,6 +73,34 @@ DECLARE_GPU_STAT_NAMED(ScreenSpaceReflections, TEXT("ScreenSpace Reflections"));
 DECLARE_GPU_STAT_NAMED(ScreenSpaceDiffuseIndirect, TEXT("Screen Space Diffuse Indirect"));
 
 
+static bool SupportScreenSpaceDiffuseIndirect(const FViewInfo& View)
+{
+	if (CVarSSGIEnable.GetValueOnRenderThread() <= 0)
+	{
+		return false;
+	}
+
+	int Quality = CVarSSGIQuality.GetValueOnRenderThread();
+
+	if (Quality <= 0)
+	{
+		return false;
+	}
+
+	if (IsAnyForwardShadingEnabled(View.GetShaderPlatform()))
+	{
+		return false;
+	}
+
+	return View.ViewState;
+}
+
+bool ShouldKeepBleedFreeSceneColor(const FViewInfo& View)
+{
+	// TODO(Guillaume): SSR as well.
+	return SupportScreenSpaceDiffuseIndirect(View) && !View.bStatePrevViewInfoIsReadOnly && CVarSSGILeakFreeReprojection.GetValueOnRenderThread() != 0;
+}
+
 bool ShouldRenderScreenSpaceReflections(const FViewInfo& View)
 {
 	if(!View.Family->EngineShowFlags.ScreenSpaceReflections)
@@ -103,24 +136,12 @@ bool ShouldRenderScreenSpaceReflections(const FViewInfo& View)
 
 bool ShouldRenderScreenSpaceDiffuseIndirect(const FViewInfo& View)
 {
-	if (CVarSSGIEnable.GetValueOnRenderThread() <= 0)
-	{
-		return false;
-	}
-
-	int Quality = CVarSSGIQuality.GetValueOnRenderThread();
-
-	if (Quality <= 0)
-	{
-		return false;
-	}
-
-	if (IsAnyForwardShadingEnabled(View.GetShaderPlatform()))
+	if (!SupportScreenSpaceDiffuseIndirect(View))
 	{
 		return false;
 	}
 
-	return View.PrevViewInfo.TemporalAAHistory.IsValid() && View.AntiAliasingMethod == AAM_TemporalAA;
+	return View.PrevViewInfo.ScreenSpaceRayTracingInput.IsValid() || View.PrevViewInfo.TemporalAAHistory.IsValid();
 }
 
 bool IsSSRTemporalPassRequired(const FViewInfo& View)
@@ -257,8 +278,9 @@ class FSSRTPrevFrameReductionCS : public FGlobalShader
 	SHADER_USE_PARAMETER_STRUCT(FSSRTPrevFrameReductionCS, FGlobalShader);
 
 	class FLowerMips : SHADER_PERMUTATION_BOOL("DIM_LOWER_MIPS");
+	class FLeakFree : SHADER_PERMUTATION_BOOL("DIM_LEAK_FREE");
 
-	using FPermutationDomain = TShaderPermutationDomain<FLowerMips>;
+	using FPermutationDomain = TShaderPermutationDomain<FLowerMips, FLeakFree>;
 
 	static bool ShouldCompilePermutation(const FGlobalShaderPermutationParameters& Parameters)
 	{
@@ -267,22 +289,33 @@ class FSSRTPrevFrameReductionCS : public FGlobalShader
 
 	BEGIN_SHADER_PARAMETER_STRUCT(FParameters, )
 		SHADER_PARAMETER(FVector4, PrevScreenPositionScaleBias)
+		SHADER_PARAMETER(FVector2D, ViewportUVToHZBBufferUV)
+		SHADER_PARAMETER(FVector2D, ReducedSceneColorSize)
+		SHADER_PARAMETER(FVector2D, ReducedSceneColorTexelSize)
 		SHADER_PARAMETER(float, PrevSceneColorPreExposureCorrection)
 		SHADER_PARAMETER(float, MinimumLuminance)
-		SHADER_PARAMETER(float, TexelSize)
-		SHADER_PARAMETER(uint32, bIsTemporalAAHistory)
+		SHADER_PARAMETER(float, HigherMipDownScaleFactor)
 		
-		SHADER_PARAMETER_RDG_TEXTURE(Texture2D, PrevTemporalAAHistory)
-		SHADER_PARAMETER_SAMPLER(SamplerState, PrevTemporalAAHistorySampler)
+		SHADER_PARAMETER_RDG_TEXTURE(Texture2D, PrevSceneColor)
+		SHADER_PARAMETER_SAMPLER(SamplerState, PrevSceneColorSampler)
+
+		SHADER_PARAMETER_RDG_TEXTURE(Texture2D, PrevSceneDepth)
+		SHADER_PARAMETER_SAMPLER(SamplerState, PrevSceneDepthSampler)
 
 		SHADER_PARAMETER_RDG_TEXTURE_SRV(Texture2D, HigherMipTexture)
+		SHADER_PARAMETER_RDG_TEXTURE_SRV(Texture2D, HigherAlphaMipTexture)
 		SHADER_PARAMETER_SAMPLER(SamplerState, HigherMipTextureSampler)
+		SHADER_PARAMETER_SAMPLER(SamplerState, HigherAlphaMipTextureSampler)
+
+		SHADER_PARAMETER_RDG_TEXTURE(Texture2D, FurthestHZBTexture)
+		SHADER_PARAMETER_SAMPLER(SamplerState, FurthestHZBTextureSampler)
 
 		SHADER_PARAMETER_STRUCT_INCLUDE(FSceneTextureParameters, SceneTextures)
 		SHADER_PARAMETER_STRUCT_INCLUDE(FSceneTextureSamplerParameters, SceneTextureSamplers)
 		SHADER_PARAMETER_STRUCT_REF(FViewUniformShaderParameters, View)
 
-		SHADER_PARAMETER_RDG_TEXTURE_UAV_ARRAY(RWTexture<float4>, ReducedSceneColorOutput, [3])
+		SHADER_PARAMETER_RDG_TEXTURE_UAV_ARRAY(RWTexture2D<float4>, ReducedSceneColorOutput, [3])
+		SHADER_PARAMETER_RDG_TEXTURE_UAV_ARRAY(RWTexture2D<float>, ReducedSceneAlphaOutput, [3])
 	END_SHADER_PARAMETER_STRUCT()
 };
 
@@ -400,9 +433,8 @@ class FScreenSpaceDiffuseIndirectCS : public FGlobalShader
 	using FPermutationDomain = TShaderPermutationDomain< FQualityDim >;
 	
 	BEGIN_SHADER_PARAMETER_STRUCT( FParameters, )
-		SHADER_PARAMETER( FVector4,		HZBUvFactorAndInvFactor )
-		SHADER_PARAMETER( FVector4,		PrevScreenPositionScaleBias )
-		SHADER_PARAMETER( float,		PrevSceneColorPreExposureCorrection )
+		SHADER_PARAMETER(FVector4, HZBUvFactorAndInvFactor)
+		SHADER_PARAMETER(FVector4, ColorBufferScaleBias)
 		SHADER_PARAMETER(float, PixelPositionToFullResPixel)
 		SHADER_PARAMETER(FVector2D, FullResPixelOffset)
 		
@@ -814,9 +846,6 @@ void RenderScreenSpaceDiffuseIndirect(
 {
 	check(ShouldRenderScreenSpaceDiffuseIndirect(View));
 
-	const FTemporalAAHistory& TemporalAAHistory = View.PrevViewInfo.TemporalAAHistory;
-	check(TemporalAAHistory.IsValid()); // TODO.
-	
 	const int32 Quality = FMath::Clamp( CVarSSGIQuality.GetValueOnRenderThread(), 1, 4 );
 
 	bool bHalfResolution = CVarSSGIHalfResolution.GetValueOnRenderThread() != 0;
@@ -833,36 +862,95 @@ void RenderScreenSpaceDiffuseIndirect(
 
 	RDG_EVENT_SCOPE(GraphBuilder, "SSGI %dx%d", Viewport.Width(), Viewport.Height());
 
+	const FVector2D ViewportUVToHZBBufferUV(
+		float(View.ViewRect.Width()) / float(2 * View.HZBMipmap0Size.X),
+		float(View.ViewRect.Height()) / float(2 * View.HZBMipmap0Size.Y)
+	);
+
 	FRDGTexture* FurthestHZBTexture = GraphBuilder.RegisterExternalTexture(View.HZB);
 	FRDGTexture* ClosestHZBTexture = GraphBuilder.RegisterExternalTexture(View.ClosestHZB);
-	FRDGTexture* ColorTexture = GraphBuilder.RegisterExternalTexture(TemporalAAHistory.RT[0]);
 
 	// Reproject and reduce previous frame color.
 	FRDGTexture* ReducedSceneColor;
+	FRDGTexture* ReducedSceneAlpha = nullptr;
 	{
 		// Number of mip skipped at the begining of the mip chain.
 		const int32 DownSamplingMip = 1;
 
+		// Number of mip in the mip chain
+		const int32 kNumMips = 5;
+
+		bool bUseLeakFree = View.PrevViewInfo.ScreenSpaceRayTracingInput != nullptr;
+
 		// Allocate ReducedSceneColor.
 		{
+			FIntPoint RequiredSize = SceneTextures.SceneDepthBuffer->Desc.Extent / (1 << DownSamplingMip);
+
+			int32 QuantizeMultiple = 1 << (kNumMips - 1);
+			FIntPoint QuantizedSize = FIntPoint::DivideAndRoundUp(RequiredSize, QuantizeMultiple);
+
 			FRDGTextureDesc Desc = FRDGTextureDesc::Create2DDesc(
-				SceneTextures.SceneDepthBuffer->Desc.Extent / (1 << DownSamplingMip),
+				FIntPoint(QuantizeMultiple * QuantizedSize.X, QuantizeMultiple * QuantizedSize.Y),
 				PF_FloatR11G11B10,
 				FClearValueBinding::None,
 				/* InFlags = */ TexCreate_None,
 				/* InTargetableFlags = */ TexCreate_ShaderResource | TexCreate_UAV,
 				/* bInForceSeparateTargetAndShaderResource = */ false);
-			Desc.NumMips = 5;
+			Desc.NumMips = kNumMips;
 
 			ReducedSceneColor = GraphBuilder.CreateTexture(Desc, TEXT("SSRTReducedSceneColor"));
+
+			if (bUseLeakFree)
+			{
+				Desc.Format = PF_A8;
+				ReducedSceneAlpha = GraphBuilder.CreateTexture(Desc, TEXT("SSRTReducedSceneAlpha"));
+			}
+		}
+
+		FSSRTPrevFrameReductionCS::FParameters DefaultPassParameters;
+		{
+			DefaultPassParameters.SceneTextures = SceneTextures;
+			SetupSceneTextureSamplers(&DefaultPassParameters.SceneTextureSamplers);
+			DefaultPassParameters.View = View.ViewUniformBuffer;
+
+			DefaultPassParameters.ReducedSceneColorSize = FVector2D(
+				ReducedSceneColor->Desc.Extent.X, ReducedSceneColor->Desc.Extent.Y);
+			DefaultPassParameters.ReducedSceneColorTexelSize = FVector2D(
+				1.0f / float(ReducedSceneColor->Desc.Extent.X), 1.0f / float(ReducedSceneColor->Desc.Extent.Y));
 		}
 
 		{
 			FSSRTPrevFrameReductionCS::FParameters* PassParameters = GraphBuilder.AllocParameters<FSSRTPrevFrameReductionCS::FParameters>();
+			*PassParameters = DefaultPassParameters;
+
+			FIntPoint ViewportOffset;
+			FIntPoint ViewportExtent;
+			FIntPoint BufferSize;
+
+			if (bUseLeakFree)
+			{
+				BufferSize = View.PrevViewInfo.ScreenSpaceRayTracingInput->GetDesc().Extent;
+				ViewportOffset = View.ViewRect.Min; // TODO
+				ViewportExtent = View.ViewRect.Size();
 
-			FIntPoint ViewportOffset = TemporalAAHistory.ViewportRect.Min;
-			FIntPoint ViewportExtent = TemporalAAHistory.ViewportRect.Size();
-			FIntPoint BufferSize = TemporalAAHistory.ReferenceBufferSize;
+				PassParameters->PrevSceneColor = GraphBuilder.RegisterExternalTexture(View.PrevViewInfo.ScreenSpaceRayTracingInput);
+				PassParameters->PrevSceneColorSampler = TStaticSamplerState<SF_Point>::GetRHI();
+
+				PassParameters->PrevSceneDepth = GraphBuilder.RegisterExternalTexture(View.PrevViewInfo.DepthBuffer);
+				PassParameters->PrevSceneDepthSampler = TStaticSamplerState<SF_Bilinear>::GetRHI();
+			}
+			else
+			{
+				BufferSize = View.PrevViewInfo.TemporalAAHistory.ReferenceBufferSize;
+				ViewportOffset = View.PrevViewInfo.TemporalAAHistory.ViewportRect.Min;
+				ViewportExtent = View.PrevViewInfo.TemporalAAHistory.ViewportRect.Size();
+
+				PassParameters->PrevSceneColor = GraphBuilder.RegisterExternalTexture(View.PrevViewInfo.TemporalAAHistory.RT[0]);
+				PassParameters->PrevSceneColorSampler = TStaticSamplerState<SF_Bilinear>::GetRHI();
+			}
+
+			PassParameters->PrevSceneColorPreExposureCorrection = View.PreExposure / View.PrevViewInfo.SceneColorPreExposure;
+			PassParameters->MinimumLuminance = CVarSSGIMinimumLuminance.GetValueOnRenderThread();
 
 			PassParameters->PrevScreenPositionScaleBias = FVector4(
 				ViewportExtent.X * 0.5f / BufferSize.X,
@@ -870,29 +958,23 @@ void RenderScreenSpaceDiffuseIndirect(
 				(ViewportExtent.X * 0.5f + ViewportOffset.X) / BufferSize.X,
 				(ViewportExtent.Y * 0.5f + ViewportOffset.Y) / BufferSize.Y);
 
-			PassParameters->PrevSceneColorPreExposureCorrection = View.PreExposure / View.PrevViewInfo.SceneColorPreExposure;
-			PassParameters->MinimumLuminance = CVarSSGIMinimumLuminance.GetValueOnRenderThread();
-
-			PassParameters->bIsTemporalAAHistory = true;
-			PassParameters->PrevTemporalAAHistory = ColorTexture;
-			PassParameters->PrevTemporalAAHistorySampler = TStaticSamplerState<SF_Bilinear>::GetRHI();
-
-			PassParameters->SceneTextures = SceneTextures;
-			SetupSceneTextureSamplers(&PassParameters->SceneTextureSamplers);
-			PassParameters->View = View.ViewUniformBuffer;
-
 			for (int32 MipLevel = 0; MipLevel < (PassParameters->ReducedSceneColorOutput.Num() - DownSamplingMip); MipLevel++)
 			{
 				PassParameters->ReducedSceneColorOutput[DownSamplingMip + MipLevel] = GraphBuilder.CreateUAV(FRDGTextureUAVDesc(ReducedSceneColor, MipLevel));
+				if (ReducedSceneAlpha)
+					PassParameters->ReducedSceneAlphaOutput[DownSamplingMip + MipLevel] = GraphBuilder.CreateUAV(FRDGTextureUAVDesc(ReducedSceneAlpha, MipLevel));
 			}
 
 			FSSRTPrevFrameReductionCS::FPermutationDomain PermutationVector;
-			PermutationVector.Set<FSSRTPrevFrameReductionCS::FLowerMips>(false);
+			PermutationVector.Set<FSSRTPrevFrameReductionCS::FLowerMips>(false); 
+			PermutationVector.Set<FSSRTPrevFrameReductionCS::FLeakFree>(bUseLeakFree);
 
 			TShaderMapRef<FSSRTPrevFrameReductionCS> ComputeShader(View.ShaderMap, PermutationVector);
 			FComputeShaderUtils::AddPass(
 				GraphBuilder,
-				RDG_EVENT_NAME("PrevFrameReduction %dx%d", View.ViewRect.Width(), View.ViewRect.Height()),
+				RDG_EVENT_NAME("PrevFrameReduction(LeakFree=%i) %dx%d",
+					bUseLeakFree ? 1 : 0,
+					View.ViewRect.Width(), View.ViewRect.Height()),
 				*ComputeShader,
 				PassParameters,
 				FComputeShaderUtils::GetGroupCount(View.ViewRect.Size(), 8));
@@ -905,26 +987,43 @@ void RenderScreenSpaceDiffuseIndirect(
 			int32 Divisor = 1 << (StartDestMip + DownSamplingMip);
 
 			FSSRTPrevFrameReductionCS::FParameters* PassParameters = GraphBuilder.AllocParameters<FSSRTPrevFrameReductionCS::FParameters>();
+			*PassParameters = DefaultPassParameters;
+
 			PassParameters->HigherMipTexture = GraphBuilder.CreateSRV(FRDGTextureSRVDesc::CreateForMipLevel(ReducedSceneColor, SrcMip));
-			PassParameters->HigherMipTextureSampler = TStaticSamplerState<SF_Bilinear>::GetRHI();
-			PassParameters->TexelSize = Divisor;
+			if (bUseLeakFree)
+			{
+				PassParameters->HigherMipTextureSampler = TStaticSamplerState<SF_Point>::GetRHI();
+				PassParameters->HigherAlphaMipTexture = GraphBuilder.CreateSRV(FRDGTextureSRVDesc::CreateForMipLevel(ReducedSceneAlpha, SrcMip));
+				PassParameters->HigherAlphaMipTextureSampler = TStaticSamplerState<SF_Point>::GetRHI();
+			}
+			else
+			{
+				PassParameters->HigherMipTextureSampler = TStaticSamplerState<SF_Bilinear>::GetRHI();
+			}
+
+			PassParameters->HigherMipDownScaleFactor = 1 << (DownSamplingMip + SrcMip);
 
-			PassParameters->SceneTextures = SceneTextures;
-			SetupSceneTextureSamplers(&PassParameters->SceneTextureSamplers);
-			PassParameters->View = View.ViewUniformBuffer;
+			PassParameters->ViewportUVToHZBBufferUV = ViewportUVToHZBBufferUV;
+			PassParameters->FurthestHZBTexture = FurthestHZBTexture;
+			PassParameters->FurthestHZBTextureSampler = TStaticSamplerState<SF_Point>::GetRHI();
 
 			for (int32 MipLevel = 0; MipLevel < PassParameters->ReducedSceneColorOutput.Num(); MipLevel++)
 			{
 				PassParameters->ReducedSceneColorOutput[MipLevel] = GraphBuilder.CreateUAV(FRDGTextureUAVDesc(ReducedSceneColor, StartDestMip + MipLevel));
+				if (ReducedSceneAlpha)
+					PassParameters->ReducedSceneAlphaOutput[MipLevel] = GraphBuilder.CreateUAV(FRDGTextureUAVDesc(ReducedSceneAlpha, StartDestMip + MipLevel));
 			}
 
 			FSSRTPrevFrameReductionCS::FPermutationDomain PermutationVector;
 			PermutationVector.Set<FSSRTPrevFrameReductionCS::FLowerMips>(true);
+			PermutationVector.Set<FSSRTPrevFrameReductionCS::FLeakFree>(bUseLeakFree);
 
 			TShaderMapRef<FSSRTPrevFrameReductionCS> ComputeShader(View.ShaderMap, PermutationVector);
 			ClearUnusedGraphResources(*ComputeShader, PassParameters);
 			GraphBuilder.AddPass(
-				RDG_EVENT_NAME("PrevFrameReduction %dx%d", View.ViewRect.Width() / Divisor, View.ViewRect.Height() / Divisor),
+				RDG_EVENT_NAME("PrevFrameReduction(LeakFree=%i) %dx%d",
+					bUseLeakFree ? 1 : 0,
+					View.ViewRect.Width() / Divisor, View.ViewRect.Height() / Divisor),
 				PassParameters,
 				ERDGPassFlags::Compute | ERDGPassFlags::GenerateMips,
 				[PassParameters, ComputeShader, &View, Divisor](FRHICommandList& RHICmdList)
@@ -1013,31 +1112,29 @@ void RenderScreenSpaceDiffuseIndirect(
 			PassParameters->FullResPixelOffset = FVector2D(0.5f, 0.5f);
 		}
 
+		{
+			// float2 SceneBufferUV;
+			// float2 PixelPos = SceneBufferUV * View.BufferSizeAndInvSize.xy - View.ViewRect.Min;
+			// PixelPos *= 0.5 // ReducedSceneColor is half resolution.
+			// float2 ReducedSceneColorUV = PixelPos / ReducedSceneColor->Extent;
+
+			PassParameters->ColorBufferScaleBias = FVector4(
+				0.5f * SceneTextures.SceneDepthBuffer->Desc.Extent.X / float(ReducedSceneColor->Desc.Extent.X),
+				0.5f * SceneTextures.SceneDepthBuffer->Desc.Extent.Y / float(ReducedSceneColor->Desc.Extent.Y),
+				-0.5f * View.ViewRect.Min.X / float(ReducedSceneColor->Desc.Extent.X),
+				-0.5f * View.ViewRect.Min.Y / float(ReducedSceneColor->Desc.Extent.Y));
+		}
+
 		PassParameters->FurthestHZBTexture = FurthestHZBTexture;
 		PassParameters->FurthestHZBTextureSampler = TStaticSamplerState<SF_Point>::GetRHI();
 		PassParameters->ColorTexture = ReducedSceneColor;
 		PassParameters->ColorTextureSampler = TStaticSamplerState<SF_Bilinear>::GetRHI();
 
-		const FVector2D HZBUvFactor(
-			float( View.ViewRect.Width() )  / float( 2 * View.HZBMipmap0Size.X ),
-			float( View.ViewRect.Height() ) / float( 2 * View.HZBMipmap0Size.Y )
-			);
-			
 		PassParameters->HZBUvFactorAndInvFactor = FVector4(
-			HZBUvFactor.X,
-			HZBUvFactor.Y,
-			1.0f / HZBUvFactor.X,
-			1.0f / HZBUvFactor.Y );
-
-		FIntPoint ViewportOffset	= TemporalAAHistory.ViewportRect.Min;
-		FIntPoint ViewportExtent	= TemporalAAHistory.ViewportRect.Size();
-		FIntPoint BufferSize		= TemporalAAHistory.ReferenceBufferSize;
-
-		PassParameters->PrevScreenPositionScaleBias = FVector4(
-			 ViewportExtent.X * 0.5f / BufferSize.X,
-			-ViewportExtent.Y * 0.5f / BufferSize.Y,
-			(ViewportExtent.X * 0.5f + ViewportOffset.X) / BufferSize.X,
-			(ViewportExtent.Y * 0.5f + ViewportOffset.Y) / BufferSize.Y );
+			ViewportUVToHZBBufferUV.X,
+			ViewportUVToHZBBufferUV.Y,
+			1.0f / ViewportUVToHZBBufferUV.X,
+			1.0f / ViewportUVToHZBBufferUV.Y );
 
 		PassParameters->SceneTextures = SceneTextures;
 		SetupSceneTextureSamplers(&PassParameters->SceneTextureSamplers);
diff --git a/Engine/Source/Runtime/Renderer/Private/ScreenSpaceRayTracing.h b/Engine/Source/Runtime/Renderer/Private/ScreenSpaceRayTracing.h
--- a/Engine/Source/Runtime/Renderer/Private/ScreenSpaceRayTracing.h
+++ b/Engine/Source/Runtime/Renderer/Private/ScreenSpaceRayTracing.h
@@ -30,10 +30,18 @@ struct FTiledScreenSpaceReflection
 	uint32 TileSize;
 };
 
+bool ShouldKeepBleedFreeSceneColor(const FViewInfo& View);
+
 bool ShouldRenderScreenSpaceReflections(const FViewInfo& View);
 
 bool ShouldRenderScreenSpaceDiffuseIndirect(const FViewInfo& View);
 
+void ProcessForNextFrameScreenSpaceRayTracing(
+	FRDGBuilder& GraphBuilder,
+	const FSceneTextureParameters& SceneTextures,
+	const FRDGTextureRef CurrentSceneColor,
+	const FViewInfo& View);
+
 void GetSSRQualityForView(const FViewInfo& View, ESSRQuality* OutQuality, IScreenSpaceDenoiser::FReflectionsRayTracingConfig* OutRayTracingConfigs);
 
 bool IsSSRTemporalPassRequired(const FViewInfo& View);
diff --git a/Engine/Source/Runtime/Renderer/Private/VelocityRendering.cpp b/Engine/Source/Runtime/Renderer/Private/VelocityRendering.cpp
--- a/Engine/Source/Runtime/Renderer/Private/VelocityRendering.cpp
+++ b/Engine/Source/Runtime/Renderer/Private/VelocityRendering.cpp
@@ -388,7 +388,9 @@ bool FDeferredShadingSceneRenderer::ShouldRenderVelocities() const
 		bool bRayTracing = IsRayTracingEnabled();
 		bool bDenoise = bRayTracing;
 
-		bNeedsVelocity |= bMotionBlur || bTemporalAA || bDistanceFieldAO || bSSRTemporal || bDenoise;
+		bool bSSGI = ShouldRenderScreenSpaceDiffuseIndirect(View);
+		
+		bNeedsVelocity |= bMotionBlur || bTemporalAA || bDistanceFieldAO || bSSRTemporal || bDenoise || bSSGI;
 	}
 
 	return bNeedsVelocity;
