Fixes compilation failure

diff --git a/Engine/Source/Runtime/Renderer/Private/ScreenSpaceRayTracing.cpp b/Engine/Source/Runtime/Renderer/Private/ScreenSpaceRayTracing.cpp
--- a/Engine/Source/Runtime/Renderer/Private/ScreenSpaceRayTracing.cpp
+++ b/Engine/Source/Runtime/Renderer/Private/ScreenSpaceRayTracing.cpp
@@ -92,7 +92,7 @@ static bool SupportScreenSpaceDiffuseIndirect(const FViewInfo& View)
 		return false;
 	}
 
-	return View.ViewState;
+	return View.ViewState != nullptr;
 }
 
 bool ShouldKeepBleedFreeSceneColor(const FViewInfo& View)
