Cherry-pick 10176375: Implements BufferUVBilinearCorrection in the denoiser

diff --git a/Engine/Shaders/Private/ScreenSpaceDenoise/SSDCommon.ush b/Engine/Shaders/Private/ScreenSpaceDenoise/SSDCommon.ush
--- a/Engine/Shaders/Private/ScreenSpaceDenoise/SSDCommon.ush
+++ b/Engine/Shaders/Private/ScreenSpaceDenoise/SSDCommon.ush
@@ -78,6 +78,7 @@ float4 BufferBilinearUVMinMax;
 float2 BufferUVToOutputPixelPosition;
 float4 BufferUVToScreenPosition;
 float4x4 ScreenToView;
+float2 BufferUVBilinearCorrection;
 
 float WorldDepthToPixelWorldRadius;
 
diff --git a/Engine/Shaders/Private/ScreenSpaceDenoise/SSDTemporalAccumulation.usf b/Engine/Shaders/Private/ScreenSpaceDenoise/SSDTemporalAccumulation.usf
--- a/Engine/Shaders/Private/ScreenSpaceDenoise/SSDTemporalAccumulation.usf
+++ b/Engine/Shaders/Private/ScreenSpaceDenoise/SSDTemporalAccumulation.usf
@@ -555,7 +555,7 @@ void TemporallyAccumulate(
 			#endif
 			
 			// VGPR configuration of the kernel.
-			KernelConfig.BufferUV = HistoryBufferUV;
+			KernelConfig.BufferUV = HistoryBufferUV + BufferUVBilinearCorrection;
 			{
 				KernelConfig.CompressedRefSceneMetadata = CompressedRefSceneMetadata;
 				KernelConfig.RefBufferUV = SceneBufferUV;
diff --git a/Engine/Source/Runtime/Renderer/Private/ScreenSpaceDenoise.cpp b/Engine/Source/Runtime/Renderer/Private/ScreenSpaceDenoise.cpp
--- a/Engine/Source/Runtime/Renderer/Private/ScreenSpaceDenoise.cpp
+++ b/Engine/Source/Runtime/Renderer/Private/ScreenSpaceDenoise.cpp
@@ -712,6 +712,7 @@ BEGIN_SHADER_PARAMETER_STRUCT(FSSDCommonParameters, )
 	SHADER_PARAMETER(float, WorldDepthToPixelWorldRadius)
 	SHADER_PARAMETER(FVector4, BufferUVToScreenPosition)
 	SHADER_PARAMETER(FMatrix, ScreenToView)
+	SHADER_PARAMETER(FVector2D, BufferUVBilinearCorrection)
 
 	SHADER_PARAMETER_STRUCT_INCLUDE(FSceneTextureParameters, SceneTextures)
 
@@ -1441,6 +1442,9 @@ static void DenoiseSignalAtConstantPixelDensity(
 			FPlane(0, 0, View.ProjectionMatrixUnadjustedForRHI.M[2][2], 1),
 			FPlane(0, 0, View.ProjectionMatrixUnadjustedForRHI.M[3][2], 0))
 			* View.ViewMatrices.GetInvProjectionMatrix();
+
+		CommonParameters.BufferUVBilinearCorrection.X = (0.5f * PixelPositionToFullResPixel - FullResPixelOffset.X) / float(FullResBufferExtent.X);
+		CommonParameters.BufferUVBilinearCorrection.Y = (0.5f * PixelPositionToFullResPixel - FullResPixelOffset.Y) / float(FullResBufferExtent.Y);
 	}
 
 	#if RHI_RAYTRACING
