Fix inverted filmic tonemapper on android GLES
#rnx

diff --git a/Engine/Source/Runtime/Renderer/Private/PostProcess/PostProcessTonemap.cpp b/Engine/Source/Runtime/Renderer/Private/PostProcess/PostProcessTonemap.cpp
--- a/Engine/Source/Runtime/Renderer/Private/PostProcess/PostProcessTonemap.cpp
+++ b/Engine/Source/Runtime/Renderer/Private/PostProcess/PostProcessTonemap.cpp
@@ -845,11 +845,6 @@ FScreenPassTexture AddTonemapPass(FRDGBuilder& GraphBuilder, const FViewInfo& Vi
 
 		EScreenPassDrawFlags DrawFlags = EScreenPassDrawFlags::AllowHMDHiddenAreaMask;
 
-		if (Inputs.bFlipYAxis)
-		{
-			DrawFlags |= EScreenPassDrawFlags::FlipYAxis;
-		}
-
 		AddDrawScreenPass(
 			GraphBuilder,
 			RDG_EVENT_NAME("Tonemap %dx%d (PS GammaOnly=%d)", OutputViewport.Rect.Width(), OutputViewport.Rect.Height(), Inputs.bGammaOnly),
