Fix an issue that causes EngineMisc memory keep increasing when world rendering is disabled

diff --git a/Engine/Source/Runtime/Engine/Private/GameViewportClient.cpp b/Engine/Source/Runtime/Engine/Private/GameViewportClient.cpp
--- a/Engine/Source/Runtime/Engine/Private/GameViewportClient.cpp
+++ b/Engine/Source/Runtime/Engine/Private/GameViewportClient.cpp
@@ -1512,6 +1512,8 @@ void UGameViewportClient::Draw(FViewport* InViewport, FCanvas* SceneCanvas)
 	}
 	else
 	{
+		GetRendererModule().PerFrameCleanupIfSkipRenderer();
+
 		// Make sure RHI resources get flushed if we're not using a renderer
 		ENQUEUE_RENDER_COMMAND(UGameViewportClient_FlushRHIResources)(
 			[](FRHICommandListImmediate& RHICmdList)
diff --git a/Engine/Source/Runtime/RenderCore/Public/RendererInterface.h b/Engine/Source/Runtime/RenderCore/Public/RendererInterface.h
--- a/Engine/Source/Runtime/RenderCore/Public/RendererInterface.h
+++ b/Engine/Source/Runtime/RenderCore/Public/RendererInterface.h
@@ -710,6 +710,9 @@ public:
 
 	virtual void PostRenderAllViewports() = 0;
 
+	/** Performs necessary per-frame cleanup. Only use when rendering through scene renderer (i.e. BeginRenderingViewFamily) is skipped */
+	virtual void PerFrameCleanupIfSkipRenderer() = 0;
+
 	virtual IAllocatedVirtualTexture* AllocateVirtualTexture(const FAllocatedVTDescription& Desc) = 0;
 	virtual void DestroyVirtualTexture(IAllocatedVirtualTexture* AllocatedVT) = 0;
 
diff --git a/Engine/Source/Runtime/Renderer/Private/RendererModule.h b/Engine/Source/Runtime/Renderer/Private/RendererModule.h
--- a/Engine/Source/Runtime/Renderer/Private/RendererModule.h
+++ b/Engine/Source/Runtime/Renderer/Private/RendererModule.h
@@ -96,6 +96,8 @@ public:
 
 	virtual void PostRenderAllViewports() override;
 
+	virtual void PerFrameCleanupIfSkipRenderer() override;
+
 	virtual IAllocatedVirtualTexture* AllocateVirtualTexture(const FAllocatedVTDescription& Desc) override;
 	virtual void DestroyVirtualTexture(IAllocatedVirtualTexture* AllocatedVT) override;
 	virtual FVirtualTextureProducerHandle RegisterVirtualTextureProducer(const FVTProducerDescription& Desc, IVirtualTexture* Producer) override;
diff --git a/Engine/Source/Runtime/Renderer/Private/SceneRendering.cpp b/Engine/Source/Runtime/Renderer/Private/SceneRendering.cpp
--- a/Engine/Source/Runtime/Renderer/Private/SceneRendering.cpp
+++ b/Engine/Source/Runtime/Renderer/Private/SceneRendering.cpp
@@ -3704,6 +3704,16 @@ void FRendererModule::PostRenderAllViewports()
 	++GFrameNumber;
 }
 
+void FRendererModule::PerFrameCleanupIfSkipRenderer()
+{
+	// Some systems (e.g. Slate) can still draw (via FRendererModule::DrawTileMesh for example) when scene renderer is not used
+	ENQUEUE_RENDER_COMMAND(CmdPerFrameCleanupIfSkipRenderer)(
+		[](FRHICommandListImmediate&)
+	{
+		GPrimitiveIdVertexBufferPool.DiscardAll();
+	});
+}
+
 void FRendererModule::UpdateMapNeedsLightingFullyRebuiltState(UWorld* World)
 {
 	World->SetMapNeedsLightingFullyRebuilt(World->Scene->GetRenderScene()->NumUncachedStaticLightingInteractions, World->Scene->GetRenderScene()->NumUnbuiltReflectionCaptures);
