Fix for wrong format dummy StencilSRV assigned to SceneTextures uniform buffer.

diff --git a/Engine/Source/Runtime/Renderer/Private/PostProcess/SceneRenderTargets.cpp b/Engine/Source/Runtime/Renderer/Private/PostProcess/SceneRenderTargets.cpp
--- a/Engine/Source/Runtime/Renderer/Private/PostProcess/SceneRenderTargets.cpp
+++ b/Engine/Source/Runtime/Renderer/Private/PostProcess/SceneRenderTargets.cpp
@@ -3077,7 +3077,7 @@ void SetupSceneTextureUniformParameters(
 		}
 
 		SceneTextureParameters.SceneDepthTextureSampler = TStaticSamplerState<SF_Point, AM_Clamp, AM_Clamp, AM_Clamp>::GetRHI();
-		SceneTextureParameters.SceneStencilTexture = bSetupDepth && SceneContext.SceneStencilSRV ? SceneContext.SceneStencilSRV : GNullColorVertexBuffer.VertexBufferSRV;
+		SceneTextureParameters.SceneStencilTexture = bSetupDepth && SceneContext.SceneStencilSRV ? SceneContext.SceneStencilSRV : GSystemTextures.StencilDummySRV;
 	}
 
 	// GBuffer
@@ -3129,7 +3129,7 @@ void SetupSceneTextureUniformParameters(
 		const bool bSetupCustomDepth = (SetupMode & ESceneTextureSetupMode::CustomDepth) != ESceneTextureSetupMode::None;
 
 		FRHITexture* CustomDepth = DepthDefault;
-		FRHIShaderResourceView* CustomStencilSRV = GSystemTextures.WhiteDummySRV;
+		FRHIShaderResourceView* CustomStencilSRV = GSystemTextures.StencilDummySRV;
 
 		// if there is no custom depth it's better to have the far distance there
 		IPooledRenderTarget* CustomDepthTarget = SceneContext.bCustomDepthIsValid ? SceneContext.CustomDepth.GetReference() : 0;
diff --git a/Engine/Source/Runtime/Renderer/Private/SystemTextures.cpp b/Engine/Source/Runtime/Renderer/Private/SystemTextures.cpp
--- a/Engine/Source/Runtime/Renderer/Private/SystemTextures.cpp
+++ b/Engine/Source/Runtime/Renderer/Private/SystemTextures.cpp
@@ -141,7 +141,21 @@ void FSystemTextures::InitializeCommonTextures(FRHICommandListImmediate& RHICmdL
 		RHICmdList.EndRenderPass();
 		RHICmdList.CopyToResolveTarget(DepthDummy->GetRenderTargetItem().TargetableTexture, DepthDummy->GetRenderTargetItem().ShaderResourceTexture, FResolveParams());
 	}
-	
+
+	// Create a dummy stencil SRV.
+	{
+		FPooledRenderTargetDesc Desc(FPooledRenderTargetDesc::Create2DDesc(FIntPoint(1, 1), PF_R8G8B8A8_UINT, FClearValueBinding::White, TexCreate_HideInVisualizeTexture, TexCreate_RenderTargetable | TexCreate_NoFastClear, false));
+		Desc.AutoWritable = false;
+		GRenderTargetPool.FindFreeElement(RHICmdList, Desc, StencilDummy, TEXT("StencilDummy"), true, ERenderTargetTransience::NonTransient);
+
+		FRHIRenderPassInfo RPInfo(StencilDummy->GetRenderTargetItem().TargetableTexture, ERenderTargetActions::Clear_Store);
+		RHICmdList.BeginRenderPass(RPInfo, TEXT("StencilDummy"));
+		RHICmdList.EndRenderPass();
+		RHICmdList.CopyToResolveTarget(StencilDummy->GetRenderTargetItem().TargetableTexture, StencilDummy->GetRenderTargetItem().ShaderResourceTexture, FResolveParams());
+
+		StencilDummySRV = RHICreateShaderResourceView((FRHITexture2D*)StencilDummy->GetRenderTargetItem().ShaderResourceTexture.GetReference(), 0);
+	}
+
 	if (!GSupportsShaderFramebufferFetch && GPixelFormats[PF_FloatRGBA].Supported)
 	{
 		// PF_FloatRGBA to encode exactly the 0.5.
@@ -675,6 +689,11 @@ FRDGTextureRef FSystemTextures::GetDepthDummy(FRDGBuilder& GraphBuilder) const
 	return GraphBuilder.RegisterExternalTexture(DepthDummy, TEXT("DepthDummy"));
 }
 
+FRDGTextureRef FSystemTextures::GetStencilDummy(FRDGBuilder& GraphBuilder) const
+{
+	return GraphBuilder.RegisterExternalTexture(StencilDummy, TEXT("StencilDummy"));
+}
+
 FRDGTextureRef FSystemTextures::GetGreenDummy(FRDGBuilder& GraphBuilder) const
 {
 	return GraphBuilder.RegisterExternalTexture(GreenDummy, TEXT("GreenDummy"));
diff --git a/Engine/Source/Runtime/Renderer/Private/SystemTextures.h b/Engine/Source/Runtime/Renderer/Private/SystemTextures.h
--- a/Engine/Source/Runtime/Renderer/Private/SystemTextures.h
+++ b/Engine/Source/Runtime/Renderer/Private/SystemTextures.h
@@ -80,6 +80,8 @@ public:
 	TRefCountPtr<IPooledRenderTarget> MaxFP16Depth;
 	/** Depth texture that holds a single depth value */
 	TRefCountPtr<IPooledRenderTarget> DepthDummy;
+	/** Stencil texture that holds a single stencil value. */
+	TRefCountPtr<IPooledRenderTarget> StencilDummy;
 	// float4(0,1,0,1)
 	TRefCountPtr<IPooledRenderTarget> GreenDummy;
 	// float4(0.5,0.5,0.5,1)
@@ -92,6 +94,8 @@ public:
 
 	// SRV for WhiteDummy Texture.
 	TRefCountPtr<FRHIShaderResourceView> WhiteDummySRV;
+	// SRV for StencilDummy Texture.
+	TRefCountPtr<FRHIShaderResourceView> StencilDummySRV;
 
 	FRDGTextureRef GetWhiteDummy(FRDGBuilder& GraphBuilder) const;
 	FRDGTextureRef GetBlackDummy(FRDGBuilder& GraphBuilder) const;
@@ -105,6 +109,7 @@ public:
 	FRDGTextureRef GetLTCAmp(FRDGBuilder& GraphBuilder) const;
 	FRDGTextureRef GetMaxFP16Depth(FRDGBuilder& GraphBuilder) const;
 	FRDGTextureRef GetDepthDummy(FRDGBuilder& GraphBuilder) const;
+	FRDGTextureRef GetStencilDummy(FRDGBuilder& GraphBuilder) const;
 	FRDGTextureRef GetGreenDummy(FRDGBuilder& GraphBuilder) const;
 	FRDGTextureRef GetDefaultNormal8Bit(FRDGBuilder& GraphBuilder) const;
 	FRDGTextureRef GetMidGreyDummy(FRDGBuilder& GraphBuilder) const;
