Cherry-pick 10232023: Fixes validation error when using r.TemporalAA.HistoryScreenPercentage

diff --git a/Engine/Source/Runtime/Renderer/Private/PostProcess/PostProcessMitchellNetravali.cpp b/Engine/Source/Runtime/Renderer/Private/PostProcess/PostProcessMitchellNetravali.cpp
--- a/Engine/Source/Runtime/Renderer/Private/PostProcess/PostProcessMitchellNetravali.cpp
+++ b/Engine/Source/Runtime/Renderer/Private/PostProcess/PostProcessMitchellNetravali.cpp
@@ -46,7 +46,7 @@ FRDGTextureRef ComputeMitchellNetravaliDownsample(
 		PF_FloatRGBA,
 		FClearValueBinding::Black,
 		TexCreate_None,
-		TexCreate_UAV,
+		TexCreate_ShaderResource | TexCreate_UAV,
 		false);
 
 	FRDGTextureRef OutputTexture = GraphBuilder.CreateTexture(OutputTextureDesc, TEXT("MitchelNetravaliDownsampleOutput"));
