Add (some) material view mode for hair.

diff --git a/Engine/Shaders/Private/HairStrands/HairStrandsDebug.usf b/Engine/Shaders/Private/HairStrands/HairStrandsDebug.usf
--- a/Engine/Shaders/Private/HairStrands/HairStrandsDebug.usf
+++ b/Engine/Shaders/Private/HairStrands/HairStrandsDebug.usf
@@ -5,16 +5,25 @@
 
 float2 OutputResolution;
 uint FastResolveMask;
+uint DebugMode;
+int SampleIndex;
 
 Texture2D<uint4> CategorizationTexture;
 Texture2D<float> ViewTransmittanceTexture;
 Texture2D<uint2> DepthStencilTexture;
+Texture2D<uint> NodeIndex;
+StructuredBuffer<FNodeData> NodeData;
 SamplerState LinearSampler;
 
 #define DEBUG_MODE_SAMPLE_COUNT 0 
 #define DEBUG_MODE_COVERAGE_TYPE 1
 #define DEBUG_MODE_TAA_RESOLVE_TYPE 2
 #define DEBUG_MODE_COVERAGE 3
+#define DEBUG_MODE_DEPTH 4
+#define DEBUG_MODE_BASECOLOR 5
+#define DEBUG_MODE_ROUGHNESS 6
+#define DEBUG_MODE_SPECULAR 7
+#define DEBUG_MODE_TANGENT 8
 
 float3 HUEtoRGB(in float H)
 {
@@ -48,25 +57,112 @@ void MainPS(
 	if (!bIsValid)
 		discard;
 
-#if PERMUTATION_DEBUG_MODE == DEBUG_MODE_SAMPLE_COUNT
-	const float3 Color = GetColorCode(CatData.SampleCount/float(HAIR_VISIBILITY_MAXSAMPLECOUNT));
-#endif
-
-#if PERMUTATION_DEBUG_MODE == DEBUG_MODE_COVERAGE_TYPE
-	const bool bIsFullyCovered = CatData.PixelCoverage >= 1;
-	const float3 Color = bIsFullyCovered ? float3(0.25,0.85,0.25) : float3(0.85,0.25,0.25);
-#endif
-
-#if PERMUTATION_DEBUG_MODE == DEBUG_MODE_TAA_RESOLVE_TYPE
-	const uint Stencil = DepthStencilTexture.Load(uint3(PixelPos,0)).y;
-	const bool bIsFastResolve = (Stencil & FastResolveMask) != 0;
-	const float3 Color = bIsFastResolve ? float3(1,1,0) : float3(0.15f, 0.25f, 0.75f);
-#endif
-
-#if PERMUTATION_DEBUG_MODE == DEBUG_MODE_COVERAGE
-	const float3 Color = GetColorCode(saturate(CatData.PixelCoverage));
-#endif
-	
+	float3 Color = 0;
+	if (DebugMode == DEBUG_MODE_SAMPLE_COUNT)
+	{
+		Color = GetColorCode(CatData.SampleCount/float(HAIR_VISIBILITY_MAXSAMPLECOUNT));
+	}
+
+	if (DebugMode == DEBUG_MODE_COVERAGE_TYPE)
+	{
+		const bool bIsFullyCovered = CatData.PixelCoverage >= 1;
+		Color = bIsFullyCovered ? float3(0.25,0.85,0.25) : float3(0.85,0.25,0.25);
+	}
+
+	if (DebugMode == DEBUG_MODE_TAA_RESOLVE_TYPE)
+	{
+		const uint Stencil = DepthStencilTexture.Load(uint3(PixelPos,0)).y;
+		const bool bIsFastResolve = (Stencil & FastResolveMask) != 0;
+		Color = bIsFastResolve ? float3(1,1,0) : float3(0.15f, 0.25f, 0.75f);
+	}
+
+	if (DebugMode == DEBUG_MODE_COVERAGE)
+	{
+		Color = GetColorCode(saturate(CatData.PixelCoverage));
+	}
+
+	const bool bIsDebugModeMaterial =
+		DebugMode == DEBUG_MODE_DEPTH ||
+		DebugMode == DEBUG_MODE_BASECOLOR ||
+		DebugMode == DEBUG_MODE_ROUGHNESS ||
+		DebugMode == DEBUG_MODE_SPECULAR || 
+		DebugMode == DEBUG_MODE_TANGENT;
+
+	const FNodeDesc NodeDesc = DecodeNodeDesc(NodeIndex.Load(uint3(PixelPos, 0)));
+	if (NodeDesc.Count > 0 && bIsDebugModeMaterial)
+	{
+		const int LocalSampleIndex = SampleIndex; 
+		const bool bAverageAllSamples = LocalSampleIndex < 0;
+		const uint SelectedIndex = LocalSampleIndex >= 0 ? clamp(LocalSampleIndex, 0, NodeDesc.Count - 1) : 0;
+		float	Depth = 0;
+		float3	Tangent = 0;
+		float3	BaseColor = 0;
+		float	Roughness = 0;
+		float	Specular = 0;
+
+		for (uint NodeIt = 0; NodeIt < NodeDesc.Count; ++NodeIt)
+		{
+			const uint NodeOffset = NodeDesc.Offset + NodeIt;
+			const FNodeData Data = NodeData.Load(NodeOffset);
+			float			SampleDepth;
+			FVisibilityData SampleVisData;
+			float3			SampleBaseColor;
+			float			SampleRoughness;
+			float			SampleSpecular;
+			DecodeNodeData(Data, SampleDepth, SampleVisData, SampleBaseColor, SampleRoughness, SampleSpecular);
+
+			if (bAverageAllSamples)
+			{
+				Depth		+= SampleDepth;
+				Tangent		+= SampleVisData.Tangent;
+				BaseColor	+= SampleBaseColor;
+				Roughness	+= SampleRoughness;
+				Specular	+= SampleSpecular;
+			}
+			else if (SelectedIndex == NodeIt)
+			{
+				Depth		= SampleDepth;
+				Tangent		= SampleVisData.Tangent;
+				BaseColor	= SampleBaseColor;
+				Roughness	= SampleRoughness;
+				Specular	= SampleSpecular;
+			}
+		}
+
+		if (bAverageAllSamples)
+		{
+			Depth		/= NodeDesc.Count;
+			Tangent		/= NodeDesc.Count;
+			BaseColor	/= NodeDesc.Count;
+			Roughness	/= NodeDesc.Count;
+			Specular	/= NodeDesc.Count;
+		}
+
+		if (DebugMode == DEBUG_MODE_DEPTH)
+		{
+			Color = Depth;
+		}
+
+		if (DebugMode == DEBUG_MODE_BASECOLOR)
+		{
+			Color = BaseColor;
+		}
+
+		if (DebugMode == DEBUG_MODE_ROUGHNESS)
+		{
+			Color = Roughness;
+		}
+
+		if (DebugMode == DEBUG_MODE_SPECULAR)
+		{
+			Color = Specular;
+		}
+
+		if (DebugMode == DEBUG_MODE_TANGENT)
+		{
+			Color = (Tangent+1) * 0.5f;
+		}
+	}
 
 	OutColor = float4(Color, 1);
 }
diff --git a/Engine/Source/Runtime/Renderer/Private/HairStrands/HairStrandsDebug.cpp b/Engine/Source/Runtime/Renderer/Private/HairStrands/HairStrandsDebug.cpp
--- a/Engine/Source/Runtime/Renderer/Private/HairStrands/HairStrandsDebug.cpp
+++ b/Engine/Source/Runtime/Renderer/Private/HairStrands/HairStrandsDebug.cpp
@@ -48,6 +48,9 @@ static FAutoConsoleVariableRef CVarHairStrandsDebugBSDFAbsorption(TEXT("r.HairSt
 static float GHairStrandsDebugPlotBsdfExposure = 1.1f;
 static FAutoConsoleVariableRef CVarHairStrandsDebugBSDFExposure(TEXT("r.HairStrands.PlotBsdf.Exposure"), GHairStrandsDebugPlotBsdfExposure, TEXT("Change the exposure of the plot."));
 
+static int GHairStrandsDebugSampleIndex = -1;
+static FAutoConsoleVariableRef CVarHairStrandsDebugMaterialSampleIndex(TEXT("r.HairStrands.DebugMode.SampleIndex"), GHairStrandsDebugSampleIndex, TEXT("Debug value for a given sample index (default:-1, i.e., average sample information)."));
+
 static int32 GHairDebugMeshProjection_SkinCacheMesh = 0;
 
 static int32 GHairDebugMeshProjection_Sim_HairRestTriangles = 0;
@@ -111,7 +114,12 @@ enum class EHairDebugMode : uint8
 	VoxelsBaseColor,
 	VoxelsRoughness,
 	MeshProjection,
-	Coverage
+	Coverage,
+	MaterialDepth,
+	MaterialBaseColor,
+	MaterialRoughness,
+	MaterialSpecular,
+	MaterialTangent
 };
 
 static EHairDebugMode GetHairDebugMode()
@@ -132,6 +140,11 @@ static EHairDebugMode GetHairDebugMode()
 	case 11: return EHairDebugMode::VoxelsRoughness;
 	case 12: return EHairDebugMode::MeshProjection;
 	case 13: return EHairDebugMode::Coverage;
+	case 14: return EHairDebugMode::MaterialDepth;
+	case 15: return EHairDebugMode::MaterialBaseColor;
+	case 16: return EHairDebugMode::MaterialRoughness;
+	case 17: return EHairDebugMode::MaterialSpecular;
+	case 18: return EHairDebugMode::MaterialTangent;
 	default: return EHairDebugMode::None;
 	};
 }
@@ -154,6 +167,11 @@ static const TCHAR* ToString(EHairDebugMode DebugMode)
 	case EHairDebugMode::VoxelsRoughness: return TEXT("Hair roughness volume");
 	case EHairDebugMode::MeshProjection: return TEXT("Hair mesh projection");
 	case EHairDebugMode::Coverage: return TEXT("Hair coverage");
+	case EHairDebugMode::MaterialDepth: return TEXT("Hair material depth");
+	case EHairDebugMode::MaterialBaseColor: return TEXT("Hair material base color");
+	case EHairDebugMode::MaterialRoughness: return TEXT("Hair material roughness");
+	case EHairDebugMode::MaterialSpecular: return TEXT("Hair material specular");
+	case EHairDebugMode::MaterialTangent: return TEXT("Hair material tangent");
 	default: return TEXT("None");
 	};
 }
@@ -273,13 +291,14 @@ class FHairDebugPS : public FGlobalShader
 	DECLARE_GLOBAL_SHADER(FHairDebugPS);
 	SHADER_USE_PARAMETER_STRUCT(FHairDebugPS, FGlobalShader);
 
-	class FDebugMode : SHADER_PERMUTATION_INT("PERMUTATION_DEBUG_MODE", 4);
-	using FPermutationDomain = TShaderPermutationDomain<FDebugMode>;
-
 	BEGIN_SHADER_PARAMETER_STRUCT(FParameters, )
 		SHADER_PARAMETER(FVector2D, OutputResolution)
 		SHADER_PARAMETER(uint32, FastResolveMask)
+		SHADER_PARAMETER(uint32, DebugMode)
+		SHADER_PARAMETER(int32, SampleIndex)
 		SHADER_PARAMETER_RDG_TEXTURE(Texture2D, CategorizationTexture)
+		SHADER_PARAMETER_RDG_TEXTURE(Texture2D, NodeIndex)
+		SHADER_PARAMETER_RDG_BUFFER_SRV(StructuredBuffer, NodeData)
 		SHADER_PARAMETER_SRV(Texture2D, DepthStencilTexture)
 		SHADER_PARAMETER_SAMPLER(SamplerState, LinearSampler)
 		SHADER_PARAMETER_STRUCT_REF(FViewUniformShaderParameters, ViewUniformBuffer)
@@ -297,40 +316,61 @@ static void AddDebugHairPass(
 	const FViewInfo* View,
 	const EHairDebugMode InDebugMode,
 	const TRefCountPtr<IPooledRenderTarget>& InCategorizationTexture,
+	const TRefCountPtr<IPooledRenderTarget>& InNodeIndex,
+	const TRefCountPtr<FPooledRDGBuffer>& InNodeData,
 	const FShaderResourceViewRHIRef& InDepthStencilTexture,
 	FRDGTextureRef& OutTarget)
 {
 	check(OutTarget);
-	check(InDebugMode == EHairDebugMode::TAAResolveType || InDebugMode == EHairDebugMode::SamplePerPixel || InDebugMode == EHairDebugMode::CoverageType || InDebugMode == EHairDebugMode::Coverage);
-
-	if (!InCategorizationTexture) return;
+	check(InDebugMode == EHairDebugMode::TAAResolveType || 
+		InDebugMode == EHairDebugMode::SamplePerPixel || 
+		InDebugMode == EHairDebugMode::CoverageType || 
+		InDebugMode == EHairDebugMode::Coverage ||
+		InDebugMode == EHairDebugMode::MaterialDepth ||
+		InDebugMode == EHairDebugMode::MaterialBaseColor ||
+		InDebugMode == EHairDebugMode::MaterialRoughness ||
+		InDebugMode == EHairDebugMode::MaterialSpecular ||
+		InDebugMode == EHairDebugMode::MaterialTangent);
+
+	if (!InCategorizationTexture || !InNodeIndex || !InNodeData) return;
 	if (InDebugMode == EHairDebugMode::TAAResolveType && !InDepthStencilTexture) return;
 
 	FRDGTextureRef CategorizationTexture = InCategorizationTexture ? GraphBuilder.RegisterExternalTexture(InCategorizationTexture, TEXT("CategorizationTexture")) : nullptr;
+	FRDGTextureRef NodeIndex = InNodeIndex ? GraphBuilder.RegisterExternalTexture(InNodeIndex, TEXT("NodeIndex")) : nullptr;
+	FRDGBufferRef  NodeData = InNodeData ? GraphBuilder.RegisterExternalBuffer(InNodeData, TEXT("NodeData")) : nullptr;
 
 	const FIntRect Viewport = View->ViewRect;
 	const FIntPoint Resolution(Viewport.Width(), Viewport.Height());
 
+	uint32 InternalDebugMode = 0;
+	switch (InDebugMode)
+	{
+		case EHairDebugMode::SamplePerPixel:	InternalDebugMode = 0; break;
+		case EHairDebugMode::CoverageType:		InternalDebugMode = 1; break;
+		case EHairDebugMode::TAAResolveType:	InternalDebugMode = 2; break;
+		case EHairDebugMode::Coverage:			InternalDebugMode = 3; break;
+		case EHairDebugMode::MaterialDepth:		InternalDebugMode = 4; break;
+		case EHairDebugMode::MaterialBaseColor:	InternalDebugMode = 5; break;
+		case EHairDebugMode::MaterialRoughness:	InternalDebugMode = 6; break;
+		case EHairDebugMode::MaterialSpecular:	InternalDebugMode = 7; break;
+		case EHairDebugMode::MaterialTangent:	InternalDebugMode = 8; break;
+	};
+
 	FHairDebugPS::FParameters* Parameters = GraphBuilder.AllocParameters<FHairDebugPS::FParameters>();
+	Parameters->ViewUniformBuffer = View->ViewUniformBuffer;
 	Parameters->OutputResolution = Resolution;
 	Parameters->FastResolveMask = STENCIL_TEMPORAL_RESPONSIVE_AA_MASK;
 	Parameters->CategorizationTexture = CategorizationTexture;
+	Parameters->NodeIndex = NodeIndex;
+	Parameters->NodeData = GraphBuilder.CreateSRV(NodeData);
 	Parameters->DepthStencilTexture = InDepthStencilTexture;
 	Parameters->LinearSampler = TStaticSamplerState<SF_Bilinear, AM_Clamp, AM_Clamp, AM_Clamp>::GetRHI();
+	Parameters->DebugMode = InternalDebugMode;
+	Parameters->SampleIndex = GHairStrandsDebugSampleIndex;
 	Parameters->RenderTargets[0] = FRenderTargetBinding(OutTarget, ERenderTargetLoadAction::ELoad, 0);
 	TShaderMapRef<FPostProcessVS> VertexShader(View->ShaderMap);
 
-	FHairDebugPS::FPermutationDomain PermutationVector;
-	uint32 DebugPermutation = 0;
-	switch (InDebugMode)
-	{
-		case EHairDebugMode::SamplePerPixel:	DebugPermutation = 0; break;
-		case EHairDebugMode::CoverageType:		DebugPermutation = 1; break;
-		case EHairDebugMode::TAAResolveType:	DebugPermutation = 2; break;
-		case EHairDebugMode::Coverage:			DebugPermutation = 3; break;
-	};
-	PermutationVector.Set<FHairDebugPS::FDebugMode>(DebugPermutation);
-	TShaderMapRef<FHairDebugPS> PixelShader(View->ShaderMap, PermutationVector);
+	TShaderMapRef<FHairDebugPS> PixelShader(View->ShaderMap);
 
 	ClearUnusedGraphResources(*PixelShader, Parameters);
 
@@ -1288,14 +1328,24 @@ void RenderHairStrandsDebugInfo(FRHICommandListImmediate& RHICmdList, TArray<FVi
 		}
 	}
 
-	if (HairDebugMode == EHairDebugMode::TAAResolveType || HairDebugMode == EHairDebugMode::SamplePerPixel || HairDebugMode == EHairDebugMode::CoverageType || HairDebugMode == EHairDebugMode::Coverage)
+	const bool bRunDebugPass =
+		HairDebugMode == EHairDebugMode::TAAResolveType ||
+		HairDebugMode == EHairDebugMode::SamplePerPixel ||
+		HairDebugMode == EHairDebugMode::CoverageType ||
+		HairDebugMode == EHairDebugMode::Coverage ||
+		HairDebugMode == EHairDebugMode::MaterialDepth ||
+		HairDebugMode == EHairDebugMode::MaterialBaseColor ||
+		HairDebugMode == EHairDebugMode::MaterialRoughness ||
+		HairDebugMode == EHairDebugMode::MaterialSpecular ||
+		HairDebugMode == EHairDebugMode::MaterialTangent;
+	if (bRunDebugPass)
 	{
 		FRDGBuilder GraphBuilder(RHICmdList);
 		FRDGTextureRef SceneColorTexture = GraphBuilder.RegisterExternalTexture(SceneTargets.GetSceneColor(), TEXT("SceneColorTexture"));
 		if (ViewIndex < uint32(HairDatas->HairVisibilityViews.HairDatas.Num()))
 		{
 			const FHairStrandsVisibilityData& VisibilityData = HairDatas->HairVisibilityViews.HairDatas[ViewIndex];
-			AddDebugHairPass(GraphBuilder, &View, HairDebugMode, VisibilityData.CategorizationTexture, SceneTargets.SceneStencilSRV, SceneColorTexture);
+			AddDebugHairPass(GraphBuilder, &View, HairDebugMode, VisibilityData.CategorizationTexture, VisibilityData.NodeIndex, VisibilityData.NodeData, SceneTargets.SceneStencilSRV, SceneColorTexture);
 			AddDebugHairPrintPass(GraphBuilder, &View, HairDebugMode, VisibilityData,  SceneTargets.SceneStencilSRV);
 		}
 
