Add visibility pass.

diff --git a/Engine/Shaders/Private/HairStrands/HairStrandsMaterialPS.usf b/Engine/Shaders/Private/HairStrands/HairStrandsMaterialPS.usf
new file mode 100644
--- /dev/null
+++ b/Engine/Shaders/Private/HairStrands/HairStrandsMaterialPS.usf
@@ -0,0 +1,379 @@
+// Copyright 1998-2019 Epic Games, Inc. All Rights Reserved.
+
+#include "HairStrandsVisibilityCommon.ush"
+#include "../Common.ush"
+#include "../SceneTexturesCommon.ush"
+#include "HairStrandsDeepShadowCommon.ush"
+#include "../DeferredShadingCommon.ush"
+#include "../VelocityCommon.ush"
+
+#include "/Engine/Generated/Material.ush"
+#include "/Engine/Generated/VertexFactory.ush" 
+
+#define SUPPORT_MATERIAL_PROPERTY 1
+uint MaterialPass_ClusterId;
+uint MaterialPass_MaterialId;  
+uint MaterialPass_PrimitiveId;
+RWStructuredBuffer<FNodeData>	OutNodeData;
+RWBuffer<float2>				OutNodeVelocity; 
+
+// For reference:
+#if 0
+struct FVertexFactoryInterpolantsVSToPS
+{
+	TANGENTTOWORLD_INTERPOLATOR_BLOCK
+
+
+#if USE_INSTANCING
+	// x = per-instance random, y = per-instance fade out amount, z = hide/show flag, w dither fade cutoff
+	float4  PerInstanceParams : COLOR1;
+#endif
+
+#if INSTANCED_STEREO
+	nointerpolation uint EyeIndex : PACKED_EYE_INDEX;
+#endif
+
+#if VF_USE_PRIMITIVE_SCENE_DATA
+	nointerpolation uint PrimitiveId : PRIMITIVE_ID;
+#endif
+#if VF_STRAND_HAIR
+	nointerpolation uint HairPrimitiveId : HAIR_PRIMITIVE_ID; // Control point ID
+	float2 HairPrimitiveUV : HAIR_PRIMITIVE_UV; // U: parameteric distance between the two surrounding control points. V: parametric distance along the width.
+#endif
+};
+
+struct FVertexFactoryInput
+{
+#if VF_USE_PRIMITIVE_SCENE_DATA
+	uint PrimitiveId : ATTRIBUTE13;
+#endif	
+	uint VertexId : SV_VertexID;
+};
+
+#endif
+
+#if VF_STRAND_HAIR
+struct FEvaluateOutput
+{
+	FNodeData NodeData;
+	float2 NodeVelocity;
+//	float3 P0;
+//	float3 P1;
+//	float3 P;
+//	float3 A;
+//	float3 B;
+//	float U;
+
+};
+
+FVertexFactoryInterpolantsVSToPS GetInterpolants(uint PrimitiveId, uint HairPrimitiveId, inout float3 World_P, inout float3 World_P_Prev)
+{
+	FMaterialVertexParameters Empty = (FMaterialVertexParameters)0;
+
+	FVertexFactoryInput Input;
+	Input.PrimitiveId = PrimitiveId;
+	Input.VertexId = HairPrimitiveId * 6; // 6: the vertex factory expect triangles list index. A quad is made of two triangles with indices: 0,1,2,3,4,5.
+
+	FVertexFactoryIntermediates	Intermediates = GetVertexFactoryIntermediates(Input);
+	FVertexFactoryInterpolantsVSToPS Interpolants = VertexFactoryGetInterpolantsVSToPS(Input, Intermediates, Empty);
+
+	World_P = VertexFactoryGetWorldPositionRaw(Input, Intermediates).xyz;
+	World_P_Prev = VertexFactoryGetPreviousWorldPosition(Input, Intermediates).xyz;
+
+	return Interpolants;
+}
+
+FEvaluateOutput Evaluate(
+	in float2 SamplePixelCoord,
+	in float SampleDepth,
+	in uint SamplePrimitiveId,
+	in uint SampleCoverage8bit,
+	in uint PrimitiveId)
+{
+	// World_P / World_P0 / World_P1 are in translated world space
+	const float SceneDepth = ConvertFromDeviceZ(SampleDepth);
+	const float2 UV = SamplePixelCoord / float2(View.BufferSizeAndInvSize.xy);
+	const float2 ScreenPosition = (UV - View.ScreenPositionScaleBias.wz) / View.ScreenPositionScaleBias.xy;
+//	const float3 World_P = mul(float4(ScreenPosition * SceneDepth, SceneDepth, 1), View.ScreenToWorld).xyz;
+	const float3 World_P = mul(float4(ScreenPosition * SceneDepth, SceneDepth, 1), View.ScreenToTranslatedWorld).xyz;
+
+//	Out.P = World_P;
+
+	float3 Velocity_World_P;
+	float3 Velocity_World_P_Prev;
+	float U = 0;
+	FVertexFactoryInterpolantsVSToPS Interpolants;
+	{
+		const uint HairPrimitiveId0 = SamplePrimitiveId;
+		const uint HairPrimitiveId1 = HairPrimitiveId0 + 1;
+
+		float3 World_P0;
+		float3 World_P1;
+
+		float3 World_P0_Prev;
+		float3 World_P1_Prev;
+
+		
+		FVertexFactoryInterpolantsVSToPS Interpolants0 = GetInterpolants(PrimitiveId, HairPrimitiveId0, World_P0, World_P0_Prev);
+		FVertexFactoryInterpolantsVSToPS Interpolants1 = GetInterpolants(PrimitiveId, HairPrimitiveId1, World_P1, World_P1_Prev);
+//		Out.P0 = World_P0;
+//		Out.P1 = World_P1;
+
+		Interpolants = Interpolants0;
+
+		// Compute the projection of the current sample point onto the hair segment (which is in the middle of the hair strands)
+		const float3 A = (World_P - World_P0);
+		const float3 B = (World_P1 - World_P0);
+		const float DotA = dot(A, B);
+		const float LengthB2 = dot(B, B);
+		const float SegmentU = LengthB2 > 0 ? saturate(DotA / LengthB2) : 0;
+
+//		Out.U = SegmentU;
+//		Out.A = A;
+//		Out.B = B;
+		const float SegmentV = 0.5f; // Todo handle case when the actual geometry is close and cover multiple pixels
+
+		Velocity_World_P	  = lerp(World_P0, World_P1, SegmentU);
+		Velocity_World_P_Prev = lerp(World_P0_Prev, World_P1_Prev, SegmentU);
+
+		Interpolants.TangentToWorld0 = lerp(Interpolants0.TangentToWorld0, Interpolants1.TangentToWorld0, SegmentU);
+		Interpolants.TangentToWorld2 = lerp(Interpolants0.TangentToWorld2, Interpolants1.TangentToWorld2, SegmentU);
+
+		Interpolants.TangentToWorld0 = normalize(Interpolants.TangentToWorld0);
+		Interpolants.TangentToWorld2 = normalize(Interpolants.TangentToWorld2);
+
+		Interpolants.PrimitiveId	 = PrimitiveId;
+		Interpolants.HairPrimitiveId = HairPrimitiveId0;
+		Interpolants.HairPrimitiveUV = float2(SegmentU, SegmentV);
+	}
+	
+
+	#if 1
+	float2 EncodedVelocity = 0;
+	{
+		const float4 ScreenPos		= mul(float4(Velocity_World_P.xyz, 1), ResolvedView.TranslatedWorldToClip);
+		const float4 PrevScreenPos	= mul(float4(Velocity_World_P_Prev.xyz, 1), ResolvedView.PrevTranslatedWorldToClip);
+
+		const float2 Velocity	= Calculate2DVelocity(ScreenPos, PrevScreenPos);
+		EncodedVelocity			= EncodeVelocityToTexture(Velocity);
+	}
+	#endif
+
+	// Sample Position < consider to be the center??
+	// VS
+	const float4 SvPosition = float4(SamplePixelCoord, SampleDepth, 1);
+
+	// PS
+	FVisibilityData OutVisibilityData = (FVisibilityData)0;	
+	OutVisibilityData.Coverage8bit = SampleCoverage8bit;
+	OutVisibilityData.PrimitiveId = SampleCoverage8bit;
+	OutVisibilityData.ClusterId = MaterialPass_ClusterId;
+
+	half3 OutBaseColor = 0;
+	float OutRoughness = 0;
+	float OutSpecular = 0;
+	{
+		// Coverage computation
+		// We don't use the coverage information
+		#if 0
+		float Coverage = 1;
+		{
+			float WorldStrandRadius = VFIntermediates.HairDimensions.y; // Coming from the vertex factory evaluation
+
+			FHairRenderInfo HairRenderInfo = GetHairRenderInfo(ResolvedView.HairRenderInfo);
+			const float SceneDepth = ConvertFromDeviceZ(SvPosition.z); // Linear depth in world unit
+			const float PixelRadius = HairRenderInfo.bIsOrthoView ? HairRenderInfo.RadiusAtDepth1Primary : SceneDepth * HairRenderInfo.RadiusAtDepth1Primary; // Not correct but the coverage is not used (we count instead the number of sub-sample covered)
+			const float StrandRealRadius = WorldStrandRadius;
+			Coverage = StrandRealRadius / max(StrandRealRadius, PixelRadius);
+		}
+		float Transmittance = saturate(1.0f - Coverage);
+		#endif
+
+		//OutVisibilityData.Tangent = normalize(Interpolants.TangentToWorld2.xyz);
+		//OutVisibilityData.PrimitiveId = Interpolants.HairPrimitiveId;
+		//OutVisibilityData.ClusterId = HairVisibilityPass_HairClusterIndex;
+		//OutVisibilityData.Coverage8bit = To8bitCoverage(Coverage);
+
+		// Material computation
+		#if 1
+		{
+			const bool bIsFrontFace = true;
+			FMaterialPixelParameters MaterialParameters = GetMaterialPixelParameters(Interpolants, SvPosition);
+			FPixelMaterialInputs PixelMaterialInputs;
+
+			float4 ScreenPosition = SvPositionToResolvedScreenPosition(SvPosition);
+			float3 TranslatedWorldPosition = SvPositionToResolvedTranslatedWorld(SvPosition);
+			CalcMaterialParametersEx(MaterialParameters, PixelMaterialInputs, SvPosition, ScreenPosition, bIsFrontFace, TranslatedWorldPosition, TranslatedWorldPosition);
+
+			OutBaseColor = GetMaterialBaseColor(PixelMaterialInputs);
+			OutRoughness = GetMaterialRoughness(PixelMaterialInputs);
+			OutSpecular = GetMaterialSpecular(PixelMaterialInputs);
+			//float MaterialAO = GetMaterialAmbientOcclusion(PixelMaterialInputs);
+
+			//OutVisibilityData.Tangent = mul(normalize(GetMaterialNormalRaw(PixelMaterialInputs)), MaterialParameters.TangentToWorld);
+			OutVisibilityData.Tangent = mul(normalize(GetMaterialNormalRaw(PixelMaterialInputs)), MaterialParameters.TangentToWorld);
+//			OutVisibilityData.Tangent = normalize(Interpolants.TangentToWorld2.xyz);
+
+			OutVisibilityData.Tangent = normalize(OutVisibilityData.Tangent); 
+		}
+		#endif
+	}
+
+	FEvaluateOutput Out;
+	Out.NodeData		= ConvertToNodeData(SampleDepth, OutVisibilityData, OutBaseColor, OutRoughness, OutSpecular);
+	Out.NodeVelocity	= EncodedVelocity;
+	return Out;
+}
+#endif
+
+void Main(
+	in float4 SvPosition: SV_Position
+	, in uint2 Resolution : DISPATCH_RESOLUTION
+	, out float4 OutColor0 : SV_Target0
+//	, out float4 OutColor1 : SV_Target1
+//	, out float4 OutColor2 : SV_Target2
+//	, out float4 OutColor3 : SV_Target3
+//	, out float4 OutColor4 : SV_Target4
+)
+{
+	ResolvedView = ResolveView();
+
+	const uint3 PixelCoord = uint3(floor(SvPosition.xy), 0);
+//	const FNodeDesc NodeDesc = DecodeNodeDesc(MaterialPassParameters.NodeIndex.Load(PixelCoord));
+
+//	OutColor0 = 0;
+	const uint SampleIndex = PixelCoord.x + PixelCoord.y * Resolution.x;// MaterialPassParameters.MaxResolution.x;
+	bool bIsValid = false;
+	if (SampleIndex < MaterialPassParameters.MaxSampleCount)
+	{
+//		FNodeData InNode = MaterialPassParameters.NodeVis[SampleIndex];
+		const FNodeVis InNode = MaterialPassParameters.NodeVis.Load(SampleIndex);
+		
+		const uint PackedSamplePixelCoord = MaterialPassParameters.NodeCoord[SampleIndex];
+		const float2 SamplePixelCoord = float2(PackedSamplePixelCoord & 0xFFFF, (PackedSamplePixelCoord>>16) & 0xFFFF) + float2(0.5f, 0.5f);
+
+#if 0
+		float			SampleDepth;
+		FVisibilityData	VisibilityData;
+		float3			BaseColor;
+		float			Roughness;
+		float			Specular;
+
+		DecodeNodeData(
+			InNode,
+			SampleDepth,
+			VisibilityData,
+			BaseColor,
+			Roughness,
+			Specular);
+#else
+		float SampleDepth = InNode.Depth;
+		uint SampleCoverage8bit = InNode.Coverage8bit;
+		uint SamplePrimitiveId = InNode.PrimitiveId;
+		uint SampleMaterialId = InNode.MaterialId;
+
+		if (SampleMaterialId == MaterialPass_MaterialId)
+		{
+#endif
+			FEvaluateOutput Out = Evaluate(SamplePixelCoord, SampleDepth, SamplePrimitiveId, SampleCoverage8bit, MaterialPass_PrimitiveId);
+			OutNodeData[SampleIndex] = Out.NodeData;
+			OutNodeVelocity[SampleIndex] = Out.NodeVelocity;
+
+			bIsValid = true;
+//		OutColor0 = float4(Out.P0, 1);
+//		OutColor1 = float4(Out.P1, 1);
+//		OutColor2 = float4(Out.P, Out.U);
+//		OutColor3 = float4(Out.A, 1);
+//		OutColor4 = float4(Out.B, 1);
+		}
+	}
+
+	OutColor0 = bIsValid? float4(0, 1, 0, 1) : float4(1, 0, 0, 1);
+
+#if 0
+	float Coverage = 1;
+	{
+		FHairRenderInfo HairRenderInfo = GetHairRenderInfo(ResolvedView.HairRenderInfo);
+		const float SceneDepth = ConvertFromDeviceZ(SvPosition.z); // Linear depth in world unit
+		const float PixelRadius = HairRenderInfo.bIsOrthoView ? HairRenderInfo.RadiusAtDepth1Primary : SceneDepth * HairRenderInfo.RadiusAtDepth1Primary; // Not correct but the coverage is not used (we count instead the number of sub-sample covered)
+		const float StrandRealRadius = WorldStrandRadius;
+		Coverage = StrandRealRadius / max(StrandRealRadius, PixelRadius);
+	}
+
+#if HAIR_RENDER_MODE == RENDER_MODE_TRANSMITTANCE
+	OutColor0 = saturate(1.0f - Coverage);
+#endif
+
+#if HAIR_RENDER_MODE == RENDER_MODE_MSAA || HAIR_RENDER_MODE == RENDER_MODE_PPLL
+	FVisibilityData VisibilityData;
+  #if VF_STRAND_HAIR
+	VisibilityData.Tangent = normalize(Interpolants.TangentToWorld2.xyz);
+	VisibilityData.PrimitiveId = Interpolants.HairPrimitiveId;
+  #else
+	VisibilityData.Tangent = float3(0, 0, 1);
+	VisibilityData.PrimitiveId = 0;
+  #endif
+	VisibilityData.ClusterId = HairVisibilityPass_HairClusterIndex;
+	VisibilityData.Coverage8bit = To8bitCoverage(Coverage);
+
+	half3 BaseColor = 0;
+	float Roughness = 0;
+	float Specular  = 0;
+	float2 Velocity = 0;
+	float2 EncodedVelocity = 0;
+#if SUPPORT_MATERIAL_PROPERTY
+	{
+		FMaterialPixelParameters MaterialParameters = GetMaterialPixelParameters(Interpolants, SvPosition);
+		FPixelMaterialInputs PixelMaterialInputs;
+	
+		float4 ScreenPosition = SvPositionToResolvedScreenPosition(SvPosition);
+		float3 TranslatedWorldPosition = SvPositionToResolvedTranslatedWorld(SvPosition);
+		CalcMaterialParametersEx(MaterialParameters, PixelMaterialInputs, SvPosition, ScreenPosition, bIsFrontFace, TranslatedWorldPosition, TranslatedWorldPosition);
+
+		BaseColor		 = GetMaterialBaseColor(PixelMaterialInputs);
+		Roughness		 = GetMaterialRoughness(PixelMaterialInputs);
+		half  Metallic	 = GetMaterialMetallic(PixelMaterialInputs);
+		Specular		 = GetMaterialSpecular(PixelMaterialInputs);
+		float MaterialAO = GetMaterialAmbientOcclusion(PixelMaterialInputs);
+
+		VisibilityData.Tangent = mul(normalize(GetMaterialNormalRaw(PixelMaterialInputs)), MaterialParameters.TangentToWorld);
+	}
+	{
+		Velocity = Calculate2DVelocity(PackedVelocityA, PackedVelocityC);
+		EncodedVelocity = EncodeVelocityToTexture(Velocity);
+	}
+#endif
+#endif
+
+#if HAIR_RENDER_MODE == RENDER_MODE_MSAA
+	OutColor0 = EncodeVisibilityData(VisibilityData);
+	OutColor1 = float4(sqrt(BaseColor), Roughness);
+	OutColor2 = float4(Specular,0,0,0);
+	OutColor3 = EncodedVelocity;
+#endif
+
+#if HAIR_RENDER_MODE == RENDER_MODE_PPLL
+	const uint2 PixelCoord = SvPosition.xy;
+
+	float2 NDC = SvPosition.xy / SvPosition.w;
+	float2 ScreenUV = NDC * ResolvedView.ScreenPositionScaleBias.xy + ResolvedView.ScreenPositionScaleBias.wz;
+
+	// Allocate space for a new node
+	uint NodeIndex;
+	InterlockedAdd(PPLLCounter[uint2(0, 0)], 1, NodeIndex);
+
+	if (NodeIndex < VisibilityPassGlobalParameters.MaxPPLLNodeCount)
+	{
+		// If we can append new node, add it and make it point to the node previously in head of list
+		uint NextNodeIndex;
+		InterlockedExchange(PPLLNodeIndex[PixelCoord], NodeIndex, NextNodeIndex);
+		const FPPLLNodeData Node = ConvertToPPLLNodeData(SvPosition.z, VisibilityData, BaseColor, Roughness, Specular, Velocity, NextNodeIndex);
+		PPLLNodes[NodeIndex] = Node;
+	}
+#endif
+
+#endif
+
+}
+
+
diff --git a/Engine/Shaders/Private/HairStrands/HairStrandsMaterialVS.usf b/Engine/Shaders/Private/HairStrands/HairStrandsMaterialVS.usf
new file mode 100644
--- /dev/null
+++ b/Engine/Shaders/Private/HairStrands/HairStrandsMaterialVS.usf
@@ -0,0 +1,30 @@
+// Copyright 1998-2019 Epic Games, Inc. All Rights Reserved.
+
+#include "HairStrandsVisibilityCommon.ush"
+#include "../Common.ush"
+#include "/Engine/Generated/Material.ush"
+#include "/Engine/Generated/VertexFactory.ush"
+
+#if VERTEXSHADER
+ 
+void Main(
+	uint VertexId : SV_VertexID, 
+	out float4 Position : SV_POSITION,
+	nointerpolation out uint2 Resolution : DISPATCH_RESOLUTION)
+{
+	const uint DispatchCountX = MaterialPassParameters.IndirectArgs.Load(0);
+	const uint DispatchCountY = MaterialPassParameters.IndirectArgs.Load(1);
+	const uint TotalNodeCount = MaterialPassParameters.NodeGroupSize * DispatchCountX * DispatchCountY;
+
+	Resolution.x = ceil(sqrt(TotalNodeCount));
+	Resolution.y = Resolution.x;
+
+	const float2 ClipCoord = ((Resolution+0.5f) / float2(MaterialPassParameters.MaxResolution)) * 2;
+
+	const float2 UV = float2(VertexId & 1, VertexId >> 1);
+	const float2 Pos = float2(-UV.x, UV.y) * 4 + float2(-1,+1) + float2(ClipCoord.x, -ClipCoord.y);
+	Position = float4(Pos, 0.5f, 1);
+
+}
+
+#endif // VERTEXSHADER
diff --git a/Engine/Shaders/Private/HairStrands/HairStrandsVelocity.usf b/Engine/Shaders/Private/HairStrands/HairStrandsVelocity.usf
new file mode 100644
--- /dev/null
+++ b/Engine/Shaders/Private/HairStrands/HairStrandsVelocity.usf
@@ -0,0 +1,101 @@
+// Copyright 1998-2019 Epic Games, Inc. All Rights Reserved.
+
+#pragma once
+
+#include "../Common.ush"
+#include "HairStrandsVisibilityCommon.ush"
+
+#define RENDER_MODE_MSAA 0
+#define VENDOR_AMD 0
+#define VENDOR_NVIDIA 1
+#define VENDOR_OTHER 2
+
+#if PERMUTATION_VENDOR == VENDOR_AMD
+#define TILE_PIXEL_SIZE_X 8
+#define TILE_PIXEL_SIZE_Y 8
+#elif PERMUTATION_VENDOR == VENDOR_NVIDIA
+#define TILE_PIXEL_SIZE_X 8
+#define TILE_PIXEL_SIZE_Y 4
+#else //VENDOR_OTHER
+#define TILE_PIXEL_SIZE_X 8
+#define TILE_PIXEL_SIZE_Y 8
+#endif
+
+uint2 ResolutionOffset;
+
+Texture2D<uint> NodeIndex;
+Buffer<float2> NodeVelocity;
+StructuredBuffer<FNodeVis> NodeVis; 
+
+#define VELOCITY_TYPE_NONE 0
+#define VELOCITY_TYPE_AVG 1  
+#define VELOCITY_TYPE_CLOSEST 2
+#define VELOCITY_TYPE_MAX 3
+#if PERMUTATION_VELOCITY != VELOCITY_TYPE_NONE
+	RWTexture2D<float2> OutVelocityTexture;
+#endif
+
+[numthreads(TILE_PIXEL_SIZE_X, TILE_PIXEL_SIZE_Y, 1)]
+void MainCS(uint GroupIndex : SV_GroupIndex, uint3 DispatchThreadId : SV_DispatchThreadID)
+{
+	const uint2 PixelCoord = DispatchThreadId.xy + ResolutionOffset;
+	const FNodeDesc NodeDesc = DecodeNodeDesc(NodeIndex.Load(uint3(PixelCoord, 0)));
+	if (NodeDesc.Count > 0)
+	{
+		// Store final sort node data 
+		#if PERMUTATION_VELOCITY == VELOCITY_TYPE_AVG
+		float2 AverageVelocity = 0;
+		#endif
+
+		#if PERMUTATION_VELOCITY == VELOCITY_TYPE_CLOSEST
+		float2 ClosestEncodedVelocity = 0;
+		float ClosestDepth = 0; // Inverse-Z
+		#endif		
+
+		#if PERMUTATION_VELOCITY == VELOCITY_TYPE_MAX
+		float2 MaxEncodedVelocity = 0;
+		float MaxVelocityMagnitude2 = 0;
+		#endif		
+
+		for (uint NodeIt = 0; NodeIt < NodeDesc.Count; ++NodeIt)
+		{
+			const uint SampleIndex = NodeDesc.Offset + NodeIt;
+			const float2 EncodedVelocity = NodeVelocity[SampleIndex];
+
+			#if PERMUTATION_VELOCITY == VELOCITY_TYPE_AVG
+			AverageVelocity += DecodeVelocityFromTexture(EncodedVelocity);
+			#endif			
+
+			#if PERMUTATION_VELOCITY == VELOCITY_TYPE_CLOSEST
+			const float NodeDepth = NodeVis[SampleIndex].Depth;
+			if (NodeDepth > ClosestDepth) // Inverse-Z
+			{
+				ClosestEncodedVelocity = EncodedVelocity;
+				ClosestDepth = NodeDepth;
+			}
+			#endif
+
+			#if PERMUTATION_VELOCITY == VELOCITY_TYPE_MAX
+			const float2 CurrentVelocity = DecodeVelocityFromTexture(EncodedVelocity);
+			const float CurrentVelocityMagnitude2 = dot(CurrentVelocity, CurrentVelocity);
+			if (CurrentVelocityMagnitude2 > MaxVelocityMagnitude2)
+			{
+				MaxEncodedVelocity = EncodedVelocity;
+				MaxVelocityMagnitude2 = CurrentVelocityMagnitude2;
+			}
+			#endif			
+		}
+
+		#if PERMUTATION_VELOCITY == VELOCITY_TYPE_AVG
+		OutVelocityTexture[PixelCoord] = EncodeVelocityToTexture(AverageVelocity / max(NodeDesc.Count, 1));
+		#endif		
+
+		#if PERMUTATION_VELOCITY == VELOCITY_TYPE_CLOSEST
+		OutVelocityTexture[PixelCoord] = ClosestEncodedVelocity;
+		#endif		
+
+		#if PERMUTATION_VELOCITY == VELOCITY_TYPE_MAX
+		OutVelocityTexture[PixelCoord] = MaxEncodedVelocity;
+		#endif		
+	}
+}
diff --git a/Engine/Shaders/Private/HairStrands/HairStrandsVisibilityCommon.ush b/Engine/Shaders/Private/HairStrands/HairStrandsVisibilityCommon.ush
--- a/Engine/Shaders/Private/HairStrands/HairStrandsVisibilityCommon.ush
+++ b/Engine/Shaders/Private/HairStrands/HairStrandsVisibilityCommon.ush
@@ -5,6 +5,7 @@
 #define RENDER_MODE_MSAA 0
 #define RENDER_MODE_TRANSMITTANCE 1
 #define RENDER_MODE_PPLL 2
+#define RENDER_MODE_MSAA_VISIBILITY 3
 
 struct FVisibilityData
 {
@@ -14,10 +15,19 @@ struct FVisibilityData
 	uint   ClusterId;	//  4 bits
 };
 
+struct FNodeVis // 128 bit aligments. TODO: Move the coverage computation into the material evaluation for saving 64 bits? 
+{
+	float Depth;		// 32 bits
+	uint  PrimitiveId;	// 32 bits
+	uint  MaterialId;	//  8 bits : this could be packed into a the 3 high bits of PrimitiveId: up to 8 material on screen
+	uint  Coverage8bit;	//  8 bits
+};
+
 // Footprint: 128bits | 16bytes
 // Option to reduce the footprint: 
 // * we don't use PrimitiveID at the moment which is encoded on 3bytes
 // * Deferred the evaluation of Coverage & material properties into some other pass?
+
 struct FNodeData
 {
 	float Depth;
@@ -350,3 +360,50 @@ FNodeData ConvertToNodeData(float Depth, FVisibilityData VisibilityData, float3
 	 return Out;
  }
 
+ uint Uint16ToUint32(uint2 In)
+ {
+	 return (In.x & 0xFFFF) | ((In.y & 0xFFFF) << 16);
+ }
+
+ uint2 Uint32ToUint16(uint In)
+ {
+	 uint2 A;
+	 A.x = In & 0xFFFF;
+	 A.y = (In >> 16) & 0xFFFF;
+	 return A;
+ }
+
+ uint Float16ToUint32(float2 In)
+ {
+	 return Uint16ToUint32(f32tof16(In));
+ }
+
+ float2 Uint32ToFloat16(uint In)
+ {
+	return f16tof32(Uint32ToUint16(In));
+ }
+
+ struct FVisSample
+ {
+	 uint PrimitiveId;
+	 uint MaterialId;
+ };
+
+ uint GetInvalidHairPrimitiveId()
+ {
+	return 0xFFFFFFFF;
+ }
+
+ uint EncodeVisSample(uint PrimitiveId, uint MaterialId)
+ {
+	 return (PrimitiveId & 0x0FFFFFFF) | ((MaterialId & 0xF) << 28);
+ }
+
+ FVisSample DecodeVisSample(uint RawVis)
+ {
+	 FVisSample Out;
+	 Out.PrimitiveId = RawVis & 0x0FFFFFFF;
+	 Out.MaterialId = (RawVis >> 28) & 0xF;
+	 return Out;
+ }
+ 
diff --git a/Engine/Shaders/Private/HairStrands/HairStrandsVisibilityCompaction.usf b/Engine/Shaders/Private/HairStrands/HairStrandsVisibilityCompaction.usf
--- a/Engine/Shaders/Private/HairStrands/HairStrandsVisibilityCompaction.usf
+++ b/Engine/Shaders/Private/HairStrands/HairStrandsVisibilityCompaction.usf
@@ -33,10 +33,18 @@ uint2 ResolutionOffset;
 
 RWTexture2D<uint> OutCompactNodeCounter;
 RWTexture2D<uint> OutCompactNodeIndex;
-RWStructuredBuffer<FNodeData> OutCompactNodeData;
+#if PERMUTATION_VISIBILITY
+	RWStructuredBuffer<FNodeVis> OutCompactNodeData;
+#else
+	RWStructuredBuffer<FNodeData> OutCompactNodeData;
+#endif
 RWStructuredBuffer<uint> OutCompactNodeCoord;
 RWTexture2D<uint4> OutCategorizationTexture;
 
+#ifndef PERMUTATION_VELOCITY
+	#error Velocity permutation needs to be defined
+#endif
+
 #define VELOCITY_TYPE_NONE 0
 #define VELOCITY_TYPE_AVG 1  
 #define VELOCITY_TYPE_CLOSEST 2
@@ -49,7 +57,6 @@ RWTexture2D<uint4> OutCategorizationTexture;
 groupshared uint AllocationNodeCount;
 groupshared uint AllocationNodeOffset;
 
-
 #if PERMUTATION_PPLL
 #define PPLL_MAX_RENDER_NODE_COUNT PERMUTATION_PPLL
 
@@ -276,7 +283,161 @@ void MainCS(uint GroupIndex : SV_GroupIndex, uint3 DispatchThreadId : SV_Dispatc
 #else // PERMUTATION_PPLL //////////////////////////////////////////////////////////////////////////////////////////////////////////////
 
 
+#if PERMUTATION_VISIBILITY
+
+void ComputeUniqueSamplesWithoutDepthTest(
+	inout uint3 OutSamples[HAIR_VISIBILITY_MAXSAMPLECOUNT],
+	inout FSampleSetDesc OutSet,
+	const uint2 PixelCoord,
+	const uint SampleCount,
+	const float OpaqueDepth,
+	Texture2DMS<uint, HAIR_VISIBILITY_MAXSAMPLECOUNT>  InMSAA_IDTexture,
+	Texture2DMS<float, HAIR_VISIBILITY_MAXSAMPLECOUNT>  InMSAA_DepthTexture)
+{
+	OutSet.UniqueSampleCount = 0;
+	OutSet.ValidSampleCount = 0;
+
+	OutSet.HairSampleCount = clamp(SampleCount, 0, HAIR_VISIBILITY_MAXSAMPLECOUNT);
+	const uint InvalidSample = 0xFFFF;
+
+	float ClosestDepth = 0; // Inverse-Z
+	for (uint SampleIt = 0; SampleIt < OutSet.HairSampleCount; ++SampleIt)
+	{
+		// Note: InMSAA_IDTexture contains both the primitive ID and tha material ID. However 
+		// the material ID is constant along the primitive, so it is correct to use this as a 
+		// sorting/deduplication key
+		const uint PrimitiveId = InMSAA_IDTexture.Load(PixelCoord, SampleIt);
+		const bool bIsValid = PrimitiveId != GetInvalidHairPrimitiveId();
+		if (!bIsValid)
+			continue;
+
+		const float SampleDepth = InMSAA_DepthTexture.Load(PixelCoord, SampleIt);
+		if (OpaqueDepth > SampleDepth) // Reverse-Z
+			continue;
+
+		++OutSet.ValidSampleCount;
+
+		bool bAlreadyExist = false;
+		for (uint UniqueIt = 0; UniqueIt < OutSet.UniqueSampleCount; ++UniqueIt)
+		{
+			const bool bIsSimilar = PrimitiveId == OutSamples[UniqueIt].x;
+
+			if (bIsSimilar)
+			{
+				OutSamples[UniqueIt].y += 1;
+				if (ClosestDepth < SampleDepth)
+				{
+					OutSamples[OutSet.UniqueSampleCount].z = SampleIt;
+					ClosestDepth = SampleDepth;
+				}
+				bAlreadyExist = true;
+				break;
+			}
+		}
+		// todo: always takes the closest depth?
+		if (!bAlreadyExist)
+		{
+			OutSamples[OutSet.UniqueSampleCount].x = PrimitiveId;
+			OutSamples[OutSet.UniqueSampleCount].y = 1;
+			OutSamples[OutSet.UniqueSampleCount].z = SampleIt;
+			++OutSet.UniqueSampleCount;
+		}
+	}
+}
+
+#ifndef HAIR_VISIBILITY_MAXSAMPLECOUNT
+#error HAIR_VISIBILITY_MAXSAMPLECOUNT is not defined
+#endif
+
+Texture2DMS<uint, HAIR_VISIBILITY_MAXSAMPLECOUNT> MSAA_IDTexture;
+Texture2DMS<float, HAIR_VISIBILITY_MAXSAMPLECOUNT> MSAA_DepthTexture;
+
+#if PERMUTATION_VIEWTRANSMITTANCE == 1
+Texture2D<float> ViewTransmittanceTexture;
+#endif
+
+[numthreads(TILE_PIXEL_SIZE_X, TILE_PIXEL_SIZE_Y, 1)]
+void MainCS(uint GroupIndex : SV_GroupIndex, uint3 DispatchThreadId : SV_DispatchThreadID)
+{
+	if (GroupIndex == 0)
+	{
+		AllocationNodeCount = 0;
+		AllocationNodeOffset = 0;
+	}
+
+	uint2 PixelCoord = DispatchThreadId.xy + ResolutionOffset;
+	if (HairVisibilityMSAASampleCount == 0 || PixelCoord.x >= OutputResolution.x || PixelCoord.y >= OutputResolution.y)
+		PixelCoord = uint2(0, 0);
+
+	FSampleSetDesc SampleDesc;
+
+	const float OpaqueDepth = SceneTexturesStruct.SceneDepthTexture.Load(uint3(PixelCoord, 0)).r;
+	uint3 Samples[HAIR_VISIBILITY_MAXSAMPLECOUNT];		// x:PrimitiveId|MaterialId, y:Weight, z:SampleIt
+	ComputeUniqueSamplesWithoutDepthTest(
+		Samples,
+		SampleDesc,
+		PixelCoord,
+		HairVisibilityMSAASampleCount,
+		OpaqueDepth,
+		MSAA_IDTexture, 
+		MSAA_DepthTexture);
+
+	FNodeDesc NodeDesc;
+	NodeDesc.Count = SampleDesc.UniqueSampleCount;
+	NodeDesc.Offset = 0;
+
+	InterlockedAdd(AllocationNodeCount, NodeDesc.Count, NodeDesc.Offset);
+	GroupMemoryBarrierWithGroupSync();
+	if (GroupIndex == 0 && AllocationNodeCount > 0)
+	{
+		InterlockedAdd(OutCompactNodeCounter[uint2(0, 0)], AllocationNodeCount, AllocationNodeOffset);
+	}
+	GroupMemoryBarrierWithGroupSync();
+
+	// Allocate node space
+	FCategorizationData Categorization = (FCategorizationData)0;
+	if (NodeDesc.Count > 0)
+	{
+		NodeDesc.Offset += AllocationNodeOffset;
+
+		// Store final sort node data 
+		if (NodeDesc.Offset + NodeDesc.Count < MaxNodeCount)
+		{
+			for (uint OutIndex = 0; OutIndex < NodeDesc.Count; ++OutIndex)
+			{
+				const FVisSample VisSample = DecodeVisSample(Samples[OutIndex].x);
 
+				// VisibilityData.Coverage8bit is a weight normalising to 1 the contribution of all the compacted samples. Because later it is weighted by Categorization.PixelCoverage.
+				FNodeVis OutNodeVis;
+				OutNodeVis.PrimitiveId	= VisSample.PrimitiveId;
+				OutNodeVis.Depth		= MSAA_DepthTexture.Load(PixelCoord, Samples[OutIndex].z);
+				OutNodeVis.Coverage8bit = To8bitCoverage(Samples[OutIndex].y / float(SampleDesc.ValidSampleCount));
+				OutNodeVis.MaterialId	= VisSample.MaterialId;
+
+				Categorization.TotalCoverage8bit += OutNodeVis.Coverage8bit;
+				Categorization.ClosestDepth = max(Categorization.ClosestDepth, OutNodeVis.Depth); // Inverse-Z
+
+				const uint StoreIndex = NodeDesc.Offset + OutIndex;
+				OutCompactNodeData[StoreIndex] = OutNodeVis;
+				OutCompactNodeCoord[StoreIndex] = (PixelCoord.x & 0xFFFF) | ((PixelCoord.y & 0xFFFF) << 16);
+			}
+
+			NodeDesc.Count = NodeDesc.Count;
+			Categorization.SampleCount = NodeDesc.Count;
+		}
+	}
+
+#if PERMUTATION_VIEWTRANSMITTANCE == 1
+	Categorization.PixelCoverage = 1.0f - ViewTransmittanceTexture.Load(uint3(PixelCoord, 0));
+#else
+	Categorization.PixelCoverage = SampleDesc.ValidSampleCount / float(HAIR_VISIBILITY_MAXSAMPLECOUNT);
+#endif
+
+	OutCompactNodeIndex[PixelCoord] = EncodeNodeDesc(NodeDesc);
+	OutCategorizationTexture[PixelCoord] = EncodeCategorizationData(Categorization);
+}
+
+#else // PERMUTATION_VISIBILITY
 
 void ComputeUniqueSamplesWithoutDepthTest(
 	inout uint3 OutSamples[HAIR_VISIBILITY_MAXSAMPLECOUNT],
@@ -333,6 +494,7 @@ void ComputeUniqueSamplesWithoutDepthTest(
 				if (ClosestDepth < SampleDepth)
 				{
 					OutSamples[OutSet.UniqueSampleCount].z = SampleIt;
+					ClosestDepth = SampleDepth;
 				}
 				bAlreadyExist = true;
 				break;
@@ -521,8 +683,7 @@ void MainCS(uint GroupIndex : SV_GroupIndex, uint3 DispatchThreadId : SV_Dispatc
 	OutCategorizationTexture[PixelCoord] = EncodeCategorizationData(Categorization);
 }
 
-
-
+#endif  // PERMUTATION_VISIBILITY
 #endif // PERMUTATION_PPLL
 
 
diff --git a/Engine/Shaders/Private/HairStrands/HairStrandsVisibilityPS.usf b/Engine/Shaders/Private/HairStrands/HairStrandsVisibilityPS.usf
--- a/Engine/Shaders/Private/HairStrands/HairStrandsVisibilityPS.usf
+++ b/Engine/Shaders/Private/HairStrands/HairStrandsVisibilityPS.usf
@@ -27,6 +27,10 @@
 #define SUPPORT_MATERIAL_PROPERTY 0
 #endif
 
+#if HAIR_RENDER_MODE == RENDER_MODE_MSAA_VISIBILITY
+#define SUPPORT_MATERIAL_PROPERTY 0
+#endif
+
 #if HAIR_RENDER_MODE == RENDER_MODE_PPLL 
 uint HairVisibilityPass_MaxPPLLNodeCount;
 RWTexture2D<uint> PPLLCounter;
@@ -35,12 +39,17 @@ RWStructuredBuffer<FPPLLNodeData> PPLLNodes;
 #endif
 
 uint HairVisibilityPass_HairClusterIndex;
+uint HairVisibilityPass_HairMaterialId;
 
 #if HAIR_RENDER_MODE == RENDER_MODE_PPLL
 EARLYDEPTHSTENCIL	// Required for UAV operations to not happen.
 #endif
 void MainVisibility(
+#if HAIR_RENDER_MODE == RENDER_MODE_MSAA_VISIBILITY
+	uint HairPrimitiveId : HAIR_PRIMITIVE_ID
+#else
 	centroid FVertexFactoryInterpolantsVSToPS Interpolants
+#endif
 	, in float4 SvPosition	: SV_Position
 #if HAIR_RENDER_MODE == RENDER_MODE_MSAA || HAIR_RENDER_MODE == RENDER_MODE_PPLL
 	, centroid in float WorldStrandRadius : TEXCOORD8
@@ -59,6 +68,8 @@ void MainVisibility(
 	, out float4 OutColor2	: SV_Target2
 	, out float2 OutColor3	: SV_Target3
   #endif
+#elif HAIR_RENDER_MODE == RENDER_MODE_MSAA_VISIBILITY
+	, out uint OutColor0 : SV_Target0
 #endif
 )
 {
@@ -73,6 +84,7 @@ void MainVisibility(
 	}
 	#endif
 
+#if HAIR_RENDER_MODE != RENDER_MODE_MSAA_VISIBILITY
 	float Coverage = 1;
 	{
 		FHairRenderInfo HairRenderInfo = GetHairRenderInfo(ResolvedView.HairRenderInfo);
@@ -81,6 +93,7 @@ void MainVisibility(
 		const float StrandRealRadius = WorldStrandRadius;
 		Coverage = StrandRealRadius / max(StrandRealRadius, PixelRadius);
 	}
+#endif
 
 #if HAIR_RENDER_MODE == RENDER_MODE_TRANSMITTANCE
 	OutColor0 = saturate(1.0f - Coverage);
@@ -154,6 +167,13 @@ void MainVisibility(
 	}
 #endif
 
+#if HAIR_RENDER_MODE == RENDER_MODE_MSAA_VISIBILITY
+  #if VF_STRAND_HAIR
+	OutColor0 = EncodeVisSample(HairPrimitiveId, HairVisibilityPass_HairMaterialId);
+  #else
+	OutColor0 = 0;
+  #endif
+#endif
 }
 
 
diff --git a/Engine/Shaders/Private/HairStrands/HairStrandsVisibilityVS.usf b/Engine/Shaders/Private/HairStrands/HairStrandsVisibilityVS.usf
--- a/Engine/Shaders/Private/HairStrands/HairStrandsVisibilityVS.usf
+++ b/Engine/Shaders/Private/HairStrands/HairStrandsVisibilityVS.usf
@@ -8,9 +8,15 @@
 
 struct FHairVisibilityAccumulateVSToPS
 {
+#if HAIR_RENDER_MODE == RENDER_MODE_MSAA_VISIBILITY
+	nointerpolation uint HairPrimitiveId : HAIR_PRIMITIVE_ID;
+#else
 	FVertexFactoryInterpolantsVSToPS FactoryInterpolants;
+#endif
 	float4 Position : SV_POSITION;
+#if HAIR_RENDER_MODE != RENDER_MODE_MSAA_VISIBILITY
 	float WorldStrandRadius : TEXCOORD8;
+#endif
 #if HAIR_RENDER_MODE == RENDER_MODE_MSAA || HAIR_RENDER_MODE == RENDER_MODE_PPLL
 	float4 PackedVelocityA : TEXCOORD9;
 	float4 PackedVelocityC : TEXCOORD12;
@@ -78,14 +84,19 @@ void Main(
 	//}
 
 	Output.Position = mul(WorldPosition, View.TranslatedWorldToClip);
+
+#if HAIR_RENDER_MODE == RENDER_MODE_MSAA_VISIBILITY
+	Output.HairPrimitiveId = VFIntermediates.HairPrimitiveId;
+#else
 	Output.FactoryInterpolants = VertexFactoryGetInterpolantsVSToPS(Input, VFIntermediates, VertexParameters);
 
-#ifdef VF_STRAND_HAIR
+  #ifdef VF_STRAND_HAIR
 	Output.WorldStrandRadius = VFIntermediates.HairDimensions.y;
-#else
+  #else
 	Output.WorldStrandRadius = 1;
+  #endif
 #endif
-	
+
 	OutputVertexID( Output );
 }
 
diff --git a/Engine/Source/Runtime/Renderer/Private/HairStrands/HairStrandsCluster.cpp b/Engine/Source/Runtime/Renderer/Private/HairStrands/HairStrandsCluster.cpp
--- a/Engine/Source/Runtime/Renderer/Private/HairStrands/HairStrandsCluster.cpp
+++ b/Engine/Source/Runtime/Renderer/Private/HairStrands/HairStrandsCluster.cpp
@@ -17,6 +17,7 @@ FHairStrandsClusterViews CreateHairStrandsClusters(
 		const FViewInfo& View = Views[ViewIndex];
 		if (View.Family)
 		{
+			int32 MaterialId = 0;
 			FHairStrandsClusterDatas& PrimitivesClusters = PrimitivesClusterViews.Views.AddDefaulted_GetRef();
 
 			if (View.HairStrandsMeshElements.Num() == 0)
@@ -26,7 +27,7 @@ FHairStrandsClusterViews CreateHairStrandsClusters(
 
 			// Spatially clusters all hair primitives within the same area, and render a single DOM for them
 			uint32 ClusterId = 0;
-			auto UpdateCluster = [&PrimitivesClusters, &View, &ClusterId](const FMeshBatchAndRelevance* MeshBatchAndRelevance, const FPrimitiveSceneProxy* Proxy)
+			auto UpdateCluster = [&PrimitivesClusters, &View, &ClusterId, &MaterialId](const FMeshBatchAndRelevance* MeshBatchAndRelevance, const FPrimitiveSceneProxy* Proxy)
 			{
 				const FBoxSphereBounds& PrimitiveBounds = Proxy->GetBounds();
 
@@ -40,7 +41,9 @@ FHairStrandsClusterViews CreateHairStrandsClusters(
 
 						if (MeshBatchAndRelevance)
 						{
-							Cluster.PrimitivesInfos.Add(*MeshBatchAndRelevance);
+							FHairStrandsClusterData::PrimitiveInfo& PrimitiveInfo = Cluster.PrimitivesInfos.AddZeroed_GetRef();
+							PrimitiveInfo.MeshBatchAndRelevance = *MeshBatchAndRelevance;
+							PrimitiveInfo.MaterialId = MaterialId++;
 						}
 						bClusterFound = true;
 						break;
@@ -53,7 +56,9 @@ FHairStrandsClusterViews CreateHairStrandsClusters(
 					Cluster.ClusterId = ClusterId++;
 					if (MeshBatchAndRelevance)
 					{
-						Cluster.PrimitivesInfos.Add(*MeshBatchAndRelevance);
+						FHairStrandsClusterData::PrimitiveInfo& PrimitiveInfo = Cluster.PrimitivesInfos.AddZeroed_GetRef();
+						PrimitiveInfo.MeshBatchAndRelevance = *MeshBatchAndRelevance;
+						PrimitiveInfo.MaterialId = MaterialId++;
 					}
 					Cluster.Bounds = PrimitiveBounds;
 					PrimitivesClusters.Datas.Add(Cluster);
diff --git a/Engine/Source/Runtime/Renderer/Private/HairStrands/HairStrandsCluster.h b/Engine/Source/Runtime/Renderer/Private/HairStrands/HairStrandsCluster.h
--- a/Engine/Source/Runtime/Renderer/Private/HairStrands/HairStrandsCluster.h
+++ b/Engine/Source/Runtime/Renderer/Private/HairStrands/HairStrandsCluster.h
@@ -14,7 +14,12 @@
 /// Hair cluster infos (i.e. group of mesh having hair material on them)
 struct FHairStrandsClusterData
 {
-	typedef TArray<FMeshBatchAndRelevance, SceneRenderingAllocator> TPrimitiveInfos;
+	struct PrimitiveInfo
+	{
+		FMeshBatchAndRelevance MeshBatchAndRelevance;
+		uint32 MaterialId;
+	};
+	typedef TArray<PrimitiveInfo, SceneRenderingAllocator> TPrimitiveInfos;
 
 	FVector GetMinBound() const { return VoxelResources.MinAABB; }
 	FVector GetMaxBound() const { return VoxelResources.MaxAABB; }
diff --git a/Engine/Source/Runtime/Renderer/Private/HairStrands/HairStrandsRasterCommon.cpp b/Engine/Source/Runtime/Renderer/Private/HairStrands/HairStrandsRasterCommon.cpp
--- a/Engine/Source/Runtime/Renderer/Private/HairStrands/HairStrandsRasterCommon.cpp
+++ b/Engine/Source/Runtime/Renderer/Private/HairStrands/HairStrandsRasterCommon.cpp
@@ -408,11 +408,11 @@ void RasterHairStrands(
 
 	FDeepShadowMeshProcessor DeepShadowMeshProcessor(Scene, ViewInfo /* is a SceneView */, DrawRenderState, &ShadowContext, ShadowPassType);
 
-	for (const FMeshBatchAndRelevance& MeshBatchAndRelevance : PrimitiveSceneInfos)
+	for (const FHairStrandsClusterData::PrimitiveInfo& PrimitveInfo : PrimitiveSceneInfos)
 	{
-		const FMeshBatch& MeshBatch = *MeshBatchAndRelevance.Mesh;
+		const FMeshBatch& MeshBatch = *PrimitveInfo.MeshBatchAndRelevance.Mesh;
 		const uint64 BatchElementMask = ~0ull;
-		DeepShadowMeshProcessor.AddMeshBatch(MeshBatch, BatchElementMask, MeshBatchAndRelevance.PrimitiveSceneProxy);
+		DeepShadowMeshProcessor.AddMeshBatch(MeshBatch, BatchElementMask, PrimitveInfo.MeshBatchAndRelevance.PrimitiveSceneProxy);
 	}
 
 	if (VisibleMeshDrawCommands.Num() > 0)
diff --git a/Engine/Source/Runtime/Renderer/Private/HairStrands/HairStrandsVisibility.cpp b/Engine/Source/Runtime/Renderer/Private/HairStrands/HairStrandsVisibility.cpp
--- a/Engine/Source/Runtime/Renderer/Private/HairStrands/HairStrandsVisibility.cpp
+++ b/Engine/Source/Runtime/Renderer/Private/HairStrands/HairStrandsVisibility.cpp
@@ -51,20 +51,44 @@ static FAutoConsoleVariableRef CVarGHairVisibilityPPLLMeanListElementCountPerPix
 static int32 GHairVisibilityPPLLMaxRenderNodePerPixel = 16;
 static FAutoConsoleVariableRef CVarGHairVisibilityPPLLMeanNodeCountPerPixel(TEXT("r.HairStrands.VisibilityPPLLMaxRenderNodePerPixel"), GHairVisibilityPPLLMaxRenderNodePerPixel, TEXT("The maximum number of node allowed to be independently shaded and composited per pixel. Total amount of node will be width*height*VisibilityPPLLMaxRenderNodePerPixel. The last node is used to aggregate all furthest strands to shade into a single one."));
 
+static int32 GHairStrandsVisibilityMaterialPass = 0;
+static FAutoConsoleVariableRef CVarHairStrandsVisibilityMaterialPass(TEXT("r.HairStrands.Visibility.MaterialPass"), GHairStrandsVisibilityMaterialPass, TEXT("Enable the deferred material pass evaluation after the hair visibility is resolved."));
 
 /////////////////////////////////////////////////////////////////////////////////////////
 
+namespace HairStrandsVisibilityInternal
+{
+	struct NodeData
+	{
+		uint32 Depth;
+		uint32 PrimitiveId_ClusterId;
+		uint32 Tangent_Coverage;
+		uint32 BaseColor_Roughness;
+		uint32 Specular;
+	};
+
+	// 128 bit alignment
+	struct NodeVis
+	{
+		uint32 Depth;
+		uint32 PrimitiveId_ClusterId;
+		uint32 Coverage_ClusterID_Pad;
+		uint32 Pad;
+	};
+}
+
 enum EHairVisibilityRenderMode
 {
 	HairVisibilityRenderMode_MSAA,
 	HairVisibilityRenderMode_Transmittance,
 	HairVisibilityRenderMode_PPLL,
+	HairVisibilityRenderMode_MSAA_Visibility,
 	HairVisibilityRenderModeCount
 };
 
 EHairVisibilityRenderMode GetHairVisibilityRenderMode()
 {
-	return GHairVisibilityPPLL >0 ? HairVisibilityRenderMode_PPLL : HairVisibilityRenderMode_MSAA;
+	return GHairVisibilityPPLL > 0 ? HairVisibilityRenderMode_PPLL : HairVisibilityRenderMode_MSAA;
 }
 
 uint32 GetPPLLMeanListElementCountPerPixel()
@@ -131,6 +155,461 @@ static bool IsCompatibleWithHairVisibility(const FMeshMaterialShaderPermutationP
 	return IsCompatibleWithHairStrands(Parameters.Material, Parameters.Platform);
 }
 
+/////////////////////////////////////////////////////////////////////////////////////////
+
+BEGIN_GLOBAL_SHADER_PARAMETER_STRUCT(FMaterialPassParameters, )
+	SHADER_PARAMETER(FIntPoint, MaxResolution)
+	SHADER_PARAMETER(uint32, MaxSampleCount)
+	SHADER_PARAMETER(uint32, NodeGroupSize)
+	SHADER_PARAMETER_TEXTURE(Texture2D<uint>, NodeIndex)
+	SHADER_PARAMETER_SRV(StructuredBuffer<uint>, NodeCoord)
+	SHADER_PARAMETER_SRV(StructuredBuffer<FNodeVis>, NodeVis)
+	SHADER_PARAMETER_SRV(Buffer<uint>, IndirectArgs)
+END_GLOBAL_SHADER_PARAMETER_STRUCT()
+IMPLEMENT_GLOBAL_SHADER_PARAMETER_STRUCT(FMaterialPassParameters, "MaterialPassParameters");
+
+class FHairMaterialVS : public FMeshMaterialShader
+{
+	DECLARE_SHADER_TYPE(FHairMaterialVS, MeshMaterial);
+
+protected:
+	FHairMaterialVS(const ShaderMetaType::CompiledShaderInitializerType& Initializer)
+		: FMeshMaterialShader(Initializer)
+	{
+		ERHIFeatureLevel::Type FeatureLevel = GetMaxSupportedFeatureLevel((EShaderPlatform)Initializer.Target.Platform);
+		check(FSceneInterface::GetShadingPath(FeatureLevel) != EShadingPath::Mobile);
+		PassUniformBuffer.Bind(Initializer.ParameterMap, FMaterialPassParameters::StaticStructMetadata.GetShaderVariableName());
+	}
+
+	FHairMaterialVS() {}
+
+	static bool ShouldCompilePermutation(const FMeshMaterialShaderPermutationParameters& Parameters)
+	{
+		const bool bIsCompatible = IsCompatibleWithHairVisibility(Parameters) && Parameters.VertexFactoryType->GetFName() == FName(TEXT("FHairStrandsVertexFactory"));
+		return bIsCompatible;
+	}
+
+	static void ModifyCompilationEnvironment(const FMaterialShaderPermutationParameters& Parameters, FShaderCompilerEnvironment& OutEnvironment)
+	{
+		FMaterialShader::ModifyCompilationEnvironment(Parameters, OutEnvironment);
+	}
+};
+IMPLEMENT_MATERIAL_SHADER_TYPE(, FHairMaterialVS, TEXT("/Engine/Private/HairStrands/HairStrandsMaterialVS.usf"), TEXT("Main"), SF_Vertex);
+
+/////////////////////////////////////////////////////////////////////////////////////////
+
+class FHairMaterialShaderElementData : public FMeshMaterialShaderElementData
+{
+public:
+	FHairMaterialShaderElementData(int32 ClusterId, int32 MaterialId, int32 PrimitiveId) : MaterialPass_ClusterId(ClusterId), MaterialPass_MaterialId(MaterialId), MaterialPass_PrimitiveId(PrimitiveId) { }
+	uint32 MaterialPass_ClusterId;
+	uint32 MaterialPass_MaterialId;
+	uint32 MaterialPass_PrimitiveId;
+};
+
+class FHairMaterialPS : public FMeshMaterialShader
+{
+	DECLARE_SHADER_TYPE(FHairMaterialPS, MeshMaterial);
+
+public:
+	FHairMaterialPS(const ShaderMetaType::CompiledShaderInitializerType& Initializer)
+		: FMeshMaterialShader(Initializer)
+	{
+		ERHIFeatureLevel::Type FeatureLevel = GetMaxSupportedFeatureLevel((EShaderPlatform)Initializer.Target.Platform);
+		check(FSceneInterface::GetShadingPath(FeatureLevel) != EShadingPath::Mobile);
+		PassUniformBuffer.Bind(Initializer.ParameterMap, FMaterialPassParameters::StaticStructMetadata.GetShaderVariableName());
+		MaterialPass_ClusterId.Bind(Initializer.ParameterMap, TEXT("MaterialPass_ClusterId"));
+		MaterialPass_MaterialId.Bind(Initializer.ParameterMap, TEXT("MaterialPass_MaterialId"));
+		MaterialPass_PrimitiveId.Bind(Initializer.ParameterMap, TEXT("MaterialPass_PrimitiveId"));
+
+		// Fake binding for avoiding error from the shader compiler. These output are actually properly bound by the rendering 
+		// pass, and not directly by this material shader
+		OutNodeData.Bind(Initializer.ParameterMap, TEXT("OutNodeData"));
+		OutNodeVelocity.Bind(Initializer.ParameterMap, TEXT("OutNodeVelocity"));
+	}
+
+	FHairMaterialPS() {}
+
+	static bool ShouldCompilePermutation(const FMeshMaterialShaderPermutationParameters& Parameters)
+	{
+		const bool bIsCompatible = IsCompatibleWithHairStrands(Parameters.Material, Parameters.Platform) && Parameters.VertexFactoryType->GetFName() == FName(TEXT("FHairStrandsVertexFactory"));
+		return bIsCompatible;
+	}
+
+	static void ModifyCompilationEnvironment(const FMaterialShaderPermutationParameters& Parameters, FShaderCompilerEnvironment& OutEnvironment)
+	{
+		FMaterialShader::ModifyCompilationEnvironment(Parameters, OutEnvironment);
+	}
+
+	virtual bool Serialize(FArchive& Ar) override
+	{
+		bool bShaderHasOutdatedParameters = FMeshMaterialShader::Serialize(Ar);
+		Ar << MaterialPass_ClusterId;
+		Ar << MaterialPass_MaterialId;
+		Ar << MaterialPass_PrimitiveId;
+		Ar << OutNodeData;
+		Ar << OutNodeVelocity;
+		return bShaderHasOutdatedParameters;
+	}
+
+	void GetShaderBindings(
+		const FScene* Scene,
+		ERHIFeatureLevel::Type FeatureLevel,
+		const FPrimitiveSceneProxy* PrimitiveSceneProxy,
+		const FMaterialRenderProxy& MaterialRenderProxy,
+		const FMaterial& Material,
+		const FMeshPassProcessorRenderState& DrawRenderState,
+		const FHairMaterialShaderElementData& ShaderElementData,
+		FMeshDrawSingleShaderBindings& ShaderBindings) const
+	{
+		FMeshMaterialShader::GetShaderBindings(Scene, FeatureLevel, PrimitiveSceneProxy, MaterialRenderProxy, Material, DrawRenderState, ShaderElementData, ShaderBindings);
+		ShaderBindings.Add(MaterialPass_ClusterId, ShaderElementData.MaterialPass_ClusterId);
+		ShaderBindings.Add(MaterialPass_MaterialId, ShaderElementData.MaterialPass_MaterialId);
+		ShaderBindings.Add(MaterialPass_PrimitiveId, ShaderElementData.MaterialPass_PrimitiveId);
+	}
+
+private:
+	FShaderParameter MaterialPass_ClusterId;
+	FShaderParameter MaterialPass_MaterialId;
+	FShaderParameter MaterialPass_PrimitiveId;
+	FShaderParameter OutNodeData;
+	FShaderParameter OutNodeVelocity;
+};
+IMPLEMENT_MATERIAL_SHADER_TYPE(, FHairMaterialPS, TEXT("/Engine/Private/HairStrands/HairStrandsMaterialPS.usf"), TEXT("Main"), SF_Pixel);
+
+/////////////////////////////////////////////////////////////////////////////////////////
+
+class FHairMaterialProcessor : public FMeshPassProcessor
+{
+public:
+	FHairMaterialProcessor(
+		const FScene* Scene,
+		const FSceneView* InViewIfDynamicMeshCommand,
+		const FMeshPassProcessorRenderState& InPassDrawRenderState,
+		FDynamicPassMeshDrawListContext* InDrawListContext);
+
+	virtual void AddMeshBatch(const FMeshBatch& RESTRICT MeshBatch, uint64 BatchElementMask, const FPrimitiveSceneProxy* RESTRICT PrimitiveSceneProxy, int32 StaticMeshId = -1) override final;
+	void AddMeshBatch(const FMeshBatch& RESTRICT MeshBatch, uint64 BatchElementMask, const FPrimitiveSceneProxy* RESTRICT PrimitiveSceneProxy, int32 StaticMeshId, int32 HairClusterId, int32 HairMaterialId);
+
+private:
+	void Process(
+		const FMeshBatch& MeshBatch,
+		uint64 BatchElementMask,
+		const FPrimitiveSceneProxy* RESTRICT PrimitiveSceneProxy,
+		int32 StaticMeshId,
+		const FMaterialRenderProxy& RESTRICT MaterialRenderProxy,
+		const FMaterial& RESTRICT MaterialResource,
+		const int32 HairClusterId,
+		const int32 HairMaterialId,
+		const int32 HairPrimitiveId);
+
+	FMeshPassProcessorRenderState PassDrawRenderState;
+};
+
+void FHairMaterialProcessor::AddMeshBatch(const FMeshBatch& RESTRICT MeshBatch, uint64 BatchElementMask, const FPrimitiveSceneProxy* RESTRICT PrimitiveSceneProxy, int32 StaticMeshId)
+{
+	AddMeshBatch(MeshBatch, BatchElementMask, PrimitiveSceneProxy, StaticMeshId, 0, 0);
+}
+
+void FHairMaterialProcessor::AddMeshBatch(const FMeshBatch& RESTRICT MeshBatch, uint64 BatchElementMask, const FPrimitiveSceneProxy* RESTRICT PrimitiveSceneProxy, int32 StaticMeshId, int32 HairClusterId, int32 HairMaterialId)
+{
+	int32 PrimitiveId = 0;
+	int32 ScenePrimitiveId = 0;
+	GetDrawCommandPrimitiveId(PrimitiveSceneProxy->GetPrimitiveSceneInfo(), MeshBatch.Elements[0], PrimitiveId, ScenePrimitiveId);
+
+	static const FVertexFactoryType* CompatibleVF = FVertexFactoryType::GetVFByName(TEXT("FHairStrandsVertexFactory"));
+
+	// Determine the mesh's material and blend mode.
+	const FMaterialRenderProxy* FallbackMaterialRenderProxyPtr = nullptr;
+	const FMaterial& Material = MeshBatch.MaterialRenderProxy->GetMaterialWithFallback(FeatureLevel, FallbackMaterialRenderProxyPtr);
+	const bool bIsCompatible = IsCompatibleWithHairStrands(&Material, FeatureLevel);
+	const bool bIsHairStrandsFactory = MeshBatch.VertexFactory->GetType()->GetId() == CompatibleVF->GetId();
+
+	if (bIsCompatible
+		&& bIsHairStrandsFactory
+		&& (!PrimitiveSceneProxy || PrimitiveSceneProxy->ShouldRenderInMainPass())
+		&& ShouldIncludeDomainInMeshPass(Material.GetMaterialDomain()))
+	{
+		// For the mesh patch to be rendered a single triangle triangle to spawn the necessary amount of thread
+		FMeshBatch MeshBatchCopy = MeshBatch;
+		for (uint32 ElementIt = 0, ElementCount=uint32(MeshBatch.Elements.Num()); ElementIt < ElementCount; ++ElementIt)
+		{
+			MeshBatchCopy.Elements[ElementIt].FirstIndex = 0;
+			MeshBatchCopy.Elements[ElementIt].NumPrimitives = 1;
+			MeshBatchCopy.Elements[ElementIt].NumInstances = 1;
+		}
+
+		const FMaterialRenderProxy& MaterialRenderProxy = FallbackMaterialRenderProxyPtr ? *FallbackMaterialRenderProxyPtr : *MeshBatch.MaterialRenderProxy;
+		Process(MeshBatchCopy, BatchElementMask, PrimitiveSceneProxy, StaticMeshId, MaterialRenderProxy, Material, HairClusterId, HairMaterialId, PrimitiveId);
+	}
+}
+
+void FHairMaterialProcessor::Process(
+	const FMeshBatch& MeshBatch,
+	uint64 BatchElementMask,
+	const FPrimitiveSceneProxy* RESTRICT PrimitiveSceneProxy,
+	int32 StaticMeshId,
+	const FMaterialRenderProxy& RESTRICT MaterialRenderProxy,
+	const FMaterial& RESTRICT MaterialResource,
+	const int32 HairClusterId,
+	const int32 HairMaterialId,
+	const int32 HairPrimitiveId)
+{
+	const FVertexFactory* VertexFactory = MeshBatch.VertexFactory;
+
+	TMeshProcessorShaders<
+		FHairMaterialVS,
+		FMeshMaterialShader,
+		FMeshMaterialShader,
+		FHairMaterialPS> PassShaders;
+	{
+		FVertexFactoryType* VertexFactoryType = VertexFactory->GetType();
+		PassShaders.VertexShader = MaterialResource.GetShader<FHairMaterialVS>(VertexFactoryType);
+		PassShaders.PixelShader = MaterialResource.GetShader<FHairMaterialPS>(VertexFactoryType);
+	}
+
+	FMeshPassProcessorRenderState DrawRenderState(PassDrawRenderState);
+	FHairMaterialShaderElementData ShaderElementData(HairClusterId, HairMaterialId, HairPrimitiveId);
+	ShaderElementData.InitializeMeshMaterialData(ViewIfDynamicMeshCommand, PrimitiveSceneProxy, MeshBatch, StaticMeshId, false);
+
+	BuildMeshDrawCommands(
+		MeshBatch,
+		BatchElementMask,
+		PrimitiveSceneProxy,
+		MaterialRenderProxy,
+		MaterialResource,
+		DrawRenderState,
+		PassShaders,
+		ERasterizerFillMode::FM_Solid,
+		ERasterizerCullMode::CM_CCW,
+		FMeshDrawCommandSortKey::Default,
+		EMeshPassFeatures::Default,
+		ShaderElementData);
+}
+
+FHairMaterialProcessor::FHairMaterialProcessor(
+	const FScene* Scene,
+	const FSceneView* InViewIfDynamicMeshCommand,
+	const FMeshPassProcessorRenderState& InPassDrawRenderState,
+	FDynamicPassMeshDrawListContext* InDrawListContext)
+	: FMeshPassProcessor(Scene, Scene->GetFeatureLevel(), InViewIfDynamicMeshCommand, InDrawListContext)
+	, PassDrawRenderState(InPassDrawRenderState)
+{
+}
+
+/////////////////////////////////////////////////////////////////////////////////////////
+
+BEGIN_SHADER_PARAMETER_STRUCT(FVisibilityMaterialPassParameters, )
+	SHADER_PARAMETER_RDG_TEXTURE(Texture2D<uint>, NodeIndex)
+	SHADER_PARAMETER_RDG_BUFFER_SRV(StructuredBuffer<uint>,	NodeCoord)
+	SHADER_PARAMETER_RDG_BUFFER_SRV(StructuredBuffer<FNodeVis>, NodeVis)
+	SHADER_PARAMETER_RDG_BUFFER_SRV(Buffer<uint>, IndirectArgs)
+	SHADER_PARAMETER_RDG_BUFFER_UAV(RWStructuredBuffer<FNodeData>, OutNodeData)
+	SHADER_PARAMETER_RDG_BUFFER_UAV(RWBuffer<float2>, OutNodeVelocity)
+	RENDER_TARGET_BINDING_SLOTS()
+END_SHADER_PARAMETER_STRUCT()
+
+struct FMaterialPassOutput
+{
+	static const EPixelFormat VelocityFormat = PF_G16R16;
+	FRDGBufferRef NodeData = nullptr;
+	FRDGBufferRef NodeVelocity = nullptr;
+};
+
+static FMaterialPassOutput AddHairMaterialPass(
+	FRDGBuilder& GraphBuilder,
+	const FScene* Scene,
+	const FViewInfo* ViewInfo,
+	const FHairStrandsClusterDatas& ClusterDatas,
+	const uint32 NodeGroupSize,
+	FRDGTextureRef CompactNodeIndex,
+	FRDGBufferRef CompactNodeVis,
+	FRDGBufferRef CompactNodeCoord,
+	FRDGBufferRef IndirectArgBuffer)
+{
+	if (!CompactNodeVis || !CompactNodeIndex)
+		return FMaterialPassOutput();
+
+	const uint32 MaxNodeCount = CompactNodeVis->Desc.NumElements;
+
+	FMaterialPassOutput Output;
+	Output.NodeData		= GraphBuilder.CreateBuffer(FRDGBufferDesc::CreateStructuredDesc(sizeof(HairStrandsVisibilityInternal::NodeData), MaxNodeCount), TEXT("HairCompactNodeData"));
+	Output.NodeVelocity = GraphBuilder.CreateBuffer(FRDGBufferDesc::CreateBufferDesc(4, CompactNodeVis->Desc.NumElements), TEXT("HairCompactNodeVelocity"));
+
+	const uint32 ResolutionDim = FMath::CeilToInt(FMath::Sqrt(MaxNodeCount));
+	const FIntPoint Resolution(ResolutionDim, ResolutionDim);
+
+	FRDGTextureDesc OutputDesc;
+	OutputDesc.Extent.X = Resolution.X;
+	OutputDesc.Extent.Y = Resolution.Y;
+	OutputDesc.Depth = 0;
+	OutputDesc.Format = PF_FloatRGBA;
+	OutputDesc.NumMips = 1;
+	OutputDesc.Flags = 0;
+	OutputDesc.TargetableFlags = TexCreate_RenderTargetable;
+	FRDGTextureRef OutDummyTexture0 = GraphBuilder.CreateTexture(OutputDesc, TEXT("HairMaterialDummyOutput"));
+
+
+	// Add resources reference to the pass parameters, in order to get the resource lifetime extended to this pass
+	FVisibilityMaterialPassParameters* PassParameters = GraphBuilder.AllocParameters<FVisibilityMaterialPassParameters>();
+	PassParameters->NodeIndex		= CompactNodeIndex;
+	PassParameters->NodeVis			= GraphBuilder.CreateSRV(CompactNodeVis);
+	PassParameters->NodeCoord		= GraphBuilder.CreateSRV(CompactNodeCoord);
+	PassParameters->IndirectArgs	= GraphBuilder.CreateSRV(IndirectArgBuffer);
+	PassParameters->OutNodeData		= GraphBuilder.CreateUAV(FRDGBufferUAVDesc(Output.NodeData));
+	PassParameters->OutNodeVelocity	= GraphBuilder.CreateUAV(FRDGBufferUAVDesc(Output.NodeVelocity, FMaterialPassOutput::VelocityFormat));
+	PassParameters->RenderTargets[0] = FRenderTargetBinding(OutDummyTexture0, ERenderTargetLoadAction::EClear, 0);
+
+
+	GraphBuilder.AddPass(
+		RDG_EVENT_NAME("HairStrandsMaterialPass"),
+		PassParameters,
+		ERDGPassFlags::Raster,
+		[PassParameters, Scene = Scene, ViewInfo, &ClusterDatas, MaxNodeCount, Resolution, NodeGroupSize](FRHICommandListImmediate& RHICmdList)
+	{
+		check(RHICmdList.IsInsideRenderPass());
+		check(IsInRenderingThread());
+
+		FMaterialPassParameters MaterialPassParameters;
+		MaterialPassParameters.MaxResolution	= Resolution;
+		MaterialPassParameters.NodeGroupSize	= NodeGroupSize;
+		MaterialPassParameters.MaxSampleCount	= MaxNodeCount;
+		MaterialPassParameters.NodeIndex		= PassParameters->NodeIndex->GetPooledRenderTarget()->GetRenderTargetItem().ShaderResourceTexture;
+		MaterialPassParameters.NodeCoord		= PassParameters->NodeCoord->GetRHI();
+		MaterialPassParameters.NodeVis			= PassParameters->NodeVis->GetRHI();
+		MaterialPassParameters.IndirectArgs		= PassParameters->IndirectArgs->GetRHI();
+		TUniformBufferRef<FMaterialPassParameters> MaterialPassParametersBuffer = TUniformBufferRef<FMaterialPassParameters>::CreateUniformBufferImmediate(MaterialPassParameters, UniformBuffer_SingleFrame);
+
+		FMeshPassProcessorRenderState DrawRenderState(*ViewInfo, MaterialPassParametersBuffer);
+		{
+			TUniformBufferRef<FViewUniformShaderParameters> ViewUniformShaderParameters;
+			const bool bEnableMSAA = false;
+			SetUpViewHairRenderInfo(*ViewInfo, bEnableMSAA, ViewInfo->CachedViewUniformShaderParameters->HairRenderInfo);
+			ViewUniformShaderParameters = TUniformBufferRef<FViewUniformShaderParameters>::CreateUniformBufferImmediate(*ViewInfo->CachedViewUniformShaderParameters, UniformBuffer_SingleFrame);
+			DrawRenderState.SetViewUniformBuffer(ViewUniformShaderParameters);
+		}
+
+		{
+			RHICmdList.SetViewport(0, 0, 0.0f, Resolution.X, Resolution.Y, 1.0f);
+			DrawRenderState.SetBlendState(TStaticBlendState<>::GetRHI());
+			DrawRenderState.SetDepthStencilState(TStaticDepthStencilState <false, CF_Always> ::GetRHI());
+			
+			FDynamicMeshDrawCommandStorage DynamicMeshDrawCommandStorage;
+			FMeshCommandOneFrameArray VisibleMeshDrawCommands;
+			FGraphicsMinimalPipelineStateSet PipelineStateSet;
+			FDynamicPassMeshDrawListContext ShadowContext(DynamicMeshDrawCommandStorage, VisibleMeshDrawCommands, PipelineStateSet);
+			FHairMaterialProcessor MeshProcessor(Scene, ViewInfo, DrawRenderState, &ShadowContext);
+
+			for (const FHairStrandsClusterData& ClusterData : ClusterDatas.Datas)
+			{
+				for (const FHairStrandsClusterData::PrimitiveInfo& PrimitiveInfo : ClusterData.PrimitivesInfos)
+				{
+					const FMeshBatch& MeshBatch = *PrimitiveInfo.MeshBatchAndRelevance.Mesh;
+					const uint64 BatchElementMask = ~0ull;
+					MeshProcessor.AddMeshBatch(MeshBatch, BatchElementMask, PrimitiveInfo.MeshBatchAndRelevance.PrimitiveSceneProxy, -1, ClusterData.ClusterId, PrimitiveInfo.MaterialId);
+				}
+			}
+
+			if (VisibleMeshDrawCommands.Num() > 0)
+			{
+				FRHIVertexBuffer* PrimitiveIdVertexBuffer = nullptr;
+				SortAndMergeDynamicPassMeshDrawCommands(ViewInfo->GetFeatureLevel(), VisibleMeshDrawCommands, DynamicMeshDrawCommandStorage, PrimitiveIdVertexBuffer, 1);
+				SubmitMeshDrawCommands(VisibleMeshDrawCommands, PipelineStateSet, PrimitiveIdVertexBuffer, 0, false, 1, RHICmdList);
+			}
+		}
+	});
+
+	return Output;
+}
+
+///////////////////////////////////////////////////////////////////////////////////////////////////
+class FHairVelocityCS: public FGlobalShader
+{
+	DECLARE_GLOBAL_SHADER(FHairVelocityCS);
+	SHADER_USE_PARAMETER_STRUCT(FHairVelocityCS, FGlobalShader);
+
+	class FVendor : SHADER_PERMUTATION_INT("PERMUTATION_VENDOR", HairVisibilityVendorCount);
+	class FVelocity : SHADER_PERMUTATION_INT("PERMUTATION_VELOCITY", 4);
+	using FPermutationDomain = TShaderPermutationDomain<FVendor, FVelocity>;
+
+	BEGIN_SHADER_PARAMETER_STRUCT(FParameters, )
+		SHADER_PARAMETER(FIntPoint, ResolutionOffset)
+		SHADER_PARAMETER_RDG_TEXTURE(Texture2D, NodeIndex)
+		SHADER_PARAMETER_RDG_BUFFER_SRV(Buffer, NodeVelocity)
+		SHADER_PARAMETER_RDG_BUFFER_SRV(StructuredBuffer<FNodeVis>, NodeVis)
+		SHADER_PARAMETER_RDG_TEXTURE_UAV(Texture2D, OutVelocityTexture)
+		SHADER_PARAMETER_STRUCT_REF(FViewUniformShaderParameters, ViewUniformBuffer)
+		SHADER_PARAMETER_STRUCT_REF(FSceneTexturesUniformParameters, SceneTexturesStruct)
+	END_SHADER_PARAMETER_STRUCT()
+
+public:
+
+	static bool ShouldCompilePermutation(const FGlobalShaderPermutationParameters& Parameters)
+	{
+		return IsHairStrandsSupported(Parameters.Platform);
+	}
+};
+
+IMPLEMENT_GLOBAL_SHADER(FHairVelocityCS, "/Engine/Private/HairStrands/HairStrandsVelocity.usf", "MainCS", SF_Compute);
+
+static void AddHairVelocityPass(
+	FRDGBuilder& GraphBuilder,
+	const FViewInfo& View,
+	const FHairStrandsClusterDatas& ClusterDatas,
+	FRDGTextureRef& NodeIndex,
+	FRDGBufferRef& NodeVis,
+	FRDGBufferRef& NodeVelocity,
+	FRDGTextureRef& OutVelocityTexture)
+{
+	const bool bWriteOutVelocity = OutVelocityTexture != nullptr;
+	if (!bWriteOutVelocity)
+		return;
+
+	check(OutVelocityTexture->Desc.Format == FMaterialPassOutput::VelocityFormat);
+
+	FSceneRenderTargets& SceneContext = FSceneRenderTargets::Get(GraphBuilder.RHICmdList);
+	FSceneTexturesUniformParameters SceneTextures;
+	SetupSceneTextureUniformParameters(SceneContext, View.FeatureLevel, ESceneTextureSetupMode::All, SceneTextures);
+
+	FHairVelocityCS::FPermutationDomain PermutationVector;
+	PermutationVector.Set<FHairVelocityCS::FVendor>(GetVendor());
+	PermutationVector.Set<FHairVelocityCS::FVelocity>(bWriteOutVelocity ? FMath::Clamp(GHairVelocityType + 1, 0, 3) : 0);
+
+	FHairVelocityCS::FParameters* PassParameters = GraphBuilder.AllocParameters<FHairVelocityCS::FParameters>();
+	PassParameters->SceneTexturesStruct = CreateUniformBufferImmediate(SceneTextures, EUniformBufferUsage::UniformBuffer_SingleDraw);
+	PassParameters->ViewUniformBuffer = View.ViewUniformBuffer;
+	PassParameters->NodeIndex = NodeIndex;
+	PassParameters->NodeVis = GraphBuilder.CreateSRV(NodeVis);
+	PassParameters->NodeVelocity = GraphBuilder.CreateSRV(NodeVelocity, FMaterialPassOutput::VelocityFormat);
+	PassParameters->OutVelocityTexture = GraphBuilder.CreateUAV(OutVelocityTexture);
+
+	FIntRect TotalRect = ComputeVisibleHairStrandsClustersRect(View.ViewRect, ClusterDatas);
+
+	// Snap the rect onto thread group boundary
+	const FIntPoint GroupSize = GetVendorOptimalGroupSize2D();
+	TotalRect.Min.X = FMath::FloorToInt(float(TotalRect.Min.X) / float(GroupSize.X)) * GroupSize.X;
+	TotalRect.Min.Y = FMath::FloorToInt(float(TotalRect.Min.Y) / float(GroupSize.Y)) * GroupSize.Y;
+	TotalRect.Max.X = FMath::CeilToInt(float(TotalRect.Max.X) / float(GroupSize.X)) * GroupSize.X;
+	TotalRect.Max.Y = FMath::CeilToInt(float(TotalRect.Max.Y) / float(GroupSize.Y)) * GroupSize.Y;
+
+	FIntPoint RectResolution(TotalRect.Width(), TotalRect.Height());
+	PassParameters->ResolutionOffset = FIntPoint(TotalRect.Min.X, TotalRect.Min.Y);
+
+	TShaderMapRef<FHairVelocityCS> ComputeShader(View.ShaderMap, PermutationVector);
+	FComputeShaderUtils::AddPass(
+		GraphBuilder,
+		RDG_EVENT_NAME("HairStrandsVelocity"),
+		*ComputeShader,
+		PassParameters,
+		FComputeShaderUtils::GetGroupCount(RectResolution, GroupSize));
+}
+
+/////////////////////////////////////////////////////////////////////////////////////////
+BEGIN_GLOBAL_SHADER_PARAMETER_STRUCT(FVisibilityPassGlobalParameters, )
+	SHADER_PARAMETER(uint32, MaxPPLLNodeCount)
+END_GLOBAL_SHADER_PARAMETER_STRUCT()
+IMPLEMENT_GLOBAL_SHADER_PARAMETER_STRUCT(FVisibilityPassGlobalParameters, "VisibilityPassGlobalParameters");
+
 BEGIN_SHADER_PARAMETER_STRUCT(FVisibilityPassParameters, )
 	SHADER_PARAMETER(uint32, HairVisibilityPass_MaxPPLLNodeCount)
 	SHADER_PARAMETER_RDG_TEXTURE_UAV(Texture2D, PPLLCounter)
@@ -167,6 +646,7 @@ protected:
 		OutEnvironment.SetDefine(TEXT("HAIR_RENDER_MODE"), RenderModeValue);
 	}
 };
+IMPLEMENT_MATERIAL_SHADER_TYPE(template<>, FHairVisibilityVS<HairVisibilityRenderMode_MSAA_Visibility>, TEXT("/Engine/Private/HairStrands/HairStrandsVisibilityVS.usf"), TEXT("Main"), SF_Vertex);
 IMPLEMENT_MATERIAL_SHADER_TYPE(template<>, FHairVisibilityVS<HairVisibilityRenderMode_MSAA>, TEXT("/Engine/Private/HairStrands/HairStrandsVisibilityVS.usf"), TEXT("Main"), SF_Vertex);
 IMPLEMENT_MATERIAL_SHADER_TYPE(template<>, FHairVisibilityVS<HairVisibilityRenderMode_Transmittance>, TEXT("/Engine/Private/HairStrands/HairStrandsVisibilityVS.usf"), TEXT("Main"), SF_Vertex);
 IMPLEMENT_MATERIAL_SHADER_TYPE(template<>, FHairVisibilityVS<HairVisibilityRenderMode_PPLL>, TEXT("/Engine/Private/HairStrands/HairStrandsVisibilityVS.usf"), TEXT("Main"), SF_Vertex);
@@ -176,8 +656,9 @@ IMPLEMENT_MATERIAL_SHADER_TYPE(template<>, FHairVisibilityVS<HairVisibilityRende
 class FHairVisibilityShaderElementData : public FMeshMaterialShaderElementData
 {
 public:
-	FHairVisibilityShaderElementData(uint32 InHairClusterId, uint32 InMaxPPLLNodeCount) : HairClusterId(InHairClusterId), MaxPPLLNodeCount(InMaxPPLLNodeCount) { }
+	FHairVisibilityShaderElementData(uint32 InHairClusterId, uint32 InHairMaterialId, uint32 InMaxPPLLNodeCount) : HairClusterId(InHairClusterId), HairMaterialId(InHairMaterialId), MaxPPLLNodeCount(InMaxPPLLNodeCount) { }
 	uint32 HairClusterId;
+	uint32 HairMaterialId;
 	uint32 MaxPPLLNodeCount;
 };
 
@@ -193,7 +674,9 @@ public:
 	{
 		ERHIFeatureLevel::Type FeatureLevel = GetMaxSupportedFeatureLevel((EShaderPlatform)Initializer.Target.Platform);
 		check(FSceneInterface::GetShadingPath(FeatureLevel) != EShadingPath::Mobile);
+		PassUniformBuffer.Bind(Initializer.ParameterMap, FVisibilityPassGlobalParameters::StaticStructMetadata.GetShaderVariableName());
 		HairVisibilityPass_HairClusterIndex.Bind(Initializer.ParameterMap, TEXT("HairVisibilityPass_HairClusterIndex"));
+		HairVisibilityPass_HairMaterialId.Bind(Initializer.ParameterMap, TEXT("HairVisibilityPass_HairMaterialId"));
 		HairVisibilityPass_MaxPPLLNodeCount.Bind(Initializer.ParameterMap, TEXT("HairVisibilityPass_MaxPPLLNodeCount"));
 		PPLLCounter.Bind(Initializer.ParameterMap, TEXT("PPLLCounter"));
 		PPLLNodeIndex.Bind(Initializer.ParameterMap, TEXT("PPLLNodeIndex"));
@@ -218,6 +701,7 @@ public:
 	{
 		bool bShaderHasOutdatedParameters = FMeshMaterialShader::Serialize(Ar);
 		Ar << HairVisibilityPass_HairClusterIndex;
+		Ar << HairVisibilityPass_HairMaterialId;
 		Ar << HairVisibilityPass_MaxPPLLNodeCount;
 		Ar << PPLLCounter;
 		Ar << PPLLNodeIndex;
@@ -237,16 +721,19 @@ public:
 	{
 		FMeshMaterialShader::GetShaderBindings(Scene, FeatureLevel, PrimitiveSceneProxy, MaterialRenderProxy, Material, DrawRenderState, ShaderElementData, ShaderBindings);
 		ShaderBindings.Add(HairVisibilityPass_HairClusterIndex, ShaderElementData.HairClusterId);
+		ShaderBindings.Add(HairVisibilityPass_HairMaterialId, ShaderElementData.HairMaterialId);
 		ShaderBindings.Add(HairVisibilityPass_MaxPPLLNodeCount, ShaderElementData.MaxPPLLNodeCount);
 	}
 
 	FShaderParameter HairVisibilityPass_HairClusterIndex;
+	FShaderParameter HairVisibilityPass_HairMaterialId;
 	FShaderParameter HairVisibilityPass_MaxPPLLNodeCount;
 	// This is on FVisibilityPassParameters but needed to avoid the compiler to assert with unbound parameters 
 	FShaderParameter PPLLCounter;
 	FShaderParameter PPLLNodeIndex;
 	FShaderParameter PPLLNodes;
 };
+IMPLEMENT_MATERIAL_SHADER_TYPE(template<>, FHairVisibilityPS<HairVisibilityRenderMode_MSAA_Visibility>, TEXT("/Engine/Private/HairStrands/HairStrandsVisibilityPS.usf"), TEXT("MainVisibility"), SF_Pixel);
 IMPLEMENT_MATERIAL_SHADER_TYPE(template<>, FHairVisibilityPS<HairVisibilityRenderMode_MSAA>, TEXT("/Engine/Private/HairStrands/HairStrandsVisibilityPS.usf"), TEXT("MainVisibility"), SF_Pixel);
 IMPLEMENT_MATERIAL_SHADER_TYPE(template<>, FHairVisibilityPS<HairVisibilityRenderMode_Transmittance>, TEXT("/Engine/Private/HairStrands/HairStrandsVisibilityPS.usf"), TEXT("MainVisibility"), SF_Pixel);
 IMPLEMENT_MATERIAL_SHADER_TYPE(template<>, FHairVisibilityPS<HairVisibilityRenderMode_PPLL>, TEXT("/Engine/Private/HairStrands/HairStrandsVisibilityPS.usf"), TEXT("MainVisibility"), SF_Pixel);
@@ -264,7 +751,7 @@ public:
 		FDynamicPassMeshDrawListContext* InDrawListContext);
 
 	virtual void AddMeshBatch(const FMeshBatch& RESTRICT MeshBatch, uint64 BatchElementMask, const FPrimitiveSceneProxy* RESTRICT PrimitiveSceneProxy, int32 StaticMeshId = -1) override final;
-	void AddMeshBatch(const FMeshBatch& RESTRICT MeshBatch, uint64 BatchElementMask, const FPrimitiveSceneProxy* RESTRICT PrimitiveSceneProxy, int32 StaticMeshId, uint32 HairClusterId);
+	void AddMeshBatch(const FMeshBatch& RESTRICT MeshBatch, uint64 BatchElementMask, const FPrimitiveSceneProxy* RESTRICT PrimitiveSceneProxy, int32 StaticMeshId, uint32 HairClusterId, uint32 HairMaterialId);
 
 	void SetMaxPPLLNodeCount(uint32 NewValue) { MaxPPLLNodeCount = NewValue; }
 
@@ -278,6 +765,7 @@ private:
 		const FMaterialRenderProxy& RESTRICT MaterialRenderProxy,
 		const FMaterial& RESTRICT MaterialResource,
 		const uint32 HairClusterId,
+		const uint32 HairMaterialId,
 		ERasterizerFillMode MeshFillMode,
 		ERasterizerCullMode MeshCullMode);
 
@@ -288,10 +776,10 @@ private:
 
 void FHairVisibilityProcessor::AddMeshBatch(const FMeshBatch& RESTRICT MeshBatch, uint64 BatchElementMask, const FPrimitiveSceneProxy* RESTRICT PrimitiveSceneProxy, int32 StaticMeshId)
 {
-	AddMeshBatch(MeshBatch, BatchElementMask, PrimitiveSceneProxy, StaticMeshId, 0);
+	AddMeshBatch(MeshBatch, BatchElementMask, PrimitiveSceneProxy, StaticMeshId, 0, 0);
 }
 
-void FHairVisibilityProcessor::AddMeshBatch(const FMeshBatch& RESTRICT MeshBatch, uint64 BatchElementMask, const FPrimitiveSceneProxy* RESTRICT PrimitiveSceneProxy, int32 StaticMeshId, uint32 HairClusterId)
+void FHairVisibilityProcessor::AddMeshBatch(const FMeshBatch& RESTRICT MeshBatch, uint64 BatchElementMask, const FPrimitiveSceneProxy* RESTRICT PrimitiveSceneProxy, int32 StaticMeshId, uint32 HairClusterId, uint32 HairMaterialId)
 {
 	static const FVertexFactoryType* CompatibleVF = FVertexFactoryType::GetVFByName(TEXT("FHairStrandsVertexFactory"));
 
@@ -309,12 +797,14 @@ void FHairVisibilityProcessor::AddMeshBatch(const FMeshBatch& RESTRICT MeshBatch
 		const FMaterialRenderProxy& MaterialRenderProxy = FallbackMaterialRenderProxyPtr ? *FallbackMaterialRenderProxyPtr : *MeshBatch.MaterialRenderProxy;
 		const ERasterizerFillMode MeshFillMode = ComputeMeshFillMode(MeshBatch, Material);
 		const ERasterizerCullMode MeshCullMode = ComputeMeshCullMode(MeshBatch, Material);
+		if (RenderMode == HairVisibilityRenderMode_MSAA_Visibility)
+			Process<HairVisibilityRenderMode_MSAA_Visibility>(MeshBatch, BatchElementMask, PrimitiveSceneProxy, StaticMeshId, MaterialRenderProxy, Material, HairClusterId, HairMaterialId, MeshFillMode, MeshCullMode);
 		if (RenderMode == HairVisibilityRenderMode_MSAA)
-			Process<HairVisibilityRenderMode_MSAA>(MeshBatch, BatchElementMask, PrimitiveSceneProxy, StaticMeshId, MaterialRenderProxy, Material, HairClusterId, MeshFillMode, MeshCullMode);
+			Process<HairVisibilityRenderMode_MSAA>(MeshBatch, BatchElementMask, PrimitiveSceneProxy, StaticMeshId, MaterialRenderProxy, Material, HairClusterId, HairMaterialId, MeshFillMode, MeshCullMode);
 		if (RenderMode == HairVisibilityRenderMode_Transmittance)
-			Process<HairVisibilityRenderMode_Transmittance>(MeshBatch, BatchElementMask, PrimitiveSceneProxy, StaticMeshId, MaterialRenderProxy, Material, HairClusterId, MeshFillMode, MeshCullMode);
+			Process<HairVisibilityRenderMode_Transmittance>(MeshBatch, BatchElementMask, PrimitiveSceneProxy, StaticMeshId, MaterialRenderProxy, Material, HairClusterId, HairMaterialId, MeshFillMode, MeshCullMode);
 		if (RenderMode == HairVisibilityRenderMode_PPLL)
-			Process<HairVisibilityRenderMode_PPLL>(MeshBatch, BatchElementMask, PrimitiveSceneProxy, StaticMeshId, MaterialRenderProxy, Material, HairClusterId, MeshFillMode, MeshCullMode);
+			Process<HairVisibilityRenderMode_PPLL>(MeshBatch, BatchElementMask, PrimitiveSceneProxy, StaticMeshId, MaterialRenderProxy, Material, HairClusterId, HairMaterialId, MeshFillMode, MeshCullMode);
 	}
 }
 
@@ -327,6 +817,7 @@ void FHairVisibilityProcessor::Process(
 	const FMaterialRenderProxy& RESTRICT MaterialRenderProxy,
 	const FMaterial& RESTRICT MaterialResource,
 	const uint32 HairClusterId,
+	const uint32 HairMaterialId,
 	ERasterizerFillMode MeshFillMode,
 	ERasterizerCullMode MeshCullMode)
 {
@@ -344,7 +835,7 @@ void FHairVisibilityProcessor::Process(
 	}
 
 	FMeshPassProcessorRenderState DrawRenderState(PassDrawRenderState);
-	FHairVisibilityShaderElementData ShaderElementData(HairClusterId, MaxPPLLNodeCount);
+	FHairVisibilityShaderElementData ShaderElementData(HairClusterId, HairMaterialId, MaxPPLLNodeCount);
 	ShaderElementData.InitializeMeshMaterialData(ViewIfDynamicMeshCommand, PrimitiveSceneProxy, MeshBatch, StaticMeshId, false);
 
 	BuildMeshDrawCommands(
@@ -524,12 +1015,13 @@ class FHairVisibilityPrimitiveIdCompactionCS : public FGlobalShader
 	DECLARE_GLOBAL_SHADER(FHairVisibilityPrimitiveIdCompactionCS);
 	SHADER_USE_PARAMETER_STRUCT(FHairVisibilityPrimitiveIdCompactionCS, FGlobalShader);
 
-	class FVendor   : SHADER_PERMUTATION_INT("PERMUTATION_VENDOR", HairVisibilityVendorCount);
-	class FVelocity : SHADER_PERMUTATION_INT("PERMUTATION_VELOCITY", 4);
+	class FVendor		: SHADER_PERMUTATION_INT("PERMUTATION_VENDOR", HairVisibilityVendorCount);
+	class FVelocity		: SHADER_PERMUTATION_INT("PERMUTATION_VELOCITY", 4);
 	class FViewTransmittance : SHADER_PERMUTATION_INT("PERMUTATION_VIEWTRANSMITTANCE", 2);
-	class FMaterial : SHADER_PERMUTATION_INT("PERMUTATION_MATERIAL_COMPACTION", 2);
-	class FPPLL     : SHADER_PERMUTATION_SPARSE_INT("PERMUTATION_PPLL", 0, 8, 16, 32); // See GetPPLLMaxRenderNodePerPixel
-	using FPermutationDomain = TShaderPermutationDomain<FVendor, FVelocity, FViewTransmittance, FMaterial, FPPLL>;
+	class FMaterial		: SHADER_PERMUTATION_INT("PERMUTATION_MATERIAL_COMPACTION", 2);
+	class FPPLL			: SHADER_PERMUTATION_SPARSE_INT("PERMUTATION_PPLL", 0, 8, 16, 32); // See GetPPLLMaxRenderNodePerPixel
+	class FVisibility	: SHADER_PERMUTATION_INT("PERMUTATION_VISIBILITY", 2);
+	using FPermutationDomain = TShaderPermutationDomain<FVendor, FVelocity, FViewTransmittance, FMaterial, FPPLL, FVisibility>;
 
 	BEGIN_SHADER_PARAMETER_STRUCT(FParameters, )
 		SHADER_PARAMETER(FIntPoint, OutputResolution)
@@ -570,6 +1062,7 @@ public:
 		if (PermutationVector.Get<FPPLL>() > 0)
 		{
 			PermutationVector.Set<FViewTransmittance>(0);
+			PermutationVector.Set<FVisibility>(0); 
 		}
 		return PermutationVector;
 	}
@@ -581,6 +1074,10 @@ public:
 		{
 			return false;
 		}
+		if (PermutationVector.Get<FPPLL>() > 0 && PermutationVector.Get<FVisibility>() > 0)
+		{
+			return false;
+		}
 		return IsHairStrandsSupported(Parameters.Platform);
 	}
 };
@@ -589,6 +1086,7 @@ IMPLEMENT_GLOBAL_SHADER(FHairVisibilityPrimitiveIdCompactionCS, "/Engine/Private
 
 static void AddHairVisibilityPrimitiveIdCompactionPass(
 	const bool bUsePPLL,
+	const bool bUseVisibility,
 	FRDGBuilder& GraphBuilder,
 	const FViewInfo& View,
 	const FHairStrandsClusterDatas& ClusterDatas,
@@ -611,11 +1109,20 @@ static void AddHairVisibilityPrimitiveIdCompactionPass(
 	}
 	else
 	{
-		check(PassParameters->MSAA_DepthTexture);
-		check(PassParameters->MSAA_IDTexture);
-		check(PassParameters->MSAA_MaterialTexture);
-		check(PassParameters->MSAA_AttributeTexture);
-		Resolution = PassParameters->MSAA_DepthTexture->Desc.Extent;
+		if (bUseVisibility)
+		{
+			check(PassParameters->MSAA_DepthTexture);
+			check(PassParameters->MSAA_IDTexture);
+			Resolution = PassParameters->MSAA_DepthTexture->Desc.Extent;
+		}
+		else
+		{
+			check(PassParameters->MSAA_DepthTexture);
+			check(PassParameters->MSAA_IDTexture);
+			check(PassParameters->MSAA_MaterialTexture);
+			check(PassParameters->MSAA_AttributeTexture);
+			Resolution = PassParameters->MSAA_DepthTexture->Desc.Extent;
+		}
 	}
 
 	
@@ -665,18 +1172,7 @@ static void AddHairVisibilityPrimitiveIdCompactionPass(
 	const uint32 SampleCount = FMath::RoundUpToPowerOfTwo(HairVisibilityMSAASampleCount);
 	const uint32 PPLLMaxRenderNodePerPixel = GetPPLLMaxRenderNodePerPixel();
 	const uint32 MaxRenderNodeCount = Resolution.X * Resolution.Y * (GetHairVisibilityRenderMode() == HairVisibilityRenderMode_MSAA ? SampleCount : PPLLMaxRenderNodePerPixel);
-	{
-		struct NodeData
-		{
-			uint32 Depth;
-			uint32 PrimitiveId_ClusterId;
-			uint32 Tangent_Coverage;
-			uint32 BaseColor_Roughness;
-			uint32 Specular;
-		};
-
-		OutCompactNodeData = GraphBuilder.CreateBuffer(FRDGBufferDesc::CreateStructuredDesc(sizeof(NodeData), MaxRenderNodeCount), TEXT("HairVisibilityPrimitiveIdCompactNodeData"));
-	}
+	OutCompactNodeData = GraphBuilder.CreateBuffer(FRDGBufferDesc::CreateStructuredDesc(bUseVisibility ? sizeof(HairStrandsVisibilityInternal::NodeVis) : sizeof(HairStrandsVisibilityInternal::NodeData), MaxRenderNodeCount), TEXT("HairVisibilityPrimitiveIdCompactNodeData"));
 
 	{
 		// Pixel coord of the node. Stored as 2*uint16, packed into a single uint32
@@ -688,12 +1184,14 @@ static void AddHairVisibilityPrimitiveIdCompactionPass(
 	SetupSceneTextureUniformParameters(SceneContext, View.FeatureLevel, ESceneTextureSetupMode::All, SceneTextures);
 
 	const bool bWriteOutVelocity = OutVelocityTexture != nullptr;
+	const uint32 VelocityPermutation = bWriteOutVelocity ? FMath::Clamp(GHairVelocityType + 1, 0, 3) : 0;
 	FHairVisibilityPrimitiveIdCompactionCS::FPermutationDomain PermutationVector;
 	PermutationVector.Set<FHairVisibilityPrimitiveIdCompactionCS::FVendor>(GetVendor());
-	PermutationVector.Set<FHairVisibilityPrimitiveIdCompactionCS::FVelocity>(bWriteOutVelocity ? FMath::Clamp(GHairVelocityType+1, 0, 3) : 0);
+	PermutationVector.Set<FHairVisibilityPrimitiveIdCompactionCS::FVelocity>(VelocityPermutation);
 	PermutationVector.Set<FHairVisibilityPrimitiveIdCompactionCS::FViewTransmittance>(PassParameters->ViewTransmittanceTexture ? 1 : 0);
 	PermutationVector.Set<FHairVisibilityPrimitiveIdCompactionCS::FMaterial>(GHairStrandsMaterialCompactionEnable ? 1 : 0);
 	PermutationVector.Set<FHairVisibilityPrimitiveIdCompactionCS::FPPLL>(bUsePPLL ? PPLLMaxRenderNodePerPixel : 0);
+	PermutationVector.Set<FHairVisibilityPrimitiveIdCompactionCS::FVisibility>(bUseVisibility ? 1 : 0);
 	PermutationVector = FHairVisibilityPrimitiveIdCompactionCS::RemapPermutation(PermutationVector);
 
 	PassParameters->OutputResolution = Resolution;
@@ -799,9 +1297,16 @@ static FRDGTextureRef AddHairVisibilityFillOpaqueDepth(
 	const FViewInfo* CapturedView = &View;
 
 	TArray<FIntRect> ClusterRects;
-	for (const FHairStrandsClusterData& ClusterData : ClusterDatas.Datas)
+	if (IsHairStrandsViewRectOptimEnable())
+	{
+		for (const FHairStrandsClusterData& ClusterData : ClusterDatas.Datas)
+		{
+			ClusterRects.Add(ClusterData.ScreenRect);
+		}
+	}
+	else
 	{
-		ClusterRects.Add(ClusterData.ScreenRect);
+		ClusterRects.Add(Viewport);
 	}
 
 	{
@@ -860,9 +1365,21 @@ static void AddHairVisibilityCommonPass(
 	const EHairVisibilityRenderMode RenderMode,
 	TPassParameters* PassParameters)
 {
+
+	auto GetPassName = [](EHairVisibilityRenderMode RenderMode)
+	{
+		switch (RenderMode)
+		{
+		case HairVisibilityRenderMode_PPLL:				return RDG_EVENT_NAME("HairStrandsVisibilityPPLLPass");
+		case HairVisibilityRenderMode_MSAA:				return RDG_EVENT_NAME("HairStrandsVisibilityMSAAPass");
+		case HairVisibilityRenderMode_MSAA_Visibility:	return RDG_EVENT_NAME("HairStrandsVisibilityMSAAVisPass");
+		case HairVisibilityRenderMode_Transmittance:	return RDG_EVENT_NAME("HairStrandsTransmittancePass");
+		default:										return RDG_EVENT_NAME("Noname");
+		}
+	};
+
 	GraphBuilder.AddPass(
-		RenderMode == HairVisibilityRenderMode_PPLL ? RDG_EVENT_NAME("HairStrandsPPLLPass") : 
-		(RenderMode == HairVisibilityRenderMode_MSAA ? RDG_EVENT_NAME("HairStrandsVisibilityPass") : RDG_EVENT_NAME("HairStrandsTransmittancePass")),
+		GetPassName(RenderMode), 
 		PassParameters,
 		ERDGPassFlags::Raster,
 		[PassParameters, Scene = Scene, ViewInfo, &ClusterDatas, RenderMode](FRHICommandListImmediate& RHICmdList)
@@ -881,7 +1398,7 @@ static void AddHairVisibilityCommonPass(
 			ViewUniformShaderParameters = TUniformBufferRef<FViewUniformShaderParameters>::CreateUniformBufferImmediate(*ViewInfo->CachedViewUniformShaderParameters, UniformBuffer_SingleFrame);
 			DrawRenderState.SetViewUniformBuffer(ViewUniformShaderParameters);
 		}
-
+		
 		{
 			RHICmdList.SetViewport(0, 0, 0.0f, ViewInfo->ViewRect.Width(), ViewInfo->ViewRect.Height(), 1.0f);
 			if (RenderMode == HairVisibilityRenderMode_MSAA)
@@ -891,6 +1408,11 @@ static void AddHairVisibilityCommonPass(
 					CW_RGBA, BO_Add, BF_One, BF_Zero, BO_Add, BF_One, BF_Zero>::GetRHI());
 				DrawRenderState.SetDepthStencilState(TStaticDepthStencilState<true, CF_DepthNearOrEqual>::GetRHI());
 			}
+			else if(RenderMode == HairVisibilityRenderMode_MSAA_Visibility)
+			{
+				DrawRenderState.SetBlendState(TStaticBlendState<CW_RGBA, BO_Add, BF_One, BF_Zero, BO_Add, BF_One, BF_Zero>::GetRHI());
+				DrawRenderState.SetDepthStencilState(TStaticDepthStencilState<true, CF_DepthNearOrEqual>::GetRHI());
+			}
 			else if (RenderMode == HairVisibilityRenderMode_Transmittance)
 			{
 				DrawRenderState.SetBlendState(TStaticBlendState<CW_RED, BO_Add, BF_DestColor, BF_Zero, BO_Add, BF_Zero, BF_Zero>::GetRHI());
@@ -916,11 +1438,11 @@ static void AddHairVisibilityCommonPass(
 
 			for (const FHairStrandsClusterData& ClusterData : ClusterDatas.Datas)
 			{
-				for (const FMeshBatchAndRelevance& MeshBatchAndRelevance : ClusterData.PrimitivesInfos)
+				for (const FHairStrandsClusterData::PrimitiveInfo& PrimitiveInfo : ClusterData.PrimitivesInfos)
 				{
-					const FMeshBatch& MeshBatch = *MeshBatchAndRelevance.Mesh;
+					const FMeshBatch& MeshBatch = *PrimitiveInfo.MeshBatchAndRelevance.Mesh;
 					const uint64 BatchElementMask = ~0ull;
-					MeshProcessor.AddMeshBatch(MeshBatch, BatchElementMask, MeshBatchAndRelevance.PrimitiveSceneProxy, -1, ClusterData.ClusterId);
+					MeshProcessor.AddMeshBatch(MeshBatch, BatchElementMask, PrimitiveInfo.MeshBatchAndRelevance.PrimitiveSceneProxy, -1, ClusterData.ClusterId, PrimitiveInfo.MaterialId);
 				}
 			}
 
@@ -935,6 +1457,7 @@ static void AddHairVisibilityCommonPass(
 }
 
 static void AddHairVisibilityMSAAPass(
+	const bool bUseVisibility,
 	FRDGBuilder& GraphBuilder,
 	const FScene* Scene,
 	const FViewInfo* ViewInfo,
@@ -948,90 +1471,123 @@ static void AddHairVisibilityMSAAPass(
 {
 	const uint32 MSAASampleCount = FMath::RoundUpToPowerOfTwo(FMath::Clamp(GHairVisibilitySampleCount, 1, 16));
 
+	if (bUseVisibility)
 	{
-		FRDGTextureDesc Desc;
-		Desc.Extent.X = Resolution.X;
-		Desc.Extent.Y = Resolution.Y;
-		Desc.Depth = 0;
-		Desc.Format = PF_R16G16B16A16_UINT;
-		Desc.NumMips = 1;
-		Desc.NumSamples = MSAASampleCount;
-		Desc.Flags = TexCreate_None;
-		Desc.TargetableFlags = TexCreate_RenderTargetable | TexCreate_ShaderResource;
-		Desc.bForceSharedTargetAndShaderResource = true;
-		OutVisibilityIdTexture = GraphBuilder.CreateTexture(Desc, TEXT("HairVisibilityIDTexture"));
+		{
+			FRDGTextureDesc Desc;
+			Desc.Extent.X = Resolution.X;
+			Desc.Extent.Y = Resolution.Y;
+			Desc.Depth = 0;
+			Desc.Format = PF_R32_UINT;
+			Desc.NumMips = 1;
+			Desc.NumSamples = MSAASampleCount;
+			Desc.Flags = TexCreate_None;
+			Desc.TargetableFlags = TexCreate_RenderTargetable | TexCreate_ShaderResource;
+			Desc.bForceSharedTargetAndShaderResource = true;
+			OutVisibilityIdTexture = GraphBuilder.CreateTexture(Desc, TEXT("HairVisibilityIDTexture"));
+		}
+		OutVisibilityMaterialTexture = nullptr;
+		OutVisibilityAttributeTexture = nullptr;
+		OutVisibilityVelocityTexture = nullptr;
+
+		AddClearGraphicPass(GraphBuilder, RDG_EVENT_NAME("HairStrandsClearVisibilityMSAAIdTexture"), ViewInfo, 0xFFFFFFFF, OutVisibilityIdTexture);
+
+		FVisibilityPassParameters* PassParameters = GraphBuilder.AllocParameters<FVisibilityPassParameters>();
+		PassParameters->RenderTargets[0] = FRenderTargetBinding(OutVisibilityIdTexture, ERenderTargetLoadAction::ELoad, 0);
+		PassParameters->RenderTargets.DepthStencil = FDepthStencilBinding(
+			OutVisibilityDepthTexture,
+			ERenderTargetLoadAction::ELoad,
+			ERenderTargetLoadAction::ENoAction,
+			FExclusiveDepthStencil::DepthWrite_StencilNop);
+		AddHairVisibilityCommonPass(GraphBuilder, Scene, ViewInfo, ClusterDatas, HairVisibilityRenderMode_MSAA_Visibility, PassParameters);
 	}
-
+	else
 	{
-		FRDGTextureDesc Desc;
-		Desc.Extent.X = Resolution.X;
-		Desc.Extent.Y = Resolution.Y;
-		Desc.Depth = 0;
-		Desc.Format = PF_R8G8B8A8;
-		Desc.NumMips = 1;
-		Desc.NumSamples = MSAASampleCount;
-		Desc.Flags = TexCreate_None;
-		Desc.TargetableFlags = TexCreate_RenderTargetable | TexCreate_ShaderResource;
-		Desc.ClearValue = FClearValueBinding(FLinearColor(0, 0, 0, 0));
-		Desc.bForceSharedTargetAndShaderResource = true;
-		OutVisibilityMaterialTexture = GraphBuilder.CreateTexture(Desc, TEXT("HairVisibilityMaterialTexture"));
-	}
+		{
+			FRDGTextureDesc Desc;
+			Desc.Extent.X = Resolution.X;
+			Desc.Extent.Y = Resolution.Y;
+			Desc.Depth = 0;
+			Desc.Format = PF_R16G16B16A16_UINT;
+			Desc.NumMips = 1;
+			Desc.NumSamples = MSAASampleCount;
+			Desc.Flags = TexCreate_None;
+			Desc.TargetableFlags = TexCreate_RenderTargetable | TexCreate_ShaderResource;
+			Desc.bForceSharedTargetAndShaderResource = true;
+			OutVisibilityIdTexture = GraphBuilder.CreateTexture(Desc, TEXT("HairVisibilityIDTexture"));
+		}
 
-	{
-		FRDGTextureDesc Desc;
-		Desc.Extent.X = Resolution.X;
-		Desc.Extent.Y = Resolution.Y;
-		Desc.Depth = 0;
-		Desc.Format = PF_R8G8B8A8;
-		Desc.NumMips = 1;
-		Desc.NumSamples = MSAASampleCount;
-		Desc.Flags = TexCreate_None;
-		Desc.TargetableFlags = TexCreate_RenderTargetable | TexCreate_ShaderResource;
-		Desc.ClearValue = FClearValueBinding(FLinearColor(0, 0, 0, 0));
-		Desc.bForceSharedTargetAndShaderResource = true;
-		OutVisibilityAttributeTexture = GraphBuilder.CreateTexture(Desc, TEXT("HairVisibilityAttributeTexture"));
-	}
+		{
+			FRDGTextureDesc Desc;
+			Desc.Extent.X = Resolution.X;
+			Desc.Extent.Y = Resolution.Y;
+			Desc.Depth = 0;
+			Desc.Format = PF_R8G8B8A8;
+			Desc.NumMips = 1;
+			Desc.NumSamples = MSAASampleCount;
+			Desc.Flags = TexCreate_None;
+			Desc.TargetableFlags = TexCreate_RenderTargetable | TexCreate_ShaderResource;
+			Desc.ClearValue = FClearValueBinding(FLinearColor(0, 0, 0, 0));
+			Desc.bForceSharedTargetAndShaderResource = true;
+			OutVisibilityMaterialTexture = GraphBuilder.CreateTexture(Desc, TEXT("HairVisibilityMaterialTexture"));
+		}
 
-	{
-		FRDGTextureDesc Desc;
-		Desc.Extent.X = Resolution.X;
-		Desc.Extent.Y = Resolution.Y;
-		Desc.Depth = 0;
-		Desc.Format = PF_G16R16;
-		Desc.NumMips = 1;
-		Desc.NumSamples = MSAASampleCount;
-		Desc.Flags = TexCreate_None;
-		Desc.TargetableFlags = TexCreate_RenderTargetable | TexCreate_ShaderResource;
-		Desc.ClearValue = FClearValueBinding(FLinearColor(0, 0, 0, 0));
-		Desc.bForceSharedTargetAndShaderResource = true;
-		OutVisibilityVelocityTexture = GraphBuilder.CreateTexture(Desc, TEXT("HairVisibilityVelocityTexture"));
-	}
-	AddClearGraphicPass(GraphBuilder, RDG_EVENT_NAME("HairStrandsClearVisibilityMSAAIdTexture"), ViewInfo, 0xFFFFFFFF, OutVisibilityIdTexture);
+		{
+			FRDGTextureDesc Desc;
+			Desc.Extent.X = Resolution.X;
+			Desc.Extent.Y = Resolution.Y;
+			Desc.Depth = 0;
+			Desc.Format = PF_R8G8B8A8;
+			Desc.NumMips = 1;
+			Desc.NumSamples = MSAASampleCount;
+			Desc.Flags = TexCreate_None;
+			Desc.TargetableFlags = TexCreate_RenderTargetable | TexCreate_ShaderResource;
+			Desc.ClearValue = FClearValueBinding(FLinearColor(0, 0, 0, 0));
+			Desc.bForceSharedTargetAndShaderResource = true;
+			OutVisibilityAttributeTexture = GraphBuilder.CreateTexture(Desc, TEXT("HairVisibilityAttributeTexture"));
+		}
 
-	// Manually clear RTs as using the Clear action on the RT, issue a global clean on all targets, while still need a special clear 
-	// for the PrimitiveId buffer
-	// const ERenderTargetLoadAction LoadAction = GHairClearVisibilityBuffer ? ERenderTargetLoadAction::EClear : ERenderTargetLoadAction::ENoAction;
-	ERenderTargetLoadAction LoadAction = ERenderTargetLoadAction::ENoAction;
-	if (GHairClearVisibilityBuffer)
-	{
-		LoadAction = ERenderTargetLoadAction::ELoad;
-		AddClearGraphicPass(GraphBuilder, RDG_EVENT_NAME("HairStrandsClearVisibilityMSAAMaterial"), ViewInfo, 0, OutVisibilityMaterialTexture);
-		AddClearGraphicPass(GraphBuilder, RDG_EVENT_NAME("HairStrandsClearVisibilityMSAAAttribute"), ViewInfo, 0, OutVisibilityAttributeTexture);
-		AddClearGraphicPass(GraphBuilder, RDG_EVENT_NAME("HairStrandsClearVisibilityMSAAVelocity"), ViewInfo, 0, OutVisibilityVelocityTexture);
-	}
+		{
+			FRDGTextureDesc Desc;
+			Desc.Extent.X = Resolution.X;
+			Desc.Extent.Y = Resolution.Y;
+			Desc.Depth = 0;
+			Desc.Format = PF_G16R16;
+			Desc.NumMips = 1;
+			Desc.NumSamples = MSAASampleCount;
+			Desc.Flags = TexCreate_None;
+			Desc.TargetableFlags = TexCreate_RenderTargetable | TexCreate_ShaderResource;
+			Desc.ClearValue = FClearValueBinding(FLinearColor(0, 0, 0, 0));
+			Desc.bForceSharedTargetAndShaderResource = true;
+			OutVisibilityVelocityTexture = GraphBuilder.CreateTexture(Desc, TEXT("HairVisibilityVelocityTexture"));
+		}
+		AddClearGraphicPass(GraphBuilder, RDG_EVENT_NAME("HairStrandsClearVisibilityMSAAIdTexture"), ViewInfo, 0xFFFFFFFF, OutVisibilityIdTexture);
 
-	FVisibilityPassParameters* PassParameters = GraphBuilder.AllocParameters<FVisibilityPassParameters>();
-	PassParameters->RenderTargets[0] = FRenderTargetBinding(OutVisibilityIdTexture, ERenderTargetLoadAction::ELoad, 0);
-	PassParameters->RenderTargets[1] = FRenderTargetBinding(OutVisibilityMaterialTexture,  LoadAction, 0);
-	PassParameters->RenderTargets[2] = FRenderTargetBinding(OutVisibilityAttributeTexture, LoadAction, 0);
-	PassParameters->RenderTargets[3] = FRenderTargetBinding(OutVisibilityVelocityTexture,  LoadAction, 0);
+		// Manually clear RTs as using the Clear action on the RT, issue a global clean on all targets, while still need a special clear 
+		// for the PrimitiveId buffer
+		// const ERenderTargetLoadAction LoadAction = GHairClearVisibilityBuffer ? ERenderTargetLoadAction::EClear : ERenderTargetLoadAction::ENoAction;
+		ERenderTargetLoadAction LoadAction = ERenderTargetLoadAction::ENoAction;
+		if (GHairClearVisibilityBuffer)
+		{
+			LoadAction = ERenderTargetLoadAction::ELoad;
+			AddClearGraphicPass(GraphBuilder, RDG_EVENT_NAME("HairStrandsClearVisibilityMSAAMaterial"), ViewInfo, 0, OutVisibilityMaterialTexture);
+			AddClearGraphicPass(GraphBuilder, RDG_EVENT_NAME("HairStrandsClearVisibilityMSAAAttribute"), ViewInfo, 0, OutVisibilityAttributeTexture);
+			AddClearGraphicPass(GraphBuilder, RDG_EVENT_NAME("HairStrandsClearVisibilityMSAAVelocity"), ViewInfo, 0, OutVisibilityVelocityTexture);
+		}
 
-	PassParameters->RenderTargets.DepthStencil = FDepthStencilBinding(
-		OutVisibilityDepthTexture,
-		ERenderTargetLoadAction::ELoad,
-		ERenderTargetLoadAction::ENoAction,
-		FExclusiveDepthStencil::DepthWrite_StencilNop);
-	AddHairVisibilityCommonPass(GraphBuilder, Scene, ViewInfo, ClusterDatas, HairVisibilityRenderMode_MSAA, PassParameters);
+		FVisibilityPassParameters* PassParameters = GraphBuilder.AllocParameters<FVisibilityPassParameters>();
+		PassParameters->RenderTargets[0] = FRenderTargetBinding(OutVisibilityIdTexture, ERenderTargetLoadAction::ELoad, 0);
+		PassParameters->RenderTargets[1] = FRenderTargetBinding(OutVisibilityMaterialTexture,  LoadAction, 0);
+		PassParameters->RenderTargets[2] = FRenderTargetBinding(OutVisibilityAttributeTexture, LoadAction, 0);
+		PassParameters->RenderTargets[3] = FRenderTargetBinding(OutVisibilityVelocityTexture,  LoadAction, 0);
+
+		PassParameters->RenderTargets.DepthStencil = FDepthStencilBinding(
+			OutVisibilityDepthTexture,
+			ERenderTargetLoadAction::ELoad,
+			ERenderTargetLoadAction::ENoAction,
+			FExclusiveDepthStencil::DepthWrite_StencilNop);
+		AddHairVisibilityCommonPass(GraphBuilder, Scene, ViewInfo, ClusterDatas, HairVisibilityRenderMode_MSAA, PassParameters);
+	}
 }
 
 static void AddHairVisibilityPPLLPass(
@@ -1295,6 +1851,8 @@ FHairStrandsVisibilityViews RenderHairStrandsVisibilityBuffer(
 
 			if (RenderMode == HairVisibilityRenderMode_MSAA)
 			{
+				const bool bIsVisiblityEnable = GHairStrandsVisibilityMaterialPass > 0;
+
 				struct FRDGMsaaVisibilityResources
 				{
 					FRDGTextureRef DepthTexture;
@@ -1312,6 +1870,7 @@ FHairStrandsVisibilityViews RenderHairStrandsVisibilityBuffer(
 					SceneDepthTexture);
 
 				AddHairVisibilityMSAAPass(
+					bIsVisiblityEnable,
 					GraphBuilder, 
 					Scene, 
 					&View, 
@@ -1325,10 +1884,13 @@ FHairStrandsVisibilityViews RenderHairStrandsVisibilityBuffer(
 
 				// This is used when compaction is not enabled.
 				GraphBuilder.QueueTextureExtraction(MsaaVisibilityResources.IdTexture, &VisibilityData.IDTexture);
-				GraphBuilder.QueueTextureExtraction(MsaaVisibilityResources.MaterialTexture, &VisibilityData.MaterialTexture);
-				GraphBuilder.QueueTextureExtraction(MsaaVisibilityResources.AttributeTexture, &VisibilityData.AttributeTexture);
-				GraphBuilder.QueueTextureExtraction(MsaaVisibilityResources.VelocityTexture, &VisibilityData.VelocityTexture);
 				GraphBuilder.QueueTextureExtraction(MsaaVisibilityResources.DepthTexture, &VisibilityData.DepthTexture);
+				if (!bIsVisiblityEnable)
+				{
+					GraphBuilder.QueueTextureExtraction(MsaaVisibilityResources.MaterialTexture, &VisibilityData.MaterialTexture);
+					GraphBuilder.QueueTextureExtraction(MsaaVisibilityResources.AttributeTexture, &VisibilityData.AttributeTexture);
+					GraphBuilder.QueueTextureExtraction(MsaaVisibilityResources.VelocityTexture, &VisibilityData.VelocityTexture);
+				}
 
 				FRDGTextureRef CategorizationTexture;
 				{
@@ -1346,6 +1908,7 @@ FHairStrandsVisibilityViews RenderHairStrandsVisibilityBuffer(
 					FRDGBufferRef IndirectArgsBuffer;
 					AddHairVisibilityPrimitiveIdCompactionPass(
 						false, // bUsePPLL
+						bIsVisiblityEnable,
 						GraphBuilder,
 						View,
 						ClusterDatas,
@@ -1357,6 +1920,35 @@ FHairStrandsVisibilityViews RenderHairStrandsVisibilityBuffer(
 						CategorizationTexture,
 						SceneVelocityTexture,
 						IndirectArgsBuffer);
+
+					if (bIsVisiblityEnable)
+					{
+						// Evaluate material based on the visiblity pass result
+						// Output both complete sample data + per-sample velocity
+						FMaterialPassOutput PassOutput = AddHairMaterialPass(
+							GraphBuilder,
+							Scene,
+							&View,
+							ClusterDatas,
+							VisibilityData.NodeGroupSize,
+							CompactNodeIndex,
+							CompactNodeData,
+							CompactNodeCoord,
+							IndirectArgsBuffer);
+
+						// Merge per-sample velocity into the scene velocity buffer
+						AddHairVelocityPass(
+							GraphBuilder,
+							View,
+							ClusterDatas,
+							CompactNodeIndex,
+							CompactNodeData,
+							PassOutput.NodeVelocity,
+							SceneVelocityTexture);
+
+						CompactNodeData = PassOutput.NodeData;
+					}
+
 					GraphBuilder.QueueTextureExtraction(CompactNodeIndex,		&VisibilityData.NodeIndex);
 					GraphBuilder.QueueTextureExtraction(CategorizationTexture,	&VisibilityData.CategorizationTexture);
 					GraphBuilder.QueueBufferExtraction(CompactNodeData,			&VisibilityData.NodeData,					FRDGResourceState::EAccess::Read, FRDGResourceState::EPipeline::Graphics);
@@ -1406,6 +1998,7 @@ FHairStrandsVisibilityViews RenderHairStrandsVisibilityBuffer(
 					FRDGBufferRef IndirectArgsBuffer;
 					AddHairVisibilityPrimitiveIdCompactionPass(
 						true, // bUsePPLL
+						false,
 						GraphBuilder,
 						View,
 						ClusterDatas,
diff --git a/Engine/Source/Runtime/Renderer/Public/MeshPassProcessor.h b/Engine/Source/Runtime/Renderer/Public/MeshPassProcessor.h
--- a/Engine/Source/Runtime/Renderer/Public/MeshPassProcessor.h
+++ b/Engine/Source/Runtime/Renderer/Public/MeshPassProcessor.h
@@ -1126,7 +1126,7 @@ public:
 		EMeshPassFeatures MeshPassFeatures,
 		const ShaderElementDataType& ShaderElementData);
 
-private:
+protected:
 	RENDERER_API void GetDrawCommandPrimitiveId(
 		const FPrimitiveSceneInfo* RESTRICT PrimitiveSceneInfo,
 		const FMeshBatchElement& BatchElement,
