Add permutation for computing hair count from the primary rays.
Use hair count texture for hair sky lighting integration (TT term, for getting more precise details)
Add opaque depth injection pass for the hair count texture (depth is injected if it is within a below a distance threshold to the closest hair).

diff --git a/Engine/Shaders/Private/HairStrands/HairStrandsEnvironmentLighting.usf b/Engine/Shaders/Private/HairStrands/HairStrandsEnvironmentLighting.usf
--- a/Engine/Shaders/Private/HairStrands/HairStrandsEnvironmentLighting.usf
+++ b/Engine/Shaders/Private/HairStrands/HairStrandsEnvironmentLighting.usf
@@ -75,6 +75,8 @@ Texture3D<float4>			HairEnergyLUTTexture;
 SamplerState 				HairLUTSampler;
 
 float HairDistanceThreshold;
+uint bHairUseViewHairCount;
+Texture2D<float> HairCountTexture;
 
 uint  MultipleScatterSampleCount;
 uint  JitterSphericalIntegration;
@@ -392,7 +394,11 @@ float3 ReflectionEnvironment(FGBufferData GBuffer, float2 ScreenPosition, float2
 			// For TT, we don't want to abort the hair counting, as this is important to get the proper transmission
 			const float DistanceThresholdTT = 100000;
 			const float DensityScaleFactor = 2;
-			const float HairCount = InternalComputeHairCountVoxel(SampleWorldPosition, L, Voxel_MinAABB, Voxel_MaxAABB, Voxel_Resolution, Voxel_DensityTexture, Voxel_DensityScale * DensityScaleFactor, 0, DistanceThresholdTT, false);
+			float HairCount = 0;
+			if (bHairUseViewHairCount)
+				HairCount = HairCountTexture.Load(uint3(SvPosition.xy, 0)) * DensityScaleFactor;
+			else
+				HairCount = InternalComputeHairCountVoxel(SampleWorldPosition, L, Voxel_MinAABB, Voxel_MaxAABB, Voxel_Resolution, Voxel_DensityTexture, Voxel_DensityScale * DensityScaleFactor, 0, DistanceThresholdTT, false);
 			Transmittance = ComputeTransmittanceFromHairCount(HairCount, GBuffer, SinLightAngle, HairScatteringLUTTexture, HairLUTSampler);
 		}
 
diff --git a/Engine/Shaders/Private/HairStrands/HairStrandsVisibilityCommon.ush b/Engine/Shaders/Private/HairStrands/HairStrandsVisibilityCommon.ush
--- a/Engine/Shaders/Private/HairStrands/HairStrandsVisibilityCommon.ush
+++ b/Engine/Shaders/Private/HairStrands/HairStrandsVisibilityCommon.ush
@@ -6,6 +6,7 @@
 #define RENDER_MODE_TRANSMITTANCE 1
 #define RENDER_MODE_PPLL 2
 #define RENDER_MODE_MSAA_VISIBILITY 3
+#define RENDER_MODE_TRANSMITTANCE_AND_HAIRCOUNT 4
 
 struct FVisibilityData
 {
diff --git a/Engine/Shaders/Private/HairStrands/HairStrandsVisibilityPS.usf b/Engine/Shaders/Private/HairStrands/HairStrandsVisibilityPS.usf
--- a/Engine/Shaders/Private/HairStrands/HairStrandsVisibilityPS.usf
+++ b/Engine/Shaders/Private/HairStrands/HairStrandsVisibilityPS.usf
@@ -23,7 +23,7 @@
 #define SUPPORT_MATERIAL_PROPERTY 1
 #endif
 
-#if HAIR_RENDER_MODE == RENDER_MODE_TRANSMITTANCE
+#if HAIR_RENDER_MODE == RENDER_MODE_TRANSMITTANCE || HAIR_RENDER_MODE == RENDER_MODE_TRANSMITTANCE_AND_HAIRCOUNT
 #define SUPPORT_MATERIAL_PROPERTY 0
 #endif
 
@@ -55,12 +55,15 @@ void MainVisibility(
 	, centroid in float WorldStrandRadius : TEXCOORD8
 	, const float4 PackedVelocityA : TEXCOORD9
 	, const float4 PackedVelocityC : TEXCOORD12
-#elif HAIR_RENDER_MODE == RENDER_MODE_TRANSMITTANCE 
+#elif HAIR_RENDER_MODE == RENDER_MODE_TRANSMITTANCE || HAIR_RENDER_MODE == RENDER_MODE_TRANSMITTANCE_AND_HAIRCOUNT
 	, centroid in float WorldStrandRadius : TEXCOORD8
 #endif
 	  OPTIONAL_IsFrontFace
 #if HAIR_RENDER_MODE == RENDER_MODE_TRANSMITTANCE
 	, out float OutColor0 : SV_Target0
+#elif HAIR_RENDER_MODE == RENDER_MODE_TRANSMITTANCE_AND_HAIRCOUNT
+	, out float OutColor0 : SV_Target0
+	, out float OutColor1 : SV_Target1
 #elif HAIR_RENDER_MODE == RENDER_MODE_MSAA
 	, out uint4 OutColor0	: SV_Target0
   #if SUPPORT_MATERIAL_PROPERTY
@@ -97,6 +100,9 @@ void MainVisibility(
 
 #if HAIR_RENDER_MODE == RENDER_MODE_TRANSMITTANCE
 	OutColor0 = saturate(1.0f - Coverage);
+#elif HAIR_RENDER_MODE == RENDER_MODE_TRANSMITTANCE_AND_HAIRCOUNT
+	OutColor0 = saturate(1.0f - Coverage);
+	OutColor1 = Coverage;
 #endif
 
 #if HAIR_RENDER_MODE == RENDER_MODE_MSAA || HAIR_RENDER_MODE == RENDER_MODE_PPLL
diff --git a/Engine/Shaders/Private/HairStrands/HairStrandsVisibilityTransmittanceDepthPS.usf b/Engine/Shaders/Private/HairStrands/HairStrandsVisibilityTransmittanceDepthPS.usf
new file mode 100644
--- /dev/null
+++ b/Engine/Shaders/Private/HairStrands/HairStrandsVisibilityTransmittanceDepthPS.usf
@@ -0,0 +1,27 @@
+// Copyright 1998-2019 Epic Games, Inc. All Rights Reserved.
+
+#include "../Common.ush"
+#include "../SceneTexturesCommon.ush"
+#include "HairStrandsVisibilityCommon.ush"
+
+Texture2D<float> SceneDepthTexture;
+Texture2D<uint4> CategorizationTexture;
+
+float DistanceThreshold;
+void MainPS(in FScreenVertexOutput Input, out float OutColor : SV_Target0)
+{
+	const uint3 PixelCoord = uint3(Input.Position.xy, 0);
+	const float FarDepth = 10e-06;
+
+	const float SceneDepth = ConvertFromDeviceZ(SceneDepthTexture.Load(PixelCoord));
+	const FCategorizationData CatData = DecodeCategorizationData(CategorizationTexture.Load(PixelCoord));
+	const float SampleDepth = ConvertFromDeviceZ(CatData.ClosestDepth);
+
+	if (abs(SampleDepth-SceneDepth) > DistanceThreshold)
+	{
+		discard;
+	}
+
+	const uint LargeHairCount = 100000;
+	OutColor = LargeHairCount;
+}
diff --git a/Engine/Source/Runtime/Renderer/Private/HairStrands/HairStrandsEnvironment.cpp b/Engine/Source/Runtime/Renderer/Private/HairStrands/HairStrandsEnvironment.cpp
--- a/Engine/Source/Runtime/Renderer/Private/HairStrands/HairStrandsEnvironment.cpp
+++ b/Engine/Source/Runtime/Renderer/Private/HairStrands/HairStrandsEnvironment.cpp
@@ -41,10 +41,8 @@ static FAutoConsoleVariableRef CVarHairStrandsSkyAOSampleCount(TEXT("r.HairStran
 
 static float GHairStrandsSkyAODistanceThreshold = 10;
 static float GHairStrandsSkyLightingDistanceThreshold = 10;
-static int32 GHairStrandsSkyLightingUseViewHairCount = 1;
 static FAutoConsoleVariableRef CVarHairStrandsSkyAOThreshold(TEXT("r.HairStrands.SkyAO.DistanceThreshold"), GHairStrandsSkyAODistanceThreshold, TEXT("Max distance for occlusion search."), ECVF_Scalability | ECVF_RenderThreadSafe);
 static FAutoConsoleVariableRef CVarHairStrandsSkyLightingDistanceThreshold(TEXT("r.HairStrands.SkyLighting.DistanceThreshold"), GHairStrandsSkyLightingDistanceThreshold, TEXT("Max distance for occlusion search."), ECVF_Scalability | ECVF_RenderThreadSafe);
-static FAutoConsoleVariableRef CVarHairStrandsSkyLightingUseViewHairCount(TEXT("r.HairStrands.SkyLighting.UseViewHairCount"), GHairStrandsSkyLightingUseViewHairCount, TEXT("Use the primary visibility hair count output to compute the TT term instead of the voxel grid."), ECVF_Scalability | ECVF_RenderThreadSafe);
 
 static bool GetHairStrandsSkyLightingEnable() { return GHairSkylightingEnable > 0; }
 static bool GetHairStrandsSkyAOEnable() { return GHairSkyAOEnable > 0; }
@@ -267,6 +265,7 @@ class FHairEnvironmentLightingCS : public FGlobalShader
 		SHADER_PARAMETER(uint32, HairComponents)
 		SHADER_PARAMETER(uint32, JitterSphericalIntegration)
 		SHADER_PARAMETER(float, HairDistanceThreshold)
+		SHADER_PARAMETER(uint32, bHairUseViewHairCount)
 
 		SHADER_PARAMETER_TEXTURE(Texture2D, PreIntegratedGF)
 		SHADER_PARAMETER_SAMPLER(SamplerState, PreIntegratedGFSampler)
@@ -277,6 +276,7 @@ class FHairEnvironmentLightingCS : public FGlobalShader
 		SHADER_PARAMETER_STRUCT_INCLUDE(FSceneTextureParameters, SceneTextures)
 		SHADER_PARAMETER_STRUCT_INCLUDE(FSceneTextureSamplerParameters, SceneTextureSamplers)
 
+		SHADER_PARAMETER_RDG_TEXTURE(Texture2D, HairCountTexture)
 		SHADER_PARAMETER_RDG_TEXTURE(TEXTURE3D, HairEnergyLUTTexture)
 		SHADER_PARAMETER_RDG_TEXTURE(TEXTURE3D, HairScatteringLUTTexture)
 		SHADER_PARAMETER_SAMPLER(SamplerState, HairLUTSampler)
@@ -329,12 +329,14 @@ static FRDGBufferRef AddHairStrandsEnvironmentLightingPassCS(
 	PassParameters->Voxel_DepthBiasScale = GetHairStrandsVoxelizationDepthBiasScale();
 	PassParameters->Voxel_TanConeAngle = FMath::Tan(FMath::DegreesToRadians(GetHairStrandsSkyLightingConeAngle()));
 	PassParameters->HairDistanceThreshold = FMath::Max(GHairStrandsSkyLightingDistanceThreshold, 1.f);
+	PassParameters->bHairUseViewHairCount = VisibilityData.ViewHairCountTexture ? 1 : 0;
 	PassParameters->MaxVisibilityNodeCount = VisibilityData.NodeData->Desc.NumElements;
 	PassParameters->MultipleScatterSampleCount = FMath::Max(uint32(GHairStrandsSkyLightingSampleCount), 1u);
 	PassParameters->JitterSphericalIntegration = GHairStrandsSkyLightingJitterSphericalIntegration ? 1 : 0;
 	PassParameters->HairComponents = ToBitfield(GetHairComponents());
 	PassParameters->PreIntegratedGF = GSystemTextures.PreintegratedGF->GetRenderTargetItem().ShaderResourceTexture;
 	PassParameters->PreIntegratedGFSampler = TStaticSamplerState<SF_Bilinear, AM_Clamp, AM_Clamp, AM_Clamp>::GetRHI();
+	PassParameters->HairCountTexture = GraphBuilder.RegisterExternalTexture(VisibilityData.ViewHairCountTexture ? VisibilityData.ViewHairCountTexture : GSystemTextures.BlackDummy);
 	PassParameters->SceneTextures = SceneTextures;
 	SetupSceneTextureSamplers(&PassParameters->SceneTextureSamplers);
 	PassParameters->ViewUniformBuffer = View.ViewUniformBuffer;
diff --git a/Engine/Source/Runtime/Renderer/Private/HairStrands/HairStrandsVisibility.cpp b/Engine/Source/Runtime/Renderer/Private/HairStrands/HairStrandsVisibility.cpp
--- a/Engine/Source/Runtime/Renderer/Private/HairStrands/HairStrandsVisibility.cpp
+++ b/Engine/Source/Runtime/Renderer/Private/HairStrands/HairStrandsVisibility.cpp
@@ -29,7 +29,6 @@ static float GHairStrandsMaterialCompactionTangentThreshold = 10.f;
 static FAutoConsoleVariableRef CVarHairStrandsMaterialCompactionDepthThreshold(TEXT("r.HairStrands.MaterialCompaction.DepthThreshold"), GHairStrandsMaterialCompactionDepthThreshold, TEXT("Compaction threshold for depth value for material compaction (in centimeters). Default 1 cm."));
 static FAutoConsoleVariableRef CVarHairStrandsMaterialCompactionTangentThreshold(TEXT("r.HairStrands.MaterialCompaction.TangentThreshold"), GHairStrandsMaterialCompactionTangentThreshold, TEXT("Compaciton threshold for tangent value for material compaction (in degrees). Default 10 deg."));
 
-
 static int32 GHairVisibilitySampleCount = 8;
 static FAutoConsoleVariableRef CVarHairVisibilitySampleCount(TEXT("r.HairStrands.VisibilitySampleCount"), GHairVisibilitySampleCount, TEXT("Hair strands visibility sample count"));
 
@@ -54,6 +53,11 @@ static FAutoConsoleVariableRef CVarGHairVisibilityPPLLMeanNodeCountPerPixel(TEXT
 static int32 GHairStrandsVisibilityMaterialPass = 0;
 static FAutoConsoleVariableRef CVarHairStrandsVisibilityMaterialPass(TEXT("r.HairStrands.Visibility.MaterialPass"), GHairStrandsVisibilityMaterialPass, TEXT("Enable the deferred material pass evaluation after the hair visibility is resolved."));
 
+static int32 GHairStrandsViewHairCount = 0;
+static float GHairStrandsViewHairCountDepthDistanceThreshold = 30.f;
+static FAutoConsoleVariableRef CVarHairStrandsViewHairCount(TEXT("r.HairStrands.Visibility.HairCount"), GHairStrandsViewHairCount, TEXT("Enable the computation of 'view-hair-count' during the transmission pass."));
+static FAutoConsoleVariableRef CVarHairStrandsViewHairCountDepthDistanceThreshold(TEXT("r.HairStrands.Visibility.HairCount.DistanceThreshold"), GHairStrandsViewHairCountDepthDistanceThreshold, TEXT("Distance threshold defining if opaque depth get injected into the 'view-hair-count' buffer."));
+
 /////////////////////////////////////////////////////////////////////////////////////////
 
 namespace HairStrandsVisibilityInternal
@@ -83,6 +87,7 @@ enum EHairVisibilityRenderMode
 	HairVisibilityRenderMode_Transmittance,
 	HairVisibilityRenderMode_PPLL,
 	HairVisibilityRenderMode_MSAA_Visibility,
+	HairVisibilityRenderMode_TransmittanceAndHairCount,
 	HairVisibilityRenderModeCount
 };
 
@@ -649,6 +654,7 @@ protected:
 IMPLEMENT_MATERIAL_SHADER_TYPE(template<>, FHairVisibilityVS<HairVisibilityRenderMode_MSAA_Visibility>, TEXT("/Engine/Private/HairStrands/HairStrandsVisibilityVS.usf"), TEXT("Main"), SF_Vertex);
 IMPLEMENT_MATERIAL_SHADER_TYPE(template<>, FHairVisibilityVS<HairVisibilityRenderMode_MSAA>, TEXT("/Engine/Private/HairStrands/HairStrandsVisibilityVS.usf"), TEXT("Main"), SF_Vertex);
 IMPLEMENT_MATERIAL_SHADER_TYPE(template<>, FHairVisibilityVS<HairVisibilityRenderMode_Transmittance>, TEXT("/Engine/Private/HairStrands/HairStrandsVisibilityVS.usf"), TEXT("Main"), SF_Vertex);
+IMPLEMENT_MATERIAL_SHADER_TYPE(template<>, FHairVisibilityVS<HairVisibilityRenderMode_TransmittanceAndHairCount>, TEXT("/Engine/Private/HairStrands/HairStrandsVisibilityVS.usf"), TEXT("Main"), SF_Vertex);
 IMPLEMENT_MATERIAL_SHADER_TYPE(template<>, FHairVisibilityVS<HairVisibilityRenderMode_PPLL>, TEXT("/Engine/Private/HairStrands/HairStrandsVisibilityVS.usf"), TEXT("Main"), SF_Vertex);
 
 /////////////////////////////////////////////////////////////////////////////////////////
@@ -736,6 +742,7 @@ public:
 IMPLEMENT_MATERIAL_SHADER_TYPE(template<>, FHairVisibilityPS<HairVisibilityRenderMode_MSAA_Visibility>, TEXT("/Engine/Private/HairStrands/HairStrandsVisibilityPS.usf"), TEXT("MainVisibility"), SF_Pixel);
 IMPLEMENT_MATERIAL_SHADER_TYPE(template<>, FHairVisibilityPS<HairVisibilityRenderMode_MSAA>, TEXT("/Engine/Private/HairStrands/HairStrandsVisibilityPS.usf"), TEXT("MainVisibility"), SF_Pixel);
 IMPLEMENT_MATERIAL_SHADER_TYPE(template<>, FHairVisibilityPS<HairVisibilityRenderMode_Transmittance>, TEXT("/Engine/Private/HairStrands/HairStrandsVisibilityPS.usf"), TEXT("MainVisibility"), SF_Pixel);
+IMPLEMENT_MATERIAL_SHADER_TYPE(template<>, FHairVisibilityPS<HairVisibilityRenderMode_TransmittanceAndHairCount>, TEXT("/Engine/Private/HairStrands/HairStrandsVisibilityPS.usf"), TEXT("MainVisibility"), SF_Pixel);
 IMPLEMENT_MATERIAL_SHADER_TYPE(template<>, FHairVisibilityPS<HairVisibilityRenderMode_PPLL>, TEXT("/Engine/Private/HairStrands/HairStrandsVisibilityPS.usf"), TEXT("MainVisibility"), SF_Pixel);
 
 /////////////////////////////////////////////////////////////////////////////////////////
@@ -803,6 +810,8 @@ void FHairVisibilityProcessor::AddMeshBatch(const FMeshBatch& RESTRICT MeshBatch
 			Process<HairVisibilityRenderMode_MSAA>(MeshBatch, BatchElementMask, PrimitiveSceneProxy, StaticMeshId, MaterialRenderProxy, Material, HairClusterId, HairMaterialId, MeshFillMode, MeshCullMode);
 		if (RenderMode == HairVisibilityRenderMode_Transmittance)
 			Process<HairVisibilityRenderMode_Transmittance>(MeshBatch, BatchElementMask, PrimitiveSceneProxy, StaticMeshId, MaterialRenderProxy, Material, HairClusterId, HairMaterialId, MeshFillMode, MeshCullMode);
+		if (RenderMode == HairVisibilityRenderMode_TransmittanceAndHairCount)
+			Process<HairVisibilityRenderMode_TransmittanceAndHairCount>(MeshBatch, BatchElementMask, PrimitiveSceneProxy, StaticMeshId, MaterialRenderProxy, Material, HairClusterId, HairMaterialId, MeshFillMode, MeshCullMode);
 		if (RenderMode == HairVisibilityRenderMode_PPLL)
 			Process<HairVisibilityRenderMode_PPLL>(MeshBatch, BatchElementMask, PrimitiveSceneProxy, StaticMeshId, MaterialRenderProxy, Material, HairClusterId, HairMaterialId, MeshFillMode, MeshCullMode);
 	}
@@ -1370,11 +1379,12 @@ static void AddHairVisibilityCommonPass(
 	{
 		switch (RenderMode)
 		{
-		case HairVisibilityRenderMode_PPLL:				return RDG_EVENT_NAME("HairStrandsVisibilityPPLLPass");
-		case HairVisibilityRenderMode_MSAA:				return RDG_EVENT_NAME("HairStrandsVisibilityMSAAPass");
-		case HairVisibilityRenderMode_MSAA_Visibility:	return RDG_EVENT_NAME("HairStrandsVisibilityMSAAVisPass");
-		case HairVisibilityRenderMode_Transmittance:	return RDG_EVENT_NAME("HairStrandsTransmittancePass");
-		default:										return RDG_EVENT_NAME("Noname");
+		case HairVisibilityRenderMode_PPLL:						return RDG_EVENT_NAME("HairStrandsVisibilityPPLLPass");
+		case HairVisibilityRenderMode_MSAA:						return RDG_EVENT_NAME("HairStrandsVisibilityMSAAPass");
+		case HairVisibilityRenderMode_MSAA_Visibility:			return RDG_EVENT_NAME("HairStrandsVisibilityMSAAVisPass");
+		case HairVisibilityRenderMode_Transmittance:			return RDG_EVENT_NAME("HairStrandsTransmittancePass");
+		case HairVisibilityRenderMode_TransmittanceAndHairCount:return RDG_EVENT_NAME("HairStrandsTransmittanceAndHairCountPass");
+		default:												return RDG_EVENT_NAME("Noname");
 		}
 	};
 
@@ -1389,7 +1399,7 @@ static void AddHairVisibilityCommonPass(
 
 		FMeshPassProcessorRenderState DrawRenderState(*ViewInfo);
 
-		if(RenderMode == HairVisibilityRenderMode_Transmittance || RenderMode == HairVisibilityRenderMode_PPLL)
+		if(RenderMode == HairVisibilityRenderMode_Transmittance || RenderMode == HairVisibilityRenderMode_TransmittanceAndHairCount || RenderMode == HairVisibilityRenderMode_PPLL)
 		{
 			TUniformBufferRef<FViewUniformShaderParameters> ViewUniformShaderParameters;
 			const bool bEnableMSAA = false;
@@ -1418,6 +1428,13 @@ static void AddHairVisibilityCommonPass(
 				DrawRenderState.SetBlendState(TStaticBlendState<CW_RED, BO_Add, BF_DestColor, BF_Zero, BO_Add, BF_Zero, BF_Zero>::GetRHI());
 				DrawRenderState.SetDepthStencilState(TStaticDepthStencilState<false, CF_DepthNearOrEqual>::GetRHI());
 			}
+			else if (RenderMode == HairVisibilityRenderMode_TransmittanceAndHairCount)
+			{
+				DrawRenderState.SetBlendState(TStaticBlendState<
+					CW_RED, BO_Add, BF_DestColor, BF_Zero, BO_Add, BF_Zero, BF_Zero,
+					CW_RED, BO_Add, BF_One, BF_One, BO_Add, BF_Zero, BF_Zero>::GetRHI());
+				DrawRenderState.SetDepthStencilState(TStaticDepthStencilState<false, CF_DepthNearOrEqual>::GetRHI());
+			}
 			else if (RenderMode == HairVisibilityRenderMode_PPLL)
 			{
 				DrawRenderState.SetBlendState(TStaticBlendState<>::GetRHI());
@@ -1659,45 +1676,141 @@ static void AddHairVisibilityPPLLPass(
 
 ///////////////////////////////////////////////////////////////////////////////////////////////////
 
-static FRDGTextureRef AddHairViewTransmittancePass(
+struct FHairPrimaryTransmittance
+{
+	FRDGTextureRef TransmittanceTexture = nullptr;
+	FRDGTextureRef HairCountTexture = nullptr;
+};
+
+static FHairPrimaryTransmittance AddHairViewTransmittancePass(
 	FRDGBuilder& GraphBuilder,
 	const FScene* Scene,
 	const FViewInfo* ViewInfo,
 	const FHairStrandsClusterDatas& ClusterDatas,
 	const FIntPoint& Resolution,
+	const bool bOutputHairCount,
 	FRDGTextureRef SceneDepthTexture)
 {
 	check(SceneDepthTexture->Desc.Extent == Resolution);
+	const EHairVisibilityRenderMode RenderMode = bOutputHairCount ? HairVisibilityRenderMode_TransmittanceAndHairCount : HairVisibilityRenderMode_Transmittance;
+
+	FRDGTextureDesc Desc;
+	Desc.Extent.X = Resolution.X;
+	Desc.Extent.Y = Resolution.Y;
+	Desc.Depth = 0;
+	Desc.Format = PF_R32_FLOAT;
+	Desc.NumMips = 1;
+	Desc.NumSamples = 1;
+	Desc.Flags = TexCreate_None;
+	Desc.TargetableFlags = TexCreate_RenderTargetable | TexCreate_ShaderResource;
+	Desc.bForceSharedTargetAndShaderResource = true;
+	Desc.ClearValue = FClearValueBinding(FLinearColor(1.0f, 1.0f, 1.0f, 1.0f)); // Clear to transmittance 1
+
+	FVisibilityPassParameters* PassParameters = GraphBuilder.AllocParameters<FVisibilityPassParameters>();
+	FHairPrimaryTransmittance Out;
 
-	FRDGTextureRef ViewTransmittanceTexture = nullptr;
+	Out.TransmittanceTexture = GraphBuilder.CreateTexture(Desc, TEXT("HairViewTransmittanceTexture"));
+	PassParameters->RenderTargets[0] = FRenderTargetBinding(Out.TransmittanceTexture, ERenderTargetLoadAction::EClear, 0);
+
+	if (RenderMode == HairVisibilityRenderMode_TransmittanceAndHairCount)
 	{
-		FRDGTextureDesc Desc;
-		Desc.Extent.X = Resolution.X;
-		Desc.Extent.Y = Resolution.Y;
-		Desc.Depth = 0;
-		Desc.Format = PF_R32_FLOAT;
-		Desc.NumMips = 1;
-		Desc.NumSamples = 1;
-		Desc.Flags = TexCreate_None;
-		Desc.TargetableFlags = TexCreate_RenderTargetable | TexCreate_ShaderResource;
-		Desc.bForceSharedTargetAndShaderResource = true;
-		Desc.ClearValue = FClearValueBinding(FLinearColor(1.0f, 1.0f, 1.0f, 1.0f)); // Clear to transmittance =1
-		ViewTransmittanceTexture = GraphBuilder.CreateTexture(Desc, TEXT("HairViewTransmittanceTexture"));
+		Desc.ClearValue = FClearValueBinding(FLinearColor(0.0f, 0.0f, 0.0f, 0.0f));
+		Out.HairCountTexture = GraphBuilder.CreateTexture(Desc, TEXT("HairViewHairCountTexture"));
+		PassParameters->RenderTargets[1] = FRenderTargetBinding(Out.HairCountTexture, ERenderTargetLoadAction::EClear, 0);
 	}
 
-	FVisibilityPassParameters* PassParameters = GraphBuilder.AllocParameters<FVisibilityPassParameters>();
-	PassParameters->RenderTargets[0] = FRenderTargetBinding(ViewTransmittanceTexture, ERenderTargetLoadAction::EClear, 0);
-
 	PassParameters->RenderTargets.DepthStencil = FDepthStencilBinding(
 		SceneDepthTexture,
 		ERenderTargetLoadAction::ELoad,
 		ERenderTargetLoadAction::ENoAction,
 		FExclusiveDepthStencil::DepthRead_StencilNop);
-	AddHairVisibilityCommonPass(GraphBuilder, Scene, ViewInfo, ClusterDatas, HairVisibilityRenderMode_Transmittance, PassParameters);
+	AddHairVisibilityCommonPass(GraphBuilder, Scene, ViewInfo, ClusterDatas, RenderMode, PassParameters);
+
+	return Out;
+}
+
+///////////////////////////////////////////////////////////////////////////////////////////////////
+// Inject depth information into the view hair count texture, to block opaque occluder
+class FHairViewTransmittanceDepthPS : public FGlobalShader
+{
+	DECLARE_GLOBAL_SHADER(FHairViewTransmittanceDepthPS);
+	SHADER_USE_PARAMETER_STRUCT(FHairViewTransmittanceDepthPS, FGlobalShader);
+
+	BEGIN_SHADER_PARAMETER_STRUCT(FParameters, )
+		SHADER_PARAMETER(float, DistanceThreshold)
+		SHADER_PARAMETER_RDG_TEXTURE(Texture2D, SceneDepthTexture)
+		SHADER_PARAMETER_RDG_TEXTURE(Texture2D, CategorizationTexture)		
+		SHADER_PARAMETER_STRUCT_REF(FViewUniformShaderParameters, ViewUniformBuffer)
+		RENDER_TARGET_BINDING_SLOTS()
+	END_SHADER_PARAMETER_STRUCT()
+
+public:
+
+	static bool ShouldCompilePermutation(const FGlobalShaderPermutationParameters& Parameters)
+	{
+		return IsHairStrandsSupported(Parameters.Platform);
+	}
+};
+
+IMPLEMENT_GLOBAL_SHADER(FHairViewTransmittanceDepthPS, "/Engine/Private/HairStrands/HairStrandsVisibilityTransmittanceDepthPS.usf", "MainPS", SF_Pixel);
+
+static void AddHairViewTransmittanceDepthPass(
+	FRDGBuilder& GraphBuilder,
+	const FViewInfo& View,
+	const FRDGTextureRef& CategorizationTexture,
+	const FRDGTextureRef& SceneDepthTexture,
+	FRDGTextureRef& HairCountTexture)
+{
+	FHairViewTransmittanceDepthPS::FParameters* Parameters = GraphBuilder.AllocParameters<FHairViewTransmittanceDepthPS::FParameters>();
+	Parameters->DistanceThreshold = FMath::Max(1.f, GHairStrandsViewHairCountDepthDistanceThreshold);
+	Parameters->CategorizationTexture = CategorizationTexture;
+	Parameters->SceneDepthTexture = SceneDepthTexture;
+	Parameters->ViewUniformBuffer = View.ViewUniformBuffer;
+	Parameters->RenderTargets[0] = FRenderTargetBinding(HairCountTexture, ERenderTargetLoadAction::ELoad);
+	
+	TShaderMapRef<FPostProcessVS> VertexShader(View.ShaderMap);
+	TShaderMapRef<FHairViewTransmittanceDepthPS> PixelShader(View.ShaderMap);
+	const TShaderMap<FGlobalShaderType>* GlobalShaderMap = View.ShaderMap;
+	const FIntRect Viewport = View.ViewRect;
+	const FIntPoint Resolution = HairCountTexture->Desc.Extent;
+	const FViewInfo* CapturedView = &View;
+	ClearUnusedGraphResources(*PixelShader, Parameters);
 
-	return ViewTransmittanceTexture;
+	GraphBuilder.AddPass(
+		RDG_EVENT_NAME("HairStrandsViewTransmittanceDepth"),
+		Parameters,
+		ERDGPassFlags::Raster,
+		[Parameters, VertexShader, PixelShader, Viewport, Resolution, CapturedView](FRHICommandList& RHICmdList)
+	{
+		FGraphicsPipelineStateInitializer GraphicsPSOInit;
+		RHICmdList.ApplyCachedRenderTargets(GraphicsPSOInit);
+		GraphicsPSOInit.BlendState = TStaticBlendState<CW_RGBA, BO_Add, BF_One, BF_One, BO_Add, BF_Zero, BF_Zero>::GetRHI();
+		GraphicsPSOInit.RasterizerState = TStaticRasterizerState<>::GetRHI();
+		GraphicsPSOInit.DepthStencilState = TStaticDepthStencilState<false, CF_Always>::GetRHI();
+
+		GraphicsPSOInit.BoundShaderState.VertexDeclarationRHI = GFilterVertexDeclaration.VertexDeclarationRHI;
+		GraphicsPSOInit.BoundShaderState.VertexShaderRHI = GETSAFERHISHADER_VERTEX(*VertexShader);
+		GraphicsPSOInit.BoundShaderState.PixelShaderRHI = GETSAFERHISHADER_PIXEL(*PixelShader);
+		GraphicsPSOInit.PrimitiveType = PT_TriangleList;
+		SetGraphicsPipelineState(RHICmdList, GraphicsPSOInit);
+
+		VertexShader->SetParameters(RHICmdList, CapturedView->ViewUniformBuffer);
+		RHICmdList.SetViewport(Viewport.Min.X, Viewport.Min.Y, 0.0f, Viewport.Max.X, Viewport.Max.Y, 1.0f);
+		SetShaderParameters(RHICmdList, *PixelShader, PixelShader->GetPixelShader(), *Parameters);
+		DrawRectangle(
+			RHICmdList,
+			0, 0,
+			Viewport.Width(), Viewport.Height(),
+			Viewport.Min.X, Viewport.Min.Y,
+			Viewport.Width(), Viewport.Height(),
+			Viewport.Size(),
+			Resolution,
+			*VertexShader,
+			EDRF_UseTriangleOptimization);
+	});
 }
 
+
 ///////////////////////////////////////////////////////////////////////////////////////////////////
 class FHairVisibilityDepthPS : public FGlobalShader
 {
@@ -1837,15 +1950,17 @@ FHairStrandsVisibilityViews RenderHairStrandsVisibilityBuffer(
 			const EHairVisibilityRenderMode RenderMode = GetHairVisibilityRenderMode();
 
 			// Run the view transmittance pass if needed (not in PPLL mode that is already a high quality render path)
-			FRDGTextureRef ViewTransmittanceTexture = nullptr;
+			FHairPrimaryTransmittance ViewTransmittance;
 			if (GHairStrandsViewTransmittancePassEnable > 0 && RenderMode != HairVisibilityRenderMode_PPLL)
 			{
-				ViewTransmittanceTexture = AddHairViewTransmittancePass(
+				const bool bOutputHairCount = GHairStrandsViewHairCount > 0;
+				ViewTransmittance = AddHairViewTransmittancePass(
 					GraphBuilder,
 					Scene,
 					&View,
 					ClusterDatas,
 					Resolution, 
+					bOutputHairCount,
 					SceneDepthTexture);
 			}
 
@@ -1900,7 +2015,7 @@ FHairStrandsVisibilityViews RenderHairStrandsVisibilityBuffer(
 					PassParameters->MSAA_MaterialTexture = MsaaVisibilityResources.MaterialTexture;
 					PassParameters->MSAA_AttributeTexture = MsaaVisibilityResources.AttributeTexture;
 					PassParameters->MSAA_VelocityTexture = MsaaVisibilityResources.VelocityTexture;
-					PassParameters->ViewTransmittanceTexture = ViewTransmittanceTexture;
+					PassParameters->ViewTransmittanceTexture = ViewTransmittance.TransmittanceTexture;
 
 					FRDGTextureRef CompactNodeIndex;
 					FRDGBufferRef CompactNodeData;
@@ -1967,6 +2082,17 @@ FHairStrandsVisibilityViews RenderHairStrandsVisibilityBuffer(
 					SceneGBufferBTexture,
 					SceneColorTexture,
 					SceneDepthTexture);
+
+				if (ViewTransmittance.HairCountTexture)
+				{
+					AddHairViewTransmittanceDepthPass(
+						GraphBuilder,
+						View,
+						CategorizationTexture,
+						SceneDepthTexture,
+						ViewTransmittance.HairCountTexture);
+					GraphBuilder.QueueTextureExtraction(ViewTransmittance.HairCountTexture, &VisibilityData.ViewHairCountTexture);
+				}
 			}
 			else if (RenderMode == HairVisibilityRenderMode_PPLL)
 			{
@@ -1990,7 +2116,7 @@ FHairStrandsVisibilityViews RenderHairStrandsVisibilityBuffer(
 					PassParameters->PPLLCounter  = PPLLNodeCounterTexture;
 					PassParameters->PPLLNodeIndex= PPLLNodeIndexTexture;
 					PassParameters->PPLLNodeData = GraphBuilder.CreateSRV(PPLLNodeDataBuffer);
-					PassParameters->ViewTransmittanceTexture = ViewTransmittanceTexture;
+					PassParameters->ViewTransmittanceTexture = ViewTransmittance.TransmittanceTexture;
 
 					FRDGTextureRef CompactNodeIndex;
 					FRDGBufferRef CompactNodeData;
diff --git a/Engine/Source/Runtime/Renderer/Private/HairStrands/HairStrandsVisibility.h b/Engine/Source/Runtime/Renderer/Private/HairStrands/HairStrandsVisibility.h
--- a/Engine/Source/Runtime/Renderer/Private/HairStrands/HairStrandsVisibility.h
+++ b/Engine/Source/Runtime/Renderer/Private/HairStrands/HairStrandsVisibility.h
@@ -22,6 +22,7 @@ struct FHairStrandsVisibilityData
 	TRefCountPtr<IPooledRenderTarget> AttributeTexture;
 	TRefCountPtr<IPooledRenderTarget> VelocityTexture;
 	TRefCountPtr<IPooledRenderTarget> CategorizationTexture;
+	TRefCountPtr<IPooledRenderTarget> ViewHairCountTexture;
 
 	TRefCountPtr<IPooledRenderTarget> PPLLNodeCounterTexture;
 	TRefCountPtr<IPooledRenderTarget> PPLLNodeIndexTexture;
