Change courtesy of [REDACTED_PERSON]
Fix a sample out-of-bound issue in single layer water pass

diff --git a/Engine/Source/Runtime/Renderer/Private/SingleLayerWaterRendering.cpp b/Engine/Source/Runtime/Renderer/Private/SingleLayerWaterRendering.cpp
--- a/Engine/Source/Runtime/Renderer/Private/SingleLayerWaterRendering.cpp
+++ b/Engine/Source/Runtime/Renderer/Private/SingleLayerWaterRendering.cpp
@@ -488,8 +488,8 @@ void FDeferredShadingSceneRenderer::CopySingleLayerWaterTextures(FRHICommandList
     
 		const FIntRect RefractionViewRect = FIntRect(FIntPoint::DivideAndRoundDown(View.ViewRect.Min, RefractionDownsampleFactor), FIntPoint::DivideAndRoundDown(View.ViewRect.Max, RefractionDownsampleFactor));
 
-		PassData.ViewData[ViewIndex].SceneWithoutSingleLayerWaterMaxUV.X = (RefractionViewRect.Size().X - 0.5f) / RefractionViewRect.Size().X;
-		PassData.ViewData[ViewIndex].SceneWithoutSingleLayerWaterMaxUV.Y = (RefractionViewRect.Size().Y - 0.5f) / RefractionViewRect.Size().Y;
+		PassData.ViewData[ViewIndex].SceneWithoutSingleLayerWaterMaxUV.X = (RefractionViewRect.Size().X - 0.5f) / RefractionResolution.X;
+		PassData.ViewData[ViewIndex].SceneWithoutSingleLayerWaterMaxUV.Y = (RefractionViewRect.Size().Y - 0.5f) / RefractionResolution.Y;
 
 	    FPixelShaderUtils::AddFullscreenPass(
 		    GraphBuilder,
