Fixes static analysis in SSRT

diff --git a/Engine/Source/Runtime/Renderer/Private/ScreenSpaceRayTracing.cpp b/Engine/Source/Runtime/Renderer/Private/ScreenSpaceRayTracing.cpp
--- a/Engine/Source/Runtime/Renderer/Private/ScreenSpaceRayTracing.cpp
+++ b/Engine/Source/Runtime/Renderer/Private/ScreenSpaceRayTracing.cpp
@@ -992,6 +992,7 @@ void RenderScreenSpaceDiffuseIndirect(
 			PassParameters->HigherMipTexture = GraphBuilder.CreateSRV(FRDGTextureSRVDesc::CreateForMipLevel(ReducedSceneColor, SrcMip));
 			if (bUseLeakFree)
 			{
+				check(ReducedSceneAlpha);
 				PassParameters->HigherMipTextureSampler = TStaticSamplerState<SF_Point>::GetRHI();
 				PassParameters->HigherAlphaMipTexture = GraphBuilder.CreateSRV(FRDGTextureSRVDesc::CreateForMipLevel(ReducedSceneAlpha, SrcMip));
 				PassParameters->HigherAlphaMipTextureSampler = TStaticSamplerState<SF_Point>::GetRHI();
