Fix shadowing error
#rnx

diff --git a/Engine/Source/Runtime/Renderer/Private/HairStrands/HairStrandsVisibility.cpp b/Engine/Source/Runtime/Renderer/Private/HairStrands/HairStrandsVisibility.cpp
--- a/Engine/Source/Runtime/Renderer/Private/HairStrands/HairStrandsVisibility.cpp
+++ b/Engine/Source/Runtime/Renderer/Private/HairStrands/HairStrandsVisibility.cpp
@@ -1375,7 +1375,7 @@ static void AddHairVisibilityCommonPass(
 	TPassParameters* PassParameters)
 {
 
-	auto GetPassName = [](EHairVisibilityRenderMode RenderMode)
+	auto GetPassName = [RenderMode]()
 	{
 		switch (RenderMode)
 		{
@@ -1389,7 +1389,7 @@ static void AddHairVisibilityCommonPass(
 	};
 
 	GraphBuilder.AddPass(
-		GetPassName(RenderMode), 
+		GetPassName(), 
 		PassParameters,
 		ERDGPassFlags::Raster,
 		[PassParameters, Scene = Scene, ViewInfo, &ClusterDatas, RenderMode](FRHICommandListImmediate& RHICmdList)
