Ray tracing skylight samples per pixel minimum value clamped to 2. Default value changed to 4.

diff --git a/Engine/Source/Runtime/Engine/Private/Components/SkyLightComponent.cpp b/Engine/Source/Runtime/Engine/Private/Components/SkyLightComponent.cpp
--- a/Engine/Source/Runtime/Engine/Private/Components/SkyLightComponent.cpp
+++ b/Engine/Source/Runtime/Engine/Private/Components/SkyLightComponent.cpp
@@ -240,7 +240,7 @@ USkyLightComponent::USkyLightComponent(const FObjectInitializer& ObjectInitializ
 	bCastRaytracedShadow = false;
 	bAffectReflection = true;
 	bAffectGlobalIllumination = true;
-	SamplesPerPixel = 1;
+	SamplesPerPixel = 4;
 }
 
 #if RHI_RAYTRACING
diff --git a/Engine/Source/Runtime/Renderer/Private/RayTracing/RaytracingSkylight.cpp b/Engine/Source/Runtime/Renderer/Private/RayTracing/RaytracingSkylight.cpp
--- a/Engine/Source/Runtime/Renderer/Private/RayTracing/RaytracingSkylight.cpp
+++ b/Engine/Source/Runtime/Renderer/Private/RayTracing/RaytracingSkylight.cpp
@@ -864,7 +864,7 @@ void FDeferredShadingSceneRenderer::GenerateSkyLightVisibilityRays(
 	// SkyLight data setup
 	FSkyLightData SkyLightData;
 	SetupSkyLightParameters(*Scene, &SkyLightData);
-	int32 SamplesPerPixel = GRayTracingSkyLightSamplesPerPixel >= 0 ? GRayTracingSkyLightSamplesPerPixel : Scene->SkyLight->SamplesPerPixel;
+ 	int32 SamplesPerPixel = GRayTracingSkyLightSamplesPerPixel >= 0 ? GRayTracingSkyLightSamplesPerPixel : FMath::Max(Scene->SkyLight->SamplesPerPixel, 2);
 	ensure(SamplesPerPixel > 0);
 	SkyLightData.SamplesPerPixel = SamplesPerPixel;
 	SkyLightData.MaxRayDistance = GRayTracingSkyLightMaxRayDistance;
@@ -926,7 +926,7 @@ void FDeferredShadingSceneRenderer::RenderRayTracingSkyLight(
 		return;
 	}
 
-	int32 SamplesPerPixel = GRayTracingSkyLightSamplesPerPixel >= 0 ? GRayTracingSkyLightSamplesPerPixel : Scene->SkyLight->SamplesPerPixel;
+	int32 SamplesPerPixel = GRayTracingSkyLightSamplesPerPixel >= 0 ? GRayTracingSkyLightSamplesPerPixel : FMath::Max(Scene->SkyLight->SamplesPerPixel, 2);
 	if (SamplesPerPixel <= 0)
 	{
 		return;
