Fix static analysis.

diff --git a/Engine/Source/Runtime/Renderer/Private/HairStrands/HairStrandsVisibility.cpp b/Engine/Source/Runtime/Renderer/Private/HairStrands/HairStrandsVisibility.cpp
--- a/Engine/Source/Runtime/Renderer/Private/HairStrands/HairStrandsVisibility.cpp
+++ b/Engine/Source/Runtime/Renderer/Private/HairStrands/HairStrandsVisibility.cpp
@@ -325,10 +325,11 @@ void FHairMaterialProcessor::AddMeshBatch(const FMeshBatch& RESTRICT MeshBatch,
 	const FMaterial& Material = MeshBatch.MaterialRenderProxy->GetMaterialWithFallback(FeatureLevel, FallbackMaterialRenderProxyPtr);
 	const bool bIsCompatible = IsCompatibleWithHairStrands(&Material, FeatureLevel);
 	const bool bIsHairStrandsFactory = MeshBatch.VertexFactory->GetType()->GetId() == CompatibleVF->GetId();
+	const bool bShouldRender = (!PrimitiveSceneProxy && MeshBatch.Elements.Num()>0) || (PrimitiveSceneProxy && PrimitiveSceneProxy->ShouldRenderInMainPass());
 
 	if (bIsCompatible
 		&& bIsHairStrandsFactory
-		&& (!PrimitiveSceneProxy || PrimitiveSceneProxy->ShouldRenderInMainPass())
+		&& bShouldRender
 		&& ShouldIncludeDomainInMeshPass(Material.GetMaterialDomain()))
 	{
 		// For the mesh patch to be rendered a single triangle triangle to spawn the necessary amount of thread
@@ -342,7 +343,8 @@ void FHairMaterialProcessor::AddMeshBatch(const FMeshBatch& RESTRICT MeshBatch,
 
 		int32 PrimitiveId = 0;
 		int32 ScenePrimitiveId = 0;
-		GetDrawCommandPrimitiveId(PrimitiveSceneProxy->GetPrimitiveSceneInfo(), MeshBatch.Elements[0], PrimitiveId, ScenePrimitiveId);
+		FPrimitiveSceneInfo* SceneInfo = PrimitiveSceneProxy ? PrimitiveSceneProxy->GetPrimitiveSceneInfo() : nullptr;
+		GetDrawCommandPrimitiveId(SceneInfo, MeshBatch.Elements[0], PrimitiveId, ScenePrimitiveId);
 
 		const FMaterialRenderProxy& MaterialRenderProxy = FallbackMaterialRenderProxyPtr ? *FallbackMaterialRenderProxyPtr : *MeshBatch.MaterialRenderProxy;
 		Process(MeshBatchCopy, BatchElementMask, PrimitiveSceneProxy, StaticMeshId, MaterialRenderProxy, Material, HairClusterId, HairMaterialId, PrimitiveId);
@@ -795,10 +797,11 @@ void FHairVisibilityProcessor::AddMeshBatch(const FMeshBatch& RESTRICT MeshBatch
 	const FMaterial& Material = MeshBatch.MaterialRenderProxy->GetMaterialWithFallback(FeatureLevel, FallbackMaterialRenderProxyPtr);
 	const bool bIsCompatible = IsCompatibleWithHairStrands(&Material, FeatureLevel);
 	const bool bIsHairStrandsFactory = MeshBatch.VertexFactory->GetType()->GetId() == CompatibleVF->GetId();
+	const bool bShouldRender = (!PrimitiveSceneProxy && MeshBatch.Elements.Num()>0) || (PrimitiveSceneProxy && PrimitiveSceneProxy->ShouldRenderInMainPass());
 
 	if (bIsCompatible 
 		&& bIsHairStrandsFactory
-		&& (!PrimitiveSceneProxy || PrimitiveSceneProxy->ShouldRenderInMainPass())
+		&& bShouldRender
 		&& ShouldIncludeDomainInMeshPass(Material.GetMaterialDomain()))
 	{
 		const FMaterialRenderProxy& MaterialRenderProxy = FallbackMaterialRenderProxyPtr ? *FallbackMaterialRenderProxyPtr : *MeshBatch.MaterialRenderProxy;
