fix crash when trying to deref empty shadow mask texture

diff --git a/Engine/Source/Runtime/Renderer/Private/LightRendering.cpp b/Engine/Source/Runtime/Renderer/Private/LightRendering.cpp
--- a/Engine/Source/Runtime/Renderer/Private/LightRendering.cpp
+++ b/Engine/Source/Runtime/Renderer/Private/LightRendering.cpp
@@ -1781,7 +1781,10 @@ void FDeferredShadingSceneRenderer::RenderLights(FRHICommandListImmediate& RHICm
 
 					if (!bShadowMaskReadable)
 					{
-						RHICmdList.TransitionResource(EResourceTransitionAccess::EReadable, ScreenShadowMaskTexture->GetRenderTargetItem().ShaderResourceTexture);
+						if(ScreenShadowMaskTexture)
+						{
+							RHICmdList.TransitionResource(EResourceTransitionAccess::EReadable, ScreenShadowMaskTexture->GetRenderTargetItem().ShaderResourceTexture);
+						}
 						if (ScreenShadowMaskSubPixelTexture)
 						{
 							RHICmdList.TransitionResource(EResourceTransitionAccess::EReadable, ScreenShadowMaskSubPixelTexture->GetRenderTargetItem().ShaderResourceTexture);
