Undo single //UE4/Release-4.24/Engine/Source/Runtime/Renderer/Private/LightRendering.cpp revision 12

diff --git a/Engine/Source/Runtime/Renderer/Private/LightRendering.cpp b/Engine/Source/Runtime/Renderer/Private/LightRendering.cpp
--- a/Engine/Source/Runtime/Renderer/Private/LightRendering.cpp
+++ b/Engine/Source/Runtime/Renderer/Private/LightRendering.cpp
@@ -1781,10 +1781,7 @@ void FDeferredShadingSceneRenderer::RenderLights(FRHICommandListImmediate& RHICm
 
 					if (!bShadowMaskReadable)
 					{
-						if(ScreenShadowMaskTexture)
-						{
-							RHICmdList.TransitionResource(EResourceTransitionAccess::EReadable, ScreenShadowMaskTexture->GetRenderTargetItem().ShaderResourceTexture);
-						}
+						RHICmdList.TransitionResource(EResourceTransitionAccess::EReadable, ScreenShadowMaskTexture->GetRenderTargetItem().ShaderResourceTexture);
 						if (ScreenShadowMaskSubPixelTexture)
 						{
 							RHICmdList.TransitionResource(EResourceTransitionAccess::EReadable, ScreenShadowMaskSubPixelTexture->GetRenderTargetItem().ShaderResourceTexture);
