Fix for VR previous crash.

diff --git a/Engine/Source/Runtime/Renderer/Private/PostProcess/PostProcessHMD.cpp b/Engine/Source/Runtime/Renderer/Private/PostProcess/PostProcessHMD.cpp
--- a/Engine/Source/Runtime/Renderer/Private/PostProcess/PostProcessHMD.cpp
+++ b/Engine/Source/Runtime/Renderer/Private/PostProcess/PostProcessHMD.cpp
@@ -117,7 +117,9 @@ FScreenPassTexture AddDefaultHMDDistortionPass(FRDGBuilder& GraphBuilder, const
 	{
 		RHICmdList.SetViewport(OutputViewRect.Min.X, OutputViewRect.Min.Y, 0.0f, OutputViewRect.Max.X, OutputViewRect.Max.Y, 1.0f);
 
-		SetScreenPassPipelineState(RHICmdList, FScreenPassPipelineState(*VertexShader, *PixelShader));
+		FScreenPassPipelineState PipelineState(*VertexShader, *PixelShader);
+		PipelineState.VertexDeclaration = GDistortionVertexDeclaration.VertexDeclarationRHI;
+		SetScreenPassPipelineState(RHICmdList, PipelineState);
 
 		SetShaderParameters(RHICmdList, *VertexShader, VertexShader->GetVertexShader(), *PassParameters);
 		SetShaderParameters(RHICmdList, *PixelShader, PixelShader->GetPixelShader(), *PassParameters);
diff --git a/Engine/Source/Runtime/Renderer/Private/ScreenPass.cpp b/Engine/Source/Runtime/Renderer/Private/ScreenPass.cpp
--- a/Engine/Source/Runtime/Renderer/Private/ScreenPass.cpp
+++ b/Engine/Source/Runtime/Renderer/Private/ScreenPass.cpp
@@ -104,7 +104,7 @@ void SetScreenPassPipelineState(FRHICommandList& RHICmdList, const FScreenPassPi
 	GraphicsPSOInit.BlendState = ScreenPassDraw.BlendState;
 	GraphicsPSOInit.RasterizerState = TStaticRasterizerState<FM_Solid, CM_None>::GetRHI();
 	GraphicsPSOInit.DepthStencilState = ScreenPassDraw.DepthStencilState;
-	GraphicsPSOInit.BoundShaderState.VertexDeclarationRHI = GFilterVertexDeclaration.VertexDeclarationRHI;
+	GraphicsPSOInit.BoundShaderState.VertexDeclarationRHI = ScreenPassDraw.VertexDeclaration;
 	GraphicsPSOInit.BoundShaderState.VertexShaderRHI = VertexShaderRHI;
 	GraphicsPSOInit.BoundShaderState.PixelShaderRHI = PixelShaderRHI;
 	GraphicsPSOInit.PrimitiveType = PT_TriangleList;
diff --git a/Engine/Source/Runtime/Renderer/Private/ScreenPass.h b/Engine/Source/Runtime/Renderer/Private/ScreenPass.h
--- a/Engine/Source/Runtime/Renderer/Private/ScreenPass.h
+++ b/Engine/Source/Runtime/Renderer/Private/ScreenPass.h
@@ -226,11 +226,13 @@ struct FScreenPassPipelineState
 		FShader* InVertexShader,
 		FShader* InPixelShader,
 		FRHIBlendState* InBlendState = FDefaultBlendState::GetRHI(),
-		FRHIDepthStencilState* InDepthStencilState = FDefaultDepthStencilState::GetRHI())
+		FRHIDepthStencilState* InDepthStencilState = FDefaultDepthStencilState::GetRHI(),
+		FRHIVertexDeclaration* InVertexDeclaration = GFilterVertexDeclaration.VertexDeclarationRHI)
 		: VertexShader(InVertexShader)
 		, PixelShader(InPixelShader)
 		, BlendState(InBlendState)
 		, DepthStencilState(InDepthStencilState)
+		, VertexDeclaration(InVertexDeclaration)
 	{}
 
 	void Validate() const
@@ -239,12 +241,14 @@ struct FScreenPassPipelineState
 		check(PixelShader);
 		check(BlendState);
 		check(DepthStencilState);
+		check(VertexDeclaration);
 	}
 
 	FShader* VertexShader = nullptr;
 	FShader* PixelShader = nullptr;
 	FRHIBlendState* BlendState = nullptr;
 	FRHIDepthStencilState* DepthStencilState = nullptr;
+	FRHIVertexDeclaration* VertexDeclaration = nullptr;
 };
 
 // Helper function which sets the pipeline state object on the command list prior to invoking a screen pass.
