Fix crash when the following two events happen during the same frame:
* a parameter of a translucent material changes (or any other material change which causes a mesh draw command to be rebuilt)
* the scene is captured to a cube map (for example, because a skylight component is marked as dirty)

UpdateDeferredCachedUniformExpressions() would normally be called by BeginRenderingViewFamily() to update the changed material, but USkyLightComponent::UpdateSkyCaptureContents() runs before the viewport is redrawn, so if we don't update the cache explicitly at the beginning of CaptureSceneToScratchCubemap(), FMaterialShader::GetShaderBindings() asserts.

diff --git a/Engine/Source/Runtime/Renderer/Private/ReflectionEnvironmentCapture.cpp b/Engine/Source/Runtime/Renderer/Private/ReflectionEnvironmentCapture.cpp
--- a/Engine/Source/Runtime/Renderer/Private/ReflectionEnvironmentCapture.cpp
+++ b/Engine/Source/Runtime/Renderer/Private/ReflectionEnvironmentCapture.cpp
@@ -759,6 +759,7 @@ void CaptureSceneToScratchCubemap(FRHICommandListImmediate& RHICmdList, FSceneRe
 
 	// update any resources that needed a deferred update
 	FDeferredUpdateResource::UpdateResources(RHICmdList);
+	FMaterialRenderProxy::UpdateDeferredCachedUniformExpressions();
 
 	const auto FeatureLevel = SceneRenderer->FeatureLevel;
 	
diff --git a/Engine/Source/Runtime/Renderer/Private/ShaderBaseClasses.cpp b/Engine/Source/Runtime/Renderer/Private/ShaderBaseClasses.cpp
--- a/Engine/Source/Runtime/Renderer/Private/ShaderBaseClasses.cpp
+++ b/Engine/Source/Runtime/Renderer/Private/ShaderBaseClasses.cpp
@@ -424,7 +424,6 @@ void FMaterialShader::GetShaderBindings(
 
 	const FUniformExpressionCache& UniformExpressionCache = MaterialRenderProxy.UniformExpressionCache[FeatureLevel];
 
-	// Trying to get context for UE-75466.
 	checkf(UniformExpressionCache.bUpToDate, TEXT("UniformExpressionCache should be up to date, RenderProxy=%s Material=%s FeatureLevel=%d"), *MaterialRenderProxy.GetFriendlyName(), *Material.GetFriendlyName(), FeatureLevel);
 	checkf(UniformExpressionCache.UniformBuffer, TEXT("NULL UniformBuffer, RenderProxy=%s Material=%s FeatureLevel=%d"), *MaterialRenderProxy.GetFriendlyName(), *Material.GetFriendlyName(), FeatureLevel);
 
