Fixed: GPU crash on iOS during reflection enviroment capture

diff --git a/Engine/Source/Runtime/Renderer/Private/ReflectionEnvironmentCapture.cpp b/Engine/Source/Runtime/Renderer/Private/ReflectionEnvironmentCapture.cpp
--- a/Engine/Source/Runtime/Renderer/Private/ReflectionEnvironmentCapture.cpp
+++ b/Engine/Source/Runtime/Renderer/Private/ReflectionEnvironmentCapture.cpp
@@ -530,7 +530,8 @@ public:
 
 		if (IsMobilePlatform(Parameters.Platform))
 		{
-			OutEnvironment.SetDefine(TEXT("MOBILE_FORCE_DEPTH_TEXTURE_READS"), 1u);
+			// SceneDepth is memoryless on mobile
+			OutEnvironment.SetDefine(TEXT("SCENE_TEXTURES_DISABLED"), 1u);
 		}
 	}
 
