Transition depth texture back to read after velocity pass.
(fixes validation error)

diff --git a/Engine/Source/Runtime/Renderer/Private/VelocityRendering.cpp b/Engine/Source/Runtime/Renderer/Private/VelocityRendering.cpp
--- a/Engine/Source/Runtime/Renderer/Private/VelocityRendering.cpp
+++ b/Engine/Source/Runtime/Renderer/Private/VelocityRendering.cpp
@@ -443,7 +443,14 @@ void FDeferredShadingSceneRenderer::RenderVelocities(FRHICommandListImmediate& R
 			RenderVelocitiesInner(RHICmdList, VelocityRT, VelocityPass);
 			RHICmdList.EndRenderPass();
 		}
-
+		if(VelocityPass != EVelocityPass::Opaque)
+		{
+			FTexture2DRHIRef DepthTexture = FSceneRenderTargets::Get(RHICmdList).GetSceneDepthTexture();
+			if(DepthTexture)
+			{
+				RHICmdList.TransitionResource(EResourceTransitionAccess::EReadable, DepthTexture);
+			}
+		}
 		RHICmdList.CopyToResolveTarget(VelocityRT->GetRenderTargetItem().TargetableTexture, VelocityRT->GetRenderTargetItem().ShaderResourceTexture, FResolveParams());
 	}
 
