Secondary views in multiview mode now render ray tracing reflections with sort materials disabled to avoid breaking the reflection pass

diff --git a/Engine/Source/Runtime/Renderer/Private/RayTracing/RayTracingReflections.cpp b/Engine/Source/Runtime/Renderer/Private/RayTracing/RayTracingReflections.cpp
--- a/Engine/Source/Runtime/Renderer/Private/RayTracing/RayTracingReflections.cpp
+++ b/Engine/Source/Runtime/Renderer/Private/RayTracing/RayTracingReflections.cpp
@@ -240,7 +240,8 @@ void FDeferredShadingSceneRenderer::PrepareRayTracingReflections(const FViewInfo
 	// Declare all RayGen shaders that require material closest hit shaders to be bound
 
 	const bool bHybridReflections = CVarRayTracingReflectionsHybrid.GetValueOnRenderThread() != 0;
-	const bool bSortMaterials = bHybridReflections || CVarRayTracingReflectionsSortMaterials.GetValueOnRenderThread() != 0;
+	const bool bIsMultiviewSecondary = View.ViewRect.Min.X > 0 || View.ViewRect.Min.Y > 0;
+	const bool bSortMaterials = (bHybridReflections || CVarRayTracingReflectionsSortMaterials.GetValueOnRenderThread() != 0) && !bIsMultiviewSecondary;
 
 	if (bSortMaterials)
 	{
@@ -289,7 +290,8 @@ void FDeferredShadingSceneRenderer::RenderRayTracingReflections(
 	const uint32 SortTileSize = CVarRayTracingReflectionsSortTileSize.GetValueOnRenderThread();
 
 	const bool bHybridReflections = CVarRayTracingReflectionsHybrid.GetValueOnRenderThread() != 0;
-	const bool bSortMaterials = bHybridReflections || CVarRayTracingReflectionsSortMaterials.GetValueOnRenderThread() != 0;
+	const bool bIsMultiviewSecondary = View.ViewRect.Min.X > 0 || View.ViewRect.Min.Y > 0;
+	const bool bSortMaterials = (bHybridReflections || CVarRayTracingReflectionsSortMaterials.GetValueOnRenderThread() != 0) && !bIsMultiviewSecondary;
 	const uint32 EnableTranslucency = GRayTracingReflectionsTranslucency > -1 ? (uint32)GRayTracingReflectionsTranslucency : (uint32) View.FinalPostProcessSettings.RayTracingReflectionsTranslucency;
 
 	FSceneRenderTargets& SceneContext = FSceneRenderTargets::Get(GraphBuilder.RHICmdList);
