Only clear last cubearray.
This is a workaround for the clearing of the cubemap overwriting the upload of it.

diff --git a/Engine/Source/Runtime/Renderer/Private/ReflectionEnvironment.cpp b/Engine/Source/Runtime/Renderer/Private/ReflectionEnvironment.cpp
--- a/Engine/Source/Runtime/Renderer/Private/ReflectionEnvironment.cpp
+++ b/Engine/Source/Runtime/Renderer/Private/ReflectionEnvironment.cpp
@@ -289,7 +289,7 @@ TUniformBufferRef<FReflectionUniformParameters> CreateReflectionUniformBuffer(co
 	return CreateUniformBufferImmediate(ReflectionStruct, Usage);
 }
 
-void ClearCubeArray(FRHICommandListImmediate& RHICmdList, TRefCountPtr<IPooledRenderTarget>& CubeArray)
+static void ClearCubeArrayLastCube(FRHICommandListImmediate& RHICmdList, TRefCountPtr<IPooledRenderTarget>& CubeArray)
 {
 	const FPooledRenderTargetDesc& Desc = CubeArray.GetReference()->GetDesc();
 	int32 NumMips = Desc.NumMips;
@@ -324,14 +324,12 @@ void ClearCubeArray(FRHICommandListImmediate& RHICmdList, TRefCountPtr<IPooledRe
 		FSceneRenderTargetItem& DestCube = CubeArray.GetReference()->GetRenderTargetItem();
 
 		// GPU copy back to the scene's texture array, which is not a render target
-		for (int32 CaptureIndex = 0; CaptureIndex < NumCaptures; CaptureIndex++)
+		int32 CaptureIndex = NumCaptures-1;
+		for (int32 MipIndex = 0; MipIndex < NumMips; MipIndex++)
 		{
-			for (int32 MipIndex = 0; MipIndex < NumMips; MipIndex++)
+			for (int32 CubeFace = 0; CubeFace < CubeFace_MAX; CubeFace++)
 			{
-				for (int32 CubeFace = 0; CubeFace < CubeFace_MAX; CubeFace++)
-				{
-					RHICmdList.CopyToResolveTarget(RT0.TargetableTexture, DestCube.ShaderResourceTexture, FResolveParams(FResolveRect(), (ECubeFace)CubeFace, MipIndex, 0, CaptureIndex));
-				}
+				RHICmdList.CopyToResolveTarget(RT0.TargetableTexture, DestCube.ShaderResourceTexture, FResolveParams(FResolveRect(), (ECubeFace)CubeFace, MipIndex, 0, CaptureIndex));
 			}
 		}
 	}
@@ -370,7 +368,7 @@ void FReflectionEnvironmentCubemapArray::InitDynamicRHI()
 		GRenderTargetPool.FindFreeElement(RHICmdList, Desc, ReflectionEnvs, TEXT("ReflectionEnvs"));
 		if(MaxCubemaps == 1)
 		{
-			ClearCubeArray(RHICmdList, ReflectionEnvs); // When rounded up to two cubemaps, we need to clear all , to avoid validation problems, which happen because we don't fill out the final cube
+			ClearCubeArrayLastCube(RHICmdList, ReflectionEnvs); // When rounded up to two cubemaps, we need to clear all , to avoid validation problems, which happen because we don't fill out the final cube
 		}
 	}
 }
