Clean up stuff StencilDummy/SRV on ReleaseDynamicRHI

diff --git a/Engine/Source/Runtime/Renderer/Private/SystemTextures.cpp b/Engine/Source/Runtime/Renderer/Private/SystemTextures.cpp
--- a/Engine/Source/Runtime/Renderer/Private/SystemTextures.cpp
+++ b/Engine/Source/Runtime/Renderer/Private/SystemTextures.cpp
@@ -622,6 +622,8 @@ void FSystemTextures::ReleaseDynamicRHI()
 	DefaultNormal8Bit.SafeRelease();
 	VolumetricBlackDummy.SafeRelease();
 	MidGreyDummy.SafeRelease();
+	StencilDummy.SafeRelease();
+	StencilDummySRV.SafeRelease();
 
 	GRenderTargetPool.FreeUnusedResources();
 
@@ -712,4 +714,4 @@ FRDGTextureRef FSystemTextures::GetMidGreyDummy(FRDGBuilder& GraphBuilder) const
 FRDGTextureRef FSystemTextures::GetVolumetricBlackDummy(FRDGBuilder& GraphBuilder) const
 {
 	return GraphBuilder.RegisterExternalTexture(VolumetricBlackDummy, TEXT("VolumetricBlackDummy"));
-}
\ No newline at end of file
+}
