Gracefully disables the SSGI denoiser when on an unsupported platform

diff --git a/Engine/Source/Runtime/Renderer/Private/IndirectLightRendering.cpp b/Engine/Source/Runtime/Renderer/Private/IndirectLightRendering.cpp
--- a/Engine/Source/Runtime/Renderer/Private/IndirectLightRendering.cpp
+++ b/Engine/Source/Runtime/Renderer/Private/IndirectLightRendering.cpp
@@ -151,6 +151,14 @@ void FDeferredShadingSceneRenderer::RenderDiffuseIndirectAndAmbientOcclusion(FRH
 		else if (bApplySSGI)
 		{
 			RenderScreenSpaceDiffuseIndirect(GraphBuilder, SceneTextures, SceneColor, View, /* out */ &RayTracingConfig, /* out */ &DenoiserInputs);
+
+			const IScreenSpaceDenoiser* DefaultDenoiser = IScreenSpaceDenoiser::GetDefaultDenoiser();
+			const IScreenSpaceDenoiser* DenoiserToUse = DenoiseMode == 1 ? DefaultDenoiser : GScreenSpaceDenoiser;
+
+			if (!DenoiserToUse->SupportsScreenSpaceDiffuseIndirectDenoiser(View.GetShaderPlatform()) && DenoiseMode > 0)
+			{
+				DenoiseMode = 0;
+			}
 		}
 		else
 		{
diff --git a/Engine/Source/Runtime/Renderer/Private/ScreenSpaceDenoise.cpp b/Engine/Source/Runtime/Renderer/Private/ScreenSpaceDenoise.cpp
--- a/Engine/Source/Runtime/Renderer/Private/ScreenSpaceDenoise.cpp
+++ b/Engine/Source/Runtime/Renderer/Private/ScreenSpaceDenoise.cpp
@@ -679,7 +679,7 @@ bool ShouldCompileSignalPipeline(ESignalProcessing SignalProcessing, EShaderPlat
 {
 	if (SignalProcessing == ESignalProcessing::ScreenSpaceDiffuseIndirect)
 	{
-		return Platform == SP_PCD3D_SM5 || Platform == SP_PS4;
+		return Platform == SP_PCD3D_SM5 || Platform == SP_PS4 || Platform == SP_XBOXONE_D3D12;
 	}
 	else if (
 		SignalProcessing == ESignalProcessing::Reflections ||
@@ -2548,13 +2548,18 @@ public:
 		return GlobalIlluminationOutputs;
 	}
 
+	bool SupportsScreenSpaceDiffuseIndirectDenoiser(EShaderPlatform Platform) const override
+	{
+		return ShouldCompileSignalPipeline(ESignalProcessing::ScreenSpaceDiffuseIndirect, Platform);
+	}
+
 	FDiffuseIndirectOutputs DenoiseScreenSpaceDiffuseIndirect(
 		FRDGBuilder& GraphBuilder,
 		const FViewInfo& View,
 		FPreviousViewInfo* PreviousViewInfos,
 		const FSceneTextureParameters& SceneTextures,
 		const FDiffuseIndirectInputs& Inputs,
-		const FAmbientOcclusionRayTracingConfig Config) const
+		const FAmbientOcclusionRayTracingConfig Config) const override
 	{
 		RDG_GPU_STAT_SCOPE(GraphBuilder, DiffuseIndirectDenoiser);
 
diff --git a/Engine/Source/Runtime/Renderer/Private/ScreenSpaceDenoise.h b/Engine/Source/Runtime/Renderer/Private/ScreenSpaceDenoise.h
--- a/Engine/Source/Runtime/Renderer/Private/ScreenSpaceDenoise.h
+++ b/Engine/Source/Runtime/Renderer/Private/ScreenSpaceDenoise.h
@@ -246,6 +246,7 @@ public:
 		const FAmbientOcclusionRayTracingConfig Config) const = 0;
 
 	/** Entry point to denoise SSGI. */
+	virtual bool SupportsScreenSpaceDiffuseIndirectDenoiser(EShaderPlatform Platform) const = 0;
 	virtual FDiffuseIndirectOutputs DenoiseScreenSpaceDiffuseIndirect(
 		FRDGBuilder& GraphBuilder,
 		const FViewInfo& View,
