Kills r.SSGI.MinimumLuminance failed experiment

diff --git a/Engine/Source/Runtime/Renderer/Private/ScreenSpaceRayTracing.cpp b/Engine/Source/Runtime/Renderer/Private/ScreenSpaceRayTracing.cpp
--- a/Engine/Source/Runtime/Renderer/Private/ScreenSpaceRayTracing.cpp
+++ b/Engine/Source/Runtime/Renderer/Private/ScreenSpaceRayTracing.cpp
@@ -63,11 +63,6 @@ static TAutoConsoleVariable<int32> CVarSSGIQuality(
 	TEXT("Quality setting to control number of ray shot with SSGI, between 1 and 4 (defaults to 4).\n"),
 	ECVF_Scalability | ECVF_RenderThreadSafe);
 
-static TAutoConsoleVariable<float> CVarSSGIMinimumLuminance(
-	TEXT("r.SSGI.MinimumLuminance"), 0.5f,
-	TEXT("Experimental"),
-	ECVF_Scalability | ECVF_RenderThreadSafe);
-
 
 DECLARE_GPU_STAT_NAMED(ScreenSpaceReflections, TEXT("ScreenSpace Reflections"));
 DECLARE_GPU_STAT_NAMED(ScreenSpaceDiffuseIndirect, TEXT("Screen Space Diffuse Indirect"));
@@ -950,7 +945,6 @@ void RenderScreenSpaceDiffuseIndirect(
 			}
 
 			PassParameters->PrevSceneColorPreExposureCorrection = View.PreExposure / View.PrevViewInfo.SceneColorPreExposure;
-			PassParameters->MinimumLuminance = CVarSSGIMinimumLuminance.GetValueOnRenderThread();
 
 			PassParameters->PrevScreenPositionScaleBias = FVector4(
 				ViewportExtent.X * 0.5f / BufferSize.X,
