Another missing transition for shadowmaps

diff --git a/Engine/Source/Runtime/Renderer/Private/ShadowRendering.cpp b/Engine/Source/Runtime/Renderer/Private/ShadowRendering.cpp
--- a/Engine/Source/Runtime/Renderer/Private/ShadowRendering.cpp
+++ b/Engine/Source/Runtime/Renderer/Private/ShadowRendering.cpp
@@ -1516,6 +1516,10 @@ bool FSceneRenderer::RenderShadowProjections(FRHICommandListImmediate& RHICmdLis
 		else
 		{
 			NormalShadows.Add(ProjectedShadowInfo);
+			if (ProjectedShadowInfo->bAllocated && ProjectedShadowInfo->RenderTargets.DepthTarget)
+			{
+				RHICmdList.TransitionResource(EResourceTransitionAccess::EReadable, ProjectedShadowInfo->RenderTargets.DepthTarget->GetRenderTargetItem().ShaderResourceTexture.GetReference());
+			}
 		}
 	}
 
