Remove ClearCubeArray and fix creation of cube array resources in the reflection environment.
 - Old code was clamping the minimum number of cubes in the reflection environment cube array to 2, to force that the RHI resource was always an array, rather than a single cube.
 - This only seemed to be necessary because the pooled render target desc function only set bIsArray to true if InArraySize > 1.
 - Adding a CreateCubemapArrayDesc function removes the need to clamp the num slices, so now we allocate the number of cubes we actually need. This also solves the validation warnings that CL 10112521 was originally working around.

diff --git a/Engine/Source/Runtime/RenderCore/Public/RendererInterface.h b/Engine/Source/Runtime/RenderCore/Public/RendererInterface.h
--- a/Engine/Source/Runtime/RenderCore/Public/RendererInterface.h
+++ b/Engine/Source/Runtime/RenderCore/Public/RendererInterface.h
@@ -190,6 +190,43 @@ public:
 		return NewDesc;
 	}
 
+	/**
+	 * Factory function to create cube map array texture description
+	 * @param InFlags bit mask combined from elements of ETextureCreateFlags e.g. TexCreate_UAV
+	 */
+	static FPooledRenderTargetDesc CreateCubemapArrayDesc(
+		uint32 InExtent,
+		EPixelFormat InFormat,
+		const FClearValueBinding& InClearValue,
+		uint32 InFlags,
+		uint32 InTargetableFlags,
+		bool bInForceSeparateTargetAndShaderResource,
+		uint32 InArraySize,
+		uint16 InNumMips = 1,
+		bool InAutowritable = true)
+	{
+		check(InExtent);
+
+		FPooledRenderTargetDesc NewDesc;
+		NewDesc.ClearValue = InClearValue;
+		NewDesc.Extent = FIntPoint(InExtent, InExtent);
+		NewDesc.Depth = 0;
+		NewDesc.ArraySize = InArraySize;
+		NewDesc.bIsArray = true;
+		NewDesc.bIsCubemap = true;
+		NewDesc.NumMips = InNumMips;
+		NewDesc.NumSamples = 1;
+		NewDesc.Format = InFormat;
+		NewDesc.Flags = InFlags;
+		NewDesc.TargetableFlags = InTargetableFlags;
+		NewDesc.bForceSeparateTargetAndShaderResource = bInForceSeparateTargetAndShaderResource;
+		NewDesc.DebugName = TEXT("UnknownTextureCubeArray");
+		NewDesc.AutoWritable = InAutowritable;
+		check(NewDesc.IsCubemap());
+
+		return NewDesc;
+	}
+
 	/** Comparison operator to test if a render target can be reused */
 	bool Compare(const FPooledRenderTargetDesc& rhs, bool bExact) const
 	{
diff --git a/Engine/Source/Runtime/Renderer/Private/ReflectionEnvironment.cpp b/Engine/Source/Runtime/Renderer/Private/ReflectionEnvironment.cpp
--- a/Engine/Source/Runtime/Renderer/Private/ReflectionEnvironment.cpp
+++ b/Engine/Source/Runtime/Renderer/Private/ReflectionEnvironment.cpp
@@ -289,54 +289,6 @@ TUniformBufferRef<FReflectionUniformParameters> CreateReflectionUniformBuffer(co
 	return CreateUniformBufferImmediate(ReflectionStruct, Usage);
 }
 
-static void ClearCubeArrayLastCube(FRHICommandListImmediate& RHICmdList, TRefCountPtr<IPooledRenderTarget>& CubeArray)
-{
-	const FPooledRenderTargetDesc& Desc = CubeArray.GetReference()->GetDesc();
-	int32 NumMips = Desc.NumMips;
-	int32 CubemapSize = Desc.Extent.X;
-	int32 NumCaptures = Desc.ArraySize;
-	TRefCountPtr<IPooledRenderTarget> TempCube;
-
-	uint32 CubeTexFlags = TexCreate_TargetArraySlicesIndependently;
-	FPooledRenderTargetDesc Desc2(FPooledRenderTargetDesc::CreateCubemapDesc(CubemapSize, Desc.Format, FClearValueBinding(FLinearColor(0, 0, 0, 0)), CubeTexFlags, TexCreate_RenderTargetable, false, 1, NumMips));
-	GRenderTargetPool.FindFreeElement(RHICmdList, Desc2, TempCube, TEXT("TempCube"), true, ERenderTargetTransience::Transient);
-
-	FSceneRenderTargetItem& RT0 = TempCube->GetRenderTargetItem();
-	{
-		SCOPED_DRAW_EVENT(RHICmdList, ClearCubeArray);
-
-		RHICmdList.TransitionResource(EResourceTransitionAccess::EWritable, RT0.TargetableTexture);
-
-		for (int32 MipIndex = 0; MipIndex < NumMips; MipIndex++)
-		{
-			for (int32 CubeFace = 0; CubeFace < CubeFace_MAX; CubeFace++)
-			{
-				FRHIRenderPassInfo RPInfo(RT0.TargetableTexture, ERenderTargetActions::Clear_Store, nullptr, MipIndex, CubeFace);
-
-				RHICmdList.BeginRenderPass(RPInfo, TEXT("ClearCubeFace"));
-				RHICmdList.EndRenderPass();
-			}
-		}
-
-		FResolveParams ResolveParams(FResolveRect(), CubeFace_PosX, -1, -1, -1);
-		RHICmdList.CopyToResolveTarget(RT0.TargetableTexture, RT0.TargetableTexture, ResolveParams);
-
-		FSceneRenderTargetItem& DestCube = CubeArray.GetReference()->GetRenderTargetItem();
-
-		// GPU copy back to the scene's texture array, which is not a render target
-		int32 CaptureIndex = NumCaptures-1;
-		for (int32 MipIndex = 0; MipIndex < NumMips; MipIndex++)
-		{
-			for (int32 CubeFace = 0; CubeFace < CubeFace_MAX; CubeFace++)
-			{
-				RHICmdList.CopyToResolveTarget(RT0.TargetableTexture, DestCube.ShaderResourceTexture, FResolveParams(FResolveRect(), (ECubeFace)CubeFace, MipIndex, 0, CaptureIndex));
-			}
-		}
-	}
-}
-
-
-
 void FReflectionEnvironmentCubemapArray::InitDynamicRHI()
 {
 	if (GetFeatureLevel() >= ERHIFeatureLevel::SM5)
@@ -345,31 +297,23 @@ void FReflectionEnvironmentCubemapArray::InitDynamicRHI()
 
 		ReleaseCubeArray();
 
-		FPooledRenderTargetDesc Desc(
-			FPooledRenderTargetDesc::CreateCubemapDesc(
-				CubemapSize,
-				// Alpha stores sky mask
-				PF_FloatRGBA, 
-				FClearValueBinding::None,
-				TexCreate_None,
-				TexCreate_None,
-				false, 
-				// Cubemap array of 1 produces a regular cubemap, so guarantee it will be allocated as an array
-				FMath::Max<uint32>(MaxCubemaps, 2),
-				NumReflectionCaptureMips
-				)
-			);
-
-		Desc.AutoWritable = false;
+		FPooledRenderTargetDesc Desc = FPooledRenderTargetDesc::CreateCubemapArrayDesc(
+			CubemapSize,
+			// Alpha stores sky mask
+			PF_FloatRGBA,
+			FClearValueBinding::None,
+			TexCreate_None,
+			TexCreate_None,
+			false,
+			MaxCubemaps,
+			NumReflectionCaptureMips,
+			false
+		);
 	
 		FRHICommandListImmediate& RHICmdList = FRHICommandListExecutor::GetImmediateCommandList();
 
 		// Allocate TextureCubeArray for the scene's reflection captures
 		GRenderTargetPool.FindFreeElement(RHICmdList, Desc, ReflectionEnvs, TEXT("ReflectionEnvs"));
-		if(MaxCubemaps == 1)
-		{
-			ClearCubeArrayLastCube(RHICmdList, ReflectionEnvs); // When rounded up to two cubemaps, we need to clear all , to avoid validation problems, which happen because we don't fill out the final cube
-		}
 	}
 }
 
