Add new RVT format: BaseColor_Normal_Specular_Mask_YCoCg
This adds an aditional 8 bit "mask" channel for generic use in RVT materials.

diff --git a/Engine/Shaders/Private/VirtualTextureCompress.usf b/Engine/Shaders/Private/VirtualTextureCompress.usf
--- a/Engine/Shaders/Private/VirtualTextureCompress.usf
+++ b/Engine/Shaders/Private/VirtualTextureCompress.usf
@@ -19,8 +19,8 @@ SamplerState TextureSampler2;
 RWTexture2D<uint2> OutCompressTexture0_u2;
 RWTexture2D<uint4> OutCompressTexture0_u4;
 RWTexture2D<uint4> OutCompressTexture1;
-RWTexture2D<uint2> OutCompressTexture2;
-
+RWTexture2D<uint2> OutCompressTexture2_u2;
+RWTexture2D<uint4> OutCompressTexture2_u4;
 
 /** Trivial vertex shader for copy shaders. */
 void CopyVS(in uint VertexId : SV_VertexID, out float4 OutPosition : SV_POSITION, out float2 OutTexCoord : TEXCOORD0)
@@ -75,42 +75,86 @@ void ReadBlockRGB( Texture2D<float4> SourceTexture, SamplerState TextureSampler,
 }
 
 /** Read a 4x4 block of XY channels from a normal texture ready for BC5 compression. */
-void ReadBlockRG( Texture2D<float4> SourceTexture, SamplerState TextureSampler, float2 UV, float2 TexelUVSize, out float BlockR[16], out float BlockG[16] )
+void ReadBlockXY( Texture2D<float4> SourceTexture, SamplerState TextureSampler, float2 UV, float2 TexelUVSize, out float BlockX[16], out float BlockY[16] )
 {
 	{
 		float4 Red = SourceTexture.GatherRed(TextureSampler, UV);
 		float4 Green = SourceTexture.GatherGreen(TextureSampler, UV);
-		BlockR[0] = Red[3]; BlockG[0] = Green[3];
-		BlockR[1] = Red[2]; BlockG[1] = Green[2];
-		BlockR[4] = Red[0]; BlockG[4] = Green[0];
-		BlockR[5] = Red[1]; BlockG[5] = Green[1];
+		BlockX[0] = Red[3]; BlockY[0] = Green[3];
+		BlockX[1] = Red[2]; BlockY[1] = Green[2];
+		BlockX[4] = Red[0]; BlockY[4] = Green[0];
+		BlockX[5] = Red[1]; BlockY[5] = Green[1];
 	}
 	{
 		float2 UVOffset = UV + float2(2.f * TexelUVSize.x, 0);
 		float4 Red = SourceTexture.GatherRed(TextureSampler, UVOffset);
 		float4 Green = SourceTexture.GatherGreen(TextureSampler, UVOffset);
-		BlockR[2] = Red[3]; BlockG[2] = Green[3];
-		BlockR[3] = Red[2]; BlockG[3] = Green[2];
-		BlockR[6] = Red[0]; BlockG[6] = Green[0];
-		BlockR[7] = Red[1]; BlockG[7] = Green[1];
+		BlockX[2] = Red[3]; BlockY[2] = Green[3];
+		BlockX[3] = Red[2]; BlockY[3] = Green[2];
+		BlockX[6] = Red[0]; BlockY[6] = Green[0];
+		BlockX[7] = Red[1]; BlockY[7] = Green[1];
 	}
 	{
 		float2 UVOffset = UV + float2(0, 2.f * TexelUVSize.y);
 		float4 Red = SourceTexture.GatherRed(TextureSampler, UVOffset);
 		float4 Green = SourceTexture.GatherGreen(TextureSampler, UVOffset);
-		BlockR[8]  = Red[3]; BlockG[8]  = Green[3];
-		BlockR[9]  = Red[2]; BlockG[9]  = Green[2];
-		BlockR[12] = Red[0]; BlockG[12] = Green[0];
-		BlockR[13] = Red[1]; BlockG[13] = Green[1];
+		BlockX[8]  = Red[3]; BlockY[8]  = Green[3];
+		BlockX[9]  = Red[2]; BlockY[9]  = Green[2];
+		BlockX[12] = Red[0]; BlockY[12] = Green[0];
+		BlockX[13] = Red[1]; BlockY[13] = Green[1];
 	}
 	{
 		float2 UVOffset = UV + 2.f * TexelUVSize;
 		float4 Red = SourceTexture.GatherRed(TextureSampler, UVOffset);
 		float4 Green = SourceTexture.GatherGreen(TextureSampler, UVOffset);
-		BlockR[10] = Red[3]; BlockG[10] = Green[3];
-		BlockR[11] = Red[2]; BlockG[11] = Green[2];
-		BlockR[14] = Red[0]; BlockG[14] = Green[0];
-		BlockR[15] = Red[1]; BlockG[15] = Green[1];
+		BlockX[10] = Red[3]; BlockY[10] = Green[3];
+		BlockX[11] = Red[2]; BlockY[11] = Green[2];
+		BlockX[14] = Red[0]; BlockY[14] = Green[0];
+		BlockX[15] = Red[1]; BlockY[15] = Green[1];
+	}
+}
+
+/** Read a 4x4 block of XYA channels ready for BC4/BC5 compression. */
+void ReadBlockXYA( Texture2D<float4> SourceTexture, SamplerState TextureSampler, float2 UV, float2 TexelUVSize, out float BlockX[16], out float BlockY[16], out float BlockA[16] )
+{
+	{
+		float4 Red = SourceTexture.GatherRed(TextureSampler, UV);
+		float4 Green = SourceTexture.GatherGreen(TextureSampler, UV);
+		float4 Blue = SourceTexture.GatherBlue(TextureSampler, UV);
+		BlockX[0] = Red[3]; BlockY[0] = Green[3]; BlockA[0] = Blue[3];
+		BlockX[1] = Red[2]; BlockY[1] = Green[2]; BlockA[1] = Blue[2];
+		BlockX[4] = Red[0]; BlockY[4] = Green[0]; BlockA[4] = Blue[0];
+		BlockX[5] = Red[1]; BlockY[5] = Green[1]; BlockA[5] = Blue[1];
+	}
+	{
+		float2 UVOffset = UV + float2(2.f * TexelUVSize.x, 0);
+		float4 Red = SourceTexture.GatherRed(TextureSampler, UVOffset);
+		float4 Green = SourceTexture.GatherGreen(TextureSampler, UVOffset);
+		float4 Blue = SourceTexture.GatherBlue(TextureSampler, UV);
+		BlockX[2] = Red[3]; BlockY[2] = Green[3]; BlockA[2] = Blue[3];
+		BlockX[3] = Red[2]; BlockY[3] = Green[2]; BlockA[3] = Blue[2];
+		BlockX[6] = Red[0]; BlockY[6] = Green[0]; BlockA[6] = Blue[0];
+		BlockX[7] = Red[1]; BlockY[7] = Green[1]; BlockA[7] = Blue[1];
+	}
+	{
+		float2 UVOffset = UV + float2(0, 2.f * TexelUVSize.y);
+		float4 Red = SourceTexture.GatherRed(TextureSampler, UVOffset);
+		float4 Green = SourceTexture.GatherGreen(TextureSampler, UVOffset);
+		float4 Blue = SourceTexture.GatherBlue(TextureSampler, UV);
+		BlockX[8]  = Red[3]; BlockY[8]  = Green[3]; BlockA[8]  = Blue[3];
+		BlockX[9]  = Red[2]; BlockY[9]  = Green[2]; BlockA[9]  = Blue[2];
+		BlockX[12] = Red[0]; BlockY[12] = Green[0]; BlockA[12] = Blue[0];
+		BlockX[13] = Red[1]; BlockY[13] = Green[1]; BlockA[13] = Blue[1];
+	}
+	{
+		float2 UVOffset = UV + 2.f * TexelUVSize;
+		float4 Red = SourceTexture.GatherRed(TextureSampler, UVOffset);
+		float4 Green = SourceTexture.GatherGreen(TextureSampler, UVOffset);
+		float4 Blue = SourceTexture.GatherBlue(TextureSampler, UV);
+		BlockX[10] = Red[3]; BlockY[10] = Green[3]; BlockA[10] = Blue[3];
+		BlockX[11] = Red[2]; BlockY[11] = Green[2]; BlockA[11] = Blue[2];
+		BlockX[14] = Red[0]; BlockY[14] = Green[0]; BlockA[14] = Blue[0];
+		BlockX[15] = Red[1]; BlockY[15] = Green[1]; BlockA[15] = Blue[1];
 	}
 }
 
@@ -161,7 +205,7 @@ void CompressBaseColorNormalSpecularCS(
 
 	float BlockNormalX[16];
 	float BlockNormalY[16];
-	ReadBlockRG(RenderTexture1, TextureSampler1, SampleUV, TexelUVSize, BlockNormalX, BlockNormalY);
+	ReadBlockXY(RenderTexture1, TextureSampler1, SampleUV, TexelUVSize, BlockNormalX, BlockNormalY);
 
 	float3 BlockSpecular[16];
 	ReadBlockRGB(RenderTexture2, TextureSampler2, SampleUV, TexelUVSize, BlockSpecular);
@@ -175,8 +219,7 @@ void CompressBaseColorNormalSpecularCS(
 	OutCompressTexture1[ThreadId.xy + DestRect.xy] = CompressBC3Block(BlockSpecular, BlockNormalY);
 }
 
-/** Copy normal X and Y to the alpha channels of the final target textures (without compression)
-  * This path is used when we disable texture compression, because we need to keep the same final channel layout. */
+/** Copy path used when we disable texture compression, because we need to keep the same final channel layout. */
 void CopyBaseColorNormalSpecularPS(
 	in float4 InPosition : SV_POSITION, 
 	in noperspective float2 InTexCoord : TEXCOORD0, 
@@ -211,7 +254,7 @@ void CompressBaseColorNormalSpecularYCoCgCS(
 
 	float BlockNormalX[16];
 	float BlockNormalY[16];
-	ReadBlockRG(RenderTexture1, TextureSampler1, SampleUV, TexelUVSize, BlockNormalX, BlockNormalY);
+	ReadBlockXY(RenderTexture1, TextureSampler1, SampleUV, TexelUVSize, BlockNormalX, BlockNormalY);
 
 	float3 BlockSpecular[16];
 	ReadBlockRGB(RenderTexture2, TextureSampler2, SampleUV, TexelUVSize, BlockSpecular);
@@ -223,11 +266,10 @@ void CompressBaseColorNormalSpecularYCoCgCS(
 
 	OutCompressTexture0_u4[ThreadId.xy + DestRect.xy] = CompressBC3BlockYCoCg(BlockBaseColor);
 	OutCompressTexture1[ThreadId.xy + DestRect.xy] = CompressBC5Block(BlockNormalX, BlockNormalY);
-	OutCompressTexture2[ThreadId.xy + DestRect.xy] = CompressBC1Block(BlockSpecular);
+	OutCompressTexture2_u2[ThreadId.xy + DestRect.xy] = CompressBC1Block(BlockSpecular);
 }
 
-/** Copy normal X and Y to the alpha channels of the final target textures (without compression)
-  * This path is used when we disable texture compression, because we need to keep the same final channel layout. */
+/** Copy path used when we disable texture compression, because we need to keep the same final channel layout. */
 void CopyBaseColorNormalSpecularYCoCgPS(
 	in float4 InPosition : SV_POSITION, 
 	in noperspective float2 InTexCoord : TEXCOORD0, 
@@ -242,10 +284,63 @@ void CopyBaseColorNormalSpecularYCoCgPS(
 	RoughnessSpecularNormalZ.z = round(RoughnessSpecularNormalZ.z);
 	
 	OutColor0 = float4(YCoCg.yz, 0, YCoCg.x);
-	OutColor1 = float4(NormalXY.xy, 0, 1);
+	OutColor1 = float4(NormalXY, 0, 1);
 	OutColor2 = float4(RoughnessSpecularNormalZ, 1);
 }
 
+/** Compress base color, normal, roughness, specular and mask to BC3, BC5, BC3 target set. */
+[numthreads(8,8,1)]
+void CompressBaseColorNormalSpecularMaskYCoCgCS(
+	uint3 ThreadId : SV_DispatchThreadID 
+	)
+{
+	uint2 SamplePos = ThreadId.xy * 4;
+	if (any(SamplePos >= DestRect.zw))
+		return;
+
+	float2 TexelUVSize = 1.f / float2(DestRect.zw);
+	float2 SampleUV = (float2(SamplePos) + 0.5f) * TexelUVSize;
+
+	float3 BlockBaseColor[16];
+	ReadBlockRGB(RenderTexture0, TextureSampler0, SampleUV, TexelUVSize, BlockBaseColor);
+
+	float BlockNormalX[16];
+	float BlockNormalY[16];
+	float BlockMask[16];
+	ReadBlockXYA(RenderTexture1, TextureSampler1, SampleUV, TexelUVSize, BlockNormalX, BlockNormalY, BlockMask);
+
+	float3 BlockSpecular[16];
+	ReadBlockRGB(RenderTexture2, TextureSampler2, SampleUV, TexelUVSize, BlockSpecular);
+
+	for (int i=0; i<16; i++)
+	{
+		BlockSpecular[i].z = round(BlockSpecular[i].z);
+	}
+
+	OutCompressTexture0_u4[ThreadId.xy + DestRect.xy] = CompressBC3BlockYCoCg(BlockBaseColor);
+	OutCompressTexture1[ThreadId.xy + DestRect.xy] = CompressBC5Block(BlockNormalX, BlockNormalY);
+	OutCompressTexture2_u4[ThreadId.xy + DestRect.xy] = CompressBC3Block(BlockSpecular, BlockMask);
+}
+
+/** Copy path used when we disable texture compression, because we need to keep the same final channel layout. */
+void CopyBaseColorNormalSpecularMaskYCoCgPS(
+	in float4 InPosition : SV_POSITION, 
+	in noperspective float2 InTexCoord : TEXCOORD0, 
+	out float4 OutColor0 : SV_Target0,
+	out float4 OutColor1 : SV_Target1,
+	out float4 OutColor2 : SV_Target2)
+{
+	float3 YCoCg = RGB2YCoCg(RenderTexture0.SampleLevel(TextureSampler0, InTexCoord, 0).xyz);
+	float3 NormalXYMask = RenderTexture1.SampleLevel(TextureSampler1, InTexCoord, 0).xyz;
+	float3 RoughnessSpecularNormalZ = RenderTexture2.SampleLevel(TextureSampler2, InTexCoord, 0).xyz;
+	
+	RoughnessSpecularNormalZ.z = round(RoughnessSpecularNormalZ.z);
+	
+	OutColor0 = float4(YCoCg.yz, 0, YCoCg.x);
+	OutColor1 = float4(NormalXYMask.xy, 0, 1);
+	OutColor2 = float4(RoughnessSpecularNormalZ, NormalXYMask.z);
+}
+
 /** Copy world height (without compression)
   * This path is used when want to copy the world height data into a thumbnail. */
 void CopyWorldHeightPS(
diff --git a/Engine/Shaders/Private/VirtualTextureMaterial.usf b/Engine/Shaders/Private/VirtualTextureMaterial.usf
--- a/Engine/Shaders/Private/VirtualTextureMaterial.usf
+++ b/Engine/Shaders/Private/VirtualTextureMaterial.usf
@@ -74,6 +74,7 @@ void FPixelShaderInOut_MainPS(
 	half3 Normal = GetVirtualTextureOutput3(MaterialParameters);
 	half WorldHeight = GetVirtualTextureOutput4(MaterialParameters);
 	float Opacity = GetVirtualTextureOutput5(MaterialParameters);
+	float Mask = GetVirtualTextureOutput6(MaterialParameters);
 #else
 	// Output is from standard material output attribute node
 	half3 BaseColor = GetMaterialBaseColor(PixelMaterialInputs);
@@ -82,6 +83,7 @@ void FPixelShaderInOut_MainPS(
 	half3 Normal = MaterialParameters.WorldNormal;
 	float WorldHeight = MaterialParameters.AbsoluteWorldPosition.z;
 	float Opacity = GetMaterialOpacity(PixelMaterialInputs);
+	float Mask = 0.f;
 #endif
 
 	// Apply debug shading
@@ -89,13 +91,10 @@ void FPixelShaderInOut_MainPS(
 
 #if defined(OUT_BASECOLOR)
 	Out.MRT[0] = float4(BaseColor, 1.f) * Opacity;
-#elif defined(OUT_BASECOLOR_NORMAL)
-	Out.MRT[0] = float4(BaseColor, 1.f) * Opacity;
-	Out.MRT[1] = float4(PackNormal(Normal).xy, 0.f, 1.f) * Opacity;
 #elif defined(OUT_BASECOLOR_NORMAL_SPECULAR)
 	float3 PackedNormal = PackNormal(Normal);
 	Out.MRT[0] = float4(BaseColor, 1.f) * Opacity;
-	Out.MRT[1] = float4(PackedNormal.xy, 0.f, 1.f) * Opacity;
+	Out.MRT[1] = float4(PackedNormal.xy, Mask, 1.f) * Opacity;
 	Out.MRT[2] = float4(Specular, Roughness, PackedNormal.z, 1.f) * Opacity;
 #elif defined(OUT_WORLDHEIGHT)
 	Out.MRT[0] = float4(PackWorldHeight(WorldHeight), 0, 0, 1);
@@ -107,9 +106,6 @@ void FPixelShaderInOut_MainPS(
 
 #if defined(OUT_BASECOLOR)
 #define PIXELSHADEROUTPUT_MRT0 1
-#elif defined(OUT_BASECOLOR_NORMAL)
-#define PIXELSHADEROUTPUT_MRT0 1
-#define PIXELSHADEROUTPUT_MRT1 1
 #elif defined(OUT_BASECOLOR_NORMAL_SPECULAR)
 #define PIXELSHADEROUTPUT_MRT0 1
 #define PIXELSHADEROUTPUT_MRT1 1
diff --git a/Engine/Source/Runtime/Engine/Classes/Materials/MaterialExpressionRuntimeVirtualTextureOutput.h b/Engine/Source/Runtime/Engine/Classes/Materials/MaterialExpressionRuntimeVirtualTextureOutput.h
--- a/Engine/Source/Runtime/Engine/Classes/Materials/MaterialExpressionRuntimeVirtualTextureOutput.h
+++ b/Engine/Source/Runtime/Engine/Classes/Materials/MaterialExpressionRuntimeVirtualTextureOutput.h
@@ -37,6 +37,10 @@ class UMaterialExpressionRuntimeVirtualTextureOutput : public UMaterialExpressio
 	UPROPERTY()
 	FExpressionInput Opacity;
 
+	/** Input for Mask to output to virtual texture. */
+	UPROPERTY()
+	FExpressionInput Mask;
+
 public:
 #if WITH_EDITOR
 	//~ Begin UMaterialExpression Interface
diff --git a/Engine/Source/Runtime/Engine/Private/Materials/MaterialExpressions.cpp b/Engine/Source/Runtime/Engine/Private/Materials/MaterialExpressions.cpp
--- a/Engine/Source/Runtime/Engine/Private/Materials/MaterialExpressions.cpp
+++ b/Engine/Source/Runtime/Engine/Private/Materials/MaterialExpressions.cpp
@@ -2025,6 +2025,10 @@ int32 UMaterialExpressionRuntimeVirtualTextureOutput::Compile(class FMaterialCom
 	{
 		CodeInput = Opacity.IsConnected() ? Opacity.Compile(Compiler) : Compiler->Constant(1.f);
 	}
+	else if (OutputIndex == 6)
+	{
+		CodeInput = Mask.IsConnected() ? Mask.Compile(Compiler) : Compiler->Constant(1.f);
+	}
 
 	Compiler->VirtualTextureOutput();
 	return Compiler->CustomOutput(this, OutputIndex, CodeInput);
@@ -2039,7 +2043,7 @@ void UMaterialExpressionRuntimeVirtualTextureOutput::GetCaption(TArray<FString>&
 
 int32 UMaterialExpressionRuntimeVirtualTextureOutput::GetNumOutputs() const
 {
-	return 6; 
+	return 7; 
 }
 
 FString UMaterialExpressionRuntimeVirtualTextureOutput::GetFunctionName() const
@@ -2100,6 +2104,7 @@ void UMaterialExpressionRuntimeVirtualTextureSample::InitOutputs()
 	Outputs.Add(FExpressionOutput(TEXT("Roughness")));
 	Outputs.Add(FExpressionOutput(TEXT("Normal")));
 	Outputs.Add(FExpressionOutput(TEXT("WorldHeight")));
+	Outputs.Add(FExpressionOutput(TEXT("Mask")));
 #endif // WITH_EDITORONLY_DATA
 }
 
@@ -2185,12 +2190,14 @@ int32 UMaterialExpressionRuntimeVirtualTextureSample::Compile(class FMaterialCom
 	bool bIsSpecularValid = false;
 	bool bIsNormalValid = false;
 	bool bIsWorldHeightValid = false;
+	bool bIsMaskValid = false;
 
 	switch (MaterialType)
 	{
 	case ERuntimeVirtualTextureMaterialType::BaseColor: bIsBaseColorValid = true; break;
 	case ERuntimeVirtualTextureMaterialType::BaseColor_Normal_Specular: bIsBaseColorValid = bIsNormalValid = bIsSpecularValid = true; break;
 	case ERuntimeVirtualTextureMaterialType::BaseColor_Normal_Specular_YCoCg: bIsBaseColorValid = bIsNormalValid = bIsSpecularValid = true; break;
+	case ERuntimeVirtualTextureMaterialType::BaseColor_Normal_Specular_Mask_YCoCg: bIsBaseColorValid = bIsNormalValid = bIsSpecularValid = bIsMaskValid = true; break;
 	case ERuntimeVirtualTextureMaterialType::WorldHeight: bIsWorldHeightValid = true; break;
 	}
 
@@ -2201,7 +2208,8 @@ int32 UMaterialExpressionRuntimeVirtualTextureSample::Compile(class FMaterialCom
 		{
 			switch (MaterialType)
 			{
-			case ERuntimeVirtualTextureMaterialType::BaseColor_Normal_Specular_YCoCg: UnpackType = EVirtualTextureUnpackType::BaseColorYCoCg; break;
+			case ERuntimeVirtualTextureMaterialType::BaseColor_Normal_Specular_YCoCg:
+			case ERuntimeVirtualTextureMaterialType::BaseColor_Normal_Specular_Mask_YCoCg: UnpackType = EVirtualTextureUnpackType::BaseColorYCoCg; break;
 			default: UnpackTarget = 0; UnpackMask = 0x7; break;
 			}
 		}
@@ -2216,7 +2224,8 @@ int32 UMaterialExpressionRuntimeVirtualTextureSample::Compile(class FMaterialCom
 			switch (MaterialType)
 			{
 			case ERuntimeVirtualTextureMaterialType::BaseColor_Normal_Specular:  UnpackTarget = 1; UnpackMask = 0x1; break;
-			case ERuntimeVirtualTextureMaterialType::BaseColor_Normal_Specular_YCoCg: UnpackTarget = 2; UnpackMask = 0x1; break;
+			case ERuntimeVirtualTextureMaterialType::BaseColor_Normal_Specular_YCoCg:
+			case ERuntimeVirtualTextureMaterialType::BaseColor_Normal_Specular_Mask_YCoCg: UnpackTarget = 2; UnpackMask = 0x1; break;
 			}
 		}
 		else
@@ -2230,7 +2239,8 @@ int32 UMaterialExpressionRuntimeVirtualTextureSample::Compile(class FMaterialCom
 			switch (MaterialType)
 			{
 			case ERuntimeVirtualTextureMaterialType::BaseColor_Normal_Specular: UnpackTarget = 1; UnpackMask = 0x2; break;
-			case ERuntimeVirtualTextureMaterialType::BaseColor_Normal_Specular_YCoCg: UnpackTarget = 2; UnpackMask = 0x2; break;
+			case ERuntimeVirtualTextureMaterialType::BaseColor_Normal_Specular_YCoCg:
+			case ERuntimeVirtualTextureMaterialType::BaseColor_Normal_Specular_Mask_YCoCg: UnpackTarget = 2; UnpackMask = 0x2; break;
 			}
 		}
 		else
@@ -2244,7 +2254,8 @@ int32 UMaterialExpressionRuntimeVirtualTextureSample::Compile(class FMaterialCom
 			switch (MaterialType)
 			{
 			case ERuntimeVirtualTextureMaterialType::BaseColor_Normal_Specular: UnpackType = EVirtualTextureUnpackType::NormalBC3BC3; break;
-			case ERuntimeVirtualTextureMaterialType::BaseColor_Normal_Specular_YCoCg: UnpackType = EVirtualTextureUnpackType::NormalBC5BC1; break;
+			case ERuntimeVirtualTextureMaterialType::BaseColor_Normal_Specular_YCoCg:
+			case ERuntimeVirtualTextureMaterialType::BaseColor_Normal_Specular_Mask_YCoCg: UnpackType = EVirtualTextureUnpackType::NormalBC5BC1; break;
 			}
 		}
 		else
@@ -2262,6 +2273,16 @@ int32 UMaterialExpressionRuntimeVirtualTextureSample::Compile(class FMaterialCom
 			return Compiler->Constant(0.f);
 		}
 		break;
+	case 5:
+		if ((bIsParameter || bIsVirtualTextureValid) && bIsMaskValid)
+		{
+			UnpackTarget = 2; UnpackMask = 0x8; break;
+		}
+		else
+		{
+			return Compiler->Constant(1.f);
+		}
+		break;
 	default:
 		return INDEX_NONE;
 	}
diff --git a/Engine/Source/Runtime/Engine/Private/VT/RuntimeVirtualTexture.cpp b/Engine/Source/Runtime/Engine/Private/VT/RuntimeVirtualTexture.cpp
--- a/Engine/Source/Runtime/Engine/Private/VT/RuntimeVirtualTexture.cpp
+++ b/Engine/Source/Runtime/Engine/Private/VT/RuntimeVirtualTexture.cpp
@@ -266,6 +266,7 @@ int32 URuntimeVirtualTexture::GetLayerCount(ERuntimeVirtualTextureMaterialType I
 	case ERuntimeVirtualTextureMaterialType::BaseColor_Normal_Specular:
 		return 2;
 	case ERuntimeVirtualTextureMaterialType::BaseColor_Normal_Specular_YCoCg:
+	case ERuntimeVirtualTextureMaterialType::BaseColor_Normal_Specular_Mask_YCoCg:
 		return 3;
 	default:
 		break;
@@ -291,6 +292,7 @@ EPixelFormat URuntimeVirtualTexture::GetLayerFormat(int32 LayerIndex) const
 			return bCompressTextures ? PF_DXT1 : PF_B8G8R8A8;
 		case ERuntimeVirtualTextureMaterialType::BaseColor_Normal_Specular:
 		case ERuntimeVirtualTextureMaterialType::BaseColor_Normal_Specular_YCoCg:
+		case ERuntimeVirtualTextureMaterialType::BaseColor_Normal_Specular_Mask_YCoCg:
 			return bCompressTextures ? PF_DXT5 : PF_B8G8R8A8;
 		case ERuntimeVirtualTextureMaterialType::WorldHeight:
 			return PF_G16;
@@ -305,6 +307,7 @@ EPixelFormat URuntimeVirtualTexture::GetLayerFormat(int32 LayerIndex) const
 		case ERuntimeVirtualTextureMaterialType::BaseColor_Normal_Specular:
 			return bCompressTextures ? PF_DXT5 : PF_B8G8R8A8;
 		case ERuntimeVirtualTextureMaterialType::BaseColor_Normal_Specular_YCoCg:
+		case ERuntimeVirtualTextureMaterialType::BaseColor_Normal_Specular_Mask_YCoCg:
 			return bCompressTextures ? PF_BC5 : PF_B8G8R8A8;
 		default:
 			break;
@@ -316,6 +319,8 @@ EPixelFormat URuntimeVirtualTexture::GetLayerFormat(int32 LayerIndex) const
 		{
 		case ERuntimeVirtualTextureMaterialType::BaseColor_Normal_Specular_YCoCg:
 			return bCompressTextures ? PF_DXT1 : PF_B8G8R8A8;
+		case ERuntimeVirtualTextureMaterialType::BaseColor_Normal_Specular_Mask_YCoCg:
+			return bCompressTextures ? PF_DXT5 : PF_B8G8R8A8;
 		default:
 			break;
 		}
@@ -335,6 +340,7 @@ bool URuntimeVirtualTexture::IsLayerSRGB(int32 LayerIndex) const
 		// Only BaseColor layer is sRGB
 		return LayerIndex == 0;
 	case ERuntimeVirtualTextureMaterialType::BaseColor_Normal_Specular_YCoCg:
+	case ERuntimeVirtualTextureMaterialType::BaseColor_Normal_Specular_Mask_YCoCg:
 	case ERuntimeVirtualTextureMaterialType::WorldHeight:
 		return false;
 	default:
@@ -348,7 +354,15 @@ bool URuntimeVirtualTexture::IsLayerSRGB(int32 LayerIndex) const
 
 bool URuntimeVirtualTexture::IsLayerYCoCg(int32 LayerIndex) const
 {
-	return MaterialType == ERuntimeVirtualTextureMaterialType::BaseColor_Normal_Specular_YCoCg && LayerIndex == 0;
+	switch (MaterialType)
+	{
+	case ERuntimeVirtualTextureMaterialType::BaseColor_Normal_Specular_YCoCg:
+	case ERuntimeVirtualTextureMaterialType::BaseColor_Normal_Specular_Mask_YCoCg:
+		return LayerIndex == 0;
+	default:
+		break;
+	}
+	return false;
 }
 
 int32 URuntimeVirtualTexture::GetEstimatedPageTableTextureMemoryKb() const
diff --git a/Engine/Source/Runtime/Engine/Public/VT/RuntimeVirtualTextureEnum.h b/Engine/Source/Runtime/Engine/Public/VT/RuntimeVirtualTextureEnum.h
--- a/Engine/Source/Runtime/Engine/Public/VT/RuntimeVirtualTextureEnum.h
+++ b/Engine/Source/Runtime/Engine/Public/VT/RuntimeVirtualTextureEnum.h
@@ -19,7 +19,8 @@ enum class ERuntimeVirtualTextureMaterialType : uint8
 	BaseColor UMETA(DisplayName = "Base Color"),
 	BaseColor_Normal_DEPRECATED UMETA(Hidden),
 	BaseColor_Normal_Specular UMETA(DisplayName = "Base Color, Normal, Roughness, Specular"),
-	BaseColor_Normal_Specular_YCoCg UMETA(DisplayName = "YCoCg Base Color, Normal, Roughness, Specular", ToolTip="Base Color is stored in YCoCg space. This requires more memory but may provide better quality."),
+	BaseColor_Normal_Specular_YCoCg UMETA(DisplayName = "YCoCg Base Color, Normal, Roughness, Specular", ToolTip = "Base Color is stored in YCoCg space. This requires more memory but may provide better quality."),
+	BaseColor_Normal_Specular_Mask_YCoCg UMETA(DisplayName = "YCoCg Base Color, Normal, Roughness, Specular, Mask", ToolTip="Base Color is stored in YCoCg space. This requires more memory but may provide better quality."),
 	WorldHeight UMETA(DisplayName = "World Height"),
 
 	Count UMETA(Hidden),
diff --git a/Engine/Source/Runtime/Renderer/Private/VT/RuntimeVirtualTextureRender.cpp b/Engine/Source/Runtime/Renderer/Private/VT/RuntimeVirtualTextureRender.cpp
--- a/Engine/Source/Runtime/Renderer/Private/VT/RuntimeVirtualTextureRender.cpp
+++ b/Engine/Source/Runtime/Renderer/Private/VT/RuntimeVirtualTextureRender.cpp
@@ -246,6 +246,7 @@ namespace RuntimeVirtualTexture
 						break;
 					case ERuntimeVirtualTextureMaterialType::BaseColor_Normal_Specular:
 					case ERuntimeVirtualTextureMaterialType::BaseColor_Normal_Specular_YCoCg:
+					case ERuntimeVirtualTextureMaterialType::BaseColor_Normal_Specular_Mask_YCoCg:
 						Process<FMaterialPolicy_BaseColorNormalSpecular>(MeshBatch, BatchElementMask, StaticMeshId, PrimitiveSceneProxy, MaterialRenderProxy, Material);
 						break;
 					case ERuntimeVirtualTextureMaterialType::WorldHeight:
@@ -420,7 +421,8 @@ namespace RuntimeVirtualTexture
 			SHADER_PARAMETER_RDG_TEXTURE_UAV(RWTexture<uint2>, OutCompressTexture0_u2)
 			SHADER_PARAMETER_RDG_TEXTURE_UAV(RWTexture<uint4>, OutCompressTexture0_u4)
 			SHADER_PARAMETER_RDG_TEXTURE_UAV(RWTexture<uint4>, OutCompressTexture1)
-			SHADER_PARAMETER_RDG_TEXTURE_UAV(RWTexture<uint2>, OutCompressTexture2)
+			SHADER_PARAMETER_RDG_TEXTURE_UAV(RWTexture<uint2>, OutCompressTexture2_u2)
+			SHADER_PARAMETER_RDG_TEXTURE_UAV(RWTexture<uint4>, OutCompressTexture2_u4)
 		END_SHADER_PARAMETER_STRUCT()
 
 		static bool ShouldCompilePermutation(FGlobalShaderPermutationParameters const& Parameters)
@@ -456,6 +458,7 @@ namespace RuntimeVirtualTexture
 	IMPLEMENT_SHADER_TYPE(template<>, FShader_VirtualTextureCompress_CS< ERuntimeVirtualTextureMaterialType::BaseColor >, TEXT("/Engine/Private/VirtualTextureCompress.usf"), TEXT("CompressBaseColorCS"), SF_Compute);
 	IMPLEMENT_SHADER_TYPE(template<>, FShader_VirtualTextureCompress_CS< ERuntimeVirtualTextureMaterialType::BaseColor_Normal_Specular >, TEXT("/Engine/Private/VirtualTextureCompress.usf"), TEXT("CompressBaseColorNormalSpecularCS"), SF_Compute);
 	IMPLEMENT_SHADER_TYPE(template<>, FShader_VirtualTextureCompress_CS< ERuntimeVirtualTextureMaterialType::BaseColor_Normal_Specular_YCoCg >, TEXT("/Engine/Private/VirtualTextureCompress.usf"), TEXT("CompressBaseColorNormalSpecularYCoCgCS"), SF_Compute);
+	IMPLEMENT_SHADER_TYPE(template<>, FShader_VirtualTextureCompress_CS< ERuntimeVirtualTextureMaterialType::BaseColor_Normal_Specular_Mask_YCoCg >, TEXT("/Engine/Private/VirtualTextureCompress.usf"), TEXT("CompressBaseColorNormalSpecularMaskYCoCgCS"), SF_Compute);
 
 
 	/** Add the BC compression pass to the graph. */
@@ -488,6 +491,9 @@ namespace RuntimeVirtualTexture
 		case ERuntimeVirtualTextureMaterialType::BaseColor_Normal_Specular_YCoCg:
 			AddCompressPass<ERuntimeVirtualTextureMaterialType::BaseColor_Normal_Specular_YCoCg>(GraphBuilder, FeatureLevel, Parameters, GroupCount);
 			break;
+		case ERuntimeVirtualTextureMaterialType::BaseColor_Normal_Specular_Mask_YCoCg:
+			AddCompressPass<ERuntimeVirtualTextureMaterialType::BaseColor_Normal_Specular_Mask_YCoCg>(GraphBuilder, FeatureLevel, Parameters, GroupCount);
+			break;
 		}
 	}
 
@@ -555,6 +561,7 @@ namespace RuntimeVirtualTexture
 	IMPLEMENT_SHADER_TYPE(template<>, FShader_VirtualTextureCopy_PS< ERuntimeVirtualTextureMaterialType::BaseColor >, TEXT("/Engine/Private/VirtualTextureCompress.usf"), TEXT("CopyBaseColorNormalSpecularPS"), SF_Pixel);
 	IMPLEMENT_SHADER_TYPE(template<>, FShader_VirtualTextureCopy_PS< ERuntimeVirtualTextureMaterialType::BaseColor_Normal_Specular >, TEXT("/Engine/Private/VirtualTextureCompress.usf"), TEXT("CopyBaseColorNormalSpecularPS"), SF_Pixel);
 	IMPLEMENT_SHADER_TYPE(template<>, FShader_VirtualTextureCopy_PS< ERuntimeVirtualTextureMaterialType::BaseColor_Normal_Specular_YCoCg >, TEXT("/Engine/Private/VirtualTextureCompress.usf"), TEXT("CopyBaseColorNormalSpecularYCoCgPS"), SF_Pixel);
+	IMPLEMENT_SHADER_TYPE(template<>, FShader_VirtualTextureCopy_PS< ERuntimeVirtualTextureMaterialType::BaseColor_Normal_Specular_Mask_YCoCg >, TEXT("/Engine/Private/VirtualTextureCompress.usf"), TEXT("CopyBaseColorNormalSpecularMaskYCoCgPS"), SF_Pixel);
 	IMPLEMENT_SHADER_TYPE(template<>, FShader_VirtualTextureCopy_PS< ERuntimeVirtualTextureMaterialType::WorldHeight >, TEXT("/Engine/Private/VirtualTextureCompress.usf"), TEXT("CopyWorldHeightPS"), SF_Pixel);
 
 
@@ -602,6 +609,9 @@ namespace RuntimeVirtualTexture
 		case ERuntimeVirtualTextureMaterialType::BaseColor_Normal_Specular_YCoCg:
 			AddCopyPass<ERuntimeVirtualTextureMaterialType::BaseColor_Normal_Specular_YCoCg>(GraphBuilder, FeatureLevel, Parameters, TextureSize);
 			break;
+		case ERuntimeVirtualTextureMaterialType::BaseColor_Normal_Specular_Mask_YCoCg:
+			AddCopyPass<ERuntimeVirtualTextureMaterialType::BaseColor_Normal_Specular_Mask_YCoCg>(GraphBuilder, FeatureLevel, Parameters, TextureSize);
+			break;
 		}
 	}
 
@@ -613,6 +623,7 @@ namespace RuntimeVirtualTexture
 		case ERuntimeVirtualTextureMaterialType::BaseColor:
 		case ERuntimeVirtualTextureMaterialType::BaseColor_Normal_Specular:
 		case ERuntimeVirtualTextureMaterialType::BaseColor_Normal_Specular_YCoCg:
+		case ERuntimeVirtualTextureMaterialType::BaseColor_Normal_Specular_Mask_YCoCg:
 			AddCopyPass<ERuntimeVirtualTextureMaterialType::BaseColor>(GraphBuilder, FeatureLevel, Parameters, TextureSize);
 			break;
 		case ERuntimeVirtualTextureMaterialType::WorldHeight:
@@ -630,7 +641,7 @@ namespace RuntimeVirtualTexture
 			bRenderPass = OutputTexture0 != nullptr;
 			bCopyThumbnailPass = bRenderPass && bIsThumbnails;
 			bCompressPass = bRenderPass && !bCopyThumbnailPass && (OutputTexture0->GetFormat() == PF_DXT1 || OutputTexture0->GetFormat() == PF_DXT5 || OutputTexture0->GetFormat() == PF_BC5);
-			bCopyPass = bRenderPass && !bCopyThumbnailPass && !bCompressPass && (MaterialType == ERuntimeVirtualTextureMaterialType::BaseColor_Normal_Specular || MaterialType == ERuntimeVirtualTextureMaterialType::BaseColor_Normal_Specular_YCoCg);
+			bCopyPass = bRenderPass && !bCopyThumbnailPass && !bCompressPass && (MaterialType == ERuntimeVirtualTextureMaterialType::BaseColor_Normal_Specular || MaterialType == ERuntimeVirtualTextureMaterialType::BaseColor_Normal_Specular_YCoCg || MaterialType == ERuntimeVirtualTextureMaterialType::BaseColor_Normal_Specular_Mask_YCoCg);
 		
 			switch (MaterialType)
 			{
@@ -681,7 +692,31 @@ namespace RuntimeVirtualTexture
 				{
 					OutputAlias0 = CompressTexture0_u4 = GraphBuilder.CreateTexture(FPooledRenderTargetDesc::Create2DDesc(TextureSize / 4, PF_R32G32B32A32_UINT, FClearValueBinding::None, TexCreate_None, TexCreate_UAV, false), TEXT("CompressTexture0"));
 					OutputAlias1 = CompressTexture1 = GraphBuilder.CreateTexture(FPooledRenderTargetDesc::Create2DDesc(TextureSize / 4, PF_R32G32B32A32_UINT, FClearValueBinding::None, TexCreate_None, TexCreate_UAV, false), TEXT("CompressTexture1"));
-					OutputAlias2 = CompressTexture2 = GraphBuilder.CreateTexture(FPooledRenderTargetDesc::Create2DDesc(TextureSize / 4, PF_R32G32_UINT, FClearValueBinding::None, TexCreate_None, TexCreate_UAV, false), TEXT("CompressTexture2"));
+					OutputAlias2 = CompressTexture2_u2 = GraphBuilder.CreateTexture(FPooledRenderTargetDesc::Create2DDesc(TextureSize / 4, PF_R32G32_UINT, FClearValueBinding::None, TexCreate_None, TexCreate_UAV, false), TEXT("CompressTexture2"));
+				}
+				if (bCopyPass)
+				{
+					OutputAlias0 = CopyTexture0 = GraphBuilder.CreateTexture(FPooledRenderTargetDesc::Create2DDesc(TextureSize, PF_B8G8R8A8, FClearValueBinding::Black, TexCreate_None, TexCreate_RenderTargetable | TexCreate_ShaderResource, false), TEXT("CopyTexture0"));
+					OutputAlias1 = CopyTexture1 = GraphBuilder.CreateTexture(FPooledRenderTargetDesc::Create2DDesc(TextureSize, PF_B8G8R8A8, FClearValueBinding::Black, TexCreate_None, TexCreate_RenderTargetable | TexCreate_ShaderResource, false), TEXT("CopyTexture1"));
+					OutputAlias2 = CopyTexture2 = GraphBuilder.CreateTexture(FPooledRenderTargetDesc::Create2DDesc(TextureSize, PF_B8G8R8A8, FClearValueBinding::Black, TexCreate_None, TexCreate_RenderTargetable | TexCreate_ShaderResource, false), TEXT("CopyTexture2"));
+				}
+				if (bCopyThumbnailPass)
+				{
+					OutputAlias0 = CopyTexture0 = GraphBuilder.CreateTexture(FPooledRenderTargetDesc::Create2DDesc(TextureSize, PF_B8G8R8A8, FClearValueBinding::Black, TexCreate_SRGB, TexCreate_RenderTargetable | TexCreate_ShaderResource, false), TEXT("CopyTexture0"));
+				}
+				break;
+			case ERuntimeVirtualTextureMaterialType::BaseColor_Normal_Specular_Mask_YCoCg:
+				if (bRenderPass)
+				{
+					OutputAlias0 = RenderTexture0 = GraphBuilder.CreateTexture(FPooledRenderTargetDesc::Create2DDesc(TextureSize, PF_B8G8R8A8, FClearValueBinding::Black, TexCreate_SRGB, TexCreate_RenderTargetable | TexCreate_ShaderResource, false), TEXT("RenderTexture0"));
+					RenderTexture1 = GraphBuilder.CreateTexture(FPooledRenderTargetDesc::Create2DDesc(TextureSize, PF_B8G8R8A8, FClearValueBinding::DefaultNormal8Bit, TexCreate_None, TexCreate_RenderTargetable | TexCreate_ShaderResource, false), TEXT("RenderTexture1"));
+					RenderTexture2 = GraphBuilder.CreateTexture(FPooledRenderTargetDesc::Create2DDesc(TextureSize, PF_B8G8R8A8, FClearValueBinding::DefaultNormal8Bit, TexCreate_None, TexCreate_RenderTargetable | TexCreate_ShaderResource, false), TEXT("RenderTexture2"));
+				}
+				if (bCompressPass)
+				{
+					OutputAlias0 = CompressTexture0_u4 = GraphBuilder.CreateTexture(FPooledRenderTargetDesc::Create2DDesc(TextureSize / 4, PF_R32G32B32A32_UINT, FClearValueBinding::None, TexCreate_None, TexCreate_UAV, false), TEXT("CompressTexture0"));
+					OutputAlias1 = CompressTexture1 = GraphBuilder.CreateTexture(FPooledRenderTargetDesc::Create2DDesc(TextureSize / 4, PF_R32G32B32A32_UINT, FClearValueBinding::None, TexCreate_None, TexCreate_UAV, false), TEXT("CompressTexture1"));
+					OutputAlias2 = CompressTexture2_u4 = GraphBuilder.CreateTexture(FPooledRenderTargetDesc::Create2DDesc(TextureSize / 4, PF_R32G32B32A32_UINT, FClearValueBinding::None, TexCreate_None, TexCreate_UAV, false), TEXT("CompressTexture2"));
 				}
 				if (bCopyPass)
 				{
@@ -720,7 +755,8 @@ namespace RuntimeVirtualTexture
 		FRDGTextureRef CompressTexture0_u2 = nullptr;
 		FRDGTextureRef CompressTexture0_u4 = nullptr;
 		FRDGTextureRef CompressTexture1 = nullptr;
-		FRDGTextureRef CompressTexture2 = nullptr;
+		FRDGTextureRef CompressTexture2_u2 = nullptr;
+		FRDGTextureRef CompressTexture2_u4 = nullptr;
 		FRDGTextureRef CopyTexture0 = nullptr;
 		FRDGTextureRef CopyTexture1 = nullptr;
 		FRDGTextureRef CopyTexture2 = nullptr;
@@ -854,7 +890,8 @@ namespace RuntimeVirtualTexture
 			PassParameters->OutCompressTexture0_u2 = GraphSetup.CompressTexture0_u2 ? GraphBuilder.CreateUAV(FRDGTextureUAVDesc(GraphSetup.CompressTexture0_u2)) : nullptr;
 			PassParameters->OutCompressTexture0_u4 = GraphSetup.CompressTexture0_u4 ? GraphBuilder.CreateUAV(FRDGTextureUAVDesc(GraphSetup.CompressTexture0_u4)) : nullptr;
 			PassParameters->OutCompressTexture1 = GraphSetup.CompressTexture1 ? GraphBuilder.CreateUAV(FRDGTextureUAVDesc(GraphSetup.CompressTexture1)) : nullptr;
-			PassParameters->OutCompressTexture2 = GraphSetup.CompressTexture2 ? GraphBuilder.CreateUAV(FRDGTextureUAVDesc(GraphSetup.CompressTexture2)) : nullptr;
+			PassParameters->OutCompressTexture2_u2 = GraphSetup.CompressTexture2_u2 ? GraphBuilder.CreateUAV(FRDGTextureUAVDesc(GraphSetup.CompressTexture2_u2)) : nullptr;
+			PassParameters->OutCompressTexture2_u4 = GraphSetup.CompressTexture2_u4 ? GraphBuilder.CreateUAV(FRDGTextureUAVDesc(GraphSetup.CompressTexture2_u4)) : nullptr;
 
 			AddCompressPass(GraphBuilder, View->GetFeatureLevel(), PassParameters, TextureSize, MaterialType);
 		}
