Try to fix a bug of mobile devices encounter black or flickering textures

diff --git a/Engine/Source/Runtime/Renderer/Private/MobileBasePass.cpp b/Engine/Source/Runtime/Renderer/Private/MobileBasePass.cpp
--- a/Engine/Source/Runtime/Renderer/Private/MobileBasePass.cpp
+++ b/Engine/Source/Runtime/Renderer/Private/MobileBasePass.cpp
@@ -254,8 +254,16 @@ ELightMapPolicyType MobileBasePass::SelectMeshLightmapPolicy(
 
 		const bool bMovableWithCSM = bUseMovableLight && MobileDirectionalLight->ShouldRenderViewIndependentWholeSceneShadows() && bPrimReceivesCSM;
 
-		const FIndirectLightingCacheAllocation* IndirectLightingCacheAllocation = PrimitiveSceneProxy->GetPrimitiveSceneInfo()->IndirectLightingCacheAllocation;
-		const bool bPrimitiveUsesILC = PrimitiveSceneProxy->GetIndirectLightingCacheQuality() != ILCQ_Off;
+		const bool bPrimitiveUsesILC = PrimitiveSceneProxy
+									&& (PrimitiveSceneProxy->IsMovable() || PrimitiveSceneProxy->NeedsUnbuiltPreviewLighting() || PrimitiveSceneProxy->GetLightmapType() == ELightmapType::ForceVolumetric)
+									&& PrimitiveSceneProxy->WillEverBeLit()
+									&& PrimitiveSceneProxy->GetIndirectLightingCacheQuality() != ILCQ_Off;
+
+		const bool bHasValidVLM = Scene && Scene->VolumetricLightmapSceneData.HasData()
+								&& ReadOnlyCVARCache.bAllowStaticLighting;
+
+		const bool bHasValidILC = Scene && Scene->PrecomputedLightVolumes.Num() > 0
+								&& IsIndirectLightingCacheAllowed(FeatureLevel);
 
 		if (LightMapInteraction.GetType() == LMIT_Texture && ReadOnlyCVARCache.bAllowStaticLighting && ReadOnlyCVARCache.bEnableLowQualityLightmaps)
 		{
@@ -299,18 +307,7 @@ ELightMapPolicyType MobileBasePass::SelectMeshLightmapPolicy(
 				}
 			}
 		}
-		else if (IsIndirectLightingCacheAllowed(FeatureLevel) /* implies bAllowStaticLighting*/
-			&& PrimitiveSceneProxy
-			// Movable objects need to get their GI from the indirect lighting cache
-			&& (PrimitiveSceneProxy->IsMovable() 
-				// Static objects with lightmap type set to ForceVolumetric also need to get GI from ILC
-				|| LightMapInteraction.GetType() == LMIT_GlobalVolume)
-			// if there is no valid ILC data, skip SH calculation by using Movable Directional Light permutation.
-			&& ((IndirectLightingCacheAllocation && IndirectLightingCacheAllocation->IsValid())
-				// if there is no valid VLM data, skip SH calculation by using Movable Directional Light permutation.
-				|| (Scene && Scene->VolumetricLightmapSceneData.HasData()))
-			// Use the indirect lighting cache shaders if the object does not turn off ILC 
-			&& bPrimitiveUsesILC)
+		else if ((bHasValidVLM || bHasValidILC) && bPrimitiveUsesILC)
 		{
 			if (bUseMovableLight)
 			{
