Fix for water refraction pass in splitscreen

#rnx

diff --git a/Engine/Source/Runtime/Renderer/Private/SingleLayerWaterRendering.cpp b/Engine/Source/Runtime/Renderer/Private/SingleLayerWaterRendering.cpp
--- a/Engine/Source/Runtime/Renderer/Private/SingleLayerWaterRendering.cpp
+++ b/Engine/Source/Runtime/Renderer/Private/SingleLayerWaterRendering.cpp
@@ -436,7 +436,7 @@ public:
 //////////////////////////////////////////////////////////////////////////
 
 
-
+PRAGMA_DISABLE_OPTIMIZATION
 
 void FDeferredShadingSceneRenderer::CopySingleLayerWaterTextures(FRHICommandListImmediate& RHICmdList, FSingleLayerWaterPassData& PassData)
 {
@@ -488,8 +488,8 @@ void FDeferredShadingSceneRenderer::CopySingleLayerWaterTextures(FRHICommandList
     
 		const FIntRect RefractionViewRect = FIntRect(FIntPoint::DivideAndRoundDown(View.ViewRect.Min, RefractionDownsampleFactor), FIntPoint::DivideAndRoundDown(View.ViewRect.Max, RefractionDownsampleFactor));
 
-		PassData.ViewData[ViewIndex].SceneWithoutSingleLayerWaterMaxUV.X = (RefractionViewRect.Size().X - 0.5f) / RefractionResolution.X;
-		PassData.ViewData[ViewIndex].SceneWithoutSingleLayerWaterMaxUV.Y = (RefractionViewRect.Size().Y - 0.5f) / RefractionResolution.Y;
+		PassData.ViewData[ViewIndex].SceneWithoutSingleLayerWaterMaxUV.X = (RefractionViewRect.Max.X - 0.5f) / RefractionResolution.X; 
+		PassData.ViewData[ViewIndex].SceneWithoutSingleLayerWaterMaxUV.Y = (RefractionViewRect.Max.Y - 0.5f) / RefractionResolution.Y;
 
 	    FPixelShaderUtils::AddFullscreenPass(
 		    GraphBuilder,
@@ -509,7 +509,7 @@ void FDeferredShadingSceneRenderer::CopySingleLayerWaterTextures(FRHICommandList
 
 	GraphBuilder.Execute();
 }
-
+PRAGMA_ENABLE_OPTIMIZATION
 void FDeferredShadingSceneRenderer::BeginRenderingWaterGBuffer(FRHICommandList& RHICmdList, FExclusiveDepthStencil::Type DepthStencilAccess, bool bBindQuadOverdrawBuffers, EShaderPlatform InShaderPlatform)
 {
 	SCOPED_DRAW_EVENT(RHICmdList, BeginRenderingWaterGBuffer);
