Fixed RTGI is disabled in split screen mode

diff --git a/Engine/Source/Runtime/Renderer/Private/IndirectLightRendering.cpp b/Engine/Source/Runtime/Renderer/Private/IndirectLightRendering.cpp
--- a/Engine/Source/Runtime/Renderer/Private/IndirectLightRendering.cpp
+++ b/Engine/Source/Runtime/Renderer/Private/IndirectLightRendering.cpp
@@ -113,12 +113,6 @@ void FDeferredShadingSceneRenderer::RenderDiffuseIndirectAndAmbientOcclusion(FRH
 		return;
 	}
 
-	if (Views.Num() > 1 && ShouldRenderRayTracingAmbientOcclusion(Views[0]))
-	{
-		//#dxr_todo: enable RTAO in multiview mode
-		return;
-	}
-
 	FRDGBuilder GraphBuilder(RHICmdListImmediate);
 	
 	FSceneTextureParameters SceneTextures;
@@ -132,7 +126,7 @@ void FDeferredShadingSceneRenderer::RenderDiffuseIndirectAndAmbientOcclusion(FRH
 		const bool bApplyRTGI = ShouldRenderRayTracingGlobalIllumination(View);
 		const bool bApplySSGI = ShouldRenderScreenSpaceDiffuseIndirect(View);
 		const bool bApplySSAO = SceneContext.bScreenSpaceAOIsValid;
-		const bool bApplyRTAO = ShouldRenderRayTracingAmbientOcclusion(View);
+		const bool bApplyRTAO = ShouldRenderRayTracingAmbientOcclusion(View) && Views.Num() == 1; //#dxr_todo: enable RTAO in multiview mode
 
 		int32 DenoiseMode = CVarDiffuseIndirectDenoiser.GetValueOnRenderThread();
 
