Sorted-deferred execution of ray tracing global illumination pass' "final gather" mode. Based on performance work done by [REDACTED_PERSON] [at] NVIDIA.

diff --git a/Engine/Shaders/Private/RayTracing/RayTracingCreateGatherPointsRGS.usf b/Engine/Shaders/Private/RayTracing/RayTracingCreateGatherPointsRGS.usf
--- a/Engine/Shaders/Private/RayTracing/RayTracingCreateGatherPointsRGS.usf
+++ b/Engine/Shaders/Private/RayTracing/RayTracingCreateGatherPointsRGS.usf
@@ -1,5 +1,14 @@
 // Copyright 1998-2019 Epic Games, Inc. All Rights Reserved.
 
+// EDeferredMaterialMode
+#define DEFERRED_MATERIAL_MODE_NONE   0
+#define DEFERRED_MATERIAL_MODE_GATHER 1
+#define DEFERRED_MATERIAL_MODE_SHADE  2
+
+#ifndef DIM_DEFERRED_MATERIAL_MODE
+#define DIM_DEFERRED_MATERIAL_MODE DEFERRED_MATERIAL_MODE_NONE
+#endif
+
 #include "../Common.ush"
 #include "../RectLight.ush"
 #include "../DeferredShadingCommon.ush"
@@ -32,7 +41,12 @@ uint EvalSkyLight;
 uint UseRussianRoulette;
 
 int2 GatherPointsResolution;
+uint2 TileAlignedResolution;
+uint SortTileSize;
+
 RWStructuredBuffer<FGatherPoints> RWGatherPointsBuffer;
+// Material buffer and tile size are used for sorted materials.
+RWStructuredBuffer<FDeferredMaterialPayload> MaterialBuffer;
 
 uint DispatchThreadIdToLinearIndex(uint2 DispatchThreadId)
 {
@@ -61,19 +75,43 @@ void GenerateCosineNormalRay(
 	RayPdf = Direction_Tangent.w;
 }
 
-RAY_TRACING_ENTRY_RAYGEN(RayTracingCreateGatherPointsRGS)
+bool ShouldTraceRay(uint2 DispatchThreadId, FGBufferData GBufferData)
 {
-	uint2 DispatchThreadId = DispatchRaysIndex().xy + View.ViewRectMin + uint2(RenderTileOffsetX, RenderTileOffsetY);
-	uint GatherPointsIndex = DispatchThreadIdToLinearIndex(DispatchThreadId);
-#if 1
-	// Default initialization
-	for (uint Index = 0; Index < SamplesPerPixel; ++Index)
+	bool bTraceRay = true;
+
+	// Cull based on shading model
+	uint ShadingModelID = GBufferData.ShadingModelID;
+	if (ShadingModelID == SHADINGMODELID_UNLIT
+		|| ShadingModelID == SHADINGMODELID_TWOSIDED_FOLIAGE
+		)
 	{
-		uint SampleIndexOffset = SampleIndex + Index;
-		FGatherSample GatherSample = (FGatherSample)0.0;
-		WriteGatherSample(RWGatherPointsBuffer, GatherPointsIndex, SampleIndexOffset, GatherPointsResolution.x * GatherPointsResolution.y, GatherSample);
+		bTraceRay = false;
 	}
-#endif
+
+	// Cull based on diffuse threshold
+	float3 DiffuseColor = GBufferData.DiffuseColor;
+	if (Luminance(DiffuseColor) < DiffuseThreshold)
+	{
+		bTraceRay = false;
+	}
+
+	// Cull based on ViewRect
+	if (any(DispatchThreadId >= GatherPointsResolution))
+	{
+		bTraceRay = false;
+	}
+
+	return bTraceRay;
+}
+
+RAY_TRACING_ENTRY_RAYGEN(RayTracingCreateGatherPointsTraceRGS)
+{
+	uint2 DispatchThreadId = DispatchRaysIndex().xy + View.ViewRectMin + uint2(RenderTileOffsetX, RenderTileOffsetY);
+
+	FDeferredMaterialPayload DeferredMaterialPayload;
+	DeferredMaterialPayload.SortKey = RAY_TRACING_DEFERRED_MATERIAL_KEY_INVALID;
+	DeferredMaterialPayload.PixelCoordinates = DispatchThreadId.x | (DispatchThreadId.y << 16);
+
 	// Get G-Buffer surface data
 	uint2 PixelCoord = GetPixelCoord(DispatchThreadId, UpscaleFactor);
 	float2 InvBufferSize = View.BufferSizeAndInvSize.zw;
@@ -85,25 +123,102 @@ RAY_TRACING_ENTRY_RAYGEN(RayTracingCreateGatherPointsRGS)
 		GBufferData.DiffuseColor = GBufferData.StoredBaseColor;
 	}
 
-	// Cull based on shading model
-	uint ShadingModelID = GBufferData.ShadingModelID;
-	if (ShadingModelID == SHADINGMODELID_UNLIT
-		|| ShadingModelID == SHADINGMODELID_TWOSIDED_FOLIAGE
-		)
+	// Construct position, normal, and view direction
+	float DeviceZ = SceneDepthBuffer.Load(int3(PixelCoord, 0)).r;
+	float3 WorldPosition;
+	float3 CameraDirection;
+	ReconstructWorldPositionAndCameraDirectionFromDeviceZ(PixelCoord, DeviceZ, WorldPosition, CameraDirection);
+	float3 WorldNormal = GBufferData.WorldNormal;
+
+	FQuasiRandom QuasiRandom;
+	QuasiRandom_Initialize(QuasiRandom, PixelCoord, 0, SamplesPerPixel);
+
+	bool bTraceRay = ShouldTraceRay(DispatchThreadId, GBufferData);
+	if (bTraceRay)
+	{
+		float3 RandSample = QuasiRandom_Sample3(QuasiRandom);
+
+		// Initialize ray
+		RayDesc Ray;
+		float RayPdf = 1.0;
+		GenerateCosineNormalRay(WorldPosition, WorldNormal, RandSample.yz, Ray.Origin, Ray.Direction, Ray.TMin, Ray.TMax, RayPdf);
+		ApplyCameraRelativeDepthBias(Ray, PixelCoord, DeviceZ, WorldNormal, MaxNormalBias);
+
+		uint RayFlags = 0;
+#if !ENABLE_TWO_SIDED_GEOMETRY
+		RayFlags |= RAY_FLAG_CULL_BACK_FACING_TRIANGLES;
+#endif
+		const uint InstanceInclusionMask = RAY_TRACING_MASK_OPAQUE;
+		DeferredMaterialPayload.SortKey = RAY_TRACING_DEFERRED_MATERIAL_KEY_RAY_MISS;
+		TraceRay(
+			TLAS,								// AccelerationStructure
+			RayFlags,
+			InstanceInclusionMask,				// InstanceInclusionMask
+			RAY_TRACING_SHADER_SLOT_MATERIAL,	// RayContributionToHitGroupIndex
+			RAY_TRACING_NUM_SHADER_SLOTS,		// MultiplierForGeometryContributionToShaderIndex
+			0,									// MissShaderIndex
+			Ray,								// RayDesc
+			DeferredMaterialPayload
+		);
+	}
+
+	uint LinearStoreIndex = DispatchRaysIndex().y * TileAlignedResolution.x + DispatchRaysIndex().x;
+	uint StoreIndex = LinearStoreIndex;
+
+	if (SortTileSize > 0) // Store payload in tile order
+	{
+		uint2 Block = DispatchRaysIndex().xy / SortTileSize;
+		uint2 Thread = DispatchRaysIndex().xy % SortTileSize;
+
+		uint IndexInsideBlock = Thread.y * SortTileSize + Thread.x;
+		uint ElementsPerBlock = SortTileSize * SortTileSize;
+		uint BlocksPerRow = TileAlignedResolution.x / SortTileSize;
+		uint BlockIndex = Block.y * BlocksPerRow + Block.x;
+
+		StoreIndex = BlockIndex * ElementsPerBlock + IndexInsideBlock;
+	}
+
+	MaterialBuffer[StoreIndex] = DeferredMaterialPayload;
+}
+
+RAY_TRACING_ENTRY_RAYGEN(RayTracingCreateGatherPointsRGS)
+{
+	FDeferredMaterialPayload DeferredMaterialPayload;
+	DeferredMaterialPayload.SortKey = RAY_TRACING_DEFERRED_MATERIAL_KEY_INVALID;
+
+	uint2 DispatchThreadId;
+	if (DIM_DEFERRED_MATERIAL_MODE == DEFERRED_MATERIAL_MODE_SHADE)
 	{
-		return;
+		DeferredMaterialPayload = MaterialBuffer[DispatchRaysIndex().x];
+		DispatchThreadId.x = DeferredMaterialPayload.PixelCoordinates & 0xFFFF;
+		DispatchThreadId.y = DeferredMaterialPayload.PixelCoordinates >> 16;
+
+		if (DeferredMaterialPayload.SortKey == RAY_TRACING_DEFERRED_MATERIAL_KEY_INVALID)
+		{
+			return;
+		}
+	}
+	else
+	{
+		DispatchThreadId = DispatchRaysIndex().xy + View.ViewRectMin + uint2(RenderTileOffsetX, RenderTileOffsetY);
 	}
 
+	uint GatherPointsIndex = DispatchThreadIdToLinearIndex(DispatchThreadId);
+
+	// Default initialization
+	FGatherSample GatherSample = (FGatherSample)0.0;
+	WriteGatherSample(RWGatherPointsBuffer, GatherPointsIndex, SampleIndex, GatherPointsResolution.x * GatherPointsResolution.y, GatherSample);
+
+	// Get G-Buffer surface data
+	uint2 PixelCoord = GetPixelCoord(DispatchThreadId, UpscaleFactor);
+	float2 InvBufferSize = View.BufferSizeAndInvSize.zw;
+	float2 UV = (float2(PixelCoord) + 0.5) * InvBufferSize;
+	FGBufferData GBufferData = GetGBufferDataFromSceneTextures(UV);
 	// Remap DiffuseColor when using SubsurfaceProfile (GBuffer decoding replaces with 100% albedo)
 	if (UseSubsurfaceProfile(GBufferData.ShadingModelID))
 	{
 		GBufferData.DiffuseColor = GBufferData.StoredBaseColor;
 	}
-	float3 DiffuseColor = GBufferData.DiffuseColor;
-	if (Luminance(DiffuseColor) < DiffuseThreshold)
-	{
-		return;
-	}
 
 	// Construct position, normal, and view direction
 	float DeviceZ = SceneDepthBuffer.Load(int3(PixelCoord, 0)).r;
@@ -114,11 +229,11 @@ RAY_TRACING_ENTRY_RAYGEN(RayTracingCreateGatherPointsRGS)
 
 	FQuasiRandom QuasiRandom;
 	QuasiRandom_Initialize(QuasiRandom, PixelCoord, 0, SamplesPerPixel);
-	//QuasiRandom_Initialize(QuasiRandom, DispatchThreadId, 0, SamplesPerPixel);
 
-	for (uint Index = 0; Index < SamplesPerPixel; ++Index)
+	bool bTraceRay = ShouldTraceRay(DispatchThreadId, GBufferData);
+	if (bTraceRay)
 	{
-		uint SampleIndexOffset = (SampleIndex + Index) % GatherSamplesPerPixel;
+		uint SampleIndexOffset = SampleIndex % GatherSamplesPerPixel;
 #if 0
 		// Adaptive rejection optimization
 		{
@@ -130,9 +245,6 @@ RAY_TRACING_ENTRY_RAYGEN(RayTracingCreateGatherPointsRGS)
 		float3 Position = 0;
 		float3 Normal = 0;
 
-		QuasiRandom_Update(QuasiRandom, Index, SamplesPerPixel);
-
-		float3 RayThroughput = 1.0;
 		float3 RandSample = QuasiRandom_Sample3(QuasiRandom);
 
 		// Initialize ray
@@ -143,8 +255,6 @@ RAY_TRACING_ENTRY_RAYGEN(RayTracingCreateGatherPointsRGS)
 		RayPdf = 1.0;
 		ApplyCameraRelativeDepthBias(Ray, PixelCoord, DeviceZ, WorldNormal, MaxNormalBias);
 
-		{
-			// Cast ray
 			uint RayFlags = 0;
 			const uint InstanceInclusionMask = RAY_TRACING_MASK_OPAQUE;
 			const bool bEnableSkyLightContribution = true;
@@ -152,6 +262,25 @@ RAY_TRACING_ENTRY_RAYGEN(RayTracingCreateGatherPointsRGS)
 #if !ENABLE_TWO_SIDED_GEOMETRY
 			RayFlags |= RAY_FLAG_CULL_BACK_FACING_TRIANGLES;
 #endif
+		if (DIM_DEFERRED_MATERIAL_MODE == DEFERRED_MATERIAL_MODE_SHADE)
+		{
+			if (DeferredMaterialPayload.SortKey < RAY_TRACING_DEFERRED_MATERIAL_KEY_RAY_MISS)
+			{
+				const float ShortRayLength = 1.0f; // 1cm is arbitrarily chosen
+				Ray.TMin = max(0, DeferredMaterialPayload.HitT - ShortRayLength * 0.5f);
+
+				// TMax can only be used if alpha masked materials are handled during initial ray tracing (material gathering)
+				// Ray.TMax = DeferredMaterialPayload.HitT + ShortRayLength * 0.5f;
+			}
+			else
+			{
+				Ray.TMax = 0;
+			}
+		}
+
+		{
+			// Cast ray
+			const bool bDisableSkyLightContribution = false;
 
 			FRayCone RayCone = (FRayCone)0;
 			FMaterialClosestHitPayload Payload = TraceMaterialRay(
@@ -163,7 +292,7 @@ RAY_TRACING_ENTRY_RAYGEN(RayTracingCreateGatherPointsRGS)
 				bEnableSkyLightContribution);
 
 			// TODO: impossible based on formulation
-			if (Payload.HitT > MaxRayDistanceForGI) continue;
+			if (Payload.HitT > MaxRayDistanceForGI) return;
 
 			// Update intersection
 			Ray.Origin += Ray.Direction * Payload.HitT;
diff --git a/Engine/Source/Runtime/Renderer/Private/RayTracing/RayTracingGlobalIllumination.cpp b/Engine/Source/Runtime/Renderer/Private/RayTracing/RayTracingGlobalIllumination.cpp
--- a/Engine/Source/Runtime/Renderer/Private/RayTracing/RayTracingGlobalIllumination.cpp
+++ b/Engine/Source/Runtime/Renderer/Private/RayTracing/RayTracingGlobalIllumination.cpp
@@ -22,6 +22,7 @@
 #include "BlueNoise.h"
 #include "SceneTextureParameters.h"
 #include "RayTracingDefinitions.h"
+#include "RayTracingDeferredMaterials.h"
 
 static TAutoConsoleVariable<int32> CVarRayTracingGlobalIllumination(
 	TEXT("r.RayTracing.GlobalIllumination"),
@@ -136,11 +137,36 @@ static TAutoConsoleVariable<int32> CVarRayTracingGlobalIlluminationEnableFinalGa
 
 static float GRayTracingGlobalIlluminationFinalGatherDistance = 10.0;
 static FAutoConsoleVariableRef CVarRayTracingGlobalIlluminationFinalGatherDistance(
-	TEXT("r.RayTracing.GlobalIllumination.FinalGatherDistance"),
+	TEXT("r.RayTracing.GlobalIllumination.FinalGather.Distance"),
 	GRayTracingGlobalIlluminationFinalGatherDistance,
 	TEXT("Maximum world-space distance for valid, reprojected final gather points (default = 10)")
 );
 
+static TAutoConsoleVariable<int32> CVarRayTracingGlobalIlluminationFinalGatherSortMaterials(
+	TEXT("r.RayTracing.GlobalIllumination.FinalGather.SortMaterials"),
+	1,
+	TEXT("Sets whether refected materials will be sorted before shading\n")
+	TEXT("0: Disabled\n ")
+	TEXT("1: Enabled, using Trace->Sort->Trace (Default)\n"),
+	ECVF_RenderThreadSafe);
+
+static TAutoConsoleVariable<int32> CVarRayTracingGlobalIlluminationFinalGatherSortTileSize(
+	TEXT("r.RayTracing.GlobalIllumination.FinalGather.SortTileSize"),
+	64,
+	TEXT("Size of pixel tiles for sorted global illumination (default = 64)\n"),
+	ECVF_RenderThreadSafe);
+
+static TAutoConsoleVariable<int32> CVarRayTracingGlobalIlluminationFinalGatherSortSize(
+	TEXT("r.RayTracing.GlobalIllumination.FinalGather.SortSize"),
+	5,
+	TEXT("Size of horizon for material ID sort\n")
+	TEXT("0: Disabled\n")
+	TEXT("1: 256 Elements\n")
+	TEXT("2: 512 Elements\n")
+	TEXT("3: 1024 Elements\n")
+	TEXT("4: 2048 Elements\n")
+	TEXT("5: 4096 Elements (Default)\n"),
+	ECVF_RenderThreadSafe);
 
 DECLARE_GPU_STAT_NAMED(RayTracingGIBruteForce, TEXT("Ray Tracing GI: Brute Force"));
 DECLARE_GPU_STAT_NAMED(RayTracingGIFinalGather, TEXT("Ray Tracing GI: Final Gather"));
@@ -357,10 +383,11 @@ class FRayTracingGlobalIlluminationCreateGatherPointsRGS : public FGlobalShader
 	DECLARE_GLOBAL_SHADER(FRayTracingGlobalIlluminationCreateGatherPointsRGS)
 	SHADER_USE_ROOT_PARAMETER_STRUCT(FRayTracingGlobalIlluminationCreateGatherPointsRGS, FGlobalShader)
 
-	class FUseAttenuationTermDim : SHADER_PERMUTATION_BOOL("USE_ATTENUATION_TERM");
+		class FUseAttenuationTermDim : SHADER_PERMUTATION_BOOL("USE_ATTENUATION_TERM");
 	class FEnableTwoSidedGeometryDim : SHADER_PERMUTATION_BOOL("ENABLE_TWO_SIDED_GEOMETRY");
+	class FDeferredMaterialMode : SHADER_PERMUTATION_ENUM_CLASS("DIM_DEFERRED_MATERIAL_MODE", EDeferredMaterialMode);
 
-	using FPermutationDomain = TShaderPermutationDomain<FUseAttenuationTermDim, FEnableTwoSidedGeometryDim>;
+	using FPermutationDomain = TShaderPermutationDomain<FUseAttenuationTermDim, FEnableTwoSidedGeometryDim, FDeferredMaterialMode>;
 
 	static bool ShouldCompilePermutation(const FGlobalShaderPermutationParameters& Parameters)
 	{
@@ -401,13 +428,80 @@ class FRayTracingGlobalIlluminationCreateGatherPointsRGS : public FGlobalShader
 		SHADER_PARAMETER_SAMPLER(SamplerState, TransmissionProfilesLinearSampler)
 
 		SHADER_PARAMETER(FIntPoint, GatherPointsResolution)
+		SHADER_PARAMETER(FIntPoint, TileAlignedResolution)
+		SHADER_PARAMETER(int32, SortTileSize)
+
 		// Output
 		SHADER_PARAMETER_RDG_BUFFER_UAV(RWStructuredBuffer<GatherPoints>, RWGatherPointsBuffer)
+		// Optional indirection buffer used for sorted materials
+		SHADER_PARAMETER_RDG_BUFFER_UAV(StructuredBuffer<FDeferredMaterialPayload>, MaterialBuffer)
 		END_SHADER_PARAMETER_STRUCT()
 };
 
 IMPLEMENT_GLOBAL_SHADER(FRayTracingGlobalIlluminationCreateGatherPointsRGS, "/Engine/Private/RayTracing/RayTracingCreateGatherPointsRGS.usf", "RayTracingCreateGatherPointsRGS", SF_RayGen);
 
+class FRayTracingGlobalIlluminationCreateGatherPointsTraceRGS : public FGlobalShader
+{
+	DECLARE_GLOBAL_SHADER(FRayTracingGlobalIlluminationCreateGatherPointsTraceRGS)
+	SHADER_USE_ROOT_PARAMETER_STRUCT(FRayTracingGlobalIlluminationCreateGatherPointsTraceRGS, FGlobalShader)
+
+		class FUseAttenuationTermDim : SHADER_PERMUTATION_BOOL("USE_ATTENUATION_TERM");
+	class FEnableTwoSidedGeometryDim : SHADER_PERMUTATION_BOOL("ENABLE_TWO_SIDED_GEOMETRY");
+	class FDeferredMaterialMode : SHADER_PERMUTATION_ENUM_CLASS("DIM_DEFERRED_MATERIAL_MODE", EDeferredMaterialMode);
+
+	using FPermutationDomain = TShaderPermutationDomain<FUseAttenuationTermDim, FEnableTwoSidedGeometryDim, FDeferredMaterialMode>;
+
+	static bool ShouldCompilePermutation(const FGlobalShaderPermutationParameters& Parameters)
+	{
+		return ShouldCompileRayTracingShadersForProject(Parameters.Platform);
+	}
+
+	BEGIN_SHADER_PARAMETER_STRUCT(FParameters, )
+		SHADER_PARAMETER(uint32, GatherSamplesPerPixel)
+		SHADER_PARAMETER(uint32, SamplesPerPixel)
+		SHADER_PARAMETER(uint32, SampleIndex)
+		SHADER_PARAMETER(uint32, MaxBounces)
+		SHADER_PARAMETER(uint32, UpscaleFactor)
+		SHADER_PARAMETER(uint32, RenderTileOffsetX)
+		SHADER_PARAMETER(uint32, RenderTileOffsetY)
+		SHADER_PARAMETER(float, MaxRayDistanceForGI)
+		SHADER_PARAMETER(float, NextEventEstimationSamples)
+		SHADER_PARAMETER(float, DiffuseThreshold)
+		SHADER_PARAMETER(float, MaxNormalBias)
+		SHADER_PARAMETER(uint32, EvalSkyLight)
+		SHADER_PARAMETER(uint32, UseRussianRoulette)
+
+		// Scene data
+		SHADER_PARAMETER_SRV(RaytracingAccelerationStructure, TLAS)
+		SHADER_PARAMETER_STRUCT_REF(FViewUniformShaderParameters, ViewUniformBuffer)
+
+		// Sampling sequence
+		SHADER_PARAMETER_STRUCT_REF(FHaltonIteration, HaltonIteration)
+		SHADER_PARAMETER_STRUCT_REF(FHaltonPrimes, HaltonPrimes)
+		SHADER_PARAMETER_STRUCT_REF(FBlueNoise, BlueNoise)
+
+		// Light data
+		SHADER_PARAMETER_STRUCT_REF(FPathTracingLightData, LightParameters)
+		SHADER_PARAMETER_STRUCT_REF(FSkyLightData, SkyLight)
+
+		// Shading data
+		SHADER_PARAMETER_STRUCT_INCLUDE(FSceneTextureParameters, SceneTextures)
+		SHADER_PARAMETER_RDG_TEXTURE(Texture2D, SSProfilesTexture)
+		SHADER_PARAMETER_SAMPLER(SamplerState, TransmissionProfilesLinearSampler)
+
+		SHADER_PARAMETER(FIntPoint, GatherPointsResolution)
+		SHADER_PARAMETER(FIntPoint, TileAlignedResolution)
+		SHADER_PARAMETER(int32, SortTileSize)
+
+		// Output
+		SHADER_PARAMETER_RDG_BUFFER_UAV(RWStructuredBuffer<GatherPoints>, RWGatherPointsBuffer)
+		// Optional indirection buffer used for sorted materials
+		SHADER_PARAMETER_RDG_BUFFER_UAV(StructuredBuffer<FDeferredMaterialPayload>, MaterialBuffer)
+		END_SHADER_PARAMETER_STRUCT()
+};
+
+IMPLEMENT_GLOBAL_SHADER(FRayTracingGlobalIlluminationCreateGatherPointsTraceRGS, "/Engine/Private/RayTracing/RayTracingCreateGatherPointsRGS.usf", "RayTracingCreateGatherPointsTraceRGS", SF_RayGen);
+
 class FRayTracingGlobalIlluminationFinalGatherRGS : public FGlobalShader
 {
 	DECLARE_GLOBAL_SHADER(FRayTracingGlobalIlluminationFinalGatherRGS)
@@ -456,6 +550,8 @@ IMPLEMENT_GLOBAL_SHADER(FRayTracingGlobalIlluminationFinalGatherRGS, "/Engine/Pr
 
 void FDeferredShadingSceneRenderer::PrepareRayTracingGlobalIllumination(const FViewInfo& View, TArray<FRHIRayTracingShader*>& OutRayGenShaders)
 {
+	const bool bSortMaterials = CVarRayTracingGlobalIlluminationFinalGatherSortMaterials.GetValueOnRenderThread() != 0;
+
 	// Declare all RayGen shaders that require material closest hit shaders to be bound
 	for (int UseAttenuationTerm = 0; UseAttenuationTerm < 2; ++UseAttenuationTerm)
 	{
@@ -467,11 +563,37 @@ void FDeferredShadingSceneRenderer::PrepareRayTracingGlobalIllumination(const FV
 			TShaderMapRef<FGlobalIlluminationRGS> RayGenerationShader(View.ShaderMap, PermutationVector);
 			OutRayGenShaders.Add(RayGenerationShader->GetRayTracingShader());
 
-			FRayTracingGlobalIlluminationCreateGatherPointsRGS::FPermutationDomain CreateGatherPointsPermutationVector;
-			CreateGatherPointsPermutationVector.Set<FRayTracingGlobalIlluminationCreateGatherPointsRGS::FUseAttenuationTermDim>(UseAttenuationTerm == 1);
-			CreateGatherPointsPermutationVector.Set<FRayTracingGlobalIlluminationCreateGatherPointsRGS::FEnableTwoSidedGeometryDim>(EnableTwoSidedGeometry == 1);
-			TShaderMapRef<FRayTracingGlobalIlluminationCreateGatherPointsRGS> CreateGatherPointsRayGenerationShader(View.ShaderMap, CreateGatherPointsPermutationVector);
-			OutRayGenShaders.Add(CreateGatherPointsRayGenerationShader->GetRayTracingShader());
+			if (bSortMaterials)
+			{
+				// Gather
+				{
+					FRayTracingGlobalIlluminationCreateGatherPointsTraceRGS::FPermutationDomain CreateGatherPointsPermutationVector;
+					CreateGatherPointsPermutationVector.Set<FRayTracingGlobalIlluminationCreateGatherPointsTraceRGS::FUseAttenuationTermDim>(UseAttenuationTerm == 1);
+					CreateGatherPointsPermutationVector.Set<FRayTracingGlobalIlluminationCreateGatherPointsTraceRGS::FEnableTwoSidedGeometryDim>(EnableTwoSidedGeometry == 1);
+					CreateGatherPointsPermutationVector.Set<FRayTracingGlobalIlluminationCreateGatherPointsTraceRGS::FDeferredMaterialMode>(EDeferredMaterialMode::Gather);
+					TShaderMapRef<FRayTracingGlobalIlluminationCreateGatherPointsTraceRGS> CreateGatherPointsRayGenerationShader(View.ShaderMap, CreateGatherPointsPermutationVector);
+					OutRayGenShaders.Add(CreateGatherPointsRayGenerationShader->GetRayTracingShader());
+				}
+
+				// Shade
+				{
+					FRayTracingGlobalIlluminationCreateGatherPointsRGS::FPermutationDomain CreateGatherPointsPermutationVector;
+					CreateGatherPointsPermutationVector.Set<FRayTracingGlobalIlluminationCreateGatherPointsRGS::FUseAttenuationTermDim>(UseAttenuationTerm == 1);
+					CreateGatherPointsPermutationVector.Set<FRayTracingGlobalIlluminationCreateGatherPointsRGS::FEnableTwoSidedGeometryDim>(EnableTwoSidedGeometry == 1);
+					CreateGatherPointsPermutationVector.Set<FRayTracingGlobalIlluminationCreateGatherPointsRGS::FDeferredMaterialMode>(EDeferredMaterialMode::Shade);
+					TShaderMapRef<FRayTracingGlobalIlluminationCreateGatherPointsRGS> CreateGatherPointsRayGenerationShader(View.ShaderMap, CreateGatherPointsPermutationVector);
+					OutRayGenShaders.Add(CreateGatherPointsRayGenerationShader->GetRayTracingShader());
+				}
+			}
+			else
+			{
+				FRayTracingGlobalIlluminationCreateGatherPointsRGS::FPermutationDomain CreateGatherPointsPermutationVector;
+				CreateGatherPointsPermutationVector.Set<FRayTracingGlobalIlluminationCreateGatherPointsRGS::FUseAttenuationTermDim>(UseAttenuationTerm == 1);
+				CreateGatherPointsPermutationVector.Set<FRayTracingGlobalIlluminationCreateGatherPointsRGS::FEnableTwoSidedGeometryDim>(EnableTwoSidedGeometry == 1);
+				CreateGatherPointsPermutationVector.Set<FRayTracingGlobalIlluminationCreateGatherPointsRGS::FDeferredMaterialMode>(EDeferredMaterialMode::None);
+				TShaderMapRef<FRayTracingGlobalIlluminationCreateGatherPointsRGS> CreateGatherPointsRayGenerationShader(View.ShaderMap, CreateGatherPointsPermutationVector);
+				OutRayGenShaders.Add(CreateGatherPointsRayGenerationShader->GetRayTracingShader());
+			}
 
 			FRayTracingGlobalIlluminationFinalGatherRGS::FPermutationDomain GatherPassPermutationVector;
 			GatherPassPermutationVector.Set<FRayTracingGlobalIlluminationFinalGatherRGS::FUseAttenuationTermDim>(UseAttenuationTerm == 1);
@@ -542,6 +664,47 @@ bool FDeferredShadingSceneRenderer::RenderRayTracingGlobalIllumination(
 }
 #endif // RHI_RAYTRACING
 
+void CopyGatherPassParameters(
+	const FRayTracingGlobalIlluminationCreateGatherPointsRGS::FParameters& PassParameters,
+	FRayTracingGlobalIlluminationCreateGatherPointsTraceRGS::FParameters* NewParameters
+)
+{
+	NewParameters->GatherSamplesPerPixel = PassParameters.GatherSamplesPerPixel;
+	NewParameters->SamplesPerPixel = PassParameters.SamplesPerPixel;
+	NewParameters->SampleIndex = PassParameters.SampleIndex;
+	NewParameters->MaxBounces = PassParameters.MaxBounces;
+	NewParameters->UpscaleFactor = PassParameters.UpscaleFactor;
+	NewParameters->RenderTileOffsetX = PassParameters.RenderTileOffsetX;
+	NewParameters->RenderTileOffsetY = PassParameters.RenderTileOffsetY;
+	NewParameters->MaxRayDistanceForGI = PassParameters.MaxRayDistanceForGI;
+	NewParameters->NextEventEstimationSamples = PassParameters.NextEventEstimationSamples;
+	NewParameters->DiffuseThreshold = PassParameters.DiffuseThreshold;
+	NewParameters->MaxNormalBias = PassParameters.MaxNormalBias;
+	NewParameters->EvalSkyLight = PassParameters.EvalSkyLight;
+	NewParameters->UseRussianRoulette = PassParameters.UseRussianRoulette;
+
+	NewParameters->TLAS = PassParameters.TLAS;
+	NewParameters->ViewUniformBuffer = PassParameters.ViewUniformBuffer;
+
+	NewParameters->HaltonIteration = PassParameters.HaltonIteration;
+	NewParameters->HaltonPrimes = PassParameters.HaltonPrimes;
+	NewParameters->BlueNoise = PassParameters.BlueNoise;
+
+	NewParameters->LightParameters = PassParameters.LightParameters;
+	NewParameters->SkyLight = PassParameters.SkyLight;
+
+	NewParameters->SceneTextures = PassParameters.SceneTextures;
+	NewParameters->SSProfilesTexture = PassParameters.SSProfilesTexture;
+	NewParameters->TransmissionProfilesLinearSampler = PassParameters.TransmissionProfilesLinearSampler;
+
+	NewParameters->GatherPointsResolution = PassParameters.GatherPointsResolution;
+	NewParameters->TileAlignedResolution = PassParameters.TileAlignedResolution;
+	NewParameters->SortTileSize = PassParameters.SortTileSize;
+
+	NewParameters->RWGatherPointsBuffer = PassParameters.RWGatherPointsBuffer;
+	NewParameters->MaterialBuffer = PassParameters.MaterialBuffer;
+}
+
 void FDeferredShadingSceneRenderer::RayTracingGlobalIlluminationCreateGatherPoints(
 	FRDGBuilder& GraphBuilder,
 	FSceneTextureParameters& SceneTextures,
@@ -555,7 +718,7 @@ void FDeferredShadingSceneRenderer::RayTracingGlobalIlluminationCreateGatherPoin
 	RDG_GPU_STAT_SCOPE(GraphBuilder, RayTracingGICreateGatherPoints);
 	RDG_EVENT_SCOPE(GraphBuilder, "Ray Tracing GI: Create Gather Points");
 
-	int32 GatherSamples = GetRayTracingGlobalIlluminationSamplesPerPixel(View);
+	int32 GatherSamples = FMath::Min(GetRayTracingGlobalIlluminationSamplesPerPixel(View), MAXIMUM_GATHER_POINTS_PER_PIXEL);
 	int32 SamplesPerPixel = 1;
 
 	uint32 IterationCount = SamplesPerPixel;
@@ -638,6 +801,12 @@ void FDeferredShadingSceneRenderer::RayTracingGlobalIlluminationCreateGatherPoin
 	PassParameters->GatherPointsResolution = FIntPoint(GatherPointsResolution.X, GatherPointsResolution.Y);
 	PassParameters->RWGatherPointsBuffer = GraphBuilder.CreateUAV(GatherPointsBuffer, EPixelFormat::PF_R32_UINT);
 
+	// When deferred materials are used, two passes are invoked:
+	// 1) Gather ray-hit data and sort by hit-shader ID
+	// 2) Re-trace "short" ray and shade
+	const bool bSortMaterials = CVarRayTracingGlobalIlluminationFinalGatherSortMaterials.GetValueOnRenderThread() != 0;
+	if (!bSortMaterials)
+	{
 	FRayTracingGlobalIlluminationCreateGatherPointsRGS::FPermutationDomain PermutationVector;
 	PermutationVector.Set<FRayTracingGlobalIlluminationCreateGatherPointsRGS::FUseAttenuationTermDim>(true);
 	PermutationVector.Set<FRayTracingGlobalIlluminationCreateGatherPointsRGS::FEnableTwoSidedGeometryDim>(CVarRayTracingGlobalIlluminationEnableTwoSidedGeometry.GetValueOnRenderThread() != 0);
@@ -655,6 +824,87 @@ void FDeferredShadingSceneRenderer::RayTracingGlobalIlluminationCreateGatherPoin
 		SetShaderParameters(GlobalResources, *RayGenerationShader, *PassParameters);
 		RHICmdList.RayTraceDispatch(View.RayTracingMaterialPipeline, RayGenerationShader->GetRayTracingShader(), RayTracingSceneRHI, GlobalResources, GatherPointsResolution.X, GatherPointsResolution.Y);
 	});
+	}
+	else
+	{
+		// Determines tile-size for sorted-deferred path
+		const int32 SortTileSize = CVarRayTracingGlobalIlluminationFinalGatherSortTileSize.GetValueOnRenderThread();
+		FIntPoint TileAlignedResolution = FIntPoint(GatherPointsResolution.X, GatherPointsResolution.Y);
+		if (SortTileSize)
+		{
+			TileAlignedResolution = FIntPoint::DivideAndRoundUp(TileAlignedResolution, SortTileSize) * SortTileSize;
+		}
+		PassParameters->TileAlignedResolution = TileAlignedResolution;
+		PassParameters->SortTileSize = SortTileSize;
+
+		FRDGBufferRef DeferredMaterialBuffer = nullptr;
+		const uint32 DeferredMaterialBufferNumElements = TileAlignedResolution.X * TileAlignedResolution.Y;
+
+		// Gather pass
+		{
+			FRayTracingGlobalIlluminationCreateGatherPointsTraceRGS::FParameters* GatherPassParameters = GraphBuilder.AllocParameters<FRayTracingGlobalIlluminationCreateGatherPointsTraceRGS::FParameters>();
+			CopyGatherPassParameters(*PassParameters, GatherPassParameters);
+
+			FRDGBufferDesc Desc = FRDGBufferDesc::CreateStructuredDesc(sizeof(FDeferredMaterialPayload), DeferredMaterialBufferNumElements);
+			DeferredMaterialBuffer = GraphBuilder.CreateBuffer(Desc, TEXT("RayTracingGlobalIlluminationMaterialBuffer"));
+			GatherPassParameters->MaterialBuffer = GraphBuilder.CreateUAV(DeferredMaterialBuffer);
+
+			FRayTracingGlobalIlluminationCreateGatherPointsTraceRGS::FPermutationDomain PermutationVector;
+			PermutationVector.Set<FRayTracingGlobalIlluminationCreateGatherPointsTraceRGS::FUseAttenuationTermDim>(true);
+			PermutationVector.Set<FRayTracingGlobalIlluminationCreateGatherPointsTraceRGS::FEnableTwoSidedGeometryDim>(CVarRayTracingGlobalIlluminationEnableTwoSidedGeometry.GetValueOnRenderThread() != 0);
+			PermutationVector.Set<FRayTracingGlobalIlluminationCreateGatherPointsTraceRGS::FDeferredMaterialMode>(EDeferredMaterialMode::Gather);
+			TShaderMapRef<FRayTracingGlobalIlluminationCreateGatherPointsTraceRGS> RayGenerationShader(GetGlobalShaderMap(FeatureLevel), PermutationVector);
+
+			ClearUnusedGraphResources(*RayGenerationShader, GatherPassParameters);
+
+			GraphBuilder.AddPass(
+				RDG_EVENT_NAME("GlobalIlluminationRayTracingGatherMaterials %dx%d", TileAlignedResolution.X, TileAlignedResolution.Y),
+				GatherPassParameters,
+				ERDGPassFlags::Compute,
+				[GatherPassParameters, this, &View, RayGenerationShader, TileAlignedResolution](FRHICommandList& RHICmdList)
+			{
+				FRayTracingPipelineState* Pipeline = BindRayTracingDeferredMaterialGatherPipeline(RHICmdList, View, RayGenerationShader->GetRayTracingShader());
+
+				FRayTracingShaderBindingsWriter GlobalResources;
+				SetShaderParameters(GlobalResources, *RayGenerationShader, *GatherPassParameters);
+
+				FRHIRayTracingScene* RayTracingSceneRHI = View.RayTracingScene.RayTracingSceneRHI;
+				RHICmdList.RayTraceDispatch(Pipeline, RayGenerationShader->GetRayTracingShader(), RayTracingSceneRHI, GlobalResources, TileAlignedResolution.X, TileAlignedResolution.Y);
+			});
+		}
+
+		// Sort by hit-shader ID
+		const uint32 SortSize = CVarRayTracingGlobalIlluminationFinalGatherSortSize.GetValueOnRenderThread();
+		SortDeferredMaterials(GraphBuilder, View, SortSize, DeferredMaterialBufferNumElements, DeferredMaterialBuffer);
+
+		// Shade pass
+		{
+			PassParameters->MaterialBuffer = GraphBuilder.CreateUAV(DeferredMaterialBuffer);
+
+			FRayTracingGlobalIlluminationCreateGatherPointsRGS::FPermutationDomain PermutationVector;
+			PermutationVector.Set<FRayTracingGlobalIlluminationCreateGatherPointsRGS::FUseAttenuationTermDim>(true);
+			PermutationVector.Set<FRayTracingGlobalIlluminationCreateGatherPointsRGS::FEnableTwoSidedGeometryDim>(CVarRayTracingGlobalIlluminationEnableTwoSidedGeometry.GetValueOnRenderThread() != 0);
+			PermutationVector.Set<FRayTracingGlobalIlluminationCreateGatherPointsRGS::FDeferredMaterialMode>(EDeferredMaterialMode::Shade);
+			TShaderMapRef<FRayTracingGlobalIlluminationCreateGatherPointsRGS> RayGenerationShader(GetGlobalShaderMap(FeatureLevel), PermutationVector);
+			ClearUnusedGraphResources(*RayGenerationShader, PassParameters);
+
+			GraphBuilder.AddPass(
+				RDG_EVENT_NAME("GlobalIlluminationRayTracingShadeMaterials %d", DeferredMaterialBufferNumElements),
+				PassParameters,
+				ERDGPassFlags::Compute,
+				[PassParameters, this, &View, RayGenerationShader, DeferredMaterialBufferNumElements](FRHICommandList& RHICmdList)
+			{
+				FRHIRayTracingScene* RayTracingSceneRHI = View.RayTracingScene.RayTracingSceneRHI;
+				FRayTracingShaderBindingsWriter GlobalResources;
+				SetShaderParameters(GlobalResources, *RayGenerationShader, *PassParameters);
+
+				// Shading pass for sorted materials uses 1D dispatch over all elements in the material buffer.
+				// This can be reduced to the number of output pixels if sorting pass guarantees that all invalid entries are moved to the end.
+				RHICmdList.RayTraceDispatch(View.RayTracingMaterialPipeline, RayGenerationShader->GetRayTracingShader(), RayTracingSceneRHI, GlobalResources, DeferredMaterialBufferNumElements, 1);
+			});
+		}
+	}
+
 }
 #else
 {
