Fixing compilation issue.

diff --git a/Engine/Source/Runtime/Renderer/Private/RayTracing/RayTracingGlobalIllumination.cpp b/Engine/Source/Runtime/Renderer/Private/RayTracing/RayTracingGlobalIllumination.cpp
--- a/Engine/Source/Runtime/Renderer/Private/RayTracing/RayTracingGlobalIllumination.cpp
+++ b/Engine/Source/Runtime/Renderer/Private/RayTracing/RayTracingGlobalIllumination.cpp
@@ -664,6 +664,7 @@ bool FDeferredShadingSceneRenderer::RenderRayTracingGlobalIllumination(
 }
 #endif // RHI_RAYTRACING
 
+#if RHI_RAYTRACING
 void CopyGatherPassParameters(
 	const FRayTracingGlobalIlluminationCreateGatherPointsRGS::FParameters& PassParameters,
 	FRayTracingGlobalIlluminationCreateGatherPointsTraceRGS::FParameters* NewParameters
@@ -704,6 +705,7 @@ void CopyGatherPassParameters(
 	NewParameters->RWGatherPointsBuffer = PassParameters.RWGatherPointsBuffer;
 	NewParameters->MaterialBuffer = PassParameters.MaterialBuffer;
 }
+#endif // RHI_RAYTRACING
 
 void FDeferredShadingSceneRenderer::RayTracingGlobalIlluminationCreateGatherPoints(
 	FRDGBuilder& GraphBuilder,
