Remove r.RayTracing.GlobalIllumination.EnableFinalGather CVar

This CVar has now been subsumed by setting r.RayTracing.GlobalIllumination to a value of 2.

diff --git a/Engine/Source/Runtime/Renderer/Private/RayTracing/RayTracingGlobalIllumination.cpp b/Engine/Source/Runtime/Renderer/Private/RayTracing/RayTracingGlobalIllumination.cpp
--- a/Engine/Source/Runtime/Renderer/Private/RayTracing/RayTracingGlobalIllumination.cpp
+++ b/Engine/Source/Runtime/Renderer/Private/RayTracing/RayTracingGlobalIllumination.cpp
@@ -128,13 +128,6 @@ static TAutoConsoleVariable<int32> CVarRayTracingGlobalIlluminationMaxLightCount
 	ECVF_RenderThreadSafe
 );
 
-static TAutoConsoleVariable<int32> CVarRayTracingGlobalIlluminationEnableFinalGather(
-	TEXT("r.RayTracing.GlobalIllumination.EnableFinalGather"),
-	0,
-	TEXT("Enables final gather algorithm for 1-bounce global illumination (default = 0)"),
-	ECVF_RenderThreadSafe
-);
-
 static float GRayTracingGlobalIlluminationFinalGatherDistance = 10.0;
 static FAutoConsoleVariableRef CVarRayTracingGlobalIlluminationFinalGatherDistance(
 	TEXT("r.RayTracing.GlobalIllumination.FinalGather.Distance"),
@@ -313,10 +306,10 @@ bool ShouldRenderRayTracingGlobalIllumination(const FViewInfo& View)
 bool IsFinalGatherEnabled(const FViewInfo& View)
 {
 
-	int32 bEnableFinalGather = CVarRayTracingGlobalIlluminationEnableFinalGather.GetValueOnRenderThread();
-	if (bEnableFinalGather >= 0)
+	int32 CVarRayTracingGlobalIlluminationValue = CVarRayTracingGlobalIllumination.GetValueOnRenderThread();
+	if (CVarRayTracingGlobalIlluminationValue >= 0)
 	{
-		return bEnableFinalGather > 0;
+		return CVarRayTracingGlobalIlluminationValue == 2;
 	}
 
 	return View.FinalPostProcessSettings.RayTracingGIType == ERayTracingGlobalIlluminationType::FinalGather;
