Remove CVar typo-ation

diff --git a/Engine/Source/Runtime/Renderer/Private/RayTracing/RayTracingGlobalIllumination.cpp b/Engine/Source/Runtime/Renderer/Private/RayTracing/RayTracingGlobalIllumination.cpp
--- a/Engine/Source/Runtime/Renderer/Private/RayTracing/RayTracingGlobalIllumination.cpp
+++ b/Engine/Source/Runtime/Renderer/Private/RayTracing/RayTracingGlobalIllumination.cpp
@@ -48,7 +48,7 @@ static FAutoConsoleVariableRef CVarRayTracingGlobalIlluminationMaxRayDistance(
 );
 
 static TAutoConsoleVariable<int32> CVarRayTracingGlobalIlluminationMaxBounces(
-	TEXT("r.RayTracing.GlobalIlluminGlobalIllumination.MaxBouncesation"),
+	TEXT("r.RayTracing.GlobalIllumination.MaxBounces"),
 	-1,
 	TEXT("Max bounces (default = -1 (driven by postprocesing volume))"),
 	ECVF_RenderThreadSafe
