Properly clamp the water refraction in the negative x,y directions.

#rnx

diff --git a/Engine/Shaders/Private/BasePassPixelShader.usf b/Engine/Shaders/Private/BasePassPixelShader.usf
--- a/Engine/Shaders/Private/BasePassPixelShader.usf
+++ b/Engine/Shaders/Private/BasePassPixelShader.usf
@@ -1269,7 +1269,8 @@ ShadingModel);
 			DirectionalLightShadow,
 			OpaqueBasePass.SceneColorWithoutSingleLayerWaterTexture, SingleLayerWaterSceneColorSampler,
 			OpaqueBasePass.SceneDepthWithoutSingleLayerWaterTexture, SingleLayerWaterSceneDepthSampler,
-			OpaqueBasePass.SceneWithoutSingleLayerWaterMaxUV,
+			OpaqueBasePass.SceneWithoutSingleLayerWaterMinMaxUV.xy,
+			OpaqueBasePass.SceneWithoutSingleLayerWaterMinMaxUV.zw,
 			Specular, OpaqueBasePass.DistortionParams,
 			SunIlluminance, WaterDiffuseIndirectIlluminance, EnvBrdf,
 			CameraIsUnderWater, WaterVisibility, EyeIndex).Luminance;
diff --git a/Engine/Shaders/Private/SingleLayerWaterCommon.ush b/Engine/Shaders/Private/SingleLayerWaterCommon.ush
--- a/Engine/Shaders/Private/SingleLayerWaterCommon.ush
+++ b/Engine/Shaders/Private/SingleLayerWaterCommon.ush
@@ -49,6 +49,7 @@ WaterVolumeLightingOutput EvaluateWaterVolumeLighting(
 #if !SIMPLE_SINGLE_LAYER_WATER
 	Texture2D SceneColorWithoutSingleLayerWaterTexture, SamplerState SceneColorWithoutSingleLayerWaterSampler,
 	Texture2D SceneDepthWithoutSingleLayerWaterTexture, SamplerState SceneDepthWithoutSingleLayerWaterSampler,
+	float2 SceneWithoutSingleLayerWaterMinUV,
 	float2 SceneWithoutSingleLayerWaterMaxUV,
 #endif // SIMPLE_SINGLE_LAYER_WATER
 	float Specular, const float4 DistortionParams,
@@ -96,7 +97,7 @@ WaterVolumeLightingOutput EvaluateWaterVolumeLighting(
 
 #else  // !SIMPLE_SINGLE_LAYER_WATER
 
-		ViewportUV = min(ViewportUV, SceneWithoutSingleLayerWaterMaxUV);
+		ViewportUV = clamp(ViewportUV, SceneWithoutSingleLayerWaterMinUV, SceneWithoutSingleLayerWaterMaxUV);
 
 		float PixelSceneDepth = SceneDepthWithoutSingleLayerWaterTexture.SampleLevel(SceneDepthWithoutSingleLayerWaterSampler, ViewportUV, 0).x * SINGLE_LAYER_WATER_DEPTH_SCALE;
 		
@@ -108,7 +109,7 @@ WaterVolumeLightingOutput EvaluateWaterVolumeLighting(
 		// Also reduce refraction according to water thickness
 		BufferUVDistortion *= saturate((PixelSceneDepth - WaterDepth) * 1.0 / 30.0f);// *(1 - saturate(WaterDepth * 1.0 / 2000.0f));
 
-		const float2 DistortedUV = min(ViewportUV + BufferUVDistortion, SceneWithoutSingleLayerWaterMaxUV);
+		const float2 DistortedUV = clamp(ViewportUV + BufferUVDistortion, SceneWithoutSingleLayerWaterMinUV, SceneWithoutSingleLayerWaterMaxUV);
 
 		float SceneDepth = SceneDepthWithoutSingleLayerWaterTexture.SampleLevel(SceneDepthWithoutSingleLayerWaterSampler, DistortedUV, 0).x * SINGLE_LAYER_WATER_DEPTH_SCALE;
 
diff --git a/Engine/Source/Runtime/Renderer/Private/BasePassRendering.cpp b/Engine/Source/Runtime/Renderer/Private/BasePassRendering.cpp
--- a/Engine/Source/Runtime/Renderer/Private/BasePassRendering.cpp
+++ b/Engine/Source/Runtime/Renderer/Private/BasePassRendering.cpp
@@ -517,7 +517,7 @@ void CreateOpaqueBasePassUniformBuffer(
 	FRHICommandListImmediate& RHICmdList, 
 	const FViewInfo& View,
 	IPooledRenderTarget* ForwardScreenSpaceShadowMask,
-	FVector2D* SceneWithoutSingleLayerWaterMaxUV,
+	FVector4* SceneWithoutSingleLayerWaterMinMaxUV,
 	IPooledRenderTarget* SceneColorWithoutSingleLayerWater,
 	IPooledRenderTarget* SceneDepthWithoutSingleLayerWater,
 	TUniformBufferRef<FOpaqueBasePassUniformParameters>& BasePassUniformBuffer)
@@ -584,7 +584,7 @@ void CreateOpaqueBasePassUniformBuffer(
 	}
 
 	// Single Layer Water
-	BasePassParameters.SceneWithoutSingleLayerWaterMaxUV = SceneWithoutSingleLayerWaterMaxUV ? *SceneWithoutSingleLayerWaterMaxUV : FVector2D(1.0f, 1.0f);
+	BasePassParameters.SceneWithoutSingleLayerWaterMinMaxUV = SceneWithoutSingleLayerWaterMinMaxUV ? *SceneWithoutSingleLayerWaterMinMaxUV : FVector4(0.0f, 0.0f, 1.0f, 1.0f);
 	BasePassParameters.SceneColorWithoutSingleLayerWaterTexture = (SceneColorWithoutSingleLayerWater ? SceneColorWithoutSingleLayerWater->GetRenderTargetItem().TargetableTexture : GSystemTextures.BlackDummy->GetRenderTargetItem().TargetableTexture);
 	BasePassParameters.SceneColorWithoutSingleLayerWaterSampler = TStaticSamplerState<SF_Bilinear>::GetRHI();
 	BasePassParameters.SceneDepthWithoutSingleLayerWaterTexture = (SceneDepthWithoutSingleLayerWater ? SceneDepthWithoutSingleLayerWater->GetRenderTargetItem().TargetableTexture : GSystemTextures.BlackDummy->GetRenderTargetItem().TargetableTexture);
diff --git a/Engine/Source/Runtime/Renderer/Private/BasePassRendering.h b/Engine/Source/Runtime/Renderer/Private/BasePassRendering.h
--- a/Engine/Source/Runtime/Renderer/Private/BasePassRendering.h
+++ b/Engine/Source/Runtime/Renderer/Private/BasePassRendering.h
@@ -79,7 +79,7 @@ BEGIN_GLOBAL_SHADER_PARAMETER_STRUCT(FOpaqueBasePassUniformParameters,)
 	SHADER_PARAMETER_SAMPLER(SamplerState, SceneDepthWithoutSingleLayerWaterSampler)
 	SHADER_PARAMETER_TEXTURE(Texture2D, PreIntegratedGFTexture)
 	SHADER_PARAMETER_SAMPLER(SamplerState, PreIntegratedGFSampler)
-	SHADER_PARAMETER(FVector2D, SceneWithoutSingleLayerWaterMaxUV)
+	SHADER_PARAMETER(FVector4, SceneWithoutSingleLayerWaterMinMaxUV)
 	SHADER_PARAMETER(FVector4, DistortionParams)
 	// Misc
 	SHADER_PARAMETER_TEXTURE(Texture2D, EyeAdaptation)
@@ -119,7 +119,7 @@ extern void CreateOpaqueBasePassUniformBuffer(
 	FRHICommandListImmediate& RHICmdList, 
 	const FViewInfo& View,
 	IPooledRenderTarget* ForwardScreenSpaceShadowMask,
-	FVector2D* SceneWithoutSingleLayerWaterMaxUV,
+	FVector4* SceneWithoutSingleLayerWaterMinMaxUV,
 	IPooledRenderTarget* SceneColorWithoutSingleLayerWater,
 	IPooledRenderTarget* SceneDepthWithoutSingleLayerWater,
 	TUniformBufferRef<FOpaqueBasePassUniformParameters>& BasePassUniformBuffer);
diff --git a/Engine/Source/Runtime/Renderer/Private/SingleLayerWaterRendering.cpp b/Engine/Source/Runtime/Renderer/Private/SingleLayerWaterRendering.cpp
--- a/Engine/Source/Runtime/Renderer/Private/SingleLayerWaterRendering.cpp
+++ b/Engine/Source/Runtime/Renderer/Private/SingleLayerWaterRendering.cpp
@@ -249,7 +249,7 @@ BEGIN_SHADER_PARAMETER_STRUCT(FSingleLayerWaterCommonShaderParameters, )
 	SHADER_PARAMETER_SAMPLER(SamplerState, PreIntegratedGFSampler)
 	SHADER_PARAMETER_RDG_TEXTURE(Texture2D, SceneNoWaterDepthTexture)
 	SHADER_PARAMETER_SAMPLER(SamplerState, SceneNoWaterDepthSampler)
-	SHADER_PARAMETER(FVector2D, SceneNoWaterMaxUV)
+	SHADER_PARAMETER(FVector4, SceneNoWaterMinMaxUV)
 	SHADER_PARAMETER_STRUCT_INCLUDE(FSceneTextureParameters, SceneTextures)	// Water scene texture
 	SHADER_PARAMETER_STRUCT_INCLUDE(FSceneTextureSamplerParameters, SceneTextureSamplers)
 	SHADER_PARAMETER_STRUCT_REF(FViewUniformShaderParameters, ViewUniformBuffer)
@@ -488,8 +488,10 @@ void FDeferredShadingSceneRenderer::CopySingleLayerWaterTextures(FRHICommandList
     
 		const FIntRect RefractionViewRect = FIntRect(FIntPoint::DivideAndRoundDown(View.ViewRect.Min, RefractionDownsampleFactor), FIntPoint::DivideAndRoundDown(View.ViewRect.Max, RefractionDownsampleFactor));
 
-		PassData.ViewData[ViewIndex].SceneWithoutSingleLayerWaterMaxUV.X = (RefractionViewRect.Max.X - 0.5f) / RefractionResolution.X; 
-		PassData.ViewData[ViewIndex].SceneWithoutSingleLayerWaterMaxUV.Y = (RefractionViewRect.Max.Y - 0.5f) / RefractionResolution.Y;
+		PassData.ViewData[ViewIndex].SceneWithoutSingleLayerWaterMinMaxUV.X = (RefractionViewRect.Min.X + 0.5f) / RefractionResolution.X; 
+		PassData.ViewData[ViewIndex].SceneWithoutSingleLayerWaterMinMaxUV.Y = (RefractionViewRect.Min.Y + 0.5f) / RefractionResolution.Y;
+		PassData.ViewData[ViewIndex].SceneWithoutSingleLayerWaterMinMaxUV.Z = (RefractionViewRect.Max.X - 0.5f) / RefractionResolution.X; 
+		PassData.ViewData[ViewIndex].SceneWithoutSingleLayerWaterMinMaxUV.W = (RefractionViewRect.Max.Y - 0.5f) / RefractionResolution.Y;
 
 	    FPixelShaderUtils::AddFullscreenPass(
 		    GraphBuilder,
@@ -619,7 +621,7 @@ void FDeferredShadingSceneRenderer::RenderSingleLayerWaterReflections(FRHIComman
 			Parameters.PreIntegratedGFSampler = TStaticSamplerState<SF_Bilinear, AM_Clamp, AM_Clamp, AM_Clamp>::GetRHI();
 			Parameters.SceneNoWaterDepthTexture = GraphBuilder.RegisterExternalTexture(PassData.SceneDepthWithoutSingleLayerWater ? PassData.SceneDepthWithoutSingleLayerWater : GSystemTextures.BlackDummy);
 			Parameters.SceneNoWaterDepthSampler = TStaticSamplerState<SF_Point>::GetRHI();
-			Parameters.SceneNoWaterMaxUV = PassData.ViewData[ViewIndex].SceneWithoutSingleLayerWaterMaxUV;
+			Parameters.SceneNoWaterMinMaxUV = PassData.ViewData[ViewIndex].SceneWithoutSingleLayerWaterMinMaxUV;
 			Parameters.SceneTextures = SceneTextures;
 			SetupSceneTextureSamplers(&Parameters.SceneTextureSamplers);
 			Parameters.ViewUniformBuffer = View.ViewUniformBuffer;
@@ -822,7 +824,7 @@ bool FDeferredShadingSceneRenderer::RenderSingleLayerWaterPass(FRHICommandListIm
 			CreateOpaqueBasePassUniformBuffer(RHICmdList, 
 				View, 
 				WhiteForwardScreenSpaceShadowMask, 
-				&PassData.ViewData[ViewIndex].SceneWithoutSingleLayerWaterMaxUV,
+				&PassData.ViewData[ViewIndex].SceneWithoutSingleLayerWaterMinMaxUV,
 				PassData.SceneColorWithoutSingleLayerWater.IsValid() ? PassData.SceneColorWithoutSingleLayerWater : GSystemTextures.BlackDummy,
 				PassData.SceneDepthWithoutSingleLayerWater,
 				WaterPassUniformBuffer);
diff --git a/Engine/Source/Runtime/Renderer/Private/SingleLayerWaterRendering.h b/Engine/Source/Runtime/Renderer/Private/SingleLayerWaterRendering.h
--- a/Engine/Source/Runtime/Renderer/Private/SingleLayerWaterRendering.h
+++ b/Engine/Source/Runtime/Renderer/Private/SingleLayerWaterRendering.h
@@ -20,7 +20,7 @@ struct FSingleLayerWaterPassData
 
 	struct FSingleLayerWaterPassViewData
 	{
-		FVector2D SceneWithoutSingleLayerWaterMaxUV;
+		FVector4 SceneWithoutSingleLayerWaterMinMaxUV;
 	};
 
 	TArray<FSingleLayerWaterPassViewData> ViewData;
