Fix a crash with ray traced reflections set to 0spp and move samples per pixel calculation to a function.

diff --git a/Engine/Source/Runtime/Renderer/Private/ReflectionEnvironment.cpp b/Engine/Source/Runtime/Renderer/Private/ReflectionEnvironment.cpp
--- a/Engine/Source/Runtime/Renderer/Private/ReflectionEnvironment.cpp
+++ b/Engine/Source/Runtime/Renderer/Private/ReflectionEnvironment.cpp
@@ -200,13 +200,18 @@ bool IsReflectionCaptureAvailable()
 }
 
 #if RHI_RAYTRACING
+int32 GetRayTracingReflectionsSamplesPerPixel(const FViewInfo& View)
+{
+	return GRayTracingReflectionsSamplesPerPixel >= 0 ? GRayTracingReflectionsSamplesPerPixel : View.FinalPostProcessSettings.RayTracingReflectionsSamplesPerPixel;
+}
+
 bool ShouldRenderRayTracingReflections(const FViewInfo& View)
 {
 	bool bThisViewHasRaytracingReflections = View.FinalPostProcessSettings.ReflectionsType == EReflectionsType::RayTracing;
 
 	const bool bReflectionsCvarEnabled = GRayTracingReflections < 0 ? bThisViewHasRaytracingReflections : (GRayTracingReflections != 0);
 	const int32 ForceAllRayTracingEffects = GetForceRayTracingEffectsCVarValue();
-	const bool bReflectionPassEnabled = (ForceAllRayTracingEffects > 0 || (bReflectionsCvarEnabled && ForceAllRayTracingEffects < 0));
+	const bool bReflectionPassEnabled = (ForceAllRayTracingEffects > 0 || (bReflectionsCvarEnabled && ForceAllRayTracingEffects < 0)) && (GetRayTracingReflectionsSamplesPerPixel(View) > 0);
 
 	return IsRayTracingEnabled() && bReflectionPassEnabled;
 }
@@ -682,7 +687,7 @@ void FDeferredShadingSceneRenderer::RenderDeferredReflectionsAndSkyLighting(FRHI
 				bDenoise = DenoiserMode != 0;
 				
 				RayTracingConfig.ResolutionFraction = FMath::Clamp(CVarReflectionScreenPercentage.GetValueOnRenderThread() / 100.0f, 0.25f, 1.0f);
-				RayTracingConfig.RayCountPerPixel = GRayTracingReflectionsSamplesPerPixel > -1 ? GRayTracingReflectionsSamplesPerPixel : View.FinalPostProcessSettings.RayTracingReflectionsSamplesPerPixel;
+				RayTracingConfig.RayCountPerPixel = GetRayTracingReflectionsSamplesPerPixel(View);
 
 				if (!bDenoise)
 				{
