Missing if/def around GetRayTracingReflectionsSamplesPerPixel

diff --git a/Engine/Source/Runtime/Renderer/Private/ReflectionEnvironment.cpp b/Engine/Source/Runtime/Renderer/Private/ReflectionEnvironment.cpp
--- a/Engine/Source/Runtime/Renderer/Private/ReflectionEnvironment.cpp
+++ b/Engine/Source/Runtime/Renderer/Private/ReflectionEnvironment.cpp
@@ -687,7 +687,11 @@ void FDeferredShadingSceneRenderer::RenderDeferredReflectionsAndSkyLighting(FRHI
 				bDenoise = DenoiserMode != 0;
 				
 				RayTracingConfig.ResolutionFraction = FMath::Clamp(CVarReflectionScreenPercentage.GetValueOnRenderThread() / 100.0f, 0.25f, 1.0f);
+#if RHI_RAYTRACING
 				RayTracingConfig.RayCountPerPixel = GetRayTracingReflectionsSamplesPerPixel(View);
+#else
+				RayTracingConfig.RayCountPerPixel = 0;
+#endif
 
 				if (!bDenoise)
 				{
