Fixes SSGI denoiser not returning AO buffer

Original author: Lucas Ammann

diff --git a/Engine/Source/Runtime/Renderer/Private/ScreenSpaceDenoise.cpp b/Engine/Source/Runtime/Renderer/Private/ScreenSpaceDenoise.cpp
--- a/Engine/Source/Runtime/Renderer/Private/ScreenSpaceDenoise.cpp
+++ b/Engine/Source/Runtime/Renderer/Private/ScreenSpaceDenoise.cpp
@@ -2591,6 +2591,7 @@ public:
 
 		FDiffuseIndirectOutputs GlobalIlluminationOutputs;
 		GlobalIlluminationOutputs.Color = SignalOutput.Textures[0];
+		GlobalIlluminationOutputs.AmbientOcclusionMask = SignalOutput.Textures[1];
 		return GlobalIlluminationOutputs;
 	}
 }; // class FDefaultScreenSpaceDenoiser
