Disable render to virtual texture on primitives with no static mesh (fixes crash when inadvertently adding a render to virtual texture on a scene primitive component that doesn't support it such as a brush component...)

diff --git a/Engine/Source/Runtime/Renderer/Private/PrimitiveSceneInfo.cpp b/Engine/Source/Runtime/Renderer/Private/PrimitiveSceneInfo.cpp
--- a/Engine/Source/Runtime/Renderer/Private/PrimitiveSceneInfo.cpp
+++ b/Engine/Source/Runtime/Renderer/Private/PrimitiveSceneInfo.cpp
@@ -11,6 +11,7 @@
 #include "SceneCore.h"
 #include "VelocityRendering.h"
 #include "ScenePrivate.h"
+#include "RendererModule.h"
 #include "HAL/LowLevelMemTracker.h"
 #include "RayTracing/RayTracingMaterialHitShaders.h"
 #include "VT/RuntimeVirtualTextureSceneProxy.h"
@@ -651,7 +652,7 @@ void FPrimitiveSceneInfo::AddToScene(FRHICommandListImmediate& RHICmdList, bool
 	{
 		int8 MinLod, MaxLod;
 		GetRuntimeVirtualTextureLODRange(StaticMeshRelevances, MinLod, MaxLod);
-		
+
 		FPrimitiveVirtualTextureLodInfo& LodInfo = Scene->PrimitiveVirtualTextureLod[PackedIndex];
 		LodInfo.MinLod = FMath::Clamp((int32)MinLod, 0, 15);
 		LodInfo.MaxLod = FMath::Clamp((int32)MaxLod, 0, 15);
@@ -748,23 +749,30 @@ void FPrimitiveSceneInfo::RemoveFromScene(bool bUpdateStaticDrawLists)
 
 void FPrimitiveSceneInfo::UpdateRuntimeVirtualTextureFlags()
 {
-	const bool bRenderToVirtualTexture = Proxy->WritesVirtualTexture();
-
-	RuntimeVirtualTextureFlags.bRenderToVirtualTexture = bRenderToVirtualTexture;
+	RuntimeVirtualTextureFlags.bRenderToVirtualTexture = false;
 	RuntimeVirtualTextureFlags.RuntimeVirtualTextureMask = 0;
 
-	if (bRenderToVirtualTexture)
+	if (Proxy->WritesVirtualTexture())
 	{
-		// Performance assumption: The arrays of runtime virtual textures are small (less that 5?) so that O(n^2) scan isn't expensive
-		for (TSparseArray<FRuntimeVirtualTextureSceneProxy*>::TConstIterator It(Scene->RuntimeVirtualTextures); It; ++It)
+		if (StaticMeshes.Num() == 0)
 		{
-			int32 SceneIndex = It.GetIndex();
-			if (SceneIndex < FPrimitiveVirtualTextureFlags::RuntimeVirtualTexture_BitCount)
+			UE_LOG(LogRenderer, Warning, TEXT("Rendering a primitive in a runtime virtual texture implies that there is a mesh to render. Please disable this option on primitive component : %s"), *Proxy->GetOwnerName().ToString());
+		}
+		else
+		{
+			RuntimeVirtualTextureFlags.bRenderToVirtualTexture = true;
+
+			// Performance assumption: The arrays of runtime virtual textures are small (less that 5?) so that O(n^2) scan isn't expensive
+			for (TSparseArray<FRuntimeVirtualTextureSceneProxy*>::TConstIterator It(Scene->RuntimeVirtualTextures); It; ++It)
 			{
-				URuntimeVirtualTexture* SceneVirtualTexture = (*It)->VirtualTexture;
-				if (Proxy->WritesVirtualTexture(SceneVirtualTexture))
+				int32 SceneIndex = It.GetIndex();
+				if (SceneIndex < FPrimitiveVirtualTextureFlags::RuntimeVirtualTexture_BitCount)
 				{
-					RuntimeVirtualTextureFlags.RuntimeVirtualTextureMask |= 1 << SceneIndex;
+					URuntimeVirtualTexture* SceneVirtualTexture = (*It)->VirtualTexture;
+					if (Proxy->WritesVirtualTexture(SceneVirtualTexture))
+					{
+						RuntimeVirtualTextureFlags.RuntimeVirtualTextureMask |= 1 << SceneIndex;
+					}
 				}
 			}
 		}
