Moves FReflectionEnvironmentSkyLightingPS into IndirectLightRendering.cpp

diff --git a/Engine/Source/Runtime/Renderer/Private/IndirectLightRendering.cpp b/Engine/Source/Runtime/Renderer/Private/IndirectLightRendering.cpp
--- a/Engine/Source/Runtime/Renderer/Private/IndirectLightRendering.cpp
+++ b/Engine/Source/Runtime/Renderer/Private/IndirectLightRendering.cpp
@@ -8,9 +8,12 @@
 #include "ScreenSpaceDenoise.h"
 #include "ScreenSpaceRayTracing.h"
 #include "DeferredShadingRenderer.h"
+#include "PostProcess/PostProcessSubsurface.h"
+#include "PostProcess/PostProcessTemporalAA.h"
 #include "PostProcessing.h" // for FPostProcessVS
 #include "RendererModule.h" 
 #include "RayTracing/RaytracingOptions.h"
+#include "DistanceFieldAmbientOcclusion.h"
 
 
 static TAutoConsoleVariable<int32> CVarDiffuseIndirectDenoiser(
@@ -18,6 +21,77 @@ static TAutoConsoleVariable<int32> CVarDiffuseIndirectDenoiser(
 	TEXT("Denoising options (default = 1)"),
 	ECVF_RenderThreadSafe);
 
+static int32 GRayTracingReflections = -1;
+static FAutoConsoleVariableRef CVarReflectionsMethod(
+	TEXT("r.RayTracing.Reflections"),
+	GRayTracingReflections,
+	TEXT("-1: Value driven by postprocess volume (default) \n")
+	TEXT("0: use traditional rasterized SSR\n")
+	TEXT("1: use ray traced reflections\n"));
+
+static TAutoConsoleVariable<float> CVarReflectionScreenPercentage(
+	TEXT("r.RayTracing.Reflections.ScreenPercentage"),
+	100.0f,
+	TEXT("Screen percentage the reflections should be ray traced at (default = 100)."),
+	ECVF_RenderThreadSafe);
+
+static int32 GRayTracingReflectionsSamplesPerPixel = -1;
+static FAutoConsoleVariableRef CVarRayTracingReflectionsSamplesPerPixel(
+	TEXT("r.RayTracing.Reflections.SamplesPerPixel"),
+	GRayTracingReflectionsSamplesPerPixel,
+	TEXT("Sets the samples-per-pixel for reflections (default = -1 (driven by postprocesing volume))"));
+
+static int32 GRayTracingReflectionsHeightFog = 1;
+static FAutoConsoleVariableRef CVarRayTracingReflectionsHeightFog(
+	TEXT("r.RayTracing.Reflections.HeightFog"),
+	GRayTracingReflectionsHeightFog,
+	TEXT("Enables height fog in ray traced reflections (default = 1)"));
+
+static TAutoConsoleVariable<int32> CVarUseReflectionDenoiser(
+	TEXT("r.Reflections.Denoiser"),
+	2,
+	TEXT("Choose the denoising algorithm.\n")
+	TEXT(" 0: Disabled;\n")
+	TEXT(" 1: Forces the default denoiser of the renderer;\n")
+	TEXT(" 2: GScreenSpaceDenoiser which may be overriden by a third party plugin (default)."),
+	ECVF_RenderThreadSafe);
+
+static TAutoConsoleVariable<int32> CVarDenoiseSSR(
+	TEXT("r.SSR.ExperimentalDenoiser"), 0,
+	TEXT("Replace SSR's TAA pass with denoiser."),
+	ECVF_RenderThreadSafe);
+
+static TAutoConsoleVariable<float> CVarSkySpecularOcclusionStrength(
+	TEXT("r.SkySpecularOcclusionStrength"),
+	1,
+	TEXT("Strength of skylight specular occlusion from DFAO (default is 1.0)"),
+	ECVF_RenderThreadSafe);
+
+
+DECLARE_GPU_STAT_NAMED(ReflectionEnvironment, TEXT("Reflection Environment"));
+DECLARE_GPU_STAT_NAMED(RayTracingReflections, TEXT("Ray Tracing Reflections"));
+DECLARE_GPU_STAT(SkyLightDiffuse);
+
+int GetReflectionEnvironmentCVar();
+
+#if RHI_RAYTRACING
+int32 GetRayTracingReflectionsSamplesPerPixel(const FViewInfo& View)
+{
+	return GRayTracingReflectionsSamplesPerPixel >= 0 ? GRayTracingReflectionsSamplesPerPixel : View.FinalPostProcessSettings.RayTracingReflectionsSamplesPerPixel;
+}
+
+bool ShouldRenderRayTracingReflections(const FViewInfo& View)
+{
+	bool bThisViewHasRaytracingReflections = View.FinalPostProcessSettings.ReflectionsType == EReflectionsType::RayTracing;
+
+	const bool bReflectionsCvarEnabled = GRayTracingReflections < 0 ? bThisViewHasRaytracingReflections : (GRayTracingReflections != 0);
+	const int32 ForceAllRayTracingEffects = GetForceRayTracingEffectsCVarValue();
+	const bool bReflectionPassEnabled = (ForceAllRayTracingEffects > 0 || (bReflectionsCvarEnabled && ForceAllRayTracingEffects < 0)) && (GetRayTracingReflectionsSamplesPerPixel(View) > 0);
+
+	return IsRayTracingEnabled() && bReflectionPassEnabled;
+}
+#endif // RHI_RAYTRACING
+
 
 bool IsAmbientCubemapPassRequired(const FSceneView& View);
 
@@ -97,9 +171,209 @@ class FAmbientCubemapCompositePS : public FGlobalShader
 	END_SHADER_PARAMETER_STRUCT()
 };
 
+/** Pixel shader that does tiled deferred culling of reflection captures, then sorts and composites them. */
+class FReflectionEnvironmentSkyLightingPS : public FGlobalShader
+{
+	DECLARE_GLOBAL_SHADER(FReflectionEnvironmentSkyLightingPS);
+	SHADER_USE_PARAMETER_STRUCT(FReflectionEnvironmentSkyLightingPS, FGlobalShader)
+
+	class FHasBoxCaptures : SHADER_PERMUTATION_BOOL("REFLECTION_COMPOSITE_HAS_BOX_CAPTURES");
+	class FHasSphereCaptures : SHADER_PERMUTATION_BOOL("REFLECTION_COMPOSITE_HAS_SPHERE_CAPTURES");
+	class FDFAOIndirectOcclusion : SHADER_PERMUTATION_BOOL("SUPPORT_DFAO_INDIRECT_OCCLUSION");
+	class FSkyLight : SHADER_PERMUTATION_BOOL("ENABLE_SKY_LIGHT");
+	class FDynamicSkyLight : SHADER_PERMUTATION_BOOL("ENABLE_DYNAMIC_SKY_LIGHT");
+	class FSkyShadowing : SHADER_PERMUTATION_BOOL("APPLY_SKY_SHADOWING");
+	class FRayTracedReflections : SHADER_PERMUTATION_BOOL("RAY_TRACED_REFLECTIONS");
+
+	using FPermutationDomain = TShaderPermutationDomain<
+		FHasBoxCaptures,
+		FHasSphereCaptures,
+		FDFAOIndirectOcclusion,
+		FSkyLight,
+		FDynamicSkyLight,
+		FSkyShadowing,
+		FRayTracedReflections>;
+
+	static FPermutationDomain RemapPermutation(FPermutationDomain PermutationVector)
+	{
+		// FSkyLightingDynamicSkyLight requires FSkyLightingSkyLight.
+		if (!PermutationVector.Get<FSkyLight>())
+		{
+			PermutationVector.Set<FDynamicSkyLight>(false);
+		}
+
+		// FSkyLightingSkyShadowing requires FSkyLightingDynamicSkyLight.
+		if (!PermutationVector.Get<FDynamicSkyLight>())
+		{
+			PermutationVector.Set<FSkyShadowing>(false);
+		}
+
+		return PermutationVector;
+	}
+
+	static FPermutationDomain BuildPermutationVector(const FViewInfo& View, bool bBoxCapturesOnly, bool bSphereCapturesOnly, bool bSupportDFAOIndirectOcclusion, bool bEnableSkyLight, bool bEnableDynamicSkyLight, bool bApplySkyShadowing, bool bRayTracedReflections)
+	{
+		FPermutationDomain PermutationVector;
+
+		PermutationVector.Set<FHasBoxCaptures>(bBoxCapturesOnly);
+		PermutationVector.Set<FHasSphereCaptures>(bSphereCapturesOnly);
+		PermutationVector.Set<FDFAOIndirectOcclusion>(bSupportDFAOIndirectOcclusion);
+		PermutationVector.Set<FSkyLight>(bEnableSkyLight);
+		PermutationVector.Set<FDynamicSkyLight>(bEnableDynamicSkyLight);
+		PermutationVector.Set<FSkyShadowing>(bApplySkyShadowing);
+		PermutationVector.Set<FRayTracedReflections>(bRayTracedReflections);
+
+		return RemapPermutation(PermutationVector);
+	}
+
+	static bool ShouldCompilePermutation(const FGlobalShaderPermutationParameters& Parameters)
+	{
+		if (!IsFeatureLevelSupported(Parameters.Platform, ERHIFeatureLevel::SM5))
+		{
+			return false;
+		}
+
+		FPermutationDomain PermutationVector(Parameters.PermutationId);
+		return PermutationVector == RemapPermutation(PermutationVector);
+	}
+
+	static void ModifyCompilationEnvironment(const FGlobalShaderPermutationParameters& Parameters, FShaderCompilerEnvironment& OutEnvironment)
+	{
+		FGlobalShader::ModifyCompilationEnvironment(Parameters, OutEnvironment);
+		OutEnvironment.SetDefine(TEXT("MAX_CAPTURES"), GMaxNumReflectionCaptures);
+		OutEnvironment.CompilerFlags.Add(CFLAG_StandardOptimization);
+		FForwardLightingParameters::ModifyCompilationEnvironment(Parameters.Platform, OutEnvironment);
+	}
+
+	BEGIN_SHADER_PARAMETER_STRUCT(FParameters, )
+
+		// Sky light parameters.
+		SHADER_PARAMETER(FVector4, OcclusionTintAndMinOcclusion)
+		SHADER_PARAMETER(FVector, ContrastAndNormalizeMulAdd)
+		SHADER_PARAMETER(float, ApplyBentNormalAO)
+		SHADER_PARAMETER(float, InvSkySpecularOcclusionStrength)
+		SHADER_PARAMETER(float, OcclusionExponent)
+		SHADER_PARAMETER(float, OcclusionCombineMode)
+
+		// Distance field AO parameters.
+		// TODO. FDFAOUpsampleParameters
+		SHADER_PARAMETER(FVector2D, AOBufferBilinearUVMax)
+		SHADER_PARAMETER(float, DistanceFadeScale)
+		SHADER_PARAMETER(float, AOMaxViewDistance)
+		SHADER_PARAMETER_RDG_TEXTURE(Texture2D, BentNormalAOTexture)
+		SHADER_PARAMETER_SAMPLER(SamplerState, BentNormalAOSampler)
+
+		SHADER_PARAMETER_RDG_TEXTURE(Texture2D, AmbientOcclusionTexture)
+		SHADER_PARAMETER_SAMPLER(SamplerState, AmbientOcclusionSampler)
+
+		SHADER_PARAMETER_RDG_TEXTURE(Texture2D, ScreenSpaceReflectionsTexture)
+		SHADER_PARAMETER_SAMPLER(SamplerState, ScreenSpaceReflectionsSampler)
+
+		SHADER_PARAMETER_TEXTURE(Texture2D, PreIntegratedGF)
+		SHADER_PARAMETER_SAMPLER(SamplerState, PreIntegratedGFSampler)
+
+		SHADER_PARAMETER_STRUCT_INCLUDE(FSceneTextureParameters, SceneTextures)
+		SHADER_PARAMETER_STRUCT_INCLUDE(FSceneTextureSamplerParameters, SceneTextureSamplers)
+
+		SHADER_PARAMETER_STRUCT_REF(FViewUniformShaderParameters, ViewUniformBuffer)
+		SHADER_PARAMETER_STRUCT_REF(FReflectionUniformParameters, ReflectionsParameters)
+		SHADER_PARAMETER_STRUCT_REF(FReflectionCaptureShaderData, ReflectionCaptureData)
+		SHADER_PARAMETER_STRUCT_REF(FForwardLightData, ForwardLightData)
+
+		RENDER_TARGET_BINDING_SLOTS()
+		END_SHADER_PARAMETER_STRUCT()
+}; // FReflectionEnvironmentSkyLightingPS
+
 IMPLEMENT_GLOBAL_SHADER(FDiffuseIndirectCompositePS, "/Engine/Private/DiffuseIndirectComposite.usf", "MainPS", SF_Pixel);
 IMPLEMENT_GLOBAL_SHADER(FAmbientCubemapCompositePS, "/Engine/Private/AmbientCubemapComposite.usf", "MainPS", SF_Pixel);
+IMPLEMENT_GLOBAL_SHADER(FReflectionEnvironmentSkyLightingPS, "/Engine/Private/ReflectionEnvironmentPixelShader.usf", "ReflectionEnvironmentSkyLighting", SF_Pixel);
+
+
+IMPLEMENT_GLOBAL_SHADER_PARAMETER_STRUCT(FReflectionUniformParameters, "ReflectionStruct");
+
+void SetupReflectionUniformParameters(const FViewInfo& View, FReflectionUniformParameters& OutParameters)
+{
+	FTexture* SkyLightTextureResource = GBlackTextureCube;
+	FTexture* SkyLightBlendDestinationTextureResource = GBlackTextureCube;
+	float ApplySkyLightMask = 0;
+	float BlendFraction = 0;
+	bool bSkyLightIsDynamic = false;
+	float SkyAverageBrightness = 1.0f;
+
+	const bool bApplySkyLight = View.Family->EngineShowFlags.SkyLighting;
+	const FScene* Scene = (const FScene*)View.Family->Scene;
+
+	if (Scene
+		&& Scene->SkyLight
+		&& Scene->SkyLight->ProcessedTexture
+		&& bApplySkyLight)
+	{
+		const FSkyLightSceneProxy& SkyLight = *Scene->SkyLight;
+		SkyLightTextureResource = SkyLight.ProcessedTexture;
+		BlendFraction = SkyLight.BlendFraction;
+
+		if (SkyLight.BlendFraction > 0.0f && SkyLight.BlendDestinationProcessedTexture)
+		{
+			if (SkyLight.BlendFraction < 1.0f)
+			{
+				SkyLightBlendDestinationTextureResource = SkyLight.BlendDestinationProcessedTexture;
+			}
+			else
+			{
+				SkyLightTextureResource = SkyLight.BlendDestinationProcessedTexture;
+				BlendFraction = 0;
+			}
+		}
+
+		ApplySkyLightMask = 1;
+		bSkyLightIsDynamic = !SkyLight.bHasStaticLighting && !SkyLight.bWantsStaticShadowing;
+		SkyAverageBrightness = SkyLight.AverageBrightness;
+	}
+
+	const int32 CubemapWidth = SkyLightTextureResource->GetSizeX();
+	const float SkyMipCount = FMath::Log2(CubemapWidth) + 1.0f;
+
+	OutParameters.SkyLightCubemap = SkyLightTextureResource->TextureRHI;
+	OutParameters.SkyLightCubemapSampler = SkyLightTextureResource->SamplerStateRHI;
+	OutParameters.SkyLightBlendDestinationCubemap = SkyLightBlendDestinationTextureResource->TextureRHI;
+	OutParameters.SkyLightBlendDestinationCubemapSampler = SkyLightBlendDestinationTextureResource->SamplerStateRHI;
+	OutParameters.SkyLightParameters = FVector4(SkyMipCount - 1.0f, ApplySkyLightMask, bSkyLightIsDynamic ? 1.0f : 0.0f, BlendFraction);
+	OutParameters.SkyLightCubemapBrightness = SkyAverageBrightness;
 
+	// Note: GBlackCubeArrayTexture has an alpha of 0, which is needed to represent invalid data so the sky cubemap can still be applied
+	FRHITexture* CubeArrayTexture = View.FeatureLevel >= ERHIFeatureLevel::SM5 ? GBlackCubeArrayTexture->TextureRHI : GBlackTextureCube->TextureRHI;
+
+	if (View.Family->EngineShowFlags.ReflectionEnvironment
+		&& View.FeatureLevel >= ERHIFeatureLevel::SM5
+		&& Scene
+		&& Scene->ReflectionSceneData.CubemapArray.IsValid()
+		&& Scene->ReflectionSceneData.RegisteredReflectionCaptures.Num())
+	{
+		CubeArrayTexture = Scene->ReflectionSceneData.CubemapArray.GetRenderTarget().ShaderResourceTexture;
+	}
+
+	OutParameters.ReflectionCubemap = CubeArrayTexture;
+	OutParameters.ReflectionCubemapSampler = TStaticSamplerState<SF_Trilinear, AM_Clamp, AM_Clamp, AM_Clamp>::GetRHI();
+
+	OutParameters.PreIntegratedGF = GSystemTextures.PreintegratedGF->GetRenderTargetItem().ShaderResourceTexture;
+	OutParameters.PreIntegratedGFSampler = TStaticSamplerState<SF_Bilinear, AM_Clamp, AM_Clamp, AM_Clamp>::GetRHI();
+}
+
+TUniformBufferRef<FReflectionUniformParameters> CreateReflectionUniformBuffer(const class FViewInfo& View, EUniformBufferUsage Usage)
+{
+	FReflectionUniformParameters ReflectionStruct;
+	SetupReflectionUniformParameters(View, ReflectionStruct);
+	return CreateUniformBufferImmediate(ReflectionStruct, Usage);
+}
+
+bool FDeferredShadingSceneRenderer::ShouldDoReflectionEnvironment() const
+{
+	const ERHIFeatureLevel::Type SceneFeatureLevel = Scene->GetFeatureLevel();
+
+	return IsReflectionEnvironmentAvailable(SceneFeatureLevel)
+		&& Scene->ReflectionSceneData.RegisteredReflectionCaptures.Num()
+		&& ViewFamily.EngineShowFlags.ReflectionEnvironment;
+}
 
 void FDeferredShadingSceneRenderer::RenderDiffuseIndirectAndAmbientOcclusion(FRHICommandListImmediate& RHICmdListImmediate)
 {
@@ -359,3 +633,347 @@ void FDeferredShadingSceneRenderer::RenderDiffuseIndirectAndAmbientOcclusion(FRH
 
 	GraphBuilder.Execute();
 }
+
+void FDeferredShadingSceneRenderer::RenderDeferredReflectionsAndSkyLighting(FRHICommandListImmediate& RHICmdList, TRefCountPtr<IPooledRenderTarget>& DynamicBentNormalAO, TRefCountPtr<IPooledRenderTarget>& VelocityRT, const FHairStrandsDatas* HairDatas)
+{
+	check(RHICmdList.IsOutsideRenderPass());
+
+	if (ViewFamily.EngineShowFlags.VisualizeLightCulling || !ViewFamily.EngineShowFlags.Lighting)
+	{
+		return;
+	}
+
+	// If we're currently capturing a reflection capture, output SpecularColor * IndirectIrradiance for metals so they are not black in reflections,
+	// Since we don't have multiple bounce specular reflections
+	bool bReflectionCapture = false;
+	for (int32 ViewIndex = 0, Num = Views.Num(); ViewIndex < Num; ViewIndex++)
+	{
+		const FViewInfo& View = Views[ViewIndex];
+		bReflectionCapture = bReflectionCapture || View.bIsReflectionCapture;
+	}
+
+	if (bReflectionCapture)
+	{
+		// if we are rendering a reflection capture then we can skip this pass entirely (no reflection and no sky contribution evaluated in this pass)
+		return;
+	}
+
+	// The specular sky light contribution is also needed by RT Reflections as a fallback.
+	const bool bSkyLight = Scene->SkyLight
+		&& Scene->SkyLight->ProcessedTexture
+		&& !Scene->SkyLight->bHasStaticLighting;
+
+	bool bDynamicSkyLight = ShouldRenderDeferredDynamicSkyLight(Scene, ViewFamily);
+	bool bApplySkyShadowing = false;
+	if (bDynamicSkyLight)
+	{
+		SCOPED_DRAW_EVENT(RHICmdList, SkyLightDiffuse);
+		SCOPED_GPU_STAT(RHICmdList, SkyLightDiffuse);
+
+		FDistanceFieldAOParameters Parameters(Scene->SkyLight->OcclusionMaxDistance, Scene->SkyLight->Contrast);
+
+		extern int32 GDistanceFieldAOApplyToStaticIndirect;
+		if (Scene->SkyLight->bCastShadows
+			&& !GDistanceFieldAOApplyToStaticIndirect
+			&& ShouldRenderDistanceFieldAO()
+			&& ViewFamily.EngineShowFlags.AmbientOcclusion)
+		{
+			// TODO: convert to RDG.
+			bApplySkyShadowing = RenderDistanceFieldLighting(RHICmdList, Parameters, VelocityRT, DynamicBentNormalAO, false, false);
+		}
+	}
+
+	check(RHICmdList.IsOutsideRenderPass());
+
+	FSceneRenderTargets& SceneContext = FSceneRenderTargets::Get(RHICmdList);
+
+	const bool bReflectionEnv = ShouldDoReflectionEnvironment();
+
+	FRDGBuilder GraphBuilder(RHICmdList);
+
+	FRDGTextureRef SceneColorTexture = GraphBuilder.RegisterExternalTexture(SceneContext.GetSceneColor());
+	FRDGTextureRef SceneColorSubPixelTexture = nullptr;
+	if (HairDatas)
+	{
+		SceneColorSubPixelTexture = GraphBuilder.RegisterExternalTexture(SceneContext.GetSceneColorSubPixel());
+	}
+	FRDGTextureRef AmbientOcclusionTexture = GraphBuilder.RegisterExternalTexture(SceneContext.bScreenSpaceAOIsValid ? SceneContext.ScreenSpaceAO : GSystemTextures.WhiteDummy);
+	FRDGTextureRef DynamicBentNormalAOTexture = GraphBuilder.RegisterExternalTexture(DynamicBentNormalAO ? DynamicBentNormalAO : GSystemTextures.WhiteDummy);
+
+	FSceneTextureParameters SceneTextures;
+	SetupSceneTextureParameters(GraphBuilder, &SceneTextures);
+
+	uint32 ViewIndex = 0;
+	for (FViewInfo& View : Views)
+	{
+		const uint32 CurrentViewIndex = ViewIndex++;
+		const bool bRayTracedReflections = ShouldRenderRayTracingReflections(View);
+		const bool bScreenSpaceReflections = !bRayTracedReflections && ShouldRenderScreenSpaceReflections(View);
+
+		const bool bComposePlanarReflections = !bRayTracedReflections && HasDeferredPlanarReflections(View);
+
+		FRDGTextureRef ReflectionsColor = nullptr;
+		if (bRayTracedReflections || bScreenSpaceReflections)
+		{
+			int32 DenoiserMode = CVarUseReflectionDenoiser.GetValueOnRenderThread();
+
+			bool bDenoise = false;
+			bool bTemporalFilter = false;
+
+			// Traces the reflections, either using screen space reflection, or ray tracing.
+			IScreenSpaceDenoiser::FReflectionsInputs DenoiserInputs;
+			IScreenSpaceDenoiser::FReflectionsRayTracingConfig RayTracingConfig;
+			if (bRayTracedReflections)
+			{
+				RDG_EVENT_SCOPE(GraphBuilder, "RayTracingReflections");
+				RDG_GPU_STAT_SCOPE(GraphBuilder, RayTracingReflections);
+
+				bDenoise = DenoiserMode != 0;
+
+				RayTracingConfig.ResolutionFraction = FMath::Clamp(CVarReflectionScreenPercentage.GetValueOnRenderThread() / 100.0f, 0.25f, 1.0f);
+#if RHI_RAYTRACING
+				RayTracingConfig.RayCountPerPixel = GetRayTracingReflectionsSamplesPerPixel(View);
+#else
+				RayTracingConfig.RayCountPerPixel = 0;
+#endif
+
+				if (!bDenoise)
+				{
+					RayTracingConfig.ResolutionFraction = 1.0f;
+				}
+
+				RenderRayTracingReflections(
+					GraphBuilder,
+					SceneTextures,
+					View,
+					RayTracingConfig.RayCountPerPixel, GRayTracingReflectionsHeightFog, RayTracingConfig.ResolutionFraction,
+					&DenoiserInputs);
+			}
+			else if (bScreenSpaceReflections)
+			{
+				bDenoise = DenoiserMode != 0 && CVarDenoiseSSR.GetValueOnRenderThread();
+				bTemporalFilter = !bDenoise && View.ViewState && IsSSRTemporalPassRequired(View);
+
+				FRDGTextureRef CurrentSceneColor = GraphBuilder.RegisterExternalTexture(SceneContext.GetSceneColor());
+
+				ESSRQuality SSRQuality;
+				GetSSRQualityForView(View, &SSRQuality, &RayTracingConfig);
+
+				RDG_EVENT_SCOPE(GraphBuilder, "ScreenSpaceReflections(Quality=%d)", int32(SSRQuality));
+
+				RenderScreenSpaceReflections(
+					GraphBuilder, SceneTextures, CurrentSceneColor, View, SSRQuality, bDenoise, &DenoiserInputs);
+			}
+			else
+			{
+				check(0);
+			}
+
+			if (bDenoise)
+			{
+				const IScreenSpaceDenoiser* DefaultDenoiser = IScreenSpaceDenoiser::GetDefaultDenoiser();
+				const IScreenSpaceDenoiser* DenoiserToUse = DenoiserMode == 1 ? DefaultDenoiser : GScreenSpaceDenoiser;
+
+				// Standard event scope for denoiser to have all profiling information not matter what, and with explicit detection of third party.
+				RDG_EVENT_SCOPE(GraphBuilder, "%s%s(Reflections) %dx%d",
+					DenoiserToUse != DefaultDenoiser ? TEXT("ThirdParty ") : TEXT(""),
+					DenoiserToUse->GetDebugName(),
+					View.ViewRect.Width(), View.ViewRect.Height());
+
+				IScreenSpaceDenoiser::FReflectionsOutputs DenoiserOutputs = DenoiserToUse->DenoiseReflections(
+					GraphBuilder,
+					View,
+					&View.PrevViewInfo,
+					SceneTextures,
+					DenoiserInputs,
+					RayTracingConfig);
+
+				ReflectionsColor = DenoiserOutputs.Color;
+			}
+			else if (bTemporalFilter)
+			{
+				check(View.ViewState);
+				FTAAPassParameters TAASettings(View);
+				TAASettings.Pass = ETAAPassConfig::ScreenSpaceReflections;
+				TAASettings.SceneColorInput = DenoiserInputs.Color;
+				TAASettings.bOutputRenderTargetable = bComposePlanarReflections;
+
+				FTAAOutputs TAAOutputs = AddTemporalAAPass(
+					GraphBuilder,
+					SceneTextures,
+					View,
+					TAASettings,
+					View.PrevViewInfo.SSRHistory,
+					&View.ViewState->PrevFrameViewInfo.SSRHistory);
+
+				ReflectionsColor = TAAOutputs.SceneColor;
+			}
+			else
+			{
+				if (bRayTracedReflections && DenoiserInputs.RayHitDistance)
+				{
+					// The performance of ray tracing does not allow to run without a denoiser in real time.
+					// Multiple rays per pixel is unsupported by the denoiser that will most likely more bound by to
+					// many rays than exporting the hit distance buffer. Therefore no permutation of the ray generation
+					// shader has been judged required to be supported.
+					GraphBuilder.RemoveUnusedTextureWarning(DenoiserInputs.RayHitDistance);
+				}
+
+				ReflectionsColor = DenoiserInputs.Color;
+			}
+		} // if (bRayTracedReflections || bScreenSpaceReflections)
+
+		if (bComposePlanarReflections)
+		{
+			check(!bRayTracedReflections);
+			RenderDeferredPlanarReflections(GraphBuilder, SceneTextures, View, /* inout */ ReflectionsColor);
+		}
+
+		bool bRequiresApply = ReflectionsColor != nullptr || bSkyLight || bDynamicSkyLight || bReflectionEnv;
+
+		if (bRequiresApply)
+		{
+			RDG_GPU_STAT_SCOPE(GraphBuilder, ReflectionEnvironment);
+
+			// Render the reflection environment with tiled deferred culling
+			bool bHasBoxCaptures = (View.NumBoxReflectionCaptures > 0);
+			bool bHasSphereCaptures = (View.NumSphereReflectionCaptures > 0);
+
+			FReflectionEnvironmentSkyLightingPS::FParameters* PassParameters = GraphBuilder.AllocParameters<FReflectionEnvironmentSkyLightingPS::FParameters>();
+
+			// Setup the parameters of the shader.
+			{
+				// Setups all shader parameters related to skylight.
+				{
+					FSkyLightSceneProxy* SkyLight = Scene->SkyLight;
+
+					float SkyLightContrast = 0.01f;
+					float SkyLightOcclusionExponent = 1.0f;
+					FVector4 SkyLightOcclusionTintAndMinOcclusion(0.0f, 0.0f, 0.0f, 0.0f);
+					EOcclusionCombineMode SkyLightOcclusionCombineMode = EOcclusionCombineMode::OCM_MAX;
+					if (SkyLight)
+					{
+						FDistanceFieldAOParameters Parameters(SkyLight->OcclusionMaxDistance, SkyLight->Contrast);
+						SkyLightContrast = Parameters.Contrast;
+						SkyLightOcclusionExponent = SkyLight->OcclusionExponent;
+						SkyLightOcclusionTintAndMinOcclusion = FVector4(SkyLight->OcclusionTint);
+						SkyLightOcclusionTintAndMinOcclusion.W = SkyLight->MinOcclusion;
+						SkyLightOcclusionCombineMode = SkyLight->OcclusionCombineMode;
+					}
+
+					// Scale and bias to remap the contrast curve to [0,1]
+					const float Min = 1 / (1 + FMath::Exp(-SkyLightContrast * (0 * 10 - 5)));
+					const float Max = 1 / (1 + FMath::Exp(-SkyLightContrast * (1 * 10 - 5)));
+					const float Mul = 1.0f / (Max - Min);
+					const float Add = -Min / (Max - Min);
+
+					PassParameters->OcclusionTintAndMinOcclusion = SkyLightOcclusionTintAndMinOcclusion;
+					PassParameters->ContrastAndNormalizeMulAdd = FVector(SkyLightContrast, Mul, Add);
+					PassParameters->OcclusionExponent = SkyLightOcclusionExponent;
+					PassParameters->OcclusionCombineMode = SkyLightOcclusionCombineMode == OCM_Minimum ? 0.0f : 1.0f;
+					PassParameters->ApplyBentNormalAO = DynamicBentNormalAO ? 1.0f : 0.0f;
+					PassParameters->InvSkySpecularOcclusionStrength = 1.0f / FMath::Max(CVarSkySpecularOcclusionStrength.GetValueOnRenderThread(), 0.1f);
+				}
+
+				// Setups all shader parameters related to distance field AO
+				{
+					FIntPoint AOBufferSize = GetBufferSizeForAO();
+					PassParameters->AOBufferBilinearUVMax = FVector2D(
+						(View.ViewRect.Width() / GAODownsampleFactor - 0.51f) / AOBufferSize.X, // 0.51 - so bilateral gather4 won't sample invalid texels
+						(View.ViewRect.Height() / GAODownsampleFactor - 0.51f) / AOBufferSize.Y);
+
+					extern float GAOViewFadeDistanceScale;
+					PassParameters->AOMaxViewDistance = GetMaxAOViewDistance();
+					PassParameters->DistanceFadeScale = 1.0f / ((1.0f - GAOViewFadeDistanceScale) * GetMaxAOViewDistance());
+
+					PassParameters->BentNormalAOTexture = DynamicBentNormalAOTexture;
+					PassParameters->BentNormalAOSampler = TStaticSamplerState<SF_Bilinear>::GetRHI();
+				}
+
+				PassParameters->AmbientOcclusionTexture = AmbientOcclusionTexture;
+				PassParameters->AmbientOcclusionSampler = TStaticSamplerState<SF_Point>::GetRHI();
+
+				PassParameters->ScreenSpaceReflectionsTexture = ReflectionsColor ? ReflectionsColor : GraphBuilder.RegisterExternalTexture(GSystemTextures.BlackDummy);
+				PassParameters->ScreenSpaceReflectionsSampler = TStaticSamplerState<SF_Point>::GetRHI();
+
+				PassParameters->PreIntegratedGF = GSystemTextures.PreintegratedGF->GetRenderTargetItem().ShaderResourceTexture;
+				PassParameters->PreIntegratedGFSampler = TStaticSamplerState<SF_Bilinear, AM_Clamp, AM_Clamp, AM_Clamp>::GetRHI();
+
+				PassParameters->SceneTextures = SceneTextures;
+				SetupSceneTextureSamplers(&PassParameters->SceneTextureSamplers);
+
+				PassParameters->ViewUniformBuffer = View.ViewUniformBuffer;
+				PassParameters->ReflectionCaptureData = View.ReflectionCaptureUniformBuffer;
+				{
+					FReflectionUniformParameters ReflectionUniformParameters;
+					SetupReflectionUniformParameters(View, ReflectionUniformParameters);
+					PassParameters->ReflectionsParameters = CreateUniformBufferImmediate(ReflectionUniformParameters, UniformBuffer_SingleDraw);
+				}
+				PassParameters->ForwardLightData = View.ForwardLightingResources->ForwardLightDataUniformBuffer;
+			}
+
+			PassParameters->RenderTargets[0] = FRenderTargetBinding(SceneColorTexture, ERenderTargetLoadAction::ELoad);
+
+			// Bind hair data
+			const bool bCheckerboardSubsurfaceRendering = IsSubsurfaceCheckerboardFormat(SceneContext.GetSceneColorFormat());
+			auto PermutationVector = FReflectionEnvironmentSkyLightingPS::BuildPermutationVector(
+				View, bHasBoxCaptures, bHasSphereCaptures, DynamicBentNormalAO != NULL, bSkyLight, bDynamicSkyLight, bApplySkyShadowing, bRayTracedReflections);
+
+			TShaderMapRef<FReflectionEnvironmentSkyLightingPS> PixelShader(View.ShaderMap, PermutationVector);
+			ClearUnusedGraphResources(*PixelShader, PassParameters);
+
+			GraphBuilder.AddPass(
+				RDG_EVENT_NAME("ReflectionEnvironmentAndSky %dx%d", View.ViewRect.Width(), View.ViewRect.Height()),
+				PassParameters,
+				ERDGPassFlags::Raster,
+				[PassParameters, &View, PixelShader, bCheckerboardSubsurfaceRendering](FRHICommandList& InRHICmdList)
+			{
+				InRHICmdList.SetViewport(View.ViewRect.Min.X, View.ViewRect.Min.Y, 0.0f, View.ViewRect.Max.X, View.ViewRect.Max.Y, 1.0f);
+
+				FGraphicsPipelineStateInitializer GraphicsPSOInit;
+				FPixelShaderUtils::InitFullscreenPipelineState(InRHICmdList, View.ShaderMap, *PixelShader, GraphicsPSOInit);
+
+				extern int32 GAOOverwriteSceneColor;
+				if (GetReflectionEnvironmentCVar() == 2 || GAOOverwriteSceneColor)
+				{
+					// override scene color for debugging
+					GraphicsPSOInit.BlendState = TStaticBlendState<>::GetRHI();
+				}
+				else
+				{
+					if (bCheckerboardSubsurfaceRendering)
+					{
+						GraphicsPSOInit.BlendState = TStaticBlendState<CW_RGB, BO_Add, BF_One, BF_One>::GetRHI();
+					}
+					else
+					{
+						GraphicsPSOInit.BlendState = TStaticBlendState<CW_RGBA, BO_Add, BF_One, BF_One, BO_Add, BF_One, BF_One>::GetRHI();
+					}
+				}
+
+				SetGraphicsPipelineState(InRHICmdList, GraphicsPSOInit);
+				SetShaderParameters(InRHICmdList, *PixelShader, PixelShader->GetPixelShader(), *PassParameters);
+				FPixelShaderUtils::DrawFullscreenTriangle(InRHICmdList);
+			});
+		} // if (bRequiresApply)
+
+
+		if (HairDatas)
+		{
+			RenderHairStrandsEnvironmentLighting(GraphBuilder, CurrentViewIndex, Views, HairDatas, SceneColorTexture, SceneColorSubPixelTexture);
+		}
+	} // for (FViewInfo& View : Views)
+
+	TRefCountPtr<IPooledRenderTarget> OutSceneColor;
+	GraphBuilder.QueueTextureExtraction(SceneColorTexture, &OutSceneColor);
+
+	TRefCountPtr<IPooledRenderTarget> OutSceneSubpixelColor;
+	if (SceneColorSubPixelTexture)
+	{
+		GraphBuilder.QueueTextureExtraction(SceneColorSubPixelTexture, &OutSceneSubpixelColor);
+	}
+
+	GraphBuilder.Execute();
+
+	ResolveSceneColor(RHICmdList);
+}
diff --git a/Engine/Source/Runtime/Renderer/Private/ReflectionEnvironment.cpp b/Engine/Source/Runtime/Renderer/Private/ReflectionEnvironment.cpp
--- a/Engine/Source/Runtime/Renderer/Private/ReflectionEnvironment.cpp
+++ b/Engine/Source/Runtime/Renderer/Private/ReflectionEnvironment.cpp
@@ -24,7 +24,6 @@
 #include "PostProcess/SceneFilterRendering.h"
 #include "PostProcess/PostProcessing.h"
 #include "PostProcess/PostProcessSubsurface.h"
-#include "PostProcess/PostProcessTemporalAA.h"
 #include "LightRendering.h"
 #include "LightPropagationVolumeSettings.h"
 #include "PipelineStateCache.h"
@@ -38,9 +37,6 @@
 #include "SceneTextureParameters.h"
 #include "HairStrands/HairStrandsRendering.h"
 
-DECLARE_GPU_STAT_NAMED(ReflectionEnvironment, TEXT("Reflection Environment"));
-DECLARE_GPU_STAT_NAMED(RayTracingReflections, TEXT("Ray Tracing Reflections"));
-DECLARE_GPU_STAT(SkyLightDiffuse);
 
 extern TAutoConsoleVariable<int32> CVarLPVMixing;
 
@@ -101,56 +97,9 @@ static TAutoConsoleVariable<int32> CVarDoTiledReflections(
 	TEXT(" 1: on (default)"),
 	ECVF_RenderThreadSafe);
 
-static TAutoConsoleVariable<float> CVarSkySpecularOcclusionStrength(
-	TEXT("r.SkySpecularOcclusionStrength"),
-	1,
-	TEXT("Strength of skylight specular occlusion from DFAO (default is 1.0)"),
-	ECVF_RenderThreadSafe);
-
-static int32 GRayTracingReflections = -1;
-static FAutoConsoleVariableRef CVarReflectionsMethod(
-	TEXT("r.RayTracing.Reflections"),
-	GRayTracingReflections,
-	TEXT("-1: Value driven by postprocess volume (default) \n")
-	TEXT("0: use traditional rasterized SSR\n")
-	TEXT("1: use ray traced reflections\n")
-);
-
-static TAutoConsoleVariable<float> CVarReflectionScreenPercentage(
-	TEXT("r.RayTracing.Reflections.ScreenPercentage"),
-	100.0f,
-	TEXT("Screen percentage the reflections should be ray traced at (default = 100)."),
-	ECVF_RenderThreadSafe);
-
-static int32 GRayTracingReflectionsSamplesPerPixel = -1;
-static FAutoConsoleVariableRef CVarRayTracingReflectionsSamplesPerPixel(
-	TEXT("r.RayTracing.Reflections.SamplesPerPixel"),
-	GRayTracingReflectionsSamplesPerPixel,
-	TEXT("Sets the samples-per-pixel for reflections (default = -1 (driven by postprocesing volume))"));
-
-static int32 GRayTracingReflectionsHeightFog = 1;
-static FAutoConsoleVariableRef CVarRayTracingReflectionsHeightFog(
-	TEXT("r.RayTracing.Reflections.HeightFog"),
-	GRayTracingReflectionsHeightFog,
-	TEXT("Enables height fog in ray traced reflections (default = 1)"));
-
-static TAutoConsoleVariable<int32> CVarUseReflectionDenoiser(
-	TEXT("r.Reflections.Denoiser"),
-	2,
-	TEXT("Choose the denoising algorithm.\n")
-	TEXT(" 0: Disabled;\n")
-	TEXT(" 1: Forces the default denoiser of the renderer;\n")
-	TEXT(" 2: GScreenSpaceDenoiser which may be overriden by a third party plugin (default)."),
-	ECVF_RenderThreadSafe);
-
-static TAutoConsoleVariable<int32> CVarDenoiseSSR(
-	TEXT("r.SSR.ExperimentalDenoiser"), 0,
-	TEXT("Replace SSR's TAA pass with denoiser."),
-	ECVF_RenderThreadSafe);
-
 
 // to avoid having direct access from many places
-static int GetReflectionEnvironmentCVar()
+int GetReflectionEnvironmentCVar()
 {
 	int32 RetVal = CVarReflectionEnvironment.GetValueOnAnyThread();
 
@@ -199,101 +148,6 @@ bool IsReflectionCaptureAvailable()
 	return (!AllowStaticLightingVar || AllowStaticLightingVar->GetInt() != 0);
 }
 
-#if RHI_RAYTRACING
-int32 GetRayTracingReflectionsSamplesPerPixel(const FViewInfo& View)
-{
-	return GRayTracingReflectionsSamplesPerPixel >= 0 ? GRayTracingReflectionsSamplesPerPixel : View.FinalPostProcessSettings.RayTracingReflectionsSamplesPerPixel;
-}
-
-bool ShouldRenderRayTracingReflections(const FViewInfo& View)
-{
-	bool bThisViewHasRaytracingReflections = View.FinalPostProcessSettings.ReflectionsType == EReflectionsType::RayTracing;
-
-	const bool bReflectionsCvarEnabled = GRayTracingReflections < 0 ? bThisViewHasRaytracingReflections : (GRayTracingReflections != 0);
-	const int32 ForceAllRayTracingEffects = GetForceRayTracingEffectsCVarValue();
-	const bool bReflectionPassEnabled = (ForceAllRayTracingEffects > 0 || (bReflectionsCvarEnabled && ForceAllRayTracingEffects < 0)) && (GetRayTracingReflectionsSamplesPerPixel(View) > 0);
-
-	return IsRayTracingEnabled() && bReflectionPassEnabled;
-}
-#endif // RHI_RAYTRACING
-
-IMPLEMENT_GLOBAL_SHADER_PARAMETER_STRUCT(FReflectionUniformParameters, "ReflectionStruct");
-
-void SetupReflectionUniformParameters(const FViewInfo& View, FReflectionUniformParameters& OutParameters)
-{
-	FTexture* SkyLightTextureResource = GBlackTextureCube;
-	FTexture* SkyLightBlendDestinationTextureResource = GBlackTextureCube;
-	float ApplySkyLightMask = 0;
-	float BlendFraction = 0;
-	bool bSkyLightIsDynamic = false;
-	float SkyAverageBrightness = 1.0f;
-
-	const bool bApplySkyLight = View.Family->EngineShowFlags.SkyLighting;
-	const FScene* Scene = (const FScene*)View.Family->Scene;
-
-	if (Scene
-		&& Scene->SkyLight 
-		&& Scene->SkyLight->ProcessedTexture
-		&& bApplySkyLight)
-	{
-		const FSkyLightSceneProxy& SkyLight = *Scene->SkyLight;
-		SkyLightTextureResource = SkyLight.ProcessedTexture;
-		BlendFraction = SkyLight.BlendFraction;
-
-		if (SkyLight.BlendFraction > 0.0f && SkyLight.BlendDestinationProcessedTexture)
-		{
-			if (SkyLight.BlendFraction < 1.0f)
-			{
-				SkyLightBlendDestinationTextureResource = SkyLight.BlendDestinationProcessedTexture;
-			}
-			else
-			{
-				SkyLightTextureResource = SkyLight.BlendDestinationProcessedTexture;
-				BlendFraction = 0;
-			}
-		}
-		
-		ApplySkyLightMask = 1;
-		bSkyLightIsDynamic = !SkyLight.bHasStaticLighting && !SkyLight.bWantsStaticShadowing;
-		SkyAverageBrightness = SkyLight.AverageBrightness;
-	}
-
-	const int32 CubemapWidth = SkyLightTextureResource->GetSizeX();
-	const float SkyMipCount = FMath::Log2(CubemapWidth) + 1.0f;
-
-	OutParameters.SkyLightCubemap = SkyLightTextureResource->TextureRHI;
-	OutParameters.SkyLightCubemapSampler = SkyLightTextureResource->SamplerStateRHI;
-	OutParameters.SkyLightBlendDestinationCubemap = SkyLightBlendDestinationTextureResource->TextureRHI;
-	OutParameters.SkyLightBlendDestinationCubemapSampler = SkyLightBlendDestinationTextureResource->SamplerStateRHI;
-	OutParameters.SkyLightParameters = FVector4(SkyMipCount - 1.0f, ApplySkyLightMask, bSkyLightIsDynamic ? 1.0f : 0.0f, BlendFraction);
-	OutParameters.SkyLightCubemapBrightness = SkyAverageBrightness;
-
-	// Note: GBlackCubeArrayTexture has an alpha of 0, which is needed to represent invalid data so the sky cubemap can still be applied
-	FRHITexture* CubeArrayTexture = View.FeatureLevel >= ERHIFeatureLevel::SM5 ? GBlackCubeArrayTexture->TextureRHI : GBlackTextureCube->TextureRHI;
-
-	if (View.Family->EngineShowFlags.ReflectionEnvironment 
-		&& View.FeatureLevel >= ERHIFeatureLevel::SM5
-		&& Scene
-		&& Scene->ReflectionSceneData.CubemapArray.IsValid()
-		&& Scene->ReflectionSceneData.RegisteredReflectionCaptures.Num())
-	{
-		CubeArrayTexture = Scene->ReflectionSceneData.CubemapArray.GetRenderTarget().ShaderResourceTexture;
-	}
-
-	OutParameters.ReflectionCubemap = CubeArrayTexture;
-	OutParameters.ReflectionCubemapSampler = TStaticSamplerState<SF_Trilinear, AM_Clamp, AM_Clamp, AM_Clamp>::GetRHI();
-
-	OutParameters.PreIntegratedGF = GSystemTextures.PreintegratedGF->GetRenderTargetItem().ShaderResourceTexture;
-	OutParameters.PreIntegratedGFSampler = TStaticSamplerState<SF_Bilinear, AM_Clamp, AM_Clamp, AM_Clamp>::GetRHI();
-}
-
-TUniformBufferRef<FReflectionUniformParameters> CreateReflectionUniformBuffer(const class FViewInfo& View, EUniformBufferUsage Usage)
-{
-	FReflectionUniformParameters ReflectionStruct;
-	SetupReflectionUniformParameters(View, ReflectionStruct);
-	return CreateUniformBufferImmediate(ReflectionStruct, Usage);
-}
-
 void FReflectionEnvironmentCubemapArray::InitDynamicRHI()
 {
 	if (GetFeatureLevel() >= ERHIFeatureLevel::SM5)
@@ -413,7 +267,7 @@ void FReflectionEnvironmentCubemapArray::ResizeCubemapArrayGPU(uint32 InMaxCubem
 
 	{
 		SCOPED_DRAW_EVENT(RHICmdList, ReflectionEnvironment_ResizeCubemapArray);
-		SCOPED_GPU_STAT(RHICmdList, ReflectionEnvironment)
+		//SCOPED_GPU_STAT(RHICmdList, ReflectionEnvironment)
 
 		// Copy the cubemaps, remapping the elements as necessary
 		FResolveParams ResolveParams;
@@ -465,471 +319,3 @@ void FReflectionEnvironmentCubemapArray::UpdateMaxCubemaps(uint32 InMaxCubemaps,
 }
 
 IMPLEMENT_GLOBAL_SHADER_PARAMETER_STRUCT(FReflectionCaptureShaderData, "ReflectionCapture");
-
-/** Pixel shader that does tiled deferred culling of reflection captures, then sorts and composites them. */
-class FReflectionEnvironmentSkyLightingPS : public FGlobalShader
-{
-	DECLARE_GLOBAL_SHADER(FReflectionEnvironmentSkyLightingPS);
-	SHADER_USE_PARAMETER_STRUCT(FReflectionEnvironmentSkyLightingPS, FGlobalShader)
-
-	class FHasBoxCaptures			: SHADER_PERMUTATION_BOOL("REFLECTION_COMPOSITE_HAS_BOX_CAPTURES");
-	class FHasSphereCaptures		: SHADER_PERMUTATION_BOOL("REFLECTION_COMPOSITE_HAS_SPHERE_CAPTURES");
-	class FDFAOIndirectOcclusion	: SHADER_PERMUTATION_BOOL("SUPPORT_DFAO_INDIRECT_OCCLUSION");
-	class FSkyLight					: SHADER_PERMUTATION_BOOL("ENABLE_SKY_LIGHT");
-	class FDynamicSkyLight			: SHADER_PERMUTATION_BOOL("ENABLE_DYNAMIC_SKY_LIGHT");
-	class FSkyShadowing				: SHADER_PERMUTATION_BOOL("APPLY_SKY_SHADOWING");
-	class FRayTracedReflections		: SHADER_PERMUTATION_BOOL("RAY_TRACED_REFLECTIONS");
-
-	using FPermutationDomain = TShaderPermutationDomain<
-		FHasBoxCaptures,
-		FHasSphereCaptures,
-		FDFAOIndirectOcclusion,
-		FSkyLight,
-		FDynamicSkyLight,
-		FSkyShadowing,
-		FRayTracedReflections>;
-
-	static FPermutationDomain RemapPermutation(FPermutationDomain PermutationVector)
-	{
-		// FSkyLightingDynamicSkyLight requires FSkyLightingSkyLight.
-		if (!PermutationVector.Get<FSkyLight>())
-		{
-			PermutationVector.Set<FDynamicSkyLight>(false);
-		}
-
-		// FSkyLightingSkyShadowing requires FSkyLightingDynamicSkyLight.
-		if (!PermutationVector.Get<FDynamicSkyLight>())
-		{
-			PermutationVector.Set<FSkyShadowing>(false);
-		}
-
-		return PermutationVector;
-	}
-
-	static FPermutationDomain BuildPermutationVector(const FViewInfo& View, bool bBoxCapturesOnly, bool bSphereCapturesOnly, bool bSupportDFAOIndirectOcclusion, bool bEnableSkyLight, bool bEnableDynamicSkyLight, bool bApplySkyShadowing, bool bRayTracedReflections)
-	{
-		FPermutationDomain PermutationVector;
-
-		PermutationVector.Set<FHasBoxCaptures>(bBoxCapturesOnly);
-		PermutationVector.Set<FHasSphereCaptures>(bSphereCapturesOnly);
-		PermutationVector.Set<FDFAOIndirectOcclusion>(bSupportDFAOIndirectOcclusion);
-		PermutationVector.Set<FSkyLight>(bEnableSkyLight);
-		PermutationVector.Set<FDynamicSkyLight>(bEnableDynamicSkyLight);
-		PermutationVector.Set<FSkyShadowing>(bApplySkyShadowing);
-		PermutationVector.Set<FRayTracedReflections>(bRayTracedReflections);
-
-		return RemapPermutation(PermutationVector);
-	}
-
-	static bool ShouldCompilePermutation(const FGlobalShaderPermutationParameters& Parameters)
-	{
-		if (!IsFeatureLevelSupported(Parameters.Platform, ERHIFeatureLevel::SM5))
-		{
-			return false;
-		}
-
-		FPermutationDomain PermutationVector(Parameters.PermutationId);
-		return PermutationVector == RemapPermutation(PermutationVector);
-	}
-
-	static void ModifyCompilationEnvironment(const FGlobalShaderPermutationParameters& Parameters, FShaderCompilerEnvironment& OutEnvironment)
-	{
-		FGlobalShader::ModifyCompilationEnvironment(Parameters, OutEnvironment);
-		OutEnvironment.SetDefine(TEXT("MAX_CAPTURES"), GMaxNumReflectionCaptures);
-		OutEnvironment.CompilerFlags.Add(CFLAG_StandardOptimization);
-		FForwardLightingParameters::ModifyCompilationEnvironment(Parameters.Platform, OutEnvironment);
-	}
-
-	BEGIN_SHADER_PARAMETER_STRUCT(FParameters, )
-
-		// Sky light parameters.
-		SHADER_PARAMETER(FVector4, OcclusionTintAndMinOcclusion)
-		SHADER_PARAMETER(FVector, ContrastAndNormalizeMulAdd)
-		SHADER_PARAMETER(float, ApplyBentNormalAO)
-		SHADER_PARAMETER(float, InvSkySpecularOcclusionStrength)
-		SHADER_PARAMETER(float, OcclusionExponent)
-		SHADER_PARAMETER(float, OcclusionCombineMode)
-		
-		// Distance field AO parameters.
-		// TODO. FDFAOUpsampleParameters
-		SHADER_PARAMETER(FVector2D, AOBufferBilinearUVMax)
-		SHADER_PARAMETER(float, DistanceFadeScale)
-		SHADER_PARAMETER(float, AOMaxViewDistance)
-		SHADER_PARAMETER_RDG_TEXTURE(Texture2D, BentNormalAOTexture)
-		SHADER_PARAMETER_SAMPLER(SamplerState,  BentNormalAOSampler)
-		
-		SHADER_PARAMETER_RDG_TEXTURE(Texture2D, AmbientOcclusionTexture)
-		SHADER_PARAMETER_SAMPLER(SamplerState,  AmbientOcclusionSampler)
-		
-		SHADER_PARAMETER_RDG_TEXTURE(Texture2D, ScreenSpaceReflectionsTexture)
-		SHADER_PARAMETER_SAMPLER(SamplerState,  ScreenSpaceReflectionsSampler)
-		
-		SHADER_PARAMETER_TEXTURE(Texture2D,    PreIntegratedGF)
-		SHADER_PARAMETER_SAMPLER(SamplerState, PreIntegratedGFSampler)
-		
-		SHADER_PARAMETER_STRUCT_INCLUDE(FSceneTextureParameters, SceneTextures)
-		SHADER_PARAMETER_STRUCT_INCLUDE(FSceneTextureSamplerParameters, SceneTextureSamplers)
-
-		SHADER_PARAMETER_STRUCT_REF(FViewUniformShaderParameters, ViewUniformBuffer)
-		SHADER_PARAMETER_STRUCT_REF(FReflectionUniformParameters, ReflectionsParameters)
-		SHADER_PARAMETER_STRUCT_REF(FReflectionCaptureShaderData, ReflectionCaptureData)
-		SHADER_PARAMETER_STRUCT_REF(FForwardLightData, ForwardLightData)
-
-		RENDER_TARGET_BINDING_SLOTS()
-	END_SHADER_PARAMETER_STRUCT()
-}; // FReflectionEnvironmentSkyLightingPS
-
-IMPLEMENT_GLOBAL_SHADER(FReflectionEnvironmentSkyLightingPS, "/Engine/Private/ReflectionEnvironmentPixelShader.usf", "ReflectionEnvironmentSkyLighting", SF_Pixel);
-
-bool FDeferredShadingSceneRenderer::ShouldDoReflectionEnvironment() const
-{
-	const ERHIFeatureLevel::Type SceneFeatureLevel = Scene->GetFeatureLevel();
-
-	return IsReflectionEnvironmentAvailable(SceneFeatureLevel)
-		&& Scene->ReflectionSceneData.RegisteredReflectionCaptures.Num()
-		&& ViewFamily.EngineShowFlags.ReflectionEnvironment;
-}
-
-void FDeferredShadingSceneRenderer::RenderDeferredReflectionsAndSkyLighting(FRHICommandListImmediate& RHICmdList, TRefCountPtr<IPooledRenderTarget>& DynamicBentNormalAO, TRefCountPtr<IPooledRenderTarget>& VelocityRT, const FHairStrandsDatas* HairDatas)
-{
-	check(RHICmdList.IsOutsideRenderPass());
-
-	if (ViewFamily.EngineShowFlags.VisualizeLightCulling || !ViewFamily.EngineShowFlags.Lighting)
-	{
-		return;
-	}
-
-	// If we're currently capturing a reflection capture, output SpecularColor * IndirectIrradiance for metals so they are not black in reflections,
-	// Since we don't have multiple bounce specular reflections
-	bool bReflectionCapture = false;
-	for (int32 ViewIndex = 0, Num = Views.Num(); ViewIndex < Num; ViewIndex++)
-	{
-		const FViewInfo& View = Views[ViewIndex];
-		bReflectionCapture = bReflectionCapture || View.bIsReflectionCapture;
-	}
-
-	if (bReflectionCapture)
-	{
-		// if we are rendering a reflection capture then we can skip this pass entirely (no reflection and no sky contribution evaluated in this pass)
-		return;	
-	}
-
-	// The specular sky light contribution is also needed by RT Reflections as a fallback.
-	const bool bSkyLight = Scene->SkyLight
-		&& Scene->SkyLight->ProcessedTexture
-		&& !Scene->SkyLight->bHasStaticLighting;
-
-	bool bDynamicSkyLight = ShouldRenderDeferredDynamicSkyLight(Scene, ViewFamily);
-	bool bApplySkyShadowing = false;
-	if (bDynamicSkyLight)
-	{
-		SCOPED_DRAW_EVENT(RHICmdList, SkyLightDiffuse);
-		SCOPED_GPU_STAT(RHICmdList, SkyLightDiffuse);
-
-		FDistanceFieldAOParameters Parameters(Scene->SkyLight->OcclusionMaxDistance, Scene->SkyLight->Contrast);
-
-		extern int32 GDistanceFieldAOApplyToStaticIndirect;
-		if (Scene->SkyLight->bCastShadows
-			&& !GDistanceFieldAOApplyToStaticIndirect
-			&& ShouldRenderDistanceFieldAO()
-			&& ViewFamily.EngineShowFlags.AmbientOcclusion)
-		{
-			// TODO: convert to RDG.
-			bApplySkyShadowing = RenderDistanceFieldLighting(RHICmdList, Parameters, VelocityRT, DynamicBentNormalAO, false, false);
-		}
-	}
-
-	check(RHICmdList.IsOutsideRenderPass());
-
-	FSceneRenderTargets& SceneContext = FSceneRenderTargets::Get(RHICmdList);
-
-	const bool bReflectionEnv = ShouldDoReflectionEnvironment();
-
-	FRDGBuilder GraphBuilder(RHICmdList);
-
-	FRDGTextureRef SceneColorTexture = GraphBuilder.RegisterExternalTexture(SceneContext.GetSceneColor());
-	FRDGTextureRef SceneColorSubPixelTexture = nullptr;
-	if (HairDatas)
-	{
-		SceneColorSubPixelTexture = GraphBuilder.RegisterExternalTexture(SceneContext.GetSceneColorSubPixel());
-	}
-	FRDGTextureRef AmbientOcclusionTexture = GraphBuilder.RegisterExternalTexture(SceneContext.bScreenSpaceAOIsValid ? SceneContext.ScreenSpaceAO : GSystemTextures.WhiteDummy);
-	FRDGTextureRef DynamicBentNormalAOTexture = GraphBuilder.RegisterExternalTexture(DynamicBentNormalAO ? DynamicBentNormalAO : GSystemTextures.WhiteDummy);
-
-	FSceneTextureParameters SceneTextures;
-	SetupSceneTextureParameters(GraphBuilder, &SceneTextures);
-
-	uint32 ViewIndex = 0;
-	for (FViewInfo& View : Views)
-	{
-		const uint32 CurrentViewIndex = ViewIndex++;
-		const bool bRayTracedReflections = ShouldRenderRayTracingReflections(View);
-		const bool bScreenSpaceReflections = !bRayTracedReflections && ShouldRenderScreenSpaceReflections(View);
-
-		const bool bComposePlanarReflections = !bRayTracedReflections && HasDeferredPlanarReflections(View);
-
-		FRDGTextureRef ReflectionsColor = nullptr;
-		if (bRayTracedReflections || bScreenSpaceReflections)
-		{
-			int32 DenoiserMode = CVarUseReflectionDenoiser.GetValueOnRenderThread();
-					
-			bool bDenoise = false;
-			bool bTemporalFilter = false;
-
-			// Traces the reflections, either using screen space reflection, or ray tracing.
-			IScreenSpaceDenoiser::FReflectionsInputs DenoiserInputs;
-			IScreenSpaceDenoiser::FReflectionsRayTracingConfig RayTracingConfig;
-			if (bRayTracedReflections)
-			{
-				RDG_EVENT_SCOPE(GraphBuilder, "RayTracingReflections");
-				RDG_GPU_STAT_SCOPE(GraphBuilder, RayTracingReflections);
-
-				bDenoise = DenoiserMode != 0;
-				
-				RayTracingConfig.ResolutionFraction = FMath::Clamp(CVarReflectionScreenPercentage.GetValueOnRenderThread() / 100.0f, 0.25f, 1.0f);
-#if RHI_RAYTRACING
-				RayTracingConfig.RayCountPerPixel = GetRayTracingReflectionsSamplesPerPixel(View);
-#else
-				RayTracingConfig.RayCountPerPixel = 0;
-#endif
-
-				if (!bDenoise)
-				{
-					RayTracingConfig.ResolutionFraction = 1.0f;
-				}
-
-				RenderRayTracingReflections(
-					GraphBuilder,
-					SceneTextures,
-					View,
-					RayTracingConfig.RayCountPerPixel, GRayTracingReflectionsHeightFog, RayTracingConfig.ResolutionFraction,
-					&DenoiserInputs);
-			}
-			else if (bScreenSpaceReflections)
-			{
-				bDenoise = DenoiserMode != 0 && CVarDenoiseSSR.GetValueOnRenderThread();
-				bTemporalFilter = !bDenoise && View.ViewState && IsSSRTemporalPassRequired(View);
-
-				FRDGTextureRef CurrentSceneColor = GraphBuilder.RegisterExternalTexture(SceneContext.GetSceneColor());
-
-				ESSRQuality SSRQuality;
-				GetSSRQualityForView(View, &SSRQuality, &RayTracingConfig);
-				
-				RDG_EVENT_SCOPE(GraphBuilder, "ScreenSpaceReflections(Quality=%d)", int32(SSRQuality));
-			
-				RenderScreenSpaceReflections(
-					GraphBuilder, SceneTextures, CurrentSceneColor, View, SSRQuality, bDenoise, &DenoiserInputs);
-			}
-			else
-			{
-				check(0);
-			}
-
-			if (bDenoise)
-			{
-				const IScreenSpaceDenoiser* DefaultDenoiser = IScreenSpaceDenoiser::GetDefaultDenoiser();
-				const IScreenSpaceDenoiser* DenoiserToUse = DenoiserMode == 1 ? DefaultDenoiser : GScreenSpaceDenoiser;
-
-				// Standard event scope for denoiser to have all profiling information not matter what, and with explicit detection of third party.
-				RDG_EVENT_SCOPE(GraphBuilder, "%s%s(Reflections) %dx%d",
-					DenoiserToUse != DefaultDenoiser ? TEXT("ThirdParty ") : TEXT(""),
-					DenoiserToUse->GetDebugName(),
-					View.ViewRect.Width(), View.ViewRect.Height());
-
-				IScreenSpaceDenoiser::FReflectionsOutputs DenoiserOutputs = DenoiserToUse->DenoiseReflections(
-					GraphBuilder,
-					View,
-					&View.PrevViewInfo,
-					SceneTextures,
-					DenoiserInputs,
-					RayTracingConfig);
-
-				ReflectionsColor = DenoiserOutputs.Color;
-			}
-			else if (bTemporalFilter)
-			{
-				check(View.ViewState);
-				FTAAPassParameters TAASettings(View);
-				TAASettings.Pass = ETAAPassConfig::ScreenSpaceReflections;
-				TAASettings.SceneColorInput = DenoiserInputs.Color;
-				TAASettings.bOutputRenderTargetable = bComposePlanarReflections;
-				
-				FTAAOutputs TAAOutputs = AddTemporalAAPass(
-					GraphBuilder,
-					SceneTextures,
-					View,
-					TAASettings,
-					View.PrevViewInfo.SSRHistory,
-					&View.ViewState->PrevFrameViewInfo.SSRHistory);
-				
-				ReflectionsColor = TAAOutputs.SceneColor;
-			}
-			else
-			{
-				if (bRayTracedReflections && DenoiserInputs.RayHitDistance)
-				{
-					// The performance of ray tracing does not allow to run without a denoiser in real time.
-					// Multiple rays per pixel is unsupported by the denoiser that will most likely more bound by to
-					// many rays than exporting the hit distance buffer. Therefore no permutation of the ray generation
-					// shader has been judged required to be supported.
-					GraphBuilder.RemoveUnusedTextureWarning(DenoiserInputs.RayHitDistance);
-				}
-				
-				ReflectionsColor = DenoiserInputs.Color;
-			}
-		} // if (bRayTracedReflections || bScreenSpaceReflections)
-
-		if (bComposePlanarReflections)
-		{
-			check(!bRayTracedReflections);
-			RenderDeferredPlanarReflections(GraphBuilder, SceneTextures, View, /* inout */ ReflectionsColor);
-		}
-
-		bool bRequiresApply = ReflectionsColor != nullptr || bSkyLight || bDynamicSkyLight || bReflectionEnv;
-
-		if (bRequiresApply)
-		{
-			RDG_GPU_STAT_SCOPE(GraphBuilder, ReflectionEnvironment);
-
-			// Render the reflection environment with tiled deferred culling
-			bool bHasBoxCaptures = (View.NumBoxReflectionCaptures > 0);
-			bool bHasSphereCaptures = (View.NumSphereReflectionCaptures > 0);
-			
-			FReflectionEnvironmentSkyLightingPS::FParameters* PassParameters = GraphBuilder.AllocParameters<FReflectionEnvironmentSkyLightingPS::FParameters>();
-
-			// Setup the parameters of the shader.
-			{
-				// Setups all shader parameters related to skylight.
-				{
-					FSkyLightSceneProxy* SkyLight = Scene->SkyLight;
-
-					float SkyLightContrast = 0.01f;
-					float SkyLightOcclusionExponent = 1.0f;
-					FVector4 SkyLightOcclusionTintAndMinOcclusion(0.0f, 0.0f, 0.0f, 0.0f);
-					EOcclusionCombineMode SkyLightOcclusionCombineMode = EOcclusionCombineMode::OCM_MAX;
-					if (SkyLight)
-					{
-						FDistanceFieldAOParameters Parameters(SkyLight->OcclusionMaxDistance, SkyLight->Contrast);
-						SkyLightContrast = Parameters.Contrast;
-						SkyLightOcclusionExponent = SkyLight->OcclusionExponent;
-						SkyLightOcclusionTintAndMinOcclusion = FVector4(SkyLight->OcclusionTint);
-						SkyLightOcclusionTintAndMinOcclusion.W = SkyLight->MinOcclusion;
-						SkyLightOcclusionCombineMode = SkyLight->OcclusionCombineMode;
-					}
-
-					// Scale and bias to remap the contrast curve to [0,1]
-					const float Min = 1 / (1 + FMath::Exp(-SkyLightContrast * (0 * 10 - 5)));
-					const float Max = 1 / (1 + FMath::Exp(-SkyLightContrast * (1 * 10 - 5)));
-					const float Mul = 1.0f / (Max - Min);
-					const float Add = -Min / (Max - Min);
-					
-					PassParameters->OcclusionTintAndMinOcclusion = SkyLightOcclusionTintAndMinOcclusion;
-					PassParameters->ContrastAndNormalizeMulAdd = FVector(SkyLightContrast, Mul, Add);
-					PassParameters->OcclusionExponent = SkyLightOcclusionExponent;
-					PassParameters->OcclusionCombineMode = SkyLightOcclusionCombineMode == OCM_Minimum ? 0.0f : 1.0f;
-					PassParameters->ApplyBentNormalAO = DynamicBentNormalAO ? 1.0f : 0.0f;
-					PassParameters->InvSkySpecularOcclusionStrength = 1.0f / FMath::Max(CVarSkySpecularOcclusionStrength.GetValueOnRenderThread(), 0.1f);
-				}
-
-				// Setups all shader parameters related to distance field AO
-				{
-					FIntPoint AOBufferSize = GetBufferSizeForAO();
-					PassParameters->AOBufferBilinearUVMax = FVector2D(
-						(View.ViewRect.Width() / GAODownsampleFactor - 0.51f) / AOBufferSize.X, // 0.51 - so bilateral gather4 won't sample invalid texels
-						(View.ViewRect.Height() / GAODownsampleFactor - 0.51f) / AOBufferSize.Y);
-
-					extern float GAOViewFadeDistanceScale;
-					PassParameters->AOMaxViewDistance = GetMaxAOViewDistance();
-					PassParameters->DistanceFadeScale = 1.0f / ((1.0f - GAOViewFadeDistanceScale) * GetMaxAOViewDistance());
-
-					PassParameters->BentNormalAOTexture = DynamicBentNormalAOTexture;
-					PassParameters->BentNormalAOSampler = TStaticSamplerState<SF_Bilinear>::GetRHI();
-				}
-
-				PassParameters->AmbientOcclusionTexture = AmbientOcclusionTexture;
-				PassParameters->AmbientOcclusionSampler = TStaticSamplerState<SF_Point>::GetRHI();
-
-				PassParameters->ScreenSpaceReflectionsTexture = ReflectionsColor ? ReflectionsColor : GraphBuilder.RegisterExternalTexture(GSystemTextures.BlackDummy);
-				PassParameters->ScreenSpaceReflectionsSampler = TStaticSamplerState<SF_Point>::GetRHI();
-
-				PassParameters->PreIntegratedGF = GSystemTextures.PreintegratedGF->GetRenderTargetItem().ShaderResourceTexture;
-				PassParameters->PreIntegratedGFSampler = TStaticSamplerState<SF_Bilinear, AM_Clamp, AM_Clamp, AM_Clamp>::GetRHI();
-				
-				PassParameters->SceneTextures = SceneTextures;
-				SetupSceneTextureSamplers(&PassParameters->SceneTextureSamplers);
-
-				PassParameters->ViewUniformBuffer = View.ViewUniformBuffer;
-				PassParameters->ReflectionCaptureData = View.ReflectionCaptureUniformBuffer;
-				{
-					FReflectionUniformParameters ReflectionUniformParameters;
-					SetupReflectionUniformParameters(View, ReflectionUniformParameters);
-					PassParameters->ReflectionsParameters = CreateUniformBufferImmediate(ReflectionUniformParameters, UniformBuffer_SingleDraw);
-				}
-				PassParameters->ForwardLightData = View.ForwardLightingResources->ForwardLightDataUniformBuffer;
-			}
-
-			PassParameters->RenderTargets[0] = FRenderTargetBinding(SceneColorTexture, ERenderTargetLoadAction::ELoad);
-
-			// Bind hair data
-			const bool bCheckerboardSubsurfaceRendering = IsSubsurfaceCheckerboardFormat(SceneContext.GetSceneColorFormat());
-			auto PermutationVector = FReflectionEnvironmentSkyLightingPS::BuildPermutationVector(
-				View, bHasBoxCaptures, bHasSphereCaptures, DynamicBentNormalAO != NULL, bSkyLight, bDynamicSkyLight, bApplySkyShadowing, bRayTracedReflections);
-
-			TShaderMapRef<FReflectionEnvironmentSkyLightingPS> PixelShader(View.ShaderMap, PermutationVector);
-			ClearUnusedGraphResources(*PixelShader, PassParameters);
-
-			GraphBuilder.AddPass(
-				RDG_EVENT_NAME("ReflectionEnvironmentAndSky %dx%d", View.ViewRect.Width(), View.ViewRect.Height()),
-				PassParameters,
-				ERDGPassFlags::Raster,
-				[PassParameters, &View, PixelShader, bCheckerboardSubsurfaceRendering](FRHICommandList& InRHICmdList)
-			{
-				InRHICmdList.SetViewport(View.ViewRect.Min.X, View.ViewRect.Min.Y, 0.0f, View.ViewRect.Max.X, View.ViewRect.Max.Y, 1.0f);
-
-				FGraphicsPipelineStateInitializer GraphicsPSOInit;
-				FPixelShaderUtils::InitFullscreenPipelineState(InRHICmdList, View.ShaderMap, *PixelShader, GraphicsPSOInit);
-
-				extern int32 GAOOverwriteSceneColor;
-				if (GetReflectionEnvironmentCVar() == 2 || GAOOverwriteSceneColor)
-				{
-					// override scene color for debugging
-					GraphicsPSOInit.BlendState = TStaticBlendState<>::GetRHI();
-				}
-				else
-				{
-					if (bCheckerboardSubsurfaceRendering)
-					{
-						GraphicsPSOInit.BlendState = TStaticBlendState<CW_RGB, BO_Add, BF_One, BF_One>::GetRHI();
-					}
-					else
-					{
-						GraphicsPSOInit.BlendState = TStaticBlendState<CW_RGBA, BO_Add, BF_One, BF_One, BO_Add, BF_One, BF_One>::GetRHI();
-					}
-				}
-
-				SetGraphicsPipelineState(InRHICmdList, GraphicsPSOInit);
-				SetShaderParameters(InRHICmdList, *PixelShader, PixelShader->GetPixelShader(), *PassParameters);
-				FPixelShaderUtils::DrawFullscreenTriangle(InRHICmdList);
-			});
-		} // if (bRequiresApply)
-
-
-		if (HairDatas)
-		{
-			RenderHairStrandsEnvironmentLighting(GraphBuilder, CurrentViewIndex, Views, HairDatas, SceneColorTexture, SceneColorSubPixelTexture);
-		}
-	} // for (FViewInfo& View : Views)
-	
-	TRefCountPtr<IPooledRenderTarget> OutSceneColor;
-	GraphBuilder.QueueTextureExtraction(SceneColorTexture, &OutSceneColor);
-
-	TRefCountPtr<IPooledRenderTarget> OutSceneSubpixelColor;
-	if (SceneColorSubPixelTexture)
-	{
-		GraphBuilder.QueueTextureExtraction(SceneColorSubPixelTexture, &OutSceneSubpixelColor);
-	}
-
-	GraphBuilder.Execute();
-	
-	ResolveSceneColor(RHICmdList);
-}
