- Make GDefaultMaterialParameterCollectionInstances a TMultiMap, instead of TMap
- It's possible for the same MPC to be loaded multiple times.  If this happens, it will get added to this map multiple times.  If it's not a multimap, the first instance that's destroyed will remove the ID from the map, which will cause further lookups against the map to fail

diff --git a/Engine/Source/Runtime/Engine/Private/Materials/ParameterCollection.cpp b/Engine/Source/Runtime/Engine/Private/Materials/ParameterCollection.cpp
--- a/Engine/Source/Runtime/Engine/Private/Materials/ParameterCollection.cpp
+++ b/Engine/Source/Runtime/Engine/Private/Materials/ParameterCollection.cpp
@@ -20,7 +20,7 @@ FAutoConsoleVariableRef CVarDeferUpdateRenderStates(
 	ECVF_RenderThreadSafe
 );
 
-TMap<FGuid, FMaterialParameterCollectionInstanceResource*> GDefaultMaterialParameterCollectionInstances;
+TMultiMap<FGuid, FMaterialParameterCollectionInstanceResource*> GDefaultMaterialParameterCollectionInstances;
 
 UMaterialParameterCollection::UMaterialParameterCollection(const FObjectInitializer& ObjectInitializer)
 	: Super(ObjectInitializer)
@@ -63,11 +63,12 @@ void UMaterialParameterCollection::BeginDestroy()
 {
 	if (DefaultResource)
 	{
+		FMaterialParameterCollectionInstanceResource* Resource = DefaultResource;
 		FGuid Id = StateId;
 		ENQUEUE_RENDER_COMMAND(RemoveDefaultResourceCommand)(
-			[Id](FRHICommandListImmediate& RHICmdList)
+			[Resource, Id](FRHICommandListImmediate& RHICmdList)
 			{	
-				GDefaultMaterialParameterCollectionInstances.Remove(Id);			
+				GDefaultMaterialParameterCollectionInstances.RemoveSingle(Id, Resource);
 			}
 		);
 	}
diff --git a/Engine/Source/Runtime/Engine/Public/ParameterCollection.h b/Engine/Source/Runtime/Engine/Public/ParameterCollection.h
--- a/Engine/Source/Runtime/Engine/Public/ParameterCollection.h
+++ b/Engine/Source/Runtime/Engine/Public/ParameterCollection.h
@@ -69,4 +69,4 @@ private:
 };
 
 // Default instance resources used when rendering a material using a parameter collection but there's no FScene present to get a FMaterialParameterCollectionInstanceResource
-extern ENGINE_API TMap<FGuid, FMaterialParameterCollectionInstanceResource*> GDefaultMaterialParameterCollectionInstances;
+extern ENGINE_API TMultiMap<FGuid, FMaterialParameterCollectionInstanceResource*> GDefaultMaterialParameterCollectionInstances;
diff --git a/Engine/Source/Runtime/Renderer/Private/ShaderBaseClasses.cpp b/Engine/Source/Runtime/Renderer/Private/ShaderBaseClasses.cpp
--- a/Engine/Source/Runtime/Renderer/Private/ShaderBaseClasses.cpp
+++ b/Engine/Source/Runtime/Renderer/Private/ShaderBaseClasses.cpp
@@ -287,7 +287,7 @@ void FMaterialShader::SetParametersInner(
 				// In a cooked project these numbers are persistent so we can track back to the original
 				// parameter collection that was being referenced and no longer exists
 				FString InstancesString;
-				TMap<FGuid, FMaterialParameterCollectionInstanceResource*>::TIterator Iter = GDefaultMaterialParameterCollectionInstances.CreateIterator();
+				TMultiMap<FGuid, FMaterialParameterCollectionInstanceResource*>::TIterator Iter = GDefaultMaterialParameterCollectionInstances.CreateIterator();
 				while (Iter)
 				{
 					FMaterialParameterCollectionInstanceResource* Instance = Iter.Value();
@@ -465,7 +465,7 @@ void FMaterialShader::GetShaderBindings(
 				// In a cooked project these numbers are persistent so we can track back to the original
 				// parameter collection that was being referenced and no longer exists
 				FString InstancesString;
-				TMap<FGuid, FMaterialParameterCollectionInstanceResource*>::TIterator Iter = GDefaultMaterialParameterCollectionInstances.CreateIterator();
+				TMultiMap<FGuid, FMaterialParameterCollectionInstanceResource*>::TIterator Iter = GDefaultMaterialParameterCollectionInstances.CreateIterator();
 				while (Iter)
 				{
 					FMaterialParameterCollectionInstanceResource* Instance = Iter.Value();
