Added height fog pass rendering for below water surface views.
- See RenderUnderWaterFog
- Generalised RenderViewFog to handle non full res output and alternative depth source (still need to follow multi view layout)
- Tested Switch like permutation and PC dx11 et dx12 and ps4.

Removed found "disable/enable optimisation" tags from 10409240.

diff --git a/Engine/Shaders/Private/HeightFogPixelShader.usf b/Engine/Shaders/Private/HeightFogPixelShader.usf
--- a/Engine/Shaders/Private/HeightFogPixelShader.usf
+++ b/Engine/Shaders/Private/HeightFogPixelShader.usf
@@ -10,10 +10,14 @@
 #include "HeightFogCommon.ush"
 
 float bOnlyOnRenderedOpaque;
+float bUseLinearDepthTexture;
 
 Texture2D OcclusionTexture;
 SamplerState OcclusionSampler;
 
+Texture2D LinearDepthTexture;
+SamplerState LinearDepthSampler;
+
 void ExponentialPixelMain(
 	float2 TexCoord : TEXCOORD0,
 	float3 ScreenVector : TEXCOORD1,
@@ -21,14 +25,23 @@ void ExponentialPixelMain(
 	out float4 OutColor : SV_Target0
 	)
 {
+	float SceneDepth;
+	float DeviceZ;
+	if (bUseLinearDepthTexture > 0.0f)
+	{
+		SceneDepth = Texture2DSampleLevel(LinearDepthTexture, LinearDepthSampler, TexCoord, 0).r * bUseLinearDepthTexture; // bUseLinearDepthTexture is also a scale to get depth in correct range
+		DeviceZ = ConvertToDeviceZ(SceneDepth);
+	}
+	else
+	{
 	#if SCENE_TEXTURES_DISABLED
 		float SceneDepth = SCENE_TEXTURES_DISABLED_SCENE_DEPTH_VALUE;
+		DeviceZ = ConvertToDeviceZ(SceneDepth);
 	#else
-		float DeviceZ = Texture2DSampleLevel(SceneTexturesStruct.SceneDepthTexture, SceneTexturesStruct.SceneDepthTextureSampler, TexCoord, 0).r;
-
-		// Fetch the depth buffer Z / W value, solve for W
-		float SceneDepth = ConvertFromDeviceZ(DeviceZ);
+		DeviceZ = Texture2DSampleLevel(SceneTexturesStruct.SceneDepthTexture, SceneTexturesStruct.SceneDepthTextureSampler, TexCoord, 0).r;
+		SceneDepth = ConvertFromDeviceZ(DeviceZ); // Fetch the depth buffer Z / W value, solve for W
 	#endif
+	}
 
 	bool bIsRendered = (DeviceZ != 0.0);
 
diff --git a/Engine/Shaders/Private/SingleLayerWaterComposite.usf b/Engine/Shaders/Private/SingleLayerWaterComposite.usf
--- a/Engine/Shaders/Private/SingleLayerWaterComposite.usf
+++ b/Engine/Shaders/Private/SingleLayerWaterComposite.usf
@@ -13,6 +13,7 @@
 #include "LightGridCommon.ush"
 
 // Scale for encoding depth into SceneDepthWithoutSingleLayerWaterTexture
+// Must match FDeferredShadingSceneRenderer::RenderUnderWaterFog
 #define SINGLE_LAYER_WATER_DEPTH_SCALE 100.0f
 
 #define ENABLE_SKY_LIGHT 1
diff --git a/Engine/Source/Runtime/Renderer/Private/DeferredShadingRenderer.cpp b/Engine/Source/Runtime/Renderer/Private/DeferredShadingRenderer.cpp
--- a/Engine/Source/Runtime/Renderer/Private/DeferredShadingRenderer.cpp
+++ b/Engine/Source/Runtime/Renderer/Private/DeferredShadingRenderer.cpp
@@ -2062,6 +2062,12 @@ void FDeferredShadingSceneRenderer::Render(FRHICommandListImmediate& RHICmdList)
 		SingleLayerWaterPassData.ViewData.SetNum(Views.Num());
 		CopySingleLayerWaterTextures(RHICmdList, SingleLayerWaterPassData);
 
+		// Render heightfog over the color buffer if it is allocated, e.g. SingleLayerWaterUsesSimpleShading is true which is not the case on Switch.
+		if (bCanOverlayRayTracingOutput && ShouldRenderFog(ViewFamily) && SingleLayerWaterPassData.SceneColorWithoutSingleLayerWater.IsValid())
+		{
+			RenderUnderWaterFog(RHICmdList, SingleLayerWaterPassData);
+		}
+
 		// Make the Depth texture writable since the water GBuffer pass will update it
 		RHICmdList.TransitionResource(EResourceTransitionAccess::EWritable,  SceneContext.GetSceneDepthSurface());
 
diff --git a/Engine/Source/Runtime/Renderer/Private/DeferredShadingRenderer.h b/Engine/Source/Runtime/Renderer/Private/DeferredShadingRenderer.h
--- a/Engine/Source/Runtime/Renderer/Private/DeferredShadingRenderer.h
+++ b/Engine/Source/Runtime/Renderer/Private/DeferredShadingRenderer.h
@@ -28,6 +28,7 @@ class UStaticMeshComponent;
 class FExponentialHeightFogSceneInfo;
 
 struct FSingleLayerWaterPassData;
+struct FHeightFogRenderingParameters;
 
 class FLightShaftsOutput
 {
@@ -214,9 +215,10 @@ private:
 
 	/** Renders the scene's fogging. */
 	bool RenderFog(FRHICommandListImmediate& RHICmdList, const FLightShaftsOutput& LightShaftsOutput);
+	void RenderUnderWaterFog(FRHICommandListImmediate& RHICmdList, FSingleLayerWaterPassData& PassData);
 	
 	/** Renders the scene's fogging for a view. */
-	void RenderViewFog(FRHICommandList& RHICmdList, const FViewInfo& View, const FLightShaftsOutput& LightShaftsOutput);
+	void RenderViewFog(FRHICommandList& RHICmdList, const FViewInfo& View, const FHeightFogRenderingParameters& Params);
 
 	/** Renders the scene's atmosphere. */
 	void RenderAtmosphere(FRHICommandListImmediate& RHICmdList, const FLightShaftsOutput& LightShaftsOutput);
diff --git a/Engine/Source/Runtime/Renderer/Private/FogRendering.cpp b/Engine/Source/Runtime/Renderer/Private/FogRendering.cpp
--- a/Engine/Source/Runtime/Renderer/Private/FogRendering.cpp
+++ b/Engine/Source/Runtime/Renderer/Private/FogRendering.cpp
@@ -10,6 +10,7 @@
 #include "ScenePrivate.h"
 #include "Engine/TextureCube.h"
 #include "PipelineStateCache.h"
+#include "SingleLayerWaterRendering.h"
 
 DECLARE_GPU_STAT(Fog);
 
@@ -43,6 +44,14 @@ static TAutoConsoleVariable<int32> CVarFog(
 
 IMPLEMENT_GLOBAL_SHADER_PARAMETER_STRUCT(FFogUniformParameters, "FogStruct");
 
+struct FHeightFogRenderingParameters
+{
+	const FLightShaftsOutput& LightShaftsOutput;
+	FTextureRHIRef LinearDepthTextureRHI;
+	FIntRect ViewRect;
+	float LinearDepthReadScale;
+};
+
 void SetupFogUniformParameters(const FViewInfo& View, FFogUniformParameters& OutParameters)
 {
 	// Exponential Height Fog
@@ -186,11 +195,14 @@ public:
 	{
 		OcclusionTexture.Bind(Initializer.ParameterMap, TEXT("OcclusionTexture"));
 		OcclusionSampler.Bind(Initializer.ParameterMap, TEXT("OcclusionSampler"));
+		LinearDepthTexture.Bind(Initializer.ParameterMap, TEXT("LinearDepthTexture"));
+		LinearDepthSampler.Bind(Initializer.ParameterMap, TEXT("LinearDepthSampler"));
 		bOnlyOnRenderedOpaque.Bind(Initializer.ParameterMap, TEXT("bOnlyOnRenderedOpaque"));
+		bUseLinearDepthTexture.Bind(Initializer.ParameterMap, TEXT("bUseLinearDepthTexture"));
 		SceneTextureParameters.Bind(Initializer);
 	}
 
-	void SetParameters(FRHICommandList& RHICmdList, const FViewInfo& View, const FLightShaftsOutput& LightShaftsOutput)
+	void SetParameters(FRHICommandList& RHICmdList, const FViewInfo& View, const FHeightFogRenderingParameters& Params)
 	{
 		FGlobalShader::SetParameters<FViewUniformShaderParameters>(RHICmdList, GetPixelShader(), View.ViewUniformBuffer);
 		SceneTextureParameters.Set(RHICmdList, GetPixelShader(), View.FeatureLevel, ESceneTextureSetupMode::All);
@@ -199,10 +211,9 @@ public:
 		SetupFogUniformParameters(View, FogUniformParameters);
 		SetUniformBufferParameterImmediate(RHICmdList, GetPixelShader(), GetUniformBufferParameter<FFogUniformParameters>(), FogUniformParameters);
 
-		FTextureRHIRef TextureRHI = LightShaftsOutput.LightShaftOcclusion ?
-			LightShaftsOutput.LightShaftOcclusion->GetRenderTargetItem().ShaderResourceTexture :
+		FTextureRHIRef TextureRHI = Params.LightShaftsOutput.LightShaftOcclusion ?
+			Params.LightShaftsOutput.LightShaftOcclusion->GetRenderTargetItem().ShaderResourceTexture :
 			GWhiteTexture->TextureRHI;
-
 		SetTextureParameter(
 			RHICmdList, 
 			GetPixelShader(),
@@ -211,8 +222,18 @@ public:
 			TextureRHI
 			);
 
-		float bOnlyOnRenderedOpaqueValue = View.bFogOnlyOnRenderedOpaque ? 1.0f : 0.0f;
-		SetShaderValue(RHICmdList, GetPixelShader(), bOnlyOnRenderedOpaque, bOnlyOnRenderedOpaqueValue);
+		const bool bUseLinearDepthTextureEnabled = Params.LinearDepthTextureRHI != nullptr;
+		FTextureRHIRef LinearDepthTextureRHI = bUseLinearDepthTextureEnabled ? Params.LinearDepthTextureRHI : GSystemTextures.DepthDummy->GetRenderTargetItem().ShaderResourceTexture;
+		SetTextureParameter(
+			RHICmdList, 
+			GetPixelShader(),
+			LinearDepthTexture, LinearDepthSampler,
+			TStaticSamplerState<SF_Point,AM_Clamp,AM_Clamp,AM_Clamp>::GetRHI(),
+			LinearDepthTextureRHI
+			);
+
+		SetShaderValue(RHICmdList, GetPixelShader(), bOnlyOnRenderedOpaque, View.bFogOnlyOnRenderedOpaque ? 1.0f : 0.0f);
+		SetShaderValue(RHICmdList, GetPixelShader(), bUseLinearDepthTexture, bUseLinearDepthTextureEnabled ? Params.LinearDepthReadScale : 0.0f);
 	}
 
 	virtual bool Serialize(FArchive& Ar) override
@@ -221,7 +242,10 @@ public:
 		Ar << SceneTextureParameters;
 		Ar << OcclusionTexture;
 		Ar << OcclusionSampler;
+		Ar << LinearDepthTexture;
+		Ar << LinearDepthSampler;
 		Ar << bOnlyOnRenderedOpaque;
+		Ar << bUseLinearDepthTexture;
 		return bShaderHasOutdatedParameters;
 	}
 
@@ -229,7 +253,10 @@ private:
 	FSceneTextureShaderParameters SceneTextureParameters;
 	FShaderResourceParameter OcclusionTexture;
 	FShaderResourceParameter OcclusionSampler;
+	FShaderResourceParameter LinearDepthTexture;
+	FShaderResourceParameter LinearDepthSampler;
 	FShaderParameter bOnlyOnRenderedOpaque;
+	FShaderParameter bUseLinearDepthTexture;
 };
 
 IMPLEMENT_SHADER_TYPE(template<>,TExponentialHeightFogPS<EHeightFogFeature::HeightFog>,TEXT("/Engine/Private/HeightFogPixelShader.usf"), TEXT("ExponentialPixelMain"),SF_Pixel)
@@ -347,7 +374,7 @@ void FSceneRenderer::InitFogConstants()
 }
 
 /** Sets the bound shader state for either the per-pixel or per-sample fog pass. */
-void SetFogShaders(FRHICommandList& RHICmdList, FGraphicsPipelineStateInitializer& GraphicsPSOInit, FScene* Scene, const FViewInfo& View, bool bShouldRenderVolumetricFog, const FLightShaftsOutput& LightShaftsOutput)
+void SetFogShaders(FRHICommandList& RHICmdList, FGraphicsPipelineStateInitializer& GraphicsPSOInit, FScene* Scene, const FViewInfo& View, bool bShouldRenderVolumetricFog, const FHeightFogRenderingParameters& Params)
 {
 	if (Scene->ExponentialFogs.Num() > 0)
 	{
@@ -364,7 +391,7 @@ void SetFogShaders(FRHICommandList& RHICmdList, FGraphicsPipelineStateInitialize
 				GraphicsPSOInit.BoundShaderState.PixelShaderRHI = GETSAFERHISHADER_PIXEL(*ExponentialHeightFogPixelShader);
 				SetGraphicsPipelineState(RHICmdList, GraphicsPSOInit);
 				VertexShader->SetParameters(RHICmdList, View);
-				ExponentialHeightFogPixelShader->SetParameters(RHICmdList, View, LightShaftsOutput);
+				ExponentialHeightFogPixelShader->SetParameters(RHICmdList, View, Params);
 			}
 			else if (View.bUseDirectionalInscattering)
 			{
@@ -373,7 +400,7 @@ void SetFogShaders(FRHICommandList& RHICmdList, FGraphicsPipelineStateInitialize
 				GraphicsPSOInit.BoundShaderState.PixelShaderRHI = GETSAFERHISHADER_PIXEL(*ExponentialHeightFogPixelShader);
 				SetGraphicsPipelineState(RHICmdList, GraphicsPSOInit);
 				VertexShader->SetParameters(RHICmdList, View);
-				ExponentialHeightFogPixelShader->SetParameters(RHICmdList, View, LightShaftsOutput);
+				ExponentialHeightFogPixelShader->SetParameters(RHICmdList, View, Params);
 			}
 			else
 			{
@@ -382,7 +409,7 @@ void SetFogShaders(FRHICommandList& RHICmdList, FGraphicsPipelineStateInitialize
 				GraphicsPSOInit.BoundShaderState.PixelShaderRHI = GETSAFERHISHADER_PIXEL(*ExponentialHeightFogPixelShader);
 				SetGraphicsPipelineState(RHICmdList, GraphicsPSOInit);
 				VertexShader->SetParameters(RHICmdList, View);
-				ExponentialHeightFogPixelShader->SetParameters(RHICmdList, View, LightShaftsOutput);
+				ExponentialHeightFogPixelShader->SetParameters(RHICmdList, View, Params);
 			}
 		}
 		else
@@ -394,7 +421,7 @@ void SetFogShaders(FRHICommandList& RHICmdList, FGraphicsPipelineStateInitialize
 				GraphicsPSOInit.BoundShaderState.PixelShaderRHI = GETSAFERHISHADER_PIXEL(*ExponentialHeightFogPixelShader);
 				SetGraphicsPipelineState(RHICmdList, GraphicsPSOInit);
 				VertexShader->SetParameters(RHICmdList, View);
-				ExponentialHeightFogPixelShader->SetParameters(RHICmdList, View, LightShaftsOutput);
+				ExponentialHeightFogPixelShader->SetParameters(RHICmdList, View, Params);
 			}
 			else if (View.bUseDirectionalInscattering)
 			{
@@ -403,7 +430,7 @@ void SetFogShaders(FRHICommandList& RHICmdList, FGraphicsPipelineStateInitialize
 				GraphicsPSOInit.BoundShaderState.PixelShaderRHI = GETSAFERHISHADER_PIXEL(*ExponentialHeightFogPixelShader);
 				SetGraphicsPipelineState(RHICmdList, GraphicsPSOInit);
 				VertexShader->SetParameters(RHICmdList, View);
-				ExponentialHeightFogPixelShader->SetParameters(RHICmdList, View, LightShaftsOutput);
+				ExponentialHeightFogPixelShader->SetParameters(RHICmdList, View, Params);
 			}
 			else
 			{
@@ -412,22 +439,22 @@ void SetFogShaders(FRHICommandList& RHICmdList, FGraphicsPipelineStateInitialize
 				GraphicsPSOInit.BoundShaderState.PixelShaderRHI = GETSAFERHISHADER_PIXEL(*ExponentialHeightFogPixelShader);
 				SetGraphicsPipelineState(RHICmdList, GraphicsPSOInit);
 				VertexShader->SetParameters(RHICmdList, View);
-				ExponentialHeightFogPixelShader->SetParameters(RHICmdList, View, LightShaftsOutput);
+				ExponentialHeightFogPixelShader->SetParameters(RHICmdList, View, Params);
 			}
 		}
 	}
 }
 
-void FDeferredShadingSceneRenderer::RenderViewFog(FRHICommandList& RHICmdList, const FViewInfo& View, const FLightShaftsOutput& LightShaftsOutput)
+void FDeferredShadingSceneRenderer::RenderViewFog(FRHICommandList& RHICmdList, const FViewInfo& View, const FHeightFogRenderingParameters& Params)
 {
 	FGraphicsPipelineStateInitializer GraphicsPSOInit;
 	RHICmdList.ApplyCachedRenderTargets(GraphicsPSOInit);
 
-	SCOPED_DRAW_EVENTF(RHICmdList, Fog, TEXT("ExponentialHeightFog %dx%d"), View.ViewRect.Width(), View.ViewRect.Height());
+	SCOPED_DRAW_EVENTF(RHICmdList, Fog, TEXT("ExponentialHeightFog %dx%d"), Params.ViewRect.Width(), Params.ViewRect.Height());
 	SCOPED_GPU_STAT(RHICmdList, Fog);
 
 	// Set the device viewport for the view.
-	RHICmdList.SetViewport(View.ViewRect.Min.X, View.ViewRect.Min.Y, 0.0f, View.ViewRect.Max.X, View.ViewRect.Max.Y, 1.0f);
+	RHICmdList.SetViewport(Params.ViewRect.Min.X, Params.ViewRect.Min.Y, 0.0f, Params.ViewRect.Max.X, Params.ViewRect.Max.Y, 1.0f);
 			
 	GraphicsPSOInit.RasterizerState = TStaticRasterizerState<FM_Solid, CM_None>::GetRHI();
 			
@@ -437,7 +464,7 @@ void FDeferredShadingSceneRenderer::RenderViewFog(FRHICommandList& RHICmdList, c
 	GraphicsPSOInit.DepthStencilState = TStaticDepthStencilState<false, CF_Always>::GetRHI();
 	GraphicsPSOInit.PrimitiveType = PT_TriangleList;
 
-	SetFogShaders(RHICmdList, GraphicsPSOInit, Scene, View, ShouldRenderVolumetricFog(), LightShaftsOutput);
+	SetFogShaders(RHICmdList, GraphicsPSOInit, Scene, View, ShouldRenderVolumetricFog(), Params);
 
 	// Draw a quad covering the view.
 	RHICmdList.SetStreamSource(0, GScreenSpaceVertexBuffer.VertexBufferRHI, 0);
@@ -456,13 +483,17 @@ bool FDeferredShadingSceneRenderer::RenderFog(FRHICommandListImmediate& RHICmdLi
 
 		SceneContext.BeginRenderingSceneColor(RHICmdList, ESimpleRenderTargetMode::EExistingColorAndDepth, FExclusiveDepthStencil::DepthRead_StencilWrite, true);
 
+		FHeightFogRenderingParameters Parameters = { LightShaftsOutput, nullptr, FIntRect(), 1.0f };
+
 		for(int32 ViewIndex = 0;ViewIndex < Views.Num();ViewIndex++)
 		{
 			const FViewInfo& View = Views[ViewIndex];
 			if (View.IsPerspectiveProjection())
 			{
 				SCOPED_GPU_MASK(RHICmdList, View.GPUMask);
-				RenderViewFog(RHICmdList, View, LightShaftsOutput);
+
+				Parameters.ViewRect = View.ViewRect;
+				RenderViewFog(RHICmdList, View, Parameters);
 			}
 		}
 
@@ -474,6 +505,42 @@ bool FDeferredShadingSceneRenderer::RenderFog(FRHICommandListImmediate& RHICmdLi
 	return false;
 }
 
+void FDeferredShadingSceneRenderer::RenderUnderWaterFog(FRHICommandListImmediate& RHICmdList, FSingleLayerWaterPassData& PassData)
+{
+	check(RHICmdList.IsOutsideRenderPass());
+
+	if (Scene->ExponentialFogs.Num() > 0
+		// Fog must be done in the base pass for MSAA to work
+		&& !IsForwardShadingEnabled(ShaderPlatform))
+	{
+
+		RHICmdList.TransitionResource(EResourceTransitionAccess::EWritable, PassData.SceneColorWithoutSingleLayerWater->GetRenderTargetItem().TargetableTexture.GetReference());
+		FRHIRenderPassInfo RPInfo(PassData.SceneColorWithoutSingleLayerWater->GetRenderTargetItem().TargetableTexture, MakeRenderTargetActions(ERenderTargetLoadAction::ELoad, ERenderTargetStoreAction::EStore));
+		RHICmdList.BeginRenderPass(RPInfo, TEXT("BeginRenderingSceneColor"));
+
+		FLightShaftsOutput LightShaftsOutput;
+		LightShaftsOutput.LightShaftOcclusion = nullptr;
+		FTextureRHIRef LinearDepthTextureRHI = PassData.SceneDepthWithoutSingleLayerWater->GetRenderTargetItem().ShaderResourceTexture;
+		const float SINGLE_LAYER_WATER_DEPTH_SCALE = 100.0f; // This must match SINGLE_LAYER_WATER_DEPTH_SCALE from SingleLayerWaterCommon.ush and SingleLayerWaterComposite.usf. TODO deduplicate
+		FHeightFogRenderingParameters Parameters = { LightShaftsOutput, LinearDepthTextureRHI, FIntRect(), SINGLE_LAYER_WATER_DEPTH_SCALE };
+
+		for (int32 ViewIndex = 0; ViewIndex < Views.Num(); ViewIndex++)
+		{
+			const FViewInfo& View = Views[ViewIndex];
+			if (View.IsPerspectiveProjection())
+			{
+				SCOPED_GPU_MASK(RHICmdList, View.GPUMask);
+
+				// Specify the low resolution view rect
+				Parameters.ViewRect = PassData.ViewData[ViewIndex].SceneWithoutSingleLayerWaterViewRect;
+
+				RenderViewFog(RHICmdList, View, Parameters);
+			}
+		}
+
+		RHICmdList.EndRenderPass();
+	}
+}
 
 bool ShouldRenderFog(const FSceneViewFamily& Family)
 {
diff --git a/Engine/Source/Runtime/Renderer/Private/SingleLayerWaterRendering.cpp b/Engine/Source/Runtime/Renderer/Private/SingleLayerWaterRendering.cpp
--- a/Engine/Source/Runtime/Renderer/Private/SingleLayerWaterRendering.cpp
+++ b/Engine/Source/Runtime/Renderer/Private/SingleLayerWaterRendering.cpp
@@ -436,7 +436,6 @@ public:
 //////////////////////////////////////////////////////////////////////////
 
 
-PRAGMA_DISABLE_OPTIMIZATION
 
 void FDeferredShadingSceneRenderer::CopySingleLayerWaterTextures(FRHICommandListImmediate& RHICmdList, FSingleLayerWaterPassData& PassData)
 {
@@ -488,6 +487,7 @@ void FDeferredShadingSceneRenderer::CopySingleLayerWaterTextures(FRHICommandList
     
 		const FIntRect RefractionViewRect = FIntRect(FIntPoint::DivideAndRoundDown(View.ViewRect.Min, RefractionDownsampleFactor), FIntPoint::DivideAndRoundDown(View.ViewRect.Max, RefractionDownsampleFactor));
 
+		PassData.ViewData[ViewIndex].SceneWithoutSingleLayerWaterViewRect   = RefractionViewRect;
 		PassData.ViewData[ViewIndex].SceneWithoutSingleLayerWaterMinMaxUV.X = (RefractionViewRect.Min.X + 0.5f) / RefractionResolution.X; 
 		PassData.ViewData[ViewIndex].SceneWithoutSingleLayerWaterMinMaxUV.Y = (RefractionViewRect.Min.Y + 0.5f) / RefractionResolution.Y;
 		PassData.ViewData[ViewIndex].SceneWithoutSingleLayerWaterMinMaxUV.Z = (RefractionViewRect.Max.X - 0.5f) / RefractionResolution.X; 
@@ -511,7 +511,7 @@ void FDeferredShadingSceneRenderer::CopySingleLayerWaterTextures(FRHICommandList
 
 	GraphBuilder.Execute();
 }
-PRAGMA_ENABLE_OPTIMIZATION
+
 void FDeferredShadingSceneRenderer::BeginRenderingWaterGBuffer(FRHICommandList& RHICmdList, FExclusiveDepthStencil::Type DepthStencilAccess, bool bBindQuadOverdrawBuffers, EShaderPlatform InShaderPlatform)
 {
 	SCOPED_DRAW_EVENT(RHICmdList, BeginRenderingWaterGBuffer);
diff --git a/Engine/Source/Runtime/Renderer/Private/SingleLayerWaterRendering.h b/Engine/Source/Runtime/Renderer/Private/SingleLayerWaterRendering.h
--- a/Engine/Source/Runtime/Renderer/Private/SingleLayerWaterRendering.h
+++ b/Engine/Source/Runtime/Renderer/Private/SingleLayerWaterRendering.h
@@ -20,6 +20,7 @@ struct FSingleLayerWaterPassData
 
 	struct FSingleLayerWaterPassViewData
 	{
+		FIntRect SceneWithoutSingleLayerWaterViewRect;
 		FVector4 SceneWithoutSingleLayerWaterMinMaxUV;
 	};
 
