Align DF tile resources grid size to 64 during allocation to prevent memory spikes due to reallocation over the frames

#rnx

diff --git a/Engine/Source/Runtime/Renderer/Private/DistanceFieldLightingShared.h b/Engine/Source/Runtime/Renderer/Private/DistanceFieldLightingShared.h
--- a/Engine/Source/Runtime/Renderer/Private/DistanceFieldLightingShared.h
+++ b/Engine/Source/Runtime/Renderer/Private/DistanceFieldLightingShared.h
@@ -501,12 +501,14 @@ public:
 
 	void Initialize()
 	{
-		TileNumCulledObjects.Initialize(sizeof(uint32), TileDimensions.X * TileDimensions.Y, PF_R32_UINT, BUF_Static);
+		FIntPoint AlignedDimensions = FIntPoint(Align(TileDimensions.X, 64), Align(TileDimensions.Y, 64));
+	
+		TileNumCulledObjects.Initialize(sizeof(uint32), AlignedDimensions.X * AlignedDimensions.Y, PF_R32_UINT, BUF_Static);
+		TileStartOffsets.Initialize(sizeof(uint32), AlignedDimensions.X * AlignedDimensions.Y, PF_R32_UINT, BUF_Static);
 		NextStartOffset.Initialize(sizeof(uint32), 1, PF_R32_UINT, BUF_Static);
-		TileStartOffsets.Initialize(sizeof(uint32), TileDimensions.X * TileDimensions.Y, PF_R32_UINT, BUF_Static);
 
 		//@todo - handle max exceeded
-		TileArrayData.Initialize(b16BitIndices ? sizeof(uint16) : sizeof(uint32), GAverageObjectsPerShadowCullTile * TileDimensions.X * TileDimensions.Y * LightTileDataStride, b16BitIndices ? PF_R16_UINT : PF_R32_UINT, BUF_Static);
+		TileArrayData.Initialize(b16BitIndices ? sizeof(uint16) : sizeof(uint32), GAverageObjectsPerShadowCullTile * AlignedDimensions.X * AlignedDimensions.Y * LightTileDataStride, b16BitIndices ? PF_R16_UINT : PF_R32_UINT, BUF_Static);
 	}
 
 	void Release()
@@ -522,6 +524,11 @@ public:
 		return TileNumCulledObjects.NumBytes + NextStartOffset.NumBytes + TileStartOffsets.NumBytes + TileArrayData.NumBytes;
 	}
 
+	static FIntPoint GetAlignedDimensions(FIntPoint InTileDimensions)
+	{
+		return FIntPoint(Align(InTileDimensions.X, 64), Align(InTileDimensions.Y, 64));
+	}
+
 	FIntPoint TileDimensions;
 
 	FRWBuffer TileNumCulledObjects;
diff --git a/Engine/Source/Runtime/Renderer/Private/DistanceFieldShadowing.cpp b/Engine/Source/Runtime/Renderer/Private/DistanceFieldShadowing.cpp
--- a/Engine/Source/Runtime/Renderer/Private/DistanceFieldShadowing.cpp
+++ b/Engine/Source/Runtime/Renderer/Private/DistanceFieldShadowing.cpp
@@ -787,8 +787,16 @@ void CullDistanceFieldObjectsForLight(
 	if (LightSceneProxy->GetLightType() == LightType_Directional && GShadowScatterTileCulling)
 	{
 		const bool b16BitObjectIndices = Scene->DistanceFieldSceneData.CanUse16BitObjectIndices();
+		bool bLightDimensionsDirty = false;
 
-		if (!TileIntersectionResources || TileIntersectionResources->TileDimensions != LightTileDimensions || TileIntersectionResources->b16BitIndices != b16BitObjectIndices)
+		if (TileIntersectionResources)
+		{
+			FIntPoint AlignedLightDimensions = FLightTileIntersectionResources::GetAlignedDimensions(LightTileDimensions);
+			FIntPoint CurrentAlignedLightDimensions = FLightTileIntersectionResources::GetAlignedDimensions(TileIntersectionResources->TileDimensions);
+			bLightDimensionsDirty = AlignedLightDimensions.X > CurrentAlignedLightDimensions.X || AlignedLightDimensions.Y > CurrentAlignedLightDimensions.Y;
+		}
+		
+		if (!TileIntersectionResources || bLightDimensionsDirty || TileIntersectionResources->b16BitIndices != b16BitObjectIndices)
 		{
 			if (TileIntersectionResources)
 			{
