Fixed: Crash adding skeletal mesh to the level while in Android ES 3.1 rendering preview

diff --git a/Engine/Source/Runtime/Renderer/Private/SceneHitProxyRendering.cpp b/Engine/Source/Runtime/Renderer/Private/SceneHitProxyRendering.cpp
--- a/Engine/Source/Runtime/Renderer/Private/SceneHitProxyRendering.cpp
+++ b/Engine/Source/Runtime/Renderer/Private/SceneHitProxyRendering.cpp
@@ -527,6 +527,8 @@ static void DoRenderHitProxies(FRHICommandListImmediate& RHICmdList, const FScen
 
 void FMobileSceneRenderer::RenderHitProxies(FRHICommandListImmediate& RHICmdList)
 {
+	Scene->UpdateAllPrimitiveSceneInfos(RHICmdList);
+	
 	PrepareViewRectsForRendering();
 
 #if WITH_EDITOR
