Merging CL#10665407

Add option to do a single renderpass for shadows which fixes flickering on some devices + improves performance.

diff --git a/Engine/Config/Android/AndroidEngine.ini b/Engine/Config/Android/AndroidEngine.ini
--- a/Engine/Config/Android/AndroidEngine.ini
+++ b/Engine/Config/Android/AndroidEngine.ini
@@ -77,3 +77,4 @@ TaskGraph.EnablePowerSavingThreadPriorityReduction=1
 ; prefer D24_S8 for depth/stencil
 r.Vulkan.Depth24Bit=1
 ;GalaxyS6Edge=577
+r.Shadow.ForceSerialSingleRenderPass=1
diff --git a/Engine/Source/Runtime/Renderer/Private/ShadowDepthRendering.cpp b/Engine/Source/Runtime/Renderer/Private/ShadowDepthRendering.cpp
--- a/Engine/Source/Runtime/Renderer/Private/ShadowDepthRendering.cpp
+++ b/Engine/Source/Runtime/Renderer/Private/ShadowDepthRendering.cpp
@@ -45,6 +45,13 @@ DECLARE_GPU_STAT_NAMED(ShadowDepths, TEXT("Shadow Depths"));
 IMPLEMENT_GLOBAL_SHADER_PARAMETER_STRUCT(FShadowDepthPassUniformParameters, "ShadowDepthPass");
 IMPLEMENT_GLOBAL_SHADER_PARAMETER_STRUCT(FMobileShadowDepthPassUniformParameters, "MobileShadowDepthPass");
 
+
+static TAutoConsoleVariable<int32> CVarShadowForceSerialSingleRenderPass(
+	TEXT("r.Shadow.ForceSerialSingleRenderPass"),
+	0,
+	TEXT("Force Serial shadow passes to render in 1 pass."),
+	ECVF_RenderThreadSafe);
+
 void SetupShadowDepthPassUniformBuffer(
 	const FProjectedShadowInfo* ShadowInfo,
 	FRHICommandListImmediate& RHICmdList,
@@ -1511,6 +1518,13 @@ void FSceneRenderer::RenderShadowDepthMapAtlases(FRHICommandListImmediate& RHICm
 
 		if (SerialShadowPasses.Num() > 0)
 		{
+			bool bForceSingleRenderPass = CVarShadowForceSerialSingleRenderPass.GetValueOnAnyThread() != 0;
+			if(bForceSingleRenderPass)
+			{
+				BeginShadowRenderPass(RHICmdList, true);
+			}
+
+
 			for (int32 ShadowIndex = 0; ShadowIndex < SerialShadowPasses.Num(); ShadowIndex++)
 			{
 				FProjectedShadowInfo* ProjectedShadowInfo = SerialShadowPasses[ShadowIndex];
@@ -1536,8 +1550,18 @@ void FSceneRenderer::RenderShadowDepthMapAtlases(FRHICommandListImmediate& RHICm
 
 				ProjectedShadowInfo->SetupShadowUniformBuffers(RHICmdList, Scene);
 				ProjectedShadowInfo->TransitionCachedShadowmap(RHICmdList, Scene);
-				BeginShadowRenderPass(RHICmdList, ShadowIndex == 0);
+				if (!bForceSingleRenderPass)
+				{
+					BeginShadowRenderPass(RHICmdList, ShadowIndex == 0);
+				}
 				ProjectedShadowInfo->RenderDepth(RHICmdList, this, BeginShadowRenderPass, false);
+				if(!bForceSingleRenderPass)
+				{
+					RHICmdList.EndRenderPass();
+				}
+			}
+			if(bForceSingleRenderPass)
+			{
 				RHICmdList.EndRenderPass();
 			}
 		}
