Use shadow bias for pre-shadows
Although there can be no self shadowing, it's possible to get shadow bias artifacts from objects that are close together (such as landscape components)

diff --git a/Engine/Source/Runtime/Renderer/Private/ShadowRendering.cpp b/Engine/Source/Runtime/Renderer/Private/ShadowRendering.cpp
--- a/Engine/Source/Runtime/Renderer/Private/ShadowRendering.cpp
+++ b/Engine/Source/Runtime/Renderer/Private/ShadowRendering.cpp
@@ -1243,15 +1243,9 @@ void FProjectedShadowInfo::UpdateShaderDepthBias()
 		SlopeScaleDepthBias = CVarCSMShadowSlopeScaleDepthBias.GetValueOnRenderThread();
 		SlopeScaleDepthBias *= LightSceneInfo->Proxy->GetUserShadowSlopeBias();
 	}
-	else if (bPreShadow)
-	{
-		// Preshadows don't need a depth bias since there is no self shadowing
-		DepthBias = 0;
-		SlopeScaleDepthBias = 0;
-	}
 	else
 	{
-		// per object shadows
+		// pre shadow or per object shadows
 		if(bDirectionalLight)
 		{
 			// we use CSMShadowDepthBias cvar but this is per object shadows, maybe we want to use different settings
