Back out changelist 7308359 due to issues caught by tests
This will reopen #jira UE-56273

diff --git a/Engine/Source/Runtime/Renderer/Private/ShadowRendering.cpp b/Engine/Source/Runtime/Renderer/Private/ShadowRendering.cpp
--- a/Engine/Source/Runtime/Renderer/Private/ShadowRendering.cpp
+++ b/Engine/Source/Runtime/Renderer/Private/ShadowRendering.cpp
@@ -1243,9 +1243,15 @@ void FProjectedShadowInfo::UpdateShaderDepthBias()
 		SlopeScaleDepthBias = CVarCSMShadowSlopeScaleDepthBias.GetValueOnRenderThread();
 		SlopeScaleDepthBias *= LightSceneInfo->Proxy->GetUserShadowSlopeBias();
 	}
+	else if (bPreShadow)
+	{
+		// Preshadows don't need a depth bias since there is no self shadowing
+		DepthBias = 0;
+		SlopeScaleDepthBias = 0;
+	}
 	else
 	{
-		// pre shadow or per object shadows
+		// per object shadows
 		if(bDirectionalLight)
 		{
 			// we use CSMShadowDepthBias cvar but this is per object shadows, maybe we want to use different settings
