AddVisualizeGBufferOverviewPass now adds the required graph nodes if any BufferVisualizationPipes are enabled even if not dumping to file or showing screen visualization.

#rnx

diff --git a/Engine/Source/Runtime/Renderer/Private/PostProcess/PostProcessVisualizeBuffer.cpp b/Engine/Source/Runtime/Renderer/Private/PostProcess/PostProcessVisualizeBuffer.cpp
--- a/Engine/Source/Runtime/Renderer/Private/PostProcess/PostProcessVisualizeBuffer.cpp
+++ b/Engine/Source/Runtime/Renderer/Private/PostProcess/PostProcessVisualizeBuffer.cpp
@@ -337,10 +337,11 @@ FScreenPassTexture AddVisualizeGBufferOverviewPass(
 	const FViewInfo& View,
 	const FVisualizeGBufferOverviewInputs& Inputs)
 {
+	const FFinalPostProcessSettings& PostProcessSettings = View.FinalPostProcessSettings;
+
 	check(Inputs.SceneColor.IsValid());
-	check(Inputs.bDumpToFile || Inputs.bOverview);
+	check(Inputs.bDumpToFile || Inputs.bOverview || PostProcessSettings.BufferVisualizationPipes.Num() > 0);
 
-	const FFinalPostProcessSettings& PostProcessSettings = View.FinalPostProcessSettings;
 
 	FScreenPassTexture Output;
 	const EPixelFormat OutputFormat = Inputs.bOutputInHDR ? PF_FloatRGBA : PF_Unknown;
