GTAO Updates to Temporal filter.

diff --git a/Engine/Shaders/Private/PostProcessAmbientOcclusion.usf b/Engine/Shaders/Private/PostProcessAmbientOcclusion.usf
--- a/Engine/Shaders/Private/PostProcessAmbientOcclusion.usf
+++ b/Engine/Shaders/Private/PostProcessAmbientOcclusion.usf
@@ -30,6 +30,11 @@
 // This define enables the thickness heuristic which assu
 #define GTAO_THICKNESS_HEURISTIC 1
 
+// Define to Use Cosine Weighted Integral
+#define GTAO_PROJECTED_NORMAL_WEIGHTING 0
+
+// Define to Use Cosine Weighted Integral
+#define GTAO_COSINE_WEIGHTED_INTEGRAL 1
 
 // 0: classic with weighted sample, 1: don't normalize and adjust the formula to be simpler and faster - can look better and is cheaper (Alchemy like?)
 #define OPTIMIZATION_O1 1
@@ -921,13 +926,13 @@ void MainSSAOSmoothCS(
 
 #if SHADER_QUALITY == 0
 	// very low
-	#define GTAO_NUMTAPS 4
+	#define GTAO_NUMTAPS 6
 #elif SHADER_QUALITY == 1
 	// low
-	#define GTAO_NUMTAPS 7
+	#define GTAO_NUMTAPS 8
 #elif SHADER_QUALITY == 2
 	// medium
-	#define GTAO_NUMTAPS 12
+	#define GTAO_NUMTAPS 10
 #elif SHADER_QUALITY == 3
 	// high
 	#define GTAO_NUMTAPS 20
@@ -975,19 +980,26 @@ FalloffParams ComputeFalloffParams(float ZDist)
 
 #if COMPUTE_SHADER
 RWTexture2D<float2> HorizonOutTexture;
-RWTexture2D<float> DepthOutTexture;
+RWTexture2D<float>	DepthOutTexture;
+RWTexture2D<float2> VelocityOutTexture;
 RWTexture2D<float> DepthsTexture;
 #endif
 
 Texture2D		HistoryTexture;
 SamplerState	HistoryTextureSampler;
 
-Texture2D		DepthHistoryTexture;
-SamplerState	DepthHistoryTextureSampler;
-
 Texture2D		ZReadTexture;
 SamplerState	ZReadTextureSampler;
 
+Texture2D		ZCurrTexture;
+SamplerState	ZCurrTextureSampler;
+
+Texture2D		ZPrevTexture;
+SamplerState	ZPrevTextureSampler;
+
+Texture2D		VelocityPrevTexture;
+SamplerState	VelocityPrevTextureSampler;
+
 float4 PrevScreenPositionScaleBias;
 float4 FilterParams;
 
@@ -1199,11 +1211,22 @@ float ComputeInnerIntegral(float2 Angles, float3 ScreenDir, float3 ViewDir, floa
 	Angles.x = Gamma + max(-Angles.x - Gamma, -(PI_HALF) );
 	Angles.y = Gamma + min( Angles.y - Gamma,  (PI_HALF) );
 
+#if GTAO_PROJECTED_NORMAL_WEIGHTING==0
+	LenProjNormal = 1;
+#endif
+
+
+#if GTAO_COSINE_WEIGHTED_INTEGRAL
 	float AO = ( LenProjNormal *  0.25 * 
 					    ( (Angles.x * SinGamma + CosGamma - cos((2.0 * Angles.x) - Gamma)) +
 				  	      (Angles.y * SinGamma + CosGamma - cos((2.0 * Angles.y) - Gamma)) ));
-
 	AO = saturate(AO / PI_HALF);
+#else
+	// Uniform weighted AO
+	float AO =  2 - (cos(Angles.x) + cos(Angles.y));
+	AO = saturate(AO / 2.0 * PI_HALF);
+#endif
+
 
 	// Fade out based on user defined distance
 	float Mul = ScreenSpaceAOParams[4].x;
@@ -1228,8 +1251,8 @@ float3 GetRandomVector(float2 UV )
 
 	// Rotate the random Vec
 	RandomVec.x = dot(RandomTexVec.xy, float2(TemporalCos, -TemporalSin ));
-	RandomVec.y = dot(RandomTexVec.xy, float2(TemporalSin, TemporalCos ));
-	RandomVec.z =RandomTexVec.z;
+    RandomVec.y = dot(RandomTexVec.xy, float2(TemporalSin, TemporalCos ));
+	RandomVec.z = RandomTexVec.z;
 
 	return RandomVec;
 }
@@ -1260,6 +1283,8 @@ void GTAOCombinedPSandCS(in float2 UV, out float4 OutColor, inout float OutDepth
 	float3 RandomAndOffset	= GetRandomVector(UV);
 	float2 RandomVec		= RandomAndOffset.xy;
 	float  Offset			= RandomAndOffset.z;
+	Offset += GTAOParams[0].z;
+
 	float3 ViewSpacePos		= ScreenToViewPos(UV,SceneDepth);
 
 	// Get Normal and transform into viewspace
@@ -1275,7 +1300,7 @@ void GTAOCombinedPSandCS(in float2 UV, out float4 OutColor, inout float OutDepth
 
 	// Given the depth determine how many pixels away we are guranteed to be outside of the falloff distance.
 	float InvTanHalfFov			= ScreenSpaceAOParams[3].w;
-	float ScreenPixelsToCutoff  = (GTAOParams[2].x * InvTanHalfFov * 7)  / ViewSpacePos.z ; 
+	float ScreenPixelsToCutoff  = (GTAOParams[2].x * InvTanHalfFov * 9)  / ViewSpacePos.z ; 
 	float Scale = ScreenPixelsToCutoff /  GTAO_NUMTAPS;
 
 	// Clamp the scale. Too small and it won't be jumping by more than 1 pixel. Too much and it'll thrash the cache
@@ -1305,7 +1330,7 @@ void GTAOCombinedPSandCS(in float2 UV, out float4 OutColor, inout float OutDepth
 	Sum = Sum   / (float)GTAO_NUMROTS;
 
 	OutColor.r	= Sum ;
-	OutDepth	= SceneDepth;
+	OutDepth	= DeviceZ;
 
 	return;
 }
@@ -1335,7 +1360,7 @@ void GTAOCombinedCS(
 	GTAOCombinedPSandCS(BufferUV , OutColor, OutDepth);
 
 	OutTexture[PixelPos]		= OutColor.r;
-	DepthOutTexture[PixelPos]   = 1.0f/OutDepth;
+	DepthOutTexture[PixelPos]   = OutDepth;
 }
 #endif
 
@@ -1450,7 +1475,7 @@ float2 HorizonSearchPSandCS(in float2 UV, inout float OutDepth)
 
 	// Given the depth determine how many pixels away we are guranteed to be outside of the falloff distance.
 	float InvTanHalfFov			= ScreenSpaceAOParams[3].w;
-	float ScreenPixelsToCutoff  = (GTAOParams[2].x * InvTanHalfFov * 7)  / ViewSpacePos.z ; 
+	float ScreenPixelsToCutoff  = (GTAOParams[2].x * InvTanHalfFov * 13)  / ViewSpacePos.z ; 
 	float Scale = ScreenPixelsToCutoff /  GTAO_NUMTAPS;
 
 	// Clamp the scale. Too small and it won't be jumping by more than 1 pixel. Too much and it'll thrash the cache
@@ -1501,16 +1526,16 @@ void HorizonSearchCS(
 Texture2D SceneVelocityTexture;
 SamplerState SceneVelocityTextureSampler;
 float4 BlendParams;
-float2 ReprojectPos(float2 UV, float Depth)
+float3 ReprojectPos(float2 UV, float Depth)
 {
 	// Given a UV reproject where this was in the previous frame
 	// Camera motion for pixel (in ScreenPos space).
 	float2 ThisScreen = (UV.xy - View.ScreenPositionScaleBias.wz) / View.ScreenPositionScaleBias.xy;
+
 	float4 ThisClip = float4( ThisScreen, Depth, 1 );
 	float4 PrevClip = mul( ThisClip, View.ClipToPrevClip );
 	float2 PrevScreen = PrevClip.xy / PrevClip.w;
 
-
 	float2 Velocity	=  Texture2DSampleLevel(SceneVelocityTexture, SceneVelocityTextureSampler, UV,0).rg;
 	if( Velocity.x > 0.0 )
 	{
@@ -1518,9 +1543,10 @@ float2 ReprojectPos(float2 UV, float Depth)
 	}
 	
 	float2 PrevUV = PrevScreen.xy * PrevScreenPositionScaleBias.xy + PrevScreenPositionScaleBias.zw;
-	return PrevUV;
+	return float3(PrevUV, PrevClip.z/ PrevClip.w);
 }
 
+
 float ReadHistory(float2 UV, float CurrDepth, float CurrAO, inout float WeightSum)
 {
 	float	BilinearWeights[4];
@@ -1548,24 +1574,60 @@ float ReadHistory(float2 UV, float CurrDepth, float CurrAO, inout float WeightSu
 	HistoryAO[2] = Texture2DSample(HistoryTexture, HistoryTextureSampler, UV + float2(    0, dUV.y)).r;
 	HistoryAO[3] = Texture2DSample(HistoryTexture, HistoryTextureSampler, UV + float2(dUV.x, dUV.y)).r;
 
-	PrevDepth[0] = Texture2DSample(DepthHistoryTexture, DepthHistoryTextureSampler, UV + float2(    0,     0)).r;
-	PrevDepth[1] = Texture2DSample(DepthHistoryTexture, DepthHistoryTextureSampler, UV + float2(dUV.x,     0)).r;
-	PrevDepth[2] = Texture2DSample(DepthHistoryTexture, DepthHistoryTextureSampler, UV + float2(    0, dUV.y)).r;
-	PrevDepth[3] = Texture2DSample(DepthHistoryTexture, DepthHistoryTextureSampler, UV + float2(dUV.x, dUV.y)).r;
+	PrevDepth[0] = Texture2DSample(ZPrevTexture, ZPrevTextureSampler, UV + float2(    0,     0)).r;
+	PrevDepth[1] = Texture2DSample(ZPrevTexture, ZPrevTextureSampler, UV + float2(dUV.x,     0)).r;
+	PrevDepth[2] = Texture2DSample(ZPrevTexture, ZPrevTextureSampler, UV + float2(    0, dUV.y)).r;
+	PrevDepth[3] = Texture2DSample(ZPrevTexture, ZPrevTextureSampler, UV + float2(dUV.x, dUV.y)).r;
 
-	const float DepthScale =0.1;
+	const float DepthScale = 0.01;
 	float VisHistory = 0;
 	for(int i=0; i<4; i++)
 	{
-		float BilateralWeight = saturate( 1- (DepthScale * (CurrDepth - PrevDepth[i])) );
-		WeightSum  += BilateralWeight;
-		VisHistory += BilinearWeights[i] * lerp( CurrAO, HistoryAO[i], BilateralWeight);
+		float DepthWeight = saturate( 1- (DepthScale * abs(CurrDepth - PrevDepth[i])) );
+		WeightSum  += DepthWeight * BilinearWeights[i];
+		VisHistory += BilinearWeights[i] * lerp( CurrAO, HistoryAO[i], DepthWeight);
 	}
 
-	WeightSum/=4;
+
 	return VisHistory;
 }
 
+float ReadHistoryClamp(float2 UV, float MinAO, float MaxAO)
+{
+	float	BilinearWeights[4];
+	float4 SizeInvSize = PostprocessInput0Size.xyzw;
+
+	float2 PixUV   = (UV * SizeInvSize.xy)-0.5;
+	float2 FloorUV = floor(PixUV);
+	float2 FracUV  = (PixUV - FloorUV); 
+	UV			   = (FloorUV * SizeInvSize.zw) + (SizeInvSize.zw*0.5);
+
+	BilinearWeights[0] = (1.0 -	FracUV.x) * ( 1.0 -	FracUV.y);
+	BilinearWeights[1] = (		FracUV.x) * ( 1.0 -	FracUV.y);
+	BilinearWeights[2] = (1.0 -	FracUV.x) * (       FracUV.y);
+	BilinearWeights[3] = (		FracUV.x) * (       FracUV.y);
+
+	// Read the 4 previous depths and History
+	float HistoryAO[4];
+
+	float2 dUV = SizeInvSize.zw;
+
+	// TODO - Use GatherR when available
+	HistoryAO[0] = Texture2DSample(HistoryTexture, HistoryTextureSampler, UV + float2(	  0,     0)).r;
+	HistoryAO[1] = Texture2DSample(HistoryTexture, HistoryTextureSampler, UV + float2(dUV.x,     0)).r;
+	HistoryAO[2] = Texture2DSample(HistoryTexture, HistoryTextureSampler, UV + float2(    0, dUV.y)).r;
+	HistoryAO[3] = Texture2DSample(HistoryTexture, HistoryTextureSampler, UV + float2(dUV.x, dUV.y)).r;
+
+	float VisHistory = 0;
+	for(int i=0; i<4; i++)
+	{
+		HistoryAO[i] = clamp(HistoryAO[i], MinAO, MaxAO);
+		VisHistory += BilinearWeights[i] * HistoryAO[i];
+	}
+
+
+	return VisHistory;
+}
 
 void NeighbourhoodClamp(float2 UV, float BaseAO, inout float MinAO, inout float MaxAO)
 {
@@ -1580,23 +1642,91 @@ void NeighbourhoodClamp(float2 UV, float BaseAO, inout float MinAO, inout float
 }
 
 
-void GTAOTemporalFilterPSAndCS(float2 UV, inout float OutAO, inout float OutDepth)
+float CompareVeloc(float2 V1, float2 V2)
+{
+	float2 V12 = V1-V2;
+	return 1-saturate( abs(V12.x + V12.y) * 100);
+
+}
+
+void GTAOTemporalFilterPSAndCS(float2 UV, inout float OutAO, inout float OutDepth, inout float2 OutVelocity)
+{
+	float NewAO					= Texture2DSample(PostprocessInput0, PostprocessInput0Sampler, UV).r;
+
+	// Current depth of the rendered Scene
+	float CurrDepthDeviceZ		= Texture2DSample(ZCurrTexture, ZCurrTextureSampler, UV).r;
+	float CurrDepth				= ConvertFromDeviceZ( CurrDepthDeviceZ);
+
+	// Previous UV value
+	float3 PrevUVDepth			= ReprojectPos( UV,  CurrDepthDeviceZ);
+	float  CurrDepthReproject	= ConvertFromDeviceZ(PrevUVDepth.z);
+	float2 PrevUV				= PrevUVDepth.xy;
+
+	float2 PixVelocity			= UV - PrevUV;
+	float  VelocityMag			= saturate(length(PixVelocity)*100);
+	
+	// Compare velocities 
+	float2 DestVeloc=0;
+	{
+		float DestDeviceZ		= Texture2DSample(ZCurrTexture, ZCurrTextureSampler, PrevUVDepth.xy).r;
+		float3 Reproj			= ReprojectPos( PrevUVDepth.xy,  DestDeviceZ); 
+		DestVeloc				= PrevUVDepth - Reproj.xy;
+	}
+
+	float VelocCompare			= CompareVeloc(PixVelocity, DestVeloc);
+
+	// Get an acceptable range of values we care about from the current AO
+	float RangeVal	   = lerp(0.1, 0.00, VelocityMag);
+	float MinAO = saturate(NewAO - RangeVal);
+	float MaxAO = saturate(NewAO + RangeVal);
+
+	// Simple history value
+	float HistoryPrevUV				= ReadHistoryClamp(PrevUV, MinAO, MaxAO);
+	float HistoryThisUV				= Texture2DSample(HistoryTexture, HistoryTextureSampler, UV ).r;
+	HistoryThisUV = clamp(HistoryThisUV, MinAO, MaxAO);
+
+	float History = HistoryPrevUV;
+
+	History = lerp(HistoryThisUV, HistoryPrevUV, VelocCompare);
+	OutAO				= lerp(History, NewAO, 0.2);
+	OutDepth			= CurrDepth;
+	OutVelocity			= UV - PrevUV;
+}
+
+
+void GTAOTemporalFilterPSAndCS2(float2 UV, inout float OutAO, inout float OutDepth, inout float2 OutVelocity)
 {
 	float NewAO			= Texture2DSample(PostprocessInput0, PostprocessInput0Sampler, UV).r;
 
-	float SceneDepth	= GetDepthFromAOInput(UV);	// FIXME - Get from HZB Buffer
-	float DeviceZ		= LookupDeviceZ( UV);
+	// Current depth of the rendered Scene
+	float CurrDepthDeviceZ		= Texture2DSample(ZCurrTexture, ZCurrTextureSampler, UV).r;
+	float CurrDepth				= ConvertFromDeviceZ( CurrDepthDeviceZ);
+	
+	// 
+	float3 PrevUVDepth			= ReprojectPos( UV,  CurrDepthDeviceZ);
+	float  CurrDepthReproject	= ConvertFromDeviceZ(PrevUVDepth.z);
+	float2 PrevUV				= PrevUVDepth.xy;
+
 
-	float2 PrevUV		= ReprojectPos( UV,  DeviceZ);
-	float WeightSum=0;
-	float CurrDepth		= SceneDepth;
+	float2 ThisVeloc = UV-PrevUVDepth.xy;
+	float2 DestVeloc=0;
+	{
+		float DestDeviceZ		= Texture2DSample(ZCurrTexture, ZCurrTextureSampler, PrevUVDepth.xy).r;
+		float3 Blah				= ReprojectPos( PrevUVDepth.xy,  DestDeviceZ); 
+		DestVeloc				= PrevUVDepth-Blah.xy;
+	}
 
-	float HistoryAO		= ReadHistory(PrevUV, SceneDepth, NewAO, WeightSum );
+	// Compare velocities 
+	float VelocCompare			= CompareVeloc(ThisVeloc, DestVeloc);
 
-	float PrevDepth		= Texture2DSample(DepthHistoryTexture, DepthHistoryTextureSampler, PrevUV).r;
+	float WeightSum=0;
+	float ReprojectDepth	= ConvertFromDeviceZ( PrevUVDepth.z);
+	float HistoryAO			= ReadHistory(PrevUV, CurrDepthReproject, NewAO, WeightSum );
 
-	float Blend			= 0.15;//lerp(1,0.15, WeightSum);
+	HistoryAO			= lerp(NewAO, HistoryAO, VelocCompare);
 
+	float PrevDepth		= Texture2DSample(ZPrevTexture, ZPrevTextureSampler, PrevUV).r;
+	float Blend			= lerp(1,0.15, WeightSum);
 
 	if((PrevUV.x <= 0.0) || (PrevUV.y <= 0.0) || (PrevUV.x >= 1.0) || (PrevUV.y >= 1.0))
 	{
@@ -1612,28 +1742,35 @@ void GTAOTemporalFilterPSAndCS(float2 UV, inout float OutAO, inout float OutDept
 	float MidAO  = (MinAO + HalfAO);
 
 	// Ensure the HalfAO is at least some minimum ammount. This helps with the history clamp.
-	HalfAO = max(HalfAO,0.05);
+	HalfAO = max(HalfAO,0.05*VelocCompare);
 	MinAO = saturate(MidAO - HalfAO);
 	MaxAO = saturate(MidAO + HalfAO);
 
 	OutAO = (NewAO*Blend) + (HistoryAO* (1-Blend));
+
+
 	OutAO = clamp(OutAO, MinAO, MaxAO);
 
+	// If a camera cut then Use current data
 	if(BlendParams.x > 0.5 )
 	{
 		OutAO = NewAO;
 	}
 
-	OutDepth = CurrDepth;
+	OutDepth	= CurrDepthReproject;
+	OutVelocity = UV - PrevUV;
 }
 
-void GTAOTemporalFilterPS(in noperspective float4 UVAndScreenPos : TEXCOORD0, float4 SvPosition : SV_POSITION, out float4 OutColor : SV_Target0, out float4 OutColor1 : SV_Target1)
+void GTAOTemporalFilterPS(in noperspective float4 UVAndScreenPos : TEXCOORD0, float4 SvPosition : SV_POSITION, out float4 OutColor : SV_Target0, out float4 OutColor1 : SV_Target1, out float4 OutColor2 : SV_Target2)
 {
-	float OutAO, OutDepth;
+	float OutAO = 0;
+	float OutDepth = 0;
+	float2 OutVelocity = 0;
 	
-	GTAOTemporalFilterPSAndCS(UVAndScreenPos.xy, OutAO, OutDepth);
+	GTAOTemporalFilterPSAndCS(UVAndScreenPos.xy, OutAO, OutDepth, OutVelocity);
 	OutColor	= OutAO;
-	OutColor1 = OutDepth;
+	OutColor1  = OutDepth;
+	OutColor2  = float4((OutVelocity*0.5) + 0.5,0,0);
 }
 
 #if COMPUTE_SHADER
@@ -1645,14 +1782,17 @@ void GTAOTemporalFilterCS(
 {
 	float OutColor		= 0;
 	float OutDepth		= 0;
+	float2 OutVelocity  = 0;
+
 	int2   PixelPos		= DispatchThreadId + ScreenSpaceAOParams[5].zw; 
 	float2 PixelCenter	= (float2)PixelPos + float2(0.5, 0.5);
 	float2 BufferUV		= PixelCenter.xy * PostprocessInput0Size.zw;
 
-	GTAOTemporalFilterPSAndCS(BufferUV, OutColor, OutDepth);
+	GTAOTemporalFilterPSAndCS(BufferUV, OutColor, OutDepth, OutVelocity);
 
-	OutTexture[PixelPos]		= OutColor;
-	DepthOutTexture[PixelPos]   = OutDepth;
+	OutTexture[PixelPos]			= OutColor;
+	DepthOutTexture[PixelPos]		= OutDepth;
+	VelocityOutTexture[PixelPos]	= (OutVelocity*0.5) + 0.5;
 }
 
 #endif
@@ -1761,9 +1901,14 @@ void GTAOUpsampleCS(
 
 float4 GTAOUpsamplePSAndCS(float2 UV)
 {
-//	UV += PostprocessInput0Size.zw * -0.5;
+	float2 Offset = PostprocessInput0Size.zw * 0.25;
+
+	float AO0 = saturate( Texture2DSample(PostprocessInput0, PostprocessInput0Sampler, UV + float2(-Offset.x, -Offset.y)) .r);
+	float AO1 = saturate( Texture2DSample(PostprocessInput0, PostprocessInput0Sampler, UV + float2(Offset.x,  -Offset.y)) .r);
+	float AO2 = saturate( Texture2DSample(PostprocessInput0, PostprocessInput0Sampler, UV + float2(-Offset.x, Offset.y)) .r);
+	float AO3 = saturate( Texture2DSample(PostprocessInput0, PostprocessInput0Sampler, UV + float2(Offset.x, Offset.y)) .r);
 
-	float AO = saturate( Texture2DSample(PostprocessInput0, PostprocessInput0Sampler, UV).r);
+	float AO = min(min(AO0, AO1), min(AO2, AO3))  ;
 	return float4(AO,0,0,0);
 }
 
@@ -1890,7 +2035,7 @@ void GTAOSpatialFilterCS(
 				SampleZDiff -= XDiff * DiffZ_X;
 				SampleZDiff -= YDiff * DiffZ_Y;
 			
-				const float SpatialFilterWeight = 100000;
+				const float SpatialFilterWeight = 1000;
 				Weight = 1.0f - saturate(SampleZDiff*SpatialFilterWeight );
 			}
 
diff --git a/Engine/Source/Runtime/Renderer/Private/CompositionLighting/PostProcessAmbientOcclusion.cpp b/Engine/Source/Runtime/Renderer/Private/CompositionLighting/PostProcessAmbientOcclusion.cpp
--- a/Engine/Source/Runtime/Renderer/Private/CompositionLighting/PostProcessAmbientOcclusion.cpp
+++ b/Engine/Source/Runtime/Renderer/Private/CompositionLighting/PostProcessAmbientOcclusion.cpp
@@ -111,7 +111,7 @@ static TAutoConsoleVariable<int32> CVarGTAOUseNormals(
 
 static TAutoConsoleVariable<float> CVarGTAOThicknessBlend(
 	TEXT("r.GTAO.ThicknessBlend"),
-	0.1f,
+	0.03f,
 	TEXT("A heuristic to bias occlusion for thin or thick objects. \n ")
 	TEXT("0  : Off \n ")
 	TEXT(">0 : On - Bigger values lead to reduced occlusion \n ")
@@ -1247,7 +1247,6 @@ FPooledRenderTargetDesc FRCPassPostProcessAmbientOcclusion::ComputeOutputDesc(EP
 
 
 
-
 /** Shader parameters needed for screen space AmbientOcclusion passes. */
 class FGTAOParameters
 {
@@ -1270,7 +1269,7 @@ public:
 		
 		const FSceneViewState* ViewState = static_cast<const FSceneViewState*>(View.State);
 
-		if (ViewState && CVarTemporalFilter->GetValueOnRenderThread() > 0)
+		if (ViewState)// && CVarTemporalFilter->GetValueOnRenderThread() > 0)
 		{
 			TemporalFrame = ViewState->GetCurrentUnclampedTemporalAASampleIndex();
 			Frame = ViewState->GetFrameIndex();
@@ -1285,7 +1284,7 @@ public:
 		float TemporalAngle = Rots[TemporalFrame % 6] * (PI / 360.0f);
 
 		// Angles of rotation that are set per frame
-		GTAOParam[0] = FVector4(cos(TemporalAngle), sin(TemporalAngle), Offsets[(TemporalFrame / 6) % 4], Offsets[TemporalFrame % 4]);
+		GTAOParam[0] = FVector4(cos(TemporalAngle), sin(TemporalAngle), Offsets[(TemporalFrame / 6) % 4] * 0.25, Offsets[TemporalFrame % 4]);
 
 		// Frame X = number , Y = Thickness param, 
 		float ThicknessBlend = CVarGTAOThicknessBlend.GetValueOnRenderThread();
@@ -1399,7 +1398,8 @@ public:
 	}
 
 	template <typename TRHICmdList>
-	void SetParametersCS(TRHICmdList& RHICmdList, const FRenderingCompositePassContext& Context, FIntPoint DestSize, FIntPoint InputTextureSize, uint32 DownScaleFactor, FRHIUnorderedAccessView *OutUAV, FRHIUnorderedAccessView *OutDepthsUAV)
+	void SetParametersCS(TRHICmdList& RHICmdList, const FRenderingCompositePassContext& Context, FIntPoint DestSize, FIntPoint InputTextureSize, uint32 DownScaleFactor,
+		FRHIUnorderedAccessView *OutUAV, FRHIUnorderedAccessView *OutDepthsUAV)
 	{
 		const FFinalPostProcessSettings& Settings = Context.View.FinalPostProcessSettings;
 		FRHIComputeShader* ShaderRHI = GetComputeShader();
@@ -2301,8 +2301,14 @@ public:
 	FShaderResourceParameter		HistoryTexture;
 	FShaderResourceParameter		HistoryTextureSampler;
 
-	FShaderResourceParameter		DepthHistoryTexture;
-	FShaderResourceParameter		DepthHistoryTextureSampler;
+	FShaderResourceParameter		ZCurrTexture;
+	FShaderResourceParameter		ZCurrTextureSampler;
+
+	FShaderResourceParameter		ZPrevTexture;
+	FShaderResourceParameter		ZPrevTextureSampler;
+
+	FShaderResourceParameter		VelocityPrevTexture;
+	FShaderResourceParameter		VelocityPrevTextureSampler;
 
 	FShaderResourceParameter		SceneVelocityTexture;
 	FShaderResourceParameter		SceneVelocityTextureSampler;
@@ -2313,6 +2319,7 @@ public:
 
 	FShaderResourceParameter		OutTexture;
 	FShaderResourceParameter 		DepthOutTexture;
+	FShaderResourceParameter 		VelocityOutTexture;
 
 
 	/** Initialization constructor. */
@@ -2322,34 +2329,41 @@ public:
 		PostprocessParameter.Bind(Initializer.ParameterMap);
 		SceneTextureParameters.Bind(Initializer);
 
-		HistoryTexture.Bind(Initializer.ParameterMap,				TEXT("HistoryTexture"));
-		HistoryTextureSampler.Bind(Initializer.ParameterMap,		TEXT("HistoryTextureSampler"));
+		HistoryTexture.Bind(Initializer.ParameterMap, TEXT("HistoryTexture"));
+		HistoryTextureSampler.Bind(Initializer.ParameterMap, TEXT("HistoryTextureSampler"));
+
+		ZCurrTexture.Bind(Initializer.ParameterMap, TEXT("ZCurrTexture"));
+		ZCurrTextureSampler.Bind(Initializer.ParameterMap, TEXT("ZCurrTextureSampler"));
 
-		DepthHistoryTexture.Bind(Initializer.ParameterMap,			TEXT("DepthHistoryTexture"));
-		DepthHistoryTextureSampler.Bind(Initializer.ParameterMap,	TEXT("DepthHistoryTextureSampler"));
+		ZPrevTexture.Bind(Initializer.ParameterMap, TEXT("ZPrevTexture"));
+		ZPrevTextureSampler.Bind(Initializer.ParameterMap, TEXT("ZPrevTextureSampler"));
 
-		SceneVelocityTexture.Bind(Initializer.ParameterMap,			TEXT("SceneVelocityTexture"));
-		SceneVelocityTextureSampler.Bind(Initializer.ParameterMap,	TEXT("SceneVelocityTextureSampler"));
+		VelocityPrevTexture.Bind(Initializer.ParameterMap, TEXT("VelocityPrevTexture"));
+		VelocityPrevTextureSampler.Bind(Initializer.ParameterMap, TEXT("VelocityPrevTextureSampler"));
 
-		PrevScreenPositionScaleBias.Bind(Initializer.ParameterMap,  TEXT("PrevScreenPositionScaleBias"));
-		BlendParams.Bind(Initializer.ParameterMap,					TEXT("BlendParams"));
+		SceneVelocityTexture.Bind(Initializer.ParameterMap, TEXT("SceneVelocityTexture"));
+		SceneVelocityTextureSampler.Bind(Initializer.ParameterMap, TEXT("SceneVelocityTextureSampler"));
+
+		PrevScreenPositionScaleBias.Bind(Initializer.ParameterMap, TEXT("PrevScreenPositionScaleBias"));
+		BlendParams.Bind(Initializer.ParameterMap, TEXT("BlendParams"));
 
 		ScreenSpaceAOParams.Bind(Initializer.ParameterMap);
 
 		if (bComputeShader)
 		{
-			OutTexture.Bind(Initializer.ParameterMap,				TEXT("OutTexture"));
-			DepthOutTexture.Bind(Initializer.ParameterMap,			TEXT("DepthOutTexture"));
+			OutTexture.Bind(Initializer.ParameterMap, TEXT("OutTexture"));
+			DepthOutTexture.Bind(Initializer.ParameterMap, TEXT("DepthOutTexture"));
+			VelocityOutTexture.Bind(Initializer.ParameterMap, TEXT("VelocityOutTexture"));
 		}
 	}
 
 	template <typename TRHICmdList>
 	void SetParametersCS(TRHICmdList& RHICmdList, FRenderingCompositePassContext& Context, FIntPoint DestSize, FIntPoint InputTextureSize, bool bCameraCut, const FGTAOTAAHistory& InputHistory, TRefCountPtr<IPooledRenderTarget>& VelocityRT,
-		FRHIUnorderedAccessView *OutUAV0, FRHIUnorderedAccessView *OutUAV1)
+		FRHIUnorderedAccessView *OutUAV_AO, FRHIUnorderedAccessView *OutUAV_Depth, FRHIUnorderedAccessView *OutUAV_Velocity)
 	{
 		const FFinalPostProcessSettings& Settings = Context.View.FinalPostProcessSettings;
 		FRHIComputeShader* ShaderRHI = GetComputeShader();
-
+		FSceneRenderTargets& SceneContext = FSceneRenderTargets::Get(Context.RHICmdList);
 		PostprocessParameter.SetCS(ShaderRHI, Context, RHICmdList, TStaticSamplerState<SF_Bilinear, AM_Clamp, AM_Clamp, AM_Clamp>::GetRHI());
 
 		FGlobalShader::SetParameters<FViewUniformShaderParameters>(RHICmdList, ShaderRHI, Context.View.ViewUniformBuffer);
@@ -2376,10 +2390,13 @@ public:
 				TStaticSamplerState<SF_Point, AM_Border, AM_Border, AM_Border, 0, 0, 0xffffffff >::GetRHI(),
 				InputHistory.RT[0]->GetRenderTargetItem().TargetableTexture);
 
-
-			SetTextureParameter(RHICmdList, ShaderRHI, DepthHistoryTexture, DepthHistoryTextureSampler,
+			SetTextureParameter(RHICmdList, ShaderRHI, ZPrevTexture, ZPrevTextureSampler,
 				TStaticSamplerState<SF_Point>::GetRHI(),
 				InputHistory.Depth[0]->GetRenderTargetItem().TargetableTexture);
+
+			SetTextureParameter(RHICmdList, ShaderRHI, VelocityPrevTexture, VelocityPrevTextureSampler,
+				TStaticSamplerState<SF_Point>::GetRHI(),
+				InputHistory.Velocity[0]->GetRenderTargetItem().TargetableTexture);
 		}
 		else
 		{
@@ -2393,9 +2410,13 @@ public:
 		SetTextureParameter(RHICmdList, ShaderRHI, SceneVelocityTexture, SceneVelocityTextureSampler,
 			TStaticSamplerState<SF_Point>::GetRHI(), VelocityRT->GetRenderTargetItem().ShaderResourceTexture);
 
+		SetTextureParameter(RHICmdList, ShaderRHI, ZCurrTexture, ZCurrTextureSampler, TStaticSamplerState<SF_Point, AM_Wrap, AM_Wrap, AM_Wrap>::GetRHI(),
+			SceneContext.ScreenSpaceGTAODepths->GetRenderTargetItem().ShaderResourceTexture);
+
 		// Bind the output UAVs
-		RHICmdList.SetUAVParameter(ShaderRHI, OutTexture.GetBaseIndex(),	  OutUAV0);
-		RHICmdList.SetUAVParameter(ShaderRHI, DepthOutTexture.GetBaseIndex(), OutUAV1);
+		RHICmdList.SetUAVParameter(ShaderRHI, OutTexture.GetBaseIndex(), OutUAV_AO);
+		RHICmdList.SetUAVParameter(ShaderRHI, DepthOutTexture.GetBaseIndex(), OutUAV_Depth);
+		RHICmdList.SetUAVParameter(ShaderRHI, VelocityOutTexture.GetBaseIndex(), OutUAV_Velocity);
 	}
 
 
@@ -2403,6 +2424,7 @@ public:
 	{
 		const FFinalPostProcessSettings& Settings = Context.View.FinalPostProcessSettings;
 		FRHIPixelShader* ShaderRHI = GetPixelShader();
+		FSceneRenderTargets& SceneContext = FSceneRenderTargets::Get(Context.RHICmdList);
 
 		PostprocessParameter.SetPS(Context.RHICmdList, ShaderRHI, Context, TStaticSamplerState<SF_Point, AM_Clamp, AM_Clamp, AM_Clamp>::GetRHI());
 
@@ -2427,13 +2449,16 @@ public:
 		if (InputHistory.IsValid())
 		{
 			SetTextureParameter(Context.RHICmdList, ShaderRHI, HistoryTexture, HistoryTextureSampler,
-				TStaticSamplerState<SF_Point,AM_Border, AM_Border, AM_Border, 0,0, 0xffffffff >::GetRHI(),
+				TStaticSamplerState<SF_Point, AM_Border, AM_Border, AM_Border, 0, 0, 0xffffffff >::GetRHI(),
 				InputHistory.RT[0]->GetRenderTargetItem().TargetableTexture);
 
-
-			SetTextureParameter(Context.RHICmdList, ShaderRHI, DepthHistoryTexture, DepthHistoryTextureSampler,
+			SetTextureParameter(Context.RHICmdList, ShaderRHI, ZPrevTexture, ZPrevTextureSampler,
 				TStaticSamplerState<SF_Point>::GetRHI(),
 				InputHistory.Depth[0]->GetRenderTargetItem().TargetableTexture);
+
+			SetTextureParameter(Context.RHICmdList, ShaderRHI, VelocityPrevTexture, VelocityPrevTextureSampler,
+				TStaticSamplerState<SF_Point>::GetRHI(),
+				InputHistory.Velocity[0]->GetRenderTargetItem().TargetableTexture);
 		}
 		else
 		{
@@ -2444,8 +2469,11 @@ public:
 
 		}
 
+		SetTextureParameter(Context.RHICmdList, ShaderRHI, ZCurrTexture, ZCurrTextureSampler, TStaticSamplerState<SF_Point, AM_Wrap, AM_Wrap, AM_Wrap>::GetRHI(),
+			SceneContext.ScreenSpaceGTAODepths->GetRenderTargetItem().ShaderResourceTexture);
+
 		SetTextureParameter(Context.RHICmdList, ShaderRHI, SceneVelocityTexture, SceneVelocityTextureSampler,
-				TStaticSamplerState<SF_Point>::GetRHI(), VelocityRT->GetRenderTargetItem().ShaderResourceTexture);
+			TStaticSamplerState<SF_Point>::GetRHI(), VelocityRT->GetRenderTargetItem().ShaderResourceTexture);
 
 	}
 
@@ -2455,6 +2483,7 @@ public:
 		FRHIComputeShader* ShaderRHI = GetComputeShader();
 		RHICmdList.SetUAVParameter(ShaderRHI, OutTexture.GetBaseIndex(), nullptr);
 		RHICmdList.SetUAVParameter(ShaderRHI, DepthOutTexture.GetBaseIndex(), nullptr);
+		RHICmdList.SetUAVParameter(ShaderRHI, VelocityOutTexture.GetBaseIndex(), nullptr);
 	}
 
 
@@ -2462,8 +2491,8 @@ public:
 	virtual bool Serialize(FArchive& Ar) override
 	{
 		bool bShaderHasOutdatedParameters = FGlobalShader::Serialize(Ar);
-		Ar << DepthHistoryTexture << DepthHistoryTextureSampler << HistoryTexture << HistoryTextureSampler << PostprocessParameter << SceneTextureParameters << PrevScreenPositionScaleBias << OutTexture;
-		Ar << SceneVelocityTexture << SceneVelocityTextureSampler << BlendParams << DepthOutTexture;
+		Ar << ZCurrTexture << ZCurrTextureSampler << HistoryTexture << HistoryTextureSampler << PostprocessParameter << SceneTextureParameters << PrevScreenPositionScaleBias << OutTexture;
+		Ar << ZPrevTexture << ZPrevTextureSampler << VelocityPrevTexture << VelocityPrevTextureSampler << SceneVelocityTexture << SceneVelocityTextureSampler << BlendParams << DepthOutTexture << VelocityOutTexture;
 		return bShaderHasOutdatedParameters;
 	}
 
@@ -2501,22 +2530,24 @@ void FRCPassPostProcessAmbientOcclusion_GTAO_TemporalFilter::Process(FRenderingC
 	const FPooledRenderTargetDesc* InputDesc0 = GetInputDesc(ePId_Input0);
 	FSceneRenderTargets& SceneContext = FSceneRenderTargets::Get(Context.RHICmdList);
 
-	const FSceneRenderTargetItem& DestRenderTarget0 = (DownScaleFactor==1) ? SceneContext.ScreenSpaceAO->GetRenderTargetItem() : PassOutputs[0].RequestSurface(Context);
+	const FSceneRenderTargetItem& DestRenderTarget0 = (DownScaleFactor == 1) ? SceneContext.ScreenSpaceAO->GetRenderTargetItem() : PassOutputs[0].RequestSurface(Context);
 	const FSceneRenderTargetItem& DestRenderTarget1 = PassOutputs[1].RequestSurface(Context);
+	const FSceneRenderTargetItem& DestRenderTarget2 = PassOutputs[2].RequestSurface(Context);
 
-	FIntPoint InputTexSize  = InputDesc0->Extent;
+	FIntPoint InputTexSize = InputDesc0->Extent;
 	FIntPoint OutputTexSize = PassOutputs[0].RenderTargetDesc.Extent;
 	FIntRect OutputFullRect(FIntPoint(0, 0), OutputTexSize);
-	
+
 	FIntRect InputViewRect = View.ViewRect / DownScaleFactor;
-	FIntRect OutputViewRect = InputViewRect ;
+	FIntRect OutputViewRect = InputViewRect;
 
 	// Whether to use camera cut shader permutation or not.
 	const bool bCameraCut = !InputHistory.IsValid() || View.bCameraCut;
 
 	OutputHistory->SafeRelease();
-	OutputHistory->RT[0]	= PassOutputs[0].PooledRenderTarget;
+	OutputHistory->RT[0] = PassOutputs[0].PooledRenderTarget;
 	OutputHistory->Depth[0] = PassOutputs[1].PooledRenderTarget;
+	OutputHistory->Velocity[0] = PassOutputs[2].PooledRenderTarget;
 	OutputHistory->ViewportRect = OutputViewRect;
 	OutputHistory->ReferenceBufferSize = OutputTexSize;
 
@@ -2525,15 +2556,16 @@ void FRCPassPostProcessAmbientOcclusion_GTAO_TemporalFilter::Process(FRenderingC
 		// Compute  Version
 		UnbindRenderTargets(Context.RHICmdList);
 		TShaderMapRef<FPostProcessGTAOTemporalFilterPSandCS<true>> ComputeShader(Context.GetShaderMap());
-		uint32 GroupSizeX = FMath::DivideAndRoundUp(OutputViewRect.Width(),  8);
+		uint32 GroupSizeX = FMath::DivideAndRoundUp(OutputViewRect.Width(), 8);
 		uint32 GroupSizeY = FMath::DivideAndRoundUp(OutputViewRect.Height(), 8);
 
 		Context.RHICmdList.TransitionResource(EResourceTransitionAccess::ERWBarrier, EResourceTransitionPipeline::EComputeToCompute, DestRenderTarget0.UAV, nullptr);
 		Context.RHICmdList.TransitionResource(EResourceTransitionAccess::ERWBarrier, EResourceTransitionPipeline::EComputeToCompute, DestRenderTarget1.UAV, nullptr);
+		Context.RHICmdList.TransitionResource(EResourceTransitionAccess::ERWBarrier, EResourceTransitionPipeline::EComputeToCompute, DestRenderTarget2.UAV, nullptr);
 
 		Context.RHICmdList.SetComputeShader(ComputeShader->GetComputeShader());
 		ComputeShader->SetParametersCS(Context.RHICmdList, Context, OutputTexSize, InputTexSize, bCameraCut, InputHistory,
-			bCameraCut ? GSystemTextures.BlackDummy : SceneContext.SceneVelocity, DestRenderTarget0.UAV, DestRenderTarget1.UAV);
+			bCameraCut ? GSystemTextures.BlackDummy : SceneContext.SceneVelocity, DestRenderTarget0.UAV, DestRenderTarget1.UAV, DestRenderTarget2.UAV);
 
 		DispatchComputeShader(Context.RHICmdList, *ComputeShader, GroupSizeX, GroupSizeY, 1);
 		ComputeShader->UnsetParameters(Context.RHICmdList);
@@ -2549,6 +2581,7 @@ void FRCPassPostProcessAmbientOcclusion_GTAO_TemporalFilter::Process(FRenderingC
 
 		Context.RHICmdList.TransitionResource(EResourceTransitionAccess::EWritable, DestRenderTarget0.TargetableTexture);
 		Context.RHICmdList.TransitionResource(EResourceTransitionAccess::EWritable, DestRenderTarget1.TargetableTexture);
+		Context.RHICmdList.TransitionResource(EResourceTransitionAccess::EWritable, DestRenderTarget2.TargetableTexture);
 
 		FRHIRenderPassInfo RPInfo(2, RenderTargets, ERenderTargetActions::Load_Store);
 
@@ -2578,7 +2611,7 @@ void FRCPassPostProcessAmbientOcclusion_GTAO_TemporalFilter::Process(FRenderingC
 
 			VertexShader->SetParameters(Context);
 			PixelShader->SetParametersPS(Context, OutputTexSize, InputTexSize, bCameraCut, InputHistory,
-											bCameraCut ? GSystemTextures.BlackDummy : SceneContext.SceneVelocity);
+				bCameraCut ? GSystemTextures.BlackDummy : SceneContext.SceneVelocity);
 
 			DrawRectangle(
 				Context.RHICmdList,
@@ -2594,10 +2627,8 @@ void FRCPassPostProcessAmbientOcclusion_GTAO_TemporalFilter::Process(FRenderingC
 
 		Context.RHICmdList.TransitionResource(EResourceTransitionAccess::EReadable, DestRenderTarget0.TargetableTexture);
 		Context.RHICmdList.TransitionResource(EResourceTransitionAccess::EReadable, DestRenderTarget1.TargetableTexture);
+		Context.RHICmdList.TransitionResource(EResourceTransitionAccess::EReadable, DestRenderTarget2.TargetableTexture);
 	}
-
-
-
 }
 
 
@@ -2611,17 +2642,21 @@ FPooledRenderTargetDesc FRCPassPostProcessAmbientOcclusion_GTAO_TemporalFilter::
 	{
 		Ret.Format = PF_G8;
 	}
-	else
+	else if (InPassOutputId == 1)
 	{
 		Ret.Format = PF_R32_FLOAT;
 	}
+	else
+	{
+		Ret.Format = PF_G16R16;
+	}
 
 	Ret.ClearValue = FClearValueBinding::None;
 	Ret.TargetableFlags &= ~TexCreate_DepthStencilTargetable;
 	Ret.TargetableFlags |= TexCreate_UAV;
 	Ret.TargetableFlags |= TexCreate_RenderTargetable | TexCreate_ShaderResource;
 
-	Ret.DebugName		= TEXT("GTAOTemporalAccumulate");
+	Ret.DebugName = TEXT("GTAOTemporalAccumulate");
 	return Ret;
 }
 
diff --git a/Engine/Source/Runtime/Renderer/Private/CompositionLighting/PostProcessAmbientOcclusion.h b/Engine/Source/Runtime/Renderer/Private/CompositionLighting/PostProcessAmbientOcclusion.h
--- a/Engine/Source/Runtime/Renderer/Private/CompositionLighting/PostProcessAmbientOcclusion.h
+++ b/Engine/Source/Runtime/Renderer/Private/CompositionLighting/PostProcessAmbientOcclusion.h
@@ -211,7 +211,7 @@ private:
 
 
 
-class FRCPassPostProcessAmbientOcclusion_GTAO_TemporalFilter : public TRenderingCompositePassBase<1, 2>
+class FRCPassPostProcessAmbientOcclusion_GTAO_TemporalFilter : public TRenderingCompositePassBase<1, 3>
 {
 public:
 
diff --git a/Engine/Source/Runtime/Renderer/Private/SceneRendering.h b/Engine/Source/Runtime/Renderer/Private/SceneRendering.h
--- a/Engine/Source/Runtime/Renderer/Private/SceneRendering.h
+++ b/Engine/Source/Runtime/Renderer/Private/SceneRendering.h
@@ -743,6 +743,7 @@ struct FGTAOTAAHistory
 	//  scene color's RGBA are in RT[0].
 	TRefCountPtr<IPooledRenderTarget> RT[kRenderTargetCount];
 	TRefCountPtr<IPooledRenderTarget> Depth[kRenderTargetCount];
+	TRefCountPtr<IPooledRenderTarget> Velocity[kRenderTargetCount];
 
 	// Reference size of RT. Might be different than RT's actual size to handle down res.
 	FIntPoint ReferenceBufferSize;
