Added support for moving translucency pass to before water rendering depending on if the camera is above or below the water surface.

Moved some translucency code into a function to be shared by the two call sites.

This introduces a water depth state on the scene view, so it's now possible to for example call IsUnderwater() on a view.

diff --git a/Engine/Source/Runtime/Engine/Classes/Engine/World.h b/Engine/Source/Runtime/Engine/Classes/Engine/World.h
--- a/Engine/Source/Runtime/Engine/Classes/Engine/World.h
+++ b/Engine/Source/Runtime/Engine/Classes/Engine/World.h
@@ -1469,7 +1469,7 @@ public:
 #endif
 
 	/** Called when the world computes how post process volumes contribute to the scene. */
-	DECLARE_EVENT_OneParam(UWorld, FOnBeginPostProcessSettings, FVector);
+	DECLARE_EVENT_TwoParams(UWorld, FOnBeginPostProcessSettings, FVector, FSceneView*);
 	FOnBeginPostProcessSettings OnBeginPostProcessSettings;
 
 	/** Inserts a post process volume into the world in priority order */
diff --git a/Engine/Source/Runtime/Engine/Private/SceneView.cpp b/Engine/Source/Runtime/Engine/Private/SceneView.cpp
--- a/Engine/Source/Runtime/Engine/Private/SceneView.cpp
+++ b/Engine/Source/Runtime/Engine/Private/SceneView.cpp
@@ -655,6 +655,7 @@ FSceneView::FSceneView(const FSceneViewInitOptions& InitOptions)
 	, bIsMobileMultiViewEnabled(false)
 	, bIsMobileMultiViewDirectEnabled(false)
 	, bShouldBindInstancedViewUB(false)
+	, UnderwaterDepth(-1.0f)
 	, GlobalClippingPlane(FPlane(0, 0, 0, 0))
 #if WITH_EDITOR
 	, OverrideLODViewOrigin(InitOptions.OverrideLODViewOrigin)
diff --git a/Engine/Source/Runtime/Engine/Private/World.cpp b/Engine/Source/Runtime/Engine/Private/World.cpp
--- a/Engine/Source/Runtime/Engine/Private/World.cpp
+++ b/Engine/Source/Runtime/Engine/Private/World.cpp
@@ -7516,7 +7516,7 @@ static void DoPostProcessVolume(IInterface_PostProcessVolume* Volume, FVector Vi
 
 void UWorld::AddPostProcessingSettings(FVector ViewLocation, FSceneView* SceneView)
 {
-	OnBeginPostProcessSettings.Broadcast(ViewLocation);
+	OnBeginPostProcessSettings.Broadcast(ViewLocation, SceneView);
 
 	for (IInterface_PostProcessVolume* PPVolume : PostProcessVolumes)
 	{
diff --git a/Engine/Source/Runtime/Engine/Public/SceneView.h b/Engine/Source/Runtime/Engine/Public/SceneView.h
--- a/Engine/Source/Runtime/Engine/Public/SceneView.h
+++ b/Engine/Source/Runtime/Engine/Public/SceneView.h
@@ -1016,6 +1016,9 @@ public:
 	/** True if we need to bind the instanced view uniform buffer parameters. */
 	bool bShouldBindInstancedViewUB;
 
+	/** How far below the water surface this view is. -1 means the view is out of water. */
+	float UnderwaterDepth;
+
 	/** Global clipping plane being applied to the scene, or all 0's if disabled.  This is used when rendering the planar reflection pass. */
 	FPlane GlobalClippingPlane;
 
@@ -1173,6 +1176,8 @@ public:
 	/** @return true:perspective, false:orthographic */
 	inline bool IsPerspectiveProjection() const { return ViewMatrices.IsPerspectiveProjection(); }
 
+	bool IsUnderwater() const { return UnderwaterDepth > 0.0f; }
+
 	/** Returns the location used as the origin for LOD computations
 	 * @param Index, 0 or 1, which LOD origin to return
 	 * @return LOD origin
diff --git a/Engine/Source/Runtime/Renderer/Private/DeferredShadingRenderer.cpp b/Engine/Source/Runtime/Renderer/Private/DeferredShadingRenderer.cpp
--- a/Engine/Source/Runtime/Renderer/Private/DeferredShadingRenderer.cpp
+++ b/Engine/Source/Runtime/Renderer/Private/DeferredShadingRenderer.cpp
@@ -2051,10 +2051,36 @@ void FDeferredShadingSceneRenderer::Render(FRHICommandListImmediate& RHICmdList)
 
 	checkSlow(RHICmdList.IsOutsideRenderPass());
 
-	
 	const bool bShouldRenderSingleLayerWater = ShouldRenderSingleLayerWater(Views, ViewFamily.EngineShowFlags);
+	// All views are considered above water if we don't render any water materials
+	bool bHasAnyViewsAbovewater = !bShouldRenderSingleLayerWater;
+
 	if(bShouldRenderSingleLayerWater)
 	{
+		bool bHasAnyViewsUnderwater = false;
+		bool bAnyViewWithRaytracingTranslucency = false;
+		for (const FViewInfo& View : Views)
+		{
+			bHasAnyViewsUnderwater = bHasAnyViewsUnderwater || View.IsUnderwater();
+			bHasAnyViewsAbovewater = bHasAnyViewsAbovewater || !View.IsUnderwater();
+#if RHI_RAYTRACING
+			bAnyViewWithRaytracingTranslucency = bAnyViewWithRaytracingTranslucency || ShouldRenderRayTracingTranslucency(View);
+#endif
+		}
+
+		// Run a translucency pass here if there are any views underwater. The views that run their translucency here will not run it later in the regular translucency pass
+		if (bHasAnyViewsUnderwater && !bAnyViewWithRaytracingTranslucency && ViewFamily.EngineShowFlags.Translucency && !ViewFamily.EngineShowFlags.VisualizeLightCulling && !ViewFamily.UseDebugViewPS())
+		{
+			CSV_SCOPED_TIMING_STAT_EXCLUSIVE(RenderTranslucency);
+			SCOPE_CYCLE_COUNTER(STAT_TranslucencyDrawTime);
+
+			RHICmdList.SetCurrentStat(GET_STATID(STAT_CLM_Translucency));
+
+			const bool bDrawUnderwaterViews = true;
+			RenderTranslucency(RHICmdList, bDrawUnderwaterViews);
+			ServiceLocalQueue();
+		}
+
 		RHICmdList.SetCurrentStat(GET_STATID(STAT_CLM_WaterPass));
 		SCOPED_DRAW_EVENTF(RHICmdList, WaterRendering, TEXT("WaterRendering"));
 		SCOPED_GPU_STAT(RHICmdList, WaterRendering);
@@ -2219,7 +2245,7 @@ void FDeferredShadingSceneRenderer::Render(FRHICommandListImmediate& RHICmdList)
 	GRenderTargetPool.AddPhaseEvent(TEXT("Translucency"));
 
 	// Draw translucency.
-	if (bCanOverlayRayTracingOutput && ViewFamily.EngineShowFlags.Translucency && !ViewFamily.EngineShowFlags.VisualizeLightCulling && !ViewFamily.UseDebugViewPS())
+	if (bHasAnyViewsAbovewater && bCanOverlayRayTracingOutput && ViewFamily.EngineShowFlags.Translucency && !ViewFamily.EngineShowFlags.VisualizeLightCulling && !ViewFamily.UseDebugViewPS())
 	{
 		CSV_SCOPED_TIMING_STAT_EXCLUSIVE(RenderTranslucency);
 		SCOPE_CYCLE_COUNTER(STAT_TranslucencyDrawTime);
@@ -2242,27 +2268,7 @@ void FDeferredShadingSceneRenderer::Render(FRHICommandListImmediate& RHICmdList)
 		else
 #endif
 		{
-			// For now there is only one resolve for all translucency passes. This can be changed by enabling the resolve in RenderTranslucency()
-			TRefCountPtr<IPooledRenderTarget> SceneColorCopy;
-			ConditionalResolveSceneColorForTranslucentMaterials(RHICmdList, SceneColorCopy);
-
-			// Disable UAV cache flushing so we have optimal VT feedback performance.
-			RHICmdList.AutomaticCacheFlushAfterComputeShader(false);
-
-			if (ViewFamily.AllowTranslucencyAfterDOF())
-			{
-				RenderTranslucency(RHICmdList, ETranslucencyPass::TPT_StandardTranslucency, SceneColorCopy);
-				// Translucency after DOF is rendered now, but stored in the separate translucency RT for later use.
-				RenderTranslucency(RHICmdList, ETranslucencyPass::TPT_TranslucencyAfterDOF, SceneColorCopy);
-			}
-			else // Otherwise render translucent primitives in a single bucket.
-			{
-				RenderTranslucency(RHICmdList, ETranslucencyPass::TPT_AllTranslucency, SceneColorCopy);
-			}
-
-			RHICmdList.AutomaticCacheFlushAfterComputeShader(true);
-			RHICmdList.FlushComputeShaderCache();
-
+			RenderTranslucency(RHICmdList);
 			ServiceLocalQueue();
 
 			static const auto DisableDistortionCVar = IConsoleManager::Get().FindTConsoleVariableDataInt(TEXT("r.DisableDistortion"));
diff --git a/Engine/Source/Runtime/Renderer/Private/DeferredShadingRenderer.h b/Engine/Source/Runtime/Renderer/Private/DeferredShadingRenderer.h
--- a/Engine/Source/Runtime/Renderer/Private/DeferredShadingRenderer.h
+++ b/Engine/Source/Runtime/Renderer/Private/DeferredShadingRenderer.h
@@ -311,8 +311,11 @@ private:
 	/** Resolve the scene color if any translucent material needs it. */
 	void ConditionalResolveSceneColorForTranslucentMaterials(FRHICommandListImmediate& RHICmdList, TRefCountPtr<IPooledRenderTarget>& SceneColorCopy);
 
-	/** Renders the scene's translucency. */
-	void RenderTranslucency(FRHICommandListImmediate& RHICmdList, ETranslucencyPass::Type TranslucencyPass, IPooledRenderTarget* SceneColorCopy);
+	/** Renders the scene's translucency passes. */
+	void RenderTranslucency(FRHICommandListImmediate& RHICmdList, bool bDrawUnderwaterViews = false);
+
+	/** Renders the scene's translucency given a specific pass. */
+	void RenderTranslucencyInner(FRHICommandListImmediate& RHICmdList, ETranslucencyPass::Type TranslucencyPass, IPooledRenderTarget* SceneColorCopy, bool bDrawUnderwaterViews = false);
 
 	/** Renders the scene's light shafts */
 	void RenderLightShaftOcclusion(FRHICommandListImmediate& RHICmdList, FLightShaftsOutput& Output);
diff --git a/Engine/Source/Runtime/Renderer/Private/TranslucentRendering.cpp b/Engine/Source/Runtime/Renderer/Private/TranslucentRendering.cpp
--- a/Engine/Source/Runtime/Renderer/Private/TranslucentRendering.cpp
+++ b/Engine/Source/Runtime/Renderer/Private/TranslucentRendering.cpp
@@ -504,7 +504,7 @@ void FDeferredShadingSceneRenderer::ConditionalResolveSceneColorForTranslucentMa
 			bNeedsResolve |= View.TranslucentPrimCount.UseSceneColorCopy((ETranslucencyPass::Type)TranslucencyPass);
 		}
 
-		if (bNeedsResolve)
+		if (bNeedsResolve && !View.IsUnderwater())
 		{
 			FSceneRenderTargets& SceneContext = FSceneRenderTargets::Get(RHICmdList);
 
@@ -857,7 +857,34 @@ void UpsampleTranslucency(FRHICommandList& RHICmdList, const FViewInfo& View, bo
 	SceneContext.FinishRenderingSceneColor(RHICmdList);
 }
 
-void FDeferredShadingSceneRenderer::RenderTranslucency(FRHICommandListImmediate& RHICmdList, ETranslucencyPass::Type TranslucencyPass, IPooledRenderTarget* SceneColorCopy)
+void FDeferredShadingSceneRenderer::RenderTranslucency(FRHICommandListImmediate& RHICmdList, bool bDrawUnderwaterViews)
+{
+	// For now there is only one resolve for all translucency passes. This can be changed by enabling the resolve in RenderTranslucencyInner()
+	TRefCountPtr<IPooledRenderTarget> SceneColorCopy = GSystemTextures.BlackDummy;
+	if (!bDrawUnderwaterViews)
+	{
+		ConditionalResolveSceneColorForTranslucentMaterials(RHICmdList, SceneColorCopy);
+	}
+
+	// Disable UAV cache flushing so we have optimal VT feedback performance.
+	RHICmdList.AutomaticCacheFlushAfterComputeShader(false);
+
+	if (ViewFamily.AllowTranslucencyAfterDOF())
+	{
+		RenderTranslucencyInner(RHICmdList, ETranslucencyPass::TPT_StandardTranslucency, SceneColorCopy, bDrawUnderwaterViews);
+		// Translucency after DOF is rendered now, but stored in the separate translucency RT for later use.
+		RenderTranslucencyInner(RHICmdList, ETranslucencyPass::TPT_TranslucencyAfterDOF, SceneColorCopy, bDrawUnderwaterViews);
+	}
+	else // Otherwise render translucent primitives in a single bucket.
+	{
+		RenderTranslucencyInner(RHICmdList, ETranslucencyPass::TPT_AllTranslucency, SceneColorCopy, bDrawUnderwaterViews);
+	}
+
+	RHICmdList.AutomaticCacheFlushAfterComputeShader(true);
+	RHICmdList.FlushComputeShaderCache();
+}
+
+void FDeferredShadingSceneRenderer::RenderTranslucencyInner(FRHICommandListImmediate& RHICmdList, ETranslucencyPass::Type TranslucencyPass, IPooledRenderTarget* SceneColorCopy, bool bDrawUnderwaterViews)
 {
 	check(RHICmdList.IsOutsideRenderPass());
 
@@ -879,12 +906,12 @@ void FDeferredShadingSceneRenderer::RenderTranslucency(FRHICommandListImmediate&
 	}
 	FScopedCommandListWaitForTasks Flusher(bUseParallel && (CVarRHICmdFlushRenderThreadTasksTranslucentPass.GetValueOnRenderThread() > 0 || CVarRHICmdFlushRenderThreadTasks.GetValueOnRenderThread() > 0), RHICmdList);
 
-	for (int32 ViewIndex = 0; ViewIndex < Views.Num(); ViewIndex++)
+	for (int32 ViewIndex = 0, NumProcessedViews = 0; ViewIndex < Views.Num(); ViewIndex++)
 	{
 		checkSlow(RHICmdList.IsOutsideRenderPass());
 
 		FViewInfo& View = Views[ViewIndex];
-		if (!View.ShouldRenderView())
+		if (!View.ShouldRenderView() || (Views[ViewIndex].IsUnderwater() != bDrawUnderwaterViews))
 		{
 			continue;
 		}
@@ -906,7 +933,7 @@ void FDeferredShadingSceneRenderer::RenderTranslucency(FRHICommandListImmediate&
 		FMeshPassProcessorRenderState DrawRenderState(View, BasePassUniformBuffer);
 
 		// If downsampling we need to render in the separate buffer. Otherwise we also need to render offscreen to apply TPT_TranslucencyAfterDOF
-		if (RenderInSeparateTranslucency(SceneContext, TranslucencyPass, View.TranslucentPrimCount.DisableOffscreenRendering(TranslucencyPass)))
+		if (!bDrawUnderwaterViews && RenderInSeparateTranslucency(SceneContext, TranslucencyPass, View.TranslucentPrimCount.DisableOffscreenRendering(TranslucencyPass)))
 		{
 			checkSlow(RHICmdList.IsOutsideRenderPass());
 
@@ -937,7 +964,7 @@ void FDeferredShadingSceneRenderer::RenderTranslucency(FRHICommandListImmediate&
 				BeginTimingSeparateTranslucencyPass(RHICmdList, View);
 			}
 
-			SceneContext.BeginRenderingSeparateTranslucency(RHICmdList, View, *this, ViewIndex == 0 || View.Family->bMultiGPUForkAndJoin);
+			SceneContext.BeginRenderingSeparateTranslucency(RHICmdList, View, *this, NumProcessedViews == 0 || View.Family->bMultiGPUForkAndJoin);
 
 			// Draw only translucent prims that are in the SeparateTranslucency pass
 			DrawRenderState.SetDepthStencilState(TStaticDepthStencilState<false, CF_DepthNearOrEqual>::GetRHI());
@@ -968,7 +995,7 @@ void FDeferredShadingSceneRenderer::RenderTranslucency(FRHICommandListImmediate&
 		}
 		else
 		{
-			SceneContext.BeginRenderingTranslucency(RHICmdList, View, *this, ViewIndex == 0 || View.Family->bMultiGPUForkAndJoin);
+			SceneContext.BeginRenderingTranslucency(RHICmdList, View, *this, NumProcessedViews == 0 || View.Family->bMultiGPUForkAndJoin);
 			DrawRenderState.SetDepthStencilState(TStaticDepthStencilState<false, CF_DepthNearOrEqual>::GetRHI());
 
 			if (bUseParallel && !ViewFamily.UseDebugViewPS())
@@ -991,6 +1018,8 @@ void FDeferredShadingSceneRenderer::RenderTranslucency(FRHICommandListImmediate&
 			STAT(View.ViewState->TranslucencyTimer.End(RHICmdList));
 		}
 #endif
+		// Keep track of number of views not skipped
+		NumProcessedViews++;
 	}
 
 	checkSlow(RHICmdList.IsOutsideRenderPass());
