Added comment to clarify that the underwater views will force translucency to render at full resolution

#rnx

diff --git a/Engine/Source/Runtime/Renderer/Private/DeferredShadingRenderer.cpp b/Engine/Source/Runtime/Renderer/Private/DeferredShadingRenderer.cpp
--- a/Engine/Source/Runtime/Renderer/Private/DeferredShadingRenderer.cpp
+++ b/Engine/Source/Runtime/Renderer/Private/DeferredShadingRenderer.cpp
@@ -2069,6 +2069,7 @@ void FDeferredShadingSceneRenderer::Render(FRHICommandListImmediate& RHICmdList)
 		}
 
 		// Run a translucency pass here if there are any views underwater. The views that run their translucency here will not run it later in the regular translucency pass
+		// The translucency pass run here will force all objects to be rendered in fullscreen pass. No partial resolution pass supported here, so that might differ from the behavior if it was rendered in the regular pass instead.
 		if (bHasAnyViewsUnderwater && !bAnyViewWithRaytracingTranslucency && ViewFamily.EngineShowFlags.Translucency && !ViewFamily.EngineShowFlags.VisualizeLightCulling && !ViewFamily.UseDebugViewPS())
 		{
 			CSV_SCOPED_TIMING_STAT_EXCLUSIVE(RenderTranslucency);
