Fix for SkyAtmosphere debug to take into account scene color clear mode.

#udn https://udn.unrealengine.com/questions/544226/ensure-that-hits-100-of-the-time-when-rendering-sk.html

diff --git a/Engine/Source/Runtime/Renderer/Private/DeferredShadingRenderer.cpp b/Engine/Source/Runtime/Renderer/Private/DeferredShadingRenderer.cpp
--- a/Engine/Source/Runtime/Renderer/Private/DeferredShadingRenderer.cpp
+++ b/Engine/Source/Runtime/Renderer/Private/DeferredShadingRenderer.cpp
@@ -1630,14 +1630,16 @@ void FDeferredShadingSceneRenderer::Render(FRHICommandListImmediate& RHICmdList)
 		}
 		else
 		{
-		bIsGBufferCurrent = true;
+			bIsGBufferCurrent = true;
 		}
 		ServiceLocalQueue();
-	}
 
-	if (bRenderSkyAtmosphereEditorNotifications)
-	{
-		RenderSkyAtmosphereEditorNotifications(RHICmdList);
+		if (bRenderSkyAtmosphereEditorNotifications)
+		{
+			// We only render this warning text when bRequiresRHIClear==true to make sure the scene color buffer is allocated at this stage.
+			// When false, the option specifies that all pixels must be written to by a sky dome anyway.
+			RenderSkyAtmosphereEditorNotifications(RHICmdList);
+		}
 	}
 
 	// Wireframe mode requires bRequiresRHIClear to be true. 
