The previous frame buffers used as inputs for the sun average and AA postprocess passes on mobile were not transitioned to readable, triggering an ensure on Vulkan.

The buffers were left as readable at the end of the frame, but were transitioned to writable automatically in FRenderTargetPool::TransitionTargetsWritable() at the beginning of the next frame. The sun average input always comes from the sun merge pass, so that can be fixed by setting the target's AutoWritable property to false. The AA input can come from a number of passes depending on how the pipeline is set up, so it's not as easy to figure out which one needs AutoWritable turned off. I've added an explicit transition for now, since this is going to be rewritten for RDG at some point anyway.

diff --git a/Engine/Source/Runtime/Renderer/Private/PostProcess/PostProcessMobile.cpp b/Engine/Source/Runtime/Renderer/Private/PostProcess/PostProcessMobile.cpp
--- a/Engine/Source/Runtime/Renderer/Private/PostProcess/PostProcessMobile.cpp
+++ b/Engine/Source/Runtime/Renderer/Private/PostProcess/PostProcessMobile.cpp
@@ -1749,6 +1749,7 @@ FPooledRenderTargetDesc FRCPassPostProcessSunMergeES2::ComputeOutputDesc(EPassOu
 	Ret.Extent.Y = FMath::Max(1, PrePostSourceViewportSize.Y/4);
 	Ret.DebugName = TEXT("SunMerge");
 	Ret.ClearValue = FClearValueBinding(FLinearColor::Black);
+	Ret.AutoWritable = false;
 	return Ret;
 }
 
@@ -1955,6 +1956,7 @@ FPooledRenderTargetDesc FRCPassPostProcessSunMergeSmallES2::ComputeOutputDesc(EP
 	Ret.Extent.Y = FMath::Max(1, PrePostSourceViewportSize.Y/4);
 	Ret.DebugName = TEXT("SunMergeSmall");
 	Ret.ClearValue = FClearValueBinding(FLinearColor::Black);
+	Ret.AutoWritable = false;
 	return Ret;
 }
 
@@ -2958,6 +2960,12 @@ void FRCPassPostProcessAaES2::Process(FRenderingCompositePassContext& Context)
 		// Full clear to avoid restore
 		LoadStoreAction = ERenderTargetActions::Clear_Store;
 	}
+
+	// The previous frame target has been transitioned to writable in FRenderTargetPool::TransitionTargetsWritable(), so we
+	// need to transition it to readable again. Ideally we'll get rid of this useless read->write->read transition when we
+	// port this over to RDG.
+	const FSceneRenderTargetItem& PrevFrameInput = GetInput(ePId_Input1)->GetOutput()->RequestInput()->GetRenderTargetItem();
+	Context.RHICmdList.TransitionResource(EResourceTransitionAccess::EReadable, PrevFrameInput.ShaderResourceTexture);
 	
 	FRHIRenderPassInfo RPInfo(DestRenderTarget.TargetableTexture, LoadStoreAction);
 	Context.RHICmdList.BeginRenderPass(RPInfo, TEXT("AaES2"));
