Fixed: Spot Lights don't work on mobile if Cast Shadows is disabled

diff --git a/Engine/Source/Runtime/Renderer/Private/LightSceneInfo.cpp b/Engine/Source/Runtime/Renderer/Private/LightSceneInfo.cpp
--- a/Engine/Source/Runtime/Renderer/Private/LightSceneInfo.cpp
+++ b/Engine/Source/Runtime/Renderer/Private/LightSceneInfo.cpp
@@ -68,14 +68,23 @@ void FLightSceneInfo::AddToScene()
 {
 	const FLightSceneInfoCompact& LightSceneInfoCompact = Scene->Lights[Id];
 
+	bool bIsValidLightTypeMobile = false;
+	if (Scene->GetShadingPath() == EShadingPath::Mobile && Proxy->IsMovable())
+	{
+		static const auto MobileEnableMovableSpotLightsVar = IConsoleManager::Get().FindTConsoleVariableDataInt(TEXT("r.Mobile.EnableMovableSpotLights"));
+		const uint8 LightType = Proxy->GetLightType();
+		bIsValidLightTypeMobile = LightType == LightType_Rect || LightType == LightType_Point
+			|| (LightType == LightType_Spot && MobileEnableMovableSpotLightsVar->GetValueOnRenderThread());
+	}
+
 	// Only need to create light interactions for lights that can cast a shadow, 
 	// As deferred shading doesn't need to know anything about the primitives that a light affects
 	if (Proxy->CastsDynamicShadow() 
 		|| Proxy->CastsStaticShadow() 
 		// Lights that should be baked need to check for interactions to track unbuilt state correctly
 		|| Proxy->HasStaticLighting()
-		// ES2 path supports dynamic point lights in the base pass using forward rendering, so we need to know the primitives
-		|| (Scene->GetFeatureLevel() < ERHIFeatureLevel::SM5 && Proxy->GetLightType() == LightType_Point && Proxy->IsMovable()))
+		// Mobile path supports dynamic point/spot lights in the base pass using forward rendering, so we need to know the primitives
+		|| bIsValidLightTypeMobile)
 	{
 		// Add the light to the scene's light octree.
 		Scene->LightOctree.AddElement(LightSceneInfoCompact);
