Mark the call to BeginRenderingSeparateTranslucency as being the first in the frame if the SeparateTransluicencyRT buffer hasn't been allocated yet this frame. This is to prevent stale data to make it to the translucency composite pass.

diff --git a/Engine/Source/Runtime/Renderer/Private/LightShaftRendering.cpp b/Engine/Source/Runtime/Renderer/Private/LightShaftRendering.cpp
--- a/Engine/Source/Runtime/Renderer/Private/LightShaftRendering.cpp
+++ b/Engine/Source/Runtime/Renderer/Private/LightShaftRendering.cpp
@@ -944,7 +944,8 @@ void ApplyLightShaftBloom(FRHICommandListImmediate& RHICmdList, const FViewInfo&
 	bool bUseSeparateTranslucency = false;
 	if (View.Family->AllowTranslucencyAfterDOF() && GLightShaftRenderAfterDOF)
 	{
-		SceneContext.BeginRenderingSeparateTranslucency(RHICmdList, View, Renderer, false);
+		// If the separate translucency RT hasn't been allocated yet this frame, make sure to mark the call to BeginRenderingSeparateTranslucency as being the first in this frame
+		SceneContext.BeginRenderingSeparateTranslucency(RHICmdList, View, Renderer, !SceneContext.SeparateTranslucencyRT.IsValid());
 		bUseSeparateTranslucency = true;
 	}
 	else
