Include the whole Unix platform group in DOF whitelist (UE-82031).

diff --git a/Engine/Source/Runtime/Renderer/Private/PostProcess/DiaphragmDOF.h b/Engine/Source/Runtime/Renderer/Private/PostProcess/DiaphragmDOF.h
--- a/Engine/Source/Runtime/Renderer/Private/PostProcess/DiaphragmDOF.h
+++ b/Engine/Source/Runtime/Renderer/Private/PostProcess/DiaphragmDOF.h
@@ -12,7 +12,7 @@
 
 
 // Whitelist diaphragm DOF for platforms that actually have been tested.
-#define WITH_DIAPHRAGM_DOF (PLATFORM_WINDOWS || PLATFORM_XBOXONE || PLATFORM_PS4 || PLATFORM_MAC || PLATFORM_LINUX || PLATFORM_IOS || PLATFORM_SWITCH)
+#define WITH_DIAPHRAGM_DOF (PLATFORM_WINDOWS || PLATFORM_XBOXONE || PLATFORM_PS4 || PLATFORM_MAC || PLATFORM_UNIX || PLATFORM_IOS || PLATFORM_SWITCH)
 
 
 class FViewInfo;
