Prevent rendering crash with attempting to invoke ray tracing sky light pass with forward renderer.

diff --git a/Engine/Source/Runtime/Renderer/Private/DeferredShadingRenderer.cpp b/Engine/Source/Runtime/Renderer/Private/DeferredShadingRenderer.cpp
--- a/Engine/Source/Runtime/Renderer/Private/DeferredShadingRenderer.cpp
+++ b/Engine/Source/Runtime/Renderer/Private/DeferredShadingRenderer.cpp
@@ -1844,7 +1844,7 @@ void FDeferredShadingSceneRenderer::Render(FRHICommandListImmediate& RHICmdList)
 	TRefCountPtr<IPooledRenderTarget> SkyLightHitDistanceRT;
 
 	const bool bRayTracingEnabled = IsRayTracingEnabled();
-	if (bRayTracingEnabled && bCanOverlayRayTracingOutput)
+	if (bRayTracingEnabled && bCanOverlayRayTracingOutput && !IsForwardShadingEnabled(ShaderPlatform))
 	{		
 		RenderRayTracingSkyLight(RHICmdList, SkyLightRT, SkyLightHitDistanceRT);
 	}
diff --git a/Engine/Source/Runtime/Renderer/Private/ScenePrivate.h b/Engine/Source/Runtime/Renderer/Private/ScenePrivate.h
--- a/Engine/Source/Runtime/Renderer/Private/ScenePrivate.h
+++ b/Engine/Source/Runtime/Renderer/Private/ScenePrivate.h
@@ -47,6 +47,7 @@
 #include "VolumetricFogShared.h"
 #include "DebugViewModeRendering.h"
 #include "PrecomputedVolumetricLightmap.h"
+#include "RayTracing/RaytracingOptions.h"
 #if RHI_RAYTRACING
 #include "RayTracing/RayTracingIESLightProfiles.h"
 #include "Halton.h"
@@ -2859,7 +2860,7 @@ public:
 
 	bool ShouldRenderSkylightInBasePass(EBlendMode BlendMode) const
 	{
-		bool bRenderSkyLight = SkyLight && !SkyLight->bHasStaticLighting && !SkyLight->bCastRayTracedShadow;
+		bool bRenderSkyLight = SkyLight && !SkyLight->bHasStaticLighting && !(ShouldRenderRayTracingSkyLight(SkyLight) && !IsForwardShadingEnabled(GetShaderPlatform()));
 
 		if (IsTranslucentBlendMode(BlendMode))
 		{
