Added an option to disable SkyAtmosphere editor debug text slowing down everything.

diff --git a/Engine/Source/Runtime/Renderer/Private/SkyAtmosphereRendering.cpp b/Engine/Source/Runtime/Renderer/Private/SkyAtmosphereRendering.cpp
--- a/Engine/Source/Runtime/Renderer/Private/SkyAtmosphereRendering.cpp
+++ b/Engine/Source/Runtime/Renderer/Private/SkyAtmosphereRendering.cpp
@@ -192,19 +192,23 @@ static TAutoConsoleVariable<float> CVarSkyAtmosphereDistantSkyLightLUTAltitude(
 	TEXT("The altitude at which the sky samples are taken to integrate the sky lighting. Default to 6km, typicaly cirrus clouds altitude.\n"),
 	ECVF_RenderThreadSafe | ECVF_Scalability);
 
-////////////////////////////////////////////////////////////////////////// Debug / Visualisation
+////////////////////////////////////////////////////////////////////////// Debug / Visualization
 
 static TAutoConsoleVariable<int32> CVarSkyAtmosphereLUT32(
 	TEXT("r.SkyAtmosphere.LUT32"), 0,
 	TEXT("Use full 32bit per-channel precision for all sky LUTs.\n"),
 	ECVF_RenderThreadSafe | ECVF_Scalability);
 
-TAutoConsoleVariable<int32> CVarSkyAtmosphereVisualize(
+static TAutoConsoleVariable<int32> CVarSkyAtmosphereVisualize(
 	TEXT("r.SkyAtmosphere.Visualize"),
 	0,
 	TEXT("DRAW DEBUG"),
 	ECVF_RenderThreadSafe);
 
+static TAutoConsoleVariable<int32> CVarSkyAtmosphereDebugText(
+	TEXT("r.SkyAtmosphere.DebugText"), 1,
+	TEXT("Renders debug text when a sky dome with a sky material is used, to warn the user when it does not cover some area of the screen.\n"));
+
 DECLARE_GPU_STAT(SkyAtmosphereLUTs);
 DECLARE_GPU_STAT(SkyAtmosphere);
 DECLARE_GPU_STAT(SkyAtmosphereEditor);
@@ -1432,7 +1436,7 @@ bool FSceneRenderer::ShouldRenderSkyAtmosphereEditorNotifications()
 	{
 		bAnyViewHasSkyMaterial |= Views[ViewIndex].bSceneHasSkyMaterial;
 	}
-	return bAnyViewHasSkyMaterial;
+	return bAnyViewHasSkyMaterial && CVarSkyAtmosphereDebugText.GetValueOnAnyThread()>0;
 #endif
 	return false;
 }
