VS2019 compile fix in DoLazyStaticMeshUpdateCVarSinkFunction

diff --git a/Engine/Source/Runtime/Renderer/Private/RendererScene.cpp b/Engine/Source/Runtime/Renderer/Private/RendererScene.cpp
--- a/Engine/Source/Runtime/Renderer/Private/RendererScene.cpp
+++ b/Engine/Source/Runtime/Renderer/Private/RendererScene.cpp
@@ -564,8 +564,8 @@ static void DoLazyStaticMeshUpdateCVarSinkFunction()
 		return;
 	}
 
-	static bool CachedDoLazyStaticMeshUpdate = CVarDoLazyStaticMeshUpdate.GetValueOnGameThread();
-	bool DoLazyStaticMeshUpdate = CVarDoLazyStaticMeshUpdate.GetValueOnGameThread();
+	static bool CachedDoLazyStaticMeshUpdate = !!CVarDoLazyStaticMeshUpdate.GetValueOnGameThread();
+	bool DoLazyStaticMeshUpdate = !!CVarDoLazyStaticMeshUpdate.GetValueOnGameThread();
 
 	if (DoLazyStaticMeshUpdate != CachedDoLazyStaticMeshUpdate)
 	{
