Fixing integration records (copy down from Dev-Main)

diff --git a/Engine/Source/Runtime/Renderer/Public/RayTracingDynamicGeometryCollection.h b/Engine/Source/Runtime/Renderer/Public/RayTracingDynamicGeometryCollection.h
--- a/Engine/Source/Runtime/Renderer/Public/RayTracingDynamicGeometryCollection.h
+++ b/Engine/Source/Runtime/Renderer/Public/RayTracingDynamicGeometryCollection.h
@@ -9,7 +9,12 @@ class RENDERER_API FRayTracingDynamicGeometryCollection
 public:
 	FRayTracingDynamicGeometryCollection();
 
-	void AddDynamicMeshBatchForGeometryUpdate(const FScene* Scene, const FSceneView* View, const FPrimitiveSceneProxy* PrimitiveSceneProxy, const FMeshBatch& MeshBatch, FRayTracingGeometry& Geometry, uint32 NumMaxVertices, FRWBuffer& Buffer);
+	void AddDynamicMeshBatchForGeometryUpdate(
+		const FScene* Scene, 
+		const FSceneView* View, 
+		const FPrimitiveSceneProxy* PrimitiveSceneProxy, 
+		FRayTracingDynamicGeometryUpdateParams Params);
+
 	void DispatchUpdates(FRHICommandListImmediate& RHICmdList);
 	void Clear();
 
