Fixed DF TileDimensions update that was set only when TileIntersectionResources was created / reallocated

#rnx

diff --git a/Engine/Source/Runtime/Renderer/Private/DistanceFieldLightingShared.h b/Engine/Source/Runtime/Renderer/Private/DistanceFieldLightingShared.h
--- a/Engine/Source/Runtime/Renderer/Private/DistanceFieldLightingShared.h
+++ b/Engine/Source/Runtime/Renderer/Private/DistanceFieldLightingShared.h
@@ -495,20 +495,22 @@ class FLightTileIntersectionResources
 {
 public:
 
-	FLightTileIntersectionResources() :
-		b16BitIndices(false)
+	FLightTileIntersectionResources() 
+		: TileDimensions(FIntPoint::ZeroValue)
+		, TileAlignedDimensions(FIntPoint::ZeroValue)
+		, b16BitIndices(false)
 	{}
 
 	void Initialize()
 	{
-		FIntPoint AlignedDimensions = FIntPoint(Align(TileDimensions.X, 64), Align(TileDimensions.Y, 64));
-	
-		TileNumCulledObjects.Initialize(sizeof(uint32), AlignedDimensions.X * AlignedDimensions.Y, PF_R32_UINT, BUF_Static);
-		TileStartOffsets.Initialize(sizeof(uint32), AlignedDimensions.X * AlignedDimensions.Y, PF_R32_UINT, BUF_Static);
+		TileAlignedDimensions = FIntPoint(Align(TileDimensions.X, 64), Align(TileDimensions.Y, 64));
+			
+		TileNumCulledObjects.Initialize(sizeof(uint32), TileAlignedDimensions.X * TileAlignedDimensions.Y, PF_R32_UINT, BUF_Static);
+		TileStartOffsets.Initialize(sizeof(uint32), TileAlignedDimensions.X * TileAlignedDimensions.Y, PF_R32_UINT, BUF_Static);
 		NextStartOffset.Initialize(sizeof(uint32), 1, PF_R32_UINT, BUF_Static);
 
 		//@todo - handle max exceeded
-		TileArrayData.Initialize(b16BitIndices ? sizeof(uint16) : sizeof(uint32), GAverageObjectsPerShadowCullTile * AlignedDimensions.X * AlignedDimensions.Y * LightTileDataStride, b16BitIndices ? PF_R16_UINT : PF_R32_UINT, BUF_Static);
+		TileArrayData.Initialize(b16BitIndices ? sizeof(uint16) : sizeof(uint32), GAverageObjectsPerShadowCullTile * TileAlignedDimensions.X * TileAlignedDimensions.Y * LightTileDataStride, b16BitIndices ? PF_R16_UINT : PF_R32_UINT, BUF_Static);
 	}
 
 	void Release()
@@ -524,12 +526,18 @@ public:
 		return TileNumCulledObjects.NumBytes + NextStartOffset.NumBytes + TileStartOffsets.NumBytes + TileArrayData.NumBytes;
 	}
 
+	FIntPoint GetTileAlignedDimensions() const 
+	{
+		return TileAlignedDimensions;
+	}
+	
 	static FIntPoint GetAlignedDimensions(FIntPoint InTileDimensions)
 	{
 		return FIntPoint(Align(InTileDimensions.X, 64), Align(InTileDimensions.Y, 64));
 	}
 
 	FIntPoint TileDimensions;
+	FIntPoint TileAlignedDimensions;
 
 	FRWBuffer TileNumCulledObjects;
 	FRWBuffer NextStartOffset;
diff --git a/Engine/Source/Runtime/Renderer/Private/DistanceFieldShadowing.cpp b/Engine/Source/Runtime/Renderer/Private/DistanceFieldShadowing.cpp
--- a/Engine/Source/Runtime/Renderer/Private/DistanceFieldShadowing.cpp
+++ b/Engine/Source/Runtime/Renderer/Private/DistanceFieldShadowing.cpp
@@ -792,8 +792,7 @@ void CullDistanceFieldObjectsForLight(
 		if (TileIntersectionResources)
 		{
 			FIntPoint AlignedLightDimensions = FLightTileIntersectionResources::GetAlignedDimensions(LightTileDimensions);
-			FIntPoint CurrentAlignedLightDimensions = FLightTileIntersectionResources::GetAlignedDimensions(TileIntersectionResources->TileDimensions);
-			bLightDimensionsDirty = AlignedLightDimensions.X > CurrentAlignedLightDimensions.X || AlignedLightDimensions.Y > CurrentAlignedLightDimensions.Y;
+			bLightDimensionsDirty = AlignedLightDimensions.X > TileIntersectionResources->GetTileAlignedDimensions().X;
 		}
 		
 		if (!TileIntersectionResources || bLightDimensionsDirty || TileIntersectionResources->b16BitIndices != b16BitObjectIndices)
@@ -807,12 +806,14 @@ void CullDistanceFieldObjectsForLight(
 				TileIntersectionResources = MakeUnique<FLightTileIntersectionResources>();
 			}
 
-			TileIntersectionResources->TileDimensions = LightTileDimensions;
 			TileIntersectionResources->b16BitIndices = b16BitObjectIndices;
-
+			TileIntersectionResources->TileDimensions = LightTileDimensions;
 			TileIntersectionResources->Initialize();
 		}
 
+		check(TileIntersectionResources);
+		TileIntersectionResources->TileDimensions = LightTileDimensions;
+
 		{
 			SCOPED_DRAW_EVENT(RHICmdList, ComputeTileStartOffsets);
 
