Fix MSAA interaction with before tonemapper post-process materials on iOS
#rnx

diff --git a/Engine/Shaders/Private/MaterialTemplate.ush b/Engine/Shaders/Private/MaterialTemplate.ush
--- a/Engine/Shaders/Private/MaterialTemplate.ush
+++ b/Engine/Shaders/Private/MaterialTemplate.ush
@@ -1650,7 +1650,11 @@ MaterialFloat4 MobileSceneTextureLookup(inout FMaterialPixelParameters Parameter
 #if POST_PROCESS_MATERIAL
 		MaterialFloat4 Input0 = Texture2DSample(PostProcessInput_0_Texture, PostProcessInput_0_SharedSampler, UV);
 		#if POST_PROCESS_MATERIAL_BEFORE_TONEMAP
-		Input0 = Decode32BPPHDR(Input0);
+			Input0 = Decode32BPPHDR(Input0);
+			#if METAL_PROFILE
+				// If MSAA is in use undo PreTonemapMSAA, it is reapplied at the end of post process material
+				Input0.rgb = lerp(Input0.rgb, Input0.rgb * rcp(Input0.r*(-0.299) + Input0.g*(-0.587) + Input0.b*(-0.114) + 1.0), ResolvedView.IsMobileMSAA);
+			#endif
 		#endif
 		// We need to preserve original SceneColor Alpha as it's used by tonemapper on mobile
 		Parameters.BackupSceneColorAlpha = Input0.a;
diff --git a/Engine/Shaders/Private/PostProcessMaterialShaders.usf b/Engine/Shaders/Private/PostProcessMaterialShaders.usf
--- a/Engine/Shaders/Private/PostProcessMaterialShaders.usf
+++ b/Engine/Shaders/Private/PostProcessMaterialShaders.usf
@@ -342,5 +342,9 @@ void MainPS(
 #endif
 #endif
 	OutColor = FullSceneColor;
+#if POST_PROCESS_MATERIAL_BEFORE_TONEMAP && METAL_PROFILE
+	// If MSAA is in use we reapply the on-chip pre-tonemap since the tonemapper will be expecting it.
+	OutColor.rgb = lerp(OutColor.rgb, OutColor.rgb* rcp(OutColor.r*0.299 + OutColor.g*0.587 + OutColor.b*0.114 + 1.0), ResolvedView.IsMobileMSAA);
+#endif
 }
 #endif //(FEATURE_LEVEL > FEATURE_LEVEL_ES3_1)
diff --git a/Engine/Source/Runtime/Engine/Public/SceneView.h b/Engine/Source/Runtime/Engine/Public/SceneView.h
--- a/Engine/Source/Runtime/Engine/Public/SceneView.h
+++ b/Engine/Source/Runtime/Engine/Public/SceneView.h
@@ -696,6 +696,7 @@ enum ETranslucencyVolumeCascade
 	VIEW_UNIFORM_BUFFER_MEMBER(FLinearColor, SkyLightColor) \
 	VIEW_UNIFORM_BUFFER_MEMBER_ARRAY(FVector4, SkyIrradianceEnvironmentMap, [7]) \
 	VIEW_UNIFORM_BUFFER_MEMBER(float, MobilePreviewMode) \
+	VIEW_UNIFORM_BUFFER_MEMBER(float, IsMobileMSAA) \
 	VIEW_UNIFORM_BUFFER_MEMBER(float, HMDEyePaddingOffset) \
 	VIEW_UNIFORM_BUFFER_MEMBER_EX(float, ReflectionCubemapMaxMip, EShaderPrecisionModifier::Half) \
 	VIEW_UNIFORM_BUFFER_MEMBER(float, ShowDecalsMask) \
diff --git a/Engine/Source/Runtime/Renderer/Private/SceneRendering.cpp b/Engine/Source/Runtime/Renderer/Private/SceneRendering.cpp
--- a/Engine/Source/Runtime/Renderer/Private/SceneRendering.cpp
+++ b/Engine/Source/Runtime/Renderer/Private/SceneRendering.cpp
@@ -1575,6 +1575,12 @@ void FViewInfo::SetupUniformBufferParameters(
 		(RHIFeatureLevel == ERHIFeatureLevel::ES2 || RHIFeatureLevel == ERHIFeatureLevel::ES3_1) &&
 		GMaxRHIFeatureLevel > ERHIFeatureLevel::ES3_1) ? 1.0f : 0.0f;
 
+	static const auto CVarMobileMSAA = IConsoleManager::Get().FindTConsoleVariableDataInt(TEXT("r.MobileMSAA"));
+	bool bMobileMSAA = Scene && IsMetalMobilePlatform(Scene->GetShaderPlatform())
+		&& CVarMobileMSAA
+		&& CVarMobileMSAA->GetValueOnAnyThread() > 1;
+	ViewUniformShaderParameters.IsMobileMSAA = bMobileMSAA ? 1.0f : 0.0f;
+
 	// Padding between the left and right eye may be introduced by an HMD, which instanced stereo needs to account for.
 	if ((IStereoRendering::IsStereoEyePass(StereoPass)) && (Family->Views.Num() > 1))
 	{
