fix for CIS unity builds.
#rnx

diff --git a/Engine/Source/Runtime/Renderer/Private/SceneRendering.cpp b/Engine/Source/Runtime/Renderer/Private/SceneRendering.cpp
--- a/Engine/Source/Runtime/Renderer/Private/SceneRendering.cpp
+++ b/Engine/Source/Runtime/Renderer/Private/SceneRendering.cpp
@@ -1575,10 +1575,10 @@ void FViewInfo::SetupUniformBufferParameters(
 		(RHIFeatureLevel == ERHIFeatureLevel::ES2 || RHIFeatureLevel == ERHIFeatureLevel::ES3_1) &&
 		GMaxRHIFeatureLevel > ERHIFeatureLevel::ES3_1) ? 1.0f : 0.0f;
 
-	static const auto CVarMobileMSAA = IConsoleManager::Get().FindTConsoleVariableDataInt(TEXT("r.MobileMSAA"));
+	static const auto MobileMSAACVar = IConsoleManager::Get().FindTConsoleVariableDataInt(TEXT("r.MobileMSAA"));
 	bool bMobileMSAA = Scene && IsMetalMobilePlatform(Scene->GetShaderPlatform())
-		&& CVarMobileMSAA
-		&& CVarMobileMSAA->GetValueOnAnyThread() > 1;
+		&& MobileMSAACVar
+		&& MobileMSAACVar->GetValueOnAnyThread() > 1;
 	ViewUniformShaderParameters.IsMobileMSAA = bMobileMSAA ? 1.0f : 0.0f;
 
 	// Padding between the left and right eye may be introduced by an HMD, which instanced stereo needs to account for.
