Fixed bloom threshold to include values > -1.0 instead of > 0.

diff --git a/Engine/Source/Runtime/Renderer/Private/PostProcess/PostProcessBloomSetup.cpp b/Engine/Source/Runtime/Renderer/Private/PostProcess/PostProcessBloomSetup.cpp
--- a/Engine/Source/Runtime/Renderer/Private/PostProcess/PostProcessBloomSetup.cpp
+++ b/Engine/Source/Runtime/Renderer/Private/PostProcess/PostProcessBloomSetup.cpp
@@ -113,7 +113,7 @@ FScreenPassTexture AddBloomSetupPass(FRDGBuilder& GraphBuilder, const FViewInfo&
 {
 	check(Inputs.SceneColor.IsValid());
 	check(Inputs.EyeAdaptationTexture);
-	check(Inputs.Threshold > 0.0f);
+	check(Inputs.Threshold > -1.0f);
 
 	const bool bIsComputePass = View.bUseComputePasses;
 
diff --git a/Engine/Source/Runtime/Renderer/Private/PostProcess/PostProcessing.cpp b/Engine/Source/Runtime/Renderer/Private/PostProcess/PostProcessing.cpp
--- a/Engine/Source/Runtime/Renderer/Private/PostProcess/PostProcessing.cpp
+++ b/Engine/Source/Runtime/Renderer/Private/PostProcess/PostProcessing.cpp
@@ -512,9 +512,9 @@ void AddPostProcessingPasses(FRDGBuilder& GraphBuilder, const FViewInfo& View, c
 			FBloomInputs PassInputs;
 			PassInputs.SceneColor = SceneColor;
 
-			const bool bBloomThresholdEnabled = View.FinalPostProcessSettings.BloomThreshold > 0.0f;
+			const bool bBloomThresholdEnabled = View.FinalPostProcessSettings.BloomThreshold > -1.0f;
 
-			// Reuse the main scene downsample chain if a threshold isn't required for bloom. 
+			// Reuse the main scene downsample chain if a threshold isn't required for bloom.
 			if (SceneDownsampleChain.IsInitialized() && !bBloomThresholdEnabled)
 			{
 				PassInputs.SceneDownsampleChain = &SceneDownsampleChain;
