Fixed bug with high resolution screenshot capture regions where the material mask would always be displayed.

diff --git a/Engine/Source/Runtime/Renderer/Private/PostProcess/PostProcessMaterial.cpp b/Engine/Source/Runtime/Renderer/Private/PostProcess/PostProcessMaterial.cpp
--- a/Engine/Source/Runtime/Renderer/Private/PostProcess/PostProcessMaterial.cpp
+++ b/Engine/Source/Runtime/Renderer/Private/PostProcess/PostProcessMaterial.cpp
@@ -766,9 +766,11 @@ FScreenPassTexture AddHighResolutionScreenshotMaskPass(
 
 	static_assert(UE_ARRAY_COUNT(PassNames) == static_cast<uint32>(EPass::MAX), "Pass names array doesn't match pass enum");
 
+	const bool bHighResScreenshotMask = View.Family->EngineShowFlags.HighResScreenshotMask != 0;
+
 	TOverridePassSequence<EPass> PassSequence(Inputs.OverrideOutput);
-	PassSequence.SetEnabled(EPass::Material, Inputs.Material != nullptr);
-	PassSequence.SetEnabled(EPass::MaskMaterial, Inputs.MaskMaterial != nullptr && GIsHighResScreenshot);
+	PassSequence.SetEnabled(EPass::Material, bHighResScreenshotMask && Inputs.Material != nullptr);
+	PassSequence.SetEnabled(EPass::MaskMaterial, bHighResScreenshotMask && Inputs.MaskMaterial != nullptr && GIsHighResScreenshot);
 	PassSequence.SetEnabled(EPass::CaptureRegionMaterial, Inputs.CaptureRegionMaterial != nullptr);
 	PassSequence.Finalize();
 
