FNM - Fix memory corruption in MeshDrawCommands when binding more than 8 SRVs and Textures

diff --git a/Engine/Source/Runtime/Renderer/Private/MeshPassProcessor.cpp b/Engine/Source/Runtime/Renderer/Private/MeshPassProcessor.cpp
--- a/Engine/Source/Runtime/Renderer/Private/MeshPassProcessor.cpp
+++ b/Engine/Source/Runtime/Renderer/Private/MeshPassProcessor.cpp
@@ -148,7 +148,7 @@ void FMeshDrawShaderBindings::SetShaderBindings(
 		checkSlow(Parameter.BaseIndex < UE_ARRAY_COUNT(ShaderBindingState.SRVs));
 
 		uint32 TypeByteIndex = SRVIndex / 8;
-		uint32 TypeBitIndex = SRVIndex - TypeByteIndex;
+		uint32 TypeBitIndex = SRVIndex % 8;
 
 		if (SRVType[TypeByteIndex] & (1 << TypeBitIndex))
 		{
@@ -233,7 +233,7 @@ void FMeshDrawShaderBindings::SetShaderBindings(
 		FShaderParameterInfo Parameter = SRVParameters[SRVIndex];
 
 		uint32 TypeByteIndex = SRVIndex / 8;
-		uint32 TypeBitIndex = SRVIndex - TypeByteIndex;
+		uint32 TypeBitIndex = SRVIndex % 8;
 
 		if (SRVType[TypeByteIndex] & (1 << TypeBitIndex))
 		{
@@ -620,7 +620,7 @@ void FMeshDrawShaderBindings::Finalize(const FMeshProcessorShaders* ShadersForDe
 			FShaderParameterInfo Parameter = SRVParameters[SRVIndex];
 
 			uint32 TypeByteIndex = SRVIndex / 8;
-			uint32 TypeBitIndex = SRVIndex - TypeByteIndex;
+			uint32 TypeBitIndex = SRVIndex % 8;
 
 			if (SRVType[TypeByteIndex] & (1 << TypeBitIndex))
 			{
diff --git a/Engine/Source/Runtime/Renderer/Public/MeshDrawShaderBindings.h b/Engine/Source/Runtime/Renderer/Public/MeshDrawShaderBindings.h
--- a/Engine/Source/Runtime/Renderer/Public/MeshDrawShaderBindings.h
+++ b/Engine/Source/Runtime/Renderer/Public/MeshDrawShaderBindings.h
@@ -320,7 +320,7 @@ private:
 		if (FoundIndex >= 0)
 		{
 			uint32 TypeByteIndex = FoundIndex / 8;
-			uint32 TypeBitIndex = FoundIndex - TypeByteIndex;
+			uint32 TypeBitIndex = FoundIndex % 8;
 			GetSRVTypeStart()[TypeByteIndex] |= 1 << TypeBitIndex;
 			GetSRVStart()[FoundIndex] = Value;
 		}
