Inject proper distance under height field in the global distance field

diff --git a/Engine/Shaders/Private/GlobalDistanceField.usf b/Engine/Shaders/Private/GlobalDistanceField.usf
--- a/Engine/Shaders/Private/GlobalDistanceField.usf
+++ b/Engine/Shaders/Private/GlobalDistanceField.usf
@@ -113,6 +113,7 @@ uint ClipmapIndex;
 float3 VolumeTexelSize;
 float3 UpdateRegionVolumeMin;
 float UpdateRegionVolumeStep;
+float HeightfieldThickness;
 uint3 CullGridDimension;
 
 RWBuffer<uint> RWCullGridObjectNum;
@@ -370,6 +371,13 @@ void CompositeHeightfieldsIntoGlobalDistanceFieldCS(
 					float DistanceToHeightfieldPlane = dot(WorldHeightfieldNormal, WorldPosition - WorldHeightfieldShadingPosition);
 					float MinDistance = clamp(DistanceToHeightfieldPlane, -AOGlobalMaxSphereQueryRadius, AOGlobalMaxSphereQueryRadius);
 
+					// Limit negative region of a heighfield to a user defined thickness  
+					const float MinInteriorDistance = -HeightfieldThickness;
+					if (MinDistance < MinInteriorDistance)
+					{
+						MinDistance = min(MinInteriorDistance - MinDistance, AOGlobalMaxSphereQueryRadius);
+					}
+
 					float3 GlobalUV = ComputeGlobalUV(WorldPosition, ClipmapIndex);
 					uint3 WriteCoordinate = (uint3)(frac(GlobalUV) * GlobalVolumeDimension);
 					RWGlobalDistanceFieldTexture[WriteCoordinate] = min(MinDistance, RWGlobalDistanceFieldTexture[WriteCoordinate]);
diff --git a/Engine/Source/Runtime/Renderer/Private/GlobalDistanceField.cpp b/Engine/Source/Runtime/Renderer/Private/GlobalDistanceField.cpp
--- a/Engine/Source/Runtime/Renderer/Private/GlobalDistanceField.cpp
+++ b/Engine/Source/Runtime/Renderer/Private/GlobalDistanceField.cpp
@@ -102,6 +102,14 @@ FAutoConsoleVariableRef CVarAOGlobalDistanceFieldRepresentHeightfields(
 	ECVF_RenderThreadSafe
 	);
 
+float GGlobalDistanceFieldHeightFieldThicknessScale = 4.0f;
+FAutoConsoleVariableRef CVarGlobalDistanceFieldHeightFieldThicknessScale(
+	TEXT("r.GlobalDistanceFieldHeightFieldThicknessScale"),
+	GGlobalDistanceFieldHeightFieldThicknessScale,
+	TEXT("Thickness of the height field when it's entered into the global distance field, measured in distance field voxels. Defaults to 4 which means 4x the voxel size as thickness."),
+	ECVF_RenderThreadSafe
+	);
+
 void FGlobalDistanceFieldInfo::UpdateParameterData(float MaxOcclusionDistance)
 {
 	if (Clipmaps.Num() > 0)
@@ -618,6 +626,7 @@ public:
 		AOGlobalMaxSphereQueryRadius.Bind(Initializer.ParameterMap, TEXT("AOGlobalMaxSphereQueryRadius"));
 		HeightfieldDescriptionParameters.Bind(Initializer.ParameterMap);
 		HeightfieldTextureParameters.Bind(Initializer.ParameterMap);
+		HeightfieldThickness.Bind(Initializer.ParameterMap, TEXT("HeightfieldThickness"));
 	}
 
 	FCompositeHeightfieldsIntoGlobalDistanceFieldCS()
@@ -653,6 +662,7 @@ public:
 		SetShaderValue(RHICmdList, ShaderRHI, UpdateRegionVolumeStep, VolumeStep);
 		SetShaderValue(RHICmdList, ShaderRHI, ClipmapIndex, ClipmapIndexValue);
 		SetShaderValue(RHICmdList, ShaderRHI, AOGlobalMaxSphereQueryRadius, GlobalMaxSphereQueryRadius);
+		SetShaderValue(RHICmdList, ShaderRHI, HeightfieldThickness, VolumeStep * GGlobalDistanceFieldHeightFieldThicknessScale);
 
 		HeightfieldDescriptionParameters.Set(RHICmdList, ShaderRHI, GetHeightfieldDescriptionsSRV(), NumHeightfieldsValue);
 		HeightfieldTextureParameters.Set(RHICmdList, ShaderRHI, HeightfieldTextureValue, NULL, VisibilityTextureValue);
@@ -679,6 +689,7 @@ public:
 		Ar << AOGlobalMaxSphereQueryRadius;
 		Ar << HeightfieldDescriptionParameters;
 		Ar << HeightfieldTextureParameters;
+		Ar << HeightfieldThickness;
 		return bShaderHasOutdatedParameters;
 	}
 
@@ -694,6 +705,7 @@ private:
 	FShaderParameter AOGlobalMaxSphereQueryRadius;
 	FHeightfieldDescriptionParameters HeightfieldDescriptionParameters;
 	FHeightfieldTextureParameters HeightfieldTextureParameters;
+	FShaderParameter HeightfieldThickness;
 };
 
 IMPLEMENT_SHADER_TYPE(, FCompositeHeightfieldsIntoGlobalDistanceFieldCS, TEXT("/Engine/Private/GlobalDistanceField.usf"), TEXT("CompositeHeightfieldsIntoGlobalDistanceFieldCS"), SF_Compute);
