Speculative fix for FLandscapeRenderSystem::FGetSectionLODBiasesTask may get out of current frame and crash

diff --git a/Engine/Source/Runtime/Renderer/Private/SceneRendering.cpp b/Engine/Source/Runtime/Renderer/Private/SceneRendering.cpp
--- a/Engine/Source/Runtime/Renderer/Private/SceneRendering.cpp
+++ b/Engine/Source/Runtime/Renderer/Private/SceneRendering.cpp
@@ -2614,13 +2614,6 @@ FSceneRenderer::~FSceneRenderer()
 	// Manually release references to TRefCountPtrs that are allocated on the mem stack, which doesn't call dtors
 	SortedShadowsForShadowDepthPass.Release();
 
-	extern TSet<IPersistentViewUniformBufferExtension*> PersistentViewUniformBufferExtensions;
-
-	for (IPersistentViewUniformBufferExtension* Extension : PersistentViewUniformBufferExtensions)
-	{
-		Extension->EndFrame();
-	}
-
 	Views.Empty();
 }
 
@@ -3479,6 +3472,14 @@ static void RenderViewFamily_RenderThread(FRHICommandListImmediate& RHICmdList,
 			SceneRenderer->Scene->DistanceFieldSceneData.PrimitiveModifiedBounds[CacheType].Reset();
 		}
 
+		// Immediately issue EndFrame() for all extensions in case any of the outstanding tasks they issued getting out of this frame
+		extern TSet<IPersistentViewUniformBufferExtension*> PersistentViewUniformBufferExtensions;
+
+		for (IPersistentViewUniformBufferExtension* Extension : PersistentViewUniformBufferExtensions)
+		{
+			Extension->EndFrame();
+		}
+
 #if STATS
 		{
 			QUICK_SCOPE_CYCLE_COUNTER(STAT_RenderViewFamily_RenderThread_MemStats);
