Added RT stats UpdatePrimitiveTransform, AddPrimitiveSceneInfos

#rnx

diff --git a/Engine/Source/Runtime/Renderer/Private/RendererScene.cpp b/Engine/Source/Runtime/Renderer/Private/RendererScene.cpp
--- a/Engine/Source/Runtime/Renderer/Private/RendererScene.cpp
+++ b/Engine/Source/Runtime/Renderer/Private/RendererScene.cpp
@@ -3629,6 +3629,7 @@ void FScene::UpdateAllPrimitiveSceneInfos(FRHICommandListImmediate& RHICmdList)
 		}
 	}
 	{
+		CSV_SCOPED_TIMING_STAT_EXCLUSIVE(AddPrimitiveSceneInfos);
 		SCOPED_NAMED_EVENT(FScene_AddPrimitiveSceneInfos, FColor::Green);
 		SCOPE_CYCLE_COUNTER(STAT_AddScenePrimitiveRenderThreadTime);
 		if (AddedLocalPrimitiveSceneInfos.Num())
@@ -3820,6 +3821,7 @@ void FScene::UpdateAllPrimitiveSceneInfos(FRHICommandListImmediate& RHICmdList)
 		}
 	}
 	{
+		CSV_SCOPED_TIMING_STAT_EXCLUSIVE(UpdatePrimitiveTransform);
 		SCOPED_NAMED_EVENT(FScene_AddPrimitiveSceneInfos, FColor::Yellow);
 		SCOPE_CYCLE_COUNTER(STAT_UpdatePrimitiveTransformRenderThreadTime);
 
