Forcing the PrevFrameReduction steps to use Float16 on Mac.
Writes to RG11B10 from compute sometimes fail so you get black. 16F appears to work correctly.

This is a gross hack because switching all RG11B10 textures to 16F at the RHI level causes rendering errors elsewhere (3 components vs 4 components maybe? Haven't done a deep dive) so I changed it at the high level to keep the change narrow.

diff --git a/Engine/Source/Runtime/Renderer/Private/ScreenSpaceRayTracing.cpp b/Engine/Source/Runtime/Renderer/Private/ScreenSpaceRayTracing.cpp
--- a/Engine/Source/Runtime/Renderer/Private/ScreenSpaceRayTracing.cpp
+++ b/Engine/Source/Runtime/Renderer/Private/ScreenSpaceRayTracing.cpp
@@ -886,7 +886,13 @@ void RenderScreenSpaceDiffuseIndirect(
 
 			FRDGTextureDesc Desc = FRDGTextureDesc::Create2DDesc(
 				FIntPoint(QuantizeMultiple * QuantizedSize.X, QuantizeMultiple * QuantizedSize.Y),
+				// HACK: This is a workaround to fix UE-84870.
+				// Ideally this would be fixed in MetalRHI but there's no 3 component 16f texture + swizzles aren't allowed on writable textures.
+#if PLATFORM_MAC
+				PF_FloatRGBA,
+#else
 				PF_FloatR11G11B10,
+#endif
 				FClearValueBinding::None,
 				/* InFlags = */ TexCreate_None,
 				/* InTargetableFlags = */ TexCreate_ShaderResource | TexCreate_UAV,
