Chery-picking fix for crash when the HairStrands plug-in is enabled. Original commit message:

Add registration of a unifrom buffer for the hair visibility pass in order to remove the validation error.

diff --git a/Engine/Source/Runtime/Renderer/Private/HairStrands/HairStrandsVisibility.cpp b/Engine/Source/Runtime/Renderer/Private/HairStrands/HairStrandsVisibility.cpp
--- a/Engine/Source/Runtime/Renderer/Private/HairStrands/HairStrandsVisibility.cpp
+++ b/Engine/Source/Runtime/Renderer/Private/HairStrands/HairStrandsVisibility.cpp
@@ -637,6 +637,7 @@ protected:
 	{
 		ERHIFeatureLevel::Type FeatureLevel = GetMaxSupportedFeatureLevel((EShaderPlatform)Initializer.Target.Platform);
 		check(FSceneInterface::GetShadingPath(FeatureLevel) != EShadingPath::Mobile);
+		PassUniformBuffer.Bind(Initializer.ParameterMap, FVisibilityPassGlobalParameters::StaticStructMetadata.GetShaderVariableName());
 	}
 
 	FHairVisibilityVS() {}
