Removed mask of TexCreate_RenderTarget mask from "DOFRecombine". It can be used as a render target in subsequent PostProcessMaterial passes.

diff --git a/Engine/Source/Runtime/Renderer/Private/PostProcess/DiaphragmDOF.cpp b/Engine/Source/Runtime/Renderer/Private/PostProcess/DiaphragmDOF.cpp
--- a/Engine/Source/Runtime/Renderer/Private/PostProcess/DiaphragmDOF.cpp
+++ b/Engine/Source/Runtime/Renderer/Private/PostProcess/DiaphragmDOF.cpp
@@ -2551,7 +2551,6 @@ FRDGTextureRef DiaphragmDOF::AddPasses(
 			FRDGTextureDesc Desc = InputSceneColor->Desc;
 			Desc.NumSamples = 1;
 			Desc.TargetableFlags |= TexCreate_UAV;
-			Desc.TargetableFlags &= ~TexCreate_RenderTargetable;
 			NewSceneColor = GraphBuilder.CreateTexture(Desc, TEXT("DOFRecombine"));
 		}
 		
