Use blending for PP materials only when alpha output is enabled

diff --git a/Engine/Source/Runtime/Renderer/Private/PostProcess/PostProcessMaterial.cpp b/Engine/Source/Runtime/Renderer/Private/PostProcess/PostProcessMaterial.cpp
--- a/Engine/Source/Runtime/Renderer/Private/PostProcess/PostProcessMaterial.cpp
+++ b/Engine/Source/Runtime/Renderer/Private/PostProcess/PostProcessMaterial.cpp
@@ -356,7 +356,12 @@ FScreenPassTexture AddPostProcessMaterialPass(
 	}
 
 	FRHIBlendState* DefaultBlendState = FScreenPassPipelineState::FDefaultBlendState::GetRHI();
-	FRHIBlendState* BlendState = GetMaterialBlendState(Material);
+	FRHIBlendState* BlendState = DefaultBlendState;
+	
+	if (IsMaterialBlendEnabled(Material))
+	{
+		BlendState = GetMaterialBlendState(Material);
+	}
 
 	// Blend / Depth Stencil usage requires that the render target have primed color data.
 	const bool bIsCompositeWithInput = DepthStencilState != DefaultDepthStencilState || BlendState != DefaultBlendState;
