Fix for Panini projection crash.

diff --git a/Engine/Source/Runtime/Renderer/Private/PostProcess/PostProcessUpscale.cpp b/Engine/Source/Runtime/Renderer/Private/PostProcess/PostProcessUpscale.cpp
--- a/Engine/Source/Runtime/Renderer/Private/PostProcess/PostProcessUpscale.cpp
+++ b/Engine/Source/Runtime/Renderer/Private/PostProcess/PostProcessUpscale.cpp
@@ -215,6 +215,7 @@ FScreenPassTexture AddUpscalePass(FRDGBuilder& GraphBuilder, const FViewInfo& Vi
 	PassParameters->PointSceneColorSampler = TStaticSamplerState<SF_Point, AM_Border, AM_Border, AM_Border>::GetRHI();
 	PassParameters->Panini = GetPaniniProjectionParameters(PaniniConfig, View);
 	PassParameters->UpscaleSoftness = FMath::Clamp(CVarUpscaleSoftness.GetValueOnRenderThread(), 0.0f, 1.0f);
+	PassParameters->View = View.ViewUniformBuffer;
 
 	FUpscalePS::FPermutationDomain PixelPermutationVector;
 	PixelPermutationVector.Set<FUpscalePS::FMethodDimension>(Inputs.Method);
