Extend distance field shadowing to also handle heightfields. This is another approach to shadow terrains. Compared to far shadow cascades, it showed lower RT and GPU overhead during internal testing.

diff --git a/Engine/Shaders/Private/DistanceFieldLightingShared.ush b/Engine/Shaders/Private/DistanceFieldLightingShared.ush
--- a/Engine/Shaders/Private/DistanceFieldLightingShared.ush
+++ b/Engine/Shaders/Private/DistanceFieldLightingShared.ush
@@ -24,15 +24,27 @@
 #define UPDATEOBJECTS_THREADGROUP_SIZE 1
 #endif
 
+#ifndef DISTANCEFIELD_PRIMITIVE_TYPE_DEFINED
+#define DISTANCEFIELD_PRIMITIVE_TYPE_DEFINED
+#define DFPT_SignedDistanceField 0
+#define DFPT_HeightField 1
+#endif
+
+#ifndef DISTANCEFIELD_PRIMITIVE_TYPE
+#define DISTANCEFIELD_PRIMITIVE_TYPE DFPT_SignedDistanceField
+#endif
+
 float3 DistanceFieldVolumePositionToUV(float3 VolumePosition, float3 UVScale, float3 UVAdd)
 {
 	float3 VolumeUV = VolumePosition * UVScale + UVAdd;
 	return VolumeUV;
 }
 
-Texture3D DistanceFieldTexture;
 SamplerState DistanceFieldSampler;
 
+#if (DISTANCEFIELD_PRIMITIVE_TYPE == DFPT_SignedDistanceField)
+Texture3D DistanceFieldTexture;
+
 float SampleMeshDistanceField(float3 VolumeUV, float2 ObjectDistanceFieldMAD)
 {
 	float TextureValue = Texture3DSampleLevel(DistanceFieldTexture, DistanceFieldSampler, VolumeUV, 0).x;
@@ -42,6 +54,25 @@ float SampleMeshDistanceField(float3 VolumeUV, float2 ObjectDistanceFieldMAD)
 	return TextureValue;
 }
 
+float SampleHeightFieldAtlas(float2 UV)
+{
+	return 0.0;
+}
+#else
+Texture2D DistanceFieldTexture;
+
+float SampleMeshDistanceField(float3 VolumeUV, float2 ObjectDistanceFieldMAD)
+{
+	return 0.0;
+}
+
+float SampleHeightFieldAtlas(float2 UV)
+{
+	float4 SampleValue = Texture2DSampleLevel(DistanceFieldTexture, DistanceFieldSampler, UV, 0);
+	return DecodePackedHeight(SampleValue.xy);
+}
+#endif
+
 float3 DistanceFieldAtlasTexelSize;
 
 RWBuffer<uint> RWObjectIndirectArguments;
@@ -54,7 +85,15 @@ uint GetCulledNumObjects()
 }
 
 // In float4's.  Must match equivalent C++ variables.
+#if (DISTANCEFIELD_PRIMITIVE_TYPE == DFPT_SignedDistanceField)
+#define OBJECT_BOUNDS_STRIDE 1
 #define OBJECT_DATA_STRIDE 17
+#elif (DISTANCEFIELD_PRIMITIVE_TYPE == DFPT_HeightField)
+#define OBJECT_BOUNDS_STRIDE 2
+#define OBJECT_DATA_STRIDE 5
+#else
+#error Unknown distance field primitive type
+#endif
 
 uint NumSceneObjects;
 
@@ -172,12 +211,19 @@ float4x4 LoadGlobalObjectOBBToWorld(uint ObjectIndex)
 }
 
 // In float4's.  Must match equivalent C++ variables.
+#if (DISTANCEFIELD_PRIMITIVE_TYPE == DFPT_SignedDistanceField)
 #define CULLED_OBJECT_DATA_STRIDE 17
 #define CULLED_OBJECT_BOX_BOUNDS_STRIDE 5
+#elif (DISTANCEFIELD_PRIMITIVE_TYPE == DFPT_HeightField)
+#define CULLED_OBJECT_DATA_STRIDE 5
+#define CULLED_OBJECT_BOX_BOUNDS_STRIDE 5
+#else
+#error Unknown distance field primitive type
+#endif
 
 float4 LoadObjectPositionAndRadiusFromBuffer(uint ObjectIndex, StructuredBuffer<float4> InBuffer)
 {
-	return InBuffer[ObjectIndex];
+	return InBuffer[ObjectIndex * OBJECT_BOUNDS_STRIDE];
 }
 
 float4x4 LoadObjectWorldToVolumeFromBuffer(uint ObjectIndex, StructuredBuffer<float4> InBuffer)
@@ -238,6 +284,31 @@ float4x4 LoadObjectOBBToWorldFromBuffer(uint ObjectIndex, StructuredBuffer<float
 	return transpose(float4x4(M0, M1, M2, float4(0.0f, 0.0f, 0.0f, 1.0f)));
 }
 
+void LoadHeightFieldObjectBoxBoundsFromBuffer(uint ObjectIndex, StructuredBuffer<float4> InBuffer, out float3 OutBoxOrigin, out float3 OutBoxExtent)
+{
+	OutBoxOrigin = InBuffer[ObjectIndex * OBJECT_BOUNDS_STRIDE + 0].xyz;
+	OutBoxExtent = InBuffer[ObjectIndex * OBJECT_BOUNDS_STRIDE + 1].xyz;
+}
+
+float4x4 LoadHeightFieldObjectWorldToLocalFromBuffer(uint ObjectIndex, StructuredBuffer<float4> InBuffer)
+{
+	float4 M0 = InBuffer[ObjectIndex * CULLED_OBJECT_DATA_STRIDE + 0];
+	float4 M1 = InBuffer[ObjectIndex * CULLED_OBJECT_DATA_STRIDE + 1];
+	float4 M2 = InBuffer[ObjectIndex * CULLED_OBJECT_DATA_STRIDE + 2];
+
+	return transpose(float4x4(M0, M1, M2, float4(0, 0, 0, 1)));
+}
+
+float4 LoadHeightFieldObjectRectFromBuffer(uint ObjectIndex, StructuredBuffer<float4> InBuffer)
+{
+	return InBuffer[ObjectIndex * CULLED_OBJECT_DATA_STRIDE + 3];
+}
+
+float4 LoadHeightFieldObjectAtlasUVScaleBiasFromBuffer(uint ObjectIndex, StructuredBuffer<float4> InBuffer)
+{
+	return InBuffer[ObjectIndex * CULLED_OBJECT_DATA_STRIDE + 4];
+}
+
 // These are structured buffers so they can be scalar memory loads on GCN
 StructuredBuffer<float4> CulledObjectBounds;
 StructuredBuffer<float4> CulledObjectData;
@@ -298,6 +369,26 @@ float4x4 LoadObjectOBBToWorld(uint ObjectIndex)
 	return LoadObjectOBBToWorldFromBuffer(ObjectIndex, CulledObjectData);
 }
 
+void LoadHeightFieldObjectBoxBounds(uint ObjectIndex, out float3 OutBoxOrigin, out float3 OutBoxExtent)
+{
+	LoadHeightFieldObjectBoxBoundsFromBuffer(ObjectIndex, CulledObjectBounds, OutBoxOrigin, OutBoxExtent);
+}
+
+float4x4 LoadHeightFieldObjectWorldToLocal(uint ObjectIndex)
+{
+	return LoadHeightFieldObjectWorldToLocalFromBuffer(ObjectIndex, CulledObjectData);
+}
+
+float4 LoadHeightFieldObjectAtlasUVScaleBias(uint ObjectIndex)
+{
+	return LoadHeightFieldObjectAtlasUVScaleBiasFromBuffer(ObjectIndex, CulledObjectData);
+}
+
+float4 LoadHeightFieldObjectRect(uint ObjectIndex)
+{
+	return LoadHeightFieldObjectRectFromBuffer(ObjectIndex, CulledObjectData);
+}
+
 // x = Offset in global buffer, y = NumLOD0, z = NumSurfels (all LODs), W = instance index
 uint4 LoadObjectSurfelCoordinate(uint ObjectIndex)
 {
diff --git a/Engine/Shaders/Private/DistanceFieldObjectCulling.usf b/Engine/Shaders/Private/DistanceFieldObjectCulling.usf
--- a/Engine/Shaders/Private/DistanceFieldObjectCulling.usf
+++ b/Engine/Shaders/Private/DistanceFieldObjectCulling.usf
@@ -15,7 +15,13 @@ Buffer<uint> UploadOperationIndices;
 Buffer<float4> UploadOperationData;
 
 // In float4's.  Must match equivalent C++ variables.
-#define UPLOAD_DATA_STRIDE (1 + OBJECT_DATA_STRIDE)
+#define UPLOAD_DATA_STRIDE (OBJECT_BOUNDS_STRIDE + OBJECT_DATA_STRIDE)
+
+void UploadBoundsFloat4(uint DestIndex, uint UploadIndex)
+{
+	float4 UploadVector = UploadOperationData[UploadIndex];
+	WriteFloat4(RWSceneObjectBounds, DestIndex, UploadVector);
+}
 
 void UploadDataFloat4(uint DestIndex, uint UploadIndex)
 {
@@ -34,14 +40,17 @@ void UploadObjectsToBufferCS(
 	if (DispatchThreadId.x < NumUploadOperations)
 	{
 		uint DestIndex = UploadOperationIndices[UploadOperationIndex];
-		float4 UploadBounds = UploadOperationData[UploadOperationIndex * UPLOAD_DATA_STRIDE + 0];
 
-		WriteFloat4(RWSceneObjectBounds, DestIndex, UploadBounds);
+		UNROLL
+		for (uint VectorIndex = 0; VectorIndex < OBJECT_BOUNDS_STRIDE; ++VectorIndex)
+		{
+			UploadBoundsFloat4(DestIndex * OBJECT_BOUNDS_STRIDE + VectorIndex, UploadOperationIndex * UPLOAD_DATA_STRIDE + VectorIndex);
+		}
 
 		UNROLL
-		for (uint VectorIndex = 0; VectorIndex < OBJECT_DATA_STRIDE; VectorIndex++)
+		for (uint VectorIndex = 0; VectorIndex < OBJECT_DATA_STRIDE; ++VectorIndex)
 		{
-			UploadDataFloat4(DestIndex * OBJECT_DATA_STRIDE + VectorIndex, UploadOperationIndex * UPLOAD_DATA_STRIDE + VectorIndex + 1);
+			UploadDataFloat4(DestIndex * OBJECT_DATA_STRIDE + VectorIndex, UploadOperationIndex * UPLOAD_DATA_STRIDE + OBJECT_BOUNDS_STRIDE + VectorIndex);
 		}
 	}
 }
@@ -54,6 +63,11 @@ RWBuffer<float> RWCopyObjectBounds;
 RWBuffer<float> RWCopyObjectData;
 #endif
 
+void CopyBoundsFloat4(uint DestIndex, uint SourceIndex)
+{
+	WriteFloat4(RWCopyObjectBounds, DestIndex, LoadFloat4(SceneObjectBounds, SourceIndex));
+}
+
 void CopyDataFloat4(uint DestIndex, uint SourceIndex)
 {
 	WriteFloat4(RWCopyObjectData, DestIndex, LoadFloat4(SceneObjectData, SourceIndex));
@@ -72,7 +86,11 @@ void CopyObjectBufferCS(
 		uint SourceIndex = CopyOperationIndex;
 		uint DestIndex = CopyOperationIndex;
 
-		WriteFloat4(RWCopyObjectBounds, DestIndex, LoadFloat4(SceneObjectBounds, SourceIndex));
+		UNROLL
+		for (uint VectorIndex = 0; VectorIndex < OBJECT_BOUNDS_STRIDE; ++VectorIndex)
+		{
+			CopyBoundsFloat4(DestIndex * OBJECT_BOUNDS_STRIDE + VectorIndex, SourceIndex * OBJECT_BOUNDS_STRIDE + VectorIndex);
+		}
 
 		UNROLL
 		for (uint VectorIndex = 0; VectorIndex < OBJECT_DATA_STRIDE; VectorIndex++)
@@ -102,6 +120,11 @@ Buffer<uint4> RemoveOperationIndices;
 	#define RWDataRemoveSource ObjectData2
 #endif
 
+void WriteBoundsFloat4(uint DestIndex, uint SourceIndex)
+{
+	WriteFloat4(RWSceneObjectBounds, DestIndex, LoadFloat4(RWBoundsRemoveSource, SourceIndex));
+}
+
 void WriteDataFloat4(uint DestIndex, uint SourceIndex)
 {
 	WriteFloat4(RWSceneObjectData, DestIndex, LoadFloat4(RWDataRemoveSource, SourceIndex));
@@ -121,7 +144,11 @@ void RemoveObjectsFromBufferCS(
 		uint SourceIndex = RemoveOperationIndices[RemoveOperationIndex].y;
 		uint DestIndex = RemoveOperationIndices[RemoveOperationIndex].x;
 
-		WriteFloat4(RWSceneObjectBounds, DestIndex, LoadFloat4(RWBoundsRemoveSource, SourceIndex));
+		UNROLL
+		for (uint VectorIndex = 0; VectorIndex < OBJECT_BOUNDS_STRIDE; ++VectorIndex)
+		{
+			WriteBoundsFloat4(DestIndex * OBJECT_BOUNDS_STRIDE + VectorIndex, SourceIndex * OBJECT_BOUNDS_STRIDE + VectorIndex);
+		}
 
 		UNROLL
 		for (uint VectorIndex = 0; VectorIndex < OBJECT_DATA_STRIDE; VectorIndex++)
diff --git a/Engine/Shaders/Private/DistanceFieldShadowing.usf b/Engine/Shaders/Private/DistanceFieldShadowing.usf
--- a/Engine/Shaders/Private/DistanceFieldShadowing.usf
+++ b/Engine/Shaders/Private/DistanceFieldShadowing.usf
@@ -13,7 +13,15 @@ uint ObjectBoundingGeometryIndexCount;
 
 void CopyCulledObjectData(uint DestIndex, uint SourceIndex)
 {
+#if (DISTANCEFIELD_PRIMITIVE_TYPE == DFPT_SignedDistanceField)
 	RWCulledObjectBounds[DestIndex] = LoadFloat4(SceneObjectBounds, SourceIndex);
+#else
+	float4 BoxOriginAndSphereRadius = LoadFloat4(SceneObjectBounds, SourceIndex * OBJECT_BOUNDS_STRIDE);
+	float4 BoxExtent = LoadFloat4(SceneObjectBounds, SourceIndex * OBJECT_BOUNDS_STRIDE + 1);
+
+	RWCulledObjectBounds[DestIndex * OBJECT_BOUNDS_STRIDE] = BoxOriginAndSphereRadius;
+	RWCulledObjectBounds[DestIndex * OBJECT_BOUNDS_STRIDE + 1] = BoxExtent;
+#endif
 
 	UNROLL
 	for (uint VectorIndex = 0; VectorIndex < CULLED_OBJECT_DATA_STRIDE; VectorIndex++)
@@ -24,10 +32,11 @@ void CopyCulledObjectData(uint DestIndex, uint SourceIndex)
 
 	float3 MinViewSpacePosition = float3(2000000, 2000000, 2000000);
 	float3 MaxViewSpacePosition = float3(-2000000, -2000000, -2000000);
+	float3 ViewSpaceBoundsVertices[8];
 
+#if (DISTANCEFIELD_PRIMITIVE_TYPE == DFPT_SignedDistanceField)
 	float4x4 ObjectOBBToWorld = LoadGlobalObjectOBBToWorld(SourceIndex);
-
-	float3 ViewSpaceBoundsVertices[8];
+#endif
 
 	for (uint i = 0; i < 8; i++)
 	{
@@ -36,9 +45,13 @@ void CopyCulledObjectData(uint DestIndex, uint SourceIndex)
 		UnitBoxVertex.y = i & 0x2 ? 1.0f : -1.0f;
 		UnitBoxVertex.z = i & 0x4 ? 1.0f : -1.0f;
 
+#if (DISTANCEFIELD_PRIMITIVE_TYPE == DFPT_SignedDistanceField)
 		UnitBoxVertex *= SDF_APPROX_MARGIN;
-
 		float3 WorldBoundsPosition = mul(float4(UnitBoxVertex, 1.0f), ObjectOBBToWorld).xyz;
+#else
+		float3 WorldBoundsPosition = UnitBoxVertex * BoxExtent.xyz + BoxOriginAndSphereRadius.xyz;
+#endif
+
 		float3 ViewSpacePosition = mul(float4(WorldBoundsPosition, 1.0f), WorldToShadow).xyz;
 		MinViewSpacePosition = min(MinViewSpacePosition, ViewSpacePosition);
 		MaxViewSpacePosition = max(MaxViewSpacePosition, ViewSpacePosition);
@@ -78,6 +91,23 @@ bool ShadowConvexHullIntersectSphere(float3 SphereOrigin, float SphereRadius)
 	return true;
 }
 
+bool ShadowConvexHullIntersectBox(float3 BoxOrigin, float3 BoxExtent)
+{
+	for (uint PlaneIndex = 0; PlaneIndex < NumShadowHullPlanes; ++PlaneIndex)
+	{
+		float4 PlaneData = ShadowConvexHull[PlaneIndex];
+		float PlaneDistance = dot(PlaneData.xyz, BoxOrigin) - PlaneData.w;
+		float PushOut = dot(abs(PlaneData.xyz), BoxExtent);
+
+		if (PlaneDistance > PushOut)
+		{
+			return false;
+		}
+	}
+
+	return true;
+}
+
 [numthreads(UPDATEOBJECTS_THREADGROUP_SIZE, 1, 1)]
 void CullObjectsForShadowCS(
 	uint3 GroupId : SV_GroupID,
@@ -102,6 +132,7 @@ void CullObjectsForShadowCS(
 
 	if (ObjectIndex < NumSceneObjects)
 	{
+#if (DISTANCEFIELD_PRIMITIVE_TYPE == DFPT_SignedDistanceField)
 		float4 ObjectBoundingSphere = LoadFloat4(SceneObjectBounds, SourceIndex);
         
         if (ShadowBoundingSphere.w == 0 && ShadowConvexHullIntersectSphere(ObjectBoundingSphere.xyz, ObjectBoundingSphere.w)
@@ -119,6 +150,17 @@ void CullObjectsForShadowCS(
                 CopyCulledObjectData(DestIndex, SourceIndex);
             }
         }
+#else
+		float3 BoxOrigin = LoadFloat4(SceneObjectBounds, SourceIndex * OBJECT_BOUNDS_STRIDE).xyz;
+		float3 BoxExtent = LoadFloat4(SceneObjectBounds, SourceIndex * OBJECT_BOUNDS_STRIDE + 1).xyz;
+
+		if (ShadowConvexHullIntersectBox(BoxOrigin, BoxExtent))
+		{
+			uint DestIndex;
+			InterlockedAdd(RWObjectIndirectArguments[1], 1U, DestIndex);
+			CopyCulledObjectData(DestIndex, SourceIndex);
+		}
+#endif
 	}
 
 #else
@@ -143,23 +185,50 @@ void CullObjectsForShadowCS(
 RWBuffer<uint> RWNextStartOffset;
 RWBuffer<uint> RWShadowTileStartOffsets;
 
+groupshared uint GroupNumIntersectingObjects;
+groupshared uint GroupStartOffset;
+
 #ifndef COMPUTE_START_OFFSET_GROUP_SIZE
 #define COMPUTE_START_OFFSET_GROUP_SIZE 1
 #endif
 
 [numthreads(COMPUTE_START_OFFSET_GROUP_SIZE, COMPUTE_START_OFFSET_GROUP_SIZE, 1)]
 void ComputeCulledTilesStartOffsetCS(
-	uint3 DispatchThreadId : SV_DispatchThreadID)
+	uint3 DispatchThreadId : SV_DispatchThreadID,
+	uint3 GroupThreadId : SV_GroupThreadID)
 {
+	bool bFirstThreadInGroup = all(GroupThreadId.xy == 0);
+
+	if (bFirstThreadInGroup)
+	{
+		GroupNumIntersectingObjects = 0;
+	}
+	
+	GroupMemoryBarrierWithGroupSync();
+
 	uint2 TileCoordinate = DispatchThreadId.xy;
+	bool bValidTile = all(TileCoordinate < ShadowTileListGroupSize);
+	uint TileIndex = TileCoordinate.y * ShadowTileListGroupSize.x + TileCoordinate.x;
+	uint TileStartOffset;
 
-	if (all(TileCoordinate < ShadowTileListGroupSize))
+	if (bValidTile)
 	{
-		uint TileIndex = DispatchThreadId.y * ShadowTileListGroupSize.x + DispatchThreadId.x;
 		uint NumIntersectingObjects = ShadowTileNumCulledObjects[TileIndex];
+		InterlockedAdd(GroupNumIntersectingObjects, NumIntersectingObjects, TileStartOffset);
+	}
+
+	GroupMemoryBarrierWithGroupSync();
 
-		uint TileStartOffset;
-		InterlockedAdd(RWNextStartOffset[0], NumIntersectingObjects, TileStartOffset);
+	if (bFirstThreadInGroup)
+	{
+		InterlockedAdd(RWNextStartOffset[0], GroupNumIntersectingObjects, GroupStartOffset);
+	}
+
+	GroupMemoryBarrierWithGroupSync();
+	
+	if (bValidTile)
+	{
+		TileStartOffset += GroupStartOffset;
 		RWShadowTileStartOffsets[TileIndex] = TileStartOffset;
 	}
 }
@@ -182,6 +251,7 @@ void ShadowObjectCullVS(
 	out float4 OutPosition : SV_POSITION
 )
 {
+#if (DISTANCEFIELD_PRIMITIVE_TYPE == DFPT_SignedDistanceField)
 	float4x4 ObjectOBBToWorld = LoadObjectOBBToWorld(ObjectIndex);
 
 	float ObjectScaleX = length(ObjectOBBToWorld[0].xyz);
@@ -193,7 +263,17 @@ void ShadowObjectCullVS(
 	ObjectOBBToWorld[2].xyz *= (SDF_APPROX_MARGIN + MinExpandRadius / ObjectScaleZ);
 
 	float3 WorldPosition = mul(float4(InPosition.xyz, 1.0f), ObjectOBBToWorld).xyz;
-	OutPosition = mul(float4(WorldPosition, 1.0f), WorldToShadow);
+#else
+	float3 BoxOrigin;
+	float3 BoxExtent;
+	LoadHeightFieldObjectBoxBounds(ObjectIndex, BoxOrigin, BoxExtent);
+
+	BoxExtent += MinExpandRadius;
+
+	float3 WorldPosition = InPosition.xyz * BoxExtent + BoxOrigin;
+#endif
+
+	OutPosition = mul(float4(WorldPosition, 1.0), WorldToShadow);
 
 	// Clamp the vertex to the near plane if it is in front of the near plane
 	if (OutPosition.z < 0)
@@ -479,6 +559,10 @@ float MinDepth;
 float MaxDepth;
 uint DownsampleFactor;
 
+float2 InvOutputBufferSize;
+Texture2D ShadowFactorsTexture;
+SamplerState ShadowFactorsSampler;
+
 [numthreads(THREADGROUP_SIZEX, THREADGROUP_SIZEY, 1)]
 void DistanceFieldShadowingCS(
 	uint3 GroupId : SV_GroupID,
@@ -526,46 +610,86 @@ void DistanceFieldShadowingCS(
 	BRANCH
 	if (bTileShouldComputeShadowing)
 	{
+		// Keeps result from going all the way sharp
+		float MinSphereRadius = .4f;
+		// Maintain reasonable culling bounds
+		float MaxSphereRadius = 100;
+
+		// Reduce shadowing when we are close the ray origin (only used for heightfield)
+		float SelfShadowFadeDistance = 100;
+		// Mitigate self-shadow caused by discontinuity on heightfield borders (only used for heightfield)
+		float SelfShadowVerticalBias = 100;
+
+#if HAS_PREVIOUS_OUTPUT
+	#if PLATFORM_SUPPORTS_TYPED_UAV_LOAD
+		float PrevResult = RWShadowFactors[DispatchThreadId.xy].x;
+	#else
+		float2 PrevResultUV = (DispatchThreadId.xy + 0.5) * InvOutputBufferSize;
+		float PrevResult = Texture2DSampleLevel(ShadowFactorsTexture, ShadowFactorsSampler, PrevResultUV, 0).x;
+	#endif
+#endif
+		bool bIsHeightField = DISTANCEFIELD_PRIMITIVE_TYPE == DFPT_HeightField;
+
+		if (bIsHeightField)
 		{
-			// World space offset along the start of the ray to avoid incorrect self-shadowing
-			float RayStartOffset = 2 + RayStartOffsetDepthScale * SceneDepth;
-			// Keeps result from going all the way sharp
-			float MinSphereRadius = .4f;
-			// Maintain reasonable culling bounds
-			float MaxSphereRadius = 100;
-
-			#if POINT_LIGHT
-
-				float3 LightVector = LightPositionAndInvRadius.xyz - OpaqueWorldPosition;
-				float LightVectorLength = length(LightVector);
-				float3 WorldRayStart = OpaqueWorldPosition + LightVector / LightVectorLength * RayStartOffset;
-				float3 WorldRayEnd = LightPositionAndInvRadius.xyz;
-				float MaxRayTime = LightVectorLength;
-				float MaxAngle = tan(10 * PI / 180.0f);
-				// Comparing tangents instead of angles, but tangent is always increasing in this range
-				float TanLightAngle = min(LightSourceRadius / LightVectorLength, MaxAngle);
-
-			#else
-
-				float3 WorldRayStart = OpaqueWorldPosition + LightDirection * RayStartOffset;
-				float3 WorldRayEnd = OpaqueWorldPosition + LightDirection * TraceDistance;
-				float MaxRayTime = TraceDistance;
-				float TanLightAngle = TanLightAngleAndNormalThreshold.x;
-
-			#endif
-
-			#if SCATTER_TILE_CULLING
-				bool bUseScatterTileCulling = true;
-			#else
-				bool bUseScatterTileCulling = false;
-			#endif
+			MaxSphereRadius = 1000;
+			OpaqueWorldPosition.z += SelfShadowVerticalBias;
+		}
+
+		// World space offset along the start of the ray to avoid incorrect self-shadowing
+		float RayStartOffset = 2 + RayStartOffsetDepthScale * SceneDepth;
+
+		#if POINT_LIGHT
+
+			float3 LightVector = LightPositionAndInvRadius.xyz - OpaqueWorldPosition;
+			float LightVectorLength = length(LightVector);
+			float3 WorldRayStart = OpaqueWorldPosition + LightVector / LightVectorLength * RayStartOffset;
+			float3 WorldRayEnd = LightPositionAndInvRadius.xyz;
+			float MaxRayTime = LightVectorLength;
+			float MaxAngle = tan(10 * PI / 180.0f);
+			// Comparing tangents instead of angles, but tangent is always increasing in this range
+			float TanLightAngle = min(LightSourceRadius / LightVectorLength, MaxAngle);
+
+		#else
+
+			float3 WorldRayStart = OpaqueWorldPosition + LightDirection * RayStartOffset;
+			float3 WorldRayEnd = OpaqueWorldPosition + LightDirection * TraceDistance;
+			float MaxRayTime = TraceDistance;
+			float TanLightAngle = TanLightAngleAndNormalThreshold.x;
+
+		#endif
+
+		#if SCATTER_TILE_CULLING
+			bool bUseScatterTileCulling = true;
+		#else
+			bool bUseScatterTileCulling = false;
+		#endif
 		
-			#if USE_CULLING
-				bool bUseCulling = true;
-			#else
-				bool bUseCulling = false;
-			#endif
+		#if USE_CULLING
+			bool bUseCulling = true;
+		#else
+			bool bUseCulling = false;
+		#endif
 
+		if (bIsHeightField)
+		{
+			Result = ShadowRayTraceThroughCulledHeightFieldObjects(
+				WorldRayStart,
+				WorldRayEnd,
+				TanLightAngle,
+				MaxSphereRadius,
+				SelfShadowFadeDistance,
+				CulledDataParameter,
+				NumIntersectingObjects,
+				bUseCulling,
+				bUseScatterTileCulling);
+
+#if HAS_PREVIOUS_OUTPUT
+			Result = min(Result, PrevResult);
+#endif
+		}
+		else
+		{
 			float SubsurfaceDensity = 0;
 			bool bUseSubsurfaceTransmission = false;
 
@@ -597,7 +721,6 @@ void DistanceFieldShadowingCS(
 				bUseSubsurfaceTransmission);
 		}
 	}
-
 #else
 	//Result = bTileShouldComputeShadowing;
 	Result = NumIntersectingObjects / 500.0f;
@@ -606,9 +729,6 @@ void DistanceFieldShadowingCS(
 	RWShadowFactors[DispatchThreadId.xy] = float2(Result, SceneDepth);
 }
 
-Texture2D ShadowFactorsTexture;
-SamplerState ShadowFactorsSampler;
-
 float FadePlaneOffset;
 float InvFadePlaneLength;
 float NearFadePlaneOffset;
diff --git a/Engine/Shaders/Private/DistanceFieldShadowingShared.ush b/Engine/Shaders/Private/DistanceFieldShadowingShared.ush
--- a/Engine/Shaders/Private/DistanceFieldShadowingShared.ush
+++ b/Engine/Shaders/Private/DistanceFieldShadowingShared.ush
@@ -194,4 +194,134 @@ float ShadowRayTraceThroughCulledObjects(
 
 	return MinConeVisibility;
 }
+
+float RayTraceHeightfieldLocal(
+	float3 LocalRayStart,
+	float3 LocalRayUnitDirection,
+	float StartRayTime,
+	float EndRayTime,
+	float4 AtlasUVScaleBias,
+	float TanLightAngle,
+	float LocalSelfShadowScale,
+	float LocalMaxSphereRadius)
+{
+#if DF_SHADOW_QUALITY == 3
+	float MaxSteps = 32;
+#elif DF_SHADOW_QUALITY == 2
+	float MaxSteps = 16;
+#else
+	float MaxSteps = 8;
+#endif
+
+	float StepRayTime = (EndRayTime - StartRayTime) / MaxSteps;
+	float SampleRayTime = StartRayTime;
+	float MinConeVisibility = 1.0;
+
+	for (float StepIndex = 0; StepIndex < MaxSteps; ++StepIndex)
+	{
+		float3 StepPosition = LocalRayStart + LocalRayUnitDirection * SampleRayTime;
+		float2 StepUV = StepPosition.xy * AtlasUVScaleBias.xy + AtlasUVScaleBias.zw;
+		float StepHeight = SampleHeightFieldAtlas(StepUV);
+
+		float SelfShadowVisibility = 1 - saturate(SampleRayTime * LocalSelfShadowScale);
+
+		float SphereRadius = clamp(TanLightAngle * SampleRayTime, 0.0001, LocalMaxSphereRadius);
+		float Distance = StepPosition.z - StepHeight;
+		float StepVisibility = max(saturate(Distance / SphereRadius), SelfShadowVisibility);
+
+		MinConeVisibility = min(MinConeVisibility, StepVisibility);
+		SampleRayTime += StepRayTime;
+
+		if (MinConeVisibility < 0.01)
+		{
+			break;
+		}
+	}
+
+	return MinConeVisibility;
+}
+
+float ShadowRayTraceThroughCulledHeightFieldObjects(
+	float3 WorldRayStart,
+	float3 WorldRayEnd,
+	float TanLightAngle,
+	float MaxSphereRadius,
+	float SelfShadowFadeDistance,
+	uint CulledDataParameter,
+	uint NumIntersectingObjects,
+	uniform bool bUseCulling,
+	uniform bool bUseScatterTileCulling)
+{
+	float MinConeVisibility = 1.0;
+
+	for (uint ListObjectIndex = 0; ListObjectIndex < NumIntersectingObjects; ++ListObjectIndex)
+	{
+		uint ObjectIndex;
+
+		if (bUseCulling)
+		{
+			if (bUseScatterTileCulling)
+			{
+				uint CulledDataStart = CulledDataParameter;
+				ObjectIndex = ShadowTileArrayData.Load((ListObjectIndex + CulledDataStart) * SHADOW_TILE_ARRAY_DATA_STRIDE);
+			}
+			else
+			{
+				uint GroupIndex = CulledDataParameter;
+				ObjectIndex = IntersectingObjectIndices[MAX_INTERSECTING_OBJECTS * GroupIndex + ListObjectIndex];
+			}
+		}
+		else
+		{
+			ObjectIndex = ListObjectIndex;
+		}
+
+		float4x4 WorldToLocal = LoadHeightFieldObjectWorldToLocal(ObjectIndex);
+		float2 HeightFieldSize = LoadHeightFieldObjectRect(ObjectIndex).zw;
+
+		float3 LocalRayStart = mul(float4(WorldRayStart, 1.0), WorldToLocal).xyz;
+		float3 LocalRayEnd = mul(float4(WorldRayEnd, 1.0), WorldToLocal).xyz;
+
+		float3 LocalBoundsMin = float3(0, 0, DecodePackedHeight(float2(0, 0)));
+		float3 LocalBoundsMax = float3(HeightFieldSize, DecodePackedHeight(float2(1, 1)));
+
+		float2 IntersectionNearFar = LineBoxIntersect(LocalRayStart, LocalRayEnd, LocalBoundsMin, LocalBoundsMax);
+		bool bValidIntersection = IntersectionNearFar.y > IntersectionNearFar.x;
+
+		if (bValidIntersection)
+		{
+			float4 AtlasUVScaleBias = LoadHeightFieldObjectAtlasUVScaleBias(ObjectIndex);
+
+			float3 LocalRayDirection = LocalRayEnd - LocalRayStart;
+			float LocalRayLength = length(LocalRayDirection);
+			float3 LocalRayUnitDirection = LocalRayDirection / LocalRayLength;
+			float StartRayTime = LocalRayLength * IntersectionNearFar.x;
+			float EndRayTime = LocalRayLength * IntersectionNearFar.y;
+
+			float WorldToLocalScale = length(WorldToLocal[0].xyz);
+			float LocalSelfShadowScale = 1.0 / (SelfShadowFadeDistance * WorldToLocalScale);
+			float LocalMaxSphereRadius = MaxSphereRadius * WorldToLocalScale;
+
+			float TempVisibility = RayTraceHeightfieldLocal(
+				LocalRayStart,
+				LocalRayUnitDirection,
+				StartRayTime,
+				EndRayTime,
+				AtlasUVScaleBias,
+				TanLightAngle,
+				LocalSelfShadowScale,
+				LocalMaxSphereRadius);
+
+			MinConeVisibility = min(MinConeVisibility, TempVisibility);
+
+			if (MinConeVisibility < 0.01)
+			{
+				MinConeVisibility = 0;
+				break;
+			}
+		}
+	}
+
+	return MinConeVisibility;
+}
 	
\ No newline at end of file
diff --git a/Engine/Shaders/Private/HeightFieldAtlasManagement.usf b/Engine/Shaders/Private/HeightFieldAtlasManagement.usf
new file mode 100644
--- /dev/null
+++ b/Engine/Shaders/Private/HeightFieldAtlasManagement.usf
@@ -0,0 +1,26 @@
+// Copyright Epic Games, Inc. All Rights Reserved.
+
+/*=============================================================================
+	HeightFieldAtlasManagement.usf
+=============================================================================*/
+
+uint4 UpdateRegionOffsetAndSize;
+float4 SourceScaleBias;
+uint SourceMipBias;
+
+Texture2D<float4> SourceTexture;
+SamplerState SourceTextureSampler;
+
+RWTexture2D<float2> RWHeightFieldAtlas;
+
+[numthreads(THREADGROUP_SIZEX, THREADGROUP_SIZEY, 1)]
+void UploadHeightFieldToAtlasCS(uint3 DispatchThreadId : SV_DispatchThreadID)
+{
+	if (all(DispatchThreadId.xy < UpdateRegionOffsetAndSize.zw))
+	{
+		float2 SourceUV = DispatchThreadId.xy * SourceScaleBias.xy + SourceScaleBias.zw;
+		uint2 WriteCoord = DispatchThreadId.xy + UpdateRegionOffsetAndSize.xy;
+
+		RWHeightFieldAtlas[WriteCoord] = SourceTexture.SampleLevel(SourceTextureSampler, SourceUV, SourceMipBias).xy;
+	}
+}
diff --git a/Engine/Source/Runtime/Engine/Private/Components/DirectionalLightComponent.cpp b/Engine/Source/Runtime/Engine/Private/Components/DirectionalLightComponent.cpp
--- a/Engine/Source/Runtime/Engine/Private/Components/DirectionalLightComponent.cpp
+++ b/Engine/Source/Runtime/Engine/Private/Components/DirectionalLightComponent.cpp
@@ -412,7 +412,9 @@ public:
 	virtual bool ShouldCreateRayTracedCascade(ERHIFeatureLevel::Type InFeatureLevel, bool bPrecomputedLightingIsValid, int32 MaxNearCascades) const override
 	{
 		static auto CVarDistanceFieldShadowing = IConsoleManager::Get().FindConsoleVariable(TEXT("r.DistanceFieldShadowing"));
-		if (CVarDistanceFieldShadowing != nullptr && CVarDistanceFieldShadowing->GetInt() == 0)
+		static IConsoleVariable* CVarHFShadowing = IConsoleManager::Get().FindConsoleVariable(TEXT("r.HeightFieldShadowing"));
+
+		if (CVarDistanceFieldShadowing != nullptr && CVarDistanceFieldShadowing->GetInt() == 0 && (!CVarHFShadowing || !CVarHFShadowing->GetInt()))
 		{
 			return false;
 		}
diff --git a/Engine/Source/Runtime/Engine/Private/DistanceFieldAtlas.cpp b/Engine/Source/Runtime/Engine/Private/DistanceFieldAtlas.cpp
--- a/Engine/Source/Runtime/Engine/Private/DistanceFieldAtlas.cpp
+++ b/Engine/Source/Runtime/Engine/Private/DistanceFieldAtlas.cpp
@@ -18,6 +18,8 @@
 #include "Misc/AutomationTest.h"
 #include "Async/ParallelFor.h"
 #include "DistanceFieldDownsampling.h"
+#include "GlobalShader.h"
+#include "RenderGraph.h"
 
 #if WITH_EDITOR
 #include "DerivedDataCacheInterface.h"
@@ -127,7 +129,22 @@ static FAutoConsoleVariableRef CVarDistanceFieldParallelAtlasUpdate(
 	TEXT("Whether to parallelize distance field data decompression and copying to upload buffer"),
 	ECVF_RenderThreadSafe);
 
+static int32 GHeightFieldAtlasTileSize = 64;
+static FAutoConsoleVariableRef CVarHeightFieldAtlasTileSize(
+	TEXT("r.HeightFields.AtlasTileSize"),
+	GHeightFieldAtlasTileSize,
+	TEXT(""),
+	ECVF_RenderThreadSafe);
+
+static int32 GHeightFieldAtlasDimInTiles = 16;
+static FAutoConsoleVariableRef CVarHeightFieldAtlasDimInTiles(
+	TEXT("r.HeightFields.AtlasDimInTiles"),
+	GHeightFieldAtlasDimInTiles,
+	TEXT(""),
+	ECVF_RenderThreadSafe);
+
 TGlobalResource<FDistanceFieldVolumeTextureAtlas> GDistanceFieldVolumeTextureAtlas = TGlobalResource<FDistanceFieldVolumeTextureAtlas>();
+TGlobalResource<FHeightFieldTextureAtlas> GHeightFieldTextureAtlas;
 
 FDistanceFieldVolumeTextureAtlas::FDistanceFieldVolumeTextureAtlas() :
 	BlockAllocator(0, 0, 0, 0, 0, 0, false, false),
@@ -1172,3 +1189,318 @@ void FDistanceFieldAsyncQueue::Shutdown()
 		FPlatformProcess::Sleep(0.01f);
 	}
 }
+
+void FHeightFieldTextureAtlas::InitializeIfNeeded()
+{
+	const uint32 LocalTileSize = (uint32)GHeightFieldAtlasTileSize;
+	const uint32 LocalDimInTiles = (uint32)GHeightFieldAtlasDimInTiles;
+
+	if (!AtlasTextureRHI
+		|| TileAllocator.TileSize != LocalTileSize
+		|| TileAllocator.DimInTiles != LocalDimInTiles)
+	{
+		TileAllocator.Init(LocalTileSize, 1, LocalDimInTiles);
+
+		for (TSet<FAllocation>::TConstIterator It(CurrentAllocations); It; ++It)
+		{
+			check(!PendingAllocations.Contains(*It));
+			PendingAllocations.Add(*It);
+		}
+
+		CurrentAllocations.Reset();
+
+		const uint32 SizeX = TileAllocator.DimInTexels;
+		const uint32 SizeY = TileAllocator.DimInTexels;
+		const uint32 Flags = TexCreate_ShaderResource | TexCreate_UAV;
+		FRHIResourceCreateInfo CreateInfo(TEXT("HeightFieldAtlas"));
+
+		AtlasTextureRHI = RHICreateTexture2D(SizeX, SizeY, PF_R8G8, 1, 1, Flags, CreateInfo);
+		AtlasUAVRHI = RHICreateUnorderedAccessView(AtlasTextureRHI, 0);
+
+		++Generation;
+	}
+}
+
+void FHeightFieldTextureAtlas::AddAllocation(const UTexture2D* Texture)
+{
+	check(Texture);
+
+	FAllocation* Found = CurrentAllocations.Find(Texture);
+	if (Found)
+	{
+		++Found->RefCount;
+		return;
+	}
+
+	Found = FailedAllocations.Find(Texture);
+	if (Found)
+	{
+		++Found->RefCount;
+		return;
+	}
+
+	Found = PendingAllocations.Find(Texture);
+	if (Found)
+	{
+		++Found->RefCount;
+	}
+	else
+	{
+		PendingAllocations.Add(Texture);
+	}
+}
+
+void FHeightFieldTextureAtlas::RemoveAllocation(const UTexture2D* Texture)
+{
+	FSetElementId Id = PendingAllocations.FindId(Texture);
+	if (Id.IsValidId())
+	{
+		check(PendingAllocations[Id].RefCount);
+		if (!--PendingAllocations[Id].RefCount)
+		{
+			PendingAllocations.Remove(Id);
+		}
+		return;
+	}
+
+	Id = FailedAllocations.FindId(Texture);
+	if (Id.IsValidId())
+	{
+		check(FailedAllocations[Id].RefCount);
+		if (!--FailedAllocations[Id].RefCount)
+		{
+			FailedAllocations.Remove(Id);
+		}
+		return;
+	}
+
+	Id = CurrentAllocations.FindId(Texture);
+	if (Id.IsValidId())
+	{
+		FAllocation& Allocation = CurrentAllocations[Id];
+		check(Allocation.RefCount && Allocation.TileIdx != INDEX_NONE);
+		if (!--Allocation.RefCount)
+		{
+			TileAllocator.FreeTile(Allocation.TileIdx);
+			CurrentAllocations.Remove(Id);
+		}
+	}
+}
+
+class FUploadHeightFieldToAtlasCS : public FGlobalShader
+{
+	DECLARE_GLOBAL_SHADER(FUploadHeightFieldToAtlasCS);
+	SHADER_USE_PARAMETER_STRUCT(FUploadHeightFieldToAtlasCS, FGlobalShader);
+
+	using FPermutationDomain = FShaderPermutationNone;
+
+	BEGIN_SHADER_PARAMETER_STRUCT(FParameters,)
+		SHADER_PARAMETER(FUintVector4, UpdateRegionOffsetAndSize)
+		SHADER_PARAMETER(FVector4, SourceScaleBias)
+		SHADER_PARAMETER(uint32, SourceMipBias)
+		SHADER_PARAMETER_TEXTURE(Texture2D, SourceTexture)
+		SHADER_PARAMETER_SAMPLER(SamplerState, SourceTextureSampler)
+		SHADER_PARAMETER_UAV(RWTexture2D<float2>, RWHeightFieldAtlas)
+	END_SHADER_PARAMETER_STRUCT()
+
+	static bool ShouldCompilePermutation(const FGlobalShaderPermutationParameters& Parameters)
+	{
+		return IsFeatureLevelSupported(Parameters.Platform, ERHIFeatureLevel::SM5) && DoesPlatformSupportDistanceFieldShadowing(Parameters.Platform);
+	}
+
+	static void ModifyCompilationEnvironment(const FGlobalShaderPermutationParameters& Parameters, FShaderCompilerEnvironment& OutEnvironment)
+	{
+		FGlobalShader::ModifyCompilationEnvironment(Parameters, OutEnvironment);
+		OutEnvironment.SetDefine(TEXT("THREADGROUP_SIZEX"), ThreadGroupSizeX);
+		OutEnvironment.SetDefine(TEXT("THREADGROUP_SIZEY"), ThreadGroupSizeY);
+	}
+
+	static constexpr uint32 ThreadGroupSizeX = 8;
+	static constexpr uint32 ThreadGroupSizeY = 8;
+};
+
+IMPLEMENT_GLOBAL_SHADER(FUploadHeightFieldToAtlasCS, "/Engine/Private/HeightFieldAtlasManagement.usf", "UploadHeightFieldToAtlasCS", SF_Compute);
+
+void FHeightFieldTextureAtlas::UpdateAllocations(FRHICommandListImmediate& RHICmdList, ERHIFeatureLevel::Type InFeatureLevel)
+{
+	InitializeIfNeeded();
+
+	TArray<uint32, TInlineAllocator<8>> UploadTileIndices;
+	TArray<FRHITexture*, TInlineAllocator<8>> UploadSourceTextures;
+
+	if (PendingAllocations.Num())
+	{
+		bool bAllocFailed = false;
+
+		for (TSet<FAllocation>::TConstIterator It(PendingAllocations); It; ++It)
+		{
+			if (!bAllocFailed)
+			{
+				const uint32 TileIdx = TileAllocator.AllocTile();
+				bAllocFailed = TileIdx == INDEX_NONE;
+				if (!bAllocFailed)
+				{
+					FAllocation TmpAllocation = *It;
+					TmpAllocation.TileIdx = TileIdx;
+					CurrentAllocations.Add(TmpAllocation);
+					UploadTileIndices.Add(TileIdx);
+					UploadSourceTextures.Add(TmpAllocation.SourceTexture->Resource->TextureRHI);
+				}
+			}
+			else
+			{
+				FailedAllocations.Add(*It);
+			}
+		}
+
+		PendingAllocations.Empty();
+	}
+
+	if (FailedAllocations.Num() && TileAllocator.CanAlloc())
+	{
+		for (TSet<FAllocation>::TIterator It(FailedAllocations); It; ++It)
+		{
+			const uint32 TileIdx = TileAllocator.AllocTile();
+			if (TileIdx != INDEX_NONE)
+			{
+				FAllocation TmpAllocation = *It;
+				TmpAllocation.TileIdx = TileIdx;
+				CurrentAllocations.Add(TmpAllocation);
+				UploadTileIndices.Add(TileIdx);
+				UploadSourceTextures.Add(TmpAllocation.SourceTexture->Resource->TextureRHI);
+				It.RemoveCurrent();
+			}
+			else
+			{
+				break;
+			}
+		}
+	}
+
+	if (UploadTileIndices.Num())
+	{
+		FRDGBuilder GraphBuilder(RHICmdList);
+
+		const uint32 TileSize = TileAllocator.TileSize;
+		const uint32 BorderSize = TileAllocator.BorderSize;
+		const uint32 TileSizeWithBorder = TileSize + BorderSize * 2;
+
+		const float SourceScale = 1.f / TileSize;
+		const float SourceBias = (.5f - BorderSize) * SourceScale;
+
+		check(TileSize && !(TileSize & (TileSize - 1)));
+		const uint32 TileLevels = FMath::CountBits(TileSize - 1) + 1;
+
+		const uint32 NumGroupsX = FMath::DivideAndRoundUp(TileSizeWithBorder, FUploadHeightFieldToAtlasCS::ThreadGroupSizeX);
+		const uint32 NumGroupsY = FMath::DivideAndRoundUp(TileSizeWithBorder, FUploadHeightFieldToAtlasCS::ThreadGroupSizeY);
+
+		TShaderMapRef<FUploadHeightFieldToAtlasCS> ComputeShader(GetGlobalShaderMap(InFeatureLevel));
+		FUploadHeightFieldToAtlasCS* ComputeShaderPtr = *ComputeShader;
+
+		RHICmdList.TransitionResource(EResourceTransitionAccess::ERWBarrier, EResourceTransitionPipeline::EComputeToCompute, AtlasUAVRHI);
+
+		for (int32 Idx = 0; Idx < UploadTileIndices.Num(); ++Idx)
+		{
+			const uint32 TileIdx = UploadTileIndices[Idx];
+			const FIntPoint TileCoord = TileAllocator.GetTileCoord(TileIdx);
+
+			FRHITexture2D* SourceTexture = (FRHITexture2D*)UploadSourceTextures[Idx];
+			const uint32 SourceLevels = SourceTexture->GetNumMips();
+			const uint32 SourceMipBias = SourceLevels > TileLevels ? SourceLevels - TileLevels : 0;
+
+			FUploadHeightFieldToAtlasCS::FParameters* Parameters = GraphBuilder.AllocParameters<FUploadHeightFieldToAtlasCS::FParameters>();
+			Parameters->UpdateRegionOffsetAndSize = FUintVector4(TileCoord.X * TileSizeWithBorder, TileCoord.Y * TileSizeWithBorder, TileSizeWithBorder, TileSizeWithBorder);
+			Parameters->SourceScaleBias = FVector4(SourceScale, SourceScale, SourceBias, SourceBias);
+			Parameters->SourceMipBias = SourceMipBias;
+			Parameters->SourceTexture = SourceTexture;
+			Parameters->SourceTextureSampler = TStaticSamplerState<SF_Bilinear>::GetRHI();
+			Parameters->RWHeightFieldAtlas = AtlasUAVRHI;
+
+			GraphBuilder.AddPass(
+				RDG_EVENT_NAME("UploadHeightFieldToAtlas"),
+				Parameters,
+				ERDGPassFlags::Compute,
+				[Parameters, ComputeShaderPtr, NumGroupsX, NumGroupsY](FRHICommandList& CmdList)
+			{
+				FComputeShaderUtils::Dispatch(CmdList, ComputeShaderPtr, *Parameters, FIntVector(NumGroupsX, NumGroupsY, 1));
+			});
+		}
+
+		GraphBuilder.Execute();
+		RHICmdList.TransitionResource(EResourceTransitionAccess::EReadable, EResourceTransitionPipeline::EComputeToGfx, AtlasUAVRHI);
+	}
+}
+
+uint32 FHeightFieldTextureAtlas::GetAllocationHandle(const UTexture2D* Texture) const
+{
+	const FAllocation* Allocation = CurrentAllocations.Find(Texture);
+	return Allocation ? Allocation->TileIdx : INDEX_NONE;
+}
+
+FVector4 FHeightFieldTextureAtlas::GetTileScaleBias(uint32 TileHandle) const
+{
+	const FIntPoint TileCoord = TileAllocator.GetTileCoord(TileHandle);
+	const float BiasX = (TileCoord.X * TileAllocator.TileSizeWithBorder + TileAllocator.BorderSize) * TileAllocator.TexelSize;
+	const float BiasY = (TileCoord.Y * TileAllocator.TileSizeWithBorder + TileAllocator.BorderSize) * TileAllocator.TexelSize;
+	return FVector4(TileAllocator.TileScale, TileAllocator.TileScale, BiasX, BiasY);
+}
+
+void FHeightFieldTextureAtlas::FSubAllocator::Init(uint32 InTileSize, uint32 InBorderSize, uint32 InDimInTiles)
+{
+	check(InDimInTiles && !(InDimInTiles & (InDimInTiles - 1)));
+
+	TileSize = InTileSize;
+	BorderSize = InBorderSize;
+	TileSizeWithBorder = InTileSize + 2 * InBorderSize;
+	DimInTiles = InDimInTiles;
+	DimInTilesShift = FMath::CountBits(InDimInTiles - 1);
+	DimInTilesMask = InDimInTiles - 1;
+	DimInTexels = InDimInTiles * TileSizeWithBorder;
+	MaxNumTiles = InDimInTiles * InDimInTiles;
+	TexelSize = 1.f / DimInTexels;
+	TileScale = TileSize * TexelSize;
+	NextFreeTileIdx = 0;
+	FreeTileIndices.Empty();
+}
+
+uint32 FHeightFieldTextureAtlas::FSubAllocator::AllocTile()
+{
+	uint32 Ret = INDEX_NONE;
+	if (FreeTileIndices.Num())
+	{
+		Ret = FreeTileIndices.Pop();
+	}
+	else if (NextFreeTileIdx < MaxNumTiles)
+	{
+		Ret = NextFreeTileIdx++;
+	}
+	return Ret;
+}
+
+void FHeightFieldTextureAtlas::FSubAllocator::FreeTile(uint32 TileIdx)
+{
+	check(TileIdx < NextFreeTileIdx);
+	FreeTileIndices.Add(TileIdx);
+}
+
+bool FHeightFieldTextureAtlas::FSubAllocator::CanAlloc() const
+{
+	return FreeTileIndices.Num() || NextFreeTileIdx < MaxNumTiles;
+}
+
+FIntPoint FHeightFieldTextureAtlas::FSubAllocator::GetTileCoord(uint32 TileIdx) const
+{
+	return FIntPoint(TileIdx & DimInTilesMask, TileIdx >> DimInTilesShift);
+}
+
+FHeightFieldTextureAtlas::FAllocation::FAllocation()
+	: SourceTexture(nullptr)
+	, RefCount(0)
+	, TileIdx(INDEX_NONE)
+{}
+
+FHeightFieldTextureAtlas::FAllocation::FAllocation(const UTexture2D* InTexture)
+	: SourceTexture(InTexture)
+	, RefCount(1)
+	, TileIdx(INDEX_NONE)
+{}
diff --git a/Engine/Source/Runtime/Engine/Public/DistanceFieldAtlas.h b/Engine/Source/Runtime/Engine/Public/DistanceFieldAtlas.h
--- a/Engine/Source/Runtime/Engine/Public/DistanceFieldAtlas.h
+++ b/Engine/Source/Runtime/Engine/Public/DistanceFieldAtlas.h
@@ -18,6 +18,7 @@
 
 class FDistanceFieldVolumeData;
 class UStaticMesh;
+class UTexture2D;
 
 template <class T> class TLockFreePointerListLIFO;
 
@@ -153,6 +154,109 @@ private:
 
 extern ENGINE_API TGlobalResource<FDistanceFieldVolumeTextureAtlas> GDistanceFieldVolumeTextureAtlas;
 
+class ENGINE_API FHeightFieldTextureAtlas : public FRenderResource
+{
+public:
+	void InitializeIfNeeded();
+
+	void AddAllocation(const UTexture2D* Texture);
+
+	void RemoveAllocation(const UTexture2D* Texture);
+
+	void UpdateAllocations(FRHICommandListImmediate& RHICmdList, ERHIFeatureLevel::Type InFeatureLevel);
+
+	uint32 GetAllocationHandle(const UTexture2D* Texture) const;
+
+	FVector4 GetTileScaleBias(uint32 TileHandle) const;
+
+	FRHITexture2D* GetAtlasTexture() const
+	{
+		return AtlasTextureRHI;
+	}
+
+	uint32 GetSizeX() const
+	{
+		return TileAllocator.DimInTexels;
+	}
+
+	uint32 GetSizeY() const
+	{
+		return TileAllocator.DimInTexels;
+	}
+
+	uint32 GetGeneration() const
+	{
+		return Generation;
+	}
+
+private:
+	class FSubAllocator
+	{
+	public:
+		void Init(uint32 InTileSize, uint32 InBorderSize, uint32 InDimInTiles);
+
+		uint32 AllocTile();
+
+		void FreeTile(uint32 TileIdx);
+
+		bool CanAlloc() const;
+
+		FIntPoint GetTileCoord(uint32 TileIdx) const;
+
+	private:
+		uint32 TileSize;
+		uint32 BorderSize;
+		uint32 TileSizeWithBorder;
+		uint32 DimInTiles;
+		uint32 DimInTilesShift;
+		uint32 DimInTilesMask;
+		uint32 DimInTexels;
+		uint32 MaxNumTiles;
+
+		float TexelSize;
+		float TileScale;
+
+		uint32 NextFreeTileIdx;
+		TArray<uint32> FreeTileIndices;
+
+		friend class FHeightFieldTextureAtlas;
+	};
+
+	struct FAllocation
+	{
+		const UTexture2D* SourceTexture;
+		uint32 RefCount;
+		uint32 TileIdx;
+
+		FAllocation();
+
+		FAllocation(const UTexture2D* InTexture);
+
+		bool operator==(const FAllocation& Other) const
+		{
+			return SourceTexture == Other.SourceTexture;
+		}
+
+		friend uint32 GetTypeHash(const FAllocation& Key)
+		{
+			return GetTypeHash(Key.SourceTexture);
+		}
+	};
+
+	FSubAllocator TileAllocator;
+
+	TSet<FAllocation> PendingAllocations;
+	TSet<FAllocation> FailedAllocations;
+	TSet<FAllocation> CurrentAllocations;
+
+	FTexture2DRHIRef AtlasTextureRHI;
+	FUnorderedAccessViewRHIRef AtlasUAVRHI;
+
+	uint32 Generation;
+};
+
+extern ENGINE_API TGlobalResource<FHeightFieldTextureAtlas> GHeightFieldTextureAtlas;
+
 /** Distance field data payload and output of the mesh build process. */
 class ENGINE_API FDistanceFieldVolumeData : public FDeferredCleanupInterface
 {
diff --git a/Engine/Source/Runtime/Renderer/Private/CapsuleShadowRendering.cpp b/Engine/Source/Runtime/Renderer/Private/CapsuleShadowRendering.cpp
--- a/Engine/Source/Runtime/Renderer/Private/CapsuleShadowRendering.cpp
+++ b/Engine/Source/Runtime/Renderer/Private/CapsuleShadowRendering.cpp
@@ -426,7 +426,25 @@ public:
 
 		if (Scene->DistanceFieldSceneData.GetCurrentObjectBuffers())
 		{
-			DistanceFieldObjectParameters.Set(RHICmdList, ShaderRHI, *Scene->DistanceFieldSceneData.GetCurrentObjectBuffers(), Scene->DistanceFieldSceneData.NumObjectsInBuffer);
+			FRHITexture* TextureAtlas;
+			int32 AtlasSizeX;
+			int32 AtlasSizeY;
+			int32 AtlasSizeZ;
+
+			TextureAtlas = GDistanceFieldVolumeTextureAtlas.VolumeTextureRHI;
+			AtlasSizeX = GDistanceFieldVolumeTextureAtlas.GetSizeX();
+			AtlasSizeY = GDistanceFieldVolumeTextureAtlas.GetSizeY();
+			AtlasSizeZ = GDistanceFieldVolumeTextureAtlas.GetSizeZ();
+
+			DistanceFieldObjectParameters.Set(
+				RHICmdList,
+				ShaderRHI,
+				*Scene->DistanceFieldSceneData.GetCurrentObjectBuffers(),
+				Scene->DistanceFieldSceneData.NumObjectsInBuffer,
+				TextureAtlas,
+				AtlasSizeX,
+				AtlasSizeY,
+				AtlasSizeZ);
 		}
 		else
 		{
@@ -505,7 +523,7 @@ private:
 	FShaderParameter CosFadeStartAngle;
 	FShaderResourceParameter LightDirectionData;
 	FShaderParameter IndirectCapsuleSelfShadowingIntensity;
-	FDistanceFieldObjectBufferParameters DistanceFieldObjectParameters;
+	TDistanceFieldObjectBufferParameters<DFPT_SignedDistanceField> DistanceFieldObjectParameters;
 };
 
 template<ECapsuleShadowingType ShadowingType, EIndirectShadowingPrimitiveTypes PrimitiveTypes>
diff --git a/Engine/Source/Runtime/Renderer/Private/DeferredShadingRenderer.cpp b/Engine/Source/Runtime/Renderer/Private/DeferredShadingRenderer.cpp
--- a/Engine/Source/Runtime/Renderer/Private/DeferredShadingRenderer.cpp
+++ b/Engine/Source/Runtime/Renderer/Private/DeferredShadingRenderer.cpp
@@ -552,10 +552,17 @@ static int32 GetCustomDepthPassLocation()
 
 void FDeferredShadingSceneRenderer::PrepareDistanceFieldScene(FRHICommandListImmediate& RHICmdList, bool bSplitDispatch)
 {
+	CSV_SCOPED_TIMING_STAT_EXCLUSIVE(RenderDFAO);
+	SCOPE_CYCLE_COUNTER(STAT_FDeferredShadingSceneRenderer_DistanceFieldAO_Init);
+
+	if (ShouldPrepareHeightFieldScene())
+	{
+		GHeightFieldTextureAtlas.UpdateAllocations(RHICmdList, FeatureLevel);
+		UpdateGlobalHeightFieldObjectBuffers(RHICmdList);
+	}
+
 	if (ShouldPrepareDistanceFieldScene())
 	{
-		CSV_SCOPED_TIMING_STAT_EXCLUSIVE(RenderDFAO);
-		SCOPE_CYCLE_COUNTER(STAT_FDeferredShadingSceneRenderer_DistanceFieldAO_Init);
 		GDistanceFieldVolumeTextureAtlas.UpdateAllocations(RHICmdList, FeatureLevel);
 		UpdateGlobalDistanceFieldObjectBuffers(RHICmdList);
 		if (bSplitDispatch)
diff --git a/Engine/Source/Runtime/Renderer/Private/DeferredShadingRenderer.h b/Engine/Source/Runtime/Renderer/Private/DeferredShadingRenderer.h
--- a/Engine/Source/Runtime/Renderer/Private/DeferredShadingRenderer.h
+++ b/Engine/Source/Runtime/Renderer/Private/DeferredShadingRenderer.h
@@ -493,8 +493,10 @@ private:
 
 	bool ShouldPrepareDistanceFieldScene() const;
 	bool ShouldPrepareGlobalDistanceField() const;
+	bool ShouldPrepareHeightFieldScene() const;
 
 	void UpdateGlobalDistanceFieldObjectBuffers(FRHICommandListImmediate& RHICmdList);
+	void UpdateGlobalHeightFieldObjectBuffers(FRHICommandListImmediate& RHICmdList);
 	void PrepareDistanceFieldScene(FRHICommandListImmediate& RHICmdList, bool bSplitDispatch);
 
 	void RenderViewTranslucency(FRHICommandListImmediate& RHICmdList, const FViewInfo& View, const FMeshPassProcessorRenderState& DrawRenderState, ETranslucencyPass::Type TranslucenyPass);
diff --git a/Engine/Source/Runtime/Renderer/Private/DistanceFieldGlobalIllumination.cpp b/Engine/Source/Runtime/Renderer/Private/DistanceFieldGlobalIllumination.cpp
--- a/Engine/Source/Runtime/Renderer/Private/DistanceFieldGlobalIllumination.cpp
+++ b/Engine/Source/Runtime/Renderer/Private/DistanceFieldGlobalIllumination.cpp
@@ -93,6 +93,7 @@ TGlobalResource<FVPLResources> GVPLResources;
 TGlobalResource<FVPLResources> GCulledVPLResources;
 
 extern TGlobalResource<FDistanceFieldObjectBufferResource> GShadowCulledObjectBuffers;
+extern TGlobalResource<FHeightFieldObjectBufferResource> GShadowCulledHeightFieldObjectBuffers;
 
 class FVPLPlacementCS : public FGlobalShader
 {
@@ -141,7 +142,18 @@ public:
 	{
 		FRHIComputeShader* ShaderRHI = GetComputeShader();
 		FGlobalShader::SetParameters<FViewUniformShaderParameters>(RHICmdList, ShaderRHI, View.ViewUniformBuffer);
-		ObjectParameters.Set(RHICmdList, ShaderRHI, GShadowCulledObjectBuffers.Buffers);
+
+		FRHITexture* TextureAtlas;
+		int32 AtlasSizeX;
+		int32 AtlasSizeY;
+		int32 AtlasSizeZ;
+
+		TextureAtlas = GDistanceFieldVolumeTextureAtlas.VolumeTextureRHI;
+		AtlasSizeX = GDistanceFieldVolumeTextureAtlas.GetSizeX();
+		AtlasSizeY = GDistanceFieldVolumeTextureAtlas.GetSizeY();
+		AtlasSizeZ = GDistanceFieldVolumeTextureAtlas.GetSizeZ();
+
+		ObjectParameters.Set(RHICmdList, ShaderRHI, GShadowCulledObjectBuffers.Buffers, TextureAtlas, AtlasSizeX, AtlasSizeY, AtlasSizeZ);
 
 		SetShaderValue(RHICmdList, ShaderRHI, InvPlacementGridSize, InvPlacementGridSizeValue);
 		SetShaderValue(RHICmdList, ShaderRHI, WorldToShadow, WorldToShadowValue);
@@ -204,7 +216,7 @@ private:
 	FShaderParameter ShadowToWorld;
 	FShaderParameter LightDirectionAndTraceDistance;
 	FShaderParameter LightColor;
-	FDistanceFieldCulledObjectBufferParameters ObjectParameters;
+	TDistanceFieldCulledObjectBufferParameters<DFPT_SignedDistanceField> ObjectParameters;
 	FLightTileIntersectionParameters LightTileIntersectionParameters;
 	FShaderParameter VPLPlacementCameraRadius;
 };
@@ -580,7 +592,18 @@ public:
 
 		RHICmdList.TransitionResource(EResourceTransitionAccess::ERWBarrier, EResourceTransitionPipeline::EComputeToCompute, GAOCulledObjectBuffers.Buffers.ObjectIndirectDispatch.UAV);
 		DispatchParameters.SetBuffer(RHICmdList, ShaderRHI, GAOCulledObjectBuffers.Buffers.ObjectIndirectDispatch);
-		ObjectParameters.Set(RHICmdList, ShaderRHI, GAOCulledObjectBuffers.Buffers);
+
+		FRHITexture* TextureAtlas;
+		int32 AtlasSizeX;
+		int32 AtlasSizeY;
+		int32 AtlasSizeZ;
+
+		TextureAtlas = GDistanceFieldVolumeTextureAtlas.VolumeTextureRHI;
+		AtlasSizeX = GDistanceFieldVolumeTextureAtlas.GetSizeX();
+		AtlasSizeY = GDistanceFieldVolumeTextureAtlas.GetSizeY();
+		AtlasSizeZ = GDistanceFieldVolumeTextureAtlas.GetSizeZ();
+
+		ObjectParameters.Set(RHICmdList, ShaderRHI, GAOCulledObjectBuffers.Buffers, TextureAtlas, AtlasSizeX, AtlasSizeY, AtlasSizeZ);
 
 		SetShaderValue(RHICmdList, ShaderRHI, ObjectProcessStride, GVPLSpreadUpdateOver);
 	}
@@ -606,7 +629,7 @@ public:
 private:
 
 	FRWShaderParameter DispatchParameters;
-	FDistanceFieldCulledObjectBufferParameters ObjectParameters;
+	TDistanceFieldCulledObjectBufferParameters<DFPT_SignedDistanceField> ObjectParameters;
 	FShaderParameter ObjectProcessStride;
 };
 
@@ -667,7 +690,17 @@ public:
 
 		const FScene* Scene = (const FScene*)View.Family->Scene;
 
-		ObjectParameters.Set(RHICmdList, ShaderRHI, GAOCulledObjectBuffers.Buffers);
+		FRHITexture* TextureAtlas;
+		int32 AtlasSizeX;
+		int32 AtlasSizeY;
+		int32 AtlasSizeZ;
+
+		TextureAtlas = GDistanceFieldVolumeTextureAtlas.VolumeTextureRHI;
+		AtlasSizeX = GDistanceFieldVolumeTextureAtlas.GetSizeX();
+		AtlasSizeY = GDistanceFieldVolumeTextureAtlas.GetSizeY();
+		AtlasSizeZ = GDistanceFieldVolumeTextureAtlas.GetSizeZ();
+
+		ObjectParameters.Set(RHICmdList, ShaderRHI, GAOCulledObjectBuffers.Buffers, TextureAtlas, AtlasSizeX, AtlasSizeY, AtlasSizeZ);
 		SurfelParameters.Set(RHICmdList, ShaderRHI, *Scene->DistanceFieldSceneData.SurfelBuffers, *Scene->DistanceFieldSceneData.InstancedSurfelBuffers);
 
 		FLightShaderParameters LightParameters;
@@ -736,7 +769,7 @@ private:
 	FShaderParameter LightSourceRadius;
 	FShaderParameter TanLightAngleAndNormalThreshold;
 	FShaderParameter LightColor;
-	FDistanceFieldCulledObjectBufferParameters ObjectParameters;
+	TDistanceFieldCulledObjectBufferParameters<DFPT_SignedDistanceField> ObjectParameters;
 	FSurfelBufferParameters SurfelParameters;
 	FLightTileIntersectionParameters LightTileIntersectionParameters;
 	FShaderParameter WorldToShadow;
@@ -1020,7 +1053,18 @@ public:
 		SceneTextureParameters.Set(RHICmdList, ShaderRHI, View.FeatureLevel, ESceneTextureSetupMode::All);
 
 		extern TGlobalResource<FDistanceFieldObjectBufferResource> GAOCulledObjectBuffers;
-		ObjectParameters.Set(RHICmdList, ShaderRHI, GAOCulledObjectBuffers.Buffers);
+		FRHITexture* TextureAtlas;
+		int32 AtlasSizeX;
+		int32 AtlasSizeY;
+		int32 AtlasSizeZ;
+
+		TextureAtlas = GDistanceFieldVolumeTextureAtlas.VolumeTextureRHI;
+		AtlasSizeX = GDistanceFieldVolumeTextureAtlas.GetSizeX();
+		AtlasSizeY = GDistanceFieldVolumeTextureAtlas.GetSizeY();
+		AtlasSizeZ = GDistanceFieldVolumeTextureAtlas.GetSizeZ();
+
+		ObjectParameters.Set(RHICmdList, ShaderRHI, GAOCulledObjectBuffers.Buffers, TextureAtlas, AtlasSizeX, AtlasSizeY, AtlasSizeZ);
+
 		AOParameters.Set(RHICmdList, ShaderRHI, Parameters);
 		ScreenGridParameters.Set(RHICmdList, ShaderRHI, View, DistanceFieldNormal);
 
@@ -1069,7 +1113,7 @@ public:
 private:
 
 	FSceneTextureShaderParameters SceneTextureParameters;
-	FDistanceFieldCulledObjectBufferParameters ObjectParameters;
+	TDistanceFieldCulledObjectBufferParameters<DFPT_SignedDistanceField> ObjectParameters;
 	FAOParameters AOParameters;
 	FScreenGridParameters ScreenGridParameters;
 	FSurfelBufferParameters SurfelParameters;
diff --git a/Engine/Source/Runtime/Renderer/Private/DistanceFieldLightingShared.h b/Engine/Source/Runtime/Renderer/Private/DistanceFieldLightingShared.h
--- a/Engine/Source/Runtime/Renderer/Private/DistanceFieldLightingShared.h
+++ b/Engine/Source/Runtime/Renderer/Private/DistanceFieldLightingShared.h
@@ -29,6 +29,7 @@ DECLARE_LOG_CATEGORY_EXTERN(LogDistanceField, Warning, All);
 extern int32 GDistanceFieldAOTileSizeX;
 extern int32 GDistanceFieldAOTileSizeY;
 extern int32 GAverageObjectsPerShadowCullTile;
+extern int32 GAverageHeightFieldObjectsPerShadowCullTile;
 
 extern int32 GDistanceFieldGI;
 
@@ -43,7 +44,15 @@ extern bool IsDistanceFieldGIAllowed(const FViewInfo& View);
 extern bool UseDistanceFieldAO();
 extern bool UseAOObjectDistanceField();
 
-class FDistanceFieldObjectBuffers
+enum EDistanceFieldPrimitiveType
+{
+	DFPT_SignedDistanceField,
+	DFPT_HeightField,
+	DFPT_Num
+};
+
+template <EDistanceFieldPrimitiveType PrimitiveType>
+class TDistanceFieldObjectBuffers
 {
 public:
 
@@ -55,7 +64,7 @@ public:
 	FRWBuffer Bounds;
 	FRWBuffer Data;
 
-	FDistanceFieldObjectBuffers()
+	TDistanceFieldObjectBuffers()
 	{
 		MaxObjects = 0;
 	}
@@ -74,6 +83,9 @@ public:
 	}
 };
 
+class FDistanceFieldObjectBuffers : public TDistanceFieldObjectBuffers<DFPT_SignedDistanceField> {};
+class FHeightFieldObjectBuffers : public TDistanceFieldObjectBuffers<DFPT_HeightField> {};
+
 class FSurfelBuffers
 {
 public:
@@ -135,7 +147,8 @@ public:
 	FRWBuffer VPLFlux;
 };
 
-class FDistanceFieldObjectBufferParameters
+template <EDistanceFieldPrimitiveType PrimitiveType>
+class TDistanceFieldObjectBufferParameters
 {
 public:
 	void Bind(const FShaderParameterMap& ParameterMap)
@@ -149,7 +162,16 @@ public:
 	}
 
 	template<typename TParamRef>
-	void Set(FRHICommandList& RHICmdList, const TParamRef& ShaderRHI, const FDistanceFieldObjectBuffers& ObjectBuffers, int32 NumObjectsValue, bool bBarrier = false)
+	void Set(
+		FRHICommandList& RHICmdList,
+		const TParamRef& ShaderRHI,
+		const TDistanceFieldObjectBuffers<PrimitiveType>& ObjectBuffers,
+		int32 NumObjectsValue,
+		FRHITexture* TextureAtlas,
+		int32 AtlasSizeX,
+		int32 AtlasSizeY,
+		int32 AtlasSizeZ,
+		bool bBarrier = false)
 	{
 		if (bBarrier)
 		{
@@ -169,18 +191,17 @@ public:
 			DistanceFieldTexture,
 			DistanceFieldSampler,
 			TStaticSamplerState<SF_Bilinear,AM_Clamp,AM_Clamp,AM_Clamp>::GetRHI(),
-			GDistanceFieldVolumeTextureAtlas.VolumeTextureRHI
-			);
+			TextureAtlas);
 
-		const int32 NumTexelsOneDimX = GDistanceFieldVolumeTextureAtlas.GetSizeX();
-		const int32 NumTexelsOneDimY = GDistanceFieldVolumeTextureAtlas.GetSizeY();
-		const int32 NumTexelsOneDimZ = GDistanceFieldVolumeTextureAtlas.GetSizeZ();
+		const int32 NumTexelsOneDimX = AtlasSizeX;
+		const int32 NumTexelsOneDimY = AtlasSizeY;
+		const int32 NumTexelsOneDimZ = AtlasSizeZ;
 		const FVector InvTextureDim(1.0f / NumTexelsOneDimX, 1.0f / NumTexelsOneDimY, 1.0f / NumTexelsOneDimZ);
 		SetShaderValue(RHICmdList, ShaderRHI, DistanceFieldAtlasTexelSize, InvTextureDim);
 	}
 
 	template<typename TParamRef>
-	void UnsetParameters(FRHICommandList& RHICmdList, const TParamRef& ShaderRHI, const FDistanceFieldObjectBuffers& ObjectBuffers, bool bBarrier = false)
+	void UnsetParameters(FRHICommandList& RHICmdList, const TParamRef& ShaderRHI, const TDistanceFieldObjectBuffers<PrimitiveType>& ObjectBuffers, bool bBarrier = false)
 	{
 		SceneObjectBounds.UnsetUAV(RHICmdList, ShaderRHI);
 		SceneObjectData.UnsetUAV(RHICmdList, ShaderRHI);
@@ -194,7 +215,7 @@ public:
 		}
 	}
 
-	friend FArchive& operator<<(FArchive& Ar, FDistanceFieldObjectBufferParameters& P)
+	friend FArchive& operator<<(FArchive& Ar, TDistanceFieldObjectBufferParameters& P)
 	{
 		Ar << P.SceneObjectBounds << P.SceneObjectData << P.NumSceneObjects << P.DistanceFieldTexture << P.DistanceFieldSampler << P.DistanceFieldAtlasTexelSize;
 		return Ar;
@@ -252,7 +273,8 @@ private:
 	FRWShaderParameter VPLFlux;
 };
 
-class FDistanceFieldCulledObjectBuffers
+template <EDistanceFieldPrimitiveType PrimitiveType>
+class TDistanceFieldCulledObjectBuffers
 {
 public:
 
@@ -268,7 +290,7 @@ public:
 	FRWBufferStructured Data;
 	FRWBufferStructured BoxBounds;
 
-	FDistanceFieldCulledObjectBuffers()
+	TDistanceFieldCulledObjectBuffers()
 	{
 		MaxObjects = 0;
 		bWantBoxBounds = false;
@@ -278,16 +300,43 @@ public:
 	{
 		if (MaxObjects > 0)
 		{
+			const TCHAR* ObjectIndirectArgumentsDebugName;
+			const TCHAR* ObjectIndirectDispatchDebugName;
+			const TCHAR* BoundsDebugName;
+			const TCHAR* DataDebugName;
+			const TCHAR* BoxBoundsDebugName;
+			uint32 BoundsNumElements;
+
+			if (PrimitiveType == DFPT_HeightField)
+			{
+				ObjectIndirectArgumentsDebugName = TEXT("FHeightFieldCulledObjectBuffers_ObjectIndirectArguments");
+				ObjectIndirectDispatchDebugName = TEXT("FHeightFieldCulledObjectBuffers_ObjectIndirectDispatch");
+				BoundsDebugName = TEXT("FHeightFieldCulledObjectBuffers_Bounds");
+				DataDebugName = TEXT("FHeightFieldCulledObjectBuffers_Data");
+				BoxBoundsDebugName = TEXT("FHeightFieldCulledObjectBuffers_BoxBounds");
+				BoundsNumElements = MaxObjects * 2;
+			}
+			else
+			{
+				check(PrimitiveType == DFPT_SignedDistanceField);
+				ObjectIndirectArgumentsDebugName = TEXT("FDistanceFieldCulledObjectBuffers_ObjectIndirectArguments");
+				ObjectIndirectDispatchDebugName = TEXT("FDistanceFieldCulledObjectBuffers_ObjectIndirectDispatch");
+				BoundsDebugName = TEXT("FDistanceFieldCulledObjectBuffers_Bounds");
+				DataDebugName = TEXT("FDistanceFieldCulledObjectBuffers_Data");
+				BoxBoundsDebugName = TEXT("FDistanceFieldCulledObjectBuffers_BoxBounds");
+				BoundsNumElements = MaxObjects;
+			}
+
 			const uint32 FastVRamFlag = GFastVRamConfig.DistanceFieldCulledObjectBuffers | ( IsTransientResourceBufferAliasingEnabled() ? BUF_Transient : BUF_None );
 
-			ObjectIndirectArguments.Initialize(sizeof(uint32), 5, PF_R32_UINT, BUF_Static | BUF_DrawIndirect, TEXT("FDistanceFieldCulledObjectBuffers::ObjectIndirectArguments"));
-			ObjectIndirectDispatch.Initialize(sizeof(uint32), 3, PF_R32_UINT, BUF_Static | BUF_DrawIndirect, TEXT("FDistanceFieldCulledObjectBuffers::ObjectIndirectDispatch"));
-			Bounds.Initialize(sizeof(FVector4), MaxObjects, BUF_Static | FastVRamFlag, TEXT("FDistanceFieldCulledObjectBuffers::Bounds"));
-			Data.Initialize(sizeof(FVector4), MaxObjects * ObjectDataStride, BUF_Static | FastVRamFlag, TEXT("FDistanceFieldCulledObjectBuffers::Data"));
+			ObjectIndirectArguments.Initialize(sizeof(uint32), 5, PF_R32_UINT, BUF_Static | BUF_DrawIndirect, ObjectIndirectArgumentsDebugName);
+			ObjectIndirectDispatch.Initialize(sizeof(uint32), 3, PF_R32_UINT, BUF_Static | BUF_DrawIndirect, ObjectIndirectDispatchDebugName);
+			Bounds.Initialize(sizeof(FVector4), BoundsNumElements, BUF_Static | FastVRamFlag, BoundsDebugName);
+			Data.Initialize(sizeof(FVector4), MaxObjects * ObjectDataStride, BUF_Static | FastVRamFlag, DataDebugName);
 
 			if (bWantBoxBounds)
 			{
-				BoxBounds.Initialize(sizeof(FVector4), MaxObjects * ObjectBoxBoundsStride, BUF_Static | FastVRamFlag, TEXT("FDistanceFieldCulledObjectBuffers::BoxBounds"));
+				BoxBounds.Initialize(sizeof(FVector4), MaxObjects * ObjectBoxBoundsStride, BUF_Static | FastVRamFlag, BoxBoundsDebugName);
 			}
 		}
 	}
@@ -327,10 +376,14 @@ public:
 	}
 };
 
-class FDistanceFieldObjectBufferResource : public FRenderResource
+class FDistanceFieldCulledObjectBuffers : public TDistanceFieldCulledObjectBuffers<DFPT_SignedDistanceField> {};
+class FHeightFieldCulledObjectBuffers : public TDistanceFieldCulledObjectBuffers<DFPT_HeightField> {};
+
+template <EDistanceFieldPrimitiveType PrimitiveType>
+class TDistanceFieldObjectBufferResource : public FRenderResource
 {
 public:
-	FDistanceFieldCulledObjectBuffers Buffers;
+	typename TChooseClass<PrimitiveType == DFPT_HeightField, FHeightFieldCulledObjectBuffers, FDistanceFieldCulledObjectBuffers>::Result Buffers;
 
 	virtual void InitDynamicRHI()  override
 	{
@@ -343,7 +396,11 @@ public:
 	}
 };
 
-class FDistanceFieldCulledObjectBufferParameters
+class FDistanceFieldObjectBufferResource : public TDistanceFieldObjectBufferResource<DFPT_SignedDistanceField> {};
+class FHeightFieldObjectBufferResource : public TDistanceFieldObjectBufferResource<DFPT_HeightField> {};
+
+template <EDistanceFieldPrimitiveType PrimitiveType>
+class TDistanceFieldCulledObjectBufferParameters
 {
 public:
 	void Bind(const FShaderParameterMap& ParameterMap)
@@ -358,7 +415,14 @@ public:
 	}
 
 	template<typename TShaderRHI, typename TRHICommandList>
-	void Set(TRHICommandList& RHICmdList, TShaderRHI* ShaderRHI, const FDistanceFieldCulledObjectBuffers& ObjectBuffers)
+	void Set(
+		TRHICommandList& RHICmdList,
+		TShaderRHI* ShaderRHI,
+		const TDistanceFieldCulledObjectBuffers<PrimitiveType>& ObjectBuffers,
+		FRHITexture* TextureAtlas,
+		int32 AtlasSizeX,
+		int32 AtlasSizeY,
+		int32 AtlasSizeZ)
 	{
 		ObjectIndirectArguments.SetBuffer(RHICmdList, ShaderRHI, ObjectBuffers.ObjectIndirectArguments);
 		CulledObjectBounds.SetBuffer(RHICmdList, ShaderRHI, ObjectBuffers.Bounds);
@@ -376,12 +440,12 @@ public:
 			DistanceFieldTexture,
 			DistanceFieldSampler,
 			TStaticSamplerState<SF_Bilinear,AM_Clamp,AM_Clamp,AM_Clamp>::GetRHI(),
-			GDistanceFieldVolumeTextureAtlas.VolumeTextureRHI
+			TextureAtlas
 			);
 
-		const int32 NumTexelsOneDimX = GDistanceFieldVolumeTextureAtlas.GetSizeX();
-		const int32 NumTexelsOneDimY = GDistanceFieldVolumeTextureAtlas.GetSizeY();
-		const int32 NumTexelsOneDimZ = GDistanceFieldVolumeTextureAtlas.GetSizeZ();
+		const int32 NumTexelsOneDimX = AtlasSizeX;
+		const int32 NumTexelsOneDimY = AtlasSizeY;
+		const int32 NumTexelsOneDimZ = AtlasSizeZ;
 		const FVector InvTextureDim(1.0f / NumTexelsOneDimX, 1.0f / NumTexelsOneDimY, 1.0f / NumTexelsOneDimZ);
 		SetShaderValue(RHICmdList, ShaderRHI, DistanceFieldAtlasTexelSize, InvTextureDim);
 	}
@@ -395,7 +459,7 @@ public:
 		CulledObjectBoxBounds.UnsetUAV(RHICmdList, ShaderRHI);
 	}
 
-	void GetUAVs(const FDistanceFieldCulledObjectBuffers& ObjectBuffers, TArray<FRHIUnorderedAccessView*>& UAVs)
+	void GetUAVs(const TDistanceFieldCulledObjectBuffers<PrimitiveType>& ObjectBuffers, TArray<FRHIUnorderedAccessView*>& UAVs)
 	{
 		uint32 MaxIndex = 0;
 		MaxIndex = FMath::Max(MaxIndex, ObjectIndirectArguments.GetUAVIndex());
@@ -428,7 +492,7 @@ public:
 		check(UAVs.Num() > 0);
 	}
 
-	friend FArchive& operator<<(FArchive& Ar, FDistanceFieldCulledObjectBufferParameters& P)
+	friend FArchive& operator<<(FArchive& Ar, TDistanceFieldCulledObjectBufferParameters<PrimitiveType>& P)
 	{
 		Ar << P.ObjectIndirectArguments << P.CulledObjectBounds << P.CulledObjectData << P.CulledObjectBoxBounds << P.DistanceFieldTexture << P.DistanceFieldSampler << P.DistanceFieldAtlasTexelSize;
 		return Ar;
@@ -501,7 +565,7 @@ public:
 		, b16BitIndices(false)
 	{}
 
-	void Initialize()
+	void Initialize(int32 MaxNumObjectsPerTile)
 	{
 		TileAlignedDimensions = FIntPoint(Align(TileDimensions.X, 64), Align(TileDimensions.Y, 64));
 			
@@ -510,7 +574,7 @@ public:
 		NextStartOffset.Initialize(sizeof(uint32), 1, PF_R32_UINT, BUF_Static);
 
 		//@todo - handle max exceeded
-		TileArrayData.Initialize(b16BitIndices ? sizeof(uint16) : sizeof(uint32), GAverageObjectsPerShadowCullTile * TileAlignedDimensions.X * TileAlignedDimensions.Y * LightTileDataStride, b16BitIndices ? PF_R16_UINT : PF_R32_UINT, BUF_Static);
+		TileArrayData.Initialize(b16BitIndices ? sizeof(uint16) : sizeof(uint32), MaxNumObjectsPerTile * TileAlignedDimensions.X * TileAlignedDimensions.Y * LightTileDataStride, b16BitIndices ? PF_R16_UINT : PF_R32_UINT, BUF_Static);
 	}
 
 	void Release()
@@ -655,6 +719,17 @@ extern void CullDistanceFieldObjectsForLight(
 	float ShadowBoundingRadius,
 	TUniquePtr<class FLightTileIntersectionResources>& TileIntersectionResources);
 
+extern void CullHeightFieldObjectsForLight(
+	FRHICommandListImmediate& RHICmdList,
+	const FViewInfo& View,
+	const FLightSceneProxy* LightSceneProxy,
+	const FMatrix& WorldToShadowValue,
+	int32 NumPlanes,
+	const FPlane* PlaneData,
+	const FVector4& ShadowBoundingSphereValue,
+	float ShadowBoundingRadius,
+	TUniquePtr<class FLightTileIntersectionResources>& TileIntersectionResources);
+
 class FUniformMeshBuffers
 {
 public:
@@ -742,3 +817,25 @@ public:
 extern TGlobalResource<FDistanceFieldObjectBufferResource> GAOCulledObjectBuffers;
 
 extern bool SupportsDistanceFieldAO(ERHIFeatureLevel::Type FeatureLevel, EShaderPlatform ShaderPlatform);
+
+template <bool bIsAType> struct TSelector {};
+
+template <>
+struct TSelector<true>
+{
+	template <typename AType, typename BType>
+	AType& operator()(AType& A, BType& B)
+	{
+		return A;
+	}
+};
+
+template <>
+struct TSelector<false>
+{
+	template <typename AType, typename BType>
+	BType& operator()(AType& A, BType& B)
+	{
+		return B;
+	}
+};
\ No newline at end of file
diff --git a/Engine/Source/Runtime/Renderer/Private/DistanceFieldObjectCulling.cpp b/Engine/Source/Runtime/Renderer/Private/DistanceFieldObjectCulling.cpp
--- a/Engine/Source/Runtime/Renderer/Private/DistanceFieldObjectCulling.cpp
+++ b/Engine/Source/Runtime/Renderer/Private/DistanceFieldObjectCulling.cpp
@@ -147,8 +147,20 @@ public:
 
 		FRHIComputeShader* ShaderRHI = GetComputeShader();
 		FGlobalShader::SetParameters<FViewUniformShaderParameters>(RHICmdList, ShaderRHI, View.ViewUniformBuffer);
-		ObjectBufferParameters.Set(RHICmdList, ShaderRHI, *(Scene->DistanceFieldSceneData.GetCurrentObjectBuffers()), Scene->DistanceFieldSceneData.NumObjectsInBuffer);
-		CulledObjectParameters.Set(RHICmdList, ShaderRHI, GAOCulledObjectBuffers.Buffers);
+
+		FRHITexture* TextureAtlas;
+		int32 AtlasSizeX;
+		int32 AtlasSizeY;
+		int32 AtlasSizeZ;
+
+		TextureAtlas = GDistanceFieldVolumeTextureAtlas.VolumeTextureRHI;
+		AtlasSizeX = GDistanceFieldVolumeTextureAtlas.GetSizeX();
+		AtlasSizeY = GDistanceFieldVolumeTextureAtlas.GetSizeY();
+		AtlasSizeZ = GDistanceFieldVolumeTextureAtlas.GetSizeZ();
+
+		ObjectBufferParameters.Set(RHICmdList, ShaderRHI, *(Scene->DistanceFieldSceneData.GetCurrentObjectBuffers()), Scene->DistanceFieldSceneData.NumObjectsInBuffer, TextureAtlas, AtlasSizeX, AtlasSizeY, AtlasSizeZ);
+		CulledObjectParameters.Set(RHICmdList, ShaderRHI, GAOCulledObjectBuffers.Buffers, TextureAtlas, AtlasSizeX, AtlasSizeY, AtlasSizeZ);
+
 		AOParameters.Set(RHICmdList, ShaderRHI, Parameters);
 
 		// Shader assumes max 6
@@ -187,8 +199,8 @@ public:
 
 private:
 
-	FDistanceFieldObjectBufferParameters ObjectBufferParameters;
-	FDistanceFieldCulledObjectBufferParameters CulledObjectParameters;
+	TDistanceFieldObjectBufferParameters<DFPT_SignedDistanceField> ObjectBufferParameters;
+	TDistanceFieldCulledObjectBufferParameters<DFPT_SignedDistanceField> CulledObjectParameters;
 	FAOParameters AOParameters;
 	FShaderParameter NumConvexHullPlanes;
 	FShaderParameter ViewFrustumConvexHull;
@@ -360,7 +372,17 @@ public:
 		FRHIVertexShader* ShaderRHI = GetVertexShader();
 		FGlobalShader::SetParameters<FViewUniformShaderParameters>(RHICmdList, ShaderRHI, View.ViewUniformBuffer);
 		
-		ObjectParameters.Set(RHICmdList, ShaderRHI, GAOCulledObjectBuffers.Buffers);
+		FRHITexture* TextureAtlas;
+		int32 AtlasSizeX;
+		int32 AtlasSizeY;
+		int32 AtlasSizeZ;
+
+		TextureAtlas = GDistanceFieldVolumeTextureAtlas.VolumeTextureRHI;
+		AtlasSizeX = GDistanceFieldVolumeTextureAtlas.GetSizeX();
+		AtlasSizeY = GDistanceFieldVolumeTextureAtlas.GetSizeY();
+		AtlasSizeZ = GDistanceFieldVolumeTextureAtlas.GetSizeZ();
+
+		ObjectParameters.Set(RHICmdList, ShaderRHI, GAOCulledObjectBuffers.Buffers, TextureAtlas, AtlasSizeX, AtlasSizeY, AtlasSizeZ);
 		AOParameters.Set(RHICmdList, ShaderRHI, Parameters);
 
 		const int32 NumRings = StencilingGeometry::GLowPolyStencilSphereVertexBuffer.GetNumRings();
@@ -382,7 +404,7 @@ public:
 	}
 
 private:
-	FDistanceFieldCulledObjectBufferParameters ObjectParameters;
+	TDistanceFieldCulledObjectBufferParameters<DFPT_SignedDistanceField> ObjectParameters;
 	FAOParameters AOParameters;
 	FShaderParameter ConservativeRadiusScale;
 };
@@ -428,7 +450,18 @@ public:
 	{
 		FRHIPixelShader* ShaderRHI = GetPixelShader();
 		FGlobalShader::SetParameters<FViewUniformShaderParameters>(RHICmdList, ShaderRHI, View.ViewUniformBuffer);
-		ObjectParameters.Set(RHICmdList, ShaderRHI, GAOCulledObjectBuffers.Buffers);
+
+		FRHITexture* TextureAtlas;
+		int32 AtlasSizeX;
+		int32 AtlasSizeY;
+		int32 AtlasSizeZ;
+
+		TextureAtlas = GDistanceFieldVolumeTextureAtlas.VolumeTextureRHI;
+		AtlasSizeX = GDistanceFieldVolumeTextureAtlas.GetSizeX();
+		AtlasSizeY = GDistanceFieldVolumeTextureAtlas.GetSizeY();
+		AtlasSizeZ = GDistanceFieldVolumeTextureAtlas.GetSizeZ();
+
+		ObjectParameters.Set(RHICmdList, ShaderRHI, GAOCulledObjectBuffers.Buffers, TextureAtlas, AtlasSizeX, AtlasSizeY, AtlasSizeZ);
 		AOParameters.Set(RHICmdList, ShaderRHI, Parameters);
 
 		FTileIntersectionResources* TileIntersectionResources = ((FSceneViewState*)View.State)->AOTileIntersectionResources;
@@ -461,7 +494,7 @@ public:
 	}
 
 private:
-	FDistanceFieldCulledObjectBufferParameters ObjectParameters;
+	TDistanceFieldCulledObjectBufferParameters<DFPT_SignedDistanceField> ObjectParameters;
 	FAOParameters AOParameters;
 	FTileIntersectionParameters TileIntersectionParameters;
 	FShaderResourceParameter TileConeAxisAndCos;
@@ -509,7 +542,17 @@ public:
 
 		FGlobalShader::SetParameters<FViewUniformShaderParameters>(RHICmdList, ShaderRHI, View.ViewUniformBuffer);
 
-		ObjectParameters.Set(RHICmdList, ShaderRHI, GAOCulledObjectBuffers.Buffers);
+		FRHITexture* TextureAtlas;
+		int32 AtlasSizeX;
+		int32 AtlasSizeY;
+		int32 AtlasSizeZ;
+
+		TextureAtlas = GDistanceFieldVolumeTextureAtlas.VolumeTextureRHI;
+		AtlasSizeX = GDistanceFieldVolumeTextureAtlas.GetSizeX();
+		AtlasSizeY = GDistanceFieldVolumeTextureAtlas.GetSizeY();
+		AtlasSizeZ = GDistanceFieldVolumeTextureAtlas.GetSizeZ();
+
+		ObjectParameters.Set(RHICmdList, ShaderRHI, GAOCulledObjectBuffers.Buffers, TextureAtlas, AtlasSizeX, AtlasSizeY, AtlasSizeZ);
 
 		FTileIntersectionResources* TileIntersectionResources = ((FSceneViewState*)View.State)->AOTileIntersectionResources;
 		
@@ -543,7 +586,7 @@ public:
 
 private:
 
-	FDistanceFieldCulledObjectBufferParameters ObjectParameters;
+	TDistanceFieldCulledObjectBufferParameters<DFPT_SignedDistanceField> ObjectParameters;
 	FTileIntersectionParameters TileIntersectionParameters;
 };
 
diff --git a/Engine/Source/Runtime/Renderer/Private/DistanceFieldObjectManagement.cpp b/Engine/Source/Runtime/Renderer/Private/DistanceFieldObjectManagement.cpp
--- a/Engine/Source/Runtime/Renderer/Private/DistanceFieldObjectManagement.cpp
+++ b/Engine/Source/Runtime/Renderer/Private/DistanceFieldObjectManagement.cpp
@@ -36,14 +36,19 @@ FAutoConsoleVariableRef CVarAOLogObjectBufferReallocation(
 	);
 
 // Must match equivalent shader defines
-int32 FDistanceFieldObjectBuffers::ObjectDataStride = 17;
-int32 FDistanceFieldCulledObjectBuffers::ObjectDataStride = 17;
-int32 FDistanceFieldCulledObjectBuffers::ObjectBoxBoundsStride = 5;
+int32 TDistanceFieldObjectBuffers<DFPT_SignedDistanceField>::ObjectDataStride = 17;
+int32 TDistanceFieldObjectBuffers<DFPT_HeightField>::ObjectDataStride = 5;
+int32 TDistanceFieldCulledObjectBuffers<DFPT_SignedDistanceField>::ObjectDataStride = 17;
+int32 TDistanceFieldCulledObjectBuffers<DFPT_SignedDistanceField>::ObjectBoxBoundsStride = 5;
+int32 TDistanceFieldCulledObjectBuffers<DFPT_HeightField>::ObjectDataStride = 5;
+int32 TDistanceFieldCulledObjectBuffers<DFPT_HeightField>::ObjectBoxBoundsStride = 5;
 
 // In float4's.  Must match corresponding usf definition
 int32 UploadObjectDataStride = 1 + FDistanceFieldObjectBuffers::ObjectDataStride;
+int32 UploadHeightFieldObjectDataStride = 2 + FHeightFieldObjectBuffers::ObjectDataStride;
 
-void FDistanceFieldObjectBuffers::Initialize()
+template <EDistanceFieldPrimitiveType PrimitiveType>
+void TDistanceFieldObjectBuffers<PrimitiveType>::Initialize()
 {
 	if (MaxObjects > 0)
 	{
@@ -58,25 +63,44 @@ void FDistanceFieldObjectBuffers::Initialize()
 			BufferFormat = PF_A32B32G32R32F;
 		}
 
-		Bounds.Initialize(GPixelFormats[BufferFormat].BlockBytes, NumComponents * MaxObjects, BufferFormat, 0, TEXT("FDistanceFieldObjectBuffers::Bounds"));
-		Data.Initialize(GPixelFormats[BufferFormat].BlockBytes, NumComponents * MaxObjects * ObjectDataStride, BufferFormat, 0, TEXT("FDistanceFieldObjectBuffers::Data"));
+		uint32 BoundsNumElements;
+		const TCHAR* BoundsDebugName;
+		const TCHAR* DataDebugName;
+
+		if (PrimitiveType == DFPT_HeightField)
+		{
+			BoundsNumElements = NumComponents * 2 * MaxObjects;
+			BoundsDebugName = TEXT("FHeightFieldObjectBuffers_Bounds");
+			DataDebugName = TEXT("FHeightFieldObjectBuffers_Data");
+		}
+		else
+		{
+			check(PrimitiveType == DFPT_SignedDistanceField);
+			BoundsNumElements = NumComponents * MaxObjects;
+			BoundsDebugName = TEXT("FDistanceFieldObjectBuffers_Bounds");
+			DataDebugName = TEXT("FDistanceFieldObjectBuffers_Data");
+		}
+
+		Bounds.Initialize(GPixelFormats[BufferFormat].BlockBytes, BoundsNumElements, BufferFormat, 0, BoundsDebugName);
+		Data.Initialize(GPixelFormats[BufferFormat].BlockBytes, NumComponents * MaxObjects * ObjectDataStride, BufferFormat, 0, DataDebugName);
 	}
 }
 
-class FDistanceFieldUploadDataResource : public FRenderResource
+template <EDistanceFieldPrimitiveType PrimitiveType>
+class TDistanceFieldUploadDataResource : public FRenderResource
 {
 public:
 
 	FCPUUpdatedBuffer UploadData;
 
-	FDistanceFieldUploadDataResource()
+	TDistanceFieldUploadDataResource()
 	{
 		// PS4 volatile only supports 8Mb, switch to dynamic once that is fixed.
 		// PS4 volatile only supports 8Mb, switch to volatile once that is fixed.
 		UploadData.bVolatile = false;
 
 		UploadData.Format = PF_A32B32G32R32F;
-		UploadData.Stride = UploadObjectDataStride;
+		UploadData.Stride = PrimitiveType == DFPT_HeightField ? UploadHeightFieldObjectDataStride : UploadObjectDataStride;
 	}
 
 	virtual void InitDynamicRHI()  override
@@ -90,7 +114,8 @@ public:
 	}
 };
 
-TGlobalResource<FDistanceFieldUploadDataResource> GDistanceFieldUploadData;
+TGlobalResource<TDistanceFieldUploadDataResource<DFPT_SignedDistanceField>> GDistanceFieldUploadData;
+TGlobalResource<TDistanceFieldUploadDataResource<DFPT_HeightField>> GHeightFieldUploadData;
 
 class FDistanceFieldUploadIndicesResource : public FRenderResource
 {
@@ -119,6 +144,7 @@ public:
 };
 
 TGlobalResource<FDistanceFieldUploadIndicesResource> GDistanceFieldUploadIndices;
+TGlobalResource<FDistanceFieldUploadIndicesResource> GHeightFieldUploadIndices;
 
 class FDistanceFieldRemoveIndicesResource : public FRenderResource
 {
@@ -144,12 +170,14 @@ public:
 };
 
 TGlobalResource<FDistanceFieldRemoveIndicesResource> GDistanceFieldRemoveIndices;
+TGlobalResource<FDistanceFieldRemoveIndicesResource> GHeightFieldRemoveIndices;
 
 const uint32 UpdateObjectsGroupSize = 64;
 
-class FUploadObjectsToBufferCS : public FGlobalShader
+template <EDistanceFieldPrimitiveType PrimitiveType>
+class TUploadObjectsToBufferCS : public FGlobalShader
 {
-	DECLARE_SHADER_TYPE(FUploadObjectsToBufferCS,Global)
+	DECLARE_SHADER_TYPE(TUploadObjectsToBufferCS,Global)
 public:
 
 	static bool ShouldCompilePermutation(const FGlobalShaderPermutationParameters& Parameters)
@@ -161,9 +189,10 @@ public:
 	{
 		FGlobalShader::ModifyCompilationEnvironment(Parameters,OutEnvironment);
 		OutEnvironment.SetDefine(TEXT("UPDATEOBJECTS_THREADGROUP_SIZE"), UpdateObjectsGroupSize);
+		OutEnvironment.SetDefine(TEXT("DISTANCEFIELD_PRIMITIVE_TYPE"), (int32)PrimitiveType);
 	}
 
-	FUploadObjectsToBufferCS(const ShaderMetaType::CompiledShaderInitializerType& Initializer)
+	TUploadObjectsToBufferCS(const ShaderMetaType::CompiledShaderInitializerType& Initializer)
 		: FGlobalShader(Initializer)
 	{
 		NumUploadOperations.Bind(Initializer.ParameterMap, TEXT("NumUploadOperations"));
@@ -172,7 +201,7 @@ public:
 		ObjectBufferParameters.Bind(Initializer.ParameterMap);
 	}
 
-	FUploadObjectsToBufferCS()
+	TUploadObjectsToBufferCS()
 	{
 	}
 
@@ -185,20 +214,41 @@ public:
 		SetSRVParameter(RHICmdList, ShaderRHI, UploadOperationIndices, InUploadOperationIndices);
 		SetSRVParameter(RHICmdList, ShaderRHI, UploadOperationData, InUploadOperationData);
 
+		constexpr bool bIsHeightField = PrimitiveType == DFPT_HeightField;
+		const FDistanceFieldSceneData& SceneData = Scene->DistanceFieldSceneData;
+		const auto& ObjectBuffers = TSelector<bIsHeightField>()(*SceneData.GetHeightFieldObjectBuffers(), *SceneData.GetCurrentObjectBuffers());
+		const uint32 NumObjectsInBuffer = bIsHeightField ? SceneData.NumHeightFieldObjectsInBuffer : SceneData.NumObjectsInBuffer;
+
+		FRHITexture* TextureAtlas;
+		int32 AtlasSizeX;
+		int32 AtlasSizeY;
+		int32 AtlasSizeZ;
 
-		ObjectBufferParameters.Set(RHICmdList, ShaderRHI, *(Scene->DistanceFieldSceneData.GetCurrentObjectBuffers()), Scene->DistanceFieldSceneData.NumObjectsInBuffer, true);
+		if (bIsHeightField)
+		{
+			TextureAtlas = GHeightFieldTextureAtlas.GetAtlasTexture();
+			AtlasSizeX = GHeightFieldTextureAtlas.GetSizeX();
+			AtlasSizeY = GHeightFieldTextureAtlas.GetSizeY();
+			AtlasSizeZ = 1;
+		}
+		else
+		{
+			TextureAtlas = GDistanceFieldVolumeTextureAtlas.VolumeTextureRHI;
+			AtlasSizeX = GDistanceFieldVolumeTextureAtlas.GetSizeX();
+			AtlasSizeY = GDistanceFieldVolumeTextureAtlas.GetSizeY();
+			AtlasSizeZ = GDistanceFieldVolumeTextureAtlas.GetSizeZ();
+		}
+
+		ObjectBufferParameters.Set(RHICmdList, ShaderRHI, ObjectBuffers, NumObjectsInBuffer, TextureAtlas, AtlasSizeX, AtlasSizeY, AtlasSizeZ, true);
 	}
 
 	void UnsetParameters(FRHICommandList& RHICmdList, const FScene* Scene)
 	{
-		const FDistanceFieldObjectBuffers& ObjectBuffers = *(Scene->DistanceFieldSceneData.GetCurrentObjectBuffers());
+		constexpr bool bIsHeightField = PrimitiveType == DFPT_HeightField;
+		const FDistanceFieldSceneData& SceneData = Scene->DistanceFieldSceneData;
+		const auto& ObjectBuffers = TSelector<bIsHeightField>()(*SceneData.GetHeightFieldObjectBuffers(), *SceneData.GetCurrentObjectBuffers());
 
 		ObjectBufferParameters.UnsetParameters(RHICmdList, GetComputeShader(), ObjectBuffers, true);
-
-		FRHIUnorderedAccessView* OutUAVs[2];
-		OutUAVs[0] = ObjectBuffers.Bounds.UAV;
-		OutUAVs[1] = ObjectBuffers.Data.UAV;
-		RHICmdList.TransitionResources(EResourceTransitionAccess::EReadable, EResourceTransitionPipeline::EComputeToCompute, OutUAVs, UE_ARRAY_COUNT(OutUAVs));
 	}
 
 	virtual bool Serialize(FArchive& Ar) override
@@ -216,15 +266,19 @@ private:
 	FShaderParameter NumUploadOperations;
 	FShaderResourceParameter UploadOperationIndices;
 	FShaderResourceParameter UploadOperationData;
-	FDistanceFieldObjectBufferParameters ObjectBufferParameters;
+	TDistanceFieldObjectBufferParameters<PrimitiveType> ObjectBufferParameters;
 };
 
-IMPLEMENT_SHADER_TYPE(,FUploadObjectsToBufferCS,TEXT("/Engine/Private/DistanceFieldObjectCulling.usf"),TEXT("UploadObjectsToBufferCS"),SF_Compute);
+IMPLEMENT_SHADER_TYPE(template<>, TUploadObjectsToBufferCS<DFPT_SignedDistanceField>, TEXT("/Engine/Private/DistanceFieldObjectCulling.usf"), TEXT("UploadObjectsToBufferCS"), SF_Compute);
+IMPLEMENT_SHADER_TYPE(template<>, TUploadObjectsToBufferCS<DFPT_HeightField>, TEXT("/Engine/Private/DistanceFieldObjectCulling.usf"), TEXT("UploadObjectsToBufferCS"), SF_Compute);
 
-class FCopyObjectBufferCS : public FGlobalShader
+template <EDistanceFieldPrimitiveType PrimitiveType>
+class TCopyObjectBufferCS : public FGlobalShader
 {
-	DECLARE_SHADER_TYPE(FCopyObjectBufferCS,Global)
+	DECLARE_SHADER_TYPE(TCopyObjectBufferCS, Global);
+
 public:
+	typedef typename TChooseClass<PrimitiveType == DFPT_HeightField, FHeightFieldObjectBuffers, FDistanceFieldObjectBuffers>::Result ObjectBuffersType;
 
 	static bool ShouldCompilePermutation(const FGlobalShaderPermutationParameters& Parameters)
 	{
@@ -235,9 +289,10 @@ public:
 	{
 		FGlobalShader::ModifyCompilationEnvironment(Parameters,OutEnvironment);
 		OutEnvironment.SetDefine(TEXT("UPDATEOBJECTS_THREADGROUP_SIZE"), UpdateObjectsGroupSize);
+		OutEnvironment.SetDefine(TEXT("DISTANCEFIELD_PRIMITIVE_TYPE"), (int32)PrimitiveType);
 	}
 
-	FCopyObjectBufferCS(const ShaderMetaType::CompiledShaderInitializerType& Initializer)
+	TCopyObjectBufferCS(const ShaderMetaType::CompiledShaderInitializerType& Initializer)
 		: FGlobalShader(Initializer)
 	{
 		CopyObjectBounds.Bind(Initializer.ParameterMap, TEXT("CopyObjectBounds"));
@@ -245,11 +300,11 @@ public:
 		ObjectBufferParameters.Bind(Initializer.ParameterMap);
 	}
 
-	FCopyObjectBufferCS()
+	TCopyObjectBufferCS()
 	{
 	}
 
-	void SetParameters(FRHICommandList& RHICmdList, FDistanceFieldObjectBuffers& ObjectBuffersSource, FDistanceFieldObjectBuffers& ObjectBuffersDest, int32 NumObjectsValue)
+	void SetParameters(FRHICommandList& RHICmdList, ObjectBuffersType& ObjectBuffersSource, ObjectBuffersType& ObjectBuffersDest, int32 NumObjectsValue)
 	{
 		FRHIComputeShader* ShaderRHI = GetComputeShader();
 
@@ -261,10 +316,30 @@ public:
 		CopyObjectBounds.SetBuffer(RHICmdList, ShaderRHI, ObjectBuffersDest.Bounds);
 		CopyObjectData.SetBuffer(RHICmdList, ShaderRHI, ObjectBuffersDest.Data);
 
-		ObjectBufferParameters.Set(RHICmdList, ShaderRHI, ObjectBuffersSource, NumObjectsValue);
+		FRHITexture* TextureAtlas;
+		int32 AtlasSizeX;
+		int32 AtlasSizeY;
+		int32 AtlasSizeZ;
+
+		if (PrimitiveType == DFPT_HeightField)
+		{
+			TextureAtlas = GHeightFieldTextureAtlas.GetAtlasTexture();
+			AtlasSizeX = GHeightFieldTextureAtlas.GetSizeX();
+			AtlasSizeY = GHeightFieldTextureAtlas.GetSizeY();
+			AtlasSizeZ = 1;
+		}
+		else
+		{
+			TextureAtlas = GDistanceFieldVolumeTextureAtlas.VolumeTextureRHI;
+			AtlasSizeX = GDistanceFieldVolumeTextureAtlas.GetSizeX();
+			AtlasSizeY = GDistanceFieldVolumeTextureAtlas.GetSizeY();
+			AtlasSizeZ = GDistanceFieldVolumeTextureAtlas.GetSizeZ();
+		}
+
+		ObjectBufferParameters.Set(RHICmdList, ShaderRHI, ObjectBuffersSource, NumObjectsValue, TextureAtlas, AtlasSizeX, AtlasSizeY, AtlasSizeZ);
 	}
 
-	void UnsetParameters(FRHICommandList& RHICmdList, FDistanceFieldObjectBuffers& ObjectBuffersDest)
+	void UnsetParameters(FRHICommandList& RHICmdList, ObjectBuffersType& ObjectBuffersDest)
 	{
 		ObjectBufferParameters.UnsetParameters(RHICmdList, GetComputeShader(), ObjectBuffersDest);
 		CopyObjectBounds.UnsetUAV(RHICmdList, GetComputeShader());
@@ -289,10 +364,11 @@ private:
 
 	FRWShaderParameter CopyObjectBounds;
 	FRWShaderParameter CopyObjectData;
-	FDistanceFieldObjectBufferParameters ObjectBufferParameters;
+	TDistanceFieldObjectBufferParameters<PrimitiveType> ObjectBufferParameters;
 };
 
-IMPLEMENT_SHADER_TYPE(,FCopyObjectBufferCS,TEXT("/Engine/Private/DistanceFieldObjectCulling.usf"),TEXT("CopyObjectBufferCS"),SF_Compute);
+IMPLEMENT_SHADER_TYPE(template<>, TCopyObjectBufferCS<DFPT_SignedDistanceField>, TEXT("/Engine/Private/DistanceFieldObjectCulling.usf"), TEXT("CopyObjectBufferCS"), SF_Compute);
+IMPLEMENT_SHADER_TYPE(template<>, TCopyObjectBufferCS<DFPT_HeightField>, TEXT("/Engine/Private/DistanceFieldObjectCulling.usf"), TEXT("CopyObjectBufferCS"), SF_Compute);
 
 class FCopySurfelBufferCS : public FGlobalShader
 {
@@ -434,7 +510,7 @@ private:
 
 IMPLEMENT_SHADER_TYPE(,FCopyVPLFluxBufferCS,TEXT("/Engine/Private/SurfelTree.usf"),TEXT("CopyVPLFluxBufferCS"),SF_Compute);
 
-template<bool bRemoveFromSameBuffer>
+template<bool bRemoveFromSameBuffer, EDistanceFieldPrimitiveType PrimitiveType>
 class TRemoveObjectsFromBufferCS : public FGlobalShader
 {
 	DECLARE_SHADER_TYPE(TRemoveObjectsFromBufferCS,Global)
@@ -450,6 +526,7 @@ public:
 		FGlobalShader::ModifyCompilationEnvironment(Parameters,OutEnvironment);
 		OutEnvironment.SetDefine(TEXT("UPDATEOBJECTS_THREADGROUP_SIZE"), UpdateObjectsGroupSize);
 		OutEnvironment.SetDefine(TEXT("REMOVE_FROM_SAME_BUFFER"), bRemoveFromSameBuffer);
+		OutEnvironment.SetDefine(TEXT("DISTANCEFIELD_PRIMITIVE_TYPE"), (int32)PrimitiveType);
 	}
 
 	TRemoveObjectsFromBufferCS(const ShaderMetaType::CompiledShaderInitializerType& Initializer)
@@ -478,14 +555,46 @@ public:
 
 		SetShaderValue(RHICmdList, ShaderRHI, NumRemoveOperations, NumRemoveOperationsValue);
 		SetSRVParameter(RHICmdList, ShaderRHI, RemoveOperationIndices, InRemoveOperationIndices);
-		ObjectBufferParameters.Set(RHICmdList, ShaderRHI, *(Scene->DistanceFieldSceneData.GetCurrentObjectBuffers()), Scene->DistanceFieldSceneData.NumObjectsInBuffer, true);
+
+		constexpr bool bIsHeightField = PrimitiveType == DFPT_HeightField;
+		const FDistanceFieldSceneData& SceneData = Scene->DistanceFieldSceneData;
+		const auto& ObjectBuffers = TSelector<bIsHeightField>()(*SceneData.GetHeightFieldObjectBuffers(), *SceneData.GetCurrentObjectBuffers());
+		const uint32 NumObjectsInBuffer = bIsHeightField ? SceneData.NumHeightFieldObjectsInBuffer : SceneData.NumObjectsInBuffer;
+
+		FRHITexture* TextureAtlas;
+		int32 AtlasSizeX;
+		int32 AtlasSizeY;
+		int32 AtlasSizeZ;
+
+		if (PrimitiveType == DFPT_HeightField)
+		{
+			TextureAtlas = GHeightFieldTextureAtlas.GetAtlasTexture();
+			AtlasSizeX = GHeightFieldTextureAtlas.GetSizeX();
+			AtlasSizeY = GHeightFieldTextureAtlas.GetSizeY();
+			AtlasSizeZ = 1;
+		}
+		else
+		{
+			check(PrimitiveType == DFPT_SignedDistanceField);
+			TextureAtlas = GDistanceFieldVolumeTextureAtlas.VolumeTextureRHI;
+			AtlasSizeX = GDistanceFieldVolumeTextureAtlas.GetSizeX();
+			AtlasSizeY = GDistanceFieldVolumeTextureAtlas.GetSizeY();
+			AtlasSizeZ = GDistanceFieldVolumeTextureAtlas.GetSizeZ();
+		}
+
+		ObjectBufferParameters.Set(RHICmdList, ShaderRHI, ObjectBuffers, NumObjectsInBuffer, TextureAtlas, AtlasSizeX, AtlasSizeY, AtlasSizeZ, true);
+
 		SetSRVParameter(RHICmdList, ShaderRHI, ObjectBounds2, InObjectBounds2);
 		SetSRVParameter(RHICmdList, ShaderRHI, ObjectData2, InObjectData2);
 	}
 
 	void UnsetParameters(FRHICommandList& RHICmdList, const FScene* Scene)
 	{
-		ObjectBufferParameters.UnsetParameters(RHICmdList, GetComputeShader(), *(Scene->DistanceFieldSceneData.GetCurrentObjectBuffers()), true);
+		constexpr bool bIsHeightField = PrimitiveType == DFPT_HeightField;
+		const FDistanceFieldSceneData& SceneData = Scene->DistanceFieldSceneData;
+		const auto& ObjectBuffers = TSelector<bIsHeightField>()(*SceneData.GetHeightFieldObjectBuffers(), *SceneData.GetCurrentObjectBuffers());
+
+		ObjectBufferParameters.UnsetParameters(RHICmdList, GetComputeShader(), ObjectBuffers, true);
 	}
 
 	virtual bool Serialize(FArchive& Ar) override
@@ -503,13 +612,18 @@ private:
 
 	FShaderParameter NumRemoveOperations;
 	FShaderResourceParameter RemoveOperationIndices;
-	FDistanceFieldObjectBufferParameters ObjectBufferParameters;
+	TDistanceFieldObjectBufferParameters<PrimitiveType> ObjectBufferParameters;
 	FShaderResourceParameter ObjectBounds2;
 	FShaderResourceParameter ObjectData2;
 };
 
-IMPLEMENT_SHADER_TYPE(template<>,TRemoveObjectsFromBufferCS<true>,TEXT("/Engine/Private/DistanceFieldObjectCulling.usf"),TEXT("RemoveObjectsFromBufferCS"),SF_Compute);
-IMPLEMENT_SHADER_TYPE(template<>,TRemoveObjectsFromBufferCS<false>,TEXT("/Engine/Private/DistanceFieldObjectCulling.usf"),TEXT("RemoveObjectsFromBufferCS"),SF_Compute);
+#define IMPLEMENT_REMOVE_OBJECTS_FROM_BUFFER_SHADER_TYPE(bRemoveFromSameBuffer, PrimitiveType) \
+	typedef TRemoveObjectsFromBufferCS<bRemoveFromSameBuffer, PrimitiveType> FRemoveObjectsFromBufferCS##bRemoveFromSameBuffer##PrimitiveType; \
+	IMPLEMENT_SHADER_TYPE(template<>, FRemoveObjectsFromBufferCS##bRemoveFromSameBuffer##PrimitiveType, TEXT("/Engine/Private/DistanceFieldObjectCulling.usf"), TEXT("RemoveObjectsFromBufferCS"), SF_Compute);
+
+IMPLEMENT_REMOVE_OBJECTS_FROM_BUFFER_SHADER_TYPE(true, DFPT_SignedDistanceField);
+IMPLEMENT_REMOVE_OBJECTS_FROM_BUFFER_SHADER_TYPE(false, DFPT_SignedDistanceField);
+IMPLEMENT_REMOVE_OBJECTS_FROM_BUFFER_SHADER_TYPE(true, DFPT_HeightField);
 
 void FSurfelBufferAllocator::RemovePrimitive(const FPrimitiveSceneInfo* Primitive)
 {
@@ -762,7 +876,7 @@ void UpdateGlobalDistanceFieldObjectRemoves(FRHICommandListImmediate& RHICmdList
 
 				if (bUseRemoveAtSwap)
 				{
-					TShaderMapRef<TRemoveObjectsFromBufferCS<true> > ComputeShader(GetGlobalShaderMap(Scene->GetFeatureLevel()));
+					TShaderMapRef<TRemoveObjectsFromBufferCS<true, DFPT_SignedDistanceField>> ComputeShader(GetGlobalShaderMap(Scene->GetFeatureLevel()));
 					RHICmdList.SetComputeShader(ComputeShader->GetComputeShader());
 					ComputeShader->SetParameters(RHICmdList, Scene, RemoveObjectIndices.Num(), GDistanceFieldRemoveIndices.RemoveIndices.BufferSRV, NULL, NULL);
 
@@ -771,7 +885,7 @@ void UpdateGlobalDistanceFieldObjectRemoves(FRHICommandListImmediate& RHICmdList
 				}
 				else
 				{
-					TShaderMapRef<TRemoveObjectsFromBufferCS<false> > ComputeShader(GetGlobalShaderMap(Scene->GetFeatureLevel()));
+					TShaderMapRef<TRemoveObjectsFromBufferCS<false, DFPT_SignedDistanceField>> ComputeShader(GetGlobalShaderMap(Scene->GetFeatureLevel()));
 					RHICmdList.SetComputeShader(ComputeShader->GetComputeShader());
 					ComputeShader->SetParameters(RHICmdList, Scene, RemoveObjectIndices.Num(), GDistanceFieldRemoveIndices.RemoveIndices.BufferSRV, CurrentObjectBuffers->Bounds.SRV, CurrentObjectBuffers->Data.SRV);
 
@@ -783,6 +897,89 @@ void UpdateGlobalDistanceFieldObjectRemoves(FRHICommandListImmediate& RHICmdList
 	}
 }
 
+void UpdateGlobalHeightFieldObjectRemoves(FRHICommandListImmediate& RHICmdList, FScene* Scene)
+{
+	FDistanceFieldSceneData& SceneData = Scene->DistanceFieldSceneData;
+	TArray<FIntRect> DstSrcIndices;
+
+	if (SceneData.PendingHeightFieldRemoveOps.Num())
+	{
+		TArray<int32, SceneRenderingAllocator> PendingRemoveObjectIndices;
+
+		for (int32 Idx = 0; Idx < SceneData.PendingHeightFieldRemoveOps.Num(); ++Idx)
+		{
+			const FHeightFieldPrimitiveRemoveInfo& RemoveInfo = SceneData.PendingHeightFieldRemoveOps[Idx];
+			check(RemoveInfo.DistanceFieldInstanceIndices.Num() == 1);
+			const int32 ObjectIdx = RemoveInfo.DistanceFieldInstanceIndices[0];
+
+			if (ObjectIdx >= 0)
+			{
+				const FGlobalDFCacheType CacheType = RemoveInfo.bOftenMoving ? GDF_Full : GDF_MostlyStatic;
+				SceneData.PrimitiveModifiedBounds[CacheType].Add(RemoveInfo.SphereBound);
+				PendingRemoveObjectIndices.Add(ObjectIdx);
+			}
+		}
+
+		SceneData.PendingHeightFieldRemoveOps.Reset();
+
+		if (PendingRemoveObjectIndices.Num())
+		{
+			check(SceneData.NumHeightFieldObjectsInBuffer >= PendingRemoveObjectIndices.Num());
+			check(SceneData.NumHeightFieldObjectsInBuffer == SceneData.HeightfieldPrimitives.Num());
+
+			Algo::Sort(PendingRemoveObjectIndices);
+
+			for (int32 Idx = 0; Idx < PendingRemoveObjectIndices.Num(); ++Idx)
+			{
+				const int32 LastIdx = PendingRemoveObjectIndices.Num() - 1;
+				const int32 RemoveIdx = PendingRemoveObjectIndices[Idx];
+				const int32 LastRemoveIdx = PendingRemoveObjectIndices[LastIdx];
+				const int32 LastObjectIdx = --SceneData.NumHeightFieldObjectsInBuffer;
+
+				if (LastRemoveIdx < LastObjectIdx)
+				{
+					DstSrcIndices.Add(FIntRect(RemoveIdx, LastObjectIdx, 0, 0));
+
+					FPrimitiveSceneInfo* Primitive = SceneData.HeightfieldPrimitives[LastObjectIdx];
+					check(Primitive && Primitive->DistanceFieldInstanceIndices.Num() == 1);
+					Primitive->DistanceFieldInstanceIndices[0] = RemoveIdx;
+					SceneData.HeightfieldPrimitives.RemoveAtSwap(RemoveIdx);
+				}
+				else
+				{
+					check(LastRemoveIdx == LastObjectIdx);
+					SceneData.HeightfieldPrimitives.RemoveAt(LastObjectIdx);
+					PendingRemoveObjectIndices.RemoveAt(LastIdx, 1, false);
+					--Idx;
+				}
+			}
+		}
+
+		if (DstSrcIndices.Num())
+		{
+			FCPUUpdatedBuffer& RemoveIndices = GHeightFieldRemoveIndices.RemoveIndices;
+			if (DstSrcIndices.Num() > RemoveIndices.MaxElements)
+			{
+				RemoveIndices.MaxElements = DstSrcIndices.Num() * 5 / 4;
+				RemoveIndices.Release();
+				RemoveIndices.Initialize();
+			}
+
+			void* LockedBuffer = RHILockVertexBuffer(RemoveIndices.Buffer, 0, RemoveIndices.Buffer->GetSize(), RLM_WriteOnly);
+			const uint32 MemcpySize = DstSrcIndices.GetTypeSize() * DstSrcIndices.Num();
+			check(RemoveIndices.Buffer->GetSize() >= MemcpySize);
+			FMemory::Memcpy(LockedBuffer, DstSrcIndices.GetData(), MemcpySize);
+			RHIUnlockVertexBuffer(RemoveIndices.Buffer);
+
+			TShaderMapRef<TRemoveObjectsFromBufferCS<true, DFPT_HeightField>> ComputeShader(GetGlobalShaderMap(Scene->GetFeatureLevel()));
+			RHICmdList.SetComputeShader(ComputeShader->GetComputeShader());
+			ComputeShader->SetParameters(RHICmdList, Scene, DstSrcIndices.Num(), RemoveIndices.BufferSRV, nullptr, nullptr);
+			DispatchComputeShader(RHICmdList, *ComputeShader, FMath::DivideAndRoundUp<uint32>(DstSrcIndices.Num(), UpdateObjectsGroupSize), 1, 1);
+			ComputeShader->UnsetParameters(RHICmdList, Scene);
+		}
+	}
+}
+
 /** Gathers the information needed to represent a single object's distance field and appends it to the upload buffers. */
 bool ProcessPrimitiveUpdate(
 	bool bIsAddOperation,
@@ -1021,15 +1218,84 @@ bool ProcessPrimitiveUpdate(
 	return true;
 }
 
+bool ProcessHeightFieldPrimitiveUpdate(
+	bool bIsAddOperation,
+	FRHICommandListImmediate& RHICmdList,
+	FScene* Scene,
+	FPrimitiveSceneInfo* PrimitiveSceneInfo,
+	int32 OriginalNumObjects,
+	TArray<uint32>& UploadObjectIndices,
+	TArray<FVector4>& UploadObjectData)
+{
+	FDistanceFieldSceneData& SceneData = Scene->DistanceFieldSceneData;
+
+	UTexture2D* HeightNormalTexture;
+	UTexture2D* DiffuseColorTexture;
+	UTexture2D* VisibilityTexture;
+	FHeightfieldComponentDescription HeightFieldCompDesc(PrimitiveSceneInfo->Proxy->GetLocalToWorld());
+	PrimitiveSceneInfo->Proxy->GetHeightfieldRepresentation(HeightNormalTexture, DiffuseColorTexture, VisibilityTexture, HeightFieldCompDesc);
+
+	uint32 UploadIdx;
+	if (bIsAddOperation)
+	{
+		++SceneData.NumHeightFieldObjectsInBuffer;
+		SceneData.HeightfieldPrimitives.Add(PrimitiveSceneInfo);
+
+		const FGlobalDFCacheType CacheType = PrimitiveSceneInfo->Proxy->IsOftenMoving() ? GDF_Full : GDF_MostlyStatic;
+		const FBoxSphereBounds Bounds = PrimitiveSceneInfo->Proxy->GetBounds();
+		SceneData.PrimitiveModifiedBounds[CacheType].Add(FVector4(Bounds.Origin, Bounds.SphereRadius));
+
+		UploadIdx = OriginalNumObjects + UploadObjectIndices.Num();
+		PrimitiveSceneInfo->DistanceFieldInstanceIndices.Empty(1);
+		PrimitiveSceneInfo->DistanceFieldInstanceIndices.Add(UploadIdx);
+	}
+	else
+	{
+		UploadIdx = PrimitiveSceneInfo->DistanceFieldInstanceIndices[0];
+	}
+
+	UploadObjectIndices.Add(UploadIdx);
+
+	const FBoxSphereBounds& Bounds = PrimitiveSceneInfo->Proxy->GetBounds();
+	const FBox BoxBound = Bounds.GetBox();
+	UploadObjectData.Add(FVector4(BoxBound.GetCenter(), Bounds.SphereRadius));
+	UploadObjectData.Add(FVector4(BoxBound.GetExtent(), 0.f));
+
+	const FMatrix& LocalToWorld = HeightFieldCompDesc.LocalToWorld;
+	check(LocalToWorld.GetMaximumAxisScale() > 0.f);
+	const FMatrix WorldToLocalT = LocalToWorld.Inverse().GetTransposed();
+	UploadObjectData.Add(*(const FVector4*)&WorldToLocalT.M[0]);
+	UploadObjectData.Add(*(const FVector4*)&WorldToLocalT.M[1]);
+	UploadObjectData.Add(*(const FVector4*)&WorldToLocalT.M[2]);
+
+	const FIntRect& HeightFieldRect = HeightFieldCompDesc.HeightfieldRect;
+	UploadObjectData.Add(FVector4(HeightFieldRect.Min.X, HeightFieldRect.Min.Y, HeightFieldRect.Width(), HeightFieldRect.Height()));
+
+	const FVector4& HeightFieldScaleBias = HeightFieldCompDesc.HeightfieldScaleBias;
+	check(HeightFieldScaleBias.Y >= 0.f && HeightFieldScaleBias.Z >= 0.f && HeightFieldScaleBias.W >= 0.f);
+	const uint32 TileHandle = GHeightFieldTextureAtlas.GetAllocationHandle(HeightNormalTexture);
+	check(TileHandle != INDEX_NONE);
+	const FVector4 TileScaleBias = GHeightFieldTextureAtlas.GetTileScaleBias(TileHandle);
+	UploadObjectData.Add(FVector4(
+		FMath::Abs(HeightFieldScaleBias.X) * TileScaleBias.X,
+		HeightFieldScaleBias.Y * TileScaleBias.Y,
+		HeightFieldScaleBias.Z * TileScaleBias.X + TileScaleBias.Z,
+		HeightFieldScaleBias.W * TileScaleBias.Y + TileScaleBias.W));
+
+	check(!(UploadObjectData.Num() % UploadHeightFieldObjectDataStride));
+
+	return true;
+}
+
 bool bVerifySceneIntegrity = false;
 
-void FDeferredShadingSceneRenderer::UpdateGlobalDistanceFieldObjectBuffers(FRHICommandListImmediate& RHICmdList) 
+void FDeferredShadingSceneRenderer::UpdateGlobalDistanceFieldObjectBuffers(FRHICommandListImmediate& RHICmdList)
 {
 	FDistanceFieldSceneData& DistanceFieldSceneData = Scene->DistanceFieldSceneData;
 
 	if (GDistanceFieldVolumeTextureAtlas.VolumeTextureRHI
-		&& (DistanceFieldSceneData.HasPendingOperations() || 
-			DistanceFieldSceneData.PendingThrottledOperations.Num()>0 || 
+		&& (DistanceFieldSceneData.HasPendingOperations() ||
+			DistanceFieldSceneData.PendingThrottledOperations.Num() > 0 ||
 			DistanceFieldSceneData.AtlasGeneration != GDistanceFieldVolumeTextureAtlas.GetGeneration()))
 	{
 		QUICK_SCOPE_CYCLE_COUNTER(STAT_UpdateObjectData);
@@ -1038,7 +1304,7 @@ void FDeferredShadingSceneRenderer::UpdateGlobalDistanceFieldObjectBuffers(FRHIC
 		SCOPED_GPU_MASK(RHICmdList, FRHIGPUMask::All());
 		SCOPED_DRAW_EVENT(RHICmdList, UpdateSceneObjectData);
 
-		if (DistanceFieldSceneData.ObjectBuffers[DistanceFieldSceneData.ObjectBufferIndex]==nullptr)
+		if (DistanceFieldSceneData.ObjectBuffers[DistanceFieldSceneData.ObjectBufferIndex] == nullptr)
 		{
 			DistanceFieldSceneData.ObjectBuffers[DistanceFieldSceneData.ObjectBufferIndex] = new FDistanceFieldObjectBuffers();
 		}
@@ -1053,7 +1319,7 @@ void FDeferredShadingSceneRenderer::UpdateGlobalDistanceFieldObjectBuffers(FRHIC
 			DistanceFieldSceneData.InstancedSurfelBuffers = new FInstancedSurfelBuffers();
 		}
 
-		if (DistanceFieldSceneData.PendingAddOperations.Num()>0)
+		if (DistanceFieldSceneData.PendingAddOperations.Num() > 0)
 		{
 			DistanceFieldSceneData.PendingThrottledOperations.Reserve(DistanceFieldSceneData.PendingThrottledOperations.Num() + DistanceFieldSceneData.PendingAddOperations.Num());
 		}
@@ -1070,7 +1336,7 @@ void FDeferredShadingSceneRenderer::UpdateGlobalDistanceFieldObjectBuffers(FRHIC
 				FPrimitiveAndInstance& PrimitiveInstance = DistanceFieldSceneData.PrimitiveInstanceMapping[PrimitiveInstanceIndex];
 
 				// Queue an update of all primitives, since the atlas layout has changed
-				if (PrimitiveInstance.InstanceIndex == 0 
+				if (PrimitiveInstance.InstanceIndex == 0
 					&& !DistanceFieldSceneData.HasPendingRemovePrimitive(PrimitiveInstance.Primitive)
 					&& !DistanceFieldSceneData.PendingAddOperations.Contains(PrimitiveInstance.Primitive)
 					&& !DistanceFieldSceneData.PendingUpdateOperations.Contains(PrimitiveInstance.Primitive))
@@ -1104,7 +1370,7 @@ void FDeferredShadingSceneRenderer::UpdateGlobalDistanceFieldObjectBuffers(FRHIC
 			int32 OriginalNumObjects = DistanceFieldSceneData.NumObjectsInBuffer;
 			int32 OriginalNumSurfels = DistanceFieldSceneData.SurfelAllocations.GetNumSurfelsInBuffer();
 			int32 OriginalNumInstancedSurfels = DistanceFieldSceneData.InstancedSurfelAllocations.GetNumSurfelsInBuffer();
-			
+
 			if (bPrepareForDistanceFieldGI)
 			{
 				for (int32 UploadPrimitiveIndex = 0; UploadPrimitiveIndex < DistanceFieldSceneData.PendingAddOperations.Num(); UploadPrimitiveIndex++)
@@ -1126,7 +1392,7 @@ void FDeferredShadingSceneRenderer::UpdateGlobalDistanceFieldObjectBuffers(FRHIC
 
 						if (PrimitiveNumSurfels > 5000)
 						{
-							UE_LOG(LogDistanceField,Warning,TEXT("Primitive %s %s used %u Surfels"), *PrimitiveSceneInfo->Proxy->GetOwnerName().ToString(), *PrimitiveSceneInfo->Proxy->GetResourceName().ToString(), PrimitiveNumSurfels);
+							UE_LOG(LogDistanceField, Warning, TEXT("Primitive %s %s used %u Surfels"), *PrimitiveSceneInfo->Proxy->GetOwnerName().ToString(), *PrimitiveSceneInfo->Proxy->GetResourceName().ToString(), PrimitiveNumSurfels);
 						}
 
 						DistanceFieldSceneData.SurfelAllocations.AddPrimitive(PrimitiveSceneInfo, PrimitiveLOD0Surfels, PrimitiveNumSurfels, 1);
@@ -1163,7 +1429,7 @@ void FDeferredShadingSceneRenderer::UpdateGlobalDistanceFieldObjectBuffers(FRHIC
 						DistanceFieldSceneData.SurfelBuffers->Initialize();
 					}
 				}
-				
+
 				if (DistanceFieldSceneData.InstancedSurfelBuffers->MaxSurfels < DistanceFieldSceneData.InstancedSurfelAllocations.GetNumSurfelsInBuffer())
 				{
 					if (DistanceFieldSceneData.InstancedSurfelBuffers->MaxSurfels > 0)
@@ -1172,7 +1438,7 @@ void FDeferredShadingSceneRenderer::UpdateGlobalDistanceFieldObjectBuffers(FRHIC
 						FInstancedSurfelBuffers* NewInstancedSurfelBuffers = new FInstancedSurfelBuffers();
 						NewInstancedSurfelBuffers->MaxSurfels = DistanceFieldSceneData.InstancedSurfelAllocations.GetNumSurfelsInBuffer() * 5 / 4;
 						NewInstancedSurfelBuffers->Initialize();
-						
+
 						{
 							TShaderMapRef<FCopyVPLFluxBufferCS> ComputeShader(GetGlobalShaderMap(Scene->GetFeatureLevel()));
 							RHICmdList.SetComputeShader(ComputeShader->GetComputeShader());
@@ -1265,7 +1531,7 @@ void FDeferredShadingSceneRenderer::UpdateGlobalDistanceFieldObjectBuffers(FRHIC
 					NewObjectBuffers->Initialize();
 
 					{
-						TShaderMapRef<FCopyObjectBufferCS> ComputeShader(GetGlobalShaderMap(Scene->GetFeatureLevel()));
+						TShaderMapRef<TCopyObjectBufferCS<DFPT_SignedDistanceField>> ComputeShader(GetGlobalShaderMap(Scene->GetFeatureLevel()));
 						RHICmdList.SetComputeShader(ComputeShader->GetComputeShader());
 						ComputeShader->SetParameters(RHICmdList, *ObjectBuffers, *NewObjectBuffers, OriginalNumObjects);
 
@@ -1314,7 +1580,7 @@ void FDeferredShadingSceneRenderer::UpdateGlobalDistanceFieldObjectBuffers(FRHIC
 			RHIUnlockVertexBuffer(GDistanceFieldUploadData.UploadData.Buffer);
 
 			{
-				TShaderMapRef<FUploadObjectsToBufferCS> ComputeShader(GetGlobalShaderMap(Scene->GetFeatureLevel()));
+				TShaderMapRef<TUploadObjectsToBufferCS<DFPT_SignedDistanceField>> ComputeShader(GetGlobalShaderMap(Scene->GetFeatureLevel()));
 				RHICmdList.SetComputeShader(ComputeShader->GetComputeShader());
 				ComputeShader->SetParameters(RHICmdList, Scene, UploadObjectIndices.Num(), GDistanceFieldUploadIndices.UploadIndices.BufferSRV, GDistanceFieldUploadData.UploadData.BufferSRV);
 
@@ -1333,6 +1599,133 @@ void FDeferredShadingSceneRenderer::UpdateGlobalDistanceFieldObjectBuffers(FRHIC
 	}
 }
 
+void FDeferredShadingSceneRenderer::UpdateGlobalHeightFieldObjectBuffers(FRHICommandListImmediate& RHICmdList)
+{
+	FDistanceFieldSceneData& DistanceFieldSceneData = Scene->DistanceFieldSceneData;
+
+	if (GHeightFieldTextureAtlas.GetAtlasTexture()
+		&& (DistanceFieldSceneData.HasPendingHeightFieldOperations()
+			|| DistanceFieldSceneData.HeightFieldAtlasGeneration != GHeightFieldTextureAtlas.GetGeneration()))
+	{
+		QUICK_SCOPE_CYCLE_COUNTER(STAT_UpdateHeightFieldSceneObjectData);
+		SCOPED_DRAW_EVENT(RHICmdList, UpdateHeightFieldSceneObjectData);
+
+		if (!DistanceFieldSceneData.HeightFieldObjectBuffers)
+		{
+			DistanceFieldSceneData.HeightFieldObjectBuffers = new FHeightFieldObjectBuffers;
+		}
+
+		if (DistanceFieldSceneData.HeightFieldAtlasGeneration != GHeightFieldTextureAtlas.GetGeneration())
+		{
+			DistanceFieldSceneData.HeightFieldAtlasGeneration = GHeightFieldTextureAtlas.GetGeneration();
+
+			for (int32 Idx = 0; Idx < DistanceFieldSceneData.HeightfieldPrimitives.Num(); ++Idx)
+			{
+				const FPrimitiveSceneInfo* Primitive = DistanceFieldSceneData.HeightfieldPrimitives[Idx];
+
+				if (!DistanceFieldSceneData.HasPendingRemoveHeightFieldPrimitive(Primitive)
+					&& !DistanceFieldSceneData.PendingHeightFieldAddOps.Contains(Primitive)
+					&& !DistanceFieldSceneData.PendingHeightFieldUpdateOps.Contains(Primitive))
+				{
+					DistanceFieldSceneData.PendingUpdateOperations.Add((FPrimitiveSceneInfo*)Primitive);
+				}
+			}
+		}
+
+		UpdateGlobalHeightFieldObjectRemoves(RHICmdList, Scene);
+
+		if (DistanceFieldSceneData.PendingHeightFieldAddOps.Num() || DistanceFieldSceneData.PendingHeightFieldUpdateOps.Num())
+		{
+			const int32 NumAddOps = DistanceFieldSceneData.PendingHeightFieldAddOps.Num();
+			const int32 NumUpdateOps = DistanceFieldSceneData.PendingHeightFieldUpdateOps.Num();
+			const int32 NumUploadOps = NumAddOps + NumUpdateOps;
+			const int32 OriginalNumObjects = DistanceFieldSceneData.NumHeightFieldObjectsInBuffer;
+			TArray<uint32> UploadHeightFieldObjectIndices;
+			TArray<FVector4> UploadHeightFieldObjectData;
+
+			UploadHeightFieldObjectIndices.Empty(NumUploadOps);
+			UploadHeightFieldObjectData.Empty(NumUploadOps * UploadHeightFieldObjectDataStride);
+
+			for (int32 Idx = 0; Idx < NumAddOps; ++Idx)
+			{
+				FPrimitiveSceneInfo* PrimitiveSceneInfo = DistanceFieldSceneData.PendingHeightFieldAddOps[Idx];
+				ProcessHeightFieldPrimitiveUpdate(true, RHICmdList, Scene, PrimitiveSceneInfo, OriginalNumObjects, UploadHeightFieldObjectIndices, UploadHeightFieldObjectData);
+			}
+
+			for (int32 Idx = 0; Idx < NumUpdateOps; ++Idx)
+			{
+				FPrimitiveSceneInfo* PrimitiveSceneInfo = DistanceFieldSceneData.PendingHeightFieldUpdateOps[Idx];
+				ProcessHeightFieldPrimitiveUpdate(false, RHICmdList, Scene, PrimitiveSceneInfo, OriginalNumObjects, UploadHeightFieldObjectIndices, UploadHeightFieldObjectData);
+			}
+
+			DistanceFieldSceneData.PendingHeightFieldAddOps.Reset();
+			DistanceFieldSceneData.PendingHeightFieldUpdateOps.Empty();
+
+			FHeightFieldObjectBuffers*& ObjectBuffers = DistanceFieldSceneData.HeightFieldObjectBuffers;
+
+			if (ObjectBuffers->MaxObjects < DistanceFieldSceneData.NumHeightFieldObjectsInBuffer)
+			{
+				if (ObjectBuffers->MaxObjects > 0)
+				{
+					FHeightFieldObjectBuffers* NewObjectBuffers = new FHeightFieldObjectBuffers;
+					NewObjectBuffers->MaxObjects = DistanceFieldSceneData.NumHeightFieldObjectsInBuffer * 5 / 4;
+					NewObjectBuffers->Initialize();
+
+					TShaderMapRef<TCopyObjectBufferCS<DFPT_HeightField>> ComputeShader(GetGlobalShaderMap(Scene->GetFeatureLevel()));
+					RHICmdList.SetComputeShader(ComputeShader->GetComputeShader());
+					ComputeShader->SetParameters(RHICmdList, *ObjectBuffers, *NewObjectBuffers, OriginalNumObjects);
+					DispatchComputeShader(RHICmdList, *ComputeShader, FMath::DivideAndRoundUp<uint32>(OriginalNumObjects, UpdateObjectsGroupSize), 1, 1);
+					ComputeShader->UnsetParameters(RHICmdList, *NewObjectBuffers);
+
+					ObjectBuffers->Release();
+					delete ObjectBuffers;
+					ObjectBuffers = NewObjectBuffers;
+				}
+				else
+				{
+					ObjectBuffers->MaxObjects = DistanceFieldSceneData.NumHeightFieldObjectsInBuffer * 5 / 4;
+					ObjectBuffers->Initialize();
+				}
+			}
+
+			const int32 NumObjectsToUpload = UploadHeightFieldObjectIndices.Num();
+
+			if (NumObjectsToUpload)
+			{
+				if (NumObjectsToUpload > GHeightFieldUploadIndices.UploadIndices.MaxElements
+					|| (GHeightFieldUploadIndices.UploadIndices.MaxElements > 1000 && GHeightFieldUploadIndices.UploadIndices.MaxElements > NumObjectsToUpload * 2))
+				{
+					GHeightFieldUploadIndices.UploadIndices.MaxElements = NumObjectsToUpload * 5 / 4;
+					GHeightFieldUploadIndices.UploadIndices.Release();
+					GHeightFieldUploadIndices.UploadIndices.Initialize();
+
+					GHeightFieldUploadData.UploadData.MaxElements = NumObjectsToUpload * 5 / 4;
+					GHeightFieldUploadData.UploadData.Release();
+					GHeightFieldUploadData.UploadData.Initialize();
+				}
+
+				void* LockedBuffer = RHILockVertexBuffer(GHeightFieldUploadIndices.UploadIndices.Buffer, 0, GHeightFieldUploadIndices.UploadIndices.Buffer->GetSize(), RLM_WriteOnly);
+				uint32 MemcpySize = UploadHeightFieldObjectIndices.GetTypeSize() * NumObjectsToUpload;
+				check(GHeightFieldUploadIndices.UploadIndices.Buffer->GetSize() >= MemcpySize);
+				FMemory::Memcpy(LockedBuffer, UploadHeightFieldObjectIndices.GetData(), MemcpySize);
+				RHIUnlockVertexBuffer(GHeightFieldUploadIndices.UploadIndices.Buffer);
+
+				LockedBuffer = RHILockVertexBuffer(GHeightFieldUploadData.UploadData.Buffer, 0, GHeightFieldUploadData.UploadData.Buffer->GetSize(), RLM_WriteOnly);
+				MemcpySize = UploadHeightFieldObjectData.GetTypeSize() * UploadHeightFieldObjectData.Num();
+				check(GHeightFieldUploadData.UploadData.Buffer->GetSize() >= MemcpySize);
+				FMemory::Memcpy(LockedBuffer, UploadHeightFieldObjectData.GetData(), MemcpySize);
+				RHIUnlockVertexBuffer(GHeightFieldUploadData.UploadData.Buffer);
+
+				TShaderMapRef<TUploadObjectsToBufferCS<DFPT_HeightField>> ComputeShader(GetGlobalShaderMap(Scene->GetFeatureLevel()));
+				RHICmdList.SetComputeShader(ComputeShader->GetComputeShader());
+				ComputeShader->SetParameters(RHICmdList, Scene, NumObjectsToUpload, GHeightFieldUploadIndices.UploadIndices.BufferSRV, GHeightFieldUploadData.UploadData.BufferSRV);
+				DispatchComputeShader(RHICmdList, *ComputeShader, FMath::DivideAndRoundUp<uint32>(NumObjectsToUpload, UpdateObjectsGroupSize), 1, 1);
+				ComputeShader->UnsetParameters(RHICmdList, Scene);
+			}
+		}
+	}
+}
+
 FString GetObjectBufferMemoryString()
 {
 	return FString::Printf(TEXT("Temp object buffers %.3fMb"), 
diff --git a/Engine/Source/Runtime/Renderer/Private/DistanceFieldScreenGridLighting.cpp b/Engine/Source/Runtime/Renderer/Private/DistanceFieldScreenGridLighting.cpp
--- a/Engine/Source/Runtime/Renderer/Private/DistanceFieldScreenGridLighting.cpp
+++ b/Engine/Source/Runtime/Renderer/Private/DistanceFieldScreenGridLighting.cpp
@@ -131,7 +131,19 @@ public:
 		FRHIComputeShader* ShaderRHI = GetComputeShader();
 		FGlobalShader::SetParameters<FViewUniformShaderParameters>(RHICmdList, ShaderRHI, View.ViewUniformBuffer);
 		SceneTextureParameters.Set(RHICmdList, ShaderRHI, View.FeatureLevel, ESceneTextureSetupMode::All);
-		ObjectParameters.Set(RHICmdList, ShaderRHI, GAOCulledObjectBuffers.Buffers);
+
+		FRHITexture* TextureAtlas;
+		int32 AtlasSizeX;
+		int32 AtlasSizeY;
+		int32 AtlasSizeZ;
+
+		TextureAtlas = GDistanceFieldVolumeTextureAtlas.VolumeTextureRHI;
+		AtlasSizeX = GDistanceFieldVolumeTextureAtlas.GetSizeX();
+		AtlasSizeY = GDistanceFieldVolumeTextureAtlas.GetSizeY();
+		AtlasSizeZ = GDistanceFieldVolumeTextureAtlas.GetSizeZ();
+
+		ObjectParameters.Set(RHICmdList, ShaderRHI, GAOCulledObjectBuffers.Buffers, TextureAtlas, AtlasSizeX, AtlasSizeY, AtlasSizeZ);
+
 		AOParameters.Set(RHICmdList, ShaderRHI, Parameters);
 		ScreenGridParameters.Set(RHICmdList, ShaderRHI, View, DistanceFieldNormal);
 
@@ -205,7 +217,7 @@ public:
 private:
 
 	FSceneTextureShaderParameters SceneTextureParameters;
-	FDistanceFieldCulledObjectBufferParameters ObjectParameters;
+	TDistanceFieldCulledObjectBufferParameters<DFPT_SignedDistanceField> ObjectParameters;
 	FAOParameters AOParameters;
 	FScreenGridParameters ScreenGridParameters;
 	FGlobalDistanceFieldParameters GlobalDistanceFieldParameters;
@@ -282,7 +294,19 @@ public:
 		FRHIComputeShader* ShaderRHI = GetComputeShader();
 		FGlobalShader::SetParameters<FViewUniformShaderParameters>(RHICmdList, ShaderRHI, View.ViewUniformBuffer);
 		SceneTextureParameters.Set(RHICmdList, ShaderRHI, View.FeatureLevel, ESceneTextureSetupMode::All);
-		ObjectParameters.Set(RHICmdList, ShaderRHI, GAOCulledObjectBuffers.Buffers);
+
+		FRHITexture* TextureAtlas;
+		int32 AtlasSizeX;
+		int32 AtlasSizeY;
+		int32 AtlasSizeZ;
+
+		TextureAtlas = GDistanceFieldVolumeTextureAtlas.VolumeTextureRHI;
+		AtlasSizeX = GDistanceFieldVolumeTextureAtlas.GetSizeX();
+		AtlasSizeY = GDistanceFieldVolumeTextureAtlas.GetSizeY();
+		AtlasSizeZ = GDistanceFieldVolumeTextureAtlas.GetSizeZ();
+
+		ObjectParameters.Set(RHICmdList, ShaderRHI, GAOCulledObjectBuffers.Buffers, TextureAtlas, AtlasSizeX, AtlasSizeY, AtlasSizeZ);
+
 		AOParameters.Set(RHICmdList, ShaderRHI, Parameters);
 		ScreenGridParameters.Set(RHICmdList, ShaderRHI, View, DistanceFieldNormal);
 		GlobalDistanceFieldParameters.Set(RHICmdList, ShaderRHI, GlobalDistanceFieldInfo.ParameterData);
@@ -356,7 +380,7 @@ public:
 private:
 
 	FSceneTextureShaderParameters SceneTextureParameters;
-	FDistanceFieldCulledObjectBufferParameters ObjectParameters;
+	TDistanceFieldCulledObjectBufferParameters<DFPT_SignedDistanceField> ObjectParameters;
 	FAOParameters AOParameters;
 	FScreenGridParameters ScreenGridParameters;
 	FGlobalDistanceFieldParameters GlobalDistanceFieldParameters;
diff --git a/Engine/Source/Runtime/Renderer/Private/DistanceFieldShadowing.cpp b/Engine/Source/Runtime/Renderer/Private/DistanceFieldShadowing.cpp
--- a/Engine/Source/Runtime/Renderer/Private/DistanceFieldShadowing.cpp
+++ b/Engine/Source/Runtime/Renderer/Private/DistanceFieldShadowing.cpp
@@ -85,6 +85,35 @@ FAutoConsoleVariableRef CVarAverageObjectsPerShadowCullTile(
 	ECVF_RenderThreadSafe | ECVF_ReadOnly
 	);
 
+static int32 GHeightFieldShadowing = 0;
+FAutoConsoleVariableRef CVarHeightFieldShadowing(
+	TEXT("r.HeightFieldShadowing"),
+	GHeightFieldShadowing,
+	TEXT("Whether the height field shadowing feature is allowed."),
+	ECVF_Scalability | ECVF_RenderThreadSafe);
+
+static int32 GHFShadowQuality = 3;
+FAutoConsoleVariableRef CVarHFShadowQuality(
+	TEXT("r.HFShadowQuality"),
+	GHFShadowQuality,
+	TEXT("Defines the height field shadow method which allows to adjust for quality or performance.\n")
+	TEXT(" 0:off, 1:low (8 steps), 2:medium (16 steps), 3:high (32 steps, default)"),
+	ECVF_Scalability | ECVF_RenderThreadSafe);
+
+static float GMinDirectionalLightAngleForRTHF = 27.f;
+static FAutoConsoleVariableRef CVarMinDirectionalLightAngleForRTHF(
+	TEXT("r.Shadow.MinDirectionalLightAngleForRTHF"),
+	GMinDirectionalLightAngleForRTHF,
+	TEXT(""),
+	ECVF_RenderThreadSafe);
+
+int32 GAverageHeightFieldObjectsPerShadowCullTile = 16;
+FAutoConsoleVariableRef CVarAverageHeightFieldObjectsPerShadowCullTile(
+	TEXT("r.HFShadowAverageObjectsPerCullTile"),
+	GAverageHeightFieldObjectsPerShadowCullTile,
+	TEXT("Determines how much memory should be allocated in height field object culling data structures.  Too much = memory waste, too little = flickering due to buffer overflow."),
+	ECVF_RenderThreadSafe | ECVF_ReadOnly);
+
 int32 const GDistanceFieldShadowTileSizeX = 8;
 int32 const GDistanceFieldShadowTileSizeY = 8;
 
@@ -99,10 +128,12 @@ FIntPoint GetBufferSizeForDFShadows()
 }
 
 TGlobalResource<FDistanceFieldObjectBufferResource> GShadowCulledObjectBuffers;
+TGlobalResource<FHeightFieldObjectBufferResource> GShadowCulledHeightFieldObjectBuffers;
 
-class FCullObjectsForShadowCS : public FGlobalShader
+template <EDistanceFieldPrimitiveType PrimitiveType>
+class TCullObjectsForShadowCS : public FGlobalShader
 {
-	DECLARE_SHADER_TYPE(FCullObjectsForShadowCS,Global)
+	DECLARE_SHADER_TYPE(TCullObjectsForShadowCS, Global);
 public:
 
 	static bool ShouldCompilePermutation(const FGlobalShaderPermutationParameters& Parameters)
@@ -114,9 +145,10 @@ public:
 	{
 		FGlobalShader::ModifyCompilationEnvironment(Parameters,OutEnvironment);
 		OutEnvironment.SetDefine(TEXT("UPDATEOBJECTS_THREADGROUP_SIZE"), UpdateObjectsGroupSize);
+		OutEnvironment.SetDefine(TEXT("DISTANCEFIELD_PRIMITIVE_TYPE"), (int32)PrimitiveType);
 	}
 
-	FCullObjectsForShadowCS(const ShaderMetaType::CompiledShaderInitializerType& Initializer)
+	TCullObjectsForShadowCS(const ShaderMetaType::CompiledShaderInitializerType& Initializer)
 		: FGlobalShader(Initializer)
 	{
 		ObjectBufferParameters.Bind(Initializer.ParameterMap);
@@ -128,24 +160,59 @@ public:
 		ShadowConvexHull.Bind(Initializer.ParameterMap, TEXT("ShadowConvexHull"));
 	}
 
-	FCullObjectsForShadowCS()
+	TCullObjectsForShadowCS()
 	{
 	}
 
-	void SetParameters(FRHICommandList& RHICmdList, const FScene* Scene, const FSceneView& View, const FMatrix& WorldToShadowValue, int32 NumPlanes, const FPlane* PlaneData, const FVector4& ShadowBoundingSphereValue)
+	void SetParameters(
+		FRHICommandList& RHICmdList,
+		const FScene* Scene,
+		const FSceneView& View,
+		const FMatrix& WorldToShadowValue,
+		int32 NumPlanes,
+		const FPlane* PlaneData,
+		const FVector4& ShadowBoundingSphereValue)
 	{
+		constexpr bool bIsHeightField = PrimitiveType == DFPT_HeightField;
+		const FDistanceFieldSceneData& SceneData = Scene->DistanceFieldSceneData;
+		const auto& SceneObjectBuffers = TSelector<bIsHeightField>()(*SceneData.GetHeightFieldObjectBuffers(), *SceneData.GetCurrentObjectBuffers());
+		auto& CulledObjectBuffers = TSelector<bIsHeightField>()(GShadowCulledHeightFieldObjectBuffers, GShadowCulledObjectBuffers);
+		
+		int32 NumObjectsInBuffer;
+		FRHITexture* TextureAtlas;
+		int32 AtlasSizeX;
+		int32 AtlasSizeY;
+		int32 AtlasSizeZ;
+
+		if (bIsHeightField)
+		{
+			NumObjectsInBuffer = SceneData.NumHeightFieldObjectsInBuffer;
+			TextureAtlas = GHeightFieldTextureAtlas.GetAtlasTexture();
+			AtlasSizeX = GHeightFieldTextureAtlas.GetSizeX();
+			AtlasSizeY = GHeightFieldTextureAtlas.GetSizeY();
+			AtlasSizeZ = 1;
+		}
+		else
+		{
+			NumObjectsInBuffer = SceneData.NumObjectsInBuffer;
+			TextureAtlas = GDistanceFieldVolumeTextureAtlas.VolumeTextureRHI;
+			AtlasSizeX = GDistanceFieldVolumeTextureAtlas.GetSizeX();
+			AtlasSizeY = GDistanceFieldVolumeTextureAtlas.GetSizeY();
+			AtlasSizeZ = GDistanceFieldVolumeTextureAtlas.GetSizeZ();
+		}
+
 		FRHIComputeShader* ShaderRHI = GetComputeShader();
 		FGlobalShader::SetParameters<FViewUniformShaderParameters>(RHICmdList, ShaderRHI, View.ViewUniformBuffer);
-		ObjectBufferParameters.Set(RHICmdList, ShaderRHI, *(Scene->DistanceFieldSceneData.GetCurrentObjectBuffers()), Scene->DistanceFieldSceneData.NumObjectsInBuffer);
+		ObjectBufferParameters.Set(RHICmdList, ShaderRHI, SceneObjectBuffers, NumObjectsInBuffer, TextureAtlas, AtlasSizeX, AtlasSizeY, AtlasSizeZ);
 
 		FRHIUnorderedAccessView* OutUAVs[4];
-		OutUAVs[0] = GShadowCulledObjectBuffers.Buffers.ObjectIndirectArguments.UAV;
-		OutUAVs[1] = GShadowCulledObjectBuffers.Buffers.Bounds.UAV;
-		OutUAVs[2] = GShadowCulledObjectBuffers.Buffers.Data.UAV;
-		OutUAVs[3] = GShadowCulledObjectBuffers.Buffers.BoxBounds.UAV;		
+		OutUAVs[0] = CulledObjectBuffers.Buffers.ObjectIndirectArguments.UAV;
+		OutUAVs[1] = CulledObjectBuffers.Buffers.Bounds.UAV;
+		OutUAVs[2] = CulledObjectBuffers.Buffers.Data.UAV;
+		OutUAVs[3] = CulledObjectBuffers.Buffers.BoxBounds.UAV;
 		RHICmdList.TransitionResources(EResourceTransitionAccess::ERWBarrier, EResourceTransitionPipeline::EComputeToCompute, OutUAVs, UE_ARRAY_COUNT(OutUAVs));
 
-		CulledObjectParameters.Set(RHICmdList, ShaderRHI, GShadowCulledObjectBuffers.Buffers);
+		CulledObjectParameters.Set(RHICmdList, ShaderRHI, CulledObjectBuffers.Buffers, TextureAtlas, AtlasSizeX, AtlasSizeY, AtlasSizeZ);
 
 		SetShaderValue(RHICmdList, ShaderRHI, ObjectBoundingGeometryIndexCount, UE_ARRAY_COUNT(GCubeIndices));
 		SetShaderValue(RHICmdList, ShaderRHI, WorldToShadow, WorldToShadowValue);
@@ -164,14 +231,19 @@ public:
 
 	void UnsetParameters(FRHICommandList& RHICmdList, const FScene* Scene)
 	{
-		ObjectBufferParameters.UnsetParameters(RHICmdList, GetComputeShader(), *(Scene->DistanceFieldSceneData.GetCurrentObjectBuffers()));
+		constexpr bool bIsHeightField = PrimitiveType == DFPT_HeightField;
+		const FDistanceFieldSceneData& SceneData = Scene->DistanceFieldSceneData;
+		const auto& SceneObjectBuffers = TSelector<bIsHeightField>()(*SceneData.GetHeightFieldObjectBuffers(), *SceneData.GetCurrentObjectBuffers());
+		auto& CulledObjectBuffers = TSelector<bIsHeightField>()(GShadowCulledHeightFieldObjectBuffers, GShadowCulledObjectBuffers);
+
+		ObjectBufferParameters.UnsetParameters(RHICmdList, GetComputeShader(), SceneObjectBuffers);
 		CulledObjectParameters.UnsetParameters(RHICmdList, GetComputeShader());
 
 		FRHIUnorderedAccessView* OutUAVs[4];
-		OutUAVs[0] = GShadowCulledObjectBuffers.Buffers.ObjectIndirectArguments.UAV;
-		OutUAVs[1] = GShadowCulledObjectBuffers.Buffers.Bounds.UAV;
-		OutUAVs[2] = GShadowCulledObjectBuffers.Buffers.Data.UAV;
-		OutUAVs[3] = GShadowCulledObjectBuffers.Buffers.BoxBounds.UAV;
+		OutUAVs[0] = CulledObjectBuffers.Buffers.ObjectIndirectArguments.UAV;
+		OutUAVs[1] = CulledObjectBuffers.Buffers.Bounds.UAV;
+		OutUAVs[2] = CulledObjectBuffers.Buffers.Data.UAV;
+		OutUAVs[3] = CulledObjectBuffers.Buffers.BoxBounds.UAV;
 		RHICmdList.TransitionResources(EResourceTransitionAccess::EReadable, EResourceTransitionPipeline::EComputeToCompute, OutUAVs, UE_ARRAY_COUNT(OutUAVs));
 	}
 
@@ -190,8 +262,8 @@ public:
 
 private:
 
-	FDistanceFieldObjectBufferParameters ObjectBufferParameters;
-	FDistanceFieldCulledObjectBufferParameters CulledObjectParameters;
+	TDistanceFieldObjectBufferParameters<PrimitiveType> ObjectBufferParameters;
+	TDistanceFieldCulledObjectBufferParameters<PrimitiveType> CulledObjectParameters;
 	FShaderParameter ObjectBoundingGeometryIndexCount;
 	FShaderParameter WorldToShadow;
 	FShaderParameter NumShadowHullPlanes;
@@ -199,12 +271,14 @@ private:
 	FShaderParameter ShadowConvexHull;
 };
 
-IMPLEMENT_SHADER_TYPE(,FCullObjectsForShadowCS,TEXT("/Engine/Private/DistanceFieldShadowing.usf"),TEXT("CullObjectsForShadowCS"),SF_Compute);
+IMPLEMENT_SHADER_TYPE(template<>, TCullObjectsForShadowCS<DFPT_SignedDistanceField>, TEXT("/Engine/Private/DistanceFieldShadowing.usf"), TEXT("CullObjectsForShadowCS"), SF_Compute);
+IMPLEMENT_SHADER_TYPE(template<>, TCullObjectsForShadowCS<DFPT_HeightField>, TEXT("/Engine/Private/DistanceFieldShadowing.usf"), TEXT("CullObjectsForShadowCS"), SF_Compute);
 
 /**  */
-class FShadowObjectCullVS : public FGlobalShader
+template <EDistanceFieldPrimitiveType PrimitiveType>
+class TShadowObjectCullVS : public FGlobalShader
 {
-	DECLARE_SHADER_TYPE(FShadowObjectCullVS,Global);
+	DECLARE_SHADER_TYPE(TShadowObjectCullVS,Global);
 public:
 
 	static bool ShouldCompilePermutation(const FGlobalShaderPermutationParameters& Parameters) 
@@ -212,7 +286,13 @@ public:
 		return IsFeatureLevelSupported(Parameters.Platform, ERHIFeatureLevel::SM5) && DoesPlatformSupportDistanceFieldShadowing(Parameters.Platform); 
 	}
 
-	FShadowObjectCullVS(const ShaderMetaType::CompiledShaderInitializerType& Initializer):
+	static void ModifyCompilationEnvironment(const FGlobalShaderPermutationParameters& Parameters, FShaderCompilerEnvironment& OutEnvironment)
+	{
+		FGlobalShader::ModifyCompilationEnvironment(Parameters, OutEnvironment);
+		OutEnvironment.SetDefine(TEXT("DISTANCEFIELD_PRIMITIVE_TYPE"), (int32)PrimitiveType);
+	}
+
+	TShadowObjectCullVS(const ShaderMetaType::CompiledShaderInitializerType& Initializer):
 		FGlobalShader(Initializer)
 	{
 		ObjectParameters.Bind(Initializer.ParameterMap);
@@ -220,18 +300,41 @@ public:
 		MinExpandRadius.Bind(Initializer.ParameterMap, TEXT("MinExpandRadius"));
 	}
 
-	FShadowObjectCullVS() {}
+	TShadowObjectCullVS() {}
 
 	void SetParameters(FRHICommandList& RHICmdList, const FSceneView& View, FVector2D NumGroupsValue, const FMatrix& WorldToShadowMatrixValue, float ShadowRadius)
 	{
 		FRHIVertexShader* ShaderRHI = GetVertexShader();
 		FGlobalShader::SetParameters<FViewUniformShaderParameters>(RHICmdList, ShaderRHI, View.ViewUniformBuffer);
 		
-		ObjectParameters.Set(RHICmdList, ShaderRHI, GShadowCulledObjectBuffers.Buffers);
+		constexpr bool bIsHeightField = PrimitiveType == DFPT_HeightField;
+		auto& CulledObjectBuffers = TSelector<bIsHeightField>()(GShadowCulledHeightFieldObjectBuffers, GShadowCulledObjectBuffers);
+
+		FRHITexture* TextureAtlas;
+		int32 AtlasSizeX;
+		int32 AtlasSizeY;
+		int32 AtlasSizeZ;
+
+		if (bIsHeightField)
+		{
+			TextureAtlas = GHeightFieldTextureAtlas.GetAtlasTexture();
+			AtlasSizeX = GHeightFieldTextureAtlas.GetSizeX();
+			AtlasSizeY = GHeightFieldTextureAtlas.GetSizeY();
+			AtlasSizeZ = 1;
+		}
+		else
+		{
+			TextureAtlas = GDistanceFieldVolumeTextureAtlas.VolumeTextureRHI;
+			AtlasSizeX = GDistanceFieldVolumeTextureAtlas.GetSizeX();
+			AtlasSizeY = GDistanceFieldVolumeTextureAtlas.GetSizeY();
+			AtlasSizeZ = GDistanceFieldVolumeTextureAtlas.GetSizeZ();
+		}
+
+		ObjectParameters.Set(RHICmdList, ShaderRHI, CulledObjectBuffers.Buffers, TextureAtlas, AtlasSizeX, AtlasSizeY, AtlasSizeZ);
 
 		SetShaderValue(RHICmdList, ShaderRHI, WorldToShadow, WorldToShadowMatrixValue);
 
-		float MinExpandRadiusValue = 1.414f * ShadowRadius / FMath::Min(NumGroupsValue.X, NumGroupsValue.Y);
+		const float MinExpandRadiusValue = (bIsHeightField ? 0.87f : 1.414f) * ShadowRadius / FMath::Min(NumGroupsValue.X, NumGroupsValue.Y);
 		SetShaderValue(RHICmdList, ShaderRHI, MinExpandRadius, MinExpandRadiusValue);
 	}
 
@@ -245,17 +348,18 @@ public:
 	}
 
 private:
-	FDistanceFieldCulledObjectBufferParameters ObjectParameters;
+	TDistanceFieldCulledObjectBufferParameters<PrimitiveType> ObjectParameters;
 	FShaderParameter WorldToShadow;
 	FShaderParameter MinExpandRadius;
 };
 
-IMPLEMENT_SHADER_TYPE(,FShadowObjectCullVS,TEXT("/Engine/Private/DistanceFieldShadowing.usf"),TEXT("ShadowObjectCullVS"),SF_Vertex);
+IMPLEMENT_SHADER_TYPE(template<>, TShadowObjectCullVS<DFPT_SignedDistanceField>, TEXT("/Engine/Private/DistanceFieldShadowing.usf"), TEXT("ShadowObjectCullVS"), SF_Vertex);
+IMPLEMENT_SHADER_TYPE(template<>, TShadowObjectCullVS<DFPT_HeightField>, TEXT("/Engine/Private/DistanceFieldShadowing.usf"), TEXT("ShadowObjectCullVS"), SF_Vertex);
 
-template <bool bCountingPass>
-class FShadowObjectCullPS : public FGlobalShader
+template <bool bCountingPass, EDistanceFieldPrimitiveType PrimitiveType>
+class TShadowObjectCullPS : public FGlobalShader
 {
-	DECLARE_SHADER_TYPE(FShadowObjectCullPS, Global);
+	DECLARE_SHADER_TYPE(TShadowObjectCullPS, Global);
 public:
 
 	static bool ShouldCompilePermutation(const FGlobalShaderPermutationParameters& Parameters)
@@ -268,13 +372,14 @@ public:
 		FGlobalShader::ModifyCompilationEnvironment(Parameters,OutEnvironment);
 		FLightTileIntersectionParameters::ModifyCompilationEnvironment(Parameters.Platform, OutEnvironment);
 		OutEnvironment.SetDefine(TEXT("SCATTER_CULLING_COUNT_PASS"), bCountingPass ? 1 : 0);
+		OutEnvironment.SetDefine(TEXT("DISTANCEFIELD_PRIMITIVE_TYPE"), (int32)PrimitiveType);
 	}
 
 	/** Default constructor. */
-	FShadowObjectCullPS() {}
+	TShadowObjectCullPS() {}
 
 	/** Initialization constructor. */
-	FShadowObjectCullPS(const ShaderMetaType::CompiledShaderInitializerType& Initializer)
+	TShadowObjectCullPS(const ShaderMetaType::CompiledShaderInitializerType& Initializer)
 		: FGlobalShader(Initializer)
 	{
 		ObjectParameters.Bind(Initializer.ParameterMap);
@@ -289,7 +394,31 @@ public:
 		FRHIPixelShader* ShaderRHI = GetPixelShader();
 		FGlobalShader::SetParameters<FViewUniformShaderParameters>(RHICmdList, ShaderRHI, View.ViewUniformBuffer);
 
-		ObjectParameters.Set(RHICmdList, ShaderRHI, GShadowCulledObjectBuffers.Buffers);
+		constexpr bool bIsHeightField = PrimitiveType == DFPT_HeightField;
+		auto& CulledObjectBuffers = TSelector<bIsHeightField>()(GShadowCulledHeightFieldObjectBuffers, GShadowCulledObjectBuffers);
+
+		FRHITexture* TextureAtlas;
+		int32 AtlasSizeX;
+		int32 AtlasSizeY;
+		int32 AtlasSizeZ;
+
+		if (bIsHeightField)
+		{
+			TextureAtlas = GHeightFieldTextureAtlas.GetAtlasTexture();
+			AtlasSizeX = GHeightFieldTextureAtlas.GetSizeX();
+			AtlasSizeY = GHeightFieldTextureAtlas.GetSizeY();
+			AtlasSizeZ = 1;
+		}
+		else
+		{
+			TextureAtlas = GDistanceFieldVolumeTextureAtlas.VolumeTextureRHI;
+			AtlasSizeX = GDistanceFieldVolumeTextureAtlas.GetSizeX();
+			AtlasSizeY = GDistanceFieldVolumeTextureAtlas.GetSizeY();
+			AtlasSizeZ = GDistanceFieldVolumeTextureAtlas.GetSizeZ();
+		}
+
+		ObjectParameters.Set(RHICmdList, ShaderRHI, CulledObjectBuffers.Buffers, TextureAtlas, AtlasSizeX, AtlasSizeY, AtlasSizeZ);
+
 		LightTileIntersectionParameters.Set(RHICmdList, ShaderRHI, *TileIntersectionResources);
 	}
 
@@ -308,12 +437,20 @@ public:
 	}
 
 private:
-	FDistanceFieldCulledObjectBufferParameters ObjectParameters;
+	TDistanceFieldCulledObjectBufferParameters<PrimitiveType> ObjectParameters;
 	FLightTileIntersectionParameters LightTileIntersectionParameters;
 };
 
-IMPLEMENT_SHADER_TYPE(template<>,FShadowObjectCullPS<true>,TEXT("/Engine/Private/DistanceFieldShadowing.usf"),TEXT("ShadowObjectCullPS"),SF_Pixel);
-IMPLEMENT_SHADER_TYPE(template<>,FShadowObjectCullPS<false>,TEXT("/Engine/Private/DistanceFieldShadowing.usf"),TEXT("ShadowObjectCullPS"),SF_Pixel);
+#define VARIATION(PrimitiveType) \
+	typedef TShadowObjectCullPS<true, PrimitiveType> FShadowObjectCullPS##1##PrimitiveType; \
+	typedef TShadowObjectCullPS<false, PrimitiveType> FShadowObjectCullPS##0##PrimitiveType; \
+	IMPLEMENT_SHADER_TYPE(, FShadowObjectCullPS##1##PrimitiveType, TEXT("/Engine/Private/DistanceFieldShadowing.usf"), TEXT("ShadowObjectCullPS"), SF_Pixel); \
+	IMPLEMENT_SHADER_TYPE(, FShadowObjectCullPS##0##PrimitiveType, TEXT("/Engine/Private/DistanceFieldShadowing.usf"), TEXT("ShadowObjectCullPS"), SF_Pixel)
+
+VARIATION(DFPT_SignedDistanceField);
+VARIATION(DFPT_HeightField);
+
+#undef VARIATION
 
 enum EDistanceFieldShadowingType
 {
@@ -322,7 +459,7 @@ enum EDistanceFieldShadowingType
 	DFS_PointLightTiledCulling
 };
 
-template<EDistanceFieldShadowingType ShadowingType, uint32 DFShadowQuality>
+template<EDistanceFieldShadowingType ShadowingType, uint32 DFShadowQuality, EDistanceFieldPrimitiveType PrimitiveType, bool bHasPrevOutput>
 class TDistanceFieldShadowingCS : public FGlobalShader
 {
 	DECLARE_SHADER_TYPE(TDistanceFieldShadowingCS, Global);
@@ -341,6 +478,9 @@ public:
 		OutEnvironment.SetDefine(TEXT("SCATTER_TILE_CULLING"), ShadowingType == DFS_DirectionalLightScatterTileCulling);
 		OutEnvironment.SetDefine(TEXT("POINT_LIGHT"), ShadowingType == DFS_PointLightTiledCulling);
 		OutEnvironment.SetDefine(TEXT("DF_SHADOW_QUALITY"), DFShadowQuality);
+		OutEnvironment.SetDefine(TEXT("DISTANCEFIELD_PRIMITIVE_TYPE"), (int32)PrimitiveType);
+		OutEnvironment.SetDefine(TEXT("HAS_PREVIOUS_OUTPUT"), (int32)bHasPrevOutput);
+		OutEnvironment.SetDefine(TEXT("PLATFORM_SUPPORTS_TYPED_UAV_LOAD"), (int32)RHISupports4ComponentUAVReadWrite(Parameters.Platform));
 	}
 
 	/** Default constructor. */
@@ -366,6 +506,9 @@ public:
 		MinDepth.Bind(Initializer.ParameterMap, TEXT("MinDepth"));
 		MaxDepth.Bind(Initializer.ParameterMap, TEXT("MaxDepth"));
 		DownsampleFactor.Bind(Initializer.ParameterMap, TEXT("DownsampleFactor"));
+		InvOutputBufferSize.Bind(Initializer.ParameterMap, TEXT("InvOutputBufferSize"));
+		ShadowFactorsTexture.Bind(Initializer.ParameterMap, TEXT("ShadowFactorsTexture"));
+		ShadowFactorsSampler.Bind(Initializer.ParameterMap, TEXT("ShadowFactorsSampler"));
 	}
 
 	template<typename TRHICommandList>
@@ -376,7 +519,8 @@ public:
 		FSceneRenderTargetItem& ShadowFactorsValue, 
 		FVector2D NumGroupsValue,
 		const FIntRect& ScissorRect,
-		FLightTileIntersectionResources* TileIntersectionResources)
+		FLightTileIntersectionResources* TileIntersectionResources,
+		FRHITexture* PrevOutput)
 	{
 		FRHIComputeShader* ShaderRHI = GetComputeShader();
 
@@ -385,7 +529,31 @@ public:
 		RHICmdList.TransitionResource(EResourceTransitionAccess::ERWBarrier, EResourceTransitionPipeline::EComputeToCompute, ShadowFactorsValue.UAV);
 		ShadowFactors.SetTexture(RHICmdList, ShaderRHI, ShadowFactorsValue.ShaderResourceTexture, ShadowFactorsValue.UAV);
 
-		ObjectParameters.Set(RHICmdList, ShaderRHI, GShadowCulledObjectBuffers.Buffers);
+		constexpr bool bIsHeightField = PrimitiveType == DFPT_HeightField;
+		auto& CulledObjectBuffers = TSelector<bIsHeightField>()(GShadowCulledHeightFieldObjectBuffers, GShadowCulledObjectBuffers);
+
+		FRHITexture* TextureAtlas;
+		int32 AtlasSizeX;
+		int32 AtlasSizeY;
+		int32 AtlasSizeZ;
+
+		if (bIsHeightField)
+		{
+			TextureAtlas = GHeightFieldTextureAtlas.GetAtlasTexture();
+			AtlasSizeX = GHeightFieldTextureAtlas.GetSizeX();
+			AtlasSizeY = GHeightFieldTextureAtlas.GetSizeY();
+			AtlasSizeZ = 1;
+		}
+		else
+		{
+			TextureAtlas = GDistanceFieldVolumeTextureAtlas.VolumeTextureRHI;
+			AtlasSizeX = GDistanceFieldVolumeTextureAtlas.GetSizeX();
+			AtlasSizeY = GDistanceFieldVolumeTextureAtlas.GetSizeY();
+			AtlasSizeZ = GDistanceFieldVolumeTextureAtlas.GetSizeZ();
+		}
+
+		ObjectParameters.Set(RHICmdList, ShaderRHI, CulledObjectBuffers.Buffers, TextureAtlas, AtlasSizeX, AtlasSizeY, AtlasSizeZ);
+
 		SceneTextureParameters.Set(RHICmdList, ShaderRHI, View.FeatureLevel, ESceneTextureSetupMode::All);
 
 		SetShaderValue(RHICmdList, ShaderRHI, NumGroups, NumGroupsValue);
@@ -403,7 +571,9 @@ public:
 
 		SetShaderValue(RHICmdList, ShaderRHI, RayStartOffsetDepthScale, LightProxy.GetRayStartOffsetDepthScale());
 
-		const float LightSourceAngle = FMath::Clamp(LightProxy.GetLightSourceAngle(), 0.001f, 5.0f) * PI / 180.0f;
+		constexpr float MaxLightAngle = bIsHeightField ? 45.0f : 5.0f;
+		const float MinLightAngle = bIsHeightField ? FMath::Min(GMinDirectionalLightAngleForRTHF, MaxLightAngle) : 0.001f;
+		const float LightSourceAngle = FMath::Clamp(LightProxy.GetLightSourceAngle(), MinLightAngle, MaxLightAngle) * PI / 180.0f;
 		const FVector TanLightAngleAndNormalThresholdValue(FMath::Tan(LightSourceAngle), FMath::Cos(PI / 2 + LightSourceAngle), LightProxy.GetTraceDistance());
 		SetShaderValue(RHICmdList, ShaderRHI, TanLightAngleAndNormalThreshold, TanLightAngleAndNormalThresholdValue);
 
@@ -428,12 +598,22 @@ public:
 		}
 		else
 		{
+			check(!bIsHeightField);
 			//@todo - set these up for point lights as well
 			SetShaderValue(RHICmdList, ShaderRHI, MinDepth, 0.0f);
 			SetShaderValue(RHICmdList, ShaderRHI, MaxDepth, HALF_WORLD_MAX);
 		}
 
 		SetShaderValue(RHICmdList, ShaderRHI, DownsampleFactor, GetDFShadowDownsampleFactor());
+
+		if (InvOutputBufferSize.IsBound())
+		{
+			check(ShadowFactorsTexture.IsBound() && ShadowFactorsSampler.IsBound() && PrevOutput);
+
+			const FIntPoint OutputBufferSize = GetBufferSizeForDFShadows();
+			SetShaderValue(RHICmdList, ShaderRHI, InvOutputBufferSize, FVector2D(1.f / OutputBufferSize.X, 1.f / OutputBufferSize.Y));
+			SetTextureParameter(RHICmdList, ShaderRHI, ShadowFactorsTexture, ShadowFactorsSampler, TStaticSamplerState<>::GetRHI(), PrevOutput);
+		}
 	}
 
 	template<typename TRHICommandList>
@@ -463,6 +643,9 @@ public:
 		Ar << MinDepth;
 		Ar << MaxDepth;
 		Ar << DownsampleFactor;
+		Ar << InvOutputBufferSize;
+		Ar << ShadowFactorsTexture;
+		Ar << ShadowFactorsSampler;
 		return bShaderHasOutdatedParameters;
 	}
 
@@ -486,7 +669,7 @@ private:
 	FShaderParameter RayStartOffsetDepthScale;
 	FShaderParameter TanLightAngleAndNormalThreshold;
 	FShaderParameter ScissorRectMinAndSize;
-	FDistanceFieldCulledObjectBufferParameters ObjectParameters;
+	TDistanceFieldCulledObjectBufferParameters<PrimitiveType> ObjectParameters;
 	FSceneTextureShaderParameters SceneTextureParameters;
 	FLightTileIntersectionParameters LightTileIntersectionParameters;
 	FShaderParameter WorldToShadow;
@@ -494,20 +677,27 @@ private:
 	FShaderParameter MinDepth;
 	FShaderParameter MaxDepth;
 	FShaderParameter DownsampleFactor;
+	FShaderParameter InvOutputBufferSize;
+	FShaderResourceParameter ShadowFactorsTexture;
+	FShaderResourceParameter ShadowFactorsSampler;
 };
 
 // #define avoids a lot of code duplication
-#define VARIATION(A) \
-	typedef TDistanceFieldShadowingCS<A, 1> FDistanceFieldShadowingCS##A##1; \
-	typedef TDistanceFieldShadowingCS<A, 2> FDistanceFieldShadowingCS##A##2; \
-    typedef TDistanceFieldShadowingCS<A, 3> FDistanceFieldShadowingCS##A##3; \
-	IMPLEMENT_SHADER_TYPE2(FDistanceFieldShadowingCS##A##1, SF_Compute); \
-	IMPLEMENT_SHADER_TYPE2(FDistanceFieldShadowingCS##A##2, SF_Compute); \
-    IMPLEMENT_SHADER_TYPE2(FDistanceFieldShadowingCS##A##3, SF_Compute);
-
-VARIATION(DFS_DirectionalLightScatterTileCulling)
-VARIATION(DFS_DirectionalLightTiledCulling)
-VARIATION(DFS_PointLightTiledCulling)
+#define VARIATION(A, B, C) \
+	typedef TDistanceFieldShadowingCS<A, 1, B, C> FDistanceFieldShadowingCS##A##1##B##C; \
+	typedef TDistanceFieldShadowingCS<A, 2, B, C> FDistanceFieldShadowingCS##A##2##B##C; \
+    typedef TDistanceFieldShadowingCS<A, 3, B, C> FDistanceFieldShadowingCS##A##3##B##C; \
+	IMPLEMENT_SHADER_TYPE2(FDistanceFieldShadowingCS##A##1##B##C, SF_Compute); \
+	IMPLEMENT_SHADER_TYPE2(FDistanceFieldShadowingCS##A##2##B##C, SF_Compute); \
+    IMPLEMENT_SHADER_TYPE2(FDistanceFieldShadowingCS##A##3##B##C, SF_Compute)
+
+VARIATION(DFS_DirectionalLightScatterTileCulling, DFPT_SignedDistanceField, false);
+VARIATION(DFS_DirectionalLightScatterTileCulling, DFPT_HeightField, false);
+VARIATION(DFS_DirectionalLightScatterTileCulling, DFPT_HeightField, true);
+VARIATION(DFS_DirectionalLightTiledCulling, DFPT_SignedDistanceField, false);
+VARIATION(DFS_DirectionalLightTiledCulling, DFPT_HeightField, false);
+VARIATION(DFS_DirectionalLightTiledCulling, DFPT_HeightField, true);
+VARIATION(DFS_PointLightTiledCulling, DFPT_SignedDistanceField, false);
 
 #undef VARIATION
 
@@ -675,7 +865,7 @@ private:
 
 IMPLEMENT_SHADER_TYPE(,FComputeCulledObjectStartOffsetCS,TEXT("/Engine/Private/DistanceFieldShadowing.usf"),TEXT("ComputeCulledTilesStartOffsetCS"),SF_Compute);
 
-template<bool bCountingPass>
+template<bool bCountingPass, EDistanceFieldPrimitiveType PrimitiveType>
 void ScatterObjectsToShadowTiles(
 	FRHICommandListImmediate& RHICmdList, 
 	const FViewInfo& View, 
@@ -684,8 +874,8 @@ void ScatterObjectsToShadowTiles(
 	FIntPoint LightTileDimensions, 
 	FLightTileIntersectionResources* TileIntersectionResources)
 {
-	TShaderMapRef<FShadowObjectCullVS> VertexShader(View.ShaderMap);
-	TShaderMapRef<FShadowObjectCullPS<bCountingPass>> PixelShader(View.ShaderMap);
+	TShaderMapRef<TShadowObjectCullVS<PrimitiveType>> VertexShader(View.ShaderMap);
+	TShaderMapRef<TShadowObjectCullPS<bCountingPass, PrimitiveType>> PixelShader(View.ShaderMap);
 
 	TArray<FRHIUnorderedAccessView*> UAVs;
 	PixelShader->GetUAVs(View, TileIntersectionResources, UAVs);
@@ -738,7 +928,8 @@ void ScatterObjectsToShadowTiles(
 	RHICmdList.TransitionResources(EResourceTransitionAccess::EReadable, EResourceTransitionPipeline::EGfxToCompute, UAVs.GetData(), UAVs.Num());
 }
 
-void CullDistanceFieldObjectsForLight(
+template <EDistanceFieldPrimitiveType PrimitiveType>
+void CullDistanceFieldObjectsForLight_Internal(
 	FRHICommandListImmediate& RHICmdList,
 	const FViewInfo& View,
 	const FLightSceneProxy* LightSceneProxy, 
@@ -749,33 +940,38 @@ void CullDistanceFieldObjectsForLight(
 	float ShadowBoundingRadius,
 	TUniquePtr<FLightTileIntersectionResources>& TileIntersectionResources)
 {
+	constexpr bool bIsHeightfield = PrimitiveType == DFPT_HeightField;
 	const FScene* Scene = (const FScene*)(View.Family->Scene);
 
 	SCOPED_DRAW_EVENT(RHICmdList, CullObjectsForLight);
 
 	{
-		if (!GShadowCulledObjectBuffers.IsInitialized()  
-			|| GShadowCulledObjectBuffers.Buffers.MaxObjects < Scene->DistanceFieldSceneData.NumObjectsInBuffer
-			|| GShadowCulledObjectBuffers.Buffers.MaxObjects > 3 * Scene->DistanceFieldSceneData.NumObjectsInBuffer
+		const FDistanceFieldSceneData& SceneData = Scene->DistanceFieldSceneData;
+		const int32 NumObjectsInBuffer = bIsHeightfield ? SceneData.NumHeightFieldObjectsInBuffer : SceneData.NumObjectsInBuffer;
+		auto& CulledObjectBuffers = TSelector<bIsHeightfield>()(GShadowCulledHeightFieldObjectBuffers, GShadowCulledObjectBuffers);
+
+		if (!CulledObjectBuffers.IsInitialized()
+			|| CulledObjectBuffers.Buffers.MaxObjects < NumObjectsInBuffer
+			|| CulledObjectBuffers.Buffers.MaxObjects > 3 * NumObjectsInBuffer
 			|| GFastVRamConfig.bDirty)
 		{
-			GShadowCulledObjectBuffers.Buffers.bWantBoxBounds = true;
-			GShadowCulledObjectBuffers.Buffers.MaxObjects = Scene->DistanceFieldSceneData.NumObjectsInBuffer * 5 / 4;
-			GShadowCulledObjectBuffers.ReleaseResource();
-			GShadowCulledObjectBuffers.InitResource();
+			CulledObjectBuffers.Buffers.bWantBoxBounds = true;
+			CulledObjectBuffers.Buffers.MaxObjects = NumObjectsInBuffer * 5 / 4;
+			CulledObjectBuffers.ReleaseResource();
+			CulledObjectBuffers.InitResource();
 		}
-		GShadowCulledObjectBuffers.Buffers.AcquireTransientResource();
+		CulledObjectBuffers.Buffers.AcquireTransientResource();
 
 		{
-			SCOPED_DRAW_EVENTF(RHICmdList, CullObjectsToFrustum, TEXT("CullObjectsToFrustum SceneObjects %u"), Scene->DistanceFieldSceneData.NumObjectsInBuffer);
+			SCOPED_DRAW_EVENTF(RHICmdList, CullObjectsToFrustum, TEXT("CullObjectsToFrustum %sObjects %d"), bIsHeightfield ? TEXT("HeightField") : TEXT("DistanceField"), NumObjectsInBuffer);
 
-			ClearUAV(RHICmdList, GShadowCulledObjectBuffers.Buffers.ObjectIndirectArguments, 0);
+			ClearUAV(RHICmdList, CulledObjectBuffers.Buffers.ObjectIndirectArguments, 0);
 
-			TShaderMapRef<FCullObjectsForShadowCS> ComputeShader(GetGlobalShaderMap(Scene->GetFeatureLevel()));
+			TShaderMapRef<TCullObjectsForShadowCS<PrimitiveType>> ComputeShader(GetGlobalShaderMap(Scene->GetFeatureLevel()));
 			RHICmdList.SetComputeShader(ComputeShader->GetComputeShader());
 			ComputeShader->SetParameters(RHICmdList, Scene, View, WorldToShadowValue, NumPlanes, PlaneData, ShadowBoundingSphereValue);
 
-			DispatchComputeShader(RHICmdList, *ComputeShader, FMath::DivideAndRoundUp<uint32>(Scene->DistanceFieldSceneData.NumObjectsInBuffer, UpdateObjectsGroupSize), 1, 1);
+			DispatchComputeShader(RHICmdList, *ComputeShader, FMath::DivideAndRoundUp<uint32>(NumObjectsInBuffer, UpdateObjectsGroupSize), 1, 1);
 			ComputeShader->UnsetParameters(RHICmdList, Scene);
 		}
 	}
@@ -808,7 +1004,7 @@ void CullDistanceFieldObjectsForLight(
 
 			TileIntersectionResources->b16BitIndices = b16BitObjectIndices;
 			TileIntersectionResources->TileDimensions = LightTileDimensions;
-			TileIntersectionResources->Initialize();
+			TileIntersectionResources->Initialize(bIsHeightfield ? GAverageHeightFieldObjectsPerShadowCullTile : GAverageObjectsPerShadowCullTile);
 		}
 
 		check(TileIntersectionResources);
@@ -821,7 +1017,7 @@ void CullDistanceFieldObjectsForLight(
 			ClearUAV(RHICmdList, TileIntersectionResources->TileNumCulledObjects, 0);
 
 			// Rasterize object bounding shapes and intersect with shadow tiles to compute how many objects intersect each tile
-			ScatterObjectsToShadowTiles<true>(RHICmdList, View, WorldToShadowValue, ShadowBoundingRadius, LightTileDimensions, TileIntersectionResources.Get());
+			ScatterObjectsToShadowTiles<true, PrimitiveType>(RHICmdList, View, WorldToShadowValue, ShadowBoundingRadius, LightTileDimensions, TileIntersectionResources.Get());
 
 			ClearUAV(RHICmdList, TileIntersectionResources->NextStartOffset, 0);
 
@@ -843,16 +1039,67 @@ void CullDistanceFieldObjectsForLight(
 			ClearUAV(RHICmdList, TileIntersectionResources->TileNumCulledObjects, 0);
 
 			// Rasterize object bounding shapes and intersect with shadow tiles, and write out intersecting tile indices for the cone tracing pass
-			ScatterObjectsToShadowTiles<false>(RHICmdList, View, WorldToShadowValue, ShadowBoundingRadius, LightTileDimensions, TileIntersectionResources.Get());
+			ScatterObjectsToShadowTiles<false, PrimitiveType>(RHICmdList, View, WorldToShadowValue, ShadowBoundingRadius, LightTileDimensions, TileIntersectionResources.Get());
 		}
 	}
 }
 
+void CullDistanceFieldObjectsForLight(
+	FRHICommandListImmediate& RHICmdList,
+	const FViewInfo& View,
+	const FLightSceneProxy* LightSceneProxy,
+	const FMatrix& WorldToShadowValue,
+	int32 NumPlanes,
+	const FPlane* PlaneData,
+	const FVector4& ShadowBoundingSphereValue,
+	float ShadowBoundingRadius,
+	TUniquePtr<FLightTileIntersectionResources>& TileIntersectionResources)
+{
+	CullDistanceFieldObjectsForLight_Internal<DFPT_SignedDistanceField>(
+		RHICmdList,
+		View,
+		LightSceneProxy,
+		WorldToShadowValue,
+		NumPlanes,
+		PlaneData,
+		ShadowBoundingSphereValue,
+		ShadowBoundingRadius,
+		TileIntersectionResources);
+}
+
+void CullHeightFieldObjectsForLight(
+	FRHICommandListImmediate& RHICmdList,
+	const FViewInfo& View,
+	const FLightSceneProxy* LightSceneProxy,
+	const FMatrix& WorldToShadowValue,
+	int32 NumPlanes,
+	const FPlane* PlaneData,
+	const FVector4& ShadowBoundingSphereValue,
+	float ShadowBoundingRadius,
+	TUniquePtr<FLightTileIntersectionResources>& TileIntersectionResources)
+{
+	CullDistanceFieldObjectsForLight_Internal<DFPT_HeightField>(
+		RHICmdList,
+		View,
+		LightSceneProxy,
+		WorldToShadowValue,
+		NumPlanes,
+		PlaneData,
+		ShadowBoundingSphereValue,
+		ShadowBoundingRadius,
+		TileIntersectionResources);
+}
+
 int32 GetDFShadowQuality()
 {
 	return FMath::Clamp(GDFShadowQuality, 0, 3);
 }
 
+int32 GetHFShadowQuality()
+{
+	return FMath::Clamp(GHFShadowQuality, 0, 3);
+}
+
 bool SupportsDistanceFieldShadows(ERHIFeatureLevel::Type FeatureLevel, EShaderPlatform ShaderPlatform)
 {
 	return GDistanceFieldShadowing 
@@ -861,6 +1108,14 @@ bool SupportsDistanceFieldShadows(ERHIFeatureLevel::Type FeatureLevel, EShaderPl
 		&& DoesPlatformSupportDistanceFieldShadowing(ShaderPlatform);
 }
 
+bool SupportsHeightFieldShadows(ERHIFeatureLevel::Type FeatureLevel, EShaderPlatform ShaderPlatform)
+{
+	return GHeightFieldShadowing
+		&& GetHFShadowQuality() > 0
+		&& FeatureLevel >= ERHIFeatureLevel::SM5
+		&& DoesPlatformSupportDistanceFieldShadowing(ShaderPlatform);
+}
+
 bool FDeferredShadingSceneRenderer::ShouldPrepareForDistanceFieldShadows() const 
 {
 	bool bSceneHasRayTracedDFShadows = false;
@@ -892,9 +1147,23 @@ bool FDeferredShadingSceneRenderer::ShouldPrepareForDistanceFieldShadows() const
 		&& SupportsDistanceFieldShadows(Scene->GetFeatureLevel(), Scene->GetShaderPlatform());
 }
 
-template<typename TRHICommandList, EDistanceFieldShadowingType DFSType>
-void RayTraceShadowsDispatch(TRHICommandList& RHICmdList, const FViewInfo& View, FProjectedShadowInfo* ProjectedShadowInfo, FLightTileIntersectionResources* TileIntersectionResources)
+bool FDeferredShadingSceneRenderer::ShouldPrepareHeightFieldScene() const
+{
+	return Scene
+		&& ViewFamily.EngineShowFlags.DynamicShadows
+		&& SupportsHeightFieldShadows(Scene->GetFeatureLevel(), Scene->GetShaderPlatform());
+}
+
+template<typename TRHICommandList, EDistanceFieldShadowingType DFSType, EDistanceFieldPrimitiveType PrimitiveType, bool bHasPrevOutput>
+void RayTraceShadowsDispatch(
+	TRHICommandList& RHICmdList,
+	const FViewInfo& View,
+	FProjectedShadowInfo* ProjectedShadowInfo,
+	FLightTileIntersectionResources* TileIntersectionResources,
+	FRHITexture* PrevOutput = nullptr)
 {
+	check(DFSType != DFS_PointLightTiledCulling || PrimitiveType != DFPT_HeightField);
+
 	FIntRect ScissorRect;
 	if (!ProjectedShadowInfo->GetLightSceneInfo().Proxy->GetScissorRect(ScissorRect, View, View.ViewRect))
 	{
@@ -908,26 +1177,26 @@ void RayTraceShadowsDispatch(TRHICommandList& RHICmdList, const FViewInfo& View,
 	{
 		RHICmdList.SetComputeShader(CS->GetComputeShader());
 		FSceneRenderTargetItem& RayTracedShadowsRTI = ProjectedShadowInfo->RayTracedShadowsRT->GetRenderTargetItem();
-		CS->SetParameters(RHICmdList, View, ProjectedShadowInfo, RayTracedShadowsRTI, FVector2D(GroupSizeX, GroupSizeY), ScissorRect, TileIntersectionResources);
+		CS->SetParameters(RHICmdList, View, ProjectedShadowInfo, RayTracedShadowsRTI, FVector2D(GroupSizeX, GroupSizeY), ScissorRect, TileIntersectionResources, PrevOutput);
 		DispatchComputeShader(RHICmdList, *CS, GroupSizeX, GroupSizeY, 1);
 		CS->UnsetParameters(RHICmdList, RayTracedShadowsRTI);
 	};
 
-	int32 const DFShadowQuality = GetDFShadowQuality();
-	
+	int32 const DFShadowQuality = PrimitiveType == DFPT_HeightField ? GetHFShadowQuality() : GetDFShadowQuality();
+
 	if (DFShadowQuality == 1)
 	{
-		TShaderMapRef<TDistanceFieldShadowingCS<DFSType, 1> > ComputeShader(View.ShaderMap);
+		TShaderMapRef<TDistanceFieldShadowingCS<DFSType, 1, PrimitiveType, bHasPrevOutput>> ComputeShader(View.ShaderMap);
 		DispatchTemplatedCS(ComputeShader);
 	}
 	else if (DFShadowQuality == 2)
 	{
-		TShaderMapRef<TDistanceFieldShadowingCS<DFSType, 2> > ComputeShader(View.ShaderMap);
+		TShaderMapRef<TDistanceFieldShadowingCS<DFSType, 2, PrimitiveType, bHasPrevOutput>> ComputeShader(View.ShaderMap);
 		DispatchTemplatedCS(ComputeShader);
 	}
 	else if (DFShadowQuality == 3)
-	{ 
-		TShaderMapRef<TDistanceFieldShadowingCS<DFSType, 3> > ComputeShader(View.ShaderMap);
+	{
+		TShaderMapRef<TDistanceFieldShadowingCS<DFSType, 3, PrimitiveType, bHasPrevOutput>> ComputeShader(View.ShaderMap);
 		DispatchTemplatedCS(ComputeShader);
 	}
 }
@@ -938,28 +1207,59 @@ void RayTraceShadows(TRHICommandList& RHICmdList, const FViewInfo& View, FProjec
 	int32 const DFShadowQuality = GetDFShadowQuality();
 	if (ProjectedShadowInfo->bDirectionalLight && GShadowScatterTileCulling)
 	{
-		RayTraceShadowsDispatch<TRHICommandList, DFS_DirectionalLightScatterTileCulling>(RHICmdList, View, ProjectedShadowInfo, TileIntersectionResources);
+		RayTraceShadowsDispatch<TRHICommandList, DFS_DirectionalLightScatterTileCulling, DFPT_SignedDistanceField, false>(RHICmdList, View, ProjectedShadowInfo, TileIntersectionResources);
 	}
 	else if (ProjectedShadowInfo->bDirectionalLight)
 	{
-		RayTraceShadowsDispatch<TRHICommandList, DFS_DirectionalLightTiledCulling>(RHICmdList, View, ProjectedShadowInfo, TileIntersectionResources);
+		RayTraceShadowsDispatch<TRHICommandList, DFS_DirectionalLightTiledCulling, DFPT_SignedDistanceField, false>(RHICmdList, View, ProjectedShadowInfo, TileIntersectionResources);
 	}
 	else
 	{
-		RayTraceShadowsDispatch<TRHICommandList, DFS_PointLightTiledCulling>(RHICmdList, View, ProjectedShadowInfo, TileIntersectionResources);
+		RayTraceShadowsDispatch<TRHICommandList, DFS_PointLightTiledCulling, DFPT_SignedDistanceField, false>(RHICmdList, View, ProjectedShadowInfo, TileIntersectionResources);
+	}
+}
+
+template <typename TRHICommandList>
+void RayTraceHeightFieldShadows(
+	TRHICommandList& RHICmdList,
+	const FViewInfo& View,
+	FProjectedShadowInfo* ProjectedShadowInfo,
+	FLightTileIntersectionResources* TileIntersectionResources,
+	bool bHasPrevOutput,
+	FRHITexture* PrevOutput)
+{
+	check(ProjectedShadowInfo->bDirectionalLight);
+
+	if (GShadowScatterTileCulling && bHasPrevOutput)
+	{
+		RayTraceShadowsDispatch<TRHICommandList, DFS_DirectionalLightScatterTileCulling, DFPT_HeightField, true>(RHICmdList, View, ProjectedShadowInfo, TileIntersectionResources, PrevOutput);
+	}
+	else if (GShadowScatterTileCulling)
+	{
+		RayTraceShadowsDispatch<TRHICommandList, DFS_DirectionalLightScatterTileCulling, DFPT_HeightField, false>(RHICmdList, View, ProjectedShadowInfo, TileIntersectionResources);
+	}
+	else if (bHasPrevOutput)
+	{
+		RayTraceShadowsDispatch<TRHICommandList, DFS_DirectionalLightTiledCulling, DFPT_HeightField, true>(RHICmdList, View, ProjectedShadowInfo, TileIntersectionResources, PrevOutput);
+	}
+	else
+	{
+		RayTraceShadowsDispatch<TRHICommandList, DFS_DirectionalLightTiledCulling, DFPT_HeightField, false>(RHICmdList, View, ProjectedShadowInfo, TileIntersectionResources);
 	}
 }
 
 void FProjectedShadowInfo::BeginRenderRayTracedDistanceFieldProjection(FRHICommandListImmediate& RHICmdList, const FViewInfo& View)
 {
-	if (SupportsDistanceFieldShadows(View.GetFeatureLevel(), View.GetShaderPlatform())
-		&& View.Family->EngineShowFlags.RayTracedDistanceFieldShadows)
+	const bool bDFShadowSupported = SupportsDistanceFieldShadows(View.GetFeatureLevel(), View.GetShaderPlatform());
+	const bool bHFShadowSupported = SupportsHeightFieldShadows(View.GetFeatureLevel(), View.GetShaderPlatform());
+	const bool bBufferAliasingEnabled = IsTransientResourceBufferAliasingEnabled();
+	const FScene* Scene = (const FScene*)(View.Family->Scene);
+
+	if (bDFShadowSupported && View.Family->EngineShowFlags.RayTracedDistanceFieldShadows)
 	{
 		QUICK_SCOPE_CYCLE_COUNTER(STAT_BeginRenderRayTracedDistanceFieldShadows);
 		SCOPED_DRAW_EVENT(RHICmdList, BeginRayTracedDistanceFieldShadow);
 
-		const FScene* Scene = (const FScene*)(View.Family->Scene);
-
 		if (GDistanceFieldVolumeTextureAtlas.VolumeTextureRHI && Scene->DistanceFieldSceneData.NumObjectsInBuffer > 0)
 		{
 			check(!Scene->DistanceFieldSceneData.HasPendingOperations());
@@ -1016,6 +1316,69 @@ void FProjectedShadowInfo::BeginRenderRayTracedDistanceFieldProjection(FRHIComma
 			UnbindRenderTargets(RHICmdList);
 
 			RayTraceShadows(RHICmdList, View, this, TileIntersectionResources);
+
+			if (bBufferAliasingEnabled)
+			{
+				GShadowCulledObjectBuffers.Buffers.DiscardTransientResource();
+			}
+		}
+	}
+
+	if (bDirectionalLight
+		&& View.Family->EngineShowFlags.RayTracedDistanceFieldShadows
+		&& GHeightFieldTextureAtlas.GetAtlasTexture()
+		&& Scene->DistanceFieldSceneData.NumHeightFieldObjectsInBuffer > 0
+		&& bHFShadowSupported)
+	{
+		QUICK_SCOPE_CYCLE_COUNTER(STAT_BeginRenderRayTracedHeightFieldShadows);
+		SCOPED_DRAW_EVENT(RHICmdList, BeginRenderRayTracedHeightFieldShadows);
+
+		check(!Scene->DistanceFieldSceneData.HasPendingHeightFieldOperations());
+
+		UnbindRenderTargets(RHICmdList);
+
+		const int32 NumPlanes = CascadeSettings.ShadowBoundsAccurate.Planes.Num();
+		const FPlane* PlaneData = CascadeSettings.ShadowBoundsAccurate.Planes.GetData();
+		const FVector4 ShadowBoundingSphereValue(0.f, 0.f, 0.f, 0.f);
+		const FMatrix WorldToShadowValue = FTranslationMatrix(PreShadowTranslation) * SubjectAndReceiverMatrix;
+
+		CullHeightFieldObjectsForLight(
+			RHICmdList,
+			View,
+			LightSceneInfo->Proxy,
+			WorldToShadowValue,
+			NumPlanes,
+			PlaneData,
+			ShadowBoundingSphereValue,
+			ShadowBounds.W,
+			LightSceneInfo->HeightFieldTileIntersectionResources);
+
+		FLightTileIntersectionResources* TileIntersectionResources = LightSceneInfo->HeightFieldTileIntersectionResources.Get();
+		const bool bHasPrevOutput = RayTracedShadowsRT;
+		TRefCountPtr<IPooledRenderTarget> PrevOutputRT;
+
+		if (!RHISupports4ComponentUAVReadWrite(View.GetShaderPlatform()))
+		{
+			PrevOutputRT = RayTracedShadowsRT;
+			RayTracedShadowsRT = nullptr;
+		}
+
+		if (!RayTracedShadowsRT)
+		{
+			const FIntPoint BufferSize = GetBufferSizeForDFShadows();
+			FPooledRenderTargetDesc Desc(FPooledRenderTargetDesc::Create2DDesc(BufferSize, PF_G16R16F, FClearValueBinding::None, TexCreate_None, TexCreate_RenderTargetable | TexCreate_UAV, false));
+			Desc.Flags |= GFastVRamConfig.DistanceFieldShadows;
+			GRenderTargetPool.FindFreeElement(RHICmdList, Desc, RayTracedShadowsRT, TEXT("RayTracedShadows"));
+		}
+
+		SCOPED_DRAW_EVENT(RHICmdList, RayTraceHeightFieldShadows);
+		UnbindRenderTargets(RHICmdList);
+
+		RayTraceHeightFieldShadows(RHICmdList, View, this, TileIntersectionResources, bHasPrevOutput, PrevOutputRT ? PrevOutputRT->GetRenderTargetItem().ShaderResourceTexture : nullptr);
+
+		if (bBufferAliasingEnabled)
+		{
+			GShadowCulledHeightFieldObjectBuffers.Buffers.DiscardTransientResource();
 		}
 	}
 }
@@ -1031,11 +1394,6 @@ void FProjectedShadowInfo::RenderRayTracedDistanceFieldProjection(FRHICommandLis
 		QUICK_SCOPE_CYCLE_COUNTER(STAT_RenderRayTracedDistanceFieldShadows);
 		SCOPED_DRAW_EVENT(RHICmdList, RayTracedDistanceFieldShadow);
 
-		if ( IsTransientResourceBufferAliasingEnabled() )
-		{
-			GShadowCulledObjectBuffers.Buffers.DiscardTransientResource();
-		}
-
 		{
 			FRHIRenderPassInfo RPInfo(ScreenShadowMaskTexture->GetRenderTargetItem().TargetableTexture, ERenderTargetActions::Load_Store);
 			RPInfo.DepthStencilRenderTarget.Action = MakeDepthStencilTargetActions(ERenderTargetActions::Load_DontStore, ERenderTargetActions::Load_Store);
diff --git a/Engine/Source/Runtime/Renderer/Private/DistanceFieldVisualization.cpp b/Engine/Source/Runtime/Renderer/Private/DistanceFieldVisualization.cpp
--- a/Engine/Source/Runtime/Renderer/Private/DistanceFieldVisualization.cpp
+++ b/Engine/Source/Runtime/Renderer/Private/DistanceFieldVisualization.cpp
@@ -66,7 +66,18 @@ public:
 		RHICmdList.TransitionResource(EResourceTransitionAccess::ERWBarrier, EResourceTransitionPipeline::EComputeToCompute, VisualizeMeshDistanceFieldsValue.UAV);
 		VisualizeMeshDistanceFields.SetTexture(RHICmdList, ShaderRHI, VisualizeMeshDistanceFieldsValue.ShaderResourceTexture, VisualizeMeshDistanceFieldsValue.UAV);
 
-		ObjectParameters.Set(RHICmdList, ShaderRHI, GAOCulledObjectBuffers.Buffers);
+		FRHITexture* TextureAtlas;
+		int32 AtlasSizeX;
+		int32 AtlasSizeY;
+		int32 AtlasSizeZ;
+
+		TextureAtlas = GDistanceFieldVolumeTextureAtlas.VolumeTextureRHI;
+		AtlasSizeX = GDistanceFieldVolumeTextureAtlas.GetSizeX();
+		AtlasSizeY = GDistanceFieldVolumeTextureAtlas.GetSizeY();
+		AtlasSizeZ = GDistanceFieldVolumeTextureAtlas.GetSizeZ();
+
+		ObjectParameters.Set(RHICmdList, ShaderRHI, GAOCulledObjectBuffers.Buffers, TextureAtlas, AtlasSizeX, AtlasSizeY, AtlasSizeZ);
+
 		AOParameters.Set(RHICmdList, ShaderRHI, Parameters);
 		SceneTextureParameters.Set(RHICmdList, ShaderRHI, View.FeatureLevel, ESceneTextureSetupMode::All);
 
@@ -101,7 +112,7 @@ private:
 
 	FRWShaderParameter VisualizeMeshDistanceFields;
 	FShaderParameter NumGroups;
-	FDistanceFieldCulledObjectBufferParameters ObjectParameters;
+	TDistanceFieldCulledObjectBufferParameters<DFPT_SignedDistanceField> ObjectParameters;
 	FSceneTextureShaderParameters SceneTextureParameters;
 	FAOParameters AOParameters;
 	FGlobalDistanceFieldParameters GlobalDistanceFieldParameters;
diff --git a/Engine/Source/Runtime/Renderer/Private/GlobalDistanceField.cpp b/Engine/Source/Runtime/Renderer/Private/GlobalDistanceField.cpp
--- a/Engine/Source/Runtime/Renderer/Private/GlobalDistanceField.cpp
+++ b/Engine/Source/Runtime/Renderer/Private/GlobalDistanceField.cpp
@@ -192,7 +192,26 @@ public:
 	{
 		FRHIComputeShader* ShaderRHI = GetComputeShader();
 		FGlobalShader::SetParameters<FViewUniformShaderParameters>(RHICmdList, ShaderRHI, View.ViewUniformBuffer);
-		ObjectBufferParameters.Set(RHICmdList, ShaderRHI, *(Scene->DistanceFieldSceneData.GetCurrentObjectBuffers()), Scene->DistanceFieldSceneData.NumObjectsInBuffer);
+
+		FRHITexture* TextureAtlas;
+		int32 AtlasSizeX;
+		int32 AtlasSizeY;
+		int32 AtlasSizeZ;
+
+		TextureAtlas = GDistanceFieldVolumeTextureAtlas.VolumeTextureRHI;
+		AtlasSizeX = GDistanceFieldVolumeTextureAtlas.GetSizeX();
+		AtlasSizeY = GDistanceFieldVolumeTextureAtlas.GetSizeY();
+		AtlasSizeZ = GDistanceFieldVolumeTextureAtlas.GetSizeZ();
+
+		ObjectBufferParameters.Set(
+			RHICmdList,
+			ShaderRHI,
+			*(Scene->DistanceFieldSceneData.GetCurrentObjectBuffers()),
+			Scene->DistanceFieldSceneData.NumObjectsInBuffer,
+			TextureAtlas,
+			AtlasSizeX,
+			AtlasSizeY,
+			AtlasSizeZ);
 
 		FRHIUnorderedAccessView* OutUAVs[4];
 		OutUAVs[0] = GGlobalDistanceFieldCulledObjectBuffers.Buffers.ObjectIndirectArguments.UAV;
@@ -201,7 +220,7 @@ public:
 		OutUAVs[3] = GGlobalDistanceFieldCulledObjectBuffers.Buffers.BoxBounds.UAV;
 		RHICmdList.TransitionResources(EResourceTransitionAccess::ERWBarrier, EResourceTransitionPipeline::EComputeToCompute, OutUAVs, UE_ARRAY_COUNT(OutUAVs));
 
-		CulledObjectParameters.Set(RHICmdList, ShaderRHI, GGlobalDistanceFieldCulledObjectBuffers.Buffers);
+		CulledObjectParameters.Set(RHICmdList, ShaderRHI, GGlobalDistanceFieldCulledObjectBuffers.Buffers, TextureAtlas, AtlasSizeX, AtlasSizeY, AtlasSizeZ);
 
 		extern float GAOConeHalfAngle;
 		const float GlobalMaxSphereQueryRadius = MaxOcclusionDistance / (1 + FMath::Tan(GAOConeHalfAngle));
@@ -247,8 +266,8 @@ public:
 
 private:
 
-	FDistanceFieldObjectBufferParameters ObjectBufferParameters;
-	FDistanceFieldCulledObjectBufferParameters CulledObjectParameters;
+	TDistanceFieldObjectBufferParameters<DFPT_SignedDistanceField> ObjectBufferParameters;
+	TDistanceFieldCulledObjectBufferParameters<DFPT_SignedDistanceField> CulledObjectParameters;
 	FShaderParameter AOGlobalMaxSphereQueryRadius;
 	FShaderParameter VolumeBounds;
 	FShaderParameter AcceptOftenMovingObjectsOnly;
@@ -368,7 +387,18 @@ public:
 	{
 		FRHIComputeShader* ShaderRHI = GetComputeShader();
 		FGlobalShader::SetParameters<FViewUniformShaderParameters>(RHICmdList, ShaderRHI, View.ViewUniformBuffer);
-		CulledObjectBufferParameters.Set(RHICmdList, ShaderRHI, GGlobalDistanceFieldCulledObjectBuffers.Buffers);
+
+		FRHITexture* TextureAtlas;
+		int32 AtlasSizeX;
+		int32 AtlasSizeY;
+		int32 AtlasSizeZ;
+
+		TextureAtlas = GDistanceFieldVolumeTextureAtlas.VolumeTextureRHI;
+		AtlasSizeX = GDistanceFieldVolumeTextureAtlas.GetSizeX();
+		AtlasSizeY = GDistanceFieldVolumeTextureAtlas.GetSizeY();
+		AtlasSizeZ = GDistanceFieldVolumeTextureAtlas.GetSizeZ();
+
+		CulledObjectBufferParameters.Set(RHICmdList, ShaderRHI, GGlobalDistanceFieldCulledObjectBuffers.Buffers, TextureAtlas, AtlasSizeX, AtlasSizeY, AtlasSizeZ);
 		GlobalDistanceFieldParameters.Set(RHICmdList, ShaderRHI, GlobalDistanceFieldInfo.ParameterData);
 
 		RHICmdList.TransitionResource(EResourceTransitionAccess::ERWBarrier, EResourceTransitionPipeline::EComputeToCompute, GObjectGridBuffers.CullGridObjectNum.UAV);
@@ -416,7 +446,7 @@ public:
 
 private:
 
-	FDistanceFieldCulledObjectBufferParameters CulledObjectBufferParameters;
+	TDistanceFieldCulledObjectBufferParameters<DFPT_SignedDistanceField> CulledObjectBufferParameters;
 	FGlobalDistanceFieldParameters GlobalDistanceFieldParameters;
 	FRWShaderParameter CullGridObjectNum;
 	FRWShaderParameter CullGridObjectArray;
@@ -503,7 +533,18 @@ public:
 	{
 		FRHIComputeShader* ShaderRHI = GetComputeShader();
 		FGlobalShader::SetParameters<FViewUniformShaderParameters>(RHICmdList, ShaderRHI, View.ViewUniformBuffer);
-		CulledObjectBufferParameters.Set(RHICmdList, ShaderRHI, GGlobalDistanceFieldCulledObjectBuffers.Buffers);
+
+		FRHITexture* TextureAtlas;
+		int32 AtlasSizeX;
+		int32 AtlasSizeY;
+		int32 AtlasSizeZ;
+
+		TextureAtlas = GDistanceFieldVolumeTextureAtlas.VolumeTextureRHI;
+		AtlasSizeX = GDistanceFieldVolumeTextureAtlas.GetSizeX();
+		AtlasSizeY = GDistanceFieldVolumeTextureAtlas.GetSizeY();
+		AtlasSizeZ = GDistanceFieldVolumeTextureAtlas.GetSizeZ();
+
+		CulledObjectBufferParameters.Set(RHICmdList, ShaderRHI, GGlobalDistanceFieldCulledObjectBuffers.Buffers, TextureAtlas, AtlasSizeX, AtlasSizeY, AtlasSizeZ);
 		GlobalDistanceFieldParameters.Set(RHICmdList, ShaderRHI, ParameterData);
 
 		const FSceneRenderTargetItem& ClipMapRTI = Clipmap.RenderTarget->GetRenderTargetItem();
@@ -567,7 +608,7 @@ public:
 
 private:
 
-	FDistanceFieldCulledObjectBufferParameters CulledObjectBufferParameters;
+	TDistanceFieldCulledObjectBufferParameters<DFPT_SignedDistanceField> CulledObjectBufferParameters;
 	FGlobalDistanceFieldParameters GlobalDistanceFieldParameters;
 	FRWShaderParameter GlobalDistanceFieldTexture;
 	FShaderResourceParameter ParentGlobalDistanceFieldTexture;
diff --git a/Engine/Source/Runtime/Renderer/Private/HeightfieldLighting.cpp b/Engine/Source/Runtime/Renderer/Private/HeightfieldLighting.cpp
--- a/Engine/Source/Runtime/Renderer/Private/HeightfieldLighting.cpp
+++ b/Engine/Source/Runtime/Renderer/Private/HeightfieldLighting.cpp
@@ -945,7 +945,18 @@ public:
 		FGlobalShader::SetParameters<FViewUniformShaderParameters>(RHICmdList, ShaderRHI, View.ViewUniformBuffer);
 		HeightfieldDescriptionParameters.Set(RHICmdList, ShaderRHI, GetHeightfieldDescriptionsSRV(), NumHeightfieldsValue);
 		GlobalHeightfieldParameters.Set(RHICmdList, ShaderRHI, Atlas);
-		ObjectParameters.Set(RHICmdList, ShaderRHI, CulledObjectBuffers.Buffers);
+
+		FRHITexture* TextureAtlas;
+		int32 AtlasSizeX;
+		int32 AtlasSizeY;
+		int32 AtlasSizeZ;
+
+		TextureAtlas = GDistanceFieldVolumeTextureAtlas.VolumeTextureRHI;
+		AtlasSizeX = GDistanceFieldVolumeTextureAtlas.GetSizeX();
+		AtlasSizeY = GDistanceFieldVolumeTextureAtlas.GetSizeY();
+		AtlasSizeZ = GDistanceFieldVolumeTextureAtlas.GetSizeZ();
+
+		ObjectParameters.Set(RHICmdList, ShaderRHI, CulledObjectBuffers.Buffers, TextureAtlas, AtlasSizeX, AtlasSizeY, AtlasSizeZ);
 
 		SetShaderValue(RHICmdList, ShaderRHI, LightDirection, ProjectedShadowInfo->GetLightSceneInfo().Proxy->GetDirection());
 
@@ -987,7 +998,7 @@ private:
 
 	FHeightfieldDescriptionParameters HeightfieldDescriptionParameters;
 	FGlobalHeightfieldParameters GlobalHeightfieldParameters;
-	FDistanceFieldCulledObjectBufferParameters ObjectParameters;
+	TDistanceFieldCulledObjectBufferParameters<DFPT_SignedDistanceField> ObjectParameters;
 	FShaderParameter LightDirection;
 	FShaderParameter TanLightAngle;
 	FShaderParameter CascadeDepthMinMax;
diff --git a/Engine/Source/Runtime/Renderer/Private/LightSceneInfo.cpp b/Engine/Source/Runtime/Renderer/Private/LightSceneInfo.cpp
--- a/Engine/Source/Runtime/Renderer/Private/LightSceneInfo.cpp
+++ b/Engine/Source/Runtime/Renderer/Private/LightSceneInfo.cpp
@@ -41,6 +41,7 @@ FLightSceneInfo::FLightSceneInfo(FLightSceneProxy* InProxy, bool InbVisible)
 	, DynamicInteractionStaticPrimitiveList(NULL)
 	, Id(INDEX_NONE)
 	, TileIntersectionResources(nullptr)
+	, HeightFieldTileIntersectionResources(nullptr)
 	, DynamicShadowMapChannel(-1)
 	, bPrecomputedLightingIsValid(InProxy->GetLightComponent()->IsPrecomputedLightingValid())
 	, bVisible(InbVisible)
diff --git a/Engine/Source/Runtime/Renderer/Private/LightSceneInfo.h b/Engine/Source/Runtime/Renderer/Private/LightSceneInfo.h
--- a/Engine/Source/Runtime/Renderer/Private/LightSceneInfo.h
+++ b/Engine/Source/Runtime/Renderer/Private/LightSceneInfo.h
@@ -168,6 +168,7 @@ public:
 
 	/** Tile intersection buffer for distance field shadowing, stored on the light to avoid reallocating each frame. */
 	mutable TUniquePtr<class FLightTileIntersectionResources> TileIntersectionResources;
+	mutable TUniquePtr<class FLightTileIntersectionResources> HeightFieldTileIntersectionResources;
 
 	mutable FVertexBufferRHIRef ShadowCapsuleShapesVertexBuffer;
 	mutable FShaderResourceViewRHIRef ShadowCapsuleShapesSRV;
diff --git a/Engine/Source/Runtime/Renderer/Private/RendererScene.cpp b/Engine/Source/Runtime/Renderer/Private/RendererScene.cpp
--- a/Engine/Source/Runtime/Renderer/Private/RendererScene.cpp
+++ b/Engine/Source/Runtime/Renderer/Private/RendererScene.cpp
@@ -370,14 +370,18 @@ bool FPixelInspectorData::AddPixelInspectorRequest(FPixelInspectorRequest *Pixel
 
 FDistanceFieldSceneData::FDistanceFieldSceneData(EShaderPlatform ShaderPlatform) 
 	: NumObjectsInBuffer(0)
+	, NumHeightFieldObjectsInBuffer(0)
 	, ObjectBufferIndex(0)
 	, SurfelBuffers(NULL)
 	, InstancedSurfelBuffers(NULL)
 	, AtlasGeneration(0)
+	, HeightFieldAtlasGeneration(0)
 {
 	ObjectBuffers[0] = nullptr;
 	ObjectBuffers[1] = nullptr;
 
+	HeightFieldObjectBuffers = nullptr;
+
 	static const auto CVar = IConsoleManager::Get().FindTConsoleVariableDataInt(TEXT("r.GenerateMeshDistanceFields"));
 
 	bTrackAllPrimitives = (DoesPlatformSupportDistanceFieldAO(ShaderPlatform) || DoesPlatformSupportDistanceFieldShadowing(ShaderPlatform)) && CVar->GetValueOnGameThread() != 0 && IsUsingDistanceFields(ShaderPlatform);
@@ -393,7 +397,7 @@ FDistanceFieldSceneData::~FDistanceFieldSceneData()
 
 void FDistanceFieldSceneData::AddPrimitive(FPrimitiveSceneInfo* InPrimitive)
 {
-	const FPrimitiveSceneProxy* Proxy = InPrimitive->Proxy;
+	FPrimitiveSceneProxy* Proxy = InPrimitive->Proxy;
 
 	if ((bTrackAllPrimitives || Proxy->CastsDynamicIndirectShadow())
 		&& Proxy->CastsDynamicShadow()
@@ -401,10 +405,17 @@ void FDistanceFieldSceneData::AddPrimitive(FPrimitiveSceneInfo* InPrimitive)
 	{
 		if (Proxy->SupportsHeightfieldRepresentation())
 		{
-			HeightfieldPrimitives.Add(InPrimitive);
-			FBoxSphereBounds PrimitiveBounds = Proxy->GetBounds();
-			FGlobalDFCacheType CacheType = Proxy->IsOftenMoving() ? GDF_Full : GDF_MostlyStatic;
-			PrimitiveModifiedBounds[CacheType].Add(FVector4(PrimitiveBounds.Origin, PrimitiveBounds.SphereRadius));
+			// Theoretically, we may need to wait for high-res mips to be streamed in.
+			// In practice, however, those mips are usually resident at this point.
+			UTexture2D* HeightAndNormal;
+			UTexture2D* Unused0, *Unused1;
+			FHeightfieldComponentDescription Unused2(FMatrix::Identity);
+			Proxy->GetHeightfieldRepresentation(HeightAndNormal, Unused0, Unused1, Unused2);
+			GHeightFieldTextureAtlas.AddAllocation(HeightAndNormal);
+
+			checkSlow(!PendingHeightFieldAddOps.Contains(InPrimitive));
+			checkSlow(!PendingHeightFieldUpdateOps.Contains(InPrimitive));
+			PendingHeightFieldAddOps.Add(InPrimitive);
 		}
 
 		if (Proxy->SupportsDistanceFieldRepresentation())
@@ -436,7 +447,7 @@ void FDistanceFieldSceneData::UpdatePrimitive(FPrimitiveSceneInfo* InPrimitive)
 
 void FDistanceFieldSceneData::RemovePrimitive(FPrimitiveSceneInfo* InPrimitive)
 {
-	const FPrimitiveSceneProxy* Proxy = InPrimitive->Proxy;
+	FPrimitiveSceneProxy* Proxy = InPrimitive->Proxy;
 
 	if ((bTrackAllPrimitives || Proxy->CastsDynamicIndirectShadow()) 
 		&& Proxy->AffectsDistanceFieldLighting())
@@ -457,11 +468,21 @@ void FDistanceFieldSceneData::RemovePrimitive(FPrimitiveSceneInfo* InPrimitive)
 
 		if (Proxy->SupportsHeightfieldRepresentation())
 		{
-			HeightfieldPrimitives.Remove(InPrimitive);
+			UTexture2D* HeightAndNormal;
+			UTexture2D* Unused0, *Unused1;
+			FHeightfieldComponentDescription Unused2(FMatrix::Identity);
+			Proxy->GetHeightfieldRepresentation(HeightAndNormal, Unused0, Unused1, Unused2);
+			GHeightFieldTextureAtlas.RemoveAllocation(HeightAndNormal);
 
-			FBoxSphereBounds PrimitiveBounds = Proxy->GetBounds();
-			FGlobalDFCacheType CacheType = Proxy->IsOftenMoving() ? GDF_Full : GDF_MostlyStatic;
-			PrimitiveModifiedBounds[CacheType].Add(FVector4(PrimitiveBounds.Origin, PrimitiveBounds.SphereRadius));
+			PendingHeightFieldAddOps.Remove(InPrimitive);
+			PendingHeightFieldUpdateOps.Remove(InPrimitive);
+
+			if (InPrimitive->DistanceFieldInstanceIndices.Num() > 0)
+			{
+				PendingHeightFieldRemoveOps.Add(FHeightFieldPrimitiveRemoveInfo(InPrimitive));
+			}
+
+			InPrimitive->DistanceFieldInstanceIndices.Empty();
 		}
 	}
 }
diff --git a/Engine/Source/Runtime/Renderer/Private/ScenePrivate.h b/Engine/Source/Runtime/Renderer/Private/ScenePrivate.h
--- a/Engine/Source/Runtime/Renderer/Private/ScenePrivate.h
+++ b/Engine/Source/Runtime/Renderer/Private/ScenePrivate.h
@@ -1746,6 +1746,19 @@ public:
 	TArray<int32, TInlineAllocator<1>> DistanceFieldInstanceIndices;
 };
 
+class FHeightFieldPrimitiveRemoveInfo : public FPrimitiveRemoveInfo
+{
+public:
+	FHeightFieldPrimitiveRemoveInfo(const FPrimitiveSceneInfo* InPrimitive)
+		: FPrimitiveRemoveInfo(InPrimitive)
+	{
+		const FBoxSphereBounds Bounds = InPrimitive->Proxy->GetBounds();
+		SphereBound = FVector4(Bounds.Origin, Bounds.SphereRadius);
+	}
+
+	FVector4 SphereBound;
+};
+
 class FSurfelBufferAllocator
 {
 public:
@@ -1787,9 +1800,14 @@ public:
 		return PendingAddOperations.Num() > 0 || PendingUpdateOperations.Num() > 0 || PendingRemoveOperations.Num() > 0;
 	}
 
+	bool HasPendingHeightFieldOperations() const
+	{
+		return PendingHeightFieldAddOps.Num() > 0 || PendingHeightFieldUpdateOps.Num() > 0 || PendingHeightFieldRemoveOps.Num() > 0;
+	}
+
 	bool HasPendingRemovePrimitive(const FPrimitiveSceneInfo* Primitive) const
 	{
-		for (int32 RemoveIndex = 0; RemoveIndex < PendingRemoveOperations.Num(); RemoveIndex++)
+		for (int32 RemoveIndex = 0; RemoveIndex < PendingRemoveOperations.Num(); ++RemoveIndex)
 		{
 			if (PendingRemoveOperations[RemoveIndex].Primitive == Primitive)
 			{
@@ -1800,23 +1818,48 @@ public:
 		return false;
 	}
 
+	bool HasPendingRemoveHeightFieldPrimitive(const FPrimitiveSceneInfo* Primitive) const
+	{
+		for (int32 RemoveIndex = 0; RemoveIndex < PendingHeightFieldRemoveOps.Num(); ++RemoveIndex)
+		{
+			if (PendingHeightFieldRemoveOps[RemoveIndex].Primitive == Primitive)
+			{
+				return true;
+			}
+		}
+
+		return false;
+	}
+
 	inline bool CanUse16BitObjectIndices() const
 	{
 		return bCanUse16BitObjectIndices && (NumObjectsInBuffer < (1 << 16));
 	}
 
+	bool CanUse16BitHeightFieldObjectIndices() const
+	{
+		return bCanUse16BitObjectIndices && (NumHeightFieldObjectsInBuffer < 65536);
+	}
+
 	const class FDistanceFieldObjectBuffers* GetCurrentObjectBuffers() const
 	{
 		return ObjectBuffers[ObjectBufferIndex];
 	}
 
+	const class FHeightFieldObjectBuffers* GetHeightFieldObjectBuffers() const
+	{
+		return HeightFieldObjectBuffers;
+	}
+
 	int32 NumObjectsInBuffer;
+	int32 NumHeightFieldObjectsInBuffer;
 	class FDistanceFieldObjectBuffers* ObjectBuffers[2];
+	class FHeightFieldObjectBuffers* HeightFieldObjectBuffers;
 	int ObjectBufferIndex;
 
 	/** Stores the primitive and instance index of every entry in the object buffer. */
 	TArray<FPrimitiveAndInstance> PrimitiveInstanceMapping;
-	TArray<const FPrimitiveSceneInfo*> HeightfieldPrimitives;
+	TArray<FPrimitiveSceneInfo*> HeightfieldPrimitives;
 
 	class FSurfelBuffers* SurfelBuffers;
 	FSurfelBufferAllocator SurfelAllocations;
@@ -1831,8 +1874,13 @@ public:
 	TArray<FPrimitiveRemoveInfo> PendingRemoveOperations;
 	TArray<FVector4> PrimitiveModifiedBounds[GDF_Num];
 
+	TArray<FPrimitiveSceneInfo*> PendingHeightFieldAddOps;
+	TArray<FPrimitiveSceneInfo*> PendingHeightFieldUpdateOps;
+	TArray<FHeightFieldPrimitiveRemoveInfo> PendingHeightFieldRemoveOps;
+
 	/** Used to detect atlas reallocations, since objects store UVs into the atlas and need to be updated when it changes. */
 	int32 AtlasGeneration;
+	int32 HeightFieldAtlasGeneration;
 
 	bool bTrackAllPrimitives;
 	bool bCanUse16BitObjectIndices;
