Fix compile error

diff --git a/Engine/Source/Runtime/Renderer/Private/DistanceFieldObjectManagement.cpp b/Engine/Source/Runtime/Renderer/Private/DistanceFieldObjectManagement.cpp
--- a/Engine/Source/Runtime/Renderer/Private/DistanceFieldObjectManagement.cpp
+++ b/Engine/Source/Runtime/Renderer/Private/DistanceFieldObjectManagement.cpp
@@ -36,12 +36,12 @@ FAutoConsoleVariableRef CVarAOLogObjectBufferReallocation(
 	);
 
 // Must match equivalent shader defines
-int32 TDistanceFieldObjectBuffers<DFPT_SignedDistanceField>::ObjectDataStride = 17;
-int32 TDistanceFieldObjectBuffers<DFPT_HeightField>::ObjectDataStride = 5;
-int32 TDistanceFieldCulledObjectBuffers<DFPT_SignedDistanceField>::ObjectDataStride = 17;
-int32 TDistanceFieldCulledObjectBuffers<DFPT_SignedDistanceField>::ObjectBoxBoundsStride = 5;
-int32 TDistanceFieldCulledObjectBuffers<DFPT_HeightField>::ObjectDataStride = 5;
-int32 TDistanceFieldCulledObjectBuffers<DFPT_HeightField>::ObjectBoxBoundsStride = 5;
+template<> int32 TDistanceFieldObjectBuffers<DFPT_SignedDistanceField>::ObjectDataStride = 17;
+template<> int32 TDistanceFieldObjectBuffers<DFPT_HeightField>::ObjectDataStride = 5;
+template<> int32 TDistanceFieldCulledObjectBuffers<DFPT_SignedDistanceField>::ObjectDataStride = 17;
+template<> int32 TDistanceFieldCulledObjectBuffers<DFPT_SignedDistanceField>::ObjectBoxBoundsStride = 5;
+template<> int32 TDistanceFieldCulledObjectBuffers<DFPT_HeightField>::ObjectDataStride = 5;
+template<> int32 TDistanceFieldCulledObjectBuffers<DFPT_HeightField>::ObjectBoxBoundsStride = 5;
 
 // In float4's.  Must match corresponding usf definition
 int32 UploadObjectDataStride = 1 + FDistanceFieldObjectBuffers::ObjectDataStride;
diff --git a/Engine/Source/Runtime/Renderer/Private/DistanceFieldShadowing.cpp b/Engine/Source/Runtime/Renderer/Private/DistanceFieldShadowing.cpp
--- a/Engine/Source/Runtime/Renderer/Private/DistanceFieldShadowing.cpp
+++ b/Engine/Source/Runtime/Renderer/Private/DistanceFieldShadowing.cpp
@@ -444,8 +444,8 @@ private:
 #define VARIATION(PrimitiveType) \
 	typedef TShadowObjectCullPS<true, PrimitiveType> FShadowObjectCullPS##1##PrimitiveType; \
 	typedef TShadowObjectCullPS<false, PrimitiveType> FShadowObjectCullPS##0##PrimitiveType; \
-	IMPLEMENT_SHADER_TYPE(, FShadowObjectCullPS##1##PrimitiveType, TEXT("/Engine/Private/DistanceFieldShadowing.usf"), TEXT("ShadowObjectCullPS"), SF_Pixel); \
-	IMPLEMENT_SHADER_TYPE(, FShadowObjectCullPS##0##PrimitiveType, TEXT("/Engine/Private/DistanceFieldShadowing.usf"), TEXT("ShadowObjectCullPS"), SF_Pixel)
+	IMPLEMENT_SHADER_TYPE(template<>, FShadowObjectCullPS##1##PrimitiveType, TEXT("/Engine/Private/DistanceFieldShadowing.usf"), TEXT("ShadowObjectCullPS"), SF_Pixel); \
+	IMPLEMENT_SHADER_TYPE(template<>, FShadowObjectCullPS##0##PrimitiveType, TEXT("/Engine/Private/DistanceFieldShadowing.usf"), TEXT("ShadowObjectCullPS"), SF_Pixel)
 
 VARIATION(DFPT_SignedDistanceField);
 VARIATION(DFPT_HeightField);
