Copying //UE4/Dev-Physics to Main (//UE4/Main)

diff --git a/Engine/Source/Runtime/Renderer/Private/HairStrands/HairStrandsCluster.cpp b/Engine/Source/Runtime/Renderer/Private/HairStrands/HairStrandsCluster.cpp
--- a/Engine/Source/Runtime/Renderer/Private/HairStrands/HairStrandsCluster.cpp
+++ b/Engine/Source/Runtime/Renderer/Private/HairStrands/HairStrandsCluster.cpp
@@ -5,77 +5,254 @@
 #include "HairStrandsUtils.h"
 #include "SceneRendering.h"
 #include "SceneManagement.h"
+#include "RendererInterface.h"
+#include "Shader.h"
+#include "GlobalShader.h"
+#include "ShaderParameters.h"
+#include "ShaderParameterStruct.h"
 
-FHairStrandsClusterViews CreateHairStrandsClusters(
+class FHairMacroGroupAABBCS : public FGlobalShader
+{
+	DECLARE_GLOBAL_SHADER(FHairMacroGroupAABBCS);
+	SHADER_USE_PARAMETER_STRUCT(FHairMacroGroupAABBCS, FGlobalShader);
+
+	BEGIN_SHADER_PARAMETER_STRUCT(FParameters, )
+		SHADER_PARAMETER(uint32, MacroGroupId)
+		SHADER_PARAMETER(uint32, MacroGroupValid)
+		SHADER_PARAMETER(uint32, bClearBuffer)
+		SHADER_PARAMETER_SRV(Buffer, InGroupAABBBuffer0)
+		SHADER_PARAMETER_SRV(Buffer, InGroupAABBBuffer1)
+		SHADER_PARAMETER_SRV(Buffer, InGroupAABBBuffer2)
+		SHADER_PARAMETER_SRV(Buffer, InGroupAABBBuffer3)
+		SHADER_PARAMETER_SRV(Buffer, InGroupAABBBuffer4)
+		SHADER_PARAMETER_SRV(Buffer, InGroupAABBBuffer5)
+		SHADER_PARAMETER_SRV(Buffer, InGroupAABBBuffer6)
+		SHADER_PARAMETER_SRV(Buffer, InGroupAABBBuffer7)
+		SHADER_PARAMETER_RDG_BUFFER_UAV(RWBuffer, OutMacroGroupAABBBuffer)
+		END_SHADER_PARAMETER_STRUCT()
+
+public:
+	static bool ShouldCompilePermutation(const FGlobalShaderPermutationParameters& Parameters) { return IsHairStrandsSupported(Parameters.Platform); }
+	static void ModifyCompilationEnvironment(const FGlobalShaderPermutationParameters& Parameters, FShaderCompilerEnvironment& OutEnvironment)
+	{
+		FGlobalShader::ModifyCompilationEnvironment(Parameters, OutEnvironment);
+		OutEnvironment.SetDefine(TEXT("SHADER_AABBUPDATE"), 1);
+	}
+};
+
+IMPLEMENT_GLOBAL_SHADER(FHairMacroGroupAABBCS, "/Engine/Private/HairStrands/HairStrandsAABB.usf", "Main", SF_Compute);
+
+static void AddHairMacroGroupAABBPass(
+	FRDGBuilder& GraphBuilder,
+	FHairStrandsMacroGroupData& MacroGroup,
+	FRDGBufferUAVRef& OutHairMacroGroupAABBBufferUAV)
+{
+	const uint32 PrimitiveCount = MacroGroup.PrimitivesInfos.Num();
+	if (PrimitiveCount == 0)
+		return;
+
+	const uint32 GroupPerPass = 8;
+	bool bNeedClear = true;
+	const uint32 MacroGroupId = MacroGroup.MacroGroupId;
+	const uint32 IterationCount = FMath::CeilToInt(PrimitiveCount / float(GroupPerPass));
+	for (uint32 PassIt = 0; PassIt < IterationCount; ++PassIt)
+	{
+		FHairMacroGroupAABBCS::FParameters* Parameters = GraphBuilder.AllocParameters<FHairMacroGroupAABBCS::FParameters>();
+		Parameters->MacroGroupId = MacroGroupId;
+		Parameters->OutMacroGroupAABBBuffer = OutHairMacroGroupAABBBufferUAV;
+
+		uint32 MacroGroupValid = 1;
+		uint32 CurrentGroupIt = 1;
+		for (uint32 PassPrimitiveIt = 0, PassPrimitiveCount = FMath::Min(GroupPerPass, PrimitiveCount - PassIt * GroupPerPass); PassPrimitiveIt < PassPrimitiveCount; ++PassPrimitiveIt)
+		{
+			const uint32 PrimitiveIndex = PassIt * GroupPerPass + PassPrimitiveIt;
+			const FHairStrandsMacroGroupData::PrimitiveInfo& PrimitiveInfo = MacroGroup.PrimitivesInfos[PrimitiveIndex];
+			const FHairStrandsPrimitiveResources& Resources = GetHairStandsPrimitiveResources(PrimitiveInfo.ResourceId);
+			FShaderResourceViewRHIRef GroupAABBBufferSRV = Resources.Groups[PrimitiveInfo.GroupIndex].GroupAABBBuffer->SRV;
+
+			// Default value
+			if (PassPrimitiveIt == 0 && PassPrimitiveCount != GroupPerPass)
+			{
+				Parameters->InGroupAABBBuffer0 = GroupAABBBufferSRV;
+				Parameters->InGroupAABBBuffer1 = GroupAABBBufferSRV;
+				Parameters->InGroupAABBBuffer2 = GroupAABBBufferSRV;
+				Parameters->InGroupAABBBuffer3 = GroupAABBBufferSRV;
+				Parameters->InGroupAABBBuffer4 = GroupAABBBufferSRV;
+				Parameters->InGroupAABBBuffer5 = GroupAABBBufferSRV;
+				Parameters->InGroupAABBBuffer6 = GroupAABBBufferSRV;
+				Parameters->InGroupAABBBuffer7 = GroupAABBBufferSRV;
+			}
+
+			switch (PassPrimitiveIt)
+			{
+				case 0 : Parameters->InGroupAABBBuffer0 = GroupAABBBufferSRV; break;
+				case 1 : Parameters->InGroupAABBBuffer1 = GroupAABBBufferSRV; break;
+				case 2 : Parameters->InGroupAABBBuffer2 = GroupAABBBufferSRV; break;
+				case 3 : Parameters->InGroupAABBBuffer3 = GroupAABBBufferSRV; break;
+				case 4 : Parameters->InGroupAABBBuffer4 = GroupAABBBufferSRV; break;
+				case 5 : Parameters->InGroupAABBBuffer5 = GroupAABBBufferSRV; break;
+				case 6 : Parameters->InGroupAABBBuffer6 = GroupAABBBufferSRV; break;
+				case 7 : Parameters->InGroupAABBBuffer7 = GroupAABBBufferSRV; break;
+			}
+			MacroGroupValid |= 1 << PassPrimitiveIt;
+
+		}
+
+		Parameters->MacroGroupValid = MacroGroupValid;
+		Parameters->bClearBuffer = bNeedClear ? 1 : 0;
+
+		TShaderMapRef<FHairMacroGroupAABBCS> ComputeShader(GetGlobalShaderMap(ERHIFeatureLevel::SM5));
+		FComputeShaderUtils::AddPass(
+			GraphBuilder,
+			RDG_EVENT_NAME("HairStrandsMacroGroupAABBUpdate"),
+			*ComputeShader,
+			Parameters,
+			FIntVector(1,1,1));
+
+		bNeedClear = false;
+	}
+}
+
+static bool DoesGroupExists(uint32 ResourceId, FHairStrandsMacroGroupData::TPrimitiveGroups& PrimitivesGroups)
+{
+	// Simple linear search as the expected number of groups is supposed to be low (<10)
+	for (FHairStrandsMacroGroupData::PrimitiveGroup& Group : PrimitivesGroups)
+	{
+		if (Group.ResourceId == ResourceId)
+		{
+			return true;
+		}
+	}
+	return false;
+}
+
+FHairStrandsMacroGroupViews CreateHairStrandsMacroGroups(
 	FRHICommandListImmediate& RHICmdList,
 	const FScene* Scene,
 	const TArray<FViewInfo>& Views)
 {
-	FHairStrandsClusterViews PrimitivesClusterViews;
+	static const FVertexFactoryType* CompatibleVF = FVertexFactoryType::GetVFByName(TEXT("FHairStrandsVertexFactory"));
+
+	FHairStrandsMacroGroupViews MacroGroupsViews;
 	for (int32 ViewIndex = 0; ViewIndex < Views.Num(); ViewIndex++)
 	{
 		const FViewInfo& View = Views[ViewIndex];
 		if (View.Family)
 		{
 			int32 MaterialId = 0;
-			FHairStrandsClusterDatas& PrimitivesClusters = PrimitivesClusterViews.Views.AddDefaulted_GetRef();
+			FHairStrandsMacroGroupDatas& MacroGroups = MacroGroupsViews.Views.AddDefaulted_GetRef();
 
 			if (View.HairStrandsMeshElements.Num() == 0)
 			{
 				continue;
 			}
 
-			// Spatially clusters all hair primitives within the same area, and render a single DOM for them
-			uint32 ClusterId = 0;
-			auto UpdateCluster = [&PrimitivesClusters, &View, &ClusterId, &MaterialId](const FMeshBatchAndRelevance* MeshBatchAndRelevance, const FPrimitiveSceneProxy* Proxy)
+			// Aggregate all hair primitives within the same area into macro groups, for allocating/rendering DOM/voxel
+			uint32 MacroGroupId = 0;
+			auto UpdateMacroGroup = [&MacroGroups, &View, &MacroGroupId, &MaterialId](const FMeshBatchAndRelevance* MeshBatchAndRelevance, const FPrimitiveSceneProxy* Proxy)
 			{
+				const bool bIsHairStrandsFactory = MeshBatchAndRelevance->Mesh->VertexFactory->GetType()->GetId() == CompatibleVF->GetId();
+				if (!bIsHairStrandsFactory)
+					return;
+
 				const FBoxSphereBounds& PrimitiveBounds = Proxy->GetBounds();
 
-				bool bClusterFound = false;
-				for (FHairStrandsClusterData& Cluster : PrimitivesClusters.Datas)
+				bool bFound = false;
+				for (FHairStrandsMacroGroupData& MacroGroup : MacroGroups.Datas)
 				{
-					const bool bIntersect = FBoxSphereBounds::SpheresIntersect(Cluster.Bounds, PrimitiveBounds);
+					const bool bIntersect = FBoxSphereBounds::SpheresIntersect(MacroGroup.Bounds, PrimitiveBounds);
 					if (bIntersect)
 					{
-						Cluster.Bounds = Union(Cluster.Bounds, PrimitiveBounds);
+						MacroGroup.Bounds = Union(MacroGroup.Bounds, PrimitiveBounds);
 
 						if (MeshBatchAndRelevance)
 						{
-							FHairStrandsClusterData::PrimitiveInfo& PrimitiveInfo = Cluster.PrimitivesInfos.AddZeroed_GetRef();
+							check(MeshBatchAndRelevance->Mesh);
+							check(MeshBatchAndRelevance->Mesh->Elements.Num() == 1);
+
+							const FHairGroupPublicData* HairGroupPublicData = reinterpret_cast<const FHairGroupPublicData*>(MeshBatchAndRelevance->Mesh->Elements[0].VertexFactoryUserData);
+
+							FHairStrandsMacroGroupData::PrimitiveInfo& PrimitiveInfo = MacroGroup.PrimitivesInfos.AddZeroed_GetRef();
 							PrimitiveInfo.MeshBatchAndRelevance = *MeshBatchAndRelevance;
 							PrimitiveInfo.MaterialId = MaterialId++;
+							PrimitiveInfo.ResourceId = reinterpret_cast<uint64>(MeshBatchAndRelevance->Mesh->Elements[0].UserData);
+							PrimitiveInfo.GroupIndex = HairGroupPublicData->GetGroupIndex();
+							check(PrimitiveInfo.GroupIndex < 32); // Sanity check
+
+							const bool bAlreadyExists = DoesGroupExists(PrimitiveInfo.ResourceId, MacroGroup.PrimitivesGroups);
+							if (!bAlreadyExists)
+							{
+								FHairStrandsMacroGroupData::PrimitiveGroup& PrimitiveGroup = MacroGroup.PrimitivesGroups.AddZeroed_GetRef();
+								PrimitiveGroup.GroupIndex = PrimitiveInfo.GroupIndex;
+								PrimitiveGroup.ResourceId = PrimitiveInfo.ResourceId;
+							}
 						}
-						bClusterFound = true;
+						bFound = true;
 						break;
 					}
 				}
 
-				if (!bClusterFound)
+				if (!bFound)
 				{
-					FHairStrandsClusterData Cluster;
-					Cluster.ClusterId = ClusterId++;
+					FHairStrandsMacroGroupData MacroGroup;
+					MacroGroup.MacroGroupId = MacroGroupId++;
 					if (MeshBatchAndRelevance)
 					{
-						FHairStrandsClusterData::PrimitiveInfo& PrimitiveInfo = Cluster.PrimitivesInfos.AddZeroed_GetRef();
+						check(MeshBatchAndRelevance->Mesh);
+						check(MeshBatchAndRelevance->Mesh->Elements.Num() == 1);
+
+						const FHairGroupPublicData* HairGroupPublicData  = reinterpret_cast<const FHairGroupPublicData*>(MeshBatchAndRelevance->Mesh->Elements[0].VertexFactoryUserData);
+
+						FHairStrandsMacroGroupData::PrimitiveInfo& PrimitiveInfo = MacroGroup.PrimitivesInfos.AddZeroed_GetRef();
 						PrimitiveInfo.MeshBatchAndRelevance = *MeshBatchAndRelevance;
 						PrimitiveInfo.MaterialId = MaterialId++;
+						PrimitiveInfo.ResourceId = reinterpret_cast<uint64>(MeshBatchAndRelevance->Mesh->Elements[0].UserData);
+						PrimitiveInfo.GroupIndex = HairGroupPublicData->GetGroupIndex();
+						check(PrimitiveInfo.GroupIndex < 32); // Sanity check
+
+						const bool bAlreadyExists = DoesGroupExists(PrimitiveInfo.ResourceId, MacroGroup.PrimitivesGroups);
+						if (!bAlreadyExists)
+						{
+							FHairStrandsMacroGroupData::PrimitiveGroup& PrimitiveGroup = MacroGroup.PrimitivesGroups.AddZeroed_GetRef();
+							PrimitiveGroup.GroupIndex = PrimitiveInfo.GroupIndex;
+							PrimitiveGroup.ResourceId = PrimitiveInfo.ResourceId;
+						}
 					}
-					Cluster.Bounds = PrimitiveBounds;
-					PrimitivesClusters.Datas.Add(Cluster);
+					MacroGroup.Bounds = PrimitiveBounds;
+					MacroGroups.Datas.Add(MacroGroup);
 				}
 			};
 
 			for (const FMeshBatchAndRelevance& MeshBatchAndRelevance : View.HairStrandsMeshElements)
 			{
-				UpdateCluster(&MeshBatchAndRelevance, MeshBatchAndRelevance.PrimitiveSceneProxy);
+				UpdateMacroGroup(&MeshBatchAndRelevance, MeshBatchAndRelevance.PrimitiveSceneProxy);
+			}
+
+			for (FHairStrandsMacroGroupData& MacroGroup : MacroGroups.Datas)
+			{
+				MacroGroup.ScreenRect = ComputeProjectedScreenRect(MacroGroup.Bounds.GetBox(), View);
 			}
 
-			for (FHairStrandsClusterData& Cluster : PrimitivesClusters.Datas)
+			// Build hair macro group AABBB
+			const uint32 MacroGroupCount = MacroGroups.Datas.Num();
+			if (MacroGroupCount > 0)
 			{
-				Cluster.ScreenRect = ComputeProjectedScreenRect(Cluster.Bounds.GetBox(), View);
+				DECLARE_GPU_STAT(HairStrandsAABB);
+				SCOPED_DRAW_EVENT(RHICmdList, HairStrandsAABB);
+				SCOPED_GPU_STAT(RHICmdList, HairStrandsAABB);
+
+				FRDGBuilder GraphBuilder(RHICmdList);
+				FRDGBufferRef MacroGroupAABBBuffer = GraphBuilder.CreateBuffer(FRDGBufferDesc::CreateBufferDesc(4, 6 * MacroGroupCount), TEXT("HairMacroGroupAABBBuffer"));
+				FRDGBufferUAVRef MacroGroupAABBBufferUAV = GraphBuilder.CreateUAV(MacroGroupAABBBuffer, PF_R32_SINT);
+				for (FHairStrandsMacroGroupData& MacroGroup : MacroGroups.Datas)
+				{				
+					AddHairMacroGroupAABBPass(GraphBuilder, MacroGroup, MacroGroupAABBBufferUAV);
+				}
+				GraphBuilder.QueueBufferExtraction(MacroGroupAABBBuffer, &MacroGroups.MacroGroupResources.MacroGroupAABBsBuffer, FRDGResourceState::EAccess::Read, FRDGResourceState::EPipeline::Compute);
+				GraphBuilder.Execute();
 			}
 		}
 	}
 
-	return PrimitivesClusterViews;
+	return MacroGroupsViews;
 }
diff --git a/Engine/Source/Runtime/Renderer/Private/HairStrands/HairStrandsCluster.h b/Engine/Source/Runtime/Renderer/Private/HairStrands/HairStrandsCluster.h
--- a/Engine/Source/Runtime/Renderer/Private/HairStrands/HairStrandsCluster.h
+++ b/Engine/Source/Runtime/Renderer/Private/HairStrands/HairStrandsCluster.h
@@ -1,7 +1,7 @@
 // Copyright Epic Games, Inc. All Rights Reserved.
 
 /*=============================================================================
-	HairStrandsCluster.h: Hair strands cluster computation implementation.
+	HairStrandsCluster.h: Hair strands macro group computation implementation.
 =============================================================================*/
 
 #pragma once
@@ -11,40 +11,73 @@
 #include "HairStrandsVoxelization.h"
 #include "SceneManagement.h"
 
-/// Hair cluster infos (i.e. group of mesh having hair material on them)
-struct FHairStrandsClusterData
+// -----------------------------
+// Hair Data structures overview 
+// -----------------------------
+// A groom component contains one or several HairGroup. These hair group are send to the 
+// render as mesh batches. These meshes batches are filtered/culled per view, and regroup 
+// into HairMacroGroup for computing voxelization/DOM data, ...
+//
+// The hierarchy of the data structure is as follow:
+//  * HairMacroGroup
+//  * HairGroup
+//  * HairCluster
+
+struct FHairMacroGroupAABBData
+{
+	TRefCountPtr<FPooledRDGBuffer>	MacroGroupAABBsBuffer;
+};
+
+/// Hair macro group infos
+struct FHairStrandsMacroGroupData
 {
+	// List of primitive/mesh batch within an instance group
 	struct PrimitiveInfo
 	{
 		FMeshBatchAndRelevance MeshBatchAndRelevance;
 		uint32 MaterialId;
+		uint32 ResourceId;
+		uint32 GroupIndex;
 	};
 	typedef TArray<PrimitiveInfo, SceneRenderingAllocator> TPrimitiveInfos;
 
+	// List of unique group within an instance group
+	struct PrimitiveGroup
+	{
+		uint32 ResourceId;
+		uint32 GroupIndex;
+	};
+	typedef TArray<PrimitiveGroup, SceneRenderingAllocator> TPrimitiveGroups;
+
+
 	FVector GetMinBound() const { return VoxelResources.MinAABB; }
 	FVector GetMaxBound() const { return VoxelResources.MaxAABB; }
 	uint32  GetResolution() const { return VoxelResources.DensityTexture ? VoxelResources.DensityTexture->GetDesc().Extent.X : 0; }
 
+	FVirtualVoxelNodeDesc VirtualVoxelNodeDesc;
 	FHairStrandsVoxelResources VoxelResources;
 	TPrimitiveInfos PrimitivesInfos;
+	TPrimitiveGroups PrimitivesGroups;
 	FBoxSphereBounds Bounds;
 	FIntRect ScreenRect;
-	uint32 ClusterId;
+	uint32 MacroGroupId;
 };
 
-/// Store all hair strandscluster infos for a given view
-struct FHairStrandsClusterDatas
+/// Store all hair strands macro group infos for a given view
+struct FHairStrandsMacroGroupDatas
 {
-	TArray<FHairStrandsClusterData, SceneRenderingAllocator> Datas;
+	TArray<FHairStrandsMacroGroupData, SceneRenderingAllocator> Datas;
+	FVirtualVoxelResources VirtualVoxelResources;
+	FHairMacroGroupAABBData MacroGroupResources;
 };
 
-/// Store all hair strands cluster info for all views
-struct FHairStrandsClusterViews
+/// Store all hair strands macro group info for all views
+struct FHairStrandsMacroGroupViews
 {
-	TArray<FHairStrandsClusterDatas, SceneRenderingAllocator> Views;
+	TArray<FHairStrandsMacroGroupDatas, SceneRenderingAllocator> Views;
 };
 
-FHairStrandsClusterViews CreateHairStrandsClusters(
+FHairStrandsMacroGroupViews CreateHairStrandsMacroGroups(
 	FRHICommandListImmediate& RHICmdList,
 	const FScene* Scene,
 	const TArray<FViewInfo>& Views);
\ No newline at end of file
diff --git a/Engine/Source/Runtime/Renderer/Private/HairStrands/HairStrandsClusters.cpp b/Engine/Source/Runtime/Renderer/Private/HairStrands/HairStrandsClusters.cpp
new file mode 100644
--- /dev/null
+++ b/Engine/Source/Runtime/Renderer/Private/HairStrands/HairStrandsClusters.cpp
@@ -0,0 +1,671 @@
+// Copyright Epic Games, Inc. All Rights Reserved.
+
+#include "HairStrandsClusters.h"
+#include "HairStrandsUtils.h"
+#include "SceneRendering.h"
+#include "SceneManagement.h"
+#include "RendererInterface.h"
+#include "Shader.h"
+#include "GlobalShader.h"
+#include "ShaderParameters.h"
+#include "ShaderParameterStruct.h"
+
+///////////////////////////////////////////////////////////////////////////////////////////////////
+
+static int32 GHairStrandsClusterCullingUsesHzb = 1;
+static FAutoConsoleVariableRef CVarHairCullingUseHzb(TEXT("r.HairStrands.Cluster.CullingUsesHzb"), GHairStrandsClusterCullingUsesHzb, TEXT("Enable/disable the use of HZB to help cull more hair clusters."));
+
+static int32 GHairStrandsClusterCulling = 1;
+static FAutoConsoleVariableRef CVarHairClusterCulling(TEXT("r.HairStrands.Cluster.Culling"), GHairStrandsClusterCulling, TEXT("Enable/Disable hair cluster culling"));
+
+static int32 GHairStrandsClusterCullingLodMode = -1;
+static FAutoConsoleVariableRef CVarHairClusterCullingLodMode(TEXT("r.HairStrands.Cluster.CullingLodMode"), GHairStrandsClusterCullingLodMode, TEXT("0/1/2 to force such a lod. Otherwise use -1 for automatic selection of Lod."));
+
+static int32 GStrandHairClusterCullingShadow = 1;
+static FAutoConsoleVariableRef CVarHairClusterCullingShadow(TEXT("r.HairStrands.Cluster.CullingShadow"), GStrandHairClusterCullingShadow, TEXT("Enable/Disable hair cluster culling for shadow views"));
+
+static int32 GHairStrandsClusterCullingShadowLodMode = -1;
+static FAutoConsoleVariableRef CVarHairClusterCullingShadowLodMode(TEXT("r.HairStrands.Cluster.CullingShadowLodMode"), GHairStrandsClusterCullingShadowLodMode, TEXT("0/1/2 to force such a lod. Otherwise use -1 for automatic selection of Lod. Used for voxelisation and DOM."));
+
+static int32 GHairStrandsClusterCullingFreezeCamera = 0;
+static FAutoConsoleVariableRef CVarHairStrandsClusterCullingFreezeCamera(TEXT("r.HairStrands.Cluster.CullingFreezeCamera"), GHairStrandsClusterCullingFreezeCamera, TEXT("Freeze camera when enabled. It will disable HZB culling because hzb buffer is not frozen."));
+
+
+///////////////////////////////////////////////////////////////////////////////////////////////////
+#if 0 // This code is currently defined within HairStrandsRendering.cpp but it would be better to move it here
+class FClearClusterAABBCS : public FGlobalShader
+{
+	DECLARE_GLOBAL_SHADER(FClearClusterAABBCS);
+	SHADER_USE_PARAMETER_STRUCT(FClearClusterAABBCS, FGlobalShader);
+
+	BEGIN_SHADER_PARAMETER_STRUCT(FParameters, )
+		SHADER_PARAMETER_UAV(RWBuffer, OutClusterAABBBuffer)
+		SHADER_PARAMETER_UAV(RWBuffer, OutGroupAABBBuffer)
+		SHADER_PARAMETER(uint32, ClusterCount)
+		END_SHADER_PARAMETER_STRUCT()
+
+public:
+	static bool ShouldCompilePermutation(const FGlobalShaderPermutationParameters& Parameters) { return IsHairStrandsSupported(Parameters.Platform); }
+	static void ModifyCompilationEnvironment(const FGlobalShaderPermutationParameters& Parameters, FShaderCompilerEnvironment& OutEnvironment)
+	{
+		FGlobalShader::ModifyCompilationEnvironment(Parameters, OutEnvironment);
+		OutEnvironment.SetDefine(TEXT("SHADER_CLEARCLUSTERAABB"), 1);
+	}
+};
+
+IMPLEMENT_GLOBAL_SHADER(FClearClusterAABBCS, "/Engine/Private/HairStrands/HairStrandsClusterCulling.usf", "MainClearClusterAABBCS", SF_Compute);
+
+static void AddClearClusterAABBPass(
+	FRDGBuilder& GraphBuilder,
+	uint32 ClusterCount,
+	FRHIUnorderedAccessView* OutClusterAABBuffer,
+	FRHIUnorderedAccessView* OutGroupAABBuffer)
+{
+	check(OutClusterAABBuffer);
+
+	FClearClusterAABBCS::FParameters* Parameters = GraphBuilder.AllocParameters<FClearClusterAABBCS::FParameters>();
+	Parameters->ClusterCount = ClusterCount;
+	Parameters->OutClusterAABBBuffer = OutClusterAABBuffer;
+	Parameters->OutGroupAABBBuffer = OutGroupAABBuffer;
+
+
+	TShaderMap<FGlobalShaderType>* ShaderMap = GetGlobalShaderMap(ERHIFeatureLevel::SM5);
+	TShaderMapRef<FClearClusterAABBCS> ComputeShader(ShaderMap);
+
+	const FIntVector DispatchCount = DispatchCount.DivideAndRoundUp(FIntVector(ClusterCount * 6, 1, 1), FIntVector(64, 1, 1));
+	FComputeShaderUtils::AddPass(
+		GraphBuilder,
+		RDG_EVENT_NAME("HairStrandsClearClusterAABB"),
+		*ComputeShader,
+		Parameters,
+		DispatchCount);
+}
+
+///////////////////////////////////////////////////////////////////////////////////////////////////
+
+class FHairClusterAABBCS : public FGlobalShader
+{
+	DECLARE_GLOBAL_SHADER(FHairClusterAABBCS);
+	SHADER_USE_PARAMETER_STRUCT(FHairClusterAABBCS, FGlobalShader);
+
+	class FGroupSize : SHADER_PERMUTATION_INT("PERMUTATION_GROUP_SIZE", 2);
+	using FPermutationDomain = TShaderPermutationDomain<FGroupSize>;
+
+	BEGIN_SHADER_PARAMETER_STRUCT(FParameters, )
+		SHADER_PARAMETER(uint32, DispatchCountX)
+		SHADER_PARAMETER(uint32, ClusterCount)
+		SHADER_PARAMETER(FVector, OutHairPositionOffset)
+		SHADER_PARAMETER(FMatrix, LocalToWorldMatrix)
+		SHADER_PARAMETER_SRV(Buffer, RenderDeformedPositionBuffer)
+		SHADER_PARAMETER_SRV(Buffer, ClusterVertexIdBuffer)
+		SHADER_PARAMETER_SRV(Buffer, ClusterInfoBuffer)
+		SHADER_PARAMETER_UAV(RWBuffer, OutClusterAABBBuffer)
+		SHADER_PARAMETER_UAV(RWBuffer, OutGroupAABBBuffer)
+		END_SHADER_PARAMETER_STRUCT()
+
+public:
+	static bool ShouldCompilePermutation(const FGlobalShaderPermutationParameters& Parameters) { return IsHairStrandsSupported(Parameters.Platform); }
+};
+
+IMPLEMENT_GLOBAL_SHADER(FHairClusterAABBCS, "/Engine/Private/HairStrands/HairStrandsInterpolation.usf", "ClusterAABBEvaluationCS", SF_Compute);
+
+static void AddHairClusterAABBPass(
+	FRDGBuilder& GraphBuilder,
+	const FHairStrandsProjectionHairData::HairGroup& InRenHairData,
+	const FVector& OutHairWorldOffset,
+	FHairStrandClusterData::FHairGroup& HairGroupClusters,
+	const FShaderResourceViewRHIRef& RenderPositionBuffer)
+{
+	const uint32 GroupSize = ComputeGroupSize();
+	const FIntVector DispatchCount = ComputeDispatchGroupCount2D(HairGroupClusters.ClusterCount);
+
+	FHairClusterAABBCS::FParameters* Parameters = GraphBuilder.AllocParameters<FHairClusterAABBCS::FParameters>();
+	Parameters->DispatchCountX = DispatchCount.X;
+	Parameters->ClusterCount = HairGroupClusters.ClusterCount;
+	Parameters->LocalToWorldMatrix = InRenHairData.LocalToWorld.ToMatrixWithScale();
+	Parameters->OutHairPositionOffset = OutHairWorldOffset;
+	Parameters->RenderDeformedPositionBuffer = RenderPositionBuffer;
+	Parameters->ClusterVertexIdBuffer = HairGroupClusters.ClusterVertexIdBuffer->SRV;
+	Parameters->ClusterInfoBuffer = HairGroupClusters.ClusterInfoBuffer->SRV;
+	Parameters->OutClusterAABBBuffer = HairGroupClusters.ClusterAABBBuffer->UAV;
+	Parameters->OutGroupAABBBuffer = HairGroupClusters.GroupAABBBuffer->UAV;
+
+	FHairClusterAABBCS::FPermutationDomain PermutationVector;
+	PermutationVector.Set<FHairClusterAABBCS::FGroupSize>(GetGroupSizePermutation(GroupSize));
+	TShaderMapRef<FHairClusterAABBCS> ComputeShader(GetGlobalShaderMap(ERHIFeatureLevel::SM5), PermutationVector);
+
+	FComputeShaderUtils::AddPass(
+		GraphBuilder,
+		RDG_EVENT_NAME("HairStrandsClusterAABB"),
+		*ComputeShader,
+		Parameters,
+		DispatchCount);
+}
+#endif
+
+///////////////////////////////////////////////////////////////////////////////////////////////////
+
+class FHairIndBufferClearCS : public FGlobalShader
+{
+	DECLARE_GLOBAL_SHADER(FHairIndBufferClearCS);
+	SHADER_USE_PARAMETER_STRUCT(FHairIndBufferClearCS, FGlobalShader);
+
+	class FSetIndirectDraw : SHADER_PERMUTATION_INT("PERMUTATION_SETINDIRECTDRAW", 2);
+	using FPermutationDomain = TShaderPermutationDomain<FSetIndirectDraw>;
+
+	BEGIN_SHADER_PARAMETER_STRUCT(FParameters, )
+		SHADER_PARAMETER_RDG_BUFFER_UAV(Buffer, DispatchIndirectParametersClusterCount)
+		SHADER_PARAMETER_UAV(Buffer, DrawIndirectParameters)
+		SHADER_PARAMETER(uint32, VertexCountPerInstance)
+	END_SHADER_PARAMETER_STRUCT()
+
+public:
+	static bool ShouldCompilePermutation(const FGlobalShaderPermutationParameters& Parameters) { return IsHairStrandsSupported(Parameters.Platform); }
+	static void ModifyCompilationEnvironment(const FGlobalShaderPermutationParameters& Parameters, FShaderCompilerEnvironment& OutEnvironment)
+	{
+		FGlobalShader::ModifyCompilationEnvironment(Parameters, OutEnvironment);
+		OutEnvironment.SetDefine(TEXT("SHADER_CLUSTERCULLINGINDCLEAR"), 1);
+	}
+};
+
+IMPLEMENT_GLOBAL_SHADER(FHairIndBufferClearCS, "/Engine/Private/HairStrands/HairStrandsClusterCulling.usf", "MainClusterCullingIndClearCS", SF_Compute);
+
+///////////////////////////////////////////////////////////////////////////////////////////////////
+
+class FHairClusterCullingCS : public FGlobalShader
+{
+	DECLARE_GLOBAL_SHADER(FHairClusterCullingCS);
+	SHADER_USE_PARAMETER_STRUCT(FHairClusterCullingCS, FGlobalShader);
+
+	class FHZBCulling : SHADER_PERMUTATION_INT("PERMUTATION_HZBCULLING", 2);
+	class FDebugAABBBuffer : SHADER_PERMUTATION_INT("PERMUTATION_DEBUGAABBBUFFER", 2);
+	class FLodDebug : SHADER_PERMUTATION_INT("PERMUTATION_LODDEBUG", 2);
+	using FPermutationDomain = TShaderPermutationDomain<FHZBCulling, FDebugAABBBuffer, FLodDebug>;
+
+	BEGIN_SHADER_PARAMETER_STRUCT(FParameters, )
+		SHADER_PARAMETER(FVector, CameraWorldPos)
+		SHADER_PARAMETER(FMatrix, WorldToClipMatrix)
+		SHADER_PARAMETER(uint32, ClusterCount)
+		SHADER_PARAMETER(int32, LodMode)
+		SHADER_PARAMETER(int32, ShadowViewMode)
+		SHADER_PARAMETER(uint32, NumConvexHullPlanes)
+		SHADER_PARAMETER(float, LodBias)
+		SHADER_PARAMETER(float, LodAverageVertexPerPixel)
+		SHADER_PARAMETER(FVector2D, ViewRectPixel)
+		SHADER_PARAMETER_ARRAY(FVector, ViewFrustumConvexHull, [6])
+		SHADER_PARAMETER_SRV(Buffer, ClusterAABBBuffer)
+		SHADER_PARAMETER_SRV(Buffer, ClusterInfoBuffer)
+		SHADER_PARAMETER_SRV(Buffer, ClusterIndexRadiusScaleInfoBuffer)
+		SHADER_PARAMETER_RDG_BUFFER_UAV(Buffer, GlobalIndexStartBuffer)
+		SHADER_PARAMETER_RDG_BUFFER_UAV(Buffer, GlobalIndexCountBuffer)
+		SHADER_PARAMETER_RDG_BUFFER_UAV(Buffer, GlobalRadiusScaleBuffer)
+		SHADER_PARAMETER_RDG_BUFFER_UAV(Buffer, ClusterDebugAABBBuffer)
+		SHADER_PARAMETER_RDG_BUFFER_UAV(Buffer, DispatchIndirectParametersClusterCount)
+		SHADER_PARAMETER_UAV(Buffer, DrawIndirectParameters)
+		SHADER_PARAMETER(FVector, HZBUvFactor)
+		SHADER_PARAMETER(FVector4, HZBSize)
+		SHADER_PARAMETER_RDG_TEXTURE(Texture2d<float>, HZBTexture)
+		SHADER_PARAMETER_SAMPLER(SamplerState, HZBSampler)
+	END_SHADER_PARAMETER_STRUCT()
+
+public:
+	static bool ShouldCompilePermutation(const FGlobalShaderPermutationParameters& Parameters) { return IsHairStrandsSupported(Parameters.Platform); }
+	static void ModifyCompilationEnvironment(const FGlobalShaderPermutationParameters& Parameters, FShaderCompilerEnvironment& OutEnvironment)
+	{
+		FGlobalShader::ModifyCompilationEnvironment(Parameters, OutEnvironment);
+		OutEnvironment.SetDefine(TEXT("SHADER_CLUSTERCULLING"), 1);
+	}
+};
+
+IMPLEMENT_GLOBAL_SHADER(FHairClusterCullingCS, "/Engine/Private/HairStrands/HairStrandsClusterCulling.usf", "MainClusterCullingCS", SF_Compute);
+
+///////////////////////////////////////////////////////////////////////////////////////////////////
+
+class FMainClusterCullingPrepareIndirectDrawsCS : public FGlobalShader
+{
+	DECLARE_GLOBAL_SHADER(FMainClusterCullingPrepareIndirectDrawsCS);
+	SHADER_USE_PARAMETER_STRUCT(FMainClusterCullingPrepareIndirectDrawsCS, FGlobalShader);
+
+	BEGIN_SHADER_PARAMETER_STRUCT(FParameters, )
+		SHADER_PARAMETER_RDG_BUFFER_SRV(Buffer, DispatchIndirectParametersClusterCount)
+		SHADER_PARAMETER_RDG_BUFFER_UAV(Buffer, DispatchIndirectParametersClusterCount2D)
+		SHADER_PARAMETER_RDG_BUFFER_UAV(Buffer, DispatchIndirectParametersClusterCountDiv512)
+		SHADER_PARAMETER_RDG_BUFFER_UAV(Buffer, DispatchIndirectParametersClusterCountDiv512Div512)
+	END_SHADER_PARAMETER_STRUCT()
+
+public:
+	static bool ShouldCompilePermutation(const FGlobalShaderPermutationParameters& Parameters) { return IsHairStrandsSupported(Parameters.Platform); }
+	static void ModifyCompilationEnvironment(const FGlobalShaderPermutationParameters& Parameters, FShaderCompilerEnvironment& OutEnvironment)
+	{
+		FGlobalShader::ModifyCompilationEnvironment(Parameters, OutEnvironment);
+		OutEnvironment.SetDefine(TEXT("SHADER_PREPAREINDIRECTDRAW"), 1);
+	}
+};
+
+IMPLEMENT_GLOBAL_SHADER(FMainClusterCullingPrepareIndirectDrawsCS, "/Engine/Private/HairStrands/HairStrandsClusterCulling.usf", "MainClusterCullingPrepareIndirectDrawsCS", SF_Compute);
+
+///////////////////////////////////////////////////////////////////////////////////////////////////
+
+class FHairClusterCullingLocalBlockPreFixSumCS : public FGlobalShader
+{
+	DECLARE_GLOBAL_SHADER(FHairClusterCullingLocalBlockPreFixSumCS);
+	SHADER_USE_PARAMETER_STRUCT(FHairClusterCullingLocalBlockPreFixSumCS, FGlobalShader);
+
+	BEGIN_SHADER_PARAMETER_STRUCT(FParameters, )
+		SHADER_PARAMETER_RDG_BUFFER(Buffer, DispatchIndirectParametersClusterCountDiv512)
+		SHADER_PARAMETER_RDG_BUFFER_SRV(Buffer, DispatchIndirectParametersClusterCount)
+		SHADER_PARAMETER_RDG_BUFFER_SRV(Buffer, GlobalIndexCountBuffer)
+		SHADER_PARAMETER_RDG_BUFFER_UAV(Buffer, PerBlocklIndexCountPreFixSumBuffer)
+		SHADER_PARAMETER_RDG_BUFFER_UAV(Buffer, PerBlocklTotalIndexCountBuffer)
+	END_SHADER_PARAMETER_STRUCT()
+
+public:
+	static bool ShouldCompilePermutation(const FGlobalShaderPermutationParameters& Parameters) { return IsHairStrandsSupported(Parameters.Platform); }
+	static void ModifyCompilationEnvironment(const FGlobalShaderPermutationParameters& Parameters, FShaderCompilerEnvironment& OutEnvironment)
+	{
+		FGlobalShader::ModifyCompilationEnvironment(Parameters, OutEnvironment);
+		OutEnvironment.SetDefine(TEXT("SHADER_PERBLOCKPREFIXSUM"), 1);
+	}
+};
+
+IMPLEMENT_GLOBAL_SHADER(FHairClusterCullingLocalBlockPreFixSumCS, "/Engine/Private/HairStrands/HairStrandsClusterCulling.usf", "MainPerBlockPreFixSumCS", SF_Compute);
+
+///////////////////////////////////////////////////////////////////////////////////////////////////
+
+class FHairClusterCullingCompactVertexIdsLocalBlockCS : public FGlobalShader
+{
+	DECLARE_GLOBAL_SHADER(FHairClusterCullingCompactVertexIdsLocalBlockCS);
+	SHADER_USE_PARAMETER_STRUCT(FHairClusterCullingCompactVertexIdsLocalBlockCS, FGlobalShader);
+
+	BEGIN_SHADER_PARAMETER_STRUCT(FParameters, )
+		SHADER_PARAMETER_RDG_BUFFER(Buffer, DispatchIndirectParametersBuffer)
+		SHADER_PARAMETER_RDG_BUFFER_SRV(Buffer, DispatchIndirectParametersClusterCount)
+		SHADER_PARAMETER_RDG_BUFFER_SRV(Buffer, DispatchIndirectParametersClusterCount2D)
+		SHADER_PARAMETER_RDG_BUFFER_SRV(Buffer, DispatchIndirectParametersClusterCountDiv512)
+		SHADER_PARAMETER_RDG_BUFFER_SRV(Buffer, PerBlocklIndexCountPreFixSumBuffer)
+		SHADER_PARAMETER_RDG_BUFFER_SRV(Buffer, PerBlocklTotalIndexCountPreFixSumBuffer)
+		SHADER_PARAMETER_RDG_BUFFER_SRV(Buffer, GlobalIndexStartBuffer)
+		SHADER_PARAMETER_RDG_BUFFER_SRV(Buffer, GlobalIndexCountBuffer)
+		SHADER_PARAMETER_RDG_BUFFER_SRV(Buffer, GlobalRadiusScaleBuffer)
+		SHADER_PARAMETER_SRV(Buffer, ClusterVertexIdBuffer)
+		SHADER_PARAMETER_UAV(Buffer, CulledCompactedIndexBuffer)
+		SHADER_PARAMETER_UAV(Buffer, CulledCompactedRadiusScaleBuffer)
+	END_SHADER_PARAMETER_STRUCT()
+
+public:
+	static bool ShouldCompilePermutation(const FGlobalShaderPermutationParameters& Parameters) { return IsHairStrandsSupported(Parameters.Platform); }
+	static void ModifyCompilationEnvironment(const FGlobalShaderPermutationParameters& Parameters, FShaderCompilerEnvironment& OutEnvironment)
+	{
+		FGlobalShader::ModifyCompilationEnvironment(Parameters, OutEnvironment);
+		OutEnvironment.SetDefine(TEXT("SHADER_CLUSTERCULLINGCOMPACTVERTEXIDLOCALBLOCK"), 1);
+	}
+};
+
+IMPLEMENT_GLOBAL_SHADER(FHairClusterCullingCompactVertexIdsLocalBlockCS, "/Engine/Private/HairStrands/HairStrandsClusterCulling.usf", "MainClusterCullingCompactVertexIdsLocalBlockCS", SF_Compute);
+
+///////////////////////////////////////////////////////////////////////////////////////////////////
+
+class FHairClusterCullingPreFixSumCS : public FGlobalShader
+{
+	DECLARE_GLOBAL_SHADER(FHairClusterCullingPreFixSumCS);
+	SHADER_USE_PARAMETER_STRUCT(FHairClusterCullingPreFixSumCS, FGlobalShader);
+
+	BEGIN_SHADER_PARAMETER_STRUCT(FParameters, )
+		SHADER_PARAMETER_RDG_BUFFER(Buffer, DispatchIndirectParameters)
+		SHADER_PARAMETER_RDG_BUFFER_SRV(Buffer, DispatchIndirectParametersClusterCount)
+		SHADER_PARAMETER_RDG_BUFFER_SRV(Buffer, GlobalIndexCountBuffer)
+		SHADER_PARAMETER_RDG_BUFFER_UAV(Buffer, GlobalIndexCountPreFixSumBuffer)
+	END_SHADER_PARAMETER_STRUCT()
+
+public:
+	static bool ShouldCompilePermutation(const FGlobalShaderPermutationParameters& Parameters) { return IsHairStrandsSupported(Parameters.Platform); }
+	static void ModifyCompilationEnvironment(const FGlobalShaderPermutationParameters& Parameters, FShaderCompilerEnvironment& OutEnvironment)
+	{
+		FGlobalShader::ModifyCompilationEnvironment(Parameters, OutEnvironment);
+		OutEnvironment.SetDefine(TEXT("SHADER_CLUSTERCULLINGPREFIXSUM"), 1);
+	}
+};
+
+IMPLEMENT_GLOBAL_SHADER(FHairClusterCullingPreFixSumCS, "/Engine/Private/HairStrands/HairStrandsClusterCulling.usf", "MainClusterCullingPreFixSumCS", SF_Compute);
+
+///////////////////////////////////////////////////////////////////////////////////////////////////
+
+static FVector CapturedCameraWorldPos;
+static FMatrix CapturedWorldToClipMatrix;
+
+struct FHairHZBParameters
+{
+	// Set from renderer for view culling
+	FVector HZBUvFactorValue;
+	FVector4 HZBSizeValue;
+	TRefCountPtr<IPooledRenderTarget> HZB;
+};
+
+static void AddClusterCullingPass(
+	FRDGBuilder& GraphBuilder,
+	TShaderMap<FGlobalShaderType>* ShaderMap,
+	const FViewInfo& View,
+	const FHairCullingParams& CullingParameters,
+	const FHairHZBParameters& HZBParameters,
+	const int32 ClusterCullingLodMode,
+	FHairStrandClusterData::FHairGroup& ClusterData)
+{
+	FRDGBufferRef DispatchIndirectParametersClusterCount = GraphBuilder.CreateBuffer(FRDGBufferDesc::CreateIndirectDesc<FRHIDispatchIndirectParameters>(), TEXT("HairDispatchIndirectParametersClusterCount"));
+	FRDGBufferRef DispatchIndirectParametersClusterCount2D = GraphBuilder.CreateBuffer(FRDGBufferDesc::CreateIndirectDesc<FRHIDispatchIndirectParameters>(), TEXT("HairDispatchIndirectParametersClusterCount2D"));
+	FRDGBufferRef DispatchIndirectParametersClusterCountDiv512 = GraphBuilder.CreateBuffer(FRDGBufferDesc::CreateIndirectDesc<FRHIDispatchIndirectParameters>(), TEXT("HairDispatchIndirectParametersClusterCountDiv512"));
+	FRDGBufferRef DispatchIndirectParametersClusterCountDiv512Div512 = GraphBuilder.CreateBuffer(FRDGBufferDesc::CreateIndirectDesc<FRHIDispatchIndirectParameters>(), TEXT("HairDispatchIndirectParametersClusterCountDiv512Div512"));
+	
+	FRDGBufferRef GlobalIndexStartBuffer = GraphBuilder.CreateBuffer(FRDGBufferDesc::CreateBufferDesc(sizeof(uint32), ClusterData.ClusterCount), TEXT("HairGlobalIndexStartBuffer"));
+	FRDGBufferRef GlobalIndexCountBuffer = GraphBuilder.CreateBuffer(FRDGBufferDesc::CreateBufferDesc(sizeof(uint32), ClusterData.ClusterCount), TEXT("HairGlobalIndexCountBuffer"));
+	FRDGBufferRef GlobalRadiusScaleBuffer = GraphBuilder.CreateBuffer(FRDGBufferDesc::CreateBufferDesc(sizeof(float), ClusterData.ClusterCount), TEXT("HairGlobalRadiusScaleBuffer"));
+
+	FRDGBufferRef PerBlocklTotalIndexCountBuffer = GraphBuilder.CreateBuffer(FRDGBufferDesc::CreateBufferDesc(sizeof(uint32), ClusterData.ClusterCount), TEXT("PerBlocklTotalIndexCountBuffer"));
+	FRDGBufferRef PerBlocklTotalIndexCountPreFixSumBuffer = GraphBuilder.CreateBuffer(FRDGBufferDesc::CreateBufferDesc(sizeof(uint32) * 2, ClusterData.ClusterCount), TEXT("PerBlocklTotalIndexCountPreFixSumBuffer"));
+	FRDGBufferRef PerBlocklIndexCountPreFixSumBuffer = GraphBuilder.CreateBuffer(FRDGBufferDesc::CreateBufferDesc(sizeof(uint32) * 2, ClusterData.ClusterCount), TEXT("PerBlocklIndexCountPreFixSumBuffer"));
+
+	FRWBuffer& DrawIndirectParametersBuffer = ClusterData.HairGroupPublicPtr->GetDrawIndirectBuffer();
+	
+	bool bClusterDebugAABBBuffer = false;
+	bool bClusterDebug = false;
+#if WITH_EDITOR
+	static IConsoleVariable* CVarClusterDebug = IConsoleManager::Get().FindConsoleVariable(TEXT("r.HairStrands.Cluster.Debug"));
+	bClusterDebugAABBBuffer = CVarClusterDebug && CVarClusterDebug->GetInt() >= 2;
+	bClusterDebug = CVarClusterDebug && CVarClusterDebug->GetInt() > 0;
+	FRDGBufferRef ClusterDebugAABBBuffer = GraphBuilder.CreateBuffer(FRDGBufferDesc::CreateBufferDesc(sizeof(uint32), ClusterData.ClusterCount), TEXT("CulledCompactedIndexBuffer"));
+#endif
+
+	/// Initialise indirect buffers to be setup during the culling process
+	{
+		FHairIndBufferClearCS::FParameters* Parameters = GraphBuilder.AllocParameters<FHairIndBufferClearCS::FParameters>();
+		Parameters->DispatchIndirectParametersClusterCount = GraphBuilder.CreateUAV(DispatchIndirectParametersClusterCount);
+		Parameters->DrawIndirectParameters = DrawIndirectParametersBuffer.UAV;
+
+		FHairIndBufferClearCS::FPermutationDomain Permutation;
+		Permutation.Set<FHairIndBufferClearCS::FSetIndirectDraw>(0);
+		TShaderMapRef<FHairIndBufferClearCS> ComputeShader(ShaderMap, Permutation);
+		FComputeShaderUtils::AddPass(
+			GraphBuilder,
+			RDG_EVENT_NAME("BufferClearCS"),
+			*ComputeShader,
+			Parameters,
+			FIntVector(1, 1, 1));
+	}
+
+	/// Cull cluster, generate indirect dispatch and prepare data to expand index buffer
+	{
+		const bool bClusterCullingFrozenCamera = GHairStrandsClusterCullingFreezeCamera > 0;
+		if (!bClusterCullingFrozenCamera)
+		{
+			CapturedCameraWorldPos = View.ViewMatrices.GetViewOrigin();
+			CapturedWorldToClipMatrix = View.ViewMatrices.GetViewProjectionMatrix();
+		}
+
+		FHairClusterCullingCS::FParameters* Parameters = GraphBuilder.AllocParameters<FHairClusterCullingCS::FParameters>();
+		Parameters->CameraWorldPos = CapturedCameraWorldPos;
+		Parameters->WorldToClipMatrix = CapturedWorldToClipMatrix;
+		Parameters->ClusterCount = ClusterData.ClusterCount;
+		Parameters->LodMode = ClusterCullingLodMode;
+		Parameters->ShadowViewMode = CullingParameters.bShadowViewMode ? 1 : 0;
+		Parameters->LodBias = ClusterData.LodBias;
+		Parameters->LodAverageVertexPerPixel = ClusterData.LodAverageVertexPerPixel;
+		Parameters->ViewRectPixel = FVector2D(float(View.UnconstrainedViewRect.Width()), float(View.UnconstrainedViewRect.Height()));
+
+		Parameters->NumConvexHullPlanes = View.ViewFrustum.Planes.Num();
+		check(Parameters->NumConvexHullPlanes <= 6);
+		for (uint32 i = 0; i < Parameters->NumConvexHullPlanes; ++i)
+		{
+			Parameters->ViewFrustumConvexHull[i] = View.ViewFrustum.Planes[i];
+		}
+
+		Parameters->ClusterAABBBuffer = ClusterData.ClusterAABBBuffer->SRV;
+		Parameters->ClusterInfoBuffer = ClusterData.ClusterInfoBuffer->SRV;
+		Parameters->ClusterIndexRadiusScaleInfoBuffer = ClusterData.ClusterIndexRadiusScaleInfoBuffer->SRV;
+
+		Parameters->GlobalIndexStartBuffer = GraphBuilder.CreateUAV(GlobalIndexStartBuffer, PF_R32_UINT);
+		Parameters->GlobalIndexCountBuffer = GraphBuilder.CreateUAV(GlobalIndexCountBuffer, PF_R32_UINT);
+		Parameters->GlobalRadiusScaleBuffer = GraphBuilder.CreateUAV(GlobalRadiusScaleBuffer, PF_R32_FLOAT);
+
+		Parameters->DispatchIndirectParametersClusterCount = GraphBuilder.CreateUAV(DispatchIndirectParametersClusterCount);
+		Parameters->DrawIndirectParameters = DrawIndirectParametersBuffer.UAV;
+
+#if WITH_EDITOR
+		Parameters->ClusterDebugAABBBuffer = GraphBuilder.CreateUAV(ClusterDebugAABBBuffer, PF_R32_SINT);
+#endif
+
+		Parameters->HZBUvFactor = HZBParameters.HZBUvFactorValue;
+		Parameters->HZBSize = HZBParameters.HZBSizeValue;
+		Parameters->HZBTexture = HZBParameters.HZB.IsValid() ? GraphBuilder.RegisterExternalTexture(HZBParameters.HZB, TEXT("HairClusterCullingHZB")) : nullptr;
+		Parameters->HZBSampler = TStaticSamplerState<SF_Point, AM_Clamp, AM_Clamp, AM_Clamp>::GetRHI();
+
+		FHairClusterCullingCS::FPermutationDomain Permutation;
+		Permutation.Set<FHairClusterCullingCS::FHZBCulling>((HZBParameters.HZB.IsValid() && !bClusterCullingFrozenCamera && GHairStrandsClusterCullingUsesHzb) ? 1 : 0);
+		Permutation.Set<FHairClusterCullingCS::FDebugAABBBuffer>(bClusterDebugAABBBuffer ? 1 : 0);
+		Permutation.Set<FHairClusterCullingCS::FLodDebug>(ClusterCullingLodMode > -1 ? 1 : 0);
+		TShaderMapRef<FHairClusterCullingCS> ComputeShader(ShaderMap, Permutation);
+		const FIntVector DispatchCount = DispatchCount.DivideAndRoundUp(FIntVector(ClusterData.ClusterCount, 1, 1), FIntVector(64, 1, 1));
+		FComputeShaderUtils::AddPass(
+			GraphBuilder,
+			RDG_EVENT_NAME("ClusterCullingCS"),
+			*ComputeShader,
+			Parameters,
+			DispatchCount);
+	}
+
+	/// Prepare some indirect draw buffers for specific compute group size
+	{
+		FMainClusterCullingPrepareIndirectDrawsCS::FParameters* Parameters = GraphBuilder.AllocParameters<FMainClusterCullingPrepareIndirectDrawsCS::FParameters>();
+		Parameters->DispatchIndirectParametersClusterCount = GraphBuilder.CreateSRV(DispatchIndirectParametersClusterCount, PF_R32_UINT);
+		Parameters->DispatchIndirectParametersClusterCount2D = GraphBuilder.CreateUAV(DispatchIndirectParametersClusterCount2D, PF_R32_UINT);
+		Parameters->DispatchIndirectParametersClusterCountDiv512 = GraphBuilder.CreateUAV(DispatchIndirectParametersClusterCountDiv512, PF_R32_UINT);
+		Parameters->DispatchIndirectParametersClusterCountDiv512Div512 = GraphBuilder.CreateUAV(DispatchIndirectParametersClusterCountDiv512Div512, PF_R32_UINT);
+
+		TShaderMapRef<FMainClusterCullingPrepareIndirectDrawsCS> ComputeShader(ShaderMap);
+		FComputeShaderUtils::AddPass(
+			GraphBuilder,
+			RDG_EVENT_NAME("PrepareIndirectDrawsCS"),
+			*ComputeShader,
+			Parameters,
+			FIntVector(2, 1, 1));
+	}
+
+	/// local prefix sum per 512 block
+	{
+		FHairClusterCullingLocalBlockPreFixSumCS::FParameters* Parameters = GraphBuilder.AllocParameters<FHairClusterCullingLocalBlockPreFixSumCS::FParameters>();
+		Parameters->DispatchIndirectParametersClusterCountDiv512 = DispatchIndirectParametersClusterCountDiv512;
+		Parameters->DispatchIndirectParametersClusterCount = GraphBuilder.CreateSRV(DispatchIndirectParametersClusterCount);
+		Parameters->GlobalIndexCountBuffer = GraphBuilder.CreateSRV(GlobalIndexCountBuffer, PF_R32_UINT);
+		Parameters->PerBlocklIndexCountPreFixSumBuffer = GraphBuilder.CreateUAV(PerBlocklIndexCountPreFixSumBuffer, PF_R32G32_UINT);
+		Parameters->PerBlocklTotalIndexCountBuffer = GraphBuilder.CreateUAV(PerBlocklTotalIndexCountBuffer, PF_R32_UINT);
+
+		TShaderMapRef<FHairClusterCullingLocalBlockPreFixSumCS> ComputeShader(ShaderMap);
+		FComputeShaderUtils::AddPass(
+			GraphBuilder,
+			RDG_EVENT_NAME("WithinBlockIndexCountPreFixSumCS"),
+			*ComputeShader,
+			Parameters,
+			DispatchIndirectParametersClusterCountDiv512, 0); // FIX ME, this could get over 65535
+		check(ClusterData.ClusterCount / 512 <= 65535);
+	}
+
+	/// Prefix sum on the total index count per block of 512
+	{
+		FHairClusterCullingPreFixSumCS::FParameters* Parameters = GraphBuilder.AllocParameters<FHairClusterCullingPreFixSumCS::FParameters>();
+		Parameters->DispatchIndirectParameters = DispatchIndirectParametersClusterCountDiv512Div512;
+		Parameters->DispatchIndirectParametersClusterCount = GraphBuilder.CreateSRV(DispatchIndirectParametersClusterCountDiv512, PF_R32_UINT);
+		Parameters->GlobalIndexCountBuffer = GraphBuilder.CreateSRV(PerBlocklTotalIndexCountBuffer, PF_R32_UINT);
+		Parameters->GlobalIndexCountPreFixSumBuffer = GraphBuilder.CreateUAV(PerBlocklTotalIndexCountPreFixSumBuffer, PF_R32G32_UINT);
+
+		TShaderMapRef<FHairClusterCullingPreFixSumCS> ComputeShader(ShaderMap);
+		FComputeShaderUtils::AddPass(
+			GraphBuilder,
+			RDG_EVENT_NAME("BlockIndexCountPreFixSumCS"),
+			*ComputeShader,
+			Parameters,
+			DispatchIndirectParametersClusterCountDiv512Div512, 0); // FIX ME, this could get over 65535
+		check(ClusterData.ClusterCount / (512*512) <= 65535);
+	}
+
+	/// Compact to VertexId buffer using hierarchical binary search or splatting
+	{
+		FHairClusterCullingCompactVertexIdsLocalBlockCS::FParameters* Parameters = GraphBuilder.AllocParameters<FHairClusterCullingCompactVertexIdsLocalBlockCS::FParameters>();
+
+		Parameters->DispatchIndirectParametersClusterCount = GraphBuilder.CreateSRV(DispatchIndirectParametersClusterCount, PF_R32_UINT);
+		Parameters->DispatchIndirectParametersClusterCount2D = GraphBuilder.CreateSRV(DispatchIndirectParametersClusterCount2D, PF_R32_UINT);
+		Parameters->DispatchIndirectParametersClusterCountDiv512 = GraphBuilder.CreateSRV(DispatchIndirectParametersClusterCountDiv512, PF_R32_UINT);
+
+		Parameters->PerBlocklIndexCountPreFixSumBuffer = GraphBuilder.CreateSRV(PerBlocklIndexCountPreFixSumBuffer, PF_R32G32_UINT);
+		Parameters->PerBlocklTotalIndexCountPreFixSumBuffer = GraphBuilder.CreateSRV(PerBlocklTotalIndexCountPreFixSumBuffer, PF_R32G32_UINT);
+
+		Parameters->GlobalIndexStartBuffer = GraphBuilder.CreateSRV(GlobalIndexStartBuffer, PF_R32_UINT);
+		Parameters->GlobalIndexCountBuffer = GraphBuilder.CreateSRV(GlobalIndexCountBuffer, PF_R32_UINT);
+		Parameters->GlobalRadiusScaleBuffer = GraphBuilder.CreateSRV(GlobalRadiusScaleBuffer, PF_R32_FLOAT);
+		Parameters->ClusterVertexIdBuffer = ClusterData.ClusterVertexIdBuffer->SRV;
+
+		check(ClusterData.GetCulledVertexIdBuffer());
+		check(ClusterData.GetCulledVertexRadiusScaleBuffer());
+		Parameters->CulledCompactedIndexBuffer = ClusterData.GetCulledVertexIdBuffer()->UAV;
+		Parameters->CulledCompactedRadiusScaleBuffer = ClusterData.GetCulledVertexRadiusScaleBuffer()->UAV;
+
+		Parameters->DispatchIndirectParametersBuffer = DispatchIndirectParametersClusterCount2D;
+		TShaderMapRef<FHairClusterCullingCompactVertexIdsLocalBlockCS> ComputeShader(ShaderMap);
+		FComputeShaderUtils::AddPass(
+			GraphBuilder,
+			RDG_EVENT_NAME("SplatCompactVertexIdsCS"),
+			*ComputeShader,
+			Parameters,
+			DispatchIndirectParametersClusterCount2D, 0); // DispatchIndirectParametersClusterCount2D is used to avoid having any dispatch dimension going above 65535.
+	}
+
+	// Should this be move onto the culling result?
+#if WITH_EDITOR
+	if (bClusterDebugAABBBuffer)
+	{
+		GraphBuilder.QueueBufferExtraction(ClusterDebugAABBBuffer, &ClusterData.ClusterDebugAABBBuffer,	FRDGResourceState::EAccess::Read, FRDGResourceState::EPipeline::Graphics);
+	}
+	if (bClusterDebug)
+	{
+		GraphBuilder.QueueBufferExtraction(DispatchIndirectParametersClusterCount, &ClusterData.CulledDispatchIndirectParametersClusterCount, FRDGResourceState::EAccess::Read, FRDGResourceState::EPipeline::Graphics);
+	}
+#endif
+}
+
+///////////////////////////////////////////////////////////////////////////////////////////////////
+
+static void AddClusterResetLod0(
+	FRDGBuilder& GraphBuilder,
+	TShaderMap<FGlobalShaderType>* ShaderMap,
+	FHairStrandClusterData::FHairGroup& ClusterData)
+{
+	// Set as culling result not available
+	ClusterData.SetCullingResultAvailable(false);
+
+	// Initialise indirect buffers to entire lod 0 dispatch
+	FHairIndBufferClearCS::FParameters* Parameters = GraphBuilder.AllocParameters<FHairIndBufferClearCS::FParameters>();
+	Parameters->DrawIndirectParameters = ClusterData.HairGroupPublicPtr->GetDrawIndirectBuffer().UAV;
+	Parameters->VertexCountPerInstance = ClusterData.HairGroupPublicPtr->GetGroupInstanceVertexCount();
+
+	FHairIndBufferClearCS::FPermutationDomain Permutation;
+	Permutation.Set<FHairIndBufferClearCS::FSetIndirectDraw>(1);
+	TShaderMapRef<FHairIndBufferClearCS> ComputeShader(ShaderMap, Permutation);
+	FComputeShaderUtils::AddPass(
+		GraphBuilder,
+		RDG_EVENT_NAME("BufferClearCS"),
+		*ComputeShader,
+		Parameters,
+		FIntVector(1, 1, 1));
+}
+
+///////////////////////////////////////////////////////////////////////////////////////////////////
+
+void ComputeHairStrandsClustersCulling(
+	FRHICommandListImmediate& RHICmdList,
+	TShaderMap<FGlobalShaderType>& ShaderMap,
+	const TArray<FViewInfo>& Views,
+	const FHairCullingParams& CullingParameters,
+	FHairStrandClusterData& ClusterDatas)
+{
+	DECLARE_GPU_STAT(HairStrandsClusterCulling);
+	SCOPED_DRAW_EVENT(RHICmdList, HairStrandsClusterCulling);
+	SCOPED_GPU_STAT(RHICmdList, HairStrandsClusterCulling);
+
+	uint32 ViewCount = Views.Num();
+	TArray<const FSceneView*> SceneViews;
+	SceneViews.SetNumUninitialized(ViewCount);
+	for (uint32 ViewId = 0; ViewId < ViewCount; ++ViewId)
+	{
+		SceneViews[ViewId] = &Views[ViewId];
+	}
+
+	FHairHZBParameters HZBParameters;
+	if (ViewCount > 0) // only handling one view for now
+	{
+		const FViewInfo& ViewInfo = Views[0];
+		HZBParameters.HZB = ViewInfo.HZB.IsValid() ? ViewInfo.HZB : nullptr;
+
+		const float kHZBTestMaxMipmap = 9.0f;
+		const float HZBMipmapCounts = FMath::Log2(FMath::Max(ViewInfo.HZBMipmap0Size.X, ViewInfo.HZBMipmap0Size.Y));
+		const FVector HZBUvFactorValue(
+			float(ViewInfo.ViewRect.Width()) / float(2 * ViewInfo.HZBMipmap0Size.X),
+			float(ViewInfo.ViewRect.Height()) / float(2 * ViewInfo.HZBMipmap0Size.Y),
+			FMath::Max(HZBMipmapCounts - kHZBTestMaxMipmap, 0.0f)
+		);
+		const FVector4 HZBSizeValue(
+			ViewInfo.HZBMipmap0Size.X,
+			ViewInfo.HZBMipmap0Size.Y,
+			1.0f / float(ViewInfo.HZBMipmap0Size.X),
+			1.0f / float(ViewInfo.HZBMipmap0Size.Y)
+		);
+		HZBParameters.HZBUvFactorValue = HZBUvFactorValue;
+		HZBParameters.HZBSizeValue = HZBSizeValue;
+	}
+	
+	const bool bCullingProcessSkipped = GHairStrandsClusterCulling <= 0 || (CullingParameters.bShadowViewMode && GStrandHairClusterCullingShadow <= 0);
+	for (const FViewInfo& View : Views)
+	{
+		for (FHairStrandClusterData::FHairGroup& ClusterData : ClusterDatas.HairGroups)		
+		{
+			if (bCullingProcessSkipped)
+			{
+				FRDGBuilder GraphBuilder(RHICmdList);
+				AddClusterResetLod0(GraphBuilder, &ShaderMap, ClusterData);
+				GraphBuilder.Execute();
+				continue;
+			}
+
+			const uint32 LODMode = CullingParameters.bShadowViewMode ? GHairStrandsClusterCullingShadowLodMode : GHairStrandsClusterCullingLodMode;
+			// TODO use compute overlap (will need to split AddClusterCullingPass)
+			FRDGBuilder GraphBuilder(RHICmdList);
+			AddClusterCullingPass(
+				GraphBuilder, 
+				&ShaderMap, 
+				View, 
+				CullingParameters,
+				HZBParameters,
+				LODMode,
+				ClusterData);
+			GraphBuilder.Execute();
+
+			ClusterData.SetCullingResultAvailable(true);
+		}
+	}
+}
+
+///////////////////////////////////////////////////////////////////////////////////////////////////
+
+void ResetHairStrandsClusterToLOD0(
+	FRHICommandListImmediate& RHICmdList,
+	TShaderMap<FGlobalShaderType>& ShaderMap,
+	FHairStrandClusterData& ClusterDatas)
+{
+	DECLARE_GPU_STAT(HairStrandsResetLod0);
+	SCOPED_DRAW_EVENT(RHICmdList, HairStrandsResetLod0);
+	SCOPED_GPU_STAT(RHICmdList, HairStrandsResetLod0);
+
+	for (FHairStrandClusterData::FHairGroup& ClusterData : ClusterDatas.HairGroups)
+	{
+		// TODO use compute overlap (will need to split AddClusterCullingPass)
+		FRDGBuilder GraphBuilder(RHICmdList);
+		AddClusterResetLod0(GraphBuilder, &ShaderMap, ClusterData);
+		GraphBuilder.Execute();
+	}
+}
diff --git a/Engine/Source/Runtime/Renderer/Private/HairStrands/HairStrandsClusters.h b/Engine/Source/Runtime/Renderer/Private/HairStrands/HairStrandsClusters.h
new file mode 100644
--- /dev/null
+++ b/Engine/Source/Runtime/Renderer/Private/HairStrands/HairStrandsClusters.h
@@ -0,0 +1,29 @@
+// Copyright Epic Games, Inc. All Rights Reserved.
+
+/*=============================================================================
+	HairStrandsCulling.h: Hair strands culling implementation.
+=============================================================================*/
+
+#pragma once
+
+#include "CoreMinimal.h"
+#include "RendererInterface.h"
+#include "HairStrandsInterface.h"
+
+struct FHairCullingParams
+{
+	bool bShadowViewMode		= false;		// Set to true for shadow views, all strand will be kept but those outside view will
+	bool bCullingProcessSkipped	= false;
+};
+
+void ComputeHairStrandsClustersCulling(
+	FRHICommandListImmediate& RHICmdList,
+	TShaderMap<FGlobalShaderType>& ShaderMap,
+	const TArray<FViewInfo>& Views,
+	const FHairCullingParams& CullingParameters,
+	FHairStrandClusterData& ClusterDatas);
+
+void ResetHairStrandsClusterToLOD0(
+	FRHICommandListImmediate& RHICmdList,
+	TShaderMap<FGlobalShaderType>& ShaderMap,
+	FHairStrandClusterData& ClusterDatas);
\ No newline at end of file
diff --git a/Engine/Source/Runtime/Renderer/Private/HairStrands/HairStrandsComposition.cpp b/Engine/Source/Runtime/Renderer/Private/HairStrands/HairStrandsComposition.cpp
--- a/Engine/Source/Runtime/Renderer/Private/HairStrands/HairStrandsComposition.cpp
+++ b/Engine/Source/Runtime/Renderer/Private/HairStrands/HairStrandsComposition.cpp
@@ -9,6 +9,7 @@
 #include "RenderGraphUtils.h"
 #include "PostProcessing.h"
 #include "HairStrandsRendering.h"
+#include "HairStrandsScatter.h"
 
 /////////////////////////////////////////////////////////////////////////////////////////
 
@@ -225,13 +226,25 @@ void RenderHairComposeSubPixel(
 		{
 			if (ViewIndex < HairVisibilityViews.HairDatas.Num())
 			{
-				TRefCountPtr<IPooledRenderTarget> CategorisationTexture = HairVisibilityViews.HairDatas[ViewIndex].CategorizationTexture;
+				const FHairStrandsMacroGroupDatas& MacroGroupDatas = HairDatas->MacroGroupsPerViews.Views[ViewIndex];
+				const FHairStrandsVisibilityData& VisibilityData = HairVisibilityViews.HairDatas[ViewIndex];
+
+				TRefCountPtr<IPooledRenderTarget> CategorisationTexture = VisibilityData.CategorizationTexture;
 				if (!CategorisationTexture)
 				{
 					continue; // Automatically skip for any view not rendering hair
 				}
 				const FRDGTextureRef RDGCategorisationTexture = CategorisationTexture ? GraphBuilder.RegisterExternalTexture(CategorisationTexture, TEXT("HairVisibilityCategorisationTexture")) : nullptr;
 
+				AddHairDiffusionPass(
+					GraphBuilder,
+					View,
+					VisibilityData,
+					MacroGroupDatas.VirtualVoxelResources,
+					SceneColorDepth,
+					SceneColorSubPixelTexture,
+					SceneColorTexture);
+
 				// #hair_todo : compose partially covered hair with transparent surface: this can be done by 
 				// rendering quad(s) covering the hair at the correct depth. This will be sorted with other 
 				// transparent surface, which should make the overall sorting workable
